- Fix memory leak due to not destroying CQDIPerInstanceData inside DestroyPerInstanceData

diff --git a/Engine/Plugins/FX/Niagara/Source/Niagara/Classes/NiagaraDataInterfaceCollisionQuery.h b/Engine/Plugins/FX/Niagara/Source/Niagara/Classes/NiagaraDataInterfaceCollisionQuery.h
--- a/Engine/Plugins/FX/Niagara/Source/Niagara/Classes/NiagaraDataInterfaceCollisionQuery.h
+++ b/Engine/Plugins/FX/Niagara/Source/Niagara/Classes/NiagaraDataInterfaceCollisionQuery.h
@@ -42,7 +42,7 @@ public:
 	/** Initializes the per instance data for this interface. Returns false if there was some error and the simulation should be disabled. */
 	virtual bool InitPerInstanceData(void* PerInstanceData, FNiagaraSystemInstance* InSystemInstance) override;
 	/** Destroys the per instence data for this interface. */
-	virtual void DestroyPerInstanceData(void* PerInstanceData, FNiagaraSystemInstance* InSystemInstance) {}
+	virtual void DestroyPerInstanceData(void* PerInstanceData, FNiagaraSystemInstance* InSystemInstance) override;
 
 	/** Ticks the per instance data for this interface, if it has any. */
 	virtual bool PerInstanceTick(void* PerInstanceData, FNiagaraSystemInstance* SystemInstance, float DeltaSeconds);
diff --git a/Engine/Plugins/FX/Niagara/Source/Niagara/Private/NiagaraDataInterfaceCollisionQuery.cpp b/Engine/Plugins/FX/Niagara/Source/Niagara/Private/NiagaraDataInterfaceCollisionQuery.cpp
--- a/Engine/Plugins/FX/Niagara/Source/Niagara/Private/NiagaraDataInterfaceCollisionQuery.cpp
+++ b/Engine/Plugins/FX/Niagara/Source/Niagara/Private/NiagaraDataInterfaceCollisionQuery.cpp
@@ -33,6 +33,11 @@ bool UNiagaraDataInterfaceCollisionQuery::InitPerInstanceData(void* PerInstanceD
 	return true;
 }
 
+void UNiagaraDataInterfaceCollisionQuery::DestroyPerInstanceData(void* PerInstanceData, FNiagaraSystemInstance* InSystemInstance)
+{
+	CQDIPerInstanceData* InstData = (CQDIPerInstanceData*)PerInstanceData;
+	InstData->~CQDIPerInstanceData();
+}
 
 void UNiagaraDataInterfaceCollisionQuery::PostInitProperties()
 {
