- Integrate fix for being able to get a specific bone off the end of the array on the GPU
- We now mirror the clamp that happens on the CPU

diff --git a/Engine/Plugins/FX/Niagara/Shaders/Private/NiagaraDataInterfaceSkeletalMesh.ush b/Engine/Plugins/FX/Niagara/Shaders/Private/NiagaraDataInterfaceSkeletalMesh.ush
--- a/Engine/Plugins/FX/Niagara/Shaders/Private/NiagaraDataInterfaceSkeletalMesh.ush
+++ b/Engine/Plugins/FX/Niagara/Shaders/Private/NiagaraDataInterfaceSkeletalMesh.ush
@@ -458,7 +458,8 @@ void DISkelMesh_GetSpecificBoneCount(in FDISkelMeshContext DIContext, out int Co
 
 void DISkelMesh_GetSpecificBoneAt(in FDISkelMeshContext DIContext, in int InBoneIndex, out int Bone)
 {
-    Bone = (int)DIContext.SpecificBones[InBoneIndex];
+    int BoneIndex = min(InBoneIndex, int(DIContext.NumSpecificBones - 1));
+    Bone = (int)DIContext.SpecificBones[BoneIndex];
 }
 
 void DISkelMesh_RandomSpecificBone(in FDISkelMeshContext DIContext, out int Bone)
