- Integrate: Fill GPU skeletal mesh data even if we don't have a component
- This makes CPU & GPU mirror one another

diff --git a/Engine/Plugins/FX/Niagara/Source/Niagara/Private/NiagaraDataInterfaceSkeletalMesh.cpp b/Engine/Plugins/FX/Niagara/Source/Niagara/Private/NiagaraDataInterfaceSkeletalMesh.cpp
--- a/Engine/Plugins/FX/Niagara/Source/Niagara/Private/NiagaraDataInterfaceSkeletalMesh.cpp
+++ b/Engine/Plugins/FX/Niagara/Source/Niagara/Private/NiagaraDataInterfaceSkeletalMesh.cpp
@@ -611,24 +611,38 @@ void FSkeletalMeshGpuDynamicBufferProxy::ReleaseRHI()
 
 void FSkeletalMeshGpuDynamicBufferProxy::NewFrame(const FNDISkeletalMesh_InstanceData* InstanceData, int32 LODIndex)
 {
+	// Grab Skeletal Component / Mesh, we must have a mesh at minimum to set the data
 	USkeletalMeshComponent* SkelComp = InstanceData != nullptr ? Cast<USkeletalMeshComponent>(InstanceData->Component.Get()) : nullptr;
-	if (SkelComp)
+	USkeletalMesh* SkelMesh = SkelComp ? SkelComp->SkeletalMesh : InstanceData->MeshSafe.Get();
+	if ( SkelMesh == nullptr )
+	{
+		return;
+	}
+
+	// Create AllSectionsRefToLocalMatrices
+	TArray<FVector4> AllSectionsRefToLocalMatrices;
+	static_assert(sizeof(FVector4) == 4 * sizeof(float), "FVector4 should match 4 * floats");
 	{
 		TArray<FMatrix> RefToLocalMatrices;
-		SkelComp->CacheRefToLocalMatrices(RefToLocalMatrices);
+		if (SkelComp)
+		{
+			//-TODO: We need to handle MasterPoseComponent here
+			SkelComp->CacheRefToLocalMatrices(RefToLocalMatrices);
+		}
+		else
+		{
+			RefToLocalMatrices = SkelMesh->RefBasesInvMatrix;
+		}
 
-		TIndirectArray<FSkeletalMeshLODRenderData>& LODRenderDataArray = SkelComp->SkeletalMesh->GetResourceForRendering()->LODRenderData;
-		check(0 <= LODIndex  && LODIndex < LODRenderDataArray.Num());
+		TIndirectArray<FSkeletalMeshLODRenderData>& LODRenderDataArray = SkelMesh->GetResourceForRendering()->LODRenderData;
+		check(0 <= LODIndex && LODIndex < LODRenderDataArray.Num());
 		FSkeletalMeshLODRenderData& LODRenderData = LODRenderDataArray[LODIndex];
 		TArray<FSkelMeshRenderSection>& Sections = LODRenderData.RenderSections;
 		uint32 SectionCount = Sections.Num();
 
-		TArray<FVector4> AllSectionsRefToLocalMatrices;
-		static_assert(sizeof(FVector4) == 4*sizeof(float), "FVector4 should match 4 * floats");
-
 		// Count number of matrices we want before appending all of them according to the per section mapping from BoneMap
 		uint32 Float4Count = 0;
-		for ( const FSkelMeshRenderSection& Section : Sections )
+		for (const FSkelMeshRenderSection& Section : Sections)
 		{
 			Float4Count += Section.BoneMap.Num() * 3;
 		}
@@ -636,7 +650,7 @@ void FSkeletalMeshGpuDynamicBufferProxy::NewFrame(const FNDISkeletalMesh_Instanc
 		AllSectionsRefToLocalMatrices.AddUninitialized(Float4Count);
 
 		Float4Count = 0;
-		for ( const FSkelMeshRenderSection& Section : Sections )
+		for (const FSkelMeshRenderSection& Section : Sections)
 		{
 			const uint32 MatrixCount = Section.BoneMap.Num();
 			for (uint32 m = 0; m < MatrixCount; ++m)
@@ -645,58 +659,84 @@ void FSkeletalMeshGpuDynamicBufferProxy::NewFrame(const FNDISkeletalMesh_Instanc
 				Float4Count += 3;
 			}
 		}
+	}
+
+	// Generate information for bone sampling
+	TArray<FVector4> BoneSamplingData;
+	{
+		BoneSamplingData.Reserve((SamplingBoneCount + SamplingSocketCount) * 2);
 
-		// Generate information for bone sampling
-		TArray<FVector4> BoneSamplingData;
+		// Append bones
+		if (SkelComp != nullptr)
 		{
+			//-TODO: We need to handle MasterPoseComponent here
 			const TArray<FTransform>& ComponentTransforms = SkelComp->GetComponentSpaceTransforms();
 			check(ComponentTransforms.Num() == SamplingBoneCount);
 
-			BoneSamplingData.Reserve((SamplingBoneCount + SamplingSocketCount) * 2);
-
-			// Append bones
 			for (const FTransform& BoneTransform : ComponentTransforms)
 			{
 				const FQuat Rotation = BoneTransform.GetRotation();
 				BoneSamplingData.Add(BoneTransform.GetLocation());
 				BoneSamplingData.Add(FVector4(Rotation.X, Rotation.Y, Rotation.Z, Rotation.W));
 			}
+		}
+		else
+		{
+			//-TODO: Opt and combine with MaterPostComponent
+			TArray<FTransform> TempBoneTransforms;
+			TempBoneTransforms.Reserve(SamplingBoneCount);
 
-			// Append sockets
-			for (const FTransform& SocketTransform : InstanceData->GetSpecificSocketsCurrBuffer())
+			const TArray<FTransform>& RefTransforms = SkelMesh->RefSkeleton.GetRefBonePose();
+			for (int32 i=0; i < RefTransforms.Num(); ++i)
 			{
-				const FQuat Rotation = SocketTransform.GetRotation();
-				BoneSamplingData.Add(SocketTransform.GetLocation());
+				FTransform BoneTransform = RefTransforms[i];
+				const int32 ParentIndex = SkelMesh->RefSkeleton.GetParentIndex(i);
+				if (TempBoneTransforms.IsValidIndex(ParentIndex))
+				{
+					BoneTransform = BoneTransform * TempBoneTransforms[ParentIndex];
+				}
+				TempBoneTransforms.Add(BoneTransform);
+
+				const FQuat Rotation = BoneTransform.GetRotation();
+				BoneSamplingData.Add(BoneTransform.GetLocation());
 				BoneSamplingData.Add(FVector4(Rotation.X, Rotation.Y, Rotation.Z, Rotation.W));
 			}
 		}
 
-		FSkeletalMeshGpuDynamicBufferProxy* ThisProxy = this;
-		ENQUEUE_RENDER_COMMAND(UpdateSpawnInfoForSkinnedMesh)(
-			[ThisProxy, AllSectionsRefToLocalMatrices = MoveTemp(AllSectionsRefToLocalMatrices), BoneSamplingData = MoveTemp(BoneSamplingData)](FRHICommandListImmediate& RHICmdList) mutable
-			{
-				ThisProxy->CurrentBoneBufferId = (ThisProxy->CurrentBoneBufferId + 1) % BufferBoneCount;
-				ThisProxy->bPrevBoneGpuBufferValid = ThisProxy->bBoneGpuBufferValid;
-				ThisProxy->bBoneGpuBufferValid = true;
+		// Append sockets
+		for (const FTransform& SocketTransform : InstanceData->GetSpecificSocketsCurrBuffer())
+		{
+			const FQuat Rotation = SocketTransform.GetRotation();
+			BoneSamplingData.Add(SocketTransform.GetLocation());
+			BoneSamplingData.Add(FVector4(Rotation.X, Rotation.Y, Rotation.Z, Rotation.W));
+		}
+	}
 
-				// Copy bone remap data matrices
-				{
-					const uint32 NumBytes = AllSectionsRefToLocalMatrices.Num() * sizeof(FVector4);
-					void* DstData = RHILockVertexBuffer(ThisProxy->GetRWBufferBone().SectionBuffer, 0, NumBytes, RLM_WriteOnly);
-					FMemory::Memcpy(DstData, AllSectionsRefToLocalMatrices.GetData(), NumBytes);
-					RHIUnlockVertexBuffer(ThisProxy->GetRWBufferBone().SectionBuffer);
-				}
+	FSkeletalMeshGpuDynamicBufferProxy* ThisProxy = this;
+	ENQUEUE_RENDER_COMMAND(UpdateSpawnInfoForSkinnedMesh)(
+		[ThisProxy, AllSectionsRefToLocalMatrices = MoveTemp(AllSectionsRefToLocalMatrices), BoneSamplingData = MoveTemp(BoneSamplingData)](FRHICommandListImmediate& RHICmdList) mutable
+		{
+			ThisProxy->CurrentBoneBufferId = (ThisProxy->CurrentBoneBufferId + 1) % BufferBoneCount;
+			ThisProxy->bPrevBoneGpuBufferValid = ThisProxy->bBoneGpuBufferValid;
+			ThisProxy->bBoneGpuBufferValid = true;
 
-				// Copy bone sampling data
-				{
-					const uint32 NumBytes = BoneSamplingData.Num() * sizeof(FVector4);
-					FVector4* DstData = reinterpret_cast<FVector4*>(RHILockVertexBuffer(ThisProxy->GetRWBufferBone().SamplingBuffer, 0, NumBytes, RLM_WriteOnly));
-					FMemory::Memcpy(DstData, BoneSamplingData.GetData(), NumBytes);
-					RHIUnlockVertexBuffer(ThisProxy->GetRWBufferBone().SamplingBuffer);
-				}
+			// Copy bone remap data matrices
+			{
+				const uint32 NumBytes = AllSectionsRefToLocalMatrices.Num() * sizeof(FVector4);
+				void* DstData = RHILockVertexBuffer(ThisProxy->GetRWBufferBone().SectionBuffer, 0, NumBytes, RLM_WriteOnly);
+				FMemory::Memcpy(DstData, AllSectionsRefToLocalMatrices.GetData(), NumBytes);
+				RHIUnlockVertexBuffer(ThisProxy->GetRWBufferBone().SectionBuffer);
 			}
-		);
-	}
+
+			// Copy bone sampling data
+			{
+				const uint32 NumBytes = BoneSamplingData.Num() * sizeof(FVector4);
+				FVector4* DstData = reinterpret_cast<FVector4*>(RHILockVertexBuffer(ThisProxy->GetRWBufferBone().SamplingBuffer, 0, NumBytes, RLM_WriteOnly));
+				FMemory::Memcpy(DstData, BoneSamplingData.GetData(), NumBytes);
+				RHIUnlockVertexBuffer(ThisProxy->GetRWBufferBone().SamplingBuffer);
+			}
+		}
+	);
 }
 
 //////////////////////////////////////////////////////////////////////////
