- Change InverseFast to Inverse to fix warnings / NaNs when we have a zero scale on one axis

diff --git a/Engine/Plugins/FX/Niagara/Source/Niagara/Private/NiagaraDataInterfaceSkeletalMesh.cpp b/Engine/Plugins/FX/Niagara/Source/Niagara/Private/NiagaraDataInterfaceSkeletalMesh.cpp
--- a/Engine/Plugins/FX/Niagara/Source/Niagara/Private/NiagaraDataInterfaceSkeletalMesh.cpp
+++ b/Engine/Plugins/FX/Niagara/Source/Niagara/Private/NiagaraDataInterfaceSkeletalMesh.cpp
@@ -1570,7 +1570,7 @@ bool FNDISkeletalMesh_InstanceData::Tick(UNiagaraDataInterfaceSkeletalMesh* Inte
 			PrevTransform = Transform;
 			PrevTransformInverseTransposed = TransformInverseTransposed;
 			Transform = Component->GetComponentToWorld().ToMatrixWithScale();
-			TransformInverseTransposed = Transform.InverseFast().GetTransposed();
+			TransformInverseTransposed = Transform.Inverse().GetTransposed();
 		}
 		else
 		{
