- Fix issue with fast path functions not able to bind after shader compiler changes

diff --git a/Engine/Plugins/FX/Niagara/Source/Niagara/Private/NiagaraFunctionLibrary.cpp b/Engine/Plugins/FX/Niagara/Source/Niagara/Private/NiagaraFunctionLibrary.cpp
--- a/Engine/Plugins/FX/Niagara/Source/Niagara/Private/NiagaraFunctionLibrary.cpp
+++ b/Engine/Plugins/FX/Niagara/Source/Niagara/Private/NiagaraFunctionLibrary.cpp
@@ -434,9 +434,9 @@ UNiagaraParameterCollectionInstance* UNiagaraFunctionLibrary::GetNiagaraParamete
 }
 
 
-const TArray<FNiagaraFunctionSignature>& UNiagaraFunctionLibrary::GetVectorVMFastPathOps()
+const TArray<FNiagaraFunctionSignature>& UNiagaraFunctionLibrary::GetVectorVMFastPathOps(bool bIgnoreConsoleVariable)
 {
-	if (GAllowFastPathFunctionLibrary == 0)
+	if (!bIgnoreConsoleVariable && (GAllowFastPathFunctionLibrary == 0))
 	{
 		static TArray<FNiagaraFunctionSignature> Empty;
 		return Empty;
diff --git a/Engine/Plugins/FX/Niagara/Source/Niagara/Public/NiagaraFunctionLibrary.h b/Engine/Plugins/FX/Niagara/Source/Niagara/Public/NiagaraFunctionLibrary.h
--- a/Engine/Plugins/FX/Niagara/Source/Niagara/Public/NiagaraFunctionLibrary.h
+++ b/Engine/Plugins/FX/Niagara/Source/Niagara/Public/NiagaraFunctionLibrary.h
@@ -59,7 +59,7 @@ public:
 	UFUNCTION(BlueprintCallable, Category = Niagara, meta = (Keywords = "niagara parameter collection", WorldContext = "WorldContextObject"))
 	static UNiagaraParameterCollectionInstance* GetNiagaraParameterCollection(UObject* WorldContextObject, UNiagaraParameterCollection* Collection);
 
-	static const TArray<FNiagaraFunctionSignature>& GetVectorVMFastPathOps();
+	static const TArray<FNiagaraFunctionSignature>& GetVectorVMFastPathOps(bool bIgnoreConsoleVariable = false);
 	static bool DefineFunctionHLSL(const FNiagaraFunctionSignature& FunctionSignature, FString& HlslOutput);
 
 	static bool GetVectorVMFastPathExternalFunction(const FVMExternalFunctionBindingInfo& BindingInfo, FVMExternalFunction &OutFunc);
diff --git a/Engine/Plugins/FX/Niagara/Source/NiagaraEditor/Private/NiagaraCompiler.cpp b/Engine/Plugins/FX/Niagara/Source/NiagaraEditor/Private/NiagaraCompiler.cpp
--- a/Engine/Plugins/FX/Niagara/Source/NiagaraEditor/Private/NiagaraCompiler.cpp
+++ b/Engine/Plugins/FX/Niagara/Source/NiagaraEditor/Private/NiagaraCompiler.cpp
@@ -13,6 +13,7 @@
 #include "NiagaraSystem.h"
 #include "NiagaraNodeEmitter.h"
 #include "NiagaraNodeInput.h"
+#include "NiagaraFunctionLibrary.h"
 #include "NiagaraScriptSource.h"
 #include "NiagaraDataInterface.h"
 #include "NiagaraDataInterfaceStaticMesh.h"
@@ -1060,6 +1061,18 @@ TOptional<FNiagaraCompileResults> FHlslNiagaraCompiler::GetCompileResult(int32 J
 				}
 			}
 
+			// Look in function library
+			if (Sig == nullptr)
+			{
+				Sig = UNiagaraFunctionLibrary::GetVectorVMFastPathOps(true).FindByPredicate(
+					[&](const FNiagaraFunctionSignature& CheckSig)
+				{
+					FString SigSymbol = FHlslNiagaraTranslator::GetFunctionSignatureSymbol(CheckSig);
+					return SigSymbol == FuncInfo.Name;
+				}
+				);
+			}
+
 			if (Sig)
 			{
 				int32 NewBindingIdx = Results.Data->CalledVMExternalFunctions.AddDefaulted();
@@ -1075,7 +1088,7 @@ TOptional<FNiagaraCompileResults> FHlslNiagaraCompiler::GetCompileResult(int32 J
 			}
 			else
 			{
-				Error(FText::Format(LOCTEXT("VectorVMExternalFunctionBindingError", "Failed to bind the exernal function call:  {0}"), FText::FromString(FuncInfo.Name)));
+				Error(FText::Format(LOCTEXT("VectorVMExternalFunctionBindingError", "Failed to bind the external function call:  {0}"), FText::FromString(FuncInfo.Name)));
 				Results.bVMSucceeded = false;
 			}
 		}
