Merging Global Distance Field and volume texture changes

diff --git a/Engine/Plugins/FX/Niagara/Source/Niagara/Classes/NiagaraDataInterfaceVolumeTexture.h b/Engine/Plugins/FX/Niagara/Source/Niagara/Classes/NiagaraDataInterfaceVolumeTexture.h
new file mode 100644
--- /dev/null
+++ b/Engine/Plugins/FX/Niagara/Source/Niagara/Classes/NiagaraDataInterfaceVolumeTexture.h
@@ -0,0 +1,73 @@
+// Copyright Epic Games, Inc. All Rights Reserved.
+
+#pragma once
+
+#include "NiagaraDataInterface.h"
+#include "NiagaraCommon.h"
+#include "VectorVM.h"
+#include "NiagaraDataInterfaceVolumeTexture.generated.h"
+
+class UVolumeTexture;
+
+/** Data Interface allowing sampling of a texture */
+UCLASS(EditInlineNew, Category = "Texture", meta = (DisplayName = "Volume Texture Sample"))
+class NIAGARA_API UNiagaraDataInterfaceVolumeTexture : public UNiagaraDataInterface
+{
+	GENERATED_UCLASS_BODY()
+public:
+
+	UPROPERTY(EditAnywhere, Category = "Texture")
+	UVolumeTexture* Texture;
+
+	//UObject Interface
+	virtual void PostInitProperties()override;
+	virtual void PostLoad() override;
+#if WITH_EDITOR
+	virtual void PostEditChangeProperty(struct FPropertyChangedEvent& PropertyChangedEvent) override;
+#endif
+	//UObject Interface End
+
+	//UNiagaraDataInterface Interface
+	virtual void GetFunctions(TArray<FNiagaraFunctionSignature>& OutFunctions)override;
+	virtual void GetVMExternalFunction(const FVMExternalFunctionBindingInfo& BindingInfo, void* InstanceData, FVMExternalFunction &OutFunc) override;
+	virtual bool CanExecuteOnTarget(ENiagaraSimTarget Target)const override { return Target == ENiagaraSimTarget::GPUComputeSim; }
+	//UNiagaraDataInterface Interface End
+
+	void SampleVolumeTexture(FVectorVMContext& Context);
+	void GetTextureDimensions(FVectorVMContext& Context);
+
+	virtual bool Equals(const UNiagaraDataInterface* Other) const override;
+
+	// GPU sim functionality
+	virtual void GetParameterDefinitionHLSL(const FNiagaraDataInterfaceGPUParamInfo& ParamInfo, FString& OutHLSL) override;
+	virtual bool GetFunctionHLSL(const FNiagaraDataInterfaceGPUParamInfo& ParamInfo, const FNiagaraDataInterfaceGeneratedFunction& FunctionInfo, int FunctionInstanceIndex, FString& OutHLSL) override;
+	virtual FNiagaraDataInterfaceParametersCS* ConstructComputeParameters()const override;
+
+	//FRWBuffer& GetGPUBuffer();
+	static const FString TextureName;
+	static const FString SamplerName;
+	static const FString DimensionsBaseName;
+
+	void SetTexture(UVolumeTexture* InTexture);
+
+protected:
+	virtual bool CopyToInternal(UNiagaraDataInterface* Destination) const override;
+
+	void PushToRenderThread();
+
+	static const FName SampleVolumeTextureName;
+	static const FName TextureDimsName;
+};
+
+struct FNiagaraDataInterfaceProxyVolumeTexture : public FNiagaraDataInterfaceProxy
+{
+	FSamplerStateRHIRef SamplerStateRHI;
+	FTextureReferenceRHIRef TextureReferenceRHI;
+	FVector TexDims;
+
+	virtual void ConsumePerInstanceDataFromGameThread(void* PerInstanceData, const FNiagaraSystemInstanceID& Instance) override { check(false); }
+	virtual int32 PerInstanceDataPassedToRenderThreadSize() const override
+	{
+		return 0;
+	}
+};
\ No newline at end of file
diff --git a/Engine/Plugins/FX/Niagara/Source/Niagara/Private/NiagaraDataInterfaceVolumeTexture.cpp b/Engine/Plugins/FX/Niagara/Source/Niagara/Private/NiagaraDataInterfaceVolumeTexture.cpp
new file mode 100644
--- /dev/null
+++ b/Engine/Plugins/FX/Niagara/Source/Niagara/Private/NiagaraDataInterfaceVolumeTexture.cpp
@@ -0,0 +1,344 @@
+// Copyright Epic Games, Inc. All Rights Reserved.
+
+#include "NiagaraDataInterfaceVolumeTexture.h"
+#include "NiagaraShader.h"
+#include "ShaderParameterUtils.h"
+#include "NiagaraCustomVersion.h"
+#include "Engine/VolumeTexture.h"
+
+
+#define LOCTEXT_NAMESPACE "UNiagaraDataInterfaceVolumeTexture"
+
+const FName UNiagaraDataInterfaceVolumeTexture::SampleVolumeTextureName(TEXT("SampleVolumeTexture"));
+const FName UNiagaraDataInterfaceVolumeTexture::TextureDimsName(TEXT("TextureDimensions3D"));
+const FString UNiagaraDataInterfaceVolumeTexture::TextureName(TEXT("Texture_"));
+const FString UNiagaraDataInterfaceVolumeTexture::SamplerName(TEXT("Sampler_"));
+const FString UNiagaraDataInterfaceVolumeTexture::DimensionsBaseName(TEXT("Dimensions_"));
+
+UNiagaraDataInterfaceVolumeTexture::UNiagaraDataInterfaceVolumeTexture(FObjectInitializer const& ObjectInitializer)
+	: Super(ObjectInitializer)
+	, Texture(nullptr)
+{
+	Proxy.Reset(new FNiagaraDataInterfaceProxyVolumeTexture());
+	PushToRenderThread();
+}
+
+void UNiagaraDataInterfaceVolumeTexture::PostInitProperties()
+{
+	Super::PostInitProperties();
+
+	if (HasAnyFlags(RF_ClassDefaultObject))
+	{
+		FNiagaraTypeRegistry::Register(FNiagaraTypeDefinition(GetClass()), true, false, false);
+	}
+
+	PushToRenderThread();
+}
+
+void UNiagaraDataInterfaceVolumeTexture::PostLoad()
+{
+	Super::PostLoad();
+
+	// Not safe since the UTexture might not have yet PostLoad() called and so UpdateResource() called.
+	// This will affect whether the SamplerStateRHI will be available or not.
+	PushToRenderThread();
+}
+
+#if WITH_EDITOR
+
+void UNiagaraDataInterfaceVolumeTexture::PostEditChangeProperty(struct FPropertyChangedEvent& PropertyChangedEvent)
+{
+	Super::PostEditChangeProperty(PropertyChangedEvent);
+	PushToRenderThread();
+}
+
+#endif
+
+bool UNiagaraDataInterfaceVolumeTexture::CopyToInternal(UNiagaraDataInterface* Destination) const
+{
+	if (!Super::CopyToInternal(Destination))
+	{
+		return false;
+	}
+	UNiagaraDataInterfaceVolumeTexture* DestinationTexture = CastChecked<UNiagaraDataInterfaceVolumeTexture>(Destination);
+	DestinationTexture->Texture = Texture;
+	DestinationTexture->PushToRenderThread();
+
+	return true;
+}
+
+bool UNiagaraDataInterfaceVolumeTexture::Equals(const UNiagaraDataInterface* Other) const
+{
+	if (!Super::Equals(Other))
+	{
+		return false;
+	}
+	const UNiagaraDataInterfaceVolumeTexture* OtherTexture = CastChecked<const UNiagaraDataInterfaceVolumeTexture>(Other);
+	return OtherTexture->Texture == Texture;
+}
+
+void UNiagaraDataInterfaceVolumeTexture::GetFunctions(TArray<FNiagaraFunctionSignature>& OutFunctions)
+{
+	{
+		FNiagaraFunctionSignature Sig;
+		Sig.Name = SampleVolumeTextureName;
+		Sig.bMemberFunction = true;
+		Sig.bRequiresContext = false;
+		Sig.Inputs.Add(FNiagaraVariable(FNiagaraTypeDefinition(GetClass()), TEXT("Texture")));
+		Sig.Inputs.Add(FNiagaraVariable(FNiagaraTypeDefinition::GetVec3Def(), TEXT("UVW")));
+		Sig.Inputs.Add(FNiagaraVariable(FNiagaraTypeDefinition::GetFloatDef(), TEXT("MipLevel")));
+		Sig.SetDescription(LOCTEXT("TextureSampleVolumeTextureDesc", "Sample the specified mip level of the input 3d texture at the specified UVW coordinates. The UVW origin (0, 0, 0) is in the bottom left hand corner of the volume."));
+		Sig.Outputs.Add(FNiagaraVariable(FNiagaraTypeDefinition::GetVec4Def(), TEXT("Value")));
+		//Sig.Owner = *GetName();
+
+		OutFunctions.Add(Sig);
+	}
+
+	{
+		FNiagaraFunctionSignature Sig;
+		Sig.Name = TextureDimsName;
+		Sig.bMemberFunction = true;
+		Sig.bRequiresContext = false;
+		Sig.Inputs.Add(FNiagaraVariable(FNiagaraTypeDefinition(GetClass()), TEXT("Texture")));
+		Sig.SetDescription(LOCTEXT("TextureDimsDesc", "Get the dimensions of mip 0 of the texture."));
+		Sig.Outputs.Add(FNiagaraVariable(FNiagaraTypeDefinition::GetVec3Def(), TEXT("Dimensions3D")));
+		//Sig.Owner = *GetName();
+
+		OutFunctions.Add(Sig);
+	}
+}
+
+DEFINE_NDI_DIRECT_FUNC_BINDER(UNiagaraDataInterfaceVolumeTexture, SampleVolumeTexture)
+void UNiagaraDataInterfaceVolumeTexture::GetVMExternalFunction(const FVMExternalFunctionBindingInfo& BindingInfo, void* InstanceData, FVMExternalFunction &OutFunc)
+{
+	if (BindingInfo.Name == SampleVolumeTextureName)
+	{
+		check(BindingInfo.GetNumInputs() == 4 && BindingInfo.GetNumOutputs() == 4);
+		NDI_FUNC_BINDER(UNiagaraDataInterfaceVolumeTexture, SampleVolumeTexture)::Bind(this, OutFunc);
+	}
+	else if (BindingInfo.Name == TextureDimsName)
+	{
+		check(BindingInfo.GetNumInputs() == 0 && BindingInfo.GetNumOutputs() == 3);
+		OutFunc = FVMExternalFunction::CreateUObject(this, &UNiagaraDataInterfaceVolumeTexture::GetTextureDimensions);
+	}
+}
+
+void UNiagaraDataInterfaceVolumeTexture::GetTextureDimensions(FVectorVMContext& Context)
+{
+	VectorVM::FExternalFuncRegisterHandler<float> OutWidth(Context);
+	VectorVM::FExternalFuncRegisterHandler<float> OutHeight(Context);
+	VectorVM::FExternalFuncRegisterHandler<float> OutDepth(Context);
+
+	if (Texture == nullptr)
+	{
+		for (int32 i = 0; i < Context.NumInstances; ++i)
+		{
+			*OutWidth.GetDestAndAdvance() = 0.0f;
+			*OutHeight.GetDestAndAdvance() = 0.0f;
+			*OutDepth.GetDestAndAdvance() = 0.0f;
+		}
+	}
+	else
+	{
+		float Width = Texture->GetSizeX();
+		float Height = Texture->GetSizeY();
+		float Depth = Texture->GetSizeZ();
+		for (int32 i = 0; i < Context.NumInstances; ++i)
+		{
+			*OutWidth.GetDestAndAdvance() = Width;
+			*OutHeight.GetDestAndAdvance() = Height;
+			*OutDepth.GetDestAndAdvance() = Depth;
+		}
+	}
+}
+
+void UNiagaraDataInterfaceVolumeTexture::SampleVolumeTexture(FVectorVMContext& Context)
+{
+	VectorVM::FExternalFuncInputHandler<float> XParam(Context);
+	VectorVM::FExternalFuncInputHandler<float> YParam(Context);
+	VectorVM::FExternalFuncInputHandler<float> ZParam(Context);
+	VectorVM::FExternalFuncInputHandler<float> MipLevelParam(Context);
+	VectorVM::FExternalFuncRegisterHandler<float> OutSampleR(Context);
+	VectorVM::FExternalFuncRegisterHandler<float> OutSampleG(Context);
+	VectorVM::FExternalFuncRegisterHandler<float> OutSampleB(Context);
+	VectorVM::FExternalFuncRegisterHandler<float> OutSampleA(Context);
+
+	for (int32 i = 0; i < Context.NumInstances; ++i)
+	{
+		float X = XParam.GetAndAdvance();
+		float Y = YParam.GetAndAdvance();
+		float Z = YParam.GetAndAdvance();
+		float Mip = MipLevelParam.GetAndAdvance();
+		*OutSampleR.GetDestAndAdvance() = 1.0;
+		*OutSampleG.GetDestAndAdvance() = 0.0;
+		*OutSampleB.GetDestAndAdvance() = 1.0;
+		*OutSampleA.GetDestAndAdvance() = 1.0;
+	}
+
+}
+
+bool UNiagaraDataInterfaceVolumeTexture::GetFunctionHLSL(const FNiagaraDataInterfaceGPUParamInfo& ParamInfo, const FNiagaraDataInterfaceGeneratedFunction& FunctionInfo, int FunctionInstanceIndex, FString& OutHLSL)
+{
+	if (FunctionInfo.DefinitionName == SampleVolumeTextureName)
+	{
+		FString HLSLTextureName = TextureName + ParamInfo.DataInterfaceHLSLSymbol;
+		FString HLSLSamplerName = SamplerName + ParamInfo.DataInterfaceHLSLSymbol;
+		OutHLSL += TEXT("void ") + FunctionInfo.InstanceName + TEXT("(in float3 In_UV, in float MipLevel, out float4 Out_Value) \n{\n");
+		OutHLSL += TEXT("\t Out_Value = ") + HLSLTextureName + TEXT(".SampleLevel(") + HLSLSamplerName + TEXT(", In_UV, MipLevel);\n");
+		OutHLSL += TEXT("\n}\n");
+		return true;
+	}
+	else if (FunctionInfo.DefinitionName == TextureDimsName)
+	{
+		FString DimsVar = DimensionsBaseName + ParamInfo.DataInterfaceHLSLSymbol;
+		OutHLSL += TEXT("void ") + FunctionInfo.InstanceName + TEXT("(out float3 Out_Value) \n{\n");
+		OutHLSL += TEXT("\t Out_Value = ") + DimsVar + TEXT(";\n");
+		OutHLSL += TEXT("\n}\n");
+		return true;
+	}
+	return false;
+}
+
+void UNiagaraDataInterfaceVolumeTexture::GetParameterDefinitionHLSL(const FNiagaraDataInterfaceGPUParamInfo& ParamInfo, FString& OutHLSL)
+{
+	FString HLSLTextureName = TextureName + ParamInfo.DataInterfaceHLSLSymbol;
+	FString HLSLSamplerName = SamplerName + ParamInfo.DataInterfaceHLSLSymbol;
+	OutHLSL += TEXT("Texture3D ") + HLSLTextureName + TEXT(";\n");
+	OutHLSL += TEXT("SamplerState ") + HLSLSamplerName + TEXT(";\n");
+	OutHLSL += TEXT("float3 ") + DimensionsBaseName + ParamInfo.DataInterfaceHLSLSymbol + TEXT(";\n");
+}
+
+
+struct FNiagaraDataInterfaceParametersCS_VolumeTexture : public FNiagaraDataInterfaceParametersCS
+{
+	virtual void Bind(const FNiagaraDataInterfaceParamRef& ParamRef, const class FShaderParameterMap& ParameterMap) override
+	{
+		FString TexName = UNiagaraDataInterfaceVolumeTexture::TextureName + ParamRef.ParameterInfo.DataInterfaceHLSLSymbol;
+		FString SampleName = (UNiagaraDataInterfaceVolumeTexture::SamplerName + ParamRef.ParameterInfo.DataInterfaceHLSLSymbol);
+		TextureParam.Bind(ParameterMap, *TexName);
+		SamplerParam.Bind(ParameterMap, *SampleName);
+
+		if (!TextureParam.IsBound())
+		{
+			UE_LOG(LogNiagara, Warning, TEXT("Binding failed for FNiagaraDataInterfaceParametersCS_VolumeTexture Texture %s. Was it optimized out?"), *TexName)
+		}
+
+		if (!SamplerParam.IsBound())
+		{
+			UE_LOG(LogNiagara, Warning, TEXT("Binding failed for FNiagaraDataInterfaceParametersCS_VolumeTexture Sampler %s. Was it optimized out?"), *SampleName)
+		}
+
+		Dimensions.Bind(ParameterMap, *(UNiagaraDataInterfaceVolumeTexture::DimensionsBaseName + ParamRef.ParameterInfo.DataInterfaceHLSLSymbol));
+
+	}
+
+	virtual void Serialize(FArchive& Ar)override
+	{
+		Ar << TextureParam;
+		Ar << SamplerParam;
+		if (Ar.IsLoading() == false || Ar.CustomVer(FNiagaraCustomVersion::GUID) >= FNiagaraCustomVersion::TextureDataInterfaceSizeSerialize)
+		{
+			Ar << Dimensions;
+		}
+	}
+
+	virtual void Set(FRHICommandList& RHICmdList, const FNiagaraDataInterfaceSetArgs& Context) const override
+	{
+		check(IsInRenderingThread());
+
+		FRHIComputeShader* ComputeShaderRHI = Context.Shader->GetComputeShader();
+		FNiagaraDataInterfaceProxyVolumeTexture* TextureDI = static_cast<FNiagaraDataInterfaceProxyVolumeTexture*>(Context.DataInterface);
+
+		FRHITexture* TextureRHI = (TextureDI && TextureDI->TextureReferenceRHI) ? TextureDI->TextureReferenceRHI->GetReferencedTexture() : nullptr;
+		if (TextureRHI)
+		{
+			FRHISamplerState* SamplerStateRHI = TextureDI->SamplerStateRHI;
+			if (!SamplerStateRHI)
+			{
+				// Fallback required because PostLoad() order affects whether RHI resources 
+				// are initalized in UNiagaraDataInterfaceVolumeTexture::PushToRenderThread().
+				SamplerStateRHI = TStaticSamplerState<SF_Point, AM_Clamp, AM_Clamp, AM_Clamp>::GetRHI();
+			}
+			SetTextureParameter(
+				RHICmdList,
+				ComputeShaderRHI,
+				TextureParam,
+				SamplerParam,
+				SamplerStateRHI,
+				TextureRHI
+			);
+			SetShaderValue(RHICmdList, ComputeShaderRHI, Dimensions, TextureDI->TexDims);
+		}
+		else
+		{
+			SetTextureParameter(
+				RHICmdList,
+				ComputeShaderRHI,
+				TextureParam,
+				SamplerParam,
+				GBlackTexture->SamplerStateRHI,
+				GBlackTexture->TextureRHI
+			);
+			FVector TexDims(0.0f, 0.0f, 0.0f);
+			SetShaderValue(RHICmdList, ComputeShaderRHI, Dimensions, TexDims);
+		}
+	}
+private:
+
+	FShaderResourceParameter TextureParam;
+	FShaderResourceParameter SamplerParam;
+	FShaderParameter Dimensions;
+};
+
+void UNiagaraDataInterfaceVolumeTexture::PushToRenderThread()
+{
+	FNiagaraDataInterfaceProxyVolumeTexture* RT_Proxy = GetProxyAs<FNiagaraDataInterfaceProxyVolumeTexture>();
+
+	float DimX = 0.0f;
+	float DimY = 0.0f;
+	float DimZ = 0.0f;
+	FTextureReferenceRHIRef RT_TextureReference;
+	FSamplerStateRHIRef RT_SamplerState;
+
+	if (Texture && Texture->TextureReference.TextureReferenceRHI.IsValid())
+	{
+		if (Texture->TextureReference.TextureReferenceRHI)
+		{
+			RT_TextureReference = Texture->TextureReference.TextureReferenceRHI;
+		}
+		if (Texture->Resource)
+		{
+			RT_SamplerState = Texture->Resource->SamplerStateRHI;
+		}
+		DimX = Texture->GetSizeX();
+		DimY = Texture->GetSizeY();
+		DimZ = Texture->GetSizeZ();
+	}
+
+	ENQUEUE_RENDER_COMMAND(FPushDITextureToRT) (
+		[RT_Proxy, RT_TextureReference, RT_SamplerState, DimX, DimY, DimZ](FRHICommandListImmediate& RHICmdList)
+	{
+		RT_Proxy->TextureReferenceRHI = RT_TextureReference;
+		RT_Proxy->SamplerStateRHI = RT_SamplerState;
+		RT_Proxy->TexDims.X = DimX;
+		RT_Proxy->TexDims.Y = DimY;
+		RT_Proxy->TexDims.Z = DimZ;
+	}
+	);
+}
+
+FNiagaraDataInterfaceParametersCS* UNiagaraDataInterfaceVolumeTexture::ConstructComputeParameters()const
+{
+	return new FNiagaraDataInterfaceParametersCS_VolumeTexture();
+}
+
+void UNiagaraDataInterfaceVolumeTexture::SetTexture(UVolumeTexture* InTexture)
+{
+	if (InTexture)
+	{
+		Texture = InTexture;
+		PushToRenderThread();
+	}
+}
+
+#undef LOCTEXT_NAMESPACE
\ No newline at end of file
diff --git a/Engine/Plugins/FX/Niagara/Source/Niagara/Private/NiagaraFunctionLibrary.cpp b/Engine/Plugins/FX/Niagara/Source/Niagara/Private/NiagaraFunctionLibrary.cpp
--- a/Engine/Plugins/FX/Niagara/Source/Niagara/Private/NiagaraFunctionLibrary.cpp
+++ b/Engine/Plugins/FX/Niagara/Source/Niagara/Private/NiagaraFunctionLibrary.cpp
@@ -13,6 +13,8 @@
 #include "NiagaraDataInterfaceStaticMesh.h"
 #include "NiagaraStats.h"
 #include "NiagaraDataInterfaceTexture.h"
+#include "NiagaraDataInterfaceVolumeTexture.h"
+#include "Engine/VolumeTexture.h"
 
 #define LOCTEXT_NAMESPACE "NiagaraFunctionLibrary"
 
@@ -398,7 +400,7 @@ void UNiagaraFunctionLibrary::SetTextureObject(UNiagaraComponent* NiagaraSystem,
 
 	if (!Texture)
 	{
-		UE_LOG(LogNiagara, Warning, TEXT("SkeletalMeshComponent in \"SetTextureObject\" is NULL, OverrideName \"%s\" and NiagaraSystem \"%s\", skipping."), *OverrideName, *NiagaraSystem->GetOwner()->GetName());
+		UE_LOG(LogNiagara, Warning, TEXT("Texture in \"SetTextureObject\" is NULL, OverrideName \"%s\" and NiagaraSystem \"%s\", skipping."), *OverrideName, *NiagaraSystem->GetOwner()->GetName());
 		return;
 	}
 
@@ -423,6 +425,41 @@ void UNiagaraFunctionLibrary::SetTextureObject(UNiagaraComponent* NiagaraSystem,
 	TexureDI->SetTexture(Texture);
 }
 
+void UNiagaraFunctionLibrary::SetVolumeTextureObject(UNiagaraComponent* NiagaraSystem, const FString& OverrideName, UVolumeTexture* Texture)
+{
+	if (!NiagaraSystem)
+	{
+		UE_LOG(LogNiagara, Warning, TEXT("NiagaraSystem in \"SetVolumeTextureObject\" is NULL, OverrideName \"%s\" and Texture \"%s\", skipping."), *OverrideName, Texture ? *Texture->GetName() : TEXT("NULL"));
+		return;
+	}
+
+	if (!Texture)
+	{
+		UE_LOG(LogNiagara, Warning, TEXT("Volume Texture in \"SetVolumeTextureObject\" is NULL, OverrideName \"%s\" and NiagaraSystem \"%s\", skipping."), *OverrideName, *NiagaraSystem->GetOwner()->GetName());
+		return;
+	}
+
+	const FNiagaraParameterStore& OverrideParameters = NiagaraSystem->GetOverrideParameters();
+
+	FNiagaraVariable Variable(FNiagaraTypeDefinition(UNiagaraDataInterfaceVolumeTexture::StaticClass()), *OverrideName);
+
+	int32 Index = OverrideParameters.IndexOf(Variable);
+	if (Index == INDEX_NONE)
+	{
+		UE_LOG(LogNiagara, Warning, TEXT("Could not find index of variable \"%s\" in the OverrideParameters map of NiagaraSystem \"%s\"."), *OverrideName, *NiagaraSystem->GetOwner()->GetName());
+		return;
+	}
+
+	UNiagaraDataInterfaceVolumeTexture* TexureDI = Cast<UNiagaraDataInterfaceVolumeTexture>(OverrideParameters.GetDataInterface(Index));
+	if (!TexureDI)
+	{
+		UE_LOG(LogNiagara, Warning, TEXT("Did not find a matching Volume Texture Data Interface variable named \"%s\" in the User variables of NiagaraSystem \"%s\" ."), *OverrideName, *NiagaraSystem->GetOwner()->GetName());
+		return;
+	}
+
+	TexureDI->SetTexture(Texture);
+}
+
 UNiagaraParameterCollectionInstance* UNiagaraFunctionLibrary::GetNiagaraParameterCollection(UObject* WorldContextObject, UNiagaraParameterCollection* Collection)
 {
 	UWorld* World = GEngine->GetWorldFromContextObject(WorldContextObject, EGetWorldErrorMode::LogAndReturnNull);
diff --git a/Engine/Plugins/FX/Niagara/Source/Niagara/Public/NiagaraFunctionLibrary.h b/Engine/Plugins/FX/Niagara/Source/Niagara/Public/NiagaraFunctionLibrary.h
--- a/Engine/Plugins/FX/Niagara/Source/Niagara/Public/NiagaraFunctionLibrary.h
+++ b/Engine/Plugins/FX/Niagara/Source/Niagara/Public/NiagaraFunctionLibrary.h
@@ -14,6 +14,7 @@
 class UNiagaraComponent;
 class UNiagaraSystem;
 class USceneComponent;
+class UVolumeTexture;
 
 /**
 * A C++ and Blueprint accessible library of utility functions for accessing Niagara simulations
@@ -52,6 +53,10 @@ public:
 	UFUNCTION(BlueprintCallable, Category = Niagara)
 	static void SetTextureObject(UNiagaraComponent* NiagaraSystem, const FString& OverrideName, UTexture* Texture);
 	
+	/** Overrides the Volume Texture for a Niagara Volume Texture Data Interface User Parameter.*/
+	UFUNCTION(BlueprintCallable, Category = Niagara)
+	static void SetVolumeTextureObject(UNiagaraComponent* NiagaraSystem, const FString& OverrideName, UVolumeTexture* Texture);
+	
 	//This is gonna be totally reworked
 // 	UFUNCTION(BlueprintCallable, Category = Niagara, meta = (Keywords = "niagara System", UnsafeDuringActorConstruction = "true"))
 // 	static void SetUpdateScriptConstant(UNiagaraComponent* Component, FName EmitterName, FName ConstantName, FVector Value);
diff --git a/Engine/Plugins/FX/Niagara/Source/NiagaraEditor/NiagaraEditor.Build.cs b/Engine/Plugins/FX/Niagara/Source/NiagaraEditor/NiagaraEditor.Build.cs
--- a/Engine/Plugins/FX/Niagara/Source/NiagaraEditor/NiagaraEditor.Build.cs
+++ b/Engine/Plugins/FX/Niagara/Source/NiagaraEditor/NiagaraEditor.Build.cs
@@ -48,6 +48,7 @@ public class NiagaraEditor : ModuleRules
 				"Projects",
                 "MainFrame",
 				"ToolMenus",
+				"Renderer",
 				"EditorWidgets",
 			}
         );
diff --git a/Engine/Plugins/FX/Niagara/Source/NiagaraEditor/Private/GlobalDistanceFieldCapture.cpp b/Engine/Plugins/FX/Niagara/Source/NiagaraEditor/Private/GlobalDistanceFieldCapture.cpp
new file mode 100644
--- /dev/null
+++ b/Engine/Plugins/FX/Niagara/Source/NiagaraEditor/Private/GlobalDistanceFieldCapture.cpp
@@ -0,0 +1,170 @@
+#include "GlobalDistanceFieldCapture.h"
+#include "Editor.h"
+#include "Stats/Stats2.h"
+#include "Renderer/Private/GlobalDistanceField.h"
+#include "HAL/IConsoleManager.h"
+#include "Engine/VolumeTexture.h"
+#include "ContentBrowser/Public/ContentBrowserModule.h"
+#include "ContentBrowser/Public/IContentBrowserSingleton.h"
+#include "LevelEditorViewport.h"
+#include "NiagaraEditorModule.h"
+
+FAutoConsoleCommand GCaptureDistanceFieldCommand(
+	TEXT("fx.CaptureGlobalDistanceField"),
+	TEXT("Creates a Volume Texture from the global distance field currently visible in the editor"),
+	FConsoleCommandWithArgsDelegate::CreateLambda(
+		[](const TArray<FString>& Args) {
+			if (Args.Num() < 3)
+			{
+				bool bRangeCompress = (Args.Num() > 0 && Args[0] == TEXT("RangeCompress"));
+				FGlobalDistanceFieldCapture::Request(bRangeCompress);
+			}
+			else
+			{
+				FVector CamPos;
+				CamPos.X = FCString::Atof(*Args[0]);
+				CamPos.Y = FCString::Atof(*Args[1]);
+				CamPos.Z = FCString::Atof(*Args[2]);
+				bool bRangeCompress = (Args.Num() > 3 && Args[3] == TEXT("RangeCompress"));
+				FGlobalDistanceFieldCapture::Request(bRangeCompress, CamPos);
+			}
+		})
+	);
+
+FGlobalDistanceFieldCapture* FGlobalDistanceFieldCapture::Singleton = nullptr;
+
+static UVolumeTexture* GetSelectedVolumeTexture()
+{
+	// Check to see if there's a volume texture selected in the content browser
+	FContentBrowserModule& ContentBrowserModule = FModuleManager::LoadModuleChecked<FContentBrowserModule>("ContentBrowser");
+	IContentBrowserSingleton& Browser = ContentBrowserModule.Get();
+
+	TArray<FAssetData> SelectedAssets;
+	Browser.GetSelectedAssets(SelectedAssets);
+
+	if (SelectedAssets.Num() == 1 && SelectedAssets[0].GetClass()->IsChildOf<UVolumeTexture>())
+	{
+		// A volume texture is selected
+		return Cast<UVolumeTexture>(SelectedAssets[0].GetAsset());
+	}
+
+	return nullptr;
+}
+
+void FGlobalDistanceFieldCapture::Request(bool bRangeCompress)
+{
+	RequestCommon(bRangeCompress, false, FVector());
+}
+
+void FGlobalDistanceFieldCapture::Request(bool bRangeCompress, const FVector& CamPos)
+{
+	RequestCommon(bRangeCompress, true, CamPos);
+}
+
+void FGlobalDistanceFieldCapture::RequestCommon(bool bRangeCompress, bool bSetCamPos, const FVector& CamPos)
+{
+	if (Singleton != nullptr)
+	{
+		UE_LOG(LogConsoleResponse, Error, TEXT("Cannot capture the global distance field because a capture is already in progress."));
+		return;
+	}
+
+	UVolumeTexture* VolumeTexture = GetSelectedVolumeTexture();
+	if (VolumeTexture == nullptr)
+	{
+		UE_LOG(LogConsoleResponse, Error, TEXT("No Volume Texture selected in the Content Browser. Select the Volume Texture you'd like to overwrite and try again."));
+		return;
+	}
+
+	Singleton = new FGlobalDistanceFieldCapture(VolumeTexture, bRangeCompress, bSetCamPos, CamPos);
+}
+
+FGlobalDistanceFieldCapture::FGlobalDistanceFieldCapture(UVolumeTexture* Tex, bool bCompress, bool bSetCamPos, const FVector& CamPos)
+{
+	ensure(Tex != nullptr);
+	
+	VolumeTex = Tex;
+	bRangeCompress = bCompress;
+
+	if (bSetCamPos && GCurrentLevelEditingViewportClient != nullptr)
+	{
+		StoredCamPos = GCurrentLevelEditingViewportClient->GetViewLocation();
+
+		// Small kludge: Add a tiny bit to the camera position to ensure the grid gets snapped to the same position. This is because they're likely using
+		// the Console Output to position the camera, and the grid is snapped to cell-sized coords, and placing the camera exactly at that position
+		// was causing it to snap one cell back in each dimension.
+		FVector ShiftedCamPos = CamPos + FVector(0.001f);
+
+		GCurrentLevelEditingViewportClient->SetViewLocation(ShiftedCamPos);
+		bRestoreCamera = true;
+	}
+
+	Readback = new FGlobalDistanceFieldReadback();
+	Readback->ReadbackComplete = FGlobalDistanceFieldReadback::FCompleteDelegate::CreateRaw(this, &FGlobalDistanceFieldCapture::OnReadbackComplete);
+	Readback->CallbackThread = ENamedThreads::GameThread;
+	RequestGlobalDistanceFieldReadback_GameThread(Readback);
+}
+
+FGlobalDistanceFieldCapture::~FGlobalDistanceFieldCapture()
+{
+	ensure(Singleton == this);
+
+	delete Readback;
+	Singleton = nullptr;
+}
+
+void FGlobalDistanceFieldCapture::OnReadbackComplete()
+{
+	ensure(Readback);
+
+	float MaxDist = -MAX_FLT;
+	float MinDist = MAX_FLT;
+	for (FFloat16Color Color : Readback->ReadbackData)
+	{
+		float Dist = Color.R;
+		MaxDist = FMath::Max(MaxDist, Dist);
+		MinDist = FMath::Min(MinDist, Dist);
+	}
+
+	if (VolumeTex.IsValid())
+	{
+		FIntVector Size = Readback->Size;
+		ETextureSourceFormat Format = TSF_RGBA16F;
+		const uint8* PixelData = (const uint8*)Readback->ReadbackData.GetData();
+		TArray<uint16> CompressedData;
+		if (bRangeCompress)
+		{
+			CompressedData.Empty(Size.X * Size.Y * Size.Z * sizeof(uint16));
+			for (FFloat16Color Color : Readback->ReadbackData)
+			{
+				CompressedData.Add(uint16(65535.0f * FMath::GetRangePct(MinDist, MaxDist, (float)Color.R)));
+			}
+			PixelData = (const uint8*)CompressedData.GetData();
+			Format = TSF_G16;
+		}
+
+		// Overwrite the volume texture with the new data
+		VolumeTex->Source.Init(Size.X, Size.Y, Size.Z, 1, Format, PixelData);
+		VolumeTex->Source2DTexture = nullptr; // we're no longer sourced from a 2D texture
+		VolumeTex->SetLightingGuid();
+		VolumeTex->UpdateResource();
+	}
+
+	FVector Center, Extents;
+	Readback->Bounds.GetCenterAndExtents(Center, Extents);
+
+	UE_LOG(LogConsoleResponse, Display, TEXT("GLOBAL DISTANCE FIELD CAPTURE: Texture: %s, Position: (%g, %g, %g), Extents: (%g, %g, %g), Distance Range: (%g, %g)"),
+		*VolumeTex->GetName(), Center.X, Center.Y, Center.Z, Extents.X, Extents.Y, Extents.Z, MinDist, MaxDist);	
+
+	if (bRestoreCamera && GCurrentLevelEditingViewportClient != nullptr)
+	{
+		GCurrentLevelEditingViewportClient->SetViewLocation(StoredCamPos);
+	}
+
+	delete Readback;
+	Readback = nullptr;
+	VolumeTex = nullptr;
+
+	// Work is done, we can kill the singleton instance now
+	delete this;
+}
\ No newline at end of file
diff --git a/Engine/Plugins/FX/Niagara/Source/NiagaraEditor/Private/GlobalDistanceFieldCapture.h b/Engine/Plugins/FX/Niagara/Source/NiagaraEditor/Private/GlobalDistanceFieldCapture.h
new file mode 100644
--- /dev/null
+++ b/Engine/Plugins/FX/Niagara/Source/NiagaraEditor/Private/GlobalDistanceFieldCapture.h
@@ -0,0 +1,39 @@
+#pragma once
+
+#include "CoreMinimal.h"
+
+class FGlobalDistanceFieldReadback;
+class UVolumeTexture;
+
+/**
+ * Utility interface used by an editor console command to capture the global distance field and store it into a volume
+ * texture selected in the Content Browser
+ */
+class FGlobalDistanceFieldCapture
+{
+public:
+	virtual ~FGlobalDistanceFieldCapture();
+	
+	/**	Requests a capture. Will overwrite the volume texture selected in the Content Browser or create a new one. */
+	static void Request(bool bRangeCompress);
+
+	/**	Requests a capture at the specified camera location. Will overwrite the volume texture selected in the Content Browser or create a new one. */
+	static void Request(bool bRangeCompress, const FVector& CameraPos);
+
+private:
+	TWeakObjectPtr<UVolumeTexture> VolumeTex;
+	FGlobalDistanceFieldReadback* Readback = nullptr;
+	FVector StoredCamPos;
+	bool bRestoreCamera = false;
+	bool bRangeCompress = false;
+
+	// Only allow singleton access of constructors	
+	FGlobalDistanceFieldCapture(UVolumeTexture* Tex, bool bCompress, bool bSetCamPos, const FVector& CamPos);
+	FGlobalDistanceFieldCapture(const FGlobalDistanceFieldCapture&) = delete;
+
+	void OnReadbackComplete();
+
+	static void RequestCommon(bool bRangeCompress, bool bSetCamPos, const FVector& CamPos);
+
+	static FGlobalDistanceFieldCapture* Singleton;
+};
\ No newline at end of file
diff --git a/Engine/Source/Developer/AssetTools/Private/AssetTypeActions/AssetTypeActions_Texture2D.cpp b/Engine/Source/Developer/AssetTools/Private/AssetTypeActions/AssetTypeActions_Texture2D.cpp
--- a/Engine/Source/Developer/AssetTools/Private/AssetTypeActions/AssetTypeActions_Texture2D.cpp
+++ b/Engine/Source/Developer/AssetTools/Private/AssetTypeActions/AssetTypeActions_Texture2D.cpp
@@ -42,8 +42,7 @@ void FAssetTypeActions_Texture2D::GetActions(const TArray<UObject*>& InObjects,
 			);
 	}
 
-	static const auto AllowVolumeTextureAssetCreationVar = IConsoleManager::Get().FindTConsoleVariableDataInt(TEXT("r.AllowVolumeTextureAssetCreation"));
-	if (InObjects.Num() == 1 && AllowVolumeTextureAssetCreationVar->GetValueOnGameThread() != 0)
+	if (InObjects.Num() == 1)
 	{
 		Section.AddMenuEntry(
 			"Texture2D_CreateVolumeTexture",
diff --git a/Engine/Source/Runtime/D3D12RHI/Private/D3D12RenderTarget.cpp b/Engine/Source/Runtime/D3D12RHI/Private/D3D12RenderTarget.cpp
--- a/Engine/Source/Runtime/D3D12RHI/Private/D3D12RenderTarget.cpp
+++ b/Engine/Source/Runtime/D3D12RHI/Private/D3D12RenderTarget.cpp
@@ -317,16 +317,16 @@ void FD3D12CommandContext::RHICopyToResolveTarget(FRHITexture* SourceTextureRHI,
 							const uint32 XBytes = (uint32)srcDesc.Width * BlockBytes;
 							const uint32 XBytesAligned = Align(XBytes, FD3D12_TEXTURE_DATA_PITCH_ALIGNMENT);
 
-							D3D12_SUBRESOURCE_FOOTPRINT destSubresource;
-							destSubresource.Depth = 1;
-							destSubresource.Height = srcDesc.Height;
-							destSubresource.Width = srcDesc.Width;
-							destSubresource.Format = srcDesc.Format;
-							destSubresource.RowPitch = XBytesAligned;
+							D3D12_SUBRESOURCE_FOOTPRINT DestSubresource;
+							DestSubresource.Depth = 1;
+							DestSubresource.Height = srcDesc.Height;
+							DestSubresource.Width = srcDesc.Width;
+							DestSubresource.Format = srcDesc.Format;
+							DestSubresource.RowPitch = XBytesAligned;
 
 							D3D12_PLACED_SUBRESOURCE_FOOTPRINT placedTexture2D = {0};
 							placedTexture2D.Offset = 0;
-							placedTexture2D.Footprint = destSubresource;
+							placedTexture2D.Footprint = DestSubresource;
 
 							CD3DX12_TEXTURE_COPY_LOCATION DestCopyLocation(DestTexture2D->GetResource()->GetResource(), placedTexture2D);
 							CD3DX12_TEXTURE_COPY_LOCATION SourceCopyLocation(SourceTexture2D->GetResource()->GetResource(), ResolveParams.SourceArrayIndex);
@@ -694,17 +694,17 @@ TRefCountPtr<FD3D12Resource> FD3D12DynamicRHI::GetStagingTexture(FRHITexture* Te
 	}
 
 	uint32 BytesPerPixel = ComputeBytesPerPixel(SourceDesc.Format);
-	D3D12_SUBRESOURCE_FOOTPRINT destSubresource;
-	destSubresource.Depth = 1;
-	destSubresource.Height = SourceDesc.Height;
-	destSubresource.Width = SourceDesc.Width;
-	destSubresource.Format = SourceDesc.Format;
-	destSubresource.RowPitch = XBytesAligned;
-	check(destSubresource.RowPitch % FD3D12_TEXTURE_DATA_PITCH_ALIGNMENT == 0);	// Make sure we align correctly.
+	D3D12_SUBRESOURCE_FOOTPRINT DestSubresource;
+	DestSubresource.Depth = 1;
+	DestSubresource.Height = SourceDesc.Height;
+	DestSubresource.Width = SourceDesc.Width;
+	DestSubresource.Format = SourceDesc.Format;
+	DestSubresource.RowPitch = XBytesAligned;
+	check(DestSubresource.RowPitch % FD3D12_TEXTURE_DATA_PITCH_ALIGNMENT == 0);	// Make sure we align correctly.
 
 	D3D12_PLACED_SUBRESOURCE_FOOTPRINT placedTexture2D = { 0 };
 	placedTexture2D.Offset = 0;
-	placedTexture2D.Footprint = destSubresource;
+	placedTexture2D.Footprint = DestSubresource;
 
 	CD3DX12_TEXTURE_COPY_LOCATION DestCopyLocation(TempTexture2D->GetResource(), placedTexture2D);
 	CD3DX12_TEXTURE_COPY_LOCATION SourceCopyLocation(Texture->GetResource()->GetResource(), Subresource);
@@ -1320,17 +1320,17 @@ void FD3D12DynamicRHI::ReadSurfaceDataMSAARaw(FRHICommandList_RecursiveHazardous
 	}
 
 	// Setup the descriptions for the copy to the readback heap.
-	D3D12_SUBRESOURCE_FOOTPRINT destSubresource;
-	destSubresource.Depth = 1;
-	destSubresource.Height = SizeY;
-	destSubresource.Width = SizeX;
-	destSubresource.Format = TextureDesc.Format;
-	destSubresource.RowPitch = XBytesAligned;
-	check(destSubresource.RowPitch % FD3D12_TEXTURE_DATA_PITCH_ALIGNMENT == 0);	// Make sure we align correctly.
+	D3D12_SUBRESOURCE_FOOTPRINT DestSubresource;
+	DestSubresource.Depth = 1;
+	DestSubresource.Height = SizeY;
+	DestSubresource.Width = SizeX;
+	DestSubresource.Format = TextureDesc.Format;
+	DestSubresource.RowPitch = XBytesAligned;
+	check(DestSubresource.RowPitch % FD3D12_TEXTURE_DATA_PITCH_ALIGNMENT == 0);	// Make sure we align correctly.
 
 	D3D12_PLACED_SUBRESOURCE_FOOTPRINT placedTexture2D = { 0 };
 	placedTexture2D.Offset = 0;
-	placedTexture2D.Footprint = destSubresource;
+	placedTexture2D.Footprint = DestSubresource;
 
 	CD3DX12_TEXTURE_COPY_LOCATION DestCopyLocation(StagingTexture2D->GetResource(), placedTexture2D);
 	CD3DX12_TEXTURE_COPY_LOCATION SourceCopyLocation(NonMSAATexture2D->GetResource(), Subresource);
@@ -1552,17 +1552,17 @@ void FD3D12DynamicRHI::RHIReadSurfaceFloatData(FRHITexture* TextureRHI, FIntRect
 	}
 
 	uint32 BytesPerPixel = ComputeBytesPerPixel(TextureDesc.Format);
-	D3D12_SUBRESOURCE_FOOTPRINT destSubresource;
-	destSubresource.Depth = 1;
-	destSubresource.Height = TextureDesc.Height;
-	destSubresource.Width = TextureDesc.Width;
-	destSubresource.Format = TextureDesc.Format;
-	destSubresource.RowPitch = XBytesAligned;
-	check(destSubresource.RowPitch % FD3D12_TEXTURE_DATA_PITCH_ALIGNMENT == 0);	// Make sure we align correctly.
+	D3D12_SUBRESOURCE_FOOTPRINT DestSubresource;
+	DestSubresource.Depth = 1;
+	DestSubresource.Height = TextureDesc.Height;
+	DestSubresource.Width = TextureDesc.Width;
+	DestSubresource.Format = TextureDesc.Format;
+	DestSubresource.RowPitch = XBytesAligned;
+	check(DestSubresource.RowPitch % FD3D12_TEXTURE_DATA_PITCH_ALIGNMENT == 0);	// Make sure we align correctly.
 
 	D3D12_PLACED_SUBRESOURCE_FOOTPRINT placedTexture2D = { 0 };
 	placedTexture2D.Offset = 0;
-	placedTexture2D.Footprint = destSubresource;
+	placedTexture2D.Footprint = DestSubresource;
 
 	CD3DX12_TEXTURE_COPY_LOCATION DestCopyLocation(TempTexture2D->GetResource(), placedTexture2D);
 	CD3DX12_TEXTURE_COPY_LOCATION SourceCopyLocation(Texture->GetResource()->GetResource(), Subresource);
@@ -1625,8 +1625,10 @@ void FD3D12DynamicRHI::RHIRead3DSurfaceFloatData(FRHITexture* TextureRHI, FIntRe
 	uint32 SizeZ = ZMinMax.Y - ZMinMax.X;
 
 	// Check the format of the surface
-	D3D12_RESOURCE_DESC const& TextureDesc11 = Texture->GetResource()->GetDesc();
-	check(TextureDesc11.Format == GPixelFormats[PF_FloatRGBA].PlatformFormat);
+	D3D12_RESOURCE_DESC const& TextureDesc = Texture->GetResource()->GetDesc();
+	bool bIsRGBAFmt = TextureDesc.Format == GPixelFormats[PF_FloatRGBA].PlatformFormat;
+	bool bIsR16FFmt = TextureDesc.Format == GPixelFormats[PF_R16F].PlatformFormat;
+	check(bIsRGBAFmt || bIsR16FFmt);	
 
 	// Allocate the output buffer.
 	OutData.Empty(SizeX * SizeY * SizeZ * sizeof(FFloat16Color));
@@ -1643,27 +1645,27 @@ void FD3D12DynamicRHI::RHIRead3DSurfaceFloatData(FRHITexture* TextureRHI, FIntRe
 	// create a temp 3d texture to copy render target to
 	TRefCountPtr<FD3D12Resource> TempTexture3D;
 	const uint32 BlockBytes = GPixelFormats[TextureRHI->GetFormat()].BlockBytes;
-	const uint32 XBytesAligned = Align(TextureDesc11.Width * BlockBytes, FD3D12_TEXTURE_DATA_PITCH_ALIGNMENT);
-	const uint32 DepthBytesAligned = XBytesAligned * TextureDesc11.Height;
-	const uint32 MipBytesAligned = DepthBytesAligned * TextureDesc11.DepthOrArraySize;
+	const uint32 XBytesAligned = Align(TextureDesc.Width * BlockBytes, FD3D12_TEXTURE_DATA_PITCH_ALIGNMENT);
+	const uint32 DepthBytesAligned = XBytesAligned * TextureDesc.Height;
+	const uint32 MipBytesAligned = DepthBytesAligned * TextureDesc.DepthOrArraySize;
 	VERIFYD3D12RESULT(Adapter->CreateBuffer(D3D12_HEAP_TYPE_READBACK, Node, Node, MipBytesAligned, TempTexture3D.GetInitReference(), nullptr));
 
 	// Copy the data to a staging resource.
 	uint32 Subresource = 0;
-	uint32 BytesPerPixel = ComputeBytesPerPixel(TextureDesc11.Format);
-	D3D12_SUBRESOURCE_FOOTPRINT destSubresource;
-	destSubresource.Depth = TextureDesc11.DepthOrArraySize;
-	destSubresource.Height = TextureDesc11.Height;
-	destSubresource.Width = TextureDesc11.Width;
-	destSubresource.Format = TextureDesc11.Format;
-	destSubresource.RowPitch = XBytesAligned;
-	check(destSubresource.RowPitch % FD3D12_TEXTURE_DATA_PITCH_ALIGNMENT == 0);	// Make sure we align correctly.
-
-	D3D12_PLACED_SUBRESOURCE_FOOTPRINT placedTexture3D = { 0 };
-	placedTexture3D.Offset = 0;
-	placedTexture3D.Footprint = destSubresource;
-
-	CD3DX12_TEXTURE_COPY_LOCATION DestCopyLocation(TempTexture3D->GetResource(), placedTexture3D);
+	uint32 BytesPerPixel = ComputeBytesPerPixel(TextureDesc.Format);
+	D3D12_SUBRESOURCE_FOOTPRINT DestSubresource;
+	DestSubresource.Depth = TextureDesc.DepthOrArraySize;
+	DestSubresource.Height = TextureDesc.Height;
+	DestSubresource.Width = TextureDesc.Width;
+	DestSubresource.Format = TextureDesc.Format;
+	DestSubresource.RowPitch = XBytesAligned;
+	check(DestSubresource.RowPitch % FD3D12_TEXTURE_DATA_PITCH_ALIGNMENT == 0);	// Make sure we align correctly.
+
+	D3D12_PLACED_SUBRESOURCE_FOOTPRINT PlacedTexture3D = { 0 };
+	PlacedTexture3D.Offset = 0;
+	PlacedTexture3D.Footprint = DestSubresource;
+
+	CD3DX12_TEXTURE_COPY_LOCATION DestCopyLocation(TempTexture3D->GetResource(), PlacedTexture3D);
 	CD3DX12_TEXTURE_COPY_LOCATION SourceCopyLocation(Texture->GetResource()->GetResource(), Subresource);
 
 	{
@@ -1695,17 +1697,38 @@ void FD3D12DynamicRHI::RHIRead3DSurfaceFloatData(FRHITexture* TextureRHI, FIntRe
 		OutData.AddZeroed(TotalCount);
 	}
 
-	// Read the data out of the buffer, converting it from ABGR to ARGB.
-	for (int32 Z = ZMinMax.X; Z < ZMinMax.Y; ++Z)
+	// Read the data out of the buffer
+	if (bIsRGBAFmt)
+	{
+		// Texture is RGBA16F format
+		for (int32 Z = ZMinMax.X; Z < ZMinMax.Y; ++Z)
+		{
+			for (int32 Y = InRect.Min.Y; Y < InRect.Max.Y; ++Y)
+			{
+				const FFloat16Color* SrcPtr = (const FFloat16Color*)((const uint8*)pData + (Y - InRect.Min.Y) * XBytesAligned + (Z - ZMinMax.X) * DepthBytesAligned);
+				int32 Index = (Y - InRect.Min.Y) * SizeX + (Z - ZMinMax.X) * SizeX * SizeY;
+				check(Index < OutData.Num());
+				FFloat16Color* DestPtr = &OutData[Index];
+				FMemory::Memcpy(DestPtr, SrcPtr, SizeX * sizeof(FFloat16Color));
+			}
+		}
+	}
+	else if (bIsR16FFmt)
 	{
-		for (int32 Y = InRect.Min.Y; Y < InRect.Max.Y; ++Y)
+		// Texture is R16F format
+		for (int32 Z = ZMinMax.X; Z < ZMinMax.Y; ++Z)
 		{
-			FFloat16Color* SrcPtr = (FFloat16Color*)((uint8*)pData + (Y - InRect.Min.Y) * XBytesAligned + (Z - ZMinMax.X) * DepthBytesAligned);
-			int32 Index = (Y - InRect.Min.Y) * SizeX + (Z - ZMinMax.X) * SizeX * SizeY;
-			check(Index < OutData.Num());
-			FFloat16Color* DestColor = &OutData[Index];
-			FFloat16* DestPtr = (FFloat16*)(DestColor);
-			FMemory::Memcpy(DestPtr, SrcPtr, SizeX * sizeof(FFloat16) * 4);
+			for (int32 Y = InRect.Min.Y; Y < InRect.Max.Y; ++Y)
+			{
+				const FFloat16* SrcPtr = (const FFloat16*)((const uint8*)pData + (Y - InRect.Min.Y) * XBytesAligned + (Z - ZMinMax.X) * DepthBytesAligned);
+				for (int32 X = InRect.Min.X; X < InRect.Max.X; ++X)
+				{
+					int32 Index = (Y - InRect.Min.Y) * SizeX + (Z - ZMinMax.X) * SizeX * SizeY + X;
+					check(Index < OutData.Num());
+					OutData[Index].R = SrcPtr[X];
+					OutData[Index].A = FFloat16(1.0f); // ensure full alpha (as if you sampled on GPU)
+				}
+			}
 		}
 	}
 
diff --git a/Engine/Source/Runtime/Engine/Private/Materials/MaterialExpressions.cpp b/Engine/Source/Runtime/Engine/Private/Materials/MaterialExpressions.cpp
--- a/Engine/Source/Runtime/Engine/Private/Materials/MaterialExpressions.cpp
+++ b/Engine/Source/Runtime/Engine/Private/Materials/MaterialExpressions.cpp
@@ -261,13 +261,11 @@ FUObjectAnnotationSparseBool GMaterialFunctionsThatNeedSamplerFixup;
 /** Returns whether the given expression class is allowed. */
 bool IsAllowedExpressionType(const UClass* const Class, const bool bMaterialFunction)
 {
-	static const auto AllowVolumeTextureAssetCreationVar = IConsoleManager::Get().FindTConsoleVariableDataInt(TEXT("r.AllowVolumeTextureAssetCreation"));
 	static const auto AllowTextureArrayAssetCreationVar = IConsoleManager::Get().FindTConsoleVariableDataInt(TEXT("r.AllowTexture2DArrayCreation"));
 
 	// Exclude comments from the expression list, as well as the base parameter expression, as it should not be used directly
 	const bool bSharedAllowed = Class != UMaterialExpressionComment::StaticClass() 
-		&& Class != UMaterialExpressionParameter::StaticClass()
-		&& (Class != UMaterialExpressionTextureSampleParameterVolume::StaticClass() || AllowVolumeTextureAssetCreationVar->GetValueOnGameThread() != 0)
+		&& Class != UMaterialExpressionParameter::StaticClass() && Class != UMaterialExpressionTextureSampleParameterVolume::StaticClass()
 		&& (Class != UMaterialExpressionTextureSampleParameter2DArray::StaticClass() || AllowTextureArrayAssetCreationVar->GetValueOnGameThread() != 0);
 
 	if (bMaterialFunction)
@@ -16841,29 +16839,29 @@ int32 UMaterialExpressionCurveAtlasRowParameter::Compile(class FMaterialCompiler
 	{
 		// Retrieve the curve index directly from the atlas rather than relying on the scalar parameter defaults
 		int32 CurveIndex = 0;
-
+		
 		if (Atlas->GetCurveIndex(Curve, CurveIndex))
 		{
-		DefaultValue = (float)CurveIndex;
-		int32 Slot = Compiler->ScalarParameter(ParameterName, DefaultValue);
+			DefaultValue = (float)CurveIndex;
+			int32 Slot = Compiler->ScalarParameter(ParameterName, DefaultValue);
 
-		// Get Atlas texture object and texture size
-		int32 AtlasRef = INDEX_NONE;
-		int32 AtlasCode = Compiler->Texture(Atlas, AtlasRef, SAMPLERTYPE_LinearColor, SSM_Clamp_WorldGroupSettings, TMVM_None);
+			// Get Atlas texture object and texture size
+			int32 AtlasRef = INDEX_NONE;
+			int32 AtlasCode = Compiler->Texture(Atlas, AtlasRef, SAMPLERTYPE_LinearColor, SSM_Clamp_WorldGroupSettings, TMVM_None);
 			if (AtlasCode != INDEX_NONE)
 			{
-		int32 AtlasSize = Compiler->ForceCast(Compiler->TextureProperty(AtlasCode, TMTM_TextureSize), MCT_Float1);
+				int32 AtlasSize = Compiler->ForceCast(Compiler->TextureProperty(AtlasCode, TMTM_TextureSize), MCT_Float1);
 
-		// Calculate UVs from size and slot
-		// if the input is hooked up, use it, otherwise use the internal constant
-		int32 Arg1 = InputTime.GetTracedInput().Expression ? InputTime.Compile(Compiler) : Compiler->Constant(0);
-		int32 Arg2 = Compiler->Div(Compiler->Add(Slot, Compiler->Constant(0.5)), AtlasSize);
+				// Calculate UVs from size and slot
+				// if the input is hooked up, use it, otherwise use the internal constant
+				int32 Arg1 = InputTime.GetTracedInput().Expression ? InputTime.Compile(Compiler) : Compiler->Constant(0);
+				int32 Arg2 = Compiler->Div(Compiler->Add(Slot, Compiler->Constant(0.5)), AtlasSize);
 
-		int32 UV = Compiler->AppendVector(Arg1, Arg2);
+				int32 UV = Compiler->AppendVector(Arg1, Arg2);
 
-		// Sample texture
-		return Compiler->TextureSample(AtlasCode, UV, SAMPLERTYPE_LinearColor, INDEX_NONE, INDEX_NONE, TMVM_None, SSM_Clamp_WorldGroupSettings, AtlasRef, false);
-	}
+				// Sample texture
+				return Compiler->TextureSample(AtlasCode, UV, SAMPLERTYPE_LinearColor, INDEX_NONE, INDEX_NONE, TMVM_None, SSM_Clamp_WorldGroupSettings, AtlasRef, false);
+			}
 			else
 			{
 				return CompilerError(Compiler, TEXT("There was an error when compiling the texture."));
diff --git a/Engine/Source/Runtime/Engine/Private/VolumeTexture.cpp b/Engine/Source/Runtime/Engine/Private/VolumeTexture.cpp
--- a/Engine/Source/Runtime/Engine/Private/VolumeTexture.cpp
+++ b/Engine/Source/Runtime/Engine/Private/VolumeTexture.cpp
@@ -12,13 +12,6 @@
 #include "DeviceProfiles/DeviceProfileManager.h"
 #include "Containers/ResourceArray.h"
 
-static TAutoConsoleVariable<int32> CVarAllowVolumeTextureAssetCreation(
-	TEXT("r.AllowVolumeTextureAssetCreation"),
-	0,
-	TEXT("Enable UVolumeTexture assets"),
-	ECVF_ReadOnly
-	);
-
 // Limit the possible depth of volume texture otherwise when the user converts 2D textures, he can crash the engine.
 const int32 MAX_VOLUME_TEXTURE_DEPTH = 512;
 
diff --git a/Engine/Source/Runtime/Renderer/Private/GlobalDistanceField.cpp b/Engine/Source/Runtime/Renderer/Private/GlobalDistanceField.cpp
--- a/Engine/Source/Runtime/Renderer/Private/GlobalDistanceField.cpp
+++ b/Engine/Source/Runtime/Renderer/Private/GlobalDistanceField.cpp
@@ -8,6 +8,7 @@
 #include "DistanceFieldLightingShared.h"
 #include "RendererModule.h"
 #include "ClearQuad.h"
+#include "Engine/VolumeTexture.h"
 
 DECLARE_GPU_STAT(GlobalDistanceFieldUpdate);
 
@@ -110,6 +111,18 @@ FAutoConsoleVariableRef CVarGlobalDistanceFieldHeightFieldThicknessScale(
 	ECVF_RenderThreadSafe
 	);
 
+// For reading back the distance field data
+static FGlobalDistanceFieldReadback* GDFReadbackRequest = nullptr;
+void RequestGlobalDistanceFieldReadback(FGlobalDistanceFieldReadback* Readback)
+{
+	if (ensure(GDFReadbackRequest == nullptr))
+	{
+		ensure(Readback->ReadbackComplete.IsBound());
+		ensure(Readback->CallbackThread != ENamedThreads::UnusedAnchor);
+		GDFReadbackRequest = Readback;
+	}
+}
+
 void FGlobalDistanceFieldInfo::UpdateParameterData(float MaxOcclusionDistance)
 {
 	if (Clipmaps.Num() > 0)
@@ -1008,7 +1021,8 @@ static void ComputeUpdateRegionsAndUpdateViewState(
 				|| ClipmapViewState.CachedGlobalDistanceFieldViewDistance != Scene->GlobalDistanceFieldViewDistance
 				|| ClipmapViewState.CacheMostlyStaticSeparately != GAOGlobalDistanceFieldCacheMostlyStaticSeparately
 				|| ClipmapViewState.LastUsedSceneDataForFullUpdate != &Scene->DistanceFieldSceneData
-				|| GAOGlobalDistanceFieldForceFullUpdate;
+				|| GAOGlobalDistanceFieldForceFullUpdate
+				|| GDFReadbackRequest != nullptr;
 
 			if (ShouldUpdateClipmapThisFrame(ClipmapIndex, View.ViewState->GlobalDistanceFieldUpdateIndex)
 				|| bForceFullUpdate)
@@ -1308,6 +1322,29 @@ void FViewInfo::SetupGlobalDistanceFieldUniformBufferParameters(FViewUniformShad
 	ViewUniformShaderParameters.GlobalDistanceFieldSampler3 = TStaticSamplerState<SF_Bilinear, AM_Wrap, AM_Wrap, AM_Wrap>::GetRHI();
 }
 
+void ReadbackDistanceFieldClipmap(FRHICommandListImmediate& RHICmdList, FGlobalDistanceFieldInfo& GlobalDistanceFieldInfo)
+{
+	FGlobalDistanceFieldReadback* Readback = GDFReadbackRequest;
+	GDFReadbackRequest = nullptr;
+
+	FGlobalDistanceFieldClipmap& ClipMap = GlobalDistanceFieldInfo.Clipmaps[0];
+	FTextureRHIRef SourceTexture = ClipMap.RenderTarget->GetRenderTargetItem().ShaderResourceTexture;
+	FIntVector Size = SourceTexture->GetSizeXYZ();
+	
+	RHICmdList.Read3DSurfaceFloatData(SourceTexture, FIntRect(0, 0, Size.X, Size.Y), FIntPoint(0, Size.Z), Readback->ReadbackData);
+	Readback->Bounds = ClipMap.Bounds;
+	Readback->Size = Size;
+	
+	// Fire the callback to notify that the request is complete
+	DECLARE_CYCLE_STAT(TEXT("FSimpleDelegateGraphTask.DistanceFieldReadbackDelegate"), STAT_FSimpleDelegateGraphTask_DistanceFieldReadbackDelegate, STATGROUP_TaskGraphTasks);
+	FSimpleDelegateGraphTask::CreateAndDispatchWhenReady(
+		Readback->ReadbackComplete,
+		GET_STATID(STAT_FSimpleDelegateGraphTask_DistanceFieldReadbackDelegate),
+		nullptr,
+		Readback->CallbackThread
+		);	
+}
+
 /** 
  * Updates the global distance field for a view.  
  * Typically issues updates for just the newly exposed regions of the volume due to camera movement.
@@ -1581,6 +1618,12 @@ void UpdateGlobalDistanceFieldVolume(
 			}
 		}
 	}
+
+	if (GDFReadbackRequest && GlobalDistanceFieldInfo.Clipmaps.Num() > 0)
+	{
+		// Read back a clipmap
+		ReadbackDistanceFieldClipmap(RHICmdList, GlobalDistanceFieldInfo);
+	}
 }
 
 void ListGlobalDistanceFieldMemory()
diff --git a/Engine/Source/Runtime/Renderer/Private/GlobalDistanceField.h b/Engine/Source/Runtime/Renderer/Private/GlobalDistanceField.h
--- a/Engine/Source/Runtime/Renderer/Private/GlobalDistanceField.h
+++ b/Engine/Source/Runtime/Renderer/Private/GlobalDistanceField.h
@@ -8,6 +8,8 @@
 
 #include "CoreMinimal.h"
 #include "DistanceFieldAmbientOcclusion.h"
+#include "Delegates/DelegateCombinations.h"
+#include "Async/TaskGraphInterfaces.h"
 
 extern int32 GAOGlobalDistanceField;
 
@@ -32,3 +34,30 @@ extern void UpdateGlobalDistanceFieldVolume(
 	const FScene* Scene, 
 	float MaxOcclusionDistance, 
 	FGlobalDistanceFieldInfo& Info);
+
+class FGlobalDistanceFieldReadback
+{
+public:
+	DECLARE_DELEGATE(FCompleteDelegate);
+
+	FBox Bounds;
+	FIntVector Size;
+	TArray<FFloat16Color> ReadbackData;	
+	FCompleteDelegate ReadbackComplete;
+	ENamedThreads::Type CallbackThread = ENamedThreads::UnusedAnchor;
+};
+
+/**
+ * Retrieves the GPU data of a global distance field clipmap for access by the CPU
+ *
+ * @note: Currently only works with the highest res clipmap on the first updated view in the frame
+ **/
+void RENDERER_API RequestGlobalDistanceFieldReadback(FGlobalDistanceFieldReadback* Readback);
+
+inline void RequestGlobalDistanceFieldReadback_GameThread(FGlobalDistanceFieldReadback* Readback)
+{
+	ENQUEUE_RENDER_COMMAND(RequestGlobalDistanceFieldReadback)(
+		[Readback](FRHICommandListImmediate& RHICmdList) {
+		RequestGlobalDistanceFieldReadback(Readback);
+	});
+}
\ No newline at end of file
diff --git a/Engine/Source/Runtime/Windows/D3D11RHI/Private/D3D11RenderTarget.cpp b/Engine/Source/Runtime/Windows/D3D11RHI/Private/D3D11RenderTarget.cpp
--- a/Engine/Source/Runtime/Windows/D3D11RHI/Private/D3D11RenderTarget.cpp
+++ b/Engine/Source/Runtime/Windows/D3D11RHI/Private/D3D11RenderTarget.cpp
@@ -1008,7 +1008,9 @@ void FD3D11DynamicRHI::RHIRead3DSurfaceFloatData(FRHITexture* TextureRHI,FIntRec
 	D3D11_TEXTURE3D_DESC TextureDesc;
 	((ID3D11Texture3D*)Texture->GetResource())->GetDesc(&TextureDesc);
 
-	check(TextureDesc.Format == GPixelFormats[PF_FloatRGBA].PlatformFormat);
+	bool bIsRGBAFmt = TextureDesc.Format == GPixelFormats[PF_FloatRGBA].PlatformFormat;
+	bool bIsR16FFmt = TextureDesc.Format == GPixelFormats[PF_R16F].PlatformFormat;	
+	check(bIsRGBAFmt || bIsR16FFmt);
 
 	// Allocate the output buffer.
 	OutData.Empty(SizeX * SizeY * SizeZ * sizeof(FFloat16Color));
@@ -1052,19 +1054,40 @@ void FD3D11DynamicRHI::RHIRead3DSurfaceFloatData(FRHITexture* TextureRHI,FIntRec
 		OutData.AddZeroed(TotalCount);
 	}
 
-	// Read the data out of the buffer, converting it from ABGR to ARGB.
-	for (int32 Z = ZMinMax.X; Z < ZMinMax.Y; ++Z)
+	// Read the data out of the buffer
+	if (bIsRGBAFmt)
+	{
+		// Texture data is RGBA16F
+		for (int32 Z = ZMinMax.X; Z < ZMinMax.Y; ++Z)
+		{
+			for (int32 Y = InRect.Min.Y; Y < InRect.Max.Y; ++Y)
+			{
+				const FFloat16Color* SrcPtr = (const FFloat16Color*)((const uint8*)LockedRect.pData + (Y - InRect.Min.Y) * LockedRect.RowPitch + (Z - ZMinMax.X) * LockedRect.DepthPitch);
+				int32 Index = (Y - InRect.Min.Y) * SizeX + (Z - ZMinMax.X) * SizeX * SizeY;
+				check(Index < OutData.Num());
+				FFloat16Color* DestPtr = &OutData[Index];
+				FMemory::Memcpy(DestPtr, SrcPtr, SizeX * sizeof(FFloat16Color));
+			}
+		}
+	}
+	else if (bIsR16FFmt)
 	{
-		for(int32 Y = InRect.Min.Y; Y < InRect.Max.Y; ++Y)
+		// Texture data is R16F
+		for (int32 Z = ZMinMax.X; Z < ZMinMax.Y; ++Z)
 		{
-			FFloat16Color* SrcPtr = (FFloat16Color*)((uint8*)LockedRect.pData + (Y - InRect.Min.Y) * LockedRect.RowPitch + (Z - ZMinMax.X) * LockedRect.DepthPitch);
-			int32 Index = (Y - InRect.Min.Y) * SizeX + (Z - ZMinMax.X) * SizeX * SizeY;
-			check(Index < OutData.Num());
-			FFloat16Color* DestColor = &OutData[Index];
-			FFloat16* DestPtr = (FFloat16*)(DestColor);
-			FMemory::Memcpy(DestPtr,SrcPtr,SizeX * sizeof(FFloat16) * 4);
+			for (int32 Y = InRect.Min.Y; Y < InRect.Max.Y; ++Y)
+			{
+				const FFloat16* SrcPtr = (const FFloat16*)((const uint8*)LockedRect.pData + (Y - InRect.Min.Y) * LockedRect.RowPitch + (Z - ZMinMax.X) * LockedRect.DepthPitch);
+				for (int32 X = InRect.Min.X; X < InRect.Max.X; ++X)
+				{
+					int32 Index = (Y - InRect.Min.Y) * SizeX + (Z - ZMinMax.X) * SizeX * SizeY + X;
+					check(Index < OutData.Num());
+					OutData[Index].R = SrcPtr[X];
+					OutData[Index].A = FFloat16(1.0f); // ensure full alpha (as if you sampled on GPU)
+				}
+			}
 		}
 	}
 
 	Direct3DDeviceIMContext->Unmap(TempTexture3D,0);
-}
+}
\ No newline at end of file
