Added missing copyright headers

diff --git a/Engine/Plugins/FX/Niagara/Source/NiagaraEditor/Private/GlobalDistanceFieldCapture.cpp b/Engine/Plugins/FX/Niagara/Source/NiagaraEditor/Private/GlobalDistanceFieldCapture.cpp
--- a/Engine/Plugins/FX/Niagara/Source/NiagaraEditor/Private/GlobalDistanceFieldCapture.cpp
+++ b/Engine/Plugins/FX/Niagara/Source/NiagaraEditor/Private/GlobalDistanceFieldCapture.cpp
@@ -1,3 +1,5 @@
+// Copyright Epic Games, Inc. All Rights Reserved.
+
 #include "GlobalDistanceFieldCapture.h"
 #include "Editor.h"
 #include "Stats/Stats2.h"
diff --git a/Engine/Plugins/FX/Niagara/Source/NiagaraEditor/Private/GlobalDistanceFieldCapture.h b/Engine/Plugins/FX/Niagara/Source/NiagaraEditor/Private/GlobalDistanceFieldCapture.h
--- a/Engine/Plugins/FX/Niagara/Source/NiagaraEditor/Private/GlobalDistanceFieldCapture.h
+++ b/Engine/Plugins/FX/Niagara/Source/NiagaraEditor/Private/GlobalDistanceFieldCapture.h
@@ -1,3 +1,5 @@
+// Copyright Epic Games, Inc. All Rights Reserved.
+
 #pragma once
 
 #include "CoreMinimal.h"
