Fix no-unity compile error

diff --git a/Engine/Plugins/FX/Niagara/Source/NiagaraEditor/Private/GlobalDistanceFieldCapture.h b/Engine/Plugins/FX/Niagara/Source/NiagaraEditor/Private/GlobalDistanceFieldCapture.h
--- a/Engine/Plugins/FX/Niagara/Source/NiagaraEditor/Private/GlobalDistanceFieldCapture.h
+++ b/Engine/Plugins/FX/Niagara/Source/NiagaraEditor/Private/GlobalDistanceFieldCapture.h
@@ -3,6 +3,7 @@
 #pragma once
 
 #include "CoreMinimal.h"
+#include "UObject/WeakObjectPtr.h"
 
 class FGlobalDistanceFieldReadback;
 class UVolumeTexture;
