Fix merge issues caused by interface changes with the new Niagara Volume Texture Data Interface

diff --git a/Engine/Plugins/FX/Niagara/Source/Niagara/Private/NiagaraDataInterfaceVolumeTexture.cpp b/Engine/Plugins/FX/Niagara/Source/Niagara/Private/NiagaraDataInterfaceVolumeTexture.cpp
--- a/Engine/Plugins/FX/Niagara/Source/Niagara/Private/NiagaraDataInterfaceVolumeTexture.cpp
+++ b/Engine/Plugins/FX/Niagara/Source/Niagara/Private/NiagaraDataInterfaceVolumeTexture.cpp
@@ -211,10 +211,12 @@ void UNiagaraDataInterfaceVolumeTexture::GetParameterDefinitionHLSL(const FNiaga
 
 struct FNiagaraDataInterfaceParametersCS_VolumeTexture : public FNiagaraDataInterfaceParametersCS
 {
-	virtual void Bind(const FNiagaraDataInterfaceParamRef& ParamRef, const class FShaderParameterMap& ParameterMap) override
+	DECLARE_TYPE_LAYOUT(FNiagaraDataInterfaceParametersCS_VolumeTexture, NonVirtual);
+public:
+	void Bind(const FNiagaraDataInterfaceGPUParamInfo& ParameterInfo, const class FShaderParameterMap& ParameterMap)
 	{
-		FString TexName = UNiagaraDataInterfaceVolumeTexture::TextureName + ParamRef.ParameterInfo.DataInterfaceHLSLSymbol;
-		FString SampleName = (UNiagaraDataInterfaceVolumeTexture::SamplerName + ParamRef.ParameterInfo.DataInterfaceHLSLSymbol);
+		FString TexName = UNiagaraDataInterfaceVolumeTexture::TextureName + ParameterInfo.DataInterfaceHLSLSymbol;
+		FString SampleName = (UNiagaraDataInterfaceVolumeTexture::SamplerName + ParameterInfo.DataInterfaceHLSLSymbol);
 		TextureParam.Bind(ParameterMap, *TexName);
 		SamplerParam.Bind(ParameterMap, *SampleName);
 
@@ -228,25 +230,15 @@ struct FNiagaraDataInterfaceParametersCS_VolumeTexture : public FNiagaraDataInte
 			UE_LOG(LogNiagara, Warning, TEXT("Binding failed for FNiagaraDataInterfaceParametersCS_VolumeTexture Sampler %s. Was it optimized out?"), *SampleName)
 		}
 
-		Dimensions.Bind(ParameterMap, *(UNiagaraDataInterfaceVolumeTexture::DimensionsBaseName + ParamRef.ParameterInfo.DataInterfaceHLSLSymbol));
+		Dimensions.Bind(ParameterMap, *(UNiagaraDataInterfaceVolumeTexture::DimensionsBaseName + ParameterInfo.DataInterfaceHLSLSymbol));
 
 	}
 
-	virtual void Serialize(FArchive& Ar)override
-	{
-		Ar << TextureParam;
-		Ar << SamplerParam;
-		if (Ar.IsLoading() == false || Ar.CustomVer(FNiagaraCustomVersion::GUID) >= FNiagaraCustomVersion::TextureDataInterfaceSizeSerialize)
-		{
-			Ar << Dimensions;
-		}
-	}
-
-	virtual void Set(FRHICommandList& RHICmdList, const FNiagaraDataInterfaceSetArgs& Context) const override
+	void Set(FRHICommandList& RHICmdList, const FNiagaraDataInterfaceSetArgs& Context) const
 	{
 		check(IsInRenderingThread());
 
-		FRHIComputeShader* ComputeShaderRHI = Context.Shader->GetComputeShader();
+		FRHIComputeShader* ComputeShaderRHI = Context.Shader.GetComputeShader();
 		FNiagaraDataInterfaceProxyVolumeTexture* TextureDI = static_cast<FNiagaraDataInterfaceProxyVolumeTexture*>(Context.DataInterface);
 
 		FRHITexture* TextureRHI = (TextureDI && TextureDI->TextureReferenceRHI) ? TextureDI->TextureReferenceRHI->GetReferencedTexture() : nullptr;
@@ -284,10 +276,9 @@ struct FNiagaraDataInterfaceParametersCS_VolumeTexture : public FNiagaraDataInte
 		}
 	}
 private:
-
-	FShaderResourceParameter TextureParam;
-	FShaderResourceParameter SamplerParam;
-	FShaderParameter Dimensions;
+	LAYOUT_FIELD(FShaderResourceParameter, TextureParam);
+	LAYOUT_FIELD(FShaderResourceParameter, SamplerParam);
+	LAYOUT_FIELD(FShaderParameter, Dimensions);
 };
 
 void UNiagaraDataInterfaceVolumeTexture::PushToRenderThread()
