Niagara - Rewrite the input value handling code in UNiagaraStackFuncitonInput to properly handle the new features and prevent crashes.

diff --git a/Engine/Plugins/FX/Niagara/Source/NiagaraEditor/Private/NiagaraEditorUtilities.cpp b/Engine/Plugins/FX/Niagara/Source/NiagaraEditor/Private/NiagaraEditorUtilities.cpp
--- a/Engine/Plugins/FX/Niagara/Source/NiagaraEditor/Private/NiagaraEditorUtilities.cpp
+++ b/Engine/Plugins/FX/Niagara/Source/NiagaraEditor/Private/NiagaraEditorUtilities.cpp
@@ -814,6 +814,14 @@ bool FNiagaraEditorUtilities::DataMatches(const FStructOnScope& StructOnScopeA,
 	return FMemory::Memcmp(StructOnScopeA.GetStructMemory(), StructOnScopeB.GetStructMemory(), StructOnScopeA.GetStruct()->GetStructureSize()) == 0;
 }
 
+void FNiagaraEditorUtilities::CopyDataTo(FStructOnScope& DestinationStructOnScope, const FStructOnScope& SourceStructOnScope, bool bCheckTypes)
+{
+	checkf(DestinationStructOnScope.GetStruct()->GetStructureSize() == SourceStructOnScope.GetStruct()->GetStructureSize() &&
+		(bCheckTypes == false || DestinationStructOnScope.GetStruct() == SourceStructOnScope.GetStruct()),
+		TEXT("Can not copy data from one struct to another if their size is different or if the type is different and type checking is enabled."));
+	FMemory::Memcpy(DestinationStructOnScope.GetStructMemory(), SourceStructOnScope.GetStructMemory(), SourceStructOnScope.GetStruct()->GetStructureSize());
+}
+
 TSharedPtr<SWidget> FNiagaraEditorUtilities::CreateInlineErrorText(TAttribute<FText> ErrorMessage, TAttribute<FText> ErrorTooltip)
 {
 	TSharedPtr<SHorizontalBox> ErrorInternalBox = SNew(SHorizontalBox);
@@ -1069,6 +1077,36 @@ bool FNiagaraEditorUtilities::ResolveConstantValue(UEdGraphPin* Pin, int32& Valu
 	return false;
 }
 
+TSharedPtr<FStructOnScope> FNiagaraEditorUtilities::StaticSwitchDefaultIntToStructOnScope(int32 InStaticSwitchDefaultValue, FNiagaraTypeDefinition InSwitchType)
+{
+	if (InSwitchType == FNiagaraTypeDefinition::GetBoolDef())
+	{
+		checkf(FNiagaraBool::StaticStruct()->GetStructureSize() == InSwitchType.GetSize(), TEXT("Value to type def size mismatch."));
+
+		FNiagaraBool BoolValue;
+		BoolValue.SetValue(InStaticSwitchDefaultValue != 0);
+
+		TSharedPtr<FStructOnScope> StructValue = MakeShared<FStructOnScope>(InSwitchType.GetStruct());
+		FMemory::Memcpy(StructValue->GetStructMemory(), (uint8*)(&BoolValue), InSwitchType.GetSize());
+
+		return StructValue;
+	}
+	else if (InSwitchType == FNiagaraTypeDefinition::GetIntDef() || InSwitchType.IsEnum())
+	{
+		checkf(FNiagaraInt32::StaticStruct()->GetStructureSize() == InSwitchType.GetSize(), TEXT("Value to type def size mismatch."));
+
+		FNiagaraInt32 IntValue;
+		IntValue.Value = InStaticSwitchDefaultValue;
+
+		TSharedPtr<FStructOnScope> StructValue = MakeShared<FStructOnScope>(InSwitchType.GetStruct());
+		FMemory::Memcpy(StructValue->GetStructMemory(), (uint8*)(&IntValue), InSwitchType.GetSize());
+
+		return StructValue;
+	}
+
+	return TSharedPtr<FStructOnScope>();
+}
+
 /* Go through the graph and attempt to auto-detect the type of any numeric pins by working back from the leaves of the graph. Only change the types of pins, not FNiagaraVariables.*/
 void PreprocessGraph(const UEdGraphSchema_Niagara* Schema, UNiagaraGraph* Graph, UNiagaraNodeOutput* OutputNode)
 {
diff --git a/Engine/Plugins/FX/Niagara/Source/NiagaraEditor/Private/ViewModels/Stack/NiagaraStackFunctionInput.cpp b/Engine/Plugins/FX/Niagara/Source/NiagaraEditor/Private/ViewModels/Stack/NiagaraStackFunctionInput.cpp
--- a/Engine/Plugins/FX/Niagara/Source/NiagaraEditor/Private/ViewModels/Stack/NiagaraStackFunctionInput.cpp
+++ b/Engine/Plugins/FX/Niagara/Source/NiagaraEditor/Private/ViewModels/Stack/NiagaraStackFunctionInput.cpp
@@ -54,28 +54,9 @@ UNiagaraStackFunctionInput::UNiagaraStackFunctionInput()
 	, bShowEditConditionInline(false)
 	, bIsInlineEditConditionToggle(false)
 	, bIsDynamicInputScriptReassignmentPending(false)
-	, bIsLocalOverride(false)
 {
 }
 
-void UNiagaraStackFunctionInput::AddReferencedObjects(UObject* InThis, FReferenceCollector& Collector)
-{
-	UNiagaraStackFunctionInput* This = Cast<UNiagaraStackFunctionInput>(InThis);
-	if (This != nullptr)
-	{
-		This->AddReferencedObjects(Collector);
-	}
-	Super::AddReferencedObjects(InThis, Collector);
-}
-
-void UNiagaraStackFunctionInput::AddReferencedObjects(FReferenceCollector& Collector)
-{
-	if (InputValues.DataObjects.IsValid() && InputValues.DataObjects.GetDefaultValueOwner() == FDataValues::EDefaultValueOwner::LocallyOwned)
-	{
-		Collector.AddReferencedObject(InputValues.DataObjects.GetDefaultValueObjectRef(), this);
-	}
-}
-
 // Traverses the path between the owning module node and the function call node this input belongs too collecting up the input handles between them.
 void GenerateInputParameterHandlePath(UNiagaraNodeFunctionCall& ModuleNode, UNiagaraNodeFunctionCall& FunctionCallNode, TArray<FNiagaraParameterHandle>& OutHandlePath)
 {
@@ -222,7 +203,19 @@ const FNiagaraTypeDefinition& UNiagaraStackFunctionInput::GetInputType() const
 
 FText UNiagaraStackFunctionInput::GetTooltipText() const
 {
-	return GetTooltipText(InputValues.Mode);
+	if (InputMetaData.IsSet() && InputMetaData->Description.IsEmptyOrWhitespace() == false)
+	{
+		return FText::Format(LOCTEXT("FunctionInputTooltipWithDescription", "Name: {0} \nType: {1} \nDesc: {2}"),
+			GetDisplayName(),
+			InputType.GetNameText(),
+			InputMetaData->Description);
+	}
+	else
+	{
+		return FText::Format(LOCTEXT("FunctionInputTooltip", "Name: {0} \nType: {1}"),
+			GetDisplayName(),
+			InputType.GetNameText());
+	}
 }
 
 bool UNiagaraStackFunctionInput::GetIsEnabled() const
@@ -241,7 +234,7 @@ UObject* UNiagaraStackFunctionInput::GetExternalAsset() const
 
 bool UNiagaraStackFunctionInput::TestCanCutWithMessage(FText& OutMessage) const
 {
-	if (InputValues.Mode == EValueMode::Invalid)
+	if (InputValues.HasEditableData() == false)
 	{
 		OutMessage = LOCTEXT("CantCutInvalidMessage", "The current input state doesn't support cutting.");
 		return false;
@@ -272,7 +265,7 @@ void UNiagaraStackFunctionInput::RemoveForCut()
 
 bool UNiagaraStackFunctionInput::TestCanCopyWithMessage(FText& OutMessage) const
 {
-	if (InputValues.Mode == EValueMode::Invalid)
+	if (InputValues.HasEditableData() == false)
 	{
 		OutMessage = LOCTEXT("CantCopyInvalidMessage", "The current input state doesn't support copying.");
 		return false;
@@ -337,56 +330,61 @@ void UNiagaraStackFunctionInput::Paste(const UNiagaraClipboardContent* Clipboard
 	}
 }
 
-FText UNiagaraStackFunctionInput::GetTooltipText(EValueMode InValueMode) const
+FText UNiagaraStackFunctionInput::GetValueToolTip() const
 {
-	FNiagaraVariable ValueVariable;
-	UNiagaraGraph* NodeGraph = nullptr;
-
-	if (InValueMode == EValueMode::Linked)
+	if (IsFinalized())
 	{
-		UEdGraphPin* OverridePin = GetOverridePin();
-		UEdGraphPin* ValuePin = OverridePin != nullptr ? OverridePin : GetDefaultPin();
-		ValueVariable = FNiagaraVariable(InputType, InputValues.LinkedHandle.GetParameterHandleString());
-		if (ValuePin != nullptr)
-		{
-			NodeGraph = Cast<UNiagaraGraph>(ValuePin->GetOwningNode()->GetGraph());
-		}
+		return FText();
 	}
-	else
+
+	if (ValueToolTipCache.IsSet() == false)
 	{
-		ValueVariable = FNiagaraVariable(InputType, InputParameterHandle.GetParameterHandleString());
-		if (OwningFunctionCallNode.IsValid() && OwningFunctionCallNode->FunctionScript != nullptr)
+		ValueToolTipCache = FText();
+		switch (InputValues.Mode)
 		{
-			UNiagaraScriptSource* Source = Cast<UNiagaraScriptSource>(OwningFunctionCallNode->FunctionScript->GetSource());
-			const UEdGraphSchema_Niagara* NiagaraSchema = GetDefault<UEdGraphSchema_Niagara>();
-			NodeGraph = Source->NodeGraph;
+		case EValueMode::Data:
+		{
+			FString DataInterfaceDescription = InputValues.DataObject->GetClass()->GetDescription();
+			if (DataInterfaceDescription.Len() > 0)
+			{
+				ValueToolTipCache = FText::FromString(DataInterfaceDescription);
+			}
+			break;
 		}
-	}
-
-	TOptional<FNiagaraVariableMetaData> MetaData;
-	if (FNiagaraConstants::IsNiagaraConstant(ValueVariable))
-	{
-		const FNiagaraVariableMetaData* FoundMetaData = FNiagaraConstants::GetConstantMetaData(ValueVariable);
-		if (FoundMetaData)
+		case EValueMode::DefaultFunction:
+			if (InputValues.DefaultFunctionNode->FunctionScript != nullptr)
+			{
+				ValueToolTipCache = InputValues.DefaultFunctionNode->FunctionScript->Description;
+			}
+			break;
+		case EValueMode::Dynamic:
+			if (InputValues.DynamicNode->FunctionScript != nullptr)
+			{
+				ValueToolTipCache = InputValues.DynamicNode->FunctionScript->Description;
+			}
+			break;
+		case EValueMode::InvalidOverride:
+			ValueToolTipCache = LOCTEXT("InvalidOverrideToolTip", "The script is in an invalid and unrecoverable state for this\ninput.  Resetting to default may fix this issue.");
+			break;
+		case EValueMode::Linked:
 		{
-			MetaData = *FoundMetaData;
+			FNiagaraVariable InputVariable(InputType, InputValues.LinkedHandle.GetParameterHandleString());
+			if (FNiagaraConstants::IsNiagaraConstant(InputVariable))
+			{
+				const FNiagaraVariableMetaData* FoundMetaData = FNiagaraConstants::GetConstantMetaData(InputVariable);
+				if (FoundMetaData != nullptr)
+				{
+					ValueToolTipCache = FoundMetaData->Description;
+				}
+			}
+			break;
+		}
+		case EValueMode::UnsupportedDefault:
+			ValueToolTipCache = LOCTEXT("UnsupportedDefault", "The defalut value defined in the script graph\nis custom and can not be shown in the selection stack.");
+			break;
 		}
 	}
-	else if (NodeGraph != nullptr)
-	{
-		MetaData = NodeGraph->GetMetaData(ValueVariable);
-	}
-
-	FText Description = FText::GetEmpty();
-	if (MetaData.IsSet())
-	{
-		Description = MetaData->Description;
-	}
-
-	return FText::Format(LOCTEXT("FunctionInputTooltip", "Name: {0} \nType: {1} \nDesc: {2}"),
-		FText::FromName(ValueVariable.GetName()),
-		ValueVariable.GetType().GetNameText(),
-		Description);
+	return ValueToolTipCache.GetValue();
 }
 
 void UNiagaraStackFunctionInput::RefreshChildrenInternal(const TArray<UNiagaraStackEntry*>& CurrentChildren, TArray<UNiagaraStackEntry*>& NewChildren, TArray<FStackIssue>& NewIssues)
@@ -518,15 +516,15 @@ void UNiagaraStackFunctionInput::RefreshChildrenInternal(const TArray<UNiagaraSt
 		}
 	}
 
-	if (InputValues.Mode == EValueMode::Data && InputValues.DataObjects.GetValueObject() != nullptr)
+	if (InputValues.Mode == EValueMode::Data && InputValues.DataObject.IsValid())
 	{
 		UNiagaraStackObject* ValueObjectEntry = FindCurrentChildOfTypeByPredicate<UNiagaraStackObject>(CurrentChildren,
-			[=](UNiagaraStackObject* CurrentObjectEntry) { return CurrentObjectEntry->GetObject() == InputValues.DataObjects.GetValueObject(); });
+			[=](UNiagaraStackObject* CurrentObjectEntry) { return CurrentObjectEntry->GetObject() == InputValues.DataObject.Get(); });
 
 		if(ValueObjectEntry == nullptr)
 		{
 			ValueObjectEntry = NewObject<UNiagaraStackObject>(this);
-			ValueObjectEntry->Initialize(CreateDefaultChildRequiredData(), InputValues.DataObjects.GetValueObject(), GetOwnerStackItemEditorDataKey(), OwningFunctionCallNode.Get());
+			ValueObjectEntry->Initialize(CreateDefaultChildRequiredData(), InputValues.DataObject.Get(), GetOwnerStackItemEditorDataKey(), OwningFunctionCallNode.Get());
 		}
 		NewChildren.Add(ValueObjectEntry);
 	}
@@ -666,152 +664,121 @@ FString UNiagaraStackFunctionInput::ResolveDisplayNameArgument(const FString& In
 	return FString();
 }
 
-
-TSharedPtr<FStructOnScope> UNiagaraStackFunctionInput::FInputValues::GetLocalStructToReuse()
-{
-	return Mode == EValueMode::Local ? LocalStruct : TSharedPtr<FStructOnScope>();
-}
-
-UNiagaraDataInterface* UNiagaraStackFunctionInput::FInputValues::GetDataDefaultValueObjectToReuse()
-{
-	return Mode == EValueMode::Data && DataObjects.IsValid() && DataObjects.GetDefaultValueOwner() == FDataValues::EDefaultValueOwner::LocallyOwned
-		? DataObjects.GetDefaultValueObject()
-		: nullptr;
-}
-
-bool UNiagaraStackFunctionInput::TryGetDefaultBinding(FNiagaraParameterHandle& LinkedValueHandle, UNiagaraScriptVariable* InVariable, UEdGraphPin& ValuePin)
-{
-	if (!InVariable)
-	{
-		return false;
-	}
-
-	if (InVariable->DefaultMode != ENiagaraDefaultMode::Binding || !InVariable->DefaultBinding.IsValid())
-	{
-		return false;
-	}
-
-	LinkedValueHandle = FNiagaraParameterHandle(InVariable->DefaultBinding.GetName());
-	return true;
-}
-
-void UNiagaraStackFunctionInput::RefreshValues(bool bFromSetLocalValue)
+void UNiagaraStackFunctionInput::RefreshValues()
 {
 	if (ensureMsgf(IsStaticParameter() || InputParameterHandle.IsModuleHandle(), TEXT("Function inputs can only be generated for module paramters.")) == false)
 	{
 		return;
 	}
 
+	// First collect the default values which are used to figure out if an input can be reset, and are used to
+	// determine the current displayed value.
+	DefaultInputValues = FInputValues();
+	UpdateValuesFromScriptDefaults(DefaultInputValues);
+
 	FInputValues OldValues = InputValues;
 	InputValues = FInputValues();
 
-	UEdGraphPin* DefaultPin = GetDefaultPin();
-	if (DefaultPin != nullptr)
+	// If there is an override pin available it's value will take precedence so check that first.
+	UEdGraphPin* OverridePin = GetOverridePin();
+	if (OverridePin != nullptr)
 	{
-		UEdGraphPin* OverridePin = GetOverridePin();
-		UEdGraphPin* ValuePin = OverridePin != nullptr ? OverridePin : DefaultPin;
-
-		if (UNiagaraGraph* FunctionGraph = Cast<UNiagaraScriptSource>(OwningFunctionCallNode->FunctionScript->GetSource())->NodeGraph)
-		{
-			Variable = FunctionGraph->GetScriptVariable(*(TEXT("Module.") + DisplayName.ToString()));
-		}
-
-		if (TryGetCurrentDataValue(InputValues.DataObjects, OverridePin, *DefaultPin, OldValues.GetDataDefaultValueObjectToReuse()))
-		{
-			InputValues.Mode = EValueMode::Data;
-			bIsLocalOverride = false;
-		}
-		else if (TryGetCurrentExpressionValue(InputValues.ExpressionNode, OverridePin))
-		{
-			InputValues.Mode = EValueMode::Expression;
-			bIsLocalOverride = false;
-		}
-		else if (TryGetCurrentDynamicValue(InputValues.DynamicNode, OverridePin))
-		{
-			InputValues.Mode = EValueMode::Dynamic;
-			bIsLocalOverride = false;
-		}
-		else if (TryGetCurrentLinkedValue(InputValues.LinkedHandle, *ValuePin))
+		UpdateValuesFromOverridePin(OldValues, InputValues, *OverridePin);
+	}
+	else
+	{
+		if(InputType.IsDataInterface())
 		{
-			InputValues.Mode = EValueMode::Linked;
-			bIsLocalOverride = false;
+			// Data interfaces must always be overridden in the stack, if there wasn't an override pin found set the mode to invalid override since the 
+			// stack graph state isn't valid.
+			InputValues.Mode = EValueMode::InvalidOverride;
 		}
-		else if (TryGetCurrentLocalValue(InputValues.LocalStruct, *DefaultPin, *ValuePin, OldValues.GetLocalStructToReuse(), Variable))
+		else if (IsRapidIterationCandidate())
 		{
-			if (OverridePin)
+			// If the value is a rapid iteration parameter it's a local value so copy it's value from the rapid iteration parameter store if it's in there,
+			// otherwise copy the value from the default.
+			InputValues.Mode = EValueMode::Local;
+			InputValues.LocalStruct = MakeShared<FStructOnScope>(InputType.GetStruct());
+			const uint8* RapidIterationParameterData = SourceScript->RapidIterationParameters.GetParameterData(RapidIterationParameter);
+			if (RapidIterationParameterData == nullptr)
 			{
-				InputValues.Mode = EValueMode::Local;
-			}
-			else
-			{
-				if (!bIsLocalOverride && TryGetDefaultBinding(InputValues.LinkedHandle, Variable, *ValuePin))
-				{
-					InputValues.Mode = EValueMode::Linked;
-					bIsLocalOverride = false; 
-				}
-				else
-				{
-					InputValues.Mode = EValueMode::Local;
-				}
+				RapidIterationParameterData = DefaultInputValues.LocalStruct->GetStructMemory();
 			}
+			FMemory::Memcpy(InputValues.LocalStruct->GetStructMemory(), RapidIterationParameterData, InputType.GetSize());
 		}
-		else if (TryGetDefaultBinding(InputValues.LinkedHandle, Variable, *ValuePin))
+		else
 		{
-			InputValues.Mode = EValueMode::Linked;
-			bIsLocalOverride = false;
+			// Otherwise if there isn't an override pin and it's not a rapid iteration parameter use the default value.
+			InputValues = DefaultInputValues;
 		}
 	}
 
 	bCanResetCache.Reset();
 	bCanResetToBaseCache.Reset();
+	ValueToolTipCache.Reset();
 	ValueChangedDelegate.Broadcast();
 }
 
 void UNiagaraStackFunctionInput::RefreshFromMetaData()
 {
-	if (OwningFunctionCallNode->FunctionScript != nullptr)
+	InputMetaData.Reset();
+	if (OwningFunctionCallNode->IsA<UNiagaraNodeAssignment>())
+	{
+		// Set variables nodes have no metadata, but if they're setting a defined constant see if there's metadata for that.
+		FNiagaraVariable InputVariable(InputType, InputParameterHandle.GetName());
+		if (FNiagaraConstants::IsNiagaraConstant(InputVariable))
+		{
+			const FNiagaraVariableMetaData* FoundMetaData = FNiagaraConstants::GetConstantMetaData(InputVariable);
+			if (FoundMetaData)
+			{
+				InputMetaData = *FoundMetaData;
+			}
+		}
+	}
+	else if (OwningFunctionCallNode->FunctionScript != nullptr)
 	{
+		// Otherwise just get it from the defining graph.
 		UNiagaraGraph* FunctionGraph = CastChecked<UNiagaraScriptSource>(OwningFunctionCallNode->FunctionScript->GetSource())->NodeGraph;
 		FNiagaraVariable InputVariable(InputType, InputParameterHandle.GetParameterHandleString());
 		InputMetaData = FunctionGraph->GetMetaData(InputVariable);
+	}
 
-		if (InputMetaData.IsSet())
-		{
-			SetIsAdvanced(InputMetaData->bAdvancedDisplay);
-
-			FText EditConditionError;
-			EditCondition.Refresh(InputMetaData->EditCondition, EditConditionError);
-			if (EditCondition.IsValid() && EditCondition.GetConditionInputType() == FNiagaraTypeDefinition::GetBoolDef())
-			{
-				TOptional<FNiagaraVariableMetaData> EditConditionInputMetadata = EditCondition.GetConditionInputMetaData();
-				if (EditConditionInputMetadata.IsSet())
-				{
-					bShowEditConditionInline = EditConditionInputMetadata->bInlineEditConditionToggle;
-				}
-			}
-			else
-			{
-				bShowEditConditionInline = false;
-			}
+	if (InputMetaData.IsSet())
+	{
+		SetIsAdvanced(InputMetaData->bAdvancedDisplay);
 
-			if (EditConditionError.IsEmpty() == false)
+		FText EditConditionError;
+		EditCondition.Refresh(InputMetaData->EditCondition, EditConditionError);
+		if (EditCondition.IsValid() && EditCondition.GetConditionInputType() == FNiagaraTypeDefinition::GetBoolDef())
+		{
+			TOptional<FNiagaraVariableMetaData> EditConditionInputMetadata = EditCondition.GetConditionInputMetaData();
+			if (EditConditionInputMetadata.IsSet())
 			{
-				UE_LOG(LogNiagaraEditor, Warning, TEXT("Edit condition failed to bind.  Function: %s Input: %s Message: %s"), 
-					*OwningFunctionCallNode->GetFunctionName(), *InputParameterHandle.GetName().ToString(), *EditConditionError.ToString());
+				bShowEditConditionInline = EditConditionInputMetadata->bInlineEditConditionToggle;
 			}
+		}
+		else
+		{
+			bShowEditConditionInline = false;
+		}
 
-			FText VisibleConditionError;
-			VisibleCondition.Refresh(InputMetaData->VisibleCondition, VisibleConditionError);
+		if (EditConditionError.IsEmpty() == false)
+		{
+			UE_LOG(LogNiagaraEditor, Warning, TEXT("Edit condition failed to bind.  Function: %s Input: %s Message: %s"), 
+				*OwningFunctionCallNode->GetFunctionName(), *InputParameterHandle.GetName().ToString(), *EditConditionError.ToString());
+		}
 
-			if (VisibleConditionError.IsEmpty() == false)
-			{
-				UE_LOG(LogNiagaraEditor, Warning, TEXT("Visible condition failed to bind.  Function: %s Input: %s Message: %s"),
-					*OwningFunctionCallNode->GetFunctionName(), *InputParameterHandle.GetName().ToString(), *VisibleConditionError.ToString());
-			}
+		FText VisibleConditionError;
+		VisibleCondition.Refresh(InputMetaData->VisibleCondition, VisibleConditionError);
 
-			bIsInlineEditConditionToggle = InputType == FNiagaraTypeDefinition::GetBoolDef() && 
-				InputMetaData->bInlineEditConditionToggle;
+		if (VisibleConditionError.IsEmpty() == false)
+		{
+			UE_LOG(LogNiagaraEditor, Warning, TEXT("Visible condition failed to bind.  Function: %s Input: %s Message: %s"),
+				*OwningFunctionCallNode->GetFunctionName(), *InputParameterHandle.GetName().ToString(), *VisibleConditionError.ToString());
 		}
+
+		bIsInlineEditConditionToggle = InputType == FNiagaraTypeDefinition::GetBoolDef() && 
+			InputMetaData->bInlineEditConditionToggle;
 	}
 }
 
@@ -843,16 +810,20 @@ void UNiagaraStackFunctionInput::SetLinkedValueHandle(const FNiagaraParameterHan
 	}
 
 	FScopedTransaction ScopedTransaction(LOCTEXT("UpdateLinkedInputValue", "Update linked input value"));
-	UEdGraphPin& OverridePin = GetOrCreateOverridePin();
-	RemoveNodesForOverridePin(OverridePin);
+	RemoveOverridePin();
+
 	if (IsRapidIterationCandidate())
 	{
 		RemoveRapidIterationParametersForAffectedScripts();
 	}
 
-	FNiagaraStackGraphUtilities::SetLinkedValueHandleForFunctionInput(OverridePin, InParameterHandle);
-	FNiagaraStackGraphUtilities::RelayoutGraph(*OwningFunctionCallNode->GetGraph());
+	if (InParameterHandle != DefaultInputValues.LinkedHandle)
+	{
+		// Only set the linked value if it's actually different from the default.
+		FNiagaraStackGraphUtilities::SetLinkedValueHandleForFunctionInput(GetOrCreateOverridePin(), InParameterHandle);
+	}
 
+	FNiagaraStackGraphUtilities::RelayoutGraph(*OwningFunctionCallNode->GetGraph());
 	RefreshValues();
 }
 
@@ -912,7 +883,7 @@ FName GetNamespaceForUsage(ENiagaraScriptUsage Usage)
 	}
 }
 
-void UNiagaraStackFunctionInput::GetAvailableParameterHandles(TArray<FNiagaraParameterHandle>& AvailableParameterHandles)
+void UNiagaraStackFunctionInput::GetAvailableParameterHandles(TArray<FNiagaraParameterHandle>& AvailableParameterHandles) const
 {
 	// Engine Handles.
 	for (const FNiagaraVariable& SystemVariable : FNiagaraConstants::GetEngineConstants())
@@ -1041,6 +1012,11 @@ void UNiagaraStackFunctionInput::GetAvailableParameterHandles(TArray<FNiagaraPar
 	}
 }
 
+UNiagaraNodeFunctionCall* UNiagaraStackFunctionInput::GetDefaultFunctionNode() const
+{
+	return InputValues.DefaultFunctionNode.Get();
+}
+
 UNiagaraNodeFunctionCall* UNiagaraStackFunctionInput::GetDynamicInputNode() const
 {
 	return InputValues.DynamicNode.Get();
@@ -1123,14 +1099,14 @@ void UNiagaraStackFunctionInput::SetCustomExpression(const FString& InCustomExpr
 	RefreshChildren();
 }
 
-TSharedPtr<FStructOnScope> UNiagaraStackFunctionInput::GetLocalValueStruct()
+TSharedPtr<const FStructOnScope> UNiagaraStackFunctionInput::GetLocalValueStruct()
 {
 	return InputValues.LocalStruct;
 }
 
 UNiagaraDataInterface* UNiagaraStackFunctionInput::GetDataValueObject()
 {
-	return InputValues.DataObjects.GetValueObject();
+	return InputValues.DataObject.Get();
 }
 
 void UNiagaraStackFunctionInput::NotifyBeginLocalValueChange()
@@ -1148,99 +1124,63 @@ void UNiagaraStackFunctionInput::NotifyEndLocalValueChange()
 
 bool UNiagaraStackFunctionInput::IsRapidIterationCandidate() const
 {
-	return !IsStaticParameter() && FNiagaraStackGraphUtilities::IsRapidIterationType(InputType);
+	// Rapid iteration parameters will only be used if the input is not static and the input value default is a
+	// local value, if it's linked in graph or through metadata or a default dynamic input the compiler generates
+	// code for that instead.
+	return !IsStaticParameter() && FNiagaraStackGraphUtilities::IsRapidIterationType(InputType) && DefaultInputValues.Mode == EValueMode::Local;
 }
 
-void UNiagaraStackFunctionInput::SetLocalValue(TSharedRef<FStructOnScope> InLocalValue, bool bIsOverride)
+void UNiagaraStackFunctionInput::SetLocalValue(TSharedRef<FStructOnScope> InLocalValue)
 {
-	if (bIsOverride)
-	{
-		bIsLocalOverride = bIsOverride;
-	}
-	TGuardValue<bool> UpdateGuard(bUpdatingLocalValueDirectly, true);
-
-	const UEdGraphSchema_Niagara* NiagaraSchema = GetDefault<UEdGraphSchema_Niagara>();
-	UEdGraphPin* DefaultPin = GetDefaultPin();
-	UEdGraphPin* OverridePin = GetOverridePin();
-	UEdGraphPin* ValuePin = DefaultPin;
-
-	// If the parameter we set is from a static switch then we don't want to use an override pin, but directly change the default pin.
-	if (IsStaticParameter())
-	{
-		FNiagaraVariable LocalValueVariable(InputType, NAME_None);
-		LocalValueVariable.SetData(InLocalValue->GetStructMemory());
-		FString PinDefaultValue;
-		if (ensureMsgf(NiagaraSchema->TryGetPinDefaultValueFromNiagaraVariable(LocalValueVariable, PinDefaultValue),
-			TEXT("Could not generate value string for static switch parameter.")))
-		{
-			DefaultPin->Modify();
-			DefaultPin->DefaultValue = PinDefaultValue;
-			Cast<UNiagaraNode>(DefaultPin->GetOwningNode())->MarkNodeRequiresSynchronization(TEXT("Default Value Changed"), true);
-		}
-		RefreshValues(true);
-		return;
-	}
-	
-	// If the default pin in the function graph is connected internally, rapid iteration parameters can't be used since
-	// the compilation currently won't use them.
-	bool bCanUseRapidIterationParameter = IsRapidIterationCandidate() && DefaultPin->LinkedTo.Num() == 0;
-	if (Variable && Variable->DefaultMode == ENiagaraDefaultMode::Binding)
-	{
-		bCanUseRapidIterationParameter = false;
-	}
+	checkf(InLocalValue->GetStruct() == InputType.GetStruct(), TEXT("Can not set an input to an unrelated type."));
 
-	if (bCanUseRapidIterationParameter == false)
-	{
-		ValuePin = OverridePin != nullptr ? OverridePin : DefaultPin;
-	}
-
-	TSharedPtr<FStructOnScope> CurrentValue;
-	bool bCanHaveLocalValue = ValuePin != nullptr;
-	bool bHasLocalValue = bCanHaveLocalValue && InputValues.Mode == EValueMode::Local && TryGetCurrentLocalValue(CurrentValue, *DefaultPin, *ValuePin, TSharedPtr<FStructOnScope>(), nullptr);
-	bool bLocalValueMatchesSetValue = bHasLocalValue && FNiagaraEditorUtilities::DataMatches(*CurrentValue.Get(), InLocalValue.Get());
-
-	if (bCanHaveLocalValue == false || bLocalValueMatchesSetValue)
+	if (InputValues.Mode == EValueMode::Local && FNiagaraEditorUtilities::DataMatches(*InputValues.LocalStruct.Get(), InLocalValue.Get()))
 	{
+		// The value matches the current value so noop.
 		return;
 	}
 
+	TGuardValue<bool> UpdateGuard(bUpdatingLocalValueDirectly, true);
 	FScopedTransaction ScopedTransaction(LOCTEXT("UpdateInputLocalValue", "Update input local value"));
-	UNiagaraGraph* EmitterGraph = Cast<UNiagaraGraph>(OwningFunctionCallNode->GetGraph());
-
 	bool bGraphWillNeedRelayout = false;
+	UEdGraphPin* OverridePin = GetOverridePin();
+
 	if (OverridePin != nullptr && OverridePin->LinkedTo.Num() > 0)
 	{
+		// If there is an override pin and it's linked we'll need to remove all of the linked nodes to set a local value.
 		RemoveNodesForOverridePin(*OverridePin);
 		bGraphWillNeedRelayout = true;
 	}
 
-	if (bCanUseRapidIterationParameter)
+	if (IsRapidIterationCandidate())
 	{
-		for (TWeakObjectPtr<UNiagaraScript> Script : AffectedScripts)
-		{
-			Script->Modify();
-		}
-
-		// If there is currently an override, we need to get rid of it.
+		// If there is currently an override pin, it must be removed to allow the rapid iteration parameter to be used.
 		if (OverridePin != nullptr)
 		{
 			UNiagaraNode* OverrideNode = CastChecked<UNiagaraNode>(OverridePin->GetOwningNode());
 			OverrideNode->Modify();
 			OverrideNode->RemovePin(OverridePin);
+			bGraphWillNeedRelayout = true;
 		}
 
+		// Update the value on all affected scripts.
 		for (TWeakObjectPtr<UNiagaraScript> Script : AffectedScripts)
 		{
 			bool bAddParameterIfMissing = true;
+			Script->Modify();
 			Script->RapidIterationParameters.SetParameterData(InLocalValue->GetStructMemory(), RapidIterationParameter, bAddParameterIfMissing);
 		}
 	}
-	else 
+	else
 	{
+		// If rapid iteration parameters can't be used the string representation of the value needs to be set on the override pin
+		// for this input.  For static switch inputs the override pin in on the owning function call node and for standard parameter
+		// pins the override pin is on the override parameter map set node.
 		FNiagaraVariable LocalValueVariable(InputType, NAME_None);
 		LocalValueVariable.SetData(InLocalValue->GetStructMemory());
 		FString PinDefaultValue;
-		if(ensureMsgf(NiagaraSchema->TryGetPinDefaultValueFromNiagaraVariable(LocalValueVariable, PinDefaultValue),
+		const UEdGraphSchema_Niagara* NiagaraSchema = GetDefault<UEdGraphSchema_Niagara>();
+		if (ensureMsgf(NiagaraSchema->TryGetPinDefaultValueFromNiagaraVariable(LocalValueVariable, PinDefaultValue),
 			TEXT("Could not generate default value string for non-rapid iteration parameter.")))
 		{
 			if (OverridePin == nullptr)
@@ -1257,103 +1197,48 @@ void UNiagaraStackFunctionInput::SetLocalValue(TSharedRef<FStructOnScope> InLoca
 	
 	if (bGraphWillNeedRelayout)
 	{
-		FNiagaraStackGraphUtilities::RelayoutGraph(*EmitterGraph);
+		FNiagaraStackGraphUtilities::RelayoutGraph(*OwningFunctionCallNode->GetNiagaraGraph());
 	}
 
-	RefreshValues(true);
+	RefreshValues();
 }
 
 bool UNiagaraStackFunctionInput::CanReset() const
 {
 	if (bCanResetCache.IsSet() == false)
 	{
-		bool bNewCanReset;
-		if (InputValues.Mode == EValueMode::Data)
+		bool bNewCanReset = false;
+		if (DefaultInputValues.Mode == EValueMode::None)
+		{
+			// Can't reset if no valid default was set.
+			bNewCanReset = false;
+		}
+		else if (InputValues.Mode != DefaultInputValues.Mode)
 		{
-			// For data values a copy of the default object should have been created automatically and attached to the override pin for this input.  If a 
-			// copy of the default object wasn't created, the input can be reset to create one.  If a copy of the data object is available it can be
-			// reset if it's different from it's default value.
-			bool bHasDataValueObject = InputValues.DataObjects.GetValueObject() != nullptr;
-			bool bHasDefaultDataValueObject = InputValues.DataObjects.GetDefaultValueObject() != nullptr;
-			bool bIsDataValueDifferentFromDefaultDataValue = bHasDataValueObject && bHasDefaultDataValueObject
-				&& InputValues.DataObjects.GetValueObject()->Equals(InputValues.DataObjects.GetDefaultValueObject()) == false;
-			bNewCanReset = bHasDataValueObject == false || bHasDefaultDataValueObject == false || bIsDataValueDifferentFromDefaultDataValue;
+			// If the current value mode is different from the default value mode, it can always be reset.
+			bNewCanReset = true;
 		}
 		else
 		{
-			UEdGraphPin* DefaultPin = GetDefaultPin();
-			if (ensure(DefaultPin != nullptr))
+			switch (InputValues.Mode)
 			{
-				if (IsStaticParameter())
-				{
-					bNewCanReset = true;
-
-					// Static switch parameters only hold a single value in the default pin, so we compare it to the default value in the switch metadata
-					UNiagaraGraph* FunctionGraph = CastChecked<UNiagaraScriptSource>(OwningFunctionCallNode->FunctionScript->GetSource())->NodeGraph;
-					const UEdGraphSchema_Niagara* Schema = CastChecked<UEdGraphSchema_Niagara>(OwningFunctionCallNode->GetSchema());
-
-					for (const FNiagaraVariable& SwitchVar : FunctionGraph->FindStaticSwitchInputs())
-					{
-						FEdGraphPinType PinType = Schema->TypeDefinitionToPinType(SwitchVar.GetType());
-						if (DefaultPin->PinName.IsEqual(SwitchVar.GetName()) && DefaultPin->PinType == PinType)
-						{
-							TOptional<FNiagaraVariableMetaData> MetaData = FunctionGraph->GetMetaData(SwitchVar);
-
-							int32 DefaultValue = MetaData.IsSet() ? MetaData->StaticSwitchDefaultValue  : 0;
-							int32 CurrentPinValue = 0;
-							if (FNiagaraEditorUtilities::ResolveConstantValue(DefaultPin, CurrentPinValue))
-							{
-								bNewCanReset = DefaultValue != CurrentPinValue;
-							}
-							else
-							{
-								bNewCanReset = false;
-							}
-							break;
-						}
-					}
-				}
-				else if (DefaultPin->LinkedTo.Num() == 0)
-				{
-					if (UEdGraphPin* OverridePin = GetOverridePin())
-					{
-						bNewCanReset = true;
-						UNiagaraGraph* FunctionGraph = CastChecked<UNiagaraScriptSource>(OwningFunctionCallNode->FunctionScript->GetSource())->NodeGraph;
-						if (FunctionGraph && OverridePin && OverridePin->LinkedTo.Num() == 1)
-						{
-							UNiagaraScriptVariable* ScriptVariable = FunctionGraph->GetScriptVariable(*(TEXT("Module.") + DisplayName.ToString()));
-							if (ScriptVariable && OverridePin->LinkedTo[0]->PinName == ScriptVariable->DefaultBinding.GetName())
-							{
-								bNewCanReset = false;
-							}
-						}
-					}
-					else if (IsRapidIterationCandidate())
-					{
-						FNiagaraVariable DefaultVar = GetDefaultVariableForRapidIterationParameter();
-						bool bHasValidLocalValue = InputValues.LocalStruct.IsValid();
-						bool bHasValidDefaultValue = DefaultVar.IsValid();
-						bNewCanReset = bHasValidLocalValue && bHasValidDefaultValue && FNiagaraEditorUtilities::DataMatches(DefaultVar, *InputValues.LocalStruct.Get()) == false;
-					}
-					else
-					{
-						bNewCanReset = false;
-					}
-				}
-				else
-				{
-					if (FNiagaraStackGraphUtilities::IsValidDefaultDynamicInput(*SourceScript, *DefaultPin))
-					{
-						UEdGraphPin* OverridePin = GetOverridePin();
-						FString UniqueEmitterName = GetEmitterViewModel().IsValid() ? GetEmitterViewModel()->GetEmitter()->GetUniqueEmitterName() : FString();
-						bNewCanReset = OverridePin == nullptr || FNiagaraStackGraphUtilities::DoesDynamicInputMatchDefault(UniqueEmitterName, *SourceScript,
-							*OwningFunctionCallNode, *OverridePin, InputParameterHandle.GetName(), *DefaultPin) == false;
-					}
-					else
-					{
-						bNewCanReset = GetOverridePin() != nullptr;
-					}
-				}
+			case EValueMode::Data:
+				bNewCanReset = InputValues.DataObject->Equals(DefaultInputValues.DataObject.Get()) == false;
+				break;
+			case EValueMode::Dynamic:
+				// For now assume that default dynamic inputs can always be reset to default since they're not currently supported properly.
+				// TODO: Fix when default dynamic inputs are properly supported.
+				bNewCanReset = false;
+				break;
+			case EValueMode::Linked:
+				bNewCanReset = InputValues.LinkedHandle != DefaultInputValues.LinkedHandle;
+				break;
+			case EValueMode::Local:
+				bNewCanReset =
+					DefaultInputValues.LocalStruct.IsValid() &&
+					InputValues.LocalStruct->GetStruct() == DefaultInputValues.LocalStruct->GetStruct() &&
+					FMemory::Memcmp(InputValues.LocalStruct->GetStructMemory(), DefaultInputValues.LocalStruct->GetStructMemory(), InputType.GetSize());
+				break;
 			}
 		}
 		bCanResetCache = bNewCanReset;
@@ -1361,19 +1246,6 @@ bool UNiagaraStackFunctionInput::CanReset() const
 	return bCanResetCache.GetValue();
 }
 
-FNiagaraVariable UNiagaraStackFunctionInput::GetDefaultVariableForRapidIterationParameter() const
-{
-	FNiagaraVariable Var;
-	UEdGraphPin* DefaultPin = GetDefaultPin();
-	if (DefaultPin != nullptr)
-	{
-		const UEdGraphSchema_Niagara* NiagaraSchema = GetDefault<UEdGraphSchema_Niagara>();
-		Var = NiagaraSchema->PinToNiagaraVariable(DefaultPin, true);
-		Var.SetName(*RapidIterationParameter.GetName().ToString());
-	}
-	return Var;
-}
-
 bool UNiagaraStackFunctionInput::UpdateRapidIterationParametersForAffectedScripts(const uint8* Data)
 {
 	for (TWeakObjectPtr<UNiagaraScript> Script : AffectedScripts)
@@ -1409,118 +1281,55 @@ bool UNiagaraStackFunctionInput::RemoveRapidIterationParametersForAffectedScript
 
 void UNiagaraStackFunctionInput::Reset()
 {
-	bIsLocalOverride = false;
-	if(InputValues.Mode == EValueMode::Data)
+	if (CanReset())
 	{
-		// For data values they are reset by making sure the data object owned by this input matches the default
-		// data object.  If there is no data object owned by the input, one is created and and updated to match the default.
-		FScopedTransaction ScopedTransaction(LOCTEXT("ResetInputObjectTransaction", "Reset the inputs data interface object to default."));
-		if (InputValues.DataObjects.GetValueObject() != nullptr && InputValues.DataObjects.GetDefaultValueObject() != nullptr)
-		{
-			InputValues.DataObjects.GetDefaultValueObject()->CopyTo(InputValues.DataObjects.GetValueObject());
-		}
-		else
+		if (DefaultInputValues.Mode == EValueMode::Data)
 		{
-			UEdGraphPin& OverridePin = GetOrCreateOverridePin();
-			RemoveNodesForOverridePin(OverridePin);
-
-			FString InputNodeName = InputParameterHandlePath[0].GetName().ToString();
-			for (int32 i = 1; i < InputParameterHandlePath.Num(); i++)
-			{
-				InputNodeName += "." + InputParameterHandlePath[i].GetName().ToString();
-			}
-
-			UNiagaraDataInterface* InputValueObject;
-			FNiagaraStackGraphUtilities::SetDataValueObjectForFunctionInput(OverridePin, const_cast<UClass*>(InputType.GetClass()), InputNodeName, InputValueObject);
-			if (InputValues.DataObjects.GetDefaultValueObject() != nullptr)
+			FScopedTransaction ScopedTransaction(LOCTEXT("ResetInputObjectTransaction", "Reset the inputs data interface object to default."));
+			if (InputValues.Mode == EValueMode::Data)
 			{
-				InputValues.DataObjects.GetDefaultValueObject()->CopyTo(InputValueObject);
+				// If there is already a valid data object just copy from the default to the current value.
+				InputValues.DataObject->Modify();
+				DefaultInputValues.DataObject->CopyTo(InputValues.DataObject.Get());
 			}
-
-			FNiagaraStackGraphUtilities::RelayoutGraph(*OwningFunctionCallNode->GetGraph());
-		}
-	}
-	else if (IsStaticParameter())
-	{
-		// Static switch parameters only hold a single value in the default pin, so we directly set it to the default value from the switch metadata
-		UNiagaraGraph* FunctionGraph = CastChecked<UNiagaraScriptSource>(OwningFunctionCallNode->FunctionScript->GetSource())->NodeGraph;
-		const UEdGraphSchema_Niagara* Schema = CastChecked<UEdGraphSchema_Niagara>(OwningFunctionCallNode->GetSchema());
-		UEdGraphPin* DefaultPin = GetDefaultPin();
-
-		for (FNiagaraVariable SwitchVar : FunctionGraph->FindStaticSwitchInputs())
-		{
-			FEdGraphPinType PinType = Schema->TypeDefinitionToPinType(SwitchVar.GetType());
-			if (DefaultPin->PinName.IsEqual(SwitchVar.GetName()) && DefaultPin->PinType == PinType)
+			else
 			{
-				TOptional<FNiagaraVariableMetaData> MetaData = FunctionGraph->GetMetaData(SwitchVar);
+				// Otherwise remove the current nodes from the override pin and set a new data object and copy the values from the default.
+				UEdGraphPin& OverridePin = GetOrCreateOverridePin();
+				RemoveNodesForOverridePin(OverridePin);
 
-				int32 DefaultValue = MetaData.IsSet() ? MetaData->StaticSwitchDefaultValue : 0;
-				SwitchVar.SetValue(DefaultValue);
-				FString PinDefaultValue;
-				if (ensureMsgf(Schema->TryGetPinDefaultValueFromNiagaraVariable(SwitchVar, PinDefaultValue), TEXT("Could not generate value string to reset static switch parameter.")))
+				FString InputNodeName = InputParameterHandlePath[0].GetName().ToString();
+				for (int32 i = 1; i < InputParameterHandlePath.Num(); i++)
 				{
-					DefaultPin->DefaultValue = PinDefaultValue;
+					InputNodeName += "." + InputParameterHandlePath[i].GetName().ToString();
 				}
 
-				OwningFunctionCallNode->GetNiagaraGraph()->NotifyGraphNeedsRecompile();
+				UNiagaraDataInterface* InputValueObject;
+				FNiagaraStackGraphUtilities::SetDataValueObjectForFunctionInput(OverridePin, const_cast<UClass*>(InputType.GetClass()), InputNodeName, InputValueObject);
+				DefaultInputValues.DataObject->CopyTo(InputValueObject);
+
 				FNiagaraStackGraphUtilities::RelayoutGraph(*OwningFunctionCallNode->GetGraph());
-				break;
 			}
 		}
-	}
-	else
-	{
-		// For all other value modes removing the nodes connected to the override pin resets them.
-		UNiagaraNodeParameterMapSet* OverrideNode = GetOverrideNode();
-		UEdGraphPin* OverridePin = GetOverridePin();
-		UEdGraphPin* DefaultPin = GetDefaultPin();
-		bool bGraphNeedsRecompile = false;
+		else if (DefaultInputValues.Mode == EValueMode::Linked)
 		{
-			FScopedTransaction ScopedTransaction(LOCTEXT("ResetInputStructTransaction", "Reset the inputs value to default."));
-			
-			if(DefaultPin->LinkedTo.Num() == 0 && IsRapidIterationCandidate())
-			{
-				if (OverrideNode != nullptr && OverridePin != nullptr)
-				{
-					RemoveNodesForOverridePin(*OverridePin);
-					OverrideNode->Modify();
-					OverrideNode->RemovePin(OverridePin);
-					bGraphNeedsRecompile = true;
-				}
-
-				// Get the default value of the graph pin and use that to reset the rapid iteration variables...
-				FNiagaraVariable DefaultVar = GetDefaultVariableForRapidIterationParameter();
-				if (DefaultVar.IsValid())
-				{
-					UpdateRapidIterationParametersForAffectedScripts(DefaultVar.GetData());
-				}
-			}
-			else if(DefaultPin->LinkedTo.Num() == 0 || FNiagaraStackGraphUtilities::IsValidDefaultDynamicInput(*SourceScript, *DefaultPin) == false)
-			{
-				if (ensureMsgf(OverrideNode != nullptr && OverridePin != nullptr, TEXT("Can not reset the value of an input that doesn't have a valid override node and override pin")))
-				{
-					RemoveNodesForOverridePin(*OverridePin);
-					OverrideNode->Modify();
-					OverrideNode->RemovePin(OverridePin);
-					bGraphNeedsRecompile =  true;
-				}
-			}
-			else
-			{
-				if (OverridePin != nullptr)
-				{
-					RemoveNodesForOverridePin(*OverridePin);
-				}
-				FNiagaraStackGraphUtilities::ResetToDefaultDynamicInput(GetSystemViewModel(), GetEmitterViewModel(), GetStackEditorData(),
-					*SourceScript, AffectedScripts, *OwningModuleNode, *OwningFunctionCallNode, InputParameterHandle.GetName(), *DefaultPin);
-				bGraphNeedsRecompile = true;
-			}
-
-			if (bGraphNeedsRecompile)
-			{
-				OwningFunctionCallNode->GetNiagaraGraph()->NotifyGraphNeedsRecompile();
-				FNiagaraStackGraphUtilities::RelayoutGraph(*OwningFunctionCallNode->GetGraph());
-			}
+			FScopedTransaction ScopedTransaction(LOCTEXT("ResetInputLinkedValueTransaction", "Reset the input to its default linked value."));
+			SetLinkedValueHandle(DefaultInputValues.LinkedHandle);
+		}
+		else if (DefaultInputValues.Mode == EValueMode::Local)
+		{
+			FScopedTransaction ScopedTransaction(LOCTEXT("ResetInputLocalValueTransaction", "Reset the input to its default local value."));
+			SetLocalValue(DefaultInputValues.LocalStruct.ToSharedRef());
+		}
+		else if (DefaultInputValues.Mode == EValueMode::DefaultFunction ||
+			DefaultInputValues.Mode == EValueMode::UnsupportedDefault)
+		{
+			FScopedTransaction ScopedTransaction(LOCTEXT("ResetInputValueTransaction", "Reset the input to its default value."));
+			RemoveOverridePin();
+		}
+		else
+		{
+			ensureMsgf(false, TEXT("Attempted to reset a function input to default without a valid default."));
 		}
 	}
 	RefreshChildren();
@@ -1720,16 +1529,10 @@ void UNiagaraStackFunctionInput::DeleteInput()
 	if (UNiagaraNodeAssignment* NodeAssignment = Cast<UNiagaraNodeAssignment>(OwningFunctionCallNode.Get()))
 	{
 		FScopedTransaction ScopedTransaction(LOCTEXT("RemoveInputTransaction", "Remove Input"));
-
-		UEdGraphPin* OverridePin = GetOverridePin();
-		if (OverridePin != nullptr)
-		{
-			// If there is an override pin and connected nodes, remove them before removing the input since removing
-			// the input will prevent us from finding the override pin.
-			RemoveNodesForOverridePin(*OverridePin);
-			UNiagaraNodeParameterMapSet* OverrideNode = GetOverrideNode();
-			OverrideNode->RemovePin(OverridePin);
-		}
+		
+		// If there is an override pin and connected nodes, remove them before removing the input since removing
+		// the input will prevent us from finding the override pin.
+		RemoveOverridePin();
 
 		FNiagaraVariable Var = FNiagaraVariable(GetInputType(), GetInputParameterHandle().GetName());
 		NodeAssignment->Modify();
@@ -1873,7 +1676,7 @@ const UNiagaraClipboardFunctionInput* UNiagaraStackFunctionInput::ToClipboardFun
 		ClipboardInput = UNiagaraClipboardFunctionInput::CreateLinkedValue(InOuter, InputName, InputType, bEditConditionValue, InputValues.LinkedHandle.GetParameterHandleString());
 		break;
 	case EValueMode::Data:
-		ClipboardInput = UNiagaraClipboardFunctionInput::CreateDataValue(InOuter, InputName, InputType, bEditConditionValue, InputValues.DataObjects.GetValueObject());
+		ClipboardInput = UNiagaraClipboardFunctionInput::CreateDataValue(InOuter, InputName, InputType, bEditConditionValue, InputValues.DataObject.Get());
 		break;
 	case EValueMode::Expression:
 		ClipboardInput = UNiagaraClipboardFunctionInput::CreateExpressionValue(InOuter, InputName, InputType, bEditConditionValue, InputValues.ExpressionNode->GetHlslText().ToString());
@@ -1891,7 +1694,9 @@ const UNiagaraClipboardFunctionInput* UNiagaraStackFunctionInput::ToClipboardFun
 
 		break;
 	}
-	case EValueMode::Invalid:
+	case EValueMode::InvalidOverride:
+	case EValueMode::UnsupportedDefault:
+	case EValueMode::DefaultFunction:
 		// Do nothing.
 		break;
 	default:
@@ -1918,8 +1723,18 @@ void UNiagaraStackFunctionInput::SetValueFromClipboardFunctionInput(const UNiaga
 			SetLinkedValueHandle(FNiagaraParameterHandle(ClipboardFunctionInput.Linked));
 			break;
 		case ENiagaraClipboardFunctionInputValueMode::Data:
-			ClipboardFunctionInput.Data->CopyTo(GetDataValueObject());
+		{
+			if (GetDataValueObject() == nullptr)
+			{
+				Reset();
+			}
+			UNiagaraDataInterface* InputDataInterface = GetDataValueObject();
+			if (ensureMsgf(InputDataInterface != nullptr, TEXT("Data interface paste failed.  Current data value object null even after reset.")))
+			{
+				ClipboardFunctionInput.Data->CopyTo(InputDataInterface);
+			}
 			break;
+		}
 		case ENiagaraClipboardFunctionInputValueMode::Expression:
 			SetCustomExpression(ClipboardFunctionInput.Expression);
 			break;
@@ -1976,16 +1791,9 @@ void UNiagaraStackFunctionInput::GetSearchItems(TArray<FStackSearchItem>& Search
 		{
 			SearchItems.Add({ FName("LinkedDynamicInputName"), InputValues.DynamicNode->GetNodeTitle(ENodeTitleType::MenuTitle)});
 		}
-		else if (InputValues.Mode == EValueMode::Data && InputValues.DataObjects.IsValid())
+		else if (InputValues.Mode == EValueMode::Data && InputValues.DataObject.IsValid())
 		{
-			if (InputValues.DataObjects.GetValueObject() != nullptr)
-			{
-				SearchItems.Add({ FName("LinkedDataInterfaceName"), FText::FromString(InputValues.DataObjects.GetValueObject()->GetName()) });
-			}
-			if (InputValues.DataObjects.GetDefaultValueObject() != nullptr)
-			{
-				SearchItems.Add({ FName("LinkedDataInterfaceDefaultName"), InputValues.DataObjects.GetDefaultValueObject()->GetClass()->GetDisplayNameText() });
-			}
+			SearchItems.Add({ FName("LinkedDataInterfaceName"), FText::FromString(InputValues.DataObject->GetName()) });
 		}
 		else if (InputValues.Mode == EValueMode::Expression && InputValues.ExpressionNode.Get() != nullptr)
 		{
@@ -2048,27 +1856,6 @@ UNiagaraNodeParameterMapSet& UNiagaraStackFunctionInput::GetOrCreateOverrideNode
 	return *OverrideNode;
 }
 
-UEdGraphPin* UNiagaraStackFunctionInput::GetDefaultPin() const
-{
-	// If we have a static switch parameter, we check the pins of the function call node for a matching name,
-	// otherwise we search the parameter map node inside the function for a matching pin
-	if (IsStaticParameter())
-	{
-		const UEdGraphSchema_Niagara* NiagaraSchema = CastChecked<UEdGraphSchema_Niagara>(OwningFunctionCallNode->GetSchema());
-		FName InputName = InputParameterHandle.GetParameterHandleString();
-		for (UEdGraphPin* Pin : OwningFunctionCallNode->Pins)
-		{
-			FNiagaraVariable InputVariable = NiagaraSchema->PinToNiagaraVariable(Pin);
-			if (InputVariable.GetName() == InputName && InputVariable.GetType() == InputType)
-			{
-				return Pin;
-			}
-		}
-		return nullptr;
-	}
-	return OwningFunctionCallNode->FindParameterMapDefaultValuePin(InputParameterHandle.GetParameterHandleString(), SourceScript->GetUsage());
-}
-
 UEdGraphPin* UNiagaraStackFunctionInput::GetOverridePin() const
 {
 	if (OverridePinCache.IsSet() == false)
@@ -2090,180 +1877,241 @@ UEdGraphPin& UNiagaraStackFunctionInput::GetOrCreateOverridePin()
 	return *OverridePin;
 }
 
-bool UNiagaraStackFunctionInput::TryGetCurrentLocalValue(TSharedPtr<FStructOnScope>& LocalValue, UEdGraphPin& DefaultPin, UEdGraphPin& ValuePin, TSharedPtr<FStructOnScope> OldValueToReuse, UNiagaraScriptVariable* InVariable)
+void UNiagaraStackFunctionInput::GetDefaultDataInterfaceValueFromDefaultPin(UEdGraphPin* DefaultPin, UNiagaraStackFunctionInput::FInputValues& InInputValues) const
 {
-	if (InputType.IsUObject() == false && ValuePin.LinkedTo.Num() == 0)
+	// Default data interfaces are stored on the input node in the graph, but if it doesn't exist or it's data interface pointer is null, just use the CDO.
+	InInputValues.Mode = EValueMode::Data;
+	if (DefaultPin->LinkedTo.Num() == 1 && DefaultPin->LinkedTo[0]->GetOwningNode() != nullptr && DefaultPin->LinkedTo[0]->GetOwningNode()->IsA<UNiagaraNodeInput>())
 	{
-		const UEdGraphSchema_Niagara* NiagaraSchema = GetDefault<UEdGraphSchema_Niagara>();
-		FNiagaraVariable ValueVariable = NiagaraSchema->PinToNiagaraVariable(&ValuePin, true);
-		if (OldValueToReuse.IsValid() && OldValueToReuse->GetStruct() == ValueVariable.GetType().GetStruct())
+		UNiagaraNodeInput* DataInputNode = CastChecked<UNiagaraNodeInput>(DefaultPin->LinkedTo[0]->GetOwningNode());
+		InInputValues.DataObject = DataInputNode->GetDataInterface();
+	}
+	if (InInputValues.DataObject.IsValid() == false)
+	{
+		InInputValues.DataObject = Cast<UNiagaraDataInterface>(InputType.GetClass()->GetDefaultObject());
+	}
+}
+
+void UNiagaraStackFunctionInput::GetDefaultLocalValueFromDefaultPin(UEdGraphPin* DefaultPin, UNiagaraStackFunctionInput::FInputValues& InInputValues) const
+{
+	// Local default values are stored in the pin's default value string.
+	InInputValues.Mode = EValueMode::Local;
+	const UEdGraphSchema_Niagara* NiagaraSchema = GetDefault<UEdGraphSchema_Niagara>();
+	FNiagaraVariable LocalValueVariable = NiagaraSchema->PinToNiagaraVariable(DefaultPin);
+	if (LocalValueVariable.IsDataAllocated())
+	{
+		InInputValues.LocalStruct = MakeShared<FStructOnScope>(InputType.GetStruct());
+		FMemory::Memcpy(InInputValues.LocalStruct->GetStructMemory(), LocalValueVariable.GetData(), InputType.GetSize());
+	}
+}
+
+struct FLinkedHandleOrFunctionNode
+{
+	TOptional<FNiagaraParameterHandle> LinkedHandle;
+	TWeakObjectPtr<UNiagaraNodeFunctionCall> LinkedFunctionCallNode;
+};
+
+void UNiagaraStackFunctionInput::GetDefaultLinkedHandleOrLinkedFunctionFromDefaultPin(UEdGraphPin* DefaultPin, UNiagaraStackFunctionInput::FInputValues& InInputValues) const
+{
+	// A default pin linked to a parameter map set is a default linked value.  Linked inputs can be setup in a chain and the first
+	// available one will be used so that case must be handled too.  So first collect up the potential linked values and validate 
+	// the linked node structure.
+	TArray<FLinkedHandleOrFunctionNode> LinkedValues;
+	UEdGraphPin* CurrentDefaultPin = DefaultPin;
+	while (CurrentDefaultPin != nullptr && CurrentDefaultPin->LinkedTo.Num() == 1)
+	{
+		UEdGraphPin* LinkedPin = CurrentDefaultPin->LinkedTo[0];
+	
+		if (LinkedPin->GetOwningNode()->IsA<UNiagaraNodeParameterMapGet>())
 		{
-			LocalValue = OldValueToReuse;
+			UNiagaraNodeParameterMapGet* GetNode = CastChecked<UNiagaraNodeParameterMapGet>(LinkedPin->GetOwningNode());
+			FLinkedHandleOrFunctionNode LinkedValue;
+			LinkedValue.LinkedHandle = FNiagaraParameterHandle(*LinkedPin->GetName());
+			LinkedValues.Add(LinkedValue);
+			CurrentDefaultPin = GetNode->GetDefaultPin(LinkedPin);
 		}
-		else
+		else if (LinkedPin->GetOwningNode()->IsA<UNiagaraNodeFunctionCall>())
 		{
-			LocalValue = MakeShared<FStructOnScope>(ValueVariable.GetType().GetStruct());
+			UNiagaraNodeFunctionCall* FunctionCallNode = CastChecked<UNiagaraNodeFunctionCall>(LinkedPin->GetOwningNode());
+			FLinkedHandleOrFunctionNode LinkedValue;
+			LinkedValue.LinkedFunctionCallNode = FunctionCallNode;
+			LinkedValues.Add(LinkedValue);
+			CurrentDefaultPin = nullptr;
 		}
-
-		// If the default pin in the function graph is connected internally, rapid iteration parameters can't be used since
-		// the compilation currently won't use them.
-		bool bCanUseRapidIterationParameter = IsRapidIterationCandidate() && DefaultPin.LinkedTo.Num() == 0;
-		bool bFoundRapidIterationParameter = false;
-		if (InVariable && InVariable->DefaultMode == ENiagaraDefaultMode::Binding)
+		else
 		{
-			bCanUseRapidIterationParameter = false;
+			// Only parameter map get nodes and function calls are valid for a default linked input chain so clear the linked
+			// handles and stop searching.
+			CurrentDefaultPin = nullptr;
+			LinkedValues.Empty();
 		}
-		if (bCanUseRapidIterationParameter)
+	}
+
+	FLinkedHandleOrFunctionNode* ValueToUse = nullptr;
+	if (LinkedValues.Num() == 1)
+	{
+		ValueToUse = &LinkedValues[0];
+	}
+	else if(LinkedValues.Num() > 1)
+	{
+		// If there are a chain of linked values use the first one that's available, otherwise just use the last one.
+		TArray<FNiagaraParameterHandle> AvailableHandles;
+		GetAvailableParameterHandles(AvailableHandles);
+		for (FLinkedHandleOrFunctionNode& LinkedValue : LinkedValues)
 		{
-			const uint8* RapidIterationParameterData = SourceScript->RapidIterationParameters.GetParameterData(RapidIterationParameter);
-			if(RapidIterationParameterData != nullptr)
+			if (LinkedValue.LinkedFunctionCallNode.IsValid() ||
+				AvailableHandles.Contains(LinkedValue.LinkedHandle.GetValue()))
 			{
-				FMemory::Memcpy(LocalValue->GetStructMemory(), RapidIterationParameterData, ValueVariable.GetSizeInBytes());
-				bFoundRapidIterationParameter = true;
+				ValueToUse = &LinkedValue;
+				break;
 			}
 		}
+	}
 
-		if(bFoundRapidIterationParameter == false)
+	if (ValueToUse != nullptr)
+	{
+		if (ValueToUse->LinkedHandle.IsSet())
+		{
+			InInputValues.Mode = EValueMode::Linked;
+			InInputValues.LinkedHandle = ValueToUse->LinkedHandle.GetValue();
+		}
+		else
 		{
-			ValueVariable.CopyTo(LocalValue->GetStructMemory());
+			InInputValues.Mode = EValueMode::DefaultFunction;
+			InInputValues.DefaultFunctionNode = ValueToUse->LinkedFunctionCallNode;
 		}
-		return true;
 	}
-	return false;
 }
 
-bool UNiagaraStackFunctionInput::TryGetCurrentDataValue(FDataValues& DataValues, UEdGraphPin* OverrideValuePin, UEdGraphPin& DefaultValuePin, UNiagaraDataInterface* LocallyOwnedDefaultDataValueObjectToReuse)
+void UNiagaraStackFunctionInput::UpdateValuesFromScriptDefaults(FInputValues& InInputValues) const
 {
-	if (InputType.IsDataInterface())
+	// Get the script variable first since it's used to determine static switch and bound input values.
+	UNiagaraScriptVariable* InputScriptVariable = nullptr;
+	if (OwningFunctionCallNode->FunctionScript != nullptr)
 	{
-		UNiagaraDataInterface* DataValueObject = nullptr;
-		if (OverrideValuePin != nullptr)
+		UNiagaraGraph* FunctionGraph = CastChecked<UNiagaraScriptSource>(OwningFunctionCallNode->FunctionScript->GetSource())->NodeGraph;
+		InputScriptVariable = FunctionGraph->GetScriptVariable(InputParameterHandle.GetParameterHandleString());
+	}
+
+	if (IsStaticParameter())
+	{
+		// Static switch parameters are always locally set values.
+		if (InputScriptVariable != nullptr)
 		{
-			if (OverrideValuePin->LinkedTo.Num() == 1)
-			{
-				UNiagaraNodeInput* InputNode = Cast<UNiagaraNodeInput>(OverrideValuePin->LinkedTo[0]->GetOwningNode());
-				if (InputNode != nullptr)
-				{
-					if (InputNode->Usage == ENiagaraInputNodeUsage::Parameter)
-					{
-						DataValueObject = InputNode->GetDataInterface();
-					}
-				}
+			TSharedPtr<FStructOnScope> StaticSwitchLocalStruct = FNiagaraEditorUtilities::StaticSwitchDefaultIntToStructOnScope(InputScriptVariable->Metadata.StaticSwitchDefaultValue, InputType);
+			if(ensureMsgf(StaticSwitchLocalStruct.IsValid(), TEXT("Unsupported static struct default value.")))
+			{ 
+				InInputValues.Mode = EValueMode::Local;
+				InInputValues.LocalStruct = StaticSwitchLocalStruct;
 			}
 		}
-
-		UNiagaraDataInterface* DefaultDataValueObject = nullptr;
-		FDataValues::EDefaultValueOwner DefaultDataValueOwner = FDataValues::EDefaultValueOwner::Invalid;
-		if (DefaultValuePin.LinkedTo.Num() == 1)
+	}
+	else
+	{
+		if (InputScriptVariable != nullptr && InputScriptVariable->DefaultMode == ENiagaraDefaultMode::Binding && InputScriptVariable->DefaultBinding.IsValid())
 		{
-			UNiagaraNodeInput* InputNode = Cast<UNiagaraNodeInput>(DefaultValuePin.LinkedTo[0]->GetOwningNode());
-			if (InputNode != nullptr && InputNode->Usage == ENiagaraInputNodeUsage::Parameter && InputNode->GetDataInterface() != nullptr)
-			{
-				DefaultDataValueObject = InputNode->GetDataInterface();
-				DefaultDataValueOwner = FDataValues::EDefaultValueOwner::FunctionOwned;
-			}
+			// The next highest precedence value is a linked value from a variable binding so check that.
+			InInputValues.Mode = EValueMode::Linked;
+			InInputValues.LinkedHandle = InputScriptVariable->DefaultBinding.GetName();
 		}
-
-		if (DefaultDataValueObject == nullptr)
+		else
 		{
-			if (LocallyOwnedDefaultDataValueObjectToReuse == nullptr)
-			{
-				DefaultDataValueObject = NewObject<UNiagaraDataInterface>(this, const_cast<UClass*>(InputType.GetClass()), NAME_None, RF_Transactional | RF_Public);
-			}
-			else
+			// Otherwise we need to check the pin that defined the variable in the graph to determine the default.
+			UEdGraphPin* DefaultPin = OwningFunctionCallNode->FindParameterMapDefaultValuePin(InputParameterHandle.GetParameterHandleString(), SourceScript->GetUsage());
+			if (DefaultPin != nullptr)
 			{
-				DefaultDataValueObject = LocallyOwnedDefaultDataValueObjectToReuse;
+				if (InputType.IsDataInterface())
+				{
+					// Data interfaces are handled differently than other values types so collect them here.
+					GetDefaultDataInterfaceValueFromDefaultPin(DefaultPin, InInputValues);
+				}
+				else
+				{
+					// Otherwise check for local and linked values.
+					if (DefaultPin->LinkedTo.Num() == 0)
+					{
+						// If the default pin isn't wired to anything then it's a local value.
+						GetDefaultLocalValueFromDefaultPin(DefaultPin, InInputValues);
+					}
+					else if(DefaultPin->LinkedTo.Num() == 1 && DefaultPin->LinkedTo[0]->GetOwningNode() != nullptr)
+					{
+						// If a default pin is linked to a parameter map it can be a linked value.
+						GetDefaultLinkedHandleOrLinkedFunctionFromDefaultPin(DefaultPin, InInputValues);
+					}
+
+					if(InInputValues.Mode == EValueMode::None)
+					{
+						// If an input mode wasn't found than the graph is configured in a way that can't be displayed in the stack.
+						InInputValues.Mode = EValueMode::UnsupportedDefault;
+					}
+				}
 			}
-			DefaultDataValueOwner = FDataValues::EDefaultValueOwner::LocallyOwned;
 		}
-		
-		DataValues = FDataValues(DataValueObject, DefaultDataValueObject, DefaultDataValueOwner);
-		return true;
 	}
-	return false;
 }
 
-bool UNiagaraStackFunctionInput::TryGetCurrentLinkedValue(FNiagaraParameterHandle& LinkedValueHandle, UEdGraphPin& ValuePin)
+void UNiagaraStackFunctionInput::UpdateValuesFromOverridePin(const FInputValues& OldInputValues, FInputValues& NewInputValues, UEdGraphPin& InOverridePin) const
 {
-	if (ValuePin.LinkedTo.Num() == 1)
+	NewInputValues.Mode = EValueMode::InvalidOverride;
+	if (InOverridePin.LinkedTo.Num() == 0)
 	{
-		UEdGraphPin* CurrentValuePin = &ValuePin;
-		TSharedPtr<TArray<FNiagaraParameterHandle>> AvailableHandles;
-		while (CurrentValuePin != nullptr)
+		// If an override pin exists but it's not connected, the only valid state is a local struct value stored in the pins default value string.
+		if (InputType.IsUObject() == false)
 		{
-			UEdGraphPin* LinkedValuePin = CurrentValuePin->LinkedTo[0];
-			CurrentValuePin = nullptr;
-
-			UNiagaraNodeParameterMapGet* GetNode = Cast<UNiagaraNodeParameterMapGet>(LinkedValuePin->GetOwningNode());
-			if (GetNode == nullptr)
+			// If there was an old local struct, reuse it if it's of the correct type.
+			TSharedPtr<FStructOnScope> LocalStruct;
+			if (OldInputValues.Mode == EValueMode::Local && OldInputValues.LocalStruct.IsValid() && OldInputValues.LocalStruct->GetStruct() == InputType.GetStruct())
 			{
-				// Only parameter map get nodes are supported for linked values.
-				return false;
+				LocalStruct = OldInputValues.LocalStruct;
 			}
-
-			// If a parameter map get node was found, the linked handle will be stored in the pin name.  
-			FNiagaraParameterHandle LinkedValueHandleFromNode(LinkedValuePin->PinName);
-
-			UEdGraphPin* LinkedValueHandleDefaultPin = GetNode->GetDefaultPin(LinkedValuePin);
-			if (LinkedValueHandleDefaultPin->LinkedTo.Num() == 0)
+			else
 			{
-				// If the default value pin for this get node isn't connected this is the last read in the chain
-				// so return the handle.
-				LinkedValueHandle = LinkedValueHandleFromNode;
-				return true;
+				LocalStruct = MakeShared<FStructOnScope>(InputType.GetStruct());
 			}
-			else
+			const UEdGraphSchema_Niagara* NiagaraSchema = GetDefault<UEdGraphSchema_Niagara>();
+			FNiagaraVariable ValueVariable = NiagaraSchema->PinToNiagaraVariable(&InOverridePin, false);
+			if (ValueVariable.IsDataAllocated())
 			{
-				// If the default value pin for the get node is connected then there are a chain of possible values.
-				// if the value of the current get node is available it can be returned, otherwise we need to check the
-				// next node.
-				if (AvailableHandles.IsValid() == false)
-				{
-					AvailableHandles = MakeShared<TArray<FNiagaraParameterHandle>>();
-					GetAvailableParameterHandles(*AvailableHandles);
-				}
-
-				if (AvailableHandles->Contains(LinkedValueHandleFromNode))
-				{
-					LinkedValueHandle = LinkedValueHandleFromNode;
-					return true;
-				}
-				else
-				{
-					CurrentValuePin = LinkedValueHandleDefaultPin;
-				}
+				ValueVariable.CopyTo(LocalStruct->GetStructMemory());
+				NewInputValues.Mode = EValueMode::Local;
+				NewInputValues.LocalStruct = LocalStruct;
 			}
 		}
 	}
-	return false;
-}
-
-bool UNiagaraStackFunctionInput::TryGetCurrentExpressionValue(TWeakObjectPtr<UNiagaraNodeCustomHlsl>& ExpressionValue, UEdGraphPin* OverridePin)
-{
-	if (OverridePin != nullptr && OverridePin->LinkedTo.Num() == 1)
+	else if (InOverridePin.LinkedTo.Num() == 1 && InOverridePin.LinkedTo[0] != nullptr && InOverridePin.LinkedTo[0]->GetOwningNode() != nullptr)
 	{
-		UNiagaraNodeCustomHlsl* DynamicNode = Cast<UNiagaraNodeCustomHlsl>(OverridePin->LinkedTo[0]->GetOwningNode());
-		if (DynamicNode != nullptr)
+		UEdGraphNode* LinkedNode = InOverridePin.LinkedTo[0]->GetOwningNode();
+		if (LinkedNode->IsA<UNiagaraNodeInput>())
 		{
-			ExpressionValue = DynamicNode;
-			return true;
+			// Input nodes handle data interface values.
+			UNiagaraNodeInput* InputNode = CastChecked<UNiagaraNodeInput>(LinkedNode);
+			if (InputNode->GetDataInterface() != nullptr)
+			{
+				NewInputValues.Mode = EValueMode::Data;
+				NewInputValues.DataObject = InputNode->GetDataInterface();
+			}
 		}
-	}
-	return false;
-}
-
-bool UNiagaraStackFunctionInput::TryGetCurrentDynamicValue(TWeakObjectPtr<UNiagaraNodeFunctionCall>& DynamicValue, UEdGraphPin* OverridePin)
-{
-	if (OverridePin != nullptr && OverridePin->LinkedTo.Num() == 1)
-	{
-		UNiagaraNodeFunctionCall* DynamicNode = Cast<UNiagaraNodeFunctionCall>(OverridePin->LinkedTo[0]->GetOwningNode());
-		if (DynamicNode != nullptr)
+		else if (LinkedNode->IsA<UNiagaraNodeParameterMapGet>())
 		{
-			DynamicValue = DynamicNode;
-			return true;
+			// Parameter map get nodes handle linked values.
+			NewInputValues.Mode = EValueMode::Linked;
+			NewInputValues.LinkedHandle = FNiagaraParameterHandle(InOverridePin.LinkedTo[0]->PinName);
+		}
+		else if (LinkedNode->IsA<UNiagaraNodeCustomHlsl>())
+		{
+			// Custom hlsl nodes handle expression values.
+			UNiagaraNodeCustomHlsl* ExpressionNode = CastChecked<UNiagaraNodeCustomHlsl>(LinkedNode);
+			NewInputValues.Mode = EValueMode::Expression;
+			NewInputValues.ExpressionNode = ExpressionNode;
+		}
+		else if (LinkedNode->IsA<UNiagaraNodeFunctionCall>())
+		{
+			// Function call nodes handle dynamic inputs.
+			UNiagaraNodeFunctionCall* DynamicNode = CastChecked<UNiagaraNodeFunctionCall>(LinkedNode);
+			NewInputValues.Mode = EValueMode::Dynamic;
+			NewInputValues.DynamicNode = DynamicNode;
 		}
 	}
-	return false;
 }
 
 void UNiagaraStackFunctionInput::RemoveNodesForOverridePin(UEdGraphPin& OverridePin)
@@ -2279,4 +2127,16 @@ void UNiagaraStackFunctionInput::RemoveNodesForOverridePin(UEdGraphPin& Override
 	}
 }
 
+void UNiagaraStackFunctionInput::RemoveOverridePin()
+{
+	UEdGraphPin* OverridePin = GetOverridePin();
+	if (OverridePin != nullptr)
+	{
+		RemoveNodesForOverridePin(*OverridePin);
+		UNiagaraNodeParameterMapSet* OverrideNode = CastChecked<UNiagaraNodeParameterMapSet>(OverridePin->GetOwningNode());
+		OverrideNode->Modify();
+		OverrideNode->RemovePin(OverridePin);
+	}
+}
+
 #undef LOCTEXT_NAMESPACE
diff --git a/Engine/Plugins/FX/Niagara/Source/NiagaraEditor/Public/NiagaraEditorUtilities.h b/Engine/Plugins/FX/Niagara/Source/NiagaraEditor/Public/NiagaraEditorUtilities.h
--- a/Engine/Plugins/FX/Niagara/Source/NiagaraEditor/Public/NiagaraEditorUtilities.h
+++ b/Engine/Plugins/FX/Niagara/Source/NiagaraEditor/Public/NiagaraEditorUtilities.h
@@ -130,6 +130,8 @@ namespace FNiagaraEditorUtilities
 	/** Returns whether the data in two structs on scope matches. */
 	bool DataMatches(const FStructOnScope& StructOnScopeA, const FStructOnScope& StructOnScopeB);
 
+	void NIAGARAEDITOR_API CopyDataTo(FStructOnScope& DestinationStructOnScope, const FStructOnScope& SourceStructOnScope, bool bCheckTypes = true);
+
 	TSharedPtr<SWidget> CreateInlineErrorText(TAttribute<FText> ErrorMessage, TAttribute<FText> ErrorTooltip);
 
 	void CompileExistingEmitters(const TArray<UNiagaraEmitter*>& AffectedEmitters);
@@ -148,6 +150,8 @@ namespace FNiagaraEditorUtilities
 
 	bool ResolveConstantValue(UEdGraphPin* Pin, int32& Value);
 
+	TSharedPtr<FStructOnScope> StaticSwitchDefaultIntToStructOnScope(int32 InStaticSwitchDefaultValue, FNiagaraTypeDefinition InSwitchType);
+
 	void PreprocessFunctionGraph(const UEdGraphSchema_Niagara* Schema, UNiagaraGraph* Graph, const TArray<UEdGraphPin*>& CallInputs, const TArray<UEdGraphPin*>& CallOutputs, ENiagaraScriptUsage ScriptUsage, const FCompileConstantResolver& ConstantResolver);
 
 	bool PODPropertyAppendCompileHash(const void* Container, FProperty* Property, const FString& PropertyName, struct FNiagaraCompileHashVisitor* InVisitor);
diff --git a/Engine/Plugins/FX/Niagara/Source/NiagaraEditor/Public/ViewModels/Stack/NiagaraParameterHandle.h b/Engine/Plugins/FX/Niagara/Source/NiagaraEditor/Public/ViewModels/Stack/NiagaraParameterHandle.h
--- a/Engine/Plugins/FX/Niagara/Source/NiagaraEditor/Public/ViewModels/Stack/NiagaraParameterHandle.h
+++ b/Engine/Plugins/FX/Niagara/Source/NiagaraEditor/Public/ViewModels/Stack/NiagaraParameterHandle.h
@@ -19,6 +19,8 @@ public:
 
 	bool operator==(const FNiagaraParameterHandle& Other) const;
 
+	bool operator!=(const FNiagaraParameterHandle& Other) const { return (*this == Other) == false; }
+
 	static FNiagaraParameterHandle CreateAliasedModuleParameterHandle(const FNiagaraParameterHandle& ModuleParameterHandle, const UNiagaraNodeFunctionCall* ModuleNode);
 
 	static FNiagaraParameterHandle CreateEngineParameterHandle(const FNiagaraVariable& SystemVariable);
diff --git a/Engine/Plugins/FX/Niagara/Source/NiagaraEditor/Public/ViewModels/Stack/NiagaraStackFunctionInput.h b/Engine/Plugins/FX/Niagara/Source/NiagaraEditor/Public/ViewModels/Stack/NiagaraStackFunctionInput.h
--- a/Engine/Plugins/FX/Niagara/Source/NiagaraEditor/Public/ViewModels/Stack/NiagaraStackFunctionInput.h
+++ b/Engine/Plugins/FX/Niagara/Source/NiagaraEditor/Public/ViewModels/Stack/NiagaraStackFunctionInput.h
@@ -43,8 +43,14 @@ public:
 		Data,
 		/** The value is provided by an expression object. */
 		Expression,
-		/** The value source for this input was not set, or couldn't be determined. */
-		Invalid
+		/** The value is a default value provided by a function call. */
+		DefaultFunction,
+		/** This input is overridden in the stack graph, but the override is invalid. */
+		InvalidOverride,
+		/** This input has a default value set in it's graph which can't be displayed in the stack view. */
+		UnsupportedDefault,
+		/** This input has no value. */
+		None
 	};
 
 	DECLARE_MULTICAST_DELEGATE(FOnValueChanged);
@@ -52,9 +58,6 @@ public:
 public:
 	UNiagaraStackFunctionInput();
 
-	static void AddReferencedObjects(UObject* InThis, FReferenceCollector& Collector);
-	void AddReferencedObjects(FReferenceCollector& Collector);
-
 	/** 
 	 * Sets the input data for this entry.
 	 * @param InRequiredEntryData The required data for all stack entries.
@@ -105,7 +108,8 @@ public:
 	virtual FText GetPasteTransactionText(const UNiagaraClipboardContent* ClipboardContent) const override;
 	virtual void Paste(const UNiagaraClipboardContent* ClipboardContent, FText& OutPasteWarning) override;
 
-	FText GetTooltipText(EValueMode InValueMode) const;
+	/** Gets the tooltip that should be shown for the value of this input. */
+	FText GetValueToolTip() const;
 
 	/** Gets the path of parameter handles from the owning module to the function call which owns this input. */
 	const TArray<FNiagaraParameterHandle>& GetInputParameterHandlePath() const;
@@ -119,8 +123,11 @@ public:
 	/** Sets the value of this input to a linked parameter handle. */
 	void SetLinkedValueHandle(const FNiagaraParameterHandle& InParameterHandle);
 
-	/** Gets the current set of available parameter handes which can be assigned to this input. */
-	void GetAvailableParameterHandles(TArray<FNiagaraParameterHandle>& AvailableParameterHandles);
+	/** Gets the current set of available parameter handles which can be assigned to this input. */
+	void GetAvailableParameterHandles(TArray<FNiagaraParameterHandle>& AvailableParameterHandles) const;
+
+	/** Gets the function node form the script graph if the current value mode is DefaultFunction. */
+	UNiagaraNodeFunctionCall* GetDefaultFunctionNode() const;
 
 	/** Gets the dynamic input node providing the value for this input, if one is available. */
 	UNiagaraNodeFunctionCall* GetDynamicInputNode() const;
@@ -138,7 +145,7 @@ public:
 	void SetCustomExpression(const FString& InCustomExpression);
 
 	/** Gets the current struct value of this input is there is one. */
-	TSharedPtr<FStructOnScope> GetLocalValueStruct();
+	TSharedPtr<const FStructOnScope> GetLocalValueStruct();
 
 	/** Gets the current data object value of this input is there is one. */
 	UNiagaraDataInterface* GetDataValueObject();
@@ -153,7 +160,7 @@ public:
 	bool IsEnabled() const;
 
 	/** Sets this input's local value. */
-	void SetLocalValue(TSharedRef<FStructOnScope> InLocalValue, bool bIsOverride = false);
+	void SetLocalValue(TSharedRef<FStructOnScope> InLocalValue);
 	
 	/** Returns whether or not the value or handle of this input has been overridden and can be reset. */
 	bool CanReset() const;
@@ -249,91 +256,35 @@ protected:
 	virtual void FinalizeInternal() override;
 	virtual void RefreshChildrenInternal(const TArray<UNiagaraStackEntry*>& CurrentChildren, TArray<UNiagaraStackEntry*>& NewChildren, TArray<FStackIssue>& NewIssues) override;
 
-	FNiagaraVariable GetDefaultVariableForRapidIterationParameter() const;
 	bool UpdateRapidIterationParametersForAffectedScripts(const uint8* Data);
 	bool RemoveRapidIterationParametersForAffectedScripts();
 	FString ResolveDisplayNameArgument(const FString& InArg) const;
 
 private:
-	struct FDataValues
-	{
-		enum class EDefaultValueOwner
-		{
-			LocallyOwned,
-			FunctionOwned,
-			Invalid
-		};
-
-		FDataValues()
-			: ValueObject(nullptr)
-			, DefaultValueObject(nullptr)
-			, DefaultValueOwner(EDefaultValueOwner::Invalid)
-			, bIsValid(false)
-		{
-		}
-
-		FDataValues(UNiagaraDataInterface* InValueObject, UNiagaraDataInterface* InDefaultValueObject, EDefaultValueOwner InDefaultValueOwner)
-			: ValueObject(InValueObject)
-			, DefaultValueObject(InDefaultValueObject)
-			, DefaultValueOwner(InDefaultValueOwner)
-			, bIsValid(true)
-		{
-			checkf(DefaultValueObject == nullptr || DefaultValueOwner != EDefaultValueOwner::Invalid, TEXT("Must specify a valid owner if the default value object is not null"));
-		}
-
-		UNiagaraDataInterface* GetValueObject() const
-		{
-			return ValueObject;
-		}
-
-		UNiagaraDataInterface* GetDefaultValueObject() const
-		{
-			return DefaultValueObject;
-		}
-
-		UNiagaraDataInterface*& GetDefaultValueObjectRef()
-		{
-			return DefaultValueObject;
-		}
-
-		EDefaultValueOwner GetDefaultValueOwner() const
-		{
-			return DefaultValueOwner;
-		}
-
-		bool IsValid() const
-		{
-			return bIsValid;
-		}
-
-	private:
-		UNiagaraDataInterface* ValueObject;
-		UNiagaraDataInterface* DefaultValueObject;
-		EDefaultValueOwner DefaultValueOwner;
-		bool bIsValid;
-	};
-
 	struct FInputValues
 	{
 		FInputValues()
-			: Mode(EValueMode::Invalid)
+			: Mode(EValueMode::None)
 		{
 		}
 
-		TSharedPtr<FStructOnScope> GetLocalStructToReuse();
-		UNiagaraDataInterface* GetDataDefaultValueObjectToReuse();
+		bool HasEditableData() const
+		{
+			return Mode != EValueMode::None && Mode != EValueMode::InvalidOverride && Mode != EValueMode::UnsupportedDefault;
+		}
 
 		EValueMode Mode;
 		TSharedPtr<FStructOnScope> LocalStruct;
 		FNiagaraParameterHandle LinkedHandle;
 		TWeakObjectPtr<UNiagaraNodeFunctionCall> DynamicNode;
 		TWeakObjectPtr<UNiagaraNodeCustomHlsl> ExpressionNode;
-		FDataValues DataObjects;
+		TWeakObjectPtr<UNiagaraDataInterface> DataObject;
+		TWeakObjectPtr<UNiagaraNodeFunctionCall> DefaultFunctionNode;
 	};
 
 private:
 	/** Refreshes the current values for this input from the state of the graph. */
-	void RefreshValues(bool bFromSetLocalValue=false);
+	void RefreshValues();
 
 	/** Refreshes additional state for this input which comes from input metadata. */
 	void RefreshFromMetaData();
@@ -354,9 +305,6 @@ private:
 	  * This will create the node and add it to the graph if it doesn't exist. */
 	UNiagaraNodeParameterMapSet& GetOrCreateOverrideNode();
 
-	/** Gets the default value pin from the map get node which generated this input. */
-	UEdGraphPin* GetDefaultPin() const;
-
 	/** Gets the pin on the override node which is associated with this input if it exists. */
 	UEdGraphPin* GetOverridePin() const;
 
@@ -364,25 +312,21 @@ private:
 	  * pin don't exist, they will be created. */
 	UEdGraphPin& GetOrCreateOverridePin();
 
-	bool TryGetDefaultBinding(FNiagaraParameterHandle& LinkedValueHandle, UNiagaraScriptVariable* Variable, UEdGraphPin& ValuePin);
+	void GetDefaultDataInterfaceValueFromDefaultPin(UEdGraphPin* DefaultPin, UNiagaraStackFunctionInput::FInputValues& InInputValues) const;
 
-	/** Tries to get a local value for this input if it exists by checking the graph data directly. */
-	bool TryGetCurrentLocalValue(TSharedPtr<FStructOnScope>& LocalValue, UEdGraphPin& DefaultPin, UEdGraphPin& ValuePin, TSharedPtr<FStructOnScope> OldValueToReuse, UNiagaraScriptVariable* Variable);
+	void GetDefaultLocalValueFromDefaultPin(UEdGraphPin* DefaultPin, UNiagaraStackFunctionInput::FInputValues& InInputValues) const;
 
-	/** Tries to get a data interface value for this input if it exists by checking the graph data directly .*/
-	bool TryGetCurrentDataValue(FDataValues& DataValues, UEdGraphPin* OverrideValuePin, UEdGraphPin& DefaultValuePin, UNiagaraDataInterface* LocallyOwnedDefaultDataValueObjectToReuse);
+	void GetDefaultLinkedHandleOrLinkedFunctionFromDefaultPin(UEdGraphPin* DefaultPin, UNiagaraStackFunctionInput::FInputValues& InInputValues) const;
 
-	/** Tries to get the linked value parameter handle for this input if it exists by checking the graph directly. */
-	bool TryGetCurrentLinkedValue(FNiagaraParameterHandle& LinkedValue, UEdGraphPin& ValuePin);
+	void UpdateValuesFromScriptDefaults(FInputValues& InInputValues) const;
 
-	/** Gets the dynamic input node providing a value to this input if one exists. */
-	bool TryGetCurrentDynamicValue(TWeakObjectPtr<UNiagaraNodeFunctionCall>& DynamicValue, UEdGraphPin* OverridePin);
+	void UpdateValuesFromOverridePin(const FInputValues& OldInputValues, FInputValues& NewInputValues, UEdGraphPin& InOverridePin) const;
 
 	/** Removes all nodes connected to the override pin which provide it's value. */
 	void RemoveNodesForOverridePin(UEdGraphPin& OverridePin);
 
-	/** Gets the expression input node providing a value to this input if one exists. */
-	bool TryGetCurrentExpressionValue(TWeakObjectPtr<UNiagaraNodeCustomHlsl>& ExpressionValue, UEdGraphPin* OverridePin);
+	/** Remove the override pin and all nodes connected to it. */
+	void RemoveOverridePin();
 
 	/** Determine if the values in this input are possibly under the control of the rapid iteration array on the script.*/
 	bool IsRapidIterationCandidate() const;
@@ -434,6 +378,9 @@ private:
 	/** Optional override for the display name*/
 	TOptional<FText> DisplayNameOverride;
 
+	/** The default value for this input defined in the defining script. */
+	FInputValues DefaultInputValues;
+
 	/** Pointers and handles to the various values this input can have. */
 	FInputValues InputValues;
 
@@ -451,6 +398,9 @@ private:
 	/** Whether or not this input can be reset to a base value defined by a parent emitter. */
 	mutable TOptional<bool> bCanResetToBaseCache;
 
+	/** A tooltip to show for the value of this input. */
+	mutable TOptional<FText> ValueToolTipCache;
+
 	/** A flag to prevent handling graph changes when it's being updated directly by this object. */
 	bool bUpdatingGraphDirectly;
 
@@ -488,7 +438,4 @@ private:
 
 	/** Whether or not the dynamic input for this input has a function script reassignment pending due to a request to fix a missing script. */
 	bool bIsDynamicInputScriptReassignmentPending;
-
-	bool bIsLocalOverride;
-	UNiagaraScriptVariable* Variable;
 };
diff --git a/Engine/Plugins/FX/Niagara/Source/NiagaraEditorWidgets/Private/Stack/SNiagaraStackFunctionInputName.cpp b/Engine/Plugins/FX/Niagara/Source/NiagaraEditorWidgets/Private/Stack/SNiagaraStackFunctionInputName.cpp
--- a/Engine/Plugins/FX/Niagara/Source/NiagaraEditorWidgets/Private/Stack/SNiagaraStackFunctionInputName.cpp
+++ b/Engine/Plugins/FX/Niagara/Source/NiagaraEditorWidgets/Private/Stack/SNiagaraStackFunctionInputName.cpp
@@ -98,7 +98,7 @@ FText SNiagaraStackFunctionInputName::GetToolTipText() const
 	{
 		return FText();
 	}
-	return FunctionInput->GetTooltipText(UNiagaraStackFunctionInput::EValueMode::Local);
+	return FunctionInput->GetTooltipText();
 }
 
 void SNiagaraStackFunctionInputName::OnNameTextCommitted(const FText& InText, ETextCommit::Type InCommitType)
diff --git a/Engine/Plugins/FX/Niagara/Source/NiagaraEditorWidgets/Private/Stack/SNiagaraStackFunctionInputValue.cpp b/Engine/Plugins/FX/Niagara/Source/NiagaraEditorWidgets/Private/Stack/SNiagaraStackFunctionInputValue.cpp
--- a/Engine/Plugins/FX/Niagara/Source/NiagaraEditorWidgets/Private/Stack/SNiagaraStackFunctionInputValue.cpp
+++ b/Engine/Plugins/FX/Niagara/Source/NiagaraEditorWidgets/Private/Stack/SNiagaraStackFunctionInputValue.cpp
@@ -30,6 +30,7 @@
 #include "NiagaraSystem.h"
 #include "SNiagaraGraphActionWidget.h"
 #include "Subsystems/AssetEditorSubsystem.h"
+#include "EditorFontGlyphs.h"
 
 #define LOCTEXT_NAMESPACE "NiagaraStackFunctionInputValue"
 
@@ -40,9 +41,7 @@ bool SNiagaraStackFunctionInputValue::bIncludeNonLibraryInputs = false;
 void SNiagaraStackFunctionInputValue::Construct(const FArguments& InArgs, UNiagaraStackFunctionInput* InFunctionInput)
 {
 	FunctionInput = InFunctionInput;
-
 	FunctionInput->OnValueChanged().AddSP(this, &SNiagaraStackFunctionInputValue::OnInputValueChanged);
-	DisplayedLocalValueStruct = FunctionInput->GetLocalValueStruct();
 
 	FMargin ItemPadding = FMargin(0);
 	ChildSlot
@@ -79,91 +78,13 @@ void SNiagaraStackFunctionInputValue::Construct(const FArguments& InArgs, UNiaga
 				]
 			]
 			+ SHorizontalBox::Slot()
+			.VAlign(VAlign_Center)
 			[
-				// TODO Don't generate all of these widgets for every input, only generate the ones that are used based on the value type.
-
-				// Local struct
-				SNew(SOverlay)
-				+ SOverlay::Slot()
-				[
-					SAssignNew(LocalValueStructContainer, SBox)
-					.Visibility(this, &SNiagaraStackFunctionInputValue::GetValueWidgetVisibility, UNiagaraStackFunctionInput::EValueMode::Local)
-					[
-						ConstructLocalValueStructWidget()
-					]
-				]
-
-				// Linked handle
-				+ SOverlay::Slot()
-				.Padding(0.0f, 0.0f, 0.0f, 2.0f)
-				[
-					SNew(SBox)
-					.Visibility(this, &SNiagaraStackFunctionInputValue::GetValueWidgetVisibility, UNiagaraStackFunctionInput::EValueMode::Linked)
-					.ToolTipText_UObject(InFunctionInput, &UNiagaraStackFunctionInput::GetTooltipText, UNiagaraStackFunctionInput::EValueMode::Linked)
-					.VAlign(VAlign_Center)
-					[
-						SNew(STextBlock)
-						.TextStyle(FNiagaraEditorStyle::Get(), "NiagaraEditor.ParameterText")
-						.Text(this, &SNiagaraStackFunctionInputValue::GetLinkedValueHandleText)
-						.OnDoubleClicked(this, &SNiagaraStackFunctionInputValue::OnLinkedInputDoubleClicked)
-					]
-				]
-
-				// Data Object
-				+ SOverlay::Slot()
-				.Padding(0.0f, 0.0f, 0.0f, 2.0f)
-				[
-					SNew(SBox)
-					.Visibility(this, &SNiagaraStackFunctionInputValue::GetValueWidgetVisibility, UNiagaraStackFunctionInput::EValueMode::Data)
-					.VAlign(VAlign_Center)
-					[
-						SNew(STextBlock)
-						.TextStyle(FNiagaraEditorStyle::Get(), "NiagaraEditor.ParameterText")
-						.Text(this, &SNiagaraStackFunctionInputValue::GetDataValueText)
-					]
-				]
-
-				// Dynamic input name
-				+ SOverlay::Slot()
-				[
-					SNew(SBox)
-					.Visibility(this, &SNiagaraStackFunctionInputValue::GetValueWidgetVisibility, UNiagaraStackFunctionInput::EValueMode::Dynamic)
-					.VAlign(VAlign_Center)
-					[
-						SNew(STextBlock)
-						.TextStyle(FNiagaraEditorStyle::Get(), "NiagaraEditor.ParameterText")
-						.Text(this, &SNiagaraStackFunctionInputValue::GetDynamicValueText)
-						.ToolTipText(this, &SNiagaraStackFunctionInputValue::GetDynamicValueToolTip)
-						.OnDoubleClicked(this, &SNiagaraStackFunctionInputValue::DynamicInputTextDoubleClicked)
-					]
-				]
-
-				// Expression input
-				+ SOverlay::Slot()
+				// Value container and widgets.
+				SAssignNew(ValueContainer, SBox)
+				.ToolTipText_UObject(FunctionInput, &UNiagaraStackFunctionInput::GetValueToolTip)
 				[
-					SNew(SBox)
-					.Visibility(this, &SNiagaraStackFunctionInputValue::GetValueWidgetVisibility, UNiagaraStackFunctionInput::EValueMode::Expression)
-					.VAlign(VAlign_Center)
-					[
-						SNew(SEditableTextBox)
-						.IsReadOnly(false)
-						.Text_UObject(FunctionInput, &UNiagaraStackFunctionInput::GetCustomExpressionText)
-						.OnTextCommitted(this, &SNiagaraStackFunctionInputValue::OnExpressionTextCommitted)
-					]
-				]
-
-				// Invalid input
-				+ SOverlay::Slot()
-				[
-					SNew(SBox)
-					.Visibility(this, &SNiagaraStackFunctionInputValue::GetValueWidgetVisibility, UNiagaraStackFunctionInput::EValueMode::Invalid)
-					.VAlign(VAlign_Center)
-					[
-						SNew(STextBlock)
-						.TextStyle(FNiagaraEditorStyle::Get(), "NiagaraEditor.ParameterText")
-						.Text(this, &SNiagaraStackFunctionInputValue::GetInvalidValueText)
-						.ToolTipText(this, &SNiagaraStackFunctionInputValue::GetInvalidValueToolTipText)
-					]
+					ConstructValueWidgets()
 				]
 			]
 
@@ -226,6 +147,8 @@ void SNiagaraStackFunctionInputValue::Construct(const FArguments& InArgs, UNiaga
 			]
 		]
 	];
+
+	ValueModeForGeneratedWidgets = FunctionInput->GetValueMode();
 }
 
 void SNiagaraStackFunctionInputValue::Tick(const FGeometry& AllottedGeometry, const double InCurrentTime, const float InDeltaTime)
@@ -237,20 +160,81 @@ void SNiagaraStackFunctionInputValue::Tick(const FGeometry& AllottedGeometry, co
 	}
 }
 
-bool SNiagaraStackFunctionInputValue::GetInputEnabled() const
+TSharedRef<SWidget> SNiagaraStackFunctionInputValue::ConstructValueWidgets()
 {
-	return FunctionInput->GetHasEditCondition() == false || FunctionInput->GetEditConditionEnabled();
+	DisplayedLocalValueStruct.Reset();
+	LocalValueStructParameterEditor.Reset();
+	LocalValueStructDetailsView.Reset();
+
+	switch (FunctionInput->GetValueMode())
+	{
+	case UNiagaraStackFunctionInput::EValueMode::Local:
+	{
+		return ConstructLocalValueStructWidget();
+	}
+	case UNiagaraStackFunctionInput::EValueMode::Linked:
+	{
+		return SNew(STextBlock)
+			.TextStyle(FNiagaraEditorStyle::Get(), "NiagaraEditor.ParameterText")
+			.Text(this, &SNiagaraStackFunctionInputValue::GetLinkedValueHandleText)
+			.OnDoubleClicked(this, &SNiagaraStackFunctionInputValue::OnLinkedInputDoubleClicked);
+	}
+	case UNiagaraStackFunctionInput::EValueMode::Data:
+	{
+		return SNew(STextBlock)
+			.TextStyle(FNiagaraEditorStyle::Get(), "NiagaraEditor.ParameterText")
+			.Text(this, &SNiagaraStackFunctionInputValue::GetDataValueText);
+	}
+	case UNiagaraStackFunctionInput::EValueMode::Dynamic:
+	{
+		return SNew(STextBlock)
+			.TextStyle(FNiagaraEditorStyle::Get(), "NiagaraEditor.ParameterText")
+			.Text(this, &SNiagaraStackFunctionInputValue::GetDynamicValueText)
+			.OnDoubleClicked(this, &SNiagaraStackFunctionInputValue::DynamicInputTextDoubleClicked);
+	}
+	case UNiagaraStackFunctionInput::EValueMode::DefaultFunction:
+	{
+		return SNew(STextBlock)
+			.TextStyle(FNiagaraEditorStyle::Get(), "NiagaraEditor.ParameterText")
+			.Text(this, &SNiagaraStackFunctionInputValue::GetDefaultFunctionText);
+	}
+	case UNiagaraStackFunctionInput::EValueMode::Expression:
+	{
+		return SNew(SEditableTextBox)
+			.IsReadOnly(false)
+			.Text_UObject(FunctionInput, &UNiagaraStackFunctionInput::GetCustomExpressionText)
+			.OnTextCommitted(this, &SNiagaraStackFunctionInputValue::OnExpressionTextCommitted);
+	}
+	case UNiagaraStackFunctionInput::EValueMode::InvalidOverride:
+	{
+		return SNew(STextBlock)
+			.TextStyle(FNiagaraEditorStyle::Get(), "NiagaraEditor.ParameterText")
+			.Text(LOCTEXT("InvalidOverrideText", "Invalid Scirpt Value"));
+	}
+	case UNiagaraStackFunctionInput::EValueMode::UnsupportedDefault:
+	{
+		return SNew(STextBlock)
+			.TextStyle(FNiagaraEditorStyle::Get(), "NiagaraEditor.ParameterText")
+			.Text(LOCTEXT("UnsupportedDefault", "Custom Default"));
+	}
+	default:
+		return SNullWidget::NullWidget;
+	}
 }
 
-EVisibility SNiagaraStackFunctionInputValue::GetValueWidgetVisibility(UNiagaraStackFunctionInput::EValueMode ValidMode) const
+
+bool SNiagaraStackFunctionInputValue::GetInputEnabled() const
 {
-	return FunctionInput->GetValueMode() == ValidMode ? EVisibility::Visible : EVisibility::Collapsed;
+	return FunctionInput->GetHasEditCondition() == false || FunctionInput->GetEditConditionEnabled();
 }
 
 TSharedRef<SWidget> SNiagaraStackFunctionInputValue::ConstructLocalValueStructWidget()
 {
 	LocalValueStructParameterEditor.Reset();
 	LocalValueStructDetailsView.Reset();
+
+	DisplayedLocalValueStruct = MakeShared<FStructOnScope>(FunctionInput->GetInputType().GetStruct());
+	FNiagaraEditorUtilities::CopyDataTo(*DisplayedLocalValueStruct.Get(), *FunctionInput->GetLocalValueStruct().Get());
 	if (DisplayedLocalValueStruct.IsValid())
 	{
 		FNiagaraEditorModule& NiagaraEditorModule = FModuleManager::GetModuleChecked<FNiagaraEditorModule>("NiagaraEditor");
@@ -299,23 +283,33 @@ TSharedRef<SWidget> SNiagaraStackFunctionInputValue::ConstructLocalValueStructWi
 
 void SNiagaraStackFunctionInputValue::OnInputValueChanged()
 {
-	TSharedPtr<FStructOnScope> NewLocalValueStruct = FunctionInput->GetLocalValueStruct();
-	if (DisplayedLocalValueStruct == NewLocalValueStruct)
+	if (ValueModeForGeneratedWidgets != FunctionInput->GetValueMode())
 	{
-		if (LocalValueStructParameterEditor.IsValid())
-		{
-			LocalValueStructParameterEditor->UpdateInternalValueFromStruct(DisplayedLocalValueStruct.ToSharedRef());
-		}
-		if (LocalValueStructDetailsView.IsValid())
-		{
-			LocalValueStructDetailsView->SetStructureData(TSharedPtr<FStructOnScope>());
-			LocalValueStructDetailsView->SetStructureData(DisplayedLocalValueStruct);
-		}
+		ValueContainer->SetContent(ConstructValueWidgets());
+		ValueModeForGeneratedWidgets = FunctionInput->GetValueMode();
 	}
 	else
 	{
-		DisplayedLocalValueStruct = NewLocalValueStruct;
-		LocalValueStructContainer->SetContent(ConstructLocalValueStructWidget());
+		if (ValueModeForGeneratedWidgets == UNiagaraStackFunctionInput::EValueMode::Local)
+		{
+			if (DisplayedLocalValueStruct->GetStruct() == FunctionInput->GetLocalValueStruct()->GetStruct())
+			{
+				FNiagaraEditorUtilities::CopyDataTo(*DisplayedLocalValueStruct.Get(), *FunctionInput->GetLocalValueStruct().Get());
+				if (LocalValueStructParameterEditor.IsValid())
+				{
+					LocalValueStructParameterEditor->UpdateInternalValueFromStruct(DisplayedLocalValueStruct.ToSharedRef());
+				}
+				if (LocalValueStructDetailsView.IsValid())
+				{
+					LocalValueStructDetailsView->SetStructureData(TSharedPtr<FStructOnScope>());
+					LocalValueStructDetailsView->SetStructureData(DisplayedLocalValueStruct);
+				}
+			}
+			else
+			{
+				ValueContainer->SetContent(ConstructLocalValueStructWidget());
+			}
+		}
 	}
 }
 
@@ -373,15 +367,15 @@ FText SNiagaraStackFunctionInputValue::GetDynamicValueText() const
 	}
 }
 
-FText SNiagaraStackFunctionInputValue::GetDynamicValueToolTip() const
+FText SNiagaraStackFunctionInputValue::GetDefaultFunctionText() const
 {
-	if (FunctionInput->GetDynamicInputNode() != nullptr)
+	if (FunctionInput->GetDefaultFunctionNode() != nullptr)
 	{
-		return FunctionInput->GetDynamicInputNode()->GetTooltipText();
+		return FText::FromString(FName::NameToDisplayString(FunctionInput->GetDefaultFunctionNode()->GetFunctionName(), false));
 	}
 	else
 	{
-		return LOCTEXT("InvalidDynamicDisplayName", "(Invalid)");
+		return LOCTEXT("InvalidDefaultFunctionDisplayName", "(Invalid)");
 	}
 }
 
@@ -390,30 +384,6 @@ void SNiagaraStackFunctionInputValue::OnExpressionTextCommitted(const FText& Nam
 	FunctionInput->SetCustomExpression(Name.ToString());
 }
 
-FText SNiagaraStackFunctionInputValue::GetInvalidValueText() const
-{
-	if (FunctionInput->CanReset())
-	{
-		return LOCTEXT("InvalidResetLabel", "Unsupported value - Reset to fix.");
-	}
-	else
-	{
-		return LOCTEXT("InvalidLabel", "Custom value");
-	}
-}
-
-FText SNiagaraStackFunctionInputValue::GetInvalidValueToolTipText() const
-{
-	if (FunctionInput->CanReset())
-	{
-		return LOCTEXT("InvalidResetToolTip", "This input has an unsupported value assigned in the stack.\nUse the reset button to remove the unsupported value.");
-	}
-	else
-	{
-		return LOCTEXT("InvalidToolTip", "The script that defines the source of this input has\n a custom default value that can not be displayed in the stack view.\nYou can set a local override value using the drop down menu.");
-	}
-}
-
 FReply SNiagaraStackFunctionInputValue::DynamicInputTextDoubleClicked(const FGeometry& MyGeometry, const FPointerEvent& PointerEvent)
 {
 	UNiagaraNodeFunctionCall* DynamicInputNode = FunctionInput->GetDynamicInputNode();
@@ -706,7 +676,7 @@ void SNiagaraStackFunctionInputValue::SetToLocalValue()
 		if (DefaultValueData.Num() == LocalValueStruct->GetStructureSize())
 		{
 			FMemory::Memcpy(LocalValue->GetStructMemory(), DefaultValueData.GetData(), DefaultValueData.Num());
-			FunctionInput->SetLocalValue(LocalValue, true);
+			FunctionInput->SetLocalValue(LocalValue);
 		}
 	}
 }
@@ -772,17 +742,19 @@ FText SNiagaraStackFunctionInputValue::GetInputIconText() const
 	switch (FunctionInput->GetValueMode())
 	{
 	case UNiagaraStackFunctionInput::EValueMode::Linked:
-		return FText::FromString(FString(TEXT("\xf0C1") /* fa-link */));
+		return FEditorFontGlyphs::Link;
 	case UNiagaraStackFunctionInput::EValueMode::Data:
-		return FText::FromString(FString(TEXT("\xf1C0") /* fa-database */));
+		return FEditorFontGlyphs::Database;
 	case UNiagaraStackFunctionInput::EValueMode::Dynamic:
-		return FText::FromString(FString(TEXT("\xf201") /* fa-line-chart */));
+		return FEditorFontGlyphs::Line_Chart;
 	case UNiagaraStackFunctionInput::EValueMode::Expression:
-		return FText::FromString(FString(TEXT("\xf120") /* fa-terminal */));
-	case UNiagaraStackFunctionInput::EValueMode::Invalid:
-		return FunctionInput->CanReset()
-			? FText::FromString(FString(TEXT("\xf128") /* fa-question */))
-			: FText::FromString(FString(TEXT("\xf005") /* fa-star */));
+		return FEditorFontGlyphs::Terminal;
+	case UNiagaraStackFunctionInput::EValueMode::DefaultFunction:
+		return FEditorFontGlyphs::Plug;
+	case UNiagaraStackFunctionInput::EValueMode::InvalidOverride:
+		return FEditorFontGlyphs::Question;
+	case UNiagaraStackFunctionInput::EValueMode::UnsupportedDefault:
+		return FEditorFontGlyphs::Star;
 	default:
 		return FText::FromString(FString(TEXT("\xf128") /* fa-question */));
 	}
@@ -801,10 +773,12 @@ FText SNiagaraStackFunctionInputValue::GetInputIconToolTip() const
 		return LOCTEXT("DynamicInputIconToolTip", "Dynamic Value");
 	case UNiagaraStackFunctionInput::EValueMode::Expression:
 		return LOCTEXT("ExpressionInputIconToolTip", "Custom Expression");
-	case UNiagaraStackFunctionInput::EValueMode::Invalid:
-		return FunctionInput->CanReset()
-			? InvalidText
-			: LOCTEXT("CustomInputIconToolTip", "Custom value");
+	case UNiagaraStackFunctionInput::EValueMode::DefaultFunction:
+		return LOCTEXT("DefaultFunctionIconToolTip", "Script Defined Default Function");
+	case UNiagaraStackFunctionInput::EValueMode::InvalidOverride:
+		return LOCTEXT("InvalidOverrideIconToolTip", "Invalid Script State");
+	case UNiagaraStackFunctionInput::EValueMode::UnsupportedDefault:
+		return LOCTEXT("UnsupportedDefaultIconToolTip", "Script Defined Custom Default");
 	default:
 		return InvalidText;
 	}
@@ -822,7 +796,9 @@ FSlateColor SNiagaraStackFunctionInputValue::GetInputIconColor() const
 		return FLinearColor(FColor::Cyan);
 	case UNiagaraStackFunctionInput::EValueMode::Expression:
 		return FLinearColor(FColor::Green);
-	case UNiagaraStackFunctionInput::EValueMode::Invalid:
+	case UNiagaraStackFunctionInput::EValueMode::InvalidOverride:
+	case UNiagaraStackFunctionInput::EValueMode::UnsupportedDefault:
+	case UNiagaraStackFunctionInput::EValueMode::DefaultFunction:
 	default:
 		return FLinearColor(FColor::White);
 	}
diff --git a/Engine/Plugins/FX/Niagara/Source/NiagaraEditorWidgets/Private/Stack/SNiagaraStackFunctionInputValue.h b/Engine/Plugins/FX/Niagara/Source/NiagaraEditorWidgets/Private/Stack/SNiagaraStackFunctionInputValue.h
--- a/Engine/Plugins/FX/Niagara/Source/NiagaraEditorWidgets/Private/Stack/SNiagaraStackFunctionInputValue.h
+++ b/Engine/Plugins/FX/Niagara/Source/NiagaraEditorWidgets/Private/Stack/SNiagaraStackFunctionInputValue.h
@@ -28,11 +28,10 @@ public:
 	virtual void Tick(const FGeometry& AllottedGeometry, const double InCurrentTime, const float InDeltaTime) override;
 
 private:
+	TSharedRef<SWidget> ConstructValueWidgets();
 
 	bool GetInputEnabled() const;
 
-	EVisibility GetValueWidgetVisibility(UNiagaraStackFunctionInput::EValueMode ValidMode) const;
-
 	TSharedRef<SWidget> ConstructLocalValueStructWidget();
 
 	void OnInputValueChanged();
@@ -51,14 +50,10 @@ private:
 
 	FText GetDynamicValueText() const;
 
-	FText GetDynamicValueToolTip() const;
+	FText GetDefaultFunctionText() const;
 
 	void OnExpressionTextCommitted(const FText& Name, ETextCommit::Type CommitInfo);
 
-	FText GetInvalidValueText() const;
-
-	FText GetInvalidValueToolTipText() const;
-
 	FReply DynamicInputTextDoubleClicked(const FGeometry& MyGeometry, const FPointerEvent& PointerEvent);
 	FReply OnLinkedInputDoubleClicked(const FGeometry& MyGeometry, const FPointerEvent& PointerEvent);
 
@@ -144,11 +139,13 @@ private:
 private:
 	UNiagaraStackFunctionInput* FunctionInput;
 
-	TSharedPtr<FStructOnScope> DisplayedLocalValueStruct;
+	TSharedPtr<SBox> ValueContainer;
+	UNiagaraStackFunctionInput::EValueMode ValueModeForGeneratedWidgets;
 
-	TSharedPtr<SBox> LocalValueStructContainer;
+	TSharedPtr<FStructOnScope> DisplayedLocalValueStruct;
 	TSharedPtr<SNiagaraParameterEditor> LocalValueStructParameterEditor;
 	TSharedPtr<IStructureDetailsView> LocalValueStructDetailsView;
+
 	TSharedPtr<SComboButton> SetFunctionInputButton;
 	TSharedPtr<SGraphActionMenu> SelectInputFunctionMenu;
 
