Function nodes in Niagara scripts now correctly display their experimental status when newly placed.

#rnx

diff --git a/Engine/Plugins/FX/Niagara/Source/NiagaraEditor/Private/NiagaraNodeFunctionCall.cpp b/Engine/Plugins/FX/Niagara/Source/NiagaraEditor/Private/NiagaraNodeFunctionCall.cpp
--- a/Engine/Plugins/FX/Niagara/Source/NiagaraEditor/Private/NiagaraNodeFunctionCall.cpp
+++ b/Engine/Plugins/FX/Niagara/Source/NiagaraEditor/Private/NiagaraNodeFunctionCall.cpp
@@ -294,6 +294,7 @@ void UNiagaraNodeFunctionCall::AllocateDefaultPins()
 		ComputeNodeName();
 	}
 
+	UpdateNodeErrorMessage();
 }
 
 // Returns true if this node is deprecated
@@ -686,11 +687,6 @@ void UNiagaraNodeFunctionCall::UpdateNodeErrorMessage()
 	}
 }
 
-void UNiagaraNodeFunctionCall::PostPlacedNewNode()
-{
-	UpdateNodeErrorMessage();
-}
-
 bool UNiagaraNodeFunctionCall::RefreshFromExternalChanges()
 {
 	bool bReload = false;
diff --git a/Engine/Plugins/FX/Niagara/Source/NiagaraEditor/Public/NiagaraNodeFunctionCall.h b/Engine/Plugins/FX/Niagara/Source/NiagaraEditor/Public/NiagaraNodeFunctionCall.h
--- a/Engine/Plugins/FX/Niagara/Source/NiagaraEditor/Public/NiagaraNodeFunctionCall.h
+++ b/Engine/Plugins/FX/Niagara/Source/NiagaraEditor/Public/NiagaraNodeFunctionCall.h
@@ -85,7 +85,6 @@ public:
 	//End UObject interface
 
 	//~ Begin UNiagaraNode Interface
-	virtual void PostPlacedNewNode() override;
 	virtual void Compile(class FHlslNiagaraTranslator* Translator, TArray<int32>& Outputs) override;
 	virtual UObject* GetReferencedAsset() const override;
 	virtual bool RefreshFromExternalChanges() override;
