Fixed nullptr crash when trying to determine the owning system of an orphaned emitter.

#rnx

diff --git a/Engine/Plugins/FX/Niagara/Source/NiagaraEditor/Private/NiagaraOverviewNode.cpp b/Engine/Plugins/FX/Niagara/Source/NiagaraEditor/Private/NiagaraOverviewNode.cpp
--- a/Engine/Plugins/FX/Niagara/Source/NiagaraEditor/Private/NiagaraOverviewNode.cpp
+++ b/Engine/Plugins/FX/Niagara/Source/NiagaraEditor/Private/NiagaraOverviewNode.cpp
@@ -99,16 +99,21 @@ FLinearColor UNiagaraOverviewNode::GetNodeTitleColor() const
 		NotIsolatedColor = FNiagaraEditorStyle::Get().GetColor("NiagaraEditor.OverviewNode.NotIsolatedColor");
 	}
 
-	if (EmitterHandleGuid.IsValid() && OwningSystem->GetIsolateEnabled())
+	if (EmitterHandleGuid.IsValid())
 	{
-		const FNiagaraEmitterHandle* Handle = FindEmitterHandleByID(OwningSystem, EmitterHandleGuid);
-		if (ensureMsgf(Handle != nullptr, TEXT("Failed to find matching emitter handle for existing overview node!")))
+		if (OwningSystem != nullptr && OwningSystem->GetIsolateEnabled())
 		{
-			return Handle->IsIsolated() ? IsolatedColor : NotIsolatedColor;
+			const FNiagaraEmitterHandle* Handle = FindEmitterHandleByID(OwningSystem, EmitterHandleGuid);
+			if (ensureMsgf(Handle != nullptr, TEXT("Failed to find matching emitter handle for existing overview node!")))
+			{
+				return Handle->IsIsolated() ? IsolatedColor : NotIsolatedColor;
+			}
 		}
+		
+		return EmitterColor;
 	}
-
-	return EmitterHandleGuid.IsValid() ? EmitterColor : SystemColor;
+	
+	return SystemColor;
 }
 
 static bool IsSystemAsset(UNiagaraSystem* System)
