Fixed a bug where renaming a parameter lost the default value settings for that parameter

diff --git a/Engine/Plugins/FX/Niagara/Source/NiagaraEditor/Private/NiagaraGraph.cpp b/Engine/Plugins/FX/Niagara/Source/NiagaraEditor/Private/NiagaraGraph.cpp
--- a/Engine/Plugins/FX/Niagara/Source/NiagaraEditor/Private/NiagaraGraph.cpp
+++ b/Engine/Plugins/FX/Niagara/Source/NiagaraEditor/Private/NiagaraGraph.cpp
@@ -945,7 +945,7 @@ bool UNiagaraGraph::RenameParameter(const FNiagaraVariable& Parameter, FName New
 	}
 
 	// Swap metadata to the new parameter
-	if (UNiagaraScriptVariable** Metadata = VariableToScriptVariable.Find(Parameter))
+	if (OldScriptVariable)
 	{
 		if (!bFromStaticSwitch || (bFromStaticSwitch && NumReferences <= 1))
 		{
@@ -958,6 +958,13 @@ bool UNiagaraGraph::RenameParameter(const FNiagaraVariable& Parameter, FName New
 					Variable->DefaultBinding.SetName(NewParameter.GetName());
 				}
 			}
+
+			// Replace the script variable data
+			UNiagaraScriptVariable* NewScriptVariable = NewObject<UNiagaraScriptVariable>(this, FName(), RF_Transactional);
+			NewScriptVariable->Variable = NewParameter;
+			NewScriptVariable->DefaultMode = (*OldScriptVariable)->DefaultMode;
+			NewScriptVariable->DefaultBinding = (*OldScriptVariable)->DefaultBinding;
+			VariableToScriptVariable.Add(NewParameter, NewScriptVariable);
 			VariableToScriptVariable.Remove(Parameter);
 		}
 	}
