Niagara Platform Set.

New feature for controling which platforms use certain aspects of Niagara.
Current users include;
- Emitter Enabled/Disabled switch
- System level scalability settings in EffectType and overrides in System.
- Emitter level scalability settings in EffectType and overrides in Emitter.

This entirely replaces previous Detail Level feature and all related code.

Now we allow users to select from EffectsQualitiy levels as a basic use case for scaling across platforms.
Further fine grained control can be acheived by setting per device profile overrides to this.
For example all Effects Quality levels could be enabled but Android could be specifically disabled. Or some subset of mobile devices.

TODO: Code is in place for this to cook out emitters but there are other impediments stopping this currently.
TODO: Sebastian has more UI changes incomming but as a v1.0 this is funtional.

diff --git a/Engine/Config/BaseScalability.ini b/Engine/Config/BaseScalability.ini
--- a/Engine/Config/BaseScalability.ini
+++ b/Engine/Config/BaseScalability.ini
@@ -397,7 +397,6 @@ r.SSS.HalfRes=1
 r.SSGI.Quality=0
 r.EmitterSpawnRateScale=0.125
 r.ParticleLightQuality=0
-r.DetailLevel=0
 r.SkyAtmosphere.AerialPerspectiveLUT.FastApplyOnOpaque=1 ; Always have FastSkyLUT 1 in this case to avoid wrong sky
 r.SkyAtmosphere.AerialPerspectiveLUT.SampleCountPerSlice=1
 r.SkyAtmosphere.AerialPerspectiveLUT.DepthResolution=8.0
@@ -409,7 +408,6 @@ r.SkyAtmosphere.SampleCountMax=16.0
 r.SkyAtmosphere.TransmittanceLUT.UseSmallFormat=1
 r.SkyAtmosphere.TransmittanceLUT.SampleCount=10.0
 r.SkyAtmosphere.MultiScatteringLUT.SampleCount=15.0
-fx.NiagaraGlobalSpawnCountScale=0.125
 
 [EffectsQuality@1]
 r.TranslucencyLightingVolumeDim=32
@@ -427,7 +425,6 @@ r.SSS.HalfRes=1
 r.SSGI.Quality=1
 r.EmitterSpawnRateScale=0.25
 r.ParticleLightQuality=0
-r.DetailLevel=1
 r.SkyAtmosphere.AerialPerspectiveLUT.FastApplyOnOpaque=1 ; Always have FastSkyLUT 1 in this case to avoid wrong sky
 r.SkyAtmosphere.AerialPerspectiveLUT.SampleCountPerSlice=1
 r.SkyAtmosphere.AerialPerspectiveLUT.DepthResolution=16.0
@@ -439,7 +436,6 @@ r.SkyAtmosphere.SampleCountMax=32.0
 r.SkyAtmosphere.TransmittanceLUT.UseSmallFormat=0
 r.SkyAtmosphere.TransmittanceLUT.SampleCount=10.0
 r.SkyAtmosphere.MultiScatteringLUT.SampleCount=15.0
-fx.NiagaraGlobalSpawnCountScale=0.25
 
 [EffectsQuality@2]
 r.TranslucencyLightingVolumeDim=48
@@ -457,7 +453,6 @@ r.SSS.HalfRes=1
 r.SSGI.Quality=2
 r.EmitterSpawnRateScale=0.5
 r.ParticleLightQuality=1
-r.DetailLevel=2
 r.SkyAtmosphere.AerialPerspectiveLUT.FastApplyOnOpaque=1 ; Always have FastSkyLUT 1 in this case to avoid wrong sky
 r.SkyAtmosphere.AerialPerspectiveLUT.SampleCountPerSlice=1
 r.SkyAtmosphere.AerialPerspectiveLUT.DepthResolution=16.0
@@ -469,7 +464,6 @@ r.SkyAtmosphere.SampleCountMax=32.0
 r.SkyAtmosphere.TransmittanceLUT.UseSmallFormat=0
 r.SkyAtmosphere.TransmittanceLUT.SampleCount=10.0
 r.SkyAtmosphere.MultiScatteringLUT.SampleCount=15.0
-fx.NiagaraGlobalSpawnCountScale=0.5
 
 [EffectsQuality@3]
 r.TranslucencyLightingVolumeDim=64
@@ -487,7 +481,6 @@ r.SSS.HalfRes=1
 r.SSGI.Quality=3
 r.EmitterSpawnRateScale=1.0
 r.ParticleLightQuality=2
-r.DetailLevel=3
 r.SkyAtmosphere.AerialPerspectiveLUT.FastApplyOnOpaque=1 ; Always have FastSkyLUT 1 in this case to avoid wrong sky
 r.SkyAtmosphere.AerialPerspectiveLUT.SampleCountPerSlice=1
 r.SkyAtmosphere.AerialPerspectiveLUT.DepthResolution=16.0
@@ -499,7 +492,6 @@ r.SkyAtmosphere.SampleCountMax=32.0
 r.SkyAtmosphere.TransmittanceLUT.UseSmallFormat=0
 r.SkyAtmosphere.TransmittanceLUT.SampleCount=10.0
 r.SkyAtmosphere.MultiScatteringLUT.SampleCount=15.0
-fx.NiagaraGlobalSpawnCountScale=1.0
 
 [EffectsQuality@Cine]
 r.TranslucencyLightingVolumeDim=64
@@ -517,7 +509,6 @@ r.SSS.HalfRes=0
 r.SSGI.Quality=4
 r.EmitterSpawnRateScale=1.0
 r.ParticleLightQuality=2
-r.DetailLevel=4
 r.SkyAtmosphere.AerialPerspectiveLUT.FastApplyOnOpaque=0
 r.SkyAtmosphere.AerialPerspectiveLUT.SampleCountPerSlice=2
 r.SkyAtmosphere.AerialPerspectiveLUT.DepthResolution=32.0
@@ -529,7 +520,6 @@ r.SkyAtmosphere.SampleCountMax=32.0
 r.SkyAtmosphere.TransmittanceLUT.UseSmallFormat=0
 r.SkyAtmosphere.TransmittanceLUT.SampleCount=30.0
 r.SkyAtmosphere.MultiScatteringLUT.SampleCount=20.0
-fx.NiagaraGlobalSpawnCountScale=1.0
 
 ;-----------------------------------------------------------------------------------------------------------------
 
diff --git a/Engine/Config/Linux/LinuxEngine.ini b/Engine/Config/Linux/LinuxEngine.ini
--- a/Engine/Config/Linux/LinuxEngine.ini
+++ b/Engine/Config/Linux/LinuxEngine.ini
@@ -4,6 +4,7 @@ PlatformHeadroomDB=0
 
 [SystemSettings]
 r.setres=1280x720
+fx.NiagaraAllowRuntimeScalabilityChanges=1
 
 [SystemSettingsEditor]
 ;r.RHICmdBypass=1
diff --git a/Engine/Config/Mac/MacEngine.ini b/Engine/Config/Mac/MacEngine.ini
--- a/Engine/Config/Mac/MacEngine.ini
+++ b/Engine/Config/Mac/MacEngine.ini
@@ -12,6 +12,7 @@ UseAudioMixer=true
 r.setres=1280x720
 g.TimeoutForBlockOnRenderFence=300000
 framepro.ScopeMinTimeMicroseconds=10
+fx.NiagaraAllowRuntimeScalabilityChanges=1
 
 [SystemSettingsEditor]
 ; Parallel rendering has not been tested in the editor so just disable it.
diff --git a/Engine/Config/Windows/WindowsEngine.ini b/Engine/Config/Windows/WindowsEngine.ini
--- a/Engine/Config/Windows/WindowsEngine.ini
+++ b/Engine/Config/Windows/WindowsEngine.ini
@@ -17,6 +17,7 @@ ConfigClass=/Engine/DeviceProfiles/SetDeviceProfileWindows.SetDeviceProfileWindo
 [SystemSettings]
 r.setres=1280x720
 framepro.ScopeMinTimeMicroseconds=10
+fx.NiagaraAllowRuntimeScalabilityChanges=1
 
 [PlatformCrypto]
 PlatformRequiresDataCrypto=True
diff --git a/Engine/Plugins/FX/Niagara/Config/BaseNiagara.ini b/Engine/Plugins/FX/Niagara/Config/BaseNiagara.ini
--- a/Engine/Plugins/FX/Niagara/Config/BaseNiagara.ini
+++ b/Engine/Plugins/FX/Niagara/Config/BaseNiagara.ini
@@ -13,6 +13,9 @@
 
 +EnumRedirects=(OldName="ENiagaraScriptUsage",ValueChanges=(("ENiagaraScriptUsage::SpawnScript","ENiagaraScriptUsage::ParticleSpawnScript"), ("ENiagaraScriptUsage::SpawnScriptInterpolated","ENiagaraScriptUsage::ParticleSpawnScriptInterpolated"),("ENiagaraScriptUsage::UpdateScript","ENiagaraScriptUsage::ParticleUpdateScript"), ("ENiagaraScriptUsage::EventScript","ENiagaraScriptUsage::ParticleEventScript"), ("ENiagaraScriptUsage::SystemScript","ENiagaraScriptUsage::SystemSpawnScript")) )
 
++StructRedirects=(OldName="NiagaraScalabilitySettings",NewName="NiagaraSystemScalabilitySettings")
++StructRedirects=(OldName="NiagaraScalabilityOverrides",NewName="NiagaraSystemScalabilityOverride")
+
 ;+EnumRedirects=(OldName="ENiagaraSortMode",ValueChanges=(("ENiagaraSortMode::SortNone","ENiagaraSortMode::None"), ("ENiagaraSortMode::SortViewDepth","ENiagaraSortMode::ViewDepth"), ("ENiagaraSortMode::SortViewDistance","ENiagaraSortMode::ViewDistance")) )
 
 [/Script/Niagara.NiagaraSettings]
diff --git a/Engine/Plugins/FX/Niagara/Source/Niagara/Classes/NiagaraEffectType.h b/Engine/Plugins/FX/Niagara/Source/Niagara/Classes/NiagaraEffectType.h
--- a/Engine/Plugins/FX/Niagara/Source/Niagara/Classes/NiagaraEffectType.h
+++ b/Engine/Plugins/FX/Niagara/Source/Niagara/Classes/NiagaraEffectType.h
@@ -4,6 +4,7 @@
 
 #include "InGamePerformanceTracker.h"
 #include "RenderCommandFence.h"
+#include "NiagaraPlatformSet.h"
 #include "NiagaraEffectType.generated.h"
 
 #define DEBUG_SCALABILITY_STATE (!UE_BUILD_SHIPPING)
@@ -40,12 +41,18 @@ enum class ENiagaraScalabilityUpdateFrequency
 	Continuous,
 };
 
-/** Settings which control high level scalability of Niagara effects. Effects culled at this level stop ticking and where possible are not spawned at all. Costing the absolute minimum at runtime. */
+//////////////////////////////////////////////////////////////////////////
+
+/** Scalability settings for Niagara Systems for a particular platform set (unless overridden). */
 USTRUCT()
-struct FNiagaraScalabilitySettings
+struct FNiagaraSystemScalabilitySettings
 {
 	GENERATED_USTRUCT_BODY()
 
+	/** The platforms on which these settings are active (unless overridden). */
+	UPROPERTY(EditAnywhere, Category = "Scalability")
+	FNiagaraPlatformSet Platforms;
+
 	UPROPERTY(EditAnywhere, Category = "Scalability", meta = (InlineEditConditionToggle))
 	uint32 bCullByDistance : 1;
 	UPROPERTY(EditAnywhere, Category = "Scalability", meta = (InlineEditConditionToggle))
@@ -74,21 +81,29 @@ struct FNiagaraScalabilitySettings
 	/** Effects will be culled if they go more than this length of time without being rendered. */
 	UPROPERTY(EditAnywhere, Category = "Scalability", meta = (EditCondition = "bCullByMaxTimeWithoutRender"))
 	float MaxTimeWithoutRender;
+	
+	FNiagaraSystemScalabilitySettings();
 
-	/** Scale factor applied to all spawn counts for effects of this type. */
-	//UPROPERTY(EditAnywhere, Category = "Scalability", meta = (EditCondition = "bCullByMaxTimeWithoutRender"))
-	//float SpawnCountScale;
+	void Clear();
+};
+
+/** Container struct for an array of system scalability settings. Enables details customization and data validation. */
+USTRUCT()
+struct FNiagaraSystemScalabilitySettingsArray
+{
+	GENERATED_USTRUCT_BODY()
 
-	//TODO: Assess viability of progressive scaling via biasing the DetailLevel that effects use and reducing this with significance etc.
-	/** Settings that control whether we should bias the DetailLevel used by effects of this type. */
-	//TArray<FNiagaraDetailLevelBiasSettings> DetailLevelBiasSettings;
+	UPROPERTY(EditAnywhere, Category = "Scalability")
+	TArray<FNiagaraSystemScalabilitySettings> Settings;
 };
 
 USTRUCT()
-struct FNiagaraScalabilityOverrides : public FNiagaraScalabilitySettings
+struct FNiagaraSystemScalabilityOverride : public FNiagaraSystemScalabilitySettings
 {
 	GENERATED_USTRUCT_BODY()
-		
+
+	FNiagaraSystemScalabilityOverride();
+
 	//TODO: Detail customization that effectively allows these values to be edit conditions for their respective properties inside FNiagaraScalabilitySettings.
 	UPROPERTY(EditAnywhere, Category = "Override")
 	uint32 bOverrideDistanceSettings : 1;
@@ -98,10 +113,73 @@ struct FNiagaraScalabilityOverrides : public FNiagaraScalabilitySettings
 	uint32 bOverrideOwnerLODSettings : 1;
 	UPROPERTY(EditAnywhere, Category = "Override")
 	uint32 bOverrideTimeSinceRendererSettings : 1;
+};
+
+/** Container struct for an array of system scalability overrides. Enables details customization and data validation. */
+USTRUCT()
+struct FNiagaraSystemScalabilityOverrides
+{
+	GENERATED_USTRUCT_BODY()
+	
+	UPROPERTY(EditAnywhere, Category = "Override")
+	TArray<FNiagaraSystemScalabilityOverride> Overrides;
+};
+
+/** Scalability settings for Niagara Emitters on a particular platform set. */
+USTRUCT()
+struct FNiagaraEmitterScalabilitySettings
+{
+	GENERATED_USTRUCT_BODY()
+
+	/** The platforms on which these settings are active (unless overridden). */
+	UPROPERTY(EditAnywhere, Category = "Scalability")
+	FNiagaraPlatformSet Platforms;
+
+	UPROPERTY(EditAnywhere, Category = "Scalability", meta = (InlineEditConditionToggle))
+	uint32 bScaleSpawnCount : 1;
+
+	/** Effects of this type are culled beyond this distance. */
+	UPROPERTY(EditAnywhere, Category = "Scalability", meta = (EditCondition = "bScaleSpawnCount"))
+	float SpawnCountScale;
+
+	FNiagaraEmitterScalabilitySettings();
+	void Clear();
+};
+
+/** Container struct for an array of emitter scalability settings. Enables details customization and data validation. */
+USTRUCT()
+struct FNiagaraEmitterScalabilitySettingsArray
+{
+	GENERATED_USTRUCT_BODY()
+
+	UPROPERTY(EditAnywhere, Category = "Scalability")
+	TArray<FNiagaraEmitterScalabilitySettings> Settings;
+};
+
+USTRUCT()
+struct FNiagaraEmitterScalabilityOverride : public FNiagaraEmitterScalabilitySettings
+{
+	GENERATED_USTRUCT_BODY()
+
+	FNiagaraEmitterScalabilityOverride();
+
+	//TODO: Detail customization that effectively allows these values to be edit conditions for their respective properties inside FNiagaraScalabilitySettings.
 	UPROPERTY(EditAnywhere, Category = "Override")
 	uint32 bOverrideSpawnCountScale : 1;
 };
 
+/** Container struct for an array of emitter scalability overrides. Enables details customization and data validation. */
+USTRUCT()
+struct FNiagaraEmitterScalabilityOverrides
+{
+	GENERATED_USTRUCT_BODY()
+
+	UPROPERTY(EditAnywhere, Category = "Override")
+	TArray<FNiagaraEmitterScalabilityOverride> Overrides;
+};
+
+//////////////////////////////////////////////////////////////////////////
+
 /** Contains settings and working data shared among many NiagaraSystems that share some commonality of type. For example ImpactFX vs EnvironmentalFX. */
 UCLASS()
 class NIAGARA_API UNiagaraEffectType : public UObject
@@ -126,18 +204,26 @@ class NIAGARA_API UNiagaraEffectType : public UObject
 	ENiagaraCullReaction CullReaction;
 
 	/** Cull settings to use at each detail level. */
+	UPROPERTY()
+	TArray<FNiagaraSystemScalabilitySettings> DetailLevelScalabilitySettings_DEPRECATED;
+
 	UPROPERTY(EditAnywhere, Category = "Effect Type")
-	TArray<FNiagaraScalabilitySettings> DetailLevelScalabilitySettings;
+	FNiagaraSystemScalabilitySettingsArray SystemScalabilitySettings;
 	
+	UPROPERTY(EditAnywhere, Category = "Effect Type")
+	FNiagaraEmitterScalabilitySettingsArray EmitterScalabilitySettings;
+
 	FORCEINLINE int32* GetCycleCounter(bool bGameThread, bool bConcurrent);
 	void ProcessLastFrameCycleCounts();
 
 	//TODO: Dynamic budgetting from perf data.
 	//void ApplyDynamicBudget(float InDynamicBudget_GT, float InDynamicBudget_GT_CNC, float InDynamicBudget_RT);
 
-	void SetDetailLevel(int32 DetailLevel);
+	FORCEINLINE const FNiagaraSystemScalabilitySettingsArray& GetSystemScalabilitySettings()const { return SystemScalabilitySettings; }
+	FORCEINLINE const FNiagaraEmitterScalabilitySettingsArray& GetEmitterScalabilitySettings()const { return EmitterScalabilitySettings; }
 
-	FORCEINLINE const TArray<FNiagaraScalabilitySettings>& GetScalabilitySettings()const { return DetailLevelScalabilitySettings; }
+	const FNiagaraSystemScalabilitySettings& GetActiveSystemScalabilitySettings()const;
+	const FNiagaraEmitterScalabilitySettings& GetActiveEmitterScalabilitySettings()const;
 
 	float GetAverageFrameTime_GT() { return AvgTimeMS_GT; }
 	float GetAverageFrameTime_GT_CNC() { return AvgTimeMS_GT_CNC; }
diff --git a/Engine/Plugins/FX/Niagara/Source/Niagara/Classes/NiagaraEmitter.h b/Engine/Plugins/FX/Niagara/Source/Niagara/Classes/NiagaraEmitter.h
--- a/Engine/Plugins/FX/Niagara/Source/Niagara/Classes/NiagaraEmitter.h
+++ b/Engine/Plugins/FX/Niagara/Source/Niagara/Classes/NiagaraEmitter.h
@@ -10,6 +10,7 @@
 #include "NiagaraCollision.h"
 #include "INiagaraMergeManager.h"
 #include "NiagaraSystemFastPath.h"
+#include "NiagaraEffectType.h"
 #include "NiagaraEmitter.generated.h"
 
 class UMaterial;
@@ -177,31 +178,22 @@ struct FNiagaraEventScriptProperties : public FNiagaraEmitterScriptProperties
 	uint32 MinSpawnNumber;
 };
 
+/** Legacy struct for spawn count scale overrides. This is now done in FNiagaraEmitterScalabilityOverrides*/
 USTRUCT()
 struct FNiagaraDetailsLevelScaleOverrides 
 {
 	GENERATED_BODY()
 
 	FNiagaraDetailsLevelScaleOverrides();
-
-	/** The Scalability spawn count scale used when Effects Quality level is set to Low instead of fx.NiagaraGlobalSpawnCountScale. */
-	UPROPERTY(EditAnywhere, Category = "Scalability")
+	UPROPERTY()
 	float Low;
-
-	/** The Scalability spawn count scale used when Effects Quality level is set to Medium instead of fx.NiagaraGlobalSpawnCountScale. */
-	UPROPERTY(EditAnywhere, Category = "Scalability")
+	UPROPERTY()
 	float Medium;
-
-	/** The Scalability spawn count scale used when Effects Quality level is set to High instead of fx.NiagaraGlobalSpawnCountScale. */
-	UPROPERTY(EditAnywhere, Category = "Scalability")
+	UPROPERTY()
 	float High;
-
-	/** The Scalability spawn count scale used when Effects Quality level is set to Epic instead of fx.NiagaraGlobalSpawnCountScale. */
-	UPROPERTY(EditAnywhere, Category = "Scalability")
+	UPROPERTY()
 	float Epic;
-
-	/** The Scalability spawn count scale used when Effects Quality level is set to Cine instead of fx.NiagaraGlobalSpawnCountScale. */
-	UPROPERTY(EditAnywhere, Category = "Scalability")
+	UPROPERTY()
 	float Cine;
 };
 
@@ -249,6 +241,7 @@ public:
 	virtual void PostEditChangeProperty(struct FPropertyChangedEvent& PropertyChangedEvent) override;
 	NIAGARA_API FOnPropertiesChanged& OnPropertiesChanged();
 	NIAGARA_API FOnRenderersChanged& OnRenderersChanged();
+	bool IsEnabledOnPlatform(const FString& PlatformName);
 #endif
 	void Serialize(FArchive& Ar)override;
 	virtual void PostInitProperties() override;
@@ -302,17 +295,18 @@ public:
 	UPROPERTY(EditAnywhere, Category = "Emitter", meta = (EditCondition = "bFixedBounds"))
 	FBox FixedBounds;
 	
-	/** If the current engine detail level is below MinDetailLevel then this emitter is disabled. */
-	UPROPERTY(EditAnywhere, Category = "Scalability", meta=(EditCondition = "bUseMinDetailLevel"))
-	int32 MinDetailLevel;
-
-	/** If the current engine detail level is above MaxDetailLevel then this emitter is disabled. */
-	UPROPERTY(EditAnywhere, Category = "Scalability", meta=(EditCondition = "bUseMaxDetailLevel"))
-	int32 MaxDetailLevel;
+	UPROPERTY()
+	int32 MinDetailLevel_DEPRECATED;
+	UPROPERTY()
+	int32 MaxDetailLevel_DEPRECATED;
+	UPROPERTY()
+	FNiagaraDetailsLevelScaleOverrides GlobalSpawnCountScaleOverrides_DEPRECATED;
+	
+	UPROPERTY(EditAnywhere, Category = "Scalability")
+	FNiagaraPlatformSet Platforms;
 
-	/** Custom Scalability spawn count scales to override the ones defined by fx.NiagaraGlobalSpawnCountScale in the GScalabilityIni config hierarchy. */
-	UPROPERTY(EditAnywhere, Category = "Scalability", meta = (EditCondition = "bOverrideGlobalSpawnCountScale"))
-	FNiagaraDetailsLevelScaleOverrides GlobalSpawnCountScaleOverrides;
+	UPROPERTY(EditAnywhere, Category = "Scalability")
+	FNiagaraEmitterScalabilityOverrides ScalabilityOverrides;
 
 	/** When enabled, this will spawn using interpolated parameter values and perform a partial update at spawn time. This adds significant additional cost for spawning but will produce much smoother spawning for high spawn rates, erratic frame rates and fast moving emitters. */
 	UPROPERTY(EditAnywhere, Category = "Emitter")
@@ -323,16 +317,16 @@ public:
 	uint32 bFixedBounds : 1;
 
 	/** Whether to use the min detail or not. */
-	UPROPERTY(EditAnywhere, Category = "Scalability", meta = (InlineEditConditionToggle))
-	uint32 bUseMinDetailLevel : 1;
+	UPROPERTY()
+	uint32 bUseMinDetailLevel_DEPRECATED : 1;
 	
 	/** Whether to use the min detail or not. */
-	UPROPERTY(EditAnywhere, Category = "Scalability", meta = (InlineEditConditionToggle))
-	uint32 bUseMaxDetailLevel : 1;
+	UPROPERTY()
+	uint32 bUseMaxDetailLevel_DEPRECATED : 1;
 
-	/** Whether or not to override the global spawn count scale based on Effects Quality levels, as defined by fx.NiagaraGlobalSpawnCountScale in the GScalabilityIni config hierarchy. */
-	UPROPERTY(EditAnywhere, Category = "Scalability", meta = (InlineEditConditionToggle))
-	uint32 bOverrideGlobalSpawnCountScale : 1;
+	/** Legacy bool to control overriding the global spawn count scales. */
+	UPROPERTY()
+	uint32 bOverrideGlobalSpawnCountScale_DEPRECATED : 1;
 
 	/** Do particles in this emitter require a persistent ID? */
 	UPROPERTY(EditAnywhere, Category = "Emitter")
@@ -423,10 +417,11 @@ public:
 	NIAGARA_API UNiagaraEmitter* DuplicateWithoutMerging(UObject* InOuter);
 #endif
 
-	float GetSpawnCountScale(int EffectsQuality = -1); 
+	FORCEINLINE const FNiagaraEmitterScalabilitySettings& GetScalabilitySettings()const { return CurrentScalabilitySettings; }
+
+	/** Returns true if this emitter's platform filter allows it on this platform and EffectsQuality level. */
+	NIAGARA_API bool IsAllowedByScalability()const;
 
-	/** Is this emitter allowed to be enabled by the current system detail level. */
-	NIAGARA_API bool IsAllowedByDetailLevel(int32 DetailLevel)const;
 	NIAGARA_API bool RequiresPersistantIDs()const;
 
 	NIAGARA_API bool IsValid()const;
@@ -484,6 +479,8 @@ public:
 	NIAGARA_API void SetParent(UNiagaraEmitter& InParent);
 #endif
 
+	void OnEffectsQualityChanged();
+
 	void InitFastPathAttributeNames();
 
 	UPROPERTY(EditAnywhere, Category = "Script Fast Path")
@@ -510,6 +507,8 @@ public:
 protected:
 	virtual void BeginDestroy() override;
 
+	void ResolveScalabilitySettings();
+
 #if WITH_EDITORONLY_DATA
 private:
 	void UpdateFromMergedCopy(const INiagaraMergeManager& MergeManager, UNiagaraEmitter* MergedEmitter);
@@ -576,6 +575,8 @@ private:
 
 	MemoryRuntimeEstimation RuntimeEstimation;
 	FCriticalSection EstimationCriticalSection;
+
+	FNiagaraEmitterScalabilitySettings CurrentScalabilitySettings;
 };
 
 
diff --git a/Engine/Plugins/FX/Niagara/Source/Niagara/Classes/NiagaraEmitterInstance.h b/Engine/Plugins/FX/Niagara/Source/Niagara/Classes/NiagaraEmitterInstance.h
--- a/Engine/Plugins/FX/Niagara/Source/Niagara/Classes/NiagaraEmitterInstance.h
+++ b/Engine/Plugins/FX/Niagara/Source/Niagara/Classes/NiagaraEmitterInstance.h
@@ -81,7 +81,7 @@ public:
 		return 0;
 	}
 	FORCEINLINE int32 GetTotalSpawnedParticles()const { return TotalSpawnedParticles; }
-	FORCEINLINE float GetSpawnCountScale(int32 EffectsQuality = -1)const { return CachedEmitter->GetSpawnCountScale(EffectsQuality); }
+	FORCEINLINE const FNiagaraEmitterScalabilitySettings& GetScalabilitySettings()const { return CachedEmitter->GetScalabilitySettings(); }
 
 	NIAGARA_API const FNiagaraEmitterHandle& GetEmitterHandle() const;
 
diff --git a/Engine/Plugins/FX/Niagara/Source/Niagara/Classes/NiagaraPlatformSet.h b/Engine/Plugins/FX/Niagara/Source/Niagara/Classes/NiagaraPlatformSet.h
new file mode 100644
--- /dev/null
+++ b/Engine/Plugins/FX/Niagara/Source/Niagara/Classes/NiagaraPlatformSet.h
@@ -0,0 +1,262 @@
+// Copyright Epic Games, Inc. All Rights Reserved.
+
+#pragma once
+
+#include "CoreMinimal.h"
+#include "NiagaraPlatformSet.generated.h"
+
+class UDeviceProfile;
+
+UENUM()
+enum class ENiagaraPlatformSelectionState : uint8
+{
+	Default,/** Neither explicitly enabled or disabled, this platform is enabled or not based on other settings in the platform set. */
+	Enabled,/** This platform is explicitly disabled. */
+	Disabled,/** This platform is explicitly disabled. */
+};
+
+USTRUCT()
+struct FNiagaraDeviceProfileStateEntry
+{
+	GENERATED_BODY();
+
+	FNiagaraDeviceProfileStateEntry()
+		: EffectsQualityMask(0)
+		, SetEffectsQualityMask(0)
+	{}
+
+	UPROPERTY(EditAnywhere, Category = Profile)
+	FName ProfileName;
+
+	/** The state of each set effects quality. */
+	UPROPERTY(EditAnywhere, Category = Profile)
+	uint32 EffectsQualityMask;
+
+	/** Which Effects Qualities are set. */
+	UPROPERTY(EditAnywhere, Category = Profile)
+	uint32 SetEffectsQualityMask;
+
+	FORCEINLINE ENiagaraPlatformSelectionState GetState(int32 EffectsQuality)const;
+
+	FORCEINLINE void SetState(int32 EffectsQuality, ENiagaraPlatformSelectionState State);
+
+	FORCEINLINE bool AllDefaulted()const { return SetEffectsQualityMask == 0; }
+};
+
+UENUM()
+enum class ENiagaraPlatformSetState : uint8
+{
+	Disabled, //This platform set is disabled.
+	Enabled, //This device profile is enabled but not active.
+	Active, //This device profile is enabled and active now.
+	Unknown UMETA(Hidden),
+};
+
+struct FNiagaraPlatformSet;
+
+USTRUCT()
+struct FNiagaraPlatformSetConflictEntry
+{
+	GENERATED_BODY()
+
+	FNiagaraPlatformSetConflictEntry()
+		: ProfileName(NAME_None)
+		, EffectsQualityMask(0)
+	{}
+
+	UPROPERTY()
+	FName ProfileName;
+
+	/** Mask of conflicting effects qualities for this profile. */
+	UPROPERTY()
+	int32 EffectsQualityMask;
+};
+
+USTRUCT()
+struct FNiagaraPlatformSetConflictInfo
+{
+	GENERATED_BODY()
+
+	FNiagaraPlatformSetConflictInfo()
+		: SetAIndex(INDEX_NONE)
+		, SetBIndex(INDEX_NONE)
+	{}
+	/** Index of the first conflicting set in the checked array. */
+	UPROPERTY()
+	int32 SetAIndex;
+
+	/** Index of the second conflicting set in the checked array. */
+	UPROPERTY()
+	int32 SetBIndex;
+
+	/** Array of all conflicts between these sets. */
+	UPROPERTY()
+	TArray<FNiagaraPlatformSetConflictEntry> Conflicts;
+};
+
+
+USTRUCT()
+struct NIAGARA_API FNiagaraPlatformSet
+{
+	GENERATED_BODY();
+
+	static FORCEINLINE int32 EQFromMask(int32 EQMask);
+	static FORCEINLINE int32 CreateEQMask(int32 EQ);
+	static FText GetEffectsQualityText(int32 EffectsQuality);
+	static FText GetEffectsQualityMaskText(int32 EffectsQualityMask);
+public:
+
+	//Runtime public API
+
+	FNiagaraPlatformSet(int32 EQMask);
+	FNiagaraPlatformSet();
+
+	/** Is this set active right now. i.e. enabled for the current device profile and effects quality. */
+	bool IsActive()const;
+	
+	/** Is this platform set enabled on any effects quality for the passed device profile. Returns the EffectsQualityMask for all enabled effects qualities for this profile. */
+	int32 IsEnabledForDeviceProfile(const UDeviceProfile* DeviceProfile)const;
+
+	/** Is this platform set enabled at this effects quality on any device profile. */
+	bool IsEnabledForEffectsQuality(int32 EffectsQuality)const;
+
+	/** Fill OutProfiles with all device proifles that have been overriden at the passed EffectsQuality. */
+	void GetOverridenDeviceProfiles(int32 EffectsQuality, TArray<UDeviceProfile*>& OutEnabledProfiles, TArray<UDeviceProfile*>& OutDisabledProfiles)const;
+
+	/** Returns true if the current platform can modify it's niagara scalability settings at runtime. */
+	static bool CanChangeScalabilityAtRuntime();
+
+	/**Will force all platform sets to regenerate their cacehd data next time they are used.*/
+	static void InvalidateCachedData();
+
+	static int32 GetEffectQualityMaskForDeviceProfile(const UDeviceProfile* Profile);
+
+	//Editor only public API
+#if WITH_EDITOR
+	/** Returns true if this set is enabled for any profiles on the specified platform. */
+	bool IsEnabledForPlatform(const FString& PlatformName);
+
+	/**
+	Does the enabled set of profiles in this set conflict with those in another set.
+	In some cases this is not allowed.
+	*/
+	//bool Conflicts(const FNiagaraPlatformSet& Other, TArray<const UDeviceProfile*>& OutConflictingProfiles, bool bIncludeProfilesWithVariableEffectsQuality=false)const;
+
+
+	/** Direct state accessors for the UI. */
+	bool IsEffectQualityEnabled(int32 EffectQuality)const;
+	void SetEnabledForEffectQuality(int32 EffectQuality, bool bEnabled);
+	void SetDeviceProfileState(UDeviceProfile* Profile, int32 EffectsQuality, ENiagaraPlatformSelectionState NewState);
+	ENiagaraPlatformSelectionState GetDeviceProfileState(UDeviceProfile* Profile, int32 EffectsQuality)const;
+	
+	/** Invalidates any cached data on this platform set when something has changed. */
+	void OnChanged();
+
+	/** Returns true if the passed platform can modify it's niagara scalability settings at runtime. */
+	static bool CanChangeScalabilityAtRuntime(const UDeviceProfile* DeviceProfile);
+
+	/** Returns true if the passed platform can modify it's niagara scalability settings at runtime. */
+	static bool CanChangeScalabilityAtRuntime(const FString& PlatformName);
+
+	/** Returns true if the passed platform should prune emitters on cook. */
+	static bool ShouldPruneEmittersOnCook(const FString& PlatformName);
+
+	/** Inspects the passed sets and generates an array of all conflicts between these sets. Used to keep arrays of platform sets orthogonal. */
+	static bool GatherConflicts(const TArray<const FNiagaraPlatformSet*>& PlatformSets, TArray<FNiagaraPlatformSetConflictInfo>& OutConflicts);
+#endif
+
+	/** Mask defining which effects qualities this set matches. */
+	UPROPERTY(EditAnywhere, Category = Platforms)
+	int32 EffectsQualityMask;
+
+	/** States of specific device profiles we've set. */
+	UPROPERTY(EditAnywhere, Category = Platforms)
+	TArray<FNiagaraDeviceProfileStateEntry> DeviceProfileStates;
+private:
+
+	//Set from outside when we need to force all cached values to be regenerated. For example on CVar changes.
+	static uint32 LastDirtiedFrame;
+
+	//Last frame we built our cached data. 
+	mutable uint32 LastBuiltFrame;
+
+	mutable bool bEnabledForCurrentProfileAndEffectQuality;
+
+	bool IsEnabled(const UDeviceProfile* Profile, int32 EffectsQuality)const;
+
+#if WITH_EDITOR
+	//Data we pull from platform ini files.
+	struct FPlatformIniSettings
+	{
+		FPlatformIniSettings()
+			:bCanChangeEffectsQualityAtRuntime(0), EffectsQualityMask(INDEX_NONE)
+		{}
+		FPlatformIniSettings(int32 InbCanChangeEffectsQualityAtRuntime, int32 InbPruneEmittersOnCook, int32 InEffectsQualityMask)
+			:bCanChangeEffectsQualityAtRuntime(InbCanChangeEffectsQualityAtRuntime), bPruneEmittersOnCook(InbPruneEmittersOnCook), EffectsQualityMask(InEffectsQualityMask)
+		{}
+		int32 bCanChangeEffectsQualityAtRuntime;
+		int32 bPruneEmittersOnCook;
+		int32 EffectsQualityMask;
+	};
+
+	//Cached data read from platform ini files.
+	static TMap<FName, FPlatformIniSettings> CachedPlatformIniSettings;
+
+	//Cached final EffectsQuality setting for each device profile.
+	static TMap<const UDeviceProfile*, int32> CachedEQMasksPerDeviceProfile;
+
+	static FPlatformIniSettings GetPlatformIniSettings(const FString& PlatformName);
+
+	static int32 GetEffectQualityMaskForPlatform(const FString& PlatformName);
+#endif
+};
+
+
+FORCEINLINE int32 FNiagaraPlatformSet::EQFromMask(int32 EQMask)
+{
+	//Iterate through bits in mask to find the first set quality level.
+	int32 EffectsQuality = 0;
+	while (EQMask != 0)
+	{
+		if (EQMask & 0x1)
+		{
+			return EffectsQuality;
+		}
+		++EffectsQuality;
+		EQMask >>= 1;
+	}
+
+	return INDEX_NONE;
+}
+
+FORCEINLINE int32 FNiagaraPlatformSet::CreateEQMask(int32 EQ)
+{
+	return EQ == INDEX_NONE ? INDEX_NONE : (1 << EQ);
+}
+
+
+FORCEINLINE ENiagaraPlatformSelectionState FNiagaraDeviceProfileStateEntry::GetState(int32 EffectsQuality)const
+{
+	int32 EQMask = FNiagaraPlatformSet::CreateEQMask(EffectsQuality);
+	return (EQMask & SetEffectsQualityMask) == 0 ? ENiagaraPlatformSelectionState::Default : ((EQMask & EffectsQualityMask) != 0 ? ENiagaraPlatformSelectionState::Enabled : ENiagaraPlatformSelectionState::Disabled);
+}
+
+FORCEINLINE void FNiagaraDeviceProfileStateEntry::SetState(int32 EffectsQuality, ENiagaraPlatformSelectionState State)
+{
+	int32 EQMask = FNiagaraPlatformSet::CreateEQMask(EffectsQuality);
+	if (State == ENiagaraPlatformSelectionState::Default)
+	{
+		SetEffectsQualityMask &= ~EQMask;
+		EffectsQualityMask &= ~EQMask;
+	}
+	else if (State == ENiagaraPlatformSelectionState::Enabled)
+	{
+		SetEffectsQualityMask |= EQMask;
+		EffectsQualityMask |= EQMask;
+	}
+	else //(State == ENiagaraPlatformSelectionState::Disabled)
+	{
+		SetEffectsQualityMask |= EQMask;
+		EffectsQualityMask &= ~EQMask;
+	}
+}
\ No newline at end of file
diff --git a/Engine/Plugins/FX/Niagara/Source/Niagara/Classes/NiagaraScript.h b/Engine/Plugins/FX/Niagara/Source/Niagara/Classes/NiagaraScript.h
--- a/Engine/Plugins/FX/Niagara/Source/Niagara/Classes/NiagaraScript.h
+++ b/Engine/Plugins/FX/Niagara/Source/Niagara/Classes/NiagaraScript.h
@@ -136,10 +136,6 @@ public:
 	TArray<FString> AdditionalDefines;
 #endif
 
-	/** Bitfield of supported detail levels in this compile.*/
-	UPROPERTY()
-	uint32 DetailLevelMask;
-
 	/** Whether or not we need to bake Rapid Iteration params. True to keep params, false to bake.*/
 	UPROPERTY()
 	uint32 bUsesRapidIterationParams : 1;
@@ -176,7 +172,6 @@ public:
 	FNiagaraVMExecutableDataId()
 		: CompilerVersionID()
 		, ScriptUsageType(ENiagaraScriptUsage::Function)
-		, DetailLevelMask(0xFFFFFFFF)
 		, bUsesRapidIterationParams(true)
 		, bInterpolatedSpawn(false)
 		, bRequiresPersistentIDs(false)
diff --git a/Engine/Plugins/FX/Niagara/Source/Niagara/Classes/NiagaraScriptSourceBase.h b/Engine/Plugins/FX/Niagara/Source/Niagara/Classes/NiagaraScriptSourceBase.h
--- a/Engine/Plugins/FX/Niagara/Source/Niagara/Classes/NiagaraScriptSourceBase.h
+++ b/Engine/Plugins/FX/Niagara/Source/Niagara/Classes/NiagaraScriptSourceBase.h
@@ -32,11 +32,7 @@ public:
 	virtual int32 GetDependentRequestCount() const = 0;
 	virtual TSharedPtr<FNiagaraCompileRequestDataBase, ESPMode::ThreadSafe> GetDependentRequest(int32 Index) = 0;
 	virtual FName ResolveEmitterAlias(FName VariableName) const = 0;
-	virtual uint32 GetDetailLevelMask() const = 0;
 	virtual bool GetUseRapidIterationParams() const = 0;
-
-	static const uint32 CookForAllDetailLevelMask = 0xFFFFFFFF;
-
 };
 
 class FNiagaraCompileOptions
@@ -63,7 +59,6 @@ public:
 	FString Name;
 	int32 TargetUsageBitmask;
 	TArray<FString> AdditionalDefines;
-	uint32 DetailLevelMask = 0xFFFFFFFF;
 };
 
 struct FNiagaraParameterStore;
diff --git a/Engine/Plugins/FX/Niagara/Source/Niagara/Classes/NiagaraSystem.h b/Engine/Plugins/FX/Niagara/Source/Niagara/Classes/NiagaraSystem.h
--- a/Engine/Plugins/FX/Niagara/Source/Niagara/Classes/NiagaraSystem.h
+++ b/Engine/Plugins/FX/Niagara/Source/Niagara/Classes/NiagaraSystem.h
@@ -356,10 +356,9 @@ public:
 	FORCEINLINE int32* GetCycleCounter(bool bGameThread, bool bConcurrent);
 
 	UNiagaraEffectType* GetEffectType()const;
-	const FNiagaraScalabilitySettings& GetScalabilitySettings(int32 DetailLevel=INDEX_NONE);
+	FORCEINLINE const FNiagaraSystemScalabilitySettings& GetScalabilitySettings() { return CurrentScalabilitySettings; }
 	
-	void ResolveScalabilityOverrides();
-	void OnDetailLevelChanges(int32 DetailLevel);
+	void OnEffectsQualityChanged();
 
 	/** Whether or not fixed bounds are enabled. */
 	UPROPERTY(EditAnywhere, Category = "System", meta = (InlineEditConditionToggle))
@@ -402,6 +401,7 @@ private:
 	void InitEmitterDataSetCompiledData(FNiagaraDataSetCompiledData& DataSetToInit, const UNiagaraEmitter* InAssociatedEmitter, const FNiagaraEmitterHandle& InAssociatedEmitterHandle);
 #endif
 
+	void ResolveScalabilitySettings();
 	void UpdatePostCompileDIInfo();
 
 protected:
@@ -411,8 +411,11 @@ protected:
 	UPROPERTY(EditAnywhere, Category = "System", meta=(InlineEditConditionToggle))
 	bool bOverrideScalabilitySettings;
 
+	UPROPERTY()
+	TArray<FNiagaraSystemScalabilityOverride> ScalabilityOverrides_DEPRECATED;
+
 	UPROPERTY(EditAnywhere, Category = "System", meta = (EditCondition="bOverrideScalabilitySettings"))
-	TArray<FNiagaraScalabilityOverrides> ScalabilityOverrides;
+	FNiagaraSystemScalabilityOverrides SystemScalabilityOverrides;
 
 	/** Handles to the emitter this System will simulate. */
 	UPROPERTY()
@@ -495,11 +498,7 @@ protected:
 	mutable TStatId StatID_RT_CNC;
 #endif
 
-	/** Resolved results of this system's overrides applied on top of it's effect type settings. */
-	UPROPERTY()
-	TArray<FNiagaraScalabilitySettings> ResolvedScalabilitySettings;
-
-	FNiagaraScalabilitySettings CurrentScalabilitySettings;
+	FNiagaraSystemScalabilitySettings CurrentScalabilitySettings;
 };
 
 extern int32 GEnableNiagaraRuntimeCycleCounts;
diff --git a/Engine/Plugins/FX/Niagara/Source/Niagara/Private/NiagaraComponent.cpp b/Engine/Plugins/FX/Niagara/Source/Niagara/Private/NiagaraComponent.cpp
--- a/Engine/Plugins/FX/Niagara/Source/Niagara/Private/NiagaraComponent.cpp
+++ b/Engine/Plugins/FX/Niagara/Source/Niagara/Private/NiagaraComponent.cpp
@@ -440,9 +440,7 @@ UNiagaraComponent::UNiagaraComponent(const FObjectInitializer& ObjectInitializer
 	, bAutoDestroy(false)
 	, MaxTimeBeforeForceUpdateTransform(5.0f)
 #if WITH_EDITOR
-	, PreviewDetailLevel(INDEX_NONE)
 	, PreviewLODDistance(0.0f)
-	, bEnablePreviewDetailLevel(false)
 	, bEnablePreviewLODDistance(false)
 	, bWaitForCompilationOnActivate(false)
 #endif
@@ -1795,25 +1793,12 @@ void UNiagaraComponent::SetMaxSimTime(float InMaxTime)
 }
 
 #if WITH_NIAGARA_COMPONENT_PREVIEW_DATA
-void UNiagaraComponent::SetPreviewDetailLevel(bool bInEnablePreviewDetailLevel, int32 InPreviewDetailLevel)
-{
-	bool bReInit = bEnablePreviewDetailLevel != bInEnablePreviewDetailLevel || (bEnablePreviewDetailLevel && PreviewDetailLevel != InPreviewDetailLevel);
-
-	bEnablePreviewDetailLevel = bInEnablePreviewDetailLevel;
-	PreviewDetailLevel = InPreviewDetailLevel;
-	if (bReInit)
-	{
-		ReinitializeSystem();
-	}
-}
-
 void UNiagaraComponent::SetPreviewLODDistance(bool bInEnablePreviewLODDistance, float InPreviewLODDistance)
 {
 	bEnablePreviewLODDistance = bInEnablePreviewLODDistance;
 	PreviewLODDistance = InPreviewLODDistance;
 }
 #else
-void UNiagaraComponent::SetPreviewDetailLevel(bool bInEnablePreviewDetailLevel, int32 InPreviewDetailLevel){}
 void UNiagaraComponent::SetPreviewLODDistance(bool bInEnablePreviewLODDistance, float InPreviewLODDistance){}
 #endif
 
diff --git a/Engine/Plugins/FX/Niagara/Source/Niagara/Private/NiagaraEffectType.cpp b/Engine/Plugins/FX/Niagara/Source/Niagara/Private/NiagaraEffectType.cpp
--- a/Engine/Plugins/FX/Niagara/Source/Niagara/Private/NiagaraEffectType.cpp
+++ b/Engine/Plugins/FX/Niagara/Source/Niagara/Private/NiagaraEffectType.cpp
@@ -42,6 +42,64 @@ bool UNiagaraEffectType::IsReadyForFinishDestroy()
 void UNiagaraEffectType::PostLoad()
 {
 	Super::PostLoad();
+
+	const int32 NiagaraVer = GetLinkerCustomVersion(FNiagaraCustomVersion::GUID);
+	if (NiagaraVer < FNiagaraCustomVersion::PlatformScalingRefactor)
+	{
+		for (int32 DL=0; DL < DetailLevelScalabilitySettings_DEPRECATED.Num(); ++DL)
+		{
+			//Transfer system settings to new platform structure.
+			FNiagaraSystemScalabilitySettings& DLSettings = DetailLevelScalabilitySettings_DEPRECATED[DL];
+			FNiagaraSystemScalabilitySettings& NewSettigns = SystemScalabilitySettings.Settings.AddDefaulted_GetRef();
+			NewSettigns = DLSettings;
+			NewSettigns.Platforms = FNiagaraPlatformSet(FNiagaraPlatformSet::CreateEQMask(DL));
+		}
+
+		//Ensure all types replicate the former global spawn count scales
+		float LegacySpawnCountScales[5] = { 0.125f, 0.25f, 0.5f, 1.0f, 1.0f };
+		for (int32 DL = 0; DL < 5; ++DL)
+		{
+			if (LegacySpawnCountScales[DL] != 1.0f)//Don't bother adding one for levels with 1.0 as we'll fallback to no scaling.
+			{
+				FNiagaraEmitterScalabilitySettings& EmitterSettings = EmitterScalabilitySettings.Settings.AddDefaulted_GetRef();
+				EmitterSettings.bScaleSpawnCount = true;
+				EmitterSettings.SpawnCountScale = LegacySpawnCountScales[DL];
+				EmitterSettings.Platforms = FNiagaraPlatformSet(FNiagaraPlatformSet::CreateEQMask(DL));
+			}
+		}
+	}
+}
+
+const FNiagaraSystemScalabilitySettings& UNiagaraEffectType::GetActiveSystemScalabilitySettings()const
+{
+	for (const FNiagaraSystemScalabilitySettings& Settings : SystemScalabilitySettings.Settings)
+	{
+		if (Settings.Platforms.IsActive())
+		{
+			return Settings;
+		}
+	}
+
+	//UE_LOG(LogNiagara, Warning, TEXT("Could not find active system scalability settings for EffectType %s"), *GetFullName());
+
+	static FNiagaraSystemScalabilitySettings Dummy;
+	return Dummy;
+}
+
+const FNiagaraEmitterScalabilitySettings& UNiagaraEffectType::GetActiveEmitterScalabilitySettings()const
+{
+	for (const FNiagaraEmitterScalabilitySettings& Settings : EmitterScalabilitySettings.Settings)
+	{
+		if (Settings.Platforms.IsActive())
+		{
+			return Settings;
+		}
+	}
+	
+	//UE_LOG(LogNiagara, Warning, TEXT("Could not find active emitter scalability settings for EffectType %s"), *GetFullName());
+
+	static FNiagaraEmitterScalabilitySettings Dummy;
+	return Dummy;
 }
 
 #if WITH_EDITOR
@@ -56,7 +114,7 @@ void UNiagaraEffectType::PostEditChangeProperty(struct FPropertyChangedEvent& Pr
 		UNiagaraSystem* System = *It;
 		if (System->GetEffectType() == this)
 		{
-			System->ResolveScalabilityOverrides();
+			System->OnEffectsQualityChanged();
 			UpdateContext.Add(System, true);
 		}
 	}
@@ -100,3 +158,48 @@ void UNiagaraEffectType::ProcessLastFrameCycleCounts()
 // 		MinSignificanceFromPerf = 0.0f;
 // 	}
 // }
+
+
+//////////////////////////////////////////////////////////////////////////
+
+FNiagaraSystemScalabilityOverride::FNiagaraSystemScalabilityOverride()
+	: bOverrideDistanceSettings(false)
+	, bOverrideInstanceCountSettings(false)
+	, bOverrideOwnerLODSettings(false)
+	, bOverrideTimeSinceRendererSettings(false)
+{
+}
+
+FNiagaraSystemScalabilitySettings::FNiagaraSystemScalabilitySettings()
+{
+	Clear();
+}
+
+void FNiagaraSystemScalabilitySettings::Clear()
+{
+	Platforms = FNiagaraPlatformSet();
+	bCullByDistance = false;
+	bCullByMaxOwnerLOD = false;
+	bCullByMaxTimeWithoutRender = false;
+	bCullMaxInstanceCount = false;
+	MaxDistance = 0.0f;
+	MaxInstances = 0;
+	MaxOwnerLOD = 0;
+	MaxTimeWithoutRender = 0.0f;
+}
+
+FNiagaraEmitterScalabilitySettings::FNiagaraEmitterScalabilitySettings()
+{
+	Clear();
+}
+
+void FNiagaraEmitterScalabilitySettings::Clear()
+{
+	SpawnCountScale = 1.0f;
+	bScaleSpawnCount = false;
+}
+
+FNiagaraEmitterScalabilityOverride::FNiagaraEmitterScalabilityOverride()
+	: bOverrideSpawnCountScale(false)
+{
+}
diff --git a/Engine/Plugins/FX/Niagara/Source/Niagara/Private/NiagaraEmitter.cpp b/Engine/Plugins/FX/Niagara/Source/Niagara/Private/NiagaraEmitter.cpp
--- a/Engine/Plugins/FX/Niagara/Source/Niagara/Private/NiagaraEmitter.cpp
+++ b/Engine/Plugins/FX/Niagara/Source/Niagara/Private/NiagaraEmitter.cpp
@@ -14,6 +14,10 @@
 #include "NiagaraStats.h"
 #include "Modules/ModuleManager.h"
 #include "Misc/ConfigCacheIni.h"
+#include "Interfaces/ITargetPlatform.h"
+
+PRAGMA_DISABLE_INLINING
+PRAGMA_DISABLE_OPTIMIZATION
 
 #if WITH_EDITOR
 const FName UNiagaraEmitter::PrivateMemberNames::EventHandlerScriptProps = GET_MEMBER_NAME_CHECKED(UNiagaraEmitter, EventHandlerScriptProps);
@@ -53,30 +57,13 @@ static FAutoConsoleVariableRef CVarEnableEmitterChangeIdMergeLogging(
 	ECVF_Default
 );
 
-float UNiagaraEmitter::GetSpawnCountScale(int EffectsQuality)
-{
-	if (bOverrideGlobalSpawnCountScale)
-	{
-		switch (EffectsQuality)
-		{
-		case 0: return GlobalSpawnCountScaleOverrides.Low;
-		case 1: return GlobalSpawnCountScaleOverrides.Medium;
-		case 2: return GlobalSpawnCountScaleOverrides.High;
-		case 3: return GlobalSpawnCountScaleOverrides.Epic;
-		case 4: return GlobalSpawnCountScaleOverrides.Cine;
-		}
-	}
-	
-	return INiagaraModule::GetGlobalSpawnCountScale(); // default to the global spawn scale when invalid effects quality
-}
-
 FNiagaraDetailsLevelScaleOverrides::FNiagaraDetailsLevelScaleOverrides()
 {
-	GConfig->GetFloat(TEXT("EffectsQuality@0"), TEXT("fx.NiagaraGlobalSpawnCountScale"), Low, GScalabilityIni);
-	GConfig->GetFloat(TEXT("EffectsQuality@1"), TEXT("fx.NiagaraGlobalSpawnCountScale"), Medium, GScalabilityIni);
-	GConfig->GetFloat(TEXT("EffectsQuality@2"), TEXT("fx.NiagaraGlobalSpawnCountScale"), High, GScalabilityIni);
-	GConfig->GetFloat(TEXT("EffectsQuality@3"), TEXT("fx.NiagaraGlobalSpawnCountScale"), Epic, GScalabilityIni);
-	GConfig->GetFloat(TEXT("EffectsQuality@Cine"), TEXT("fx.NiagaraGlobalSpawnCountScale"), Cine, GScalabilityIni);
+	Low = 0.125f;
+	Medium = 0.25f;
+	High = 0.5f;
+	Epic = 1.0f;
+	Cine = 1.0f;
 }
 
 void FNiagaraEmitterScriptProperties::InitDataSetAccess()
@@ -128,12 +115,12 @@ UNiagaraEmitter::UNiagaraEmitter(const FObjectInitializer& Initializer)
 : Super(Initializer)
 , PreAllocationCount(0)
 , FixedBounds(FBox(FVector(-100), FVector(100)))
-, MinDetailLevel(0)
-, MaxDetailLevel(4)
+, MinDetailLevel_DEPRECATED(0)
+, MaxDetailLevel_DEPRECATED(4)
 , bInterpolatedSpawning(false)
 , bFixedBounds(false)
-, bUseMinDetailLevel(false)
-, bUseMaxDetailLevel(false)
+, bUseMinDetailLevel_DEPRECATED(false)
+, bUseMaxDetailLevel_DEPRECATED(false)
 , bRequiresPersistentIDs(false)
 , MaxDeltaTimePerTick(0.125)
 , DefaultShaderStageIndex(0)
@@ -168,6 +155,8 @@ void UNiagaraEmitter::PostInitProperties()
 
 	}
 	UniqueEmitterName = TEXT("Emitter");
+
+	ResolveScalabilitySettings();
 }
 
 #if WITH_EDITORONLY_DATA
@@ -308,6 +297,55 @@ void UNiagaraEmitter::PostLoad()
 		}
 	}
 
+	const int32 NiagaraVer = GetLinkerCustomVersion(FNiagaraCustomVersion::GUID);
+	if (NiagaraVer < FNiagaraCustomVersion::PlatformScalingRefactor)
+	{
+		int32 MinDetailLevel = bUseMaxDetailLevel_DEPRECATED ? MinDetailLevel_DEPRECATED : 0;
+		int32 MaxDetailLevel = bUseMaxDetailLevel_DEPRECATED ? MaxDetailLevel_DEPRECATED : 4;
+		int32 NewEQMask = 0;
+		//Currently all detail levels were direct mappings to effects quality so just transfer them over to the new mask in PlatformSet.
+		for (int32 EQ = MinDetailLevel; EQ <= MaxDetailLevel; ++EQ)
+		{
+			NewEQMask |= (1 << EQ);
+		}
+
+		Platforms = FNiagaraPlatformSet(NewEQMask);
+
+		//Transfer spawn rate scaling overrides
+		if (bOverrideGlobalSpawnCountScale_DEPRECATED)
+		{
+			FNiagaraEmitterScalabilityOverride& LowOverride = ScalabilityOverrides.Overrides.AddDefaulted_GetRef();
+			LowOverride.Platforms = FNiagaraPlatformSet(FNiagaraPlatformSet::CreateEQMask(0));
+			LowOverride.bOverrideSpawnCountScale = true;
+			LowOverride.bScaleSpawnCount = true;
+			LowOverride.SpawnCountScale = GlobalSpawnCountScaleOverrides_DEPRECATED.Low;
+
+			FNiagaraEmitterScalabilityOverride& MediumOverride = ScalabilityOverrides.Overrides.AddDefaulted_GetRef();
+			MediumOverride.Platforms = FNiagaraPlatformSet(FNiagaraPlatformSet::CreateEQMask(1));
+			MediumOverride.bOverrideSpawnCountScale = true;
+			MediumOverride.bScaleSpawnCount = true;
+			MediumOverride.SpawnCountScale = GlobalSpawnCountScaleOverrides_DEPRECATED.Medium;
+
+			FNiagaraEmitterScalabilityOverride& HighOverride = ScalabilityOverrides.Overrides.AddDefaulted_GetRef();
+			HighOverride.Platforms = FNiagaraPlatformSet(FNiagaraPlatformSet::CreateEQMask(2));
+			HighOverride.bOverrideSpawnCountScale = true;
+			HighOverride.bScaleSpawnCount = true;
+			HighOverride.SpawnCountScale = GlobalSpawnCountScaleOverrides_DEPRECATED.High;
+
+			FNiagaraEmitterScalabilityOverride& EpicOverride = ScalabilityOverrides.Overrides.AddDefaulted_GetRef();
+			EpicOverride.Platforms = FNiagaraPlatformSet(FNiagaraPlatformSet::CreateEQMask(3));
+			EpicOverride.bOverrideSpawnCountScale = true;
+			EpicOverride.bScaleSpawnCount = true;
+			EpicOverride.SpawnCountScale = GlobalSpawnCountScaleOverrides_DEPRECATED.Epic;
+
+			FNiagaraEmitterScalabilityOverride& CineOverride = ScalabilityOverrides.Overrides.AddDefaulted_GetRef();
+			CineOverride.Platforms = FNiagaraPlatformSet(FNiagaraPlatformSet::CreateEQMask(4));
+			CineOverride.bOverrideSpawnCountScale = true;
+			CineOverride.bScaleSpawnCount = true;
+			CineOverride.SpawnCountScale = GlobalSpawnCountScaleOverrides_DEPRECATED.Cine;
+		}
+	}
+
 	if (!GPUComputeScript)
 	{
 		GPUComputeScript = NewObject<UNiagaraScript>(this, "GPUComputeScript", EObjectFlags::RF_Transactional);
@@ -482,6 +520,8 @@ void UNiagaraEmitter::PostLoad()
 		}
 	}
 #endif
+
+	ResolveScalabilitySettings();
 }
 
 #if WITH_EDITOR
@@ -606,17 +646,8 @@ void UNiagaraEmitter::PostEditChangeProperty(struct FPropertyChangedEvent& Prope
 		UNiagaraSystem::RequestCompileForEmitter(this);
 #endif
 	}
-	if (PropertyName == GET_MEMBER_NAME_CHECKED(UNiagaraEmitter, bOverrideGlobalSpawnCountScale))
-	{
-		if (GraphSource != nullptr)
-		{
-			GraphSource->MarkNotSynchronized(TEXT("Emitter Override Global Spawn Count Scale changed."));
-		}
 
-#if WITH_EDITORONLY_DATA
-		UNiagaraSystem::RequestCompileForEmitter(this);
-#endif
-	}
+	ResolveScalabilitySettings();
 
 	ThumbnailImageOutOfDate = true;
 	UpdateChangeId(TEXT("PostEditChangeProperty"));
@@ -633,6 +664,12 @@ UNiagaraEmitter::FOnPropertiesChanged& UNiagaraEmitter::OnRenderersChanged()
 {
 	return OnRenderersChangedDelegate;
 }
+
+bool UNiagaraEmitter::IsEnabledOnPlatform(const FString& PlatformName)
+{
+	bool bCanPrune = FNiagaraPlatformSet::ShouldPruneEmittersOnCook(PlatformName);
+	return bCanPrune && Platforms.IsEnabledForPlatform(PlatformName);
+}
 #endif
 
 
@@ -767,14 +804,9 @@ UNiagaraScript* UNiagaraEmitter::GetScript(ENiagaraScriptUsage Usage, FGuid Usag
 	return nullptr;
 }
 
-bool UNiagaraEmitter::IsAllowedByDetailLevel(int32 DetailLevel)const
+bool UNiagaraEmitter::IsAllowedByScalability()const
 {
-	if ((bUseMinDetailLevel && DetailLevel < MinDetailLevel) || (bUseMaxDetailLevel && DetailLevel > MaxDetailLevel))
-	{
-		return false;
-	}
-
-	return true;
+	return Platforms.IsActive();
 }
 
 bool UNiagaraEmitter::RequiresPersistantIDs()const
@@ -1444,6 +1476,36 @@ void UNiagaraEmitter::SetParent(UNiagaraEmitter& InParent)
 }
 #endif
 
+void UNiagaraEmitter::ResolveScalabilitySettings()
+{
+	CurrentScalabilitySettings.Clear();
+
+	if (UNiagaraSystem* Owner = GetTypedOuter<UNiagaraSystem>())
+	{
+		if(UNiagaraEffectType* ActualEffectType = Owner->GetEffectType())
+		{
+			CurrentScalabilitySettings = ActualEffectType->GetActiveEmitterScalabilitySettings();
+		}
+	}
+
+	for (FNiagaraEmitterScalabilityOverride& Override : ScalabilityOverrides.Overrides)
+	{
+		if (Override.Platforms.IsActive())
+		{
+			if (Override.bOverrideSpawnCountScale)
+			{
+				CurrentScalabilitySettings.bScaleSpawnCount = Override.bScaleSpawnCount;
+				CurrentScalabilitySettings.SpawnCountScale = Override.SpawnCountScale;
+			}
+		}
+	}
+}
+
+void UNiagaraEmitter::OnEffectsQualityChanged()
+{
+	ResolveScalabilitySettings();
+}
+
 void UNiagaraEmitter::InitFastPathAttributeNames()
 {
 	auto InitParameters = [](const FNiagaraParameters& Parameters, const FString& EmitterName, FNiagaraFastPathAttributeNames& FastPathParameterNames)
@@ -1474,3 +1536,6 @@ void UNiagaraEmitter::InitFastPathAttributeNames()
 	InitParameters(SpawnScriptProps.Script->GetVMExecutableData().Parameters, UniqueEmitterName, SpawnFastPathAttributeNames);
 	InitParameters(UpdateScriptProps.Script->GetVMExecutableData().Parameters, UniqueEmitterName, UpdateFastPathAttributeNames);
 }
+
+PRAGMA_ENABLE_INLINING
+PRAGMA_ENABLE_OPTIMIZATION
\ No newline at end of file
diff --git a/Engine/Plugins/FX/Niagara/Source/Niagara/Private/NiagaraEmitterInstance.cpp b/Engine/Plugins/FX/Niagara/Source/Niagara/Private/NiagaraEmitterInstance.cpp
--- a/Engine/Plugins/FX/Niagara/Source/Niagara/Private/NiagaraEmitterInstance.cpp
+++ b/Engine/Plugins/FX/Niagara/Source/Niagara/Private/NiagaraEmitterInstance.cpp
@@ -209,10 +209,9 @@ void FNiagaraEmitterInstance::Dump()const
 
 bool FNiagaraEmitterInstance::IsAllowedToExecute() const
 {
-	int32 DetailLevel = ParentSystemInstance->GetDetailLevel();
 	const FNiagaraEmitterHandle& EmitterHandle = GetEmitterHandle();
 	return EmitterHandle.GetIsEnabled() &&
-		CachedEmitter->IsAllowedByDetailLevel(DetailLevel) &&
+		CachedEmitter->IsAllowedByScalability() &&
 		// TODO: fall back to CPU sim instead once we have scalability functionality to do so
 		(CachedEmitter->SimTarget != ENiagaraSimTarget::GPUComputeSim || (Batcher && FNiagaraUtilities::AllowGPUParticles(Batcher->GetShaderPlatform())));
 }
diff --git a/Engine/Plugins/FX/Niagara/Source/Niagara/Private/NiagaraModule.cpp b/Engine/Plugins/FX/Niagara/Source/Niagara/Private/NiagaraModule.cpp
--- a/Engine/Plugins/FX/Niagara/Source/Niagara/Private/NiagaraModule.cpp
+++ b/Engine/Plugins/FX/Niagara/Source/Niagara/Private/NiagaraModule.cpp
@@ -23,6 +23,7 @@
 #include "DeviceProfiles/DeviceProfileManager.h"
 #include "Interfaces/ITargetPlatform.h"
 #include "DeviceProfiles/DeviceProfile.h"
+#include "Scalability.h"
 
 IMPLEMENT_MODULE(INiagaraModule, Niagara);
 
@@ -39,7 +40,7 @@ static FAutoConsoleVariableRef CVarLogCompileIdGeneration(
 
 float INiagaraModule::EngineGlobalSpawnCountScale = 1.0f;
 float INiagaraModule::EngineGlobalSystemCountScale = 1.0f;
-int32 INiagaraModule::EngineDetailLevel = 4;
+int32 INiagaraModule::EngineEffectsQuality = 4;
 
 int32 GEnableVerboseNiagaraChangeIdLogging = 0;
 static FAutoConsoleVariableRef CVarEnableVerboseNiagaraChangeIdLogging(
@@ -59,34 +60,6 @@ static TAutoConsoleVariable<int32> CVarUseShaderStages(
 	TEXT("Enable or not the shader stages within Niagara (WIP feature only there for temporary testing)."),
 	ECVF_Default);
 
-/**
-Detail Level CVar.
-Effectively replaces the DetaiMode feature but allows for a rolling range of new hardware and emitters to target them.
-TODO: Possible that this might be more broadly useful across the engine as a replacement for DetailMode so placing in "r." rather than "fx."
-*/
-static TAutoConsoleVariable<float> CVarDetailLevel(
-	TEXT("r.DetailLevel"),
-	4,
-	TEXT("The detail level for use with Niagara.\n")
-	TEXT("If this value does not fall within an Emitter's MinDetailLevel and MaxDetailLevel range, then it will be disabled. \n")
-	TEXT("\n")
-	TEXT("Default = 4"),
-	ECVF_Scalability);
-
-static TAutoConsoleVariable<float> CVarPruneEmittersOnCookByDetailLevel(
-	TEXT("fx.NiagaraPruneEmittersOnCookByDetailLevel"),
-	0,
-	TEXT("Whether to eliminate all emitters that don't match the detail level.\n")
-	TEXT("This will only work if scalability settings affecting detail level can not be changed at runtime (depends on platform).\n"),
-	ECVF_ReadOnly);
-
-static FAutoConsoleVariableRef CVarNiaraGlobalSpawnCountScale(
-	TEXT("fx.NiagaraGlobalSpawnCountScale"),
-	INiagaraModule::EngineGlobalSpawnCountScale,
-	TEXT("A global scale on spawn counts in Niagara. \n"),
-	ECVF_Scalability
-);
-
 static FAutoConsoleVariableRef CVarNiaraGlobalSystemCountScale(
 	TEXT("fx.NiagaraGlobalSystemCountScale"),
 	INiagaraModule::EngineGlobalSystemCountScale,
@@ -194,9 +167,6 @@ void INiagaraModule::StartupModule()
 	FModuleManager::Get().LoadModule(TEXT("NiagaraEditor"));
 #endif
 
-	CVarDetailLevel.AsVariable()->SetOnChangedCallback(FConsoleVariableDelegate::CreateRaw(this, &INiagaraModule::OnChangeDetailLevel));
-	EngineDetailLevel = CVarDetailLevel.GetValueOnGameThread();
-
 	//Init commonly used FNiagaraVariables
 
 	Engine_DeltaTime = FNiagaraVariable(FNiagaraTypeDefinition::GetFloatDef(), TEXT("Engine.DeltaTime"));
@@ -322,7 +292,6 @@ void INiagaraModule::ShutdownModule()
 	INiagaraShaderModule& NiagaraShaderModule = FModuleManager::LoadModuleChecked<INiagaraShaderModule>("NiagaraShader");
 	NiagaraShaderModule.ResetOnRequestDefaultDataInterfaceHandler();
 
-	CVarDetailLevel.AsVariable()->SetOnChangedCallback(FConsoleVariableDelegate());
 	ShutdownRenderingResources();
 }
 
@@ -428,61 +397,19 @@ void INiagaraModule::UnregisterPrecompiler(FDelegateHandle DelegateHandle)
 
 #endif
 
-
-bool INiagaraModule::IsTargetPlatformIncludedInLevelRangeForCook(const ITargetPlatform* InTargetPlatform, const UNiagaraEmitter* InEmitter)
+void INiagaraModule::OnEffectsQualityChanged(int32 NewEffectsQuality)
 {
-#if WITH_EDITORONLY_DATA
-	if (UDeviceProfile* DeviceProfile = UDeviceProfileManager::Get().FindProfile(InTargetPlatform->IniPlatformName()))
+	if (NewEffectsQuality != EngineEffectsQuality && FNiagaraPlatformSet::CanChangeScalabilityAtRuntime())
 	{
-		// get local scalability CVars that could cull this actor
-		int32 CVarCullBasedOnDetailLevel;
-		if (DeviceProfile->GetConsolidatedCVarValue(TEXT("fx.NiagaraPruneEmittersOnCookByDetailLevel"), CVarCullBasedOnDetailLevel) && CVarCullBasedOnDetailLevel == 1)
-		{
-			int32 CVarDetailLevelFoundValue;
-			if (InEmitter && DeviceProfile->GetConsolidatedCVarValue(TEXT("r.DetailLevel"), CVarDetailLevelFoundValue))
-			{
-				// Check emitter's detail level range contains the platform's level
-				// If e.g. the emitter's detail level range is between 0 and 2  and the platform detail is 3 only,
-				// then we should cull it.
-				if (InEmitter->IsAllowedByDetailLevel(CVarDetailLevelFoundValue))
-				{
-					return true;
-				}
-				return false;
-			}
-		}
-	}
-#endif
-	return true;
-}
+		EngineEffectsQuality = NewEffectsQuality;
 
-void INiagaraModule::OnChangeDetailLevel(class IConsoleVariable* CVar)
-{
-	int32 NewDetailLevel = CVar->GetInt();
-	if (EngineDetailLevel != NewDetailLevel)
-	{
-		EngineDetailLevel = NewDetailLevel;
+		FNiagaraPlatformSet::InvalidateCachedData();
 
 		for (TObjectIterator<UNiagaraSystem> It; It; ++It)
 		{
 			UNiagaraSystem* System = *It;
 			check(System);
-			System->OnDetailLevelChanges(NewDetailLevel);
-		}
-
-		// If the detail level has changed we have to reset all systems,
-		// and only activate ones that were previously active
-		for (TObjectIterator<UNiagaraComponent> It; It; ++It)
-		{
-			UNiagaraComponent* Comp = *It;
-			check(Comp);
-
-			const bool bWasActive = Comp->IsActive();
-			Comp->DestroyInstance();
-			if ( bWasActive )
-			{
-				Comp->Activate(true);
-			}
+			System->OnEffectsQualityChanged();
 		}
 	}
 }
@@ -1020,6 +947,17 @@ void INiagaraModule::ProcessShaderCompilationQueue()
 	return OnProcessQueue.Execute();
 }
 
+
+void CVarSinkFunc()
+{
+	//This Cvar sink can happen before the one which primes the cached scalability cvars so we must grab this ourselves.
+	IConsoleManager& ConsoleMan = IConsoleManager::Get();
+	static const auto CVarEQ = ConsoleMan.FindTConsoleVariableDataInt(TEXT("sg.EffectsQuality"));
+	INiagaraModule::OnEffectsQualityChanged(CVarEQ->GetValueOnGameThread());
+}
+
+static FAutoConsoleVariableSink CVarSink(FConsoleCommandDelegate::CreateStatic(&CVarSinkFunc));
+
 #if WITH_EDITOR
 const TArray<FNiagaraVariable>& FNiagaraGlobalParameters::GetVariables()
 {
diff --git a/Engine/Plugins/FX/Niagara/Source/Niagara/Private/NiagaraPlatformSet.cpp b/Engine/Plugins/FX/Niagara/Source/Niagara/Private/NiagaraPlatformSet.cpp
new file mode 100644
--- /dev/null
+++ b/Engine/Plugins/FX/Niagara/Source/Niagara/Private/NiagaraPlatformSet.cpp
@@ -0,0 +1,472 @@
+// Copyright Epic Games, Inc. All Rights Reserved.
+
+#include "NiagaraPlatformSet.h"
+#include "NiagaraModule.h"
+#include "DeviceProfiles/DeviceProfile.h"
+#include "DeviceProfiles/DeviceProfileManager.h"
+#include "Scalability.h"
+
+#define LOCTEXT_NAMESPACE "NiagaraPlatformSet"
+
+/**
+Whether a platform can change it's scalability settings at runtime.
+Defaults to false for all platforms and is explicitly enabled for desktop platforms.
+*/
+const TCHAR* CanChangeEQCVarName = TEXT("fx.NiagaraAllowRuntimeScalabilityChanges");
+int32 GbNiagaraAllowRuntimeScalabilityChanges = 0;
+static FAutoConsoleVariableRef CVarNiagaraAllowRuntimeScalabilityChanges(
+	CanChangeEQCVarName,
+	GbNiagaraAllowRuntimeScalabilityChanges,
+	TEXT("If > 0 this platform allows scalability level changes at runtime. \n"),
+	ECVF_ReadOnly
+);
+
+uint32 FNiagaraPlatformSet::LastDirtiedFrame = 0;
+#if WITH_EDITOR
+TMap<const UDeviceProfile*, int32> FNiagaraPlatformSet::CachedEQMasksPerDeviceProfile;
+TMap<FName, FNiagaraPlatformSet::FPlatformIniSettings> FNiagaraPlatformSet::CachedPlatformIniSettings;
+#endif
+
+FText FNiagaraPlatformSet::GetEffectsQualityText(int32 EffectsQuality)
+{
+	Scalability::FQualityLevels Counts = Scalability::GetQualityLevelCounts();
+	return Scalability::GetQualityLevelText(EffectsQuality, Counts.EffectsQuality);
+}
+
+FText FNiagaraPlatformSet::GetEffectsQualityMaskText(int32 EffectsQualityMask)
+{
+	if (EffectsQualityMask == INDEX_NONE)
+	{
+		return LOCTEXT("EffectsQualityAll", "All");
+	}
+	else if (EffectsQualityMask == 0)
+	{
+		return LOCTEXT("EffectsQualityNone", "None");
+	}
+	else
+	{
+		Scalability::FQualityLevels Counts = Scalability::GetQualityLevelCounts();
+		return Scalability::GetQualityLevelText(EQFromMask(EffectsQualityMask), Counts.EffectsQuality);
+	}
+}
+
+FNiagaraPlatformSet::FNiagaraPlatformSet(int32 EQMask)
+	: EffectsQualityMask(EQMask)
+	, LastBuiltFrame(0)
+	, bEnabledForCurrentProfileAndEffectQuality(false)
+{
+	IsActive();
+}
+
+FNiagaraPlatformSet::FNiagaraPlatformSet()
+	: EffectsQualityMask(INDEX_NONE)
+	, LastBuiltFrame(0)
+	, bEnabledForCurrentProfileAndEffectQuality(false)
+{
+}
+
+bool FNiagaraPlatformSet::IsActive()const
+{
+	if (LastBuiltFrame <= LastDirtiedFrame || LastBuiltFrame == 0)
+	{
+		bEnabledForCurrentProfileAndEffectQuality = IsEnabled(UDeviceProfileManager::Get().GetActiveProfile(), INiagaraModule::GetEffectsQuality());
+		LastBuiltFrame = GFrameNumber;
+	}
+	return bEnabledForCurrentProfileAndEffectQuality;
+}
+
+int32 FNiagaraPlatformSet::IsEnabledForDeviceProfile(const UDeviceProfile* DeviceProfile)const
+{
+	Scalability::FQualityLevels NumLevels = Scalability::GetQualityLevelCounts();
+	int32 RetEQMask = 0;
+	for (int32 i = 0; i < NumLevels.EffectsQuality; ++i)
+	{
+		if (IsEnabled(DeviceProfile, i))
+		{
+			RetEQMask |= CreateEQMask(i);
+		}
+	}
+
+	return RetEQMask;
+}
+
+bool FNiagaraPlatformSet::IsEnabledForEffectsQuality(int32 EffectsQuality)const
+{
+	for (UObject* DeviceProfileObj : UDeviceProfileManager::Get().Profiles)
+	{
+		if (UDeviceProfile* Profile = Cast<UDeviceProfile>(DeviceProfileObj))
+		{
+			if (IsEnabled(Profile, EffectsQuality))
+			{
+				return true;
+			}
+		}
+	}
+
+	return false;
+}
+
+void FNiagaraPlatformSet::GetOverridenDeviceProfiles(int32 EffectsQuality, TArray<UDeviceProfile*>& OutEnabledProfiles, TArray<UDeviceProfile*>& OutDisabledProfiles)const
+{
+	int32 EQMask = CreateEQMask(EffectsQuality);
+	for (const FNiagaraDeviceProfileStateEntry& Entry : DeviceProfileStates)
+	{
+		if (UObject** DeviceProfile = UDeviceProfileManager::Get().Profiles.FindByPredicate([&](UObject* CheckProfile) {return CheckProfile->GetFName() == Entry.ProfileName;}))
+		{
+			UDeviceProfile* Profile = CastChecked<UDeviceProfile>(*DeviceProfile);
+			//If this platform cannot change at runtime then we store all EQs in the state so that the device is still overridden if someone changes it's EQ CVar.
+			//So here we must also check that this EffectsQuality is the right one for the platforms current setting.
+			int32 ProfileEQMask = GetEffectQualityMaskForDeviceProfile(Profile);
+			if (ProfileEQMask == INDEX_NONE || (EQMask & ProfileEQMask) != 0)
+			{
+				ENiagaraPlatformSelectionState State = Entry.GetState(EffectsQuality);
+				if (State == ENiagaraPlatformSelectionState::Enabled)
+				{
+					OutEnabledProfiles.Add(Profile);
+				}
+				else if (State == ENiagaraPlatformSelectionState::Disabled)
+				{
+					OutDisabledProfiles.Add(Profile);
+				}
+			}
+		}
+	}
+}
+
+bool FNiagaraPlatformSet::CanChangeScalabilityAtRuntime()
+{
+	//For the current platform we can just read direct as this CVar is readonly.
+	return GbNiagaraAllowRuntimeScalabilityChanges != 0;
+}
+
+int32 FNiagaraPlatformSet::GetEffectQualityMaskForDeviceProfile(const UDeviceProfile* Profile)
+{
+#if WITH_EDITOR
+
+	//When in the editor we may be asking for the EQ of a platform other than the current one.
+	//So we have to look in ini files and device profiles to find it.
+
+	if (int32* CachedEQMask = CachedEQMasksPerDeviceProfile.Find(Profile))
+	{
+		//We've seen this Profile before so return the cached value;
+		return *CachedEQMask;
+	}
+
+	int32 EQMask = INDEX_NONE;
+	FPlatformIniSettings PlatformSettings = GetPlatformIniSettings(Profile->DeviceType);
+	
+	if (PlatformSettings.bCanChangeEffectsQualityAtRuntime)
+	{
+		EQMask = INDEX_NONE;
+	}
+	else
+	{
+		//Default to the platform EQ but allow device profiles to override.
+		EQMask = PlatformSettings.EffectsQualityMask;
+
+		//Check if the DPs set EQ directly.
+		int32 EffectsQuality = INDEX_NONE;
+		if (Profile->GetConsolidatedCVarValue(TEXT("sg.EffectsQuality"), EffectsQuality))
+		{
+			EQMask = CreateEQMask(EffectsQuality);
+		}
+	}
+
+	CachedEQMasksPerDeviceProfile.Add(Profile) = EQMask;
+	return EQMask;
+
+#else
+
+	//When not in editor we can assume we're asking about the current platform.
+	check(Profile == UDeviceProfileManager::Get().GetActiveProfile());
+	bool bCanChangeEQAtRuntime = CanChangeScalabilityAtRuntime();
+
+	return CreateEQMask(bCanChangeEQAtRuntime ? INDEX_NONE : INiagaraModule::GetEffectsQuality());
+
+#endif
+}
+
+bool FNiagaraPlatformSet::IsEnabled(const UDeviceProfile* Profile, int32 EffectsQuality)const
+{
+	checkSlow(Profile);
+
+	int32 TestEQMask = CreateEQMask(EffectsQuality);
+	int32 ProfileEQMask = GetEffectQualityMaskForDeviceProfile(Profile);
+	bool bIsEnabledByEQ = (EffectsQualityMask & ProfileEQMask & TestEQMask) != 0;
+
+	if (DeviceProfileStates.Num() > 0)
+	{
+		//Walk up the parent hierarchy to see if we have an explicit state for this profile set.
+		const UDeviceProfile* CurrProfile = Profile;
+		while (CurrProfile)
+		{
+			if (const FNiagaraDeviceProfileStateEntry* StateEntry = DeviceProfileStates.FindByPredicate([&](const FNiagaraDeviceProfileStateEntry& ProfileState) {return ProfileState.ProfileName == CurrProfile->GetFName(); }))
+			{
+				ENiagaraPlatformSelectionState SelectionState = StateEntry->GetState(EffectsQuality);
+				if (SelectionState != ENiagaraPlatformSelectionState::Default)
+				{
+					return SelectionState == ENiagaraPlatformSelectionState::Enabled;
+				}
+			}
+			CurrProfile = Cast<UDeviceProfile>(CurrProfile->Parent);
+		}
+	}
+
+	return bIsEnabledByEQ;
+}
+
+
+void FNiagaraPlatformSet::InvalidateCachedData()
+{
+#if WITH_EDITOR
+	CachedEQMasksPerDeviceProfile.Empty();
+	CachedPlatformIniSettings.Empty();
+#endif
+
+	LastDirtiedFrame = GFrameNumber;
+}
+
+#if WITH_EDITOR
+
+bool FNiagaraPlatformSet::IsEnabledForPlatform(const FString& PlatformName)
+{
+	for (const UObject* ProfileObj : UDeviceProfileManager::Get().Profiles)
+	{
+		if (const UDeviceProfile* Profile = Cast<const UDeviceProfile>(ProfileObj))
+		{
+			if (Profile->DeviceType == PlatformName)
+			{
+				if (IsEnabledForDeviceProfile(Profile) != 0)
+				{
+					return true;//At least one profile for this platform is enabled.
+				}
+			}
+		}
+	}
+
+	//No enabled profiles for this platform.
+	return false;
+}
+
+// bool FNiagaraPlatformSet::Conflicts(const FNiagaraPlatformSet& Other, TArray<const UDeviceProfile*>& OutConflictingProfiles, bool bIncludeProfilesWithVariableEffectsQuality)const
+// {
+// 	for (const UObject* ProfileObj : UDeviceProfileManager::Get().Profiles)
+// 	{
+// 		if (const UDeviceProfile* Profile = Cast<const UDeviceProfile>(ProfileObj))
+// 		{
+// 			if (bIncludeProfilesWithVariableEffectsQuality || CanChangeScalabilityAtRuntime(Profile) == false)
+// 			{
+// 				if (IsEnabledForDeviceProfile(Profile) && Other.IsEnabledForDeviceProfile(Profile))
+// 				{
+// 					OutConflictingProfiles.AddUnique(Profile);
+// 				}
+// 			}
+// 		}
+// 	}
+// 
+// 	return OutConflictingProfiles.Num() > 0;
+// }
+
+bool FNiagaraPlatformSet::IsEffectQualityEnabled(int32 EffectQuality)const
+{
+	return ((1 << EffectQuality) & EffectsQualityMask) != 0;
+}
+
+void FNiagaraPlatformSet::SetEnabledForEffectQuality(int32 EffectQuality, bool bEnabled)
+{
+	int32 EQBit = (1 << EffectQuality);
+	if (bEnabled)
+	{
+		EffectsQualityMask |= EQBit;
+	}
+	else
+	{
+		EffectsQualityMask &= ~EQBit;
+	}
+	OnChanged();
+}
+
+void FNiagaraPlatformSet::SetDeviceProfileState(UDeviceProfile* Profile, int32 EffectsQuality, ENiagaraPlatformSelectionState NewState)
+{
+	int32 Index = INDEX_NONE;
+	Index = DeviceProfileStates.IndexOfByPredicate([&](const FNiagaraDeviceProfileStateEntry& Entry) { return Entry.ProfileName == Profile->GetFName(); });
+
+	int32 ProfileEQMask = GetEffectQualityMaskForDeviceProfile(Profile);
+	if (ProfileEQMask != INDEX_NONE)
+	{
+		//For platforms that cannot change EQ at runtime we mark all state bits when setting state here so that if someone changes their EQ setting in the future, the state will be preserved.
+		EffectsQuality = INDEX_NONE;
+	}
+
+	if(Index == INDEX_NONE)
+	{
+		if (NewState != ENiagaraPlatformSelectionState::Default)
+		{
+			FNiagaraDeviceProfileStateEntry& NewEntry = DeviceProfileStates.AddDefaulted_GetRef();
+			NewEntry.ProfileName = Profile->GetFName();
+			NewEntry.SetState(EffectsQuality, NewState);
+		}
+	}
+	else
+	{
+		DeviceProfileStates[Index].SetState(EffectsQuality, NewState);
+
+		if (DeviceProfileStates[Index].AllDefaulted())
+		{
+			DeviceProfileStates.RemoveAtSwap(Index);//We don't need to store the default state. It's implied by no entry.
+		}
+	}
+	OnChanged();
+}
+
+ENiagaraPlatformSelectionState FNiagaraPlatformSet::GetDeviceProfileState(UDeviceProfile* Profile, int32 EffectsQuality)const
+{
+	if (const FNiagaraDeviceProfileStateEntry* ExistingEntry = DeviceProfileStates.FindByPredicate([&](const FNiagaraDeviceProfileStateEntry& Entry) { return Entry.ProfileName == Profile->GetFName(); }))
+	{
+		int32 ProfileEQMask = GetEffectQualityMaskForDeviceProfile(Profile);
+		if (ProfileEQMask == INDEX_NONE || ProfileEQMask & CreateEQMask(EffectsQuality))
+		{
+			//For profiles that cannot change scalability at runtime we store all flags in their state so that if anyone ever changes their EQ Cvar, the state setting remains valid.
+			//This just means we also have to ensure this is the correct EQ here.
+
+			return ExistingEntry->GetState(EffectsQuality);
+		}
+	}
+	return ENiagaraPlatformSelectionState::Default;
+}
+
+void FNiagaraPlatformSet::OnChanged()
+{
+	LastBuiltFrame = 0;
+}
+
+bool FNiagaraPlatformSet::CanChangeScalabilityAtRuntime(const FString& PlatformName)
+{
+	int32 EQMask = GetEffectQualityMaskForPlatform(PlatformName);
+	//If we're on a platform that can change EQ then it will have the full mask.
+	return EQMask == INDEX_NONE;
+}
+
+bool FNiagaraPlatformSet::CanChangeScalabilityAtRuntime(const UDeviceProfile* DeviceProfile)
+{
+	if (ensure(DeviceProfile))
+	{
+		int32 EQMask = GetEffectQualityMaskForDeviceProfile(DeviceProfile);
+		//If we're on a platform that can change EQ then it will have the full mask.
+		return EQMask == INDEX_NONE;
+	}
+	return true;//Assuming true if we fail to find the platform seems safest.
+}
+
+bool FNiagaraPlatformSet::ShouldPruneEmittersOnCook(const FString& PlatformName)
+{
+	FPlatformIniSettings Settings = GetPlatformIniSettings(PlatformName);
+	return Settings.bPruneEmittersOnCook != 0;
+}
+
+bool FNiagaraPlatformSet::GatherConflicts(const TArray<const FNiagaraPlatformSet*>& PlatformSets, TArray<FNiagaraPlatformSetConflictInfo>& OutConflicts)
+{
+	Scalability::FQualityLevels NumLevels = Scalability::GetQualityLevelCounts();
+
+	FNiagaraPlatformSetConflictInfo* CurrentConflict = nullptr;
+	for (int32 A = 0; A < PlatformSets.Num(); ++A)
+	{
+		for (int32 B = A + 1; B < PlatformSets.Num(); ++B)
+		{
+			check(A != B);
+
+			if (PlatformSets[A] == nullptr || PlatformSets[B] == nullptr) continue;
+
+			const FNiagaraPlatformSet& SetA = *PlatformSets[A];
+			const FNiagaraPlatformSet& SetB = *PlatformSets[B];
+
+			CurrentConflict = nullptr;
+			for (UObject* DPObj : UDeviceProfileManager::Get().Profiles)
+			{
+				UDeviceProfile* Profile = CastChecked<UDeviceProfile>(DPObj);
+				int32 AEnabledMask = SetA.IsEnabledForDeviceProfile(Profile);
+				int32 BEnabledMask = SetB.IsEnabledForDeviceProfile(Profile);
+				int32 ConflictMask = AEnabledMask & BEnabledMask;
+
+				if (ConflictMask != 0)
+				{
+					//We have a conflict so add it to the output.
+					if (CurrentConflict == nullptr)
+					{
+						CurrentConflict = &OutConflicts.AddDefaulted_GetRef();
+						CurrentConflict->SetAIndex = A;
+						CurrentConflict->SetBIndex = B;
+					}
+					FNiagaraPlatformSetConflictEntry& ConflictEntry = CurrentConflict->Conflicts.AddDefaulted_GetRef();
+					ConflictEntry.ProfileName = Profile->GetFName();
+					ConflictEntry.EffectsQualityMask = ConflictMask;
+				}
+			}
+		}
+	}
+
+	return OutConflicts.Num() > 0;
+}
+
+FNiagaraPlatformSet::FPlatformIniSettings FNiagaraPlatformSet::GetPlatformIniSettings(const FString& PlatformName)
+{
+	if (FPlatformIniSettings* CachedSettings = CachedPlatformIniSettings.Find(*PlatformName))
+	{
+		return *CachedSettings;
+	}
+
+	//Load config files in which we can reasonable expect to find sg.EffectQuality and may be set.
+	FConfigFile EngineSettings;
+	FConfigCacheIni::LoadLocalIniFile(EngineSettings, TEXT("Engine"), true, *PlatformName);//Should use BaseProfileName? Are either of these ensured to be correct? I worry this is brittle.
+
+	FConfigFile GameSettings;
+	FConfigCacheIni::LoadLocalIniFile(GameSettings, TEXT("Game"), true, *PlatformName);//Should use BaseProfileName? Are either of these ensured to be correct? I worry this is brittle.
+
+	FConfigFile ScalabilitySettings;
+	FConfigCacheIni::LoadLocalIniFile(ScalabilitySettings, TEXT("Scalability"), true, *PlatformName);//Should use BaseProfileName? Are either of these ensured to be correct? I worry this is brittle.
+
+	auto FindCVarValue = [&](const TCHAR* Section, const TCHAR* CVarName, int32& OutVal)
+	{
+		if (!ScalabilitySettings.GetInt(Section, CVarName, OutVal))
+		{
+			if (!GameSettings.GetInt(Section, CVarName, OutVal))
+			{
+				EngineSettings.GetInt(Section, CVarName, OutVal);
+			}
+		}
+	};
+
+	int32 EffectsQuality = INDEX_NONE;
+	int32 EQMask = INDEX_NONE;
+
+	//If this platform can change quality settings at runtime then we return the full mask.
+	int32 CanChangeEffectQuality = 0;
+
+	FindCVarValue(TEXT("SystemSettings"), CanChangeEQCVarName, CanChangeEffectQuality);
+	if (CanChangeEffectQuality != 0)
+	{
+		//If this platform can change it's EffectsQuality at runtime we return the ALL mask.
+		EQMask = INDEX_NONE;
+	}
+	else
+	{
+		FindCVarValue(TEXT("ScalabilityGroups"), TEXT("sg.EffectsQuality"), EffectsQuality);
+		EQMask = CreateEQMask(EffectsQuality);
+	}
+
+	int32 PruneEmittersOnCook = 0;
+	FindCVarValue(TEXT("SystemSettings"), TEXT("fx.PruneEmittersOnCook"), PruneEmittersOnCook);
+
+	FPlatformIniSettings& NewSetting = CachedPlatformIniSettings.Add(*PlatformName);
+	NewSetting = FPlatformIniSettings(CanChangeEffectQuality, PruneEmittersOnCook, EQMask);
+	return NewSetting;
+}
+
+int32 FNiagaraPlatformSet::GetEffectQualityMaskForPlatform(const FString& PlatformName)
+{
+	FPlatformIniSettings PlatformSettings = GetPlatformIniSettings(PlatformName);
+	return PlatformSettings.EffectsQualityMask;
+}
+#endif
+
+#undef LOCTEXT_NAMESPACE
\ No newline at end of file
diff --git a/Engine/Plugins/FX/Niagara/Source/Niagara/Private/NiagaraScalabilityManager.cpp b/Engine/Plugins/FX/Niagara/Source/Niagara/Private/NiagaraScalabilityManager.cpp
--- a/Engine/Plugins/FX/Niagara/Source/Niagara/Private/NiagaraScalabilityManager.cpp
+++ b/Engine/Plugins/FX/Niagara/Source/Niagara/Private/NiagaraScalabilityManager.cpp
@@ -181,7 +181,7 @@ void FNiagaraScalabilityManager::Update(FNiagaraWorldManager* WorldMan)
 #endif
 
 		UNiagaraSystem* System = Component->GetAsset();
-		const FNiagaraScalabilitySettings& ScalabilitySettings = System->GetScalabilitySettings(Component->GetPreviewDetailLevel());
+		const FNiagaraSystemScalabilitySettings& ScalabilitySettings = System->GetScalabilitySettings();
 
 		SignificanceSortedIndices.Add(i);
 		bNeedSortedSignificanceCull = ScalabilitySettings.bCullMaxInstanceCount && ScalabilitySettings.MaxInstances > 0;
@@ -204,7 +204,7 @@ void FNiagaraScalabilityManager::Update(FNiagaraWorldManager* WorldMan)
 
 			bool bOldCulled = CompState.bCulled;
 
-			const FNiagaraScalabilitySettings& ScalabilitySettings = System->GetScalabilitySettings(Component->GetPreviewDetailLevel());
+			const FNiagaraSystemScalabilitySettings& ScalabilitySettings = System->GetScalabilitySettings();
 			WorldMan->SortedSignificanceCull(EffectType, ScalabilitySettings, CompState.Significance, i, CompState);
 
 			CompState.bDirty |= CompState.bCulled != bOldCulled;
@@ -333,4 +333,4 @@ void FNiagaraScalabilityManager::Dump()
 	UE_LOG(LogNiagara, Display, TEXT("-------------------------------------------------------------------------------\n%s"), *DetailString);
 }
 
-#endif
\ No newline at end of file
+#endif
diff --git a/Engine/Plugins/FX/Niagara/Source/Niagara/Private/NiagaraScript.cpp b/Engine/Plugins/FX/Niagara/Source/Niagara/Private/NiagaraScript.cpp
--- a/Engine/Plugins/FX/Niagara/Source/Niagara/Private/NiagaraScript.cpp
+++ b/Engine/Plugins/FX/Niagara/Source/Niagara/Private/NiagaraScript.cpp
@@ -140,7 +140,6 @@ bool FNiagaraVMExecutableDataId::operator==(const FNiagaraVMExecutableDataId& Re
 #if WITH_EDITORONLY_DATA
 		BaseScriptCompileHash != ReferenceSet.BaseScriptCompileHash ||
 #endif
-		DetailLevelMask != ReferenceSet.DetailLevelMask ||
 		bUsesRapidIterationParams != ReferenceSet.bUsesRapidIterationParams ||
 		bInterpolatedSpawn != ReferenceSet.bInterpolatedSpawn ||
 		bRequiresPersistentIDs != ReferenceSet.bRequiresPersistentIDs )
@@ -212,14 +211,6 @@ void FNiagaraVMExecutableDataId::AppendKeyString(FString& KeyString, const FStri
 		KeyString += TEXT("[AdditionalDefines]") + Delimiter;
 	}
 
-	if (DetailLevelMask != 0xFFFFFFFF)
-	{
-		KeyString += FString::Printf(TEXT("DL_%d"), DetailLevelMask);
-	}
-	else
-	{
-		KeyString += TEXT("ALLDL") + Delimiter;
-	}
 
 	if (bUsesRapidIterationParams)
 	{
@@ -315,7 +306,6 @@ void UNiagaraScript::ComputeVMCompilationId(FNiagaraVMExecutableDataId& Id) cons
 	Id.bUsesRapidIterationParams = true;
 	Id.bInterpolatedSpawn = false;
 	Id.bRequiresPersistentIDs = false;
-	Id.DetailLevelMask = 0xFFFFFFFF; // Unused for now.
 	
 	// Ideally we wouldn't want to do this but rather than push the data down
 	// from the emitter.
@@ -348,10 +338,6 @@ void UNiagaraScript::ComputeVMCompilationId(FNiagaraVMExecutableDataId& Id) cons
 		{
 			Id.AdditionalDefines.Add(TEXT("Emitter.Determinism"));
 		}
-		if (Emitter->bOverrideGlobalSpawnCountScale)
-		{
-			Id.AdditionalDefines.Add(TEXT("Emitter.OverrideGlobalSpawnCountScale"));
-		}
 
 		if (!Emitter->bBakeOutRapidIteration)
 		{
@@ -1643,7 +1629,7 @@ void UNiagaraScript::CacheResourceShadersForCooking(EShaderPlatform ShaderPlatfo
 
 			NewResource->SetScript(this, (ERHIFeatureLevel::Type)TargetFeatureLevel, CachedScriptVMId.CompilerVersionID, CachedScriptVMId.AdditionalDefines,
 				CachedScriptVMId.BaseScriptCompileHash,	CachedScriptVMId.ReferencedCompileHashes, 
-				CachedScriptVMId.bUsesRapidIterationParams, CachedScriptVMId.DetailLevelMask, GetFullName());
+				CachedScriptVMId.bUsesRapidIterationParams, GetFullName());
 			ResourceToCache = NewResource;
 
 			check(ResourceToCache);
@@ -1709,7 +1695,7 @@ void UNiagaraScript::CacheResourceShadersForRendering(bool bRegenerateId, bool b
 			ERHIFeatureLevel::Type CacheFeatureLevel = GMaxRHIFeatureLevel;
 			ScriptResource.SetScript(this, FeatureLevel, CachedScriptVMId.CompilerVersionID, CachedScriptVMId.AdditionalDefines,
 				CachedScriptVMId.BaseScriptCompileHash, CachedScriptVMId.ReferencedCompileHashes, 
-				CachedScriptVMId.bUsesRapidIterationParams, CachedScriptVMId.DetailLevelMask, GetFullName());
+				CachedScriptVMId.bUsesRapidIterationParams, GetFullName());
 
 			//if (ScriptResourcesByFeatureLevel[FeatureLevel])
 			{
diff --git a/Engine/Plugins/FX/Niagara/Source/Niagara/Private/NiagaraSystem.cpp b/Engine/Plugins/FX/Niagara/Source/Niagara/Private/NiagaraSystem.cpp
--- a/Engine/Plugins/FX/Niagara/Source/Niagara/Private/NiagaraSystem.cpp
+++ b/Engine/Plugins/FX/Niagara/Source/Niagara/Private/NiagaraSystem.cpp
@@ -20,6 +20,7 @@
 #include "NiagaraWorldManager.h"
 #include "NiagaraEmitterInstanceBatcher.h"
 #include "NiagaraSettings.h"
+#include "Interfaces/ITargetPlatform.h"
 
 
 #if WITH_EDITOR
@@ -42,6 +43,14 @@ static FAutoConsoleVariableRef CVarNiagaraForceSystemsToCookOutRapidIterationOnL
 	ECVF_Default
 );
 
+static int GNiagaraVerboseCookOutput = 0;
+static FAutoConsoleVariableRef CVarNiagaraVerboseCookOutput(
+	TEXT("fx.NiagaraVerboseCookOutput"),
+	GNiagaraVerboseCookOutput,
+	TEXT("When enabled UNiagaraSystems generate verbose output on cook."),
+	ECVF_Default
+);
+
 //////////////////////////////////////////////////////////////////////////
 
 UNiagaraSystem::UNiagaraSystem(const FObjectInitializer& ObjectInitializer)
@@ -123,7 +132,7 @@ void UNiagaraSystem::PostInitProperties()
 #endif
 	}
 
-	ResolveScalabilityOverrides();
+	ResolveScalabilitySettings();
 }
 
 bool UNiagaraSystem::IsLooping() const
@@ -201,8 +210,67 @@ void UNiagaraSystem::RequestCompileForEmitter(UNiagaraEmitter* InEmitter)
 
 void UNiagaraSystem::Serialize(FArchive& Ar)
 {
+	//TODO: Cook out emitters. This requires some wrangling on the compilation side for system scripts.
+// 	//When cooking we clear out any emitters that are disabled entirely or just for the target platform.
+// 	TArray<FNiagaraEmitterHandle> OldEmitters;
+// 	if (Ar.IsCooking())
+// 	{
+// 		OldEmitters = EmitterHandles;
+// 		const ITargetPlatform* Target = Ar.CookingTarget();
+// 		for (FNiagaraEmitterHandle& Handle : EmitterHandles)
+// 		{
+// 			if (Handle.GetIsEnabled() && Handle.GetInstance()->IsEnabledOnPlatform(Target->PlatformName()))
+// 			{
+// 				Handle.ClearEmitter();
+// 			}
+// 		}
+// 	}
+
+	//Some log spam to display emitters we should be cooking out.
+#if WITH_EDITOR && 0
+	if (Ar.IsCooking() && GNiagaraVerboseCookOutput)
+	{
+		TBitArray<TInlineAllocator<8>> EnabledEmitters;
+		EnabledEmitters.Reserve(EmitterHandles.Num());
+		int32 NumEnabled = 0;
+		const ITargetPlatform* Target = Ar.CookingTarget();
+		for (FNiagaraEmitterHandle& Handle : EmitterHandles)
+		{
+			bool bEnabled = Handle.GetInstance()->IsEnabledOnPlatform(Target->PlatformName());
+			EnabledEmitters.Add(bEnabled);
+			if (bEnabled)
+			{
+				NumEnabled++;
+			}
+		}
+
+		if (NumEnabled != EmitterHandles.Num())
+		{
+			UE_LOG(LogNiagara, Log, TEXT("==============================================="));
+			UE_LOG(LogNiagara, Log, TEXT("System has emitters that could be cooked out!"), *GetName());
+
+			for (int32 EmitterIdx = 0; EmitterIdx < EmitterHandles.Num(); ++EmitterIdx)
+			{
+				if (EnabledEmitters[EmitterIdx] == false)
+				{
+					UE_LOG(LogNiagara, Log, TEXT("%s"), *EmitterHandles[EmitterIdx].GetName().ToString());
+				}
+			}
+			UE_LOG(LogNiagara, Log, TEXT("==============================================="));
+		}
+	}
+
+#endif
+
 	Super::Serialize(Ar);
 
+	//TODO: Cook out emitters. This requires some wrangling on the compilation side for system scripts.
+// 	//Restore old emitters post cook pruning.
+// 	if (Ar.IsCooking())
+// 	{
+// 		EmitterHandles = OldEmitters;
+// 	}
+
 	Ar.UsingCustomVersion(FNiagaraCustomVersion::GUID);
 
 	if (Ar.CustomVer(FNiagaraCustomVersion::GUID) >= FNiagaraCustomVersion::ChangeEmitterCompiledDataToSharedRefs)
@@ -274,7 +342,7 @@ void UNiagaraSystem::PostEditChangeProperty(struct FPropertyChangedEvent& Proper
 		}
 	}
 
-	ResolveScalabilityOverrides();
+	ResolveScalabilitySettings();
 
 	UpdateContext.CommitUpdate();
 }
@@ -313,6 +381,18 @@ void UNiagaraSystem::PostLoad()
 		}
 	}
 
+	const int32 NiagaraVer = GetLinkerCustomVersion(FNiagaraCustomVersion::GUID);
+	if (NiagaraVer < FNiagaraCustomVersion::PlatformScalingRefactor)
+	{
+		for (int32 DL = 0; DL < ScalabilityOverrides_DEPRECATED.Num(); ++DL)
+		{
+			FNiagaraSystemScalabilityOverride& LegacyOverride = ScalabilityOverrides_DEPRECATED[DL];
+			FNiagaraSystemScalabilityOverride& NewOverride = SystemScalabilityOverrides.Overrides.AddDefaulted_GetRef();
+			NewOverride = LegacyOverride;
+			NewOverride.Platforms = FNiagaraPlatformSet(FNiagaraPlatformSet::CreateEQMask(DL));
+		}
+	}
+
 #if UE_EDITOR
 	ExposedParameters.RecreateRedirections();
 #endif
@@ -349,8 +429,6 @@ void UNiagaraSystem::PostLoad()
 		}
 		AllSystemScripts.Add(SystemUpdateScript);
 
-		const int32 NiagaraVer = GetLinkerCustomVersion(FNiagaraCustomVersion::GUID);
-
 		//TODO: This causes a crash becuase the script source ptr is null? Fix
 		//For existing emitters before the lifecylce rework, ensure they have the system lifecycle module.
 		if (NiagaraVer < FNiagaraCustomVersion::LifeCycleRework)
@@ -492,7 +570,7 @@ void UNiagaraSystem::PostLoad()
 		bIsReadyToRunCached = IsReadyToRunInternal();
 	}
 
-	ResolveScalabilityOverrides();
+	ResolveScalabilitySettings();
 }
 
 #if WITH_EDITORONLY_DATA
@@ -1383,101 +1461,58 @@ UNiagaraEffectType* UNiagaraSystem::GetEffectType()const
 	return Settings->GetDefaultEffectType();
 }
 
-const FNiagaraScalabilitySettings& UNiagaraSystem::GetScalabilitySettings(int32 DetailLevel)
+void UNiagaraSystem::ResolveScalabilitySettings()
 {
-#if WITH_NIAGARA_COMPONENT_PREVIEW_DATA
-	int32 UseDetailLevel = DetailLevel == INDEX_NONE ? INiagaraModule::GetDetailLevel() : DetailLevel;
-
-	int32 SettingsIndex = ResolvedScalabilitySettings.IsValidIndex(UseDetailLevel) ? UseDetailLevel : ResolvedScalabilitySettings.Num() - 1;
-	check(SettingsIndex != INDEX_NONE);//We should always have at least one entry in our settings array.
-	return ResolvedScalabilitySettings[SettingsIndex];
-#else
-	return CurrentScalabilitySettings;
-#endif
-}
-
-void UNiagaraSystem::ResolveScalabilityOverrides()
-{
-	ResolvedScalabilitySettings.Empty();
-
-	UNiagaraEffectType* ActualEffectType = GetEffectType();
-	int32 NumEffectTypeSettings = 0;
-	if (ActualEffectType)
+	CurrentScalabilitySettings.Clear();
+	if (UNiagaraEffectType* ActualEffectType = GetEffectType())
 	{
-		NumEffectTypeSettings = ActualEffectType->GetScalabilitySettings().Num();
+		CurrentScalabilitySettings = ActualEffectType->GetActiveSystemScalabilitySettings();
 	}
-	int32 NumOverrides = ScalabilityOverrides.Num();
 
-	int32 ResolvedCount = FMath::Max(NumEffectTypeSettings, NumOverrides);
-	if (ResolvedCount == 0)
-	{
-		//Just add a single defaulted entry to make other calling code simpler.
-		ResolvedScalabilitySettings.AddDefaulted();
-	}
-	else
+	for (FNiagaraSystemScalabilityOverride& Override : SystemScalabilityOverrides.Overrides)
 	{
-
-		ResolvedScalabilitySettings.Reserve(ResolvedCount);
-		for (int32 i = 0; i < ResolvedCount; ++i)
+		if (Override.Platforms.IsActive())
 		{
-			int32 NewIdx = ResolvedScalabilitySettings.AddDefaulted();
-			FNiagaraScalabilitySettings& NewSettings = ResolvedScalabilitySettings[NewIdx];
-			if (ActualEffectType)
+			if (Override.bOverrideDistanceSettings)
 			{
-				const TArray<FNiagaraScalabilitySettings>& EffectTypeSettings = ActualEffectType->GetScalabilitySettings();
-				int32 LastIndex = EffectTypeSettings.Num() - 1;
-				if (LastIndex >= 0)
-				{
-					NewSettings = EffectTypeSettings[FMath::Min(i, LastIndex)];
-				}
+				CurrentScalabilitySettings.bCullByDistance = Override.bCullByDistance;
+				CurrentScalabilitySettings.MaxDistance = Override.MaxDistance;
 			}
 
-			if (ScalabilityOverrides.IsValidIndex(i) && bOverrideScalabilitySettings)
+			if (Override.bOverrideInstanceCountSettings)
 			{
-				const FNiagaraScalabilityOverrides& Overrides = ScalabilityOverrides[i];
-
-				if (Overrides.bOverrideDistanceSettings)
-				{
-					NewSettings.bCullByDistance = Overrides.bCullByDistance;
-					NewSettings.MaxDistance = Overrides.MaxDistance;
-				}
-
-				if (Overrides.bOverrideOwnerLODSettings)
-				{
-					NewSettings.bCullByMaxOwnerLOD = Overrides.bCullByMaxOwnerLOD;
-					NewSettings.MaxOwnerLOD = Overrides.MaxOwnerLOD;
-				}
-
-				if (Overrides.bOverrideInstanceCountSettings)
-				{
-					NewSettings.bCullMaxInstanceCount = Overrides.bCullMaxInstanceCount;
-					NewSettings.MaxInstances = Overrides.MaxInstances;
-				}
+				CurrentScalabilitySettings.bCullMaxInstanceCount = Override.bCullMaxInstanceCount;
+				CurrentScalabilitySettings.MaxInstances = Override.MaxInstances;
+			}
 
-				if (Overrides.bOverrideTimeSinceRendererSettings)
-				{
-					NewSettings.bCullByMaxTimeWithoutRender = Overrides.bCullByMaxTimeWithoutRender;
-					NewSettings.MaxTimeWithoutRender = Overrides.MaxTimeWithoutRender;
-				}
+			if (Override.bOverrideOwnerLODSettings)
+			{
+				CurrentScalabilitySettings.bCullByMaxOwnerLOD = Override.bCullByMaxOwnerLOD;
+				CurrentScalabilitySettings.MaxOwnerLOD = Override.MaxOwnerLOD;
+			}
 
-// 				if (Overrides.bOverrideSpawnCountScale)
-// 				{
-// 					NewSettings.SpawnCountScale = Overrides.SpawnCountScale;
-// 				}
+			if (Override.bOverrideTimeSinceRendererSettings)
+			{
+				CurrentScalabilitySettings.bCullByMaxTimeWithoutRender = Override.bCullByMaxTimeWithoutRender;
+				CurrentScalabilitySettings.MaxTimeWithoutRender = Override.MaxTimeWithoutRender;
 			}
+			break;//These overrides *should* be for orthogonal platform sets so we can exit after we've found a match.
 		}
 	}
-
-	OnDetailLevelChanges(INiagaraModule::GetDetailLevel());
 }
 
-void UNiagaraSystem::OnDetailLevelChanges(int32 DetailLevel)
+void UNiagaraSystem::OnEffectsQualityChanged()
 {
-	int32 SettingsIndex = ResolvedScalabilitySettings.IsValidIndex(DetailLevel) ? DetailLevel : ResolvedScalabilitySettings.Num() - 1;
-	check(SettingsIndex != INDEX_NONE);//We should always have at least one entry in our settings array.
-	CurrentScalabilitySettings = ResolvedScalabilitySettings[SettingsIndex];
-}
+	ResolveScalabilitySettings();
 
+	for (FNiagaraEmitterHandle& Handle : EmitterHandles)
+	{
+		Handle.GetInstance()->OnEffectsQualityChanged();
+	}
+
+	FNiagaraSystemUpdateContext UpdateCtx;
+	UpdateCtx.Add(this, true);
+}
 
 FNiagaraEmitterCompiledData::FNiagaraEmitterCompiledData()
 {
diff --git a/Engine/Plugins/FX/Niagara/Source/Niagara/Private/NiagaraSystemInstance.cpp b/Engine/Plugins/FX/Niagara/Source/Niagara/Private/NiagaraSystemInstance.cpp
--- a/Engine/Plugins/FX/Niagara/Source/Niagara/Private/NiagaraSystemInstance.cpp
+++ b/Engine/Plugins/FX/Niagara/Source/Niagara/Private/NiagaraSystemInstance.cpp
@@ -15,6 +15,9 @@
 #include "GameFramework/PlayerController.h"
 
 
+PRAGMA_DISABLE_INLINING
+PRAGMA_DISABLE_OPTIMIZATION
+
 DECLARE_CYCLE_STAT(TEXT("System Activate [GT]"), STAT_NiagaraSystemActivate, STATGROUP_Niagara);
 DECLARE_CYCLE_STAT(TEXT("System Deactivate [GT]"), STAT_NiagaraSystemDeactivate, STATGROUP_Niagara);
 DECLARE_CYCLE_STAT(TEXT("System Complete [GT]"), STAT_NiagaraSystemComplete, STATGROUP_Niagara);
@@ -1386,18 +1389,6 @@ bool FNiagaraSystemInstance::GetPerInstanceDataAndOffsets(void*& OutData, uint32
 	return DataInterfaceInstanceDataOffsets.Num() != 0;
 }
 
-int32 FNiagaraSystemInstance::GetDetailLevel()const
-{
-	int32 DetailLevel = INiagaraModule::GetDetailLevel();
-#if WITH_EDITOR
-	if (Component && Component->bEnablePreviewDetailLevel)
-	{
-		DetailLevel = Component->PreviewDetailLevel;
-	}
-#endif
-	return DetailLevel;
-}
-
 void FNiagaraSystemInstance::TickDataInterfaces(float DeltaSeconds, bool bPostSimulate)
 {
 	if (!GetSystem() || !Component || IsDisabled())
@@ -1519,6 +1510,8 @@ float FNiagaraSystemInstance::GetLODDistance()
 			LODDistance = FMath::Sqrt(LODDistanceSqr);
 		}
 	}
+
+	bLODDistanceIsValid = true;
 	return LODDistance;
 }
 
@@ -1630,7 +1623,8 @@ void FNiagaraSystemInstance::TickInstanceParameters_GameThread(float DeltaSecond
 		{
 			CurrentEmitterParameters.EmitterNumParticles = Emitter->GetNumParticles();
 			CurrentEmitterParameters.EmitterTotalSpawnedParticles = Emitter->GetTotalSpawnedParticles();
-			CurrentEmitterParameters.EmitterSpawnCountScale = Emitter->GetSpawnCountScale(EffectsQualityLevel);
+			const FNiagaraEmitterScalabilitySettings& ScalabilitySettings = Emitter->GetScalabilitySettings();
+			CurrentEmitterParameters.EmitterSpawnCountScale = ScalabilitySettings.bScaleSpawnCount ? ScalabilitySettings.SpawnCountScale : 1.0f;
 			++GatheredInstanceParameters.NumAlive;
 		}
 		else
@@ -2242,3 +2236,6 @@ FNiagaraSystemInstance::FOnDestroyed& FNiagaraSystemInstance::OnDestroyed()
 	return OnDestroyedDelegate;
 }
 #endif
+
+PRAGMA_ENABLE_INLINING
+PRAGMA_ENABLE_OPTIMIZATION
diff --git a/Engine/Plugins/FX/Niagara/Source/Niagara/Private/NiagaraSystemSimulation.cpp b/Engine/Plugins/FX/Niagara/Source/Niagara/Private/NiagaraSystemSimulation.cpp
--- a/Engine/Plugins/FX/Niagara/Source/Niagara/Private/NiagaraSystemSimulation.cpp
+++ b/Engine/Plugins/FX/Niagara/Source/Niagara/Private/NiagaraSystemSimulation.cpp
@@ -1139,7 +1139,8 @@ void FNiagaraSystemSimulation::TickFastPath(FNiagaraSystemSimulationTickContext&
 
 			if (UNiagaraEmitter* CachedEmitter = EmitterInstance->GetCachedEmitter())
 			{
-				EmitterMap.Emitter.SpawnCountScale = CachedEmitter->GetSpawnCountScale(Context.EffectsQuality);
+				const FNiagaraEmitterScalabilitySettings& ScalabilitySettings = CachedEmitter->GetScalabilitySettings();
+				EmitterMap.Emitter.SpawnCountScale = ScalabilitySettings.SpawnCountScale;
 			}
 		}
 	}
diff --git a/Engine/Plugins/FX/Niagara/Source/Niagara/Private/NiagaraWorldManager.cpp b/Engine/Plugins/FX/Niagara/Source/Niagara/Private/NiagaraWorldManager.cpp
--- a/Engine/Plugins/FX/Niagara/Source/Niagara/Private/NiagaraWorldManager.cpp
+++ b/Engine/Plugins/FX/Niagara/Source/Niagara/Private/NiagaraWorldManager.cpp
@@ -30,6 +30,14 @@ static FAutoConsoleVariableRef CVarNiagaraAllowAsyncWorkToEndOfFrame(
 	GNiagaraAllowAsyncWorkToEndOfFrame,
 	TEXT("Allow async work to continue until the end of the frame, if false it will complete within the tick group it's started in."),
 	ECVF_Default
+); 
+
+static int GNiagaraWaitOnPreGC = 1;
+static FAutoConsoleVariableRef CVarNiagaraWaitOnPreGC(
+	TEXT("fx.Niagara.WaitOnPreGC"),
+	GNiagaraWaitOnPreGC,
+	TEXT("Toggles whether Niagara will wait for all async tasks to complete before any GC calls."),
+	ECVF_Default
 );
 
 FAutoConsoleCommandWithWorld DumpNiagaraWorldManagerCommand(
@@ -52,7 +60,8 @@ FDelegateHandle FNiagaraWorldManager::OnWorldCleanupHandle;
 FDelegateHandle FNiagaraWorldManager::OnPreWorldFinishDestroyHandle;
 FDelegateHandle FNiagaraWorldManager::OnWorldBeginTearDownHandle;
 FDelegateHandle FNiagaraWorldManager::TickWorldHandle;
-FDelegateHandle FNiagaraWorldManager::PostGCHandle; 
+FDelegateHandle FNiagaraWorldManager::PreGCHandle;
+FDelegateHandle FNiagaraWorldManager::PostGCHandle;
 FDelegateHandle FNiagaraWorldManager::PreGCBeginDestroyHandle; 
 TMap<class UWorld*, class FNiagaraWorldManager*> FNiagaraWorldManager::WorldManagers;
 
@@ -169,6 +178,7 @@ void FNiagaraWorldManager::OnStartup()
 	OnWorldBeginTearDownHandle = FWorldDelegates::OnWorldBeginTearDown.AddStatic(&FNiagaraWorldManager::OnWorldBeginTearDown);
 	TickWorldHandle = FWorldDelegates::OnWorldPostActorTick.AddStatic(&FNiagaraWorldManager::TickWorld);
 
+	PreGCHandle = FCoreUObjectDelegates::GetPreGarbageCollectDelegate().AddStatic(&FNiagaraWorldManager::OnPreGarbageCollect);
 	PostGCHandle = FCoreUObjectDelegates::GetPostGarbageCollect().AddStatic(&FNiagaraWorldManager::OnPostGarbageCollect);
 	PreGCBeginDestroyHandle = FCoreUObjectDelegates::PreGarbageCollectConditionalBeginDestroy.AddStatic(&FNiagaraWorldManager::OnPreGarbageCollectBeginDestroy);
 }
@@ -181,6 +191,7 @@ void FNiagaraWorldManager::OnShutdown()
 	FWorldDelegates::OnWorldBeginTearDown.Remove(OnWorldBeginTearDownHandle);
 	FWorldDelegates::OnWorldPostActorTick.Remove(TickWorldHandle);
 
+	FCoreUObjectDelegates::GetPreGarbageCollectDelegate().Remove(PreGCHandle);
 	FCoreUObjectDelegates::GetPostGarbageCollect().Remove(PostGCHandle);
 	FCoreUObjectDelegates::PreGarbageCollectConditionalBeginDestroy.Remove(PreGCBeginDestroyHandle);
 
@@ -353,6 +364,20 @@ void FNiagaraWorldManager::OnWorldCleanup(bool bSessionEnded, bool bCleanupResou
 	ScalabilityManagers.Empty();
 }
 
+void FNiagaraWorldManager::PreGarbageCollect()
+{
+	if (GNiagaraWaitOnPreGC)
+	{
+		for (int TG = 0; TG < NiagaraNumTickGroups; ++TG)
+		{
+			for (TPair<UNiagaraSystem*, TSharedRef<FNiagaraSystemSimulation, ESPMode::ThreadSafe>>& SimPair : SystemSimulations[TG])
+			{
+				SimPair.Value->WaitForInstancesTickComplete();
+			}
+		}
+	}
+}
+
 void FNiagaraWorldManager::PostGarbageCollect()
 {
 	//Clear out and scalability managers who's EffectTypes have been GCd.
@@ -435,6 +460,14 @@ void FNiagaraWorldManager::TickWorld(UWorld* World, ELevelTick TickType, float D
 	Get(World)->PostActorTick(DeltaSeconds);
 }
 
+void FNiagaraWorldManager::OnPreGarbageCollect()
+{
+	for (TPair<UWorld*, FNiagaraWorldManager*>& Pair : WorldManagers)
+	{
+		Pair.Value->PreGarbageCollect();
+	}
+}
+
 void FNiagaraWorldManager::OnPostGarbageCollect()
 {
 	for (TPair<UWorld*, FNiagaraWorldManager*>& Pair : WorldManagers)
@@ -672,7 +705,7 @@ bool FNiagaraWorldManager::ShouldPreCull(UNiagaraSystem* System, UNiagaraCompone
 			if (CanPreCull(EffectType))
 			{
 				FNiagaraScalabilityState State;
-				const FNiagaraScalabilitySettings& ScalabilitySettings = System->GetScalabilitySettings(Component->GetPreviewDetailLevel());
+				const FNiagaraSystemScalabilitySettings& ScalabilitySettings = System->GetScalabilitySettings();
 				CalculateScalabilityState(System, ScalabilitySettings, EffectType, Component, true, State);
 				return State.bCulled;
 			}
@@ -690,7 +723,7 @@ bool FNiagaraWorldManager::ShouldPreCull(UNiagaraSystem* System, FVector Locatio
 			if (CanPreCull(EffectType))
 			{
 				FNiagaraScalabilityState State;
-				const FNiagaraScalabilitySettings& ScalabilitySettings = System->GetScalabilitySettings();
+				const FNiagaraSystemScalabilitySettings& ScalabilitySettings = System->GetScalabilitySettings();
 				CalculateScalabilityState(System, ScalabilitySettings, EffectType, Location, true, State);
 				return State.bCulled;
 			}
@@ -699,7 +732,7 @@ bool FNiagaraWorldManager::ShouldPreCull(UNiagaraSystem* System, FVector Locatio
 	return false;
 }
 
-void FNiagaraWorldManager::CalculateScalabilityState(UNiagaraSystem* System, const FNiagaraScalabilitySettings& ScalabilitySettings, UNiagaraEffectType* EffectType, FVector Location, bool bIsPreCull, FNiagaraScalabilityState& OutState)
+void FNiagaraWorldManager::CalculateScalabilityState(UNiagaraSystem* System, const FNiagaraSystemScalabilitySettings& ScalabilitySettings, UNiagaraEffectType* EffectType, FVector Location, bool bIsPreCull, FNiagaraScalabilityState& OutState)
 {
 	float DistSignificance = DistanceSignificance(EffectType, ScalabilitySettings, Location);
 
@@ -721,12 +754,9 @@ void FNiagaraWorldManager::CalculateScalabilityState(UNiagaraSystem* System, con
 	OutState.bDirty = OutState.bCulled != bOldCulled;
 
 	//TODO: More progressive scalability options?
-	//We can possibly adjust the DetailLevel each effect can use so we can more gradually scale down effects at lower significance.
-	//Possibly do similar to the Significance manager and define a set budget of instances that are allowed at each level.
-	//E.g. 5 max quality environmental effects, 10 bias to -1 detail level, 20 biased to -2 etc.
 }
 
-void FNiagaraWorldManager::CalculateScalabilityState(UNiagaraSystem* System, const FNiagaraScalabilitySettings& ScalabilitySettings, UNiagaraEffectType* EffectType, UNiagaraComponent* Component, bool bIsPreCull, FNiagaraScalabilityState& OutState)
+void FNiagaraWorldManager::CalculateScalabilityState(UNiagaraSystem* System, const FNiagaraSystemScalabilitySettings& ScalabilitySettings, UNiagaraEffectType* EffectType, UNiagaraComponent* Component, bool bIsPreCull, FNiagaraScalabilityState& OutState)
 {
 	float DistSignificance = DistanceSignificance(EffectType, ScalabilitySettings, Component);
 	
@@ -757,9 +787,6 @@ void FNiagaraWorldManager::CalculateScalabilityState(UNiagaraSystem* System, con
 	OutState.bDirty = OutState.bCulled != bOldCulled;
 
 	//TODO: More progressive scalability options?
-	//We can possibly adjust the DetailLevel each effect can use so we can more gradually scale down effects at lower significance.
-	//Possibly do similar to the Significance manager and define a set budget of instances that are allowed at each level.
-	//E.g. 5 max quality environmental effects, 10 bias to -1 detail level, 20 biased to -2 etc.
 }
 
 bool FNiagaraWorldManager::CanPreCull(UNiagaraEffectType* EffectType)
@@ -768,7 +795,7 @@ bool FNiagaraWorldManager::CanPreCull(UNiagaraEffectType* EffectType)
 	return EffectType->CullReaction == ENiagaraCullReaction::Deactivate || EffectType->CullReaction == ENiagaraCullReaction::DeactivateImmediate;
 }
 
-void FNiagaraWorldManager::SortedSignificanceCull(UNiagaraEffectType* EffectType, const FNiagaraScalabilitySettings& ScalabilitySettings, float Significance, int32 Index, FNiagaraScalabilityState& OutState)
+void FNiagaraWorldManager::SortedSignificanceCull(UNiagaraEffectType* EffectType, const FNiagaraSystemScalabilitySettings& ScalabilitySettings, float Significance, int32 Index, FNiagaraScalabilityState& OutState)
 {
 	//Cull all but the N most significance FX.
 	bool bCull = ScalabilitySettings.bCullMaxInstanceCount && Index >= ScalabilitySettings.MaxInstances;
@@ -778,7 +805,7 @@ void FNiagaraWorldManager::SortedSignificanceCull(UNiagaraEffectType* EffectType
 #endif
 }
 
-void FNiagaraWorldManager::SignificanceCull(UNiagaraEffectType* EffectType, const FNiagaraScalabilitySettings& ScalabilitySettings, float Significance, FNiagaraScalabilityState& OutState)
+void FNiagaraWorldManager::SignificanceCull(UNiagaraEffectType* EffectType, const FNiagaraSystemScalabilitySettings& ScalabilitySettings, float Significance, FNiagaraScalabilityState& OutState)
 {
 	float MinSignificance = 0.0f;
 
@@ -799,7 +826,7 @@ void FNiagaraWorldManager::SignificanceCull(UNiagaraEffectType* EffectType, cons
 #endif
 }
 
-void FNiagaraWorldManager::VisibilityCull(UNiagaraEffectType* EffectType, const FNiagaraScalabilitySettings& ScalabilitySettings, UNiagaraComponent* Component, FNiagaraScalabilityState& OutState)
+void FNiagaraWorldManager::VisibilityCull(UNiagaraEffectType* EffectType, const FNiagaraSystemScalabilitySettings& ScalabilitySettings, UNiagaraComponent* Component, FNiagaraScalabilityState& OutState)
 {
 	float TimeSinceRendered = Component->GetSafeTimeSinceRendered(World->TimeSeconds);
 	bool bCull = ScalabilitySettings.bCullByMaxTimeWithoutRender && TimeSinceRendered > ScalabilitySettings.MaxTimeWithoutRender;
@@ -810,7 +837,7 @@ void FNiagaraWorldManager::VisibilityCull(UNiagaraEffectType* EffectType, const
 #endif
 }
 
-void FNiagaraWorldManager::OwnerLODCull(UNiagaraEffectType* EffectType, const FNiagaraScalabilitySettings& ScalabilitySettings, UNiagaraComponent* Component, FNiagaraScalabilityState& OutState)
+void FNiagaraWorldManager::OwnerLODCull(UNiagaraEffectType* EffectType, const FNiagaraSystemScalabilitySettings& ScalabilitySettings, UNiagaraComponent* Component, FNiagaraScalabilityState& OutState)
 {
 	bool bCull = ScalabilitySettings.bCullByMaxOwnerLOD && Component->GetOwnerLOD() > ScalabilitySettings.MaxOwnerLOD;
 	OutState.bCulled |= bCull;
@@ -819,7 +846,7 @@ void FNiagaraWorldManager::OwnerLODCull(UNiagaraEffectType* EffectType, const FN
 #endif
 }
 
-void FNiagaraWorldManager::InstanceCountCull(UNiagaraEffectType* EffectType, const FNiagaraScalabilitySettings& ScalabilitySettings, FNiagaraScalabilityState& OutState)
+void FNiagaraWorldManager::InstanceCountCull(UNiagaraEffectType* EffectType, const FNiagaraSystemScalabilitySettings& ScalabilitySettings, FNiagaraScalabilityState& OutState)
 {
 	bool bCull = ScalabilitySettings.bCullMaxInstanceCount && EffectType->NumInstances > ScalabilitySettings.MaxInstances;
 	OutState.bCulled |= bCull;
@@ -828,7 +855,7 @@ void FNiagaraWorldManager::InstanceCountCull(UNiagaraEffectType* EffectType, con
 #endif
 }
 
-float FNiagaraWorldManager::DistanceSignificance(UNiagaraEffectType* EffectType, const FNiagaraScalabilitySettings& ScalabilitySettings, UNiagaraComponent* Component)
+float FNiagaraWorldManager::DistanceSignificance(UNiagaraEffectType* EffectType, const FNiagaraSystemScalabilitySettings& ScalabilitySettings, UNiagaraComponent* Component)
 {
 	float LODDistance = 0.0f;
 
@@ -867,7 +894,7 @@ float FNiagaraWorldManager::DistanceSignificance(UNiagaraEffectType* EffectType,
 	return 1.0f;
 }
 
-float FNiagaraWorldManager::DistanceSignificance(UNiagaraEffectType* EffectType, const FNiagaraScalabilitySettings& ScalabilitySettings, FVector Location)
+float FNiagaraWorldManager::DistanceSignificance(UNiagaraEffectType* EffectType, const FNiagaraSystemScalabilitySettings& ScalabilitySettings, FVector Location)
 {
 	if (ScalabilitySettings.bCullByDistance && bCachedPlayerViewLocationsValid)
 	{
diff --git a/Engine/Plugins/FX/Niagara/Source/Niagara/Public/NiagaraComponent.h b/Engine/Plugins/FX/Niagara/Source/Niagara/Public/NiagaraComponent.h
--- a/Engine/Plugins/FX/Niagara/Source/Niagara/Public/NiagaraComponent.h
+++ b/Engine/Plugins/FX/Niagara/Source/Niagara/Public/NiagaraComponent.h
@@ -154,6 +154,8 @@ public:
 
 	virtual void Activate(bool bReset = false)override;
 	virtual void Deactivate()override;
+
+	UFUNCTION(BlueprintCallable, Category = Niagara)
 	void DeactivateImmediate();
 
 	FORCEINLINE ENiagaraExecutionState GetRequestedExecutionState()const { return SystemInstance ? SystemInstance->GetRequestedExecutionState() : ENiagaraExecutionState::Complete; }
@@ -403,21 +405,12 @@ public:
 #endif
 	//~ End UObject Interface
 
-	UFUNCTION(BlueprintCallable, Category = Preview, meta = (Keywords = "preview detail level scalability"))
-	void SetPreviewDetailLevel(bool bEnablePreviewDetailLevel, int32 PreviewDetailLevel);
-
 	UFUNCTION(BlueprintCallable, Category = Preview, meta = (Keywords = "preview LOD Distance scalability"))
 	void SetPreviewLODDistance(bool bEnablePreviewLODDistance, float PreviewLODDistance);
 
-	UFUNCTION(BlueprintCallable, Category = Preview, meta = (Keywords = "preview detail level scalability"))
-	FORCEINLINE bool GetPreviewDetailLevelEnabled()const;
-
 	UFUNCTION(BlueprintCallable, Category = Preview, meta = (Keywords = "preview LOD Distance scalability"))
 	FORCEINLINE bool GetPreviewLODDistanceEnabled()const;
 
-	UFUNCTION(BlueprintCallable, Category = Preview, meta = (Keywords = "preview detail level scalability"))
-	FORCEINLINE int32 GetPreviewDetailLevel()const;
-
 	UFUNCTION(BlueprintCallable, Category = Preview, meta = (Keywords = "preview LOD Distance scalability"))
 	FORCEINLINE int32 GetPreviewLODDistance()const;
 
@@ -507,9 +500,7 @@ public:
 	virtual void SetUseAutoManageAttachment(bool bAutoManage) override { bAutoManageAttachment = bAutoManage; }
 
 #if WITH_NIAGARA_COMPONENT_PREVIEW_DATA
-	int32 PreviewDetailLevel;
 	float PreviewLODDistance;
-	uint32 bEnablePreviewDetailLevel : 1;
 	uint32 bEnablePreviewLODDistance : 1;
 #endif
 
@@ -564,14 +555,10 @@ private:
 };
 
 #if WITH_NIAGARA_COMPONENT_PREVIEW_DATA
-FORCEINLINE bool UNiagaraComponent::GetPreviewDetailLevelEnabled()const { return bEnablePreviewDetailLevel; }
 FORCEINLINE bool UNiagaraComponent::GetPreviewLODDistanceEnabled()const { return bEnablePreviewLODDistance; }
-FORCEINLINE int32 UNiagaraComponent::GetPreviewDetailLevel()const { return bEnablePreviewDetailLevel ? PreviewDetailLevel : INDEX_NONE; }
 FORCEINLINE int32 UNiagaraComponent::GetPreviewLODDistance()const { return bEnablePreviewLODDistance ? PreviewLODDistance : 0.0f; }
 #else
-FORCEINLINE bool UNiagaraComponent::GetPreviewDetailLevelEnabled()const { return false; }
 FORCEINLINE bool UNiagaraComponent::GetPreviewLODDistanceEnabled()const { return false; }
-FORCEINLINE int32 UNiagaraComponent::GetPreviewDetailLevel()const { return INDEX_NONE; }
 FORCEINLINE int32 UNiagaraComponent::GetPreviewLODDistance()const { return 0.0f; }
 #endif
 
diff --git a/Engine/Plugins/FX/Niagara/Source/Niagara/Public/NiagaraModule.h b/Engine/Plugins/FX/Niagara/Source/Niagara/Public/NiagaraModule.h
--- a/Engine/Plugins/FX/Niagara/Source/Niagara/Public/NiagaraModule.h
+++ b/Engine/Plugins/FX/Niagara/Source/Niagara/Public/NiagaraModule.h
@@ -72,10 +72,10 @@ public:
 
 #endif
 
-	FORCEINLINE static int32 GetDetailLevel() { return EngineDetailLevel; }
+	FORCEINLINE static int32 GetEffectsQuality() { return EngineEffectsQuality; }
+
 	FORCEINLINE static float GetGlobalSpawnCountScale() { return EngineGlobalSpawnCountScale; }
 	FORCEINLINE static float GetGlobalSystemCountScale() { return EngineGlobalSystemCountScale; }
-	static bool IsTargetPlatformIncludedInLevelRangeForCook(const ITargetPlatform* InTargetPlatform, const class UNiagaraEmitter* InEmitter);
 
 	static float EngineGlobalSpawnCountScale;
 	static float EngineGlobalSystemCountScale;
@@ -165,6 +165,8 @@ public:
 	FORCEINLINE static const FNiagaraVariable&  GetVar_DataInstance_Alive() { return DataInstance_Alive; }
 	FORCEINLINE static const FNiagaraVariable&  GetVar_BeginDefaults() { return Translator_BeginDefaults; }
 
+	static void OnEffectsQualityChanged(int32 NewEffectsQuality);
+
 	FOnProcessQueue OnProcessQueue;
 
 #if WITH_EDITORONLY_DATA
@@ -176,8 +178,8 @@ public:
 	FOnPrecompile ObjectPrecompilerDelegate;
 #endif
 
-	void OnChangeDetailLevel(class IConsoleVariable* CVar);
-	static int32 EngineDetailLevel;
+	static int32 EngineEffectsQuality;
+
 
 private:
 	static FNiagaraVariable Engine_DeltaTime;
diff --git a/Engine/Plugins/FX/Niagara/Source/Niagara/Public/NiagaraSystemInstance.h b/Engine/Plugins/FX/Niagara/Source/Niagara/Public/NiagaraSystemInstance.h
--- a/Engine/Plugins/FX/Niagara/Source/Niagara/Public/NiagaraSystemInstance.h
+++ b/Engine/Plugins/FX/Niagara/Source/Niagara/Public/NiagaraSystemInstance.h
@@ -265,8 +265,6 @@ public:
 	/*void SetHasGPUEmitters(bool bInHasGPUEmitters) { bHasGPUEmitters = bInHasGPUEmitters; }*/
 	bool HasGPUEmitters() { return bHasGPUEmitters;  }
 
-	int32 GetDetailLevel()const;
-
 	FORCEINLINE void BeginAsyncWork()
 	{
 		bAsyncWorkInProgress = true;
diff --git a/Engine/Plugins/FX/Niagara/Source/Niagara/Public/NiagaraWorldManager.h b/Engine/Plugins/FX/Niagara/Source/Niagara/Public/NiagaraWorldManager.h
--- a/Engine/Plugins/FX/Niagara/Source/Niagara/Public/NiagaraWorldManager.h
+++ b/Engine/Plugins/FX/Niagara/Source/Niagara/Public/NiagaraWorldManager.h
@@ -125,6 +125,7 @@ public:
 
 	void OnWorldCleanup(bool bSessionEnded, bool bCleanupResources);
 
+	void PreGarbageCollect();
 	void PostGarbageCollect();
 	void PreGarbageCollectBeginDestroy();
 	
@@ -152,10 +153,10 @@ public:
 	NIAGARA_API bool ShouldPreCull(UNiagaraSystem* System, UNiagaraComponent* Component);
 	NIAGARA_API bool ShouldPreCull(UNiagaraSystem* System, FVector Location);
 
-	void CalculateScalabilityState(UNiagaraSystem* System, const FNiagaraScalabilitySettings& ScalabilitySettings, UNiagaraEffectType* EffectType, UNiagaraComponent* Component, bool bIsPreCull, FNiagaraScalabilityState& OutState);
-	void CalculateScalabilityState(UNiagaraSystem* System, const FNiagaraScalabilitySettings& ScalabilitySettings, UNiagaraEffectType* EffectType, FVector Location, bool bIsPreCull, FNiagaraScalabilityState& OutState);
+	void CalculateScalabilityState(UNiagaraSystem* System, const FNiagaraSystemScalabilitySettings& ScalabilitySettings, UNiagaraEffectType* EffectType, UNiagaraComponent* Component, bool bIsPreCull, FNiagaraScalabilityState& OutState);
+	void CalculateScalabilityState(UNiagaraSystem* System, const FNiagaraSystemScalabilitySettings& ScalabilitySettings, UNiagaraEffectType* EffectType, FVector Location, bool bIsPreCull, FNiagaraScalabilityState& OutState);
 
-	/*FORCEINLINE_DEBUGGABLE*/ void SortedSignificanceCull(UNiagaraEffectType* EffectType, const FNiagaraScalabilitySettings& ScalabilitySettings, float Significance, int32 Index, FNiagaraScalabilityState& OutState);
+	/*FORCEINLINE_DEBUGGABLE*/ void SortedSignificanceCull(UNiagaraEffectType* EffectType, const FNiagaraSystemScalabilitySettings& ScalabilitySettings, float Significance, int32 Index, FNiagaraScalabilityState& OutState);
 
 #if DEBUG_SCALABILITY_STATE
 	void DumpScalabilityState();
@@ -178,6 +179,9 @@ private:
 	// Callback for when a world is ticked.
 	static void TickWorld(UWorld* World, ELevelTick TickType, float DeltaSeconds);
 
+	// Callback to handle any pre GC processing needed.
+	static void OnPreGarbageCollect();
+
 	// Callback to handle any post GC processing needed.
 	static void OnPostGarbageCollect();
 
@@ -189,20 +193,21 @@ private:
 
 	FORCEINLINE_DEBUGGABLE bool CanPreCull(UNiagaraEffectType* EffectType);
 
-	FORCEINLINE_DEBUGGABLE void SignificanceCull(UNiagaraEffectType* EffectType, const FNiagaraScalabilitySettings& ScalabilitySettings, float Significance, FNiagaraScalabilityState& OutState);
-	FORCEINLINE_DEBUGGABLE void VisibilityCull(UNiagaraEffectType* EffectType, const FNiagaraScalabilitySettings& ScalabilitySettings, UNiagaraComponent* Component, FNiagaraScalabilityState& OutState);
-	FORCEINLINE_DEBUGGABLE void OwnerLODCull(UNiagaraEffectType* EffectType, const FNiagaraScalabilitySettings& ScalabilitySettings, UNiagaraComponent* Component, FNiagaraScalabilityState& OutState);
-	FORCEINLINE_DEBUGGABLE void InstanceCountCull(UNiagaraEffectType* EffectType, const FNiagaraScalabilitySettings& ScalabilitySettings, FNiagaraScalabilityState& OutState);
+	FORCEINLINE_DEBUGGABLE void SignificanceCull(UNiagaraEffectType* EffectType, const FNiagaraSystemScalabilitySettings& ScalabilitySettings, float Significance, FNiagaraScalabilityState& OutState);
+	FORCEINLINE_DEBUGGABLE void VisibilityCull(UNiagaraEffectType* EffectType, const FNiagaraSystemScalabilitySettings& ScalabilitySettings, UNiagaraComponent* Component, FNiagaraScalabilityState& OutState);
+	FORCEINLINE_DEBUGGABLE void OwnerLODCull(UNiagaraEffectType* EffectType, const FNiagaraSystemScalabilitySettings& ScalabilitySettings, UNiagaraComponent* Component, FNiagaraScalabilityState& OutState);
+	FORCEINLINE_DEBUGGABLE void InstanceCountCull(UNiagaraEffectType* EffectType, const FNiagaraSystemScalabilitySettings& ScalabilitySettings, FNiagaraScalabilityState& OutState);
 
 	/** Calculate significance contribution from the distance to nearest view. */
-	FORCEINLINE_DEBUGGABLE float DistanceSignificance(UNiagaraEffectType* EffectType, const FNiagaraScalabilitySettings& ScalabilitySettings, FVector Location);
-	FORCEINLINE_DEBUGGABLE float DistanceSignificance(UNiagaraEffectType* EffectType, const FNiagaraScalabilitySettings& ScalabilitySettings, UNiagaraComponent* Component);
+	FORCEINLINE_DEBUGGABLE float DistanceSignificance(UNiagaraEffectType* EffectType, const FNiagaraSystemScalabilitySettings& ScalabilitySettings, FVector Location);
+	FORCEINLINE_DEBUGGABLE float DistanceSignificance(UNiagaraEffectType* EffectType, const FNiagaraSystemScalabilitySettings& ScalabilitySettings, UNiagaraComponent* Component);
 	
 	static FDelegateHandle OnWorldInitHandle;
 	static FDelegateHandle OnWorldCleanupHandle;
 	static FDelegateHandle OnPreWorldFinishDestroyHandle;
 	static FDelegateHandle OnWorldBeginTearDownHandle;
 	static FDelegateHandle TickWorldHandle;
+	static FDelegateHandle PreGCHandle;
 	static FDelegateHandle PostGCHandle;
 	static FDelegateHandle PreGCBeginDestroyHandle;
 
diff --git a/Engine/Plugins/FX/Niagara/Source/NiagaraCore/Public/NiagaraCustomVersion.h b/Engine/Plugins/FX/Niagara/Source/NiagaraCore/Public/NiagaraCustomVersion.h
--- a/Engine/Plugins/FX/Niagara/Source/NiagaraCore/Public/NiagaraCustomVersion.h
+++ b/Engine/Plugins/FX/Niagara/Source/NiagaraCore/Public/NiagaraCustomVersion.h
@@ -140,6 +140,7 @@ struct FNiagaraCustomVersion
 
 		AddGeneratedFunctionsToGPUParamInfo, // Save the functions generated by a GPU data interface inside FNiagaraDataInterfaceGPUParamInfo
 
+		PlatformScalingRefactor, //Removed DetailLevel in favor of FNiagaraPlatfomSet based selection of per platform settings.
 		// DO NOT ADD A NEW VERSION UNLESS YOU HAVE TALKED TO THE NIAGARA LEAD. Mismanagement of these versions can lead to data loss if it is adjusted in multiple streams simultaneously.
 		// -----<new versions can be added above this line>  -------------------------------------------------
 		VersionPlusOne,
diff --git a/Engine/Plugins/FX/Niagara/Source/NiagaraEditor/NiagaraEditor.Build.cs b/Engine/Plugins/FX/Niagara/Source/NiagaraEditor/NiagaraEditor.Build.cs
--- a/Engine/Plugins/FX/Niagara/Source/NiagaraEditor/NiagaraEditor.Build.cs
+++ b/Engine/Plugins/FX/Niagara/Source/NiagaraEditor/NiagaraEditor.Build.cs
@@ -41,6 +41,7 @@ public class NiagaraEditor : ModuleRules
 				"GraphEditor",
                 "ShaderFormatVectorVM",
                 "TargetPlatform",
+                "DesktopPlatform",
                 "AppFramework",
 				"MovieSceneTools",
                 "MovieSceneTracks",
diff --git a/Engine/Plugins/FX/Niagara/Source/NiagaraEditor/Private/Customizations/NiagaraTypeCustomizations.cpp b/Engine/Plugins/FX/Niagara/Source/NiagaraEditor/Private/Customizations/NiagaraTypeCustomizations.cpp
--- a/Engine/Plugins/FX/Niagara/Source/NiagaraEditor/Private/Customizations/NiagaraTypeCustomizations.cpp
+++ b/Engine/Plugins/FX/Niagara/Source/NiagaraEditor/Private/Customizations/NiagaraTypeCustomizations.cpp
@@ -2,28 +2,35 @@
 
 #include "NiagaraTypeCustomizations.h"
 #include "CoreMinimal.h"
-#include "PropertyHandle.h"
-#include "DetailWidgetRow.h"
 #include "DetailLayoutBuilder.h"
-#include "NiagaraTypes.h"
-#include "Widgets/Text/STextBlock.h"
-#include "Widgets/Input/SCheckBox.h"
+#include "DetailWidgetRow.h"
 #include "IDetailChildrenBuilder.h"
-#include "NiagaraRendererProperties.h"
+#include "NiagaraConstants.h"
+#include "NiagaraConstants.h"
+#include "NiagaraEditorStyle.h"
 #include "NiagaraEmitter.h"
-#include "SGraphActionMenu.h"
-#include "Framework/Application/SlateApplication.h"
-#include "ScopedTransaction.h"
-#include "NiagaraParameterMapHistory.h"
-#include "NiagaraScriptSource.h"
 #include "NiagaraNodeOutput.h"
-#include "NiagaraConstants.h"
 #include "NiagaraNodeParameterMapBase.h"
-#include "Widgets/Input/STextComboBox.h"
-#include "NiagaraSystem.h"
-
+#include "NiagaraParameterMapHistory.h"
+#include "NiagaraPlatformSet.h"
+#include "NiagaraRendererProperties.h"
+#include "NiagaraScriptSource.h"
 #include "NiagaraScriptVariable.h"
-#include "NiagaraConstants.h"
+#include "NiagaraSystem.h"
+#include "NiagaraTypes.h"
+#include "PlatformInfo.h"
+#include "PropertyHandle.h"
+#include "SGraphActionMenu.h"
+#include "Scalability.h"
+#include "ScopedTransaction.h"
+#include "DeviceProfiles/DeviceProfile.h"
+#include "DeviceProfiles/DeviceProfileManager.h"
+#include "Framework/Application/SlateApplication.h"
+#include "Widgets/Input/SButton.h"
+#include "Widgets/Input/SCheckBox.h"
+#include "Widgets/Input/STextComboBox.h"
+#include "Widgets/Layout/SWrapBox.h"
+#include "Widgets/Text/STextBlock.h"
 
 #define LOCTEXT_NAMESPACE "FNiagaraVariableAttributeBindingCustomization"
 
@@ -457,11 +464,11 @@ void FNiagaraUserParameterBindingCustomization::CustomizeHeader(TSharedRef<IProp
 				.Font(IDetailLayoutBuilder::GetDetailFont())
 				]
 				];
+
 			bAddDefault = false;
 		}
 	}
 
-
 	if (bAddDefault)
 	{
 		HeaderRow
@@ -662,14 +669,14 @@ void FNiagaraScriptVariableBindingCustomization::CustomizeHeader(TSharedRef<IPro
 					[
 						SNew(SComboButton)
 						.OnGetMenuContent(this, &FNiagaraScriptVariableBindingCustomization::OnGetMenuContent)
-					.ContentPadding(1)
-					.ToolTipText(this, &FNiagaraScriptVariableBindingCustomization::GetTooltipText)
-					.ButtonContent()
-					[
-						SNew(STextBlock)
-						.Text(this, &FNiagaraScriptVariableBindingCustomization::GetCurrentText)
-					.Font(IDetailLayoutBuilder::GetDetailFont())
-					]
+						.ContentPadding(1)
+						.ToolTipText(this, &FNiagaraScriptVariableBindingCustomization::GetTooltipText)
+						.ButtonContent()
+						[
+							SNew(STextBlock)
+							.Text(this, &FNiagaraScriptVariableBindingCustomization::GetCurrentText)
+							.Font(IDetailLayoutBuilder::GetDetailFont())
+						]
 					];
 				bAddDefault = false;
 			}
@@ -700,7 +707,938 @@ void FNiagaraScriptVariableBindingCustomization::CustomizeHeader(TSharedRef<IPro
 			];
 	}
 }
+//////////////////////////////////////////////////////////////////////////
+
+void FNiagaraPlatformSetTypeCustomization::CustomizeHeader(TSharedRef<IPropertyHandle> InPropertyHandle, FDetailWidgetRow& HeaderRow, IPropertyTypeCustomizationUtils& CustomizationUtils)
+{
+	PropertyHandle = InPropertyHandle;
+	TArray<UObject*> Objects;
+	PropertyHandle->GetOuterObjects(Objects);
+	TargetPlatformSet = (FNiagaraPlatformSet*)PropertyHandle->GetValueBaseAddress((uint8*)Objects[0]);
+
+	if (PlatformSelectionStates.Num() == 0)
+	{
+		PlatformSelectionStates.Add(MakeShared<ENiagaraPlatformSelectionState>(ENiagaraPlatformSelectionState::Enabled));
+		PlatformSelectionStates.Add(MakeShared<ENiagaraPlatformSelectionState>(ENiagaraPlatformSelectionState::Disabled));
+		PlatformSelectionStates.Add(MakeShared<ENiagaraPlatformSelectionState>(ENiagaraPlatformSelectionState::Default));
+	}
+
+	//Look for outer types for which we need to ensure there are no conflicting settings.
+	SystemScalabilitySettings = nullptr;
+	EmitterScalabilitySettings = nullptr;
+	SystemScalabilityOverrides = nullptr;
+	EmitterScalabilityOverrides = nullptr;
+	PlatformSetArray.Reset();
+	PlatformSetArrayIndex = INDEX_NONE;
+
+	//Look whether this platform set belongs to a class which must keep orthogonal platform sets.
+	//We then interrogate the other sets via these ptrs to look for conflicts.
+	TSharedPtr<IPropertyHandle> CurrHandle = PropertyHandle->GetParentHandle();
+	while (CurrHandle)
+	{
+		int32 ThisIndex = CurrHandle->GetIndexInArray();
+		if (ThisIndex != INDEX_NONE)
+		{
+			PlatformSetArray = CurrHandle->GetParentHandle()->AsArray();
+			PlatformSetArrayIndex = ThisIndex;
+		}
+
+		if (UProperty* CurrProperty = CurrHandle->GetProperty())
+		{
+			if (UStruct* CurrStruct = CurrProperty->GetOwnerStruct())
+			{
+				TSharedPtr<IPropertyHandle> ParentHandle = CurrHandle->GetParentHandle();
+				if (CurrStruct == FNiagaraSystemScalabilitySettingsArray::StaticStruct())
+				{
+					SystemScalabilitySettings = (FNiagaraSystemScalabilitySettingsArray*)ParentHandle->GetValueBaseAddress((uint8*)Objects[0]);
+					break;
+				}
+				else if (CurrStruct == FNiagaraEmitterScalabilitySettingsArray::StaticStruct())
+				{
+					EmitterScalabilitySettings = (FNiagaraEmitterScalabilitySettingsArray*)ParentHandle->GetValueBaseAddress((uint8*)Objects[0]);
+					break;
+				}
+				else if (CurrStruct == FNiagaraSystemScalabilityOverrides::StaticStruct())
+				{
+					SystemScalabilityOverrides = (FNiagaraSystemScalabilityOverrides*)ParentHandle->GetValueBaseAddress((uint8*)Objects[0]);
+					break;
+				}
+				else if (CurrStruct == FNiagaraEmitterScalabilityOverrides::StaticStruct())
+				{
+					EmitterScalabilityOverrides = (FNiagaraEmitterScalabilityOverrides*)ParentHandle->GetValueBaseAddress((uint8*)Objects[0]);
+					break;
+				}
+			}
+		}
+
+		CurrHandle = CurrHandle->GetParentHandle();
+	}
+
+	UpdateCachedConflicts();
+	
+	BaseSystem = Objects[0]->GetTypedOuter<UNiagaraSystem>();
+	PropertyHandle->SetOnPropertyValueChanged(FSimpleDelegate::CreateSP(this, &FNiagaraPlatformSetTypeCustomization::OnPropertyValueChanged));
+
+	HeaderRow
+		.WholeRowContent()
+		[ 	
+			SAssignNew(EffectsQualityWidgetBox, SWrapBox)
+			.UseAllottedWidth(true)
+		];
+
+	GenerateEffectsQualitySelectionWidgets();
+}
+
+void FNiagaraPlatformSetTypeCustomization::CustomizeChildren(TSharedRef<IPropertyHandle> StructPropertyHandle, class IDetailChildrenBuilder& ChildBuilder, IPropertyTypeCustomizationUtils& StructCustomizationUtils)
+{
+}
+
+FText FNiagaraPlatformSetTypeCustomization::GetCurrentText() const
+{
+	return LOCTEXT("Platforms", "Platforms");
+}
+
+FText FNiagaraPlatformSetTypeCustomization::GetTooltipText() const
+{
+	return LOCTEXT("Platforms", "Platforms");
+}
+
+void FNiagaraPlatformSetTypeCustomization::GenerateEffectsQualitySelectionWidgets()
+{
+	EffectsQualityWidgetBox->ClearChildren();
+
+	Scalability::FQualityLevels Counts = Scalability::GetQualityLevelCounts();
+	EffectsQualityMenus.SetNum(Counts.EffectsQuality);
+
+	for (int32 EffectsQuality = 0; EffectsQuality < Counts.EffectsQuality; ++EffectsQuality)
+	{
+		bool First = EffectsQuality == 0;
+		bool Last = EffectsQuality == (Counts.EffectsQuality - 1);
+
+		EffectsQualityWidgetBox->AddSlot()
+			.Padding(0, 0, 1, 0)
+			[
+				SNew(SBox)
+				.WidthOverride(100)
+				.VAlign(VAlign_Center)
+				[
+					SNew(SVerticalBox)
+					+ SVerticalBox::Slot()
+					.AutoHeight()
+					[
+						SNew(SCheckBox)
+						.Style(FNiagaraEditorStyle::Get(), First ? "NiagaraEditor.PlatformSet.StartButton" : 
+							(Last ? "NiagaraEditor.PlatformSet.EndButton" : "NiagaraEditor.PlatformSet.MiddleButton"))
+						.IsChecked(this, &FNiagaraPlatformSetTypeCustomization::IsEQChecked, EffectsQuality)
+						.OnCheckStateChanged(this, &FNiagaraPlatformSetTypeCustomization::EQCheckStateChanged, EffectsQuality)
+						[
+							SNew(SHorizontalBox)
+							+ SHorizontalBox::Slot()
+							.HAlign(HAlign_Fill)
+							.Padding(6,2,2,4)
+							[
+								SNew(STextBlock)
+								.TextStyle(FNiagaraEditorStyle::Get(), "NiagaraEditor.PlatformSet.ButtonText")
+								.Text(FNiagaraPlatformSet::GetEffectsQualityText(EffectsQuality))
+								.ColorAndOpacity(this, &FNiagaraPlatformSetTypeCustomization::GetEffectsQualityButtonTextColor, EffectsQuality)
+							]
+							// error icon
+							+ SHorizontalBox::Slot()
+							.AutoWidth()
+							.HAlign(HAlign_Right)
+							.VAlign(VAlign_Center)
+							.Padding(0,0,2,0)
+							[
+								SNew(SBox)
+								.WidthOverride(12)
+								.HeightOverride(12)
+								[
+									SNew(SImage)
+									.ToolTipText(this, &FNiagaraPlatformSetTypeCustomization::GetEffectsQualityErrorToolTip, EffectsQuality)
+									.Visibility(this, &FNiagaraPlatformSetTypeCustomization::GetEffectsQualityErrorVisibility, EffectsQuality)
+									.Image(FEditorStyle::GetBrush("Icons.Error"))
+								]
+							]
+							// dropdown button
+							+ SHorizontalBox::Slot()
+							.AutoWidth()
+							.HAlign(HAlign_Right)
+							.VAlign(VAlign_Center)
+							.Padding(0,0,2,0)
+							[
+								// dropdown button
+								SAssignNew(EffectsQualityMenus[EffectsQuality], SMenuAnchor)
+								.ToolTipText(LOCTEXT("AddPlatformOverride", "Add an override for a specific platform."))
+								.OnGetMenuContent(this, &FNiagaraPlatformSetTypeCustomization::GenerateDeviceProfileTreeWidget, EffectsQuality)
+								[
+									SNew(SButton)
+									.Visibility_Lambda([this, EffectsQuality]()
+									{
+										return EffectsQualityWidgetBox->GetChildren()->GetChildAt(EffectsQuality)->IsHovered() ? EVisibility::Visible : EVisibility::Hidden;
+									})
+									.ButtonStyle(FEditorStyle::Get(), "HoverHintOnly")
+									.ForegroundColor(FSlateColor::UseForeground())
+									.OnClicked(this, &FNiagaraPlatformSetTypeCustomization::ToggleMenuOpenForEffectsQuality, EffectsQuality)
+									[
+										SNew(SBox)
+										.WidthOverride(8)
+										.HeightOverride(8)
+										.VAlign(VAlign_Center)
+										[
+											SNew(SImage)
+											.Image(FNiagaraEditorStyle::Get().GetBrush("NiagaraEditor.PlatformSet.DropdownButton"))
+										]
+									]
+								]
+							]
+						]
+					]
+					+ SVerticalBox::Slot()
+					.AutoHeight()
+					[
+						GenerateAdditionalDevicesWidgetForEQ(EffectsQuality)
+					]
+				]
+			];
+	}
+}
+
+void FNiagaraPlatformSetTypeCustomization::UpdateCachedConflicts()
+{
+	TArray<const FNiagaraPlatformSet*> PlatformSets;
+	if (SystemScalabilityOverrides != nullptr)
+	{
+		for (FNiagaraSystemScalabilityOverride& Override : SystemScalabilityOverrides->Overrides)
+		{
+			PlatformSets.Add(&Override.Platforms);
+		}
+	}
+	else if (EmitterScalabilityOverrides != nullptr)
+	{
+		for (FNiagaraEmitterScalabilityOverride& Override : EmitterScalabilityOverrides->Overrides)
+		{
+			PlatformSets.Add(&Override.Platforms);
+		}
+	}
+	else if (SystemScalabilitySettings != nullptr)
+	{
+		for (FNiagaraSystemScalabilitySettings& Settings : SystemScalabilitySettings->Settings)
+		{
+			PlatformSets.Add(&Settings.Platforms);
+		}
+	}
+	else if (EmitterScalabilitySettings != nullptr)
+	{
+		for (FNiagaraEmitterScalabilitySettings& Settings : EmitterScalabilitySettings->Settings)
+		{
+			PlatformSets.Add(&Settings.Platforms);
+		}
+	}
+
+	CachedConflicts.Reset();
+	FNiagaraPlatformSet::GatherConflicts(PlatformSets, CachedConflicts);
+}
+
+static TSharedPtr<IPropertyHandle> FindChildPlatformSet(TSharedPtr<IPropertyHandle> PropertyHandle)
+{
+	if (UStructProperty* Property = Cast<UStructProperty>(PropertyHandle->GetProperty()))
+	{
+		if (Property->Struct == FNiagaraPlatformSet::StaticStruct())
+		{
+			return PropertyHandle;
+		}
+	}
+
+	// recurse
+	uint32 NumChildren = 0;
+	PropertyHandle->GetNumChildren(NumChildren);
+
+	for (int32 Idx = 0; Idx < (int32) NumChildren; ++Idx)
+	{
+		TSharedPtr<IPropertyHandle> Child = PropertyHandle->GetChildHandle(Idx);
+
+		TSharedPtr<IPropertyHandle> ChildResult = FindChildPlatformSet(Child);
+		if (ChildResult.IsValid())
+		{
+			return ChildResult;
+		}
+	}
+
+	return TSharedPtr<IPropertyHandle>();
+}
+
+void FNiagaraPlatformSetTypeCustomization::InvalidateSiblingConflicts() const
+{
+	if (!PlatformSetArray.IsValid())
+	{
+		return;
+	}
+
+	uint32 ArrayCount = 0;
+	PlatformSetArray->GetNumElements(ArrayCount);
+	for (int32 Idx = 0; Idx < (int32) ArrayCount; ++Idx)
+	{
+		if (Idx == PlatformSetArrayIndex)
+		{
+			continue; // skip self
+		}
+
+		TSharedRef<IPropertyHandle> Sibling = PlatformSetArray->GetElement(Idx);
+
+		TSharedPtr<IPropertyHandle> SiblingPlatformSet = FindChildPlatformSet(Sibling);
+		if (SiblingPlatformSet.IsValid())
+		{
+			SiblingPlatformSet->NotifyPostChange();
+		}
+	}
+}
+
+EVisibility FNiagaraPlatformSetTypeCustomization::GetEffectsQualityErrorVisibility(int32 EffectsQuality) const
+{
+	// not part of an array
+	if (PlatformSetArrayIndex == INDEX_NONE)
+	{
+		return EVisibility::Collapsed;
+	}
+
+	for (const FNiagaraPlatformSetConflictInfo& ConflictInfo : CachedConflicts)
+	{
+		if (ConflictInfo.SetAIndex == PlatformSetArrayIndex || 
+			ConflictInfo.SetBIndex == PlatformSetArrayIndex)
+		{
+			// this conflict applies to this platform set, check if it applies to this effects quality button
+			const int32 EQMask = FNiagaraPlatformSet::CreateEQMask(EffectsQuality);
+
+			for (const FNiagaraPlatformSetConflictEntry& Conflict : ConflictInfo.Conflicts)
+			{
+				if ((EQMask & Conflict.EffectsQualityMask) != 0)
+				{
+					return EVisibility::Visible;
+				}
+			}
+		}
+	}
+
+	return EVisibility::Collapsed;
+}
+
+FText FNiagaraPlatformSetTypeCustomization::GetEffectsQualityErrorToolTip(int32 EffectsQuality) const
+{
+	if (PlatformSetArrayIndex == INDEX_NONE)
+	{
+		return FText::GetEmpty();
+	}
+
+	for (const FNiagaraPlatformSetConflictInfo& ConflictInfo : CachedConflicts)
+	{
+		int32 OtherIndex = INDEX_NONE;
+
+		if (ConflictInfo.SetAIndex == PlatformSetArrayIndex)
+		{
+			OtherIndex = ConflictInfo.SetBIndex;
+		}
+
+		if (ConflictInfo.SetBIndex == PlatformSetArrayIndex)
+		{
+			OtherIndex = ConflictInfo.SetAIndex;
+		}
+
+		if (OtherIndex != INDEX_NONE)
+		{
+			const int32 EQMask = FNiagaraPlatformSet::CreateEQMask(EffectsQuality);
+
+			for (const FNiagaraPlatformSetConflictEntry& Conflict : ConflictInfo.Conflicts)
+			{
+				if ((EQMask & Conflict.EffectsQualityMask) != 0)
+				{
+					FText FormatString = LOCTEXT("EffectsQualityConflictToolTip", "This effect quality conflicts with the set at index {Index} in this array."); 
+
+					FFormatNamedArguments Args;
+					Args.Add(TEXT("Index"), OtherIndex);
+
+					return FText::Format(FormatString, Args);
+				}
+			}
+		}
+	}
+
+	return FText::GetEmpty();
+}
+
+FSlateColor FNiagaraPlatformSetTypeCustomization::GetEffectsQualityButtonTextColor(int32 EffectsQuality) const
+{
+	return TargetPlatformSet->IsEffectQualityEnabled(EffectsQuality) ? 
+		FSlateColor(FLinearColor(0.95f, 0.95f, 0.95f)) :
+		FSlateColor::UseForeground();
+}
+
+TSharedRef<SWidget> FNiagaraPlatformSetTypeCustomization::GenerateAdditionalDevicesWidgetForEQ(int32 EffectsQuality)
+{
+	TSharedRef<SVerticalBox> Container = SNew(SVerticalBox);
+
+	auto AddDeviceProfileOverrideWidget = [&](UDeviceProfile* Profile, bool bEnabled)
+	{
+		TSharedPtr<SHorizontalBox> DeviceBox;
+
+		Container->AddSlot()
+			.AutoHeight()
+			[
+				SAssignNew(DeviceBox, SHorizontalBox)
+			];
+
+		const FText DeviceNameText = FText::FromName(Profile->GetFName());
+
+		DeviceBox->AddSlot()
+			.AutoWidth()
+			.HAlign(HAlign_Left)
+			.VAlign(VAlign_Center)
+			.Padding(3, 0, 3, 0)
+			[
+				SNew(SBox)
+				.WidthOverride(8)
+				.HeightOverride(8)
+				.ToolTipText(bEnabled ?
+					LOCTEXT("DeviceIncludedToolTip", "This device is included.") :
+					LOCTEXT("DeviceExcludedToolTip", "This device is excluded."))
+				[
+					SNew(SImage)
+					.Image(bEnabled ?
+						FNiagaraEditorStyle::Get().GetBrush("NiagaraEditor.PlatformSet.Include") :
+						FNiagaraEditorStyle::Get().GetBrush("NiagaraEditor.PlatformSet.Exclude"))
+					.ColorAndOpacity(bEnabled ?
+						FLinearColor(0, 1, 0) :
+						FLinearColor(1, 0, 0))
+					]
+				];
 
+		DeviceBox->AddSlot()
+			.HAlign(HAlign_Fill)
+			.Padding(0, 0, 2, 0)
+			[
+				SNew(STextBlock)
+				.Text(DeviceNameText)
+				.ToolTipText(DeviceNameText)
+			];
+
+		DeviceBox->AddSlot()
+			.AutoWidth()
+			.HAlign(HAlign_Right)
+			.VAlign(VAlign_Center)
+			.Padding(0, 0, 2, 0)
+			[
+				SNew(SBox)
+				.WidthOverride(12)
+				.HeightOverride(12)
+				[
+					SNew(SImage)
+					.ToolTipText(this, &FNiagaraPlatformSetTypeCustomization::GetDeviceProfileErrorToolTip, Profile, EffectsQuality)
+					.Visibility(this, &FNiagaraPlatformSetTypeCustomization::GetDeviceProfileErrorVisibility, Profile, EffectsQuality)
+					.Image(FEditorStyle::GetBrush("Icons.Error"))
+				]
+			];
 
+		DeviceBox->AddSlot()
+			.AutoWidth()
+			.HAlign(HAlign_Right)
+			.VAlign(VAlign_Center)
+			.Padding(0, 0, 2, 0)
+			[
+				SNew(SButton)
+				.ButtonStyle(FEditorStyle::Get(), "HoverHintOnly")
+				.ForegroundColor(FSlateColor::UseForeground())
+				.OnClicked(this, &FNiagaraPlatformSetTypeCustomization::RemoveDeviceProfile, Profile, EffectsQuality)
+				.ToolTipText(LOCTEXT("RemoveDevice", "Remove this device override."))
+				[
+					SNew(SBox)
+					.WidthOverride(8)
+					.HeightOverride(8)
+					.VAlign(VAlign_Center)
+					[
+						SNew(SImage)
+						.Image(FNiagaraEditorStyle::Get().GetBrush("NiagaraEditor.PlatformSet.Remove"))
+					]
+				]
+			];
+		};
+
+	TArray<UDeviceProfile*> EnabledProfiles;
+	TArray<UDeviceProfile*> DisabledProfiles;
+	TargetPlatformSet->GetOverridenDeviceProfiles(EffectsQuality, EnabledProfiles, DisabledProfiles);
+
+	for (UDeviceProfile* Profile : EnabledProfiles)
+	{
+		AddDeviceProfileOverrideWidget(Profile, true);
+	}
+
+	for (UDeviceProfile* Profile : DisabledProfiles)
+	{
+		AddDeviceProfileOverrideWidget(Profile, false);
+	}
+
+	return Container;
+}
+
+FReply FNiagaraPlatformSetTypeCustomization::RemoveDeviceProfile(UDeviceProfile* Profile, int32 EffectsQuality)
+{
+	PropertyHandle->NotifyPreChange();
+	TargetPlatformSet->SetDeviceProfileState(Profile, EffectsQuality, ENiagaraPlatformSelectionState::Default);
+	PropertyHandle->NotifyPostChange(EPropertyChangeType::ValueSet);
+
+	InvalidateSiblingConflicts();
+
+	return FReply::Handled();
+}
+
+EVisibility FNiagaraPlatformSetTypeCustomization::GetDeviceProfileErrorVisibility(UDeviceProfile* Profile, int32 EffectsQuality) const
+{
+	// not part of an array
+	if (PlatformSetArrayIndex == INDEX_NONE)
+	{
+		return EVisibility::Collapsed;
+	}
+
+	for (const FNiagaraPlatformSetConflictInfo& ConflictInfo : CachedConflicts)
+	{
+		if (ConflictInfo.SetAIndex == PlatformSetArrayIndex || 
+			ConflictInfo.SetBIndex == PlatformSetArrayIndex)
+		{
+			const int32 EQMask = FNiagaraPlatformSet::CreateEQMask(EffectsQuality);
+
+			for (const FNiagaraPlatformSetConflictEntry& Entry : ConflictInfo.Conflicts)
+			{
+				if ((Entry.EffectsQualityMask & EQMask) != 0 &&
+					Entry.ProfileName == Profile->GetFName())
+				{
+					return EVisibility::Visible;
+				}
+			}
+		}
+	}
+
+	return EVisibility::Collapsed;
+}
+
+FText FNiagaraPlatformSetTypeCustomization::GetDeviceProfileErrorToolTip(UDeviceProfile* Profile, int32 EffectsQuality) const
+{
+	for (const FNiagaraPlatformSetConflictInfo& ConflictInfo : CachedConflicts)
+	{
+		int32 OtherIndex = INDEX_NONE;
+
+		if (ConflictInfo.SetAIndex == PlatformSetArrayIndex)
+		{
+			OtherIndex = ConflictInfo.SetBIndex;
+		}
+
+		if (ConflictInfo.SetBIndex == PlatformSetArrayIndex)
+		{
+			OtherIndex = ConflictInfo.SetAIndex;
+		}
+
+		if (OtherIndex != INDEX_NONE)
+		{
+			const int32 EQMask = FNiagaraPlatformSet::CreateEQMask(EffectsQuality);
+
+			for (const FNiagaraPlatformSetConflictEntry& Entry : ConflictInfo.Conflicts)
+			{
+				if ((Entry.EffectsQualityMask & EQMask) != 0 &&
+					Entry.ProfileName == Profile->GetFName())
+				{
+					FText FormatString = LOCTEXT("PlatformOverrideConflictToolTip", "This platform override conflicts with the set at index {Index} in this array."); 
+
+					FFormatNamedArguments Args;
+					Args.Add(TEXT("Index"), OtherIndex);
+
+					return FText::Format(FormatString, Args);
+				}
+			}
+		}
+	}
+
+	return FText::GetEmpty();
+}
+
+FReply FNiagaraPlatformSetTypeCustomization::ToggleMenuOpenForEffectsQuality(int32 EffectsQuality)
+{
+	check(EffectsQualityMenus.IsValidIndex(EffectsQuality));
+	
+	TSharedPtr<SMenuAnchor> MenuAnchor = EffectsQualityMenus[EffectsQuality];
+	MenuAnchor->SetIsOpen(!MenuAnchor->IsOpen());
+
+	return FReply::Handled();
+}
+
+// Is or does a viewmodel contain any children active at the given quality?
+bool FNiagaraPlatformSetTypeCustomization::IsTreeActiveForEQ(const TSharedPtr<FNiagaraDeviceProfileViewModel>& Tree, int32 EffectsQualityMask) const
+{
+	int32 Mask = TargetPlatformSet->GetEffectQualityMaskForDeviceProfile(Tree->Profile);
+	if ((Mask & EffectsQualityMask) != 0)
+	{
+		return true;
+	}
+
+	for (const TSharedPtr<FNiagaraDeviceProfileViewModel>& Child : Tree->Children)
+	{
+		if (IsTreeActiveForEQ(Child, EffectsQualityMask))
+		{
+			return true;
+		}
+	}
+
+	return false;
+}
+
+void FNiagaraPlatformSetTypeCustomization::FilterTreeForEQ(const TSharedPtr<FNiagaraDeviceProfileViewModel>& SourceTree, TSharedPtr<FNiagaraDeviceProfileViewModel>& FilteredTree, int32 EffectsQualityMask)
+{
+	for (const TSharedPtr<FNiagaraDeviceProfileViewModel>& SourceChild : SourceTree->Children)
+	{
+		if (IsTreeActiveForEQ(SourceChild, EffectsQualityMask))
+		{
+			TSharedPtr<FNiagaraDeviceProfileViewModel>& FilteredChild = FilteredTree->Children.Add_GetRef(MakeShared<FNiagaraDeviceProfileViewModel>());
+			FilteredChild->Profile = SourceChild->Profile;
+
+			FilterTreeForEQ(SourceChild, FilteredChild, EffectsQualityMask);
+		}
+	}
+}
+
+void FNiagaraPlatformSetTypeCustomization::CreateDeviceProfileTree()
+{
+	//Pull device profiles out by their hierarchy depth.
+	TArray<TArray<UDeviceProfile*>> DeviceProfilesByHierarchyLevel;
+	for (UObject* Profile : UDeviceProfileManager::Get().Profiles)
+	{
+		UDeviceProfile* DeviceProfile = CastChecked<UDeviceProfile>(Profile);
+
+		TFunction<void(int32&, UDeviceProfile*)> FindDepth;
+		FindDepth = [&](int32& Depth, UDeviceProfile* CurrProfile)
+		{
+			if (CurrProfile->Parent)
+			{
+				FindDepth(++Depth, Cast<UDeviceProfile>(CurrProfile->Parent));
+			}
+		};
+
+		int32 ProfileDepth = 0;
+		FindDepth(ProfileDepth, DeviceProfile);
+		DeviceProfilesByHierarchyLevel.SetNum(FMath::Max(ProfileDepth+1, DeviceProfilesByHierarchyLevel.Num()));
+		DeviceProfilesByHierarchyLevel[ProfileDepth].Add(DeviceProfile);
+	}
+	
+	FullDeviceProfileTree.Reset(DeviceProfilesByHierarchyLevel[0].Num());
+	for (int32 RootProfileIdx = 0; RootProfileIdx < DeviceProfilesByHierarchyLevel[0].Num(); ++RootProfileIdx)
+	{
+		TFunction<void(FNiagaraDeviceProfileViewModel*, int32)> BuildProfileTree;
+		BuildProfileTree = [&](FNiagaraDeviceProfileViewModel* CurrRoot, int32 CurrLevel)
+		{
+			int32 NextLevel = CurrLevel + 1;
+			if (NextLevel < DeviceProfilesByHierarchyLevel.Num())
+			{
+				for (UDeviceProfile* PossibleChild : DeviceProfilesByHierarchyLevel[NextLevel])
+				{
+					if (PossibleChild->Parent == CurrRoot->Profile)
+					{
+						TSharedPtr<FNiagaraDeviceProfileViewModel>& NewChild = CurrRoot->Children.Add_GetRef(MakeShared<FNiagaraDeviceProfileViewModel>());
+						NewChild->Profile = PossibleChild;
+						BuildProfileTree(NewChild.Get(), NextLevel);
+					}
+				}
+			}
+		};
+
+		//Add all root nodes and build their trees.
+		TSharedPtr<FNiagaraDeviceProfileViewModel> CurrRoot = MakeShared<FNiagaraDeviceProfileViewModel>();
+		CurrRoot->Profile = DeviceProfilesByHierarchyLevel[0][RootProfileIdx];
+		BuildProfileTree(CurrRoot.Get(), 0);
+		FullDeviceProfileTree.Add(CurrRoot);
+	}
+
+	Scalability::FQualityLevels Counts = Scalability::GetQualityLevelCounts();
+	FilteredDeviceProfileTrees.SetNum(Counts.EffectsQuality);
+	
+	for (TSharedPtr<FNiagaraDeviceProfileViewModel>& FullDeviceRoot : FullDeviceProfileTree)
+	{
+		for (int32 EffectsQuality = 0; EffectsQuality < Counts.EffectsQuality; ++EffectsQuality)
+		{
+			int32 EffectsQualityMask = FNiagaraPlatformSet::CreateEQMask(EffectsQuality);
+
+			if (IsTreeActiveForEQ(FullDeviceRoot, EffectsQualityMask))
+			{
+				TArray<TSharedPtr<FNiagaraDeviceProfileViewModel>>& FilteredRoots = FilteredDeviceProfileTrees[EffectsQuality];
+
+				TSharedPtr<FNiagaraDeviceProfileViewModel>& FilteredRoot = FilteredRoots.Add_GetRef(MakeShared<FNiagaraDeviceProfileViewModel>());
+				FilteredRoot->Profile = FullDeviceRoot->Profile;
+
+				FilterTreeForEQ(FullDeviceRoot, FilteredRoot, EffectsQualityMask);
+			}
+		}
+	}
+}
+
+TSharedRef<SWidget> FNiagaraPlatformSetTypeCustomization::GenerateDeviceProfileTreeWidget(int32 EffectsQuality)
+{
+	if (FullDeviceProfileTree.Num() == 0)
+	{
+		CreateDeviceProfileTree();	
+	}
+
+	TArray<TSharedPtr<FNiagaraDeviceProfileViewModel>>* TreeToUse = &FullDeviceProfileTree;
+	if (EffectsQuality != INDEX_NONE)
+	{
+		check(EffectsQuality < FilteredDeviceProfileTrees.Num());
+		TreeToUse = &FilteredDeviceProfileTrees[EffectsQuality];
+	}
+
+	return SNew(SBorder)
+		.BorderImage(FEditorStyle::Get().GetBrush("Menu.Background"))
+		[
+			SAssignNew(DeviceProfileTreeWidget, STreeView<TSharedPtr<FNiagaraDeviceProfileViewModel>>)
+			.TreeItemsSource(TreeToUse)
+			.OnGenerateRow(this, &FNiagaraPlatformSetTypeCustomization::OnGenerateDeviceProfileTreeRow, EffectsQuality)
+			.OnGetChildren(this, &FNiagaraPlatformSetTypeCustomization::OnGetDeviceProfileTreeChildren, EffectsQuality)
+			.SelectionMode(ESelectionMode::None)
+		];
+}
+
+TSharedRef<ITableRow> FNiagaraPlatformSetTypeCustomization::OnGenerateDeviceProfileTreeRow(TSharedPtr<FNiagaraDeviceProfileViewModel> InItem, const TSharedRef<STableViewBase>& OwnerTable, int32 EffectsQuality)
+{
+	TSharedPtr<SHorizontalBox> RowContainer;
+	SAssignNew(RowContainer, SHorizontalBox);
+
+	int32 ProfileMask = TargetPlatformSet->GetEffectQualityMaskForDeviceProfile(InItem->Profile);
+	FText NameTooltip = FText::Format(LOCTEXT("ProfileEQTooltipFmt", "Effects Quality: {0}"), FNiagaraPlatformSet::GetEffectsQualityMaskText(ProfileMask));
+	
+	//Top level profile. Look for a platform icon.
+	if (InItem->Profile->Parent == nullptr)
+	{
+		if (const PlatformInfo::FPlatformInfo* Info = PlatformInfo::FindPlatformInfo(*InItem->Profile->DeviceType))
+		{
+			const FSlateBrush* DeviceProfileTypeIcon = FEditorStyle::GetBrush(Info->GetIconStyleName(PlatformInfo::EPlatformIconSize::Normal));
+			if (DeviceProfileTypeIcon != FEditorStyle::Get().GetDefaultBrush())
+			{
+				RowContainer->AddSlot()
+					.AutoWidth()
+					.VAlign(VAlign_Center)
+					.HAlign(HAlign_Left)
+					.Padding(4, 0, 0, 0)
+					[
+						SNew(SBox)
+						.WidthOverride(16)
+						.HeightOverride(16)
+						[
+							SNew(SImage)
+							.Image(DeviceProfileTypeIcon)
+						]
+					];
+			}
+		}
+	}
+
+	FName TextStyleName("NormalText");
+	FSlateColor TextColor(FSlateColor::UseForeground());
+	ENiagaraPlatformSelectionState CurrentState = TargetPlatformSet->GetDeviceProfileState(InItem->Profile, EffectsQuality);
+	
+	if (CurrentState == ENiagaraPlatformSelectionState::Enabled)
+	{
+		TextStyleName = "RichTextBlock.Bold";
+		TextColor = FSlateColor(FLinearColor(FVector4(0,1,0,1)));
+	}
+	else if (CurrentState == ENiagaraPlatformSelectionState::Disabled)
+	{
+		TextStyleName = "RichTextBlock.Italic";
+		TextColor = FSlateColor(FLinearColor(FVector4(1,0,0,1)));
+	}
+
+
+	int32 EffectsQualityMask = FNiagaraPlatformSet::CreateEQMask(EffectsQuality);
+	if ((ProfileMask & EffectsQualityMask) == 0)
+	{
+		TextColor = FSlateColor::UseSubduedForeground();
+	}
+
+	RowContainer->AddSlot()
+		.Padding(4, 2, 0, 2)
+		.HAlign(HAlign_Fill)
+		.VAlign(VAlign_Center)
+		[
+			SNew(STextBlock)
+			.TextStyle(FEditorStyle::Get(), TextStyleName)
+			.Text(FText::FromString(InItem->Profile->GetName()))
+			.ToolTipText(NameTooltip)
+			.ColorAndOpacity(TextColor)
+		];
+
+	RowContainer->AddSlot()
+		.AutoWidth()
+		.Padding(12, 2, 4, 4)
+		.HAlign(HAlign_Right)
+		.VAlign(VAlign_Center)
+		[
+			SNew(SButton)
+			.ButtonStyle(FEditorStyle::Get(), "HoverHintOnly")
+			.HAlign(HAlign_Center)
+			.VAlign(VAlign_Fill)
+			.Visibility(this, &FNiagaraPlatformSetTypeCustomization::GetProfileMenuButtonVisibility, InItem, EffectsQuality)
+			.OnClicked(this, &FNiagaraPlatformSetTypeCustomization::OnProfileMenuButtonClicked, InItem, EffectsQuality)
+			.ToolTipText(this, &FNiagaraPlatformSetTypeCustomization::GetProfileMenuButtonToolTip, InItem, EffectsQuality)
+			[
+				SNew(SBox)
+				.WidthOverride(8)
+				.HeightOverride(8)
+				.VAlign(VAlign_Center)
+				.HAlign(HAlign_Center)
+				[
+					SNew(SImage)
+					.Image(this, &FNiagaraPlatformSetTypeCustomization::GetProfileMenuButtonImage, InItem, EffectsQuality)
+				]
+			]
+		];
+
+	return SNew(STableRow<TSharedPtr<FNiagaraDeviceProfileViewModel>>, OwnerTable)
+		.Style(FNiagaraEditorStyle::Get(), "NiagaraEditor.PlatformSet.TreeView")
+		[
+			RowContainer.ToSharedRef()
+		];
+}
+
+enum class EProfileButtonMode
+{
+	None,
+	Include,
+	Exclude,
+	Remove
+};
+
+static EProfileButtonMode GetProfileMenuButtonMode(FNiagaraPlatformSet* PlatformSet, TSharedPtr<FNiagaraDeviceProfileViewModel> Item, int32 EffectsQuality)
+{
+	int32 Mask = PlatformSet->GetEffectQualityMaskForDeviceProfile(Item->Profile);
+	int32 EQMask = FNiagaraPlatformSet::CreateEQMask(EffectsQuality);
+
+	if ((Mask & EQMask) == 0)
+	{
+		return EProfileButtonMode::None;
+	}
+
+	ENiagaraPlatformSelectionState CurrentState = PlatformSet->GetDeviceProfileState(Item->Profile, EffectsQuality);
+
+	bool bQualityEnabled = PlatformSet->IsEffectQualityEnabled(EffectsQuality);
+	bool bIsDefault = CurrentState == ENiagaraPlatformSelectionState::Default;
+
+	if (bIsDefault && bQualityEnabled)
+	{
+		return EProfileButtonMode::Exclude;
+	}
+
+	if (bIsDefault && !bQualityEnabled)
+	{
+		return EProfileButtonMode::Include;
+	}
+
+	return EProfileButtonMode::Remove;
+}
+
+
+FText FNiagaraPlatformSetTypeCustomization::GetProfileMenuButtonToolTip(TSharedPtr<FNiagaraDeviceProfileViewModel> Item, int32 EffectsQuality) const
+{
+	EProfileButtonMode Mode = GetProfileMenuButtonMode(TargetPlatformSet, Item, EffectsQuality);
+	switch(Mode)
+	{
+		case EProfileButtonMode::Include:
+			return LOCTEXT("IncludePlatform", "Include this platform.");
+		case EProfileButtonMode::Exclude:
+			return LOCTEXT("ExcludePlatform", "Exclude this platform.");
+		case EProfileButtonMode::Remove:
+			return LOCTEXT("RemovePlatform", "Remove this platform override.");
+	}
+
+	return FText::GetEmpty();
+}
+
+EVisibility FNiagaraPlatformSetTypeCustomization::GetProfileMenuButtonVisibility(TSharedPtr<FNiagaraDeviceProfileViewModel> Item, int32 EffectsQuality) const
+{
+	EProfileButtonMode Mode = GetProfileMenuButtonMode(TargetPlatformSet, Item, EffectsQuality);
+	if (Mode == EProfileButtonMode::None)
+	{
+		return EVisibility::Collapsed;
+	}
+
+	return EVisibility::Visible;
+}
+
+const FSlateBrush* FNiagaraPlatformSetTypeCustomization::GetProfileMenuButtonImage(TSharedPtr<FNiagaraDeviceProfileViewModel> Item, int32 EffectsQuality) const
+{
+	EProfileButtonMode Mode = GetProfileMenuButtonMode(TargetPlatformSet, Item, EffectsQuality);
+	switch(Mode)
+	{
+		case EProfileButtonMode::Include:
+			return FNiagaraEditorStyle::Get().GetBrush("NiagaraEditor.PlatformSet.Include");
+		case EProfileButtonMode::Exclude:
+			return FNiagaraEditorStyle::Get().GetBrush("NiagaraEditor.PlatformSet.Exclude");
+		case EProfileButtonMode::Remove:
+			return FNiagaraEditorStyle::Get().GetBrush("NiagaraEditor.PlatformSet.Remove");
+	}
+
+	return FEditorStyle::GetBrush("NoBrush");
+}
+
+FReply FNiagaraPlatformSetTypeCustomization::OnProfileMenuButtonClicked(TSharedPtr<FNiagaraDeviceProfileViewModel> Item, int32 EffectsQuality)
+{
+	ENiagaraPlatformSelectionState TargetState;
+
+	EProfileButtonMode Mode = GetProfileMenuButtonMode(TargetPlatformSet, Item, EffectsQuality);
+	switch(Mode)
+	{
+		case EProfileButtonMode::Include:
+			TargetState = ENiagaraPlatformSelectionState::Enabled;
+			break;
+		case EProfileButtonMode::Exclude:
+			TargetState = ENiagaraPlatformSelectionState::Disabled;
+			break;
+		case EProfileButtonMode::Remove:
+			TargetState = ENiagaraPlatformSelectionState::Default;
+			break;
+		default:
+			return FReply::Handled(); // shouldn't happen, button should be collapsed
+	}
+
+	PropertyHandle->NotifyPreChange();
+	TargetPlatformSet->SetDeviceProfileState(Item->Profile, EffectsQuality, TargetState);
+	PropertyHandle->NotifyPostChange(EPropertyChangeType::ValueSet);
+
+	InvalidateSiblingConflicts();
+
+	return FReply::Handled();
+}
+
+void FNiagaraPlatformSetTypeCustomization::OnGetDeviceProfileTreeChildren(TSharedPtr<FNiagaraDeviceProfileViewModel> InItem, TArray< TSharedPtr<FNiagaraDeviceProfileViewModel> >& OutChildren, int32 EffectsQuality)
+{
+	if (TargetPlatformSet->GetDeviceProfileState(InItem->Profile, EffectsQuality) == ENiagaraPlatformSelectionState::Default)
+	{
+		OutChildren = InItem->Children;
+	}
+}
+
+ECheckBoxState FNiagaraPlatformSetTypeCustomization::IsEQChecked(int32 EffectsQuality)const
+{
+	if (TargetPlatformSet)
+	{
+		return TargetPlatformSet->IsEffectQualityEnabled(EffectsQuality) ? ECheckBoxState::Checked : ECheckBoxState::Unchecked;
+	}
+	return ECheckBoxState::Undetermined;
+}
+
+void FNiagaraPlatformSetTypeCustomization::EQCheckStateChanged(ECheckBoxState CheckState, int32 EffectsQuality)
+{
+	PropertyHandle->NotifyPreChange();
+	TargetPlatformSet->SetEnabledForEffectQuality(EffectsQuality, CheckState == ECheckBoxState::Checked);
+	PropertyHandle->NotifyPostChange(EPropertyChangeType::ValueSet);
+
+	InvalidateSiblingConflicts();
+}
+
+void FNiagaraPlatformSetTypeCustomization::OnPropertyValueChanged()
+{
+	GenerateEffectsQualitySelectionWidgets();
+	UpdateCachedConflicts();
+}
 
 #undef LOCTEXT_NAMESPACE
diff --git a/Engine/Plugins/FX/Niagara/Source/NiagaraEditor/Private/Customizations/NiagaraTypeCustomizations.h b/Engine/Plugins/FX/Niagara/Source/NiagaraEditor/Private/Customizations/NiagaraTypeCustomizations.h
--- a/Engine/Plugins/FX/Niagara/Source/NiagaraEditor/Private/Customizations/NiagaraTypeCustomizations.h
+++ b/Engine/Plugins/FX/Niagara/Source/NiagaraEditor/Private/Customizations/NiagaraTypeCustomizations.h
@@ -5,14 +5,20 @@
 #include "CoreMinimal.h"
 #include "IDetailCustomization.h"
 #include "IPropertyTypeCustomization.h"
-#include "Styling/SlateTypes.h"
+#include "NiagaraPlatformSet.h"
 #include "EdGraph/EdGraphSchema.h"
+#include "Styling/SlateTypes.h"
+#include "Widgets/Views/STreeView.h"
+
 #include "NiagaraTypeCustomizations.generated.h"
 
+class FDetailWidgetRow;
 class FNiagaraScriptViewModel;
 class IPropertyHandle;
-class FDetailWidgetRow;
+class IPropertyHandleArray;
 class IPropertyTypeCustomization;
+class SMenuAnchor;
+class SWrapBox;
 
 class FNiagaraNumericCustomization : public IPropertyTypeCustomization
 {
@@ -189,3 +195,89 @@ private:
 	class UNiagaraScriptVariable* BaseScriptVariable;
 	struct FNiagaraScriptVariableBinding* TargetVariableBinding;
 };
+
+struct FNiagaraDeviceProfileViewModel
+{
+	class UDeviceProfile* Profile;
+	TArray<TSharedPtr<FNiagaraDeviceProfileViewModel>> Children;
+};
+
+class FNiagaraPlatformSetTypeCustomization : public IPropertyTypeCustomization
+{
+public:
+	/** @return A new instance of this class */
+	static TSharedRef<IPropertyTypeCustomization> MakeInstance()
+	{
+		return MakeShared<FNiagaraPlatformSetTypeCustomization>();
+	}
+
+	/** IPropertyTypeCustomization interface begin */
+	virtual void CustomizeHeader(TSharedRef<IPropertyHandle> StructPropertyHandle, class FDetailWidgetRow& HeaderRow, IPropertyTypeCustomizationUtils& StructCustomizationUtils) override;
+	virtual void CustomizeChildren(TSharedRef<IPropertyHandle> StructPropertyHandle, class IDetailChildrenBuilder& ChildBuilder, IPropertyTypeCustomizationUtils& StructCustomizationUtils) override;
+	/** IPropertyTypeCustomization interface end */
+private:
+
+	void CreateDeviceProfileTree();
+
+	TSharedRef<SWidget> GenerateDeviceProfileTreeWidget(int32 EffectsQuality);
+	TSharedRef<ITableRow> OnGenerateDeviceProfileTreeRow(TSharedPtr<FNiagaraDeviceProfileViewModel> InItem, const TSharedRef<STableViewBase>& OwnerTable, int32 EffectsQuality);
+	void OnGetDeviceProfileTreeChildren(TSharedPtr<FNiagaraDeviceProfileViewModel> InItem, TArray< TSharedPtr<FNiagaraDeviceProfileViewModel> >& OutChildren, int32 EffectsQuality);
+	FReply RemoveDeviceProfile(UDeviceProfile* Profile, int32 EffectsQuality);
+
+	EVisibility GetDeviceProfileErrorVisibility(UDeviceProfile* Profile, int32 EffectsQuality) const;
+	FText GetDeviceProfileErrorToolTip(UDeviceProfile* Profile, int32 EffectsQuality) const;
+
+	FSlateColor GetEffectsQualityButtonTextColor(int32 EffectsQuality) const;
+	EVisibility GetEffectsQualityErrorVisibility(int32 EffectsQuality) const;
+	FText GetEffectsQualityErrorToolTip(int32 EffectsQuality) const;
+
+	const FSlateBrush* GetProfileMenuButtonImage(TSharedPtr<FNiagaraDeviceProfileViewModel> Item, int32 EffectsQuality) const;
+	EVisibility GetProfileMenuButtonVisibility(TSharedPtr<FNiagaraDeviceProfileViewModel> Item, int32 EffectsQuality) const;
+	FText GetProfileMenuButtonToolTip(TSharedPtr<FNiagaraDeviceProfileViewModel> Item, int32 EffectsQuality) const;
+	FReply OnProfileMenuButtonClicked(TSharedPtr<FNiagaraDeviceProfileViewModel> InItem, int32 EffectsQuality);
+
+	TSharedRef<SWidget> GetCurrentDeviceProfileSelectionWidget(TSharedPtr<FNiagaraDeviceProfileViewModel> ProfileView);
+	TSharedRef<SWidget> OnGenerateDeviceProfileSelectionWidget(TSharedPtr<ENiagaraPlatformSelectionState> InItem);
+
+	FText GetCurrentText() const;
+	FText GetTooltipText() const;
+
+	void GenerateEffectsQualitySelectionWidgets();
+	TSharedRef<SWidget> GenerateAdditionalDevicesWidgetForEQ(int32 EffectsQuality);
+
+	bool IsTreeActiveForEQ(const TSharedPtr<FNiagaraDeviceProfileViewModel>& Tree, int32 EffectsQualityMask) const;
+	void FilterTreeForEQ(const TSharedPtr<FNiagaraDeviceProfileViewModel>& SourceTree, TSharedPtr<FNiagaraDeviceProfileViewModel>& FilteredTree, int32 EffectsQualityMask);
+
+	ECheckBoxState IsEQChecked(int32 EffectsQuality) const;
+	void EQCheckStateChanged(ECheckBoxState CheckState, int32 EffectsQuality);
+
+	FReply ToggleMenuOpenForEffectsQuality(int32 EffectsQuality);
+
+	void UpdateCachedConflicts();
+	void InvalidateSiblingConflicts() const;
+	void OnPropertyValueChanged();
+
+private:
+	TArray<FNiagaraPlatformSetConflictInfo> CachedConflicts;
+	TSharedPtr<IPropertyHandleArray> PlatformSetArray;
+	int32 PlatformSetArrayIndex;
+
+	TArray<TSharedPtr<SMenuAnchor>> EffectsQualityMenus;
+	TSharedPtr<SWrapBox> EffectsQualityWidgetBox;
+
+	TSharedPtr<IPropertyHandle> PropertyHandle;
+	struct FNiagaraPlatformSet* TargetPlatformSet;
+	class UNiagaraSystem* BaseSystem;
+
+	TArray<TSharedPtr<FNiagaraDeviceProfileViewModel>> FullDeviceProfileTree;
+	TArray<TArray<TSharedPtr<FNiagaraDeviceProfileViewModel>>> FilteredDeviceProfileTrees;
+
+	TArray<TSharedPtr<ENiagaraPlatformSelectionState>> PlatformSelectionStates;
+
+	TSharedPtr<STreeView<TSharedPtr<FNiagaraDeviceProfileViewModel>>> DeviceProfileTreeWidget;
+
+	struct FNiagaraSystemScalabilitySettingsArray* SystemScalabilitySettings;
+	struct FNiagaraEmitterScalabilitySettingsArray* EmitterScalabilitySettings;
+	struct FNiagaraSystemScalabilityOverrides* SystemScalabilityOverrides;
+	struct FNiagaraEmitterScalabilityOverrides* EmitterScalabilityOverrides;
+};
diff --git a/Engine/Plugins/FX/Niagara/Source/NiagaraEditor/Private/NiagaraCompiler.cpp b/Engine/Plugins/FX/Niagara/Source/NiagaraEditor/Private/NiagaraCompiler.cpp
--- a/Engine/Plugins/FX/Niagara/Source/NiagaraEditor/Private/NiagaraCompiler.cpp
+++ b/Engine/Plugins/FX/Niagara/Source/NiagaraEditor/Private/NiagaraCompiler.cpp
@@ -542,7 +542,6 @@ TSharedPtr<FNiagaraCompileRequestDataBase, ESPMode::ThreadSafe> FNiagaraEditorMo
 		// Create an array of variables that we *did* encounter when traversing the graphs.
 		TArray<FNiagaraVariable> EncounteredExposedVars;
 		check(System->GetSystemSpawnScript()->GetSource() == System->GetSystemUpdateScript()->GetSource());
-		uint32 DetailLevelMask = 0xFFFFFFFF;
 
 		// First deep copy all the emitter graphs referenced by the system so that we can later hook up emitter handles in the system traversal.
 		BasePtr->EmitterData.Empty();
diff --git a/Engine/Plugins/FX/Niagara/Source/NiagaraEditor/Private/NiagaraEditorModule.cpp b/Engine/Plugins/FX/Niagara/Source/NiagaraEditor/Private/NiagaraEditorModule.cpp
--- a/Engine/Plugins/FX/Niagara/Source/NiagaraEditor/Private/NiagaraEditorModule.cpp
+++ b/Engine/Plugins/FX/Niagara/Source/NiagaraEditor/Private/NiagaraEditorModule.cpp
@@ -118,6 +118,9 @@
 #include "SourceControlOperations.h"
 #include "ISourceControlProvider.h"
 #include "ISourceControlModule.h"
+#include "DeviceProfiles/DeviceProfileManager.h"
+#include "Interfaces/ITargetPlatform.h"
+#include "DeviceProfiles/DeviceProfile.h"
 
 IMPLEMENT_MODULE( FNiagaraEditorModule, NiagaraEditor );
 
@@ -728,6 +731,8 @@ void FNiagaraEditorModule::StartupModule()
 	
 	// Any attempt to use GEditor right now will fail as it hasn't been initialized yet. Waiting for post engine init resolves that.
 	FCoreDelegates::OnPostEngineInit.AddRaw(this, &FNiagaraEditorModule::OnPostEngineInit);
+
+	DeviceProfileManagerUpdatedHandle = UDeviceProfileManager::Get().OnManagerUpdated().AddRaw(this, &FNiagaraEditorModule::OnDeviceProfileManagerUpdated);
 	
 	// register details customization
 	FPropertyEditorModule& PropertyModule = FModuleManager::LoadModuleChecked<FPropertyEditorModule>("PropertyEditor");
@@ -771,7 +776,15 @@ void FNiagaraEditorModule::StartupModule()
 	PropertyModule.RegisterCustomPropertyTypeLayout("NiagaraScriptVariableBinding",
 		FOnGetPropertyTypeCustomizationInstance::CreateStatic(&FNiagaraScriptVariableBindingCustomization::MakeInstance)
 	);
-		
+
+	PropertyModule.RegisterCustomPropertyTypeLayout("NiagaraPlatformSet",
+		FOnGetPropertyTypeCustomizationInstance::CreateStatic(&FNiagaraPlatformSetTypeCustomization::MakeInstance)
+	);
+
+// 	PropertyModule.RegisterCustomPropertyTypeLayout("NiagaraSystemScalabilitySettingsArray",
+// 		FOnGetPropertyTypeCustomizationInstance::CreateStatic(&FNiagaraSystemScalabilitySettingsArrayTypeCustomization::MakeInstance)
+// 	);
+	
 	PropertyModule.RegisterCustomPropertyTypeLayout("NiagaraUserParameterBinding",
 		FOnGetPropertyTypeCustomizationInstance::CreateStatic(&FNiagaraUserParameterBindingCustomization::MakeInstance)
 	);
@@ -986,6 +999,7 @@ void FNiagaraEditorModule::ShutdownModule()
 
 	FCoreUObjectDelegates::GetPreGarbageCollectDelegate().RemoveAll(this);
 	FCoreDelegates::OnPostEngineInit.RemoveAll(this);
+	UDeviceProfileManager::Get().OnManagerUpdated().Remove(DeviceProfileManagerUpdatedHandle);
 	
 	if (GEditor)
 	{
@@ -1082,6 +1096,11 @@ void FNiagaraEditorModule::OnPostEngineInit()
 	}
 }
 
+void FNiagaraEditorModule::OnDeviceProfileManagerUpdated()
+{
+	FNiagaraPlatformSet::InvalidateCachedData();
+}
+
 FNiagaraEditorModule& FNiagaraEditorModule::Get()
 {
 	return FModuleManager::LoadModuleChecked<FNiagaraEditorModule>("NiagaraEditor");
diff --git a/Engine/Plugins/FX/Niagara/Source/NiagaraEditor/Private/NiagaraEditorStyle.cpp b/Engine/Plugins/FX/Niagara/Source/NiagaraEditor/Private/NiagaraEditorStyle.cpp
--- a/Engine/Plugins/FX/Niagara/Source/NiagaraEditor/Private/NiagaraEditorStyle.cpp
+++ b/Engine/Plugins/FX/Niagara/Source/NiagaraEditor/Private/NiagaraEditorStyle.cpp
@@ -40,13 +40,16 @@ NIAGARAEDITOR_API FString RelativePathToPluginPath(const FString& RelativePath,
 	return (ContentDir / RelativePath) + Extension;
 }
 
-#define IMAGE_PLUGIN_BRUSH( RelativePath, ... ) FSlateImageBrush( RelativePathToPluginPath( RelativePath, ".png" ), __VA_ARGS__ )
 #define IMAGE_BRUSH( RelativePath, ... ) FSlateImageBrush( Style->RootToContentDir( RelativePath, TEXT(".png") ), __VA_ARGS__ )
+#define IMAGE_CORE_BRUSH( RelativePath, ... ) FSlateImageBrush( FPaths::EngineContentDir() / "Editor/Slate" / RelativePath + TEXT(".png"), __VA_ARGS__ )
+#define IMAGE_PLUGIN_BRUSH( RelativePath, ... ) FSlateImageBrush( RelativePathToPluginPath( RelativePath, ".png" ), __VA_ARGS__ )
 #define BOX_BRUSH( RelativePath, ... ) FSlateBoxBrush( Style->RootToContentDir( RelativePath, TEXT(".png") ), __VA_ARGS__ )
-#define BORDER_BRUSH( RelativePath, ... ) FSlateBorderBrush( Style->RootToContentDir( RelativePath, TEXT(".png") ), __VA_ARGS__ )
 #define BOX_CORE_BRUSH( RelativePath, ... ) FSlateBoxBrush( FPaths::EngineContentDir() / "Editor/Slate" / RelativePath + TEXT(".png"), __VA_ARGS__ )
+#define BORDER_BRUSH( RelativePath, ... ) FSlateBorderBrush( Style->RootToContentDir( RelativePath, TEXT(".png") ), __VA_ARGS__ )
 #define DEFAULT_FONT(...) FCoreStyle::GetDefaultFontStyle(__VA_ARGS__)
 
+const FVector2D Icon8x8(8.0f, 8.0f);
+const FVector2D Icon12x12(12.0f, 12.0f);
 const FVector2D Icon16x16(16.0f, 16.0f);
 const FVector2D Icon20x20(20.0f, 20.0f);
 const FVector2D Icon40x40(40.0f, 40.0f);
@@ -260,9 +263,74 @@ TSharedRef< FSlateStyleSet > FNiagaraEditorStyle::Create()
 	// Niagara sequence
 	Style->Set("NiagaraEditor.NiagaraSequence.DefaultTrackColor", FLinearColor(0, .25f, 0));
 
+	// Niagara platform set customization
+	Style->Set("NiagaraEditor.PlatformSet.DropdownButton", new IMAGE_CORE_BRUSH("Common/ComboArrow", Icon8x8));
+
+	Style->Set("NiagaraEditor.PlatformSet.ButtonText", FTextBlockStyle(NormalText)
+		.SetFont(DEFAULT_FONT("Bold", 10))
+		.SetColorAndOpacity(FLinearColor(0.72f, 0.72f, 0.72f))
+		.SetHighlightColor(FLinearColor(1, 1, 1)));
+
+	const FString SmallRoundedButtonStart(TEXT("Common/SmallRoundedButtonLeft"));
+	const FString SmallRoundedButtonMiddle(TEXT("Common/SmallRoundedButtonCentre"));
+	const FString SmallRoundedButtonEnd(TEXT("Common/SmallRoundedButtonRight"));
+
+	const FSlateColor SelectionColor = FEditorStyle::GetSlateColor("SelectionColor");
+	const FSlateColor SelectionColor_Pressed = FEditorStyle::GetSlateColor("SelectionColor_Pressed");
+
+	{
+		const FLinearColor NormalColor(0.15, 0.15, 0.15, 1);
+
+		Style->Set("NiagaraEditor.PlatformSet.StartButton", FCheckBoxStyle()
+			.SetCheckBoxType(ESlateCheckBoxType::ToggleButton)
+			.SetUncheckedImage(BOX_CORE_BRUSH(*SmallRoundedButtonStart, FMargin(7.f / 16.f), NormalColor))
+			.SetUncheckedPressedImage(BOX_CORE_BRUSH(*SmallRoundedButtonStart, FMargin(7.f / 16.f), SelectionColor_Pressed))
+			.SetUncheckedHoveredImage(BOX_CORE_BRUSH(*SmallRoundedButtonStart, FMargin(7.f / 16.f), SelectionColor_Pressed))
+			.SetCheckedHoveredImage(BOX_CORE_BRUSH(*SmallRoundedButtonStart, FMargin(7.f / 16.f), SelectionColor))
+			.SetCheckedPressedImage(BOX_CORE_BRUSH(*SmallRoundedButtonStart, FMargin(7.f / 16.f), SelectionColor))
+			.SetCheckedImage(BOX_CORE_BRUSH(*SmallRoundedButtonStart, FMargin(7.f / 16.f), SelectionColor)));
+
+		Style->Set("NiagaraEditor.PlatformSet.MiddleButton", FCheckBoxStyle()
+			.SetCheckBoxType(ESlateCheckBoxType::ToggleButton)
+			.SetUncheckedImage(BOX_CORE_BRUSH(*SmallRoundedButtonMiddle, FMargin(7.f / 16.f), NormalColor))
+			.SetUncheckedPressedImage(BOX_CORE_BRUSH(*SmallRoundedButtonMiddle, FMargin(7.f / 16.f), SelectionColor_Pressed))
+			.SetUncheckedHoveredImage(BOX_CORE_BRUSH(*SmallRoundedButtonMiddle, FMargin(7.f / 16.f), SelectionColor_Pressed))
+			.SetCheckedHoveredImage(BOX_CORE_BRUSH(*SmallRoundedButtonMiddle, FMargin(7.f / 16.f), SelectionColor))
+			.SetCheckedPressedImage(BOX_CORE_BRUSH(*SmallRoundedButtonMiddle, FMargin(7.f / 16.f), SelectionColor))
+			.SetCheckedImage(BOX_CORE_BRUSH(*SmallRoundedButtonMiddle, FMargin(7.f / 16.f), SelectionColor)));
+
+		Style->Set("NiagaraEditor.PlatformSet.EndButton", FCheckBoxStyle()
+			.SetCheckBoxType(ESlateCheckBoxType::ToggleButton)
+			.SetUncheckedImage(BOX_CORE_BRUSH(*SmallRoundedButtonEnd, FMargin(7.f / 16.f), NormalColor))
+			.SetUncheckedPressedImage(BOX_CORE_BRUSH(*SmallRoundedButtonEnd, FMargin(7.f / 16.f), SelectionColor_Pressed))
+			.SetUncheckedHoveredImage(BOX_CORE_BRUSH(*SmallRoundedButtonEnd, FMargin(7.f / 16.f), SelectionColor_Pressed))
+			.SetCheckedHoveredImage(BOX_CORE_BRUSH(*SmallRoundedButtonEnd, FMargin(7.f / 16.f), SelectionColor))
+			.SetCheckedPressedImage(BOX_CORE_BRUSH(*SmallRoundedButtonEnd, FMargin(7.f / 16.f), SelectionColor))
+			.SetCheckedImage(BOX_CORE_BRUSH(*SmallRoundedButtonEnd, FMargin(7.f / 16.f), SelectionColor)));
+	}
+
+	Style->Set("NiagaraEditor.PlatformSet.Include", new IMAGE_CORE_BRUSH("Icons/PlusSymbol_12x", Icon12x12));
+	Style->Set("NiagaraEditor.PlatformSet.Exclude", new IMAGE_CORE_BRUSH("Icons/MinusSymbol_12x", Icon12x12));
+	Style->Set("NiagaraEditor.PlatformSet.Remove", new IMAGE_CORE_BRUSH("Icons/Cross_12x", Icon12x12));
+
+	const FSlateColor SelectionColor_Inactive = FEditorStyle::GetSlateColor("SelectionColor_Inactive");
+
+	Style->Set("NiagaraEditor.PlatformSet.TreeView", FTableRowStyle()
+		.SetEvenRowBackgroundBrush(FSlateNoResource())
+		.SetEvenRowBackgroundHoveredBrush(IMAGE_CORE_BRUSH("Common/Selection", Icon8x8, SelectionColor_Inactive))
+		.SetOddRowBackgroundBrush(FSlateNoResource())
+		.SetOddRowBackgroundHoveredBrush(IMAGE_CORE_BRUSH("Common/Selection", Icon8x8, SelectionColor_Inactive))
+		.SetSelectorFocusedBrush(FSlateNoResource())
+		.SetActiveBrush(IMAGE_CORE_BRUSH("Common/Selection", Icon8x8, SelectionColor))
+		.SetActiveHoveredBrush(IMAGE_CORE_BRUSH("Common/Selection", Icon8x8, SelectionColor))
+		.SetInactiveBrush(IMAGE_CORE_BRUSH("Common/Selection", Icon8x8, SelectionColor_Inactive))
+		.SetInactiveHoveredBrush(IMAGE_CORE_BRUSH("Common/Selection", Icon8x8, SelectionColor_Inactive)));
+
 	return Style;
 }
 
+#undef IMAGE_PLUGIN_BRUSH
+#undef IMAGE_CORE_BRUSH
 #undef IMAGE_BRUSH
 #undef BOX_BRUSH
 #undef BORDER_BRUSH
diff --git a/Engine/Plugins/FX/Niagara/Source/NiagaraEditor/Private/NiagaraHlslTranslator.h b/Engine/Plugins/FX/Niagara/Source/NiagaraEditor/Private/NiagaraHlslTranslator.h
--- a/Engine/Plugins/FX/Niagara/Source/NiagaraEditor/Private/NiagaraHlslTranslator.h
+++ b/Engine/Plugins/FX/Niagara/Source/NiagaraEditor/Private/NiagaraHlslTranslator.h
@@ -82,7 +82,7 @@ struct FNiagaraTranslateResults
 class FNiagaraCompileRequestData : public FNiagaraCompileRequestDataBase
 {
 public:
-	FNiagaraCompileRequestData() : bUseRapidIterationParams(true), DetailLevelMask(FNiagaraCompileRequestDataBase::CookForAllDetailLevelMask)
+	FNiagaraCompileRequestData() : bUseRapidIterationParams(true)
 	{
 
 	}
@@ -109,9 +109,6 @@ public:
 		return EmitterData[Index];
 	}
 	void AddRapidIterationParameters(const FNiagaraParameterStore& InParamStore, FCompileConstantResolver InResolver);
-	virtual uint32 GetDetailLevelMask(void) const override {
-		return DetailLevelMask;
-	};
 	virtual bool GetUseRapidIterationParams() const override { return bUseRapidIterationParams; }
 
 	// If this is being held onto for any length of time, make sure to hold onto it in a gc-aware object. Right now in this information-passing struct,
@@ -126,7 +123,6 @@ public:
 	UNiagaraScriptSource* Source;
 	FString SourceName;
 	bool bUseRapidIterationParams = true;
-	uint32 DetailLevelMask = 0xFFFFFFFF;
 
 	UEnum* ENiagaraScriptCompileStatusEnum;
 	UEnum* ENiagaraScriptUsageEnum;
diff --git a/Engine/Plugins/FX/Niagara/Source/NiagaraEditor/Private/NiagaraParameterMapHistory.cpp b/Engine/Plugins/FX/Niagara/Source/NiagaraEditor/Private/NiagaraParameterMapHistory.cpp
--- a/Engine/Plugins/FX/Niagara/Source/NiagaraEditor/Private/NiagaraParameterMapHistory.cpp
+++ b/Engine/Plugins/FX/Niagara/Source/NiagaraEditor/Private/NiagaraParameterMapHistory.cpp
@@ -1430,11 +1430,6 @@ bool FCompileConstantResolver::ResolveConstant(FNiagaraVariable& OutConstant) co
 		OutConstant.SetValue(Emitter->bDeterminism ? FNiagaraBool(true) : FNiagaraBool(false));
 		return true;
 	}
-	if (Emitter && OutConstant == FNiagaraVariable(FNiagaraTypeDefinition::GetBoolDef(), TEXT("Emitter.OverrideGlobalSpawnCountScale")))
-	{
-		OutConstant.SetValue(Emitter->bOverrideGlobalSpawnCountScale ? FNiagaraBool(true) : FNiagaraBool(false));
-		return true;
-	}
 	if (Emitter && OutConstant == FNiagaraVariable(FNiagaraTypeDefinition::GetSimulationTargetEnum(), TEXT("Emitter.SimulationTarget")))
 	{
 		FNiagaraInt32 EnumValue;
diff --git a/Engine/Plugins/FX/Niagara/Source/NiagaraEditor/Private/ViewModels/NiagaraEmitterViewModel.cpp b/Engine/Plugins/FX/Niagara/Source/NiagaraEditor/Private/ViewModels/NiagaraEmitterViewModel.cpp
--- a/Engine/Plugins/FX/Niagara/Source/NiagaraEditor/Private/ViewModels/NiagaraEmitterViewModel.cpp
+++ b/Engine/Plugins/FX/Niagara/Source/NiagaraEditor/Private/ViewModels/NiagaraEmitterViewModel.cpp
@@ -24,7 +24,7 @@ template<> TMap<UNiagaraEmitter*, TArray<FNiagaraEmitterViewModel*>> TNiagaraVie
 
 const FText FNiagaraEmitterViewModel::StatsFormat = NSLOCTEXT("NiagaraEmitterViewModel", "StatsFormat", "{0} Particles | {1} ms | {2} MB | {3}");
 const FText FNiagaraEmitterViewModel::StatsParticleCountFormat = NSLOCTEXT("NiagaraEmitterViewModel", "StatsParticleCountFormat", "{0} Particles");
-const FText FNiagaraEmitterViewModel::ParticleDisabledDueToDetailLevel = NSLOCTEXT("NiagaraEmitterViewModel", "ParticleDisabledDueToDetailLevel", "Disabled due to detail level mismatch (Current: {0}). See Scalability options for valid range. ");
+const FText FNiagaraEmitterViewModel::ParticleDisabledDueToScalability = NSLOCTEXT("NiagaraEmitterViewModel", "ParticleDisabledDueToScalability", "Disabled due to scalability settings (Current Effects Quality: {0}). See Scalability options for valid range. ");
 
 namespace NiagaraCommands
 {
@@ -201,11 +201,11 @@ FText FNiagaraEmitterViewModel::GetStatsText() const
 					return LOCTEXT("InvalidInstance", "Invalid Emitter! May have compile errors.");
 				}
 
-				int32 DetailLevel = SimInstance->GetParentSystemInstance()->GetDetailLevel();
-				if (!HandleEmitter->IsAllowedByDetailLevel(DetailLevel))
-				{
-					return FText::Format(ParticleDisabledDueToDetailLevel, FText::AsNumber(DetailLevel));
-				}
+ 				int32 EffectsQuality = INiagaraModule::GetEffectsQuality();
+  				if (!HandleEmitter->IsAllowedByScalability())
+  				{
+					return FText::Format(ParticleDisabledDueToScalability, FNiagaraPlatformSet::GetEffectsQualityText(EffectsQuality));
+  				}
 
 				if (NiagaraCommands::EmitterStatsFormat->GetInt() == 1)
 				{
@@ -235,25 +235,6 @@ TSharedRef<FNiagaraScriptViewModel> FNiagaraEmitterViewModel::GetSharedScriptVie
 	return SharedScriptViewModel.ToSharedRef();
 }
 
-bool FNiagaraEmitterViewModel::IsEnabledByDetailLevel() const
-{
-	if (Simulation.IsValid())
-	{
-		TSharedPtr<FNiagaraEmitterInstance, ESPMode::ThreadSafe> SimInstance = Simulation.Pin();
-		if (SimInstance.IsValid())
-		{
-			const FNiagaraEmitterHandle& Handle = SimInstance->GetEmitterHandle();
-			if (Handle.GetInstance())
-			{
-				UNiagaraEmitter* HandleEmitter = Handle.GetInstance();
-				int32 DetailLevel = SimInstance->GetParentSystemInstance()->GetDetailLevel();
-				return HandleEmitter->IsAllowedByDetailLevel(DetailLevel);
-			}
-		}
-	}
-	return false;
-}
-
 const UNiagaraEmitterEditorData& FNiagaraEmitterViewModel::GetEditorData() const
 {
 	check(Emitter.IsValid());
diff --git a/Engine/Plugins/FX/Niagara/Source/NiagaraEditor/Public/NiagaraEditorModule.h b/Engine/Plugins/FX/Niagara/Source/NiagaraEditor/Public/NiagaraEditorModule.h
--- a/Engine/Plugins/FX/Niagara/Source/NiagaraEditor/Public/NiagaraEditorModule.h
+++ b/Engine/Plugins/FX/Niagara/Source/NiagaraEditor/Public/NiagaraEditorModule.h
@@ -124,6 +124,7 @@ private:
 	void OnPreGarbageCollection();
 	void OnExecParticleInvoked(const TCHAR* InStr);
 	void OnPostEngineInit();
+	void OnDeviceProfileManagerUpdated();
 
 	/** FGCObject interface */
 	virtual void AddReferencedObjects( FReferenceCollector& Collector ) override;
@@ -161,6 +162,8 @@ private:
 	FDelegateHandle CompileResultHandle;
 	FDelegateHandle PrecompilerHandle;
 
+	FDelegateHandle DeviceProfileManagerUpdatedHandle;
+
 	USequencerSettings* SequencerSettings;
 	
 	TSharedPtr<INiagaraEditorWidgetProvider> WidgetProvider;
diff --git a/Engine/Plugins/FX/Niagara/Source/NiagaraEditor/Public/ViewModels/NiagaraEmitterViewModel.h b/Engine/Plugins/FX/Niagara/Source/NiagaraEditor/Public/ViewModels/NiagaraEmitterViewModel.h
--- a/Engine/Plugins/FX/Niagara/Source/NiagaraEditor/Public/ViewModels/NiagaraEmitterViewModel.h
+++ b/Engine/Plugins/FX/Niagara/Source/NiagaraEditor/Public/ViewModels/NiagaraEmitterViewModel.h
@@ -99,8 +99,6 @@ public:
 
 	void Cleanup();
 
-	bool IsEnabledByDetailLevel() const;
-
 private:
 	/** Sets this view model to a different emitter. */
 	void SetEmitter(UNiagaraEmitter* InEmitter);
@@ -124,8 +122,8 @@ private:
 	/** The text format stats to only display particles count. */
 	static const FText StatsParticleCountFormat;
 
-	/** The text format stats to display when detail level mismatches. */
-	static const FText ParticleDisabledDueToDetailLevel;
+	/** The text format stats to display when an emitter is disabled due to scalability. */
+	static const FText ParticleDisabledDueToScalability;
 
 	/** The emitter object being displayed by the control .*/
 	TWeakObjectPtr<UNiagaraEmitter> Emitter;
diff --git a/Engine/Plugins/FX/Niagara/Source/NiagaraShader/Private/NiagaraShader.cpp b/Engine/Plugins/FX/Niagara/Source/NiagaraShader/Private/NiagaraShader.cpp
--- a/Engine/Plugins/FX/Niagara/Source/NiagaraShader/Private/NiagaraShader.cpp
+++ b/Engine/Plugins/FX/Niagara/Source/NiagaraShader/Private/NiagaraShader.cpp
@@ -189,7 +189,11 @@ void FNiagaraShaderMapId::Serialize(FArchive& Ar)
 	
 	if (NiagaraVer >= FNiagaraCustomVersion::AddRIAndDetailLevel)
 	{
-		Ar << DetailLevelMask;
+		if (NiagaraVer < FNiagaraCustomVersion::PlatformScalingRefactor)
+		{
+			int32 Dummy;
+			Ar << Dummy;
+		}
 		Ar << bUsesRapidIterationParams;
 	}
 	
@@ -235,7 +239,6 @@ bool FNiagaraShaderMapId::operator==(const FNiagaraShaderMapId& ReferenceSet) co
 		|| */BaseCompileHash != ReferenceSet.BaseCompileHash
 		|| FeatureLevel != ReferenceSet.FeatureLevel
 		|| CompilerVersionID != ReferenceSet.CompilerVersionID 
-		|| DetailLevelMask != ReferenceSet.DetailLevelMask 
 		|| bUsesRapidIterationParams != ReferenceSet.bUsesRapidIterationParams)
 	{
 		return false;
@@ -305,15 +308,6 @@ void FNiagaraShaderMapId::AppendKeyString(FString& KeyString) const
 	KeyString += CompilerVersionID.ToString();
 	KeyString += TEXT("_");
 
-	if (DetailLevelMask != 0xFFFFFFFF)
-	{
-		KeyString += FString::Printf(TEXT("DL_%d"), DetailLevelMask);
-	}
-	else
-	{
-		KeyString += TEXT("ALLDL_");
-	}
-
 	if (bUsesRapidIterationParams)
 	{
 		KeyString += TEXT("USESRI_");
diff --git a/Engine/Plugins/FX/Niagara/Source/NiagaraShader/Private/NiagaraShared.cpp b/Engine/Plugins/FX/Niagara/Source/NiagaraShader/Private/NiagaraShared.cpp
--- a/Engine/Plugins/FX/Niagara/Source/NiagaraShader/Private/NiagaraShared.cpp
+++ b/Engine/Plugins/FX/Niagara/Source/NiagaraShader/Private/NiagaraShared.cpp
@@ -114,7 +114,6 @@ bool FNiagaraShaderScript::IsSame(const FNiagaraShaderMapId& InId) const
 	return
 		InId.FeatureLevel == FeatureLevel &&/*
 		InId.BaseScriptID == BaseScriptId &&*/
-		InId.DetailLevelMask == DetailLevelMask &&
 		InId.bUsesRapidIterationParams == bUsesRapidIterationParams &&
 		InId.BaseCompileHash == BaseCompileHash &&
 		InId.CompilerVersionID == CompilerVersionId;
@@ -148,7 +147,6 @@ NIAGARASHADER_API void FNiagaraShaderScript::GetShaderMapId(EShaderPlatform Plat
 		GetDependentShaderTypes(Platform, ShaderTypes);
 		OutId.FeatureLevel = GetFeatureLevel();/*
 		OutId.BaseScriptID = BaseScriptId;*/
-		OutId.DetailLevelMask = DetailLevelMask;
 		OutId.bUsesRapidIterationParams = bUsesRapidIterationParams;		
 		BaseCompileHash.ToSHAHash(OutId.BaseCompileHash);
 		OutId.CompilerVersionID = FNiagaraCustomVersion::LatestScriptCompileVersion;
@@ -259,7 +257,7 @@ void FNiagaraShaderScript::SerializeShaderMap(FArchive& Ar)
 
 void FNiagaraShaderScript::SetScript(UNiagaraScript *InScript, ERHIFeatureLevel::Type InFeatureLevel, const FGuid& InCompilerVersionID,  const TArray<FString>& InAdditionalDefines,
 		const FNiagaraCompileHash& InBaseCompileHash, const TArray<FNiagaraCompileHash>& InReferencedCompileHashes, 
-		bool bInUsesRapidIterationParams, uint32 InDetailLevelMask, FString InFriendlyName)
+		bool bInUsesRapidIterationParams, FString InFriendlyName)
 {
 	checkf(InBaseCompileHash.IsValid(), TEXT("Invalid base compile hash.  Script caching will fail."))
 	BaseVMScript = InScript;
@@ -267,7 +265,6 @@ void FNiagaraShaderScript::SetScript(UNiagaraScript *InScript, ERHIFeatureLevel:
 	//BaseScriptId = InBaseScriptID;
 	AdditionalDefines = InAdditionalDefines;
 	bUsesRapidIterationParams = bInUsesRapidIterationParams;
-	DetailLevelMask = InDetailLevelMask;
 	BaseCompileHash = InBaseCompileHash;
 	ReferencedCompileHashes = InReferencedCompileHashes;
 	FriendlyName = InFriendlyName;
diff --git a/Engine/Plugins/FX/Niagara/Source/NiagaraShader/Public/NiagaraShared.h b/Engine/Plugins/FX/Niagara/Source/NiagaraShader/Public/NiagaraShared.h
--- a/Engine/Plugins/FX/Niagara/Source/NiagaraShader/Public/NiagaraShared.h
+++ b/Engine/Plugins/FX/Niagara/Source/NiagaraShader/Public/NiagaraShared.h
@@ -226,10 +226,6 @@ public:
 	/** The compile hashes of the top level scripts the script is dependent on. */
 	LAYOUT_FIELD(TMemoryImageArray<FSHAHash>, ReferencedCompileHashes);
 
-	
-	/** Bitfield of supported detail levels in this compile.*/
-	LAYOUT_FIELD_INITIALIZED(uint32, DetailLevelMask, 0xFFFFFFFF);
-
 	/** Whether or not we need to bake Rapid Iteration params. True to keep params, false to bake.*/
 	LAYOUT_FIELD_INITIALIZED(bool, bUsesRapidIterationParams, true);
 
@@ -716,7 +712,7 @@ public:
 
 	NIAGARASHADER_API void SetScript(UNiagaraScript *InScript, ERHIFeatureLevel::Type InFeatureLevel, const FGuid& InCompilerVersion, const TArray<FString>& InAdditionalDefines,
 		const FNiagaraCompileHash& InBaseCompileHash, const TArray<FNiagaraCompileHash>& InReferencedCompileHashes, 
-		bool bInUsesRapidIterationParams, uint32 InDetailLevelMask, FString InFriendlyName);
+		bool bInUsesRapidIterationParams, FString InFriendlyName);
 
 	UNiagaraScript *GetBaseVMScript()
 	{
@@ -792,9 +788,6 @@ private:
 
 	/** Configuration options */
 	TArray<FString> AdditionalDefines;
-	
-	/** Bitfield of supported detail levels in this compile.*/
-	uint32 DetailLevelMask = 0xFFFFFFFF;
 
 	/** Whether or not we need to bake Rapid Iteration params. True to keep params, false to bake.*/
 	bool bUsesRapidIterationParams = true;
diff --git a/Engine/Source/Runtime/Engine/Private/DeviceProfiles/DeviceProfile.cpp b/Engine/Source/Runtime/Engine/Private/DeviceProfiles/DeviceProfile.cpp
--- a/Engine/Source/Runtime/Engine/Private/DeviceProfiles/DeviceProfile.cpp
+++ b/Engine/Source/Runtime/Engine/Private/DeviceProfiles/DeviceProfile.cpp
@@ -319,8 +319,12 @@ const TMap<FString, FString>& UDeviceProfile::GetConsolidatedCVars() const
 		{
 			FString CVarKey, CVarValue;
 			if (CurrentCVar.Split(TEXT("="), &CVarKey, &CVarValue))
-			{
-				InOutMap.Add(CVarKey, CVarValue);
+			{	
+				//Prevent parents from overwriting the values in child profiles.
+				if (InOutMap.Find(CVarKey) == nullptr)
+				{
+					InOutMap.Add(CVarKey, CVarValue);
+				}
 			}
 		}
 	};
diff --git a/Engine/Source/Runtime/Engine/Private/Scalability.cpp b/Engine/Source/Runtime/Engine/Private/Scalability.cpp
--- a/Engine/Source/Runtime/Engine/Private/Scalability.cpp
+++ b/Engine/Source/Runtime/Engine/Private/Scalability.cpp
@@ -959,10 +959,34 @@ FQualityLevels GetQualityLevelCounts()
 
 void LoadPlatformScalability(FString PlatformName)
 {
-	
+}
+
+#define LOCTEXT_NAMESPACE "Scalability"
+
+FText GetQualityLevelText(int32 QualityLevel, int32 NumLevels)
+{
+	//This matches logic in editor scalability settings UI. TODO: Unify.
+	const FText Names[5] = { LOCTEXT("QualityLowLabel", "Low"), LOCTEXT("QualityMediumLabel", "Medium"), LOCTEXT("QualityHighLabel", "High"), LOCTEXT("QualityEpicLabel", "Epic"), LOCTEXT("QualityCineLabel", "Cinematic") };
 
+	QualityLevel = FMath::Clamp(QualityLevel, 0, NumLevels - 1);
+	if (NumLevels == 5)
+	{
+		return Names[QualityLevel];
+	}
+	else
+	{
+		if (QualityLevel == NumLevels - 1)
+		{
+			return Names[4];
+		}
+		else
+		{
+			return FText::AsNumber(QualityLevel);
+		}
+	}
 }
 
+#undef LOCTEXT_NAMESPACE
 
 }
 
diff --git a/Engine/Source/Runtime/Engine/Public/Scalability.h b/Engine/Source/Runtime/Engine/Public/Scalability.h
--- a/Engine/Source/Runtime/Engine/Public/Scalability.h
+++ b/Engine/Source/Runtime/Engine/Public/Scalability.h
@@ -170,4 +170,6 @@ namespace Scalability
 	/** Set an Editor preview scalability platform */
 	void ENGINE_API ChangeScalabilityPreviewPlatform(FName NewPlatformScalabilityName);
 #endif
+
+	ENGINE_API FText GetQualityLevelText(int32 Value, int32 NumLevels);
 }
