Add locked axis and camera offset to mesh renderer (last features needed to make mesh facing modes view-independent)
Some clean-up of NiagaraMeshVertexFactory (no more working with transposed matrices)

diff --git a/Engine/Plugins/FX/Niagara/Shaders/Private/NiagaraMeshVertexFactory.ush b/Engine/Plugins/FX/Niagara/Shaders/Private/NiagaraMeshVertexFactory.ush
--- a/Engine/Plugins/FX/Niagara/Shaders/Private/NiagaraMeshVertexFactory.ush
+++ b/Engine/Plugins/FX/Niagara/Shaders/Private/NiagaraMeshVertexFactory.ush
@@ -26,6 +26,7 @@
 	Buffer<float> NiagaraParticleDataMaterialParam1;
 	Buffer<float> NiagaraParticleDataMaterialParam2;
 	Buffer<float> NiagaraParticleDataMaterialParam3;
+	Buffer<float> NiagaraParticleDataCameraOffset;
 #endif
 
 float3 TransformPosition(float3 InPosition)
@@ -56,8 +57,6 @@ float3 TransformVector(float3 InVector)
 //#define USE_PREVIOUS_TRANSFORM (FEATURE_LEVEL >= FEATURE_LEVEL_SM4)
 #define USE_PREVIOUS_TRANSFORM 0
 
-uint MeshFacingMode;
-
 #if USE_PREVIOUS_TRANSFORM
 	Buffer<float4> PrevTransformBuffer;
 #endif
@@ -67,6 +66,10 @@ uint MeshFacingMode;
 #define MESH_FACING_CAMERA_POSITION 2
 #define MESH_FACING_CAMERA_PLANE 3
 
+#define SPACE_SIMULATION 0
+#define SPACE_WORLD 1
+#define SPACE_LOCAL 2
+
 /** Sorted index buffer for particles if we need to sort their draw order (translucency etc). */
 Buffer<int> SortedIndices;
 uint SortedIndicesOffset;
@@ -423,93 +426,132 @@ float4 CalcPrevWorldPosition(FVertexFactoryInput Input, FVertexFactoryIntermedia
     return float4(WorldPosition, Input.Position.w);
 }
 
-// From UnrealMath::FindBetween
-float4 FindQuatBetween(float3 From, float3 To)
-{
-	// It appears that length(normalize(To)) is not always exactly equal to 1 and can create issues. 
-	// See comments around NEEDS_PARTICLE_VELOCITY
-	// float NormAB = sqrt(dot(NormAB, NormAB) * dot(To, To));
-
-	float NormAB = 1.0f;
-	float w = NormAB + dot(From, To);
-	float4 Result;
-
-	if (w >= 1e-6f * NormAB)
-	{
-		//Axis = cross(From, To);
-		Result = float4(From.y * To.z - From.z * To.y,
-					   From.z * To.x - From.x * To.z,
-					   From.x * To.y - From.y * To.x,
-					   w);
-	}
-	else
-	{
-		// From and To point in opposite directions
-		w = 0.f;
-		Result = abs(From.x) > abs(From.y)
-				? float4(-From.z, 0.f, From.x, w)
-				: float4(0.f, -From.z, From.y, w);
-	}
-
-	Result = normalize(Result);
-	return Result;
-}
-
-float4x4 QuatTo4x4(float4 Q)
+float3x3 QuatTo3x3(float4 Q)
 {
 	float x2 = Q.x + Q.x;  float y2 = Q.y + Q.y;  float z2 = Q.z + Q.z;
 	float xx = Q.x * x2;   float xy = Q.x * y2;   float xz = Q.x * z2;
 	float yy = Q.y * y2;   float yz = Q.y * z2;   float zz = Q.z * z2;
 	float wx = Q.w * x2;   float wy = Q.w * y2;   float wz = Q.w * z2;
 
-	float4x4 Mat = {
-	/*M[0][0]*/ 1.0f - (yy + zz),	/*M[1][0]*/ xy - wz,				/*M[2][0]*/ xz + wy,		0.0f,
-	/*M[0][1]*/ xy + wz,			/*M[1][1]*/ 1.0f - (xx + zz),		/*M[2][1]*/ yz - wx,		0.0f,
-	/*M[0][2]*/ xz - wy,			/*M[1][2]*/ yz + wx,				/*M[2][2]*/ 1.0f - (xx + yy),0.0f,
-	0.0f,							0.0f,								0.0f,						1.0f
+	float3x3 Mat = {
+		1.0f - (yy + zz),	xy + wz,			xz - wy,
+		xy - wz,			1.0f - (xx + zz),	yz + wx,
+		xz + wy,			yz - wx,			1.0f - (xx + yy),
 	};
+
 	return Mat;
 }
 
-void AxisAngleToPitchYawRoll(float4 AxisAngle, out float Pitch, out float Yaw, out float Roll)
+float3x3 GetFacingMatrix(in FVertexFactoryIntermediates Intermediates)
 {
-	AxisAngle.xyz= normalize(AxisAngle).xyz;
-	Yaw = atan2(AxisAngle.y, AxisAngle.x) * (180.f / PI);
-	if (AxisAngle.y == 0.0f && AxisAngle.x == 0.0f) // atan2 is undefined for 0,0
+	float3 WorldX = float3(1, 0, 0);
+	float3 WorldZ = float3(0, 0, 1);
+	float3 FacingDir = float3(0, 0, 0);
+	const bool bLocalSpace = NiagaraMeshVF.bLocalSpace;
+
+	// Select facing direction
+	switch (NiagaraMeshVF.FacingMode)
 	{
-		Yaw = 0.0f;
+		case MESH_FACING_VELOCITY:
+		{
+			FacingDir = Intermediates.ParticleVelocity.w > 0.0f ? Intermediates.ParticleVelocity.xyz : WorldZ;
+			break;
+		}
+
+		case MESH_FACING_CAMERA_POSITION:
+		{
+			float3 WorldPosition = TransformPosition(Intermediates.ParticlePosition.xyz);
+			float3 CameraVector = ResolvedView.WorldCameraOrigin - WorldPosition;
+			FacingDir = normalize(CameraVector);
+			break;
+		}
+
+		default:
+		case MESH_FACING_CAMERA_PLANE:
+		{
+			FacingDir = -ResolvedView.ViewForward;
+			break;
+		}
 	}
-	Pitch = atan2(AxisAngle.z,sqrt(AxisAngle.x*AxisAngle.x+AxisAngle.y*AxisAngle.y)) * (180.f / PI);
-	Roll = AxisAngle.w* (180.f / PI);
-}
 
+	// Check to rotate the facing direction to local-space
+	if (bLocalSpace)
+	{
+		FacingDir = RotateWorldToLocal(FacingDir);
+	}
 
-float4 EulerToQuaternion(float Pitch, float Yaw, float Roll)
-{
-	float DEG_TO_RAD = PI/(180.f);
-	float DIVIDE_BY_2 = DEG_TO_RAD/2.f;
-	float SP = sin(Pitch*DIVIDE_BY_2);
-	float SY = sin(Yaw*DIVIDE_BY_2);
-	float SR = sin(Roll*DIVIDE_BY_2);
-	float CP = cos(Pitch*DIVIDE_BY_2);
-	float CY = cos(Yaw*DIVIDE_BY_2);
-	float CR = cos(Roll*DIVIDE_BY_2);
-
-	float4 RotationQuat;
-	RotationQuat.x =  CR*SP*SY - SR*CP*CY;
-	RotationQuat.y = -CR*SP*CY - SR*CP*SY;
-	RotationQuat.z =  CR*CP*SY - SR*SP*CY;
-	RotationQuat.w =  CR*CP*CY + SR*SP*SY;
-
-	return RotationQuat;
+	float3 XAxis = { 1, 0, 0 };
+	float3 YAxis = { 0, 1, 0 };
+	float3 ZAxis = { 0, 0, 1 };
+	if (NiagaraMeshVF.bLockedAxisEnable)
+	{
+		// This is a special case where we want to lock the Z-Axis to the locked axis and get the X-Axis as close to facing direction as possible
+		const bool bWorldSpaceAxis = (NiagaraMeshVF.LockedAxisSpace == SPACE_WORLD) || (NiagaraMeshVF.LockedAxisSpace == SPACE_SIMULATION && !bLocalSpace);
+		float3 LockedAxis = NiagaraMeshVF.LockedAxis;
+		if (bWorldSpaceAxis && bLocalSpace)
+		{
+			// Transform the world-space axis to local space
+			LockedAxis = RotateWorldToLocal(LockedAxis);
+		}
+		else if (!bWorldSpaceAxis && !bLocalSpace)
+		{
+			// Transform the local-space axis to world space
+			LockedAxis = RotateLocalToWorld(LockedAxis);
+		}
+
+		ZAxis = LockedAxis;
+		YAxis = normalize(cross(ZAxis, FacingDir));
+		XAxis = cross(YAxis, ZAxis);
+	}
+	else
+	{
+		// Determine a reference vector to use for up
+		float3 RefVector;
+		if (NiagaraMeshVF.FacingMode == MESH_FACING_CAMERA_PLANE)
+		{
+			// Use the camera upwards direction as a reference vector
+			//-TODO: Add ability to remove HMD roll in VR
+			RefVector = ResolvedView.ViewUp;
+		}
+		else
+		{
+			// Prefer to use world up as a reference vector, fall back to negative world X-axis when facing up or down
+			RefVector = WorldZ;
+		}
+
+		// Check to rotate the reference direction to local-space
+		if (bLocalSpace)
+		{
+			RefVector = RotateWorldToLocal(RefVector);
+		}
+
+		// Orthonormalize the look-at rotation and generate a matrix
+		XAxis = FacingDir;
+		YAxis = normalize(cross(RefVector, FacingDir));
+		ZAxis = cross(XAxis, YAxis);
+	}
+	
+	return float3x3(XAxis, YAxis, ZAxis);
 }
 
-float4x4 EulerToMat4x4(float Pitch, float Yaw, float Roll)
+float3 CalculateCameraOffset(FVertexFactoryIntermediates Intermediates, uint InstanceId)
 {
-	float4 Quat = EulerToQuaternion(Pitch, Yaw, Roll);
-	return QuatTo4x4(Quat);
-}
+	if (NiagaraMeshVF.CameraOffsetDataOffset == -1)
+	{
+		return (float3)0.0f;
+	}
+
+	float Offset = GetFloat(NIAGARA_LOOSEPARAM_STRUCT(NiagaraParticleDataCameraOffset), NiagaraMeshVF.CameraOffsetDataOffset, InstanceId);
+	float3 WorldPosition = TransformPosition(Intermediates.ParticlePosition.xyz);
+	float3 CameraVector = ResolvedView.WorldCameraOrigin - WorldPosition;
+	float3 CameraDir = normalize(CameraVector);
+	if (NiagaraMeshVF.bLocalSpace)
+	{
+		CameraDir = RotateWorldToLocal(CameraDir);
+	}
 
+	return CameraDir * Offset;
+}
 
 FVertexFactoryIntermediates GetVertexFactoryIntermediates(FVertexFactoryInput Input)
 {
@@ -541,16 +583,17 @@ FVertexFactoryIntermediates GetVertexFactoryIntermediates(FVertexFactoryInput In
 		Intermediates.ParticleColor = GetVec4(NIAGARA_LOOSEPARAM_STRUCT(NiagaraParticleDataColor), NiagaraMeshVF.ColorDataOffset, Input.InstanceId);
 	}
 
-	float3 ParticleVelocity;
-	if (NiagaraMeshVF.VelocityDataOffset == -1)
-	{
-		ParticleVelocity = float3(0.0f, 0.0f, 0.0f);
-	}
-	else
+	float3 ParticleVelocity = (float3)0.0f;
+	if (NiagaraMeshVF.VelocityDataOffset != -1)
 	{
 		ParticleVelocity = float3(GetVec3(NIAGARA_LOOSEPARAM_STRUCT(NiagaraParticleDataVelocity), NiagaraMeshVF.VelocityDataOffset, Input.InstanceId));
 	}
 
+	// Calculate the particle velocity (NOTE: always do this regardless of NEEDS_PARTICLE_VELOCITY because of velocity facing particles)
+	const float ParticleVelocityMagnitude = length(ParticleVelocity);
+	Intermediates.ParticleVelocity.xyz = (ParticleVelocityMagnitude > 0) ? ParticleVelocity / ParticleVelocityMagnitude : float3(0.0f, 0.0f, 0.0f);
+	Intermediates.ParticleVelocity.w = ParticleVelocityMagnitude;
+
 #if (DYNAMIC_PARAMETERS_MASK & 1)
 	if (NiagaraMeshVF.MaterialParamDataOffset == -1)
 	{
@@ -612,108 +655,80 @@ FVertexFactoryIntermediates GetVertexFactoryIntermediates(FVertexFactoryInput In
 	{
 		Intermediates.ParticleRandom = GetFloat(NIAGARA_LOOSEPARAM_STRUCT(NiagaraParticleDataMaterialRandom), NiagaraMeshVF.MaterialRandomDataOffset, Input.InstanceId);
 	}
-#endif
+#endif	
 
-	const float ParticleVelocityMagnitude = length(ParticleVelocity);
-#if NEEDS_PARTICLE_VELOCITY
-	Intermediates.ParticleVelocity.xyz = (ParticleVelocityMagnitude > 0) ? normalize(ParticleVelocity) : float3(0.0f, 0.0f, 0.0f);
-	Intermediates.ParticleVelocity.w = ParticleVelocityMagnitude;
-#endif
-
-	float4 Transform1 = float4(1.0f, 0.0f, 0.0f, Intermediates.ParticlePosition.x);
-	float4 Transform2 = float4(0.0f, 1.0f, 0.0f, Intermediates.ParticlePosition.y);
-	float4 Transform3 = float4(0.0f, 0.0f, 1.0f, Intermediates.ParticlePosition.z);
+	// Swizzle vertex color.
+	Intermediates.VertexColor = Input.VertexColor FCOLOR_COMPONENT_SWIZZLE;
 
+	// Retrieve rotation
+	float3x3 RotationMat = {
+		1, 0, 0,
+		0, 1, 0,
+		0, 0, 1
+	};
 	if(NiagaraMeshVF.TransformDataOffset != -1)
 	{
 		float4 Quat = GetVec4(NIAGARA_LOOSEPARAM_STRUCT(NiagaraParticleDataTransform), NiagaraMeshVF.TransformDataOffset, Input.InstanceId);
-
-		//Transform axis to handle local space.
-		//AxisAngle.xyz = TransformVector(AxisAngle.xyz);
-
-
-		float4x4 EulerToMat = QuatTo4x4(normalize(Quat));
-		Transform1.xyz = EulerToMat[0].xyz;
-		Transform2.xyz = EulerToMat[1].xyz;
-		Transform3.xyz = EulerToMat[2].xyz;
+		RotationMat = QuatTo3x3(normalize(Quat));
 	}
-	float4x4 ScaleMatrix = float4x4(float4(1.0f, 0.0f, 0.0f, 0.0f), float4(0.0f, 1.0f, 0.0f, 0.0f), float4(0.0f, 0.0f, 1.0f, 0.0f), float4(0.0f, 0.0f, 0.0f, 1.0f));
+	
+	// Retrieve Scale
+	float3 Scale = (float3)1.0f;
 	if(NiagaraMeshVF.ScaleDataOffset != -1)
 	{
-		float3 Scale = GetVec3(NIAGARA_LOOSEPARAM_STRUCT(NiagaraParticleDataScale), NiagaraMeshVF.ScaleDataOffset, Input.InstanceId);
-		ScaleMatrix[0][0] *= Scale.x;
-		ScaleMatrix[1][1] *= Scale.y;
-		ScaleMatrix[2][2] *= Scale.z;
-	}
-	// Swizzle vertex color.
-	Intermediates.VertexColor = Input.VertexColor FCOLOR_COMPONENT_SWIZZLE;
-
-	// World position and localtoworld transform.
-	float4x4 LocalToWorldMat = float4x4(Transform1, Transform2, Transform3, float4(0.0f, 0.0f, 0.0f, 1.0f));
-
-	Intermediates.ParticleTranslatedWorldPosition = Intermediates.ParticlePosition.xyz + ResolvedView.PreViewTranslation.xyz;
-
-	if (MeshFacingMode == MESH_FACING_VELOCITY)
-	{
-		// The result of normalize(ParticleVelocity) seems affected by whether lenght(ParticleVelocity) is computed above or not (see NEEDS_PARTICLE_VELOCITY).
-		// This makes z fighting possible between the basepass and the depthpass (which never defines NEEDS_PARTICLE_VELOCITY) for opaque materials.
-		// The workaround for this issue is to ensure that lenght(ParticleVelocity) gets computed in both shaders.
-		float4 ToAlignmentQuat = FindQuatBetween(float3(1.0f, 0.0f, 0.0f), (ParticleVelocityMagnitude > 0) ? normalize(ParticleVelocity) : float3(0.0f, 0.0f, 0.0f));
-		float4x4 ToAlignmentMat = QuatTo4x4( normalize(ToAlignmentQuat) );
-		Intermediates.LocalToWorld =  mul(LocalToWorldMat, mul(ToAlignmentMat, ScaleMatrix));
-	}
-	else if (MeshFacingMode == MESH_FACING_CAMERA_POSITION)
-	{
-		float3 WorldPosition = TransformPosition(Intermediates.ParticlePosition);
-		float3 CameraVector = ResolvedView.WorldCameraOrigin.xyz - WorldPosition.xyz;
-		float3 CameraForward = normalize(CameraVector);
-		float3 CameraRight = normalize(cross(CameraForward, ResolvedView.ViewUp));
-		float3 CameraUp = cross(CameraRight, CameraForward);
-
-		float4x4 ToAlignmentMat;
-		ToAlignmentMat[0] = float4(CameraForward.x, -CameraRight.x, CameraUp.x, 0.0f);
-		ToAlignmentMat[1] = float4(CameraForward.y, -CameraRight.y, CameraUp.y, 0.0f);
-		ToAlignmentMat[2] = float4(CameraForward.z, -CameraRight.z, CameraUp.z, 0.0f);
-		ToAlignmentMat[3] = float4(0.0f, 0.0f, 0.0f, 1.0f);
-
-		Intermediates.LocalToWorld = mul(LocalToWorldMat, mul(ToAlignmentMat, ScaleMatrix));
+		Scale = GetVec3(NIAGARA_LOOSEPARAM_STRUCT(NiagaraParticleDataScale), NiagaraMeshVF.ScaleDataOffset, Input.InstanceId);		
 	}
-	else if (MeshFacingMode == MESH_FACING_CAMERA_PLANE)
-	{
-		float4x4 ToAlignmentMat;
-		ToAlignmentMat[0] = float4(-ResolvedView.ViewForward.x, -ResolvedView.ViewRight.x, ResolvedView.ViewUp.x, 0.0f);
-		ToAlignmentMat[1] = float4(-ResolvedView.ViewForward.y, -ResolvedView.ViewRight.y, ResolvedView.ViewUp.y, 0.0f);
-		ToAlignmentMat[2] = float4(-ResolvedView.ViewForward.z, -ResolvedView.ViewRight.z, ResolvedView.ViewUp.z, 0.0f);
-		ToAlignmentMat[3] = float4(0.0f, 0.0f, 0.0f, 1.0f);
-		Intermediates.LocalToWorld =  mul(LocalToWorldMat, mul(ToAlignmentMat, ScaleMatrix));
-	}
-	else
+	
+	float3x3 ScaleMat = {
+		Scale.x,	0,			0,
+		0,			Scale.y,	0,
+		0,			0,			Scale.z
+	};
+	
+	float3x3 ScaleRotMat = mul(ScaleMat, RotationMat);
+	if (NiagaraMeshVF.FacingMode != MESH_FACING_DEFAULT)
 	{
-		Intermediates.LocalToWorld =  mul(LocalToWorldMat, ScaleMatrix);
+		// Factor in facing rotation
+		float3x3 FacingMat = GetFacingMatrix(Intermediates);
+		ScaleRotMat = mul(ScaleRotMat, FacingMat);
 	}
 
-	Intermediates.LocalToWorld = transpose(Intermediates.LocalToWorld);
+	// Apply CameraOffset
+	Intermediates.ParticlePosition.xyz += CalculateCameraOffset(Intermediates, Input.InstanceId);
+
+	// Determine final "simulation space" matrix
+	Intermediates.LocalToWorld = float4x4(
+		float4(ScaleRotMat[0], 0),
+		float4(ScaleRotMat[1], 0),
+		float4(ScaleRotMat[2], 0),
+		Intermediates.ParticlePosition
+	);
 
-    //-TODO: If / when we store previous transform data use them here (providing they are in local space, if not handle below)
+	//-TODO: If / when we store previous transform data use them here (providing they are in local space, if not handle below)
 	//if (MeshParticleVF.PrevTransformAvailable > 0)
 	//{
     //   //-TODO:
 	//}
     //else
     {
+		// NOTE: Velocity is in simulation space, so calculate prev transform here, before handling local space
         Intermediates.PrevLocalToWorld = Intermediates.LocalToWorld;
         Intermediates.PrevLocalToWorld[3].xyz -= ParticleVelocity * NiagaraMeshVF.DeltaSeconds;
     }
 
 	//Handle local space. TODO: Optimize behind define.
-	//Apply local space to local to world, also pretranspse here to avoid it in the vs&ps
+	//Apply local space to local to world, also pretranspose here to avoid it in the vs&ps
+
 	if (NiagaraMeshVF.bLocalSpace)
 	{
 		Intermediates.PrevLocalToWorld = mul(Intermediates.PrevLocalToWorld, Primitive.PreviousLocalToWorld);
 		Intermediates.LocalToWorld = mul(Intermediates.LocalToWorld, Primitive.LocalToWorld);
-		Intermediates.ParticlePosition = TransformLocalToWorld(Intermediates.ParticlePosition);
+		Intermediates.ParticlePosition = Intermediates.LocalToWorld[3];
 	}
 
+	// Translated world position.	
+	Intermediates.ParticleTranslatedWorldPosition = Intermediates.ParticlePosition.xyz + ResolvedView.PreViewTranslation.xyz;
+
 #if NUM_MATERIAL_TEXCOORDS_VERTEX
 	for(int CoordinateIndex = 0; CoordinateIndex < NUM_MATERIAL_TEXCOORDS_VERTEX; CoordinateIndex++)
 	{
diff --git a/Engine/Plugins/FX/Niagara/Source/Niagara/Private/NiagaraMeshRendererProperties.cpp b/Engine/Plugins/FX/Niagara/Source/Niagara/Private/NiagaraMeshRendererProperties.cpp
--- a/Engine/Plugins/FX/Niagara/Source/Niagara/Private/NiagaraMeshRendererProperties.cpp
+++ b/Engine/Plugins/FX/Niagara/Source/Niagara/Private/NiagaraMeshRendererProperties.cpp
@@ -44,9 +44,14 @@ bool FNiagaraMeshMaterialOverride::SerializeFromMismatchedTag(const struct FProp
 UNiagaraMeshRendererProperties::UNiagaraMeshRendererProperties()
 	: ParticleMesh(nullptr)
 	, SortMode(ENiagaraSortMode::None)
+	, bOverrideMaterials(false)
 	, bSortOnlyWhenTranslucent(true)
 	, SubImageSize(1.0f, 1.0f)
 	, bSubImageBlend(false)
+	, FacingMode(ENiagaraMeshFacingMode::Default)
+	, bLockedAxisEnable(false)
+	, LockedAxis(0.0f, 0.0f, 1.0f)
+	, LockedAxisSpace(ENiagaraMeshLockedAxisSpace::Simulation)
 {
 }
 
@@ -134,6 +139,7 @@ void UNiagaraMeshRendererProperties::InitBindings()
 		ScaleBinding = FNiagaraConstants::GetAttributeDefaultBinding(SYS_PARAM_PARTICLES_SCALE);
 		MaterialRandomBinding = FNiagaraConstants::GetAttributeDefaultBinding(SYS_PARAM_PARTICLES_MATERIAL_RANDOM);
 		NormalizedAgeBinding = FNiagaraConstants::GetAttributeDefaultBinding(SYS_PARAM_PARTICLES_NORMALIZED_AGE);
+		CameraOffsetBinding = FNiagaraConstants::GetAttributeDefaultBinding(SYS_PARAM_PARTICLES_CAMERA_OFFSET);
 
 		//Default custom sorting to age
 		CustomSortingBinding = FNiagaraConstants::GetAttributeDefaultBinding(SYS_PARAM_PARTICLES_NORMALIZED_AGE);
diff --git a/Engine/Plugins/FX/Niagara/Source/Niagara/Private/NiagaraRendererMeshes.cpp b/Engine/Plugins/FX/Niagara/Source/Niagara/Private/NiagaraRendererMeshes.cpp
--- a/Engine/Plugins/FX/Niagara/Source/Niagara/Private/NiagaraRendererMeshes.cpp
+++ b/Engine/Plugins/FX/Niagara/Source/Niagara/Private/NiagaraRendererMeshes.cpp
@@ -78,6 +78,7 @@ namespace ENiagaraMeshVFLayout
 		DynamicParam1,
 		DynamicParam2,
 		DynamicParam3,
+		CameraOffset,
 
 		Num,
 	};
@@ -96,12 +97,21 @@ FNiagaraRendererMeshes::FNiagaraRendererMeshes(ERHIFeatureLevel::Type FeatureLev
 	MeshRenderData = Properties->ParticleMesh->RenderData.Get();
 
 	FacingMode = Properties->FacingMode;
+	bLockedAxisEnable = Properties->bLockedAxisEnable;
+	LockedAxis = Properties->LockedAxis;
+	LockedAxisSpace = Properties->LockedAxisSpace;
 	SortMode = Properties->SortMode;
 	bSortOnlyWhenTranslucent = Properties->bSortOnlyWhenTranslucent;
 	bOverrideMaterials = Properties->bOverrideMaterials;
 	SubImageSize = Properties->SubImageSize;
 	bSubImageBlend = Properties->bSubImageBlend;
 
+	// Ensure valid value for the locked axis
+	if (!LockedAxis.Normalize())
+	{
+		LockedAxis.Set(0.0f, 0.0f, 1.0f);
+	}
+
 	const FNiagaraDataSet& Data = Emitter->GetData();
 
 	MaterialParamValidMask = 0;
@@ -116,6 +126,7 @@ FNiagaraRendererMeshes::FNiagaraRendererMeshes(ERHIFeatureLevel::Type FeatureLev
 	SetVertexFactoryVariable(Data, Properties->NormalizedAgeBinding.DataSetVariable, ENiagaraMeshVFLayout::NormalizedAge);
 	SetVertexFactoryVariable(Data, Properties->CustomSortingBinding.DataSetVariable, ENiagaraMeshVFLayout::CustomSorting);
 	SetVertexFactoryVariable(Data, Properties->SubImageIndexBinding.DataSetVariable, ENiagaraMeshVFLayout::SubImage);
+	SetVertexFactoryVariable(Data, Properties->CameraOffsetBinding.DataSetVariable, ENiagaraMeshVFLayout::CameraOffset);
 	MaterialParamValidMask |= SetVertexFactoryVariable(Data, Properties->DynamicMaterialBinding.DataSetVariable, ENiagaraMeshVFLayout::DynamicParam0) ? 0x1 : 0;
 	MaterialParamValidMask |= SetVertexFactoryVariable(Data, Properties->DynamicMaterial1Binding.DataSetVariable, ENiagaraMeshVFLayout::DynamicParam1) ? 0x2 : 0;
 	MaterialParamValidMask |= SetVertexFactoryVariable(Data, Properties->DynamicMaterial2Binding.DataSetVariable, ENiagaraMeshVFLayout::DynamicParam2) ? 0x4 : 0;
@@ -310,7 +321,6 @@ void FNiagaraRendererMeshes::GetDynamicMeshElements(const TArray<const FSceneVie
 					VertexFactories.Add(VertexFactory);
 
 					VertexFactory->SetParticleFactoryType(NVFT_Mesh);
-					VertexFactory->SetMeshFacingMode((uint32)FacingMode);
 					VertexFactory->SetLODIndex(LODIndex);
 					VertexFactory->InitResource();
 					SetupVertexFactory(VertexFactory, LODModel);
@@ -324,7 +334,7 @@ void FNiagaraRendererMeshes::GetDynamicMeshElements(const TArray<const FSceneVie
 
 				PerViewUniformParameters.bLocalSpace = bLocalSpace;
 				PerViewUniformParameters.PrevTransformAvailable = false;
-				PerViewUniformParameters.DeltaSeconds = ViewFamily.DeltaWorldTime;
+				PerViewUniformParameters.DeltaSeconds = ViewFamily.DeltaWorldTime;				
 
 				PerViewUniformParameters.PositionDataOffset = VFVariables[ENiagaraMeshVFLayout::Position].GetGPUOffset();
 				PerViewUniformParameters.VelocityDataOffset = VFVariables[ENiagaraMeshVFLayout::Velocity].GetGPUOffset();
@@ -338,12 +348,17 @@ void FNiagaraRendererMeshes::GetDynamicMeshElements(const TArray<const FSceneVie
 				PerViewUniformParameters.MaterialParam1DataOffset = VFVariables[ENiagaraMeshVFLayout::DynamicParam1].GetGPUOffset();
 				PerViewUniformParameters.MaterialParam2DataOffset = VFVariables[ENiagaraMeshVFLayout::DynamicParam2].GetGPUOffset();
 				PerViewUniformParameters.MaterialParam3DataOffset = VFVariables[ENiagaraMeshVFLayout::DynamicParam3].GetGPUOffset();
+				PerViewUniformParameters.CameraOffsetDataOffset = VFVariables[ENiagaraMeshVFLayout::CameraOffset].GetGPUOffset();
 
 				PerViewUniformParameters.MaterialParamValidMask = MaterialParamValidMask;
 				PerViewUniformParameters.SizeDataOffset = INDEX_NONE;
 				PerViewUniformParameters.DefaultPos = bLocalSpace ? FVector4(0.0f, 0.0f, 0.0f, 1.0f) : FVector4(SceneProxy->GetLocalToWorld().GetOrigin());
 				PerViewUniformParameters.SubImageSize = FVector4(SubImageSize.X, SubImageSize.Y, 1.0f / SubImageSize.X, 1.0f / SubImageSize.Y);
 				PerViewUniformParameters.SubImageBlendMode = bSubImageBlend;
+				PerViewUniformParameters.FacingMode = (uint32)FacingMode;
+				PerViewUniformParameters.bLockedAxisEnable = bLockedAxisEnable;
+				PerViewUniformParameters.LockedAxis = LockedAxis;
+				PerViewUniformParameters.LockedAxisSpace = (uint32)LockedAxisSpace;
 
 				//Sort particles if needed.
 				CollectorResources.VertexFactory->SetSortedIndices(nullptr, 0xFFFFFFFF);
diff --git a/Engine/Plugins/FX/Niagara/Source/Niagara/Public/NiagaraMeshRendererProperties.h b/Engine/Plugins/FX/Niagara/Source/Niagara/Public/NiagaraMeshRendererProperties.h
--- a/Engine/Plugins/FX/Niagara/Source/Niagara/Public/NiagaraMeshRendererProperties.h
+++ b/Engine/Plugins/FX/Niagara/Source/Niagara/Public/NiagaraMeshRendererProperties.h
@@ -27,6 +27,16 @@ enum class ENiagaraMeshFacingMode : uint8
 	CameraPlane
 };
 
+UENUM()
+enum class ENiagaraMeshLockedAxisSpace : uint8 {
+	/** The locked axis is in the emitter's local space if the emitter is marked as local-space, or in world space otherwise */
+	Simulation,
+	/** The locked axis is in world space */
+	World,
+	/** The locked axis is in the emitter's local space */
+	Local
+};
+
 USTRUCT()
 struct FNiagaraMeshMaterialOverride 
 {
@@ -127,7 +137,19 @@ public:
 
 	/** Determines how the mesh orients itself relative to the camera.*/
 	UPROPERTY(EditAnywhere, Category = "Mesh Rendering")
-	ENiagaraMeshFacingMode FacingMode; 
+	ENiagaraMeshFacingMode FacingMode;
+
+	/** If true and in a non-default facing mode, will lock facing direction to an arbitrary plane of rotation */
+	UPROPERTY(EditAnywhere, Category = "Mesh Rendering")
+	uint32 bLockedAxisEnable : 1;
+
+	/** Arbitrary axis by which to lock facing rotations */
+	UPROPERTY(EditAnywhere, Category = "Mesh Rendering", meta = (EditCondition = "bLockedAxisEnable"))
+	FVector LockedAxis;
+
+	/** Specifies what space the locked axis is in */
+	UPROPERTY(EditAnywhere, Category = "Mesh Rendering", meta = (EditCondition = "bLockedAxisEnable"))
+	ENiagaraMeshLockedAxisSpace LockedAxisSpace;
 	
 	/** Which attribute should we use for position when generating instanced meshes?*/
 	UPROPERTY(EditAnywhere, Category = "Bindings")
@@ -181,6 +203,10 @@ public:
 	UPROPERTY(EditAnywhere, Category = "Bindings")
 	FNiagaraVariableAttributeBinding NormalizedAgeBinding;
 
+	/** Which attribute should we use for camera offset when rendering meshes?*/
+	UPROPERTY(EditAnywhere, Category = "Bindings")
+	FNiagaraVariableAttributeBinding CameraOffsetBinding;
+
 
 protected:
 	bool FindBinding(const FNiagaraUserParameterBinding& InBinding, const FNiagaraEmitterInstance* InEmitter, TArray<UMaterialInterface*>& OutMaterials);
diff --git a/Engine/Plugins/FX/Niagara/Source/Niagara/Public/NiagaraRendererMeshes.h b/Engine/Plugins/FX/Niagara/Source/Niagara/Public/NiagaraRendererMeshes.h
--- a/Engine/Plugins/FX/Niagara/Source/Niagara/Public/NiagaraRendererMeshes.h
+++ b/Engine/Plugins/FX/Niagara/Source/Niagara/Public/NiagaraRendererMeshes.h
@@ -47,10 +47,14 @@ private:
 	ENiagaraMeshFacingMode FacingMode;
 	uint32 bOverrideMaterials : 1;
 	uint32 bSortOnlyWhenTranslucent : 1;
+	uint32 bLockedAxisEnable : 1;
 
 	uint32 bSubImageBlend : 1;
 	FVector2D SubImageSize;
 
+	FVector LockedAxis;
+	ENiagaraMeshLockedAxisSpace LockedAxisSpace;
+
 	uint32 MaterialParamValidMask;
 
 	int32 MeshMinimumLOD = 0;
diff --git a/Engine/Plugins/FX/Niagara/Source/NiagaraVertexFactories/Private/NiagaraMeshVertexFactory.cpp b/Engine/Plugins/FX/Niagara/Source/NiagaraVertexFactories/Private/NiagaraMeshVertexFactory.cpp
--- a/Engine/Plugins/FX/Niagara/Source/NiagaraVertexFactories/Private/NiagaraMeshVertexFactory.cpp
+++ b/Engine/Plugins/FX/Niagara/Source/NiagaraVertexFactories/Private/NiagaraMeshVertexFactory.cpp
@@ -10,7 +10,7 @@
 #include "ShaderParameterUtils.h"
 #include "MeshMaterialShader.h"
 
-IMPLEMENT_GLOBAL_SHADER_PARAMETER_STRUCT(FNiagaraMeshUniformParameters,"NiagaraMeshVF");
+IMPLEMENT_GLOBAL_SHADER_PARAMETER_STRUCT(FNiagaraMeshUniformParameters, "NiagaraMeshVF");
 
 class FNiagaraMeshVertexFactoryShaderParametersVS : public FVertexFactoryShaderParameters
 {
@@ -31,6 +31,7 @@ public:
 		NiagaraParticleDataMaterialParam2.Bind(ParameterMap, TEXT("NiagaraParticleDataMaterialParam2"));
 		NiagaraParticleDataMaterialParam3.Bind(ParameterMap, TEXT("NiagaraParticleDataMaterialParam3"));
 		NiagaraParticleDataSubImage.Bind(ParameterMap, TEXT("NiagaraParticleDataSubImage"));
+		NiagaraParticleDataCameraOffset.Bind(ParameterMap, TEXT("NiagaraParticleDataCameraOffset"));
 
 		FloatDataStride.Bind(ParameterMap, TEXT("NiagaraFloatDataStride"));
 
@@ -38,7 +39,6 @@ public:
 		// 		FloatDataOffset.Bind(ParameterMap, TEXT("NiagaraInt32DataOffset"));
 		// 		FloatDataStride.Bind(ParameterMap, TEXT("NiagaraInt3DataStride"));
 
-		MeshFacingMode.Bind(ParameterMap, TEXT("MeshFacingMode"));
 		SortedIndices.Bind(ParameterMap, TEXT("SortedIndices"));
 		SortedIndicesOffset.Bind(ParameterMap, TEXT("SortedIndicesOffset"));
 	}
@@ -57,8 +57,6 @@ public:
 		FNiagaraMeshVertexFactory* NiagaraMeshVF = (FNiagaraMeshVertexFactory*)VertexFactory;
 		ShaderBindings.Add(Shader->GetUniformBufferParameter<FNiagaraMeshUniformParameters>(), NiagaraMeshVF->GetUniformBuffer());
 
-		ShaderBindings.Add(MeshFacingMode, NiagaraMeshVF->GetMeshFacingMode());
-
 		ShaderBindings.Add(NiagaraParticleDataPosition, NiagaraMeshVF->GetParticleDataFloatSRV());
 		ShaderBindings.Add(NiagaraParticleDataVelocity, NiagaraMeshVF->GetParticleDataFloatSRV());
 		ShaderBindings.Add(NiagaraParticleDataColor, NiagaraMeshVF->GetParticleDataFloatSRV());
@@ -71,6 +69,7 @@ public:
 		ShaderBindings.Add(NiagaraParticleDataMaterialParam2, NiagaraMeshVF->GetParticleDataFloatSRV());
 		ShaderBindings.Add(NiagaraParticleDataMaterialParam3, NiagaraMeshVF->GetParticleDataFloatSRV());
 		ShaderBindings.Add(NiagaraParticleDataSubImage, NiagaraMeshVF->GetParticleDataFloatSRV());
+		ShaderBindings.Add(NiagaraParticleDataCameraOffset, NiagaraMeshVF->GetParticleDataFloatSRV());
 
 		ShaderBindings.Add(FloatDataStride, NiagaraMeshVF->GetFloatDataStride());
 
@@ -81,7 +80,7 @@ public:
 
 private:
 
-	
+
 		//LAYOUT_FIELD(FShaderResourceParameter, PrevTransformBuffer)
 
 	LAYOUT_FIELD(FShaderResourceParameter, NiagaraParticleDataPosition);
@@ -96,16 +95,15 @@ private:
 	LAYOUT_FIELD(FShaderResourceParameter, NiagaraParticleDataMaterialParam2);
 	LAYOUT_FIELD(FShaderResourceParameter, NiagaraParticleDataMaterialParam3);
 	LAYOUT_FIELD(FShaderResourceParameter, NiagaraParticleDataSubImage);
+	LAYOUT_FIELD(FShaderResourceParameter, NiagaraParticleDataCameraOffset);
 	LAYOUT_FIELD(FShaderParameter, FloatDataStride);
 
 		// 	LAYOUT_FIELD(FShaderResourceParameter, NiagaraParticleDataInt)
 		// 	LAYOUT_FIELD(FShaderParameter, Int32DataOffset)
 		// 	LAYOUT_FIELD(FShaderParameter, Int32DataStride)
-
-		LAYOUT_FIELD(FShaderParameter, MeshFacingMode)
 		LAYOUT_FIELD(FShaderResourceParameter, SortedIndices)
 		LAYOUT_FIELD(FShaderParameter, SortedIndicesOffset)
-	
+
 };
 
 class FNiagaraMeshVertexFactoryShaderParametersPS : public FVertexFactoryShaderParameters
diff --git a/Engine/Plugins/FX/Niagara/Source/NiagaraVertexFactories/Public/NiagaraMeshVertexFactory.h b/Engine/Plugins/FX/Niagara/Source/NiagaraVertexFactories/Public/NiagaraMeshVertexFactory.h
--- a/Engine/Plugins/FX/Niagara/Source/NiagaraVertexFactories/Public/NiagaraMeshVertexFactory.h
+++ b/Engine/Plugins/FX/Niagara/Source/NiagaraVertexFactories/Public/NiagaraMeshVertexFactory.h
@@ -48,8 +48,13 @@ BEGIN_GLOBAL_SHADER_PARAMETER_STRUCT(FNiagaraMeshUniformParameters, NIAGARAVERTE
 	SHADER_PARAMETER(int, MaterialParam3DataOffset)
 	SHADER_PARAMETER(int, NormalizedAgeDataOffset)
 	SHADER_PARAMETER(int, MaterialRandomDataOffset)
+	SHADER_PARAMETER(int, CameraOffsetDataOffset)
 	SHADER_PARAMETER(FVector4, DefaultPos)
 	SHADER_PARAMETER(int, SubImageBlendMode)
+	SHADER_PARAMETER(uint32, FacingMode)
+	SHADER_PARAMETER(uint32, bLockedAxisEnable)
+	SHADER_PARAMETER(FVector, LockedAxis)
+	SHADER_PARAMETER(uint32, LockedAxisSpace)
 END_GLOBAL_SHADER_PARAMETER_STRUCT()
 
 typedef TUniformBufferRef<FNiagaraMeshUniformParameters> FNiagaraMeshUniformBufferRef;
@@ -69,7 +74,6 @@ public:
 	FNiagaraMeshVertexFactory(ENiagaraVertexFactoryType InType, ERHIFeatureLevel::Type InFeatureLevel)
 		: FNiagaraVertexFactoryBase(InType, InFeatureLevel)
 		, LODIndex(-1)
-		, MeshFacingMode(0)
 		, InstanceVerticesCPU(nullptr)
 		, FloatDataStride(0)
 		, SortedIndicesOffset(0)
@@ -78,7 +82,6 @@ public:
 	FNiagaraMeshVertexFactory()
 		: FNiagaraVertexFactoryBase(NVFT_MAX, ERHIFeatureLevel::Num)
 		, LODIndex(-1)
-		, MeshFacingMode(0)
 		, InstanceVerticesCPU(nullptr)
 		, FloatDataStride(0)
 		, SortedIndicesOffset(0)
@@ -173,14 +176,10 @@ public:
 
 	int32 GetLODIndex() const { return LODIndex; }
 	void SetLODIndex(int32 InLODIndex) { LODIndex = InLODIndex; }
-
-	uint32 GetMeshFacingMode() const { return MeshFacingMode; }
-	void SetMeshFacingMode(uint32 InMode) { MeshFacingMode = InMode; }
-
+	
 protected:
 	FStaticMeshDataType Data;
-	int32 LODIndex;
-	uint32 MeshFacingMode;
+	int32 LODIndex;	
 
 	/** Uniform buffer with mesh particle parameters. */
 	FRHIUniformBuffer* MeshParticleUniformBuffer;
diff --git a/Engine/Shaders/Private/VertexFactoryCommon.ush b/Engine/Shaders/Private/VertexFactoryCommon.ush
--- a/Engine/Shaders/Private/VertexFactoryCommon.ush
+++ b/Engine/Shaders/Private/VertexFactoryCommon.ush
@@ -49,6 +49,18 @@ float3 RotateLocalToWorld(float3 LocalDirection)
 		InvScale.z * Primitive.LocalToWorld[2].xyz * LocalDirection.zzz;
 }
 
+float3 RotateWorldToLocal(float3 WorldDirection)
+{
+	const float3 InvScale = Primitive.InvNonUniformScaleAndDeterminantSign.xyz;
+	float3x3 InvRot = {
+		InvScale.x * Primitive.LocalToWorld[0].xyz,
+		InvScale.y * Primitive.LocalToWorld[1].xyz,
+		InvScale.z * Primitive.LocalToWorld[2].xyz
+	};
+	InvRot = transpose(InvRot);
+	return mul(WorldDirection, InvRot);
+}
+
 // Octahedron Normal Vectors
 // [Cigolle 2014, "A Survey of Efficient Representations for Independent Unit Vectors"]
 //						Mean	Max
