Fixed stale entries in the ID to index buffers for CPU emitters.

The code used to update the output ID to index table, but that left stale entries in the input buffer set. When particles died, one of the buffers had -1 for the dead IDs, but the other still had the old execution indexes. This could lead to direct reads returning bogus data, instead of reporting that the ID is no longer in use. The solution is to clear the ID table at the beginning of the tick, and let the UpdateID() function fill in the IDs which are still in use.

This problem didn't occur on GPU, where the tables were already being cleared.

diff --git a/Engine/Plugins/FX/Niagara/Source/Niagara/Private/NiagaraDataInterfaceParticleRead.cpp b/Engine/Plugins/FX/Niagara/Source/Niagara/Private/NiagaraDataInterfaceParticleRead.cpp
--- a/Engine/Plugins/FX/Niagara/Source/Niagara/Private/NiagaraDataInterfaceParticleRead.cpp
+++ b/Engine/Plugins/FX/Niagara/Source/Niagara/Private/NiagaraDataInterfaceParticleRead.cpp
@@ -309,7 +309,11 @@ struct FNiagaraDataInterfaceParametersCS_ParticleRead : public FNiagaraDataInter
 		check(Proxy);
 
 		FNDIParticleRead_RenderThreadData* RTData = Proxy->GetRenderDataForSystem(Context.SystemInstance);
-		check(RTData);
+		if (!RTData)
+		{
+			SetErrorParams(RHICmdList, ComputeShader, false);
+			return;
+		}
 
 		if (RTData->SourceEmitterGPUContext == nullptr)
 		{
diff --git a/Engine/Plugins/FX/Niagara/Source/Niagara/Private/NiagaraDataSet.cpp b/Engine/Plugins/FX/Niagara/Source/Niagara/Private/NiagaraDataSet.cpp
--- a/Engine/Plugins/FX/Niagara/Source/Niagara/Private/NiagaraDataSet.cpp
+++ b/Engine/Plugins/FX/Niagara/Source/Niagara/Private/NiagaraDataSet.cpp
@@ -236,14 +236,6 @@ void FNiagaraDataSet::Allocate(int32 NumInstances, bool bMaintainExisting)
 		{
 			//UE_LOG(LogNiagara, Warning, TEXT("Growing ID Table! OldSize:%d | NewSize:%d"), ExistingNumIDs, RequiredIDs);
 			int32 NewNumIds = RequiredIDs - ExistingNumIDs;
-#if 0
-			DestinationIDTable.SetNumUninitialized(RequiredIDs);
-#else
-			while (DestinationIDTable.Num() < RequiredIDs)
-			{
-				DestinationIDTable.Add(INDEX_NONE);
-			}
-#endif
 
 			//Free ID Table must always be at least as large as the data buffer + it's current size in the case all particles die this frame.
 			FreeIDsTable.AddUninitialized(NewNumIds);
@@ -293,8 +285,13 @@ void FNiagaraDataSet::Allocate(int32 NumInstances, bool bMaintainExisting)
 			RequiredIDs = ExistingNumIDs;
 		}
 
+		// We need to clear the ID to index table to -1 so we don't have stale entries for particles which died in the previous
+		// frame (when the results were written to another buffer). All the entries which are in use will be filled in by the script.
 		DestinationIDTable.SetNumUninitialized(RequiredIDs);
-		MaxUsedID = INDEX_NONE;//reset the max ID ready for it to be filled in during simulation.
+		FMemory::Memset(DestinationIDTable.GetData(), -1, RequiredIDs * DestinationIDTable.GetTypeSize());
+
+		//reset the max ID ready for it to be filled in during simulation.
+		MaxUsedID = INDEX_NONE;
 
 // 		UE_LOG(LogNiagara, Warning, TEXT("DataSetAllocate: NumInstances:%d | ID Table Size:%d | NumFreeIDs:%d | FreeTableSize:%d"), NumInstances, DestinationData->GetIDTable().Num(), NumFreeIDs, FreeIDsTable.Num());
 // 		UE_LOG(LogNiagara, Warning, TEXT("== FreeIDs %d =="), NumFreeIDs);
diff --git a/Engine/Source/Runtime/VectorVM/Private/VectorVM.cpp b/Engine/Source/Runtime/VectorVM/Private/VectorVM.cpp
--- a/Engine/Source/Runtime/VectorVM/Private/VectorVM.cpp
+++ b/Engine/Source/Runtime/VectorVM/Private/VectorVM.cpp
@@ -1504,7 +1504,6 @@ struct FScalarKernelUpdateID
 			{
 				//Add the ID to a thread local list of IDs to free which are actually added to the list safely at the end of this chunk's execution.
 				IDsToFree.Add(InstanceId);
-				IDTable[InstanceId] = INDEX_NONE;
 				//UE_LOG(LogVectorVM, Warning, TEXT("FreeingID: InstanceID:%d."), InstanceId);
 			}
 			else
