Removing a few edits that crept into last cl in error.

diff --git a/Engine/Plugins/FX/Niagara/Source/Niagara/Private/NiagaraEmitter.cpp b/Engine/Plugins/FX/Niagara/Source/Niagara/Private/NiagaraEmitter.cpp
--- a/Engine/Plugins/FX/Niagara/Source/Niagara/Private/NiagaraEmitter.cpp
+++ b/Engine/Plugins/FX/Niagara/Source/Niagara/Private/NiagaraEmitter.cpp
@@ -16,9 +16,6 @@
 #include "Misc/ConfigCacheIni.h"
 #include "Interfaces/ITargetPlatform.h"
 
-PRAGMA_DISABLE_INLINING
-PRAGMA_DISABLE_OPTIMIZATION
-
 #if WITH_EDITOR
 const FName UNiagaraEmitter::PrivateMemberNames::EventHandlerScriptProps = GET_MEMBER_NAME_CHECKED(UNiagaraEmitter, EventHandlerScriptProps);
 
@@ -1536,6 +1533,3 @@ void UNiagaraEmitter::InitFastPathAttributeNames()
 	InitParameters(SpawnScriptProps.Script->GetVMExecutableData().Parameters, UniqueEmitterName, SpawnFastPathAttributeNames);
 	InitParameters(UpdateScriptProps.Script->GetVMExecutableData().Parameters, UniqueEmitterName, UpdateFastPathAttributeNames);
 }
-
-PRAGMA_ENABLE_INLINING
-PRAGMA_ENABLE_OPTIMIZATION
\ No newline at end of file
diff --git a/Engine/Plugins/FX/Niagara/Source/Niagara/Private/NiagaraSystemInstance.cpp b/Engine/Plugins/FX/Niagara/Source/Niagara/Private/NiagaraSystemInstance.cpp
--- a/Engine/Plugins/FX/Niagara/Source/Niagara/Private/NiagaraSystemInstance.cpp
+++ b/Engine/Plugins/FX/Niagara/Source/Niagara/Private/NiagaraSystemInstance.cpp
@@ -15,9 +15,6 @@
 #include "GameFramework/PlayerController.h"
 
 
-PRAGMA_DISABLE_INLINING
-PRAGMA_DISABLE_OPTIMIZATION
-
 DECLARE_CYCLE_STAT(TEXT("System Activate [GT]"), STAT_NiagaraSystemActivate, STATGROUP_Niagara);
 DECLARE_CYCLE_STAT(TEXT("System Deactivate [GT]"), STAT_NiagaraSystemDeactivate, STATGROUP_Niagara);
 DECLARE_CYCLE_STAT(TEXT("System Complete [GT]"), STAT_NiagaraSystemComplete, STATGROUP_Niagara);
@@ -2236,6 +2233,3 @@ FNiagaraSystemInstance::FOnDestroyed& FNiagaraSystemInstance::OnDestroyed()
 	return OnDestroyedDelegate;
 }
 #endif
-
-PRAGMA_ENABLE_INLINING
-PRAGMA_ENABLE_OPTIMIZATION
diff --git a/Engine/Plugins/FX/Niagara/Source/Niagara/Public/NiagaraComponent.h b/Engine/Plugins/FX/Niagara/Source/Niagara/Public/NiagaraComponent.h
--- a/Engine/Plugins/FX/Niagara/Source/Niagara/Public/NiagaraComponent.h
+++ b/Engine/Plugins/FX/Niagara/Source/Niagara/Public/NiagaraComponent.h
@@ -154,8 +154,6 @@ public:
 
 	virtual void Activate(bool bReset = false)override;
 	virtual void Deactivate()override;
-
-	UFUNCTION(BlueprintCallable, Category = Niagara)
 	void DeactivateImmediate();
 
 	FORCEINLINE ENiagaraExecutionState GetRequestedExecutionState()const { return SystemInstance ? SystemInstance->GetRequestedExecutionState() : ENiagaraExecutionState::Complete; }
