Early out in the impossible case of a NULL readback request to fix static analysis error

diff --git a/Engine/Plugins/FX/Niagara/Source/NiagaraEditor/Private/GlobalDistanceFieldCapture.cpp b/Engine/Plugins/FX/Niagara/Source/NiagaraEditor/Private/GlobalDistanceFieldCapture.cpp
--- a/Engine/Plugins/FX/Niagara/Source/NiagaraEditor/Private/GlobalDistanceFieldCapture.cpp
+++ b/Engine/Plugins/FX/Niagara/Source/NiagaraEditor/Private/GlobalDistanceFieldCapture.cpp
@@ -117,7 +117,10 @@ FGlobalDistanceFieldCapture::~FGlobalDistanceFieldCapture()
 
 void FGlobalDistanceFieldCapture::OnReadbackComplete()
 {
-	ensure(Readback);
+	if (!ensureMsgf(Readback, TEXT("FGlobalDistanceFieldCapture::OnReadbackComplete called without a pending request!")))
+	{
+		return;
+	}
 
 	float MaxDist = -MAX_FLT;
 	float MinDist = MAX_FLT;
