- Updated some NiagaraTypeCustomization code to FProperty

diff --git a/Engine/Plugins/FX/Niagara/Source/NiagaraEditor/Private/Customizations/NiagaraTypeCustomizations.cpp b/Engine/Plugins/FX/Niagara/Source/NiagaraEditor/Private/Customizations/NiagaraTypeCustomizations.cpp
--- a/Engine/Plugins/FX/Niagara/Source/NiagaraEditor/Private/Customizations/NiagaraTypeCustomizations.cpp
+++ b/Engine/Plugins/FX/Niagara/Source/NiagaraEditor/Private/Customizations/NiagaraTypeCustomizations.cpp
@@ -743,7 +743,7 @@ void FNiagaraPlatformSetTypeCustomization::CustomizeHeader(TSharedRef<IPropertyH
 			PlatformSetArrayIndex = ThisIndex;
 		}
 
-		if (UProperty* CurrProperty = CurrHandle->GetProperty())
+		if (FProperty* CurrProperty = CurrHandle->GetProperty())
 		{
 			if (UStruct* CurrStruct = CurrProperty->GetOwnerStruct())
 			{
@@ -941,7 +941,7 @@ void FNiagaraPlatformSetTypeCustomization::UpdateCachedConflicts()
 
 static TSharedPtr<IPropertyHandle> FindChildPlatformSet(TSharedPtr<IPropertyHandle> PropertyHandle)
 {
-	if (UStructProperty* Property = Cast<UStructProperty>(PropertyHandle->GetProperty()))
+	if (FStructProperty* Property = CastField<FStructProperty>(PropertyHandle->GetProperty()))
 	{
 		if (Property->Struct == FNiagaraPlatformSet::StaticStruct())
 		{
