Modules can now be renamed from the Niagara stack panel. The module will still display its type in parentheses when renamed, and this does not affect compilation in any way.

CanRenameInput and its ilk moved from NiagaraStackFunctionInput to NiagaraStackEntry, all stack entries now use NiagaraStackEditorData::GetStackEntryIsRenamePending. Search text highlighting moved to SNiagaraStackEntryWidget.

diff --git a/Engine/Plugins/FX/Niagara/Source/Niagara/Classes/NiagaraEmitter.h b/Engine/Plugins/FX/Niagara/Source/Niagara/Classes/NiagaraEmitter.h
--- a/Engine/Plugins/FX/Niagara/Source/Niagara/Classes/NiagaraEmitter.h
+++ b/Engine/Plugins/FX/Niagara/Source/Niagara/Classes/NiagaraEmitter.h
@@ -18,6 +18,7 @@ class UNiagaraEmitter;
 class UNiagaraEventReceiverEmitterAction;
 class UNiagaraRendererProperties;
 class UNiagaraShaderStageBase;
+class UNiagaraEditorDataBase;
 
 //TODO: Event action that spawns other whole Systems?
 //One that calls a BP exposed delegate?
@@ -376,10 +377,10 @@ public:
 
 	/* Gets a Guid which is updated any time data in this emitter is changed. */
 	FGuid NIAGARA_API GetChangeId() const;
-	
-	/** Data used by the editor to maintain UI state etc.. */
-	UPROPERTY()
-	UObject* EditorData;
+
+	NIAGARA_API UNiagaraEditorDataBase* GetEditorData() const;
+
+	NIAGARA_API void SetEditorData(UNiagaraEditorDataBase* InEditorData);
 
 	/** Internal: The thumbnail image.*/
 	UPROPERTY()
@@ -525,11 +526,17 @@ private:
 
 	void GraphSourceChanged();
 
+	void PersistentEditorDataChanged();
+
 private:
 	/** Adjusted every time that we compile this emitter. Lets us know that we might differ from any cached versions.*/
 	UPROPERTY()
 	FGuid ChangeId;
 
+	/** Data used by the editor to maintain UI state etc.. */
+	UPROPERTY()
+	UNiagaraEditorDataBase* EditorData;
+
 	/** A multicast delegate which is called whenever all the scripts for this emitter have been compiled (successfully or not). */
 	FOnEmitterCompiled OnVMScriptCompiledDelegate;
 #endif
diff --git a/Engine/Plugins/FX/Niagara/Source/Niagara/Private/NiagaraEmitter.cpp b/Engine/Plugins/FX/Niagara/Source/Niagara/Private/NiagaraEmitter.cpp
--- a/Engine/Plugins/FX/Niagara/Source/Niagara/Private/NiagaraEmitter.cpp
+++ b/Engine/Plugins/FX/Niagara/Source/Niagara/Private/NiagaraEmitter.cpp
@@ -15,6 +15,7 @@
 #include "Modules/ModuleManager.h"
 #include "Misc/ConfigCacheIni.h"
 #include "Interfaces/ITargetPlatform.h"
+#include "NiagaraEditorDataBase.h"
 
 #if WITH_EDITOR
 const FName UNiagaraEmitter::PrivateMemberNames::EventHandlerScriptProps = GET_MEMBER_NAME_CHECKED(UNiagaraEmitter, EventHandlerScriptProps);
@@ -515,6 +516,11 @@ void UNiagaraEmitter::PostLoad()
 		{
 			Renderer->OnChanged().AddUObject(this, &UNiagaraEmitter::RendererChanged);
 		}
+
+		if (EditorData != nullptr)
+		{
+			EditorData->OnPersistentDataChanged().AddUObject(this, &UNiagaraEmitter::PersistentEditorDataChanged);
+		}
 	}
 #endif
 
@@ -818,6 +824,26 @@ FGuid UNiagaraEmitter::GetChangeId() const
 	return ChangeId;
 }
 
+UNiagaraEditorDataBase* UNiagaraEmitter::GetEditorData() const
+{
+	return EditorData;
+}
+
+void UNiagaraEmitter::SetEditorData(UNiagaraEditorDataBase* InEditorData)
+{
+	if (EditorData != nullptr)
+	{
+		EditorData->OnPersistentDataChanged().RemoveAll(this);
+	}
+
+	EditorData = InEditorData;
+	
+	if (EditorData != nullptr)
+	{
+		EditorData->OnPersistentDataChanged().AddUObject(this, &UNiagaraEmitter::PersistentEditorDataChanged);
+	}
+}
+
 bool UNiagaraEmitter::AreAllScriptAndSourcesSynchronized() const
 {
 	if (SpawnScriptProps.Script->IsCompilable() && !SpawnScriptProps.Script->AreScriptAndSourceSynchronized())
@@ -1132,6 +1158,8 @@ void UNiagaraEmitter::UpdateFromMergedCopy(const INiagaraMergeManager& MergeMana
 		MergedRenderer->OnChanged().AddUObject(this, &UNiagaraEmitter::RendererChanged);
 	}
 
+	SetEditorData(MergedEmitter->GetEditorData());
+
 	// Update the change id since we don't know what's changed.
 	UpdateChangeId(TEXT("Updated from merged copy"));
 }
@@ -1353,6 +1381,11 @@ void UNiagaraEmitter::GraphSourceChanged()
 {
 	UpdateChangeId(TEXT("Graph source changed."));
 }
+
+void UNiagaraEmitter::PersistentEditorDataChanged()
+{
+	UpdateChangeId(TEXT("Persistent editor data changed."));
+}
 #endif
 
 TStatId UNiagaraEmitter::GetStatID(bool bGameThread, bool bConcurrent)const
diff --git a/Engine/Plugins/FX/Niagara/Source/Niagara/Public/NiagaraEditorDataBase.h b/Engine/Plugins/FX/Niagara/Source/Niagara/Public/NiagaraEditorDataBase.h
--- a/Engine/Plugins/FX/Niagara/Source/Niagara/Public/NiagaraEditorDataBase.h
+++ b/Engine/Plugins/FX/Niagara/Source/Niagara/Public/NiagaraEditorDataBase.h
@@ -49,6 +49,11 @@ class UNiagaraEditorDataBase : public UObject
 public:
 #if WITH_EDITORONLY_DATA
 	virtual void PostLoadFromOwner(UObject* InOwner) PURE_VIRTUAL(UNiagaraEditorDataBase::PostLoadFromOwner, );
+
+	NIAGARA_API FSimpleMulticastDelegate& OnPersistentDataChanged() { return PersistentDataChangedDelegate; }
+
+private:
+	FSimpleMulticastDelegate PersistentDataChangedDelegate;
 #endif
 };
 
diff --git a/Engine/Plugins/FX/Niagara/Source/NiagaraEditor/Private/NiagaraEditorUtilities.cpp b/Engine/Plugins/FX/Niagara/Source/NiagaraEditor/Private/NiagaraEditorUtilities.cpp
--- a/Engine/Plugins/FX/Niagara/Source/NiagaraEditor/Private/NiagaraEditorUtilities.cpp
+++ b/Engine/Plugins/FX/Niagara/Source/NiagaraEditor/Private/NiagaraEditorUtilities.cpp
@@ -1594,7 +1594,7 @@ bool FNiagaraEditorUtilities::AddParameter(FNiagaraVariable& NewParameterVariabl
 	bool bSuccess = TargetParameterStore.AddParameter(NewParameterVariable);
 	if (bSuccess)
 	{
-		StackEditorData.SetModuleInputIsRenamePending(NewParameterVariable.GetName().ToString(), true);
+		StackEditorData.SetStackEntryIsRenamePending(NewParameterVariable.GetName().ToString(), true);
 	}
 	return bSuccess;
 }
diff --git a/Engine/Plugins/FX/Niagara/Source/NiagaraEditor/Private/NiagaraEmitterEditorData.cpp b/Engine/Plugins/FX/Niagara/Source/NiagaraEditor/Private/NiagaraEmitterEditorData.cpp
--- a/Engine/Plugins/FX/Niagara/Source/NiagaraEditor/Private/NiagaraEmitterEditorData.cpp
+++ b/Engine/Plugins/FX/Niagara/Source/NiagaraEditor/Private/NiagaraEmitterEditorData.cpp
@@ -6,6 +6,12 @@
 UNiagaraEmitterEditorData::UNiagaraEmitterEditorData(const FObjectInitializer& ObjectInitializer)
 {
 	StackEditorData = ObjectInitializer.CreateDefaultSubobject<UNiagaraStackEditorData>(this, TEXT("StackEditorData"));
+
+	if (StackEditorData != nullptr)
+	{
+		StackEditorData->OnPersistentDataChanged().AddUObject(this, &UNiagaraEmitterEditorData::StackEditorDataChanged);
+	}
+	
 	PlaybackRangeMin = 0;
 	PlaybackRangeMax = 10;
 }
@@ -16,6 +22,7 @@ void UNiagaraEmitterEditorData::PostLoad()
 	if (StackEditorData == nullptr)
 	{
 		StackEditorData = NewObject<UNiagaraStackEditorData>(this, TEXT("StackEditorData"), RF_Transactional);
+		StackEditorData->OnPersistentDataChanged().AddUObject(this, &UNiagaraEmitterEditorData::StackEditorDataChanged);
 	}
 }
 
@@ -33,4 +40,11 @@ void UNiagaraEmitterEditorData::SetPlaybackRange(TRange<float> InPlaybackRange)
 {
 	PlaybackRangeMin = InPlaybackRange.GetLowerBoundValue();
 	PlaybackRangeMax = InPlaybackRange.GetUpperBoundValue();
+
+	OnPersistentDataChanged().Broadcast();
+}
+
+void UNiagaraEmitterEditorData::StackEditorDataChanged()
+{
+	OnPersistentDataChanged().Broadcast();
 }
\ No newline at end of file
diff --git a/Engine/Plugins/FX/Niagara/Source/NiagaraEditor/Private/NiagaraEmitterEditorData.h b/Engine/Plugins/FX/Niagara/Source/NiagaraEditor/Private/NiagaraEmitterEditorData.h
--- a/Engine/Plugins/FX/Niagara/Source/NiagaraEditor/Private/NiagaraEmitterEditorData.h
+++ b/Engine/Plugins/FX/Niagara/Source/NiagaraEditor/Private/NiagaraEmitterEditorData.h
@@ -2,13 +2,15 @@
 
 #pragma once
 
+#include "NiagaraEditorDataBase.h"
+
 #include "NiagaraEmitterEditorData.generated.h"
 
 class UNiagaraStackEditorData;
 
 /** Editor only UI data for emitters. */
 UCLASS()
-class UNiagaraEmitterEditorData : public UObject
+class UNiagaraEmitterEditorData : public UNiagaraEditorDataBase
 {
 	GENERATED_BODY()
 
@@ -32,4 +34,6 @@ private:
 
 	UPROPERTY()
 	float PlaybackRangeMax;
+
+	void StackEditorDataChanged();
 };
\ No newline at end of file
diff --git a/Engine/Plugins/FX/Niagara/Source/NiagaraEditor/Private/NiagaraScriptMergeManager.cpp b/Engine/Plugins/FX/Niagara/Source/NiagaraEditor/Private/NiagaraScriptMergeManager.cpp
--- a/Engine/Plugins/FX/Niagara/Source/NiagaraEditor/Private/NiagaraScriptMergeManager.cpp
+++ b/Engine/Plugins/FX/Niagara/Source/NiagaraEditor/Private/NiagaraScriptMergeManager.cpp
@@ -20,6 +20,7 @@
 #include "ViewModels/Stack/NiagaraStackGraphUtilities.h"
 #include "NiagaraEditorModule.h"
 #include "NiagaraEmitterEditorData.h"
+#include "NiagaraStackEditorData.h"
 
 #include "UObject/PropertyPortFlags.h"
 
@@ -615,6 +616,8 @@ void FNiagaraEmitterMergeAdapter::Initialize(const UNiagaraEmitter& InEmitter, U
 	{
 		Renderers.Add(MakeShared<FNiagaraRendererMergeAdapter>(*RendererProperties));
 	}
+
+	EditorData = Cast<const UNiagaraEmitterEditorData>(Emitter->GetEditorData());
 }
 
 UNiagaraEmitter* FNiagaraEmitterMergeAdapter::GetEditableEmitter() const
@@ -657,6 +660,11 @@ const TArray<TSharedRef<FNiagaraRendererMergeAdapter>> FNiagaraEmitterMergeAdapt
 	return Renderers;
 }
 
+const UNiagaraEmitterEditorData* FNiagaraEmitterMergeAdapter::GetEditorData() const
+{
+	return EditorData.Get();
+}
+
 TSharedPtr<FNiagaraScriptStackMergeAdapter> FNiagaraEmitterMergeAdapter::GetScriptStack(ENiagaraScriptUsage Usage, FGuid ScriptUsageId)
 {
 	switch (Usage)
@@ -808,7 +816,8 @@ bool FNiagaraEmitterDiffResults::IsEmpty() const
 		RemovedBaseRenderers.Num() == 0 &&
 		AddedOtherRenderers.Num() == 0 &&
 		ModifiedBaseRenderers.Num() == 0 &&
-		ModifiedOtherRenderers.Num() == 0;
+		ModifiedOtherRenderers.Num() == 0 &&
+		ModifiedStackEntryDisplayNames.Num() == 0;
 }
 
 void FNiagaraEmitterDiffResults::AddError(FText ErrorMessage)
@@ -1094,11 +1103,14 @@ INiagaraMergeManager::FMergeEmitterResults FNiagaraScriptMergeManager::MergeEmit
 		MergeResults.ErrorMessages.Append(RendererResults.ErrorMessages);
 
 		CopyPropertiesToBase(MergedInstance, &Instance, DiffResults.DifferentEmitterProperties);
-		if (Instance.EditorData != nullptr)
+
+		FApplyDiffResults StackEntryDisplayNameDiffs = ApplyStackEntryDisplayNameDiffs(*MergedInstance, DiffResults);
+		if (StackEntryDisplayNameDiffs.bSucceeded == false)
 		{
-			// We keep the instance editor data here so that the UI state in the system stays more consistent.
-			MergedInstance->EditorData = Cast<UNiagaraEmitterEditorData>(StaticDuplicateObject(Instance.EditorData, MergedInstance));
+			MergeResults.MergeResult = INiagaraMergeManager::EMergeEmitterResult::FailedToMerge;
 		}
+		MergeResults.bModifiedGraph |= StackEntryDisplayNameDiffs.bModifiedGraph;
+		MergeResults.ErrorMessages.Append(StackEntryDisplayNameDiffs.ErrorMessages);
 
 #if 0
 		UE_LOG(LogNiagaraEditor, Log, TEXT("A"));
@@ -1479,6 +1491,7 @@ FNiagaraEmitterDiffResults FNiagaraScriptMergeManager::DiffEmitters(UNiagaraEmit
 	DiffShaderStages(BaseEmitterAdapter->GetShaderStages(), OtherEmitterAdapter->GetShaderStages(), EmitterDiffResults);
 	DiffRenderers(BaseEmitterAdapter->GetRenderers(), OtherEmitterAdapter->GetRenderers(), EmitterDiffResults);
 	DiffEditableProperties(&BaseEmitter, &OtherEmitter, *UNiagaraEmitter::StaticClass(), EmitterDiffResults.DifferentEmitterProperties);
+	DiffStackEntryDisplayNames(BaseEmitterAdapter->GetEditorData(), OtherEmitterAdapter->GetEditorData(), EmitterDiffResults.ModifiedStackEntryDisplayNames);
 
 	return EmitterDiffResults;
 }
@@ -1779,6 +1792,23 @@ void FNiagaraScriptMergeManager::DiffEditableProperties(const void* BaseDataAddr
 	}
 }
 
+void FNiagaraScriptMergeManager::DiffStackEntryDisplayNames(const UNiagaraEmitterEditorData* BaseEditorData, const UNiagaraEmitterEditorData* OtherEditorData, TMap<FString, FText>& OutModifiedStackEntryDisplayNames) const
+{
+	if (BaseEditorData != nullptr && OtherEditorData != nullptr)
+	{
+		// find display names that have been added or changed in the instance
+		const TMap<FString, FText>& OtherRenames = OtherEditorData->GetStackEditorData().GetAllStackEntryDisplayNames();
+		for (auto& Pair : OtherRenames)
+		{
+			const FText* BaseDisplayName = BaseEditorData->GetStackEditorData().GetStackEntryDisplayName(Pair.Key);
+			if (BaseDisplayName == nullptr || !BaseDisplayName->EqualTo(Pair.Value))
+			{
+				OutModifiedStackEntryDisplayNames.Add(Pair.Key, Pair.Value);
+			}
+		}
+	}
+}
+
 TOptional<bool> FNiagaraScriptMergeManager::DoFunctionInputOverridesMatch(TSharedRef<FNiagaraStackFunctionInputOverrideMergeAdapter> BaseFunctionInputAdapter, TSharedRef<FNiagaraStackFunctionInputOverrideMergeAdapter> OtherFunctionInputAdapter) const
 {
 	// Local String Value.
@@ -2538,6 +2568,29 @@ FNiagaraScriptMergeManager::FApplyDiffResults FNiagaraScriptMergeManager::ApplyR
 	return Results;
 }
 
+FNiagaraScriptMergeManager::FApplyDiffResults FNiagaraScriptMergeManager::ApplyStackEntryDisplayNameDiffs(UNiagaraEmitter& Emitter, const FNiagaraEmitterDiffResults& DiffResults) const
+{
+	if (DiffResults.ModifiedStackEntryDisplayNames.Num() > 0)
+	{
+		UNiagaraEmitterEditorData* EditorData = Cast<UNiagaraEmitterEditorData>(Emitter.GetEditorData());
+		if (EditorData == nullptr)
+		{
+			EditorData = NewObject<UNiagaraEmitterEditorData>(&Emitter, NAME_None, RF_Transactional);
+			Emitter.SetEditorData(EditorData);
+		}
+
+		for (auto& Pair : DiffResults.ModifiedStackEntryDisplayNames)
+		{
+			EditorData->GetStackEditorData().SetStackEntryDisplayName(Pair.Key, Pair.Value);
+		}
+	}
+
+	FApplyDiffResults Results;
+	Results.bSucceeded = true;
+	Results.bModifiedGraph = false;
+	return Results;
+}
+
 TSharedRef<FNiagaraEmitterMergeAdapter> FNiagaraScriptMergeManager::GetEmitterMergeAdapterUsingCache(const UNiagaraEmitter& Emitter)
 {
 	FCachedMergeAdapter* CachedMergeAdapter = CachedMergeAdapters.Find(FObjectKey(&Emitter));
diff --git a/Engine/Plugins/FX/Niagara/Source/NiagaraEditor/Private/NiagaraScriptMergeManager.h b/Engine/Plugins/FX/Niagara/Source/NiagaraEditor/Private/NiagaraScriptMergeManager.h
--- a/Engine/Plugins/FX/Niagara/Source/NiagaraEditor/Private/NiagaraScriptMergeManager.h
+++ b/Engine/Plugins/FX/Niagara/Source/NiagaraEditor/Private/NiagaraScriptMergeManager.h
@@ -14,6 +14,7 @@
 #include "UObject/ObjectKey.h"
 
 class UNiagaraEmitter;
+class UNiagaraEmitterEditorData;
 class UNiagaraScript;
 class UNiagaraNodeOutput;
 class UNiagaraNodeInput;
@@ -211,6 +212,7 @@ public:
 	const TArray<TSharedRef<FNiagaraEventHandlerMergeAdapter>> GetEventHandlers() const;
 	const TArray<TSharedRef<FNiagaraShaderStageMergeAdapter>> GetShaderStages() const;
 	const TArray<TSharedRef<FNiagaraRendererMergeAdapter>> GetRenderers() const;
+	const UNiagaraEmitterEditorData* GetEditorData() const;
 
 	TSharedPtr<FNiagaraScriptStackMergeAdapter> GetScriptStack(ENiagaraScriptUsage Usage, FGuid UsageId);
 
@@ -233,6 +235,7 @@ private:
 	TArray<TSharedRef<FNiagaraEventHandlerMergeAdapter>> EventHandlers;
 	TArray<TSharedRef<FNiagaraShaderStageMergeAdapter>> ShaderStages;
 	TArray<TSharedRef<FNiagaraRendererMergeAdapter>> Renderers;
+	TWeakObjectPtr<const UNiagaraEmitterEditorData> EditorData;
 };
 
 struct FNiagaraScriptStackDiffResults
@@ -319,6 +322,8 @@ struct FNiagaraEmitterDiffResults
 	TArray<TSharedRef<FNiagaraRendererMergeAdapter>> ModifiedBaseRenderers;
 	TArray<TSharedRef<FNiagaraRendererMergeAdapter>> ModifiedOtherRenderers;
 
+	TMap<FString, FText> ModifiedStackEntryDisplayNames;
+
 private:
 	bool bIsValid;
 	TArray<FText> ErrorMessages;
@@ -389,8 +394,11 @@ public:
 
 	virtual void DiffEditableProperties(const void* BaseDataAddress, const void* OtherDataAddress, UStruct& Struct, TArray<FProperty*>& OutDifferentProperties) const override;
 
+	void DiffStackEntryDisplayNames(const UNiagaraEmitterEditorData* BaseEditorData, const UNiagaraEmitterEditorData* OtherEditorData, TMap<FString, FText>& OutModifiedStackEntryDisplayNames) const;
+
 	virtual void CopyPropertiesToBase(void* BaseDataAddress, const void* OtherDataAddress, TArray<FProperty*> PropertiesToCopy) const override;
 
+
 private:
 	TOptional<bool> DoFunctionInputOverridesMatch(TSharedRef<FNiagaraStackFunctionInputOverrideMergeAdapter> BaseFunctionInputAdapter, TSharedRef<FNiagaraStackFunctionInputOverrideMergeAdapter> OtherFunctionInputAdapter) const;
 
@@ -402,6 +410,8 @@ private:
 
 	FApplyDiffResults ApplyRendererDiff(UNiagaraEmitter& BaseEmitter, const FNiagaraEmitterDiffResults& DiffResults) const;
 
+	FApplyDiffResults ApplyStackEntryDisplayNameDiffs(UNiagaraEmitter& Emitter, const FNiagaraEmitterDiffResults& DiffResults) const;
+
 	FApplyDiffResults AddModule(FString UniqueEmitterName, UNiagaraScript& OwningScript, UNiagaraNodeOutput& TargetOutputNode, TSharedRef<FNiagaraStackFunctionMergeAdapter> AddModule) const;
 
 	FApplyDiffResults RemoveInputOverride(UNiagaraScript& OwningScript, TSharedRef<FNiagaraStackFunctionInputOverrideMergeAdapter> OverrideToRemove) const;
diff --git a/Engine/Plugins/FX/Niagara/Source/NiagaraEditor/Private/NiagaraStackEditorData.cpp b/Engine/Plugins/FX/Niagara/Source/NiagaraEditor/Private/NiagaraStackEditorData.cpp
--- a/Engine/Plugins/FX/Niagara/Source/NiagaraEditor/Private/NiagaraStackEditorData.cpp
+++ b/Engine/Plugins/FX/Niagara/Source/NiagaraEditor/Private/NiagaraStackEditorData.cpp
@@ -2,15 +2,15 @@
 
 #include "NiagaraStackEditorData.h"
 
-bool UNiagaraStackEditorData::GetModuleInputIsRenamePending(const FString& ModuleInputKey) const
+bool UNiagaraStackEditorData::GetStackEntryIsRenamePending(const FString& StackEntryKey) const
 {
-	const bool* bIsRenamePendingPtr = ModuleInputKeyToRenamePendingMap.Find(ModuleInputKey);
+	const bool* bIsRenamePendingPtr = StackEntryKeyToRenamePendingMap.Find(StackEntryKey);
 	return bIsRenamePendingPtr != nullptr && *bIsRenamePendingPtr;
 }
 
-void UNiagaraStackEditorData::SetModuleInputIsRenamePending(const FString& ModuleInputKey, bool bIsRenamePending)
+void UNiagaraStackEditorData::SetStackEntryIsRenamePending(const FString& StackEntryKey, bool bIsRenamePending)
 {
-	ModuleInputKeyToRenamePendingMap.FindOrAdd(ModuleInputKey) = bIsRenamePending;
+	StackEntryKeyToRenamePendingMap.FindOrAdd(StackEntryKey) = bIsRenamePending;
 }
 
 bool UNiagaraStackEditorData::GetStackEntryIsExpanded(const FString& StackEntryKey, bool bIsExpandedDefault) const
@@ -55,6 +55,33 @@ void UNiagaraStackEditorData::SetStackItemShowAdvanced(const FString& StackEntry
 	}
 }
 
+const FText* UNiagaraStackEditorData::GetStackEntryDisplayName(const FString& StackEntryKey) const
+{
+	return StackEntryKeyToDisplayName.Find(StackEntryKey);
+}
+
+void UNiagaraStackEditorData::SetStackEntryDisplayName(const FString& StackEntryKey, const FText& InDisplayName)
+{
+	bool bBroadcast = false;
+
+	if (InDisplayName.IsEmptyOrWhitespace())
+	{
+		// we assume here that the display name has changed
+		StackEntryKeyToDisplayName.Remove(StackEntryKey);
+		bBroadcast = true;
+	}
+	else if (ensureMsgf(StackEntryKey.IsEmpty() == false, TEXT("Can not set the display name with an empty key")))
+	{
+		StackEntryKeyToDisplayName.FindOrAdd(StackEntryKey) = InDisplayName;
+		bBroadcast = true;
+	}
+
+	if (bBroadcast)
+	{
+		OnPersistentDataChanged().Broadcast();
+	}
+}
+
 bool UNiagaraStackEditorData::GetShowAllAdvanced() const
 {
 	return bShowAllAdvanced;
diff --git a/Engine/Plugins/FX/Niagara/Source/NiagaraEditor/Private/NiagaraSystemEditorData.cpp b/Engine/Plugins/FX/Niagara/Source/NiagaraEditor/Private/NiagaraSystemEditorData.cpp
--- a/Engine/Plugins/FX/Niagara/Source/NiagaraEditor/Private/NiagaraSystemEditorData.cpp
+++ b/Engine/Plugins/FX/Niagara/Source/NiagaraEditor/Private/NiagaraSystemEditorData.cpp
@@ -174,7 +174,7 @@ void UNiagaraSystemEditorData::UpdatePlaybackRangeFromEmitters(UNiagaraSystem& O
 
 		for (const FNiagaraEmitterHandle& EmitterHandle : OwnerSystem.GetEmitterHandles())
 		{
-			UNiagaraEmitterEditorData* EmitterEditorData = Cast<UNiagaraEmitterEditorData>(EmitterHandle.GetInstance()->EditorData);
+			UNiagaraEmitterEditorData* EmitterEditorData = Cast<UNiagaraEmitterEditorData>(EmitterHandle.GetInstance()->GetEditorData());
 			if (EmitterEditorData != nullptr)
 			{
 				EmitterPlaybackRangeMin = FMath::Min(PlaybackRangeMin, EmitterEditorData->GetPlaybackRange().GetLowerBoundValue());
diff --git a/Engine/Plugins/FX/Niagara/Source/NiagaraEditor/Private/ViewModels/NiagaraEmitterViewModel.cpp b/Engine/Plugins/FX/Niagara/Source/NiagaraEditor/Private/ViewModels/NiagaraEmitterViewModel.cpp
--- a/Engine/Plugins/FX/Niagara/Source/NiagaraEditor/Private/ViewModels/NiagaraEmitterViewModel.cpp
+++ b/Engine/Plugins/FX/Niagara/Source/NiagaraEditor/Private/ViewModels/NiagaraEmitterViewModel.cpp
@@ -239,7 +239,7 @@ const UNiagaraEmitterEditorData& FNiagaraEmitterViewModel::GetEditorData() const
 {
 	check(Emitter.IsValid());
 
-	const UNiagaraEmitterEditorData* EditorData = Cast<UNiagaraEmitterEditorData>(Emitter->EditorData);
+	const UNiagaraEmitterEditorData* EditorData = Cast<UNiagaraEmitterEditorData>(Emitter->GetEditorData());
 	if (EditorData == nullptr)
 	{
 		EditorData = GetDefault<UNiagaraEmitterEditorData>();
@@ -250,12 +250,12 @@ const UNiagaraEmitterEditorData& FNiagaraEmitterViewModel::GetEditorData() const
 UNiagaraEmitterEditorData& FNiagaraEmitterViewModel::GetOrCreateEditorData()
 {
 	check(Emitter.IsValid());
-	UNiagaraEmitterEditorData* EditorData = Cast<UNiagaraEmitterEditorData>(Emitter->EditorData);
+	UNiagaraEmitterEditorData* EditorData = Cast<UNiagaraEmitterEditorData>(Emitter->GetEditorData());
 	if (EditorData == nullptr)
 	{
 		EditorData = NewObject<UNiagaraEmitterEditorData>(Emitter.Get(), NAME_None, RF_Transactional);
-		Emitter->Modify();
-		Emitter->EditorData = EditorData;
+		Emitter->Modify(); 
+		Emitter->SetEditorData(EditorData);
 	}
 	return *EditorData;
 }
diff --git a/Engine/Plugins/FX/Niagara/Source/NiagaraEditor/Private/ViewModels/Stack/NiagaraStackEntry.cpp b/Engine/Plugins/FX/Niagara/Source/NiagaraEditor/Private/ViewModels/Stack/NiagaraStackEntry.cpp
--- a/Engine/Plugins/FX/Niagara/Source/NiagaraEditor/Private/ViewModels/Stack/NiagaraStackEntry.cpp
+++ b/Engine/Plugins/FX/Niagara/Source/NiagaraEditor/Private/ViewModels/Stack/NiagaraStackEntry.cpp
@@ -8,6 +8,7 @@
 #include "NiagaraStackEditorData.h"
 #include "NiagaraScriptMergeManager.h"
 #include "Misc/SecureHash.h"
+#include "ScopedTransaction.h"
 
 const FName UNiagaraStackEntry::FExecutionCategoryNames::System = TEXT("System");
 const FName UNiagaraStackEntry::FExecutionCategoryNames::Emitter = TEXT("Emitter");
@@ -178,6 +179,11 @@ FText UNiagaraStackEntry::GetDisplayName() const
 	return FText();
 }
 
+FText UNiagaraStackEntry::GetOriginalName() const
+{
+	return FText();
+}
+
 UObject* UNiagaraStackEntry::GetDisplayedObject() const
 {
 	return nullptr;
@@ -737,3 +743,30 @@ TOptional<UNiagaraStackEntry::FDropRequestResponse> UNiagaraStackEntry::ChildReq
 			: TOptional<FDropRequestResponse>();
 	}
 }
+
+bool UNiagaraStackEntry::GetIsRenamePending() const
+{
+	return SupportsRename() && GetStackEditorData().GetStackEntryIsRenamePending(StackEditorDataKey);
+}
+
+void UNiagaraStackEntry::SetIsRenamePending(bool bIsRenamePending)
+{
+	if (SupportsRename())
+	{
+		GetStackEditorData().SetStackEntryIsRenamePending(StackEditorDataKey, bIsRenamePending);
+	}
+}
+
+void UNiagaraStackEntry::OnRenamed(FText NewName)
+{
+	if (SupportsRename())
+	{
+		if (!NewName.EqualTo(GetDisplayName()))
+		{
+			FScopedTransaction ScopedTransaction(NSLOCTEXT("NiagaraStackEntry", "RenameModule", "Rename Module"));
+
+			GetStackEditorData().Modify();
+			GetStackEditorData().SetStackEntryDisplayName(GetStackEditorDataKey(), NewName);
+		}
+	}
+}
diff --git a/Engine/Plugins/FX/Niagara/Source/NiagaraEditor/Private/ViewModels/Stack/NiagaraStackFunctionInput.cpp b/Engine/Plugins/FX/Niagara/Source/NiagaraEditor/Private/ViewModels/Stack/NiagaraStackFunctionInput.cpp
--- a/Engine/Plugins/FX/Niagara/Source/NiagaraEditor/Private/ViewModels/Stack/NiagaraStackFunctionInput.cpp
+++ b/Engine/Plugins/FX/Niagara/Source/NiagaraEditor/Private/ViewModels/Stack/NiagaraStackFunctionInput.cpp
@@ -1433,28 +1433,16 @@ FNiagaraVariable UNiagaraStackFunctionInput::CreateRapidIterationVariable(const
 	return FNiagaraStackGraphUtilities::CreateRapidIterationParameter(UniqueEmitterName, OutputNode->GetUsage(), InName, InputType);
 }
 
-bool UNiagaraStackFunctionInput::CanRenameInput() const
+bool UNiagaraStackFunctionInput::SupportsRename() const
 {
 	// Only module level assignment node inputs can be renamed.
 	return OwningAssignmentNode.IsValid() && InputParameterHandlePath.Num() == 1 &&
 		OwningAssignmentNode->FindAssignmentTarget(InputParameterHandle.GetName()) != INDEX_NONE;
 }
 
-bool UNiagaraStackFunctionInput::GetIsRenamePending() const
-{
-	return CanRenameInput() && GetStackEditorData().GetModuleInputIsRenamePending(StackEditorDataKey);
-}
-
-void UNiagaraStackFunctionInput::SetIsRenamePending(bool bIsRenamePending)
-{
-	if (CanRenameInput())
-	{
-		GetStackEditorData().SetModuleInputIsRenamePending(StackEditorDataKey, bIsRenamePending);
-	}
-}
-
-void UNiagaraStackFunctionInput::RenameInput(FName NewName)
+void UNiagaraStackFunctionInput::OnRenamed(FText NewNameText)
 {
+	FName NewName(*NewNameText.ToString());
 	if (OwningAssignmentNode.IsValid() && InputParameterHandlePath.Num() == 1 && InputParameterHandle.GetName() != NewName)
 	{
 		FScopedTransaction ScopedTransaction(LOCTEXT("RenameInput", "Rename this function's input."));
diff --git a/Engine/Plugins/FX/Niagara/Source/NiagaraEditor/Private/ViewModels/Stack/NiagaraStackModuleItem.cpp b/Engine/Plugins/FX/Niagara/Source/NiagaraEditor/Private/ViewModels/Stack/NiagaraStackModuleItem.cpp
--- a/Engine/Plugins/FX/Niagara/Source/NiagaraEditor/Private/ViewModels/Stack/NiagaraStackModuleItem.cpp
+++ b/Engine/Plugins/FX/Niagara/Source/NiagaraEditor/Private/ViewModels/Stack/NiagaraStackModuleItem.cpp
@@ -136,18 +136,47 @@ void UNiagaraStackModuleItem::Initialize(FRequiredEntryData InRequiredEntryData,
 		AddChildFilter(FOnFilterChild::CreateUObject(this, &UNiagaraStackModuleItem::FilterOutputCollection));
 		AddChildFilter(FOnFilterChild::CreateUObject(this, &UNiagaraStackModuleItem::FilterLinkedInputCollection));
 	}
+
+	OnEditorDataChanged();
+	GetStackEditorData().OnPersistentDataChanged().AddUObject(this, &UNiagaraStackModuleItem::OnEditorDataChanged);
+}
+
+void UNiagaraStackModuleItem::OnEditorDataChanged()
+{
+	const FText* NewDisplayName = GetStackEditorData().GetStackEntryDisplayName(GetStackEditorDataKey());
+	if (NewDisplayName == nullptr)
+	{
+		CustomDisplayName.Reset();
+	}
+	else
+	{
+		CustomDisplayName = *NewDisplayName;
+	}
 }
 
 FText UNiagaraStackModuleItem::GetDisplayName() const
+{
+	if (IsFinalized())
+	{
+		return FText::FromName(NAME_None);
+	}
+
+	if (CustomDisplayName.IsSet())
+	{
+		return CustomDisplayName.GetValue();
+	}
+	
+	return GetOriginalName();
+}
+
+FText UNiagaraStackModuleItem::GetOriginalName() const
 {
 	if (FunctionCallNode != nullptr)
 	{
 		return FunctionCallNode->GetNodeTitle(ENodeTitleType::ListView);
 	}
-	else
-	{
+
 	return FText::FromName(NAME_None);
-	}
 }
 
 UObject* UNiagaraStackModuleItem::GetDisplayedObject() const
diff --git a/Engine/Plugins/FX/Niagara/Source/NiagaraEditor/Private/ViewModels/Stack/NiagaraStackParameterStoreEntry.cpp b/Engine/Plugins/FX/Niagara/Source/NiagaraEditor/Private/ViewModels/Stack/NiagaraStackParameterStoreEntry.cpp
--- a/Engine/Plugins/FX/Niagara/Source/NiagaraEditor/Private/ViewModels/Stack/NiagaraStackParameterStoreEntry.cpp
+++ b/Engine/Plugins/FX/Niagara/Source/NiagaraEditor/Private/ViewModels/Stack/NiagaraStackParameterStoreEntry.cpp
@@ -170,24 +170,6 @@ void UNiagaraStackParameterStoreEntry::Reset()
 	GetSystemViewModel()->ResetSystem();
 }
 
-bool UNiagaraStackParameterStoreEntry::CanRenameInput() const
-{
-	return true; 
-}
-
-bool UNiagaraStackParameterStoreEntry::GetIsRenamePending() const
-{
-	return CanRenameInput() && GetStackEditorData().GetModuleInputIsRenamePending(ParameterName.ToString());
-}
-
-void UNiagaraStackParameterStoreEntry::SetIsRenamePending(bool bIsRenamePending)
-{
-	if (CanRenameInput())
-	{
-		GetStackEditorData().SetModuleInputIsRenamePending(ParameterName.ToString(), bIsRenamePending);
-	}
-}
-
 TArray<UEdGraphPin*> UNiagaraStackParameterStoreEntry::GetOwningPins()
 {
 	TArray<UNiagaraGraph*> GraphsToCheck;
@@ -233,14 +215,13 @@ TArray<UEdGraphPin*> UNiagaraStackParameterStoreEntry::GetOwningPins()
 	return OwningPins;
 }
 
-void UNiagaraStackParameterStoreEntry::RenameInput(FString NewName)
+void UNiagaraStackParameterStoreEntry::OnRenamed(FText NewName)
 {
-	FName NewFName = FName(*NewName);
-	FString ActualNameString = NewName;
+	FString ActualNameString = NewName.ToString();
 	FString NamespacePrefix = FNiagaraParameterHandle::UserNamespace.ToString() + ".";
-	if (NewName.Contains(NamespacePrefix))
+	if (ActualNameString.Contains(NamespacePrefix))
 	{
-		ActualNameString = NewName.Replace(*NamespacePrefix, TEXT(""));
+		ActualNameString = ActualNameString.Replace(*NamespacePrefix, TEXT(""));
 	}
 	FName ActualName = FName(*ActualNameString);
 	// what if it's not user namespace? dehardcode.
@@ -248,7 +229,6 @@ void UNiagaraStackParameterStoreEntry::RenameInput(FString NewName)
 	FName VariableName = ParameterHandle.GetParameterHandleString();
 	if (VariableName != ParameterName)
 	{
-
 		// destroy links, rename parameter and rebuild links
 		TArray<UEdGraphPin*> OwningPins = GetOwningPins();
 		TArray<UEdGraphPin*> LinkedPins;
diff --git a/Engine/Plugins/FX/Niagara/Source/NiagaraEditor/Private/ViewModels/Stack/NiagaraStackScriptItemGroup.cpp b/Engine/Plugins/FX/Niagara/Source/NiagaraEditor/Private/ViewModels/Stack/NiagaraStackScriptItemGroup.cpp
--- a/Engine/Plugins/FX/Niagara/Source/NiagaraEditor/Private/ViewModels/Stack/NiagaraStackScriptItemGroup.cpp
+++ b/Engine/Plugins/FX/Niagara/Source/NiagaraEditor/Private/ViewModels/Stack/NiagaraStackScriptItemGroup.cpp
@@ -269,7 +269,7 @@ private:
 			FString FunctionInputEditorDataKey = FNiagaraStackGraphUtilities::GenerateStackFunctionInputEditorDataKey(*NewAssignmentModule, InputPins[0]->PinName);
 			if (bRenameParameterOnAdd)
 			{
-				StackEditorData.SetModuleInputIsRenamePending(FunctionInputEditorDataKey, true);
+				StackEditorData.SetStackEntryIsRenamePending(FunctionInputEditorDataKey, true);
 			}
 		}
 
diff --git a/Engine/Plugins/FX/Niagara/Source/NiagaraEditor/Public/NiagaraStackEditorData.h b/Engine/Plugins/FX/Niagara/Source/NiagaraEditor/Public/NiagaraStackEditorData.h
--- a/Engine/Plugins/FX/Niagara/Source/NiagaraEditor/Public/NiagaraStackEditorData.h
+++ b/Engine/Plugins/FX/Niagara/Source/NiagaraEditor/Public/NiagaraStackEditorData.h
@@ -2,28 +2,30 @@
 
 #pragma once
 
+#include "NiagaraEditorDataBase.h"
+
 #include "NiagaraStackEditorData.generated.h"
 struct FStackIssue;
 
 /** Editor only UI data for emitters. */
 UCLASS()
-class UNiagaraStackEditorData : public UObject
+class UNiagaraStackEditorData : public UNiagaraEditorDataBase
 {
 	GENERATED_BODY()
 
 public:
 	/*
-	* Gets whether or not a module has a rename pending.
-	* @param ModuleInputKey A unique key for the module.
+	* Gets whether or not a stack entry has a rename pending.
+	* @param StackEntryKey A unique key for the stack entry.
 	*/
-	bool GetModuleInputIsRenamePending(const FString& ModuleInputKey) const;
+	bool GetStackEntryIsRenamePending(const FString& StackEntryKey) const;
 
 	/*
-	* Sets whether or not a module is pinned.
-	* @param ModuleInputKey A unique key for the module.
-	* @param bIsRenamePending Whether or not the module is pinned.
+	* Sets whether or not a stack entry has a rename pending.
+	* @param StackEntryKey A unique key for the stack entry.
+	* @param bIsRenamePending Whether or not the stack entry has a rename pending.
 	*/
-	void SetModuleInputIsRenamePending(const FString& ModuleInputKey, bool bIsRenamePending);
+	void SetStackEntryIsRenamePending(const FString& StackEntryKey, bool bIsRenamePending);
 
 	/*
 	 * Gets whether or not a stack entry is Expanded.
@@ -67,6 +69,24 @@ public:
 	*/
 	void SetStackItemShowAdvanced(const FString& StackEntryKey, bool bShowAdanced);
 
+	/*
+	* Gets a stack entry's display name. Returns null if none is found.
+	* @param StackEntryKey A unique key for the stack entry.
+	*/
+	NIAGARAEDITOR_API const FText* GetStackEntryDisplayName(const FString& StackEntryKey) const;
+
+	/*
+	* Gets a map of all renamed stack entry keys to their display name.
+	*/
+	const TMap<FString, FText>& GetAllStackEntryDisplayNames() const { return StackEntryKeyToDisplayName; }
+
+	/*
+	* Sets a stack entry's display name.
+	* @param StackEntryKey A unique key for the stack entry.
+	* @param InDisplayName The display name to set for this entry.
+	*/
+	void SetStackEntryDisplayName(const FString& StackEntryKey, const FText& InDisplayName);
+
 	/* Gets whether or not all advanced items should be shown in the stack. */
 	bool GetShowAllAdvanced() const;
 
@@ -109,31 +129,25 @@ public:
 	NIAGARAEDITOR_API const TArray<FString>& GetDismissedStackIssueIds();
 
 private:
-	UPROPERTY(Transient)
-	TMap<FString, bool> ModuleInputKeyToRenamePendingMap;
+	TMap<FString, bool> StackEntryKeyToRenamePendingMap;
 
-	UPROPERTY(Transient)
 	TMap<FString, bool> StackEntryKeyToExpandedMap;
 
-	UPROPERTY(Transient)
 	TMap<FString, bool> StackEntryKeyToPreSearchExpandedMap;
 
-	UPROPERTY(Transient)
 	TMap<FString, bool> StackItemKeyToShowAdvancedMap;
 
-	UPROPERTY(Transient)
+	UPROPERTY()
+	TMap<FString, FText> StackEntryKeyToDisplayName;
+
 	bool bShowAllAdvanced;
 
-	UPROPERTY(Transient)
 	bool bShowOutputs;
 
-	UPROPERTY(Transient)
 	bool bShowLinkedInputs;
 
-	UPROPERTY(Transient)
 	bool bShowOnlyIssues;
 
-	UPROPERTY(Transient)
 	double LastScrollPosition;
 
 	UPROPERTY()
diff --git a/Engine/Plugins/FX/Niagara/Source/NiagaraEditor/Public/ViewModels/Stack/NiagaraStackEntry.h b/Engine/Plugins/FX/Niagara/Source/NiagaraEditor/Public/ViewModels/Stack/NiagaraStackEntry.h
--- a/Engine/Plugins/FX/Niagara/Source/NiagaraEditor/Public/ViewModels/Stack/NiagaraStackEntry.h
+++ b/Engine/Plugins/FX/Niagara/Source/NiagaraEditor/Public/ViewModels/Stack/NiagaraStackEntry.h
@@ -225,6 +225,8 @@ public:
 
 	virtual FText GetDisplayName() const;
 
+	virtual FText GetOriginalName() const;
+
 	virtual UObject* GetDisplayedObject() const;
 
 	UNiagaraStackEditorData& GetStackEditorData() const;
@@ -377,6 +379,18 @@ public:
 
 	virtual void Delete() { }
 
+	/** Returns whether or not this entry can be renamed. */
+	virtual bool SupportsRename() const { return false; }
+
+	/** Gets whether this entry has a rename pending. */
+	virtual bool GetIsRenamePending() const;
+
+	/** Sets whether this entry has a rename pending. */
+	virtual void SetIsRenamePending(bool bIsRenamePending);
+
+	/** Handler for when a rename is committed for this stack entry. */
+	virtual void OnRenamed(FText NewName);
+
 protected:
 	virtual void BeginDestroy() override;
 
diff --git a/Engine/Plugins/FX/Niagara/Source/NiagaraEditor/Public/ViewModels/Stack/NiagaraStackFunctionInput.h b/Engine/Plugins/FX/Niagara/Source/NiagaraEditor/Public/ViewModels/Stack/NiagaraStackFunctionInput.h
--- a/Engine/Plugins/FX/Niagara/Source/NiagaraEditor/Public/ViewModels/Stack/NiagaraStackFunctionInput.h
+++ b/Engine/Plugins/FX/Niagara/Source/NiagaraEditor/Public/ViewModels/Stack/NiagaraStackFunctionInput.h
@@ -184,16 +184,10 @@ public:
 	void ResetToBase();
 
 	/** Returns whether or not this input can be renamed. */
-	bool CanRenameInput() const;
-
-	/** Gets whether this input has a rename pending. */
-	bool GetIsRenamePending() const;
-
-	/** Sets whether this input has a rename pending. */
-	void SetIsRenamePending(bool bIsRenamePending);
+	virtual bool SupportsRename() const override;
 
 	/** Renames this input to the name specified. */
-	void RenameInput(FName NewName);
+	virtual void OnRenamed(FText NewName) override;
 
 	/** Returns whether or not this input can be deleted. */
 	bool CanDeleteInput() const;
diff --git a/Engine/Plugins/FX/Niagara/Source/NiagaraEditor/Public/ViewModels/Stack/NiagaraStackModuleItem.h b/Engine/Plugins/FX/Niagara/Source/NiagaraEditor/Public/ViewModels/Stack/NiagaraStackModuleItem.h
--- a/Engine/Plugins/FX/Niagara/Source/NiagaraEditor/Public/ViewModels/Stack/NiagaraStackModuleItem.h
+++ b/Engine/Plugins/FX/Niagara/Source/NiagaraEditor/Public/ViewModels/Stack/NiagaraStackModuleItem.h
@@ -29,6 +29,7 @@ public:
 
 	virtual FText GetDisplayName() const override;
 	virtual UObject* GetDisplayedObject() const override;
+	virtual FText GetOriginalName() const override;
 	virtual FText GetTooltipText() const override;
 
 	INiagaraStackItemGroupAddUtilities* GetGroupAddUtilities();
@@ -37,6 +38,8 @@ public:
 	bool CanRefresh() const;
 	void Refresh();
 
+	virtual bool SupportsRename() const override { return true; }
+
 	virtual bool SupportsChangeEnabled() const override { return true; }
 	virtual bool GetIsEnabled() const override;
 
@@ -100,8 +103,7 @@ private:
 	bool FilterLinkedInputCollection(const UNiagaraStackEntry& Child) const;
 	bool FilterLinkedInputCollectionChild(const UNiagaraStackEntry& Child) const;
 	void RefreshIssues(TArray<FStackIssue>& NewIssues);
-
-private:
+	void OnEditorDataChanged();
 	void RefreshIsEnabled();
 
 private:
@@ -111,6 +113,8 @@ private:
 	bool bIsEnabled;
 	bool bCanRefresh;
 
+	TOptional<FText> CustomDisplayName;
+
 	UPROPERTY()
 	UNiagaraStackModuleItemLinkedInputCollection* LinkedInputCollection;
 
diff --git a/Engine/Plugins/FX/Niagara/Source/NiagaraEditor/Public/ViewModels/Stack/NiagaraStackParameterStoreEntry.h b/Engine/Plugins/FX/Niagara/Source/NiagaraEditor/Public/ViewModels/Stack/NiagaraStackParameterStoreEntry.h
--- a/Engine/Plugins/FX/Niagara/Source/NiagaraEditor/Public/ViewModels/Stack/NiagaraStackParameterStoreEntry.h
+++ b/Engine/Plugins/FX/Niagara/Source/NiagaraEditor/Public/ViewModels/Stack/NiagaraStackParameterStoreEntry.h
@@ -73,16 +73,10 @@ public:
 	void Reset();
 
 	/** Returns whether or not this input can be renamed. */
-	bool CanRenameInput() const;
-
-	/** Gets whether this input has a rename pending. */
-	bool GetIsRenamePending() const;
-
-	/** Sets whether this input has a rename pending. */
-	void SetIsRenamePending(bool bIsRenamePending);
+	virtual bool SupportsRename() const override { return true; }
 
 	/** Renames this input to the name specified. */
-	void RenameInput(FString NewName);
+	virtual void OnRenamed(FText NewName) override;
 
 	/** Checks if the chosen name is unique (not duplicate) */
 	bool IsUniqueName(FString NewName);
diff --git a/Engine/Plugins/FX/Niagara/Source/NiagaraEditorWidgets/Private/NiagaraEditorWidgetsStyle.cpp b/Engine/Plugins/FX/Niagara/Source/NiagaraEditorWidgets/Private/NiagaraEditorWidgetsStyle.cpp
--- a/Engine/Plugins/FX/Niagara/Source/NiagaraEditorWidgets/Private/NiagaraEditorWidgetsStyle.cpp
+++ b/Engine/Plugins/FX/Niagara/Source/NiagaraEditorWidgets/Private/NiagaraEditorWidgetsStyle.cpp
@@ -71,7 +71,15 @@ TSharedRef< FSlateStyleSet > FNiagaraEditorWidgetsStyle::Create()
 	FTextBlockStyle StackGroupText = FTextBlockStyle(NormalText)
 		.SetFont(StackGroupFont)
 		.SetColorAndOpacity(FLinearColor(1.0f, 1.0f, 1.0f, 1.0f));
-	Style->Set("NiagaraEditor.Stack.GroupText", StackGroupText);
+	
+	FEditableTextBoxStyle EditableGroupText = FEditableTextBoxStyle(NormalEditableTextBox)
+		.SetFont(StackGroupFont);
+
+	FInlineEditableTextBlockStyle StackEditableGroupText = FInlineEditableTextBlockStyle()
+		.SetEditableTextBoxStyle(EditableGroupText)
+		.SetTextStyle(StackGroupText); 
+
+	Style->Set("NiagaraEditor.Stack.GroupText", StackEditableGroupText);
 
 	FSlateFontInfo StackDefaultFont = DEFAULT_FONT("Regular", 10);
 	FTextBlockStyle StackDefaultText = FTextBlockStyle(NormalText)
@@ -97,14 +105,25 @@ TSharedRef< FSlateStyleSet > FNiagaraEditorWidgetsStyle::Create()
 	FSlateFontInfo StackItemFont = DEFAULT_FONT("Regular", 11);
 	FTextBlockStyle StackItemText = FTextBlockStyle(NormalText)
 		.SetFont(StackItemFont);
-	Style->Set("NiagaraEditor.Stack.ItemText", StackItemText);
+	
+	FEditableTextBoxStyle EditableItemText = FEditableTextBoxStyle(NormalEditableTextBox)
+		.SetFont(StackItemFont);
+
+	FInlineEditableTextBlockStyle StackEditableItemText = FInlineEditableTextBlockStyle()
+		.SetTextStyle(StackItemText)
+		.SetEditableTextBoxStyle(EditableItemText);
+	Style->Set("NiagaraEditor.Stack.ItemText", StackEditableItemText);
+
+	FSlateFontInfo TypeNameFont = DEFAULT_FONT("Regular", 9);
+	FTextBlockStyle TypeNameText = FTextBlockStyle(NormalText)
+		.SetFont(TypeNameFont);
+	Style->Set("NiagaraEditor.Stack.TypeNameText", TypeNameText);
 
 	FSlateFontInfo SystemOverviewListHeaderFont = DEFAULT_FONT("Bold", 12);
 	FTextBlockStyle SystemOverviewListHeaderText = FTextBlockStyle(NormalText)
 		.SetFont(SystemOverviewListHeaderFont);
 	Style->Set("NiagaraEditor.SystemOverview.ListHeaderText", SystemOverviewListHeaderText);
 
-
 	FSlateFontInfo SystemOverviewItemFont = DEFAULT_FONT("Regular", 9);
 	FTextBlockStyle SystemOverviewItemText = FTextBlockStyle(NormalText)
 		.SetFont(SystemOverviewItemFont);
diff --git a/Engine/Plugins/FX/Niagara/Source/NiagaraEditorWidgets/Private/NiagaraStackCommandContext.cpp b/Engine/Plugins/FX/Niagara/Source/NiagaraEditorWidgets/Private/NiagaraStackCommandContext.cpp
--- a/Engine/Plugins/FX/Niagara/Source/NiagaraEditorWidgets/Private/NiagaraStackCommandContext.cpp
+++ b/Engine/Plugins/FX/Niagara/Source/NiagaraEditorWidgets/Private/NiagaraStackCommandContext.cpp
@@ -60,7 +60,8 @@ bool FNiagaraStackCommandContext::AddEditMenuItems(FMenuBuilder& MenuBuilder)
 	bool bSupportsCopy = SelectedEntries.ContainsByPredicate([](const UNiagaraStackEntry* SelectedEntry) { return SelectedEntry->SupportsCopy(); });
 	bool bSupportsPaste = SelectedEntries.ContainsByPredicate([](const UNiagaraStackEntry* SelectedEntry) { return SelectedEntry->SupportsPaste(); });
 	bool bSupportsDelete = SelectedEntries.ContainsByPredicate([](const UNiagaraStackEntry* SelectedEntry) { return SelectedEntry->SupportsDelete(); });
-	if (bSupportsCut || bSupportsCopy || bSupportsPaste || bSupportsDelete)
+	bool bSupportsRename = SelectedEntries.ContainsByPredicate([](const UNiagaraStackEntry* SelectedEntry) { return SelectedEntry->SupportsRename(); });
+	if (bSupportsCut || bSupportsCopy || bSupportsPaste || bSupportsDelete || bSupportsRename)
 	{
 		MenuBuilder.BeginSection("EntryEdit", LOCTEXT("EntryEditActions", "Edit"));
 		{
@@ -84,6 +85,11 @@ bool FNiagaraStackCommandContext::AddEditMenuItems(FMenuBuilder& MenuBuilder)
 				TAttribute<FText> CanDeleteToolTip = TAttribute<FText>::Create(TAttribute<FText>::FGetter::CreateSP(this, &FNiagaraStackCommandContext::GetCanDeleteSelectedEntriesToolTip));
 				MenuBuilder.AddMenuEntry(FGenericCommands::Get().Delete, NAME_None, TAttribute<FText>(), CanDeleteToolTip);
 			}
+			if (bSupportsRename)
+			{
+				TAttribute<FText> CanRenameToolTip = TAttribute<FText>::Create(TAttribute<FText>::FGetter::CreateSP(this, &FNiagaraStackCommandContext::GetCanRenameSelectedEntriesToolTip));
+				MenuBuilder.AddMenuEntry(FGenericCommands::Get().Rename, NAME_None, TAttribute<FText>(), CanRenameToolTip);
+			}
 		}
 		MenuBuilder.EndSection();
 		return true;
@@ -105,6 +111,9 @@ void FNiagaraStackCommandContext::SetupCommands()
 	Commands->MapAction(FGenericCommands::Get().Delete, FUIAction(
 		FExecuteAction::CreateSP(this, &FNiagaraStackCommandContext::DeleteSelectedEntries),
 		FCanExecuteAction::CreateSP(this, &FNiagaraStackCommandContext::CanDeleteSelectedEntries)));
+	Commands->MapAction(FGenericCommands::Get().Rename, FUIAction(
+		FExecuteAction::CreateSP(this, &FNiagaraStackCommandContext::RenameSelectedEntries),
+		FCanExecuteAction::CreateSP(this, &FNiagaraStackCommandContext::CanRenameSelectedEntries)));
 }
 
 bool FNiagaraStackCommandContext::CanCutSelectedEntries() const
@@ -204,5 +213,38 @@ void FNiagaraStackCommandContext::DeleteSelectedEntries() const
 	FNiagaraStackClipboardUtilities::DeleteSelection(SelectedEntries);
 }
 
+FText FNiagaraStackCommandContext::GetCanRenameSelectedEntriesToolTip() const
+{
+	if (SelectedEntries.Num() > 1)
+	{
+		LOCTEXT("CantRenameMultipleToolTip", "Can't rename multiple items.");
+	}
+
+	if (SelectedEntries[0]->SupportsRename())
+	{
+		return LOCTEXT("CantRenameToolTip", "The selected item cannot be renamed.");
+	}
+
+	return LOCTEXT("RenameToolTip", "Rename the selected item.");
+}
+
+bool FNiagaraStackCommandContext::CanRenameSelectedEntries() const
+{
+	if (SelectedEntries.Num() == 1)
+	{
+		return SelectedEntries[0]->SupportsRename();
+	}
+
+	return false;
+}
+
+void FNiagaraStackCommandContext::RenameSelectedEntries()
+{
+	if (SelectedEntries.Num() == 1)
+	{
+		SelectedEntries[0]->SetIsRenamePending(true);
+	}
+}
+
 #undef LOCTEXT_NAMESPACE
 
diff --git a/Engine/Plugins/FX/Niagara/Source/NiagaraEditorWidgets/Private/NiagaraStackCommandContext.h b/Engine/Plugins/FX/Niagara/Source/NiagaraEditorWidgets/Private/NiagaraStackCommandContext.h
--- a/Engine/Plugins/FX/Niagara/Source/NiagaraEditorWidgets/Private/NiagaraStackCommandContext.h
+++ b/Engine/Plugins/FX/Niagara/Source/NiagaraEditorWidgets/Private/NiagaraStackCommandContext.h
@@ -49,6 +49,12 @@ private:
 
 	void DeleteSelectedEntries() const;
 
+	FText GetCanRenameSelectedEntriesToolTip() const;
+
+	bool CanRenameSelectedEntries() const;
+
+	void RenameSelectedEntries();
+
 private:
 	TSharedRef<FUICommandList> Commands;
 
diff --git a/Engine/Plugins/FX/Niagara/Source/NiagaraEditorWidgets/Private/SNiagaraStackEntryWidget.cpp b/Engine/Plugins/FX/Niagara/Source/NiagaraEditorWidgets/Private/SNiagaraStackEntryWidget.cpp
--- a/Engine/Plugins/FX/Niagara/Source/NiagaraEditorWidgets/Private/SNiagaraStackEntryWidget.cpp
+++ b/Engine/Plugins/FX/Niagara/Source/NiagaraEditorWidgets/Private/SNiagaraStackEntryWidget.cpp
@@ -2,9 +2,11 @@
 
 #include "SNiagaraStackEntryWidget.h"
 #include "NiagaraEditorWidgetsStyle.h"
+#include "NiagaraStackEditorData.h"
 #include "Widgets/Layout/SWrapBox.h"
 #include "Widgets/Layout/SBox.h"
 #include "Widgets/Text/STextBlock.h"
+#include "Widgets/Text/SInlineEditableTextBlock.h"
 
 #define LOCTEXT_NAMESPACE "SNiagaraStackEntryWidget"
 
@@ -14,6 +16,9 @@ void SNiagaraStackDisplayName::Construct(const FArguments& InArgs, UNiagaraStack
 	StackViewModel = &InStackViewModel;
 	TextStyleName = InTextStyleName;
 
+	TypeNameStyle = InArgs._TypeNameStyle;
+	OnRenameCommitted = InArgs._OnRenameCommitted;
+
 	StackViewModel->OnStructureChanged().AddSP(this, &SNiagaraStackDisplayName::StackViewModelStructureChanged);
 
 	ChildSlot
@@ -32,13 +37,52 @@ SNiagaraStackDisplayName::~SNiagaraStackDisplayName()
 
 TSharedRef<SWidget> SNiagaraStackDisplayName::ConstructChildren()
 {
-	TSharedRef<STextBlock> BaseNameWidget = SNew(STextBlock)
-		.TextStyle(FNiagaraEditorWidgetsStyle::Get(), TextStyleName)
-		.ToolTipText_UObject(StackEntryItem, &UNiagaraStackEntry::GetTooltipText)
-		.Text_UObject(StackEntryItem, &UNiagaraStackEntry::GetDisplayName)
-		.HighlightText_UObject(StackViewModel, &UNiagaraStackViewModel::GetCurrentSearchText)
-		.ColorAndOpacity(this, &SNiagaraStackDisplayName::GetTextColorForSearch, FSlateColor::UseForeground())
-		.IsEnabled(this, &SNiagaraStackDisplayName::GetIsEnabled);
+	const FInlineEditableTextBlockStyle& EditableStyle = FNiagaraEditorWidgetsStyle::Get().GetWidgetStyle<FInlineEditableTextBlockStyle>(TextStyleName);
+
+	TSharedPtr<SWidget> BaseNameWidget;
+	if (StackEntryItem->SupportsRename())
+	{
+		BaseNameWidget = SNew(SHorizontalBox)
+			+ SHorizontalBox::Slot()
+			.AutoWidth()
+			.HAlign(HAlign_Left)
+			.VAlign(VAlign_Bottom)
+			[
+				SAssignNew(EditableTextBlock, SInlineEditableTextBlock)
+					.Style(&EditableStyle)
+					.ToolTipText_UObject(StackEntryItem, &UNiagaraStackEntry::GetTooltipText)
+					.Text_UObject(StackEntryItem, &UNiagaraStackEntry::GetDisplayName)
+					.HighlightText_UObject(StackViewModel, &UNiagaraStackViewModel::GetCurrentSearchText)
+					.ColorAndOpacity(this, &SNiagaraStackEntryWidget::GetTextColorForSearch, FSlateColor::UseForeground())
+					.OnTextCommitted(OnRenameCommitted)
+					.IsEnabled_UObject(StackEntryItem, &UNiagaraStackEntry::GetIsEnabledAndOwnerIsEnabled)
+			]
+			+ SHorizontalBox::Slot()
+			.AutoWidth()
+			.Padding(5,0,0,0)
+			.HAlign(HAlign_Left)
+			.VAlign(VAlign_Bottom)
+			[
+				SNew(STextBlock)
+				.TextStyle(TypeNameStyle)
+				.ToolTipText_UObject(StackEntryItem, &UNiagaraStackEntry::GetTooltipText)
+				.Text(this, &SNiagaraStackDisplayName::GetOriginalName)
+				.HighlightText_UObject(StackViewModel, &UNiagaraStackViewModel::GetCurrentSearchText)
+				.ColorAndOpacity(this, &SNiagaraStackDisplayName::GetTextColorForSearch, FSlateColor::UseSubduedForeground())
+				.IsEnabled_UObject(StackEntryItem, &UNiagaraStackEntry::GetIsEnabledAndOwnerIsEnabled)
+				.Visibility(this, &SNiagaraStackDisplayName::ShouldShowOriginalName)
+			];
+	}
+	else
+	{
+		BaseNameWidget = SNew(STextBlock)
+			.TextStyle(&EditableStyle.TextStyle)
+			.ToolTipText_UObject(StackEntryItem, &UNiagaraStackEntry::GetTooltipText)
+			.Text_UObject(StackEntryItem, &UNiagaraStackEntry::GetDisplayName)
+			.HighlightText_UObject(StackViewModel, &UNiagaraStackViewModel::GetCurrentSearchText)
+			.ColorAndOpacity(this, &SNiagaraStackDisplayName::GetTextColorForSearch, FSlateColor::UseForeground())
+			.IsEnabled_UObject(StackEntryItem, &UNiagaraStackEntry::GetIsEnabledAndOwnerIsEnabled);
+	}
 
 	int32 NumTopLevelEmitters = 0;
 	for (const TSharedRef<UNiagaraStackViewModel::FTopLevelViewModel>& TopLevelViewModel : StackViewModel->GetTopLevelViewModels())
@@ -52,7 +96,7 @@ TSharedRef<SWidget> SNiagaraStackDisplayName::ConstructChildren()
 	if (NumTopLevelEmitters <= 1)
 	{
 		TopLevelViewModelCountAtLastConstruction = 1;
-		return BaseNameWidget;
+		return BaseNameWidget.ToSharedRef();
 	}
 	else
 	{
@@ -71,11 +115,11 @@ TSharedRef<SWidget> SNiagaraStackDisplayName::ConstructChildren()
 					.Text(this, &SNiagaraStackDisplayName::GetTopLevelDisplayName, TWeakPtr<UNiagaraStackViewModel::FTopLevelViewModel>(TopLevelViewModel))
 					.HighlightText_UObject(StackViewModel, &UNiagaraStackViewModel::GetCurrentSearchText)
 					.ColorAndOpacity(this, &SNiagaraStackDisplayName::GetTextColorForSearch, FSlateColor::UseForeground())
-					.IsEnabled(this, &SNiagaraStackDisplayName::GetIsEnabled)
+					.IsEnabled_UObject(StackEntryItem, &UNiagaraStackEntry::GetIsEnabledAndOwnerIsEnabled)
 				]
 				+ SWrapBox::Slot()
 				[
-					BaseNameWidget
+					BaseNameWidget.ToSharedRef()
 				];
 		}
 	}
@@ -109,24 +153,41 @@ void SNiagaraStackDisplayName::StackViewModelStructureChanged()
 	}
 }
 
-bool SNiagaraStackDisplayName::GetIsEnabled() const
+FText SNiagaraStackDisplayName::GetOriginalName() const
+{
+	if (StackEntryItem->IsFinalized() == false)
+	{
+		return FText::Format(FTextFormat::FromString(TEXT("({0})")), StackEntryItem->GetOriginalName());
+	}
+	return FText::GetEmpty();
+}
+
+EVisibility SNiagaraStackDisplayName::ShouldShowOriginalName() const
+{
+	const FText* CurrentDisplayName = StackEntryItem->GetStackEditorData().GetStackEntryDisplayName(StackEntryItem->GetStackEditorDataKey());
+	return CurrentDisplayName != nullptr ? EVisibility::Visible : EVisibility::Collapsed;
+}
+
+void SNiagaraStackDisplayName::StartRename()
 {
-	return StackEntryItem->GetOwnerIsEnabled() && StackEntryItem->GetIsEnabled();
+	if (EditableTextBlock.IsValid())
+	{
+		EditableTextBlock->EnterEditingMode();
+	}
 }
 
 FSlateColor SNiagaraStackEntryWidget::GetTextColorForSearch(FSlateColor DefaultColor) const
 {
 	if (IsCurrentSearchMatch())
 	{
-		return FNiagaraEditorWidgetsStyle::Get().GetColor("NiagaraEditor.Stack.SearchHighlightColor");
-	} 
-	
+		return FSlateColor(FLinearColor(FColor::Orange));
+	}
 	return DefaultColor;
 }
 
 bool SNiagaraStackEntryWidget::IsCurrentSearchMatch() const
 {
-	auto FocusedEntry = StackViewModel->GetCurrentFocusedEntry();
+	UNiagaraStackEntry* FocusedEntry = StackViewModel->GetCurrentFocusedEntry();
 	return StackEntryItem != nullptr && FocusedEntry == StackEntryItem;
 }
 
diff --git a/Engine/Plugins/FX/Niagara/Source/NiagaraEditorWidgets/Private/SNiagaraStackEntryWidget.h b/Engine/Plugins/FX/Niagara/Source/NiagaraEditorWidgets/Private/SNiagaraStackEntryWidget.h
--- a/Engine/Plugins/FX/Niagara/Source/NiagaraEditorWidgets/Private/SNiagaraStackEntryWidget.h
+++ b/Engine/Plugins/FX/Niagara/Source/NiagaraEditorWidgets/Private/SNiagaraStackEntryWidget.h
@@ -5,17 +5,20 @@
 #include "Widgets/SCompoundWidget.h"
 #include "Widgets/DeclarativeSyntaxSupport.h"
 #include "ViewModels/Stack/NiagaraStackViewModel.h"
+#include "Framework/SlateDelegates.h"
 
 class UNiagaraStackViewModel;
 class UNiagaraStackEntry;
 class SBox;
+class SInlineEditableTextBlock;
 
 class SNiagaraStackEntryWidget : public SCompoundWidget
 {
 public:
-	FSlateColor GetTextColorForSearch(FSlateColor DefaultColor) const;
 	FReply ExpandEntry();
 	
+	FSlateColor GetTextColorForSearch(FSlateColor DefaultColor) const;
+
 protected:
 	bool IsCurrentSearchMatch() const;
 	
@@ -28,26 +31,38 @@ class SNiagaraStackDisplayName : public SNiagaraStackEntryWidget
 {
 public:
 	SLATE_BEGIN_ARGS(SNiagaraStackDisplayName) { }
+		SLATE_EVENT(FOnTextCommitted, OnRenameCommitted);
+		SLATE_STYLE_ARGUMENT(FTextBlockStyle, TypeNameStyle)
 	SLATE_END_ARGS();
 
 	void Construct(const FArguments& InArgs, UNiagaraStackEntry& InStackEntry, UNiagaraStackViewModel& InStackViewModel, FName InTextStyleName);
 
 	~SNiagaraStackDisplayName();
 
+	void StartRename();
+
 private:
+	void OnEndRename();
+
 	TSharedRef<SWidget> ConstructChildren();
 
 	FText GetTopLevelDisplayName(TWeakPtr<UNiagaraStackViewModel::FTopLevelViewModel> TopLevelViewModelWeak) const;
 
 	void StackViewModelStructureChanged();
 
-	bool GetIsEnabled() const;
+	FText GetOriginalName() const;
+
+	EVisibility ShouldShowOriginalName() const;
 
 private:
 	FName TextStyleName;
+	const FTextBlockStyle* TypeNameStyle;
 
 	TSharedPtr<SBox> Container;
 
+	TSharedPtr<SInlineEditableTextBlock> EditableTextBlock;
+	FOnTextCommitted OnRenameCommitted;
+
 	mutable FText TopLevelDisplayNameCache;
 	mutable FText TopLevelDisplayNameFormattedCache;
 	int32 TopLevelViewModelCountAtLastConstruction;
diff --git a/Engine/Plugins/FX/Niagara/Source/NiagaraEditorWidgets/Private/Stack/SNiagaraStackErrorItem.cpp b/Engine/Plugins/FX/Niagara/Source/NiagaraEditorWidgets/Private/Stack/SNiagaraStackErrorItem.cpp
--- a/Engine/Plugins/FX/Niagara/Source/NiagaraEditorWidgets/Private/Stack/SNiagaraStackErrorItem.cpp
+++ b/Engine/Plugins/FX/Niagara/Source/NiagaraEditorWidgets/Private/Stack/SNiagaraStackErrorItem.cpp
@@ -54,7 +54,6 @@ void SNiagaraStackErrorItem::Construct(const FArguments& InArgs, UNiagaraStackEr
 			SNew(STextBlock)
 			.TextStyle(FNiagaraEditorStyle::Get(), "NiagaraEditor.ParameterText")
 			.Text_UObject(ErrorItem, &UNiagaraStackErrorItem::GetDisplayName)
-			.ColorAndOpacity(this, &SNiagaraStackErrorItem::GetTextColorForSearch, FSlateColor::UseForeground())
 			.HighlightText_UObject(StackViewModel, &UNiagaraStackViewModel::GetCurrentSearchText)
 		]
 	];
diff --git a/Engine/Plugins/FX/Niagara/Source/NiagaraEditorWidgets/Private/Stack/SNiagaraStackFunctionInputName.cpp b/Engine/Plugins/FX/Niagara/Source/NiagaraEditorWidgets/Private/Stack/SNiagaraStackFunctionInputName.cpp
--- a/Engine/Plugins/FX/Niagara/Source/NiagaraEditorWidgets/Private/Stack/SNiagaraStackFunctionInputName.cpp
+++ b/Engine/Plugins/FX/Niagara/Source/NiagaraEditorWidgets/Private/Stack/SNiagaraStackFunctionInputName.cpp
@@ -44,7 +44,6 @@ void SNiagaraStackFunctionInputName::Construct(const FArguments& InArgs, UNiagar
 			.IsSelected(this, &SNiagaraStackFunctionInputName::GetIsNameWidgetSelected)
 			.OnTextCommitted(this, &SNiagaraStackFunctionInputName::OnNameTextCommitted)
 			.HighlightText_UObject(InStackViewModel, &UNiagaraStackViewModel::GetCurrentSearchText)
-			.ColorAndOpacity(this, &SNiagaraStackFunctionInputName::GetTextColorForSearch, FSlateColor::UseForeground())
 			.ToolTipText(this, &SNiagaraStackFunctionInputName::GetToolTipText)
 		]
 	];
@@ -77,7 +76,7 @@ void SNiagaraStackFunctionInputName::OnEditConditionCheckStateChanged(ECheckBoxS
 
 bool SNiagaraStackFunctionInputName::GetIsNameReadOnly() const
 {
-	return FunctionInput->CanRenameInput() == false;
+	return FunctionInput->SupportsRename() == false;
 }
 
 bool SNiagaraStackFunctionInputName::GetIsNameWidgetSelected() const
@@ -103,5 +102,5 @@ FText SNiagaraStackFunctionInputName::GetToolTipText() const
 
 void SNiagaraStackFunctionInputName::OnNameTextCommitted(const FText& InText, ETextCommit::Type InCommitType)
 {
-	FunctionInput->RenameInput(*InText.ToString());
+	FunctionInput->OnRenamed(InText);
 }
diff --git a/Engine/Plugins/FX/Niagara/Source/NiagaraEditorWidgets/Private/Stack/SNiagaraStackModuleItem.cpp b/Engine/Plugins/FX/Niagara/Source/NiagaraEditorWidgets/Private/Stack/SNiagaraStackModuleItem.cpp
--- a/Engine/Plugins/FX/Niagara/Source/NiagaraEditorWidgets/Private/Stack/SNiagaraStackModuleItem.cpp
+++ b/Engine/Plugins/FX/Niagara/Source/NiagaraEditorWidgets/Private/Stack/SNiagaraStackModuleItem.cpp
@@ -26,6 +26,7 @@
 #include "NiagaraEditorWidgetsUtilities.h"
 #include "Subsystems/AssetEditorSubsystem.h"
 #include "Editor.h"
+#include "EditorFontGlyphs.h"
 
 #define LOCTEXT_NAMESPACE "NiagaraStackModuleItem"
 
@@ -52,7 +53,9 @@ void SNiagaraStackModuleItem::Construct(const FArguments& InArgs, UNiagaraStackM
 			.VAlign(VAlign_Center)
 			.Padding(2, 0, 0, 0)
 			[
-				SNew(SNiagaraStackDisplayName, InModuleItem, *InStackViewModel, "NiagaraEditor.Stack.ItemText")
+				SAssignNew(DisplayNameWidget, SNiagaraStackDisplayName, InModuleItem, *InStackViewModel, "NiagaraEditor.Stack.ItemText")
+				.OnRenameCommitted(this, &SNiagaraStackModuleItem::OnRenameCommitted)
+				.TypeNameStyle(FNiagaraEditorWidgetsStyle::Get(), "NiagaraEditor.Stack.TypeNameText")
 			]
 			// Raise Action Menu button
 			+ SHorizontalBox::Slot()
@@ -86,7 +89,7 @@ void SNiagaraStackModuleItem::Construct(const FArguments& InArgs, UNiagaraStackM
 				[
 					SNew(STextBlock)
 					.Font(FEditorStyle::Get().GetFontStyle("FontAwesome.10"))
-					.Text(FText::FromString(FString(TEXT("\xf021"))))
+					.Text(FEditorFontGlyphs::Refresh)
 				]
 			]
 			// Delete button
@@ -105,7 +108,7 @@ void SNiagaraStackModuleItem::Construct(const FArguments& InArgs, UNiagaraStackM
 				[
 					SNew(STextBlock)
 					.Font(FEditorStyle::Get().GetFontStyle("FontAwesome.10"))
-					.Text(FText::FromString(FString(TEXT("\xf1f8"))))
+					.Text(FEditorFontGlyphs::Trash)
 				]
 			]
 			// Enabled checkbox
@@ -156,6 +159,16 @@ void SNiagaraStackModuleItem::Tick(const FGeometry& AllottedGeometry, const doub
 		ModuleItem->SetIsModuleScriptReassignmentPending(false);
 		ShowReassignModuleScriptMenu();
 	}
+
+	if (StackEntryItem->GetIsRenamePending())
+	{
+		if (DisplayNameWidget.IsValid())
+		{
+			DisplayNameWidget->StartRename();
+		}
+		StackEntryItem->SetIsRenamePending(false);
+	}
+	SNiagaraStackEntryWidget::Tick(AllottedGeometry, InCurrentTime, InDeltaTime);
 }
 
 ECheckBoxState SNiagaraStackModuleItem::GetCheckState() const
@@ -360,4 +373,12 @@ void SNiagaraStackModuleItem::ShowReassignModuleScriptMenu()
 	FSlateApplication::Get().PushMenu(AsShared(), FWidgetPath(), MenuWidget, MenuPosition, FPopupTransitionEffect::ContextMenu);
 }
 
+void SNiagaraStackModuleItem::OnRenameCommitted(const FText& NewName, ETextCommit::Type CommitType)
+{
+	if (CommitType != ETextCommit::OnCleared)
+	{
+		StackEntryItem->OnRenamed(NewName);
+	}
+}
+
 #undef LOCTEXT_NAMESPACE
\ No newline at end of file
diff --git a/Engine/Plugins/FX/Niagara/Source/NiagaraEditorWidgets/Private/Stack/SNiagaraStackModuleItem.h b/Engine/Plugins/FX/Niagara/Source/NiagaraEditorWidgets/Private/Stack/SNiagaraStackModuleItem.h
--- a/Engine/Plugins/FX/Niagara/Source/NiagaraEditorWidgets/Private/Stack/SNiagaraStackModuleItem.h
+++ b/Engine/Plugins/FX/Niagara/Source/NiagaraEditorWidgets/Private/Stack/SNiagaraStackModuleItem.h
@@ -5,9 +5,11 @@
 #include "SNiagaraStackEntryWidget.h"
 #include "Styling/SlateTypes.h"
 #include "Layout/Visibility.h"
+#include "Types/SlateEnums.h"
 
 class UNiagaraStackModuleItem;
 class UNiagaraStackViewModel;
+class SNiagaraStackDisplayName;
 
 class SNiagaraStackModuleItem : public SNiagaraStackEntryWidget
 {
@@ -58,6 +60,9 @@ private:
 
 	void ShowReassignModuleScriptMenu();
 
+	void OnRenameCommitted(const FText& NewName, ETextCommit::Type CommitType);
+
 private:
 	UNiagaraStackModuleItem* ModuleItem;
+	TSharedPtr<SNiagaraStackDisplayName> DisplayNameWidget;
 };
\ No newline at end of file
diff --git a/Engine/Plugins/FX/Niagara/Source/NiagaraEditorWidgets/Private/Stack/SNiagaraStackParameterStoreEntryName.cpp b/Engine/Plugins/FX/Niagara/Source/NiagaraEditorWidgets/Private/Stack/SNiagaraStackParameterStoreEntryName.cpp
--- a/Engine/Plugins/FX/Niagara/Source/NiagaraEditorWidgets/Private/Stack/SNiagaraStackParameterStoreEntryName.cpp
+++ b/Engine/Plugins/FX/Niagara/Source/NiagaraEditorWidgets/Private/Stack/SNiagaraStackParameterStoreEntryName.cpp
@@ -31,7 +31,6 @@ void SNiagaraStackParameterStoreEntryName::Construct(const FArguments& InArgs, U
 			.OnTextCommitted(this, &SNiagaraStackParameterStoreEntryName::OnNameTextCommitted)
 			.OnVerifyTextChanged(this, &SNiagaraStackParameterStoreEntryName::VerifyNameTextChanged)
 			.HighlightText_UObject(InStackViewModel, &UNiagaraStackViewModel::GetCurrentSearchText)
-			.ColorAndOpacity(this, &SNiagaraStackParameterStoreEntryName::GetTextColorForSearch, FSlateColor::UseForeground())
 		]
 	];
 }
@@ -42,14 +41,14 @@ void SNiagaraStackParameterStoreEntryName::Tick(const FGeometry& AllottedGeometr
 	if (StackEntry->GetIsRenamePending() && NameTextBlock.IsValid())
 	{
 		NameTextBlock->EnterEditingMode();
-		StackEntry->SetIsRenamePending(false);
 	}
+	StackEntry->SetIsRenamePending(false);
 	SCompoundWidget::Tick(AllottedGeometry, InCurrentTime, InDeltaTime);
 }
 
 bool SNiagaraStackParameterStoreEntryName::GetIsNameReadOnly() const
 {
-	return StackEntry->CanRenameInput() == false;
+	return StackEntry->SupportsRename() == false;
 }
 
 bool SNiagaraStackParameterStoreEntryName::GetIsNameWidgetSelected() const
@@ -75,9 +74,9 @@ bool SNiagaraStackParameterStoreEntryName::VerifyNameTextChanged(const FText& Ne
 
 void SNiagaraStackParameterStoreEntryName::OnNameTextCommitted(const FText& InText, ETextCommit::Type InCommitType)
 {
-	if (StackEntry->GetDisplayName().ToString() != *InText.ToString())
+	if (!StackEntry->GetDisplayName().EqualTo(InText))
 	{
-		StackEntry->RenameInput(*InText.ToString());
+		StackEntry->OnRenamed(InText);
 		// toast notification
 		FNotificationInfo Info(FText::Format(LOCTEXT("NiagaraRenamedUserParameter", "System exposed parameter was renamed.\n{0}\n(All links to inner variables were updated in the process.)"), StackEntry->GetDisplayName()));
 		Info.ExpireDuration = 5.0f;
diff --git a/Engine/Plugins/FX/Niagara/Source/NiagaraEditorWidgets/Private/Stack/SNiagaraStackParameterStoreEntryValue.cpp b/Engine/Plugins/FX/Niagara/Source/NiagaraEditorWidgets/Private/Stack/SNiagaraStackParameterStoreEntryValue.cpp
--- a/Engine/Plugins/FX/Niagara/Source/NiagaraEditorWidgets/Private/Stack/SNiagaraStackParameterStoreEntryValue.cpp
+++ b/Engine/Plugins/FX/Niagara/Source/NiagaraEditorWidgets/Private/Stack/SNiagaraStackParameterStoreEntryValue.cpp
@@ -93,7 +93,6 @@ void SNiagaraStackParameterStoreEntryValue::Construct(const FArguments& InArgs,
 			.IsFocusable(false)
 			.ForegroundColor(FSlateColor::UseForeground())
 			.ToolTipText(LOCTEXT("DeleteToolTip", "Delete this parameter"))
-			.Visibility(this, &SNiagaraStackParameterStoreEntryValue::GetDeleteButtonVisibility)
 			.OnClicked(this, &SNiagaraStackParameterStoreEntryValue::DeleteClicked)
 			.Content()
 			[
@@ -250,11 +249,6 @@ void SNiagaraStackParameterStoreEntryValue::ParameterPropertyValueChanged(const
 	StackEntry->NotifyValueChanged();
 }
 
-EVisibility SNiagaraStackParameterStoreEntryValue::GetDeleteButtonVisibility() const
-{
-	return StackEntry->CanRenameInput() ? EVisibility::Visible : EVisibility::Collapsed;
-}
-
 EVisibility SNiagaraStackParameterStoreEntryValue::GetReferenceVisibility() const
 {
 	return EVisibility::Collapsed;
diff --git a/Engine/Plugins/FX/Niagara/Source/NiagaraEditorWidgets/Private/Stack/SNiagaraStackParameterStoreEntryValue.h b/Engine/Plugins/FX/Niagara/Source/NiagaraEditorWidgets/Private/Stack/SNiagaraStackParameterStoreEntryValue.h
--- a/Engine/Plugins/FX/Niagara/Source/NiagaraEditorWidgets/Private/Stack/SNiagaraStackParameterStoreEntryValue.h
+++ b/Engine/Plugins/FX/Niagara/Source/NiagaraEditorWidgets/Private/Stack/SNiagaraStackParameterStoreEntryValue.h
@@ -37,8 +37,6 @@ private:
 
 	void ParameterPropertyValueChanged(const FPropertyChangedEvent& PropertyChangedEvent);
 
-	EVisibility GetDeleteButtonVisibility() const;
-
 	EVisibility GetReferenceVisibility() const;
 
 	EVisibility GetResetButtonVisibility() const;
