Adding emitter reparenting and moving all emitter/stack context options to a more visible menu
#rnx

diff --git a/Engine/Plugins/FX/Niagara/Source/Niagara/Classes/NiagaraEmitter.h b/Engine/Plugins/FX/Niagara/Source/Niagara/Classes/NiagaraEmitter.h
--- a/Engine/Plugins/FX/Niagara/Source/Niagara/Classes/NiagaraEmitter.h
+++ b/Engine/Plugins/FX/Niagara/Source/Niagara/Classes/NiagaraEmitter.h
@@ -478,6 +478,8 @@ public:
 	NIAGARA_API void RemoveParent();
 
 	NIAGARA_API void SetParent(UNiagaraEmitter& InParent);
+
+	NIAGARA_API	void Reparent(UNiagaraEmitter& InParent);
 #endif
 
 	void OnEffectsQualityChanged();
diff --git a/Engine/Plugins/FX/Niagara/Source/Niagara/Private/NiagaraEmitter.cpp b/Engine/Plugins/FX/Niagara/Source/Niagara/Private/NiagaraEmitter.cpp
--- a/Engine/Plugins/FX/Niagara/Source/Niagara/Private/NiagaraEmitter.cpp
+++ b/Engine/Plugins/FX/Niagara/Source/Niagara/Private/NiagaraEmitter.cpp
@@ -212,20 +212,11 @@ INiagaraMergeManager::FMergeEmitterResults UNiagaraEmitter::MergeChangesFromPare
 		return MergeResults;
 	}
 
-	if (ParentAtLastMerge == nullptr)
-	{
-		// If we don't have a copy of the last merged parent emitter, this emitter can't safely be
-		// merged.
-		INiagaraMergeManager::FMergeEmitterResults MergeResults;
-		MergeResults.MergeResult = INiagaraMergeManager::EMergeEmitterResult::FailedToDiff;
-		MergeResults.bModifiedGraph = false;
-		MergeResults.ErrorMessages.Add(NSLOCTEXT("NiagaraEmitter", "NoLastMergedParentErrorMessage", "This emitter has no 'ParentAtLastMerge' so changes can't be merged in."));
-		return MergeResults;
-	}
+	const bool bNoParentAtLastMerge = (ParentAtLastMerge == nullptr);
 
 	INiagaraModule& NiagaraModule = FModuleManager::Get().GetModuleChecked<INiagaraModule>("Niagara");
 	const INiagaraMergeManager& MergeManager = NiagaraModule.GetMergeManager();
-	INiagaraMergeManager::FMergeEmitterResults MergeResults = MergeManager.MergeEmitter(*Parent, *ParentAtLastMerge, *this);
+	INiagaraMergeManager::FMergeEmitterResults MergeResults = MergeManager.MergeEmitter(*Parent, ParentAtLastMerge, *this);
 	if (MergeResults.MergeResult == INiagaraMergeManager::EMergeEmitterResult::SucceededDifferencesApplied || MergeResults.MergeResult == INiagaraMergeManager::EMergeEmitterResult::SucceededNoDifferences)
 	{
 		if (MergeResults.MergeResult == INiagaraMergeManager::EMergeEmitterResult::SucceededDifferencesApplied)
@@ -1504,6 +1495,13 @@ void UNiagaraEmitter::SetParent(UNiagaraEmitter& InParent)
 	ParentAtLastMerge->ClearFlags(RF_Standalone | RF_Public);
 	GraphSource->MarkNotSynchronized(TEXT("Emitter parent changed"));
 }
+
+void UNiagaraEmitter::Reparent(UNiagaraEmitter& InParent)
+{
+	Parent = &InParent;
+	ParentAtLastMerge = nullptr;
+	GraphSource->MarkNotSynchronized(TEXT("Emitter parent changed"));
+}
 #endif
 
 void UNiagaraEmitter::ResolveScalabilitySettings()
diff --git a/Engine/Plugins/FX/Niagara/Source/Niagara/Public/INiagaraMergeManager.h b/Engine/Plugins/FX/Niagara/Source/Niagara/Public/INiagaraMergeManager.h
--- a/Engine/Plugins/FX/Niagara/Source/Niagara/Public/INiagaraMergeManager.h
+++ b/Engine/Plugins/FX/Niagara/Source/Niagara/Public/INiagaraMergeManager.h
@@ -38,7 +38,7 @@ public:
 		}
 	};
 
-	virtual FMergeEmitterResults MergeEmitter(UNiagaraEmitter& Parent, UNiagaraEmitter& ParentAtLastMerge, UNiagaraEmitter& Instance) const = 0;
+	virtual FMergeEmitterResults MergeEmitter(UNiagaraEmitter& Parent, UNiagaraEmitter* ParentAtLastMerge, UNiagaraEmitter& Instance) const = 0;
 
 	virtual void DiffEditableProperties(const void* BaseDataAddress, const void* OtherDataAddress, UStruct& Struct, TArray<FProperty*>& OutDifferentProperties) const = 0;
 
diff --git a/Engine/Plugins/FX/Niagara/Source/NiagaraEditor/Private/NiagaraEditorUtilities.cpp b/Engine/Plugins/FX/Niagara/Source/NiagaraEditor/Private/NiagaraEditorUtilities.cpp
--- a/Engine/Plugins/FX/Niagara/Source/NiagaraEditor/Private/NiagaraEditorUtilities.cpp
+++ b/Engine/Plugins/FX/Niagara/Source/NiagaraEditor/Private/NiagaraEditorUtilities.cpp
@@ -52,6 +52,7 @@
 #include "Subsystems/AssetEditorSubsystem.h"
 #include "Framework/Commands/GenericCommands.h"
 #include "UObject/TextProperty.h"
+#include "Editor/EditorEngine.h"
 
 #define LOCTEXT_NAMESPACE "FNiagaraEditorUtilities"
 
@@ -1605,6 +1606,15 @@ void FNiagaraEditorUtilities::ShowParentEmitterInContentBrowser(TSharedRef<FNiag
 	ContentBrowserModule.Get().SyncBrowserToAssets(TArray<FAssetData> { FAssetData(Emitter->GetParentEmitter()) });
 }
 
+void FNiagaraEditorUtilities::OpenParentEmitterForEdit(TSharedRef<FNiagaraEmitterViewModel> Emitter)
+{
+	UNiagaraEmitter* ParentEmitter = const_cast<UNiagaraEmitter*>(Emitter->GetParentEmitter());
+	if (ParentEmitter != nullptr)
+	{
+		GEditor->GetEditorSubsystem<UAssetEditorSubsystem>()->OpenEditorForAsset(ParentEmitter);
+	}
+}
+
 ECheckBoxState FNiagaraEditorUtilities::GetSelectedEmittersEnabledCheckState(TSharedRef<FNiagaraSystemViewModel> SystemViewModel)
 {
 	bool bFirst = true;
@@ -1781,11 +1791,10 @@ bool FNiagaraEditorUtilities::AddEmitterContextMenuActions(FMenuBuilder& MenuBui
 		TSharedRef<FNiagaraSystemViewModel> OwningSystemViewModel = EmitterHandleViewModel->GetOwningSystemViewModel();
 
 		bool bSingleSelection = OwningSystemViewModel->GetSelectionViewModel()->GetSelectedEmitterHandleIds().Num() == 1;
-
+		TSharedRef<FNiagaraEmitterViewModel> EmitterViewModel = EmitterHandleViewModel->GetEmitterViewModel();
 		MenuBuilder.BeginSection("EmitterActions", LOCTEXT("EmitterActions", "Emitter Actions"));
 		{
-			TSharedRef<FNiagaraEmitterViewModel> EmitterViewModel = EmitterHandleViewModel->GetEmitterViewModel();
-
+		
 			if (OwningSystemViewModel->GetEditMode() == ENiagaraSystemViewModelEditMode::SystemAsset)
 			{
 				MenuBuilder.AddMenuEntry(
@@ -1825,49 +1834,85 @@ bool FNiagaraEditorUtilities::AddEmitterContextMenuActions(FMenuBuilder& MenuBui
 			}
 
 			MenuBuilder.AddMenuEntry(
-				LOCTEXT("RemoveParentEmitter", "Remove Parent Emitter"),
-				LOCTEXT("RemoveParentEmitterToolTip", "Remove this emitter's parent, preventing inheritance of any further changes."),
+				LOCTEXT("CreateAssetFromThisEmitter", "Create Asset From This"),
+				LOCTEXT("CreateAssetFromThisEmitterToolTip", "Create an emitter asset from this emitter."),
 				FSlateIcon(),
 				FUIAction(
-					FExecuteAction::CreateSP(EmitterViewModel, &FNiagaraEmitterViewModel::RemoveParentEmitter),
+					FExecuteAction::CreateStatic(&FNiagaraEditorUtilities::CreateAssetFromEmitter, EmitterHandleViewModel),
 					FCanExecuteAction::CreateLambda(
-						[bSingleSelection, bHasParent = EmitterViewModel->HasParentEmitter()]()
+						[bSingleSelection, EmitterHandleViewModel]()
 						{
-							return bSingleSelection && bHasParent;
+							return bSingleSelection && EmitterHandleViewModel->GetOwningSystemEditMode() == ENiagaraSystemViewModelEditMode::SystemAsset;
 						}
 					)
 				)
 			);
+		}
 
+		MenuBuilder.EndSection();
+		MenuBuilder.BeginSection("EmitterActions", LOCTEXT("ParentActions", "Parent Actions"));
+		{
 			MenuBuilder.AddMenuEntry(
-				LOCTEXT("ShowEmitterInContentBrowser", "Show in Content Browser"),
-				LOCTEXT("ShowEmitterInContentBrowserToolTip", "Show the selected emitter in the Content Browser."),
+				LOCTEXT("UpdateParentEmitter", "Update Parent Emitter"),
+				LOCTEXT("UpdateParentEmitterToolTip", "Change or add a parent emitter."),
 				FSlateIcon(),
 				FUIAction(
-					FExecuteAction::CreateStatic(&FNiagaraEditorUtilities::ShowParentEmitterInContentBrowser, EmitterViewModel),
+					FExecuteAction::CreateSP(EmitterViewModel, &FNiagaraEmitterViewModel::CreateNewParentWindow, EmitterHandleViewModel),
 					FCanExecuteAction::CreateLambda(
-						[bSingleSelection, bHasParent = EmitterViewModel->HasParentEmitter()]()
+						[bSingleSelection]()
 						{
-							return bSingleSelection && bHasParent;
+							return bSingleSelection;
 						}
 					)
 				)
 			);
 
-			MenuBuilder.AddMenuEntry(
-				LOCTEXT("CreateAssetFromThisEmitter", "Create Asset From This"),
-				LOCTEXT("CreateAssetFromThisEmitterToolTip", "Create an emitter asset from this emitter."),
-				FSlateIcon(),
-				FUIAction(
-					FExecuteAction::CreateStatic(&FNiagaraEditorUtilities::CreateAssetFromEmitter, EmitterHandleViewModel),
-					FCanExecuteAction::CreateLambda(
-						[bSingleSelection, EmitterHandleViewModel]()
-						{
-							return bSingleSelection && EmitterHandleViewModel->GetOwningSystemEditMode() == ENiagaraSystemViewModelEditMode::SystemAsset;
-						}
-					)
+		MenuBuilder.AddMenuEntry(
+			LOCTEXT("RemoveParentEmitter", "Remove Parent Emitter"),
+			LOCTEXT("RemoveParentEmitterToolTip", "Remove this emitter's parent, preventing inheritance of any further changes."),
+			FSlateIcon(),
+			FUIAction(
+				FExecuteAction::CreateSP(EmitterViewModel, &FNiagaraEmitterViewModel::RemoveParentEmitter),
+				FCanExecuteAction::CreateLambda(
+					[bSingleSelection, bHasParent = EmitterViewModel->HasParentEmitter()]()
+					{
+						return bSingleSelection && bHasParent;
+					}
 				)
-			);
+			)
+		);
+
+		MenuBuilder.AddMenuEntry(
+			LOCTEXT("OpenParentEmitterForEdit", "Open Parent For Edit"),
+			LOCTEXT("OpenParentEmitterForEditToolTip", "Open and Focus Parent Emitter."),
+			FSlateIcon(),
+			FUIAction(
+				FExecuteAction::CreateStatic(&FNiagaraEditorUtilities::OpenParentEmitterForEdit, EmitterViewModel),
+				FCanExecuteAction::CreateLambda(
+					[bSingleSelection, bHasParent = EmitterViewModel->HasParentEmitter()]()
+		{
+			return bSingleSelection && bHasParent;
+		}
+		)
+			)
+		);
+
+		MenuBuilder.AddMenuEntry(
+			LOCTEXT("ShowParentEmitterInContentBrowser", "Show Parent in Content Browser"),
+			LOCTEXT("ShowParentEmitterInContentBrowserToolTip", "Show the selected emitter's parent emitter in the Content Browser."),
+			FSlateIcon(),
+			FUIAction(
+				FExecuteAction::CreateStatic(&FNiagaraEditorUtilities::ShowParentEmitterInContentBrowser, EmitterViewModel),
+				FCanExecuteAction::CreateLambda(
+					[bSingleSelection, bHasParent = EmitterViewModel->HasParentEmitter()]()
+					{
+						return bSingleSelection && bHasParent;
+					}
+				)
+			)
+		);
+
+	
 		}
 		MenuBuilder.EndSection();
 
diff --git a/Engine/Plugins/FX/Niagara/Source/NiagaraEditor/Private/NiagaraScriptMergeManager.cpp b/Engine/Plugins/FX/Niagara/Source/NiagaraEditor/Private/NiagaraScriptMergeManager.cpp
--- a/Engine/Plugins/FX/Niagara/Source/NiagaraEditor/Private/NiagaraScriptMergeManager.cpp
+++ b/Engine/Plugins/FX/Niagara/Source/NiagaraEditor/Private/NiagaraScriptMergeManager.cpp
@@ -978,13 +978,14 @@ FNiagaraScriptMergeManager::FApplyDiffResults FNiagaraScriptMergeManager::Resolv
 	return DiffResults;
 }
 
-INiagaraMergeManager::FMergeEmitterResults FNiagaraScriptMergeManager::MergeEmitter(UNiagaraEmitter& Parent, UNiagaraEmitter& ParentAtLastMerge, UNiagaraEmitter& Instance) const
+INiagaraMergeManager::FMergeEmitterResults FNiagaraScriptMergeManager::MergeEmitter(UNiagaraEmitter& Parent, UNiagaraEmitter* ParentAtLastMerge, UNiagaraEmitter& Instance) const
 {
 	SCOPE_CYCLE_COUNTER(STAT_NiagaraEditor_ScriptMergeManager_MergeEmitter);
 	INiagaraMergeManager::FMergeEmitterResults MergeResults;
-
-	FNiagaraEmitterDiffResults DiffResults = DiffEmitters(ParentAtLastMerge, Instance);
-
+	const bool bNoParentAtLastMerge = (ParentAtLastMerge == nullptr);
+	UNiagaraEmitter* FirstEmitterToDiffAgainst = bNoParentAtLastMerge ? &Parent : ParentAtLastMerge;
+	FNiagaraEmitterDiffResults DiffResults = DiffEmitters(*FirstEmitterToDiffAgainst, Instance);
+	
 	if (DiffResults.IsValid() == false)
 	{
 		MergeResults.MergeResult = INiagaraMergeManager::EMergeEmitterResult::FailedToDiff;
@@ -1040,13 +1041,13 @@ INiagaraMergeManager::FMergeEmitterResults FNiagaraScriptMergeManager::MergeEmit
 		TMap<FGuid, FGuid> LastChangeIds;
 		TMap<FGuid, FGuid> ChangeIdsThatNeedToBeReset;
 		FNiagaraEditorUtilities::GatherChangeIds(Parent, SourceChangeIds, TEXT("Source"));
-		FNiagaraEditorUtilities::GatherChangeIds(ParentAtLastMerge, PreviousSourceChangeIds, TEXT("MergeLast"));
+		FNiagaraEditorUtilities::GatherChangeIds(*FirstEmitterToDiffAgainst, PreviousSourceChangeIds, TEXT("MergeLast"));
 		FNiagaraEditorUtilities::GatherChangeIds(Instance, LastChangeIds, TEXT("Instance"));
 		DiffChangeIds(SourceChangeIds, PreviousSourceChangeIds, LastChangeIds, ChangeIdsThatNeedToBeReset);
 
 		MergeResults.MergeResult = INiagaraMergeManager::EMergeEmitterResult::SucceededDifferencesApplied;
 
-		FApplyDiffResults EmitterSpawnResults = ApplyScriptStackDiff(MergedInstanceAdapter->GetEmitterSpawnStack().ToSharedRef(), DiffResults.EmitterSpawnDiffResults);
+		FApplyDiffResults EmitterSpawnResults = ApplyScriptStackDiff(MergedInstanceAdapter->GetEmitterSpawnStack().ToSharedRef(), DiffResults.EmitterSpawnDiffResults, bNoParentAtLastMerge);
 		if (EmitterSpawnResults.bSucceeded == false)
 		{
 			MergeResults.MergeResult = INiagaraMergeManager::EMergeEmitterResult::FailedToMerge;
@@ -1054,7 +1055,7 @@ INiagaraMergeManager::FMergeEmitterResults FNiagaraScriptMergeManager::MergeEmit
 		MergeResults.bModifiedGraph |= EmitterSpawnResults.bModifiedGraph;
 		MergeResults.ErrorMessages.Append(EmitterSpawnResults.ErrorMessages);
 
-		FApplyDiffResults EmitterUpdateResults = ApplyScriptStackDiff(MergedInstanceAdapter->GetEmitterUpdateStack().ToSharedRef(), DiffResults.EmitterUpdateDiffResults);
+		FApplyDiffResults EmitterUpdateResults = ApplyScriptStackDiff(MergedInstanceAdapter->GetEmitterUpdateStack().ToSharedRef(), DiffResults.EmitterUpdateDiffResults, bNoParentAtLastMerge);
 		if (EmitterUpdateResults.bSucceeded == false)
 		{
 			MergeResults.MergeResult = INiagaraMergeManager::EMergeEmitterResult::FailedToMerge;
@@ -1062,7 +1063,7 @@ INiagaraMergeManager::FMergeEmitterResults FNiagaraScriptMergeManager::MergeEmit
 		MergeResults.bModifiedGraph |= EmitterUpdateResults.bModifiedGraph;
 		MergeResults.ErrorMessages.Append(EmitterUpdateResults.ErrorMessages);
 
-		FApplyDiffResults ParticleSpawnResults = ApplyScriptStackDiff(MergedInstanceAdapter->GetParticleSpawnStack().ToSharedRef(), DiffResults.ParticleSpawnDiffResults);
+		FApplyDiffResults ParticleSpawnResults = ApplyScriptStackDiff(MergedInstanceAdapter->GetParticleSpawnStack().ToSharedRef(), DiffResults.ParticleSpawnDiffResults, bNoParentAtLastMerge);
 		if (ParticleSpawnResults.bSucceeded == false)
 		{
 			MergeResults.MergeResult = INiagaraMergeManager::EMergeEmitterResult::FailedToMerge;
@@ -1070,7 +1071,7 @@ INiagaraMergeManager::FMergeEmitterResults FNiagaraScriptMergeManager::MergeEmit
 		MergeResults.bModifiedGraph |= ParticleSpawnResults.bModifiedGraph;
 		MergeResults.ErrorMessages.Append(ParticleSpawnResults.ErrorMessages);
 
-		FApplyDiffResults ParticleUpdateResults = ApplyScriptStackDiff(MergedInstanceAdapter->GetParticleUpdateStack().ToSharedRef(), DiffResults.ParticleUpdateDiffResults);
+		FApplyDiffResults ParticleUpdateResults = ApplyScriptStackDiff(MergedInstanceAdapter->GetParticleUpdateStack().ToSharedRef(), DiffResults.ParticleUpdateDiffResults, bNoParentAtLastMerge);
 		if (ParticleUpdateResults.bSucceeded == false)
 		{
 			MergeResults.MergeResult = INiagaraMergeManager::EMergeEmitterResult::FailedToMerge;
@@ -1078,7 +1079,7 @@ INiagaraMergeManager::FMergeEmitterResults FNiagaraScriptMergeManager::MergeEmit
 		MergeResults.bModifiedGraph |= ParticleUpdateResults.bModifiedGraph;
 		MergeResults.ErrorMessages.Append(ParticleUpdateResults.ErrorMessages);
 
-		FApplyDiffResults EventHandlerResults = ApplyEventHandlerDiff(MergedInstanceAdapter, DiffResults);
+		FApplyDiffResults EventHandlerResults = ApplyEventHandlerDiff(MergedInstanceAdapter, DiffResults, bNoParentAtLastMerge);
 		if (EventHandlerResults.bSucceeded == false)
 		{
 			MergeResults.MergeResult = INiagaraMergeManager::EMergeEmitterResult::FailedToMerge;
@@ -1086,7 +1087,7 @@ INiagaraMergeManager::FMergeEmitterResults FNiagaraScriptMergeManager::MergeEmit
 		MergeResults.bModifiedGraph |= EventHandlerResults.bModifiedGraph;
 		MergeResults.ErrorMessages.Append(EventHandlerResults.ErrorMessages);
 
-		FApplyDiffResults ShaderStageResults = ApplyShaderStageDiff(MergedInstanceAdapter, DiffResults);
+		FApplyDiffResults ShaderStageResults = ApplyShaderStageDiff(MergedInstanceAdapter, DiffResults, bNoParentAtLastMerge);
 		if (ShaderStageResults.bSucceeded == false)
 		{
 			MergeResults.MergeResult = INiagaraMergeManager::EMergeEmitterResult::FailedToMerge;
@@ -1094,7 +1095,7 @@ INiagaraMergeManager::FMergeEmitterResults FNiagaraScriptMergeManager::MergeEmit
 		MergeResults.bModifiedGraph |= ShaderStageResults.bModifiedGraph;
 		MergeResults.ErrorMessages.Append(ShaderStageResults.ErrorMessages);
 
-		FApplyDiffResults RendererResults = ApplyRendererDiff(*MergedInstance, DiffResults);
+		FApplyDiffResults RendererResults = ApplyRendererDiff(*MergedInstance, DiffResults, bNoParentAtLastMerge);
 		if (RendererResults.bSucceeded == false)
 		{
 			MergeResults.MergeResult = INiagaraMergeManager::EMergeEmitterResult::FailedToMerge;
@@ -1286,7 +1287,7 @@ FNiagaraScriptMergeManager::FApplyDiffResults FNiagaraScriptMergeManager::ResetM
 	ResetDiffResults.ModifiedBaseInputOverrides.RemoveAll(FindUnrelatedInputOverrides);
 	ResetDiffResults.ModifiedOtherInputOverrides.RemoveAll(FindUnrelatedInputOverrides);
 
-	return ApplyScriptStackDiff(EmitterAdapter->GetScriptStack(ScriptUsage, ScriptUsageId).ToSharedRef(), ResetDiffResults);
+	return ApplyScriptStackDiff(EmitterAdapter->GetScriptStack(ScriptUsage, ScriptUsageId).ToSharedRef(), ResetDiffResults, false);
 }
 
 bool FNiagaraScriptMergeManager::HasBaseEventHandler(const UNiagaraEmitter& BaseEmitter, FGuid EventScriptUsageId)
@@ -1425,7 +1426,7 @@ void FNiagaraScriptMergeManager::ResetRendererToBase(UNiagaraEmitter& Emitter, c
 	ResetDiffResults.ModifiedBaseRenderers.RemoveAll(FindUnrelatedRenderers);
 	ResetDiffResults.ModifiedOtherRenderers.RemoveAll(FindUnrelatedRenderers);
 
-	ApplyRendererDiff(Emitter, ResetDiffResults);
+	ApplyRendererDiff(Emitter, ResetDiffResults, false);
 }
 
 bool FNiagaraScriptMergeManager::IsEmitterEditablePropertySetDifferentFromBase(const UNiagaraEmitter& Emitter, const UNiagaraEmitter& BaseEmitter)
@@ -2162,7 +2163,7 @@ void FNiagaraScriptMergeManager::CopyPropertiesToBase(void* BaseDataAddress, con
 	}
 }
 
-FNiagaraScriptMergeManager::FApplyDiffResults FNiagaraScriptMergeManager::ApplyScriptStackDiff(TSharedRef<FNiagaraScriptStackMergeAdapter> BaseScriptStackAdapter, const FNiagaraScriptStackDiffResults& DiffResults) const
+FNiagaraScriptMergeManager::FApplyDiffResults FNiagaraScriptMergeManager::ApplyScriptStackDiff(TSharedRef<FNiagaraScriptStackMergeAdapter> BaseScriptStackAdapter, const FNiagaraScriptStackDiffResults& DiffResults, const bool bNoParentAtLastMerge) const
 {
 	FApplyDiffResults Results;
 
@@ -2187,17 +2188,21 @@ FNiagaraScriptMergeManager::FApplyDiffResults FNiagaraScriptMergeManager::ApplyS
 	TArray<TSharedRef<FNiagaraStackFunctionMergeAdapter>> EnableModules;
 	TArray<TSharedRef<FNiagaraStackFunctionMergeAdapter>> DisableModules;
 
-	for (TSharedRef<FNiagaraStackFunctionMergeAdapter> RemovedModule : DiffResults.RemovedBaseModules)
+	if (!bNoParentAtLastMerge)
 	{
-		TSharedPtr<FNiagaraStackFunctionMergeAdapter> MatchingModuleAdapter = BaseScriptStackAdapter->GetModuleFunctionById(RemovedModule->GetFunctionCallNode()->NodeGuid);
-		if (MatchingModuleAdapter.IsValid())
+		for (TSharedRef<FNiagaraStackFunctionMergeAdapter> RemovedModule : DiffResults.RemovedBaseModules)
 		{
-			RemoveModules.Add(MatchingModuleAdapter.ToSharedRef());
+			TSharedPtr<FNiagaraStackFunctionMergeAdapter> MatchingModuleAdapter = BaseScriptStackAdapter->GetModuleFunctionById(RemovedModule->GetFunctionCallNode()->NodeGuid);
+			if (MatchingModuleAdapter.IsValid())
+			{
+				RemoveModules.Add(MatchingModuleAdapter.ToSharedRef());
+			}
 		}
 	}
 
 	AddModules.Append(DiffResults.AddedOtherModules);
 
+
 	for (TSharedRef<FNiagaraStackFunctionInputOverrideMergeAdapter> RemovedInputOverrideAdapter : DiffResults.RemovedBaseInputOverrides)
 	{
 		TSharedPtr<FNiagaraStackFunctionMergeAdapter> MatchingModuleAdapter = BaseScriptStackAdapter->GetModuleFunctionById(RemovedInputOverrideAdapter->GetOwningFunctionCall()->NodeGuid);
@@ -2293,12 +2298,15 @@ FNiagaraScriptMergeManager::FApplyDiffResults FNiagaraScriptMergeManager::ApplyS
 		Results.ErrorMessages.Append(AddModuleResults.ErrorMessages);
 	}
 
-	for (TSharedRef<FNiagaraStackFunctionInputOverrideMergeAdapter> RemoveInputOverrideItem : RemoveInputOverrides)
+	if (!bNoParentAtLastMerge)
 	{
-		FApplyDiffResults RemoveInputOverrideResults = RemoveInputOverride(*BaseScriptStackAdapter->GetScript(), RemoveInputOverrideItem);
-		Results.bSucceeded &= RemoveInputOverrideResults.bSucceeded;
-		Results.bModifiedGraph |= RemoveInputOverrideResults.bModifiedGraph;
-		Results.ErrorMessages.Append(RemoveInputOverrideResults.ErrorMessages);
+		for (TSharedRef<FNiagaraStackFunctionInputOverrideMergeAdapter> RemoveInputOverrideItem : RemoveInputOverrides)
+		{
+			FApplyDiffResults RemoveInputOverrideResults = RemoveInputOverride(*BaseScriptStackAdapter->GetScript(), RemoveInputOverrideItem);
+			Results.bSucceeded &= RemoveInputOverrideResults.bSucceeded;
+			Results.bModifiedGraph |= RemoveInputOverrideResults.bModifiedGraph;
+			Results.ErrorMessages.Append(RemoveInputOverrideResults.ErrorMessages);
+		}
 	}
 
 	for (const FAddInputOverrideActionData& AddInputOverrideActionData : AddInputOverrideActionDatas)
@@ -2322,11 +2330,12 @@ FNiagaraScriptMergeManager::FApplyDiffResults FNiagaraScriptMergeManager::ApplyS
 	return Results;
 }
 
-FNiagaraScriptMergeManager::FApplyDiffResults FNiagaraScriptMergeManager::ApplyEventHandlerDiff(TSharedRef<FNiagaraEmitterMergeAdapter> BaseEmitterAdapter, const FNiagaraEmitterDiffResults& DiffResults) const
+FNiagaraScriptMergeManager::FApplyDiffResults FNiagaraScriptMergeManager::ApplyEventHandlerDiff(TSharedRef<FNiagaraEmitterMergeAdapter> BaseEmitterAdapter, const FNiagaraEmitterDiffResults& DiffResults, const bool bNoParentAtLastMerge) const
 {
 	FApplyDiffResults Results;
 	if (DiffResults.RemovedBaseEventHandlers.Num() > 0)
 	{
+		// If this becomes supported, it needs to handle the bNoParentAtLastMerge case
 		Results.bSucceeded = false;
 		Results.bModifiedGraph = false;
 		Results.ErrorMessages.Add(LOCTEXT("RemovedEventHandlersUnsupported", "Apply diff failed, removed event handlers are currently unsupported."));
@@ -2361,7 +2370,7 @@ FNiagaraScriptMergeManager::FApplyDiffResults FNiagaraScriptMergeManager::ApplyE
 				}
 				if (ModifiedEventHandler.ScriptDiffResults.IsEmpty() == false)
 				{
-					FApplyDiffResults ApplyEventHandlerStackDiffResults = ApplyScriptStackDiff(MatchingBaseEventHandlerAdapter->GetEventStack().ToSharedRef(), ModifiedEventHandler.ScriptDiffResults);
+					FApplyDiffResults ApplyEventHandlerStackDiffResults = ApplyScriptStackDiff(MatchingBaseEventHandlerAdapter->GetEventStack().ToSharedRef(), ModifiedEventHandler.ScriptDiffResults, bNoParentAtLastMerge);
 					Results.bSucceeded &= ApplyEventHandlerStackDiffResults.bSucceeded;
 					Results.bModifiedGraph |= ApplyEventHandlerStackDiffResults.bModifiedGraph;
 					Results.ErrorMessages.Append(ApplyEventHandlerStackDiffResults.ErrorMessages);
@@ -2422,13 +2431,14 @@ FNiagaraScriptMergeManager::FApplyDiffResults FNiagaraScriptMergeManager::ApplyE
 	return Results;
 }
 
-FNiagaraScriptMergeManager::FApplyDiffResults FNiagaraScriptMergeManager::ApplyShaderStageDiff(TSharedRef<FNiagaraEmitterMergeAdapter> BaseEmitterAdapter, const FNiagaraEmitterDiffResults& DiffResults) const
+FNiagaraScriptMergeManager::FApplyDiffResults FNiagaraScriptMergeManager::ApplyShaderStageDiff(TSharedRef<FNiagaraEmitterMergeAdapter> BaseEmitterAdapter, const FNiagaraEmitterDiffResults& DiffResults, const bool bNoParentAtLastMerge) const
 {
 	FApplyDiffResults Results;
 	if (DiffResults.RemovedBaseShaderStages.Num() > 0)
 	{
 		Results.bSucceeded = false;
 		Results.bModifiedGraph = false;
+		// If this becomes supported, it needs to handle the bNoParentAtLastMerge case
 		Results.ErrorMessages.Add(LOCTEXT("RemovedShaderStagesUnsupported", "Apply diff failed, removed shader stages are currently unsupported."));
 		return Results;
 	}
@@ -2460,7 +2470,7 @@ FNiagaraScriptMergeManager::FApplyDiffResults FNiagaraScriptMergeManager::ApplyS
 				}
 				if (ModifiedShaderStage.ScriptDiffResults.IsEmpty() == false)
 				{
-					FApplyDiffResults ApplyShaderStageStackDiffResults = ApplyScriptStackDiff(MatchingBaseShaderStageAdapter->GetShaderStageStack().ToSharedRef(), ModifiedShaderStage.ScriptDiffResults);
+					FApplyDiffResults ApplyShaderStageStackDiffResults = ApplyScriptStackDiff(MatchingBaseShaderStageAdapter->GetShaderStageStack().ToSharedRef(), ModifiedShaderStage.ScriptDiffResults, bNoParentAtLastMerge);
 					Results.bSucceeded &= ApplyShaderStageStackDiffResults.bSucceeded;
 					Results.bModifiedGraph |= ApplyShaderStageStackDiffResults.bModifiedGraph;
 					Results.ErrorMessages.Append(ApplyShaderStageStackDiffResults.ErrorMessages);
@@ -2521,18 +2531,21 @@ FNiagaraScriptMergeManager::FApplyDiffResults FNiagaraScriptMergeManager::ApplyS
 	return Results;
 }
 
-FNiagaraScriptMergeManager::FApplyDiffResults FNiagaraScriptMergeManager::ApplyRendererDiff(UNiagaraEmitter& BaseEmitter, const FNiagaraEmitterDiffResults& DiffResults) const
+FNiagaraScriptMergeManager::FApplyDiffResults FNiagaraScriptMergeManager::ApplyRendererDiff(UNiagaraEmitter& BaseEmitter, const FNiagaraEmitterDiffResults& DiffResults, const bool bNoParentAtLastMerge) const
 {
 	TArray<UNiagaraRendererProperties*> RenderersToRemove;
 	TArray<UNiagaraRendererProperties*> RenderersToAdd;
 
-	for (TSharedRef<FNiagaraRendererMergeAdapter> RemovedRenderer : DiffResults.RemovedBaseRenderers)
+	if (!bNoParentAtLastMerge)
 	{
-		auto FindRendererByMergeId = [=](UNiagaraRendererProperties* Renderer) { return Renderer->GetMergeId() == RemovedRenderer->GetRenderer()->GetMergeId(); };
-		UNiagaraRendererProperties*const* MatchingRendererPtr = BaseEmitter.GetRenderers().FindByPredicate(FindRendererByMergeId);
-		if (MatchingRendererPtr != nullptr)
+		for (TSharedRef<FNiagaraRendererMergeAdapter> RemovedRenderer : DiffResults.RemovedBaseRenderers)
 		{
-			RenderersToRemove.Add(*MatchingRendererPtr);
+			auto FindRendererByMergeId = [=](UNiagaraRendererProperties* Renderer) { return Renderer->GetMergeId() == RemovedRenderer->GetRenderer()->GetMergeId(); };
+			UNiagaraRendererProperties* const* MatchingRendererPtr = BaseEmitter.GetRenderers().FindByPredicate(FindRendererByMergeId);
+			if (MatchingRendererPtr != nullptr)
+			{
+				RenderersToRemove.Add(*MatchingRendererPtr);
+			}
 		}
 	}
 
diff --git a/Engine/Plugins/FX/Niagara/Source/NiagaraEditor/Private/NiagaraScriptMergeManager.h b/Engine/Plugins/FX/Niagara/Source/NiagaraEditor/Private/NiagaraScriptMergeManager.h
--- a/Engine/Plugins/FX/Niagara/Source/NiagaraEditor/Private/NiagaraScriptMergeManager.h
+++ b/Engine/Plugins/FX/Niagara/Source/NiagaraEditor/Private/NiagaraScriptMergeManager.h
@@ -346,7 +346,7 @@ public:
 	};
 
 public:
-	virtual INiagaraMergeManager::FMergeEmitterResults MergeEmitter(UNiagaraEmitter& Parent, UNiagaraEmitter& ParentAtLastMerge, UNiagaraEmitter& Instance) const override;
+	virtual INiagaraMergeManager::FMergeEmitterResults MergeEmitter(UNiagaraEmitter& Parent, UNiagaraEmitter* ParentAtLastMerge, UNiagaraEmitter& Instance) const override;
 
 	static TSharedRef<FNiagaraScriptMergeManager> Get();
 
@@ -402,13 +402,13 @@ public:
 private:
 	TOptional<bool> DoFunctionInputOverridesMatch(TSharedRef<FNiagaraStackFunctionInputOverrideMergeAdapter> BaseFunctionInputAdapter, TSharedRef<FNiagaraStackFunctionInputOverrideMergeAdapter> OtherFunctionInputAdapter) const;
 
-	FApplyDiffResults ApplyScriptStackDiff(TSharedRef<FNiagaraScriptStackMergeAdapter> BaseScriptStackAdapter, const FNiagaraScriptStackDiffResults& DiffResults) const;
+	FApplyDiffResults ApplyScriptStackDiff(TSharedRef<FNiagaraScriptStackMergeAdapter> BaseScriptStackAdapter, const FNiagaraScriptStackDiffResults& DiffResults, const bool bNoParentAtLastMerge) const;
 
-	FApplyDiffResults ApplyEventHandlerDiff(TSharedRef<FNiagaraEmitterMergeAdapter> BaseEmitterAdapter, const FNiagaraEmitterDiffResults& DiffResults) const;
+	FApplyDiffResults ApplyEventHandlerDiff(TSharedRef<FNiagaraEmitterMergeAdapter> BaseEmitterAdapter, const FNiagaraEmitterDiffResults& DiffResults, const bool bNoParentAtLastMerge) const;
 
-	FApplyDiffResults ApplyShaderStageDiff(TSharedRef<FNiagaraEmitterMergeAdapter> BaseEmitterAdapter, const FNiagaraEmitterDiffResults& DiffResults) const;
+	FApplyDiffResults ApplyShaderStageDiff(TSharedRef<FNiagaraEmitterMergeAdapter> BaseEmitterAdapter, const FNiagaraEmitterDiffResults& DiffResults, const bool bNoParentAtLastMerge) const;
 
-	FApplyDiffResults ApplyRendererDiff(UNiagaraEmitter& BaseEmitter, const FNiagaraEmitterDiffResults& DiffResults) const;
+	FApplyDiffResults ApplyRendererDiff(UNiagaraEmitter& BaseEmitter, const FNiagaraEmitterDiffResults& DiffResults, const bool bNoParentAtLastMerge) const;
 
 	FApplyDiffResults ApplyStackEntryDisplayNameDiffs(UNiagaraEmitter& Emitter, const FNiagaraEmitterDiffResults& DiffResults) const;
 
diff --git a/Engine/Plugins/FX/Niagara/Source/NiagaraEditor/Private/Toolkits/NiagaraSystemToolkit.cpp b/Engine/Plugins/FX/Niagara/Source/NiagaraEditor/Private/Toolkits/NiagaraSystemToolkit.cpp
--- a/Engine/Plugins/FX/Niagara/Source/NiagaraEditor/Private/Toolkits/NiagaraSystemToolkit.cpp
+++ b/Engine/Plugins/FX/Niagara/Source/NiagaraEditor/Private/Toolkits/NiagaraSystemToolkit.cpp
@@ -1128,7 +1128,7 @@ void FNiagaraSystemToolkit::UpdateOriginalEmitter()
 		{
 			checkfSlow(Script->AreScriptAndSourceSynchronized(), TEXT("Editable Emitter Script change ID is out of date when applying to Original Emitter!"));
 		}
-
+		Emitter->PreEditChange(nullptr);
 		// overwrite the original script in place by constructing a new one with the same name
 		Emitter = (UNiagaraEmitter*)StaticDuplicateObject(EditableEmitter, Emitter->GetOuter(),
 			Emitter->GetFName(), RF_AllFlags, Emitter->GetClass());
diff --git a/Engine/Plugins/FX/Niagara/Source/NiagaraEditor/Private/ViewModels/NiagaraEmitterViewModel.cpp b/Engine/Plugins/FX/Niagara/Source/NiagaraEditor/Private/ViewModels/NiagaraEmitterViewModel.cpp
--- a/Engine/Plugins/FX/Niagara/Source/NiagaraEditor/Private/ViewModels/NiagaraEmitterViewModel.cpp
+++ b/Engine/Plugins/FX/Niagara/Source/NiagaraEditor/Private/ViewModels/NiagaraEmitterViewModel.cpp
@@ -17,6 +17,12 @@
 #include "NiagaraSystemInstance.h"
 
 #include "ScopedTransaction.h"
+#include "IContentBrowserSingleton.h"
+#include "Framework/Application/SlateApplication.h"
+#include "ContentBrowserModule.h"
+#include "Framework/Docking/TabManager.h"
+#include "Widgets/SWindow.h"
+#include "ViewModels/NiagaraEmitterHandleViewModel.h"
 
 #define LOCTEXT_NAMESPACE "EmitterEditorViewModel"
 
@@ -42,6 +48,11 @@ FNiagaraEmitterViewModel::FNiagaraEmitterViewModel()
 
 void FNiagaraEmitterViewModel::Cleanup()
 {
+	if (NewParentWindow.IsValid())
+	{
+		NewParentWindow.Reset();
+	}
+
 	if (Emitter.IsValid())
 	{
 		Emitter->OnEmitterVMCompiled().RemoveAll(this);
@@ -153,6 +164,61 @@ FText FNiagaraEmitterViewModel::GetParentPathNameText() const
 	return FText();
 }
 
+void FNiagaraEmitterViewModel::CreateNewParentWindow(TSharedRef<FNiagaraEmitterHandleViewModel> EmitterHandleViewModel)
+{
+	FAssetPickerConfig AssetPickerConfig;
+	AssetPickerConfig.SelectionMode = ESelectionMode::SingleToggle;
+	AssetPickerConfig.InitialAssetViewType = EAssetViewType::List;
+	AssetPickerConfig.Filter.ClassNames.Add(UNiagaraEmitter::StaticClass()->GetFName());
+	AssetPickerConfig.OnAssetsActivated.BindSP(this, &FNiagaraEmitterViewModel::UpdateParentEmitter, EmitterHandleViewModel);
+
+	FContentBrowserModule& ContentBrowserModule = FModuleManager::Get().LoadModuleChecked<FContentBrowserModule>(TEXT("ContentBrowser"));
+	TSharedRef<SWidget> AssetPicker = ContentBrowserModule.Get().CreateAssetPicker(AssetPickerConfig);
+
+	const FText TitleText = LOCTEXT("NewParentEmitter", "New Parent Emitter");
+	// Create the window to pick the class
+	SAssignNew(NewParentWindow, SWindow)
+		.Title(TitleText)
+		.SizingRule(ESizingRule::UserSized)
+		.ClientSize(FVector2D(700.f, 300.f))
+		.AutoCenter(EAutoCenter::PreferredWorkArea)
+		.SupportsMinimize(false);
+	NewParentWindow->SetContent(AssetPicker);
+	TSharedPtr<SWindow> RootWindow = FGlobalTabmanager::Get()->GetRootWindow();
+	if (RootWindow.IsValid())
+	{
+		FSlateApplication::Get().AddWindowAsNativeChild(NewParentWindow.ToSharedRef(), RootWindow.ToSharedRef());
+	}
+	else
+	{
+		FSlateApplication::Get().AddWindow(NewParentWindow.ToSharedRef());
+	}
+}
+
+void FNiagaraEmitterViewModel::UpdateParentEmitter(const TArray<FAssetData> & ActivatedAssets, EAssetTypeActivationMethod::Type ActivationMethod, TSharedRef<FNiagaraEmitterHandleViewModel> EmitterHandleViewModel)
+{
+	FScopedTransaction ScopedTransaction(LOCTEXT("RemoveParentEmitterTransaction", "Remove Parent Emitter"));
+	if ((ActivationMethod == EAssetTypeActivationMethod::DoubleClicked || ActivationMethod == EAssetTypeActivationMethod::Opened) && ActivatedAssets.Num() == 1)
+	{
+		if (Emitter.IsValid())
+		{
+			FAssetData NewParentData = ActivatedAssets[0];
+			UNiagaraEmitter* NewParentPtr = Cast<UNiagaraEmitter>(NewParentData.GetAsset());
+			if (NewParentPtr && NewParentPtr != Emitter)
+			{
+				Emitter->PreEditChange(nullptr);
+				Emitter->Reparent(*NewParentPtr);
+				Emitter->MergeChangesFromParent();
+				Emitter->PostEditChange();
+				EmitterHandleViewModel->GetOwningSystemViewModel()->RefreshAll();
+				NewParentWindow->RequestDestroyWindow();
+			}
+			
+		}
+		
+	}
+}
+
 void FNiagaraEmitterViewModel::RemoveParentEmitter()
 {
 	FScopedTransaction ScopedTransaction(LOCTEXT("RemoveParentEmitterTransaction", "Remove Parent Emitter"));
diff --git a/Engine/Plugins/FX/Niagara/Source/NiagaraEditor/Public/NiagaraEditorUtilities.h b/Engine/Plugins/FX/Niagara/Source/NiagaraEditor/Public/NiagaraEditorUtilities.h
--- a/Engine/Plugins/FX/Niagara/Source/NiagaraEditor/Public/NiagaraEditorUtilities.h
+++ b/Engine/Plugins/FX/Niagara/Source/NiagaraEditor/Public/NiagaraEditorUtilities.h
@@ -226,6 +226,7 @@ namespace FNiagaraEditorUtilities
 
 	void ShowParentEmitterInContentBrowser(TSharedRef<FNiagaraEmitterViewModel> EmitterViewModel);
 
+	void OpenParentEmitterForEdit(TSharedRef<FNiagaraEmitterViewModel> Emitter);
 	ECheckBoxState GetSelectedEmittersEnabledCheckState(TSharedRef<FNiagaraSystemViewModel> SystemViewModel);
 	void ToggleSelectedEmittersEnabled(TSharedRef<FNiagaraSystemViewModel> SystemViewModel);
 
diff --git a/Engine/Plugins/FX/Niagara/Source/NiagaraEditor/Public/ViewModels/NiagaraEmitterViewModel.h b/Engine/Plugins/FX/Niagara/Source/NiagaraEditor/Public/ViewModels/NiagaraEmitterViewModel.h
--- a/Engine/Plugins/FX/Niagara/Source/NiagaraEditor/Public/ViewModels/NiagaraEmitterViewModel.h
+++ b/Engine/Plugins/FX/Niagara/Source/NiagaraEditor/Public/ViewModels/NiagaraEmitterViewModel.h
@@ -5,6 +5,7 @@
 #include "NiagaraCommon.h"
 #include "ViewModels/TNiagaraViewModelManager.h"
 #include "UObject/ObjectKey.h"
+#include "IAssetTypeActions.h"
 
 class UNiagaraEmitter;
 class UNiagaraScript;
@@ -16,6 +17,8 @@ class FNiagaraEmitterInstance;
 struct FNiagaraVariable;
 struct FNiagaraParameterStore;
 struct FEdGraphEditAction;
+class SWindow;
+class FNiagaraEmitterHandleViewModel;
 
 /** The view model for the UNiagaraEmitter objects */
 class FNiagaraEmitterViewModel : public TSharedFromThis<FNiagaraEmitterViewModel>,  public TNiagaraViewModelManager<UNiagaraEmitter, FNiagaraEmitterViewModel>
@@ -60,6 +63,9 @@ public:
 	/** Gets the text representation of the parent emitter path. */
 	NIAGARAEDITOR_API FText GetParentPathNameText() const;
 
+	NIAGARAEDITOR_API void CreateNewParentWindow(TSharedRef<FNiagaraEmitterHandleViewModel> EmitterHandleViewModel);
+	void UpdateParentEmitter(const TArray<FAssetData>& ActivatedAssets, EAssetTypeActivationMethod::Type ActivationMethod, TSharedRef<FNiagaraEmitterHandleViewModel> EmitterHandleViewModel);
+
 	/** Removes the parent emitter from this emitter. */
 	NIAGARAEDITOR_API void RemoveParentEmitter();
 
@@ -162,4 +168,6 @@ private:
 
 	/** A mapping of script to the delegate handle for it's on parameter map changed delegate. */
 	TMap<FObjectKey, FDelegateHandle> ScriptToOnParameterStoreChangedHandleMap;
+
+	TSharedPtr<SWindow> NewParentWindow;
 };
\ No newline at end of file
diff --git a/Engine/Plugins/FX/Niagara/Source/NiagaraEditor/Public/ViewModels/NiagaraSystemViewModel.h b/Engine/Plugins/FX/Niagara/Source/NiagaraEditor/Public/ViewModels/NiagaraSystemViewModel.h
--- a/Engine/Plugins/FX/Niagara/Source/NiagaraEditor/Public/ViewModels/NiagaraSystemViewModel.h
+++ b/Engine/Plugins/FX/Niagara/Source/NiagaraEditor/Public/ViewModels/NiagaraSystemViewModel.h
@@ -241,7 +241,7 @@ public:
 	void Cleanup();
 
 	/** Reinitializes all System instances, and rebuilds emitter handle view models and tracks. */
-	void RefreshAll();
+	NIAGARAEDITOR_API void RefreshAll();
 
 	/** Called to notify the system view model that one of the data objects in the system was modified. */
 	void NotifyDataObjectChanged(UObject* ChangedObject);
diff --git a/Engine/Plugins/FX/Niagara/Source/NiagaraEditorWidgets/Private/SNiagaraStack.cpp b/Engine/Plugins/FX/Niagara/Source/NiagaraEditorWidgets/Private/SNiagaraStack.cpp
--- a/Engine/Plugins/FX/Niagara/Source/NiagaraEditorWidgets/Private/SNiagaraStack.cpp
+++ b/Engine/Plugins/FX/Niagara/Source/NiagaraEditorWidgets/Private/SNiagaraStack.cpp
@@ -54,6 +54,7 @@
 #include "NiagaraStackCommandContext.h"
 #include "NiagaraEditorUtilities.h"
 #include "Subsystems/AssetEditorSubsystem.h"
+#include "EditorFontGlyphs.h"
 
 #define LOCTEXT_NAMESPACE "NiagaraStack"
 
@@ -63,6 +64,7 @@ public:
 	SLATE_BEGIN_ARGS(SNiagaraStackEmitterHeader) {}
 		SLATE_ATTRIBUTE(EVisibility, IssueIconVisibility);
 		SLATE_EVENT(FSimpleDelegate, OnCycleThroughIssues);
+		SLATE_ARGUMENT(TSharedPtr<SNiagaraStack>, ParentStack)
 	SLATE_END_ARGS();
 
 	void Construct(const FArguments& InArgs, TSharedRef<FNiagaraEmitterHandleViewModel> InEmitterHandleViewModel, UNiagaraStackEntry* InRootEntry, UNiagaraStackViewModel* InStackViewModel)
@@ -71,7 +73,7 @@ public:
 		OnCycleThroughIssues = InArgs._OnCycleThroughIssues;
 		StackViewModel = InStackViewModel;
 		TopLevelViewModel = StackViewModel->GetTopLevelViewModelForEntry(*InRootEntry);
-
+		ParentStackPtr = InArgs._ParentStack;
 		ChildSlot
 		[
 			SNew(SVerticalBox)
@@ -137,28 +139,27 @@ public:
 					.Visibility(InArgs._IssueIconVisibility)
 					.OnClicked(this, &SNiagaraStackEmitterHeader::OnIssueIconClicked)
 				]
-				// Open and Focus Source Emitter
 				+ SHorizontalBox::Slot()
 				.AutoWidth()
 				.HAlign(HAlign_Fill)
 				.VAlign(VAlign_Center)
 				.Padding(0, 0, 2, 0)
 				[
-					SNew(SButton)
-					.IsFocusable(false)
-					.ToolTipText(LOCTEXT("OpenAndFocusParentEmitterToolTip", "Open and Focus Parent Emitter"))
-					.ButtonStyle(FEditorStyle::Get(), "HoverHintOnly")
-					.ForegroundColor(FNiagaraEditorWidgetsStyle::Get().GetColor("NiagaraEditor.Stack.ForegroundColor"))
-					.ContentPadding(2)
-					.OnClicked(this, &SNiagaraStackEmitterHeader::OpenParentEmitter)
-					.Visibility(this, &SNiagaraStackEmitterHeader::GetOpenSourceEmitterVisibility)
-					// GoToSource icon is 30x30px so we scale it down to stay in line with other 12x12px UI
-					.DesiredSizeScale(FVector2D(0.55f, 0.55f))
-					.Content()
+					SAssignNew(SettingsAnchor, SMenuAnchor)
+					.Placement(MenuPlacement_MenuLeft)
+					.OnGetMenuContent(this, &SNiagaraStackEmitterHeader::OnGetContent)
 					[
-						SNew(SImage)
-						.Image(FNiagaraEditorWidgetsStyle::Get().GetBrush("NiagaraEditor.Stack.GoToSourceIcon"))
-						.ColorAndOpacity(FNiagaraEditorWidgetsStyle::Get().GetColor("NiagaraEditor.Stack.FlatButtonColor"))
+						SNew(SButton)
+						.ForegroundColor(FSlateColor::UseForeground())
+						.ButtonStyle(FEditorStyle::Get(), TEXT("FlatButton"))
+						.OnClicked(this, &SNiagaraStackEmitterHeader::OpenSubmenu)
+						.VAlign(VAlign_Center)
+						.HAlign(HAlign_Fill)
+						[
+							SNew(STextBlock)
+							.Font(FEditorStyle::Get().GetFontStyle("FontAwesome.16"))
+							.Text(FEditorFontGlyphs::Cogs)
+						]
 					]
 				]
 			]
@@ -211,17 +212,6 @@ private:
 		return bIsRenamed ? EVisibility::Visible : EVisibility::Collapsed;
 	}
 
-
-	FReply OpenParentEmitter()
-	{
-		UNiagaraEmitter* ParentEmitter = const_cast<UNiagaraEmitter*>(EmitterHandleViewModel->GetEmitterViewModel()->GetParentEmitter());
-		if (ParentEmitter != nullptr)
-		{
-			GEditor->GetEditorSubsystem<UAssetEditorSubsystem>()->OpenEditorForAsset(ParentEmitter);
-		}
-		return FReply::Handled();
-	}
-
 	EVisibility GetOpenSourceEmitterVisibility() const
 	{
 		return EmitterHandleViewModel->GetEmitterViewModel()->GetEmitter()->GetParent() != nullptr ? EVisibility::Visible : EVisibility::Collapsed;
@@ -234,6 +224,21 @@ private:
 		return FReply::Handled();
 	}
 
+	FReply OpenSubmenu()
+	{
+		SettingsAnchor->SetIsOpen(!SettingsAnchor->IsOpen());
+		return FReply::Handled();
+	}
+
+	TSharedRef<SWidget> OnGetContent() const
+	{
+		if (ParentStackPtr.IsValid())
+		{
+			return ParentStackPtr.Pin()->GenerateStackMenu(TopLevelViewModel).ToSharedRef();
+		}
+		return SNullWidget::NullWidget;
+	}
+
 private:
 	TSharedPtr<FNiagaraEmitterHandleViewModel> EmitterHandleViewModel;
 	TSharedPtr<UNiagaraStackViewModel::FTopLevelViewModel> TopLevelViewModel;
@@ -241,6 +246,8 @@ private:
 	TSharedPtr<SInlineEditableTextBlock> EmitterNameTextBlock;
 	FSimpleDelegate OnCycleThroughIssues;
 	UNiagaraStackViewModel* StackViewModel;
+	TWeakPtr<SNiagaraStack> ParentStackPtr;
+	TSharedPtr<SMenuAnchor> SettingsAnchor;
 };
 
 const float SpacerHeight = 6;
@@ -707,21 +714,20 @@ TSharedRef<ITableRow> SNiagaraStack::OnGenerateRowForTopLevelObject(TSharedRef<U
 	{
 		Content = SNew(SNiagaraStackEmitterHeader, Item->EmitterHandleViewModel.ToSharedRef(), Item->RootEntry.Get(), StackViewModel)
 			.IssueIconVisibility(this, &SNiagaraStack::GetIssueIconVisibility)
-			.OnCycleThroughIssues(this, &SNiagaraStack::OnCycleThroughIssues);
+			.OnCycleThroughIssues(this, &SNiagaraStack::OnCycleThroughIssues)
+			.ParentStack(SharedThis(this));
 	}
 
-	Content->SetOnMouseButtonUp(FPointerEventHandler::CreateSP(this, &SNiagaraStack::OnTopLevelRowMouseButtonDown, TWeakPtr<UNiagaraStackViewModel::FTopLevelViewModel>(Item)));
-
 	return SNew(STableRow<TSharedRef<UNiagaraStackViewModel::FTopLevelViewModel>>, OwnerTable)
 		[
 			Content.ToSharedRef()
 		];
 }
 
-FReply SNiagaraStack::OnTopLevelRowMouseButtonDown(const FGeometry&, const FPointerEvent& MouseEvent, TWeakPtr<UNiagaraStackViewModel::FTopLevelViewModel> TopLevelViewModelWeak)
+TSharedPtr<SWidget> SNiagaraStack::GenerateStackMenu(TWeakPtr<UNiagaraStackViewModel::FTopLevelViewModel> TopLevelViewModelWeak)
 {
 	TSharedPtr<UNiagaraStackViewModel::FTopLevelViewModel> TopLevelViewModel = TopLevelViewModelWeak.Pin();
-	if (TopLevelViewModel.IsValid() && MouseEvent.GetEffectingButton() == EKeys::RightMouseButton)
+	if (TopLevelViewModel.IsValid())
 	{
 		FMenuBuilder MenuBuilder(true, nullptr);
 
@@ -762,11 +768,9 @@ FReply SNiagaraStack::OnTopLevelRowMouseButtonDown(const FGeometry&, const FPoin
 		}
 		MenuBuilder.EndSection();
 
-		FWidgetPath WidgetPath = MouseEvent.GetEventPath() != nullptr ? *MouseEvent.GetEventPath() : FWidgetPath();
-		FSlateApplication::Get().PushMenu(AsShared(), WidgetPath, MenuBuilder.MakeWidget(), MouseEvent.GetScreenSpacePosition(), FPopupTransitionEffect(FPopupTransitionEffect::ContextMenu));
-		return FReply::Handled();
+		return MenuBuilder.MakeWidget();
 	}
-	return FReply::Unhandled();
+	return SNullWidget::NullWidget;
 }
 
 TSharedRef<SNiagaraStackTableRow> SNiagaraStack::ConstructContainerForItem(UNiagaraStackEntry* Item)
diff --git a/Engine/Plugins/FX/Niagara/Source/NiagaraEditorWidgets/Private/SNiagaraStack.h b/Engine/Plugins/FX/Niagara/Source/NiagaraEditorWidgets/Private/SNiagaraStack.h
--- a/Engine/Plugins/FX/Niagara/Source/NiagaraEditorWidgets/Private/SNiagaraStack.h
+++ b/Engine/Plugins/FX/Niagara/Source/NiagaraEditorWidgets/Private/SNiagaraStack.h
@@ -17,7 +17,7 @@ class SNiagaraStackTableRow;
 class SSearchBox;
 class FReply;
 class FNiagaraStackCommandContext;
-
+class SWidget;
 class SNiagaraStack : public SCompoundWidget
 {
 public:
@@ -28,6 +28,8 @@ public:
 
 	void Construct(const FArguments& InArgs, UNiagaraStackViewModel* InStackViewModel);
 
+	TSharedPtr<SWidget> GenerateStackMenu(TWeakPtr<UNiagaraStackViewModel::FTopLevelViewModel> TopLevelViewModelWeak);
+
 private:
 	struct FRowWidgets
 	{
@@ -52,7 +54,7 @@ private:
 
 	TSharedRef<ITableRow> OnGenerateRowForTopLevelObject(TSharedRef<UNiagaraStackViewModel::FTopLevelViewModel> Item, const TSharedRef<STableViewBase>& OwnerTable);
 
-	FReply OnTopLevelRowMouseButtonDown(const FGeometry&, const FPointerEvent& MouseEvent, TWeakPtr<UNiagaraStackViewModel::FTopLevelViewModel> TopLevelViewModelWeak);
+	
 
 	TSharedRef<SNiagaraStackTableRow> ConstructContainerForItem(UNiagaraStackEntry* Item);
 
