Fix issue with velocity-facing mesh particles when the velocity is perfectly vertical

diff --git a/Engine/Plugins/FX/Niagara/Shaders/Private/NiagaraMeshVertexFactory.ush b/Engine/Plugins/FX/Niagara/Shaders/Private/NiagaraMeshVertexFactory.ush
--- a/Engine/Plugins/FX/Niagara/Shaders/Private/NiagaraMeshVertexFactory.ush
+++ b/Engine/Plugins/FX/Niagara/Shaders/Private/NiagaraMeshVertexFactory.ush
@@ -515,8 +515,9 @@ float3x3 GetFacingMatrix(in FVertexFactoryIntermediates Intermediates)
 		}
 		else
 		{
-			// Prefer to use world up as a reference vector, fall back to negative world X-axis when facing up or down
-			RefVector = WorldZ;
+			// Prefer to use world up as a reference vector, fall back to world X-axis when facing up or down
+			float DotWorldZ = dot(FacingDir, WorldZ);
+			RefVector = abs(DotWorldZ) > 0.99f ? (-sign(DotWorldZ) * WorldX) : WorldZ;
 		}
 
 		// Check to rotate the reference direction to local-space
