Merge from game. Note that this is just to unblock robomerge. Functionality is broken, likely as a result of LoadTimes changes to shaders.
Copying over shader compilation fixes from working branch
CL 11205946 GPU compile errors now flush to the UI, making it possible to handle odd situations where the GPU doesn't compile but the CPU does. Also flushing out stale GPU script results.
CL 11259759 Making sure we clear out the compiled error list when successfully compiling.

Various fixes in here.
1) Collision data interfaces require specific functions for CPU/GPU. These weren't generating warnings at the right time. Now they do.
2) When compiling CPU sims. make sure to invalidate the resulting GPU sim. This keeps us from getting stale results in the UI.
3) Added a new definition for CPU/GPU sims to script compilation. This lets us propagate information based off of the emitter's state. Hopefully this will let us fix up the collision DI in the future to just do the right thing in a static switch.
4) Also now generating an error when compiling event reads/writes on GPU sims.
5) Had to move the recompile request to after all other work for post edit change happens so that we don't race the UI in displaying the errors.

diff --git a/Engine/Plugins/FX/Niagara/Source/Niagara/Classes/NiagaraEmitter.h b/Engine/Plugins/FX/Niagara/Source/Niagara/Classes/NiagaraEmitter.h
--- a/Engine/Plugins/FX/Niagara/Source/Niagara/Classes/NiagaraEmitter.h
+++ b/Engine/Plugins/FX/Niagara/Source/Niagara/Classes/NiagaraEmitter.h
@@ -403,6 +403,15 @@ public:
 	/** Callback issued whenever a VM compilation successfully happened (even if the results are a script that cannot be executed due to errors)*/
 	NIAGARA_API FOnEmitterCompiled& OnEmitterVMCompiled();
 
+	/** Callback issued whenever a VM compilation successfully happened (even if the results are a script that cannot be executed due to errors)*/
+	NIAGARA_API FOnEmitterCompiled& OnEmitterGPUCompiled();
+
+	/** Callback issued whenever a GPU compilation successfully happened (even if the results are a script that cannot be executed due to errors)*/
+	NIAGARA_API FOnEmitterCompiled& OnGPUCompilationComplete()
+	{
+		return OnGPUScriptCompiledDelegate;
+	}
+	
 	NIAGARA_API static bool GetForceCompileOnLoad();
 
 	/** Whether or not this emitter is synchronized with its parent emitter. */
@@ -541,6 +550,11 @@ private:
 
 	/** A multicast delegate which is called whenever all the scripts for this emitter have been compiled (successfully or not). */
 	FOnEmitterCompiled OnVMScriptCompiledDelegate;
+
+	/** A multicast delegate which is called whenever all the scripts for this emitter have been compiled (successfully or not). */
+	FOnEmitterCompiled OnGPUScriptCompiledDelegate;
+
+	void RaiseOnEmitterGPUCompiled(UNiagaraScript* InScript);
 #endif
 
 	UPROPERTY()
diff --git a/Engine/Plugins/FX/Niagara/Source/Niagara/Classes/NiagaraScript.h b/Engine/Plugins/FX/Niagara/Source/Niagara/Classes/NiagaraScript.h
--- a/Engine/Plugins/FX/Niagara/Source/Niagara/Classes/NiagaraScript.h
+++ b/Engine/Plugins/FX/Niagara/Source/Niagara/Classes/NiagaraScript.h
@@ -537,6 +537,12 @@ public:
 	{
 		return &ScriptResource;
 	}
+	const FNiagaraShaderScript *GetRenderThreadScript() const
+	{
+		return &ScriptResource;
+	}
+
+	
 
 	FComputeShaderRHIRef GetScriptShader() 
 	{
@@ -597,9 +603,12 @@ public:
 	/** Creates a string key for the derived data cache */
 	FString GetNiagaraDDCKeyString();
 
-	/** Callback issued whenever a compilation successfully happened (even if the results are a script that cannot be executed due to errors)*/
+	/** Callback issued whenever a VM script compilation successfully happened (even if the results are a script that cannot be executed due to errors)*/
 	NIAGARA_API FOnScriptCompiled& OnVMScriptCompiled();
-	
+
+	/** Callback issued whenever a GPU script compilation successfully happened (even if the results are a script that cannot be executed due to errors)*/
+	NIAGARA_API FOnScriptCompiled& OnGPUScriptCompiled();
+
 	/** External call used to identify the values for a successful VM script compilation. OnVMScriptCompiled will be issued in this case.*/
 	void SetVMCompilationResults(const FNiagaraVMExecutableDataId& InCompileId, FNiagaraVMExecutableData& InScriptVM, FNiagaraCompileRequestDataBase* InRequestData);
 
@@ -609,10 +618,12 @@ public:
 	NIAGARA_API bool SynchronizeExecutablesWithMaster(const UNiagaraScript* Script, const TMap<FString, FString>& RenameMap);
 
 	NIAGARA_API void SyncAliases(const TMap<FString, FString>& RenameMap);
+
 #endif
 	
 	UFUNCTION()
-	void OnCompilationComplete();
+	void RaiseOnGPUCompilationComplete();
+
 
 	NIAGARA_API FNiagaraVMExecutableData& GetVMExecutableData() { return CachedScriptVM; }
 	NIAGARA_API const FNiagaraVMExecutableData& GetVMExecutableData() const { return CachedScriptVM; }
@@ -633,7 +644,7 @@ public:
 	void InvalidateExecutionReadyParameterStores();
 
 private:
-
+	bool OwnerCanBeRunOnGpu() const;
 	bool LegacyCanBeRunOnGpu()const;
 
 	/** Return the expected SimTarget for this script. Only returns a valid target if there is valid data to run with. */
@@ -668,6 +679,7 @@ private:
 	
 	/** A multicast delegate which is called whenever the script has been compiled (successfully or not). */
 	FOnScriptCompiled OnVMScriptCompiledDelegate;
+	FOnScriptCompiled OnGPUScriptCompiledDelegate;
 
 	mutable FNiagaraVMExecutableDataId LastReportedVMId;
 
diff --git a/Engine/Plugins/FX/Niagara/Source/Niagara/Private/NiagaraEmitter.cpp b/Engine/Plugins/FX/Niagara/Source/Niagara/Private/NiagaraEmitter.cpp
--- a/Engine/Plugins/FX/Niagara/Source/Niagara/Private/NiagaraEmitter.cpp
+++ b/Engine/Plugins/FX/Niagara/Source/Niagara/Private/NiagaraEmitter.cpp
@@ -152,6 +152,14 @@ void UNiagaraEmitter::PostInitProperties()
 		GPUComputeScript->SetUsage(ENiagaraScriptUsage::ParticleGPUComputeScript);
 
 	}
+
+#if WITH_EDITORONLY_DATA
+	if (GPUComputeScript)
+	{
+		GPUComputeScript->OnGPUScriptCompiled().AddUObject(this, &UNiagaraEmitter::RaiseOnEmitterGPUCompiled);
+	}
+#endif
+
 	UniqueEmitterName = TEXT("Emitter");
 
 	ResolveScalabilitySettings();
@@ -344,6 +352,14 @@ void UNiagaraEmitter::PostLoad()
 #endif
 	}
 
+
+#if WITH_EDITORONLY_DATA
+	if (GPUComputeScript)
+	{
+		GPUComputeScript->OnGPUScriptCompiled().AddUObject(this, &UNiagaraEmitter::RaiseOnEmitterGPUCompiled);
+	}
+#endif
+
 	if (EmitterSpawnScriptProps.Script == nullptr || EmitterUpdateScriptProps.Script == nullptr)
 	{
 		EmitterSpawnScriptProps.Script = NewObject<UNiagaraScript>(this, "EmitterSpawnScript", EObjectFlags::RF_Transactional);
@@ -579,6 +595,7 @@ void UNiagaraEmitter::PostEditChangeProperty(struct FPropertyChangedEvent& Prope
 		PropertyName = PropertyChangedEvent.Property->GetFName();
 	}
 
+	bool bNeedsRecompile = false;
 	if (PropertyName == GET_MEMBER_NAME_CHECKED(UNiagaraEmitter, bInterpolatedSpawning))
 	{
 		bool bActualInterpolatedSpawning = SpawnScriptProps.Script->IsInterpolatedParticleSpawnScript();
@@ -591,9 +608,7 @@ void UNiagaraEmitter::PostEditChangeProperty(struct FPropertyChangedEvent& Prope
 			{
 				GraphSource->MarkNotSynchronized(TEXT("Emitter interpolated spawn changed"));
 			}
-#if WITH_EDITORONLY_DATA
-			UNiagaraSystem::RequestCompileForEmitter(this);
-#endif
+			bNeedsRecompile = true;
 		}
 	}
 	if (PropertyName == GET_MEMBER_NAME_CHECKED(UNiagaraEmitter, SimTarget))
@@ -602,10 +617,7 @@ void UNiagaraEmitter::PostEditChangeProperty(struct FPropertyChangedEvent& Prope
 		{
 			GraphSource->MarkNotSynchronized(TEXT("Emitter simulation target changed."));
 		}
-
-#if WITH_EDITORONLY_DATA
-		UNiagaraSystem::RequestCompileForEmitter(this);
-#endif
+		bNeedsRecompile = true;
 	}
 	if (PropertyName == GET_MEMBER_NAME_CHECKED(UNiagaraEmitter, bRequiresPersistentIDs))
 	{
@@ -613,10 +625,7 @@ void UNiagaraEmitter::PostEditChangeProperty(struct FPropertyChangedEvent& Prope
 		{
 			GraphSource->MarkNotSynchronized(TEXT("Emitter Requires Persistent IDs changed."));
 		}
-
-#if WITH_EDITORONLY_DATA
-		UNiagaraSystem::RequestCompileForEmitter(this);
-#endif
+		bNeedsRecompile = true;
 	}
 	if (PropertyName == GET_MEMBER_NAME_CHECKED(UNiagaraEmitter, bLocalSpace))
 	{
@@ -625,9 +634,7 @@ void UNiagaraEmitter::PostEditChangeProperty(struct FPropertyChangedEvent& Prope
 			GraphSource->MarkNotSynchronized(TEXT("Emitter LocalSpace changed."));
 		}
 
-#if WITH_EDITORONLY_DATA
-		UNiagaraSystem::RequestCompileForEmitter(this);
-#endif
+		bNeedsRecompile = true;
 	}
 	if (PropertyName == GET_MEMBER_NAME_CHECKED(UNiagaraEmitter, bDeterminism))
 	{
@@ -646,6 +653,13 @@ void UNiagaraEmitter::PostEditChangeProperty(struct FPropertyChangedEvent& Prope
 	ThumbnailImageOutOfDate = true;
 	UpdateChangeId(TEXT("PostEditChangeProperty"));
 	OnPropertiesChangedDelegate.Broadcast();
+
+#if WITH_EDITORONLY_DATA
+	if (bNeedsRecompile)
+	{
+		UNiagaraSystem::RequestCompileForEmitter(this);
+	}
+#endif
 }
 
 
@@ -887,6 +901,11 @@ UNiagaraEmitter::FOnEmitterCompiled& UNiagaraEmitter::OnEmitterVMCompiled()
 	return OnVMScriptCompiledDelegate;
 }
 
+UNiagaraEmitter::FOnEmitterCompiled& UNiagaraEmitter::OnEmitterGPUCompiled()
+{
+	return OnGPUScriptCompiledDelegate;
+}
+
 void  UNiagaraEmitter::InvalidateCompileResults()
 {
 	TArray<UNiagaraScript*> Scripts;
@@ -936,6 +955,12 @@ void UNiagaraEmitter::OnPostCompile()
 		}
 	}
 
+	// If we have a GPU script but the SimTarget isn't GPU, we need to clear out the old results.
+	if (SimTarget != ENiagaraSimTarget::GPUComputeSim && GPUComputeScript->GetLastCompileStatus() != ENiagaraScriptCompileStatus::NCS_Unknown)
+	{
+		GPUComputeScript->InvalidateCompileResults(TEXT("Not a GPU emitter."));
+	}
+
 	RuntimeEstimation = MemoryRuntimeEstimation();
 
 	OnEmitterVMCompiled().Broadcast(this);
@@ -1333,6 +1358,10 @@ void UNiagaraEmitter::BeginDestroy()
 	{
 		GraphSource->OnChanged().RemoveAll(this);
 	}
+	if (GPUComputeScript)
+	{
+		GPUComputeScript->OnGPUScriptCompiled().RemoveAll(this);
+	}
 #endif
 	Super::BeginDestroy();
 }
@@ -1373,6 +1402,11 @@ void UNiagaraEmitter::GraphSourceChanged()
 	UpdateChangeId(TEXT("Graph source changed."));
 }
 
+void UNiagaraEmitter::RaiseOnEmitterGPUCompiled(UNiagaraScript* InScript)
+{
+	OnGPUScriptCompiledDelegate.Broadcast(this);
+}
+
 void UNiagaraEmitter::PersistentEditorDataChanged()
 {
 	UpdateChangeId(TEXT("Persistent editor data changed."));
diff --git a/Engine/Plugins/FX/Niagara/Source/Niagara/Private/NiagaraScript.cpp b/Engine/Plugins/FX/Niagara/Source/Niagara/Private/NiagaraScript.cpp
--- a/Engine/Plugins/FX/Niagara/Source/Niagara/Private/NiagaraScript.cpp
+++ b/Engine/Plugins/FX/Niagara/Source/Niagara/Private/NiagaraScript.cpp
@@ -256,7 +256,7 @@ UNiagaraScript::UNiagaraScript(const FObjectInitializer& ObjectInitializer)
 #endif
 {
 #if WITH_EDITORONLY_DATA
-	ScriptResource.OnCompilationComplete().AddUniqueDynamic(this, &UNiagaraScript::OnCompilationComplete);
+	ScriptResource.OnCompilationComplete().AddUniqueDynamic(this, &UNiagaraScript::RaiseOnGPUCompilationComplete);
 
 	RapidIterationParameters.DebugName = *GetFullName();
 #endif	
@@ -307,6 +307,7 @@ void UNiagaraScript::ComputeVMCompilationId(FNiagaraVMExecutableDataId& Id) cons
 	Id.bInterpolatedSpawn = false;
 	Id.bRequiresPersistentIDs = false;
 	
+	ENiagaraSimTarget SimTargetToBuild = ENiagaraSimTarget::CPUSim;
 	// Ideally we wouldn't want to do this but rather than push the data down
 	// from the emitter.
 	UObject* Obj = GetOuter();
@@ -325,6 +326,10 @@ void UNiagaraScript::ComputeVMCompilationId(FNiagaraVMExecutableDataId& Id) cons
 			Id.bInterpolatedSpawn = true;
 			Id.AdditionalDefines.Add(TEXT("InterpolatedSpawn"));
 		}
+		if (IsParticleScript(Usage))
+		{
+			SimTargetToBuild = Emitter->SimTarget;
+		}
 		if (Emitter->RequiresPersistantIDs())
 		{
 			Id.bRequiresPersistentIDs = true;
@@ -369,6 +374,19 @@ void UNiagaraScript::ComputeVMCompilationId(FNiagaraVMExecutableDataId& Id) cons
 		}
 	}
 
+	switch (SimTargetToBuild)
+	{
+	case ENiagaraSimTarget::CPUSim:
+		Id.AdditionalDefines.Add(TEXT("CPUSim"));
+		break;
+	case ENiagaraSimTarget::GPUComputeSim:
+		Id.AdditionalDefines.Add(TEXT("GPUComputeSim"));
+		break;
+	default:
+		checkf(false, TEXT("Unknown sim target type!"));
+	}
+
+
 	// If we aren't using rapid iteration parameters, we need to bake them into the hashstate for the compile id. This 
 	// makes their values part of the lookup. 
 	if (false == Id.bUsesRapidIterationParams)
@@ -1288,6 +1306,8 @@ void UNiagaraScript::SetVMCompilationResults(const FNiagaraVMExecutableDataId& I
 	CachedScriptVMId = InCompileId;
 	CachedScriptVM = InScriptVM;
 	CachedParameterCollectionReferences.Empty();
+	// Proactively clear out the script resource, because it might be stale now.
+	ScriptResource.Invalidate();
 	
 	if (CachedScriptVM.LastCompileStatus == ENiagaraScriptCompileStatus::NCS_Error)
 	{
@@ -1463,6 +1483,7 @@ bool UNiagaraScript::RequestExternallyManagedAsyncCompile(const TSharedPtr<FNiag
 		FNiagaraCompileOptions Options(GetUsage(), GetUsageId(), ModuleUsageBitmask, GetPathName(), GetFullName(), GetName());
 		Options.AdditionalDefines = LastGeneratedVMId.AdditionalDefines;
 		OutAsyncHandle = NiagaraModule.StartScriptCompileJob(RequestData.Get(), Options);
+		UE_LOG(LogNiagara, Verbose, TEXT("Script '%s' is requesting compile.."), *GetFullName());
 		return true;
 	}
 	else
@@ -1474,9 +1495,10 @@ bool UNiagaraScript::RequestExternallyManagedAsyncCompile(const TSharedPtr<FNiag
 }
 #endif
 
-void UNiagaraScript::OnCompilationComplete()
+void UNiagaraScript::RaiseOnGPUCompilationComplete()
 {
 #if WITH_EDITORONLY_DATA
+	OnGPUScriptCompiled().Broadcast(this);
 	FNiagaraSystemUpdateContext(this, true);
 #endif
 }
@@ -1689,7 +1711,8 @@ void UNiagaraScript::CacheResourceShadersForRendering(bool bRegenerateId, bool b
 
 	if (CanBeRunOnGpu())
 	{
-		if (Source)
+		// Need to make sure the owner supports GPU scripts, otherwise this is a wasted compile.
+		if (Source && OwnerCanBeRunOnGpu())
 		{
 			FNiagaraShaderScript* ResourceToCache;
 			ERHIFeatureLevel::Type CacheFeatureLevel = GMaxRHIFeatureLevel;
@@ -1708,6 +1731,10 @@ void UNiagaraScript::CacheResourceShadersForRendering(bool bRegenerateId, bool b
 				}
 			}
 		}
+		else
+		{
+			ScriptResource.Invalidate();
+		}
 	}
 }
 
@@ -1822,6 +1849,11 @@ UNiagaraScript::FOnScriptCompiled& UNiagaraScript::OnVMScriptCompiled()
 	return OnVMScriptCompiledDelegate;
 }
 
+UNiagaraScript::FOnScriptCompiled& UNiagaraScript::OnGPUScriptCompiled()
+{
+	return OnGPUScriptCompiledDelegate;
+}
+
 
 
 #endif
@@ -1999,6 +2031,19 @@ bool UNiagaraScript::CanBeRunOnGpu()const
 }
 
 
+bool UNiagaraScript::OwnerCanBeRunOnGpu() const
+{
+	if (UNiagaraEmitter* Emitter = GetTypedOuter<UNiagaraEmitter>())
+	{
+		if (Emitter->SimTarget == ENiagaraSimTarget::GPUComputeSim)
+		{
+			return true;
+		}
+	}
+
+	return false;
+}
+
 bool UNiagaraScript::LegacyCanBeRunOnGpu() const
 {
 	if (UNiagaraEmitter* Emitter = GetTypedOuter<UNiagaraEmitter>())
diff --git a/Engine/Plugins/FX/Niagara/Source/NiagaraCore/Private/NiagaraCustomVersion.cpp b/Engine/Plugins/FX/Niagara/Source/NiagaraCore/Private/NiagaraCustomVersion.cpp
--- a/Engine/Plugins/FX/Niagara/Source/NiagaraCore/Private/NiagaraCustomVersion.cpp
+++ b/Engine/Plugins/FX/Niagara/Source/NiagaraCore/Private/NiagaraCustomVersion.cpp
@@ -8,4 +8,4 @@ const FGuid FNiagaraCustomVersion::GUID(0xFCF57AFA, 0x50764283, 0xB9A9E658, 0xFF
 // Register the custom version with core
 FCustomVersionRegistration GRegisterNiagaraCustomVersion(FNiagaraCustomVersion::GUID, FNiagaraCustomVersion::LatestVersion, TEXT("NiagaraVer"));
 
-const FGuid FNiagaraCustomVersion::LatestScriptCompileVersion(0x6B79BB03, 0x495F4956, 0x9F7C5D3E, 0x42B9AC36);
\ No newline at end of file
+const FGuid FNiagaraCustomVersion::LatestScriptCompileVersion(0x0288EF7C, 0x881D4F16, 0xB4A7513C, 0xFA3EECED);
diff --git a/Engine/Plugins/FX/Niagara/Source/NiagaraEditor/Private/NiagaraHlslTranslator.cpp b/Engine/Plugins/FX/Niagara/Source/NiagaraEditor/Private/NiagaraHlslTranslator.cpp
--- a/Engine/Plugins/FX/Niagara/Source/NiagaraEditor/Private/NiagaraHlslTranslator.cpp
+++ b/Engine/Plugins/FX/Niagara/Source/NiagaraEditor/Private/NiagaraHlslTranslator.cpp
@@ -4520,7 +4520,10 @@ void FHlslNiagaraTranslator::HandleParameterRead(int32 ParamMapHistoryIdx, const
 	}
 }
 
-
+bool FHlslNiagaraTranslator::IsCompileOptionDefined(const TCHAR* InDefineStr)
+{
+	return CompileOptions.AdditionalDefines.Contains(InDefineStr);
+}
 
 void FHlslNiagaraTranslator::ReadDataSet(const FNiagaraDataSetID DataSet, const TArray<FNiagaraVariable>& Variables, ENiagaraDataSetAccessMode AccessMode, int32 InputChunk, TArray<int32>& Outputs)
 {
@@ -5531,12 +5534,25 @@ void FHlslNiagaraTranslator::RegisterFunctionCall(ENiagaraScriptUsage ScriptUsag
 				TArray<FNiagaraFunctionSignature> DataInterfaceFunctions;
 				CDO->GetFunctions(DataInterfaceFunctions);
 
-				const bool bFoundMatch = DataInterfaceFunctions.ContainsByPredicate([&](const FNiagaraFunctionSignature& Sig) -> bool { return Sig.EqualsIgnoringSpecifiers(OutSignature); });
-				if (!bFoundMatch)
+				const int32 FoundMatch = DataInterfaceFunctions.IndexOfByPredicate([&](const FNiagaraFunctionSignature& Sig) -> bool { return Sig.EqualsIgnoringSpecifiers(OutSignature); });
+				if (FoundMatch < 0)
 				{
 					Error(LOCTEXT("FunctionCallDataInterfaceMissing", "Function call signature does not match DataInterface possible signatures?"), nullptr, nullptr);
 					return;
 				}
+				else
+				{
+					bool bCPUSim = IsCompileOptionDefined(TEXT("CPUSim"));
+					bool bGPUSim = IsCompileOptionDefined(TEXT("GPUComputeSim"));
+					if (bCPUSim && !DataInterfaceFunctions[FoundMatch].bSupportsCPU)
+					{
+						Error(FText::Format(LOCTEXT("FunctionCallDataInterfaceCPUMissing", "Function call \"{0}\" does not work on CPU sims."), FText::FromName(OutSignature.Name)), nullptr, nullptr);
+					}
+					else if (bGPUSim && !DataInterfaceFunctions[FoundMatch].bSupportsGPU)
+					{
+						Error(FText::Format(LOCTEXT("FunctionCallDataInterfaceGPUMissing", "Function call \"{0}\" does not work on GPU sims."), FText::FromName(OutSignature.Name)), nullptr, nullptr);
+					}
+				}
 
 				// We only use this for GPU systems currently so that we emit only the functionality required
 				DataInterfaceRegisteredFunctions.FindOrAdd(Info.Type.GetFName()).Add(OutSignature);
diff --git a/Engine/Plugins/FX/Niagara/Source/NiagaraEditor/Private/NiagaraHlslTranslator.h b/Engine/Plugins/FX/Niagara/Source/NiagaraEditor/Private/NiagaraHlslTranslator.h
--- a/Engine/Plugins/FX/Niagara/Source/NiagaraEditor/Private/NiagaraHlslTranslator.h
+++ b/Engine/Plugins/FX/Niagara/Source/NiagaraEditor/Private/NiagaraHlslTranslator.h
@@ -449,6 +449,7 @@ public:
 	virtual int32 GetParameter(const FNiagaraVariable& Parameter);
 	virtual int32 GetRapidIterationParameter(const FNiagaraVariable& Parameter);
 
+	bool IsCompileOptionDefined(const TCHAR* InDefineStr);
 
 	virtual int32 GetAttribute(const FNiagaraVariable& Attribute);
 
diff --git a/Engine/Plugins/FX/Niagara/Source/NiagaraEditor/Private/NiagaraNodeReadDataSet.cpp b/Engine/Plugins/FX/Niagara/Source/NiagaraEditor/Private/NiagaraNodeReadDataSet.cpp
--- a/Engine/Plugins/FX/Niagara/Source/NiagaraEditor/Private/NiagaraNodeReadDataSet.cpp
+++ b/Engine/Plugins/FX/Niagara/Source/NiagaraEditor/Private/NiagaraNodeReadDataSet.cpp
@@ -100,6 +100,13 @@ void UNiagaraNodeReadDataSet::Compile(class FHlslNiagaraTranslator* Translator,
 {
 	TArray<int32> Inputs;
 	CompileInputPins(Translator, Inputs);
+	
+	bool bGPUSim = Translator->IsCompileOptionDefined(TEXT("GPUComputeSim"));
+
+	if (bGPUSim)
+	{
+		Translator->Error(LOCTEXT("CannotRunReadDataSetGPU", "Cannot use an event read node on GPU sims!"), this, nullptr);
+	}
 
 	FString IssuesWithStruct;
 	if (!IsSynchronizedWithStruct(false, &IssuesWithStruct, false))
diff --git a/Engine/Plugins/FX/Niagara/Source/NiagaraEditor/Private/NiagaraNodeWriteDataSet.cpp b/Engine/Plugins/FX/Niagara/Source/NiagaraEditor/Private/NiagaraNodeWriteDataSet.cpp
--- a/Engine/Plugins/FX/Niagara/Source/NiagaraEditor/Private/NiagaraNodeWriteDataSet.cpp
+++ b/Engine/Plugins/FX/Niagara/Source/NiagaraEditor/Private/NiagaraNodeWriteDataSet.cpp
@@ -71,6 +71,14 @@ void UNiagaraNodeWriteDataSet::Compile(class FHlslNiagaraTranslator* Translator,
 	TArray<int32> Inputs;
 	CompileInputPins(Translator, Inputs);
 
+
+	bool bGPUSim = Translator->IsCompileOptionDefined(TEXT("GPUComputeSim"));
+
+	if (bGPUSim)
+	{
+		Translator->Error(LOCTEXT("CannotRunWriteDataSetGPU", "Cannot use an event write node on GPU sims!"), this, nullptr);
+	}
+
 	FString IssuesWithStruct;
 	if (!IsSynchronizedWithStruct(true, &IssuesWithStruct,false))
 	{
diff --git a/Engine/Plugins/FX/Niagara/Source/NiagaraEditor/Private/ViewModels/NiagaraEmitterViewModel.cpp b/Engine/Plugins/FX/Niagara/Source/NiagaraEditor/Private/ViewModels/NiagaraEmitterViewModel.cpp
--- a/Engine/Plugins/FX/Niagara/Source/NiagaraEditor/Private/ViewModels/NiagaraEmitterViewModel.cpp
+++ b/Engine/Plugins/FX/Niagara/Source/NiagaraEditor/Private/ViewModels/NiagaraEmitterViewModel.cpp
@@ -95,6 +95,7 @@ void FNiagaraEmitterViewModel::SetEmitter(UNiagaraEmitter* InEmitter)
 	if (Emitter.IsValid())
 	{
 		Emitter->OnEmitterVMCompiled().RemoveAll(this);
+		Emitter->OnEmitterGPUCompiled().RemoveAll(this);
 		Emitter->OnPropertiesChanged().RemoveAll(this);
 		UnregisterViewModelWithMap(RegisteredHandle);
 	}
@@ -109,6 +110,7 @@ void FNiagaraEmitterViewModel::SetEmitter(UNiagaraEmitter* InEmitter)
 	{
 		Emitter->OnEmitterVMCompiled().AddSP(this, &FNiagaraEmitterViewModel::OnVMCompiled);
 		Emitter->OnPropertiesChanged().AddSP(this, &FNiagaraEmitterViewModel::OnEmitterPropertiesChanged);
+		Emitter->OnEmitterGPUCompiled().AddSP(this, &FNiagaraEmitterViewModel::OnGPUCompiled);
 	}
 
 	AddScriptEventHandlers();
@@ -326,6 +328,13 @@ UNiagaraEmitterEditorData& FNiagaraEmitterViewModel::GetOrCreateEditorData()
 	return *EditorData;
 }
 
+
+void FNiagaraEmitterViewModel::OnGPUCompiled(UNiagaraEmitter* InEmitter)
+{
+	// Flush compilation logs just like the CPU script compiled.
+	OnVMCompiled(InEmitter);
+}
+
 void FNiagaraEmitterViewModel::OnVMCompiled(UNiagaraEmitter* InEmitter)
 {
 	if (Emitter.IsValid() && InEmitter == Emitter.Get())
@@ -371,14 +380,60 @@ void FNiagaraEmitterViewModel::OnVMCompiled(UNiagaraEmitter* InEmitter)
 
 		}
 
+		// Now handle any GPU sims..
+		TArray<UNiagaraScript*> GPUScripts;
+		
+		// Going ahead and making an array for future proofing..
+		{
+			UNiagaraScript* GPUScript = Emitter->GetGPUComputeScript();
+			GPUScripts.Add(GPUScript);
+		}
+
+		// Push in GPU script errors into the other scripts
+		for (UNiagaraScript* GPUScript : GPUScripts)
+		{
+			if (GPUScript)
+			{
+				const FNiagaraShaderScript* ShaderScript = GPUScript->GetRenderThreadScript();
+				if (!ShaderScript)
+					return;
+
+				for (int32 i = 0; i < Scripts.Num(); i++)
+				{
+					if (Scripts[i] != nullptr && CompileStatuses.Num() > i && CompileErrors.Num() > i && GPUScript->ContainsUsage(Scripts[i]->Usage))
+					{
+						const TArray<FString>& NewErrors = ShaderScript->GetCompileErrors();
+						bool bErrors = false;
+						for (const FString& ErrorStr : NewErrors)
+						{
+							if (ErrorStr.Contains(TEXT("err0r")))
+							{
+								CompileErrors[i] += ErrorStr + TEXT("\n");
+								bErrors = true;
+							}
+						}
+
+						if (bErrors)
+						{
+							CompileStatuses[i] = ENiagaraScriptCompileStatus::NCS_Error;
+						}
+					}
+				}
+
+			}
+		}
+
 		for (int32 i = 0; i < CompileStatuses.Num(); i++)
 		{
 			AggregateStatus = FNiagaraEditorUtilities::UnionCompileStatus(AggregateStatus, CompileStatuses[i]);
 			AggregateErrors += CompilePaths[i] + TEXT(" ") + FNiagaraEditorUtilities::StatusToText(CompileStatuses[i]).ToString() + TEXT("\n");
 			AggregateErrors += CompileErrors[i] + TEXT("\n");
 		}
-		check(SharedScriptViewModel.IsValid());
-		SharedScriptViewModel->UpdateCompileStatus(AggregateStatus, AggregateErrors, CompileStatuses, CompileErrors, CompilePaths, Scripts);
+		
+		if (SharedScriptViewModel.IsValid())
+		{
+			SharedScriptViewModel->UpdateCompileStatus(AggregateStatus, AggregateErrors, CompileStatuses, CompileErrors, CompilePaths, Scripts);
+		}
 	}
 	OnScriptCompiled().Broadcast();
 }
diff --git a/Engine/Plugins/FX/Niagara/Source/NiagaraEditor/Private/ViewModels/NiagaraScriptViewModel.cpp b/Engine/Plugins/FX/Niagara/Source/NiagaraEditor/Private/ViewModels/NiagaraScriptViewModel.cpp
--- a/Engine/Plugins/FX/Niagara/Source/NiagaraEditor/Private/ViewModels/NiagaraScriptViewModel.cpp
+++ b/Engine/Plugins/FX/Niagara/Source/NiagaraEditor/Private/ViewModels/NiagaraScriptViewModel.cpp
@@ -64,6 +64,11 @@ void FNiagaraScriptViewModel::OnVMScriptCompiled(UNiagaraScript* InScript)
 	}
 }
 
+void FNiagaraScriptViewModel::OnGPUScriptCompiled(UNiagaraScript* InScript)
+{
+	// Do nothing in base implementation
+}
+
 bool FNiagaraScriptViewModel::IsGraphDirty() const
 {
 	for (int32 i = 0; i < Scripts.Num(); i++)
@@ -101,6 +106,7 @@ FNiagaraScriptViewModel::~FNiagaraScriptViewModel()
 		if (Scripts[i].IsValid())
 		{
 			Scripts[i]->OnVMScriptCompiled().RemoveAll(this);
+			Scripts[i]->OnGPUScriptCompiled().RemoveAll(this);
 		}
 	}
 
@@ -136,6 +142,7 @@ void FNiagaraScriptViewModel::SetScripts(UNiagaraScriptSource* InScriptSource, T
 		if (Scripts[i].IsValid())
 		{
 			Scripts[i]->OnVMScriptCompiled().RemoveAll(this);
+			Scripts[i]->OnGPUScriptCompiled().RemoveAll(this);
 		}
 	}
 
@@ -151,6 +158,7 @@ void FNiagaraScriptViewModel::SetScripts(UNiagaraScriptSource* InScriptSource, T
 		int32 i = Scripts.Add(Script);
 		check(Script->GetSource() == InScriptSource);
 		Scripts[i]->OnVMScriptCompiled().AddSP(this, &FNiagaraScriptViewModel::OnVMScriptCompiled);
+		Scripts[i]->OnGPUScriptCompiled().AddSP(this, &FNiagaraScriptViewModel::OnGPUScriptCompiled);
 	}
 	Source = InScriptSource;
 
diff --git a/Engine/Plugins/FX/Niagara/Source/NiagaraEditor/Private/ViewModels/NiagaraSystemViewModel.cpp b/Engine/Plugins/FX/Niagara/Source/NiagaraEditor/Private/ViewModels/NiagaraSystemViewModel.cpp
--- a/Engine/Plugins/FX/Niagara/Source/NiagaraEditor/Private/ViewModels/NiagaraSystemViewModel.cpp
+++ b/Engine/Plugins/FX/Niagara/Source/NiagaraEditor/Private/ViewModels/NiagaraSystemViewModel.cpp
@@ -677,13 +677,17 @@ void FNiagaraSystemViewModel::SendLastCompileMessageJobs() const
 	{
 		const UNiagaraEmitter* EmitterInSystem = Handle.GetInstance();
 		TArray<UNiagaraScript*> EmitterScripts;
-		EmitterInSystem->GetScripts(EmitterScripts);
+		EmitterInSystem->GetScripts(EmitterScripts, false);
 		for (UNiagaraScript* EmitterScript : EmitterScripts)
 		{
 			ScriptsToGetCompileEventsFrom.Add(FNiagaraScriptAndOwningScriptNameString(EmitterScript, EmitterInSystem->GetUniqueEmitterName()));
 		}
 	}
 
+	// Clear out existing compile event messages from compile event message jobs for this asset.
+	FNiagaraMessageManager::Get()->RefreshMessagesForAssetKeyAndMessageJobType(SystemMessageLogGuidKey.GetValue(), ENiagaraMessageJobType::CompileEventMessageJob);
+
+	// Make new messages and messages jobs from the compile.
 	TArray<TSharedPtr<const INiagaraMessageJob>> JobBatchToQueue;
 	// Iterate from back to front to avoid reordering the events when they are queued
 	for (int i = ScriptsToGetCompileEventsFrom.Num()-1; i >=0; --i)
@@ -704,9 +708,33 @@ void FNiagaraSystemViewModel::SendLastCompileMessageJobs() const
 
 			JobBatchToQueue.Add(MakeShared<FNiagaraMessageJobCompileEvent>(CompileEvent, MakeWeakObjectPtr(const_cast<UNiagaraScript*>(ScriptInfo.Script)), ScriptInfo.OwningScriptNameString));
 		}
+		
+		// Check if there are any GPU compile errors and if so push them.
+		const FNiagaraShaderScript* ShaderScript = ScriptInfo.Script->GetRenderThreadScript();
+		if (ShaderScript != nullptr && ShaderScript->IsCompilationFinished() && 
+			ScriptInfo.Script->Usage == ENiagaraScriptUsage::ParticleGPUComputeScript)
+		{
+			const TArray<FString>& GPUCompileErrors = ShaderScript->GetCompileErrors();
+			for (const FString& String : GPUCompileErrors)
+			{
+				FNiagaraCompileEventSeverity Severity = FNiagaraCompileEventSeverity::Warning;
+				if (String.Contains(TEXT("err0r")))
+				{
+					Severity = FNiagaraCompileEventSeverity::Error;
+					ErrorCount++;
+				}
+				else
+				{
+					WarningCount++;
+				}
+				FNiagaraCompileEvent CompileEvent = FNiagaraCompileEvent(Severity, String);
+				JobBatchToQueue.Add(MakeShared<FNiagaraMessageJobCompileEvent>(CompileEvent, MakeWeakObjectPtr(const_cast<UNiagaraScript*>(ScriptInfo.Script)),
+					ScriptInfo.OwningScriptNameString));
+			}
+		}
+		
 	}
 	JobBatchToQueue.Insert(MakeShared<FNiagaraMessageJobPostCompileSummary>(ErrorCount, WarningCount, GetLatestCompileStatus(), FText::FromString("System")), 0);
-	FNiagaraMessageManager::Get()->RefreshMessagesForAssetKeyAndMessageJobType(SystemMessageLogGuidKey.GetValue(), ENiagaraMessageJobType::CompileEventMessageJob);
 	FNiagaraMessageManager::Get()->QueueMessageJobBatch(JobBatchToQueue, SystemMessageLogGuidKey.GetValue());
 }
 
diff --git a/Engine/Plugins/FX/Niagara/Source/NiagaraEditor/Private/ViewModels/Stack/NiagaraStackEntry.cpp b/Engine/Plugins/FX/Niagara/Source/NiagaraEditor/Private/ViewModels/Stack/NiagaraStackEntry.cpp
--- a/Engine/Plugins/FX/Niagara/Source/NiagaraEditor/Private/ViewModels/Stack/NiagaraStackEntry.cpp
+++ b/Engine/Plugins/FX/Niagara/Source/NiagaraEditor/Private/ViewModels/Stack/NiagaraStackEntry.cpp
@@ -73,7 +73,11 @@ UNiagaraStackEntry::FStackIssue::FStackIssue(EStackIssueSeverity InSeverity, FTe
 	, Fixes(InFixes)
 {
 	checkf(ShortDescription.IsEmptyOrWhitespace() == false, TEXT("Short description can not be empty."));
-	checkf(LongDescription.IsEmptyOrWhitespace() == false, TEXT("Long description can not be empty."));
+	//checkf(LongDescription.IsEmptyOrWhitespace() == false, TEXT("Long description can not be empty."));
+	if (LongDescription.IsEmptyOrWhitespace())
+	{
+		LongDescription = ShortDescription;
+	}
 	checkf(InStackEditorDataKey.IsEmpty() == false, TEXT("Stack editor data key can not be empty."));
 }
 
diff --git a/Engine/Plugins/FX/Niagara/Source/NiagaraEditor/Public/ViewModels/NiagaraEmitterViewModel.h b/Engine/Plugins/FX/Niagara/Source/NiagaraEditor/Public/ViewModels/NiagaraEmitterViewModel.h
--- a/Engine/Plugins/FX/Niagara/Source/NiagaraEditor/Public/ViewModels/NiagaraEmitterViewModel.h
+++ b/Engine/Plugins/FX/Niagara/Source/NiagaraEditor/Public/ViewModels/NiagaraEmitterViewModel.h
@@ -110,6 +110,7 @@ private:
 	void SetEmitter(UNiagaraEmitter* InEmitter);
 
 	void OnVMCompiled(UNiagaraEmitter* InEmitter);
+	void OnGPUCompiled(UNiagaraEmitter* InEmitter);
 
 	void AddScriptEventHandlers();
 
diff --git a/Engine/Plugins/FX/Niagara/Source/NiagaraEditor/Public/ViewModels/NiagaraScriptViewModel.h b/Engine/Plugins/FX/Niagara/Source/NiagaraEditor/Public/ViewModels/NiagaraScriptViewModel.h
--- a/Engine/Plugins/FX/Niagara/Source/NiagaraEditor/Public/ViewModels/NiagaraScriptViewModel.h
+++ b/Engine/Plugins/FX/Niagara/Source/NiagaraEditor/Public/ViewModels/NiagaraScriptViewModel.h
@@ -98,6 +98,8 @@ protected:
 
 	virtual void OnVMScriptCompiled(UNiagaraScript* InScript);
 
+	virtual void OnGPUScriptCompiled(UNiagaraScript* InScript);
+
 	TWeakObjectPtr<UNiagaraScriptSource> Source;
 
 	/** The view model for the input parameter collection. */
diff --git a/Engine/Plugins/FX/Niagara/Source/NiagaraShader/Private/NiagaraShaderCompilationManager.cpp b/Engine/Plugins/FX/Niagara/Source/NiagaraShader/Private/NiagaraShaderCompilationManager.cpp
--- a/Engine/Plugins/FX/Niagara/Source/NiagaraShader/Private/NiagaraShaderCompilationManager.cpp
+++ b/Engine/Plugins/FX/Niagara/Source/NiagaraShader/Private/NiagaraShaderCompilationManager.cpp
@@ -245,6 +245,8 @@ void FNiagaraShaderCompilationManager::ProcessCompiledNiagaraShaderMaps(TMap<int
 						}
 						else
 						{
+							Script->SetCompileErrors(Errors);
+
 							// if we succeeded and our shader map is not complete this could be because the script was being edited quicker then the compile could be completed
 							// Don't modify scripts for which the compiled shader map is no longer complete
 							// This can happen if a script being compiled is edited
diff --git a/Engine/Plugins/FX/Niagara/Source/NiagaraShader/Private/NiagaraShared.cpp b/Engine/Plugins/FX/Niagara/Source/NiagaraShader/Private/NiagaraShared.cpp
--- a/Engine/Plugins/FX/Niagara/Source/NiagaraShader/Private/NiagaraShared.cpp
+++ b/Engine/Plugins/FX/Niagara/Source/NiagaraShader/Private/NiagaraShared.cpp
@@ -82,6 +82,8 @@ NIAGARASHADER_API void FNiagaraShaderScript::Invalidate()
 {
 	CancelCompilation();
 	ReleaseShaderMap();
+	CompileErrors.Empty();
+	HlslOutput.Empty();
 }
 
 
