- Fix for static analysis warning

#rnx

diff --git a/Engine/Plugins/FX/Niagara/Source/Niagara/Private/NiagaraDataInterfaceSkeletalMesh.cpp b/Engine/Plugins/FX/Niagara/Source/Niagara/Private/NiagaraDataInterfaceSkeletalMesh.cpp
--- a/Engine/Plugins/FX/Niagara/Source/Niagara/Private/NiagaraDataInterfaceSkeletalMesh.cpp
+++ b/Engine/Plugins/FX/Niagara/Source/Niagara/Private/NiagaraDataInterfaceSkeletalMesh.cpp
@@ -612,8 +612,21 @@ void FSkeletalMeshGpuDynamicBufferProxy::ReleaseRHI()
 void FSkeletalMeshGpuDynamicBufferProxy::NewFrame(const FNDISkeletalMesh_InstanceData* InstanceData, int32 LODIndex)
 {
 	// Grab Skeletal Component / Mesh, we must have a mesh at minimum to set the data
-	USkeletalMeshComponent* SkelComp = InstanceData != nullptr ? Cast<USkeletalMeshComponent>(InstanceData->Component.Get()) : nullptr;
-	USkeletalMesh* SkelMesh = SkelComp ? SkelComp->SkeletalMesh : InstanceData->MeshSafe.Get();
+	USkeletalMeshComponent* SkelComp = nullptr;
+	USkeletalMesh* SkelMesh = nullptr;
+	if (InstanceData != nullptr)
+	{
+		SkelComp = Cast<USkeletalMeshComponent>(InstanceData->Component.Get());
+		if ( SkelComp != nullptr )
+		{
+			SkelMesh = SkelComp->SkeletalMesh;
+		}
+		if (SkelMesh == nullptr)
+		{
+			SkelMesh = InstanceData->MeshSafe.Get();
+		}
+	}
+
 	if ( SkelMesh == nullptr )
 	{
 		return;
