Optimized calculating the bounds of CPU particles; emitter of 100 sprites has the cost reduced by 40%

diff --git a/Engine/Plugins/FX/Niagara/Source/Niagara/Classes/NiagaraDataSet.h b/Engine/Plugins/FX/Niagara/Source/Niagara/Classes/NiagaraDataSet.h
--- a/Engine/Plugins/FX/Niagara/Source/Niagara/Classes/NiagaraDataSet.h
+++ b/Engine/Plugins/FX/Niagara/Source/Niagara/Classes/NiagaraDataSet.h
@@ -768,6 +768,8 @@ struct FNiagaraDataSetAccessor<float> : public FNiagaraDataSetAccessorBase
 		DestBase[Index] = InValue;
 	}
 
+	FORCEINLINE const float* const GetX() const { return SrcBase; }
+
 private:
 	float* SrcBase;
 	float* DestBase;
@@ -853,6 +855,9 @@ struct FNiagaraDataSetAccessor<FVector2D> : public FNiagaraDataSetAccessorBase
 		DestYBase[Index] = InValue.Y;
 	}
 
+	FORCEINLINE const float* const GetX() const { return SrcXBase; }
+	FORCEINLINE const float* const GetY() const { return SrcYBase; }
+
 private:
 
 	float* SrcXBase;
@@ -947,6 +952,10 @@ struct FNiagaraDataSetAccessor<FVector> : public FNiagaraDataSetAccessorBase
 		DestZBase[Index] = InValue.Z;
 	}
 
+	FORCEINLINE const float* const GetX() const { return SrcXBase; }
+	FORCEINLINE const float* const GetY() const { return SrcYBase; }
+	FORCEINLINE const float* const GetZ() const { return SrcZBase; }
+
 private:
 
 	float* SrcXBase;
@@ -1515,6 +1524,7 @@ private:
 	uint32 CurrIdx;
 };
 
+#if WITH_EDITOR
 /**
 Allows immediate access to GPU data on the CPU, you can then use FNiagaraDataSetAccessor to access the data.
 This will make a copy of the GPU data and will stall the CPU until the data is ready from the GPU,
@@ -1542,7 +1552,7 @@ private:
 	NiagaraEmitterInstanceBatcher* Batcher = nullptr;
 	uint32				NumInstances = 0;
 };
-
+#endif
 
 //////////////////////////////////////////////////////////////////////////
 
diff --git a/Engine/Plugins/FX/Niagara/Source/Niagara/Classes/NiagaraEmitterInstance.h b/Engine/Plugins/FX/Niagara/Source/Niagara/Classes/NiagaraEmitterInstance.h
--- a/Engine/Plugins/FX/Niagara/Source/Niagara/Classes/NiagaraEmitterInstance.h
+++ b/Engine/Plugins/FX/Niagara/Source/Niagara/Classes/NiagaraEmitterInstance.h
@@ -58,9 +58,10 @@ public:
 
 	uint32 CalculateEventSpawnCount(const FNiagaraEventScriptProperties &EventHandlerProps, TArray<int32, TInlineAllocator<16>>& EventSpawnCounts, FNiagaraDataSet *EventSet);
 
-	/** Generate system bounds, reading back data from the GPU will introduce a stall and should only be used for debug purposes. */
-	FBox CalculateDynamicBounds(const bool bReadGPUSimulation = false);
+#if WITH_EDITOR
+	/** Potentially reads back data from the GPU which will introduce a stall and should only be used for debug purposes. */
 	NIAGARA_API void CalculateFixedBounds(const FTransform& ToWorldSpace);
+#endif
 
 	FNiagaraDataSet& GetData()const { return *ParticleDataSet; }
 
@@ -134,6 +135,9 @@ private:
 		const FNiagaraScriptExecutionContext& ExecContext,
 		FScriptExecutionConstantBufferTable& ConstantBufferTable) const;
 
+	/** Generate emitter bounds */
+	FBox InternalCalculateDynamicBounds(int32 ParticleCount) const;
+
 	template<typename TBindingType, typename TVariableType>
 	static void AddBinding(FName ParameterName, TVariableType ParameterType, TBindingType* SourceValuePtr, FNiagaraParameterStore& TargetParameterStore, TArray<TNiagaraFastPathAttributeBinding<TBindingType>>& TargetBindings)
 	{
diff --git a/Engine/Plugins/FX/Niagara/Source/Niagara/Private/NiagaraBoundsCalculatorHelper.h b/Engine/Plugins/FX/Niagara/Source/Niagara/Private/NiagaraBoundsCalculatorHelper.h
new file mode 100644
--- /dev/null
+++ b/Engine/Plugins/FX/Niagara/Source/Niagara/Private/NiagaraBoundsCalculatorHelper.h
@@ -0,0 +1,140 @@
+// Copyright Epic Games, Inc. All Rights Reserved.
+
+
+#pragma once
+
+#include "NiagaraBoundsCalculator.h"
+#include "NiagaraDataSet.h"
+
+template<bool bUsedWithSprites, bool bUsedWithMeshes, bool bUsedWithRibbons>
+class FNiagaraBoundsCalculatorHelper : public FNiagaraBoundsCalculator
+{
+public:
+	FNiagaraBoundsCalculatorHelper() = default;
+	FNiagaraBoundsCalculatorHelper(const FVector& InMeshExtents)
+		: MeshExtents(InMeshExtents)
+	{}
+
+	virtual void InitAccessors(FNiagaraDataSet& DataSet) override final
+	{
+		static const FName PositionName(TEXT("Position"));
+		static const FName SpriteSizeName(TEXT("SpriteSize"));
+		static const FName ScaleName(TEXT("Scale"));
+		static const FName RibbonWidthName(TEXT("RibbonWidth"));
+
+		PositionAccessor = FNiagaraDataSetAccessor<FVector>(DataSet, FNiagaraVariable(FNiagaraTypeDefinition::GetVec3Def(), PositionName));
+
+		if (bUsedWithSprites)
+		{
+			SpriteSizeAccessor = FNiagaraDataSetAccessor<FVector2D>(DataSet, FNiagaraVariable(FNiagaraTypeDefinition::GetVec2Def(), SpriteSizeName));
+		}
+		else
+		{
+			SpriteSizeAccessor = FNiagaraDataSetAccessor<FVector2D>();
+		}
+
+		if (bUsedWithMeshes)
+		{
+			ScaleAccessor = FNiagaraDataSetAccessor<FVector>(DataSet, FNiagaraVariable(FNiagaraTypeDefinition::GetVec3Def(), ScaleName));
+		}
+		else
+		{
+			ScaleAccessor = FNiagaraDataSetAccessor<FVector>();
+		}
+
+		if (bUsedWithRibbons)
+		{
+			RibbonWidthAccessor = FNiagaraDataSetAccessor<float>(DataSet, FNiagaraVariable(FNiagaraTypeDefinition::GetFloatDef(), RibbonWidthName));
+		}
+		else
+		{
+			RibbonWidthAccessor = FNiagaraDataSetAccessor<float>();
+		}
+	}
+
+	virtual void RefreshAccessors() override final
+	{
+		PositionAccessor.InitForAccess();
+
+		if (bUsedWithSprites)
+		{
+			SpriteSizeAccessor.InitForAccess();
+		}
+		if (bUsedWithMeshes)
+		{
+			ScaleAccessor.InitForAccess();
+		}
+		if (bUsedWithRibbons)
+		{
+			RibbonWidthAccessor.InitForAccess();
+		}
+	}
+
+	virtual FBox CalculateBounds(const int32 NumInstances) const override final
+	{
+		if (!NumInstances || !PositionAccessor.IsValid())
+		{
+			return FBox(ForceInit);
+		}
+
+		constexpr float kDefaultSize = 50.0f;
+
+		const FVector PositionMax = FVector(
+			GetFloatMax(MakeArrayView(PositionAccessor.GetX(), NumInstances)),
+			GetFloatMax(MakeArrayView(PositionAccessor.GetY(), NumInstances)),
+			GetFloatMax(MakeArrayView(PositionAccessor.GetZ(), NumInstances)));
+
+		const FVector PositionMin = FVector(
+			GetFloatMin(MakeArrayView(PositionAccessor.GetX(), NumInstances)),
+			GetFloatMin(MakeArrayView(PositionAccessor.GetY(), NumInstances)),
+			GetFloatMin(MakeArrayView(PositionAccessor.GetZ(), NumInstances)));
+
+		float MaxSize = KINDA_SMALL_NUMBER;
+		if (bUsedWithMeshes)
+		{
+			FVector MaxScale = FVector(kDefaultSize, kDefaultSize, kDefaultSize);
+			if (ScaleAccessor.IsValid())
+			{
+				MaxScale = FVector(
+					GetFloatMax(MakeArrayView(ScaleAccessor.GetX(), NumInstances)),
+					GetFloatMax(MakeArrayView(ScaleAccessor.GetY(), NumInstances)),
+					GetFloatMax(MakeArrayView(ScaleAccessor.GetZ(), NumInstances)));
+			}
+
+			const FVector ScaledExtents = MeshExtents * (MaxScale.IsNearlyZero() ? FVector::OneVector : MaxScale);
+			MaxSize = FMath::Max(MaxSize, ScaledExtents.GetMax());
+		}
+
+		if (bUsedWithSprites)
+		{
+			float MaxSpriteSize = kDefaultSize;
+
+			if (SpriteSizeAccessor.IsValid())
+			{
+				MaxSpriteSize = FMath::Max(
+					GetFloatMax(MakeArrayView(SpriteSizeAccessor.GetX(), NumInstances)),
+					GetFloatMax(MakeArrayView(SpriteSizeAccessor.GetY(), NumInstances)));
+			}
+			MaxSize = FMath::Max(MaxSize, FMath::IsNearlyZero(MaxSpriteSize) ? 1.0f : MaxSpriteSize);
+		}
+
+		if (bUsedWithRibbons)
+		{
+			float MaxRibbonWidth = kDefaultSize;
+			if (RibbonWidthAccessor.IsValid())
+			{
+				MaxRibbonWidth = GetFloatMax(MakeArrayView(RibbonWidthAccessor.GetX(), NumInstances));
+			}
+
+			MaxSize = FMath::Max(MaxSize, FMath::IsNearlyZero(MaxRibbonWidth) ? 1.0f : MaxRibbonWidth);
+		}
+
+		return FBox(PositionMin, PositionMax).ExpandBy(MaxSize);
+	}
+
+	FNiagaraDataSetAccessor<FVector> PositionAccessor;
+	FNiagaraDataSetAccessor<FVector2D> SpriteSizeAccessor;
+	FNiagaraDataSetAccessor<FVector> ScaleAccessor;
+	FNiagaraDataSetAccessor<float> RibbonWidthAccessor;
+	const FVector MeshExtents = FVector::OneVector;
+};
diff --git a/Engine/Plugins/FX/Niagara/Source/Niagara/Private/NiagaraDataSet.cpp b/Engine/Plugins/FX/Niagara/Source/Niagara/Private/NiagaraDataSet.cpp
--- a/Engine/Plugins/FX/Niagara/Source/Niagara/Private/NiagaraDataSet.cpp
+++ b/Engine/Plugins/FX/Niagara/Source/Niagara/Private/NiagaraDataSet.cpp
@@ -1026,6 +1026,7 @@ void FNiagaraDataBuffer::ReleaseGPUInstanceCount(FNiagaraGPUInstanceCountManager
 }
 
 
+#if WITH_EDITOR
 void FScopedNiagaraDataSetGPUReadback::ReadbackData(NiagaraEmitterInstanceBatcher* InBatcher, FNiagaraDataSet* InDataSet)
 {
 	check(DataSet == nullptr);
@@ -1083,6 +1084,7 @@ void FScopedNiagaraDataSetGPUReadback::ReadbackData(NiagaraEmitterInstanceBatche
 	);
 	FlushRenderingCommands();
 }
+#endif
 
 //////////////////////////////////////////////////////////////////////////
 
diff --git a/Engine/Plugins/FX/Niagara/Source/Niagara/Private/NiagaraEmitterInstance.cpp b/Engine/Plugins/FX/Niagara/Source/Niagara/Private/NiagaraEmitterInstance.cpp
--- a/Engine/Plugins/FX/Niagara/Source/Niagara/Private/NiagaraEmitterInstance.cpp
+++ b/Engine/Plugins/FX/Niagara/Source/Niagara/Private/NiagaraEmitterInstance.cpp
@@ -404,10 +404,12 @@ void FNiagaraEmitterInstance::Init(int32 InEmitterIdx, FNiagaraSystemInstanceID
 		}
 		FNiagaraUtilities::CollectScriptDataInterfaceParameters(*CachedEmitter, MakeArrayView(Scripts), ScriptDefinedDataInterfaceParameters);
 
-		// Initialize bounds calculators
-		//-OPT: Could skip creating this if we won't ever use it
+		// Initialize bounds calculators - skip creating if we won't ever use it.  We leave the GPU sims in there with the editor so that we can
+		// generate the bounds from the readback in the tool.
+#if !WITH_EDITOR
 		bool bUseDynamicBounds = !CachedSystemFixedBounds.IsSet() && !CachedEmitter->bFixedBounds && CachedEmitter->SimTarget == ENiagaraSimTarget::CPUSim;
 		if (bUseDynamicBounds)
+#endif
 		{
 			BoundsCalculators.Reserve(CachedEmitter->GetRenderers().Num());
 			for (UNiagaraRendererProperties* RendererProperties : CachedEmitter->GetRenderers())
@@ -764,18 +766,44 @@ int FNiagaraEmitterInstance::GetTotalBytesUsed()
 	return BytesUsed;
 }
 
-FBox FNiagaraEmitterInstance::CalculateDynamicBounds(const bool bReadGPUSimulation)
+FBox FNiagaraEmitterInstance::InternalCalculateDynamicBounds(int32 ParticleCount) const
 {
-	if (IsComplete() || !BoundsCalculators.Num() || CachedEmitter == nullptr)
+	if (!ParticleCount || !BoundsCalculators.Num())
+	{
 		return FBox(ForceInit);
+	}
+
+	FBox Ret;
+	Ret.Init();
+
+	for (const TUniquePtr<FNiagaraBoundsCalculator>& BoundsCalculator : BoundsCalculators)
+	{
+		BoundsCalculator->RefreshAccessors();
+		Ret += BoundsCalculator->CalculateBounds(ParticleCount);
+	}
+
+	return Ret;
+}
+
+#if WITH_EDITOR
+void FNiagaraEmitterInstance::CalculateFixedBounds(const FTransform& ToWorldSpace)
+{
+	check(CachedEmitter);
+
+	if (IsComplete() || CachedEmitter == nullptr)
+	{
+		return ;
+	}
 
 	FScopedNiagaraDataSetGPUReadback ScopedGPUReadback;
 
 	int32 NumInstances = 0;
 	if (CachedEmitter->SimTarget == ENiagaraSimTarget::GPUComputeSim)
 	{
-		if (!bReadGPUSimulation || (GPUExecContext == nullptr))
-			return FBox(ForceInit);
+		if (GPUExecContext == nullptr)
+		{
+			return;
+		}
 
 		ScopedGPUReadback.ReadbackData(Batcher, GPUExecContext->MainDataSet);
 		NumInstances = ScopedGPUReadback.GetNumInstances();
@@ -786,40 +814,11 @@ FBox FNiagaraEmitterInstance::CalculateDynamicBounds(const bool bReadGPUSimulati
 	}
 
 	if (NumInstances == 0)
-		return FBox(ForceInit);
-
-	FBox Ret;
-	Ret.Init();
-
-	bool bContainsNaN = false;
-	for ( const TUniquePtr<FNiagaraBoundsCalculator>& BoundsCalculator : BoundsCalculators )
 	{
-		Ret += BoundsCalculator->CalculateBounds(NumInstances, bContainsNaN);
-	}
-
-#if !UE_BUILD_SHIPPING && !UE_BUILD_TEST
-	if (bContainsNaN && ParentSystemInstance != nullptr && CachedEmitter != nullptr && ParentSystemInstance->GetSystem() != nullptr)
-	{
-		const FString OnScreenMessage = FString::Printf(TEXT("Niagara Particle position data contains NaNs. Likely a divide by zero somewhere in your modules. Emitter \"%s\" in System \"%s\".  Use fx.Niagara.DumpNansOnce to get full log."), *CachedEmitter->GetName(), *ParentSystemInstance->GetSystem()->GetName());
-		GEngine->AddOnScreenDebugMessage((uint64)((PTRINT)this), 3.f, FColor::Red, OnScreenMessage);
-
-		if (GbNiagaraDumpNans || GbNiagaraDumpNansOnce)
-		{
-			GbNiagaraDumpNansOnce = 0;
-			UE_LOG(LogNiagara, Warning, TEXT("Particle position data contains NaNs. Likely a divide by zero somewhere in your modules. Emitter \"%s\" in System \"%s\""), *CachedEmitter->GetName(), *ParentSystemInstance->GetSystem()->GetName());
-			ParentSystemInstance->Dump();
-		}
+		return;
 	}
-#endif
-
-	return Ret;
-}
 
-void FNiagaraEmitterInstance::CalculateFixedBounds(const FTransform& ToWorldSpace)
-{
-	check(CachedEmitter);
-
-	FBox Bounds = CalculateDynamicBounds(true);
+	FBox Bounds = InternalCalculateDynamicBounds(NumInstances);
 	if (!Bounds.IsValid)
 		return;
 
@@ -836,6 +835,7 @@ void FNiagaraEmitterInstance::CalculateFixedBounds(const FTransform& ToWorldSpac
 
 	CachedBounds = Bounds;
 }
+#endif
 
 /** 
   * Do any post work such as calculating dynamic bounds.
@@ -859,7 +859,7 @@ void FNiagaraEmitterInstance::PostTick()
 	}
 	else
 	{
-		FBox DynamicBounds = CalculateDynamicBounds();
+		FBox DynamicBounds = InternalCalculateDynamicBounds(ParticleDataSet->GetCurrentDataChecked().GetNumInstances());
 		if (DynamicBounds.IsValid)
 		{
 			if (CachedEmitter->bLocalSpace)
diff --git a/Engine/Plugins/FX/Niagara/Source/Niagara/Private/NiagaraMeshRendererProperties.cpp b/Engine/Plugins/FX/Niagara/Source/Niagara/Private/NiagaraMeshRendererProperties.cpp
--- a/Engine/Plugins/FX/Niagara/Source/Niagara/Private/NiagaraMeshRendererProperties.cpp
+++ b/Engine/Plugins/FX/Niagara/Source/Niagara/Private/NiagaraMeshRendererProperties.cpp
@@ -71,8 +71,8 @@ FNiagaraBoundsCalculator* UNiagaraMeshRendererProperties::CreateBoundsCalculator
 {
 	if (ParticleMesh)
 	{
-		FNiagaraBoundsCalculatorHelper<false, true, false>* BoundsCalculator = new FNiagaraBoundsCalculatorHelper<false, true, false>();
-		BoundsCalculator->MeshExtents = ParticleMesh->GetBounds().BoxExtent;
+		FNiagaraBoundsCalculatorHelper<false, true, false>* BoundsCalculator
+			= new FNiagaraBoundsCalculatorHelper<false, true, false>(ParticleMesh->GetBounds().BoxExtent);
 		return BoundsCalculator;
 	}
 
diff --git a/Engine/Plugins/FX/Niagara/Source/Niagara/Public/NiagaraBoundsCalculator.h b/Engine/Plugins/FX/Niagara/Source/Niagara/Public/NiagaraBoundsCalculator.h
--- a/Engine/Plugins/FX/Niagara/Source/Niagara/Public/NiagaraBoundsCalculator.h
+++ b/Engine/Plugins/FX/Niagara/Source/Niagara/Public/NiagaraBoundsCalculator.h
@@ -7,10 +7,37 @@
 #include "Stats/Stats.h"
 #include "NiagaraModule.h"
 
+class FNiagaraDataSet;
+
 class FNiagaraBoundsCalculator
 {
 public:
 	virtual ~FNiagaraBoundsCalculator() { }
 	virtual void InitAccessors(FNiagaraDataSet& DataSet) = 0;
-	virtual FBox CalculateBounds(const int32 NumInstances, bool& bOutContainsNaN) = 0;
+	virtual void RefreshAccessors() = 0;
+	virtual FBox CalculateBounds(const int32 NumInstances) const = 0;
+
+	FORCEINLINE static float GetFloatMin(TArrayView<const float> Values)
+	{
+		float MinValue = FLT_MAX;
+
+		for (auto Value : Values)
+		{
+			MinValue = FMath::Min(MinValue, Value);
+		}
+
+		return MinValue;
+	}
+
+	FORCEINLINE static float GetFloatMax(TArrayView<const float> Values)
+	{
+		float MaxValue = -FLT_MAX;
+
+		for (auto Value : Values)
+		{
+			MaxValue = FMath::Max(MaxValue, Value);
+		}
+
+		return MaxValue;
+	}
 };
diff --git a/Engine/Plugins/FX/Niagara/Source/Niagara/Public/NiagaraBoundsCalculatorHelper.h b/Engine/Plugins/FX/Niagara/Source/Niagara/Public/NiagaraBoundsCalculatorHelper.h
deleted file mode 100644
