- Wait for system tick to complete before adding instances otherwise we can race on the data set

diff --git a/Engine/Plugins/FX/Niagara/Source/Niagara/Private/NiagaraSystemSimulation.cpp b/Engine/Plugins/FX/Niagara/Source/Niagara/Private/NiagaraSystemSimulation.cpp
--- a/Engine/Plugins/FX/Niagara/Source/Niagara/Private/NiagaraSystemSimulation.cpp
+++ b/Engine/Plugins/FX/Niagara/Source/Niagara/Private/NiagaraSystemSimulation.cpp
@@ -1776,6 +1776,8 @@ void FNiagaraSystemSimulation::AddInstance(FNiagaraSystemInstance* Instance)
 	check(IsInGameThread());
 	check(Instance->SystemInstanceIndex == INDEX_NONE);
 
+	WaitForSystemTickComplete();
+
 	Instance->SetPendingSpawn(true);
 	Instance->SystemInstanceIndex = PendingSystemInstances.Add(Instance);
 
