Niagara - Scratch Pad MVP

diff --git a/Engine/Plugins/FX/Niagara/Source/Niagara/Classes/NiagaraEmitter.h b/Engine/Plugins/FX/Niagara/Source/Niagara/Classes/NiagaraEmitter.h
--- a/Engine/Plugins/FX/Niagara/Source/Niagara/Classes/NiagaraEmitter.h
+++ b/Engine/Plugins/FX/Niagara/Source/Niagara/Classes/NiagaraEmitter.h
@@ -399,6 +399,12 @@ public:
 
 	UPROPERTY(EditAnywhere, AdvancedDisplay, Category = "Asset Options", AssetRegistrySearchable)
 	FText TemplateAssetDescription;
+
+	UPROPERTY()
+	TArray<UNiagaraScript*> ScratchPadScripts;
+
+	UPROPERTY()
+	TArray<UNiagaraScript*> ParentScratchPadScripts;
 	
 	/** Callback issued whenever a VM compilation successfully happened (even if the results are a script that cannot be executed due to errors)*/
 	NIAGARA_API FOnEmitterCompiled& OnEmitterVMCompiled();
diff --git a/Engine/Plugins/FX/Niagara/Source/Niagara/Classes/NiagaraSystem.h b/Engine/Plugins/FX/Niagara/Source/Niagara/Classes/NiagaraSystem.h
--- a/Engine/Plugins/FX/Niagara/Source/Niagara/Classes/NiagaraSystem.h
+++ b/Engine/Plugins/FX/Niagara/Source/Niagara/Classes/NiagaraSystem.h
@@ -303,6 +303,9 @@ public:
 	UPROPERTY(EditAnywhere, AdvancedDisplay, Category = "Asset Options", AssetRegistrySearchable)
 	FText TemplateAssetDescription;
 
+	UPROPERTY()
+	TArray<UNiagaraScript*> ScratchPadScripts;
+
 	bool GetIsolateEnabled() const;
 	void SetIsolateEnabled(bool bIsolate);
 	
diff --git a/Engine/Plugins/FX/Niagara/Source/Niagara/Private/NiagaraEmitter.cpp b/Engine/Plugins/FX/Niagara/Source/Niagara/Private/NiagaraEmitter.cpp
--- a/Engine/Plugins/FX/Niagara/Source/Niagara/Private/NiagaraEmitter.cpp
+++ b/Engine/Plugins/FX/Niagara/Source/Niagara/Private/NiagaraEmitter.cpp
@@ -542,6 +542,8 @@ UNiagaraEmitter* UNiagaraEmitter::CreateWithParentAndOwner(UNiagaraEmitter& InPa
 	NewEmitter->Parent = &InParentEmitter;
 	NewEmitter->ParentAtLastMerge = Cast<UNiagaraEmitter>(StaticDuplicateObject(&InParentEmitter, NewEmitter));
 	NewEmitter->ParentAtLastMerge->ClearFlags(RF_Standalone | RF_Public);
+	NewEmitter->ParentScratchPadScripts.Append(NewEmitter->ScratchPadScripts);
+	NewEmitter->ScratchPadScripts.Empty();
 	NewEmitter->SetUniqueEmitterName(InName.GetPlainNameString());
 	NewEmitter->GraphSource->MarkNotSynchronized(InitialNotSynchronizedReason);
 	return NewEmitter;
@@ -1174,6 +1176,15 @@ void UNiagaraEmitter::UpdateFromMergedCopy(const INiagaraMergeManager& MergeMana
 		MergedRenderer->OnChanged().AddUObject(this, &UNiagaraEmitter::RendererChanged);
 	}
 
+	// Copy parent scratch pad scripts.
+	ParentScratchPadScripts.Empty();
+
+	for (UNiagaraScript* MergedParentScratchPadScript : MergedEmitter->ParentScratchPadScripts)
+	{
+		ReouterMergedObject(this, MergedParentScratchPadScript);
+		ParentScratchPadScripts.Add(MergedParentScratchPadScript);
+	}
+
 	SetEditorData(MergedEmitter->GetEditorData());
 
 	// Update the change id since we don't know what's changed.
diff --git a/Engine/Plugins/FX/Niagara/Source/Niagara/Public/INiagaraMergeManager.h b/Engine/Plugins/FX/Niagara/Source/Niagara/Public/INiagaraMergeManager.h
--- a/Engine/Plugins/FX/Niagara/Source/Niagara/Public/INiagaraMergeManager.h
+++ b/Engine/Plugins/FX/Niagara/Source/Niagara/Public/INiagaraMergeManager.h
@@ -19,6 +19,7 @@ public:
 		FMergeEmitterResults()
 			: MergeResult(EMergeEmitterResult::None)
 			, bModifiedGraph(false)
+			, MergedInstance(nullptr)
 		{
 		}
 
diff --git a/Engine/Plugins/FX/Niagara/Source/NiagaraEditor/Private/Customizations/NiagaraScriptDetails.cpp b/Engine/Plugins/FX/Niagara/Source/NiagaraEditor/Private/Customizations/NiagaraScriptDetails.cpp
--- a/Engine/Plugins/FX/Niagara/Source/NiagaraEditor/Private/Customizations/NiagaraScriptDetails.cpp
+++ b/Engine/Plugins/FX/Niagara/Source/NiagaraEditor/Private/Customizations/NiagaraScriptDetails.cpp
@@ -109,17 +109,37 @@ private:
 	TSharedPtr<SComboButton> ComboButton;
 };
 
-TSharedRef<IDetailCustomization> FNiagaraScriptDetails::MakeInstance(TWeakPtr<FNiagaraScriptViewModel> ScriptViewModel)
+TSharedRef<IDetailCustomization> FNiagaraScriptDetails::MakeInstance()
 {
-	return MakeShared<FNiagaraScriptDetails>(ScriptViewModel.Pin());
+	return MakeShared<FNiagaraScriptDetails>();
 }
 
-FNiagaraScriptDetails::FNiagaraScriptDetails(TSharedPtr<FNiagaraScriptViewModel> InScriptViewModel)
-	: ScriptViewModel(InScriptViewModel)
-{}
+FNiagaraScriptDetails::FNiagaraScriptDetails()
+{
+}
 
 void FNiagaraScriptDetails::CustomizeDetails(IDetailLayoutBuilder& DetailBuilder)
 {
+	TArray<TWeakObjectPtr<UObject>> ObjectsBeingCustomized;
+	DetailBuilder.GetObjectsBeingCustomized(ObjectsBeingCustomized);
+	if (ObjectsBeingCustomized.Num() != 1 || ObjectsBeingCustomized[0]->IsA<UNiagaraScript>() == false)
+	{
+		return;
+	}
+
+	UNiagaraScript* ScriptBeingCustomized = CastChecked<UNiagaraScript>(ObjectsBeingCustomized[0]);
+	TArray<TSharedPtr<FNiagaraScriptViewModel>> ExistingViewModels;
+	TNiagaraViewModelManager<UNiagaraScript, FNiagaraScriptViewModel>::GetAllViewModelsForObject(ScriptBeingCustomized, ExistingViewModels);
+	
+	if (ensureMsgf(ExistingViewModels.Num() == 1, TEXT("Missing or duplicate script view models detected.  Can not create script details customization.")))
+	{
+		ScriptViewModel = ExistingViewModels[0];
+	}
+	else
+	{
+		return;
+	}
+
 	static const FName InputParamCategoryName = TEXT("NiagaraScript_InputParams");
 	static const FName OutputParamCategoryName = TEXT("NiagaraScript_OutputParams");
 	static const FName ScriptCategoryName = TEXT("Script");
diff --git a/Engine/Plugins/FX/Niagara/Source/NiagaraEditor/Private/Customizations/NiagaraScriptDetails.h b/Engine/Plugins/FX/Niagara/Source/NiagaraEditor/Private/Customizations/NiagaraScriptDetails.h
--- a/Engine/Plugins/FX/Niagara/Source/NiagaraEditor/Private/Customizations/NiagaraScriptDetails.h
+++ b/Engine/Plugins/FX/Niagara/Source/NiagaraEditor/Private/Customizations/NiagaraScriptDetails.h
@@ -18,9 +18,9 @@ class IPropertyTypeCustomizationUtils;
 class FNiagaraScriptDetails : public IDetailCustomization
 {
 public:
-	static TSharedRef<IDetailCustomization> MakeInstance(TWeakPtr<FNiagaraScriptViewModel> ScriptViewModel);
+	static TSharedRef<IDetailCustomization> MakeInstance();
 	
-	FNiagaraScriptDetails(TSharedPtr<FNiagaraScriptViewModel> InScriptViewModel);
+	FNiagaraScriptDetails();
 
 	virtual void CustomizeDetails(IDetailLayoutBuilder& DetailBuilder);
 
diff --git a/Engine/Plugins/FX/Niagara/Source/NiagaraEditor/Private/EdGraphSchema_Niagara.cpp b/Engine/Plugins/FX/Niagara/Source/NiagaraEditor/Private/EdGraphSchema_Niagara.cpp
--- a/Engine/Plugins/FX/Niagara/Source/NiagaraEditor/Private/EdGraphSchema_Niagara.cpp
+++ b/Engine/Plugins/FX/Niagara/Source/NiagaraEditor/Private/EdGraphSchema_Niagara.cpp
@@ -403,12 +403,11 @@ TArray<TSharedPtr<FNiagaraSchemaAction_NewNode> > UEdGraphSchema_Niagara::GetGra
 		FText Keywords;
 		ScriptAsset.GetTagValue(GET_MEMBER_NAME_CHECKED(UNiagaraScript, Keywords), Keywords);
 
-		FString DisplayNameString = FName::NameToDisplayString(ScriptAsset.AssetName.ToString(), false);
+		bool bIsInLibrary = FNiagaraEditorUtilities::IsScriptAssetInLibrary(ScriptAsset);
+		const FText MenuDesc = FNiagaraEditorUtilities::FormatScriptName(ScriptAsset.AssetName, bIsInLibrary);
+		const FText TooltipDesc = FNiagaraEditorUtilities::FormatScriptDescription(AssetDesc, ScriptAsset.ObjectPath, bIsInLibrary);
 
-		const FText MenuDesc = FText::FromString(DisplayNameString);
-		const FText TooltipDesc = FNiagaraEditorUtilities::FormatScriptAssetDescription(AssetDesc, ScriptAsset.ObjectPath);
-
-		TSharedPtr<FNiagaraSchemaAction_NewNode> FunctionCallAction = AddNewNodeAction(NewActions, Category, MenuDesc, *DisplayNameString, TooltipDesc, Keywords);
+		TSharedPtr<FNiagaraSchemaAction_NewNode> FunctionCallAction = AddNewNodeAction(NewActions, Category, MenuDesc, ScriptAsset.AssetName, TooltipDesc, Keywords);
 
 		UNiagaraNodeFunctionCall* FunctionCallNode = NewObject<UNiagaraNodeFunctionCall>(OwnerOfTemporaries);
 		FunctionCallNode->FunctionScriptAssetObjectPath = ScriptAsset.ObjectPath;
diff --git a/Engine/Plugins/FX/Niagara/Source/NiagaraEditor/Private/NiagaraCompiler.cpp b/Engine/Plugins/FX/Niagara/Source/NiagaraEditor/Private/NiagaraCompiler.cpp
--- a/Engine/Plugins/FX/Niagara/Source/NiagaraEditor/Private/NiagaraCompiler.cpp
+++ b/Engine/Plugins/FX/Niagara/Source/NiagaraEditor/Private/NiagaraCompiler.cpp
@@ -193,6 +193,10 @@ void FNiagaraCompileRequestData::VisitReferencedGraphsRecursive(UNiagaraGraph* I
 							else
 							{
 								DupeScript = DuplicateObject<UNiagaraScript>(FunctionScript, InNode, FunctionScript->GetFName());
+								if (DupeScript->GetSource()->GetOuter() != DupeScript)
+								{
+									DupeScript->SetSource(DuplicateObject<UNiagaraScriptSource>(CastChecked<UNiagaraScriptSource>(DupeScript->GetSource()), DupeScript, DupeScript->GetSource()->GetFName()));
+								}
 								ProcessedGraph = Cast<UNiagaraScriptSource>(DupeScript->GetSource())->NodeGraph;
 								FEdGraphUtilities::MergeChildrenGraphsIn(ProcessedGraph, ProcessedGraph, /*bRequireSchemaMatch=*/ true);
 								FNiagaraEditorUtilities::PreprocessFunctionGraph(Schema, ProcessedGraph, CallInputs, CallOutputs, ScriptUsage, ConstantResolver);
diff --git a/Engine/Plugins/FX/Niagara/Source/NiagaraEditor/Private/NiagaraEditorModule.cpp b/Engine/Plugins/FX/Niagara/Source/NiagaraEditor/Private/NiagaraEditorModule.cpp
--- a/Engine/Plugins/FX/Niagara/Source/NiagaraEditor/Private/NiagaraEditorModule.cpp
+++ b/Engine/Plugins/FX/Niagara/Source/NiagaraEditor/Private/NiagaraEditorModule.cpp
@@ -95,6 +95,15 @@
 #include "Customizations/NiagaraTypeCustomizations.h"
 #include "Customizations/NiagaraEventScriptPropertiesCustomization.h"
 #include "Customizations/NiagaraScriptVariableCustomization.h"
+#include "Customizations/NiagaraScriptDetails.h"
+
+#include "NiagaraComponent.h"
+#include "NiagaraNodeStaticSwitch.h"
+#include "NiagaraScriptVariable.h"
+#include "NiagaraScript.h"
+#include "NiagaraCommon.h"
+#include "NiagaraScriptHighlight.h"
+
 #include "HAL/IConsoleManager.h"
 #include "NiagaraHlslTranslator.h"
 #include "NiagaraThumbnailRenderer.h"
@@ -104,7 +113,6 @@
 #include "NiagaraNodeFunctionCall.h"
 #include "Engine/Selection.h"
 #include "NiagaraActor.h"
-#include "NiagaraNodeFunctionCall.h"
 #include "INiagaraEditorOnlyDataUtlities.h"
 
 #include "Editor.h"
@@ -736,46 +744,58 @@ void FNiagaraEditorModule::StartupModule()
 	
 	// register details customization
 	FPropertyEditorModule& PropertyModule = FModuleManager::LoadModuleChecked<FPropertyEditorModule>("PropertyEditor");
-	PropertyModule.RegisterCustomClassLayout("NiagaraComponent", FOnGetDetailCustomizationInstance::CreateStatic(&FNiagaraComponentDetails::MakeInstance));
 
-	PropertyModule.RegisterCustomClassLayout("NiagaraNodeStaticSwitch", FOnGetDetailCustomizationInstance::CreateStatic(&FNiagaraStaticSwitchNodeDetails::MakeInstance));
+	PropertyModule.RegisterCustomClassLayout(
+		UNiagaraComponent::StaticClass()->GetFName(),
+		FOnGetDetailCustomizationInstance::CreateStatic(&FNiagaraComponentDetails::MakeInstance));
 
-	PropertyModule.RegisterCustomClassLayout("NiagaraScriptVariable", FOnGetDetailCustomizationInstance::CreateStatic(&FNiagaraScriptVariableDetails::MakeInstance));
+	PropertyModule.RegisterCustomClassLayout(
+		UNiagaraNodeStaticSwitch::StaticClass()->GetFName(),
+		FOnGetDetailCustomizationInstance::CreateStatic(&FNiagaraStaticSwitchNodeDetails::MakeInstance));
 
-	PropertyModule.RegisterCustomClassLayout("NiagaraNodeFunctionCall", FOnGetDetailCustomizationInstance::CreateStatic(&FNiagaraFunctionCallNodeDetails::MakeInstance));
-	
-	PropertyModule.RegisterCustomPropertyTypeLayout("NiagaraFloat",
-		FOnGetPropertyTypeCustomizationInstance::CreateStatic(&FNiagaraNumericCustomization::MakeInstance)
-	);
+	PropertyModule.RegisterCustomClassLayout(
+		UNiagaraScriptVariable::StaticClass()->GetFName(), 
+		FOnGetDetailCustomizationInstance::CreateStatic(&FNiagaraScriptVariableDetails::MakeInstance));
 
-	PropertyModule.RegisterCustomPropertyTypeLayout("NiagaraInt32",
-		FOnGetPropertyTypeCustomizationInstance::CreateStatic(&FNiagaraNumericCustomization::MakeInstance)
-	);
+	PropertyModule.RegisterCustomClassLayout(
+		UNiagaraNodeFunctionCall::StaticClass()->GetFName(),
+		FOnGetDetailCustomizationInstance::CreateStatic(&FNiagaraFunctionCallNodeDetails::MakeInstance));
 
-	PropertyModule.RegisterCustomPropertyTypeLayout("NiagaraNumeric",
-		FOnGetPropertyTypeCustomizationInstance::CreateStatic(&FNiagaraNumericCustomization::MakeInstance)
-	);
+	PropertyModule.RegisterCustomClassLayout(
+		UNiagaraScript::StaticClass()->GetFName(),
+		FOnGetDetailCustomizationInstance::CreateStatic(&FNiagaraScriptDetails::MakeInstance));
+	
+	PropertyModule.RegisterCustomPropertyTypeLayout(
+		FNiagaraFloat::StaticStruct()->GetFName(),
+		FOnGetPropertyTypeCustomizationInstance::CreateStatic(&FNiagaraNumericCustomization::MakeInstance));
 
-	PropertyModule.RegisterCustomPropertyTypeLayout("NiagaraParameterMap",
-		FOnGetPropertyTypeCustomizationInstance::CreateStatic(&FNiagaraNumericCustomization::MakeInstance)
-	);
+	PropertyModule.RegisterCustomPropertyTypeLayout(
+		FNiagaraInt32::StaticStruct()->GetFName(),
+		FOnGetPropertyTypeCustomizationInstance::CreateStatic(&FNiagaraNumericCustomization::MakeInstance));
 
+	PropertyModule.RegisterCustomPropertyTypeLayout(
+		FNiagaraNumeric::StaticStruct()->GetFName(),
+		FOnGetPropertyTypeCustomizationInstance::CreateStatic(&FNiagaraNumericCustomization::MakeInstance));
 
-	PropertyModule.RegisterCustomPropertyTypeLayout("NiagaraBool",
-		FOnGetPropertyTypeCustomizationInstance::CreateStatic(&FNiagaraBoolCustomization::MakeInstance)
-	);
+	PropertyModule.RegisterCustomPropertyTypeLayout(
+		FNiagaraParameterMap::StaticStruct()->GetFName(),
+		FOnGetPropertyTypeCustomizationInstance::CreateStatic(&FNiagaraNumericCustomization::MakeInstance));
 
-	PropertyModule.RegisterCustomPropertyTypeLayout("NiagaraMatrix",
-		FOnGetPropertyTypeCustomizationInstance::CreateStatic(&FNiagaraMatrixCustomization::MakeInstance)
-	);
+	PropertyModule.RegisterCustomPropertyTypeLayout(
+		FNiagaraBool::StaticStruct()->GetFName(),
+		FOnGetPropertyTypeCustomizationInstance::CreateStatic(&FNiagaraBoolCustomization::MakeInstance));
 
-	PropertyModule.RegisterCustomPropertyTypeLayout("NiagaraVariableAttributeBinding",
-		FOnGetPropertyTypeCustomizationInstance::CreateStatic(&FNiagaraVariableAttributeBindingCustomization::MakeInstance)
-	);
+	PropertyModule.RegisterCustomPropertyTypeLayout(
+		FNiagaraMatrix::StaticStruct()->GetFName(),
+		FOnGetPropertyTypeCustomizationInstance::CreateStatic(&FNiagaraMatrixCustomization::MakeInstance));
+
+	PropertyModule.RegisterCustomPropertyTypeLayout(
+		FNiagaraVariableAttributeBinding::StaticStruct()->GetFName(),
+		FOnGetPropertyTypeCustomizationInstance::CreateStatic(&FNiagaraVariableAttributeBindingCustomization::MakeInstance));
 	
-	PropertyModule.RegisterCustomPropertyTypeLayout("NiagaraScriptVariableBinding",
-		FOnGetPropertyTypeCustomizationInstance::CreateStatic(&FNiagaraScriptVariableBindingCustomization::MakeInstance)
-	);
+	PropertyModule.RegisterCustomPropertyTypeLayout(
+		FNiagaraScriptVariableBinding::StaticStruct()->GetFName(),
+		FOnGetPropertyTypeCustomizationInstance::CreateStatic(&FNiagaraScriptVariableBindingCustomization::MakeInstance));
 
 	PropertyModule.RegisterCustomPropertyTypeLayout("NiagaraPlatformSet",
 		FOnGetPropertyTypeCustomizationInstance::CreateStatic(&FNiagaraPlatformSetTypeCustomization::MakeInstance)
@@ -785,9 +805,13 @@ void FNiagaraEditorModule::StartupModule()
 // 		FOnGetPropertyTypeCustomizationInstance::CreateStatic(&FNiagaraSystemScalabilitySettingsArrayTypeCustomization::MakeInstance)
 // 	);
 	
-	PropertyModule.RegisterCustomPropertyTypeLayout("NiagaraUserParameterBinding",
-		FOnGetPropertyTypeCustomizationInstance::CreateStatic(&FNiagaraUserParameterBindingCustomization::MakeInstance)
-	);
+	PropertyModule.RegisterCustomPropertyTypeLayout(
+		FNiagaraUserParameterBinding::StaticStruct()->GetFName(),
+		FOnGetPropertyTypeCustomizationInstance::CreateStatic(&FNiagaraUserParameterBindingCustomization::MakeInstance));
+
+	PropertyModule.RegisterCustomPropertyTypeLayout(
+		FNiagaraScriptHighlight::StaticStruct()->GetFName(),
+		FOnGetPropertyTypeCustomizationInstance::CreateStatic(&FNiagaraScriptHighlightDetails::MakeInstance));
 
 	FNiagaraEditorStyle::Initialize();
 	ReinitializeStyleCommand = IConsoleManager::Get().RegisterConsoleCommand(
diff --git a/Engine/Plugins/FX/Niagara/Source/NiagaraEditor/Private/NiagaraEditorUtilities.cpp b/Engine/Plugins/FX/Niagara/Source/NiagaraEditor/Private/NiagaraEditorUtilities.cpp
--- a/Engine/Plugins/FX/Niagara/Source/NiagaraEditor/Private/NiagaraEditorUtilities.cpp
+++ b/Engine/Plugins/FX/Niagara/Source/NiagaraEditor/Private/NiagaraEditorUtilities.cpp
@@ -53,6 +53,9 @@
 #include "Framework/Commands/GenericCommands.h"
 #include "UObject/TextProperty.h"
 #include "Editor/EditorEngine.h"
+#include "Widgets/Notifications/SNotificationList.h"
+#include "Styling/CoreStyle.h"
+#include "Framework/Notifications/NotificationManager.h"
 
 #define LOCTEXT_NAMESPACE "FNiagaraEditorUtilities"
 
@@ -1400,11 +1403,42 @@ const FNiagaraEmitterHandle* FNiagaraEditorUtilities::GetEmitterHandleForEmitter
 		[&Emitter](const FNiagaraEmitterHandle& EmitterHandle) { return EmitterHandle.GetInstance() == &Emitter; });
 }
 
-FText FNiagaraEditorUtilities::FormatScriptAssetDescription(FText Description, FName Path)
+bool FNiagaraEditorUtilities::IsScriptAssetInLibrary(const FAssetData& ScriptAssetData)
 {
+	bool bIsInLibrary;
+	bool bIsLibraryTagFound = ScriptAssetData.GetTagValue(GET_MEMBER_NAME_CHECKED(UNiagaraScript, bExposeToLibrary), bIsInLibrary);
+	if (bIsLibraryTagFound == false)
+	{
+		if (ScriptAssetData.IsAssetLoaded())
+		{
+			UNiagaraScript* Script = static_cast<UNiagaraScript*>(ScriptAssetData.GetAsset());
+			if (Script != nullptr)
+			{
+				bIsInLibrary = Script->bExposeToLibrary;
+			}
+		}
+		else
+		{
+			bIsInLibrary = false;
+		}
+	}
+	return bIsInLibrary;
+}
+
+NIAGARAEDITOR_API FText FNiagaraEditorUtilities::FormatScriptName(FName Name, bool bIsInLibrary)
+{
+	return FText::FromString(FName::NameToDisplayString(Name.ToString(), false) + (bIsInLibrary ? TEXT("") : TEXT("*")));
+}
+
+FText FNiagaraEditorUtilities::FormatScriptDescription(FText Description, FName Path, bool bIsInLibrary)
+{
+	FText LibrarySuffix = bIsInLibrary
+		? LOCTEXT("LibrarySuffix", "\n* Script is not exposed to the library.")
+		: FText();
+
 	return Description.IsEmptyOrWhitespace()
-		? FText::Format(LOCTEXT("ScriptAssetDescriptionFormatPathOnly", "Path: {0}"), FText::FromName(Path))
-		: FText::Format(LOCTEXT("ScriptAssetDescriptionFormat", "Description: {1}\nPath: {0}"), FText::FromName(Path), Description);
+		? FText::Format(LOCTEXT("ScriptAssetDescriptionFormatPathOnly", "Path: {0}{1}"), FText::FromName(Path), LibrarySuffix)
+		: FText::Format(LOCTEXT("ScriptAssetDescriptionFormat", "Description: {1}\nPath: {0}{2}"), FText::FromName(Path), Description, LibrarySuffix);
 }
 
 void FNiagaraEditorUtilities::ResetSystemsThatReferenceSystemViewModel(const FNiagaraSystemViewModel& ReferencedSystemViewModel)
@@ -1922,4 +1956,15 @@ bool FNiagaraEditorUtilities::AddEmitterContextMenuActions(FMenuBuilder& MenuBui
 	return false;
 }
 
+void FNiagaraEditorUtilities::WarnWithToastAndLog(FText WarningMessage)
+{
+	FNotificationInfo WarningNotification(WarningMessage);
+	WarningNotification.ExpireDuration = 5.0f;
+	WarningNotification.bFireAndForget = true;
+	WarningNotification.bUseLargeFont = false;
+	WarningNotification.Image = FCoreStyle::Get().GetBrush(TEXT("MessageLog.Warning"));
+	FSlateNotificationManager::Get().AddNotification(WarningNotification);
+	UE_LOG(LogNiagaraEditor, Warning, TEXT("%s"), *WarningMessage.ToString());
+}
+
 #undef LOCTEXT_NAMESPACE
diff --git a/Engine/Plugins/FX/Niagara/Source/NiagaraEditor/Private/NiagaraScriptMergeManager.cpp b/Engine/Plugins/FX/Niagara/Source/NiagaraEditor/Private/NiagaraScriptMergeManager.cpp
--- a/Engine/Plugins/FX/Niagara/Source/NiagaraEditor/Private/NiagaraScriptMergeManager.cpp
+++ b/Engine/Plugins/FX/Niagara/Source/NiagaraEditor/Private/NiagaraScriptMergeManager.cpp
@@ -1186,6 +1186,13 @@ INiagaraMergeManager::FMergeEmitterResults FNiagaraScriptMergeManager::MergeEmit
 		TMap<FGuid, FGuid> FinalChangeIds;
 		FNiagaraEditorUtilities::GatherChangeIds(*MergedInstance, FinalChangeIds, TEXT("Final"));
 	}
+
+	if(MergeResults.MergedInstance != nullptr)
+	{
+		MergeResults.MergedInstance->ParentScratchPadScripts.Append(MergeResults.MergedInstance->ScratchPadScripts);
+		MergeResults.MergedInstance->ScratchPadScripts.Empty();
+	}
+
 	return MergeResults;
 }
 
@@ -1725,8 +1732,13 @@ void FNiagaraScriptMergeManager::DiffScriptStacks(TSharedRef<FNiagaraScriptStack
 			DiffResults.EnabledChangedOtherModules.Add(CommonValuePair.OtherValue);
 		}
 
-		if (CommonValuePair.BaseValue->GetFunctionCallNode()->FunctionScript == CommonValuePair.OtherValue->GetFunctionCallNode()->FunctionScript ||
-			CommonValuePair.BaseValue->GetFunctionCallNode()->IsA<UNiagaraNodeAssignment>())
+		UNiagaraScript* BaseFunctionScript = CommonValuePair.BaseValue->GetFunctionCallNode()->FunctionScript;
+		UNiagaraScript* OtherFunctionScript = CommonValuePair.OtherValue->GetFunctionCallNode()->FunctionScript;
+		bool bFunctionScriptsMatch = BaseFunctionScript == OtherFunctionScript;
+		bool bFunctionScriptsAreNotAssets =
+			BaseFunctionScript != nullptr && BaseFunctionScript->IsAsset() == false &&
+			OtherFunctionScript != nullptr && OtherFunctionScript->IsAsset() == false;
+		if (bFunctionScriptsMatch || bFunctionScriptsAreNotAssets)
 		{
 			DiffFunctionInputs(CommonValuePair.BaseValue, CommonValuePair.OtherValue, DiffResults);
 		}
diff --git a/Engine/Plugins/FX/Niagara/Source/NiagaraEditor/Private/Toolkits/NiagaraScriptToolkit.cpp b/Engine/Plugins/FX/Niagara/Source/NiagaraEditor/Private/Toolkits/NiagaraScriptToolkit.cpp
--- a/Engine/Plugins/FX/Niagara/Source/NiagaraEditor/Private/Toolkits/NiagaraScriptToolkit.cpp
+++ b/Engine/Plugins/FX/Niagara/Source/NiagaraEditor/Private/Toolkits/NiagaraScriptToolkit.cpp
@@ -341,13 +341,6 @@ TSharedRef<SDockTab> FNiagaraScriptToolkit::SpawnTabScriptDetails(const FSpawnTa
 	TSharedRef<IDetailsView> DetailsView = PropertyEditorModule.CreateDetailView(DetailsViewArgs);
 
 	DetailsView->OnFinishedChangingProperties().AddRaw(this, &FNiagaraScriptToolkit::OnEditedScriptPropertyFinishedChanging);
-	DetailsView->RegisterInstancedCustomPropertyLayout(
-		UNiagaraScript::StaticClass(),
-		FOnGetDetailCustomizationInstance::CreateStatic(&FNiagaraScriptDetails::MakeInstance, ScriptViewModelWeakPtr));
-	DetailsView->RegisterInstancedCustomPropertyTypeLayout(
-		FNiagaraScriptHighlight::StaticStruct()->GetFName(),
-		FOnGetPropertyTypeCustomizationInstance::CreateStatic(&FNiagaraScriptHighlightDetails::MakeInstance));
-	
 	DetailsView->SetObjects(DetailsScriptSelection->GetSelectedObjects().Array());
 
 	return SNew(SDockTab)
diff --git a/Engine/Plugins/FX/Niagara/Source/NiagaraEditor/Private/Toolkits/NiagaraSystemToolkit.cpp b/Engine/Plugins/FX/Niagara/Source/NiagaraEditor/Private/Toolkits/NiagaraSystemToolkit.cpp
--- a/Engine/Plugins/FX/Niagara/Source/NiagaraEditor/Private/Toolkits/NiagaraSystemToolkit.cpp
+++ b/Engine/Plugins/FX/Niagara/Source/NiagaraEditor/Private/Toolkits/NiagaraSystemToolkit.cpp
@@ -71,6 +71,7 @@ const FName FNiagaraSystemToolkit::PreviewSettingsTabId(TEXT("NiagaraSystemEdito
 const FName FNiagaraSystemToolkit::GeneratedCodeTabID(TEXT("NiagaraSystemEditor_GeneratedCode"));
 const FName FNiagaraSystemToolkit::MessageLogTabID(TEXT("NiagaraSystemEditor_MessageLog"));
 const FName FNiagaraSystemToolkit::SystemOverviewTabID(TEXT("NiagaraSystemEditor_SystemOverview"));
+const FName FNiagaraSystemToolkit::ScratchPadTabID(TEXT("NiagaraSystemEditor_ScratchPad"));
 
 static int32 GbLogNiagaraSystemChanges = 0;
 static FAutoConsoleVariableRef CVarSuppressNiagaraSystems(
@@ -148,6 +149,10 @@ void FNiagaraSystemToolkit::RegisterTabSpawners(const TSharedRef<class FTabManag
 	InTabManager->RegisterTabSpawner(SystemOverviewTabID, FOnSpawnTab::CreateSP(this, &FNiagaraSystemToolkit::SpawnTab_SystemOverview))
 		.SetDisplayName(LOCTEXT("SystemOverviewTabName", "System Overview"))
 		.SetGroup(WorkspaceMenuCategory.ToSharedRef());
+
+	InTabManager->RegisterTabSpawner(ScratchPadTabID, FOnSpawnTab::CreateSP(this, &FNiagaraSystemToolkit::SpawnTab_ScratchPad))
+		.SetDisplayName(LOCTEXT("ScratchPadTabName", "Scratch Pad"))
+		.SetGroup(WorkspaceMenuCategory.ToSharedRef());
 }
 
 void FNiagaraSystemToolkit::UnregisterTabSpawners(const TSharedRef<class FTabManager>& InTabManager)
@@ -166,6 +171,7 @@ void FNiagaraSystemToolkit::UnregisterTabSpawners(const TSharedRef<class FTabMan
 	InTabManager->UnregisterTabSpawner(PreviewSettingsTabId);
 	InTabManager->UnregisterTabSpawner(GeneratedCodeTabID);
 	InTabManager->UnregisterTabSpawner(SystemOverviewTabID);
+	InTabManager->UnregisterTabSpawner(ScratchPadTabID);
 }
 
 
@@ -305,11 +311,12 @@ void FNiagaraSystemToolkit::InitializeInternal(const EToolkitMode::Type Mode, co
 	SystemViewModel->GetSelectionViewModel()->OnEmitterHandleIdSelectionChanged().AddSP(this, &FNiagaraSystemToolkit::OnSystemSelectionChanged);
 	SystemViewModel->GetOnPinnedEmittersChanged().AddSP(this, &FNiagaraSystemToolkit::RefreshParameters);
 	SystemViewModel->GetOnPinnedCurvesChanged().AddSP(this, &FNiagaraSystemToolkit::OnPinnedCurvesChanged);
+	SystemViewModel->OnRequestFocusTab().AddSP(this, &FNiagaraSystemToolkit::OnViewModelRequestFocusTab);
 
 	const float InTime = -0.02f;
 	const float OutTime = 3.2f;
 
-	TSharedRef<FTabManager::FLayout> StandaloneDefaultLayout = FTabManager::NewLayout("Standalone_Niagara_System_Layout_v20")
+	TSharedRef<FTabManager::FLayout> StandaloneDefaultLayout = FTabManager::NewLayout("Standalone_Niagara_System_Layout_v22")
 		->AddArea
 		(
 			FTabManager::NewPrimaryArea()->SetOrientation(Orient_Vertical)
@@ -350,6 +357,8 @@ void FNiagaraSystemToolkit::InitializeInternal(const EToolkitMode::Type Mode, co
 							FTabManager::NewStack()
 							->SetSizeCoefficient(0.6f)
 							->AddTab(SystemOverviewTabID, ETabState::OpenedTab)
+							->AddTab(ScratchPadTabID, ETabState::OpenedTab)
+							->SetForegroundTab(SystemOverviewTabID)
 						)
 					)
 					->Split
@@ -639,11 +648,21 @@ TSharedRef<SDockTab> FNiagaraSystemToolkit::SpawnTab_MessageLog(const FSpawnTabA
 
 TSharedRef<SDockTab> FNiagaraSystemToolkit::SpawnTab_SystemOverview(const FSpawnTabArgs& Args)
 {
-	FNiagaraEditorModule& NiagaraEditorModule = FModuleManager::LoadModuleChecked<FNiagaraEditorModule>("NiagaraEditor");
 	TSharedRef<SDockTab> SpawnedTab = SNew(SDockTab)
 		.Label(LOCTEXT("SystemOverviewTabLabel", "System Overview"))
 		[
-			NiagaraEditorModule.GetWidgetProvider()->CreateSystemOverview(SystemViewModel.ToSharedRef())
+			FNiagaraEditorModule::Get().GetWidgetProvider()->CreateSystemOverview(SystemViewModel.ToSharedRef())
+		];
+
+	return SpawnedTab;
+}
+
+TSharedRef<SDockTab> FNiagaraSystemToolkit::SpawnTab_ScratchPad(const FSpawnTabArgs& Args)
+{
+	TSharedRef<SDockTab> SpawnedTab = SNew(SDockTab)
+		.Label(LOCTEXT("ScratchPadTabLabel", "Scratch Pad"))
+		[
+			FNiagaraEditorModule::Get().GetWidgetProvider()->CreateScriptScratchPad(*SystemViewModel->GetScriptScratchPadViewModel())
 		];
 
 	return SpawnedTab;
@@ -1424,4 +1443,9 @@ void FNiagaraSystemToolkit::OnSystemSelectionChanged()
 	RefreshParameters();
 }
 
+void FNiagaraSystemToolkit::OnViewModelRequestFocusTab(FName TabName)
+{
+	GetTabManager()->InvokeTab(TabName);
+}
+
 #undef LOCTEXT_NAMESPACE
diff --git a/Engine/Plugins/FX/Niagara/Source/NiagaraEditor/Private/Toolkits/NiagaraSystemToolkit.h b/Engine/Plugins/FX/Niagara/Source/NiagaraEditor/Private/Toolkits/NiagaraSystemToolkit.h
--- a/Engine/Plugins/FX/Niagara/Source/NiagaraEditor/Private/Toolkits/NiagaraSystemToolkit.h
+++ b/Engine/Plugins/FX/Niagara/Source/NiagaraEditor/Private/Toolkits/NiagaraSystemToolkit.h
@@ -107,6 +107,7 @@ private:
 	TSharedRef<SDockTab> SpawnTab_GeneratedCode(const FSpawnTabArgs& Args);
 	TSharedRef<SDockTab> SpawnTab_MessageLog(const FSpawnTabArgs& Args);
 	TSharedRef<SDockTab> SpawnTab_SystemOverview(const FSpawnTabArgs& Args);
+	TSharedRef<SDockTab> SpawnTab_ScratchPad(const FSpawnTabArgs& Args);
 
 	/** Builds the toolbar widget */
 	void ExtendToolbar();	
@@ -134,6 +135,7 @@ private:
 	void OnPinnedCurvesChanged();
 	void RefreshParameters();
 	void OnSystemSelectionChanged();
+	void OnViewModelRequestFocusTab(FName TabName);
 
 	TSharedRef<SWidget> GenerateBoundsMenuContent(TSharedRef<FUICommandList> InCommandList);
 	const FName GetNiagaraSystemMessageLogName(UNiagaraSystem* InSystem) const;
@@ -175,6 +177,7 @@ private:
 
 	bool bChangesDiscarded;
 
+public:
 	static const FName ViewportTabID;
 	static const FName CurveEditorTabID;
 	static const FName SequencerTabID;
@@ -188,7 +191,9 @@ private:
 	static const FName GeneratedCodeTabID;
 	static const FName MessageLogTabID;
 	static const FName SystemOverviewTabID;
+	static const FName ScratchPadTabID;
 
+private:
 	static bool bShowLibraryOnly;
 	static bool bShowTemplateOnly;
 };
diff --git a/Engine/Plugins/FX/Niagara/Source/NiagaraEditor/Private/ViewModels/NiagaraScratchPadScriptViewModel.cpp b/Engine/Plugins/FX/Niagara/Source/NiagaraEditor/Private/ViewModels/NiagaraScratchPadScriptViewModel.cpp
new file mode 100644
--- /dev/null
+++ b/Engine/Plugins/FX/Niagara/Source/NiagaraEditor/Private/ViewModels/NiagaraScratchPadScriptViewModel.cpp
@@ -0,0 +1,75 @@
+// Copyright Epic Games, Inc. All Rights Reserved.
+
+#include "ViewModels/NiagaraScratchPadScriptViewModel.h"
+#include "ViewModels/Stack/NiagaraStackGraphUtilities.h"
+#include "NiagaraNodeFunctionCall.h"
+#include "NiagaraSystem.h"
+#include "NiagaraEmitter.h"
+
+#include "ObjectTools.h"
+
+FNiagaraScratchPadScriptViewModel::FNiagaraScratchPadScriptViewModel()
+	: FNiagaraScriptViewModel(TAttribute<FText>::Create(TAttribute<FText>::FGetter::CreateRaw(this, &FNiagaraScratchPadScriptViewModel::GetDisplayNameInternal)), ENiagaraParameterEditMode::EditAll)
+	, bIsPendingRename(false)
+{
+}
+
+void FNiagaraScratchPadScriptViewModel::Initialize(UNiagaraScript* Script)
+{
+	SetScript(Script);
+}
+
+UNiagaraScript* FNiagaraScratchPadScriptViewModel::GetScript() const
+{
+	return GetScripts()[0].Get();
+}
+
+bool FNiagaraScratchPadScriptViewModel::GetIsPendingRename() const
+{
+	return bIsPendingRename;
+}
+
+void FNiagaraScratchPadScriptViewModel::SetIsPendingRename(bool bInIsPendingRename)
+{
+	bIsPendingRename = bInIsPendingRename;
+}
+
+void FNiagaraScratchPadScriptViewModel::SetScriptName(FText InScriptName)
+{
+	UNiagaraScript* Script = GetScript();
+	FString NewName = ObjectTools::SanitizeObjectName(InScriptName.ToString());
+	if (Script->GetName() != NewName)
+	{
+		FName NewUniqueName = MakeUniqueObjectName(Script->GetOuter(), UNiagaraScript::StaticClass(), *NewName);
+		Script->Rename(*NewUniqueName.ToString());
+
+		for (TObjectIterator<UNiagaraNodeFunctionCall> It; It; ++It)
+		{
+			UNiagaraNodeFunctionCall* FunctionCallNode = *It;
+			if (FunctionCallNode->FunctionScript == Script)
+			{
+				FString OldFunctionName = FunctionCallNode->GetFunctionName();
+				FunctionCallNode->SuggestName(FString());
+				const FString NewFunctionName = FunctionCallNode->GetFunctionName();
+				UNiagaraSystem* System = FunctionCallNode->GetTypedOuter<UNiagaraSystem>();
+				UNiagaraEmitter* Emitter = FunctionCallNode->GetTypedOuter<UNiagaraEmitter>();
+				if (System != nullptr)
+				{
+					FNiagaraStackGraphUtilities::RenameReferencingParameters(*System, Emitter, *FunctionCallNode, OldFunctionName, NewFunctionName);
+				}
+			}
+		}
+
+		OnRenamedDelegate.Broadcast();
+	}
+}
+
+FNiagaraScratchPadScriptViewModel::FOnRenamed& FNiagaraScratchPadScriptViewModel::OnRenamed()
+{
+	return OnRenamedDelegate;
+}
+
+FText FNiagaraScratchPadScriptViewModel::GetDisplayNameInternal() const
+{
+	return FText::FromString(GetScript()->GetName());
+}
\ No newline at end of file
diff --git a/Engine/Plugins/FX/Niagara/Source/NiagaraEditor/Private/ViewModels/NiagaraScratchPadViewModel.cpp b/Engine/Plugins/FX/Niagara/Source/NiagaraEditor/Private/ViewModels/NiagaraScratchPadViewModel.cpp
new file mode 100644
--- /dev/null
+++ b/Engine/Plugins/FX/Niagara/Source/NiagaraEditor/Private/ViewModels/NiagaraScratchPadViewModel.cpp
@@ -0,0 +1,308 @@
+// Copyright Epic Games, Inc. All Rights Reserved.
+
+#include "ViewModels/NiagaraScratchPadViewModel.h"
+#include "ViewModels/NiagaraSystemViewModel.h"
+#include "ViewModels/NiagaraScratchPadScriptViewModel.h"
+#include "ViewModels/NiagaraScriptGraphViewModel.h"
+#include "NiagaraSystem.h"
+#include "NiagaraScriptSource.h"
+#include "NiagaraNodeOutput.h"
+#include "NiagaraObjectSelection.h"
+#include "NiagaraEditorSettings.h"
+#include "NiagaraNodeFunctionCall.h"
+
+#include "ScopedTransaction.h"
+
+#define LOCTEXT_NAMESPACE "NiagaraScratchPadViewModel"
+
+void UNiagaraScratchPadViewModel::Initialize(TSharedRef<FNiagaraSystemViewModel> InSystemViewModel)
+{
+	SystemViewModelWeak = InSystemViewModel;
+	ObjectSelection = MakeShared<FNiagaraObjectSelection>();
+	RefreshScriptViewModels();
+	AvailableUsages = { ENiagaraScriptUsage::DynamicInput, ENiagaraScriptUsage::Module };
+}
+
+void UNiagaraScratchPadViewModel::Finalize()
+{
+}
+
+void GetOuterAndTargetScripts(TSharedRef<FNiagaraSystemViewModel> SystemViewModel, UObject*& OutOuter, TArray<UNiagaraScript*>*& OutTargetScripts)
+{
+	if (SystemViewModel->GetEditMode() == ENiagaraSystemViewModelEditMode::SystemAsset)
+	{
+		OutOuter = &SystemViewModel->GetSystem();
+		OutTargetScripts = &SystemViewModel->GetSystem().ScratchPadScripts;
+	}
+	else
+	{
+		if (SystemViewModel->GetSystem().GetEmitterHandles().Num() == 1)
+		{
+			UNiagaraEmitter* TargetEmitter = SystemViewModel->GetSystem().GetEmitterHandles()[0].GetInstance();
+			OutOuter = TargetEmitter;
+			OutTargetScripts = &TargetEmitter->ScratchPadScripts;
+		}
+		else
+		{
+			OutOuter = nullptr;
+			OutTargetScripts = nullptr;
+		}
+	}
+}
+
+void UNiagaraScratchPadViewModel::RefreshScriptViewModels()
+{
+	TArray<TSharedRef<FNiagaraScratchPadScriptViewModel>> OldScriptViewModels = ScriptViewModels;
+	ScriptViewModels.Empty();
+
+	bool bNewViewModels = false;
+
+	UObject* ScriptOuter;
+	TArray<UNiagaraScript*>* TargetScripts;
+	GetOuterAndTargetScripts(GetSystemViewModel(), ScriptOuter, TargetScripts);
+
+	if (ScriptOuter != nullptr && TargetScripts != nullptr)
+	{
+		for (UNiagaraScript* ScratchPadScript : (*TargetScripts))
+		{
+			TSharedPtr<FNiagaraScratchPadScriptViewModel> ScriptViewModel;
+
+			TSharedRef<FNiagaraScratchPadScriptViewModel>* OldScriptViewModel = OldScriptViewModels.FindByPredicate(
+				[ScratchPadScript](TSharedRef<FNiagaraScratchPadScriptViewModel> ScriptViewModel) { return ScriptViewModel->GetScript() == ScratchPadScript; });
+			if (OldScriptViewModel != nullptr)
+			{
+				ScriptViewModel = *OldScriptViewModel;
+				OldScriptViewModels.Remove(ScriptViewModel.ToSharedRef());
+			}
+			else
+			{
+				ScriptViewModel = MakeShared<FNiagaraScratchPadScriptViewModel>();
+				ScriptViewModel->Initialize(ScratchPadScript);
+				ScriptViewModel->GetGraphViewModel()->GetNodeSelection()->OnSelectedObjectsChanged().AddUObject(this, &UNiagaraScratchPadViewModel::ScriptGraphNodeSelectionChanged, TWeakPtr<FNiagaraScratchPadScriptViewModel>(ScriptViewModel));
+				ScriptViewModel->OnRenamed().AddUObject(this, &UNiagaraScratchPadViewModel::ScriptViewModelScriptRenamed);
+				bNewViewModels = true;
+			}
+			ScriptViewModels.Add(ScriptViewModel.ToSharedRef());
+		}
+	}
+
+	for (TSharedRef<FNiagaraScratchPadScriptViewModel> OldScriptViewModel : OldScriptViewModels)
+	{
+		OldScriptViewModel->OnRenamed().RemoveAll(this);
+	}
+
+	if (bNewViewModels || OldScriptViewModels.Num() > 0)
+	{
+		OnScriptViewModelsChangedDelegate.Broadcast();
+	}
+}
+
+const TArray<TSharedRef<FNiagaraScratchPadScriptViewModel>>& UNiagaraScratchPadViewModel::GetScriptViewModels() const
+{
+	return ScriptViewModels;
+}
+
+TSharedPtr<FNiagaraScratchPadScriptViewModel> UNiagaraScratchPadViewModel::GetViewModelForScript(UNiagaraScript* InScript)
+{
+	TSharedRef<FNiagaraScratchPadScriptViewModel>* ViewModelForScript = ScriptViewModels.FindByPredicate([InScript](TSharedRef<FNiagaraScratchPadScriptViewModel>& ScriptViewModel) { return ScriptViewModel->GetScript() == InScript; });
+	if (ViewModelForScript != nullptr)
+	{
+		return *ViewModelForScript;
+	}
+	return TSharedPtr<FNiagaraScratchPadScriptViewModel>();
+}
+
+const TArray<ENiagaraScriptUsage>& UNiagaraScratchPadViewModel::GetAvailableUsages() const
+{
+	return AvailableUsages;
+}
+
+FText UNiagaraScratchPadViewModel::GetDisplayNameForUsage(ENiagaraScriptUsage InUsage) const
+{
+	switch (InUsage)
+	{
+	case ENiagaraScriptUsage::DynamicInput:
+		return LOCTEXT("DynamicInputDisplayName", "Dynamic Inputs");
+	case ENiagaraScriptUsage::Module:
+		return LOCTEXT("ModuleDisplayName", "Modules");
+	case ENiagaraScriptUsage::Function:
+		return LOCTEXT("FunctionDisplayName", "Functions");
+	default: 
+		return LOCTEXT("InvalidUsageDisplayName", "Invalid");
+	}
+}
+
+TSharedRef<FNiagaraObjectSelection> UNiagaraScratchPadViewModel::GetObjectSelection()
+{
+	return ObjectSelection.ToSharedRef();
+}
+
+UNiagaraScript* UNiagaraScratchPadViewModel::GetActiveScript()
+{
+	return ActiveScript;
+}
+
+void UNiagaraScratchPadViewModel::SetActiveScript(UNiagaraScript* InActiveScript)
+{
+	ActiveScript = InActiveScript;
+	if (ActiveScript != nullptr)
+	{
+		ObjectSelection->SetSelectedObject(ActiveScript);
+	}
+	else
+	{
+		ObjectSelection->ClearSelectedObjects();
+	}
+	OnActiveScriptChangedDelegate.Broadcast();
+}
+
+void UNiagaraScratchPadViewModel::DeleteActiveScript()
+{
+	if (ActiveScript != nullptr)
+	{
+		FScopedTransaction DeleteTransaction(LOCTEXT("DeleteScratchPadScriptTransaction", "Delete scratch pad script."));
+		for (TObjectIterator<UNiagaraNodeFunctionCall> It; It; ++It)
+		{
+			UNiagaraNodeFunctionCall* FunctionCallNode = *It;
+			if (FunctionCallNode->FunctionScript == ActiveScript)
+			{
+				FunctionCallNode->Modify();
+				FunctionCallNode->FunctionScript = nullptr;
+			}
+		}
+
+		UObject* ScriptOuter;
+		TArray<UNiagaraScript*>* TargetScripts;
+		GetOuterAndTargetScripts(GetSystemViewModel(), ScriptOuter, TargetScripts);
+
+		ScriptOuter->Modify();
+		TargetScripts->Remove(ActiveScript);
+
+		OnScriptDeletedDelegate.Broadcast();
+		RefreshScriptViewModels();
+	}
+}
+
+TSharedPtr<FNiagaraScratchPadScriptViewModel> UNiagaraScratchPadViewModel::CreateNewScript(ENiagaraScriptUsage InScriptUsage, ENiagaraScriptUsage InTargetSupportedUsage, FNiagaraTypeDefinition InOutputType)
+{
+	UObject* ScriptOuter;
+	TArray<UNiagaraScript*>* TargetScripts;
+	GetOuterAndTargetScripts(GetSystemViewModel(), ScriptOuter, TargetScripts);
+
+	UNiagaraScript* NewScript = nullptr;
+	switch (InScriptUsage)
+	{
+	case ENiagaraScriptUsage::DynamicInput:
+	{
+		UNiagaraScript* DefaultDynamicInput = Cast<UNiagaraScript>(GetDefault<UNiagaraEditorSettings>()->DefaultDynamicInputScript.TryLoad());
+		if (DefaultDynamicInput != nullptr)
+		{
+			NewScript = CastChecked<UNiagaraScript>(StaticDuplicateObject(DefaultDynamicInput, ScriptOuter, MakeUniqueObjectName(ScriptOuter, UNiagaraScript::StaticClass(), TEXT("NewScratchDynamicInput"))));
+			TArray<UNiagaraNodeOutput*> OutputNodes;
+			CastChecked<UNiagaraScriptSource>(NewScript->GetSource())->NodeGraph->GetNodesOfClass(OutputNodes);
+			if (OutputNodes.Num() == 1)
+			{
+				if (InOutputType.IsValid())
+				{
+					UNiagaraNodeOutput* DynamicInputOutputNode = OutputNodes[0];
+					DynamicInputOutputNode->Outputs.Empty();
+					DynamicInputOutputNode->Outputs.Add(FNiagaraVariable(InOutputType, "Output"));
+					FPropertyChangedEvent OutputsChangedEvent(FindFieldChecked<UProperty>(UNiagaraNodeOutput::StaticClass(), GET_MEMBER_NAME_CHECKED(UNiagaraNodeOutput, Outputs)));
+					DynamicInputOutputNode->PostEditChangeProperty(OutputsChangedEvent);
+				}
+			}
+		}
+		break;
+	}
+	case ENiagaraScriptUsage::Module:
+	{
+		UNiagaraScript* DefaultModule = Cast<UNiagaraScript>(GetDefault<UNiagaraEditorSettings>()->DefaultModuleScript.TryLoad());
+		if (DefaultModule != nullptr)
+		{
+			NewScript = CastChecked<UNiagaraScript>(StaticDuplicateObject(DefaultModule, ScriptOuter, MakeUniqueObjectName(ScriptOuter, UNiagaraScript::StaticClass(), TEXT("NewScratchModule"))));
+		}
+		break;
+	}
+	}
+
+	if (NewScript != nullptr)
+	{
+		NewScript->ClearFlags(RF_Public | RF_Standalone);
+		ScriptOuter->Modify();
+		TargetScripts->Add(NewScript);
+		NewScript->ModuleUsageBitmask |= (1 << (int32)InTargetSupportedUsage);
+		RefreshScriptViewModels();
+	}
+
+	return GetViewModelForScript(NewScript);
+}
+
+TSharedPtr<FNiagaraScratchPadScriptViewModel> UNiagaraScratchPadViewModel::CreateNewScriptAsDuplicate(UNiagaraScript* ScriptToDuplicate)
+{
+	UObject* ScriptOuter;
+	TArray<UNiagaraScript*>* TargetScripts;
+	GetOuterAndTargetScripts(GetSystemViewModel(), ScriptOuter, TargetScripts);
+
+	UNiagaraScript* NewScript = NewScript = CastChecked<UNiagaraScript>(StaticDuplicateObject(ScriptToDuplicate, ScriptOuter, MakeUniqueObjectName(ScriptOuter, UNiagaraScript::StaticClass(), ScriptToDuplicate->GetFName())));
+	NewScript->ClearFlags(RF_Public | RF_Standalone);
+	ScriptOuter->Modify();
+	TargetScripts->Add(NewScript);
+	RefreshScriptViewModels();
+
+	return GetViewModelForScript(NewScript);
+}
+
+UNiagaraScratchPadViewModel::FOnScriptViewModelsChanged& UNiagaraScratchPadViewModel::OnScriptViewModelsChanged()
+{
+	return OnScriptViewModelsChangedDelegate;
+}
+
+UNiagaraScratchPadViewModel::FOnActiveScriptChanged& UNiagaraScratchPadViewModel::OnActiveScriptChanged()
+{
+	return OnActiveScriptChangedDelegate;
+}
+
+UNiagaraScratchPadViewModel::FOnScriptRenamed& UNiagaraScratchPadViewModel::OnScriptRenamed()
+{
+	return OnScriptRenamedDelegate;
+}
+
+UNiagaraScratchPadViewModel::FOnScriptDeleted& UNiagaraScratchPadViewModel::OnScriptDeleted()
+{
+	return OnScriptDeletedDelegate;
+}
+
+TSharedRef<FNiagaraSystemViewModel> UNiagaraScratchPadViewModel::GetSystemViewModel()
+{
+	TSharedPtr<FNiagaraSystemViewModel> SystemViewModel = SystemViewModelWeak.Pin();
+	checkf(SystemViewModel.IsValid(), TEXT("SystemViewModel destroyed before scratch pad view model."));
+	return SystemViewModel.ToSharedRef();
+}
+
+void UNiagaraScratchPadViewModel::ScriptGraphNodeSelectionChanged(TWeakPtr<FNiagaraScratchPadScriptViewModel> InScriptViewModelWeak)
+{
+	TSharedPtr<FNiagaraScratchPadScriptViewModel> InScriptViewModel = InScriptViewModelWeak.Pin();
+	if (InScriptViewModel.IsValid())
+	{
+		TArray<UObject*> SelectedNodes = InScriptViewModel->GetGraphViewModel()->GetNodeSelection()->GetSelectedObjects().Array();
+		if (SelectedNodes.Num() > 0)
+		{
+			ObjectSelection->SetSelectedObjects(SelectedNodes);
+		}
+		else if (ActiveScript != nullptr)
+		{
+			ObjectSelection->SetSelectedObject(ActiveScript);
+		}
+		else
+		{
+			ObjectSelection->ClearSelectedObjects();
+		}
+	}
+}
+
+void UNiagaraScratchPadViewModel::ScriptViewModelScriptRenamed()
+{
+	OnScriptRenamed().Broadcast();
+}
+
+#undef LOCTEXT_NAMESPACE
\ No newline at end of file
diff --git a/Engine/Plugins/FX/Niagara/Source/NiagaraEditor/Private/ViewModels/NiagaraScriptGraphViewModel.cpp b/Engine/Plugins/FX/Niagara/Source/NiagaraEditor/Private/ViewModels/NiagaraScriptGraphViewModel.cpp
--- a/Engine/Plugins/FX/Niagara/Source/NiagaraEditor/Private/ViewModels/NiagaraScriptGraphViewModel.cpp
+++ b/Engine/Plugins/FX/Niagara/Source/NiagaraEditor/Private/ViewModels/NiagaraScriptGraphViewModel.cpp
@@ -17,7 +17,7 @@
 
 #define LOCTEXT_NAMESPACE "NiagaraScriptGraphViewModel"
 
-FNiagaraScriptGraphViewModel::FNiagaraScriptGraphViewModel(FText InDisplayName)
+FNiagaraScriptGraphViewModel::FNiagaraScriptGraphViewModel(TAttribute<FText> InDisplayName)
 	: DisplayName(InDisplayName)
 	, Commands(MakeShareable(new FUICommandList()))
 	, NodeSelection(MakeShareable(new FNiagaraObjectSelection()))
@@ -41,7 +41,7 @@ void FNiagaraScriptGraphViewModel::SetScriptSource(UNiagaraScriptSource* InScrip
 
 FText FNiagaraScriptGraphViewModel::GetDisplayName() const
 {
-	return DisplayName;
+	return DisplayName.Get();
 }
 
 UNiagaraScriptSource* FNiagaraScriptGraphViewModel::GetScriptSource()
diff --git a/Engine/Plugins/FX/Niagara/Source/NiagaraEditor/Private/ViewModels/NiagaraScriptGraphViewModel.h b/Engine/Plugins/FX/Niagara/Source/NiagaraEditor/Private/ViewModels/NiagaraScriptGraphViewModel.h
--- a/Engine/Plugins/FX/Niagara/Source/NiagaraEditor/Private/ViewModels/NiagaraScriptGraphViewModel.h
+++ b/Engine/Plugins/FX/Niagara/Source/NiagaraEditor/Private/ViewModels/NiagaraScriptGraphViewModel.h
@@ -23,7 +23,7 @@ public:
 
 public:
 	/** Create a new view model with the supplied  display name. */
-	FNiagaraScriptGraphViewModel(FText InDisplayName);
+	FNiagaraScriptGraphViewModel(TAttribute<FText> InDisplayName);
 
 	~FNiagaraScriptGraphViewModel();
 
@@ -89,7 +89,7 @@ private:
 	TWeakObjectPtr<UNiagaraScriptSource> ScriptSource;
 
 	/** The display name for the script graph. */
-	FText DisplayName;
+	TAttribute<FText> DisplayName;
 
 	/** Commands for editing the graph. */
 	TSharedRef<FUICommandList> Commands;
diff --git a/Engine/Plugins/FX/Niagara/Source/NiagaraEditor/Private/ViewModels/NiagaraScriptInputCollectionViewModel.cpp b/Engine/Plugins/FX/Niagara/Source/NiagaraEditor/Private/ViewModels/NiagaraScriptInputCollectionViewModel.cpp
--- a/Engine/Plugins/FX/Niagara/Source/NiagaraEditor/Private/ViewModels/NiagaraScriptInputCollectionViewModel.cpp
+++ b/Engine/Plugins/FX/Niagara/Source/NiagaraEditor/Private/ViewModels/NiagaraScriptInputCollectionViewModel.cpp
@@ -18,9 +18,9 @@
 
 FText DisplayNameFormat = NSLOCTEXT("ScriptInputCollection", "DisplayNameFormat", "{0} Inputs");
 
-FNiagaraScriptInputCollectionViewModel::FNiagaraScriptInputCollectionViewModel(FText InDisplayName, ENiagaraParameterEditMode InParameterEditMode)
+FNiagaraScriptInputCollectionViewModel::FNiagaraScriptInputCollectionViewModel(TAttribute<FText> InDisplayName, ENiagaraParameterEditMode InParameterEditMode)
 	: FNiagaraParameterCollectionViewModel(InParameterEditMode)
-	, DisplayName(FText::Format(DisplayNameFormat, InDisplayName))
+	, DisplayName(InDisplayName)
 	, bCanHaveNumericParameters(true)
 {
 }
@@ -82,7 +82,7 @@ void FNiagaraScriptInputCollectionViewModel::SetScripts(TArray<UNiagaraScript*>
 
 FText FNiagaraScriptInputCollectionViewModel::GetDisplayName() const 
 {
-	return DisplayName;
+	return FText::Format(DisplayNameFormat, DisplayName.Get());
 }
 
 FVector2D GetNewNodeLocation(UNiagaraGraph* Graph, UNiagaraNode* NewInputNode, float VerticalNodeOffset, float HorizontalNodeOffset)
diff --git a/Engine/Plugins/FX/Niagara/Source/NiagaraEditor/Private/ViewModels/NiagaraScriptInputCollectionViewModel.h b/Engine/Plugins/FX/Niagara/Source/NiagaraEditor/Private/ViewModels/NiagaraScriptInputCollectionViewModel.h
--- a/Engine/Plugins/FX/Niagara/Source/NiagaraEditor/Private/ViewModels/NiagaraScriptInputCollectionViewModel.h
+++ b/Engine/Plugins/FX/Niagara/Source/NiagaraEditor/Private/ViewModels/NiagaraScriptInputCollectionViewModel.h
@@ -16,7 +16,7 @@ class UNiagaraEmitter;
 class FNiagaraScriptInputCollectionViewModel : public FNiagaraParameterCollectionViewModel, public TSharedFromThis<FNiagaraScriptInputCollectionViewModel>
 {
 public:
-	FNiagaraScriptInputCollectionViewModel(FText InDisplayName, ENiagaraParameterEditMode InParameterEditMode);
+	FNiagaraScriptInputCollectionViewModel(TAttribute<FText> InDisplayName, ENiagaraParameterEditMode InParameterEditMode);
 
 	~FNiagaraScriptInputCollectionViewModel();
 
@@ -69,7 +69,7 @@ private:
 	TWeakObjectPtr<UNiagaraGraph> Graph;
 
 	/** The display name for the view model. */
-	FText DisplayName;
+	TAttribute<FText> DisplayName;
 
 	/** The handle to the graph changed delegate. */
 	FDelegateHandle OnGraphChangedHandle;
diff --git a/Engine/Plugins/FX/Niagara/Source/NiagaraEditor/Private/ViewModels/NiagaraScriptViewModel.cpp b/Engine/Plugins/FX/Niagara/Source/NiagaraEditor/Private/ViewModels/NiagaraScriptViewModel.cpp
--- a/Engine/Plugins/FX/Niagara/Source/NiagaraEditor/Private/ViewModels/NiagaraScriptViewModel.cpp
+++ b/Engine/Plugins/FX/Niagara/Source/NiagaraEditor/Private/ViewModels/NiagaraScriptViewModel.cpp
@@ -19,7 +19,7 @@
 
 template<> TMap<UNiagaraScript*, TArray<FNiagaraScriptViewModel*>> TNiagaraViewModelManager<UNiagaraScript, FNiagaraScriptViewModel>::ObjectsToViewModels{};
 
-FNiagaraScriptViewModel::FNiagaraScriptViewModel(FText DisplayName, ENiagaraParameterEditMode InParameterEditMode)
+FNiagaraScriptViewModel::FNiagaraScriptViewModel(TAttribute<FText> DisplayName, ENiagaraParameterEditMode InParameterEditMode)
 	: InputCollectionViewModel(MakeShareable(new FNiagaraScriptInputCollectionViewModel(DisplayName, InParameterEditMode)))
 	, OutputCollectionViewModel(MakeShareable(new FNiagaraScriptOutputCollectionViewModel(InParameterEditMode)))
 	, GraphViewModel(MakeShareable(new FNiagaraScriptGraphViewModel(DisplayName)))
@@ -120,6 +120,15 @@ FNiagaraScriptViewModel::~FNiagaraScriptViewModel()
 
 }
 
+FText FNiagaraScriptViewModel::GetDisplayName() const
+{
+	return GraphViewModel->GetDisplayName();
+}
+
+const TArray<TWeakObjectPtr<UNiagaraScript>>& FNiagaraScriptViewModel::GetScripts() const
+{
+	return Scripts;
+}
 
 void FNiagaraScriptViewModel::SetScripts(UNiagaraScriptSource* InScriptSource, TArray<UNiagaraScript*>& InScripts)
 {
@@ -289,7 +298,6 @@ UNiagaraScript* FNiagaraScriptViewModel::GetStandaloneScript()
 	return nullptr;
 }
 
-
 void FNiagaraScriptViewModel::UpdateCompileStatus(ENiagaraScriptCompileStatus InAggregateCompileStatus, const FString& InAggregateCompileErrorString,
 	const TArray<ENiagaraScriptCompileStatus>& InCompileStatuses, const TArray<FString>& InCompileErrors, const TArray<FString>& InCompilePaths,
 	const TArray<UNiagaraScript*>& InCompileSources)
diff --git a/Engine/Plugins/FX/Niagara/Source/NiagaraEditor/Private/ViewModels/NiagaraSystemViewModel.cpp b/Engine/Plugins/FX/Niagara/Source/NiagaraEditor/Private/ViewModels/NiagaraSystemViewModel.cpp
--- a/Engine/Plugins/FX/Niagara/Source/NiagaraEditor/Private/ViewModels/NiagaraSystemViewModel.cpp
+++ b/Engine/Plugins/FX/Niagara/Source/NiagaraEditor/Private/ViewModels/NiagaraSystemViewModel.cpp
@@ -10,6 +10,7 @@
 #include "ViewModels/NiagaraEmitterHandleViewModel.h"
 #include "ViewModels/NiagaraEmitterViewModel.h"
 #include "ViewModels/NiagaraSystemSelectionViewModel.h"
+#include "ViewModels/NiagaraScratchPadViewModel.h"
 #include "ViewModels/Stack/NiagaraStackViewModel.h"
 #include "ViewModels/Stack/NiagaraStackGraphUtilities.h"
 #include "ViewModels/Stack/NiagaraStackEntry.h"
@@ -71,6 +72,7 @@ FNiagaraSystemViewModel::FNiagaraSystemViewModel()
 	, bCompilePendingCompletion(false)
 	, SystemStackViewModel(nullptr)
 	, SelectionViewModel(nullptr)
+	, ScriptScratchPadViewModel(nullptr)
 {
 	GEditor->RegisterForUndo(this);
 }
@@ -101,6 +103,11 @@ void FNiagaraSystemViewModel::Initialize(UNiagaraSystem& InSystem, FNiagaraSyste
 	SystemStackViewModel->InitializeWithViewModels(this->AsShared(), TSharedPtr<FNiagaraEmitterHandleViewModel>(), FNiagaraStackViewModelOptions(true, false));
 	SystemStackViewModel->OnStructureChanged().AddSP(this, &FNiagaraSystemViewModel::StackViewModelStructureChanged);
 
+	ScriptScratchPadViewModel = NewObject<UNiagaraScratchPadViewModel>(GetTransientPackage());
+	ScriptScratchPadViewModel->Initialize(this->AsShared());
+	ScriptScratchPadViewModel->OnScriptRenamed().AddSP(this, &FNiagaraSystemViewModel::ScratchPadScriptsChanged);
+	ScriptScratchPadViewModel->OnScriptDeleted().AddSP(this, &FNiagaraSystemViewModel::ScratchPadScriptsChanged);
+
 	SetupPreviewComponentAndInstance();
 	SetupSequencer();
 	RefreshAll();
@@ -178,6 +185,14 @@ void FNiagaraSystemViewModel::Cleanup()
 		SelectionViewModel->Finalize();
 		SelectionViewModel = nullptr;
 	}
+
+	if (ScriptScratchPadViewModel != nullptr)
+	{
+		ScriptScratchPadViewModel->OnScriptRenamed().RemoveAll(this);
+		ScriptScratchPadViewModel->OnScriptDeleted().RemoveAll(this);
+		ScriptScratchPadViewModel->Finalize();
+		ScriptScratchPadViewModel = nullptr;
+	}
 }
 
 FNiagaraSystemViewModel::~FNiagaraSystemViewModel()
@@ -478,6 +493,10 @@ void FNiagaraSystemViewModel::AddReferencedObjects(FReferenceCollector& Collecto
 	{
 		Collector.AddReferencedObject(SelectionViewModel);
 	}
+	if (ScriptScratchPadViewModel != nullptr)
+	{
+		Collector.AddReferencedObject(ScriptScratchPadViewModel);
+	}
 }
 
 void FNiagaraSystemViewModel::PostUndo(bool bSuccess)
@@ -580,6 +599,16 @@ FNiagaraSystemViewModel::FOnPreClose& FNiagaraSystemViewModel::OnPreClose()
 	return OnPreCloseDelegate;
 }
 
+FNiagaraSystemViewModel::FOnRequestFocusTab& FNiagaraSystemViewModel::OnRequestFocusTab()
+{
+	return OnRequestFocusTabDelegate;
+}
+
+void FNiagaraSystemViewModel::FocusTab(FName TabName)
+{
+	OnRequestFocusTabDelegate.Broadcast(TabName);
+}
+
 TSharedPtr<FUICommandList> FNiagaraSystemViewModel::GetToolkitCommands()
 {
 	return ToolkitCommands.Pin();
@@ -642,6 +671,11 @@ UNiagaraSystemSelectionViewModel* FNiagaraSystemViewModel::GetSelectionViewModel
 	return SelectionViewModel;
 }
 
+UNiagaraScratchPadViewModel* FNiagaraSystemViewModel::GetScriptScratchPadViewModel()
+{
+	return ScriptScratchPadViewModel;
+}
+
 TStatId FNiagaraSystemViewModel::GetStatId() const
 {
 	RETURN_QUICK_DECLARE_CYCLE_STAT(FNiagaraSystemViewModel, STATGROUP_Tickables);
@@ -765,6 +799,7 @@ void FNiagaraSystemViewModel::RefreshAll()
 	RefreshEmitterHandleViewModels();
 	RefreshSequencerTracks();
 	ResetCurveData();
+	ScriptScratchPadViewModel->RefreshScriptViewModels();
 }
 
 void FNiagaraSystemViewModel::NotifyDataObjectChanged(UObject* ChangedObject)
@@ -1890,6 +1925,15 @@ void FNiagaraSystemViewModel::StackViewModelStructureChanged()
 	}
 }
 
+void FNiagaraSystemViewModel::ScratchPadScriptsChanged()
+{
+	SystemStackViewModel->GetRootEntry()->RefreshChildren();
+	for (TSharedRef<FNiagaraEmitterHandleViewModel> EmitterHandleViewModel : EmitterHandleViewModels)
+	{
+		EmitterHandleViewModel->GetEmitterStackViewModel()->GetRootEntry()->RefreshChildren();
+	}
+}
+
 TSharedPtr<FNiagaraOverviewGraphViewModel> FNiagaraSystemViewModel::GetOverviewGraphViewModel() const
 {
 	return OverviewGraphViewModel;
diff --git a/Engine/Plugins/FX/Niagara/Source/NiagaraEditor/Private/ViewModels/Stack/NiagaraStackFunctionInput.cpp b/Engine/Plugins/FX/Niagara/Source/NiagaraEditor/Private/ViewModels/Stack/NiagaraStackFunctionInput.cpp
--- a/Engine/Plugins/FX/Niagara/Source/NiagaraEditor/Private/ViewModels/Stack/NiagaraStackFunctionInput.cpp
+++ b/Engine/Plugins/FX/Niagara/Source/NiagaraEditor/Private/ViewModels/Stack/NiagaraStackFunctionInput.cpp
@@ -3,6 +3,8 @@
 #include "ViewModels/Stack/NiagaraStackFunctionInput.h"
 #include "ViewModels/Stack/NiagaraStackFunctionInputCollection.h"
 #include "ViewModels/Stack/NiagaraStackObject.h"
+#include "ViewModels/NiagaraScratchPadViewModel.h"
+#include "ViewModels/NiagaraScratchPadScriptViewModel.h"
 #include "NiagaraNodeFunctionCall.h"
 #include "NiagaraNodeAssignment.h"
 #include "NiagaraNodeParameterMapGet.h"
@@ -30,6 +32,7 @@
 #include "ViewModels/Stack/NiagaraParameterHandle.h"
 #include "NiagaraEmitter.h"
 #include "NiagaraClipboard.h"
+#include "Toolkits/NiagaraSystemToolkit.h"
 
 #include "Materials/Material.h"
 #include "Materials/MaterialExpressionDynamicParameter.h"
@@ -148,7 +151,6 @@ void UNiagaraStackFunctionInput::Initialize(
 	InputParameterHandlePath.Add(InputParameterHandle);
 
 	DisplayName = FText::FromName(InputParameterHandle.GetName());
-	AliasedInputParameterHandle = FNiagaraParameterHandle::CreateAliasedModuleParameterHandle(InputParameterHandle, OwningFunctionCallNode.Get());
 
 	InputType = InInputType;
 	StackEditorDataKey = FNiagaraStackGraphUtilities::GenerateStackFunctionInputEditorDataKey(*OwningFunctionCallNode.Get(), InputParameterHandle);
@@ -389,6 +391,7 @@ FText UNiagaraStackFunctionInput::GetValueToolTip() const
 
 void UNiagaraStackFunctionInput::RefreshChildrenInternal(const TArray<UNiagaraStackEntry*>& CurrentChildren, TArray<UNiagaraStackEntry*>& NewChildren, TArray<FStackIssue>& NewIssues)
 {
+	AliasedInputParameterHandle = FNiagaraParameterHandle::CreateAliasedModuleParameterHandle(InputParameterHandle, OwningFunctionCallNode.Get());
 	RapidIterationParameter = CreateRapidIterationVariable(AliasedInputParameterHandle.GetParameterHandleString());
 
 	RefreshFromMetaData();
@@ -1030,30 +1033,47 @@ void UNiagaraStackFunctionInput::GetAvailableDynamicInputs(TArray<UNiagaraScript
 	DynamicInputScriptFilterOptions.bIncludeNonLibraryScripts = bIncludeNonLibraryInputs;
 	FNiagaraEditorUtilities::GetFilteredScriptAssets(DynamicInputScriptFilterOptions, DynamicInputAssets);
 
+	auto MatchesInputType = [this](UNiagaraScript* Script)
+	{
+		UNiagaraScriptSource* DynamicInputScriptSource = Cast<UNiagaraScriptSource>(Script->GetSource());
+		TArray<UNiagaraNodeOutput*> OutputNodes;
+		DynamicInputScriptSource->NodeGraph->GetNodesOfClass<UNiagaraNodeOutput>(OutputNodes);
+		if (OutputNodes.Num() == 1)
+		{
+			TArray<UEdGraphPin*> InputPins;
+			OutputNodes[0]->GetInputPins(InputPins);
+			if (InputPins.Num() == 1)
+			{
+				const UEdGraphSchema_Niagara* NiagaraSchema = GetDefault<UEdGraphSchema_Niagara>();
+				FNiagaraTypeDefinition PinType = NiagaraSchema->PinToTypeDefinition(InputPins[0]);
+				if (PinType == InputType)
+				{
+					return true;
+				}
+			}
+		}
+		return false;
+	};
+
 	for (const FAssetData& DynamicInputAsset : DynamicInputAssets)
 	{
 		UNiagaraScript* DynamicInputScript = Cast<UNiagaraScript>(DynamicInputAsset.GetAsset());
 		if (DynamicInputScript != nullptr)
 		{
-			UNiagaraScriptSource* DynamicInputScriptSource = Cast<UNiagaraScriptSource>(DynamicInputScript->GetSource());
-			TArray<UNiagaraNodeOutput*> OutputNodes;
-			DynamicInputScriptSource->NodeGraph->GetNodesOfClass<UNiagaraNodeOutput>(OutputNodes);
-			if (OutputNodes.Num() == 1)
+			if(MatchesInputType(DynamicInputScript))
 			{
-				TArray<UEdGraphPin*> InputPins;
-				OutputNodes[0]->GetInputPins(InputPins);
-				if (InputPins.Num() == 1)
-				{
-					const UEdGraphSchema_Niagara* NiagaraSchema = GetDefault<UEdGraphSchema_Niagara>();
-					FNiagaraTypeDefinition PinType = NiagaraSchema->PinToTypeDefinition(InputPins[0]);
-					if (PinType == InputType)
-					{
-						AvailableDynamicInputs.Add(DynamicInputScript);
-					}
-				}
+				AvailableDynamicInputs.Add(DynamicInputScript);
 			}
 		}
 	}
+
+	for (TSharedRef<FNiagaraScratchPadScriptViewModel> ScratchPadScriptViewModel : GetSystemViewModel()->GetScriptScratchPadViewModel()->GetScriptViewModels())
+	{
+		if (MatchesInputType(ScratchPadScriptViewModel->GetScript()))
+		{
+			AvailableDynamicInputs.Add(ScratchPadScriptViewModel->GetScript());
+		}
+	}
 }
 
 void UNiagaraStackFunctionInput::SetDynamicInput(UNiagaraScript* DynamicInput, FString SuggestedName)
@@ -1099,6 +1119,18 @@ void UNiagaraStackFunctionInput::SetCustomExpression(const FString& InCustomExpr
 	RefreshChildren();
 }
 
+void UNiagaraStackFunctionInput::SetScratch()
+{
+	FScopedTransaction ScopedTransaction(LOCTEXT("SetScratch", "Make new scratch dynamic input"));
+	TSharedPtr<FNiagaraScratchPadScriptViewModel> ScratchScriptViewModel = GetSystemViewModel()->GetScriptScratchPadViewModel()->CreateNewScript(ENiagaraScriptUsage::DynamicInput, SourceScript->GetUsage(), InputType);
+	if (ScratchScriptViewModel.IsValid())
+	{
+		SetDynamicInput(ScratchScriptViewModel->GetScript());
+		ScratchScriptViewModel->SetIsPendingRename(true);
+	}
+	GetSystemViewModel()->FocusTab(FNiagaraSystemToolkit::ScratchPadTabID);
+}
+
 TSharedPtr<const FStructOnScope> UNiagaraStackFunctionInput::GetLocalValueStruct()
 {
 	return InputValues.LocalStruct;
@@ -1730,7 +1762,11 @@ void UNiagaraStackFunctionInput::SetValueFromClipboardFunctionInput(const UNiaga
 			if (ensureMsgf(ClipboardFunctionInput.Dynamic->ScriptMode == ENiagaraClipboardFunctionScriptMode::ScriptAsset,
 				TEXT("Can not set dynamic input value from clipboard, only script asset funcitons can be set.")))
 			{
-				SetDynamicInput(ClipboardFunctionInput.Dynamic->Script, ClipboardFunctionInput.Dynamic->FunctionName);
+				UNiagaraScript* ClipboardScript = ClipboardFunctionInput.Dynamic->Script->IsAsset()
+					? ClipboardFunctionInput.Dynamic->Script
+					: GetSystemViewModel()->GetScriptScratchPadViewModel()->CreateNewScriptAsDuplicate(ClipboardFunctionInput.Dynamic->Script)->GetScript();
+				SetDynamicInput(ClipboardScript, ClipboardFunctionInput.Dynamic->FunctionName);
+
 				TArray<UNiagaraStackFunctionInputCollection*> DynamicInputCollections;
 				GetUnfilteredChildrenOfType(DynamicInputCollections);
 				for (UNiagaraStackFunctionInputCollection* DynamicInputCollection : DynamicInputCollections)
diff --git a/Engine/Plugins/FX/Niagara/Source/NiagaraEditor/Private/ViewModels/Stack/NiagaraStackGraphUtilities.cpp b/Engine/Plugins/FX/Niagara/Source/NiagaraEditor/Private/ViewModels/Stack/NiagaraStackGraphUtilities.cpp
--- a/Engine/Plugins/FX/Niagara/Source/NiagaraEditor/Private/ViewModels/Stack/NiagaraStackGraphUtilities.cpp
+++ b/Engine/Plugins/FX/Niagara/Source/NiagaraEditor/Private/ViewModels/Stack/NiagaraStackGraphUtilities.cpp
@@ -1252,6 +1252,20 @@ UNiagaraNodeFunctionCall* FNiagaraStackGraphUtilities::AddScriptModuleToStack(FA
 	return NewModuleNode;
 }
 
+UNiagaraNodeFunctionCall* FNiagaraStackGraphUtilities::AddScriptModuleToStack(UNiagaraScript* ModuleScript, UNiagaraNodeOutput& TargetOutputNode, int32 TargetIndex)
+{
+	UEdGraph* Graph = TargetOutputNode.GetGraph();
+	Graph->Modify();
+
+	FGraphNodeCreator<UNiagaraNodeFunctionCall> ModuleNodeCreator(*Graph);
+	UNiagaraNodeFunctionCall* NewModuleNode = ModuleNodeCreator.CreateNode();
+	NewModuleNode->FunctionScript = ModuleScript;
+	ModuleNodeCreator.Finalize();
+
+	ConnectModuleNode(*NewModuleNode, TargetOutputNode, TargetIndex);
+	return NewModuleNode;
+}
+
 UNiagaraNodeAssignment* FNiagaraStackGraphUtilities::AddParameterModuleToStack(const TArray<FNiagaraVariable>& ParameterVariables, UNiagaraNodeOutput& TargetOutputNode, int32 TargetIndex, const TArray<FString>& InDefaultValues)
 {
 	UEdGraph* Graph = TargetOutputNode.GetGraph();
diff --git a/Engine/Plugins/FX/Niagara/Source/NiagaraEditor/Private/ViewModels/Stack/NiagaraStackScriptItemGroup.cpp b/Engine/Plugins/FX/Niagara/Source/NiagaraEditor/Private/ViewModels/Stack/NiagaraStackScriptItemGroup.cpp
--- a/Engine/Plugins/FX/Niagara/Source/NiagaraEditor/Private/ViewModels/Stack/NiagaraStackScriptItemGroup.cpp
+++ b/Engine/Plugins/FX/Niagara/Source/NiagaraEditor/Private/ViewModels/Stack/NiagaraStackScriptItemGroup.cpp
@@ -6,6 +6,8 @@
 #include "ViewModels/NiagaraSystemSelectionViewModel.h"
 #include "ViewModels/NiagaraEmitterViewModel.h"
 #include "ViewModels/NiagaraScriptViewModel.h"
+#include "ViewModels/NiagaraScratchPadViewModel.h"
+#include "ViewModels/NiagaraScratchPadScriptViewModel.h"
 #include "NiagaraScriptGraphViewModel.h"
 #include "NiagaraGraph.h"
 #include "NiagaraNodeOutput.h"
@@ -20,6 +22,7 @@
 #include "NiagaraEmitterHandle.h"
 #include "NiagaraEditorSettings.h"
 #include "NiagaraClipboard.h"
+#include "Toolkits/NiagaraSystemToolkit.h"
 
 #include "Internationalization/Internationalization.h"
 #include "ScopedTransaction.h"
@@ -44,40 +47,33 @@ public:
 			Category = LOCTEXT("ModuleNotCategorized", "Uncategorized Modules");
 		}
 
-		bool bIsLibraryScript = true;
-		bool bIsLibraryTagFound = false;
-		bIsLibraryTagFound = AssetData.GetTagValue(GET_MEMBER_NAME_CHECKED(UNiagaraScript, bExposeToLibrary), bIsLibraryScript);
-		if (bIsLibraryTagFound == false)
-		{
-			if (AssetData.IsAssetLoaded())
-			{
-				UNiagaraScript* Script = static_cast<UNiagaraScript*>(AssetData.GetAsset());
-				if (Script != nullptr)
-				{
-					bIsLibraryScript = Script->bExposeToLibrary;
-				}
-			}
-		}
+		bool bIsInLibrary = FNiagaraEditorUtilities::IsScriptAssetInLibrary(AssetData);
 
-		FString DisplayNameString = FName::NameToDisplayString(AssetData.AssetName.ToString(), false);
-		if (bIsLibraryScript == false)
-		{
-			DisplayNameString += "*";
-		}
-		FText DisplayName = FText::FromString(DisplayNameString);
+		FText DisplayName = FNiagaraEditorUtilities::FormatScriptName(AssetData.AssetName, bIsInLibrary);
 
 		FText AssetDescription;
 		AssetData.GetTagValue(GET_MEMBER_NAME_CHECKED(UNiagaraScript, Description), AssetDescription);
-		FText Description = FNiagaraEditorUtilities::FormatScriptAssetDescription(AssetDescription, AssetData.ObjectPath);
-		if (bIsLibraryScript == false)
-		{
-			Description = FText::Format(LOCTEXT("ScriptDescriptionWithLibraryCommentFormat", "{0}\n{1}"), Description, LOCTEXT("NotExposedToLibrary", "*Not exposed to library"));
-		}
+		FText Description = FNiagaraEditorUtilities::FormatScriptDescription(AssetDescription, AssetData.ObjectPath, bIsInLibrary);
 
 		FText Keywords;
 		AssetData.GetTagValue(GET_MEMBER_NAME_CHECKED(UNiagaraScript, Keywords), Keywords);
 
-		return MakeShareable(new FScriptGroupAddAction(Category, DisplayName, Description, Keywords, FNiagaraVariable(), false, AssetData, false));
+		return MakeShareable(new FScriptGroupAddAction(Category, DisplayName, Description, Keywords, FNiagaraVariable(), false, AssetData, nullptr, false));
+	}
+
+	static TSharedRef<FScriptGroupAddAction> CreateScriptModuleAction(UNiagaraScript* ModuleScript)
+	{
+		FText Category = ModuleScript->Category;
+		if (Category.IsEmptyOrWhitespace())
+		{
+			Category = LOCTEXT("ModuleNotCategorized", "Uncategorized Modules");
+		}
+
+		FText DisplayName = FNiagaraEditorUtilities::FormatScriptName(ModuleScript->GetFName(), ModuleScript->bExposeToLibrary);
+		FText Description = FNiagaraEditorUtilities::FormatScriptDescription(ModuleScript->Description, *ModuleScript->GetPathName(), ModuleScript->bExposeToLibrary);
+		FText Keywords = ModuleScript->Keywords;
+
+		return MakeShareable(new FScriptGroupAddAction(Category, DisplayName, Description, Keywords, FNiagaraVariable(), false, FAssetData(), ModuleScript, false));
 	}
 
 	static TSharedRef<FScriptGroupAddAction> CreateExistingParameterModuleAction(FNiagaraVariable ParameterVariable)
@@ -90,7 +86,7 @@ public:
 		FText AttributeDescription = FNiagaraConstants::GetAttributeDescription(ParameterVariable);
 		FText Description = FText::Format(LOCTEXT("ExistingParameterModuleDescriptoinFormat", "Description: Set the parameter {0}. {1}"), FText::FromName(ParameterVariable.GetName()), AttributeDescription);
 
-		return MakeShareable(new FScriptGroupAddAction(Category, DisplayName, Description, FText(), ParameterVariable, false, FAssetData(), false));
+		return MakeShareable(new FScriptGroupAddAction(Category, DisplayName, Description, FText(), ParameterVariable, false, FAssetData(), nullptr, false));
 	}
 
 	static TSharedRef<FScriptGroupAddAction> CreateNewParameterModuleAction(FName NewParameterNamespace, FNiagaraTypeDefinition NewParameterType)
@@ -102,7 +98,16 @@ public:
 		FNiagaraParameterHandle NewParameterHandle(NewParameterNamespace, *(TEXT("New") + NewParameterType.GetName()));
 		FNiagaraVariable NewParameter(NewParameterType, NewParameterHandle.GetParameterHandleString());
 
-		return MakeShareable(new FScriptGroupAddAction(Category, DisplayName, Description, FText(), NewParameter, true, FAssetData(), false));
+		return MakeShareable(new FScriptGroupAddAction(Category, DisplayName, Description, FText(), NewParameter, true, FAssetData(), nullptr, false));
+	}
+
+	static TSharedRef<FScriptGroupAddAction> CreateNewScratchModuleAction()
+	{
+		FText Category = LOCTEXT("NewScratchModuleCategory", "Scratch");
+		FText DisplayName = LOCTEXT("NewScratchModuleName", "New Scratch Pad Module");
+		FText Description = LOCTEXT("NewScratchModuleDescription", "Description: Create a new scratch pad module.");
+
+		return MakeShareable(new FScriptGroupAddAction(Category, DisplayName, Description, FText(), FNiagaraVariable(), false, FAssetData(), nullptr, true));
 	}
 
 	virtual FText GetCategory() const override
@@ -140,13 +145,21 @@ public:
 		return ModuleAssetData;
 	}
 
-	bool GetIsMaterialParameterModuleAction() const
+	UNiagaraScript* GetModuleScript() const
+	{
+		return ModuleScript;
+	}
+
+	bool GetIsNewScratchModuleAction() const
 	{
-		return bIsMaterialParameterModuleAction;
+		return bIsNewScratchModuleAction;
 	}
 
 private:
-	FScriptGroupAddAction(FText InCategory, FText InDisplayName, FText InDescription, FText InKeywords, FNiagaraVariable InModuleParameterVariable, bool bInRenameParameterOnAdd, FAssetData InModuleAssetData, bool bInIsMaterialParameterModuleAction)
+	FScriptGroupAddAction(FText InCategory, FText InDisplayName, FText InDescription, FText InKeywords,
+		FNiagaraVariable InModuleParameterVariable, bool bInRenameParameterOnAdd,
+		FAssetData InModuleAssetData, UNiagaraScript* InModuleScript, bool bInIsNewScratchModuleAction)
+
 		: Category(InCategory)
 		, DisplayName(InDisplayName)
 		, Description(InDescription)
@@ -154,7 +167,8 @@ private:
 		, ModuleParameterVariable(InModuleParameterVariable)
 		, bRenameParameterOnAdd(bInRenameParameterOnAdd)
 		, ModuleAssetData(InModuleAssetData)
-		, bIsMaterialParameterModuleAction(bInIsMaterialParameterModuleAction)
+		, ModuleScript(InModuleScript)
+		, bIsNewScratchModuleAction(bInIsNewScratchModuleAction)
 	{
 	}
 
@@ -166,7 +180,9 @@ private:
 	FNiagaraVariable ModuleParameterVariable;
 	bool bRenameParameterOnAdd;
 	FAssetData ModuleAssetData;
+	UNiagaraScript* ModuleScript;
 	bool bIsMaterialParameterModuleAction;
+	bool bIsNewScratchModuleAction;
 };
 
 class FScriptItemGroupAddUtilities : public TNiagaraStackItemGroupAddUtilities<UNiagaraNodeFunctionCall*>
@@ -227,6 +243,21 @@ public:
 				OutAddActions.Add(FScriptGroupAddAction::CreateNewParameterModuleAction(NewParameterNamespace.GetValue(), AvailableType));
 			}
 		}
+
+		// Generate actions for scratch pad modules
+		for (TSharedRef<FNiagaraScratchPadScriptViewModel> ScratchPadScriptViewModel : SystemViewModel.Pin()->GetScriptScratchPadViewModel()->GetScriptViewModels())
+		{
+			UNiagaraScript* ScratchPadScript = ScratchPadScriptViewModel->GetScript();
+			if (ScratchPadScript->Usage == ENiagaraScriptUsage::Module)
+			{
+				TArray<ENiagaraScriptUsage> SupportedUsages = ScratchPadScript->GetSupportedUsageContexts();
+				if (SupportedUsages.Contains(OutputNode->GetUsage()))
+				{
+					OutAddActions.Add(FScriptGroupAddAction::CreateScriptModuleAction(ScratchPadScript));
+				}
+			}
+		}
+		OutAddActions.Add(FScriptGroupAddAction::CreateNewScratchModuleAction());
 	}
 
 	virtual void ExecuteAddAction(TSharedRef<INiagaraStackItemGroupAddAction> AddAction, int32 TargetIndex) override
@@ -242,11 +273,27 @@ public:
 		{
 			NewModuleNode = AddParameterModule(ScriptGroupAddAction->GetModuleParameterVariable(), ScriptGroupAddAction->GetRenameParameterOnAdd(), TargetIndex);
 		}
+		else if (ScriptGroupAddAction->GetModuleScript() != nullptr)
+		{
+			NewModuleNode = FNiagaraStackGraphUtilities::AddScriptModuleToStack(ScriptGroupAddAction->GetModuleScript(), *OutputNode, TargetIndex);
+		}
+		else if (ScriptGroupAddAction->GetIsNewScratchModuleAction())
+		{
+			TSharedPtr<FNiagaraScratchPadScriptViewModel> ScratchPadScriptViewModel = SystemViewModel.Pin()->GetScriptScratchPadViewModel()->CreateNewScript(ENiagaraScriptUsage::Module, OutputNode->GetUsage(), FNiagaraTypeDefinition());
+			if (ScratchPadScriptViewModel.IsValid())
+			{
+				ScratchPadScriptViewModel->SetIsPendingRename(true);
+				NewModuleNode = FNiagaraStackGraphUtilities::AddScriptModuleToStack(ScratchPadScriptViewModel->GetScript(), *OutputNode, TargetIndex);
+				SystemViewModel.Pin()->FocusTab(FNiagaraSystemToolkit::ScratchPadTabID);
+			}
+		}
 
-		checkf(NewModuleNode != nullptr, TEXT("Add module action failed"));
-		FNiagaraStackGraphUtilities::InitializeStackFunctionInputs(SystemViewModel.Pin().ToSharedRef(), EmitterViewModel.Pin(), StackEditorData, *NewModuleNode, *NewModuleNode);
-		FNiagaraStackGraphUtilities::RelayoutGraph(*OutputNode->GetGraph());
-		OnItemAdded.ExecuteIfBound(NewModuleNode);
+		if (NewModuleNode != nullptr)
+		{
+			FNiagaraStackGraphUtilities::InitializeStackFunctionInputs(SystemViewModel.Pin().ToSharedRef(), EmitterViewModel.Pin(), StackEditorData, *NewModuleNode, *NewModuleNode);
+			FNiagaraStackGraphUtilities::RelayoutGraph(*OutputNode->GetGraph());
+			OnItemAdded.ExecuteIfBound(NewModuleNode);
+		}
 	}
 private:
 	UNiagaraNodeFunctionCall* AddScriptAssetModule(const FAssetData& AssetData, int32 TargetIndex)
@@ -1061,6 +1108,9 @@ void UNiagaraStackScriptItemGroup::PasteModules(const UNiagaraClipboardContent*
 			{
 				if (ClipboardFunction->Script->GetSupportedUsageContexts().Contains(OutputNode->GetUsage()))
 				{
+					UNiagaraScript* ClipboardScript = ClipboardFunction->Script->IsAsset()
+						? ClipboardFunction->Script
+						: GetSystemViewModel()->GetScriptScratchPadViewModel()->CreateNewScriptAsDuplicate(ClipboardFunction->Script)->GetScript();
 					NewFunctionCallNode = FNiagaraStackGraphUtilities::AddScriptModuleToStack(ClipboardFunction->Script, *OutputNode, CurrentPasteIndex);
 				}
 				else
diff --git a/Engine/Plugins/FX/Niagara/Source/NiagaraEditor/Private/Widgets/SDynamicLayoutBox.cpp b/Engine/Plugins/FX/Niagara/Source/NiagaraEditor/Private/Widgets/SDynamicLayoutBox.cpp
new file mode 100644
--- /dev/null
+++ b/Engine/Plugins/FX/Niagara/Source/NiagaraEditor/Private/Widgets/SDynamicLayoutBox.cpp
@@ -0,0 +1,58 @@
+// Copyright 1998-2020 Epic Games, Inc. All Rights Reserved.
+
+#include "Widgets/SDynamicLayoutBox.h"
+
+SDynamicLayoutBox::FNamedWidgetProvider::FNamedWidgetProvider()
+{
+}
+
+SDynamicLayoutBox::FNamedWidgetProvider::FNamedWidgetProvider(SDynamicLayoutBox::FOnGenerateNamedWidget InGenerateNamedWidget)
+{
+	GenerateNamedWidget = InGenerateNamedWidget;
+}
+
+TSharedRef<SWidget> SDynamicLayoutBox::FNamedWidgetProvider::GetNamedWidget(FName InWidgetName) const
+{
+	if (GenerateNamedWidget.IsBound())
+	{
+		TSharedRef<SWidget>* CachedNamedWidget = NameToGeneratedWidgetCache.Find(InWidgetName);
+		if (CachedNamedWidget != nullptr)
+		{
+			return *CachedNamedWidget;
+		}
+		else
+		{
+			TSharedRef<SWidget> NamedWidget = GenerateNamedWidget.Execute(InWidgetName);
+			NameToGeneratedWidgetCache.Add(InWidgetName, NamedWidget);
+			return NamedWidget;
+		}
+	}
+	return SNullWidget::NullWidget;
+}
+
+void SDynamicLayoutBox::Construct(const FArguments& InArgs)
+{
+	if (InArgs._GenerateNamedWidget.IsBound())
+	{
+		NamedWidgetProvider = FNamedWidgetProvider(InArgs._GenerateNamedWidget);
+	}
+	GenerateNamedLayout = InArgs._GenerateNamedLayout;
+	ChooseLayout = InArgs._ChooseLayout;
+}
+
+void SDynamicLayoutBox::Tick(const FGeometry& AllottedGeometry, const double InCurrentTime, const float InDeltaTime)
+{
+	if (GetCachedGeometry() != AllottedGeometry)
+	{
+		FName NewLayoutName;
+		if (ChooseLayout.IsBound())
+		{
+			NewLayoutName = ChooseLayout.Execute();
+		}
+		if (GenerateNamedLayout.IsBound() && (CurrentLayout.IsSet() == false || CurrentLayout.GetValue() != NewLayoutName))
+		{
+			ChildSlot[GenerateNamedLayout.Execute(NewLayoutName, NamedWidgetProvider)];
+			CurrentLayout = NewLayoutName;
+		}
+	}
+}
\ No newline at end of file
diff --git a/Engine/Plugins/FX/Niagara/Source/NiagaraEditor/Private/Widgets/SNiagaraScriptGraph.cpp b/Engine/Plugins/FX/Niagara/Source/NiagaraEditor/Private/Widgets/SNiagaraScriptGraph.cpp
--- a/Engine/Plugins/FX/Niagara/Source/NiagaraEditor/Private/Widgets/SNiagaraScriptGraph.cpp
+++ b/Engine/Plugins/FX/Niagara/Source/NiagaraEditor/Private/Widgets/SNiagaraScriptGraph.cpp
@@ -42,6 +42,10 @@ void SNiagaraScriptGraph::Construct(const FArguments& InArgs, TSharedRef<FNiagar
 	GraphTitle = InArgs._GraphTitle;
 
 	GraphEditor = ConstructGraphEditor();
+	if (InArgs._ZoomToFitOnLoad)
+	{
+		GraphEditor->ZoomToFit(false);
+	}
 
 	ChildSlot
 	[
@@ -202,6 +206,7 @@ TSharedRef<SGraphEditor> SNiagaraScriptGraph::ConstructGraphEditor()
 
 	// Set a niagara node factory.
 	CreatedGraphEditor->SetNodeFactory(MakeShareable(new FNiagaraNodeFactory()));
+	CreatedGraphEditor->ZoomToFit(false);
 
 	return CreatedGraphEditor;
 }
diff --git a/Engine/Plugins/FX/Niagara/Source/NiagaraEditor/Private/Widgets/SVerticalResizeBox.cpp b/Engine/Plugins/FX/Niagara/Source/NiagaraEditor/Private/Widgets/SVerticalResizeBox.cpp
new file mode 100644
--- /dev/null
+++ b/Engine/Plugins/FX/Niagara/Source/NiagaraEditor/Private/Widgets/SVerticalResizeBox.cpp
@@ -0,0 +1,113 @@
+// Copyright Epic Games, Inc. All Rights Reserved.
+
+#include "Widgets/SVerticalResizeBox.h"
+#include "Brushes/SlateColorBrush.h"
+#include "Widgets/Layout/SBox.h"
+
+void SVerticalResizeBox::Construct(const FArguments& InArgs)
+{
+	ContentHeight = InArgs._ContentHeight;
+	HandleHeight = InArgs._HandleHeight;
+	HandleColor = InArgs._HandleColor;
+	HandleHighlightColor = InArgs._HandleHighlightColor;
+	HandleBrush = FSlateColorBrush(FLinearColor::White);
+	ContentHeightChanged = InArgs._ContentHeightChanged;
+
+	ChildSlot
+	[
+		SNew(SBox)
+		.HeightOverride(this, &SVerticalResizeBox::GetHeightOverride)
+		.Padding(FMargin(0, 0, 0, HandleHeight))
+		[
+			InArgs._Content.Widget
+		]
+	];
+}
+
+FReply SVerticalResizeBox::OnMouseButtonDown(const FGeometry& MyGeometry, const FPointerEvent& MouseEvent)
+{
+	if (MouseEvent.GetEffectingButton() == EKeys::LeftMouseButton)
+	{
+		FVector2D MouseLocation = MyGeometry.AbsoluteToLocal(MouseEvent.GetScreenSpacePosition());
+		if (MyGeometry.GetLocalSize().Y - MouseLocation.Y < HandleHeight)
+		{
+			DragStartLocation = MouseLocation.Y;
+			DragStartContentHeight = ContentHeight.Get();
+			return FReply::Handled().CaptureMouse(SharedThis(this));
+		}
+	}
+	return FReply::Unhandled();
+}
+
+FReply SVerticalResizeBox::OnMouseButtonUp(const FGeometry& MyGeometry, const FPointerEvent& MouseEvent)
+{
+	if (this->HasMouseCapture())
+	{
+		return FReply::Handled().ReleaseMouseCapture();
+	}
+	return FReply::Unhandled();
+}
+
+FReply SVerticalResizeBox::OnMouseMove(const FGeometry& MyGeometry, const FPointerEvent& MouseEvent)
+{
+	FVector2D MouseLocation = MyGeometry.AbsoluteToLocal(MouseEvent.GetScreenSpacePosition());
+	LastMouseLocation = MouseLocation.Y;
+	if (this->HasMouseCapture())
+	{
+		float NewContentHeight = DragStartContentHeight + (MouseLocation.Y - DragStartLocation);
+		if (ContentHeight.IsBound() && ContentHeightChanged.IsBound())
+		{
+			ContentHeightChanged.Execute(NewContentHeight);
+		}
+		else
+		{
+			ContentHeight = NewContentHeight;
+		}
+		return FReply::Handled();
+	}
+	return FReply::Unhandled();
+}
+
+FCursorReply SVerticalResizeBox::OnCursorQuery(const FGeometry& MyGeometry, const FPointerEvent& CursorEvent) const
+{
+	FVector2D CursorLocation = MyGeometry.AbsoluteToLocal(CursorEvent.GetScreenSpacePosition());
+	if (MyGeometry.GetLocalSize().Y - CursorLocation.Y < HandleHeight)
+	{
+		return FCursorReply::Cursor(EMouseCursor::ResizeUpDown);
+	}
+	return FCursorReply::Unhandled();
+}
+
+int32 SVerticalResizeBox::OnPaint(const FPaintArgs& Args, const FGeometry& AllottedGeometry, const FSlateRect& MyClippingRect, FSlateWindowElementList& OutDrawElements, int32 LayerId, const FWidgetStyle& InWidgetStyle, bool bParentEnabled) const
+{
+	FLinearColor HandleBoxColor;
+	int32 HandleLayerId = LayerId + 1;
+	FVector2D LocalSize = AllottedGeometry.GetLocalSize();
+	if (IsHovered() && LastMouseLocation.IsSet() && LastMouseLocation.GetValue() >= LocalSize.Y - HandleHeight && LastMouseLocation.GetValue() <= LocalSize.Y)
+	{
+		HandleBoxColor = HandleHighlightColor.Get();
+	}
+	else
+	{
+		HandleBoxColor = HandleColor.Get();
+	}
+
+	FVector2D HandleLocation(0, AllottedGeometry.GetLocalSize().Y - HandleHeight);
+	FVector2D HandleSize(AllottedGeometry.GetLocalSize().X, HandleHeight);
+	FSlateDrawElement::MakeBox
+	(
+		OutDrawElements,
+		HandleLayerId,
+		AllottedGeometry.ToPaintGeometry(HandleLocation, HandleSize),
+		&HandleBrush,
+		ESlateDrawEffect::None,
+		HandleBoxColor
+	);
+
+	return SCompoundWidget::OnPaint(Args, AllottedGeometry, MyClippingRect, OutDrawElements, HandleLayerId, InWidgetStyle, bParentEnabled);
+}
+
+FOptionalSize SVerticalResizeBox::GetHeightOverride() const
+{
+	return ContentHeight.Get() + HandleHeight;
+}
\ No newline at end of file
diff --git a/Engine/Plugins/FX/Niagara/Source/NiagaraEditor/Public/NiagaraEditorModule.h b/Engine/Plugins/FX/Niagara/Source/NiagaraEditor/Public/NiagaraEditorModule.h
--- a/Engine/Plugins/FX/Niagara/Source/NiagaraEditor/Public/NiagaraEditorModule.h
+++ b/Engine/Plugins/FX/Niagara/Source/NiagaraEditor/Public/NiagaraEditorModule.h
@@ -27,6 +27,7 @@ class INiagaraEditorOnlyDataUtilities;
 class FNiagaraEditorCommands;
 struct FNiagaraScriptHighlight;
 class FNiagaraClipboard;
+class UNiagaraScratchPadViewModel;
 class FHlslNiagaraCompiler;
 
 DECLARE_STATS_GROUP(TEXT("Niagara Editor"), STATGROUP_NiagaraEditor, STATCAT_Advanced);
@@ -38,6 +39,7 @@ public:
 	virtual TSharedRef<SWidget> CreateStackView(UNiagaraStackViewModel& StackViewModel) const = 0;
 	virtual TSharedRef<SWidget> CreateSystemOverview(TSharedRef<FNiagaraSystemViewModel> SystemViewModel) const = 0;
 	virtual TSharedRef<SWidget> CreateStackIssueIcon(UNiagaraStackViewModel& StackViewModel, UNiagaraStackEntry& StackEntry) const = 0;
+	virtual TSharedRef<SWidget> CreateScriptScratchPad(UNiagaraScratchPadViewModel& ScriptScratchPadViewModel) const = 0;
 	virtual FLinearColor GetColorForExecutionCategory(FName ExecutionCategory) const = 0;
 };
 
diff --git a/Engine/Plugins/FX/Niagara/Source/NiagaraEditor/Public/NiagaraEditorUtilities.h b/Engine/Plugins/FX/Niagara/Source/NiagaraEditor/Public/NiagaraEditorUtilities.h
--- a/Engine/Plugins/FX/Niagara/Source/NiagaraEditor/Public/NiagaraEditorUtilities.h
+++ b/Engine/Plugins/FX/Niagara/Source/NiagaraEditor/Public/NiagaraEditorUtilities.h
@@ -194,7 +194,11 @@ namespace FNiagaraEditorUtilities
 	 */
 	const FNiagaraEmitterHandle* GetEmitterHandleForEmitter(UNiagaraSystem& System, UNiagaraEmitter& Emitter);
 
-	NIAGARAEDITOR_API FText FormatScriptAssetDescription(FText Description, FName Path);
+	NIAGARAEDITOR_API bool IsScriptAssetInLibrary(const FAssetData& ScriptAssetData);
+
+	NIAGARAEDITOR_API FText FormatScriptName(FName Name, bool bIsInLibrary);
+
+	NIAGARAEDITOR_API FText FormatScriptDescription(FText Description, FName Path, bool bIsInLibrary);
 
 	void ResetSystemsThatReferenceSystemViewModel(const FNiagaraSystemViewModel& ReferencedSystemViewModel);
 
@@ -234,4 +238,6 @@ namespace FNiagaraEditorUtilities
 	void ToggleSelectedEmittersIsolated(TSharedRef<FNiagaraSystemViewModel> SystemViewModel);
 
 	void CreateAssetFromEmitter(TSharedRef<FNiagaraEmitterHandleViewModel> EmitterHandleViewModel);
+
+	NIAGARAEDITOR_API void WarnWithToastAndLog(FText WarningMessage);
 };
diff --git a/Engine/Plugins/FX/Niagara/Source/NiagaraEditor/Public/ViewModels/NiagaraScratchPadScriptViewModel.h b/Engine/Plugins/FX/Niagara/Source/NiagaraEditor/Public/ViewModels/NiagaraScratchPadScriptViewModel.h
new file mode 100644
--- /dev/null
+++ b/Engine/Plugins/FX/Niagara/Source/NiagaraEditor/Public/ViewModels/NiagaraScratchPadScriptViewModel.h
@@ -0,0 +1,33 @@
+// Copyright Epic Games, Inc. All Rights Reserved.
+
+#pragma once
+
+#include "ViewModels/NiagaraScriptViewModel.h"
+
+class NIAGARAEDITOR_API FNiagaraScratchPadScriptViewModel : public FNiagaraScriptViewModel
+{
+public:
+	DECLARE_MULTICAST_DELEGATE(FOnRenamed);
+
+public:
+	FNiagaraScratchPadScriptViewModel();
+
+	void Initialize(UNiagaraScript* Script);
+
+	UNiagaraScript* GetScript() const;
+
+	bool GetIsPendingRename() const;
+
+	void SetIsPendingRename(bool bInIsPendingRename);
+
+	void SetScriptName(FText InScriptName);
+
+	FOnRenamed& OnRenamed();
+
+private:
+	FText GetDisplayNameInternal() const;
+
+private:
+	bool bIsPendingRename;
+	FOnRenamed OnRenamedDelegate;
+};
\ No newline at end of file
diff --git a/Engine/Plugins/FX/Niagara/Source/NiagaraEditor/Public/ViewModels/NiagaraScratchPadViewModel.h b/Engine/Plugins/FX/Niagara/Source/NiagaraEditor/Public/ViewModels/NiagaraScratchPadViewModel.h
new file mode 100644
--- /dev/null
+++ b/Engine/Plugins/FX/Niagara/Source/NiagaraEditor/Public/ViewModels/NiagaraScratchPadViewModel.h
@@ -0,0 +1,85 @@
+// Copyright Epic Games, Inc. All Rights Reserved.
+
+#pragma once
+
+#include "CoreMinimal.h"
+#include "NiagaraCommon.h"
+#include "NiagaraScratchPadViewModel.generated.h"
+
+class FNiagaraSystemViewModel;
+class FNiagaraScratchPadScriptViewModel;
+class FNiagaraObjectSelection;
+class UNiagaraScript;
+
+UCLASS()
+class NIAGARAEDITOR_API UNiagaraScratchPadViewModel : public UObject
+{
+	GENERATED_BODY()
+
+public:
+	DECLARE_MULTICAST_DELEGATE(FOnScriptViewModelsChanged);
+	DECLARE_MULTICAST_DELEGATE(FOnActiveScriptChanged);
+	DECLARE_MULTICAST_DELEGATE(FOnScriptRenamed);
+	DECLARE_MULTICAST_DELEGATE(FOnScriptDeleted);
+
+public:
+	void Initialize(TSharedRef<FNiagaraSystemViewModel> InSystemViewModel);
+
+	void Finalize();
+
+	void RefreshScriptViewModels();
+
+	const TArray<TSharedRef<FNiagaraScratchPadScriptViewModel>>& GetScriptViewModels() const;
+
+	TSharedPtr<FNiagaraScratchPadScriptViewModel> GetViewModelForScript(UNiagaraScript* InScript);
+
+	const TArray<ENiagaraScriptUsage>& GetAvailableUsages() const;
+
+	FText GetDisplayNameForUsage(ENiagaraScriptUsage InUsage) const;
+
+	TSharedRef<FNiagaraObjectSelection> GetObjectSelection();
+
+	UNiagaraScript* GetActiveScript();
+
+	void SetActiveScript(UNiagaraScript* InActiveScript);
+
+	void DeleteActiveScript();
+
+	TSharedPtr<FNiagaraScratchPadScriptViewModel> CreateNewScript(ENiagaraScriptUsage InScriptUsage, ENiagaraScriptUsage InTargetSupportedUsage, FNiagaraTypeDefinition InOutputType);
+
+	TSharedPtr<FNiagaraScratchPadScriptViewModel> CreateNewScriptAsDuplicate(UNiagaraScript* ScriptToDuplicate);
+
+	FOnScriptViewModelsChanged& OnScriptViewModelsChanged();
+
+	FOnActiveScriptChanged& OnActiveScriptChanged();
+
+	FOnScriptRenamed& OnScriptRenamed();
+
+	FOnScriptDeleted& OnScriptDeleted();
+
+private:
+	TSharedRef<FNiagaraSystemViewModel> GetSystemViewModel();
+
+	void ScriptGraphNodeSelectionChanged(TWeakPtr<FNiagaraScratchPadScriptViewModel> InScriptViewModelWeak);
+
+	void ScriptViewModelScriptRenamed();
+
+private:
+	TSharedPtr<FNiagaraObjectSelection> ObjectSelection;
+
+	UNiagaraScript* ActiveScript;
+
+	TWeakPtr<FNiagaraSystemViewModel> SystemViewModelWeak;
+
+	TArray<TSharedRef<FNiagaraScratchPadScriptViewModel>> ScriptViewModels;
+
+	TArray<ENiagaraScriptUsage> AvailableUsages;
+
+	FOnScriptViewModelsChanged OnScriptViewModelsChangedDelegate;
+
+	FOnActiveScriptChanged OnActiveScriptChangedDelegate;
+
+	FOnScriptRenamed OnScriptRenamedDelegate;
+
+	FOnScriptDeleted OnScriptDeletedDelegate;
+};
\ No newline at end of file
diff --git a/Engine/Plugins/FX/Niagara/Source/NiagaraEditor/Public/ViewModels/NiagaraScriptViewModel.h b/Engine/Plugins/FX/Niagara/Source/NiagaraEditor/Public/ViewModels/NiagaraScriptViewModel.h
--- a/Engine/Plugins/FX/Niagara/Source/NiagaraEditor/Public/ViewModels/NiagaraScriptViewModel.h
+++ b/Engine/Plugins/FX/Niagara/Source/NiagaraEditor/Public/ViewModels/NiagaraScriptViewModel.h
@@ -23,9 +23,13 @@ class FNiagaraObjectSelection;
 class FNiagaraScriptViewModel : public TSharedFromThis<FNiagaraScriptViewModel>, public FEditorUndoClient, public TNiagaraViewModelManager<UNiagaraScript, FNiagaraScriptViewModel>
 {
 public:
-	FNiagaraScriptViewModel(FText DisplayName, ENiagaraParameterEditMode InParameterEditMode);
+	FNiagaraScriptViewModel(TAttribute<FText> DisplayName, ENiagaraParameterEditMode InParameterEditMode);
 
-	~FNiagaraScriptViewModel();
+	virtual ~FNiagaraScriptViewModel();
+
+	NIAGARAEDITOR_API FText GetDisplayName() const;
+
+	NIAGARAEDITOR_API const TArray<TWeakObjectPtr<UNiagaraScript>>& GetScripts() const;
 
 	/** Sets the view model to a different script. */
 	void SetScript(UNiagaraScript* InScript);
@@ -42,7 +46,7 @@ public:
 	TSharedRef<INiagaraParameterCollectionViewModel> GetOutputParameterMapViewModel();
 
 	/** Gets the view model for the graph. */
-	TSharedRef<FNiagaraScriptGraphViewModel> GetGraphViewModel();
+	NIAGARAEDITOR_API TSharedRef<FNiagaraScriptGraphViewModel> GetGraphViewModel();
 
 	/** Gets the currently selected script variables. */
 	TSharedRef<FNiagaraObjectSelection> GetVariableSelection();
diff --git a/Engine/Plugins/FX/Niagara/Source/NiagaraEditor/Public/ViewModels/NiagaraSystemViewModel.h b/Engine/Plugins/FX/Niagara/Source/NiagaraEditor/Public/ViewModels/NiagaraSystemViewModel.h
--- a/Engine/Plugins/FX/Niagara/Source/NiagaraEditor/Public/ViewModels/NiagaraSystemViewModel.h
+++ b/Engine/Plugins/FX/Niagara/Source/NiagaraEditor/Public/ViewModels/NiagaraSystemViewModel.h
@@ -36,6 +36,7 @@ struct FEdGraphEditAction;
 class UNiagaraNodeFunctionCall;
 class FNiagaraEmitterViewModel;
 class FNiagaraOverviewGraphViewModel;
+class UNiagaraScratchPadViewModel;
 
 /** Defines different editing modes for this system view model. */
 enum class NIAGARAEDITOR_API ENiagaraSystemViewModelEditMode
@@ -102,6 +103,8 @@ public:
 
 	DECLARE_MULTICAST_DELEGATE(FOnPreClose);
 
+	DECLARE_MULTICAST_DELEGATE_OneParam(FOnRequestFocusTab, FName /* TabName */);
+
 public:
 	struct FEmitterHandleToDuplicate
 	{
@@ -275,6 +278,10 @@ public:
 	NIAGARAEDITOR_API void NotifyPreClose();
 
 	NIAGARAEDITOR_API FOnPreClose& OnPreClose();
+
+	FOnRequestFocusTab& OnRequestFocusTab();
+
+	void FocusTab(FName TabName);
 	
 	/** Gets the system toolkit command list. */
 	NIAGARAEDITOR_API TSharedPtr<FUICommandList> GetToolkitCommands();
@@ -297,6 +304,8 @@ public:
 	/** Gets the a view model representing the selected entries in the overview. */
 	NIAGARAEDITOR_API UNiagaraSystemSelectionViewModel* GetSelectionViewModel();
 
+	NIAGARAEDITOR_API UNiagaraScratchPadViewModel* GetScriptScratchPadViewModel();
+
 	/** Duplicates a set of emitters and refreshes everything.*/
 	void DuplicateEmitters(TArray<FEmitterHandleToDuplicate> EmitterHandlesToDuplicate);
 
@@ -424,6 +433,9 @@ private:
 	/** Called whenever one of the owned stack viewmodels structure changes. */
 	void StackViewModelStructureChanged();
 
+	/** Called whenever one of the scripts in the scratch pad changes. */
+	void ScratchPadScriptsChanged();
+
 private:
 	/** The System being viewed and edited by this view model. */
 	UNiagaraSystem* System;
@@ -488,6 +500,8 @@ private:
 	/** A multicast delegate which is called whenever this has been notified it's owner will be closing. */
 	FOnPreClose OnPreCloseDelegate;
 
+	FOnRequestFocusTab OnRequestFocusTabDelegate;
+
 	/** A flag for preventing reentrancy when syncrhonizing sequencer data. */
 	bool bUpdatingEmittersFromSequencerDataChange;
 
@@ -553,4 +567,6 @@ private:
 	UNiagaraStackViewModel* SystemStackViewModel;
 
 	UNiagaraSystemSelectionViewModel* SelectionViewModel;
+
+	UNiagaraScratchPadViewModel* ScriptScratchPadViewModel;
 };
\ No newline at end of file
diff --git a/Engine/Plugins/FX/Niagara/Source/NiagaraEditor/Public/ViewModels/Stack/NiagaraStackFunctionInput.h b/Engine/Plugins/FX/Niagara/Source/NiagaraEditor/Public/ViewModels/Stack/NiagaraStackFunctionInput.h
--- a/Engine/Plugins/FX/Niagara/Source/NiagaraEditor/Public/ViewModels/Stack/NiagaraStackFunctionInput.h
+++ b/Engine/Plugins/FX/Niagara/Source/NiagaraEditor/Public/ViewModels/Stack/NiagaraStackFunctionInput.h
@@ -144,6 +144,9 @@ public:
 	/** Sets the dynamic custom expression script for this input. */
 	void SetCustomExpression(const FString& InCustomExpression);
 
+	/** Create a new scratch pad dynamic inputs and set this input to use it. */
+	void SetScratch();
+
 	/** Gets the current struct value of this input is there is one. */
 	TSharedPtr<const FStructOnScope> GetLocalValueStruct();
 
diff --git a/Engine/Plugins/FX/Niagara/Source/NiagaraEditor/Public/ViewModels/Stack/NiagaraStackGraphUtilities.h b/Engine/Plugins/FX/Niagara/Source/NiagaraEditor/Public/ViewModels/Stack/NiagaraStackGraphUtilities.h
--- a/Engine/Plugins/FX/Niagara/Source/NiagaraEditor/Public/ViewModels/Stack/NiagaraStackGraphUtilities.h
+++ b/Engine/Plugins/FX/Niagara/Source/NiagaraEditor/Public/ViewModels/Stack/NiagaraStackGraphUtilities.h
@@ -118,6 +118,8 @@ namespace FNiagaraStackGraphUtilities
 	bool RemoveModuleFromStack(UNiagaraScript& OwningScript, UNiagaraNodeFunctionCall& ModuleNode, TArray<TWeakObjectPtr<UNiagaraNodeInput>>& OutRemovedInputNodes);
 
 	UNiagaraNodeFunctionCall* AddScriptModuleToStack(FAssetData ModuleScriptAsset, UNiagaraNodeOutput& TargetOutputNode, int32 TargetIndex = INDEX_NONE);
+
+	UNiagaraNodeFunctionCall* AddScriptModuleToStack(UNiagaraScript* ModuleScript, UNiagaraNodeOutput& TargetOutputNode, int32 TargetIndex = INDEX_NONE);
 	
 	bool FindScriptModulesInStack(FAssetData ModuleScriptAsset, UNiagaraNodeOutput& TargetOutputNode, TArray<UNiagaraNodeFunctionCall*> OutFunctionCalls);
 
diff --git a/Engine/Plugins/FX/Niagara/Source/NiagaraEditor/Public/Widgets/SDynamicLayoutBox.h b/Engine/Plugins/FX/Niagara/Source/NiagaraEditor/Public/Widgets/SDynamicLayoutBox.h
new file mode 100644
--- /dev/null
+++ b/Engine/Plugins/FX/Niagara/Source/NiagaraEditor/Public/Widgets/SDynamicLayoutBox.h
@@ -0,0 +1,47 @@
+// Copyright 1998-2020 Epic Games, Inc. All Rights Reserved.
+
+#pragma once
+
+#include "Widgets/SCompoundWidget.h"
+#include "Widgets/DeclarativeSyntaxSupport.h"
+
+class NIAGARAEDITOR_API SDynamicLayoutBox : public SCompoundWidget
+{
+public:
+	DECLARE_DELEGATE_RetVal_OneParam(TSharedRef<SWidget>, FOnGenerateNamedWidget, FName /* InWidgetName */);
+
+	class NIAGARAEDITOR_API FNamedWidgetProvider
+	{
+	public:
+		FNamedWidgetProvider();
+		FNamedWidgetProvider(FOnGenerateNamedWidget InGeneratedNamedWidget);
+
+		TSharedRef<SWidget> GetNamedWidget(FName InWidgetName) const;
+
+	private:
+		mutable TMap<FName, TSharedRef<SWidget>> NameToGeneratedWidgetCache;
+		FOnGenerateNamedWidget GenerateNamedWidget;
+	};
+
+	DECLARE_DELEGATE_RetVal_TwoParams(TSharedRef<SWidget>, FOnGenerateNamedLayout, FName /* InLayoutName */, const FNamedWidgetProvider& /* InNamedWidgetProvider */);
+
+	DECLARE_DELEGATE_RetVal(FName, FOnChooseLayout);
+
+public:
+	SLATE_BEGIN_ARGS(SDynamicLayoutBox)
+	{}
+		SLATE_EVENT(FOnGenerateNamedWidget, GenerateNamedWidget)
+		SLATE_EVENT(FOnGenerateNamedLayout, GenerateNamedLayout)
+		SLATE_EVENT(FOnChooseLayout, ChooseLayout);
+	SLATE_END_ARGS();
+
+	void Construct(const FArguments& InArgs);
+
+	virtual void Tick(const FGeometry& AllottedGeometry, const double InCurrentTime, const float InDeltaTime) override;
+
+private:
+	FNamedWidgetProvider NamedWidgetProvider;
+	FOnGenerateNamedLayout GenerateNamedLayout;
+	FOnChooseLayout ChooseLayout;
+	TOptional<FName> CurrentLayout;
+};
\ No newline at end of file
diff --git a/Engine/Plugins/FX/Niagara/Source/NiagaraEditor/Public/Widgets/SItemSelector.h b/Engine/Plugins/FX/Niagara/Source/NiagaraEditor/Public/Widgets/SItemSelector.h
--- a/Engine/Plugins/FX/Niagara/Source/NiagaraEditor/Public/Widgets/SItemSelector.h
+++ b/Engine/Plugins/FX/Niagara/Source/NiagaraEditor/Public/Widgets/SItemSelector.h
@@ -20,10 +20,12 @@ public:
 	DECLARE_DELEGATE_RetVal_OneParam(TArray<CategoryType>, FOnGetCategoriesForItem, const ItemType& /* Item */);
 	DECLARE_DELEGATE_RetVal_TwoParams(bool, FOnCompareCategoriesForEquality, const CategoryType& /* CategoryA */, const CategoryType& /* CategoryB */);
 	DECLARE_DELEGATE_RetVal_TwoParams(bool, FOnCompareCategoriesForSorting, const CategoryType& /* CategoryA */, const CategoryType& /* CateogoryB */);
+	DECLARE_DELEGATE_RetVal_TwoParams(bool, FOnCompareItemsForEquality, const ItemType& /* ItemA */, const ItemType& /* ItemB */);
 	DECLARE_DELEGATE_RetVal_TwoParams(bool, FOnCompareItemsForSorting, const ItemType& /* ItemA */, const ItemType& /* ItemB */);
 	DECLARE_DELEGATE_RetVal_TwoParams(bool, FOnDoesItemMatchFilterText, const FText& /* Filter text */, const ItemType& /* Item */);
 	DECLARE_DELEGATE_RetVal_OneParam(TSharedRef<SWidget>, FOnGenerateWidgetForCategory, const CategoryType& /* Category */);
 	DECLARE_DELEGATE_RetVal_OneParam(TSharedRef<SWidget>, FOnGenerateWidgetForItem, const ItemType& /* Item */);
+	DECLARE_DELEGATE(FOnSelectionChanged);
 	DECLARE_DELEGATE_OneParam(FOnItemActivated, const ItemType& /* Item */);
 
 public:
@@ -31,8 +33,17 @@ public:
 		: _AllowMultiselect(false)
 		, _ClickActivateMode(EItemSelectorClickActivateMode::DoubleClick)
 	{}
+		/** The items to display in the item selector. */
+		SLATE_ARGUMENT(TArray<ItemType>, Items)
+
+		/** The default categories to show in additional to ones generated from the items.
+		NOTE: The OnCompareCategoriesForEquality, and OnGenerateWidgetForCategory delegates must be bound if this argument is supplied. */
+		SLATE_ARGUMENT(TArray<CategoryType>, DefaultCategories)
+
+		/** The default category paths to show in additional to ones generated from the items. 
+		NOTE: The OnCompareCategoriesForEquality, and OnGenerateWidgetForCategory delegates must be bound if this argument is supplied. */
+		SLATE_ARGUMENT(TArray<TArray<CategoryType>>, DefaultCategoryPaths)
 
-	SLATE_ARGUMENT(TArray<ItemType>, Items)
 		/** Whether or not this item selector should allow multiple items to be selected. */
 		SLATE_ARGUMENT(bool, AllowMultiselect)
 
@@ -49,6 +60,9 @@ public:
 		/** An optional delegate which determines the sorting for categories.  If not bound categories will be ordered by the order they're encountered while processing items. */
 		SLATE_EVENT(FOnCompareCategoriesForSorting, OnCompareCategoriesForSorting)
 
+		/** An optional delegate which compares items for equality.  This is required for using the SetSelectedItems api. */
+		SLATE_EVENT(FOnCompareItemsForEquality, OnCompareItemsForEquality)
+
 		/** An optional delegate which determines the sorting for items within each category. */
 		SLATE_EVENT(FOnCompareItemsForSorting, OnCompareItemsForSorting)
 
@@ -63,6 +77,9 @@ public:
 
 		/** A delegate which is called when an item is activated by either double clicking on it or by pressing enter while it's selected. */
 		SLATE_EVENT(FOnItemActivated, OnItemActivated)
+
+		/** A delegate which is called when the selection changes. */
+		SLATE_EVENT(FOnSelectionChanged, OnSelectionChanged)
 	SLATE_END_ARGS();
 
 private:
@@ -79,9 +96,11 @@ private:
 		{
 		}
 
+		virtual bool IsFiltering() const = 0;
 		virtual bool DoesItemMatchFilterText(const ItemType& InItem) const = 0;
 		virtual bool CompareCategoriesForEquality(const CategoryType& CategoryA, const CategoryType& CategoryB) const = 0;
 		virtual const FOnCompareCategoriesForSorting& GetOnCompareCategoriesForSorting() const = 0;
+		virtual bool CompareItemsForEquality(const ItemType& ItemA, const ItemType& ItemB) const = 0;
 		virtual const FOnCompareItemsForSorting& GetOnCompareItemsForSorting() const = 0;
 	};
 
@@ -112,15 +131,22 @@ private:
 
 		void GetChildren(TArray<TSharedRef<FItemSelectorItemViewModel>>& OutChildren) const
 		{
-			TArray<TSharedRef<FItemSelectorItemViewModel>> Children;
-			GetChildrenInternal(Children);
-			for (TSharedRef<FItemSelectorItemViewModel> Child : Children)
+			if (GetItemUtilities()->IsFiltering())
 			{
-				if (Child->PassesFilter())
+				TArray<TSharedRef<FItemSelectorItemViewModel>> Children;
+				GetChildrenInternal(Children);
+				for (TSharedRef<FItemSelectorItemViewModel> Child : Children)
 				{
-					OutChildren.Add(Child);
+					if (Child->PassesFilter())
+					{
+						OutChildren.Add(Child);
+					}
 				}
 			}
+			else
+			{
+				GetChildrenInternal(OutChildren);
+			}
 		}
 
 		virtual bool PassesFilter() const
@@ -175,6 +201,23 @@ private:
 			return Category;
 		}
 
+		void GetItemViewModelsForItems(const TArray<ItemType>& InItems, TArray<TSharedRef<FItemSelectorItemViewModel>>& OutItemViewModelsForItems)
+		{
+			for (TSharedRef<FItemSelectorItemCategoryViewModel> ChildCategoryViewModel : ChildCategoryViewModels)
+			{
+				ChildCategoryViewModel->GetItemViewModelsForItems(InItems, OutItemViewModelsForItems);
+			}
+			for (TSharedRef<FItemSelectorItemContainerViewModel> ChildItemViewModel : ChildItemViewModels)
+			{
+				const ItemType& ChildItem = ChildItemViewModel->GetItem();
+				TSharedRef<IItemSelectorItemViewModelUtilities> MyItemUtilities = this->GetItemUtilities();
+				if (InItems.ContainsByPredicate([ChildItem, MyItemUtilities](const ItemType& Item) { return MyItemUtilities->CompareItemsForEquality(Item, ChildItem); }))
+				{
+					OutItemViewModelsForItems.Add(ChildItemViewModel);
+				}
+			}
+		}
+
 		TSharedRef<FItemSelectorItemCategoryViewModel> AddCategory(const CategoryType& InCategory)
 		{
 			TSharedRef<FItemSelectorItemCategoryViewModel> NewCategoryViewModel = MakeShared<FItemSelectorItemCategoryViewModel>(this->GetItemUtilities(), InCategory);
@@ -259,11 +302,17 @@ private:
 	class FItemSelectorViewModel : public IItemSelectorItemViewModelUtilities, public TSharedFromThis<FItemSelectorViewModel>
 	{
 	public:
-		FItemSelectorViewModel(TArray<ItemType> InItems, FOnGetCategoriesForItem InOnGetCategoriesForItem, FOnCompareCategoriesForEquality InOnCompareCategoriesForEquality, FOnCompareCategoriesForSorting InOnCompareCategoriesForSorting, FOnCompareItemsForSorting InOnCompareItemsForSorting, FOnDoesItemMatchFilterText InOnDoesItemMatchFilterText)
+		FItemSelectorViewModel(TArray<ItemType> InItems, TArray<TArray<CategoryType>> InDefaultCategoryPaths,
+			FOnGetCategoriesForItem InOnGetCategoriesForItem,
+			FOnCompareCategoriesForEquality InOnCompareCategoriesForEquality, FOnCompareCategoriesForSorting InOnCompareCategoriesForSorting,
+			FOnCompareItemsForEquality InOnCompareItemsForEquality, FOnCompareItemsForSorting InOnCompareItemsForSorting,
+			FOnDoesItemMatchFilterText InOnDoesItemMatchFilterText)
 			: Items(InItems)
+			, DefaultCategoryPaths(InDefaultCategoryPaths)
 			, OnGetCategoriesForItem(InOnGetCategoriesForItem)
 			, OnCompareCategoriesForEquality(InOnCompareCategoriesForEquality)
 			, OnCompareCategoriesForSorting(InOnCompareCategoriesForSorting)
+			, OnCompareItemsForEquality(InOnCompareItemsForEquality)
 			, OnCompareItemsForSorting(InOnCompareItemsForSorting)
 			, OnDoesItemMatchFilterText(InOnDoesItemMatchFilterText)
 		{
@@ -274,44 +323,62 @@ private:
 			if (RootCategoryViewModel.IsValid() == false)
 			{
 				RootCategoryViewModel = MakeShared<FItemSelectorItemCategoryViewModel>(this->AsShared(), RootCategory);
-				for (const ItemType& Item : Items)
-				{
-					// Cache the category array since the category view models hold a const reference to the generated categories and without this they
-					// are deleted.
-					TSharedRef<TArray<CategoryType>> ItemCategories = MakeShared<TArray<CategoryType>>();
-					ItemCategoriesCache.Add(ItemCategories);
-
-					if (OnGetCategoriesForItem.IsBound())
-					{
-						ItemCategories->Append(OnGetCategoriesForItem.Execute(Item));
-					}
 
-					TSharedPtr<FItemSelectorItemCategoryViewModel> CurrentCategoryViewModel = RootCategoryViewModel;
-					if (ItemCategories->Num() > 0)
-					{
+				for (const TArray<CategoryType>& DefaultCategoryPath : DefaultCategoryPaths)
+				{
+					AddCategory(DefaultCategoryPath);
+				}
 
-						for (const CategoryType& ItemCategory : (*ItemCategories))
-						{
-							TSharedPtr<FItemSelectorItemCategoryViewModel> ExistingCategoryViewModel = CurrentCategoryViewModel->FindChildCategory(ItemCategory);
-							if (ExistingCategoryViewModel.IsValid())
-							{
-								CurrentCategoryViewModel = ExistingCategoryViewModel;
-							}
-							else
-							{
-								TSharedRef<FItemSelectorItemCategoryViewModel> NewItemCategoryViewModel = CurrentCategoryViewModel->AddCategory(ItemCategory);
-								CurrentCategoryViewModel = NewItemCategoryViewModel;
-							}
-						}
-					}
-					CurrentCategoryViewModel->AddItem(Item);
+				for (const ItemType& Item : Items)
+				{
+					AddItem(Item);
 				}
+
 				RootCategoryViewModel->SortChildren();
 				RootCategoryViewModel->GetChildren(RootTreeCategories);
 			}
 			return &RootTreeCategories;
 		}
 
+		void GetItemViewModelsForItems(const TArray<ItemType>& InItems, TArray<TSharedRef<FItemSelectorItemViewModel>>& OutItemViewModelsForItems)
+		{
+			RootCategoryViewModel->GetItemViewModelsForItems(InItems, OutItemViewModelsForItems);
+		}
+
+		TSharedRef<FItemSelectorItemCategoryViewModel> AddCategory(const TArray<CategoryType>& CategoryPath)
+		{
+			TSharedPtr<FItemSelectorItemCategoryViewModel> CurrentCategoryViewModel = RootCategoryViewModel;
+			for (const CategoryType& Category : CategoryPath)
+			{
+				TSharedPtr<FItemSelectorItemCategoryViewModel> ExistingCategoryViewModel = CurrentCategoryViewModel->FindChildCategory(Category);
+				if (ExistingCategoryViewModel.IsValid())
+				{
+					CurrentCategoryViewModel = ExistingCategoryViewModel;
+				}
+				else
+				{
+					TSharedRef<FItemSelectorItemCategoryViewModel> NewItemCategoryViewModel = CurrentCategoryViewModel->AddCategory(Category);
+					CurrentCategoryViewModel = NewItemCategoryViewModel;
+				}
+			}
+			return CurrentCategoryViewModel.ToSharedRef();
+		}
+
+		void AddItem(const ItemType& Item)
+		{
+			// Cache the category array since the category view models hold a const reference to the generated categories and without this they are deleted.
+			TSharedRef<TArray<CategoryType>> ItemCategories = MakeShared<TArray<CategoryType>>();
+			ItemCategoriesCache.Add(ItemCategories);
+
+			if (OnGetCategoriesForItem.IsBound())
+			{
+				ItemCategories->Append(OnGetCategoriesForItem.Execute(Item));
+			}
+
+			TSharedRef<FItemSelectorItemCategoryViewModel> ItemCategory = AddCategory(ItemCategories.Get());
+			ItemCategory->AddItem(Item);
+		}
+
 		const FText& GetFilterText() const
 		{
 			return FilterText;
@@ -324,6 +391,11 @@ private:
 			RootCategoryViewModel->GetChildren(RootTreeCategories);
 		}
 
+		virtual bool IsFiltering() const override
+		{
+			return FilterText.IsEmptyOrWhitespace() == false;
+		}
+
 		virtual bool DoesItemMatchFilterText(const ItemType& InItem) const override
 		{
 			return FilterText.IsEmpty() || OnDoesItemMatchFilterText.IsBound() == false || OnDoesItemMatchFilterText.Execute(FilterText, InItem);
@@ -339,6 +411,11 @@ private:
 			return OnCompareCategoriesForSorting;
 		}
 
+		virtual bool CompareItemsForEquality(const ItemType& ItemA, const ItemType& ItemB) const override
+		{
+			return OnCompareItemsForEquality.Execute(ItemA, ItemB);
+		}
+
 		virtual const FOnCompareItemsForSorting& GetOnCompareItemsForSorting() const override
 		{
 			return OnCompareItemsForSorting;
@@ -354,15 +431,30 @@ private:
 			return OnCompareItemsForSorting.Execute(ItemA, ItemB);
 		}
 
+		void Refresh(const TArray<ItemType>& InItems, const TArray<TArray<CategoryType>>& InDefaultCategoryPaths)
+		{
+			if (RootCategoryViewModel.IsValid())
+			{
+				ItemCategoriesCache.Empty();
+				RootCategoryViewModel.Reset();
+				RootTreeCategories.Empty();
+			}
+			Items = InItems;
+			DefaultCategoryPaths = InDefaultCategoryPaths;
+			GetRootItems();
+		}
+
 	private:
 		CategoryType RootCategory;
 
 		TArray<ItemType> Items;
+		TArray<TArray<CategoryType>> DefaultCategoryPaths;
 		TArray<TSharedRef<TArray<CategoryType>>> ItemCategoriesCache;
 
 		FOnGetCategoriesForItem OnGetCategoriesForItem;
 		FOnCompareCategoriesForEquality OnCompareCategoriesForEquality;
 		FOnCompareCategoriesForSorting OnCompareCategoriesForSorting;
+		FOnCompareItemsForEquality OnCompareItemsForEquality;
 		FOnCompareItemsForSorting OnCompareItemsForSorting;
 		FOnDoesItemMatchFilterText OnDoesItemMatchFilterText;
 
@@ -398,21 +490,36 @@ public:
 	void Construct(const FArguments& InArgs)
 	{
 		Items = InArgs._Items;
+		TArray<TArray<CategoryType>> DefaultCategoryPathsFromDefaultCategories;
+		for (const CategoryType& DefaultCategory : InArgs._DefaultCategories)
+		{
+			DefaultCategoryPaths.AddDefaulted();
+			DefaultCategoryPaths.Last().Add(DefaultCategory);
+		}
+		DefaultCategoryPaths.Append(InArgs._DefaultCategoryPaths);
 		ClickActivateMode = InArgs._ClickActivateMode;
 		OnGetCategoriesForItem = InArgs._OnGetCategoriesForItem;
 		OnCompareCategoriesForEquality = InArgs._OnCompareCategoriesForEquality;
 		OnCompareCategoriesForSorting = InArgs._OnCompareCategoriesForSorting;
+		OnCompareItemsForEquality = InArgs._OnCompareItemsForEquality;
 		OnCompareItemsForSorting = InArgs._OnCompareItemsForSorting;
 		OnDoesItemMatchFilterText = InArgs._OnDoesItemMatchFilterText;
 		OnGenerateWidgetForCategory = InArgs._OnGenerateWidgetForCategory;
 		OnGenerateWidgetForItem = InArgs._OnGenerateWidgetForItem;
 		OnItemActivated = InArgs._OnItemActivated;
+		OnSelectionChanged = InArgs._OnSelectionChanged;
+		bIsSettingSelection = false;
 
+		checkf(DefaultCategoryPaths.Num() == 0 || OnCompareCategoriesForEquality.IsBound(), TEXT("OnCompareCategoriesForEquality must be bound if default categories are supplied."));
+		checkf(DefaultCategoryPaths.Num() == 0 || OnGenerateWidgetForCategory.IsBound(), TEXT("OnGenerateWidgetForCategory must be bound if default categories are supplied."));
 		checkf(OnGetCategoriesForItem.IsBound() == false || OnCompareCategoriesForEquality.IsBound(), TEXT("OnCompareCategoriesForEquality must be bound if OnGenerateCategoriesForItem is bound."));
 		checkf(OnGetCategoriesForItem.IsBound() == false || OnGenerateWidgetForCategory.IsBound(), TEXT("OnGenerateWidgetForCategory must be bound if OnGenerateCategoriesForItem is bound."));
 		checkf(OnGenerateWidgetForItem.IsBound(), TEXT("OnGenerateWidgetForItem must be bound"));
 
-		ViewModel = MakeShared<FItemSelectorViewModel>(Items, OnGetCategoriesForItem, OnCompareCategoriesForEquality, OnCompareCategoriesForSorting, OnCompareItemsForSorting, OnDoesItemMatchFilterText);
+		ViewModel = MakeShared<FItemSelectorViewModel>(
+			Items, DefaultCategoryPaths, OnGetCategoriesForItem, 
+			OnCompareCategoriesForEquality, OnCompareCategoriesForSorting, 
+			OnCompareItemsForEquality, OnCompareItemsForSorting, OnDoesItemMatchFilterText);
 
 		ChildSlot
 		[
@@ -434,6 +541,7 @@ public:
 				.OnGetChildren(this, &SItemSelector::OnGetChildren)
 				.OnMouseButtonClick(this, &SItemSelector::OnMouseClick)
 				.OnMouseButtonDoubleClick(this, &SItemSelector::OnMouseDoubleClick)
+				.OnSelectionChanged(this, &SItemSelector::OnTreeSelectionChanged)
 				.TreeItemsSource(ViewModel->GetRootItems())
 			]
 		];
@@ -458,6 +566,31 @@ public:
 		return SelectedItems;
 	}
 
+	void SetSelectedItems(const TArray<ItemType>& NewSelectedItems)
+	{
+		checkf(OnCompareItemsForEquality.IsBound(), TEXT("OnCompareItemsForEquality event must be handled to use the SetSelectedItems function."));
+		TGuardValue<bool> SelectionGuard(bIsSettingSelection, true);
+		if (NewSelectedItems.Num() == 0)
+		{
+			ItemTree->ClearSelection();
+		}
+		else
+		{
+			TArray<TSharedRef<FItemSelectorItemViewModel>> SelectedItemViewModels;
+			ViewModel->GetItemViewModelsForItems(NewSelectedItems, SelectedItemViewModels);
+			for (TSharedRef<FItemSelectorItemViewModel> SelectedItemViewModel : SelectedItemViewModels)
+			{
+				ItemTree->SetSelection(SelectedItemViewModel);
+			}
+		}
+		OnSelectionChanged.ExecuteIfBound();
+	}
+
+	void ClearSelectedItems()
+	{
+		ItemTree->ClearSelection();
+	}
+
 	virtual FReply OnKeyDown(const FGeometry& MyGeometry, const FKeyEvent& InKeyEvent)
 	{
 		if (InKeyEvent.GetKey() == EKeys::Enter && OnItemActivated.IsBound())
@@ -474,6 +607,30 @@ public:
 		return SCompoundWidget::OnKeyDown(MyGeometry, InKeyEvent);
 	}
 
+	void RefreshItemsAndDefaultCategories(const TArray<ItemType>& InItems, const TArray<TArray<CategoryType>>& InDefaultCategoryPaths)
+	{
+		ViewModel->Refresh(InItems, InDefaultCategoryPaths);
+		ExpandTree();
+		ItemTree->RequestTreeRefresh();
+	}
+
+	void RefreshItemsAndDefaultCategories(const TArray<ItemType>& InItems, const TArray<CategoryType> &InDefaultCategories)
+	{
+		TArray<TArray<CategoryType>> DefaultCategoryPathsFromDefaultCategories;
+		for (const CategoryType& DefaultCategory : InDefaultCategories)
+		{
+			DefaultCategoryPathsFromDefaultCategories.AddDefaulted();
+			DefaultCategoryPathsFromDefaultCategories.Last().Add(DefaultCategory);
+		}
+		RefreshItemsAndDefaultCategories(InItems, DefaultCategoryPathsFromDefaultCategories);
+	}
+
+	void RefreshItems(const TArray<ItemType>& InItems)
+	{
+		TArray<TArray<CategoryType>> UnusedDefaultCategoryPaths;
+		RefreshItemsAndDefaultCategories(InItems, UnusedDefaultCategoryPaths);
+	}
+
 private:
 	EVisibility GetSearchBoxVisibility() const
 	{
@@ -540,6 +697,14 @@ private:
 		}
 	}
 
+	void OnTreeSelectionChanged(TSharedPtr<FItemSelectorItemViewModel> SelectedItem, ESelectInfo::Type)
+	{
+		if (bIsSettingSelection == false)
+		{
+			OnSelectionChanged.ExecuteIfBound();
+		}
+	}
+
 	void ExpandTree()
 	{
 		TArray<TSharedRef<FItemSelectorItemViewModel>> ItemsToProcess;
@@ -555,19 +720,24 @@ private:
 
 private:
 	TArray<ItemType> Items;
+	TArray<TArray<CategoryType>> DefaultCategoryPaths;
 
 	EItemSelectorClickActivateMode ClickActivateMode;
 
 	FOnGetCategoriesForItem OnGetCategoriesForItem;
 	FOnCompareCategoriesForEquality OnCompareCategoriesForEquality;
 	FOnCompareCategoriesForSorting OnCompareCategoriesForSorting;
+	FOnCompareItemsForEquality OnCompareItemsForEquality;
 	FOnCompareItemsForSorting OnCompareItemsForSorting;
 	FOnDoesItemMatchFilterText OnDoesItemMatchFilterText;
 	FOnGenerateWidgetForCategory OnGenerateWidgetForCategory;
 	FOnGenerateWidgetForItem OnGenerateWidgetForItem;
 	FOnItemActivated OnItemActivated;
+	FOnSelectionChanged OnSelectionChanged;
 
 	TSharedPtr<FItemSelectorViewModel> ViewModel;
 	TSharedPtr<SSearchBox> SearchBox;
 	TSharedPtr<STreeView<TSharedRef<FItemSelectorItemViewModel>>> ItemTree;
+
+	bool bIsSettingSelection;
 };
\ No newline at end of file
diff --git a/Engine/Plugins/FX/Niagara/Source/NiagaraEditor/Public/Widgets/SNiagaraScriptGraph.h b/Engine/Plugins/FX/Niagara/Source/NiagaraEditor/Public/Widgets/SNiagaraScriptGraph.h
--- a/Engine/Plugins/FX/Niagara/Source/NiagaraEditor/Public/Widgets/SNiagaraScriptGraph.h
+++ b/Engine/Plugins/FX/Niagara/Source/NiagaraEditor/Public/Widgets/SNiagaraScriptGraph.h
@@ -16,13 +16,14 @@ class SNiagaraScriptGraph : public SCompoundWidget
 {
 public:
 	SLATE_BEGIN_ARGS(SNiagaraScriptGraph)
+		: _ZoomToFitOnLoad(false)
 	{}
 		/** The text displayed in the title bar of the graph. */
 		SLATE_ATTRIBUTE(FText, GraphTitle)
-
+		SLATE_ARGUMENT(bool, ZoomToFitOnLoad)
 	SLATE_END_ARGS();
 
-	void Construct(const FArguments& InArgs, TSharedRef<FNiagaraScriptGraphViewModel> InViewModel);
+	NIAGARAEDITOR_API void Construct(const FArguments& InArgs, TSharedRef<FNiagaraScriptGraphViewModel> InViewModel);
 
 	const TSharedPtr<SGraphEditor> GetGraphEditor() { return GraphEditor; };
 
diff --git a/Engine/Plugins/FX/Niagara/Source/NiagaraEditor/Public/Widgets/SNiagaraSelectedObjectsDetails.h b/Engine/Plugins/FX/Niagara/Source/NiagaraEditor/Public/Widgets/SNiagaraSelectedObjectsDetails.h
--- a/Engine/Plugins/FX/Niagara/Source/NiagaraEditor/Public/Widgets/SNiagaraSelectedObjectsDetails.h
+++ b/Engine/Plugins/FX/Niagara/Source/NiagaraEditor/Public/Widgets/SNiagaraSelectedObjectsDetails.h
@@ -19,8 +19,8 @@ public:
 
 	SLATE_END_ARGS();
 
-	void Construct(const FArguments& InArgs, TSharedRef<FNiagaraObjectSelection> InSelectedObjects);
-	void Construct(const FArguments& InArgs, TSharedRef<FNiagaraObjectSelection> InSelectedObjects, TSharedRef<FNiagaraObjectSelection> InSelectedObjects2);
+	NIAGARAEDITOR_API void Construct(const FArguments& InArgs, TSharedRef<FNiagaraObjectSelection> InSelectedObjects);
+	NIAGARAEDITOR_API void Construct(const FArguments& InArgs, TSharedRef<FNiagaraObjectSelection> InSelectedObjects, TSharedRef<FNiagaraObjectSelection> InSelectedObjects2);
 
 	/** Delegate to know when one of the properties has been changed.*/
 	FOnFinishedChangingProperties& OnFinishedChangingProperties() { return OnFinishedChangingPropertiesDelegate; }
diff --git a/Engine/Plugins/FX/Niagara/Source/NiagaraEditor/Public/Widgets/SVerticalResizeBox.h b/Engine/Plugins/FX/Niagara/Source/NiagaraEditor/Public/Widgets/SVerticalResizeBox.h
new file mode 100644
--- /dev/null
+++ b/Engine/Plugins/FX/Niagara/Source/NiagaraEditor/Public/Widgets/SVerticalResizeBox.h
@@ -0,0 +1,54 @@
+// Copyright Epic Games, Inc. All Rights Reserved.
+
+#pragma once
+
+#include "Widgets/SCompoundWidget.h"
+#include "Widgets/DeclarativeSyntaxSupport.h"
+#include "Types/SlateStructs.h"
+
+class NIAGARAEDITOR_API SVerticalResizeBox : public SCompoundWidget
+{
+public:
+	DECLARE_DELEGATE_OneParam(FOnContentHeightChanged, float);
+
+public:
+	SLATE_BEGIN_ARGS(SVerticalResizeBox)
+		: _HandleHeight(5)
+		, _ContentHeight(50)
+		, _HandleColor(FLinearColor(0.0f, 0.0f, 0.0f, 0.0f))
+		, _HandleHighlightColor(FLinearColor(1.0f, 1.0f, 1.0f, 0.5f))
+	{}
+		SLATE_ARGUMENT(float, HandleHeight)
+		SLATE_ATTRIBUTE(float, ContentHeight)
+		SLATE_ATTRIBUTE(FLinearColor, HandleColor)
+		SLATE_ATTRIBUTE(FLinearColor, HandleHighlightColor)
+		SLATE_EVENT(FOnContentHeightChanged, ContentHeightChanged)
+		SLATE_DEFAULT_SLOT(FArguments, Content);
+	SLATE_END_ARGS();
+
+	void Construct(const FArguments& InArgs);
+
+	virtual FReply OnMouseButtonDown(const FGeometry& MyGeometry, const FPointerEvent& MouseEvent) override;
+	virtual FReply OnMouseButtonUp(const FGeometry& MyGeometry, const FPointerEvent& MouseEvent) override;
+	virtual FReply OnMouseMove(const FGeometry& MyGeometry, const FPointerEvent& MouseEvent) override;
+	virtual FCursorReply OnCursorQuery(const FGeometry& MyGeometry, const FPointerEvent& CursorEvent) const override;
+	virtual int32 OnPaint(const FPaintArgs& Args, const FGeometry& AllottedGeometry, const FSlateRect& MyClippingRect, FSlateWindowElementList& OutDrawElements, int32 LayerId, const FWidgetStyle& InWidgetStyle, bool bParentEnabled) const override;
+
+private:
+	FOptionalSize GetHeightOverride() const;
+
+private:
+	TOptional<float> LastMouseLocation;
+
+	TAttribute<float> ContentHeight;
+	float HandleHeight;
+
+	float DragStartLocation;
+	float DragStartContentHeight;
+
+	TAttribute<FLinearColor> HandleColor;
+	TAttribute<FLinearColor> HandleHighlightColor;
+	FSlateBrush HandleBrush;
+
+	FOnContentHeightChanged ContentHeightChanged;
+};
\ No newline at end of file
diff --git a/Engine/Plugins/FX/Niagara/Source/NiagaraEditorWidgets/Private/DetailCustomizations/NiagaraDataInterfaceCurveDetails.cpp b/Engine/Plugins/FX/Niagara/Source/NiagaraEditorWidgets/Private/DetailCustomizations/NiagaraDataInterfaceCurveDetails.cpp
--- a/Engine/Plugins/FX/Niagara/Source/NiagaraEditorWidgets/Private/DetailCustomizations/NiagaraDataInterfaceCurveDetails.cpp
+++ b/Engine/Plugins/FX/Niagara/Source/NiagaraEditorWidgets/Private/DetailCustomizations/NiagaraDataInterfaceCurveDetails.cpp
@@ -19,6 +19,7 @@
 #include "Brushes/SlateColorBrush.h"
 #include "Modules/ModuleManager.h"
 
+#include "Widgets/SVerticalResizeBox.h"
 #include "Widgets/Input/SComboButton.h"
 #include "Widgets/Input/SButton.h"
 #include "Widgets/Images/SImage.h"
@@ -42,141 +43,6 @@ FRichCurve* GetCurveFromPropertyHandle(TSharedPtr<IPropertyHandle> Handle)
 	return RawData.Num() == 1 ? static_cast<FRichCurve*>(RawData[0]) : nullptr;
 }
 
-class SNiagaraResizeBox : public SCompoundWidget
-{
-public:
-	DECLARE_DELEGATE_OneParam(FOnContentHeightChanged, float);
-
-public:
-	SLATE_BEGIN_ARGS(SNiagaraResizeBox)
-		: _HandleHeight(5)
-		, _ContentHeight(50)
-		, _HandleColor(FLinearColor(0.0f, 0.0f, 0.0f, 0.0f))
-		, _HandleHighlightColor(FLinearColor(1.0f, 1.0f, 1.0f, 0.5f))
-	{}
-		SLATE_ARGUMENT(float, HandleHeight)
-		SLATE_ATTRIBUTE(float, ContentHeight)
-		SLATE_ATTRIBUTE(FLinearColor, HandleColor)
-		SLATE_ATTRIBUTE(FLinearColor, HandleHighlightColor)
-		SLATE_EVENT(FOnContentHeightChanged, ContentHeightChanged)
-		SLATE_DEFAULT_SLOT(FArguments, Content)
-	SLATE_END_ARGS()
-
-	void Construct(const FArguments& InArgs)
-	{
-		ContentHeight = InArgs._ContentHeight;
-		HandleHeight = InArgs._HandleHeight;
-		HandleColor = InArgs._HandleColor;
-		HandleHighlightColor = InArgs._HandleHighlightColor;
-		HandleBrush = FSlateColorBrush(FLinearColor::White);
-		ContentHeightChanged = InArgs._ContentHeightChanged;
-
-		ChildSlot
-		[
-			SNew(SBox)
-			.HeightOverride(this, &SNiagaraResizeBox::GetHeightOverride)
-			.Padding(FMargin(0, 0, 0, HandleHeight))
-			[
-				InArgs._Content.Widget
-			]
-		];
-	}
-
-	virtual FReply OnMouseButtonDown(const FGeometry& MyGeometry, const FPointerEvent& MouseEvent) override
-	{
-		if (MouseEvent.GetEffectingButton() == EKeys::LeftMouseButton)
-		{
-			FVector2D MouseLocation = MyGeometry.AbsoluteToLocal(MouseEvent.GetScreenSpacePosition());
-			if (MyGeometry.GetLocalSize().Y - MouseLocation.Y < HandleHeight)
-			{
-				DragStartLocation = MouseLocation.Y;
-				DragStartContentHeight = ContentHeight.Get();
-				return FReply::Handled().CaptureMouse(SharedThis(this));
-			}
-		}
-		return FReply::Unhandled();
-	}
-
-	virtual FReply OnMouseButtonUp(const FGeometry& MyGeometry, const FPointerEvent& MouseEvent) override
-	{
-		if (this->HasMouseCapture())
-		{
-			return FReply::Handled().ReleaseMouseCapture();
-		}
-		return FReply::Unhandled();
-	}
-
-	virtual FReply OnMouseMove(const FGeometry& MyGeometry, const FPointerEvent& MouseEvent) override
-	{
-		FVector2D MouseLocation = MyGeometry.AbsoluteToLocal(MouseEvent.GetScreenSpacePosition());
-		LastMouseLocation = MouseLocation.Y;
-		if (this->HasMouseCapture())
-		{
-			float NewContentHeight = DragStartContentHeight + (MouseLocation.Y - DragStartLocation);
-			if (ContentHeight.IsBound() && ContentHeightChanged.IsBound())
-			{
-				ContentHeightChanged.Execute(NewContentHeight);
-			}
-			else
-			{
-				ContentHeight = NewContentHeight;
-			}
-			return FReply::Handled();
-		}
-		return FReply::Unhandled();
-	}
-
-	int32 OnPaint(const FPaintArgs& Args, const FGeometry& AllottedGeometry, const FSlateRect& MyClippingRect, FSlateWindowElementList& OutDrawElements, int32 LayerId, const FWidgetStyle& InWidgetStyle, bool bParentEnabled) const
-	{
-		FLinearColor HandleBoxColor;
-		int32 HandleLayerId = LayerId + 1;
-		FVector2D LocalSize = AllottedGeometry.GetLocalSize();
-		if (IsHovered() && LastMouseLocation.IsSet() && LastMouseLocation.GetValue() >= LocalSize.Y - HandleHeight && LastMouseLocation.GetValue() <= LocalSize.Y)
-		{
-			HandleBoxColor = HandleHighlightColor.Get();
-		}
-		else
-		{
-			HandleBoxColor = HandleColor.Get();
-		}
-
-		FVector2D HandleLocation(0, AllottedGeometry.GetLocalSize().Y - HandleHeight);
-		FVector2D HandleSize(AllottedGeometry.GetLocalSize().X, HandleHeight);
-		FSlateDrawElement::MakeBox
-		(
-			OutDrawElements,
-			HandleLayerId,
-			AllottedGeometry.ToPaintGeometry(HandleLocation, HandleSize),
-			&HandleBrush,
-			ESlateDrawEffect::None,
-			HandleBoxColor
-		);
-
-		return SCompoundWidget::OnPaint(Args, AllottedGeometry, MyClippingRect, OutDrawElements, HandleLayerId, InWidgetStyle, bParentEnabled);
-	}
-
-private:
-	FOptionalSize GetHeightOverride() const
-	{
-		return ContentHeight.Get() + HandleHeight;
-	}
-
-private:
-	TOptional<float> LastMouseLocation;
-
-	TAttribute<float> ContentHeight;
-	float HandleHeight;
-
-	float DragStartLocation;
-	float DragStartContentHeight;
-
-	TAttribute<FLinearColor> HandleColor;
-	TAttribute<FLinearColor> HandleHighlightColor;
-	FSlateBrush HandleBrush;
-
-	FOnContentHeightChanged ContentHeightChanged;
-};
-
 class SNiagaraDataInterfaceCurveEditor : public SCompoundWidget
 {
 public:
@@ -427,7 +293,7 @@ void FNiagaraDataInterfaceCurveDetailsBase::CustomizeDetails(IDetailLayoutBuilde
 	CurveCategory.AddCustomRow(NSLOCTEXT("NiagaraDataInterfaceCurveDetails", "CurveFilterText", "Curve"))
 		.WholeRowContent()
 		[
-			SNew(SNiagaraResizeBox)
+			SNew(SVerticalResizeBox)
 			.ContentHeight(StackCurveEditorOptions, &FNiagaraStackCurveEditorOptions::GetHeight)
 			.ContentHeightChanged(StackCurveEditorOptions, &FNiagaraStackCurveEditorOptions::SetHeight)
 			.Content()
diff --git a/Engine/Plugins/FX/Niagara/Source/NiagaraEditorWidgets/Private/NiagaraEditorWidgetsModule.cpp b/Engine/Plugins/FX/Niagara/Source/NiagaraEditorWidgets/Private/NiagaraEditorWidgetsModule.cpp
--- a/Engine/Plugins/FX/Niagara/Source/NiagaraEditorWidgets/Private/NiagaraEditorWidgetsModule.cpp
+++ b/Engine/Plugins/FX/Niagara/Source/NiagaraEditorWidgets/Private/NiagaraEditorWidgetsModule.cpp
@@ -10,6 +10,7 @@
 #include "SNiagaraOverviewGraph.h"
 #include "NiagaraEditorWidgetsUtilities.h"
 #include "Stack/SNiagaraStackIssueIcon.h"
+#include "SNiagaraScratchPad.h"
 
 #include "Modules/ModuleManager.h"
 #include "PropertyEditorModule.h"
@@ -187,6 +188,11 @@ TSharedRef<SWidget> FNiagaraEditorWidgetsModule::FNiagaraEditorWidgetProvider::C
 	return SNew(SNiagaraStackIssueIcon, &StackViewModel, &StackEntry);
 }
 
+TSharedRef<SWidget> FNiagaraEditorWidgetsModule::FNiagaraEditorWidgetProvider::CreateScriptScratchPad(UNiagaraScratchPadViewModel& ScriptScratchPadViewModel) const
+{
+	return SNew(SNiagaraScratchPad, &ScriptScratchPadViewModel);
+}
+
 FLinearColor FNiagaraEditorWidgetsModule::FNiagaraEditorWidgetProvider::GetColorForExecutionCategory(FName ExecutionCategory) const
 {
 	return FNiagaraEditorWidgetsStyle::Get().GetColor(FNiagaraStackEditorWidgetsUtilities::GetIconColorNameForExecutionCategory(ExecutionCategory));
diff --git a/Engine/Plugins/FX/Niagara/Source/NiagaraEditorWidgets/Private/NiagaraEditorWidgetsStyle.cpp b/Engine/Plugins/FX/Niagara/Source/NiagaraEditorWidgets/Private/NiagaraEditorWidgetsStyle.cpp
--- a/Engine/Plugins/FX/Niagara/Source/NiagaraEditorWidgets/Private/NiagaraEditorWidgetsStyle.cpp
+++ b/Engine/Plugins/FX/Niagara/Source/NiagaraEditorWidgets/Private/NiagaraEditorWidgetsStyle.cpp
@@ -237,6 +237,20 @@ TSharedRef< FSlateStyleSet > FNiagaraEditorWidgetsStyle::Create()
 
 	Style->Set("NiagaraEditor.ShowInCurveEditorIcon", new IMAGE_PLUGIN_BRUSH("Icons/ShowInCurveEditor", Icon16x16, FLinearColor::White));
 
+	// Scratch pad
+	FSlateFontInfo ScratchPadLargeHeaderFont = DEFAULT_FONT("Bold", 11);
+	FTextBlockStyle ScratchPadLargeHeaderText = FTextBlockStyle(NormalText)
+		.SetFont(ScratchPadLargeHeaderFont);
+	Style->Set("NiagaraEditor.ScratchPad.LargeHeaderText", ScratchPadLargeHeaderText);
+
+	FSlateFontInfo ScratchPadSmallHeaderFont = DEFAULT_FONT("Bold", 9);
+	FTextBlockStyle ScratchPadSmallHeaderText = FTextBlockStyle(NormalText)
+		.SetFont(ScratchPadSmallHeaderFont);
+	Style->Set("NiagaraEditor.ScratchPad.SmallHeaderText", ScratchPadSmallHeaderText);
+
+	Style->Set("NiagaraEditor.ScratchPad.HeaderColor", FLinearColor(FColor(48, 48, 48)));
+
+
 	return Style;
 }
 
diff --git a/Engine/Plugins/FX/Niagara/Source/NiagaraEditorWidgets/Private/NiagaraScratchPadCommandContext.cpp b/Engine/Plugins/FX/Niagara/Source/NiagaraEditorWidgets/Private/NiagaraScratchPadCommandContext.cpp
new file mode 100644
--- /dev/null
+++ b/Engine/Plugins/FX/Niagara/Source/NiagaraEditorWidgets/Private/NiagaraScratchPadCommandContext.cpp
@@ -0,0 +1,150 @@
+// Copyright Epic Games, Inc. All Rights Reserved.
+
+#include "NiagaraScratchPadCommandContext.h"
+#include "ViewModels/NiagaraScratchPadViewModel.h"
+
+#include "Framework/Notifications/NotificationManager.h"
+#include "Logging/LogMacros.h"
+#include "Widgets/Notifications/SNotificationList.h"
+#include "Framework/Commands/GenericCommands.h"
+#include "Framework/Commands/UICommandList.h"
+#include "Framework/MultiBox/MultiBoxBuilder.h"
+
+#define LOCTEXT_NAMESPACE "NiagaraScratchPadCommandContext"
+
+FNiagaraScratchPadCommandContext::FNiagaraScratchPadCommandContext(UNiagaraScratchPadViewModel* InScratchPadViewModel)
+	: Commands(MakeShared<FUICommandList>())
+	, bCommandsAreSetup(false)
+	, bProcessingCommandBindings(false)
+	, ScratchPadViewModel(InScratchPadViewModel)
+{
+}
+
+TSharedRef<FUICommandList> FNiagaraScratchPadCommandContext::GetCommands()
+{
+	if (bCommandsAreSetup == false)
+	{
+		SetupCommands();
+		bCommandsAreSetup = true;
+	}
+	return Commands;
+}
+
+bool FNiagaraScratchPadCommandContext::ProcessCommandBindings(const FKeyEvent& InKeyEvent)
+{
+	TGuardValue<bool> ProcessingGuard(bProcessingCommandBindings, true);
+	return GetCommands()->ProcessCommandBindings(InKeyEvent);
+}
+
+bool FNiagaraScratchPadCommandContext::AddEditMenuItems(FMenuBuilder& MenuBuilder)
+{
+	bool bSupportsCut = false;
+	bool bSupportsCopy = false;
+	bool bSupportsPaste = false;
+	bool bSupportsDelete = true;
+	if (bSupportsCut || bSupportsCopy || bSupportsPaste || bSupportsDelete)
+	{
+		MenuBuilder.BeginSection("EntryEdit", LOCTEXT("EntryEditActions", "Edit"));
+		{
+			if (bSupportsCut)
+			{
+				TAttribute<FText> CanCutToolTip = TAttribute<FText>::Create(TAttribute<FText>::FGetter::CreateSP(this, &FNiagaraScratchPadCommandContext::GetCanCutSelectedScriptsToolTip));
+				MenuBuilder.AddMenuEntry(FGenericCommands::Get().Cut, NAME_None, TAttribute<FText>(), CanCutToolTip);
+			}
+			if (bSupportsCopy)
+			{
+				TAttribute<FText> CanCopyToolTip = TAttribute<FText>::Create(TAttribute<FText>::FGetter::CreateSP(this, &FNiagaraScratchPadCommandContext::GetCanCopySelectedScriptsToolTip));
+				MenuBuilder.AddMenuEntry(FGenericCommands::Get().Copy, NAME_None, TAttribute<FText>(), CanCopyToolTip);
+			}
+			if (bSupportsPaste)
+			{
+				TAttribute<FText> CanPasteToolTip = TAttribute<FText>::Create(TAttribute<FText>::FGetter::CreateSP(this, &FNiagaraScratchPadCommandContext::GetCanPasteSelectedScriptsToolTip));
+				MenuBuilder.AddMenuEntry(FGenericCommands::Get().Paste, NAME_None, TAttribute<FText>(), CanPasteToolTip);
+			}
+			if (bSupportsDelete)
+			{
+				TAttribute<FText> CanDeleteToolTip = TAttribute<FText>::Create(TAttribute<FText>::FGetter::CreateSP(this, &FNiagaraScratchPadCommandContext::GetCanDeleteSelectedScriptsToolTip));
+				MenuBuilder.AddMenuEntry(FGenericCommands::Get().Delete, NAME_None, TAttribute<FText>(), CanDeleteToolTip);
+			}
+		}
+		MenuBuilder.EndSection();
+		return true;
+	}
+	return false;
+}
+
+void FNiagaraScratchPadCommandContext::SetupCommands()
+{
+	Commands->MapAction(FGenericCommands::Get().Cut, FUIAction(
+		FExecuteAction::CreateSP(this, &FNiagaraScratchPadCommandContext::CutSelectedScripts),
+		FCanExecuteAction::CreateSP(this, &FNiagaraScratchPadCommandContext::CanCutSelectedScripts)));
+	Commands->MapAction(FGenericCommands::Get().Copy, FUIAction(
+		FExecuteAction::CreateSP(this, &FNiagaraScratchPadCommandContext::CopySelectedScripts),
+		FCanExecuteAction::CreateSP(this, &FNiagaraScratchPadCommandContext::CanCopySelectedScripts)));
+	Commands->MapAction(FGenericCommands::Get().Paste, FUIAction(
+		FExecuteAction::CreateSP(this, &FNiagaraScratchPadCommandContext::PasteSelectedScripts),
+		FCanExecuteAction::CreateSP(this, &FNiagaraScratchPadCommandContext::CanPasteSelectedScripts)));
+	Commands->MapAction(FGenericCommands::Get().Delete, FUIAction(
+		FExecuteAction::CreateSP(this, &FNiagaraScratchPadCommandContext::DeleteSelectedScripts),
+		FCanExecuteAction::CreateSP(this, &FNiagaraScratchPadCommandContext::CanDeleteSelectedScripts)));
+}
+
+bool FNiagaraScratchPadCommandContext::CanCutSelectedScripts() const
+{
+	return false;
+}
+
+FText FNiagaraScratchPadCommandContext::GetCanCutSelectedScriptsToolTip() const
+{
+	return FText();
+}
+
+void FNiagaraScratchPadCommandContext::CutSelectedScripts() const
+{
+}
+
+bool FNiagaraScratchPadCommandContext::CanCopySelectedScripts() const
+{
+	return false;
+}
+
+FText FNiagaraScratchPadCommandContext::GetCanCopySelectedScriptsToolTip() const
+{
+	return FText();
+}
+
+void FNiagaraScratchPadCommandContext::CopySelectedScripts() const
+{
+}
+
+bool FNiagaraScratchPadCommandContext::CanPasteSelectedScripts() const
+{
+	return false;
+}
+
+FText FNiagaraScratchPadCommandContext::GetCanPasteSelectedScriptsToolTip() const
+{
+	return FText();
+}
+
+void FNiagaraScratchPadCommandContext::PasteSelectedScripts() const
+{
+}
+
+bool FNiagaraScratchPadCommandContext::CanDeleteSelectedScripts() const
+{
+	return true;
+}
+
+FText FNiagaraScratchPadCommandContext::GetCanDeleteSelectedScriptsToolTip() const
+{
+	return LOCTEXT("DeleteScriptMessage", "Delete this script and reset references to it.");
+}
+
+void FNiagaraScratchPadCommandContext::DeleteSelectedScripts() const
+{
+	ScratchPadViewModel->DeleteActiveScript();
+}
+
+#undef LOCTEXT_NAMESPACE
+
diff --git a/Engine/Plugins/FX/Niagara/Source/NiagaraEditorWidgets/Private/NiagaraScratchPadCommandContext.h b/Engine/Plugins/FX/Niagara/Source/NiagaraEditorWidgets/Private/NiagaraScratchPadCommandContext.h
new file mode 100644
--- /dev/null
+++ b/Engine/Plugins/FX/Niagara/Source/NiagaraEditorWidgets/Private/NiagaraScratchPadCommandContext.h
@@ -0,0 +1,58 @@
+// Copyright Epic Games, Inc. All Rights Reserved.
+
+#pragma once
+
+#include "Templates/SharedPointer.h"
+
+class UNiagaraScratchPadViewModel;
+class FUICommandList;
+class FMenuBuilder;
+struct FKeyEvent;
+
+class FNiagaraScratchPadCommandContext : public TSharedFromThis<FNiagaraScratchPadCommandContext>
+{
+public:
+	FNiagaraScratchPadCommandContext(UNiagaraScratchPadViewModel* InScratchPadViewModel);
+
+	TSharedRef<FUICommandList> GetCommands();
+
+	bool ProcessCommandBindings(const FKeyEvent& InKeyEvent);
+
+	bool AddEditMenuItems(FMenuBuilder& MenuBuilder);
+
+private:
+	void SetupCommands();
+
+	bool CanCutSelectedScripts() const;
+
+	FText GetCanCutSelectedScriptsToolTip() const;
+
+	void CutSelectedScripts() const;
+
+	bool CanCopySelectedScripts() const;
+
+	FText GetCanCopySelectedScriptsToolTip() const;
+
+	void CopySelectedScripts() const;
+
+	bool CanPasteSelectedScripts() const;
+
+	FText GetCanPasteSelectedScriptsToolTip() const;
+
+	void PasteSelectedScripts() const;
+
+	bool CanDeleteSelectedScripts() const;
+
+	FText GetCanDeleteSelectedScriptsToolTip() const;
+
+	void DeleteSelectedScripts() const;
+
+private:
+	TSharedRef<FUICommandList> Commands;
+
+	bool bCommandsAreSetup;
+
+	bool bProcessingCommandBindings;
+
+	UNiagaraScratchPadViewModel* ScratchPadViewModel;
+};
\ No newline at end of file
diff --git a/Engine/Plugins/FX/Niagara/Source/NiagaraEditorWidgets/Private/NiagaraStackCommandContext.cpp b/Engine/Plugins/FX/Niagara/Source/NiagaraEditorWidgets/Private/NiagaraStackCommandContext.cpp
--- a/Engine/Plugins/FX/Niagara/Source/NiagaraEditorWidgets/Private/NiagaraStackCommandContext.cpp
+++ b/Engine/Plugins/FX/Niagara/Source/NiagaraEditorWidgets/Private/NiagaraStackCommandContext.cpp
@@ -2,6 +2,7 @@
 
 #include "NiagaraStackCommandContext.h"
 #include "NiagaraEditorCommon.h"
+#include "NiagaraEditorUtilities.h"
 #include "ViewModels/Stack/NiagaraStackClipboardUtilities.h"
 #include "ViewModels/Stack/NiagaraStackEntry.h"
 
@@ -14,18 +15,6 @@
 
 #define LOCTEXT_NAMESPACE "NiagaraStackCommandContext"
 
-void WarnWithToastAndLog(FText WarningMessage)
-{
-	FText IncompleteDeleteMessage = LOCTEXT("DeleteIncompleteMessage", "Not all items could be deleted because they either\ndon't support being deleted or they are inherited.");
-	FNotificationInfo WarningNotification(WarningMessage);
-	WarningNotification.ExpireDuration = 5.0f;
-	WarningNotification.bFireAndForget = true;
-	WarningNotification.bUseLargeFont = false;
-	WarningNotification.Image = FCoreStyle::Get().GetBrush(TEXT("MessageLog.Warning"));
-	FSlateNotificationManager::Get().AddNotification(WarningNotification);
-	UE_LOG(LogNiagaraEditor, Warning, TEXT("%s"), *WarningMessage.ToString());
-}
-
 FNiagaraStackCommandContext::FNiagaraStackCommandContext()
 	: Commands(MakeShared<FUICommandList>())
 	, bCommandsAreSetup(false)
@@ -122,7 +111,7 @@ bool FNiagaraStackCommandContext::CanCutSelectedEntries() const
 	bool bCanCut = FNiagaraStackClipboardUtilities::TestCanCutSelectionWithMessage(SelectedEntries, CanCutMessage);
 	if (bProcessingCommandBindings && bCanCut == false && CanCutMessage.IsEmptyOrWhitespace() == false)
 	{
-		WarnWithToastAndLog(CanCutMessage);
+		FNiagaraEditorUtilities::WarnWithToastAndLog(CanCutMessage);
 	}
 	return bCanCut;
 }
@@ -145,7 +134,7 @@ bool FNiagaraStackCommandContext::CanCopySelectedEntries() const
 	bool bCanCopy = FNiagaraStackClipboardUtilities::TestCanCopySelectionWithMessage(SelectedEntries, CanCopyMessage);
 	if (bProcessingCommandBindings && bCanCopy == false && CanCopyMessage.IsEmptyOrWhitespace() == false)
 	{
-		WarnWithToastAndLog(CanCopyMessage);
+		FNiagaraEditorUtilities::WarnWithToastAndLog(CanCopyMessage);
 	}
 	return bCanCopy;
 }
@@ -168,7 +157,7 @@ bool FNiagaraStackCommandContext::CanPasteSelectedEntries() const
 	bool bCanPaste = FNiagaraStackClipboardUtilities::TestCanPasteSelectionWithMessage(SelectedEntries, CanPasteMessage);
 	if (bProcessingCommandBindings && bCanPaste == false && CanPasteMessage.IsEmptyOrWhitespace() == false)
 	{
-		WarnWithToastAndLog(CanPasteMessage);
+		FNiagaraEditorUtilities::WarnWithToastAndLog(CanPasteMessage);
 	}
 	return bCanPaste;
 }
@@ -186,7 +175,7 @@ void FNiagaraStackCommandContext::PasteSelectedEntries() const
 	FNiagaraStackClipboardUtilities::PasteSelection(SelectedEntries, PasteWarning);
 	if (PasteWarning.IsEmptyOrWhitespace() == false)
 	{
-		WarnWithToastAndLog(PasteWarning);
+		FNiagaraEditorUtilities::WarnWithToastAndLog(PasteWarning);
 	}
 }
 
@@ -196,7 +185,7 @@ bool FNiagaraStackCommandContext::CanDeleteSelectedEntries() const
 	bool bCanDelete = FNiagaraStackClipboardUtilities::TestCanDeleteSelectionWithMessage(SelectedEntries, CanDeleteMessage);
 	if (bProcessingCommandBindings && bCanDelete == false && CanDeleteMessage.IsEmptyOrWhitespace() == false)
 	{
-		WarnWithToastAndLog(CanDeleteMessage);
+		FNiagaraEditorUtilities::WarnWithToastAndLog(CanDeleteMessage);
 	}
 	return bCanDelete;
 }
diff --git a/Engine/Plugins/FX/Niagara/Source/NiagaraEditorWidgets/Private/SNiagaraScratchPad.cpp b/Engine/Plugins/FX/Niagara/Source/NiagaraEditorWidgets/Private/SNiagaraScratchPad.cpp
new file mode 100644
--- /dev/null
+++ b/Engine/Plugins/FX/Niagara/Source/NiagaraEditorWidgets/Private/SNiagaraScratchPad.cpp
@@ -0,0 +1,600 @@
+// Copyright Epic Games, Inc. All Rights Reserved.
+
+#include "SNiagaraScratchPad.h"
+#include "ViewModels/NiagaraScratchPadViewModel.h"
+#include "ViewModels/NiagaraScratchPadScriptViewModel.h"
+#include "ViewModels/Stack/NiagaraStackFunctionInputCollection.h"
+#include "Widgets/SDynamicLayoutBox.h"
+#include "Widgets/SItemSelector.h"
+#include "Widgets/SNiagaraSelectedObjectsDetails.h"
+#include "Widgets/SVerticalResizeBox.h"
+#include "Widgets/SNiagaraScriptGraph.h"
+#include "NiagaraObjectSelection.h"
+#include "NiagaraNodeFunctionCall.h"
+#include "NiagaraEditorWidgetsStyle.h"
+#include "NiagaraScratchPadCommandContext.h"
+
+#include "Widgets/Text/SInlineEditableTextBlock.h"
+#include "Widgets/Input/SButton.h"
+#include "EditorStyleSet.h"
+#include "Framework/MultiBox/MultiBoxBuilder.h"
+#include "Framework/Application/SlateApplication.h"
+#include "Layout/WidgetPath.h"
+
+#define LOCTEXT_NAMESPACE "NiagaraScratchPad"
+
+FName ScriptSelectorName = "ScriptSelector";
+FName ScriptEditorName = "ScriptEditor";
+FName SelectionEditorName = "SelectionEditor";
+FName WideLayoutName = "Wide";
+FName NarrowLayoutName = "Narrow";
+
+class SNiagaraScratchPadScriptRow : public SCompoundWidget
+{
+	SLATE_BEGIN_ARGS(SNiagaraScratchPadScriptRow) {}
+		SLATE_ATTRIBUTE(bool, IsSelected);
+	SLATE_END_ARGS()
+
+	void Construct(
+		const FArguments& InArgs,
+		UNiagaraScratchPadViewModel* InScratchPadViewModel,
+		TSharedRef<FNiagaraScratchPadScriptViewModel> InScriptViewModel,
+		TSharedPtr<FNiagaraScratchPadCommandContext> InCommandContext)
+	{
+		ScratchPadViewModel = InScratchPadViewModel;
+		ScriptViewModel = InScriptViewModel;
+		CommandContext = InCommandContext;
+		IsSelected = InArgs._IsSelected;
+
+		ChildSlot
+		[
+			SAssignNew(NameEditableText, SInlineEditableTextBlock)
+			.Text(this, &SNiagaraScratchPadScriptRow::GetNameText)
+			.IsSelected(this, &SNiagaraScratchPadScriptRow::GetIsSelected)
+			.OnTextCommitted(this, &SNiagaraScratchPadScriptRow::OnNameTextCommitted)
+		];
+	}
+
+	virtual void Tick(const FGeometry& AllottedGeometry, const double InCurrentTime, const float InDeltaTime) override
+	{
+		if (ScriptViewModel->GetIsPendingRename())
+		{
+			ScriptViewModel->SetIsPendingRename(false);
+			NameEditableText->EnterEditingMode();
+		}
+	}
+
+	virtual FReply OnMouseButtonDown(const FGeometry& MyGeometry, const FPointerEvent& MouseEvent) override
+	{
+		if (MouseEvent.GetEffectingButton() == EKeys::RightMouseButton)
+		{
+			return FReply::Handled()
+				.CaptureMouse(SharedThis(this));
+		}
+		return SCompoundWidget::OnMouseButtonDown(MyGeometry, MouseEvent);
+	}
+
+	virtual FReply OnMouseButtonUp(const FGeometry& MyGeometry, const FPointerEvent& MouseEvent) override
+	{
+		if (MouseEvent.GetEffectingButton() == EKeys::RightMouseButton)
+		{
+			// Set this script to be the active one.
+			ScratchPadViewModel->SetActiveScript(ScriptViewModel->GetScript());
+
+			FMenuBuilder MenuBuilder(true, CommandContext->GetCommands());
+			CommandContext->AddEditMenuItems(MenuBuilder);
+
+			FWidgetPath WidgetPath = MouseEvent.GetEventPath() != nullptr ? *MouseEvent.GetEventPath() : FWidgetPath();
+			FSlateApplication::Get().PushMenu(AsShared(), WidgetPath, MenuBuilder.MakeWidget(), MouseEvent.GetScreenSpacePosition(), FPopupTransitionEffect(FPopupTransitionEffect::ContextMenu));
+			return FReply::Handled().ReleaseMouseCapture();
+		}
+		return SCompoundWidget::OnMouseButtonUp(MyGeometry, MouseEvent);
+	}
+
+private:
+	FText GetNameText() const
+	{
+		return ScriptViewModel->GetDisplayName();
+	}
+
+	bool GetIsSelected() const
+	{
+		return IsSelected.Get(false);
+	}
+
+	void OnNameTextCommitted(const FText& InText, ETextCommit::Type InCommitType)
+	{
+		ScriptViewModel->SetScriptName(InText);
+	}
+
+private:
+	UNiagaraScratchPadViewModel* ScratchPadViewModel;
+	TSharedPtr<FNiagaraScratchPadScriptViewModel> ScriptViewModel;
+	TSharedPtr<FNiagaraScratchPadCommandContext> CommandContext;
+	TAttribute<bool> IsSelected;
+	TSharedPtr<SInlineEditableTextBlock> NameEditableText;
+};
+
+class SNiagaraScratchPadScriptSelector : public SCompoundWidget
+{
+	SLATE_BEGIN_ARGS(SNiagaraScratchPadScriptSelector) {}
+	SLATE_END_ARGS()
+
+	void Construct(const FArguments& InArgs, UNiagaraScratchPadViewModel* InViewModel, TSharedPtr<FNiagaraScratchPadCommandContext> InCommandContext)
+	{
+		ViewModel = InViewModel;
+		CommandContext = InCommandContext;
+		ViewModel->OnScriptViewModelsChanged().AddSP(this, &SNiagaraScratchPadScriptSelector::ScriptViewModelsChanged);
+		ViewModel->OnActiveScriptChanged().AddSP(this, &SNiagaraScratchPadScriptSelector::ActiveScriptChanged);
+		bIsUpdatingSelection = false;
+
+		ChildSlot
+		[
+			SAssignNew(ScriptSelector, SNiagaraScriptViewModelSelector)
+			.ClickActivateMode(EItemSelectorClickActivateMode::SingleClick)
+			.Items(ViewModel->GetScriptViewModels())
+			.DefaultCategories(ViewModel->GetAvailableUsages())
+			.OnGetCategoriesForItem(this, &SNiagaraScratchPadScriptSelector::OnGetCategoriesForItem)
+			.OnCompareCategoriesForEquality(this, &SNiagaraScratchPadScriptSelector::OnCompareCategoriesForEquality)
+			.OnCompareCategoriesForSorting(this, &SNiagaraScratchPadScriptSelector::OnCompareCategoriesForSorting)
+			.OnCompareItemsForEquality(this, &SNiagaraScratchPadScriptSelector::OnCompareItemsForEquality)
+			.OnCompareItemsForSorting(this, &SNiagaraScratchPadScriptSelector::OnCompareItemsForSorting)
+			.OnDoesItemMatchFilterText(this, &SNiagaraScratchPadScriptSelector::OnDoesItemMatchFilterText)
+			.OnGenerateWidgetForCategory(this, &SNiagaraScratchPadScriptSelector::OnGenerateWidgetForCategory)
+			.OnGenerateWidgetForItem(this, &SNiagaraScratchPadScriptSelector::OnGenerateWidgetForItem)
+			.OnItemActivated(this, &SNiagaraScratchPadScriptSelector::OnScriptActivated)
+			.OnSelectionChanged(this, &SNiagaraScratchPadScriptSelector::OnSelectionChanged)
+		];
+	}
+
+	virtual FReply OnKeyDown(const FGeometry& MyGeometry, const FKeyEvent& InKeyEvent) override
+	{
+		if (CommandContext->ProcessCommandBindings(InKeyEvent))
+		{
+			return FReply::Handled();
+		}
+		return SCompoundWidget::OnKeyDown(MyGeometry, InKeyEvent);
+	}
+
+private:
+	void ScriptViewModelsChanged()
+	{
+		if (ScriptSelector.IsValid())
+		{
+			ScriptSelector->RefreshItemsAndDefaultCategories(ViewModel->GetScriptViewModels(), ViewModel->GetAvailableUsages());
+		}
+	}
+
+	void ActiveScriptChanged()
+	{
+		if (bIsUpdatingSelection == false)
+		{
+			TGuardValue<bool> UpdateGuard(bIsUpdatingSelection, true);
+			TSharedPtr<FNiagaraScratchPadScriptViewModel> ActiveScriptViewModel = ViewModel->GetViewModelForScript(ViewModel->GetActiveScript());
+			if (ActiveScriptViewModel.IsValid())
+			{
+				TArray<TSharedRef<FNiagaraScratchPadScriptViewModel>> SelectedItems;
+				SelectedItems.Add(ActiveScriptViewModel.ToSharedRef());
+				ScriptSelector->SetSelectedItems(SelectedItems);
+			}
+			else
+			{
+				ScriptSelector->ClearSelectedItems();
+			}
+		}
+	}
+
+	TArray<ENiagaraScriptUsage> OnGetCategoriesForItem(const TSharedRef<FNiagaraScratchPadScriptViewModel>& Item)
+	{
+		TArray<ENiagaraScriptUsage> Categories;
+		Categories.Add(Item->GetScripts()[0]->GetUsage());
+		return Categories;
+	}
+
+	bool OnCompareCategoriesForEquality(const ENiagaraScriptUsage& CategoryA, const ENiagaraScriptUsage& CategoryB) const
+	{
+		return CategoryA == CategoryB;
+	}
+
+	bool OnCompareCategoriesForSorting(const ENiagaraScriptUsage& CategoryA, const ENiagaraScriptUsage& CategoryB) const
+	{
+		return ((int32)CategoryA) < ((int32)CategoryB);
+	}
+
+	bool OnCompareItemsForEquality(const TSharedRef<FNiagaraScratchPadScriptViewModel>& ItemA, const TSharedRef<FNiagaraScratchPadScriptViewModel>& ItemB) const
+	{
+		return ItemA == ItemB;
+	}
+
+	bool OnCompareItemsForSorting(const TSharedRef<FNiagaraScratchPadScriptViewModel>& ItemA, const TSharedRef<FNiagaraScratchPadScriptViewModel>& ItemB) const
+	{
+		return ItemA->GetDisplayName().CompareTo(ItemB->GetDisplayName()) < 0;
+	}
+
+	bool OnDoesItemMatchFilterText(const FText& FilterText, const TSharedRef<FNiagaraScratchPadScriptViewModel>& Item)
+	{
+		return Item->GetDisplayName().ToString().Find(FilterText.ToString(), ESearchCase::IgnoreCase) != INDEX_NONE;
+	}
+
+	TSharedRef<SWidget> OnGenerateWidgetForCategory(const ENiagaraScriptUsage& Category)
+	{
+		return SNew(SBorder)
+		.BorderImage(FEditorStyle::GetBrush("WhiteBrush"))
+		.BorderBackgroundColor(FNiagaraEditorWidgetsStyle::Get().GetColor("NiagaraEditor.ScratchPad.HeaderColor"))
+		[
+			SNew(SHorizontalBox)
+			+ SHorizontalBox::Slot()
+			.VAlign(VAlign_Center)
+			.Padding(5, 0, 0, 0)
+			[
+				SNew(STextBlock)
+				.TextStyle(FNiagaraEditorWidgetsStyle::Get(), "NiagaraEditor.ScratchPad.SmallHeaderText")
+				.Text(ViewModel->GetDisplayNameForUsage(Category))
+			]
+			+ SHorizontalBox::Slot()
+			.AutoWidth()
+			[
+				SNew(SButton)
+				.OnClicked(this, &SNiagaraScratchPadScriptSelector::ScriptSelectorAddButtonClicked, Category)
+				.Content()
+				[
+					SNew(STextBlock)
+					.Text(LOCTEXT("AddButtonText", "Add"))
+				]
+			]
+		];
+	}
+
+	TSharedRef<SWidget> OnGenerateWidgetForItem(const TSharedRef<FNiagaraScratchPadScriptViewModel>& Item)
+	{
+		return SNew(SNiagaraScratchPadScriptRow, ViewModel, Item, CommandContext)
+		.IsSelected(this, &SNiagaraScratchPadScriptSelector::GetItemIsSelected, TWeakPtr<FNiagaraScratchPadScriptViewModel>(Item));
+	}
+
+	void OnScriptActivated(const TSharedRef<FNiagaraScratchPadScriptViewModel>& ActivatedScript)
+	{
+		if (bIsUpdatingSelection == false)
+		{
+			TGuardValue<bool> UpdateGuard(bIsUpdatingSelection, true);
+			ViewModel->SetActiveScript(ActivatedScript->GetScript());
+		}
+	}
+
+	void OnSelectionChanged()
+	{
+		if (bIsUpdatingSelection == false)
+		{
+			TGuardValue<bool> UpdateGuard(bIsUpdatingSelection, true);
+			TArray<TSharedRef<FNiagaraScratchPadScriptViewModel>> SelectedScripts = ScriptSelector->GetSelectedItems();
+			if (SelectedScripts.Num() == 0)
+			{
+				ViewModel->SetActiveScript(nullptr);
+			}
+			else if (SelectedScripts.Num())
+			{
+				ViewModel->SetActiveScript(SelectedScripts[0]->GetScript());
+			}
+		}
+	}
+
+	FReply ScriptSelectorAddButtonClicked(ENiagaraScriptUsage Usage)
+	{
+		TSharedPtr<FNiagaraScratchPadScriptViewModel> NewScriptViewModel = ViewModel->CreateNewScript(Usage, ENiagaraScriptUsage::ParticleUpdateScript, FNiagaraTypeDefinition());
+		if (NewScriptViewModel.IsValid())
+		{
+			NewScriptViewModel->SetIsPendingRename(true);
+		}
+		return FReply::Handled();
+	}
+
+	bool GetItemIsSelected(TWeakPtr<FNiagaraScratchPadScriptViewModel> ItemWeak) const
+	{
+		TSharedPtr<FNiagaraScratchPadScriptViewModel> Item = ItemWeak.Pin();
+		return Item.IsValid() && ViewModel->GetActiveScript() == Item->GetScripts()[0];
+	}
+
+private:
+	TSharedPtr<SNiagaraScriptViewModelSelector> ScriptSelector;
+	UNiagaraScratchPadViewModel* ViewModel;
+	TSharedPtr<FNiagaraScratchPadCommandContext> CommandContext;
+	bool bIsUpdatingSelection;
+};
+
+class SNiagaraScratchPadScriptEditor : public SCompoundWidget
+{
+	SLATE_BEGIN_ARGS(SNiagaraScratchPadScriptEditor) {}
+	SLATE_END_ARGS()
+
+	void Construct(const FArguments& InArgs, TSharedRef<FNiagaraScratchPadScriptViewModel> InScriptViewModel)
+	{
+		ScriptViewModel = InScriptViewModel;
+		ChildSlot
+		[
+			SNew(SVerticalBox)
+			+ SVerticalBox::Slot()
+			.AutoHeight()
+			[
+				SNew(SHorizontalBox)
+				+ SHorizontalBox::Slot()
+				.VAlign(VAlign_Center)
+				[
+					SNew(STextBlock)
+					.TextStyle(FNiagaraEditorWidgetsStyle::Get(), "NiagaraEditor.ScratchPad.LargeHeaderText")
+					.Text(this, &SNiagaraScratchPadScriptEditor::GetNameText)
+				]
+				+ SHorizontalBox::Slot()
+				.AutoWidth()
+				[
+					SNew(SButton)
+					.OnClicked(this, &SNiagaraScratchPadScriptEditor::OnRefreshButtonClicked)
+					.ToolTipText(LOCTEXT("RefreshButtonToolTip", "Refresh the inputs for this script in the selection stack."))
+					.Content()
+					[
+						SNew(STextBlock)
+						.Text(LOCTEXT("RefreshButtonText", "Refresh"))
+					]
+				]
+			]
+			+ SVerticalBox::Slot()
+			[
+				SNew(SNiagaraScriptGraph, ScriptViewModel->GetGraphViewModel())
+				.ZoomToFitOnLoad(true)
+			]
+		];
+	}
+
+private:
+	FText GetNameText() const
+	{
+		return ScriptViewModel->GetDisplayName();
+	}
+
+	FReply OnRefreshButtonClicked()
+	{
+		TArray<UNiagaraNodeFunctionCall*> FunctionCallNodesToRefresh;
+		for (TObjectIterator<UNiagaraNodeFunctionCall> It; It; ++It)
+		{
+			UNiagaraNodeFunctionCall* FunctionCallNode = *It;
+			if (FunctionCallNode->FunctionScript == ScriptViewModel->GetScript())
+			{
+				FunctionCallNodesToRefresh.Add(FunctionCallNode);
+			}
+		}
+
+		if (FunctionCallNodesToRefresh.Num())
+		{
+			for (TObjectIterator<UNiagaraStackFunctionInputCollection> It; It; ++It)
+			{
+				UNiagaraStackFunctionInputCollection* StackFunctionInputCollection = *It;
+				if (StackFunctionInputCollection->IsFinalized() == false && FunctionCallNodesToRefresh.Contains(StackFunctionInputCollection->GetInputFunctionCallNode()))
+				{
+					StackFunctionInputCollection->RefreshChildren();
+				}
+			}
+		}
+
+		return FReply::Handled();
+	}
+
+private:
+	TSharedPtr<FNiagaraScratchPadScriptViewModel> ScriptViewModel;
+};
+
+class SNiagaraScratchPadScriptEditorList : public SCompoundWidget
+{
+	SLATE_BEGIN_ARGS(SNiagaraScratchPadScriptEditor) {}
+
+	SLATE_END_ARGS();
+
+	void Construct(const FArguments& InArgs, UNiagaraScratchPadViewModel* InViewModel)
+	{
+		ViewModel = InViewModel;
+		ViewModel->OnScriptViewModelsChanged().AddSP(this, &SNiagaraScratchPadScriptEditorList::ScriptViewModelsChanged);
+		ViewModel->OnActiveScriptChanged().AddSP(this, &SNiagaraScratchPadScriptEditorList::ActiveScriptChanged);
+		bIsUpdatingSelection = false;
+
+		ChildSlot
+		[
+			SAssignNew(ScriptEditors, SListView<TSharedRef<FNiagaraScratchPadScriptViewModel>>)
+			.ListItemsSource(&ViewModel->GetScriptViewModels())
+			.OnGenerateRow(this, &SNiagaraScratchPadScriptEditorList::OnGenerateScriptEditorRow)
+			.OnSelectionChanged(this, &SNiagaraScratchPadScriptEditorList::OnSelectionChanged)
+		];
+	}
+
+private:
+	void ScriptViewModelsChanged()
+	{
+		if (ScriptEditors.IsValid())
+		{
+			ScriptEditors->RequestListRefresh();
+		}
+	}
+
+	void ActiveScriptChanged()
+	{
+		if (bIsUpdatingSelection == false)
+		{
+			TGuardValue<bool> UpdateGuard(bIsUpdatingSelection, true);
+			TSharedPtr<FNiagaraScratchPadScriptViewModel> ActiveScriptViewModel = ViewModel->GetViewModelForScript(ViewModel->GetActiveScript());
+			if (ActiveScriptViewModel.IsValid())
+			{
+				ScriptEditors->SetSelection(ActiveScriptViewModel.ToSharedRef());
+			}
+			else
+			{
+				ScriptEditors->ClearSelection();
+			}
+		}
+	}
+
+	TSharedRef<ITableRow> OnGenerateScriptEditorRow(TSharedRef<FNiagaraScratchPadScriptViewModel> Item, const TSharedRef<STableViewBase>& OwnerTable)
+	{
+		return SNew(STableRow<TSharedRef<FNiagaraScratchPadScriptViewModel>>, OwnerTable)
+		[
+			SNew(SVerticalResizeBox)
+			.ContentHeight(300)
+			[
+				SNew(SNiagaraScratchPadScriptEditor, Item)
+			]
+		];
+	}
+
+	void OnSelectionChanged(TSharedPtr<FNiagaraScratchPadScriptViewModel> InNewSelection, ESelectInfo::Type SelectInfo)
+	{
+		if (bIsUpdatingSelection == false)
+		{
+			TGuardValue<bool> UpdateGuard(bIsUpdatingSelection, true);
+			TArray<TSharedRef<FNiagaraScratchPadScriptViewModel>> SelectedScripts;
+			ScriptEditors->GetSelectedItems(SelectedScripts);
+			if (SelectedScripts.Num() == 0)
+			{
+				ViewModel->SetActiveScript(nullptr);
+			}
+			else if (SelectedScripts.Num())
+			{
+				ViewModel->SetActiveScript(SelectedScripts[0]->GetScript());
+			}
+		}
+	}
+
+private:
+	UNiagaraScratchPadViewModel* ViewModel;
+	TSharedPtr<SListView<TSharedRef<FNiagaraScratchPadScriptViewModel>>> ScriptEditors;
+	bool bIsUpdatingSelection;
+};
+
+void SNiagaraScratchPad::Construct(const FArguments& InArgs, UNiagaraScratchPadViewModel* InViewModel)
+{
+	ViewModel = InViewModel;
+	CommandContext = MakeShared<FNiagaraScratchPadCommandContext>(InViewModel);
+
+	ChildSlot
+	[
+		SNew(SDynamicLayoutBox)
+		.GenerateNamedWidget_Lambda([this](FName InWidgetName)
+		{
+			if (InWidgetName == ScriptSelectorName)
+			{
+				return ConstructScriptSelector();
+			}
+			else if (InWidgetName == ScriptEditorName)
+			{
+				return ConstructScriptEditor();
+			}
+			else if (InWidgetName == SelectionEditorName)
+			{
+				return ConstructSelectionEditor();
+			}
+			else
+			{
+				return SNullWidget::NullWidget;
+			}
+		})
+		.GenerateNamedLayout_Lambda([this](FName InLayoutName, const SDynamicLayoutBox::FNamedWidgetProvider& InNamedWidgetProvider)
+		{
+			TSharedPtr<SWidget> Layout;
+			if (InLayoutName == WideLayoutName)
+			{
+				Layout = SNew(SSplitter)
+				.Orientation(Orient_Horizontal)
+				+ SSplitter::Slot()
+				.Value(0.15f)
+				[
+					InNamedWidgetProvider.GetNamedWidget(ScriptSelectorName)
+				]
+				+ SSplitter::Slot()
+				.Value(0.6f)
+				[
+					InNamedWidgetProvider.GetNamedWidget(ScriptEditorName)
+				]
+				+ SSplitter::Slot()
+				.Value(0.25f)
+				[
+					InNamedWidgetProvider.GetNamedWidget(SelectionEditorName)
+				];
+			}
+			else if (InLayoutName == NarrowLayoutName)
+			{
+				Layout = SNew(SSplitter)
+				.Orientation(Orient_Horizontal)
+				+ SSplitter::Slot()
+				.Value(0.3f)
+				[
+					SNew(SSplitter)
+					.Orientation(Orient_Vertical)
+					+ SSplitter::Slot()
+					.Value(0.5f)
+					[
+						InNamedWidgetProvider.GetNamedWidget(ScriptSelectorName)
+					]
+					+ SSplitter::Slot()
+					.Value(0.5f)
+					[
+						InNamedWidgetProvider.GetNamedWidget(SelectionEditorName)
+					]
+				]
+				+ SSplitter::Slot()
+				.Value(0.7f)
+				[
+					InNamedWidgetProvider.GetNamedWidget(ScriptEditorName)
+				];
+			}
+			else
+			{
+				Layout = SNullWidget::NullWidget;
+			}
+			return Layout.ToSharedRef();
+		})
+		.ChooseLayout_Lambda([this]() 
+		{ 
+			if (GetCachedGeometry().GetLocalSize().X < 1500)
+			{
+				return NarrowLayoutName;
+			}
+			else
+			{
+				return WideLayoutName;
+			}
+		})
+	];
+}
+
+TSharedRef<SWidget> SNiagaraScratchPad::ConstructScriptSelector()
+{
+	return SNew(SVerticalBox)
+	+ SVerticalBox::Slot()
+	.AutoHeight()
+	[
+		SNew(STextBlock)
+		.TextStyle(FNiagaraEditorWidgetsStyle::Get(), "NiagaraEditor.ScratchPad.LargeHeaderText")
+		.Text(LOCTEXT("ScriptSelector", "Scratch Script Selector"))
+	]
+	+ SVerticalBox::Slot()
+	[
+		SNew(SNiagaraScratchPadScriptSelector, ViewModel.Get(), CommandContext)
+	];
+}
+
+TSharedRef<SWidget> SNiagaraScratchPad::ConstructScriptEditor()
+{
+	return SNew(SNiagaraScratchPadScriptEditorList, ViewModel.Get());
+}
+
+TSharedRef<SWidget> SNiagaraScratchPad::ConstructSelectionEditor()
+{
+	return SNew(SVerticalBox)
+	+ SVerticalBox::Slot()
+	.AutoHeight()
+	[
+		SNew(STextBlock)
+		.TextStyle(FNiagaraEditorWidgetsStyle::Get(), "NiagaraEditor.ScratchPad.LargeHeaderText")
+		.Text(LOCTEXT("ScratchPadSelection", "Scratch Pad Selection"))
+	]
+	+ SVerticalBox::Slot()
+	[
+		SNew(SNiagaraSelectedObjectsDetails, ViewModel->GetObjectSelection())
+	];
+}
+
+#undef LOCTEXT_NAMESPACE
\ No newline at end of file
diff --git a/Engine/Plugins/FX/Niagara/Source/NiagaraEditorWidgets/Private/SNiagaraScratchPad.h b/Engine/Plugins/FX/Niagara/Source/NiagaraEditorWidgets/Private/SNiagaraScratchPad.h
new file mode 100644
--- /dev/null
+++ b/Engine/Plugins/FX/Niagara/Source/NiagaraEditorWidgets/Private/SNiagaraScratchPad.h
@@ -0,0 +1,33 @@
+// Copyright Epic Games, Inc. All Rights Reserved.
+
+#pragma once
+
+#include "NiagaraCommon.h"
+#include "Widgets/SItemSelector.h"
+#include "Widgets/SCompoundWidget.h"
+
+class UNiagaraScratchPadViewModel;
+class FNiagaraScratchPadScriptViewModel;
+class FNiagaraScratchPadCommandContext;
+
+typedef SItemSelector<ENiagaraScriptUsage, TSharedRef<FNiagaraScratchPadScriptViewModel>> SNiagaraScriptViewModelSelector;
+
+class SNiagaraScratchPad : public SCompoundWidget
+{
+	SLATE_BEGIN_ARGS(SNiagaraScratchPad) {}
+	SLATE_END_ARGS()
+
+	void Construct(const FArguments& InArgs, UNiagaraScratchPadViewModel* InViewModel);
+
+private:
+	TSharedRef<SWidget> ConstructScriptSelector();
+
+	TSharedRef<SWidget> ConstructScriptEditor();
+
+	TSharedRef<SWidget> ConstructSelectionEditor();
+
+private:
+	TWeakObjectPtr<UNiagaraScratchPadViewModel> ViewModel;
+
+	TSharedPtr<FNiagaraScratchPadCommandContext> CommandContext;
+};
\ No newline at end of file
diff --git a/Engine/Plugins/FX/Niagara/Source/NiagaraEditorWidgets/Private/Stack/SNiagaraStackFunctionInputValue.cpp b/Engine/Plugins/FX/Niagara/Source/NiagaraEditorWidgets/Private/Stack/SNiagaraStackFunctionInputValue.cpp
--- a/Engine/Plugins/FX/Niagara/Source/NiagaraEditorWidgets/Private/Stack/SNiagaraStackFunctionInputValue.cpp
+++ b/Engine/Plugins/FX/Niagara/Source/NiagaraEditorWidgets/Private/Stack/SNiagaraStackFunctionInputValue.cpp
@@ -535,8 +535,8 @@ void SNiagaraStackFunctionInputValue::CollectAllActions(FGraphActionListBuilderB
 		FunctionInput->GetAvailableDynamicInputs(DynamicInputScripts, bIncludeNonLibraryInputs);
 		for (UNiagaraScript* DynamicInputScript : DynamicInputScripts)
 		{
-			const FText DynamicInputText = FText::FromString(FName::NameToDisplayString(DynamicInputScript->GetName() + (DynamicInputScript->bExposeToLibrary ? "" : "*"), false));
-			const FText Tooltip = FNiagaraEditorUtilities::FormatScriptAssetDescription(DynamicInputScript->Description, *(DynamicInputScript->GetPathName() + (DynamicInputScript->bExposeToLibrary ? "" : "\n*Not exposed to library")));
+			const FText DynamicInputText = FNiagaraEditorUtilities::FormatScriptName(DynamicInputScript->GetFName(), DynamicInputScript->bExposeToLibrary);
+			const FText Tooltip = FNiagaraEditorUtilities::FormatScriptDescription(DynamicInputScript->Description, *DynamicInputScript->GetPathName(), DynamicInputScript->bExposeToLibrary);
 			TSharedPtr<FNiagaraMenuAction> DynamicInputAction(new FNiagaraMenuAction(CategoryName, DynamicInputText, Tooltip, 0, DynamicInputScript->Keywords,
 				FNiagaraMenuAction::FOnExecuteStackAction::CreateSP(this, &SNiagaraStackFunctionInputValue::DynamicInputScriptSelected, DynamicInputScript)));
 
@@ -653,6 +653,15 @@ void SNiagaraStackFunctionInputValue::CollectAllActions(FGraphActionListBuilderB
 		OutAllActions.AddAction(ExpressionAction);
 	}
 
+	{
+		const FText CategoryName = LOCTEXT("ScratchCategory", "Scratch");
+		const FText CreateDisplayName = LOCTEXT("ScratchLabel", "Make New Scratch Dynamic Input");
+		const FText CreateTooltip = LOCTEXT("ScratchToolTipl", "Create a new dynamic input in the scratch pad.");
+		TSharedPtr<FNiagaraMenuAction> CreateScratchAction(new FNiagaraMenuAction(CategoryName, CreateDisplayName, CreateTooltip, 0, FText(),
+			FNiagaraMenuAction::FOnExecuteStackAction::CreateSP(this, &SNiagaraStackFunctionInputValue::CreateScratchSelected)));
+		OutAllActions.AddAction(CreateScratchAction);
+	}
+
 	if (FunctionInput->CanDeleteInput())
 	{
 		const FText NameText = LOCTEXT("DeleteInput", "Remove");
@@ -691,6 +700,11 @@ void SNiagaraStackFunctionInputValue::CustomExpressionSelected()
 	FunctionInput->SetCustomExpression(TEXT("// Insert expression here"));
 }
 
+void SNiagaraStackFunctionInputValue::CreateScratchSelected()
+{
+	FunctionInput->SetScratch();
+}
+
 void SNiagaraStackFunctionInputValue::ParameterHandleSelected(FNiagaraParameterHandle Handle)
 {
 	FunctionInput->SetLinkedValueHandle(Handle);
@@ -852,8 +866,8 @@ void CollectDynamicInputActions(FGraphActionListBuilderBase& DynamicInputActions
 	FunctionInput->GetAvailableDynamicInputs(DynamicInputScripts);
 	for (UNiagaraScript* DynamicInputScript : DynamicInputScripts)
 	{
-		const FText DynamicInputText = FText::FromString(FName::NameToDisplayString(DynamicInputScript->GetName(), false));
-		const FText Tooltip = FNiagaraEditorUtilities::FormatScriptAssetDescription(DynamicInputScript->Description, *DynamicInputScript->GetPathName());
+		const FText DynamicInputText = FNiagaraEditorUtilities::FormatScriptName(DynamicInputScript->GetFName(), DynamicInputScript->bExposeToLibrary);
+		const FText Tooltip = FNiagaraEditorUtilities::FormatScriptDescription(DynamicInputScript->Description, *DynamicInputScript->GetPathName(), DynamicInputScript->bExposeToLibrary);
 		TSharedPtr<FNiagaraMenuAction> DynamicInputAction(new FNiagaraMenuAction(CategoryName, DynamicInputText, Tooltip, 0, DynamicInputScript->Keywords,
 			FNiagaraMenuAction::FOnExecuteStackAction::CreateStatic(&ReassignDynamicInputScript, FunctionInput, DynamicInputScript)));
 		DynamicInputActions.AddAction(DynamicInputAction);
diff --git a/Engine/Plugins/FX/Niagara/Source/NiagaraEditorWidgets/Private/Stack/SNiagaraStackFunctionInputValue.h b/Engine/Plugins/FX/Niagara/Source/NiagaraEditorWidgets/Private/Stack/SNiagaraStackFunctionInputValue.h
--- a/Engine/Plugins/FX/Niagara/Source/NiagaraEditorWidgets/Private/Stack/SNiagaraStackFunctionInputValue.h
+++ b/Engine/Plugins/FX/Niagara/Source/NiagaraEditorWidgets/Private/Stack/SNiagaraStackFunctionInputValue.h
@@ -106,6 +106,8 @@ private:
 
 	void CustomExpressionSelected();
 
+	void CreateScratchSelected();
+
 	void ParameterHandleSelected(FNiagaraParameterHandle Handle);
 
 	EVisibility GetResetButtonVisibility() const;
diff --git a/Engine/Plugins/FX/Niagara/Source/NiagaraEditorWidgets/Private/Stack/SNiagaraStackModuleItem.cpp b/Engine/Plugins/FX/Niagara/Source/NiagaraEditorWidgets/Private/Stack/SNiagaraStackModuleItem.cpp
--- a/Engine/Plugins/FX/Niagara/Source/NiagaraEditorWidgets/Private/Stack/SNiagaraStackModuleItem.cpp
+++ b/Engine/Plugins/FX/Niagara/Source/NiagaraEditorWidgets/Private/Stack/SNiagaraStackModuleItem.cpp
@@ -334,12 +334,13 @@ void CollectModuleActions(FGraphActionListBuilderBase& ModuleActions, UNiagaraSt
 			Category = LOCTEXT("ModuleNotCategorized", "Uncategorized Modules");
 		}
 
-		FString DisplayNameString = FName::NameToDisplayString(ModuleAsset.AssetName.ToString(), false);
-		FText DisplayName = FText::FromString(DisplayNameString);
+		bool bIsInLibrary = FNiagaraEditorUtilities::IsScriptAssetInLibrary(ModuleAsset);
+
+		FText DisplayName = FNiagaraEditorUtilities::FormatScriptName(ModuleAsset.AssetName, bIsInLibrary);
 
 		FText AssetDescription;
 		ModuleAsset.GetTagValue(GET_MEMBER_NAME_CHECKED(UNiagaraScript, Description), AssetDescription);
-		FText Description = FNiagaraEditorUtilities::FormatScriptAssetDescription(AssetDescription, ModuleAsset.ObjectPath);
+		FText Description = FNiagaraEditorUtilities::FormatScriptDescription(AssetDescription, ModuleAsset.ObjectPath, bIsInLibrary);
 
 		FText Keywords;
 		ModuleAsset.GetTagValue(GET_MEMBER_NAME_CHECKED(UNiagaraScript, Keywords), Keywords);
diff --git a/Engine/Plugins/FX/Niagara/Source/NiagaraEditorWidgets/Public/NiagaraEditorWidgetsModule.h b/Engine/Plugins/FX/Niagara/Source/NiagaraEditorWidgets/Public/NiagaraEditorWidgetsModule.h
--- a/Engine/Plugins/FX/Niagara/Source/NiagaraEditorWidgets/Public/NiagaraEditorWidgetsModule.h
+++ b/Engine/Plugins/FX/Niagara/Source/NiagaraEditorWidgets/Public/NiagaraEditorWidgetsModule.h
@@ -54,6 +54,7 @@ private:
 		virtual TSharedRef<SWidget> CreateStackView(UNiagaraStackViewModel& StackViewModel) const override;
 		virtual TSharedRef<SWidget> CreateSystemOverview(TSharedRef<FNiagaraSystemViewModel> SystemViewModel) const override;
 		virtual TSharedRef<SWidget> CreateStackIssueIcon(UNiagaraStackViewModel& StackViewModel, UNiagaraStackEntry& StackEntry) const override;
+		virtual TSharedRef<SWidget> CreateScriptScratchPad(UNiagaraScratchPadViewModel& ScriptScratchPadViewModel) const override;
 		virtual FLinearColor GetColorForExecutionCategory(FName ExecutionCategory) const override;
 	};
 
