Niagara - Fix copyrights.

diff --git a/Engine/Plugins/FX/Niagara/Source/NiagaraEditor/Private/Widgets/SDynamicLayoutBox.cpp b/Engine/Plugins/FX/Niagara/Source/NiagaraEditor/Private/Widgets/SDynamicLayoutBox.cpp
--- a/Engine/Plugins/FX/Niagara/Source/NiagaraEditor/Private/Widgets/SDynamicLayoutBox.cpp
+++ b/Engine/Plugins/FX/Niagara/Source/NiagaraEditor/Private/Widgets/SDynamicLayoutBox.cpp
@@ -1,4 +1,4 @@
-// Copyright 1998-2020 Epic Games, Inc. All Rights Reserved.
+// Copyright Epic Games, Inc. All Rights Reserved.
 
 #include "Widgets/SDynamicLayoutBox.h"
 
diff --git a/Engine/Plugins/FX/Niagara/Source/NiagaraEditor/Public/Widgets/SDynamicLayoutBox.h b/Engine/Plugins/FX/Niagara/Source/NiagaraEditor/Public/Widgets/SDynamicLayoutBox.h
--- a/Engine/Plugins/FX/Niagara/Source/NiagaraEditor/Public/Widgets/SDynamicLayoutBox.h
+++ b/Engine/Plugins/FX/Niagara/Source/NiagaraEditor/Public/Widgets/SDynamicLayoutBox.h
@@ -1,4 +1,4 @@
-// Copyright 1998-2020 Epic Games, Inc. All Rights Reserved.
+// Copyright Epic Games, Inc. All Rights Reserved.
 
 #pragma once
 
