Fix deprecation warning
#rnx

diff --git a/Engine/Plugins/FX/Niagara/Source/NiagaraEditor/Private/ViewModels/NiagaraScratchPadViewModel.cpp b/Engine/Plugins/FX/Niagara/Source/NiagaraEditor/Private/ViewModels/NiagaraScratchPadViewModel.cpp
--- a/Engine/Plugins/FX/Niagara/Source/NiagaraEditor/Private/ViewModels/NiagaraScratchPadViewModel.cpp
+++ b/Engine/Plugins/FX/Niagara/Source/NiagaraEditor/Private/ViewModels/NiagaraScratchPadViewModel.cpp
@@ -207,7 +207,7 @@ TSharedPtr<FNiagaraScratchPadScriptViewModel> UNiagaraScratchPadViewModel::Creat
 					UNiagaraNodeOutput* DynamicInputOutputNode = OutputNodes[0];
 					DynamicInputOutputNode->Outputs.Empty();
 					DynamicInputOutputNode->Outputs.Add(FNiagaraVariable(InOutputType, "Output"));
-					FPropertyChangedEvent OutputsChangedEvent(FindFieldChecked<UProperty>(UNiagaraNodeOutput::StaticClass(), GET_MEMBER_NAME_CHECKED(UNiagaraNodeOutput, Outputs)));
+					FPropertyChangedEvent OutputsChangedEvent(FindFieldChecked<FProperty>(UNiagaraNodeOutput::StaticClass(), GET_MEMBER_NAME_CHECKED(UNiagaraNodeOutput, Outputs)));
 					DynamicInputOutputNode->PostEditChangeProperty(OutputsChangedEvent);
 				}
 			}
