Consistent variable naming.

diff --git a/Engine/Plugins/FX/Niagara/Source/Niagara/Classes/NiagaraDataSet.h b/Engine/Plugins/FX/Niagara/Source/Niagara/Classes/NiagaraDataSet.h
--- a/Engine/Plugins/FX/Niagara/Source/Niagara/Classes/NiagaraDataSet.h
+++ b/Engine/Plugins/FX/Niagara/Source/Niagara/Classes/NiagaraDataSet.h
@@ -246,7 +246,7 @@ struct NIAGARA_API FNiagaraDataSetCompiledData
 
 	/** Whether or not this dataset require persistent IDs. */
 	UPROPERTY()
-	uint32 bNeedsPersistentIDs : 1;
+	uint32 bRequiresPersistentIDs : 1;
 
 	/** Unique ID for this DataSet. Used to allow referencing from other emitters and Systems. */
 	UPROPERTY()
@@ -302,10 +302,10 @@ public:
 	/** Returns size in bytes for all data buffers currently allocated by this dataset. */
 	uint32 GetSizeBytes()const;
 
-	FORCEINLINE bool IsInitialized()const { return bInitialized; }
-	FORCEINLINE ENiagaraSimTarget GetSimTarget()const { return CompiledData.SimTarget; }
-	FORCEINLINE FNiagaraDataSetID GetID()const { return CompiledData.ID; }	
-	FORCEINLINE bool GetNeedsPersistentIDs()const { return CompiledData.bNeedsPersistentIDs; }
+	FORCEINLINE bool IsInitialized() const { return bInitialized; }
+	FORCEINLINE ENiagaraSimTarget GetSimTarget() const { return CompiledData.SimTarget; }
+	FORCEINLINE FNiagaraDataSetID GetID() const { return CompiledData.ID; }	
+	FORCEINLINE bool RequiresPersistentIDs() const { return CompiledData.bRequiresPersistentIDs; }
 
 	FORCEINLINE TArray<int32>& GetFreeIDTable() { return FreeIDsTable; }
 	FORCEINLINE TArray<int32>& GetSpawnedIDsTable() { return SpawnedIDsTable; }
diff --git a/Engine/Plugins/FX/Niagara/Source/Niagara/Classes/NiagaraEmitter.h b/Engine/Plugins/FX/Niagara/Source/Niagara/Classes/NiagaraEmitter.h
--- a/Engine/Plugins/FX/Niagara/Source/Niagara/Classes/NiagaraEmitter.h
+++ b/Engine/Plugins/FX/Niagara/Source/Niagara/Classes/NiagaraEmitter.h
@@ -438,7 +438,7 @@ public:
 	/** Returns true if this emitter's platform filter allows it on this platform and EffectsQuality level. */
 	NIAGARA_API bool IsAllowedByScalability()const;
 
-	NIAGARA_API bool RequiresPersistantIDs()const;
+	NIAGARA_API bool RequiresPersistentIDs() const;
 
 	NIAGARA_API bool IsValid()const;
 	NIAGARA_API bool IsReadyToRun() const;
diff --git a/Engine/Plugins/FX/Niagara/Source/Niagara/Classes/NiagaraEmitterInstance.h b/Engine/Plugins/FX/Niagara/Source/Niagara/Classes/NiagaraEmitterInstance.h
--- a/Engine/Plugins/FX/Niagara/Source/Niagara/Classes/NiagaraEmitterInstance.h
+++ b/Engine/Plugins/FX/Niagara/Source/Niagara/Classes/NiagaraEmitterInstance.h
@@ -52,7 +52,7 @@ public:
 	void PostTick();
 	bool HandleCompletion(bool bForce = false);
 
-	bool RequiredPersistentID()const;
+	bool RequiresPersistentIDs() const;
 
 	FORCEINLINE bool ShouldTick()const { return ExecutionState == ENiagaraExecutionState::Active || GetNumParticles() > 0; }
 
diff --git a/Engine/Plugins/FX/Niagara/Source/Niagara/Private/NiagaraDataSet.cpp b/Engine/Plugins/FX/Niagara/Source/Niagara/Private/NiagaraDataSet.cpp
--- a/Engine/Plugins/FX/Niagara/Source/Niagara/Private/NiagaraDataSet.cpp
+++ b/Engine/Plugins/FX/Niagara/Source/Niagara/Private/NiagaraDataSet.cpp
@@ -217,7 +217,7 @@ void FNiagaraDataSet::Allocate(int32 NumInstances, bool bMaintainExisting)
 	CheckForNaNs();
 #endif
 
-	if (GetNeedsPersistentIDs())
+	if (RequiresPersistentIDs())
 	{
 		TArray<int32>& CurrentIDTable = CurrentData->GetIDTable();
 		TArray<int32>& DestinationIDTable = DestinationData->GetIDTable();
@@ -349,7 +349,7 @@ void FNiagaraDataSet::ReleaseGPUInstanceCounts(FNiagaraGPUInstanceCountManager&
 
 void FNiagaraDataSet::AllocateGPUFreeIDs(uint32 InNumInstances, FRHICommandList& RHICmdList, ERHIFeatureLevel::Type FeatureLevel, const TCHAR* DebugSimName)
 {
-	checkSlow(GetSimTarget() == ENiagaraSimTarget::GPUComputeSim && GetNeedsPersistentIDs());
+	checkSlow(GetSimTarget() == ENiagaraSimTarget::GPUComputeSim && RequiresPersistentIDs());
 
 	// Clearing and compacting the ID table must run over all the allocated elements, so we must use a chunk size which balances
 	// between reallocation frequency and the cost of processing unused elements.
@@ -714,7 +714,7 @@ void FNiagaraDataBuffer::AllocateGPU(uint32 InNumInstances, FNiagaraGPUInstanceC
 			}
 		}
 
-		if (Owner->GetNeedsPersistentIDs())
+		if (Owner->RequiresPersistentIDs())
 		{
 			uint32 NumExistingElems = GPUIDToIndexTable.Buffer ? GPUIDToIndexTable.Buffer->GetSize() / sizeof(int32) : 0;
 			uint32 NumNeededElems = Owner->GetGPUNumAllocatedIDs();
@@ -1144,7 +1144,7 @@ FNiagaraDataSetCompiledData::FNiagaraDataSetCompiledData()
 
 void FNiagaraDataSetCompiledData::Empty()
 {
-	bNeedsPersistentIDs = 0;
+	bRequiresPersistentIDs = 0;
 	TotalFloatComponents = 0;
 	TotalInt32Components = 0;
 	Variables.Empty();
diff --git a/Engine/Plugins/FX/Niagara/Source/Niagara/Private/NiagaraEmitter.cpp b/Engine/Plugins/FX/Niagara/Source/Niagara/Private/NiagaraEmitter.cpp
--- a/Engine/Plugins/FX/Niagara/Source/Niagara/Private/NiagaraEmitter.cpp
+++ b/Engine/Plugins/FX/Niagara/Source/Niagara/Private/NiagaraEmitter.cpp
@@ -819,7 +819,7 @@ bool UNiagaraEmitter::IsAllowedByScalability()const
 	return Platforms.IsActive();
 }
 
-bool UNiagaraEmitter::RequiresPersistantIDs()const
+bool UNiagaraEmitter::RequiresPersistentIDs() const
 {
 	return bRequiresPersistentIDs;
 }
diff --git a/Engine/Plugins/FX/Niagara/Source/Niagara/Private/NiagaraEmitterInstance.cpp b/Engine/Plugins/FX/Niagara/Source/Niagara/Private/NiagaraEmitterInstance.cpp
--- a/Engine/Plugins/FX/Niagara/Source/Niagara/Private/NiagaraEmitterInstance.cpp
+++ b/Engine/Plugins/FX/Niagara/Source/Niagara/Private/NiagaraEmitterInstance.cpp
@@ -899,10 +899,10 @@ bool FNiagaraEmitterInstance::HandleCompletion(bool bForce)
 	return false;
 }
 
-bool FNiagaraEmitterInstance::RequiredPersistentID()const
+bool FNiagaraEmitterInstance::RequiresPersistentIDs() const
 {
 	//TODO: can we have this be enabled at runtime from outside the system?
-	return GetEmitterHandle().GetInstance()->RequiresPersistantIDs() || ParticleDataSet->HasVariable(SYS_PARAM_PARTICLES_ID);
+	return GetEmitterHandle().GetInstance()->RequiresPersistentIDs() || ParticleDataSet->HasVariable(SYS_PARAM_PARTICLES_ID);
 }
 
 /** 
diff --git a/Engine/Plugins/FX/Niagara/Source/Niagara/Private/NiagaraEmitterInstanceBatcher.cpp b/Engine/Plugins/FX/Niagara/Source/Niagara/Private/NiagaraEmitterInstanceBatcher.cpp
--- a/Engine/Plugins/FX/Niagara/Source/Niagara/Private/NiagaraEmitterInstanceBatcher.cpp
+++ b/Engine/Plugins/FX/Niagara/Source/Niagara/Private/NiagaraEmitterInstanceBatcher.cpp
@@ -388,7 +388,7 @@ void NiagaraEmitterInstanceBatcher::ResizeBuffersAndGatherResources(FOverlappabl
 				continue;
 			}
 
-			const bool bNeedsPersistentIDs = Context->MainDataSet->GetNeedsPersistentIDs();
+			const bool bRequiresPersistentIDs = Context->MainDataSet->RequiresPersistentIDs();
 
 			//The buffer containing current simulation state.
 			Instance.CurrentData = Context->MainDataSet->GetCurrentData();
@@ -407,7 +407,7 @@ void NiagaraEmitterInstanceBatcher::ResizeBuffersAndGatherResources(FOverlappabl
 			const uint32 RequiredInstances = FMath::Max(PrevNumInstances, NewNumInstances);
 			const uint32 AllocatedInstances = FMath::Max(RequiredInstances, Instance.SpawnInfo.MaxParticleCount);
 
-			if (bNeedsPersistentIDs)
+			if (bRequiresPersistentIDs)
 			{
 				Context->MainDataSet->AllocateGPUFreeIDs(AllocatedInstances + 1, RHICmdList, FeatureLevel, Context->GetDebugSimName());
 				ReadBuffers.Add(Context->MainDataSet->GetGPUFreeIDs().UAV);
@@ -443,7 +443,7 @@ void NiagaraEmitterInstanceBatcher::ResizeBuffersAndGatherResources(FOverlappabl
 				}
 			}
 
-			if (bNeedsPersistentIDs)
+			if (bRequiresPersistentIDs)
 			{
 				WriteBuffers.Add(DestinationData.GetGPUIDToIndexTable().UAV);
 			}
@@ -1208,8 +1208,8 @@ void NiagaraEmitterInstanceBatcher::Run(const FNiagaraGPUSystemTick& Tick, const
 
 	// set the shader and data set params 
 	//
-	const bool bNeedsPersistentIDs = Context->MainDataSet->GetNeedsPersistentIDs();
-	SetSRVParameter(RHICmdList, Shader.GetComputeShader(), Shader->FreeIDBufferParam, bNeedsPersistentIDs ? Context->MainDataSet->GetGPUFreeIDs().SRV : FNiagaraRenderer::GetDummyIntBuffer().SRV);
+	const bool bRequiresPersistentIDs = Context->MainDataSet->RequiresPersistentIDs();
+	SetSRVParameter(RHICmdList, Shader.GetComputeShader(), Shader->FreeIDBufferParam, bRequiresPersistentIDs ? Context->MainDataSet->GetGPUFreeIDs().SRV : FNiagaraRenderer::GetDummyIntBuffer().SRV);
 	CurrentData.SetShaderParams(Shader.GetShader(), RHICmdList, true);
 	DestinationData.SetShaderParams(Shader.GetShader(), RHICmdList, false);
 
diff --git a/Engine/Plugins/FX/Niagara/Source/Niagara/Private/NiagaraScript.cpp b/Engine/Plugins/FX/Niagara/Source/Niagara/Private/NiagaraScript.cpp
--- a/Engine/Plugins/FX/Niagara/Source/Niagara/Private/NiagaraScript.cpp
+++ b/Engine/Plugins/FX/Niagara/Source/Niagara/Private/NiagaraScript.cpp
@@ -330,7 +330,7 @@ void UNiagaraScript::ComputeVMCompilationId(FNiagaraVMExecutableDataId& Id) cons
 		{
 			SimTargetToBuild = Emitter->SimTarget;
 		}
-		if (Emitter->RequiresPersistantIDs())
+		if (Emitter->RequiresPersistentIDs())
 		{
 			Id.bRequiresPersistentIDs = true;
 			Id.AdditionalDefines.Add(TEXT("RequiresPersistentIDs"));
diff --git a/Engine/Plugins/FX/Niagara/Source/Niagara/Private/NiagaraSystem.cpp b/Engine/Plugins/FX/Niagara/Source/Niagara/Private/NiagaraSystem.cpp
--- a/Engine/Plugins/FX/Niagara/Source/Niagara/Private/NiagaraSystem.cpp
+++ b/Engine/Plugins/FX/Niagara/Source/Niagara/Private/NiagaraSystem.cpp
@@ -1136,7 +1136,7 @@ void UNiagaraSystem::InitEmitterDataSetCompiledData(FNiagaraDataSetCompiledData&
 		DataSetToInit.Variables.AddUnique(Var);
 	}
 
-	DataSetToInit.bNeedsPersistentIDs = InAssociatedEmitter->RequiresPersistantIDs() || DataSetToInit.Variables.Contains(SYS_PARAM_PARTICLES_ID);
+	DataSetToInit.bRequiresPersistentIDs = InAssociatedEmitter->RequiresPersistentIDs() || DataSetToInit.Variables.Contains(SYS_PARAM_PARTICLES_ID);
 	DataSetToInit.ID = FNiagaraDataSetID(InAssociatedEmitterHandle.GetIdName(), ENiagaraDataSetType::ParticleData);
 	DataSetToInit.SimTarget = InAssociatedEmitter->SimTarget;
 
@@ -1358,7 +1358,7 @@ void UNiagaraSystem::InitSystemCompiledData()
 			CompiledData.Variables.AddUnique(Var);
 		}
 
-		CompiledData.bNeedsPersistentIDs = false;
+		CompiledData.bRequiresPersistentIDs = false;
 		CompiledData.ID = FNiagaraDataSetID();
 		CompiledData.SimTarget = ENiagaraSimTarget::CPUSim;
 
diff --git a/Engine/Plugins/FX/Niagara/Source/NiagaraEditor/Private/NiagaraHlslTranslator.cpp b/Engine/Plugins/FX/Niagara/Source/NiagaraEditor/Private/NiagaraHlslTranslator.cpp
--- a/Engine/Plugins/FX/Niagara/Source/NiagaraEditor/Private/NiagaraHlslTranslator.cpp
+++ b/Engine/Plugins/FX/Niagara/Source/NiagaraEditor/Private/NiagaraHlslTranslator.cpp
@@ -823,7 +823,7 @@ const FNiagaraTranslateResults &FHlslNiagaraTranslator::Translate(const FNiagara
 		return TranslateResults;
 	}
 
-	bool bNeedsPersistentIDs = CompileOptions.AdditionalDefines.Contains(TEXT("RequiresPersistentIDs"));
+	bool bRequiresPersistentIDs = CompileOptions.AdditionalDefines.Contains(TEXT("RequiresPersistentIDs"));
 
 	TranslationStages.Empty();
 	ActiveStageIdx = 0;
@@ -928,7 +928,7 @@ const FNiagaraTranslateResults &FHlslNiagaraTranslator::Translate(const FNiagara
 				UNiagaraNodeOutput* TargetOutputNode = TranslationStages[ParamMapIdx].OutputNode;
 				if (FoundHistory.GetFinalOutputNode() == TargetOutputNode)
 				{
-					if (bNeedsPersistentIDs)
+					if (bRequiresPersistentIDs)
 					{
 						//TODO: Setup alias for current level to decouple from "Particles". Would we ever want emitter or system persistent IDs?
 						FNiagaraVariable Var = FNiagaraVariable(FNiagaraTypeDefinition::GetIDDef(), TEXT("Particles.ID"));
@@ -1878,7 +1878,7 @@ void FHlslNiagaraTranslator::DefineMainGPUFunctions(
 		return false;
 	}();
 
-	const bool bNeedsPersistentIDs = CompileOptions.AdditionalDefines.Contains(TEXT("RequiresPersistentIDs"));
+	const bool bRequiresPersistentIDs = CompileOptions.AdditionalDefines.Contains(TEXT("RequiresPersistentIDs"));
 
 	// A list of constant to reset after Emitter_SpawnGroup gets modified by GetEmitterSpawnInfoForParticle()
 	TArray<FString> EmitterSpawnGroupReinit;
@@ -1937,7 +1937,7 @@ void FHlslNiagaraTranslator::DefineMainGPUFunctions(
 			HlslOutput += TEXT("\n") + ContextName + TEXT("DataInstance.Alive=true;\n");
 		}
 
-		if (bNeedsPersistentIDs)
+		if (bRequiresPersistentIDs)
 		{
 			HlslOutput += TEXT("\n\tint IDIndex, IDAcquireTag;\n\tAcquireID(0, IDIndex, IDAcquireTag);\n");
 			HlslOutput += ContextName + TEXT("Particles.ID.Index = IDIndex;\n");
@@ -2048,7 +2048,7 @@ void FHlslNiagaraTranslator::DefineMainGPUFunctions(
 
 		HlslOutput += TEXT("\tif (bValid)\n\t{\n");
 
-		if (bNeedsPersistentIDs)
+		if (bRequiresPersistentIDs)
 		{
 			HlslOutput += FString::Printf(TEXT("\t\tUpdateID(0, %sParticles.ID.Index, WriteIndex);\n"), *ContextName);
 		}
@@ -2136,8 +2136,8 @@ void FHlslNiagaraTranslator::DefineMain(FString &OutHlslOutput,
 		DefineDataSetVariableReads(HlslOutput, DataSetID, VarArrayIdx, ArrayRef);
 	}
 
-	bool bNeedsPersistentIDs = CompileOptions.AdditionalDefines.Contains(TEXT("RequiresPersistentIDs"));
-	if (bNeedsPersistentIDs && UNiagaraScript::IsSpawnScript(CompileOptions.TargetUsage))
+	bool bRequiresPersistentIDs = CompileOptions.AdditionalDefines.Contains(TEXT("RequiresPersistentIDs"));
+	if (bRequiresPersistentIDs && UNiagaraScript::IsSpawnScript(CompileOptions.TargetUsage))
 	{
 		FString MapName = UNiagaraScript::IsInterpolatedParticleSpawnScript(CompileOptions.TargetUsage) ? TEXT("Context.MapSpawn") : TEXT("Context.Map");
 		//Add code to handle persistent IDs.
@@ -2286,8 +2286,8 @@ void FHlslNiagaraTranslator::DefineDataSetVariableWrites(FString &OutHlslOutput,
 		OutHlslOutput += "\tint TmpWriteIndex = OutputIndex(0, true, true);\n";
 	}
 
-	bool bNeedsPersistentIDs = CompileOptions.AdditionalDefines.Contains(TEXT("RequiresPersistentIDs"));
-	if (bNeedsPersistentIDs && DataSetIndex == 0)
+	bool bRequiresPersistentIDs = CompileOptions.AdditionalDefines.Contains(TEXT("RequiresPersistentIDs"));
+	if (bRequiresPersistentIDs && DataSetIndex == 0)
 	{
 		FString MapName = GetParameterMapInstanceName(0);
 		OutHlslOutput += FString::Printf(TEXT("\tUpdateID(0, %s.Particles.ID.Index, TmpWriteIndex);\n"), *MapName);
