Fix format string error in assert.

diff --git a/Engine/Plugins/FX/Niagara/Source/NiagaraShader/Private/NiagaraShader.cpp b/Engine/Plugins/FX/Niagara/Source/NiagaraShader/Private/NiagaraShader.cpp
--- a/Engine/Plugins/FX/Niagara/Source/NiagaraShader/Private/NiagaraShader.cpp
+++ b/Engine/Plugins/FX/Niagara/Source/NiagaraShader/Private/NiagaraShader.cpp
@@ -1391,7 +1391,7 @@ void FNiagaraDataInterfaceParamRef::Bind(const FNiagaraDataInterfaceGPUParamInfo
 	{
 		DIType = TIndexedPtr<UNiagaraDataInterfaceBase>(Base); // TODO - clean up TIndexedPtr::operator=()
 		Parameters = Base->ConstructComputeParameters();
-		checkf(Parameters != nullptr, TEXT("Niagara type '%s' is missing ComputeParameters, missing a IMPLEMENT_NIAGARA_DI_PARAMETER?"));
+		checkf(Parameters != nullptr, TEXT("Niagara type '%s' is missing ComputeParameters, missing a IMPLEMENT_NIAGARA_DI_PARAMETER?"), *InParameterInfo.DIClassName);
 		Parameters->DIType = DIType;
 		Base->BindParameters(Parameters, InParameterInfo, ParameterMap);
 	}
