Updated the particle read DI to work with the memory freezing stuff

diff --git a/Engine/Plugins/FX/Niagara/Source/Niagara/Classes/NiagaraDataInterfaceParticleRead.h b/Engine/Plugins/FX/Niagara/Source/Niagara/Classes/NiagaraDataInterfaceParticleRead.h
--- a/Engine/Plugins/FX/Niagara/Source/Niagara/Classes/NiagaraDataInterfaceParticleRead.h
+++ b/Engine/Plugins/FX/Niagara/Source/Niagara/Classes/NiagaraDataInterfaceParticleRead.h
@@ -12,6 +12,8 @@ class NIAGARA_API UNiagaraDataInterfaceParticleRead : public UNiagaraDataInterfa
 {
 	GENERATED_UCLASS_BODY()
 public:
+	DECLARE_NIAGARA_DI_PARAMETER();
+
 	UPROPERTY(EditAnywhere, Category = "ParticleRead")
 	FString EmitterName;
 
diff --git a/Engine/Plugins/FX/Niagara/Source/Niagara/Private/NiagaraDataInterfaceParticleRead.cpp b/Engine/Plugins/FX/Niagara/Source/Niagara/Private/NiagaraDataInterfaceParticleRead.cpp
--- a/Engine/Plugins/FX/Niagara/Source/Niagara/Private/NiagaraDataInterfaceParticleRead.cpp
+++ b/Engine/Plugins/FX/Niagara/Source/Niagara/Private/NiagaraDataInterfaceParticleRead.cpp
@@ -50,6 +50,8 @@ enum class ENiagaraParticleDataValueType : uint8
 	Quat
 };
 
+DECLARE_INTRINSIC_TYPE_LAYOUT(ENiagaraParticleDataValueType);
+
 static const TCHAR* NiagaraParticleDataValueTypeName(ENiagaraParticleDataValueType Type)
 {
 	switch (Type)
@@ -131,7 +133,10 @@ private:
 
 struct FNiagaraDataInterfaceParametersCS_ParticleRead : public FNiagaraDataInterfaceParametersCS
 {
+	DECLARE_TYPE_LAYOUT(FNiagaraDataInterfaceParametersCS_ParticleRead, NonVirtual);
+
 	FNiagaraDataInterfaceParametersCS_ParticleRead() : 
+		AcquireTagRegisterIndex(-1),
 		CachedDataSet(nullptr),
 		bSourceEmitterNotGPUErrorShown(false)
 	{
@@ -150,7 +155,7 @@ struct FNiagaraDataInterfaceParametersCS_ParticleRead : public FNiagaraDataInter
 		return ENiagaraParticleDataValueType::Invalid;
 	}
 
-	void Bind(const FNiagaraDataInterfaceGPUParamInfo& ParameterInfo, const class FShaderParameterMap& ParameterMap)
+	void Bind(const FNiagaraDataInterfaceGPUParamInfo& ParameterInfo, const FShaderParameterMap& ParameterMap)
 	{
 		NumSpawnedParticlesParam.Bind(ParameterMap, *(NumSpawnedParticlesBaseName + ParameterInfo.DataInterfaceHLSLSymbol));
 		SpawnedParticlesAcquireTagParam.Bind(ParameterMap, *(SpawnedParticlesAcquireTagBaseName + ParameterInfo.DataInterfaceHLSLSymbol));
@@ -184,9 +189,6 @@ struct FNiagaraDataInterfaceParametersCS_ParticleRead : public FNiagaraDataInter
 				AttributeTypes[FuncIdx] = ENiagaraParticleDataValueType::Invalid;
 			}
 		}
-
-		int NumAttrIndices = Align(AttributeNames.Num(), 4);
-		AttributeIndices.SetNum(NumAttrIndices);
 	}
 
 	void SetErrorParams(FRHICommandList& RHICmdList, FRHIComputeShader* ComputeShader, bool bSkipSpawnInfo) const
@@ -233,7 +235,7 @@ struct FNiagaraDataInterfaceParametersCS_ParticleRead : public FNiagaraDataInter
 
 	void FindAttributeIndices(const FNiagaraDataSet* SourceDataSet, const TCHAR* SourceEmitterName) const
 	{
-		check(AttributeIndices.Num() == Align(AttributeNames.Num(), 4));
+		check(AttributeNames.Num() == AttributeTypes.Num());
 
 		const TArray<FNiagaraVariable>& SourceEmitterVariables = SourceDataSet->GetVariables();
 		const TArray<FNiagaraVariableLayoutInfo>& SourceEmitterVariableLayouts = SourceDataSet->GetVariableLayouts();
@@ -303,6 +305,12 @@ struct FNiagaraDataInterfaceParametersCS_ParticleRead : public FNiagaraDataInter
 	{
 		check(IsInRenderingThread());
 
+		if (AttributeIndices.Num() == 0)
+		{
+			int NumAttrIndices = Align(AttributeNames.Num(), 4);
+			AttributeIndices.SetNum(NumAttrIndices);
+		}
+
 		FRHIComputeShader* ComputeShader = RHICmdList.GetBoundComputeShader();
 
 		FNiagaraDataInterfaceProxyParticleRead* Proxy = static_cast<FNiagaraDataInterfaceProxyParticleRead*>(Context.DataInterface);
@@ -409,24 +417,26 @@ struct FNiagaraDataInterfaceParametersCS_ParticleRead : public FNiagaraDataInter
 	}
 	
 private:
-	FShaderParameter NumSpawnedParticlesParam;
-	FShaderParameter SpawnedParticlesAcquireTagParam;
-	FShaderResourceParameter SpawnedIDsBufferParam;
-	FShaderResourceParameter IDToIndexTableParam;
-	FShaderResourceParameter InputFloatBufferParam;
-	FShaderResourceParameter InputIntBufferParam;
-	FShaderParameter ParticleStrideFloatParam;
-	FShaderParameter ParticleStrideIntParam;
-	FShaderParameter AttributeIndicesParam;
-	FShaderParameter AcquireTagRegisterIndexParam;
-	TArray<FName> AttributeNames;
-	TArray<ENiagaraParticleDataValueType> AttributeTypes;
-	mutable TArray<int32> AttributeIndices;
-	mutable int32 AcquireTagRegisterIndex;
-	mutable FNiagaraDataSet* CachedDataSet;
-	mutable bool bSourceEmitterNotGPUErrorShown;
+	LAYOUT_FIELD(FShaderParameter, NumSpawnedParticlesParam);
+	LAYOUT_FIELD(FShaderParameter, SpawnedParticlesAcquireTagParam);
+	LAYOUT_FIELD(FShaderResourceParameter, SpawnedIDsBufferParam);
+	LAYOUT_FIELD(FShaderResourceParameter, IDToIndexTableParam);
+	LAYOUT_FIELD(FShaderResourceParameter, InputFloatBufferParam);
+	LAYOUT_FIELD(FShaderResourceParameter, InputIntBufferParam);
+	LAYOUT_FIELD(FShaderParameter, ParticleStrideFloatParam);
+	LAYOUT_FIELD(FShaderParameter, ParticleStrideIntParam);
+	LAYOUT_FIELD(FShaderParameter, AttributeIndicesParam);
+	LAYOUT_FIELD(FShaderParameter, AcquireTagRegisterIndexParam);
+	LAYOUT_FIELD(TMemoryImageArray<FName>, AttributeNames);
+	LAYOUT_FIELD(TMemoryImageArray<ENiagaraParticleDataValueType>, AttributeTypes);
+	LAYOUT_MUTABLE_FIELD(TArray<int32>, AttributeIndices);
+	LAYOUT_MUTABLE_FIELD_INITIALIZED(int32, AcquireTagRegisterIndex, -1);
+	LAYOUT_MUTABLE_FIELD_INITIALIZED(FNiagaraDataSet*, CachedDataSet, nullptr);
+	LAYOUT_MUTABLE_FIELD_INITIALIZED(bool, bSourceEmitterNotGPUErrorShown, false);
 };
 
+IMPLEMENT_TYPE_LAYOUT(FNiagaraDataInterfaceParametersCS_ParticleRead);
+
 UNiagaraDataInterfaceParticleRead::UNiagaraDataInterfaceParticleRead(FObjectInitializer const& ObjectInitializer)
 	: Super(ObjectInitializer)
 {
@@ -1314,4 +1324,6 @@ void UNiagaraDataInterfaceParticleRead::GetEmitterDependencies(void* PerInstance
 	}
 }
 
-#undef LOCTEXT_NAMESPACE
\ No newline at end of file
+IMPLEMENT_NIAGARA_DI_PARAMETER(UNiagaraDataInterfaceParticleRead, FNiagaraDataInterfaceParametersCS_ParticleRead);
+
+#undef LOCTEXT_NAMESPACE
diff --git a/Engine/Plugins/FX/Niagara/Source/NiagaraCore/Private/NiagaraCustomVersion.cpp b/Engine/Plugins/FX/Niagara/Source/NiagaraCore/Private/NiagaraCustomVersion.cpp
--- a/Engine/Plugins/FX/Niagara/Source/NiagaraCore/Private/NiagaraCustomVersion.cpp
+++ b/Engine/Plugins/FX/Niagara/Source/NiagaraCore/Private/NiagaraCustomVersion.cpp
@@ -8,4 +8,4 @@ const FGuid FNiagaraCustomVersion::GUID(0xFCF57AFA, 0x50764283, 0xB9A9E658, 0xFF
 // Register the custom version with core
 FCustomVersionRegistration GRegisterNiagaraCustomVersion(FNiagaraCustomVersion::GUID, FNiagaraCustomVersion::LatestVersion, TEXT("NiagaraVer"));
 
-const FGuid FNiagaraCustomVersion::LatestScriptCompileVersion(0x0288EF7C, 0x881D4F16, 0xB4A7513C, 0xFA3EECED);
+const FGuid FNiagaraCustomVersion::LatestScriptCompileVersion(0x4A797DE3, 0x02B740CA, 0x9E719DF4, 0x53470E2A);
diff --git a/Engine/Plugins/FX/Niagara/Source/NiagaraShader/Private/NiagaraShaderParticleID.cpp b/Engine/Plugins/FX/Niagara/Source/NiagaraShader/Private/NiagaraShaderParticleID.cpp
--- a/Engine/Plugins/FX/Niagara/Source/NiagaraShader/Private/NiagaraShaderParticleID.cpp
+++ b/Engine/Plugins/FX/Niagara/Source/NiagaraShader/Private/NiagaraShaderParticleID.cpp
@@ -30,7 +30,7 @@ public:
 		OutEnvironment.SetDefine(TEXT("THREAD_COUNT"), THREAD_COUNT);
 	}
 
-	void Execute(FRHICommandList& RHICmdList, uint32 NumElementsToAlloc, FRWBuffer& NewBuffer, uint32 NumExistingElements, FRHIShaderResourceView* ExistingBuffer)
+	void Execute(FRHICommandList& RHICmdList, FRHIComputeShader* ComputeShader, uint32 NumElementsToAlloc, FRWBuffer& NewBuffer, uint32 NumExistingElements, FRHIShaderResourceView* ExistingBuffer)
 	{
 		// To simplify the shader code, the size of the ID table must be a multiple of the thread count.
 		check(NumElementsToAlloc % THREAD_COUNT == 0);
@@ -39,7 +39,6 @@ public:
 		check(NumElementsToAlloc >= NumExistingElements);
 		uint32 NumNewElements = NumElementsToAlloc - NumExistingElements;
 
-		FRHIComputeShader* ComputeShader = RHICmdList.GetBoundComputeShader();
 		RHICmdList.SetComputeShader(ComputeShader);
 
 		NewBufferParam.SetBuffer(RHICmdList, ComputeShader, NewBuffer);
@@ -67,7 +66,8 @@ IMPLEMENT_GLOBAL_SHADER(FNiagaraInitFreeIDBufferCS, "/Plugin/FX/Niagara/Private/
 void NiagaraInitGPUFreeIDList(FRHICommandList& RHICmdList, ERHIFeatureLevel::Type FeatureLevel, uint32 NumElementsToAlloc, FRWBuffer& NewBuffer, uint32 NumExistingElements, FRHIShaderResourceView* ExistingBuffer)
 {
 	TShaderMapRef<FNiagaraInitFreeIDBufferCS> InitIDBufferCS(GetGlobalShaderMap(FeatureLevel));
-	InitIDBufferCS->Execute(RHICmdList, NumElementsToAlloc, NewBuffer, NumExistingElements, ExistingBuffer);
+	FRHIComputeShader* ComputeShader = InitIDBufferCS.GetComputeShader();
+	InitIDBufferCS->Execute(RHICmdList, ComputeShader, NumElementsToAlloc, NewBuffer, NumExistingElements, ExistingBuffer);
 }
 
 class NiagaraComputeFreeIDsCS : public FGlobalShader
@@ -95,7 +95,7 @@ public:
 		FGlobalShader::ModifyCompilationEnvironment(Parameters, OutEnvironment);
 	}
 
-	void Execute(FRHICommandList& RHICmdList, ERHIFeatureLevel::Type FeatureLevel, uint32 NumIDs, FRHIShaderResourceView* IDToIndexTable, FRWBuffer& FreeIDList, FRWBuffer& FreeIDListSizes, uint32 FreeIDListIndex)
+	void Execute(FRHICommandList& RHICmdList, ERHIFeatureLevel::Type FeatureLevel, FRHIComputeShader* ComputeShader, uint32 NumIDs, FRHIShaderResourceView* IDToIndexTable, FRWBuffer& FreeIDList, FRWBuffer& FreeIDListSizes, uint32 FreeIDListIndex)
 	{
 		const EShaderPlatform Platform = GShaderPlatformForFeatureLevel[FeatureLevel];
 		const uint32 THREAD_COUNT = (Platform == SP_XBOXONE_D3D12 || Platform == SP_PS4) ? 64 : 128;
@@ -103,7 +103,6 @@ public:
 		// To simplify the shader code, the size of the ID table must be a multiple of the thread count.
 		check(NumIDs % THREAD_COUNT == 0);
 
-		FRHIComputeShader* ComputeShader = RHICmdList.GetBoundComputeShader();
 		RHICmdList.SetComputeShader(ComputeShader);
 
 		SetSRVParameter(RHICmdList, ComputeShader, IDToIndexTableParam, IDToIndexTable);
@@ -130,7 +129,8 @@ IMPLEMENT_GLOBAL_SHADER(NiagaraComputeFreeIDsCS, "/Plugin/FX/Niagara/Private/Nia
 void NiagaraComputeGPUFreeIDs(FRHICommandList& RHICmdList, ERHIFeatureLevel::Type FeatureLevel, uint32 NumIDs, FRHIShaderResourceView* IDToIndexTable, FRWBuffer& FreeIDList, FRWBuffer& FreeIDListSizes, uint32 FreeIDListIndex)
 {
 	TShaderMapRef<NiagaraComputeFreeIDsCS> ComputeFreeIDsCS(GetGlobalShaderMap(FeatureLevel));
-	ComputeFreeIDsCS->Execute(RHICmdList, FeatureLevel, NumIDs, IDToIndexTable, FreeIDList, FreeIDListSizes, FreeIDListIndex);
+	FRHIComputeShader* ComputeShader = ComputeFreeIDsCS.GetComputeShader();
+	ComputeFreeIDsCS->Execute(RHICmdList, FeatureLevel, ComputeShader, NumIDs, IDToIndexTable, FreeIDList, FreeIDListSizes, FreeIDListIndex);
 }
 
 class NiagaraFillIntBufferCS : public FGlobalShader
@@ -152,13 +152,12 @@ public:
 		return RHISupportsComputeShaders(Parameters.Platform);
 	}
 
-	void Execute(FRHICommandList& RHICmdList, FRWBuffer& Buffer, int Value)
+	void Execute(FRHICommandList& RHICmdList, FRHIComputeShader* ComputeShader, FRWBuffer& Buffer, int Value)
 	{
 		const uint32 THREAD_COUNT = 64;
 		const uint32 NumInts = Buffer.NumBytes / sizeof(int);
 		const uint32 ThreadGroups = FMath::DivideAndRoundUp(NumInts, THREAD_COUNT);
 
-		FRHIComputeShader* ComputeShader = RHICmdList.GetBoundComputeShader();
 		RHICmdList.SetComputeShader(ComputeShader);
 
 		SetUAVParameter(RHICmdList, ComputeShader, TargetBufferParam, Buffer.UAV);
@@ -181,5 +180,6 @@ IMPLEMENT_GLOBAL_SHADER(NiagaraFillIntBufferCS, "/Plugin/FX/Niagara/Private/Niag
 void NiagaraFillGPUIntBuffer(FRHICommandList& RHICmdList, ERHIFeatureLevel::Type FeatureLevel, FRWBuffer& Buffer, int Value)
 {
 	TShaderMapRef<NiagaraFillIntBufferCS> FillCS(GetGlobalShaderMap(FeatureLevel));
-	FillCS->Execute(RHICmdList, Buffer, Value);
+	FRHIComputeShader* ComputeShader = FillCS.GetComputeShader();
+	FillCS->Execute(RHICmdList, ComputeShader, Buffer, Value);
 }
