Fix for warning after upgrading to PVS-Studio 7.0.4:

V1043: A global object variable 'INTERPOLATED_PARAMETER_PREFIX' is declared in the header. Multiple copies of it will be created in all translation units that include this header file.

#rnx

diff --git a/Engine/Plugins/FX/Niagara/Source/Niagara/Public/NiagaraCommon.h b/Engine/Plugins/FX/Niagara/Source/Niagara/Public/NiagaraCommon.h
--- a/Engine/Plugins/FX/Niagara/Source/Niagara/Public/NiagaraCommon.h
+++ b/Engine/Plugins/FX/Niagara/Source/Niagara/Public/NiagaraCommon.h
@@ -25,7 +25,7 @@ struct FNiagaraParameterStore;
 const uint32 NIAGARA_COMPUTE_THREADGROUP_SIZE = 64;
 const uint32 NIAGARA_MAX_COMPUTE_THREADGROUPS = 65535;
 
-const FString INTERPOLATED_PARAMETER_PREFIX = TEXT("PREV_");
+#define INTERPOLATED_PARAMETER_PREFIX TEXT("PREV_")
 
 /** The maximum number of spawn infos we can run on the GPU, modifying this will require a version update as it is used in the shader compiler  */
 constexpr uint32 NIAGARA_MAX_GPU_SPAWN_INFOS = 8;
