Fix for warning after upgrading to PVS-Studio 7.0.4:

V1043: A global object variable 'NumVoxelsName' is declared in the header. Multiple copies of it will be created in all translation units that include this header file.

#rnx

diff --git a/Engine/Plugins/FX/Niagara/Source/Niagara/Classes/NiagaraDataInterfaceRW.h b/Engine/Plugins/FX/Niagara/Source/Niagara/Classes/NiagaraDataInterfaceRW.h
--- a/Engine/Plugins/FX/Niagara/Source/Niagara/Classes/NiagaraDataInterfaceRW.h
+++ b/Engine/Plugins/FX/Niagara/Source/Niagara/Classes/NiagaraDataInterfaceRW.h
@@ -7,31 +7,31 @@
 class FNiagaraSystemInstance;
 
 // Global HLSL variable base names, used by HLSL.
-static const FString NumVoxelsName(TEXT("NumVoxels_"));
-static const FString VoxelSizeName(TEXT("VoxelSize_"));
-static const FString WorldBBoxSizeName(TEXT("WorldBBoxSize_"));
+extern NIAGARA_API const FString NumVoxelsName;
+extern NIAGARA_API const FString VoxelSizeName;
+extern NIAGARA_API const FString WorldBBoxSizeName;
 
-static const FString NumCellsName(TEXT("NumCells_"));
-static const FString CellSizeName(TEXT("CellSize_"));
+extern NIAGARA_API const FString NumCellsName;
+extern NIAGARA_API const FString CellSizeName;
 
 // Global VM function names, also used by the shaders code generation methods.
-static const FName NumVoxelsFunctionName("GetNumVoxels");
-static const FName VoxelSizeFunctionName("GetVoxelSize");
+extern NIAGARA_API const FName NumVoxelsFunctionName;
+extern NIAGARA_API const FName VoxelSizeFunctionName;
 
-static const FName NumCellsFunctionName("GetNumCells");
-static const FName CellSizeFunctionName("GetCellSize");
+extern NIAGARA_API const FName NumCellsFunctionName;
+extern NIAGARA_API const FName CellSizeFunctionName;
 
-static const FName WorldBBoxSizeFunctionName("GetWorldBBoxSize");
+extern NIAGARA_API const FName WorldBBoxSizeFunctionName;
 
-static const FName SimulationToUnitFunctionName("SimulationToUnit");
-static const FName UnitToSimulationFunctionName("UnitToSimulation");
-static const FName UnitToIndexFunctionName("UnitToIndex");
-static const FName IndexToUnitFunctionName("IndexToUnit");
-static const FName IndexToUnitStaggeredXFunctionName("IndexToUnitStaggeredX");
-static const FName IndexToUnitStaggeredYFunctionName("IndexToUnitStaggeredY");
+extern NIAGARA_API const FName SimulationToUnitFunctionName;
+extern NIAGARA_API const FName UnitToSimulationFunctionName;
+extern NIAGARA_API const FName UnitToIndexFunctionName;
+extern NIAGARA_API const FName IndexToUnitFunctionName;
+extern NIAGARA_API const FName IndexToUnitStaggeredXFunctionName;
+extern NIAGARA_API const FName IndexToUnitStaggeredYFunctionName;
 
-static const FName IndexToLinearFunctionName("IndexToLinear");
-static const FName LinearToIndexFunctionName("LinearToIndex");
+extern NIAGARA_API const FName IndexToLinearFunctionName;
+extern NIAGARA_API const FName LinearToIndexFunctionName;
 
 
 
diff --git a/Engine/Plugins/FX/Niagara/Source/Niagara/Private/NiagaraDataInterfaceRW.cpp b/Engine/Plugins/FX/Niagara/Source/Niagara/Private/NiagaraDataInterfaceRW.cpp
--- a/Engine/Plugins/FX/Niagara/Source/Niagara/Private/NiagaraDataInterfaceRW.cpp
+++ b/Engine/Plugins/FX/Niagara/Source/Niagara/Private/NiagaraDataInterfaceRW.cpp
@@ -7,6 +7,33 @@
 
 #define LOCTEXT_NAMESPACE "NiagaraDataInterfaceRW"
 
+// Global HLSL variable base names, used by HLSL.
+NIAGARA_API extern const FString NumVoxelsName(TEXT("NumVoxels_"));
+NIAGARA_API extern const FString VoxelSizeName(TEXT("VoxelSize_"));
+NIAGARA_API extern const FString WorldBBoxSizeName(TEXT("WorldBBoxSize_"));
+
+NIAGARA_API extern const FString NumCellsName(TEXT("NumCells_"));
+NIAGARA_API extern const FString CellSizeName(TEXT("CellSize_"));
+
+// Global VM function names, also used by the shaders code generation methods.
+NIAGARA_API extern const FName NumVoxelsFunctionName("GetNumVoxels");
+NIAGARA_API extern const FName VoxelSizeFunctionName("GetVoxelSize");
+
+NIAGARA_API extern const FName NumCellsFunctionName("GetNumCells");
+NIAGARA_API extern const FName CellSizeFunctionName("GetCellSize");
+
+NIAGARA_API extern const FName WorldBBoxSizeFunctionName("GetWorldBBoxSize");
+
+NIAGARA_API extern const FName SimulationToUnitFunctionName("SimulationToUnit");
+NIAGARA_API extern const FName UnitToSimulationFunctionName("UnitToSimulation");
+NIAGARA_API extern const FName UnitToIndexFunctionName("UnitToIndex");
+NIAGARA_API extern const FName IndexToUnitFunctionName("IndexToUnit");
+NIAGARA_API extern const FName IndexToUnitStaggeredXFunctionName("IndexToUnitStaggeredX");
+NIAGARA_API extern const FName IndexToUnitStaggeredYFunctionName("IndexToUnitStaggeredY");
+
+NIAGARA_API extern const FName IndexToLinearFunctionName("IndexToLinear");
+NIAGARA_API extern const FName LinearToIndexFunctionName("LinearToIndex");
+
 
 UNiagaraDataInterfaceRWBase::UNiagaraDataInterfaceRWBase(FObjectInitializer const& ObjectInitializer)
 	: Super(ObjectInitializer)
