- Mark Niagara Component parameter store UObjects as dirty post reachability analysis, this fixes an issue where instances can reference stale UObjects.  We may be able to refactor this later but for the moment this fixes a pretty nasty GC related issue

diff --git a/Engine/Plugins/FX/Niagara/Source/Niagara/Private/NiagaraWorldManager.cpp b/Engine/Plugins/FX/Niagara/Source/Niagara/Private/NiagaraWorldManager.cpp
--- a/Engine/Plugins/FX/Niagara/Source/Niagara/Private/NiagaraWorldManager.cpp
+++ b/Engine/Plugins/FX/Niagara/Source/Niagara/Private/NiagaraWorldManager.cpp
@@ -61,6 +61,7 @@ FDelegateHandle FNiagaraWorldManager::OnPreWorldFinishDestroyHandle;
 FDelegateHandle FNiagaraWorldManager::OnWorldBeginTearDownHandle;
 FDelegateHandle FNiagaraWorldManager::TickWorldHandle;
 FDelegateHandle FNiagaraWorldManager::PreGCHandle;
+FDelegateHandle FNiagaraWorldManager::PostReachabilityAnalysisHandle;
 FDelegateHandle FNiagaraWorldManager::PostGCHandle;
 FDelegateHandle FNiagaraWorldManager::PreGCBeginDestroyHandle; 
 TMap<class UWorld*, class FNiagaraWorldManager*> FNiagaraWorldManager::WorldManagers;
@@ -179,6 +180,7 @@ void FNiagaraWorldManager::OnStartup()
 	TickWorldHandle = FWorldDelegates::OnWorldPostActorTick.AddStatic(&FNiagaraWorldManager::TickWorld);
 
 	PreGCHandle = FCoreUObjectDelegates::GetPreGarbageCollectDelegate().AddStatic(&FNiagaraWorldManager::OnPreGarbageCollect);
+	PostReachabilityAnalysisHandle = FCoreUObjectDelegates::PostReachabilityAnalysis.AddStatic(&FNiagaraWorldManager::OnPostReachabilityAnalysis);
 	PostGCHandle = FCoreUObjectDelegates::GetPostGarbageCollect().AddStatic(&FNiagaraWorldManager::OnPostGarbageCollect);
 	PreGCBeginDestroyHandle = FCoreUObjectDelegates::PreGarbageCollectConditionalBeginDestroy.AddStatic(&FNiagaraWorldManager::OnPreGarbageCollectBeginDestroy);
 }
@@ -192,6 +194,7 @@ void FNiagaraWorldManager::OnShutdown()
 	FWorldDelegates::OnWorldPostActorTick.Remove(TickWorldHandle);
 
 	FCoreUObjectDelegates::GetPreGarbageCollectDelegate().Remove(PreGCHandle);
+	FCoreUObjectDelegates::PostReachabilityAnalysis.Remove(PostReachabilityAnalysisHandle);
 	FCoreUObjectDelegates::GetPostGarbageCollect().Remove(PostGCHandle);
 	FCoreUObjectDelegates::PreGarbageCollectConditionalBeginDestroy.Remove(PreGCBeginDestroyHandle);
 
@@ -368,6 +371,8 @@ void FNiagaraWorldManager::PreGarbageCollect()
 {
 	if (GNiagaraWaitOnPreGC)
 	{
+		// We must wait for system simulation & instance async ticks to complete before garbage collection can start
+		// The reason for this is that our async ticks could be referencing GC objects, i.e. skel meshes, etc, and we don't want them to become unreachable while we are potentially using them
 		for (int TG = 0; TG < NiagaraNumTickGroups; ++TG)
 		{
 			for (TPair<UNiagaraSystem*, TSharedRef<FNiagaraSystemSimulation, ESPMode::ThreadSafe>>& SimPair : SystemSimulations[TG])
@@ -378,6 +383,15 @@ void FNiagaraWorldManager::PreGarbageCollect()
 	}
 }
 
+void FNiagaraWorldManager::PostReachabilityAnalysis()
+{
+	for (TObjectIterator<UNiagaraComponent> ComponentIt; ComponentIt; ++ComponentIt)
+	{
+		ComponentIt->GetOverrideParameters().MarkUObjectsDirty();
+	}
+}
+
+
 void FNiagaraWorldManager::PostGarbageCollect()
 {
 	//Clear out and scalability managers who's EffectTypes have been GCd.
@@ -468,6 +482,14 @@ void FNiagaraWorldManager::OnPreGarbageCollect()
 	}
 }
 
+void FNiagaraWorldManager::OnPostReachabilityAnalysis()
+{
+	for (TPair<UWorld*, FNiagaraWorldManager*>& Pair : WorldManagers)
+	{
+		Pair.Value->PostReachabilityAnalysis();
+	}
+}
+
 void FNiagaraWorldManager::OnPostGarbageCollect()
 {
 	for (TPair<UWorld*, FNiagaraWorldManager*>& Pair : WorldManagers)
diff --git a/Engine/Plugins/FX/Niagara/Source/Niagara/Public/NiagaraWorldManager.h b/Engine/Plugins/FX/Niagara/Source/Niagara/Public/NiagaraWorldManager.h
--- a/Engine/Plugins/FX/Niagara/Source/Niagara/Public/NiagaraWorldManager.h
+++ b/Engine/Plugins/FX/Niagara/Source/Niagara/Public/NiagaraWorldManager.h
@@ -126,6 +126,7 @@ public:
 	void OnWorldCleanup(bool bSessionEnded, bool bCleanupResources);
 
 	void PreGarbageCollect();
+	void PostReachabilityAnalysis();
 	void PostGarbageCollect();
 	void PreGarbageCollectBeginDestroy();
 	
@@ -182,6 +183,9 @@ private:
 	// Callback to handle any pre GC processing needed.
 	static void OnPreGarbageCollect();
 
+	// Callback post reachability
+	static void OnPostReachabilityAnalysis();
+
 	// Callback to handle any post GC processing needed.
 	static void OnPostGarbageCollect();
 
@@ -208,6 +212,7 @@ private:
 	static FDelegateHandle OnWorldBeginTearDownHandle;
 	static FDelegateHandle TickWorldHandle;
 	static FDelegateHandle PreGCHandle;
+	static FDelegateHandle PostReachabilityAnalysisHandle;
 	static FDelegateHandle PostGCHandle;
 	static FDelegateHandle PreGCBeginDestroyHandle;
 
