Fix for redundant niagara GPUComputeShaders being saved into cooked builds because BeginCacheForCookedPlatformData is called multiple (non-deterministic) times during a cook
-saves about 10kb for each of our .uexp files containing GPUComputeShaders, which translates into less cooking, less saving and less loading (multiple shaders weren't being allocated in the game)
-resolves one source of non-deterministic cooks as depending on how much work async work the cooker had it could have created more than the guaranteed 3 of these

diff --git a/Engine/Plugins/FX/Niagara/Source/Niagara/Private/NiagaraScript.cpp b/Engine/Plugins/FX/Niagara/Source/Niagara/Private/NiagaraScript.cpp
--- a/Engine/Plugins/FX/Niagara/Source/Niagara/Private/NiagaraScript.cpp
+++ b/Engine/Plugins/FX/Niagara/Source/Niagara/Private/NiagaraScript.cpp
@@ -1642,9 +1642,19 @@ void UNiagaraScript::CacheResourceShadersForCooking(EShaderPlatform ShaderPlatfo
 		// spawn and update are combined on GPU, so we only compile spawn scripts
 		if (Usage == ENiagaraScriptUsage::ParticleGPUComputeScript)
 		{
-			FNiagaraShaderScript *ResourceToCache = nullptr;
 			ERHIFeatureLevel::Type TargetFeatureLevel = GetMaxSupportedFeatureLevel(ShaderPlatform);
+			const auto FindExistingScriptPredicate = [&](const FNiagaraShaderScript* ExistingScript)
+			{
+				return ExistingScript->MatchesScript(TargetFeatureLevel, CachedScriptVMId);
+			};
 
+			// see if the script has already been added before adding a new version
+			if (InOutCachedResources.ContainsByPredicate(FindExistingScriptPredicate))
+			{
+				return;
+			}
+			
+			FNiagaraShaderScript *ResourceToCache = nullptr;
 			FNiagaraShaderScript* NewResource = AllocateResource();
 			check(CachedScriptVMId.CompilerVersionID.IsValid());
 			check(CachedScriptVMId.BaseScriptCompileHash.IsValid());
diff --git a/Engine/Plugins/FX/Niagara/Source/NiagaraShader/Private/NiagaraShared.cpp b/Engine/Plugins/FX/Niagara/Source/NiagaraShader/Private/NiagaraShared.cpp
--- a/Engine/Plugins/FX/Niagara/Source/NiagaraShader/Private/NiagaraShared.cpp
+++ b/Engine/Plugins/FX/Niagara/Source/NiagaraShader/Private/NiagaraShared.cpp
@@ -8,6 +8,7 @@
 #include "NiagaraShaderModule.h"
 #include "NiagaraShaderType.h"
 #include "NiagaraShader.h"
+#include "NiagaraScript.h"
 #include "Stats/StatsMisc.h"
 #include "UObject/CoreObjectVersion.h"
 #include "Misc/App.h"
@@ -273,6 +274,18 @@ void FNiagaraShaderScript::SetScript(UNiagaraScript *InScript, ERHIFeatureLevel:
 	SetFeatureLevel(InFeatureLevel);
 }
 
+#if WITH_EDITOR
+bool FNiagaraShaderScript::MatchesScript(ERHIFeatureLevel::Type InFeatureLevel, const FNiagaraVMExecutableDataId& ScriptId) const
+{
+	return CompilerVersionId == ScriptId.CompilerVersionID
+		&& AdditionalDefines == ScriptId.AdditionalDefines
+		&& bUsesRapidIterationParams == ScriptId.bUsesRapidIterationParams
+		&& BaseCompileHash == ScriptId.BaseScriptCompileHash
+		&& ReferencedCompileHashes == ScriptId.ReferencedCompileHashes
+		&& FeatureLevel == InFeatureLevel;
+}
+#endif
+
 NIAGARASHADER_API  void FNiagaraShaderScript::SetRenderingThreadShaderMap(FNiagaraShaderMap* InShaderMap)
 {
 	check(IsInRenderingThread());
diff --git a/Engine/Plugins/FX/Niagara/Source/NiagaraShader/Public/NiagaraShared.h b/Engine/Plugins/FX/Niagara/Source/NiagaraShader/Public/NiagaraShared.h
--- a/Engine/Plugins/FX/Niagara/Source/NiagaraShader/Public/NiagaraShared.h
+++ b/Engine/Plugins/FX/Niagara/Source/NiagaraShader/Public/NiagaraShared.h
@@ -30,6 +30,8 @@ class FNiagaraShaderMap;
 class FNiagaraShader;
 class FNiagaraShaderMapId;
 class UNiagaraScript;
+struct FNiagaraVMExecutableDataId;
+
 #define MAX_CONCURRENT_EVENT_DATASETS 4
 
 /** Defines the compile event types for translation/compilation.*/
@@ -713,6 +715,9 @@ public:
 	NIAGARASHADER_API void SetScript(UNiagaraScript *InScript, ERHIFeatureLevel::Type InFeatureLevel, const FGuid& InCompilerVersion, const TArray<FString>& InAdditionalDefines,
 		const FNiagaraCompileHash& InBaseCompileHash, const TArray<FNiagaraCompileHash>& InReferencedCompileHashes, 
 		bool bInUsesRapidIterationParams, FString InFriendlyName);
+#if WITH_EDITOR
+	NIAGARASHADER_API bool MatchesScript(ERHIFeatureLevel::Type InFeatureLevel, const FNiagaraVMExecutableDataId& ScriptId) const;
+#endif
 
 	UNiagaraScript *GetBaseVMScript()
 	{
