Enabling drop shadows for effects quality button text.

#rnx

diff --git a/Engine/Plugins/FX/Niagara/Source/NiagaraEditor/Private/Customizations/NiagaraTypeCustomizations.cpp b/Engine/Plugins/FX/Niagara/Source/NiagaraEditor/Private/Customizations/NiagaraTypeCustomizations.cpp
--- a/Engine/Plugins/FX/Niagara/Source/NiagaraEditor/Private/Customizations/NiagaraTypeCustomizations.cpp
+++ b/Engine/Plugins/FX/Niagara/Source/NiagaraEditor/Private/Customizations/NiagaraTypeCustomizations.cpp
@@ -841,6 +841,7 @@ void FNiagaraPlatformSetTypeCustomization::GenerateEffectsQualitySelectionWidget
 								.TextStyle(FNiagaraEditorStyle::Get(), "NiagaraEditor.PlatformSet.ButtonText")
 								.Text(FNiagaraPlatformSet::GetEffectsQualityText(EffectsQuality))
 								.ColorAndOpacity(this, &FNiagaraPlatformSetTypeCustomization::GetEffectsQualityButtonTextColor, EffectsQuality)
+								.ShadowOffset(FVector2D(1, 1))
 							]
 							// error icon
 							+ SHorizontalBox::Slot()
