- Allocate space in ParameterStore for missing parameters to fix potential memory overwrite

#[fyi] shaun.kime

diff --git a/Engine/Plugins/FX/Niagara/Source/Niagara/Private/NiagaraParameterStore.cpp b/Engine/Plugins/FX/Niagara/Source/Niagara/Private/NiagaraParameterStore.cpp
--- a/Engine/Plugins/FX/Niagara/Source/Niagara/Private/NiagaraParameterStore.cpp
+++ b/Engine/Plugins/FX/Niagara/Source/Niagara/Private/NiagaraParameterStore.cpp
@@ -646,6 +646,8 @@ void FNiagaraParameterStore::SanityCheckData(bool bInitInterfaces)
 	// Additional protections were added for safety.
 	bool OwnerDirtied = false;
 
+	int32 ParameterDataSize = 0;
+
 	TArray<FNiagaraVariableWithOffset>::TConstIterator It = SortedParameterOffsets.CreateConstIterator();
 	while (It)
 	{
@@ -695,10 +697,16 @@ void FNiagaraParameterStore::SanityCheckData(bool bInitInterfaces)
 
 					OwnerDirtied = true;
 				}
+				ParameterDataSize = FMath::Max(ParameterDataSize, SrcIndex + Size);
 			}
 		}
 	}
 
+	if (ParameterData.Num() < ParameterDataSize)
+	{
+		ParameterData.AddZeroed(ParameterDataSize - ParameterData.Num());
+	}
+
 	if (Owner && OwnerDirtied)
 	{
 		UE_LOG(LogNiagara, Warning, TEXT("%s needs to be resaved to prevent above warnings due to the parameter state being stale."), *Owner->GetFullName());
