TArray -> TMemoryImageArray, fixes compile error on 32bit

diff --git a/Engine/Plugins/FX/Niagara/Source/Niagara/Private/NiagaraDataInterfaceParticleRead.cpp b/Engine/Plugins/FX/Niagara/Source/Niagara/Private/NiagaraDataInterfaceParticleRead.cpp
--- a/Engine/Plugins/FX/Niagara/Source/Niagara/Private/NiagaraDataInterfaceParticleRead.cpp
+++ b/Engine/Plugins/FX/Niagara/Source/Niagara/Private/NiagaraDataInterfaceParticleRead.cpp
@@ -429,7 +429,7 @@ private:
 	LAYOUT_FIELD(FShaderParameter, AcquireTagRegisterIndexParam);
 	LAYOUT_FIELD(TMemoryImageArray<FName>, AttributeNames);
 	LAYOUT_FIELD(TMemoryImageArray<ENiagaraParticleDataValueType>, AttributeTypes);
-	LAYOUT_MUTABLE_FIELD(TArray<int32>, AttributeIndices);
+	LAYOUT_MUTABLE_FIELD(TMemoryImageArray<int32>, AttributeIndices);
 	LAYOUT_MUTABLE_FIELD_INITIALIZED(int32, AcquireTagRegisterIndex, -1);
 	LAYOUT_MUTABLE_FIELD_INITIALIZED(FNiagaraDataSet*, CachedDataSet, nullptr);
 	LAYOUT_MUTABLE_FIELD_INITIALIZED(bool, bSourceEmitterNotGPUErrorShown, false);
