Moving check for circular connection to higher priority
#rnx

diff --git a/Engine/Plugins/FX/Niagara/Source/NiagaraEditor/Private/EdGraphSchema_Niagara.cpp b/Engine/Plugins/FX/Niagara/Source/NiagaraEditor/Private/EdGraphSchema_Niagara.cpp
--- a/Engine/Plugins/FX/Niagara/Source/NiagaraEditor/Private/EdGraphSchema_Niagara.cpp
+++ b/Engine/Plugins/FX/Niagara/Source/NiagaraEditor/Private/EdGraphSchema_Niagara.cpp
@@ -940,6 +940,13 @@ const FPinConnectionResponse UEdGraphSchema_Niagara::CanCreateConnection(const U
 		return FPinConnectionResponse(CONNECT_RESPONSE_DISALLOW, TEXT("Directions are not compatible"));
 	}
 
+	// Check for a circular connection before checking any type compatibility
+	TSet<const UEdGraphNode*> VisitedNodes;
+	if (UEdGraphSchema_Niagara::CheckCircularConnection(VisitedNodes, OutputPin->GetOwningNode(), InputPin->GetOwningNode()))
+	{
+		return FPinConnectionResponse(CONNECT_RESPONSE_DISALLOW, TEXT("Circular connection found"));
+	}
+
 	if (PinA->PinType.PinCategory != TEXT("wildcard") && PinB->PinType.PinCategory != TEXT("wildcard"))
 	{
 		// Check for compatible type pins.
@@ -1011,12 +1018,6 @@ const FPinConnectionResponse UEdGraphSchema_Niagara::CanCreateConnection(const U
 		}
 	}
 
-	TSet<const UEdGraphNode*> VisitedNodes;
-	if (UEdGraphSchema_Niagara::CheckCircularConnection(VisitedNodes, OutputPin->GetOwningNode(), InputPin->GetOwningNode()))
-	{
-		return FPinConnectionResponse(CONNECT_RESPONSE_DISALLOW, TEXT("Circular connection found"));
-	}
-
 	// See if we want to break existing connections (if its an input with an existing connection)
 	const bool bBreakExistingDueToDataInput = (InputPin->LinkedTo.Num() > 0);
 	if (bBreakExistingDueToDataInput)
