Moved the cached data from the compute shader params to the per-instance render data, so the memory freezing stuff doesn't need to worry about it.

diff --git a/Engine/Plugins/FX/Niagara/Source/Niagara/Private/NiagaraDataInterfaceParticleRead.cpp b/Engine/Plugins/FX/Niagara/Source/Niagara/Private/NiagaraDataInterfaceParticleRead.cpp
--- a/Engine/Plugins/FX/Niagara/Source/Niagara/Private/NiagaraDataInterfaceParticleRead.cpp
+++ b/Engine/Plugins/FX/Niagara/Source/Niagara/Private/NiagaraDataInterfaceParticleRead.cpp
@@ -75,38 +75,56 @@ struct FNDIParticleRead_InstanceData
 	FNiagaraEmitterInstance* EmitterInstance;
 };
 
-struct FNDIParticleRead_RenderThreadData
+struct FNDIParticleRead_GameToRenderData
 {
-	FNDIParticleRead_RenderThreadData() : SourceEmitterGPUContext(nullptr) {}
+	FNiagaraComputeExecutionContext* SourceEmitterGPUContext;
+	FString SourceEmitterName;
+};
+
+struct FNDIParticleRead_RenderInstanceData
+{
+	FNDIParticleRead_RenderInstanceData() :
+		SourceEmitterGPUContext(nullptr)
+		, CachedDataSet(nullptr)
+		, AcquireTagRegisterIndex(-1)
+		, bSourceEmitterNotGPUErrorShown(false)
+	{
+	}
 
 	FNiagaraComputeExecutionContext* SourceEmitterGPUContext;
 	FString SourceEmitterName;
+	const FNiagaraDataSet* CachedDataSet;
+	TArray<int32> AttributeIndices;
+	int32 AcquireTagRegisterIndex;
+	bool bSourceEmitterNotGPUErrorShown;
 };
 
 struct FNiagaraDataInterfaceProxyParticleRead : public FNiagaraDataInterfaceProxy
 {
-	virtual void ConsumePerInstanceDataFromGameThread(void* PerInstanceData, const FNiagaraSystemInstanceID& InstanceID) override
+	virtual void ConsumePerInstanceDataFromGameThread(void* Data, const FNiagaraSystemInstanceID& InstanceID) override
 	{
-		FNDIParticleRead_RenderThreadData* SystemData = SystemsRenderData.Find(InstanceID);
-		if (!ensure(SystemData))
+		FNDIParticleRead_RenderInstanceData* InstanceData = SystemsRenderData.Find(InstanceID);
+		if (!ensure(InstanceData))
 		{
 			return;
 		}
 
-		FNDIParticleRead_RenderThreadData* IncomingData = static_cast<FNDIParticleRead_RenderThreadData*>(PerInstanceData);
-		if (!IncomingData)
+		FNDIParticleRead_GameToRenderData* IncomingData = static_cast<FNDIParticleRead_GameToRenderData*>(Data);
+		if (IncomingData)
 		{
-			SystemData->SourceEmitterGPUContext = nullptr;
-			SystemData->SourceEmitterName = TEXT("");
-			return;
+			InstanceData->SourceEmitterGPUContext = IncomingData->SourceEmitterGPUContext;
+			InstanceData->SourceEmitterName = IncomingData->SourceEmitterName;
+		}
+		else
+		{
+			InstanceData->SourceEmitterGPUContext = nullptr;
+			InstanceData->SourceEmitterName = TEXT("");
 		}
-
-		*SystemData = *IncomingData;
 	}
 	
 	virtual int32 PerInstanceDataPassedToRenderThreadSize() const override
 	{
-		return sizeof(FNDIParticleRead_RenderThreadData);
+		return sizeof(FNDIParticleRead_GameToRenderData);
 	}
 
 	void CreateRenderThreadSystemData(const FNiagaraSystemInstanceID& InstanceID)
@@ -122,26 +140,19 @@ struct FNiagaraDataInterfaceProxyParticleRead : public FNiagaraDataInterfaceProx
 		SystemsRenderData.Remove(InstanceID);
 	}
 
-	FNDIParticleRead_RenderThreadData* GetRenderDataForSystem(const FNiagaraSystemInstanceID& InstanceID)
+	FNDIParticleRead_RenderInstanceData* GetRenderDataForSystem(const FNiagaraSystemInstanceID& InstanceID)
 	{
 		return SystemsRenderData.Find(InstanceID);
 	}
 
 private:
-	TMap<FNiagaraSystemInstanceID, FNDIParticleRead_RenderThreadData> SystemsRenderData;
+	TMap<FNiagaraSystemInstanceID, FNDIParticleRead_RenderInstanceData> SystemsRenderData;
 };
 
 struct FNiagaraDataInterfaceParametersCS_ParticleRead : public FNiagaraDataInterfaceParametersCS
 {
 	DECLARE_TYPE_LAYOUT(FNiagaraDataInterfaceParametersCS_ParticleRead, NonVirtual);
 
-	FNiagaraDataInterfaceParametersCS_ParticleRead() : 
-		AcquireTagRegisterIndex(-1),
-		CachedDataSet(nullptr),
-		bSourceEmitterNotGPUErrorShown(false)
-	{
-	}
-
 	ENiagaraParticleDataValueType GetValueTypeFromFuncName(const FName& FuncName)
 	{
 		if (FuncName == GetIntAttributeFunctionName) return ENiagaraParticleDataValueType::Int;
@@ -193,13 +204,16 @@ struct FNiagaraDataInterfaceParametersCS_ParticleRead : public FNiagaraDataInter
 
 	void SetErrorParams(FRHICommandList& RHICmdList, FRHIComputeShader* ComputeShader, bool bSkipSpawnInfo) const
 	{
-		CachedDataSet = nullptr;
-
+		// Set all the indices to -1, so that any reads return false.
+		TArray<int32, TInlineAllocator<32>> AttributeIndices;
+		int NumAttrIndices = Align(AttributeNames.Num(), 4);
+		AttributeIndices.SetNumUninitialized(NumAttrIndices);
 		for (int AttrIdx = 0; AttrIdx < AttributeIndices.Num(); ++AttrIdx)
 		{
 			AttributeIndices[AttrIdx] = -1;
 		}
-		AcquireTagRegisterIndex = -1;
+
+		int32 AcquireTagRegisterIndex = -1;
 
 		if (!bSkipSpawnInfo)
 		{
@@ -233,10 +247,14 @@ struct FNiagaraDataInterfaceParametersCS_ParticleRead : public FNiagaraDataInter
 		}
 	}
 
-	void FindAttributeIndices(const FNiagaraDataSet* SourceDataSet, const TCHAR* SourceEmitterName) const
+	void FindAttributeIndices(FNDIParticleRead_RenderInstanceData* InstanceData, const FNiagaraDataSet* SourceDataSet) const
 	{
 		check(AttributeNames.Num() == AttributeTypes.Num());
 
+		int NumAttrIndices = Align(AttributeNames.Num(), 4);
+		InstanceData->AttributeIndices.SetNumUninitialized(NumAttrIndices);
+
+		// Find the register index for each named attribute in the source emitter.
 		const TArray<FNiagaraVariable>& SourceEmitterVariables = SourceDataSet->GetVariables();
 		const TArray<FNiagaraVariableLayoutInfo>& SourceEmitterVariableLayouts = SourceDataSet->GetVariableLayouts();
 		for (int AttrNameIdx = 0; AttrNameIdx < AttributeNames.Num(); ++AttrNameIdx)
@@ -244,7 +262,7 @@ struct FNiagaraDataInterfaceParametersCS_ParticleRead : public FNiagaraDataInter
 			const FName& AttrName = AttributeNames[AttrNameIdx];
 			if (AttrName == NAME_None)
 			{
-				AttributeIndices[AttrNameIdx] = -1;
+				InstanceData->AttributeIndices[AttrNameIdx] = -1;
 				continue;
 			}
 
@@ -258,14 +276,14 @@ struct FNiagaraDataInterfaceParametersCS_ParticleRead : public FNiagaraDataInter
 					if (CheckVariableType(Var.GetType(), AttributeType))
 					{
 						const FNiagaraVariableLayoutInfo& Layout = SourceEmitterVariableLayouts[VarIdx];
-						AttributeIndices[AttrNameIdx] = (AttributeType == ENiagaraParticleDataValueType::Int || AttributeType == ENiagaraParticleDataValueType::Bool) ? Layout.Int32ComponentStart : Layout.FloatComponentStart;
+						InstanceData->AttributeIndices[AttrNameIdx] = (AttributeType == ENiagaraParticleDataValueType::Int || AttributeType == ENiagaraParticleDataValueType::Bool) ? Layout.Int32ComponentStart : Layout.FloatComponentStart;
 					}
 					else
 					{
 						UE_LOG(LogNiagara, Error, TEXT("Variable '%s' in emitter '%s' has type '%s', but particle read DI tried to access it as '%s'."),
-							*Var.GetName().ToString(), SourceEmitterName, *Var.GetType().GetName(), NiagaraParticleDataValueTypeName(AttributeType)
+							*Var.GetName().ToString(), *InstanceData->SourceEmitterName, *Var.GetType().GetName(), NiagaraParticleDataValueTypeName(AttributeType)
 						);
-						AttributeIndices[AttrNameIdx] = -1;
+						InstanceData->AttributeIndices[AttrNameIdx] = -1;
 					}
 					FoundVariable = true;
 					break;
@@ -274,30 +292,31 @@ struct FNiagaraDataInterfaceParametersCS_ParticleRead : public FNiagaraDataInter
 
 			if (!FoundVariable)
 			{
-				UE_LOG(LogNiagara, Error, TEXT("Particle read DI is trying to access inexistent variable '%s' in emitter '%s'."), *AttrName.ToString(), SourceEmitterName);
-				AttributeIndices[AttrNameIdx] = -1;
+				UE_LOG(LogNiagara, Error, TEXT("Particle read DI is trying to access inexistent variable '%s' in emitter '%s'."), *AttrName.ToString(), *InstanceData->SourceEmitterName);
+				InstanceData->AttributeIndices[AttrNameIdx] = -1;
 			}
 		}
 
-		AcquireTagRegisterIndex = -1;
+		// Find the register index for the AcquireTag part of the particle ID in the source emitter.
+		InstanceData->AcquireTagRegisterIndex = -1;
 		for (int VarIdx = 0; VarIdx < SourceEmitterVariables.Num(); ++VarIdx)
 		{
 			const FNiagaraVariable& Var = SourceEmitterVariables[VarIdx];
 			if (Var.GetName().ToString() == TEXT("ID"))
 			{
-				AcquireTagRegisterIndex = SourceEmitterVariableLayouts[VarIdx].Int32ComponentStart + 1;
+				InstanceData->AcquireTagRegisterIndex = SourceEmitterVariableLayouts[VarIdx].Int32ComponentStart + 1;
 				break;
 			}
 		}
-
-		if (AcquireTagRegisterIndex == -1)
+		if (InstanceData->AcquireTagRegisterIndex == -1)
 		{
-			UE_LOG(LogNiagara, Error, TEXT("Particle read DI cannot find ID variable in emitter '%s'."), SourceEmitterName);
+			UE_LOG(LogNiagara, Error, TEXT("Particle read DI cannot find ID variable in emitter '%s'."), *InstanceData->SourceEmitterName);
 		}
 
-		for (int AttrIdx = AttributeNames.Num(); AttrIdx < AttributeIndices.Num(); ++AttrIdx)
+		// Initialize the buffer padding too, so we don't move garbage around.
+		for (int AttrIdx = AttributeNames.Num(); AttrIdx < InstanceData->AttributeIndices.Num(); ++AttrIdx)
 		{
-			AttributeIndices[AttrIdx] = -1;
+			InstanceData->AttributeIndices[AttrIdx] = -1;
 		}
 	}
 
@@ -305,48 +324,44 @@ struct FNiagaraDataInterfaceParametersCS_ParticleRead : public FNiagaraDataInter
 	{
 		check(IsInRenderingThread());
 
-		if (AttributeIndices.Num() == 0)
-		{
-			int NumAttrIndices = Align(AttributeNames.Num(), 4);
-			AttributeIndices.SetNum(NumAttrIndices);
-		}
-
 		FRHIComputeShader* ComputeShader = RHICmdList.GetBoundComputeShader();
 
 		FNiagaraDataInterfaceProxyParticleRead* Proxy = static_cast<FNiagaraDataInterfaceProxyParticleRead*>(Context.DataInterface);
 		check(Proxy);
 
-		FNDIParticleRead_RenderThreadData* RTData = Proxy->GetRenderDataForSystem(Context.SystemInstance);
-		if (!RTData)
+		FNDIParticleRead_RenderInstanceData* InstanceData = Proxy->GetRenderDataForSystem(Context.SystemInstance);
+		if (!InstanceData)
 		{
 			SetErrorParams(RHICmdList, ComputeShader, false);
 			return;
 		}
 
-		if (RTData->SourceEmitterGPUContext == nullptr)
+		if (InstanceData->SourceEmitterGPUContext == nullptr)
 		{
 			// This means the source emitter isn't running on GPU.
-			if (!bSourceEmitterNotGPUErrorShown)
+			if (!InstanceData->bSourceEmitterNotGPUErrorShown)
 			{
-				UE_LOG(LogNiagara, Error, TEXT("GPU particle read DI is set to access CPU emitter '%s'."), *RTData->SourceEmitterName);
-				bSourceEmitterNotGPUErrorShown = true;
+				UE_LOG(LogNiagara, Error, TEXT("GPU particle read DI is set to access CPU emitter '%s'."), *InstanceData->SourceEmitterName);
+				InstanceData->bSourceEmitterNotGPUErrorShown = true;
 			}
 			SetErrorParams(RHICmdList, ComputeShader, false);
+			InstanceData->CachedDataSet = nullptr;
 			return;
 		}
 
-		bSourceEmitterNotGPUErrorShown = false;
+		InstanceData->bSourceEmitterNotGPUErrorShown = false;
 
-		FNiagaraDataSet* SourceDataSet = RTData->SourceEmitterGPUContext->MainDataSet;
+		FNiagaraDataSet* SourceDataSet = InstanceData->SourceEmitterGPUContext->MainDataSet;
 		if (!SourceDataSet)
 		{
 			SetErrorParams(RHICmdList, ComputeShader, false);
+			InstanceData->CachedDataSet = nullptr;
 			return;
 		}
 
 		FNiagaraDataBuffer* SourceData;
 		uint32 NumSpawnedInstances = 0, IDAcquireTag = 0;
-		const bool bReadingOwnEmitter = Context.ComputeInstanceData->Context == RTData->SourceEmitterGPUContext;
+		const bool bReadingOwnEmitter = Context.ComputeInstanceData->Context == InstanceData->SourceEmitterGPUContext;
 		if (bReadingOwnEmitter)
 		{
 			// If the current execution context is the same as the source emitter's context, it means we're reading from
@@ -383,10 +398,10 @@ struct FNiagaraDataInterfaceParametersCS_ParticleRead : public FNiagaraDataInter
 			return;
 		}
 
-		if (CachedDataSet != SourceDataSet)
+		if (InstanceData->CachedDataSet != SourceDataSet)
 		{
-			FindAttributeIndices(SourceDataSet, *RTData->SourceEmitterName);
-			CachedDataSet = SourceDataSet;
+			FindAttributeIndices(InstanceData, SourceDataSet);
+			InstanceData->CachedDataSet = SourceDataSet;
 		}
 
 		if (!SourceData->GetGPUIDToIndexTable().Buffer)
@@ -412,8 +427,8 @@ struct FNiagaraDataInterfaceParametersCS_ParticleRead : public FNiagaraDataInter
 		SetSRVParameter(RHICmdList, ComputeShader, InputIntBufferParam, SourceData->GetGPUBufferInt().SRV);
 		SetShaderValue(RHICmdList, ComputeShader, ParticleStrideFloatParam, ParticleStrideFloat);
 		SetShaderValue(RHICmdList, ComputeShader, ParticleStrideIntParam, ParticleStrideInt);
-		SetShaderValueArray(RHICmdList, ComputeShader, AttributeIndicesParam, AttributeIndices.GetData(), AttributeIndices.Num());
-		SetShaderValue(RHICmdList, ComputeShader, AcquireTagRegisterIndexParam, AcquireTagRegisterIndex);
+		SetShaderValueArray(RHICmdList, ComputeShader, AttributeIndicesParam, InstanceData->AttributeIndices.GetData(), InstanceData->AttributeIndices.Num());
+		SetShaderValue(RHICmdList, ComputeShader, AcquireTagRegisterIndexParam, InstanceData->AcquireTagRegisterIndex);
 	}
 	
 private:
@@ -429,10 +444,6 @@ private:
 	LAYOUT_FIELD(FShaderParameter, AcquireTagRegisterIndexParam);
 	LAYOUT_FIELD(TMemoryImageArray<FName>, AttributeNames);
 	LAYOUT_FIELD(TMemoryImageArray<ENiagaraParticleDataValueType>, AttributeTypes);
-	LAYOUT_MUTABLE_FIELD(TMemoryImageArray<int32>, AttributeIndices);
-	LAYOUT_MUTABLE_FIELD_INITIALIZED(int32, AcquireTagRegisterIndex, -1);
-	LAYOUT_MUTABLE_FIELD_INITIALIZED(FNiagaraDataSet*, CachedDataSet, nullptr);
-	LAYOUT_MUTABLE_FIELD_INITIALIZED(bool, bSourceEmitterNotGPUErrorShown, false);
 };
 
 IMPLEMENT_TYPE_LAYOUT(FNiagaraDataInterfaceParametersCS_ParticleRead);
@@ -1306,7 +1317,7 @@ bool UNiagaraDataInterfaceParticleRead::GetFunctionHLSL(const FNiagaraDataInterf
 
 void UNiagaraDataInterfaceParticleRead::ProvidePerInstanceDataForRenderThread(void* DataForRenderThread, void* PerInstanceData, const FNiagaraSystemInstanceID& SystemInstance)
 {
-	FNDIParticleRead_RenderThreadData* RTData = new (DataForRenderThread) FNDIParticleRead_RenderThreadData;
+	FNDIParticleRead_GameToRenderData* RTData = new (DataForRenderThread) FNDIParticleRead_GameToRenderData;
 	FNDIParticleRead_InstanceData* PIData = static_cast<FNDIParticleRead_InstanceData*>(PerInstanceData);
 	if (PIData && PIData->EmitterInstance)
 	{
diff --git a/Engine/Plugins/FX/Niagara/Source/NiagaraCore/Private/NiagaraCustomVersion.cpp b/Engine/Plugins/FX/Niagara/Source/NiagaraCore/Private/NiagaraCustomVersion.cpp
--- a/Engine/Plugins/FX/Niagara/Source/NiagaraCore/Private/NiagaraCustomVersion.cpp
+++ b/Engine/Plugins/FX/Niagara/Source/NiagaraCore/Private/NiagaraCustomVersion.cpp
@@ -8,4 +8,4 @@ const FGuid FNiagaraCustomVersion::GUID(0xFCF57AFA, 0x50764283, 0xB9A9E658, 0xFF
 // Register the custom version with core
 FCustomVersionRegistration GRegisterNiagaraCustomVersion(FNiagaraCustomVersion::GUID, FNiagaraCustomVersion::LatestVersion, TEXT("NiagaraVer"));
 
-const FGuid FNiagaraCustomVersion::LatestScriptCompileVersion(0x4A797DE3, 0x02B740CA, 0x9E719DF4, 0x53470E2A);
+const FGuid FNiagaraCustomVersion::LatestScriptCompileVersion(0xEBAD283C, 0x50794391, 0x9D128C9D, 0xA1EDBF7A);
