Fixed missing initialization.

diff --git a/Engine/Plugins/FX/Niagara/Source/Niagara/Private/NiagaraDataInterfaceParticleRead.cpp b/Engine/Plugins/FX/Niagara/Source/Niagara/Private/NiagaraDataInterfaceParticleRead.cpp
--- a/Engine/Plugins/FX/Niagara/Source/Niagara/Private/NiagaraDataInterfaceParticleRead.cpp
+++ b/Engine/Plugins/FX/Niagara/Source/Niagara/Private/NiagaraDataInterfaceParticleRead.cpp
@@ -77,6 +77,8 @@ struct FNDIParticleRead_InstanceData
 
 struct FNDIParticleRead_GameToRenderData
 {
+	FNDIParticleRead_GameToRenderData() : SourceEmitterGPUContext(nullptr) {}
+
 	FNiagaraComputeExecutionContext* SourceEmitterGPUContext;
 	FString SourceEmitterName;
 };
