Replaced raw class name with compile-safe version when registering details customization.

#rnx

diff --git a/Engine/Plugins/FX/Niagara/Source/NiagaraEditor/Private/NiagaraEditorModule.cpp b/Engine/Plugins/FX/Niagara/Source/NiagaraEditor/Private/NiagaraEditorModule.cpp
--- a/Engine/Plugins/FX/Niagara/Source/NiagaraEditor/Private/NiagaraEditorModule.cpp
+++ b/Engine/Plugins/FX/Niagara/Source/NiagaraEditor/Private/NiagaraEditorModule.cpp
@@ -797,13 +797,10 @@ void FNiagaraEditorModule::StartupModule()
 		FNiagaraScriptVariableBinding::StaticStruct()->GetFName(),
 		FOnGetPropertyTypeCustomizationInstance::CreateStatic(&FNiagaraScriptVariableBindingCustomization::MakeInstance));
 
-	PropertyModule.RegisterCustomPropertyTypeLayout("NiagaraPlatformSet",
+	PropertyModule.RegisterCustomPropertyTypeLayout(
+		FNiagaraPlatformSet::StaticStruct()->GetFName(),
 		FOnGetPropertyTypeCustomizationInstance::CreateStatic(&FNiagaraPlatformSetTypeCustomization::MakeInstance)
 	);
-
-// 	PropertyModule.RegisterCustomPropertyTypeLayout("NiagaraSystemScalabilitySettingsArray",
-// 		FOnGetPropertyTypeCustomizationInstance::CreateStatic(&FNiagaraSystemScalabilitySettingsArrayTypeCustomization::MakeInstance)
-// 	);
 	
 	PropertyModule.RegisterCustomPropertyTypeLayout(
 		FNiagaraUserParameterBinding::StaticStruct()->GetFName(),
@@ -819,8 +816,6 @@ void FNiagaraEditorModule::StartupModule()
 		TEXT("Reinitializes the style for the niagara editor module.  Used in conjuction with live coding for UI tweaks.  May crash the editor if style objects are in use."),
 		FConsoleCommandDelegate::CreateRaw(this, &FNiagaraEditorModule::ReinitializeStyle));
 
-
-
 	FNiagaraEditorCommands::Register();
 
 	TSharedPtr<FNiagaraScriptGraphPanelPinFactory> GraphPanelPinFactory = MakeShareable(new FNiagaraScriptGraphPanelPinFactory());
@@ -856,7 +851,6 @@ void FNiagaraEditorModule::StartupModule()
 	GraphPanelPinFactory->RegisterTypePin(FNiagaraTypeDefinition::GetParameterMapStruct(), FNiagaraScriptGraphPanelPinFactory::FCreateGraphPin::CreateLambda(
 		[](UEdGraphPin* GraphPin) -> TSharedRef<SGraphPin> { return SNew(SGraphPinExec, GraphPin); }));
 
-
 	EnumTypeUtilities = MakeShareable(new FNiagaraEditorEnumTypeUtilities());
 	RegisterTypeUtilities(FNiagaraTypeDefinition::GetFloatDef(), MakeShareable(new FNiagaraEditorFloatTypeUtilities()));
 	RegisterTypeUtilities(FNiagaraTypeDefinition::GetIntDef(), MakeShareable(new FNiagaraEditorIntegerTypeUtilities()));
