IWYU - removed unnecessary includes and forward declarations.

#rnx

diff --git a/Engine/Plugins/FX/Niagara/Source/NiagaraEditor/Private/NiagaraNodeEmitter.cpp b/Engine/Plugins/FX/Niagara/Source/NiagaraEditor/Private/NiagaraNodeEmitter.cpp
--- a/Engine/Plugins/FX/Niagara/Source/NiagaraEditor/Private/NiagaraNodeEmitter.cpp
+++ b/Engine/Plugins/FX/Niagara/Source/NiagaraEditor/Private/NiagaraNodeEmitter.cpp
@@ -11,7 +11,6 @@
 #include "NiagaraScriptSource.h"
 #include "NiagaraGraph.h"
 #include "NiagaraNodeInput.h"
-#include "SNiagaraGraphPinAdd.h"
 #include "NiagaraGraph.h"
 #include "NiagaraNodeParameterMapBase.h"
 #include "NiagaraNodeOutput.h"
diff --git a/Engine/Plugins/FX/Niagara/Source/NiagaraEditor/Private/NiagaraNodeWithDynamicPins.cpp b/Engine/Plugins/FX/Niagara/Source/NiagaraEditor/Private/NiagaraNodeWithDynamicPins.cpp
--- a/Engine/Plugins/FX/Niagara/Source/NiagaraEditor/Private/NiagaraNodeWithDynamicPins.cpp
+++ b/Engine/Plugins/FX/Niagara/Source/NiagaraEditor/Private/NiagaraNodeWithDynamicPins.cpp
@@ -9,7 +9,6 @@
 #include "ToolMenus.h"
 #include "Widgets/Input/SEditableTextBox.h"
 #include "Widgets/Layout/SBox.h"
-#include "SNiagaraGraphPinAdd.h"
 #include "INiagaraEditorTypeUtilities.h"
 #include "NiagaraEditorUtilities.h"
 #include "SNiagaraParameterMapView.h"
diff --git a/Engine/Plugins/FX/Niagara/Source/NiagaraEditor/Private/Widgets/SNiagaraParameterMapView.cpp b/Engine/Plugins/FX/Niagara/Source/NiagaraEditor/Private/Widgets/SNiagaraParameterMapView.cpp
--- a/Engine/Plugins/FX/Niagara/Source/NiagaraEditor/Private/Widgets/SNiagaraParameterMapView.cpp
+++ b/Engine/Plugins/FX/Niagara/Source/NiagaraEditor/Private/Widgets/SNiagaraParameterMapView.cpp
@@ -7,7 +7,6 @@
 #include "Widgets/Input/SButton.h"
 #include "Widgets/Input/SComboButton.h"
 #include "Widgets/Input/SSearchBox.h"
-#include "Widgets/SNiagaraGraphPinAdd.h"
 #include "NiagaraCommon.h"
 #include "NiagaraScript.h"
 #include "NiagaraScriptSource.h"
diff --git a/Engine/Plugins/FX/Niagara/Source/NiagaraEditor/Private/Widgets/SNiagaraParameterMapView.h b/Engine/Plugins/FX/Niagara/Source/NiagaraEditor/Private/Widgets/SNiagaraParameterMapView.h
--- a/Engine/Plugins/FX/Niagara/Source/NiagaraEditor/Private/Widgets/SNiagaraParameterMapView.h
+++ b/Engine/Plugins/FX/Niagara/Source/NiagaraEditor/Private/Widgets/SNiagaraParameterMapView.h
@@ -19,7 +19,6 @@ class SEditableTextBox;
 class SExpanderArrow;
 class SSearchBox;
 class SComboButton;
-class SNiagaraGraphPinAdd;
 class FNiagaraObjectSelection;
 class UNiagaraGraph;
 struct FEdGraphSchemaAction;
diff --git a/Engine/Plugins/FX/Niagara/Source/NiagaraEditor/Public/NiagaraNodeWithDynamicPins.h b/Engine/Plugins/FX/Niagara/Source/NiagaraEditor/Public/NiagaraNodeWithDynamicPins.h
--- a/Engine/Plugins/FX/Niagara/Source/NiagaraEditor/Public/NiagaraNodeWithDynamicPins.h
+++ b/Engine/Plugins/FX/Niagara/Source/NiagaraEditor/Public/NiagaraNodeWithDynamicPins.h
@@ -10,7 +10,6 @@
 #include "NiagaraNodeWithDynamicPins.generated.h"
 
 class UEdGraphPin;
-class SNiagaraGraphPinAdd;
 
 DECLARE_DELEGATE_OneParam(FOnAddParameter, const FNiagaraVariable&);
 
