Display a message when multi-selecting Niagara system actors to explain why the Override Parameters category is missing.

#rnx

diff --git a/Engine/Plugins/FX/Niagara/Source/NiagaraEditor/Private/Customizations/NiagaraComponentDetails.cpp b/Engine/Plugins/FX/Niagara/Source/NiagaraEditor/Private/Customizations/NiagaraComponentDetails.cpp
--- a/Engine/Plugins/FX/Niagara/Source/NiagaraEditor/Private/Customizations/NiagaraComponentDetails.cpp
+++ b/Engine/Plugins/FX/Niagara/Source/NiagaraEditor/Private/Customizations/NiagaraComponentDetails.cpp
@@ -53,14 +53,25 @@
 class FNiagaraComponentNodeBuilder : public IDetailCustomNodeBuilder
 {
 public:
-	FNiagaraComponentNodeBuilder(UNiagaraComponent* InComponent, UNiagaraScript* SourceSpawn, UNiagaraScript* SourceUpdate)						   
+	FNiagaraComponentNodeBuilder(UNiagaraComponent* InComponent)						   
 	{
 		Component = InComponent;
 		Component->OnSynchronizedWithAssetParameters().AddRaw(this, &FNiagaraComponentNodeBuilder::ComponentSynchronizedWithAssetParameters);
-		if (SourceSpawn)
-			OriginalScripts.Add(SourceSpawn);
-		if (SourceUpdate)
-			OriginalScripts.Add(SourceUpdate);
+
+		if (InComponent->GetAsset() != nullptr)
+		{
+			UNiagaraScript* ScriptSpawn = InComponent->GetAsset()->GetSystemSpawnScript();
+			if (ScriptSpawn != nullptr)
+			{
+				OriginalScripts.Add(ScriptSpawn);
+			}
+
+			UNiagaraScript* ScriptUpdate = InComponent->GetAsset()->GetSystemUpdateScript();
+			if (ScriptUpdate != nullptr)
+			{
+				OriginalScripts.Add(ScriptUpdate);
+			}
+		}
 		//UE_LOG(LogNiagaraEditor, Log, TEXT("FNiagaraComponentNodeBuilder %p Component %p"), this, Component.Get());
 	}
 
@@ -464,20 +475,20 @@ void FNiagaraComponentDetails::CustomizeDetails(IDetailLayoutBuilder& DetailBuil
 		}
 
 		FGameDelegates::Get().GetEndPlayMapDelegate().AddRaw(this, &FNiagaraComponentDetails::OnPiEEnd);
-
-		if (Component->GetAsset())
-		{
-			UNiagaraScript* ScriptSpawn = Component->GetAsset()->GetSystemSpawnScript();
-			UNiagaraScript* ScriptUpdate = Component->GetAsset()->GetSystemUpdateScript();
-
-			IDetailCategoryBuilder& InputParamCategory = DetailBuilder.EditCategory(ParamCategoryName, LOCTEXT("ParamCategoryName", "Override Parameters"));
-			InputParamCategory.AddCustomBuilder(MakeShared<FNiagaraComponentNodeBuilder>(Component.Get(), ScriptSpawn, ScriptUpdate));
-		}
-		else
-		{
-			IDetailCategoryBuilder& InputParamCategory = DetailBuilder.EditCategory(ParamCategoryName, LOCTEXT("ParamCategoryName", "Override Parameters"));
-			InputParamCategory.AddCustomBuilder(MakeShared<FNiagaraComponentNodeBuilder>(Component.Get(), nullptr, nullptr));
-		}
+			
+		IDetailCategoryBuilder& InputParamCategory = DetailBuilder.EditCategory(ParamCategoryName, LOCTEXT("ParamCategoryName", "Override Parameters"));
+		InputParamCategory.AddCustomBuilder(MakeShared<FNiagaraComponentNodeBuilder>(Component.Get()));
+	}
+	else if (ObjectsCustomized.Num() > 1)
+	{
+		IDetailCategoryBuilder& InputParamCategory = DetailBuilder.EditCategory(ParamCategoryName, LOCTEXT("ParamCategoryName", "Override Parameters"));
+		InputParamCategory.AddCustomRow(LOCTEXT("ParamCategoryName", "Override Parameters"))
+			.WholeRowContent()
+			[
+				SNew(STextBlock)
+				.TextStyle(FEditorStyle::Get(), "SmallText")
+				.Text(LOCTEXT("OverrideParameterMultiselectionUnsupported", "Multiple override parameter sets cannot be edited simultaneously."))
+			];
 	}
 }
 
