The Niagara Emitter command "Create Duplicate Parent" will now close editors for the asset before reparenting it so that the visual state matches the state on disk.

#rnx

diff --git a/Engine/Plugins/FX/Niagara/Source/NiagaraEditor/Private/AssetTypeActions/AssetTypeActions_NiagaraEmitter.cpp b/Engine/Plugins/FX/Niagara/Source/NiagaraEditor/Private/AssetTypeActions/AssetTypeActions_NiagaraEmitter.cpp
--- a/Engine/Plugins/FX/Niagara/Source/NiagaraEditor/Private/AssetTypeActions/AssetTypeActions_NiagaraEmitter.cpp
+++ b/Engine/Plugins/FX/Niagara/Source/NiagaraEditor/Private/AssetTypeActions/AssetTypeActions_NiagaraEmitter.cpp
@@ -17,6 +17,7 @@
 #include "NiagaraEditorUtilities.h"
 #include "ViewModels/Stack/NiagaraStackGraphUtilities.h"
 #include "Modules/ModuleManager.h"
+#include "Subsystems/AssetEditorSubsystem.h"
 
 #define LOCTEXT_NAMESPACE "AssetTypeActions"
 
@@ -132,22 +133,25 @@ void FAssetTypeActions_NiagaraEmitter::ExecuteCreateDuplicateParent(TArray<TWeak
 	TArray<UObject*> ObjectsToSync;
 	FAssetToolsModule& AssetToolsModule = FModuleManager::GetModuleChecked<FAssetToolsModule>("AssetTools");
 
-	for (TWeakObjectPtr<UNiagaraEmitter> Emitter : Emitters)
+	for (TWeakObjectPtr<UNiagaraEmitter> EmitterPtr : Emitters)
 	{
-		if (Emitter.IsValid())
+		UNiagaraEmitter* Emitter = EmitterPtr.Get();
+		if (Emitter != nullptr)
 		{
 			FString Name;
 			FString PackageName;
 			CreateUniqueAssetName(Emitter->GetOutermost()->GetName(), DefaultSuffix, PackageName, Name);
 
 			// Create the factory used to generate the asset
-			UNiagaraEmitter* NewEmitter = Cast<UNiagaraEmitter>(AssetToolsModule.Get().DuplicateAssetWithDialog(Name, FPackageName::GetLongPackagePath(PackageName), Emitter.Get()));
+			UNiagaraEmitter* NewEmitter = Cast<UNiagaraEmitter>(AssetToolsModule.Get().DuplicateAssetWithDialog(Name, FPackageName::GetLongPackagePath(PackageName), Emitter));
 
 			if (NewEmitter != nullptr)
 			{
+				GEditor->GetEditorSubsystem<UAssetEditorSubsystem>()->CloseAllEditorsForAsset(Emitter);
+
 				NewEmitter->Modify();
 				NewEmitter->SetUniqueEmitterName(Name);
-				Emitter->GraphSource->MarkNotSynchronized(TEXT("Emitter created"));
+				NewEmitter->GraphSource->MarkNotSynchronized(TEXT("Emitter created"));
 
 				Emitter->Modify();
 				Emitter->SetParent(*NewEmitter);
