Fixed: if multiple instances of a DI were used in a script, the HLSL translator produced all function combinations for each instance.

DataInterfaceRegisteredFunctions accumulated all registered functions from all the instances, then the generator iterated over it for each instance, which was incorrect. The information was already available per instance in FNiagaraScriptDataInterfaceCompileInfo.RegisteredFunctions, so we can just use that and remove DataInterfaceRegisteredFunctions entirely.

Also cleaned up some unused variables.

diff --git a/Engine/Plugins/FX/Niagara/Source/NiagaraCore/Private/NiagaraCustomVersion.cpp b/Engine/Plugins/FX/Niagara/Source/NiagaraCore/Private/NiagaraCustomVersion.cpp
--- a/Engine/Plugins/FX/Niagara/Source/NiagaraCore/Private/NiagaraCustomVersion.cpp
+++ b/Engine/Plugins/FX/Niagara/Source/NiagaraCore/Private/NiagaraCustomVersion.cpp
@@ -8,4 +8,4 @@ const FGuid FNiagaraCustomVersion::GUID(0xFCF57AFA, 0x50764283, 0xB9A9E658, 0xFF
 // Register the custom version with core
 FCustomVersionRegistration GRegisterNiagaraCustomVersion(FNiagaraCustomVersion::GUID, FNiagaraCustomVersion::LatestVersion, TEXT("NiagaraVer"));
 
-const FGuid FNiagaraCustomVersion::LatestScriptCompileVersion(0xEBAD283C, 0x50794391, 0x9D128C9D, 0xA1EDBF7A);
+const FGuid FNiagaraCustomVersion::LatestScriptCompileVersion(0xD6750F29, 0x8A554BC2, 0xACB577F3, 0xC6BD12CD);
diff --git a/Engine/Plugins/FX/Niagara/Source/NiagaraEditor/Private/NiagaraHlslTranslator.cpp b/Engine/Plugins/FX/Niagara/Source/NiagaraEditor/Private/NiagaraHlslTranslator.cpp
--- a/Engine/Plugins/FX/Niagara/Source/NiagaraEditor/Private/NiagaraHlslTranslator.cpp
+++ b/Engine/Plugins/FX/Niagara/Source/NiagaraEditor/Private/NiagaraHlslTranslator.cpp
@@ -1742,19 +1742,13 @@ void FHlslNiagaraTranslator::DefineDataSetWriteFunction(FString &HlslOutputStrin
 	HlslOutput += TEXT("}\n\n");
 }
 
-void FHlslNiagaraTranslator::DefineDataInterfaceHLSL(FString &InHlslOutput)
+void FHlslNiagaraTranslator::DefineDataInterfaceHLSL(FString& InHlslOutput)
 {
 	FString InterfaceCommonHLSL;
 	FString InterfaceUniformHLSL;
 	FString InterfaceFunctionHLSL;
-	TArray<FString> BufferParamNames;
 	TSet<FName> InterfaceClasses;
-	for (uint32 i = 0; i < 32; i++)
-	{
-		BufferParamNames.Add(TEXT("DataInterfaceBuffer_") + FString::FromInt(i));
-	}
 
-	uint32 CurBufferIndex = 0;
 	for (int32 i = 0; i < CompilationOutput.ScriptData.DataInterfaceInfo.Num(); i++)
 	{
 		FNiagaraScriptDataInterfaceCompileInfo& Info = CompilationOutput.ScriptData.DataInterfaceInfo[i];
@@ -1773,37 +1767,31 @@ void FHlslNiagaraTranslator::DefineDataInterfaceHLSL(FString &InHlslOutput)
 			FString OwnerIDString = Info.Name.ToString();
 			FString SanitizedOwnerIDString = GetSanitizedSymbolName(OwnerIDString, true);
 
-			// grab the buffer definition from the interface
-			//
 			FNiagaraDataInterfaceGPUParamInfo& DIInstanceInfo = DIParamInfo.AddDefaulted_GetRef();
 			DIInstanceInfo.DataInterfaceHLSLSymbol = SanitizedOwnerIDString;
 			DIInstanceInfo.DIClassName = Info.Type.GetClass()->GetName();
 
 			// Build a list of function instances that will be generated for this DI.
-			const TSet<FNiagaraFunctionSignature>* DataInterfaceFunctions = DataInterfaceRegisteredFunctions.Find(Info.Type.GetFName());
-			if (DataInterfaceFunctions != nullptr)
+			DIInstanceInfo.GeneratedFunctions.Reserve(Info.RegisteredFunctions.Num());
+			for (const FNiagaraFunctionSignature& OriginalSig : Info.RegisteredFunctions)
 			{
-				DIInstanceInfo.GeneratedFunctions.Reserve(DataInterfaceFunctions->Num());
-				for (const FNiagaraFunctionSignature& OriginalSig : *DataInterfaceFunctions)
+				if (!OriginalSig.bSupportsGPU)
 				{
-					if (!OriginalSig.bSupportsGPU)
-					{
-						Error(FText::Format(LOCTEXT("GPUDataInterfaceFunctionNotSupported", "DataInterface {0} function {1} cannot run on the GPU."), FText::FromName(Info.Type.GetFName()), FText::FromName(OriginalSig.Name)), nullptr, nullptr);
-						continue;
-					}
+					Error(FText::Format(LOCTEXT("GPUDataInterfaceFunctionNotSupported", "DataInterface {0} function {1} cannot run on the GPU."), FText::FromName(Info.Type.GetFName()), FText::FromName(OriginalSig.Name)), nullptr, nullptr);
+					continue;
+				}
 
-					// make a copy so we can modify the owner id and get the correct hlsl signature
-					FNiagaraFunctionSignature Sig = OriginalSig;
-					Sig.OwnerName = Info.Name;
+				// make a copy so we can modify the owner id and get the correct hlsl signature
+				FNiagaraFunctionSignature Sig = OriginalSig;
+				Sig.OwnerName = Info.Name;
 
-					FNiagaraDataInterfaceGeneratedFunction& DIFunc = DIInstanceInfo.GeneratedFunctions.AddDefaulted_GetRef();
-					DIFunc.DefinitionName = Sig.Name;
-					DIFunc.InstanceName = GetFunctionSignatureSymbol(Sig);
-					DIFunc.Specifiers.Empty(Sig.FunctionSpecifiers.Num());
-					for (const TTuple<FName, FName>& Specifier : Sig.FunctionSpecifiers)
-					{
-						DIFunc.Specifiers.Add(Specifier);
-					}
+				FNiagaraDataInterfaceGeneratedFunction& DIFunc = DIInstanceInfo.GeneratedFunctions.AddDefaulted_GetRef();
+				DIFunc.DefinitionName = Sig.Name;
+				DIFunc.InstanceName = GetFunctionSignatureSymbol(Sig);
+				DIFunc.Specifiers.Empty(Sig.FunctionSpecifiers.Num());
+				for (const TTuple<FName, FName>& Specifier : Sig.FunctionSpecifiers)
+				{
+					DIFunc.Specifiers.Add(Specifier);
 				}
 			}
 
@@ -5067,7 +5055,6 @@ void FHlslNiagaraTranslator::HandleCustomHlslNode(UNiagaraNodeCustomHlsl* Custom
 				if (NumFound != 0)
 				{
 					AddedFuncs.Add(Sig);
-					DataInterfaceRegisteredFunctions.FindOrAdd(Input.GetType().GetFName()).Add(Sig);
 
 					if (Info.UserPtrIdx != INDEX_NONE && CompilationTarget != ENiagaraSimTarget::GPUComputeSim)
 					{
@@ -5554,9 +5541,6 @@ void FHlslNiagaraTranslator::RegisterFunctionCall(ENiagaraScriptUsage ScriptUsag
 					}
 				}
 
-				// We only use this for GPU systems currently so that we emit only the functionality required
-				DataInterfaceRegisteredFunctions.FindOrAdd(Info.Type.GetFName()).Add(OutSignature);
-
 				if (Info.UserPtrIdx != INDEX_NONE && CompilationTarget != ENiagaraSimTarget::GPUComputeSim)
 				{
 					//This interface requires per instance data via a user ptr so place the index to it at the end of the inputs.
diff --git a/Engine/Plugins/FX/Niagara/Source/NiagaraEditor/Private/NiagaraHlslTranslator.h b/Engine/Plugins/FX/Niagara/Source/NiagaraEditor/Private/NiagaraHlslTranslator.h
--- a/Engine/Plugins/FX/Niagara/Source/NiagaraEditor/Private/NiagaraHlslTranslator.h
+++ b/Engine/Plugins/FX/Niagara/Source/NiagaraEditor/Private/NiagaraHlslTranslator.h
@@ -684,6 +684,4 @@ private:
 	TArray<FNiagaraVariable> InitialNamespaceVariablesMissingDefault;
 	// Variables that need to be initialized in the body or at the end of spawn.
 	TArray<FNiagaraVariable> DeferredVariablesMissingDefault;
-
-	TMap<FName, TSet<FNiagaraFunctionSignature>> DataInterfaceRegisteredFunctions;
 };
