Fix shader compilation in NiagaraRibbonVertexFactory.ush for instanced stereo.

diff --git a/Engine/Plugins/FX/Niagara/Shaders/Private/NiagaraRibbonVertexFactory.ush b/Engine/Plugins/FX/Niagara/Shaders/Private/NiagaraRibbonVertexFactory.ush
--- a/Engine/Plugins/FX/Niagara/Shaders/Private/NiagaraRibbonVertexFactory.ush
+++ b/Engine/Plugins/FX/Niagara/Shaders/Private/NiagaraRibbonVertexFactory.ush
@@ -902,7 +902,7 @@ FVertexFactoryInterpolantsVSToDS VertexFactoryInterpolate(FVertexFactoryInterpol
 #endif
 
 #if INSTANCED_STEREO
-	Interpolants.EyeIndex = 0;
+	O.InterpolantsVSToPS.EyeIndex = 0;
 #endif
 
 	return O;
