Merging //UE4/Dev-RenderPlat-Staging to Dev-VR (//UE4/Dev-VR)

Not sure why Robomerge didn't merge all these... still more it's missing after this I presume, because shader compile worker is failing.  Win64 Development Editor seems to build now at least.

diff --git a/Engine/Plugins/FX/Niagara/Source/NiagaraEditorWidgets/Private/Stack/SNiagaraStackModuleItem.cpp b/Engine/Plugins/FX/Niagara/Source/NiagaraEditorWidgets/Private/Stack/SNiagaraStackModuleItem.cpp
--- a/Engine/Plugins/FX/Niagara/Source/NiagaraEditorWidgets/Private/Stack/SNiagaraStackModuleItem.cpp
+++ b/Engine/Plugins/FX/Niagara/Source/NiagaraEditorWidgets/Private/Stack/SNiagaraStackModuleItem.cpp
@@ -66,6 +66,7 @@ void SNiagaraStackModuleItem::Construct(const FArguments& InArgs, UNiagaraStackM
 				.IssueSeverity_UObject(ModuleItem, &UNiagaraStackModuleItem::GetHighestStackIssueSeverity)
 				.ErrorTooltip(this, &SNiagaraStackModuleItem::GetErrorButtonTooltipText)
 				.Visibility(this, &SNiagaraStackModuleItem::GetStackIssuesWarningVisibility)
+				.IsEnabled(this, &SNiagaraStackModuleItem::GetButtonsEnabled)
 				.OnButtonClicked(this, &SNiagaraStackModuleItem::ExpandEntry)
 			]
 			// Raise Action Menu button
@@ -81,6 +82,7 @@ void SNiagaraStackModuleItem::Construct(const FArguments& InArgs, UNiagaraStackM
 				.HAlign(HAlign_Center)
 				.VAlign(VAlign_Center)
 				.Visibility(this, &SNiagaraStackModuleItem::GetRaiseActionMenuVisibility)
+				.IsEnabled(this, &SNiagaraStackModuleItem::GetButtonsEnabled)
 			]
 			// Refresh button
 			+ SHorizontalBox::Slot()
@@ -93,6 +95,7 @@ void SNiagaraStackModuleItem::Construct(const FArguments& InArgs, UNiagaraStackM
 				.ForegroundColor(FNiagaraEditorWidgetsStyle::Get().GetColor("NiagaraEditor.Stack.FlatButtonColor"))
 				.ToolTipText(LOCTEXT("RefreshTooltip", "Refresh this module"))
 				.Visibility(this, &SNiagaraStackModuleItem::GetRefreshVisibility)
+				.IsEnabled(this, &SNiagaraStackModuleItem::GetButtonsEnabled)
 				.OnClicked(this, &SNiagaraStackModuleItem::RefreshClicked)
 				.Content()
 				[
@@ -110,8 +113,8 @@ void SNiagaraStackModuleItem::Construct(const FArguments& InArgs, UNiagaraStackM
 				.ButtonStyle(FEditorStyle::Get(), "HoverHintOnly")
 				.IsFocusable(false)
 				.ForegroundColor(FNiagaraEditorWidgetsStyle::Get().GetColor("NiagaraEditor.Stack.FlatButtonColor"))
-				.ToolTipText(LOCTEXT("DeleteToolTip", "Delete this module"))
-				.Visibility(this, &SNiagaraStackModuleItem::GetEditButtonVisibility)
+				.ToolTipText(this, &SNiagaraStackModuleItem::GetDeleteButtonToolTipText)
+				.IsEnabled(this, &SNiagaraStackModuleItem::GetDeleteButtonEnabled)
 				.OnClicked(this, &SNiagaraStackModuleItem::DeleteClicked)
 				.Content()
 				[
@@ -128,6 +131,7 @@ void SNiagaraStackModuleItem::Construct(const FArguments& InArgs, UNiagaraStackM
 				SNew(SCheckBox)
 				.IsChecked(this, &SNiagaraStackModuleItem::GetCheckState)
 				.OnCheckStateChanged(this, &SNiagaraStackModuleItem::OnCheckStateChanged)
+				.IsEnabled(this, &SNiagaraStackModuleItem::GetEnabledCheckBoxEnabled)
 			]
 		]
 	];
@@ -179,34 +183,8 @@ void SNiagaraStackModuleItem::OnCheckStateChanged(ECheckBoxState InCheckState)
 	ModuleItem->SetIsEnabled(InCheckState == ECheckBoxState::Checked);
 }
 
-EVisibility SNiagaraStackModuleItem::GetEditButtonVisibility() const
+bool SNiagaraStackModuleItem::GetButtonsEnabled() const
 {
->>>> ORIGINAL //UE4/Dev-RenderPlat-Staging/Engine/Plugins/FX/Niagara/Source/NiagaraEditorWidgets/Private/Stack/SNiagaraStackModuleItem.cpp#3
-	return ModuleItem->GetOwnerIsEnabled() && ModuleItem->GetIsEnabled();
-}
-
-FText SNiagaraStackModuleItem::GetDeleteButtonToolTipText() const
-{
-	if (ModuleItem->CanMoveAndDelete())
-	{
-		return LOCTEXT("DeleteToolTip", "Delete this module");
-	}
-	else
-	{
-		// TODO: This message should be handled by the entry.
-		return LOCTEXT("CantDeleteToolTip", "This module can not be deleted becaue it is inherited.");
-	}
-}
-
-bool SNiagaraStackModuleItem::GetDeleteButtonEnabled() const
-{
-	return ModuleItem->GetOwnerIsEnabled() && ModuleItem->CanMoveAndDelete();
-}
-
-bool SNiagaraStackModuleItem::GetEnabledCheckBoxEnabled() const
-{
-	return ModuleItem->GetOwnerIsEnabled();
-==== THEIRS //UE4/Dev-RenderPlat-Staging/Engine/Plugins/FX/Niagara/Source/NiagaraEditorWidgets/Private/Stack/SNiagaraStackModuleItem.cpp#4
 	return ModuleItem->GetOwnerIsEnabled() && ModuleItem->GetIsEnabled();
 }
 
@@ -226,9 +204,6 @@ bool SNiagaraStackModuleItem::GetDeleteButtonEnabled() const
 bool SNiagaraStackModuleItem::GetEnabledCheckBoxEnabled() const
 {
 	return ModuleItem->GetOwnerIsEnabled();
-==== YOURS //joe.conley_RDU-WD-2786_2965/Engine/Plugins/FX/Niagara/Source/NiagaraEditorWidgets/Private/Stack/SNiagaraStackModuleItem.cpp
-	return ModuleItem->CanMoveAndDelete() ? EVisibility::Visible : EVisibility::Collapsed;
-<<<<
 }
 
 EVisibility SNiagaraStackModuleItem::GetRaiseActionMenuVisibility() const
diff --git a/Engine/Source/Runtime/D3D12RHI/Private/D3D12RHIPrivate.h b/Engine/Source/Runtime/D3D12RHI/Private/D3D12RHIPrivate.h
--- a/Engine/Source/Runtime/D3D12RHI/Private/D3D12RHIPrivate.h
+++ b/Engine/Source/Runtime/D3D12RHI/Private/D3D12RHIPrivate.h
@@ -27,10 +27,6 @@
 
 #define D3D12_RHI_RAYTRACING (RHI_RAYTRACING)
 
-#ifndef WITH_DX_PERF
-#define WITH_DX_PERF 0
-#endif
-
 // Dependencies.
 #include "CoreMinimal.h"
 #include "RHI.h"
@@ -337,7 +333,7 @@ public:
 	virtual void RHIBindDebugLabelName(FRHITexture* Texture, const TCHAR* Name) final override;
 	virtual void RHIReadSurfaceData(FRHITexture* Texture, FIntRect Rect, TArray<FColor>& OutData, FReadSurfaceDataFlags InFlags) final override;
 	virtual void RHIReadSurfaceData(FRHITexture* TextureRHI, FIntRect InRect, TArray<FLinearColor>& OutData, FReadSurfaceDataFlags InFlags) final override;
-	virtual void RHIMapStagingSurface(FRHITexture* Texture, void*& OutData, int32& OutWidth, int32& OutHeight, uint32 GPUIndex = 0) final override;
+	virtual void RHIMapStagingSurface(FRHITexture* Texture, FRHIGPUFence* Fence, void*& OutData, int32& OutWidth, int32& OutHeight, uint32 GPUIndex = 0) final override;
 	virtual void RHIUnmapStagingSurface(FRHITexture* Texture, uint32 GPUIndex = 0) final override;
 	virtual void RHIReadSurfaceFloatData(FRHITexture* Texture, FIntRect Rect, TArray<FFloat16Color>& OutData, ECubeFace CubeFace, int32 ArrayIndex, int32 MipIndex) final override;
 	virtual void RHIRead3DSurfaceFloatData(FRHITexture* Texture, FIntRect Rect, FIntPoint ZMinMax, TArray<FFloat16Color>& OutData) final override;
diff --git a/Engine/Source/Runtime/D3D12RHI/Private/D3D12RenderTarget.cpp b/Engine/Source/Runtime/D3D12RHI/Private/D3D12RenderTarget.cpp
--- a/Engine/Source/Runtime/D3D12RHI/Private/D3D12RenderTarget.cpp
+++ b/Engine/Source/Runtime/D3D12RHI/Private/D3D12RenderTarget.cpp
@@ -1397,7 +1397,7 @@ void FD3D12DynamicRHI::ReadSurfaceDataMSAARaw(FRHICommandList_RecursiveHazardous
 	}
 }
 
-void FD3D12DynamicRHI::RHIMapStagingSurface(FRHITexture* TextureRHI, void*& OutData, int32& OutWidth, int32& OutHeight, uint32 GPUIndex)
+void FD3D12DynamicRHI::RHIMapStagingSurface(FRHITexture* TextureRHI, FRHIGPUFence* FenceRHI, void*& OutData, int32& OutWidth, int32& OutHeight, uint32 GPUIndex)
 {
 	FD3D12Texture2D* DestTexture2D = ResourceCast(TextureRHI->GetTexture2D());
 
diff --git a/Engine/Source/Runtime/Engine/Private/RendererSettings.cpp b/Engine/Source/Runtime/Engine/Private/RendererSettings.cpp
--- a/Engine/Source/Runtime/Engine/Private/RendererSettings.cpp
+++ b/Engine/Source/Runtime/Engine/Private/RendererSettings.cpp
@@ -73,6 +73,7 @@ URendererSettings::URendererSettings(const FObjectInitializer& ObjectInitializer
 	GPUSimulationTextureSizeX = 1024;
 	GPUSimulationTextureSizeY = 1024;
 	bEnableRayTracing = 0;
+	bLPV = 1;
 }
 
 void URendererSettings::PostInitProperties()
diff --git a/Engine/Source/Runtime/Experimental/Chaos/Private/Chaos/HeightField.cpp b/Engine/Source/Runtime/Experimental/Chaos/Private/Chaos/HeightField.cpp
--- a/Engine/Source/Runtime/Experimental/Chaos/Private/Chaos/HeightField.cpp
+++ b/Engine/Source/Runtime/Experimental/Chaos/Private/Chaos/HeightField.cpp
@@ -415,21 +415,6 @@ namespace Chaos
 		}
 	}
 
-
-		{
-			// No resample, just push new heights into the data
-			for(int32 RowIdx = InBeginRow; RowIdx < EndRow; ++RowIdx)
-			{
-				for(int32 ColIdx = InBeginCol; ColIdx < EndCol; ++ColIdx)
-				{
-					const int32 HeightIndex = RowIdx * NumCols + ColIdx;
-					const int32 NewSetIndex = (RowIdx - InBeginRow) * NumCols + (ColIdx - InBeginCol);
-					OutData.Heights[HeightIndex] = static_cast<typename FDataType::StorageType>((ToRealFunc(BufferView[NewSetIndex]) - OutData.MinValue) / OutData.HeightPerUnit);
-				}
-			}
-		}
-	}
-
 	template <typename T>
 	THeightField<T>::THeightField(TArray<T>&& Height, int32 NumRows, int32 NumCols, const TVector<T, 3>& InScale)
 		: TImplicitObject<T, 3>(EImplicitObject::HasBoundingBox, ImplicitObjectType::HeightField)
diff --git a/Engine/Source/Runtime/NullDrv/Public/NullRHI.h b/Engine/Source/Runtime/NullDrv/Public/NullRHI.h
--- a/Engine/Source/Runtime/NullDrv/Public/NullRHI.h
+++ b/Engine/Source/Runtime/NullDrv/Public/NullRHI.h
@@ -389,13 +389,13 @@ public:
 	}
 
 
-	virtual void RHIMapStagingSurface(FRHITexture* Texture,void*& OutData,int32& OutWidth,int32& OutHeight, uint32 GPUIndex = 0) final override
+	virtual void RHIMapStagingSurface(FRHITexture* Texture, FRHIGPUFence* Fence, void*& OutData, int32& OutWidth, int32& OutHeight, uint32 GPUIndex) final override
 	{
 
 	}
 
 
-	virtual void RHIUnmapStagingSurface(FRHITexture* Texture, uint32 GPUIndex = 0) final override
+	virtual void RHIUnmapStagingSurface(FRHITexture* Texture, uint32 GPUIndex) final override
 	{
 
 	}
diff --git a/Engine/Source/Runtime/OpenGLDrv/Private/OpenGLRenderTarget.cpp b/Engine/Source/Runtime/OpenGLDrv/Private/OpenGLRenderTarget.cpp
--- a/Engine/Source/Runtime/OpenGLDrv/Private/OpenGLRenderTarget.cpp
+++ b/Engine/Source/Runtime/OpenGLDrv/Private/OpenGLRenderTarget.cpp
@@ -936,7 +936,7 @@ void FOpenGLDynamicRHI::RHIReadSurfaceData(FRHITexture* TextureRHI, FIntRect Rec
 	GetContextStateForCurrentContext().Framebuffer = (GLuint)-1;
 }
 
-void FOpenGLDynamicRHI::RHIMapStagingSurface(FRHITexture* TextureRHI,void*& OutData,int32& OutWidth,int32& OutHeight, uint32 GPUIndex)
+void FOpenGLDynamicRHI::RHIMapStagingSurface(FRHITexture* TextureRHI, FRHIGPUFence* FenceRHI, void*& OutData, int32& OutWidth, int32& OutHeight, uint32 GPUIndex)
 {
 	FRHICommandListImmediate& RHICmdList = FRHICommandListExecutor::GetImmediateCommandList();
 
diff --git a/Engine/Source/Runtime/OpenGLDrv/Public/OpenGLDrv.h b/Engine/Source/Runtime/OpenGLDrv/Public/OpenGLDrv.h
--- a/Engine/Source/Runtime/OpenGLDrv/Public/OpenGLDrv.h
+++ b/Engine/Source/Runtime/OpenGLDrv/Public/OpenGLDrv.h
@@ -411,7 +411,7 @@ public:
 	virtual void RHIBindDebugLabelName(FRHITexture* Texture, const TCHAR* Name) final override;
 	virtual void RHIReadSurfaceData(FRHITexture* Texture,FIntRect Rect,TArray<FColor>& OutData,FReadSurfaceDataFlags InFlags) final override;
 	virtual void RHIReadSurfaceData(FRHITexture* Texture, FIntRect Rect, TArray<FLinearColor>& OutData, FReadSurfaceDataFlags InFlags) final override;
-	virtual void RHIMapStagingSurface(FRHITexture* Texture,void*& OutData,int32& OutWidth,int32& OutHeight, uint32 GPUIndex = 0) final override;
+	virtual void RHIMapStagingSurface(FRHITexture* Texture, FRHIGPUFence* Fence, void*& OutData, int32& OutWidth, int32& OutHeight, uint32 GPUIndex = 0) final override;
 	virtual void RHIUnmapStagingSurface(FRHITexture* Texture, uint32 GPUIndex = 0) final override;
 	virtual void RHIReadSurfaceFloatData(FRHITexture* Texture,FIntRect Rect,TArray<FFloat16Color>& OutData,ECubeFace CubeFace,int32 ArrayIndex,int32 MipIndex) final override;
 	virtual void RHIRead3DSurfaceFloatData(FRHITexture* Texture,FIntRect Rect,FIntPoint ZMinMax,TArray<FFloat16Color>& OutData) final override;
diff --git a/Engine/Source/Runtime/RHI/Public/DynamicRHI.h b/Engine/Source/Runtime/RHI/Public/DynamicRHI.h
--- a/Engine/Source/Runtime/RHI/Public/DynamicRHI.h
+++ b/Engine/Source/Runtime/RHI/Public/DynamicRHI.h
@@ -914,7 +914,7 @@ public:
 
 	/** Watch out for OutData to be 0 (can happen on DXGI_ERROR_DEVICE_REMOVED), don't call RHIUnmapStagingSurface in that case. */
 	// FlushType: Flush Immediate (seems wrong)
-	virtual void RHIMapStagingSurface(FRHITexture* Texture, void*& OutData, int32& OutWidth, int32& OutHeight, uint32 GPUIndex = 0) = 0;
+	virtual void RHIMapStagingSurface(FRHITexture* Texture, FRHIGPUFence* Fence, void*& OutData, int32& OutWidth, int32& OutHeight, uint32 GPUIndex = 0) = 0;
 
 	/** call after a succesful RHIMapStagingSurface() call */
 	// FlushType: Flush Immediate (seems wrong)
@@ -1206,7 +1206,7 @@ public:
 	virtual void RHIDiscardTransientResource_RenderThread(FRHIStructuredBuffer* Buffer) { }
 
 
-	virtual void RHIMapStagingSurface_RenderThread(class FRHICommandListImmediate& RHICmdList, FRHITexture* Texture, void*& OutData, int32& OutWidth, int32& OutHeight);
+	virtual void RHIMapStagingSurface_RenderThread(class FRHICommandListImmediate& RHICmdList, FRHITexture* Texture, FRHIGPUFence* Fence, void*& OutData, int32& OutWidth, int32& OutHeight);
 	virtual void RHIUnmapStagingSurface_RenderThread(class FRHICommandListImmediate& RHICmdList, FRHITexture* Texture);
 	virtual void RHIReadSurfaceFloatData_RenderThread(class FRHICommandListImmediate& RHICmdList, FRHITexture* Texture, FIntRect Rect, TArray<FFloat16Color>& OutData, ECubeFace CubeFace, int32 ArrayIndex, int32 MipIndex);
 
diff --git a/Engine/Source/Runtime/RHI/Public/RHIValidation.h b/Engine/Source/Runtime/RHI/Public/RHIValidation.h
--- a/Engine/Source/Runtime/RHI/Public/RHIValidation.h
+++ b/Engine/Source/Runtime/RHI/Public/RHIValidation.h
@@ -914,9 +914,9 @@ public:
 
 	/** Watch out for OutData to be 0 (can happen on DXGI_ERROR_DEVICE_REMOVED), don't call RHIUnmapStagingSurface in that case. */
 	// FlushType: Flush Immediate (seems wrong)
-	virtual void RHIMapStagingSurface(FRHITexture* Texture, void*& OutData, int32& OutWidth, int32& OutHeight, uint32 GPUIndex = 0) override final
+	virtual void RHIMapStagingSurface(FRHITexture* Texture, FRHIGPUFence* Fence, void*& OutData, int32& OutWidth, int32& OutHeight, uint32 GPUIndex = 0) final override
 	{
-		RHI->RHIMapStagingSurface(Texture, OutData, OutWidth, OutHeight, GPUIndex);
+		RHI->RHIMapStagingSurface(Texture, Fence, OutData, OutWidth, OutHeight, GPUIndex);
 	}
 
 	/** call after a succesful RHIMapStagingSurface() call */
diff --git a/Engine/Source/Runtime/VulkanRHI/Private/VulkanRenderTarget.cpp b/Engine/Source/Runtime/VulkanRHI/Private/VulkanRenderTarget.cpp
--- a/Engine/Source/Runtime/VulkanRHI/Private/VulkanRenderTarget.cpp
+++ b/Engine/Source/Runtime/VulkanRHI/Private/VulkanRenderTarget.cpp
@@ -926,14 +926,31 @@ void FVulkanDynamicRHI::RHIReadSurfaceData(FRHITexture* TextureRHI, FIntRect Rec
 	}
 }
 
-void FVulkanDynamicRHI::RHIMapStagingSurface(FRHITexture* TextureRHI,void*& OutData,int32& OutWidth,int32& OutHeight, uint32 GPUIndex)
+void FVulkanDynamicRHI::RHIMapStagingSurface(FRHITexture* TextureRHI, FRHIGPUFence* FenceRHI, void*& OutData, int32& OutWidth, int32& OutHeight, uint32 GPUIndex)
 {
 	FRHITexture2D* TextureRHI2D = TextureRHI->GetTexture2D();
 	check(TextureRHI2D);
 	FVulkanTexture2D* Texture2D = ResourceCast(TextureRHI2D);
 
+	int32 Pitch = Texture2D->GetSizeX();
+	if (ensureMsgf(Texture2D->Surface.GetTiling() == VK_IMAGE_TILING_LINEAR, TEXT("RHIMapStagingSurface() called with a %d x %d texture in non-linear tiling %d, the result will likely be garbled."), static_cast<int32>(Texture2D->GetSizeX()), static_cast<int32>(Texture2D->GetSizeY()), static_cast<int32>(Texture2D->Surface.GetTiling())))
+	{
+		// Pitch can be only retrieved from linear textures.
+		VkImageSubresource ImageSubResource;
+		FMemory::Memzero(ImageSubResource);
+
+		ImageSubResource.aspectMask = Texture2D->Surface.GetFullAspectMask();
+		ImageSubResource.mipLevel = 0;
+		ImageSubResource.arrayLayer = 0;
+
+		VkSubresourceLayout SubResourceLayout;
+		VulkanRHI::vkGetImageSubresourceLayout(Device->GetInstanceHandle(), Texture2D->Surface.Image, &ImageSubResource, &SubResourceLayout);
+
+		int32 BytesPerPixel = GetNumBitsPerPixel(Texture2D->Surface.StorageFormat) / 8;
+		Pitch = SubResourceLayout.rowPitch / BytesPerPixel;
+	}
 
-	OutWidth = Texture2D->GetSizeX();
+	OutWidth = Pitch;
 	OutHeight = Texture2D->GetSizeY();
 
 	OutData = Texture2D->Surface.GetMappedPointer();
@@ -941,9 +958,9 @@ void FVulkanDynamicRHI::RHIMapStagingSurface(FRHITexture* TextureRHI,void*& OutD
 }
 
 
-void FVulkanDynamicRHI::RHIMapStagingSurface_RenderThread(class FRHICommandListImmediate& RHICmdList, FRHITexture* Texture, void*& OutData, int32& OutWidth, int32& OutHeight)
+void FVulkanDynamicRHI::RHIMapStagingSurface_RenderThread(class FRHICommandListImmediate& RHICmdList, FRHITexture* Texture, FRHIGPUFence* Fence, void*& OutData, int32& OutWidth, int32& OutHeight)
 {
-	RHIMapStagingSurface(Texture, OutData, OutWidth, OutHeight, RHICmdList.GetGPUMask().ToIndex());
+	RHIMapStagingSurface(Texture, Fence, OutData, OutWidth, OutHeight, RHICmdList.GetGPUMask().ToIndex());
 }
 
 void FVulkanDynamicRHI::RHIUnmapStagingSurface_RenderThread(class FRHICommandListImmediate& RHICmdList, FRHITexture* Texture)
diff --git a/Engine/Source/Runtime/VulkanRHI/Public/VulkanDynamicRHI.h b/Engine/Source/Runtime/VulkanRHI/Public/VulkanDynamicRHI.h
--- a/Engine/Source/Runtime/VulkanRHI/Public/VulkanDynamicRHI.h
+++ b/Engine/Source/Runtime/VulkanRHI/Public/VulkanDynamicRHI.h
@@ -147,9 +147,9 @@ public:
 	virtual void RHIBindDebugLabelName(FRHIUnorderedAccessView* UnorderedAccessViewRHI, const TCHAR* Name) final override;
 	virtual void RHIReadSurfaceData(FRHITexture* Texture, FIntRect Rect, TArray<FColor>& OutData, FReadSurfaceDataFlags InFlags) final override;
 	virtual void RHIReadSurfaceData(FRHITexture* Texture, FIntRect Rect, TArray<FLinearColor>& OutData, FReadSurfaceDataFlags InFlags) final override;
-	virtual void RHIMapStagingSurface(FRHITexture* Texture, void*& OutData, int32& OutWidth, int32& OutHeight, uint32 GPUIndex = 0) final override;
+	virtual void RHIMapStagingSurface(FRHITexture* Texture, FRHIGPUFence* Fence, void*& OutData, int32& OutWidth, int32& OutHeight, uint32 GPUIndex = 0) final override;
 	virtual void RHIUnmapStagingSurface(FRHITexture* Texture, uint32 GPUIndex = 0) final override;
-	virtual void RHIMapStagingSurface_RenderThread(class FRHICommandListImmediate& RHICmdList, FRHITexture* Texture, void*& OutData, int32& OutWidth, int32& OutHeight);
+	virtual void RHIMapStagingSurface_RenderThread(class FRHICommandListImmediate& RHICmdList, FRHITexture* Texture, FRHIGPUFence* Fence, void*& OutData, int32& OutWidth, int32& OutHeight);
 	virtual void RHIUnmapStagingSurface_RenderThread(class FRHICommandListImmediate& RHICmdList, FRHITexture* Texture);
 
 
diff --git a/Engine/Source/Runtime/Windows/D3D11RHI/Private/D3D11RHIPrivate.h b/Engine/Source/Runtime/Windows/D3D11RHI/Private/D3D11RHIPrivate.h
--- a/Engine/Source/Runtime/Windows/D3D11RHI/Private/D3D11RHIPrivate.h
+++ b/Engine/Source/Runtime/Windows/D3D11RHI/Private/D3D11RHIPrivate.h
@@ -469,7 +469,7 @@ public:
 	virtual void RHIBindDebugLabelName(FRHIUnorderedAccessView* UnorderedAccessViewRHI, const TCHAR* Name) final override;
 	virtual void RHIReadSurfaceData(FRHITexture* Texture,FIntRect Rect,TArray<FColor>& OutData,FReadSurfaceDataFlags InFlags) final override;
 	virtual void RHIReadSurfaceData(FRHITexture* TextureRHI, FIntRect InRect, TArray<FLinearColor>& OutData, FReadSurfaceDataFlags InFlags) final override;
-	virtual void RHIMapStagingSurface(FRHITexture* Texture,void*& OutData,int32& OutWidth,int32& OutHeight, uint32 GPUIndex = 0) final override;
+	virtual void RHIMapStagingSurface(FRHITexture* Texture, FRHIGPUFence* Fence, void*& OutData, int32& OutWidth, int32& OutHeight, uint32 GPUIndex = 0) final override;
 	virtual void RHIUnmapStagingSurface(FRHITexture* Texture, uint32 GPUIndex = 0) final override;
 	virtual void RHIReadSurfaceFloatData(FRHITexture* Texture,FIntRect Rect,TArray<FFloat16Color>& OutData,ECubeFace CubeFace,int32 ArrayIndex,int32 MipIndex) final override;
 	virtual void RHIRead3DSurfaceFloatData(FRHITexture* Texture,FIntRect Rect,FIntPoint ZMinMax,TArray<FFloat16Color>& OutData) final override;
diff --git a/Engine/Source/Runtime/Windows/D3D11RHI/Private/D3D11RenderTarget.cpp b/Engine/Source/Runtime/Windows/D3D11RHI/Private/D3D11RenderTarget.cpp
--- a/Engine/Source/Runtime/Windows/D3D11RHI/Private/D3D11RenderTarget.cpp
+++ b/Engine/Source/Runtime/Windows/D3D11RHI/Private/D3D11RenderTarget.cpp
@@ -802,7 +802,7 @@ void FD3D11DynamicRHI::ReadSurfaceDataMSAARaw(FRHICommandList_RecursiveHazardous
 	}
 }
 
-void FD3D11DynamicRHI::RHIMapStagingSurface(FRHITexture* TextureRHI,void*& OutData,int32& OutWidth,int32& OutHeight, uint32 GPUIndex)
+void FD3D11DynamicRHI::RHIMapStagingSurface(FRHITexture* TextureRHI, FRHIGPUFence* FenceRHI, void*& OutData, int32& OutWidth, int32& OutHeight, uint32 GPUIndex)
 {
 	ID3D11Texture2D* Texture = (ID3D11Texture2D*)(GetD3D11TextureFromRHITexture(TextureRHI)->GetResource());
 	
