Fixing access to parameters that are not in shipping and test

diff --git a/Engine/Plugins/FX/Niagara/Source/Niagara/Private/NiagaraRendererMeshes.cpp b/Engine/Plugins/FX/Niagara/Source/Niagara/Private/NiagaraRendererMeshes.cpp
--- a/Engine/Plugins/FX/Niagara/Source/Niagara/Private/NiagaraRendererMeshes.cpp
+++ b/Engine/Plugins/FX/Niagara/Source/Niagara/Private/NiagaraRendererMeshes.cpp
@@ -542,14 +542,19 @@ void FNiagaraRendererMeshes::GetDynamicRayTracingInstances(FRayTracingMaterialGa
 		MeshBatch.MaterialRenderProxy = MaterialProxy;
 		MeshBatch.SegmentIndex = SectionIndex;
 		MeshBatch.LODIndex = LODIndex;
+
+#if !(UE_BUILD_SHIPPING || UE_BUILD_TEST)
 		MeshBatch.VisualizeLODIndex = LODIndex;
+#endif
 		MeshBatch.CastRayTracedShadow = MeshBatch.CastShadow;
 		FMeshBatchElement& MeshBatchElement = MeshBatch.Elements[0];
 		MeshBatchElement.VertexFactoryUserData = VFs.VertexFactory.GetUniformBuffer();
 		MeshBatchElement.MinVertexIndex = Section.MinVertexIndex;
 		MeshBatchElement.MaxVertexIndex = Section.MaxVertexIndex;
+		
+#if !(UE_BUILD_SHIPPING || UE_BUILD_TEST)
 		MeshBatchElement.VisualizeElementIndex = SectionIndex;
-
+#endif
 		RayTracingInstance.Materials.Add(MeshBatch);
 	}
 
