Fixed ParameterOffsets not being converted to SortedParameterOffsets after being loaded from FNiagaraUserRedirectionParameterStore::SerializeFromMismatchedTag().

diff --git a/Engine/Plugins/FX/Niagara/Source/Niagara/Private/NiagaraUserRedirectionParameterStore.cpp b/Engine/Plugins/FX/Niagara/Source/Niagara/Private/NiagaraUserRedirectionParameterStore.cpp
--- a/Engine/Plugins/FX/Niagara/Source/Niagara/Private/NiagaraUserRedirectionParameterStore.cpp
+++ b/Engine/Plugins/FX/Niagara/Source/Niagara/Private/NiagaraUserRedirectionParameterStore.cpp
@@ -99,6 +99,8 @@ bool FNiagaraUserRedirectionParameterStore::SerializeFromMismatchedTag(const FPr
 	{
 		FNiagaraParameterStore OldStore;
 		FNiagaraParameterStore::StaticStruct()->SerializeItem(Slot, &OldStore, nullptr);
+		// Call PostLoad() to convert the serialized ParameterOffsets to SortedParameterOffsets.
+		OldStore.PostLoad();
 		*this = OldStore;
 		return true;
 	}
