Fixed open todo merging shader compile result

diff --git a/Engine/Plugins/FX/Niagara/Source/NiagaraEditor/Private/NiagaraCompiler.cpp b/Engine/Plugins/FX/Niagara/Source/NiagaraEditor/Private/NiagaraCompiler.cpp
--- a/Engine/Plugins/FX/Niagara/Source/NiagaraEditor/Private/NiagaraCompiler.cpp
+++ b/Engine/Plugins/FX/Niagara/Source/NiagaraEditor/Private/NiagaraCompiler.cpp
@@ -1065,11 +1065,13 @@ TOptional<FNiagaraCompileResults> FHlslNiagaraCompiler::GetCompileResult(int32 J
 				int32 NewBindingIdx = Results.Data->CalledVMExternalFunctions.AddDefaulted();
 				Results.Data->CalledVMExternalFunctions[NewBindingIdx].Name = *Sig->GetName();
 				Results.Data->CalledVMExternalFunctions[NewBindingIdx].OwnerName = Sig->OwnerName;
-				//#todo: fix me
-				//Results.Data->CalledVMExternalFunctions[NewBindingIdx].Specifiers = Sig->FunctionSpecifiers;
-
 				Results.Data->CalledVMExternalFunctions[NewBindingIdx].InputParamLocations = FuncInfo.InputParamLocations;
 				Results.Data->CalledVMExternalFunctions[NewBindingIdx].NumOutputs = FuncInfo.NumOutputs;
+				for (auto it = Sig->FunctionSpecifiers.CreateConstIterator(); it; ++it)
+				{
+					// we convert the map into an array here to save runtime memory
+					Results.Data->CalledVMExternalFunctions[NewBindingIdx].FunctionSpecifiers.Emplace(it->Key, it->Value);
+				}
 			}
 			else
 			{
