Fix for deprecation of ClearUAV with Grid2DCollection

Fix deprecation warning for ClearUAV with NiagaraDataInterfacePressureGrid.cpp

diff --git a/Engine/Plugins/Experimental/HairStrands/Source/HairStrandsNiagara/Private/NiagaraDataInterfacePressureGrid.cpp b/Engine/Plugins/Experimental/HairStrands/Source/HairStrandsNiagara/Private/NiagaraDataInterfacePressureGrid.cpp
--- a/Engine/Plugins/Experimental/HairStrands/Source/HairStrandsNiagara/Private/NiagaraDataInterfacePressureGrid.cpp
+++ b/Engine/Plugins/Experimental/HairStrands/Source/HairStrandsNiagara/Private/NiagaraDataInterfacePressureGrid.cpp
@@ -86,7 +86,7 @@ void FNDIPressureGridBuffer::ReleaseRHI()
 
 void FNDIPressureGridBuffer::ClearBuffers(FRHICommandList& RHICmdList)
 {
-	ClearUAV(RHICmdList, GridDataBuffer, FLinearColor(0, 0, 0, 0));
+	RHICmdList.ClearUAVFloat(GridDataBuffer.UAV, FVector4(ForceInitToZero));
 }
 
 //------------------------------------------------------------------------------------------------------------
diff --git a/Engine/Plugins/FX/Niagara/Source/Niagara/Private/NiagaraDataInterfaceGrid2DCollection.cpp b/Engine/Plugins/FX/Niagara/Source/Niagara/Private/NiagaraDataInterfaceGrid2DCollection.cpp
--- a/Engine/Plugins/FX/Niagara/Source/Niagara/Private/NiagaraDataInterfaceGrid2DCollection.cpp
+++ b/Engine/Plugins/FX/Niagara/Source/Niagara/Private/NiagaraDataInterfaceGrid2DCollection.cpp
@@ -698,7 +698,7 @@ void FNiagaraDataInterfaceProxyGrid2DCollection::PreStage(FRHICommandList& RHICm
 		// memory efficient since it would theoretically not require any double buffering.		
 		if (!Context.IsIterationStage)
 		{
-			RHICmdList.ClearUAVFloat(ProxyData->DestinationData->GridBuffer.UAV, FVector4(0, 0, 0, 0));
+			RHICmdList.ClearUAVFloat(ProxyData->DestinationData->GridBuffer.UAV, FVector4(ForceInitToZero));
 		}
 	}
 }
@@ -726,7 +726,7 @@ void FNiagaraDataInterfaceProxyGrid2DCollection::ResetData(FRHICommandList& RHIC
 	{
 		if (Buffer)
 		{
-			ClearUAV(RHICmdList, Buffer->GridBuffer, FLinearColor(0, 0, 0, 0));
+			RHICmdList.ClearUAVFloat(Buffer->GridBuffer.UAV, FVector4(ForceInitToZero));
 		}	
 	}	
 }
