Additional ClearUAV warning cleanup

diff --git a/Engine/Plugins/Experimental/HairStrands/Source/HairStrandsNiagara/Private/NiagaraDataInterfacePressureGrid.cpp b/Engine/Plugins/Experimental/HairStrands/Source/HairStrandsNiagara/Private/NiagaraDataInterfacePressureGrid.cpp
--- a/Engine/Plugins/Experimental/HairStrands/Source/HairStrandsNiagara/Private/NiagaraDataInterfacePressureGrid.cpp
+++ b/Engine/Plugins/Experimental/HairStrands/Source/HairStrandsNiagara/Private/NiagaraDataInterfacePressureGrid.cpp
@@ -86,7 +86,9 @@ void FNDIPressureGridBuffer::ReleaseRHI()
 
 void FNDIPressureGridBuffer::ClearBuffers(FRHICommandList& RHICmdList)
 {
+	RHICmdList.TransitionResource(EResourceTransitionAccess::ERWBarrier, EResourceTransitionPipeline::EComputeToCompute, GridDataBuffer.UAV);
 	RHICmdList.ClearUAVFloat(GridDataBuffer.UAV, FVector4(ForceInitToZero));
+	RHICmdList.TransitionResource(EResourceTransitionAccess::ERWBarrier, EResourceTransitionPipeline::EComputeToCompute, GridDataBuffer.UAV);
 }
 
 //------------------------------------------------------------------------------------------------------------
diff --git a/Engine/Plugins/FX/Niagara/Source/Niagara/Private/NiagaraDataInterfaceGrid2DCollection.cpp b/Engine/Plugins/FX/Niagara/Source/Niagara/Private/NiagaraDataInterfaceGrid2DCollection.cpp
--- a/Engine/Plugins/FX/Niagara/Source/Niagara/Private/NiagaraDataInterfaceGrid2DCollection.cpp
+++ b/Engine/Plugins/FX/Niagara/Source/Niagara/Private/NiagaraDataInterfaceGrid2DCollection.cpp
@@ -698,7 +698,9 @@ void FNiagaraDataInterfaceProxyGrid2DCollection::PreStage(FRHICommandList& RHICm
 		// memory efficient since it would theoretically not require any double buffering.		
 		if (!Context.IsIterationStage)
 		{
+			RHICmdList.TransitionResource(EResourceTransitionAccess::ERWBarrier, EResourceTransitionPipeline::EComputeToCompute, ProxyData->DestinationData->GridBuffer.UAV);
 			RHICmdList.ClearUAVFloat(ProxyData->DestinationData->GridBuffer.UAV, FVector4(ForceInitToZero));
+			RHICmdList.TransitionResource(EResourceTransitionAccess::ERWBarrier, EResourceTransitionPipeline::EComputeToCompute, ProxyData->DestinationData->GridBuffer.UAV);
 		}
 	}
 }
@@ -726,8 +728,10 @@ void FNiagaraDataInterfaceProxyGrid2DCollection::ResetData(FRHICommandList& RHIC
 	{
 		if (Buffer)
 		{
+			RHICmdList.TransitionResource(EResourceTransitionAccess::ERWBarrier, EResourceTransitionPipeline::EComputeToCompute, Buffer->GridBuffer.UAV);
 			RHICmdList.ClearUAVFloat(Buffer->GridBuffer.UAV, FVector4(ForceInitToZero));
-		}	
+			RHICmdList.TransitionResource(EResourceTransitionAccess::ERWBarrier, EResourceTransitionPipeline::EComputeToCompute, Buffer->GridBuffer.UAV);
+		}
 	}	
 }
 #undef LOCTEXT_NAMESPACE
\ No newline at end of file
diff --git a/Engine/Plugins/FX/Niagara/Source/Niagara/Private/NiagaraEmitterInstanceBatcher.cpp b/Engine/Plugins/FX/Niagara/Source/Niagara/Private/NiagaraEmitterInstanceBatcher.cpp
--- a/Engine/Plugins/FX/Niagara/Source/Niagara/Private/NiagaraEmitterInstanceBatcher.cpp
+++ b/Engine/Plugins/FX/Niagara/Source/Niagara/Private/NiagaraEmitterInstanceBatcher.cpp
@@ -452,7 +452,8 @@ void NiagaraEmitterInstanceBatcher::ResizeBuffersAndGatherResources(FOverlappabl
 	{
 		SCOPED_DRAW_EVENT(RHICmdList, NiagaraGPUClearFreeIDListSizeTable);
 		RHICmdList.TransitionResource(EResourceTransitionAccess::ERWBarrier, EResourceTransitionPipeline::EComputeToCompute, FreeIDListSizesBuffer.UAV);
-		ClearUAV(RHICmdList, FreeIDListSizesBuffer, 0);
+		RHICmdList.ClearUAVUint(FreeIDListSizesBuffer.UAV, FUintVector4(ForceInitToZero));
+		RHICmdList.TransitionResource(EResourceTransitionAccess::ERWBarrier, EResourceTransitionPipeline::EComputeToCompute, FreeIDListSizesBuffer.UAV);
 	}
 }
 
@@ -1199,10 +1200,12 @@ void NiagaraEmitterInstanceBatcher::Run(const FNiagaraGPUSystemTick& Tick, const
 	{
 		// Put INDEX_NONE in all the slots of the ID to index table. The simulation will fill in the
 		// indices for the IDs which are currently in use, so we can compute the free ID list based on
-		// the unused slots. No need for an explicit barrier here, ClearUAV() inserts a RW barrier.
+		// the unused slots.
 		SCOPED_DRAW_EVENT(RHICmdList, NiagaraGPUClearIDTable);
 		SCOPED_GPU_STAT(RHICmdList, NiagaraGPUClearIDTable);
-		ClearUAV(RHICmdList, DestinationData.GetGPUIDToIndexTable(), INDEX_NONE);
+		RHICmdList.TransitionResource(EResourceTransitionAccess::ERWBarrier, EResourceTransitionPipeline::EComputeToCompute, DestinationData.GetGPUIDToIndexTable().UAV);
+		RHICmdList.ClearUAVUint(DestinationData.GetGPUIDToIndexTable().UAV, FUintVector4(INDEX_NONE));
+		RHICmdList.TransitionResource(EResourceTransitionAccess::ERWBarrier, EResourceTransitionPipeline::EComputeToCompute, DestinationData.GetGPUIDToIndexTable().UAV);
 	}
 
 	RHICmdList.SetComputeShader(Shader->GetComputeShader());
