Possible fix for bone sampling crash.
When using a master pose component with SkinOnTheFly we could end up reading from an empty array.
The base skeletalmesh component's transform array is empty when using a master pose component.

diff --git a/Engine/Plugins/FX/Niagara/Source/Niagara/Classes/NiagaraDataInterfaceSkeletalMesh.h b/Engine/Plugins/FX/Niagara/Source/Niagara/Classes/NiagaraDataInterfaceSkeletalMesh.h
--- a/Engine/Plugins/FX/Niagara/Source/Niagara/Classes/NiagaraDataInterfaceSkeletalMesh.h
+++ b/Engine/Plugins/FX/Niagara/Source/Niagara/Classes/NiagaraDataInterfaceSkeletalMesh.h
@@ -145,6 +145,9 @@ struct FSkeletalMeshSkinningData
 	}
 
 private:
+
+	void UpdateBoneTransforms();
+
 	FCriticalSection CriticalSection; 
 	
 	TWeakObjectPtr<USkeletalMeshComponent> MeshComp;
diff --git a/Engine/Plugins/FX/Niagara/Source/Niagara/Private/NiagaraDataInterfaceSkeletalMesh.cpp b/Engine/Plugins/FX/Niagara/Source/Niagara/Private/NiagaraDataInterfaceSkeletalMesh.cpp
--- a/Engine/Plugins/FX/Niagara/Source/Niagara/Private/NiagaraDataInterfaceSkeletalMesh.cpp
+++ b/Engine/Plugins/FX/Niagara/Source/Niagara/Private/NiagaraDataInterfaceSkeletalMesh.cpp
@@ -117,9 +117,9 @@ void FSkeletalMeshSkinningData::RegisterUser(FSkeletalMeshSkinningDataUsage Usag
 	if (Usage.NeedsDataImmediately())
 	{
 		check(IsInGameThread());
-		if (CurrBoneRefToLocals().Num() == 0)
+		if (CurrBoneRefToLocals().Num() == 0 || CurrComponentTransforms().Num() == 0 )
 		{
-			SkelComp->CacheRefToLocalMatrices(CurrBoneRefToLocals()); 
+			UpdateBoneTransforms();
 		}
 		
 		//Prime the prev matrices if they're missing.
@@ -128,10 +128,6 @@ void FSkeletalMeshSkinningData::RegisterUser(FSkeletalMeshSkinningDataUsage Usag
 			PrevBoneRefToLocals() = CurrBoneRefToLocals();
 		}
 
-		if (CurrComponentTransforms().Num() == 0)
-		{
-			CurrComponentTransforms() = SkelComp->GetComponentSpaceTransforms();
-		}
 		if (PrevComponentTransforms().Num() != CurrComponentTransforms().Num())
 		{
 			PrevComponentTransforms() = CurrComponentTransforms();
@@ -187,66 +183,87 @@ bool FSkeletalMeshSkinningData::IsUsed()const
 	return false;
 }
 
-bool FSkeletalMeshSkinningData::Tick(float InDeltaSeconds, bool bRequirePreskin)
+void FSkeletalMeshSkinningData::UpdateBoneTransforms()
 {
 	USkeletalMeshComponent* SkelComp = MeshComp.Get();
 	check(SkelComp);
-	DeltaSeconds = InDeltaSeconds;
-	CurrIndex ^= 1;
 
-	if (BoneMatrixUsers > 0)
+	const TArray<FTransform>& BaseCompSpaceTransforms = SkelComp->GetComponentSpaceTransforms();
+	TArray<FMatrix>& CurrBones = CurrBoneRefToLocals();
+	TArray<FTransform>& CurrTransforms = CurrComponentTransforms();
+
+	if (USkinnedMeshComponent* MasterComponent = SkelComp->MasterPoseComponent.Get())
 	{
-		TArray<FMatrix>& CurrBones = CurrBoneRefToLocals();
-		if ( USkinnedMeshComponent* MasterComponent = SkelComp->MasterPoseComponent.Get() )
+		const USkeletalMesh* SkelMesh = SkelComp->SkeletalMesh;
+		const TArray<int32>& MasterBoneMap = SkelComp->GetMasterBoneMap();
+		const int32 NumBones = MasterBoneMap.Num();
+
+		check(SkelMesh);
+		if (NumBones == 0)
+		{
+			// This case indicates an invalid master pose component (e.g. no skeletal mesh)
+			CurrBones.Empty(SkelMesh->RefSkeleton.GetNum());
+			CurrBones.AddDefaulted(SkelMesh->RefSkeleton.GetNum());
+			CurrTransforms.Empty(SkelMesh->RefSkeleton.GetNum());
+			CurrTransforms.AddDefaulted(SkelMesh->RefSkeleton.GetNum());
+		}
+		else
 		{
-			const USkeletalMesh* SkelMesh = SkelComp->SkeletalMesh;
-			const TArray<int32>& MasterBoneMap = SkelComp->GetMasterBoneMap();
-			const int32 NumBones = MasterBoneMap.Num();
+			CurrBones.SetNumUninitialized(NumBones);
+			CurrTransforms.SetNumUninitialized(NumBones);
 
-			check(SkelMesh);
-			if ( NumBones == 0 )
+			const TArray<FTransform>& MasterTransforms = MasterComponent->GetComponentSpaceTransforms();
+			for (int32 BoneIndex = 0; BoneIndex < NumBones; ++BoneIndex)
 			{
-				// This case indicates an invalid master pose component (e.g. no skeletal mesh)
-				CurrBones.Empty(SkelMesh->RefSkeleton.GetNum());
-				CurrBones.AddDefaulted(SkelMesh->RefSkeleton.GetNum());
-			}
-			else
-			{
-				CurrBones.SetNumUninitialized(NumBones);
-
-				const TArray<FTransform>& MasterTransforms = MasterComponent->GetComponentSpaceTransforms();
-				for (int32 BoneIndex = 0; BoneIndex < NumBones; ++BoneIndex)
+				bool bFoundMaster = false;
+				FTransform CompSpaceTransform;
+				if (MasterBoneMap.IsValidIndex(BoneIndex))
 				{
-					bool bFoundMaster = false;
-					if (MasterBoneMap.IsValidIndex(BoneIndex))
+					const int32 MasterIndex = MasterBoneMap[BoneIndex];
+					if (MasterIndex != INDEX_NONE && MasterIndex < MasterTransforms.Num())
 					{
-						const int32 MasterIndex = MasterBoneMap[BoneIndex];
-						if (MasterIndex != INDEX_NONE && MasterIndex < MasterTransforms.Num())
-						{
-							CurrBones[BoneIndex] = MasterTransforms[MasterIndex].ToMatrixWithScale();
-							bFoundMaster = true;
-						}
+						CurrTransforms[BoneIndex] = MasterTransforms[MasterIndex];
+						CurrBones[BoneIndex] = SkelMesh->RefBasesInvMatrix[BoneIndex] * MasterTransforms[MasterIndex].ToMatrixWithScale();
+						bFoundMaster = true;
 					}
+				}
 
-					if ( !bFoundMaster )
-					{
-						const int32 ParentIndex = SkelMesh->RefSkeleton.GetParentIndex(BoneIndex);
+				if (!bFoundMaster)
+				{
+					const int32 ParentIndex = SkelMesh->RefSkeleton.GetParentIndex(BoneIndex);
 
-						CurrBones[BoneIndex] = SkelMesh->RefSkeleton.GetRefBonePose()[BoneIndex].ToMatrixWithScale();
-						if ( CurrBones.IsValidIndex(ParentIndex) && ParentIndex < BoneIndex )
-						{
-							CurrBones[BoneIndex] = CurrBones[ParentIndex] * CurrBones[BoneIndex];
-						}
+					if (CurrTransforms.IsValidIndex(ParentIndex) && ParentIndex < BoneIndex)
+					{
+						CurrTransforms[BoneIndex] = CurrTransforms[ParentIndex] * SkelMesh->RefSkeleton.GetRefBonePose()[BoneIndex];
+					}
+					else
+					{
+						CurrTransforms[BoneIndex] = SkelMesh->RefSkeleton.GetRefBonePose()[BoneIndex];
 					}
+
+					CurrBones[BoneIndex] = CurrTransforms[BoneIndex].ToMatrixWithScale();
+
 				}
 			}
 		}
-		else
-		{
-			SkelComp->CacheRefToLocalMatrices(CurrBones);
-		}
+	}
+	else
+	{
+		SkelComp->CacheRefToLocalMatrices(CurrBones);
+		CurrTransforms = SkelComp->GetComponentSpaceTransforms();
+	}
+}
 
-		CurrComponentTransforms() = SkelComp->GetComponentSpaceTransforms();
+bool FSkeletalMeshSkinningData::Tick(float InDeltaSeconds, bool bRequirePreskin)
+{
+	USkeletalMeshComponent* SkelComp = MeshComp.Get();
+	check(SkelComp);
+	DeltaSeconds = InDeltaSeconds;
+	CurrIndex ^= 1;
+
+	if (BoneMatrixUsers > 0)
+	{
+		UpdateBoneTransforms();
 	}
 
 	//Prime the prev matrices if they're missing.
@@ -2322,4 +2339,4 @@ void FSkeletalMeshAccessorHelper::Init<
 	SamplingRegionBuiltData = &SamplingInfo.GetRegionBuiltData(InstData->SamplingRegionIndices[0]);
 }
 
-#undef LOCTEXT_NAMESPACE
+#undef LOCTEXT_NAMESPACE
\ No newline at end of file
