Skipping some misc work, including potentally expensive GPU readback when GPU particles are not enabled.

diff --git a/Engine/Plugins/FX/Niagara/Source/Niagara/Private/NiagaraEmitterInstanceBatcher.cpp b/Engine/Plugins/FX/Niagara/Source/Niagara/Private/NiagaraEmitterInstanceBatcher.cpp
--- a/Engine/Plugins/FX/Niagara/Source/Niagara/Private/NiagaraEmitterInstanceBatcher.cpp
+++ b/Engine/Plugins/FX/Niagara/Source/Niagara/Private/NiagaraEmitterInstanceBatcher.cpp
@@ -128,6 +128,12 @@ void NiagaraEmitterInstanceBatcher::GiveSystemTick_RenderThread(FNiagaraGPUSyste
 {
 	check(IsInRenderingThread());
 
+	if (!FNiagaraUtilities::AllowGPUParticles(GetShaderPlatform()))
+	{
+		return;
+	}
+
+
 	// @todo REMOVE THIS HACK
 	if (GFrameNumberRenderThread > LastFrameThatDrainedData + GNiagaraGpuMaxQueuedRenderFrames)
 	{
@@ -554,6 +560,11 @@ void NiagaraEmitterInstanceBatcher::DispatchAllOnCompute(FOverlappableTicks& Ove
 
 void NiagaraEmitterInstanceBatcher::PostRenderOpaque(FRHICommandListImmediate& RHICmdList, FRHIUniformBuffer* ViewUniformBuffer, const class FShaderParametersMetadata* SceneTexturesUniformBufferStruct, FRHIUniformBuffer* SceneTexturesUniformBuffer, bool bAllowGPUParticleUpdate)
 {
+	if (!FNiagaraUtilities::AllowGPUParticles(GetShaderPlatform()))
+	{
+		return;
+	}
+
 	LLM_SCOPE(ELLMTag::Niagara);
 
 	if (bAllowGPUParticleUpdate)
@@ -808,6 +819,11 @@ void NiagaraEmitterInstanceBatcher::ExecuteAll(FRHICommandList& RHICmdList, FRHI
 
 void NiagaraEmitterInstanceBatcher::PreInitViews(FRHICommandListImmediate& RHICmdList, bool bAllowGPUParticleUpdate)
 {
+	if (!FNiagaraUtilities::AllowGPUParticles(GetShaderPlatform()))
+	{
+		return;
+	}
+
 	LLM_SCOPE(ELLMTag::Niagara);
 
 	// Reset the list of GPUSort tasks and release any resources they hold on to.
@@ -894,6 +910,11 @@ void NiagaraEmitterInstanceBatcher::PreInitViews(FRHICommandListImmediate& RHICm
 
 void NiagaraEmitterInstanceBatcher::PostInitViews(FRHICommandListImmediate& RHICmdList, FRHIUniformBuffer* ViewUniformBuffer, bool bAllowGPUParticleUpdate)
 {
+	if (!FNiagaraUtilities::AllowGPUParticles(GetShaderPlatform()))
+	{
+		return;
+	}
+
 	LLM_SCOPE(ELLMTag::Niagara);
 
 	if (bAllowGPUParticleUpdate)
@@ -943,6 +964,11 @@ bool NiagaraEmitterInstanceBatcher::RequiresEarlyViewUniformBuffer() const
 
 void NiagaraEmitterInstanceBatcher::PreRender(FRHICommandListImmediate& RHICmdList, const class FGlobalDistanceFieldParameterData* GlobalDistanceFieldParameterData, bool bAllowGPUParticleUpdate)
 {
+	if (!FNiagaraUtilities::AllowGPUParticles(GetShaderPlatform()))
+	{
+		return;
+	}
+
 	LLM_SCOPE(ELLMTag::Niagara);
 
 	GlobalDistanceFieldParams = GlobalDistanceFieldParameterData ? *GlobalDistanceFieldParameterData : FGlobalDistanceFieldParameterData();
