Fix for padding in pushing owner vectors (position/velocity/axis/scale) into the data sets used by the system scripts.

diff --git a/Engine/Plugins/FX/Niagara/Source/Niagara/Private/NiagaraSystem.cpp b/Engine/Plugins/FX/Niagara/Source/Niagara/Private/NiagaraSystem.cpp
--- a/Engine/Plugins/FX/Niagara/Source/Niagara/Private/NiagaraSystem.cpp
+++ b/Engine/Plugins/FX/Niagara/Source/Niagara/Private/NiagaraSystem.cpp
@@ -1568,7 +1568,16 @@ void FNiagaraParameterDataSetBindingCollection::BuildInternal(const TArray<FNiag
 			}
 		}
 
-		ParameterOffset += Var.GetSizeInBytes();
+		// we need to take into account potential padding that is in the constant buffers based similar to what is done
+		// in the NiagaraHlslTranslator, where Vec2/Vec3 are treated as Vec4.
+		int32 ParameterSize = Var.GetSizeInBytes();
+		const FNiagaraTypeDefinition& Type = Var.GetType();
+		if (Type == FNiagaraTypeDefinition::GetVec2Def() || Type == FNiagaraTypeDefinition::GetVec3Def())
+		{
+			ParameterSize = Align(ParameterSize, FNiagaraTypeDefinition::GetVec4Def().GetSize());
+		}
+
+		ParameterOffset += ParameterSize;
 	}
 
 	FloatOffsets.Shrink();
