It is no longer possible to rename emitters to reserved names: [Emitter, Engine, User, Local, Module, NPC, Particles, System, Transient]. Doing so could cause the system to fail to compile.

Names are now verified when changed from the Overview and the Timeline panels as well as the Selection panel.

Removed "Rename Emitter" command, since emitter context menus now all include the generic "Rename" command.

diff --git a/Engine/Plugins/FX/Niagara/Source/NiagaraEditor/Private/NiagaraEditorUtilities.cpp b/Engine/Plugins/FX/Niagara/Source/NiagaraEditor/Private/NiagaraEditorUtilities.cpp
--- a/Engine/Plugins/FX/Niagara/Source/NiagaraEditor/Private/NiagaraEditorUtilities.cpp
+++ b/Engine/Plugins/FX/Niagara/Source/NiagaraEditor/Private/NiagaraEditorUtilities.cpp
@@ -1856,15 +1856,6 @@ bool FNiagaraEditorUtilities::AddEmitterContextMenuActions(FMenuBuilder& MenuBui
 					NAME_None,
 					EUserInterfaceActionType::ToggleButton
 				);
-
-				MenuBuilder.AddMenuEntry(
-					LOCTEXT("RenameEmitter", "Rename Emitter"),
-					LOCTEXT("RenameEmitterToolTip", "Rename this local emitter copy."),
-					FSlateIcon(),
-					FUIAction(
-						FExecuteAction::CreateSP(EmitterHandleViewModel, &FNiagaraEmitterHandleViewModel::SetIsRenamePending, true)
-					)
-				);
 			}
 
 			MenuBuilder.AddMenuEntry(
diff --git a/Engine/Plugins/FX/Niagara/Source/NiagaraEditor/Private/Sequencer/NiagaraSequence/MovieSceneNiagaraEmitterTrack.cpp b/Engine/Plugins/FX/Niagara/Source/NiagaraEditor/Private/Sequencer/NiagaraSequence/MovieSceneNiagaraEmitterTrack.cpp
--- a/Engine/Plugins/FX/Niagara/Source/NiagaraEditor/Private/Sequencer/NiagaraSequence/MovieSceneNiagaraEmitterTrack.cpp
+++ b/Engine/Plugins/FX/Niagara/Source/NiagaraEditor/Private/Sequencer/NiagaraSequence/MovieSceneNiagaraEmitterTrack.cpp
@@ -97,6 +97,11 @@ bool UMovieSceneNiagaraEmitterTrack::CanRename() const
 	return false;
 }
 
+bool UMovieSceneNiagaraEmitterTrack::ValidateDisplayName(const FText& NewDisplayName, FText& OutErrorMessage) const
+{
+	return EmitterHandleViewModel.Pin()->VerifyNameTextChanged(NewDisplayName, OutErrorMessage);
+}
+
 FNiagaraSystemViewModel& UMovieSceneNiagaraEmitterTrack::GetSystemViewModel() const
 {
 	return *SystemViewModel;
diff --git a/Engine/Plugins/FX/Niagara/Source/NiagaraEditor/Private/Sequencer/NiagaraSequence/MovieSceneNiagaraEmitterTrack.h b/Engine/Plugins/FX/Niagara/Source/NiagaraEditor/Private/Sequencer/NiagaraSequence/MovieSceneNiagaraEmitterTrack.h
--- a/Engine/Plugins/FX/Niagara/Source/NiagaraEditor/Private/Sequencer/NiagaraSequence/MovieSceneNiagaraEmitterTrack.h
+++ b/Engine/Plugins/FX/Niagara/Source/NiagaraEditor/Private/Sequencer/NiagaraSequence/MovieSceneNiagaraEmitterTrack.h
@@ -82,6 +82,7 @@ public:
 	virtual bool IsEmpty() const override;
 	virtual const TArray<UMovieSceneSection*>& GetAllSections() const override;
 	virtual bool SupportsMultipleRows() const override;
+	virtual bool ValidateDisplayName(const FText& NewDisplayName, FText& OutErrorMessage) const override;
 
 	/** Gets the unique id for the emitter handle that was associated with this track; used for copy/paste detection */
 	FGuid GetEmitterHandleId() const;
diff --git a/Engine/Plugins/FX/Niagara/Source/NiagaraEditor/Private/ViewModels/NiagaraEmitterHandleViewModel.cpp b/Engine/Plugins/FX/Niagara/Source/NiagaraEditor/Private/ViewModels/NiagaraEmitterHandleViewModel.cpp
--- a/Engine/Plugins/FX/Niagara/Source/NiagaraEditor/Private/ViewModels/NiagaraEmitterHandleViewModel.cpp
+++ b/Engine/Plugins/FX/Niagara/Source/NiagaraEditor/Private/ViewModels/NiagaraEmitterHandleViewModel.cpp
@@ -14,9 +14,9 @@
 #include "NiagaraGraph.h"
 #include "NiagaraNodeInput.h"
 #include "NiagaraScriptOutputCollectionViewModel.h"
+#include "Algo/Find.h"
 #include "Framework/Notifications/NotificationManager.h"
 #include "Widgets/Notifications/SNotificationList.h"
-
 #include "ScopedTransaction.h"
 #include "NiagaraRendererProperties.h"
 #include "ViewModels/Stack/NiagaraStackRoot.h"
@@ -208,12 +208,29 @@ void FNiagaraEmitterHandleViewModel::OnNameTextComitted(const FText& InText, ETe
 
 bool FNiagaraEmitterHandleViewModel::VerifyNameTextChanged(const FText& NewText, FText& OutErrorMessage)
 {
-	FName NewName = *NewText.ToString();
-	if (NewName == FName())
+	static const FString ReservedNames[] =
+	{
+		TEXT("Particles"),
+		TEXT("Local"),
+		TEXT("Engine"),
+		TEXT("Transient"),
+		TEXT("User"),
+		TEXT("Emitter"),
+		TEXT("Module"),
+		TEXT("NPC")
+	};
+
+	FString NewString = NewText.ToString();
+	if (NewText.IsEmptyOrWhitespace())
 	{
 		OutErrorMessage = NSLOCTEXT("NiagaraEmitterEditor", "NiagaraInputNameEmptyWarn", "Cannot have empty name!");
 		return false;
 	}
+	else if (Algo::Find(ReservedNames, NewString) != nullptr)
+	{
+		OutErrorMessage = FText::Format(NSLOCTEXT("NiagaraEmitterEditor", "NiagaraInputNameReservedWarn", "Cannot use reserved name \"{0}\"!"), NewText);
+		return false;
+	}
 	return true;
 }
 
diff --git a/Engine/Plugins/FX/Niagara/Source/NiagaraEditor/Private/ViewModels/NiagaraSystemViewModel.cpp b/Engine/Plugins/FX/Niagara/Source/NiagaraEditor/Private/ViewModels/NiagaraSystemViewModel.cpp
--- a/Engine/Plugins/FX/Niagara/Source/NiagaraEditor/Private/ViewModels/NiagaraSystemViewModel.cpp
+++ b/Engine/Plugins/FX/Niagara/Source/NiagaraEditor/Private/ViewModels/NiagaraSystemViewModel.cpp
@@ -516,6 +516,11 @@ void FNiagaraSystemViewModel::Tick(float DeltaTime)
 		SendLastCompileMessageJobs();
 	}
 
+	if (!SystemScriptViewModel.IsValid())
+	{
+		return;
+	}
+
 	if (bForceAutoCompileOnce || (GetDefault<UNiagaraEditorSettings>()->GetAutoCompile() && bCanAutoCompile))
 	{
 		bool bRecompile = false;
diff --git a/Engine/Plugins/FX/Niagara/Source/NiagaraEditorWidgets/Private/SNiagaraOverviewGraph.cpp b/Engine/Plugins/FX/Niagara/Source/NiagaraEditorWidgets/Private/SNiagaraOverviewGraph.cpp
--- a/Engine/Plugins/FX/Niagara/Source/NiagaraEditorWidgets/Private/SNiagaraOverviewGraph.cpp
+++ b/Engine/Plugins/FX/Niagara/Source/NiagaraEditorWidgets/Private/SNiagaraOverviewGraph.cpp
@@ -44,6 +44,7 @@ void SNiagaraOverviewGraph::Construct(const FArguments& InArgs, TSharedRef<FNiag
 	SGraphEditor::FGraphEditorEvents Events;
 	Events.OnSelectionChanged = SGraphEditor::FOnSelectionChanged::CreateSP(this, &SNiagaraOverviewGraph::GraphSelectionChanged);
 	Events.OnCreateActionMenu = SGraphEditor::FOnCreateActionMenu::CreateSP(this, &SNiagaraOverviewGraph::OnCreateGraphActionMenu);
+	Events.OnVerifyTextCommit = FOnNodeVerifyTextCommit::CreateSP(this, &SNiagaraOverviewGraph::OnVerifyNodeTitle);
 	Events.OnTextCommitted = FOnNodeTextCommitted::CreateSP(this, &SNiagaraOverviewGraph::OnNodeTitleCommitted);
 
 	FGraphAppearanceInfo AppearanceInfo;
@@ -251,6 +252,19 @@ void SNiagaraOverviewGraph::OnClearIsolated()
 	ViewModel->GetSystemViewModel()->IsolateEmitters(TArray<FGuid>());
 }
 
+bool SNiagaraOverviewGraph::OnVerifyNodeTitle(const FText& NewText, UEdGraphNode* Node, FText& OutErrorMessage) const
+{
+	UNiagaraOverviewNode* NiagaraNode = Cast<UNiagaraOverviewNode>(Node);
+
+	TSharedPtr<FNiagaraEmitterHandleViewModel> NodeEmitterHandleViewModel = ViewModel->GetSystemViewModel()->GetEmitterHandleViewModelById(NiagaraNode->GetEmitterHandleGuid());
+	if (ensureMsgf(NodeEmitterHandleViewModel.IsValid(), TEXT("Failed to find EmitterHandleViewModel with matching Emitter GUID to Overview Node!")))
+	{
+		return NodeEmitterHandleViewModel->VerifyNameTextChanged(NewText, OutErrorMessage);
+	}
+
+	return true;
+}
+
 void SNiagaraOverviewGraph::OnNodeTitleCommitted(const FText& NewText, ETextCommit::Type CommitInfo, UEdGraphNode* NodeBeingChanged)
 {
 	if (NodeBeingChanged)
diff --git a/Engine/Plugins/FX/Niagara/Source/NiagaraEditorWidgets/Private/SNiagaraOverviewGraph.h b/Engine/Plugins/FX/Niagara/Source/NiagaraEditorWidgets/Private/SNiagaraOverviewGraph.h
--- a/Engine/Plugins/FX/Niagara/Source/NiagaraEditorWidgets/Private/SNiagaraOverviewGraph.h
+++ b/Engine/Plugins/FX/Niagara/Source/NiagaraEditorWidgets/Private/SNiagaraOverviewGraph.h
@@ -38,6 +38,7 @@ private:
 	void OnCreateComment();
 	void OnClearIsolated();
 
+	bool OnVerifyNodeTitle(const FText& NewText, UEdGraphNode* Node, FText& OutErrorMessage) const;
 	void OnNodeTitleCommitted(const FText& NewText, ETextCommit::Type CommitInfo, UEdGraphNode* NodeBeingChanged);
 
 	void CreateAddEmitterMenuContent(FMenuBuilder& MenuBuilder, UEdGraph* InGraph);
diff --git a/Engine/Source/Editor/Sequencer/Private/DisplayNodes/SequencerTrackNode.cpp b/Engine/Source/Editor/Sequencer/Private/DisplayNodes/SequencerTrackNode.cpp
--- a/Engine/Source/Editor/Sequencer/Private/DisplayNodes/SequencerTrackNode.cpp
+++ b/Engine/Source/Editor/Sequencer/Private/DisplayNodes/SequencerTrackNode.cpp
@@ -345,6 +345,16 @@ bool FSequencerTrackNode::CanRenameNode() const
 }
 
 
+bool FSequencerTrackNode::ValidateDisplayName(const FText& NewDisplayName, FText& OutErrorMessage) const
+{
+	auto NameableTrack = Cast<UMovieSceneNameableTrack>(AssociatedTrack.Get());
+	if (NameableTrack != nullptr)
+	{
+		return NameableTrack->ValidateDisplayName(NewDisplayName, OutErrorMessage);
+	}
+	return true;
+}
+
 FReply FSequencerTrackNode::CreateNewSection() const
 {
 	UMovieSceneTrack* Track = GetTrack();
diff --git a/Engine/Source/Editor/Sequencer/Private/DisplayNodes/SequencerTrackNode.h b/Engine/Source/Editor/Sequencer/Private/DisplayNodes/SequencerTrackNode.h
--- a/Engine/Source/Editor/Sequencer/Private/DisplayNodes/SequencerTrackNode.h
+++ b/Engine/Source/Editor/Sequencer/Private/DisplayNodes/SequencerTrackNode.h
@@ -127,6 +127,7 @@ public:
 
 	virtual void BuildContextMenu( FMenuBuilder& MenuBuilder );
 	virtual bool CanRenameNode() const override;
+	virtual bool ValidateDisplayName(const FText& NewDisplayName, FText& OutErrorMessage) const override;
 	virtual TSharedRef<SWidget> GetCustomOutlinerContent() override;
 	virtual void GetChildKeyAreaNodesRecursively(TArray<TSharedRef<FSequencerSectionKeyAreaNode>>& OutNodes) const override;
 	virtual FText GetDisplayName() const override;
diff --git a/Engine/Source/Runtime/MovieScene/Public/MovieSceneNameableTrack.h b/Engine/Source/Runtime/MovieScene/Public/MovieSceneNameableTrack.h
--- a/Engine/Source/Runtime/MovieScene/Public/MovieSceneNameableTrack.h
+++ b/Engine/Source/Runtime/MovieScene/Public/MovieSceneNameableTrack.h
@@ -34,6 +34,14 @@ public:
 	 */
 	MOVIESCENE_API virtual bool CanRename() const { return true; }
 
+	/** 
+	 * Validate the new display name. 
+	 *
+	 * @return True if the given display name is valid, false if it is not. 
+	 * Error message should be set if the name is not valid.
+	 */
+	MOVIESCENE_API virtual bool ValidateDisplayName(const FText& NewDisplayName, FText& OutErrorMessage) const { return true; }
+
 public:
 
 	// UMovieSceneTrack interface
