Niagara "Set Variables" modules now display the name of the variable that they are setting, and the number of additional variables changed in parentheses.

Set Specific Parameters menu entries now display the actual parameter name instead of adding spaces for no reason.

diff --git a/Engine/Plugins/FX/Niagara/Source/NiagaraEditor/Private/NiagaraNodeAssignment.cpp b/Engine/Plugins/FX/Niagara/Source/NiagaraEditor/Private/NiagaraNodeAssignment.cpp
--- a/Engine/Plugins/FX/Niagara/Source/NiagaraEditor/Private/NiagaraNodeAssignment.cpp
+++ b/Engine/Plugins/FX/Niagara/Source/NiagaraEditor/Private/NiagaraNodeAssignment.cpp
@@ -32,7 +32,37 @@ void UNiagaraNodeAssignment::AllocateDefaultPins()
 
 FText UNiagaraNodeAssignment::GetNodeTitle(ENodeTitleType::Type TitleType) const
 {
-	return LOCTEXT("NodeTitle", "Set Variables");
+	return Title;
+}
+
+void UNiagaraNodeAssignment::RefreshTitle()
+{
+	if (AssignmentTargets.Num() == 1)
+	{
+		Title = FText::Format(LOCTEXT("NodeTitleSingle", "Set {0}"), FText::FromName(AssignmentTargets[0].GetName()));
+	}
+	else if (AssignmentTargets.Num() > 1)
+	{
+		Title = FText::Format(LOCTEXT("NodeTitleMultiple", "Set {0} (+{1})"), FText::FromName(AssignmentTargets[0].GetName()), AssignmentTargets.Num() - 1);
+	}
+	else
+	{
+		Title = LOCTEXT("NodeTitle", "Set Variables");
+	}
+}
+
+FText UNiagaraNodeAssignment::GetTooltipText() const
+{
+	FText BaseText = LOCTEXT("NodeTooltipFormat", "Sets these parameters in the stack:");
+
+	TArray<FText> TargetNames;
+	TargetNames.Add(BaseText);
+	for (const FNiagaraVariable& Var : AssignmentTargets)
+	{
+		TargetNames.Add(FText::Format(LOCTEXT("Indent", "\t{0}"), FText::FromName(Var.GetName())));
+	}
+
+	return FText::Join(FText::FromString("\n"), TargetNames);
 }
 
 bool UNiagaraNodeAssignment::RefreshFromExternalChanges()
@@ -40,6 +70,7 @@ bool UNiagaraNodeAssignment::RefreshFromExternalChanges()
 	FunctionScript = nullptr;
 	GenerateScript();
 	ReallocatePins();
+	RefreshTitle();
 	return true;
 }
 
@@ -175,6 +206,8 @@ void UNiagaraNodeAssignment::PostLoad()
 			}
 		}
 	}
+
+	RefreshTitle();
 }
 
 void UNiagaraNodeAssignment::BuildParameterMapHistory(FNiagaraParameterMapHistoryBuilder& OutHistory, bool bRecursive /*= true*/, bool bFilterForCompilation /*= true*/) const
@@ -213,11 +246,10 @@ void UNiagaraNodeAssignment::BuildAddParameterMenu(FMenuBuilder& MenuBuilder, EN
 
 	for (const FNiagaraVariable& AvailableParameter : AvailableParameters)
 	{
-		FString DisplayNameString = FName::NameToDisplayString(AvailableParameter.GetName().ToString(), false);
-		const FText NameText = FText::FromString(DisplayNameString);
+		const FText NameText = FText::FromName(AvailableParameter.GetName());
 		FText VarDesc = FNiagaraConstants::GetAttributeDescription(AvailableParameter);
 		FString VarDefaultValue = FNiagaraConstants::GetAttributeDefaultValue(AvailableParameter);
-		const FText TooltipDesc = FText::Format(LOCTEXT("SetFunctionPopupTooltip", "Description: Set the parameter {0}. {1}"), FText::FromName(AvailableParameter.GetName()), VarDesc);
+		const FText TooltipDesc = FText::Format(LOCTEXT("SetFunctionPopupTooltip", "Description: Set the parameter {0}. {1}"), NameText, VarDesc);
 		FText CategoryName = LOCTEXT("ModuleSetCategory", "Set Specific Parameters");
 		bool bCanExecute = AssignmentTargets.Contains(AvailableParameter) == false; 
 
@@ -302,6 +334,8 @@ void UNiagaraNodeAssignment::AddParameter(FNiagaraVariable InVar, FString InDefa
 	RefreshFromExternalChanges();
 	MarkNodeRequiresSynchronization(__FUNCTION__, true);
 	OnInputsChangedDelegate.Broadcast();
+
+	RefreshTitle();
 }
 
 void UNiagaraNodeAssignment::RemoveParameter(const FNiagaraVariable& InVar)
@@ -330,6 +364,8 @@ void UNiagaraNodeAssignment::RemoveParameter(const FNiagaraVariable& InVar)
 	RefreshFromExternalChanges();
 	MarkNodeRequiresSynchronization(__FUNCTION__, true);
 	OnInputsChangedDelegate.Broadcast();
+
+	RefreshTitle();
 }
 
 void UNiagaraNodeAssignment::UpdateUsageBitmaskFromOwningScript()
@@ -520,6 +556,8 @@ void UNiagaraNodeAssignment::InitializeScript(UNiagaraScript* NewScript)
 			}
 		}
 	}
+
+	RefreshTitle();
 }
 
 int32 UsageToBitmask(ENiagaraScriptUsage Usage)
@@ -629,6 +667,7 @@ bool UNiagaraNodeAssignment::SetAssignmentTargetName(int32 Idx, const FName& InN
 	{
 		AssignmentTargets[Idx].SetName(InName);
 		MarkNodeRequiresSynchronization(__FUNCTION__, true);
+		RefreshTitle();
 		return true;
 	}
 	return false;
diff --git a/Engine/Plugins/FX/Niagara/Source/NiagaraEditor/Private/ViewModels/Stack/NiagaraStackScriptItemGroup.cpp b/Engine/Plugins/FX/Niagara/Source/NiagaraEditor/Private/ViewModels/Stack/NiagaraStackScriptItemGroup.cpp
--- a/Engine/Plugins/FX/Niagara/Source/NiagaraEditor/Private/ViewModels/Stack/NiagaraStackScriptItemGroup.cpp
+++ b/Engine/Plugins/FX/Niagara/Source/NiagaraEditor/Private/ViewModels/Stack/NiagaraStackScriptItemGroup.cpp
@@ -80,9 +80,7 @@ public:
 	{
 		FText Category = LOCTEXT("ExistingParameterModuleCategory", "Set Specific Parameters");
 
-		FString DisplayNameString = FName::NameToDisplayString(ParameterVariable.GetName().ToString(), false);
-		FText DisplayName = FText::FromString(DisplayNameString);
-
+		FText DisplayName = FText::FromName(ParameterVariable.GetName());
 		FText AttributeDescription = FNiagaraConstants::GetAttributeDescription(ParameterVariable);
 		FText Description = FText::Format(LOCTEXT("ExistingParameterModuleDescriptoinFormat", "Description: Set the parameter {0}. {1}"), FText::FromName(ParameterVariable.GetName()), AttributeDescription);
 
diff --git a/Engine/Plugins/FX/Niagara/Source/NiagaraEditor/Public/NiagaraNodeAssignment.h b/Engine/Plugins/FX/Niagara/Source/NiagaraEditor/Public/NiagaraNodeAssignment.h
--- a/Engine/Plugins/FX/Niagara/Source/NiagaraEditor/Public/NiagaraNodeAssignment.h
+++ b/Engine/Plugins/FX/Niagara/Source/NiagaraEditor/Public/NiagaraNodeAssignment.h
@@ -52,6 +52,8 @@ public:
 
 	void UpdateUsageBitmaskFromOwningScript();
 
+	virtual FText GetTooltipText() const override;
+
 protected:
 
 	UPROPERTY()
@@ -69,11 +71,15 @@ protected:
 	UPROPERTY()
 	FString OldFunctionCallName;
 
+	FText Title;
+
 private:
 	void GenerateScript();
 
 	void InitializeScript(UNiagaraScript* NewScript);
 
 	int32 CalculateScriptUsageBitmask();
+
+	void RefreshTitle();
 };
 
