Fixed issue with effect type not saving it's custom version.
It would cause us to reapply some postload code in error.

diff --git a/Engine/Plugins/FX/Niagara/Source/Niagara/Classes/NiagaraEffectType.h b/Engine/Plugins/FX/Niagara/Source/Niagara/Classes/NiagaraEffectType.h
--- a/Engine/Plugins/FX/Niagara/Source/Niagara/Classes/NiagaraEffectType.h
+++ b/Engine/Plugins/FX/Niagara/Source/Niagara/Classes/NiagaraEffectType.h
@@ -189,6 +189,7 @@ class NIAGARA_API UNiagaraEffectType : public UObject
 	//UObject Interface
 	virtual void BeginDestroy()override;
 	virtual bool IsReadyForFinishDestroy()override;
+	virtual void Serialize(FArchive& Ar)override;
 	virtual void PostLoad()override;
 #if WITH_EDITOR
 	virtual void PostEditChangeProperty(struct FPropertyChangedEvent& PropertyChangedEvent);
diff --git a/Engine/Plugins/FX/Niagara/Source/Niagara/Private/NiagaraEffectType.cpp b/Engine/Plugins/FX/Niagara/Source/Niagara/Private/NiagaraEffectType.cpp
--- a/Engine/Plugins/FX/Niagara/Source/Niagara/Private/NiagaraEffectType.cpp
+++ b/Engine/Plugins/FX/Niagara/Source/Niagara/Private/NiagaraEffectType.cpp
@@ -39,6 +39,12 @@ bool UNiagaraEffectType::IsReadyForFinishDestroy()
 	return ReleaseFence.IsFenceComplete() && Super::IsReadyForFinishDestroy();
 }
 
+void UNiagaraEffectType::Serialize(FArchive& Ar)
+{
+	Super::Serialize(Ar);
+	Ar.UsingCustomVersion(FNiagaraCustomVersion::GUID);
+}
+
 void UNiagaraEffectType::PostLoad()
 {
 	Super::PostLoad();
