Changing the trashcan button to go through the same path as the other delete methods.
#rnx

diff --git a/Engine/Plugins/FX/Niagara/Source/NiagaraEditorWidgets/Private/Stack/SNiagaraStackModuleItem.cpp b/Engine/Plugins/FX/Niagara/Source/NiagaraEditorWidgets/Private/Stack/SNiagaraStackModuleItem.cpp
--- a/Engine/Plugins/FX/Niagara/Source/NiagaraEditorWidgets/Private/Stack/SNiagaraStackModuleItem.cpp
+++ b/Engine/Plugins/FX/Niagara/Source/NiagaraEditorWidgets/Private/Stack/SNiagaraStackModuleItem.cpp
@@ -27,6 +27,7 @@
 #include "Subsystems/AssetEditorSubsystem.h"
 #include "Editor.h"
 #include "EditorFontGlyphs.h"
+#include "ViewModels/Stack/NiagaraStackClipboardUtilities.h"
 
 #define LOCTEXT_NAMESPACE "NiagaraStackModuleItem"
 
@@ -216,7 +217,9 @@ EVisibility SNiagaraStackModuleItem::GetRefreshVisibility() const
 
 FReply SNiagaraStackModuleItem::DeleteClicked()
 {
-	ModuleItem->Delete();
+	TArray<UNiagaraStackEntry*> EntriesToDelete;
+	EntriesToDelete.Add(ModuleItem);
+	FNiagaraStackClipboardUtilities::DeleteSelection(EntriesToDelete);
 	return FReply::Handled();
 }
 
