Removing error logging from a step in the enum loading process - real errors are handled after both methods are tried.
#rnx

diff --git a/Engine/Plugins/FX/Niagara/Source/NiagaraEditor/Private/TypeEditorUtilities/NiagaraEnumTypeEditorUtilities.cpp b/Engine/Plugins/FX/Niagara/Source/NiagaraEditor/Private/TypeEditorUtilities/NiagaraEnumTypeEditorUtilities.cpp
--- a/Engine/Plugins/FX/Niagara/Source/NiagaraEditor/Private/TypeEditorUtilities/NiagaraEnumTypeEditorUtilities.cpp
+++ b/Engine/Plugins/FX/Niagara/Source/NiagaraEditor/Private/TypeEditorUtilities/NiagaraEnumTypeEditorUtilities.cpp
@@ -110,7 +110,7 @@ bool FNiagaraEditorEnumTypeUtilities::SetValueFromPinDefaultString(const FString
 	checkf(Enum != nullptr, TEXT("Variable is not an enum type."));
 
 	FNiagaraInt32 EnumValue;
-	EnumValue.Value = (int32)Enum->GetValueByNameString(StringValue, EGetByNameFlags::ErrorIfNotFound);
+	EnumValue.Value = (int32)Enum->GetValueByNameString(StringValue);
 	if(EnumValue.Value != INDEX_NONE)
 	{
 		Variable.AllocateData();
