Define COMPUTESHADER for niagara GPU sim, to properly cull functionality that can't be used in compute

diff --git a/Engine/Plugins/FX/Niagara/Source/NiagaraCore/Private/NiagaraCustomVersion.cpp b/Engine/Plugins/FX/Niagara/Source/NiagaraCore/Private/NiagaraCustomVersion.cpp
--- a/Engine/Plugins/FX/Niagara/Source/NiagaraCore/Private/NiagaraCustomVersion.cpp
+++ b/Engine/Plugins/FX/Niagara/Source/NiagaraCore/Private/NiagaraCustomVersion.cpp
@@ -8,4 +8,4 @@ const FGuid FNiagaraCustomVersion::GUID(0xFCF57AFA, 0x50764283, 0xB9A9E658, 0xFF
 // Register the custom version with core
 FCustomVersionRegistration GRegisterNiagaraCustomVersion(FNiagaraCustomVersion::GUID, FNiagaraCustomVersion::LatestVersion, TEXT("NiagaraVer"));
 
-const FGuid FNiagaraCustomVersion::LatestScriptCompileVersion(0x6B0875B7, 0x486C447E, 0x8FF73178, 0xAE5C0EFC);
+const FGuid FNiagaraCustomVersion::LatestScriptCompileVersion(0xB5AC73B2, 0x6F2A47E3, 0x804376B4, 0xA03114E2);
diff --git a/Engine/Plugins/FX/Niagara/Source/NiagaraEditor/Private/NiagaraCompiler.cpp b/Engine/Plugins/FX/Niagara/Source/NiagaraEditor/Private/NiagaraCompiler.cpp
--- a/Engine/Plugins/FX/Niagara/Source/NiagaraEditor/Private/NiagaraCompiler.cpp
+++ b/Engine/Plugins/FX/Niagara/Source/NiagaraEditor/Private/NiagaraCompiler.cpp
@@ -712,6 +712,12 @@ void FNiagaraEditorModule::TestCompileScriptFromConsole(const TArray<FString>& A
 			Input.VirtualSourceFilePath = TEXT("/Plugin/FX/Niagara/Private/NiagaraEmitterInstanceShader.usf");
 			Input.EntryPointName = TEXT("SimulateMain");
 			Input.Environment.SetDefine(TEXT("VM_SIMULATION"), 1);
+			Input.Environment.SetDefine(TEXT("COMPUTESHADER"), 1);
+			Input.Environment.SetDefine(TEXT("PIXELSHADER"), 0);
+			Input.Environment.SetDefine(TEXT("DOMAINSHADER"), 0);
+			Input.Environment.SetDefine(TEXT("HULLSHADER"), 0);
+			Input.Environment.SetDefine(TEXT("VERTEXSHADER"), 0);
+			Input.Environment.SetDefine(TEXT("GEOMETRYSHADER"), 0);
 			Input.Environment.IncludeVirtualPathToContentsMap.Add(TEXT("/Engine/Generated/NiagaraEmitterInstance.ush"), TranslatedHLSL);
 
 			FShaderCompilerOutput Output;
@@ -800,6 +806,12 @@ int32 FHlslNiagaraCompiler::CompileScript(const FNiagaraCompileRequestData* InCo
 	Input.VirtualSourceFilePath = TEXT("/Plugin/FX/Niagara/Private/NiagaraEmitterInstanceShader.usf");
 	Input.EntryPointName = TEXT("SimulateMain");
 	Input.Environment.SetDefine(TEXT("VM_SIMULATION"), 1);
+	Input.Environment.SetDefine(TEXT("COMPUTESHADER"), 1);
+	Input.Environment.SetDefine(TEXT("PIXELSHADER"), 0);
+	Input.Environment.SetDefine(TEXT("DOMAINSHADER"), 0);
+	Input.Environment.SetDefine(TEXT("HULLSHADER"), 0);
+	Input.Environment.SetDefine(TEXT("VERTEXSHADER"), 0);
+	Input.Environment.SetDefine(TEXT("GEOMETRYSHADER"), 0);
 	Input.Environment.IncludeVirtualPathToContentsMap.Add(TEXT("/Engine/Generated/NiagaraEmitterInstance.ush"), TranslatedHLSL);
 	Input.bGenerateDirectCompileFile = false;
 	Input.DumpDebugInfoRootPath = GShaderCompilingManager->GetAbsoluteShaderDebugInfoDirectory() / TEXT("VM");
