Fixed: GPU systems which ran in PostInitViews never read back their instance counts.

diff --git a/Engine/Plugins/FX/Niagara/Source/Niagara/Classes/NiagaraEmitterInstanceBatcher.h b/Engine/Plugins/FX/Niagara/Source/Niagara/Classes/NiagaraEmitterInstanceBatcher.h
--- a/Engine/Plugins/FX/Niagara/Source/Niagara/Classes/NiagaraEmitterInstanceBatcher.h
+++ b/Engine/Plugins/FX/Niagara/Source/Niagara/Classes/NiagaraEmitterInstanceBatcher.h
@@ -145,18 +145,16 @@ public:
 	/** Loop over all the data interfaces and call the poststage methods */
 	void PostStageInterface(const FNiagaraGPUSystemTick& Tick, FNiagaraComputeInstanceData *Instance, FRHICommandList &RHICmdList, const FNiagaraShaderRef& ComputeShader, const uint32 ShaderStageIndex) const;
 
-	/** Run the dispatch over multiple stages */
-	void DispatchMultipleStages(const FNiagaraGPUSystemTick& Tick, FNiagaraComputeInstanceData *Instance, FRHICommandList &RHICmdList, FRHIUniformBuffer* ViewUniformBuffer, const FNiagaraShaderRef& ComputeShader);
-
 	/** Get the shared SortManager, used in the rendering loop to call FGPUSortManager::OnPreRender() and FGPUSortManager::OnPostRenderOpaque() */
 	virtual FGPUSortManager* GetGPUSortManager() const override;
 
 private:
 	using FEmitterInstanceList = TArray<FNiagaraComputeInstanceData*>;
 
-	void ExecuteAll(FRHICommandList& RHICmdList, FRHIUniformBuffer* ViewUniformBuffer, bool bSetReadback, ETickStage TickStage);
+	void ExecuteAll(FRHICommandList& RHICmdList, FRHIUniformBuffer* ViewUniformBuffer, ETickStage TickStage);
 	void ResizeBuffersAndGatherResources(FOverlappableTicks& OverlappableTick, FRHICommandList& RHICmdList, FNiagaraBufferArray& ReadBuffers, FNiagaraBufferArray& WriteBuffers, FNiagaraBufferArray& OutputGraphicsBuffers, FEmitterInstanceList& InstancesWithPersistentIDs);
-	void DispatchAllOnCompute(FOverlappableTicks& OverlappableTick, FRHICommandList& RHICmdList, FRHIUniformBuffer* ViewUniformBuffer, FNiagaraBufferArray& ReadBuffers, FNiagaraBufferArray& WriteBuffers, bool bSetReadback);
+	void DispatchAllOnCompute(FOverlappableTicks& OverlappableTick, FRHICommandList& RHICmdList, FRHIUniformBuffer* ViewUniformBuffer, FNiagaraBufferArray& ReadBuffers, FNiagaraBufferArray& WriteBuffers);
+	void DispatchMultipleStages(const FNiagaraGPUSystemTick& Tick, FNiagaraComputeInstanceData *Instance, FRHICommandList &RHICmdList, FRHIUniformBuffer* ViewUniformBuffer, const FNiagaraShaderRef& ComputeShader);
 
 	bool ShouldTickForStage(const FNiagaraGPUSystemTick& Tick, ETickStage TickStage) const;
 
diff --git a/Engine/Plugins/FX/Niagara/Source/Niagara/Private/NiagaraEmitterInstanceBatcher.cpp b/Engine/Plugins/FX/Niagara/Source/Niagara/Private/NiagaraEmitterInstanceBatcher.cpp
--- a/Engine/Plugins/FX/Niagara/Source/Niagara/Private/NiagaraEmitterInstanceBatcher.cpp
+++ b/Engine/Plugins/FX/Niagara/Source/Niagara/Private/NiagaraEmitterInstanceBatcher.cpp
@@ -489,7 +489,7 @@ void NiagaraEmitterInstanceBatcher::ResizeBuffersAndGatherResources(FOverlappabl
 	}
 }
 
-void NiagaraEmitterInstanceBatcher::DispatchAllOnCompute(FOverlappableTicks& OverlappableTick, FRHICommandList& RHICmdList, FRHIUniformBuffer* ViewUniformBuffer, FNiagaraBufferArray& ReadBuffers, FNiagaraBufferArray& WriteBuffers, bool bSetReadback)
+void NiagaraEmitterInstanceBatcher::DispatchAllOnCompute(FOverlappableTicks& OverlappableTick, FRHICommandList& RHICmdList, FRHIUniformBuffer* ViewUniformBuffer, FNiagaraBufferArray& ReadBuffers, FNiagaraBufferArray& WriteBuffers)
 {
 	FRHICommandListImmediate& RHICmdListImmediate = FRHICommandListExecutor::GetImmediateCommandList();
 
@@ -537,12 +537,12 @@ void NiagaraEmitterInstanceBatcher::DispatchAllOnCompute(FOverlappableTicks& Ove
 				DispatchMultipleStages(*Tick, &Instance, RHICmdList, ViewUniformBuffer, Context->GPUScript_RT->GetShader());
 
 				FNiagaraDataBuffer* CurrentData = Instance.CurrentData;
-				if (bSetReadback && Tick->bIsFinalTick)
+				if (!GPUInstanceCounterManager.HasPendingGPUReadback() && Tick->bIsFinalTick)
 				{
 					// Now that the current data is not required anymore, stage it for readback.
 					if (CurrentData->GetNumInstances() && Context->EmitterInstanceReadback.GPUCountOffset == INDEX_NONE && CurrentData->GetGPUInstanceCountBufferOffset() != INDEX_NONE)
 					{
-						// Transfer the GPU instance counter ownership to the context. Note that when bSetReadback is true, a readback request will be performed later in the tick update.
+						// Transfer the GPU instance counter ownership to the context. Note that a readback request will be performed later in the tick update, unless there's already a pending readback.
 						Context->EmitterInstanceReadback.GPUCountOffset = CurrentData->GetGPUInstanceCountBufferOffset();
 						Context->EmitterInstanceReadback.CPUCount = CurrentData->GetNumInstances();
 						CurrentData->ClearGPUInstanceCountBufferOffset();
@@ -570,7 +570,7 @@ void NiagaraEmitterInstanceBatcher::PostRenderOpaque(FRHICommandListImmediate& R
 	if (bAllowGPUParticleUpdate)
 	{
 		// Setup new readback since if there is no pending request, there is no risk of having invalid data read (offset being allocated after the readback was sent).
-		ExecuteAll(RHICmdList, ViewUniformBuffer, !GPUInstanceCounterManager.HasPendingGPUReadback(), ETickStage::PostOpaqueRender);
+		ExecuteAll(RHICmdList, ViewUniformBuffer, ETickStage::PostOpaqueRender);
 
 		RHICmdList.BeginUAVOverlap();
 		UpdateFreeIDBuffers(RHICmdList, DeferredIDBufferUpdates);
@@ -673,7 +673,7 @@ void NiagaraEmitterInstanceBatcher::UpdateFreeIDBuffers(FRHICommandList& RHICmdL
 	bFreeIDListSizesBufferCleared = false;
 }
 
-void NiagaraEmitterInstanceBatcher::ExecuteAll(FRHICommandList& RHICmdList, FRHIUniformBuffer* ViewUniformBuffer, bool bSetReadback, ETickStage TickStage)
+void NiagaraEmitterInstanceBatcher::ExecuteAll(FRHICommandList& RHICmdList, FRHIUniformBuffer* ViewUniformBuffer, ETickStage TickStage)
 {
 	SCOPE_CYCLE_COUNTER(STAT_NiagaraGPUSimTick_RT);
 
@@ -780,7 +780,7 @@ void NiagaraEmitterInstanceBatcher::ExecuteAll(FRHICommandList& RHICmdList, FRHI
 		{
 			SCOPED_DRAW_EVENT(RHICmdList, NiagaraGPUSimulation);
 			SCOPED_GPU_STAT(RHICmdList, NiagaraGPUSimulation);
-			DispatchAllOnCompute(SimPass, RHICmdList, ViewUniformBuffer, ReadBuffers, WriteBuffers, bSetReadback);
+			DispatchAllOnCompute(SimPass, RHICmdList, ViewUniformBuffer, ReadBuffers, WriteBuffers);
 		}
 
 		if (InstancesWithPersistentIDs.Num() == 0)
@@ -899,7 +899,7 @@ void NiagaraEmitterInstanceBatcher::PreInitViews(FRHICommandListImmediate& RHICm
 
 		if (GNiagaraAllowTickBeforeRender)
 		{
-			ExecuteAll(RHICmdList, nullptr, !GPUInstanceCounterManager.HasPendingGPUReadback(), ETickStage::PreInitViews);
+			ExecuteAll(RHICmdList, nullptr, ETickStage::PreInitViews);
 		}
 	}
 	else
@@ -919,7 +919,7 @@ void NiagaraEmitterInstanceBatcher::PostInitViews(FRHICommandListImmediate& RHIC
 
 	if (bAllowGPUParticleUpdate)
 	{
-		ExecuteAll(RHICmdList, ViewUniformBuffer, false, ETickStage::PostInitViews);
+		ExecuteAll(RHICmdList, ViewUniformBuffer, ETickStage::PostInitViews);
 	}
 }
 
