Fix for possible crash in scalability manager.

diff --git a/Engine/Plugins/FX/Niagara/Source/Niagara/Private/NiagaraComponent.cpp b/Engine/Plugins/FX/Niagara/Source/Niagara/Private/NiagaraComponent.cpp
--- a/Engine/Plugins/FX/Niagara/Source/Niagara/Private/NiagaraComponent.cpp
+++ b/Engine/Plugins/FX/Niagara/Source/Niagara/Private/NiagaraComponent.cpp
@@ -1129,6 +1129,8 @@ void UNiagaraComponent::OnUnregister()
 
 	SetActiveFlag(false);
 
+	UnregisterWithScalabilityManager();
+
 	if (SystemInstance)
 	{
 		//UE_LOG(LogNiagara, Log, TEXT("UNiagaraComponent::OnUnregister: %p  %s\n"), SystemInstance.Get(), *GetAsset()->GetFullName());
