Wait for Niagara Async work on PreGC.
Possible fix for several rare crashes that could be caused by GC nulling objects while we have async work in flight referencing them.

diff --git a/Engine/Plugins/FX/Niagara/Source/Niagara/Private/NiagaraWorldManager.cpp b/Engine/Plugins/FX/Niagara/Source/Niagara/Private/NiagaraWorldManager.cpp
--- a/Engine/Plugins/FX/Niagara/Source/Niagara/Private/NiagaraWorldManager.cpp
+++ b/Engine/Plugins/FX/Niagara/Source/Niagara/Private/NiagaraWorldManager.cpp
@@ -23,6 +23,7 @@
 DECLARE_CYCLE_STAT(TEXT("Niagara Manager Update Scalability Managers [GT]"), STAT_UpdateScalabilityManagers, STATGROUP_Niagara);
 DECLARE_CYCLE_STAT(TEXT("Niagara Manager Tick [GT]"), STAT_NiagaraWorldManTick, STATGROUP_Niagara);
 DECLARE_CYCLE_STAT(TEXT("Niagara Manager Wait On Render [GT]"), STAT_NiagaraWorldManWaitOnRender, STATGROUP_Niagara);
+DECLARE_CYCLE_STAT(TEXT("Niagara Manager Wait Pre Garbage Collect [GT]"), STAT_NiagaraWorldManWaitPreGC, STATGROUP_Niagara);
 
 static int GNiagaraAllowAsyncWorkToEndOfFrame = 1;
 static FAutoConsoleVariableRef CVarNiagaraAllowAsyncWorkToEndOfFrame(
@@ -371,6 +372,7 @@ void FNiagaraWorldManager::PreGarbageCollect()
 {
 	if (GNiagaraWaitOnPreGC)
 	{
+		SCOPE_CYCLE_COUNTER(STAT_NiagaraWorldManWaitPreGC);
 		// We must wait for system simulation & instance async ticks to complete before garbage collection can start
 		// The reason for this is that our async ticks could be referencing GC objects, i.e. skel meshes, etc, and we don't want them to become unreachable while we are potentially using them
 		for (int TG = 0; TG < NiagaraNumTickGroups; ++TG)
