Clean up UDeviceProfileManager handle before Module starts exiting as that single will have its delegate deleted from under it

diff --git a/Engine/Plugins/FX/Niagara/Source/NiagaraEditor/Private/NiagaraEditorModule.cpp b/Engine/Plugins/FX/Niagara/Source/NiagaraEditor/Private/NiagaraEditorModule.cpp
--- a/Engine/Plugins/FX/Niagara/Source/NiagaraEditor/Private/NiagaraEditorModule.cpp
+++ b/Engine/Plugins/FX/Niagara/Source/NiagaraEditor/Private/NiagaraEditorModule.cpp
@@ -715,6 +715,12 @@ class FNiagaraSystemColorParameterTrackEditor : public FNiagaraSystemParameterTr
 	}
 };
 
+// This will be called before UObjects are destroyed, so clean up anything we need to related to UObjects here
+void FNiagaraEditorModule::OnPreExit()
+{
+	UDeviceProfileManager::Get().OnManagerUpdated().Remove(DeviceProfileManagerUpdatedHandle);
+}
+
 void FNiagaraEditorModule::StartupModule()
 {
 	bThumbnailRenderersRegistered = false;
@@ -741,6 +747,7 @@ void FNiagaraEditorModule::StartupModule()
 	FCoreDelegates::OnPostEngineInit.AddRaw(this, &FNiagaraEditorModule::OnPostEngineInit);
 
 	DeviceProfileManagerUpdatedHandle = UDeviceProfileManager::Get().OnManagerUpdated().AddRaw(this, &FNiagaraEditorModule::OnDeviceProfileManagerUpdated);
+	FCoreDelegates::OnPreExit.AddRaw(this, &FNiagaraEditorModule::OnPreExit);
 	
 	// register details customization
 	FPropertyEditorModule& PropertyModule = FModuleManager::LoadModuleChecked<FPropertyEditorModule>("PropertyEditor");
@@ -1017,7 +1024,7 @@ void FNiagaraEditorModule::ShutdownModule()
 
 	FCoreUObjectDelegates::GetPreGarbageCollectDelegate().RemoveAll(this);
 	FCoreDelegates::OnPostEngineInit.RemoveAll(this);
-	UDeviceProfileManager::Get().OnManagerUpdated().Remove(DeviceProfileManagerUpdatedHandle);
+	FCoreDelegates::OnPreExit.RemoveAll(this);
 	
 	if (GEditor)
 	{
diff --git a/Engine/Plugins/FX/Niagara/Source/NiagaraEditor/Public/NiagaraEditorModule.h b/Engine/Plugins/FX/Niagara/Source/NiagaraEditor/Public/NiagaraEditorModule.h
--- a/Engine/Plugins/FX/Niagara/Source/NiagaraEditor/Public/NiagaraEditorModule.h
+++ b/Engine/Plugins/FX/Niagara/Source/NiagaraEditor/Public/NiagaraEditorModule.h
@@ -127,6 +127,7 @@ private:
 	void OnExecParticleInvoked(const TCHAR* InStr);
 	void OnPostEngineInit();
 	void OnDeviceProfileManagerUpdated();
+	void OnPreExit();
 
 	/** FGCObject interface */
 	virtual void AddReferencedObjects( FReferenceCollector& Collector ) override;
