Niagara parameter pin tooltips reformatted to put the description first.

#rnx

diff --git a/Engine/Plugins/FX/Niagara/Source/NiagaraEditor/Private/NiagaraEditorUtilities.cpp b/Engine/Plugins/FX/Niagara/Source/NiagaraEditor/Private/NiagaraEditorUtilities.cpp
--- a/Engine/Plugins/FX/Niagara/Source/NiagaraEditor/Private/NiagaraEditorUtilities.cpp
+++ b/Engine/Plugins/FX/Niagara/Source/NiagaraEditor/Private/NiagaraEditorUtilities.cpp
@@ -1438,7 +1438,17 @@ FText FNiagaraEditorUtilities::FormatScriptDescription(FText Description, FName
 
 	return Description.IsEmptyOrWhitespace()
 		? FText::Format(LOCTEXT("ScriptAssetDescriptionFormatPathOnly", "Path: {0}{1}"), FText::FromName(Path), LibrarySuffix)
-		: FText::Format(LOCTEXT("ScriptAssetDescriptionFormat", "Description: {1}\nPath: {0}{2}"), FText::FromName(Path), Description, LibrarySuffix);
+		: FText::Format(LOCTEXT("ScriptAssetDescriptionFormat", "{1}\nPath: {0}{2}"), FText::FromName(Path), Description, LibrarySuffix);
+}
+
+FText FNiagaraEditorUtilities::FormatVariableDescription(FText Description, FText Name, FText Type)
+{
+	if (Description.IsEmptyOrWhitespace() == false)
+	{
+		return FText::Format(LOCTEXT("VariableDescriptionFormat", "{0}\nName: \"{1}\"\nType: {2}"), Description, Name, Type);
+	}
+
+	return FText::Format(LOCTEXT("VariableDescriptionFormat_NoDesc", "Name: \"{0}\"\nType: {1}"), Name, Type);
 }
 
 void FNiagaraEditorUtilities::ResetSystemsThatReferenceSystemViewModel(const FNiagaraSystemViewModel& ReferencedSystemViewModel)
diff --git a/Engine/Plugins/FX/Niagara/Source/NiagaraEditor/Private/NiagaraNodeParameterMapBase.cpp b/Engine/Plugins/FX/Niagara/Source/NiagaraEditor/Private/NiagaraNodeParameterMapBase.cpp
--- a/Engine/Plugins/FX/Niagara/Source/NiagaraEditor/Private/NiagaraNodeParameterMapBase.cpp
+++ b/Engine/Plugins/FX/Niagara/Source/NiagaraEditor/Private/NiagaraNodeParameterMapBase.cpp
@@ -174,18 +174,9 @@ void UNiagaraNodeParameterMapBase::GetPinHoverText(const UEdGraphPin& Pin, FStri
 
 			FNiagaraVariable Var = FNiagaraVariable(TypeDef, Pin.PinName);
 			TOptional<FNiagaraVariableMetaData> Metadata = NiagaraGraph->GetMetaData(Var);
-			if (Metadata.IsSet())
-			{
-				FText Desc = FText::Format(LOCTEXT("GetVarTooltip", "Name: \"{0}\"\nType: {1}\nDesc: {2}"), FText::FromName(Pin.PinName),
-					TypeDef.GetNameText(), Metadata->Description);
-				HoverTextOut = Desc.ToString();
-			}
-			else
-			{
-				FText Desc = FText::Format(LOCTEXT("GetVarTooltip_NoDesc", "Name: \"{0}\"\nType: {1}\nDesc: None"), FText::FromName(Pin.PinName),
-					TypeDef.GetNameText());
-				HoverTextOut = Desc.ToString();
-			}
+			FText Description = Metadata.IsSet() ? Metadata->Description : FText::GetEmpty();
+			FText ToolTipText = FNiagaraEditorUtilities::FormatVariableDescription(Description, FText::FromName(Pin.PinName), TypeDef.GetNameText());
+			HoverTextOut = ToolTipText.ToString();
 		}
 	}
 }
diff --git a/Engine/Plugins/FX/Niagara/Source/NiagaraEditor/Private/NiagaraNodeParameterMapGet.cpp b/Engine/Plugins/FX/Niagara/Source/NiagaraEditor/Private/NiagaraNodeParameterMapGet.cpp
--- a/Engine/Plugins/FX/Niagara/Source/NiagaraEditor/Private/NiagaraNodeParameterMapGet.cpp
+++ b/Engine/Plugins/FX/Niagara/Source/NiagaraEditor/Private/NiagaraNodeParameterMapGet.cpp
@@ -470,24 +470,15 @@ void UNiagaraNodeParameterMapGet::GetPinHoverText(const UEdGraphPin& Pin, FStrin
 					HoverTextOut = Desc.ToString();
 					return;
 				}
-				
 			}
 			else
 			{
 				FNiagaraVariable Var = FNiagaraVariable(TypeDef, Pin.PinName);
 				TOptional<FNiagaraVariableMetaData> Metadata = NiagaraGraph->GetMetaData(Var);
-				if (Metadata.IsSet())
-				{
-					FText Desc = FText::Format(LOCTEXT("GetVarTooltip", "Name: \"{0}\"\nType: {1}\nDesc: {2}"), FText::FromName(Pin.PinName),
-						TypeDef.GetNameText(), Metadata->Description);
-					HoverTextOut = Desc.ToString();
-				}
-				else
-				{
-					FText Desc = FText::Format(LOCTEXT("GetVarTooltip_NoDesc", "Name: \"{0}\"\nType: {1}\nDesc: None"), FText::FromName(Pin.PinName),
-						TypeDef.GetNameText());
-					HoverTextOut = Desc.ToString();
-				}
+
+				FText Description = Metadata.IsSet() ? Metadata->Description : FText::GetEmpty();
+				FText ToolTipText = FNiagaraEditorUtilities::FormatVariableDescription(Description, FText::FromName(Pin.PinName), TypeDef.GetNameText());
+				HoverTextOut = ToolTipText.ToString();
 			}
 		}
 	}
diff --git a/Engine/Plugins/FX/Niagara/Source/NiagaraEditor/Private/ViewModels/Stack/NiagaraStackFunctionInput.cpp b/Engine/Plugins/FX/Niagara/Source/NiagaraEditor/Private/ViewModels/Stack/NiagaraStackFunctionInput.cpp
--- a/Engine/Plugins/FX/Niagara/Source/NiagaraEditor/Private/ViewModels/Stack/NiagaraStackFunctionInput.cpp
+++ b/Engine/Plugins/FX/Niagara/Source/NiagaraEditor/Private/ViewModels/Stack/NiagaraStackFunctionInput.cpp
@@ -205,19 +205,8 @@ const FNiagaraTypeDefinition& UNiagaraStackFunctionInput::GetInputType() const
 
 FText UNiagaraStackFunctionInput::GetTooltipText() const
 {
-	if (InputMetaData.IsSet() && InputMetaData->Description.IsEmptyOrWhitespace() == false)
-	{
-		return FText::Format(LOCTEXT("FunctionInputTooltipWithDescription", "Name: {0} \nType: {1} \nDesc: {2}"),
-			GetDisplayName(),
-			InputType.GetNameText(),
-			InputMetaData->Description);
-	}
-	else
-	{
-		return FText::Format(LOCTEXT("FunctionInputTooltip", "Name: {0} \nType: {1}"),
-			GetDisplayName(),
-			InputType.GetNameText());
-	}
+	FText Description = InputMetaData.IsSet() ? InputMetaData->Description : FText::GetEmpty();
+	return FNiagaraEditorUtilities::FormatVariableDescription(Description, GetDisplayName(), InputType.GetNameText());
 }
 
 bool UNiagaraStackFunctionInput::GetIsEnabled() const
diff --git a/Engine/Plugins/FX/Niagara/Source/NiagaraEditor/Public/NiagaraEditorUtilities.h b/Engine/Plugins/FX/Niagara/Source/NiagaraEditor/Public/NiagaraEditorUtilities.h
--- a/Engine/Plugins/FX/Niagara/Source/NiagaraEditor/Public/NiagaraEditorUtilities.h
+++ b/Engine/Plugins/FX/Niagara/Source/NiagaraEditor/Public/NiagaraEditorUtilities.h
@@ -200,6 +200,8 @@ namespace FNiagaraEditorUtilities
 
 	NIAGARAEDITOR_API FText FormatScriptDescription(FText Description, FName Path, bool bIsInLibrary);
 
+	NIAGARAEDITOR_API FText FormatVariableDescription(FText Description, FText Name, FText Type);
+
 	void ResetSystemsThatReferenceSystemViewModel(const FNiagaraSystemViewModel& ReferencedSystemViewModel);
 
 	TArray<UNiagaraComponent*> GetComponentsThatReferenceSystem(const UNiagaraSystem& ReferencedSystem);
