The dropdown arrow in the header of a Set Variables modules is now a + to more clearly reflect the functionality available in the dropdown.

#rnx

diff --git a/Engine/Plugins/FX/Niagara/Source/NiagaraEditorWidgets/Private/Stack/SNiagaraStackModuleItem.cpp b/Engine/Plugins/FX/Niagara/Source/NiagaraEditorWidgets/Private/Stack/SNiagaraStackModuleItem.cpp
--- a/Engine/Plugins/FX/Niagara/Source/NiagaraEditorWidgets/Private/Stack/SNiagaraStackModuleItem.cpp
+++ b/Engine/Plugins/FX/Niagara/Source/NiagaraEditorWidgets/Private/Stack/SNiagaraStackModuleItem.cpp
@@ -63,7 +63,7 @@ void SNiagaraStackModuleItem::Construct(const FArguments& InArgs, UNiagaraStackM
 			.AutoWidth()
 			[
 				SNew(SComboButton)
-				.HasDownArrow(true)
+				.HasDownArrow(false)
 				.ButtonStyle(FEditorStyle::Get(), "HoverHintOnly")
 				.ForegroundColor(FSlateColor::UseForeground())
 				.OnGetMenuContent(this, &SNiagaraStackModuleItem::RaiseActionMenuClicked)
@@ -72,6 +72,11 @@ void SNiagaraStackModuleItem::Construct(const FArguments& InArgs, UNiagaraStackM
 				.VAlign(VAlign_Center)
 				.Visibility(this, &SNiagaraStackModuleItem::GetRaiseActionMenuVisibility)
 				.IsEnabled(this, &SNiagaraStackModuleItem::GetButtonsEnabled)
+				.ButtonContent()
+				[
+					SNew(SImage)
+					.Image(FEditorStyle::Get().GetBrush("PropertyWindow.Button_AddToArray"))
+				]
 			]
 			// Refresh button
 			+ SHorizontalBox::Slot()
