Workaround for an issue around sharedgroup in NiagaraGPU simulation.

diff --git a/Engine/Plugins/FX/Niagara/Shaders/Private/NiagaraEmitterInstanceShader.usf b/Engine/Plugins/FX/Niagara/Shaders/Private/NiagaraEmitterInstanceShader.usf
--- a/Engine/Plugins/FX/Niagara/Shaders/Private/NiagaraEmitterInstanceShader.usf
+++ b/Engine/Plugins/FX/Niagara/Shaders/Private/NiagaraEmitterInstanceShader.usf
@@ -484,7 +484,7 @@ int rand_int(int x, int Seed1, int Seed2, int Seed3)
 	}
 
 #define USE_WAVE_INTRISICS (XBOXONE_PROFILE)
-#define USE_GROUP_SHARED ((THREADGROUP_SIZE == 64 || THREADGROUP_SIZE == 32)  && !USE_WAVE_INTRISICS && !VULKAN_PROFILE_SM5 && !VULKAN_PROFILE)
+#define USE_GROUP_SHARED ((THREADGROUP_SIZE == 64 || THREADGROUP_SIZE == 32)  && !USE_WAVE_INTRISICS)
 
 #if USE_GROUP_SHARED
 	#if THREADGROUP_SIZE == 64
@@ -581,22 +581,24 @@ int rand_int(int x, int Seed1, int Seed2, int Seed3)
 			GroupMemoryBarrierWithGroupSync(); 
 
 			// Group by 64
-			if (GGroupThreadId.x == 0)
+			if ((GGroupThreadId.x & 0x3F) == 0)
 			{
-				const uint ActiveCount1 = GroupSharedIndex16[0];
-				const uint ActiveCount2 = ActiveCount1 + GroupSharedIndex16[1];
-		#if THREADGROUP_SIZE == 64
-				const uint ActiveCount3 = ActiveCount2 + GroupSharedIndex16[2];
-				const uint ActiveCount4 = ActiveCount3 + GroupSharedIndex16[3];
-		#endif
-				GroupSharedIndex16[0] = 0;
-				GroupSharedIndex16[1] = ActiveCount1;
+				const uint Index = GGroupThreadId.x / 16;
+				uint RetPrevIdx = 0;
+		
+				const uint ActiveCount1 = GroupSharedIndex16[Index];
+				const uint ActiveCount2 = ActiveCount1 + GroupSharedIndex16[Index + 1];
 		#if THREADGROUP_SIZE == 64
-				GroupSharedIndex16[2] = ActiveCount2;
-				GroupSharedIndex16[3] = ActiveCount3;
+				const uint ActiveCount3 = ActiveCount2 + GroupSharedIndex16[Index + 2];
+				const uint ActiveCount4 = ActiveCount3 + GroupSharedIndex16[Index + 3];
 		#endif
-				uint RetPrevIdx = 0;
+
+				GroupSharedIndex16[Index] = 0;
+				GroupSharedIndex16[Index + 1] = ActiveCount1;
 		#if THREADGROUP_SIZE == 64
+				GroupSharedIndex16[Index + 2] = ActiveCount2;
+				GroupSharedIndex16[Index + 3] = ActiveCount3;
+
 				InterlockedAdd(RWInstanceCounts[WriteInstanceCountOffset], ActiveCount4, RetPrevIdx);
 		#elif THREADGROUP_SIZE == 32
 				InterlockedAdd(RWInstanceCounts[WriteInstanceCountOffset], ActiveCount2, RetPrevIdx);
