Niagara platform set customization is now top justified when there are platform overrides.

#rnx

diff --git a/Engine/Plugins/FX/Niagara/Source/NiagaraEditor/Private/Customizations/NiagaraTypeCustomizations.cpp b/Engine/Plugins/FX/Niagara/Source/NiagaraEditor/Private/Customizations/NiagaraTypeCustomizations.cpp
--- a/Engine/Plugins/FX/Niagara/Source/NiagaraEditor/Private/Customizations/NiagaraTypeCustomizations.cpp
+++ b/Engine/Plugins/FX/Niagara/Source/NiagaraEditor/Private/Customizations/NiagaraTypeCustomizations.cpp
@@ -820,7 +820,7 @@ void FNiagaraPlatformSetTypeCustomization::GenerateEffectsQualitySelectionWidget
 			[
 				SNew(SBox)
 				.WidthOverride(100)
-				.VAlign(VAlign_Center)
+				.VAlign(VAlign_Top)
 				[
 					SNew(SVerticalBox)
 					+ SVerticalBox::Slot()
