Fix for FX not reactivating when changing FX Quality.

diff --git a/Engine/Plugins/FX/Niagara/Source/Niagara/Public/NiagaraComponent.h b/Engine/Plugins/FX/Niagara/Source/Niagara/Public/NiagaraComponent.h
--- a/Engine/Plugins/FX/Niagara/Source/Niagara/Public/NiagaraComponent.h
+++ b/Engine/Plugins/FX/Niagara/Source/Niagara/Public/NiagaraComponent.h
@@ -517,6 +517,7 @@ public:
 	UFUNCTION(BlueprintCallable, Category = Scalability, meta = (Keywords = "LOD scalability"))
 	void SetAllowScalability(bool bAllow);
 
+	FORCEINLINE int32 GetScalabilityManagerHandle()const { return ScalabilityManagerHandle; }
 private:
 	/** Did we try and activate but fail due to the asset being not yet ready. Keep looping.*/
 	uint32 bAwaitingActivationDueToNotReady : 1;
