Niagara - Sequencer fixes.
+ Merging 11427994 - Add aditional options to the niagara system life cycle track to better control system lifetime and age updating.
+ Merging 11447307 - Change the sequencer integration to fully deactivate the component instead of just setting the requested execution state.
+ Fix bool user parameter animation.
+ Auto-test updates.

diff --git a/Engine/Plugins/FX/Niagara/Source/Niagara/Private/MovieScene/MovieSceneNiagaraSystemSpawnSection.cpp b/Engine/Plugins/FX/Niagara/Source/Niagara/Private/MovieScene/MovieSceneNiagaraSystemSpawnSection.cpp
--- a/Engine/Plugins/FX/Niagara/Source/Niagara/Private/MovieScene/MovieSceneNiagaraSystemSpawnSection.cpp
+++ b/Engine/Plugins/FX/Niagara/Source/Niagara/Private/MovieScene/MovieSceneNiagaraSystemSpawnSection.cpp
@@ -5,4 +5,30 @@
 UMovieSceneNiagaraSystemSpawnSection::UMovieSceneNiagaraSystemSpawnSection()
 {
 	EvalOptions.CompletionMode = EMovieSceneCompletionMode::RestoreState;
+	SectionStartBehavior = ENiagaraSystemSpawnSectionStartBehavior::Activate;
+	SectionEvaluateBehavior = ENiagaraSystemSpawnSectionEvaluateBehavior::None;
+	SectionEndBehavior = ENiagaraSystemSpawnSectionEndBehavior::SetSystemInactive;
+	AgeUpdateMode = ENiagaraAgeUpdateMode::TickDeltaTime;
 }
+
+ENiagaraSystemSpawnSectionStartBehavior UMovieSceneNiagaraSystemSpawnSection::GetSectionStartBehavior() const
+{
+	return SectionStartBehavior;
+}
+
+ENiagaraSystemSpawnSectionEvaluateBehavior UMovieSceneNiagaraSystemSpawnSection::GetSectionEvaluateBehavior() const
+{
+	return SectionEvaluateBehavior;
+}
+
+ENiagaraSystemSpawnSectionEndBehavior UMovieSceneNiagaraSystemSpawnSection::GetSectionEndBehavior() const
+{
+	return SectionEndBehavior;
+}
+
+ENiagaraAgeUpdateMode UMovieSceneNiagaraSystemSpawnSection::GetAgeUpdateMode() const
+{
+	return AgeUpdateMode;
+}
+
+
diff --git a/Engine/Plugins/FX/Niagara/Source/Niagara/Private/MovieScene/MovieSceneNiagaraSystemTrack.cpp b/Engine/Plugins/FX/Niagara/Source/Niagara/Private/MovieScene/MovieSceneNiagaraSystemTrack.cpp
--- a/Engine/Plugins/FX/Niagara/Source/Niagara/Private/MovieScene/MovieSceneNiagaraSystemTrack.cpp
+++ b/Engine/Plugins/FX/Niagara/Source/Niagara/Private/MovieScene/MovieSceneNiagaraSystemTrack.cpp
@@ -34,7 +34,10 @@ void UMovieSceneNiagaraSystemTrack::PostCompile(FMovieSceneEvaluationTrack& OutT
 	{
 		UMovieSceneNiagaraSystemSpawnSection* SpawnSection = CastChecked<UMovieSceneNiagaraSystemSpawnSection>(*SpawnSectionPtr);
 		UMovieScene* ParentMovieScene = GetTypedOuter<UMovieScene>();
-		OutTrack.SetTrackImplementation(FMovieSceneNiagaraSystemTrackImplementation(SpawnSection->GetInclusiveStartFrame(), SpawnSection->GetExclusiveEndFrame()));
+		OutTrack.SetTrackImplementation(FMovieSceneNiagaraSystemTrackImplementation(
+			SpawnSection->GetInclusiveStartFrame(), SpawnSection->GetExclusiveEndFrame(),
+			SpawnSection->GetSectionStartBehavior(), SpawnSection->GetSectionEvaluateBehavior(),
+			SpawnSection->GetSectionEndBehavior(), SpawnSection->GetAgeUpdateMode()));
 	}
 }
 
diff --git a/Engine/Plugins/FX/Niagara/Source/Niagara/Private/MovieScene/MovieSceneNiagaraSystemTrackTemplate.cpp b/Engine/Plugins/FX/Niagara/Source/Niagara/Private/MovieScene/MovieSceneNiagaraSystemTrackTemplate.cpp
--- a/Engine/Plugins/FX/Niagara/Source/Niagara/Private/MovieScene/MovieSceneNiagaraSystemTrackTemplate.cpp
+++ b/Engine/Plugins/FX/Niagara/Source/Niagara/Private/MovieScene/MovieSceneNiagaraSystemTrackTemplate.cpp
@@ -77,10 +77,16 @@ struct FPreAnimatedNiagaraComponentTokenProducer : IMovieScenePreAnimatedTokenPr
 
 struct FNiagaraSystemUpdateDesiredAgeExecutionToken : IMovieSceneExecutionToken
 {
-	FNiagaraSystemUpdateDesiredAgeExecutionToken(FFrameNumber InSpawnSectionStartFrame, FFrameNumber InSpawnSectionEndFrame)
+	FNiagaraSystemUpdateDesiredAgeExecutionToken(
+		FFrameNumber InSpawnSectionStartFrame, FFrameNumber InSpawnSectionEndFrame,
+		ENiagaraSystemSpawnSectionStartBehavior InSpawnSectionStartBehavior, ENiagaraSystemSpawnSectionEvaluateBehavior InSpawnSectionEvaluateBehavior,
+		ENiagaraSystemSpawnSectionEndBehavior InSpawnSectionEndBehavior, ENiagaraAgeUpdateMode InAgeUpdateMode)
 		: SpawnSectionStartFrame(InSpawnSectionStartFrame)
 		, SpawnSectionEndFrame(InSpawnSectionEndFrame)
-		
+		, SpawnSectionStartBehavior(InSpawnSectionStartBehavior)
+		, SpawnSectionEvaluateBehavior(InSpawnSectionEvaluateBehavior)
+		, SpawnSectionEndBehavior(InSpawnSectionEndBehavior)
+		, AgeUpdateMode(InAgeUpdateMode)
 	{}
 
 	virtual void Execute(const FMovieSceneContext& Context, const FMovieSceneEvaluationOperand& Operand, FPersistentEvaluationData& PersistentData, IMovieScenePlayer& Player) override
@@ -94,7 +100,7 @@ struct FNiagaraSystemUpdateDesiredAgeExecutionToken : IMovieSceneExecutionToken
 			Player.SavePreAnimatedState(*NiagaraComponent, TypeID, FPreAnimatedNiagaraComponentTokenProducer(), PersistentData.GetTrackKey());
 
 			NiagaraComponent->SetForceSolo(true);
-			NiagaraComponent->SetAgeUpdateMode(ENiagaraAgeUpdateMode::DesiredAge);
+			NiagaraComponent->SetAgeUpdateMode(AgeUpdateMode);
 
 			UMovieSceneSequence* MovieSceneSequence = Player.GetEvaluationTemplate().GetSequence(Operand.SequenceID);
 			if (MovieSceneSequence != nullptr)
@@ -110,28 +116,54 @@ struct FNiagaraSystemUpdateDesiredAgeExecutionToken : IMovieSceneExecutionToken
 
 			if (Context.GetTime() < SpawnSectionStartFrame)
 			{
-				if (NiagaraComponent->IsActive() && SystemInstance != nullptr)
+				if (SpawnSectionStartBehavior == ENiagaraSystemSpawnSectionStartBehavior::Activate)
 				{
-					SystemInstance->SetRequestedExecutionState(ENiagaraExecutionState::Complete);
+					if (NiagaraComponent->IsActive())
+					{
+						NiagaraComponent->DeactivateImmediate();
+					}
 				}
 			}
-			else if (Context.GetTime() < SpawnSectionEndFrame)
+			else if (Context.GetRange().Overlaps(TRange<FFrameTime>(FFrameTime(SpawnSectionStartFrame))))
 			{
-				if (NiagaraComponent->IsActive() == false)
+				if (SpawnSectionStartBehavior == ENiagaraSystemSpawnSectionStartBehavior::Activate)
 				{
-					NiagaraComponent->Activate();
+					if (NiagaraComponent->IsActive() == false)
+					{
+						NiagaraComponent->Activate();
+					}
 				}
-				
-				if (SystemInstance != nullptr)
+			}
+			else if (Context.GetTime() < SpawnSectionEndFrame)
+			{
+				if (SpawnSectionEvaluateBehavior == ENiagaraSystemSpawnSectionEvaluateBehavior::ActivateIfInactive)
 				{
-					SystemInstance->SetRequestedExecutionState(ENiagaraExecutionState::Active);
+					if (NiagaraComponent->IsActive() == false)
+					{
+						NiagaraComponent->Activate();
+					}
+
+					if (SystemInstance != nullptr)
+					{
+						SystemInstance->SetRequestedExecutionState(ENiagaraExecutionState::Active);
+					}
 				}
 			}
 			else
 			{
-				if (SystemInstance != nullptr)
+				if (SpawnSectionEndBehavior == ENiagaraSystemSpawnSectionEndBehavior::SetSystemInactive)
 				{
-					SystemInstance->SetRequestedExecutionState(ENiagaraExecutionState::Inactive);
+					if (SystemInstance != nullptr)
+					{
+						SystemInstance->SetRequestedExecutionState(ENiagaraExecutionState::Inactive);
+					}
+				}
+				else if (SpawnSectionEndBehavior == ENiagaraSystemSpawnSectionEndBehavior::Deactivate)
+				{
+					if (NiagaraComponent->IsActive())
+					{
+						NiagaraComponent->DeactivateImmediate();
+					}
 				}
 			}
 
@@ -155,15 +187,31 @@ struct FNiagaraSystemUpdateDesiredAgeExecutionToken : IMovieSceneExecutionToken
 
 	FFrameNumber SpawnSectionStartFrame;
 	FFrameNumber SpawnSectionEndFrame;
+	ENiagaraSystemSpawnSectionStartBehavior SpawnSectionStartBehavior;
+	ENiagaraSystemSpawnSectionEvaluateBehavior SpawnSectionEvaluateBehavior;
+	ENiagaraSystemSpawnSectionEndBehavior SpawnSectionEndBehavior;
+	ENiagaraAgeUpdateMode AgeUpdateMode;
 };
 
-FMovieSceneNiagaraSystemTrackImplementation::FMovieSceneNiagaraSystemTrackImplementation(FFrameNumber InSpawnSectionStartFrame, FFrameNumber InSpawnSectionEndFrame)
+FMovieSceneNiagaraSystemTrackImplementation::FMovieSceneNiagaraSystemTrackImplementation(
+	FFrameNumber InSpawnSectionStartFrame, FFrameNumber InSpawnSectionEndFrame,
+	ENiagaraSystemSpawnSectionStartBehavior InSpawnSectionStartBehavior, ENiagaraSystemSpawnSectionEvaluateBehavior InSpawnSectionEvaluateBehavior,
+	ENiagaraSystemSpawnSectionEndBehavior InSpawnSectionEndBehavior, ENiagaraAgeUpdateMode InAgeUpdateMode)
 	: SpawnSectionStartFrame(InSpawnSectionStartFrame)
 	, SpawnSectionEndFrame(InSpawnSectionEndFrame)
-{}
+	, SpawnSectionStartBehavior(InSpawnSectionStartBehavior)
+	, SpawnSectionEvaluateBehavior(InSpawnSectionEvaluateBehavior)
+	, SpawnSectionEndBehavior(InSpawnSectionEndBehavior)
+	, AgeUpdateMode(InAgeUpdateMode)
+	
+{
+}
 
 void FMovieSceneNiagaraSystemTrackImplementation::Evaluate(const FMovieSceneEvaluationTrack& Track, FMovieSceneSegmentIdentifier SegmentID, const FMovieSceneEvaluationOperand& Operand, const FMovieSceneContext& Context, const FPersistentEvaluationData& PersistentData, FMovieSceneExecutionTokens& ExecutionTokens) const
 {
 	ExecutionTokens.SetContext(Context);
-	ExecutionTokens.Add(FNiagaraSystemUpdateDesiredAgeExecutionToken(SpawnSectionStartFrame, SpawnSectionEndFrame));
+	ExecutionTokens.Add(FNiagaraSystemUpdateDesiredAgeExecutionToken(
+		SpawnSectionStartFrame, SpawnSectionEndFrame,
+		SpawnSectionStartBehavior, SpawnSectionEvaluateBehavior,
+		SpawnSectionEndBehavior, AgeUpdateMode));
 }
\ No newline at end of file
diff --git a/Engine/Plugins/FX/Niagara/Source/Niagara/Private/MovieScene/MovieSceneNiagaraSystemTrackTemplate.h b/Engine/Plugins/FX/Niagara/Source/Niagara/Private/MovieScene/MovieSceneNiagaraSystemTrackTemplate.h
--- a/Engine/Plugins/FX/Niagara/Source/Niagara/Private/MovieScene/MovieSceneNiagaraSystemTrackTemplate.h
+++ b/Engine/Plugins/FX/Niagara/Source/Niagara/Private/MovieScene/MovieSceneNiagaraSystemTrackTemplate.h
@@ -2,6 +2,8 @@
 
 #pragma once
 
+#include "NiagaraCommon.h"
+#include "MovieScene/MovieSceneNiagaraSystemSpawnSection.h"
 #include "Evaluation/MovieSceneEvalTemplate.h"
 #include "Evaluation/MovieSceneTrackImplementation.h"
 #include "MovieSceneNiagaraSystemTrackTemplate.generated.h"
@@ -26,7 +28,11 @@ struct FMovieSceneNiagaraSystemTrackImplementation : public FMovieSceneTrackImpl
 public:
 	FMovieSceneNiagaraSystemTrackImplementation() { }
 	FMovieSceneNiagaraSystemTrackImplementation(FFrameNumber InSpawnSectionStartFrame, FFrameNumber InSpawnSectionEndFrame);
-	
+	FMovieSceneNiagaraSystemTrackImplementation(
+		FFrameNumber InSpawnSectionStartFrame, FFrameNumber InSpawnSectionEndFrame,
+		ENiagaraSystemSpawnSectionStartBehavior InSectionStartBehavior, ENiagaraSystemSpawnSectionEvaluateBehavior InSectionEvaluateBehavior,
+		ENiagaraSystemSpawnSectionEndBehavior InSectionEndBehavior, ENiagaraAgeUpdateMode InAgeUpdateMode);
+
 	virtual void SetupOverrides() override
 	{
 		EnableOverrides(CustomEvaluateFlag);
@@ -41,4 +47,8 @@ private:
 	FFrameNumber SpawnSectionStartFrame;
 	UPROPERTY()
 	FFrameNumber SpawnSectionEndFrame;
+	ENiagaraSystemSpawnSectionStartBehavior SpawnSectionStartBehavior;
+	ENiagaraSystemSpawnSectionEvaluateBehavior SpawnSectionEvaluateBehavior;
+	ENiagaraSystemSpawnSectionEndBehavior SpawnSectionEndBehavior;
+	ENiagaraAgeUpdateMode AgeUpdateMode;
 };
\ No newline at end of file
diff --git a/Engine/Plugins/FX/Niagara/Source/Niagara/Private/MovieScene/Parameters/MovieSceneNiagaraBoolParameterSectionTemplate.cpp b/Engine/Plugins/FX/Niagara/Source/Niagara/Private/MovieScene/Parameters/MovieSceneNiagaraBoolParameterSectionTemplate.cpp
--- a/Engine/Plugins/FX/Niagara/Source/Niagara/Private/MovieScene/Parameters/MovieSceneNiagaraBoolParameterSectionTemplate.cpp
+++ b/Engine/Plugins/FX/Niagara/Source/Niagara/Private/MovieScene/Parameters/MovieSceneNiagaraBoolParameterSectionTemplate.cpp
@@ -27,5 +27,5 @@ void FMovieSceneNiagaraBoolParameterSectionTemplate::GetParameterValue(FFrameTim
 	}
 	
 	OutAnimatedValueData.AddUninitialized(sizeof(FNiagaraBool));
-	FMemory::Memcpy(OutAnimatedValueData.GetData(), (uint8*)&AnimatedValue, sizeof(FNiagaraBool));
+	FMemory::Memcpy(OutAnimatedValueData.GetData(), (uint8*)&AnimatedNiagaraValue, sizeof(FNiagaraBool));
 }
\ No newline at end of file
diff --git a/Engine/Plugins/FX/Niagara/Source/Niagara/Private/NiagaraComponent.cpp b/Engine/Plugins/FX/Niagara/Source/Niagara/Private/NiagaraComponent.cpp
--- a/Engine/Plugins/FX/Niagara/Source/Niagara/Private/NiagaraComponent.cpp
+++ b/Engine/Plugins/FX/Niagara/Source/Niagara/Private/NiagaraComponent.cpp
@@ -433,6 +433,7 @@ UNiagaraComponent::UNiagaraComponent(const FObjectInitializer& ObjectInitializer
 	, bForceSolo(false)
 	, AgeUpdateMode(ENiagaraAgeUpdateMode::TickDeltaTime)
 	, DesiredAge(0.0f)
+	, LastHandledDesiredAge(0.0f)
 	, bCanRenderWhileSeeking(true)
 	, SeekDelta(1 / 30.0f)
 	, MaxSimTime(33.0f / 1000.0f)
@@ -587,7 +588,7 @@ void UNiagaraComponent::TickComponent(float DeltaSeconds, enum ELevelTick TickTy
 		{
 			SystemInstance->ComponentTick(DeltaSeconds, ThisTickFunction->GetCompletionHandle());
 		}
-		else
+		else if(AgeUpdateMode == ENiagaraAgeUpdateMode::DesiredAge)
 		{
 			float AgeDiff = FMath::Max(DesiredAge, 0.0f) - SystemInstance->GetAge();
 			int32 TicksToProcess = 0;
@@ -628,6 +629,33 @@ void UNiagaraComponent::TickComponent(float DeltaSeconds, enum ELevelTick TickTy
 				bIsSeeking = false;
 			}
 		}
+		else if (AgeUpdateMode == ENiagaraAgeUpdateMode::DesiredAgeNoSeek)
+		{
+			int32 MaxForwardFrames = 5; // HACK - for some reason sequencer jumps forwards by multiple frames on pause, so this is being added to allow for FX to stay alive when being controlled by sequencer in the editor.  This should be lowered once that issue is fixed.
+			float AgeDiff = DesiredAge - LastHandledDesiredAge;
+			if (AgeDiff < 0)
+			{
+				if (FMath::Abs(AgeDiff) >= SeekDelta)
+				{
+					// When going back in time for a frame or more, reset and simulate a single frame.  We ignore small negative changes to delta
+					// time which can happen when controlling time with the timeline and the time snaps to a previous time when paused.
+					SystemInstance->Reset(FNiagaraSystemInstance::EResetMode::ResetAll);
+					SystemInstance->ComponentTick(SeekDelta, nullptr);
+				}
+			}
+			else if (AgeDiff < MaxForwardFrames * SeekDelta)
+			{
+				// Allow ticks between 0 and MaxForwardFrames, but don't ever send more then 2 x the seek delta.
+				SystemInstance->ComponentTick(FMath::Min(AgeDiff, 2 * SeekDelta), nullptr);
+			}
+			else
+			{
+				// When going forward in time for more than MaxForwardFrames, reset and simulate a single frame.
+				SystemInstance->Reset(FNiagaraSystemInstance::EResetMode::ResetAll);
+				SystemInstance->ComponentTick(SeekDelta, nullptr);
+			}
+			LastHandledDesiredAge = DesiredAge;
+		}
 
 		if (SceneProxy != nullptr)
 		{
diff --git a/Engine/Plugins/FX/Niagara/Source/Niagara/Public/MovieScene/MovieSceneNiagaraSystemSpawnSection.h b/Engine/Plugins/FX/Niagara/Source/Niagara/Public/MovieScene/MovieSceneNiagaraSystemSpawnSection.h
--- a/Engine/Plugins/FX/Niagara/Source/Niagara/Public/MovieScene/MovieSceneNiagaraSystemSpawnSection.h
+++ b/Engine/Plugins/FX/Niagara/Source/Niagara/Public/MovieScene/MovieSceneNiagaraSystemSpawnSection.h
@@ -2,13 +2,71 @@
 
 #pragma once
 
+#include "NiagaraCommon.h"
 #include "MovieSceneSection.h"
 #include "MovieSceneNiagaraSystemSpawnSection.generated.h"
 
+/** Defines options for system life cycle for before the section is evaluating up to the first frame the section evaluates. */
+UENUM()
+enum class ENiagaraSystemSpawnSectionStartBehavior
+{
+	/** When the time before the section evaluates the particle system's component will be deactivated and on the first frame of the section the
+	 system's component will be activated. */
+	Activate
+};
+
+/** Defines options for system life cycle for when the section is evaluating from the 2nd frame until the last frame of the section. */
+UENUM()
+enum class ENiagaraSystemSpawnSectionEvaluateBehavior
+{
+	/** The system's component will be activated on any frame where it is inactive.  This is useful for continuous emitters, especially if the sequencer will start in the middle of the section. */
+	ActivateIfInactive,
+	/** There sill be no changes to the system life cycle while the section is evaluating. */
+	None
+};
+
+/** Defines options for system life cycle for the time after the section. */
+UENUM()
+enum class ENiagaraSystemSpawnSectionEndBehavior
+{
+	//** When the section ends the system is set to inactive which stops spawning but lets existing particles simulate until death.
+	SetSystemInactive,
+	//** When the section ends the system's component is deactivated which will kill all existing particles.
+	Deactivate,
+	//** Does nothing when the section ends and allows the system to continue to run as normal.
+	None
+};
+
 UCLASS(MinimalAPI)
 class UMovieSceneNiagaraSystemSpawnSection : public UMovieSceneSection
 {
 	GENERATED_BODY()
 
+public:
 	UMovieSceneNiagaraSystemSpawnSection();
+
+	ENiagaraSystemSpawnSectionStartBehavior GetSectionStartBehavior() const;
+
+	ENiagaraSystemSpawnSectionEvaluateBehavior GetSectionEvaluateBehavior() const;
+
+	ENiagaraSystemSpawnSectionEndBehavior GetSectionEndBehavior() const;
+
+	ENiagaraAgeUpdateMode GetAgeUpdateMode() const;
+
+private:
+	/** Specifies what should happen to the niagara system from before the section evaluates up until the first frame of the section. */
+	UPROPERTY(EditAnywhere, Category = "Life Cycle")
+	ENiagaraSystemSpawnSectionStartBehavior SectionStartBehavior;
+
+	/** Specifies what should happen to the niagara system when section is evaluating from the 2nd frame until the last frame. */
+	UPROPERTY(EditAnywhere, Category = "Life Cycle")
+	ENiagaraSystemSpawnSectionEvaluateBehavior SectionEvaluateBehavior;
+
+	/** Specifies what should happen to the niagara system when section evaluation finishes. */
+	UPROPERTY(EditAnywhere, Category = "Life Cycle")
+	ENiagaraSystemSpawnSectionEndBehavior SectionEndBehavior;
+
+	/** Specifies how sequencer should update the age of the controlled niagara system. */
+	UPROPERTY(EditAnywhere, Category = "Life Cycle")
+	ENiagaraAgeUpdateMode AgeUpdateMode;
 };
\ No newline at end of file
diff --git a/Engine/Plugins/FX/Niagara/Source/Niagara/Public/NiagaraCommon.h b/Engine/Plugins/FX/Niagara/Source/Niagara/Public/NiagaraCommon.h
--- a/Engine/Plugins/FX/Niagara/Source/Niagara/Public/NiagaraCommon.h
+++ b/Engine/Plugins/FX/Niagara/Source/Niagara/Public/NiagaraCommon.h
@@ -83,10 +83,14 @@ enum class ENiagaraSimTarget : uint8
 UENUM()
 enum class ENiagaraAgeUpdateMode : uint8
 {
-	/** Update the age using the delta time supplied to the tick function. */
+	/** Update the age using the delta time supplied to the component tick function. */
 	TickDeltaTime,
 	/** Update the age by seeking to the DesiredAge. To prevent major perf loss, we clamp to MaxClampTime*/
-	DesiredAge
+	DesiredAge,
+	/** Update the age by tracking changes to the desired age, but when the desired age goes backwards in time,
+	or jumps forwards in time by more than a few steps, the system is reset and simulated forward by a single step.
+	This mode is useful for continuous effects controlled by sequencer. */
+	DesiredAgeNoSeek
 };
 
 
diff --git a/Engine/Plugins/FX/Niagara/Source/Niagara/Public/NiagaraComponent.h b/Engine/Plugins/FX/Niagara/Source/Niagara/Public/NiagaraComponent.h
--- a/Engine/Plugins/FX/Niagara/Source/Niagara/Public/NiagaraComponent.h
+++ b/Engine/Plugins/FX/Niagara/Source/Niagara/Public/NiagaraComponent.h
@@ -98,6 +98,9 @@ private:
 	/** The desired age of the System instance.  This is only relevant when using the DesiredAge age update mode. */
 	float DesiredAge;
 
+	/** The last desired age value that was handled by the tick function.  This is only relevant when using the DesiredAgeNoSeek age update mode. */
+	float LastHandledDesiredAge;
+
 	/** Whether or not the component can render while seeking to the desired age. */
 	bool bCanRenderWhileSeeking;
 
