Don't tick Niagara or Cascade world managers on dedicated servers.

diff --git a/Engine/Plugins/FX/Niagara/Source/Niagara/Private/NiagaraWorldManager.cpp b/Engine/Plugins/FX/Niagara/Source/Niagara/Private/NiagaraWorldManager.cpp
--- a/Engine/Plugins/FX/Niagara/Source/Niagara/Private/NiagaraWorldManager.cpp
+++ b/Engine/Plugins/FX/Niagara/Source/Niagara/Private/NiagaraWorldManager.cpp
@@ -146,6 +146,7 @@ FNiagaraWorldManager::FNiagaraWorldManager(UWorld* InWorld)
 		TickFunc.EndTickGroup = GNiagaraAllowAsyncWorkToEndOfFrame ? TG_LastDemotable : (ETickingGroup)TickFunc.TickGroup;
 		TickFunc.bCanEverTick = true;
 		TickFunc.bStartWithTickEnabled = true;
+		TickFunc.bAllowTickOnDedicatedServer = false;
 		TickFunc.bHighPriority = true;
 		TickFunc.Owner = this;
 		TickFunc.RegisterTickFunction(InWorld->PersistentLevel);
diff --git a/Engine/Source/Runtime/Engine/Private/Particles/ParticleSystemManager.cpp b/Engine/Source/Runtime/Engine/Private/Particles/ParticleSystemManager.cpp
--- a/Engine/Source/Runtime/Engine/Private/Particles/ParticleSystemManager.cpp
+++ b/Engine/Source/Runtime/Engine/Private/Particles/ParticleSystemManager.cpp
@@ -228,6 +228,7 @@ FParticleSystemWorldManager::FParticleSystemWorldManager(UWorld* InWorld)
 		TickFunc.EndTickGroup = TickFunc.TickGroup;
 		TickFunc.bCanEverTick = true;
 		TickFunc.bStartWithTickEnabled = true;
+		TickFunc.bAllowTickOnDedicatedServer = false;
 		TickFunc.bHighPriority = true;
 		TickFunc.Owner = this;
 		TickFunc.RegisterTickFunction(InWorld->PersistentLevel);
