Merging in various scalability fixes from other branch.

- Preventing distance and visbility culling when editor doesn't have focus to avoid false positives.
- Adding new stats NiagaraSystemCounts to help accurately tracks total instance counts per system.
- Insta completing systems with no active emitters.
- Removing PendingKillOrUnreachable components during the PreGCBeginDestroy callback. Covers the calse where the whole world is made unreachable but not torn down correctly.
- More cleanly handle burst FX using Asleep FXTypes and post a log message for us to fix the content when this occurs.

diff --git a/Engine/Plugins/FX/Niagara/Source/Niagara/Classes/NiagaraSystem.h b/Engine/Plugins/FX/Niagara/Source/Niagara/Classes/NiagaraSystem.h
--- a/Engine/Plugins/FX/Niagara/Source/Niagara/Classes/NiagaraSystem.h
+++ b/Engine/Plugins/FX/Niagara/Source/Niagara/Classes/NiagaraSystem.h
@@ -368,6 +368,7 @@ public:
 	uint32 bFixedBounds : 1;
 
 	TStatId GetStatID(bool bGameThread, bool bConcurrent)const;
+	void SetInstanceCountStat(int32 NumInstances)const;
 
 	UPROPERTY(EditAnywhere, Category = "Script Fast Path")
 	ENiagaraFastPathMode FastPathMode;
@@ -499,6 +500,8 @@ protected:
 	mutable TStatId StatID_GT_CNC;
 	mutable TStatId StatID_RT;
 	mutable TStatId StatID_RT_CNC;
+
+	mutable TStatId StatID_InstanceCount;
 #endif
 
 	FNiagaraSystemScalabilitySettings CurrentScalabilitySettings;
diff --git a/Engine/Plugins/FX/Niagara/Source/Niagara/Niagara.Build.cs b/Engine/Plugins/FX/Niagara/Source/Niagara/Niagara.Build.cs
--- a/Engine/Plugins/FX/Niagara/Source/Niagara/Niagara.Build.cs
+++ b/Engine/Plugins/FX/Niagara/Source/Niagara/Niagara.Build.cs
@@ -17,9 +17,10 @@ public class Niagara : ModuleRules
                 "TimeManagement",
                 "Renderer",
                 "JsonUtilities",
-				"Json",
-				"AudioPlatformConfiguration",
-				"SignalProcessing"
+		"Json",
+		"AudioPlatformConfiguration",
+		"SignalProcessing",
+		"ApplicationCore"
             }
         );
 
diff --git a/Engine/Plugins/FX/Niagara/Source/Niagara/Private/NiagaraComponent.cpp b/Engine/Plugins/FX/Niagara/Source/Niagara/Private/NiagaraComponent.cpp
--- a/Engine/Plugins/FX/Niagara/Source/Niagara/Private/NiagaraComponent.cpp
+++ b/Engine/Plugins/FX/Niagara/Source/Niagara/Private/NiagaraComponent.cpp
@@ -449,6 +449,7 @@ UNiagaraComponent::UNiagaraComponent(const FObjectInitializer& ObjectInitializer
 	, bActivateShouldResetWhenReady(false)
 	, bDidAutoAttach(false)
 	, bAllowScalability(true)
+	, bIsCulledByScalability(false)
 	//, bIsChangingAutoAttachment(false)
 	, ScalabilityManagerHandle(INDEX_NONE)
 	, OwnerLOD(0)
@@ -868,16 +869,19 @@ void UNiagaraComponent::ActivateInternal(bool bReset /* = false */, bool bIsScal
 		UnregisterWithScalabilityManager();
 	}
 
-	if (!bIsScalabilityCull && ScalabilityManagerHandle != INDEX_NONE)
+	if (!bIsScalabilityCull && bIsCulledByScalability)
 	{
+		check(ScalabilityManagerHandle != INDEX_NONE);
 		//If this is a non scalability activate call and we're still registered with the manager.
 		//If we reach this point then we must have been previously culled by scalability so bail here.
 		return;
 	}
 
-	if (RegisterWithScalabilityManagerOrPreCull())
+	bIsCulledByScalability = false;
+	if (ShouldPreCull())
 	{
 		//We have decided to pre cull the system.
+		bIsCulledByScalability = true;
 		return;
 	}
 
@@ -936,6 +940,8 @@ void UNiagaraComponent::ActivateInternal(bool bReset /* = false */, bool bIsScal
 		return;
 	}
 
+	RegisterWithScalabilityManager();
+
 	SystemInstance->Activate(ResetMode);
 
 	/** We only need to tick the component if we require solo mode. */
@@ -949,10 +955,14 @@ void UNiagaraComponent::Deactivate()
 
 void UNiagaraComponent::DeactivateInternal(bool bIsScalabilityCull /* = false */)
 {
-	// Unregister with the scalability manager if this is a genuine deactivation from outside.
-	// The scalability manager itself can call this function when culling systems.
-	if (bIsScalabilityCull == false)
+	if (bIsScalabilityCull)
 	{
+		bIsCulledByScalability = true;
+	}
+	else
+	{
+		// Unregister with the scalability manager if this is a genuine deactivation from outside.
+		// The scalability manager itself can call this function when culling systems.
 		UnregisterWithScalabilityManager();
 	}
 
@@ -998,7 +1008,11 @@ void UNiagaraComponent::DeactivateImmediateInternal(bool bIsScalabilityCull)
 
 	//Unregister with the scalability manager if this is a genuine deactivation from outside.
 	//The scalability manager itself can call this function when culling systems.
-	if (bIsScalabilityCull == false)
+	if (bIsScalabilityCull)
+	{
+		bIsCulledByScalability = true;
+	}
+	else
 	{
 		UnregisterWithScalabilityManager();
 	}
@@ -1011,9 +1025,9 @@ void UNiagaraComponent::DeactivateImmediateInternal(bool bIsScalabilityCull)
 	}
 }
 
-bool UNiagaraComponent::RegisterWithScalabilityManagerOrPreCull()
+bool UNiagaraComponent::ShouldPreCull()
 {
-	if (ScalabilityManagerHandle == INDEX_NONE && bAllowScalability)
+	if (bAllowScalability)
 	{
 		if (UNiagaraSystem* System = GetAsset())
 		{
@@ -1026,10 +1040,6 @@ bool UNiagaraComponent::RegisterWithScalabilityManagerOrPreCull()
 						//If we're just set to check on spawn then check for precull here.
 						return WorldMan->ShouldPreCull(GetAsset(), this);
 					}
-					else
-					{
-						WorldMan->RegisterWithScalabilityManager(this);
-					}
 				}
 			}
 		}
@@ -1038,6 +1048,23 @@ bool UNiagaraComponent::RegisterWithScalabilityManagerOrPreCull()
 	return false;
 }
 
+void UNiagaraComponent::RegisterWithScalabilityManager()
+{
+	if (ScalabilityManagerHandle == INDEX_NONE && bAllowScalability)
+	{
+		if (UNiagaraSystem* System = GetAsset())
+		{
+			if (UNiagaraEffectType* EffectType = System->GetEffectType())
+			{
+				if (FNiagaraWorldManager* WorldMan = FNiagaraWorldManager::Get(GetWorld()))
+				{
+					WorldMan->RegisterWithScalabilityManager(this);
+				}
+			}
+		}
+	}
+}
+
 void UNiagaraComponent::UnregisterWithScalabilityManager()
 {
 	if (ScalabilityManagerHandle != INDEX_NONE)
@@ -1047,6 +1074,7 @@ void UNiagaraComponent::UnregisterWithScalabilityManager()
 			WorldMan->UnregisterWithScalabilityManager(this);
 		}
 	}
+	bIsCulledByScalability = false;
 	ScalabilityManagerHandle = INDEX_NONE;//Just to be sure our state is unregistered.
 }
 
@@ -1082,6 +1110,24 @@ void UNiagaraComponent::OnSystemComplete()
 	{
 		CancelAutoAttachment(/*bDetachFromParent=*/ true);
 	}
+
+	if (!bIsCulledByScalability && ScalabilityManagerHandle != INDEX_NONE)
+	{
+		//Can we be sure this isn't going to spam erroneously?
+		if (UNiagaraEffectType* EffectType = GetAsset()->GetEffectType())
+		{
+			//Only trigger warning if we're not being deactivated/completed from the outside and this is a natural completion by the system itself.
+			if (EffectType->CullReaction == ENiagaraCullReaction::DeactivateImmediateResume || EffectType->CullReaction == ENiagaraCullReaction::DeactivateResume)
+			{
+				//If we're completing naturally, i.e. we're a burst/non-looping system then we shouldn't be using a mode reactivates the effect.
+				UE_LOG(LogNiagara, Warning, TEXT("Niagara Effect has completed naturally but has an effect type with the \"Asleep\" cull reaction. If an effect like this is culled before it can complete then it could leak into the scalability manager and be reactivated incorrectly. Please verify this is using the correct EffctType.\nComponent:%s\nSystem:%s")
+					, *GetFullName(), *GetAsset()->GetFullName());
+			}
+		}
+
+		//We've completed naturally so unregister with the scalability manager.
+		UnregisterWithScalabilityManager();
+	}
 }
 
 void UNiagaraComponent::DestroyInstance()
diff --git a/Engine/Plugins/FX/Niagara/Source/Niagara/Private/NiagaraScalabilityManager.cpp b/Engine/Plugins/FX/Niagara/Source/Niagara/Private/NiagaraScalabilityManager.cpp
--- a/Engine/Plugins/FX/Niagara/Source/Niagara/Private/NiagaraScalabilityManager.cpp
+++ b/Engine/Plugins/FX/Niagara/Source/Niagara/Private/NiagaraScalabilityManager.cpp
@@ -42,17 +42,22 @@ void FNiagaraScalabilityManager::AddReferencedObjects(FReferenceCollector& Colle
 void FNiagaraScalabilityManager::PreGarbageCollectBeginDestroy()
 {
 	//After the GC has potentially nulled out references to the components we were tracking we clear them out here.
-	//This should only be in the case where MarkPendingKill() is called directly. Typical component destruction will unregister in OnComponentDestroyed().
+	//This should only be in the case where MarkPendingKill() is called directly. Typical component destruction will unregister in OnComponentDestroyed() or OnUnregister().
 	//Components then just clear their handle in BeginDestroy knowing they've already been removed from the manager.
 	//I would prefer some pre BeginDestroy() callback into the component in which I could cleanly unregister with the manager in all cases but I don't think that's possible.
 	int32 CompIdx = ManagedComponents.Num();
 	while (--CompIdx >= 0)
 	{
-		if (ManagedComponents[CompIdx] == nullptr)
+		UNiagaraComponent* Comp = ManagedComponents[CompIdx];
+		if (Comp == nullptr)
 		{
 			//UE_LOG(LogNiagara, Warning, TEXT("Unregister from PreGCBeginDestroy @%d/%d - %s"), CompIdx, ManagedComponents.Num(), *EffectType->GetName());
 			UnregisterAt(CompIdx);
 		}
+		else if (Comp->IsPendingKillOrUnreachable())
+		{
+			Unregister(Comp);
+		}
 	}
 }
 
diff --git a/Engine/Plugins/FX/Niagara/Source/Niagara/Private/NiagaraStats.h b/Engine/Plugins/FX/Niagara/Source/Niagara/Private/NiagaraStats.h
--- a/Engine/Plugins/FX/Niagara/Source/Niagara/Private/NiagaraStats.h
+++ b/Engine/Plugins/FX/Niagara/Source/Niagara/Private/NiagaraStats.h
@@ -26,4 +26,6 @@ DECLARE_CYCLE_STAT(TEXT("RT Concurrent Total"), STAT_NiagaraOverview_RT_CNC, STA
 
 
 DECLARE_STATS_GROUP(TEXT("Niagara Systems"), STATGROUP_NiagaraSystems, STATCAT_Advanced);
-DECLARE_STATS_GROUP(TEXT("Niagara Emitters"), STATGROUP_NiagaraEmitters, STATCAT_Advanced);
\ No newline at end of file
+DECLARE_STATS_GROUP(TEXT("Niagara Emitters"), STATGROUP_NiagaraEmitters, STATCAT_Advanced);
+
+DECLARE_STATS_GROUP(TEXT("Niagara System Counts"), STATGROUP_NiagaraSystemCounts, STATCAT_Advanced);
\ No newline at end of file
diff --git a/Engine/Plugins/FX/Niagara/Source/Niagara/Private/NiagaraSystem.cpp b/Engine/Plugins/FX/Niagara/Source/Niagara/Private/NiagaraSystem.cpp
--- a/Engine/Plugins/FX/Niagara/Source/Niagara/Private/NiagaraSystem.cpp
+++ b/Engine/Plugins/FX/Niagara/Source/Niagara/Private/NiagaraSystem.cpp
@@ -1439,13 +1439,31 @@ TStatId UNiagaraSystem::GetStatID(bool bGameThread, bool bConcurrent)const
 	return TStatId();
 }
 
+void UNiagaraSystem::SetInstanceCountStat(int32 NumInstances)const
+{
+#if STATS
+	if (!StatID_GT.IsValidStat())
+	{
+		GenerateStatID();
+	}
+
+	if (FThreadStats::IsCollectingData())
+	{
+		FThreadStats::AddMessage(StatID_InstanceCount.GetName(), EStatOperation::Set, int64(NumInstances));
+		TRACE_STAT_SET(StatID_InstanceCount.GetName(), int64(Value));
+	}
+#endif
+}
+
 void UNiagaraSystem::GenerateStatID()const
 {
 #if STATS
-	StatID_GT = FDynamicStats::CreateStatId<FStatGroup_STATGROUP_NiagaraSystems>(GetPathName() + TEXT("[GT]"));
-	StatID_GT_CNC = FDynamicStats::CreateStatId<FStatGroup_STATGROUP_NiagaraSystems>(GetPathName() + TEXT("[GT_CNC]"));
-	StatID_RT = FDynamicStats::CreateStatId<FStatGroup_STATGROUP_NiagaraSystems>(GetPathName() + TEXT("[RT]"));
-	StatID_RT_CNC = FDynamicStats::CreateStatId<FStatGroup_STATGROUP_NiagaraSystems>(GetPathName() + TEXT("[RT_CNC]"));
+	StatID_GT = FDynamicStats::CreateStatId<FStatGroup_STATGROUP_NiagaraSystems>(GetPathName() + TEXT(" [GT]"));
+	StatID_GT_CNC = FDynamicStats::CreateStatId<FStatGroup_STATGROUP_NiagaraSystems>(GetPathName() + TEXT(" [GT_CNC]"));
+	StatID_RT = FDynamicStats::CreateStatId<FStatGroup_STATGROUP_NiagaraSystems>(GetPathName() + TEXT(" [RT]"));
+	StatID_RT_CNC = FDynamicStats::CreateStatId<FStatGroup_STATGROUP_NiagaraSystems>(GetPathName() + TEXT(" [RT_CNC]"));
+
+	StatID_InstanceCount = FDynamicStats::CreateStatIdInt64<FStatGroup_STATGROUP_NiagaraSystemCounts>(GetPathName());
 #endif
 }
 
diff --git a/Engine/Plugins/FX/Niagara/Source/Niagara/Private/NiagaraSystemInstance.cpp b/Engine/Plugins/FX/Niagara/Source/Niagara/Private/NiagaraSystemInstance.cpp
--- a/Engine/Plugins/FX/Niagara/Source/Niagara/Private/NiagaraSystemInstance.cpp
+++ b/Engine/Plugins/FX/Niagara/Source/Niagara/Private/NiagaraSystemInstance.cpp
@@ -790,47 +790,66 @@ void FNiagaraSystemInstance::Reset(FNiagaraSystemInstance::EResetMode Mode)
 		bBindParams = !IsDisabled();
 	}
 	
-	if (bBindParams)
+	//If none of our emitters actually made it out of the init process we can just bail here before we ever tick.
+	bool bHasActiveEmitters = false;
+	for (auto& Inst : Emitters)
 	{
-		ResetParameters();
-		BindParameters();
+		if (!Inst->IsComplete())
+		{
+			bHasActiveEmitters = true;
+			break;
+		}
 	}
 
 	SetRequestedExecutionState(ENiagaraExecutionState::Active);
-	SetActualExecutionState(ENiagaraExecutionState::Active);
-
-	if (bBindParams)
+	if (bHasActiveEmitters)
 	{
-		InitDataInterfaces();
-	}
+		if (bBindParams)
+		{
+			ResetParameters();
+			BindParameters();
+		}
 
-	//Interface init can disable the system.
-	if (!IsComplete())
-	{
-		ComputeEmittersExecutionOrder();
+		SetActualExecutionState(ENiagaraExecutionState::Active);
 
-		bPendingSpawn = true;
-		SystemSimulation->AddInstance(this);
+		if (bBindParams)
+		{
+			InitDataInterfaces();
+		}
 
-		UNiagaraSystem* System = GetSystem();
-		if (System->NeedsWarmup())
+		//Interface init can disable the system.
+		if (!IsComplete())
 		{
-			int32 WarmupTicks = System->GetWarmupTickCount();
-			float WarmupDt = System->GetWarmupTickDelta();
-			
-			AdvanceSimulation(WarmupTicks, WarmupDt);
+			ComputeEmittersExecutionOrder();
+
+			bPendingSpawn = true;
+			SystemSimulation->AddInstance(this);
 
-			//Reset age to zero.
-			Age = 0.0f;
-			TickCount = 0;
+			UNiagaraSystem* System = GetSystem();
+			if (System->NeedsWarmup())
+			{
+				int32 WarmupTicks = System->GetWarmupTickCount();
+				float WarmupDt = System->GetWarmupTickDelta();
+
+				AdvanceSimulation(WarmupTicks, WarmupDt);
+
+				//Reset age to zero.
+				Age = 0.0f;
+				TickCount = 0;
+			}
 		}
-	}
 
-	if (Component)
+		if (Component)
+		{
+			// This system may not tick again immediately so we mark the render state dirty here so that
+			// the renderers will be reset this frame.
+			Component->MarkRenderDynamicDataDirty();
+		}
+	}
+	else
 	{
-		// This system may not tick again immediately so we mark the render state dirty here so that
-		// the renderers will be reset this frame.
-		Component->MarkRenderDynamicDataDirty();
+		SetActualExecutionState(ENiagaraExecutionState::Complete);
+		Complete();
 	}
 }
 
diff --git a/Engine/Plugins/FX/Niagara/Source/Niagara/Private/NiagaraSystemSimulation.cpp b/Engine/Plugins/FX/Niagara/Source/Niagara/Private/NiagaraSystemSimulation.cpp
--- a/Engine/Plugins/FX/Niagara/Source/Niagara/Private/NiagaraSystemSimulation.cpp
+++ b/Engine/Plugins/FX/Niagara/Source/Niagara/Private/NiagaraSystemSimulation.cpp
@@ -43,6 +43,9 @@ DECLARE_CYCLE_STAT(TEXT("System Sim Init (DirectBindings) [GT]"), STAT_NiagaraSy
 DECLARE_CYCLE_STAT(TEXT("ForcedWaitForAsync"), STAT_NiagaraSystemSim_ForceWaitForAsync, STATGROUP_Niagara);
 DECLARE_CYCLE_STAT(TEXT("ForcedWait Fake Stall"), STAT_NiagaraSystemSim_ForceWaitFakeStall, STATGROUP_Niagara);
 
+
+DECLARE_DWORD_COUNTER_STAT(TEXT("Total System Instances"), STAT_TotalNiagaraSystemInstances, STATGROUP_NiagaraSystemCounts);
+
 static int32 GbDumpSystemData = 0;
 static FAutoConsoleVariableRef CVarNiagaraDumpSystemData(
 	TEXT("fx.DumpSystemData"),
@@ -792,6 +795,9 @@ void FNiagaraSystemSimulation::Tick_GameThread(float DeltaSeconds, const FGraphE
 	static const auto EffectsQualityCVar = IConsoleManager::Get().FindConsoleVariable(TEXT("sg.EffectsQuality"));
 	FNiagaraSystemSimulationTickContext Context(this, SystemInstances, MainDataSet, DeltaSeconds, SpawnNum, EffectsQualityCVar->GetInt(), MyCompletionGraphEvent);
 
+	System->SetInstanceCountStat(SystemInstances.Num());
+	INC_DWORD_STAT_BY(STAT_TotalNiagaraSystemInstances, SystemInstances.Num());
+
 	//Now kick of the concurrent tick.
 	if (Context.bTickAsync)
 	{
diff --git a/Engine/Plugins/FX/Niagara/Source/Niagara/Private/NiagaraWorldManager.cpp b/Engine/Plugins/FX/Niagara/Source/Niagara/Private/NiagaraWorldManager.cpp
--- a/Engine/Plugins/FX/Niagara/Source/Niagara/Private/NiagaraWorldManager.cpp
+++ b/Engine/Plugins/FX/Niagara/Source/Niagara/Private/NiagaraWorldManager.cpp
@@ -19,6 +19,7 @@
 #include "NiagaraComponentPool.h"
 #include "NiagaraComponent.h"
 #include "NiagaraEffectType.h"
+#include "HAL/PlatformApplicationMisc.h"
 
 DECLARE_CYCLE_STAT(TEXT("Niagara Manager Update Scalability Managers [GT]"), STAT_UpdateScalabilityManagers, STATGROUP_Niagara);
 DECLARE_CYCLE_STAT(TEXT("Niagara Manager Tick [GT]"), STAT_NiagaraWorldManTick, STATGROUP_Niagara);
@@ -138,6 +139,7 @@ FName FNiagaraWorldManagerTickFunction::DiagnosticContext(bool bDetailed)
 FNiagaraWorldManager::FNiagaraWorldManager(UWorld* InWorld)
 	: World(InWorld)
 	, CachedEffectsQuality(INDEX_NONE)
+	, bAppHasFocus(true)
 {
 	for (int32 TickGroup=0; TickGroup < NiagaraNumTickGroups; ++TickGroup)
 	{
@@ -571,6 +573,8 @@ void FNiagaraWorldManager::Tick(ETickingGroup TickGroup, float DeltaSeconds, ELe
 	{
 		FNiagaraSharedObject::FlushDeletionList();
 
+		bAppHasFocus = FPlatformApplicationMisc::IsThisApplicationForeground();
+
 		// Cache player view locations for all system instances to access
 		//-TODO: Do we need to do this per tick group?
 		bCachedPlayerViewLocationsValid = true;
@@ -798,7 +802,8 @@ void FNiagaraWorldManager::CalculateScalabilityState(UNiagaraSystem* System, con
 	OwnerLODCull(EffectType, ScalabilitySettings, Component, OutState);
 	
 	//Can't cull dynamic bounds by visibility
-	if (System->bFixedBounds)
+
+	if (System->bFixedBounds && bAppHasFocus)
 	{
 		VisibilityCull(EffectType, ScalabilitySettings, Component, OutState);
 	}
diff --git a/Engine/Plugins/FX/Niagara/Source/Niagara/Public/NiagaraComponent.h b/Engine/Plugins/FX/Niagara/Source/Niagara/Public/NiagaraComponent.h
--- a/Engine/Plugins/FX/Niagara/Source/Niagara/Public/NiagaraComponent.h
+++ b/Engine/Plugins/FX/Niagara/Source/Niagara/Public/NiagaraComponent.h
@@ -174,7 +174,8 @@ public:
 	void DeactivateInternal(bool bIsScalabilityCull);
 	void DeactivateImmediateInternal(bool bIsScalabilityCull);
 
-	bool RegisterWithScalabilityManagerOrPreCull();
+	bool ShouldPreCull();
+	void RegisterWithScalabilityManager();
 	void UnregisterWithScalabilityManager();
 
 	public:
@@ -533,6 +534,9 @@ private:
 	/** True if this component is allowed to perform scalability checks and potentially be culled etc. Occasionally it is useful to disable this for specific components. E.g. Effects on the local player. */
 	uint32 bAllowScalability : 1;
 
+	/** True if this component has been culled by the scalability manager. */
+	uint32 bIsCulledByScalability : 1;
+
 	/** Flag to mark us as currently changing auto attachment as part of Activate/Deactivate so we don't reset in the OnAttachmentChanged() callback. */
 	//uint32 bIsChangingAutoAttachment : 1;
 
diff --git a/Engine/Plugins/FX/Niagara/Source/Niagara/Public/NiagaraWorldManager.h b/Engine/Plugins/FX/Niagara/Source/Niagara/Public/NiagaraWorldManager.h
--- a/Engine/Plugins/FX/Niagara/Source/Niagara/Public/NiagaraWorldManager.h
+++ b/Engine/Plugins/FX/Niagara/Source/Niagara/Public/NiagaraWorldManager.h
@@ -241,5 +241,8 @@ private:
 
 	UPROPERTY(transient)
 	TMap<UNiagaraEffectType*, FNiagaraScalabilityManager> ScalabilityManagers;
+
+	/** True if the app has focus. We prevent some culling if the app doesn't have focus as it can interefre. */
+	bool bAppHasFocus;
 };
 
