Crash fix for non-desktop platforms.

diff --git a/Engine/Plugins/FX/Niagara/Source/Niagara/Private/NiagaraWorldManager.cpp b/Engine/Plugins/FX/Niagara/Source/Niagara/Private/NiagaraWorldManager.cpp
--- a/Engine/Plugins/FX/Niagara/Source/Niagara/Private/NiagaraWorldManager.cpp
+++ b/Engine/Plugins/FX/Niagara/Source/Niagara/Private/NiagaraWorldManager.cpp
@@ -573,7 +573,11 @@ void FNiagaraWorldManager::Tick(ETickingGroup TickGroup, float DeltaSeconds, ELe
 	{
 		FNiagaraSharedObject::FlushDeletionList();
 
+#if PLATFORM_DESKTOP
 		bAppHasFocus = FPlatformApplicationMisc::IsThisApplicationForeground();
+#else
+		bAppHasFocus = true;
+#endif
 
 		// Cache player view locations for all system instances to access
 		//-TODO: Do we need to do this per tick group?
