Niagara
- Mark FNiagaraUserParameterBinding as NIAGARA_API so it can be used in other modules

diff --git a/Engine/Plugins/FX/Niagara/Source/Niagara/Public/NiagaraCommon.h b/Engine/Plugins/FX/Niagara/Source/Niagara/Public/NiagaraCommon.h
--- a/Engine/Plugins/FX/Niagara/Source/Niagara/Public/NiagaraCommon.h
+++ b/Engine/Plugins/FX/Niagara/Source/Niagara/Public/NiagaraCommon.h
@@ -766,7 +766,7 @@ namespace FNiagaraUtilities
 };
 
 USTRUCT()
-struct FNiagaraUserParameterBinding
+struct NIAGARA_API FNiagaraUserParameterBinding
 {
 	GENERATED_USTRUCT_BODY()
 
