Niagara GpuSim fix for Android
- Disable overlap compute on Android
- Workaround for Mali Compiler bug that assumes texture index is warp invariant

diff --git a/Engine/Plugins/FX/Niagara/Shaders/Private/NiagaraDrawIndirectArgsGen.usf b/Engine/Plugins/FX/Niagara/Shaders/Private/NiagaraDrawIndirectArgsGen.usf
--- a/Engine/Plugins/FX/Niagara/Shaders/Private/NiagaraDrawIndirectArgsGen.usf
+++ b/Engine/Plugins/FX/Niagara/Shaders/Private/NiagaraDrawIndirectArgsGen.usf
@@ -29,7 +29,30 @@ void MainCS(uint TaskIndex : SV_DispatchThreadID)
 		RWDrawIndirectArgs[ArgOffset + 3] = 0; // BaseVertexLocation
 		RWDrawIndirectArgs[ArgOffset + 4] = 0; // StartInstanceLocation
 	}
-	else if (TaskIndex < TaskCount.z)
+/*
+Explanation of why we are using two ifs here instead of if/else: 
+Mali requires texture indexes to be warp invariant, which means that the if / if code translates to
+	if id.x < uniform :
+	[to 5 stores]
+	[branch to end]
+	if <the other condition>
+	do 1 image store
+	<end>
+and the if / else translates to
+	if uniform > ID.x: goto 4 image stores
+	if <other condition> :
+	setup registers
+	goto <last store>
+	<4 images stores branch target>
+	do 4 stores
+	setup registers
+	<last store target>
+	do 1 store
+Notice how it completely ignores the fact that there are two different texture buffers, which leads the compiler to do an "optimization" and separate the last store.This means that the else condition does indeed evaluate correctly, but we are writing in the wrong texture!
+Note: This is imagestore specific and not compute specific and is a known issue for Arm. This is fixed in driver version R21
+*/
+
+	if (TaskIndex < TaskCount.z && TaskIndex >= TaskCount.x)
 	{
 		const uint InfoOffset = TaskCount.x * NIAGARA_DRAW_INDIRECT_TASK_INFO_SIZE + (TaskIndex - TaskCount.x);
 		RWInstanceCounts[TaskInfos[InfoOffset + 0]] = 0;
diff --git a/Engine/Plugins/FX/Niagara/Source/Niagara/Classes/NiagaraEmitterInstanceBatcher.h b/Engine/Plugins/FX/Niagara/Source/Niagara/Classes/NiagaraEmitterInstanceBatcher.h
--- a/Engine/Plugins/FX/Niagara/Source/Niagara/Classes/NiagaraEmitterInstanceBatcher.h
+++ b/Engine/Plugins/FX/Niagara/Source/Niagara/Classes/NiagaraEmitterInstanceBatcher.h
@@ -191,6 +191,7 @@ private:
 		return (FNiagaraComputeInstanceData*)(PackedData + sizeof(uint32));
 	}
 
+	bool UseOverlapCompute();
 	void FinishDispatches();
 	void ReleaseTicks();
 	void ResizeFreeIDsListSizesBuffer(uint32 NumInstances);
diff --git a/Engine/Plugins/FX/Niagara/Source/Niagara/Private/NiagaraEmitterInstanceBatcher.cpp b/Engine/Plugins/FX/Niagara/Source/Niagara/Private/NiagaraEmitterInstanceBatcher.cpp
--- a/Engine/Plugins/FX/Niagara/Source/Niagara/Private/NiagaraEmitterInstanceBatcher.cpp
+++ b/Engine/Plugins/FX/Niagara/Source/Niagara/Private/NiagaraEmitterInstanceBatcher.cpp
@@ -194,6 +194,11 @@ void NiagaraEmitterInstanceBatcher::ReleaseTicks()
 	Ticks_RT.Empty(0);
 }
 
+bool NiagaraEmitterInstanceBatcher::UseOverlapCompute()
+{
+	return !IsMobilePlatform(ShaderPlatform) && GNiagaraOverlapCompute;
+}
+
 bool NiagaraEmitterInstanceBatcher::ResetDataInterfaces(const FNiagaraGPUSystemTick& Tick, FNiagaraComputeInstanceData *Instance, FRHICommandList &RHICmdList, const FNiagaraShaderRef& ComputeShader ) const
 {
 	bool ValidSpawnStage = true;
@@ -711,14 +716,14 @@ void NiagaraEmitterInstanceBatcher::ExecuteAll(FRHICommandList& RHICmdList, FRHI
 		}
 
 		// Set bIsFinalTick for the last tick of each context and reset the scratch index.
-		const int32 ScrachIndexReset = GNiagaraOverlapCompute ? 0 : INDEX_NONE;
+		const int32 ScrachIndexReset = UseOverlapCompute() ? 0 : INDEX_NONE;
 		for (FNiagaraComputeExecutionContext* Context : RelevantContexts)
 		{
 			RelevantTicks[Context->ScratchIndex]->bIsFinalTick = true;
 			Context->ScratchIndex = ScrachIndexReset;
 		}
 
-		if (GNiagaraOverlapCompute)
+		if (UseOverlapCompute())
 		{
 			// Transpose now only once the data to get all independent tick per pass
 			SimPasses.Reserve(2); // Safe bet!
diff --git a/Engine/Plugins/FX/Niagara/Source/NiagaraShader/Private/NiagaraShaderDerivedDataVersion.h b/Engine/Plugins/FX/Niagara/Source/NiagaraShader/Private/NiagaraShaderDerivedDataVersion.h
--- a/Engine/Plugins/FX/Niagara/Source/NiagaraShader/Private/NiagaraShaderDerivedDataVersion.h
+++ b/Engine/Plugins/FX/Niagara/Source/NiagaraShader/Private/NiagaraShaderDerivedDataVersion.h
@@ -11,4 +11,4 @@ NiagaraShaderDerivedDataVersion.h: Shader derived data version for Niagara.
 // In case of merge conflicts with DDC versions, you MUST generate a new GUID and set this new
 // guid as version
 
-#define NIAGARASHADERMAP_DERIVEDDATA_VER		TEXT("{C6FAD17E-D2B1-4A9D-8242-D5FAC8A833EB}")
+#define NIAGARASHADERMAP_DERIVEDDATA_VER		TEXT("906AEEB436374DF1ACE9A4AA6BABCED1")
