Fixes Niagara sprite velocity aligment and motion blur when particle are in local space.

diff --git a/Engine/Plugins/FX/Niagara/Shaders/Private/NiagaraSpriteVertexFactory.ush b/Engine/Plugins/FX/Niagara/Shaders/Private/NiagaraSpriteVertexFactory.ush
--- a/Engine/Plugins/FX/Niagara/Shaders/Private/NiagaraSpriteVertexFactory.ush
+++ b/Engine/Plugins/FX/Niagara/Shaders/Private/NiagaraSpriteVertexFactory.ush
@@ -157,8 +157,6 @@ struct FVertexFactoryInterpolantsVSToPS
 
 struct FVertexFactoryIntermediates
 {
-	/** The position of the particle in non-translated world space. */
-	float3 ParticleWorldPosition;
 	/** The position of the vertex in translated world space. */
 	float3 VertexWorldPosition;
 	/** Last frame's position of the vertex in translated world space. */
@@ -185,6 +183,7 @@ struct FVertexFactoryIntermediates
 	/** The velocity of the particle, XYZ: direction, W: speed. */
 	float4 ParticleVelocity;
 
+	/** The particle worldspace position */
 	float3 Position;
 	float3 OldPosition;
 	
@@ -440,7 +439,7 @@ float3 SafeNormalizeWithFallback(float3 V, float3 Fallback)
 	return L > 0.000001 ? V / L : Fallback; 
 }
 
-void GetTangents(FVertexFactoryInput Input, FVertexFactoryIntermediates Intermediates, float3 WorldPosition, float3 OldWorldPosition, out float3 OutRight, out float3 OutUp)
+void GetTangents(FVertexFactoryInput Input, FVertexFactoryIntermediates Intermediates, float3 ParticleDirection, out float3 OutRight, out float3 OutUp)
 {
 	// Select camera up/right vectors.
 	float3 ResolvedViewRight = lerp(ResolvedView.ViewRight, ResolvedView.HMDViewNoRollRight, NiagaraSpriteVF.RemoveHMDRoll);
@@ -452,9 +451,8 @@ void GetTangents(FVertexFactoryInput Input, FVertexFactoryIntermediates Intermed
 	float3 CameraFacingVector = -ResolvedView.ViewForward;
 
 	// Determine the vector from the particle to the camera and the particle's movement direction.
-	float3 CameraVector = ResolvedView.WorldCameraOrigin.xyz - WorldPosition.xyz;
-	float3 CameraDirection = SafeNormalize(CameraVector);
-	float3 ParticleDirection	= SafeNormalizeWithFallback(WorldPosition.xyz - OldWorldPosition.xyz, float3(0,0,1));
+	const float3 CameraVector = ResolvedView.WorldCameraOrigin - Intermediates.Position;
+	const float3 CameraDirection = SafeNormalize(CameraVector);
 
 	float3 CustomFacingVector = NiagaraSpriteVFLooseParameters.ParticleFacingMode == FACING_MODE_CUSTOMFACINGVECTOR ? Intermediates.CustomFacingVector : CameraDirection;
 	
@@ -555,14 +553,14 @@ float3x3 CalcTangentBasis(FVertexFactoryIntermediates Intermediates)
 	// ENM_Spherical
 	else if (NiagaraSpriteVF.NormalsType < 1.5f)
 	{
-		float3 TangentZ = normalize(Intermediates.ParticleWorldPosition - NiagaraSpriteVF.NormalsSphereCenter.xyz);
+		float3 TangentZ = normalize(Intermediates.Position - NiagaraSpriteVF.NormalsSphereCenter.xyz);
 		Result[2] = TangentZ;	
 	}
 	// ENM_Cylindrical
 	else
 	{
-		float3 ClosestPointOnCylinder = NiagaraSpriteVF.NormalsSphereCenter.xyz + dot(NiagaraSpriteVF.NormalsCylinderUnitDirection.xyz, Intermediates.ParticleWorldPosition - NiagaraSpriteVF.NormalsSphereCenter.xyz) * NiagaraSpriteVF.NormalsCylinderUnitDirection.xyz;
-		float3 TangentZ = normalize(Intermediates.ParticleWorldPosition - ClosestPointOnCylinder);
+		float3 ClosestPointOnCylinder = NiagaraSpriteVF.NormalsSphereCenter.xyz + dot(NiagaraSpriteVF.NormalsCylinderUnitDirection.xyz, Intermediates.Position - NiagaraSpriteVF.NormalsSphereCenter.xyz) * NiagaraSpriteVF.NormalsCylinderUnitDirection.xyz;
+		float3 TangentZ = normalize(Intermediates.Position - ClosestPointOnCylinder);
 		Result[2] = TangentZ;	
 	}
 
@@ -836,7 +834,6 @@ float3 GetNiagaraParticleAlignmentVector(uint InstanceID)
 FVertexFactoryIntermediates GetVertexFactoryIntermediates(FVertexFactoryInput Input)
 {
 	FVertexFactoryIntermediates Intermediates = (FVertexFactoryIntermediates)0;
-
 	uint ParticleID = Input.InstanceId;
 
 #if !SHADOW_DEPTH_SHADER && !SPHERICAL_OPACITY_FOR_SHADOW_DEPTHS
@@ -846,8 +843,7 @@ FVertexFactoryIntermediates GetVertexFactoryIntermediates(FVertexFactoryInput In
 		ParticleID = NiagaraSpriteVFLooseParameters.SortedIndices[NiagaraSpriteVFLooseParameters.SortedIndicesOffset + Input.InstanceId];
 	}
 #endif
-	Intermediates.Position = GetNiagaraParticlePosition(ParticleID);
-	Intermediates.OldPosition = GetNiagaraParticlePrevPosition(ParticleID) - GetNiagaraParticleVelocity(ParticleID) * NiagaraSpriteVF.DeltaSeconds;
+
 	float2 ParticleSize = GetNiagaraParticleSize(ParticleID);
 	Intermediates.SizeRotSubImage = float4(ParticleSize.x, ParticleSize.y, GetNiagaraParticleRotation(ParticleID), GetNiagaraParticleSubimage(ParticleID));
 	Intermediates.Color	= GetNiagaraParticleColor(ParticleID);
@@ -867,46 +863,51 @@ FVertexFactoryIntermediates GetVertexFactoryIntermediates(FVertexFactoryInput In
 	Intermediates.DynamicParameter3 = GetNiagaraParticleDynamicParameters3(ParticleID);
 #endif
 
+	// Particle world position, before any sprite aligment.
+	Intermediates.Position = GetNiagaraParticlePosition(ParticleID);
 
-	// World position.
-	
-	float3 ParticleWorldPosition = Intermediates.Position;
-	float3 ParticleOldWorldPosition = Intermediates.OldPosition;
-	float CameraOffset = GetNiagaraCameraOffset(ParticleID);
-	float3 CameraVector = ResolvedView.WorldCameraOrigin.xyz - ParticleWorldPosition;
-	float3 CameraDirection = SafeNormalize(CameraVector);
-
-	// Do we need the hack that forces particles at the camera origin to be *really* behind the camera origin?
-	float3 OffsetPosition = CameraDirection * CameraOffset;
-
-	ParticleWorldPosition = ParticleWorldPosition + OffsetPosition;
-	ParticleOldWorldPosition = ParticleOldWorldPosition + OffsetPosition;
-	
-	Intermediates.ParticleWorldPosition = ParticleWorldPosition;
+	const float3 ParticleVelocity = GetNiagaraParticleVelocity(ParticleID);
 
 	// Tangents.
-	float3 Right,Up;
-	GetTangents(Input, Intermediates, ParticleWorldPosition,ParticleOldWorldPosition,Right,Up);
-	Intermediates.TangentUp = Up;
-	Intermediates.TangentRight = Right;
+	{
+		const float3 ParticleDirection = SafeNormalizeWithFallback(ParticleVelocity, float3(0,0,1));
+		float3 Right,Up;
+
+		GetTangents(Input, Intermediates, ParticleDirection, Right, Up);
+		Intermediates.TangentUp = Up;
+		Intermediates.TangentRight = Right;
+	}
 
 	float2 UVForPosition;
 	float2 UVForTexturing;
 	float2 UVForTexturingUnflipped;
 	ComputeBillboardUVs(Input, Intermediates, UVForPosition, UVForTexturing, UVForTexturingUnflipped);
 
-	float2 UVScale = GetNiagaraUVScale(ParticleID);
+	const float2 UVScale = GetNiagaraUVScale(ParticleID);
 	UVForTexturing.xy = float2(0.5f, 0.5f) + (UVForTexturing.xy - float2(0.5f, 0.5f) ) *  UVScale.xy; 
+	const float2 Size = abs(Intermediates.SizeRotSubImage.xy);
+
+	// Vertex position and prev vertex position.
+	{
+		// Do we need the hack that forces particles at the camera origin to be *really* behind the camera origin?
+		const float3 CameraOffset = SafeNormalize(ResolvedView.WorldCameraOrigin - Intermediates.Position) * GetNiagaraCameraOffset(ParticleID);
 
-	// Vertex position.
-	float4 VertexWorldPosition = float4(ParticleWorldPosition + ResolvedView.PreViewTranslation,1);
-	float2 Size = abs(Intermediates.SizeRotSubImage.xy);
-	VertexWorldPosition += Size.x * (UVForPosition.x + NiagaraSpriteVF.PivotOffset.x) * float4(Right,0);
-	VertexWorldPosition += Size.y * (UVForPosition.y + NiagaraSpriteVF.PivotOffset.y) * float4(Up,0);
-	Intermediates.VertexWorldPosition = VertexWorldPosition.xyz;
+		const float2x3 Tangents = float2x3(Intermediates.TangentRight, Intermediates.TangentUp);
+		const float3 VertexOffset = CameraOffset + mul(Size * (UVForPosition + NiagaraSpriteVF.PivotOffset), Tangents);
 
-	// Previous vertex position.
-	Intermediates.PreviousVertexWorldPosition = VertexWorldPosition.xyz + (ParticleOldWorldPosition - ParticleWorldPosition);
+		Intermediates.VertexWorldPosition = Intermediates.Position + VertexOffset;
+		Intermediates.PreviousVertexWorldPosition = GetNiagaraParticlePrevPosition(ParticleID) + VertexOffset + ParticleVelocity * NiagaraSpriteVF.DeltaSeconds;
+
+#if USE_PARTICLE_POSITION
+		// Do we really want the camera offset here ? What about when GENERATE_SPHERICAL_PARTICLE_NORMALS
+		Intermediates.TranslatedWorldPositionAndSize.xyz = Intermediates.Position + CameraOffset + ResolvedView.PreViewTranslation;
+		Intermediates.TranslatedWorldPositionAndSize.w = .5f * min(Size.x, Size.y);
+#if GENERATE_SPHERICAL_PARTICLE_NORMALS
+		// Recenter the particle position for GENERATE_SPHERICAL_PARTICLE_NORMALS to use the right center.
+		Intermediates.TranslatedWorldPositionAndSize.xyz += mul(Size * (float2(.5, .5) + NiagaraSpriteVF.PivotOffset), Tangents);
+#endif // GENERATE_SPHERICAL_PARTICLE_NORMALS
+#endif // USE_PARTICLE_POSITION
+	}
 	
 	// SubUV.
 	float SubImageIndex = Intermediates.SizeRotSubImage.w;
@@ -924,24 +925,13 @@ FVertexFactoryIntermediates GetVertexFactoryIntermediates(FVertexFactoryInput In
 	Intermediates.TexCoord_Unflipped.zw = (float2(SubImageBH, SubImageBV) + UVForTexturingUnflipped) * NiagaraSpriteVF.SubImageSize.zw;
 #endif
 	Intermediates.SubImageLerp = (NiagaraSpriteVF.SubImageBlendMode == 1) ? SubImageLerp : 0.0f;
-	
-
-#if USE_PARTICLE_POSITION
-	float ParticleRadius = .5f * min(Intermediates.SizeRotSubImage.x, Intermediates.SizeRotSubImage.y);
-	float3 ParticleTranslatedWorldPosition = ParticleWorldPosition + ResolvedView.PreViewTranslation.xyz;
-	Intermediates.TranslatedWorldPositionAndSize = float4(ParticleTranslatedWorldPosition, ParticleRadius);
-
-#if GENERATE_SPHERICAL_PARTICLE_NORMALS
-	// Recenter the particle position for GENERATE_SPHERICAL_PARTICLE_NORMALS to use the right center.
-	Intermediates.TranslatedWorldPositionAndSize.xyz += Size.x * (NiagaraSpriteVF.PivotOffset.x + .5) * Right + Size.y * (NiagaraSpriteVF.PivotOffset.y + .5) * Up;
-#endif
-#endif
 
 #if USE_PARTICLE_VELOCITY
-	float3 ParticleVelocity = (ParticleWorldPosition - ParticleOldWorldPosition) * (1.0f/NiagaraSpriteVF.DeltaSeconds);
-	float3 DirectionVector = lerp(ParticleVelocity, Intermediates.CustomAlignmentVector, NiagaraSpriteVFLooseParameters.ParticleAlignmentMode == ALIGNMENT_MODE_CUSTOMALIGNMENT);
-	Intermediates.ParticleVelocity.xyz = normalize(DirectionVector);
-	Intermediates.ParticleVelocity.w = length(DirectionVector);
+	{
+		const float3 DirectionVector = NiagaraSpriteVFLooseParameters.ParticleAlignmentMode == ALIGNMENT_MODE_CUSTOMALIGNMENT ? Intermediates.CustomAlignmentVector : ParticleVelocity;
+		Intermediates.ParticleVelocity.xyz = normalize(DirectionVector);
+		Intermediates.ParticleVelocity.w = length(DirectionVector);
+	}
 #endif
 
 #if USE_PARTICLE_TIME
@@ -969,7 +959,7 @@ FVertexFactoryIntermediates GetVertexFactoryIntermediates(FVertexFactoryInput In
 
 float4 VertexFactoryGetWorldPosition(FVertexFactoryInput Input, FVertexFactoryIntermediates Intermediates)
 {
-	return float4(Intermediates.VertexWorldPosition,1);
+	return float4(Intermediates.VertexWorldPosition + ResolvedView.PreViewTranslation, 1);
 }
 
 float3 VertexFactoryGetWorldNormal(FVertexFactoryInput Input, FVertexFactoryIntermediates Intermediates)
@@ -995,7 +985,7 @@ float3 VertexFactoryGetPositionForVertexLighting(FVertexFactoryInput Input, FVer
 {
 #if (TRANSLUCENCY_LIGHTING_VOLUMETRIC_PERVERTEX_DIRECTIONAL || TRANSLUCENCY_LIGHTING_VOLUMETRIC_PERVERTEX_NONDIRECTIONAL)
 	// Do per-vertex lighting with particle position instead of vertex position when we're using SubUV cutouts, since the vertex positions vary per-frame which would cause popping
-	return NiagaraSpriteVFLooseParameters.NumCutoutVerticesPerFrame > 0 ? (Intermediates.ParticleWorldPosition + ResolvedView.PreViewTranslation.xyz) : TranslatedWorldPosition;
+	return NiagaraSpriteVFLooseParameters.NumCutoutVerticesPerFrame > 0 ? (Intermediates.Position + ResolvedView.PreViewTranslation.xyz) : TranslatedWorldPosition;
 #else
 	return TranslatedWorldPosition;
 #endif
@@ -1185,7 +1175,7 @@ FVertexFactoryInterpolantsVSToPS VertexFactoryAssignInterpolants(FVertexFactoryI
 
 float4 VertexFactoryGetPreviousWorldPosition(FVertexFactoryInput Input, FVertexFactoryIntermediates Intermediates)
 {
-	return float4(Intermediates.PreviousVertexWorldPosition, 1.0f);
+	return float4(Intermediates.PreviousVertexWorldPosition + ResolvedView.PrevPreViewTranslation, 1.0f);
 }
 
 /**
