Making useful runtime error not generate an error during valid test case.

diff --git a/Engine/Plugins/FX/Niagara/Source/Niagara/Private/NiagaraComponent.cpp b/Engine/Plugins/FX/Niagara/Source/Niagara/Private/NiagaraComponent.cpp
--- a/Engine/Plugins/FX/Niagara/Source/Niagara/Private/NiagaraComponent.cpp
+++ b/Engine/Plugins/FX/Niagara/Source/Niagara/Private/NiagaraComponent.cpp
@@ -66,6 +66,14 @@ static FAutoConsoleVariableRef CVarNiagaraComponentWarnNullAsset(
 	ECVF_Default
 );
 
+static int32 GNiagaraComponentWarnAsleepCullReaction = 1;
+static FAutoConsoleVariableRef CVarNiagaraComponentWarnAsleepCullReaction(
+	TEXT("fx.Niagara.ComponentWarnAsleepCullReaction"),
+	GNiagaraComponentWarnAsleepCullReaction,
+	TEXT("When enabled we will warn if a NiagaraComponent completes naturally but has Asleep mode set for cullreaction."),
+	ECVF_Default
+);
+
 void DumpNiagaraComponents(UWorld* World)
 {
 	for (TActorIterator<AActor> ActorItr(World); ActorItr; ++ActorItr)
@@ -1119,9 +1127,12 @@ void UNiagaraComponent::OnSystemComplete()
 			//Only trigger warning if we're not being deactivated/completed from the outside and this is a natural completion by the system itself.
 			if (EffectType->CullReaction == ENiagaraCullReaction::DeactivateImmediateResume || EffectType->CullReaction == ENiagaraCullReaction::DeactivateResume)
 			{
-				//If we're completing naturally, i.e. we're a burst/non-looping system then we shouldn't be using a mode reactivates the effect.
-				UE_LOG(LogNiagara, Warning, TEXT("Niagara Effect has completed naturally but has an effect type with the \"Asleep\" cull reaction. If an effect like this is culled before it can complete then it could leak into the scalability manager and be reactivated incorrectly. Please verify this is using the correct EffctType.\nComponent:%s\nSystem:%s")
-					, *GetFullName(), *GetAsset()->GetFullName());
+				if (GNiagaraComponentWarnAsleepCullReaction == 1)
+				{
+					//If we're completing naturally, i.e. we're a burst/non-looping system then we shouldn't be using a mode reactivates the effect.
+					UE_LOG(LogNiagara, Warning, TEXT("Niagara Effect has completed naturally but has an effect type with the \"Asleep\" cull reaction. If an effect like this is culled before it can complete then it could leak into the scalability manager and be reactivated incorrectly. Please verify this is using the correct EffctType.\nComponent:%s\nSystem:%s")
+						, *GetFullName(), *GetAsset()->GetFullName());
+				}
 			}
 		}
 
