Attempting to mitigate crash until better repro case is found

diff --git a/Engine/Plugins/FX/Niagara/Source/Niagara/Classes/NiagaraDataInterfaceStaticMesh.h b/Engine/Plugins/FX/Niagara/Source/Niagara/Classes/NiagaraDataInterfaceStaticMesh.h
--- a/Engine/Plugins/FX/Niagara/Source/Niagara/Classes/NiagaraDataInterfaceStaticMesh.h
+++ b/Engine/Plugins/FX/Niagara/Source/Niagara/Classes/NiagaraDataInterfaceStaticMesh.h
@@ -99,6 +99,8 @@ protected:
 
 struct FNDIStaticMesh_InstanceData
 {
+	FNDIStaticMesh_InstanceData() : Mesh(nullptr) {}
+
 	 //Cached ptr to component we sample from. 
 	TWeakObjectPtr<USceneComponent> Component;
 
diff --git a/Engine/Plugins/FX/Niagara/Source/Niagara/Private/NiagaraDataInterfaceStaticMesh.cpp b/Engine/Plugins/FX/Niagara/Source/Niagara/Private/NiagaraDataInterfaceStaticMesh.cpp
--- a/Engine/Plugins/FX/Niagara/Source/Niagara/Private/NiagaraDataInterfaceStaticMesh.cpp
+++ b/Engine/Plugins/FX/Niagara/Source/Niagara/Private/NiagaraDataInterfaceStaticMesh.cpp
@@ -29,6 +29,14 @@ const FString UNiagaraDataInterfaceStaticMesh::InstanceWorldVelocityName(TEXT("I
 const FString UNiagaraDataInterfaceStaticMesh::AreaWeightedSamplingName(TEXT("AreaWeightedSamplingName_"));
 const FString UNiagaraDataInterfaceStaticMesh::NumTexCoordName(TEXT("NumTexCoordName_"));
 
+static int32 GNiagaraFailStaticMeshDataInterface = 1;
+static FAutoConsoleVariableRef CVarNiagaraFailStaticMeshDataInterface(
+	TEXT("fx.Niagara.FailStaticMeshDataInterface"),
+	GNiagaraFailStaticMeshDataInterface,
+	TEXT("When enabled we will fail out using static mesh data interfaces."),
+	ECVF_Default
+);
+
 FStaticMeshFilteredAreaWeightedSectionSampler::FStaticMeshFilteredAreaWeightedSectionSampler()
 	: Res(nullptr)
 	, Owner(nullptr)
@@ -233,6 +241,12 @@ bool FNDIStaticMesh_InstanceData::Init(UNiagaraDataInterfaceStaticMesh* Interfac
 		TransformInverseTransposed = Transform.InverseFast().GetTransposed();
 	}
 
+	if (GNiagaraFailStaticMeshDataInterface != 0)
+	{
+		UE_LOG(LogNiagara, Log, TEXT("StaticMesh data interface aborting because \"fx.Niagara.FailStaticMeshDataInterface 1\". Failed InitPerInstanceData - %s"), *Interface->GetFullName());
+		return false;
+	}
+
 	if (!Mesh)
 	{
 		UE_LOG(LogNiagara, Log, TEXT("StaticMesh data interface has no valid mesh. Failed InitPerInstanceData - %s"), *Interface->GetFullName());
@@ -278,6 +292,7 @@ bool FNDIStaticMesh_InstanceData::Init(UNiagaraDataInterfaceStaticMesh* Interfac
 
 	Sampler.Init(&Res, this);
 
+
 	return true;
 }
 
