Niagara - Fix an issue where you can not create a new set variables node with a data interface and you can not add data interfaces to existing set variables nodes.

diff --git a/Engine/Plugins/FX/Niagara/Source/NiagaraEditor/Private/ViewModels/Stack/NiagaraStackGraphUtilities.cpp b/Engine/Plugins/FX/Niagara/Source/NiagaraEditor/Private/ViewModels/Stack/NiagaraStackGraphUtilities.cpp
--- a/Engine/Plugins/FX/Niagara/Source/NiagaraEditor/Private/ViewModels/Stack/NiagaraStackGraphUtilities.cpp
+++ b/Engine/Plugins/FX/Niagara/Source/NiagaraEditor/Private/ViewModels/Stack/NiagaraStackGraphUtilities.cpp
@@ -1506,7 +1506,7 @@ void FNiagaraStackGraphUtilities::GetNewParameterAvailableTypes(TArray<FNiagaraT
 	for (const FNiagaraTypeDefinition& RegisteredParameterType : FNiagaraTypeRegistry::GetRegisteredParameterTypes())
 	{
 		//Object types only allowed in user namespace at the moment.
-		if (RegisteredParameterType.IsUObject() && Namespace != FNiagaraParameterHandle::UserNamespace)
+		if (RegisteredParameterType.IsUObject() && RegisteredParameterType.IsDataInterface() == false && Namespace != FNiagaraParameterHandle::UserNamespace)
 		{
 			continue;
 		}
