Working to make the static mesh data interface more safe. In some cases the areaweightedsampler array was empty. Also doing a weak poitner check in case the mesh goes away due to GC.

diff --git a/Engine/Plugins/FX/Niagara/Source/Niagara/Classes/NiagaraDataInterfaceStaticMesh.h b/Engine/Plugins/FX/Niagara/Source/Niagara/Classes/NiagaraDataInterfaceStaticMesh.h
--- a/Engine/Plugins/FX/Niagara/Source/Niagara/Classes/NiagaraDataInterfaceStaticMesh.h
+++ b/Engine/Plugins/FX/Niagara/Source/Niagara/Classes/NiagaraDataInterfaceStaticMesh.h
@@ -102,7 +102,10 @@ struct FNDIStaticMesh_InstanceData
 	FNDIStaticMesh_InstanceData() : Mesh(nullptr) {}
 
 	 //Cached ptr to component we sample from. 
-	TWeakObjectPtr<USceneComponent> Component;
+	TWeakObjectPtr<USceneComponent> SafeComponent_GT;
+
+	// Cached ptr to the mesh so that we can make sure that we haven't been deleted.
+	TWeakObjectPtr<UStaticMesh> SafeMesh_GT;
 
 	//Cached ptr to actual mesh we sample from. 
 	UStaticMesh* Mesh;
diff --git a/Engine/Plugins/FX/Niagara/Source/Niagara/Private/NiagaraDataInterfaceStaticMesh.cpp b/Engine/Plugins/FX/Niagara/Source/Niagara/Private/NiagaraDataInterfaceStaticMesh.cpp
--- a/Engine/Plugins/FX/Niagara/Source/Niagara/Private/NiagaraDataInterfaceStaticMesh.cpp
+++ b/Engine/Plugins/FX/Niagara/Source/Niagara/Private/NiagaraDataInterfaceStaticMesh.cpp
@@ -29,7 +29,7 @@ const FString UNiagaraDataInterfaceStaticMesh::InstanceWorldVelocityName(TEXT("I
 const FString UNiagaraDataInterfaceStaticMesh::AreaWeightedSamplingName(TEXT("AreaWeightedSamplingName_"));
 const FString UNiagaraDataInterfaceStaticMesh::NumTexCoordName(TEXT("NumTexCoordName_"));
 
-static int32 GNiagaraFailStaticMeshDataInterface = 1;
+static int32 GNiagaraFailStaticMeshDataInterface = 0;
 static FAutoConsoleVariableRef CVarNiagaraFailStaticMeshDataInterface(
 	TEXT("fx.Niagara.FailStaticMeshDataInterface"),
 	GNiagaraFailStaticMeshDataInterface,
@@ -57,12 +57,27 @@ float FStaticMeshFilteredAreaWeightedSectionSampler::GetWeights(TArray<float>& O
 	float Total = 0.0f;
 	OutWeights.Empty(Owner->GetValidSections().Num());
 	FStaticMeshLODResources& LODRes = Owner->Mesh->RenderData->LODResources[0];
-	for (int32 i = 0; i < Owner->GetValidSections().Num(); ++i)
+
+	if (Owner->Mesh->bSupportUniformlyDistributedSampling && LODRes.AreaWeightedSectionSamplers.Num() > 0)
 	{
-		int32 SecIdx = Owner->GetValidSections()[i];
-		float T = LODRes.AreaWeightedSectionSamplers[SecIdx].GetTotalWeight();
-		OutWeights.Add(T);
-		Total += T;
+		for (int32 i = 0; i < Owner->GetValidSections().Num(); ++i)
+		{
+			int32 SecIdx = Owner->GetValidSections()[i];
+			float T = LODRes.AreaWeightedSectionSamplers[SecIdx].GetTotalWeight();
+			OutWeights.Add(T);
+			Total += T;
+		}
+	}
+	else
+	{
+		for (int32 i = 0; i < Owner->GetValidSections().Num(); ++i)
+		{
+			int32 SecIdx = Owner->GetValidSections()[i];
+			float T = 1.0f;
+			OutWeights.Add(T);
+			Total += T;
+		}
+
 	}
 	return Total;
 }
@@ -150,7 +165,9 @@ bool FNDIStaticMesh_InstanceData::Init(UNiagaraDataInterfaceStaticMesh* Interfac
 {
 	check(SystemInstance);
 	UStaticMesh* PrevMesh = Mesh;
-	Component = nullptr;
+	SafeComponent_GT = nullptr;
+	SafeMesh_GT = nullptr;
+
 	Mesh = nullptr;
 	Transform = FMatrix::Identity;
 	TransformInverseTransposed = FMatrix::Identity;
@@ -161,7 +178,7 @@ bool FNDIStaticMesh_InstanceData::Init(UNiagaraDataInterfaceStaticMesh* Interfac
 
 	if (Interface->SourceComponent)
 	{
-		Component = Interface->SourceComponent;
+		SafeComponent_GT = Interface->SourceComponent;
 		Mesh = Interface->SourceComponent->GetStaticMesh();
 	}
 	else if (Interface->Source)
@@ -180,11 +197,11 @@ bool FNDIStaticMesh_InstanceData::Init(UNiagaraDataInterfaceStaticMesh* Interfac
 		if (SourceComp)
 		{
 			Mesh = SourceComp->GetStaticMesh();
-			Component = SourceComp;
+			SafeComponent_GT = SourceComp;
 		}
 		else
 		{
-			Component = Interface->Source->GetRootComponent();
+			SafeComponent_GT = Interface->Source->GetRootComponent();
 		}
 	}
 	else
@@ -193,12 +210,12 @@ bool FNDIStaticMesh_InstanceData::Init(UNiagaraDataInterfaceStaticMesh* Interfac
 		{
 			if (UStaticMeshComponent* ParentComp = Cast<UStaticMeshComponent>(SimComp->GetAttachParent()))
 			{
-				Component = ParentComp;
+				SafeComponent_GT = ParentComp;
 				Mesh = ParentComp->GetStaticMesh();
 			}
 			else if (UStaticMeshComponent* OuterComp = SimComp->GetTypedOuter<UStaticMeshComponent>())
 			{
-				Component = OuterComp;
+				SafeComponent_GT = OuterComp;
 				Mesh = OuterComp->GetStaticMesh();
 			}
 			else if (AActor* Owner = SimComp->GetAttachmentRootActor())
@@ -212,32 +229,32 @@ bool FNDIStaticMesh_InstanceData::Init(UNiagaraDataInterfaceStaticMesh* Interfac
 						if (PossibleMesh != nullptr && PossibleMesh->bAllowCPUAccess)
 						{
 							Mesh = PossibleMesh;
-							Component = SourceComp;
+							SafeComponent_GT = SourceComp;
 							break;
 						}
 					}
 				}
 			}
 
-			if (!Component.IsValid())
+			if (!SafeComponent_GT.IsValid())
 			{
-				Component = SimComp;
+				SafeComponent_GT = SimComp;
 			}
 		}
 	}
 
-	check(Component.IsValid());
+	check(SafeComponent_GT.IsValid());
 
 	if (!Mesh && Interface->DefaultMesh)
 	{
 		Mesh = Interface->DefaultMesh;
 	}
 
-	if (Component.IsValid() && Mesh)
+	if (SafeComponent_GT.IsValid() && Mesh)
 	{
 		PrevTransform = Transform;
 		PrevTransformInverseTransposed = TransformInverseTransposed;
-		Transform = Component->GetComponentToWorld().ToMatrixWithScale();
+		Transform = SafeComponent_GT->GetComponentToWorld().ToMatrixWithScale();
 		TransformInverseTransposed = Transform.InverseFast().GetTransposed();
 	}
 
@@ -259,7 +276,7 @@ bool FNDIStaticMesh_InstanceData::Init(UNiagaraDataInterfaceStaticMesh* Interfac
 		return false;
 	}
 
-	if (!Component.IsValid())
+	if (!SafeComponent_GT.IsValid())
 	{
 		UE_LOG(LogNiagara, Log, TEXT("StaticMesh data interface has no valid component. Failed InitPerInstanceData - %s"), *Interface->GetFullName());
 		return false;
@@ -290,6 +307,7 @@ bool FNDIStaticMesh_InstanceData::Init(UNiagaraDataInterfaceStaticMesh* Interfac
 		return false;
 	}
 
+	SafeMesh_GT = Mesh;
 	Sampler.Init(&Res, this);
 
 
@@ -300,12 +318,19 @@ bool FNDIStaticMesh_InstanceData::ResetRequired(UNiagaraDataInterfaceStaticMesh*
 {
 	check(GetActualMesh());
 
-	if (!Component.IsValid())
+	if (!SafeComponent_GT.IsValid())
 	{
 		//The component we were bound to is no longer valid so we have to trigger a reset.
 		return true;
 	}
 
+	UStaticMesh* MeshSafe = SafeMesh_GT.Get();
+	if (MeshSafe != Mesh)
+	{
+		//The static mesh we were bound to is no longer valid so we have to trigger a reset.
+		return true;
+	}
+
 	if (Interface != nullptr && ChangeId != Interface->ChangeId)
 	{
 		return true;
@@ -334,11 +359,11 @@ bool FNDIStaticMesh_InstanceData::Tick(UNiagaraDataInterfaceStaticMesh* Interfac
 	else
 	{
 		DeltaSeconds = InDeltaSeconds;
-		if (Component.IsValid() && Mesh)
+		if (SafeComponent_GT.IsValid() && SafeMesh_GT.IsValid())
 		{
 			PrevTransform = Transform;
 			PrevTransformInverseTransposed = TransformInverseTransposed;
-			Transform = Component->GetComponentToWorld().ToMatrixWithScale();
+			Transform = SafeComponent_GT->GetComponentToWorld().ToMatrixWithScale();
 			TransformInverseTransposed = Transform.InverseFast().GetTransposed();
 		}
 		else
@@ -994,7 +1019,7 @@ DEFINE_NDI_FUNC_BINDER(UNiagaraDataInterfaceStaticMesh, GetVertexPosition);
 void UNiagaraDataInterfaceStaticMesh::GetVMExternalFunction(const FVMExternalFunctionBindingInfo& BindingInfo, void* InstanceData, FVMExternalFunction &OutFunc)
 {
 	FNDIStaticMesh_InstanceData* InstData = (FNDIStaticMesh_InstanceData*)InstanceData;
-	check(InstData && InstData->Mesh && InstData->Component.IsValid());
+	check(InstData && InstData->Mesh && InstData->SafeComponent_GT.IsValid());
 
 	bool bNeedsVertexPositions = false;
 	bool bNeedsVertexColors = false;
@@ -1250,14 +1275,28 @@ template<>
 FORCEINLINE int32 UNiagaraDataInterfaceStaticMesh::RandomSection<TIntegralConstant<bool, true>, true>(FRandomStream& RandStream, FStaticMeshLODResources& Res, FNDIStaticMesh_InstanceData* InstData)
 {
 	checkSlow(InstData->GetValidSections().Num() > 0);
-	int32 Idx = InstData->GetAreaWeigtedSampler().GetEntryIndex(RandStream.GetFraction(), RandStream.GetFraction());
-	return InstData->GetValidSections()[Idx];
+	if (InstData->GetAreaWeigtedSampler().GetNumEntries() > 0)
+	{
+		int32 Idx = InstData->GetAreaWeigtedSampler().GetEntryIndex(RandStream.GetFraction(), RandStream.GetFraction());
+		return InstData->GetValidSections()[Idx];
+	}
+	else
+	{
+		return 0;
+	}
 }
 
 template<>
 FORCEINLINE int32 UNiagaraDataInterfaceStaticMesh::RandomSection<TIntegralConstant<bool, true>, false>(FRandomStream& RandStream, FStaticMeshLODResources& Res, FNDIStaticMesh_InstanceData* InstData)
 {
-	return Res.AreaWeightedSampler.GetEntryIndex(RandStream.GetFraction(), RandStream.GetFraction());
+	if (Res.AreaWeightedSampler.GetNumEntries() > 0)
+	{
+		return Res.AreaWeightedSampler.GetEntryIndex(RandStream.GetFraction(), RandStream.GetFraction());
+	}
+	else
+	{
+		return 0;
+	}
 }
 
 template<>
@@ -1294,36 +1333,60 @@ template<>
 FORCEINLINE int32 UNiagaraDataInterfaceStaticMesh::RandomTriIndex<TIntegralConstant<bool, true>, true>(FRandomStream& RandStream, FStaticMeshLODResources& Res, FNDIStaticMesh_InstanceData* InstData)
 {
 	int32 SecIdx = RandomSection<TIntegralConstant<bool, true>, true>(RandStream, Res, InstData);
-	FStaticMeshSection&  Sec = Res.Sections[SecIdx];
-	int32 Tri = Res.AreaWeightedSectionSamplers[SecIdx].GetEntryIndex(RandStream.GetFraction(), RandStream.GetFraction());
-	return (Sec.FirstIndex / 3) + Tri;
+	if (SecIdx < Res.Sections.Num() && SecIdx < Res.AreaWeightedSectionSamplers.Num())
+	{
+		FStaticMeshSection& Sec = Res.Sections[SecIdx];
+		if (Res.AreaWeightedSectionSamplers[SecIdx].GetNumEntries() > 0)
+		{
+			int32 Tri = Res.AreaWeightedSectionSamplers[SecIdx].GetEntryIndex(RandStream.GetFraction(), RandStream.GetFraction());
+			return (Sec.FirstIndex / 3) + Tri;
+		}
+		return (Sec.FirstIndex / 3);
+	}
+	return 0;
 }
 
 template<>
 FORCEINLINE int32 UNiagaraDataInterfaceStaticMesh::RandomTriIndex<TIntegralConstant<bool, true>, false>(FRandomStream& RandStream, FStaticMeshLODResources& Res, FNDIStaticMesh_InstanceData* InstData)
 {
 	int32 SecIdx = RandomSection<TIntegralConstant<bool, true>, false>(RandStream, Res, InstData);
-	FStaticMeshSection&  Sec = Res.Sections[SecIdx];
-	int32 Tri = Res.AreaWeightedSectionSamplers[SecIdx].GetEntryIndex(RandStream.GetFraction(), RandStream.GetFraction());
-	return (Sec.FirstIndex / 3) + Tri;
+	if (SecIdx < Res.Sections.Num() && SecIdx < Res.AreaWeightedSectionSamplers.Num())
+	{
+		FStaticMeshSection&  Sec = Res.Sections[SecIdx];
+		if (Res.AreaWeightedSectionSamplers[SecIdx].GetNumEntries() > 0)
+		{
+			int32 Tri = Res.AreaWeightedSectionSamplers[SecIdx].GetEntryIndex(RandStream.GetFraction(), RandStream.GetFraction());
+			return (Sec.FirstIndex / 3) + Tri;
+		}
+		return (Sec.FirstIndex / 3);
+	}
+	return 0;
 }
 
 template<>
 FORCEINLINE int32 UNiagaraDataInterfaceStaticMesh::RandomTriIndex<TIntegralConstant<bool, false>, true>(FRandomStream& RandStream, FStaticMeshLODResources& Res, FNDIStaticMesh_InstanceData* InstData)
 {
 	int32 SecIdx = RandomSection<TIntegralConstant<bool, false>, true>(RandStream, Res, InstData);
-	FStaticMeshSection&  Sec = Res.Sections[SecIdx];
-	int32 Tri = RandStream.RandRange(0, Sec.NumTriangles - 1);
-	return (Sec.FirstIndex / 3) + Tri;
+	if (SecIdx < Res.Sections.Num())
+	{
+		FStaticMeshSection& Sec = Res.Sections[SecIdx];
+		int32 Tri = RandStream.RandRange(0, Sec.NumTriangles - 1);
+		return (Sec.FirstIndex / 3) + Tri;
+	}
+	return 0;
 }
 
 template<>
 FORCEINLINE int32 UNiagaraDataInterfaceStaticMesh::RandomTriIndex<TIntegralConstant<bool, false>, false>(FRandomStream& RandStream, FStaticMeshLODResources& Res, FNDIStaticMesh_InstanceData* InstData)
 {
 	int32 SecIdx = RandomSection<TIntegralConstant<bool, false>, false>(RandStream, Res, InstData);
-	FStaticMeshSection&  Sec = Res.Sections[SecIdx];
-	int32 Tri = RandStream.RandRange(0, Sec.NumTriangles - 1);
-	return (Sec.FirstIndex / 3) + Tri;
+	if (SecIdx < Res.Sections.Num())
+	{
+		FStaticMeshSection& Sec = Res.Sections[SecIdx];
+		int32 Tri = RandStream.RandRange(0, Sec.NumTriangles - 1);
+		return (Sec.FirstIndex / 3) + Tri;
+	}
+	return 0;
 }
 
 template<typename TAreaWeighted>
@@ -1742,7 +1805,7 @@ void UNiagaraDataInterfaceStaticMesh::GetTriCoordPositionAndVelocity(FVectorVMCo
 	const FPositionVertexBuffer& Positions = Res.VertexBuffers.PositionVertexBuffer;
 
 	const int32 NumTriangles = Indices.Num() / 3;
-	float InvDt = 1.0f / InstData->DeltaSeconds;
+	float InvDt = InstData->DeltaSeconds > 0.0f ? (1.0f / InstData->DeltaSeconds) : 0.0f;
 	for (int32 i = 0; i < Context.NumInstances; ++i)
 	{
 		const int32 Tri = (TriParam.Get() % NumTriangles) * 3;
