Fixed: if multiple instances of a DI were used in a script, the HLSL translator produced all function combinations for each instance.

DataInterfaceRegisteredFunctions accumulated all registered functions from all the instances, then the generator iterated over it for each instance, which was incorrect. The information was already available per instance in FNiagaraScriptDataInterfaceCompileInfo.RegisteredFunctions, so we can just use that and remove DataInterfaceRegisteredFunctions entirely.

Already submitted this once as 11386264, but that introduced a new bug, because RegisteredFunctions can contain duplicates. Did it properly this time, so that GetFunctionHLSL is called once per unique function instance.

diff --git a/Engine/Plugins/FX/Niagara/Source/NiagaraCore/Private/NiagaraCustomVersion.cpp b/Engine/Plugins/FX/Niagara/Source/NiagaraCore/Private/NiagaraCustomVersion.cpp
--- a/Engine/Plugins/FX/Niagara/Source/NiagaraCore/Private/NiagaraCustomVersion.cpp
+++ b/Engine/Plugins/FX/Niagara/Source/NiagaraCore/Private/NiagaraCustomVersion.cpp
@@ -8,4 +8,4 @@ const FGuid FNiagaraCustomVersion::GUID(0xFCF57AFA, 0x50764283, 0xB9A9E658, 0xFF
 // Register the custom version with core
 FCustomVersionRegistration GRegisterNiagaraCustomVersion(FNiagaraCustomVersion::GUID, FNiagaraCustomVersion::LatestVersion, TEXT("NiagaraVer"));
 
-const FGuid FNiagaraCustomVersion::LatestScriptCompileVersion(0xD6750F29, 0x8A554BC2, 0xACB577F3, 0xC6BD12CD);
+const FGuid FNiagaraCustomVersion::LatestScriptCompileVersion(0x99FA2658, 0xE7E24678, 0xBC4E3F17, 0x603EE37F);
diff --git a/Engine/Plugins/FX/Niagara/Source/NiagaraEditor/Private/NiagaraHlslTranslator.cpp b/Engine/Plugins/FX/Niagara/Source/NiagaraEditor/Private/NiagaraHlslTranslator.cpp
--- a/Engine/Plugins/FX/Niagara/Source/NiagaraEditor/Private/NiagaraHlslTranslator.cpp
+++ b/Engine/Plugins/FX/Niagara/Source/NiagaraEditor/Private/NiagaraHlslTranslator.cpp
@@ -1772,9 +1772,16 @@ void FHlslNiagaraTranslator::DefineDataInterfaceHLSL(FString& InHlslOutput)
 			DIInstanceInfo.DIClassName = Info.Type.GetClass()->GetName();
 
 			// Build a list of function instances that will be generated for this DI.
+			TSet<FNiagaraFunctionSignature> SeenFunctions;
 			DIInstanceInfo.GeneratedFunctions.Reserve(Info.RegisteredFunctions.Num());
 			for (const FNiagaraFunctionSignature& OriginalSig : Info.RegisteredFunctions)
 			{
+				if (SeenFunctions.Contains(OriginalSig))
+				{
+					continue;
+				}
+				SeenFunctions.Add(OriginalSig);
+
 				if (!OriginalSig.bSupportsGPU)
 				{
 					Error(FText::Format(LOCTEXT("GPUDataInterfaceFunctionNotSupported", "DataInterface {0} function {1} cannot run on the GPU."), FText::FromName(Info.Type.GetFName()), FText::FromName(OriginalSig.Name)), nullptr, nullptr);
