Fixed: if multiple instances of a DI were used in a script, the HLSL translator produced all function combinations for each instance.

DataInterfaceRegisteredFunctions accumulated all registered functions from all the instances, then the generator iterated over it for each instance, which was incorrect. The information was already available per instance in FNiagaraScriptDataInterfaceCompileInfo.RegisteredFunctions, so we can just use that and remove DataInterfaceRegisteredFunctions entirely.

Already submitted this once as 11386264, but that introduced a new bug, because RegisteredFunctions can contain duplicates. Did it properly this time, so that GetFunctionHLSL is called once per unique function instance.

diff --git a/Engine/Plugins/FX/Niagara/Source/NiagaraCore/Private/NiagaraCustomVersion.cpp b/Engine/Plugins/FX/Niagara/Source/NiagaraCore/Private/NiagaraCustomVersion.cpp
--- a/Engine/Plugins/FX/Niagara/Source/NiagaraCore/Private/NiagaraCustomVersion.cpp
+++ b/Engine/Plugins/FX/Niagara/Source/NiagaraCore/Private/NiagaraCustomVersion.cpp
@@ -8,4 +8,4 @@ const FGuid FNiagaraCustomVersion::GUID(0xFCF57AFA, 0x50764283, 0xB9A9E658, 0xFF
 // Register the custom version with core
 FCustomVersionRegistration GRegisterNiagaraCustomVersion(FNiagaraCustomVersion::GUID, FNiagaraCustomVersion::LatestVersion, TEXT("NiagaraVer"));
 
-const FGuid FNiagaraCustomVersion::LatestScriptCompileVersion(0xB5AC73B2, 0x6F2A47E3, 0x804376B4, 0xA03114E2);
+const FGuid FNiagaraCustomVersion::LatestScriptCompileVersion(0x36511D5D, 0x860A47B4, 0x975E699B, 0x4599D5D4);
