Freeze-related fixes for several niagara DI, and change name of virtual function to make future mistakes of this type into compile errors

diff --git a/Engine/Plugins/FX/Niagara/Source/Niagara/Classes/NiagaraDataInterfaceAudioOscilloscope.h b/Engine/Plugins/FX/Niagara/Source/Niagara/Classes/NiagaraDataInterfaceAudioOscilloscope.h
--- a/Engine/Plugins/FX/Niagara/Source/Niagara/Classes/NiagaraDataInterfaceAudioOscilloscope.h
+++ b/Engine/Plugins/FX/Niagara/Source/Niagara/Classes/NiagaraDataInterfaceAudioOscilloscope.h
@@ -120,6 +120,8 @@ class NIAGARA_API UNiagaraDataInterfaceAudioOscilloscope : public UNiagaraDataIn
 {
 	GENERATED_UCLASS_BODY()
 public:
+
+	DECLARE_NIAGARA_DI_PARAMETER();
 	
 	UPROPERTY(EditAnywhere, Category = "Oscilloscope")
 	USoundSubmix* Submix;
@@ -153,7 +155,6 @@ public:
 
 	virtual bool GetFunctionHLSL(const FNiagaraDataInterfaceGPUParamInfo& ParamInfo, const FNiagaraDataInterfaceGeneratedFunction& FunctionInfo, int FunctionInstanceIndex, FString& OutHLSL) override;
 	virtual void GetParameterDefinitionHLSL(const FNiagaraDataInterfaceGPUParamInfo& ParamInfo, FString& OutHLSL) override;
-	virtual FNiagaraDataInterfaceParametersCS* ConstructComputeParameters() const override;
 
 #if WITH_EDITOR
 	virtual void PostEditChangeProperty(struct FPropertyChangedEvent& PropertyChangedEvent) override;
diff --git a/Engine/Plugins/FX/Niagara/Source/Niagara/Classes/NiagaraDataInterfaceVolumeTexture.h b/Engine/Plugins/FX/Niagara/Source/Niagara/Classes/NiagaraDataInterfaceVolumeTexture.h
--- a/Engine/Plugins/FX/Niagara/Source/Niagara/Classes/NiagaraDataInterfaceVolumeTexture.h
+++ b/Engine/Plugins/FX/Niagara/Source/Niagara/Classes/NiagaraDataInterfaceVolumeTexture.h
@@ -16,6 +16,8 @@ class NIAGARA_API UNiagaraDataInterfaceVolumeTexture : public UNiagaraDataInterf
 	GENERATED_UCLASS_BODY()
 public:
 
+	DECLARE_NIAGARA_DI_PARAMETER();
+
 	UPROPERTY(EditAnywhere, Category = "Texture")
 	UVolumeTexture* Texture;
 
@@ -41,7 +43,6 @@ public:
 	// GPU sim functionality
 	virtual void GetParameterDefinitionHLSL(const FNiagaraDataInterfaceGPUParamInfo& ParamInfo, FString& OutHLSL) override;
 	virtual bool GetFunctionHLSL(const FNiagaraDataInterfaceGPUParamInfo& ParamInfo, const FNiagaraDataInterfaceGeneratedFunction& FunctionInfo, int FunctionInstanceIndex, FString& OutHLSL) override;
-	virtual FNiagaraDataInterfaceParametersCS* ConstructComputeParameters()const override;
 
 	//FRWBuffer& GetGPUBuffer();
 	static const FString TextureName;
diff --git a/Engine/Plugins/FX/Niagara/Source/Niagara/Private/NiagaraDataInterfaceAudioOscilloscope.cpp b/Engine/Plugins/FX/Niagara/Source/Niagara/Private/NiagaraDataInterfaceAudioOscilloscope.cpp
--- a/Engine/Plugins/FX/Niagara/Source/Niagara/Private/NiagaraDataInterfaceAudioOscilloscope.cpp
+++ b/Engine/Plugins/FX/Niagara/Source/Niagara/Private/NiagaraDataInterfaceAudioOscilloscope.cpp
@@ -309,14 +309,14 @@ void UNiagaraDataInterfaceAudioOscilloscope::GetParameterDefinitionHLSL(const FN
 
 struct FNiagaraDataInterfaceParametersCS_AudioOscilloscope : public FNiagaraDataInterfaceParametersCS
 {
-	
+	DECLARE_INLINE_TYPE_LAYOUT(FNiagaraDataInterfaceParametersCS_AudioOscilloscope, NonVirtual);
 
-	FNiagaraDataInterfaceParametersCS_AudioOscilloscope(UNiagaraDataInterfaceAudioOscilloscope* InDataInterface)
+	void Bind(const FNiagaraDataInterfaceGPUParamInfo& ParameterInfo, const class FShaderParameterMap& ParameterMap)
 	{
+		NumChannels.Bind(ParameterMap, *(NumChannelsName + ParameterInfo.DataInterfaceHLSLSymbol));
+		AudioBuffer.Bind(ParameterMap, *(AudioBufferName + ParameterInfo.DataInterfaceHLSLSymbol));
 	}
 
-	virtual ~FNiagaraDataInterfaceParametersCS_AudioOscilloscope() {}
-
 	void Set(FRHICommandList& RHICmdList, const FNiagaraDataInterfaceSetArgs& Context) const
 	{
 		check(IsInRenderingThread());
@@ -330,14 +330,11 @@ struct FNiagaraDataInterfaceParametersCS_AudioOscilloscope : public FNiagaraData
 		RHICmdList.SetShaderResourceViewParameter(ComputeShaderRHI, AudioBuffer.GetBaseIndex(), AudioBufferSRV.SRV);
 	}
 
-	FShaderParameter NumChannels;
-	FShaderResourceParameter AudioBuffer;
+	LAYOUT_FIELD(FShaderParameter, NumChannels);
+	LAYOUT_FIELD(FShaderResourceParameter, AudioBuffer);
 };
 
-FNiagaraDataInterfaceParametersCS* UNiagaraDataInterfaceAudioOscilloscope::ConstructComputeParameters()const
-{
-	return new FNiagaraDataInterfaceParametersCS_AudioOscilloscope(const_cast<UNiagaraDataInterfaceAudioOscilloscope*>(this));
-}
+IMPLEMENT_NIAGARA_DI_PARAMETER(UNiagaraDataInterfaceAudioOscilloscope, FNiagaraDataInterfaceParametersCS_AudioOscilloscope);
 
 void FNiagaraDataInterfaceProxyOscilloscope::OnUpdateResampling(int32 InResolution, float InScopeInMilliseconds)
 {
diff --git a/Engine/Plugins/FX/Niagara/Source/Niagara/Private/NiagaraDataInterfaceVolumeTexture.cpp b/Engine/Plugins/FX/Niagara/Source/Niagara/Private/NiagaraDataInterfaceVolumeTexture.cpp
--- a/Engine/Plugins/FX/Niagara/Source/Niagara/Private/NiagaraDataInterfaceVolumeTexture.cpp
+++ b/Engine/Plugins/FX/Niagara/Source/Niagara/Private/NiagaraDataInterfaceVolumeTexture.cpp
@@ -211,7 +211,7 @@ void UNiagaraDataInterfaceVolumeTexture::GetParameterDefinitionHLSL(const FNiaga
 
 struct FNiagaraDataInterfaceParametersCS_VolumeTexture : public FNiagaraDataInterfaceParametersCS
 {
-	DECLARE_TYPE_LAYOUT(FNiagaraDataInterfaceParametersCS_VolumeTexture, NonVirtual);
+	DECLARE_INLINE_TYPE_LAYOUT(FNiagaraDataInterfaceParametersCS_VolumeTexture, NonVirtual);
 public:
 	void Bind(const FNiagaraDataInterfaceGPUParamInfo& ParameterInfo, const class FShaderParameterMap& ParameterMap)
 	{
@@ -281,6 +281,8 @@ private:
 	LAYOUT_FIELD(FShaderParameter, Dimensions);
 };
 
+IMPLEMENT_NIAGARA_DI_PARAMETER(UNiagaraDataInterfaceVolumeTexture, FNiagaraDataInterfaceParametersCS_VolumeTexture);
+
 void UNiagaraDataInterfaceVolumeTexture::PushToRenderThread()
 {
 	FNiagaraDataInterfaceProxyVolumeTexture* RT_Proxy = GetProxyAs<FNiagaraDataInterfaceProxyVolumeTexture>();
@@ -318,11 +320,6 @@ void UNiagaraDataInterfaceVolumeTexture::PushToRenderThread()
 	);
 }
 
-FNiagaraDataInterfaceParametersCS* UNiagaraDataInterfaceVolumeTexture::ConstructComputeParameters()const
-{
-	return new FNiagaraDataInterfaceParametersCS_VolumeTexture();
-}
-
 void UNiagaraDataInterfaceVolumeTexture::SetTexture(UVolumeTexture* InTexture)
 {
 	if (InTexture)
diff --git a/Engine/Plugins/FX/Niagara/Source/NiagaraCore/Public/NiagaraDataInterfaceBase.h b/Engine/Plugins/FX/Niagara/Source/NiagaraCore/Public/NiagaraDataInterfaceBase.h
--- a/Engine/Plugins/FX/Niagara/Source/NiagaraCore/Public/NiagaraDataInterfaceBase.h
+++ b/Engine/Plugins/FX/Niagara/Source/NiagaraCore/Public/NiagaraDataInterfaceBase.h
@@ -63,10 +63,10 @@ class NIAGARACORE_API UNiagaraDataInterfaceBase : public UNiagaraMergeable
 public:
 	
 	/** Constructs the correct CS parameter type for this DI (if any). The object type returned by this can only vary by class and not per object data. */
-	virtual FNiagaraDataInterfaceParametersCS* ConstructComputeParameters() const { return nullptr; }
+	virtual FNiagaraDataInterfaceParametersCS* CreateComputeParameters() const { return nullptr; }
 	virtual const FTypeLayoutDesc* GetComputeParametersTypeDesc() const { return nullptr; }
 
-	/** Methods that operate on an instance of type FNiagaraDataInterfaceParametersCS*, created by the above ConstructComputeParameters() method */
+	/** Methods that operate on an instance of type FNiagaraDataInterfaceParametersCS*, created by the above CreateComputeParameters() method */
 	virtual void BindParameters(FNiagaraDataInterfaceParametersCS* Base, const FNiagaraDataInterfaceGPUParamInfo& ParameterInfo, const class FShaderParameterMap& ParameterMap) {}
 	virtual void SetParameters(const FNiagaraDataInterfaceParametersCS* Base, FRHICommandList& RHICmdList, const FNiagaraDataInterfaceSetArgs& Context) const {}
 	virtual void UnsetParameters(const FNiagaraDataInterfaceParametersCS* Base, FRHICommandList& RHICmdList, const FNiagaraDataInterfaceSetArgs& Context) const {}
@@ -74,7 +74,7 @@ public:
 
 /** This goes in class declaration for UNiagaraDataInterfaceBase-derived types, that need custom parameter type */
 #define DECLARE_NIAGARA_DI_PARAMETER() \
-	virtual FNiagaraDataInterfaceParametersCS* ConstructComputeParameters() const override; \
+	virtual FNiagaraDataInterfaceParametersCS* CreateComputeParameters() const override; \
 	virtual const FTypeLayoutDesc* GetComputeParametersTypeDesc() const override; \
 	virtual void BindParameters(FNiagaraDataInterfaceParametersCS* Base, const FNiagaraDataInterfaceGPUParamInfo& ParameterInfo, const class FShaderParameterMap& ParameterMap) override; \
 	virtual void SetParameters(const FNiagaraDataInterfaceParametersCS* Base, FRHICommandList& RHICmdList, const FNiagaraDataInterfaceSetArgs& Context) const override; \
@@ -83,7 +83,7 @@ public:
 /** This goes in cpp file matched with class declartion using DECLARE_NIAGARA_DI_PARAMETER() */
 #define IMPLEMENT_NIAGARA_DI_PARAMETER(T, ParameterType) \
 	static_assert(ParameterType::InterfaceType == ETypeLayoutInterface::NonVirtual, "DI ParameterType must be non-virtual"); \
-	FNiagaraDataInterfaceParametersCS* T::ConstructComputeParameters() const { return new ParameterType(); } \
+	FNiagaraDataInterfaceParametersCS* T::CreateComputeParameters() const { return new ParameterType(); } \
 	const FTypeLayoutDesc* T::GetComputeParametersTypeDesc() const { return &StaticGetTypeLayoutDesc<ParameterType>(); } \
 	void T::BindParameters(FNiagaraDataInterfaceParametersCS* Base, const FNiagaraDataInterfaceGPUParamInfo& ParameterInfo, const class FShaderParameterMap& ParameterMap) { static_cast<ParameterType*>(Base)->Bind(ParameterInfo, ParameterMap); } \
 	void T::SetParameters(const FNiagaraDataInterfaceParametersCS* Base, FRHICommandList& RHICmdList, const FNiagaraDataInterfaceSetArgs& Context) const { static_cast<const ParameterType*>(Base)->Set(RHICmdList, Context); } \
diff --git a/Engine/Plugins/FX/Niagara/Source/NiagaraShader/Private/NiagaraShader.cpp b/Engine/Plugins/FX/Niagara/Source/NiagaraShader/Private/NiagaraShader.cpp
--- a/Engine/Plugins/FX/Niagara/Source/NiagaraShader/Private/NiagaraShader.cpp
+++ b/Engine/Plugins/FX/Niagara/Source/NiagaraShader/Private/NiagaraShader.cpp
@@ -1390,10 +1390,13 @@ void FNiagaraDataInterfaceParamRef::Bind(const FNiagaraDataInterfaceGPUParamInfo
 	if (Base)
 	{
 		DIType = TIndexedPtr<UNiagaraDataInterfaceBase>(Base); // TODO - clean up TIndexedPtr::operator=()
-		Parameters = Base->ConstructComputeParameters();
-		checkf(Parameters != nullptr, TEXT("Niagara type '%s' is missing ComputeParameters, missing a IMPLEMENT_NIAGARA_DI_PARAMETER?"), *InParameterInfo.DIClassName);
-		Parameters->DIType = DIType;
-		Base->BindParameters(Parameters, InParameterInfo, ParameterMap);
+		Parameters = Base->CreateComputeParameters();
+		if (Parameters)
+		{
+			checkf(Base->GetComputeParametersTypeDesc(), TEXT("Missing TypeDesc for compute parameters for class %s"), *InParameterInfo.DIClassName);
+			Parameters->DIType = DIType;
+			Base->BindParameters(Parameters, InParameterInfo, ParameterMap);
+		}
 	}
 }
 
