Since CreateRenderer can be called from a worker thread now due to Sceneproxy creation, using the concurrent GetMaterialRelevance call instead.

diff --git a/Engine/Plugins/FX/Niagara/Source/Niagara/Private/NiagaraRenderer.cpp b/Engine/Plugins/FX/Niagara/Source/Niagara/Private/NiagaraRenderer.cpp
--- a/Engine/Plugins/FX/Niagara/Source/Niagara/Private/NiagaraRenderer.cpp
+++ b/Engine/Plugins/FX/Niagara/Source/Niagara/Private/NiagaraRenderer.cpp
@@ -267,7 +267,7 @@ void FNiagaraRenderer::Initialize(const UNiagaraRendererProperties *InProps, con
 		{
 			Mat = UMaterial::GetDefaultMaterial(MD_Surface);
 		}
-		BaseMaterialRelevance_GT |= Mat->GetRelevance(FeatureLevel);
+		BaseMaterialRelevance_GT |= Mat->GetRelevance_Concurrent(FeatureLevel);
 	}
 }
 
