Added a function which retrieves the number of particles in the source emitter to the direct  read DI.
Fixed: the GetNumSpawnedParticles function did not work if the source emitter didn't have persistent IDs enabled.
Fixed: GetNumSpawnedParticles did not work on GPU.

diff --git a/Engine/Plugins/FX/Niagara/Source/Niagara/Classes/NiagaraDataInterfaceParticleRead.h b/Engine/Plugins/FX/Niagara/Source/Niagara/Classes/NiagaraDataInterfaceParticleRead.h
--- a/Engine/Plugins/FX/Niagara/Source/Niagara/Classes/NiagaraDataInterfaceParticleRead.h
+++ b/Engine/Plugins/FX/Niagara/Source/Niagara/Classes/NiagaraDataInterfaceParticleRead.h
@@ -37,6 +37,7 @@ public:
 
 	void GetNumSpawnedParticles(FVectorVMContext& Context);
 	void GetSpawnedIDAtIndex(FVectorVMContext& Context);
+	void GetNumParticles(FVectorVMContext& Context);
 	void ReadFloat(FVectorVMContext& Context, FName AttributeToRead);
 	void ReadVector2(FVectorVMContext& Context, FName AttributeToRead);
 	void ReadVector3(FVectorVMContext& Context, FName AttributeToRead);
diff --git a/Engine/Plugins/FX/Niagara/Source/Niagara/Private/NiagaraDataInterfaceParticleRead.cpp b/Engine/Plugins/FX/Niagara/Source/Niagara/Private/NiagaraDataInterfaceParticleRead.cpp
--- a/Engine/Plugins/FX/Niagara/Source/Niagara/Private/NiagaraDataInterfaceParticleRead.cpp
+++ b/Engine/Plugins/FX/Niagara/Source/Niagara/Private/NiagaraDataInterfaceParticleRead.cpp
@@ -10,6 +10,7 @@
 
 static const FName GetNumSpawnedParticlesFunctionName("Get Num Spawned Particles");
 static const FName GetSpawnedIDAtIndexFunctionName("Get Spawned ID At Index");
+static const FName GetNumParticlesFunctionName("Get Num Particles");
 static const FName GetIntAttributeFunctionName("Get int Attribute");
 static const FName GetFloatAttributeFunctionName("Get float Attribute");
 static const FName GetVec2AttributeFunctionName("Get Vector2 Attribute");
@@ -21,6 +22,7 @@ static const FName GetQuatAttributeFunctionName("Get Quaternion Attribute");
 
 static const FString NumSpawnedParticlesBaseName(TEXT("NumSpawnedParticles_"));
 static const FString SpawnedParticlesAcquireTagBaseName(TEXT("SpawnedParticlesAcquireTag_"));
+static const FString NumParticlesBaseName(TEXT("NumParticles_"));
 static const FString SpawnedIDsBufferBaseName(TEXT("SpawnedIDsBuffer_"));
 static const FString IDToIndexTableBaseName(TEXT("IDToIndexTable_"));
 static const FString InputFloatBufferBaseName(TEXT("InputFloatBuffer_"));
@@ -172,6 +174,7 @@ struct FNiagaraDataInterfaceParametersCS_ParticleRead : public FNiagaraDataInter
 	{
 		NumSpawnedParticlesParam.Bind(ParameterMap, *(NumSpawnedParticlesBaseName + ParameterInfo.DataInterfaceHLSLSymbol));
 		SpawnedParticlesAcquireTagParam.Bind(ParameterMap, *(SpawnedParticlesAcquireTagBaseName + ParameterInfo.DataInterfaceHLSLSymbol));
+		NumParticlesParam.Bind(ParameterMap, *(NumParticlesBaseName + ParameterInfo.DataInterfaceHLSLSymbol));
 		SpawnedIDsBufferParam.Bind(ParameterMap, *(SpawnedIDsBufferBaseName + ParameterInfo.DataInterfaceHLSLSymbol));
 		IDToIndexTableParam.Bind(ParameterMap, *(IDToIndexTableBaseName + ParameterInfo.DataInterfaceHLSLSymbol));
 		InputFloatBufferParam.Bind(ParameterMap, *(InputFloatBufferBaseName + ParameterInfo.DataInterfaceHLSLSymbol));
@@ -222,6 +225,7 @@ struct FNiagaraDataInterfaceParametersCS_ParticleRead : public FNiagaraDataInter
 			SetShaderValue(RHICmdList, ComputeShader, NumSpawnedParticlesParam, 0);
 			SetShaderValue(RHICmdList, ComputeShader, SpawnedParticlesAcquireTagParam, 0);
 			SetSRVParameter(RHICmdList, ComputeShader, SpawnedIDsBufferParam, FNiagaraRenderer::GetDummyIntBuffer().SRV);
+			SetShaderValue(RHICmdList, ComputeShader, NumParticlesParam, 0);
 		}
 
 		SetSRVParameter(RHICmdList, ComputeShader, IDToIndexTableParam, FNiagaraRenderer::GetDummyIntBuffer().SRV);
@@ -300,19 +304,22 @@ struct FNiagaraDataInterfaceParametersCS_ParticleRead : public FNiagaraDataInter
 		}
 
 		// Find the register index for the AcquireTag part of the particle ID in the source emitter.
-		InstanceData->AcquireTagRegisterIndex = -1;
-		for (int VarIdx = 0; VarIdx < SourceEmitterVariables.Num(); ++VarIdx)
+		if (AcquireTagRegisterIndexParam.IsBound())
 		{
-			const FNiagaraVariable& Var = SourceEmitterVariables[VarIdx];
-			if (Var.GetName().ToString() == TEXT("ID"))
+			InstanceData->AcquireTagRegisterIndex = -1;
+			for (int VarIdx = 0; VarIdx < SourceEmitterVariables.Num(); ++VarIdx)
 			{
-				InstanceData->AcquireTagRegisterIndex = SourceEmitterVariableLayouts[VarIdx].Int32ComponentStart + 1;
-				break;
+				const FNiagaraVariable& Var = SourceEmitterVariables[VarIdx];
+				if (Var.GetName().ToString() == TEXT("ID"))
+				{
+					InstanceData->AcquireTagRegisterIndex = SourceEmitterVariableLayouts[VarIdx].Int32ComponentStart + 1;
+					break;
+				}
+			}
+			if (InstanceData->AcquireTagRegisterIndex == -1)
+			{
+				UE_LOG(LogNiagara, Error, TEXT("Particle read DI cannot find ID variable in emitter '%s'."), *InstanceData->SourceEmitterName);
 			}
-		}
-		if (InstanceData->AcquireTagRegisterIndex == -1)
-		{
-			UE_LOG(LogNiagara, Error, TEXT("Particle read DI cannot find ID variable in emitter '%s'."), *InstanceData->SourceEmitterName);
 		}
 
 		// Initialize the buffer padding too, so we don't move garbage around.
@@ -322,6 +329,16 @@ struct FNiagaraDataInterfaceParametersCS_ParticleRead : public FNiagaraDataInter
 		}
 	}
 
+	FRHIShaderResourceView* GetIntSRVWithFallback(FRWBuffer& Buffer) const
+	{
+		return Buffer.SRV ? Buffer.SRV : FNiagaraRenderer::GetDummyIntBuffer().SRV;
+	}
+
+	FRHIShaderResourceView* GetFloatSRVWithFallback(FRWBuffer& Buffer) const
+	{
+		return Buffer.SRV ? Buffer.SRV : FNiagaraRenderer::GetDummyFloatBuffer().SRV;
+	}
+
 	void Set(FRHICommandList& RHICmdList, const FNiagaraDataInterfaceSetArgs& Context) const
 	{
 		check(IsInRenderingThread());
@@ -392,7 +409,8 @@ struct FNiagaraDataInterfaceParametersCS_ParticleRead : public FNiagaraDataInter
 
 		SetShaderValue(RHICmdList, ComputeShader, NumSpawnedParticlesParam, NumSpawnedInstances);
 		SetShaderValue(RHICmdList, ComputeShader, SpawnedParticlesAcquireTagParam, IDAcquireTag);
-		SetSRVParameter(RHICmdList, ComputeShader, SpawnedIDsBufferParam, SourceDataSet->GetGPUFreeIDs().SRV);
+		SetSRVParameter(RHICmdList, ComputeShader, SpawnedIDsBufferParam, GetIntSRVWithFallback(SourceDataSet->GetGPUFreeIDs()));
+		SetShaderValue(RHICmdList, ComputeShader, NumParticlesParam, SourceData ? SourceData->GetNumInstances() : 0);
 
 		if (!SourceData)
 		{
@@ -424,9 +442,9 @@ struct FNiagaraDataInterfaceParametersCS_ParticleRead : public FNiagaraDataInter
 		const uint32 ParticleStrideFloat = SourceData->GetFloatStride() / sizeof(float);
 		const uint32 ParticleStrideInt = SourceData->GetInt32Stride() / sizeof(int32);
 
-		SetSRVParameter(RHICmdList, ComputeShader, IDToIndexTableParam, SourceData->GetGPUIDToIndexTable().SRV);
-		SetSRVParameter(RHICmdList, ComputeShader, InputFloatBufferParam, SourceData->GetGPUBufferFloat().SRV);
-		SetSRVParameter(RHICmdList, ComputeShader, InputIntBufferParam, SourceData->GetGPUBufferInt().SRV);
+		SetSRVParameter(RHICmdList, ComputeShader, IDToIndexTableParam, GetIntSRVWithFallback(SourceData->GetGPUIDToIndexTable()));
+		SetSRVParameter(RHICmdList, ComputeShader, InputFloatBufferParam, GetFloatSRVWithFallback(SourceData->GetGPUBufferFloat()));
+		SetSRVParameter(RHICmdList, ComputeShader, InputIntBufferParam, GetIntSRVWithFallback(SourceData->GetGPUBufferInt()));
 		SetShaderValue(RHICmdList, ComputeShader, ParticleStrideFloatParam, ParticleStrideFloat);
 		SetShaderValue(RHICmdList, ComputeShader, ParticleStrideIntParam, ParticleStrideInt);
 		SetShaderValueArray(RHICmdList, ComputeShader, AttributeIndicesParam, InstanceData->AttributeIndices.GetData(), InstanceData->AttributeIndices.Num());
@@ -436,6 +454,7 @@ struct FNiagaraDataInterfaceParametersCS_ParticleRead : public FNiagaraDataInter
 private:
 	LAYOUT_FIELD(FShaderParameter, NumSpawnedParticlesParam);
 	LAYOUT_FIELD(FShaderParameter, SpawnedParticlesAcquireTagParam);
+	LAYOUT_FIELD(FShaderParameter, NumParticlesParam);
 	LAYOUT_FIELD(FShaderResourceParameter, SpawnedIDsBufferParam);
 	LAYOUT_FIELD(FShaderResourceParameter, IDToIndexTableParam);
 	LAYOUT_FIELD(FShaderResourceParameter, InputFloatBufferParam);
@@ -541,6 +560,15 @@ void UNiagaraDataInterfaceParticleRead::GetFunctions(TArray<FNiagaraFunctionSign
 		Sig.bRequiresContext = false;
 		OutFunctions.Add(Sig);
 	}
+	{
+		FNiagaraFunctionSignature Sig;
+		Sig.Name = GetNumParticlesFunctionName;
+		Sig.Inputs.Add(FNiagaraVariable(FNiagaraTypeDefinition(GetClass()), TEXT("Particle Reader")));
+		Sig.Outputs.Add(FNiagaraVariable(FNiagaraTypeDefinition::GetIntDef(), TEXT("Num Particles")));
+		Sig.bMemberFunction = true;
+		Sig.bRequiresContext = false;
+		OutFunctions.Add(Sig);
+	}
 	{
 		FNiagaraFunctionSignature Sig;
 		Sig.Name = GetFloatAttributeFunctionName;
@@ -665,6 +693,7 @@ void UNiagaraDataInterfaceParticleRead::GetFunctions(TArray<FNiagaraFunctionSign
 
 DEFINE_NDI_DIRECT_FUNC_BINDER(UNiagaraDataInterfaceParticleRead, GetNumSpawnedParticles);
 DEFINE_NDI_DIRECT_FUNC_BINDER(UNiagaraDataInterfaceParticleRead, GetSpawnedIDAtIndex);
+DEFINE_NDI_DIRECT_FUNC_BINDER(UNiagaraDataInterfaceParticleRead, GetNumParticles);
 DEFINE_NDI_DIRECT_FUNC_BINDER_WITH_PAYLOAD(UNiagaraDataInterfaceParticleRead, ReadFloat);
 DEFINE_NDI_DIRECT_FUNC_BINDER_WITH_PAYLOAD(UNiagaraDataInterfaceParticleRead, ReadVector2);
 DEFINE_NDI_DIRECT_FUNC_BINDER_WITH_PAYLOAD(UNiagaraDataInterfaceParticleRead, ReadVector3);
@@ -688,6 +717,12 @@ void UNiagaraDataInterfaceParticleRead::GetVMExternalFunction(const FVMExternalF
 		return;
 	}
 
+	if (BindingInfo.Name == GetNumParticlesFunctionName)
+	{
+		NDI_FUNC_BINDER(UNiagaraDataInterfaceParticleRead, GetNumParticles)::Bind(this, OutFunc);
+		return;
+	}
+
 	bool bBindSuccessful = false;
 	FNDIParticleRead_InstanceData* PIData = static_cast<FNDIParticleRead_InstanceData*>(InstanceData);
 	static const FName NAME_Attribute("Attribute");
@@ -784,9 +819,8 @@ void UNiagaraDataInterfaceParticleRead::GetNumSpawnedParticles(FVectorVMContext&
 	VectorVM::FUserPtrHandler<FNDIParticleRead_InstanceData> InstData(Context);
 	VectorVM::FExternalFuncRegisterHandler<int32> OutNumSpawned(Context);
 
-	FNiagaraEmitterInstance* EmitterInstance = InstData.Get()->EmitterInstance;
-	const TArray<int32>& SpawnedIDsTable = EmitterInstance->GetData().GetSpawnedIDsTable();
-	int32 NumSpawned = SpawnedIDsTable.Num();
+	const FNiagaraEmitterInstance* EmitterInstance = InstData.Get()->EmitterInstance;
+	const int32 NumSpawned = EmitterInstance->GetData().GetCurrentData()->GetNumSpawnedInstances();
 
 	for (int32 InstanceIdx = 0; InstanceIdx < Context.NumInstances; ++InstanceIdx)
 	{
@@ -834,6 +868,20 @@ void UNiagaraDataInterfaceParticleRead::GetSpawnedIDAtIndex(FVectorVMContext& Co
 	}
 }
 
+void UNiagaraDataInterfaceParticleRead::GetNumParticles(FVectorVMContext& Context)
+{
+	VectorVM::FUserPtrHandler<FNDIParticleRead_InstanceData> InstData(Context);
+	VectorVM::FExternalFuncRegisterHandler<int32> OutNumParticles (Context);
+
+	const FNiagaraEmitterInstance* EmitterInstance = InstData.Get()->EmitterInstance;
+	const int32 NumParticles = EmitterInstance->GetData().GetCurrentData()->GetNumInstances();
+
+	for (int32 InstanceIdx = 0; InstanceIdx < Context.NumInstances; ++InstanceIdx)
+	{
+		*OutNumParticles.GetDestAndAdvance() = NumParticles;
+	}
+}
+
 void UNiagaraDataInterfaceParticleRead::ReadFloat(FVectorVMContext& Context, FName AttributeToRead)
 {
 	VectorVM::FExternalFuncInputHandler<int32> ParticleIDIndexParam(Context);
@@ -1073,10 +1121,10 @@ T UNiagaraDataInterfaceParticleRead::RetrieveValueWithCheck(FNiagaraEmitterInsta
 		{
 			FNiagaraID ActualID = IDData.GetSafe(ParticleIndex, NIAGARA_INVALID_ID);
 			if (ActualID == ParticleID)
-		{
+			{
 				Value = ValueData.GetSafe(ParticleIndex, T());
 				bValid = true;
-		}
+			}
 		}
 
 		return Value;
@@ -1107,6 +1155,7 @@ void UNiagaraDataInterfaceParticleRead::GetParameterDefinitionHLSL(const FNiagar
 	static const TCHAR *FormatDeclarations = TEXT(
 		"int {NumSpawnedParticlesName};\n"
 		"int {SpawnedParticlesAcquireTagName};\n"
+		"int {NumParticlesName};\n"
 		"uint {ParticleStrideFloatName};\n"
 		"uint {ParticleStrideIntName};\n"
 		"int {AcquireTagRegisterIndexName};\n"
@@ -1125,6 +1174,7 @@ void UNiagaraDataInterfaceParticleRead::GetParameterDefinitionHLSL(const FNiagar
 	TMap<FString, FStringFormatArg> ArgsDeclarations;
 	ArgsDeclarations.Add(TEXT("NumSpawnedParticlesName"), NumSpawnedParticlesBaseName + ParamInfo.DataInterfaceHLSLSymbol);
 	ArgsDeclarations.Add(TEXT("SpawnedParticlesAcquireTagName"), SpawnedParticlesAcquireTagBaseName + ParamInfo.DataInterfaceHLSLSymbol);
+	ArgsDeclarations.Add(TEXT("NumParticlesName"), NumParticlesBaseName + ParamInfo.DataInterfaceHLSLSymbol);
 	ArgsDeclarations.Add(TEXT("ParticleStrideFloatName"), ParticleStrideFloatBaseName + ParamInfo.DataInterfaceHLSLSymbol);
 	ArgsDeclarations.Add(TEXT("ParticleStrideIntName"), ParticleStrideIntBaseName + ParamInfo.DataInterfaceHLSLSymbol);
 	ArgsDeclarations.Add(TEXT("AcquireTagRegisterIndexName"), AcquireTagRegisterIndexBaseName + ParamInfo.DataInterfaceHLSLSymbol);
@@ -1232,7 +1282,7 @@ bool UNiagaraDataInterfaceParticleRead::GetFunctionHLSL(const FNiagaraDataInterf
 		static const TCHAR* FuncTemplate = TEXT(
 			"void {FunctionName}(out int Out_NumSpawned)\n"
 			"{\n"
-			"    return {NumSpawnedParticlesName};\n"
+			"    Out_NumSpawned = {NumSpawnedParticlesName};\n"
 			"}\n\n"
 		);
 
@@ -1274,6 +1324,23 @@ bool UNiagaraDataInterfaceParticleRead::GetFunctionHLSL(const FNiagaraDataInterf
 		return true;
 	}
 
+	if (FunctionInfo.DefinitionName == GetNumParticlesFunctionName)
+	{
+		static const TCHAR* FuncTemplate = TEXT(
+			"void {FunctionName}(out int Out_NumParticles)\n"
+			"{\n"
+			"    Out_NumParticles = {NumParticlesName};\n"
+			"}\n\n"
+		);
+
+		TMap<FString, FStringFormatArg> FuncTemplateArgs;
+		FuncTemplateArgs.Add(TEXT("FunctionName"), FunctionInfo.InstanceName);
+		FuncTemplateArgs.Add(TEXT("NumParticlesName"), NumParticlesBaseName + ParamInfo.DataInterfaceHLSLSymbol);
+
+		OutHLSL += FString::Format(FuncTemplate, FuncTemplateArgs);
+		return true;
+	}
+
 	if (FunctionInfo.DefinitionName == GetIntAttributeFunctionName)
 	{
 		return GenerateGetFunctionHLSL(ParamInfo, FunctionInfo, FunctionInstanceIndex, ENiagaraParticleDataComponentType::Int, 1, OutHLSL);
diff --git a/Engine/Plugins/FX/Niagara/Source/Niagara/Private/NiagaraEmitterInstance.cpp b/Engine/Plugins/FX/Niagara/Source/Niagara/Private/NiagaraEmitterInstance.cpp
--- a/Engine/Plugins/FX/Niagara/Source/Niagara/Private/NiagaraEmitterInstance.cpp
+++ b/Engine/Plugins/FX/Niagara/Source/Niagara/Private/NiagaraEmitterInstance.cpp
@@ -1441,9 +1441,10 @@ void FNiagaraEmitterInstance::Tick(float DeltaSeconds)
 		}
 	}
 
-	// The spawn count and ID acquire tag in FNiagaraDataBuffer are currently only used by the GPU path,
-	// but we'll set them here anyway, to prevent surprises in case someone sees them and decides to use them.
-	Data.GetDestinationDataChecked().SetNumSpawnedInstances(Data.GetSpawnedIDsTable().Num());
+	int32 NumAfterSpawn = Data.GetCurrentDataChecked().GetNumInstances();
+	int32 TotalNumSpawned = NumAfterSpawn - NumBeforeSpawn;
+
+	Data.GetDestinationDataChecked().SetNumSpawnedInstances(TotalNumSpawned);
 	Data.GetDestinationDataChecked().SetIDAcquireTag(Data.GetIDAcquireTag());
 
 	//We're done with this simulation pass.
@@ -1466,8 +1467,6 @@ void FNiagaraEmitterInstance::Tick(float DeltaSeconds)
 
 	//Now pull out any debug info we need.
 #if WITH_EDITORONLY_DATA
-	int32 NumAfterSpawn = Data.GetCurrentDataChecked().GetNumInstances();
-	int32 TotalNumSpawned = NumAfterSpawn - NumBeforeSpawn;
 	if (ParentSystemInstance->ShouldCaptureThisFrame())
 	{
 		//Pull out update data.
diff --git a/Engine/Plugins/FX/Niagara/Source/NiagaraCore/Private/NiagaraCustomVersion.cpp b/Engine/Plugins/FX/Niagara/Source/NiagaraCore/Private/NiagaraCustomVersion.cpp
--- a/Engine/Plugins/FX/Niagara/Source/NiagaraCore/Private/NiagaraCustomVersion.cpp
+++ b/Engine/Plugins/FX/Niagara/Source/NiagaraCore/Private/NiagaraCustomVersion.cpp
@@ -8,4 +8,4 @@ const FGuid FNiagaraCustomVersion::GUID(0xFCF57AFA, 0x50764283, 0xB9A9E658, 0xFF
 // Register the custom version with core
 FCustomVersionRegistration GRegisterNiagaraCustomVersion(FNiagaraCustomVersion::GUID, FNiagaraCustomVersion::LatestVersion, TEXT("NiagaraVer"));
 
-const FGuid FNiagaraCustomVersion::LatestScriptCompileVersion(0x36511D5D, 0x860A47B4, 0x975E699B, 0x4599D5D4);
+const FGuid FNiagaraCustomVersion::LatestScriptCompileVersion(0x19CA6492, 0x8CCB46EC, 0xAC439F3F, 0x0BAAE91B);
