Sim Stages Experimental support added

diff --git a/Engine/Plugins/Experimental/HairStrands/Source/HairStrandsNiagara/Private/HairStrandsNiagara.cpp b/Engine/Plugins/Experimental/HairStrands/Source/HairStrandsNiagara/Private/HairStrandsNiagara.cpp
--- a/Engine/Plugins/Experimental/HairStrands/Source/HairStrandsNiagara/Private/HairStrandsNiagara.cpp
+++ b/Engine/Plugins/Experimental/HairStrands/Source/HairStrandsNiagara/Private/HairStrandsNiagara.cpp
@@ -8,9 +8,6 @@ IMPLEMENT_MODULE(FHairStrandsNiagara, HairStrandsNiagara);
 
 void FHairStrandsNiagara::StartupModule()
 {
-	static const auto UseShaderStagesCVar = IConsoleManager::Get().FindConsoleVariable(TEXT("fx.UseShaderStages"));
-	UseShaderStagesCVar->Set(1);
-
 	// Maps virtual shader source directory /Plugin/FX/Niagara to the plugin's actual Shaders directory.
 	FString PluginShaderDir = FPaths::Combine(IPluginManager::Get().FindPlugin(TEXT("HairStrands"))->GetBaseDir(), TEXT("Shaders"));
 	AddShaderSourceDirectoryMapping(TEXT("/Plugin/Experimental/HairStrands"), PluginShaderDir);
diff --git a/Engine/Plugins/Experimental/HairStrands/Source/HairStrandsNiagara/Private/NiagaraDataInterfaceHairStrands.cpp b/Engine/Plugins/Experimental/HairStrands/Source/HairStrandsNiagara/Private/NiagaraDataInterfaceHairStrands.cpp
--- a/Engine/Plugins/Experimental/HairStrands/Source/HairStrandsNiagara/Private/NiagaraDataInterfaceHairStrands.cpp
+++ b/Engine/Plugins/Experimental/HairStrands/Source/HairStrandsNiagara/Private/NiagaraDataInterfaceHairStrands.cpp
@@ -3804,7 +3804,7 @@ void FNDIHairStrandsProxy::PreStage(FRHICommandList& RHICmdList, const FNiagaraD
 
 void FNDIHairStrandsProxy::PostStage(FRHICommandList& RHICmdList, const FNiagaraDataInterfaceSetArgs& Context)
 {
-	if (Context.ShaderStageIndex == 0)
+	if (Context.SimulationStageIndex == 0)
 	{
 		FNDIHairStrandsData* ProxyData =
 			SystemInstancesToProxyData.Find(Context.SystemInstance);
diff --git a/Engine/Plugins/FX/Niagara/Shaders/Private/NiagaraEmitterInstanceShader.usf b/Engine/Plugins/FX/Niagara/Shaders/Private/NiagaraEmitterInstanceShader.usf
--- a/Engine/Plugins/FX/Niagara/Shaders/Private/NiagaraEmitterInstanceShader.usf
+++ b/Engine/Plugins/FX/Niagara/Shaders/Private/NiagaraEmitterInstanceShader.usf
@@ -6,6 +6,7 @@ NiagaraSimulationShader.usf:
 #pragma warning(disable:4008)
 #include "/Engine/Public/Platform.ush"
 
+
 #if GPU_SIMULATION
 	#include "/Engine/Private/Common.ush"
 	#include "/Engine/Private/GlobalDistanceFieldShared.ush"
@@ -57,12 +58,15 @@ bool NiagaraAll(bool4 b) { return b.x && b.y && b.z && b.w; }
 	static uint GEmitterTickCounter;
 	static uint GSimStart;
 	static uint GRandomSeedOffset = 0;
+	
+	const static uint GSpawnPhase = 0;
+	const static uint GUpdatePhase = 1;
 
 	// To be removed as well but don't know who is using it
 	#include "/Engine/Private/SceneTexturesCommon.ush"
 
 	// Temporary include while everything is moved to plugins
-	#if USE_SHADER_STAGES == 1
+	#if USE_SIMULATION_STAGES == 1
 		#include "NiagaraStrandsPhysics.ush"
 	#endif
 
@@ -432,8 +436,8 @@ int rand_int(int x, int Seed1, int Seed2, int Seed3)
 	 * declared as int instead of uint to avoid compiler warnings because we can't expose uint on Niagara module parameters.  
 	 * This is still declared as a uint in the corresponding C++ classes
 	 */
-	int DefaultShaderStageIndex;
-	int ShaderStageIndex;
+	int DefaultSimulationStageIndex;
+	int SimulationStageIndex;
 
 	int IterationInterfaceCount;
 
@@ -510,12 +514,12 @@ int rand_int(int x, int Seed1, int Seed2, int Seed3)
 		MustBe0[DataSetID] = 0;
 		// End static assert
 
-		#if USE_SHADER_STAGES
-			if(ShaderStageIndex != 0)
+		#if USE_SIMULATION_STAGES
+			if(SimulationStageIndex != 0)
 			{
 				bDoAcquire = false; // Prevent any change to the atomic
 			}
-		#endif // USE_SHADER_STAGES
+		#endif // USE_SIMULATION_STAGES
 
 		int PrevIdx = GSpawnStartInstance + SpawnedInstances;	// scratch instance as default; write to that for dead particles
 
@@ -622,12 +626,12 @@ int rand_int(int x, int Seed1, int Seed2, int Seed3)
 
 		#endif // USE_WAVE_INTRISICS || USE_GROUP_SHARED
 
-		#if USE_SHADER_STAGES
-			if(ShaderStageIndex != 0)
+		#if USE_SIMULATION_STAGES
+			if(SimulationStageIndex != 0)
 			{
 				return ExecIndex();
 			}
-		#endif // USE_SHADER_STAGES
+		#endif // USE_SIMULATION_STAGES
 
 		return PrevIdx;
 	}
@@ -680,108 +684,15 @@ int rand_int(int x, int Seed1, int Seed2, int Seed3)
 	 */
 	int OutputIndex(const int DataSetID, const bool bIsPersistent, const bool bIsValid)
 	{
-		#if USE_SHADER_STAGES
+		#if USE_SIMULATION_STAGES
 			const int DenseIndex = AcquireIndex(DataSetID,bIsValid);
 			const int SparseIndex = GDispatchThreadId.x;
 			return (bIsPersistent && (DenseIndex != SparseIndex) ) ? SparseIndex : DenseIndex;
 		#else
 			return AcquireIndex(DataSetID,bIsValid);
-		#endif // USE_SHADER_STAGES
+		#endif // USE_SIMULATION_STAGES
 	}
 
 //Include the simulation shader code generated by the node graph.
 #include "/Engine/Generated/NiagaraEmitterInstance.ush"
 
-
-#if GPU_SIMULATION
-
-/* 
- *	CS wrapper for our generated code; calls spawn and update functions on the corresponding instances in the buffer
- */
-
-[numthreads(THREADGROUP_SIZE, 1, 1)]
-void SimulateMainComputeCS(
-	uint3 GroupId : SV_GroupID,
-	uint3 DispatchThreadId : SV_DispatchThreadID,
-	uint3 GroupThreadId : SV_GroupThreadID)
-{
-	GDispatchThreadId = DispatchThreadId;
-	GGroupThreadId = GroupThreadId;
-	GCurrentPhase = -1;
-	GEmitterTickCounter = EmitterTickCounter;
-	GSimStart = SimStart;
-	GRandomSeedOffset = 0;
-
-	/*
-	if(CopyInstancesBeforeStart == 1)
-	{
-		UpdateStartInstance = 0;
-	}
-	*/
-
-	const uint InstanceID = UpdateStartInstance + DispatchThreadId.x;
-	if (ReadInstanceCountOffset == 0xFFFFFFFF)
-	{
-		GSpawnStartInstance = 0;
-	}
-	else
-	{
-		GSpawnStartInstance = RWInstanceCounts[ReadInstanceCountOffset];				// needed by ExecIndex()
-	}
-	const int MaxInstances = GSpawnStartInstance + SpawnedInstances;
-
-	bool bRunPhase1 = InstanceID < GSpawnStartInstance && InstanceID < UpdateStartInstance + MaxInstances;
-	bool bRunPhase0 = InstanceID >= GSpawnStartInstance && InstanceID < UpdateStartInstance + MaxInstances;
-	
-	#if USE_SHADER_STAGES
-	if (IterationInterfaceCount > 0)
-	{
-		bRunPhase1 = InstanceID < IterationInterfaceCount && GSimStart != 1;
-		bRunPhase0 = InstanceID < IterationInterfaceCount && GSimStart == 1;
-	}
-	#endif // USE_SHADER_STAGES
-
-	const float RandomSeedInitialisation = NiagaraInternalNoise(InstanceID*16384, 0 * 8196, (bRunPhase1 ? 4096 : 0) + EmitterTickCounter);	// initialise the random state seed
-
-	FSimulationContext Context = (FSimulationContext)0;
-
-	BRANCH
-	if (bRunPhase1)
-	{
-		GCurrentPhase = 1;
-
-		SetupExecIndexForGPU();
-
-		InitConstants(Context);
-		LoadUpdateVariables(Context, InstanceID);
-		ReadDataSets(Context);
-
-		SimulateMapUpdate(Context);
-		WriteDataSets(Context);
-	}
-	else if (bRunPhase0)
-	{
-		GCurrentPhase = 0;
-
-		SetupExecIndexAndSpawnInfoForGPU();
-
-		InitConstants(Context);
-		InitSpawnVariables(Context);
-		ReadDataSets(Context);
-
-		Context.MapSpawn.Particles.UniqueID = Engine_Emitter_TotalSpawnedParticles + ExecIndex();
-		ConditionalInterpolateParameters(Context);
-		SimulateMapSpawn(Context);
-
-		GCurrentPhase = 1;
-
-		TransferAttributes(Context);
-
-		SimulateMapUpdate(Context);
-		WriteDataSets(Context);
-	}
-
-	StoreUpdateVariables(Context);
-}
-
-#endif//GPU_SIMULATION
diff --git a/Engine/Plugins/FX/Niagara/Shaders/Private/NiagaraMeshVertexFactory.ush b/Engine/Plugins/FX/Niagara/Shaders/Private/NiagaraMeshVertexFactory.ush
--- a/Engine/Plugins/FX/Niagara/Shaders/Private/NiagaraMeshVertexFactory.ush
+++ b/Engine/Plugins/FX/Niagara/Shaders/Private/NiagaraMeshVertexFactory.ush
@@ -723,7 +723,7 @@ FVertexFactoryIntermediates GetVertexFactoryIntermediates(FVertexFactoryInput In
 	if (NiagaraMeshVF.bLocalSpace)
 	{
 		Intermediates.PrevLocalToWorld = mul(Intermediates.PrevLocalToWorld, Primitive.PreviousLocalToWorld);
-		Intermediates.LocalToWorld = mul(Intermediates.LocalToWorld, Primitive.LocalToWorld);
+		//Intermediates.LocalToWorld = mul(Intermediates.LocalToWorld, Primitive.LocalToWorld);
 		Intermediates.ParticlePosition = Intermediates.LocalToWorld[3];
 	}
 
diff --git a/Engine/Plugins/FX/Niagara/Source/Niagara/Classes/NiagaraDataInterface.h b/Engine/Plugins/FX/Niagara/Source/Niagara/Classes/NiagaraDataInterface.h
--- a/Engine/Plugins/FX/Niagara/Source/Niagara/Classes/NiagaraDataInterface.h
+++ b/Engine/Plugins/FX/Niagara/Source/Niagara/Classes/NiagaraDataInterface.h
@@ -167,19 +167,20 @@ struct FNiagaraDataInterfaceProxy : TSharedFromThis<FNiagaraDataInterfaceProxy,
 	virtual void ConsumePerInstanceDataFromGameThread(void* PerInstanceData, const FNiagaraSystemInstanceID& Instance) { check(false); }
 
 	// #todo(dmp): move all of this stuff to the RW interface to keep it out of here?
-
+	FName SourceDIName;
+	
 	// a set of the shader stages that require the data interface for data output
-	TSet<int> OutputShaderStages;
+	TSet<int> OutputSimulationStages_DEPRECATED;
 
 	// a set of the shader stages that require the data interface for setting number of output elements
-	TSet<int> IterationShaderStages;
+	TSet<int> IterationSimulationStages_DEPRECATED;
 	
 	// number of elements to output to
 	uint32 ElementCount;
 
 	void SetElementCount(uint32 Count) { ElementCount = Count;  }
-	virtual bool IsOutputStage(uint32 CurrentStage) const { return OutputShaderStages.Contains(CurrentStage); }
-	virtual bool IsIterationStage(uint32 CurrentStage) const { return IterationShaderStages.Contains(CurrentStage); }
+	virtual bool IsOutputStage_DEPRECATED(uint32 CurrentStage) const { return OutputSimulationStages_DEPRECATED.Contains(CurrentStage); }
+	virtual bool IsIterationStage_DEPRECATED(uint32 CurrentStage) const { return IterationSimulationStages_DEPRECATED.Contains(CurrentStage); }
 
 	virtual void ResetData(FRHICommandList& RHICmdList, const FNiagaraDataInterfaceSetArgs& Context) { }
 
diff --git a/Engine/Plugins/FX/Niagara/Source/Niagara/Classes/NiagaraDataInterfaceGrid2DCollection.h b/Engine/Plugins/FX/Niagara/Source/Niagara/Classes/NiagaraDataInterfaceGrid2DCollection.h
--- a/Engine/Plugins/FX/Niagara/Source/Niagara/Classes/NiagaraDataInterfaceGrid2DCollection.h
+++ b/Engine/Plugins/FX/Niagara/Source/Niagara/Classes/NiagaraDataInterfaceGrid2DCollection.h
@@ -96,7 +96,7 @@ struct FNiagaraDataInterfaceProxyGrid2DCollection : public FNiagaraDataInterface
 	TMap<FNiagaraSystemInstanceID, Grid2DCollectionRWInstanceData> SystemInstancesToProxyData;
 };
 
-UCLASS(EditInlineNew, Category = "Grid", meta = (DisplayName = "Grid2D Collection"), Blueprintable, BlueprintType)
+UCLASS(EditInlineNew, Category = "Grid", meta = (DisplayName = "Grid2D Collection", Experimental), Blueprintable, BlueprintType)
 class NIAGARA_API UNiagaraDataInterfaceGrid2DCollection : public UNiagaraDataInterfaceGrid2D
 {
 	GENERATED_UCLASS_BODY()
diff --git a/Engine/Plugins/FX/Niagara/Source/Niagara/Classes/NiagaraDataInterfaceRW.h b/Engine/Plugins/FX/Niagara/Source/Niagara/Classes/NiagaraDataInterfaceRW.h
--- a/Engine/Plugins/FX/Niagara/Source/Niagara/Classes/NiagaraDataInterfaceRW.h
+++ b/Engine/Plugins/FX/Niagara/Source/Niagara/Classes/NiagaraDataInterfaceRW.h
@@ -118,7 +118,6 @@ protected:
 	virtual bool CopyToInternal(UNiagaraDataInterface* Destination) const override;
 	virtual void PushToRenderThread() {};
 
-	FNiagaraDataInterfaceProxyRW* RWProxy;
 };
 
 
diff --git a/Engine/Plugins/FX/Niagara/Source/Niagara/Classes/NiagaraDataSet.h b/Engine/Plugins/FX/Niagara/Source/Niagara/Classes/NiagaraDataSet.h
--- a/Engine/Plugins/FX/Niagara/Source/Niagara/Classes/NiagaraDataSet.h
+++ b/Engine/Plugins/FX/Niagara/Source/Niagara/Classes/NiagaraDataSet.h
@@ -291,7 +291,7 @@ public:
 	void ResetBuffers();
 
 	/** Begins a new simulation pass and grabs a destination buffer. Returns the new destination data buffer. */
-	FNiagaraDataBuffer& BeginSimulate();
+	FNiagaraDataBuffer& BeginSimulate(bool bResetDestinationData = true);
 
 	/** Ends a simulation pass and sets the current simulation state. */
 	void EndSimulate(bool SetCurrentData = true);
diff --git a/Engine/Plugins/FX/Niagara/Source/Niagara/Classes/NiagaraEmitter.h b/Engine/Plugins/FX/Niagara/Source/Niagara/Classes/NiagaraEmitter.h
--- a/Engine/Plugins/FX/Niagara/Source/Niagara/Classes/NiagaraEmitter.h
+++ b/Engine/Plugins/FX/Niagara/Source/Niagara/Classes/NiagaraEmitter.h
@@ -17,7 +17,7 @@ class UMaterial;
 class UNiagaraEmitter;
 class UNiagaraEventReceiverEmitterAction;
 class UNiagaraRendererProperties;
-class UNiagaraShaderStageBase;
+class UNiagaraSimulationStageBase;
 class UNiagaraEditorDataBase;
 
 //TODO: Event action that spawns other whole Systems?
@@ -337,18 +337,26 @@ public:
 	UPROPERTY(EditAnywhere, AdvancedDisplay, Category = "Emitter", meta = (EditCondition = "bLimitDeltaTime"))
 	float MaxDeltaTimePerTick;
 
-	/** Get the max number of iteration that the CS is going to be launched. */
-	UPROPERTY(EditAnywhere, AdvancedDisplay, Category = "Emitter")
+	/** Get the default shader stage index. */
+	UPROPERTY(EditAnywhere, AdvancedDisplay, Category = "Simulation Stages", meta = (EditCondition = "bSimulationStagesEnabled || bDeprecatedShaderStagesEnabled", DisplayAfter = "bSimulationStagesEnabled"))
 	uint32 DefaultShaderStageIndex;
 
-	/** Get the max number of iteration that the CS is going to be launched. */
-	UPROPERTY(EditAnywhere, AdvancedDisplay, Category = "Emitter")
+	/** Get the number of shader stages that we fire off. */
+	UPROPERTY(EditAnywhere, AdvancedDisplay, Category = "Simulation Stages", meta = (EditCondition = "bSimulationStagesEnabled || bDeprecatedShaderStagesEnabled", DisplayAfter = "DefaultShaderStageIndex"))
 	uint32 MaxUpdateIterations;
 
-	/** Get the max number of iteration that the CS is going to be launched. */
-	UPROPERTY(EditAnywhere, AdvancedDisplay, Category = "Emitter")
+	/** Get whether or not shaderstages spwn. */
+	UPROPERTY(EditAnywhere, AdvancedDisplay, Category = "Simulation Stages", meta = (EditCondition = "bSimulationStagesEnabled || bDeprecatedShaderStagesEnabled", DisplayAfter = "MaxUpdateIterations"))
 	TSet<uint32> SpawnStages;
 
+	/** Get whether or not to use simulation stages. */
+	UPROPERTY(EditAnywhere, AdvancedDisplay, Category = "Simulation Stages", meta = (DisplayName = "Enable Simulation Stages (Experimental GPU Only)"))
+	uint32 bSimulationStagesEnabled : 1;
+
+	/** Get whether or not to use shader stages. */
+	UPROPERTY(EditAnywhere, AdvancedDisplay, Category = "Simulation Stages", meta = (DisplayName = "Enable Deprecated Shader Stages (Experimental GPU Only)"))
+	uint32 bDeprecatedShaderStagesEnabled : 1;
+
 	/** Whether to limit the max tick delta time or not. */
 	UPROPERTY(EditAnywhere, AdvancedDisplay, Category = "Emitter", meta = (InlineEditConditionToggle))
 	uint32 bLimitDeltaTime : 1;
@@ -466,15 +474,15 @@ public:
 
 	void NIAGARA_API RemoveEventHandlerByUsageId(FGuid EventHandlerUsageId);
 
-	NIAGARA_API const TArray<UNiagaraShaderStageBase*>& GetShaderStages() const { return ShaderStages; }
+	NIAGARA_API const TArray<UNiagaraSimulationStageBase*>& GetSimulationStages() const { return SimulationStages; }
 
-	NIAGARA_API UNiagaraShaderStageBase* GetShaderStageById(FGuid ScriptUsageId) const;
+	NIAGARA_API UNiagaraSimulationStageBase* GetSimulationStageById(FGuid ScriptUsageId) const;
 
-	void NIAGARA_API AddShaderStage(UNiagaraShaderStageBase* ShaderStage);
+	void NIAGARA_API AddSimulationStage(UNiagaraSimulationStageBase* SimulationStage);
 
-	void NIAGARA_API RemoveShaderStage(UNiagaraShaderStageBase* ShaderStage);
+	void NIAGARA_API RemoveSimulationStage(UNiagaraSimulationStageBase* SimulationStage);
 
-	void NIAGARA_API MoveShaderStageToIndex(UNiagaraShaderStageBase* ShaderStage, int32 TargetIndex);
+	void NIAGARA_API MoveSimulationStageToIndex(UNiagaraSimulationStageBase* SimulationStage, int32 TargetIndex);
 
 	/* Gets whether or not the supplied event generator id matches an event generator which is shared between the particle spawn and update scrips. */
 	bool IsEventGeneratorShared(FName EventGeneratorId) const;
@@ -537,7 +545,7 @@ private:
 
 	void ScriptRapidIterationParameterChanged();
 
-	void ShaderStageChanged();
+	void SimulationStageChanged();
 
 	void RendererChanged();
 
@@ -573,7 +581,7 @@ private:
 	TArray<FNiagaraEventScriptProperties> EventHandlerScriptProps;
 
 	UPROPERTY(meta = (NiagaraNoMerge))
-	TArray<UNiagaraShaderStageBase*> ShaderStages;
+	TArray<UNiagaraSimulationStageBase*> SimulationStages;
 
 	UPROPERTY()
 	UNiagaraScript* GPUComputeScript;
diff --git a/Engine/Plugins/FX/Niagara/Source/Niagara/Classes/NiagaraEmitterInstanceBatcher.h b/Engine/Plugins/FX/Niagara/Source/Niagara/Classes/NiagaraEmitterInstanceBatcher.h
--- a/Engine/Plugins/FX/Niagara/Source/Niagara/Classes/NiagaraEmitterInstanceBatcher.h
+++ b/Engine/Plugins/FX/Niagara/Source/Niagara/Classes/NiagaraEmitterInstanceBatcher.h
@@ -108,7 +108,7 @@ public:
 
 	void ProcessDebugInfo(FRHICommandList &RHICmdLis, const FNiagaraComputeExecutionContext* Context) const;
 
-	void SetDataInterfaceParameters(const TArray<FNiagaraDataInterfaceProxy*>& DataInterfaceProxies, const FNiagaraShaderRef& Shader, FRHICommandList &RHICmdList, const FNiagaraComputeInstanceData* Instance, const FNiagaraGPUSystemTick& Tick, uint32 ShaderStageIndex) const;
+	void SetDataInterfaceParameters(const TArray<FNiagaraDataInterfaceProxy*>& DataInterfaceProxies, const FNiagaraShaderRef& Shader, FRHICommandList &RHICmdList, const FNiagaraComputeInstanceData* Instance, const FNiagaraGPUSystemTick& Tick, uint32 SimulationStageIndex) const;
 	void UnsetDataInterfaceParameters(const TArray<FNiagaraDataInterfaceProxy*>& DataInterfaceProxies, const FNiagaraShaderRef& Shader, FRHICommandList& RHICmdList, const FNiagaraComputeInstanceData* Instance, const FNiagaraGPUSystemTick& Tick) const;
 
 	void Run(	const FNiagaraGPUSystemTick& Tick,
@@ -120,8 +120,8 @@ public:
 				FRHIUniformBuffer* ViewUniformBuffer, 
 				const FNiagaraGpuSpawnInfo& SpawnInfo,
 				bool bCopyBeforeStart = false,
-				uint32 DefaultShaderStageIndex = 0,
-				uint32 ShaderStageIndex = 0,
+				uint32 DefaultSimulationStageIndex = 0,
+				uint32 SimulationStageIndex = 0,
 				FNiagaraDataInterfaceProxy *IterationInterface = nullptr,
 				bool HasRunParticleStage = false
 			);
@@ -137,13 +137,13 @@ public:
 	bool ResetDataInterfaces(const FNiagaraGPUSystemTick& Tick, FNiagaraComputeInstanceData *Instance, FRHICommandList &RHICmdList, const FNiagaraShaderRef& ComputeShader ) const;
 
 	/** Given a shader stage index, find the corresponding data interface */
-	FNiagaraDataInterfaceProxy* FindIterationInterface(FNiagaraComputeInstanceData *Instance, const uint32 ShaderStageIndex) const;
+	FNiagaraDataInterfaceProxy* FindIterationInterface(FNiagaraComputeInstanceData *Instance, const uint32 SimulationStageIndex) const;
 
 	/** Loop over all the data interfaces and call the prestage methods */
-	void PreStageInterface(const FNiagaraGPUSystemTick& Tick, FNiagaraComputeInstanceData *Instance, FRHICommandList &RHICmdList, const FNiagaraShaderRef& ComputeShader, const uint32 ShaderStageIndex) const;
+	void PreStageInterface(const FNiagaraGPUSystemTick& Tick, FNiagaraComputeInstanceData *Instance, FRHICommandList &RHICmdList, const FNiagaraShaderRef& ComputeShader, const uint32 SimulationStageIndex) const;
 
 	/** Loop over all the data interfaces and call the poststage methods */
-	void PostStageInterface(const FNiagaraGPUSystemTick& Tick, FNiagaraComputeInstanceData *Instance, FRHICommandList &RHICmdList, const FNiagaraShaderRef& ComputeShader, const uint32 ShaderStageIndex) const;
+	void PostStageInterface(const FNiagaraGPUSystemTick& Tick, FNiagaraComputeInstanceData *Instance, FRHICommandList &RHICmdList, const FNiagaraShaderRef& ComputeShader, const uint32 SimulationStageIndex) const;
 
 	/** Get the shared SortManager, used in the rendering loop to call FGPUSortManager::OnPreRender() and FGPUSortManager::OnPostRenderOpaque() */
 	virtual FGPUSortManager* GetGPUSortManager() const override;
diff --git a/Engine/Plugins/FX/Niagara/Source/Niagara/Classes/NiagaraScript.h b/Engine/Plugins/FX/Niagara/Source/Niagara/Classes/NiagaraScript.h
--- a/Engine/Plugins/FX/Niagara/Source/Niagara/Classes/NiagaraScript.h
+++ b/Engine/Plugins/FX/Niagara/Source/Niagara/Classes/NiagaraScript.h
@@ -111,6 +111,37 @@ struct FNiagaraScriptDebuggerInfo
 	TAtomic<bool> bWritten;
 };
 
+USTRUCT()
+struct NIAGARA_API FSimulationStageMetaData
+{
+	GENERATED_USTRUCT_BODY()
+public:
+
+	/** The Data Interface that we iterate over for this stage. If None, then use particles.*/
+	UPROPERTY()
+	FName IterationSource;
+
+	/** Is this stage a spawn-only stage? */
+	UPROPERTY()
+	uint32 bSpawnOnly : 1;
+
+	/** Do we write to particles this stage?*/
+	UPROPERTY()
+	uint32 bWritesParticles : 1;
+
+	/** DataInterfaces that we write to in this stage.*/
+	UPROPERTY()
+	TArray<FName> OutputDestinations;
+
+	/** Index of the simulation stage where we begin iterating. This is meant to encompass iteration count without having an entry for each iteration.*/
+	UPROPERTY()
+	int32 MinStage;
+
+	/** Index of the simulation stage where we end iterating. This is meant to encompass iteration count without having an entry for each iteration.*/
+	UPROPERTY()
+	int32 MaxStage;
+};
+
 
 /** Struct containing all of the data necessary to look up a NiagaraScript's VM executable results from the Derived Data Cache.*/
 USTRUCT()
@@ -314,6 +345,9 @@ public:
 	UPROPERTY()
 	ENiagaraScriptCompileStatus LastCompileStatus;
 
+	UPROPERTY()
+	TArray<FSimulationStageMetaData> SimulationStageMetaData;
+
 #if WITH_EDITORONLY_DATA
 	UPROPERTY()
 	bool bReadsAttributeData;
diff --git a/Engine/Plugins/FX/Niagara/Source/Niagara/Classes/NiagaraScriptExecutionContext.h b/Engine/Plugins/FX/Niagara/Source/Niagara/Classes/NiagaraScriptExecutionContext.h
--- a/Engine/Plugins/FX/Niagara/Source/Niagara/Classes/NiagaraScriptExecutionContext.h
+++ b/Engine/Plugins/FX/Niagara/Source/Niagara/Classes/NiagaraScriptExecutionContext.h
@@ -201,7 +201,7 @@ struct FNiagaraComputeExecutionContext
 
 	void Reset(NiagaraEmitterInstanceBatcher* Batcher);
 
-	void InitParams(UNiagaraScript* InGPUComputeScript, ENiagaraSimTarget InSimTarget, const FString& InDebugSimName, const uint32 InDefaultShaderStageIndex, int32 InMaxUpdateIterations, const TSet<uint32> InSpawnStages);
+	void InitParams(UNiagaraScript* InGPUComputeScript, ENiagaraSimTarget InSimTarget, const FString& InDebugSimName, const uint32 InDefaultSimulationStageIndex, int32 InMaxUpdateIterations, const TSet<uint32> InSpawnStages);
 	void DirtyDataInterfaces();
 	bool Tick(FNiagaraSystemInstance* ParentSystemInstance);
 
@@ -217,7 +217,7 @@ struct FNiagaraComputeExecutionContext
 		// The CPU instance count at the time the GPU count readback was issued. Always bigger or equal to the GPU count.
 		uint32 CPUCount = 0;
 	}  EmitterInstanceReadback;
-
+	
 private:
 	void ResetInternal(NiagaraEmitterInstanceBatcher* Batcher);
 
@@ -232,6 +232,8 @@ public:
 	FORCEINLINE const TCHAR* GetDebugSimName() const { return TEXT(""); }
 #endif
 
+	const TArray<UNiagaraDataInterface*>& GetDataInterfaces()const { return CombinedParamStore.GetDataInterfaces(); }
+
 	class FNiagaraDataSet *MainDataSet;
 	UNiagaraScript* GPUScript;
 	class FNiagaraShaderScript*  GPUScript_RT;
@@ -258,7 +260,7 @@ public:
 	// Game thread spawn info will be sent to the render thread inside FNiagaraComputeInstanceData
 	FNiagaraGpuSpawnInfo GpuSpawnInfo_GT;
 
-	uint32 DefaultShaderStageIndex;
+	uint32 DefaultSimulationStageIndex;
 	uint32 MaxUpdateIterations;
 	TSet<uint32> SpawnStages;
 
@@ -267,6 +269,13 @@ public:
 	/** Temp data used in NiagaraEmitterInstanceBatcher::ExecuteAll() to avoid creating a map per FNiagaraComputeExecutionContext */
 	mutable int32 ScratchIndex = INDEX_NONE;
 
+	TArray < FSimulationStageMetaData> SimStageInfo;
+
+	bool IsOutputStage(FNiagaraDataInterfaceProxy* DIProxy, uint32 CurrentStage) const;
+	bool IsIterationStage(FNiagaraDataInterfaceProxy* DIProxy, uint32 CurrentStage) const;
+	FNiagaraDataInterfaceProxy* FindIterationInterface(const TArray<FNiagaraDataInterfaceProxy*>& InProxies, uint32 SimulationStageIndex) const;
+	const FSimulationStageMetaData* GetSimStageMetaData(uint32 SimulationStageIndex) const;
+
 #if WITH_EDITORONLY_DATA
 	mutable FRHIGPUMemoryReadback *GPUDebugDataReadbackFloat;
 	mutable FRHIGPUMemoryReadback *GPUDebugDataReadbackInt;
@@ -294,6 +303,16 @@ struct FNiagaraDataInterfaceInstanceData
 
 //TODO: Rename FNiagaraGPUEmitterTick?
 
+
+struct FNiagaraSimStageData
+{
+	FNiagaraDataBuffer* Source;
+	FNiagaraDataBuffer* Destination;
+	FNiagaraDataInterfaceProxy* AlternateIterationSource;
+	uint32 SourceCountOffset;
+	uint32 DestinationCountOffset;
+};
+
 struct FNiagaraComputeInstanceData
 {
 	FNiagaraGpuSpawnInfo SpawnInfo;
@@ -304,13 +323,17 @@ struct FNiagaraComputeInstanceData
 	uint8* ExternalParamData = nullptr;
 	FNiagaraComputeExecutionContext* Context = nullptr;
 	TArray<FNiagaraDataInterfaceProxy*> DataInterfaceProxies;
+	bool bUsesSimStages = false;
+	bool bUsesOldShaderStages = false;
+	TArray<FNiagaraSimStageData, TInlineAllocator<1>> SimStageData;
 
-	//Buffer containing current state that this tick will read from. Initialized at the start of processing this tick on the RT.
-	FNiagaraDataBuffer* CurrentData = nullptr;
-	//Buffer into which we'll write the new simulation state. Initialized at the start of processing this tick on the RT.
-	FNiagaraDataBuffer* DestinationData = nullptr;
+	bool IsOutputStage(FNiagaraDataInterfaceProxy* DIProxy, uint32 CurrentStage) const;
+	bool IsIterationStage(FNiagaraDataInterfaceProxy* DIProxy, uint32 CurrentStage) const;
+	FNiagaraDataInterfaceProxy* FindIterationInterface(uint32 SimulationStageIndex) const;
 };
 
+
+
 /*
 	Represents all the information needed to dispatch a single tick of a FNiagaraSystemInstance.
 	This object will be created on the game thread and passed to the renderthread.
@@ -342,6 +365,7 @@ public:
 	FORCEINLINE FNiagaraComputeInstanceData* GetInstanceData()const{ return reinterpret_cast<FNiagaraComputeInstanceData*>(InstanceData_ParamData_Packed); }
 
 	uint32 Count;
+	uint32 TotalDispatches;
 	FNiagaraSystemInstanceID SystemInstanceID;
 	FNiagaraDataInterfaceInstanceData* DIInstanceData;
 	uint8* InstanceData_ParamData_Packed;
@@ -350,4 +374,5 @@ public:
 	bool bRequiresEarlyViewData = false;
 	bool bNeedsReset = false;
 	bool bIsFinalTick = false;
+	uint32 NumInstancesWithSimStages = 0;
 };
diff --git a/Engine/Plugins/FX/Niagara/Source/Niagara/Private/NiagaraDataInterfaceGrid2DCollection.cpp b/Engine/Plugins/FX/Niagara/Source/Niagara/Private/NiagaraDataInterfaceGrid2DCollection.cpp
--- a/Engine/Plugins/FX/Niagara/Source/Niagara/Private/NiagaraDataInterfaceGrid2DCollection.cpp
+++ b/Engine/Plugins/FX/Niagara/Source/Niagara/Private/NiagaraDataInterfaceGrid2DCollection.cpp
@@ -117,7 +117,6 @@ UNiagaraDataInterfaceGrid2DCollection::UNiagaraDataInterfaceGrid2DCollection(FOb
 	: Super(ObjectInitializer)
 {
 	Proxy.Reset(new FNiagaraDataInterfaceProxyGrid2DCollection());
-	RWProxy = (FNiagaraDataInterfaceProxyRW*) Proxy.Get();
 	PushToRenderThread();
 }
 
@@ -146,6 +145,7 @@ void UNiagaraDataInterfaceGrid2DCollection::GetFunctions(TArray<FNiagaraFunction
 		Sig.Inputs.Add(FNiagaraVariable(FNiagaraTypeDefinition::GetIntDef(), TEXT("AttributeIndex")));
 		Sig.Outputs.Add(FNiagaraVariable(FNiagaraTypeDefinition::GetFloatDef(), TEXT("Value")));
 
+		Sig.bExperimental = true;
 		Sig.bMemberFunction = true;
 		Sig.bRequiresContext = false;
 		OutFunctions.Add(Sig);
@@ -161,8 +161,10 @@ void UNiagaraDataInterfaceGrid2DCollection::GetFunctions(TArray<FNiagaraFunction
 		Sig.Inputs.Add(FNiagaraVariable(FNiagaraTypeDefinition::GetFloatDef(), TEXT("Value")));
 		Sig.Outputs.Add(FNiagaraVariable(FNiagaraTypeDefinition::GetIntDef(), TEXT("IGNORE")));
 
+		Sig.bExperimental = true;
 		Sig.bMemberFunction = true;
 		Sig.bRequiresContext = false;
+		Sig.bWriteFunction = true;
 		OutFunctions.Add(Sig);
 	}
 
@@ -175,6 +177,7 @@ void UNiagaraDataInterfaceGrid2DCollection::GetFunctions(TArray<FNiagaraFunction
 		Sig.Inputs.Add(FNiagaraVariable(FNiagaraTypeDefinition::GetIntDef(), TEXT("AttributeIndex")));
 		Sig.Outputs.Add(FNiagaraVariable(FNiagaraTypeDefinition::GetFloatDef(), TEXT("Value")));
 
+		Sig.bExperimental = true;
 		Sig.bMemberFunction = true;
 		Sig.bRequiresContext = false;
 		OutFunctions.Add(Sig);
@@ -328,11 +331,11 @@ bool UNiagaraDataInterfaceGrid2DCollection::CopyToInternal(UNiagaraDataInterface
 bool UNiagaraDataInterfaceGrid2DCollection::InitPerInstanceData(void* PerInstanceData, FNiagaraSystemInstance* SystemInstance)
 {
 	check(Proxy);
-
+	UE_LOG(LogNiagara, Log, TEXT("UNiagaraDataInterfaceGrid2DCollection::InitPerInstanceData DI: %p PID: %p SI: %p %s"), this, PerInstanceData, SystemInstance, *SystemInstance->GetSystem()->GetPathName())
 	Grid2DCollectionRWInstanceData* InstanceData = new (PerInstanceData) Grid2DCollectionRWInstanceData();
 
 	FNiagaraDataInterfaceProxyGrid2DCollection* RT_Proxy = GetProxyAs<FNiagaraDataInterfaceProxyGrid2DCollection>();
-
+	
 	int RT_NumCellsX = NumCellsX;
 	int RT_NumCellsY = NumCellsY;
 
@@ -379,6 +382,13 @@ bool UNiagaraDataInterfaceGrid2DCollection::InitPerInstanceData(void* PerInstanc
 	FVector2D RT_CellSize = RT_WorldBBoxSize / FVector2D(RT_NumCellsX, RT_NumCellsY);
 	InstanceData->CellSize = RT_CellSize;
 	InstanceData->WorldBBoxSize = RT_WorldBBoxSize;
+
+	//UE_LOG(LogNiagara, Log, TEXT("UNiagaraDataInterfaceGrid2DCollection::Proxy %p DI: %p"), RT_Proxy, this);
+	//for (const int32 i : IterationShaderStages)
+	//{
+	//	UE_LOG(LogNiagara, Log, TEXT("%d"), i);
+	//}
+
 	
 	// Push Updates to Proxy.
 	ENQUEUE_RENDER_COMMAND(FUpdateData)(
@@ -397,8 +407,8 @@ bool UNiagaraDataInterfaceGrid2DCollection::InitPerInstanceData(void* PerInstanc
 		
 		TargetData->WorldBBoxSize = RT_WorldBBoxSize;
 
-		RT_Proxy->OutputShaderStages = RT_OutputShaderStages;
-		RT_Proxy->IterationShaderStages = RT_IterationShaderStages;
+		RT_Proxy->OutputSimulationStages_DEPRECATED = RT_OutputShaderStages;
+		RT_Proxy->IterationSimulationStages_DEPRECATED = RT_IterationShaderStages;
 
 		RT_Proxy->SetElementCount(TargetData->NumCellsX * TargetData->NumCellsY);
 	});
@@ -473,8 +483,11 @@ bool UNiagaraDataInterfaceGrid2DCollection::FillTexture2D(const UNiagaraComponen
 			CopyInfo.SourcePosition = FIntVector(StartX, StartY, 0);
 
 			RHICmdList.TransitionResource(EResourceTransitionAccess::EReadable, Grid2DInstanceData->CurrentData->GridBuffer.Buffer);
+			RHICmdList.TransitionResource(EResourceTransitionAccess::EWritable, Dest->Resource->TextureRHI);
 
 			RHICmdList.CopyTexture(Grid2DInstanceData->CurrentData->GridBuffer.Buffer, Dest->Resource->TextureRHI, CopyInfo);
+
+			RHICmdList.TransitionResource(EResourceTransitionAccess::EReadable, Dest->Resource->TextureRHI);
 		}
 	});
 
@@ -530,9 +543,11 @@ bool UNiagaraDataInterfaceGrid2DCollection::FillRawTexture2D(const UNiagaraCompo
 		if (Dest && Dest->Resource && Grid2DInstanceDataTmp && Grid2DInstanceDataTmp->CurrentData)
 		{
 			RHICmdList.TransitionResource(EResourceTransitionAccess::EReadable, Grid2DInstanceDataTmp->CurrentData->GridBuffer.Buffer);
+			RHICmdList.TransitionResource(EResourceTransitionAccess::EWritable, Dest->Resource->TextureRHI);
 
 			FRHICopyTextureInfo CopyInfo;
 			RHICmdList.CopyTexture(Grid2DInstanceDataTmp->CurrentData->GridBuffer.Buffer, Dest->Resource->TextureRHI, CopyInfo);
+			RHICmdList.TransitionResource(EResourceTransitionAccess::EReadable, Dest->Resource->TextureRHI);
 		}
 	});
 
diff --git a/Engine/Plugins/FX/Niagara/Source/Niagara/Private/NiagaraDataInterfaceNeighborGrid3D.cpp b/Engine/Plugins/FX/Niagara/Source/Niagara/Private/NiagaraDataInterfaceNeighborGrid3D.cpp
--- a/Engine/Plugins/FX/Niagara/Source/Niagara/Private/NiagaraDataInterfaceNeighborGrid3D.cpp
+++ b/Engine/Plugins/FX/Niagara/Source/Niagara/Private/NiagaraDataInterfaceNeighborGrid3D.cpp
@@ -147,6 +147,7 @@ void UNiagaraDataInterfaceNeighborGrid3D::GetFunctions(TArray<FNiagaraFunctionSi
 		Sig.Inputs.Add(FNiagaraVariable(FNiagaraTypeDefinition(GetClass()), TEXT("Grid")));
 		Sig.Outputs.Add(FNiagaraVariable(FNiagaraTypeDefinition::GetIntDef(), TEXT("MaxNeighborsPerVoxel")));
 
+		Sig.bExperimental = true;
 		Sig.bMemberFunction = true;
 		Sig.bRequiresContext = false;
 		OutFunctions.Add(Sig);
@@ -162,6 +163,7 @@ void UNiagaraDataInterfaceNeighborGrid3D::GetFunctions(TArray<FNiagaraFunctionSi
 		Sig.Inputs.Add(FNiagaraVariable(FNiagaraTypeDefinition::GetIntDef(), TEXT("Neighbor")));
 		Sig.Outputs.Add(FNiagaraVariable(FNiagaraTypeDefinition::GetIntDef(), TEXT("Linear Index")));
 
+		Sig.bExperimental = true;
 		Sig.bMemberFunction = true;
 		Sig.bRequiresContext = false;
 		OutFunctions.Add(Sig);
@@ -174,6 +176,7 @@ void UNiagaraDataInterfaceNeighborGrid3D::GetFunctions(TArray<FNiagaraFunctionSi
 		Sig.Inputs.Add(FNiagaraVariable(FNiagaraTypeDefinition::GetIntDef(), TEXT("Linear")));
 		Sig.Outputs.Add(FNiagaraVariable(FNiagaraTypeDefinition::GetIntDef(), TEXT("NeighborIndex")));
 
+		Sig.bExperimental = true;
 		Sig.bMemberFunction = true;
 		Sig.bRequiresContext = false;
 		OutFunctions.Add(Sig);
@@ -187,6 +190,8 @@ void UNiagaraDataInterfaceNeighborGrid3D::GetFunctions(TArray<FNiagaraFunctionSi
 		Sig.Inputs.Add(FNiagaraVariable(FNiagaraTypeDefinition::GetIntDef(), TEXT("NeighborIndex")));
 		Sig.Outputs.Add(FNiagaraVariable(FNiagaraTypeDefinition::GetIntDef(), TEXT("IGNORE")));
 
+		Sig.bExperimental = true;
+		Sig.bWriteFunction = true;
 		Sig.bMemberFunction = true;
 		Sig.bRequiresContext = false;
 		OutFunctions.Add(Sig);
@@ -199,6 +204,7 @@ void UNiagaraDataInterfaceNeighborGrid3D::GetFunctions(TArray<FNiagaraFunctionSi
 		Sig.Inputs.Add(FNiagaraVariable(FNiagaraTypeDefinition::GetIntDef(), TEXT("Linear")));
 		Sig.Outputs.Add(FNiagaraVariable(FNiagaraTypeDefinition::GetIntDef(), TEXT("NeighborCount")));
 
+		Sig.bExperimental = true;
 		Sig.bMemberFunction = true;
 		Sig.bRequiresContext = false;
 		OutFunctions.Add(Sig);
@@ -212,6 +218,8 @@ void UNiagaraDataInterfaceNeighborGrid3D::GetFunctions(TArray<FNiagaraFunctionSi
 		Sig.Inputs.Add(FNiagaraVariable(FNiagaraTypeDefinition::GetIntDef(), TEXT("Increment")));
 		Sig.Outputs.Add(FNiagaraVariable(FNiagaraTypeDefinition::GetIntDef(), TEXT("PrevNeighborCount")));
 
+		Sig.bExperimental = true;
+		Sig.bWriteFunction = true;
 		Sig.bMemberFunction = true;
 		Sig.bRequiresContext = false;
 		OutFunctions.Add(Sig);
@@ -430,8 +438,8 @@ bool UNiagaraDataInterfaceNeighborGrid3D::InitPerInstanceData(void* PerInstanceD
 		TargetData->MaxNeighborsPerVoxel = RT_MaxNeighborsPerVoxel;		
 		TargetData->WorldBBoxSize = RT_WorldBBoxSize;
 
-		RT_Proxy->OutputShaderStages = RT_OutputShaderStages;
-		RT_Proxy->IterationShaderStages = RT_IterationShaderStages;
+		RT_Proxy->OutputSimulationStages_DEPRECATED = RT_OutputShaderStages;
+		RT_Proxy->IterationSimulationStages_DEPRECATED = RT_IterationShaderStages;
 
 		// #todo(dmp): element count is still defined on the proxy and not the instance data
 		RT_Proxy->SetElementCount(RT_NumVoxels.X *  RT_NumVoxels.Y * RT_NumVoxels.Z);
@@ -450,6 +458,8 @@ void UNiagaraDataInterfaceNeighborGrid3D::DestroyPerInstanceData(void* PerInstan
 	InstanceData->~NeighborGrid3DRWInstanceData();
 
 	FNiagaraDataInterfaceProxyNeighborGrid3D* ThisProxy = GetProxyAs<FNiagaraDataInterfaceProxyNeighborGrid3D>();
+	if (!ThisProxy)
+		return;
 
 	ENQUEUE_RENDER_COMMAND(FNiagaraDIDestroyInstanceData) (
 		[ThisProxy, InstanceID = SystemInstance->GetId(), Batcher = SystemInstance->GetBatcher()](FRHICommandListImmediate& CmdList)
diff --git a/Engine/Plugins/FX/Niagara/Source/Niagara/Private/NiagaraDataInterfaceParticleRead.cpp b/Engine/Plugins/FX/Niagara/Source/Niagara/Private/NiagaraDataInterfaceParticleRead.cpp
--- a/Engine/Plugins/FX/Niagara/Source/Niagara/Private/NiagaraDataInterfaceParticleRead.cpp
+++ b/Engine/Plugins/FX/Niagara/Source/Niagara/Private/NiagaraDataInterfaceParticleRead.cpp
@@ -387,14 +387,14 @@ struct FNiagaraDataInterfaceParametersCS_ParticleRead : public FNiagaraDataInter
 			// ourselves. We can't use SourceDataSet->GetCurrentData() in that case, because EndSimulate() has already been
 			// called on the current emitter, and the current data has been set to the destination data. We need to use the
 			// current compute instance data to get to the input buffers.
-			SourceData = Context.ComputeInstanceData->CurrentData;
+			SourceData = Context.ComputeInstanceData->SimStageData[0].Source;
 
 			// We still want to get the spawn count and ID acquire tag from the destination data, because that's where
 			// NiagaraEmitterInstanceBatcher::Run() stores them.
-			if (Context.ComputeInstanceData->DestinationData != nullptr)
+			if (Context.ComputeInstanceData->SimStageData[0].Destination != nullptr)
 			{
-				NumSpawnedInstances = Context.ComputeInstanceData->DestinationData->GetNumSpawnedInstances();
-				IDAcquireTag = Context.ComputeInstanceData->DestinationData->GetIDAcquireTag();
+				NumSpawnedInstances = Context.ComputeInstanceData->SimStageData[0].Destination->GetNumSpawnedInstances();
+				IDAcquireTag = Context.ComputeInstanceData->SimStageData[0].Destination->GetIDAcquireTag();
 			}
 		}
 		else
diff --git a/Engine/Plugins/FX/Niagara/Source/Niagara/Private/NiagaraDataInterfaceRW.cpp b/Engine/Plugins/FX/Niagara/Source/Niagara/Private/NiagaraDataInterfaceRW.cpp
--- a/Engine/Plugins/FX/Niagara/Source/Niagara/Private/NiagaraDataInterfaceRW.cpp
+++ b/Engine/Plugins/FX/Niagara/Source/Niagara/Private/NiagaraDataInterfaceRW.cpp
@@ -78,7 +78,6 @@ UNiagaraDataInterfaceGrid3D::UNiagaraDataInterfaceGrid3D(FObjectInitializer cons
 	, WorldBBoxSize(100., 100., 100.)
 {
 	Proxy.Reset(new FNiagaraDataInterfaceProxyRW());
-	RWProxy = (FNiagaraDataInterfaceProxyRW*) Proxy.Get();
 	PushToRenderThread();
 }
 
diff --git a/Engine/Plugins/FX/Niagara/Source/Niagara/Private/NiagaraDataSet.cpp b/Engine/Plugins/FX/Niagara/Source/Niagara/Private/NiagaraDataSet.cpp
--- a/Engine/Plugins/FX/Niagara/Source/Niagara/Private/NiagaraDataSet.cpp
+++ b/Engine/Plugins/FX/Niagara/Source/Niagara/Private/NiagaraDataSet.cpp
@@ -154,7 +154,7 @@ void FNiagaraDataSet::ReleaseBuffers()
 	GPUNumAllocatedIDs = 0;
 }
 
-FNiagaraDataBuffer& FNiagaraDataSet::BeginSimulate()
+FNiagaraDataBuffer& FNiagaraDataSet::BeginSimulate(bool bResetDestinationData)
 {
 	//CheckCorrectThread();
 	check(DestinationData == nullptr);
@@ -179,8 +179,11 @@ FNiagaraDataBuffer& FNiagaraDataSet::BeginSimulate()
 		checkSlow(DestinationData->IsBeingWritten());
 	}
 
-	DestinationData->SetNumInstances(0);
-	DestinationData->GetIDTable().Reset();
+	if (bResetDestinationData)
+	{
+		DestinationData->SetNumInstances(0);
+		DestinationData->GetIDTable().Reset();
+	}
 
 	return GetDestinationDataChecked();
 }
@@ -652,10 +655,12 @@ void FNiagaraDataBuffer::AllocateGPU(uint32 InNumInstances, FNiagaraGPUInstanceC
 
 	checkSlow(Owner->GetSimTarget() == ENiagaraSimTarget::GPUComputeSim);
 
+	//uint32 OldOffset = GPUInstanceCountBufferOffset;
 	// Release previous entry if any.
 	GPUInstanceCountManager.FreeEntry(GPUInstanceCountBufferOffset);
 	// Get a new entry currently set to 0, since simulation will increment it to the actual instance count.
 	GPUInstanceCountBufferOffset = GPUInstanceCountManager.AcquireEntry();
+	//UE_LOG(LogNiagara, Log, TEXT("AllocateGPU %p GPUInstanceCountBufferOffsetOld: %d New: %d"), this, OldOffset, GPUInstanceCountBufferOffset);
 
 	// ALLOC_CHUNKSIZE must be greater than zero and divisible by the thread group size
 	const uint32 ALLOC_CHUNKSIZE = 4096;
diff --git a/Engine/Plugins/FX/Niagara/Source/Niagara/Private/NiagaraEmitter.cpp b/Engine/Plugins/FX/Niagara/Source/Niagara/Private/NiagaraEmitter.cpp
--- a/Engine/Plugins/FX/Niagara/Source/Niagara/Private/NiagaraEmitter.cpp
+++ b/Engine/Plugins/FX/Niagara/Source/Niagara/Private/NiagaraEmitter.cpp
@@ -5,7 +5,7 @@
 #include "NiagaraScript.h"
 #include "NiagaraScriptSourceBase.h"
 #include "NiagaraRendererProperties.h"
-#include "NiagaraShaderStageBase.h"
+#include "NiagaraSimulationStageBase.h"
 #include "NiagaraCustomVersion.h"
 #include "UObject/Package.h"
 #include "UObject/Linker.h"
@@ -264,6 +264,16 @@ UNiagaraEmitter* UNiagaraEmitter::DuplicateWithoutMerging(UObject* InOuter)
 
 void UNiagaraEmitter::Serialize(FArchive& Ar)
 {
+
+	for (UNiagaraSimulationStageBase* Stage : SimulationStages)
+	{
+		ensure(Stage);
+		ensure(Stage->Script);
+		if (!HasAnyFlags(RF_Transient))
+		{
+			ensure(!Stage->Script->HasAnyFlags(RF_Transient));
+		}
+	}
 	Super::Serialize(Ar);
 
 	Ar.UsingCustomVersion(FNiagaraCustomVersion::GUID);
@@ -286,11 +296,12 @@ void UNiagaraEmitter::PostLoad()
 		}
 	}
 
-	for (int32 ShaderStageIndex = ShaderStages.Num() - 1; ShaderStageIndex >= 0; --ShaderStageIndex)
+
+	for (int32 SimulationStageIndex = SimulationStages.Num() - 1; SimulationStageIndex >= 0; --SimulationStageIndex)
 	{
-		if (ensureMsgf(ShaderStages[ShaderStageIndex] != nullptr && ShaderStages[ShaderStageIndex]->Script != nullptr, TEXT("Null shader stage, or shader stage with a null script found in %s at index %i, removing it to prevent crashes."), *GetPathName(), ShaderStageIndex) == false)
+		if (ensureMsgf(SimulationStages[SimulationStageIndex] != nullptr && SimulationStages[SimulationStageIndex]->Script != nullptr, TEXT("Null simulation stage, or simulation stage with a null script found in %s at index %i, removing it to prevent crashes."), *GetPathName(), SimulationStageIndex) == false)
 		{
-			ShaderStages.RemoveAt(ShaderStageIndex);
+			SimulationStages.RemoveAt(SimulationStageIndex);
 		}
 	}
 
@@ -512,10 +523,10 @@ void UNiagaraEmitter::PostLoad()
 				FNiagaraParameterStore::FOnChanged::FDelegate::CreateUObject(this, &UNiagaraEmitter::ScriptRapidIterationParameterChanged));
 		}
 
-		for (UNiagaraShaderStageBase* ShaderStage : ShaderStages)
+		for (UNiagaraSimulationStageBase* SimulationStage : SimulationStages)
 		{
-			ShaderStage->OnChanged().AddUObject(this, &UNiagaraEmitter::ShaderStageChanged);
-			ShaderStage->Script->RapidIterationParameters.AddOnChangedHandler(
+			SimulationStage->OnChanged().AddUObject(this, &UNiagaraEmitter::SimulationStageChanged);
+			SimulationStage->Script->RapidIterationParameters.AddOnChangedHandler(
 				FNiagaraParameterStore::FOnChanged::FDelegate::CreateUObject(this, &UNiagaraEmitter::ScriptRapidIterationParameterChanged));
 		}
 
@@ -645,9 +656,24 @@ void UNiagaraEmitter::PostEditChangeProperty(struct FPropertyChangedEvent& Prope
 			GraphSource->MarkNotSynchronized(TEXT("Emitter Determinism changed."));
 		}
 
-#if WITH_EDITORONLY_DATA
-		UNiagaraSystem::RequestCompileForEmitter(this);
-#endif
+		bNeedsRecompile = true;
+	}
+	if (PropertyName == GET_MEMBER_NAME_CHECKED(UNiagaraEmitter, bSimulationStagesEnabled))
+	{
+		if (GraphSource != nullptr)
+		{
+			GraphSource->MarkNotSynchronized(TEXT("SimulationStagesEnabled changed."));
+		}
+		bNeedsRecompile = true;
+	}
+	if (PropertyName == GET_MEMBER_NAME_CHECKED(UNiagaraEmitter, bDeprecatedShaderStagesEnabled))
+	{
+		if (GraphSource != nullptr)
+		{
+			GraphSource->MarkNotSynchronized(TEXT("DeprecatedShaderStagesEnabled changed."));
+		}
+		bNeedsRecompile = true;
+
 	}
 
 	ResolveScalabilitySettings();
@@ -785,11 +811,11 @@ void UNiagaraEmitter::GetScripts(TArray<UNiagaraScript*>& OutScripts, bool bComp
 
 	if (!bCompilableOnly)
 	{
-		for (int32 i = 0; i < ShaderStages.Num(); i++)
+		for (int32 i = 0; i < SimulationStages.Num(); i++)
 		{
-			if (ShaderStages[i] && ShaderStages[i]->Script)
+			if (SimulationStages[i] && SimulationStages[i]->Script)
 			{
-				OutScripts.Add(ShaderStages[i]->Script);
+				OutScripts.Add(SimulationStages[i]->Script);
 			}
 		}
 	}
@@ -881,9 +907,9 @@ bool UNiagaraEmitter::AreAllScriptAndSourcesSynchronized() const
 		}
 	}
 
-	for (int32 i = 0; i < ShaderStages.Num(); i++)
+	for (int32 i = 0; i < SimulationStages.Num(); i++)
 	{
-		if (ShaderStages[i] && ShaderStages[i]->Script  && ShaderStages[i]->Script->IsCompilable() && !ShaderStages[i]->Script->AreScriptAndSourceSynchronized())
+		if (SimulationStages[i] && SimulationStages[i]->Script  && SimulationStages[i]->Script->IsCompilable() && !SimulationStages[i]->Script->AreScriptAndSourceSynchronized())
 		{
 			return false;
 		}
@@ -1146,19 +1172,19 @@ void UNiagaraEmitter::UpdateFromMergedCopy(const INiagaraMergeManager& MergeMana
 	}
 
 	// Copy shader stages
-	for (UNiagaraShaderStageBase*& ShaderStage : ShaderStages)
+	for (UNiagaraSimulationStageBase*& SimulationStage : SimulationStages)
 	{
-		ShaderStage->OnChanged().RemoveAll(this);
-		ShaderStage->Script->RapidIterationParameters.RemoveAllOnChangedHandlers(this);
+		SimulationStage->OnChanged().RemoveAll(this);
+		SimulationStage->Script->RapidIterationParameters.RemoveAllOnChangedHandlers(this);
 	}
-	ShaderStages.Empty();
+	SimulationStages.Empty();
 
-	for (UNiagaraShaderStageBase* MergedShaderStage : MergedEmitter->ShaderStages)
+	for (UNiagaraSimulationStageBase* MergedSimulationStage : MergedEmitter->SimulationStages)
 	{
-		ReouterMergedObject(this, MergedShaderStage);
-		ShaderStages.Add(MergedShaderStage);
-		MergedShaderStage->OnChanged().AddUObject(this, &UNiagaraEmitter::ShaderStageChanged);
-		MergedShaderStage->Script->RapidIterationParameters.AddOnChangedHandler(
+		ReouterMergedObject(this, MergedSimulationStage);
+		SimulationStages.Add(MergedSimulationStage);
+		MergedSimulationStage->OnChanged().AddUObject(this, &UNiagaraEmitter::SimulationStageChanged);
+		MergedSimulationStage->Script->RapidIterationParameters.AddOnChangedHandler(
 			FNiagaraParameterStore::FOnChanged::FDelegate::CreateUObject(this, &UNiagaraEmitter::ScriptRapidIterationParameterChanged));
 	}
 
@@ -1307,52 +1333,52 @@ void UNiagaraEmitter::RemoveEventHandlerByUsageId(FGuid EventHandlerUsageId)
 #endif
 }
 
-UNiagaraShaderStageBase* UNiagaraEmitter::GetShaderStageById(FGuid ScriptUsageId) const
+UNiagaraSimulationStageBase* UNiagaraEmitter::GetSimulationStageById(FGuid ScriptUsageId) const
 {
-	UNiagaraShaderStageBase*const* FoundShaderStagePtr = ShaderStages.FindByPredicate([&ScriptUsageId](UNiagaraShaderStageBase* ShaderStage) { return ShaderStage->Script->GetUsageId() == ScriptUsageId; });
-	return FoundShaderStagePtr != nullptr ? *FoundShaderStagePtr : nullptr;
+	UNiagaraSimulationStageBase*const* FoundSimulationStagePtr = SimulationStages.FindByPredicate([&ScriptUsageId](UNiagaraSimulationStageBase* SimulationStage) { return SimulationStage->Script->GetUsageId() == ScriptUsageId; });
+	return FoundSimulationStagePtr != nullptr ? *FoundSimulationStagePtr : nullptr;
 }
 
-void UNiagaraEmitter::AddShaderStage(UNiagaraShaderStageBase* ShaderStage)
+void UNiagaraEmitter::AddSimulationStage(UNiagaraSimulationStageBase* SimulationStage)
 {
 	Modify();
-	ShaderStages.Add(ShaderStage);
+	SimulationStages.Add(SimulationStage);
 #if WITH_EDITOR
-	ShaderStage->OnChanged().AddUObject(this, &UNiagaraEmitter::ShaderStageChanged);
-	ShaderStage->Script->RapidIterationParameters.AddOnChangedHandler(
+	SimulationStage->OnChanged().AddUObject(this, &UNiagaraEmitter::SimulationStageChanged);
+	SimulationStage->Script->RapidIterationParameters.AddOnChangedHandler(
 		FNiagaraParameterStore::FOnChanged::FDelegate::CreateUObject(this, &UNiagaraEmitter::ScriptRapidIterationParameterChanged));
 	UpdateChangeId(TEXT("Shader stage added"));
 #endif
 }
 
-void UNiagaraEmitter::RemoveShaderStage(UNiagaraShaderStageBase* ShaderStage)
+void UNiagaraEmitter::RemoveSimulationStage(UNiagaraSimulationStageBase* SimulationStage)
 {
 	Modify();
-	bool bRemoved = ShaderStages.Remove(ShaderStage) != 0;
+	bool bRemoved = SimulationStages.Remove(SimulationStage) != 0;
 #if WITH_EDITOR
 	if (bRemoved)
 	{
-		ShaderStage->OnChanged().RemoveAll(this);
-		ShaderStage->Script->RapidIterationParameters.RemoveAllOnChangedHandlers(this);
-		UpdateChangeId(TEXT("Shader stage removed"));
+		SimulationStage->OnChanged().RemoveAll(this);
+		SimulationStage->Script->RapidIterationParameters.RemoveAllOnChangedHandlers(this);
+		UpdateChangeId(TEXT("Simulation stage removed"));
 	}
 #endif
 }
 
-void UNiagaraEmitter::MoveShaderStageToIndex(UNiagaraShaderStageBase* ShaderStageToMove, int32 TargetIndex)
+void UNiagaraEmitter::MoveSimulationStageToIndex(UNiagaraSimulationStageBase* SimulationStageToMove, int32 TargetIndex)
 {
-	int32 CurrentIndex = ShaderStages.IndexOfByKey(ShaderStageToMove);
-	checkf(CurrentIndex != INDEX_NONE, TEXT("Shader stage could not be moved because it is not owned by this emitter."));
+	int32 CurrentIndex = SimulationStages.IndexOfByKey(SimulationStageToMove);
+	checkf(CurrentIndex != INDEX_NONE, TEXT("Simulation stage could not be moved because it is not owned by this emitter."));
 	if (TargetIndex != CurrentIndex)
 	{
 		int32 AdjustedTargetIndex = CurrentIndex < TargetIndex
 			? TargetIndex - 1 // If the current index is less than the target index, the target index needs to be decreased to make up for the item being removed.
 			: TargetIndex;
 
-		ShaderStages.Remove(ShaderStageToMove);
-		ShaderStages.Insert(ShaderStageToMove, AdjustedTargetIndex);
+		SimulationStages.Remove(SimulationStageToMove);
+		SimulationStages.Insert(SimulationStageToMove, AdjustedTargetIndex);
 #if WITH_EDITOR
-		UpdateChangeId("Shader stage moved.");
+		UpdateChangeId("Simulation stage moved.");
 #endif
 	}
 }
@@ -1398,9 +1424,9 @@ void UNiagaraEmitter::ScriptRapidIterationParameterChanged()
 	UpdateChangeId(TEXT("Script rapid iteration parameter changed."));
 }
 
-void UNiagaraEmitter::ShaderStageChanged()
+void UNiagaraEmitter::SimulationStageChanged()
 {
-	UpdateChangeId(TEXT("Shader Stage Changed"));
+	UpdateChangeId(TEXT("Simulation Stage Changed"));
 }
 
 void UNiagaraEmitter::RendererChanged()
diff --git a/Engine/Plugins/FX/Niagara/Source/Niagara/Private/NiagaraEmitterInstance.cpp b/Engine/Plugins/FX/Niagara/Source/Niagara/Private/NiagaraEmitterInstance.cpp
--- a/Engine/Plugins/FX/Niagara/Source/Niagara/Private/NiagaraEmitterInstance.cpp
+++ b/Engine/Plugins/FX/Niagara/Source/Niagara/Private/NiagaraEmitterInstance.cpp
@@ -14,6 +14,7 @@
 #include "NiagaraParameterCollection.h"
 #include "NiagaraScriptExecutionContext.h"
 #include "NiagaraWorldManager.h"
+#include "NiagaraSimulationStageBase.h"
 
 DECLARE_DWORD_COUNTER_STAT(TEXT("Num Custom Events"), STAT_NiagaraNumCustomEvents, STATGROUP_Niagara);
 
@@ -358,6 +359,11 @@ void FNiagaraEmitterInstance::Init(int32 InEmitterIdx, FNiagaraSystemInstanceID
 
 			SpawnExecContext.Parameters.Bind(&GPUExecContext->CombinedParamStore);
 			UpdateExecContext.Parameters.Bind(&GPUExecContext->CombinedParamStore);
+
+			for (int32 i = 0; i < CachedEmitter->GetSimulationStages().Num(); i++)
+			{
+				CachedEmitter->GetSimulationStages()[i]->Script->RapidIterationParameters.Bind(&GPUExecContext->CombinedParamStore);
+			}
 		}
 	
 		EventExecContexts.SetNum(CachedEmitter->GetEventHandlers().Num());
@@ -709,6 +715,11 @@ void FNiagaraEmitterInstance::BindParameters(bool bExternalOnly)
 		{
 			SpawnExecContext.Parameters.Bind(&GPUExecContext->CombinedParamStore);
 			UpdateExecContext.Parameters.Bind(&GPUExecContext->CombinedParamStore);
+
+			for (int32 i = 0; i < CachedEmitter->GetSimulationStages().Num(); i++)
+			{
+				CachedEmitter->GetSimulationStages()[i]->Script->RapidIterationParameters.Bind(&GPUExecContext->CombinedParamStore);
+			}
 		}
 	}
 }
diff --git a/Engine/Plugins/FX/Niagara/Source/Niagara/Private/NiagaraEmitterInstanceBatcher.cpp b/Engine/Plugins/FX/Niagara/Source/Niagara/Private/NiagaraEmitterInstanceBatcher.cpp
--- a/Engine/Plugins/FX/Niagara/Source/Niagara/Private/NiagaraEmitterInstanceBatcher.cpp
+++ b/Engine/Plugins/FX/Niagara/Source/Niagara/Private/NiagaraEmitterInstanceBatcher.cpp
@@ -226,29 +226,13 @@ bool NiagaraEmitterInstanceBatcher::ResetDataInterfaces(const FNiagaraGPUSystemT
 	return ValidSpawnStage;
 }
 
-FNiagaraDataInterfaceProxy* NiagaraEmitterInstanceBatcher::FindIterationInterface( FNiagaraComputeInstanceData *Instance, const uint32 ShaderStageIndex) const
+FNiagaraDataInterfaceProxy* NiagaraEmitterInstanceBatcher::FindIterationInterface( FNiagaraComputeInstanceData *Instance, const uint32 SimulationStageIndex) const
 {
 	// Determine if the iteration is outputting to a custom data size
-	FNiagaraDataInterfaceProxy* IterationInterface = nullptr;
-
-	for (FNiagaraDataInterfaceProxy* Interface : Instance->DataInterfaceProxies)
-	{
-		if (Interface->IsIterationStage(ShaderStageIndex))
-		{
-			if (IterationInterface)
-			{
-				UE_LOG(LogNiagara, Error, TEXT("Multiple output Data Interfaces found for current stage"));
-			}
-			else
-			{
-				IterationInterface = Interface;
-			}
-		}
-	}
-	return IterationInterface;
+	return Instance->FindIterationInterface(SimulationStageIndex);
 }
 
-void NiagaraEmitterInstanceBatcher::PreStageInterface(const FNiagaraGPUSystemTick& Tick, FNiagaraComputeInstanceData *Instance, FRHICommandList &RHICmdList, const FNiagaraShaderRef& ComputeShader, const uint32 ShaderStageIndex) const
+void NiagaraEmitterInstanceBatcher::PreStageInterface(const FNiagaraGPUSystemTick& Tick, FNiagaraComputeInstanceData *Instance, FRHICommandList &RHICmdList, const FNiagaraShaderRef& ComputeShader, const uint32 SimulationStageIndex) const
 {
 	uint32 InterfaceIndex = 0;
 	for (FNiagaraDataInterfaceProxy* Interface : Instance->DataInterfaceProxies)
@@ -261,16 +245,16 @@ void NiagaraEmitterInstanceBatcher::PreStageInterface(const FNiagaraGPUSystemTic
 			TmpContext.DataInterface = Interface;
 			TmpContext.SystemInstance = Tick.SystemInstanceID;
 			TmpContext.Batcher = this;
-			TmpContext.ShaderStageIndex = ShaderStageIndex;
-			TmpContext.IsOutputStage = Interface->IsOutputStage(ShaderStageIndex);
-			TmpContext.IsIterationStage = Interface->IsIterationStage(ShaderStageIndex);
+			TmpContext.SimulationStageIndex = SimulationStageIndex;
+			TmpContext.IsOutputStage = Instance->IsOutputStage(Interface, SimulationStageIndex);
+			TmpContext.IsIterationStage = Instance->IsIterationStage(Interface, SimulationStageIndex);
 			Interface->PreStage(RHICmdList, TmpContext);
 		}
 		InterfaceIndex++;
 	}
 }
 
-void NiagaraEmitterInstanceBatcher::PostStageInterface(const FNiagaraGPUSystemTick& Tick, FNiagaraComputeInstanceData *Instance, FRHICommandList &RHICmdList, const FNiagaraShaderRef& ComputeShader, const uint32 ShaderStageIndex) const
+void NiagaraEmitterInstanceBatcher::PostStageInterface(const FNiagaraGPUSystemTick& Tick, FNiagaraComputeInstanceData *Instance, FRHICommandList &RHICmdList, const FNiagaraShaderRef& ComputeShader, const uint32 SimulationStageIndex) const
 {
 	uint32 InterfaceIndex = 0;
 	for (FNiagaraDataInterfaceProxy* Interface : Instance->DataInterfaceProxies)
@@ -283,9 +267,9 @@ void NiagaraEmitterInstanceBatcher::PostStageInterface(const FNiagaraGPUSystemTi
 			TmpContext.DataInterface = Interface;
 			TmpContext.SystemInstance = Tick.SystemInstanceID;
 			TmpContext.Batcher = this;
-			TmpContext.ShaderStageIndex = ShaderStageIndex;
-			TmpContext.IsOutputStage = Interface->IsOutputStage(ShaderStageIndex);
-			TmpContext.IsIterationStage = Interface->IsIterationStage(ShaderStageIndex);
+			TmpContext.SimulationStageIndex = SimulationStageIndex;
+			TmpContext.IsOutputStage = Instance->IsOutputStage(Interface, SimulationStageIndex);
+			TmpContext.IsIterationStage = Instance->IsIterationStage(Interface, SimulationStageIndex);
 			Interface->PostStage(RHICmdList, TmpContext);
 		}
 		InterfaceIndex++;
@@ -298,32 +282,30 @@ void NiagaraEmitterInstanceBatcher::DispatchMultipleStages(const FNiagaraGPUSyst
 
 	if (!ResetDataInterfaces(Tick, Instance, RHICmdList, ComputeShader)) return;
 
-	static const auto UseShaderStagesCVar = IConsoleManager::Get().FindConsoleVariable(TEXT("fx.UseShaderStages"));
-	if (UseShaderStagesCVar->GetInt() == 1)
+	if (Tick.NumInstancesWithSimStages > 0)
 	{
-
 		bool HasRunParticleStage = false;
-
-		FNiagaraDataBuffer* CurrentData = Instance->CurrentData;
-		FNiagaraDataBuffer* DestinationData = Instance->DestinationData;
-
+		
 		const uint32 NumStages = Instance->Context->MaxUpdateIterations;
-		const uint32 DefaultShaderStageIndex = Instance->Context->DefaultShaderStageIndex;
+		const uint32 DefaultSimulationStageIndex = Instance->Context->DefaultSimulationStageIndex;
 		bool bTransitionCurrentBuffer = false;
 
-		for (uint32 ShaderStageIndex = 0; ShaderStageIndex < NumStages; ++ShaderStageIndex)
+		for (uint32 SimulationStageIndex = 0; SimulationStageIndex < NumStages; ++SimulationStageIndex)
 		{
 			// Determine if the iteration is outputting to a custom data size
-			FNiagaraDataInterfaceProxy *IterationInterface = FindIterationInterface(Instance, ShaderStageIndex);
+			FNiagaraDataInterfaceProxy *IterationInterface = Instance->SimStageData[SimulationStageIndex].AlternateIterationSource;
+
+			//UE_LOG(LogNiagara, Log, TEXT("Starting sim stage %d Iteration %p. %s %s"), SimulationStageIndex, IterationInterface, Tick.bNeedsReset? TEXT("bNeedsReset") : TEXT("!bNeedsReset"), Context->GetDebugSimName());
 
 			if (IterationInterface && Context->SpawnStages.Num() > 0 &&
-				((Tick.bNeedsReset && !Context->SpawnStages.Contains(ShaderStageIndex)) ||
-				(!Tick.bNeedsReset && Context->SpawnStages.Contains(ShaderStageIndex))))
+				((Tick.bNeedsReset && !Context->SpawnStages.Contains(SimulationStageIndex)) ||
+				(!Tick.bNeedsReset && Context->SpawnStages.Contains(SimulationStageIndex))))
 			{
+				//UE_LOG(LogNiagara, Log, TEXT("Skipping sim stage %d  because iteration interface and spawn stage not on a reset. %s"), SimulationStageIndex, Context->GetDebugSimName());
 				continue;
 			}
 
-			PreStageInterface(Tick, Instance, RHICmdList, ComputeShader, ShaderStageIndex);
+			PreStageInterface(Tick, Instance, RHICmdList, ComputeShader, SimulationStageIndex);
 
 			// If we are reading from current data we need to transition the resource if it was previously written to
 			if (bTransitionCurrentBuffer && (ComputeShader->FloatInputBufferParam.IsBound() || ComputeShader->IntInputBufferParam.IsBound()))
@@ -331,13 +313,13 @@ void NiagaraEmitterInstanceBatcher::DispatchMultipleStages(const FNiagaraGPUSyst
 				bTransitionCurrentBuffer = false;
 
 				TArray<FRHIUnorderedAccessView*, TInlineAllocator<2>> Resources;
-				if (Instance->CurrentData->GetGPUBufferFloat().UAV.IsValid())
+				if (Instance->SimStageData[SimulationStageIndex].Source->GetGPUBufferFloat().UAV.IsValid())
 				{
-					Resources.Add(Instance->CurrentData->GetGPUBufferFloat().UAV);
+					Resources.Add(Instance->SimStageData[SimulationStageIndex].Source->GetGPUBufferFloat().UAV);
 				}
-				if (Instance->CurrentData->GetGPUBufferInt().UAV.IsValid())
+				if (Instance->SimStageData[SimulationStageIndex].Source->GetGPUBufferInt().UAV.IsValid())
 				{
-					Resources.Add(Instance->CurrentData->GetGPUBufferInt().UAV);
+					Resources.Add(Instance->SimStageData[SimulationStageIndex].Source->GetGPUBufferInt().UAV);
 				}
 				if (Resources.Num() > 0)
 				{
@@ -347,28 +329,24 @@ void NiagaraEmitterInstanceBatcher::DispatchMultipleStages(const FNiagaraGPUSyst
 
 			if (!IterationInterface)
 			{
-				Run(Tick, Instance, 0, Instance->DestinationData->GetNumInstances(), ComputeShader, RHICmdList, ViewUniformBuffer, Instance->SpawnInfo, false, DefaultShaderStageIndex, ShaderStageIndex,  nullptr, HasRunParticleStage);
+				Run(Tick, Instance, 0, Instance->SimStageData[SimulationStageIndex].Destination->GetNumInstances(), ComputeShader, RHICmdList, ViewUniformBuffer, Instance->SpawnInfo, false, DefaultSimulationStageIndex, SimulationStageIndex,  nullptr, HasRunParticleStage);
 				HasRunParticleStage = true;
 				bTransitionCurrentBuffer = true;
-
-				FNiagaraDataBuffer* StoreData = Instance->CurrentData;
-				Instance->CurrentData = Instance->DestinationData;
-				Instance->DestinationData = StoreData;
 			}
 			else
 			{
 				// run with correct number of instances.  This will make curr data junk or empty
-				Run(Tick, Instance, 0, IterationInterface->ElementCount, ComputeShader, RHICmdList, ViewUniformBuffer, Instance->SpawnInfo, false, DefaultShaderStageIndex, ShaderStageIndex, IterationInterface);
+				Run(Tick, Instance, 0, IterationInterface->ElementCount, ComputeShader, RHICmdList, ViewUniformBuffer, Instance->SpawnInfo, false, DefaultSimulationStageIndex, SimulationStageIndex, IterationInterface);
 			}
-			PostStageInterface(Tick, Instance, RHICmdList, ComputeShader, ShaderStageIndex);
+			PostStageInterface(Tick, Instance, RHICmdList, ComputeShader, SimulationStageIndex);
 		}
-		Instance->CurrentData = CurrentData;
-		Instance->DestinationData = DestinationData;
+
 	}
 	else
 	{
 		// run shader, sim and spawn in a single dispatch
-		Run(Tick, Instance, 0, Instance->DestinationData->GetNumInstances(), ComputeShader, RHICmdList, ViewUniformBuffer, Instance->SpawnInfo);
+		check(Instance->SimStageData.Num() > 0);
+		Run(Tick, Instance, 0, Instance->SimStageData[0].Destination->GetNumInstances(), ComputeShader, RHICmdList, ViewUniformBuffer, Instance->SpawnInfo);
 	}
 }
 
@@ -379,6 +357,8 @@ void NiagaraEmitterInstanceBatcher::ResizeBuffersAndGatherResources(FOverlappabl
 	//UE_LOG(LogNiagara, Warning, TEXT("NiagaraEmitterInstanceBatcher::ResizeBuffersAndGatherResources:  %0xP"), this);
 	for (FNiagaraGPUSystemTick* Tick : OverlappableTick)
 	{
+		//UE_LOG(LogNiagara, Warning, TEXT("NiagaraEmitterInstanceBatcher::ResizeBuffersAndGatherResources Tick:  %p  Count: %d"), Tick, Tick->Count);
+
 		const uint32 DispatchCount = Tick->Count;
 		const bool bIsFinalTick = Tick->bIsFinalTick;
 		const bool bNeedsReset = Tick->bNeedsReset;
@@ -401,16 +381,18 @@ void NiagaraEmitterInstanceBatcher::ResizeBuffersAndGatherResources(FOverlappabl
 
 			const bool bRequiresPersistentIDs = Context->MainDataSet->RequiresPersistentIDs();
 
+			check(Instance.SimStageData.Num() == Context->MaxUpdateIterations);
+
 			//The buffer containing current simulation state.
-			Instance.CurrentData = Context->MainDataSet->GetCurrentData();
+			Instance.SimStageData[0].Source = Context->MainDataSet->GetCurrentData();
 			//The buffer we're going to write simulation results to.
-			Instance.DestinationData = &Context->MainDataSet->BeginSimulate();
+			Instance.SimStageData[0].Destination = &Context->MainDataSet->BeginSimulate();
 
-			check(Instance.CurrentData && Instance.DestinationData);
-			FNiagaraDataBuffer& CurrentData = *Instance.CurrentData;
-			FNiagaraDataBuffer& DestinationData = *Instance.DestinationData;
+			check(Instance.SimStageData[0].Source && Instance.SimStageData[0].Destination);
+			FNiagaraDataBuffer* CurrentData = Instance.SimStageData[0].Source;
+			FNiagaraDataBuffer* DestinationData = Instance.SimStageData[0].Destination;
 
-			const uint32 PrevNumInstances = bNeedsReset ? 0 : CurrentData.GetNumInstances();
+			const uint32 PrevNumInstances = bNeedsReset ? 0 : CurrentData->GetNumInstances();
 			const uint32 NewNumInstances = Instance.SpawnInfo.SpawnRateInstances + Instance.SpawnInfo.EventSpawnTotal + PrevNumInstances;
 
 			//We must assume all particles survive when allocating here. 
@@ -425,44 +407,121 @@ void NiagaraEmitterInstanceBatcher::ResizeBuffersAndGatherResources(FOverlappabl
 				InstancesWithPersistentIDs.Add(&Instance);
 			}
 
-			DestinationData.AllocateGPU(AllocatedInstances + 1, GPUInstanceCounterManager, RHICmdList, FeatureLevel, Context->GetDebugSimName());
-			DestinationData.SetNumInstances(RequiredInstances);
+			DestinationData->AllocateGPU(AllocatedInstances + 1, GPUInstanceCounterManager, RHICmdList, FeatureLevel, Context->GetDebugSimName());
+			DestinationData->SetNumInstances(RequiredInstances);
+			Instance.SimStageData[0].SourceCountOffset = Instance.SimStageData[0].Source->GetGPUInstanceCountBufferOffset();
+			if (Instance.SimStageData[0].SourceCountOffset == INDEX_NONE) // It is possible that this has been queued for readback, taking ownership of the data. Use that instead.
+			{
+				Instance.SimStageData[0].SourceCountOffset = Context->EmitterInstanceReadback.GPUCountOffset;
+			}
+			Instance.SimStageData[0].DestinationCountOffset = Instance.SimStageData[0].Destination->GetGPUInstanceCountBufferOffset();
 
 			if ( Shader->FloatInputBufferParam.IsBound() )
 			{
-				ReadBuffers.Add(CurrentData.GetGPUBufferFloat().UAV);
+				ReadBuffers.Add(CurrentData->GetGPUBufferFloat().UAV);
 			}
 			if ( Shader->IntInputBufferParam.IsBound() )
 			{
-				ReadBuffers.Add(CurrentData.GetGPUBufferInt().UAV);
+				ReadBuffers.Add(CurrentData->GetGPUBufferInt().UAV);
 			}
 
 			if ( Shader->FloatOutputBufferParam.IsBound() )
 			{
-				WriteBuffers.Add(DestinationData.GetGPUBufferFloat().UAV);
-				if (bIsFinalTick)
-				{
-					OutputGraphicsBuffers.Add(DestinationData.GetGPUBufferFloat().UAV);
-				}
+				WriteBuffers.Add(DestinationData->GetGPUBufferFloat().UAV);
 			}
 			if ( Shader->IntOutputBufferParam.IsBound() )
 			{
-				WriteBuffers.Add(DestinationData.GetGPUBufferInt().UAV);
-				if (bIsFinalTick)
-				{
-					OutputGraphicsBuffers.Add(DestinationData.GetGPUBufferInt().UAV);
-				}
+				WriteBuffers.Add(DestinationData->GetGPUBufferInt().UAV);
 			}
 
+			//UE_LOG(LogScript, Warning, TEXT("ResizeBuffersAndGatherResources [%d][%d] Run ReqInst: %d Cur: %p Dest: %p "), Index, 0, RequiredInstances, Instance.SimStageData[0].Source, Instance.SimStageData[0].Destination);
 			if (bRequiresPersistentIDs)
 			{
-				WriteBuffers.Add(DestinationData.GetGPUIDToIndexTable().UAV);
+				WriteBuffers.Add(DestinationData->GetGPUIDToIndexTable().UAV);
 			}
 
+
 			Context->MainDataSet->EndSimulate();
+
+			// Go ahead and reserve the readback data...
+			//uint32 ComputeCountOffsetOverride = INDEX_NONE;
+			if (!GPUInstanceCounterManager.HasPendingGPUReadback() && Tick->bIsFinalTick)
+			{
+				// Now that the current data is not required anymore, stage it for readback.
+				if (CurrentData->GetNumInstances() && Context->EmitterInstanceReadback.GPUCountOffset == INDEX_NONE && CurrentData->GetGPUInstanceCountBufferOffset() != INDEX_NONE)
+				{
+					// Transfer the GPU instance counter ownership to the context. Note that a readback request will be performed later in the tick update, unless there's already a pending readback.
+					Context->EmitterInstanceReadback.GPUCountOffset = CurrentData->GetGPUInstanceCountBufferOffset();
+					Context->EmitterInstanceReadback.CPUCount = CurrentData->GetNumInstances();
+					CurrentData->ClearGPUInstanceCountBufferOffset();
+
+					//UE_LOG(LogNiagara, Log, TEXT("EmitterInstanceReadback.CPUCount dispatch %d  Offset: %d"), Context->EmitterInstanceReadback.CPUCount, Context->EmitterInstanceReadback.GPUCountOffset);
+				}
+			}
+			
+			uint32 NumBufferIterations = 1;
+			if (Tick->NumInstancesWithSimStages > 0)
+			{
+
+				bool HasRunParticleStage = false;
+
+				const uint32 NumStages = Instance.Context->MaxUpdateIterations;
+				if (NumStages > 1)
+				{
+					HasRunParticleStage = true;
+					for (uint32 SimulationStageIndex = 0; SimulationStageIndex < NumStages; SimulationStageIndex++)
+					{
+						if (SimulationStageIndex != 0)
+						{
+							Instance.SimStageData[SimulationStageIndex].Source = Instance.SimStageData[SimulationStageIndex - 1].Source;
+							Instance.SimStageData[SimulationStageIndex].Destination = Instance.SimStageData[SimulationStageIndex - 1].Destination;
+
+							Instance.SimStageData[SimulationStageIndex].SourceCountOffset = Instance.SimStageData[SimulationStageIndex - 1].SourceCountOffset;
+							Instance.SimStageData[SimulationStageIndex].DestinationCountOffset = Instance.SimStageData[SimulationStageIndex - 1].DestinationCountOffset;
+						}
+
+						// Determine if the iteration is outputting to a custom data size
+						FNiagaraDataInterfaceProxy* IterationInterface = FindIterationInterface(&Instance, SimulationStageIndex);
+
+						Instance.SimStageData[SimulationStageIndex].AlternateIterationSource = IterationInterface;
+
+						if (IterationInterface && Context->SpawnStages.Num() > 0 &&
+							((Tick->bNeedsReset && !Context->SpawnStages.Contains(SimulationStageIndex)) ||
+							(!Tick->bNeedsReset && Context->SpawnStages.Contains(SimulationStageIndex))))
+						{
+							continue;
+						}
+
+						if (!IterationInterface && SimulationStageIndex != 0)
+						{
+							// Go ahead and grab the write buffer, which may be too small, so make sure to resize it.
+							Instance.SimStageData[SimulationStageIndex].Source = Context->MainDataSet->GetCurrentData();
+							DestinationData = &Context->MainDataSet->BeginSimulate(false);
+							Instance.SimStageData[SimulationStageIndex].Destination = DestinationData;
+							DestinationData->AllocateGPU(AllocatedInstances + 1, GPUInstanceCounterManager, RHICmdList, FeatureLevel, Context->GetDebugSimName());
+							DestinationData->SetNumInstances(RequiredInstances);
+							Instance.SimStageData[SimulationStageIndex].SourceCountOffset = Instance.SimStageData[SimulationStageIndex].Source->GetGPUInstanceCountBufferOffset();
+							Instance.SimStageData[SimulationStageIndex].DestinationCountOffset = Instance.SimStageData[SimulationStageIndex].Destination->GetGPUInstanceCountBufferOffset();
+							
+							//UE_LOG(LogScript, Warning, TEXT("ResizeBuffersAndGatherResources [%d][%d] Run  ReqInst: %d Cur: %p Dest: %p "), Index, SimulationStageIndex, RequiredInstances, Instance.SimStageData[SimulationStageIndex].Source, Instance.SimStageData[SimulationStageIndex].Destination);
+
+							// We don't actually write we just map out the buffers here. This toggles src and dest...
+							Context->MainDataSet->EndSimulate();
+						}
+
+
+					}
+				}
+			}
+
+			CurrentData = Context->MainDataSet->GetCurrentData();
 			if (bIsFinalTick)
 			{
-				Context->SetDataToRender(Instance.DestinationData);
+				//UE_LOG(LogScript, Warning, TEXT("ResizeBuffersAndGatherResources [%d] DataSetToRender %p "),Index, CurrentData);
+
+				Context->SetDataToRender(CurrentData);
+				OutputGraphicsBuffers.Add(CurrentData->GetGPUBufferFloat().UAV);
+				OutputGraphicsBuffers.Add(CurrentData->GetGPUBufferInt().UAV);
 			}
 		}
 	}
@@ -480,7 +539,7 @@ void NiagaraEmitterInstanceBatcher::ResizeBuffersAndGatherResources(FOverlappabl
 		for (uint32 i = 0; i < NumInstancesWithPersistentIDs; ++i)
 		{
 			FNiagaraComputeInstanceData* Instance = InstancesWithPersistentIDs[i];
-			IDToIndexTables[i] = Instance->DestinationData->GetGPUIDToIndexTable().UAV;
+			IDToIndexTables[i] = Instance->SimStageData[0].Destination->GetGPUIDToIndexTable().UAV;
 		}
 		// TODO: is it sufficient to do a CS cache flush before all this and get rid of these explicit barriers?
 		RHICmdList.TransitionResources(EResourceTransitionAccess::EWritable, EResourceTransitionPipeline::EComputeToCompute, IDToIndexTables.GetData(), IDToIndexTables.Num());
@@ -489,7 +548,7 @@ void NiagaraEmitterInstanceBatcher::ResizeBuffersAndGatherResources(FOverlappabl
 		{
 			FNiagaraComputeInstanceData* Instance = InstancesWithPersistentIDs[i];
 			SCOPED_DRAW_EVENTF(RHICmdList, NiagaraGPUComputeClearIDToIndexBuffer, TEXT("Clear ID To Index Table - %s"), Instance->Context->GetDebugSimName());
-			NiagaraFillGPUIntBuffer(RHICmdList, FeatureLevel, Instance->DestinationData->GetGPUIDToIndexTable(), INDEX_NONE);
+			NiagaraFillGPUIntBuffer(RHICmdList, FeatureLevel, Instance->SimStageData[0].Destination->GetGPUIDToIndexTable(), INDEX_NONE);
 		}
 	}
 }
@@ -540,19 +599,6 @@ void NiagaraEmitterInstanceBatcher::DispatchAllOnCompute(FOverlappableTicks& Ove
 
 				// run shader, sim and spawn in a single dispatch
 				DispatchMultipleStages(*Tick, &Instance, RHICmdList, ViewUniformBuffer, Context->GPUScript_RT->GetShader());
-
-				FNiagaraDataBuffer* CurrentData = Instance.CurrentData;
-				if (!GPUInstanceCounterManager.HasPendingGPUReadback() && Tick->bIsFinalTick)
-				{
-					// Now that the current data is not required anymore, stage it for readback.
-					if (CurrentData->GetNumInstances() && Context->EmitterInstanceReadback.GPUCountOffset == INDEX_NONE && CurrentData->GetGPUInstanceCountBufferOffset() != INDEX_NONE)
-					{
-						// Transfer the GPU instance counter ownership to the context. Note that a readback request will be performed later in the tick update, unless there's already a pending readback.
-						Context->EmitterInstanceReadback.GPUCountOffset = CurrentData->GetGPUInstanceCountBufferOffset();
-						Context->EmitterInstanceReadback.CPUCount = CurrentData->GetNumInstances();
-						CurrentData->ClearGPUInstanceCountBufferOffset();
-					}
-				}
 			}
 		}
 	}
@@ -655,7 +701,7 @@ void NiagaraEmitterInstanceBatcher::UpdateFreeIDBuffers(FRHICommandList& RHICmdL
 	FNiagaraBufferArray ReadBuffers, WriteBuffers;
 	for(FNiagaraComputeInstanceData* Instance : Instances)
 	{
-		ReadBuffers.Add(Instance->DestinationData->GetGPUIDToIndexTable().UAV);
+		ReadBuffers.Add(Instance->SimStageData[0].Destination->GetGPUIDToIndexTable().UAV);
 		WriteBuffers.Add(Instance->Context->MainDataSet->GetGPUFreeIDs().UAV);
 	}
 
@@ -669,7 +715,7 @@ void NiagaraEmitterInstanceBatcher::UpdateFreeIDBuffers(FRHICommandList& RHICmdL
 	{
 		FNiagaraComputeInstanceData* Instance = Instances[InstanceIdx];
 		FNiagaraDataSet* MainDataSet = Instance->Context->MainDataSet;
-		FNiagaraDataBuffer* DestinationData = Instance->DestinationData;
+		FNiagaraDataBuffer* DestinationData = Instance->SimStageData[0].Destination;
 
 		SCOPED_DRAW_EVENTF(RHICmdList, NiagaraGPUComputeFreeIDsEmitter, TEXT("Update Free ID Buffer - %s"), Instance->Context->GetDebugSimName());
 		NiagaraComputeGPUFreeIDs(RHICmdList, FeatureLevel, MainDataSet->GetGPUNumAllocatedIDs(), DestinationData->GetGPUIDToIndexTable().SRV, MainDataSet->GetGPUFreeIDs(), FreeIDListSizesBuffer, InstanceIdx);
@@ -842,7 +888,7 @@ void NiagaraEmitterInstanceBatcher::PreInitViews(FRHICommandListImmediate& RHICm
 		int32 TotalDispatchCount = 0;
 		for (FNiagaraGPUSystemTick& Tick : Ticks_RT)
 		{
-			TotalDispatchCount += (int32)Tick.Count;
+			TotalDispatchCount += (int32)Tick.TotalDispatches;
 
 			// Cancel any pending readback if the emitter is resetting.
 			if (Tick.bNeedsReset)
@@ -884,6 +930,7 @@ void NiagaraEmitterInstanceBatcher::PreInitViews(FRHICommandListImmediate& RHICm
 								if (DeadInstanceCount <= CurrentData->GetNumInstances())
 								{
 									CurrentData->SetNumInstances(CurrentData->GetNumInstances() - DeadInstanceCount); 
+									//UE_LOG(LogNiagara, Log, TEXT("GPU Readback Offset: %d %p = %d"), Context->EmitterInstanceReadback.GPUCountOffset, CurrentData, CurrentData->GetNumInstances());
 								}
 							}
 
@@ -1108,7 +1155,7 @@ void NiagaraEmitterInstanceBatcher::ProcessDebugInfo(FRHICommandList &RHICmdList
 
 /* Set shader parameters for data interfaces
  */
-void NiagaraEmitterInstanceBatcher::SetDataInterfaceParameters(const TArray<FNiagaraDataInterfaceProxy*>& DataInterfaceProxies, const FNiagaraShaderRef& Shader, FRHICommandList &RHICmdList, const FNiagaraComputeInstanceData* Instance, const FNiagaraGPUSystemTick& Tick, uint32 ShaderStageIndex) const
+void NiagaraEmitterInstanceBatcher::SetDataInterfaceParameters(const TArray<FNiagaraDataInterfaceProxy*>& DataInterfaceProxies, const FNiagaraShaderRef& Shader, FRHICommandList& RHICmdList, const FNiagaraComputeInstanceData* Instance, const FNiagaraGPUSystemTick& Tick, uint32 SimulationStageIndex) const
 {
 	// set up data interface buffers, as defined by the DIs during compilation
 	//
@@ -1129,8 +1176,8 @@ void NiagaraEmitterInstanceBatcher::SetDataInterfaceParameters(const TArray<FNia
 			Context.SystemInstance = SystemInstance;
 			Context.Batcher = this;
 			Context.ComputeInstanceData = Instance;
-			Context.ShaderStageIndex = ShaderStageIndex;
-			Context.IsOutputStage = Interface->IsOutputStage(ShaderStageIndex);
+			Context.SimulationStageIndex = SimulationStageIndex;
+			Context.IsOutputStage = Instance->IsOutputStage(Interface, SimulationStageIndex);
 			DIParam.DIType.Get(PointerTable.DITypes)->SetParameters(DIParam.Parameters.Get(), RHICmdList, Context);
 		}
 
@@ -1191,26 +1238,34 @@ static void SetConstantBuffer(FRHICommandList &RHICmdList, FRHIComputeShader* Co
 /* Kick off a simulation/spawn run
  */
 void NiagaraEmitterInstanceBatcher::Run(const FNiagaraGPUSystemTick& Tick, const FNiagaraComputeInstanceData* Instance, uint32 UpdateStartInstance, const uint32 TotalNumInstances, const FNiagaraShaderRef& Shader,
-	FRHICommandList &RHICmdList, FRHIUniformBuffer* ViewUniformBuffer, const FNiagaraGpuSpawnInfo& SpawnInfo, bool bCopyBeforeStart, uint32 DefaultShaderStageIndex, uint32 ShaderStageIndex, FNiagaraDataInterfaceProxy *IterationInterface, bool HasRunParticleStage)
+	FRHICommandList &RHICmdList, FRHIUniformBuffer* ViewUniformBuffer, const FNiagaraGpuSpawnInfo& SpawnInfo, bool bCopyBeforeStart, uint32 DefaultSimulationStageIndex, uint32 SimulationStageIndex, FNiagaraDataInterfaceProxy *IterationInterface, bool HasRunParticleStage)
 {
 	FNiagaraComputeExecutionContext* Context = Instance->Context;
 
-	SCOPED_DRAW_EVENTF(RHICmdList, NiagaraGPUSimulationCS, TEXT("Niagara Gpu Sim - %s - NumInstances: %u - StageNumber: %u"),
-		Context->GetDebugSimName(),
-		TotalNumInstances,
-		ShaderStageIndex);
 
 	if (TotalNumInstances == 0)
 	{
 		return;
 	}
 
+	/*UE_LOG(LogNiagara, Log, TEXT("Niagara Gpu Sim - % s - NumInstances: % u - StageNumber : % u"), Context->GetDebugSimName(),
+		TotalNumInstances,
+		SimulationStageIndex);
+		*/
+
+	SCOPED_DRAW_EVENTF(RHICmdList, NiagaraGPUSimulationCS, TEXT("Niagara Gpu Sim - %s - NumInstances: %u - StageNumber: %u"),
+		Context->GetDebugSimName(),
+		TotalNumInstances,
+		SimulationStageIndex);
+
 	//UE_LOG(LogNiagara, Warning, TEXT("Run"));
 	
 	const TArray<FNiagaraDataInterfaceProxy*>& DataInterfaceProxies = Instance->DataInterfaceProxies;
-	check(Instance->CurrentData && Instance->DestinationData);
-	FNiagaraDataBuffer& DestinationData = *Instance->DestinationData;
-	FNiagaraDataBuffer& CurrentData = *Instance->CurrentData;
+	check(Instance->SimStageData[SimulationStageIndex].Source && Instance->SimStageData[SimulationStageIndex].Destination);
+	FNiagaraDataBuffer& DestinationData = *Instance->SimStageData[SimulationStageIndex].Destination;
+	FNiagaraDataBuffer& CurrentData = *Instance->SimStageData[SimulationStageIndex].Source;
+
+	//UE_LOG(LogScript, Warning, TEXT("Run [%d] TotalInstances %d  src:%p dest:%p"), SimulationStageIndex, TotalNumInstances, Instance->SimStageData[SimulationStageIndex].Source, Instance->SimStageData[SimulationStageIndex].Destination);
 
 	int32 InstancesToSpawnThisFrame = Instance->SpawnInfo.SpawnRateInstances + Instance->SpawnInfo.EventSpawnTotal;
 
@@ -1235,7 +1290,7 @@ void NiagaraEmitterInstanceBatcher::Run(const FNiagaraGPUSystemTick& Tick, const
 		RHICmdList.SetShaderUniformBuffer(ComputeShader, Shader->ViewUniformBufferParam.GetBaseIndex(), ViewUniformBuffer);
 	}
 
-	SetDataInterfaceParameters(DataInterfaceProxies, Shader, RHICmdList, Instance, Tick, ShaderStageIndex);
+	SetDataInterfaceParameters(DataInterfaceProxies, Shader, RHICmdList, Instance, Tick, SimulationStageIndex);
 
 	// set the shader and data set params 
 	//
@@ -1250,8 +1305,9 @@ void NiagaraEmitterInstanceBatcher::Run(const FNiagaraGPUSystemTick& Tick, const
 	{
 		RHICmdList.TransitionResource(EResourceTransitionAccess::ERWNoBarrier, EResourceTransitionPipeline::EComputeToCompute, GPUInstanceCounterManager.GetInstanceCountBuffer().UAV);
 		Shader->InstanceCountsParam.SetBuffer(RHICmdList, ComputeShader, GPUInstanceCounterManager.GetInstanceCountBuffer());
-		const uint32 ReadOffset = Tick.bNeedsReset ? INDEX_NONE : CurrentData.GetGPUInstanceCountBufferOffset();
-		const uint32 WriteOffset = DestinationData.GetGPUInstanceCountBufferOffset();
+		const uint32 ReadOffset = (Tick.bNeedsReset && SimulationStageIndex == 0) ? INDEX_NONE : Instance->SimStageData[SimulationStageIndex].SourceCountOffset;
+		const uint32 WriteOffset = Instance->SimStageData[SimulationStageIndex].DestinationCountOffset;
+		//UE_LOG(LogNiagara, Log, TEXT("Instance count setup R: %d W: %d reset? %s %d"), ReadOffset, WriteOffset, Tick.bNeedsReset ? TEXT("T") : TEXT("F"), CurrentData.GetGPUInstanceCountBufferOffset());
 		SetShaderValue(RHICmdList, ComputeShader, Shader->ReadInstanceCountOffsetParam, ReadOffset);
 		SetShaderValue(RHICmdList, ComputeShader, Shader->WriteInstanceCountOffsetParam, WriteOffset);
 	}
@@ -1269,8 +1325,8 @@ void NiagaraEmitterInstanceBatcher::Run(const FNiagaraGPUSystemTick& Tick, const
 
 	SetShaderValue(RHICmdList, ComputeShader, Shader->UpdateStartInstanceParam, UpdateStartInstance);					// 0, except for event handler runs
 	SetShaderValue(RHICmdList, ComputeShader, Shader->NumSpawnedInstancesParam, InstancesToSpawnThisFrame);				// number of instances in the spawn run
-	SetShaderValue(RHICmdList, ComputeShader, Shader->DefaultShaderStageIndexParam, DefaultShaderStageIndex);					// 0, except if several stages are defined
-	SetShaderValue(RHICmdList, ComputeShader, Shader->ShaderStageIndexParam, ShaderStageIndex);					// 0, except if several stages are defined
+	SetShaderValue(RHICmdList, ComputeShader, Shader->DefaultSimulationStageIndexParam, DefaultSimulationStageIndex);					// 0, except if several stages are defined
+	SetShaderValue(RHICmdList, ComputeShader, Shader->SimulationStageIndexParam, SimulationStageIndex);					// 0, except if several stages are defined
 	const int32 DefaultIterationCount = -1;
 	SetShaderValue(RHICmdList, ComputeShader, Shader->IterationInterfaceCount, DefaultIterationCount);					// 0, except if several stages are defined
 
@@ -1315,7 +1371,7 @@ void NiagaraEmitterInstanceBatcher::Run(const FNiagaraGPUSystemTick& Tick, const
 	}
 
 	//UE_LOG(LogNiagara, Log, TEXT("Num Instance : %d | Num Group : %d | Spawned Istance : %d | Start Instance : %d | Num Indices : %d | Stage Index : %d"), 
-		//TotalNumInstances, NumThreadGroups, InstancesToSpawnThisFrame, UpdateStartInstance, Context->NumIndicesPerInstance, ShaderStageIndex);
+		//TotalNumInstances, NumThreadGroups, InstancesToSpawnThisFrame, UpdateStartInstance, Context->NumIndicesPerInstance, SimulationStageIndex);
 
 	// Dispatch, if anything needs to be done
 	if (TotalNumInstances)
@@ -1335,7 +1391,7 @@ void NiagaraEmitterInstanceBatcher::Run(const FNiagaraGPUSystemTick& Tick, const
 	{
 		//UE_LOG(LogNiagara, Warning, TEXT("Queued up!"));
 
-		if (!Context->GPUDebugDataReadbackFloat && !Context->GPUDebugDataReadbackInt && !Context->GPUDebugDataReadbackCounts && DestinationData.GetGPUInstanceCountBufferOffset() != INDEX_NONE && ShaderStageIndex == Context->MaxUpdateIterations - 1)
+		if (!Context->GPUDebugDataReadbackFloat && !Context->GPUDebugDataReadbackInt && !Context->GPUDebugDataReadbackCounts && DestinationData.GetGPUInstanceCountBufferOffset() != INDEX_NONE && SimulationStageIndex == Context->MaxUpdateIterations - 1)
 		{
 			Context->GPUDebugDataFloatSize = 0;
 			Context->GPUDebugDataIntSize = 0;
diff --git a/Engine/Plugins/FX/Niagara/Source/Niagara/Private/NiagaraGPUInstanceCountManager.cpp b/Engine/Plugins/FX/Niagara/Source/Niagara/Private/NiagaraGPUInstanceCountManager.cpp
--- a/Engine/Plugins/FX/Niagara/Source/Niagara/Private/NiagaraGPUInstanceCountManager.cpp
+++ b/Engine/Plugins/FX/Niagara/Source/Niagara/Private/NiagaraGPUInstanceCountManager.cpp
@@ -75,7 +75,7 @@ uint32 FNiagaraGPUInstanceCountManager::AcquireEntry()
 	{
 		// @TODO : add realloc the buffer and copy the current content to it. Might require reallocating the readback in FNiagaraGPUInstanceCountManager::EnqueueGPUReadback()
 		ensure(UsedInstanceCounts < AllocatedInstanceCounts);
-		UE_LOG(LogNiagara, Error, TEXT("Niagara.MinGPUInstanceCount too small."));
+		UE_LOG(LogNiagara, Error, TEXT("Niagara.MinGPUInstanceCount too small. UsedInstanceCounts: %d < AllocatedInstanceCounts: %d"), UsedInstanceCounts, AllocatedInstanceCounts);
 		return INDEX_NONE;
 	}
 }
@@ -86,8 +86,13 @@ void FNiagaraGPUInstanceCountManager::FreeEntry(uint32& BufferOffset)
 
 	if (BufferOffset != INDEX_NONE)
 	{
+		//UE_LOG(LogNiagara, Warning, TEXT("FNiagaraGPUInstanceCountManager::FreeEntry %d"), BufferOffset);
 		// Add a reset to 0 task.
 		// The entry will only become available/reusable after being reset to 0 in UpdateDrawIndirectBuffer()
+		//if (InstanceCountClearTasks.Find(BufferOffset) != -1)
+		//{
+		//	UE_LOG(LogNiagara, Warning, TEXT("FNiagaraGPUInstanceCountManager::FreeEntry DUPLICATED!!!!!!!!!!!!!! %d"), BufferOffset);
+		//}
 		InstanceCountClearTasks.Add(BufferOffset);
 		BufferOffset = INDEX_NONE;
 	}
@@ -107,6 +112,7 @@ void FNiagaraGPUInstanceCountManager::ResizeBuffers(FRHICommandListImmediate& RH
 			TResourceArray<uint32> InitData;
 			InitData.AddZeroed(AllocatedInstanceCounts);
 			CountBuffer.Initialize(sizeof(uint32), AllocatedInstanceCounts, EPixelFormat::PF_R32_UINT, BUF_Static | BUF_SourceCopy, TEXT("NiagaraGPUInstanceCounts"), &InitData);
+			UE_LOG(LogNiagara, Log, TEXT("FNiagaraGPUInstanceCountManager::ResizeBuffers Alloc AllocatedInstanceCounts: %d ReservedInstanceCounts: %d"), AllocatedInstanceCounts, ReservedInstanceCounts);
 		}
 		// If we need to increase the buffer size to RecommendedInstanceCounts because the buffer is too small.
 		else if (RequiredInstanceCounts > AllocatedInstanceCounts)
@@ -134,6 +140,7 @@ void FNiagaraGPUInstanceCountManager::ResizeBuffers(FRHICommandListImmediate& RH
 			// Swap the buffers
 			AllocatedInstanceCounts = RecommendedInstanceCounts;
 			FMemory::Memswap(&NextCountBuffer, &CountBuffer, sizeof(NextCountBuffer));
+			UE_LOG(LogNiagara, Log, TEXT("FNiagaraGPUInstanceCountManager::ResizeBuffers Resize AllocatedInstanceCounts: %d ReservedInstanceCounts: %d"), AllocatedInstanceCounts, ReservedInstanceCounts);
 		}
 		// If we need to shrink the buffer size because use way to much buffer size.
 		else if ((int32)(RecommendedInstanceCounts * BufferSlack) < AllocatedInstanceCounts)
@@ -199,6 +206,11 @@ void FNiagaraGPUInstanceCountManager::UpdateDrawIndirectBuffer(FRHICommandList&
 			RHICmdList.TransitionResource(EResourceTransitionAccess::ERWBarrier, EResourceTransitionPipeline::EComputeToCompute, CountBuffer.UAV);
 
 			FReadBuffer TaskInfosBuffer;
+
+			//for (int32 i = 0; i < InstanceCountClearTasks.Num(); i++)
+			//{
+			//		UE_LOG(LogNiagara, Log, TEXT("InstanceCountClearTasks[%d] = %d"), i, InstanceCountClearTasks[i]);
+			//}
 			{
 				// All draw indirect args task are run first because of the binding between the task ID and arg write offset.
 				const uint32 ArgGenSize = DrawIndirectArgGenTasks.Num() * sizeof(FArgGenTaskInfo);
@@ -229,6 +241,11 @@ void FNiagaraGPUInstanceCountManager::UpdateDrawIndirectBuffer(FRHICommandList&
 		// Once cleared to 0, the count are reusable.
 		FreeEntries.Append(InstanceCountClearTasks);
 
+		//for (int32 i = 0; i < FreeEntries.Num(); i++)
+		//{
+		//	UE_LOG(LogNiagara, Log, TEXT("FreeEntries[%d] = %d"), i, FreeEntries[i]);
+		//}
+
 		DrawIndirectArgGenTasks.Empty();
 		DrawIndirectArgMap.Empty();
 		InstanceCountClearTasks.Empty();
diff --git a/Engine/Plugins/FX/Niagara/Source/Niagara/Private/NiagaraModule.cpp b/Engine/Plugins/FX/Niagara/Source/Niagara/Private/NiagaraModule.cpp
--- a/Engine/Plugins/FX/Niagara/Source/Niagara/Private/NiagaraModule.cpp
+++ b/Engine/Plugins/FX/Niagara/Source/Niagara/Private/NiagaraModule.cpp
@@ -50,16 +50,6 @@ static FAutoConsoleVariableRef CVarEnableVerboseNiagaraChangeIdLogging(
 	ECVF_Default
 );
 
-/**
-Use Shader Stages CVar.
-Enable the custom dispatch for multiple shader stages 
-*/
-static TAutoConsoleVariable<int32> CVarUseShaderStages(
-	TEXT("fx.UseShaderStages"), 
-	0, 
-	TEXT("Enable or not the shader stages within Niagara (WIP feature only there for temporary testing)."),
-	ECVF_Default);
-
 static FAutoConsoleVariableRef CVarNiaraGlobalSystemCountScale(
 	TEXT("fx.NiagaraGlobalSystemCountScale"),
 	INiagaraModule::EngineGlobalSystemCountScale,
diff --git a/Engine/Plugins/FX/Niagara/Source/Niagara/Private/NiagaraScript.cpp b/Engine/Plugins/FX/Niagara/Source/Niagara/Private/NiagaraScript.cpp
--- a/Engine/Plugins/FX/Niagara/Source/Niagara/Private/NiagaraScript.cpp
+++ b/Engine/Plugins/FX/Niagara/Source/Niagara/Private/NiagaraScript.cpp
@@ -38,7 +38,7 @@
 
 #include "NiagaraFunctionLibrary.h"
 #include "VectorVM.h"
-
+#include "NiagaraSimulationStageBase.h"
 #include "Async/Async.h"
 
 DECLARE_STATS_GROUP(TEXT("Niagara Detailed"), STATGROUP_NiagaraDetailed, STATCAT_Advanced);
@@ -348,6 +348,36 @@ void UNiagaraScript::ComputeVMCompilationId(FNiagaraVMExecutableDataId& Id) cons
 		{
 			Id.bUsesRapidIterationParams = true;
 		}
+
+		if (Emitter->bSimulationStagesEnabled && Usage == ENiagaraScriptUsage::ParticleGPUComputeScript/* TODO limit to just with stages in the future! Leaving like this so what can convert! && Emitter->GetSimulationStages().Num() > 0*/)
+		{
+			Id.AdditionalDefines.Add(TEXT("Emitter.UseSimulationStages"));
+
+			FSHA1 HashState;
+			FNiagaraCompileHashVisitor Visitor(HashState);
+			for (UNiagaraSimulationStageBase* Base : Emitter->GetSimulationStages())
+			{
+				// bool AppendCompileHash(FNiagaraCompileHashVisitor* InVisitor) const;
+				if (Base)
+				{
+					Base->AppendCompileHash(&Visitor);
+				}
+			}
+			HashState.Final();
+
+			TArray<uint8> DataHash;
+			DataHash.AddUninitialized(FSHA1::DigestSize);
+			HashState.GetHash(DataHash.GetData());
+
+			FNiagaraCompileHash Hash(DataHash);
+			Id.ReferencedCompileHashes.Add(Hash);
+			Id.DebugReferencedObjects.Add(TEXT("SimulationStageHeaders"));
+
+		}
+		else if (Emitter->bDeprecatedShaderStagesEnabled && Usage == ENiagaraScriptUsage::ParticleGPUComputeScript/* TODO limit to just with stages in the future! Leaving like this so what can convert! && Emitter->GetSimulationStages().Num() > 0*/)
+		{
+			Id.AdditionalDefines.Add(TEXT("Emitter.UseOldShaderStages"));
+		}
 	}
 
 	if (UNiagaraSystem* System = Cast<UNiagaraSystem>(Obj))
@@ -426,13 +456,7 @@ void UNiagaraScript::ComputeVMCompilationId(FNiagaraVMExecutableDataId& Id) cons
 		Id.ReferencedCompileHashes.Add(Hash);
 		Id.DebugReferencedObjects.Add(TEXT("RIParams"));
 	}
-
-	static const auto UseShaderStagesCVar = IConsoleManager::Get().FindConsoleVariable(TEXT("fx.UseShaderStages"));
-	if (UseShaderStagesCVar->GetInt())
-	{
-		Id.AdditionalDefines.Add(TEXT("fx.UseShaderStages"));
-	}
-
+	
 	Source->ComputeVMCompilationId(Id, Usage, UsageId);	
 	
 	if (GNiagaraDumpKeyGen == 1 && Id != LastGeneratedVMId)
@@ -466,6 +490,11 @@ bool UNiagaraScript::ContainsUsage(ENiagaraScriptUsage InUsage) const
 		return true;
 	}
 
+	if (Usage == ENiagaraScriptUsage::ParticleGPUComputeScript && InUsage == ENiagaraScriptUsage::ParticleSimulationStageScript)
+	{
+		return true;
+	}
+
 	if (InUsage == ENiagaraScriptUsage::ParticleUpdateScript && Usage == ENiagaraScriptUsage::ParticleSpawnScriptInterpolated)
 	{
 		return true;
@@ -530,6 +559,7 @@ TOptional<ENiagaraSimTarget> UNiagaraScript::GetSimTarget() const
 	case ENiagaraScriptUsage::ParticleSpawnScriptInterpolated:
 	case ENiagaraScriptUsage::ParticleUpdateScript:
 	case ENiagaraScriptUsage::ParticleEventScript:
+	case ENiagaraScriptUsage::ParticleSimulationStageScript:
 	case ENiagaraScriptUsage::ParticleGPUComputeScript:
 		if (UNiagaraEmitter* OwningEmitter = GetTypedOuter<UNiagaraEmitter>())
 		{
@@ -795,7 +825,7 @@ bool UNiagaraScript::IsUsageDependentOn(ENiagaraScriptUsage InUsageA, ENiagaraSc
 
 	// The GPU compute script is always dependent on the other particle scripts.
 	if ((InUsageA == ENiagaraScriptUsage::ParticleGPUComputeScript)
-		&& (InUsageB == ENiagaraScriptUsage::ParticleSpawnScript || InUsageB == ENiagaraScriptUsage::ParticleSpawnScriptInterpolated || InUsageB == ENiagaraScriptUsage::ParticleUpdateScript || InUsageB == ENiagaraScriptUsage::ParticleEventScript))
+		&& (InUsageB == ENiagaraScriptUsage::ParticleSpawnScript || InUsageB == ENiagaraScriptUsage::ParticleSpawnScriptInterpolated || InUsageB == ENiagaraScriptUsage::ParticleUpdateScript || InUsageB == ENiagaraScriptUsage::ParticleEventScript || InUsageB == ENiagaraScriptUsage::ParticleSimulationStageScript))
 	{
 		return true;
 	}
@@ -854,12 +884,12 @@ void UNiagaraScript::PostLoad()
 	const int32 NiagaraVer = GetLinkerCustomVersion(FNiagaraCustomVersion::GUID);
 
 #if WITH_EDITORONLY_DATA
-	if (NiagaraVer < FNiagaraCustomVersion::AddShaderStageUsageEnum)
+	if (NiagaraVer < FNiagaraCustomVersion::AddSimulationStageUsageEnum)
 	{
-		uint8 ShaderStageIndex = (uint8)ENiagaraScriptUsage::ParticleShaderStageScript;
+		uint8 SimulationStageIndex = (uint8)ENiagaraScriptUsage::ParticleSimulationStageScript;
 		uint8 MaxIndex = (uint8)ENiagaraScriptUsage::SystemUpdateScript;
 		// Start at the end and shift the bits down to account for the new shader stage bit.
-		for (uint8 CurrentIndex = MaxIndex; CurrentIndex > ShaderStageIndex; CurrentIndex--)
+		for (uint8 CurrentIndex = MaxIndex; CurrentIndex > SimulationStageIndex; CurrentIndex--)
 		{
 			uint8 OldIndex = CurrentIndex - 1;
 			if ((ModuleUsageBitmask & (1 << OldIndex)) != 0)
@@ -872,7 +902,7 @@ void UNiagaraScript::PostLoad()
 			}
 		}
 		// Clear the shader stage bit.
-		ModuleUsageBitmask &= ~(1 << ShaderStageIndex);
+		ModuleUsageBitmask &= ~(1 << SimulationStageIndex);
 	}
 
 	if (Source != nullptr)
@@ -1423,6 +1453,7 @@ void UNiagaraScript::RequestCompile(bool bForceCompile)
 			return;
 		}
 
+
 		CachedScriptVM.LastCompileStatus = ENiagaraScriptCompileStatus::NCS_BeingCreated;
 
 		TArray<TSharedPtr<FNiagaraCompileRequestDataBase, ESPMode::ThreadSafe>> DependentRequests;
@@ -1477,6 +1508,7 @@ bool UNiagaraScript::RequestExternallyManagedAsyncCompile(const TSharedPtr<FNiag
 			CachedScriptVMId = LastGeneratedVMId;
 			return false;
 		}
+
 		INiagaraModule& NiagaraModule = FModuleManager::Get().LoadModuleChecked<INiagaraModule>(TEXT("Niagara"));
 		CachedScriptVM.LastCompileStatus = ENiagaraScriptCompileStatus::NCS_BeingCreated;
 
@@ -1785,6 +1817,33 @@ void UNiagaraScript::SyncAliases(const TMap<FString, FString>& RenameMap)
 		}
 	}
 
+	// Sync up any simulation stage name references.
+	for (int32 i = 0; i < GetVMExecutableData().SimulationStageMetaData.Num(); i++)
+	{
+		if (!GetVMExecutableData().SimulationStageMetaData[i].IterationSource.IsNone())
+		{
+			FNiagaraVariable Var(FNiagaraTypeDefinition(UNiagaraDataInterface::StaticClass()), GetVMExecutableData().SimulationStageMetaData[i].IterationSource);
+			FNiagaraVariable NewVar = FNiagaraVariable::ResolveAliases(Var, RenameMap);
+			if (NewVar.GetName() != Var.GetName())
+			{
+				GetVMExecutableData().SimulationStageMetaData[i].IterationSource = NewVar.GetName();
+			}
+		}
+
+		for (int32 DestIdx = 0; DestIdx < GetVMExecutableData().SimulationStageMetaData[i].OutputDestinations.Num(); DestIdx++)
+		{
+			if (!GetVMExecutableData().SimulationStageMetaData[i].OutputDestinations[DestIdx].IsNone())
+			{
+				FNiagaraVariable Var(FNiagaraTypeDefinition(UNiagaraDataInterface::StaticClass()), GetVMExecutableData().SimulationStageMetaData[i].OutputDestinations[DestIdx]);
+				FNiagaraVariable NewVar = FNiagaraVariable::ResolveAliases(Var, RenameMap);
+				if (NewVar.GetName() != Var.GetName())
+				{
+					GetVMExecutableData().SimulationStageMetaData[i].OutputDestinations[DestIdx] = NewVar.GetName();
+				}
+			}
+		}
+	}
+
 	// Also handle any data set mappings...
 	auto Iterator = GetVMExecutableData().DataSetToParameters.CreateIterator();
 	while (Iterator)
diff --git a/Engine/Plugins/FX/Niagara/Source/Niagara/Private/NiagaraScriptExecutionContext.cpp b/Engine/Plugins/FX/Niagara/Source/Niagara/Private/NiagaraScriptExecutionContext.cpp
--- a/Engine/Plugins/FX/Niagara/Source/Niagara/Private/NiagaraScriptExecutionContext.cpp
+++ b/Engine/Plugins/FX/Niagara/Source/Niagara/Private/NiagaraScriptExecutionContext.cpp
@@ -10,6 +10,7 @@
 #include "NiagaraEmitterInstance.h"
 #include "NiagaraFunctionLibrary.h"
 #include "NiagaraGPUInstanceCountManager.h"
+#include "NiagaraDataInterfaceRW.h"
 
 DECLARE_CYCLE_STAT(TEXT("Register Setup"), STAT_NiagaraSimRegisterSetup, STATGROUP_Niagara);
 DECLARE_CYCLE_STAT(TEXT("Context Ticking"), STAT_NiagaraScriptExecContextTick, STATGROUP_Niagara);
@@ -306,7 +307,9 @@ void FNiagaraGPUSystemTick::Init(FNiagaraSystemInstance* InSystemInstance)
 	int32 TickCount = InSystemInstance->GetTickCount();
 	check(TickCount > 0);
 	bNeedsReset = ( TickCount == 1);
+	NumInstancesWithSimStages = 0;
 
+	TotalDispatches = 0;
 	uint8* GlobalParamData = ParamDataBufferPtr;
 	ParamDataBufferPtr += 2 * sizeof(FNiagaraGlobalParameters);
 	uint8* SystemParamData = ParamDataBufferPtr;
@@ -340,6 +343,21 @@ void FNiagaraGPUSystemTick::Init(FNiagaraSystemInstance* InSystemInstance)
 			InstanceData->SystemParamData = SystemParamData;
 			InstanceData->OwnerParamData = OwnerParamData;
 			InstanceData->EmitterParamData = ParamDataBufferPtr;
+			UNiagaraEmitter* EmitterRaw = Emitter->GetCachedEmitter();
+			if (EmitterRaw)
+			{
+				InstanceData->bUsesSimStages = EmitterRaw->bSimulationStagesEnabled/* TODO limit to just with stages in the future! Leaving like this so what can convert! && EmitterRaw->GetSimulationStages().Num() > 0*/;
+				InstanceData->bUsesOldShaderStages = EmitterRaw->bDeprecatedShaderStagesEnabled;
+			}
+			else
+			{
+				InstanceData->bUsesSimStages = false;
+			}
+
+			if (InstanceData->bUsesSimStages || InstanceData->bUsesOldShaderStages)
+			{
+				NumInstancesWithSimStages++;
+			}
 			ParamDataBufferPtr += 2 * sizeof(FNiagaraEmitterParameters);
 			InstanceData->ExternalParamData = ParamDataBufferPtr;
 			ParamDataBufferPtr += ParmSize;
@@ -356,14 +374,20 @@ void FNiagaraGPUSystemTick::Init(FNiagaraSystemInstance* InSystemInstance)
 
 			GPUContext->CombinedParamStore.CopyParameterDataToPaddedBuffer(InstanceData->ExternalParamData, ParmSize);
 
+			check(GPUContext->MaxUpdateIterations > 0);
+			InstanceData->SimStageData.AddZeroed(GPUContext->MaxUpdateIterations);
+			TotalDispatches += FMath::Max<int32>(GPUContext->MaxUpdateIterations, 1);
+
 			// @todo-threadsafety Think of a better way to do this!
 			const TArray<UNiagaraDataInterface*>& DataInterfaces = GPUContext->CombinedParamStore.GetDataInterfaces();
 			InstanceData->DataInterfaceProxies.Reserve(DataInterfaces.Num());
+			//UE_LOG(LogNiagara, Log, TEXT("InitTick %s %d"), GPUContext->GetDebugSimName())
 			for (UNiagaraDataInterface* DI : DataInterfaces)
 			{
 				check(DI->GetProxy());
 				InstanceData->DataInterfaceProxies.Add(DI->GetProxy());
-			}			
+				//UE_LOG(LogNiagara, Log, TEXT("Proxy %p for DI %p"), DI->GetProxy(), DI);
+			}
 		}
 	}
 
@@ -463,20 +487,127 @@ void FNiagaraComputeExecutionContext::Reset(NiagaraEmitterInstanceBatcher* Batch
 	);
 }
 
-void FNiagaraComputeExecutionContext::InitParams(UNiagaraScript* InGPUComputeScript, ENiagaraSimTarget InSimTarget, const FString& InDebugSimName, const uint32 InDefaultShaderStageIndex, const int32 InMaxUpdateIterations, const TSet<uint32> InSpawnStages)
+void FNiagaraComputeExecutionContext::InitParams(UNiagaraScript* InGPUComputeScript, ENiagaraSimTarget InSimTarget, const FString& InDebugSimName, const uint32 InDefaultSimulationStageIndex, const int32 InMaxUpdateIterations, const TSet<uint32> InSpawnStages)
 {
 #if !UE_BUILD_SHIPPING
 	DebugSimName = InDebugSimName;
 #endif
 	GPUScript = InGPUComputeScript;
 	CombinedParamStore.InitFromOwningContext(InGPUComputeScript, InSimTarget, true);
-	DefaultShaderStageIndex = InDefaultShaderStageIndex;
+	DefaultSimulationStageIndex = InDefaultSimulationStageIndex;
 	MaxUpdateIterations = InMaxUpdateIterations;
 	SpawnStages.Empty();
 
 	SpawnStages.Append(InSpawnStages);
 	
 	HasInterpolationParameters = GPUScript && GPUScript->GetComputedVMCompilationId().HasInterpolatedParameters();
+
+	if (InGPUComputeScript)
+	{
+		FNiagaraVMExecutableData& VMData = InGPUComputeScript->GetVMExecutableData();
+		if (VMData.IsValid() && VMData.SimulationStageMetaData.Num() > 0)
+		{
+
+			SimStageInfo = VMData.SimulationStageMetaData;
+
+			int32 FoundMaxUpdateIterations = SimStageInfo[SimStageInfo.Num() - 1].MaxStage;
+
+			// Some useful debugging code should we need to look up differences between old and new
+			const bool bDebugSimStages = false;
+			if (bDebugSimStages)
+			{
+				UE_LOG(LogNiagara, Log, TEXT("Stored vs:"));
+				bool bPass = FoundMaxUpdateIterations == MaxUpdateIterations;
+				UE_LOG(LogNiagara, Log, TEXT("MaxUpdateIterations: %d vs %d %s"), FoundMaxUpdateIterations, MaxUpdateIterations, bPass ? TEXT("Pass") : TEXT("FAIL!!!!!!!!"));
+
+				int32 NumSpawnFound = 0;
+				bool bMatchesFound = true;
+				for (int32 i = 0; i < SimStageInfo.Num(); i++)
+				{
+					if (SimStageInfo[i].bSpawnOnly)
+					{
+						NumSpawnFound++;
+
+						if (!SpawnStages.Contains(SimStageInfo[i].MinStage))
+						{
+							bMatchesFound = false;
+							UE_LOG(LogNiagara, Log, TEXT("Missing spawn stage: %d FAIL!!!!!!!!!"), SimStageInfo[i].MinStage);
+						}
+					}
+				}
+
+				bPass = SpawnStages.Num() == NumSpawnFound;
+				UE_LOG(LogNiagara, Log, TEXT("SpawnStages.Num(): %d vs %d %s"), NumSpawnFound, SpawnStages.Num(), bPass ? TEXT("Pass") : TEXT("FAIL!!!!!!!!"));
+
+				TArray<FNiagaraVariable> Params;
+				CombinedParamStore.GetParameters(Params);
+				for (FNiagaraVariable& Var : Params)
+				{
+					if (!Var.IsDataInterface())
+						continue;
+
+					UNiagaraDataInterface* DI = CombinedParamStore.GetDataInterface(Var);
+					UNiagaraDataInterfaceRWBase* DIRW = Cast<UNiagaraDataInterfaceRWBase>(DI);
+					if (DIRW)
+					{
+						for (int32 i = 0; i < SimStageInfo.Num(); i++)
+						{
+							if (SimStageInfo[i].IterationSource == Var.GetName())
+							{
+								if (!DIRW->IterationShaderStages.Contains(SimStageInfo[i].MinStage))
+								{
+									UE_LOG(LogNiagara, Log, TEXT("Missing iteration stage for %s: %d FAIL!!!!!!!!!"), *Var.GetName().ToString(), SimStageInfo[i].MinStage);
+								}
+							}
+
+							if (SimStageInfo[i].OutputDestinations.Contains(Var.GetName()))
+							{
+								if (!DIRW->OutputShaderStages.Contains(SimStageInfo[i].MinStage))
+								{
+									UE_LOG(LogNiagara, Log, TEXT("Missing output stage for %s: %d FAIL!!!!!!!!!"), *Var.GetName().ToString(), SimStageInfo[i].MinStage);
+								}
+							}
+						}
+					}
+
+				}
+			}
+
+
+			// Set the values that we are using from compiled data instead...
+			MaxUpdateIterations = SimStageInfo[SimStageInfo.Num() - 1].MaxStage;
+			SpawnStages.Empty();
+
+			for (int32 i = 0; i < SimStageInfo.Num(); i++)
+			{
+				if (SimStageInfo[i].bSpawnOnly)
+				{
+					SpawnStages.Add(SimStageInfo[i].MinStage);
+				}
+			}
+
+			// We need to store the name of each DI source variable here so that we can look it up later when looking for the 
+			// iteration interface.
+			TArray<FNiagaraVariable> Params;
+			CombinedParamStore.GetParameters(Params);
+			for (FNiagaraVariable& Var : Params)
+			{
+				if (!Var.IsDataInterface())
+					continue;
+
+				UNiagaraDataInterface* DI = CombinedParamStore.GetDataInterface(Var);
+				if (DI)
+				{
+					FNiagaraDataInterfaceProxy* Proxy = DI->GetProxy();
+					if (Proxy)
+					{
+						Proxy->SourceDIName = Var.GetName();
+					}
+				}
+			}
+		}
+	}
+
 	
 #if DO_CHECK
 	FNiagaraShaderRef Shader = InGPUComputeScript->GetRenderThreadScript()->GetShaderGameThread();
@@ -499,6 +630,87 @@ void FNiagaraComputeExecutionContext::InitParams(UNiagaraScript* InGPUComputeScr
 #endif
 }
 
+
+const FSimulationStageMetaData* FNiagaraComputeExecutionContext::GetSimStageMetaData(uint32 SimulationStageIndex) const
+{
+	if (SimStageInfo.Num() > 0)
+	{
+		for (int32 i = 0; i < SimStageInfo.Num(); i++)
+		{
+			if (SimulationStageIndex >= (uint32)SimStageInfo[i].MinStage  && SimulationStageIndex < (uint32)SimStageInfo[i].MaxStage)
+			{
+				return &SimStageInfo[i];
+			}
+		}
+	}
+	return nullptr;
+}
+
+bool FNiagaraComputeExecutionContext::IsOutputStage(FNiagaraDataInterfaceProxy* DIProxy, uint32 CurrentStage) const
+{
+	const FSimulationStageMetaData* MetaData = GetSimStageMetaData(CurrentStage);
+	if (MetaData && DIProxy && !DIProxy->SourceDIName.IsNone())
+	{
+		if (MetaData->OutputDestinations.Contains(DIProxy->SourceDIName))
+			return true;
+	}
+	else if (DIProxy && SimStageInfo.Num() == 0)
+	{
+		return DIProxy->IsOutputStage_DEPRECATED(CurrentStage);
+	}
+	return false;
+}
+
+bool FNiagaraComputeExecutionContext::IsIterationStage(FNiagaraDataInterfaceProxy* DIProxy, uint32 CurrentStage) const
+{
+	const FSimulationStageMetaData* MetaData = GetSimStageMetaData(CurrentStage);
+	if (MetaData && DIProxy && !DIProxy->SourceDIName.IsNone())
+	{
+		if (MetaData->IterationSource.IsNone()) // Per particle iteration...
+			return false;
+
+		if (MetaData->IterationSource == DIProxy->SourceDIName)
+			return true;
+	}
+	else if (DIProxy && SimStageInfo.Num() == 0)
+	{
+		return DIProxy->IsIterationStage_DEPRECATED(CurrentStage);
+	}
+	return false;
+}
+
+FNiagaraDataInterfaceProxy* FNiagaraComputeExecutionContext::FindIterationInterface(const TArray<FNiagaraDataInterfaceProxy*>& InProxies, uint32 CurrentStage) const
+{
+	const FSimulationStageMetaData* MetaData = GetSimStageMetaData(CurrentStage);
+	if (MetaData)
+	{
+		if (MetaData->IterationSource.IsNone()) // Per particle iteration...
+			return nullptr;
+
+		for (FNiagaraDataInterfaceProxy* Proxy : InProxies)
+		{
+			if (Proxy->SourceDIName == MetaData->IterationSource)
+				return Proxy;
+		}
+
+		UE_LOG(LogNiagara, Verbose, TEXT("FNiagaraComputeExecutionContext::FindIterationInterface could not find IterationInterface %s"), *MetaData->IterationSource.ToString());
+
+		return nullptr;
+	}
+	else if (SimStageInfo.Num() == 0)
+	{
+
+		// Fallback to old shader stages
+		for (FNiagaraDataInterfaceProxy* Proxy : InProxies)
+		{
+			if (Proxy->IsIterationStage_DEPRECATED(CurrentStage))
+				return Proxy;
+		}
+	}
+
+	return nullptr;
+}
+
 void FNiagaraComputeExecutionContext::DirtyDataInterfaces()
 {
 	CombinedParamStore.MarkInterfacesDirty();
@@ -587,8 +799,66 @@ void FNiagaraComputeExecutionContext::SetDataToRender(FNiagaraDataBuffer* InData
 
 	DataToRender = InDataToRender;
 
+
 	if (DataToRender)
 	{
 		DataToRender->AddReadRef();
 	}
 }
+
+bool FNiagaraComputeInstanceData::IsOutputStage(FNiagaraDataInterfaceProxy* DIProxy, uint32 CurrentStage) const
+{
+	if (bUsesOldShaderStages)
+	{
+		return DIProxy->IsOutputStage_DEPRECATED(CurrentStage);
+	}
+	else if (bUsesSimStages)
+	{
+		return Context->IsOutputStage(DIProxy, CurrentStage);
+	}
+	return false;
+}
+
+bool FNiagaraComputeInstanceData::IsIterationStage(FNiagaraDataInterfaceProxy* DIProxy, uint32 CurrentStage) const
+{
+
+	if (bUsesOldShaderStages)
+	{
+		return DIProxy->IsIterationStage_DEPRECATED(CurrentStage);
+	}
+	else if (bUsesSimStages)
+	{
+		return Context->IsIterationStage(DIProxy, CurrentStage);
+	}
+	return false;
+}
+
+FNiagaraDataInterfaceProxy* FNiagaraComputeInstanceData::FindIterationInterface(uint32 SimulationStageIndex) const
+{
+
+	if (bUsesOldShaderStages)
+	{
+		FNiagaraDataInterfaceProxy* IterationInterface = nullptr;
+
+		for (FNiagaraDataInterfaceProxy* Interface : DataInterfaceProxies)
+		{
+			if (Interface->IsIterationStage_DEPRECATED(SimulationStageIndex))
+			{
+				if (IterationInterface)
+				{
+					UE_LOG(LogNiagara, Error, TEXT("Multiple output Data Interfaces found for current stage"));
+				}
+				else
+				{
+					IterationInterface = Interface;
+				}
+			}
+		}
+		return IterationInterface;
+	}
+	else if (bUsesSimStages)
+	{
+		return Context->FindIterationInterface(DataInterfaceProxies, SimulationStageIndex);
+	}
+	return false;
+}
\ No newline at end of file
diff --git a/Engine/Plugins/FX/Niagara/Source/Niagara/Private/NiagaraShaderStageBase.cpp b/Engine/Plugins/FX/Niagara/Source/Niagara/Private/NiagaraShaderStageBase.cpp
deleted file mode 100644
diff --git a/Engine/Plugins/FX/Niagara/Source/Niagara/Private/NiagaraSimulationStageBase.cpp b/Engine/Plugins/FX/Niagara/Source/Niagara/Private/NiagaraSimulationStageBase.cpp
new file mode 100644
--- /dev/null
+++ b/Engine/Plugins/FX/Niagara/Source/Niagara/Private/NiagaraSimulationStageBase.cpp
@@ -0,0 +1,66 @@
+// Copyright Epic Games, Inc. All Rights Reserved.
+
+#include "NiagaraSimulationStageBase.h"
+#include "NiagaraEmitter.h"
+#include "NiagaraSystem.h"
+#include "NiagaraScriptSourceBase.h"
+
+
+bool UNiagaraSimulationStageGeneric::AppendCompileHash(FNiagaraCompileHashVisitor* InVisitor) const 
+{
+	InVisitor->UpdatePOD(TEXT("Iterations"), Iterations);
+	InVisitor->UpdatePOD(TEXT("IterationSource"), (int32)IterationSource);
+	InVisitor->UpdatePOD(TEXT("bSpawnOnly"), bSpawnOnly ? 1 : 0);
+	InVisitor->UpdateString(TEXT("DataInterface"), DataInterface.BoundVariable.GetName().ToString());
+	InVisitor->UpdateString(TEXT("SimulationStageName"), SimulationStageName.ToString());
+	return true;
+}
+#if WITH_EDITOR
+void UNiagaraSimulationStageGeneric::PostEditChangeProperty(struct FPropertyChangedEvent& PropertyChangedEvent) 
+{
+	Super::PostEditChangeProperty(PropertyChangedEvent);
+
+	FName PropertyName;
+	if (PropertyChangedEvent.Property)
+	{
+		PropertyName = PropertyChangedEvent.Property->GetFName();
+	}
+
+	UNiagaraEmitter* Emitter = Cast< UNiagaraEmitter>(GetOuter());
+
+	bool bNeedsRecompile = false;
+	if (PropertyName == GET_MEMBER_NAME_CHECKED(UNiagaraSimulationStageGeneric, Iterations) && Emitter)
+	{
+		bNeedsRecompile = true;
+	}
+	else if (PropertyName == GET_MEMBER_NAME_CHECKED(UNiagaraSimulationStageGeneric, IterationSource) && Emitter)
+	{
+		bNeedsRecompile = true;
+	}
+	else if (PropertyName == GET_MEMBER_NAME_CHECKED(UNiagaraSimulationStageGeneric, bSpawnOnly) && Emitter)
+	{
+		bNeedsRecompile = true;
+	}
+	else if (PropertyName == GET_MEMBER_NAME_CHECKED(UNiagaraSimulationStageGeneric, DataInterface) && Emitter)
+	{
+		bNeedsRecompile = true;
+	}
+	else if (PropertyName == GET_MEMBER_NAME_CHECKED(UNiagaraSimulationStageGeneric, SimulationStageName) && Emitter)
+	{
+		bNeedsRecompile = true;
+	}
+
+#if WITH_EDITORONLY_DATA
+	if (bNeedsRecompile)
+	{
+		UNiagaraScriptSourceBase* GraphSource = Emitter->UpdateScriptProps.Script->GetSource();
+		if (GraphSource != nullptr)
+		{
+			GraphSource->MarkNotSynchronized(TEXT("SimulationStageGeneric changed."));
+		}
+
+		UNiagaraSystem::RequestCompileForEmitter(Emitter);
+	}
+#endif
+}
+#endif
\ No newline at end of file
diff --git a/Engine/Plugins/FX/Niagara/Source/Niagara/Private/NiagaraSystemInstance.cpp b/Engine/Plugins/FX/Niagara/Source/Niagara/Private/NiagaraSystemInstance.cpp
--- a/Engine/Plugins/FX/Niagara/Source/Niagara/Private/NiagaraSystemInstance.cpp
+++ b/Engine/Plugins/FX/Niagara/Source/Niagara/Private/NiagaraSystemInstance.cpp
@@ -1321,6 +1321,7 @@ void FNiagaraSystemInstance::InitDataInterfaces()
 				int32* ExistingInstanceDataOffset = DataInterfaceInstanceDataOffsets.Find(Interface);
 				if (!ExistingInstanceDataOffset)//Don't add instance data for interfaces we've seen before.
 				{
+					//UE_LOG(LogNiagara, Log, TEXT("Adding DI %p %s %s"), Interface, *Interface->GetClass()->GetName(), *Interface->GetPathName());
 					DataInterfaceInstanceDataOffsets.Add(Interface) = InstanceDataSize;
 					// Assume that some of our data is going to be 16 byte aligned, so enforce that 
 					// all per-instance data is aligned that way.
@@ -1357,6 +1358,12 @@ void FNiagaraSystemInstance::InitDataInterfaces()
 			CalcInstDataSize(Sim.GetEventExecutionContexts()[i].GetDataInterfaces());
 		}
 
+		if (Sim.GetCachedEmitter() && Sim.GetCachedEmitter()->SimTarget == ENiagaraSimTarget::GPUComputeSim && Sim.GetCachedEmitter()->bSimulationStagesEnabled)
+		{
+			CalcInstDataSize(Sim.GetGPUContext()->GetDataInterfaces());
+		}
+
+
 		//Also force a rebind while we're here.
 		Sim.DirtyDataInterfaces();
 	}
diff --git a/Engine/Plugins/FX/Niagara/Source/Niagara/Private/NiagaraSystemSimulation.cpp b/Engine/Plugins/FX/Niagara/Source/Niagara/Private/NiagaraSystemSimulation.cpp
--- a/Engine/Plugins/FX/Niagara/Source/Niagara/Private/NiagaraSystemSimulation.cpp
+++ b/Engine/Plugins/FX/Niagara/Source/Niagara/Private/NiagaraSystemSimulation.cpp
@@ -1648,6 +1648,12 @@ void FNiagaraSystemSimulation::TransferSystemSimResults(FNiagaraSystemSimulation
 				FNiagaraScriptExecutionContext& UpdateContext = EmitterInst.GetUpdateExecutionContext();
 				DataSetToEmitterUpdateParameters[EmitterIdx].DataSetToParameterStore(UpdateContext.Parameters, Context.DataSet, SystemIndex);
 
+				FNiagaraComputeExecutionContext* GPUContext = EmitterInst.GetGPUContext();
+				if (GPUContext)
+				{
+					DataSetToEmitterGPUParameters[EmitterIdx].DataSetToParameterStore(GPUContext->CombinedParamStore, Context.DataSet, SystemIndex);
+				}
+
 				TArray<FNiagaraScriptExecutionContext>& EventContexts = EmitterInst.GetEventExecutionContexts();
 				for (int32 EventIdx = 0; EventIdx < EventContexts.Num(); ++EventIdx)
 				{
@@ -1923,6 +1929,7 @@ void FNiagaraSystemSimulation::InitParameterDataSetBindings(FNiagaraSystemInstan
 		DataSetToEmitterSpawnParameters.SetNum(EmitterCount);
 		DataSetToEmitterUpdateParameters.SetNum(EmitterCount);
 		DataSetToEmitterEventParameters.SetNum(EmitterCount);
+		DataSetToEmitterGPUParameters.SetNum(EmitterCount);
 
 		const FString EmitterNamespace = TEXT("Emitter");
 
@@ -1937,6 +1944,13 @@ void FNiagaraSystemSimulation::InitParameterDataSetBindings(FNiagaraSystemInstan
 			FNiagaraScriptExecutionContext& UpdateContext = EmitterInst.GetUpdateExecutionContext();
 			DataSetToEmitterUpdateParameters[EmitterIdx].Init(MainDataSet, UpdateContext.Parameters);
 
+
+			FNiagaraComputeExecutionContext* GPUContext = EmitterInst.GetGPUContext();
+			if (GPUContext)
+			{
+				DataSetToEmitterGPUParameters[EmitterIdx].Init(MainDataSet, GPUContext->CombinedParamStore);
+			}
+
 			TArray<FNiagaraScriptExecutionContext>& EventContexts = EmitterInst.GetEventExecutionContexts();
 			const int32 EventCount = EventContexts.Num();
 			DataSetToEmitterEventParameters[EmitterIdx].SetNum(EventCount);
diff --git a/Engine/Plugins/FX/Niagara/Source/Niagara/Public/NiagaraCommon.h b/Engine/Plugins/FX/Niagara/Source/Niagara/Public/NiagaraCommon.h
--- a/Engine/Plugins/FX/Niagara/Source/Niagara/Public/NiagaraCommon.h
+++ b/Engine/Plugins/FX/Niagara/Source/Niagara/Public/NiagaraCommon.h
@@ -272,6 +272,10 @@ struct NIAGARA_API FNiagaraFunctionSignature
 	UPROPERTY()
 	uint32 bSupportsGPU : 1;
 
+	/** Writes to the variable this is bound to */
+	UPROPERTY()
+	uint32 bWriteFunction : 1;
+
 	/** Function specifiers verified at bind time. */
 	UPROPERTY()
 	TMap<FName, FName> FunctionSpecifiers;
@@ -288,6 +292,7 @@ struct NIAGARA_API FNiagaraFunctionSignature
 		, bExperimental(false)
 		, bSupportsCPU(true)
 		, bSupportsGPU(true)
+		, bWriteFunction(false)
 	{
 	}
 
@@ -300,6 +305,7 @@ struct NIAGARA_API FNiagaraFunctionSignature
 		, bExperimental(false)
 		, bSupportsCPU(true)
 		, bSupportsGPU(true)
+		, bWriteFunction(false)
 	{
 
 	}
@@ -313,6 +319,7 @@ struct NIAGARA_API FNiagaraFunctionSignature
 		, bExperimental(false)
 		, bSupportsCPU(true)
 		, bSupportsGPU(true)
+		, bWriteFunction(false)
 		, FunctionSpecifiers(InFunctionSpecifiers)
 	{
 
@@ -540,12 +547,12 @@ struct NIAGARA_API FNiagaraSystemUpdateContext
 {
 	GENERATED_BODY()
 
-	FNiagaraSystemUpdateContext(const UNiagaraSystem* System, bool bReInit) :bDestroyOnAdd(false) { Add(System, bReInit); }
+	FNiagaraSystemUpdateContext(const UNiagaraSystem* System, bool bReInit, bool bInDestroyOnAdd = false) :bDestroyOnAdd(bInDestroyOnAdd) { Add(System, bReInit); }
 #if WITH_EDITORONLY_DATA
-	FNiagaraSystemUpdateContext(const UNiagaraEmitter* Emitter, bool bReInit) : bDestroyOnAdd(false) { Add(Emitter, bReInit); }
-	FNiagaraSystemUpdateContext(const UNiagaraScript* Script, bool bReInit) :bDestroyOnAdd(false) { Add(Script, bReInit); }
+	FNiagaraSystemUpdateContext(const UNiagaraEmitter* Emitter, bool bReInit, bool bInDestroyOnAdd = false) : bDestroyOnAdd(bInDestroyOnAdd) { Add(Emitter, bReInit); }
+	FNiagaraSystemUpdateContext(const UNiagaraScript* Script, bool bReInit, bool bInDestroyOnAdd = false) :bDestroyOnAdd(bInDestroyOnAdd) { Add(Script, bReInit); }
 	//FNiagaraSystemUpdateContext(UNiagaraDataInterface* Interface, bool bReinit) : Add(Interface, bReinit) {}
-	FNiagaraSystemUpdateContext(const UNiagaraParameterCollection* Collection, bool bReInit) :bDestroyOnAdd(false) { Add(Collection, bReInit); }
+	FNiagaraSystemUpdateContext(const UNiagaraParameterCollection* Collection, bool bReInit, bool bInDestroyOnAdd = false) :bDestroyOnAdd(bInDestroyOnAdd) { Add(Collection, bReInit); }
 #endif
 	FNiagaraSystemUpdateContext():bDestroyOnAdd(false){ }
 
@@ -597,25 +604,25 @@ enum class ENiagaraScriptUsage : uint8
 	/** The script defines a dynamic input for use in particle, emitter, or system scripts. */
 	DynamicInput,
 	/** The script is called when spawning particles. */
-	ParticleSpawnScript UMETA(Hidden),
+	ParticleSpawnScript,
 	/** Particle spawn script that handles intra-frame spawning and also pulls in the update script. */
 	ParticleSpawnScriptInterpolated UMETA(Hidden),
 	/** The script is called to update particles every frame. */
-	ParticleUpdateScript UMETA(Hidden),
+	ParticleUpdateScript,
 	/** The script is called to update particles in response to an event. */
-	ParticleEventScript UMETA(Hidden),
-	/** The script is called as a particle shader stage. */
-	ParticleShaderStageScript UMETA(Hidden),
+	ParticleEventScript ,
+	/** The script is called as a particle simulation stage. */
+	ParticleSimulationStageScript,
 	/** The script is called to update particles on the GPU. */
 	ParticleGPUComputeScript UMETA(Hidden),
 	/** The script is called once when the emitter spawns. */
-	EmitterSpawnScript UMETA(Hidden),
+	EmitterSpawnScript,
 	/** The script is called every frame to tick the emitter. */
-	EmitterUpdateScript UMETA(Hidden),
+	EmitterUpdateScript ,
 	/** The script is called once when the system spawns. */
-	SystemSpawnScript UMETA(Hidden),
+	SystemSpawnScript ,
 	/** The script is called every frame to tick the system. */
-	SystemUpdateScript UMETA(Hidden),
+	SystemUpdateScript,
 };
 
 UENUM()
@@ -628,6 +635,14 @@ enum class ENiagaraScriptGroup : uint8
 };
 
 
+UENUM()
+enum class ENiagaraIterationSource : uint8
+{
+	Particles = 0,
+	DataInterface
+};
+
+
 /** Defines all you need to know about a variable.*/
 USTRUCT()
 struct FNiagaraVariableInfo
@@ -680,11 +695,12 @@ struct FNiagaraVariableDataInterfaceBinding
 	FNiagaraVariableDataInterfaceBinding() {}
 	FNiagaraVariableDataInterfaceBinding(const FNiagaraVariable& InVar) : BoundVariable(InVar)
 	{
-		check(InVar.IsDataInterface() == true);
+		ensure(InVar.IsDataInterface() == true);
 	}
 
 	UPROPERTY()
 	FNiagaraVariable BoundVariable;
+
 };
 
 /** Primarily a wrapper around an FName to be used for customizations in the Selected Details panel 
diff --git a/Engine/Plugins/FX/Niagara/Source/Niagara/Public/NiagaraShaderStageBase.h b/Engine/Plugins/FX/Niagara/Source/Niagara/Public/NiagaraShaderStageBase.h
deleted file mode 100644
diff --git a/Engine/Plugins/FX/Niagara/Source/Niagara/Public/NiagaraSimulationStageBase.h b/Engine/Plugins/FX/Niagara/Source/Niagara/Public/NiagaraSimulationStageBase.h
new file mode 100644
--- /dev/null
+++ b/Engine/Plugins/FX/Niagara/Source/Niagara/Public/NiagaraSimulationStageBase.h
@@ -0,0 +1,59 @@
+// Copyright Epic Games, Inc. All Rights Reserved.
+
+#pragma once
+
+#include "NiagaraMergeable.h"
+#include "NiagaraCommon.h"
+#include "NiagaraTypes.h"
+#include "NiagaraSimulationStageBase.generated.h"
+
+class UNiagaraScript;
+
+/**
+* A base class for niagara simulation stages.  This class should be derived to add stage specific information.
+*/
+UCLASS()
+class NIAGARA_API UNiagaraSimulationStageBase : public UNiagaraMergeable
+{
+	GENERATED_BODY()
+
+public:
+	UPROPERTY()
+	UNiagaraScript* Script;
+
+	UPROPERTY(EditAnywhere, Category = "Simulation Stage")
+	FName SimulationStageName;
+
+	virtual bool AppendCompileHash(FNiagaraCompileHashVisitor* InVisitor) const { return true; }
+
+};
+
+UCLASS(meta = (DisplayName = "Generic Simulation Stage"))
+class NIAGARA_API UNiagaraSimulationStageGeneric : public UNiagaraSimulationStageBase
+{
+	GENERATED_BODY()
+
+public:
+	UNiagaraSimulationStageGeneric()
+		: Iterations(1)
+	{
+	}
+
+	/** Determine which elements this script is iterating over. You are not allowed to */
+	UPROPERTY(EditAnywhere, Category = "Simulation Stage")
+	ENiagaraIterationSource IterationSource;
+
+	UPROPERTY(EditAnywhere, Category = "Simulation Stage")
+	int32 Iterations;
+
+	UPROPERTY(EditAnywhere, Category = "Simulation Stage")
+	uint32 bSpawnOnly : 1;
+
+	UPROPERTY(EditAnywhere, Category = "Simulation Stage", meta = (editcondition = "IterationSource == ENiagaraIterationSource::DataInterface"))
+	FNiagaraVariableDataInterfaceBinding DataInterface;
+
+	virtual bool AppendCompileHash(FNiagaraCompileHashVisitor* InVisitor) const override;
+#if WITH_EDITOR
+	virtual void PostEditChangeProperty(struct FPropertyChangedEvent& PropertyChangedEvent) override;
+#endif
+};
\ No newline at end of file
diff --git a/Engine/Plugins/FX/Niagara/Source/Niagara/Public/NiagaraSystemSimulation.h b/Engine/Plugins/FX/Niagara/Source/Niagara/Public/NiagaraSystemSimulation.h
--- a/Engine/Plugins/FX/Niagara/Source/Niagara/Public/NiagaraSystemSimulation.h
+++ b/Engine/Plugins/FX/Niagara/Source/Niagara/Public/NiagaraSystemSimulation.h
@@ -303,6 +303,9 @@ protected:
 	TArray<FNiagaraParameterStoreToDataSetBinding> DataSetToEmitterUpdateParameters;
 	/** Binding to push system attributes into each emitter event parameters. */
 	TArray<TArray<FNiagaraParameterStoreToDataSetBinding>> DataSetToEmitterEventParameters;
+	/** Binding to push system attributes into each emitter gpu parameters. */
+	TArray<FNiagaraParameterStoreToDataSetBinding> DataSetToEmitterGPUParameters;
+
 
 	/** Direct bindings for Engine variables in System Spawn and Update scripts. */
 	FNiagaraParameterDirectBinding<int32> SpawnNumSystemInstancesParam;
diff --git a/Engine/Plugins/FX/Niagara/Source/NiagaraCore/Private/NiagaraCustomVersion.cpp b/Engine/Plugins/FX/Niagara/Source/NiagaraCore/Private/NiagaraCustomVersion.cpp
--- a/Engine/Plugins/FX/Niagara/Source/NiagaraCore/Private/NiagaraCustomVersion.cpp
+++ b/Engine/Plugins/FX/Niagara/Source/NiagaraCore/Private/NiagaraCustomVersion.cpp
@@ -8,4 +8,4 @@ const FGuid FNiagaraCustomVersion::GUID(0xFCF57AFA, 0x50764283, 0xB9A9E658, 0xFF
 // Register the custom version with core
 FCustomVersionRegistration GRegisterNiagaraCustomVersion(FNiagaraCustomVersion::GUID, FNiagaraCustomVersion::LatestVersion, TEXT("NiagaraVer"));
 
-const FGuid FNiagaraCustomVersion::LatestScriptCompileVersion(0x19CA6492, 0x8CCB46EC, 0xAC439F3F, 0x0BAAE91B);
+const FGuid FNiagaraCustomVersion::LatestScriptCompileVersion(0x751E8859, 0x84474334, 0x9CC3D9A1, 0xB554796E);
diff --git a/Engine/Plugins/FX/Niagara/Source/NiagaraCore/Public/NiagaraCustomVersion.h b/Engine/Plugins/FX/Niagara/Source/NiagaraCore/Public/NiagaraCustomVersion.h
--- a/Engine/Plugins/FX/Niagara/Source/NiagaraCore/Public/NiagaraCustomVersion.h
+++ b/Engine/Plugins/FX/Niagara/Source/NiagaraCore/Public/NiagaraCustomVersion.h
@@ -136,7 +136,7 @@ struct FNiagaraCustomVersion
 
 		MemorySaving, // Convert TMap into TArray to save memory, TMap contains an inline allocator which pushes the size to 80 bytes
 
-		AddShaderStageUsageEnum, // Added a new value to the script usage enum, and we need a custom version to fix the existing bitfields.
+		AddSimulationStageUsageEnum, // Added a new value to the script usage enum, and we need a custom version to fix the existing bitfields.
 
 		AddGeneratedFunctionsToGPUParamInfo, // Save the functions generated by a GPU data interface inside FNiagaraDataInterfaceGPUParamInfo
 
diff --git a/Engine/Plugins/FX/Niagara/Source/NiagaraCore/Public/NiagaraDataInterfaceBase.h b/Engine/Plugins/FX/Niagara/Source/NiagaraCore/Public/NiagaraDataInterfaceBase.h
--- a/Engine/Plugins/FX/Niagara/Source/NiagaraCore/Public/NiagaraDataInterfaceBase.h
+++ b/Engine/Plugins/FX/Niagara/Source/NiagaraCore/Public/NiagaraDataInterfaceBase.h
@@ -31,7 +31,7 @@ struct FNiagaraDataInterfaceSetArgs
 	FNiagaraSystemInstanceID SystemInstance;
 	const NiagaraEmitterInstanceBatcher* Batcher;
 	const FNiagaraComputeInstanceData* ComputeInstanceData;
-	uint32 ShaderStageIndex;
+	uint32 SimulationStageIndex;
 	bool IsOutputStage;
 	bool IsIterationStage;
 };
diff --git a/Engine/Plugins/FX/Niagara/Source/NiagaraEditor/Private/Customizations/NiagaraTypeCustomizations.cpp b/Engine/Plugins/FX/Niagara/Source/NiagaraEditor/Private/Customizations/NiagaraTypeCustomizations.cpp
--- a/Engine/Plugins/FX/Niagara/Source/NiagaraEditor/Private/Customizations/NiagaraTypeCustomizations.cpp
+++ b/Engine/Plugins/FX/Niagara/Source/NiagaraEditor/Private/Customizations/NiagaraTypeCustomizations.cpp
@@ -30,7 +30,9 @@
 #include "Widgets/Input/SCheckBox.h"
 #include "Widgets/Input/STextComboBox.h"
 #include "Widgets/Layout/SWrapBox.h"
+#include "NiagaraSimulationStageBase.h"
 #include "Widgets/Text/STextBlock.h"
+#include "NiagaraDataInterfaceRW.h"
 
 #define LOCTEXT_NAMESPACE "FNiagaraVariableAttributeBindingCustomization"
 
@@ -486,7 +488,214 @@ void FNiagaraUserParameterBindingCustomization::CustomizeHeader(TSharedRef<IProp
 	}
 }
 
+//////////////////////////////////////////////////////////////////////////
+
+FText FNiagaraDataInterfaceBindingCustomization::GetCurrentText() const
+{
+	if (BaseStage && TargetDataInterfaceBinding)
+	{
+		return FText::FromName(TargetDataInterfaceBinding->BoundVariable.GetName());
+	}
+	return FText::FromString(TEXT("Missing"));
+}
+
+FText FNiagaraDataInterfaceBindingCustomization::GetTooltipText() const
+{
+	if (BaseStage && TargetDataInterfaceBinding && TargetDataInterfaceBinding->BoundVariable.IsValid())
+	{
+		FText TooltipDesc = FText::Format(LOCTEXT("ParameterBindingTooltip", "Bound to the user parameter \"{0}\""), FText::FromName(TargetDataInterfaceBinding->BoundVariable.GetName()));
+		return TooltipDesc;
+	}
+	return FText::FromString(TEXT("Missing"));
+}
+
+TSharedRef<SWidget> FNiagaraDataInterfaceBindingCustomization::OnGetMenuContent() const
+{
+	FGraphActionMenuBuilder MenuBuilder;
+
+	return SNew(SBorder)
+		.BorderImage(FEditorStyle::GetBrush("Menu.Background"))
+		.Padding(5)
+		[
+			SNew(SBox)
+			[
+				SNew(SGraphActionMenu)
+				.OnActionSelected(const_cast<FNiagaraDataInterfaceBindingCustomization*>(this), &FNiagaraDataInterfaceBindingCustomization::OnActionSelected)
+		.OnCreateWidgetForAction(SGraphActionMenu::FOnCreateWidgetForAction::CreateSP(const_cast<FNiagaraDataInterfaceBindingCustomization*>(this), &FNiagaraDataInterfaceBindingCustomization::OnCreateWidgetForAction))
+		.OnCollectAllActions(const_cast<FNiagaraDataInterfaceBindingCustomization*>(this), &FNiagaraDataInterfaceBindingCustomization::CollectAllActions)
+		.AutoExpandActionMenu(false)
+		.ShowFilterTextBox(true)
+			]
+		];
+}
+
+TArray<FName> FNiagaraDataInterfaceBindingCustomization::GetNames() const
+{
+	TArray<FName> Names;
+
+	if (BaseStage && TargetDataInterfaceBinding)
+	{
+		UNiagaraEmitter* Emitter = BaseStage->GetTypedOuter<UNiagaraEmitter>();
+
+		if (Emitter)
+		{
+			// Find all used emitter and particle data interface variables that can be iterated upon.
+			TArray<UNiagaraScript*> AllScripts;
+			Emitter->GetScripts(AllScripts, false);
+
+			TArray<UNiagaraGraph*> Graphs;
+			for (const UNiagaraScript* Script : AllScripts)
+			{
+				const UNiagaraScriptSource* Source = Cast<UNiagaraScriptSource>(Script->GetSource());
+				if (Source)
+				{
+					Graphs.AddUnique(Source->NodeGraph);
+				}
+			}
+
+			for (const UNiagaraGraph* Graph : Graphs)
+			{
+				const TMap<FNiagaraVariable, FNiagaraGraphParameterReferenceCollection>& ParameterReferenceMap = Graph->GetParameterReferenceMap();
+				for (const auto& ParameterToReferences : ParameterReferenceMap)
+				{
+					const FNiagaraVariable& ParameterVariable = ParameterToReferences.Key;
+					if (ParameterVariable.IsDataInterface())
+					{
+						const UClass* Class = ParameterVariable.GetType().GetClass();
+						if (Class)
+						{
+							const UObject* DefaultObjDI = Class->GetDefaultObject();
+							if (DefaultObjDI != nullptr && DefaultObjDI->IsA<UNiagaraDataInterfaceRWBase>())
+							{
+								Names.AddUnique(ParameterVariable.GetName());
+							}
+						}
+					}
+				}
+			}
+		}
+	}
+
+	return Names;
+}
+
+void FNiagaraDataInterfaceBindingCustomization::CollectAllActions(FGraphActionListBuilderBase& OutAllActions)
+{
+	TArray<FName> UserParamNames = GetNames();
+	for (FName UserParamName : UserParamNames)
+	{
+		FText CategoryName = FText();
+		FString DisplayNameString = FName::NameToDisplayString(UserParamName.ToString(), false);
+		const FText NameText = FText::FromString(DisplayNameString);
+		const FText TooltipDesc = FText::Format(LOCTEXT("BindToDataInterface", "Bind to the User Parameter \"{0}\" "), FText::FromString(DisplayNameString));
+		TSharedPtr<FNiagaraStackAssetAction_VarBind> NewNodeAction(new FNiagaraStackAssetAction_VarBind(UserParamName, CategoryName, NameText,
+			TooltipDesc, 0, FText()));
+		OutAllActions.AddAction(NewNodeAction);
+	}
+}
+
+TSharedRef<SWidget> FNiagaraDataInterfaceBindingCustomization::OnCreateWidgetForAction(struct FCreateWidgetForActionData* const InCreateData)
+{
+	return SNew(SVerticalBox)
+		+ SVerticalBox::Slot()
+		.AutoHeight()
+		[
+			SNew(STextBlock)
+			.Text(InCreateData->Action->GetMenuDescription())
+			.ToolTipText(InCreateData->Action->GetTooltipDescription())
+		];
+}
+
+
+void FNiagaraDataInterfaceBindingCustomization::OnActionSelected(const TArray< TSharedPtr<FEdGraphSchemaAction> >& SelectedActions, ESelectInfo::Type InSelectionType)
+{
+	if (InSelectionType == ESelectInfo::OnMouseClick || InSelectionType == ESelectInfo::OnKeyPress || SelectedActions.Num() == 0)
+	{
+		for (int32 ActionIndex = 0; ActionIndex < SelectedActions.Num(); ActionIndex++)
+		{
+			TSharedPtr<FEdGraphSchemaAction> CurrentAction = SelectedActions[ActionIndex];
+
+			if (CurrentAction.IsValid())
+			{
+				FSlateApplication::Get().DismissAllMenus();
+				FNiagaraStackAssetAction_VarBind* EventSourceAction = (FNiagaraStackAssetAction_VarBind*)CurrentAction.Get();
+				ChangeSource(EventSourceAction->VarName);
+			}
+		}
+	}
+}
+
+void FNiagaraDataInterfaceBindingCustomization::ChangeSource(FName InVarName)
+{
+	FScopedTransaction Transaction(FText::Format(LOCTEXT("ChangeParameterSource", " Change Data Interface Source to \"{0}\" "), FText::FromName(InVarName)));
+	TArray<UObject*> Objects;
+	PropertyHandle->GetOuterObjects(Objects);
+	for (UObject* Obj : Objects)
+	{
+		Obj->Modify();
+	}
+
+	PropertyHandle->NotifyPreChange();
+	TargetDataInterfaceBinding->BoundVariable.SetName(InVarName);
+	PropertyHandle->NotifyPostChange();
+	PropertyHandle->NotifyFinishedChangingProperties();
+}
+
+void FNiagaraDataInterfaceBindingCustomization::CustomizeHeader(TSharedRef<IPropertyHandle> InPropertyHandle, FDetailWidgetRow& HeaderRow, IPropertyTypeCustomizationUtils& CustomizationUtils)
+{
+	PropertyHandle = InPropertyHandle;
+	TArray<UObject*> Objects;
+	PropertyHandle->GetOuterObjects(Objects);
+	bool bAddDefault = true;
+	if (Objects.Num() == 1)
+	{
+		BaseStage = Cast<UNiagaraSimulationStageBase>(Objects[0]);
+		if (BaseStage)
+		{
+			TargetDataInterfaceBinding = (FNiagaraVariableDataInterfaceBinding*)PropertyHandle->GetValueBaseAddress((uint8*)Objects[0]);
+
+			HeaderRow
+				.NameContent()
+				[
+					PropertyHandle->CreatePropertyNameWidget()
+				]
+			.ValueContent()
+				.MaxDesiredWidth(200.f)
+				[
+					SNew(SComboButton)
+					.OnGetMenuContent(this, &FNiagaraDataInterfaceBindingCustomization::OnGetMenuContent)
+					.ContentPadding(1)
+					.ToolTipText(this, &FNiagaraDataInterfaceBindingCustomization::GetTooltipText)
+					.ButtonContent()
+				[
+					SNew(STextBlock)
+					.Text(this, &FNiagaraDataInterfaceBindingCustomization::GetCurrentText)
+					.Font(IDetailLayoutBuilder::GetDetailFont())
+				]
+				];
+			bAddDefault = false;
+		}
+	}
+
+
+	if (bAddDefault)
+	{
+		HeaderRow
+			.NameContent()
+			[
+				PropertyHandle->CreatePropertyNameWidget()
+			]
+			.ValueContent()
+			.MaxDesiredWidth(200.f)
+			[
+				SNew(STextBlock)
+				.Text(FText::FromString(FName::NameToDisplayString(CastField<FStructProperty>(PropertyHandle->GetProperty())->Struct->GetName(), false)))
+				.Font(IDetailLayoutBuilder::GetDetailFont())
+			];
+	}
+}
 
+////////////////////////
 FText FNiagaraScriptVariableBindingCustomization::GetCurrentText() const
 {
 	if (BaseGraph && TargetVariableBinding && TargetVariableBinding->IsValid())
diff --git a/Engine/Plugins/FX/Niagara/Source/NiagaraEditor/Private/Customizations/NiagaraTypeCustomizations.h b/Engine/Plugins/FX/Niagara/Source/NiagaraEditor/Private/Customizations/NiagaraTypeCustomizations.h
--- a/Engine/Plugins/FX/Niagara/Source/NiagaraEditor/Private/Customizations/NiagaraTypeCustomizations.h
+++ b/Engine/Plugins/FX/Niagara/Source/NiagaraEditor/Private/Customizations/NiagaraTypeCustomizations.h
@@ -163,6 +163,35 @@ private:
 
 };
 
+class FNiagaraDataInterfaceBindingCustomization : public IPropertyTypeCustomization
+{
+public:
+	/** @return A new instance of this class */
+	static TSharedRef<IPropertyTypeCustomization> MakeInstance()
+	{
+		return MakeShared<FNiagaraDataInterfaceBindingCustomization>();
+	}
+
+	/** IPropertyTypeCustomization interface begin */
+	virtual void CustomizeHeader(TSharedRef<IPropertyHandle> StructPropertyHandle, class FDetailWidgetRow& HeaderRow, IPropertyTypeCustomizationUtils& StructCustomizationUtils) override;
+	virtual void CustomizeChildren(TSharedRef<IPropertyHandle> StructPropertyHandle, class IDetailChildrenBuilder& ChildBuilder, IPropertyTypeCustomizationUtils& StructCustomizationUtils) override {};
+	/** IPropertyTypeCustomization interface end */
+private:
+	FText GetCurrentText() const;
+	FText GetTooltipText() const;
+	TSharedRef<SWidget> OnGetMenuContent() const;
+	TArray<FName> GetNames() const;
+	void ChangeSource(FName InVarName);
+	void CollectAllActions(FGraphActionListBuilderBase& OutAllActions);
+	TSharedRef<SWidget> OnCreateWidgetForAction(struct FCreateWidgetForActionData* const InCreateData);
+	void OnActionSelected(const TArray< TSharedPtr<FEdGraphSchemaAction> >& SelectedActions, ESelectInfo::Type InSelectionType);
+
+	TSharedPtr<IPropertyHandle> PropertyHandle;
+	class UNiagaraSimulationStageBase* BaseStage;
+	struct FNiagaraVariableDataInterfaceBinding* TargetDataInterfaceBinding;
+
+};
+
 
 /** The primary goal of this class is to search through type matched and defined Niagara variables 
     in the UNiagaraScriptVariable customization panel to provide a default binding for module inputs. */
diff --git a/Engine/Plugins/FX/Niagara/Source/NiagaraEditor/Private/EdGraphSchema_Niagara.cpp b/Engine/Plugins/FX/Niagara/Source/NiagaraEditor/Private/EdGraphSchema_Niagara.cpp
--- a/Engine/Plugins/FX/Niagara/Source/NiagaraEditor/Private/EdGraphSchema_Niagara.cpp
+++ b/Engine/Plugins/FX/Niagara/Source/NiagaraEditor/Private/EdGraphSchema_Niagara.cpp
@@ -722,8 +722,6 @@ TArray<TSharedPtr<FNiagaraSchemaAction_NewNode> > UEdGraphSchema_Niagara::GetGra
 			UNiagaraNodeParameterMapSet* BaseNode = NewObject<UNiagaraNodeParameterMapSet>(OwnerOfTemporaries);
 			Action->NodeTemplate = BaseNode;
 		}
-		static const auto UseShaderStagesCVar = IConsoleManager::Get().FindConsoleVariable(TEXT("fx.UseShaderStages"));
-		if (UseShaderStagesCVar->GetInt() == 1)
 		{
 			FString Name = TEXT("Parameter Map For");
 			TSharedPtr<FNiagaraSchemaAction_NewNode> Action = AddNewNodeAction(NewActions, MenuCat, FText::FromString(Name), *Name, FText::GetEmpty());
diff --git a/Engine/Plugins/FX/Niagara/Source/NiagaraEditor/Private/NiagaraCompiler.cpp b/Engine/Plugins/FX/Niagara/Source/NiagaraEditor/Private/NiagaraCompiler.cpp
--- a/Engine/Plugins/FX/Niagara/Source/NiagaraEditor/Private/NiagaraCompiler.cpp
+++ b/Engine/Plugins/FX/Niagara/Source/NiagaraEditor/Private/NiagaraCompiler.cpp
@@ -31,6 +31,7 @@
 #include "ShaderCompiler.h"
 #include "NiagaraShader.h"
 #include "NiagaraScript.h"
+#include "NiagaraSimulationStageBase.h"
 #include "Serialization/MemoryReader.h"
 #include "HAL/ThreadSafeBool.h"
 
@@ -423,7 +424,7 @@ void FNiagaraCompileRequestData::AddRapidIterationParameters(const FNiagaraParam
 	}
 }
 
-void FNiagaraCompileRequestData::FinishPrecompile(UNiagaraScriptSource* ScriptSource, const TArray<FNiagaraVariable>& EncounterableVariables, ENiagaraScriptUsage InUsage, FCompileConstantResolver ConstantResolver)
+void FNiagaraCompileRequestData::FinishPrecompile(UNiagaraScriptSource* ScriptSource, const TArray<FNiagaraVariable>& EncounterableVariables, ENiagaraScriptUsage InUsage, FCompileConstantResolver ConstantResolver, const TArray<UNiagaraSimulationStageBase*>* SimStages)
 {
 	{
 		ENiagaraScriptCompileStatusEnum = StaticEnum<ENiagaraScriptCompileStatus>();
@@ -435,6 +436,7 @@ void FNiagaraCompileRequestData::FinishPrecompile(UNiagaraScriptSource* ScriptSo
 		NodeGraphDeepCopy->FindOutputNodes(OutputNodes);
 		PrecompiledHistories.Empty();
 
+		int32 NumSimStageNodes = 0;
 		for (UNiagaraNodeOutput* FoundOutputNode : OutputNodes)
 		{
 			// Map all for this output node
@@ -458,10 +460,40 @@ void FNiagaraCompileRequestData::FinishPrecompile(UNiagaraScriptSource* ScriptSo
 				}
 			}
 
+			if (FoundOutputNode->GetUsage() == ENiagaraScriptUsage::ParticleSimulationStageScript)
+			{
+				NumSimStageNodes++;
+			}
+
 			PrecompiledHistories.Append(Histories);
 			Builder.EndTranslation(TranslationName);
 		}
 
+		if (SimStages && NumSimStageNodes)
+		{
+			SpawnOnlyPerStage.AddZeroed(NumSimStageNodes);
+			NumIterationsPerStage.AddZeroed(NumSimStageNodes);
+			IterationSourcePerStage.AddZeroed(NumSimStageNodes);
+			StageGuids.AddDefaulted(NumSimStageNodes);
+			StageNames.AddDefaulted(NumSimStageNodes);
+			int32 NumProvidedStages = SimStages->Num();
+
+			for (int32 i = 0; i < NumSimStageNodes && i < NumProvidedStages; i++)
+			{
+				UNiagaraSimulationStageGeneric* GenericStage = Cast<UNiagaraSimulationStageGeneric>((*SimStages)[i]);
+
+				if (GenericStage)
+				{
+					NumIterationsPerStage[i] = GenericStage->Iterations;
+					IterationSourcePerStage[i] = GenericStage->IterationSource == ENiagaraIterationSource::DataInterface ? GenericStage->DataInterface.BoundVariable.GetName() : FName();
+					SpawnOnlyPerStage[i] = GenericStage->bSpawnOnly;
+					StageGuids[i] = GenericStage->Script->GetUsageId();
+					StageNames[i] = GenericStage->SimulationStageName;
+
+				}
+			}
+		}
+
 
 		// Generate CDO's for any referenced data interfaces...
 		for (int32 i = 0; i < PrecompiledHistories.Num(); i++)
@@ -530,7 +562,7 @@ TSharedPtr<FNiagaraCompileRequestDataBase, ESPMode::ThreadSafe> FNiagaraEditorMo
 		UNiagaraScriptSource* Source = Cast<UNiagaraScriptSource>(Script->GetSource());
 		BasePtr->DeepCopyGraphs(Source, Script->GetUsage(), EmptyResolver);
 		const TArray<FNiagaraVariable> EncounterableVariables;
-		BasePtr->FinishPrecompile(Source, EncounterableVariables, Script->GetUsage(), EmptyResolver);
+		BasePtr->FinishPrecompile(Source, EncounterableVariables, Script->GetUsage(), EmptyResolver, nullptr);
 	}
 	else if (System)
 	{
@@ -569,7 +601,7 @@ TSharedPtr<FNiagaraCompileRequestDataBase, ESPMode::ThreadSafe> FNiagaraEditorMo
 		BasePtr->DeepCopyGraphs(Source, ENiagaraScriptUsage::SystemSpawnScript, FCompileConstantResolver());
 		BasePtr->AddRapidIterationParameters(System->GetSystemSpawnScript()->RapidIterationParameters, FCompileConstantResolver());
 		BasePtr->AddRapidIterationParameters(System->GetSystemUpdateScript()->RapidIterationParameters, FCompileConstantResolver());
-		BasePtr->FinishPrecompile(Source, EncounterableVars, ENiagaraScriptUsage::SystemSpawnScript, EmptyResolver);
+		BasePtr->FinishPrecompile(Source, EncounterableVars, ENiagaraScriptUsage::SystemSpawnScript, EmptyResolver, nullptr);
 
 		// Add the User and System variables that we did encounter to the list that emitters might also encounter.
 		BasePtr->GatherPreCompiledVariables(TEXT("User"), EncounterableVars);
@@ -592,7 +624,7 @@ TSharedPtr<FNiagaraCompileRequestDataBase, ESPMode::ThreadSafe> FNiagaraEditorMo
 						BasePtr->EmitterData[i]->AddRapidIterationParameters(EmitterScripts[ScriptIdx]->RapidIterationParameters, ConstantResolver);
 					}
 				}
-				BasePtr->EmitterData[i]->FinishPrecompile(Cast<UNiagaraScriptSource>(Handle.GetInstance()->GraphSource), EncounterableVars, ENiagaraScriptUsage::EmitterSpawnScript, ConstantResolver);
+				BasePtr->EmitterData[i]->FinishPrecompile(Cast<UNiagaraScriptSource>(Handle.GetInstance()->GraphSource), EncounterableVars, ENiagaraScriptUsage::EmitterSpawnScript, ConstantResolver, &Handle.GetInstance()->GetSimulationStages());
 				BasePtr->MergeInEmitterPrecompiledData(BasePtr->EmitterData[i].Get());
 			}
 		}
diff --git a/Engine/Plugins/FX/Niagara/Source/NiagaraEditor/Private/NiagaraEditorModule.cpp b/Engine/Plugins/FX/Niagara/Source/NiagaraEditor/Private/NiagaraEditorModule.cpp
--- a/Engine/Plugins/FX/Niagara/Source/NiagaraEditor/Private/NiagaraEditorModule.cpp
+++ b/Engine/Plugins/FX/Niagara/Source/NiagaraEditor/Private/NiagaraEditorModule.cpp
@@ -817,6 +817,12 @@ void FNiagaraEditorModule::StartupModule()
 		FNiagaraScriptHighlight::StaticStruct()->GetFName(),
 		FOnGetPropertyTypeCustomizationInstance::CreateStatic(&FNiagaraScriptHighlightDetails::MakeInstance));
 
+
+	PropertyModule.RegisterCustomPropertyTypeLayout(
+	    FNiagaraVariableDataInterfaceBinding::StaticStruct()->GetFName(),
+		FOnGetPropertyTypeCustomizationInstance::CreateStatic(&FNiagaraDataInterfaceBindingCustomization::MakeInstance)
+	);
+
 	FNiagaraEditorStyle::Initialize();
 	ReinitializeStyleCommand = IConsoleManager::Get().RegisterConsoleCommand(
 		TEXT("fx.NiagaraEditor.ReinitializeStyle"),
diff --git a/Engine/Plugins/FX/Niagara/Source/NiagaraEditor/Private/NiagaraEditorUtilities.cpp b/Engine/Plugins/FX/Niagara/Source/NiagaraEditor/Private/NiagaraEditorUtilities.cpp
--- a/Engine/Plugins/FX/Niagara/Source/NiagaraEditor/Private/NiagaraEditorUtilities.cpp
+++ b/Engine/Plugins/FX/Niagara/Source/NiagaraEditor/Private/NiagaraEditorUtilities.cpp
@@ -56,6 +56,7 @@
 #include "Widgets/Notifications/SNotificationList.h"
 #include "Styling/CoreStyle.h"
 #include "Framework/Notifications/NotificationManager.h"
+#include "NiagaraSimulationStageBase.h"
 
 #define LOCTEXT_NAMESPACE "FNiagaraEditorUtilities"
 
@@ -1392,6 +1393,16 @@ UNiagaraScript* FNiagaraEditorUtilities::GetScriptFromSystem(UNiagaraSystem& Sys
 					}
 				}
 			}
+			else if (UNiagaraScript::IsEquivalentUsage(Usage, ENiagaraScriptUsage::ParticleSimulationStageScript))
+			{
+				for (const UNiagaraSimulationStageBase* SimulationStage : ScriptEmitterHandle->GetInstance()->GetSimulationStages())
+				{
+					if (SimulationStage && SimulationStage->Script && SimulationStage->Script->GetUsageId() == UsageId)
+					{
+						return SimulationStage->Script;
+					}
+				}
+			}
 		}
 	}
 	return nullptr;
diff --git a/Engine/Plugins/FX/Niagara/Source/NiagaraEditor/Private/NiagaraGraph.cpp b/Engine/Plugins/FX/Niagara/Source/NiagaraEditor/Private/NiagaraGraph.cpp
--- a/Engine/Plugins/FX/Niagara/Source/NiagaraEditor/Private/NiagaraGraph.cpp
+++ b/Engine/Plugins/FX/Niagara/Source/NiagaraEditor/Private/NiagaraGraph.cpp
@@ -1787,8 +1787,6 @@ void UNiagaraGraph::RefreshParameterReferences() const
 		VariableToScriptVariable.Remove(Variable);
 	}
 	
-	static const auto UseShaderStagesCVar = IConsoleManager::Get().FindConsoleVariable(TEXT("fx.UseShaderStages"));
-	if (UseShaderStagesCVar->GetInt() == 1)
 	{
 		// Add the array indices to the parameters
 		// When a particle attribute is created we need access the corresponding RegisterIdx if we want to 
diff --git a/Engine/Plugins/FX/Niagara/Source/NiagaraEditor/Private/NiagaraHlslTranslator.cpp b/Engine/Plugins/FX/Niagara/Source/NiagaraEditor/Private/NiagaraHlslTranslator.cpp
--- a/Engine/Plugins/FX/Niagara/Source/NiagaraEditor/Private/NiagaraHlslTranslator.cpp
+++ b/Engine/Plugins/FX/Niagara/Source/NiagaraEditor/Private/NiagaraHlslTranslator.cpp
@@ -30,6 +30,7 @@
 #include "NiagaraEditorUtilities.h"
 #include "NiagaraEditorModule.h"
 #include "NiagaraNodeReroute.h"
+#include "NiagaraSimulationStageBase.h"
 
 #include "NiagaraFunctionLibrary.h"
 
@@ -607,8 +608,7 @@ FString FHlslNiagaraTranslator::BuildParameterMapHlslDefinitions(TArray<FNiagara
 		}
 	}
 
-	static const auto UseShaderStagesCVar = IConsoleManager::Get().FindConsoleVariable(TEXT("fx.UseShaderStages"));
-	if (UseShaderStagesCVar->GetInt() == 1)
+	if (GetUsesOldShaderStages() || GetUsesSimulationStages())
 	{
 		// Add the attribute indices to the list of unique variables
 		TArray<FString> RegisterNames;
@@ -713,6 +713,17 @@ FString FHlslNiagaraTranslator::BuildParameterMapHlslDefinitions(TArray<FNiagara
 	return HlslOutputString;
 }
 
+
+bool FHlslNiagaraTranslator::GetUsesOldShaderStages() const
+{
+	return CompileOptions.AdditionalDefines.Contains(TEXT("Emitter.UseOldShaderStages"));
+}
+
+bool FHlslNiagaraTranslator::GetUsesSimulationStages() const
+{
+	return CompileOptions.AdditionalDefines.Contains(TEXT("Emitter.UseSimulationStages"));
+}
+
 bool FHlslNiagaraTranslator::ShouldConsiderTargetParameterMap(ENiagaraScriptUsage InUsage) const
 {
 	ENiagaraScriptUsage TargetUsage = GetTargetUsage();
@@ -824,6 +835,7 @@ const FNiagaraTranslateResults &FHlslNiagaraTranslator::Translate(const FNiagara
 	}
 
 	bool bRequiresPersistentIDs = CompileOptions.AdditionalDefines.Contains(TEXT("RequiresPersistentIDs"));
+	bool bUsesSimStages = (GetUsesOldShaderStages() || GetUsesSimulationStages());
 
 	TranslationStages.Empty();
 	ActiveStageIdx = 0;
@@ -832,6 +844,8 @@ const FNiagaraTranslateResults &FHlslNiagaraTranslator::Translate(const FNiagara
 	ParamMapHistories.Empty();
 	ParamMapSetVariablesToChunks.Empty();
 
+	OtherOutputParamMapHistories = CompileData->GetPrecomputedHistories();
+
 	switch (CompileOptions.TargetUsage)
 	{
 	case ENiagaraScriptUsage::ParticleSpawnScriptInterpolated:
@@ -844,6 +858,10 @@ const FNiagaraTranslateResults &FHlslNiagaraTranslator::Translate(const FNiagara
 		TranslationStages[0].OutputNode = SourceGraph->FindEquivalentOutputNode(ENiagaraScriptUsage::ParticleSpawnScript, TranslationStages[0].UsageId);
 		TranslationStages[1].OutputNode = SourceGraph->FindEquivalentOutputNode(TranslationStages[1].ScriptUsage, TranslationStages[1].UsageId);
 		TranslationStages[1].bInterpolatePreviousParams = true;
+		TranslationStages[0].SimulationStageIndexMin = 0;
+		TranslationStages[0].SimulationStageIndexMax = 1;
+		TranslationStages[1].SimulationStageIndexMin = 0;
+		TranslationStages[1].SimulationStageIndexMax = 1;
 		ParamMapHistories.AddDefaulted(2);
 		ParamMapSetVariablesToChunks.AddDefaulted(2);
 		break;
@@ -857,14 +875,104 @@ const FNiagaraTranslateResults &FHlslNiagaraTranslator::Translate(const FNiagara
 		TranslationStages[0].OutputNode = SourceGraph->FindEquivalentOutputNode(ENiagaraScriptUsage::ParticleSpawnScript, TranslationStages[0].UsageId);
 		TranslationStages[1].OutputNode = SourceGraph->FindEquivalentOutputNode(TranslationStages[1].ScriptUsage, TranslationStages[1].UsageId);
 		TranslationStages[1].bInterpolatePreviousParams = bHasInterpolatedSpawn;
+		TranslationStages[0].SimulationStageIndexMin = 0;
+		TranslationStages[0].SimulationStageIndexMax = 1;
+		TranslationStages[1].SimulationStageIndexMin = 0;
+		TranslationStages[1].SimulationStageIndexMax = 1;
 		ParamMapHistories.AddDefaulted(2);
 		ParamMapSetVariablesToChunks.AddDefaulted(2);
+
+		{
+			int32 SimStageStartIndex = 1;
+			int32 SimStageIndex = 0;
+			
+			// OutputNode order in traversal doesn't necessarily match the stack ordering. Use the GUID order to define the actual stages.
+			TArray<const UNiagaraNodeOutput*> FoundOutputNodes;
+			TArray<int32> FoundStageHistories;
+
+			for (const FGuid& StageGuid : InCompileData->StageGuids)
+			{
+				for (int32 FoundHistoryIdx = 0; FoundHistoryIdx < OtherOutputParamMapHistories.Num(); FoundHistoryIdx++)
+				{
+					FNiagaraParameterMapHistory& FoundHistory = OtherOutputParamMapHistories[FoundHistoryIdx];
+					const UNiagaraNodeOutput* HistoryOutputNode = FoundHistory.GetFinalOutputNode();
+					if (HistoryOutputNode && HistoryOutputNode->GetUsageId() == StageGuid)
+					{
+						FoundOutputNodes.Add(HistoryOutputNode);
+						FoundStageHistories.Add(FoundHistoryIdx);
+						break;
+					}
+				}
+			}
+
+			// Now iterate the nodes in the order we found them.
+			for (int32 FoundIdx = 0; FoundIdx < FoundOutputNodes.Num(); FoundIdx++)
+			{
+				const UNiagaraNodeOutput* HistoryOutputNode = FoundOutputNodes[FoundIdx];
+				const FNiagaraParameterMapHistory& FoundHistory = OtherOutputParamMapHistories[FoundStageHistories[FoundIdx]];
+
+				if (HistoryOutputNode && HistoryOutputNode->ScriptType == ENiagaraScriptUsage::ParticleSimulationStageScript)
+				{
+					bool bSpawnOnly = InCompileData->NumIterationsPerStage.Num() > SimStageIndex ? InCompileData->SpawnOnlyPerStage[SimStageIndex] : false;
+					int32 NumIterationsThisStage = InCompileData->NumIterationsPerStage.Num() > SimStageIndex ? InCompileData->NumIterationsPerStage[SimStageIndex] : 1;
+					FName IterationSrc = InCompileData->IterationSourcePerStage.Num() > SimStageIndex ? InCompileData->IterationSourcePerStage[SimStageIndex] : FName();
+					FString StageName = InCompileData->StageNames.Num() > SimStageIndex ? InCompileData->StageNames[SimStageIndex].ToString() : TEXT("Unnamed");
+					StageName = TEXT("_") + GetSanitizedFunctionNameSuffix(StageName);
+					int32 Index = TranslationStages.Add(FHlslNiagaraTranslationStage(HistoryOutputNode->ScriptType, HistoryOutputNode->ScriptTypeId));
+
+					TranslationStages[Index].PassNamespace = FString::Printf(TEXT("MapSimStage%d%s"), SimStageIndex + 1, *StageName);
+					TranslationStages[Index].ChunkModeIndex = (ENiagaraCodeChunkMode)(((int32)ENiagaraCodeChunkMode::SimulationStageBody) + (Index - 2));
+					if (TranslationStages[Index].ChunkModeIndex >= ENiagaraCodeChunkMode::SimulationStageBodyMax)
+					{
+						Error(FText::Format(LOCTEXT("TooManySimulationStages", "Cannot support more than %d simulation stages when adding %d!"),
+							FText::AsNumber((int32)ENiagaraCodeChunkMode::SimulationStageBodyMax - (int32)ENiagaraCodeChunkMode::SimulationStageBody),
+							FText::AsNumber((int32)TranslationStages[Index].ChunkModeIndex)), nullptr, nullptr);
+					}
+					TranslationStages[Index].OutputNode = SourceGraph->FindEquivalentOutputNode(TranslationStages[Index].ScriptUsage, TranslationStages[Index].UsageId);
+					ensure(TranslationStages[Index].OutputNode == HistoryOutputNode);
+					TranslationStages[Index].bInterpolatePreviousParams = false;
+					TranslationStages[Index].bCopyPreviousParams = false;
+					TranslationStages[Index].SimulationStageIndexMin = SimStageStartIndex;
+					TranslationStages[Index].SimulationStageIndexMax = SimStageStartIndex + NumIterationsThisStage;
+					TranslationStages[Index].NumIterationsThisStage = NumIterationsThisStage;
+					TranslationStages[Index].bSpawnOnly = bSpawnOnly;
+					TranslationStages[Index].IterationSource = IterationSrc;
+					TranslationStages[Index].SourceSimStage = SimStageIndex;
+					SimStageStartIndex += NumIterationsThisStage;
+					ParamMapHistories.AddDefaulted(1);
+
+					// Set up the compile output for the shader stages so that we can properly execute at runtime.
+					int32 IterationMetaIdx = CompilationOutput.ScriptData.SimulationStageMetaData.AddDefaulted();
+					CompilationOutput.ScriptData.SimulationStageMetaData[IterationMetaIdx].bSpawnOnly = bSpawnOnly;
+					CompilationOutput.ScriptData.SimulationStageMetaData[IterationMetaIdx].IterationSource = IterationSrc;
+					CompilationOutput.ScriptData.SimulationStageMetaData[IterationMetaIdx].MinStage = TranslationStages[Index].SimulationStageIndexMin;
+					CompilationOutput.ScriptData.SimulationStageMetaData[IterationMetaIdx].MaxStage = TranslationStages[Index].SimulationStageIndexMax;
+
+					// See if we write any "particle" attributes
+					for (const FNiagaraVariable& Var : FoundHistory.Variables)
+					{
+						if (FNiagaraParameterMapHistory::IsAttribute(Var))
+						{
+							CompilationOutput.ScriptData.SimulationStageMetaData[IterationMetaIdx].bWritesParticles = true;
+							break;
+						}
+					}
+
+					// Other outputs are written to as appropriate data interfaces are found. See HandleDataInterfaceCall for details.
+
+					ParamMapSetVariablesToChunks.AddDefaulted(1);
+					SimStageIndex++;
+				}
+			}
+		}
 		break;
 	default:
 		TranslationStages.Add(FHlslNiagaraTranslationStage(CompileOptions.TargetUsage, CompileOptions.TargetUsageId));
 		TranslationStages[0].PassNamespace = TEXT("Map");
 		TranslationStages[0].OutputNode = SourceGraph->FindEquivalentOutputNode(TranslationStages[0].ScriptUsage, TranslationStages[0].UsageId);
 		TranslationStages[0].ChunkModeIndex = ENiagaraCodeChunkMode::Body;
+		TranslationStages[0].SimulationStageIndexMin = 0;
+		TranslationStages[0].SimulationStageIndexMax = 0;
 		ParamMapHistories.AddDefaulted(1);
 		ParamMapSetVariablesToChunks.AddDefaulted(1);
 		break;
@@ -898,8 +1006,6 @@ const FNiagaraTranslateResults &FHlslNiagaraTranslator::Translate(const FNiagara
 
 
 	// Get all the parameter map histories traced to this graph from output nodes. We'll revisit this shortly in order to build out just the ones we care about for this translation.
-	OtherOutputParamMapHistories = CompileData->GetPrecomputedHistories();
-
 	if (ParamMapHistories.Num() == 1 && OtherOutputParamMapHistories.Num() == 1 && (CompileOptions.TargetUsage == ENiagaraScriptUsage::Function || CompileOptions.TargetUsage == ENiagaraScriptUsage::DynamicInput))
 	{
 		ParamMapHistories[0] = (OtherOutputParamMapHistories[0]);
@@ -1053,6 +1159,15 @@ const FNiagaraTranslateResults &FHlslNiagaraTranslator::Translate(const FNiagara
 				TranslationStages[i].OutputNode->Compile(ThisTranslator, OutputChunks);
 				AddBodyComment(TEXT("//End Update Script!\n\n"));
 			}
+			else if (TranslationStages[i].ScriptUsage == ENiagaraScriptUsage::ParticleSimulationStageScript)
+			{
+				AddBodyComment(FString::Printf(TEXT("//Begin Stage Script: %s!"), *TranslationStages[i].PassNamespace));
+				//Now we compile the simulation stage and read from the temp values written above.
+				CurrentParamMapIndices.Empty();
+				CurrentParamMapIndices.Add(1);
+				TranslationStages[i].OutputNode->Compile(ThisTranslator, OutputChunks);
+				AddBodyComment(FString::Printf(TEXT("//End Simulation Stage Script: %s\n\n"), *TranslationStages[i].PassNamespace));
+			}
 		}
 		CurrentBodyChunkMode = ENiagaraCodeChunkMode::Body;
 	}
@@ -1839,8 +1954,7 @@ void FHlslNiagaraTranslator::DefineMainGPUFunctions(
 	const TMap<FNiagaraDataSetID, int32>& DataSetReads,
 	const TMap<FNiagaraDataSetID, int32>& DataSetWrites)
 {
-	static const auto UseShaderStagesCVar = IConsoleManager::Get().FindConsoleVariable(TEXT("fx.UseShaderStages"));
-	const bool bUseShaderStages = UseShaderStagesCVar->GetInt() == 1;
+	const bool bUseSimulationStages = (GetUsesOldShaderStages() || GetUsesSimulationStages());
  
 	TArray<FNiagaraDataSetID> ReadDataSetIDs;
 	TArray<FNiagaraDataSetID> WriteDataSetIDs;
@@ -1945,48 +2059,76 @@ void FHlslNiagaraTranslator::DefineMainGPUFunctions(
 	// LoadUpdateVariables()
 	HlslOutput += TEXT("void LoadUpdateVariables(inout FSimulationContext Context, int InstanceIdx)\n{\n");
 	{
-		FString ContextName = TEXT("\tContext.Map.");
-		if (TranslationStages.Num() > 1) // Last context is "MapUpdate"
-		{
-			ContextName = FString::Printf(TEXT("\tContext.%s."), *TranslationStages.Last().PassNamespace);
-		}
-
-		for (int32 DataSetIndex = 0, IntCounter = 0, FloatCounter = 0; DataSetIndex < DataSetReads.Num(); ++DataSetIndex)
+		int32 StartIdx = 1;
+		for (int32 i = StartIdx; i < TranslationStages.Num(); i++)
 		{
-			const FNiagaraDataSetID DataSetID = ReadDataSetIDs[DataSetIndex];
-			const TArray<FNiagaraVariable>& NiagaraVariables = DataSetVariables[DataSetReads[DataSetID]];
-			for (const FNiagaraVariable& Var : NiagaraVariables)
+			if (TranslationStages.Num() > 2)
 			{
-				const FString VarName = ContextName + GetSanitizedSymbolName(Var.GetName().ToString());
-				FString VarFmt;
-
-				// If the NiagaraClearEachFrame value is set on the data set, we don't bother reading it in each frame as we know that it is is invalid. However,
-				// this is only used for the base data set. Other reads are potentially from events and are therefore perfectly valid.
-				if (DataSetIndex == 0 && Var.GetType().GetScriptStruct() != nullptr && Var.GetType().GetScriptStruct()->GetMetaData(TEXT("NiagaraClearEachFrame")).Equals(TEXT("true"), ESearchCase::IgnoreCase))
+				if (StartIdx == i)
 				{
-					VarFmt = VarName + TEXT("{0} = {4};\n");
+					HlslOutput += TEXT("\tBRANCH\n\tif ");
 				}
 				else
 				{
-					VarFmt = VarName + TEXT("{0} = InputData{1}({2}, {3}, InstanceIdx);\n");
+					HlslOutput += TEXT("\telse if ");
+				}
+
+				HlslOutput += FString::Printf(TEXT("(SimulationStageIndex >= %d && SimulationStageIndex < %d)\n"
+					"	{\n"), TranslationStages[i].SimulationStageIndexMin, TranslationStages[i].SimulationStageIndexMax);
+			}
+			else
+			{
+				// No need to add anything here..
+			}
+
+			FString ContextName = TEXT("\t\tContext.Map.");
+			if (TranslationStages.Num() > 1) // Second context is "MapUpdate"
+			{
+				ContextName = FString::Printf(TEXT("\t\tContext.%s."), *TranslationStages[i].PassNamespace);
+			}
+
+			for (int32 DataSetIndex = 0, IntCounter = 0, FloatCounter = 0; DataSetIndex < DataSetReads.Num(); ++DataSetIndex)
+			{
+				const FNiagaraDataSetID DataSetID = ReadDataSetIDs[DataSetIndex];
+				const TArray<FNiagaraVariable>& NiagaraVariables = DataSetVariables[DataSetReads[DataSetID]];
+				for (const FNiagaraVariable& Var : NiagaraVariables)
+				{
+					const FString VarName = ContextName + GetSanitizedSymbolName(Var.GetName().ToString());
+					FString VarFmt;
 
-					if (bUseShaderStages)
+					// If the NiagaraClearEachFrame value is set on the data set, we don't bother reading it in each frame as we know that it is is invalid. However,
+					// this is only used for the base data set. Other reads are potentially from events and are therefore perfectly valid.
+					if (DataSetIndex == 0 && Var.GetType().GetScriptStruct() != nullptr && Var.GetType().GetScriptStruct()->GetMetaData(TEXT("NiagaraClearEachFrame")).Equals(TEXT("true"), ESearchCase::IgnoreCase))
 					{
-						if (FNiagaraParameterMapHistory::IsAttribute(Var))
+						VarFmt = VarName + TEXT("{0} = {4};\n");
+					}
+					else
+					{
+						VarFmt = VarName + TEXT("{0} = InputData{1}({2}, {3}, InstanceIdx);\n");
+
+						if (bUseSimulationStages)
 						{
-							FString RegisterName = VarName;
-							RegisterName.ReplaceInline(PARAM_MAP_ATTRIBUTE_STR, PARAM_MAP_INDICES_STR);
-							const int32 RegisterValue = Var.GetType().IsFloatPrimitive() ? FloatCounter : IntCounter;
-							HlslOutput += RegisterName + FString::Printf(TEXT(" = %d;\n"), RegisterValue);
+							if (FNiagaraParameterMapHistory::IsAttribute(Var))
+							{
+								FString RegisterName = VarName;
+								RegisterName.ReplaceInline(PARAM_MAP_ATTRIBUTE_STR, PARAM_MAP_INDICES_STR);
+								const int32 RegisterValue = Var.GetType().IsFloatPrimitive() ? FloatCounter : IntCounter;
+								HlslOutput += RegisterName + FString::Printf(TEXT(" = %d;\n"), RegisterValue);
+							}
 						}
 					}
+					GatherVariableForDataSetAccess(Var, VarFmt, IntCounter, FloatCounter, DataSetIndex, TEXT(""), HlslOutput);
 				}
-				GatherVariableForDataSetAccess(Var, VarFmt, IntCounter, FloatCounter, DataSetIndex, TEXT(""), HlslOutput);
 			}
-		}
-		if (bUsesAlive)
-		{
-			HlslOutput += TEXT("\n") + ContextName + TEXT("DataInstance.Alive=true;\n");
+			if (bUsesAlive)
+			{
+				HlslOutput += ContextName + TEXT("DataInstance.Alive=true;\n");
+			}
+
+			if (TranslationStages.Num() > 2)
+			{
+				HlslOutput += TEXT("\t}\n");
+			}
 		}
 	}
 	HlslOutput += TEXT("}\n\n");
@@ -2004,63 +2146,152 @@ void FHlslNiagaraTranslator::DefineMainGPUFunctions(
 
 	//////////////////////
 	// TransferAttibutes()
+
+	
 	HlslOutput += TEXT("void TransferAttributes(inout FSimulationContext Context)\n{\n");
 	{
-		if (TranslationStages.Last().bCopyPreviousParams)
+		int32 StartIdx = 1;
+		for (int32 i = StartIdx; i < TranslationStages.Num(); i++)
 		{
-			if (ParamMapDefinedAttributesToNamespaceVars.Num() != 0)
+			if (TranslationStages.Num() > 2)
 			{
-				HlslOutput += TEXT("\tContext.") + TranslationStages.Last().PassNamespace + TEXT(".Particles = Context.") + TranslationStages[0].PassNamespace + TEXT(".Particles;\n");
-				if (bUsesAlive)
+				if (StartIdx == i)
 				{
-					HlslOutput += TEXT("\tContext.") + TranslationStages.Last().PassNamespace + TEXT(".DataInstance = Context.") + TranslationStages[0].PassNamespace + TEXT(".DataInstance;\n");
+					HlslOutput += TEXT("\tBRANCH\n\tif "); 
 				}
+				else
+				{
+					HlslOutput += TEXT("\telse if ");
+				}
+
+				HlslOutput += FString::Printf(TEXT("(SimulationStageIndex >= %d && SimulationStageIndex < %d)\n"
+					"	{\n"), TranslationStages[i].SimulationStageIndexMin, TranslationStages[i].SimulationStageIndexMax );
+			}
+
+			if (TranslationStages[i].bCopyPreviousParams)
+			{
+				if (ParamMapDefinedAttributesToNamespaceVars.Num() != 0)
+				{
+					HlslOutput += TEXT("\t\tContext.") + TranslationStages[i].PassNamespace + TEXT(".Particles = Context.") + TranslationStages[i - 1].PassNamespace + TEXT(".Particles;\n");
+					if (bUsesAlive)
+					{
+						HlslOutput += TEXT("\t\tContext.") + TranslationStages[i].PassNamespace + TEXT(".DataInstance = Context.") + TranslationStages[i - 1].PassNamespace + TEXT(".DataInstance;\n");
+					}
+				}
+			}
+			
+			if (TranslationStages.Num() > 2)
+			{
+				HlslOutput += TEXT("\t}\n");
 			}
 		}
 	}
 	HlslOutput += TEXT("}\n\n");
-
+	
 	/////////////////////////
 	// StoreUpdateVariables()
+
 	HlslOutput += TEXT("void StoreUpdateVariables(in FSimulationContext Context)\n{\n");
 	{
-		if (bUsesAlive)
+		int32 StartIdx = 1;
+		for (int32 i = StartIdx; i < TranslationStages.Num(); i++)
 		{
-			HlslOutput += TEXT("\tconst bool bValid = Context.") + TranslationStages.Last().PassNamespace + TEXT(".DataInstance.Alive;\n");
-			HlslOutput += TEXT("\tconst int WriteIndex = OutputIndex(0, false, bValid);\n");
-		}
-		else
-		{
-			HlslOutput += TEXT("\tconst bool bValid = GCurrentPhase != -1;\n");
-			HlslOutput += TEXT("\tconst int WriteIndex = OutputIndex(0, true, bValid);\n");
-		}
+			if (TranslationStages.Num() > 2)
+			{
+				if (StartIdx == i)
+				{
+					HlslOutput += TEXT("\tBRANCH\n\tif ");
+				}
+				else
+				{
+					HlslOutput += TEXT("\telse if ");
+				}
 
-		FString ContextName = TEXT("Context.Map.");
-		if (TranslationStages.Num() > 1) // Last context is "MapUpdate"
-		{
-			ContextName = FString::Printf(TEXT("Context.%s."), *TranslationStages.Last().PassNamespace);
-		}
+				HlslOutput += FString::Printf(TEXT("(SimulationStageIndex >= %d && SimulationStageIndex < %d)\n"
+					"	{\n"), TranslationStages[i].SimulationStageIndexMin, TranslationStages[i].SimulationStageIndexMax);
+			}
 
-		HlslOutput += TEXT("\tif (bValid)\n\t{\n");
+			if (bUsesAlive)
+			{
+				HlslOutput += TEXT("\t\tconst bool bValid = Context.") + TranslationStages[i].PassNamespace + TEXT(".DataInstance.Alive;\n");
+				HlslOutput += TEXT("\t\tconst int WriteIndex = OutputIndex(0, false, bValid);\n");
+			}
+			else
+			{
+				HlslOutput += TEXT("\t\tconst bool bValid = GCurrentPhase != -1;\n");
+				HlslOutput += TEXT("\t\tconst int WriteIndex = OutputIndex(0, true, bValid);\n");
+			}
+
+			FString ContextName = TEXT("Context.Map.");
+			if (TranslationStages.Num() > 1) // Last context is "MapUpdate"
+			{
+				ContextName = FString::Printf(TEXT("Context.%s."), *TranslationStages[i].PassNamespace);
+			}
+
+			HlslOutput += TEXT("\t\tif (bValid)\n\t\t{\n");
+
+			if (bRequiresPersistentIDs)
+			{
+				HlslOutput += FString::Printf(TEXT("\t\t\tUpdateID(0, %sParticles.ID.Index, WriteIndex);\n"), *ContextName);
+			}
+
+			for (int32 DataSetIndex = 0, IntCounter = 0, FloatCounter = 0; DataSetIndex < DataSetWrites.Num(); ++DataSetIndex)
+			{
+				const FNiagaraDataSetID DataSetID = ReadDataSetIDs[DataSetIndex];
+				const TArray<FNiagaraVariable>& NiagaraVariables = DataSetVariables[DataSetWrites[DataSetID]];
+				for (const FNiagaraVariable& Var : NiagaraVariables)
+				{
+					// If coming from a parameter map, use the one on the context, otherwise use the output.
+					FString VarFmt = TEXT("\t\t\tOutputData{1}(0, {2}, {3}, ") + ContextName + GetSanitizedSymbolName(Var.GetName().ToString()) + TEXT("{0});\n");
+					GatherVariableForDataSetAccess(Var, VarFmt, IntCounter, FloatCounter, -1, TEXT("WriteIndex"), HlslOutput);
+				}
+			}
+
+			HlslOutput += TEXT("\t\t}\n");
+
+			if (TranslationStages.Num() > 2)
+			{
+				HlslOutput += TEXT("\t}\n");
+			}
 
-		if (bRequiresPersistentIDs)
-		{
-			HlslOutput += FString::Printf(TEXT("\t\tUpdateID(0, %sParticles.ID.Index, WriteIndex);\n"), *ContextName);
 		}
+	}
+	HlslOutput += TEXT("\n}\n\n");
+
+	/////////////////////////
+	// SimulateDoWork()
 
-		for (int32 DataSetIndex = 0, IntCounter = 0, FloatCounter = 0; DataSetIndex < DataSetWrites.Num(); ++DataSetIndex)
+	//HlslOutput += TEXT("void SimulateDoWork(in FSimulationContext Context)\n{\n");
+	FString SimDoWorkString;// = TEXT("\t{ \n");
+	{
+		int32 StartIdx = 1;
+		for (int32 i = StartIdx; i < TranslationStages.Num(); i++)
 		{
-			const FNiagaraDataSetID DataSetID = ReadDataSetIDs[DataSetIndex];
-			const TArray<FNiagaraVariable>& NiagaraVariables = DataSetVariables[DataSetWrites[DataSetID]];
-			for (const FNiagaraVariable& Var : NiagaraVariables)
+			if (TranslationStages.Num() > 2)
 			{
-				// If coming from a parameter map, use the one on the context, otherwise use the output.
-				FString VarFmt = TEXT("\t\tOutputData{1}(0, {2}, {3}, ") + ContextName + GetSanitizedSymbolName(Var.GetName().ToString()) + TEXT("{0});\n");
-				GatherVariableForDataSetAccess(Var, VarFmt, IntCounter, FloatCounter, -1, TEXT("WriteIndex"), HlslOutput);
+				if (StartIdx == i)
+				{
+					SimDoWorkString += TEXT("\t\t//BRANCH\n\t\tif (SimulationStageIndex == 0)\n"
+						"\t\t{\n");
+					SimDoWorkString += TEXT("\t\t\tSimulate") + TranslationStages[i].PassNamespace + TEXT("(Context); \n");
+				}
+				else
+				{
+					SimDoWorkString += TEXT("\t\telse if ");
+					SimDoWorkString += FString::Printf(TEXT("(SimulationStageIndex >= %d && SimulationStageIndex < %d)\n"
+						"\t\t{\n"), TranslationStages[i].SimulationStageIndexMin, TranslationStages[i].SimulationStageIndexMax);
+
+					SimDoWorkString += TEXT("\t\t\tSimulate") + TranslationStages[i].PassNamespace + TEXT("(Context); \n");
+				}
+				SimDoWorkString += TEXT("\t\t}\n");
+			}
+			else
+			{
+				SimDoWorkString += TEXT("\t\tSimulate") + TranslationStages[i].PassNamespace + TEXT("(Context); \n");
 			}
 		}
+		//SimDoWorkString += TEXT("\t}\n");
 	}
-	HlslOutput += TEXT("\t}\n}\n\n");
 
 	/////////////////////////
 	// CopyInstance()
@@ -2094,6 +2325,107 @@ void FHlslNiagaraTranslator::DefineMainGPUFunctions(
 #endif
 	}
 	HlslOutput += TEXT("}\n");
+	
+	HlslOutput +=
+		TEXT("\n\n/*\n"
+			"*	CS wrapper for our generated code; calls spawn and update functions on the corresponding instances in the buffer\n"
+			" */ \n"
+			"\n"
+			"[numthreads(THREADGROUP_SIZE, 1, 1)]\n"
+			"void SimulateMainComputeCS(\n"
+			"	uint3 GroupId : SV_GroupID, \n"
+			"	uint3 DispatchThreadId : SV_DispatchThreadID, \n"
+			"	uint3 GroupThreadId : SV_GroupThreadID)\n"
+			"{\n"
+			"	GDispatchThreadId = DispatchThreadId; \n"
+			"	GGroupThreadId = GroupThreadId; \n"
+			"	GCurrentPhase = -1; \n"
+			"	GEmitterTickCounter = EmitterTickCounter; \n"
+			"	GSimStart = SimStart; \n"
+			"	GRandomSeedOffset = 0; \n"
+			"	\n"
+			"	/*\n"
+			"	if(CopyInstancesBeforeStart == 1)\n"
+			"	{\n"
+			"		UpdateStartInstance = 0;\n"
+			"	}\n"
+			"	*/\n"
+			"	\n"
+			"	const uint InstanceID = UpdateStartInstance + DispatchThreadId.x; \n"
+			"	if (ReadInstanceCountOffset == 0xFFFFFFFF)\n"
+			"	{\n"
+			"			GSpawnStartInstance = 0; \n"
+			"	}\n"
+			"	else\n"
+			"	{\n"
+			"		GSpawnStartInstance = RWInstanceCounts[ReadInstanceCountOffset];				// needed by ExecIndex()\n"
+			"	}\n"
+			"	const int MaxInstances = GSpawnStartInstance + SpawnedInstances; \n"
+			"		\n"
+			"	bool bRunUpdateLogic = InstanceID < GSpawnStartInstance && InstanceID < UpdateStartInstance + MaxInstances; \n"
+			"	bool bRunSpawnLogic = InstanceID >= GSpawnStartInstance && InstanceID < UpdateStartInstance + MaxInstances; \n"
+			"		\n"
+			"	#if USE_SIMULATION_STAGES\n"
+			"	if (IterationInterfaceCount > 0)\n"
+			"	{\n"
+			"		bRunUpdateLogic = InstanceID < IterationInterfaceCount && GSimStart != 1; \n"
+			"		bRunSpawnLogic = InstanceID < IterationInterfaceCount && GSimStart == 1; \n"
+			"	}\n"
+			"	#endif // USE_SIMULATION_STAGES\n"
+			"		\n"
+			"	const float RandomSeedInitialisation = NiagaraInternalNoise(InstanceID * 16384, 0 * 8196, (bRunUpdateLogic ? 4096 : 0) + EmitterTickCounter);	// initialise the random state seed\n"
+			"	\n"
+			"	FSimulationContext Context = (FSimulationContext)0; \n"
+			"	\n"
+			"	BRANCH\n"
+			"	if (bRunUpdateLogic)\n"
+			"	{\n"
+			"		GCurrentPhase = GUpdatePhase; \n"
+			"		SetupExecIndexForGPU(); \n"
+			"		InitConstants(Context); \n"
+			"		LoadUpdateVariables(Context, InstanceID); \n"
+			"		ReadDataSets(Context); \n"
+		) + SimDoWorkString
+		+ TEXT(
+			"		//SimulateDoWork(Context); \n"
+			"		//SimulateMapUpdate(Context);\n"
+			"		WriteDataSets(Context); \n"
+			"	}\n"
+			"	else if (bRunSpawnLogic)\n"
+			"	{\n"
+			"		GCurrentPhase = GSpawnPhase; \n"
+			"		SetupExecIndexAndSpawnInfoForGPU(); \n"
+			"		InitConstants(Context); \n"
+			"		InitSpawnVariables(Context); \n"
+			"		ReadDataSets(Context); \n"
+			"		\n"
+			"		Context.MapSpawn.Particles.UniqueID = Engine_Emitter_TotalSpawnedParticles + ExecIndex(); \n"
+			"		ConditionalInterpolateParameters(Context); \n"
+			"		SimulateMapSpawn(Context); \n"
+			"		\n"
+			"		GCurrentPhase = GUpdatePhase; \n"
+			"		\n"
+			"		TransferAttributes(Context); \n"
+			"		\n"
+			"		//SimulateDoWork(Context); \n"
+			"		//SimulateMapUpdate(Context);\n"
+		) + SimDoWorkString
+		+ TEXT(
+			"		WriteDataSets(Context); \n"
+			"	}\n"
+			"	\n"
+			"	StoreUpdateVariables(Context); \n\n"
+			"	#if USE_SIMULATION_STAGES \n"
+			"	// We want to keep pushing along the counts so that we can  \n"
+			"	// use any of the write locations in the future. Only simulation stage 0 will increment the counters. \n"
+			"	if (SimulationStageIndex > 0 && GDispatchThreadId.x == 0) \n"
+			"	{ \n"
+			"		RWInstanceCounts[WriteInstanceCountOffset] = RWInstanceCounts[ReadInstanceCountOffset]; \n"
+			"	} \n"
+			"	#endif \n"	
+		);
+
+	HlslOutput += TEXT("}\n");
 }
 
 void FHlslNiagaraTranslator::DefineMain(FString &OutHlslOutput,
@@ -2388,8 +2720,7 @@ void FHlslNiagaraTranslator::DefineDataSetVariableReads(FString &OutHlslOutput,
 			{
 				Fmt = VariableName + TEXT("{0} = InputData{1}({2}, {3});\n");
 
-				static const auto UseShaderStagesCVar = IConsoleManager::Get().FindConsoleVariable(TEXT("fx.UseShaderStages"));
-				if (UseShaderStagesCVar->GetInt() == 1)
+				if ((GetUsesOldShaderStages() || GetUsesSimulationStages()))
 				{
 					if (FNiagaraParameterMapHistory::IsAttribute(Var))
 					{
@@ -2823,7 +3154,8 @@ int32 FHlslNiagaraTranslator::AddBodyChunk(const FString& Value)
 
 int32 FHlslNiagaraTranslator::AddBodyChunk(FString SymbolName, FString Definition, const FNiagaraTypeDefinition& Type, TArray<int32>& SourceChunks, bool bDecl, bool bIsTerminated)
 {
-	check(CurrentBodyChunkMode == ENiagaraCodeChunkMode::Body || CurrentBodyChunkMode == ENiagaraCodeChunkMode::SpawnBody || CurrentBodyChunkMode == ENiagaraCodeChunkMode::UpdateBody);
+	check(CurrentBodyChunkMode == ENiagaraCodeChunkMode::Body || CurrentBodyChunkMode == ENiagaraCodeChunkMode::SpawnBody || CurrentBodyChunkMode == ENiagaraCodeChunkMode::UpdateBody ||
+		(CurrentBodyChunkMode >= ENiagaraCodeChunkMode::SimulationStageBody && CurrentBodyChunkMode < ENiagaraCodeChunkMode::SimulationStageBodyMax));
 
 	int32 Ret = CodeChunks.AddDefaulted();
 	FNiagaraCodeChunk& Chunk = CodeChunks[Ret];
@@ -2843,7 +3175,8 @@ int32 FHlslNiagaraTranslator::AddBodyChunk(FString SymbolName, FString Definitio
 
 int32 FHlslNiagaraTranslator::AddBodyChunk(FString SymbolName, FString Definition, const FNiagaraTypeDefinition& Type, int32 SourceChunk, bool bDecl, bool bIsTerminated)
 {
-	check(CurrentBodyChunkMode == ENiagaraCodeChunkMode::Body || CurrentBodyChunkMode == ENiagaraCodeChunkMode::SpawnBody || CurrentBodyChunkMode == ENiagaraCodeChunkMode::UpdateBody);
+	check(CurrentBodyChunkMode == ENiagaraCodeChunkMode::Body || CurrentBodyChunkMode == ENiagaraCodeChunkMode::SpawnBody || CurrentBodyChunkMode == ENiagaraCodeChunkMode::UpdateBody ||
+		(CurrentBodyChunkMode >= ENiagaraCodeChunkMode::SimulationStageBody && CurrentBodyChunkMode < ENiagaraCodeChunkMode::SimulationStageBodyMax));
 
 	int32 Ret = CodeChunks.AddDefaulted();
 	FNiagaraCodeChunk& Chunk = CodeChunks[Ret];
@@ -2861,7 +3194,8 @@ int32 FHlslNiagaraTranslator::AddBodyChunk(FString SymbolName, FString Definitio
 
 int32 FHlslNiagaraTranslator::AddBodyChunk(FString SymbolName, FString Definition, const FNiagaraTypeDefinition& Type, bool bDecl, bool bIsTerminated)
 {
-	check(CurrentBodyChunkMode == ENiagaraCodeChunkMode::Body || CurrentBodyChunkMode == ENiagaraCodeChunkMode::SpawnBody || CurrentBodyChunkMode == ENiagaraCodeChunkMode::UpdateBody);
+	check(CurrentBodyChunkMode == ENiagaraCodeChunkMode::Body || CurrentBodyChunkMode == ENiagaraCodeChunkMode::SpawnBody || CurrentBodyChunkMode == ENiagaraCodeChunkMode::UpdateBody ||
+			(CurrentBodyChunkMode >= ENiagaraCodeChunkMode::SimulationStageBody && CurrentBodyChunkMode < ENiagaraCodeChunkMode::SimulationStageBodyMax));
 
 	int32 Ret = CodeChunks.AddDefaulted();
 	FNiagaraCodeChunk& Chunk = CodeChunks[Ret];
@@ -3805,9 +4139,10 @@ bool FHlslNiagaraTranslator::ParameterMapRegisterExternalConstantNamespaceVariab
 				const auto& SystemVar = ParamMapDefinedSystemVars.FindChecked(InVariable.GetName());
 				UniformChunk = SystemVar.ChunkIndex;
 			}
-			static const auto UseShaderStagesCVar = IConsoleManager::Get().FindConsoleVariable(TEXT("fx.UseShaderStages"));
+			
+			bool bUseSimulationStages = (GetUsesOldShaderStages() || GetUsesSimulationStages());
 			// Avoid overriding the register indices
-			if ((UseShaderStagesCVar->GetInt()==1 && !FNiagaraParameterMapHistory::IsInNamespace(InVariable, PARAM_MAP_INDICES_STR)) || UseShaderStagesCVar->GetInt()==0)
+			if ((bUseSimulationStages && !FNiagaraParameterMapHistory::IsInNamespace(InVariable, PARAM_MAP_INDICES_STR)) || !bUseSimulationStages)
 			{
 				//Add this separately as the same uniform can appear in the pre sim chunks more than once in different param maps.
 				MainPreSimulateChunks.AddUnique(FString::Printf(TEXT("%s.%s = %s;"), *ParameterMapInstanceName, *GetSanitizedSymbolName(VarName), *GetCodeAsSource(UniformChunk)));
@@ -5073,6 +5408,8 @@ void FHlslNiagaraTranslator::HandleCustomHlslNode(UNiagaraNodeCustomHlsl* Custom
 					Sig.OwnerName = Info.Name;
 					Info.RegisteredFunctions.Add(Sig);
 					Functions.FindOrAdd(Sig);
+
+					HandleDataInterfaceCall(Info, Sig);
 				}
 			}
 			SigInputs.Add(Input);
@@ -5177,6 +5514,28 @@ void FHlslNiagaraTranslator::HandleCustomHlslNode(UNiagaraNodeCustomHlsl* Custom
 	OutCustomHlsl = TEXT("\n") + OutCustomHlsl + TEXT("\n");
 }
 
+void FHlslNiagaraTranslator::HandleDataInterfaceCall(FNiagaraScriptDataInterfaceCompileInfo& Info, const FNiagaraFunctionSignature& InMatchingSignature)
+{
+	bool bCPUSim = IsCompileOptionDefined(TEXT("CPUSim"));
+	bool bGPUSim = IsCompileOptionDefined(TEXT("GPUComputeSim"));
+	const UNiagaraNode* CurNode = ActiveHistoryForFunctionCalls.GetCallingContext();
+	if (bCPUSim && !InMatchingSignature.bSupportsCPU)
+	{
+		Error(FText::Format(LOCTEXT("FunctionCallDataInterfaceCPUMissing", "Function call \"{0}\" does not work on CPU sims."), FText::FromName(InMatchingSignature.Name)), CurNode, nullptr);
+	}
+	else if (bGPUSim && !InMatchingSignature.bSupportsGPU)
+	{
+		Error(FText::Format(LOCTEXT("FunctionCallDataInterfaceGPUMissing", "Function call \"{0}\" does not work on GPU sims."), FText::FromName(InMatchingSignature.Name)), CurNode, nullptr);
+	}
+
+	if (InMatchingSignature.bWriteFunction && CompilationOutput.ScriptData.SimulationStageMetaData.Num() != 0 && TranslationStages[ActiveStageIdx].SourceSimStage != -1)
+	{
+		int32 SourceSimStage = TranslationStages[ActiveStageIdx].SourceSimStage;
+		ensure(CompilationOutput.ScriptData.SimulationStageMetaData.Num() > SourceSimStage);
+		CompilationOutput.ScriptData.SimulationStageMetaData[SourceSimStage].OutputDestinations.AddUnique(Info.Name);
+	}
+}
+
 void FHlslNiagaraTranslator::RegisterFunctionCall(ENiagaraScriptUsage ScriptUsage, const FString& InName, const FString& InFullName, const FGuid& CallNodeId, UNiagaraScriptSource* Source, FNiagaraFunctionSignature& InSignature, bool bIsCustomHlsl, const FString& InCustomHlsl, TArray<int32>& Inputs, const TArray<UEdGraphPin*>& CallInputs, const TArray<UEdGraphPin*>& CallOutputs,
 	FNiagaraFunctionSignature& OutSignature)
 {
@@ -5254,8 +5613,8 @@ void FHlslNiagaraTranslator::RegisterFunctionCall(ENiagaraScriptUsage ScriptUsag
 			FString MinParamSpawn;
 			FString MaxParamSpawn;
 
-			static const auto UseShaderStagesCVar = IConsoleManager::Get().FindConsoleVariable(TEXT("fx.UseShaderStages"));
-			const bool UseShaderStages = (UseShaderStagesCVar->GetInt() == 1) && CompilationTarget == ENiagaraSimTarget::GPUComputeSim;
+			const bool UseSimulationStages = ((GetUsesOldShaderStages() || GetUsesSimulationStages())) && CompilationTarget == ENiagaraSimTarget::GPUComputeSim;
+			const bool UseOldShaderStages = GetUsesOldShaderStages() && CompilationTarget == ENiagaraSimTarget::GPUComputeSim;
 
 			//We've not compiled this function yet so compile it now.
 			EnterFunction(InName, OutSignature, Inputs, CallNodeId);
@@ -5305,7 +5664,7 @@ void FHlslNiagaraTranslator::RegisterFunctionCall(ENiagaraScriptUsage ScriptUsag
 								const FNiagaraVariable& AliasedVar = History.VariablesWithOriginalAliasesIntact[VarIdx];
 								bool bIsAliased = Var.GetName() != AliasedVar.GetName();
 
-								if (bIsAliased && UseShaderStages)
+								if (bIsAliased && UseOldShaderStages)
 								{
 									if (Var.GetName().ToString().Contains(TEXT("MinStage")))
 									{
@@ -5376,17 +5735,17 @@ void FHlslNiagaraTranslator::RegisterFunctionCall(ENiagaraScriptUsage ScriptUsag
 
 				if (bIsModuleFunction)
 				{
-					if (UseShaderStages)
+					if (UseOldShaderStages)
 					{
 						if (bStageMinFilter && bStageMaxFilter)
 						{
-							FunctionDefStr += TEXT("if ((GCurrentPhase == 1 && ShaderStageIndex >= ") + MinParam + TEXT(" && ShaderStageIndex <= ") + MaxParam + TEXT(") || ") +
-								TEXT("(GCurrentPhase == 0 && ShaderStageIndex >= ") + MinParamSpawn + TEXT(" && ShaderStageIndex <= ") + MaxParamSpawn + TEXT(")") +
+							FunctionDefStr += TEXT("if ((GCurrentPhase == GUpdatePhase && SimulationStageIndex >= ") + MinParam + TEXT(" && SimulationStageIndex <= ") + MaxParam + TEXT(") || ") +
+								TEXT("(GCurrentPhase == GSpawnPhase && SimulationStageIndex >= ") + MinParamSpawn + TEXT(" && SimulationStageIndex <= ") + MaxParamSpawn + TEXT(")") +
 								TEXT(")\n{\n");
 						}
 						else
 						{
-							FunctionDefStr += TEXT("if ((GCurrentPhase == 0 && ShaderStageIndex == 0) || (GCurrentPhase == 1 && ShaderStageIndex == DefaultShaderStageIndex))\n{\n");
+							FunctionDefStr += TEXT("if ((GCurrentPhase == GSpawnPhase && SimulationStageIndex == 0) || (GCurrentPhase == GUpdatePhase && SimulationStageIndex == DefaultSimulationStageIndex))\n{\n");
 						}
 					}
 				}
@@ -5398,7 +5757,7 @@ void FHlslNiagaraTranslator::RegisterFunctionCall(ENiagaraScriptUsage ScriptUsag
 
 				if (bIsModuleFunction)
 				{
-					if (UseShaderStages)
+					if (UseOldShaderStages)
 					{
 						FunctionDefStr += TEXT("}\n");
 					}
@@ -5536,16 +5895,7 @@ void FHlslNiagaraTranslator::RegisterFunctionCall(ENiagaraScriptUsage ScriptUsag
 				}
 				else
 				{
-					bool bCPUSim = IsCompileOptionDefined(TEXT("CPUSim"));
-					bool bGPUSim = IsCompileOptionDefined(TEXT("GPUComputeSim"));
-					if (bCPUSim && !DataInterfaceFunctions[FoundMatch].bSupportsCPU)
-					{
-						Error(FText::Format(LOCTEXT("FunctionCallDataInterfaceCPUMissing", "Function call \"{0}\" does not work on CPU sims."), FText::FromName(OutSignature.Name)), nullptr, nullptr);
-					}
-					else if (bGPUSim && !DataInterfaceFunctions[FoundMatch].bSupportsGPU)
-					{
-						Error(FText::Format(LOCTEXT("FunctionCallDataInterfaceGPUMissing", "Function call \"{0}\" does not work on GPU sims."), FText::FromName(OutSignature.Name)), nullptr, nullptr);
-					}
+					HandleDataInterfaceCall(Info, DataInterfaceFunctions[FoundMatch]);
 				}
 
 				if (Info.UserPtrIdx != INDEX_NONE && CompilationTarget != ENiagaraSimTarget::GPUComputeSim)
diff --git a/Engine/Plugins/FX/Niagara/Source/NiagaraEditor/Private/NiagaraHlslTranslator.h b/Engine/Plugins/FX/Niagara/Source/NiagaraEditor/Private/NiagaraHlslTranslator.h
--- a/Engine/Plugins/FX/Niagara/Source/NiagaraEditor/Private/NiagaraHlslTranslator.h
+++ b/Engine/Plugins/FX/Niagara/Source/NiagaraEditor/Private/NiagaraHlslTranslator.h
@@ -101,7 +101,7 @@ public:
 	const FString& GetUniqueEmitterName() const { return EmitterUniqueName; }
 	void VisitReferencedGraphs(UNiagaraGraph* InSrcGraph, UNiagaraGraph* InDupeGraph, ENiagaraScriptUsage InUsage, FCompileConstantResolver ConstantResolver);
 	void DeepCopyGraphs(UNiagaraScriptSource* ScriptSource, ENiagaraScriptUsage InUsage, FCompileConstantResolver ConstantResolver);
-	void FinishPrecompile(UNiagaraScriptSource* ScriptSource, const TArray<FNiagaraVariable>& EncounterableVariables, ENiagaraScriptUsage InUsage, FCompileConstantResolver ConstantResolver);
+	void FinishPrecompile(UNiagaraScriptSource* ScriptSource, const TArray<FNiagaraVariable>& EncounterableVariables, ENiagaraScriptUsage InUsage, FCompileConstantResolver ConstantResolver, const TArray<class UNiagaraSimulationStageBase*>* SimStages);
 	virtual int32 GetDependentRequestCount() const override {
 		return EmitterData.Num();
 	};
@@ -111,6 +111,13 @@ public:
 	void AddRapidIterationParameters(const FNiagaraParameterStore& InParamStore, FCompileConstantResolver InResolver);
 	virtual bool GetUseRapidIterationParams() const override { return bUseRapidIterationParams; }
 
+	// Simulation Stage Variables. Sim stage of 0 is always Spawn/Update
+	TArray<uint32> NumIterationsPerStage;
+	TArray<FName> IterationSourcePerStage;
+	TArray<bool> SpawnOnlyPerStage;
+	TArray<FGuid> StageGuids;
+	TArray<FName> StageNames;
+
 	// If this is being held onto for any length of time, make sure to hold onto it in a gc-aware object. Right now in this information-passing struct,
 	// we could have a leaked garbage collected pointer if not held onto by someone capable of registering a reference.
 	UNiagaraGraph* NodeGraphDeepCopy;
@@ -188,6 +195,8 @@ enum class ENiagaraCodeChunkMode : uint8
 	SpawnBody,
 	UpdateBody,
 	InitializerBody,
+	SimulationStageBody,
+	SimulationStageBodyMax = SimulationStageBody + 100,
 	Num,
 };
 
@@ -291,7 +300,7 @@ public:
 class NIAGARAEDITOR_API FHlslNiagaraTranslationStage
 {
 public:
-	FHlslNiagaraTranslationStage(ENiagaraScriptUsage InScriptUsage, FGuid InUsageId) : ScriptUsage(InScriptUsage), UsageId(InUsageId), OutputNode(nullptr), bInterpolatePreviousParams(false), bCopyPreviousParams(true), ChunkModeIndex((ENiagaraCodeChunkMode)-1){}
+	FHlslNiagaraTranslationStage(ENiagaraScriptUsage InScriptUsage, FGuid InUsageId) : ScriptUsage(InScriptUsage), UsageId(InUsageId), OutputNode(nullptr), bInterpolatePreviousParams(false), bCopyPreviousParams(true), ChunkModeIndex((ENiagaraCodeChunkMode)-1), IterationSource(), bSpawnOnly(false){}
 
 	ENiagaraScriptUsage ScriptUsage;
 	FGuid UsageId;
@@ -300,6 +309,12 @@ public:
 	bool bInterpolatePreviousParams;
 	bool bCopyPreviousParams;
 	ENiagaraCodeChunkMode ChunkModeIndex;
+	FName IterationSource;
+	int32 SimulationStageIndexMin = -1;
+	int32 SimulationStageIndexMax = -1;
+	int32 NumIterationsThisStage = 1;
+	int32 SourceSimStage = -1;
+	bool bSpawnOnly;
 };
 
 class NIAGARAEDITOR_API FHlslNiagaraTranslator
@@ -429,6 +444,9 @@ protected:
 	FString GetFunctionDefinitions();
 
 	FString GetUniqueEmitterName() const;
+
+	void HandleDataInterfaceCall(FNiagaraScriptDataInterfaceCompileInfo& Info, const FNiagaraFunctionSignature& InMatchingSignature);
+
 public:
 
 	FHlslNiagaraTranslator();
@@ -532,6 +550,9 @@ public:
 	bool GetLiteralConstantVariable(FNiagaraVariable& OutVar) const;
 
 private:
+	bool GetUsesOldShaderStages() const;
+	bool GetUsesSimulationStages() const;
+
 	void InitializeParameterMapDefaults(int32 ParamMapHistoryIdx);
 	void HandleParameterRead(int32 ParamMapHistoryIdx, const FNiagaraVariable& Var, const UEdGraphPin* DefaultPin, UNiagaraNode* ErrorNode, int32& OutputChunkId, UNiagaraScriptVariable* Variable, bool bTreatAsUnknownParameterMap = false);
 	bool ShouldConsiderTargetParameterMap(ENiagaraScriptUsage InUsage) const;
diff --git a/Engine/Plugins/FX/Niagara/Source/NiagaraEditor/Private/NiagaraNodeEmitter.cpp b/Engine/Plugins/FX/Niagara/Source/NiagaraEditor/Private/NiagaraNodeEmitter.cpp
--- a/Engine/Plugins/FX/Niagara/Source/NiagaraEditor/Private/NiagaraNodeEmitter.cpp
+++ b/Engine/Plugins/FX/Niagara/Source/NiagaraEditor/Private/NiagaraNodeEmitter.cpp
@@ -349,6 +349,7 @@ void UNiagaraNodeEmitter::BuildParameterMapHistory(FNiagaraParameterMapHistoryBu
 		Usages.Add(ENiagaraScriptUsage::ParticleSpawnScriptInterpolated);
 		Usages.Add(ENiagaraScriptUsage::ParticleUpdateScript);
 		Usages.Add(ENiagaraScriptUsage::ParticleEventScript);
+		Usages.Add(ENiagaraScriptUsage::ParticleSimulationStageScript);
 	
 		uint32 NodeIdx = OutHistory.BeginNodeVisitation(ParamMapIdx, this);
 		for (ENiagaraScriptUsage OutputNodeUsage : Usages)
diff --git a/Engine/Plugins/FX/Niagara/Source/NiagaraEditor/Private/NiagaraNodeOutput.cpp b/Engine/Plugins/FX/Niagara/Source/NiagaraEditor/Private/NiagaraNodeOutput.cpp
--- a/Engine/Plugins/FX/Niagara/Source/NiagaraEditor/Private/NiagaraNodeOutput.cpp
+++ b/Engine/Plugins/FX/Niagara/Source/NiagaraEditor/Private/NiagaraNodeOutput.cpp
@@ -169,6 +169,11 @@ FText UNiagaraNodeOutput::GetNodeTitle(ENodeTitleType::Type TitleType) const
 		}
 		return FText::Format(NSLOCTEXT("NiagaraNodeOutput", "OutputEvent", "Output Event {0}"), EventName);
 	}
+	else if (ScriptType == ENiagaraScriptUsage::ParticleSimulationStageScript)
+	{
+		FText EventName = FText::FromString(ScriptTypeId.ToString());
+		return FText::Format(NSLOCTEXT("NiagaraNodeOutput", "OutputStage", "Output Stage {0}"), EventName);
+	}
 	else if (ScriptType == ENiagaraScriptUsage::Function)
 	{
 		return NSLOCTEXT("NiagaraNodeOutput", "OutputFunction", "Output Function");
diff --git a/Engine/Plugins/FX/Niagara/Source/NiagaraEditor/Private/NiagaraParameterMapHistory.cpp b/Engine/Plugins/FX/Niagara/Source/NiagaraEditor/Private/NiagaraParameterMapHistory.cpp
--- a/Engine/Plugins/FX/Niagara/Source/NiagaraEditor/Private/NiagaraParameterMapHistory.cpp
+++ b/Engine/Plugins/FX/Niagara/Source/NiagaraEditor/Private/NiagaraParameterMapHistory.cpp
@@ -51,12 +51,7 @@ void FNiagaraParameterMapHistory::GetValidNamespacesForReading(ENiagaraScriptUsa
 	OutputNamespaces.Add(PARAM_MAP_USER_STR);
 	OutputNamespaces.Add(PARAM_MAP_SYSTEM_STR);
 	OutputNamespaces.Add(PARAM_MAP_EMITTER_STR);
-
-	static const auto UseShaderStagesCVar = IConsoleManager::Get().FindConsoleVariable(TEXT("fx.UseShaderStages"));
-	if (UseShaderStagesCVar->GetInt() == 1)
-	{
-		OutputNamespaces.Add(PARAM_MAP_INDICES_STR);
-	}
+	OutputNamespaces.Add(PARAM_MAP_INDICES_STR);
 
 	for (ENiagaraScriptUsage Usage : SupportedContexts)
 	{
@@ -98,12 +93,8 @@ bool FNiagaraParameterMapHistory::IsValidNamespaceForReading(ENiagaraScriptUsage
 	ConcernedNamespaces.Add(PARAM_MAP_SYSTEM_STR);
 	ConcernedNamespaces.Add(PARAM_MAP_EMITTER_STR);
 	ConcernedNamespaces.Add(PARAM_MAP_ATTRIBUTE_STR);
+	ConcernedNamespaces.Add(PARAM_MAP_INDICES_STR);
 
-	static const auto UseShaderStagesCVar = IConsoleManager::Get().FindConsoleVariable(TEXT("fx.UseShaderStages"));
-	if (UseShaderStagesCVar->GetInt() == 1)
-	{
-		ConcernedNamespaces.Add(PARAM_MAP_INDICES_STR);
-	}
 	
 	if (!Namespace.EndsWith(TEXT(".")))
 	{
@@ -580,14 +571,11 @@ bool FNiagaraParameterMapHistory::IsExternalConstantNamespace(const FNiagaraVari
 		return true;
 	}
 
-	static const auto UseShaderStagesCVar = IConsoleManager::Get().FindConsoleVariable(TEXT("fx.UseShaderStages"));
-	if (UseShaderStagesCVar->GetInt() == 1)
+	if (IsInNamespace(InVar, PARAM_MAP_INDICES_STR))
 	{
-		if (IsInNamespace(InVar, PARAM_MAP_INDICES_STR))
-		{
-			return true;
-		}
+		return true;
 	}
+	
 
 	// Modules and functions need to act as if they are within the script types that they 
 	// say that they support rather than using their exact script type.
diff --git a/Engine/Plugins/FX/Niagara/Source/NiagaraEditor/Private/NiagaraScriptMergeManager.cpp b/Engine/Plugins/FX/Niagara/Source/NiagaraEditor/Private/NiagaraScriptMergeManager.cpp
--- a/Engine/Plugins/FX/Niagara/Source/NiagaraEditor/Private/NiagaraScriptMergeManager.cpp
+++ b/Engine/Plugins/FX/Niagara/Source/NiagaraEditor/Private/NiagaraScriptMergeManager.cpp
@@ -4,7 +4,7 @@
 #include "NiagaraEmitter.h"
 #include "NiagaraScript.h"
 #include "NiagaraScriptSource.h"
-#include "NiagaraShaderStageBase.h"
+#include "NiagaraSimulationStageBase.h"
 #include "NiagaraRendererProperties.h"
 #include "EdGraphSchema_Niagara.h"
 #include "NiagaraGraph.h"
@@ -419,47 +419,47 @@ TSharedPtr<FNiagaraScriptStackMergeAdapter> FNiagaraEventHandlerMergeAdapter::Ge
 	return EventStack;
 }
 
-FNiagaraShaderStageMergeAdapter::FNiagaraShaderStageMergeAdapter(const UNiagaraEmitter& InEmitter, const UNiagaraShaderStageBase* InShaderStage, UNiagaraNodeOutput* InOutputNode)
+FNiagaraSimulationStageMergeAdapter::FNiagaraSimulationStageMergeAdapter(const UNiagaraEmitter& InEmitter, const UNiagaraSimulationStageBase* InSimulationStage, UNiagaraNodeOutput* InOutputNode)
 {
-	Initialize(InEmitter, InShaderStage, nullptr, InOutputNode);
+	Initialize(InEmitter, InSimulationStage, nullptr, InOutputNode);
 }
 
-FNiagaraShaderStageMergeAdapter::FNiagaraShaderStageMergeAdapter(const UNiagaraEmitter& InEmitter, UNiagaraShaderStageBase* InShaderStage, UNiagaraNodeOutput* InOutputNode)
+FNiagaraSimulationStageMergeAdapter::FNiagaraSimulationStageMergeAdapter(const UNiagaraEmitter& InEmitter, UNiagaraSimulationStageBase* InSimulationStage, UNiagaraNodeOutput* InOutputNode)
 {
-	Initialize(InEmitter, InShaderStage, InShaderStage, InOutputNode);
+	Initialize(InEmitter, InSimulationStage, InSimulationStage, InOutputNode);
 }
 
-FNiagaraShaderStageMergeAdapter::FNiagaraShaderStageMergeAdapter(const UNiagaraEmitter& InEmitter, UNiagaraNodeOutput* InOutputNode)
+FNiagaraSimulationStageMergeAdapter::FNiagaraSimulationStageMergeAdapter(const UNiagaraEmitter& InEmitter, UNiagaraNodeOutput* InOutputNode)
 {
 	Initialize(InEmitter, nullptr, nullptr, InOutputNode);
 }
 
-void FNiagaraShaderStageMergeAdapter::Initialize(const UNiagaraEmitter& InEmitter, const UNiagaraShaderStageBase* InShaderStage, UNiagaraShaderStageBase* InEditableShaderStage, UNiagaraNodeOutput* InOutputNode)
+void FNiagaraSimulationStageMergeAdapter::Initialize(const UNiagaraEmitter& InEmitter, const UNiagaraSimulationStageBase* InSimulationStage, UNiagaraSimulationStageBase* InEditableSimulationStage, UNiagaraNodeOutput* InOutputNode)
 {
 	Emitter = MakeWeakObjectPtr(const_cast<UNiagaraEmitter*>(&InEmitter));
 
-	ShaderStage = InShaderStage;
-	EditableShaderStage = InEditableShaderStage;
+	SimulationStage = InSimulationStage;
+	EditableSimulationStage = InEditableSimulationStage;
 
 	OutputNode = InOutputNode;
 
-	if (ShaderStage != nullptr && OutputNode != nullptr)
+	if (SimulationStage != nullptr && OutputNode != nullptr)
 	{
-		ShaderStageStack = MakeShared<FNiagaraScriptStackMergeAdapter>(*OutputNode.Get(), *ShaderStage->Script, Emitter->GetUniqueEmitterName());
-		InputNode = ShaderStageStack->GetInputNode();
+		SimulationStageStack = MakeShared<FNiagaraScriptStackMergeAdapter>(*OutputNode.Get(), *SimulationStage->Script, Emitter->GetUniqueEmitterName());
+		InputNode = SimulationStageStack->GetInputNode();
 	}
 }
 
-const UNiagaraEmitter* FNiagaraShaderStageMergeAdapter::GetEmitter() const
+const UNiagaraEmitter* FNiagaraSimulationStageMergeAdapter::GetEmitter() const
 {
 	return Emitter.Get();
 }
 
-FGuid FNiagaraShaderStageMergeAdapter::GetUsageId() const
+FGuid FNiagaraSimulationStageMergeAdapter::GetUsageId() const
 {
-	if (ShaderStage != nullptr)
+	if (SimulationStage != nullptr)
 	{
-		return ShaderStage->Script->GetUsageId();
+		return SimulationStage->Script->GetUsageId();
 	}
 	else
 	{
@@ -467,29 +467,29 @@ FGuid FNiagaraShaderStageMergeAdapter::GetUsageId() const
 	}
 }
 
-const UNiagaraShaderStageBase* FNiagaraShaderStageMergeAdapter::GetShaderStage() const
+const UNiagaraSimulationStageBase* FNiagaraSimulationStageMergeAdapter::GetSimulationStage() const
 {
-	return ShaderStage;
+	return SimulationStage;
 }
 
-UNiagaraShaderStageBase* FNiagaraShaderStageMergeAdapter::GetEditableShaderStage() const
+UNiagaraSimulationStageBase* FNiagaraSimulationStageMergeAdapter::GetEditableSimulationStage() const
 {
-	return EditableShaderStage;
+	return EditableSimulationStage;
 }
 
-UNiagaraNodeOutput* FNiagaraShaderStageMergeAdapter::GetOutputNode() const
+UNiagaraNodeOutput* FNiagaraSimulationStageMergeAdapter::GetOutputNode() const
 {
 	return OutputNode.Get();
 }
 
-UNiagaraNodeInput* FNiagaraShaderStageMergeAdapter::GetInputNode() const
+UNiagaraNodeInput* FNiagaraSimulationStageMergeAdapter::GetInputNode() const
 {
 	return InputNode.Get();
 }
 
-TSharedPtr<FNiagaraScriptStackMergeAdapter> FNiagaraShaderStageMergeAdapter::GetShaderStageStack() const
+TSharedPtr<FNiagaraScriptStackMergeAdapter> FNiagaraSimulationStageMergeAdapter::GetSimulationStageStack() const
 {
-	return ShaderStageStack;
+	return SimulationStageStack;
 }
 
 FNiagaraRendererMergeAdapter::FNiagaraRendererMergeAdapter(UNiagaraRendererProperties& InRenderer)
@@ -522,7 +522,7 @@ void FNiagaraEmitterMergeAdapter::Initialize(const UNiagaraEmitter& InEmitter, U
 	Graph->GetNodesOfClass<UNiagaraNodeOutput>(OutputNodes);
 
 	TArray<UNiagaraNodeOutput*> EventOutputNodes;
-	TArray<UNiagaraNodeOutput*> ShaderStageOutputNodes;
+	TArray<UNiagaraNodeOutput*> SimulationStageOutputNodes;
 	for (UNiagaraNodeOutput* OutputNode : OutputNodes)
 	{
 		if (UNiagaraScript::IsEquivalentUsage(OutputNode->GetUsage(), ENiagaraScriptUsage::EmitterSpawnScript))
@@ -545,9 +545,9 @@ void FNiagaraEmitterMergeAdapter::Initialize(const UNiagaraEmitter& InEmitter, U
 		{
 			EventOutputNodes.Add(OutputNode);
 		}
-		else if (UNiagaraScript::IsEquivalentUsage(OutputNode->GetUsage(), ENiagaraScriptUsage::ParticleShaderStageScript))
+		else if (UNiagaraScript::IsEquivalentUsage(OutputNode->GetUsage(), ENiagaraScriptUsage::ParticleSimulationStageScript))
 		{
-			ShaderStageOutputNodes.Add(OutputNode);
+			SimulationStageOutputNodes.Add(OutputNode);
 		}
 	}
 
@@ -583,32 +583,32 @@ void FNiagaraEmitterMergeAdapter::Initialize(const UNiagaraEmitter& InEmitter, U
 
 	// Create an shader stage adapter for each usage id even if it's missing a shader stage object or an output node.  These
 	// incomplete adapters will be caught if they are diffed.
-	for (const UNiagaraShaderStageBase* ShaderStage : Emitter->GetShaderStages())
+	for (const UNiagaraSimulationStageBase* SimulationStage : Emitter->GetSimulationStages())
 	{
-		UNiagaraNodeOutput** MatchingOutputNodePtr = ShaderStageOutputNodes.FindByPredicate(
-			[=](UNiagaraNodeOutput* ShaderStageOutputNode) { return ShaderStageOutputNode->GetUsageId() == ShaderStage->Script->GetUsageId(); });
+		UNiagaraNodeOutput** MatchingOutputNodePtr = SimulationStageOutputNodes.FindByPredicate(
+			[=](UNiagaraNodeOutput* SimulationStageOutputNode) { return SimulationStageOutputNode->GetUsageId() == SimulationStage->Script->GetUsageId(); });
 
 		UNiagaraNodeOutput* MatchingOutputNode = MatchingOutputNodePtr != nullptr ? *MatchingOutputNodePtr : nullptr;
 
 		if (EditableEmitter == nullptr)
 		{
-			ShaderStages.Add(MakeShared<FNiagaraShaderStageMergeAdapter>(*Emitter.Get(), ShaderStage, MatchingOutputNode));
+			SimulationStages.Add(MakeShared<FNiagaraSimulationStageMergeAdapter>(*Emitter.Get(), SimulationStage, MatchingOutputNode));
 		}
 		else
 		{
-			UNiagaraShaderStageBase* EditableShaderStage = EditableEmitter->GetShaderStageById(ShaderStage->Script->GetUsageId());
-			ShaderStages.Add(MakeShared<FNiagaraShaderStageMergeAdapter>(*Emitter.Get(), EditableShaderStage, MatchingOutputNode));
+			UNiagaraSimulationStageBase* EditableSimulationStage = EditableEmitter->GetSimulationStageById(SimulationStage->Script->GetUsageId());
+			SimulationStages.Add(MakeShared<FNiagaraSimulationStageMergeAdapter>(*Emitter.Get(), EditableSimulationStage, MatchingOutputNode));
 		}
 
 		if (MatchingOutputNode != nullptr)
 		{
-			ShaderStageOutputNodes.Remove(MatchingOutputNode);
+			SimulationStageOutputNodes.Remove(MatchingOutputNode);
 		}
 	}
 
-	for (UNiagaraNodeOutput* ShaderStageOutputNode : ShaderStageOutputNodes)
+	for (UNiagaraNodeOutput* SimulationStageOutputNode : SimulationStageOutputNodes)
 	{
-		ShaderStages.Add(MakeShared<FNiagaraShaderStageMergeAdapter>(*Emitter.Get(), ShaderStageOutputNode));
+		SimulationStages.Add(MakeShared<FNiagaraSimulationStageMergeAdapter>(*Emitter.Get(), SimulationStageOutputNode));
 	}
 
 	// Renderers
@@ -650,9 +650,9 @@ const TArray<TSharedRef<FNiagaraEventHandlerMergeAdapter>> FNiagaraEmitterMergeA
 	return EventHandlers;
 }
 
-const TArray<TSharedRef<FNiagaraShaderStageMergeAdapter>> FNiagaraEmitterMergeAdapter::GetShaderStages() const
+const TArray<TSharedRef<FNiagaraSimulationStageMergeAdapter>> FNiagaraEmitterMergeAdapter::GetSimulationStages() const
 {
-	return ShaderStages;
+	return SimulationStages;
 }
 
 const TArray<TSharedRef<FNiagaraRendererMergeAdapter>> FNiagaraEmitterMergeAdapter::GetRenderers() const
@@ -705,16 +705,16 @@ TSharedPtr<FNiagaraEventHandlerMergeAdapter> FNiagaraEmitterMergeAdapter::GetEve
 	return TSharedPtr<FNiagaraEventHandlerMergeAdapter>();
 }
 
-TSharedPtr<FNiagaraShaderStageMergeAdapter> FNiagaraEmitterMergeAdapter::GetShaderStage(FGuid ShaderStageUsageId)
+TSharedPtr<FNiagaraSimulationStageMergeAdapter> FNiagaraEmitterMergeAdapter::GetSimulationStage(FGuid SimulationStageUsageId)
 {
-	for (TSharedRef<FNiagaraShaderStageMergeAdapter> ShaderStage : ShaderStages)
+	for (TSharedRef<FNiagaraSimulationStageMergeAdapter> SimulationStage : SimulationStages)
 	{
-		if (ShaderStage->GetUsageId() == ShaderStageUsageId)
+		if (SimulationStage->GetUsageId() == SimulationStageUsageId)
 		{
-			return ShaderStage;
+			return SimulationStage;
 		}
 	}
-	return TSharedPtr<FNiagaraShaderStageMergeAdapter>();
+	return TSharedPtr<FNiagaraSimulationStageMergeAdapter>();
 }
 
 TSharedPtr<FNiagaraRendererMergeAdapter> FNiagaraEmitterMergeAdapter::GetRenderer(FGuid RendererMergeId)
@@ -782,18 +782,18 @@ bool FNiagaraEmitterDiffResults::IsValid() const
 			break;
 		}
 	}
-	bool bShaderStageDiffsAreValid = true;
-	for (const FNiagaraModifiedShaderStageDiffResults& ShaderStageDiffResults : ModifiedShaderStages)
+	bool bSimulationStageDiffsAreValid = true;
+	for (const FNiagaraModifiedSimulationStageDiffResults& SimulationStageDiffResults : ModifiedSimulationStages)
 	{
-		if (ShaderStageDiffResults.ScriptDiffResults.IsValid() == false)
+		if (SimulationStageDiffResults.ScriptDiffResults.IsValid() == false)
 		{
-			bShaderStageDiffsAreValid = false;
+			bSimulationStageDiffsAreValid = false;
 			break;
 		}
 	}
 	return bIsValid &&
 		bEventHandlerDiffsAreValid &&
-		bShaderStageDiffsAreValid &&
+		bSimulationStageDiffsAreValid &&
 		EmitterSpawnDiffResults.IsValid() &&
 		EmitterUpdateDiffResults.IsValid() &&
 		ParticleSpawnDiffResults.IsValid() &&
@@ -810,9 +810,9 @@ bool FNiagaraEmitterDiffResults::IsEmpty() const
 		RemovedBaseEventHandlers.Num() == 0 &&
 		AddedOtherEventHandlers.Num() == 0 &&
 		ModifiedEventHandlers.Num() == 0 &&
-		RemovedBaseShaderStages.Num() == 0 &&
-		AddedOtherShaderStages.Num() == 0 &&
-		ModifiedShaderStages.Num() == 0 &&
+		RemovedBaseSimulationStages.Num() == 0 &&
+		AddedOtherSimulationStages.Num() == 0 &&
+		ModifiedSimulationStages.Num() == 0 &&
 		RemovedBaseRenderers.Num() == 0 &&
 		AddedOtherRenderers.Num() == 0 &&
 		ModifiedBaseRenderers.Num() == 0 &&
@@ -1087,13 +1087,13 @@ INiagaraMergeManager::FMergeEmitterResults FNiagaraScriptMergeManager::MergeEmit
 		MergeResults.bModifiedGraph |= EventHandlerResults.bModifiedGraph;
 		MergeResults.ErrorMessages.Append(EventHandlerResults.ErrorMessages);
 
-		FApplyDiffResults ShaderStageResults = ApplyShaderStageDiff(MergedInstanceAdapter, DiffResults, bNoParentAtLastMerge);
-		if (ShaderStageResults.bSucceeded == false)
+		FApplyDiffResults SimulationStageResults = ApplySimulationStageDiff(MergedInstanceAdapter, DiffResults, bNoParentAtLastMerge);
+		if (SimulationStageResults.bSucceeded == false)
 		{
 			MergeResults.MergeResult = INiagaraMergeManager::EMergeEmitterResult::FailedToMerge;
 		}
-		MergeResults.bModifiedGraph |= ShaderStageResults.bModifiedGraph;
-		MergeResults.ErrorMessages.Append(ShaderStageResults.ErrorMessages);
+		MergeResults.bModifiedGraph |= SimulationStageResults.bModifiedGraph;
+		MergeResults.ErrorMessages.Append(SimulationStageResults.ErrorMessages);
 
 		FApplyDiffResults RendererResults = ApplyRendererDiff(*MergedInstance, DiffResults, bNoParentAtLastMerge);
 		if (RendererResults.bSucceeded == false)
@@ -1341,47 +1341,47 @@ void FNiagaraScriptMergeManager::ResetEventHandlerPropertySetToBase(UNiagaraEmit
 	Emitter.PostEditChange();
 }
 
-bool FNiagaraScriptMergeManager::HasBaseShaderStage(const UNiagaraEmitter& BaseEmitter, FGuid ShaderStageScriptUsageId)
+bool FNiagaraScriptMergeManager::HasBaseSimulationStage(const UNiagaraEmitter& BaseEmitter, FGuid SimulationStageScriptUsageId)
 {
 	TSharedRef<FNiagaraEmitterMergeAdapter> BaseEmitterAdapter = GetEmitterMergeAdapterUsingCache(BaseEmitter);
-	return BaseEmitterAdapter->GetShaderStage(ShaderStageScriptUsageId).IsValid();
+	return BaseEmitterAdapter->GetSimulationStage(SimulationStageScriptUsageId).IsValid();
 }
 
-bool FNiagaraScriptMergeManager::IsShaderStagePropertySetDifferentFromBase(UNiagaraEmitter& Emitter, const UNiagaraEmitter& BaseEmitter, FGuid ShaderStageScriptUsageId)
+bool FNiagaraScriptMergeManager::IsSimulationStagePropertySetDifferentFromBase(UNiagaraEmitter& Emitter, const UNiagaraEmitter& BaseEmitter, FGuid SimulationStageScriptUsageId)
 {
 	TSharedRef<FNiagaraEmitterMergeAdapter> EmitterAdapter = GetEmitterMergeAdapterUsingCache(Emitter);
 	TSharedRef<FNiagaraEmitterMergeAdapter> BaseEmitterAdapter = GetEmitterMergeAdapterUsingCache(BaseEmitter);
 
-	TSharedPtr<FNiagaraShaderStageMergeAdapter> ShaderStageAdapter = EmitterAdapter->GetShaderStage(ShaderStageScriptUsageId);
-	TSharedPtr<FNiagaraShaderStageMergeAdapter> BaseShaderStageAdapter = BaseEmitterAdapter->GetShaderStage(ShaderStageScriptUsageId);
+	TSharedPtr<FNiagaraSimulationStageMergeAdapter> SimulationStageAdapter = EmitterAdapter->GetSimulationStage(SimulationStageScriptUsageId);
+	TSharedPtr<FNiagaraSimulationStageMergeAdapter> BaseSimulationStageAdapter = BaseEmitterAdapter->GetSimulationStage(SimulationStageScriptUsageId);
 
-	if (ShaderStageAdapter->GetEditableShaderStage() == nullptr || BaseShaderStageAdapter->GetShaderStage() == nullptr)
+	if (SimulationStageAdapter->GetEditableSimulationStage() == nullptr || BaseSimulationStageAdapter->GetSimulationStage() == nullptr)
 	{
 		return true;
 	}
 
 	TArray<FProperty*> DifferentProperties;
-	DiffEditableProperties(BaseShaderStageAdapter->GetShaderStage(), ShaderStageAdapter->GetShaderStage(), *BaseShaderStageAdapter->GetShaderStage()->GetClass(), DifferentProperties);
+	DiffEditableProperties(BaseSimulationStageAdapter->GetSimulationStage(), SimulationStageAdapter->GetSimulationStage(), *BaseSimulationStageAdapter->GetSimulationStage()->GetClass(), DifferentProperties);
 	return DifferentProperties.Num() > 0;
 }
 
-void FNiagaraScriptMergeManager::ResetShaderStagePropertySetToBase(UNiagaraEmitter& Emitter, const UNiagaraEmitter& BaseEmitter, FGuid ShaderStageScriptUsageId)
+void FNiagaraScriptMergeManager::ResetSimulationStagePropertySetToBase(UNiagaraEmitter& Emitter, const UNiagaraEmitter& BaseEmitter, FGuid SimulationStageScriptUsageId)
 {
 	TSharedRef<FNiagaraEmitterMergeAdapter> EmitterAdapter = GetEmitterMergeAdapterUsingCache(Emitter);
 	TSharedRef<FNiagaraEmitterMergeAdapter> BaseEmitterAdapter = GetEmitterMergeAdapterUsingCache(BaseEmitter);
 
-	TSharedPtr<FNiagaraShaderStageMergeAdapter> ShaderStageAdapter = EmitterAdapter->GetShaderStage(ShaderStageScriptUsageId);
-	TSharedPtr<FNiagaraShaderStageMergeAdapter> BaseShaderStageAdapter = BaseEmitterAdapter->GetShaderStage(ShaderStageScriptUsageId);
+	TSharedPtr<FNiagaraSimulationStageMergeAdapter> SimulationStageAdapter = EmitterAdapter->GetSimulationStage(SimulationStageScriptUsageId);
+	TSharedPtr<FNiagaraSimulationStageMergeAdapter> BaseSimulationStageAdapter = BaseEmitterAdapter->GetSimulationStage(SimulationStageScriptUsageId);
 
-	if (ShaderStageAdapter->GetEditableShaderStage() == nullptr || BaseShaderStageAdapter->GetShaderStage() == nullptr)
+	if (SimulationStageAdapter->GetEditableSimulationStage() == nullptr || BaseSimulationStageAdapter->GetSimulationStage() == nullptr)
 	{
 		// TODO: Display an error to the user.
 		return;
 	}
 
 	TArray<FProperty*> DifferentProperties;
-	DiffEditableProperties(BaseShaderStageAdapter->GetShaderStage(), ShaderStageAdapter->GetShaderStage(), *BaseShaderStageAdapter->GetShaderStage()->GetClass(), DifferentProperties);
-	CopyPropertiesToBase(ShaderStageAdapter->GetEditableShaderStage(), BaseShaderStageAdapter->GetShaderStage(), DifferentProperties);
+	DiffEditableProperties(BaseSimulationStageAdapter->GetSimulationStage(), SimulationStageAdapter->GetSimulationStage(), *BaseSimulationStageAdapter->GetSimulationStage()->GetClass(), DifferentProperties);
+	CopyPropertiesToBase(SimulationStageAdapter->GetEditableSimulationStage(), BaseSimulationStageAdapter->GetSimulationStage(), DifferentProperties);
 	Emitter.PostEditChange();
 }
 
@@ -1496,7 +1496,7 @@ FNiagaraEmitterDiffResults FNiagaraScriptMergeManager::DiffEmitters(UNiagaraEmit
 	}
 
 	DiffEventHandlers(BaseEmitterAdapter->GetEventHandlers(), OtherEmitterAdapter->GetEventHandlers(), EmitterDiffResults);
-	DiffShaderStages(BaseEmitterAdapter->GetShaderStages(), OtherEmitterAdapter->GetShaderStages(), EmitterDiffResults);
+	DiffSimulationStages(BaseEmitterAdapter->GetSimulationStages(), OtherEmitterAdapter->GetSimulationStages(), EmitterDiffResults);
 	DiffRenderers(BaseEmitterAdapter->GetRenderers(), OtherEmitterAdapter->GetRenderers(), EmitterDiffResults);
 	DiffEditableProperties(&BaseEmitter, &OtherEmitter, *UNiagaraEmitter::StaticClass(), EmitterDiffResults.DifferentEmitterProperties);
 	DiffStackEntryDisplayNames(BaseEmitterAdapter->GetEditorData(), OtherEmitterAdapter->GetEditorData(), EmitterDiffResults.ModifiedStackEntryDisplayNames);
@@ -1617,51 +1617,51 @@ void FNiagaraScriptMergeManager::DiffEventHandlers(const TArray<TSharedRef<FNiag
 	}
 }
 
-void FNiagaraScriptMergeManager::DiffShaderStages(const TArray<TSharedRef<FNiagaraShaderStageMergeAdapter>>& BaseShaderStages, const TArray<TSharedRef<FNiagaraShaderStageMergeAdapter>>& OtherShaderStages, FNiagaraEmitterDiffResults& DiffResults) const
+void FNiagaraScriptMergeManager::DiffSimulationStages(const TArray<TSharedRef<FNiagaraSimulationStageMergeAdapter>>& BaseSimulationStages, const TArray<TSharedRef<FNiagaraSimulationStageMergeAdapter>>& OtherSimulationStages, FNiagaraEmitterDiffResults& DiffResults) const
 {
-	FListDiffResults<TSharedRef<FNiagaraShaderStageMergeAdapter>> ShaderStageListDiffResults = DiffLists<TSharedRef<FNiagaraShaderStageMergeAdapter>, FGuid>(
-		BaseShaderStages,
-		OtherShaderStages,
-		[](TSharedRef<FNiagaraShaderStageMergeAdapter> ShaderStage) { return ShaderStage->GetUsageId(); });
+	FListDiffResults<TSharedRef<FNiagaraSimulationStageMergeAdapter>> SimulationStageListDiffResults = DiffLists<TSharedRef<FNiagaraSimulationStageMergeAdapter>, FGuid>(
+		BaseSimulationStages,
+		OtherSimulationStages,
+		[](TSharedRef<FNiagaraSimulationStageMergeAdapter> SimulationStage) { return SimulationStage->GetUsageId(); });
 
-	DiffResults.RemovedBaseShaderStages.Append(ShaderStageListDiffResults.RemovedBaseValues);
-	DiffResults.AddedOtherShaderStages.Append(ShaderStageListDiffResults.AddedOtherValues);
+	DiffResults.RemovedBaseSimulationStages.Append(SimulationStageListDiffResults.RemovedBaseValues);
+	DiffResults.AddedOtherSimulationStages.Append(SimulationStageListDiffResults.AddedOtherValues);
 
-	for (const FCommonValuePair<TSharedRef<FNiagaraShaderStageMergeAdapter>>& CommonValuePair : ShaderStageListDiffResults.CommonValuePairs)
+	for (const FCommonValuePair<TSharedRef<FNiagaraSimulationStageMergeAdapter>>& CommonValuePair : SimulationStageListDiffResults.CommonValuePairs)
 	{
-		if (CommonValuePair.BaseValue->GetShaderStage() == nullptr || CommonValuePair.BaseValue->GetOutputNode() == nullptr)
+		if (CommonValuePair.BaseValue->GetSimulationStage() == nullptr || CommonValuePair.BaseValue->GetOutputNode() == nullptr)
 		{
-			DiffResults.AddError(FText::Format(LOCTEXT("InvalidBaseShaderStageDiffFailedFormat", "Failed to diff shader stages, the base shader stage was invalid.  Script Usage Id: {0}"),
+			DiffResults.AddError(FText::Format(LOCTEXT("InvalidBaseSimulationStageDiffFailedFormat", "Failed to diff shader stages, the base shader stage was invalid.  Script Usage Id: {0}"),
 				FText::FromString(CommonValuePair.BaseValue->GetUsageId().ToString())));
 		}
-		else if (CommonValuePair.OtherValue->GetShaderStage() == nullptr || CommonValuePair.OtherValue->GetOutputNode() == nullptr)
+		else if (CommonValuePair.OtherValue->GetSimulationStage() == nullptr || CommonValuePair.OtherValue->GetOutputNode() == nullptr)
 		{
-			DiffResults.AddError(FText::Format(LOCTEXT("InvalidOtherShaderStageDiffFailedFormat", "Failed to diff shader stage, the other shader stage was invalid.  Script Usage Id: {0}"),
+			DiffResults.AddError(FText::Format(LOCTEXT("InvalidOtherSimulationStageDiffFailedFormat", "Failed to diff shader stage, the other shader stage was invalid.  Script Usage Id: {0}"),
 				FText::FromString(CommonValuePair.OtherValue->GetUsageId().ToString())));
 		}
 		else
 		{
 			TArray<FProperty*> DifferentProperties;
-			DiffEditableProperties(CommonValuePair.BaseValue->GetShaderStage(), CommonValuePair.OtherValue->GetShaderStage(), *CommonValuePair.BaseValue->GetShaderStage()->GetClass(), DifferentProperties);
+			DiffEditableProperties(CommonValuePair.BaseValue->GetSimulationStage(), CommonValuePair.OtherValue->GetSimulationStage(), *CommonValuePair.BaseValue->GetSimulationStage()->GetClass(), DifferentProperties);
 
-			FNiagaraScriptStackDiffResults ShaderStageScriptStackDiffResults;
-			DiffScriptStacks(CommonValuePair.BaseValue->GetShaderStageStack().ToSharedRef(), CommonValuePair.OtherValue->GetShaderStageStack().ToSharedRef(), ShaderStageScriptStackDiffResults);
+			FNiagaraScriptStackDiffResults SimulationStageScriptStackDiffResults;
+			DiffScriptStacks(CommonValuePair.BaseValue->GetSimulationStageStack().ToSharedRef(), CommonValuePair.OtherValue->GetSimulationStageStack().ToSharedRef(), SimulationStageScriptStackDiffResults);
 
-			if (DifferentProperties.Num() > 0 || ShaderStageScriptStackDiffResults.IsValid() == false || ShaderStageScriptStackDiffResults.IsEmpty() == false)
+			if (DifferentProperties.Num() > 0 || SimulationStageScriptStackDiffResults.IsValid() == false || SimulationStageScriptStackDiffResults.IsEmpty() == false)
 			{
-				FNiagaraModifiedShaderStageDiffResults ModifiedShaderStageResults;
-				ModifiedShaderStageResults.BaseAdapter = CommonValuePair.BaseValue;
-				ModifiedShaderStageResults.OtherAdapter = CommonValuePair.OtherValue;
-				ModifiedShaderStageResults.ChangedProperties.Append(DifferentProperties);
-				ModifiedShaderStageResults.ScriptDiffResults = ShaderStageScriptStackDiffResults;
-				DiffResults.ModifiedShaderStages.Add(ModifiedShaderStageResults);
+				FNiagaraModifiedSimulationStageDiffResults ModifiedSimulationStageResults;
+				ModifiedSimulationStageResults.BaseAdapter = CommonValuePair.BaseValue;
+				ModifiedSimulationStageResults.OtherAdapter = CommonValuePair.OtherValue;
+				ModifiedSimulationStageResults.ChangedProperties.Append(DifferentProperties);
+				ModifiedSimulationStageResults.ScriptDiffResults = SimulationStageScriptStackDiffResults;
+				DiffResults.ModifiedSimulationStages.Add(ModifiedSimulationStageResults);
 			}
 
-			if (ShaderStageScriptStackDiffResults.IsValid() == false)
+			if (SimulationStageScriptStackDiffResults.IsValid() == false)
 			{
-				for (const FText& ShaderStageScriptStackDiffErrorMessage : ShaderStageScriptStackDiffResults.GetErrorMessages())
+				for (const FText& SimulationStageScriptStackDiffErrorMessage : SimulationStageScriptStackDiffResults.GetErrorMessages())
 				{
-					DiffResults.AddError(ShaderStageScriptStackDiffErrorMessage);
+					DiffResults.AddError(SimulationStageScriptStackDiffErrorMessage);
 				}
 			}
 		}
@@ -2443,49 +2443,49 @@ FNiagaraScriptMergeManager::FApplyDiffResults FNiagaraScriptMergeManager::ApplyE
 	return Results;
 }
 
-FNiagaraScriptMergeManager::FApplyDiffResults FNiagaraScriptMergeManager::ApplyShaderStageDiff(TSharedRef<FNiagaraEmitterMergeAdapter> BaseEmitterAdapter, const FNiagaraEmitterDiffResults& DiffResults, const bool bNoParentAtLastMerge) const
+FNiagaraScriptMergeManager::FApplyDiffResults FNiagaraScriptMergeManager::ApplySimulationStageDiff(TSharedRef<FNiagaraEmitterMergeAdapter> BaseEmitterAdapter, const FNiagaraEmitterDiffResults& DiffResults, const bool bNoParentAtLastMerge) const
 {
 	FApplyDiffResults Results;
-	if (DiffResults.RemovedBaseShaderStages.Num() > 0)
+	if (DiffResults.RemovedBaseSimulationStages.Num() > 0)
 	{
 		Results.bSucceeded = false;
 		Results.bModifiedGraph = false;
 		// If this becomes supported, it needs to handle the bNoParentAtLastMerge case
-		Results.ErrorMessages.Add(LOCTEXT("RemovedShaderStagesUnsupported", "Apply diff failed, removed shader stages are currently unsupported."));
+		Results.ErrorMessages.Add(LOCTEXT("RemovedSimulationStagesUnsupported", "Apply diff failed, removed shader stages are currently unsupported."));
 		return Results;
 	}
 
-	for (const FNiagaraModifiedShaderStageDiffResults& ModifiedShaderStage : DiffResults.ModifiedShaderStages)
+	for (const FNiagaraModifiedSimulationStageDiffResults& ModifiedSimulationStage : DiffResults.ModifiedSimulationStages)
 	{
-		if (ModifiedShaderStage.OtherAdapter->GetShaderStage() == nullptr)
+		if (ModifiedSimulationStage.OtherAdapter->GetSimulationStage() == nullptr)
 		{
 			Results.bSucceeded = false;
 			Results.ErrorMessages.Add(FText::Format(
-				LOCTEXT("MissingModifiedShaderStageObjectFormat", "Apply diff failed.  The modified shader stage with id: {0} was missing it's shader stage object."),
-				FText::FromString(ModifiedShaderStage.OtherAdapter->GetUsageId().ToString(EGuidFormats::DigitsWithHyphens))));
+				LOCTEXT("MissingModifiedSimulationStageObjectFormat", "Apply diff failed.  The modified shader stage with id: {0} was missing it's shader stage object."),
+				FText::FromString(ModifiedSimulationStage.OtherAdapter->GetUsageId().ToString(EGuidFormats::DigitsWithHyphens))));
 		}
-		else if (ModifiedShaderStage.OtherAdapter->GetOutputNode() == nullptr)
+		else if (ModifiedSimulationStage.OtherAdapter->GetOutputNode() == nullptr)
 		{
 			Results.bSucceeded = false;
 			Results.ErrorMessages.Add(FText::Format(
-				LOCTEXT("MissingModifiedShaderStageOutputNodeFormat", "Apply diff failed.  The modified shader stage with id: {0} was missing it's output node."),
-				FText::FromString(ModifiedShaderStage.OtherAdapter->GetUsageId().ToString(EGuidFormats::DigitsWithHyphens))));
+				LOCTEXT("MissingModifiedSimulationStageOutputNodeFormat", "Apply diff failed.  The modified shader stage with id: {0} was missing it's output node."),
+				FText::FromString(ModifiedSimulationStage.OtherAdapter->GetUsageId().ToString(EGuidFormats::DigitsWithHyphens))));
 		}
 		else
 		{
-			TSharedPtr<FNiagaraShaderStageMergeAdapter> MatchingBaseShaderStageAdapter = BaseEmitterAdapter->GetShaderStage(ModifiedShaderStage.OtherAdapter->GetUsageId());
-			if (MatchingBaseShaderStageAdapter.IsValid())
+			TSharedPtr<FNiagaraSimulationStageMergeAdapter> MatchingBaseSimulationStageAdapter = BaseEmitterAdapter->GetSimulationStage(ModifiedSimulationStage.OtherAdapter->GetUsageId());
+			if (MatchingBaseSimulationStageAdapter.IsValid())
 			{
-				if (ModifiedShaderStage.ChangedProperties.Num() > 0)
+				if (ModifiedSimulationStage.ChangedProperties.Num() > 0)
 				{
-					CopyPropertiesToBase(MatchingBaseShaderStageAdapter->GetEditableShaderStage(), ModifiedShaderStage.OtherAdapter->GetEditableShaderStage(), ModifiedShaderStage.ChangedProperties);
+					CopyPropertiesToBase(MatchingBaseSimulationStageAdapter->GetEditableSimulationStage(), ModifiedSimulationStage.OtherAdapter->GetEditableSimulationStage(), ModifiedSimulationStage.ChangedProperties);
 				}
-				if (ModifiedShaderStage.ScriptDiffResults.IsEmpty() == false)
+				if (ModifiedSimulationStage.ScriptDiffResults.IsEmpty() == false)
 				{
-					FApplyDiffResults ApplyShaderStageStackDiffResults = ApplyScriptStackDiff(MatchingBaseShaderStageAdapter->GetShaderStageStack().ToSharedRef(), ModifiedShaderStage.ScriptDiffResults, bNoParentAtLastMerge);
-					Results.bSucceeded &= ApplyShaderStageStackDiffResults.bSucceeded;
-					Results.bModifiedGraph |= ApplyShaderStageStackDiffResults.bModifiedGraph;
-					Results.ErrorMessages.Append(ApplyShaderStageStackDiffResults.ErrorMessages);
+					FApplyDiffResults ApplySimulationStageStackDiffResults = ApplyScriptStackDiff(MatchingBaseSimulationStageAdapter->GetSimulationStageStack().ToSharedRef(), ModifiedSimulationStage.ScriptDiffResults, bNoParentAtLastMerge);
+					Results.bSucceeded &= ApplySimulationStageStackDiffResults.bSucceeded;
+					Results.bModifiedGraph |= ApplySimulationStageStackDiffResults.bModifiedGraph;
+					Results.ErrorMessages.Append(ApplySimulationStageStackDiffResults.ErrorMessages);
 				}
 			}
 		}
@@ -2493,49 +2493,49 @@ FNiagaraScriptMergeManager::FApplyDiffResults FNiagaraScriptMergeManager::ApplyS
 
 	UNiagaraScriptSource* EmitterSource = CastChecked<UNiagaraScriptSource>(BaseEmitterAdapter->GetEditableEmitter()->GraphSource);
 	UNiagaraGraph* EmitterGraph = EmitterSource->NodeGraph;
-	for (TSharedRef<FNiagaraShaderStageMergeAdapter> AddedOtherShaderStage : DiffResults.AddedOtherShaderStages)
+	for (TSharedRef<FNiagaraSimulationStageMergeAdapter> AddedOtherSimulationStage : DiffResults.AddedOtherSimulationStages)
 	{
-		if (AddedOtherShaderStage->GetShaderStage() == nullptr)
+		if (AddedOtherSimulationStage->GetSimulationStage() == nullptr)
 		{
 			Results.bSucceeded = false;
 			Results.ErrorMessages.Add(FText::Format(
-				LOCTEXT("MissingAddedShaderStageObjectFormat", "Apply diff failed.  The added shader stage with id: {0} was missing it's shader stage object."),
-				FText::FromString(AddedOtherShaderStage->GetUsageId().ToString(EGuidFormats::DigitsWithHyphens))));
+				LOCTEXT("MissingAddedSimulationStageObjectFormat", "Apply diff failed.  The added shader stage with id: {0} was missing it's shader stage object."),
+				FText::FromString(AddedOtherSimulationStage->GetUsageId().ToString(EGuidFormats::DigitsWithHyphens))));
 		}
-		else if (AddedOtherShaderStage->GetOutputNode() == nullptr)
+		else if (AddedOtherSimulationStage->GetOutputNode() == nullptr)
 		{
 			Results.bSucceeded = false;
 			Results.ErrorMessages.Add(FText::Format(
-				LOCTEXT("MissingAddedShaderStageOutputNodeFormat", "Apply diff failed.  The added shader stage with id: {0} was missing it's output node."),
-				FText::FromString(AddedOtherShaderStage->GetUsageId().ToString(EGuidFormats::DigitsWithHyphens))));
+				LOCTEXT("MissingAddedSimulationStageOutputNodeFormat", "Apply diff failed.  The added shader stage with id: {0} was missing it's output node."),
+				FText::FromString(AddedOtherSimulationStage->GetUsageId().ToString(EGuidFormats::DigitsWithHyphens))));
 		}
 		else
 		{
 			UNiagaraEmitter* BaseEmitter = BaseEmitterAdapter->GetEditableEmitter();
-			UNiagaraShaderStageBase* AddedShaderStage = CastChecked<UNiagaraShaderStageBase>(StaticDuplicateObject(AddedOtherShaderStage->GetShaderStage(), BaseEmitter));
-			AddedShaderStage->Script = NewObject<UNiagaraScript>(AddedShaderStage, MakeUniqueObjectName(AddedShaderStage, UNiagaraScript::StaticClass(), "ShaderStage"), EObjectFlags::RF_Transactional);
-			AddedShaderStage->Script->SetUsage(ENiagaraScriptUsage::ParticleShaderStageScript);
-			AddedShaderStage->Script->SetUsageId(AddedOtherShaderStage->GetUsageId());
-			AddedShaderStage->Script->SetSource(EmitterSource);
-			BaseEmitter->AddShaderStage(AddedShaderStage);
-
-			FGuid PreferredOutputNodeGuid = AddedOtherShaderStage->GetOutputNode()->NodeGuid;
-			FGuid PreferredInputNodeGuid = AddedOtherShaderStage->GetInputNode()->NodeGuid;
-			UNiagaraNodeOutput* ShaderStageOutputNode = FNiagaraStackGraphUtilities::ResetGraphForOutput(*EmitterGraph, ENiagaraScriptUsage::ParticleShaderStageScript, AddedShaderStage->Script->GetUsageId(), PreferredOutputNodeGuid, PreferredInputNodeGuid);
-			for (TSharedRef<FNiagaraStackFunctionMergeAdapter> ModuleAdapter : AddedOtherShaderStage->GetShaderStageStack()->GetModuleFunctions())
+			UNiagaraSimulationStageBase* AddedSimulationStage = CastChecked<UNiagaraSimulationStageBase>(StaticDuplicateObject(AddedOtherSimulationStage->GetSimulationStage(), BaseEmitter));
+			AddedSimulationStage->Script = NewObject<UNiagaraScript>(AddedSimulationStage, MakeUniqueObjectName(AddedSimulationStage, UNiagaraScript::StaticClass(), "SimulationStage"), EObjectFlags::RF_Transactional);
+			AddedSimulationStage->Script->SetUsage(ENiagaraScriptUsage::ParticleSimulationStageScript);
+			AddedSimulationStage->Script->SetUsageId(AddedOtherSimulationStage->GetUsageId());
+			AddedSimulationStage->Script->SetSource(EmitterSource);
+			BaseEmitter->AddSimulationStage(AddedSimulationStage);
+
+			FGuid PreferredOutputNodeGuid = AddedOtherSimulationStage->GetOutputNode()->NodeGuid;
+			FGuid PreferredInputNodeGuid = AddedOtherSimulationStage->GetInputNode()->NodeGuid;
+			UNiagaraNodeOutput* SimulationStageOutputNode = FNiagaraStackGraphUtilities::ResetGraphForOutput(*EmitterGraph, ENiagaraScriptUsage::ParticleSimulationStageScript, AddedSimulationStage->Script->GetUsageId(), PreferredOutputNodeGuid, PreferredInputNodeGuid);
+			for (TSharedRef<FNiagaraStackFunctionMergeAdapter> ModuleAdapter : AddedOtherSimulationStage->GetSimulationStageStack()->GetModuleFunctions())
 			{
-				FApplyDiffResults AddModuleResults = AddModule(BaseEmitter->GetUniqueEmitterName(), *AddedShaderStage->Script, *ShaderStageOutputNode, ModuleAdapter);
+				FApplyDiffResults AddModuleResults = AddModule(BaseEmitter->GetUniqueEmitterName(), *AddedSimulationStage->Script, *SimulationStageOutputNode, ModuleAdapter);
 				Results.bSucceeded &= AddModuleResults.bSucceeded;
 				Results.ErrorMessages.Append(AddModuleResults.ErrorMessages);
 			}
 
 			// Force the base compile id of the new shader stage to match the added instance shader stage.
-			UNiagaraScriptSource* AddedShaderStageSourceFromDiff = Cast<UNiagaraScriptSource>(AddedOtherShaderStage->GetShaderStage()->Script->GetSource());
-			UNiagaraGraph* AddedShaderStageGraphFromDiff = AddedShaderStageSourceFromDiff->NodeGraph;
-			FGuid ScriptBaseIdFromDiff = AddedShaderStageGraphFromDiff->GetBaseId(ENiagaraScriptUsage::ParticleShaderStageScript, AddedOtherShaderStage->GetUsageId());
-			UNiagaraScriptSource* AddedShaderStageSource = Cast<UNiagaraScriptSource>(AddedShaderStage->Script->GetSource());
-			UNiagaraGraph* AddedShaderStageGraph = AddedShaderStageSource->NodeGraph;
-			AddedShaderStageGraph->ForceBaseId(ENiagaraScriptUsage::ParticleShaderStageScript, AddedOtherShaderStage->GetUsageId(), ScriptBaseIdFromDiff);
+			UNiagaraScriptSource* AddedSimulationStageSourceFromDiff = Cast<UNiagaraScriptSource>(AddedOtherSimulationStage->GetSimulationStage()->Script->GetSource());
+			UNiagaraGraph* AddedSimulationStageGraphFromDiff = AddedSimulationStageSourceFromDiff->NodeGraph;
+			FGuid ScriptBaseIdFromDiff = AddedSimulationStageGraphFromDiff->GetBaseId(ENiagaraScriptUsage::ParticleSimulationStageScript, AddedOtherSimulationStage->GetUsageId());
+			UNiagaraScriptSource* AddedSimulationStageSource = Cast<UNiagaraScriptSource>(AddedSimulationStage->Script->GetSource());
+			UNiagaraGraph* AddedSimulationStageGraph = AddedSimulationStageSource->NodeGraph;
+			AddedSimulationStageGraph->ForceBaseId(ENiagaraScriptUsage::ParticleSimulationStageScript, AddedOtherSimulationStage->GetUsageId(), ScriptBaseIdFromDiff);
 
 			Results.bModifiedGraph = true;
 		}
diff --git a/Engine/Plugins/FX/Niagara/Source/NiagaraEditor/Private/NiagaraScriptMergeManager.h b/Engine/Plugins/FX/Niagara/Source/NiagaraEditor/Private/NiagaraScriptMergeManager.h
--- a/Engine/Plugins/FX/Niagara/Source/NiagaraEditor/Private/NiagaraScriptMergeManager.h
+++ b/Engine/Plugins/FX/Niagara/Source/NiagaraEditor/Private/NiagaraScriptMergeManager.h
@@ -24,7 +24,7 @@ class UEdGraphNode;
 class UEdGraphPin;
 class UNiagaraDataInterface;
 struct FNiagaraEventScriptProperties;
-class UNiagaraShaderStageBase;
+class UNiagaraSimulationStageBase;
 class UNiagaraRendererProperties;
 
 class FNiagaraStackFunctionMergeAdapter;
@@ -158,32 +158,32 @@ private:
 	TSharedPtr<FNiagaraScriptStackMergeAdapter> EventStack;
 };
 
-class FNiagaraShaderStageMergeAdapter
+class FNiagaraSimulationStageMergeAdapter
 {
 public:
-	FNiagaraShaderStageMergeAdapter(const UNiagaraEmitter& InEmitter, const UNiagaraShaderStageBase* InShaderStage, UNiagaraNodeOutput* InOutputNode);
-	FNiagaraShaderStageMergeAdapter(const UNiagaraEmitter& InEmitter, UNiagaraShaderStageBase* InShaderStage, UNiagaraNodeOutput* InOutputNode);
-	FNiagaraShaderStageMergeAdapter(const UNiagaraEmitter& InEmitter, UNiagaraNodeOutput* InOutputNode);
+	FNiagaraSimulationStageMergeAdapter(const UNiagaraEmitter& InEmitter, const UNiagaraSimulationStageBase* InSimulationStage, UNiagaraNodeOutput* InOutputNode);
+	FNiagaraSimulationStageMergeAdapter(const UNiagaraEmitter& InEmitter, UNiagaraSimulationStageBase* InSimulationStage, UNiagaraNodeOutput* InOutputNode);
+	FNiagaraSimulationStageMergeAdapter(const UNiagaraEmitter& InEmitter, UNiagaraNodeOutput* InOutputNode);
 
 	FGuid GetUsageId() const;
 	const UNiagaraEmitter* GetEmitter() const;
-	const UNiagaraShaderStageBase* GetShaderStage() const;
-	UNiagaraShaderStageBase* GetEditableShaderStage() const;
+	const UNiagaraSimulationStageBase* GetSimulationStage() const;
+	UNiagaraSimulationStageBase* GetEditableSimulationStage() const;
 	UNiagaraNodeOutput* GetOutputNode() const;
 	UNiagaraNodeInput* GetInputNode() const;
-	TSharedPtr<FNiagaraScriptStackMergeAdapter> GetShaderStageStack() const;
+	TSharedPtr<FNiagaraScriptStackMergeAdapter> GetSimulationStageStack() const;
 
 private:
-	void Initialize(const UNiagaraEmitter& InEmitter, const UNiagaraShaderStageBase* InShaderStage, UNiagaraShaderStageBase* InEditableShaderStage, UNiagaraNodeOutput* InOutputNode);
+	void Initialize(const UNiagaraEmitter& InEmitter, const UNiagaraSimulationStageBase* InSimulationStage, UNiagaraSimulationStageBase* InEditableSimulationStage, UNiagaraNodeOutput* InOutputNode);
 
 private:
 	TWeakObjectPtr<UNiagaraEmitter> Emitter;
-	const UNiagaraShaderStageBase* ShaderStage;
-	UNiagaraShaderStageBase* EditableShaderStage;
+	const UNiagaraSimulationStageBase* SimulationStage;
+	UNiagaraSimulationStageBase* EditableSimulationStage;
 	TWeakObjectPtr<UNiagaraNodeOutput> OutputNode;
 	TWeakObjectPtr<UNiagaraNodeInput> InputNode;
 
-	TSharedPtr<FNiagaraScriptStackMergeAdapter> ShaderStageStack;
+	TSharedPtr<FNiagaraScriptStackMergeAdapter> SimulationStageStack;
 };
 
 class FNiagaraRendererMergeAdapter
@@ -210,7 +210,7 @@ public:
 	TSharedPtr<FNiagaraScriptStackMergeAdapter> GetParticleSpawnStack() const;
 	TSharedPtr<FNiagaraScriptStackMergeAdapter> GetParticleUpdateStack() const;
 	const TArray<TSharedRef<FNiagaraEventHandlerMergeAdapter>> GetEventHandlers() const;
-	const TArray<TSharedRef<FNiagaraShaderStageMergeAdapter>> GetShaderStages() const;
+	const TArray<TSharedRef<FNiagaraSimulationStageMergeAdapter>> GetSimulationStages() const;
 	const TArray<TSharedRef<FNiagaraRendererMergeAdapter>> GetRenderers() const;
 	const UNiagaraEmitterEditorData* GetEditorData() const;
 
@@ -218,7 +218,7 @@ public:
 
 	TSharedPtr<FNiagaraEventHandlerMergeAdapter> GetEventHandler(FGuid EventScriptUsageId);
 
-	TSharedPtr<FNiagaraShaderStageMergeAdapter> GetShaderStage(FGuid ShaderStageUsageId);
+	TSharedPtr<FNiagaraSimulationStageMergeAdapter> GetSimulationStage(FGuid SimulationStageUsageId);
 
 	TSharedPtr<FNiagaraRendererMergeAdapter> GetRenderer(FGuid RendererMergeId);
 
@@ -233,7 +233,7 @@ private:
 	TSharedPtr<FNiagaraScriptStackMergeAdapter> ParticleSpawnStack;
 	TSharedPtr<FNiagaraScriptStackMergeAdapter> ParticleUpdateStack;
 	TArray<TSharedRef<FNiagaraEventHandlerMergeAdapter>> EventHandlers;
-	TArray<TSharedRef<FNiagaraShaderStageMergeAdapter>> ShaderStages;
+	TArray<TSharedRef<FNiagaraSimulationStageMergeAdapter>> SimulationStages;
 	TArray<TSharedRef<FNiagaraRendererMergeAdapter>> Renderers;
 	TWeakObjectPtr<const UNiagaraEmitterEditorData> EditorData;
 };
@@ -280,10 +280,10 @@ struct FNiagaraModifiedEventHandlerDiffResults
 	FNiagaraScriptStackDiffResults ScriptDiffResults;
 };
 
-struct FNiagaraModifiedShaderStageDiffResults
+struct FNiagaraModifiedSimulationStageDiffResults
 {
-	TSharedPtr<FNiagaraShaderStageMergeAdapter> BaseAdapter;
-	TSharedPtr<FNiagaraShaderStageMergeAdapter> OtherAdapter;
+	TSharedPtr<FNiagaraSimulationStageMergeAdapter> BaseAdapter;
+	TSharedPtr<FNiagaraSimulationStageMergeAdapter> OtherAdapter;
 	TArray<FProperty*> ChangedProperties;
 	FNiagaraScriptStackDiffResults ScriptDiffResults;
 };
@@ -313,9 +313,9 @@ struct FNiagaraEmitterDiffResults
 	TArray<TSharedRef<FNiagaraEventHandlerMergeAdapter>> AddedOtherEventHandlers;
 	TArray<FNiagaraModifiedEventHandlerDiffResults> ModifiedEventHandlers;
 
-	TArray<TSharedRef<FNiagaraShaderStageMergeAdapter>> RemovedBaseShaderStages;
-	TArray<TSharedRef<FNiagaraShaderStageMergeAdapter>> AddedOtherShaderStages;
-	TArray<FNiagaraModifiedShaderStageDiffResults> ModifiedShaderStages;
+	TArray<TSharedRef<FNiagaraSimulationStageMergeAdapter>> RemovedBaseSimulationStages;
+	TArray<TSharedRef<FNiagaraSimulationStageMergeAdapter>> AddedOtherSimulationStages;
+	TArray<FNiagaraModifiedSimulationStageDiffResults> ModifiedSimulationStages;
 
 	TArray<TSharedRef<FNiagaraRendererMergeAdapter>> RemovedBaseRenderers;
 	TArray<TSharedRef<FNiagaraRendererMergeAdapter>> AddedOtherRenderers;
@@ -366,11 +366,11 @@ public:
 
 	void ResetEventHandlerPropertySetToBase(UNiagaraEmitter& Emitter, const UNiagaraEmitter& BaseEmitter, FGuid EventScriptUsageId);
 
-	bool HasBaseShaderStage(const UNiagaraEmitter& BaseEmitter, FGuid ShaderStageScriptUsageId);
+	bool HasBaseSimulationStage(const UNiagaraEmitter& BaseEmitter, FGuid SimulationStageScriptUsageId);
 
-	bool IsShaderStagePropertySetDifferentFromBase(UNiagaraEmitter& Emitter, const UNiagaraEmitter& BaseEmitter, FGuid ShaderStageScriptUsageId);
+	bool IsSimulationStagePropertySetDifferentFromBase(UNiagaraEmitter& Emitter, const UNiagaraEmitter& BaseEmitter, FGuid SimulationStageScriptUsageId);
 
-	void ResetShaderStagePropertySetToBase(UNiagaraEmitter& Emitter, const UNiagaraEmitter& BaseEmitter, FGuid ShaderStageScriptUsageId);
+	void ResetSimulationStagePropertySetToBase(UNiagaraEmitter& Emitter, const UNiagaraEmitter& BaseEmitter, FGuid SimulationStageScriptUsageId);
 
 	bool HasBaseRenderer(const UNiagaraEmitter& BaseEmitter, FGuid RendererMergeId);
 
@@ -384,7 +384,7 @@ public:
 
 	void DiffEventHandlers(const TArray<TSharedRef<FNiagaraEventHandlerMergeAdapter>>& BaseEventHandlers, const TArray<TSharedRef<FNiagaraEventHandlerMergeAdapter>>& OtherEventHandlers, FNiagaraEmitterDiffResults& DiffResults) const;
 
-	void DiffShaderStages(const TArray<TSharedRef<FNiagaraShaderStageMergeAdapter>>& BaseShaderStages, const TArray<TSharedRef<FNiagaraShaderStageMergeAdapter>>& OtherShaderStages, FNiagaraEmitterDiffResults& DiffResults) const;
+	void DiffSimulationStages(const TArray<TSharedRef<FNiagaraSimulationStageMergeAdapter>>& BaseSimulationStages, const TArray<TSharedRef<FNiagaraSimulationStageMergeAdapter>>& OtherSimulationStages, FNiagaraEmitterDiffResults& DiffResults) const;
 
 	void DiffRenderers(const TArray<TSharedRef<FNiagaraRendererMergeAdapter>>& BaseRenderers, const TArray<TSharedRef<FNiagaraRendererMergeAdapter>>& OtherRenderers, FNiagaraEmitterDiffResults& DiffResults) const;
 
@@ -406,7 +406,7 @@ private:
 
 	FApplyDiffResults ApplyEventHandlerDiff(TSharedRef<FNiagaraEmitterMergeAdapter> BaseEmitterAdapter, const FNiagaraEmitterDiffResults& DiffResults, const bool bNoParentAtLastMerge) const;
 
-	FApplyDiffResults ApplyShaderStageDiff(TSharedRef<FNiagaraEmitterMergeAdapter> BaseEmitterAdapter, const FNiagaraEmitterDiffResults& DiffResults, const bool bNoParentAtLastMerge) const;
+	FApplyDiffResults ApplySimulationStageDiff(TSharedRef<FNiagaraEmitterMergeAdapter> BaseEmitterAdapter, const FNiagaraEmitterDiffResults& DiffResults, const bool bNoParentAtLastMerge) const;
 
 	FApplyDiffResults ApplyRendererDiff(UNiagaraEmitter& BaseEmitter, const FNiagaraEmitterDiffResults& DiffResults, const bool bNoParentAtLastMerge) const;
 
diff --git a/Engine/Plugins/FX/Niagara/Source/NiagaraEditor/Private/ViewModels/NiagaraSystemViewModel.cpp b/Engine/Plugins/FX/Niagara/Source/NiagaraEditor/Private/ViewModels/NiagaraSystemViewModel.cpp
--- a/Engine/Plugins/FX/Niagara/Source/NiagaraEditor/Private/ViewModels/NiagaraSystemViewModel.cpp
+++ b/Engine/Plugins/FX/Niagara/Source/NiagaraEditor/Private/ViewModels/NiagaraSystemViewModel.cpp
@@ -1309,6 +1309,7 @@ void FNiagaraSystemViewModel::UpdateCompiledDataInterfaces(UNiagaraDataInterface
 			case ENiagaraScriptUsage::ParticleSpawnScript:
 			case ENiagaraScriptUsage::ParticleSpawnScriptInterpolated:
 			case ENiagaraScriptUsage::ParticleUpdateScript:
+			case ENiagaraScriptUsage::ParticleSimulationStageScript:
 			case ENiagaraScriptUsage::ParticleEventScript:
 				UpdateCompiledDataInterfacesForScript(*OwningScript, OuterInputNode->Input.GetName(), *ChangedDataInterface);
 				break;
diff --git a/Engine/Plugins/FX/Niagara/Source/NiagaraEditor/Private/ViewModels/Stack/NiagaraStackEntry.cpp b/Engine/Plugins/FX/Niagara/Source/NiagaraEditor/Private/ViewModels/Stack/NiagaraStackEntry.cpp
--- a/Engine/Plugins/FX/Niagara/Source/NiagaraEditor/Private/ViewModels/Stack/NiagaraStackEntry.cpp
+++ b/Engine/Plugins/FX/Niagara/Source/NiagaraEditor/Private/ViewModels/Stack/NiagaraStackEntry.cpp
@@ -19,7 +19,7 @@ const FName UNiagaraStackEntry::FExecutionSubcategoryNames::Settings = TEXT("Set
 const FName UNiagaraStackEntry::FExecutionSubcategoryNames::Spawn = TEXT("Spawn");
 const FName UNiagaraStackEntry::FExecutionSubcategoryNames::Update = TEXT("Update");
 const FName UNiagaraStackEntry::FExecutionSubcategoryNames::Event = TEXT("Event");
-const FName UNiagaraStackEntry::FExecutionSubcategoryNames::ShaderStage = TEXT("Shader Stage");
+const FName UNiagaraStackEntry::FExecutionSubcategoryNames::SimulationStage = TEXT("Simulation Stage");
 const FName UNiagaraStackEntry::FExecutionSubcategoryNames::Render = TEXT("Render");
 
 UNiagaraStackEntry::FStackIssueFix::FStackIssueFix()
diff --git a/Engine/Plugins/FX/Niagara/Source/NiagaraEditor/Private/ViewModels/Stack/NiagaraStackGraphUtilities.cpp b/Engine/Plugins/FX/Niagara/Source/NiagaraEditor/Private/ViewModels/Stack/NiagaraStackGraphUtilities.cpp
--- a/Engine/Plugins/FX/Niagara/Source/NiagaraEditor/Private/ViewModels/Stack/NiagaraStackGraphUtilities.cpp
+++ b/Engine/Plugins/FX/Niagara/Source/NiagaraEditor/Private/ViewModels/Stack/NiagaraStackGraphUtilities.cpp
@@ -1552,6 +1552,7 @@ void FNiagaraStackGraphUtilities::GetAvailableParametersForScript(UNiagaraNodeOu
 	if (ScriptOutputNode.GetUsage() == ENiagaraScriptUsage::ParticleSpawnScript ||
 		ScriptOutputNode.GetUsage() == ENiagaraScriptUsage::ParticleSpawnScriptInterpolated ||
 		ScriptOutputNode.GetUsage() == ENiagaraScriptUsage::ParticleUpdateScript ||
+		ScriptOutputNode.GetUsage() == ENiagaraScriptUsage::ParticleSimulationStageScript ||
 		ScriptOutputNode.GetUsage() == ENiagaraScriptUsage::ParticleEventScript)
 	{
 		OutAvailableParameters.Append(FNiagaraConstants::GetCommonParticleAttributes());
@@ -1593,6 +1594,7 @@ TOptional<FName> FNiagaraStackGraphUtilities::GetNamespaceForScriptUsage(ENiagar
 	case ENiagaraScriptUsage::ParticleSpawnScriptInterpolated:
 	case ENiagaraScriptUsage::ParticleUpdateScript:
 	case ENiagaraScriptUsage::ParticleEventScript:
+	case ENiagaraScriptUsage::ParticleSimulationStageScript:
 		return FNiagaraParameterHandle::ParticleAttributeNamespace;
 	case ENiagaraScriptUsage::EmitterSpawnScript:
 	case ENiagaraScriptUsage::EmitterUpdateScript:
@@ -1845,7 +1847,7 @@ bool FNiagaraStackGraphUtilities::CanWriteParameterFromUsage(FNiagaraVariable Pa
 	case ENiagaraScriptUsage::ParticleSpawnScriptInterpolated:
 	case ENiagaraScriptUsage::ParticleUpdateScript:
 	case ENiagaraScriptUsage::ParticleEventScript:
-	case ENiagaraScriptUsage::ParticleShaderStageScript:
+	case ENiagaraScriptUsage::ParticleSimulationStageScript:
 		return ParameterHandle.IsParticleAttributeHandle();
 	default:
 		return false;
diff --git a/Engine/Plugins/FX/Niagara/Source/NiagaraEditor/Private/ViewModels/Stack/NiagaraStackRoot.cpp b/Engine/Plugins/FX/Niagara/Source/NiagaraEditor/Private/ViewModels/Stack/NiagaraStackRoot.cpp
--- a/Engine/Plugins/FX/Niagara/Source/NiagaraEditor/Private/ViewModels/Stack/NiagaraStackRoot.cpp
+++ b/Engine/Plugins/FX/Niagara/Source/NiagaraEditor/Private/ViewModels/Stack/NiagaraStackRoot.cpp
@@ -6,8 +6,8 @@
 #include "ViewModels/Stack/NiagaraStackRenderItemGroup.h"
 #include "ViewModels/Stack/NiagaraStackEventHandlerGroup.h"
 #include "ViewModels/Stack/NiagaraStackEventScriptItemGroup.h"
-#include "ViewModels/Stack/NiagaraStackShaderStagesGroup.h"
-#include "ViewModels/Stack/NiagaraStackShaderStageGroup.h"
+#include "ViewModels/Stack/NiagaraStackSimulationStagesGroup.h"
+#include "ViewModels/Stack/NiagaraStackSimulationStageGroup.h"
 #include "ViewModels/NiagaraSystemViewModel.h"
 #include "ViewModels/NiagaraEmitterViewModel.h"
 #include "NiagaraSystemScriptViewModel.h"
@@ -19,14 +19,6 @@
 
 #define LOCTEXT_NAMESPACE "NiagaraStackViewModel"
 
-int32 GbShowNewShaderStageStackUIOptions = 0;
-static FAutoConsoleVariableRef CVarShowNewShaderStageStackUIOptions(
-	TEXT("fx.Niagara.ShowNewShaderStageStackUIOptions"),
-	GbShowNewShaderStageStackUIOptions,
-	TEXT("If > 0 the new shader state UI will be shown in the stack view.  This UI will generate data, but that data is currently unused.\n"),
-	ECVF_Default
-);
-
 UNiagaraStackRoot::UNiagaraStackRoot()
 	: SystemSettingsGroup(nullptr)
 	, EmitterSettingsGroup(nullptr)
@@ -35,7 +27,7 @@ UNiagaraStackRoot::UNiagaraStackRoot()
 	, ParticleSpawnGroup(nullptr)
 	, ParticleUpdateGroup(nullptr)
 	, AddEventHandlerGroup(nullptr)
-	, AddShaderStageGroup(nullptr)
+	, AddSimulationStageGroup(nullptr)
 	, RenderGroup(nullptr)
 {
 }
@@ -52,7 +44,7 @@ void UNiagaraStackRoot::Initialize(FRequiredEntryData InRequiredEntryData, bool
 	ParticleSpawnGroup = nullptr;
 	ParticleUpdateGroup = nullptr;
 	AddEventHandlerGroup = nullptr;
-	AddShaderStageGroup = nullptr;
+	AddSimulationStageGroup = nullptr;
 	RenderGroup = nullptr;
 }
 
@@ -165,14 +157,14 @@ void UNiagaraStackRoot::RefreshChildrenInternal(const TArray<UNiagaraStackEntry*
 		AddEventHandlerGroup->SetOnItemAdded(UNiagaraStackEventHandlerGroup::FOnItemAdded::CreateUObject(this, &UNiagaraStackRoot::EmitterArraysChanged));
 	}
 
-	if (bIncludeEmitterInformation && AddShaderStageGroup == nullptr && GbShowNewShaderStageStackUIOptions > 0)
+	if (bIncludeEmitterInformation && AddSimulationStageGroup == nullptr && GetEmitterViewModel()->GetEmitter()->bSimulationStagesEnabled)
 	{
-		AddShaderStageGroup = NewObject<UNiagaraStackShaderStagesGroup>(this);
+		AddSimulationStageGroup = NewObject<UNiagaraStackSimulationStagesGroup>(this);
 		FRequiredEntryData RequiredEntryData(GetSystemViewModel(), GetEmitterViewModel(),
-			FExecutionCategoryNames::Particle, FExecutionSubcategoryNames::ShaderStage,
+			FExecutionCategoryNames::Particle, FExecutionSubcategoryNames::SimulationStage,
 			GetEmitterViewModel()->GetOrCreateEditorData().GetStackEditorData());
-		AddShaderStageGroup->Initialize(RequiredEntryData);
-		AddShaderStageGroup->SetOnItemAdded(UNiagaraStackShaderStagesGroup::FOnItemAdded::CreateUObject(this, &UNiagaraStackRoot::EmitterArraysChanged));
+		AddSimulationStageGroup->Initialize(RequiredEntryData);
+		AddSimulationStageGroup->SetOnItemAdded(UNiagaraStackSimulationStagesGroup::FOnItemAdded::CreateUObject(this, &UNiagaraStackRoot::EmitterArraysChanged));
 	}
 
 	if (bIncludeEmitterInformation && RenderGroup == nullptr)
@@ -221,27 +213,27 @@ void UNiagaraStackRoot::RefreshChildrenInternal(const TArray<UNiagaraStackEntry*
 
 		NewChildren.Add(AddEventHandlerGroup);
 
-		if (GbShowNewShaderStageStackUIOptions > 0)
+		if (GetEmitterViewModel()->GetEmitter()->bSimulationStagesEnabled)
 		{
-			for (UNiagaraShaderStageBase* ShaderStage : GetEmitterViewModel()->GetEmitter()->GetShaderStages())
+			for (UNiagaraSimulationStageBase* SimulationStage : GetEmitterViewModel()->GetEmitter()->GetSimulationStages())
 			{
-				UNiagaraStackShaderStageGroup* ShaderStageGroup = FindCurrentChildOfTypeByPredicate<UNiagaraStackShaderStageGroup>(CurrentChildren,
-					[ShaderStage](UNiagaraStackShaderStageGroup* CurrentShaderStageGroup) { return CurrentShaderStageGroup->GetShaderStage() == ShaderStage; });
+				UNiagaraStackSimulationStageGroup* SimulationStageGroup = FindCurrentChildOfTypeByPredicate<UNiagaraStackSimulationStageGroup>(CurrentChildren,
+					[SimulationStage](UNiagaraStackSimulationStageGroup* CurrentSimulationStageGroup) { return CurrentSimulationStageGroup->GetSimulationStage() == SimulationStage; });
 
-				if (ShaderStageGroup == nullptr)
+				if (SimulationStageGroup == nullptr)
 				{
-					ShaderStageGroup = NewObject<UNiagaraStackShaderStageGroup>(this);
+					SimulationStageGroup = NewObject<UNiagaraStackSimulationStageGroup>(this);
 					FRequiredEntryData RequiredEntryData(GetSystemViewModel(), GetEmitterViewModel(),
-						FExecutionCategoryNames::Particle, FExecutionSubcategoryNames::ShaderStage,
+						FExecutionCategoryNames::Particle, FExecutionSubcategoryNames::SimulationStage,
 						GetEmitterViewModel()->GetEditorData().GetStackEditorData());
-					ShaderStageGroup->Initialize(RequiredEntryData, GetEmitterViewModel()->GetSharedScriptViewModel(), ShaderStage);
-					ShaderStageGroup->SetOnModifiedShaderStages(UNiagaraStackShaderStageGroup::FOnModifiedShaderStages::CreateUObject(this, &UNiagaraStackRoot::EmitterArraysChanged));
+					SimulationStageGroup->Initialize(RequiredEntryData, GetEmitterViewModel()->GetSharedScriptViewModel(), SimulationStage);
+					SimulationStageGroup->SetOnModifiedSimulationStages(UNiagaraStackSimulationStageGroup::FOnModifiedSimulationStages::CreateUObject(this, &UNiagaraStackRoot::EmitterArraysChanged));
 				}
 
-				NewChildren.Add(ShaderStageGroup);
+				NewChildren.Add(SimulationStageGroup);
 			}
 
-			NewChildren.Add(AddShaderStageGroup);
+			NewChildren.Add(AddSimulationStageGroup);
 		}
 
 		NewChildren.Add(RenderGroup);
diff --git a/Engine/Plugins/FX/Niagara/Source/NiagaraEditor/Private/ViewModels/Stack/NiagaraStackShaderStageGroup.cpp b/Engine/Plugins/FX/Niagara/Source/NiagaraEditor/Private/ViewModels/Stack/NiagaraStackShaderStageGroup.cpp
deleted file mode 100644
diff --git a/Engine/Plugins/FX/Niagara/Source/NiagaraEditor/Private/ViewModels/Stack/NiagaraStackShaderStagesGroup.cpp b/Engine/Plugins/FX/Niagara/Source/NiagaraEditor/Private/ViewModels/Stack/NiagaraStackShaderStagesGroup.cpp
deleted file mode 100644
diff --git a/Engine/Plugins/FX/Niagara/Source/NiagaraEditor/Private/ViewModels/Stack/NiagaraStackSimulationStageGroup.cpp b/Engine/Plugins/FX/Niagara/Source/NiagaraEditor/Private/ViewModels/Stack/NiagaraStackSimulationStageGroup.cpp
new file mode 100644
--- /dev/null
+++ b/Engine/Plugins/FX/Niagara/Source/NiagaraEditor/Private/ViewModels/Stack/NiagaraStackSimulationStageGroup.cpp
@@ -0,0 +1,356 @@
+// Copyright Epic Games, Inc. All Rights Reserved.
+
+#include "ViewModels/Stack/NiagaraStackSimulationStageGroup.h"
+#include "NiagaraEmitter.h"
+#include "NiagaraSimulationStageBase.h"
+#include "ViewModels/Stack/NiagaraStackObject.h"
+#include "ViewModels/Stack/NiagaraStackModuleItem.h"
+#include "ViewModels/NiagaraScriptViewModel.h"
+#include "NiagaraScriptGraphViewModel.h"
+#include "ViewModels/NiagaraEmitterViewModel.h"
+#include "NiagaraGraph.h"
+#include "NiagaraNodeOutput.h"
+#include "NiagaraNodeFunctionCall.h"
+#include "EdGraphSchema_Niagara.h"
+#include "ViewModels/Stack/NiagaraStackErrorItem.h"
+#include "NiagaraScriptSource.h"
+#include "ViewModels/Stack/NiagaraStackGraphUtilities.h"
+#include "NiagaraScriptMergeManager.h"
+#include "NiagaraStackEditorData.h"
+#include "ViewModels/NiagaraSystemViewModel.h"
+
+#include "Internationalization/Internationalization.h"
+#include "ScopedTransaction.h"
+#include "IDetailTreeNode.h"
+
+#define LOCTEXT_NAMESPACE "UNiagaraStackSimulationStageGroup"
+
+void UNiagaraStackSimulationStagePropertiesItem::Initialize(FRequiredEntryData InRequiredEntryData, UNiagaraSimulationStageBase* InSimulationStage)
+{
+	checkf(SimulationStage.IsValid() == false, TEXT("Can not initialize more than once."));
+	SimulationStage = InSimulationStage;
+	FString SimulationStageStackEditorDataKey = FString::Printf(TEXT("SimulationStage-%s-Properties"), *SimulationStage->Script->GetUsageId().ToString(EGuidFormats::DigitsWithHyphens));
+	Super::Initialize(InRequiredEntryData, SimulationStageStackEditorDataKey);
+
+	SimulationStage->OnChanged().AddUObject(this, &UNiagaraStackSimulationStagePropertiesItem::SimulationStagePropertiesChanged);
+}
+
+void UNiagaraStackSimulationStagePropertiesItem::FinalizeInternal()
+{
+	if (SimulationStage.IsValid())
+	{
+		SimulationStage->OnChanged().RemoveAll(this);
+	}
+	Super::FinalizeInternal();
+}
+
+FText UNiagaraStackSimulationStagePropertiesItem::GetDisplayName() const
+{
+	return FText::Format(LOCTEXT("SimulationStagePropertiesDisplayNameFormat", "{0} Settings"), SimulationStage->GetClass()->GetDisplayNameText());
+}
+
+bool UNiagaraStackSimulationStagePropertiesItem::TestCanResetToBaseWithMessage(FText& OutCanResetToBaseMessage) const
+{
+	if (bCanResetToBaseCache.IsSet() == false)
+	{
+		if (HasBaseSimulationStage())
+		{
+			UNiagaraEmitter* Emitter = GetEmitterViewModel()->GetEmitter();
+			const UNiagaraEmitter* BaseEmitter = Emitter->GetParent();
+			if (BaseEmitter != nullptr && Emitter != BaseEmitter)
+			{
+				TSharedRef<FNiagaraScriptMergeManager> MergeManager = FNiagaraScriptMergeManager::Get();
+				bCanResetToBaseCache = MergeManager->IsSimulationStagePropertySetDifferentFromBase(*Emitter, *BaseEmitter, SimulationStage->Script->GetUsageId());
+			}
+			else
+			{
+				bCanResetToBaseCache = false;
+			}
+		}
+		else
+		{
+			bCanResetToBaseCache = false;
+		}
+	}
+	if (bCanResetToBaseCache.GetValue())
+	{
+		OutCanResetToBaseMessage = LOCTEXT("CanResetToBase", "Reset this simulation stage to the one defined in the parent emitter.");
+		return true;
+	}
+	else
+	{
+		OutCanResetToBaseMessage = LOCTEXT("CanNotResetToBase", "No parent to reset to, or not different from parent.");
+		return false;
+	}
+}
+
+void UNiagaraStackSimulationStagePropertiesItem::ResetToBase()
+{
+	FText Unused;
+	if (TestCanResetToBaseWithMessage(Unused))
+	{
+		UNiagaraEmitter* Emitter = GetEmitterViewModel()->GetEmitter();
+		const UNiagaraEmitter* BaseEmitter = Emitter->GetParent();
+		TSharedRef<FNiagaraScriptMergeManager> MergeManager = FNiagaraScriptMergeManager::Get();
+		MergeManager->ResetSimulationStagePropertySetToBase(*Emitter, *BaseEmitter, SimulationStage->Script->GetUsageId());
+		RefreshChildren();
+	}
+}
+
+void UNiagaraStackSimulationStagePropertiesItem::RefreshChildrenInternal(const TArray<UNiagaraStackEntry*>& CurrentChildren, TArray<UNiagaraStackEntry*>& NewChildren, TArray<FStackIssue>& NewIssues)
+{
+	if (SimulationStageObject == nullptr)
+	{
+		SimulationStageObject = NewObject<UNiagaraStackObject>(this);
+		SimulationStageObject->Initialize(CreateDefaultChildRequiredData(), SimulationStage.Get(), GetStackEditorDataKey());
+	}
+
+	NewChildren.Add(SimulationStageObject);
+
+	bCanResetToBaseCache.Reset();
+	bHasBaseSimulationStageCache.Reset();
+
+	Super::RefreshChildrenInternal(CurrentChildren, NewChildren, NewIssues);
+}
+
+void UNiagaraStackSimulationStagePropertiesItem::SimulationStagePropertiesChanged()
+{
+	bCanResetToBaseCache.Reset();
+}
+
+bool UNiagaraStackSimulationStagePropertiesItem::HasBaseSimulationStage() const
+{
+	if (bHasBaseSimulationStageCache.IsSet() == false)
+	{
+		UNiagaraEmitter* Emitter = GetEmitterViewModel()->GetEmitter();
+		const UNiagaraEmitter* BaseEmitter = Emitter->GetParent();
+		if (BaseEmitter != nullptr && Emitter != BaseEmitter)
+		{
+			TSharedRef<FNiagaraScriptMergeManager> MergeManager = FNiagaraScriptMergeManager::Get();
+			bHasBaseSimulationStageCache = MergeManager->HasBaseSimulationStage(*BaseEmitter, SimulationStage->Script->GetUsageId());
+		}
+		else
+		{
+			bHasBaseSimulationStageCache = false;
+		}
+	}
+	return bHasBaseSimulationStageCache.GetValue();
+}
+
+void UNiagaraStackSimulationStageGroup::Initialize(
+	FRequiredEntryData InRequiredEntryData,
+	TSharedRef<FNiagaraScriptViewModel> InScriptViewModel,
+	UNiagaraSimulationStageBase* InSimulationStage)
+{
+	SimulationStage = InSimulationStage;
+	SimulationStage->OnChanged().AddUObject(this, &UNiagaraStackSimulationStageGroup::SimulationStagePropertiesChanged);
+
+	FText ToolTip = LOCTEXT("SimulationStageGroupTooltip", "Defines properties and script modules for a simulation stage.");
+	FText DisplayName = FText::FromName(SimulationStage->SimulationStageName);
+	Super::Initialize(InRequiredEntryData, DisplayName, ToolTip, InScriptViewModel, ENiagaraScriptUsage::ParticleSimulationStageScript, SimulationStage->Script->GetUsageId());
+}
+
+UNiagaraSimulationStageBase* UNiagaraStackSimulationStageGroup::GetSimulationStage() const
+{
+	return SimulationStage.Get();
+}
+
+void UNiagaraStackSimulationStageGroup::FinalizeInternal()
+{
+	if (SimulationStage.IsValid())
+	{
+		SimulationStage->OnChanged().RemoveAll(this);
+	}
+	Super::FinalizeInternal();
+}
+
+void UNiagaraStackSimulationStageGroup::RefreshChildrenInternal(const TArray<UNiagaraStackEntry*>& CurrentChildren, TArray<UNiagaraStackEntry*>& NewChildren, TArray<FStackIssue>& NewIssues)
+{
+	bHasBaseSimulationStageCache.Reset();
+
+	UNiagaraEmitter* Emitter = GetEmitterViewModel()->GetEmitter();
+
+	SetDisplayName(FText::FromName(SimulationStage->SimulationStageName));
+
+	if (SimulationStageProperties == nullptr)
+	{
+		SimulationStageProperties = NewObject<UNiagaraStackSimulationStagePropertiesItem>(this);
+		SimulationStageProperties->Initialize(CreateDefaultChildRequiredData(), SimulationStage.Get());
+	}
+	NewChildren.Add(SimulationStageProperties);
+	
+	Super::RefreshChildrenInternal(CurrentChildren, NewChildren, NewIssues);
+}
+
+TOptional<UNiagaraStackEntry::FDropRequestResponse> UNiagaraStackSimulationStageGroup::CanDropInternal(const FDropRequest& DropRequest)
+{
+	if (DropRequest.DragDropOperation->IsOfType<FNiagaraStackEntryDragDropOp>() && 
+		(DropRequest.DropZone == EItemDropZone::AboveItem || DropRequest.DropZone == EItemDropZone::BelowItem))
+	{
+		TSharedRef<const FNiagaraStackEntryDragDropOp> StackEntryDragDropOp = StaticCastSharedRef<const FNiagaraStackEntryDragDropOp>(DropRequest.DragDropOperation);
+		bool SimulationStageEntriesDragged = false;
+		for (UNiagaraStackEntry* DraggedEntry : StackEntryDragDropOp->GetDraggedEntries())
+		{
+			if (DraggedEntry->IsA<UNiagaraStackSimulationStageGroup>())
+			{
+				SimulationStageEntriesDragged = true;
+				break;
+			}
+		}
+
+		if (SimulationStageEntriesDragged == false)
+		{
+			// Only handle dragged SimulationStage items.
+			return Super::CanDropInternal(DropRequest);
+		}
+
+		if (DropRequest.DropOptions != UNiagaraStackEntry::EDropOptions::Overview)
+		{
+			// Only allow dropping in the overview stacks.
+			return FDropRequestResponse(TOptional<EItemDropZone>(), LOCTEXT("CantDropSimulationStageOnStack", "Simulation stages can only be dropped into the overview."));
+		}
+
+		if (StackEntryDragDropOp->GetDraggedEntries().Num() != 1)
+		{
+			// Only handle a single items.
+			return FDropRequestResponse(TOptional<EItemDropZone>(), LOCTEXT("CantDropMultipleSimulationStages", "Only single simulation stages can be dragged and dropped."));
+		}
+
+		UNiagaraStackSimulationStageGroup* SourceSimulationStageGroup = CastChecked<UNiagaraStackSimulationStageGroup>(StackEntryDragDropOp->GetDraggedEntries()[0]);
+		if (DropRequest.DragOptions != EDragOptions::Copy && SourceSimulationStageGroup->HasBaseSimulationStage())
+		{
+			return FDropRequestResponse(TOptional<EItemDropZone>(), LOCTEXT("CantMoveSimulationStageError", "This inherited simulation stage can't be moved."));
+		}
+
+		if (SourceSimulationStageGroup == this)
+		{
+			// Don't allow dropping on yourself.
+			return Super::CanDropInternal(DropRequest);
+		}
+
+		UNiagaraEmitter* OwningEmitter = GetEmitterViewModel()->GetEmitter();
+		int32 SourceIndex = OwningEmitter->GetSimulationStages().IndexOfByKey(SourceSimulationStageGroup->GetSimulationStage());
+		if (SourceIndex == INDEX_NONE)
+		{
+			return FDropRequestResponse(TOptional<EItemDropZone>(), LOCTEXT("CantDropSimulationStageFromOtherEmitterError", "This simulation stage can't be moved here\nbecause it's owned by a different emitter."));
+		}
+
+		int32 TargetIndexOffset = DropRequest.DropZone == EItemDropZone::BelowItem ? 1 : 0;
+		int32 TargetIndex = OwningEmitter->GetSimulationStages().IndexOfByKey(SimulationStage.Get()) + TargetIndexOffset;
+		if (SourceIndex < TargetIndex)
+		{
+			// If the source index is less than the target index, the target index will decrease by 1 after the source is removed.
+			TargetIndex--;
+		}
+
+		if (SourceIndex == TargetIndex)
+		{
+			// Only handle the drag if the item would actually move.
+			return Super::CanDropInternal(DropRequest);
+		}
+
+		return FDropRequestResponse(DropRequest.DropZone, LOCTEXT("MoveSimulationStageDragMessage", "Move this simulation stage here."));
+	}
+	return Super::CanDropInternal(DropRequest);
+}
+
+TOptional<UNiagaraStackEntry::FDropRequestResponse> UNiagaraStackSimulationStageGroup::DropInternal(const FDropRequest& DropRequest)
+{
+	UNiagaraEmitter* OwningEmitter = GetEmitterViewModel()->GetEmitter();
+
+	TSharedRef<const FNiagaraStackEntryDragDropOp> StackEntryDragDropOp = StaticCastSharedRef<const FNiagaraStackEntryDragDropOp>(DropRequest.DragDropOperation);
+	UNiagaraStackSimulationStageGroup* SourceSimulationStageGroup = Cast<UNiagaraStackSimulationStageGroup>(StackEntryDragDropOp->GetDraggedEntries()[0]);
+
+	if (SourceSimulationStageGroup)
+	{
+		int32 SourceIndex = OwningEmitter->GetSimulationStages().IndexOfByKey(SourceSimulationStageGroup->GetSimulationStage());
+		if (SourceIndex != INDEX_NONE)
+		{
+			int32 TargetOffset = DropRequest.DropZone == EItemDropZone::BelowItem ? 1 : 0;
+			int32 TargetIndex = OwningEmitter->GetSimulationStages().IndexOfByKey(SimulationStage.Get()) + TargetOffset;
+
+			FScopedTransaction Transaction(FText::Format(LOCTEXT("MoveSimulationStage", "Move Shader Stage {0}"), GetDisplayName()));
+			OwningEmitter->MoveSimulationStageToIndex(SourceSimulationStageGroup->GetSimulationStage(), TargetIndex);
+
+			OnRequestFullRefreshDeferred().Broadcast();
+			return FDropRequestResponse(DropRequest.DropZone, FText());
+		}
+	}
+	return Super::DropInternal(DropRequest);
+}
+
+bool UNiagaraStackSimulationStageGroup::TestCanDeleteWithMessage(FText& OutCanDeleteMessage) const
+{
+	if (HasBaseSimulationStage())
+	{
+		OutCanDeleteMessage = LOCTEXT("CantDeleteInherited", "Can not delete this simulation stage because it's inherited.");
+		return false;
+	}
+	else
+	{
+		OutCanDeleteMessage = LOCTEXT("CanDelete", "Delete this simulation stage.");
+		return true;
+	}
+}
+
+void UNiagaraStackSimulationStageGroup::Delete()
+{
+	TSharedPtr<FNiagaraScriptViewModel> ScriptViewModelPinned = ScriptViewModel.Pin();
+	checkf(ScriptViewModelPinned.IsValid(), TEXT("Can not delete when the script view model has been deleted."));
+
+	UNiagaraEmitter* Emitter = GetEmitterViewModel()->GetEmitter();
+	UNiagaraScriptSource* Source = Cast<UNiagaraScriptSource>(Emitter->GraphSource);
+
+	if (!Source || !Source->NodeGraph)
+	{
+		return;
+	}
+
+	FScopedTransaction Transaction(FText::Format(LOCTEXT("DeleteSimulationStage", "Delete {0}"), GetDisplayName()));
+
+	Emitter->Modify();
+	Source->NodeGraph->Modify();
+	TArray<UNiagaraNode*> SimulationStageNodes;
+	Source->NodeGraph->BuildTraversal(SimulationStageNodes, GetScriptUsage(), GetScriptUsageId());
+	for (UNiagaraNode* Node : SimulationStageNodes)
+	{
+		Node->Modify();
+	}
+	
+	// First, remove the simulation stage object.
+	Emitter->RemoveSimulationStage(SimulationStage.Get());
+	
+	// Now remove all graph nodes associated with the simulation stage.
+	for (UNiagaraNode* Node : SimulationStageNodes)
+	{
+		Node->DestroyNode();
+	}
+
+	// Set the emitter here to that the internal state of the view model is updated.
+	// TODO: Move the logic for managing additional scripts into the emitter view model or script view model.
+	ScriptViewModelPinned->SetScripts(Emitter);
+	
+	OnModifiedSimulationStagesDelegate.ExecuteIfBound();
+}
+
+void UNiagaraStackSimulationStageGroup::SimulationStagePropertiesChanged()
+{
+	SetDisplayName(FText::FromName(SimulationStage->SimulationStageName));
+}
+
+bool UNiagaraStackSimulationStageGroup::HasBaseSimulationStage() const
+{
+	if (bHasBaseSimulationStageCache.IsSet() == false)
+	{
+		const UNiagaraEmitter* BaseEmitter = GetEmitterViewModel()->GetEmitter()->GetParent();
+		bHasBaseSimulationStageCache = BaseEmitter != nullptr && FNiagaraScriptMergeManager::Get()->HasBaseSimulationStage(*BaseEmitter, GetScriptUsageId());
+	}
+	return bHasBaseSimulationStageCache.GetValue();
+}
+
+void UNiagaraStackSimulationStageGroup::SetOnModifiedSimulationStages(FOnModifiedSimulationStages OnModifiedSimulationStages)
+{
+	OnModifiedSimulationStagesDelegate = OnModifiedSimulationStages;
+}
+
+#undef LOCTEXT_NAMESPACE
diff --git a/Engine/Plugins/FX/Niagara/Source/NiagaraEditor/Private/ViewModels/Stack/NiagaraStackSimulationStagesGroup.cpp b/Engine/Plugins/FX/Niagara/Source/NiagaraEditor/Private/ViewModels/Stack/NiagaraStackSimulationStagesGroup.cpp
new file mode 100644
--- /dev/null
+++ b/Engine/Plugins/FX/Niagara/Source/NiagaraEditor/Private/ViewModels/Stack/NiagaraStackSimulationStagesGroup.cpp
@@ -0,0 +1,135 @@
+// Copyright Epic Games, Inc. All Rights Reserved.
+
+#include "ViewModels/Stack/NiagaraStackSimulationStagesGroup.h"
+#include "ViewModels/NiagaraEmitterViewModel.h"
+#include "ViewModels/NiagaraScriptViewModel.h"
+#include "NiagaraScriptGraphViewModel.h"
+#include "NiagaraEmitter.h"
+#include "NiagaraSimulationStageBase.h"
+#include "NiagaraGraph.h"
+#include "NiagaraScriptSource.h"
+#include "ViewModels/Stack/INiagaraStackItemGroupAddUtilities.h"
+#include "ViewModels/Stack/NiagaraStackGraphUtilities.h"
+
+#include "ScopedTransaction.h"
+
+#define LOCTEXT_NAMESPACE "NiagaraStackSimulationStagesGroup"
+
+class FSimulationStagesGroupAddAction : public INiagaraStackItemGroupAddAction
+{
+public:
+	FSimulationStagesGroupAddAction(UClass* InSimulationStageClass)
+		: SimulationStageClass(InSimulationStageClass)
+	{
+	}
+
+	virtual FText GetCategory() const override
+	{
+		return LOCTEXT("AddSimulationStageCategory", "Add Simulation Stage");
+	}
+
+	virtual FText GetDisplayName() const override
+	{
+		return SimulationStageClass->GetDisplayNameText();
+	}
+
+	virtual FText GetDescription() const override
+	{
+		return FText::FromString(SimulationStageClass->GetDescription());
+	}
+
+	virtual FText GetKeywords() const override
+	{
+		return FText();
+	}
+
+	UClass* GetSimulationStageClass() const
+	{
+		return SimulationStageClass;
+	}
+
+private:
+	UClass* SimulationStageClass;
+};
+
+class FSimulationStagesGroupAddUtilities : public TNiagaraStackItemGroupAddUtilities<UNiagaraSimulationStageBase*>
+{
+public:
+	FSimulationStagesGroupAddUtilities(TSharedRef<FNiagaraEmitterViewModel> InEmitterViewModel, FOnItemAdded InOnItemAdded)
+		: TNiagaraStackItemGroupAddUtilities(LOCTEXT("SimulationStagesGroupAddItemName", "Simulation Stage"), EAddMode::AddFromAction, true, InOnItemAdded)
+		, EmitterViewModel(InEmitterViewModel)
+	{
+	}
+
+	virtual void AddItemDirectly() override	{ unimplemented(); }
+
+	virtual void GenerateAddActions(TArray<TSharedRef<INiagaraStackItemGroupAddAction>>& OutAddActions, const FNiagaraStackItemGroupAddOptions& AddProperties) const override 
+	{
+		TArray<UClass*> SimulationStageClasses;
+		GetDerivedClasses(UNiagaraSimulationStageBase::StaticClass(), SimulationStageClasses);
+		for (UClass* SimulationStageClass : SimulationStageClasses)
+		{
+			OutAddActions.Add(MakeShared<FSimulationStagesGroupAddAction>(SimulationStageClass));
+		}
+	}
+
+	virtual void ExecuteAddAction(TSharedRef<INiagaraStackItemGroupAddAction> AddAction, int32 TargetIndex) override 
+	{ 
+		TSharedPtr<FNiagaraEmitterViewModel> EmitterViewModelPinned = EmitterViewModel.Pin();
+		if (EmitterViewModelPinned.IsValid() == false)
+		{
+			return;
+		}
+
+		UNiagaraEmitter* Emitter = EmitterViewModelPinned->GetEmitter();
+		UNiagaraScriptSource* Source = EmitterViewModelPinned->GetSharedScriptViewModel()->GetGraphViewModel()->GetScriptSource();
+		UNiagaraGraph* Graph = EmitterViewModelPinned->GetSharedScriptViewModel()->GetGraphViewModel()->GetGraph();
+		TSharedRef<FSimulationStagesGroupAddAction> SimulationStageAddAction = StaticCastSharedRef<FSimulationStagesGroupAddAction>(AddAction);
+
+		// The stack should not have been created if any of these are null, so bail out if it happens somehow rather than try to handle all of these cases.
+		checkf(Emitter != nullptr && Source != nullptr && Graph != nullptr, TEXT("Stack created for invalid emitter or graph."));
+
+		FScopedTransaction ScopedTransaction(LOCTEXT("AddNewSimulationStagesTransaction", "Add new shader stage"));
+
+		Emitter->Modify();
+		UNiagaraSimulationStageBase* SimulationStage = NewObject<UNiagaraSimulationStageBase>(Emitter, SimulationStageAddAction->GetSimulationStageClass(), NAME_None, RF_Transactional);
+		SimulationStage->Script = NewObject<UNiagaraScript>(SimulationStage, MakeUniqueObjectName(SimulationStage, UNiagaraScript::StaticClass(), "SimulationStage"), EObjectFlags::RF_Transactional);
+		SimulationStage->Script->SetUsage(ENiagaraScriptUsage::ParticleSimulationStageScript);
+		SimulationStage->Script->SetUsageId(SimulationStage->GetMergeId());
+		SimulationStage->Script->SetSource(Source);
+		Emitter->AddSimulationStage(SimulationStage);
+		FNiagaraStackGraphUtilities::ResetGraphForOutput(*Graph, ENiagaraScriptUsage::ParticleSimulationStageScript, SimulationStage->Script->GetUsageId());
+
+		// Set the emitter here so that the internal state of the view model is updated.
+		// TODO: Move the logic for managing additional scripts into the emitter view model or script view model.
+		TWeakPtr<FNiagaraEmitterInstance, ESPMode::ThreadSafe> Simulation = EmitterViewModelPinned->GetSimulation();
+		EmitterViewModelPinned->Reset();
+		EmitterViewModelPinned->Initialize(Emitter, Simulation);
+
+		OnItemAdded.ExecuteIfBound(SimulationStage);
+	}
+
+private:
+	TWeakPtr<FNiagaraEmitterViewModel> EmitterViewModel;
+};
+
+void UNiagaraStackSimulationStagesGroup::Initialize(FRequiredEntryData InRequiredEntryData)
+{
+	FText DisplayName = LOCTEXT("SimulationStageGroupName", "Add Simulation Stage");
+	FText ToolTip = LOCTEXT("SimulationStageGroupTooltip", "Shader stages for this emitter.");
+	AddUtilities = MakeShared<FSimulationStagesGroupAddUtilities>(InRequiredEntryData.EmitterViewModel.ToSharedRef(), 
+		TNiagaraStackItemGroupAddUtilities<UNiagaraSimulationStageBase*>::FOnItemAdded::CreateUObject(this, &UNiagaraStackSimulationStagesGroup::ItemAddedFromUtilties));
+	Super::Initialize(InRequiredEntryData, DisplayName, ToolTip, AddUtilities.Get());
+}
+
+void UNiagaraStackSimulationStagesGroup::SetOnItemAdded(FOnItemAdded InOnItemAdded)
+{
+	ItemAddedDelegate = InOnItemAdded;
+}
+
+void UNiagaraStackSimulationStagesGroup::ItemAddedFromUtilties(UNiagaraSimulationStageBase* AddedSimulationStage)
+{
+	ItemAddedDelegate.ExecuteIfBound();
+}
+
+#undef LOCTEXT_NAMESPACE
diff --git a/Engine/Plugins/FX/Niagara/Source/NiagaraEditor/Public/ViewModels/Stack/NiagaraStackEntry.h b/Engine/Plugins/FX/Niagara/Source/NiagaraEditor/Public/ViewModels/Stack/NiagaraStackEntry.h
--- a/Engine/Plugins/FX/Niagara/Source/NiagaraEditor/Public/ViewModels/Stack/NiagaraStackEntry.h
+++ b/Engine/Plugins/FX/Niagara/Source/NiagaraEditor/Public/ViewModels/Stack/NiagaraStackEntry.h
@@ -110,7 +110,7 @@ public:
 		static const FName Spawn;
 		static const FName Update;
 		static const FName Event;
-		static const FName ShaderStage;
+		static const FName SimulationStage;
 		static const FName Render;
 	};
 
diff --git a/Engine/Plugins/FX/Niagara/Source/NiagaraEditor/Public/ViewModels/Stack/NiagaraStackRoot.h b/Engine/Plugins/FX/Niagara/Source/NiagaraEditor/Public/ViewModels/Stack/NiagaraStackRoot.h
--- a/Engine/Plugins/FX/Niagara/Source/NiagaraEditor/Public/ViewModels/Stack/NiagaraStackRoot.h
+++ b/Engine/Plugins/FX/Niagara/Source/NiagaraEditor/Public/ViewModels/Stack/NiagaraStackRoot.h
@@ -10,7 +10,7 @@ class UNiagaraStackSystemSettingsGroup;
 class UNiagaraStackEmitterSettingsGroup;
 class UNiagaraStackScriptItemGroup;
 class UNiagaraStackEventHandlerGroup;
-class UNiagaraStackShaderStagesGroup;
+class UNiagaraStackSimulationStagesGroup;
 class UNiagaraStackRenderItemGroup;
 
 UCLASS()
@@ -64,7 +64,7 @@ private:
 	UNiagaraStackEventHandlerGroup* AddEventHandlerGroup;
 
 	UPROPERTY()
-	UNiagaraStackShaderStagesGroup* AddShaderStageGroup;
+	UNiagaraStackSimulationStagesGroup* AddSimulationStageGroup;
 
 	UPROPERTY()
 	UNiagaraStackRenderItemGroup* RenderGroup;
diff --git a/Engine/Plugins/FX/Niagara/Source/NiagaraEditor/Public/ViewModels/Stack/NiagaraStackShaderStageGroup.h b/Engine/Plugins/FX/Niagara/Source/NiagaraEditor/Public/ViewModels/Stack/NiagaraStackShaderStageGroup.h
deleted file mode 100644
diff --git a/Engine/Plugins/FX/Niagara/Source/NiagaraEditor/Public/ViewModels/Stack/NiagaraStackShaderStagesGroup.h b/Engine/Plugins/FX/Niagara/Source/NiagaraEditor/Public/ViewModels/Stack/NiagaraStackShaderStagesGroup.h
deleted file mode 100644
diff --git a/Engine/Plugins/FX/Niagara/Source/NiagaraEditor/Public/ViewModels/Stack/NiagaraStackSimulationStageGroup.h b/Engine/Plugins/FX/Niagara/Source/NiagaraEditor/Public/ViewModels/Stack/NiagaraStackSimulationStageGroup.h
new file mode 100644
--- /dev/null
+++ b/Engine/Plugins/FX/Niagara/Source/NiagaraEditor/Public/ViewModels/Stack/NiagaraStackSimulationStageGroup.h
@@ -0,0 +1,99 @@
+// Copyright Epic Games, Inc. All Rights Reserved.
+
+#pragma once
+
+#include "NiagaraCommon.h"
+#include "NiagaraEmitter.h"
+#include "ViewModels/Stack/NiagaraStackItem.h"
+#include "ViewModels/Stack/NiagaraStackScriptItemGroup.h"
+#include "NiagaraStackSimulationStageGroup.generated.h"
+
+class FNiagaraScriptViewModel;
+class UNiagaraStackObject;
+struct FNiagaraEventScriptProperties;
+class IDetailTreeNode;
+
+UCLASS()
+class NIAGARAEDITOR_API UNiagaraStackSimulationStagePropertiesItem : public UNiagaraStackItem
+{
+	GENERATED_BODY()
+		 
+public:
+	void Initialize(FRequiredEntryData InRequiredEntryData, UNiagaraSimulationStageBase* InSimulationStage);
+
+	virtual FText GetDisplayName() const override;
+
+	virtual bool SupportsResetToBase() const override { return true; }
+	virtual bool TestCanResetToBaseWithMessage(FText& OutCanResetToBaseMessage) const override;
+	virtual void ResetToBase() override;
+
+protected:
+	virtual void FinalizeInternal() override;
+
+	virtual void RefreshChildrenInternal(const TArray<UNiagaraStackEntry*>& CurrentChildren, TArray<UNiagaraStackEntry*>& NewChildren, TArray<FStackIssue>& NewIssues) override;
+
+private:
+	void SimulationStagePropertiesChanged();
+
+	bool HasBaseSimulationStage() const;
+
+private:
+	TWeakObjectPtr<UNiagaraSimulationStageBase> SimulationStage;
+
+	mutable TOptional<bool> bHasBaseSimulationStageCache;
+
+	mutable TOptional<bool> bCanResetToBaseCache;
+
+	UPROPERTY()
+	UNiagaraStackObject* SimulationStageObject;
+};
+
+UCLASS()
+/** Meant to contain a single binding of a Emitter::EventScriptProperties to the stack.*/
+class NIAGARAEDITOR_API UNiagaraStackSimulationStageGroup : public UNiagaraStackScriptItemGroup
+{
+	GENERATED_BODY()
+
+public:
+	DECLARE_DELEGATE(FOnModifiedSimulationStages);
+
+public:
+	void Initialize(
+		FRequiredEntryData InRequiredEntryData,
+		TSharedRef<FNiagaraScriptViewModel> InScriptViewModel,
+		UNiagaraSimulationStageBase* InSimulationStage);
+
+	UNiagaraSimulationStageBase* GetSimulationStage() const;
+
+	void SetOnModifiedSimulationStages(FOnModifiedSimulationStages OnModifiedSimulationStages);
+
+	virtual bool SupportsDelete() const override { return true; }
+	virtual bool TestCanDeleteWithMessage(FText& OutCanDeleteMessage) const override;
+	virtual void Delete() override;
+
+	virtual bool CanDrag() const override { return true; }
+
+protected:
+	virtual void FinalizeInternal() override;
+
+	virtual void RefreshChildrenInternal(const TArray<UNiagaraStackEntry*>& CurrentChildren, TArray<UNiagaraStackEntry*>& NewChildren, TArray<FStackIssue>& NewIssues) override;
+
+	virtual TOptional<FDropRequestResponse> CanDropInternal(const FDropRequest& DropRequest) override;
+
+	virtual TOptional<FDropRequestResponse> DropInternal(const FDropRequest& DropRequest) override;
+
+private:
+	void SimulationStagePropertiesChanged();
+
+	bool HasBaseSimulationStage() const;
+	
+private:
+	TWeakObjectPtr<UNiagaraSimulationStageBase> SimulationStage;
+
+	mutable TOptional<bool> bHasBaseSimulationStageCache;
+
+	FOnModifiedSimulationStages OnModifiedSimulationStagesDelegate;
+
+	UPROPERTY()
+	UNiagaraStackSimulationStagePropertiesItem* SimulationStageProperties;
+};
diff --git a/Engine/Plugins/FX/Niagara/Source/NiagaraEditor/Public/ViewModels/Stack/NiagaraStackSimulationStagesGroup.h b/Engine/Plugins/FX/Niagara/Source/NiagaraEditor/Public/ViewModels/Stack/NiagaraStackSimulationStagesGroup.h
new file mode 100644
--- /dev/null
+++ b/Engine/Plugins/FX/Niagara/Source/NiagaraEditor/Public/ViewModels/Stack/NiagaraStackSimulationStagesGroup.h
@@ -0,0 +1,30 @@
+// Copyright Epic Games, Inc. All Rights Reserved.
+
+#pragma once
+
+#include "ViewModels/Stack/NiagaraStackItemGroup.h"
+#include "NiagaraStackSimulationStagesGroup.generated.h"
+
+class UNiagaraSimulationStageBase;
+
+UCLASS()
+/** Container for one or more NiagaraStackEventScriptItemGroups, allowing multiple event handlers per script.*/
+class NIAGARAEDITOR_API UNiagaraStackSimulationStagesGroup : public UNiagaraStackItemGroup
+{
+	GENERATED_BODY()
+
+public:
+	DECLARE_DELEGATE(FOnItemAdded);
+
+public:
+	void Initialize(FRequiredEntryData InRequiredEntryData);
+
+	void SetOnItemAdded(FOnItemAdded InOnItemAdded);
+
+private:
+	void ItemAddedFromUtilties(UNiagaraSimulationStageBase* AddedShaderStage);
+
+private:
+	TSharedPtr<INiagaraStackItemGroupAddUtilities> AddUtilities;
+	FOnItemAdded ItemAddedDelegate;
+};
diff --git a/Engine/Plugins/FX/Niagara/Source/NiagaraEditorWidgets/Private/NiagaraEditorWidgetsStyle.cpp b/Engine/Plugins/FX/Niagara/Source/NiagaraEditorWidgets/Private/NiagaraEditorWidgetsStyle.cpp
--- a/Engine/Plugins/FX/Niagara/Source/NiagaraEditorWidgets/Private/NiagaraEditorWidgetsStyle.cpp
+++ b/Engine/Plugins/FX/Niagara/Source/NiagaraEditorWidgets/Private/NiagaraEditorWidgetsStyle.cpp
@@ -199,14 +199,14 @@ TSharedRef< FSlateStyleSet > FNiagaraEditorWidgetsStyle::Create()
 	Style->Set("NiagaraEditor.Stack.SpawnIcon", new IMAGE_PLUGIN_BRUSH("Icons/Spawn", Icon12x12, FLinearColor::White));
 	Style->Set("NiagaraEditor.Stack.UpdateIcon", new IMAGE_PLUGIN_BRUSH("Icons/Update", Icon12x12, FLinearColor::White));
 	Style->Set("NiagaraEditor.Stack.EventIcon", new IMAGE_PLUGIN_BRUSH("Icons/Event", Icon12x12, FLinearColor::White));
-	Style->Set("NiagaraEditor.Stack.ShaderStageIcon", new IMAGE_PLUGIN_BRUSH("Icons/ShaderStage", Icon12x12, FLinearColor::White));
+	Style->Set("NiagaraEditor.Stack.SimulationStageIcon", new IMAGE_PLUGIN_BRUSH("Icons/SimulationStage", Icon12x12, FLinearColor::White));
 	Style->Set("NiagaraEditor.Stack.RenderIcon", new IMAGE_PLUGIN_BRUSH("Icons/Render", Icon12x12, FLinearColor::White));
 
 	Style->Set("NiagaraEditor.Stack.ParametersIconHighlighted", new IMAGE_PLUGIN_BRUSH("Icons/SystemParams", Icon16x16, FLinearColor::White));
 	Style->Set("NiagaraEditor.Stack.SpawnIconHighlighted", new IMAGE_PLUGIN_BRUSH("Icons/Spawn", Icon16x16, FLinearColor::White));
 	Style->Set("NiagaraEditor.Stack.UpdateIconHighlighted", new IMAGE_PLUGIN_BRUSH("Icons/Update", Icon16x16, FLinearColor::White));
 	Style->Set("NiagaraEditor.Stack.EventIconHighlighted", new IMAGE_PLUGIN_BRUSH("Icons/Event", Icon16x16, FLinearColor::White));
-	Style->Set("NiagaraEditor.Stack.ShaderStageIconHighlighted", new IMAGE_PLUGIN_BRUSH("Icons/ShaderStage", Icon16x16, FLinearColor::White));
+	Style->Set("NiagaraEditor.Stack.SimulationStageIconHighlighted", new IMAGE_PLUGIN_BRUSH("Icons/SimulationStage", Icon16x16, FLinearColor::White));
 	Style->Set("NiagaraEditor.Stack.RenderIconHighlighted", new IMAGE_PLUGIN_BRUSH("Icons/Render", Icon16x16, FLinearColor::White));
 
 	Style->Set("NiagaraEditor.Stack.IconHighlightedSize", 16.0f);
diff --git a/Engine/Plugins/FX/Niagara/Source/NiagaraEditorWidgets/Private/NiagaraEditorWidgetsUtilities.cpp b/Engine/Plugins/FX/Niagara/Source/NiagaraEditorWidgets/Private/NiagaraEditorWidgetsUtilities.cpp
--- a/Engine/Plugins/FX/Niagara/Source/NiagaraEditorWidgets/Private/NiagaraEditorWidgetsUtilities.cpp
+++ b/Engine/Plugins/FX/Niagara/Source/NiagaraEditorWidgets/Private/NiagaraEditorWidgetsUtilities.cpp
@@ -64,9 +64,9 @@ FName FNiagaraStackEditorWidgetsUtilities::GetIconNameForExecutionSubcategory(FN
 		{
 			return "NiagaraEditor.Stack.EventIconHighlighted";
 		}
-		else if (ExecutionSubcategoryName == UNiagaraStackEntry::FExecutionSubcategoryNames::ShaderStage)
+		else if (ExecutionSubcategoryName == UNiagaraStackEntry::FExecutionSubcategoryNames::SimulationStage)
 		{
-			return "NiagaraEditor.Stack.ShaderStageIconHighlighted";
+			return "NiagaraEditor.Stack.SimulationStageIconHighlighted";
 		}
 		else if (ExecutionSubcategoryName == UNiagaraStackEntry::FExecutionSubcategoryNames::Render)
 		{
@@ -91,9 +91,9 @@ FName FNiagaraStackEditorWidgetsUtilities::GetIconNameForExecutionSubcategory(FN
 		{
 			return "NiagaraEditor.Stack.EventIcon";
 		}
-		else if (ExecutionSubcategoryName == UNiagaraStackEntry::FExecutionSubcategoryNames::ShaderStage)
+		else if (ExecutionSubcategoryName == UNiagaraStackEntry::FExecutionSubcategoryNames::SimulationStage)
 		{
-			return "NiagaraEditor.Stack.ShaderStageIcon";
+			return "NiagaraEditor.Stack.SimulationStageIcon";
 		}
 		else if (ExecutionSubcategoryName == UNiagaraStackEntry::FExecutionSubcategoryNames::Render)
 		{
diff --git a/Engine/Plugins/FX/Niagara/Source/NiagaraShader/Private/NiagaraShader.cpp b/Engine/Plugins/FX/Niagara/Source/NiagaraShader/Private/NiagaraShader.cpp
--- a/Engine/Plugins/FX/Niagara/Source/NiagaraShader/Private/NiagaraShader.cpp
+++ b/Engine/Plugins/FX/Niagara/Source/NiagaraShader/Private/NiagaraShader.cpp
@@ -383,9 +383,9 @@ TSharedRef<FShaderCommonCompileJob, ESPMode::ThreadSafe> FNiagaraShaderType::Beg
 	NewJob->Input.Environment.SetDefine(TEXT("DISKELMESH_BONE_INFLUENCES"), 0);
 	NewJob->Input.Environment.IncludeVirtualPathToContentsMap.Add(TEXT("/Engine/Generated/NiagaraEmitterInstance.ush"), Script->HlslOutput);
 
-	static const auto UseShaderStagesCVar = IConsoleManager::Get().FindConsoleVariable(TEXT("fx.UseShaderStages"));
-	NewJob->Input.Environment.SetDefine(TEXT("USE_SHADER_STAGES"), UseShaderStagesCVar->GetInt());
 	
+	NewJob->Input.Environment.SetDefine(TEXT("USE_SIMULATION_STAGES"), Script->GetUseSimStagesDefine());
+
 	AddReferencedUniformBufferIncludes(NewJob->Input.Environment, NewJob->Input.SourceFilePrefix, (EShaderPlatform)Target.Platform);
 	
 	FShaderCompilerEnvironment& ShaderEnvironment = NewJob->Input.Environment;
@@ -1312,8 +1312,8 @@ void FNiagaraShader::BindParams(const TArray<FNiagaraDataInterfaceGPUParamInfo>&
 
 	NumSpawnedInstancesParam.Bind(ParameterMap, TEXT("SpawnedInstances"));
 	UpdateStartInstanceParam.Bind(ParameterMap, TEXT("UpdateStartInstance"));
-	DefaultShaderStageIndexParam.Bind(ParameterMap, TEXT("DefaultShaderStageIndex"));
-	ShaderStageIndexParam.Bind(ParameterMap, TEXT("ShaderStageIndex"));
+	DefaultSimulationStageIndexParam.Bind(ParameterMap, TEXT("DefaultSimulationStageIndex"));
+	SimulationStageIndexParam.Bind(ParameterMap, TEXT("SimulationStageIndex"));
 	IterationInterfaceCount.Bind(ParameterMap, TEXT("IterationInterfaceCount"));
 
 	ComponentBufferSizeReadParam.Bind(ParameterMap, TEXT("ComponentBufferSizeRead"));
diff --git a/Engine/Plugins/FX/Niagara/Source/NiagaraShader/Private/NiagaraShared.cpp b/Engine/Plugins/FX/Niagara/Source/NiagaraShader/Private/NiagaraShared.cpp
--- a/Engine/Plugins/FX/Niagara/Source/NiagaraShader/Private/NiagaraShared.cpp
+++ b/Engine/Plugins/FX/Niagara/Source/NiagaraShader/Private/NiagaraShared.cpp
@@ -53,6 +53,23 @@ NIAGARASHADER_API bool FNiagaraShaderScript::ShouldCache(EShaderPlatform Platfor
 	return true;
 }
 
+NIAGARASHADER_API uint32 FNiagaraShaderScript::GetUseSimStagesDefine() const
+{
+	if (AdditionalDefines.Contains(TEXT("Emitter.UseSimulationStages")))
+	{
+		return 1;
+	}
+	else if (AdditionalDefines.Contains(TEXT("Emitter.UseOldShaderStages")))
+	{
+		return 1;
+	}
+	else
+	{
+		return 0;
+	}
+}
+
+
 NIAGARASHADER_API void FNiagaraShaderScript::NotifyCompilationFinished()
 {
 	OnCompilationCompleteDelegate.Broadcast();
diff --git a/Engine/Plugins/FX/Niagara/Source/NiagaraShader/Public/NiagaraShader.h b/Engine/Plugins/FX/Niagara/Source/NiagaraShader/Public/NiagaraShader.h
--- a/Engine/Plugins/FX/Niagara/Source/NiagaraShader/Public/NiagaraShader.h
+++ b/Engine/Plugins/FX/Niagara/Source/NiagaraShader/Public/NiagaraShader.h
@@ -116,8 +116,8 @@ public:
 	LAYOUT_FIELD(FShaderParameter, CopyInstancesBeforeStartParam);
 	LAYOUT_FIELD(FShaderParameter, NumSpawnedInstancesParam);
 	LAYOUT_FIELD(FShaderParameter, UpdateStartInstanceParam);
-	LAYOUT_FIELD(FShaderParameter, DefaultShaderStageIndexParam);
-	LAYOUT_FIELD(FShaderParameter, ShaderStageIndexParam);
+	LAYOUT_FIELD(FShaderParameter, DefaultSimulationStageIndexParam);
+	LAYOUT_FIELD(FShaderParameter, SimulationStageIndexParam);
 	LAYOUT_FIELD(FShaderParameter, IterationInterfaceCount);
 	LAYOUT_FIELD(FShaderParameter, ComponentBufferSizeReadParam);
 	LAYOUT_FIELD(FShaderParameter, ComponentBufferSizeWriteParam);
diff --git a/Engine/Plugins/FX/Niagara/Source/NiagaraShader/Public/NiagaraShared.h b/Engine/Plugins/FX/Niagara/Source/NiagaraShader/Public/NiagaraShared.h
--- a/Engine/Plugins/FX/Niagara/Source/NiagaraShader/Public/NiagaraShared.h
+++ b/Engine/Plugins/FX/Niagara/Source/NiagaraShader/Public/NiagaraShared.h
@@ -587,6 +587,8 @@ public:
 	NIAGARASHADER_API  bool CacheShaders(EShaderPlatform Platform, bool bApplyCompletedShaderMapForRendering, bool bForceRecompile, bool bSynchronous = false);
 	bool CacheShaders(const FNiagaraShaderMapId& ShaderMapId, EShaderPlatform Platform, bool bApplyCompletedShaderMapForRendering, bool bForceRecompile, bool bSynchronous = false);
 
+
+	NIAGARASHADER_API uint32 GetUseSimStagesDefine() const;
 	/**
 	 * Should the shader for this script with the given platform, shader type and vertex
 	 * factory type combination be compiled
