Fixing strict compile error

diff --git a/Engine/Plugins/FX/Niagara/Source/Niagara/Private/NiagaraScriptExecutionContext.cpp b/Engine/Plugins/FX/Niagara/Source/Niagara/Private/NiagaraScriptExecutionContext.cpp
--- a/Engine/Plugins/FX/Niagara/Source/Niagara/Private/NiagaraScriptExecutionContext.cpp
+++ b/Engine/Plugins/FX/Niagara/Source/Niagara/Private/NiagaraScriptExecutionContext.cpp
@@ -860,5 +860,5 @@ FNiagaraDataInterfaceProxy* FNiagaraComputeInstanceData::FindIterationInterface(
 	{
 		return Context->FindIterationInterface(DataInterfaceProxies, SimulationStageIndex);
 	}
-	return false;
+	return nullptr;
 }
\ No newline at end of file
