Making shader compile error if using sim stages with CPU
Making shader compile error if using events with GPU
Fixing incorrect shader compile error on sim target selector.

diff --git a/Engine/Plugins/FX/Niagara/Source/Niagara/Private/NiagaraScript.cpp b/Engine/Plugins/FX/Niagara/Source/Niagara/Private/NiagaraScript.cpp
--- a/Engine/Plugins/FX/Niagara/Source/Niagara/Private/NiagaraScript.cpp
+++ b/Engine/Plugins/FX/Niagara/Source/Niagara/Private/NiagaraScript.cpp
@@ -349,7 +349,7 @@ void UNiagaraScript::ComputeVMCompilationId(FNiagaraVMExecutableDataId& Id) cons
 			Id.bUsesRapidIterationParams = true;
 		}
 
-		if (Emitter->bSimulationStagesEnabled && Usage == ENiagaraScriptUsage::ParticleGPUComputeScript/* TODO limit to just with stages in the future! Leaving like this so what can convert! && Emitter->GetSimulationStages().Num() > 0*/)
+		if (Emitter->bSimulationStagesEnabled)
 		{
 			Id.AdditionalDefines.Add(TEXT("Emitter.UseSimulationStages"));
 
@@ -374,7 +374,7 @@ void UNiagaraScript::ComputeVMCompilationId(FNiagaraVMExecutableDataId& Id) cons
 			Id.DebugReferencedObjects.Add(TEXT("SimulationStageHeaders"));
 
 		}
-		else if (Emitter->bDeprecatedShaderStagesEnabled && Usage == ENiagaraScriptUsage::ParticleGPUComputeScript/* TODO limit to just with stages in the future! Leaving like this so what can convert! && Emitter->GetSimulationStages().Num() > 0*/)
+		else if (Emitter->bDeprecatedShaderStagesEnabled)
 		{
 			Id.AdditionalDefines.Add(TEXT("Emitter.UseOldShaderStages"));
 		}
diff --git a/Engine/Plugins/FX/Niagara/Source/NiagaraCore/Private/NiagaraCustomVersion.cpp b/Engine/Plugins/FX/Niagara/Source/NiagaraCore/Private/NiagaraCustomVersion.cpp
--- a/Engine/Plugins/FX/Niagara/Source/NiagaraCore/Private/NiagaraCustomVersion.cpp
+++ b/Engine/Plugins/FX/Niagara/Source/NiagaraCore/Private/NiagaraCustomVersion.cpp
@@ -8,4 +8,4 @@ const FGuid FNiagaraCustomVersion::GUID(0xFCF57AFA, 0x50764283, 0xB9A9E658, 0xFF
 // Register the custom version with core
 FCustomVersionRegistration GRegisterNiagaraCustomVersion(FNiagaraCustomVersion::GUID, FNiagaraCustomVersion::LatestVersion, TEXT("NiagaraVer"));
 
-const FGuid FNiagaraCustomVersion::LatestScriptCompileVersion(0x751E8859, 0x84474334, 0x9CC3D9A1, 0xB554796E);
+const FGuid FNiagaraCustomVersion::LatestScriptCompileVersion(0x410E4624, 0xFD724393, 0x87D800BE, 0x2B8217CE);
diff --git a/Engine/Plugins/FX/Niagara/Source/NiagaraEditor/Private/NiagaraHlslTranslator.cpp b/Engine/Plugins/FX/Niagara/Source/NiagaraEditor/Private/NiagaraHlslTranslator.cpp
--- a/Engine/Plugins/FX/Niagara/Source/NiagaraEditor/Private/NiagaraHlslTranslator.cpp
+++ b/Engine/Plugins/FX/Niagara/Source/NiagaraEditor/Private/NiagaraHlslTranslator.cpp
@@ -846,6 +846,23 @@ const FNiagaraTranslateResults &FHlslNiagaraTranslator::Translate(const FNiagara
 
 	OtherOutputParamMapHistories = CompileData->GetPrecomputedHistories();
 
+	bool bCPUSim = IsCompileOptionDefined(TEXT("CPUSim"));
+	bool bGPUSim = IsCompileOptionDefined(TEXT("GPUComputeSim"));
+
+	if (bUsesSimStages && bCPUSim)
+	{
+		Error(LOCTEXT("CannotUseSimStagesWithCPU", "Cannot use CPU simulations with Experimental Simulation Stages or Deprecated Shader Stages!"), nullptr, nullptr);
+		return TranslateResults;
+	}
+
+
+	if (CompileOptions.TargetUsage == ENiagaraScriptUsage::ParticleEventScript && bGPUSim)
+	{
+		Error(LOCTEXT("CannotUseEventsWithGPU", "GPU Events scripts are currently unsupported. Consider using DirectReads instead!"), nullptr, nullptr);
+		return TranslateResults;
+	}
+
+
 	switch (CompileOptions.TargetUsage)
 	{
 	case ENiagaraScriptUsage::ParticleSpawnScriptInterpolated:
@@ -985,7 +1002,7 @@ const FNiagaraTranslateResults &FHlslNiagaraTranslator::Translate(const FNiagara
 			Error(FText::Format(LOCTEXT("GetOutputNodeFail", "Cannot find output node of type {0}!"), FText::AsNumber((int32)TranslationStages[i].ScriptUsage)), nullptr, nullptr);
 			return TranslateResults;
 		}
-
+		
 		ValidateTypePins(TranslationStages[i].OutputNode);
 		{
 			bool bHasAnyConnections = false;
diff --git a/Engine/Plugins/FX/Niagara/Source/NiagaraEditor/Private/NiagaraNodeSimTargetSelector.cpp b/Engine/Plugins/FX/Niagara/Source/NiagaraEditor/Private/NiagaraNodeSimTargetSelector.cpp
--- a/Engine/Plugins/FX/Niagara/Source/NiagaraEditor/Private/NiagaraNodeSimTargetSelector.cpp
+++ b/Engine/Plugins/FX/Niagara/Source/NiagaraEditor/Private/NiagaraNodeSimTargetSelector.cpp
@@ -73,13 +73,17 @@ void UNiagaraNodeSimTargetSelector::Compile(class FHlslNiagaraTranslator* Transl
 	TArray<UEdGraphPin*> OutputPins;
 	GetOutputPins(OutputPins);
 
-	ENiagaraSimTarget SimulationTarget = Translator->GetSimulationTarget();
+	//ENiagaraSimTarget SimulationTarget = Translator->GetSimulationTarget();
+	bool bCPUSim = Translator->IsCompileOptionDefined(TEXT("CPUSim"));
+	bool bGPUSim = Translator->IsCompileOptionDefined(TEXT("GPUComputeSim"));
+
+
 	int32 VarIdx;
-	if (SimulationTarget == ENiagaraSimTarget::CPUSim)
+	if (bCPUSim/*SimulationTarget == ENiagaraSimTarget::CPUSim*/)
 	{
 		VarIdx = 0;
 	}
-	else if (SimulationTarget == ENiagaraSimTarget::GPUComputeSim)
+	else if (bGPUSim/*SimulationTarget == ENiagaraSimTarget::GPUComputeSim*/)
 	{
 		VarIdx = InputPins.Num() / 2;
 	}
