Pushing licensee fix

Add OnEnginePreExit() delegate to core delegates. This delegate is fired off before closing out the AssetTools, WorldBrowser and AssetRegistry modules, for handling shutdown tasks that would require these.

Handle lingering Niagara shader compiles on engine exit inside OnEnginePreExit() callback to avoid asserting if shader compile would require querying the asset registry module after it has shutdown.

diff --git a/Engine/Plugins/FX/Niagara/Source/NiagaraEditor/Private/NiagaraEditorModule.cpp b/Engine/Plugins/FX/Niagara/Source/NiagaraEditor/Private/NiagaraEditorModule.cpp
--- a/Engine/Plugins/FX/Niagara/Source/NiagaraEditor/Private/NiagaraEditorModule.cpp
+++ b/Engine/Plugins/FX/Niagara/Source/NiagaraEditor/Private/NiagaraEditorModule.cpp
@@ -719,6 +719,16 @@ class FNiagaraSystemColorParameterTrackEditor : public FNiagaraSystemParameterTr
 void FNiagaraEditorModule::OnPreExit()
 {
 	UDeviceProfileManager::Get().OnManagerUpdated().Remove(DeviceProfileManagerUpdatedHandle);
+
+	// Ensure that we don't have any lingering compiles laying around that will explode after this module shuts down.
+	for (TObjectIterator<UNiagaraSystem> It; It; ++It)
+	{
+		UNiagaraSystem* Sys = *It;
+		if (Sys)
+		{
+			Sys->WaitForCompilationComplete();
+		}
+	}
 }
 
 void FNiagaraEditorModule::StartupModule()
@@ -747,7 +757,7 @@ void FNiagaraEditorModule::StartupModule()
 	FCoreDelegates::OnPostEngineInit.AddRaw(this, &FNiagaraEditorModule::OnPostEngineInit);
 
 	DeviceProfileManagerUpdatedHandle = UDeviceProfileManager::Get().OnManagerUpdated().AddRaw(this, &FNiagaraEditorModule::OnDeviceProfileManagerUpdated);
-	FCoreDelegates::OnPreExit.AddRaw(this, &FNiagaraEditorModule::OnPreExit);
+	FCoreDelegates::OnEnginePreExit.AddRaw(this, &FNiagaraEditorModule::OnPreExit);
 	
 	// register details customization
 	FPropertyEditorModule& PropertyModule = FModuleManager::LoadModuleChecked<FPropertyEditorModule>("PropertyEditor");
@@ -1003,16 +1013,6 @@ void FNiagaraEditorModule::StartupModule()
 
 void FNiagaraEditorModule::ShutdownModule()
 {
-	// Ensure that we don't have any lingering compiles laying around that will explode after this module shuts down.
-	for (TObjectIterator<UNiagaraSystem> It; It; ++It)
-	{
-		UNiagaraSystem* Sys = *It;
-		if (Sys)
-		{
-			Sys->WaitForCompilationComplete();
-		}
-	}
-
 	MenuExtensibilityManager.Reset();
 	ToolBarExtensibilityManager.Reset();
 	
@@ -1030,7 +1030,7 @@ void FNiagaraEditorModule::ShutdownModule()
 
 	FCoreUObjectDelegates::GetPreGarbageCollectDelegate().RemoveAll(this);
 	FCoreDelegates::OnPostEngineInit.RemoveAll(this);
-	FCoreDelegates::OnPreExit.RemoveAll(this);
+	FCoreDelegates::OnEnginePreExit.RemoveAll(this);
 	
 	if (GEditor)
 	{
diff --git a/Engine/Source/Runtime/Core/Private/Misc/CoreDelegates.cpp b/Engine/Source/Runtime/Core/Private/Misc/CoreDelegates.cpp
--- a/Engine/Source/Runtime/Core/Private/Misc/CoreDelegates.cpp
+++ b/Engine/Source/Runtime/Core/Private/Misc/CoreDelegates.cpp
@@ -85,6 +85,7 @@ FSimpleMulticastDelegate FCoreDelegates::OnPostEngineInit;
 FSimpleMulticastDelegate FCoreDelegates::OnFEngineLoopInitComplete;
 FSimpleMulticastDelegate FCoreDelegates::OnExit;
 FSimpleMulticastDelegate FCoreDelegates::OnPreExit;
+FSimpleMulticastDelegate FCoreDelegates::OnEnginePreExit;
 FCoreDelegates::FGatherAdditionalLocResPathsDelegate FCoreDelegates::GatherAdditionalLocResPathsCallback;
 FSimpleMulticastDelegate FCoreDelegates::ColorPickerChanged;
 FSimpleMulticastDelegate FCoreDelegates::OnBeginFrame;
diff --git a/Engine/Source/Runtime/Core/Public/Misc/CoreDelegates.h b/Engine/Source/Runtime/Core/Public/Misc/CoreDelegates.h
--- a/Engine/Source/Runtime/Core/Public/Misc/CoreDelegates.h
+++ b/Engine/Source/Runtime/Core/Public/Misc/CoreDelegates.h
@@ -248,6 +248,9 @@ public:
 	// Called when before the application is exiting.
 	static FSimpleMulticastDelegate OnPreExit;
 
+	// Called before the engine exits. Separate from OnPreExit as OnEnginePreExit occurs before shutting down any core modules.
+	static FSimpleMulticastDelegate OnEnginePreExit;
+
 	/** Delegate for gathering up additional localization paths that are unknown to the UE4 core (such as plugins) */
 	DECLARE_MULTICAST_DELEGATE_OneParam(FGatherAdditionalLocResPathsDelegate, TArray<FString>&);
 	static FGatherAdditionalLocResPathsDelegate GatherAdditionalLocResPathsCallback;
diff --git a/Engine/Source/Runtime/Engine/Private/UnrealEngine.cpp b/Engine/Source/Runtime/Engine/Private/UnrealEngine.cpp
--- a/Engine/Source/Runtime/Engine/Private/UnrealEngine.cpp
+++ b/Engine/Source/Runtime/Engine/Private/UnrealEngine.cpp
@@ -1775,6 +1775,8 @@ void UEngine::ShutdownAudioDeviceManager()
 
 void UEngine::PreExit()
 {
+	FCoreDelegates::OnEnginePreExit.Broadcast();
+
 #if UE_NET_TRACE_ENABLED
 	FNetTrace::SetTraceVerbosity(0);
 #endif
