Fixed crash when trying to read from an emitter which has no current data.

This can't happen under normal circumstances because the source emitter is executed before the emitter which contains the DI. However, if a CPU emitter tries to read from a GPU emitter, they are not ordered, and the CPU emitter may run before the GPU emitter is initialized. This scenario is not supported, but it shouldn't crash.

diff --git a/Engine/Plugins/FX/Niagara/Source/Niagara/Classes/NiagaraDataSet.h b/Engine/Plugins/FX/Niagara/Source/Niagara/Classes/NiagaraDataSet.h
--- a/Engine/Plugins/FX/Niagara/Source/Niagara/Classes/NiagaraDataSet.h
+++ b/Engine/Plugins/FX/Niagara/Source/Niagara/Classes/NiagaraDataSet.h
@@ -157,6 +157,7 @@ public:
 	bool CheckForNaNs()const;
 
 	FORCEINLINE TArray<int32>& GetIDTable() { return IDToIndexTable; }
+	FORCEINLINE const TArray<int32>& GetIDTable() const { return IDToIndexTable; }
 
 	void SetShaderParams(class FNiagaraShader *Shader, FRHICommandList &CommandList, bool bInput);
 	void UnsetShaderParams(class FNiagaraShader *Shader, FRHICommandList &CommandList);
diff --git a/Engine/Plugins/FX/Niagara/Source/Niagara/Private/NiagaraDataInterfaceParticleRead.cpp b/Engine/Plugins/FX/Niagara/Source/Niagara/Private/NiagaraDataInterfaceParticleRead.cpp
--- a/Engine/Plugins/FX/Niagara/Source/Niagara/Private/NiagaraDataInterfaceParticleRead.cpp
+++ b/Engine/Plugins/FX/Niagara/Source/Niagara/Private/NiagaraDataInterfaceParticleRead.cpp
@@ -820,7 +820,8 @@ void UNiagaraDataInterfaceParticleRead::GetNumSpawnedParticles(FVectorVMContext&
 	VectorVM::FExternalFuncRegisterHandler<int32> OutNumSpawned(Context);
 
 	const FNiagaraEmitterInstance* EmitterInstance = InstData.Get()->EmitterInstance;
-	const int32 NumSpawned = EmitterInstance->GetData().GetCurrentData()->GetNumSpawnedInstances();
+	const FNiagaraDataBuffer* CurrentData = EmitterInstance->GetData().GetCurrentData();
+	const int32 NumSpawned = CurrentData ? CurrentData->GetNumSpawnedInstances() : 0;
 
 	for (int32 InstanceIdx = 0; InstanceIdx < Context.NumInstances; ++InstanceIdx)
 	{
@@ -874,7 +875,8 @@ void UNiagaraDataInterfaceParticleRead::GetNumParticles(FVectorVMContext& Contex
 	VectorVM::FExternalFuncRegisterHandler<int32> OutNumParticles (Context);
 
 	const FNiagaraEmitterInstance* EmitterInstance = InstData.Get()->EmitterInstance;
-	const int32 NumParticles = EmitterInstance->GetData().GetCurrentData()->GetNumInstances();
+	const FNiagaraDataBuffer* CurrentData = EmitterInstance->GetData().GetCurrentData();
+	const int32 NumParticles = CurrentData ? CurrentData->GetNumInstances() : 0;
 
 	for (int32 InstanceIdx = 0; InstanceIdx < Context.NumInstances; ++InstanceIdx)
 	{
@@ -1099,36 +1101,39 @@ void UNiagaraDataInterfaceParticleRead::ReadQuat(FVectorVMContext& Context, FNam
 template <typename T>
 T UNiagaraDataInterfaceParticleRead::RetrieveValueWithCheck(FNiagaraEmitterInstance* EmitterInstance, const FNiagaraTypeDefinition& Type, const FName& Attr, const FNiagaraID& ParticleID, bool& bValid)
 {
-	TArray<int32>& IDTable = EmitterInstance->GetData().GetCurrentData()->GetIDTable();
-	if (ParticleID.Index < 0 || ParticleID.Index >= IDTable.Num())
+	T Value = T();
+	bValid = false;
+
+	const FNiagaraDataBuffer* CurrentData = EmitterInstance->GetData().GetCurrentData();
+	if (!CurrentData)
 	{
-		bValid = false;
-		return T();
+		return Value;
 	}
-	else
+
+	const TArray<int32>& IDTable = CurrentData->GetIDTable();
+	if (ParticleID.Index < 0 || ParticleID.Index >= IDTable.Num())
 	{
-		FNiagaraVariable ReadVar(Type, Attr);
-		FNiagaraDataSetAccessor<T> ValueData(EmitterInstance->GetData(), ReadVar);
+		return Value;
+	}
 
-		FNiagaraVariable IDVar(FNiagaraTypeDefinition::GetIDDef(), "ID");
-		FNiagaraDataSetAccessor<FNiagaraID> IDData(EmitterInstance->GetData(), IDVar);
+	FNiagaraVariable ReadVar(Type, Attr);
+	FNiagaraDataSetAccessor<T> ValueData(EmitterInstance->GetData(), ReadVar);
 
-		int32 ParticleIndex = IDTable[ParticleID.Index];
-		T Value = T();
-		bValid = false;
+	FNiagaraVariable IDVar(FNiagaraTypeDefinition::GetIDDef(), "ID");
+	FNiagaraDataSetAccessor<FNiagaraID> IDData(EmitterInstance->GetData(), IDVar);
 
-		if (ParticleIndex >= 0)
+	int32 ParticleIndex = IDTable[ParticleID.Index];
+	if (ParticleIndex >= 0)
+	{
+		FNiagaraID ActualID = IDData.GetSafe(ParticleIndex, NIAGARA_INVALID_ID);
+		if (ActualID == ParticleID)
 		{
-			FNiagaraID ActualID = IDData.GetSafe(ParticleIndex, NIAGARA_INVALID_ID);
-			if (ActualID == ParticleID)
-			{
-				Value = ValueData.GetSafe(ParticleIndex, T());
-				bValid = true;
-			}
+			Value = ValueData.GetSafe(ParticleIndex, T());
+			bValid = true;
 		}
-
-		return Value;
 	}
+
+	return Value;
 }
 
 bool UNiagaraDataInterfaceParticleRead::Equals(const UNiagaraDataInterface* Other) const
