Fixed issue where repeated function calls fail to register write functions for DI's

diff --git a/Engine/Plugins/FX/Niagara/Source/NiagaraCore/Private/NiagaraCustomVersion.cpp b/Engine/Plugins/FX/Niagara/Source/NiagaraCore/Private/NiagaraCustomVersion.cpp
--- a/Engine/Plugins/FX/Niagara/Source/NiagaraCore/Private/NiagaraCustomVersion.cpp
+++ b/Engine/Plugins/FX/Niagara/Source/NiagaraCore/Private/NiagaraCustomVersion.cpp
@@ -8,4 +8,4 @@ const FGuid FNiagaraCustomVersion::GUID(0xFCF57AFA, 0x50764283, 0xB9A9E658, 0xFF
 // Register the custom version with core
 FCustomVersionRegistration GRegisterNiagaraCustomVersion(FNiagaraCustomVersion::GUID, FNiagaraCustomVersion::LatestVersion, TEXT("NiagaraVer"));
 
-const FGuid FNiagaraCustomVersion::LatestScriptCompileVersion(0x410E4624, 0xFD724393, 0x87D800BE, 0x2B8217CE);
+const FGuid FNiagaraCustomVersion::LatestScriptCompileVersion(0xDB33E371, 0xCE1C40FD, 0xBA714031, 0xE8DACCB4);
diff --git a/Engine/Plugins/FX/Niagara/Source/NiagaraEditor/Private/NiagaraHlslTranslator.cpp b/Engine/Plugins/FX/Niagara/Source/NiagaraEditor/Private/NiagaraHlslTranslator.cpp
--- a/Engine/Plugins/FX/Niagara/Source/NiagaraEditor/Private/NiagaraHlslTranslator.cpp
+++ b/Engine/Plugins/FX/Niagara/Source/NiagaraEditor/Private/NiagaraHlslTranslator.cpp
@@ -1507,6 +1507,25 @@ const FNiagaraTranslateResults &FHlslNiagaraTranslator::Translate(const FNiagara
 			}
 		}
 
+		// Log out all the information we've built thus far for assistance debugging
+		FString Preamble = TEXT("// Shader generated by Niagara HLSL Translator\n\n");
+		
+		for (int32 i = 0; i < CompilationOutput.ScriptData.SimulationStageMetaData.Num(); i++)
+		{
+			Preamble += FString::Printf(TEXT("// SimStage[%d]\n//\t Iteration Src: \"%s\"\n//\tMinIndex: %d  MaxIndex: %d\n//\tbSpawnOnly: %s\n//\tWritesParticles: %s\n"), i,
+				*CompilationOutput.ScriptData.SimulationStageMetaData[i].IterationSource.ToString(),
+				CompilationOutput.ScriptData.SimulationStageMetaData[i].MinStage,
+				CompilationOutput.ScriptData.SimulationStageMetaData[i].MaxStage,
+				CompilationOutput.ScriptData.SimulationStageMetaData[i].bSpawnOnly ? TEXT("True") : TEXT("False"),
+				CompilationOutput.ScriptData.SimulationStageMetaData[i].bWritesParticles ? TEXT("True") : TEXT("False"));
+			for (const FName& Dest : CompilationOutput.ScriptData.SimulationStageMetaData[i].OutputDestinations)
+			{
+				Preamble += FString::Printf(TEXT("//\tOutputs to: \"%s\"\n"), *Dest.ToString());
+			}
+		}
+
+		HlslOutput = Preamble + TEXT("\n\n") +  HlslOutput;
+
 		// We may have created some transient data interfaces. This cleans up the ones that we created.
 		CompilationOutput.ScriptData.DIParamInfo = DIParamInfo;
 		if (InstanceRead.Variables.Num() == 1 && InstanceRead.Variables[0].GetName() == TEXT("Particles.UniqueID")) {
@@ -2993,6 +3012,8 @@ FString FHlslNiagaraTranslator::GetUniqueSymbolName(FName BaseName)
 void FHlslNiagaraTranslator::EnterFunction(const FString& Name, FNiagaraFunctionSignature& Signature, TArray<int32>& Inputs, const FGuid& InGuid)
 {
 	FunctionContextStack.Emplace(Name, Signature, Inputs, InGuid);
+	TArray<FName> Entries;
+	ActiveStageWriteTargets.Push(Entries);
 	//May need some more heavy and scoped symbol tracking?
 
 	//Add new scope for pin reuse.
@@ -3006,6 +3027,16 @@ void FHlslNiagaraTranslator::ExitFunction()
 
 	//Pop pin reuse scope.
 	PinToCodeChunks.Pop();
+
+	// Accumulate the write targets.
+	TArray<FName> Entries = ActiveStageWriteTargets.Pop();
+	if (ActiveStageWriteTargets.Num())
+	{
+		for (const FName& Entry : Entries)
+		{
+			ActiveStageWriteTargets.Top().AddUnique(Entry);
+		}
+	}
 }
 
 FString FHlslNiagaraTranslator::GeneratedConstantString(float Constant)
@@ -5166,6 +5197,7 @@ void FHlslNiagaraTranslator::FunctionCall(UNiagaraNodeFunctionCall* FunctionNode
 		return;
 	}
 
+	//UE_LOG(LogNiagaraEditor, Log, TEXT("Function Call: %s %d"), *FunctionNode->GetFunctionName(), ActiveStageIdx);
 	// We need the generated string to generate the proper signature for now.
 	ActiveHistoryForFunctionCalls.EnterFunction(FunctionNode->GetFunctionName(), FunctionNode->FunctionScript, FunctionNode);
 
@@ -5545,11 +5577,15 @@ void FHlslNiagaraTranslator::HandleDataInterfaceCall(FNiagaraScriptDataInterface
 		Error(FText::Format(LOCTEXT("FunctionCallDataInterfaceGPUMissing", "Function call \"{0}\" does not work on GPU sims."), FText::FromName(InMatchingSignature.Name)), CurNode, nullptr);
 	}
 
+	//UE_LOG(LogNiagaraEditor, Log, TEXT("HandleDataInterfaceCall %d %s %s %s"), ActiveStageIdx, *InMatchingSignature.Name.ToString(), InMatchingSignature.bWriteFunction ? TEXT("true") : TEXT("False"), *Info.Name.ToString());
+
 	if (InMatchingSignature.bWriteFunction && CompilationOutput.ScriptData.SimulationStageMetaData.Num() != 0 && TranslationStages[ActiveStageIdx].SourceSimStage != -1)
 	{
 		int32 SourceSimStage = TranslationStages[ActiveStageIdx].SourceSimStage;
 		ensure(CompilationOutput.ScriptData.SimulationStageMetaData.Num() > SourceSimStage);
 		CompilationOutput.ScriptData.SimulationStageMetaData[SourceSimStage].OutputDestinations.AddUnique(Info.Name);
+		ActiveStageWriteTargets.Top().AddUnique(Info.Name);
+		//UE_LOG(LogNiagaraEditor, Log, TEXT(" Add to %d %d"), ActiveStageIdx, SourceSimStage);
 	}
 }
 
@@ -5832,10 +5868,28 @@ void FHlslNiagaraTranslator::RegisterFunctionCall(ENiagaraScriptUsage ScriptUsag
 				}
 
 				Functions.Add(OutSignature, FunctionDefStr);
+				FunctionStageWriteTargets.Add(OutSignature, ActiveStageWriteTargets.Top());
 			}
 
 			ExitFunction();
 		}
+		else
+		{
+			// Just because we had a cached call, doesn't mean that we should ignore adding an writetargets
+			TArray<FName>* Entries = FunctionStageWriteTargets.Find(OutSignature);
+			if (Entries)
+			{
+				for (const FName& Entry : *Entries)
+				{
+					int32 SourceSimStage = TranslationStages[ActiveStageIdx].SourceSimStage;
+					ensure(CompilationOutput.ScriptData.SimulationStageMetaData.Num() > SourceSimStage);
+					CompilationOutput.ScriptData.SimulationStageMetaData[SourceSimStage].OutputDestinations.AddUnique(Entry);
+					ActiveStageWriteTargets.Top().AddUnique(Entry);
+					//UE_LOG(LogNiagaraEditor, Log, TEXT(" Add to %d %d"), ActiveStageIdx, SourceSimStage);
+
+				}
+			}
+		}
 	}
 	else
 	{
@@ -5863,6 +5917,7 @@ void FHlslNiagaraTranslator::RegisterFunctionCall(ENiagaraScriptUsage ScriptUsag
 				}
 
 				Functions.Add(OutSignature, FunctionDefStr);
+				FunctionStageWriteTargets.Add(OutSignature, ActiveStageWriteTargets.Top());
 
 				ExitFunction();
 			}
diff --git a/Engine/Plugins/FX/Niagara/Source/NiagaraEditor/Private/NiagaraHlslTranslator.h b/Engine/Plugins/FX/Niagara/Source/NiagaraEditor/Private/NiagaraHlslTranslator.h
--- a/Engine/Plugins/FX/Niagara/Source/NiagaraEditor/Private/NiagaraHlslTranslator.h
+++ b/Engine/Plugins/FX/Niagara/Source/NiagaraEditor/Private/NiagaraHlslTranslator.h
@@ -381,6 +381,9 @@ protected:
 
 	/** All functions called in the script. */
 	TMap<FNiagaraFunctionSignature, FString> Functions;
+	TMap<FNiagaraFunctionSignature, TArray<FName> > FunctionStageWriteTargets;
+	TArray<TArray<FName>> ActiveStageWriteTargets;
+
 	/** Map of function graphs we've seen before and already pre-processed. */
 	TMap<const UNiagaraGraph*, UNiagaraGraph*> PreprocessedFunctions;
 
