Adaptive unity build fixes.

diff --git a/Engine/Plugins/FX/Niagara/Source/NiagaraEditor/Private/TypeEditorUtilities/NiagaraEnumTypeEditorUtilities.cpp b/Engine/Plugins/FX/Niagara/Source/NiagaraEditor/Private/TypeEditorUtilities/NiagaraEnumTypeEditorUtilities.cpp
--- a/Engine/Plugins/FX/Niagara/Source/NiagaraEditor/Private/TypeEditorUtilities/NiagaraEnumTypeEditorUtilities.cpp
+++ b/Engine/Plugins/FX/Niagara/Source/NiagaraEditor/Private/TypeEditorUtilities/NiagaraEnumTypeEditorUtilities.cpp
@@ -8,6 +8,7 @@
 
 #include "Widgets/Input/SSpinBox.h"
 #include "EditorWidgets/Public/SEnumComboBox.h"
+#include "EditorStyleSet.h"
 
 class SNiagaraEnumParameterEditor : public SNiagaraParameterEditor
 {
