Fix path - SEnumComboBox.h -> SEnumCombobox.h

#rnx

diff --git a/Engine/Plugins/FX/Niagara/Source/NiagaraEditor/Private/TypeEditorUtilities/NiagaraEnumTypeEditorUtilities.cpp b/Engine/Plugins/FX/Niagara/Source/NiagaraEditor/Private/TypeEditorUtilities/NiagaraEnumTypeEditorUtilities.cpp
--- a/Engine/Plugins/FX/Niagara/Source/NiagaraEditor/Private/TypeEditorUtilities/NiagaraEnumTypeEditorUtilities.cpp
+++ b/Engine/Plugins/FX/Niagara/Source/NiagaraEditor/Private/TypeEditorUtilities/NiagaraEnumTypeEditorUtilities.cpp
@@ -7,7 +7,7 @@
 #include "NiagaraEditorCommon.h"
 
 #include "Widgets/Input/SSpinBox.h"
-#include "EditorWidgets/Public/SEnumComboBox.h"
+#include "EditorWidgets/Public/SEnumCombobox.h"
 #include "EditorStyleSet.h"
 
 class SNiagaraEnumParameterEditor : public SNiagaraParameterEditor
