Fixing extraneous logging

diff --git a/Engine/Plugins/FX/Niagara/Source/Niagara/Private/NiagaraGPUInstanceCountManager.cpp b/Engine/Plugins/FX/Niagara/Source/Niagara/Private/NiagaraGPUInstanceCountManager.cpp
--- a/Engine/Plugins/FX/Niagara/Source/Niagara/Private/NiagaraGPUInstanceCountManager.cpp
+++ b/Engine/Plugins/FX/Niagara/Source/Niagara/Private/NiagaraGPUInstanceCountManager.cpp
@@ -75,7 +75,7 @@ uint32 FNiagaraGPUInstanceCountManager::AcquireEntry()
 	{
 		// @TODO : add realloc the buffer and copy the current content to it. Might require reallocating the readback in FNiagaraGPUInstanceCountManager::EnqueueGPUReadback()
 		ensure(UsedInstanceCounts < AllocatedInstanceCounts);
-		UE_LOG(LogNiagara, Error, TEXT("Niagara.MinGPUInstanceCount too small. UsedInstanceCounts: %d < AllocatedInstanceCounts: %d"), UsedInstanceCounts, AllocatedInstanceCounts);
+		//UE_LOG(LogNiagara, Error, TEXT("Niagara.MinGPUInstanceCount too small. UsedInstanceCounts: %d < AllocatedInstanceCounts: %d"), UsedInstanceCounts, AllocatedInstanceCounts);
 		return INDEX_NONE;
 	}
 }
@@ -112,7 +112,7 @@ void FNiagaraGPUInstanceCountManager::ResizeBuffers(FRHICommandListImmediate& RH
 			TResourceArray<uint32> InitData;
 			InitData.AddZeroed(AllocatedInstanceCounts);
 			CountBuffer.Initialize(sizeof(uint32), AllocatedInstanceCounts, EPixelFormat::PF_R32_UINT, BUF_Static | BUF_SourceCopy, TEXT("NiagaraGPUInstanceCounts"), &InitData);
-			UE_LOG(LogNiagara, Log, TEXT("FNiagaraGPUInstanceCountManager::ResizeBuffers Alloc AllocatedInstanceCounts: %d ReservedInstanceCounts: %d"), AllocatedInstanceCounts, ReservedInstanceCounts);
+			//UE_LOG(LogNiagara, Log, TEXT("FNiagaraGPUInstanceCountManager::ResizeBuffers Alloc AllocatedInstanceCounts: %d ReservedInstanceCounts: %d"), AllocatedInstanceCounts, ReservedInstanceCounts);
 		}
 		// If we need to increase the buffer size to RecommendedInstanceCounts because the buffer is too small.
 		else if (RequiredInstanceCounts > AllocatedInstanceCounts)
@@ -140,7 +140,7 @@ void FNiagaraGPUInstanceCountManager::ResizeBuffers(FRHICommandListImmediate& RH
 			// Swap the buffers
 			AllocatedInstanceCounts = RecommendedInstanceCounts;
 			FMemory::Memswap(&NextCountBuffer, &CountBuffer, sizeof(NextCountBuffer));
-			UE_LOG(LogNiagara, Log, TEXT("FNiagaraGPUInstanceCountManager::ResizeBuffers Resize AllocatedInstanceCounts: %d ReservedInstanceCounts: %d"), AllocatedInstanceCounts, ReservedInstanceCounts);
+			//UE_LOG(LogNiagara, Log, TEXT("FNiagaraGPUInstanceCountManager::ResizeBuffers Resize AllocatedInstanceCounts: %d ReservedInstanceCounts: %d"), AllocatedInstanceCounts, ReservedInstanceCounts);
 		}
 		// If we need to shrink the buffer size because use way to much buffer size.
 		else if ((int32)(RecommendedInstanceCounts * BufferSlack) < AllocatedInstanceCounts)
