Fixing mesh local space tests.

diff --git a/Engine/Plugins/FX/Niagara/Shaders/Private/NiagaraMeshVertexFactory.ush b/Engine/Plugins/FX/Niagara/Shaders/Private/NiagaraMeshVertexFactory.ush
--- a/Engine/Plugins/FX/Niagara/Shaders/Private/NiagaraMeshVertexFactory.ush
+++ b/Engine/Plugins/FX/Niagara/Shaders/Private/NiagaraMeshVertexFactory.ush
@@ -723,7 +723,7 @@ FVertexFactoryIntermediates GetVertexFactoryIntermediates(FVertexFactoryInput In
 	if (NiagaraMeshVF.bLocalSpace)
 	{
 		Intermediates.PrevLocalToWorld = mul(Intermediates.PrevLocalToWorld, Primitive.PreviousLocalToWorld);
-		//Intermediates.LocalToWorld = mul(Intermediates.LocalToWorld, Primitive.LocalToWorld);
+		Intermediates.LocalToWorld = mul(Intermediates.LocalToWorld, Primitive.LocalToWorld);
 		Intermediates.ParticlePosition = Intermediates.LocalToWorld[3];
 	}
 
