Fixing a crash with a single emitter active system where we reinitialize immediately on level start while the in-level component is already auto activated.

diff --git a/Engine/Plugins/FX/Niagara/Source/Niagara/Private/NiagaraComponent.cpp b/Engine/Plugins/FX/Niagara/Source/Niagara/Private/NiagaraComponent.cpp
--- a/Engine/Plugins/FX/Niagara/Source/Niagara/Private/NiagaraComponent.cpp
+++ b/Engine/Plugins/FX/Niagara/Source/Niagara/Private/NiagaraComponent.cpp
@@ -1147,6 +1147,11 @@ void UNiagaraComponent::DestroyInstance()
 	//UE_LOG(LogNiagara, Log, TEXT("DestroyInstance: %u - %s"), this, *Asset->GetName());
 	SetActiveFlag(false);
 
+	// Before we can destory the instance, we need to deactivate it.
+	if (SystemInstance)
+	{
+		SystemInstance->Deactivate(true);
+	}
 	UnregisterWithScalabilityManager();
 	
 	// Rather than setting the unique ptr to null here, we allow it to transition ownership to the system's deferred deletion queue. This allows us to safely
