Fixed unbound UAV parameters causing crashes on Vulkan and consoles.
* Added functions to the instance batcher to retrieve empty UAVs from a pool which gets reset after each dispatch.
* Changed the GetDummy*Buffer() functions in NiagaraRenderer to return an FRHIShaderResourceView* instead of an FRWBuffer. We never want to use these for UAVs.
* Fixed a validation error in FNiagaraDataInterfaceProxyNeighborGrid3D (ClearUAVUint used on SINT buffers).

Also merged CL 11535205 to fix some game/render thread races in UNiagaraDataInterfaceGrid2DCollection.

diff --git a/Engine/Plugins/FX/Niagara/Source/Niagara/Classes/NiagaraDataInterfaceGrid2DCollection.h b/Engine/Plugins/FX/Niagara/Source/Niagara/Classes/NiagaraDataInterfaceGrid2DCollection.h
--- a/Engine/Plugins/FX/Niagara/Source/Niagara/Classes/NiagaraDataInterfaceGrid2DCollection.h
+++ b/Engine/Plugins/FX/Niagara/Source/Niagara/Classes/NiagaraDataInterfaceGrid2DCollection.h
@@ -11,10 +11,10 @@
 class FNiagaraSystemInstance;
 class UTextureRenderTarget2D;
 
-class Grid2DBuffer
+class FGrid2DBuffer
 {
 public:
-	Grid2DBuffer(int NumX, int NumY)
+	FGrid2DBuffer(int NumX, int NumY)
 	{		
 		GridBuffer.Initialize(4, NumX, NumY, EPixelFormat::PF_R32_FLOAT);
 	}
@@ -22,70 +22,32 @@ public:
 	FTextureRWBuffer2D GridBuffer;	
 };
 
-class Grid2DCollectionRWInstanceData
+struct FGrid2DCollectionRWInstanceData_GameThread
 {
-public:
-
-	int NumCellsX;
-	int NumCellsY;
-	int NumAttributes;
-	int NumTilesX;
-	int NumTilesY;
-
-	FVector2D CellSize;
-	bool SetGridFromVoxelSize;
-	
-	FVector2D WorldBBoxSize;
-
-	TArray<Grid2DBuffer*> Buffers;
-
-	Grid2DBuffer* CurrentData;
-	Grid2DBuffer* DestinationData;
-
-	Grid2DCollectionRWInstanceData() : CurrentData(nullptr), DestinationData(nullptr) {}
+	FIntPoint NumCells = FIntPoint(EForceInit::ForceInitToZero);
+	FIntPoint NumTiles = FIntPoint(EForceInit::ForceInitToZero);
+	FVector2D CellSize = FVector2D::ZeroVector;
+	FVector2D WorldBBoxSize = FVector2D::ZeroVector;
+};
 
-	~Grid2DCollectionRWInstanceData()
-	{
-		for (Grid2DBuffer* Buffer : Buffers)
-		{
-			Buffer->GridBuffer.Release();
-			delete Buffer;
-		}
-	}
+struct FGrid2DCollectionRWInstanceData_RenderThread
+{
+	FIntPoint NumCells = FIntPoint(EForceInit::ForceInitToZero);
+	FIntPoint NumTiles = FIntPoint(EForceInit::ForceInitToZero);
+	FVector2D CellSize = FVector2D::ZeroVector;
+	FVector2D WorldBBoxSize = FVector2D::ZeroVector;
 
-	Grid2DBuffer* GetCurrentData() { return CurrentData; }
-	Grid2DBuffer* GetDestinationData() { return DestinationData; }
-	Grid2DBuffer &BeginSimulate()
-	{
-		for (Grid2DBuffer* Buffer : Buffers)
-		{
-			check(Buffer);
-			if (Buffer != CurrentData)
-			{
-				DestinationData = Buffer;
-				break;
-			}
-		}
-
-		if (DestinationData == nullptr)
-		{
-			DestinationData = new Grid2DBuffer(NumCellsX * NumTilesX, NumCellsY * NumTilesY);
-			Buffers.Add(DestinationData);
-		}
-
-		return *DestinationData;
-	}
+	TArray<TUniquePtr<FGrid2DBuffer>> Buffers;
+	FGrid2DBuffer* CurrentData = nullptr;
+	FGrid2DBuffer* DestinationData = nullptr;
 
-	void EndSimulate()
-	{
-		CurrentData = DestinationData;
-		DestinationData = nullptr;
-	}
+	void BeginSimulate();
+	void EndSimulate();
 };
 
-struct FNiagaraDataInterfaceProxyGrid2DCollection : public FNiagaraDataInterfaceProxyRW
+struct FNiagaraDataInterfaceProxyGrid2DCollectionProxy : public FNiagaraDataInterfaceProxyRW
 {
-	FNiagaraDataInterfaceProxyGrid2DCollection() {}		
+	FNiagaraDataInterfaceProxyGrid2DCollectionProxy() {}
 
 	virtual void PreStage(FRHICommandList& RHICmdList, const FNiagaraDataInterfaceSetArgs& Context) override;
 	virtual void PostStage(FRHICommandList& RHICmdList, const FNiagaraDataInterfaceSetArgs& Context) override;
@@ -93,7 +55,7 @@ struct FNiagaraDataInterfaceProxyGrid2DCollection : public FNiagaraDataInterface
 
 	/* List of proxy data for each system instances*/
 	// #todo(dmp): this should all be refactored to avoid duplicate code
-	TMap<FNiagaraSystemInstanceID, Grid2DCollectionRWInstanceData> SystemInstancesToProxyData;
+	TMap<FNiagaraSystemInstanceID, FGrid2DCollectionRWInstanceData_RenderThread> SystemInstancesToProxyData_RT;
 };
 
 UCLASS(EditInlineNew, Category = "Grid", meta = (DisplayName = "Grid2D Collection", Experimental), Blueprintable, BlueprintType)
@@ -122,7 +84,7 @@ public:
 	virtual bool InitPerInstanceData(void* PerInstanceData, FNiagaraSystemInstance* SystemInstance) override;
 	virtual void DestroyPerInstanceData(void* PerInstanceData, FNiagaraSystemInstance* SystemInstance) override;
 	virtual bool PerInstanceTick(void* PerInstanceData, FNiagaraSystemInstance* SystemInstance, float DeltaSeconds) override { return false; }
-	virtual int32 PerInstanceDataSize()const override { return sizeof(Grid2DCollectionRWInstanceData); }
+	virtual int32 PerInstanceDataSize()const override { return sizeof(FGrid2DCollectionRWInstanceData_GameThread); }
 	//~ UNiagaraDataInterface interface END
 
 	// Fills a texture render target 2d with the current data from the simulation
@@ -146,6 +108,6 @@ protected:
 	//~ UNiagaraDataInterface interface
 	virtual bool CopyToInternal(UNiagaraDataInterface* Destination) const override;
 	//~ UNiagaraDataInterface interface END
-};
-
 
+	TMap<FNiagaraSystemInstanceID, FGrid2DCollectionRWInstanceData_GameThread*> SystemInstancesToProxyData_GT;
+};
diff --git a/Engine/Plugins/FX/Niagara/Source/Niagara/Private/NiagaraDataInterfaceGrid2DCollection.cpp b/Engine/Plugins/FX/Niagara/Source/Niagara/Private/NiagaraDataInterfaceGrid2DCollection.cpp
--- a/Engine/Plugins/FX/Niagara/Source/Niagara/Private/NiagaraDataInterfaceGrid2DCollection.cpp
+++ b/Engine/Plugins/FX/Niagara/Source/Niagara/Private/NiagaraDataInterfaceGrid2DCollection.cpp
@@ -7,6 +7,7 @@
 #include "Engine/Texture2D.h"
 #include "NiagaraEmitterInstanceBatcher.h"
 #include "NiagaraSystemInstance.h"
+#include "NiagaraRenderer.h"
 #include "Engine/TextureRenderTarget2D.h"
 
 #define LOCTEXT_NAMESPACE "NiagaraDataInterfaceGrid2DCollection"
@@ -51,19 +52,19 @@ public:
 
 		// Get shader and DI
 		FRHIComputeShader* ComputeShaderRHI = Context.Shader.GetComputeShader();
-		FNiagaraDataInterfaceProxyGrid2DCollection* VFDI = static_cast<FNiagaraDataInterfaceProxyGrid2DCollection*>(Context.DataInterface);
+		FNiagaraDataInterfaceProxyGrid2DCollectionProxy* VFDI = static_cast<FNiagaraDataInterfaceProxyGrid2DCollectionProxy*>(Context.DataInterface);
 		
-		Grid2DCollectionRWInstanceData* ProxyData = VFDI->SystemInstancesToProxyData.Find(Context.SystemInstance);
+		FGrid2DCollectionRWInstanceData_RenderThread* ProxyData = VFDI->SystemInstancesToProxyData_RT.Find(Context.SystemInstance);
 		check(ProxyData);
 
 		int NumCellsTmp[2];
-		NumCellsTmp[0] = ProxyData->NumCellsX;
-		NumCellsTmp[1] = ProxyData->NumCellsY;
-		SetShaderValue(RHICmdList, ComputeShaderRHI, NumCellsParam, NumCellsTmp);		
-		
+		NumCellsTmp[0] = ProxyData->NumCells.X;
+		NumCellsTmp[1] = ProxyData->NumCells.Y;
+		SetShaderValue(RHICmdList, ComputeShaderRHI, NumCellsParam, NumCellsTmp);	
+
 		int NumTilesTmp[2];
-		NumTilesTmp[0] = ProxyData->NumTilesX;
-		NumTilesTmp[1] = ProxyData->NumTilesY;
+		NumTilesTmp[0] = ProxyData->NumTiles.X;
+		NumTilesTmp[1] = ProxyData->NumTiles.Y;
 		SetShaderValue(RHICmdList, ComputeShaderRHI, NumTilesParam, NumTilesTmp);		
 
 		SetShaderValue(RHICmdList, ComputeShaderRHI, CellSizeParam, ProxyData->CellSize);		
@@ -73,17 +74,34 @@ public:
 		FRHISamplerState *SamplerState = TStaticSamplerState<SF_Bilinear, AM_Clamp, AM_Clamp, AM_Clamp>::GetRHI();
 		SetSamplerParameter(RHICmdList, ComputeShaderRHI, SamplerParam, SamplerState);
 
-		if (GridParam.IsBound() && ProxyData->GetCurrentData() != NULL)
+		if (GridParam.IsBound())
 		{
-			RHICmdList.TransitionResource(EResourceTransitionAccess::EReadable, EResourceTransitionPipeline::EComputeToCompute, ProxyData->GetCurrentData()->GridBuffer.UAV);
-			SetSRVParameter(RHICmdList, Context.Shader.GetComputeShader(), GridParam, ProxyData->GetCurrentData()->GridBuffer.SRV);
+			FRHIShaderResourceView* InputGridBuffer;
+			if (ProxyData->CurrentData != nullptr)
+			{
+				RHICmdList.TransitionResource(EResourceTransitionAccess::EReadable, EResourceTransitionPipeline::EComputeToCompute, ProxyData->CurrentData->GridBuffer.UAV);
+				InputGridBuffer = ProxyData->CurrentData->GridBuffer.SRV;
+			}
+			else
+			{
+				InputGridBuffer = FNiagaraRenderer::GetDummyTextureReadBuffer2D();
+			}
+			SetSRVParameter(RHICmdList, Context.Shader.GetComputeShader(), GridParam, InputGridBuffer);
 		}
 
-		if (Context.IsOutputStage && OutputGridParam.IsBound() && ProxyData->GetDestinationData() != NULL)
+		if ( OutputGridParam.IsUAVBound() )
 		{
-			const FTextureRWBuffer2D& TextureBuffer = ProxyData->GetDestinationData()->GridBuffer;
-			RHICmdList.TransitionResource(EResourceTransitionAccess::EWritable, EResourceTransitionPipeline::EComputeToCompute, ProxyData->GetDestinationData()->GridBuffer.UAV);
-			OutputGridParam.SetTexture(RHICmdList, Context.Shader.GetComputeShader(), TextureBuffer.Buffer, TextureBuffer.UAV);
+			FRHIUnorderedAccessView* OutputGridUAV;
+			if (Context.IsOutputStage && ProxyData->DestinationData != nullptr)
+			{
+				OutputGridUAV = ProxyData->DestinationData->GridBuffer.UAV;
+				RHICmdList.TransitionResource(EResourceTransitionAccess::EWritable, EResourceTransitionPipeline::EComputeToCompute, OutputGridUAV);
+			}
+			else
+			{
+				OutputGridUAV = Context.Batcher->GetEmptyRWTextureFromPool(RHICmdList, PF_R32_FLOAT);
+			}
+			RHICmdList.SetUAVParameter(ComputeShaderRHI, OutputGridParam.GetUAVIndex(), OutputGridUAV);
 		}
 	}
 
@@ -116,7 +134,7 @@ IMPLEMENT_NIAGARA_DI_PARAMETER(UNiagaraDataInterfaceGrid2DCollection, FNiagaraDa
 UNiagaraDataInterfaceGrid2DCollection::UNiagaraDataInterfaceGrid2DCollection(FObjectInitializer const& ObjectInitializer)
 	: Super(ObjectInitializer)
 {
-	Proxy.Reset(new FNiagaraDataInterfaceProxyGrid2DCollection());
+	Proxy.Reset(new FNiagaraDataInterfaceProxyGrid2DCollectionProxy());
 	PushToRenderThread();
 }
 
@@ -232,7 +250,7 @@ void UNiagaraDataInterfaceGrid2DCollection::GetParameterDefinitionHLSL(const FNi
 		RWTexture2D<float> RW{OutputGridName};
 		int2 {NumTiles};
 		SamplerState {SamplerName};
-
+	
 	)");
 	TMap<FString, FStringFormatArg> ArgsDeclarations = {				
 		{ TEXT("GridName"),    GridName + ParamInfo.DataInterfaceHLSLSymbol },
@@ -274,7 +292,7 @@ bool UNiagaraDataInterfaceGrid2DCollection::GetFunctionHLSL(const FNiagaraDataIn
 	{
 		static const TCHAR *FormatBounds = TEXT(R"(
 			void {FunctionName}(int In_IndexX, int In_IndexY, int In_AttributeIndex, float In_Value, out int val)
-			{				
+			{			
 				int TileIndexX = In_AttributeIndex % {NumTiles}.x;
 				int TileIndexY = In_AttributeIndex / {NumTiles}.x;
 	
@@ -331,13 +349,12 @@ bool UNiagaraDataInterfaceGrid2DCollection::CopyToInternal(UNiagaraDataInterface
 bool UNiagaraDataInterfaceGrid2DCollection::InitPerInstanceData(void* PerInstanceData, FNiagaraSystemInstance* SystemInstance)
 {
 	check(Proxy);
-	UE_LOG(LogNiagara, Log, TEXT("UNiagaraDataInterfaceGrid2DCollection::InitPerInstanceData DI: %p PID: %p SI: %p %s"), this, PerInstanceData, SystemInstance, *SystemInstance->GetSystem()->GetPathName())
-	Grid2DCollectionRWInstanceData* InstanceData = new (PerInstanceData) Grid2DCollectionRWInstanceData();
 
-	FNiagaraDataInterfaceProxyGrid2DCollection* RT_Proxy = GetProxyAs<FNiagaraDataInterfaceProxyGrid2DCollection>();
-	
-	int RT_NumCellsX = NumCellsX;
-	int RT_NumCellsY = NumCellsY;
+	FGrid2DCollectionRWInstanceData_GameThread* InstanceData = new (PerInstanceData) FGrid2DCollectionRWInstanceData_GameThread();
+	SystemInstancesToProxyData_GT.Emplace(SystemInstance->GetId(), InstanceData);
+
+	InstanceData->NumCells.X = NumCellsX;
+	InstanceData->NumCells.Y = NumCellsY;
 
 	// #todo(dmp): refactor
 	int MaxDim = 16384;
@@ -348,69 +365,53 @@ bool UNiagaraDataInterfaceGrid2DCollection::InitPerInstanceData(void* PerInstanc
 	int NumTilesX = NumAttributes <= MaxTilesX ? NumAttributes : MaxTilesX;
 	int NumTilesY = ceil(1.0*NumAttributes / NumTilesX);
 
-	int RT_NumTilesX = NumTilesX;
-	int RT_NumTilesY = NumTilesY;
+	InstanceData->NumTiles.X = NumTilesX;
+	InstanceData->NumTiles.Y = NumTilesY;
 	
-	FVector2D RT_WorldBBoxSize = WorldBBoxSize;
+	InstanceData->WorldBBoxSize = WorldBBoxSize;
 
 	// If we are setting the grid from the voxel size, then recompute NumVoxels and change bbox	
 	if (SetGridFromMaxAxis)
 	{
 		float CellSize = FMath::Max(WorldBBoxSize.X, WorldBBoxSize.Y) / NumCellsMaxAxis;
 
-		RT_NumCellsX = WorldBBoxSize.X / CellSize;
-		RT_NumCellsY = WorldBBoxSize.Y / CellSize;
+		InstanceData->NumCells.X = WorldBBoxSize.X / CellSize;
+		InstanceData->NumCells.Y = WorldBBoxSize.Y / CellSize;
 
 		// Pad grid by 1 voxel if our computed bounding box is too small
-		if (WorldBBoxSize.X > WorldBBoxSize.Y && !FMath::IsNearlyEqual(CellSize * RT_NumCellsY, WorldBBoxSize.Y))
+		if (WorldBBoxSize.X > WorldBBoxSize.Y && !FMath::IsNearlyEqual(CellSize * InstanceData->NumCells.Y, WorldBBoxSize.Y))
 		{
-			RT_NumCellsY++;
+			InstanceData->NumCells.Y++;
 		}
-		else if (WorldBBoxSize.X < WorldBBoxSize.Y && !FMath::IsNearlyEqual(CellSize * RT_NumCellsX, WorldBBoxSize.X))
+		else if (WorldBBoxSize.X < WorldBBoxSize.Y && !FMath::IsNearlyEqual(CellSize * InstanceData->NumCells.X, WorldBBoxSize.X))
 		{
-			RT_NumCellsX++;
+			InstanceData->NumCells.X++;
 		}		
 		
-		RT_WorldBBoxSize = FVector2D(RT_NumCellsX, RT_NumCellsY) * CellSize;
-		NumCellsX = RT_NumCellsX;
-		NumCellsY = RT_NumCellsY;
+		InstanceData->WorldBBoxSize = FVector2D(InstanceData->NumCells.X, InstanceData->NumCells.Y) * CellSize;
+		NumCellsX = InstanceData->NumCells.X;
+		NumCellsY = InstanceData->NumCells.Y;
 	}	
 
-	TSet<int> RT_OutputShaderStages = OutputShaderStages;
-	TSet<int> RT_IterationShaderStages = IterationShaderStages;
-
-	FVector2D RT_CellSize = RT_WorldBBoxSize / FVector2D(RT_NumCellsX, RT_NumCellsY);
-	InstanceData->CellSize = RT_CellSize;
-	InstanceData->WorldBBoxSize = RT_WorldBBoxSize;
-
-	//UE_LOG(LogNiagara, Log, TEXT("UNiagaraDataInterfaceGrid2DCollection::Proxy %p DI: %p"), RT_Proxy, this);
-	//for (const int32 i : IterationShaderStages)
-	//{
-	//	UE_LOG(LogNiagara, Log, TEXT("%d"), i);
-	//}
+	InstanceData->CellSize = InstanceData->WorldBBoxSize / FVector2D(InstanceData->NumCells.X, InstanceData->NumCells.Y);
 
-	
 	// Push Updates to Proxy.
+	FNiagaraDataInterfaceProxyGrid2DCollectionProxy* RT_Proxy = GetProxyAs<FNiagaraDataInterfaceProxyGrid2DCollectionProxy>();
 	ENQUEUE_RENDER_COMMAND(FUpdateData)(
-		[RT_Proxy, RT_NumCellsX, RT_NumCellsY, RT_NumTilesX, RT_NumTilesY, RT_CellSize, RT_WorldBBoxSize, RT_OutputShaderStages, RT_IterationShaderStages, InstanceID = SystemInstance->GetId()](FRHICommandListImmediate& RHICmdList)
+		[RT_Proxy, InstanceID = SystemInstance->GetId(), RT_InstanceData=*InstanceData, RT_OutputShaderStages=OutputShaderStages, RT_IterationShaderStages= IterationShaderStages](FRHICommandListImmediate& RHICmdList)
 	{
-		check(!RT_Proxy->SystemInstancesToProxyData.Contains(InstanceID));
-		Grid2DCollectionRWInstanceData* TargetData = &RT_Proxy->SystemInstancesToProxyData.Add(InstanceID);
-
-		TargetData->NumCellsX = RT_NumCellsX;
-		TargetData->NumCellsY = RT_NumCellsY;
+		check(!RT_Proxy->SystemInstancesToProxyData_RT.Contains(InstanceID));
+		FGrid2DCollectionRWInstanceData_RenderThread* TargetData = &RT_Proxy->SystemInstancesToProxyData_RT.Add(InstanceID);
 
-		TargetData->NumTilesX = RT_NumTilesX;
-		TargetData->NumTilesY = RT_NumTilesY;
-
-		TargetData->CellSize = RT_CellSize;
-		
-		TargetData->WorldBBoxSize = RT_WorldBBoxSize;
+		TargetData->NumCells = RT_InstanceData.NumCells;
+		TargetData->NumTiles = RT_InstanceData.NumTiles;
+		TargetData->CellSize = RT_InstanceData.CellSize;
+		TargetData->WorldBBoxSize = RT_InstanceData.WorldBBoxSize;
 
 		RT_Proxy->OutputSimulationStages_DEPRECATED = RT_OutputShaderStages;
 		RT_Proxy->IterationSimulationStages_DEPRECATED = RT_IterationShaderStages;
 
-		RT_Proxy->SetElementCount(TargetData->NumCellsX * TargetData->NumCellsY);
+		RT_Proxy->SetElementCount(TargetData->NumCells.X * TargetData->NumCells.Y);
 	});
 
 	return true;
@@ -419,26 +420,25 @@ bool UNiagaraDataInterfaceGrid2DCollection::InitPerInstanceData(void* PerInstanc
 
 void UNiagaraDataInterfaceGrid2DCollection::DestroyPerInstanceData(void* PerInstanceData, FNiagaraSystemInstance* SystemInstance)
 {
-	Grid2DCollectionRWInstanceData* InstanceData = static_cast<Grid2DCollectionRWInstanceData*>(PerInstanceData);
+	SystemInstancesToProxyData_GT.Remove(SystemInstance->GetId());
 
-	InstanceData->~Grid2DCollectionRWInstanceData();
+	FGrid2DCollectionRWInstanceData_GameThread* InstanceData = static_cast<FGrid2DCollectionRWInstanceData_GameThread*>(PerInstanceData);
 
-	FNiagaraDataInterfaceProxyGrid2DCollection* ThisProxy = GetProxyAs<FNiagaraDataInterfaceProxyGrid2DCollection>();
+	InstanceData->~FGrid2DCollectionRWInstanceData_GameThread();
 
+	FNiagaraDataInterfaceProxyGrid2DCollectionProxy* RT_Proxy = GetProxyAs<FNiagaraDataInterfaceProxyGrid2DCollectionProxy>();
 	ENQUEUE_RENDER_COMMAND(FNiagaraDIDestroyInstanceData) (
-		[ThisProxy, InstanceID = SystemInstance->GetId(), Batcher = SystemInstance->GetBatcher()](FRHICommandListImmediate& CmdList)
-	{
-		//check(ThisProxy->SystemInstancesToProxyData.Contains(InstanceID));
-		ThisProxy->SystemInstancesToProxyData.Remove(InstanceID);
-	}
+		[RT_Proxy, InstanceID=SystemInstance->GetId(), Batcher=SystemInstance->GetBatcher()](FRHICommandListImmediate& CmdList)
+		{
+			//check(ThisProxy->SystemInstancesToProxyData.Contains(InstanceID));
+			RT_Proxy->SystemInstancesToProxyData_RT.Remove(InstanceID);
+		}
 	);
 }
 
 UFUNCTION(BlueprintCallable, Category = Niagara)
 bool UNiagaraDataInterfaceGrid2DCollection::FillTexture2D(const UNiagaraComponent *Component, UTextureRenderTarget2D *Dest, int AttributeIndex)
 {
-	FNiagaraDataInterfaceProxyGrid2DCollection* TProxy = GetProxyAs<FNiagaraDataInterfaceProxyGrid2DCollection>();
-
 	if (!Component || !Dest)
 	{
 		return false;
@@ -464,30 +464,29 @@ bool UNiagaraDataInterfaceGrid2DCollection::FillTexture2D(const UNiagaraComponen
 		return false;
 	}
 
-	FNiagaraSystemInstanceID InstanceID = SystemInstance->GetId();
-
+	FNiagaraDataInterfaceProxyGrid2DCollectionProxy* RT_Proxy = GetProxyAs<FNiagaraDataInterfaceProxyGrid2DCollectionProxy>();
 	ENQUEUE_RENDER_COMMAND(FUpdateDIColorCurve)(
-		[Dest, TProxy, InstanceID, AttributeIndex](FRHICommandListImmediate& RHICmdList)
+		[RT_Proxy, InstanceID=SystemInstance->GetId(), RT_TextureResource=Dest->Resource, AttributeIndex](FRHICommandListImmediate& RHICmdList)
 	{
-		Grid2DCollectionRWInstanceData* Grid2DInstanceData = TProxy->SystemInstancesToProxyData.Find(InstanceID);
+		FGrid2DCollectionRWInstanceData_RenderThread* Grid2DInstanceData = RT_Proxy->SystemInstancesToProxyData_RT.Find(InstanceID);
 
-		if (Dest && Dest->Resource && Grid2DInstanceData && Grid2DInstanceData->CurrentData)
+		if (RT_TextureResource && RT_TextureResource->TextureRHI.IsValid() && Grid2DInstanceData && Grid2DInstanceData->CurrentData)
 		{
 			FRHICopyTextureInfo CopyInfo;
-			CopyInfo.Size = FIntVector(Grid2DInstanceData->NumCellsX, Grid2DInstanceData->NumCellsY, 1);
+			CopyInfo.Size = FIntVector(Grid2DInstanceData->NumCells.X, Grid2DInstanceData->NumCells.Y, 1);
 
-			int TileIndexX = AttributeIndex % Grid2DInstanceData->NumTilesX;
-			int TileIndexY = AttributeIndex / Grid2DInstanceData->NumTilesX;
-			int StartX = TileIndexX * Grid2DInstanceData->NumCellsX;
-			int StartY = TileIndexY * Grid2DInstanceData->NumCellsY;
+			int TileIndexX = AttributeIndex % Grid2DInstanceData->NumTiles.X;
+			int TileIndexY = AttributeIndex / Grid2DInstanceData->NumTiles.X;
+			int StartX = TileIndexX * Grid2DInstanceData->NumCells.X;
+			int StartY = TileIndexY * Grid2DInstanceData->NumCells.Y;
 			CopyInfo.SourcePosition = FIntVector(StartX, StartY, 0);
 
 			RHICmdList.TransitionResource(EResourceTransitionAccess::EReadable, Grid2DInstanceData->CurrentData->GridBuffer.Buffer);
-			RHICmdList.TransitionResource(EResourceTransitionAccess::EWritable, Dest->Resource->TextureRHI);
+			RHICmdList.TransitionResource(EResourceTransitionAccess::EWritable, RT_TextureResource->TextureRHI);
 
-			RHICmdList.CopyTexture(Grid2DInstanceData->CurrentData->GridBuffer.Buffer, Dest->Resource->TextureRHI, CopyInfo);
+			RHICmdList.CopyTexture(Grid2DInstanceData->CurrentData->GridBuffer.Buffer, RT_TextureResource->TextureRHI, CopyInfo);
 
-			RHICmdList.TransitionResource(EResourceTransitionAccess::EReadable, Dest->Resource->TextureRHI);
+			RHICmdList.TransitionResource(EResourceTransitionAccess::EReadable, RT_TextureResource->TextureRHI);
 		}
 	});
 
@@ -497,8 +496,6 @@ bool UNiagaraDataInterfaceGrid2DCollection::FillTexture2D(const UNiagaraComponen
 UFUNCTION(BlueprintCallable, Category = Niagara)
 bool UNiagaraDataInterfaceGrid2DCollection::FillRawTexture2D(const UNiagaraComponent *Component, UTextureRenderTarget2D *Dest, int &TilesX, int &TilesY)
 {
-	FNiagaraDataInterfaceProxyGrid2DCollection* TProxy = GetProxyAs<FNiagaraDataInterfaceProxyGrid2DCollection>();
-
 	if (!Component)
 	{
 		TilesX = -1;
@@ -506,25 +503,24 @@ bool UNiagaraDataInterfaceGrid2DCollection::FillRawTexture2D(const UNiagaraCompo
 		return false;
 	}
 
-	FNiagaraSystemInstance *SystemInstance = Component->GetSystemInstance();
+	FNiagaraSystemInstance* SystemInstance = Component->GetSystemInstance();
 	if (!SystemInstance)
 	{
 		TilesX = -1;
 		TilesY = -1;
 		return false;
 	}
-	FNiagaraSystemInstanceID InstanceID = SystemInstance->GetId();
-	
-	Grid2DCollectionRWInstanceData* Grid2DInstanceData = TProxy->SystemInstancesToProxyData.Find(InstanceID);
+
+	FGrid2DCollectionRWInstanceData_GameThread* Grid2DInstanceData = SystemInstancesToProxyData_GT.FindRef(SystemInstance->GetId());
 	if (!Grid2DInstanceData)
 	{
 		TilesX = -1;
 		TilesY = -1;
 		return false;
 	}
-
-	TilesX = Grid2DInstanceData->NumTilesX;
-	TilesY = Grid2DInstanceData->NumTilesY;	
+	
+	TilesX = Grid2DInstanceData->NumTiles.X;
+	TilesY = Grid2DInstanceData->NumTiles.Y;
 
 	// check dest size and type needs to be float
 	// #todo(dmp): don't hardcode float since we might do other stuff in the future
@@ -534,20 +530,20 @@ bool UNiagaraDataInterfaceGrid2DCollection::FillRawTexture2D(const UNiagaraCompo
 		return false;
 	}
 
-	// #todo(dmp): pass grid2dinstancedata through?
+	FNiagaraDataInterfaceProxyGrid2DCollectionProxy* RT_Proxy = GetProxyAs<FNiagaraDataInterfaceProxyGrid2DCollectionProxy>();
 	ENQUEUE_RENDER_COMMAND(FUpdateDIColorCurve)(
-		[Dest, TProxy, InstanceID](FRHICommandListImmediate& RHICmdList)
+		[RT_Proxy, RT_InstanceID=SystemInstance->GetId(), RT_TextureResource=Dest->Resource](FRHICommandListImmediate& RHICmdList)
 	{
-		Grid2DCollectionRWInstanceData* Grid2DInstanceDataTmp = TProxy->SystemInstancesToProxyData.Find(InstanceID);
-
-		if (Dest && Dest->Resource && Grid2DInstanceDataTmp && Grid2DInstanceDataTmp->CurrentData)
+		FGrid2DCollectionRWInstanceData_RenderThread* RT_Grid2DInstanceData = RT_Proxy->SystemInstancesToProxyData_RT.Find(RT_InstanceID);
+		if (RT_TextureResource && RT_TextureResource->TextureRHI.IsValid() && RT_Grid2DInstanceData && RT_Grid2DInstanceData->CurrentData)
 		{
-			RHICmdList.TransitionResource(EResourceTransitionAccess::EReadable, Grid2DInstanceDataTmp->CurrentData->GridBuffer.Buffer);
-			RHICmdList.TransitionResource(EResourceTransitionAccess::EWritable, Dest->Resource->TextureRHI);
+			RHICmdList.TransitionResource(EResourceTransitionAccess::EReadable, RT_Grid2DInstanceData->CurrentData->GridBuffer.Buffer);
+			RHICmdList.TransitionResource(EResourceTransitionAccess::EWritable, RT_TextureResource->TextureRHI);
 
 			FRHICopyTextureInfo CopyInfo;
-			RHICmdList.CopyTexture(Grid2DInstanceDataTmp->CurrentData->GridBuffer.Buffer, Dest->Resource->TextureRHI, CopyInfo);
-			RHICmdList.TransitionResource(EResourceTransitionAccess::EReadable, Dest->Resource->TextureRHI);
+			RHICmdList.CopyTexture(RT_Grid2DInstanceData->CurrentData->GridBuffer.Buffer, RT_TextureResource->TextureRHI, CopyInfo);
+
+			RHICmdList.TransitionResource(EResourceTransitionAccess::EReadable, RT_TextureResource->TextureRHI);
 		}
 	});
 
@@ -557,8 +553,6 @@ bool UNiagaraDataInterfaceGrid2DCollection::FillRawTexture2D(const UNiagaraCompo
 UFUNCTION(BlueprintCallable, Category = Niagara)
 void UNiagaraDataInterfaceGrid2DCollection::GetRawTextureSize(const UNiagaraComponent *Component, int &SizeX, int &SizeY)
 {
-	FNiagaraDataInterfaceProxyGrid2DCollection* TProxy = GetProxyAs<FNiagaraDataInterfaceProxyGrid2DCollection>();
-
 	if (!Component)
 	{
 		SizeX = -1;
@@ -575,30 +569,21 @@ void UNiagaraDataInterfaceGrid2DCollection::GetRawTextureSize(const UNiagaraComp
 	}
 	FNiagaraSystemInstanceID InstanceID = SystemInstance->GetId();
 
-	Grid2DCollectionRWInstanceData* Grid2DInstanceData = TProxy->SystemInstancesToProxyData.Find(InstanceID);
+	FGrid2DCollectionRWInstanceData_GameThread* Grid2DInstanceData = SystemInstancesToProxyData_GT.FindRef(InstanceID);
 	if (!Grid2DInstanceData)
 	{
 		SizeX = -1;
 		SizeY = -1;
 		return;
 	}
-	
-	if (Grid2DInstanceData->Buffers.Num() <= 0)
-	{
-		SizeX = -1;
-		SizeY = -1;
-		return;
-	}
 
-	SizeX = Grid2DInstanceData->Buffers[0]->GridBuffer.Buffer->GetSizeX();
-	SizeY = Grid2DInstanceData->Buffers[0]->GridBuffer.Buffer->GetSizeY();
+	SizeX = Grid2DInstanceData->NumCells.X * Grid2DInstanceData->NumTiles.X;
+	SizeY = Grid2DInstanceData->NumCells.Y * Grid2DInstanceData->NumTiles.Y;
 }
 
 UFUNCTION(BlueprintCallable, Category = Niagara)
 void UNiagaraDataInterfaceGrid2DCollection::GetTextureSize(const UNiagaraComponent *Component, int &SizeX, int &SizeY)
 {
-	FNiagaraDataInterfaceProxyGrid2DCollection* TProxy = GetProxyAs<FNiagaraDataInterfaceProxyGrid2DCollection>();
-
 	if (!Component)
 	{
 		SizeX = -1;
@@ -615,7 +600,7 @@ void UNiagaraDataInterfaceGrid2DCollection::GetTextureSize(const UNiagaraCompone
 	}
 	FNiagaraSystemInstanceID InstanceID = SystemInstance->GetId();
 
-	Grid2DCollectionRWInstanceData* Grid2DInstanceData = TProxy->SystemInstancesToProxyData.Find(InstanceID);
+	FGrid2DCollectionRWInstanceData_GameThread* Grid2DInstanceData = SystemInstancesToProxyData_GT.FindRef(InstanceID);
 	if (!Grid2DInstanceData)
 	{
 		SizeX = -1;
@@ -623,14 +608,13 @@ void UNiagaraDataInterfaceGrid2DCollection::GetTextureSize(const UNiagaraCompone
 		return;
 	}
 
-	SizeX = Grid2DInstanceData->NumCellsX;
-	SizeY = Grid2DInstanceData->NumCellsY;
+	SizeX = Grid2DInstanceData->NumCells.X;
+	SizeY = Grid2DInstanceData->NumCells.Y;
 }
 
-
 void UNiagaraDataInterfaceGrid2DCollection::GetWorldBBoxSize(FVectorVMContext& Context)
 {
-	VectorVM::FUserPtrHandler<Grid2DCollectionRWInstanceData> InstData(Context);
+	VectorVM::FUserPtrHandler<FGrid2DCollectionRWInstanceData_GameThread> InstData(Context);
 	VectorVM::FExternalFuncRegisterHandler<float> OutWorldBoundsX(Context);
 	VectorVM::FExternalFuncRegisterHandler<float> OutWorldBoundsY(Context);
 
@@ -644,7 +628,7 @@ void UNiagaraDataInterfaceGrid2DCollection::GetWorldBBoxSize(FVectorVMContext& C
 
 void UNiagaraDataInterfaceGrid2DCollection::GetCellSize(FVectorVMContext& Context)
 {
-	VectorVM::FUserPtrHandler<Grid2DCollectionRWInstanceData> InstData(Context);
+	VectorVM::FUserPtrHandler<FGrid2DCollectionRWInstanceData_GameThread> InstData(Context);
 	VectorVM::FExternalFuncRegisterHandler<float> OutCellSizeX(Context);
 	VectorVM::FExternalFuncRegisterHandler<float> OutCellSizeY(Context);
 
@@ -655,12 +639,37 @@ void UNiagaraDataInterfaceGrid2DCollection::GetCellSize(FVectorVMContext& Contex
 	}
 }
 
-void FNiagaraDataInterfaceProxyGrid2DCollection::PreStage(FRHICommandList& RHICmdList, const FNiagaraDataInterfaceSetArgs& Context)
+void FGrid2DCollectionRWInstanceData_RenderThread::BeginSimulate()
+{
+	for (TUniquePtr<FGrid2DBuffer>& Buffer : Buffers)
+	{
+		check(Buffer.IsValid());
+		if (Buffer.Get() != CurrentData)
+		{
+			DestinationData = Buffer.Get();
+			break;
+		}
+	}
+
+	if (DestinationData == nullptr)
+	{
+		DestinationData = new FGrid2DBuffer(NumCells.X * NumTiles.X, NumCells.Y * NumTiles.Y);
+		Buffers.Emplace(DestinationData);
+	}
+}
+
+void FGrid2DCollectionRWInstanceData_RenderThread::EndSimulate()
+{
+	CurrentData = DestinationData;
+	DestinationData = nullptr;
+}
+
+void FNiagaraDataInterfaceProxyGrid2DCollectionProxy::PreStage(FRHICommandList& RHICmdList, const FNiagaraDataInterfaceSetArgs& Context)
 {
 	// #todo(dmp): Context doesnt need to specify if a stage is output or not since we moved pre/post stage to the DI itself.  Not sure which design is better for the future
 	if (Context.IsOutputStage)
 	{
-		Grid2DCollectionRWInstanceData* ProxyData = SystemInstancesToProxyData.Find(Context.SystemInstance);
+		FGrid2DCollectionRWInstanceData_RenderThread* ProxyData = SystemInstancesToProxyData_RT.Find(Context.SystemInstance);
 
 		ProxyData->BeginSimulate();
 
@@ -668,7 +677,7 @@ void FNiagaraDataInterfaceProxyGrid2DCollection::PreStage(FRHICommandList& RHICm
 		// we start.  If a user wants to access data from the previous stage, then they can read from the current data.
 
 		// #todo(dmp): we might want to expose an option where we have buffers that are write only and need a clear (ie: no buffering like the neighbor grid).  They would be considered transient perhaps?  It'd be more
-		// memory efficient since it would theoretically not require any double buffering.
+		// memory efficient since it would theoretically not require any double buffering.		
 		if (!Context.IsIterationStage)
 		{
 			RHICmdList.TransitionResource(EResourceTransitionAccess::ERWBarrier, EResourceTransitionPipeline::EComputeToCompute, ProxyData->DestinationData->GridBuffer.UAV);
@@ -678,33 +687,32 @@ void FNiagaraDataInterfaceProxyGrid2DCollection::PreStage(FRHICommandList& RHICm
 	}
 }
 
-void FNiagaraDataInterfaceProxyGrid2DCollection::PostStage(FRHICommandList& RHICmdList, const FNiagaraDataInterfaceSetArgs& Context)
+void FNiagaraDataInterfaceProxyGrid2DCollectionProxy::PostStage(FRHICommandList& RHICmdList, const FNiagaraDataInterfaceSetArgs& Context)
 {
 	
 	if (Context.IsOutputStage)
 	{
-		Grid2DCollectionRWInstanceData* ProxyData = SystemInstancesToProxyData.Find(Context.SystemInstance);
+		FGrid2DCollectionRWInstanceData_RenderThread* ProxyData = SystemInstancesToProxyData_RT.Find(Context.SystemInstance);
 		ProxyData->EndSimulate();
 	}
 }
 
-void FNiagaraDataInterfaceProxyGrid2DCollection::ResetData(FRHICommandList& RHICmdList, const FNiagaraDataInterfaceSetArgs& Context)
+void FNiagaraDataInterfaceProxyGrid2DCollectionProxy::ResetData(FRHICommandList& RHICmdList, const FNiagaraDataInterfaceSetArgs& Context)
 {	
-	Grid2DCollectionRWInstanceData* ProxyData = SystemInstancesToProxyData.Find(Context.SystemInstance);
-
+	FGrid2DCollectionRWInstanceData_RenderThread* ProxyData = SystemInstancesToProxyData_RT.Find(Context.SystemInstance);
 	if (!ProxyData)
 	{
 		return;
 	}
 
-	for (Grid2DBuffer* Buffer : ProxyData->Buffers)
+	for (TUniquePtr<FGrid2DBuffer>& Buffer : ProxyData->Buffers)
 	{
-		if (Buffer)
+		if (Buffer.IsValid())
 		{
 			RHICmdList.TransitionResource(EResourceTransitionAccess::ERWBarrier, EResourceTransitionPipeline::EComputeToCompute, Buffer->GridBuffer.UAV);
 			RHICmdList.ClearUAVFloat(Buffer->GridBuffer.UAV, FVector4(ForceInitToZero));
 			RHICmdList.TransitionResource(EResourceTransitionAccess::ERWBarrier, EResourceTransitionPipeline::EComputeToCompute, Buffer->GridBuffer.UAV);
-		}
+		}		
 	}	
 }
-#undef LOCTEXT_NAMESPACE
\ No newline at end of file
+#undef LOCTEXT_NAMESPACE
diff --git a/Engine/Plugins/FX/Niagara/Source/Niagara/Private/NiagaraDataInterfaceNeighborGrid3D.cpp b/Engine/Plugins/FX/Niagara/Source/Niagara/Private/NiagaraDataInterfaceNeighborGrid3D.cpp
--- a/Engine/Plugins/FX/Niagara/Source/Niagara/Private/NiagaraDataInterfaceNeighborGrid3D.cpp
+++ b/Engine/Plugins/FX/Niagara/Source/Niagara/Private/NiagaraDataInterfaceNeighborGrid3D.cpp
@@ -4,7 +4,8 @@
 #include "ShaderParameterUtils.h"
 #include "NiagaraEmitterInstanceBatcher.h"
 #include "NiagaraSystemInstance.h"
-
+#include "NiagaraRenderer.h"
+#include "NiagaraShaderParticleID.h"
 
 #define LOCTEXT_NAMESPACE "NiagaraDataInterfaceNeighborGrid3D"
 
@@ -82,9 +83,22 @@ public:
 				RHICmdList.TransitionResource(EResourceTransitionAccess::EReadable, EResourceTransitionPipeline::EComputeToCompute, ProxyData->NeighborhoodCountBuffer.UAV);
 				SetSRVParameter(RHICmdList, Context.Shader.GetComputeShader(), ParticleNeighborCountGridParam, ProxyData->NeighborhoodCountBuffer.SRV);
 			}
+
+			if (OutputParticleNeighborsGridParam.IsUAVBound())
+			{
+				RHICmdList.SetUAVParameter(ComputeShaderRHI, OutputParticleNeighborsGridParam.GetUAVIndex(), Context.Batcher->GetEmptyRWBufferFromPool(RHICmdList, PF_R32_SINT));
+			}
+
+			if (OutputParticleNeighborCountGridParam.IsUAVBound())
+			{
+				RHICmdList.SetUAVParameter(ComputeShaderRHI, OutputParticleNeighborCountGridParam.GetUAVIndex(), Context.Batcher->GetEmptyRWBufferFromPool(RHICmdList, PF_R32_SINT));
+			}
 		}
 		else
 		{
+			SetSRVParameter(RHICmdList, Context.Shader.GetComputeShader(), ParticleNeighborsGridParam, FNiagaraRenderer::GetDummyIntBuffer());
+			SetSRVParameter(RHICmdList, Context.Shader.GetComputeShader(), ParticleNeighborCountGridParam, FNiagaraRenderer::GetDummyIntBuffer());
+
 			if (OutputParticleNeighborsGridParam.IsBound())
 			{
 				RHICmdList.TransitionResource(EResourceTransitionAccess::EWritable, EResourceTransitionPipeline::EComputeToCompute, ProxyData->NeighborhoodBuffer.UAV);
@@ -444,8 +458,8 @@ bool UNiagaraDataInterfaceNeighborGrid3D::InitPerInstanceData(void* PerInstanceD
 		// #todo(dmp): element count is still defined on the proxy and not the instance data
 		RT_Proxy->SetElementCount(RT_NumVoxels.X *  RT_NumVoxels.Y * RT_NumVoxels.Z);
 
-		TargetData->NeighborhoodCountBuffer.Initialize(sizeof(int32), TargetData->NumVoxels.X*TargetData->NumVoxels.Y*TargetData->NumVoxels.Z, EPixelFormat::PF_R32_SINT, BUF_Static);
-		TargetData->NeighborhoodBuffer.Initialize(sizeof(int32), TargetData->NumVoxels.X*TargetData->NumVoxels.Y*TargetData->NumVoxels.Z*TargetData->MaxNeighborsPerVoxel, EPixelFormat::PF_R32_SINT, BUF_Static);
+		TargetData->NeighborhoodCountBuffer.Initialize(sizeof(int32), TargetData->NumVoxels.X*TargetData->NumVoxels.Y*TargetData->NumVoxels.Z, EPixelFormat::PF_R32_SINT, BUF_Static, TEXT("NiagaraNeighborGrid3D::NeighborCount"));
+		TargetData->NeighborhoodBuffer.Initialize(sizeof(int32), TargetData->NumVoxels.X*TargetData->NumVoxels.Y*TargetData->NumVoxels.Z*TargetData->MaxNeighborsPerVoxel, EPixelFormat::PF_R32_SINT, BUF_Static, TEXT("NiagaraNeighborGrid3D::NeighborsGrid"));
 
 	});
 
@@ -476,8 +490,10 @@ void FNiagaraDataInterfaceProxyNeighborGrid3D::PreStage(FRHICommandList& RHICmdL
 	{
 		NeighborGrid3DRWInstanceData* ProxyData = SystemInstancesToProxyData.Find(Context.SystemInstance);
 
-		RHICmdList.ClearUAVUint(ProxyData->NeighborhoodBuffer.UAV, FUintVector4((uint32)-1, (uint32)-1, (uint32)-1, (uint32)-1));
-		RHICmdList.ClearUAVUint(ProxyData->NeighborhoodCountBuffer.UAV, FUintVector4(0, 0, 0 ,0));
+		SCOPED_DRAW_EVENT(RHICmdList, NiagaraNeighborGrid3DClearNeighborInfo);
+		ERHIFeatureLevel::Type FeatureLevel = Context.Batcher->GetFeatureLevel();
+		NiagaraFillGPUIntBuffer(RHICmdList, FeatureLevel, ProxyData->NeighborhoodBuffer, -1);
+		NiagaraFillGPUIntBuffer(RHICmdList, FeatureLevel, ProxyData->NeighborhoodCountBuffer, 0);
 	}
 }
 
diff --git a/Engine/Plugins/FX/Niagara/Source/Niagara/Private/NiagaraDataInterfaceParticleRead.cpp b/Engine/Plugins/FX/Niagara/Source/Niagara/Private/NiagaraDataInterfaceParticleRead.cpp
--- a/Engine/Plugins/FX/Niagara/Source/Niagara/Private/NiagaraDataInterfaceParticleRead.cpp
+++ b/Engine/Plugins/FX/Niagara/Source/Niagara/Private/NiagaraDataInterfaceParticleRead.cpp
@@ -224,13 +224,13 @@ struct FNiagaraDataInterfaceParametersCS_ParticleRead : public FNiagaraDataInter
 		{
 			SetShaderValue(RHICmdList, ComputeShader, NumSpawnedParticlesParam, 0);
 			SetShaderValue(RHICmdList, ComputeShader, SpawnedParticlesAcquireTagParam, 0);
-			SetSRVParameter(RHICmdList, ComputeShader, SpawnedIDsBufferParam, FNiagaraRenderer::GetDummyIntBuffer().SRV);
+			SetSRVParameter(RHICmdList, ComputeShader, SpawnedIDsBufferParam, FNiagaraRenderer::GetDummyIntBuffer());
 			SetShaderValue(RHICmdList, ComputeShader, NumParticlesParam, 0);
 		}
 
-		SetSRVParameter(RHICmdList, ComputeShader, IDToIndexTableParam, FNiagaraRenderer::GetDummyIntBuffer().SRV);
-		SetSRVParameter(RHICmdList, ComputeShader, InputFloatBufferParam, FNiagaraRenderer::GetDummyFloatBuffer().SRV);
-		SetSRVParameter(RHICmdList, ComputeShader, InputIntBufferParam, FNiagaraRenderer::GetDummyIntBuffer().SRV);
+		SetSRVParameter(RHICmdList, ComputeShader, IDToIndexTableParam, FNiagaraRenderer::GetDummyIntBuffer());
+		SetSRVParameter(RHICmdList, ComputeShader, InputFloatBufferParam, FNiagaraRenderer::GetDummyFloatBuffer());
+		SetSRVParameter(RHICmdList, ComputeShader, InputIntBufferParam, FNiagaraRenderer::GetDummyIntBuffer());
 		SetShaderValue(RHICmdList, ComputeShader, ParticleStrideFloatParam, 0);
 		SetShaderValue(RHICmdList, ComputeShader, ParticleStrideIntParam, 0);
 		SetShaderValueArray(RHICmdList, ComputeShader, AttributeIndicesParam, AttributeIndices.GetData(), AttributeIndices.Num());
@@ -331,12 +331,12 @@ struct FNiagaraDataInterfaceParametersCS_ParticleRead : public FNiagaraDataInter
 
 	FRHIShaderResourceView* GetIntSRVWithFallback(FRWBuffer& Buffer) const
 	{
-		return Buffer.SRV ? Buffer.SRV : FNiagaraRenderer::GetDummyIntBuffer().SRV;
+		return Buffer.SRV ? Buffer.SRV : FNiagaraRenderer::GetDummyIntBuffer();
 	}
 
 	FRHIShaderResourceView* GetFloatSRVWithFallback(FRWBuffer& Buffer) const
 	{
-		return Buffer.SRV ? Buffer.SRV : FNiagaraRenderer::GetDummyFloatBuffer().SRV;
+		return Buffer.SRV ? Buffer.SRV : FNiagaraRenderer::GetDummyFloatBuffer();
 	}
 
 	void Set(FRHICommandList& RHICmdList, const FNiagaraDataInterfaceSetArgs& Context) const
diff --git a/Engine/Plugins/FX/Niagara/Source/Niagara/Private/NiagaraDataInterfaceRW.cpp b/Engine/Plugins/FX/Niagara/Source/Niagara/Private/NiagaraDataInterfaceRW.cpp
--- a/Engine/Plugins/FX/Niagara/Source/Niagara/Private/NiagaraDataInterfaceRW.cpp
+++ b/Engine/Plugins/FX/Niagara/Source/Niagara/Private/NiagaraDataInterfaceRW.cpp
@@ -610,7 +610,7 @@ bool UNiagaraDataInterfaceGrid2D::GetFunctionHLSL(const FNiagaraDataInterfaceGPU
 		static const TCHAR *FormatSample = TEXT(R"(
 			void {FunctionName}(float2 In_Unit, out int Out_IndexX, out int Out_IndexY)
 			{
-				int2 Out_IndexTmp = round(In_Unit * {NumCellsName}  - .5);
+				int2 Out_IndexTmp = round(In_Unit * float2({NumCellsName})  - .5);
 				Out_IndexX = Out_IndexTmp.x;
 				Out_IndexY = Out_IndexTmp.y;				
 			}
@@ -624,7 +624,7 @@ bool UNiagaraDataInterfaceGrid2D::GetFunctionHLSL(const FNiagaraDataInterfaceGPU
 		static const TCHAR *FormatSample = TEXT(R"(
 			void {FunctionName}(float In_IndexX, float In_IndexY, out float3 Out_Unit)
 			{
-				Out_Unit = float3((float2(In_IndexX, In_IndexY) + .5) / {NumCellsName}, 0);
+				Out_Unit = float3((float2(In_IndexX, In_IndexY) + .5) / float2({NumCellsName}), 0);
 			}
 		)");
 
@@ -636,7 +636,7 @@ bool UNiagaraDataInterfaceGrid2D::GetFunctionHLSL(const FNiagaraDataInterfaceGPU
 		static const TCHAR *FormatSample = TEXT(R"(
 			void {FunctionName}(float In_IndexX, float In_IndexY, out float3 Out_Unit)
 			{
-				Out_Unit = float3((float2(In_IndexX, In_IndexY) + float2(0.0, 0.5)) / {NumCellsName}, 0);
+				Out_Unit = float3((float2(In_IndexX, In_IndexY) + float2(0.0, 0.5)) / float2({NumCellsName}), 0);
 			}
 		)");
 
@@ -648,7 +648,7 @@ bool UNiagaraDataInterfaceGrid2D::GetFunctionHLSL(const FNiagaraDataInterfaceGPU
 		static const TCHAR *FormatSample = TEXT(R"(
 			void {FunctionName}(float In_IndexX, float In_IndexY, out float3 Out_Unit)
 			{
-				Out_Unit = float3((float2(In_IndexX, In_IndexY) +  + float2(0.5, 0.0)) / {NumCellsName}, 0);
+				Out_Unit = float3((float2(In_IndexX, In_IndexY) +  + float2(0.5, 0.0)) / float2({NumCellsName}), 0);
 			}
 		)");
 		
diff --git a/Engine/Plugins/FX/Niagara/Source/Niagara/Private/NiagaraDataInterfaceSkeletalMesh.cpp b/Engine/Plugins/FX/Niagara/Source/Niagara/Private/NiagaraDataInterfaceSkeletalMesh.cpp
--- a/Engine/Plugins/FX/Niagara/Source/Niagara/Private/NiagaraDataInterfaceSkeletalMesh.cpp
+++ b/Engine/Plugins/FX/Niagara/Source/Niagara/Private/NiagaraDataInterfaceSkeletalMesh.cpp
@@ -896,7 +896,7 @@ public:
 			}
 			else
 			{
-				SetSRVParameter(RHICmdList, ComputeShaderRHI, MeshTexCoordBuffer, FNiagaraRenderer::GetDummyFloatBuffer().SRV);
+				SetSRVParameter(RHICmdList, ComputeShaderRHI, MeshTexCoordBuffer, FNiagaraRenderer::GetDummyFloatBuffer());
 			}
 			if (StaticBuffers->GetBufferColorSRV())
 			{
@@ -904,7 +904,7 @@ public:
 			}
 			else
 			{
-				SetSRVParameter(RHICmdList, ComputeShaderRHI, MeshColorBuffer, FNiagaraRenderer::GetDummyFloatBuffer().SRV);
+				SetSRVParameter(RHICmdList, ComputeShaderRHI, MeshColorBuffer, FNiagaraRenderer::GetDummyFloatBuffer());
 			}
 			SetShaderValue(RHICmdList, ComputeShaderRHI, MeshTriangleCount, StaticBuffers->GetTriangleCount());
 			SetShaderValue(RHICmdList, ComputeShaderRHI, MeshVertexCount, StaticBuffers->GetVertexCount());
@@ -915,8 +915,8 @@ public:
 			}
 			else
 			{
-				SetSRVParameter(RHICmdList, ComputeShaderRHI, MeshTriangleSamplerProbaBuffer, FNiagaraRenderer::GetDummyUIntBuffer().SRV);
-				SetSRVParameter(RHICmdList, ComputeShaderRHI, MeshTriangleSamplerAliasBuffer, FNiagaraRenderer::GetDummyUIntBuffer().SRV);
+				SetSRVParameter(RHICmdList, ComputeShaderRHI, MeshTriangleSamplerProbaBuffer, FNiagaraRenderer::GetDummyUIntBuffer());
+				SetSRVParameter(RHICmdList, ComputeShaderRHI, MeshTriangleSamplerAliasBuffer, FNiagaraRenderer::GetDummyUIntBuffer());
 			}
 
 			SetSRVParameter(RHICmdList, ComputeShaderRHI, MeshSkinWeightBuffer, InstanceData->MeshSkinWeightBufferSrv);
@@ -943,14 +943,14 @@ public:
 			else
 			{
 				SetShaderValue(RHICmdList, ComputeShaderRHI, MeshNumWeights, 0);
-				SetSRVParameter(RHICmdList, ComputeShaderRHI, MeshCurrBonesBuffer, FNiagaraRenderer::GetDummyFloat4Buffer().SRV);
-				SetSRVParameter(RHICmdList, ComputeShaderRHI, MeshPrevBonesBuffer, FNiagaraRenderer::GetDummyFloat4Buffer().SRV);
-				SetSRVParameter(RHICmdList, ComputeShaderRHI, MeshCurrSamplingBonesBuffer, FNiagaraRenderer::GetDummyFloat4Buffer().SRV);
-				SetSRVParameter(RHICmdList, ComputeShaderRHI, MeshPrevSamplingBonesBuffer, FNiagaraRenderer::GetDummyFloat4Buffer().SRV);
-				SetSRVParameter(RHICmdList, ComputeShaderRHI, MeshTriangleMatricesOffsetBuffer, FNiagaraRenderer::GetDummyUIntBuffer().SRV);
+				SetSRVParameter(RHICmdList, ComputeShaderRHI, MeshCurrBonesBuffer, FNiagaraRenderer::GetDummyFloat4Buffer());
+				SetSRVParameter(RHICmdList, ComputeShaderRHI, MeshPrevBonesBuffer, FNiagaraRenderer::GetDummyFloat4Buffer());
+				SetSRVParameter(RHICmdList, ComputeShaderRHI, MeshCurrSamplingBonesBuffer, FNiagaraRenderer::GetDummyFloat4Buffer());
+				SetSRVParameter(RHICmdList, ComputeShaderRHI, MeshPrevSamplingBonesBuffer, FNiagaraRenderer::GetDummyFloat4Buffer());
+				SetSRVParameter(RHICmdList, ComputeShaderRHI, MeshTriangleMatricesOffsetBuffer, FNiagaraRenderer::GetDummyUIntBuffer());
 			}
 
-			FRHIShaderResourceView* SpecificBonesSRV = StaticBuffers->GetNumSpecificBones() > 0 ? StaticBuffers->GetSpecificBonesSRV() : FNiagaraRenderer::GetDummyUIntBuffer().SRV.GetReference();
+			FRHIShaderResourceView* SpecificBonesSRV = StaticBuffers->GetNumSpecificBones() > 0 ? StaticBuffers->GetSpecificBonesSRV() : FNiagaraRenderer::GetDummyUIntBuffer();
 			SetShaderValue(RHICmdList, ComputeShaderRHI, NumSpecificBones, StaticBuffers->GetNumSpecificBones());
 			SetSRVParameter(RHICmdList, ComputeShaderRHI, SpecificBones, SpecificBonesSRV);
 
@@ -968,33 +968,33 @@ public:
 			// Bind dummy buffers
 			ensure(!InstanceData || InstanceData->StaticBuffers);
 
-			SetSRVParameter(RHICmdList, ComputeShaderRHI, MeshVertexBuffer, FNiagaraRenderer::GetDummyFloatBuffer().SRV);
-			SetSRVParameter(RHICmdList, ComputeShaderRHI, MeshIndexBuffer, FNiagaraRenderer::GetDummyUIntBuffer().SRV);
-			SetSRVParameter(RHICmdList, ComputeShaderRHI, MeshTangentBuffer, FNiagaraRenderer::GetDummyFloatBuffer().SRV);
+			SetSRVParameter(RHICmdList, ComputeShaderRHI, MeshVertexBuffer, FNiagaraRenderer::GetDummyFloatBuffer());
+			SetSRVParameter(RHICmdList, ComputeShaderRHI, MeshIndexBuffer, FNiagaraRenderer::GetDummyUIntBuffer());
+			SetSRVParameter(RHICmdList, ComputeShaderRHI, MeshTangentBuffer, FNiagaraRenderer::GetDummyFloatBuffer());
 
 			SetShaderValue(RHICmdList, ComputeShaderRHI, MeshNumTexCoord, 0);
-			SetSRVParameter(RHICmdList, ComputeShaderRHI, MeshTexCoordBuffer, FNiagaraRenderer::GetDummyFloatBuffer().SRV);
-			SetSRVParameter(RHICmdList, ComputeShaderRHI, MeshColorBuffer, FNiagaraRenderer::GetDummyFloatBuffer().SRV);
+			SetSRVParameter(RHICmdList, ComputeShaderRHI, MeshTexCoordBuffer, FNiagaraRenderer::GetDummyFloatBuffer());
+			SetSRVParameter(RHICmdList, ComputeShaderRHI, MeshColorBuffer, FNiagaraRenderer::GetDummyFloatBuffer());
 			SetShaderValue(RHICmdList, ComputeShaderRHI, MeshTriangleCount, 0);
 			SetShaderValue(RHICmdList, ComputeShaderRHI, MeshVertexCount, 0);
-			SetSRVParameter(RHICmdList, ComputeShaderRHI, MeshTriangleSamplerProbaBuffer, FNiagaraRenderer::GetDummyUIntBuffer().SRV);
-			SetSRVParameter(RHICmdList, ComputeShaderRHI, MeshTriangleSamplerAliasBuffer, FNiagaraRenderer::GetDummyUIntBuffer().SRV);
+			SetSRVParameter(RHICmdList, ComputeShaderRHI, MeshTriangleSamplerProbaBuffer, FNiagaraRenderer::GetDummyUIntBuffer());
+			SetSRVParameter(RHICmdList, ComputeShaderRHI, MeshTriangleSamplerAliasBuffer, FNiagaraRenderer::GetDummyUIntBuffer());
 
-			SetSRVParameter(RHICmdList, ComputeShaderRHI, MeshSkinWeightBuffer, FNiagaraRenderer::GetDummyUIntBuffer().SRV);
-			SetSRVParameter(RHICmdList, ComputeShaderRHI, MeshSkinWeightLookupBuffer, FNiagaraRenderer::GetDummyUIntBuffer().SRV);
+			SetSRVParameter(RHICmdList, ComputeShaderRHI, MeshSkinWeightBuffer, FNiagaraRenderer::GetDummyUIntBuffer());
+			SetSRVParameter(RHICmdList, ComputeShaderRHI, MeshSkinWeightLookupBuffer, FNiagaraRenderer::GetDummyUIntBuffer());
 
 			SetShaderValue(RHICmdList, ComputeShaderRHI, MeshWeightStride, 0);
 			SetShaderValue(RHICmdList, ComputeShaderRHI, MeshSkinWeightIndexSize, 0);
 			SetShaderValue(RHICmdList, ComputeShaderRHI, MeshNumWeights, 0);
 
-			SetSRVParameter(RHICmdList, ComputeShaderRHI, MeshCurrBonesBuffer, FNiagaraRenderer::GetDummyFloat4Buffer().SRV);
-			SetSRVParameter(RHICmdList, ComputeShaderRHI, MeshPrevBonesBuffer, FNiagaraRenderer::GetDummyFloat4Buffer().SRV);
-			SetSRVParameter(RHICmdList, ComputeShaderRHI, MeshCurrSamplingBonesBuffer, FNiagaraRenderer::GetDummyFloat4Buffer().SRV);
-			SetSRVParameter(RHICmdList, ComputeShaderRHI, MeshPrevSamplingBonesBuffer, FNiagaraRenderer::GetDummyFloat4Buffer().SRV);
-			SetSRVParameter(RHICmdList, ComputeShaderRHI, MeshTriangleMatricesOffsetBuffer, FNiagaraRenderer::GetDummyUIntBuffer().SRV);
+			SetSRVParameter(RHICmdList, ComputeShaderRHI, MeshCurrBonesBuffer, FNiagaraRenderer::GetDummyFloat4Buffer());
+			SetSRVParameter(RHICmdList, ComputeShaderRHI, MeshPrevBonesBuffer, FNiagaraRenderer::GetDummyFloat4Buffer());
+			SetSRVParameter(RHICmdList, ComputeShaderRHI, MeshCurrSamplingBonesBuffer, FNiagaraRenderer::GetDummyFloat4Buffer());
+			SetSRVParameter(RHICmdList, ComputeShaderRHI, MeshPrevSamplingBonesBuffer, FNiagaraRenderer::GetDummyFloat4Buffer());
+			SetSRVParameter(RHICmdList, ComputeShaderRHI, MeshTriangleMatricesOffsetBuffer, FNiagaraRenderer::GetDummyUIntBuffer());
 
 			SetShaderValue(RHICmdList, ComputeShaderRHI, NumSpecificBones, 0);
-			SetSRVParameter(RHICmdList, ComputeShaderRHI, SpecificBones, FNiagaraRenderer::GetDummyUIntBuffer().SRV);
+			SetSRVParameter(RHICmdList, ComputeShaderRHI, SpecificBones, FNiagaraRenderer::GetDummyUIntBuffer());
 			SetShaderValue(RHICmdList, ComputeShaderRHI, NumSpecificSockets, 0);
 			SetShaderValue(RHICmdList, ComputeShaderRHI, SpecificSocketBoneOffset, 0);
 
diff --git a/Engine/Plugins/FX/Niagara/Source/Niagara/Private/NiagaraDataInterfaceStaticMesh.cpp b/Engine/Plugins/FX/Niagara/Source/Niagara/Private/NiagaraDataInterfaceStaticMesh.cpp
--- a/Engine/Plugins/FX/Niagara/Source/Niagara/Private/NiagaraDataInterfaceStaticMesh.cpp
+++ b/Engine/Plugins/FX/Niagara/Source/Niagara/Private/NiagaraDataInterfaceStaticMesh.cpp
@@ -533,7 +533,7 @@ public:
 				}
 				else
 				{
-					SetSRVParameter(RHICmdList, ComputeShaderRHI, MeshTexCoordBuffer, FNiagaraRenderer::GetDummyFloatBuffer().SRV);
+					SetSRVParameter(RHICmdList, ComputeShaderRHI, MeshTexCoordBuffer, FNiagaraRenderer::GetDummyFloatBuffer());
 				}
 
 				if(SpawnBuffer->GetBufferColorSRV())
@@ -542,7 +542,7 @@ public:
 				}
 				else
 				{
-					SetSRVParameter(RHICmdList, ComputeShaderRHI, MeshColorBuffer, FNiagaraRenderer::GetDummyFloatBuffer().SRV);
+					SetSRVParameter(RHICmdList, ComputeShaderRHI, MeshColorBuffer, FNiagaraRenderer::GetDummyFloatBuffer());
 				}
 
 				check(SpawnBuffer); // should always be allocated, we always need the GPU buffer for a GpuSimulation.
@@ -554,7 +554,7 @@ public:
 				}
 				else
 				{
-					SetSRVParameter(RHICmdList, ComputeShaderRHI, MeshTriangleBuffer, FNiagaraRenderer::GetDummyFloatBuffer().SRV);
+					SetSRVParameter(RHICmdList, ComputeShaderRHI, MeshTriangleBuffer, FNiagaraRenderer::GetDummyFloatBuffer());
 				}
 
 				const float InvDeltaTime = Data->DeltaSeconds > 0.0f ? 1.0f / Data->DeltaSeconds : 0.0f;
@@ -569,17 +569,17 @@ public:
 			}
 			else
 			{
-				SetSRVParameter(RHICmdList, ComputeShaderRHI, MeshVertexBuffer, FNiagaraRenderer::GetDummyFloatBuffer().SRV);
-				SetSRVParameter(RHICmdList, ComputeShaderRHI, MeshTangentBuffer, FNiagaraRenderer::GetDummyFloatBuffer().SRV);
-				SetSRVParameter(RHICmdList, ComputeShaderRHI, MeshIndexBuffer, FNiagaraRenderer::GetDummyUIntBuffer().SRV);
+				SetSRVParameter(RHICmdList, ComputeShaderRHI, MeshVertexBuffer, FNiagaraRenderer::GetDummyFloatBuffer());
+				SetSRVParameter(RHICmdList, ComputeShaderRHI, MeshTangentBuffer, FNiagaraRenderer::GetDummyFloatBuffer());
+				SetSRVParameter(RHICmdList, ComputeShaderRHI, MeshIndexBuffer, FNiagaraRenderer::GetDummyUIntBuffer());
 
 				SetShaderValue(RHICmdList, ComputeShaderRHI, NumTexCoord, 0);
-				SetSRVParameter(RHICmdList, ComputeShaderRHI, MeshTexCoordBuffer, FNiagaraRenderer::GetDummyFloatBuffer().SRV);
-				SetSRVParameter(RHICmdList, ComputeShaderRHI, MeshColorBuffer, FNiagaraRenderer::GetDummyFloatBuffer().SRV);
+				SetSRVParameter(RHICmdList, ComputeShaderRHI, MeshTexCoordBuffer, FNiagaraRenderer::GetDummyFloatBuffer());
+				SetSRVParameter(RHICmdList, ComputeShaderRHI, MeshColorBuffer, FNiagaraRenderer::GetDummyFloatBuffer());
 
 				SetShaderValue(RHICmdList, ComputeShaderRHI, SectionCount, 0);
-				SetSRVParameter(RHICmdList, ComputeShaderRHI, MeshSectionBuffer, FNiagaraRenderer::GetDummyUIntBuffer().SRV);
-				SetSRVParameter(RHICmdList, ComputeShaderRHI, MeshTriangleBuffer, FNiagaraRenderer::GetDummyUIntBuffer().SRV);
+				SetSRVParameter(RHICmdList, ComputeShaderRHI, MeshSectionBuffer, FNiagaraRenderer::GetDummyUIntBuffer());
+				SetSRVParameter(RHICmdList, ComputeShaderRHI, MeshTriangleBuffer, FNiagaraRenderer::GetDummyUIntBuffer());
 
 				SetShaderValue(RHICmdList, ComputeShaderRHI, InstanceTransform, FMatrix::Identity);
 				SetShaderValue(RHICmdList, ComputeShaderRHI, InstanceTransformInverseTransposed, FMatrix::Identity);
diff --git a/Engine/Plugins/FX/Niagara/Source/Niagara/Private/NiagaraDataSet.cpp b/Engine/Plugins/FX/Niagara/Source/Niagara/Private/NiagaraDataSet.cpp
--- a/Engine/Plugins/FX/Niagara/Source/Niagara/Private/NiagaraDataSet.cpp
+++ b/Engine/Plugins/FX/Niagara/Source/Niagara/Private/NiagaraDataSet.cpp
@@ -373,15 +373,23 @@ void FNiagaraDataSet::AllocateGPUFreeIDs(uint32 InNumInstances, FRHICommandList&
 	FRWBuffer NewFreeIDsBuffer;
 	NewFreeIDsBuffer.Initialize(sizeof(int32), NumIDsToAlloc, EPixelFormat::PF_R32_SINT, BUF_Static, DebugBufferName);
 
-	// We must maintain the existing list of free IDs.
-	FRWBuffer& ExistingBuffer = (GPUNumAllocatedIDs > 0) ? GPUFreeIDs : FNiagaraRenderer::GetDummyIntBuffer();
+	FRHIShaderResourceView* ExistingBuffer;
+	if (GPUNumAllocatedIDs > 0)
+	{
+		// We must maintain the existing list of free IDs.
+		// The free IDs buffer was written in the previous simulation step, but hasn't been transitioned to read yet, so we must
+		// transition it explicitly here. The new buffer will be transitioned by NiagaraEmitterInstanceBatcher::DispatchAllOnCompute(),
+		// so there's no need for a barrier at the end of this function.
+		RHICmdList.TransitionResource(EResourceTransitionAccess::EReadable, EResourceTransitionPipeline::EComputeToCompute, GPUFreeIDs.UAV);
+		ExistingBuffer = GPUFreeIDs.SRV;
+	}
+	else
+	{
+		ExistingBuffer = FNiagaraRenderer::GetDummyIntBuffer();
+	}
 
-	// The free IDs buffer was written in the previous simulation step, but hasn't been transitioned to read yet, so we must
-	// transition it explicitly here. The new buffer will be transitioned by NiagaraEmitterInstanceBatcher::DispatchAllOnCompute(),
-	// so there's no need for a barrier at the end of this function.
-	RHICmdList.TransitionResource(EResourceTransitionAccess::EReadable, EResourceTransitionPipeline::EComputeToCompute, ExistingBuffer.UAV);
 	RHICmdList.TransitionResource(EResourceTransitionAccess::EWritable, EResourceTransitionPipeline::EComputeToCompute, NewFreeIDsBuffer.UAV);
-	NiagaraInitGPUFreeIDList(RHICmdList, FeatureLevel, NumIDsToAlloc, NewFreeIDsBuffer, GPUNumAllocatedIDs, ExistingBuffer.SRV);
+	NiagaraInitGPUFreeIDList(RHICmdList, FeatureLevel, NumIDsToAlloc, NewFreeIDsBuffer, GPUNumAllocatedIDs, ExistingBuffer);
 
 	GPUFreeIDs = MoveTemp(NewFreeIDsBuffer);
 	GPUNumAllocatedIDs = NumIDsToAlloc;
@@ -985,8 +993,8 @@ void FNiagaraDataBuffer::SetShaderParams(FNiagaraShader *Shader, FRHICommandList
 	{
 		const bool InstancesAllocated = GetNumInstancesAllocated() > 0;
 
-		SetSRVParameter(CommandList, ComputeShader, Shader->FloatInputBufferParam, InstancesAllocated ? GetGPUBufferFloat().SRV : FNiagaraRenderer::GetDummyFloatBuffer().SRV);
-		SetSRVParameter(CommandList, ComputeShader, Shader->IntInputBufferParam, InstancesAllocated ? GetGPUBufferInt().SRV : FNiagaraRenderer::GetDummyIntBuffer().SRV);
+		SetSRVParameter(CommandList, ComputeShader, Shader->FloatInputBufferParam, InstancesAllocated ? GetGPUBufferFloat().SRV : FNiagaraRenderer::GetDummyFloatBuffer());
+		SetSRVParameter(CommandList, ComputeShader, Shader->IntInputBufferParam, InstancesAllocated ? GetGPUBufferInt().SRV : FNiagaraRenderer::GetDummyIntBuffer());
 		SetShaderValue(CommandList, ComputeShader, Shader->ComponentBufferSizeReadParam, SafeBufferSize);
 	}
 	else
diff --git a/Engine/Plugins/FX/Niagara/Source/Niagara/Private/NiagaraEmitterInstanceBatcher.cpp b/Engine/Plugins/FX/Niagara/Source/Niagara/Private/NiagaraEmitterInstanceBatcher.cpp
--- a/Engine/Plugins/FX/Niagara/Source/Niagara/Private/NiagaraEmitterInstanceBatcher.cpp
+++ b/Engine/Plugins/FX/Niagara/Source/Niagara/Private/NiagaraEmitterInstanceBatcher.cpp
@@ -1295,21 +1295,28 @@ void NiagaraEmitterInstanceBatcher::Run(const FNiagaraGPUSystemTick& Tick, const
 	// set the shader and data set params 
 	//
 	const bool bRequiresPersistentIDs = Context->MainDataSet->RequiresPersistentIDs();
-	SetSRVParameter(RHICmdList, Shader.GetComputeShader(), Shader->FreeIDBufferParam, bRequiresPersistentIDs ? Context->MainDataSet->GetGPUFreeIDs().SRV : FNiagaraRenderer::GetDummyIntBuffer().SRV);
+	SetSRVParameter(RHICmdList, Shader.GetComputeShader(), Shader->FreeIDBufferParam, bRequiresPersistentIDs ? Context->MainDataSet->GetGPUFreeIDs().SRV : FNiagaraRenderer::GetDummyIntBuffer());
 	CurrentData.SetShaderParams(Shader.GetShader(), RHICmdList, true);
 	DestinationData.SetShaderParams(Shader.GetShader(), RHICmdList, false);
 
-	// set the index buffer uav
+	// set the instance count uav
 	//
-	if (Shader->InstanceCountsParam.IsBound() && !IterationInterface )
+	if (Shader->InstanceCountsParam.IsBound())
 	{
-		RHICmdList.TransitionResource(EResourceTransitionAccess::ERWNoBarrier, EResourceTransitionPipeline::EComputeToCompute, GPUInstanceCounterManager.GetInstanceCountBuffer().UAV);
-		Shader->InstanceCountsParam.SetBuffer(RHICmdList, ComputeShader, GPUInstanceCounterManager.GetInstanceCountBuffer());
-		const uint32 ReadOffset = (Tick.bNeedsReset && SimulationStageIndex == 0) ? INDEX_NONE : Instance->SimStageData[SimulationStageIndex].SourceCountOffset;
-		const uint32 WriteOffset = Instance->SimStageData[SimulationStageIndex].DestinationCountOffset;
-		//UE_LOG(LogNiagara, Log, TEXT("Instance count setup R: %d W: %d reset? %s %d"), ReadOffset, WriteOffset, Tick.bNeedsReset ? TEXT("T") : TEXT("F"), CurrentData.GetGPUInstanceCountBufferOffset());
-		SetShaderValue(RHICmdList, ComputeShader, Shader->ReadInstanceCountOffsetParam, ReadOffset);
-		SetShaderValue(RHICmdList, ComputeShader, Shader->WriteInstanceCountOffsetParam, WriteOffset);
+		if (IterationInterface)
+		{
+			RHICmdList.SetUAVParameter(ComputeShader, Shader->InstanceCountsParam.GetUAVIndex(), GetEmptyRWBufferFromPool(RHICmdList, PF_R32_UINT));
+		}
+		else
+		{
+			RHICmdList.TransitionResource(EResourceTransitionAccess::ERWNoBarrier, EResourceTransitionPipeline::EComputeToCompute, GPUInstanceCounterManager.GetInstanceCountBuffer().UAV);
+			Shader->InstanceCountsParam.SetBuffer(RHICmdList, ComputeShader, GPUInstanceCounterManager.GetInstanceCountBuffer());
+			const uint32 ReadOffset = (Tick.bNeedsReset && SimulationStageIndex == 0) ? INDEX_NONE : Instance->SimStageData[SimulationStageIndex].SourceCountOffset;
+			const uint32 WriteOffset = Instance->SimStageData[SimulationStageIndex].DestinationCountOffset;
+			//UE_LOG(LogNiagara, Log, TEXT("Instance count setup R: %d W: %d reset? %s %d"), ReadOffset, WriteOffset, Tick.bNeedsReset ? TEXT("T") : TEXT("F"), CurrentData.GetGPUInstanceCountBufferOffset());
+			SetShaderValue(RHICmdList, ComputeShader, Shader->ReadInstanceCountOffsetParam, ReadOffset);
+			SetShaderValue(RHICmdList, ComputeShader, Shader->WriteInstanceCountOffsetParam, WriteOffset);
+		}
 	}
 
 	// set the execution parameters
@@ -1430,9 +1437,77 @@ void NiagaraEmitterInstanceBatcher::Run(const FNiagaraGPUSystemTick& Tick, const
 	CurrentData.UnsetShaderParams(Shader.GetShader(), RHICmdList);
 	DestinationData.UnsetShaderParams(Shader.GetShader(), RHICmdList);
 	Shader->InstanceCountsParam.UnsetUAV(RHICmdList, ComputeShader);
+
+	ResetEmptyUAVPools(RHICmdList);
 }
 
 FGPUSortManager* NiagaraEmitterInstanceBatcher::GetGPUSortManager() const
 {
 	return GPUSortManager;
 }
+
+NiagaraEmitterInstanceBatcher::DummyUAV::~DummyUAV()
+{
+	UAV.SafeRelease();
+	Buffer.SafeRelease();
+	Texture.SafeRelease();
+}
+
+void NiagaraEmitterInstanceBatcher::DummyUAV::Init(FRHICommandList& RHICmdList, EPixelFormat Format, bool IsTexture, const TCHAR* DebugName)
+{
+	checkSlow(IsInRenderingThread());
+
+	FRHIResourceCreateInfo CreateInfo;
+	CreateInfo.DebugName = DebugName;
+
+	if (IsTexture)
+	{
+		Texture = RHICreateTexture2D(1, 1, Format, 1, 1, TexCreate_ShaderResource | TexCreate_UAV, CreateInfo);
+		UAV = RHICreateUnorderedAccessView(Texture, 0);
+	}
+	else
+	{
+		uint32 BytesPerElement = GPixelFormats[Format].BlockBytes;
+		Buffer = RHICreateVertexBuffer(BytesPerElement, BUF_UnorderedAccess | BUF_ShaderResource, CreateInfo);
+		UAV = RHICreateUnorderedAccessView(Buffer, Format);
+	}
+
+	RHICmdList.TransitionResource(EResourceTransitionAccess::EWritable, EResourceTransitionPipeline::EComputeToCompute, UAV);
+}
+
+FRHIUnorderedAccessView* NiagaraEmitterInstanceBatcher::GetEmptyUAVFromPool(FRHICommandList& RHICmdList, EPixelFormat Format, bool IsTexture) const
+{
+	TMap<EPixelFormat, DummyUAVPool>& UAVMap = IsTexture ? DummyTexturePool : DummyBufferPool;
+	DummyUAVPool& Pool = UAVMap.FindOrAdd(Format);
+	checkSlow(Pool.NextFreeIndex <= Pool.UAVs.Num());
+	if (Pool.NextFreeIndex == Pool.UAVs.Num())
+	{
+		DummyUAV& NewUAV = Pool.UAVs.AddDefaulted_GetRef();
+		NewUAV.Init(RHICmdList, Format, IsTexture, TEXT("NiagaraEmitterInstanceBatcher::DummyUAV"));
+	}
+
+	FRHIUnorderedAccessView* UAV = Pool.UAVs[Pool.NextFreeIndex].UAV;
+	++Pool.NextFreeIndex;
+	return UAV;
+}
+
+void NiagaraEmitterInstanceBatcher::ResetEmptyUAVPool(TMap<EPixelFormat, DummyUAVPool>& UAVMap, TArray<FRHIUnorderedAccessView*>& Transitions)
+{
+	for (TPair<EPixelFormat, DummyUAVPool>& Entry : UAVMap)
+	{
+		for (int UsedIdx = 0; UsedIdx < Entry.Value.NextFreeIndex; ++UsedIdx)
+		{
+			Transitions.Add(Entry.Value.UAVs[UsedIdx].UAV);
+		}
+		Entry.Value.NextFreeIndex = 0;
+	}
+}
+
+void NiagaraEmitterInstanceBatcher::ResetEmptyUAVPools(FRHICommandList& RHICmdList)
+{
+	TArray<FRHIUnorderedAccessView*> Transitions;
+	Transitions.Reserve(32);
+	ResetEmptyUAVPool(DummyBufferPool, Transitions);
+	ResetEmptyUAVPool(DummyTexturePool, Transitions);
+	RHICmdList.TransitionResources(EResourceTransitionAccess::ERWNoBarrier, EResourceTransitionPipeline::EComputeToCompute, Transitions.GetData(), Transitions.Num());
+}
diff --git a/Engine/Plugins/FX/Niagara/Source/Niagara/Private/NiagaraRenderer.cpp b/Engine/Plugins/FX/Niagara/Source/Niagara/Private/NiagaraRenderer.cpp
--- a/Engine/Plugins/FX/Niagara/Source/Niagara/Private/NiagaraRenderer.cpp
+++ b/Engine/Plugins/FX/Niagara/Source/Niagara/Private/NiagaraRenderer.cpp
@@ -26,112 +26,90 @@ static FAutoConsoleVariableRef CVarNiagaraRadixSortThreshold(
 	ECVF_Default
 );
 
-class FNiagaraDummyRWBufferFloat : public FRenderResource
+class FNiagaraEmptyBufferSRV : public FRenderResource
 {
 public:
-	FNiagaraDummyRWBufferFloat(const FString InDebugId) : DebugId(InDebugId) {}
-	FString DebugId;
-	FRWBuffer Buffer;
+	FNiagaraEmptyBufferSRV(EPixelFormat InPixelFormat, const FString& InDebugName) : PixelFormat(InPixelFormat), DebugName(InDebugName) {}
+	EPixelFormat PixelFormat;
+	FString DebugName;
+	FVertexBufferRHIRef Buffer;
+	FShaderResourceViewRHIRef SRV;
 
 	virtual void InitRHI() override
 	{
-		UE_LOG(LogNiagara, Log, TEXT("FNiagaraDummyRWBufferFloat InitRHI %s"), *DebugId);
-		Buffer.Initialize(sizeof(float), 1, EPixelFormat::PF_R32_FLOAT, BUF_Static, *DebugId);
+		uint32 BytesPerElement = GPixelFormats[PixelFormat].BlockBytes;
+		uint32 NumElements = 1;
+		uint32 NumBytes = BytesPerElement * NumElements;
+		FRHIResourceCreateInfo CreateInfo;
+		CreateInfo.DebugName = *DebugName;
+		Buffer = RHICreateVertexBuffer(NumBytes, BUF_ShaderResource | BUF_Static, CreateInfo);
+		SRV = RHICreateShaderResourceView(Buffer, BytesPerElement, PixelFormat);
 	}
 
 	virtual void ReleaseRHI() override
 	{
-		UE_LOG(LogNiagara, Log, TEXT("FNiagaraDummyRWBufferFloat ReleaseRHI %s"), *DebugId);
-		Buffer.Release();
+		SRV.SafeRelease();
+		Buffer.SafeRelease();
 	}
 };
 
-class FNiagaraDummyRWBufferFloat4 : public FRenderResource
+class FNiagaraEmptyTextureSRV : public FRenderResource
 {
 public:
-	FNiagaraDummyRWBufferFloat4(const FString InDebugId) : DebugId(InDebugId) {}
-	FString DebugId;
-	FRWBuffer Buffer;
+	FNiagaraEmptyTextureSRV(EPixelFormat InPixelFormat, const FString& InDebugName) : PixelFormat(InPixelFormat), DebugName(InDebugName) {}
+	EPixelFormat PixelFormat;
+	FString DebugName;
+	FTexture2DRHIRef Buffer;
+	FShaderResourceViewRHIRef SRV;
 
 	virtual void InitRHI() override
 	{
-		UE_LOG(LogNiagara, Log, TEXT("FNiagaraDummyRWBufferFloat InitRHI %s"), *DebugId);
-		Buffer.Initialize(sizeof(float) * 4, 1, EPixelFormat::PF_A32B32G32R32F, BUF_Static, *DebugId);
+		FRHIResourceCreateInfo CreateInfo;
+		CreateInfo.DebugName = *DebugName;
+		Buffer = RHICreateTexture2D(1, 1, PixelFormat, 1, 1, TexCreate_ShaderResource, CreateInfo);
+		SRV = RHICreateShaderResourceView(Buffer, 0);
 	}
 
 	virtual void ReleaseRHI() override
 	{
-		UE_LOG(LogNiagara, Log, TEXT("FNiagaraDummyRWBufferFloat4 ReleaseRHI %s"), *DebugId);
-		Buffer.Release();
+		SRV.SafeRelease();
+		Buffer.SafeRelease();
 	}
 };
 
-class FNiagaraDummyRWBufferInt : public FRenderResource
+FRHIShaderResourceView* FNiagaraRenderer::GetDummyFloatBuffer()
 {
-public:
-	FNiagaraDummyRWBufferInt(const FString InDebugId) : DebugId(InDebugId) {}
-	FString DebugId;
-	FRWBuffer Buffer;
-
-	virtual void InitRHI() override
-	{
-		UE_LOG(LogNiagara, Log, TEXT("FNiagaraDummyRWBufferInt InitRHI %s"), *DebugId);
-		Buffer.Initialize(sizeof(int32), 1, EPixelFormat::PF_R32_SINT, BUF_Static, *DebugId);
-	}
-
-	virtual void ReleaseRHI() override
-	{
-		UE_LOG(LogNiagara, Log, TEXT("FNiagaraDummyRWBufferInt ReleaseRHI %s"), *DebugId);
-		Buffer.Release();
-	}
-};
-
-class FNiagaraDummyRWBufferUInt : public FRenderResource
-{
-public:
-	FNiagaraDummyRWBufferUInt(const FString InDebugId) : DebugId(InDebugId) {}
-	FString DebugId;
-	FRWBuffer Buffer;
-
-	virtual void InitRHI() override
-	{
-		UE_LOG(LogNiagara, Log, TEXT("FNiagaraDummyRWBufferUInt InitRHI %s"), *DebugId);
-		Buffer.Initialize(sizeof(uint32), 1, EPixelFormat::PF_R32_UINT, BUF_Static, *DebugId);
-	}
-
-	virtual void ReleaseRHI() override
-	{
-		UE_LOG(LogNiagara, Log, TEXT("FNiagaraDummyRWBufferUInt ReleaseRHI %s"), *DebugId);
-		Buffer.Release();
-	}
-};
+	check(IsInRenderingThread());
+	static TGlobalResource<FNiagaraEmptyBufferSRV> DummyFloatBuffer(PF_R32_FLOAT, TEXT("NiagaraRenderer::DummyFloat"));
+	return DummyFloatBuffer.SRV;
+}
 
-FRWBuffer& FNiagaraRenderer::GetDummyFloatBuffer()
+FRHIShaderResourceView* FNiagaraRenderer::GetDummyFloat4Buffer()
 {
 	check(IsInRenderingThread());
-	static TGlobalResource<FNiagaraDummyRWBufferFloat> DummyFloatBuffer(TEXT("NiagaraRenderer::DummyFloat"));
-	return DummyFloatBuffer.Buffer;
+	static TGlobalResource<FNiagaraEmptyBufferSRV> DummyFloat4Buffer(PF_A32B32G32R32F, TEXT("NiagaraRenderer::DummyFloat4"));
+	return DummyFloat4Buffer.SRV;
 }
 
-FRWBuffer& FNiagaraRenderer::GetDummyFloat4Buffer()
+FRHIShaderResourceView* FNiagaraRenderer::GetDummyIntBuffer()
 {
 	check(IsInRenderingThread());
-	static TGlobalResource<FNiagaraDummyRWBufferFloat4> GetDummyFloat4Buffer(TEXT("NiagaraRenderer::DummyFloat4"));
-	return GetDummyFloat4Buffer.Buffer;
+	static TGlobalResource<FNiagaraEmptyBufferSRV> DummyIntBuffer(PF_R32_SINT, TEXT("NiagaraRenderer::DummyInt"));
+	return DummyIntBuffer.SRV;
 }
 
-FRWBuffer& FNiagaraRenderer::GetDummyIntBuffer()
+FRHIShaderResourceView* FNiagaraRenderer::GetDummyUIntBuffer()
 {
 	check(IsInRenderingThread());
-	static TGlobalResource<FNiagaraDummyRWBufferInt> DummyIntBuffer(TEXT("NiagaraRenderer::DummyInt"));
-	return DummyIntBuffer.Buffer;
+	static TGlobalResource<FNiagaraEmptyBufferSRV> DummyUIntBuffer(PF_R32_UINT, TEXT("NiagaraRenderer::DummyUInt"));
+	return DummyUIntBuffer.SRV;
 }
 
-FRWBuffer& FNiagaraRenderer::GetDummyUIntBuffer()
+FRHIShaderResourceView* FNiagaraRenderer::GetDummyTextureReadBuffer2D()
 {
 	check(IsInRenderingThread());
-	static TGlobalResource<FNiagaraDummyRWBufferUInt> DummyUIntBuffer(TEXT("NiagaraRenderer::DummyUInt"));
-	return DummyUIntBuffer.Buffer;
+	static TGlobalResource<FNiagaraEmptyTextureSRV> DummyTextureReadBuffer2D(PF_R32_FLOAT, TEXT("NiagaraRenderer::DummyTextureReadBuffer2D"));
+	return DummyTextureReadBuffer2D.SRV;
 }
 
 bool FNiagaraRenderer::SetVertexFactoryVariable(const FNiagaraDataSet& DataSet, const FNiagaraVariable& Var, int32 VFVarOffset)
diff --git a/Engine/Plugins/FX/Niagara/Source/Niagara/Private/NiagaraRendererMeshes.cpp b/Engine/Plugins/FX/Niagara/Source/Niagara/Private/NiagaraRendererMeshes.cpp
--- a/Engine/Plugins/FX/Niagara/Source/Niagara/Private/NiagaraRendererMeshes.cpp
+++ b/Engine/Plugins/FX/Niagara/Source/Niagara/Private/NiagaraRendererMeshes.cpp
@@ -449,7 +449,7 @@ void FNiagaraRendererMeshes::GetDynamicMeshElements(const TArray<const FSceneVie
 					}
 					else
 					{
-						CollectorResources.VertexFactory->SetParticleData(FNiagaraRenderer::GetDummyFloatBuffer().SRV, 0);
+						CollectorResources.VertexFactory->SetParticleData(FNiagaraRenderer::GetDummyFloatBuffer(), 0);
 					}
 				}
 
diff --git a/Engine/Plugins/FX/Niagara/Source/Niagara/Private/NiagaraRendererSprites.cpp b/Engine/Plugins/FX/Niagara/Source/Niagara/Private/NiagaraRendererSprites.cpp
--- a/Engine/Plugins/FX/Niagara/Source/Niagara/Private/NiagaraRendererSprites.cpp
+++ b/Engine/Plugins/FX/Niagara/Source/Niagara/Private/NiagaraRendererSprites.cpp
@@ -414,7 +414,7 @@ void FNiagaraRendererSprites::SetVertexFactoryParticleData(
 			}
 			else
 			{
-				OutVertexFactory.SetParticleData(FNiagaraRenderer::GetDummyFloatBuffer().SRV, 0);
+				OutVertexFactory.SetParticleData(FNiagaraRenderer::GetDummyFloatBuffer(), 0);
 			}
 		}
 	}
diff --git a/Engine/Plugins/FX/Niagara/Source/Niagara/Public/NiagaraRenderer.h b/Engine/Plugins/FX/Niagara/Source/Niagara/Public/NiagaraRenderer.h
--- a/Engine/Plugins/FX/Niagara/Source/Niagara/Public/NiagaraRenderer.h
+++ b/Engine/Plugins/FX/Niagara/Source/Niagara/Public/NiagaraRenderer.h
@@ -115,11 +115,12 @@ public:
 	virtual void GetDynamicRayTracingInstances(FRayTracingMaterialGatheringContext& Context, TArray<FRayTracingInstance>& OutRayTracingInstances, const FNiagaraSceneProxy* Proxy) {}
 #endif
 
-	NIAGARA_API static FRWBuffer& GetDummyFloatBuffer();
-	NIAGARA_API static FRWBuffer& GetDummyFloat4Buffer();
-	NIAGARA_API static FRWBuffer& GetDummyIntBuffer();
-	NIAGARA_API static FRWBuffer& GetDummyUIntBuffer();
-	
+	NIAGARA_API static FRHIShaderResourceView* GetDummyFloatBuffer();
+	NIAGARA_API static FRHIShaderResourceView* GetDummyFloat4Buffer();
+	NIAGARA_API static FRHIShaderResourceView* GetDummyIntBuffer();
+	NIAGARA_API static FRHIShaderResourceView* GetDummyUIntBuffer();
+	NIAGARA_API static FRHIShaderResourceView* GetDummyTextureReadBuffer2D();
+
 	FORCEINLINE ENiagaraSimTarget GetSimTarget() const { return SimTarget; }
 
 protected:
diff --git a/Engine/Plugins/FX/Niagara/Source/NiagaraCore/Private/NiagaraCustomVersion.cpp b/Engine/Plugins/FX/Niagara/Source/NiagaraCore/Private/NiagaraCustomVersion.cpp
--- a/Engine/Plugins/FX/Niagara/Source/NiagaraCore/Private/NiagaraCustomVersion.cpp
+++ b/Engine/Plugins/FX/Niagara/Source/NiagaraCore/Private/NiagaraCustomVersion.cpp
@@ -8,4 +8,4 @@ const FGuid FNiagaraCustomVersion::GUID(0xFCF57AFA, 0x50764283, 0xB9A9E658, 0xFF
 // Register the custom version with core
 FCustomVersionRegistration GRegisterNiagaraCustomVersion(FNiagaraCustomVersion::GUID, FNiagaraCustomVersion::LatestVersion, TEXT("NiagaraVer"));
 
-const FGuid FNiagaraCustomVersion::LatestScriptCompileVersion(0xDB33E371, 0xCE1C40FD, 0xBA714031, 0xE8DACCB4);
+const FGuid FNiagaraCustomVersion::LatestScriptCompileVersion(0XC9A39220, 0X1AE54A45, 0X8A01CC89, 0X64E91B28);
