Forgot to submit one file in CL 11589537...

diff --git a/Engine/Plugins/FX/Niagara/Source/Niagara/Classes/NiagaraEmitterInstanceBatcher.h b/Engine/Plugins/FX/Niagara/Source/Niagara/Classes/NiagaraEmitterInstanceBatcher.h
--- a/Engine/Plugins/FX/Niagara/Source/Niagara/Classes/NiagaraEmitterInstanceBatcher.h
+++ b/Engine/Plugins/FX/Niagara/Source/Niagara/Classes/NiagaraEmitterInstanceBatcher.h
@@ -145,6 +145,9 @@ public:
 	/** Loop over all the data interfaces and call the poststage methods */
 	void PostStageInterface(const FNiagaraGPUSystemTick& Tick, FNiagaraComputeInstanceData *Instance, FRHICommandList &RHICmdList, const FNiagaraShaderRef& ComputeShader, const uint32 SimulationStageIndex) const;
 
+	FRHIUnorderedAccessView* GetEmptyRWBufferFromPool(FRHICommandList& RHICmdList, EPixelFormat Format) const { return GetEmptyUAVFromPool(RHICmdList, Format, false); }
+	FRHIUnorderedAccessView* GetEmptyRWTextureFromPool(FRHICommandList& RHICmdList, EPixelFormat Format) const { return GetEmptyUAVFromPool(RHICmdList, Format, true); }
+
 	/** Get the shared SortManager, used in the rendering loop to call FGPUSortManager::OnPreRender() and FGPUSortManager::OnPostRenderOpaque() */
 	virtual FGPUSortManager* GetGPUSortManager() const override;
 
@@ -224,4 +227,27 @@ private:
 
 	/** List of emitter instances which need their free ID buffers updated post render. */
 	FEmitterInstanceList DeferredIDBufferUpdates;
+
+	struct DummyUAV
+	{
+		FVertexBufferRHIRef Buffer;
+		FTexture2DRHIRef Texture;
+		FUnorderedAccessViewRHIRef UAV;
+
+		~DummyUAV();
+		void Init(FRHICommandList& RHICmdList, EPixelFormat Format, bool IsTexture, const TCHAR* DebugName);
+	};
+
+	struct DummyUAVPool
+	{
+		int32 NextFreeIndex = 0;
+		TArray<DummyUAV> UAVs;
+	};
+
+	mutable TMap<EPixelFormat, DummyUAVPool> DummyBufferPool;
+	mutable TMap<EPixelFormat, DummyUAVPool> DummyTexturePool;
+
+	FRHIUnorderedAccessView* GetEmptyUAVFromPool(FRHICommandList& RHICmdList, EPixelFormat Format, bool IsTexture) const;
+	void ResetEmptyUAVPool(TMap<EPixelFormat, DummyUAVPool>& UAVMap, TArray<FRHIUnorderedAccessView*>& Transitions);
+	void ResetEmptyUAVPools(FRHICommandList& RHICmdList);
 };
