One more build fix.

diff --git a/Engine/Plugins/FX/Niagara/Source/Niagara/Private/NiagaraDataInterfaceParticleRead.cpp b/Engine/Plugins/FX/Niagara/Source/Niagara/Private/NiagaraDataInterfaceParticleRead.cpp
--- a/Engine/Plugins/FX/Niagara/Source/Niagara/Private/NiagaraDataInterfaceParticleRead.cpp
+++ b/Engine/Plugins/FX/Niagara/Source/Niagara/Private/NiagaraDataInterfaceParticleRead.cpp
@@ -331,12 +331,12 @@ struct FNiagaraDataInterfaceParametersCS_ParticleRead : public FNiagaraDataInter
 
 	FRHIShaderResourceView* GetIntSRVWithFallback(FRWBuffer& Buffer) const
 	{
-		return Buffer.SRV ? Buffer.SRV : FNiagaraRenderer::GetDummyIntBuffer();
+		return Buffer.SRV ? Buffer.SRV.GetReference() : FNiagaraRenderer::GetDummyIntBuffer();
 	}
 
 	FRHIShaderResourceView* GetFloatSRVWithFallback(FRWBuffer& Buffer) const
 	{
-		return Buffer.SRV ? Buffer.SRV : FNiagaraRenderer::GetDummyFloatBuffer();
+		return Buffer.SRV ? Buffer.SRV.GetReference() : FNiagaraRenderer::GetDummyFloatBuffer();
 	}
 
 	void Set(FRHICommandList& RHICmdList, const FNiagaraDataInterfaceSetArgs& Context) const
