Call IMPLEMENT_VERTEX_FACTORY_PARAMETER_TYPE for FNiagaraMeshVertexFactory

diff --git a/Engine/Plugins/FX/Niagara/Source/NiagaraVertexFactories/Private/NiagaraMeshVertexFactory.cpp b/Engine/Plugins/FX/Niagara/Source/NiagaraVertexFactories/Private/NiagaraMeshVertexFactory.cpp
--- a/Engine/Plugins/FX/Niagara/Source/NiagaraVertexFactories/Private/NiagaraMeshVertexFactory.cpp
+++ b/Engine/Plugins/FX/Niagara/Source/NiagaraVertexFactories/Private/NiagaraMeshVertexFactory.cpp
@@ -242,6 +242,10 @@ void FNiagaraMeshVertexFactory::SetData(const FStaticMeshDataType& InData)
 }
 
 IMPLEMENT_VERTEX_FACTORY_PARAMETER_TYPE(FNiagaraMeshVertexFactory, SF_Vertex, FNiagaraMeshVertexFactoryShaderParametersVS);
+#if RHI_RAYTRACING
+IMPLEMENT_VERTEX_FACTORY_PARAMETER_TYPE(FNiagaraMeshVertexFactory, SF_Compute, FNiagaraMeshVertexFactoryShaderParametersVS);
+IMPLEMENT_VERTEX_FACTORY_PARAMETER_TYPE(FNiagaraMeshVertexFactory, SF_RayHitGroup, FNiagaraMeshVertexFactoryShaderParametersVS);
+#endif
 IMPLEMENT_VERTEX_FACTORY_PARAMETER_TYPE(FNiagaraMeshVertexFactory, SF_Pixel, FNiagaraMeshVertexFactoryShaderParametersPS);
 
 IMPLEMENT_VERTEX_FACTORY_TYPE(FNiagaraMeshVertexFactory, "/Plugin/FX/Niagara/Private/NiagaraMeshVertexFactory.ush", true, false, true, false, false);
