Niagara - Fix an issue where data interfaces defined in emitters could end up with name collisions which would caused unintended sharing of data interface objects.  When registering data interfaces while compiling emitters the names are made unique.

diff --git a/Engine/Plugins/FX/Niagara/Source/Niagara/Private/NiagaraScript.cpp b/Engine/Plugins/FX/Niagara/Source/Niagara/Private/NiagaraScript.cpp
--- a/Engine/Plugins/FX/Niagara/Source/Niagara/Private/NiagaraScript.cpp
+++ b/Engine/Plugins/FX/Niagara/Source/Niagara/Private/NiagaraScript.cpp
@@ -1372,13 +1372,13 @@ void UNiagaraScript::SetVMCompilationResults(const FNiagaraVMExecutableDataId& I
 	{
 		int32 Idx = CachedDefaultDataInterfaces.AddDefaulted();
 		CachedDefaultDataInterfaces[Idx].UserPtrIdx = Info.UserPtrIdx;
-		CachedDefaultDataInterfaces[Idx].Name = Info.Name;
+		CachedDefaultDataInterfaces[Idx].Name = InRequestData->ResolveEmitterAlias(Info.Name);
 		CachedDefaultDataInterfaces[Idx].Type = Info.Type;
 		CachedDefaultDataInterfaces[Idx].RegisteredParameterMapRead = InRequestData->ResolveEmitterAlias(Info.RegisteredParameterMapRead);
 		CachedDefaultDataInterfaces[Idx].RegisteredParameterMapWrite = InRequestData->ResolveEmitterAlias(Info.RegisteredParameterMapWrite);
 
 		// We compiled it just a bit ago, so we should be able to resolve it from the table that we passed in.
-		UNiagaraDataInterface*const* FindDIById = InRequestData->GetObjectNameMap().Find(Info.Name);
+		UNiagaraDataInterface*const* FindDIById = InRequestData->GetObjectNameMap().Find(CachedDefaultDataInterfaces[Idx].Name);
 		if (FindDIById != nullptr && *(FindDIById) != nullptr)
 		{
 			CachedDefaultDataInterfaces[Idx].DataInterface = CopyDataInterface(*(FindDIById), this);
diff --git a/Engine/Plugins/FX/Niagara/Source/NiagaraCore/Private/NiagaraCustomVersion.cpp b/Engine/Plugins/FX/Niagara/Source/NiagaraCore/Private/NiagaraCustomVersion.cpp
--- a/Engine/Plugins/FX/Niagara/Source/NiagaraCore/Private/NiagaraCustomVersion.cpp
+++ b/Engine/Plugins/FX/Niagara/Source/NiagaraCore/Private/NiagaraCustomVersion.cpp
@@ -8,4 +8,4 @@ const FGuid FNiagaraCustomVersion::GUID(0xFCF57AFA, 0x50764283, 0xB9A9E658, 0xFF
 // Register the custom version with core
 FCustomVersionRegistration GRegisterNiagaraCustomVersion(FNiagaraCustomVersion::GUID, FNiagaraCustomVersion::LatestVersion, TEXT("NiagaraVer"));
 
-const FGuid FNiagaraCustomVersion::LatestScriptCompileVersion(0XC9A39220, 0X1AE54A45, 0X8A01CC89, 0X64E91B28);
+const FGuid FNiagaraCustomVersion::LatestScriptCompileVersion(0x4122F10C, 0xA1FA4669, 0xB73B66FA, 0x2B230FC6);
diff --git a/Engine/Plugins/FX/Niagara/Source/NiagaraEditor/Private/NiagaraCompiler.cpp b/Engine/Plugins/FX/Niagara/Source/NiagaraEditor/Private/NiagaraCompiler.cpp
--- a/Engine/Plugins/FX/Niagara/Source/NiagaraEditor/Private/NiagaraCompiler.cpp
+++ b/Engine/Plugins/FX/Niagara/Source/NiagaraEditor/Private/NiagaraCompiler.cpp
@@ -130,7 +130,8 @@ void FNiagaraCompileRequestData::VisitReferencedGraphsRecursive(UNiagaraGraph* I
 				if (InputNode->Input.IsDataInterface())
 				{
 					UNiagaraDataInterface* DataInterface = InputNode->GetDataInterface();
-					FName DIName = InputNode->Input.GetName();
+					bool bIsParameterMapDataInterface = false;
+					FName DIName = FHlslNiagaraTranslator::GetDataInterfaceName(InputNode->Input.GetName(), EmitterUniqueName, bIsParameterMapDataInterface);
 					UNiagaraDataInterface* Dupe = DuplicateObject<UNiagaraDataInterface>(DataInterface, GetTransientPackage());
 					CopiedDataInterfacesByName.Add(DIName, Dupe);
 				}
@@ -586,11 +587,11 @@ TSharedPtr<FNiagaraCompileRequestDataBase, ESPMode::ThreadSafe> FNiagaraEditorMo
 			const FNiagaraEmitterHandle& Handle = System->GetEmitterHandle(i);
 			FCompileConstantResolver ConstantResolver(Handle.GetInstance());
 			TSharedPtr<FNiagaraCompileRequestData, ESPMode::ThreadSafe> EmitterPtr = MakeShared<FNiagaraCompileRequestData, ESPMode::ThreadSafe>();
+			EmitterPtr->EmitterUniqueName = Handle.GetInstance()->GetUniqueEmitterName();
 			if (Handle.GetIsEnabled()) // Don't need to copy the graph if we aren't going to use it.
 			{
 				EmitterPtr->DeepCopyGraphs(Cast<UNiagaraScriptSource>(Handle.GetInstance()->GraphSource), ENiagaraScriptUsage::EmitterSpawnScript, ConstantResolver);
 			}
-			EmitterPtr->EmitterUniqueName = Handle.GetInstance()->GetUniqueEmitterName();
 			EmitterPtr->SourceName = BasePtr->SourceName;
 			EmitterPtr->bUseRapidIterationParams = BasePtr->bUseRapidIterationParams || (!Handle.GetInstance()->bBakeOutRapidIteration);
 			BasePtr->EmitterData.Add(EmitterPtr);
diff --git a/Engine/Plugins/FX/Niagara/Source/NiagaraEditor/Private/NiagaraHlslTranslator.cpp b/Engine/Plugins/FX/Niagara/Source/NiagaraEditor/Private/NiagaraHlslTranslator.cpp
--- a/Engine/Plugins/FX/Niagara/Source/NiagaraEditor/Private/NiagaraHlslTranslator.cpp
+++ b/Engine/Plugins/FX/Niagara/Source/NiagaraEditor/Private/NiagaraHlslTranslator.cpp
@@ -5040,16 +5040,8 @@ int32 FHlslNiagaraTranslator::RegisterDataInterface(FNiagaraVariable& Var, UNiag
 	}
 
 	//If we get here then this is a new data interface.
-	FName DataInterfaceName;
-	if (FNiagaraParameterMapHistory::IsAliasedEmitterParameter(Var.GetName().ToString()))
-	{
-		FNiagaraVariable AliasedVar = ActiveHistoryForFunctionCalls.ResolveAliases(Var);
-		DataInterfaceName = AliasedVar.GetName();
-	}
-	else
-	{
-		DataInterfaceName = Var.GetName();
-	}
+	const FString* EmitterAlias = ActiveHistoryForFunctionCalls.GetEmitterAlias();
+	FName DataInterfaceName = GetDataInterfaceName(Var.GetName(), EmitterAlias != nullptr ? *EmitterAlias : FString(), bAddParameterMapRead);
 
 	int32 Idx = CompilationOutput.ScriptData.DataInterfaceInfo.IndexOfByPredicate([&](const FNiagaraScriptDataInterfaceCompileInfo& OtherInfo)
 	{
@@ -5071,7 +5063,6 @@ int32 FHlslNiagaraTranslator::RegisterDataInterface(FNiagaraVariable& Var, UNiag
 	}
 	else
 	{
-		check(CompilationOutput.ScriptData.DataInterfaceInfo[Idx].Name == Var.GetName());
 		check(CompilationOutput.ScriptData.DataInterfaceInfo[Idx].Type == Var.GetType());
 	}
 
@@ -6147,6 +6138,24 @@ FString FHlslNiagaraTranslator::GetFunctionSignatureSymbol(const FNiagaraFunctio
 	return GetSanitizedSymbolName(SigStr);
 }
 
+FName FHlslNiagaraTranslator::GetDataInterfaceName(FName BaseName, const FString& UniqueEmitterName, bool bIsParameterMapDataInterface)
+{
+	if (UniqueEmitterName.IsEmpty() == false)
+	{
+		if (FNiagaraParameterMapHistory::IsAliasedEmitterParameter(BaseName.ToString()))
+		{
+			return FNiagaraParameterMapHistory::ResolveEmitterAlias(BaseName, UniqueEmitterName);
+		}
+		else if(bIsParameterMapDataInterface == false)
+		{
+			// Don't mangle the parameter map reads for emitter scripts because they are from the system or user parameter stores and
+			// they won't bind correctly.
+			return *(UniqueEmitterName + TEXT(".") + BaseName.ToString());
+		}
+	}
+	return BaseName;
+}
+
 FString FHlslNiagaraTranslator::GetFunctionSignature(const FNiagaraFunctionSignature& Sig)
 {
 	FString SigStr = TEXT("void ") + GetFunctionSignatureSymbol(Sig);
diff --git a/Engine/Plugins/FX/Niagara/Source/NiagaraEditor/Private/NiagaraHlslTranslator.h b/Engine/Plugins/FX/Niagara/Source/NiagaraEditor/Private/NiagaraHlslTranslator.h
--- a/Engine/Plugins/FX/Niagara/Source/NiagaraEditor/Private/NiagaraHlslTranslator.h
+++ b/Engine/Plugins/FX/Niagara/Source/NiagaraEditor/Private/NiagaraHlslTranslator.h
@@ -549,6 +549,8 @@ public:
 
 	static FString GetFunctionSignatureSymbol(const FNiagaraFunctionSignature& Sig);
 
+	static FName GetDataInterfaceName(FName BaseName, const FString& UniqueEmitterName, bool bIsParameterMapDataInterface);
+
 	/** If OutVar can be replaced by a literal constant, it's data is initialized with the correct value and we return true. Returns false otherwise. */
 	bool GetLiteralConstantVariable(FNiagaraVariable& OutVar) const;
 
diff --git a/Engine/Plugins/FX/Niagara/Source/NiagaraEditor/Private/ViewModels/NiagaraSystemViewModel.cpp b/Engine/Plugins/FX/Niagara/Source/NiagaraEditor/Private/ViewModels/NiagaraSystemViewModel.cpp
--- a/Engine/Plugins/FX/Niagara/Source/NiagaraEditor/Private/ViewModels/NiagaraSystemViewModel.cpp
+++ b/Engine/Plugins/FX/Niagara/Source/NiagaraEditor/Private/ViewModels/NiagaraSystemViewModel.cpp
@@ -23,6 +23,7 @@
 #include "NiagaraSystemEditorData.h"
 #include "NiagaraGraph.h"
 #include "NiagaraNodeInput.h"
+#include "NiagaraNodeOutput.h"
 #include "NiagaraEditorModule.h"
 #include "NiagaraNodeFunctionCall.h"
 #include "EdGraphSchema_NiagaraSystemOverview.h"
@@ -1272,6 +1273,51 @@ void FNiagaraSystemViewModel::ResetCurveData()
 	OnCurveOwnerChangedDelegate.Broadcast();
 }
 
+void GetCompiledScriptAndEmitterNameFromInputNode(UNiagaraNode& StackNode, UNiagaraSystem& OwningSystem, UNiagaraScript*& OutCompiledScript, FString& OutEmitterName)
+{
+	OutCompiledScript = nullptr;
+	OutEmitterName = FString();
+	UNiagaraNodeOutput* OutputNode = FNiagaraStackGraphUtilities::GetEmitterOutputNodeForStackNode(StackNode);
+	if (OutputNode != nullptr)
+	{
+		if (OutputNode->GetUsage() == ENiagaraScriptUsage::SystemSpawnScript)
+		{
+			OutCompiledScript = OwningSystem.GetSystemSpawnScript();
+		}
+		else if (OutputNode->GetUsage() == ENiagaraScriptUsage::SystemUpdateScript)
+		{
+			OutCompiledScript = OwningSystem.GetSystemUpdateScript();
+		}
+		else
+		{
+			const FNiagaraEmitterHandle* OwningEmitterHandle = OwningSystem.GetEmitterHandles().FindByPredicate([&StackNode](const FNiagaraEmitterHandle& EmitterHandle)
+			{
+				return CastChecked<UNiagaraScriptSource>(EmitterHandle.GetInstance()->GraphSource)->NodeGraph == StackNode.GetNiagaraGraph();
+			});
+
+			if(OwningEmitterHandle != nullptr)
+			{
+				OutEmitterName = OwningEmitterHandle->GetInstance()->GetUniqueEmitterName();
+				switch (OutputNode->GetUsage())
+				{
+				case ENiagaraScriptUsage::EmitterSpawnScript:
+					OutCompiledScript = OwningSystem.GetSystemSpawnScript();
+					break;
+				case ENiagaraScriptUsage::EmitterUpdateScript:
+					OutCompiledScript = OwningSystem.GetSystemUpdateScript();
+					break;
+				case ENiagaraScriptUsage::ParticleSpawnScript:
+				case ENiagaraScriptUsage::ParticleUpdateScript:
+				case ENiagaraScriptUsage::ParticleEventScript:
+				case ENiagaraScriptUsage::ParticleSimulationStageScript:
+					OutCompiledScript = OwningEmitterHandle->GetInstance()->GetScript(OutputNode->GetUsage(), OutputNode->GetUsageId());
+					break;
+				}
+			}
+		}
+	}
+}
+
 void UpdateCompiledDataInterfacesForScript(UNiagaraScript& TargetScript, FName TargetDataInterfaceName, UNiagaraDataInterface& SourceDataInterface)
 {
 	for (FNiagaraScriptDataInterfaceInfo& DataInterfaceInfo : TargetScript.GetCachedDefaultDataInterfaces())
@@ -1297,32 +1343,14 @@ void FNiagaraSystemViewModel::UpdateCompiledDataInterfaces(UNiagaraDataInterface
 		}
 
 		// If the data interface was owned by an input node, then we need to try to update the compiled version.
-		UNiagaraEmitter* OwningEmitter;
-		UNiagaraScript* OwningScript;
-		FNiagaraStackGraphUtilities::GetOwningEmitterAndScriptForStackNode(*OuterInputNode, GetSystem(), OwningEmitter, OwningScript);
-		if (ensureMsgf(OwningScript != nullptr, TEXT("Could not find owning script for data interface input node.")))
-		{
-			switch (OwningScript->GetUsage())
-			{
-			case ENiagaraScriptUsage::SystemSpawnScript:
-			case ENiagaraScriptUsage::SystemUpdateScript:
-			case ENiagaraScriptUsage::ParticleSpawnScript:
-			case ENiagaraScriptUsage::ParticleSpawnScriptInterpolated:
-			case ENiagaraScriptUsage::ParticleUpdateScript:
-			case ENiagaraScriptUsage::ParticleSimulationStageScript:
-			case ENiagaraScriptUsage::ParticleEventScript:
-				UpdateCompiledDataInterfacesForScript(*OwningScript, OuterInputNode->Input.GetName(), *ChangedDataInterface);
-				break;
-			case ENiagaraScriptUsage::EmitterSpawnScript:
-			case ENiagaraScriptUsage::EmitterUpdateScript:
-				if (ensureMsgf(OwningEmitter != nullptr, TEXT("Could not find owning emitter for data interface input node.")))
-				{
-					UNiagaraScript& TargetScript = OwningScript->GetUsage() == ENiagaraScriptUsage::EmitterSpawnScript ? *GetSystem().GetSystemSpawnScript() : *GetSystem().GetSystemUpdateScript();
-					FName AliasedInputNodeName = FNiagaraParameterMapHistory::ResolveEmitterAlias(OuterInputNode->Input.GetName(), OwningEmitter->GetUniqueEmitterName());
-					UpdateCompiledDataInterfacesForScript(TargetScript, AliasedInputNodeName, *ChangedDataInterface);
-				}
-				break;
-			}
+		UNiagaraScript* CompiledScript;
+		FString EmitterName;
+		GetCompiledScriptAndEmitterNameFromInputNode(*OuterInputNode, GetSystem(), CompiledScript, EmitterName);
+		if (ensureMsgf(CompiledScript != nullptr, TEXT("Could not find owning script for data interface input node.")))
+		{
+			bool bIsParameterMapDataInterface = false;
+			FName DataInterfaceName = FHlslNiagaraTranslator::GetDataInterfaceName(OuterInputNode->Input.GetName(), EmitterName, bIsParameterMapDataInterface);
+			UpdateCompiledDataInterfacesForScript(*CompiledScript, DataInterfaceName, *ChangedDataInterface);
 		}
 	}
 	else
diff --git a/Engine/Plugins/FX/Niagara/Source/NiagaraEditor/Private/ViewModels/Stack/NiagaraStackGraphUtilities.cpp b/Engine/Plugins/FX/Niagara/Source/NiagaraEditor/Private/ViewModels/Stack/NiagaraStackGraphUtilities.cpp
--- a/Engine/Plugins/FX/Niagara/Source/NiagaraEditor/Private/ViewModels/Stack/NiagaraStackGraphUtilities.cpp
+++ b/Engine/Plugins/FX/Niagara/Source/NiagaraEditor/Private/ViewModels/Stack/NiagaraStackGraphUtilities.cpp
@@ -1607,41 +1607,6 @@ TOptional<FName> FNiagaraStackGraphUtilities::GetNamespaceForScriptUsage(ENiagar
 	}
 }
 
-void FNiagaraStackGraphUtilities::GetOwningEmitterAndScriptForStackNode(UNiagaraNode& StackNode, UNiagaraSystem& OwningSystem, UNiagaraEmitter*& OutEmitter, UNiagaraScript*& OutScript)
-{
-	OutEmitter = nullptr;
-	OutScript = nullptr;
-	UNiagaraNodeOutput* OutputNode = GetEmitterOutputNodeForStackNode(StackNode);
-	if (OutputNode != nullptr)
-	{
-		switch (OutputNode->GetUsage())
-		{
-			case ENiagaraScriptUsage::SystemSpawnScript:
-				OutScript = OwningSystem.GetSystemSpawnScript();
-				break;
-			case ENiagaraScriptUsage::SystemUpdateScript:
-				OutScript = OwningSystem.GetSystemUpdateScript();
-				break;
-			case ENiagaraScriptUsage::EmitterSpawnScript:
-			case ENiagaraScriptUsage::EmitterUpdateScript:
-			case ENiagaraScriptUsage::ParticleSpawnScript:
-			case ENiagaraScriptUsage::ParticleUpdateScript:
-			case ENiagaraScriptUsage::ParticleEventScript:
-				for (const FNiagaraEmitterHandle& EmitterHandle : OwningSystem.GetEmitterHandles())
-				{
-					UNiagaraScriptSource* EmitterSource = CastChecked<UNiagaraScriptSource>(EmitterHandle.GetInstance()->GraphSource);
-					if (EmitterSource->NodeGraph == StackNode.GetNiagaraGraph())
-					{
-						OutEmitter = EmitterHandle.GetInstance();
-						OutScript = OutEmitter->GetScript(OutputNode->GetUsage(), OutputNode->GetUsageId());
-						break;
-					}
-				}
-				break;
-		}
-	}
-}
-
 struct FRapidIterationParameterContext
 {
 	FRapidIterationParameterContext()
diff --git a/Engine/Plugins/FX/Niagara/Source/NiagaraEditor/Public/ViewModels/Stack/NiagaraStackGraphUtilities.h b/Engine/Plugins/FX/Niagara/Source/NiagaraEditor/Public/ViewModels/Stack/NiagaraStackGraphUtilities.h
--- a/Engine/Plugins/FX/Niagara/Source/NiagaraEditor/Public/ViewModels/Stack/NiagaraStackGraphUtilities.h
+++ b/Engine/Plugins/FX/Niagara/Source/NiagaraEditor/Public/ViewModels/Stack/NiagaraStackGraphUtilities.h
@@ -149,8 +149,6 @@ namespace FNiagaraStackGraphUtilities
 
 	TOptional<FName> GetNamespaceForScriptUsage(ENiagaraScriptUsage ScriptUsage);
 
-	void GetOwningEmitterAndScriptForStackNode(UNiagaraNode& StackNode, UNiagaraSystem& OwningSystem, UNiagaraEmitter*& OutEmitter, UNiagaraScript*& OutScript);
-
 	bool IsValidDefaultDynamicInput(UNiagaraScript& OwningScript, UEdGraphPin& DefaultPin);
 
 	bool CanWriteParameterFromUsage(FNiagaraVariable Parameter, ENiagaraScriptUsage Usage);
