Fixed crash if a GPU system containing a neighbor grid runs before the per-instance data has been set.

This shouldn't happen, but it does when editing emitters, so this is a stopgap until we figure out why. It doesn't hurt for the DI to be more robust anyway.

diff --git a/Engine/Plugins/FX/Niagara/Source/Niagara/Private/NiagaraDataInterfaceNeighborGrid3D.cpp b/Engine/Plugins/FX/Niagara/Source/Niagara/Private/NiagaraDataInterfaceNeighborGrid3D.cpp
--- a/Engine/Plugins/FX/Niagara/Source/Niagara/Private/NiagaraDataInterfaceNeighborGrid3D.cpp
+++ b/Engine/Plugins/FX/Niagara/Source/Niagara/Private/NiagaraDataInterfaceNeighborGrid3D.cpp
@@ -56,12 +56,33 @@ public:
 		FNiagaraDataInterfaceProxyNeighborGrid3D* VFDI = static_cast<FNiagaraDataInterfaceProxyNeighborGrid3D*>(Context.DataInterface);
 
 		NeighborGrid3DRWInstanceData* ProxyData = VFDI->SystemInstancesToProxyData.Find(Context.SystemInstance);
-		check(ProxyData);
+		float VoxelSizeTmp[3];
+
+		if (!ProxyData)
+		{
+			VoxelSizeTmp[0] = VoxelSizeTmp[1] = VoxelSizeTmp[2] = 1.0f;
+			SetShaderValue(RHICmdList, ComputeShaderRHI, NumVoxelsParam, 0);
+			SetShaderValue(RHICmdList, ComputeShaderRHI, VoxelSizeParam, VoxelSizeTmp);
+			SetShaderValue(RHICmdList, ComputeShaderRHI, MaxNeighborsPerVoxelParam, 0);
+			SetShaderValue(RHICmdList, ComputeShaderRHI, WorldBBoxSizeParam, FVector(0.0f, 0.0f, 0.0f));
+			SetSRVParameter(RHICmdList, Context.Shader.GetComputeShader(), ParticleNeighborsGridParam, FNiagaraRenderer::GetDummyIntBuffer());
+			SetSRVParameter(RHICmdList, Context.Shader.GetComputeShader(), ParticleNeighborCountGridParam, FNiagaraRenderer::GetDummyIntBuffer());
+			if (OutputParticleNeighborsGridParam.IsUAVBound())
+			{
+				RHICmdList.SetUAVParameter(ComputeShaderRHI, OutputParticleNeighborsGridParam.GetUAVIndex(), Context.Batcher->GetEmptyRWBufferFromPool(RHICmdList, PF_R32_SINT));
+			}
+
+			if (OutputParticleNeighborCountGridParam.IsUAVBound())
+			{
+				RHICmdList.SetUAVParameter(ComputeShaderRHI, OutputParticleNeighborCountGridParam.GetUAVIndex(), Context.Batcher->GetEmptyRWBufferFromPool(RHICmdList, PF_R32_SINT));
+			}
+
+			return;
+		}
 
 		SetShaderValue(RHICmdList, ComputeShaderRHI, NumVoxelsParam, ProxyData->NumVoxels);
 
 		// #todo(dmp): remove this computation here
-		float VoxelSizeTmp[3];
 		VoxelSizeTmp[0] = ProxyData->WorldBBoxSize.X / ProxyData->NumVoxels.X;
 		VoxelSizeTmp[1] = ProxyData->WorldBBoxSize.Y / ProxyData->NumVoxels.Y;
 		VoxelSizeTmp[2] = ProxyData->WorldBBoxSize.Z / ProxyData->NumVoxels.Z;
