Fix for issue where calling Sample Audio Buffer for positions 0.5-1.0 would return zeroes.

diff --git a/Engine/Plugins/FX/Niagara/Source/Niagara/Private/NiagaraDataInterfaceAudioOscilloscope.cpp b/Engine/Plugins/FX/Niagara/Source/Niagara/Private/NiagaraDataInterfaceAudioOscilloscope.cpp
--- a/Engine/Plugins/FX/Niagara/Source/Niagara/Private/NiagaraDataInterfaceAudioOscilloscope.cpp
+++ b/Engine/Plugins/FX/Niagara/Source/Niagara/Private/NiagaraDataInterfaceAudioOscilloscope.cpp
@@ -251,9 +251,9 @@ bool UNiagaraDataInterfaceAudioOscilloscope::GetFunctionHLSL(const FNiagaraDataI
 		static const TCHAR *FormatBounds = TEXT(R"(
 			void {FunctionName}(float In_NormalizedPosition, int In_ChannelIndex, out float Out_Val)
 			{
-				float FrameIndex = In_NormalizedPosition * {AudioBufferNumFrames};
+				float FrameIndex = In_NormalizedPosition * {AudioBufferNumSamples} / {ChannelCount};
 				int LowerIndex = floor(FrameIndex);
-				int UpperIndex =  LowerIndex < {AudioBufferNumFrames} ? LowerIndex + 1.0 : LowerIndex;
+				int UpperIndex =  LowerIndex < {AudioBufferNumSamples} ? LowerIndex + 1.0 : LowerIndex;
 				float Fraction = FrameIndex - LowerIndex;
 				float LowerValue = {AudioBuffer}.Load(LowerIndex * {ChannelCount} + In_ChannelIndex);
 				float UpperValue = {AudioBuffer}.Load(UpperIndex * {ChannelCount} + In_ChannelIndex);
@@ -263,7 +263,7 @@ bool UNiagaraDataInterfaceAudioOscilloscope::GetFunctionHLSL(const FNiagaraDataI
 		TMap<FString, FStringFormatArg> ArgsBounds = {
 			{TEXT("FunctionName"), FStringFormatArg(FunctionInfo.InstanceName)},
 			{TEXT("ChannelCount"), FStringFormatArg(NumChannelsName + ParamInfo.DataInterfaceHLSLSymbol)},
-			{TEXT("AudioBufferNumFrames"), FStringFormatArg(Resolution)},
+			{TEXT("AudioBufferNumSamples"), FStringFormatArg(Resolution)},
 			{TEXT("AudioBuffer"), FStringFormatArg(AudioBufferName + ParamInfo.DataInterfaceHLSLSymbol)},
 		};
 		OutHLSL += FString::Format(FormatBounds, ArgsBounds);
@@ -447,7 +447,7 @@ void FNiagaraDataInterfaceProxyOscilloscope::PostAudioToGPU()
 		[&](FRHICommandListImmediate& RHICmdList)
 	{
 		DownsampleAudioToBuffer();
-		uint32 BufferSize = Resolution * NumChannelsInDownsampledBuffer * sizeof(float);
+		size_t BufferSize = DownsampledBuffer.Num() * sizeof(float);
 		if (BufferSize != 0 && !GPUDownsampledBuffer.NumBytes)
 		{
 			GPUDownsampledBuffer.Initialize(sizeof(float), Resolution * NumChannelsInDownsampledBuffer, EPixelFormat::PF_R32_FLOAT, BUF_Static);
@@ -504,7 +504,7 @@ void FNiagaraDataInterfaceProxyOscilloscope::DownsampleAudioToBuffer()
 	}
 	else
 	{
-		PopBuffer.SetNumZeroed(NumFramesToPop * NumChannelsInDownsampledBuffer);
+		PopBuffer.SetNumZeroed(NumSamplesToPop);
 	}
 
 	// Downsample buffer in place.
