Niagara Namespace Refactor 1.0: Add precompiler stage to allow automatically promoting parameters linked across script executions from uniforms to the dataset by fixing namespace. Also add per system option to cull dataset parameters that are not used. Add metadata to define scope and input/ouput argument of variables.
Niagara Parameter Panel 2.0: Reimplement Parameter Panel to support displaying parameter names as namestrings with separate UI to configure associated metadata. Implement categorization of parameters to manage inputs/references/outputs/initial values separately.

diff --git a/Engine/Plugins/FX/Niagara/Source/Niagara/Classes/NiagaraConstants.h b/Engine/Plugins/FX/Niagara/Source/Niagara/Classes/NiagaraConstants.h
--- a/Engine/Plugins/FX/Niagara/Source/Niagara/Classes/NiagaraConstants.h
+++ b/Engine/Plugins/FX/Niagara/Source/Niagara/Classes/NiagaraConstants.h
@@ -19,7 +19,14 @@
 #define PARAM_MAP_RAPID_ITERATION_STR TEXT("Constants.")
 #define PARAM_MAP_INDICES_STR TEXT("Array.")
 #define PARAM_MAP_RAPID_ITERATION_BASE_STR TEXT("Constants")
+#define PARAM_MAP_SCRIPT_PERSISTENT_STR TEXT("ScriptPersistent.")
+#define PARAM_MAP_SCRIPT_TRANSIENT_STR TEXT("ScriptTransient.")
+#define PARAM_MAP_INTERMEDIATE_STR TEXT("Intermedate.")
+#define PARAM_MAP_LOCAL_STR TEXT("Local.")
+#define PARAM_MAP_OUTPUT_STR TEXT("Output.")
 
+#define TRANSLATOR_SET_VARIABLES_STR TEXT("SetVariables")
+#define TRANSLATOR_SET_VARIABLES_UNDERSCORE_STR TEXT("SetVariables_")
 
 #define SYS_PARAM_ENGINE_DELTA_TIME                      INiagaraModule::GetVar_Engine_DeltaTime()
 #define SYS_PARAM_ENGINE_INV_DELTA_TIME                  INiagaraModule::GetVar_Engine_InvDeltaTime()
@@ -134,6 +141,9 @@ struct NIAGARA_API FNiagaraConstants
 
 	static bool IsEngineManagedAttribute(const FNiagaraVariable& Var);
 
+	static const FName InputPinName;
+	static const FName OutputPinName;
+
 private:
 	static TArray<FNiagaraVariable> SystemParameters;
 	static TArray<FNiagaraVariable> TranslatorParameters;
diff --git a/Engine/Plugins/FX/Niagara/Source/Niagara/Classes/NiagaraScript.h b/Engine/Plugins/FX/Niagara/Source/Niagara/Classes/NiagaraScript.h
--- a/Engine/Plugins/FX/Niagara/Source/Niagara/Classes/NiagaraScript.h
+++ b/Engine/Plugins/FX/Niagara/Source/Niagara/Classes/NiagaraScript.h
@@ -466,6 +466,7 @@ public:
 	NIAGARA_API const FNiagaraVMExecutableDataId& GetComputedVMCompilationId() const { return CachedScriptVMId; }
 #endif
 
+
 	void SetUsage(ENiagaraScriptUsage InUsage) { Usage = InUsage; }
 	ENiagaraScriptUsage GetUsage() const { return Usage; }
 
@@ -526,6 +527,7 @@ public:
 	static bool NIAGARA_API ConvertUsageToGroup(ENiagaraScriptUsage InUsage, ENiagaraScriptGroup& OutGroup);
 
 #if WITH_EDITORONLY_DATA
+	NIAGARA_API TArray<ENiagaraParameterScope> GetUnsupportedParameterScopes() const;
 	NIAGARA_API TArray<ENiagaraScriptUsage> GetSupportedUsageContexts() const;
 	static NIAGARA_API TArray<ENiagaraScriptUsage> GetSupportedUsageContextsForBitmask(int32 InModuleUsageBitmask);
 #endif
@@ -567,17 +569,16 @@ public:
 	void CacheShadersForResources(EShaderPlatform ShaderPlatform, FNiagaraShaderScript *ResourceToCache, bool bApplyCompletedShaderMapForRendering, bool bForceRecompile = false, bool bCooking=false);
 #endif // WITH_EDITOR
 	FNiagaraShaderScript* AllocateResource();
-	FNiagaraShaderScript *GetRenderThreadScript()
+	FNiagaraShaderScript* GetRenderThreadScript()
 	{
 		return &ScriptResource;
 	}
-	const FNiagaraShaderScript *GetRenderThreadScript() const
+
+	const FNiagaraShaderScript* GetRenderThreadScript() const
 	{
 		return &ScriptResource;
 	}
 
-	
-
 	FComputeShaderRHIRef GetScriptShader() 
 	{
 		if (!ScriptShader)
diff --git a/Engine/Plugins/FX/Niagara/Source/Niagara/Classes/NiagaraSystem.h b/Engine/Plugins/FX/Niagara/Source/Niagara/Classes/NiagaraSystem.h
--- a/Engine/Plugins/FX/Niagara/Source/Niagara/Classes/NiagaraSystem.h
+++ b/Engine/Plugins/FX/Niagara/Source/Niagara/Classes/NiagaraSystem.h
@@ -331,6 +331,11 @@ public:
 	/** Experimental feature that allows us to bake out rapid iteration parameters during the normal compile process. */
 	UPROPERTY(EditAnywhere, AdvancedDisplay, Category = "Emitter")
 	uint32 bBakeOutRapidIteration : 1;
+
+	/** Experimental feature to depromote parameters in the dataset that are not accessed during the precompile process. */
+	UPROPERTY(EditAnywhere, AdvancedDisplay, Category = "Emitter")
+	uint32 bCullDatasetParameters : 1;
+
 #endif
 
 	FORCEINLINE UNiagaraParameterCollectionInstance* GetParameterCollectionOverride(UNiagaraParameterCollection* Collection)
diff --git a/Engine/Plugins/FX/Niagara/Source/Niagara/Private/NiagaraConstants.cpp b/Engine/Plugins/FX/Niagara/Source/Niagara/Private/NiagaraConstants.cpp
--- a/Engine/Plugins/FX/Niagara/Source/Niagara/Private/NiagaraConstants.cpp
+++ b/Engine/Plugins/FX/Niagara/Source/Niagara/Private/NiagaraConstants.cpp
@@ -20,6 +20,9 @@ TMap<FNiagaraVariable, FNiagaraVariable> FNiagaraConstants::AttrDefaultsValueMap
 TMap<FNiagaraVariable, FNiagaraVariable> FNiagaraConstants::AttrDataSetKeyMap;
 TArray<FNiagaraVariable> FNiagaraConstants::EngineManagedAttributes;
 
+const FName FNiagaraConstants::InputPinName("InputPin");
+const FName FNiagaraConstants::OutputPinName("OutputPin");
+
 void FNiagaraConstants::Init()
 {
 	if (SystemParameters.Num() == 0)
diff --git a/Engine/Plugins/FX/Niagara/Source/Niagara/Private/NiagaraLightRendererProperties.cpp b/Engine/Plugins/FX/Niagara/Source/Niagara/Private/NiagaraLightRendererProperties.cpp
--- a/Engine/Plugins/FX/Niagara/Source/Niagara/Private/NiagaraLightRendererProperties.cpp
+++ b/Engine/Plugins/FX/Niagara/Source/Niagara/Private/NiagaraLightRendererProperties.cpp
@@ -135,6 +135,37 @@ void UNiagaraLightRendererProperties::FixMaterial(UMaterial* Material)
 {
 }
 
+const TArray<FNiagaraVariable>& UNiagaraLightRendererProperties::GetBoundAttributes()
+{
+	CurrentAttributeBindings.Reset();
+
+	TArray<const FNiagaraVariableAttributeBinding*> AttributeBindings;
+	AttributeBindings.Add(&LightRenderingEnabledBinding);
+	AttributeBindings.Add(&LightExponentBinding);
+	AttributeBindings.Add(&PositionBinding);
+	AttributeBindings.Add(&ColorBinding);
+	AttributeBindings.Add(&RadiusBinding);
+	AttributeBindings.Add(&VolumetricScatteringBinding);
+
+	for (const FNiagaraVariableAttributeBinding* AttributeBinding : AttributeBindings)
+	{
+		if (AttributeBinding->BoundVariable.IsValid())
+		{
+			CurrentAttributeBindings.Add(AttributeBinding->BoundVariable);
+		}
+		else if (AttributeBinding->DataSetVariable.IsValid())
+		{
+			CurrentAttributeBindings.Add(AttributeBinding->DataSetVariable);
+		}
+		else
+		{
+			CurrentAttributeBindings.Add(AttributeBinding->DefaultValueIfNonExistent);
+		}
+	}
+
+	return CurrentAttributeBindings;
+}
+
 #endif // WITH_EDITORONLY_DATA
 
 #undef LOCTEXT_NAMESPACE
\ No newline at end of file
diff --git a/Engine/Plugins/FX/Niagara/Source/Niagara/Private/NiagaraMeshRendererProperties.cpp b/Engine/Plugins/FX/Niagara/Source/Niagara/Private/NiagaraMeshRendererProperties.cpp
--- a/Engine/Plugins/FX/Niagara/Source/Niagara/Private/NiagaraMeshRendererProperties.cpp
+++ b/Engine/Plugins/FX/Niagara/Source/Niagara/Private/NiagaraMeshRendererProperties.cpp
@@ -255,6 +255,43 @@ void UNiagaraMeshRendererProperties::FixMaterial(UMaterial* Material)
 	Material->ForceRecompileForRendering();
 }
 
+const TArray<FNiagaraVariable>& UNiagaraMeshRendererProperties::GetBoundAttributes()
+{
+	CurrentAttributeBindings.Reset();
+
+	TArray<const FNiagaraVariableAttributeBinding*> AttributeBindings;
+	AttributeBindings.Add(&PositionBinding);
+	AttributeBindings.Add(&ColorBinding);
+	AttributeBindings.Add(&VelocityBinding);
+	AttributeBindings.Add(&MeshOrientationBinding);
+	AttributeBindings.Add(&ScaleBinding);
+	AttributeBindings.Add(&DynamicMaterialBinding);
+	AttributeBindings.Add(&DynamicMaterial1Binding);
+	AttributeBindings.Add(&DynamicMaterial2Binding);
+	AttributeBindings.Add(&DynamicMaterial3Binding);
+	AttributeBindings.Add(&MaterialRandomBinding);
+	AttributeBindings.Add(&CustomSortingBinding);
+	AttributeBindings.Add(&NormalizedAgeBinding);
+
+	for (const FNiagaraVariableAttributeBinding* AttributeBinding : AttributeBindings)
+	{
+		if (AttributeBinding->BoundVariable.IsValid())
+		{
+			CurrentAttributeBindings.Add(AttributeBinding->BoundVariable);
+		}
+		else if (AttributeBinding->DataSetVariable.IsValid())
+		{
+			CurrentAttributeBindings.Add(AttributeBinding->DataSetVariable);
+		}
+		else
+		{
+			CurrentAttributeBindings.Add(AttributeBinding->DefaultValueIfNonExistent);
+		}
+	}
+
+	return CurrentAttributeBindings;
+}
+
 const TArray<FNiagaraVariable>& UNiagaraMeshRendererProperties::GetRequiredAttributes()
 {
 	static TArray<FNiagaraVariable> Attrs;
diff --git a/Engine/Plugins/FX/Niagara/Source/Niagara/Private/NiagaraModule.cpp b/Engine/Plugins/FX/Niagara/Source/Niagara/Private/NiagaraModule.cpp
--- a/Engine/Plugins/FX/Niagara/Source/Niagara/Private/NiagaraModule.cpp
+++ b/Engine/Plugins/FX/Niagara/Source/Niagara/Private/NiagaraModule.cpp
@@ -427,6 +427,8 @@ UEnum* FNiagaraTypeDefinition::SimulationTargetEnum;
 UEnum* FNiagaraTypeDefinition::ExecutionStateSourceEnum;
 UEnum* FNiagaraTypeDefinition::ScriptUsageEnum;
 
+UEnum* FNiagaraTypeDefinition::ParameterScopeEnum;
+
 FNiagaraTypeDefinition FNiagaraTypeDefinition::ParameterMapDef;
 FNiagaraTypeDefinition FNiagaraTypeDefinition::IDDef;
 FNiagaraTypeDefinition FNiagaraTypeDefinition::NumericDef;
@@ -543,6 +545,8 @@ void FNiagaraTypeDefinition::Init()
 	ExecutionStateSourceEnum = StaticEnum<ENiagaraExecutionStateSource>();
 	SimulationTargetEnum = StaticEnum<ENiagaraSimTarget>();
 	ScriptUsageEnum = StaticEnum<ENiagaraScriptUsage>();
+
+	ParameterScopeEnum = StaticEnum<ENiagaraParameterScope>();
 	
 	RecreateUserDefinedTypeRegistry();
 }
diff --git a/Engine/Plugins/FX/Niagara/Source/Niagara/Private/NiagaraRibbonRendererProperties.cpp b/Engine/Plugins/FX/Niagara/Source/Niagara/Private/NiagaraRibbonRendererProperties.cpp
--- a/Engine/Plugins/FX/Niagara/Source/Niagara/Private/NiagaraRibbonRendererProperties.cpp
+++ b/Engine/Plugins/FX/Niagara/Source/Niagara/Private/NiagaraRibbonRendererProperties.cpp
@@ -238,6 +238,45 @@ void UNiagaraRibbonRendererProperties::FixMaterial(UMaterial* InMaterial)
 	InMaterial->ForceRecompileForRendering();
 }
 
+const TArray<FNiagaraVariable>& UNiagaraRibbonRendererProperties::GetBoundAttributes()
+{
+	CurrentAttributeBindings.Reset();
+
+	TArray<const FNiagaraVariableAttributeBinding*> AttributeBindings;
+	AttributeBindings.Add(&PositionBinding);
+	AttributeBindings.Add(&ColorBinding);
+	AttributeBindings.Add(&VelocityBinding);
+	AttributeBindings.Add(&NormalizedAgeBinding);
+	AttributeBindings.Add(&RibbonTwistBinding);
+	AttributeBindings.Add(&RibbonWidthBinding);
+	AttributeBindings.Add(&RibbonFacingBinding);
+	AttributeBindings.Add(&RibbonIdBinding);
+	AttributeBindings.Add(&RibbonLinkOrderBinding);
+	AttributeBindings.Add(&MaterialRandomBinding);
+	AttributeBindings.Add(&DynamicMaterialBinding);
+	AttributeBindings.Add(&DynamicMaterial1Binding);
+	AttributeBindings.Add(&DynamicMaterial2Binding);
+	AttributeBindings.Add(&DynamicMaterial3Binding);
+
+	for (const FNiagaraVariableAttributeBinding* AttributeBinding : AttributeBindings)
+	{
+		if (AttributeBinding->BoundVariable.IsValid())
+		{
+			CurrentAttributeBindings.Add(AttributeBinding->BoundVariable);
+		}
+		else if (AttributeBinding->DataSetVariable.IsValid())
+		{
+			CurrentAttributeBindings.Add(AttributeBinding->DataSetVariable);
+		}
+		else
+		{
+			CurrentAttributeBindings.Add(AttributeBinding->DefaultValueIfNonExistent);
+		}
+	}
+
+	return CurrentAttributeBindings;
+}
+
 bool UNiagaraRibbonRendererProperties::CanEditChange(const FProperty* InProperty) const
 {
 
diff --git a/Engine/Plugins/FX/Niagara/Source/Niagara/Private/NiagaraScript.cpp b/Engine/Plugins/FX/Niagara/Source/Niagara/Private/NiagaraScript.cpp
--- a/Engine/Plugins/FX/Niagara/Source/Niagara/Private/NiagaraScript.cpp
+++ b/Engine/Plugins/FX/Niagara/Source/Niagara/Private/NiagaraScript.cpp
@@ -2058,6 +2058,35 @@ FNiagaraShaderScript* UNiagaraScript::AllocateResource()
 }
 
 #if WITH_EDITORONLY_DATA
+TArray<ENiagaraParameterScope> UNiagaraScript::GetUnsupportedParameterScopes() const
+{
+	TArray<ENiagaraParameterScope> UnsupportedParameterScopes;
+	UnsupportedParameterScopes.Add(ENiagaraParameterScope::System);
+	UnsupportedParameterScopes.Add(ENiagaraParameterScope::Emitter);
+	UnsupportedParameterScopes.Add(ENiagaraParameterScope::Particles);
+
+	const TArray<ENiagaraScriptUsage> SupportedUsages = GetSupportedUsageContextsForBitmask(ModuleUsageBitmask);
+	for (ENiagaraScriptUsage SupportedUsage : SupportedUsages)
+	{
+		if (SupportedUsage == ENiagaraScriptUsage::ParticleSpawnScript || SupportedUsage == ENiagaraScriptUsage::ParticleUpdateScript || SupportedUsage == ENiagaraScriptUsage::ParticleSpawnScriptInterpolated
+		|| SupportedUsage == ENiagaraScriptUsage::ParticleGPUComputeScript || SupportedUsage == ENiagaraScriptUsage::ParticleEventScript)
+		{
+			UnsupportedParameterScopes.Empty();
+			return UnsupportedParameterScopes;
+		}
+		else if (SupportedUsage == ENiagaraScriptUsage::EmitterSpawnScript || SupportedUsage == ENiagaraScriptUsage::EmitterUpdateScript)
+		{
+			UnsupportedParameterScopes.Remove(ENiagaraParameterScope::System);
+			UnsupportedParameterScopes.Remove(ENiagaraParameterScope::Emitter);
+		}
+		else if (SupportedUsage == ENiagaraScriptUsage::SystemSpawnScript || SupportedUsage == ENiagaraScriptUsage::SystemUpdateScript)
+		{
+			UnsupportedParameterScopes.Remove(ENiagaraParameterScope::System);
+		}
+	}
+	return UnsupportedParameterScopes;
+}
+
 TArray<ENiagaraScriptUsage> UNiagaraScript::GetSupportedUsageContexts() const
 {
 	return GetSupportedUsageContextsForBitmask(ModuleUsageBitmask);
diff --git a/Engine/Plugins/FX/Niagara/Source/Niagara/Private/NiagaraSpriteRendererProperties.cpp b/Engine/Plugins/FX/Niagara/Source/Niagara/Private/NiagaraSpriteRendererProperties.cpp
--- a/Engine/Plugins/FX/Niagara/Source/Niagara/Private/NiagaraSpriteRendererProperties.cpp
+++ b/Engine/Plugins/FX/Niagara/Source/Niagara/Private/NiagaraSpriteRendererProperties.cpp
@@ -215,6 +215,48 @@ void UNiagaraSpriteRendererProperties::PostEditChangeProperty(struct FPropertyCh
 	Super::PostEditChangeProperty(PropertyChangedEvent);
 }
 
+const TArray<FNiagaraVariable>& UNiagaraSpriteRendererProperties::GetBoundAttributes()
+{
+	CurrentAttributeBindings.Reset();
+
+	TArray<const FNiagaraVariableAttributeBinding*> AttributeBindings;
+	AttributeBindings.Add(&PositionBinding);
+	AttributeBindings.Add(&ColorBinding);
+	AttributeBindings.Add(&VelocityBinding);
+	AttributeBindings.Add(&SpriteRotationBinding);
+	AttributeBindings.Add(&SpriteSizeBinding);
+	AttributeBindings.Add(&SpriteFacingBinding);
+	AttributeBindings.Add(&SpriteAlignmentBinding);
+	AttributeBindings.Add(&SubImageIndexBinding);
+	AttributeBindings.Add(&DynamicMaterialBinding);
+	AttributeBindings.Add(&DynamicMaterial1Binding);
+	AttributeBindings.Add(&DynamicMaterial2Binding);
+	AttributeBindings.Add(&DynamicMaterial3Binding);
+	AttributeBindings.Add(&CameraOffsetBinding);
+	AttributeBindings.Add(&UVScaleBinding);
+	AttributeBindings.Add(&MaterialRandomBinding);
+	AttributeBindings.Add(&CustomSortingBinding);
+	AttributeBindings.Add(&NormalizedAgeBinding);
+
+	for (const FNiagaraVariableAttributeBinding* AttributeBinding : AttributeBindings)
+	{
+		if (AttributeBinding->BoundVariable.IsValid())
+		{
+			CurrentAttributeBindings.Add(AttributeBinding->BoundVariable);
+		}
+		else if (AttributeBinding->DataSetVariable.IsValid())
+		{
+			CurrentAttributeBindings.Add(AttributeBinding->DataSetVariable);
+		}
+		else
+		{
+			CurrentAttributeBindings.Add(AttributeBinding->DefaultValueIfNonExistent);
+		}
+	}
+
+	return CurrentAttributeBindings;
+}
+
 const TArray<FNiagaraVariable>& UNiagaraSpriteRendererProperties::GetRequiredAttributes()
 {
 	static TArray<FNiagaraVariable> Attrs;
diff --git a/Engine/Plugins/FX/Niagara/Source/Niagara/Public/NiagaraLightRendererProperties.h b/Engine/Plugins/FX/Niagara/Source/Niagara/Public/NiagaraLightRendererProperties.h
--- a/Engine/Plugins/FX/Niagara/Source/Niagara/Public/NiagaraLightRendererProperties.h
+++ b/Engine/Plugins/FX/Niagara/Source/Niagara/Public/NiagaraLightRendererProperties.h
@@ -33,6 +33,7 @@ public:
 #if WITH_EDITORONLY_DATA
 	virtual bool IsMaterialValidForRenderer(UMaterial* Material, FText& InvalidMessage) override;
 	virtual void FixMaterial(UMaterial* Material) override;
+	virtual const TArray<FNiagaraVariable>& GetBoundAttributes() override;
 	virtual const TArray<FNiagaraVariable>& GetRequiredAttributes() override;
 	virtual const TArray<FNiagaraVariable>& GetOptionalAttributes() override;
 	virtual void GetRendererWidgets(const FNiagaraEmitterInstance* InEmitter, TArray<TSharedPtr<SWidget>>& OutWidgets, TSharedPtr<FAssetThumbnailPool> InThumbnailPool) const override;
diff --git a/Engine/Plugins/FX/Niagara/Source/Niagara/Public/NiagaraMeshRendererProperties.h b/Engine/Plugins/FX/Niagara/Source/Niagara/Public/NiagaraMeshRendererProperties.h
--- a/Engine/Plugins/FX/Niagara/Source/Niagara/Public/NiagaraMeshRendererProperties.h
+++ b/Engine/Plugins/FX/Niagara/Source/Niagara/Public/NiagaraMeshRendererProperties.h
@@ -95,6 +95,7 @@ public:
 #if WITH_EDITORONLY_DATA
 	virtual bool IsMaterialValidForRenderer(UMaterial* Material, FText& InvalidMessage) override;
 	virtual void FixMaterial(UMaterial* Material) override;
+	virtual const TArray<FNiagaraVariable>& GetBoundAttributes() override;
 	virtual const TArray<FNiagaraVariable>& GetRequiredAttributes() override;
 	virtual const TArray<FNiagaraVariable>& GetOptionalAttributes() override;
 	virtual	void GetRendererWidgets(const FNiagaraEmitterInstance* InEmitter, TArray<TSharedPtr<SWidget>>& OutWidgets, TSharedPtr<FAssetThumbnailPool> InThumbnailPool) const override;
diff --git a/Engine/Plugins/FX/Niagara/Source/Niagara/Public/NiagaraRendererProperties.h b/Engine/Plugins/FX/Niagara/Source/Niagara/Public/NiagaraRendererProperties.h
--- a/Engine/Plugins/FX/Niagara/Source/Niagara/Public/NiagaraRendererProperties.h
+++ b/Engine/Plugins/FX/Niagara/Source/Niagara/Public/NiagaraRendererProperties.h
@@ -50,6 +50,7 @@ public:
 
 	virtual void FixMaterial(UMaterial* Material) { }
 
+	virtual const TArray<FNiagaraVariable>& GetBoundAttributes() { static TArray<FNiagaraVariable> Vars; return Vars; };
 	virtual const TArray<FNiagaraVariable>& GetRequiredAttributes() { static TArray<FNiagaraVariable> Vars; return Vars; };
 	virtual const TArray<FNiagaraVariable>& GetOptionalAttributes() { static TArray<FNiagaraVariable> Vars; return Vars; };
 
@@ -80,4 +81,8 @@ public:
 	/** Is motion blur enabled on this renderer or not, the material must also have motion blur enabled. */
 	UPROPERTY(EditAnywhere, Category = "Rendering")
 	bool bMotionBlurEnabled;
+
+protected:
+	// Copy of variables in the attribute binding, updated when GetRequiredAttributes() is called.
+	TArray<FNiagaraVariable> CurrentAttributeBindings;
 };
diff --git a/Engine/Plugins/FX/Niagara/Source/Niagara/Public/NiagaraRibbonRendererProperties.h b/Engine/Plugins/FX/Niagara/Source/Niagara/Public/NiagaraRibbonRendererProperties.h
--- a/Engine/Plugins/FX/Niagara/Source/Niagara/Public/NiagaraRibbonRendererProperties.h
+++ b/Engine/Plugins/FX/Niagara/Source/Niagara/Public/NiagaraRibbonRendererProperties.h
@@ -83,6 +83,7 @@ public:
 #if WITH_EDITOR
 	virtual bool IsMaterialValidForRenderer(UMaterial* Material, FText& InvalidMessage) override;
 	virtual void FixMaterial(UMaterial* Material);
+	virtual const TArray<FNiagaraVariable>& GetBoundAttributes() override;
 	virtual const TArray<FNiagaraVariable>& GetRequiredAttributes() override;
 	virtual const TArray<FNiagaraVariable>& GetOptionalAttributes() override;
 	virtual void GetRendererWidgets(const FNiagaraEmitterInstance* InEmitter, TArray<TSharedPtr<SWidget>>& OutWidgets, TSharedPtr<FAssetThumbnailPool> InThumbnailPool) const override;
diff --git a/Engine/Plugins/FX/Niagara/Source/Niagara/Public/NiagaraSpriteRendererProperties.h b/Engine/Plugins/FX/Niagara/Source/Niagara/Public/NiagaraSpriteRendererProperties.h
--- a/Engine/Plugins/FX/Niagara/Source/Niagara/Public/NiagaraSpriteRendererProperties.h
+++ b/Engine/Plugins/FX/Niagara/Source/Niagara/Public/NiagaraSpriteRendererProperties.h
@@ -68,6 +68,7 @@ public:
 #if WITH_EDITOR
 	virtual bool IsMaterialValidForRenderer(UMaterial* Material, FText& InvalidMessage) override;
 	virtual void FixMaterial(UMaterial* Material) override;
+	virtual const TArray<FNiagaraVariable>& GetBoundAttributes() override;
 	virtual const TArray<FNiagaraVariable>& GetRequiredAttributes() override;
 	virtual const TArray<FNiagaraVariable>& GetOptionalAttributes() override;
 	virtual void GetRendererWidgets(const FNiagaraEmitterInstance* InEmitter, TArray<TSharedPtr<SWidget>>& OutWidgets, TSharedPtr<FAssetThumbnailPool> InThumbnailPool) const override;
diff --git a/Engine/Plugins/FX/Niagara/Source/Niagara/Public/NiagaraTypes.h b/Engine/Plugins/FX/Niagara/Source/Niagara/Public/NiagaraTypes.h
--- a/Engine/Plugins/FX/Niagara/Source/Niagara/Public/NiagaraTypes.h
+++ b/Engine/Plugins/FX/Niagara/Source/Niagara/Public/NiagaraTypes.h
@@ -10,6 +10,7 @@
 #include "NiagaraTypes.generated.h"
 
 class UNiagaraDataInterfaceBase;
+
 DECLARE_LOG_CATEGORY_EXTERN(LogNiagara, Log, Verbose);
 
 // basic type struct definitions
@@ -439,6 +440,52 @@ public:
 };
 
 
+UENUM()
+enum class ENiagaraParameterScope : uint32
+{
+	Input,
+
+	User,
+
+	Engine,
+
+	System,
+
+	Emitter,
+
+	Particles,
+
+	ScriptPersistent UMETA(Hidden), //@todo(ng) hiding until autotest verification is made.
+
+	ScriptTransient,
+
+	Local UMETA(Hidden), //Convenience markup for ScopeToString functions, only use in conjunction with ENiagaraScriptParameterUsage::Local.
+
+	// insert new scopes before
+	None UMETA(Hidden),
+
+	Num UMETA(Hidden)
+};
+
+UENUM()
+enum class ENiagaraScriptParameterUsage : uint32
+{
+	Input,
+
+	Output,
+
+	Local,
+
+	InputOutput,
+
+	InitialValueInput,
+
+	// insert new script parameter usages before
+	None UMETA(Hidden),
+	
+	Num UMETA(Hidden)
+};
+
 /** Defines options for conditionally editing and showing script inputs in the UI. */
 USTRUCT()
 struct NIAGARA_API FNiagaraInputConditionMetadata
@@ -463,10 +510,14 @@ public:
 		: bAdvancedDisplay(false)
 		, EditorSortPriority(0)
 		, bInlineEditConditionToggle(false)
+		, Scope(ENiagaraParameterScope::None)
+		, Usage(ENiagaraScriptParameterUsage::None)
 		, bIsStaticSwitch(false)
 		, StaticSwitchDefaultValue(0)
-	{
-	}
+		, bAddedToNodeGraphDeepCopy(false)
+		, bOutputIsPersistent(false)
+		, bCreatedInSystemEditor(false)
+	{};
 public:
 	UPROPERTY(EditAnywhere, Category = "Variable", meta = (MultiLine = true, SkipForCompileHash = "true"))
 	FText Description;
@@ -497,6 +548,14 @@ public:
 	UPROPERTY(EditAnywhere, Category = "Variable", DisplayName = "Property Metadata", meta = (ToolTip = "Property Metadata", SkipForCompileHash = "true"))
 	TMap<FName, FString> PropertyMetaData;
 
+	/** Defines the scope of a variable that is an input to a script. */
+	UPROPERTY()
+	ENiagaraParameterScope Scope;
+
+	/** Defines the usage of a variable as an argument or output relative to the script. */
+	UPROPERTY()
+	ENiagaraScriptParameterUsage Usage;
+
 	/** This is a read-only variable that designates if the metadata is tied to a static switch or not. */
 	UPROPERTY()
 	bool bIsStaticSwitch; // TODO: This should be moved to the UNiagaraScriptVariable in the future
@@ -504,6 +563,26 @@ public:
 	/** The default value to use when creating new pins or stack entries for a static switch parameter */
 	UPROPERTY()
 	int32 StaticSwitchDefaultValue;  // TODO: This should be moved to the UNiagaraScriptVariable in the future
+
+	/** Transient data to mark variables set in the node graph deep copy as having been derived from a module namespace parameter default. */
+	UPROPERTY(Transient)
+	bool bAddedToNodeGraphDeepCopy;
+
+	/** Only valid if Usage is Output. Marks the associated FNiagaraVariable as Persistent across script runs and therefore should be retained in the Dataset during compilation/translation. */
+	UPROPERTY()
+	bool bOutputIsPersistent;
+
+	/** Namespace-less name for associated FNiagaraVariable. Edited directly by user and then used to generate full Name of associated FNiagaraVariable. */
+	UPROPERTY()
+	FName CachedNamespacelessVariableName;
+
+	/** Track if the associated parameter was created in the Emitter/System editor. Used to determine whether the associated parameter can be deleted from the Emitter/System editor. */
+	UPROPERTY()
+	bool bCreatedInSystemEditor;
+
+public:
+	FORCEINLINE bool IsInputUsage() const { return Usage == ENiagaraScriptParameterUsage::Input || Usage == ENiagaraScriptParameterUsage::InputOutput; };
+	FORCEINLINE bool IsInputOrLocalUsage() const { return Usage == ENiagaraScriptParameterUsage::Input || Usage == ENiagaraScriptParameterUsage::InputOutput || Usage == ENiagaraScriptParameterUsage::InitialValueInput || Usage == ENiagaraScriptParameterUsage::Local; };
 };
 
 USTRUCT()
@@ -749,6 +828,8 @@ public:
 	static UEnum* GetSimulationTargetEnum() { return SimulationTargetEnum; }
 	static UEnum* GetScriptUsageEnum() { return ScriptUsageEnum; }
 
+	static UEnum* GetParameterScopeEnum() { return ParameterScopeEnum; }
+
 	static const FNiagaraTypeDefinition& GetCollisionEventDef() { return CollisionEventDef; }
 
 	static bool IsScalarDefinition(const FNiagaraTypeDefinition& Type);
@@ -805,6 +886,8 @@ private:
 	static UEnum* ExecutionStateEnum;
 	static UEnum* ExecutionStateSourceEnum;
 
+	static UEnum* ParameterScopeEnum;
+
 	static UScriptStruct* ParameterMapStruct;
 	static UScriptStruct* IDStruct;
 
diff --git a/Engine/Plugins/FX/Niagara/Source/NiagaraCore/Public/NiagaraCustomVersion.h b/Engine/Plugins/FX/Niagara/Source/NiagaraCore/Public/NiagaraCustomVersion.h
--- a/Engine/Plugins/FX/Niagara/Source/NiagaraCore/Public/NiagaraCustomVersion.h
+++ b/Engine/Plugins/FX/Niagara/Source/NiagaraCore/Public/NiagaraCustomVersion.h
@@ -141,6 +141,9 @@ struct FNiagaraCustomVersion
 		AddGeneratedFunctionsToGPUParamInfo, // Save the functions generated by a GPU data interface inside FNiagaraDataInterfaceGPUParamInfo
 
 		PlatformScalingRefactor, //Removed DetailLevel in favor of FNiagaraPlatfomSet based selection of per platform settings.
+
+		PrecompileNamespaceFixup, // Promote parameters used across script executions to the Dataset, and Demote unused parameters.
+
 		// DO NOT ADD A NEW VERSION UNLESS YOU HAVE TALKED TO THE NIAGARA LEAD. Mismanagement of these versions can lead to data loss if it is adjusted in multiple streams simultaneously.
 		// -----<new versions can be added above this line>  -------------------------------------------------
 		VersionPlusOne,
diff --git a/Engine/Plugins/FX/Niagara/Source/NiagaraEditor/Private/NiagaraActions.cpp b/Engine/Plugins/FX/Niagara/Source/NiagaraEditor/Private/NiagaraActions.cpp
--- a/Engine/Plugins/FX/Niagara/Source/NiagaraEditor/Private/NiagaraActions.cpp
+++ b/Engine/Plugins/FX/Niagara/Source/NiagaraEditor/Private/NiagaraActions.cpp
@@ -9,6 +9,7 @@
 #include "Framework/Application/MenuStack.h"
 #include "Framework/Application/SlateApplication.h"
 #include "Layout/WidgetPath.h"
+#include "NiagaraParameterPanelViewModel.h"
 
 #define LOCTEXT_NAMESPACE "NiagaraActions"
 
@@ -38,6 +39,24 @@ FNiagaraParameterAction::FNiagaraParameterAction(const FNiagaraVariable& InParam
 {
 }
 
+FNiagaraParameterAction::FNiagaraParameterAction(const FNiagaraVariable& InParameter,
+	FText InNodeCategory, FText InMenuDesc, FText InToolTip, const int32 InGrouping, FText InKeywords, int32 InSectionID)
+	: FEdGraphSchemaAction(MoveTemp(InNodeCategory), MoveTemp(InMenuDesc), MoveTemp(InToolTip), InGrouping, MoveTemp(InKeywords), InSectionID)
+	, Parameter(InParameter)
+{
+}
+
+/************************************************************************/
+/* FNiagaraScriptVarAndViewInfoAction									*/
+/************************************************************************/
+FNiagaraScriptVarAndViewInfoAction::FNiagaraScriptVarAndViewInfoAction(const FNiagaraScriptVariableAndViewInfo& InScriptVariableAndViewInfo,
+	FText InNodeCategory, FText InMenuDesc, FText InToolTip, const int32 InGrouping, FText InKeywords, int32 InSectionID)
+	: FEdGraphSchemaAction(MoveTemp(InNodeCategory), MoveTemp(InMenuDesc), MoveTemp(InToolTip), InGrouping, MoveTemp(InKeywords), InSectionID)
+	, ScriptVariableAndViewInfo(InScriptVariableAndViewInfo)
+{
+}
+
+
 /************************************************************************/
 /* FNiagaraParameterGraphDragOperation									*/
 /************************************************************************/
diff --git a/Engine/Plugins/FX/Niagara/Source/NiagaraEditor/Private/NiagaraCompiler.cpp b/Engine/Plugins/FX/Niagara/Source/NiagaraEditor/Private/NiagaraCompiler.cpp
--- a/Engine/Plugins/FX/Niagara/Source/NiagaraEditor/Private/NiagaraCompiler.cpp
+++ b/Engine/Plugins/FX/Niagara/Source/NiagaraEditor/Private/NiagaraCompiler.cpp
@@ -31,6 +31,7 @@
 #include "ShaderCompiler.h"
 #include "NiagaraShader.h"
 #include "NiagaraScript.h"
+#include "NiagaraRendererProperties.h"
 #include "NiagaraSimulationStageBase.h"
 #include "Serialization/MemoryReader.h"
 #include "HAL/ThreadSafeBool.h"
@@ -441,8 +442,9 @@ void FNiagaraCompileRequestData::FinishPrecompile(UNiagaraScriptSource* ScriptSo
 		for (UNiagaraNodeOutput* FoundOutputNode : OutputNodes)
 		{
 			// Map all for this output node
-			FNiagaraParameterMapHistoryBuilder Builder;
+			FNiagaraParameterMapHistoryWithMetaDataBuilder Builder;
 			Builder.ConstantResolver = ConstantResolver;
+			Builder.AddGraphToCallingGraphContextStack(NodeGraphDeepCopy);
 			Builder.RegisterEncounterableVariables(EncounterableVariables);
 
 			FString TranslationName = TEXT("Emitter");
@@ -455,12 +457,12 @@ void FNiagaraCompileRequestData::FinishPrecompile(UNiagaraScriptSource* ScriptSo
 
 			for (FNiagaraParameterMapHistory& History : Histories)
 			{
+				History.OriginatingScriptUsage = FoundOutputNode->GetUsage();
 				for (FNiagaraVariable& Var : History.Variables)
 				{
 					check(Var.GetType() != FNiagaraTypeDefinition::GetGenericNumericDef());
 				}
 			}
-
 			if (FoundOutputNode->GetUsage() == ENiagaraScriptUsage::ParticleSimulationStageScript)
 			{
 				NumSimStageNodes++;
@@ -534,7 +536,6 @@ void FNiagaraCompileRequestData::FinishPrecompile(UNiagaraScriptSource* ScriptSo
 				}
 			}
 		}
-
 	}
 }
 
@@ -577,7 +578,6 @@ TSharedPtr<FNiagaraCompileRequestDataBase, ESPMode::ThreadSafe> FNiagaraEditorMo
 		TArray<FNiagaraVariable> EncounterableVars(OriginalExposedParams);
 
 		// Create an array of variables that we *did* encounter when traversing the graphs.
-		TArray<FNiagaraVariable> EncounteredExposedVars;
 		check(System->GetSystemSpawnScript()->GetSource() == System->GetSystemUpdateScript()->GetSource());
 
 		// First deep copy all the emitter graphs referenced by the system so that we can later hook up emitter handles in the system traversal.
@@ -615,7 +615,7 @@ TSharedPtr<FNiagaraCompileRequestDataBase, ESPMode::ThreadSafe> FNiagaraEditorMo
 			FCompileConstantResolver ConstantResolver(Handle.GetInstance());
 			if (Handle.GetIsEnabled()) // Don't pull in the emitter if it isn't going to be used.
 			{
-               TArray<UNiagaraScript*> EmitterScripts;
+				TArray<UNiagaraScript*> EmitterScripts;
 				Handle.GetInstance()->GetScripts(EmitterScripts, false);
 
 				for (int32 ScriptIdx = 0; ScriptIdx < EmitterScripts.Num(); ScriptIdx++)
@@ -627,9 +627,89 @@ TSharedPtr<FNiagaraCompileRequestDataBase, ESPMode::ThreadSafe> FNiagaraEditorMo
 				}
 				BasePtr->EmitterData[i]->FinishPrecompile(Cast<UNiagaraScriptSource>(Handle.GetInstance()->GraphSource), EncounterableVars, ENiagaraScriptUsage::EmitterSpawnScript, ConstantResolver, &Handle.GetInstance()->GetSimulationStages());
 				BasePtr->MergeInEmitterPrecompiledData(BasePtr->EmitterData[i].Get());
+
+				for (UNiagaraRendererProperties* RendererProperty : Handle.GetInstance()->GetRenderers())
+				{
+					BasePtr->EmitterData[i]->RequiredRendererVariables.Append(RendererProperty->GetBoundAttributes());
+				}
+			}
+		}
+
+		// Build parameter map histories just for System scripts in order to separate out System scope parameters from Emitter scope for the history handles.
+		TArray<FNiagaraParameterMapHistory> SystemOnlyHistories;
+		FNiagaraParameterMapHistoryWithMetaDataBuilder Builder;
+		Builder.ConstantResolver = EmptyResolver;
+		Builder.bShouldBuildSubHistories = false;
+		FString TranslationName = TEXT("Emitter");
+		Builder.AddGraphToCallingGraphContextStack(BasePtr->GetPrecomputedNodeGraph());
+		Builder.BeginTranslation(TranslationName);
+
+		TArray<UNiagaraNodeOutput*> SystemOutputNodes;
+		BasePtr->NodeGraphDeepCopy->FindOutputNodes(SystemOutputNodes);
+		for (UNiagaraNodeOutput* FoundOutputNode : SystemOutputNodes)
+		{
+			Builder.EnableScriptWhitelist(true, FoundOutputNode->GetUsage());
+			Builder.BuildParameterMaps(FoundOutputNode, true);
+			
+			for (FNiagaraParameterMapHistory& History : Builder.Histories)
+			{
+				// FinishPrecompile() has not been called on this history so set the OriginatingScriptUsage.
+				History.OriginatingScriptUsage = FoundOutputNode->GetUsage();
+				SystemOnlyHistories.Add(History);
+			}
+			Builder.Histories.Empty();
+		}
+
+		// Build history handles for System scripts.
+		TArray<FNiagaraParameterMapHistory>& SystemScriptHistories = BasePtr->GetPrecomputedHistories();
+		TArray<FNiagaraParameterMapHistoryHandle> SystemScriptHistoryHandles;
+		for (int i = 0; i < SystemScriptHistories.Num(); ++i)
+		{
+			SystemScriptHistoryHandles.Add(FNiagaraParameterMapHistoryHandle(SystemScriptHistories[i], SystemOnlyHistories[i], BasePtr->RequiredRendererVariables));
+		}
+
+		// Build history handles for Emitter and Particle scripts.
+		TArray<FNiagaraParameterMapHistoryHandle> EmitterAndParticleScriptHistoryHandles;
+		for (TSharedPtr<FNiagaraCompileRequestData, ESPMode::ThreadSafe> EmitterRequest : BasePtr->EmitterData)
+		{
+			FName EmitterUniqueName = FName(*EmitterRequest->GetUniqueEmitterName());
+			TArray<FNiagaraParameterMapHistory>& EmitterScriptHistories = EmitterRequest->GetPrecomputedHistories();
+			for (int i = 0; i < EmitterScriptHistories.Num(); ++i)
+			{
+				if (FNiagaraStackGraphUtilities::GetScopeForScriptUsage(EmitterScriptHistories[i].OriginatingScriptUsage) == ENiagaraParameterScope::Particles)
+				{
+					EmitterAndParticleScriptHistoryHandles.Add(FNiagaraParameterMapHistoryHandle(EmitterScriptHistories[i], EmitterRequest->RequiredRendererVariables, EmitterUniqueName));
+				}
+				else
+				{
+					FNiagaraParameterMapHistory* SameUsageSystemHistory = SystemScriptHistories.FindByPredicate([&EmitterScriptHistories, &i](const FNiagaraParameterMapHistory& SystemScriptHistory){
+						ENiagaraScriptUsage EmitterOriginatingScriptUsage = EmitterScriptHistories[i].OriginatingScriptUsage;
+						if (EmitterOriginatingScriptUsage == ENiagaraScriptUsage::EmitterSpawnScript)
+						{
+							return SystemScriptHistory.OriginatingScriptUsage == ENiagaraScriptUsage::SystemSpawnScript;
+						}
+						else if (EmitterOriginatingScriptUsage == ENiagaraScriptUsage::EmitterUpdateScript)
+						{
+							return SystemScriptHistory.OriginatingScriptUsage == ENiagaraScriptUsage::SystemUpdateScript;
+						}
+						else
+						{
+							checkf(false, TEXT("Encountered unexpected script usage for history!"));
+						}
+						return false;
+					});
+
+					checkf(SameUsageSystemHistory != nullptr, TEXT("Failed to find System script with same run index as Emitter script!"));
+					EmitterAndParticleScriptHistoryHandles.Add(FNiagaraParameterMapHistoryHandle(*SameUsageSystemHistory, EmitterScriptHistories[i], EmitterRequest->RequiredRendererVariables, EmitterUniqueName));
+				}
 			}
 		}
 
+		// Fixup namespaces of parameters that will be used in the histories for translation.
+		TArray<FNiagaraParameterMapHistoryHandle> AllHistoryHandles;
+		AllHistoryHandles.Append(SystemScriptHistoryHandles);
+		AllHistoryHandles.Append(EmitterAndParticleScriptHistoryHandles);
+		FNiagaraParameterMapHistoryWithMetaDataBuilder::FixupHistoryVariableNamespaces(AllHistoryHandles, System->bCullDatasetParameters);
 	}
 
 	UE_LOG(LogNiagaraEditor, Verbose, TEXT("'%s' Precompile took %f sec."), *InObj->GetOutermost()->GetName(),
diff --git a/Engine/Plugins/FX/Niagara/Source/NiagaraEditor/Private/NiagaraEditorCommon.cpp b/Engine/Plugins/FX/Niagara/Source/NiagaraEditor/Private/NiagaraEditorCommon.cpp
--- a/Engine/Plugins/FX/Niagara/Source/NiagaraEditor/Private/NiagaraEditorCommon.cpp
+++ b/Engine/Plugins/FX/Niagara/Source/NiagaraEditor/Private/NiagaraEditorCommon.cpp
@@ -16,6 +16,13 @@ DEFINE_LOG_CATEGORY(LogNiagaraEditor);
 TMap<FName, int32> FNiagaraOpInfo::OpInfoMap;
 TArray<FNiagaraOpInfo> FNiagaraOpInfo::OpInfos;
 
+int32 FNiagaraEditorCommonCVar::GNiagaraEnableParameterPanel2 = 1;
+FAutoConsoleVariableRef FNiagaraEditorCommonCVar::CVarEnableParameterPanel2(
+	TEXT("fx.NiagaraEnableParameterPanel2"),
+	FNiagaraEditorCommonCVar::GNiagaraEnableParameterPanel2,
+	TEXT("If true, enable enhanced parameter panel features and dynamic parameter linking in precompiler."),
+	ECVF_Default);
+
 const FNiagaraOpInfo* FNiagaraOpInfo::GetOpInfo(FName OpName)
 {
 	const int32* Idx = OpInfoMap.Find(OpName);
diff --git a/Engine/Plugins/FX/Niagara/Source/NiagaraEditor/Private/NiagaraEditorUtilities.cpp b/Engine/Plugins/FX/Niagara/Source/NiagaraEditor/Private/NiagaraEditorUtilities.cpp
--- a/Engine/Plugins/FX/Niagara/Source/NiagaraEditor/Private/NiagaraEditorUtilities.cpp
+++ b/Engine/Plugins/FX/Niagara/Source/NiagaraEditor/Private/NiagaraEditorUtilities.cpp
@@ -1838,6 +1838,60 @@ void FNiagaraEditorUtilities::CreateAssetFromEmitter(TSharedRef<FNiagaraEmitterH
 	}
 }
 
+void FNiagaraEditorUtilities::GetScriptRunAndExecutionIndexFromUsage(const ENiagaraScriptUsage& InUsage, int32& OutRunIndex, int32&OutExecutionIndex)
+{
+	switch (InUsage) //@todo(ng) fix numeric scheme Execution Index needing to continuously increment
+	{
+	case ENiagaraScriptUsage::SystemSpawnScript:
+		OutRunIndex = 0;
+		OutExecutionIndex = 0;
+		break;
+	case ENiagaraScriptUsage::EmitterSpawnScript:
+		OutRunIndex = 0;
+		OutExecutionIndex = 1;
+		break;
+	case ENiagaraScriptUsage::ParticleSpawnScript:
+	case ENiagaraScriptUsage::ParticleSpawnScriptInterpolated:
+		OutRunIndex = 2;
+		OutExecutionIndex = 2;
+		break;
+	case ENiagaraScriptUsage::SystemUpdateScript:
+		OutRunIndex = 1;
+		OutExecutionIndex = 0;
+		break;
+	case ENiagaraScriptUsage::EmitterUpdateScript:
+		OutRunIndex = 1;
+		OutExecutionIndex = 1;
+		break;
+	case ENiagaraScriptUsage::ParticleUpdateScript:
+	case ENiagaraScriptUsage::ParticleGPUComputeScript:
+		OutRunIndex = 2;
+		OutExecutionIndex = 3;
+		break;
+	case ENiagaraScriptUsage::ParticleEventScript:
+		OutRunIndex = 2;
+		OutExecutionIndex = 4;
+		break;
+	case ENiagaraScriptUsage::ParticleSimulationStageScript:
+		//@todo(ng) implement getter for shader stages; for now same as particle update
+		
+		OutRunIndex = 2;
+		OutExecutionIndex = 4;
+		break;
+	case ENiagaraScriptUsage::DynamicInput:
+	case ENiagaraScriptUsage::Function:
+	case ENiagaraScriptUsage::Module:
+		OutRunIndex = INDEX_NONE;
+		OutExecutionIndex = INDEX_NONE;
+		break;
+	default:
+		checkf(false, TEXT("No execution index implemented for usage!"));
+		OutRunIndex = INDEX_NONE;
+		OutExecutionIndex = INDEX_NONE;
+		break;
+	}
+}
+
 bool FNiagaraEditorUtilities::AddEmitterContextMenuActions(FMenuBuilder& MenuBuilder, const TSharedPtr<FNiagaraEmitterHandleViewModel>& EmitterHandleViewModelPtr)
 {
 	if (EmitterHandleViewModelPtr.IsValid())
diff --git a/Engine/Plugins/FX/Niagara/Source/NiagaraEditor/Private/NiagaraGraph.cpp b/Engine/Plugins/FX/Niagara/Source/NiagaraEditor/Private/NiagaraGraph.cpp
--- a/Engine/Plugins/FX/Niagara/Source/NiagaraEditor/Private/NiagaraGraph.cpp
+++ b/Engine/Plugins/FX/Niagara/Source/NiagaraEditor/Private/NiagaraGraph.cpp
@@ -327,6 +327,40 @@ void UNiagaraGraph::PostLoad()
 		}
 	}
 
+	if (NiagaraVer < FNiagaraCustomVersion::PrecompileNamespaceFixup)
+	{
+		// Collect all input and output pins to infer usages of variables they reference.
+		const TMap<FNiagaraVariable, FInputPinsAndOutputPins> VarToPinsMap = CollectVarsToInOutPinsMap();
+
+		// Now iterate all variables
+		for (auto It = VariableToScriptVariable.CreateIterator(); It; ++It)
+		{
+			FNiagaraVariable Var = It.Key();
+			UNiagaraScriptVariable* ScriptVar = It.Value();
+
+			if (Var.GetName() == FNiagaraConstants::InputPinName || Var.GetName() == FNiagaraConstants::OutputPinName)
+			{
+				//@todo Pins leaked into variable maps at some point, need to clean.
+				continue;
+			}
+
+			if (ScriptVar->Metadata.bIsStaticSwitch)
+			{
+				// Do not need to assign scopes and usages to static switches, so assign the name and continue.
+				ScriptVar->Metadata.CachedNamespacelessVariableName = ScriptVar->Variable.GetName();
+				continue;
+			}
+
+			// Resolve the namespace of the variable to a scope, and assign the usage to local if "local" is in the name string.
+			FNiagaraStackGraphUtilities::SetParameterMetaDataFromName(Var.GetName(), ScriptVar->Metadata.Scope, ScriptVar->Metadata.Usage, ScriptVar->Metadata.CachedNamespacelessVariableName);
+
+			// If the usage is not local, select the usage based on the associated input/output pins.
+			SetScriptVariableUsageForPins(VarToPinsMap, ScriptVar);
+		}
+		
+		NotifyGraphNeedsRecompile();
+	}
+
 	InvalidateCachedParameterData();
 }
 
@@ -754,6 +788,52 @@ TArray<UEdGraphNode*> UNiagaraGraph::FindReachbleNodes() const
 	return ResultNodes;
 }
 
+bool UNiagaraGraph::UpdateUsageForScriptVariable(UNiagaraScriptVariable* ScriptVariable) const
+{
+	const TMap<FNiagaraVariable, FInputPinsAndOutputPins> VarToPinsMap = CollectVarsToInOutPinsMap();
+	return SetScriptVariableUsageForPins(VarToPinsMap, ScriptVariable);
+}
+
+bool UNiagaraGraph::SetScriptVariableUsageForPins(const TMap<FNiagaraVariable, FInputPinsAndOutputPins>& VarToPinsMap, UNiagaraScriptVariable* ScriptVariable) const
+{
+	if (ScriptVariable->Metadata.Usage == ENiagaraScriptParameterUsage::Local || ScriptVariable->Metadata.Usage == ENiagaraScriptParameterUsage::InitialValueInput)
+	{
+		return true;
+	}
+
+	const FInputPinsAndOutputPins* InOutPins = VarToPinsMap.Find(ScriptVariable->Variable);
+	if (InOutPins != nullptr)
+	{
+		if (InOutPins->InputPins.Num() > 0)
+		{
+			if (InOutPins->OutputPins.Num() > 0)
+			{
+				ScriptVariable->Metadata.Usage = ENiagaraScriptParameterUsage::InputOutput;
+				return true;
+			}
+			ScriptVariable->Metadata.Usage = ENiagaraScriptParameterUsage::Output;
+			return true;
+		}
+		else if (InOutPins->OutputPins.Num() > 0)
+		{
+			ScriptVariable->Metadata.Usage = ENiagaraScriptParameterUsage::Input;
+			return true;
+		}
+		else
+		{
+			ensureMsgf(false, TEXT("Var to pins map had entry with 0 pins!"));
+			ScriptVariable->Metadata.Usage = ENiagaraScriptParameterUsage::Local;
+			return true;
+		}
+	}
+	else
+	{
+		ScriptVariable->Metadata.Usage = ENiagaraScriptParameterUsage::Local;
+		return false;
+	}
+
+}
+
 void UNiagaraGraph::GetParameters(TArray<FNiagaraVariable>& Inputs, TArray<FNiagaraVariable>& Outputs)const
 {
 	Inputs.Empty();
@@ -828,7 +908,7 @@ UNiagaraScriptVariable* UNiagaraGraph::GetScriptVariable(FName ParameterName) co
 	return nullptr;
 }
 
-void UNiagaraGraph::AddParameter(const FNiagaraVariable& Parameter, bool bIsStaticSwitch)
+void UNiagaraGraph::AddParameter(const FNiagaraVariable& Parameter, bool bIsStaticSwitch /*= false*/)
 {
 	FNiagaraGraphParameterReferenceCollection* FoundParameterReferenceCollection = ParameterToReferencesMap.Find(Parameter);
 	if (!FoundParameterReferenceCollection)
@@ -848,6 +928,77 @@ void UNiagaraGraph::AddParameter(const FNiagaraVariable& Parameter, bool bIsStat
 	}
 }
 
+void UNiagaraGraph::AddParameter(FNiagaraVariable& Parameter, const FAddParameterOptions Options)
+{
+	FNiagaraGraphParameterReferenceCollection* FoundParameterReferenceCollection = ParameterToReferencesMap.Find(Parameter);
+	if (!FoundParameterReferenceCollection)
+	{
+		FNiagaraGraphParameterReferenceCollection NewReferenceCollection = FNiagaraGraphParameterReferenceCollection(true /*bCreated*/);
+		NewReferenceCollection.Graph = this;
+		ParameterToReferencesMap.Add(Parameter, NewReferenceCollection);
+	}
+
+	UNiagaraScriptVariable** FoundScriptVariable = VariableToScriptVariable.Find(Parameter);
+	if (!FoundScriptVariable)
+	{
+		if (Options.bMakeParameterNameUnique)
+		{
+			TSet<FName> Names;
+			for (const auto& ParameterElement : GetParameterReferenceMap())
+			{
+				Names.Add(ParameterElement.Key.GetName());
+			}
+			Parameter.SetName(FNiagaraUtilities::GetUniqueName(Parameter.GetName(), Names));
+		}
+
+		UNiagaraScriptVariable* NewScriptVariable = NewObject<UNiagaraScriptVariable>(this, FName(), RF_Transactional);
+		NewScriptVariable->Variable = Parameter;
+		NewScriptVariable->Metadata.bIsStaticSwitch = Options.bIsStaticSwitch;
+		NewScriptVariable->Metadata.bCreatedInSystemEditor = Options.bAddedFromSystemEditor;
+		VariableToScriptVariable.Add(Parameter, NewScriptVariable);
+
+		bool bSkipRefreshMetaDataScopeAndUsage = false;
+		if (Options.NewParameterUsage.IsSet())
+		{
+			bSkipRefreshMetaDataScopeAndUsage = true;
+			NewScriptVariable->Metadata.Usage = Options.NewParameterUsage.GetValue();
+		}
+
+		if (Options.NewParameterScope.IsSet())
+		{
+			bSkipRefreshMetaDataScopeAndUsage = true;
+
+			bool bIsInitialValue = false;
+			if (Options.NewParameterUsage.IsSet())
+			{
+				bIsInitialValue = Options.NewParameterUsage == ENiagaraScriptParameterUsage::InitialValueInput;
+			}
+			
+			const FName NewScriptVariableName = FNiagaraStackGraphUtilities::GetVariableNameForScope(NewScriptVariable->Variable.GetName(), Options.NewParameterScope.GetValue(), bIsInitialValue);
+			RenameParameter(NewScriptVariable->Variable, NewScriptVariableName, Options.bIsStaticSwitch, Options.NewParameterScope.GetValue());
+		}
+
+		if (bSkipRefreshMetaDataScopeAndUsage == false && Options.bRefreshMetaDataScopeAndUsage)
+		{
+			GenerateMetaDataForScriptVariable(NewScriptVariable);
+		}
+	}
+	else if (Options.bRefreshMetaDataScopeAndUsage)
+	{
+		// The script variable already exists so do not add it, but refresh the existing metadata scope and usage as the option is set.
+		UpdateUsageForScriptVariable(*FoundScriptVariable);
+	}
+}
+
+void UNiagaraGraph::AddParameterReference(const FNiagaraVariable& Parameter, FNiagaraGraphParameterReference& NewParameterReference)
+{
+	FNiagaraGraphParameterReferenceCollection* FoundParameterReferenceCollection = ParameterToReferencesMap.Find(Parameter);
+	if (ensureMsgf(FoundParameterReferenceCollection != nullptr, TEXT("Failed to find parameter reference collection when adding graph parameter reference!")))
+	{
+		FoundParameterReferenceCollection->ParameterReferences.Add(NewParameterReference);
+	}
+}
+
 void UNiagaraGraph::RemoveParameter(const FNiagaraVariable& Parameter)
 {
 	FNiagaraGraphParameterReferenceCollection* ReferenceCollection = ParameterToReferencesMap.Find(Parameter);
@@ -893,7 +1044,7 @@ void UNiagaraGraph::RemoveParameter(const FNiagaraVariable& Parameter)
 	}
 }
 
-bool UNiagaraGraph::RenameParameter(const FNiagaraVariable& Parameter, FName NewName, bool bFromStaticSwitch)
+bool UNiagaraGraph::RenameParameter(const FNiagaraVariable& Parameter, FName NewName, bool bFromStaticSwitch, ENiagaraParameterScope NewScope)
 {
 	// Block rename when already renaming. This prevents recursion when CommitEditablePinName is called on referenced nodes. 
 	if (bIsRenamingParameter)
@@ -918,8 +1069,38 @@ bool UNiagaraGraph::RenameParameter(const FNiagaraVariable& Parameter, FName New
 		}
 		OldMetaData = (*OldScriptVariable)->Metadata;
 	}
+
+	// Update the OldMetaData which will be applied to the new parameter to synchronize the cached namespace-less name and scope.
+	OldMetaData.CachedNamespacelessVariableName = FName(*FNiagaraStackGraphUtilities::GetNamespacelessVariableNameString(NewParameter.GetName()));
+	if (NewScope != ENiagaraParameterScope::None)
+	{
+		OldMetaData.Scope = NewScope;
+	}
 		
+	// Swap metadata to the new parameter; put the new parameter into VariableToScriptVariable
+	if (OldScriptVariable)
+	{
+		// Rename all the bindings that point to the old parameter 
+		for (auto It : VariableToScriptVariable)
+		{
+			UNiagaraScriptVariable* Variable = It.Value;
+			if (Variable && Variable->DefaultBinding.GetName() == Parameter.GetName())
+			{
+				Variable->DefaultBinding.SetName(NewParameter.GetName());
+			}
+		}
+
+		// Replace the script variable data
+		UNiagaraScriptVariable* NewScriptVariable = NewObject<UNiagaraScriptVariable>(this, FName(), RF_Transactional);
+		NewScriptVariable->Variable = NewParameter;
+		NewScriptVariable->DefaultMode = (*OldScriptVariable)->DefaultMode;
+		NewScriptVariable->DefaultBinding = (*OldScriptVariable)->DefaultBinding;
+		VariableToScriptVariable.Remove(Parameter);
+		VariableToScriptVariable.Add(NewParameter, NewScriptVariable);
+	}
+	SetPerScriptMetaData(NewParameter, OldMetaData);
 
+	// Fixup reference collection and pin names
 	FNiagaraGraphParameterReferenceCollection* ReferenceCollection = ParameterToReferencesMap.Find(Parameter);
 	int NumReferences = 0;
 	if (ReferenceCollection)
@@ -932,6 +1113,7 @@ bool UNiagaraGraph::RenameParameter(const FNiagaraVariable& Parameter, FName New
 			UNiagaraNode* Node = Reference.Value.Get();
 			if (Node && Node->GetGraph() == this)
 			{
+				Node->Modify();
 				UEdGraphPin* Pin = Node->GetPinByPersistentGuid(Reference.Key);
 				if (Pin)
 				{
@@ -944,32 +1126,6 @@ bool UNiagaraGraph::RenameParameter(const FNiagaraVariable& Parameter, FName New
 		ParameterToReferencesMap.Add(NewParameter, NewReferences);
 	}
 
-	// Swap metadata to the new parameter
-	if (OldScriptVariable)
-	{
-		if (!bFromStaticSwitch || (bFromStaticSwitch && NumReferences <= 1))
-		{
-			// Rename all the bindings that point to the old parameter 
-			for (auto It : VariableToScriptVariable)
-			{
-				UNiagaraScriptVariable* Variable = It.Value;
-				if (Variable && Variable->DefaultBinding.GetName() == Parameter.GetName())
-				{
-					Variable->DefaultBinding.SetName(NewParameter.GetName());
-				}
-			}
-
-			// Replace the script variable data
-			UNiagaraScriptVariable* NewScriptVariable = NewObject<UNiagaraScriptVariable>(this, FName(), RF_Transactional);
-			NewScriptVariable->Variable = NewParameter;
-			NewScriptVariable->DefaultMode = (*OldScriptVariable)->DefaultMode;
-			NewScriptVariable->DefaultBinding = (*OldScriptVariable)->DefaultBinding;
-			VariableToScriptVariable.Add(NewParameter, NewScriptVariable);
-			VariableToScriptVariable.Remove(Parameter);
-		}
-	}
-	SetMetaData(NewParameter, OldMetaData);
-
 	bIsRenamingParameter = false;
 
 	NotifyGraphChanged();
@@ -1403,6 +1559,57 @@ void UNiagaraGraph::CopyCachedReferencesMap(UNiagaraGraph* TargetGraph)
 	TargetGraph->ParameterToReferencesMap = ParameterToReferencesMap;
 }
 
+void UNiagaraGraph::FixupParameterUsageMetaData()
+{
+	// Collect all input and output pins to infer usages of variables they reference.
+	const TMap<FNiagaraVariable, FInputPinsAndOutputPins> VarToPinsMap = CollectVarsToInOutPinsMap();
+
+	// Now iterate all variables
+	for (auto It = VariableToScriptVariable.CreateIterator(); It; ++It)
+	{
+		FNiagaraVariable Var = It.Key();
+		UNiagaraScriptVariable* ScriptVar = It.Value();
+
+		if (Var.GetName() == FNiagaraConstants::InputPinName || Var.GetName() == FNiagaraConstants::OutputPinName)
+		{
+			//@todo Pins leaked into variable maps at some point, need to clean.
+			continue;
+		}
+
+		if (ScriptVar->Metadata.bIsStaticSwitch)
+		{
+			// Do not need to assign scopes and usages to static switches, so assign the name and continue.
+			ScriptVar->Metadata.CachedNamespacelessVariableName = ScriptVar->Variable.GetName();
+			continue;
+		}
+
+		// Resolve the namespace of the variable to a scope, and assign the usage to local if "local" is in the name string.
+		FNiagaraStackGraphUtilities::SetParameterMetaDataFromName(Var.GetName(), ScriptVar->Metadata.Scope, ScriptVar->Metadata.Usage, ScriptVar->Metadata.CachedNamespacelessVariableName);
+
+		// If the usage is not local, select the usage based on the associated input/output pins.
+		SetScriptVariableUsageForPins(VarToPinsMap, ScriptVar);
+	}
+
+	NotifyGraphNeedsRecompile();
+}
+
+bool UNiagaraGraph::IsPinVisualWidgetProviderRegistered() const
+{
+	return OnGetPinVisualWidgetDelegate.IsBound();
+}
+
+TSharedRef<SWidget> UNiagaraGraph::GetPinVisualWidget(const UEdGraphPin* Pin) const
+{
+	return OnGetPinVisualWidgetDelegate.Execute(Pin);
+}
+
+FDelegateHandle UNiagaraGraph::RegisterPinVisualWidgetProvider(FOnGetPinVisualWidget OnGetPinVisualWidget)
+{
+	checkf(OnGetPinVisualWidgetDelegate.IsBound() == false, TEXT("Only one handler is allowed for the OnGetPinVisualWidget delegate"));
+	OnGetPinVisualWidgetDelegate = OnGetPinVisualWidget;
+	return OnGetPinVisualWidgetDelegate.GetHandle();
+}
+
 const class UEdGraphSchema_Niagara* UNiagaraGraph::GetNiagaraSchema() const
 {
 	return Cast<UEdGraphSchema_Niagara>(GetSchema());
@@ -1611,6 +1818,28 @@ void UNiagaraGraph::SetMetaData(const FNiagaraVariable& InVar, const FNiagaraVar
 	}
 }
 
+void UNiagaraGraph::SetPerScriptMetaData(const FNiagaraVariable& InVar, const FNiagaraVariableMetaData& InMetaData)
+{
+	if (UNiagaraScriptVariable** FoundMetaData = VariableToScriptVariable.Find(InVar))
+	{
+		if (*FoundMetaData)
+		{
+			// Replace the old metadata scope, usage and cached name, but retain persistent metadata such as description.
+			FNiagaraVariableMetaData& OldMetaData = (*FoundMetaData)->Metadata;
+			OldMetaData.Usage = InMetaData.Usage;
+			OldMetaData.Scope = InMetaData.Scope;
+			OldMetaData.CachedNamespacelessVariableName = InMetaData.CachedNamespacelessVariableName;
+			OldMetaData.bCreatedInSystemEditor = InMetaData.bCreatedInSystemEditor;
+		}
+	}
+	else
+	{
+		UNiagaraScriptVariable*& NewScriptVariable = VariableToScriptVariable.Add(InVar, NewObject<UNiagaraScriptVariable>(this, FName(), RF_Transactional));
+		NewScriptVariable->Variable = InVar;
+		NewScriptVariable->Metadata = InMetaData;
+	}
+}
+
 UNiagaraGraph::FOnDataInterfaceChanged& UNiagaraGraph::OnDataInterfaceChanged()
 {
 	return OnDataInterfaceChangedDelegate;
@@ -1703,7 +1932,7 @@ void UNiagaraGraph::RefreshParameterReferences() const
 
 	// Check all pins on all nodes in the graph to find parameter pins which may have been missed in the parameter map traversal.  This
 	// can happen for nodes which are not fully connected and therefore don't show up in the traversal.
-	const UEdGraphSchema_Niagara* NiagaraSchema = Cast<UEdGraphSchema_Niagara>(Schema);
+	const UEdGraphSchema_Niagara* NiagaraSchema = GetNiagaraSchema();
 	for (UEdGraphNode* Node : Nodes)
 	{
 		UNiagaraNodeStaticSwitch* SwitchNode = Cast<UNiagaraNodeStaticSwitch>(Node);
@@ -1757,6 +1986,7 @@ void UNiagaraGraph::RefreshParameterReferences() const
 			UNiagaraScriptVariable* NewScriptVariable = NewObject<UNiagaraScriptVariable>(const_cast<UNiagaraGraph*>(this));
 			NewScriptVariable->Variable = Variable;
 			NewScriptVariable->Metadata.bIsStaticSwitch = false;
+			GenerateMetaDataForScriptVariable(NewScriptVariable);
 			VariableToScriptVariable.Add(Variable, NewScriptVariable);
 		}
 	};
@@ -1821,4 +2051,63 @@ void UNiagaraGraph::InvalidateCachedParameterData()
 	bParameterReferenceRefreshPending = true;
 }
 
-#undef LOCTEXT_NAMESPACE
\ No newline at end of file
+void UNiagaraGraph::GenerateMetaDataForScriptVariable(UNiagaraScriptVariable* InScriptVariable) const
+{
+	FNiagaraStackGraphUtilities::SetParameterMetaDataFromName(InScriptVariable->Variable.GetName(), InScriptVariable->Metadata.Scope, InScriptVariable->Metadata.Usage, InScriptVariable->Metadata.CachedNamespacelessVariableName);
+	UpdateUsageForScriptVariable(InScriptVariable);
+}
+
+const TMap<FNiagaraVariable, FInputPinsAndOutputPins> UNiagaraGraph::CollectVarsToInOutPinsMap() const
+{
+	const UEdGraphSchema_Niagara* NiagaraSchema = GetNiagaraSchema();
+	TMap<FNiagaraVariable, FInputPinsAndOutputPins> VarToPinsMap;
+
+	// Collect all input and output nodes to inspect the pins and infer usages of variables they reference.
+	TArray<UNiagaraNodeParameterMapSet*> MapSetNodes;
+	TArray<UNiagaraNodeParameterMapGet*> MapGetNodes;
+	GetNodesOfClass<UNiagaraNodeParameterMapSet>(MapSetNodes);
+	GetNodesOfClass<UNiagaraNodeParameterMapGet>(MapGetNodes);
+
+	for (UNiagaraNodeParameterMapGet* MapGetNode : MapGetNodes)
+	{
+		TArray<UEdGraphPin*> MapGetOutputPins;
+		MapGetNode->GetOutputPins(MapGetOutputPins);
+		for (UEdGraphPin* Pin : MapGetOutputPins)
+		{
+			if (Pin->PinName == UNiagaraNodeParameterMapBase::SourcePinName || Pin->PinName == UNiagaraNodeParameterMapBase::AddPinName)
+			{
+				continue;
+			}
+
+			FNiagaraVariable Var = FNiagaraVariable(NiagaraSchema->PinToTypeDefinition(Pin), Pin->PinName);
+			if (VariableToScriptVariable.Find(Var))
+			{
+				FInputPinsAndOutputPins& InOutPins = VarToPinsMap.FindOrAdd(Var);
+				InOutPins.OutputPins.Add(Pin);
+			}
+		}
+	}
+
+	for (UNiagaraNodeParameterMapSet* MapSetNode : MapSetNodes)
+	{
+		TArray<UEdGraphPin*> MapSetInputPins;
+		MapSetNode->GetInputPins(MapSetInputPins);
+		for (UEdGraphPin* Pin : MapSetInputPins)
+		{
+			if (Pin->PinName == UNiagaraNodeParameterMapBase::SourcePinName || Pin->PinName == UNiagaraNodeParameterMapBase::AddPinName)
+			{
+				continue;
+			}
+
+			FNiagaraVariable Var = FNiagaraVariable(NiagaraSchema->PinToTypeDefinition(Pin), Pin->PinName);
+			if (VariableToScriptVariable.Find(Var))
+			{
+				FInputPinsAndOutputPins& InOutPins = VarToPinsMap.FindOrAdd(Var);
+				InOutPins.InputPins.Add(Pin);
+			}
+		}
+	}
+	return VarToPinsMap;
+}
+
+#undef LOCTEXT_NAMESPACE
diff --git a/Engine/Plugins/FX/Niagara/Source/NiagaraEditor/Private/NiagaraHlslTranslator.cpp b/Engine/Plugins/FX/Niagara/Source/NiagaraEditor/Private/NiagaraHlslTranslator.cpp
--- a/Engine/Plugins/FX/Niagara/Source/NiagaraEditor/Private/NiagaraHlslTranslator.cpp
+++ b/Engine/Plugins/FX/Niagara/Source/NiagaraEditor/Private/NiagaraHlslTranslator.cpp
@@ -4353,7 +4353,7 @@ void FHlslNiagaraTranslator::Emitter(class UNiagaraNodeEmitter* EmitterNode, TAr
 		Error(LOCTEXT("EmitterMissingParamMapIndex", "Emitter call missing valid ParameterMap index!"), EmitterNode, nullptr);
 		return;
 	}
-	ActiveHistoryForFunctionCalls.EnterEmitter(EmitterUniqueName, EmitterNode);
+	ActiveHistoryForFunctionCalls.EnterEmitter(EmitterUniqueName, EmitterNode->GetCalledGraph(), EmitterNode);
 
 	// Clear out the parameter map writes to emitter module parameters as they should not be shared across emitters.
 	if (ParamMapHistoryIdx != -1 && ParamMapHistoryIdx < ParamMapHistories.Num())
@@ -5190,7 +5190,7 @@ void FHlslNiagaraTranslator::FunctionCall(UNiagaraNodeFunctionCall* FunctionNode
 
 	//UE_LOG(LogNiagaraEditor, Log, TEXT("Function Call: %s %d"), *FunctionNode->GetFunctionName(), ActiveStageIdx);
 	// We need the generated string to generate the proper signature for now.
-	ActiveHistoryForFunctionCalls.EnterFunction(FunctionNode->GetFunctionName(), FunctionNode->FunctionScript, FunctionNode);
+	ActiveHistoryForFunctionCalls.EnterFunction(FunctionNode->GetFunctionName(), FunctionNode->FunctionScript, FunctionNode->GetNiagaraGraph(), FunctionNode);
 
 	// Check if there are static switch parameters being set directly by a set node from the stack UI.
 	// This can happen if a module was changed and the original parameter was replaced by a static switch with the same name, but the emitter was not yet updated.
diff --git a/Engine/Plugins/FX/Niagara/Source/NiagaraEditor/Private/NiagaraHlslTranslator.h b/Engine/Plugins/FX/Niagara/Source/NiagaraEditor/Private/NiagaraHlslTranslator.h
--- a/Engine/Plugins/FX/Niagara/Source/NiagaraEditor/Private/NiagaraHlslTranslator.h
+++ b/Engine/Plugins/FX/Niagara/Source/NiagaraEditor/Private/NiagaraHlslTranslator.h
@@ -96,8 +96,7 @@ public:
 
 	TArray<FNiagaraParameterMapHistory>& GetPrecomputedHistories() { return PrecompiledHistories; }
 	const TArray<FNiagaraParameterMapHistory>& GetPrecomputedHistories() const { return PrecompiledHistories; }
-	class UNiagaraGraph* GetPrecomputedNodeGraph() { return NodeGraphDeepCopy; }
-	const class UNiagaraGraph* GetPrecomputedNodeGraph() const { return NodeGraphDeepCopy; }
+	virtual const class UNiagaraGraph* GetPrecomputedNodeGraph() const { return NodeGraphDeepCopy; }
 	const FString& GetUniqueEmitterName() const { return EmitterUniqueName; }
 	void VisitReferencedGraphs(UNiagaraGraph* InSrcGraph, UNiagaraGraph* InDupeGraph, ENiagaraScriptUsage InUsage, FCompileConstantResolver ConstantResolver);
 	void DeepCopyGraphs(UNiagaraScriptSource* ScriptSource, ENiagaraScriptUsage InUsage, FCompileConstantResolver ConstantResolver);
@@ -147,6 +146,10 @@ public:
 	};
 	TMap<const UNiagaraGraph*, TArray<FunctionData>> PreprocessedFunctions;
 	TArray<UNiagaraGraph*> ClonedGraphs;
+
+	// Copy of the variables that are required by Renderers if this is a compile request for an Emitter. 
+	TArray<FNiagaraVariable> RequiredRendererVariables;
+
 protected:
 	void VisitReferencedGraphsRecursive(UNiagaraGraph* InGraph, const FCompileConstantResolver& ConstantResolver);
 };
diff --git a/Engine/Plugins/FX/Niagara/Source/NiagaraEditor/Private/NiagaraNodeAssignment.cpp b/Engine/Plugins/FX/Niagara/Source/NiagaraEditor/Private/NiagaraNodeAssignment.cpp
--- a/Engine/Plugins/FX/Niagara/Source/NiagaraEditor/Private/NiagaraNodeAssignment.cpp
+++ b/Engine/Plugins/FX/Niagara/Source/NiagaraEditor/Private/NiagaraNodeAssignment.cpp
@@ -225,7 +225,7 @@ void UNiagaraNodeAssignment::GenerateScript()
 {
 	if (FunctionScript == nullptr)
 	{
-		FunctionScript = NewObject<UNiagaraScript>(this, FName(*(TEXT("SetVariables_") + NodeGuid.ToString())), RF_Transactional);
+		FunctionScript = NewObject<UNiagaraScript>(this, FName(*(TRANSLATOR_SET_VARIABLES_UNDERSCORE_STR + NodeGuid.ToString())), RF_Transactional);
 		FunctionScript->SetUsage(ENiagaraScriptUsage::Module);
 		FunctionScript->Description = LOCTEXT("AssignmentNodeDesc", "Sets one or more variables in the stack.");
 		InitializeScript(FunctionScript);
diff --git a/Engine/Plugins/FX/Niagara/Source/NiagaraEditor/Private/NiagaraNodeEmitter.cpp b/Engine/Plugins/FX/Niagara/Source/NiagaraEditor/Private/NiagaraNodeEmitter.cpp
--- a/Engine/Plugins/FX/Niagara/Source/NiagaraEditor/Private/NiagaraNodeEmitter.cpp
+++ b/Engine/Plugins/FX/Niagara/Source/NiagaraEditor/Private/NiagaraNodeEmitter.cpp
@@ -338,9 +338,9 @@ void UNiagaraNodeEmitter::BuildParameterMapHistory(FNiagaraParameterMapHistoryBu
 
 	FString EmitterUniqueName = GetEmitterUniqueName();
 	UNiagaraGraph* Graph = GetCalledGraph();
-	if (Graph && ParamMapIdx != INDEX_NONE)
+	if (Graph && ParamMapIdx != INDEX_NONE && OutHistory.bShouldBuildSubHistories)
 	{
-		OutHistory.EnterEmitter(EmitterUniqueName, this);
+		OutHistory.EnterEmitter(EmitterUniqueName, Graph, this);
 
 		TArray<ENiagaraScriptUsage> Usages;
 		Usages.Add(ENiagaraScriptUsage::EmitterSpawnScript);
@@ -364,7 +364,7 @@ void UNiagaraNodeEmitter::BuildParameterMapHistory(FNiagaraParameterMapHistoryBu
 			ChildBuilder.RegisterEncounterableVariables(OutHistory.GetEncounterableVariables());
 			ChildBuilder.EnableScriptWhitelist(true, GetUsage());
 			FString LocalEmitterName = TEXT("Emitter");
-			ChildBuilder.EnterEmitter(LocalEmitterName, this);
+			ChildBuilder.EnterEmitter(LocalEmitterName, Graph, this);
 			for (UNiagaraNodeOutput* OutputNode : OutputNodes)
 			{
 				ChildBuilder.BuildParameterMaps(OutputNode, true);
@@ -384,11 +384,7 @@ void UNiagaraNodeEmitter::BuildParameterMapHistory(FNiagaraParameterMapHistoryBu
 					int32 ExistingIdx = OutHistory.Histories[ParamMapIdx].FindVariable(Var.GetName(), Var.GetType());
 					if (ExistingIdx == INDEX_NONE)
 					{
-						ExistingIdx = OutHistory.Histories[ParamMapIdx].Variables.Add(Var);
-						OutHistory.Histories[ParamMapIdx].VariablesWithOriginalAliasesIntact.Add(History.VariablesWithOriginalAliasesIntact[SrcVarIdx]);
-						OutHistory.Histories[ParamMapIdx].PerVariableReadHistory.AddDefaulted(1);
-						OutHistory.Histories[ParamMapIdx].PerVariableWriteHistory.AddDefaulted(1);
-						OutHistory.Histories[ParamMapIdx].PerVariableWarnings.AddDefaulted(1);
+						ExistingIdx = OutHistory.AddVariableToHistory(OutHistory.Histories[ParamMapIdx], Var, History.VariablesWithOriginalAliasesIntact[SrcVarIdx], nullptr);
 					}
 					ensure(ExistingIdx < OutHistory.Histories[ParamMapIdx].PerVariableWarnings.Num());
 					ensure(ExistingIdx < OutHistory.Histories[ParamMapIdx].PerVariableReadHistory.Num());
diff --git a/Engine/Plugins/FX/Niagara/Source/NiagaraEditor/Private/NiagaraNodeFunctionCall.cpp b/Engine/Plugins/FX/Niagara/Source/NiagaraEditor/Private/NiagaraNodeFunctionCall.cpp
--- a/Engine/Plugins/FX/Niagara/Source/NiagaraEditor/Private/NiagaraNodeFunctionCall.cpp
+++ b/Engine/Plugins/FX/Niagara/Source/NiagaraEditor/Private/NiagaraNodeFunctionCall.cpp
@@ -839,7 +839,7 @@ void UNiagaraNodeFunctionCall::BuildParameterMapHistory(FNiagaraParameterMapHist
 			}
 		}
 
-		OutHistory.EnterFunction(GetFunctionName(), FunctionScript, this);
+		OutHistory.EnterFunction(GetFunctionName(), FunctionScript, FunctionGraph, this);
 		if (ParamMapIdx != INDEX_NONE)
 		{
 			NodeIdx = OutHistory.BeginNodeVisitation(ParamMapIdx, this);
diff --git a/Engine/Plugins/FX/Niagara/Source/NiagaraEditor/Private/NiagaraNodeParameterMapBase.cpp b/Engine/Plugins/FX/Niagara/Source/NiagaraEditor/Private/NiagaraNodeParameterMapBase.cpp
--- a/Engine/Plugins/FX/Niagara/Source/NiagaraEditor/Private/NiagaraNodeParameterMapBase.cpp
+++ b/Engine/Plugins/FX/Niagara/Source/NiagaraEditor/Private/NiagaraNodeParameterMapBase.cpp
@@ -28,6 +28,7 @@
 const FName UNiagaraNodeParameterMapBase::ParameterPinSubCategory("ParameterPin");
 const FName UNiagaraNodeParameterMapBase::SourcePinName("Source");
 const FName UNiagaraNodeParameterMapBase::DestPinName("Dest");
+const FName UNiagaraNodeParameterMapBase::AddPinName("Add");
 
 UNiagaraNodeParameterMapBase::UNiagaraNodeParameterMapBase() 
 	: UNiagaraNodeWithDynamicPins()
@@ -70,7 +71,7 @@ TArray<FNiagaraParameterMapHistory> UNiagaraNodeParameterMapBase::GetParameterMa
 	Builder.RegisterEncounterableVariables(EncounterableVariables);
 	if (!EmitterNameOverride.IsEmpty())
 	{
-		Builder.EnterEmitter(EmitterNameOverride, nullptr);
+		Builder.EnterEmitter(EmitterNameOverride, InGraphEnd->GetNiagaraGraph(), nullptr);
 	}
 
 	if (bLimitToOutputScriptType)
@@ -210,7 +211,7 @@ void UNiagaraNodeParameterMapBase::OnPinRenamed(UEdGraphPin* RenamedPin, const F
 			Names.Add(Pin->GetFName());
 		}
 	}
-	const FName NewUniqueName = FNiagaraUtilities::GetUniqueName(*RenamedPin->GetName(), Names);
+	const FName NewUniqueName = FNiagaraUtilities::GetUniqueName(*RenamedPin->GetName(), Names); //@todo(ng) remove
 
 	FNiagaraTypeDefinition VarType = CastChecked<UEdGraphSchema_Niagara>(GetSchema())->PinToTypeDefinition(RenamedPin);
 	FNiagaraVariable Var(VarType, *OldName);
diff --git a/Engine/Plugins/FX/Niagara/Source/NiagaraEditor/Private/NiagaraNodeParameterMapBase.h b/Engine/Plugins/FX/Niagara/Source/NiagaraEditor/Private/NiagaraNodeParameterMapBase.h
--- a/Engine/Plugins/FX/Niagara/Source/NiagaraEditor/Private/NiagaraNodeParameterMapBase.h
+++ b/Engine/Plugins/FX/Niagara/Source/NiagaraEditor/Private/NiagaraNodeParameterMapBase.h
@@ -50,4 +50,5 @@ public:
 
 	static const FName SourcePinName;
 	static const FName DestPinName;
+	static const FName AddPinName;
 };
diff --git a/Engine/Plugins/FX/Niagara/Source/NiagaraEditor/Private/NiagaraNodeParameterMapGet.cpp b/Engine/Plugins/FX/Niagara/Source/NiagaraEditor/Private/NiagaraNodeParameterMapGet.cpp
--- a/Engine/Plugins/FX/Niagara/Source/NiagaraEditor/Private/NiagaraNodeParameterMapGet.cpp
+++ b/Engine/Plugins/FX/Niagara/Source/NiagaraEditor/Private/NiagaraNodeParameterMapGet.cpp
@@ -134,17 +134,6 @@ void UNiagaraNodeParameterMapGet::OnNewTypedPinAdded(UEdGraphPin* NewPin)
 		TArray<UEdGraphPin*> OutputPins;
 		GetOutputPins(OutputPins);
 
-		TSet<FName> Names;
-		for (const UEdGraphPin* Pin : OutputPins)
-		{
-			if (Pin != NewPin)
-			{
-				Names.Add(Pin->GetFName());
-			}
-		}
-		const FName NewUniqueName = FNiagaraUtilities::GetUniqueName(*NewPin->GetName(), Names);
-		NewPin->PinName = NewUniqueName;
-
 		const UEdGraphSchema_Niagara* Schema = GetDefault<UEdGraphSchema_Niagara>();
 		FNiagaraTypeDefinition TypeDef = Schema->PinToTypeDefinition(NewPin);
 
@@ -155,6 +144,7 @@ void UNiagaraNodeParameterMapGet::OnNewTypedPinAdded(UEdGraphPin* NewPin)
 		}
 
 		NewPin->PinType.PinSubCategory = UNiagaraNodeParameterMapBase::ParameterPinSubCategory;
+		UpdateAddedPinMetaData(NewPin);
 	}
 
 	if (HasAnyFlags(RF_NeedLoad | RF_NeedPostLoad | RF_NeedInitialization))
@@ -171,7 +161,6 @@ void UNiagaraNodeParameterMapGet::OnNewTypedPinAdded(UEdGraphPin* NewPin)
 
 void UNiagaraNodeParameterMapGet::RemoveDynamicPin(UEdGraphPin* Pin)
 {
-	RemovePin(Pin);
 	if (Pin->Direction == EEdGraphPinDirection::EGPD_Output)
 	{
 		UEdGraphPin* DefaultPin = GetDefaultPin(Pin);
@@ -180,6 +169,8 @@ void UNiagaraNodeParameterMapGet::RemoveDynamicPin(UEdGraphPin* Pin)
 			RemovePin(DefaultPin);
 		}
 	}
+
+	Super::RemoveDynamicPin(Pin);
 }
 
 UEdGraphPin* UNiagaraNodeParameterMapGet::GetDefaultPin(UEdGraphPin* OutputPin) const
diff --git a/Engine/Plugins/FX/Niagara/Source/NiagaraEditor/Private/NiagaraNodeParameterMapGet.h b/Engine/Plugins/FX/Niagara/Source/NiagaraEditor/Private/NiagaraNodeParameterMapGet.h
--- a/Engine/Plugins/FX/Niagara/Source/NiagaraEditor/Private/NiagaraNodeParameterMapGet.h
+++ b/Engine/Plugins/FX/Niagara/Source/NiagaraEditor/Private/NiagaraNodeParameterMapGet.h
@@ -59,6 +59,9 @@ protected:
 	/** Properly set up the default input pin for an output pin.*/
 	UEdGraphPin* CreateDefaultPin(UEdGraphPin* OutputPin);
 
+	/** Convenience method to determine whether this Node is a Map Get or Map Set when adding a parameter through the parameter panel. */
+	virtual EEdGraphPinDirection GetPinDirectionForNewParameters() { return EEdGraphPinDirection::EGPD_Output; };
+
 	UPROPERTY(meta = (SkipForCompileHash="true"))
 	TMap<FGuid, FGuid> PinOutputToPinDefaultPersistentId;
 };
diff --git a/Engine/Plugins/FX/Niagara/Source/NiagaraEditor/Private/NiagaraNodeParameterMapSet.cpp b/Engine/Plugins/FX/Niagara/Source/NiagaraEditor/Private/NiagaraNodeParameterMapSet.cpp
--- a/Engine/Plugins/FX/Niagara/Source/NiagaraEditor/Private/NiagaraNodeParameterMapSet.cpp
+++ b/Engine/Plugins/FX/Niagara/Source/NiagaraEditor/Private/NiagaraNodeParameterMapSet.cpp
@@ -57,6 +57,11 @@ bool UNiagaraNodeParameterMapSet::IsPinNameEditableUponCreation(const UEdGraphPi
 	}
 }
 
+void UNiagaraNodeParameterMapSet::RemoveDynamicPin(UEdGraphPin* Pin)
+{
+	// Call NiagaraNodeWithDynamicPins::RemoveDynamicPin() instead of base class to fixup the associated pin variable's metadata.
+	Super::RemoveDynamicPin(Pin);
+}
 
 bool UNiagaraNodeParameterMapSet::VerifyEditablePinName(const FText& InName, FText& OutErrorMessage, const UEdGraphPin* InGraphPinObj) const
 {
@@ -91,6 +96,7 @@ void UNiagaraNodeParameterMapSet::OnNewTypedPinAdded(UEdGraphPin* NewPin)
 		const FName NewUniqueName = FNiagaraUtilities::GetUniqueName(*NewPin->GetName(), Names);
 		NewPin->PinName = NewUniqueName;
 		NewPin->PinType.PinSubCategory = UNiagaraNodeParameterMapBase::ParameterPinSubCategory;
+		UpdateAddedPinMetaData(NewPin);
 	}
 
 	if (!NewPin->PersistentGuid.IsValid())
diff --git a/Engine/Plugins/FX/Niagara/Source/NiagaraEditor/Private/NiagaraNodeParameterMapSet.h b/Engine/Plugins/FX/Niagara/Source/NiagaraEditor/Private/NiagaraNodeParameterMapSet.h
--- a/Engine/Plugins/FX/Niagara/Source/NiagaraEditor/Private/NiagaraNodeParameterMapSet.h
+++ b/Engine/Plugins/FX/Niagara/Source/NiagaraEditor/Private/NiagaraNodeParameterMapSet.h
@@ -22,6 +22,7 @@ public:
 
 	virtual bool IsPinNameEditable(const UEdGraphPin* GraphPinObj) const override;
 	virtual bool IsPinNameEditableUponCreation(const UEdGraphPin* GraphPinObj) const override;
+	virtual void RemoveDynamicPin(UEdGraphPin* Pin) override;
 	virtual bool VerifyEditablePinName(const FText& InName, FText& OutErrorMessage, const UEdGraphPin* InGraphPinObj) const override;
 	virtual bool CommitEditablePinName(const FText& InName, UEdGraphPin* InGraphPinObj)  override;
 	virtual bool CancelEditablePinName(const FText& InName, UEdGraphPin* InGraphPinObj) override;
@@ -39,4 +40,7 @@ public:
 protected:
 	virtual void OnNewTypedPinAdded(UEdGraphPin* NewPin); 
 	virtual void OnPinRenamed(UEdGraphPin* RenamedPin, const FString& OldName) override;
+
+	/** Convenience method to determine whether this Node is a Map Get or Map Set when adding a parameter through the parameter panel. */
+	virtual EEdGraphPinDirection GetPinDirectionForNewParameters() { return EEdGraphPinDirection::EGPD_Input; };
 };
diff --git a/Engine/Plugins/FX/Niagara/Source/NiagaraEditor/Private/NiagaraNodeReroute.cpp b/Engine/Plugins/FX/Niagara/Source/NiagaraEditor/Private/NiagaraNodeReroute.cpp
--- a/Engine/Plugins/FX/Niagara/Source/NiagaraEditor/Private/NiagaraNodeReroute.cpp
+++ b/Engine/Plugins/FX/Niagara/Source/NiagaraEditor/Private/NiagaraNodeReroute.cpp
@@ -3,6 +3,7 @@
 #include "NiagaraNodeReroute.h"
 #include "NiagaraEditorUtilities.h"
 #include "NiagaraHlslTranslator.h"
+#include "NiagaraConstants.h"
 
 #define LOCTEXT_NAMESPACE "NiagaraNodeReroute"
 
@@ -37,12 +38,10 @@ void UNiagaraNodeReroute::PostLoad()
 
 void UNiagaraNodeReroute::AllocateDefaultPins()
 {
-	const FName InputPinName(TEXT("InputPin"));
-	UEdGraphPin* MyInputPin = CreatePin(EGPD_Input, PC_Wildcard_Niagara, InputPinName);
+	UEdGraphPin* MyInputPin = CreatePin(EGPD_Input, PC_Wildcard_Niagara, FNiagaraConstants::InputPinName);
 	MyInputPin->bDefaultValueIsIgnored = true;
 
-	const FName OutputPinName(TEXT("OutputPin"));
-	CreatePin(EGPD_Output, PC_Wildcard_Niagara, OutputPinName);
+	CreatePin(EGPD_Output, PC_Wildcard_Niagara, FNiagaraConstants::OutputPinName);
 }
 
 FText UNiagaraNodeReroute::GetTooltipText() const
diff --git a/Engine/Plugins/FX/Niagara/Source/NiagaraEditor/Private/NiagaraNodeWithDynamicPins.cpp b/Engine/Plugins/FX/Niagara/Source/NiagaraEditor/Private/NiagaraNodeWithDynamicPins.cpp
--- a/Engine/Plugins/FX/Niagara/Source/NiagaraEditor/Private/NiagaraNodeWithDynamicPins.cpp
+++ b/Engine/Plugins/FX/Niagara/Source/NiagaraEditor/Private/NiagaraNodeWithDynamicPins.cpp
@@ -14,6 +14,7 @@
 #include "SNiagaraParameterMapView.h"
 #include "Framework/Application/SlateApplication.h"
 #include "NiagaraNodeParameterMapBase.h"
+#include "NiagaraScriptVariable.h"
 
 #define LOCTEXT_NAMESPACE "NiagaraNodeWithDynamicPins"
 
@@ -100,6 +101,19 @@ void UNiagaraNodeWithDynamicPins::CreateAddPin(EEdGraphPinDirection Direction)
 	CreatePin(Direction, FEdGraphPinType(UEdGraphSchema_Niagara::PinCategoryMisc, AddPinSubCategory, nullptr, EPinContainerType::None, false, FEdGraphTerminalType()), TEXT("Add"));
 }
 
+void UNiagaraNodeWithDynamicPins::UpdateAddedPinMetaData(const UEdGraphPin* AddedPin)
+{
+	if (UNiagaraGraph* Graph = GetNiagaraGraph())
+	{
+		const UEdGraphSchema_Niagara* Schema = GetDefault<UEdGraphSchema_Niagara>();
+		FNiagaraVariable PinVariable = Schema->PinToNiagaraVariable(AddedPin, false);
+		if (UNiagaraScriptVariable** ScriptVariable = Graph->GetAllMetaData().Find(PinVariable))
+		{
+			Graph->UpdateUsageForScriptVariable(*ScriptVariable);
+		}
+	}
+}
+
 bool UNiagaraNodeWithDynamicPins::IsAddPin(const UEdGraphPin* Pin) const
 {
 	return Pin->PinType.PinCategory == UEdGraphSchema_Niagara::PinCategoryMisc && 
@@ -246,7 +260,7 @@ void UNiagaraNodeWithDynamicPins::CollectAddPinActions(FGraphActionListBuilderBa
 			const FText Tooltip = FText::Format(LOCTEXT("AddButtonTypeEntryToolTipFormat", "Add a new {0} pin"), RegisteredType.GetNameText());
 			TSharedPtr<FNiagaraMenuAction> Action(new FNiagaraMenuAction(
 				FText::GetEmpty(), DisplayName, Tooltip, 0, FText::GetEmpty(),
-				FNiagaraMenuAction::FOnExecuteStackAction::CreateUObject(this, &UNiagaraNodeWithDynamicPins::AddParameter, Var, Pin)));
+				FNiagaraMenuAction::FOnExecuteStackAction::CreateUObject(this, &UNiagaraNodeWithDynamicPins::AddParameter, Var, (const UEdGraphPin*)Pin)));
 
 			OutActions.AddAction(Action);
 		}
@@ -255,37 +269,70 @@ void UNiagaraNodeWithDynamicPins::CollectAddPinActions(FGraphActionListBuilderBa
 	bOutCreateRemainingActions = false;
 }
 
-void UNiagaraNodeWithDynamicPins::AddParameter(FNiagaraVariable Parameter, UEdGraphPin* AddPin)
+void UNiagaraNodeWithDynamicPins::AddParameter(FNiagaraVariable Parameter, const UEdGraphPin* AddPin)
 {
 	FScopedTransaction AddNewPinTransaction(LOCTEXT("AddNewPinTransaction", "Add pin to node"));
-		{
-		// We need to create new UNiagaraScriptVariable instances only when creating a new variable
-		if (UNiagaraNodeParameterMapBase* Map = Cast<UNiagaraNodeParameterMapBase>(this))
-		{
-			if (UNiagaraGraph* Graph = GetNiagaraGraph())
-			{	
-				UNiagaraScriptVariable* Variable = Graph->GetScriptVariable(Parameter);
-				if (Variable == nullptr)
-				{
-					TSet<FName> Names;
-					for (const auto& ParameterElement : Graph->GetParameterReferenceMap())
-					{
-						Names.Add(ParameterElement.Key.GetName());
-					}
-					Parameter.SetName(FNiagaraUtilities::GetUniqueName(Parameter.GetName(), Names));
-					Graph->AddParameter(Parameter);
-				}
-			}
-		}
+	UNiagaraGraph::FAddParameterOptions AddParameterOptions = UNiagaraGraph::FAddParameterOptions();
+	if (AddPin->Direction == EEdGraphPinDirection::EGPD_Input)
+	{
+		AddParameterOptions.NewParameterUsage = ENiagaraScriptParameterUsage::Output;
 	}
+	else
+	{
+		AddParameterOptions.NewParameterUsage = ENiagaraScriptParameterUsage::Input;
+	}
+	// Default scope of new parameters on a pin to Particles, by convention. //@todo(ng) choose based on usage
+	AddParameterOptions.NewParameterScope = ENiagaraParameterScope::Particles;
+	//AddParameterOptions.bMakeParameterNameUnique //@todo(ng) consider
+	
+	UNiagaraGraph* Graph = GetNiagaraGraph();
+	checkf(Graph != nullptr, TEXT("Failed to get niagara graph when adding pin!"));
+	
+	Graph->Modify();
+	Graph->AddParameter(Parameter, AddParameterOptions);
+
+	Modify();
 	UEdGraphPin* Pin = this->RequestNewTypedPin(AddPin->Direction, Parameter.GetType(), Parameter.GetName());
+	FNiagaraGraphParameterReference NewParameterReference = FNiagaraGraphParameterReference(Pin->PersistentGuid, this);
+	Graph->AddParameterReference(Parameter, NewParameterReference);
+}
 
+void UNiagaraNodeWithDynamicPins::AddParameter(FNiagaraVariable Parameter, const struct UNiagaraGraph::FAddParameterOptions AddParameterOptions)
+{
+	UNiagaraGraph* Graph = GetNiagaraGraph();
+	checkf(Graph != nullptr, TEXT("Failed to get niagara graph when adding pin!"));
+
+	Graph->Modify();
+	Graph->AddParameter(Parameter, AddParameterOptions);
+
+	EEdGraphPinDirection NewPinDirection = GetPinDirectionForNewParameters();
+	checkf(NewPinDirection != EGPD_MAX, TEXT("Could not determine direction for new pin! Did you derive a new node from UNiagaraNodeWithDynamicPins?"));
+	
+	Modify();
+	UEdGraphPin* Pin = this->RequestNewTypedPin(NewPinDirection, Parameter.GetType(), Parameter.GetName());
+	FNiagaraGraphParameterReference NewParameterReference = FNiagaraGraphParameterReference(Pin->PersistentGuid, this);
+	Graph->AddParameterReference(Parameter, NewParameterReference);
 }
 
 void UNiagaraNodeWithDynamicPins::RemoveDynamicPin(UEdGraphPin* Pin)
 {
 	RemovePin(Pin);
 	MarkNodeRequiresSynchronization(__FUNCTION__, true);
+
+	if (UNiagaraGraph* Graph = GetNiagaraGraph())
+	{
+		const UEdGraphSchema_Niagara* Schema = GetDefault<UEdGraphSchema_Niagara>();
+		FNiagaraVariable PinVariable = Schema->PinToNiagaraVariable(Pin, false);
+		UNiagaraScriptVariable** PinAssociatedScriptVariable = Graph->GetAllMetaData().Find(PinVariable);
+		if (PinAssociatedScriptVariable != nullptr)
+		{
+			bool bShouldRemoveScriptVariable = !Graph->UpdateUsageForScriptVariable(*PinAssociatedScriptVariable);
+			if (bShouldRemoveScriptVariable)
+			{
+				Graph->RemoveParameter((*PinAssociatedScriptVariable)->Variable);
+			}
+		}
+	}
 }
 
 FText UNiagaraNodeWithDynamicPins::GetPinNameText(UEdGraphPin* Pin) const
diff --git a/Engine/Plugins/FX/Niagara/Source/NiagaraEditor/Private/NiagaraParameterMapHistory.cpp b/Engine/Plugins/FX/Niagara/Source/NiagaraEditor/Private/NiagaraParameterMapHistory.cpp
--- a/Engine/Plugins/FX/Niagara/Source/NiagaraEditor/Private/NiagaraParameterMapHistory.cpp
+++ b/Engine/Plugins/FX/Niagara/Source/NiagaraEditor/Private/NiagaraParameterMapHistory.cpp
@@ -1,4 +1,5 @@
 // Copyright Epic Games, Inc. All Rights Reserved.
+
 #include "NiagaraParameterMapHistory.h"
 
 #include "NiagaraEditorCommon.h"
@@ -22,13 +23,29 @@
 #include "NiagaraNodeStaticSwitch.h"
 
 #include "NiagaraScriptVariable.h"
+#include "NiagaraNodeParameterMapSet.h"
 
 #define LOCTEXT_NAMESPACE "NiagaraEditor"
 
+static int32 GNiagaraLogNamespaceFixup = 0;
+static FAutoConsoleVariableRef CVarNiagaraLogNamespaceFixup(
+	TEXT("fx.NiagaraLogNamespaceFixup"),
+	GNiagaraLogNamespaceFixup,
+	TEXT("Log matched variables and pin name changes in precompile. \n"),
+	ECVF_Default
+);
+
+static int32 GNiagaraForcePrecompilerCullDataset = 0;
+static FAutoConsoleVariableRef CVarNiagaraForcePrecompilerCullDataset(
+	TEXT("fx.NiagaraForcePrecompilerCullDataset"),
+	GNiagaraForcePrecompilerCullDataset,
+	TEXT("Force the precompiler process to cull unused Dataset parameters. \n"),
+	ECVF_Default
+);
 
 FNiagaraParameterMapHistory::FNiagaraParameterMapHistory() 
 {
-	
+	OriginatingScriptUsage = ENiagaraScriptUsage::Function;
 }
 
 void FNiagaraParameterMapHistory::GetValidNamespacesForReading(const UNiagaraScript* InScript, TArray<FString>& OutputNamespaces)
@@ -174,7 +191,11 @@ int32 FNiagaraParameterMapHistory::FindVariable(const FName& VariableName, const
 	return FoundIdx;
 }
 
-int32 FNiagaraParameterMapHistory::AddVariable(const FNiagaraVariable& InVar, const FNiagaraVariable& InAliasedVar, const UEdGraphPin* InPin)
+int32 FNiagaraParameterMapHistory::AddVariable(
+	  const FNiagaraVariable& InVar
+	, const FNiagaraVariable& InAliasedVar
+	, const UEdGraphPin* InPin
+	, TOptional<FNiagaraVariableMetaData> InMetaData /*= TOptional<FNiagaraVariableMetaData>()*/)
 {
 	FNiagaraVariable Var = InVar;
 	
@@ -186,6 +207,11 @@ int32 FNiagaraParameterMapHistory::AddVariable(const FNiagaraVariable& InVar, co
 		PerVariableWarnings.AddDefaulted(1);
 		PerVariableWriteHistory.AddDefaulted(1);
 		PerVariableReadHistory.AddDefaulted(1);
+		if (InMetaData.IsSet())
+		{
+			VariableMetaData.Add(InMetaData.GetValue());
+			check(Variables.Num() == VariableMetaData.Num());
+		}
 
 		if (InPin != nullptr)
 		{
@@ -925,12 +951,13 @@ bool FNiagaraParameterMapHistoryBuilder::InTopLevelFunctionCall(ENiagaraScriptUs
 }
 
 
-void FNiagaraParameterMapHistoryBuilder::EnterFunction(const FString& InNodeName, const UNiagaraScript* InScript, const UNiagaraNode* Node)
+void FNiagaraParameterMapHistoryBuilder::EnterFunction(const FString& InNodeName, const UNiagaraScript* InScript, const UNiagaraGraph* InGraph, const UNiagaraNode* Node)
 {
 	if (InScript != nullptr )
 	{
 		RegisterNodeVisitation(Node);
 		CallingContext.Push(Node);
+		CallingGraphContext.Push(InGraph);
 		PinToParameterMapIndices.Emplace();
 		FunctionNameContextStack.Emplace(*InNodeName);
 		BuildCurrentAliases();
@@ -947,6 +974,7 @@ void FNiagaraParameterMapHistoryBuilder::ExitFunction(const FString& InNodeName,
 	if (InScript != nullptr)
 	{
 		CallingContext.Pop();
+		CallingGraphContext.Pop();
 		PinToParameterMapIndices.Pop();
 		FunctionNameContextStack.Pop();
 		BuildCurrentAliases();
@@ -954,10 +982,11 @@ void FNiagaraParameterMapHistoryBuilder::ExitFunction(const FString& InNodeName,
 	}
 }
 
-void FNiagaraParameterMapHistoryBuilder::EnterEmitter(const FString& InEmitterName, const UNiagaraNode* Node)
+void FNiagaraParameterMapHistoryBuilder::EnterEmitter(const FString& InEmitterName, const UNiagaraGraph* InGraph, const UNiagaraNode* Node)
 {
 	RegisterNodeVisitation(Node);
 	CallingContext.Push(Node);
+	CallingGraphContext.Push(InGraph);
 	EmitterNameContextStack.Emplace(*InEmitterName);
 	BuildCurrentAliases();
 
@@ -987,6 +1016,7 @@ void FNiagaraParameterMapHistoryBuilder::EnterEmitter(const FString& InEmitterNa
 void FNiagaraParameterMapHistoryBuilder::ExitEmitter(const FString& InEmitterName, const UNiagaraNode* Node)
 {
 	CallingContext.Pop();
+	CallingGraphContext.Pop();
 	EmitterNameContextStack.Pop();
 	BuildCurrentAliases();
 	ContextuallyVisitedNodes.Pop();
@@ -1131,7 +1161,7 @@ int32 FNiagaraParameterMapHistoryBuilder::HandleVariableWrite(int32 ParamMapIdx,
 	FNiagaraVariable AliasedVar = Var;
 	Var = ResolveAliases(Var);
 
-	return Histories[ParamMapIdx].AddVariable(Var,AliasedVar, InPin);
+	return AddVariableToHistory(Histories[ParamMapIdx], Var, AliasedVar, InPin);
 }
 
 int32 FNiagaraParameterMapHistoryBuilder::HandleVariableWrite(int32 ParameterMapIndex, const FNiagaraVariable& Var)
@@ -1144,7 +1174,7 @@ int32 FNiagaraParameterMapHistoryBuilder::HandleVariableWrite(int32 ParameterMap
 	}
 	FNiagaraVariable ResolvedVar = ResolveAliases(Var);
 
-	return Histories[ParameterMapIndex].AddVariable(ResolvedVar, Var, nullptr);
+	return AddVariableToHistory(Histories[ParameterMapIndex], ResolvedVar, Var, nullptr);
 }
 
 int32 FNiagaraParameterMapHistoryBuilder::HandleVariableRead(int32 ParamMapIdx, const UEdGraphPin* InPin, bool RegisterReadsAsVariables, const UEdGraphPin* InDefaultPin, bool bFilterForCompilation, bool& OutUsedDefault)
@@ -1177,44 +1207,28 @@ int32 FNiagaraParameterMapHistoryBuilder::HandleVariableRead(int32 ParamMapIdx,
 	{
 		if (RegisterReadsAsVariables)
 		{
-			OutUsedDefault = false;
 			if (InDefaultPin)
 			{
-				OutUsedDefault = true;
 				VisitInputPin(InDefaultPin, Cast<UNiagaraNode>(InDefaultPin->GetOwningNode()), bFilterForCompilation);
-				FoundIdx = Histories[ParamMapIdx].FindVariable(Var.GetName(), Var.GetType());
 			}
 
-			if (FoundIdx == -1)
+			FoundIdx = AddVariableToHistory(Histories[ParamMapIdx], Var, AliasedVar, InDefaultPin);
+
+			// Add the default binding as well to the parameter history, if used.
+			if (UNiagaraGraph* Graph = Cast<UNiagaraGraph>(InPin->GetOwningNode()->GetGraph()))
 			{
-				FoundIdx = Histories[ParamMapIdx].Variables.Add(Var);
-				Histories[ParamMapIdx].VariablesWithOriginalAliasesIntact.Add(AliasedVar);
-				Histories[ParamMapIdx].PerVariableWarnings.AddDefaulted(1);
-				Histories[ParamMapIdx].PerVariableWriteHistory.AddDefaulted(1);
-				Histories[ParamMapIdx].PerVariableReadHistory.AddDefaulted(1);
-
-				// Add the default binding as well to the parameter history, if used.
-				if (UNiagaraGraph* Graph = Cast<UNiagaraGraph>(InPin->GetOwningNode()->GetGraph()))
+				UNiagaraScriptVariable* Variable = Graph->GetScriptVariable(AliasedVar);
+				if (Variable && Variable->DefaultMode == ENiagaraDefaultMode::Binding && Variable->DefaultBinding.IsValid())
 				{
-					UNiagaraScriptVariable* Variable = Graph->GetScriptVariable(AliasedVar);
-					if (Variable && Variable->DefaultMode == ENiagaraDefaultMode::Binding && Variable->DefaultBinding.IsValid())
-					{
-						int FoundIdxBinding = Histories[ParamMapIdx].Variables.Add(FNiagaraVariable(Var.GetType(), Variable->DefaultBinding.GetName()));
-						Histories[ParamMapIdx].VariablesWithOriginalAliasesIntact.Add(FNiagaraVariable(Var.GetType(), Variable->DefaultBinding.GetName()));
-						Histories[ParamMapIdx].PerVariableWarnings.AddDefaulted(1);
-						Histories[ParamMapIdx].PerVariableWriteHistory.AddDefaulted(1);
-						Histories[ParamMapIdx].PerVariableReadHistory.AddDefaulted(1);
-						
-						Histories[ParamMapIdx].PerVariableReadHistory[FoundIdxBinding].Add(TTuple<const UEdGraphPin*, const UEdGraphPin*>(InDefaultPin, nullptr));
-					}
+					FNiagaraVariable TempVar = FNiagaraVariable(Var.GetType(), Variable->DefaultBinding.GetName());
+					int32 FoundIdxBinding = AddVariableToHistory(Histories[ParamMapIdx], TempVar, TempVar, nullptr);
+
+					Histories[ParamMapIdx].PerVariableReadHistory[FoundIdxBinding].Add(TTuple<const UEdGraphPin*, const UEdGraphPin*>(InDefaultPin, nullptr));
 				}
 			}
+
 			Histories[ParamMapIdx].PerVariableReadHistory[FoundIdx].Add(TTuple<const UEdGraphPin*, const UEdGraphPin*>(InPin, nullptr));
 
-			if (InDefaultPin && OutUsedDefault)
-			{
-				Histories[ParamMapIdx].PerVariableWriteHistory[FoundIdx].Add(InDefaultPin);				
-			}
 		}
 		check(Histories[ParamMapIdx].Variables.Num() == Histories[ParamMapIdx].PerVariableWarnings.Num());
 		check(Histories[ParamMapIdx].Variables.Num() == Histories[ParamMapIdx].PerVariableWriteHistory.Num());
@@ -1335,14 +1349,7 @@ int32 FNiagaraParameterMapHistoryBuilder::HandleExternalVariableRead(int32 Param
 
 		if (Var.IsValid())
 		{
-			FoundIdx = Histories[ParamMapIdx].Variables.Add(Var);
-			Histories[ParamMapIdx].VariablesWithOriginalAliasesIntact.Add(AliasedVar);
-			Histories[ParamMapIdx].PerVariableWarnings.AddDefaulted(1);
-			Histories[ParamMapIdx].PerVariableWriteHistory.AddDefaulted(1);
-			Histories[ParamMapIdx].PerVariableReadHistory.AddDefaulted(1);
-			check(Histories[ParamMapIdx].Variables.Num() == Histories[ParamMapIdx].PerVariableWarnings.Num());
-			check(Histories[ParamMapIdx].Variables.Num() == Histories[ParamMapIdx].PerVariableWriteHistory.Num());
-			check(Histories[ParamMapIdx].Variables.Num() == Histories[ParamMapIdx].PerVariableReadHistory.Num());
+			FoundIdx = AddVariableToHistory(Histories[ParamMapIdx], Var, AliasedVar, nullptr);
 		}
 		else
 		{
@@ -1360,6 +1367,11 @@ int32 FNiagaraParameterMapHistoryBuilder::HandleExternalVariableRead(int32 Param
 }
 
 
+int32 FNiagaraParameterMapHistoryBuilder::AddVariableToHistory(FNiagaraParameterMapHistory& History, const FNiagaraVariable& InVar, const FNiagaraVariable& InAliasedVar, const UEdGraphPin* InPin)
+{
+	return History.AddVariable(InVar, InAliasedVar, InPin);
+}
+
 void FNiagaraParameterMapHistoryBuilder::BuildCurrentAliases()
 {
 	AliasMap = TMap<FString, FString>();
@@ -1429,3 +1441,610 @@ bool FCompileConstantResolver::ResolveConstant(FNiagaraVariable& OutConstant) co
 }
 
 #undef LOCTEXT_NAMESPACE
+
+void FNiagaraParameterMapHistoryWithMetaDataBuilder::FixupHistoryVariableNamespaces(TArray<FNiagaraParameterMapHistoryHandle>& InHistoryHandles, bool bCullDatasetOutputs)
+{
+	if (FNiagaraEditorCommonCVar::GNiagaraEnableParameterPanel2 == 0)
+	{
+		return;
+	}
+
+	struct FMatchedVariableInfo
+	{
+	public:
+		FMatchedVariableInfo(const FName& InEmitterAliasedName, const FName& InName, const FNiagaraTypeDefinition& InTypeDef)
+			: EmitterAliasedName(InEmitterAliasedName)
+			, Name(InName)
+			, TypeDef(InTypeDef)
+		{};
+
+		const FName& GetEmitterAliasedName() const { return EmitterAliasedName; };
+		const FName& GetName() const { return Name; };
+		const FNiagaraTypeDefinition& GetTypeDef() const { return TypeDef; };
+
+	private:
+		const FName EmitterAliasedName;
+		const FName Name;
+		const FNiagaraTypeDefinition TypeDef;
+	};
+
+	struct FSortedHistoryHandle
+	{
+		FSortedHistoryHandle()
+		{
+			Handle = nullptr;
+		}
+
+		FSortedHistoryHandle(FNiagaraParameterMapHistoryHandle& InHandle, const TArray<FNiagaraVariable>& VariablesToSkip, const TArray<FNiagaraVariable>& VariablesToNotDemote)
+		{
+			Handle = &InHandle;
+
+			auto SortVars = [&InHandle, &VariablesToSkip](bool bGatherInputs)->TArray<FNiagaraHistoryVariable> {
+				TArray<FNiagaraHistoryVariable> OutHistoryVars = InHandle.HistoryVariables.FilterByPredicate([&VariablesToSkip, bGatherInputs](const FNiagaraHistoryVariable& HistoryVar)
+				{
+					ENiagaraParameterScope HistoryScope = HistoryVar.VarMetaData->Scope;
+					if (HistoryScope == ENiagaraParameterScope::User || HistoryScope == ENiagaraParameterScope::Engine || HistoryScope == ENiagaraParameterScope::ScriptTransient || HistoryScope == ENiagaraParameterScope::None)
+					{
+						return false;
+					}
+					else if (HistoryVar.VarMetaData->Usage == ENiagaraScriptParameterUsage::Input && HistoryVar.VarMetaData->Scope == ENiagaraParameterScope::Input)
+					{
+						// Inputs are equivalent to module namespace, these should not change.
+						return false;
+					}
+					else if (VariablesToSkip.Contains(*HistoryVar.VarWithOriginalAlias))
+					{
+						return false;
+					}
+
+					if (HistoryVar.VarMetaData->bAddedToNodeGraphDeepCopy == false)
+					{
+						FNiagaraParameterHandle VarHandle = FNiagaraParameterHandle(HistoryVar.VarWithOriginalAlias->GetName());
+						if (VarHandle.GetNamespace().ToString().Contains(TRANSLATOR_SET_VARIABLES_STR))
+						{
+							return false;
+						}
+					}
+					
+					if (bGatherInputs)
+					{
+						return HistoryVar.VarMetaData->Usage == ENiagaraScriptParameterUsage::Input || HistoryVar.VarMetaData->Usage == ENiagaraScriptParameterUsage::InputOutput || HistoryVar.VarMetaData->bAddedToNodeGraphDeepCopy;
+					}
+					else
+					{
+						return HistoryVar.VarMetaData->Usage == ENiagaraScriptParameterUsage::Output || HistoryVar.VarMetaData->Usage == ENiagaraScriptParameterUsage::InputOutput;
+					}
+					return false;
+				});
+				return OutHistoryVars;
+			};
+
+			auto MarkDoNotDemoteVars = [&VariablesToNotDemote](TArray<FNiagaraHistoryVariable>& VarsToMark) {
+				TArray<FName> VariableNamesToNotDemote;
+				for (const FNiagaraVariable& VariableToNotDemote : VariablesToNotDemote)
+				{
+					VariableNamesToNotDemote.Add( FNiagaraParameterHandle(VariableToNotDemote.GetName() ).GetName() );
+				}
+
+				for (FNiagaraHistoryVariable& VarToMark : VarsToMark)
+				{
+					if (VariableNamesToNotDemote.Contains(FNiagaraParameterHandle( VarToMark.VarWithOriginalAlias->GetName() ).GetName() ) )
+					{
+						VarToMark.bDoNotDemote = true;
+					}
+				}
+			};
+
+			OutputHistoryVars = SortVars(false);
+			InputHistoryVars = SortVars(true);
+			MarkDoNotDemoteVars(OutputHistoryVars);
+// 			ResolveScriptAliasParameters(OutputHistoryVars);
+// 			ResolveScriptAliasParameters(InputHistoryVars);
+		}
+
+		TArray<FNiagaraHistoryVariable> OutputHistoryVars;
+		TArray<FNiagaraHistoryVariable> InputHistoryVars;
+		FNiagaraParameterMapHistoryHandle* Handle;
+
+		const int32 GetScriptExecutionIndex() const { return Handle->GetScriptExecutionIndex(); };
+		const int32 GetScriptRunIndex() const { return Handle->GetScriptRunIndex(); };
+		const ENiagaraParameterScope GetScriptScope() const { return Handle->GetScriptScope(); };
+		const FName& GetEmitterUniqueName() const { return Handle->GetEmitterUniqueName(); };
+		const bool IsSystemHistoryHandle() const { return Handle->IsSystemHistoryHandle(); };
+		const ENiagaraScriptUsage GetOriginatingScriptUsage() const { return Handle->GetOriginatingScriptUsage(); };
+
+// 		void ResolveScriptAliasParameters(TArray<FNiagaraHistoryVariable>& VarsToResolve)
+// 		{
+// 			const ENiagaraParameterScope HistoryScope = Handle->GetScriptScope();
+// 			const FString HistoryScopeString = FNiagaraStackGraphUtilities::GetNamespaceStringForScriptParameterScope(HistoryScope);
+// 			for (FNiagaraHistoryVariable& HistoryVar : VarsToResolve)
+// 			{
+// 				if (HistoryVar.VarMetaData->Scope == ENiagaraParameterScope::ScriptAlias)
+// 				{
+// 					HistoryVar.SetVarWithOriginalAliasName()
+// 					FNiagaraParameterHandle HistoryVarParameterHandle = FNiagaraParameterHandle(HistoryVar.Var->GetName());
+// 
+// 					FString FixedNamespaceString = FString();
+// 					FixedNamespaceString.Append(HistoryScopeString);
+// 					FixedNamespaceString.Append(HistoryVarParameterHandle.GetName().ToString());
+// 					FName FixedNamespaceName = FName(*FixedNamespaceString);
+// 					HistoryVar.Var->SetName(FixedNamespaceName);
+// 				}
+// 			}
+// 		}
+	};
+
+	auto GenerateInitialVarsToSkip = [](TArray<FNiagaraVariable>& Vars)->TArray<FNiagaraVariable> {
+		TArray<FNiagaraVariable> InitialVars = Vars;
+		for (FNiagaraVariable& InitialVar : InitialVars)
+		{
+			FNiagaraParameterHandle InitialVarHandle = FNiagaraParameterHandle(InitialVar.GetName());
+			FName NewCopyVarName = FName(*(InitialVarHandle.GetNamespace().ToString() + "." + PARAM_MAP_INITIAL_STR + InitialVarHandle.GetName().ToString()));
+			InitialVar.SetName(NewCopyVarName);
+		}
+		return InitialVars;
+	};
+
+
+	TArray<FSortedHistoryHandle> SortedHandles;
+	TArray<FNiagaraVariable> EmitterScopeRendererVars;
+	TArray<TArray<FNiagaraVariable>> PerHistoryParticleScopeRendererVars;
+	PerHistoryParticleScopeRendererVars.AddDefaulted(InHistoryHandles.Num());
+
+	for (int HandleIdx = 0; HandleIdx < InHistoryHandles.Num(); ++HandleIdx)
+	{
+		const FNiagaraParameterMapHistoryHandle& Handle = InHistoryHandles[HandleIdx];
+		for (const FNiagaraVariable& Var : Handle.GetRequiredRendererVariables())
+		{
+			FNiagaraParameterHandle VarHandle = FNiagaraParameterHandle(Var.GetName());
+			if (VarHandle.IsParticleAttributeHandle())
+			{
+				PerHistoryParticleScopeRendererVars[HandleIdx].Add(Var);
+			}
+			else
+			{
+				EmitterScopeRendererVars.Add(Var);
+			}
+		}
+	}
+
+	TArray<FNiagaraVariable> EngineVarsToSkip = FNiagaraConstants::GetEngineConstants();
+	EngineVarsToSkip.Add(SYS_PARAM_INSTANCE_ALIVE);
+	EngineVarsToSkip.Add(SYS_PARAM_PARTICLES_ID);
+	EngineVarsToSkip.Add(SYS_PARAM_PARTICLES_UNIQUE_ID);
+
+	for (int HandleIdx = 0; HandleIdx < InHistoryHandles.Num(); ++HandleIdx)
+	{
+		FNiagaraParameterMapHistoryHandle& Handle = InHistoryHandles[HandleIdx];
+		checkf(Handle.GetScriptExecutionIndex() != INDEX_NONE || Handle.GetScriptRunIndex() != INDEX_NONE, TEXT("Bad handle in arguments!"));
+		TArray<FNiagaraVariable> VarsToSkip;
+		TArray<FNiagaraVariable> VarsToNotDemote;
+		VarsToNotDemote.Append(EmitterScopeRendererVars);
+		VarsToNotDemote.Append(PerHistoryParticleScopeRendererVars[HandleIdx]);
+		VarsToSkip.Append(GenerateInitialVarsToSkip(VarsToNotDemote));
+		VarsToSkip.Append(EngineVarsToSkip);
+		
+		FSortedHistoryHandle HistoryHandle = FSortedHistoryHandle(Handle, VarsToSkip, VarsToNotDemote);
+		SortedHandles.Add(HistoryHandle);
+	}
+
+	// If an input variable is matched with an output variable upstream, then add the variable name and top level scope to this array for fixing up additional input variables that exist further downstream.
+	TArray<FMatchedVariableInfo> MatchedUpstreamVariableInfos;
+
+	// Sort the history handles in order of their script execution as we want to record the upstream histories during iteration.
+	SortedHandles.Sort([](const FSortedHistoryHandle& A, const FSortedHistoryHandle& B) {
+		if (A.GetScriptRunIndex() < B.GetScriptRunIndex())
+		{
+			return true;
+		}
+		else if (A.GetScriptRunIndex() == B.GetScriptRunIndex() && A.GetScriptExecutionIndex() < B.GetScriptExecutionIndex())
+		{
+			return true;
+		}
+		return false;
+	});
+
+	TArray<FSortedHistoryHandle*> ValidUpstreamHandles;
+
+	auto GetValidUpstreamHandles = [&SortedHandles](FSortedHistoryHandle& TargetHandle)->TArray<FSortedHistoryHandle*> {
+		TArray<FSortedHistoryHandle*> OutHandles;
+		bool bPerEmitter = TargetHandle.IsSystemHistoryHandle();
+
+		for(FSortedHistoryHandle& PotentialHandle : SortedHandles)
+		{
+			if (bPerEmitter && PotentialHandle.GetEmitterUniqueName() != TargetHandle.GetEmitterUniqueName())
+			{
+				continue;
+			}
+
+			if (PotentialHandle.GetScriptRunIndex() == TargetHandle.GetScriptRunIndex())
+			{
+				if (PotentialHandle.GetScriptExecutionIndex() < TargetHandle.GetScriptExecutionIndex())
+				{
+					OutHandles.Add(&PotentialHandle);
+				}
+			}
+			else if (PotentialHandle.GetScriptRunIndex() < TargetHandle.GetScriptRunIndex())
+			{
+				if (PotentialHandle.GetScriptExecutionIndex() <= TargetHandle.GetScriptExecutionIndex())
+				{
+					OutHandles.Add(&PotentialHandle);
+				}
+			}
+		}
+
+		// Sort in reverse execution order to walk up execution history from the target history.
+		OutHandles.Sort([](const FSortedHistoryHandle& A, const FSortedHistoryHandle& B) {
+			if (A.GetScriptRunIndex() > B.GetScriptRunIndex())
+			{
+				return true;
+			}
+			else if (A.GetScriptRunIndex() == B.GetScriptRunIndex() && A.GetScriptExecutionIndex() > B.GetScriptExecutionIndex())
+			{
+				return true;
+			}
+			return false;
+		});
+
+		return OutHandles;
+	};
+
+	// Fix names of module namespace parameters that share names with other parameters in the same scope (e.g. InitializeParticle.Position and Particles.Position)
+	TMap<uint32, TArray<FSortedHistoryHandle*>> ExecutionIdxToSortedHandlesMap;
+	for (FSortedHistoryHandle& Handle : SortedHandles)
+	{
+		ExecutionIdxToSortedHandlesMap.FindOrAdd(Handle.GetScriptExecutionIndex()).Add(&Handle);
+	}
+
+	for (auto Iter = ExecutionIdxToSortedHandlesMap.CreateConstIterator(); Iter; ++Iter)
+	{
+		TArray<FSortedHistoryHandle*> PerExecIdxSortedHandles = Iter.Value();
+		TArray<FNiagaraHistoryVariable> ModuleHistoryVars;
+		TArray<FNiagaraHistoryVariable> InputAndOutputVars;
+		for (FSortedHistoryHandle* Handle : PerExecIdxSortedHandles)
+		{
+			ModuleHistoryVars.Append(Handle->OutputHistoryVars.FilterByPredicate([](const FNiagaraHistoryVariable& HistoryVar) {return HistoryVar.VarMetaData->bAddedToNodeGraphDeepCopy; }));
+			InputAndOutputVars.Append(Handle->InputHistoryVars);
+			InputAndOutputVars.Append(Handle->OutputHistoryVars.FilterByPredicate([](const FNiagaraHistoryVariable& HistoryVar) {return HistoryVar.VarMetaData->bAddedToNodeGraphDeepCopy == false; }));
+		}
+
+		TMap<FName, uint32> ModuleCollisionNameCounts;
+		for (FNiagaraHistoryVariable& ModuleHistoryVar : ModuleHistoryVars)
+		{
+			if (InputAndOutputVars.ContainsByPredicate([&ModuleHistoryVar](const FNiagaraHistoryVariable& TargetHistoryVar) {
+				return FNiagaraParameterHandle(ModuleHistoryVar.VarWithOriginalAlias->GetName()).GetName() == FNiagaraParameterHandle(TargetHistoryVar.VarWithOriginalAlias->GetName()).GetName();
+			}))
+			{
+				if (GNiagaraLogNamespaceFixup > 0)
+					UE_LOG(LogNiagaraEditor, Display, TEXT("Fixed aliased node graph deep copy param: %s"), *ModuleHistoryVar.VarWithOriginalAlias->GetName().ToString());
+				uint32& ModuleCollisionNameCount = ModuleCollisionNameCounts.FindOrAdd(ModuleHistoryVar.VarWithOriginalAlias->GetName());
+				++(ModuleCollisionNameCount);
+
+				FString ModuleNameString = ModuleHistoryVar.VarWithOriginalAlias->GetName().ToString();
+				ModuleNameString += LexToString(ModuleCollisionNameCount);
+				ModuleHistoryVar.SetVarWithOriginalAliasName(FName(*ModuleNameString));
+			}
+		}
+	}
+
+
+	// Track names of parameters we demote to ensure we do not have name aliases.
+	TMap<FName, uint32> DemotedNameCounts;
+
+	for (FSortedHistoryHandle& Handle : SortedHandles)
+	{ 
+		TArray<FSortedHistoryHandle*> UpstreamHandles = GetValidUpstreamHandles(Handle);
+		for (FSortedHistoryHandle* UpstreamSortedHistoryHandle : UpstreamHandles)
+		{
+			checkf(Handle.GetScriptExecutionIndex() != INDEX_NONE || Handle.GetScriptRunIndex() != INDEX_NONE, TEXT("Tried to access invalid handle from array!"));
+
+			for (int VarIdx = Handle.InputHistoryVars.Num() - 1; VarIdx > -1; --VarIdx)
+			{
+				// First check that we have not already established a link upstream.	
+				FNiagaraHistoryVariable& CurrentHistoryVar = Handle.InputHistoryVars[VarIdx];
+				FNiagaraParameterHandle CurrentInputVarHandle = FNiagaraParameterHandle(CurrentHistoryVar.VarWithOriginalAlias->GetName());
+
+				const FMatchedVariableInfo* UpstreamVarInfo = MatchedUpstreamVariableInfos.FindByPredicate(
+					[&CurrentInputVarHandle, &CurrentHistoryVar](const FMatchedVariableInfo& UpstreamVarInfo)
+				{ return CurrentHistoryVar.Var->GetType() == UpstreamVarInfo.GetTypeDef() && CurrentInputVarHandle.GetName() == UpstreamVarInfo.GetName(); }
+				);
+
+				if (UpstreamVarInfo != nullptr)
+				{
+					// There was already a match upstream, fix the variable and pin names.
+					FNiagaraVariable* MatchedInputVar = CurrentHistoryVar.Var;
+					const FNiagaraVariable* MatchedInputVarWithOriginalAlias = CurrentHistoryVar.VarWithOriginalAlias;
+					FName TargetName = UpstreamVarInfo->GetName();
+
+					if(GNiagaraLogNamespaceFixup > 0)
+						UE_LOG(LogNiagaraEditor, Display, TEXT("Matched Already Linked Input: %s | New Name: %s"), *MatchedInputVar->GetName().ToString(), *TargetName.ToString());
+
+					MatchedInputVar->SetName(TargetName);
+					Handle.Handle->FixupReadHistoryPinNames(MatchedInputVarWithOriginalAlias, TargetName);
+					Handle.InputHistoryVars.RemoveAt(VarIdx);
+				}
+				else
+				{
+					// This variable has not already been linked upstream, search for a potential link upstream.
+					int32 MatchIdx = INDEX_NONE;
+
+					auto FindUpstreamVarMatchIdx = [&Handle, &MatchIdx, &CurrentHistoryVar, &CurrentInputVarHandle, &UpstreamSortedHistoryHandle]()->bool {
+						const ENiagaraParameterScope UpstreamScope = UpstreamSortedHistoryHandle->GetScriptScope();
+						ENiagaraParameterScope TargetScope = ENiagaraParameterScope::None;
+						if (CurrentHistoryVar.VarMetaData->bAddedToNodeGraphDeepCopy)
+						{
+							// If the current history var is at the node graph deep copy level, it can link to any output, so the target scope is always that of the upstream history.
+							TargetScope = UpstreamScope;
+						}
+						else if (CurrentHistoryVar.VarMetaData->Scope == ENiagaraParameterScope::ScriptPersistent)
+						{
+							// If the current history var is targeting the script alias scope then get the scope of the current history.
+							TargetScope = Handle.GetScriptScope();
+						}
+						else
+						{
+							TargetScope = CurrentHistoryVar.VarMetaData->Scope;
+						}
+
+						MatchIdx = UpstreamSortedHistoryHandle->OutputHistoryVars.IndexOfByPredicate(
+							[&CurrentHistoryVar, &CurrentInputVarHandle, &TargetScope, &UpstreamScope](const FNiagaraHistoryVariable& UpstreamHistoryVar){
+							return TargetScope == UpstreamScope &&
+								CurrentHistoryVar.Var->GetType() == UpstreamHistoryVar.Var->GetType() &&
+								CurrentInputVarHandle.GetName() == FNiagaraParameterHandle(UpstreamHistoryVar.VarWithOriginalAlias->GetName()).GetName(); 
+						});
+
+						return MatchIdx != INDEX_NONE;
+					};
+
+					if (FindUpstreamVarMatchIdx())
+					{
+						// Found a match upstream. Continue checking for upstream matches, recording those matches if they occur and fixing those matches' variable and pin names. Then fix the current history variable's name and pin names.
+						FString FixedNamespaceString = FString();
+						FixedNamespaceString.Append(FNiagaraStackGraphUtilities::GetNamespaceStringForScriptParameterScope(UpstreamSortedHistoryHandle->GetScriptScope()));
+						FixedNamespaceString.Append(CurrentInputVarHandle.GetName().ToString());
+						FName FixedNamespaceName = FName(*FixedNamespaceString);
+
+						FString FixedEmitterAliasNamespaceString = FString();
+						FixedEmitterAliasNamespaceString.Append(UpstreamSortedHistoryHandle->GetEmitterUniqueName().ToString() + TEXT("."));
+						FixedEmitterAliasNamespaceString.Append(CurrentInputVarHandle.GetName().ToString());
+						FName FixedEmitterAliasNamespaceName = FName(*FixedEmitterAliasNamespaceString);
+
+						FName TargetFixedName;
+						if (UpstreamSortedHistoryHandle->GetScriptScope() == ENiagaraParameterScope::Emitter && Handle.GetScriptScope() != ENiagaraParameterScope::Emitter)
+						{
+							// If reading an Emitter parameter and not in the Emitter namespace, use the aliased name.
+							TargetFixedName = FixedEmitterAliasNamespaceName;
+						}
+						else
+						{
+							TargetFixedName = FixedNamespaceName;
+						}
+						
+						do {
+							// Fix the upstream matched output var's name and associated written pin names. Record the match to MatchedUpstreamVariables, and remove the history var from OutputHistoryVars.
+							FNiagaraHistoryVariable& MatchedHistoryOutputVar = UpstreamSortedHistoryHandle->OutputHistoryVars[MatchIdx];
+							FNiagaraVariable* MatchedOutputVar = MatchedHistoryOutputVar.Var;
+							const FNiagaraVariable* MatchedOutputVarWithOriginalAlias = MatchedHistoryOutputVar.VarWithOriginalAlias;
+
+							if (GNiagaraLogNamespaceFixup > 0)
+								UE_LOG(LogNiagaraEditor, Display, TEXT("Matched Output: %s | New Name: %s"), *MatchedOutputVar->GetName().ToString(), *TargetFixedName.ToString());
+
+							MatchedOutputVar->SetName(TargetFixedName);
+							UpstreamSortedHistoryHandle->Handle->FixupWrittenHistoryPinNames(MatchedOutputVarWithOriginalAlias, TargetFixedName);
+
+							// Record the upstream match in case there are additional downstream parameters to match to.
+							FMatchedVariableInfo MatchedVariableInfo = FMatchedVariableInfo(FixedEmitterAliasNamespaceName, FixedNamespaceName, MatchedHistoryOutputVar.Var->GetType());
+							MatchedUpstreamVariableInfos.Add(MatchedVariableInfo);
+
+							UpstreamSortedHistoryHandle->OutputHistoryVars.RemoveAtSwap(MatchIdx);
+						} while (FindUpstreamVarMatchIdx());
+
+						FNiagaraVariable* MatchedInputVar = CurrentHistoryVar.Var;
+						const FNiagaraVariable* MatchedInputVarWithOriginalAlias = CurrentHistoryVar.VarWithOriginalAlias;
+						if (GNiagaraLogNamespaceFixup > 0)
+							UE_LOG(LogNiagaraEditor, Display, TEXT("Matched Input: %s | New Name: %s"), *MatchedInputVar->GetName().ToString(), *FixedNamespaceName.ToString());
+						MatchedInputVar->SetName(FixedNamespaceName);
+						Handle.Handle->FixupReadHistoryPinNames(MatchedInputVarWithOriginalAlias, FixedNamespaceName);
+						Handle.InputHistoryVars.RemoveAt(VarIdx);
+					}
+				}
+			}
+		}
+	}
+
+	
+	if (GNiagaraForcePrecompilerCullDataset > 0 || bCullDatasetOutputs)
+	{
+		// Now change the namespace of any remaining unlinked inputs and outputs to one outside the dataset
+		auto ChangeVariableNamespace = [&DemotedNameCounts](FNiagaraVariable* InVar, const FString& InNewNamespaceString)->FName {
+			FString NewNameString = FString();
+			NewNameString.Append(InNewNamespaceString);
+			FName InVarName = InVar->GetName();
+			FString InVarNameString = InVarName.ToString();
+
+			uint32* DemotedNameCount = DemotedNameCounts.Find(InVarName);
+			if (DemotedNameCount == nullptr)
+			{
+				DemotedNameCounts.Add(InVarName) = 1;
+			}
+			else
+			{
+				InVarNameString += LexToString(*DemotedNameCount);
+				++(*DemotedNameCount);
+			}
+
+			NewNameString.Append(InVarNameString);
+			FName NewName = FName(*NewNameString);
+			InVar->SetName(NewName);
+			return NewName;
+		};
+
+
+		for (FSortedHistoryHandle& Handle : SortedHandles)
+		{
+			for (FNiagaraHistoryVariable& InputHistoryVar : Handle.InputHistoryVars)
+			{
+				if (InputHistoryVar.VarMetaData->Scope == ENiagaraParameterScope::None)
+				{
+					// Only need to fixup variables that are already in the dataset.
+					continue;
+				}
+				else if (InputHistoryVar.VarMetaData->Usage == ENiagaraScriptParameterUsage::InputOutput)
+				{
+					//@todo these may be culled if the param map write history does not link the input and output pins in any way
+					continue;
+				}
+
+				if (GNiagaraLogNamespaceFixup > 0)
+					UE_LOG(LogNiagaraEditor, Display, TEXT("Unmatched Input: %s"), *InputHistoryVar.VarWithOriginalAlias->GetName().ToString());
+				FName NewName = ChangeVariableNamespace(InputHistoryVar.Var, PARAM_MAP_INTERMEDIATE_STR);
+				Handle.Handle->FixupReadHistoryPinNames(InputHistoryVar.VarWithOriginalAlias, NewName);
+				Handle.Handle->FixupWrittenHistoryPinNames(InputHistoryVar.VarWithOriginalAlias, NewName);
+			}
+
+			for (FNiagaraHistoryVariable& OutputHistoryVar : Handle.OutputHistoryVars)
+			{
+				if (OutputHistoryVar.bDoNotDemote)
+				{
+					continue;
+				}
+				else if (OutputHistoryVar.VarMetaData->Scope == ENiagaraParameterScope::None)
+				{
+					// Only need to fixup variables that are already in the dataset.
+					continue;
+				}
+				else if (OutputHistoryVar.VarMetaData->Usage == ENiagaraScriptParameterUsage::InputOutput)
+				{
+					//@todo these may be culled if the param map write history does not link the input and output pins in any way
+					continue;
+				}
+				else if (OutputHistoryVar.VarMetaData->Usage == ENiagaraScriptParameterUsage::Output && OutputHistoryVar.VarMetaData->bOutputIsPersistent)
+				{
+					if (Handle.InputHistoryVars.ContainsByPredicate([&OutputHistoryVar](const FNiagaraHistoryVariable& InputHistoryVar) {return InputHistoryVar.VarWithOriginalAlias == OutputHistoryVar.VarWithOriginalAlias; }))
+					{
+						// Skip unmatched persistent outputs that are inputs to the same history as they must remain in the dataset to be persistent.
+						continue;
+					}
+				}
+
+				if (GNiagaraLogNamespaceFixup > 0)
+					UE_LOG(LogNiagaraEditor, Display, TEXT("Unmatched Output: %s"), *OutputHistoryVar.VarWithOriginalAlias->GetName().ToString());
+				FName NewName = ChangeVariableNamespace(OutputHistoryVar.Var, PARAM_MAP_INTERMEDIATE_STR);
+				Handle.Handle->FixupReadHistoryPinNames(OutputHistoryVar.VarWithOriginalAlias, NewName);
+				Handle.Handle->FixupWrittenHistoryPinNames(OutputHistoryVar.VarWithOriginalAlias, NewName);
+			}
+		}
+	}
+}
+
+int32 FNiagaraParameterMapHistoryWithMetaDataBuilder::AddVariableToHistory(FNiagaraParameterMapHistory& History, const FNiagaraVariable& InVar, const FNiagaraVariable& InAliasedVar, const UEdGraphPin* InPin)
+{
+	TOptional<FNiagaraVariableMetaData> MetaData = CallingGraphContext.Last()->GetMetaData(InAliasedVar);
+	if (MetaData.IsSet())
+	{
+		if (CallingGraphContext.Num() == 1)
+		{
+			// Mark vars added to the node graph deep copy to fixup name during precompile, and allow linking downstream even if their scope is local.
+			MetaData->bAddedToNodeGraphDeepCopy = true;
+		}
+		return History.AddVariable(InVar, InAliasedVar, InPin, MetaData);
+	}
+	UE_LOG(LogNiagaraEditor, Warning, TEXT("Variable added to parameter map history did not have metadata! Variable: %s"), *InVar.GetName().ToString());
+	TOptional<FNiagaraVariableMetaData> BlankMetaData = FNiagaraVariableMetaData();
+	return History.AddVariable(InVar, InAliasedVar, InPin, BlankMetaData);
+}
+
+FNiagaraParameterMapHistoryHandle::FNiagaraParameterMapHistoryHandle()
+{
+	ScriptExecutionIndex = INDEX_NONE;
+	ScriptRunIndex = INDEX_NONE;
+	Scope = ENiagaraParameterScope::None;
+	HistoryVariables.Empty();
+	EmitterUniqueName = FName();
+}
+
+FNiagaraParameterMapHistoryHandle::FNiagaraParameterMapHistoryHandle(FNiagaraParameterMapHistory& InHistory, const TArray<FNiagaraVariable>& InRequiredRendererVariables, const FName InEmitterUniqueName /*= FName()*/)
+{
+	History = &InHistory;
+	RequiredRendererVariables = InRequiredRendererVariables;
+	EmitterUniqueName = InEmitterUniqueName;
+
+	Scope = FNiagaraStackGraphUtilities::GetScopeForScriptUsage(InHistory.OriginatingScriptUsage);
+	const FString HistoryScopeString = FNiagaraStackGraphUtilities::GetNamespaceStringForScriptParameterScope(Scope);
+	FNiagaraEditorUtilities::GetScriptRunAndExecutionIndexFromUsage(InHistory.OriginatingScriptUsage, ScriptRunIndex, ScriptExecutionIndex);
+
+	for (int VarIdx = 0; VarIdx < InHistory.Variables.Num(); ++VarIdx)
+	{
+		FNiagaraVariable* Var = &InHistory.Variables[VarIdx];
+		FNiagaraVariable* VarWithOriginalAliasIntact = &InHistory.VariablesWithOriginalAliasesIntact[VarIdx];
+		FNiagaraVariableMetaData* VarMetaData = &InHistory.VariableMetaData[VarIdx];
+		HistoryVariables.Add(FNiagaraHistoryVariable(Var, VarWithOriginalAliasIntact, VarMetaData));
+	}
+}
+
+//** Constructs a history handle as a view of a secondary history to a primary history. e.g. the primary history encompasses all parameters in a system and the secondary history only has parameters for a single emitter. //
+FNiagaraParameterMapHistoryHandle::FNiagaraParameterMapHistoryHandle(FNiagaraParameterMapHistory& InPrimaryHistory, FNiagaraParameterMapHistory& InSecondaryHistory, const TArray<FNiagaraVariable>& InRequiredRendererVariables, const FName InEmitterUniqueName /*= FName()*/)
+{
+	History = &InPrimaryHistory;
+	RequiredRendererVariables = InRequiredRendererVariables;
+	EmitterUniqueName = InEmitterUniqueName;
+
+	Scope = FNiagaraStackGraphUtilities::GetScopeForScriptUsage(InSecondaryHistory.OriginatingScriptUsage);
+	const FString HistoryScopeString = FNiagaraStackGraphUtilities::GetNamespaceStringForScriptParameterScope(Scope);
+	FNiagaraEditorUtilities::GetScriptRunAndExecutionIndexFromUsage(InSecondaryHistory.OriginatingScriptUsage, ScriptRunIndex, ScriptExecutionIndex);
+
+	for (int32 SecondaryVarIdx = 0; SecondaryVarIdx < InSecondaryHistory.Variables.Num(); ++SecondaryVarIdx)
+	{
+		FNiagaraVariable& SecondaryVar = InSecondaryHistory.Variables[SecondaryVarIdx];
+		int32 FoundPrimaryVarIdx = InPrimaryHistory.Variables.IndexOfByPredicate([SecondaryVar](const FNiagaraVariable& PrimaryVar) {return SecondaryVar == PrimaryVar; });
+		if (FoundPrimaryVarIdx != INDEX_NONE)
+		{
+			FNiagaraVariable* PrimaryVar = &InPrimaryHistory.Variables[FoundPrimaryVarIdx];
+			FNiagaraVariable* PrimaryVarWithOriginalAliasIntact = &InPrimaryHistory.VariablesWithOriginalAliasesIntact[FoundPrimaryVarIdx];
+
+			// Assign the secondary history metadata as it is canonical in regards to parameter linking.
+			FNiagaraVariableMetaData* ScriptVarMetaData = &InSecondaryHistory.VariableMetaData[SecondaryVarIdx];
+			HistoryVariables.Add(FNiagaraHistoryVariable(PrimaryVar, PrimaryVarWithOriginalAliasIntact, ScriptVarMetaData));
+		}
+	}
+}
+
+void FNiagaraParameterMapHistoryHandle::FixupWrittenHistoryPinNames(int32& VarIdx, const FName& NewName)
+{
+	for (const UEdGraphPin* Pin : History->PerVariableWriteHistory[VarIdx])
+	{
+		UEdGraphPin* NonConstPin = const_cast<UEdGraphPin*>(Pin);
+		if (GNiagaraLogNamespaceFixup > 0)
+			UE_LOG(LogNiagaraEditor, Display, TEXT("Fixed map set input pin: %s"), *NonConstPin->PinName.ToString());
+		NonConstPin->PinName = NewName;
+	}
+}
+
+void FNiagaraParameterMapHistoryHandle::FixupWrittenHistoryPinNames(const FNiagaraVariable* AssociatedVar, const FName& NewName)
+{
+	int32 VarIdx = GetSourceVariableHistoryIndex(AssociatedVar);
+	FixupWrittenHistoryPinNames(VarIdx, NewName);
+}
+
+void FNiagaraParameterMapHistoryHandle::FixupReadHistoryPinNames(int32& VarIdx, const FName& NewName)
+{
+	for (TTuple<const UEdGraphPin*, const UEdGraphPin*>& PinTuple : History->PerVariableReadHistory[VarIdx])
+	{
+		UEdGraphPin* Pin = const_cast<UEdGraphPin*>(PinTuple.Key);
+		if (GNiagaraLogNamespaceFixup > 0)
+			UE_LOG(LogNiagaraEditor, Display, TEXT("Fixed map get output pin: %s"), *Pin->PinName.ToString());
+		Pin->PinName = NewName;
+	}
+}
+
+void FNiagaraParameterMapHistoryHandle::FixupReadHistoryPinNames(const FNiagaraVariable* AssociatedVar, const FName& NewName)
+{
+	int32 VarIdx = GetSourceVariableHistoryIndex(AssociatedVar);
+	FixupReadHistoryPinNames(VarIdx, NewName);
+}
+
+void FNiagaraHistoryVariable::SetVarWithOriginalAliasName(const FName& NewName)
+{
+	FNiagaraVariable* NonConstVarWithOriginalAlias = const_cast<FNiagaraVariable*>(VarWithOriginalAlias);
+	NonConstVarWithOriginalAlias->SetName(NewName);
+}
diff --git a/Engine/Plugins/FX/Niagara/Source/NiagaraEditor/Private/TNiagaraGraphPinEditableName.h b/Engine/Plugins/FX/Niagara/Source/NiagaraEditor/Private/TNiagaraGraphPinEditableName.h
--- a/Engine/Plugins/FX/Niagara/Source/NiagaraEditor/Private/TNiagaraGraphPinEditableName.h
+++ b/Engine/Plugins/FX/Niagara/Source/NiagaraEditor/Private/TNiagaraGraphPinEditableName.h
@@ -81,21 +81,29 @@ protected:
 		UNiagaraNode* ParentNode = Cast<UNiagaraNode>(this->GraphPinObj->GetOwningNode());
 		if (ParentNode && ParentNode->IsPinNameEditable(this->GraphPinObj))
 		{
-			CreatedTextBlock = SNew(SInlineEditableTextBlock)
-				.Style(&FEditorStyle::Get().GetWidgetStyle<FInlineEditableTextBlockStyle>("Graph.Node.InlineEditablePinName"))
-				.Text(this, &TNiagaraGraphPinEditableName<BaseClass>::GetParentPinLabel)
-				.Visibility(this, &TNiagaraGraphPinEditableName<BaseClass>::GetParentPinVisibility)
-				.ColorAndOpacity(this, &TNiagaraGraphPinEditableName<BaseClass>::GetParentPinTextColor)
-				.OnVerifyTextChanged(this, &TNiagaraGraphPinEditableName<BaseClass>::OnVerifyTextChanged)
-				.OnTextCommitted(this, &TNiagaraGraphPinEditableName<BaseClass>::OnTextCommitted);
+			UNiagaraGraph* NiagaraGraph = ParentNode->GetNiagaraGraph();
+			if (FNiagaraEditorCommonCVar::GNiagaraEnableParameterPanel2 > 0 && NiagaraGraph->IsPinVisualWidgetProviderRegistered())
+			{
+				return NiagaraGraph->GetPinVisualWidget(this->GraphPinObj);
+			}
+			else
+			{
+				CreatedTextBlock = SNew(SInlineEditableTextBlock)
+					.Style(&FEditorStyle::Get().GetWidgetStyle<FInlineEditableTextBlockStyle>("Graph.Node.InlineEditablePinName"))
+					.Text(this, &TNiagaraGraphPinEditableName<BaseClass>::GetParentPinLabel)
+					.Visibility(this, &TNiagaraGraphPinEditableName<BaseClass>::GetParentPinVisibility)
+					.ColorAndOpacity(this, &TNiagaraGraphPinEditableName<BaseClass>::GetParentPinTextColor)
+					.OnVerifyTextChanged(this, &TNiagaraGraphPinEditableName<BaseClass>::OnVerifyTextChanged)
+					.OnTextCommitted(this, &TNiagaraGraphPinEditableName<BaseClass>::OnTextCommitted);
 
-			TSharedRef<SWidget> EditableTextBlock = CreatedTextBlock.ToSharedRef();
+				TSharedRef<SWidget> EditableTextBlock = CreatedTextBlock.ToSharedRef();
 
-			if (ParentNode->IsPinNameEditableUponCreation(this->GraphPinObj))
-			{
-				bPendingRename = true;
+				if (ParentNode->IsPinNameEditableUponCreation(this->GraphPinObj))
+				{
+					bPendingRename = true;
+				}
+				return EditableTextBlock;
 			}
-			return EditableTextBlock;
 		}
 		else
 		{
diff --git a/Engine/Plugins/FX/Niagara/Source/NiagaraEditor/Private/Toolkits/NiagaraScriptToolkit.cpp b/Engine/Plugins/FX/Niagara/Source/NiagaraEditor/Private/Toolkits/NiagaraScriptToolkit.cpp
--- a/Engine/Plugins/FX/Niagara/Source/NiagaraEditor/Private/Toolkits/NiagaraScriptToolkit.cpp
+++ b/Engine/Plugins/FX/Niagara/Source/NiagaraEditor/Private/Toolkits/NiagaraScriptToolkit.cpp
@@ -15,6 +15,7 @@
 #include "SNiagaraParameterCollection.h"
 #include "UObject/Package.h"
 #include "SNiagaraParameterMapView.h"
+#include "SNiagaraParameterPanel.h"
 #include "NiagaraEditorStyle.h"
 #include "NiagaraSystem.h"
 #include "BusyCursor.h"
@@ -47,6 +48,7 @@
 #include "NiagaraMessageManager.h"
 #include "NiagaraStandaloneScriptViewModel.h"
 #include "NiagaraMessageLogViewModel.h"
+#include "NiagaraParameterPanelViewModel.h"
 
 #define LOCTEXT_NAMESPACE "NiagaraScriptToolkit"
 
@@ -56,6 +58,7 @@ const FName FNiagaraScriptToolkit::NodeGraphTabId(TEXT("NiagaraEditor_NodeGraph"
 const FName FNiagaraScriptToolkit::ScriptDetailsTabId(TEXT("NiagaraEditor_ScriptDetails"));
 const FName FNiagaraScriptToolkit::SelectedDetailsTabId(TEXT("NiagaraEditor_SelectedDetails"));
 const FName FNiagaraScriptToolkit::ParametersTabId(TEXT("NiagaraEditor_Parameters"));
+const FName FNiagaraScriptToolkit::ParametersTabId2(TEXT("NiagaraEditor_Paramters2"));
 const FName FNiagaraScriptToolkit::StatsTabId(TEXT("NiagaraEditor_Stats"));
 const FName FNiagaraScriptToolkit::MessageLogTabID(TEXT("NiagaraEditor_MessageLog"));
 
@@ -97,6 +100,13 @@ void FNiagaraScriptToolkit::RegisterTabSpawners(const TSharedRef<class FTabManag
 	InTabManager->RegisterTabSpawner(ParametersTabId, FOnSpawnTab::CreateSP(this, &FNiagaraScriptToolkit::SpawnTabScriptParameters))
 		.SetDisplayName(LOCTEXT("ParametersTab", "Parameters"))
 		.SetGroup(WorkspaceMenuCategoryRef);
+
+	if (FNiagaraEditorCommonCVar::GNiagaraEnableParameterPanel2 > 0)
+	{
+		InTabManager->RegisterTabSpawner(ParametersTabId2, FOnSpawnTab::CreateSP(this, &FNiagaraScriptToolkit::SpawnTabScriptParameters2))
+			.SetDisplayName(LOCTEXT("ParametersTab2", "Parameters2"))
+			.SetGroup(WorkspaceMenuCategoryRef);
+	}
 		
 	InTabManager->RegisterTabSpawner(StatsTabId, FOnSpawnTab::CreateSP(this, &FNiagaraScriptToolkit::SpawnTabStats))
 		.SetDisplayName(LOCTEXT("StatsTab", "Stats"))
@@ -116,6 +126,7 @@ void FNiagaraScriptToolkit::UnregisterTabSpawners(const TSharedRef<class FTabMan
 	InTabManager->UnregisterTabSpawner(ScriptDetailsTabId);
 	InTabManager->UnregisterTabSpawner(SelectedDetailsTabId);
 	InTabManager->UnregisterTabSpawner(ParametersTabId);
+	InTabManager->UnregisterTabSpawner(ParametersTabId2);
 	InTabManager->UnregisterTabSpawner(StatsTabId);
 }
 
@@ -149,6 +160,8 @@ void FNiagaraScriptToolkit::Initialize( const EToolkitMode::Type Mode, const TSh
 
 	ScriptViewModel = MakeShareable(new FNiagaraStandaloneScriptViewModel(DisplayName, ENiagaraParameterEditMode::EditAll, NiagaraMessageLogViewModel, MessageLogGuidKey));
 	ScriptViewModel->Initialize(EditedNiagaraScript, OriginalNiagaraScript);
+	ParameterPanelViewModel = MakeShareable(new FNiagaraScriptToolkitParameterPanelViewModel(ScriptViewModel));
+	ParameterPanelViewModel->InitBindings();
 
 	OnEditedScriptGraphChangedHandle = ScriptViewModel->GetGraphViewModel()->GetGraph()->AddOnGraphNeedsRecompileHandler(
 		FOnGraphChanged::FDelegate::CreateRaw(this, &FNiagaraScriptToolkit::OnEditedScriptGraphChanged));
@@ -381,6 +394,19 @@ TSharedRef<SDockTab> FNiagaraScriptToolkit::SpawnTabScriptParameters(const FSpaw
 	return SpawnedTab;
 }
 
+TSharedRef<SDockTab> FNiagaraScriptToolkit::SpawnTabScriptParameters2(const FSpawnTabArgs& Args)
+{
+	checkf(Args.GetTabId().TabType == ParametersTabId2, TEXT("Wrong tab ID in NiagaraScriptToolkit"));
+
+	TSharedRef<SDockTab> SpawnedTab =
+		SNew(SDockTab)
+		[
+			SNew(SNiagaraParameterPanel, ParameterPanelViewModel, GetToolkitCommands())
+		];
+
+	return SpawnedTab;
+}
+
 TSharedRef<SDockTab> FNiagaraScriptToolkit::SpawnTabStats(const FSpawnTabArgs& Args)
 {
 	check(Args.GetTabId() == StatsTabId);
diff --git a/Engine/Plugins/FX/Niagara/Source/NiagaraEditor/Private/Toolkits/NiagaraScriptToolkit.h b/Engine/Plugins/FX/Niagara/Source/NiagaraEditor/Private/Toolkits/NiagaraScriptToolkit.h
--- a/Engine/Plugins/FX/Niagara/Source/NiagaraEditor/Private/Toolkits/NiagaraScriptToolkit.h
+++ b/Engine/Plugins/FX/Niagara/Source/NiagaraEditor/Private/Toolkits/NiagaraScriptToolkit.h
@@ -22,6 +22,7 @@ class FNiagaraObjectSelection;
 struct FEdGraphEditAction;
 class FNiagaraMessageLogViewModel;
 class FNiagaraStandaloneScriptViewModel;
+class FNiagaraScriptToolkitParameterPanelViewModel;
 
 /** Viewer/editor for a DataTable */
 class FNiagaraScriptToolkit : public FAssetEditorToolkit, public FGCObject
@@ -82,6 +83,7 @@ private:
 
 	/** Spawns the tab with the script details inside. */
 	TSharedRef<SDockTab> SpawnTabScriptParameters(const FSpawnTabArgs& Args);
+	TSharedRef<SDockTab> SpawnTabScriptParameters2(const FSpawnTabArgs& Args);
 	TSharedRef<SDockTab> SpawnTabStats(const FSpawnTabArgs& Args);
 
 	TSharedRef<SDockTab> SpawnTabMessageLog(const FSpawnTabArgs& Args);
@@ -125,6 +127,9 @@ private:
 	/** The Script being edited */
 	TSharedPtr<FNiagaraStandaloneScriptViewModel> ScriptViewModel;
 
+	/** The Parameter Panel displaying graph variables */
+	TSharedPtr<FNiagaraScriptToolkitParameterPanelViewModel> ParameterPanelViewModel;
+
 	/** The selection displayed by the details tab. */
 	TSharedPtr<FNiagaraObjectSelection> DetailsScriptSelection;
 
@@ -137,6 +142,7 @@ private:
 	static const FName ScriptDetailsTabId;
 	static const FName SelectedDetailsTabId;
 	static const FName ParametersTabId;
+	static const FName ParametersTabId2;
 	static const FName StatsTabId;
 	static const FName MessageLogTabID;
 
diff --git a/Engine/Plugins/FX/Niagara/Source/NiagaraEditor/Private/Toolkits/NiagaraSystemToolkit.cpp b/Engine/Plugins/FX/Niagara/Source/NiagaraEditor/Private/Toolkits/NiagaraSystemToolkit.cpp
--- a/Engine/Plugins/FX/Niagara/Source/NiagaraEditor/Private/Toolkits/NiagaraSystemToolkit.cpp
+++ b/Engine/Plugins/FX/Niagara/Source/NiagaraEditor/Private/Toolkits/NiagaraSystemToolkit.cpp
@@ -14,7 +14,8 @@
 #include "Widgets/SNiagaraSystemScript.h"
 #include "Widgets/SNiagaraSystemViewport.h"
 #include "Widgets/SNiagaraSelectedObjectsDetails.h"
-#include "Widgets/SNiagaraParameterMapView.h"
+#include "Widgets/SNiagaraParameterMapView.h" //@todo(ng) cleanup
+#include "Widgets/SNiagaraParameterPanel.h"
 #include "Widgets/SNiagaraSpreadsheetView.h"
 #include "Widgets/SNiagaraGeneratedCodeView.h"
 #include "Widgets/SNiagaraScriptGraph.h"
@@ -53,6 +54,7 @@
 #include "ViewModels/NiagaraOverviewGraphViewModel.h"
 #include "NiagaraScriptSourceBase.h"
 #include "Widgets/Input/SCheckBox.h"
+#include "NiagaraParameterPanelViewModel.h"
 
 #define LOCTEXT_NAMESPACE "NiagaraSystemEditor"
 
@@ -64,6 +66,7 @@ const FName FNiagaraSystemToolkit::SequencerTabID(TEXT("NiagaraSystemEditor_Sequ
 const FName FNiagaraSystemToolkit::SystemScriptTabID(TEXT("NiagaraSystemEditor_SystemScript"));
 const FName FNiagaraSystemToolkit::SystemDetailsTabID(TEXT("NiagaraSystemEditor_SystemDetails"));
 const FName FNiagaraSystemToolkit::SystemParametersTabID(TEXT("NiagaraSystemEditor_SystemParameters"));
+const FName FNiagaraSystemToolkit::SystemParametersTabID2(TEXT("NiagaraSystemEditor_SystemParameters2"));
 const FName FNiagaraSystemToolkit::SelectedEmitterStackTabID(TEXT("NiagaraSystemEditor_SelectedEmitterStack"));
 const FName FNiagaraSystemToolkit::SelectedEmitterGraphTabID(TEXT("NiagaraSystemEditor_SelectedEmitterGraph"));
 const FName FNiagaraSystemToolkit::DebugSpreadsheetTabID(TEXT("NiagaraSystemEditor_DebugAttributeSpreadsheet"));
@@ -120,6 +123,13 @@ void FNiagaraSystemToolkit::RegisterTabSpawners(const TSharedRef<class FTabManag
 		.SetDisplayName(LOCTEXT("SystemParameters", "Parameters"))
 		.SetGroup(WorkspaceMenuCategory.ToSharedRef());
 
+	if (FNiagaraEditorCommonCVar::GNiagaraEnableParameterPanel2 > 0)
+	{
+		InTabManager->RegisterTabSpawner(SystemParametersTabID2, FOnSpawnTab::CreateSP(this, &FNiagaraSystemToolkit::SpawnTab_SystemParameters2))
+			.SetDisplayName(LOCTEXT("SystemParameters2", "Parameters2"))
+			.SetGroup(WorkspaceMenuCategory.ToSharedRef());
+	}
+
 	InTabManager->RegisterTabSpawner(SelectedEmitterStackTabID, FOnSpawnTab::CreateSP(this, &FNiagaraSystemToolkit::SpawnTab_SelectedEmitterStack))
 		.SetDisplayName(LOCTEXT("SelectedEmitterStacks", "Selected Emitters"))
 		.SetGroup(WorkspaceMenuCategory.ToSharedRef());
@@ -208,6 +218,8 @@ void FNiagaraSystemToolkit::InitializeWithSystem(const EToolkitMode::Type Mode,
 
 	SystemViewModel = MakeShared<FNiagaraSystemViewModel>();
 	SystemViewModel->Initialize(*System, SystemOptions);
+	ParameterPanelViewModel = MakeShared<FNiagaraSystemToolkitParameterPanelViewModel>(SystemViewModel.ToSharedRef());
+	ParameterPanelViewModel->InitBindings();
 	SystemViewModel->SetToolkitCommands(GetToolkitCommands());
 	SystemToolkitMode = ESystemToolkitMode::System;
 
@@ -277,6 +289,8 @@ void FNiagaraSystemToolkit::InitializeWithEmitter(const EToolkitMode::Type Mode,
 
 	SystemViewModel = MakeShared<FNiagaraSystemViewModel>();
 	SystemViewModel->Initialize(*System, SystemOptions);
+	ParameterPanelViewModel = MakeShared<FNiagaraSystemToolkitParameterPanelViewModel>(SystemViewModel);
+	ParameterPanelViewModel->InitBindings();
 	SystemViewModel->GetEditorData().SetOwningSystemIsPlaceholder(true, *System);
 	SystemViewModel->SetToolkitCommands(GetToolkitCommands());
 	SystemViewModel->AddEmitter(*Emitter);
@@ -518,7 +532,25 @@ TSharedRef<SDockTab> FNiagaraSystemToolkit::SpawnTab_SystemParameters(const FSpa
 	TSharedRef<SDockTab> SpawnedTab =
 		SNew(SDockTab)
 		[
-			SAssignNew(ParameterMapView, SNiagaraParameterMapView, ObjectSelections, SNiagaraParameterMapView::EToolkitType::SYSTEM, GetToolkitCommands())
+			SAssignNew(ParameterMapView, SNiagaraParameterMapView, ObjectSelections, SNiagaraParameterMapView::EToolkitType::SYSTEM, GetToolkitCommands()) //@todo(ng) cleanup
+		];
+	RefreshParameters();
+
+	return SpawnedTab;
+}
+
+TSharedRef<SDockTab> FNiagaraSystemToolkit::SpawnTab_SystemParameters2(const FSpawnTabArgs& Args)
+{
+	check(Args.GetTabId().TabType == SystemParametersTabID2);
+
+
+	TArray<TSharedRef<FNiagaraObjectSelection>> ObjectSelections;
+	ObjectSelections.Add(ObjectSelectionForParameterMapView.ToSharedRef());
+
+	TSharedRef<SDockTab> SpawnedTab =
+		SNew(SDockTab)
+		[
+			SAssignNew(ParameterPanel, SNiagaraParameterPanel, ParameterPanelViewModel, GetToolkitCommands())
 		];
 	RefreshParameters();
 
diff --git a/Engine/Plugins/FX/Niagara/Source/NiagaraEditor/Private/Toolkits/NiagaraSystemToolkit.h b/Engine/Plugins/FX/Niagara/Source/NiagaraEditor/Private/Toolkits/NiagaraSystemToolkit.h
--- a/Engine/Plugins/FX/Niagara/Source/NiagaraEditor/Private/Toolkits/NiagaraSystemToolkit.h
+++ b/Engine/Plugins/FX/Niagara/Source/NiagaraEditor/Private/Toolkits/NiagaraSystemToolkit.h
@@ -30,6 +30,7 @@ struct FAssetData;
 class FMenuBuilder;
 class ISequencer;
 class FNiagaraMessageLogViewModel;
+class FNiagaraSystemToolkitParameterPanelViewModel;
 
 /** Viewer/editor for a NiagaraSystem
 */
@@ -100,6 +101,7 @@ private:
 	TSharedRef<SDockTab> SpawnTab_Sequencer(const FSpawnTabArgs& Args);
 	TSharedRef<SDockTab> SpawnTab_SystemScript(const FSpawnTabArgs& Args);
 	TSharedRef<SDockTab> SpawnTab_SystemParameters(const FSpawnTabArgs& Args);
+	TSharedRef<SDockTab> SpawnTab_SystemParameters2(const FSpawnTabArgs& Args);
 	TSharedRef<SDockTab> SpawnTab_SelectedEmitterStack(const FSpawnTabArgs& Args);
 	TSharedRef<SDockTab> SpawnTab_SelectedEmitterGraph(const FSpawnTabArgs& Args);
 	TSharedRef<SDockTab> SpawnTab_DebugSpreadsheet(const FSpawnTabArgs& Args);
@@ -171,7 +173,10 @@ private:
 	/** The command list for this editor */
 	TSharedPtr<FUICommandList> EditorCommands;
 
-	TSharedPtr<class SNiagaraParameterMapView> ParameterMapView;
+	TSharedPtr<class SNiagaraParameterMapView> ParameterMapView; //@todo(ng) cleanup
+
+	TSharedPtr<FNiagaraSystemToolkitParameterPanelViewModel> ParameterPanelViewModel;
+	TSharedPtr<class SNiagaraParameterPanel> ParameterPanel;
 
 	TSharedPtr<FNiagaraObjectSelection> ObjectSelectionForParameterMapView;
 
@@ -184,6 +189,7 @@ public:
 	static const FName SystemScriptTabID;
 	static const FName SystemDetailsTabID;
 	static const FName SystemParametersTabID;
+	static const FName SystemParametersTabID2;
 	static const FName SelectedEmitterStackTabID;
 	static const FName SelectedEmitterGraphTabID;
 	static const FName DebugSpreadsheetTabID;
diff --git a/Engine/Plugins/FX/Niagara/Source/NiagaraEditor/Private/ViewModels/NiagaraParameterNameViewModel.cpp b/Engine/Plugins/FX/Niagara/Source/NiagaraEditor/Private/ViewModels/NiagaraParameterNameViewModel.cpp
new file mode 100644
--- /dev/null
+++ b/Engine/Plugins/FX/Niagara/Source/NiagaraEditor/Private/ViewModels/NiagaraParameterNameViewModel.cpp
@@ -0,0 +1,197 @@
+// Copyright Epic Games, Inc. All Rights Reserved.
+
+#include "NiagaraParameterNameViewModel.h"
+#include "EditorWidgets/Public/SEnumComboBox.h"
+#include "Widgets/Text/SInlineEditableTextBlock.h"
+#include "SGraphActionMenu.h"
+#include "NiagaraActions.h"
+#include "ViewModels/Stack/NiagaraStackGraphUtilities.h"
+#include "NiagaraNode.h"
+#include "NiagaraParameterPanelViewModel.h"
+#include "NiagaraEditorUtilities.h"
+#include "NiagaraEditorModule.h"
+
+///////////////////////////////////////////////////////////////////////////////
+/// Parameter Panel Entry Parameter Name ViewModel							///
+///////////////////////////////////////////////////////////////////////////////
+
+FNiagaraParameterPanelEntryParameterNameViewModel::FNiagaraParameterPanelEntryParameterNameViewModel(FCreateWidgetForActionData* const InCreateData, const FNiagaraScriptVariableAndViewInfo& InScriptVarAndViewInfo)
+	: CreateData(InCreateData)
+	, CachedScriptVarAndViewInfo(InScriptVarAndViewInfo)
+{};
+
+TSharedRef<SWidget> FNiagaraParameterPanelEntryParameterNameViewModel::CreateScopeSlotWidget() const
+{
+	UEnum* ParameterScopeEnum = FNiagaraTypeDefinition::GetParameterScopeEnum();
+	TSharedPtr< SWidget > ScopeComboBoxWidget;
+	FNiagaraEditorModule& NiagaraEditorModule = FModuleManager::LoadModuleChecked<FNiagaraEditorModule>("NiagaraEditor");
+
+	bool bEnableScopeSlotWidget = true;
+	if (CachedScriptVarAndViewInfo.MetaData.Usage != ENiagaraScriptParameterUsage::Input)
+	{
+		bEnableScopeSlotWidget = false;
+	}
+	else if (CachedScriptVarAndViewInfo.MetaData.Scope == ENiagaraParameterScope::Local)
+	{
+		bEnableScopeSlotWidget = false;
+	}
+	else if (CachedScriptVarAndViewInfo.MetaData.bIsStaticSwitch)
+	{
+		bEnableScopeSlotWidget = false;
+	}
+
+	SAssignNew(ScopeComboBoxWidget, SBox)
+		.MinDesiredWidth(80.0f) //@todo(ng) tune and make const static
+		.Content()
+		[
+			SNew(SEnumComboBox, ParameterScopeEnum)
+			.IsEnabled(bEnableScopeSlotWidget)
+			//.ForegroundColor(NiagaraEditorModule.GetWidgetProvider()->GetColorForParameterScope(CachedScriptVarAndViewInfo.MetaData.Scope)) //@todo(ng) get background color instead
+			.CurrentValue(this, &FNiagaraParameterPanelEntryParameterNameViewModel::GetScopeValue)
+			.OnEnumSelectionChanged(this, &FNiagaraParameterPanelEntryParameterNameViewModel::OnScopeValueChanged)
+			.OnGetToolTipForValue(this, &FNiagaraParameterPanelEntryParameterNameViewModel::GetToolTipForScopeValue)
+			//.OnIsValueEnabled(this, &FNiagaraParameterPanelEntryParameterNameViewModel::GetScopeValueIsEnabled) //@todo(ng) impl
+		];
+
+	return ScopeComboBoxWidget.ToSharedRef();
+}
+
+TSharedRef<SWidget> FNiagaraParameterPanelEntryParameterNameViewModel::CreateTextSlotWidget() const
+{
+	TSharedPtr< SWidget > DisplayWidget;
+	const FName FontType = FName("Italic");
+	FSlateFontInfo NameFont = FCoreStyle::GetDefaultFontStyle(FontType, 10); //@todo(ng) get correct font
+
+	SAssignNew(DisplayWidget, SInlineEditableTextBlock)
+		.Text(this, &FNiagaraParameterPanelEntryParameterNameViewModel::GetParameterNameText)
+		.Font(NameFont)
+		.HighlightText(CreateData->HighlightText)
+		.OnTextCommitted(this, &FNiagaraParameterPanelEntryParameterNameViewModel::OnParameterRenamed)
+		//.OnVerifyTextChanged(InArgs._OnVerifyParameterNameChanged) //@todo(ng) impl
+		.IsSelected(CreateData->IsRowSelectedDelegate)
+		.IsReadOnly(bIsReadOnly);
+
+	//CreateData->OnRenameRequest->BindSP(static_cast<SInlineEditableTextBlock*>(DisplayWidget.Get()), &SInlineEditableTextBlock::EnterEditingMode); //@Todo(ng) figure out a way to make this work
+
+	return DisplayWidget.ToSharedRef();
+}
+
+int32 FNiagaraParameterPanelEntryParameterNameViewModel::GetScopeValue() const
+{
+	return (int32)CachedScriptVarAndViewInfo.MetaData.Scope;
+}
+
+void FNiagaraParameterPanelEntryParameterNameViewModel::OnScopeValueChanged(int32 NewScopeValue, ESelectInfo::Type SelectionType) const
+{
+	OnScopeSelectionChangedDelegate.ExecuteIfBound(CachedScriptVarAndViewInfo.ScriptVariable, CachedScriptVarAndViewInfo.MetaData, ENiagaraParameterScope(NewScopeValue));
+}
+
+FText FNiagaraParameterPanelEntryParameterNameViewModel::GetToolTipForScopeValue(int32 ScopeValue) const
+{
+	return CachedScriptVarAndViewInfo.ParameterScopeToDisplayInfo[ScopeValue].Tooltip;
+}
+
+bool FNiagaraParameterPanelEntryParameterNameViewModel::GetScopeValueIsEnabled(int32 ScopeValue) const
+{
+	return CachedScriptVarAndViewInfo.ParameterScopeToDisplayInfo[ScopeValue].bEnabled;
+}
+
+FText FNiagaraParameterPanelEntryParameterNameViewModel::GetParameterNameText() const
+{
+	return FText::FromName(CachedScriptVarAndViewInfo.MetaData.CachedNamespacelessVariableName);
+}
+
+void FNiagaraParameterPanelEntryParameterNameViewModel::OnParameterRenamed(const FText& NewNameText, ETextCommit::Type TextCommitType) const
+{
+	OnParameterRenamedDelegate.ExecuteIfBound(CachedScriptVarAndViewInfo.ScriptVariable, CachedScriptVarAndViewInfo.MetaData, NewNameText);
+}
+
+///////////////////////////////////////////////////////////////////////////////
+/// EdGraphPin Parameter Name ViewModel										///
+///////////////////////////////////////////////////////////////////////////////
+
+FNiagaraGraphPinParameterNameViewModel::FNiagaraGraphPinParameterNameViewModel(
+	const UEdGraphPin* InOwningPin
+	, const FNiagaraScriptVariableAndViewInfo& InScriptVarAndViewInfo
+	, const FNiagaraScriptToolkitParameterPanelViewModel* InParameterPanelViewModel
+)
+	: OwningPin(InOwningPin)
+	, CachedScriptVarAndViewInfo(InScriptVarAndViewInfo)
+	, ParameterPanelViewModel(InParameterPanelViewModel)
+{};
+
+TSharedRef<SWidget> FNiagaraGraphPinParameterNameViewModel::CreateScopeSlotWidget() const
+{
+	UEnum* ParameterScopeEnum = FNiagaraTypeDefinition::GetParameterScopeEnum();
+	TSharedPtr< SWidget > ScopeComboBoxWidget;
+	FNiagaraEditorModule& NiagaraEditorModule = FModuleManager::LoadModuleChecked<FNiagaraEditorModule>("NiagaraEditor");
+
+	bool bEnableScopeSlotWidget = OwningPin->Direction == EEdGraphPinDirection::EGPD_Output;
+
+	SAssignNew(ScopeComboBoxWidget, SBox)
+		.MinDesiredWidth(80.0f) //@todo(ng) tune and make const static
+		.Content()
+		[
+			SNew(SEnumComboBox, ParameterScopeEnum)
+			.IsEnabled(bEnableScopeSlotWidget)
+			//.ForegroundColor(NiagaraEditorModule.GetWidgetProvider()->GetColorForParameterScope(CachedScriptVarAndViewInfo.MetaData.Scope)) //@todo(ng) get background color instead
+			.CurrentValue(this, &FNiagaraGraphPinParameterNameViewModel::GetScopeValue)
+			.OnEnumSelectionChanged(this, &FNiagaraGraphPinParameterNameViewModel::OnScopeValueChanged)
+			.OnGetToolTipForValue(this, &FNiagaraGraphPinParameterNameViewModel::GetToolTipForScopeValue)
+			//.OnIsValueEnabled(this, &FNiagaraGraphPinParameterNameViewModel::GetScopeValueIsEnabled) //@todo(ng) impl
+		];
+
+	return ScopeComboBoxWidget.ToSharedRef();
+}
+
+TSharedRef<SWidget> FNiagaraGraphPinParameterNameViewModel::CreateTextSlotWidget() const
+{
+	TSharedPtr< SWidget > DisplayWidget;
+	const FName FontType = FName("Italic");
+	FSlateFontInfo NameFont = FCoreStyle::GetDefaultFontStyle(FontType, 10); //@todo(ng) get correct style
+
+	SAssignNew(DisplayWidget, SInlineEditableTextBlock)
+		.Text(this, &FNiagaraGraphPinParameterNameViewModel::GetParameterNameText)
+		.Font(NameFont)
+		.OnTextCommitted(this, &FNiagaraGraphPinParameterNameViewModel::OnParameterRenamed)
+		//.OnVerifyTextChanged(InArgs._OnVerifyParameterNameChanged) //@todo(ng) impl
+		.IsReadOnly(bIsReadOnly);
+
+	return DisplayWidget.ToSharedRef();
+}
+
+int32 FNiagaraGraphPinParameterNameViewModel::GetScopeValue() const
+{
+	return (int32)CachedScriptVarAndViewInfo.MetaData.Scope;
+}
+
+void FNiagaraGraphPinParameterNameViewModel::OnScopeValueChanged(int32 NewScopeValue, ESelectInfo::Type SelectionType) const
+{
+	ParameterPanelViewModel->ChangePinScope(OwningPin, ENiagaraParameterScope(NewScopeValue));
+
+	//UEdGraphPin* NonConstOwningPin = const_cast<UEdGraphPin*>(OwningPin); //@todo refactor
+	//Cast<UNiagaraNode>(OwningPin->GetOwningNode())->CommitEditablePinName(FText::FromName(NewPinName), NonConstOwningPin);
+}
+
+FText FNiagaraGraphPinParameterNameViewModel::GetToolTipForScopeValue(int32 ScopeValue) const
+{
+	return CachedScriptVarAndViewInfo.ParameterScopeToDisplayInfo[ScopeValue].Tooltip;
+}
+
+bool FNiagaraGraphPinParameterNameViewModel::GetScopeValueIsEnabled(int32 ScopeValue) const
+{
+	return CachedScriptVarAndViewInfo.ParameterScopeToDisplayInfo[ScopeValue].bEnabled;
+}
+
+FText FNiagaraGraphPinParameterNameViewModel::GetParameterNameText() const
+{
+	return FText::FromName(CachedScriptVarAndViewInfo.MetaData.CachedNamespacelessVariableName);
+}
+
+void FNiagaraGraphPinParameterNameViewModel::OnParameterRenamed(const FText& NewNameText, ETextCommit::Type SelectionType) const
+{
+	ParameterPanelViewModel->RenamePin(OwningPin, NewNameText);
+
+	//UEdGraphPin* NonConstOwningPin = const_cast<UEdGraphPin*>(OwningPin); //@todo refactor
+	//Cast<UNiagaraNode>(OwningPin->GetOwningNode())->CommitEditablePinName(NewNameText, NonConstOwningPin);
+}
diff --git a/Engine/Plugins/FX/Niagara/Source/NiagaraEditor/Private/ViewModels/NiagaraParameterNameViewModel.h b/Engine/Plugins/FX/Niagara/Source/NiagaraEditor/Private/ViewModels/NiagaraParameterNameViewModel.h
new file mode 100644
--- /dev/null
+++ b/Engine/Plugins/FX/Niagara/Source/NiagaraEditor/Private/ViewModels/NiagaraParameterNameViewModel.h
@@ -0,0 +1,94 @@
+// Copyright Epic Games, Inc. All Rights Reserved.
+
+#pragma once
+
+#include "NiagaraEditorCommon.h"
+#include "Templates/SharedPointer.h"
+
+struct FCreateWidgetForActionData;
+class FNiagaraScriptToolkitParameterPanelViewModel;
+
+/** Interface for view models to the parameter name view widget. */
+class INiagaraParameterNameViewModel : public TSharedFromThis<INiagaraParameterNameViewModel>
+{
+public:
+	virtual ~INiagaraParameterNameViewModel() { }
+
+	virtual TSharedRef<SWidget> CreateScopeSlotWidget() const = 0;
+	virtual TSharedRef<SWidget> CreateTextSlotWidget() const = 0;
+
+	virtual int32 GetScopeValue() const = 0;
+	virtual void OnScopeValueChanged(int32 NewScopeValue, ESelectInfo::Type SelectionType) const = 0;
+	virtual FText GetToolTipForScopeValue(int32 ScopeValue) const = 0;
+	virtual bool GetScopeValueIsEnabled(int32 ScopeValue) const = 0;
+
+	virtual FText GetParameterNameText() const = 0;
+	//bool VerifyParameterNameChanged(const FText& NewNameText, FText&) //@todo(ng) impl
+	virtual void OnParameterRenamed(const FText& NewNameText, ETextCommit::Type SelectionType) const = 0;
+};
+
+class FNiagaraParameterPanelEntryParameterNameViewModel : public INiagaraParameterNameViewModel
+{
+public:
+	DECLARE_DELEGATE_ThreeParams(FOnScopeSelectionChanged, const FNiagaraVariable&, const FNiagaraVariableMetaData&, const ENiagaraParameterScope)
+	DECLARE_DELEGATE_ThreeParams(FOnParameterRenamed, const FNiagaraVariable&, const FNiagaraVariableMetaData&, const FText&)
+
+	FNiagaraParameterPanelEntryParameterNameViewModel(FCreateWidgetForActionData* const InCreateData, const FNiagaraScriptVariableAndViewInfo& InScriptVarAndViewInfo);
+
+	/** Begin INiagaraParameterNameViewModel Interface */
+	virtual TSharedRef<SWidget> CreateScopeSlotWidget() const override;
+	virtual TSharedRef<SWidget> CreateTextSlotWidget() const override;
+
+	virtual int32 GetScopeValue() const override;
+	virtual void OnScopeValueChanged(int32 NewScopeValue, ESelectInfo::Type SelectionType) const override;
+	virtual FText GetToolTipForScopeValue(int32 ScopeValue) const override;
+	virtual bool GetScopeValueIsEnabled(int32 ScopeValue) const override;
+
+	virtual FText GetParameterNameText() const override;
+	//bool VerifyParameterNameChanged(const FText& NewNameText, FText&) //@todo(ng) impl
+	virtual void OnParameterRenamed(const FText& NewNameText, ETextCommit::Type SelectionType) const override;
+	/** End INiagaraParameterNameViewModel Interface */
+
+	FOnScopeSelectionChanged& GetOnScopeSelectionChangedDelegate() { return OnScopeSelectionChangedDelegate; };
+	FOnParameterRenamed& GetOnParameterRenamedDelegate() { return OnParameterRenamedDelegate; };
+
+private:
+	FOnScopeSelectionChanged OnScopeSelectionChangedDelegate;
+	FOnParameterRenamed OnParameterRenamedDelegate;
+
+	FCreateWidgetForActionData* const CreateData;
+	const FNiagaraScriptVariableAndViewInfo CachedScriptVarAndViewInfo;
+
+	TAttribute<bool> bIsReadOnly; //@todo(ng) debug
+};
+
+class FNiagaraGraphPinParameterNameViewModel : public INiagaraParameterNameViewModel
+{
+public:
+	FNiagaraGraphPinParameterNameViewModel(
+		  const UEdGraphPin* InOwningPin
+		, const FNiagaraScriptVariableAndViewInfo& InScriptVarAndViewInfo
+		, const FNiagaraScriptToolkitParameterPanelViewModel* InParameterPanelViewModel
+	);
+
+	/** Begin INiagaraParameterNameViewModel Interface */
+	virtual TSharedRef<SWidget> CreateScopeSlotWidget() const override;
+	virtual TSharedRef<SWidget> CreateTextSlotWidget() const override;
+
+	virtual int32 GetScopeValue() const override;
+	virtual void OnScopeValueChanged(int32 NewScopeValue, ESelectInfo::Type SelectionType) const override;
+	virtual FText GetToolTipForScopeValue(int32 ScopeValue) const override;
+	virtual bool GetScopeValueIsEnabled(int32 ScopeValue) const override;
+
+	virtual FText GetParameterNameText() const override;
+	//bool VerifyParameterNameChanged(const FText& NewNameText, FText&) //@todo(ng) impl
+	virtual void OnParameterRenamed(const FText& NewNameText, ETextCommit::Type SelectionType) const override;
+	/** End INiagaraParameterNameViewModel Interface */
+
+private:
+	const UEdGraphPin* OwningPin;
+	const FNiagaraScriptVariableAndViewInfo CachedScriptVarAndViewInfo;
+	const FNiagaraScriptToolkitParameterPanelViewModel* ParameterPanelViewModel;
+
+	TAttribute<bool> bIsReadOnly; //@todo(ng) debug
+};
diff --git a/Engine/Plugins/FX/Niagara/Source/NiagaraEditor/Private/ViewModels/NiagaraParameterPanelViewModel.cpp b/Engine/Plugins/FX/Niagara/Source/NiagaraEditor/Private/ViewModels/NiagaraParameterPanelViewModel.cpp
new file mode 100644
--- /dev/null
+++ b/Engine/Plugins/FX/Niagara/Source/NiagaraEditor/Private/ViewModels/NiagaraParameterPanelViewModel.cpp
@@ -0,0 +1,799 @@
+// Copyright Epic Games, Inc. All Rights Reserved.
+
+#include "NiagaraParameterPanelViewModel.h"
+#include "NiagaraScriptVariable.h"
+#include "NiagaraGraph.h"
+#include "NiagaraActions.h"
+#include "SGraphActionMenu.h"
+#include "SNiagaraParameterNameView.h"
+#include "NiagaraParameterNameViewModel.h"
+#include "NiagaraStandaloneScriptViewModel.h"
+#include "ViewModels/NiagaraSystemViewModel.h"
+#include "NiagaraConstants.h"
+#include "EdGraphSchema_Niagara.h"
+#include "ViewModels/Stack/NiagaraStackGraphUtilities.h"
+#include "ViewModels/NiagaraEmitterHandleViewModel.h"
+#include "NiagaraScriptGraphViewModel.h"
+#include "NiagaraNodeParameterMapGet.h"
+#include "NiagaraNodeParameterMapSet.h"
+#include "NiagaraSystemScriptViewModel.h"
+#include "ViewModels/NiagaraSystemSelectionViewModel.h"
+#include "NiagaraEmitterHandle.h"
+#include "NiagaraEmitter.h"
+#include "NiagaraSystem.h"
+#include "NiagaraObjectSelection.h"
+#include "EdGraph/EdGraphSchema.h"
+
+#define LOCTEXT_NAMESPACE "NiagaraParameterPanelViewModel"
+
+FText NiagaraParameterPanelSectionID::OnGetSectionTitle(const NiagaraParameterPanelSectionID::Type InSection)
+{
+	/* Setup an appropriate name for the section for this node */
+	FText SectionTitle;
+	switch (InSection)
+	{
+	case NiagaraParameterPanelSectionID::ENGINE:
+		SectionTitle = LOCTEXT("EngineParamSection", "Engine");
+		break;
+	case NiagaraParameterPanelSectionID::USER:
+		SectionTitle = LOCTEXT("UserParamSection", "User");
+		break;
+	case NiagaraParameterPanelSectionID::SYSTEM:
+		SectionTitle = LOCTEXT("SystemParamSection", "System");
+		break;
+	case NiagaraParameterPanelSectionID::EMITTER:
+		SectionTitle = LOCTEXT("EmitterParamSection", "Emitter");
+		break;
+	case NiagaraParameterPanelSectionID::PARTICLES:
+		SectionTitle = LOCTEXT("ParticlesParamSection", "Particles");
+		break;
+	case NiagaraParameterPanelSectionID::INPUTS:
+		SectionTitle = LOCTEXT("InputsParamSection", "Inputs");
+		break;
+	case NiagaraParameterPanelSectionID::REFERENCES:
+		SectionTitle = LOCTEXT("ReferencesParamSection", "References");
+		break;
+	case NiagaraParameterPanelSectionID::OUTPUTS:
+		SectionTitle = LOCTEXT("OutputsParamSection", "Outputs");
+		break;
+	case NiagaraParameterPanelSectionID::LOCALS:
+		SectionTitle = LOCTEXT("LocalsParamSection", "Locals");
+		break;
+	case NiagaraParameterPanelSectionID::INITIALVALUES:
+		SectionTitle = LOCTEXT("InitialValuesParamSection", "Initial Values");
+		break;
+// 	case NiagaraParameterPanelSectionID::PARAMETERCOLLECTION: //@TODO implement parameter collection handling
+// 		SectionTitle = NSLOCTEXT("GraphActionNode", "ParameterCollection", "Parameter Collection");
+// 		break;
+	case NiagaraParameterPanelSectionID::NONE:
+	default:
+		SectionTitle = FText::GetEmpty();
+		break;
+	}
+
+	return SectionTitle;
+}
+
+const NiagaraParameterPanelSectionID::Type NiagaraParameterPanelSectionID::GetSectionForScope(ENiagaraParameterScope InScope)
+{
+	switch (InScope) {
+	case ENiagaraParameterScope::Engine:
+		return NiagaraParameterPanelSectionID::ENGINE;
+	case ENiagaraParameterScope::User:
+		return NiagaraParameterPanelSectionID::USER;
+	case ENiagaraParameterScope::System:
+		return NiagaraParameterPanelSectionID::SYSTEM;
+	case ENiagaraParameterScope::Emitter:
+		return NiagaraParameterPanelSectionID::EMITTER;
+	case ENiagaraParameterScope::Particles:
+		return NiagaraParameterPanelSectionID::PARTICLES;
+	case ENiagaraParameterScope::ScriptPersistent:
+	case ENiagaraParameterScope::ScriptTransient:
+		// This is a potential situation if a script alias param has not had its scope cached by compiling.
+		return NiagaraParameterPanelSectionID::NONE;
+	default:
+		ensureMsgf(false, TEXT("Failed to find matching section ID for script parameter scope!"));
+	};
+	return NiagaraParameterPanelSectionID::NONE;
+}
+
+const NiagaraParameterPanelSectionID::Type NiagaraParameterPanelSectionID::GetSectionForParameterMetaData(const FNiagaraVariableMetaData& MetaData)
+{
+	if (MetaData.Usage == ENiagaraScriptParameterUsage::Local)
+	{
+		return NiagaraParameterPanelSectionID::LOCALS;
+	}
+	else if (MetaData.Usage == ENiagaraScriptParameterUsage::InitialValueInput)
+	{
+		return NiagaraParameterPanelSectionID::INITIALVALUES;
+	}
+	else if (MetaData.Usage == ENiagaraScriptParameterUsage::Input)
+	{
+		if (MetaData.Scope == ENiagaraParameterScope::Input)
+		{
+			return NiagaraParameterPanelSectionID::INPUTS;
+		}
+		else
+		{
+			return NiagaraParameterPanelSectionID::REFERENCES;
+		}
+	}
+	else if (MetaData.bIsStaticSwitch)
+	{
+		return NiagaraParameterPanelSectionID::INPUTS;
+	}
+	else if (MetaData.Usage == ENiagaraScriptParameterUsage::Output)
+	{
+		return NiagaraParameterPanelSectionID::OUTPUTS;
+	}
+	else if (MetaData.Usage == ENiagaraScriptParameterUsage::InputOutput)
+	{
+		ensureMsgf(false, TEXT("Encountered an InputOutput parameter usage when getting section ID for parameter panel!"));
+		return NiagaraParameterPanelSectionID::REFERENCES;
+	}
+	else
+	{
+		ensureMsgf(false, TEXT("Failed to find matching section ID for script parameter usage!"));
+	}
+	return NiagaraParameterPanelSectionID::LOCALS;
+}
+
+ENiagaraParameterScope NiagaraParameterPanelSectionID::GetScopeForNewParametersInSection(NiagaraParameterPanelSectionID::Type InSection)
+{
+	switch (InSection) {
+	case NiagaraParameterPanelSectionID::ENGINE:
+		return ENiagaraParameterScope::Engine;
+	case NiagaraParameterPanelSectionID::USER:
+		return ENiagaraParameterScope::User;
+	case NiagaraParameterPanelSectionID::SYSTEM:
+		return ENiagaraParameterScope::System;
+	case NiagaraParameterPanelSectionID::EMITTER:
+		return ENiagaraParameterScope::Emitter;
+	case NiagaraParameterPanelSectionID::PARTICLES:
+		return ENiagaraParameterScope::Particles;
+	case NiagaraParameterPanelSectionID::LOCALS:
+		return ENiagaraParameterScope::Local;
+	case NiagaraParameterPanelSectionID::INITIALVALUES:
+		return ENiagaraParameterScope::Particles;
+
+	// Default to Particles scope if coming from section IDs that are not directly associated with scope.
+	case NiagaraParameterPanelSectionID::INPUTS:
+	case NiagaraParameterPanelSectionID::REFERENCES:
+	case NiagaraParameterPanelSectionID::OUTPUTS:
+		return ENiagaraParameterScope::Particles;
+
+	case NiagaraParameterPanelSectionID::NONE:
+		ensureMsgf(false, TEXT("Encountered invalid parameter panel section ID NONE when getting scope from section!"));
+		break;
+	};
+	ensureMsgf(false, TEXT("Did not encounter a known section ID when getting scope from section!"));
+	return ENiagaraParameterScope::Particles;
+}
+
+ENiagaraScriptParameterUsage NiagaraParameterPanelSectionID::GetUsageForNewParametersInSection(NiagaraParameterPanelSectionID::Type InSection)
+{
+	switch (InSection) {
+	case NiagaraParameterPanelSectionID::INPUTS:
+	case NiagaraParameterPanelSectionID::REFERENCES:
+		return ENiagaraScriptParameterUsage::Input;
+	case NiagaraParameterPanelSectionID::OUTPUTS:
+		return ENiagaraScriptParameterUsage::Output;
+	case NiagaraParameterPanelSectionID::LOCALS:
+		return ENiagaraScriptParameterUsage::Local;
+	case NiagaraParameterPanelSectionID::INITIALVALUES:
+		return ENiagaraScriptParameterUsage::InitialValueInput;
+
+	// By convention, default new parameters created in system editor to output usage.
+	case NiagaraParameterPanelSectionID::ENGINE:
+	case NiagaraParameterPanelSectionID::USER:
+	case NiagaraParameterPanelSectionID::SYSTEM:
+	case NiagaraParameterPanelSectionID::EMITTER:
+	case NiagaraParameterPanelSectionID::PARTICLES:
+		return ENiagaraScriptParameterUsage::Output;
+
+	case NiagaraParameterPanelSectionID::NONE:
+		ensureMsgf(false, TEXT("Encountered invalid parameter panel section ID NONE when getting usage from section!"));
+		break;
+	};
+	ensureMsgf(false, TEXT("Did not encounter a known section ID when getting usage from section!"));
+	return ENiagaraScriptParameterUsage::Output;
+}
+
+TSharedRef<SWidget> INiagaraParameterPanelViewModel::GetScriptParameterVisualWidget(FCreateWidgetForActionData* const InCreateData) const
+{
+	TSharedPtr<FNiagaraScriptVarAndViewInfoAction> ScriptVarAndViewInfoAction = StaticCastSharedPtr<FNiagaraScriptVarAndViewInfoAction>(InCreateData->Action);
+	const FNiagaraScriptVariableAndViewInfo& ScriptVarAndViewInfo = ScriptVarAndViewInfoAction.Get()->ScriptVariableAndViewInfo;
+	TSharedPtr<FNiagaraParameterPanelEntryParameterNameViewModel> ParameterNameViewModel = MakeShared<FNiagaraParameterPanelEntryParameterNameViewModel>(InCreateData, ScriptVarAndViewInfo);
+	ParameterNameViewModel->GetOnParameterRenamedDelegate().BindSP(this, &INiagaraParameterPanelViewModel::RenameParameter);
+	ParameterNameViewModel->GetOnScopeSelectionChangedDelegate().BindSP(this, &INiagaraParameterPanelViewModel::ChangeParameterScope);
+
+	TSharedPtr<SWidget> ScriptParameterVisualWidget = SNew(SNiagaraParameterNameView, ParameterNameViewModel);
+	return ScriptParameterVisualWidget->AsShared();
+}
+
+TArray<TWeakObjectPtr<UNiagaraGraph>> FNiagaraScriptToolkitParameterPanelViewModel::GetEditableGraphs() const
+{
+	TArray<TWeakObjectPtr<UNiagaraGraph>> EditableGraphs;
+	EditableGraphs.Add(TWeakObjectPtr<UNiagaraGraph>(ScriptViewModel->GetGraphViewModel()->GetGraph()));
+	return EditableGraphs;
+}
+
+///////////////////////////////////////////////////////////////////////////////
+/// System Toolkit Parameter Panel View Model								///
+///////////////////////////////////////////////////////////////////////////////
+
+FNiagaraSystemToolkitParameterPanelViewModel::FNiagaraSystemToolkitParameterPanelViewModel(TSharedPtr<FNiagaraSystemViewModel> InSystemViewModel)
+{
+	SystemViewModel = InSystemViewModel;
+	OverviewSelectionViewModel = SystemViewModel->GetSelectionViewModel();
+	SystemScriptGraph = SystemViewModel->GetSystemScriptViewModel()->GetGraphViewModel()->GetGraph();
+	GEditor->RegisterForUndo(this);
+}
+
+FNiagaraSystemToolkitParameterPanelViewModel::~FNiagaraSystemToolkitParameterPanelViewModel()
+{
+	GEditor->UnregisterForUndo(this);
+}
+
+void FNiagaraSystemToolkitParameterPanelViewModel::InitBindings()
+{
+	if (SystemViewModel->GetEditMode() == ENiagaraSystemViewModelEditMode::SystemAsset)
+	{
+		SystemViewModel->GetSelectionViewModel()->OnEmitterHandleIdSelectionChanged().AddSP(this, &FNiagaraSystemToolkitParameterPanelViewModel::RefreshSelectedEmitterScriptGraphs);
+		SystemViewModel->OnEmitterHandleViewModelsChanged().AddSP(this, &FNiagaraSystemToolkitParameterPanelViewModel::RefreshSelectedEmitterScriptGraphs);
+	}
+}
+
+void FNiagaraSystemToolkitParameterPanelViewModel::Refresh() const
+{
+	OnParameterPanelViewModelRefreshed.ExecuteIfBound();
+}
+
+void FNiagaraSystemToolkitParameterPanelViewModel::CollectStaticSections(TArray<int32>& StaticSectionIDs) const
+{
+	// Generic Emitter/System view, categorize by scope
+	StaticSectionIDs.Add(NiagaraParameterPanelSectionID::USER);
+	StaticSectionIDs.Add(NiagaraParameterPanelSectionID::ENGINE);
+	StaticSectionIDs.Add(NiagaraParameterPanelSectionID::SYSTEM);
+	StaticSectionIDs.Add(NiagaraParameterPanelSectionID::EMITTER);
+	StaticSectionIDs.Add(NiagaraParameterPanelSectionID::PARTICLES);
+}
+
+NiagaraParameterPanelSectionID::Type FNiagaraSystemToolkitParameterPanelViewModel::GetSectionForVarAndViewInfo(const FNiagaraScriptVariableAndViewInfo& VarAndViewInfo) const
+{
+	// Generic Emitter/System view, categorize by scope
+	if (VarAndViewInfo.MetaData.Usage == ENiagaraScriptParameterUsage::Input && VarAndViewInfo.MetaData.Scope != ENiagaraParameterScope::Input)
+	{
+		return NiagaraParameterPanelSectionID::GetSectionForScope(VarAndViewInfo.MetaData.Scope);
+	}
+	else
+	{
+		//return NiagaraParameterPanelSectionID::GetSectionForScope(VarAndViewInfo.ScriptVariableMetaData.LastPrecompileScope); //@todo(ng) cache the known scope in the niagarasystem
+		ENiagaraParameterScope ReturnScope = ENiagaraParameterScope::None;
+		ENiagaraScriptParameterUsage EmptyUsage;
+		FName EmptyName;
+		FNiagaraStackGraphUtilities::SetParameterMetaDataFromName(VarAndViewInfo.ScriptVariable.GetName(), ReturnScope, EmptyUsage, EmptyName);
+		return NiagaraParameterPanelSectionID::GetSectionForScope(ReturnScope);
+	}
+}
+
+void FNiagaraSystemToolkitParameterPanelViewModel::AddParameter(const FNiagaraVariable& InVariableToAdd, const FNiagaraVariableMetaData& InVariableMetaDataToAssign)
+{
+	FScopedTransaction AddParameter(LOCTEXT("AddParameter", "Add Parameter"));
+	bool bSystemIsSelected = OverviewSelectionViewModel->GetSystemIsSelected();
+
+	UNiagaraGraph::FAddParameterOptions AddParameterOptions = UNiagaraGraph::FAddParameterOptions();
+	AddParameterOptions.NewParameterScope = InVariableMetaDataToAssign.Scope;
+	AddParameterOptions.NewParameterUsage = InVariableMetaDataToAssign.Usage;
+	AddParameterOptions.bAddedFromSystemEditor = true;
+	FNiagaraVariable DuplicateVar = InVariableToAdd; //@todo(ng) rewrite
+
+	if (InVariableMetaDataToAssign.Scope == ENiagaraParameterScope::User)
+	{
+		SystemViewModel->GetSystem().Modify();
+		SystemViewModel->GetSystem().GetExposedParameters().AddParameter(InVariableToAdd); //@Todo(ng) copy ParameterMapView impl
+	}
+	else if (InVariableMetaDataToAssign.Scope == ENiagaraParameterScope::System)
+	{
+		UNiagaraGraph* Graph = SystemViewModel->GetSystemScriptViewModel()->GetGraphViewModel()->GetGraph();
+		Graph->Modify();
+		Graph->AddParameter(DuplicateVar, AddParameterOptions); //@todo(ng) verify
+	}
+	else
+	{
+		for (TWeakObjectPtr<UNiagaraGraph> Graph : GetEditableEmitterScriptGraphs())
+		{
+			if (ensureMsgf(Graph.IsValid(), TEXT("Editable Emitter Script Graph was stale when adding parameter!")))
+			{
+				Graph->Modify();
+				Graph->AddParameter(DuplicateVar, AddParameterOptions);
+			}
+		}
+	}
+
+	if (bSystemIsSelected)
+	{
+		//SystemViewModel-> //@todo(ng) handle user params
+	}
+	
+	Refresh();
+}
+
+void FNiagaraSystemToolkitParameterPanelViewModel::RemoveParameter(const FNiagaraVariable& TargetVariableToRemove, const FNiagaraVariableMetaData& TargetVariableMetaData)
+{
+	FScopedTransaction RemoveParameter(LOCTEXT("RemoveParameter", "Remove Parameter"));
+	if (TargetVariableMetaData.IsInputUsage() && TargetVariableMetaData.Scope == ENiagaraParameterScope::User) //@todo(ng) verify
+	{
+		SystemViewModel->GetSystem().Modify();
+		SystemViewModel->GetSystem().GetExposedParameters().RemoveParameter(TargetVariableToRemove);
+	}
+
+	
+	for (const TWeakObjectPtr<UNiagaraGraph>& GraphWeakPtr : SelectedEmitterScriptGraphs)
+	{
+		if (GraphWeakPtr.IsValid())
+		{
+			UNiagaraGraph* Graph = GraphWeakPtr.Get();
+			Graph->Modify();
+			Graph->RemoveParameter(TargetVariableToRemove);
+		}
+	}
+	Refresh();
+}
+
+bool FNiagaraSystemToolkitParameterPanelViewModel::CanRemoveParameter(const FNiagaraVariable& TargetVariableToRemove, const FNiagaraVariableMetaData& TargetVariableMetaData) const
+{
+	return TargetVariableMetaData.bCreatedInSystemEditor == true;
+}
+
+void FNiagaraSystemToolkitParameterPanelViewModel::RenameParameter(const FNiagaraVariable& TargetVariableToRename, const FNiagaraVariableMetaData& TargetVariableMetaData, const FText& NewVariableNameText) const
+{
+	FScopedTransaction RenameParameter(LOCTEXT("RenameParameter", "Rename Parameter"));
+	const FString ExistingVariableScopeString = FNiagaraStackGraphUtilities::GetNamespaceStringForScriptParameterScope(TargetVariableMetaData.Scope);
+	const FName NewVariableName = FName(*(ExistingVariableScopeString + NewVariableNameText.ToString()));
+	for (const TWeakObjectPtr<UNiagaraGraph>& Graph : GetEditableEmitterScriptGraphs())
+	{
+		if (ensureMsgf(Graph.IsValid(), TEXT("Editable Emitter Script Graph was stale when renaming parameter!")))
+		{
+			Graph->Modify();
+			Graph->RenameParameter(TargetVariableToRename, NewVariableName, TargetVariableMetaData.bIsStaticSwitch, TargetVariableMetaData.Scope); //@todo(ng) handle renaming system params
+		}
+	}
+	Refresh();
+}
+
+void FNiagaraSystemToolkitParameterPanelViewModel::ChangeParameterScope(const FNiagaraVariable& TargetVariableToModify, const FNiagaraVariableMetaData& TargetVariableMetaData, const ENiagaraParameterScope NewVariableScope) const
+{
+	// Parameter scope is not editable for System toolkit.
+	return;
+}
+
+bool FNiagaraSystemToolkitParameterPanelViewModel::CanModifyParameter(const FNiagaraVariable& TargetVariableToModify, const FNiagaraVariableMetaData& TargetVariableMetaData) const
+{
+	return TargetVariableMetaData.bCreatedInSystemEditor == true;
+}
+
+bool FNiagaraSystemToolkitParameterPanelViewModel::CanRenameParameter(const FNiagaraVariable& TargetVariableToRename, const FNiagaraVariableMetaData& TargetVariableMetaData, const FText& NewVariableNameText) const
+{
+	return TargetVariableMetaData.bCreatedInSystemEditor == true;
+}
+
+FReply FNiagaraSystemToolkitParameterPanelViewModel::HandleActionDragged(const TSharedPtr<FEdGraphSchemaAction>& InAction, const FPointerEvent& MouseEvent) const
+{
+	const FText TooltipFormat = LOCTEXT("Parameters", "Name: {0} \nType: {1}"); //@todo(ng) move to static define
+	const FNiagaraScriptVarAndViewInfoAction* ScriptVarAction = static_cast<FNiagaraScriptVarAndViewInfoAction*>(InAction.Get());
+	const FNiagaraScriptVariableAndViewInfo& ScriptVarAndViewInfo = ScriptVarAction->ScriptVariableAndViewInfo;
+	NiagaraParameterPanelSectionID::Type Section = NiagaraParameterPanelSectionID::GetSectionForParameterMetaData(ScriptVarAndViewInfo.MetaData);
+	const FNiagaraVariable& Var = ScriptVarAndViewInfo.ScriptVariable;
+	const FText Name = FText::FromName(Var.GetName());
+	const FText Tooltip = FText::Format(TooltipFormat, FText::FromName(Var.GetName()), Var.GetType().GetNameText());
+
+	TSharedPtr<FNiagaraParameterAction> ParameterAction = MakeShared<FNiagaraParameterAction>(Var, FText::GetEmpty(), Name, Tooltip, 0, FText(), Section);
+	TSharedRef<FNiagaraParameterDragOperation> DragOperation = MakeShared<FNiagaraParameterDragOperation>(ParameterAction);
+	DragOperation->CurrentHoverText = InAction->GetMenuDescription();
+	DragOperation->SetupDefaults();
+	DragOperation->Construct();
+	return FReply::Handled().BeginDragDrop(DragOperation);
+}
+
+TArray<TWeakObjectPtr<UNiagaraGraph>> FNiagaraSystemToolkitParameterPanelViewModel::GetEditableGraphs() const
+{
+	TArray<TWeakObjectPtr<UNiagaraGraph>> EditableGraphs;
+	if (SystemViewModel->GetEditMode() == ENiagaraSystemViewModelEditMode::SystemAsset)
+	{
+		EditableGraphs.Add(SystemScriptGraph);
+		EditableGraphs.Append(SelectedEmitterScriptGraphs);
+	}
+	else
+	{
+		EditableGraphs.Add(static_cast<UNiagaraScriptSource*>(SystemViewModel->GetEmitterHandleViewModels()[0]->GetEmitterHandle()->GetInstance()->GraphSource)->NodeGraph);
+	}
+	return EditableGraphs;
+}
+
+const TArray<FNiagaraScriptVariableAndViewInfo> FNiagaraSystemToolkitParameterPanelViewModel::GetViewedParameters()
+{
+	TArray<FNiagaraScriptVariableAndViewInfo> ViewedParameters;
+	for (TWeakObjectPtr<UNiagaraGraph> Graph : GetEditableGraphs())
+	{
+		if (ensureMsgf(Graph.IsValid(), TEXT("Invalid Graph visited when trying to get viewed parameters for system toolkit parameter panel!")))
+		{
+			const TMap<FNiagaraVariable, UNiagaraScriptVariable*>& GraphVarToScriptVarMap = Graph->GetAllMetaData();
+			for (auto Iter = GraphVarToScriptVarMap.CreateConstIterator(); Iter; ++Iter)
+			{
+				const FNiagaraVariableMetaData& MetaData = Iter.Value()->Metadata;
+
+				if (MetaData.Scope == ENiagaraParameterScope::None || MetaData.Usage == ENiagaraScriptParameterUsage::None)
+				{
+					ensureMsgf(false, TEXT("Invalid MetaData found for graph variable: %s"), *Iter.Value()->Variable.GetName().ToString());
+					continue;
+				}
+				else if (MetaData.Scope == ENiagaraParameterScope::Local)
+				{
+					// Note, the MetaData.Usage being local is fine to display for the System toolkit.
+					continue;
+				}
+				else if (MetaData.Scope == ENiagaraParameterScope::Input) // Do not expose inputs as configurable values
+				{
+					continue;
+				}
+				else if (MetaData.Scope == ENiagaraParameterScope::ScriptPersistent || MetaData.Scope == ENiagaraParameterScope::ScriptTransient)
+				{
+					//@todo(ng) Skip script alias parameters until we can resolve them!
+					continue;
+				}
+				ViewedParameters.Add(FNiagaraScriptVariableAndViewInfo(Iter.Value()->Variable, MetaData));
+			}
+		}
+	}
+	CachedViewedParameters = ViewedParameters;
+	return ViewedParameters;
+}
+
+void FNiagaraSystemToolkitParameterPanelViewModel::PostUndo(bool bSuccess)
+{
+	Refresh();
+}
+
+void FNiagaraSystemToolkitParameterPanelViewModel::RefreshSelectedEmitterScriptGraphs()
+{
+	SelectedEmitterScriptGraphs.Reset();
+
+	const TArray<FGuid>& SelectedEmitterHandleIds = OverviewSelectionViewModel->GetSelectedEmitterHandleIds();
+
+	if (SelectedEmitterHandleIds.Num() > 0)
+	{
+		const TArray<TSharedRef<FNiagaraEmitterHandleViewModel>>& EmitterHandleViewModels = SystemViewModel->GetEmitterHandleViewModels();
+		for (const TSharedRef<FNiagaraEmitterHandleViewModel>& EmitterHandleViewModel : EmitterHandleViewModels)
+		{
+			if (SelectedEmitterHandleIds.Contains(EmitterHandleViewModel->GetId()))
+			{
+				SelectedEmitterScriptGraphs.Add(static_cast<UNiagaraScriptSource*>(EmitterHandleViewModel->GetEmitterHandle()->GetInstance()->GraphSource)->NodeGraph);
+			}
+		}
+	}
+
+	Refresh();
+}
+
+TArray<TWeakObjectPtr<UNiagaraGraph>> FNiagaraSystemToolkitParameterPanelViewModel::GetEditableEmitterScriptGraphs() const
+{
+	if (SystemViewModel->GetEditMode() == ENiagaraSystemViewModelEditMode::SystemAsset)
+	{
+		return SelectedEmitterScriptGraphs;
+	}
+	else
+	{
+		TArray<TWeakObjectPtr<UNiagaraGraph>> EditableEmitterScriptGraphs;
+		EditableEmitterScriptGraphs.Add(static_cast<UNiagaraScriptSource*>(SystemViewModel->GetEmitterHandleViewModels()[0]->GetEmitterHandle()->GetInstance()->GraphSource)->NodeGraph);
+		return EditableEmitterScriptGraphs;
+	}
+}
+
+///////////////////////////////////////////////////////////////////////////////
+/// Script Toolkit Parameter Panel View Model								///
+///////////////////////////////////////////////////////////////////////////////
+
+FNiagaraScriptToolkitParameterPanelViewModel::FNiagaraScriptToolkitParameterPanelViewModel(TSharedPtr<FNiagaraStandaloneScriptViewModel> InScriptViewModel)
+{
+	ScriptViewModel = InScriptViewModel;
+	VariableObjectSelection = ScriptViewModel->GetVariableSelection();
+}
+
+FNiagaraScriptToolkitParameterPanelViewModel::~FNiagaraScriptToolkitParameterPanelViewModel()
+{
+	UNiagaraGraph* NiagaraGraph = static_cast<UNiagaraGraph*>(ScriptViewModel->GetGraphViewModel()->GetGraph());
+	NiagaraGraph->RemoveOnGraphChangedHandler(OnGraphChangedHandle);
+	NiagaraGraph->RemoveOnGraphNeedsRecompileHandler(OnGraphNeedsRecompileHandle);
+}
+
+void FNiagaraScriptToolkitParameterPanelViewModel::InitBindings()
+{
+	UNiagaraGraph* NiagaraGraph = static_cast<UNiagaraGraph*>(ScriptViewModel->GetGraphViewModel()->GetGraph());
+	OnGraphChangedHandle = NiagaraGraph->AddOnGraphChangedHandler(
+		FOnGraphChanged::FDelegate::CreateRaw(this, &FNiagaraScriptToolkitParameterPanelViewModel::HandleOnGraphChanged));
+	OnGraphNeedsRecompileHandle = NiagaraGraph->AddOnGraphNeedsRecompileHandler(
+		FOnGraphChanged::FDelegate::CreateRaw(this, &FNiagaraScriptToolkitParameterPanelViewModel::HandleOnGraphChanged));
+	ScriptViewModel->GetGraphViewModel()->GetGraph()->RegisterPinVisualWidgetProvider(UNiagaraGraph::FOnGetPinVisualWidget::CreateSP(this, &FNiagaraScriptToolkitParameterPanelViewModel::GetScriptParameterVisualWidget));
+}
+
+void FNiagaraScriptToolkitParameterPanelViewModel::Refresh() const
+{
+	OnParameterPanelViewModelRefreshed.ExecuteIfBound();
+}
+
+void FNiagaraScriptToolkitParameterPanelViewModel::CollectStaticSections(TArray<int32>& StaticSectionIDs) const
+{
+	StaticSectionIDs.Add(NiagaraParameterPanelSectionID::INPUTS);
+	StaticSectionIDs.Add(NiagaraParameterPanelSectionID::REFERENCES);
+	StaticSectionIDs.Add(NiagaraParameterPanelSectionID::OUTPUTS);
+	StaticSectionIDs.Add(NiagaraParameterPanelSectionID::LOCALS);
+	StaticSectionIDs.Add(NiagaraParameterPanelSectionID::INITIALVALUES);
+}
+
+NiagaraParameterPanelSectionID::Type FNiagaraScriptToolkitParameterPanelViewModel::GetSectionForVarAndViewInfo(const FNiagaraScriptVariableAndViewInfo& VarAndViewInfo) const
+{
+	return NiagaraParameterPanelSectionID::GetSectionForParameterMetaData(VarAndViewInfo.MetaData);
+}
+
+void FNiagaraScriptToolkitParameterPanelViewModel::AddParameter(const FNiagaraVariable& InVariableToAdd, const FNiagaraVariableMetaData& InVariableMetaDataToAssign)
+{
+	FScopedTransaction AddParameter(LOCTEXT("AddParameter", "Add parameter and create new associated node"));
+	UNiagaraGraph::FAddParameterOptions AddParameterOptions = UNiagaraGraph::FAddParameterOptions();
+	AddParameterOptions.NewParameterScope = InVariableMetaDataToAssign.Scope;
+	AddParameterOptions.NewParameterUsage = InVariableMetaDataToAssign.Usage;
+	FNiagaraVariable DuplicateVar = InVariableToAdd; //@todo(ng) rewrite
+
+	UNiagaraGraph* Graph = ScriptViewModel->GetGraphViewModel()->GetGraph();
+	Graph->Modify();
+	const FVector2D NewNodePos = Graph->GetGoodPlaceForNewNode();
+	if (InVariableMetaDataToAssign.Usage == ENiagaraScriptParameterUsage::Input)
+	{
+		UNiagaraNodeParameterMapGet* NewMapGet = FNiagaraSchemaAction_NewNode::SpawnNodeFromTemplate<UNiagaraNodeParameterMapGet>(Graph, NewObject<UNiagaraNodeParameterMapGet>(), NewNodePos);
+		NewMapGet->Modify();
+		NewMapGet->AddParameter(DuplicateVar, AddParameterOptions);
+	}
+	else if (InVariableMetaDataToAssign.Usage == ENiagaraScriptParameterUsage::Output)
+	{
+		UNiagaraNodeParameterMapSet* NewMapSet = FNiagaraSchemaAction_NewNode::SpawnNodeFromTemplate<UNiagaraNodeParameterMapSet>(Graph, NewObject<UNiagaraNodeParameterMapSet>(), NewNodePos);
+		NewMapSet->Modify();
+		NewMapSet->AddParameter(DuplicateVar, AddParameterOptions);
+	}
+	else if (InVariableMetaDataToAssign.Usage == ENiagaraScriptParameterUsage::InputOutput)
+	{
+		ensureMsgf(false, TEXT("Unexpected usage encountered when adding parameter through parameter panel view model!"));
+	}
+	else
+	{
+		Graph->AddParameter(DuplicateVar, AddParameterOptions);
+	}
+}
+
+void FNiagaraScriptToolkitParameterPanelViewModel::RemoveParameter(const FNiagaraVariable& TargetVariableToRemove, const FNiagaraVariableMetaData& TargetVariableMetaData)
+{
+	FScopedTransaction RemoveParametersWithPins(LOCTEXT("RemoveParametersWithPins", "Remove parameter and referenced pins"));
+	UNiagaraGraph* Graph = ScriptViewModel->GetGraphViewModel()->GetGraph();
+	Graph->Modify();
+	Graph->RemoveParameter(TargetVariableToRemove);
+}
+
+bool FNiagaraScriptToolkitParameterPanelViewModel::CanRemoveParameter(const FNiagaraVariable& TargetVariableToRemove, const FNiagaraVariableMetaData& TargetVariableMetaData) const
+{
+	return true;
+}
+
+void FNiagaraScriptToolkitParameterPanelViewModel::RenameParameter(const FNiagaraVariable& TargetVariableToRename, const FNiagaraVariableMetaData& TargetVariableMetaData, const FText& NewVariableNameText) const
+{
+	FScopedTransaction RenameParameterAndReferencedPins(LOCTEXT("RenameParameterAndReferencedPins", "Rename parameter and referenced pins"));
+	const FString ExistingVariableScopeString = FNiagaraStackGraphUtilities::GetNamespaceStringForScriptParameterMetaData(TargetVariableMetaData);
+	const FName NewVariableName = FName(*(ExistingVariableScopeString + NewVariableNameText.ToString()));
+	UNiagaraGraph* Graph = ScriptViewModel->GetGraphViewModel()->GetGraph();
+	Graph->Modify();
+	Graph->RenameParameter(TargetVariableToRename, NewVariableName, TargetVariableMetaData.bIsStaticSwitch, TargetVariableMetaData.Scope);
+}
+
+void FNiagaraScriptToolkitParameterPanelViewModel::ChangeParameterScope(const FNiagaraVariable& TargetVariableToModify, const FNiagaraVariableMetaData& TargetVariableMetaData, const ENiagaraParameterScope NewVariableScope) const
+{
+	if (ensureMsgf(TargetVariableMetaData.Usage != ENiagaraScriptParameterUsage::Output, TEXT("Tried to change scope of output parameter!")))
+	{
+		FScopedTransaction ChangeParameterScopeAndReferencedPins(LOCTEXT("ChangeParameterScopeAndReferencedPins", "Change parameter scope, Rename parameter and referenced pins"));
+
+		bool bIsInitialValue = TargetVariableMetaData.Usage == ENiagaraScriptParameterUsage::InitialValueInput;
+		const FName NewVariableName = FNiagaraStackGraphUtilities::GetVariableNameForScope(TargetVariableToModify.GetName(), NewVariableScope, bIsInitialValue);
+		UNiagaraGraph* Graph = ScriptViewModel->GetGraphViewModel()->GetGraph();
+		Graph->Modify();
+		Graph->RenameParameter(TargetVariableToModify, NewVariableName, TargetVariableMetaData.bIsStaticSwitch, NewVariableScope);
+	}
+}
+
+bool FNiagaraScriptToolkitParameterPanelViewModel::CanModifyParameter(const FNiagaraVariable& TargetVariableToModify, const FNiagaraVariableMetaData& TargetVariableMetaData) const
+{
+	//@todo return false for parameter library entries
+	return true;
+}
+
+bool FNiagaraScriptToolkitParameterPanelViewModel::CanRenameParameter(const FNiagaraVariable& TargetVariableToRename, const FNiagaraVariableMetaData& TargetVariableMetaData, const FText& NewVariableNameText) const
+{
+	// Prevent name values that would alias an existing parameter
+	FName NewName = FName(*NewVariableNameText.ToString());
+	for (const FNiagaraScriptVariableAndViewInfo& ViewedParameter : CachedViewedParameters)
+	{
+		if (NewName == ViewedParameter.MetaData.CachedNamespacelessVariableName)
+		{
+			return false; //@todo(ng) wrap this into the verify logic
+		}
+	}
+	return true;
+}
+
+void FNiagaraScriptToolkitParameterPanelViewModel::HandleActionSelected(const TSharedPtr<FEdGraphSchemaAction>& InAction, ESelectInfo::Type InSelectionType)
+{
+	const FNiagaraScriptVarAndViewInfoAction* Action = static_cast<FNiagaraScriptVarAndViewInfoAction*>(InAction.Get());
+	UNiagaraScriptVariable** ScriptVarPtr = ScriptViewModel->GetGraphViewModel()->GetGraph()->GetAllMetaData().Find(Action->ScriptVariableAndViewInfo.ScriptVariable);
+	if (ensureMsgf(ScriptVarPtr != nullptr, TEXT("Failed to get UNiagaraScriptVariable from selected action!")))
+	{
+		VariableObjectSelection->SetSelectedObject(*ScriptVarPtr);
+	}
+}
+
+FReply FNiagaraScriptToolkitParameterPanelViewModel::HandleActionDragged(const TSharedPtr<FEdGraphSchemaAction>& InAction, const FPointerEvent& MouseEvent) const
+{
+	TSharedRef<FNiagaraParameterGraphDragOperation> DragOperation = FNiagaraParameterGraphDragOperation::New(InAction); //@todo(ng) do not drag drop static switches
+	DragOperation->SetAltDrag(MouseEvent.IsAltDown());
+	DragOperation->SetCtrlDrag(MouseEvent.IsLeftControlDown() || MouseEvent.IsRightControlDown());
+	return FReply::Handled().BeginDragDrop(DragOperation);
+}
+
+TSharedRef<SWidget> FNiagaraScriptToolkitParameterPanelViewModel::GetScriptParameterVisualWidget(const UEdGraphPin* Pin) const
+{
+	const UEdGraphSchema_Niagara* Schema = GetDefault<UEdGraphSchema_Niagara>();
+	const FNiagaraVariable PinVar = Schema->PinToNiagaraVariable(Pin);
+
+	const FNiagaraScriptVariableAndViewInfo* ScriptVarAndViewInfo = CachedViewedParameters.FindByPredicate([PinVar](const FNiagaraScriptVariableAndViewInfo& Entry) {return Entry.ScriptVariable == PinVar; });
+	if (ScriptVarAndViewInfo == nullptr)
+	{
+		// Failed to find the parameter name in the cache, try to find the variable in the graph script variables and generate view info.
+		const UNiagaraScriptVariable* const* ScriptVarPtr = ScriptViewModel->GetGraphViewModel()->GetGraph()->GetAllMetaData().Find(PinVar);
+		if (ScriptVarPtr != nullptr)
+		{
+			const UNiagaraScriptVariable* ScriptVar = *ScriptVarPtr;
+			const TStaticArray<FScopeIsEnabledAndTooltip, (int32)ENiagaraParameterScope::Num> PerScopeInfo = GetParameterScopesEnabledAndTooltips(ScriptVar->Variable, ScriptVar->Metadata);
+			FNiagaraScriptVariableAndViewInfo NewScriptVarAndViewInfo = FNiagaraScriptVariableAndViewInfo(ScriptVar->Variable, ScriptVar->Metadata, PerScopeInfo);
+			ScriptVarAndViewInfo = &NewScriptVarAndViewInfo;
+		}
+		else
+		{
+			// Cannot resolve the parameter from the pin, put an error widget in.
+			TSharedPtr<SWidget> ErrorTextBlock = SNew(STextBlock).Text(FText::FromString("Could not resolve parameter!")); //@todo(ng) make error item method
+			return ErrorTextBlock->AsShared();
+		}
+	}
+	TSharedPtr<FNiagaraGraphPinParameterNameViewModel> ParameterNameViewModel = MakeShared<FNiagaraGraphPinParameterNameViewModel>(Pin, *ScriptVarAndViewInfo, this);
+	
+	TSharedPtr<SWidget> ScriptParameterVisualWidget = SNew(SNiagaraParameterNameView, ParameterNameViewModel);
+	return ScriptParameterVisualWidget->AsShared();
+}
+
+TStaticArray<FScopeIsEnabledAndTooltip, (int32)ENiagaraParameterScope::Num> FNiagaraScriptToolkitParameterPanelViewModel::GetParameterScopesEnabledAndTooltips(const FNiagaraVariable& InVar, const FNiagaraVariableMetaData& InVarMetaData) const
+{
+	TStaticArray<FScopeIsEnabledAndTooltip, (int32)ENiagaraParameterScope::Num> PerScopeInfo;
+	UEnum* ParameterScopeEnum = FNiagaraTypeDefinition::GetParameterScopeEnum();
+
+	// Add defaulted entries for every possible enum value
+	TMap<ENiagaraParameterScope, TPair<bool, const FText>> OutEntriesAndInfo;
+	for (int32 i = 0; i < ParameterScopeEnum->NumEnums() - 1; i++)
+	{
+		if (ParameterScopeEnum->HasMetaData(TEXT("Hidden"), i) == false)
+		{
+			PerScopeInfo[i] = FScopeIsEnabledAndTooltip(true, FText()); //@todo(ng) put the scope in the tooltip
+		}
+	}
+
+	// Prevent setting enum values that would alias an existing parameter
+	const FName InName = InVarMetaData.CachedNamespacelessVariableName;
+	for (const FNiagaraScriptVariableAndViewInfo& ViewedParameter : CachedViewedParameters)
+	{
+		if (InName == ViewedParameter.MetaData.CachedNamespacelessVariableName)
+		{
+			PerScopeInfo[(int32)ViewedParameter.MetaData.Scope].bEnabled = false;
+			PerScopeInfo[(int32)ViewedParameter.MetaData.Scope].Tooltip = LOCTEXT("NiagaraInvalidScopeSelectionNameAlias", "Cannot select scope {0}: Parameter with same name already has this scope."); //@todo(ng) get scope
+		}
+	}
+
+	// Prevent making Module namespace parameters in function and dynamic input scripts
+	if (ScriptViewModel->GetStandaloneScript()->GetUsage() != ENiagaraScriptUsage::Module)
+	{
+		PerScopeInfo[(int32)ENiagaraParameterScope::Input].bEnabled = false;
+		PerScopeInfo[(int32)ENiagaraParameterScope::Input].Tooltip = LOCTEXT("NiagaraInvalidScopeSelectionModule", "Cannot select scope {0}: Scope is only valid in Module Scripts."); //@todo(ng) get scope
+	}
+
+	const TArray<ENiagaraParameterScope> InvalidParameterScopes = ScriptViewModel->GetStandaloneScript()->GetUnsupportedParameterScopes();
+	for (ENiagaraParameterScope InvalidScope : InvalidParameterScopes)
+	{
+		PerScopeInfo[(int32)InvalidScope].bEnabled = false;
+		PerScopeInfo[(int32)InvalidScope].Tooltip = LOCTEXT("NiagaraInvalidScopeSelectionUsageBitmask", "Cannot select scope {0}: Script Usage flags do not support a usage with this scope."); //@todo(ng) rewrite
+	}
+
+	return PerScopeInfo;
+}
+
+const TArray<FNiagaraScriptVariableAndViewInfo> FNiagaraScriptToolkitParameterPanelViewModel::GetViewedParameters()
+{
+	TArray<FNiagaraScriptVariableAndViewInfo> ViewedParameters; //@todo(ng) cached viewed parameters are too behind for GetParameterScopesEnabledAndTooltips, refactor
+	UNiagaraGraph* ViewedGraph = ScriptViewModel->GetGraphViewModel()->GetGraph();
+	if (ensureMsgf(ViewedGraph != nullptr, TEXT("Invalid Graph found when trying to get viewed parameters for script toolkit parameter panel!")))
+	{
+		const TMap<FNiagaraVariable, UNiagaraScriptVariable*>& GraphVarToScriptVarMap = ViewedGraph->GetAllMetaData();
+		for (auto Iter = GraphVarToScriptVarMap.CreateConstIterator(); Iter; ++Iter)
+		{
+			const FNiagaraVariable& Variable = Iter.Value()->Variable;
+
+			if (Variable.GetName() == FNiagaraConstants::InputPinName || Variable.GetName() == FNiagaraConstants::OutputPinName)
+			{
+				//@todo Pins leaked into variable maps at some point, need to clean.
+				continue;
+			}
+
+			const FNiagaraVariableMetaData& MetaData = Iter.Value()->Metadata;
+
+			if (MetaData.bIsStaticSwitch == false && (MetaData.Scope == ENiagaraParameterScope::None || MetaData.Usage == ENiagaraScriptParameterUsage::None) )
+			{
+				// Parameters that are not static switches must have a scope and usage set.
+				ensureMsgf(false, TEXT("Invalid MetaData found for graph variable: %s"), *Iter.Value()->Variable.GetName().ToString());
+				continue;
+			}
+			else if (MetaData.Usage == ENiagaraScriptParameterUsage::InputOutput)
+			{
+				// Need two stack entries to represent this script parameter as both an input and output.
+				FNiagaraVariableMetaData InputVariableMetaData = MetaData;
+				FNiagaraVariableMetaData OutputVariableMetaData = MetaData;
+				InputVariableMetaData.Usage = ENiagaraScriptParameterUsage::Input;
+				OutputVariableMetaData.Usage = ENiagaraScriptParameterUsage::Output;
+
+				const TStaticArray<FScopeIsEnabledAndTooltip, (int32)ENiagaraParameterScope::Num> PerScopeInfo = GetParameterScopesEnabledAndTooltips(Variable, MetaData);
+				ViewedParameters.Add(FNiagaraScriptVariableAndViewInfo(Variable, InputVariableMetaData, PerScopeInfo));
+				ViewedParameters.Add(FNiagaraScriptVariableAndViewInfo(Variable, OutputVariableMetaData));
+			}
+			else if(MetaData.Usage == ENiagaraScriptParameterUsage::Input || MetaData.Usage == ENiagaraScriptParameterUsage::InitialValueInput)
+			{
+				const TStaticArray<FScopeIsEnabledAndTooltip, (int32)ENiagaraParameterScope::Num> PerScopeInfo = GetParameterScopesEnabledAndTooltips(Variable, MetaData);
+				ViewedParameters.Add(FNiagaraScriptVariableAndViewInfo(Variable, MetaData, PerScopeInfo));
+			}
+			else
+			{
+				ViewedParameters.Add(FNiagaraScriptVariableAndViewInfo(Variable, MetaData));
+			}
+			
+		}
+	}
+	CachedViewedParameters = ViewedParameters;
+	return ViewedParameters;
+}
+
+void FNiagaraScriptToolkitParameterPanelViewModel::RenamePin(const UEdGraphPin* TargetPinToRename, const FText& NewNameText) const
+{
+	UNiagaraScriptVariable* ScriptVarToRename = ScriptViewModel->GetGraphViewModel()->GetGraph()->GetScriptVariable(TargetPinToRename->GetFName());
+	if (ensureMsgf(ScriptVarToRename != nullptr, TEXT("Failed to find script variable with same name as pin while renaming pin!")))
+	{
+		RenameParameter(ScriptVarToRename->Variable, ScriptVarToRename->Metadata, NewNameText);
+	}
+}
+
+void FNiagaraScriptToolkitParameterPanelViewModel::ChangePinScope(const UEdGraphPin* TargetPin, const ENiagaraParameterScope NewScope) const
+{
+	if (ensureMsgf(TargetPin->Direction != EEdGraphPinDirection::EGPD_Input, TEXT("Tried to edit scope of input pin, this should not happen!")))
+	{
+		UNiagaraScriptVariable* ScriptVarToRename = ScriptViewModel->GetGraphViewModel()->GetGraph()->GetScriptVariable(TargetPin->GetFName());
+		if (ensureMsgf(ScriptVarToRename != nullptr, TEXT("Failed to find script variable with same name as pin while changing pin scope!")))
+		{
+			ChangeParameterScope(ScriptVarToRename->Variable, ScriptVarToRename->Metadata, NewScope);
+		}
+	}
+}
+
+void FNiagaraScriptToolkitParameterPanelViewModel::HandleOnGraphChanged(const struct FEdGraphEditAction& InAction)
+{
+	Refresh();
+}
+
+#undef LOCTEXT_NAMESPACE // "FNiagaraParameterPanelViewModel"
diff --git a/Engine/Plugins/FX/Niagara/Source/NiagaraEditor/Private/ViewModels/NiagaraParameterPanelViewModel.h b/Engine/Plugins/FX/Niagara/Source/NiagaraEditor/Private/ViewModels/NiagaraParameterPanelViewModel.h
new file mode 100644
--- /dev/null
+++ b/Engine/Plugins/FX/Niagara/Source/NiagaraEditor/Private/ViewModels/NiagaraParameterPanelViewModel.h
@@ -0,0 +1,195 @@
+// Copyright Epic Games, Inc. All Rights Reserved.
+
+#pragma once
+
+#include "NiagaraEditorCommon.h"
+
+class UNiagaraSystem;
+class FNiagaraSystemViewModel;
+class UNiagaraSystemSelectionViewModel;
+class UNiagaraGraph;
+class FNiagaraStandaloneScriptViewModel;
+class FDelegateHandle;
+struct FCreateWidgetForActionData;
+class FNiagaraObjectSelection;
+
+/** Interface for view models to the parameter panel. */
+class INiagaraParameterPanelViewModel : public TSharedFromThis<INiagaraParameterPanelViewModel>
+{
+
+public:
+	/** Delegate to signal the view model's state has changed. */
+	DECLARE_DELEGATE(FOnParameterPanelViewModelRefreshed);
+
+	virtual ~INiagaraParameterPanelViewModel() { }
+
+	/** Separate method to call after ctor to create delegate bindings as we must be fully constructed to do so. */
+	virtual void InitBindings() = 0;
+
+	virtual void Refresh() const = 0;
+
+	virtual void CollectStaticSections(TArray<int32>& StaticSectionIDs) const = 0;
+
+	virtual NiagaraParameterPanelSectionID::Type GetSectionForVarAndViewInfo(const FNiagaraScriptVariableAndViewInfo& VarAndViewInfo) const = 0;
+	
+	virtual void AddParameter(const FNiagaraVariable& InVariableToAdd, const FNiagaraVariableMetaData& InVariableMetaDataToAssign) = 0;
+
+	virtual void RemoveParameter(const FNiagaraVariable& TargetVariableToRemove, const FNiagaraVariableMetaData& TargetVariableMetaData) = 0;
+
+	virtual bool CanRemoveParameter(const FNiagaraVariable& TargetVariableToRemove, const FNiagaraVariableMetaData& TargetVariableMetaData) const = 0;
+
+	virtual void RenameParameter(const FNiagaraVariable& TargetVariableToRename, const FNiagaraVariableMetaData& TargetVariableMetaData, const FText& NewVariableNameText) const = 0;
+
+	virtual void ChangeParameterScope(const FNiagaraVariable& TargetVariableToModify, const FNiagaraVariableMetaData& TargetVariableMetaData, const ENiagaraParameterScope NewVariableScope) const = 0;
+
+	virtual bool CanModifyParameter(const FNiagaraVariable& TargetVariableToModify, const FNiagaraVariableMetaData& TargetVariableMetaData) const = 0;
+
+	virtual bool CanRenameParameter(const FNiagaraVariable& TargetVariableToRename, const FNiagaraVariableMetaData& TargetVariableMetaData, const FText& NewVariableNameText) const = 0;
+
+	virtual void HandleActionSelected(const TSharedPtr<FEdGraphSchemaAction>& InAction, ESelectInfo::Type InSelectionType) {};
+
+	virtual FReply HandleActionDragged(const TSharedPtr<FEdGraphSchemaAction>& InAction, const FPointerEvent& MouseEvent) const = 0;
+
+	//virtual void HandleParameterDropped() {}; //@todo(ng) impl
+
+	virtual bool CanDropParameter(const FNiagaraVariable& InTargetScriptVariableToDrop) const { return false; };
+
+	virtual bool CanDropParameters(const TArray<FNiagaraVariable>& InTargetScriptVariablesToDrop) const { return false; };
+
+	virtual bool CanDragParameter(const FNiagaraVariable& InTargetScriptVariableToDrag) const { return false; };
+
+	virtual bool CanDragParameters(const TArray<FNiagaraVariable>& InTargetScriptVariablesToDrag) const { return false; };
+
+	/** Returns a list of Graphs that are valid for operations to edit their variables and/or metadata. 
+	 *Should collect all Graphs that are currently selected, but also Graphs that are implicitly selected, e.g. the node graph for the script toolkit.
+	 */
+	virtual TArray<TWeakObjectPtr<UNiagaraGraph>> GetEditableGraphs() const = 0;
+
+	virtual const TArray<FNiagaraScriptVariableAndViewInfo> GetViewedParameters() = 0;
+
+	TSharedRef<SWidget> GetScriptParameterVisualWidget(FCreateWidgetForActionData* const InCreateData) const;
+
+	FOnParameterPanelViewModelRefreshed& GetOnRefreshed() { return OnParameterPanelViewModelRefreshed; };
+
+protected:
+	FOnParameterPanelViewModelRefreshed OnParameterPanelViewModelRefreshed;
+
+	/** Cached list of parameters sent to SNiagarParameterPanel, updated whenever GetViewedParameters is called. */
+	TArray<FNiagaraScriptVariableAndViewInfo> CachedViewedParameters;
+
+};
+
+class FNiagaraSystemToolkitParameterPanelViewModel : public INiagaraParameterPanelViewModel, public FEditorUndoClient
+{
+public:
+	/** Construct a SystemToolkit Parameter Panel View Model from a System View Model. */
+	FNiagaraSystemToolkitParameterPanelViewModel(TSharedPtr<FNiagaraSystemViewModel> InSystemViewModel);
+
+	~FNiagaraSystemToolkitParameterPanelViewModel();
+
+	/** Begin INiagaraParameterPanelViewModel interface. */
+	virtual void InitBindings() override;
+
+	virtual void Refresh() const override;
+
+	virtual void CollectStaticSections(TArray<int32>& StaticSectionIDs) const override;
+
+	virtual NiagaraParameterPanelSectionID::Type GetSectionForVarAndViewInfo(const FNiagaraScriptVariableAndViewInfo& VarAndViewInfo) const override;
+
+	virtual void AddParameter(const FNiagaraVariable& VariableToAdd, const FNiagaraVariableMetaData& VariableMetaDataToAssign) override;
+
+	virtual void RemoveParameter(const FNiagaraVariable& TargetVariableToRemove, const FNiagaraVariableMetaData& TargetVariableMetaData) override;
+
+	virtual bool CanRemoveParameter(const FNiagaraVariable& TargetVariableToRemove, const FNiagaraVariableMetaData& TargetVariableMetaData) const override;
+
+	virtual void RenameParameter(const FNiagaraVariable& TargetVariableToRename, const FNiagaraVariableMetaData& TargetVariableMetaData, const FText& NewVariableNameText) const override;
+
+	virtual void ChangeParameterScope(const FNiagaraVariable& TargetVariableToModify, const FNiagaraVariableMetaData& TargetVariableMetaData, const ENiagaraParameterScope NewVariableScope) const override;
+
+	virtual bool CanModifyParameter(const FNiagaraVariable& TargetVariableToModify, const FNiagaraVariableMetaData& TargetVariableMetaData) const override;
+
+	virtual bool CanRenameParameter(const FNiagaraVariable& TargetVariableToRename, const FNiagaraVariableMetaData& TargetVariableMetaData, const FText& NewVariableNameText) const override;
+
+	virtual FReply HandleActionDragged(const TSharedPtr<FEdGraphSchemaAction>& InAction, const FPointerEvent& MouseEvent) const override;
+
+	virtual TArray<TWeakObjectPtr<UNiagaraGraph>> GetEditableGraphs() const override;
+
+	virtual const TArray<FNiagaraScriptVariableAndViewInfo> GetViewedParameters() override;
+	/** End INiagaraParameterPanelViewModel interface. */
+
+	/** Begin FEditorUndoClient Interface */
+	virtual void PostUndo(bool bSuccess) override;
+	virtual void PostRedo(bool bSuccess) override { PostUndo(bSuccess); }
+	/** End FEditorUndoClient Interface */
+
+private:
+	//** Updates the SelectedEmitterScriptGraphs array and then calls Refresh(). Use if the Emitter Script graphs change (e.g. Emitter deleted from System). */
+	void RefreshSelectedEmitterScriptGraphs();
+
+	TArray<TWeakObjectPtr<UNiagaraGraph>> GetEditableEmitterScriptGraphs() const;
+
+	// Graphs viewed to gather UNiagaraScriptVariables that are displayed by the Parameter Panel.
+	TWeakObjectPtr<UNiagaraGraph> SystemScriptGraph;
+	TArray<TWeakObjectPtr<UNiagaraGraph>>  SelectedEmitterScriptGraphs;
+
+	TSharedPtr<FNiagaraSystemViewModel> SystemViewModel;
+	UNiagaraSystemSelectionViewModel* OverviewSelectionViewModel;
+};
+
+class FNiagaraScriptToolkitParameterPanelViewModel : public INiagaraParameterPanelViewModel
+{
+public:
+	/** Construct a ScriptToolkit Parameter Panel View Model from a Script View Model. */
+	FNiagaraScriptToolkitParameterPanelViewModel(TSharedPtr<FNiagaraStandaloneScriptViewModel> InScriptViewModel);
+
+	~FNiagaraScriptToolkitParameterPanelViewModel();
+
+	/** Begin INiagaraParameterPanelViewModel interface. */
+	virtual void InitBindings() override;
+
+	virtual void Refresh() const override;
+
+	virtual void CollectStaticSections(TArray<int32>& StaticSectionIDs) const override;
+
+	virtual NiagaraParameterPanelSectionID::Type GetSectionForVarAndViewInfo(const FNiagaraScriptVariableAndViewInfo& VarAndViewInfo) const override;
+
+	virtual void AddParameter(const FNiagaraVariable& VariableToAdd, const FNiagaraVariableMetaData& VariableMetaDataToAssign) override;
+
+	virtual void RemoveParameter(const FNiagaraVariable& TargetVariableToRemove, const FNiagaraVariableMetaData& TargetVariableMetaData) override;
+
+	virtual bool CanRemoveParameter(const FNiagaraVariable& TargetVariableToRemove, const FNiagaraVariableMetaData& TargetVariableMetaData) const override;
+
+	virtual void RenameParameter(const FNiagaraVariable& TargetVariableToRename, const FNiagaraVariableMetaData& TargetVariableMetaData, const FText& NewVariableNameText) const override;
+
+	virtual void ChangeParameterScope(const FNiagaraVariable& TargetVariableToModify, const FNiagaraVariableMetaData& TargetVariableMetaData, const ENiagaraParameterScope NewVariableScope) const override;
+
+	virtual bool CanModifyParameter(const FNiagaraVariable& TargetVariableToModify, const FNiagaraVariableMetaData& TargetVariableMetaData) const override;
+
+	virtual bool CanRenameParameter(const FNiagaraVariable& TargetVariableToRename, const FNiagaraVariableMetaData& TargetVariableMetaData, const FText& NewVariableNameText) const override;
+
+	virtual void HandleActionSelected(const TSharedPtr<FEdGraphSchemaAction>& InAction, ESelectInfo::Type InSelectionType) override;
+
+	virtual FReply HandleActionDragged(const TSharedPtr<FEdGraphSchemaAction>& InAction, const FPointerEvent& MouseEvent) const override;
+
+	virtual TArray<TWeakObjectPtr<UNiagaraGraph>> GetEditableGraphs() const override;
+
+	virtual const TArray<FNiagaraScriptVariableAndViewInfo> GetViewedParameters() override;
+	/** End INiagaraParameterPanelViewModel interface. */
+
+	void RenamePin(const UEdGraphPin* TargetPinToRename, const FText& NewNameText) const;
+	void ChangePinScope(const UEdGraphPin* TargetPin, const ENiagaraParameterScope NewScope) const;
+
+private:
+	TSharedPtr<FNiagaraStandaloneScriptViewModel> ScriptViewModel;
+
+	FDelegateHandle OnGraphChangedHandle;
+	FDelegateHandle OnGraphNeedsRecompileHandle;
+
+	void HandleOnGraphChanged(const struct FEdGraphEditAction& InAction);
+
+	TStaticArray<FScopeIsEnabledAndTooltip, (int32)ENiagaraParameterScope::Num> GetParameterScopesEnabledAndTooltips(const FNiagaraVariable& InVar, const FNiagaraVariableMetaData& InVarMetaData) const;
+
+	TSharedRef<SWidget> GetScriptParameterVisualWidget(const UEdGraphPin* Pin) const;
+
+	TSharedPtr<FNiagaraObjectSelection> VariableObjectSelection;
+};
diff --git a/Engine/Plugins/FX/Niagara/Source/NiagaraEditor/Private/ViewModels/NiagaraStandaloneScriptViewModel.cpp b/Engine/Plugins/FX/Niagara/Source/NiagaraEditor/Private/ViewModels/NiagaraStandaloneScriptViewModel.cpp
--- a/Engine/Plugins/FX/Niagara/Source/NiagaraEditor/Private/ViewModels/NiagaraStandaloneScriptViewModel.cpp
+++ b/Engine/Plugins/FX/Niagara/Source/NiagaraEditor/Private/ViewModels/NiagaraStandaloneScriptViewModel.cpp
@@ -22,6 +22,12 @@ void FNiagaraStandaloneScriptViewModel::Initialize(UNiagaraScript* InScript, UNi
 	SendLastCompileMessageJobs(SourceScript);
 }
 
+UNiagaraScript* FNiagaraStandaloneScriptViewModel::GetStandaloneScript()
+{
+	checkf(Scripts.Num() == 1, TEXT("StandaloneScriptViewModel did not have exactly one script!"));
+	return Scripts[0].Get();
+}
+
 void FNiagaraStandaloneScriptViewModel::OnVMScriptCompiled(UNiagaraScript* InScript)
 {
 	FNiagaraScriptViewModel::OnVMScriptCompiled(InScript);
diff --git a/Engine/Plugins/FX/Niagara/Source/NiagaraEditor/Private/ViewModels/NiagaraStandaloneScriptViewModel.h b/Engine/Plugins/FX/Niagara/Source/NiagaraEditor/Private/ViewModels/NiagaraStandaloneScriptViewModel.h
--- a/Engine/Plugins/FX/Niagara/Source/NiagaraEditor/Private/ViewModels/NiagaraStandaloneScriptViewModel.h
+++ b/Engine/Plugins/FX/Niagara/Source/NiagaraEditor/Private/ViewModels/NiagaraStandaloneScriptViewModel.h
@@ -20,6 +20,7 @@ public:
 	);
 
 	void Initialize(UNiagaraScript* InScript, UNiagaraScript* InSourceScript);
+	UNiagaraScript* GetStandaloneScript();
 
 private:
 	virtual void OnVMScriptCompiled(UNiagaraScript* InScript) override;
diff --git a/Engine/Plugins/FX/Niagara/Source/NiagaraEditor/Private/ViewModels/Stack/NiagaraParameterHandle.cpp b/Engine/Plugins/FX/Niagara/Source/NiagaraEditor/Private/ViewModels/Stack/NiagaraParameterHandle.cpp
--- a/Engine/Plugins/FX/Niagara/Source/NiagaraEditor/Private/ViewModels/Stack/NiagaraParameterHandle.cpp
+++ b/Engine/Plugins/FX/Niagara/Source/NiagaraEditor/Private/ViewModels/Stack/NiagaraParameterHandle.cpp
@@ -13,6 +13,8 @@ const FName FNiagaraParameterHandle::ParticleAttributeNamespace(TEXT("Particles"
 const FName FNiagaraParameterHandle::ModuleNamespace(TEXT("Module"));
 const FName FNiagaraParameterHandle::ParameterCollectionNamespace(TEXT("NPC"));
 const FString FNiagaraParameterHandle::InitialPrefix(TEXT("Initial"));
+const FName FNiagaraParameterHandle::LocalNamespace(TEXT("Local"));
+const FName FNiagaraParameterHandle::InitialNamespace(TEXT("Initial"));
 
 FNiagaraParameterHandle::FNiagaraParameterHandle() 
 {
diff --git a/Engine/Plugins/FX/Niagara/Source/NiagaraEditor/Private/ViewModels/Stack/NiagaraStackGraphUtilities.cpp b/Engine/Plugins/FX/Niagara/Source/NiagaraEditor/Private/ViewModels/Stack/NiagaraStackGraphUtilities.cpp
--- a/Engine/Plugins/FX/Niagara/Source/NiagaraEditor/Private/ViewModels/Stack/NiagaraStackGraphUtilities.cpp
+++ b/Engine/Plugins/FX/Niagara/Source/NiagaraEditor/Private/ViewModels/Stack/NiagaraStackGraphUtilities.cpp
@@ -1607,6 +1607,67 @@ TOptional<FName> FNiagaraStackGraphUtilities::GetNamespaceForScriptUsage(ENiagar
 	}
 }
 
+ENiagaraParameterScope FNiagaraStackGraphUtilities::GetScopeForScriptUsage(ENiagaraScriptUsage ScriptUsage)
+{
+	switch (ScriptUsage)
+	{
+	case ENiagaraScriptUsage::ParticleSpawnScript:
+	case ENiagaraScriptUsage::ParticleSpawnScriptInterpolated:
+	case ENiagaraScriptUsage::ParticleUpdateScript:
+	case ENiagaraScriptUsage::ParticleEventScript:
+	case ENiagaraScriptUsage::ParticleSimulationStageScript:
+		return ENiagaraParameterScope::Particles;
+	case ENiagaraScriptUsage::EmitterSpawnScript:
+	case ENiagaraScriptUsage::EmitterUpdateScript:
+		return ENiagaraParameterScope::Emitter;
+	case ENiagaraScriptUsage::SystemSpawnScript:
+	case ENiagaraScriptUsage::SystemUpdateScript:
+		return ENiagaraParameterScope::System;
+	case ENiagaraScriptUsage::Module:
+	case ENiagaraScriptUsage::Function:
+	case ENiagaraScriptUsage::DynamicInput:
+		return ENiagaraParameterScope::None; // These script usages do not have associated scopes.
+	default:
+		checkf(false, TEXT("Script usage does not have known scope!"));
+		return ENiagaraParameterScope::None;
+	}
+}
+
+void FNiagaraStackGraphUtilities::GetOwningEmitterAndScriptForStackNode(UNiagaraNode& StackNode, UNiagaraSystem& OwningSystem, UNiagaraEmitter*& OutEmitter, UNiagaraScript*& OutScript)
+{
+	OutEmitter = nullptr;
+	OutScript = nullptr;
+	UNiagaraNodeOutput* OutputNode = GetEmitterOutputNodeForStackNode(StackNode);
+	if (OutputNode != nullptr)
+	{
+		switch (OutputNode->GetUsage())
+		{
+			case ENiagaraScriptUsage::SystemSpawnScript:
+				OutScript = OwningSystem.GetSystemSpawnScript();
+				break;
+			case ENiagaraScriptUsage::SystemUpdateScript:
+				OutScript = OwningSystem.GetSystemUpdateScript();
+				break;
+			case ENiagaraScriptUsage::EmitterSpawnScript:
+			case ENiagaraScriptUsage::EmitterUpdateScript:
+			case ENiagaraScriptUsage::ParticleSpawnScript:
+			case ENiagaraScriptUsage::ParticleUpdateScript:
+			case ENiagaraScriptUsage::ParticleEventScript:
+				for (const FNiagaraEmitterHandle& EmitterHandle : OwningSystem.GetEmitterHandles())
+				{
+					UNiagaraScriptSource* EmitterSource = CastChecked<UNiagaraScriptSource>(EmitterHandle.GetInstance()->GraphSource);
+					if (EmitterSource->NodeGraph == StackNode.GetNiagaraGraph())
+					{
+						OutEmitter = EmitterHandle.GetInstance();
+						OutScript = OutEmitter->GetScript(OutputNode->GetUsage(), OutputNode->GetUsageId());
+						break;
+					}
+				}
+				break;
+		}
+	}
+}
+
 struct FRapidIterationParameterContext
 {
 	FRapidIterationParameterContext()
@@ -2496,4 +2557,160 @@ void FNiagaraStackGraphUtilities::RenameReferencingParameters(UNiagaraSystem& Sy
 	}
 }
 
+FString FNiagaraStackGraphUtilities::GetNamespaceStringForScriptParameterScope(const ENiagaraParameterScope& InScope)
+{
+	switch (InScope) {
+		case ENiagaraParameterScope::Engine:
+			return PARAM_MAP_ENGINE_STR;
+			break;
+		case ENiagaraParameterScope::User:
+			return PARAM_MAP_USER_STR;
+			break;
+		case ENiagaraParameterScope::System:
+			return PARAM_MAP_SYSTEM_STR;
+			break;
+		case ENiagaraParameterScope::Emitter:
+			return PARAM_MAP_EMITTER_STR;
+			break;
+		case ENiagaraParameterScope::Particles:
+			return PARAM_MAP_ATTRIBUTE_STR;
+			break;
+		case ENiagaraParameterScope::ScriptPersistent:
+			return PARAM_MAP_SCRIPT_PERSISTENT_STR;
+			break;
+		case ENiagaraParameterScope::ScriptTransient:
+			return PARAM_MAP_SCRIPT_TRANSIENT_STR;
+			break;
+		case ENiagaraParameterScope::Input:
+			return PARAM_MAP_MODULE_STR;
+			break;
+		case ENiagaraParameterScope::Local:
+			return PARAM_MAP_LOCAL_STR;
+			break;
+		default:
+			checkf(false, TEXT("Unhandled parameter scope encountered!"));
+			return FString();
+			break;
+	};
+}
+
+FString FNiagaraStackGraphUtilities::GetNamespaceStringForScriptParameterMetaData(const FNiagaraVariableMetaData& InMetaData)
+{
+	if (InMetaData.IsInputOrLocalUsage())
+	{
+		FString NamespaceString = FNiagaraStackGraphUtilities::GetNamespaceStringForScriptParameterScope(InMetaData.Scope);
+		if (InMetaData.Usage == ENiagaraScriptParameterUsage::InitialValueInput)
+		{
+			NamespaceString.Append(PARAM_MAP_INITIAL_STR);
+		}
+		return NamespaceString;
+	}
+	return PARAM_MAP_OUTPUT_STR;
+}
+
+TArray<FName> FNiagaraStackGraphUtilities::DecomposeVariableNamespace(const FName& InVarName, FName& OutFriendlyName)
+{
+	int32 DotIndex;
+	TArray<FName> OutNamespaces;
+	FString VarNameString = InVarName.ToString();
+	while (VarNameString.FindChar(TEXT('.'), DotIndex))
+	{
+		OutNamespaces.Add(FName(*VarNameString.Left(DotIndex)));
+		VarNameString = *VarNameString.RightChop(DotIndex + 1);
+	}
+	OutFriendlyName = FName(*VarNameString);
+	return OutNamespaces;
+}
+
+void FNiagaraStackGraphUtilities::SetParameterMetaDataFromName(const FName& InVarName, ENiagaraParameterScope& OutScope, ENiagaraScriptParameterUsage& OutUsage, FName& OutFriendlyName)
+{
+	TArray<FName> DecomposedNamespaces = FNiagaraStackGraphUtilities::DecomposeVariableNamespace(InVarName, OutFriendlyName);
+	if (DecomposedNamespaces.Num() == 0)
+	{
+		UE_LOG(LogNiagaraEditor, Display, TEXT("Unexpected parameter encountered without a namespace: %s"), *InVarName.ToString());
+	}
+
+	bool bCheckIsInitialValue = false;
+	for (int NamespaceIdx = 0; NamespaceIdx < DecomposedNamespaces.Num(); ++NamespaceIdx)
+	{
+		const FName& Namespace = DecomposedNamespaces[NamespaceIdx];
+
+		if (bCheckIsInitialValue)
+		{
+			if (Namespace == FNiagaraParameterHandle::InitialNamespace)
+			{
+				OutUsage = ENiagaraScriptParameterUsage::InitialValueInput;
+			}
+			return;
+		}
+
+		if (Namespace == FNiagaraParameterHandle::LocalNamespace)
+		{
+			OutScope = ENiagaraParameterScope::Local;
+			OutUsage = ENiagaraScriptParameterUsage::Local;
+			return;
+		}
+		else if (Namespace == FNiagaraParameterHandle::ModuleNamespace)
+		{
+			OutScope = ENiagaraParameterScope::Input;
+			return;
+		}
+		else if (Namespace == FNiagaraParameterHandle::UserNamespace)
+		{
+			OutScope = ENiagaraParameterScope::User;
+			return;
+		}
+		else if (Namespace == FNiagaraParameterHandle::EngineNamespace)
+		{
+			OutScope = ENiagaraParameterScope::Engine;
+			return;
+		}
+		else if (Namespace == FNiagaraParameterHandle::SystemNamespace)
+		{
+			OutScope = ENiagaraParameterScope::System;
+			bCheckIsInitialValue = true;
+		}
+		else if (Namespace == FNiagaraParameterHandle::EmitterNamespace)
+		{
+			OutScope = ENiagaraParameterScope::Emitter;
+			bCheckIsInitialValue = true;
+		}
+		else if (Namespace == FNiagaraParameterHandle::ParticleAttributeNamespace)
+		{
+			OutScope = ENiagaraParameterScope::Particles;
+			bCheckIsInitialValue = true;
+		}
+	}
+	
+	if (bCheckIsInitialValue == false)
+	{
+		// custom user namespace
+		OutScope = ENiagaraParameterScope::ScriptTransient;
+	}
+}
+
+FString FNiagaraStackGraphUtilities::GetNamespacelessVariableNameString(const FName& InVarName)
+{
+	int32 DotIndex;
+	FString VarNameString = InVarName.ToString();
+	if (VarNameString.FindLastChar(TEXT('.'), DotIndex))
+	{
+		return VarNameString.RightChop(DotIndex + 1);
+	}
+	// No dot index, must be a namespaceless variable name (e.g. static switch name) just return the name to string.
+	return VarNameString;
+}
+
+FName FNiagaraStackGraphUtilities::GetVariableNameForScope(const FName& CurrentVariableName, const ENiagaraParameterScope NewVariableScope, bool bIsInitialValue)
+{
+	FString NewVariableScopeString = FNiagaraStackGraphUtilities::GetNamespaceStringForScriptParameterScope(NewVariableScope);
+	if (bIsInitialValue)
+	{
+		NewVariableScopeString.Append(PARAM_MAP_INITIAL_STR);
+	}
+	const FString ExistingVariableNameString = FNiagaraStackGraphUtilities::GetNamespacelessVariableNameString(CurrentVariableName);
+	FName NewVariableName = FName(*(NewVariableScopeString + ExistingVariableNameString));
+	return NewVariableName;
+}
+
 #undef LOCTEXT_NAMESPACE
diff --git a/Engine/Plugins/FX/Niagara/Source/NiagaraEditor/Private/Widgets/SNiagaraGraphPinAdd.cpp b/Engine/Plugins/FX/Niagara/Source/NiagaraEditor/Private/Widgets/SNiagaraGraphPinAdd.cpp
--- a/Engine/Plugins/FX/Niagara/Source/NiagaraEditor/Private/Widgets/SNiagaraGraphPinAdd.cpp
+++ b/Engine/Plugins/FX/Niagara/Source/NiagaraEditor/Private/Widgets/SNiagaraGraphPinAdd.cpp
@@ -11,8 +11,8 @@
 #include "Widgets/Input/SEditableTextBox.h"
 #include "Widgets/Text/SInlineEditableTextBlock.h"
 #include "Widgets/Layout/SBox.h"
-#include "SNiagaraParameterMapView.h"
 #include "Framework/Application/SlateApplication.h"
+#include "SNiagaraParameterPanel.h"
 
 #define LOCTEXT_NAMESPACE "NiagaraGraphPinAdd"
 
@@ -68,12 +68,14 @@ TSharedRef<SWidget> SNiagaraGraphPinAdd::OnGetAddButtonMenuContent()
 	{
 		TArray<TWeakObjectPtr<UNiagaraGraph>> Graphs;
 		Graphs.Add(OwningNode->GetNiagaraGraph());
-		TSharedRef<SNiagaraAddParameterMenu> MenuWidget = SNew(SNiagaraAddParameterMenu, Graphs)
-			.OnAddParameter_UObject(OwningNode, &UNiagaraNodeWithDynamicPins::AddParameter, GetPinObj()) // For non custom actions
+
+		// Creating a new parameter from a graph pin always defaults to Particles scope.
+		TSharedRef<SNiagaraAddParameterMenu2> MenuWidget = SNew(SNiagaraAddParameterMenu2, Graphs, ENiagaraParameterScope::Particles)
+			.OnAddParameter_UObject(OwningNode, &UNiagaraNodeWithDynamicPins::AddParameter, (const UEdGraphPin*)GetPinObj())
 			.OnCollectCustomActions_UObject(OwningNode, &UNiagaraNodeWithDynamicPins::CollectAddPinActions, GetPinObj())
 			.OnAllowMakeType_UObject(OwningNode, &UNiagaraNodeWithDynamicPins::AllowNiagaraTypeForAddPin)
 			.IsParameterRead(GraphPinObj ? GraphPinObj->Direction == EGPD_Output : true);
-		
+
 		AddButton->SetMenuContentWidgetToFocus(MenuWidget->GetSearchBox());
 		return MenuWidget;
 	}
@@ -81,7 +83,6 @@ TSharedRef<SWidget> SNiagaraGraphPinAdd::OnGetAddButtonMenuContent()
 	{
 		return SNullWidget::NullWidget;
 	}
-
 }
 
 #undef LOCTEXT_NAMESPACE
diff --git a/Engine/Plugins/FX/Niagara/Source/NiagaraEditor/Private/Widgets/SNiagaraParameterNameView.cpp b/Engine/Plugins/FX/Niagara/Source/NiagaraEditor/Private/Widgets/SNiagaraParameterNameView.cpp
new file mode 100644
--- /dev/null
+++ b/Engine/Plugins/FX/Niagara/Source/NiagaraEditor/Private/Widgets/SNiagaraParameterNameView.cpp
@@ -0,0 +1,33 @@
+// Copyright Epic Games, Inc. All Rights Reserved.
+
+#include "SNiagaraParameterNameView.h"
+#include "ViewModels/NiagaraParameterNameViewModel.h"
+#include "Widgets/SBoxPanel.h"
+
+class SWidget;
+
+void SNiagaraParameterNameView::Construct(const FArguments& InArgs, const TSharedPtr<INiagaraParameterNameViewModel>& InParameterNameViewModel)
+{
+	ParameterNameViewModel = InParameterNameViewModel;
+
+	TSharedRef<SWidget> ScopeSlotWidget = ParameterNameViewModel->CreateScopeSlotWidget();
+	TSharedRef<SWidget> NameSlotWidget = ParameterNameViewModel->CreateTextSlotWidget();
+
+	ChildSlot
+	[
+		SNew(SHorizontalBox)
+		+SHorizontalBox::Slot()
+		.AutoWidth()
+		.VAlign(VAlign_Center)
+		.Padding(3, 0)
+		[
+			ScopeSlotWidget
+		]
+		+SHorizontalBox::Slot()
+		.VAlign(VAlign_Center)
+		.Padding(3, 0)
+		[
+			NameSlotWidget
+		]
+	];
+}
diff --git a/Engine/Plugins/FX/Niagara/Source/NiagaraEditor/Private/Widgets/SNiagaraParameterNameView.h b/Engine/Plugins/FX/Niagara/Source/NiagaraEditor/Private/Widgets/SNiagaraParameterNameView.h
new file mode 100644
--- /dev/null
+++ b/Engine/Plugins/FX/Niagara/Source/NiagaraEditor/Private/Widgets/SNiagaraParameterNameView.h
@@ -0,0 +1,84 @@
+// Copyright Epic Games, Inc. All Rights Reserved.
+
+#pragma once
+
+struct FNiagaraScriptVariableAndViewInfo;
+struct FSlateFontInfo;
+struct FCreateWidgetForActionData;
+class INiagaraParameterNameViewModel;
+
+struct FCreateParameterNameViewInfo
+{
+	FCreateParameterNameViewInfo() = delete;
+
+	FCreateParameterNameViewInfo(const FNiagaraScriptVariableAndViewInfo& InScriptVarAndViewInfo, const FSlateFontInfo& InNameFont, bool bInIsReadOnly)
+		: ScriptVarAndViewInfo(InScriptVarAndViewInfo)
+		, NameFont(InNameFont)
+		, bIsReadOnly(bInIsReadOnly)
+	{};
+
+	FCreateParameterNameViewInfo(const FNiagaraScriptVariableAndViewInfo& InScriptVarAndViewInfo, const FSlateFontInfo& InNameFont, bool bInIsReadOnly, const FCreateWidgetForActionData* const InCreateData)
+		: ScriptVarAndViewInfo(InScriptVarAndViewInfo)
+		, NameFont(InNameFont)
+		, bIsReadOnly(bInIsReadOnly)
+		, CreateData(InCreateData)
+	{};
+
+	const FNiagaraScriptVariableAndViewInfo& ScriptVarAndViewInfo;
+	const FSlateFontInfo& NameFont;
+	bool bIsReadOnly;
+	TOptional<const FCreateWidgetForActionData* const> CreateData;
+};
+
+/** A widget for viewing and editing FNiagaraVariable names by separating their names into a namespace scope SEnumComboBox and name SInlineEditableTextBox. */
+class SNiagaraParameterNameView : public SCompoundWidget
+{
+public:
+
+	/** Scope ComboBox Delegates */
+	DECLARE_DELEGATE_RetVal(int32, FOnGetScopeValue);
+	DECLARE_DELEGATE_TwoParams(FOnScopeValueChanged, int32, ESelectInfo::Type);
+	DECLARE_DELEGATE_RetVal_OneParam(FText, FOnGetToolTipForScopeValue, int32 /* Value */)
+	DECLARE_DELEGATE_RetVal_OneParam(bool, FOnIsScopeValueEnabled, int32 /* Value */)
+
+	/** Name InlineEditableTextBox Delegates */
+	DECLARE_DELEGATE_RetVal(FText, FOnGetParameterNameText);
+	DECLARE_DELEGATE_RetVal_TwoParams(bool, FOnVerifyParameterNameChanged, const FText&, FText&)
+	DECLARE_DELEGATE_TwoParams(FOnParameterRenamed, const FText&, const struct FNiagaraScriptVarAndViewInfoAction&)
+
+
+	SLATE_BEGIN_ARGS(SNiagaraParameterNameView)
+	{}
+		/** Scope ComboBox Delegates */
+		SLATE_EVENT(FOnGetScopeValue, OnGetScopeValue)
+		SLATE_EVENT(FOnScopeValueChanged, OnScopeValueChanged)
+		SLATE_EVENT(FOnGetToolTipForScopeValue, OnGetToolTipForScopeValue)
+		SLATE_EVENT(FOnIsScopeValueEnabled, OnIsScopeValueEnabled)
+
+		/** Name InlineEditableTextBox Delegates */
+		SLATE_EVENT(FOnGetParameterNameText, OnGetParameterNameText)
+		SLATE_EVENT(FOnVerifyParameterNameChanged, OnVerifyParameterNameChanged)
+		SLATE_EVENT(FOnParameterRenamed, OnParameterRenamed)
+	SLATE_END_ARGS();
+
+	void Construct(const FArguments& InArgs, const TSharedPtr<INiagaraParameterNameViewModel>& InParameterNameViewModel);
+
+protected:
+	TSharedPtr<INiagaraParameterNameViewModel> ParameterNameViewModel;
+
+	/** Scope ComboBox Delegates */
+	FOnGetScopeValue OnGetScopeValue;
+	FOnScopeValueChanged OnScopeValueChanged;
+	FOnGetToolTipForScopeValue OnGetToolTipForScopeValue;
+	FOnIsScopeValueEnabled OnIsScopeValueEnabled;
+
+	/** Name InlineEditableTextBox Delegates */
+	FOnGetParameterNameText OnGetParameterNameText;
+	FOnVerifyParameterNameChanged OnVerifyParameterNameChanged;
+	FOnParameterRenamed OnParameterRenamed;
+
+	int32 GetScopeValue() const;
+	FText GetTooltipForScopeValue(int32 Value) const;
+	bool GetScopeValueIsEnabled(int32 Value) const;
+	FText GetParameterNameText() const;
+};
diff --git a/Engine/Plugins/FX/Niagara/Source/NiagaraEditor/Private/Widgets/SNiagaraParameterPanel.cpp b/Engine/Plugins/FX/Niagara/Source/NiagaraEditor/Private/Widgets/SNiagaraParameterPanel.cpp
new file mode 100644
--- /dev/null
+++ b/Engine/Plugins/FX/Niagara/Source/NiagaraEditor/Private/Widgets/SNiagaraParameterPanel.cpp
@@ -0,0 +1,534 @@
+// Copyright Epic Games, Inc. All Rights Reserved.
+
+#include "SNiagaraParameterPanel.h"
+#include "NiagaraEditorCommon.h"
+#include "NiagaraParameterPanelViewModel.h"
+#include "Framework/Commands/UICommandList.h"
+#include "SGraphActionMenu.h"
+#include "NiagaraActions.h"
+#include "NiagaraEditorUtilities.h"
+#include "Framework/Commands/GenericCommands.h"
+#include "DetailLayoutBuilder.h"
+#include "Widgets/Input/SSearchBox.h"
+#include "ViewModels/Stack/NiagaraStackGraphUtilities.h"
+#include "SNiagaraParameterPanelPaletteItem.h"
+#include "SNiagaraGraphActionWidget.h"
+#include "NiagaraScriptVariable.h"
+
+#define LOCTEXT_NAMESPACE "NiagaraParameterPanel"
+
+SNiagaraParameterPanel::~SNiagaraParameterPanel()
+{
+	ParameterPanelViewModel->GetOnRefreshed().Unbind();
+
+	// Unregister all commands for right click on action node
+	ToolkitCommands->UnmapAction(FNiagaraParameterPanelCommands::Get().DeleteEntry);
+	ToolkitCommands->UnmapAction(FGenericCommands::Get().Rename);
+}
+
+void FNiagaraParameterPanelCommands::RegisterCommands()
+{
+	UI_COMMAND(DeleteEntry, "Delete", "Delete this parameter", EUserInterfaceActionType::Button, FInputChord(EKeys::Platform_Delete));
+}
+
+void SNiagaraParameterPanel::Construct(const FArguments& InArgs, const TSharedPtr<INiagaraParameterPanelViewModel>& InParameterPanelViewModel, const TSharedPtr<FUICommandList>& InToolkitCommands)
+{
+	ParameterPanelViewModel = InParameterPanelViewModel;
+	ToolkitCommands = InToolkitCommands;
+
+	ParameterPanelViewModel->GetOnRefreshed().BindRaw(this, &SNiagaraParameterPanel::Refresh);
+
+	AddParameterButtons.SetNum(NiagaraParameterPanelSectionID::PARTICLES + 1); //@Todo(ng) verify
+
+	// Register all commands for right click on action node
+	FNiagaraParameterPanelCommands::Register();
+	TSharedPtr<FUICommandList> ToolKitCommandList = ToolkitCommands;
+	ToolKitCommandList->MapAction(FNiagaraParameterPanelCommands::Get().DeleteEntry,
+		FExecuteAction::CreateSP(this, &SNiagaraParameterPanel::TryDeleteEntries),
+		FCanExecuteAction::CreateSP(this, &SNiagaraParameterPanel::CanTryDeleteEntries));
+	ToolKitCommandList->MapAction(FGenericCommands::Get().Rename,
+		FExecuteAction::CreateSP(this, &SNiagaraParameterPanel::OnRequestRenameOnActionNode),
+		FCanExecuteAction::CreateSP(this, &SNiagaraParameterPanel::CanRequestRenameOnActionNode));
+
+	// create the main action list piece of this widget
+	SAssignNew(GraphActionMenu, SGraphActionMenu, false)
+		.OnCreateWidgetForAction(this, &SNiagaraParameterPanel::OnCreateWidgetForAction)
+		.OnCollectAllActions(this, &SNiagaraParameterPanel::CollectAllActions)
+ 		.OnCollectStaticSections(this, &SNiagaraParameterPanel::CollectStaticSections)
+ 		.OnActionDragged(this, &SNiagaraParameterPanel::OnActionDragged)
+		.OnActionSelected(this, &SNiagaraParameterPanel::OnActionSelected)
+// 		.OnActionDoubleClicked(this, &SNiagaraParameterMapView::OnActionDoubleClicked) //@todo(ng) impl
+// 		.OnContextMenuOpening(this, &SNiagaraParameterMapView::OnContextMenuOpening)
+// 		.OnCanRenameSelectedAction(this, &SNiagaraParameterMapView::CanRequestRenameOnActionNode)
+ 		.OnGetSectionTitle(this, &SNiagaraParameterPanel::OnGetSectionTitle)
+ 		.OnGetSectionWidget(this, &SNiagaraParameterPanel::OnGetSectionWidget)
+ 		.OnCreateCustomRowExpander_Static(&SNiagaraParameterPanel::CreateCustomActionExpander)
+		.OnActionMatchesName(this, &SNiagaraParameterPanel::HandleActionMatchesName)
+		.AutoExpandActionMenu(false)
+		.AlphaSortItems(false)
+		.UseSectionStyling(true)
+		.ShowFilterTextBox(true);
+
+	ChildSlot
+	[
+		SNew(SBox)
+		.MinDesiredWidth(300)
+		[
+			SNew(SVerticalBox)	
+			+ SVerticalBox::Slot()
+			.FillHeight(1.0f)
+			[
+				GraphActionMenu.ToSharedRef()
+			]
+		]
+	];
+}
+
+void SNiagaraParameterPanel::Tick(const FGeometry& AllottedGeometry, const double InCurrentTime, const float InDeltaTime)
+{
+	if (bNeedsRefresh)
+	{
+		GraphActionMenu->RefreshAllActions(true);
+		bNeedsRefresh = false;
+	}
+}
+
+void SNiagaraParameterPanel::AddParameter(FNiagaraVariable NewVariable, const NiagaraParameterPanelSectionID::Type SectionID)
+{
+	FNiagaraVariableMetaData NewVariableMetaData = FNiagaraVariableMetaData();
+	NewVariableMetaData.Scope = NiagaraParameterPanelSectionID::GetScopeForNewParametersInSection(SectionID);
+	NewVariableMetaData.Usage = NiagaraParameterPanelSectionID::GetUsageForNewParametersInSection(SectionID);
+
+	ParameterPanelViewModel->AddParameter(NewVariable, NewVariableMetaData);
+
+	GraphActionMenu->RefreshAllActions(true);
+	GraphActionMenu->SelectItemByName(NewVariable.GetName());
+	GraphActionMenu->OnRequestRenameOnActionNode();
+}
+
+bool SNiagaraParameterPanel::AllowMakeType(const FNiagaraTypeDefinition& InType) const
+{
+	return InType != FNiagaraTypeDefinition::GetParameterMapDef();
+}
+
+TSharedRef<SWidget> SNiagaraParameterPanel::OnCreateWidgetForAction(struct FCreateWidgetForActionData* const InCreateData)
+{
+	TSharedPtr<SWidget> ParameterVisualWidget = ParameterPanelViewModel->GetScriptParameterVisualWidget(InCreateData);
+	return SNew(SNiagaraParameterPanelPaletteItem, InCreateData, ParameterVisualWidget->AsShared());
+}
+
+void SNiagaraParameterPanel::CollectAllActions(FGraphActionListBuilderBase& OutAllActions)
+{
+	const FText TooltipFormat = LOCTEXT("Parameters", "Name: {0} \nType: {1}"); //@todo(ng) move to const static
+
+	TArray<FNiagaraScriptVariableAndViewInfo> VarAndViewInfos = ParameterPanelViewModel->GetViewedParameters();
+	VarAndViewInfos.Sort([](const FNiagaraScriptVariableAndViewInfo& A, const FNiagaraScriptVariableAndViewInfo& B) { return A.MetaData.CachedNamespacelessVariableName.LexicalLess(B.MetaData.CachedNamespacelessVariableName); });
+	for (const FNiagaraScriptVariableAndViewInfo& VarAndViewInfo : VarAndViewInfos)
+	{
+		const FText Name = FText::FromName(VarAndViewInfo.ScriptVariable.GetName());
+		const FText Tooltip = FText::Format(TooltipFormat, FText::FromName(VarAndViewInfo.ScriptVariable.GetName()), VarAndViewInfo.ScriptVariable.GetType().GetNameText());
+		const NiagaraParameterPanelSectionID::Type Section = ParameterPanelViewModel->GetSectionForVarAndViewInfo(VarAndViewInfo);
+		TSharedPtr<FNiagaraScriptVarAndViewInfoAction> ScriptVarAndViewInfoAction(new FNiagaraScriptVarAndViewInfoAction(VarAndViewInfo, FText::GetEmpty(), Name, Tooltip, 0, FText(), Section /*= 0*/));
+		OutAllActions.AddAction(ScriptVarAndViewInfoAction);
+	}
+}
+ 
+void SNiagaraParameterPanel::CollectStaticSections(TArray<int32>& StaticSectionIDs) const
+{
+	return ParameterPanelViewModel->CollectStaticSections(StaticSectionIDs);
+}
+
+FReply SNiagaraParameterPanel::OnActionDragged(const TArray<TSharedPtr<FEdGraphSchemaAction>>& InActions, const FPointerEvent& MouseEvent)
+{
+	if (InActions.Num() == 1 && InActions[0].IsValid())
+	{
+		return ParameterPanelViewModel->HandleActionDragged(InActions[0], MouseEvent);
+	}
+	return FReply::Handled();
+}
+
+void SNiagaraParameterPanel::OnActionSelected(const TArray< TSharedPtr<FEdGraphSchemaAction> >& InActions, ESelectInfo::Type InSelectionType)
+{
+	if (InActions.Num() == 1 && InActions[0].IsValid())
+	{
+		ParameterPanelViewModel->HandleActionSelected(InActions[0], InSelectionType);
+	}
+}
+
+FText SNiagaraParameterPanel::OnGetSectionTitle(int32 InSectionID)
+{
+	return NiagaraParameterPanelSectionID::OnGetSectionTitle((NiagaraParameterPanelSectionID::Type)InSectionID);
+}
+
+TSharedRef<SWidget> SNiagaraParameterPanel::OnGetSectionWidget(TSharedRef<SWidget> RowWidget, int32 InSectionID)
+{
+	TWeakPtr<SWidget> WeakRowWidget = RowWidget;
+	FText AddNewText = LOCTEXT("AddNewParameter", "Add Parameter");
+	FName MetaDataTag = TEXT("AddNewParameter");
+	return CreateAddToSectionButton((NiagaraParameterPanelSectionID::Type) InSectionID, WeakRowWidget, AddNewText, MetaDataTag); //@todo(ng) refactor to get rid of int32 insectionid
+}
+
+TSharedRef<SWidget> SNiagaraParameterPanel::CreateAddToSectionButton(const NiagaraParameterPanelSectionID::Type InSection, TWeakPtr<SWidget> WeakRowWidget, FText AddNewText, FName MetaDataTag)
+{
+	TSharedPtr<SComboButton> Button;
+	SAssignNew(Button, SComboButton)
+	.ButtonStyle(FEditorStyle::Get(), "RoundButton")
+	.ForegroundColor(FEditorStyle::GetSlateColor("DefaultForeground"))
+	.ContentPadding(FMargin(2, 0))
+	.OnGetMenuContent(this, &SNiagaraParameterPanel::OnGetParameterMenu, InSection)
+	//.IsEnabled(this, &SNiagaraParameterMapView::ParameterAddEnabled) //@Todo(ng) impl
+	.HAlign(HAlign_Center)
+	.VAlign(VAlign_Center)
+	.HasDownArrow(false)
+	.AddMetaData<FTagMetaData>(FTagMetaData(MetaDataTag))
+	.ButtonContent()
+	[
+		SNew(SHorizontalBox)
+		+ SHorizontalBox::Slot()
+		.AutoWidth()
+		.Padding(FMargin(0, 1))
+		[
+			SNew(SImage)
+			.Image(FEditorStyle::GetBrush("Plus"))
+		]
+
+		+ SHorizontalBox::Slot()
+		.VAlign(VAlign_Center)
+		.AutoWidth()
+		.Padding(FMargin(2, 0, 0, 0))
+		[
+			SNew(STextBlock)
+			.Font(IDetailLayoutBuilder::GetDetailFontBold())
+			.Text(AddNewText)
+			//.Visibility(this, &SNiagaraParameterMapView::OnAddButtonTextVisibility, WeakRowWidget, InSection) //@todo(ng) impl
+			.ShadowOffset(FVector2D(1, 1))
+		]
+	];
+
+	AddParameterButtons[(int32)InSection] = Button;
+
+	return Button.ToSharedRef();
+}
+
+TSharedRef<SWidget> SNiagaraParameterPanel::OnGetParameterMenu(const NiagaraParameterPanelSectionID::Type InSection)
+{
+	ENiagaraParameterScope NewParameterScope = NiagaraParameterPanelSectionID::GetScopeForNewParametersInSection(InSection);
+
+	TSharedRef<SNiagaraAddParameterMenu2> MenuWidget = SNew(SNiagaraAddParameterMenu2, ParameterPanelViewModel->GetEditableGraphs(), NewParameterScope)
+		.OnAddParameter(this, &SNiagaraParameterPanel::AddParameter, InSection)
+		.OnAllowMakeType(this, &SNiagaraParameterPanel::AllowMakeType)
+		.ShowGraphParameters(false)
+		.AutoExpandMenu(true);
+
+	AddParameterButtons[(int32)InSection]->SetMenuContentWidgetToFocus(MenuWidget->GetSearchBox()->AsShared());
+	return MenuWidget;
+}
+
+void SNiagaraParameterPanel::TryDeleteEntries()
+{
+	TArray<TSharedPtr<FEdGraphSchemaAction>> SelectedActions;
+	GraphActionMenu->GetSelectedActions(SelectedActions);
+
+	for (auto& Action : SelectedActions)
+	{
+		TSharedPtr<FNiagaraScriptVarAndViewInfoAction> ScriptVarAndViewInfoAction = StaticCastSharedPtr<FNiagaraScriptVarAndViewInfoAction>(Action);
+		if (ScriptVarAndViewInfoAction.Get())
+		{
+			const FNiagaraVariable& Parameter = ScriptVarAndViewInfoAction->ScriptVariableAndViewInfo.ScriptVariable;
+			const FNiagaraVariableMetaData& ParameterMetaData = ScriptVarAndViewInfoAction->ScriptVariableAndViewInfo.MetaData;
+			if (ParameterPanelViewModel->CanRemoveParameter(Parameter, ParameterMetaData))
+			{
+				ParameterPanelViewModel->RemoveParameter(Parameter, ParameterMetaData);
+			}
+		}
+	}
+}
+
+bool SNiagaraParameterPanel::CanTryDeleteEntries() const
+{
+	TArray<TSharedPtr<FEdGraphSchemaAction>> SelectedActions;
+	GraphActionMenu->GetSelectedActions(SelectedActions);
+
+	for (auto& Action : SelectedActions)
+	{
+		TSharedPtr<FNiagaraScriptVarAndViewInfoAction> ScriptVarAndViewInfoAction = StaticCastSharedPtr<FNiagaraScriptVarAndViewInfoAction>(Action);
+		if (ScriptVarAndViewInfoAction.Get())
+		{
+			const FNiagaraVariable& Parameter = ScriptVarAndViewInfoAction->ScriptVariableAndViewInfo.ScriptVariable;
+			const FNiagaraVariableMetaData& ParameterMetaData = ScriptVarAndViewInfoAction->ScriptVariableAndViewInfo.MetaData;
+			if (ParameterPanelViewModel->CanRemoveParameter(Parameter, ParameterMetaData))
+			{
+				return true;
+			}
+		}
+	}
+	return false;
+}
+
+void SNiagaraParameterPanel::OnRequestRenameOnActionNode()
+{
+	GraphActionMenu->OnRequestRenameOnActionNode();
+}
+
+bool SNiagaraParameterPanel::CanRequestRenameOnActionNode(TWeakPtr<struct FGraphActionNode> InSelectedNode) const
+{
+	return true; //@todo(ng) delegate to viewmodel
+}
+
+bool SNiagaraParameterPanel::CanRequestRenameOnActionNode() const
+{
+	TArray<TSharedPtr<FEdGraphSchemaAction> > SelectedActions;
+	GraphActionMenu->GetSelectedActions(SelectedActions);
+
+	// If there is anything selected in the GraphActionMenu, check the item for if it can be renamed.
+	if (SelectedActions.Num() > 0)
+	{
+		return GraphActionMenu->CanRequestRenameOnActionNode();
+	}
+	return false;
+}
+
+bool SNiagaraParameterPanel::HandleActionMatchesName(FEdGraphSchemaAction* InAction, const FName& InName) const
+{
+	return FName(*InAction->GetMenuDescription().ToString()) == InName;
+}
+
+void SNiagaraParameterPanel::Refresh()
+{
+	bNeedsRefresh = true;
+}
+
+	/************************************************************************/
+	/* SNiagaraAddParameterMenu2                                             */
+	/************************************************************************/
+void SNiagaraAddParameterMenu2::Construct(const FArguments& InArgs, TArray<TWeakObjectPtr<UNiagaraGraph>> InGraphs, ENiagaraParameterScope InNewParameterScope)
+{
+	this->OnAddParameter = InArgs._OnAddParameter;
+	this->OnCollectCustomActions = InArgs._OnCollectCustomActions;
+	this->OnAllowMakeType = InArgs._OnAllowMakeType;
+	this->AllowCreatingNew = InArgs._AllowCreatingNew;
+	this->ShowGraphParameters = InArgs._ShowGraphParameters;
+	this->AutoExpandMenu = InArgs._AutoExpandMenu;
+	this->IsParameterRead = InArgs._IsParameterRead;
+
+	Graphs = InGraphs;
+	NewParameterScope = InNewParameterScope;
+	NewParameterScopeText = FText::FromString(FNiagaraStackGraphUtilities::GetNamespaceStringForScriptParameterScope(NewParameterScope));
+
+	ChildSlot
+	[
+		SNew(SBorder)
+		.BorderImage(FEditorStyle::GetBrush("Menu.Background"))
+		.Padding(5)
+		[
+			SNew(SBox)
+			.MinDesiredWidth(300)
+			[
+				SAssignNew(GraphMenu, SGraphActionMenu)
+				.OnActionSelected(this, &SNiagaraAddParameterMenu2::OnActionSelected)
+				.OnCollectAllActions(this, &SNiagaraAddParameterMenu2::CollectAllActions)
+				.AutoExpandActionMenu(AutoExpandMenu.Get())
+				.ShowFilterTextBox(true)
+				.OnCreateCustomRowExpander_Static(&SNiagaraParameterPanel::CreateCustomActionExpander)
+				.OnCreateWidgetForAction_Lambda([](const FCreateWidgetForActionData* InData)
+				{
+					return SNew(SNiagaraGraphActionWidget, InData);
+				})
+			]
+		]
+	];
+}
+
+TSharedRef<SEditableTextBox> SNiagaraAddParameterMenu2::GetSearchBox()
+{
+	return GraphMenu->GetFilterTextBox();
+}
+
+void SNiagaraAddParameterMenu2::OnActionSelected(const TArray<TSharedPtr<FEdGraphSchemaAction>>& SelectedActions, ESelectInfo::Type InSelectionType)
+{
+	if (InSelectionType == ESelectInfo::OnMouseClick || InSelectionType == ESelectInfo::OnKeyPress || SelectedActions.Num() == 0)
+	{
+		for (int32 ActionIndex = 0; ActionIndex < SelectedActions.Num(); ActionIndex++)
+		{
+			TSharedPtr<FNiagaraMenuAction> CurrentAction = StaticCastSharedPtr<FNiagaraMenuAction>(SelectedActions[ActionIndex]);
+
+			if (CurrentAction.IsValid())
+			{
+				FSlateApplication::Get().DismissAllMenus();
+				CurrentAction->ExecuteAction();
+			}
+		}
+	}
+}
+
+void SNiagaraAddParameterMenu2::CollectAllActions(FGraphActionListBuilderBase& OutAllActions)
+{
+	if (OnCollectCustomActions.IsBound())
+	{
+		bool bCreateRemainingActions = true;
+		OnCollectCustomActions.Execute(OutAllActions, bCreateRemainingActions);
+		if (!bCreateRemainingActions)
+		{
+			return;
+		}
+	}
+
+	struct FNiagaraMenuActionInfo
+	{
+		FNiagaraMenuActionInfo()
+			: CategoryText()
+			, DisplayNameText()
+			, TooltipText()
+			, NewVariable()
+			, Grouping(0)
+			, KeywordsText()
+			, bIsExperimental(false)
+		{};
+
+		FText CategoryText;
+		FText DisplayNameText;
+		FText TooltipText;
+		FNiagaraVariable NewVariable;
+		int32 Grouping;
+		FText KeywordsText;
+		bool bIsExperimental;
+	};
+
+	if (AllowCreatingNew.Get())
+	{
+		// Create actions to create new FNiagaraVariables of every allowed FNiagaraTypeDefinition type.
+		TArray<FNiagaraMenuActionInfo> NewParameterMenuActionInfos;
+
+		TArray<FNiagaraTypeDefinition> Types = FNiagaraTypeRegistry::GetRegisteredTypes();
+		for (const FNiagaraTypeDefinition& RegisteredType : Types)
+		{
+			bool bAllowType = true;
+			if (OnAllowMakeType.IsBound())
+			{
+				bAllowType = OnAllowMakeType.Execute(RegisteredType);
+			}
+
+			if (bAllowType)
+			{
+				FNiagaraMenuActionInfo NewMenuActionInfo = FNiagaraMenuActionInfo();
+				NewMenuActionInfo.CategoryText = LOCTEXT("NiagaraCreateNewParameterMenu", "Create New Parameter");
+
+				FText RegisteredTypeNameText = RegisteredType.GetNameText();
+				NewMenuActionInfo.DisplayNameText = RegisteredTypeNameText;
+
+				if (const UStruct* VariableStruct = RegisteredType.GetStruct())
+				{
+					NewMenuActionInfo.TooltipText = VariableStruct->GetToolTipText(true);
+				}
+
+				const FName NewVariableName = FName(*(NewParameterScopeText.ToString() + RegisteredTypeNameText.ToString()));
+				FNiagaraVariable NewVar = FNiagaraVariable(RegisteredType, NewVariableName);
+				FNiagaraEditorUtilities::ResetVariableToDefaultValue(NewVar);
+				NewMenuActionInfo.NewVariable = NewVar;
+
+				if (NewVar.IsDataInterface())
+				{
+					if (const UClass* DataInterfaceClass = NewVar.GetType().GetClass())
+					{
+						NewMenuActionInfo.bIsExperimental = DataInterfaceClass->GetMetaData("DevelopmentStatus") == TEXT("Experimental");
+					}
+				}
+
+				NewParameterMenuActionInfos.Add(NewMenuActionInfo);
+			}
+		}
+
+		// Sort new action infos to add to the GraphActionMenu in alphabetical order of their display name.
+		NewParameterMenuActionInfos.Sort([](const FNiagaraMenuActionInfo& A, const FNiagaraMenuActionInfo& B)
+			{ return A.DisplayNameText.CompareToCaseIgnored(B.DisplayNameText) < 0; }
+		);
+
+		// Create the actual actions and mark them experimental if so.
+		for (const FNiagaraMenuActionInfo& ActionInfo : NewParameterMenuActionInfos)
+		{
+			TSharedPtr<FNiagaraMenuAction> Action(new FNiagaraMenuAction(ActionInfo.CategoryText, ActionInfo.DisplayNameText, ActionInfo.TooltipText, ActionInfo.Grouping, ActionInfo.KeywordsText,
+				FNiagaraMenuAction::FOnExecuteStackAction::CreateSP(this, &SNiagaraAddParameterMenu2::AddParameterSelected, ActionInfo.NewVariable)));
+
+			Action->IsExperimental = ActionInfo.bIsExperimental;
+
+			OutAllActions.AddAction(Action);
+		}
+	}
+
+
+	// If used in a graph, add all existing graph parameter names to the add menu.
+	if (ShowGraphParameters.Get())
+	{
+		for (TWeakObjectPtr<UNiagaraGraph>& Graph : Graphs)
+		{
+			TMap<FNiagaraVariable, UNiagaraScriptVariable*> GraphParameters = Graph.Get()->GetAllMetaData(); //@todo(ng) try an array view so we can sort it without copying
+			GraphParameters.KeySort([](const FNiagaraVariable& A, const FNiagaraVariable& B) { return A.GetName().LexicalLess(B.GetName()); }); //@todo(ng) sort by scope
+
+			for (TTuple<FNiagaraVariable, UNiagaraScriptVariable*>& GraphParameter : GraphParameters)
+			{
+				const FText Category = LOCTEXT("NiagaraAddExistingParameterMenu", "Add Existing Parameter");
+				const FText DisplayName = FText::FromName(GraphParameter.Key.GetName());
+
+				const FText Tooltip = GraphParameter.Value->Metadata.Description;
+				FNiagaraVariable Variable = GraphParameter.Key;
+				TSharedPtr<FNiagaraMenuAction> Action(new FNiagaraMenuAction(
+					Category, DisplayName, Tooltip, 0, FText::GetEmpty(),
+					FNiagaraMenuAction::FOnExecuteStackAction::CreateSP(this, &SNiagaraAddParameterMenu2::AddParameterSelected, Variable)));
+
+				OutAllActions.AddAction(Action);
+			}
+		}
+	}
+}
+
+void SNiagaraAddParameterMenu2::AddParameterSelected(FNiagaraVariable NewVariable)
+{
+	OnAddParameter.ExecuteIfBound(NewVariable);
+}
+
+class SNiagaraActionMenuExpander2 : public SExpanderArrow
+{
+	SLATE_BEGIN_ARGS(SNiagaraActionMenuExpander2) {}
+	SLATE_ATTRIBUTE(float, IndentAmount)
+		SLATE_END_ARGS()
+
+public:
+	void Construct(const FArguments& InArgs, const FCustomExpanderData& ActionMenuData)
+	{
+		OwnerRowPtr = ActionMenuData.TableRow;
+		IndentAmount = InArgs._IndentAmount;
+		if (!ActionMenuData.RowAction.IsValid())
+		{
+			SExpanderArrow::FArguments SuperArgs;
+			SuperArgs._IndentAmount = InArgs._IndentAmount;
+
+			SExpanderArrow::Construct(SuperArgs, ActionMenuData.TableRow);
+		}
+		else
+		{
+			ChildSlot
+				.Padding(TAttribute<FMargin>(this, &SNiagaraActionMenuExpander2::GetCustomIndentPadding))
+				[
+					SNew(SBox)
+				];
+		}
+	}
+
+private:
+	FMargin GetCustomIndentPadding() const
+	{
+		return SExpanderArrow::GetExpanderPadding();
+	}
+};
+
+TSharedRef<SExpanderArrow> SNiagaraParameterPanel::CreateCustomActionExpander(const FCustomExpanderData& ActionMenuData)
+{
+	return SNew(SNiagaraActionMenuExpander2, ActionMenuData);
+}
+
+void SNiagaraParameterPanel::SetVariablesObjectSelection(const TSharedRef<FNiagaraObjectSelection>& InVariablesObjectSelection)
+{
+	SelectedVariableObjects = InVariablesObjectSelection;
+}
+
+#undef LOCTEXT_NAMESPACE // "NiagaraParameterPanel"
diff --git a/Engine/Plugins/FX/Niagara/Source/NiagaraEditor/Private/Widgets/SNiagaraParameterPanel.h b/Engine/Plugins/FX/Niagara/Source/NiagaraEditor/Private/Widgets/SNiagaraParameterPanel.h
new file mode 100644
--- /dev/null
+++ b/Engine/Plugins/FX/Niagara/Source/NiagaraEditor/Private/Widgets/SNiagaraParameterPanel.h
@@ -0,0 +1,177 @@
+// Copyright Epic Games, Inc. All Rights Reserved.
+
+#pragma once
+
+#include "Widgets/SCompoundWidget.h"
+#include "Internationalization/Text.h"
+#include "NiagaraTypes.h"
+#include "NiagaraActions.h"
+
+class SGraphActionMenu;
+class SEditableTextBox;
+class SExpanderArrow;
+class SSearchBox;
+class SComboButton;
+class SNiagaraGraphPinAdd;
+class FNiagaraObjectSelection;
+class UNiagaraGraph;
+struct FEdGraphSchemaAction; //@todo(ng) cleanup
+class FUICommandList;
+class INiagaraParameterPanelViewModel;
+struct FNiagaraVariableMetaData;
+
+class FNiagaraParameterPanelCommands : public TCommands<FNiagaraParameterPanelCommands>
+{
+public:
+	/** Constructor */
+	FNiagaraParameterPanelCommands()
+		: TCommands<FNiagaraParameterPanelCommands>(TEXT("NiagaraParameterMapViewCommands"), NSLOCTEXT("Contexts", "NiagaraParameterPanel", "NiagaraParameterPanel"), NAME_None, FEditorStyle::GetStyleSetName())
+	{
+	}
+
+	// Basic operations
+	TSharedPtr<FUICommandInfo> DeleteEntry;
+
+	/** Initialize commands */
+	virtual void RegisterCommands() override;
+};
+
+/** A widget for viewing and editing UNiagaraScriptVariables provided by an INiagaraParameterPanelViewModel */
+class SNiagaraParameterPanel : public SCompoundWidget
+{
+public:
+	enum EToolkitType
+	{
+		SCRIPT,
+		SYSTEM,
+	};
+
+	SLATE_BEGIN_ARGS(SNiagaraParameterPanel)
+	{}
+	SLATE_END_ARGS();
+
+	~SNiagaraParameterPanel();
+
+	void Construct(const FArguments& InArgs, const TSharedPtr<INiagaraParameterPanelViewModel>& InParameterPanelViewModel, const TSharedPtr<FUICommandList>& InToolkitCommands);
+
+	virtual void Tick(const FGeometry& AllottedGeometry, const double InCurrentTime, const float InDeltaTime) override;
+
+	/** Whether the add parameter button should be enabled. */
+// 	bool ParameterAddEnabled() const; //@todo(ng) impl
+
+	void AddParameter(FNiagaraVariable NewVariable, const NiagaraParameterPanelSectionID::Type SectionID);
+
+ 	/** Refreshes the palette items for the graph action menu on the next tick. */
+ 	void Refresh();
+
+	static TSharedRef<SExpanderArrow> CreateCustomActionExpander(const struct FCustomExpanderData& ActionMenuData);
+
+	void SetVariablesObjectSelection(const TSharedRef<FNiagaraObjectSelection>& InVariablesObjectSelection);
+
+private:
+	/** Function to bind to SNiagaraAddParameterMenus to filter types we allow creating */
+	bool AllowMakeType(const FNiagaraTypeDefinition& InType) const;
+
+	// SGraphActionMenu delegates
+	TSharedRef<SWidget> OnCreateWidgetForAction(struct FCreateWidgetForActionData* const InCreateData);
+	void CollectAllActions(FGraphActionListBuilderBase& OutAllActions);
+	void CollectStaticSections(TArray<int32>& StaticSectionIDs) const;
+	FReply OnActionDragged(const TArray<TSharedPtr<FEdGraphSchemaAction>>& InActions, const FPointerEvent& MouseEvent);
+	void OnActionSelected(const TArray<TSharedPtr<FEdGraphSchemaAction>>& InActions, ESelectInfo::Type InSelectionType);
+// 	void OnActionDoubleClicked(const TArray<TSharedPtr<FEdGraphSchemaAction>>& InActions); //@todo(ng) impl
+// 	TSharedPtr<SWidget> OnContextMenuOpening();
+ 	FText OnGetSectionTitle(int32 InSectionID);
+	TSharedRef<SWidget> OnGetSectionWidget(TSharedRef<SWidget> RowWidget, int32 InSectionID);
+	TSharedRef<SWidget> CreateAddToSectionButton(const NiagaraParameterPanelSectionID::Type InSection, TWeakPtr<SWidget> WeakRowWidget, FText AddNewText, FName MetaDataTag);
+// 
+// 	/** Checks if the selected action has context menu */
+// 	bool SelectionHasContextMenu() const; //@todo(ng) impl
+// 
+	TSharedRef<SWidget> OnGetParameterMenu(const NiagaraParameterPanelSectionID::Type InSection);
+	
+	//EVisibility OnAddButtonTextVisibility(TWeakPtr<SWidget> RowWidget, const NiagaraParameterMapSectionID::Type InSection) const; //@todo(ng) impl
+
+
+// 	//Callbacks
+	/** Try to delete all entries that are selected. If only some selected entries pass CanDeleteEntry(), then only those selected entries will be deleted. */
+	void TryDeleteEntries();
+	/** Test if ALL selected entries cannot be deleted. If only one selected entry passes INiagaraParameterPanelViewModel::CanRemoveParameter(), return true. Per entry deletion handling is done in TryDeleteEntries(). */
+	bool CanTryDeleteEntries() const;
+	void OnRequestRenameOnActionNode();
+	bool CanRequestRenameOnActionNode(TWeakPtr<struct FGraphActionNode> InSelectedNode) const;
+	bool CanRequestRenameOnActionNode() const;
+	void HandlePaletteItemParameterRenamed(const FText& InText, const FNiagaraScriptVarAndViewInfoAction& InAction);
+	void HandlePaletteItemScopeComboSelectionChanged(ENiagaraParameterScope InScope, const FNiagaraScriptVarAndViewInfoAction& InAction);
+
+	/** Delegate handler used to match an FName to an action in the list, used for renaming keys */
+	bool HandleActionMatchesName(struct FEdGraphSchemaAction* InAction, const FName& InName) const;
+
+private:
+	/** Graph Action Menu for displaying all our variables and functions */
+	TSharedPtr<SGraphActionMenu> GraphActionMenu;
+
+	/** The filter box that handles filtering for both graph action menus. */
+	TSharedPtr<SSearchBox> FilterBox;
+ 
+	/** Add parameter buttons for all sections. */
+	TArray<TSharedPtr<SComboButton>> AddParameterButtons;
+
+	TSharedPtr<FNiagaraObjectSelection> SelectedVariableObjects;
+
+	TSharedPtr<FUICommandList> ToolkitCommands; //@todo(ng) add Find And Rename Parameter command
+
+	bool bNeedsRefresh;
+
+	TSharedPtr<INiagaraParameterPanelViewModel> ParameterPanelViewModel;
+};
+
+class SNiagaraAddParameterMenu2 : public SCompoundWidget
+{
+public:
+	/** Delegate that can be used to create a widget for a particular action */
+	DECLARE_DELEGATE_OneParam(FOnAddParameter, FNiagaraVariable);
+	DECLARE_DELEGATE_TwoParams(FOnCollectCustomActions, FGraphActionListBuilderBase&, bool&);
+	DECLARE_DELEGATE_RetVal_OneParam(bool, FOnAllowMakeType, const FNiagaraTypeDefinition&);
+
+	SLATE_BEGIN_ARGS(SNiagaraAddParameterMenu2)
+		: _AllowCreatingNew(true)
+		, _ShowGraphParameters(true)
+		, _AutoExpandMenu(false)
+		, _IsParameterRead(true) {}
+		SLATE_EVENT(FOnAddParameter, OnAddParameter)
+		SLATE_EVENT(FOnCollectCustomActions, OnCollectCustomActions)
+		SLATE_EVENT(FOnAllowMakeType, OnAllowMakeType)
+		SLATE_ATTRIBUTE(bool, AllowCreatingNew)
+		SLATE_ATTRIBUTE(bool, ShowGraphParameters)
+		SLATE_ATTRIBUTE(bool, AutoExpandMenu)
+		SLATE_ATTRIBUTE(bool, IsParameterRead)
+	SLATE_END_ARGS();
+
+	void Construct(const FArguments& InArgs, TArray<TWeakObjectPtr<UNiagaraGraph>> InGraphs, ENiagaraParameterScope InNewParameterScope);
+
+	TSharedRef<SEditableTextBox> GetSearchBox();
+
+private:
+	void OnActionSelected(const TArray<TSharedPtr<FEdGraphSchemaAction>>& SelectedActions, ESelectInfo::Type InSelectionType);
+	void CollectAllActions(FGraphActionListBuilderBase& OutAllActions);
+	void AddParameterSelected(FNiagaraVariable NewVariable);
+	void AddParameterAndMetaDataSelected(FNiagaraVariable NewVariable, const FNiagaraVariableMetaData NewVariableMetaData);
+
+	TSharedPtr<SGraphActionMenu> GraphMenu;
+
+	/** Delegate that gets fired when a parameter was added. */
+	FOnAddParameter OnAddParameter;
+	FOnCollectCustomActions OnCollectCustomActions;
+	FOnAllowMakeType OnAllowMakeType;
+
+	TArray<TWeakObjectPtr<UNiagaraGraph>> Graphs;
+
+	TAttribute<bool> AllowCreatingNew;
+	TAttribute<bool> ShowGraphParameters;
+	TAttribute<bool> AutoExpandMenu;
+	TAttribute<bool> IsParameterRead;
+
+	/** Default scope to give to new parameters created through this menu. Used when generating new variable names so that their namespace is correct. */
+	ENiagaraParameterScope NewParameterScope;
+	FText NewParameterScopeText;
+};
diff --git a/Engine/Plugins/FX/Niagara/Source/NiagaraEditor/Private/Widgets/SNiagaraParameterPanelPaletteItem.cpp b/Engine/Plugins/FX/Niagara/Source/NiagaraEditor/Private/Widgets/SNiagaraParameterPanelPaletteItem.cpp
new file mode 100644
--- /dev/null
+++ b/Engine/Plugins/FX/Niagara/Source/NiagaraEditor/Private/Widgets/SNiagaraParameterPanelPaletteItem.cpp
@@ -0,0 +1,52 @@
+// Copyright Epic Games, Inc. All Rights Reserved.
+
+#include "SNiagaraParameterPanelPaletteItem.h"
+#include "NiagaraActions.h"
+#include "TutorialMetaData.h"
+#include "EdGraphSchema_Niagara.h"
+
+#define LOCTEXT_NAMESPACE "NiagaraParameterMapPalleteItem"
+
+void SNiagaraParameterPanelPaletteItem::Construct(const FArguments& InArgs, FCreateWidgetForActionData* const InCreateData, TSharedRef<SWidget> ParameterNameViewWidget)
+{
+	check(InCreateData->Action.IsValid());
+	TSharedPtr<FNiagaraScriptVarAndViewInfoAction> Action = StaticCastSharedPtr<FNiagaraScriptVarAndViewInfoAction>(InCreateData->Action);
+	ActionPtr = InCreateData->Action;
+
+	FTutorialMetaData TagMeta("PaletteItem");
+
+	const FLinearColor TypeColor = UEdGraphSchema_Niagara::GetTypeColor(Action->GetScriptVarType());
+	FSlateBrush const* IconBrush = FEditorStyle::GetBrush(TEXT("Kismet.AllClasses.VariableIcon"));
+	FSlateBrush const* SecondaryBrush = FEditorStyle::GetBrush(TEXT("NoBrush"));
+	FSlateColor        IconColor = FSlateColor(TypeColor);
+	FSlateColor        SecondaryIconColor = IconColor;
+	FText			   IconToolTip = FText::GetEmpty();
+	FString			   IconDocLink, IconDocExcerpt;
+	TSharedRef<SWidget> IconWidget = CreateIconWidget(IconToolTip, IconBrush, IconColor, IconDocLink, IconDocExcerpt, SecondaryBrush, SecondaryIconColor);
+	IconWidget->SetEnabled(true);
+
+	static const FName BoldFontName = FName("Bold");
+	static const FName ItalicFontName = FName("Italic");
+	const FName FontType = ItalicFontName;
+	FSlateFontInfo Font = FCoreStyle::GetDefaultFontStyle(FontType, 10);
+
+	// now, create the actual widget
+	ChildSlot
+	[
+		SNew(SHorizontalBox)
+		.AddMetaData<FTutorialMetaData>(TagMeta)
+		.ToolTipText(Action->GetTooltipDescription())
+		// icon slot
+		+SHorizontalBox::Slot()
+		.AutoWidth()
+		[
+			IconWidget
+		]
+		+SHorizontalBox::Slot()
+		[
+			ParameterNameViewWidget
+		]
+	];
+}
+
+#undef LOCTEXT_NAMESPACE // "SNiagaraParameterPanelPalleteItem"
diff --git a/Engine/Plugins/FX/Niagara/Source/NiagaraEditor/Private/Widgets/SNiagaraParameterPanelPaletteItem.h b/Engine/Plugins/FX/Niagara/Source/NiagaraEditor/Private/Widgets/SNiagaraParameterPanelPaletteItem.h
new file mode 100644
--- /dev/null
+++ b/Engine/Plugins/FX/Niagara/Source/NiagaraEditor/Private/Widgets/SNiagaraParameterPanelPaletteItem.h
@@ -0,0 +1,18 @@
+// Copyright Epic Games, Inc. All Rights Reserved.
+
+#pragma once
+
+#include "NiagaraEditorCommon.h"
+#include "SGraphPalette.h"
+
+struct FCreateWidgetForActionData;
+
+class SNiagaraParameterPanelPaletteItem : public SGraphPaletteItem
+{
+public:
+	SLATE_BEGIN_ARGS(SNiagaraParameterPanelPaletteItem)
+	{}
+	SLATE_END_ARGS();
+
+	void Construct(const FArguments& InArgs, FCreateWidgetForActionData* const InCreateData, TSharedRef<SWidget> ParameterNameViewWidget);
+};
\ No newline at end of file
diff --git a/Engine/Plugins/FX/Niagara/Source/NiagaraEditor/Public/NiagaraActions.h b/Engine/Plugins/FX/Niagara/Source/NiagaraEditor/Public/NiagaraActions.h
--- a/Engine/Plugins/FX/Niagara/Source/NiagaraEditor/Public/NiagaraActions.h
+++ b/Engine/Plugins/FX/Niagara/Source/NiagaraEditor/Public/NiagaraActions.h
@@ -3,11 +3,11 @@
 #pragma once
 
 #include "CoreMinimal.h"
-#include "NiagaraTypes.h"
 #include "NiagaraGraph.h"
 #include "EdGraph/EdGraphSchema.h"
 #include "GraphEditorDragDropAction.h"
 #include "DragAndDrop/DecoratedDragDropOp.h"
+#include "NiagaraEditorCommon.h"
 #include "NiagaraActions.generated.h"
 
 USTRUCT()
@@ -43,12 +43,24 @@ private:
 	FCanExecuteStackAction CanPerformAction;
 };
 
+struct NIAGARAEDITOR_API FNiagaraScriptVarAndViewInfoAction : public FEdGraphSchemaAction
+{
+	FNiagaraScriptVarAndViewInfoAction(const FNiagaraScriptVariableAndViewInfo& InScriptVariableAndViewInfo,
+		FText InNodeCategory, FText InMenuDesc, FText InToolTip, const int32 InGrouping, FText InKeywords, int32 InSectionID = 0);
+
+	const FNiagaraTypeDefinition GetScriptVarType() const { return ScriptVariableAndViewInfo.ScriptVariable.GetType(); };
+
+	FNiagaraScriptVariableAndViewInfo ScriptVariableAndViewInfo;
+};
+
 struct NIAGARAEDITOR_API FNiagaraParameterAction : public FEdGraphSchemaAction
 {
 	FNiagaraParameterAction() {}
 	FNiagaraParameterAction(const FNiagaraVariable& InParameter,
 		const TArray<FNiagaraGraphParameterReferenceCollection>& InReferenceCollection,
 		FText InNodeCategory, FText InMenuDesc, FText InToolTip, const int32 InGrouping, FText InKeywords, int32 InSectionID = 0);
+	FNiagaraParameterAction(const FNiagaraVariable& InParameter, 
+		FText InNodeCategory, FText InMenuDesc, FText InToolTip, const int32 InGrouping, FText InKeywords, int32 InSectionID = 0);
 
 	const FNiagaraVariable& GetParameter() const { return Parameter; }
 
@@ -57,6 +69,12 @@ struct NIAGARAEDITOR_API FNiagaraParameterAction : public FEdGraphSchemaAction
 	TArray<FNiagaraGraphParameterReferenceCollection> ReferenceCollection;
 };
 
+struct NIAGARAEDITOR_API FNiagaraScriptParameterAction : public FEdGraphSchemaAction
+{
+	FNiagaraScriptParameterAction() {}
+	FNiagaraScriptParameterAction(const FNiagaraVariable& InVariable, const FNiagaraVariableMetaData& InVariableMetaData);
+};
+
 class NIAGARAEDITOR_API FNiagaraParameterGraphDragOperation : public FGraphSchemaActionDragDropAction
 {
 public:
diff --git a/Engine/Plugins/FX/Niagara/Source/NiagaraEditor/Public/NiagaraEditorCommon.h b/Engine/Plugins/FX/Niagara/Source/NiagaraEditor/Public/NiagaraEditorCommon.h
--- a/Engine/Plugins/FX/Niagara/Source/NiagaraEditor/Public/NiagaraEditorCommon.h
+++ b/Engine/Plugins/FX/Niagara/Source/NiagaraEditor/Public/NiagaraEditorCommon.h
@@ -2,7 +2,6 @@
 
 #pragma once
 
-#include "CoreMinimal.h"
 #include "NiagaraCommon.h"
 
 NIAGARAEDITOR_API DECLARE_LOG_CATEGORY_EXTERN(LogNiagaraEditor, All, All);
@@ -133,4 +132,92 @@ private:
 	const TSharedPtr<INiagaraScriptGraphFocusInfo> ScriptGraphFocusInfo;
 };
 
+/** Convenience wrapper for generating entries for scope enum combo boxes in SNiagaraParameterNameView. */
+struct FScopeIsEnabledAndTooltip
+{
+	FScopeIsEnabledAndTooltip()
+		: bEnabled(true)
+		, Tooltip()
+	{};
+
+	FScopeIsEnabledAndTooltip(bool bInEnabled, FText InTooltip)
+		: bEnabled(bInEnabled)
+		, Tooltip(InTooltip)
+	{};
+
+	bool bEnabled;
+	FText Tooltip;
+};
+
+/** Helper struct for passing along info on parameters and how to display them in SNiagaraParameterNameView */
+struct FNiagaraScriptVariableAndViewInfo
+{
+	FNiagaraScriptVariableAndViewInfo() {}
+	FNiagaraScriptVariableAndViewInfo(const FNiagaraVariable& InScriptVariable, const FNiagaraVariableMetaData& InScriptVariableMetaData)
+		: bIsSelectionRelevant(false)
+	{
+		ScriptVariable = InScriptVariable;
+		MetaData = InScriptVariableMetaData;
+	};
+
+	FNiagaraScriptVariableAndViewInfo(const FNiagaraVariable& InScriptVariable, const FNiagaraVariableMetaData& InMetaData, const TStaticArray<FScopeIsEnabledAndTooltip, (int32)ENiagaraParameterScope::Num>& InParameterScopeToDisplayInfo)
+		: ParameterScopeToDisplayInfo(InParameterScopeToDisplayInfo)
+		, bIsSelectionRelevant(false)
+	{
+		ScriptVariable = InScriptVariable;
+		MetaData = InMetaData;
+	};
+
+	FNiagaraScriptVariableAndViewInfo(const FNiagaraScriptVariableAndViewInfo& Other)
+		: ScriptVariable(Other.ScriptVariable)
+		, MetaData(Other.MetaData)
+		, ParameterScopeToDisplayInfo(Other.ParameterScopeToDisplayInfo)
+		, bIsSelectionRelevant(Other.bIsSelectionRelevant)
+	{};
+
+	FNiagaraVariable ScriptVariable;
+	FNiagaraVariableMetaData MetaData;
+
+	// Array indexed by ENiagaraParameterScope value containing info on how to present each scope in a ComboBox.
+	TStaticArray<FScopeIsEnabledAndTooltip, (int32)ENiagaraParameterScope::Num> ParameterScopeToDisplayInfo;
 
+	// Whether this entry is related to the current selection state, e.g. if a module is selected in the Stack, mark this entry if ScriptVariable is a member of that module.
+	bool bIsSelectionRelevant;
+};
+
+/* Legacy style Namespace::Type Enum for Parameter Panel section management. */
+namespace NiagaraParameterPanelSectionID
+{
+	enum Type
+	{
+		NONE = 0,
+
+		// Group A, for contextual information when editing scripts
+		INPUTS,
+		REFERENCES,
+		OUTPUTS,
+		LOCALS,
+		INITIALVALUES,
+
+		// Group B, for contextual information when editing systems
+		ENGINE,
+		USER,
+		SYSTEM,
+		EMITTER,
+		PARTICLES,
+	};
+
+	FText OnGetSectionTitle(const NiagaraParameterPanelSectionID::Type InSection);
+	static const NiagaraParameterPanelSectionID::Type GetSectionForScope(ENiagaraParameterScope InScope);
+	static const NiagaraParameterPanelSectionID::Type GetSectionForParameterMetaData(const FNiagaraVariableMetaData& MetaData);
+	ENiagaraParameterScope GetScopeForNewParametersInSection(NiagaraParameterPanelSectionID::Type InSection);
+	ENiagaraScriptParameterUsage GetUsageForNewParametersInSection(NiagaraParameterPanelSectionID::Type InSection);
+
+}
+
+class FNiagaraEditorCommonCVar
+{
+public:
+	static int32 GNiagaraEnableParameterPanel2;
+	static FAutoConsoleVariableRef CVarEnableParameterPanel2;
+};
diff --git a/Engine/Plugins/FX/Niagara/Source/NiagaraEditor/Public/NiagaraEditorModule.h b/Engine/Plugins/FX/Niagara/Source/NiagaraEditor/Public/NiagaraEditorModule.h
--- a/Engine/Plugins/FX/Niagara/Source/NiagaraEditor/Public/NiagaraEditorModule.h
+++ b/Engine/Plugins/FX/Niagara/Source/NiagaraEditor/Public/NiagaraEditorModule.h
@@ -41,6 +41,7 @@ public:
 	virtual TSharedRef<SWidget> CreateStackIssueIcon(UNiagaraStackViewModel& StackViewModel, UNiagaraStackEntry& StackEntry) const = 0;
 	virtual TSharedRef<SWidget> CreateScriptScratchPad(UNiagaraScratchPadViewModel& ScriptScratchPadViewModel) const = 0;
 	virtual FLinearColor GetColorForExecutionCategory(FName ExecutionCategory) const = 0;
+	virtual FLinearColor GetColorForParameterScope(ENiagaraParameterScope ParameterScope) const = 0;
 };
 
 extern int32 GbShowFastPathOptions;
diff --git a/Engine/Plugins/FX/Niagara/Source/NiagaraEditor/Public/NiagaraEditorUtilities.h b/Engine/Plugins/FX/Niagara/Source/NiagaraEditor/Public/NiagaraEditorUtilities.h
--- a/Engine/Plugins/FX/Niagara/Source/NiagaraEditor/Public/NiagaraEditorUtilities.h
+++ b/Engine/Plugins/FX/Niagara/Source/NiagaraEditor/Public/NiagaraEditorUtilities.h
@@ -242,4 +242,6 @@ namespace FNiagaraEditorUtilities
 	void CreateAssetFromEmitter(TSharedRef<FNiagaraEmitterHandleViewModel> EmitterHandleViewModel);
 
 	NIAGARAEDITOR_API void WarnWithToastAndLog(FText WarningMessage);
+
+	void GetScriptRunAndExecutionIndexFromUsage(const ENiagaraScriptUsage& InUsage, int32& OutRunIndex, int32&OutExecutionIndex);
 };
diff --git a/Engine/Plugins/FX/Niagara/Source/NiagaraEditor/Public/NiagaraGraph.h b/Engine/Plugins/FX/Niagara/Source/NiagaraEditor/Public/NiagaraGraph.h
--- a/Engine/Plugins/FX/Niagara/Source/NiagaraEditor/Public/NiagaraGraph.h
+++ b/Engine/Plugins/FX/Niagara/Source/NiagaraEditor/Public/NiagaraGraph.h
@@ -19,6 +19,12 @@ enum ENiagaraGraphActionType
 	GRAPHACTION_GenericNeedsRecompile = 0x1 << 16,
 };
 
+struct FInputPinsAndOutputPins
+{
+	TArray<UEdGraphPin*> InputPins;
+	TArray<UEdGraphPin*> OutputPins;
+};
+
 class UNiagaraNode;
 class UNiagaraGraph;
 
@@ -102,6 +108,7 @@ class UNiagaraGraph : public UEdGraph
 	GENERATED_UCLASS_BODY()
 
 	DECLARE_MULTICAST_DELEGATE(FOnDataInterfaceChanged);
+	DECLARE_DELEGATE_RetVal_OneParam(TSharedRef<SWidget>, FOnGetPinVisualWidget, const UEdGraphPin*);
 
 	//~ Begin UObject Interface
 	virtual void PostLoad() override;
@@ -158,6 +165,32 @@ class UNiagaraGraph : public UEdGraph
 		FGuid TargetScriptUsageId;
 	};
 
+	/** Options for the AddParameter function. */
+	struct FAddParameterOptions
+	{
+		FAddParameterOptions()
+			: NewParameterScope(TOptional<ENiagaraParameterScope>())
+			, NewParameterUsage(TOptional<ENiagaraScriptParameterUsage>())
+			, bRefreshMetaDataScopeAndUsage(false)
+			, bIsStaticSwitch(false)
+			, bMakeParameterNameUnique(false)
+			, bAddedFromSystemEditor(false)
+		{};
+
+		/** Optional scope to assign to the new parameter. New scope is force assigned and will ignore any scope that would be assigned if bRefreshMetaDataScopeAndUsage is true. */
+		TOptional<ENiagaraParameterScope> NewParameterScope;
+		/** Optional usage to assign to the new parameter. New usage is force assigned and will ignore any usage that would be assigned if bRefreshMetaDataScopeAndUsage is true.*/
+		TOptional<ENiagaraScriptParameterUsage> NewParameterUsage;
+		/** Whether the new or already existing parameter should refresh its scope and usage. If no scope is assigned, assign a scope from the variable name. Set a new usage depending on associated in/out pins in the graph. */
+		bool bRefreshMetaDataScopeAndUsage;
+		/** Whether the new parameter should be a static switch parameter. */
+		bool bIsStaticSwitch;
+		/** Whether the new parameter should have a unique name (append 001 onwards if new name is an alias.) */
+		bool bMakeParameterNameUnique;
+		/** Whether the new parameter was added from the System Toolkit. Will set bCreatedInSystemEditor on the new parameter's metadata. */
+		bool bAddedFromSystemEditor;
+	};
+
 	/** Finds input nodes in the graph with. */
 	void FindInputNodes(TArray<class UNiagaraNodeInput*>& OutInputNodes, FFindInputNodeOptions Options = FFindInputNodeOptions()) const;
 
@@ -224,9 +257,13 @@ class UNiagaraGraph : public UEdGraph
 	/** Gets the meta-data associated with this variable, if it exists.*/
 	TOptional<FNiagaraVariableMetaData> GetMetaData(const FNiagaraVariable& InVar) const;
 
-	/** Sets the meta-data associated with this variable.*/
+	/** Sets the meta-data associated with this variable. Creates a new UNiagaraScriptVariable if the target variable cannot be found. Illegal to call on FNiagaraVariables that are Niagara Constants. */
 	void SetMetaData(const FNiagaraVariable& InVar, const FNiagaraVariableMetaData& MetaData);
 
+	/** Sets the usage metadata associated with this variable for a given script. Creates a new UNiagaraScriptVariable if the target variable cannot be found. */
+	void SetPerScriptMetaData(const FNiagaraVariable& InVar, const FNiagaraVariableMetaData& InMetaData);
+
+
 	const TMap<FNiagaraVariable, UNiagaraScriptVariable*>& GetAllMetaData() const;
 	TMap<FNiagaraVariable, UNiagaraScriptVariable*>& GetAllMetaData();
 
@@ -238,14 +275,32 @@ class UNiagaraGraph : public UEdGraph
 	/** Adds parameter to parameters map setting it as created by the user.*/
 	void AddParameter(const FNiagaraVariable& Parameter, bool bIsStaticSwitch = false);
 
-	/** Adds parameter to parameters map setting it as created by the user.*/
-	void AddParameterReference(const FNiagaraVariable& Parameter, const UEdGraphPin* Pin);
+	void AddParameter(FNiagaraVariable& Parameter, const FAddParameterOptions Options);
+
+	/** Adds an FNiagaraGraphParameterReference to the ParameterToReferenceMap. */
+	void AddParameterReference(const FNiagaraVariable& Parameter, FNiagaraGraphParameterReference& NewParameterReference);
 
 	/** Remove parameter from map and all the pins associated. */
 	void RemoveParameter(const FNiagaraVariable& Parameter);
 
-	/** Rename parameter from map and all the pins associated. */
-	bool RenameParameter(const FNiagaraVariable& Parameter, FName NewName, bool bFromStaticSwitch = false);
+	/**
+	* Rename parameter from map and all the pins associated.
+	*
+	* @param Parameter				The target FNiagaraVariable key to lookup the canonical UNiagaraScriptVariable to rename.
+	* @param NewName				The new name to apply.
+	* @param bFromStaticSwitch		Whether the target parameter is from a static switch. Used to determine whether to fixup binding strings.
+	* @param NewScope				The new scope from the FNiagaraVariable's associated FNiagaraVariableMetaData. Used if changing the scope triggered a rename. Applied to the canonical UNiagaraScriptVariable's Metadata.
+	* @return						true if the new name was applied. 
+	*/
+	bool RenameParameter(const FNiagaraVariable& Parameter, FName NewName, bool bFromStaticSwitch = false, ENiagaraParameterScope NewScope = ENiagaraParameterScope::None);
+
+	/** 
+	 * Sets the ENiagaraScriptParameterUsage on the Metadata of a ScriptVariable depending on its pin usage in the node graph (used on input/output pins). 
+	 * 
+	 * @param ScriptVariable		The ScriptVariable to update the Metadata of.
+	 * @return						false if the ScriptVariable no longer has any referencing Map Get or Set pins in its owning Graph.
+	*/
+	bool UpdateUsageForScriptVariable(UNiagaraScriptVariable* ScriptVariable) const;
 
 	/** Gets a delegate which is called whenever a contained data interfaces changes. */
 	FOnDataInterfaceChanged& OnDataInterfaceChanged();
@@ -274,6 +329,17 @@ class UNiagaraGraph : public UEdGraph
 
 	void CopyCachedReferencesMap(UNiagaraGraph* TargetGraph);
 
+	void FixupParameterUsageMetaData();
+
+	bool IsPinVisualWidgetProviderRegistered() const;
+
+	/** Wrapper to provide visual widgets for pins that are managed by external viewmodels.
+	 *  Used with FNiagaraScriptToolkitParameterPanelViewModel to provide SNiagaraParameterNameView widget with validation delegates.
+	 */
+	TSharedRef<SWidget> GetPinVisualWidget(const UEdGraphPin* Pin) const;
+
+	FDelegateHandle RegisterPinVisualWidgetProvider(FOnGetPinVisualWidget OnGetPinVisualWidget);
+
 protected:
 	void RebuildNumericCache();
 	bool bNeedNumericCacheRebuilt;
@@ -290,6 +356,20 @@ private:
 	/** Marks the found parameter collections as invalid so they're rebuilt the next time they're requested. */
 	void InvalidateCachedParameterData();
 
+	/** When a new variable is added to the VariableToScriptVariableMap, generate appropriate scope and usage. */
+	void GenerateMetaDataForScriptVariable(UNiagaraScriptVariable* InScriptVariable) const;
+
+	/** Helper to get a map of variables to all input/output pins with the same name. */
+	const TMap<FNiagaraVariable, FInputPinsAndOutputPins> CollectVarsToInOutPinsMap() const;
+
+	/**
+	 * Set the usage of a script variable depending on input/output pins with same name.
+	 * @param VarToPinsMap			Mapping of Pins to the associated UNiagaraScriptVariable.
+	 * @param ScriptVariable		The UNiagaraScriptVariable to modify its Metadata Usage.
+	 * @return ShouldDelete			True if the ScriptVariable was previously of a non-local type and no longer has any associated Input or Output pins.
+	 */
+	bool SetScriptVariableUsageForPins(const TMap<FNiagaraVariable, FInputPinsAndOutputPins>& VarToPinsMap, UNiagaraScriptVariable* ScriptVariable) const;
+
 	/** A delegate that broadcasts a notification whenever the graph needs recompile due to structural change. */
 	FOnGraphChanged OnGraphNeedsRecompile;
 
@@ -327,5 +407,6 @@ private:
 	bool bIsRenamingParameter;
 
 	mutable bool bParameterReferenceRefreshPending;
-};
 
+	FOnGetPinVisualWidget OnGetPinVisualWidgetDelegate;
+};
diff --git a/Engine/Plugins/FX/Niagara/Source/NiagaraEditor/Public/NiagaraNodeWithDynamicPins.h b/Engine/Plugins/FX/Niagara/Source/NiagaraEditor/Public/NiagaraNodeWithDynamicPins.h
--- a/Engine/Plugins/FX/Niagara/Source/NiagaraEditor/Public/NiagaraNodeWithDynamicPins.h
+++ b/Engine/Plugins/FX/Niagara/Source/NiagaraEditor/Public/NiagaraNodeWithDynamicPins.h
@@ -41,14 +41,19 @@ public:
 	virtual void CollectAddPinActions(FGraphActionListBuilderBase& OutActions, bool& bOutCreateRemainingActions, UEdGraphPin* Pin);
 
 	/** Request a new pin. */
-	void AddParameter(FNiagaraVariable Parameter, UEdGraphPin* AddPin);
+	void AddParameter(FNiagaraVariable Parameter, const UEdGraphPin* AddPin);
+	void AddParameter(FNiagaraVariable Parameter, const struct UNiagaraGraph::FAddParameterOptions AddParameterOptions);
 
+	/** Request a new pin and assign MetaData to the newly generated UNiagaraScriptVariable. */
+//	void AddParameterAndSetMetaData(FNiagaraVariable& Parameter, const FNiagaraVariableMetaData& ParameterMetaData, UEdGraphPin* AddPin);
 protected:
 	virtual bool AllowDynamicPins() const { return true; }
 
 	/** Creates an add pin on the node for the specified direction. */
 	void CreateAddPin(EEdGraphPinDirection Direction);
 
+	void UpdateAddedPinMetaData(const UEdGraphPin* AddedPin);
+
 	/** Called when a new typed pin is added by the user. */
 	virtual void OnNewTypedPinAdded(UEdGraphPin* NewPin) { }
 	
@@ -75,6 +80,9 @@ protected:
 
 	virtual bool OnVerifyTextChanged(const FText& NewText, FText& OutMessage) { return true; };
 
+	/** Convenience method to determine whether this Node is a Map Get or Map Set when adding a parameter through the parameter panel. */
+	virtual EEdGraphPinDirection GetPinDirectionForNewParameters() { return EEdGraphPinDirection::EGPD_MAX; };
+
 private:
 
 	/** Gets the display text for a pin. */
diff --git a/Engine/Plugins/FX/Niagara/Source/NiagaraEditor/Public/NiagaraParameterMapHistory.h b/Engine/Plugins/FX/Niagara/Source/NiagaraEditor/Public/NiagaraParameterMapHistory.h
--- a/Engine/Plugins/FX/Niagara/Source/NiagaraEditor/Public/NiagaraParameterMapHistory.h
+++ b/Engine/Plugins/FX/Niagara/Source/NiagaraEditor/Public/NiagaraParameterMapHistory.h
@@ -35,9 +35,14 @@ struct FNiagaraParameterMapHistory
 public:
 	FNiagaraParameterMapHistory();
 
+	ENiagaraScriptUsage OriginatingScriptUsage;
+
 	/** The variables that have been identified during the traversal. */
 	TArray<FNiagaraVariable> Variables;
 
+	/** The metadata associated with each variable identified during the traversal. Only gathered by FNiagaraParameterMapHistoryWithMetaDataBuilder. */
+	TArray<FNiagaraVariableMetaData> VariableMetaData;
+
 	TArray<FNiagaraVariable> VariablesWithOriginalAliasesIntact;
 
 	/** Used parameter collections identified during the traversal. TODO: Need to ensure these cannot be GCd if the asset is deleted while it's being used in an in flight compilation. */
@@ -201,7 +206,7 @@ public:
 	/**
 	* Helper to add a variable to the known list for a parameter map.
 	*/
-	int32 AddVariable(const FNiagaraVariable& InVar, const FNiagaraVariable& InAliasedVar, const UEdGraphPin* InPin);
+	int32 AddVariable(const FNiagaraVariable& InVar, const FNiagaraVariable& InAliasedVar, const UEdGraphPin* InPin, TOptional<FNiagaraVariableMetaData> InMetaData = TOptional<FNiagaraVariableMetaData>());
 
 	/** Get the default value for this variable.*/
 	const UEdGraphPin* GetDefaultValuePin(int32 VarIdx) const;
@@ -218,6 +223,71 @@ public:
 	bool ShouldIgnoreVariableDefault(const FNiagaraVariable& Var)const;
 };
 
+// helper struct for handling tracked variable metadata and aliases from histories.
+struct FNiagaraHistoryVariable
+{
+	FNiagaraHistoryVariable(FNiagaraVariable* InVar, const FNiagaraVariable* InVarWithOriginalAlias, const FNiagaraVariableMetaData* InVarMetaData)
+		: Var(InVar)
+		, VarWithOriginalAlias(InVarWithOriginalAlias)
+		, VarMetaData(InVarMetaData)
+		, bDoNotDemote(false)
+	{};
+
+	FNiagaraVariable* Var;
+	const FNiagaraVariable* VarWithOriginalAlias;
+	const FNiagaraVariableMetaData* VarMetaData;
+
+	// Used to track whether this history variable should not be demoted during pre-compilation.
+	bool bDoNotDemote;
+
+	// Use carefully, implemented to fixup names of module namespace parameters that alias with other parameters in same scope.
+	void SetVarWithOriginalAliasName(const FName& NewName);
+};
+
+struct FNiagaraParameterMapHistoryHandle
+{
+public:
+	FNiagaraParameterMapHistoryHandle();
+
+	FNiagaraParameterMapHistoryHandle(FNiagaraParameterMapHistory& InHistory, const TArray<FNiagaraVariable>& InRequiredRendererVariables, const FName InEmitterUniqueName = FName());
+
+	FNiagaraParameterMapHistoryHandle(FNiagaraParameterMapHistory& InPrimaryHistory, FNiagaraParameterMapHistory& InSecondaryHistory, const TArray<FNiagaraVariable>& InRequiredRendererVariables, const FName InEmitterUniqueName = FName());
+
+	const int32 GetScriptExecutionIndex() const { return ScriptExecutionIndex; };
+	const int32 GetScriptRunIndex() const { return ScriptRunIndex; };
+	const ENiagaraParameterScope GetScriptScope() const { return Scope; };
+	const ENiagaraScriptUsage GetOriginatingScriptUsage() const { return History->OriginatingScriptUsage; };
+
+	TArray<FNiagaraHistoryVariable> HistoryVariables;
+
+	void SetScriptExecutionIndex(int32& InIndex) { ScriptExecutionIndex = InIndex; };
+	void SetScriptRunIndex(int32& InIndex) { ScriptRunIndex = InIndex; };
+	void SetScriptScope(ENiagaraParameterScope InScope) { Scope = InScope; };
+
+	void FixupWrittenHistoryPinNames(int32& VarIdx, const FName& NewName);
+	void FixupWrittenHistoryPinNames(const FNiagaraVariable* AssociatedVar, const FName& NewName);
+	void FixupReadHistoryPinNames(int32& VarIdx, const FName& NewName);
+	void FixupReadHistoryPinNames(const FNiagaraVariable* AssociatedVar, const FName& NewName);
+
+	const int32 GetSourceVariableHistoryIndex(const FNiagaraVariable* InVar) const { return History->VariablesWithOriginalAliasesIntact.IndexOfByPredicate([InVar](const FNiagaraVariable& Var) {return InVar == &Var; }); };
+
+	const TArray<FNiagaraVariable>& GetRequiredRendererVariables() const { return RequiredRendererVariables; };
+	const FName& GetEmitterUniqueName() const { return EmitterUniqueName; };
+	const bool IsSystemHistoryHandle() const { return EmitterUniqueName.IsNone(); };
+
+private:
+	/** Records a number associated with the required order of execution of the script that has been visited. Used to determine if parameters will be accessed from prior or proceeding scripts. */ 
+	int32 ScriptExecutionIndex;
+	int32 ScriptRunIndex;
+
+	ENiagaraParameterScope Scope;
+
+	FNiagaraParameterMapHistory* History;
+
+	TArray<FNiagaraVariable> RequiredRendererVariables;
+
+	FName EmitterUniqueName;
+};
 
 class FNiagaraParameterMapHistoryBuilder
 {
@@ -228,6 +298,8 @@ public:
 	/** Constructor*/
 	FNiagaraParameterMapHistoryBuilder();
 
+	virtual ~FNiagaraParameterMapHistoryBuilder() {};
+
 	/** Add a new parameter map to the array.*/
 	int32 CreateParameterMap();
 
@@ -262,7 +334,7 @@ public:
 	/**
 	* Record that we have entered a new function scope.
 	*/
-	void EnterFunction(const FString& InNodeName, const class UNiagaraScript* InScript, const class UNiagaraNode* Node);
+	void EnterFunction(const FString& InNodeName, const class UNiagaraScript* InScript, const class UNiagaraGraph* InGraph, const class UNiagaraNode* Node);
 
 	/**
 	* Record that we have exited a function scope.
@@ -272,7 +344,7 @@ public:
 	/**
 	* Record that we have entered an emitter scope.
 	*/
-	void EnterEmitter(const FString& InEmitterName, const class UNiagaraNode* Node);
+	void EnterEmitter(const FString& InEmitterName, const class UNiagaraGraph* InGraph, const class UNiagaraNode* Node);
 
 	/**
 	* Record that we have exited an emitter scope.
@@ -316,6 +388,11 @@ public:
 	int32 HandleVariableRead(int32 ParameterMapIndex, const UEdGraphPin* InPin, bool RegisterReadsAsVariables, const UEdGraphPin* InDefaultPin, bool bFilterForCompilation, bool& OutUsedDefault);
 
 	int32 HandleExternalVariableRead(int32 ParamMapIdx, const FName& InVarName);
+
+	/**
+	* Virtual wrapper to call correct AddVariable method on the target ParameterMapHistory.
+	*/
+	virtual int32 AddVariableToHistory(FNiagaraParameterMapHistory& History, const FNiagaraVariable& InVar, const FNiagaraVariable& InAliasedVar, const UEdGraphPin* InPin);
 	
 	/**
 	* Get the string that the "Module" namespace maps to currently (if it exists)
@@ -358,6 +435,8 @@ public:
 
 	FCompileConstantResolver ConstantResolver;
 
+	bool bShouldBuildSubHistories = true;
+
 protected:
 	/**
 	* Generate the internal alias map from the current traversal state.
@@ -372,6 +451,10 @@ protected:
 
 	/** Contains the hierarchy of nodes leading to the current graph being processed. Usually made up of FunctionCall and Emitter nodes.*/
 	TArray<const UNiagaraNode*> CallingContext;
+
+	/** Contains the hierarchy of graphs leading to and including the current graph being processed. Is not in sync with CallingContext as there is an additional 0th entry for the NodeGraphDeepCopy. */
+	TArray<const UNiagaraGraph*> CallingGraphContext;
+
 	/** Tracker for each context level of the parameter map index associated with a given pin. Used to trace parameter maps through the graph.*/
 	TArray<TMap<const UEdGraphPin*, int32> > PinToParameterMapIndices;
 	/** List of previously visited nodes per context. Note that the same node may be visited multiple times across all graph traversals, but only one time per context level.*/
@@ -397,4 +480,19 @@ protected:
 	bool bIgnoreDisabled;
 
 	TArray<FNiagaraVariable> EncounterableExternalVariables;
-};
\ No newline at end of file
+};
+
+class FNiagaraParameterMapHistoryWithMetaDataBuilder : public FNiagaraParameterMapHistoryBuilder
+{
+public:
+	/** Iterate over histories and fix variable namespaces so the translator can link parameters across script executions. */
+	static void FixupHistoryVariableNamespaces(TArray<FNiagaraParameterMapHistoryHandle>& InHistoryHandles, bool bCullDatasetOutputs);
+
+	//** Add a graph to the calling graph context stack. Generally used to prime the context stack with the node graph deep copy during precompilation. */
+	void AddGraphToCallingGraphContextStack(const UNiagaraGraph* InGraph) { CallingGraphContext.Add(InGraph); };
+
+	/**
+	* Virtual wrapper to call correct AddVariable method on the target ParameterMapHistory.
+	*/
+	virtual int32 AddVariableToHistory(FNiagaraParameterMapHistory& History, const FNiagaraVariable& InVar, const FNiagaraVariable& InAliasedVar, const UEdGraphPin* InPin);
+};
diff --git a/Engine/Plugins/FX/Niagara/Source/NiagaraEditor/Public/ViewModels/Stack/NiagaraParameterHandle.h b/Engine/Plugins/FX/Niagara/Source/NiagaraEditor/Public/ViewModels/Stack/NiagaraParameterHandle.h
--- a/Engine/Plugins/FX/Niagara/Source/NiagaraEditor/Public/ViewModels/Stack/NiagaraParameterHandle.h
+++ b/Engine/Plugins/FX/Niagara/Source/NiagaraEditor/Public/ViewModels/Stack/NiagaraParameterHandle.h
@@ -68,6 +68,8 @@ public:
 	static const FName ModuleNamespace;
 	static const FName ParameterCollectionNamespace;
 	static const FString InitialPrefix;
+	static const FName LocalNamespace;
+	static const FName InitialNamespace;
 
 private:
 	FName ParameterHandleName;
diff --git a/Engine/Plugins/FX/Niagara/Source/NiagaraEditor/Public/ViewModels/Stack/NiagaraStackGraphUtilities.h b/Engine/Plugins/FX/Niagara/Source/NiagaraEditor/Public/ViewModels/Stack/NiagaraStackGraphUtilities.h
--- a/Engine/Plugins/FX/Niagara/Source/NiagaraEditor/Public/ViewModels/Stack/NiagaraStackGraphUtilities.h
+++ b/Engine/Plugins/FX/Niagara/Source/NiagaraEditor/Public/ViewModels/Stack/NiagaraStackGraphUtilities.h
@@ -148,7 +148,8 @@ namespace FNiagaraStackGraphUtilities
 	void GetAvailableParametersForScript(UNiagaraNodeOutput& ScriptOutputNode, TArray<FNiagaraVariable>& OutAvailableParameters);
 
 	TOptional<FName> GetNamespaceForScriptUsage(ENiagaraScriptUsage ScriptUsage);
-
+	ENiagaraParameterScope GetScopeForScriptUsage(ENiagaraScriptUsage ScriptUsage);
+	
 	bool IsValidDefaultDynamicInput(UNiagaraScript& OwningScript, UEdGraphPin& DefaultPin);
 
 	bool CanWriteParameterFromUsage(FNiagaraVariable Parameter, ENiagaraScriptUsage Usage);
@@ -188,4 +189,14 @@ namespace FNiagaraStackGraphUtilities
 
 	void GatherRenamedStackFunctionOutputVariableNames(UNiagaraEmitter* Emitter, UNiagaraNodeFunctionCall& FunctionCallNode, const FString& OldFunctionName, const FString& NewFunctionName, TMap<FName, FName>& OutOldToNewNameMap);
 	void GatherRenamedStackFunctionInputAndOutputVariableNames(UNiagaraEmitter* Emitter, UNiagaraNodeFunctionCall& FunctionCallNode, const FString& OldFunctionName, const FString& NewFunctionName, TMap<FName, FName>& OutOldToNewNameMap);
+	FString GetNamespaceStringForScriptParameterScope(const ENiagaraParameterScope& InScope);
+	FString GetNamespaceStringForScriptParameterMetaData(const FNiagaraVariableMetaData& InMetaData);
+
+	TArray<FName> DecomposeVariableNamespace(const FName& InVarName, FName& OutFriendlyName);
+
+	void SetParameterMetaDataFromName(const FName& InVarName, ENiagaraParameterScope& OutScope, ENiagaraScriptParameterUsage& OutUsage, FName& OutFriendlyName);
+
+	FString GetNamespacelessVariableNameString(const FName& InVarName);
+
+	FName GetVariableNameForScope(const FName& CurrentVariableName, const ENiagaraParameterScope NewVariableScope, bool bIsInitialValue);
 }
diff --git a/Engine/Plugins/FX/Niagara/Source/NiagaraEditorWidgets/Private/NiagaraEditorWidgetsModule.cpp b/Engine/Plugins/FX/Niagara/Source/NiagaraEditorWidgets/Private/NiagaraEditorWidgetsModule.cpp
--- a/Engine/Plugins/FX/Niagara/Source/NiagaraEditorWidgets/Private/NiagaraEditorWidgetsModule.cpp
+++ b/Engine/Plugins/FX/Niagara/Source/NiagaraEditorWidgets/Private/NiagaraEditorWidgetsModule.cpp
@@ -197,3 +197,8 @@ FLinearColor FNiagaraEditorWidgetsModule::FNiagaraEditorWidgetProvider::GetColor
 {
 	return FNiagaraEditorWidgetsStyle::Get().GetColor(FNiagaraStackEditorWidgetsUtilities::GetIconColorNameForExecutionCategory(ExecutionCategory));
 }
+
+FLinearColor FNiagaraEditorWidgetsModule::FNiagaraEditorWidgetProvider::GetColorForParameterScope(ENiagaraParameterScope ParameterScope) const
+{
+	return FNiagaraEditorWidgetsStyle::Get().GetColor(FNiagaraStackEditorWidgetsUtilities::GetColorNameForParameterScope(ParameterScope));
+}
diff --git a/Engine/Plugins/FX/Niagara/Source/NiagaraEditorWidgets/Private/NiagaraEditorWidgetsStyle.cpp b/Engine/Plugins/FX/Niagara/Source/NiagaraEditorWidgets/Private/NiagaraEditorWidgetsStyle.cpp
--- a/Engine/Plugins/FX/Niagara/Source/NiagaraEditorWidgets/Private/NiagaraEditorWidgetsStyle.cpp
+++ b/Engine/Plugins/FX/Niagara/Source/NiagaraEditorWidgets/Private/NiagaraEditorWidgetsStyle.cpp
@@ -250,6 +250,13 @@ TSharedRef< FSlateStyleSet > FNiagaraEditorWidgetsStyle::Create()
 
 	Style->Set("NiagaraEditor.ScratchPad.HeaderColor", FLinearColor(FColor(48, 48, 48)));
 
+	Style->Set("NiagaraEditor.Scope.Engine", FLinearColor(FColor(230, 102, 102)));
+	Style->Set("NiagaraEditor.Scope.User", FLinearColor(FColor(114, 226, 254)));
+	Style->Set("NiagaraEditor.Scope.System", FLinearColor(FColor(1, 202, 252)));
+	Style->Set("NiagaraEditor.Scope.Emitter", FLinearColor(FColor(241, 99, 6)));
+	Style->Set("NiagaraEditor.Scope.Particles", FLinearColor(FColor(131, 218, 9)));
+	Style->Set("NiagaraEditor.Scope.ScriptPersistent", FLinearColor(FColor(255, 247, 77)));
+	Style->Set("NiagaraEditor.Scope.ScriptTransient", FLinearColor(FColor(255, 247, 77)));
 
 	return Style;
 }
diff --git a/Engine/Plugins/FX/Niagara/Source/NiagaraEditorWidgets/Private/NiagaraEditorWidgetsUtilities.cpp b/Engine/Plugins/FX/Niagara/Source/NiagaraEditorWidgets/Private/NiagaraEditorWidgetsUtilities.cpp
--- a/Engine/Plugins/FX/Niagara/Source/NiagaraEditorWidgets/Private/NiagaraEditorWidgetsUtilities.cpp
+++ b/Engine/Plugins/FX/Niagara/Source/NiagaraEditorWidgets/Private/NiagaraEditorWidgetsUtilities.cpp
@@ -128,6 +128,28 @@ FName FNiagaraStackEditorWidgetsUtilities::GetIconColorNameForExecutionCategory(
 	}
 }
 
+FName FNiagaraStackEditorWidgetsUtilities::GetColorNameForParameterScope(ENiagaraParameterScope ParameterScope)
+{
+	switch (ParameterScope) {
+	case ENiagaraParameterScope::Engine:
+		return "NiagaraEditor.Scope.Engine";
+	case ENiagaraParameterScope::User:
+		return "NiagaraEditor.Scope.User";
+	case ENiagaraParameterScope::System:
+		return "NiagaraEditor.Scope.System";
+	case ENiagaraParameterScope::Emitter:
+		return "NiagaraEditor.Scope.Emitter";
+	case ENiagaraParameterScope::Particles:
+		return "NiagaraEditor.Scope.Particles";
+	case ENiagaraParameterScope::ScriptPersistent:
+		return "NiagaraEditor.Scope.ScriptPersistent";
+	case ENiagaraParameterScope::ScriptTransient:
+		return "NiagaraEditor.Scope.ScriptTransient";
+	};
+	return NAME_None;
+}
+
+
 void OpenSourceAsset(TWeakObjectPtr<UNiagaraStackEntry> StackEntryWeak)
 {
 	UNiagaraStackEntry* StackEntry = StackEntryWeak.Get();
diff --git a/Engine/Plugins/FX/Niagara/Source/NiagaraEditorWidgets/Private/NiagaraEditorWidgetsUtilities.h b/Engine/Plugins/FX/Niagara/Source/NiagaraEditorWidgets/Private/NiagaraEditorWidgetsUtilities.h
--- a/Engine/Plugins/FX/Niagara/Source/NiagaraEditorWidgets/Private/NiagaraEditorWidgetsUtilities.h
+++ b/Engine/Plugins/FX/Niagara/Source/NiagaraEditorWidgets/Private/NiagaraEditorWidgetsUtilities.h
@@ -5,6 +5,7 @@
 #include "Templates/SharedPointer.h"
 #include "UObject/NameTypes.h"
 #include "ViewModels/Stack/NiagaraStackEntry.h"
+#include "NiagaraTypes.h"
 
 class FMenuBuilder;
 class SWidget;
@@ -20,6 +21,8 @@ namespace FNiagaraStackEditorWidgetsUtilities
 
 	FName GetIconColorNameForExecutionCategory(FName ExecutionCategoryName);
 	
+	FName GetColorNameForParameterScope(ENiagaraParameterScope ParameterScope);
+
 	bool AddStackEntryAssetContextMenuActions(FMenuBuilder& MenuBuilder, UNiagaraStackEntry& StackEntry);
 
 	bool AddStackItemContextMenuActions(FMenuBuilder& MenuBuilder, UNiagaraStackItem& StackItem);
diff --git a/Engine/Plugins/FX/Niagara/Source/NiagaraEditorWidgets/Public/NiagaraEditorWidgetsModule.h b/Engine/Plugins/FX/Niagara/Source/NiagaraEditorWidgets/Public/NiagaraEditorWidgetsModule.h
--- a/Engine/Plugins/FX/Niagara/Source/NiagaraEditorWidgets/Public/NiagaraEditorWidgetsModule.h
+++ b/Engine/Plugins/FX/Niagara/Source/NiagaraEditorWidgets/Public/NiagaraEditorWidgetsModule.h
@@ -56,6 +56,7 @@ private:
 		virtual TSharedRef<SWidget> CreateStackIssueIcon(UNiagaraStackViewModel& StackViewModel, UNiagaraStackEntry& StackEntry) const override;
 		virtual TSharedRef<SWidget> CreateScriptScratchPad(UNiagaraScratchPadViewModel& ScriptScratchPadViewModel) const override;
 		virtual FLinearColor GetColorForExecutionCategory(FName ExecutionCategory) const override;
+		virtual FLinearColor GetColorForParameterScope(ENiagaraParameterScope ParameterScope) const override;
 	};
 
 public:
