Non unity fixes.

diff --git a/Engine/Plugins/FX/Niagara/Source/NiagaraEditor/Private/ViewModels/NiagaraParameterPanelViewModel.h b/Engine/Plugins/FX/Niagara/Source/NiagaraEditor/Private/ViewModels/NiagaraParameterPanelViewModel.h
--- a/Engine/Plugins/FX/Niagara/Source/NiagaraEditor/Private/ViewModels/NiagaraParameterPanelViewModel.h
+++ b/Engine/Plugins/FX/Niagara/Source/NiagaraEditor/Private/ViewModels/NiagaraParameterPanelViewModel.h
@@ -3,6 +3,7 @@
 #pragma once
 
 #include "NiagaraEditorCommon.h"
+#include "Types/SlateEnums.h"
 
 class UNiagaraSystem;
 class FNiagaraSystemViewModel;
@@ -67,7 +68,7 @@ public:
 
 	virtual const TArray<FNiagaraScriptVariableAndViewInfo> GetViewedParameters() = 0;
 
-	TSharedRef<SWidget> GetScriptParameterVisualWidget(FCreateWidgetForActionData* const InCreateData) const;
+	TSharedRef<class SWidget> GetScriptParameterVisualWidget(FCreateWidgetForActionData* const InCreateData) const;
 
 	FOnParameterPanelViewModelRefreshed& GetOnRefreshed() { return OnParameterPanelViewModelRefreshed; };
 
@@ -189,7 +190,7 @@ private:
 
 	TStaticArray<FScopeIsEnabledAndTooltip, (int32)ENiagaraParameterScope::Num> GetParameterScopesEnabledAndTooltips(const FNiagaraVariable& InVar, const FNiagaraVariableMetaData& InVarMetaData) const;
 
-	TSharedRef<SWidget> GetScriptParameterVisualWidget(const UEdGraphPin* Pin) const;
+	TSharedRef<class SWidget> GetScriptParameterVisualWidget(const UEdGraphPin* Pin) const;
 
 	TSharedPtr<FNiagaraObjectSelection> VariableObjectSelection;
 };
diff --git a/Engine/Plugins/FX/Niagara/Source/NiagaraEditor/Private/Widgets/SNiagaraParameterNameView.h b/Engine/Plugins/FX/Niagara/Source/NiagaraEditor/Private/Widgets/SNiagaraParameterNameView.h
--- a/Engine/Plugins/FX/Niagara/Source/NiagaraEditor/Private/Widgets/SNiagaraParameterNameView.h
+++ b/Engine/Plugins/FX/Niagara/Source/NiagaraEditor/Private/Widgets/SNiagaraParameterNameView.h
@@ -2,6 +2,9 @@
 
 #pragma once
 
+#include "CoreMinimal.h"
+#include "Widgets/SCompoundWidget.h"
+
 struct FNiagaraScriptVariableAndViewInfo;
 struct FSlateFontInfo;
 struct FCreateWidgetForActionData;
