Integrating Niagara audio changes

Spectrum Niagara Data Interface

diff --git a/Engine/Plugins/FX/Niagara/Source/Niagara/Classes/NiagaraDataInterfaceAudio.h b/Engine/Plugins/FX/Niagara/Source/Niagara/Classes/NiagaraDataInterfaceAudio.h
new file mode 100644
--- /dev/null
+++ b/Engine/Plugins/FX/Niagara/Source/Niagara/Classes/NiagaraDataInterfaceAudio.h
@@ -0,0 +1,142 @@
+// Copyright Epic Games, Inc. All Rights Reserved.
+#pragma once
+
+#include "CoreMinimal.h"
+#include "NiagaraDataInterface.h"
+#include "AudioDevice.h"
+#include "AudioDeviceManager.h"
+#include "Sound/SoundSubmix.h"
+#include "DSP/MultithreadedPatching.h"
+#include "NiagaraDataInterfaceAudio.generated.h"
+
+/** Class used to to capture the audio stream of an arbitrary submix. */
+class NIAGARA_API FNiagaraSubmixListener : public ISubmixBufferListener
+{
+public:
+
+	/** Construct an FNiagaraSubmixListener
+	 * 
+	 * @param InMixer - Input FPatchMixer. Submix audio will be mixed into this patch mixer.
+	 * @param InNumSamplesToBuffer - Number of samples to hold in patch input. Low values may cause 
+	 *                               data to be overwritten by threads that produce audio. High 
+	 *                               values require more memory.
+	 */
+	FNiagaraSubmixListener(Audio::FPatchMixer& InMixer, int32 InNumSamplesToBuffer);
+
+	FNiagaraSubmixListener(const FNiagaraSubmixListener& Other)
+	{
+		// Copy constructor technically required to compile TMap, but not used during runtime if move constructor is available.
+		// If you're hitting this, consider using Emplace or Add(MoveTemp()).
+		checkNoEntry();
+	}
+
+	/** Move submix listener. */
+	FNiagaraSubmixListener(FNiagaraSubmixListener&& Other);
+
+	virtual ~FNiagaraSubmixListener();
+
+	/** Returns the current sample rate of the current submix. */
+	float GetSampleRate() const;
+
+	/** Returns the number of channels of the current submix. */
+	int32 GetNumChannels() const;
+
+	// Begin ISubmixBufferListener overrides
+	virtual void OnNewSubmixBuffer(const USoundSubmix* OwningSubmix, float* AudioData, int32 NumSamples, int32 NumChannels, const int32 SampleRate, double AudioClock) override;
+	// End ISubmixBufferListener overrides
+
+private:
+	FNiagaraSubmixListener();
+
+	TAtomic<int32> NumChannelsInSubmix;
+	TAtomic<int32> SubmixSampleRate;
+
+	Audio::FPatchInput MixerInput;
+};
+
+/** Niagara data interface proxy for audio submix listener. */
+struct FNiagaraDataInterfaceProxySubmix : public FNiagaraDataInterfaceProxy 
+{
+	/** Construct a FNiagaraDataInterfaceProxySubmix
+	 *
+	 *  @params InNumSamplesToBuffer - Number of audio sampels to buffer internally. 
+	 */
+	FNiagaraDataInterfaceProxySubmix(int32 InNumSamplesToBuffer);
+
+	virtual ~FNiagaraDataInterfaceProxySubmix();
+
+	void OnBeginDestroy();
+
+	/** @return the number of channels in the buffer.  */
+	int32 GetNumChannels() const;
+
+	/** @return the sample rate of audio in the buffer..  */
+	float GetSampleRate() const;
+
+	/** Called when the Submix property changes. */
+	virtual void OnUpdateSubmix(USoundSubmix* Submix);
+
+	/** Copies the minimum of NumSamples or however many samples are available into OutBuffer. Returns the number of samples copied, or -1 if this output's corresponding input has been destroyed. */
+	int32 PopAudio(float* OutBuffer, int32 NumSamples, bool bUseLatestAudio);
+
+	/** Returns the current number of samples buffered on this output. */
+	int32 GetNumSamplesAvailable();
+
+	/** Returns the current number of frames buffered on this output. If NumChannels is zero, then num frames
+	 * will also be zero even when samples are available. */
+	int32 GetNumFramesAvailable();
+
+	virtual int32 PerInstanceDataPassedToRenderThreadSize() const override;
+
+private:
+	void AddSubmixListener(const Audio::FDeviceId InDeviceId);
+	void RemoveSubmixListener(const Audio::FDeviceId InDeviceId);
+
+	void RegisterToAllAudioDevices();
+	void UnregisterFromAllAudioDevices(USoundSubmix* Submix);
+
+	void OnNewDeviceCreated(Audio::FDeviceId InDeviceId);
+	void OnDeviceDestroyed(Audio::FDeviceId InDeviceId);
+
+	// Map of audio devices to submix listeners. Needed for editor where multiple audio devices
+	// may exist.
+	TMap<Audio::FDeviceId, FNiagaraSubmixListener> SubmixListeners;
+
+	// Mixer for sending audio 
+	Audio::FPatchMixer PatchMixer;
+
+	USoundSubmix* SubmixRegisteredTo;
+	bool bIsSubmixListenerRegistered;
+
+	int32 NumSamplesToBuffer;
+	
+	FDelegateHandle DeviceCreatedHandle;
+	FDelegateHandle DeviceDestroyedHandle;
+};
+
+/** Data Interface allowing sampling of recent audio data. */
+UCLASS(EditInlineNew, Category = "Audio", meta = (DisplayName = "Audio Submix"))
+class NIAGARA_API UNiagaraDataInterfaceAudioSubmix : public UNiagaraDataInterface
+{
+	GENERATED_UCLASS_BODY()
+public:
+	
+	/** The audio submix where analyzed. */
+	UPROPERTY(EditAnywhere, Category = "Audio")
+	USoundSubmix* Submix;
+
+
+#if WITH_EDITOR
+	virtual void PostEditChangeProperty(struct FPropertyChangedEvent& PropertyChangedEvent) override;
+#endif // WITH_EDITOR
+
+	virtual void PostInitProperties() override;
+	virtual void BeginDestroy() override;
+	virtual void PostLoad() override;
+
+protected:
+	virtual bool CopyToInternal(UNiagaraDataInterface* Destination) const override;
+private:
+
+};
+
diff --git a/Engine/Plugins/FX/Niagara/Source/Niagara/Classes/NiagaraDataInterfaceAudioOscilloscope.h b/Engine/Plugins/FX/Niagara/Source/Niagara/Classes/NiagaraDataInterfaceAudioOscilloscope.h
--- a/Engine/Plugins/FX/Niagara/Source/Niagara/Classes/NiagaraDataInterfaceAudioOscilloscope.h
+++ b/Engine/Plugins/FX/Niagara/Source/Niagara/Classes/NiagaraDataInterfaceAudioOscilloscope.h
@@ -9,39 +9,9 @@
 #include "AudioDevice.h"
 #include "AudioDeviceManager.h"
 #include "DSP/MultithreadedPatching.h"
+#include "NiagaraDataInterfaceAudio.h"
 #include "NiagaraDataInterfaceAudioOscilloscope.generated.h"
 
-// Class used to to capture the audio stream of an arbitrary submix.
-class NIAGARA_API FNiagaraSubmixListener : public ISubmixBufferListener
-{
-public:
-	FNiagaraSubmixListener(Audio::FPatchMixer& InMixer, int32 InNumSamplesToBuffer);
-	FNiagaraSubmixListener(const FNiagaraSubmixListener& Other)
-	{
-		// Copy constructor technically required to compile TMap, but not used during runtime if move constructor is available.
-		// If you're hitting this, consider using Emplace or Add(MoveTemp()).
-		checkNoEntry();
-	}
-
-	FNiagaraSubmixListener(FNiagaraSubmixListener&& Other);
-
-	virtual ~FNiagaraSubmixListener();
-
-	float GetSampleRate();
-	int32 GetNumChannels();
-
-	// Begin ISubmixBufferListener overrides
-	virtual void OnNewSubmixBuffer(const USoundSubmix* OwningSubmix, float* AudioData, int32 NumSamples, int32 NumChannels, const int32 SampleRate, double AudioClock) override;
-	// End ISubmixBufferListener overrides
-
-private:
-	FNiagaraSubmixListener();
-
-	int32 NumChannelsInSubmix;
-	int32 SubmixSampleRate;
-
-	Audio::FPatchInput MixerInput;
-};
 
 struct FNiagaraDataInterfaceProxyOscilloscope : public FNiagaraDataInterfaceProxy
 {
diff --git a/Engine/Plugins/FX/Niagara/Source/Niagara/Classes/NiagaraDataInterfaceAudioSpectrum.h b/Engine/Plugins/FX/Niagara/Source/Niagara/Classes/NiagaraDataInterfaceAudioSpectrum.h
new file mode 100644
--- /dev/null
+++ b/Engine/Plugins/FX/Niagara/Source/Niagara/Classes/NiagaraDataInterfaceAudioSpectrum.h
@@ -0,0 +1,190 @@
+// Copyright Epic Games, Inc. All Rights Reserved.
+#pragma once
+
+#include "NiagaraDataInterface.h"
+#include "DSP/SlidingWindow.h"
+#include "DSP/DeinterleaveView.h"
+#include "DSP/ConstantQ.h"
+#include "DSP/FFTAlgorithm.h"
+#include "DSP/FloatArrayMath.h"
+#include "NiagaraDataInterfaceAudio.h"
+#include "NiagaraDataInterfaceAudioSpectrum.generated.h"
+
+/** FNiagaraDataInterfaceProxySpectrum
+ *  
+ *  Proxy class for calculating spectrums on the rendering thread for Niagara.
+ */
+struct FNiagaraDataInterfaceProxySpectrum : public FNiagaraDataInterfaceProxySubmix
+{
+	/** Construct an FNiagaraDataInterfaceProxySpectrum
+	 *
+	 * 	@param InMinimumFrequency - Minimum frequency in spectrum.
+	 * 	@param InMaximumFrequency - Maximum frequency in spectrum.
+	 * 	@param InNumBands - Number of bands in spectrum.
+	 * 	@param InNoiseFloorDb - Decibels level considered as silence. 
+	 * 	@param InNumSamplesToBuffer - Number of samples to buffer internally across thread boundaries.
+	 */
+	FNiagaraDataInterfaceProxySpectrum(float InMinimumFrequency, float InMaximumFrequency, int32 InNumBands, float InNoiseFloorDb, int32 InNumSamplesToBuffer);
+
+	virtual ~FNiagaraDataInterfaceProxySpectrum();
+
+	/** Sample value from the spectrum buffer.
+	 *  
+	 *  @param InNormalizedPositionInSpectrum - Value between 0 and 1 which represents the relative frequency in the spectrum.  
+	 *                                          0 corresponds to the MinimumFrequency, 1 corresponds to the MaximumFrequency.
+	 *                                          Values between 0 and 1 correspond to logarithmically spaced frequencys in-between
+	 *                                          the limits.
+	 *  @param InChannelIndex - The index of the audio channel to sample.
+	 *
+	 *  @return Amplitude at this position. 0.0 corresponds to the nosie floor, while 1.0 corresponds to full volume. Some values
+	 *          may exceed 1.0.
+	 */
+	float GetSpectrumValue(float InNormalizedPositionInSpectrum, int32 InChannelIndex);
+
+	/** Updates the minimum and maximum frequency of the CQT on the render thread. */
+	void UpdateCQT(float InMinimumFrequency, float InMaximumFrequency, int32 InNumBands);
+
+	/** Updates the value considered silence in decibles on the render thread. */
+	void UpdateNoiseFloor(float InNoiseFloorDb);
+
+
+	/** This function enqueues a render thread command to pop audio off of the SubmixListener, transform it into a CQT, and 
+	 * post it to the GPUAudioBuffer.
+	 */
+	void PostDataToGPU();
+	FReadBuffer& ComputeAndPostSRV();
+
+	virtual int32 PerInstanceDataPassedToRenderThreadSize() const override
+	{
+		return 0;
+	}
+
+private:
+	typedef Audio::TSlidingBuffer<float> FSlidingBuffer;
+	typedef Audio::TAutoSlidingWindow<float, Audio::FAudioBufferAlignedAllocator> FSlidingWindow;
+	typedef Audio::TAutoDeinterleaveView<float, Audio::FAudioBufferAlignedAllocator> FDeinterleaveView;
+	typedef FDeinterleaveView::TChannel<Audio::FAudioBufferAlignedAllocator> FChannel;
+
+	// Calculates the CQT
+	void UpdateSpectrum();
+
+	// Updates internal objects for num channels and samplerate
+	void SetAudioFormat(int32 InNumChannels, float InSampleRate);
+
+	void SetNoiseFloor(float InNoiseFloorDb);
+
+	// Resize internal abuffers
+	void ResizeCQT(float InMinimumFrequency, float InMaximumFrequency, int32 InNumBands);
+	void ResizeSpectrumBuffer(int32 InNumChannels, int32 InNumBands);
+	void ResizeAudioTransform(float InMinimumFrequency, float InMaximumFrequency, float InSampleRate, int32 InNumBands);
+	void ResizeSlidingBuffer(int32 InNumChannels, int32 InFFTSize);
+
+
+	// Get settings 
+	static void Clamp(float& InMinimumFrequency, float& InMaximumFrequency, int32& InNumBand);
+	static float GetNumOctaves(float InMinimumFrequency, float InMaximumFrequency);
+	static float GetNumBandsPerOctave(int32 InNumBands, float InNumOctaves);
+	static float GetBandwidthStretch(float InSampleRate, float InFFTSize, float InNumBandsPerOctave, float InMinimumFrequency);
+
+	static Audio::FPseudoConstantQKernelSettings GetConstantQSettings(float InMinimumFrequency, float InMaximumFrequency, int32 InNumBands, float InNumBandsPerOctave, float InBandwidthStretch);
+
+	static Audio::FFFTSettings GetFFTSettings(float InMinimumFrequency, float InSampleRate, float InNumBandsPerOctave);
+
+	float MinimumFrequency;
+	float MaximumFrequency;
+	int32 NumBands;
+	float NoiseFloorDb;
+
+	int32 NumChannels;
+	float SampleRate;
+	float FFTScale;
+
+	Audio::AlignedFloatBuffer PopBuffer;
+	TUniquePtr< FSlidingBuffer > SlidingBuffer;
+
+	TArray<Audio::AlignedFloatBuffer> ChannelSpectrumBuffers;
+
+	TUniquePtr<Audio::FContiguousSparse2DKernelTransform> CQTKernel;
+	TUniquePtr<Audio::IFFTAlgorithm> FFTAlgorithm;
+
+	Audio::AlignedFloatBuffer InterleavedWindowBuffer;
+	Audio::AlignedFloatBuffer WindowBuffer;
+	Audio::AlignedFloatBuffer FFTInputBuffer;
+	Audio::AlignedFloatBuffer FFTOutputBuffer;
+	Audio::AlignedFloatBuffer PowerSpectrumBuffer;
+	Audio::AlignedFloatBuffer SpectrumBuffer;
+
+	// Handle for the SRV used by the generated HLSL.
+	FReadBuffer GPUBuffer;
+
+	FCriticalSection BufferLock;
+};
+
+/** Data Interface allowing sampling of recent audio spectrum. */
+UCLASS(EditInlineNew, Category = "Audio", meta = (DisplayName = "Audio Spectrum"))
+class NIAGARA_API UNiagaraDataInterfaceAudioSpectrum : public UNiagaraDataInterfaceAudioSubmix
+{
+	GENERATED_UCLASS_BODY()
+public:
+
+	DECLARE_NIAGARA_DI_PARAMETER();
+
+	// VM function names
+	static const FName GetSpectrumFunctionName;
+	static const FName GetNumChannelsFunctionName;
+
+	// Global variable prefixes
+	static const FString GetSpectrumName;
+	static const FString NumChannelsName;
+	
+	/** The number of spectrum samples to pass to the GPU */
+	UPROPERTY(EditAnywhere, Category = "Spectrum", meta = (ClampMin = "16", ClampMax = "1024") )
+	int32 Resolution;
+
+	/** The minimum frequency represented in the spectrum. */
+	UPROPERTY(EditAnywhere, Category = "Spectrum", meta = (ClampMin = "20.0", ClampMax = "20000.0"))
+	float MinimumFrequency;
+
+	/** The maximum frequency represented in the spectrum. */
+	UPROPERTY(EditAnywhere, Category = "Spectrum", meta = (ClampMin = "20.0", ClampMax = "20000.0"))
+	float MaximumFrequency;
+
+	/** The decibel level considered as silence. This is used to scale the output of the spectrum. */
+	UPROPERTY(EditAnywhere, AdvancedDisplay, Category = "Spectrum", meta = (ClampMin = "-120.0", ClampMax = "0.0"))
+	float NoiseFloorDb;
+
+	//VM function overrides:
+	void GetSpectrumValue(FVectorVMContext& Context) {}
+	void GetNumChannels(FVectorVMContext& Context) {}
+
+	virtual void GetFunctions(TArray<FNiagaraFunctionSignature>& OutFunctions) override;
+	virtual void GetVMExternalFunction(const FVMExternalFunctionBindingInfo& BindingInfo, void* InstanceData, FVMExternalFunction &OutFunc) override;
+
+	virtual bool CanExecuteOnTarget(ENiagaraSimTarget Target) const override { 
+		if (Target == ENiagaraSimTarget::GPUComputeSim)
+		{
+			return true;
+		}
+		else
+		{
+			return false;
+		}
+	}
+
+	virtual bool GetFunctionHLSL(const FNiagaraDataInterfaceGPUParamInfo& ParamInfo, const FNiagaraDataInterfaceGeneratedFunction& FunctionInfo, int FunctionInstanceIndex, FString& OutHLSL) override;
+
+	virtual void GetParameterDefinitionHLSL(const FNiagaraDataInterfaceGPUParamInfo& ParamInfo, FString& OutHLSL) override;
+
+#if WITH_EDITOR
+	virtual void PostEditChangeProperty(struct FPropertyChangedEvent& PropertyChangedEvent) override;
+#endif // WITH_EDITOR
+
+	virtual void PostInitProperties() override;
+	virtual void PostLoad() override;
+
+protected:
+	virtual bool CopyToInternal(UNiagaraDataInterface* Destination) const override;
+private:
+
+};
+
diff --git a/Engine/Plugins/FX/Niagara/Source/Niagara/Private/NiagaraDataInterfaceAudio.cpp b/Engine/Plugins/FX/Niagara/Source/Niagara/Private/NiagaraDataInterfaceAudio.cpp
new file mode 100644
--- /dev/null
+++ b/Engine/Plugins/FX/Niagara/Source/Niagara/Private/NiagaraDataInterfaceAudio.cpp
@@ -0,0 +1,284 @@
+// Copyright Epic Games, Inc. All Rights Reserved.
+
+#include "NiagaraDataInterfaceAudio.h"
+
+FNiagaraSubmixListener::FNiagaraSubmixListener(Audio::FPatchMixer& InMixer, int32 InNumSamplesToBuffer)
+	: NumChannelsInSubmix(0)
+	, SubmixSampleRate(0)
+	, MixerInput(InMixer.AddNewInput(InNumSamplesToBuffer, 1.0f))
+{
+
+}
+
+FNiagaraSubmixListener::FNiagaraSubmixListener()
+	: NumChannelsInSubmix(0)
+	, SubmixSampleRate(0)
+{
+}
+
+FNiagaraSubmixListener::FNiagaraSubmixListener(FNiagaraSubmixListener&& Other)
+{
+	NumChannelsInSubmix = Other.NumChannelsInSubmix.Load();
+	Other.NumChannelsInSubmix = 0;
+
+	SubmixSampleRate = Other.SubmixSampleRate.Load();
+	Other.SubmixSampleRate = 0;
+
+	MixerInput = MoveTemp(Other.MixerInput);
+}
+
+FNiagaraSubmixListener::~FNiagaraSubmixListener()
+{
+}
+
+float FNiagaraSubmixListener::GetSampleRate() const
+{
+	return (float) SubmixSampleRate;
+}
+
+int32 FNiagaraSubmixListener::GetNumChannels() const
+{
+	return NumChannelsInSubmix;
+}
+
+void FNiagaraSubmixListener::OnNewSubmixBuffer(const USoundSubmix* OwningSubmix, float* AudioData, int32 NumSamples, int32 NumChannels, const int32 SampleRate, double AudioClock)
+{
+	NumChannelsInSubmix = NumChannels;
+	SubmixSampleRate = SampleRate;
+	MixerInput.PushAudio(AudioData, NumSamples);
+}
+
+///////////////////////////////////////////////////////////////////////////////////////////////////
+
+
+FNiagaraDataInterfaceProxySubmix::FNiagaraDataInterfaceProxySubmix(int32 InNumSamplesToBuffer)
+	: PatchMixer()
+	, SubmixRegisteredTo(nullptr)
+	, bIsSubmixListenerRegistered(false)
+	, NumSamplesToBuffer(InNumSamplesToBuffer)
+{
+	check(NumSamplesToBuffer > 0);
+
+	DeviceCreatedHandle = FAudioDeviceManagerDelegates::OnAudioDeviceCreated.AddRaw(this, &FNiagaraDataInterfaceProxySubmix::OnNewDeviceCreated);
+
+	DeviceDestroyedHandle = FAudioDeviceManagerDelegates::OnAudioDeviceDestroyed.AddRaw(this, &FNiagaraDataInterfaceProxySubmix::OnDeviceDestroyed);
+}
+
+FNiagaraDataInterfaceProxySubmix::~FNiagaraDataInterfaceProxySubmix()
+{
+	FAudioDeviceManagerDelegates::OnAudioDeviceCreated.Remove(DeviceCreatedHandle);
+	FAudioDeviceManagerDelegates::OnAudioDeviceDestroyed.Remove(DeviceDestroyedHandle);
+}
+
+void FNiagaraDataInterfaceProxySubmix::RegisterToAllAudioDevices()
+{
+	if (FAudioDeviceManager* DeviceManager = FAudioDeviceManager::Get())
+	{
+		// Register a new submix listener for every audio device that currently exists.
+		DeviceManager->IterateOverAllDevices([&](Audio::FDeviceId DeviceId, FAudioDevice* InDevice)
+		{
+			AddSubmixListener(DeviceId);
+		});
+	}
+}
+
+void FNiagaraDataInterfaceProxySubmix::UnregisterFromAllAudioDevices(USoundSubmix* Submix)
+{
+	if (FAudioDeviceManager* DeviceManager = FAudioDeviceManager::Get())
+	{
+		// Register a new submix listener for every audio device that currently exists.
+		DeviceManager->IterateOverAllDevices([&](Audio::FDeviceId DeviceId, FAudioDevice* InDevice)
+		{
+			RemoveSubmixListener(DeviceId);
+		});
+	}
+}
+
+void FNiagaraDataInterfaceProxySubmix::OnUpdateSubmix(USoundSubmix* Submix)
+{
+	if (bIsSubmixListenerRegistered)
+	{
+		UnregisterFromAllAudioDevices(Submix);
+	}
+
+	SubmixRegisteredTo = Submix;
+	
+	RegisterToAllAudioDevices();
+	bIsSubmixListenerRegistered = true;
+}
+
+
+void FNiagaraDataInterfaceProxySubmix::OnNewDeviceCreated(Audio::FDeviceId InDeviceId)
+{
+	if (bIsSubmixListenerRegistered)
+	{
+		AddSubmixListener(InDeviceId);
+	}
+}
+
+void FNiagaraDataInterfaceProxySubmix::OnDeviceDestroyed(Audio::FDeviceId InDeviceId)
+{
+	if (bIsSubmixListenerRegistered)
+	{
+		RemoveSubmixListener(InDeviceId);
+	}
+}
+
+void FNiagaraDataInterfaceProxySubmix::AddSubmixListener(const Audio::FDeviceId InDeviceId)
+{
+	check(!SubmixListeners.Contains(InDeviceId));
+
+	FAudioDeviceHandle DeviceHandle = FAudioDeviceManager::Get()->GetAudioDevice(InDeviceId);
+	check(DeviceHandle);
+
+	FNiagaraSubmixListener& Listener = SubmixListeners.Emplace(InDeviceId, FNiagaraSubmixListener(PatchMixer, NumSamplesToBuffer));
+
+	DeviceHandle->RegisterSubmixBufferListener(&Listener, SubmixRegisteredTo);
+}
+
+void FNiagaraDataInterfaceProxySubmix::RemoveSubmixListener(Audio::FDeviceId InDeviceId)
+{
+	if (SubmixListeners.Contains(InDeviceId))
+	{
+		FAudioDeviceHandle DeviceHandle = FAudioDeviceManager::Get()->GetAudioDevice(InDeviceId);
+		check(DeviceHandle);
+
+		DeviceHandle->UnregisterSubmixBufferListener(&SubmixListeners[InDeviceId], SubmixRegisteredTo);
+		SubmixListeners.Remove(InDeviceId);
+	}
+}
+
+int32 FNiagaraDataInterfaceProxySubmix::GetNumChannels() const
+{
+	TArray<Audio::FDeviceId> DeviceIds;
+	SubmixListeners.GetKeys(DeviceIds);
+	for (const Audio::FDeviceId DeviceId : DeviceIds)
+	{
+		int32 NumChannels = SubmixListeners[DeviceId].GetNumChannels();
+
+		if (NumChannels != 0)
+		{
+			// Not sure which listener is receiving callbacks, or if more than one is,
+			// so return first non-zero value.
+			return NumChannels;
+		}
+	}
+	return 0;
+}
+
+float FNiagaraDataInterfaceProxySubmix::GetSampleRate() const
+{
+	TArray<Audio::FDeviceId> DeviceIds;
+	SubmixListeners.GetKeys(DeviceIds);
+	for (const Audio::FDeviceId DeviceId : DeviceIds)
+	{
+		float SampleRate = SubmixListeners[DeviceId].GetSampleRate();
+
+		if (SampleRate > 0.f)
+		{
+			// Not sure which listener is receiving callbacks, or if more than one is,
+			// so return first non-zero value.
+			return SampleRate;
+		}
+	}
+	return 0.f;
+}
+
+void FNiagaraDataInterfaceProxySubmix::OnBeginDestroy()
+{
+	if (bIsSubmixListenerRegistered)
+	{
+		UnregisterFromAllAudioDevices(SubmixRegisteredTo);
+	}
+}
+
+int32 FNiagaraDataInterfaceProxySubmix::PopAudio(float* OutBuffer, int32 NumSamples, bool bUseLatestAudio)
+{
+	return PatchMixer.PopAudio(OutBuffer, NumSamples, bUseLatestAudio);
+}
+
+int32 FNiagaraDataInterfaceProxySubmix::GetNumSamplesAvailable()
+{
+	return PatchMixer.MaxNumberOfSamplesThatCanBePopped();
+}
+
+int32 FNiagaraDataInterfaceProxySubmix::GetNumFramesAvailable()
+{
+	int32 NumChannels = GetNumChannels();
+
+	if (NumChannels < 1)
+	{
+		return 0;
+	}
+
+	return GetNumSamplesAvailable() / NumChannels;
+}
+
+int32 FNiagaraDataInterfaceProxySubmix::PerInstanceDataPassedToRenderThreadSize() const
+{
+	return 0;
+}
+
+UNiagaraDataInterfaceAudioSubmix::UNiagaraDataInterfaceAudioSubmix(FObjectInitializer const& ObjectInitializer)
+	: Super(ObjectInitializer)
+	, Submix(nullptr)
+{
+	int32 DefaultNumSamplesToBuffer = 16384;
+	Proxy = MakeUnique<FNiagaraDataInterfaceProxySubmix>(DefaultNumSamplesToBuffer);
+}
+
+
+#if WITH_EDITOR
+void UNiagaraDataInterfaceAudioSubmix::PostEditChangeProperty(struct FPropertyChangedEvent& PropertyChangedEvent)
+{
+	Super::PostEditChangeProperty(PropertyChangedEvent);
+	
+	static FName SubmixFName = GET_MEMBER_NAME_CHECKED(UNiagaraDataInterfaceAudioSubmix, Submix);
+
+	// Regenerate on save any compressed sound formats or if analysis needs to be re-done
+	if (FProperty* PropertyThatChanged = PropertyChangedEvent.Property)
+	{
+		const FName& Name = PropertyThatChanged->GetFName();
+		if (Name == SubmixFName)
+		{
+			GetProxyAs<FNiagaraDataInterfaceProxySubmix>()->OnUpdateSubmix(Submix);
+		}
+	}
+}
+#endif //WITH_EDITOR
+
+void UNiagaraDataInterfaceAudioSubmix::PostInitProperties()
+{
+	Super::PostInitProperties();
+
+	GetProxyAs<FNiagaraDataInterfaceProxySubmix>()->OnUpdateSubmix(Submix);
+}
+
+
+void UNiagaraDataInterfaceAudioSubmix::BeginDestroy()
+{
+	GetProxyAs<FNiagaraDataInterfaceProxySubmix>()->OnBeginDestroy();
+
+	Super::BeginDestroy();
+}
+
+void UNiagaraDataInterfaceAudioSubmix::PostLoad()
+{
+	Super::PostLoad();
+	GetProxyAs<FNiagaraDataInterfaceProxySubmix>()->OnUpdateSubmix(Submix);
+}
+
+bool UNiagaraDataInterfaceAudioSubmix::CopyToInternal(UNiagaraDataInterface* Destination) const
+{
+	Super::CopyToInternal(Destination);
+
+	UNiagaraDataInterfaceAudioSubmix* CastedDestination = Cast<UNiagaraDataInterfaceAudioSubmix>(Destination);
+
+	if (CastedDestination)
+	{
+		CastedDestination->Submix = Submix;
+	}
+	
+	return true;
+}
+
diff --git a/Engine/Plugins/FX/Niagara/Source/Niagara/Private/NiagaraDataInterfaceAudioOscilloscope.cpp b/Engine/Plugins/FX/Niagara/Source/Niagara/Private/NiagaraDataInterfaceAudioOscilloscope.cpp
--- a/Engine/Plugins/FX/Niagara/Source/Niagara/Private/NiagaraDataInterfaceAudioOscilloscope.cpp
+++ b/Engine/Plugins/FX/Niagara/Source/Niagara/Private/NiagaraDataInterfaceAudioOscilloscope.cpp
@@ -534,51 +534,5 @@ void FNiagaraDataInterfaceProxyOscilloscope::DownsampleAudioToBuffer()
 	DownsampledBuffer.SetNumZeroed(Resolution * NumChannelsInDownsampledBuffer);
 }
 
-FNiagaraSubmixListener::FNiagaraSubmixListener(Audio::FPatchMixer& InMixer, int32 InNumSamplesToBuffer)
-	: NumChannelsInSubmix(0)
-	, SubmixSampleRate(0)
-	, MixerInput(InMixer.AddNewInput(InNumSamplesToBuffer, 1.0f))
-{
-
-}
-
-FNiagaraSubmixListener::FNiagaraSubmixListener()
-	: NumChannelsInSubmix(0)
-	, SubmixSampleRate(0)
-{
-
-}
-
-FNiagaraSubmixListener::FNiagaraSubmixListener(FNiagaraSubmixListener&& Other)
-{
-	NumChannelsInSubmix = Other.NumChannelsInSubmix;
-	Other.NumChannelsInSubmix = 0;
-
-	SubmixSampleRate = Other.SubmixSampleRate;
-	Other.SubmixSampleRate = 0;
-
-	MixerInput = MoveTemp(Other.MixerInput);
-}
-
-FNiagaraSubmixListener::~FNiagaraSubmixListener()
-{
-}
-
-float FNiagaraSubmixListener::GetSampleRate()
-{
-	return (float) SubmixSampleRate;
-}
-
-int32 FNiagaraSubmixListener::GetNumChannels()
-{
-	return NumChannelsInSubmix;
-}
-
-void FNiagaraSubmixListener::OnNewSubmixBuffer(const USoundSubmix* OwningSubmix, float* AudioData, int32 NumSamples, int32 NumChannels, const int32 SampleRate, double AudioClock)
-{
-	NumChannelsInSubmix = NumChannels;
-	SubmixSampleRate = SampleRate;
-	MixerInput.PushAudio(AudioData, NumSamples);
-}
 
 #undef LOCTEXT_NAMESPACE
\ No newline at end of file
diff --git a/Engine/Plugins/FX/Niagara/Source/Niagara/Private/NiagaraDataInterfaceAudioSpectrum.cpp b/Engine/Plugins/FX/Niagara/Source/Niagara/Private/NiagaraDataInterfaceAudioSpectrum.cpp
new file mode 100644
--- /dev/null
+++ b/Engine/Plugins/FX/Niagara/Source/Niagara/Private/NiagaraDataInterfaceAudioSpectrum.cpp
@@ -0,0 +1,777 @@
+// Copyright Epic Games, Inc. All Rights Reserved.
+
+#include "NiagaraDataInterfaceAudioSpectrum.h"
+#include "NiagaraTypes.h"
+#include "NiagaraCustomVersion.h"
+#include "NiagaraShader.h"
+#include "ShaderParameterUtils.h"
+#include "Engine/Engine.h"
+#include "NiagaraComponent.h"
+
+#include "DSP/DeinterleaveView.h"
+#include "DSP/SlidingWindow.h"
+#include "DSP/ConstantQ.h"
+#include "DSP/FFTAlgorithm.h"
+#include "DSP/FloatArrayMath.h"
+
+#define LOCTEXT_NAMESPACE "NiagaraDataInterfaceGridAudioSpectrum"
+
+FNiagaraDataInterfaceProxySpectrum::FNiagaraDataInterfaceProxySpectrum(float InMinimumFrequency, float InMaximumFrequency, int32 InNumBands, float InNoiseFloorDb, int32 InNumSamplesToBuffer)
+
+	: FNiagaraDataInterfaceProxySubmix(InNumSamplesToBuffer)
+	, MinimumFrequency(0.f)
+	, MaximumFrequency(0.f)
+	, NumBands(0)
+	, NoiseFloorDb(0.f)
+	, NumChannels(0)
+	, SampleRate(0.f)
+	, FFTScale(1.f)
+{
+	int32 InitialNumChannels = 2;
+	float InitialSampleRate = 48000.f;
+
+	SetAudioFormat(InitialNumChannels, InitialSampleRate);
+
+	ResizeCQT(InMinimumFrequency, InMaximumFrequency, InNumBands);
+
+	SetNoiseFloor(InNoiseFloorDb);
+}
+
+FNiagaraDataInterfaceProxySpectrum::~FNiagaraDataInterfaceProxySpectrum()
+{
+	check(IsInRenderingThread());
+	GPUBuffer.Release();
+}
+
+float FNiagaraDataInterfaceProxySpectrum::GetSpectrumValue(float InNormalizedPositionInSpectrum, int32 InChannelIndex)
+{
+	check(InChannelIndex >= 0);
+
+	// Check if channel index is within channel spectrum buffers
+	if (ChannelSpectrumBuffers.Num() >= InChannelIndex)
+	{
+		return 0.f;
+	}
+
+	const Audio::AlignedFloatBuffer& Buffer = ChannelSpectrumBuffers[InChannelIndex];
+	const int32 MaxIndex = Buffer.Num() - 1;
+
+	if (MaxIndex < 0)
+	{
+		return 0.0f;
+	}
+
+	InNormalizedPositionInSpectrum = FMath::Clamp(InNormalizedPositionInSpectrum, 0.0f, 1.0f);
+
+	float Index = InNormalizedPositionInSpectrum * Buffer.Num();
+
+	int32 LowerIndex = FMath::FloorToInt(Index);
+	LowerIndex = FMath::Clamp(LowerIndex, 0, MaxIndex);
+
+	int32 UpperIndex = FMath::CeilToInt(Index);
+	UpperIndex = FMath::Clamp(UpperIndex, 0, MaxIndex);
+
+	float Alpha = Index - LowerIndex;
+
+	return FMath::Lerp<float>(Buffer[LowerIndex], Buffer[UpperIndex], Alpha);
+}
+
+
+void FNiagaraDataInterfaceProxySpectrum::UpdateCQT(float InMinimumFrequency, float InMaximumFrequency, int32 InNumBands)
+{
+	ENQUEUE_RENDER_COMMAND(FUpdateSpectrumBuffer)(
+		[&, InMinimumFrequency, InMaximumFrequency, InNumBands](FRHICommandListImmediate& RHICmdList)
+	{
+		ResizeCQT(InMinimumFrequency, InMaximumFrequency, InNumBands);
+	});
+}
+
+void FNiagaraDataInterfaceProxySpectrum::ResizeCQT(float InMinimumFrequency, float InMaximumFrequency, int32 InNumBands)
+{
+	Clamp(InMinimumFrequency, InMaximumFrequency, InNumBands);
+
+	// Check which values were changed
+	bool bUpdatedMinimumFrequency = InMinimumFrequency != MinimumFrequency;
+	bool bUpdatedMaximumFrequency = InMaximumFrequency != MaximumFrequency;
+	bool bUpdatedNumBands = InNumBands != NumBands;
+
+	bool bUpdatedValues = bUpdatedMinimumFrequency || bUpdatedMaximumFrequency || bUpdatedNumBands;
+
+	if (bUpdatedNumBands)
+	{
+		ResizeSpectrumBuffer(NumChannels, InNumBands);
+	}
+
+	if (bUpdatedValues)
+	{
+		int32 OriginalFFTSize = FFTAlgorithm.IsValid() ? FFTAlgorithm->Size() : 0;
+
+		ResizeAudioTransform(InMinimumFrequency, InMaximumFrequency, SampleRate, InNumBands);
+
+		int32 NewFFTSize = FFTAlgorithm.IsValid() ? FFTAlgorithm->Size() : 0;
+
+		if (NewFFTSize != OriginalFFTSize)
+		{
+			ResizeSlidingBuffer(NumChannels, NewFFTSize);
+		}
+
+	}
+
+	// Copy values internally
+	MinimumFrequency = InMinimumFrequency;
+	MaximumFrequency = InMaximumFrequency;
+	NumBands = InNumBands;
+}
+
+void FNiagaraDataInterfaceProxySpectrum::UpdateNoiseFloor(float InNoiseFloorDb)
+{
+	ENQUEUE_RENDER_COMMAND(FUpdateSpectrumBuffer)(
+		[&, InNoiseFloorDb](FRHICommandListImmediate& RHICmdList)
+	{
+		SetNoiseFloor(InNoiseFloorDb);
+	});
+}
+
+
+void FNiagaraDataInterfaceProxySpectrum::SetNoiseFloor(float InNoiseFloorDb)
+{
+	NoiseFloorDb = FMath::Clamp(InNoiseFloorDb, -200.f, 10.f);
+}
+
+FReadBuffer& FNiagaraDataInterfaceProxySpectrum::ComputeAndPostSRV()
+{
+	PostDataToGPU();
+	return GPUBuffer;
+}
+
+
+void FNiagaraDataInterfaceProxySpectrum::PostDataToGPU()
+{
+	ENQUEUE_RENDER_COMMAND(FUpdateSpectrumBuffer)(
+		[&](FRHICommandListImmediate& RHICmdList)
+	{
+		UpdateSpectrum();
+
+		// Resize GPU data if needed
+		const int32 NumSamplesInBuffer = NumBands * NumChannels;
+		const int32 NumBytesInChannelBuffer = NumBands * sizeof(float);
+		const int32 NumBytesInBuffer = NumBytesInChannelBuffer * NumChannels;
+		
+		if (NumBytesInBuffer != GPUBuffer.NumBytes)
+		{
+			if (GPUBuffer.NumBytes > 0)
+			{
+				GPUBuffer.Release();
+			}
+
+			if (NumSamplesInBuffer > 0)
+			{
+				GPUBuffer.Initialize(sizeof(float), NumSamplesInBuffer, EPixelFormat::PF_R32_FLOAT, BUF_Static);
+			}
+		}
+
+		// Copy to GPU data
+		if (GPUBuffer.NumBytes > 0)
+		{
+			float *BufferData = static_cast<float*>(RHILockVertexBuffer(GPUBuffer.Buffer, 0, NumBytesInBuffer, EResourceLockMode::RLM_WriteOnly));
+
+			FScopeLock ScopeLock(&BufferLock);
+
+			for (int32 ChannelIndex = 0; ChannelIndex < NumChannels; ChannelIndex++)
+			{
+				int32 Pos = ChannelIndex * NumBands;
+				FPlatformMemory::Memcpy(&BufferData[Pos], ChannelSpectrumBuffers[ChannelIndex].GetData(), NumBytesInChannelBuffer);
+			}
+
+			RHIUnlockVertexBuffer(GPUBuffer.Buffer);
+		}
+	});
+}
+
+void FNiagaraDataInterfaceProxySpectrum::UpdateSpectrum()
+{
+	check(IsInRenderingThread());
+
+	FScopeLock ScopeLock(&BufferLock);
+
+	int32 SourceNumChannels = GetNumChannels();
+	float SourceSampleRate = GetSampleRate();
+
+	// Perform internal update if new samplerate or channel count encountered.
+	if ((SourceNumChannels != NumChannels) || (SourceSampleRate != SampleRate))
+	{
+		// This will update values of NumChannels and SampleRate member variables.
+		SetAudioFormat(SourceNumChannels, SourceSampleRate);
+	}
+
+	// Set data to zero if we are skipping update due to bad samplerate or channel count.
+	if (NumChannels == 0 || FMath::IsNearlyZero(SampleRate))
+	{
+		for (int32 ChannelIndex = 0; ChannelIndex < ChannelSpectrumBuffers.Num(); ChannelIndex++)
+		{
+			Audio::AlignedFloatBuffer& Buffer = ChannelSpectrumBuffers[ChannelIndex];
+			if (Buffer.Num() > 0)
+			{
+				FMemory::Memset(Buffer.GetData(), 0, sizeof(float) * Buffer.Num());
+			}
+		}
+		return;
+	}
+
+	
+	// Grab audio from buffer.
+	int32 NumSamplesToPop = GetNumSamplesAvailable();
+	NumSamplesToPop -= (NumSamplesToPop % NumChannels);
+
+	if (NumSamplesToPop > 0)
+	{
+		PopBuffer.Reset();
+		PopBuffer.AddUninitialized(NumSamplesToPop);
+
+		bool bUseLatestAudio = false;
+		PopAudio(PopBuffer.GetData(), NumSamplesToPop, bUseLatestAudio);
+	}
+
+	// If we're in a bad internal state, give up here.	
+	if (!FFTAlgorithm.IsValid() || !CQTKernel.IsValid())
+	{
+		return;
+	}
+
+	// Run sliding window over available audio.
+	FSlidingWindow SlidingWindow(*SlidingBuffer, PopBuffer, InterleavedWindowBuffer);
+
+	int32 NumWindows = 0;
+	for (Audio::AlignedFloatBuffer& InterleavedWindow : SlidingWindow)
+	{
+		if (0 == NumWindows)
+		{
+			// Zero output on first window.  We need to wait for at least one window to process
+			// because we do not want to zero out data if no windows are available.
+			for (int32 ChannelIndex = 0; ChannelIndex < NumChannels; ChannelIndex++)
+			{
+				FMemory::Memset(ChannelSpectrumBuffers[ChannelIndex].GetData(), 0, sizeof(float) * NumBands);
+			}
+		}
+
+		NumWindows++;
+
+		// Run deinterleave view over window.
+		FDeinterleaveView DeinterleaveView(InterleavedWindow, WindowBuffer, NumChannels);
+		for(FChannel Channel : DeinterleaveView)
+		{
+			// Perform FFT
+			FMemory::Memcpy(FFTInputBuffer.GetData(), Channel.Values.GetData(), sizeof(float) * Channel.Values.Num());
+			FFTAlgorithm->ForwardRealToComplex(FFTInputBuffer.GetData(), FFTOutputBuffer.GetData());
+
+			// Transflate FFTOutput to power spectrum
+			Audio::ArrayComplexToPower(FFTOutputBuffer, PowerSpectrumBuffer);
+
+			// Take CQT of power spectrum
+			CQTKernel->TransformArray(PowerSpectrumBuffer, SpectrumBuffer);
+
+			// Accumulate power spectrum CQT output.
+			Audio::ArrayAddInPlace(SpectrumBuffer, ChannelSpectrumBuffers[Channel.ChannelIndex]);
+		}
+	}
+
+	if (NumWindows > 0)
+	{
+		// Two scalings are applied. One for FFT, another for number of windows calculated.
+		float LinearScale = FFTScale / static_cast<float>(NumWindows);
+
+		// This scaling moves output to roughly 0 to 1 range.
+		float DbScale = 1.f / FMath::Max(1.f, -NoiseFloorDb);
+
+		// Apply scaling for each channel.
+		for (int32 ChannelIndex = 0; ChannelIndex < NumChannels; ChannelIndex++)
+		{
+			Audio::AlignedFloatBuffer& Buffer = ChannelSpectrumBuffers[ChannelIndex];
+
+			Audio::ArrayMultiplyByConstantInPlace(Buffer, LinearScale);
+
+			Audio::ArrayPowerToDecibelInPlace(Buffer);
+
+			Audio::ArraySubtractByConstantInPlace(Buffer, NoiseFloorDb);
+
+			Audio::ArrayMultiplyByConstantInPlace(Buffer, DbScale);
+		}
+	}
+}
+
+void FNiagaraDataInterfaceProxySpectrum::SetAudioFormat(int32 InNumChannels, float InSampleRate)
+{
+	if (InSampleRate != SampleRate)
+	{
+		ResizeAudioTransform(MinimumFrequency, MaximumFrequency, InSampleRate, NumBands);
+	}
+
+	if (InNumChannels != NumChannels)
+	{
+		ResizeSpectrumBuffer(InNumChannels, NumBands);
+	}
+
+	if ((InSampleRate != SampleRate) || (InNumChannels != NumChannels))
+	{
+		int32 FFTSize = FFTAlgorithm.IsValid() ? FFTAlgorithm->Size() : 0;
+
+		ResizeSlidingBuffer(InNumChannels, FFTSize);
+	}
+
+	NumChannels = InNumChannels;
+	SampleRate = InSampleRate;
+}
+
+void FNiagaraDataInterfaceProxySpectrum::ResizeSpectrumBuffer(int32 InNumChannels, int32 InNumBands)
+{
+	for (int32 ChannelIndex = ChannelSpectrumBuffers.Num(); ChannelIndex < InNumChannels; ChannelIndex++)
+	{
+		ChannelSpectrumBuffers.AddDefaulted();
+	}
+
+	for (int32 ChannelIndex = 0; ChannelIndex < InNumChannels; ChannelIndex++)
+	{
+		Audio::AlignedFloatBuffer& Buffer = ChannelSpectrumBuffers[ChannelIndex];
+		Buffer.Reset();
+		if (InNumBands > 0)
+		{
+			Buffer.AddUninitialized(InNumBands);
+		}
+	}
+}
+
+void FNiagaraDataInterfaceProxySpectrum::ResizeSlidingBuffer(int32 InNumChannels, int32 InFFTSize)
+{
+	int32 NumWindowSamples = FMath::Max(1, InFFTSize);
+	NumWindowSamples = FMath::Min(NumWindowSamples, 1024);
+	if (InNumChannels > 0)
+	{
+		NumWindowSamples *= InNumChannels;
+	}
+
+	// 50% overlap
+	int32 NumHopSamples = FMath::Max(1, NumWindowSamples / 2);
+
+	SlidingBuffer = MakeUnique<FSlidingBuffer>(NumWindowSamples, NumHopSamples);
+}
+
+void FNiagaraDataInterfaceProxySpectrum::ResizeAudioTransform(float InMinimumFrequency, float InMaximumFrequency, float InSampleRate, int32 InNumBands)
+{
+	FFTAlgorithm.Reset();
+	CQTKernel.Reset();
+	FFTInputBuffer.Reset();
+	FFTOutputBuffer.Reset();
+	PowerSpectrumBuffer.Reset();
+	SpectrumBuffer.Reset();
+
+	FFTScale = 1.f;
+
+	if (InSampleRate <= 0.f)
+	{
+		return;
+	}
+
+	if (InNumBands < 1)
+	{
+		return;
+	}
+
+	float NumOctaves = GetNumOctaves(InMinimumFrequency, InMaximumFrequency);
+	float NumBandsPerOctave = GetNumBandsPerOctave(InNumBands, NumOctaves);
+
+	Audio::FFFTSettings FFTSettings = GetFFTSettings(InMinimumFrequency, InSampleRate, NumBandsPerOctave);
+	if (!Audio::FFFTFactory::AreFFTSettingsSupported(FFTSettings))
+	{
+		UE_LOG(LogNiagara, Warning, TEXT("FFT settings are not supported"));
+		return;
+	}
+
+	FFTAlgorithm = Audio::FFFTFactory::NewFFTAlgorithm(FFTSettings);
+
+	if (!FFTAlgorithm.IsValid())
+	{
+		UE_LOG(LogNiagara, Error, TEXT("Failed to create FFT"));
+		return;
+	}
+
+	FFTInputBuffer.AddZeroed(FFTAlgorithm->NumInputFloats());
+	FFTOutputBuffer.AddUninitialized(FFTAlgorithm->NumOutputFloats());
+	PowerSpectrumBuffer.AddUninitialized(FFTOutputBuffer.Num() / 2);
+	SpectrumBuffer.AddUninitialized(InNumBands);
+
+	float FFTSize = static_cast<float>(FFTAlgorithm->Size());
+	check(FFTSize > 0.f);
+
+	// We want to have all FFT implementations
+	// return the same scaling so that the energy conservatino property of the fourier transform
+	// is supported.  This scaling factor is applied to the power spectrum, so we square the 
+	// scaling we would have performed on the magnitude spectrum.
+	switch (FFTAlgorithm->ForwardScaling())
+	{
+		case Audio::EFFTScaling::MultipliedByFFTSize:
+			FFTScale = 1.f / (FFTSize * FFTSize);
+			break;
+			
+		case Audio::EFFTScaling::MultipliedBySqrtFFTSize:
+			FFTScale = 1.f / FFTSize;
+			break;
+
+		case Audio::EFFTScaling::DividedByFFTSize:
+			FFTScale = FFTSize * FFTSize;
+			break;
+
+		case Audio::EFFTScaling::DividedBySqrtFFTSize:
+			FFTScale = FFTSize;
+			break;
+
+		default:
+			FFTScale = 1.f;
+			break;
+	}
+
+	float BandwidthStretch = GetBandwidthStretch(InSampleRate, FFTSize, NumBandsPerOctave, InMinimumFrequency);
+
+	Audio::FPseudoConstantQKernelSettings CQTSettings = GetConstantQSettings(InMinimumFrequency, InMaximumFrequency, InNumBands, NumBandsPerOctave, BandwidthStretch);
+	CQTKernel = Audio::NewPseudoConstantQKernelTransform(CQTSettings, FFTAlgorithm->Size(), InSampleRate);
+
+	if (!CQTKernel.IsValid())
+	{
+		UE_LOG(LogNiagara, Error, TEXT("Failed to create CQT kernel."));
+		return;
+	}
+}
+
+void FNiagaraDataInterfaceProxySpectrum::Clamp(float& InMinimumFrequency, float& InMaximumFrequency, int32& InNumBands) 
+{
+	const int32 MinimumSupportedNumBands = 1;
+	const int32 MaximumSupportedNumBands = 16384;
+	const float MinimumSupportedFrequency = 20.f;
+	const float MaximumSupportedFrequency = 20000.f;
+
+	InMinimumFrequency = FMath::Max(InMinimumFrequency, MinimumSupportedFrequency);
+	InMaximumFrequency = FMath::Min(InMaximumFrequency, MaximumSupportedFrequency);
+	InMaximumFrequency = FMath::Max(InMinimumFrequency, InMaximumFrequency);
+
+	InNumBands = FMath::Clamp(InNumBands, MinimumSupportedNumBands, MaximumSupportedNumBands);
+}
+
+float FNiagaraDataInterfaceProxySpectrum::GetNumOctaves(float InMinimumFrequency, float InMaximumFrequency) 
+{
+	const float MinimumSupportedNumOctaves = 0.01f;
+
+	float NumOctaves = FMath::Log2(InMaximumFrequency) - FMath::Log2(InMinimumFrequency);
+	NumOctaves = FMath::Max(MinimumSupportedNumOctaves, NumOctaves);
+
+	return NumOctaves;
+}
+
+float FNiagaraDataInterfaceProxySpectrum::GetNumBandsPerOctave(int32 InNumBands, float InNumOctaves)
+{
+	InNumOctaves = FMath::Max(0.01f, InNumOctaves);
+	const float MinimumSupportedNumBandsPerOctave = 0.01f;
+
+	float NumBandsPerOctave = FMath::Max(MinimumSupportedNumBandsPerOctave, static_cast<float>(InNumBands) / InNumOctaves);
+
+	return NumBandsPerOctave;
+}
+
+float FNiagaraDataInterfaceProxySpectrum::GetBandwidthStretch(float InSampleRate, float InFFTSize, float InNumBandsPerOctave, float InMinimumFrequency)
+{
+	InFFTSize = FMath::Max(1.f, InFFTSize);
+
+	float MinimumFrequencySpacing = (FMath::Pow(2.f, 1.f / InNumBandsPerOctave) - 1.f) * InMinimumFrequency;
+	MinimumFrequencySpacing = FMath::Max(0.01f, MinimumFrequencySpacing);
+
+	float FFTBinSpacing = InSampleRate / InFFTSize;
+	float BandwidthStretch = FFTBinSpacing / MinimumFrequencySpacing;
+
+	return FMath::Clamp(BandwidthStretch, 0.5f, 2.f);
+}
+
+Audio::FPseudoConstantQKernelSettings FNiagaraDataInterfaceProxySpectrum::GetConstantQSettings(float InMinimumFrequency, float InMaximumFrequency, int32 InNumBands, float InNumBandsPerOctave, float InBandwidthStretch)
+{
+	Audio::FPseudoConstantQKernelSettings CQTKernelSettings;
+
+	CQTKernelSettings.NumBands = InNumBands;
+	CQTKernelSettings.KernelLowestCenterFreq = InMinimumFrequency; 
+	CQTKernelSettings.NumBandsPerOctave = InNumBandsPerOctave;
+	CQTKernelSettings.BandWidthStretch = InBandwidthStretch;
+	CQTKernelSettings.Normalization = Audio::EPseudoConstantQNormalization::EqualEnergy;
+
+	return CQTKernelSettings;
+}
+
+Audio::FFFTSettings FNiagaraDataInterfaceProxySpectrum::GetFFTSettings(float InMinimumFrequency, float InSampleRate, float InNumBandsPerOctave)
+{
+	const int32 MaximumSupportedLog2FFTSize = 13;
+	const int32 MinimumSupportedLog2FFTSize = 8;
+
+	InNumBandsPerOctave = FMath::Max(0.01f, InNumBandsPerOctave);
+
+	float MinimumFrequencySpacing = (FMath::Pow(2.f, 1.f / InNumBandsPerOctave) - 1.f) * InMinimumFrequency;
+
+	MinimumFrequencySpacing = FMath::Max(0.01f, MinimumFrequencySpacing);
+
+	int32 DesiredFFTSize = FMath::CeilToInt(3.f * InSampleRate / MinimumFrequencySpacing);
+
+	int32 Log2FFTSize = MinimumSupportedLog2FFTSize;
+	while ((DesiredFFTSize > (1 << Log2FFTSize)) && (Log2FFTSize < MaximumSupportedLog2FFTSize))
+	{
+		Log2FFTSize++;
+	}
+
+	Audio::FFFTSettings FFTSettings;
+
+	FFTSettings.Log2Size = Log2FFTSize;
+    FFTSettings.bArrays128BitAligned = true;
+    FFTSettings.bEnableHardwareAcceleration = true;
+
+	return FFTSettings;
+}
+
+
+UNiagaraDataInterfaceAudioSpectrum::UNiagaraDataInterfaceAudioSpectrum(FObjectInitializer const& ObjectInitializer)
+	: Super(ObjectInitializer)
+	, Resolution(128)
+	, MinimumFrequency(40.f)
+	, MaximumFrequency(10000.f)
+	, NoiseFloorDb(-60.f)
+	
+{
+	if (!HasAnyFlags(RF_ClassDefaultObject))
+	{
+		Proxy = MakeUnique<FNiagaraDataInterfaceProxySpectrum>(MinimumFrequency, MaximumFrequency, Resolution, NoiseFloorDb, 16384);
+	}
+}
+
+
+void UNiagaraDataInterfaceAudioSpectrum::GetFunctions(TArray<FNiagaraFunctionSignature>& OutFunctions)
+{
+	Super::GetFunctions(OutFunctions);
+
+	{
+		FNiagaraFunctionSignature GetSpectrumSignature;
+		GetSpectrumSignature.Name = GetSpectrumFunctionName;
+		GetSpectrumSignature.Inputs.Add(FNiagaraVariable(GetClass(), TEXT("Spectrum")));
+		GetSpectrumSignature.Inputs.Add(FNiagaraVariable(FNiagaraTypeDefinition::GetFloatDef(), TEXT("NormalizedPositionInSpectrum")));
+		GetSpectrumSignature.Inputs.Add(FNiagaraVariable(FNiagaraTypeDefinition::GetIntDef(), TEXT("ChannelIndex")));
+		GetSpectrumSignature.Outputs.Add(FNiagaraVariable(FNiagaraTypeDefinition::GetFloatDef(), TEXT("Amplitude")));
+
+		GetSpectrumSignature.bMemberFunction = true;
+		GetSpectrumSignature.bRequiresContext = false;
+		OutFunctions.Add(GetSpectrumSignature);
+	}
+
+	{
+		FNiagaraFunctionSignature NumChannelsSignature;
+		NumChannelsSignature.Name = GetNumChannelsFunctionName;
+		NumChannelsSignature.Inputs.Add(FNiagaraVariable(GetClass(), TEXT("Spectrum")));
+		NumChannelsSignature.Outputs.Add(FNiagaraVariable(FNiagaraTypeDefinition::GetIntDef(), TEXT("NumChannels")));
+
+		NumChannelsSignature.bMemberFunction = true;
+		NumChannelsSignature.bRequiresContext = false;
+		OutFunctions.Add(NumChannelsSignature);
+	}
+}
+
+DEFINE_NDI_DIRECT_FUNC_BINDER(UNiagaraDataInterfaceAudioSpectrum, GetSpectrumValue);
+DEFINE_NDI_DIRECT_FUNC_BINDER(UNiagaraDataInterfaceAudioSpectrum, GetNumChannels);
+
+void UNiagaraDataInterfaceAudioSpectrum::GetVMExternalFunction(const FVMExternalFunctionBindingInfo& BindingInfo, void* InstanceData, FVMExternalFunction &OutFunc)
+{
+	if (BindingInfo.Name == GetSpectrumFunctionName)
+	{
+		NDI_FUNC_BINDER(UNiagaraDataInterfaceAudioSpectrum, GetSpectrumValue)::Bind(this, OutFunc);
+	}
+	else if (BindingInfo.Name == GetNumChannelsFunctionName)
+	{
+		NDI_FUNC_BINDER(UNiagaraDataInterfaceAudioSpectrum, GetNumChannels)::Bind(this, OutFunc);
+	}
+	else
+	{
+		ensureMsgf(false, TEXT("Error! Function defined for this class but not bound."));
+	}
+}
+
+bool UNiagaraDataInterfaceAudioSpectrum::GetFunctionHLSL(const FNiagaraDataInterfaceGPUParamInfo& ParamInfo, const FNiagaraDataInterfaceGeneratedFunction& FunctionInfo, int FunctionInstanceIndex, FString& OutHLSL)
+{
+	bool ParentRet = Super::GetFunctionHLSL(ParamInfo, FunctionInfo, FunctionInstanceIndex, OutHLSL);
+	if (ParentRet)
+	{
+		return true;
+	}
+	else if (FunctionInfo.DefinitionName == GetSpectrumFunctionName)
+	{
+		/**
+		 * float FrameIndex = NormalizedPositionInBuffer * SpectrumBuffer.Num();
+		 * int32 LowerFrameIndex = FMath::FloorToInt(FrameIndex);
+		 * float LowerFrameAmplitude = SpectrumBuffer[LowerFrameIndex * NumChannels + ChannelIndex];
+		 * int32 HigherFrameIndex = FMath::CeilToInt(FrameIndex);
+		 * float HigherFrameAmplitude = SpectrumBuffer[HigherFrameIndex * NumChannels + ChannelIndex];
+		 * float Fraction = HigherFrameIndex - FrameIndex;
+		 * return FMath::Lerp<float>(LowerFrameAmplitude, HigherFrameAmplitude, Fraction);
+		 */
+		static const TCHAR *FormatBounds = TEXT(R"(
+			void {FunctionName}(float In_NormalizedPosition, int In_ChannelIndex, out float Out_Val)
+			{
+				float FrameIndex = In_NormalizedPosition * {SpectrumResolution};
+				int LowerIndex = floor(FrameIndex);
+				LowerIndex = LowerIndex < {SpectrumResolution} ? LowerIndex : {SpectrumResolution} - 1.0;
+				int UpperIndex =  LowerIndex < {SpectrumResolution} ? LowerIndex + 1.0 : LowerIndex;
+				float Fraction = FrameIndex - LowerIndex;
+				float LowerValue = {SpectrumBuffer}.Load(In_ChannelIndex * {SpectrumResolution} + LowerIndex);
+				float UpperValue = {SpectrumBuffer}.Load(In_ChannelIndex * {SpectrumResolution} + UpperIndex);
+
+				Out_Val = lerp(LowerValue, UpperValue, Fraction);
+			}
+		)");
+		TMap<FString, FStringFormatArg> ArgsBounds = {
+			{TEXT("FunctionName"), FStringFormatArg(FunctionInfo.InstanceName)},
+			{TEXT("ChannelCount"), FStringFormatArg(NumChannelsName + ParamInfo.DataInterfaceHLSLSymbol)},
+			{TEXT("SpectrumResolution"), FStringFormatArg(Resolution)},
+			{TEXT("SpectrumBuffer"), FStringFormatArg(GetSpectrumName + ParamInfo.DataInterfaceHLSLSymbol)},
+		};
+		OutHLSL += FString::Format(FormatBounds, ArgsBounds);
+		return true;
+	}
+	else if (FunctionInfo.DefinitionName == GetNumChannelsFunctionName)
+	{
+		// See GetNumChannels()
+		static const TCHAR *FormatBounds = TEXT(R"(
+			void {FunctionName}(out int Out_Val)
+			{
+				Out_Val = {ChannelCount};
+			}
+		)");
+		TMap<FString, FStringFormatArg> ArgsBounds = {
+			{TEXT("FunctionName"), FStringFormatArg(FunctionInfo.InstanceName)},
+			{TEXT("ChannelCount"), FStringFormatArg(NumChannelsName + ParamInfo.DataInterfaceHLSLSymbol)},
+		};
+		OutHLSL += FString::Format(FormatBounds, ArgsBounds);
+		return true;
+	}
+	else
+	{
+		return false;
+	}
+}
+
+void UNiagaraDataInterfaceAudioSpectrum::GetParameterDefinitionHLSL(const FNiagaraDataInterfaceGPUParamInfo& ParamInfo, FString& OutHLSL)
+{
+	Super::GetParameterDefinitionHLSL(ParamInfo, OutHLSL);
+
+	static const TCHAR *FormatDeclarations = TEXT(R"(				
+		Buffer<float> {GetSpectrumName};
+		int {NumChannelsName};
+
+	)");
+	TMap<FString, FStringFormatArg> ArgsDeclarations = {
+		{ TEXT("GetSpectrumName"),    FStringFormatArg(GetSpectrumName + ParamInfo.DataInterfaceHLSLSymbol) },
+		{ TEXT("NumChannelsName"),    FStringFormatArg(NumChannelsName + ParamInfo.DataInterfaceHLSLSymbol) },
+	};
+	OutHLSL += FString::Format(FormatDeclarations, ArgsDeclarations);
+}
+
+struct FNiagaraDataInterfaceParametersCS_AudioSpectrum : public FNiagaraDataInterfaceParametersCS
+{
+	DECLARE_INLINE_TYPE_LAYOUT(FNiagaraDataInterfaceParametersCS_AudioSpectrum, NonVirtual);
+
+	void Bind(const FNiagaraDataInterfaceGPUParamInfo& ParameterInfo, const class FShaderParameterMap& ParameterMap)
+	{
+		NumChannels.Bind(ParameterMap, *(UNiagaraDataInterfaceAudioSpectrum::NumChannelsName + ParameterInfo.DataInterfaceHLSLSymbol));
+		SpectrumBuffer.Bind(ParameterMap, *(UNiagaraDataInterfaceAudioSpectrum::GetSpectrumName + ParameterInfo.DataInterfaceHLSLSymbol));
+	}
+
+	void Set(FRHICommandList& RHICmdList, const FNiagaraDataInterfaceSetArgs& Context) const 
+	{
+		check(IsInRenderingThread());
+
+		FRHIComputeShader* ComputeShaderRHI = Context.Shader.GetComputeShader();
+
+		FNiagaraDataInterfaceProxySpectrum* NDI = (FNiagaraDataInterfaceProxySpectrum*)Context.DataInterface;
+		FReadBuffer& SpectrumBufferSRV = NDI->ComputeAndPostSRV();
+
+		SetShaderValue(RHICmdList, ComputeShaderRHI, NumChannels, NDI->GetNumChannels());
+		RHICmdList.SetShaderResourceViewParameter(ComputeShaderRHI, SpectrumBuffer.GetBaseIndex(), SpectrumBufferSRV.SRV);
+	}
+
+	LAYOUT_FIELD(FShaderParameter, NumChannels);
+	LAYOUT_FIELD(FShaderResourceParameter,SpectrumBuffer);
+};
+
+IMPLEMENT_NIAGARA_DI_PARAMETER(UNiagaraDataInterfaceAudioSpectrum, FNiagaraDataInterfaceParametersCS_AudioSpectrum);
+
+
+#if WITH_EDITOR
+void UNiagaraDataInterfaceAudioSpectrum::PostEditChangeProperty(struct FPropertyChangedEvent& PropertyChangedEvent)
+{
+	Super::PostEditChangeProperty(PropertyChangedEvent);
+	
+	static const FName ResolutionFName = GET_MEMBER_NAME_CHECKED(UNiagaraDataInterfaceAudioSpectrum, Resolution);
+	static const FName MinimumFrequencyFName = GET_MEMBER_NAME_CHECKED(UNiagaraDataInterfaceAudioSpectrum, MinimumFrequency);
+	static const FName MaximumFrequencyFName = GET_MEMBER_NAME_CHECKED(UNiagaraDataInterfaceAudioSpectrum, MaximumFrequency);
+	static const FName NoiseFloorFName = GET_MEMBER_NAME_CHECKED(UNiagaraDataInterfaceAudioSpectrum, NoiseFloorDb);
+
+	// Regenerate on save any compressed sound formats or if analysis needs to be re-done
+	if (FProperty* PropertyThatChanged = PropertyChangedEvent.Property)
+	{
+		const FName& Name = PropertyThatChanged->GetFName();
+		if (Name == ResolutionFName || Name == MaximumFrequencyFName || Name == MinimumFrequencyFName)
+		{
+			GetProxyAs<FNiagaraDataInterfaceProxySpectrum>()->UpdateCQT(MinimumFrequency, MaximumFrequency, Resolution);
+		}
+		else if (Name == NoiseFloorFName)
+		{
+			GetProxyAs<FNiagaraDataInterfaceProxySpectrum>()->UpdateNoiseFloor(NoiseFloorDb);
+		}
+	}
+}
+#endif //WITH_EDITOR
+
+void UNiagaraDataInterfaceAudioSpectrum::PostInitProperties()
+{
+	Super::PostInitProperties();
+
+	if (HasAnyFlags(RF_ClassDefaultObject))
+	{
+		FNiagaraTypeRegistry::Register(FNiagaraTypeDefinition(GetClass()), /*bCanBeParameter*/ true, /*bCanBePayload*/ false, /*bIsUserDefined*/ false);
+	}
+
+	GetProxyAs<FNiagaraDataInterfaceProxySpectrum>()->UpdateCQT(MinimumFrequency, MaximumFrequency, Resolution);
+	GetProxyAs<FNiagaraDataInterfaceProxySpectrum>()->UpdateNoiseFloor(NoiseFloorDb);
+}
+
+void UNiagaraDataInterfaceAudioSpectrum::PostLoad()
+{
+	Super::PostLoad();
+	GetProxyAs<FNiagaraDataInterfaceProxySpectrum>()->UpdateCQT(MinimumFrequency, MaximumFrequency, Resolution);
+	GetProxyAs<FNiagaraDataInterfaceProxySpectrum>()->UpdateNoiseFloor(NoiseFloorDb);
+}
+
+bool UNiagaraDataInterfaceAudioSpectrum::CopyToInternal(UNiagaraDataInterface* Destination) const
+{
+	Super::CopyToInternal(Destination);
+
+	UNiagaraDataInterfaceAudioSpectrum* CastedDestination = Cast<UNiagaraDataInterfaceAudioSpectrum>(Destination);
+
+	if (CastedDestination)
+	{
+		CastedDestination->Resolution = Resolution;
+		CastedDestination->MinimumFrequency = MinimumFrequency;
+		CastedDestination->MaximumFrequency = MaximumFrequency;
+		CastedDestination->NoiseFloorDb = NoiseFloorDb;
+	}
+	
+	return true;
+}
+
+
+// Global VM function names, also used by the shaders code generation methods.
+const FName UNiagaraDataInterfaceAudioSpectrum::GetSpectrumFunctionName("AudioSpectrum");
+const FName UNiagaraDataInterfaceAudioSpectrum::GetNumChannelsFunctionName("GetNumChannels");
+
+// Global variable prefixes, used in HLSL parameter declarations.
+const FString UNiagaraDataInterfaceAudioSpectrum::GetSpectrumName(TEXT("SpectrumBuffer_"));
+const FString UNiagaraDataInterfaceAudioSpectrum::NumChannelsName(TEXT("NumChannels_"));
+
+#undef LOCTEXT_NAMESPACE
diff --git a/Engine/Plugins/FX/Niagara/Source/Niagara/Private/NiagaraModule.cpp b/Engine/Plugins/FX/Niagara/Source/Niagara/Private/NiagaraModule.cpp
--- a/Engine/Plugins/FX/Niagara/Source/Niagara/Private/NiagaraModule.cpp
+++ b/Engine/Plugins/FX/Niagara/Source/Niagara/Private/NiagaraModule.cpp
@@ -265,6 +265,9 @@ void INiagaraModule::StartupModule()
 	{
 		return new NiagaraEmitterInstanceBatcher(InFeatureLevel, InShaderPlatform, InGPUSortManager);
 	}));
+
+	// Needed for NiagaraDataInterfaceAudioSpectrum
+	FModuleManager::Get().LoadModule(TEXT("SignalProcessing"));
 }
 
 void INiagaraModule::ShutdownRenderingResources()
diff --git a/Engine/Source/Runtime/SignalProcessing/Private/FloatArrayMath.cpp b/Engine/Source/Runtime/SignalProcessing/Private/FloatArrayMath.cpp
--- a/Engine/Source/Runtime/SignalProcessing/Private/FloatArrayMath.cpp
+++ b/Engine/Source/Runtime/SignalProcessing/Private/FloatArrayMath.cpp
@@ -3,13 +3,18 @@
 #include "DSP/FloatArrayMath.h"
 #include "CoreMinimal.h"
 
-namespace
-{
-	const float LOGE10 = FMath::Loge(10.f);
-}
 
 namespace Audio
 {
+	namespace MathIntrinsics
+	{
+		const float Loge10 = FMath::Loge(10.f);
+		const int32 SimdMask = 0xFFFFFFFC;
+		const int32 NotSimdMask = 0x00000003;
+
+		
+	}
+
 	void ArrayCumulativeSum(TArrayView<const float> InView, TArray<float>& OutData)
 	{
 		// Initialize output data
@@ -299,23 +304,114 @@ namespace Audio
 		}
 	}
 
-	void ArrayMultiplyByConstantInPlace(TArrayView<float> InView, float InMultiplier)
+	void ArrayMultiplyByConstantInPlace(TArrayView<float> InValues, float InMultiplier)
 	{
-		const int32 Num = InView.Num();
-		float* InViewData = InView.GetData();
+		const int32 Num = InValues.Num();
+		float* InData = InValues.GetData();
+
 		for (int32 i = 0; i < Num; i++)
 		{
-			InViewData[i] *= InMultiplier;
+			InData[i] *= InMultiplier;
 		}
 	}
 
-	void ArraySubtractByConstantInPlace(TArrayView<float> InView, float InSubtrahend)
+	void ArrayMultiplyByConstantInPlace(AlignedFloatBuffer& InValues, float InMultiplier)
 	{
-		const int32 Num = InView.Num();
-		float* InViewData = InView.GetData();
+		const int32 Num = InValues.Num();
+		const int32 NumToSimd = Num & MathIntrinsics::SimdMask;
+		const int32 NumNotToSimd = Num & MathIntrinsics::NotSimdMask;
+
+		float* InData = InValues.GetData();
+
+		if (NumToSimd)
+		{
+			MultiplyBufferByConstantInPlace(InData, NumToSimd, InMultiplier);
+		}
+
+		if (NumNotToSimd)
+		{
+			TArrayView<float> ValuesView(&InData[NumToSimd],  NumNotToSimd);
+
+			ArrayMultiplyByConstantInPlace(ValuesView, InMultiplier);
+		}
+		
+	}
+
+	void ArrayAddInPlace(TArrayView<const float> InValues, TArrayView<float> InAccumulateValues)
+	{
+		check(InValues.Num() == InAccumulateValues.Num());
+
+		const int32 Num = InValues.Num();
+
+		const float* InData = InValues.GetData();
+		float* InAccumulateData = InAccumulateValues.GetData();
+
 		for (int32 i = 0; i < Num; i++)
 		{
-			InViewData[i] -= InSubtrahend;
+			InAccumulateData[i] += InData[i];
+		}
+	}
+
+	void ArrayAddInPlace(const AlignedFloatBuffer& InValues, AlignedFloatBuffer& InAccumulateValues)
+	{
+		check(InValues.Num() == InAccumulateValues.Num());
+
+		const int32 Num = InAccumulateValues.Num();
+		const int32 NumToSimd = Num & MathIntrinsics::SimdMask;
+		const int32 NumNotToSimd = Num & MathIntrinsics::NotSimdMask;
+
+		const float* InData = InValues.GetData();
+		float* InAccumulateData = InAccumulateValues.GetData();
+
+		for (int32 i = 0; i < NumToSimd; i += 4)
+		{
+			VectorRegister VectorData = VectorLoadAligned(&InData[i]);
+			VectorRegister VectorAccumData = VectorLoadAligned(&InAccumulateData[i]);
+
+			VectorRegister VectorOut = VectorAdd(VectorData, VectorAccumData);
+			VectorStoreAligned(VectorOut, &InAccumulateData[i]);
+		}
+
+		if (NumNotToSimd)
+		{
+			TArrayView<const float> ValuesView(&InData[NumToSimd],  NumNotToSimd);
+			TArrayView<float> AccumulateView(&InAccumulateData[NumToSimd], NumNotToSimd);
+
+			ArrayAddInPlace(ValuesView, AccumulateView);
+		}
+	}
+
+	void ArraySubtractByConstantInPlace(TArrayView<float> InValues, float InSubtrahend)
+	{
+		const int32 Num = InValues.Num();
+		float* InValuesData = InValues.GetData();
+		for (int32 i = 0; i < Num; i++)
+		{
+			InValuesData[i] -= InSubtrahend;
+		}
+	}
+
+	void ArraySubtractByConstantInPlace(AlignedFloatBuffer& InValues, float InSubtrahend)
+	{
+		const int32 Num = InValues.Num();
+		const int32 NumToSimd = Num & MathIntrinsics::SimdMask;
+		const int32 NumNotToSimd = Num & MathIntrinsics::NotSimdMask;
+
+		float* InData = InValues.GetData();
+
+		const VectorRegister VectorSubtrahend = VectorSetFloat1(InSubtrahend);
+
+		for (int32 i = 0; i < NumToSimd; i += 4)
+		{
+			VectorRegister VectorData = VectorLoadAligned(&InData[i]);
+			VectorData = VectorSubtract(VectorData, VectorSubtrahend);
+			VectorStoreAligned(VectorData, &InData[i]);
+		}
+
+		if (NumNotToSimd)
+		{
+			TArrayView<float> View(&InData[NumToSimd], NumNotToSimd);
+			ArraySubtractByConstantInPlace(View, InSubtrahend);
 		}
 	}
 
@@ -344,26 +440,139 @@ namespace Audio
 		}
 	}
 
-	void ArrayMagnitudeToDecibelInPlace(TArrayView<float> InView)
+	void ArrayMagnitudeToDecibelInPlace(TArrayView<float> InValues)
 	{
-		const int32 Num = InView.Num();
-		float* InViewData = InView.GetData();
+		const int32 Num = InValues.Num();
+		float* InValuesData = InValues.GetData();
 		for (int32 i = 0; i < Num; i++)
 		{
-			InViewData[i] = 20.f * FMath::Loge(InViewData[i]) / LOGE10;
+			InValuesData[i] = 20.f * FMath::Loge(InValuesData[i]) / MathIntrinsics::Loge10;
 		}
 	}
 
-	void ArrayPowerToDecibelInPlace(TArrayView<float> InView)
+	void ArrayMagnitudeToDecibelInPlace(AlignedFloatBuffer& InValues)
 	{
-		const int32 Num = InView.Num();
-		float* InViewData = InView.GetData();
+		const int32 Num = InValues.Num();
+		const int32 NumToSimd = Num & MathIntrinsics::SimdMask;
+		const int32 NumNotToSimd = Num & MathIntrinsics::NotSimdMask;
+
+		float* InData = InValues.GetData();
+
+		const VectorRegister VectorScale = VectorSetFloat1(20.f / MathIntrinsics::Loge10);
+
+		for (int32 i = 0; i < NumToSimd; i += 4)
+		{
+			VectorRegister VectorData = VectorLoadAligned(&InData[i]);
+
+			VectorData = VectorLog(VectorData);
+			VectorData = VectorMultiply(VectorData, VectorScale);
+
+			VectorStoreAligned(VectorData, &InData[i]);
+		}
+
+		if (NumNotToSimd)
+		{
+			TArrayView<float> InView(&InData[NumToSimd], NumNotToSimd);
+			ArrayMagnitudeToDecibelInPlace(InView);
+		}
+	}
+
+	void ArrayPowerToDecibelInPlace(TArrayView<float> InValues)
+	{
+		const int32 Num = InValues.Num();
+		float* InValuesData = InValues.GetData();
 		for (int32 i = 0; i < Num; i++)
 		{
-			InViewData[i] = 10.f * FMath::Loge(InViewData[i]) / LOGE10;
+			InValuesData[i] = 10.f * FMath::Loge(InValuesData[i]) / MathIntrinsics::Loge10;
+		}
+	}
+
+	void ArrayPowerToDecibelInPlace(AlignedFloatBuffer& InValues)
+	{
+		const int32 Num = InValues.Num();
+		const int32 NumToSimd = Num & MathIntrinsics::SimdMask;
+		const int32 NumNotToSimd = Num & MathIntrinsics::NotSimdMask;
+
+		float* InData = InValues.GetData();
+
+		const VectorRegister VectorScale = VectorSetFloat1(10.f / MathIntrinsics::Loge10);
+
+		for (int32 i = 0; i < NumToSimd; i += 4)
+		{
+			VectorRegister VectorData = VectorLoadAligned(&InData[i]);
+
+			VectorData = VectorLog(VectorData);
+			VectorData = VectorMultiply(VectorData, VectorScale);
+
+			VectorStoreAligned(VectorData, &InData[i]);
+		}
+
+		if (NumNotToSimd)
+		{
+			TArrayView<float> InView(&InData[NumToSimd], NumNotToSimd);
+			ArrayPowerToDecibelInPlace(InView);
+		}
+	}
+
+	void ArrayComplexToPower(TArrayView<const float> InComplexValues, TArrayView<float> OutPowerValues)
+	{
+		check((InComplexValues.Num() % 2) == 0);
+		check(InComplexValues.Num() == (OutPowerValues.Num() * 2));
+
+		const int32 NumOut = OutPowerValues.Num();
+
+		const float* InComplexData = InComplexValues.GetData();
+		float* OutPowerData = OutPowerValues.GetData();
+
+		for (int32 i = 0; i < NumOut; i++)
+		{
+			int32 ComplexPos = 2 * i;
+
+			float RealValue = InComplexData[ComplexPos];
+			float ImagValue = InComplexData[ComplexPos + 1];
+
+			OutPowerData[i] = (RealValue * RealValue) + (ImagValue * ImagValue);
 		}
 	}
 
+	void ArrayComplexToPower(const AlignedFloatBuffer& InComplexValues, AlignedFloatBuffer& OutPowerValues)
+	{
+		check((InComplexValues.Num() % 2) == 0);
+		check(InComplexValues.Num() == (OutPowerValues.Num() * 2));
+
+		const int32 NumOut = OutPowerValues.Num();
+		const int32 NumToSimd = NumOut & MathIntrinsics::SimdMask;
+		const int32 NumNotToSimd = NumOut & MathIntrinsics::NotSimdMask;
+
+		const float* InComplexData = InComplexValues.GetData();
+		float* OutPowerData = OutPowerValues.GetData();
+
+		for (int32 i = 0; i < NumToSimd; i += 4)
+		{
+			VectorRegister VectorComplex1 = VectorLoadAligned(&InComplexData[2 * i]);
+			VectorRegister VectorSquared1 = VectorMultiply (VectorComplex1, VectorComplex1);
+
+			VectorRegister VectorComplex2 = VectorLoadAligned(&InComplexData[(2 * i) + 4]);
+			VectorRegister VectorSquared2 = VectorMultiply (VectorComplex2, VectorComplex2);
+
+			VectorRegister VectorSquareReal = VectorShuffle(VectorSquared1, VectorSquared2, 0, 2, 0, 2);
+			VectorRegister VectorSquareImag = VectorShuffle(VectorSquared1, VectorSquared2, 1, 3, 1, 3);
+
+			VectorRegister VectorOut = VectorAdd(VectorSquareReal, VectorSquareImag);
+			
+			VectorStoreAligned(VectorOut, &OutPowerData[i]);
+		}
+
+		if (NumNotToSimd)
+		{
+			TArrayView<const float> ComplexView(&InComplexData[2 * NumToSimd], 2 * NumNotToSimd);
+			TArrayView<float> PowerView(&OutPowerData[NumToSimd], NumNotToSimd);
+
+			ArrayComplexToPower(ComplexView, PowerView);
+		}
+	}
+
+
 	FContiguousSparse2DKernelTransform::FContiguousSparse2DKernelTransform(const int32 NumInElements, const int32 NumOutElements)
 	:	NumIn(NumInElements)
 	,	NumOut(NumOutElements)
diff --git a/Engine/Source/Runtime/SignalProcessing/Public/DSP/FloatArrayMath.h b/Engine/Source/Runtime/SignalProcessing/Public/DSP/FloatArrayMath.h
--- a/Engine/Source/Runtime/SignalProcessing/Public/DSP/FloatArrayMath.h
+++ b/Engine/Source/Runtime/SignalProcessing/Public/DSP/FloatArrayMath.h
@@ -83,7 +83,16 @@ namespace Audio
 	SIGNALPROCESSING_API void ArrayMinMaxNormalize(TArrayView<const float> InView, TArray<float>& OutArray);
 
 	/** Multiplies each element in InView by InMultiplier */
-	SIGNALPROCESSING_API void ArrayMultiplyByConstantInPlace(TArrayView<float> InView, float InMultiplier);
+	SIGNALPROCESSING_API void ArrayMultiplyByConstantInPlace(TArrayView<float> InValues, float InMultiplier);
+	SIGNALPROCESSING_API void ArrayMultiplyByConstantInPlace(AlignedFloatBuffer& InValues, float InMultiplier);
+
+	/** Add arrays element-wise in place. InAccumulateValues[i] += InValues[i]
+	 *
+	 *  InValues is the array to add.
+	 *  InAccumulateValues is the array which holds the sum.
+	 */
+	SIGNALPROCESSING_API void ArrayAddInPlace(TArrayView<const float> InValues, TArrayView<float> InAccumulateValues);
+	SIGNALPROCESSING_API void ArrayAddInPlace(const AlignedFloatBuffer& InValues, AlignedFloatBuffer& InAccumulateValues);
 
 	/** Subract arrays element-wise. OutArray = InMinuend - InSubtrahend
 	 *
@@ -93,14 +102,26 @@ namespace Audio
 	 */
 	SIGNALPROCESSING_API void ArraySubtract(TArrayView<const float> InMinuend, TArrayView<const float> InSubtrahend, TArray<float>& OutArray);
 
-	/** Subtract value from each element in InView */
-	SIGNALPROCESSING_API void ArraySubtractByConstantInPlace(TArrayView<float> InView, float InSubtrahend);
+	/** Subtract value from each element in InValues */
+	SIGNALPROCESSING_API void ArraySubtractByConstantInPlace(TArrayView<float> InValues, float InSubtrahend);
+	SIGNALPROCESSING_API void ArraySubtractByConstantInPlace(AlignedFloatBuffer& InValues, float InSubtrahend);
 
 	/** Convert magnitude values to decibel values in place. db = 20 * log10(val) */
-	SIGNALPROCESSING_API void ArrayMagnitudeToDecibelInPlace(TArrayView<float> InView);
+	SIGNALPROCESSING_API void ArrayMagnitudeToDecibelInPlace(TArrayView<float> InValues);
+	SIGNALPROCESSING_API void ArrayMagnitudeToDecibelInPlace(AlignedFloatBuffer& InValues);
 
 	/** Convert power values to decibel values in place. db = 10 * log10(val) */
-	SIGNALPROCESSING_API void ArrayPowerToDecibelInPlace(TArrayView<float> InView);
+	SIGNALPROCESSING_API void ArrayPowerToDecibelInPlace(TArrayView<float> InValues);
+	SIGNALPROCESSING_API void ArrayPowerToDecibelInPlace(AlignedFloatBuffer& InValues);
+
+
+	/** Compute power of complex data. Out[i] = Complex[2 * i] * Complex[2 * i] + Complex[2 * i + 1] * Complex[2 * i + 1] */
+	SIGNALPROCESSING_API void ArrayComplexToPower(TArrayView<float> InComplexSamples, TArray<float>& OutPowerSamples);
+
+	/** Compute power of complex data. Out[i] = Complex[2 * i] * Complex[2 * i] + Complex[2 * i + 1] * Complex[2 * i + 1] */
+	SIGNALPROCESSING_API void ArrayComplexToPower(const AlignedFloatBuffer& InComplexSamples, AlignedFloatBuffer& OutPowerSamples);
+
+
 
 	/** FContiguousSparse2DKernelTransform
 	 *
