Non-unity fixes.

diff --git a/Engine/Plugins/FX/Niagara/Source/NiagaraEditor/Private/ViewModels/NiagaraParameterNameViewModel.cpp b/Engine/Plugins/FX/Niagara/Source/NiagaraEditor/Private/ViewModels/NiagaraParameterNameViewModel.cpp
--- a/Engine/Plugins/FX/Niagara/Source/NiagaraEditor/Private/ViewModels/NiagaraParameterNameViewModel.cpp
+++ b/Engine/Plugins/FX/Niagara/Source/NiagaraEditor/Private/ViewModels/NiagaraParameterNameViewModel.cpp
@@ -10,6 +10,7 @@
 #include "NiagaraParameterPanelViewModel.h"
 #include "NiagaraEditorUtilities.h"
 #include "NiagaraEditorModule.h"
+#include "CoreMinimal.h"
 
 ///////////////////////////////////////////////////////////////////////////////
 /// Parameter Panel Entry Parameter Name ViewModel							///
diff --git a/Engine/Plugins/FX/Niagara/Source/NiagaraEditor/Private/ViewModels/NiagaraParameterNameViewModel.h b/Engine/Plugins/FX/Niagara/Source/NiagaraEditor/Private/ViewModels/NiagaraParameterNameViewModel.h
--- a/Engine/Plugins/FX/Niagara/Source/NiagaraEditor/Private/ViewModels/NiagaraParameterNameViewModel.h
+++ b/Engine/Plugins/FX/Niagara/Source/NiagaraEditor/Private/ViewModels/NiagaraParameterNameViewModel.h
@@ -4,9 +4,12 @@
 
 #include "NiagaraEditorCommon.h"
 #include "Templates/SharedPointer.h"
+#include "Types/SlateEnums.h"
 
 struct FCreateWidgetForActionData;
 class FNiagaraScriptToolkitParameterPanelViewModel;
+class SWidget;
+struct FNiagaraScriptVariableAndViewInfo;
 
 /** Interface for view models to the parameter name view widget. */
 class INiagaraParameterNameViewModel : public TSharedFromThis<INiagaraParameterNameViewModel>
diff --git a/Engine/Plugins/FX/Niagara/Source/NiagaraEditor/Private/ViewModels/NiagaraParameterPanelViewModel.h b/Engine/Plugins/FX/Niagara/Source/NiagaraEditor/Private/ViewModels/NiagaraParameterPanelViewModel.h
--- a/Engine/Plugins/FX/Niagara/Source/NiagaraEditor/Private/ViewModels/NiagaraParameterPanelViewModel.h
+++ b/Engine/Plugins/FX/Niagara/Source/NiagaraEditor/Private/ViewModels/NiagaraParameterPanelViewModel.h
@@ -4,6 +4,9 @@
 
 #include "NiagaraEditorCommon.h"
 #include "Types/SlateEnums.h"
+#include "EditorUndoClient.h"
+#include "EdGraph/EdGraphSchema.h"
+#include "Input/Events.h"
 
 class UNiagaraSystem;
 class FNiagaraSystemViewModel;
diff --git a/Engine/Plugins/FX/Niagara/Source/NiagaraEditor/Private/Widgets/SNiagaraParameterNameView.h b/Engine/Plugins/FX/Niagara/Source/NiagaraEditor/Private/Widgets/SNiagaraParameterNameView.h
--- a/Engine/Plugins/FX/Niagara/Source/NiagaraEditor/Private/Widgets/SNiagaraParameterNameView.h
+++ b/Engine/Plugins/FX/Niagara/Source/NiagaraEditor/Private/Widgets/SNiagaraParameterNameView.h
@@ -66,7 +66,7 @@ public:
 		SLATE_EVENT(FOnParameterRenamed, OnParameterRenamed)
 	SLATE_END_ARGS();
 
-	void Construct(const FArguments& InArgs, const TSharedPtr<INiagaraParameterNameViewModel>& InParameterNameViewModel);
+	void Construct(const FArguments& InArgs, const TSharedPtr<class INiagaraParameterNameViewModel>& InParameterNameViewModel);
 
 protected:
 	TSharedPtr<INiagaraParameterNameViewModel> ParameterNameViewModel;
diff --git a/Engine/Plugins/FX/Niagara/Source/NiagaraEditor/Private/Widgets/SNiagaraParameterPanel.h b/Engine/Plugins/FX/Niagara/Source/NiagaraEditor/Private/Widgets/SNiagaraParameterPanel.h
--- a/Engine/Plugins/FX/Niagara/Source/NiagaraEditor/Private/Widgets/SNiagaraParameterPanel.h
+++ b/Engine/Plugins/FX/Niagara/Source/NiagaraEditor/Private/Widgets/SNiagaraParameterPanel.h
@@ -6,6 +6,7 @@
 #include "Internationalization/Text.h"
 #include "NiagaraTypes.h"
 #include "NiagaraActions.h"
+#include "Framework/Commands/Commands.h"
 
 class SGraphActionMenu;
 class SEditableTextBox;
