Integrating 11621164

Fixing CDO access to proxy for Spectrum NDI

diff --git a/Engine/Plugins/FX/Niagara/Source/Niagara/Private/NiagaraDataInterfaceAudioSpectrum.cpp b/Engine/Plugins/FX/Niagara/Source/Niagara/Private/NiagaraDataInterfaceAudioSpectrum.cpp
--- a/Engine/Plugins/FX/Niagara/Source/Niagara/Private/NiagaraDataInterfaceAudioSpectrum.cpp
+++ b/Engine/Plugins/FX/Niagara/Source/Niagara/Private/NiagaraDataInterfaceAudioSpectrum.cpp
@@ -712,17 +712,19 @@ void UNiagaraDataInterfaceAudioSpectrum::PostEditChangeProperty(struct FProperty
 	static const FName MaximumFrequencyFName = GET_MEMBER_NAME_CHECKED(UNiagaraDataInterfaceAudioSpectrum, MaximumFrequency);
 	static const FName NoiseFloorFName = GET_MEMBER_NAME_CHECKED(UNiagaraDataInterfaceAudioSpectrum, NoiseFloorDb);
 
-	// Regenerate on save any compressed sound formats or if analysis needs to be re-done
-	if (FProperty* PropertyThatChanged = PropertyChangedEvent.Property)
+	if (!HasAnyFlags(RF_ClassDefaultObject))
 	{
-		const FName& Name = PropertyThatChanged->GetFName();
-		if (Name == ResolutionFName || Name == MaximumFrequencyFName || Name == MinimumFrequencyFName)
-		{
-			GetProxyAs<FNiagaraDataInterfaceProxySpectrum>()->UpdateCQT(MinimumFrequency, MaximumFrequency, Resolution);
-		}
-		else if (Name == NoiseFloorFName)
+		if (FProperty* PropertyThatChanged = PropertyChangedEvent.Property)
 		{
-			GetProxyAs<FNiagaraDataInterfaceProxySpectrum>()->UpdateNoiseFloor(NoiseFloorDb);
+			const FName& Name = PropertyThatChanged->GetFName();
+			if (Name == ResolutionFName || Name == MaximumFrequencyFName || Name == MinimumFrequencyFName)
+			{
+				GetProxyAs<FNiagaraDataInterfaceProxySpectrum>()->UpdateCQT(MinimumFrequency, MaximumFrequency, Resolution);
+			}
+			else if (Name == NoiseFloorFName)
+			{
+				GetProxyAs<FNiagaraDataInterfaceProxySpectrum>()->UpdateNoiseFloor(NoiseFloorDb);
+			}
 		}
 	}
 }
@@ -736,16 +738,22 @@ void UNiagaraDataInterfaceAudioSpectrum::PostInitProperties()
 	{
 		FNiagaraTypeRegistry::Register(FNiagaraTypeDefinition(GetClass()), /*bCanBeParameter*/ true, /*bCanBePayload*/ false, /*bIsUserDefined*/ false);
 	}
-
-	GetProxyAs<FNiagaraDataInterfaceProxySpectrum>()->UpdateCQT(MinimumFrequency, MaximumFrequency, Resolution);
-	GetProxyAs<FNiagaraDataInterfaceProxySpectrum>()->UpdateNoiseFloor(NoiseFloorDb);
+	else
+	{
+		GetProxyAs<FNiagaraDataInterfaceProxySpectrum>()->UpdateCQT(MinimumFrequency, MaximumFrequency, Resolution);
+		GetProxyAs<FNiagaraDataInterfaceProxySpectrum>()->UpdateNoiseFloor(NoiseFloorDb);
+	}
 }
 
 void UNiagaraDataInterfaceAudioSpectrum::PostLoad()
 {
 	Super::PostLoad();
-	GetProxyAs<FNiagaraDataInterfaceProxySpectrum>()->UpdateCQT(MinimumFrequency, MaximumFrequency, Resolution);
-	GetProxyAs<FNiagaraDataInterfaceProxySpectrum>()->UpdateNoiseFloor(NoiseFloorDb);
+
+	if (!HasAnyFlags(RF_ClassDefaultObject))
+	{
+		GetProxyAs<FNiagaraDataInterfaceProxySpectrum>()->UpdateCQT(MinimumFrequency, MaximumFrequency, Resolution);
+		GetProxyAs<FNiagaraDataInterfaceProxySpectrum>()->UpdateNoiseFloor(NoiseFloorDb);
+	}
 }
 
 bool UNiagaraDataInterfaceAudioSpectrum::CopyToInternal(UNiagaraDataInterface* Destination) const
