Add "Owner" scope to handle Engine.Owner namespaced parameters.

diff --git a/Engine/Plugins/FX/Niagara/Source/Niagara/Public/NiagaraTypes.h b/Engine/Plugins/FX/Niagara/Source/Niagara/Public/NiagaraTypes.h
--- a/Engine/Plugins/FX/Niagara/Source/Niagara/Public/NiagaraTypes.h
+++ b/Engine/Plugins/FX/Niagara/Source/Niagara/Public/NiagaraTypes.h
@@ -449,6 +449,8 @@ enum class ENiagaraParameterScope : uint32
 
 	Engine,
 
+	Owner,
+
 	System,
 
 	Emitter,
diff --git a/Engine/Plugins/FX/Niagara/Source/NiagaraEditor/Private/ViewModels/NiagaraParameterPanelViewModel.cpp b/Engine/Plugins/FX/Niagara/Source/NiagaraEditor/Private/ViewModels/NiagaraParameterPanelViewModel.cpp
--- a/Engine/Plugins/FX/Niagara/Source/NiagaraEditor/Private/ViewModels/NiagaraParameterPanelViewModel.cpp
+++ b/Engine/Plugins/FX/Niagara/Source/NiagaraEditor/Private/ViewModels/NiagaraParameterPanelViewModel.cpp
@@ -79,6 +79,8 @@ const NiagaraParameterPanelSectionID::Type NiagaraParameterPanelSectionID::GetSe
 	switch (InScope) {
 	case ENiagaraParameterScope::Engine:
 		return NiagaraParameterPanelSectionID::ENGINE;
+	case ENiagaraParameterScope::Owner:
+		return NiagaraParameterPanelSectionID::OWNER;
 	case ENiagaraParameterScope::User:
 		return NiagaraParameterPanelSectionID::USER;
 	case ENiagaraParameterScope::System:
@@ -143,6 +145,8 @@ ENiagaraParameterScope NiagaraParameterPanelSectionID::GetScopeForNewParametersI
 	switch (InSection) {
 	case NiagaraParameterPanelSectionID::ENGINE:
 		return ENiagaraParameterScope::Engine;
+	case NiagaraParameterPanelSectionID::OWNER:
+		return ENiagaraParameterScope::Owner;
 	case NiagaraParameterPanelSectionID::USER:
 		return ENiagaraParameterScope::User;
 	case NiagaraParameterPanelSectionID::SYSTEM:
diff --git a/Engine/Plugins/FX/Niagara/Source/NiagaraEditor/Private/ViewModels/Stack/NiagaraParameterHandle.cpp b/Engine/Plugins/FX/Niagara/Source/NiagaraEditor/Private/ViewModels/Stack/NiagaraParameterHandle.cpp
--- a/Engine/Plugins/FX/Niagara/Source/NiagaraEditor/Private/ViewModels/Stack/NiagaraParameterHandle.cpp
+++ b/Engine/Plugins/FX/Niagara/Source/NiagaraEditor/Private/ViewModels/Stack/NiagaraParameterHandle.cpp
@@ -15,6 +15,7 @@ const FName FNiagaraParameterHandle::ParameterCollectionNamespace(TEXT("NPC"));
 const FString FNiagaraParameterHandle::InitialPrefix(TEXT("Initial"));
 const FName FNiagaraParameterHandle::LocalNamespace(TEXT("Local"));
 const FName FNiagaraParameterHandle::InitialNamespace(TEXT("Initial"));
+const FName FNiagaraParameterHandle::OwnerNamespace(TEXT("Owner"));
 
 FNiagaraParameterHandle::FNiagaraParameterHandle() 
 {
diff --git a/Engine/Plugins/FX/Niagara/Source/NiagaraEditor/Private/ViewModels/Stack/NiagaraStackGraphUtilities.cpp b/Engine/Plugins/FX/Niagara/Source/NiagaraEditor/Private/ViewModels/Stack/NiagaraStackGraphUtilities.cpp
--- a/Engine/Plugins/FX/Niagara/Source/NiagaraEditor/Private/ViewModels/Stack/NiagaraStackGraphUtilities.cpp
+++ b/Engine/Plugins/FX/Niagara/Source/NiagaraEditor/Private/ViewModels/Stack/NiagaraStackGraphUtilities.cpp
@@ -2563,6 +2563,9 @@ FString FNiagaraStackGraphUtilities::GetNamespaceStringForScriptParameterScope(c
 		case ENiagaraParameterScope::Engine:
 			return PARAM_MAP_ENGINE_STR;
 			break;
+		case ENiagaraParameterScope::Owner:
+			return PARAM_MAP_ENGINE_OWNER_STR;
+			break;
 		case ENiagaraParameterScope::User:
 			return PARAM_MAP_USER_STR;
 			break;
@@ -2672,7 +2675,14 @@ void FNiagaraStackGraphUtilities::SetParameterMetaDataFromName(const FName& InVa
 		}
 		else if (Namespace == FNiagaraParameterHandle::EmitterNamespace)
 		{
-			OutScope = ENiagaraParameterScope::Emitter;
+			if (Namespace == FNiagaraParameterHandle::OwnerNamespace)
+			{
+				OutScope = ENiagaraParameterScope::Owner;
+			}
+			else
+			{
+				OutScope = ENiagaraParameterScope::Emitter;
+			}
 			bCheckIsInitialValue = true;
 		}
 		else if (Namespace == FNiagaraParameterHandle::ParticleAttributeNamespace)
diff --git a/Engine/Plugins/FX/Niagara/Source/NiagaraEditor/Public/NiagaraEditorCommon.h b/Engine/Plugins/FX/Niagara/Source/NiagaraEditor/Public/NiagaraEditorCommon.h
--- a/Engine/Plugins/FX/Niagara/Source/NiagaraEditor/Public/NiagaraEditorCommon.h
+++ b/Engine/Plugins/FX/Niagara/Source/NiagaraEditor/Public/NiagaraEditorCommon.h
@@ -201,6 +201,7 @@ namespace NiagaraParameterPanelSectionID
 
 		// Group B, for contextual information when editing systems
 		ENGINE,
+		OWNER,
 		USER,
 		SYSTEM,
 		EMITTER,
diff --git a/Engine/Plugins/FX/Niagara/Source/NiagaraEditor/Public/ViewModels/Stack/NiagaraParameterHandle.h b/Engine/Plugins/FX/Niagara/Source/NiagaraEditor/Public/ViewModels/Stack/NiagaraParameterHandle.h
--- a/Engine/Plugins/FX/Niagara/Source/NiagaraEditor/Public/ViewModels/Stack/NiagaraParameterHandle.h
+++ b/Engine/Plugins/FX/Niagara/Source/NiagaraEditor/Public/ViewModels/Stack/NiagaraParameterHandle.h
@@ -70,6 +70,7 @@ public:
 	static const FString InitialPrefix;
 	static const FName LocalNamespace;
 	static const FName InitialNamespace;
+	static const FName OwnerNamespace;
 
 private:
 	FName ParameterHandleName;
diff --git a/Engine/Plugins/FX/Niagara/Source/NiagaraEditorWidgets/Private/NiagaraEditorWidgetsStyle.cpp b/Engine/Plugins/FX/Niagara/Source/NiagaraEditorWidgets/Private/NiagaraEditorWidgetsStyle.cpp
--- a/Engine/Plugins/FX/Niagara/Source/NiagaraEditorWidgets/Private/NiagaraEditorWidgetsStyle.cpp
+++ b/Engine/Plugins/FX/Niagara/Source/NiagaraEditorWidgets/Private/NiagaraEditorWidgetsStyle.cpp
@@ -251,6 +251,7 @@ TSharedRef< FSlateStyleSet > FNiagaraEditorWidgetsStyle::Create()
 	Style->Set("NiagaraEditor.ScratchPad.HeaderColor", FLinearColor(FColor(48, 48, 48)));
 
 	Style->Set("NiagaraEditor.Scope.Engine", FLinearColor(FColor(230, 102, 102)));
+	Style->Set("NiagaraEditor.Scope.Owner", FLinearColor(FColor(210, 112, 112)));
 	Style->Set("NiagaraEditor.Scope.User", FLinearColor(FColor(114, 226, 254)));
 	Style->Set("NiagaraEditor.Scope.System", FLinearColor(FColor(1, 202, 252)));
 	Style->Set("NiagaraEditor.Scope.Emitter", FLinearColor(FColor(241, 99, 6)));
diff --git a/Engine/Plugins/FX/Niagara/Source/NiagaraEditorWidgets/Private/NiagaraEditorWidgetsUtilities.cpp b/Engine/Plugins/FX/Niagara/Source/NiagaraEditorWidgets/Private/NiagaraEditorWidgetsUtilities.cpp
--- a/Engine/Plugins/FX/Niagara/Source/NiagaraEditorWidgets/Private/NiagaraEditorWidgetsUtilities.cpp
+++ b/Engine/Plugins/FX/Niagara/Source/NiagaraEditorWidgets/Private/NiagaraEditorWidgetsUtilities.cpp
@@ -133,6 +133,8 @@ FName FNiagaraStackEditorWidgetsUtilities::GetColorNameForParameterScope(ENiagar
 	switch (ParameterScope) {
 	case ENiagaraParameterScope::Engine:
 		return "NiagaraEditor.Scope.Engine";
+	case ENiagaraParameterScope::Owner:
+		return "NiagaraEditor.Scope.Owner";
 	case ENiagaraParameterScope::User:
 		return "NiagaraEditor.Scope.User";
 	case ENiagaraParameterScope::System:
