Non-unity fixes.

diff --git a/Engine/Plugins/FX/Niagara/Source/NiagaraEditor/Private/ViewModels/NiagaraParameterNameViewModel.cpp b/Engine/Plugins/FX/Niagara/Source/NiagaraEditor/Private/ViewModels/NiagaraParameterNameViewModel.cpp
--- a/Engine/Plugins/FX/Niagara/Source/NiagaraEditor/Private/ViewModels/NiagaraParameterNameViewModel.cpp
+++ b/Engine/Plugins/FX/Niagara/Source/NiagaraEditor/Private/ViewModels/NiagaraParameterNameViewModel.cpp
@@ -11,6 +11,9 @@
 #include "NiagaraEditorUtilities.h"
 #include "NiagaraEditorModule.h"
 #include "CoreMinimal.h"
+#include "Modules/ModuleManager.h"
+#include "EdGraph/EdGraphPin.h"
+#include "NiagaraEditorCommon.h"
 
 ///////////////////////////////////////////////////////////////////////////////
 /// Parameter Panel Entry Parameter Name ViewModel							///
diff --git a/Engine/Plugins/FX/Niagara/Source/NiagaraEditor/Private/ViewModels/NiagaraParameterNameViewModel.h b/Engine/Plugins/FX/Niagara/Source/NiagaraEditor/Private/ViewModels/NiagaraParameterNameViewModel.h
--- a/Engine/Plugins/FX/Niagara/Source/NiagaraEditor/Private/ViewModels/NiagaraParameterNameViewModel.h
+++ b/Engine/Plugins/FX/Niagara/Source/NiagaraEditor/Private/ViewModels/NiagaraParameterNameViewModel.h
@@ -5,11 +5,13 @@
 #include "NiagaraEditorCommon.h"
 #include "Templates/SharedPointer.h"
 #include "Types/SlateEnums.h"
+#include "Misc/Attribute.h"
 
 struct FCreateWidgetForActionData;
 class FNiagaraScriptToolkitParameterPanelViewModel;
 class SWidget;
 struct FNiagaraScriptVariableAndViewInfo;
+class UEdGraphPin;
 
 /** Interface for view models to the parameter name view widget. */
 class INiagaraParameterNameViewModel : public TSharedFromThis<INiagaraParameterNameViewModel>
diff --git a/Engine/Plugins/FX/Niagara/Source/NiagaraEditor/Private/ViewModels/NiagaraParameterPanelViewModel.cpp b/Engine/Plugins/FX/Niagara/Source/NiagaraEditor/Private/ViewModels/NiagaraParameterPanelViewModel.cpp
--- a/Engine/Plugins/FX/Niagara/Source/NiagaraEditor/Private/ViewModels/NiagaraParameterPanelViewModel.cpp
+++ b/Engine/Plugins/FX/Niagara/Source/NiagaraEditor/Private/ViewModels/NiagaraParameterPanelViewModel.cpp
@@ -23,6 +23,8 @@
 #include "NiagaraSystem.h"
 #include "NiagaraObjectSelection.h"
 #include "EdGraph/EdGraphSchema.h"
+#include "Editor.h"
+#include "ScopedTransaction.h"
 
 #define LOCTEXT_NAMESPACE "NiagaraParameterPanelViewModel"
 
@@ -347,7 +349,7 @@ void FNiagaraSystemToolkitParameterPanelViewModel::RemoveParameter(const FNiagar
 
 bool FNiagaraSystemToolkitParameterPanelViewModel::CanRemoveParameter(const FNiagaraVariable& TargetVariableToRemove, const FNiagaraVariableMetaData& TargetVariableMetaData) const
 {
-	return TargetVariableMetaData.bCreatedInSystemEditor == true;
+	return TargetVariableMetaData.bCreatedInSystemEditor;
 }
 
 void FNiagaraSystemToolkitParameterPanelViewModel::RenameParameter(const FNiagaraVariable& TargetVariableToRename, const FNiagaraVariableMetaData& TargetVariableMetaData, const FText& NewVariableNameText) const
diff --git a/Engine/Plugins/FX/Niagara/Source/NiagaraEditor/Private/ViewModels/NiagaraParameterPanelViewModel.h b/Engine/Plugins/FX/Niagara/Source/NiagaraEditor/Private/ViewModels/NiagaraParameterPanelViewModel.h
--- a/Engine/Plugins/FX/Niagara/Source/NiagaraEditor/Private/ViewModels/NiagaraParameterPanelViewModel.h
+++ b/Engine/Plugins/FX/Niagara/Source/NiagaraEditor/Private/ViewModels/NiagaraParameterPanelViewModel.h
@@ -7,6 +7,8 @@
 #include "EditorUndoClient.h"
 #include "EdGraph/EdGraphSchema.h"
 #include "Input/Events.h"
+#include "Input/Reply.h"
+#include "Templates/SharedPointer.h"
 
 class UNiagaraSystem;
 class FNiagaraSystemViewModel;
