Fixed Vulkan ensure() when using the dummy texture SRV

The dummy resources need to be initialized.

diff --git a/Engine/Plugins/FX/Niagara/Source/Niagara/Private/NiagaraRenderer.cpp b/Engine/Plugins/FX/Niagara/Source/Niagara/Private/NiagaraRenderer.cpp
--- a/Engine/Plugins/FX/Niagara/Source/Niagara/Private/NiagaraRenderer.cpp
+++ b/Engine/Plugins/FX/Niagara/Source/Niagara/Private/NiagaraRenderer.cpp
@@ -37,13 +37,18 @@ public:
 
 	virtual void InitRHI() override
 	{
-		uint32 BytesPerElement = GPixelFormats[PixelFormat].BlockBytes;
-		uint32 NumElements = 1;
-		uint32 NumBytes = BytesPerElement * NumElements;
+		// Create a buffer with one element.
+		uint32 NumBytes = GPixelFormats[PixelFormat].BlockBytes;
 		FRHIResourceCreateInfo CreateInfo;
 		CreateInfo.DebugName = *DebugName;
 		Buffer = RHICreateVertexBuffer(NumBytes, BUF_ShaderResource | BUF_Static, CreateInfo);
-		SRV = RHICreateShaderResourceView(Buffer, BytesPerElement, PixelFormat);
+
+		// Zero the buffer memory.
+		void* Data = RHILockVertexBuffer(Buffer, 0, NumBytes, RLM_WriteOnly);
+		memset(Data, 0, NumBytes);
+		RHIUnlockVertexBuffer(Buffer);
+
+		SRV = RHICreateShaderResourceView(Buffer, NumBytes, PixelFormat);
 	}
 
 	virtual void ReleaseRHI() override
@@ -59,21 +64,29 @@ public:
 	FNiagaraEmptyTextureSRV(EPixelFormat InPixelFormat, const FString& InDebugName) : PixelFormat(InPixelFormat), DebugName(InDebugName) {}
 	EPixelFormat PixelFormat;
 	FString DebugName;
-	FTexture2DRHIRef Buffer;
+	FTexture2DRHIRef Texture;
 	FShaderResourceViewRHIRef SRV;
 
 	virtual void InitRHI() override
 	{
+		// Create a 1x1 texture.
 		FRHIResourceCreateInfo CreateInfo;
 		CreateInfo.DebugName = *DebugName;
-		Buffer = RHICreateTexture2D(1, 1, PixelFormat, 1, 1, TexCreate_ShaderResource, CreateInfo);
-		SRV = RHICreateShaderResourceView(Buffer, 0);
+		Texture = RHICreateTexture2D(1, 1, PixelFormat, 1, 1, TexCreate_ShaderResource, CreateInfo);
+
+		// Zero the texture memory (there's only 1 row, so we can use the stride).
+		uint32 Stride;
+		void* Pixels = RHILockTexture2D(Texture, 0, RLM_WriteOnly, Stride, false);
+		memset(Pixels, 0, Stride);
+		RHIUnlockTexture2D(Texture, 0, false);
+
+		SRV = RHICreateShaderResourceView(Texture, 0);
 	}
 
 	virtual void ReleaseRHI() override
 	{
 		SRV.SafeRelease();
-		Buffer.SafeRelease();
+		Texture.SafeRelease();
 	}
 };
 
