Use FMemory::Memset() instead of memset().

diff --git a/Engine/Plugins/FX/Niagara/Source/Niagara/Private/NiagaraRenderer.cpp b/Engine/Plugins/FX/Niagara/Source/Niagara/Private/NiagaraRenderer.cpp
--- a/Engine/Plugins/FX/Niagara/Source/Niagara/Private/NiagaraRenderer.cpp
+++ b/Engine/Plugins/FX/Niagara/Source/Niagara/Private/NiagaraRenderer.cpp
@@ -45,7 +45,7 @@ public:
 
 		// Zero the buffer memory.
 		void* Data = RHILockVertexBuffer(Buffer, 0, NumBytes, RLM_WriteOnly);
-		memset(Data, 0, NumBytes);
+		FMemory::Memset(Data, 0, NumBytes);
 		RHIUnlockVertexBuffer(Buffer);
 
 		SRV = RHICreateShaderResourceView(Buffer, NumBytes, PixelFormat);
@@ -77,7 +77,7 @@ public:
 		// Zero the texture memory (there's only 1 row, so we can use the stride).
 		uint32 Stride;
 		void* Pixels = RHILockTexture2D(Texture, 0, RLM_WriteOnly, Stride, false);
-		memset(Pixels, 0, Stride);
+		FMemory::Memset(Pixels, 0, Stride);
 		RHIUnlockTexture2D(Texture, 0, false);
 
 		SRV = RHICreateShaderResourceView(Texture, 0);
