Rewrite GetScriptParameterVisualWidget to avoid passing ptr out of scope if variable and view info were sampled from the graph script variables.

diff --git a/Engine/Plugins/FX/Niagara/Source/NiagaraEditor/Private/ViewModels/NiagaraParameterPanelViewModel.cpp b/Engine/Plugins/FX/Niagara/Source/NiagaraEditor/Private/ViewModels/NiagaraParameterPanelViewModel.cpp
--- a/Engine/Plugins/FX/Niagara/Source/NiagaraEditor/Private/ViewModels/NiagaraParameterPanelViewModel.cpp
+++ b/Engine/Plugins/FX/Niagara/Source/NiagaraEditor/Private/ViewModels/NiagaraParameterPanelViewModel.cpp
@@ -656,7 +656,13 @@ TSharedRef<SWidget> FNiagaraScriptToolkitParameterPanelViewModel::GetScriptParam
 	const FNiagaraVariable PinVar = Schema->PinToNiagaraVariable(Pin);
 
 	const FNiagaraScriptVariableAndViewInfo* ScriptVarAndViewInfo = CachedViewedParameters.FindByPredicate([PinVar](const FNiagaraScriptVariableAndViewInfo& Entry) {return Entry.ScriptVariable == PinVar; });
-	if (ScriptVarAndViewInfo == nullptr)
+	if (ScriptVarAndViewInfo != nullptr)
+	{
+		TSharedPtr<FNiagaraGraphPinParameterNameViewModel> ParameterNameViewModel = MakeShared<FNiagaraGraphPinParameterNameViewModel>(Pin, *ScriptVarAndViewInfo, this);
+		TSharedPtr<SWidget> ScriptParameterVisualWidget = SNew(SNiagaraParameterNameView, ParameterNameViewModel);
+		return ScriptParameterVisualWidget->AsShared();
+	}
+	else
 	{
 		// Failed to find the parameter name in the cache, try to find the variable in the graph script variables and generate view info.
 		const UNiagaraScriptVariable* const* ScriptVarPtr = ScriptViewModel->GetGraphViewModel()->GetGraph()->GetAllMetaData().Find(PinVar);
@@ -665,19 +671,16 @@ TSharedRef<SWidget> FNiagaraScriptToolkitParameterPanelViewModel::GetScriptParam
 			const UNiagaraScriptVariable* ScriptVar = *ScriptVarPtr;
 			const TStaticArray<FScopeIsEnabledAndTooltip, (int32)ENiagaraParameterScope::Num> PerScopeInfo = GetParameterScopesEnabledAndTooltips(ScriptVar->Variable, ScriptVar->Metadata);
 			FNiagaraScriptVariableAndViewInfo NewScriptVarAndViewInfo = FNiagaraScriptVariableAndViewInfo(ScriptVar->Variable, ScriptVar->Metadata, PerScopeInfo);
-			ScriptVarAndViewInfo = &NewScriptVarAndViewInfo;
-		}
-		else
-		{
-			// Cannot resolve the parameter from the pin, put an error widget in.
-			TSharedPtr<SWidget> ErrorTextBlock = SNew(STextBlock).Text(FText::FromString("Could not resolve parameter!")); //@todo(ng) make error item method
-			return ErrorTextBlock->AsShared();
+
+			TSharedPtr<FNiagaraGraphPinParameterNameViewModel> ParameterNameViewModel = MakeShared<FNiagaraGraphPinParameterNameViewModel>(Pin, NewScriptVarAndViewInfo, this);
+			TSharedPtr<SWidget> ScriptParameterVisualWidget = SNew(SNiagaraParameterNameView, ParameterNameViewModel);
+			return ScriptParameterVisualWidget->AsShared();
 		}
 	}
-	TSharedPtr<FNiagaraGraphPinParameterNameViewModel> ParameterNameViewModel = MakeShared<FNiagaraGraphPinParameterNameViewModel>(Pin, *ScriptVarAndViewInfo, this);
-	
-	TSharedPtr<SWidget> ScriptParameterVisualWidget = SNew(SNiagaraParameterNameView, ParameterNameViewModel);
-	return ScriptParameterVisualWidget->AsShared();
+
+	// Cannot resolve the parameter from the pin, put an error widget in.
+	TSharedPtr<SWidget> ErrorTextBlock = SNew(STextBlock).Text(FText::FromString("Could not resolve parameter!")); //@todo(ng) make error item method
+	return ErrorTextBlock->AsShared();
 }
 
 TStaticArray<FScopeIsEnabledAndTooltip, (int32)ENiagaraParameterScope::Num> FNiagaraScriptToolkitParameterPanelViewModel::GetParameterScopesEnabledAndTooltips(const FNiagaraVariable& InVar, const FNiagaraVariableMetaData& InVarMetaData) const
