Get the command list for the menu so the UI actions are bound correctly.
#rnx

diff --git a/Engine/Plugins/FX/Niagara/Source/NiagaraEditorWidgets/Private/SNiagaraStack.cpp b/Engine/Plugins/FX/Niagara/Source/NiagaraEditorWidgets/Private/SNiagaraStack.cpp
--- a/Engine/Plugins/FX/Niagara/Source/NiagaraEditorWidgets/Private/SNiagaraStack.cpp
+++ b/Engine/Plugins/FX/Niagara/Source/NiagaraEditorWidgets/Private/SNiagaraStack.cpp
@@ -55,6 +55,9 @@
 #include "NiagaraEditorUtilities.h"
 #include "Subsystems/AssetEditorSubsystem.h"
 #include "EditorFontGlyphs.h"
+#include "Framework/Commands/UICommandList.h"
+#include "ViewModels/NiagaraOverviewGraphViewModel.h"
+
 
 #define LOCTEXT_NAMESPACE "NiagaraStack"
 
@@ -729,19 +732,21 @@ TSharedPtr<SWidget> SNiagaraStack::GenerateStackMenu(TWeakPtr<UNiagaraStackViewM
 	TSharedPtr<UNiagaraStackViewModel::FTopLevelViewModel> TopLevelViewModel = TopLevelViewModelWeak.Pin();
 	if (TopLevelViewModel.IsValid())
 	{
-		FMenuBuilder MenuBuilder(true, nullptr);
+		TSharedPtr<FUICommandList> GraphCommandList = TopLevelViewModel->RootEntry->GetSystemViewModel()->GetOverviewGraphViewModel()->GetCommands();
+		FMenuBuilder MenuBuilder(true, GraphCommandList);
 
 		FNiagaraEditorUtilities::AddEmitterContextMenuActions(MenuBuilder, TopLevelViewModel->EmitterHandleViewModel);
 
-		MenuBuilder.BeginSection("EmitterEditSection", LOCTEXT("Edit", "Edit"));
+		{
+			MenuBuilder.BeginSection("EmitterEditSection", LOCTEXT("Edit", "Edit"));
 
-		MenuBuilder.AddMenuEntry(FGenericCommands::Get().Cut);
-		MenuBuilder.AddMenuEntry(FGenericCommands::Get().Copy);
-		MenuBuilder.AddMenuEntry(FGenericCommands::Get().Delete);
-		MenuBuilder.AddMenuEntry(FGenericCommands::Get().Rename);
-
-		MenuBuilder.EndSection();
+			MenuBuilder.AddMenuEntry(FGenericCommands::Get().Cut);
+			MenuBuilder.AddMenuEntry(FGenericCommands::Get().Copy);
+			MenuBuilder.AddMenuEntry(FGenericCommands::Get().Delete);
+			MenuBuilder.AddMenuEntry(FGenericCommands::Get().Rename);
 
+			MenuBuilder.EndSection();
+		}
 		MenuBuilder.BeginSection("StackActions", LOCTEXT("StackActions", "Stack Actions"));
 		{
 			if (StackViewModel->HasDismissedStackIssues())
