Fix Linux build break (case sensitive headers)

diff --git a/Engine/Plugins/FX/Niagara/Source/NiagaraEditor/Private/ViewModels/NiagaraParameterNameViewModel.cpp b/Engine/Plugins/FX/Niagara/Source/NiagaraEditor/Private/ViewModels/NiagaraParameterNameViewModel.cpp
--- a/Engine/Plugins/FX/Niagara/Source/NiagaraEditor/Private/ViewModels/NiagaraParameterNameViewModel.cpp
+++ b/Engine/Plugins/FX/Niagara/Source/NiagaraEditor/Private/ViewModels/NiagaraParameterNameViewModel.cpp
@@ -1,7 +1,7 @@
 // Copyright Epic Games, Inc. All Rights Reserved.
 
 #include "NiagaraParameterNameViewModel.h"
-#include "EditorWidgets/Public/SEnumComboBox.h"
+#include "EditorWidgets/Public/SEnumCombobox.h"
 #include "Widgets/Text/SInlineEditableTextBlock.h"
 #include "SGraphActionMenu.h"
 #include "NiagaraActions.h"
