CIS Fix - Fix a mismerge where CL 11619064 coming into main seems to have removed ::GetOwningEmitterAndScriptForStackNode that was submitted in CL 11592184

#rnx nicholas.goldstein

diff --git a/Engine/Plugins/FX/Niagara/Source/NiagaraEditor/Public/ViewModels/Stack/NiagaraStackGraphUtilities.h b/Engine/Plugins/FX/Niagara/Source/NiagaraEditor/Public/ViewModels/Stack/NiagaraStackGraphUtilities.h
--- a/Engine/Plugins/FX/Niagara/Source/NiagaraEditor/Public/ViewModels/Stack/NiagaraStackGraphUtilities.h
+++ b/Engine/Plugins/FX/Niagara/Source/NiagaraEditor/Public/ViewModels/Stack/NiagaraStackGraphUtilities.h
@@ -150,6 +150,8 @@ namespace FNiagaraStackGraphUtilities
 	TOptional<FName> GetNamespaceForScriptUsage(ENiagaraScriptUsage ScriptUsage);
 	ENiagaraParameterScope GetScopeForScriptUsage(ENiagaraScriptUsage ScriptUsage);
 	
+	void GetOwningEmitterAndScriptForStackNode(UNiagaraNode& StackNode, UNiagaraSystem& OwningSystem, UNiagaraEmitter*& OutEmitter, UNiagaraScript*& OutScript);
+
 	bool IsValidDefaultDynamicInput(UNiagaraScript& OwningScript, UEdGraphPin& DefaultPin);
 
 	bool CanWriteParameterFromUsage(FNiagaraVariable Parameter, ENiagaraScriptUsage Usage);
