Niagara - Remove a stack graph utility function which as readded due to bad merges.

diff --git a/Engine/Plugins/FX/Niagara/Source/NiagaraEditor/Private/ViewModels/Stack/NiagaraStackGraphUtilities.cpp b/Engine/Plugins/FX/Niagara/Source/NiagaraEditor/Private/ViewModels/Stack/NiagaraStackGraphUtilities.cpp
--- a/Engine/Plugins/FX/Niagara/Source/NiagaraEditor/Private/ViewModels/Stack/NiagaraStackGraphUtilities.cpp
+++ b/Engine/Plugins/FX/Niagara/Source/NiagaraEditor/Private/ViewModels/Stack/NiagaraStackGraphUtilities.cpp
@@ -1633,41 +1633,6 @@ ENiagaraParameterScope FNiagaraStackGraphUtilities::GetScopeForScriptUsage(ENiag
 	}
 }
 
-void FNiagaraStackGraphUtilities::GetOwningEmitterAndScriptForStackNode(UNiagaraNode& StackNode, UNiagaraSystem& OwningSystem, UNiagaraEmitter*& OutEmitter, UNiagaraScript*& OutScript)
-{
-	OutEmitter = nullptr;
-	OutScript = nullptr;
-	UNiagaraNodeOutput* OutputNode = GetEmitterOutputNodeForStackNode(StackNode);
-	if (OutputNode != nullptr)
-	{
-		switch (OutputNode->GetUsage())
-		{
-			case ENiagaraScriptUsage::SystemSpawnScript:
-				OutScript = OwningSystem.GetSystemSpawnScript();
-				break;
-			case ENiagaraScriptUsage::SystemUpdateScript:
-				OutScript = OwningSystem.GetSystemUpdateScript();
-				break;
-			case ENiagaraScriptUsage::EmitterSpawnScript:
-			case ENiagaraScriptUsage::EmitterUpdateScript:
-			case ENiagaraScriptUsage::ParticleSpawnScript:
-			case ENiagaraScriptUsage::ParticleUpdateScript:
-			case ENiagaraScriptUsage::ParticleEventScript:
-				for (const FNiagaraEmitterHandle& EmitterHandle : OwningSystem.GetEmitterHandles())
-				{
-					UNiagaraScriptSource* EmitterSource = CastChecked<UNiagaraScriptSource>(EmitterHandle.GetInstance()->GraphSource);
-					if (EmitterSource->NodeGraph == StackNode.GetNiagaraGraph())
-					{
-						OutEmitter = EmitterHandle.GetInstance();
-						OutScript = OutEmitter->GetScript(OutputNode->GetUsage(), OutputNode->GetUsageId());
-						break;
-					}
-				}
-				break;
-		}
-	}
-}
-
 struct FRapidIterationParameterContext
 {
 	FRapidIterationParameterContext()
diff --git a/Engine/Plugins/FX/Niagara/Source/NiagaraEditor/Public/ViewModels/Stack/NiagaraStackGraphUtilities.h b/Engine/Plugins/FX/Niagara/Source/NiagaraEditor/Public/ViewModels/Stack/NiagaraStackGraphUtilities.h
--- a/Engine/Plugins/FX/Niagara/Source/NiagaraEditor/Public/ViewModels/Stack/NiagaraStackGraphUtilities.h
+++ b/Engine/Plugins/FX/Niagara/Source/NiagaraEditor/Public/ViewModels/Stack/NiagaraStackGraphUtilities.h
@@ -150,8 +150,6 @@ namespace FNiagaraStackGraphUtilities
 	TOptional<FName> GetNamespaceForScriptUsage(ENiagaraScriptUsage ScriptUsage);
 	ENiagaraParameterScope GetScopeForScriptUsage(ENiagaraScriptUsage ScriptUsage);
 	
-	void GetOwningEmitterAndScriptForStackNode(UNiagaraNode& StackNode, UNiagaraSystem& OwningSystem, UNiagaraEmitter*& OutEmitter, UNiagaraScript*& OutScript);
-
 	bool IsValidDefaultDynamicInput(UNiagaraScript& OwningScript, UEdGraphPin& DefaultPin);
 
 	bool CanWriteParameterFromUsage(FNiagaraVariable Parameter, ENiagaraScriptUsage Usage);
