Reduce severity of logspam around parameter addition to history it can occur for legitimate reasons.

diff --git a/Engine/Plugins/FX/Niagara/Source/NiagaraEditor/Private/NiagaraParameterMapHistory.cpp b/Engine/Plugins/FX/Niagara/Source/NiagaraEditor/Private/NiagaraParameterMapHistory.cpp
--- a/Engine/Plugins/FX/Niagara/Source/NiagaraEditor/Private/NiagaraParameterMapHistory.cpp
+++ b/Engine/Plugins/FX/Niagara/Source/NiagaraEditor/Private/NiagaraParameterMapHistory.cpp
@@ -1949,7 +1949,7 @@ int32 FNiagaraParameterMapHistoryWithMetaDataBuilder::AddVariableToHistory(FNiag
 		}
 		return History.AddVariable(InVar, InAliasedVar, InPin, MetaData);
 	}
-	UE_LOG(LogNiagaraEditor, Warning, TEXT("Variable added to parameter map history did not have metadata! Variable: %s"), *InVar.GetName().ToString());
+	UE_LOG(LogNiagaraEditor, Display, TEXT("Variable added to parameter map history did not have metadata! Variable: %s"), *InVar.GetName().ToString());
 	TOptional<FNiagaraVariableMetaData> BlankMetaData = FNiagaraVariableMetaData();
 	return History.AddVariable(InVar, InAliasedVar, InPin, BlankMetaData);
 }
