Clamp niagara parameters menu to 700 to avoid flickering bug

diff --git a/Engine/Plugins/FX/Niagara/Source/NiagaraEditor/Private/Widgets/SNiagaraParameterMapView.cpp b/Engine/Plugins/FX/Niagara/Source/NiagaraEditor/Private/Widgets/SNiagaraParameterMapView.cpp
--- a/Engine/Plugins/FX/Niagara/Source/NiagaraEditor/Private/Widgets/SNiagaraParameterMapView.cpp
+++ b/Engine/Plugins/FX/Niagara/Source/NiagaraEditor/Private/Widgets/SNiagaraParameterMapView.cpp
@@ -831,6 +831,7 @@ void SNiagaraAddParameterMenu::Construct(const FArguments& InArgs, TArray<TWeakO
 		[
 			SNew(SBox)
 			.MinDesiredWidth(300)
+			.MaxDesiredHeight(700) // Set max desired height to prevent flickering bug for menu larger than screen
 			[
 				SAssignNew(GraphMenu, SGraphActionMenu)
 				.OnActionSelected(this, &SNiagaraAddParameterMenu::OnActionSelected)
