Fixing CIS issue with static analysis

diff --git a/Engine/Plugins/FX/Niagara/Source/Niagara/Private/NiagaraEmitter.cpp b/Engine/Plugins/FX/Niagara/Source/Niagara/Private/NiagaraEmitter.cpp
--- a/Engine/Plugins/FX/Niagara/Source/Niagara/Private/NiagaraEmitter.cpp
+++ b/Engine/Plugins/FX/Niagara/Source/Niagara/Private/NiagaraEmitter.cpp
@@ -265,15 +265,33 @@ UNiagaraEmitter* UNiagaraEmitter::DuplicateWithoutMerging(UObject* InOuter)
 void UNiagaraEmitter::Serialize(FArchive& Ar)
 {
 
+#if WITH_EDITORONLY_DATA
 	for (UNiagaraSimulationStageBase* Stage : SimulationStages)
 	{
-		ensure(Stage);
-		ensure(Stage->Script);
-		if (!HasAnyFlags(RF_Transient))
+		if (Stage)
+		{
+			if (Stage->Script)
+			{
+				if (!HasAnyFlags(RF_Transient))
+				{
+					if (Stage->Script->HasAnyFlags(RF_Transient))
+					{
+						UE_LOG(LogNiagara, Error, TEXT("Emitter \"%s\" has a simulation stage with a Transient script and the emitter itself isn't transient!"), *GetPathName());
+					}
+				}
+			}
+			else
+			{
+				UE_LOG(LogNiagara, Error, TEXT("Emitter \"%s\" has a simulation stage with a null Script entry!"), *GetPathName());
+			}
+
+		}
+		else
 		{
-			ensure(!Stage->Script->HasAnyFlags(RF_Transient));
+			UE_LOG(LogNiagara, Error, TEXT("Emitter \"%s\" has a simulation stage with a null entry!"), *GetPathName());
 		}
 	}
+#endif
 	Super::Serialize(Ar);
 
 	Ar.UsingCustomVersion(FNiagaraCustomVersion::GUID);
