Now warning when fixed bounds not set.

diff --git a/Engine/Plugins/FX/Niagara/Source/NiagaraEditor/Private/ViewModels/Stack/NiagaraStackEmitterSettingsGroup.cpp b/Engine/Plugins/FX/Niagara/Source/NiagaraEditor/Private/ViewModels/Stack/NiagaraStackEmitterSettingsGroup.cpp
--- a/Engine/Plugins/FX/Niagara/Source/NiagaraEditor/Private/ViewModels/Stack/NiagaraStackEmitterSettingsGroup.cpp
+++ b/Engine/Plugins/FX/Niagara/Source/NiagaraEditor/Private/ViewModels/Stack/NiagaraStackEmitterSettingsGroup.cpp
@@ -94,11 +94,30 @@ void UNiagaraStackEmitterPropertiesItem::RefreshChildrenInternal(const TArray<UN
 	NewChildren.Add(EmitterObject);
 	bCanResetToBaseCache.Reset();
 	Super::RefreshChildrenInternal(CurrentChildren, NewChildren, NewIssues);
+	RefreshIssues(NewIssues);
+}
+
+
+void UNiagaraStackEmitterPropertiesItem::RefreshIssues(TArray<FStackIssue>& NewIssues)
+{
+	UNiagaraEmitter* ActualEmitter = GetEmitterViewModel()->GetEmitter();
+	if (ActualEmitter && ActualEmitter->SimTarget == ENiagaraSimTarget::GPUComputeSim && ActualEmitter->bFixedBounds == false)
+	{
+		FStackIssue MissingRequiredFixedBoundsModuleError(
+			EStackIssueSeverity::Warning,
+			LOCTEXT("RequiredFixedBoundsWarningFormat", "The emitter is GPU but the fixed bounds checkbox is not set.\r\nPlease update the Emitter properties, otherwise existing value for fixed bounds will be used."),
+			LOCTEXT("MissingFixedBounds", "Missing fixed bounds."),
+			GetStackEditorDataKey(),
+			false);
+
+		NewIssues.Add(MissingRequiredFixedBoundsModuleError);
+	}
 }
 
 void UNiagaraStackEmitterPropertiesItem::EmitterPropertiesChanged()
 {
 	bCanResetToBaseCache.Reset();
+	RefreshChildren();
 }
 
 UNiagaraStackEmitterSettingsGroup::UNiagaraStackEmitterSettingsGroup()
diff --git a/Engine/Plugins/FX/Niagara/Source/NiagaraEditor/Public/ViewModels/Stack/NiagaraStackEmitterSettingsGroup.h b/Engine/Plugins/FX/Niagara/Source/NiagaraEditor/Public/ViewModels/Stack/NiagaraStackEmitterSettingsGroup.h
--- a/Engine/Plugins/FX/Niagara/Source/NiagaraEditor/Public/ViewModels/Stack/NiagaraStackEmitterSettingsGroup.h
+++ b/Engine/Plugins/FX/Niagara/Source/NiagaraEditor/Public/ViewModels/Stack/NiagaraStackEmitterSettingsGroup.h
@@ -33,6 +33,8 @@ protected:
 
 	virtual void RefreshChildrenInternal(const TArray<UNiagaraStackEntry*>& CurrentChildren, TArray<UNiagaraStackEntry*>& NewChildren, TArray<FStackIssue>& NewIssues) override;
 
+	virtual void RefreshIssues(TArray<FStackIssue>& NewIssues);
+
 private:
 	void EmitterPropertiesChanged();
 
