Multiple fixes to enable simstages to work properly for other modules.

diff --git a/Engine/Plugins/FX/Niagara/Source/NiagaraEditor/Private/NiagaraHlslTranslator.cpp b/Engine/Plugins/FX/Niagara/Source/NiagaraEditor/Private/NiagaraHlslTranslator.cpp
--- a/Engine/Plugins/FX/Niagara/Source/NiagaraEditor/Private/NiagaraHlslTranslator.cpp
+++ b/Engine/Plugins/FX/Niagara/Source/NiagaraEditor/Private/NiagaraHlslTranslator.cpp
@@ -729,7 +729,7 @@ bool FHlslNiagaraTranslator::ShouldConsiderTargetParameterMap(ENiagaraScriptUsag
 	ENiagaraScriptUsage TargetUsage = GetTargetUsage();
 	if (TargetUsage >= ENiagaraScriptUsage::ParticleSpawnScript && TargetUsage <= ENiagaraScriptUsage::ParticleEventScript)
 	{
-		return InUsage >= ENiagaraScriptUsage::ParticleSpawnScript && InUsage <= ENiagaraScriptUsage::ParticleEventScript;
+		return InUsage >= ENiagaraScriptUsage::ParticleSpawnScript && InUsage <= ENiagaraScriptUsage::ParticleSimulationStageScript;
 	}
 	else if (TargetUsage == ENiagaraScriptUsage::SystemSpawnScript)
 	{
@@ -1181,7 +1181,7 @@ const FNiagaraTranslateResults &FHlslNiagaraTranslator::Translate(const FNiagara
 				AddBodyComment(FString::Printf(TEXT("//Begin Stage Script: %s!"), *TranslationStages[i].PassNamespace));
 				//Now we compile the simulation stage and read from the temp values written above.
 				CurrentParamMapIndices.Empty();
-				CurrentParamMapIndices.Add(1);
+				CurrentParamMapIndices.Add(i);
 				TranslationStages[i].OutputNode->Compile(ThisTranslator, OutputChunks);
 				AddBodyComment(FString::Printf(TEXT("//End Simulation Stage Script: %s\n\n"), *TranslationStages[i].PassNamespace));
 			}
