Adding simulation stages to log section for message manageer

diff --git a/Engine/Plugins/FX/Niagara/Source/NiagaraEditor/Private/NiagaraMessageManager.cpp b/Engine/Plugins/FX/Niagara/Source/NiagaraEditor/Private/NiagaraMessageManager.cpp
--- a/Engine/Plugins/FX/Niagara/Source/NiagaraEditor/Private/NiagaraMessageManager.cpp
+++ b/Engine/Plugins/FX/Niagara/Source/NiagaraEditor/Private/NiagaraMessageManager.cpp
@@ -86,6 +86,9 @@ const TOptional<const FString> FNiagaraMessageManager::GetStringForScriptUsageIn
 	case ENiagaraScriptUsage::ParticleEventScript:
 		return TOptional<const FString>(TEXT("Particle Event Script"));
 
+	case ENiagaraScriptUsage::ParticleSimulationStageScript:
+		return TOptional<const FString>(FString("Particle Simulation Stage Script"));
+
 	case ENiagaraScriptUsage::EmitterSpawnScript:
 		return TOptional<const FString>(TEXT("Emitter Spawn Script"));
 
