Adding simulation stages to stack function input and making some code safe against numerics.

diff --git a/Engine/Plugins/FX/Niagara/Source/NiagaraEditor/Private/ViewModels/Stack/NiagaraStackFunctionInput.cpp b/Engine/Plugins/FX/Niagara/Source/NiagaraEditor/Private/ViewModels/Stack/NiagaraStackFunctionInput.cpp
--- a/Engine/Plugins/FX/Niagara/Source/NiagaraEditor/Private/ViewModels/Stack/NiagaraStackFunctionInput.cpp
+++ b/Engine/Plugins/FX/Niagara/Source/NiagaraEditor/Private/ViewModels/Stack/NiagaraStackFunctionInput.cpp
@@ -692,11 +692,22 @@ void UNiagaraStackFunctionInput::RefreshValues()
 			InputValues.Mode = EValueMode::Local;
 			InputValues.LocalStruct = MakeShared<FStructOnScope>(InputType.GetStruct());
 			const uint8* RapidIterationParameterData = SourceScript->RapidIterationParameters.GetParameterData(RapidIterationParameter);
-			if (RapidIterationParameterData == nullptr)
+			if (DefaultInputValues.LocalStruct) // Numeric types can trigger crashes below, so extra safety has been added.
 			{
-				RapidIterationParameterData = DefaultInputValues.LocalStruct->GetStructMemory();
+				if (RapidIterationParameterData == nullptr)
+				{
+					RapidIterationParameterData = DefaultInputValues.LocalStruct->GetStructMemory();
+				}
+
+				if (InputType.GetSize() > 0 && InputValues.LocalStruct->GetStructMemory() && RapidIterationParameterData)
+				{
+					FMemory::Memcpy(InputValues.LocalStruct->GetStructMemory(), RapidIterationParameterData, InputType.GetSize());
+				}
+				else
+				{
+					UE_LOG(LogNiagaraEditor, Warning, TEXT("Type %s has no data! Cannot refresh values."), *InputType.GetName())
+				}
 			}
-			FMemory::Memcpy(InputValues.LocalStruct->GetStructMemory(), RapidIterationParameterData, InputType.GetSize());
 		}
 		else
 		{
@@ -830,7 +841,8 @@ bool UsageRunsBefore(ENiagaraScriptUsage UsageA, ENiagaraScriptUsage UsageB, boo
 		ENiagaraScriptUsage::ParticleSpawnScript,
 		ENiagaraScriptUsage::ParticleEventScript,	// When not using interpolated spawn
 		ENiagaraScriptUsage::ParticleUpdateScript,
-		ENiagaraScriptUsage::ParticleEventScript	// When using interpolated spawn and is spawn
+		ENiagaraScriptUsage::ParticleEventScript,	// When using interpolated spawn and is spawn
+		ENiagaraScriptUsage::ParticleSimulationStageScript
 	};
 
 	int32 IndexA;
@@ -863,6 +875,7 @@ FName GetNamespaceForUsage(ENiagaraScriptUsage Usage)
 	case ENiagaraScriptUsage::ParticleSpawnScript:
 	case ENiagaraScriptUsage::ParticleUpdateScript:
 	case ENiagaraScriptUsage::ParticleEventScript:
+	case ENiagaraScriptUsage::ParticleSimulationStageScript:
 		return FNiagaraParameterHandle::ParticleAttributeNamespace;
 	case ENiagaraScriptUsage::EmitterSpawnScript:
 	case ENiagaraScriptUsage::EmitterUpdateScript:
