Making the shaderstage variables only show up when the checkbox is set

diff --git a/Engine/Plugins/FX/Niagara/Source/Niagara/Classes/NiagaraEmitter.h b/Engine/Plugins/FX/Niagara/Source/Niagara/Classes/NiagaraEmitter.h
--- a/Engine/Plugins/FX/Niagara/Source/Niagara/Classes/NiagaraEmitter.h
+++ b/Engine/Plugins/FX/Niagara/Source/Niagara/Classes/NiagaraEmitter.h
@@ -338,15 +338,15 @@ public:
 	float MaxDeltaTimePerTick;
 
 	/** Get the default shader stage index. */
-	UPROPERTY(EditAnywhere, AdvancedDisplay, Category = "Simulation Stages", meta = (EditCondition = "bSimulationStagesEnabled || bDeprecatedShaderStagesEnabled", DisplayAfter = "bSimulationStagesEnabled"))
+	UPROPERTY(EditAnywhere, AdvancedDisplay, Category = "Simulation Stages", meta = (EditCondition = "bDeprecatedShaderStagesEnabled", DisplayAfter = "bDeprecatedShaderStagesEnabled"))
 	uint32 DefaultShaderStageIndex;
 
 	/** Get the number of shader stages that we fire off. */
-	UPROPERTY(EditAnywhere, AdvancedDisplay, Category = "Simulation Stages", meta = (EditCondition = "bSimulationStagesEnabled || bDeprecatedShaderStagesEnabled", DisplayAfter = "DefaultShaderStageIndex"))
+	UPROPERTY(EditAnywhere, AdvancedDisplay, Category = "Simulation Stages", meta = (EditCondition = "bDeprecatedShaderStagesEnabled", DisplayAfter = "DefaultShaderStageIndex"))
 	uint32 MaxUpdateIterations;
 
 	/** Get whether or not shaderstages spwn. */
-	UPROPERTY(EditAnywhere, AdvancedDisplay, Category = "Simulation Stages", meta = (EditCondition = "bSimulationStagesEnabled || bDeprecatedShaderStagesEnabled", DisplayAfter = "MaxUpdateIterations"))
+	UPROPERTY(EditAnywhere, AdvancedDisplay, Category = "Simulation Stages", meta = (EditCondition = "bDeprecatedShaderStagesEnabled", DisplayAfter = "MaxUpdateIterations"))
 	TSet<uint32> SpawnStages;
 
 	/** Get whether or not to use simulation stages. */
