Merging 11640024 from main - Niagara - Fix additional determinism issues related to looping in the editor and looping in the level sequencer.

diff --git a/Engine/Plugins/FX/Niagara/Source/Niagara/Private/MovieScene/MovieSceneNiagaraSystemTrackTemplate.cpp b/Engine/Plugins/FX/Niagara/Source/Niagara/Private/MovieScene/MovieSceneNiagaraSystemTrackTemplate.cpp
--- a/Engine/Plugins/FX/Niagara/Source/Niagara/Private/MovieScene/MovieSceneNiagaraSystemTrackTemplate.cpp
+++ b/Engine/Plugins/FX/Niagara/Source/Niagara/Private/MovieScene/MovieSceneNiagaraSystemTrackTemplate.cpp
@@ -121,6 +121,10 @@ struct FNiagaraSystemUpdateDesiredAgeExecutionToken : IMovieSceneExecutionToken
 					if (NiagaraComponent->IsActive())
 					{
 						NiagaraComponent->DeactivateImmediate();
+						if (NiagaraComponent->GetSystemInstance() != nullptr)
+						{
+							NiagaraComponent->GetSystemInstance()->Reset(FNiagaraSystemInstance::EResetMode::ResetAll);
+						}
 					}
 				}
 			}
@@ -163,6 +167,10 @@ struct FNiagaraSystemUpdateDesiredAgeExecutionToken : IMovieSceneExecutionToken
 					if (NiagaraComponent->IsActive())
 					{
 						NiagaraComponent->DeactivateImmediate();
+						if (NiagaraComponent->GetSystemInstance() != nullptr)
+						{
+							NiagaraComponent->GetSystemInstance()->Reset(FNiagaraSystemInstance::EResetMode::ResetAll);
+						}
 					}
 				}
 			}
diff --git a/Engine/Plugins/FX/Niagara/Source/NiagaraEditor/Private/ViewModels/NiagaraSystemViewModel.cpp b/Engine/Plugins/FX/Niagara/Source/NiagaraEditor/Private/ViewModels/NiagaraSystemViewModel.cpp
--- a/Engine/Plugins/FX/Niagara/Source/NiagaraEditor/Private/ViewModels/NiagaraSystemViewModel.cpp
+++ b/Engine/Plugins/FX/Niagara/Source/NiagaraEditor/Private/ViewModels/NiagaraSystemViewModel.cpp
@@ -1689,7 +1689,7 @@ void FNiagaraSystemViewModel::SequencerTimeChanged()
 				// We don't want to reset the current time if we're scrubbing.
 				if (CurrentStatus == EMovieScenePlayerStatus::Playing)
 				{
-					ResetSystem(ETimeResetMode::AllowResetTime, EMultiResetMode::ResetThisInstance, EReinitMode::ResetSystem);
+					ResetSystem(ETimeResetMode::AllowResetTime, EMultiResetMode::ResetThisInstance, EReinitMode::ReinitializeSystem);
 				}
 				else
 				{
