Simulation stages can kill particles.

Also:
* fix for dead particles not being removed if the emitter had stages
* workaround for a bug in HLSLcc which caused emitters with stages to have 0 particles on Vulkan and GLES.
* MaxUpdateIterations was always used, even if use shader stages was disabled. This caused some regular GPU emitters to incorrectly run multiple dispatches.

diff --git a/Engine/Plugins/FX/Niagara/Shaders/Private/NiagaraEmitterInstanceShader.usf b/Engine/Plugins/FX/Niagara/Shaders/Private/NiagaraEmitterInstanceShader.usf
--- a/Engine/Plugins/FX/Niagara/Shaders/Private/NiagaraEmitterInstanceShader.usf
+++ b/Engine/Plugins/FX/Niagara/Shaders/Private/NiagaraEmitterInstanceShader.usf
@@ -58,6 +58,7 @@ bool NiagaraAll(bool4 b) { return b.x && b.y && b.z && b.w; }
 	static uint GEmitterTickCounter;
 	static uint GSimStart;
 	static uint GRandomSeedOffset = 0;
+	static bool GStageUsesAlive = false;
 	
 	const static uint GSpawnPhase = 0;
 	const static uint GUpdatePhase = 1;
@@ -514,13 +515,6 @@ int rand_int(int x, int Seed1, int Seed2, int Seed3)
 		MustBe0[DataSetID] = 0;
 		// End static assert
 
-		#if USE_SIMULATION_STAGES
-			if(SimulationStageIndex != 0)
-			{
-				bDoAcquire = false; // Prevent any change to the atomic
-			}
-		#endif // USE_SIMULATION_STAGES
-
 		int PrevIdx = GSpawnStartInstance + SpawnedInstances;	// scratch instance as default; write to that for dead particles
 
 		#if USE_WAVE_INTRISICS
@@ -626,13 +620,6 @@ int rand_int(int x, int Seed1, int Seed2, int Seed3)
 
 		#endif // USE_WAVE_INTRISICS || USE_GROUP_SHARED
 
-		#if USE_SIMULATION_STAGES
-			if(SimulationStageIndex != 0)
-			{
-				return ExecIndex();
-			}
-		#endif // USE_SIMULATION_STAGES
-
 		return PrevIdx;
 	}
 
@@ -682,17 +669,23 @@ int rand_int(int x, int Seed1, int Seed2, int Seed3)
 	/* 
 	 * Get the index to write onto the output buffer
 	 */
-	int OutputIndex(const int DataSetID, const bool bIsPersistent, const bool bIsValid)
+	int OutputIndex(const int DataSetID, const bool bStageKillsParticles, const bool bIsValid)
 	{
 		#if USE_SIMULATION_STAGES
-			const int DenseIndex = AcquireIndex(DataSetID,bIsValid);
-			const int SparseIndex = GDispatchThreadId.x;
-			return (bIsPersistent && (DenseIndex != SparseIndex) ) ? SparseIndex : DenseIndex;
-		#else
-			return AcquireIndex(DataSetID,bIsValid);
+			// If this stage cannot kill particles, we can just write them out in the same order as they
+			// appear in the input. We must use an if here (as opposed to a ternary operator, or some
+			// other branchless construct), because we don't want to call AcquireIndex() at all, since
+			// that manipulates the RWInstanceCounts UAV. The generated code will copy the source count
+			// at the end of the shader.
+			BRANCH
+			if (SimulationStageIndex > 0 && !bStageKillsParticles)
+			{
+				return GDispatchThreadId.x;
+			}
 		#endif // USE_SIMULATION_STAGES
+
+			return AcquireIndex(DataSetID, bIsValid);
 	}
 
 //Include the simulation shader code generated by the node graph.
 #include "/Engine/Generated/NiagaraEmitterInstance.ush"
-
diff --git a/Engine/Plugins/FX/Niagara/Source/Niagara/Private/NiagaraEmitterInstance.cpp b/Engine/Plugins/FX/Niagara/Source/Niagara/Private/NiagaraEmitterInstance.cpp
--- a/Engine/Plugins/FX/Niagara/Source/Niagara/Private/NiagaraEmitterInstance.cpp
+++ b/Engine/Plugins/FX/Niagara/Source/Niagara/Private/NiagaraEmitterInstance.cpp
@@ -353,7 +353,8 @@ void FNiagaraEmitterInstance::Init(int32 InEmitterIdx, FNiagaraSystemInstanceID
 		if (CachedEmitter->SimTarget == ENiagaraSimTarget::GPUComputeSim)
 		{
 			GPUExecContext = new FNiagaraComputeExecutionContext();
-			GPUExecContext->InitParams(CachedEmitter->GetGPUComputeScript(), CachedEmitter->SimTarget, CachedEmitter->GetUniqueEmitterName(), CachedEmitter->DefaultShaderStageIndex, CachedEmitter->MaxUpdateIterations, CachedEmitter->SpawnStages);
+			const uint32 MaxUpdateIterations = CachedEmitter->bDeprecatedShaderStagesEnabled ? CachedEmitter->MaxUpdateIterations : 1;
+			GPUExecContext->InitParams(CachedEmitter->GetGPUComputeScript(), CachedEmitter->SimTarget, CachedEmitter->GetUniqueEmitterName(), CachedEmitter->DefaultShaderStageIndex, MaxUpdateIterations, CachedEmitter->SpawnStages);
 			GPUExecContext->MainDataSet = &Data;
 			GPUExecContext->GPUScript_RT = CachedEmitter->GetGPUComputeScript()->GetRenderThreadScript();
 
diff --git a/Engine/Plugins/FX/Niagara/Source/NiagaraCore/Private/NiagaraCustomVersion.cpp b/Engine/Plugins/FX/Niagara/Source/NiagaraCore/Private/NiagaraCustomVersion.cpp
--- a/Engine/Plugins/FX/Niagara/Source/NiagaraCore/Private/NiagaraCustomVersion.cpp
+++ b/Engine/Plugins/FX/Niagara/Source/NiagaraCore/Private/NiagaraCustomVersion.cpp
@@ -8,4 +8,4 @@ const FGuid FNiagaraCustomVersion::GUID(0xFCF57AFA, 0x50764283, 0xB9A9E658, 0xFF
 // Register the custom version with core
 FCustomVersionRegistration GRegisterNiagaraCustomVersion(FNiagaraCustomVersion::GUID, FNiagaraCustomVersion::LatestVersion, TEXT("NiagaraVer"));
 
-const FGuid FNiagaraCustomVersion::LatestScriptCompileVersion(0x4122F10C, 0xA1FA4669, 0xB73B66FA, 0x2B230FC6);
+const FGuid FNiagaraCustomVersion::LatestScriptCompileVersion(0x318C2E9D, 0xB5794FBD, 0xABE176B1, 0xC31904C9);
diff --git a/Engine/Plugins/FX/Niagara/Source/NiagaraEditor/Private/NiagaraHlslTranslator.cpp b/Engine/Plugins/FX/Niagara/Source/NiagaraEditor/Private/NiagaraHlslTranslator.cpp
--- a/Engine/Plugins/FX/Niagara/Source/NiagaraEditor/Private/NiagaraHlslTranslator.cpp
+++ b/Engine/Plugins/FX/Niagara/Source/NiagaraEditor/Private/NiagaraHlslTranslator.cpp
@@ -1999,29 +1999,28 @@ void FHlslNiagaraTranslator::DefineMainGPUFunctions(
 	DataSetWrites.GetKeys(WriteDataSetIDs);
 
 	// Whether Alive is used and must be set at each run
-	const bool bUsesAlive = [&]
+	bool bUsesAlive = false;
+	TArray<FName> DataSetNames;
+	for (const FNiagaraDataSetID& ReadId : ReadDataSetIDs)
 	{
-		TArray<FName> DataSetNames;
-		for (const FNiagaraDataSetID& ReadId : ReadDataSetIDs)
-		{
-			DataSetNames.AddUnique(ReadId.Name);
-		}
-		for (const FNiagaraDataSetID& WriteId : WriteDataSetIDs)
-		{
-			DataSetNames.AddUnique(WriteId.Name);
-		}
-		for (int32 i = 0; i < ParamMapHistories.Num(); i++)
+		DataSetNames.AddUnique(ReadId.Name);
+	}
+	for (const FNiagaraDataSetID& WriteId : WriteDataSetIDs)
+	{
+		DataSetNames.AddUnique(WriteId.Name);
+	}
+	for (int32 i = 0; i < ParamMapHistories.Num(); i++)
+	{
+		for (const FName& DataSetName : DataSetNames)
 		{
-			for (FName DataSetName : DataSetNames)
+			if (ParamMapHistories[i].FindVariable(*(DataSetName.ToString() + TEXT(".Alive")), FNiagaraTypeDefinition::GetBoolDef()) != INDEX_NONE)
 			{
-				if (ParamMapHistories[i].FindVariable(*(DataSetName.ToString() + TEXT(".Alive")), FNiagaraTypeDefinition::GetBoolDef()) != INDEX_NONE)
-				{
-					return true;
-				}
+				bUsesAlive = true;
+				TranslationStages[i].bUsesAlive = true;
+				break;
 			}
 		}
-		return false;
-	}();
+	}
 
 	const bool bRequiresPersistentIDs = CompileOptions.AdditionalDefines.Contains(TEXT("RequiresPersistentIDs"));
 
@@ -2247,15 +2246,16 @@ void FHlslNiagaraTranslator::DefineMainGPUFunctions(
 					"	{\n"), TranslationStages[i].SimulationStageIndexMin, TranslationStages[i].SimulationStageIndexMax);
 			}
 
-			if (bUsesAlive)
+			if (TranslationStages[i].bUsesAlive)
 			{
+				HlslOutput += TEXT("\t\tGStageUsesAlive = true;\n");
 				HlslOutput += TEXT("\t\tconst bool bValid = Context.") + TranslationStages[i].PassNamespace + TEXT(".DataInstance.Alive;\n");
-				HlslOutput += TEXT("\t\tconst int WriteIndex = OutputIndex(0, false, bValid);\n");
+				HlslOutput += TEXT("\t\tconst int WriteIndex = OutputIndex(0, true, bValid);\n");
 			}
 			else
 			{
 				HlslOutput += TEXT("\t\tconst bool bValid = GCurrentPhase != -1;\n");
-				HlslOutput += TEXT("\t\tconst int WriteIndex = OutputIndex(0, true, bValid);\n");
+				HlslOutput += TEXT("\t\tconst int WriteIndex = OutputIndex(0, false, bValid);\n");
 			}
 
 			FString ContextName = TEXT("Context.Map.");
@@ -2450,16 +2450,20 @@ void FHlslNiagaraTranslator::DefineMainGPUFunctions(
 			"		WriteDataSets(Context); \n"
 			"	}\n"
 			"	\n"
-			"	StoreUpdateVariables(Context); \n\n"
-			"	#if USE_SIMULATION_STAGES \n"
-			"	// We want to keep pushing along the counts so that we can  \n"
-			"	// use any of the write locations in the future. Only simulation stage 0 will increment the counters. \n"
-			"	if (SimulationStageIndex > 0 && GDispatchThreadId.x == 0) \n"
+			"	StoreUpdateVariables(Context);\n\n"
+		);
+
+	if (bUseSimulationStages)
+	{
+		HlslOutput += TEXT(
+			"	// If a stage doesn't kill particles, StoreUpdateVariables() never calls AcquireIndex(), so the\n"
+			"   // count isn't updated. In that case we must manually copy the original count here.\n"
+			"	if (SimulationStageIndex > 0 && !GStageUsesAlive && GDispatchThreadId.x == 0) \n"
 			"	{ \n"
 			"		RWInstanceCounts[WriteInstanceCountOffset] = RWInstanceCounts[ReadInstanceCountOffset]; \n"
 			"	} \n"
-			"	#endif \n"	
 		);
+	}
 
 	HlslOutput += TEXT("}\n");
 }
@@ -2642,11 +2646,11 @@ void FHlslNiagaraTranslator::DefineDataSetVariableWrites(FString &OutHlslOutput,
 	// grab the current ouput index to write datas 
 	if (bUsesAlive)
 	{
-		OutHlslOutput += "\tint TmpWriteIndex = OutputIndex(0, false, bValid);\n";
+		OutHlslOutput += "\tint TmpWriteIndex = OutputIndex(0, true, bValid);\n";
 	}
 	else
 	{
-		OutHlslOutput += "\tint TmpWriteIndex = OutputIndex(0, true, true);\n";
+		OutHlslOutput += "\tint TmpWriteIndex = OutputIndex(0, false, true);\n";
 	}
 
 	bool bRequiresPersistentIDs = CompileOptions.AdditionalDefines.Contains(TEXT("RequiresPersistentIDs"));
diff --git a/Engine/Plugins/FX/Niagara/Source/NiagaraEditor/Private/NiagaraHlslTranslator.h b/Engine/Plugins/FX/Niagara/Source/NiagaraEditor/Private/NiagaraHlslTranslator.h
--- a/Engine/Plugins/FX/Niagara/Source/NiagaraEditor/Private/NiagaraHlslTranslator.h
+++ b/Engine/Plugins/FX/Niagara/Source/NiagaraEditor/Private/NiagaraHlslTranslator.h
@@ -303,7 +303,7 @@ public:
 class NIAGARAEDITOR_API FHlslNiagaraTranslationStage
 {
 public:
-	FHlslNiagaraTranslationStage(ENiagaraScriptUsage InScriptUsage, FGuid InUsageId) : ScriptUsage(InScriptUsage), UsageId(InUsageId), OutputNode(nullptr), bInterpolatePreviousParams(false), bCopyPreviousParams(true), ChunkModeIndex((ENiagaraCodeChunkMode)-1), IterationSource(), bSpawnOnly(false){}
+	FHlslNiagaraTranslationStage(ENiagaraScriptUsage InScriptUsage, FGuid InUsageId) : ScriptUsage(InScriptUsage), UsageId(InUsageId), OutputNode(nullptr), bInterpolatePreviousParams(false), bCopyPreviousParams(true), ChunkModeIndex((ENiagaraCodeChunkMode)-1), IterationSource(), bSpawnOnly(false), bUsesAlive(false){}
 
 	ENiagaraScriptUsage ScriptUsage;
 	FGuid UsageId;
@@ -318,6 +318,7 @@ public:
 	int32 NumIterationsThisStage = 1;
 	int32 SourceSimStage = -1;
 	bool bSpawnOnly;
+	bool bUsesAlive;
 };
 
 class NIAGARAEDITOR_API FHlslNiagaraTranslator
