Fixed format string error in log call.
Fixed crash when disabling an emitter which has sim stages.
 * the GPU execution context isn't created in that case, so we can't do things with it in FNiagaraSystemInstance::InitDataInterfaces().

diff --git a/Engine/Plugins/FX/Niagara/Source/Niagara/Private/NiagaraScript.cpp b/Engine/Plugins/FX/Niagara/Source/Niagara/Private/NiagaraScript.cpp
--- a/Engine/Plugins/FX/Niagara/Source/Niagara/Private/NiagaraScript.cpp
+++ b/Engine/Plugins/FX/Niagara/Source/Niagara/Private/NiagaraScript.cpp
@@ -1904,7 +1904,7 @@ bool UNiagaraScript::SynchronizeExecutablesWithMaster(const UNiagaraScript* Scri
 
 void UNiagaraScript::InvalidateCompileResults(const FString& Reason)
 {
-	UE_LOG(LogNiagara, Verbose, TEXT("InvalidateCompileResults Script:%s Reason:%s"), *GetPathName());
+	UE_LOG(LogNiagara, Verbose, TEXT("InvalidateCompileResults Script:%s Reason:%s"), *GetPathName(), *Reason);
 	CachedScriptVM.Reset();
 	ScriptResource.Invalidate();
 	CachedScriptVMId.Invalidate();
diff --git a/Engine/Plugins/FX/Niagara/Source/Niagara/Private/NiagaraSystemInstance.cpp b/Engine/Plugins/FX/Niagara/Source/Niagara/Private/NiagaraSystemInstance.cpp
--- a/Engine/Plugins/FX/Niagara/Source/Niagara/Private/NiagaraSystemInstance.cpp
+++ b/Engine/Plugins/FX/Niagara/Source/Niagara/Private/NiagaraSystemInstance.cpp
@@ -1358,7 +1358,7 @@ void FNiagaraSystemInstance::InitDataInterfaces()
 			CalcInstDataSize(Sim.GetEventExecutionContexts()[i].GetDataInterfaces());
 		}
 
-		if (Sim.GetCachedEmitter() && Sim.GetCachedEmitter()->SimTarget == ENiagaraSimTarget::GPUComputeSim && Sim.GetCachedEmitter()->bSimulationStagesEnabled)
+		if (Sim.GetCachedEmitter() && Sim.GetCachedEmitter()->SimTarget == ENiagaraSimTarget::GPUComputeSim && Sim.GetCachedEmitter()->bSimulationStagesEnabled && Sim.GetGPUContext())
 		{
 			CalcInstDataSize(Sim.GetGPUContext()->GetDataInterfaces());
 		}
