Remove over-eager parameter linking for node graph deep copy level parameters.
Remove some commented out implementation.

diff --git a/Engine/Plugins/FX/Niagara/Source/NiagaraEditor/Private/NiagaraParameterMapHistory.cpp b/Engine/Plugins/FX/Niagara/Source/NiagaraEditor/Private/NiagaraParameterMapHistory.cpp
--- a/Engine/Plugins/FX/Niagara/Source/NiagaraEditor/Private/NiagaraParameterMapHistory.cpp
+++ b/Engine/Plugins/FX/Niagara/Source/NiagaraEditor/Private/NiagaraParameterMapHistory.cpp
@@ -1538,8 +1538,6 @@ void FNiagaraParameterMapHistoryWithMetaDataBuilder::FixupHistoryVariableNamespa
 			OutputHistoryVars = SortVars(false);
 			InputHistoryVars = SortVars(true);
 			MarkDoNotDemoteVars(OutputHistoryVars);
-// 			ResolveScriptAliasParameters(OutputHistoryVars);
-// 			ResolveScriptAliasParameters(InputHistoryVars);
 		}
 
 		TArray<FNiagaraHistoryVariable> OutputHistoryVars;
@@ -1552,26 +1550,6 @@ void FNiagaraParameterMapHistoryWithMetaDataBuilder::FixupHistoryVariableNamespa
 		const FName& GetEmitterUniqueName() const { return Handle->GetEmitterUniqueName(); };
 		const bool IsSystemHistoryHandle() const { return Handle->IsSystemHistoryHandle(); };
 		const ENiagaraScriptUsage GetOriginatingScriptUsage() const { return Handle->GetOriginatingScriptUsage(); };
-
-// 		void ResolveScriptAliasParameters(TArray<FNiagaraHistoryVariable>& VarsToResolve)
-// 		{
-// 			const ENiagaraParameterScope HistoryScope = Handle->GetScriptScope();
-// 			const FString HistoryScopeString = FNiagaraStackGraphUtilities::GetNamespaceStringForScriptParameterScope(HistoryScope);
-// 			for (FNiagaraHistoryVariable& HistoryVar : VarsToResolve)
-// 			{
-// 				if (HistoryVar.VarMetaData->Scope == ENiagaraParameterScope::ScriptAlias)
-// 				{
-// 					HistoryVar.SetVarWithOriginalAliasName()
-// 					FNiagaraParameterHandle HistoryVarParameterHandle = FNiagaraParameterHandle(HistoryVar.Var->GetName());
-// 
-// 					FString FixedNamespaceString = FString();
-// 					FixedNamespaceString.Append(HistoryScopeString);
-// 					FixedNamespaceString.Append(HistoryVarParameterHandle.GetName().ToString());
-// 					FName FixedNamespaceName = FName(*FixedNamespaceString);
-// 					HistoryVar.Var->SetName(FixedNamespaceName);
-// 				}
-// 			}
-// 		}
 	};
 
 	auto GenerateInitialVarsToSkip = [](TArray<FNiagaraVariable>& Vars)->TArray<FNiagaraVariable> {
@@ -1727,7 +1705,6 @@ void FNiagaraParameterMapHistoryWithMetaDataBuilder::FixupHistoryVariableNamespa
 		}
 	}
 
-
 	// Track names of parameters we demote to ensure we do not have name aliases.
 	TMap<FName, uint32> DemotedNameCounts;
 
@@ -1771,12 +1748,7 @@ void FNiagaraParameterMapHistoryWithMetaDataBuilder::FixupHistoryVariableNamespa
 					auto FindUpstreamVarMatchIdx = [&Handle, &MatchIdx, &CurrentHistoryVar, &CurrentInputVarHandle, &UpstreamSortedHistoryHandle]()->bool {
 						const ENiagaraParameterScope UpstreamScope = UpstreamSortedHistoryHandle->GetScriptScope();
 						ENiagaraParameterScope TargetScope = ENiagaraParameterScope::None;
-						if (CurrentHistoryVar.VarMetaData->bAddedToNodeGraphDeepCopy)
-						{
-							// If the current history var is at the node graph deep copy level, it can link to any output, so the target scope is always that of the upstream history.
-							TargetScope = UpstreamScope;
-						}
-						else if (CurrentHistoryVar.VarMetaData->Scope == ENiagaraParameterScope::ScriptPersistent)
+						if (CurrentHistoryVar.VarMetaData->Scope == ENiagaraParameterScope::ScriptPersistent)
 						{
 							// If the current history var is targeting the script alias scope then get the scope of the current history.
 							TargetScope = Handle.GetScriptScope();
