Reduce Samplers

diff --git a/Engine/Plugins/FX/Niagara/Shaders/Private/NiagaraGPURayTracingTransforms.usf b/Engine/Plugins/FX/Niagara/Shaders/Private/NiagaraGPURayTracingTransforms.usf
--- a/Engine/Plugins/FX/Niagara/Shaders/Private/NiagaraGPURayTracingTransforms.usf
+++ b/Engine/Plugins/FX/Niagara/Shaders/Private/NiagaraGPURayTracingTransforms.usf
@@ -63,7 +63,7 @@ void NiagaraGPURayTracingTransformsCS(uint3 DispatchThreadId : SV_DispatchThread
 
 		if (PositionOffset != -1)
 		{
-			InstancePosition = float4(GetVec3(NIAGARA_LOOSEPARAM_STRUCT(NiagaraParticleDataFloat), PositionOffset, InstanceIndex), 1);
+			InstancePosition = float4(GetVec3(PositionOffset, InstanceIndex), 1);
 		}
 
 		float4 Transform1 = float4(1.0f, 0.0f, 0.0f, InstancePosition.x);
@@ -72,7 +72,7 @@ void NiagaraGPURayTracingTransformsCS(uint3 DispatchThreadId : SV_DispatchThread
 
 		if (TransformOffset != -1)
 		{
-			float4 Quat = GetVec4(NIAGARA_LOOSEPARAM_STRUCT(NiagaraParticleDataFloat), TransformOffset, InstanceIndex);
+			float4 Quat = GetVec4(TransformOffset, InstanceIndex);
 
 			float4x4 EulerToMat = QuatTo4x4(normalize(Quat));
 			Transform1.xyz = EulerToMat[0].xyz;
@@ -83,7 +83,7 @@ void NiagaraGPURayTracingTransformsCS(uint3 DispatchThreadId : SV_DispatchThread
 		float4x4 ScaleMatrix = float4x4(float4(1.0f, 0.0f, 0.0f, 0.0f), float4(0.0f, 1.0f, 0.0f, 0.0f), float4(0.0f, 0.0f, 1.0f, 0.0f), float4(0.0f, 0.0f, 0.0f, 1.0f));
 		if (ScaleOffset != -1)
 		{
-			float3 Scale = GetVec3(NIAGARA_LOOSEPARAM_STRUCT(NiagaraParticleDataFloat), ScaleOffset, InstanceIndex);
+			float3 Scale = GetVec3(ScaleOffset, InstanceIndex);
 
 			ScaleMatrix[0][0] *= Scale.x;
 			ScaleMatrix[1][1] *= Scale.y;
diff --git a/Engine/Plugins/FX/Niagara/Shaders/Private/NiagaraMeshVertexFactory.ush b/Engine/Plugins/FX/Niagara/Shaders/Private/NiagaraMeshVertexFactory.ush
--- a/Engine/Plugins/FX/Niagara/Shaders/Private/NiagaraMeshVertexFactory.ush
+++ b/Engine/Plugins/FX/Niagara/Shaders/Private/NiagaraMeshVertexFactory.ush
@@ -13,22 +13,6 @@
 #define USE_PARTICLE_RANDOM (NEEDS_PARTICLE_RANDOM)
 #define USE_PARTICLE_TRANSFORM (NEEDS_PARTICLE_TRANSFORM)
 
-#if NIAGARA_LOOSEPARAM_BINDING
-	Buffer<float> NiagaraParticleDataPosition;
-	Buffer<float> NiagaraParticleDataVelocity;
-	Buffer<float> NiagaraParticleDataColor;
-	Buffer<float> NiagaraParticleDataScale;
-	Buffer<float> NiagaraParticleDataTransform;
-	Buffer<float> NiagaraParticleDataNormalizedAge;
-	Buffer<float> NiagaraParticleDataSubImage;
-	Buffer<float> NiagaraParticleDataMaterialRandom;
-	Buffer<float> NiagaraParticleDataMaterialParam0;
-	Buffer<float> NiagaraParticleDataMaterialParam1;
-	Buffer<float> NiagaraParticleDataMaterialParam2;
-	Buffer<float> NiagaraParticleDataMaterialParam3;
-	Buffer<float> NiagaraParticleDataCameraOffset;
-#endif
-
 float3 TransformPosition(float3 InPosition)
 {
 	if (NiagaraMeshVF.bLocalSpace)
@@ -542,7 +526,7 @@ float3 CalculateCameraOffset(FVertexFactoryIntermediates Intermediates, uint Ins
 		return (float3)0.0f;
 	}
 
-	float Offset = GetFloat(NIAGARA_LOOSEPARAM_STRUCT(NiagaraParticleDataCameraOffset), NiagaraMeshVF.CameraOffsetDataOffset, InstanceId);
+	float Offset = GetFloat(NiagaraMeshVF.CameraOffsetDataOffset, InstanceId);
 	float3 WorldPosition = TransformPosition(Intermediates.ParticlePosition.xyz);
 	float3 CameraVector = ResolvedView.WorldCameraOrigin - WorldPosition;
 	float3 CameraDir = normalize(CameraVector);
@@ -572,7 +556,7 @@ FVertexFactoryIntermediates GetVertexFactoryIntermediates(FVertexFactoryInput In
 	}
 	else
 	{
-		Intermediates.ParticlePosition = float4(GetVec3(NIAGARA_LOOSEPARAM_STRUCT(NiagaraParticleDataPosition), NiagaraMeshVF.PositionDataOffset, Input.InstanceId), 1);
+		Intermediates.ParticlePosition = float4(GetVec3(NiagaraMeshVF.PositionDataOffset, Input.InstanceId), 1);
 	}
 
 	if (NiagaraMeshVF.ColorDataOffset == -1)
@@ -581,13 +565,13 @@ FVertexFactoryIntermediates GetVertexFactoryIntermediates(FVertexFactoryInput In
 	}
 	else
 	{
-		Intermediates.ParticleColor = GetVec4(NIAGARA_LOOSEPARAM_STRUCT(NiagaraParticleDataColor), NiagaraMeshVF.ColorDataOffset, Input.InstanceId);
+		Intermediates.ParticleColor = GetVec4(NiagaraMeshVF.ColorDataOffset, Input.InstanceId);
 	}
 
 	float3 ParticleVelocity = (float3)0.0f;
 	if (NiagaraMeshVF.VelocityDataOffset != -1)
 	{
-		ParticleVelocity = float3(GetVec3(NIAGARA_LOOSEPARAM_STRUCT(NiagaraParticleDataVelocity), NiagaraMeshVF.VelocityDataOffset, Input.InstanceId));
+		ParticleVelocity = float3(GetVec3(NiagaraMeshVF.VelocityDataOffset, Input.InstanceId));
 	}
 
 	// Calculate the particle velocity (NOTE: always do this regardless of NEEDS_PARTICLE_VELOCITY because of velocity facing particles)
@@ -602,7 +586,7 @@ FVertexFactoryIntermediates GetVertexFactoryIntermediates(FVertexFactoryInput In
 	}
 	else
 	{
-		Intermediates.DynamicParameter = GetVec4(NIAGARA_LOOSEPARAM_STRUCT(NiagaraParticleDataMaterialParam0), NiagaraMeshVF.MaterialParamDataOffset, Input.InstanceId);
+		Intermediates.DynamicParameter = GetVec4(NiagaraMeshVF.MaterialParamDataOffset, Input.InstanceId);
 	}
 #endif
 #if (DYNAMIC_PARAMETERS_MASK & 2)
@@ -612,7 +596,7 @@ FVertexFactoryIntermediates GetVertexFactoryIntermediates(FVertexFactoryInput In
 	}	
 	else
 	{
-		Intermediates.DynamicParameter1 = GetVec4(NIAGARA_LOOSEPARAM_STRUCT(NiagaraParticleDataMaterialParam1), NiagaraMeshVF.MaterialParam1DataOffset, Input.InstanceId);
+		Intermediates.DynamicParameter1 = GetVec4(NiagaraMeshVF.MaterialParam1DataOffset, Input.InstanceId);
 	}
 #endif
 #if (DYNAMIC_PARAMETERS_MASK & 4)
@@ -622,7 +606,7 @@ FVertexFactoryIntermediates GetVertexFactoryIntermediates(FVertexFactoryInput In
 	}
 	else
 	{
-		Intermediates.DynamicParameter2 = GetVec4(NIAGARA_LOOSEPARAM_STRUCT(NiagaraParticleDataMaterialParam2), NiagaraMeshVF.MaterialParam2DataOffset, Input.InstanceId);
+		Intermediates.DynamicParameter2 = GetVec4(NiagaraMeshVF.MaterialParam2DataOffset, Input.InstanceId);
 	}
 #endif
 #if (DYNAMIC_PARAMETERS_MASK & 8)
@@ -632,7 +616,7 @@ FVertexFactoryIntermediates GetVertexFactoryIntermediates(FVertexFactoryInput In
 	}
 	else
 	{
-		Intermediates.DynamicParameter3 = GetVec4(NIAGARA_LOOSEPARAM_STRUCT(NiagaraParticleDataMaterialParam3), NiagaraMeshVF.MaterialParam3DataOffset, Input.InstanceId);
+		Intermediates.DynamicParameter3 = GetVec4(NiagaraMeshVF.MaterialParam3DataOffset, Input.InstanceId);
 	}
 #endif
 
@@ -643,7 +627,7 @@ FVertexFactoryIntermediates GetVertexFactoryIntermediates(FVertexFactoryInput In
 	}
 	else
 	{
-		Intermediates.RelativeTime = GetFloat(NIAGARA_LOOSEPARAM_STRUCT(NiagaraParticleDataNormalizedAge), NiagaraMeshVF.NormalizedAgeDataOffset, Input.InstanceId);
+		Intermediates.RelativeTime = GetFloat(NiagaraMeshVF.NormalizedAgeDataOffset, Input.InstanceId);
 	}
 #endif
 
@@ -654,7 +638,7 @@ FVertexFactoryIntermediates GetVertexFactoryIntermediates(FVertexFactoryInput In
 	}
 	else
 	{
-		Intermediates.ParticleRandom = GetFloat(NIAGARA_LOOSEPARAM_STRUCT(NiagaraParticleDataMaterialRandom), NiagaraMeshVF.MaterialRandomDataOffset, Input.InstanceId);
+		Intermediates.ParticleRandom = GetFloat(NiagaraMeshVF.MaterialRandomDataOffset, Input.InstanceId);
 	}
 #endif	
 
@@ -669,7 +653,7 @@ FVertexFactoryIntermediates GetVertexFactoryIntermediates(FVertexFactoryInput In
 	};
 	if(NiagaraMeshVF.TransformDataOffset != -1)
 	{
-		float4 Quat = GetVec4(NIAGARA_LOOSEPARAM_STRUCT(NiagaraParticleDataTransform), NiagaraMeshVF.TransformDataOffset, Input.InstanceId);
+		float4 Quat = GetVec4(NiagaraMeshVF.TransformDataOffset, Input.InstanceId);
 		RotationMat = QuatTo3x3(normalize(Quat));
 	}
 	
@@ -677,7 +661,7 @@ FVertexFactoryIntermediates GetVertexFactoryIntermediates(FVertexFactoryInput In
 	float3 Scale = (float3)1.0f;
 	if(NiagaraMeshVF.ScaleDataOffset != -1)
 	{
-		Scale = GetVec3(NIAGARA_LOOSEPARAM_STRUCT(NiagaraParticleDataScale), NiagaraMeshVF.ScaleDataOffset, Input.InstanceId);		
+		Scale = GetVec3(NiagaraMeshVF.ScaleDataOffset, Input.InstanceId);		
 	}
 	
 	float3x3 ScaleMat = {
@@ -747,7 +731,7 @@ FVertexFactoryIntermediates GetVertexFactoryIntermediates(FVertexFactoryInput In
 	else
 	{
 		// SubUV.
-		float SubImageIndex = GetFloat(NIAGARA_LOOSEPARAM_STRUCT(NiagaraParticleDataSubImage), NiagaraMeshVF.SubImageDataOffset, Input.InstanceId);
+		float SubImageIndex = GetFloat(NiagaraMeshVF.SubImageDataOffset, Input.InstanceId);
 		float SubImageLerp = frac(SubImageIndex);
 		float SubImageA = SubImageIndex - SubImageLerp;
 		float SubImageB = SubImageA + 1;
diff --git a/Engine/Plugins/FX/Niagara/Shaders/Private/NiagaraRibbonVertexFactory.ush b/Engine/Plugins/FX/Niagara/Shaders/Private/NiagaraRibbonVertexFactory.ush
--- a/Engine/Plugins/FX/Niagara/Shaders/Private/NiagaraRibbonVertexFactory.ush
+++ b/Engine/Plugins/FX/Niagara/Shaders/Private/NiagaraRibbonVertexFactory.ush
@@ -14,21 +14,6 @@
 #define USE_PARTICLE_RANDOM (NEEDS_PARTICLE_RANDOM)
 #define USE_PARTICLE_INTERPOLATION 1
 
-#if NIAGARA_LOOSEPARAM_BINDING
-	Buffer<float> NiagaraParticleDataPosition;
-	Buffer<float> NiagaraParticleDataVelocity;
-	Buffer<float> NiagaraParticleDataColor;
-	Buffer<float> NiagaraParticleDataWidth;
-	Buffer<float> NiagaraParticleDataTwist;
-	Buffer<float> NiagaraParticleDataFacing;
-	Buffer<float> NiagaraParticleDataNormalizedAge;
-	Buffer<float> NiagaraParticleDataMaterialRandom;
-	Buffer<float> NiagaraParticleDataMaterialParam0;
-	Buffer<float> NiagaraParticleDataMaterialParam1;
-	Buffer<float> NiagaraParticleDataMaterialParam2;
-	Buffer<float> NiagaraParticleDataMaterialParam3;
-#endif
-
 float3 TransformPosition(float3 InPosition)
 {
 	if (NiagaraRibbonVF.bLocalSpace)
@@ -159,8 +144,8 @@ void GetNiagaraParticlePositionAndDirection(in FInterpVtxData InterpVtxData, in
 	BRANCH // Also acts as ISOLATE
 	if (NiagaraRibbonVF.PositionDataOffset != -1)
 	{
-		const float3 Pos0 = GetVec3(NIAGARA_LOOSEPARAM_STRUCT(NiagaraParticleDataPosition), NiagaraRibbonVF.PositionDataOffset, InterpVtxData.ParticleId0);
-		const float3 Pos1 = GetVec3(NIAGARA_LOOSEPARAM_STRUCT(NiagaraParticleDataPosition), NiagaraRibbonVF.PositionDataOffset, InterpVtxData.ParticleId1);
+		const float3 Pos0 = GetVec3(NiagaraRibbonVF.PositionDataOffset, InterpVtxData.ParticleId0);
+		const float3 Pos1 = GetVec3(NiagaraRibbonVF.PositionDataOffset, InterpVtxData.ParticleId1);
 		const float4 TangentAndDist0 = NiagaraRibbonVFLooseParameters.TangentsAndDistances[InterpVtxData.RawParticleId0];
 		const float4 TangentAndDist1 = NiagaraRibbonVFLooseParameters.TangentsAndDistances[InterpVtxData.RawParticleId1];
 		
@@ -176,8 +161,8 @@ void GetNiagaraParticleVelocity(in FInterpVtxData InterpVtxData, in out FVertexF
 	BRANCH
 	if(NiagaraRibbonVF.VelocityDataOffset != -1)
 	{
-		const float3 Velocity0 = GetVec3(NIAGARA_LOOSEPARAM_STRUCT(NiagaraParticleDataVelocity), NiagaraRibbonVF.VelocityDataOffset, InterpVtxData.ParticleId0);
-		const float3 Velocity1 = GetVec3(NIAGARA_LOOSEPARAM_STRUCT(NiagaraParticleDataVelocity), NiagaraRibbonVF.VelocityDataOffset, InterpVtxData.ParticleId1);
+		const float3 Velocity0 = GetVec3(NiagaraRibbonVF.VelocityDataOffset, InterpVtxData.ParticleId0);
+		const float3 Velocity1 = GetVec3(NiagaraRibbonVF.VelocityDataOffset, InterpVtxData.ParticleId1);
 		const float3 Velocity = TransformVector(lerp(Velocity0, Velocity1, InterpVtxData.Alpha));
 
 		Intermediates.ParticleVelocity = float4(normalize(Velocity), length(Velocity));
@@ -191,8 +176,8 @@ void GetNiagaraParticleColor(in FInterpVtxData InterpVtxData, in out FVertexFact
 
 	if (NiagaraRibbonVF.ColorDataOffset != -1)
 	{
-		const float4 Color0 = GetVec4(NIAGARA_LOOSEPARAM_STRUCT(NiagaraParticleDataColor), NiagaraRibbonVF.ColorDataOffset, InterpVtxData.ParticleId0);
-		const float4 Color1 = GetVec4(NIAGARA_LOOSEPARAM_STRUCT(NiagaraParticleDataColor), NiagaraRibbonVF.ColorDataOffset, InterpVtxData.ParticleId1);
+		const float4 Color0 = GetVec4(NiagaraRibbonVF.ColorDataOffset, InterpVtxData.ParticleId0);
+		const float4 Color1 = GetVec4(NiagaraRibbonVF.ColorDataOffset, InterpVtxData.ParticleId1);
 		Intermediates.Color = lerp(Color0, Color1, InterpVtxData.Alpha);
 	}
 }
@@ -206,12 +191,7 @@ void GetNiagaraParticleDynamicParameters(in FInterpVtxData InterpVtxData, in out
 	FLATTEN
 	if (NiagaraRibbonVF.MaterialParamDataOffset != -1)
 	{
-		Intermediates.DynamicParameter = lerp
-		(
-			GetVec4(NIAGARA_LOOSEPARAM_STRUCT(NiagaraParticleDataMaterialParam0), NiagaraRibbonVF.MaterialParamDataOffset, InterpVtxData.ParticleId0), 
-			GetVec4(NIAGARA_LOOSEPARAM_STRUCT(NiagaraParticleDataMaterialParam0), NiagaraRibbonVF.MaterialParamDataOffset, InterpVtxData.ParticleId1), 
-			InterpVtxData.Alpha
-		);
+		Intermediates.DynamicParameter = lerp(GetVec4(NiagaraRibbonVF.MaterialParamDataOffset, InterpVtxData.ParticleId0), GetVec4(NiagaraRibbonVF.MaterialParamDataOffset, InterpVtxData.ParticleId1), InterpVtxData.Alpha);
 	}
 #endif
 #if (DYNAMIC_PARAMETERS_MASK & 2)
@@ -221,12 +201,7 @@ void GetNiagaraParticleDynamicParameters(in FInterpVtxData InterpVtxData, in out
 	FLATTEN
 	if (NiagaraRibbonVF.MaterialParam1DataOffset != -1)
 	{
-		Intermediates.DynamicParameter1 = lerp
-		(
-			GetVec4(NIAGARA_LOOSEPARAM_STRUCT(NiagaraParticleDataMaterialParam1), NiagaraRibbonVF.MaterialParam1DataOffset, InterpVtxData.ParticleId0), 
-			GetVec4(NIAGARA_LOOSEPARAM_STRUCT(NiagaraParticleDataMaterialParam1), NiagaraRibbonVF.MaterialParam1DataOffset, InterpVtxData.ParticleId1), 
-			InterpVtxData.Alpha
-		);
+		Intermediates.DynamicParameter1 = lerp(GetVec4(NiagaraRibbonVF.MaterialParam1DataOffset, InterpVtxData.ParticleId0), GetVec4(NiagaraRibbonVF.MaterialParam1DataOffset, InterpVtxData.ParticleId1), InterpVtxData.Alpha);
 	}
 #endif
 #if (DYNAMIC_PARAMETERS_MASK & 4)
@@ -236,12 +211,7 @@ void GetNiagaraParticleDynamicParameters(in FInterpVtxData InterpVtxData, in out
 	FLATTEN
 	if (NiagaraRibbonVF.MaterialParam2DataOffset != -1)
 	{
-		Intermediates.DynamicParameter2 = lerp
-		(
-			GetVec4(NIAGARA_LOOSEPARAM_STRUCT(NiagaraParticleDataMaterialParam2), NiagaraRibbonVF.MaterialParam2DataOffset, InterpVtxData.ParticleId0), 
-			GetVec4(NIAGARA_LOOSEPARAM_STRUCT(NiagaraParticleDataMaterialParam2), NiagaraRibbonVF.MaterialParam2DataOffset, InterpVtxData.ParticleId1), 
-			InterpVtxData.Alpha
-		);
+		Intermediates.DynamicParameter2 = lerp(GetVec4(NiagaraRibbonVF.MaterialParam2DataOffset, InterpVtxData.ParticleId0), GetVec4(NiagaraRibbonVF.MaterialParam2DataOffset, InterpVtxData.ParticleId1), InterpVtxData.Alpha);
 	}
 #endif
 #if (DYNAMIC_PARAMETERS_MASK & 8)
@@ -251,12 +221,7 @@ void GetNiagaraParticleDynamicParameters(in FInterpVtxData InterpVtxData, in out
 	FLATTEN
 	if (NiagaraRibbonVF.MaterialParam3DataOffset != -1)
 	{
-		Intermediates.DynamicParameter3 = lerp
-		(
-			GetVec4(NIAGARA_LOOSEPARAM_STRUCT(NiagaraParticleDataMaterialParam3), NiagaraRibbonVF.MaterialParam3DataOffset, InterpVtxData.ParticleId0), 
-			GetVec4(NIAGARA_LOOSEPARAM_STRUCT(NiagaraParticleDataMaterialParam3), NiagaraRibbonVF.MaterialParam3DataOffset, InterpVtxData.ParticleId1), 
-			InterpVtxData.Alpha
-		);
+		Intermediates.DynamicParameter3 = lerp(GetVec4(NiagaraRibbonVF.MaterialParam3DataOffset, InterpVtxData.ParticleId0), GetVec4(NiagaraRibbonVF.MaterialParam3DataOffset, InterpVtxData.ParticleId1), InterpVtxData.Alpha);
 	}
 #endif
 }
@@ -302,12 +267,7 @@ void GetNiagaraParticleRibbonWidth(in FInterpVtxData InterpVtxData, in out FVert
 	FLATTEN
 	if (NiagaraRibbonVF.WidthDataOffset != -1)
 	{
-		Intermediates.RibbonWidth = lerp
-		(
-			GetFloat(NIAGARA_LOOSEPARAM_STRUCT(NiagaraParticleDataWidth), NiagaraRibbonVF.WidthDataOffset, InterpVtxData.ParticleId0), 
-			GetFloat(NIAGARA_LOOSEPARAM_STRUCT(NiagaraParticleDataWidth), NiagaraRibbonVF.WidthDataOffset, InterpVtxData.ParticleId1), 
-			InterpVtxData.Alpha
-		);
+		Intermediates.RibbonWidth = lerp(GetFloat(NiagaraRibbonVF.WidthDataOffset, InterpVtxData.ParticleId0), GetFloat(NiagaraRibbonVF.WidthDataOffset, InterpVtxData.ParticleId1), InterpVtxData.Alpha);
 	}
 }
 
@@ -317,12 +277,7 @@ void GetNiagaraNormalizedAge(in FInterpVtxData InterpVtxData, in out FVertexFact
 	FLATTEN
 	if (NiagaraRibbonVF.NormalizedAgeDataOffset != -1)
 	{
-		Intermediates.RelativeTime = lerp
-		(
-			GetFloat(NIAGARA_LOOSEPARAM_STRUCT(NiagaraParticleDataNormalizedAge), NiagaraRibbonVF.NormalizedAgeDataOffset, InterpVtxData.ParticleId0), 
-			GetFloat(NIAGARA_LOOSEPARAM_STRUCT(NiagaraParticleDataNormalizedAge), NiagaraRibbonVF.NormalizedAgeDataOffset, InterpVtxData.ParticleId1), 
-			InterpVtxData.Alpha
-		);
+		Intermediates.RelativeTime = lerp(GetFloat(NiagaraRibbonVF.NormalizedAgeDataOffset, InterpVtxData.ParticleId0), GetFloat(NiagaraRibbonVF.NormalizedAgeDataOffset, InterpVtxData.ParticleId1), InterpVtxData.Alpha);
 	}
 }
 #endif
@@ -333,12 +288,7 @@ void GetNiagaraMaterialRandom(in FInterpVtxData InterpVtxData, in out FVertexFac
 	FLATTEN
 	if (NiagaraRibbonVF.MaterialRandomDataOffset != -1)
 	{
-		Intermediates.ParticleRandom = lerp
-		(
-			GetFloat(NIAGARA_LOOSEPARAM_STRUCT(NiagaraParticleDataMaterialRandom), NiagaraRibbonVF.MaterialRandomDataOffset, InterpVtxData.ParticleId0), 
-			GetFloat(NIAGARA_LOOSEPARAM_STRUCT(NiagaraParticleDataMaterialRandom), NiagaraRibbonVF.MaterialRandomDataOffset, InterpVtxData.ParticleId1), 
-			InterpVtxData.Alpha
-		);
+		Intermediates.ParticleRandom = lerp(GetFloat(NiagaraRibbonVF.MaterialRandomDataOffset, InterpVtxData.ParticleId0), GetFloat(NiagaraRibbonVF.MaterialRandomDataOffset, InterpVtxData.ParticleId1), InterpVtxData.Alpha);
 	}
 }
 #endif
@@ -624,8 +574,8 @@ void GetNiagaraParticleTangents(in FInterpVtxData InterpVtxData, in out FVertexF
 	BRANCH
 	if(NiagaraRibbonVF.FacingDataOffset != -1)
 	{
-		const float3 Facing0 = GetVec3(NIAGARA_LOOSEPARAM_STRUCT(NiagaraParticleDataFacing), NiagaraRibbonVF.FacingDataOffset, InterpVtxData.ParticleId0);
-		const float3 Facing1 = GetVec3(NIAGARA_LOOSEPARAM_STRUCT(NiagaraParticleDataFacing), NiagaraRibbonVF.FacingDataOffset, InterpVtxData.ParticleId1);
+		const float3 Facing0 = GetVec3(NiagaraRibbonVF.FacingDataOffset, InterpVtxData.ParticleId0);
+		const float3 Facing1 = GetVec3(NiagaraRibbonVF.FacingDataOffset, InterpVtxData.ParticleId1);
 		Intermediates.Facing = normalize(TransformVector(lerp(Facing0, Facing1, InterpVtxData.Alpha)));
 	}
 	else
@@ -651,8 +601,8 @@ void GetNiagaraParticleTangents(in FInterpVtxData InterpVtxData, in out FVertexF
 	BRANCH
 	if (NiagaraRibbonVF.TwistDataOffset != -1)
 	{
-		const float Twist0 = GetFloat(NIAGARA_LOOSEPARAM_STRUCT(NiagaraParticleDataTwist), NiagaraRibbonVF.TwistDataOffset, InterpVtxData.ParticleId0);
-		const float Twist1 = GetFloat(NIAGARA_LOOSEPARAM_STRUCT(NiagaraParticleDataTwist), NiagaraRibbonVF.TwistDataOffset, InterpVtxData.ParticleId1);
+		const float Twist0 = GetFloat(NiagaraRibbonVF.TwistDataOffset, InterpVtxData.ParticleId0);
+		const float Twist1 = GetFloat(NiagaraRibbonVF.TwistDataOffset, InterpVtxData.ParticleId1);
 		float SinTwist, CosTwist;
 		sincos(lerp(Twist0, Twist1, InterpVtxData.Alpha) /* + ResolvedView.RealTime * .5 */, SinTwist, CosTwist);
 
diff --git a/Engine/Plugins/FX/Niagara/Shaders/Private/NiagaraSortKeyGen.usf b/Engine/Plugins/FX/Niagara/Shaders/Private/NiagaraSortKeyGen.usf
--- a/Engine/Plugins/FX/Niagara/Shaders/Private/NiagaraSortKeyGen.usf
+++ b/Engine/Plugins/FX/Niagara/Shaders/Private/NiagaraSortKeyGen.usf
@@ -74,7 +74,7 @@ void GenerateParticleSortKeys(uint ParticleIndex : SV_DispatchThreadID)
 		{
 			if (SortMode == SORT_VIEW_DEPTH || SortMode == SORT_VIEW_DISTANCE)
 			{ 
-				float3 ParticlePosition = GetVec3(NIAGARA_LOOSEPARAM_STRUCT(NiagaraParticleDataFloat), SortAttributeOffset, ParticleIndex); 
+				float3 ParticlePosition = GetVec3(SortAttributeOffset, ParticleIndex); 
 				float3 CameraVector = ParticlePosition - CameraPosition;
 				SortValue = SortMode == SORT_VIEW_DEPTH ? dot(CameraVector, CameraDirection) : length(CameraVector);
 				
@@ -88,7 +88,7 @@ void GenerateParticleSortKeys(uint ParticleIndex : SV_DispatchThreadID)
 			}
 			else
 			{
-				SortValue = GetFloat(NIAGARA_LOOSEPARAM_STRUCT(NiagaraParticleDataFloat), SortAttributeOffset, ParticleIndex);
+				SortValue = GetFloat(SortAttributeOffset, ParticleIndex);
 
 #if SORT_MAX_PRECISION
 				SortKey = (asuint(SortValue)) >> SortKeyShift;
diff --git a/Engine/Plugins/FX/Niagara/Shaders/Private/NiagaraSpriteVertexFactory.ush b/Engine/Plugins/FX/Niagara/Shaders/Private/NiagaraSpriteVertexFactory.ush
--- a/Engine/Plugins/FX/Niagara/Shaders/Private/NiagaraSpriteVertexFactory.ush
+++ b/Engine/Plugins/FX/Niagara/Shaders/Private/NiagaraSpriteVertexFactory.ush
@@ -35,26 +35,6 @@
 /** Blends between FaceCamera and FaceCameraPosition.*/
 #define FACING_MODE_FACECAMERADISTANCEBLEND 4
 
-#if NIAGARA_LOOSEPARAM_BINDING
-	Buffer<float> NiagaraParticleDataPosition;
-	Buffer<float> NiagaraParticleDataVelocity;
-	Buffer<float> NiagaraParticleDataColor;
-	Buffer<float> NiagaraParticleDataRotation;
-	Buffer<float> NiagaraParticleDataSize;
-	Buffer<float> NiagaraParticleDataFacing;
-	Buffer<float> NiagaraParticleDataAlignment;
-	Buffer<float> NiagaraParticleDataSubImage;
-	Buffer<float> NiagaraParticleDataCameraOffset;
-	Buffer<float> NiagaraParticleDataUVScale;
-	Buffer<float> NiagaraParticleDataNormalizedAge;
-	Buffer<float> NiagaraParticleDataMaterialRandom;
-	Buffer<float> NiagaraParticleDataCustomSorting;
-	Buffer<float> NiagaraParticleDataMaterialParam0;
-	Buffer<float> NiagaraParticleDataMaterialParam1;
-	Buffer<float> NiagaraParticleDataMaterialParam2;
-	Buffer<float> NiagaraParticleDataMaterialParam3;
-#endif
-
 struct FVertexFactoryInput
 {
 	float2	TexCoord				: ATTRIBUTE0;
@@ -655,7 +635,7 @@ float3 GetNiagaraParticlePosition(uint InstanceID)
 		// See FNiagaraConstants for a synchronized value...
 		return TransformPosition(NiagaraSpriteVF.DefaultPos.xyz);
 	}
-	return TransformPosition(GetVec3(NIAGARA_LOOSEPARAM_STRUCT(NiagaraParticleDataPosition), NiagaraSpriteVF.PositionDataOffset, InstanceID));
+	return TransformPosition(GetVec3(NiagaraSpriteVF.PositionDataOffset, InstanceID));
 }
 float3 GetNiagaraParticlePrevPosition(uint InstanceID)
 {
@@ -664,7 +644,7 @@ float3 GetNiagaraParticlePrevPosition(uint InstanceID)
 		// See FNiagaraConstants for a synchronized value...
 		return TransformPrevPosition(NiagaraSpriteVF.DefaultPos.xyz);
 	}
-	return TransformPrevPosition(GetVec3(NIAGARA_LOOSEPARAM_STRUCT(NiagaraParticleDataPosition), NiagaraSpriteVF.PositionDataOffset, InstanceID));
+	return TransformPrevPosition(GetVec3(NiagaraSpriteVF.PositionDataOffset, InstanceID));
 }
 float3 GetNiagaraParticleVelocity(uint InstanceID)
 {
@@ -673,7 +653,7 @@ float3 GetNiagaraParticleVelocity(uint InstanceID)
 		// See FNiagaraConstants for a synchronized value...
 		return TransformVector(float3(0.0f,0.0f,0.0f));
 	}
-	return TransformVector(GetVec3(NIAGARA_LOOSEPARAM_STRUCT(NiagaraParticleDataVelocity), NiagaraSpriteVF.VelocityDataOffset, InstanceID));
+	return TransformVector(GetVec3(NiagaraSpriteVF.VelocityDataOffset, InstanceID));
 }
 float2 GetNiagaraParticleSize(uint InstanceID)
 {
@@ -682,7 +662,7 @@ float2 GetNiagaraParticleSize(uint InstanceID)
 		// See FNiagaraConstants for a synchronized value...
 		return float2(50.0f, 50.0f);
 	}
-	return GetVec2(NIAGARA_LOOSEPARAM_STRUCT(NiagaraParticleDataSize), NiagaraSpriteVF.SizeDataOffset, InstanceID);
+	return GetVec2(NiagaraSpriteVF.SizeDataOffset, InstanceID);
 }
 float GetNiagaraParticleRotation(uint InstanceID)
 {
@@ -691,7 +671,7 @@ float GetNiagaraParticleRotation(uint InstanceID)
 		// See FNiagaraConstants for a synchronized value...
 		return 0.0f;
 	}
-	return (GetFloat(NIAGARA_LOOSEPARAM_STRUCT(NiagaraParticleDataRotation), NiagaraSpriteVF.RotationDataOffset, InstanceID) / 180.0f) * PI;
+	return (GetFloat(NiagaraSpriteVF.RotationDataOffset, InstanceID) / 180.0f) * PI;
 }
 float4 GetNiagaraParticleColor(uint InstanceID)
 {
@@ -700,7 +680,7 @@ float4 GetNiagaraParticleColor(uint InstanceID)
 		// See FNiagaraConstants for a synchronized value...
 		return float4(1.0f, 1.0f, 1.0f, 1.0f);
 	}
-	return GetVec4(NIAGARA_LOOSEPARAM_STRUCT(NiagaraParticleDataColor), NiagaraSpriteVF.ColorDataOffset, InstanceID);
+	return GetVec4(NiagaraSpriteVF.ColorDataOffset, InstanceID);
 }
 
 float2 GetNiagaraUVScale(uint InstanceID)
@@ -710,7 +690,7 @@ float2 GetNiagaraUVScale(uint InstanceID)
 		// See FNiagaraConstants for a synchronized value...
 		return float2(1.0f, 1.0f);
 	}
-	return GetVec2(NIAGARA_LOOSEPARAM_STRUCT(NiagaraParticleDataUVScale), NiagaraSpriteVF.UVScaleDataOffset, InstanceID);	
+	return GetVec2(NiagaraSpriteVF.UVScaleDataOffset, InstanceID);	
 }
 
 float GetNiagaraParticleRandom(uint InstanceID)
@@ -720,7 +700,7 @@ float GetNiagaraParticleRandom(uint InstanceID)
 		// See FNiagaraConstants for a synchronized value...
 		return 0.0f;
 	}
-	return GetFloat(NIAGARA_LOOSEPARAM_STRUCT(NiagaraParticleDataMaterialRandom), NiagaraSpriteVF.MaterialRandomDataOffset, InstanceID);	
+	return GetFloat(NiagaraSpriteVF.MaterialRandomDataOffset, InstanceID);	
 }
 
 float GetNiagaraCameraOffset(uint InstanceID)
@@ -730,7 +710,7 @@ float GetNiagaraCameraOffset(uint InstanceID)
 		// See FNiagaraConstants for a synchronized value...
 		return 0.0f;
 	}
-	return GetFloat(NIAGARA_LOOSEPARAM_STRUCT(NiagaraParticleDataCameraOffset), NiagaraSpriteVF.CameraOffsetDataOffset, InstanceID);
+	return GetFloat(NiagaraSpriteVF.CameraOffsetDataOffset, InstanceID);
 }
 
 #if USE_PARTICLE_TIME
@@ -740,7 +720,7 @@ float GetNiagaraNormalizedAge(uint InstanceID)
 	{
 		return 0;
 	}
-	return GetFloat(NIAGARA_LOOSEPARAM_STRUCT(NiagaraParticleDataNormalizedAge), NiagaraSpriteVF.NormalizedAgeDataOffset, InstanceID);
+	return GetFloat(NiagaraSpriteVF.NormalizedAgeDataOffset, InstanceID);
 }
 #endif
 
@@ -749,7 +729,7 @@ float4 GetNiagaraParticleDynamicParameters(uint InstanceID)
 {
 	if (NiagaraSpriteVF.MaterialParamDataOffset != -1)
 	{
-		return GetVec4(NIAGARA_LOOSEPARAM_STRUCT(NiagaraParticleDataMaterialParam0), NiagaraSpriteVF.MaterialParamDataOffset, InstanceID);
+		return GetVec4(NiagaraSpriteVF.MaterialParamDataOffset, InstanceID);
 	}
 	else
 	{
@@ -762,7 +742,7 @@ float4 GetNiagaraParticleDynamicParameters1(uint InstanceID)
 {
 	if (NiagaraSpriteVF.MaterialParam1DataOffset != -1)
 	{
-		return GetVec4(NIAGARA_LOOSEPARAM_STRUCT(NiagaraParticleDataMaterialParam1), NiagaraSpriteVF.MaterialParam1DataOffset, InstanceID);
+		return GetVec4(NiagaraSpriteVF.MaterialParam1DataOffset, InstanceID);
 	}
 	else
 	{
@@ -775,7 +755,7 @@ float4 GetNiagaraParticleDynamicParameters2(uint InstanceID)
 {
 	if (NiagaraSpriteVF.MaterialParam2DataOffset != -1)
 	{
-		return GetVec4(NIAGARA_LOOSEPARAM_STRUCT(NiagaraParticleDataMaterialParam2), NiagaraSpriteVF.MaterialParam2DataOffset, InstanceID);
+		return GetVec4(NiagaraSpriteVF.MaterialParam2DataOffset, InstanceID);
 	}
 	else
 	{
@@ -788,7 +768,7 @@ float4 GetNiagaraParticleDynamicParameters3(uint InstanceID)
 {
 	if (NiagaraSpriteVF.MaterialParam3DataOffset != -1)
 	{
-		return GetVec4(NIAGARA_LOOSEPARAM_STRUCT(NiagaraParticleDataMaterialParam3), NiagaraSpriteVF.MaterialParam3DataOffset, InstanceID);
+		return GetVec4(NiagaraSpriteVF.MaterialParam3DataOffset, InstanceID);
 	}
 	else
 	{
@@ -804,7 +784,7 @@ float GetNiagaraParticleSubimage(uint InstanceID)
 	float Val = 0.0f;
     if(NiagaraSpriteVF.SubimageDataOffset != -1)
 	{
-		Val = GetFloat(NIAGARA_LOOSEPARAM_STRUCT(NiagaraParticleDataSubImage), NiagaraSpriteVF.SubimageDataOffset, InstanceID);
+		Val = GetFloat(NiagaraSpriteVF.SubimageDataOffset, InstanceID);
 	}
 	return Val;
 }
@@ -815,7 +795,7 @@ float3 GetNiagaraParticleFacingVector(uint InstanceID)
 	float3 Val = float3(1.0f, 0.0f, 0.0f);
     if(NiagaraSpriteVF.FacingDataOffset != -1)
 	{
-		Val = GetVec3(NIAGARA_LOOSEPARAM_STRUCT(NiagaraParticleDataFacing), NiagaraSpriteVF.FacingDataOffset, InstanceID);
+		Val = GetVec3(NiagaraSpriteVF.FacingDataOffset, InstanceID);
 	}
 	return TransformVector(Val);
 }
@@ -825,7 +805,7 @@ float3 GetNiagaraParticleAlignmentVector(uint InstanceID)
 	float3 Val = float3(1.0f, 0.0f, 0.0f);
     if(NiagaraSpriteVF.AlignmentDataOffset != -1)
 	{
-		Val = GetVec3(NIAGARA_LOOSEPARAM_STRUCT(NiagaraParticleDataAlignment), NiagaraSpriteVF.AlignmentDataOffset, InstanceID);
+		Val = GetVec3(NiagaraSpriteVF.AlignmentDataOffset, InstanceID);
 	}
 	return TransformVector(Val);
 }
diff --git a/Engine/Plugins/FX/Niagara/Shaders/Private/NiagaraVFParticleAccess.usf b/Engine/Plugins/FX/Niagara/Shaders/Private/NiagaraVFParticleAccess.usf
--- a/Engine/Plugins/FX/Niagara/Shaders/Private/NiagaraVFParticleAccess.usf
+++ b/Engine/Plugins/FX/Niagara/Shaders/Private/NiagaraVFParticleAccess.usf
@@ -12,40 +12,37 @@
 #define NIAGARA_KEYGEN 0
 #endif
 
-#define NIAGARA_LOOSEPARAM_BINDING (!RAYHITGROUPSHADER && !COMPUTESHADER) || NIAGARA_KEYGEN
+#if (!RAYHITGROUPSHADER && !COMPUTESHADER) || NIAGARA_KEYGEN
 
-#if NIAGARA_LOOSEPARAM_BINDING
-
-#define NIAGARA_LOOSEPARAM_STRUCT(x) x
+Buffer<float> NiagaraParticleDataFloat;
 uint NiagaraFloatDataStride;
 
-float GetFloat(Buffer<float> NiagaraBuffer, int RegisterIdx, uint InstanceID)
+float GetFloat(int RegisterIdx, uint InstanceID)
 {
-	return NiagaraBuffer[(RegisterIdx * NiagaraFloatDataStride + InstanceID)];
+	return NiagaraParticleDataFloat[(RegisterIdx * NiagaraFloatDataStride + InstanceID)];
 }
 
 #else
-#define NIAGARA_LOOSEPARAM_STRUCT(x) NiagaraVFLooseParameters.x
 
-float GetFloat(Buffer<float> NiagaraBuffer, int RegisterIdx, uint InstanceID)
+float GetFloat(int RegisterIdx, uint InstanceID)
 {
-	return NiagaraBuffer[(RegisterIdx * NiagaraVFLooseParameters.NiagaraFloatDataStride + InstanceID)];
+	return NiagaraVFLooseParameters.NiagaraParticleDataFloat[(RegisterIdx * NiagaraVFLooseParameters.NiagaraFloatDataStride + InstanceID)];
 }
 
 #endif
 
-float2 GetVec2(Buffer<float> NiagaraBuffer, int RegisterIndex, uint InstanceID)
+float2 GetVec2(int RegisterIndex, uint InstanceID)
 {
-	return float2(GetFloat(NiagaraBuffer, RegisterIndex, InstanceID), GetFloat(NiagaraBuffer, RegisterIndex+1, InstanceID));
+	return float2(GetFloat(RegisterIndex, InstanceID), GetFloat(RegisterIndex+1, InstanceID));
 }
 
-float3 GetVec3(Buffer<float> NiagaraBuffer, int RegisterIndex, uint InstanceID)
+float3 GetVec3(int RegisterIndex, uint InstanceID)
 {
-	return float3(GetFloat(NiagaraBuffer, RegisterIndex, InstanceID), GetFloat(NiagaraBuffer, RegisterIndex+1, InstanceID), GetFloat(NiagaraBuffer, RegisterIndex+2, InstanceID));
+	return float3(GetFloat(RegisterIndex, InstanceID), GetFloat(RegisterIndex+1, InstanceID), GetFloat(RegisterIndex+2, InstanceID));
 }
 
-float4 GetVec4(Buffer<float> NiagaraBuffer, int RegisterIndex, uint InstanceID)
+float4 GetVec4(int RegisterIndex, uint InstanceID)
 {
-	return float4(GetFloat(NiagaraBuffer, RegisterIndex, InstanceID), GetFloat(NiagaraBuffer, RegisterIndex+1, InstanceID), GetFloat(NiagaraBuffer, RegisterIndex+2, InstanceID), GetFloat(NiagaraBuffer, RegisterIndex+3, InstanceID));
+	return float4(GetFloat(RegisterIndex, InstanceID), GetFloat(RegisterIndex+1, InstanceID), GetFloat(RegisterIndex+2, InstanceID), GetFloat(RegisterIndex+3, InstanceID));
 }
 
diff --git a/Engine/Plugins/FX/Niagara/Source/Niagara/Private/NiagaraRendererRibbons.cpp b/Engine/Plugins/FX/Niagara/Source/Niagara/Private/NiagaraRendererRibbons.cpp
--- a/Engine/Plugins/FX/Niagara/Source/Niagara/Private/NiagaraRendererRibbons.cpp
+++ b/Engine/Plugins/FX/Niagara/Source/Niagara/Private/NiagaraRendererRibbons.cpp
@@ -796,18 +796,7 @@ void FNiagaraRendererRibbons::SetupMeshBatchAndCollectorResourceForView(
 	VFLooseParams.TangentsAndDistances = TangentsAndDistancesBuffer.SRV;
 	VFLooseParams.MultiRibbonIndices = MultiRibbonIndicesBuffer.SRV;
 	VFLooseParams.PackedPerRibbonDataByIndex = PackedPerRibbonDataByIndexBuffer.SRV;
-	VFLooseParams.NiagaraParticleDataPosition = CollectorResources.VertexFactory.GetParticleDataFloatSRV();
-	VFLooseParams.NiagaraParticleDataVelocity = CollectorResources.VertexFactory.GetParticleDataFloatSRV();
-	VFLooseParams.NiagaraParticleDataColor = CollectorResources.VertexFactory.GetParticleDataFloatSRV();
-	VFLooseParams.NiagaraParticleDataWidth = CollectorResources.VertexFactory.GetParticleDataFloatSRV();
-	VFLooseParams.NiagaraParticleDataTwist = CollectorResources.VertexFactory.GetParticleDataFloatSRV();
-	VFLooseParams.NiagaraParticleDataFacing = CollectorResources.VertexFactory.GetParticleDataFloatSRV();
-	VFLooseParams.NiagaraParticleDataNormalizedAge = CollectorResources.VertexFactory.GetParticleDataFloatSRV();
-	VFLooseParams.NiagaraParticleDataMaterialRandom = CollectorResources.VertexFactory.GetParticleDataFloatSRV();
-	VFLooseParams.NiagaraParticleDataMaterialParam0 = CollectorResources.VertexFactory.GetParticleDataFloatSRV();
-	VFLooseParams.NiagaraParticleDataMaterialParam1 = CollectorResources.VertexFactory.GetParticleDataFloatSRV();
-	VFLooseParams.NiagaraParticleDataMaterialParam2 = CollectorResources.VertexFactory.GetParticleDataFloatSRV();
-	VFLooseParams.NiagaraParticleDataMaterialParam3 = CollectorResources.VertexFactory.GetParticleDataFloatSRV();
+	VFLooseParams.NiagaraParticleDataFloat = CollectorResources.VertexFactory.GetParticleDataFloatSRV();
 	VFLooseParams.NiagaraFloatDataStride = CollectorResources.VertexFactory.GetFloatDataStride();
 	VFLooseParams.SortedIndicesOffset = CollectorResources.VertexFactory.GetSortedIndicesOffset();
 	VFLooseParams.FacingMode = CollectorResources.VertexFactory.GetFacingMode();
diff --git a/Engine/Plugins/FX/Niagara/Source/Niagara/Private/NiagaraRendererSprites.cpp b/Engine/Plugins/FX/Niagara/Source/Niagara/Private/NiagaraRendererSprites.cpp
--- a/Engine/Plugins/FX/Niagara/Source/Niagara/Private/NiagaraRendererSprites.cpp
+++ b/Engine/Plugins/FX/Niagara/Source/Niagara/Private/NiagaraRendererSprites.cpp
@@ -460,25 +460,7 @@ void FNiagaraRendererSprites::CreateMeshBatchForView(
 	FNiagaraSpriteVFLooseParameters VFLooseParams;
 	VFLooseParams.NumCutoutVerticesPerFrame = CollectorResources.VertexFactory.GetNumCutoutVerticesPerFrame();
 	VFLooseParams.CutoutGeometry = CollectorResources.VertexFactory.GetCutoutGeometrySRV() ? CollectorResources.VertexFactory.GetCutoutGeometrySRV() : GFNiagaraNullCutoutVertexBuffer.VertexBufferSRV.GetReference();
-	
-	VFLooseParams.NiagaraParticleDataPosition = CollectorResources.VertexFactory.GetParticleDataFloatSRV();
-	VFLooseParams.NiagaraParticleDataVelocity = CollectorResources.VertexFactory.GetParticleDataFloatSRV();
-	VFLooseParams.NiagaraParticleDataColor = CollectorResources.VertexFactory.GetParticleDataFloatSRV();
-	VFLooseParams.NiagaraParticleDataRotation = CollectorResources.VertexFactory.GetParticleDataFloatSRV();
-	VFLooseParams.NiagaraParticleDataSize = CollectorResources.VertexFactory.GetParticleDataFloatSRV();
-	VFLooseParams.NiagaraParticleDataFacing = CollectorResources.VertexFactory.GetParticleDataFloatSRV();
-	VFLooseParams.NiagaraParticleDataAlignment = CollectorResources.VertexFactory.GetParticleDataFloatSRV();
-	VFLooseParams.NiagaraParticleDataSubImage = CollectorResources.VertexFactory.GetParticleDataFloatSRV();
-	VFLooseParams.NiagaraParticleDataCameraOffset = CollectorResources.VertexFactory.GetParticleDataFloatSRV();
-	VFLooseParams.NiagaraParticleDataUVScale = CollectorResources.VertexFactory.GetParticleDataFloatSRV();
-	VFLooseParams.NiagaraParticleDataNormalizedAge = CollectorResources.VertexFactory.GetParticleDataFloatSRV();
-	VFLooseParams.NiagaraParticleDataMaterialRandom = CollectorResources.VertexFactory.GetParticleDataFloatSRV();
-	VFLooseParams.NiagaraParticleDataCustomSorting = CollectorResources.VertexFactory.GetParticleDataFloatSRV();
-	VFLooseParams.NiagaraParticleDataMaterialParam0 = CollectorResources.VertexFactory.GetParticleDataFloatSRV();
-	VFLooseParams.NiagaraParticleDataMaterialParam1 = CollectorResources.VertexFactory.GetParticleDataFloatSRV();
-	VFLooseParams.NiagaraParticleDataMaterialParam2 = CollectorResources.VertexFactory.GetParticleDataFloatSRV();
-	VFLooseParams.NiagaraParticleDataMaterialParam3 = CollectorResources.VertexFactory.GetParticleDataFloatSRV();
-
+	VFLooseParams.NiagaraParticleDataFloat = CollectorResources.VertexFactory.GetParticleDataFloatSRV();
 	VFLooseParams.NiagaraFloatDataStride = CollectorResources.VertexFactory.GetFloatDataStride();
 	VFLooseParams.ParticleAlignmentMode = CollectorResources.VertexFactory.GetAlignmentMode();
 	VFLooseParams.ParticleFacingMode = CollectorResources.VertexFactory.GetFacingMode();
diff --git a/Engine/Plugins/FX/Niagara/Source/NiagaraVertexFactories/Private/NiagaraMeshVertexFactory.cpp b/Engine/Plugins/FX/Niagara/Source/NiagaraVertexFactories/Private/NiagaraMeshVertexFactory.cpp
--- a/Engine/Plugins/FX/Niagara/Source/NiagaraVertexFactories/Private/NiagaraMeshVertexFactory.cpp
+++ b/Engine/Plugins/FX/Niagara/Source/NiagaraVertexFactories/Private/NiagaraMeshVertexFactory.cpp
@@ -19,20 +19,7 @@ public:
 	void Bind(const FShaderParameterMap& ParameterMap)
 	{
 		//PrevTransformBuffer.Bind(ParameterMap, TEXT("PrevTransformBuffer"));
-		NiagaraParticleDataPosition.Bind(ParameterMap, TEXT("NiagaraParticleDataPosition"));
-		NiagaraParticleDataVelocity.Bind(ParameterMap, TEXT("NiagaraParticleDataVelocity"));
-		NiagaraParticleDataColor.Bind(ParameterMap, TEXT("NiagaraParticleDataColor"));
-		NiagaraParticleDataScale.Bind(ParameterMap, TEXT("NiagaraParticleDataScale"));
-		NiagaraParticleDataTransform.Bind(ParameterMap, TEXT("NiagaraParticleDataTransform"));
-		NiagaraParticleDataNormalizedAge.Bind(ParameterMap, TEXT("NiagaraParticleDataNormalizedAge"));
-		NiagaraParticleDataMaterialRandom.Bind(ParameterMap, TEXT("NiagaraParticleDataMaterialRandom"));
-		NiagaraParticleDataMaterialParam0.Bind(ParameterMap, TEXT("NiagaraParticleDataMaterialParam0"));
-		NiagaraParticleDataMaterialParam1.Bind(ParameterMap, TEXT("NiagaraParticleDataMaterialParam1"));
-		NiagaraParticleDataMaterialParam2.Bind(ParameterMap, TEXT("NiagaraParticleDataMaterialParam2"));
-		NiagaraParticleDataMaterialParam3.Bind(ParameterMap, TEXT("NiagaraParticleDataMaterialParam3"));
-		NiagaraParticleDataSubImage.Bind(ParameterMap, TEXT("NiagaraParticleDataSubImage"));
-		NiagaraParticleDataCameraOffset.Bind(ParameterMap, TEXT("NiagaraParticleDataCameraOffset"));
-
+		NiagaraParticleDataFloat.Bind(ParameterMap, TEXT("NiagaraParticleDataFloat"));
 		FloatDataStride.Bind(ParameterMap, TEXT("NiagaraFloatDataStride"));
 
 		// 		NiagaraParticleDataInt.Bind(ParameterMap, TEXT("NiagaraParticleDataInt"));
@@ -57,20 +44,7 @@ public:
 		FNiagaraMeshVertexFactory* NiagaraMeshVF = (FNiagaraMeshVertexFactory*)VertexFactory;
 		ShaderBindings.Add(Shader->GetUniformBufferParameter<FNiagaraMeshUniformParameters>(), NiagaraMeshVF->GetUniformBuffer());
 
-		ShaderBindings.Add(NiagaraParticleDataPosition, NiagaraMeshVF->GetParticleDataFloatSRV());
-		ShaderBindings.Add(NiagaraParticleDataVelocity, NiagaraMeshVF->GetParticleDataFloatSRV());
-		ShaderBindings.Add(NiagaraParticleDataColor, NiagaraMeshVF->GetParticleDataFloatSRV());
-		ShaderBindings.Add(NiagaraParticleDataScale, NiagaraMeshVF->GetParticleDataFloatSRV());
-		ShaderBindings.Add(NiagaraParticleDataTransform, NiagaraMeshVF->GetParticleDataFloatSRV());
-		ShaderBindings.Add(NiagaraParticleDataNormalizedAge, NiagaraMeshVF->GetParticleDataFloatSRV());
-		ShaderBindings.Add(NiagaraParticleDataMaterialRandom, NiagaraMeshVF->GetParticleDataFloatSRV());
-		ShaderBindings.Add(NiagaraParticleDataMaterialParam0, NiagaraMeshVF->GetParticleDataFloatSRV());
-		ShaderBindings.Add(NiagaraParticleDataMaterialParam1, NiagaraMeshVF->GetParticleDataFloatSRV());
-		ShaderBindings.Add(NiagaraParticleDataMaterialParam2, NiagaraMeshVF->GetParticleDataFloatSRV());
-		ShaderBindings.Add(NiagaraParticleDataMaterialParam3, NiagaraMeshVF->GetParticleDataFloatSRV());
-		ShaderBindings.Add(NiagaraParticleDataSubImage, NiagaraMeshVF->GetParticleDataFloatSRV());
-		ShaderBindings.Add(NiagaraParticleDataCameraOffset, NiagaraMeshVF->GetParticleDataFloatSRV());
-
+		ShaderBindings.Add(NiagaraParticleDataFloat, NiagaraMeshVF->GetParticleDataFloatSRV());
 		ShaderBindings.Add(FloatDataStride, NiagaraMeshVF->GetFloatDataStride());
 
 		FRHIShaderResourceView* SortedSRV = NiagaraMeshVF->GetSortedIndicesSRV();
@@ -82,21 +56,8 @@ private:
 
 
 		//LAYOUT_FIELD(FShaderResourceParameter, PrevTransformBuffer)
-
-	LAYOUT_FIELD(FShaderResourceParameter, NiagaraParticleDataPosition);
-	LAYOUT_FIELD(FShaderResourceParameter, NiagaraParticleDataVelocity);
-	LAYOUT_FIELD(FShaderResourceParameter, NiagaraParticleDataColor);
-	LAYOUT_FIELD(FShaderResourceParameter, NiagaraParticleDataScale);
-	LAYOUT_FIELD(FShaderResourceParameter, NiagaraParticleDataTransform);
-	LAYOUT_FIELD(FShaderResourceParameter, NiagaraParticleDataNormalizedAge);
-	LAYOUT_FIELD(FShaderResourceParameter, NiagaraParticleDataMaterialRandom);
-	LAYOUT_FIELD(FShaderResourceParameter, NiagaraParticleDataMaterialParam0);
-	LAYOUT_FIELD(FShaderResourceParameter, NiagaraParticleDataMaterialParam1);
-	LAYOUT_FIELD(FShaderResourceParameter, NiagaraParticleDataMaterialParam2);
-	LAYOUT_FIELD(FShaderResourceParameter, NiagaraParticleDataMaterialParam3);
-	LAYOUT_FIELD(FShaderResourceParameter, NiagaraParticleDataSubImage);
-	LAYOUT_FIELD(FShaderResourceParameter, NiagaraParticleDataCameraOffset);
-	LAYOUT_FIELD(FShaderParameter, FloatDataStride);
+		LAYOUT_FIELD(FShaderResourceParameter, NiagaraParticleDataFloat);
+		LAYOUT_FIELD(FShaderParameter, FloatDataStride);
 
 		// 	LAYOUT_FIELD(FShaderResourceParameter, NiagaraParticleDataInt)
 		// 	LAYOUT_FIELD(FShaderParameter, Int32DataOffset)
diff --git a/Engine/Plugins/FX/Niagara/Source/NiagaraVertexFactories/Private/NiagaraRibbonVertexFactory.cpp b/Engine/Plugins/FX/Niagara/Source/NiagaraVertexFactories/Private/NiagaraRibbonVertexFactory.cpp
--- a/Engine/Plugins/FX/Niagara/Source/NiagaraVertexFactories/Private/NiagaraRibbonVertexFactory.cpp
+++ b/Engine/Plugins/FX/Niagara/Source/NiagaraVertexFactories/Private/NiagaraRibbonVertexFactory.cpp
@@ -30,19 +30,7 @@ class FNiagaraRibbonVertexFactoryShaderParametersVS : public FNiagaraRibbonVerte
 public:
 	void Bind(const FShaderParameterMap& ParameterMap)
 	{
-		NiagaraParticleDataPosition.Bind(ParameterMap, TEXT("NiagaraParticleDataPosition"));
-		NiagaraParticleDataVelocity.Bind(ParameterMap, TEXT("NiagaraParticleDataVelocity"));
-		NiagaraParticleDataColor.Bind(ParameterMap, TEXT("NiagaraParticleDataColor"));
-		NiagaraParticleDataWidth.Bind(ParameterMap, TEXT("NiagaraParticleDataWidth"));
-		NiagaraParticleDataTwist.Bind(ParameterMap, TEXT("NiagaraParticleDataTwist"));
-		NiagaraParticleDataFacing.Bind(ParameterMap, TEXT("NiagaraParticleDataFacing"));
-		NiagaraParticleDataNormalizedAge.Bind(ParameterMap, TEXT("NiagaraParticleDataNormalizedAge"));
-		NiagaraParticleDataMaterialRandom.Bind(ParameterMap, TEXT("NiagaraParticleDataMaterialRandom"));
-		NiagaraParticleDataMaterialParam0.Bind(ParameterMap, TEXT("NiagaraParticleDataMaterialParam0"));
-		NiagaraParticleDataMaterialParam1.Bind(ParameterMap, TEXT("NiagaraParticleDataMaterialParam1"));
-		NiagaraParticleDataMaterialParam2.Bind(ParameterMap, TEXT("NiagaraParticleDataMaterialParam2"));
-		NiagaraParticleDataMaterialParam3.Bind(ParameterMap, TEXT("NiagaraParticleDataMaterialParam3"));
-
+		NiagaraParticleDataFloat.Bind(ParameterMap, TEXT("NiagaraParticleDataFloat"));
 		FloatDataStride.Bind(ParameterMap, TEXT("NiagaraFloatDataStride"));
 	}
 
@@ -60,35 +48,12 @@ public:
 		FNiagaraRibbonVertexFactory* RibbonVF = (FNiagaraRibbonVertexFactory*)VertexFactory;
 		ShaderBindings.Add(Shader->GetUniformBufferParameter<FNiagaraRibbonUniformParameters>(), RibbonVF->GetRibbonUniformBuffer());
 		ShaderBindings.Add(Shader->GetUniformBufferParameter<FNiagaraRibbonVFLooseParameters>(), RibbonVF->LooseParameterUniformBuffer);
-		ShaderBindings.Add(NiagaraParticleDataPosition, RibbonVF->GetParticleDataFloatSRV());
-		ShaderBindings.Add(NiagaraParticleDataVelocity, RibbonVF->GetParticleDataFloatSRV());
-		ShaderBindings.Add(NiagaraParticleDataColor, RibbonVF->GetParticleDataFloatSRV());
-		ShaderBindings.Add(NiagaraParticleDataWidth, RibbonVF->GetParticleDataFloatSRV());
-		ShaderBindings.Add(NiagaraParticleDataTwist, RibbonVF->GetParticleDataFloatSRV());
-		ShaderBindings.Add(NiagaraParticleDataFacing, RibbonVF->GetParticleDataFloatSRV());
-		ShaderBindings.Add(NiagaraParticleDataNormalizedAge, RibbonVF->GetParticleDataFloatSRV());
-		ShaderBindings.Add(NiagaraParticleDataMaterialRandom, RibbonVF->GetParticleDataFloatSRV());
-		ShaderBindings.Add(NiagaraParticleDataMaterialParam0, RibbonVF->GetParticleDataFloatSRV());
-		ShaderBindings.Add(NiagaraParticleDataMaterialParam1, RibbonVF->GetParticleDataFloatSRV());
-		ShaderBindings.Add(NiagaraParticleDataMaterialParam2, RibbonVF->GetParticleDataFloatSRV());
-		ShaderBindings.Add(NiagaraParticleDataMaterialParam3, RibbonVF->GetParticleDataFloatSRV());
+		ShaderBindings.Add(NiagaraParticleDataFloat, RibbonVF->GetParticleDataFloatSRV());
 		ShaderBindings.Add(FloatDataStride, RibbonVF->GetFloatDataStride());
 	}
 
 private:
-	LAYOUT_FIELD(FShaderResourceParameter, NiagaraParticleDataPosition);
-	LAYOUT_FIELD(FShaderResourceParameter, NiagaraParticleDataVelocity);
-	LAYOUT_FIELD(FShaderResourceParameter, NiagaraParticleDataColor);
-	LAYOUT_FIELD(FShaderResourceParameter, NiagaraParticleDataWidth);
-	LAYOUT_FIELD(FShaderResourceParameter, NiagaraParticleDataTwist);
-	LAYOUT_FIELD(FShaderResourceParameter, NiagaraParticleDataFacing);
-	LAYOUT_FIELD(FShaderResourceParameter, NiagaraParticleDataNormalizedAge);
-	LAYOUT_FIELD(FShaderResourceParameter, NiagaraParticleDataMaterialRandom);
-	LAYOUT_FIELD(FShaderResourceParameter, NiagaraParticleDataMaterialParam0);
-	LAYOUT_FIELD(FShaderResourceParameter, NiagaraParticleDataMaterialParam1);
-	LAYOUT_FIELD(FShaderResourceParameter, NiagaraParticleDataMaterialParam2);
-	LAYOUT_FIELD(FShaderResourceParameter, NiagaraParticleDataMaterialParam3);
-
+	LAYOUT_FIELD(FShaderResourceParameter, NiagaraParticleDataFloat);
 	LAYOUT_FIELD(FShaderParameter, FloatDataStride);
 };
 
diff --git a/Engine/Plugins/FX/Niagara/Source/NiagaraVertexFactories/Private/NiagaraSpriteVertexFactory.cpp b/Engine/Plugins/FX/Niagara/Source/NiagaraVertexFactories/Private/NiagaraSpriteVertexFactory.cpp
--- a/Engine/Plugins/FX/Niagara/Source/NiagaraVertexFactories/Private/NiagaraSpriteVertexFactory.cpp
+++ b/Engine/Plugins/FX/Niagara/Source/NiagaraVertexFactories/Private/NiagaraSpriteVertexFactory.cpp
@@ -36,24 +36,7 @@ public:
 		NumCutoutVerticesPerFrame.Bind(ParameterMap, TEXT("NumCutoutVerticesPerFrame"));
 		CutoutGeometry.Bind(ParameterMap, TEXT("CutoutGeometry"));
 
-		NiagaraParticleDataPosition.Bind(ParameterMap, TEXT("NiagaraParticleDataPosition"));
-		NiagaraParticleDataVelocity.Bind(ParameterMap, TEXT("NiagaraParticleDataVelocity"));
-		NiagaraParticleDataColor.Bind(ParameterMap, TEXT("NiagaraParticleDataColor"));
-		NiagaraParticleDataRotation.Bind(ParameterMap, TEXT("NiagaraParticleDataRotation"));
-		NiagaraParticleDataSize.Bind(ParameterMap, TEXT("NiagaraParticleDataSize"));
-		NiagaraParticleDataFacing.Bind(ParameterMap, TEXT("NiagaraParticleDataFacing"));
-		NiagaraParticleDataAlignment.Bind(ParameterMap, TEXT("NiagaraParticleDataAlignment"));
-		NiagaraParticleDataSubImage.Bind(ParameterMap, TEXT("NiagaraParticleDataSubImage"));
-		NiagaraParticleDataCameraOffset.Bind(ParameterMap, TEXT("NiagaraParticleDataCameraOffset"));
-		NiagaraParticleDataUVScale.Bind(ParameterMap, TEXT("NiagaraParticleDataUVScale"));
-		NiagaraParticleDataNormalizedAge.Bind(ParameterMap, TEXT("NiagaraParticleDataNormalizedAge"));
-		NiagaraParticleDataMaterialRandom.Bind(ParameterMap, TEXT("NiagaraParticleDataMaterialRandom"));
-		NiagaraParticleDataCustomSorting.Bind(ParameterMap, TEXT("NiagaraParticleDataCustomSorting"));
-		NiagaraParticleDataMaterialParam0.Bind(ParameterMap, TEXT("NiagaraParticleDataMaterialParam0"));
-		NiagaraParticleDataMaterialParam1.Bind(ParameterMap, TEXT("NiagaraParticleDataMaterialParam1"));
-		NiagaraParticleDataMaterialParam2.Bind(ParameterMap, TEXT("NiagaraParticleDataMaterialParam2"));
-		NiagaraParticleDataMaterialParam3.Bind(ParameterMap, TEXT("NiagaraParticleDataMaterialParam3"));
-
+		NiagaraParticleDataFloat.Bind(ParameterMap, TEXT("NiagaraParticleDataFloat"));
 		FloatDataStride.Bind(ParameterMap, TEXT("NiagaraFloatDataStride"));
 
 //  		NiagaraParticleDataInt.Bind(ParameterMap, TEXT("NiagaraParticleDataInt"));
@@ -90,23 +73,7 @@ public:
 		ShaderBindings.Add(ParticleAlignmentMode, SpriteVF->GetAlignmentMode());
 		ShaderBindings.Add(ParticleFacingMode, SpriteVF->GetFacingMode());
 
-		ShaderBindings.Add(NiagaraParticleDataPosition, SpriteVF->GetParticleDataFloatSRV());
-		ShaderBindings.Add(NiagaraParticleDataVelocity, SpriteVF->GetParticleDataFloatSRV());
-		ShaderBindings.Add(NiagaraParticleDataColor, SpriteVF->GetParticleDataFloatSRV());
-		ShaderBindings.Add(NiagaraParticleDataRotation, SpriteVF->GetParticleDataFloatSRV());
-		ShaderBindings.Add(NiagaraParticleDataSize, SpriteVF->GetParticleDataFloatSRV());
-		ShaderBindings.Add(NiagaraParticleDataFacing, SpriteVF->GetParticleDataFloatSRV());
-		ShaderBindings.Add(NiagaraParticleDataAlignment, SpriteVF->GetParticleDataFloatSRV());
-		ShaderBindings.Add(NiagaraParticleDataSubImage, SpriteVF->GetParticleDataFloatSRV());
-		ShaderBindings.Add(NiagaraParticleDataCameraOffset, SpriteVF->GetParticleDataFloatSRV());
-		ShaderBindings.Add(NiagaraParticleDataUVScale, SpriteVF->GetParticleDataFloatSRV());
-		ShaderBindings.Add(NiagaraParticleDataNormalizedAge, SpriteVF->GetParticleDataFloatSRV());
-		ShaderBindings.Add(NiagaraParticleDataMaterialRandom, SpriteVF->GetParticleDataFloatSRV());
-		ShaderBindings.Add(NiagaraParticleDataCustomSorting, SpriteVF->GetParticleDataFloatSRV());
-		ShaderBindings.Add(NiagaraParticleDataMaterialParam0, SpriteVF->GetParticleDataFloatSRV());
-		ShaderBindings.Add(NiagaraParticleDataMaterialParam1, SpriteVF->GetParticleDataFloatSRV());
-		ShaderBindings.Add(NiagaraParticleDataMaterialParam2, SpriteVF->GetParticleDataFloatSRV());
-		ShaderBindings.Add(NiagaraParticleDataMaterialParam3, SpriteVF->GetParticleDataFloatSRV());
+		ShaderBindings.Add(NiagaraParticleDataFloat, SpriteVF->GetParticleDataFloatSRV());
 		ShaderBindings.Add(FloatDataStride, SpriteVF->GetFloatDataStride());
 
 		ShaderBindings.Add(SortedIndices, SpriteVF->GetSortedIndicesSRV() ? SpriteVF->GetSortedIndicesSRV() : GFNiagaraNullSortedIndicesVertexBuffer.VertexBufferSRV.GetReference());
@@ -122,24 +89,8 @@ private:
 
 		LAYOUT_FIELD(FShaderResourceParameter, CutoutGeometry);
 
-		LAYOUT_FIELD(FShaderResourceParameter, NiagaraParticleDataPosition);
-		LAYOUT_FIELD(FShaderResourceParameter, NiagaraParticleDataVelocity);
-		LAYOUT_FIELD(FShaderResourceParameter, NiagaraParticleDataColor);
-		LAYOUT_FIELD(FShaderResourceParameter, NiagaraParticleDataRotation);
-		LAYOUT_FIELD(FShaderResourceParameter, NiagaraParticleDataSize);
-		LAYOUT_FIELD(FShaderResourceParameter, NiagaraParticleDataFacing);
-		LAYOUT_FIELD(FShaderResourceParameter, NiagaraParticleDataAlignment);
-		LAYOUT_FIELD(FShaderResourceParameter, NiagaraParticleDataSubImage);
-		LAYOUT_FIELD(FShaderResourceParameter, NiagaraParticleDataCameraOffset);
-		LAYOUT_FIELD(FShaderResourceParameter, NiagaraParticleDataUVScale);
-		LAYOUT_FIELD(FShaderResourceParameter, NiagaraParticleDataNormalizedAge);
-		LAYOUT_FIELD(FShaderResourceParameter, NiagaraParticleDataMaterialRandom);
-		LAYOUT_FIELD(FShaderResourceParameter, NiagaraParticleDataCustomSorting);
-		LAYOUT_FIELD(FShaderResourceParameter, NiagaraParticleDataMaterialParam0);
-		LAYOUT_FIELD(FShaderResourceParameter, NiagaraParticleDataMaterialParam1);
-		LAYOUT_FIELD(FShaderResourceParameter, NiagaraParticleDataMaterialParam2);
-		LAYOUT_FIELD(FShaderResourceParameter, NiagaraParticleDataMaterialParam3);
-		LAYOUT_FIELD(FShaderParameter, FloatDataStride);
+	LAYOUT_FIELD(FShaderResourceParameter, NiagaraParticleDataFloat);
+	LAYOUT_FIELD(FShaderParameter, FloatDataStride);
 
 	//  	LAYOUT_FIELD(FShaderResourceParameter, NiagaraParticleDataInt);
 	//  	LAYOUT_FIELD(FShaderParameter, Int32DataOffset);
diff --git a/Engine/Plugins/FX/Niagara/Source/NiagaraVertexFactories/Public/NiagaraRibbonVertexFactory.h b/Engine/Plugins/FX/Niagara/Source/NiagaraVertexFactories/Public/NiagaraRibbonVertexFactory.h
--- a/Engine/Plugins/FX/Niagara/Source/NiagaraVertexFactories/Public/NiagaraRibbonVertexFactory.h
+++ b/Engine/Plugins/FX/Niagara/Source/NiagaraVertexFactories/Public/NiagaraRibbonVertexFactory.h
@@ -61,18 +61,7 @@ BEGIN_GLOBAL_SHADER_PARAMETER_STRUCT(FNiagaraRibbonVFLooseParameters, NIAGARAVER
 	SHADER_PARAMETER_SRV(Buffer<float4>, TangentsAndDistances)
 	SHADER_PARAMETER_SRV(Buffer<uint>, MultiRibbonIndices)
 	SHADER_PARAMETER_SRV(Buffer<float>, PackedPerRibbonDataByIndex)
-	SHADER_PARAMETER_SRV(Buffer<float>, NiagaraParticleDataPosition)
-	SHADER_PARAMETER_SRV(Buffer<float>, NiagaraParticleDataVelocity)
-	SHADER_PARAMETER_SRV(Buffer<float>, NiagaraParticleDataColor)
-	SHADER_PARAMETER_SRV(Buffer<float>, NiagaraParticleDataWidth)
-	SHADER_PARAMETER_SRV(Buffer<float>, NiagaraParticleDataTwist)
-	SHADER_PARAMETER_SRV(Buffer<float>, NiagaraParticleDataFacing)
-	SHADER_PARAMETER_SRV(Buffer<float>, NiagaraParticleDataNormalizedAge)
-	SHADER_PARAMETER_SRV(Buffer<float>, NiagaraParticleDataMaterialRandom)
-	SHADER_PARAMETER_SRV(Buffer<float>, NiagaraParticleDataMaterialParam0)
-	SHADER_PARAMETER_SRV(Buffer<float>, NiagaraParticleDataMaterialParam1)
-	SHADER_PARAMETER_SRV(Buffer<float>, NiagaraParticleDataMaterialParam2)
-	SHADER_PARAMETER_SRV(Buffer<float>, NiagaraParticleDataMaterialParam3)
+	SHADER_PARAMETER_SRV(Buffer<float>, NiagaraParticleDataFloat)
 END_GLOBAL_SHADER_PARAMETER_STRUCT()
 typedef TUniformBufferRef<FNiagaraRibbonVFLooseParameters> FNiagaraRibbonVFLooseParametersRef;
 
diff --git a/Engine/Plugins/FX/Niagara/Source/NiagaraVertexFactories/Public/NiagaraSpriteVertexFactory.h b/Engine/Plugins/FX/Niagara/Source/NiagaraVertexFactories/Public/NiagaraSpriteVertexFactory.h
--- a/Engine/Plugins/FX/Niagara/Source/NiagaraVertexFactories/Public/NiagaraSpriteVertexFactory.h
+++ b/Engine/Plugins/FX/Niagara/Source/NiagaraVertexFactories/Public/NiagaraSpriteVertexFactory.h
@@ -63,23 +63,7 @@ BEGIN_GLOBAL_SHADER_PARAMETER_STRUCT(FNiagaraSpriteVFLooseParameters, NIAGARAVER
 	SHADER_PARAMETER(uint32, SortedIndicesOffset)
 	SHADER_PARAMETER(uint32, IndirectArgsOffset)
 	SHADER_PARAMETER_SRV(Buffer<float2>, CutoutGeometry)
-	SHADER_PARAMETER_SRV(Buffer<float>, NiagaraParticleDataPosition)
-	SHADER_PARAMETER_SRV(Buffer<float>, NiagaraParticleDataVelocity)
-	SHADER_PARAMETER_SRV(Buffer<float>, NiagaraParticleDataColor)
-	SHADER_PARAMETER_SRV(Buffer<float>, NiagaraParticleDataRotation)
-	SHADER_PARAMETER_SRV(Buffer<float>, NiagaraParticleDataSize)
-	SHADER_PARAMETER_SRV(Buffer<float>, NiagaraParticleDataFacing)
-	SHADER_PARAMETER_SRV(Buffer<float>, NiagaraParticleDataAlignment)
-	SHADER_PARAMETER_SRV(Buffer<float>, NiagaraParticleDataSubImage)
-	SHADER_PARAMETER_SRV(Buffer<float>, NiagaraParticleDataCameraOffset)
-	SHADER_PARAMETER_SRV(Buffer<float>, NiagaraParticleDataUVScale)
-	SHADER_PARAMETER_SRV(Buffer<float>, NiagaraParticleDataNormalizedAge)
-	SHADER_PARAMETER_SRV(Buffer<float>, NiagaraParticleDataMaterialRandom)
-	SHADER_PARAMETER_SRV(Buffer<float>, NiagaraParticleDataCustomSorting)
-	SHADER_PARAMETER_SRV(Buffer<float>, NiagaraParticleDataMaterialParam0)
-	SHADER_PARAMETER_SRV(Buffer<float>, NiagaraParticleDataMaterialParam1)
-	SHADER_PARAMETER_SRV(Buffer<float>, NiagaraParticleDataMaterialParam2)
-	SHADER_PARAMETER_SRV(Buffer<float>, NiagaraParticleDataMaterialParam3)
+	SHADER_PARAMETER_SRV(Buffer<float>, NiagaraParticleDataFloat)
 	SHADER_PARAMETER_SRV(Buffer<int>, SortedIndices)
 	SHADER_PARAMETER_SRV(Buffer<uint>, IndirectArgsBuffer)
 END_GLOBAL_SHADER_PARAMETER_STRUCT()
