Merging 11604247 - Niagara - Several fixes for scratchpad UX.
+ Selection refactor.
+ Pinning support.
+ Proper apply support.
+ Cut/Copy/Paste support.
+ Fix trailing numbers on script names when creating new scripts and renaming scripts.
+ General UI engine consistency.

diff --git a/Engine/Plugins/FX/Niagara/Source/Niagara/Classes/NiagaraScript.h b/Engine/Plugins/FX/Niagara/Source/Niagara/Classes/NiagaraScript.h
--- a/Engine/Plugins/FX/Niagara/Source/Niagara/Classes/NiagaraScript.h
+++ b/Engine/Plugins/FX/Niagara/Source/Niagara/Classes/NiagaraScript.h
@@ -386,6 +386,7 @@ public:
 
 #if WITH_EDITOR
 	DECLARE_MULTICAST_DELEGATE_OneParam(FOnScriptCompiled, UNiagaraScript*);
+	DECLARE_MULTICAST_DELEGATE_OneParam(FOnPropertyChanged, FPropertyChangedEvent& /* PropertyChangedEvent */)
 #endif
 
 private:
@@ -644,6 +645,9 @@ public:
 	/** Callback issued whenever a GPU script compilation successfully happened (even if the results are a script that cannot be executed due to errors)*/
 	NIAGARA_API FOnScriptCompiled& OnGPUScriptCompiled();
 
+	/** Callback issues whenever post edit changed is called on this script. */
+	NIAGARA_API FOnPropertyChanged& OnPropertyChanged();
+
 	/** External call used to identify the values for a successful VM script compilation. OnVMScriptCompiled will be issued in this case.*/
 	void SetVMCompilationResults(const FNiagaraVMExecutableDataId& InCompileId, FNiagaraVMExecutableData& InScriptVM, FNiagaraCompileRequestDataBase* InRequestData);
 
@@ -715,6 +719,7 @@ private:
 	/** A multicast delegate which is called whenever the script has been compiled (successfully or not). */
 	FOnScriptCompiled OnVMScriptCompiledDelegate;
 	FOnScriptCompiled OnGPUScriptCompiledDelegate;
+	FOnPropertyChanged OnPropertyChangedDelegate;
 
 	mutable FNiagaraVMExecutableDataId LastReportedVMId;
 
diff --git a/Engine/Plugins/FX/Niagara/Source/Niagara/Private/NiagaraScript.cpp b/Engine/Plugins/FX/Niagara/Source/Niagara/Private/NiagaraScript.cpp
--- a/Engine/Plugins/FX/Niagara/Source/Niagara/Private/NiagaraScript.cpp
+++ b/Engine/Plugins/FX/Niagara/Source/Niagara/Private/NiagaraScript.cpp
@@ -1057,6 +1057,7 @@ void UNiagaraScript::PostEditChangeProperty(FPropertyChangedEvent& PropertyChang
 	}
 
 	CustomAssetRegistryTagCache.Reset();
+	OnPropertyChangedDelegate.Broadcast(PropertyChangedEvent);
 }
 
 #endif
@@ -1923,7 +1924,10 @@ UNiagaraScript::FOnScriptCompiled& UNiagaraScript::OnGPUScriptCompiled()
 	return OnGPUScriptCompiledDelegate;
 }
 
-
+UNiagaraScript::FOnPropertyChanged& UNiagaraScript::OnPropertyChanged()
+{
+	return OnPropertyChangedDelegate;
+}
 
 #endif
 
diff --git a/Engine/Plugins/FX/Niagara/Source/NiagaraEditor/Private/NiagaraEditorCommands.cpp b/Engine/Plugins/FX/Niagara/Source/NiagaraEditor/Private/NiagaraEditorCommands.cpp
--- a/Engine/Plugins/FX/Niagara/Source/NiagaraEditor/Private/NiagaraEditorCommands.cpp
+++ b/Engine/Plugins/FX/Niagara/Source/NiagaraEditor/Private/NiagaraEditorCommands.cpp
@@ -8,7 +8,7 @@
 
 void FNiagaraEditorCommands::RegisterCommands()
 {
-	UI_COMMAND(Apply, "Apply", "Push the currently compiled script to the world.", EUserInterfaceActionType::Button, FInputChord());
+	UI_COMMAND(Apply, "Apply", "Apply unsaved changes to the current asset.", EUserInterfaceActionType::Button, FInputChord());
 	UI_COMMAND(Compile, "Compile", "Compile the current scripts", EUserInterfaceActionType::Button, FInputChord());
 	UI_COMMAND(RefreshNodes, "Refresh", "Refreshes the current graph nodes, and updates pins due to external changes.", EUserInterfaceActionType::Button, FInputChord());
 	UI_COMMAND(ResetSimulation, "Reset", "Resets the current simulation", EUserInterfaceActionType::Button, FInputChord(EKeys::SpaceBar));
diff --git a/Engine/Plugins/FX/Niagara/Source/NiagaraEditor/Private/NiagaraEditorUtilities.cpp b/Engine/Plugins/FX/Niagara/Source/NiagaraEditor/Private/NiagaraEditorUtilities.cpp
--- a/Engine/Plugins/FX/Niagara/Source/NiagaraEditor/Private/NiagaraEditorUtilities.cpp
+++ b/Engine/Plugins/FX/Niagara/Source/NiagaraEditor/Private/NiagaraEditorUtilities.cpp
@@ -2033,4 +2033,36 @@ void FNiagaraEditorUtilities::WarnWithToastAndLog(FText WarningMessage)
 	UE_LOG(LogNiagaraEditor, Warning, TEXT("%s"), *WarningMessage.ToString());
 }
 
+FName FNiagaraEditorUtilities::GetUniqueObjectName(UObject* Outer, UClass* ObjectClass, const FString& CandidateName)
+{
+	if (StaticFindObject(ObjectClass, Outer, *CandidateName) == nullptr)
+	{
+		return *CandidateName;
+	}
+	else
+	{
+		FString BaseCandidateName;
+		int32 LastUnderscoreIndex;
+		int32 NameIndex = 0;
+		if (CandidateName.FindLastChar('_', LastUnderscoreIndex) && LexTryParseString(NameIndex, *CandidateName.RightChop(LastUnderscoreIndex + 1)))
+		{
+			BaseCandidateName = CandidateName.Left(LastUnderscoreIndex);
+			NameIndex++;
+		}
+		else
+		{
+			BaseCandidateName = CandidateName;
+			NameIndex = 1;
+		}
+
+		FString UniqueCandidateName = FString::Printf(TEXT("%s_%02i"), *BaseCandidateName, NameIndex);
+		while (StaticFindObject(ObjectClass, Outer, *UniqueCandidateName) != nullptr)
+		{
+			NameIndex++;
+			UniqueCandidateName = FString::Printf(TEXT("%s_%02i"), *BaseCandidateName, NameIndex);
+		}
+		return *UniqueCandidateName;
+	}
+}
+
 #undef LOCTEXT_NAMESPACE
diff --git a/Engine/Plugins/FX/Niagara/Source/NiagaraEditor/Private/NiagaraNodeFunctionCall.cpp b/Engine/Plugins/FX/Niagara/Source/NiagaraEditor/Private/NiagaraNodeFunctionCall.cpp
--- a/Engine/Plugins/FX/Niagara/Source/NiagaraEditor/Private/NiagaraNodeFunctionCall.cpp
+++ b/Engine/Plugins/FX/Niagara/Source/NiagaraEditor/Private/NiagaraNodeFunctionCall.cpp
@@ -756,7 +756,7 @@ void UNiagaraNodeFunctionCall::SubsumeExternalDependencies(TMap<const UObject*,
 
 void UNiagaraNodeFunctionCall::GatherExternalDependencyData(ENiagaraScriptUsage InMasterUsage, const FGuid& InMasterUsageId, TArray<FNiagaraCompileHash>& InReferencedCompileHashes, TArray<FString>& InReferencedObjs) const
 {
-	if (FunctionScript && FunctionScript->GetOutermost() != this->GetOutermost())
+	if (FunctionScript)
 	{
 		UNiagaraScriptSource* Source = CastChecked<UNiagaraScriptSource>(FunctionScript->GetSource());
 		UNiagaraGraph* FunctionGraph = CastChecked<UNiagaraGraph>(Source->NodeGraph);
diff --git a/Engine/Plugins/FX/Niagara/Source/NiagaraEditor/Private/Toolkits/NiagaraSystemToolkit.cpp b/Engine/Plugins/FX/Niagara/Source/NiagaraEditor/Private/Toolkits/NiagaraSystemToolkit.cpp
--- a/Engine/Plugins/FX/Niagara/Source/NiagaraEditor/Private/Toolkits/NiagaraSystemToolkit.cpp
+++ b/Engine/Plugins/FX/Niagara/Source/NiagaraEditor/Private/Toolkits/NiagaraSystemToolkit.cpp
@@ -9,6 +9,8 @@
 #include "ViewModels/NiagaraEmitterHandleViewModel.h"
 #include "ViewModels/NiagaraEmitterViewModel.h"
 #include "ViewModels/NiagaraSystemSelectionViewModel.h"
+#include "ViewModels/NiagaraScratchPadScriptViewModel.h"
+#include "ViewModels/NiagaraScratchPadViewModel.h"
 #include "NiagaraSystemScriptViewModel.h"
 #include "Widgets/SNiagaraCurveEditor.h"
 #include "Widgets/SNiagaraSystemScript.h"
@@ -415,6 +417,7 @@ void FNiagaraSystemToolkit::InitializeInternal(const EToolkitMode::Type Mode, co
 	RegenerateMenusAndToolbars();
 
 	bChangesDiscarded = false;
+	bScratchPadChangesDiscarded = false;
 }
 
 FName FNiagaraSystemToolkit::GetToolkitFName() const
@@ -1367,12 +1370,45 @@ bool FNiagaraSystemToolkit::OnRequestClose()
 
 	SystemViewModel->NotifyPreClose();
 
+	bool bHasUnappliedScratchPadChanges = false;
+	for (TSharedRef<FNiagaraScratchPadScriptViewModel> ScratchPadViewModel : SystemViewModel->GetScriptScratchPadViewModel()->GetScriptViewModels())
+	{
+		if (ScratchPadViewModel->CanApplyChanges())
+		{
+			bHasUnappliedScratchPadChanges = true;
+			break;
+		}
+	}
+
+	if (bScratchPadChangesDiscarded == false && bHasUnappliedScratchPadChanges)
+	{
+		// find out the user wants to do with their dirty scratch pad scripts.
+		EAppReturnType::Type YesNoCancelReply = FMessageDialog::Open(EAppMsgType::YesNoCancel,
+			NSLOCTEXT("NiagaraEditor", "UnsavedScratchPadScriptsPrompt", "Would you like to apply changes to scratch pad scripts? (No will discard unapplied changes)"));
+
+		switch (YesNoCancelReply)
+		{
+		case EAppReturnType::Yes:
+			for (TSharedRef<FNiagaraScratchPadScriptViewModel> ScratchPadViewModel : SystemViewModel->GetScriptScratchPadViewModel()->GetScriptViewModels())
+			{
+				ScratchPadViewModel->ApplyChanges();
+			}
+			break;
+		case EAppReturnType::No:
+			bScratchPadChangesDiscarded = true;
+			break;
+		case EAppReturnType::Cancel:
+			return false;
+			break;
+		}
+	}
+
 	if (SystemToolkitMode == ESystemToolkitMode::Emitter)
 	{
 		TSharedPtr<FNiagaraEmitterViewModel> EmitterViewModel = SystemViewModel->GetEmitterHandleViewModels()[0]->GetEmitterViewModel();
 		if (bChangesDiscarded == false && (EmitterViewModel->GetEmitter()->GetChangeId() != LastSyncedEmitterChangeId || bEmitterThumbnailUpdated))
 		{
-			// find out the user wants to do with this dirty NiagaraScript
+			// find out the user wants to do with this dirty emitter.
 			EAppReturnType::Type YesNoCancelReply = FMessageDialog::Open(EAppMsgType::YesNoCancel,
 				FText::Format(
 					NSLOCTEXT("UnrealEd", "Prompt_NiagaraEmitterEditorClose", "Would you like to apply changes to this Emitter to the original Emitter?\n{0}\n(No will lose all changes!)"),
@@ -1391,6 +1427,7 @@ bool FNiagaraSystemToolkit::OnRequestClose()
 				break;
 			case EAppReturnType::Cancel:
 				// don't exit
+				bScratchPadChangesDiscarded = false;
 				return false;
 			}
 		}
diff --git a/Engine/Plugins/FX/Niagara/Source/NiagaraEditor/Private/Toolkits/NiagaraSystemToolkit.h b/Engine/Plugins/FX/Niagara/Source/NiagaraEditor/Private/Toolkits/NiagaraSystemToolkit.h
--- a/Engine/Plugins/FX/Niagara/Source/NiagaraEditor/Private/Toolkits/NiagaraSystemToolkit.h
+++ b/Engine/Plugins/FX/Niagara/Source/NiagaraEditor/Private/Toolkits/NiagaraSystemToolkit.h
@@ -182,6 +182,8 @@ private:
 
 	bool bChangesDiscarded;
 
+	bool bScratchPadChangesDiscarded;
+
 public:
 	static const FName ViewportTabID;
 	static const FName CurveEditorTabID;
diff --git a/Engine/Plugins/FX/Niagara/Source/NiagaraEditor/Private/ViewModels/NiagaraScratchPadScriptViewModel.cpp b/Engine/Plugins/FX/Niagara/Source/NiagaraEditor/Private/ViewModels/NiagaraScratchPadScriptViewModel.cpp
--- a/Engine/Plugins/FX/Niagara/Source/NiagaraEditor/Private/ViewModels/NiagaraScratchPadScriptViewModel.cpp
+++ b/Engine/Plugins/FX/Niagara/Source/NiagaraEditor/Private/ViewModels/NiagaraScratchPadScriptViewModel.cpp
@@ -2,26 +2,70 @@
 
 #include "ViewModels/NiagaraScratchPadScriptViewModel.h"
 #include "ViewModels/Stack/NiagaraStackGraphUtilities.h"
+#include "ViewModels/Stack/NiagaraStackFunctionInputCollection.h"
 #include "NiagaraNodeFunctionCall.h"
 #include "NiagaraSystem.h"
 #include "NiagaraEmitter.h"
+#include "NiagaraEditorUtilities.h"
 
 #include "ObjectTools.h"
+#include "ScopedTransaction.h"
+
+#define LOCTEXT_NAMESPACE "NiagaraScratchPadScriptViewModel"
 
 FNiagaraScratchPadScriptViewModel::FNiagaraScratchPadScriptViewModel()
 	: FNiagaraScriptViewModel(TAttribute<FText>::Create(TAttribute<FText>::FGetter::CreateRaw(this, &FNiagaraScratchPadScriptViewModel::GetDisplayNameInternal)), ENiagaraParameterEditMode::EditAll)
 	, bIsPendingRename(false)
+	, bIsPinned(false)
+	, EditorHeight(300)
+	, bHasPendingChanges(false)
 {
 }
 
+FNiagaraScratchPadScriptViewModel::~FNiagaraScratchPadScriptViewModel()
+{
+	if (EditScript != nullptr)
+	{
+		if (EditScript->GetSource() != nullptr)
+		{
+			UNiagaraScriptSource* EditScriptSource = CastChecked<UNiagaraScriptSource>(EditScript->GetSource());
+			EditScriptSource->NodeGraph->RemoveOnGraphNeedsRecompileHandler(OnGraphNeedsRecompileHandle);
+		}
+		EditScript->OnPropertyChanged().RemoveAll(this);
+		EditScript = nullptr;
+	}
+}
+
 void FNiagaraScratchPadScriptViewModel::Initialize(UNiagaraScript* Script)
 {
-	SetScript(Script);
+	OriginalScript = Script;
+	EditScript = CastChecked<UNiagaraScript>(StaticDuplicateObject(Script, GetTransientPackage()));
+	SetScript(EditScript);
+	UNiagaraScriptSource* EditScriptSource = CastChecked<UNiagaraScriptSource>(EditScript->GetSource());
+	OnGraphNeedsRecompileHandle = EditScriptSource->NodeGraph->AddOnGraphNeedsRecompileHandler(FOnGraphChanged::FDelegate::CreateSP(this, &FNiagaraScratchPadScriptViewModel::OnScriptGraphChanged));
+	EditScript->OnPropertyChanged().AddSP(this, &FNiagaraScratchPadScriptViewModel::OnScriptPropertyChanged);
+}
+
+void FNiagaraScratchPadScriptViewModel::AddReferencedObjects(FReferenceCollector& Collector)
+{
+	Collector.AddReferencedObject(EditScript);
+}
+
+UNiagaraScript* FNiagaraScratchPadScriptViewModel::GetOriginalScript() const
+{
+	return OriginalScript;
 }
 
-UNiagaraScript* FNiagaraScratchPadScriptViewModel::GetScript() const
+UNiagaraScript* FNiagaraScratchPadScriptViewModel::GetEditScript() const
 {
-	return GetScripts()[0].Get();
+	return EditScript;
+}
+
+FText FNiagaraScratchPadScriptViewModel::GetToolTip() const
+{
+	return FText::Format(LOCTEXT("ScratchPadScriptToolTipFormat", "Description: {0}{1}"),
+		EditScript->Description.IsEmptyOrWhitespace() ? LOCTEXT("NoDescription", "(none)") : EditScript->Description,
+		bHasPendingChanges ? LOCTEXT("HasPendingChangesStatus", "\n* Has pending changes to apply") : FText());
 }
 
 bool FNiagaraScratchPadScriptViewModel::GetIsPendingRename() const
@@ -34,29 +78,42 @@ void FNiagaraScratchPadScriptViewModel::SetIsPendingRename(bool bInIsPendingRena
 	bIsPendingRename = bInIsPendingRename;
 }
 
+void GetReferencingFunctionCallNodes(UNiagaraScript* Script, TArray<UNiagaraNodeFunctionCall*>& OutReferencingFunctionCallNodes)
+{
+	for (TObjectIterator<UNiagaraNodeFunctionCall> It; It; ++It)
+	{
+		UNiagaraNodeFunctionCall* FunctionCallNode = *It;
+		if (FunctionCallNode->FunctionScript == Script)
+		{
+			OutReferencingFunctionCallNodes.Add(FunctionCallNode);
+		}
+	}
+}
+
 void FNiagaraScratchPadScriptViewModel::SetScriptName(FText InScriptName)
 {
-	UNiagaraScript* Script = GetScript();
 	FString NewName = ObjectTools::SanitizeObjectName(InScriptName.ToString());
-	if (Script->GetName() != NewName)
+	if (OriginalScript->GetName() != NewName)
 	{
-		FName NewUniqueName = MakeUniqueObjectName(Script->GetOuter(), UNiagaraScript::StaticClass(), *NewName);
-		Script->Rename(*NewUniqueName.ToString());
+		FScopedTransaction RenameTransaction(LOCTEXT("RenameScriptTransaction", "Rename scratch pad script."));
+
+		FName NewUniqueName = FNiagaraEditorUtilities::GetUniqueObjectName<UNiagaraScript>(OriginalScript->GetOuter(), *NewName);
+		OriginalScript->Modify();
+		OriginalScript->Rename(*NewUniqueName.ToString(), nullptr, REN_DontCreateRedirectors);
 
-		for (TObjectIterator<UNiagaraNodeFunctionCall> It; It; ++It)
+		TArray<UNiagaraNodeFunctionCall*> ReferencingFunctionCallNodes;
+		GetReferencingFunctionCallNodes(OriginalScript, ReferencingFunctionCallNodes);
+		for(UNiagaraNodeFunctionCall* ReferencingFunctionCallNode : ReferencingFunctionCallNodes)
 		{
-			UNiagaraNodeFunctionCall* FunctionCallNode = *It;
-			if (FunctionCallNode->FunctionScript == Script)
+			ReferencingFunctionCallNode->Modify();
+			FString OldFunctionName = ReferencingFunctionCallNode->GetFunctionName();
+			ReferencingFunctionCallNode->SuggestName(FString());
+			const FString NewFunctionName = ReferencingFunctionCallNode->GetFunctionName();
+			UNiagaraSystem* System = ReferencingFunctionCallNode->GetTypedOuter<UNiagaraSystem>();
+			UNiagaraEmitter* Emitter = ReferencingFunctionCallNode->GetTypedOuter<UNiagaraEmitter>();
+			if (System != nullptr)
 			{
-				FString OldFunctionName = FunctionCallNode->GetFunctionName();
-				FunctionCallNode->SuggestName(FString());
-				const FString NewFunctionName = FunctionCallNode->GetFunctionName();
-				UNiagaraSystem* System = FunctionCallNode->GetTypedOuter<UNiagaraSystem>();
-				UNiagaraEmitter* Emitter = FunctionCallNode->GetTypedOuter<UNiagaraEmitter>();
-				if (System != nullptr)
-				{
-					FNiagaraStackGraphUtilities::RenameReferencingParameters(*System, Emitter, *FunctionCallNode, OldFunctionName, NewFunctionName);
-				}
+				FNiagaraStackGraphUtilities::RenameReferencingParameters(*System, Emitter, *ReferencingFunctionCallNode, OldFunctionName, NewFunctionName);
 			}
 		}
 
@@ -64,12 +121,96 @@ void FNiagaraScratchPadScriptViewModel::SetScriptName(FText InScriptName)
 	}
 }
 
+bool FNiagaraScratchPadScriptViewModel::GetIsPinned() const
+{
+	return bIsPinned;
+}
+
+void FNiagaraScratchPadScriptViewModel::SetIsPinned(bool bInIsPinned)
+{
+	if (bIsPinned != bInIsPinned)
+	{
+		bIsPinned = bInIsPinned;
+		OnPinnedChangedDelegate.Broadcast();
+	}
+}
+
+float FNiagaraScratchPadScriptViewModel::GetEditorHeight() const
+{
+	return EditorHeight;
+}
+
+void FNiagaraScratchPadScriptViewModel::SetEditorHeight(float InEditorHeight)
+{
+	EditorHeight = InEditorHeight;
+}
+
+bool FNiagaraScratchPadScriptViewModel::CanApplyChanges() const
+{
+	return bHasPendingChanges;
+}
+
+void FNiagaraScratchPadScriptViewModel::ApplyChanges()
+{
+	ResetLoaders(OriginalScript->GetOutermost()); // Make sure that we're not going to get invalid version number linkers into the package we are going into. 
+	OriginalScript->GetOutermost()->LinkerCustomVersion.Empty();
+
+	OriginalScript = (UNiagaraScript*)StaticDuplicateObject(EditScript, OriginalScript->GetOuter(), OriginalScript->GetFName(),
+		RF_AllFlags,
+		OriginalScript->GetClass());
+	bHasPendingChanges = false;
+
+	TArray<UNiagaraNodeFunctionCall*> FunctionCallNodesToRefresh;
+	GetReferencingFunctionCallNodes(OriginalScript, FunctionCallNodesToRefresh);
+
+	TArray<UNiagaraStackFunctionInputCollection*> InputCollectionsToRefresh;
+	if (FunctionCallNodesToRefresh.Num())
+	{
+		for (TObjectIterator<UNiagaraStackFunctionInputCollection> It; It; ++It)
+		{
+			UNiagaraStackFunctionInputCollection* StackFunctionInputCollection = *It;
+			if (StackFunctionInputCollection->IsFinalized() == false && FunctionCallNodesToRefresh.Contains(StackFunctionInputCollection->GetInputFunctionCallNode()))
+			{
+				InputCollectionsToRefresh.Add(StackFunctionInputCollection);
+			}
+		}
+	}
+
+	for (UNiagaraNodeFunctionCall* FunctionCallNodeToRefresh : FunctionCallNodesToRefresh)
+	{
+		FunctionCallNodeToRefresh->RefreshFromExternalChanges();
+		FunctionCallNodeToRefresh->MarkNodeRequiresSynchronization(TEXT("ScratchPadChangesApplied"), true);
+	}
+
+	for (UNiagaraStackFunctionInputCollection* InputCollectionToRefresh : InputCollectionsToRefresh)
+	{
+		InputCollectionToRefresh->RefreshChildren();
+	}
+}
+
 FNiagaraScratchPadScriptViewModel::FOnRenamed& FNiagaraScratchPadScriptViewModel::OnRenamed()
 {
 	return OnRenamedDelegate;
 }
 
+FNiagaraScratchPadScriptViewModel::FOnPinnedChanged& FNiagaraScratchPadScriptViewModel::OnPinnedChanged()
+{
+	return OnPinnedChangedDelegate;
+}
+
 FText FNiagaraScratchPadScriptViewModel::GetDisplayNameInternal() const
 {
-	return FText::FromString(GetScript()->GetName());
-}
\ No newline at end of file
+	return FText::Format(LOCTEXT("DisplayNameFormat", "{0}{1}"), FText::FromString(OriginalScript->GetName()), bHasPendingChanges ? LOCTEXT("HasPendingChanges", "*") : FText());
+}
+
+void FNiagaraScratchPadScriptViewModel::OnScriptGraphChanged(const FEdGraphEditAction &Action)
+{
+	bHasPendingChanges = true;
+}
+
+void FNiagaraScratchPadScriptViewModel::OnScriptPropertyChanged(FPropertyChangedEvent& PropertyChangedEvent)
+{
+	bHasPendingChanges = true;
+}
+
+#undef LOCTEXT_NAMESPACE
\ No newline at end of file
diff --git a/Engine/Plugins/FX/Niagara/Source/NiagaraEditor/Private/ViewModels/NiagaraScratchPadViewModel.cpp b/Engine/Plugins/FX/Niagara/Source/NiagaraEditor/Private/ViewModels/NiagaraScratchPadViewModel.cpp
--- a/Engine/Plugins/FX/Niagara/Source/NiagaraEditor/Private/ViewModels/NiagaraScratchPadViewModel.cpp
+++ b/Engine/Plugins/FX/Niagara/Source/NiagaraEditor/Private/ViewModels/NiagaraScratchPadViewModel.cpp
@@ -10,6 +10,8 @@
 #include "NiagaraObjectSelection.h"
 #include "NiagaraEditorSettings.h"
 #include "NiagaraNodeFunctionCall.h"
+#include "NiagaraEditorModule.h"
+#include "NiagaraClipboard.h"
 
 #include "ScopedTransaction.h"
 
@@ -20,6 +22,10 @@ void UNiagaraScratchPadViewModel::Initialize(TSharedRef<FNiagaraSystemViewModel>
 	SystemViewModelWeak = InSystemViewModel;
 	ObjectSelection = MakeShared<FNiagaraObjectSelection>();
 	RefreshScriptViewModels();
+	if (ScriptViewModels.Num() > 0)
+	{
+		SetActiveScriptViewModel(ScriptViewModels[0]);
+	}
 	AvailableUsages = { ENiagaraScriptUsage::DynamicInput, ENiagaraScriptUsage::Module };
 }
 
@@ -55,7 +61,7 @@ void UNiagaraScratchPadViewModel::RefreshScriptViewModels()
 	TArray<TSharedRef<FNiagaraScratchPadScriptViewModel>> OldScriptViewModels = ScriptViewModels;
 	ScriptViewModels.Empty();
 
-	bool bNewViewModels = false;
+	bool bViewModelsChanged = false;
 
 	UObject* ScriptOuter;
 	TArray<UNiagaraScript*>* TargetScripts;
@@ -68,7 +74,7 @@ void UNiagaraScratchPadViewModel::RefreshScriptViewModels()
 			TSharedPtr<FNiagaraScratchPadScriptViewModel> ScriptViewModel;
 
 			TSharedRef<FNiagaraScratchPadScriptViewModel>* OldScriptViewModel = OldScriptViewModels.FindByPredicate(
-				[ScratchPadScript](TSharedRef<FNiagaraScratchPadScriptViewModel> ScriptViewModel) { return ScriptViewModel->GetScript() == ScratchPadScript; });
+				[ScratchPadScript](TSharedRef<FNiagaraScratchPadScriptViewModel> ScriptViewModel) { return ScriptViewModel->GetOriginalScript() == ScratchPadScript; });
 			if (OldScriptViewModel != nullptr)
 			{
 				ScriptViewModel = *OldScriptViewModel;
@@ -80,18 +86,51 @@ void UNiagaraScratchPadViewModel::RefreshScriptViewModels()
 				ScriptViewModel->Initialize(ScratchPadScript);
 				ScriptViewModel->GetGraphViewModel()->GetNodeSelection()->OnSelectedObjectsChanged().AddUObject(this, &UNiagaraScratchPadViewModel::ScriptGraphNodeSelectionChanged, TWeakPtr<FNiagaraScratchPadScriptViewModel>(ScriptViewModel));
 				ScriptViewModel->OnRenamed().AddUObject(this, &UNiagaraScratchPadViewModel::ScriptViewModelScriptRenamed);
-				bNewViewModels = true;
+				ScriptViewModel->OnPinnedChanged().AddUObject(this, &UNiagaraScratchPadViewModel::ScriptViewModelPinnedChanged, TWeakPtr<FNiagaraScratchPadScriptViewModel>(ScriptViewModel));
+				bViewModelsChanged = true;
 			}
 			ScriptViewModels.Add(ScriptViewModel.ToSharedRef());
 		}
 	}
 
-	for (TSharedRef<FNiagaraScratchPadScriptViewModel> OldScriptViewModel : OldScriptViewModels)
+	if (OldScriptViewModels.Num() > 0)
+	{
+		for (TSharedRef<FNiagaraScratchPadScriptViewModel> OldScriptViewModel : OldScriptViewModels)
+		{
+			OldScriptViewModel->OnRenamed().RemoveAll(this);
+			OldScriptViewModel->OnPinnedChanged().RemoveAll(this);
+		}
+		bViewModelsChanged = true;
+	}
+
+	bool bEditViewModelsChanged = false;
+	if (ActiveScriptViewModel.IsValid() && ScriptViewModels.Contains(ActiveScriptViewModel.ToSharedRef()) == false)
 	{
-		OldScriptViewModel->OnRenamed().RemoveAll(this);
+		ActiveScriptViewModel.Reset();
+		bEditViewModelsChanged = true;
 	}
 
-	if (bNewViewModels || OldScriptViewModels.Num() > 0)
+	TArray<TSharedRef<FNiagaraScratchPadScriptViewModel>> OldPinnedScriptViewModels = PinnedScriptViewModels;
+	PinnedScriptViewModels.Empty();
+	for (TSharedRef<FNiagaraScratchPadScriptViewModel> OldPinnedScriptViewModel : OldPinnedScriptViewModels)
+	{
+		// Remove pinned view models which are no longer valid, but add them one at a time to maintain the pin order.
+		if (ScriptViewModels.Contains(OldPinnedScriptViewModel))
+		{
+			PinnedScriptViewModels.Add(OldPinnedScriptViewModel);
+		}
+		else
+		{
+			bEditViewModelsChanged = true;
+		}
+	}
+
+	if (bEditViewModelsChanged)
+	{
+		RefreshEditScriptViewModels();
+	}
+	
+	if (bViewModelsChanged)
 	{
 		OnScriptViewModelsChangedDelegate.Broadcast();
 	}
@@ -102,9 +141,14 @@ const TArray<TSharedRef<FNiagaraScratchPadScriptViewModel>>& UNiagaraScratchPadV
 	return ScriptViewModels;
 }
 
+const TArray<TSharedRef<FNiagaraScratchPadScriptViewModel>>& UNiagaraScratchPadViewModel::GetEditScriptViewModels() const
+{
+	return EditScriptViewModels;
+}
+
 TSharedPtr<FNiagaraScratchPadScriptViewModel> UNiagaraScratchPadViewModel::GetViewModelForScript(UNiagaraScript* InScript)
 {
-	TSharedRef<FNiagaraScratchPadScriptViewModel>* ViewModelForScript = ScriptViewModels.FindByPredicate([InScript](TSharedRef<FNiagaraScratchPadScriptViewModel>& ScriptViewModel) { return ScriptViewModel->GetScript() == InScript; });
+	TSharedRef<FNiagaraScratchPadScriptViewModel>* ViewModelForScript = ScriptViewModels.FindByPredicate([InScript](TSharedRef<FNiagaraScratchPadScriptViewModel>& ScriptViewModel) { return ScriptViewModel->GetOriginalScript() == InScript; });
 	if (ViewModelForScript != nullptr)
 	{
 		return *ViewModelForScript;
@@ -137,29 +181,65 @@ TSharedRef<FNiagaraObjectSelection> UNiagaraScratchPadViewModel::GetObjectSelect
 	return ObjectSelection.ToSharedRef();
 }
 
-UNiagaraScript* UNiagaraScratchPadViewModel::GetActiveScript()
+TSharedPtr<FNiagaraScratchPadScriptViewModel> UNiagaraScratchPadViewModel::GetActiveScriptViewModel()
 {
-	return ActiveScript;
+	return ActiveScriptViewModel;
 }
 
-void UNiagaraScratchPadViewModel::SetActiveScript(UNiagaraScript* InActiveScript)
+void UNiagaraScratchPadViewModel::SetActiveScriptViewModel(TSharedRef<FNiagaraScratchPadScriptViewModel> InActiveScriptViewModel )
 {
-	ActiveScript = InActiveScript;
-	if (ActiveScript != nullptr)
+	if (ensureMsgf(ScriptViewModels.Contains(InActiveScriptViewModel), TEXT("Can only set an active view model from this scratch pad view model.")))
 	{
-		ObjectSelection->SetSelectedObject(ActiveScript);
+		ActiveScriptViewModel = InActiveScriptViewModel;
+		ObjectSelection->SetSelectedObject(ActiveScriptViewModel->GetEditScript());
+		RefreshEditScriptViewModels();
+		OnActiveScriptChangedDelegate.Broadcast();
 	}
-	else
+}
+
+void UNiagaraScratchPadViewModel::ResetActiveScriptViewModel()
+{
+	if (ActiveScriptViewModel.IsValid())
 	{
+		ActiveScriptViewModel.Reset();
 		ObjectSelection->ClearSelectedObjects();
+		RefreshEditScriptViewModels();
+		OnActiveScriptChangedDelegate.Broadcast();
+	}
+}
+
+void UNiagaraScratchPadViewModel::CopyActiveScript()
+{
+	if (ActiveScriptViewModel.IsValid())
+	{
+		UNiagaraClipboardContent* ClipboardContent = UNiagaraClipboardContent::Create();
+		ClipboardContent->Scripts.Add(CastChecked<UNiagaraScript>(StaticDuplicateObject(ActiveScriptViewModel->GetOriginalScript(), ClipboardContent)));
+		FNiagaraEditorModule::Get().GetClipboard().SetClipboardContent(ClipboardContent);
+	}
+}
+
+bool UNiagaraScratchPadViewModel::CanPasteScript() const
+{
+	const UNiagaraClipboardContent* ClipboardContent = FNiagaraEditorModule::Get().GetClipboard().GetClipboardContent();
+	return ClipboardContent != nullptr && ClipboardContent->Scripts.Num() == 1;
+}
+
+void UNiagaraScratchPadViewModel::PasteScript()
+{
+	if (CanPasteScript())
+	{
+		FScopedTransaction Transaction(LOCTEXT("PasteScratchPadScriptTransaction", "Paste the scratch pad script from the system clipboard."));
+		const UNiagaraClipboardContent* ClipboardContent = FNiagaraEditorModule::Get().GetClipboard().GetClipboardContent();
+		TSharedPtr<FNiagaraScratchPadScriptViewModel> PastedScriptViewModel = CreateNewScriptAsDuplicate(ClipboardContent->Scripts[0]);
+		SetActiveScriptViewModel(PastedScriptViewModel.ToSharedRef());
 	}
-	OnActiveScriptChangedDelegate.Broadcast();
 }
 
 void UNiagaraScratchPadViewModel::DeleteActiveScript()
 {
-	if (ActiveScript != nullptr)
+	if (ActiveScriptViewModel.IsValid())
 	{
+		UNiagaraScript* ActiveScript = ActiveScriptViewModel->GetOriginalScript();
 		FScopedTransaction DeleteTransaction(LOCTEXT("DeleteScratchPadScriptTransaction", "Delete scratch pad script."));
 		for (TObjectIterator<UNiagaraNodeFunctionCall> It; It; ++It)
 		{
@@ -183,6 +263,11 @@ void UNiagaraScratchPadViewModel::DeleteActiveScript()
 	}
 }
 
+FName GetUniqueScriptName(UObject* Outer, const FString& CandidateName)
+{
+	return FNiagaraEditorUtilities::GetUniqueObjectName<UNiagaraScript>(Outer, CandidateName);
+}
+
 TSharedPtr<FNiagaraScratchPadScriptViewModel> UNiagaraScratchPadViewModel::CreateNewScript(ENiagaraScriptUsage InScriptUsage, ENiagaraScriptUsage InTargetSupportedUsage, FNiagaraTypeDefinition InOutputType)
 {
 	UObject* ScriptOuter;
@@ -197,7 +282,7 @@ TSharedPtr<FNiagaraScratchPadScriptViewModel> UNiagaraScratchPadViewModel::Creat
 		UNiagaraScript* DefaultDynamicInput = Cast<UNiagaraScript>(GetDefault<UNiagaraEditorSettings>()->DefaultDynamicInputScript.TryLoad());
 		if (DefaultDynamicInput != nullptr)
 		{
-			NewScript = CastChecked<UNiagaraScript>(StaticDuplicateObject(DefaultDynamicInput, ScriptOuter, MakeUniqueObjectName(ScriptOuter, UNiagaraScript::StaticClass(), TEXT("NewScratchDynamicInput"))));
+			NewScript = CastChecked<UNiagaraScript>(StaticDuplicateObject(DefaultDynamicInput, ScriptOuter, GetUniqueScriptName(ScriptOuter, TEXT("NewScratchDynamicInput"))));
 			TArray<UNiagaraNodeOutput*> OutputNodes;
 			CastChecked<UNiagaraScriptSource>(NewScript->GetSource())->NodeGraph->GetNodesOfClass(OutputNodes);
 			if (OutputNodes.Num() == 1)
@@ -219,7 +304,7 @@ TSharedPtr<FNiagaraScratchPadScriptViewModel> UNiagaraScratchPadViewModel::Creat
 		UNiagaraScript* DefaultModule = Cast<UNiagaraScript>(GetDefault<UNiagaraEditorSettings>()->DefaultModuleScript.TryLoad());
 		if (DefaultModule != nullptr)
 		{
-			NewScript = CastChecked<UNiagaraScript>(StaticDuplicateObject(DefaultModule, ScriptOuter, MakeUniqueObjectName(ScriptOuter, UNiagaraScript::StaticClass(), TEXT("NewScratchModule"))));
+			NewScript = CastChecked<UNiagaraScript>(StaticDuplicateObject(DefaultModule, ScriptOuter, GetUniqueScriptName(ScriptOuter, TEXT("NewScratchModule"))));
 		}
 		break;
 	}
@@ -237,13 +322,13 @@ TSharedPtr<FNiagaraScratchPadScriptViewModel> UNiagaraScratchPadViewModel::Creat
 	return GetViewModelForScript(NewScript);
 }
 
-TSharedPtr<FNiagaraScratchPadScriptViewModel> UNiagaraScratchPadViewModel::CreateNewScriptAsDuplicate(UNiagaraScript* ScriptToDuplicate)
+TSharedPtr<FNiagaraScratchPadScriptViewModel> UNiagaraScratchPadViewModel::CreateNewScriptAsDuplicate(const UNiagaraScript* ScriptToDuplicate)
 {
 	UObject* ScriptOuter;
 	TArray<UNiagaraScript*>* TargetScripts;
 	GetOuterAndTargetScripts(GetSystemViewModel(), ScriptOuter, TargetScripts);
 
-	UNiagaraScript* NewScript = NewScript = CastChecked<UNiagaraScript>(StaticDuplicateObject(ScriptToDuplicate, ScriptOuter, MakeUniqueObjectName(ScriptOuter, UNiagaraScript::StaticClass(), ScriptToDuplicate->GetFName())));
+	UNiagaraScript* NewScript = NewScript = CastChecked<UNiagaraScript>(StaticDuplicateObject(ScriptToDuplicate, ScriptOuter, GetUniqueScriptName(ScriptOuter, *ScriptToDuplicate->GetFName().ToString())));
 	NewScript->ClearFlags(RF_Public | RF_Standalone);
 	ScriptOuter->Modify();
 	TargetScripts->Add(NewScript);
@@ -257,6 +342,11 @@ UNiagaraScratchPadViewModel::FOnScriptViewModelsChanged& UNiagaraScratchPadViewM
 	return OnScriptViewModelsChangedDelegate;
 }
 
+UNiagaraScratchPadViewModel::FOnScriptViewModelsChanged& UNiagaraScratchPadViewModel::OnEditScriptViewModelsChanged()
+{
+	return OnEditScriptViewModelsChangedDelegate;
+}
+
 UNiagaraScratchPadViewModel::FOnActiveScriptChanged& UNiagaraScratchPadViewModel::OnActiveScriptChanged()
 {
 	return OnActiveScriptChangedDelegate;
@@ -279,6 +369,17 @@ TSharedRef<FNiagaraSystemViewModel> UNiagaraScratchPadViewModel::GetSystemViewMo
 	return SystemViewModel.ToSharedRef();
 }
 
+void UNiagaraScratchPadViewModel::RefreshEditScriptViewModels()
+{
+	EditScriptViewModels.Empty();
+	EditScriptViewModels.Append(PinnedScriptViewModels);
+	if (ActiveScriptViewModel.IsValid())
+	{
+		EditScriptViewModels.AddUnique(ActiveScriptViewModel.ToSharedRef());
+	}
+	OnEditScriptViewModelsChangedDelegate.Broadcast();
+}
+
 void UNiagaraScratchPadViewModel::ScriptGraphNodeSelectionChanged(TWeakPtr<FNiagaraScratchPadScriptViewModel> InScriptViewModelWeak)
 {
 	TSharedPtr<FNiagaraScratchPadScriptViewModel> InScriptViewModel = InScriptViewModelWeak.Pin();
@@ -289,9 +390,9 @@ void UNiagaraScratchPadViewModel::ScriptGraphNodeSelectionChanged(TWeakPtr<FNiag
 		{
 			ObjectSelection->SetSelectedObjects(SelectedNodes);
 		}
-		else if (ActiveScript != nullptr)
+		else if (ActiveScriptViewModel.IsValid())
 		{
-			ObjectSelection->SetSelectedObject(ActiveScript);
+			ObjectSelection->SetSelectedObject(ActiveScriptViewModel->GetEditScript());
 		}
 		else
 		{
@@ -305,4 +406,48 @@ void UNiagaraScratchPadViewModel::ScriptViewModelScriptRenamed()
 	OnScriptRenamed().Broadcast();
 }
 
+void UNiagaraScratchPadViewModel::ScriptViewModelPinnedChanged(TWeakPtr<FNiagaraScratchPadScriptViewModel> ScriptViewModelWeak)
+{
+	TSharedPtr<FNiagaraScratchPadScriptViewModel> ScriptViewModel = ScriptViewModelWeak.Pin();
+	bool bPinnedCollectionChanged = false;
+	if (ScriptViewModel.IsValid())
+	{
+		if (ScriptViewModel->GetIsPinned())
+		{
+			PinnedScriptViewModels.AddUnique(ScriptViewModel.ToSharedRef());
+			if (ActiveScriptViewModel != ScriptViewModel)
+			{
+				SetActiveScriptViewModel(ScriptViewModel.ToSharedRef());
+			}
+			else
+			{
+				RefreshEditScriptViewModels();
+			}
+		}
+		else
+		{
+			PinnedScriptViewModels.Remove(ScriptViewModel.ToSharedRef());
+			if (PinnedScriptViewModels.Num() == 0)
+			{
+				if (ActiveScriptViewModel.IsValid() == false)
+				{
+					// When unpinning the last script, and there is no active script set it as the active script so that it remains displayed in the UI.
+					SetActiveScriptViewModel(ScriptViewModel.ToSharedRef());
+				}	
+			}
+			else
+			{
+				if (ActiveScriptViewModel == ScriptViewModel)
+				{
+					SetActiveScriptViewModel(PinnedScriptViewModels.Last());
+				}
+				else
+				{
+					RefreshEditScriptViewModels();
+				}
+			}
+		}
+	}
+}
+
 #undef LOCTEXT_NAMESPACE
\ No newline at end of file
diff --git a/Engine/Plugins/FX/Niagara/Source/NiagaraEditor/Private/ViewModels/Stack/NiagaraStackFunctionInput.cpp b/Engine/Plugins/FX/Niagara/Source/NiagaraEditor/Private/ViewModels/Stack/NiagaraStackFunctionInput.cpp
--- a/Engine/Plugins/FX/Niagara/Source/NiagaraEditor/Private/ViewModels/Stack/NiagaraStackFunctionInput.cpp
+++ b/Engine/Plugins/FX/Niagara/Source/NiagaraEditor/Private/ViewModels/Stack/NiagaraStackFunctionInput.cpp
@@ -1071,9 +1071,9 @@ void UNiagaraStackFunctionInput::GetAvailableDynamicInputs(TArray<UNiagaraScript
 
 	for (TSharedRef<FNiagaraScratchPadScriptViewModel> ScratchPadScriptViewModel : GetSystemViewModel()->GetScriptScratchPadViewModel()->GetScriptViewModels())
 	{
-		if (MatchesInputType(ScratchPadScriptViewModel->GetScript()))
+		if (MatchesInputType(ScratchPadScriptViewModel->GetOriginalScript()))
 		{
-			AvailableDynamicInputs.Add(ScratchPadScriptViewModel->GetScript());
+			AvailableDynamicInputs.Add(ScratchPadScriptViewModel->GetOriginalScript());
 		}
 	}
 }
@@ -1127,10 +1127,11 @@ void UNiagaraStackFunctionInput::SetScratch()
 	TSharedPtr<FNiagaraScratchPadScriptViewModel> ScratchScriptViewModel = GetSystemViewModel()->GetScriptScratchPadViewModel()->CreateNewScript(ENiagaraScriptUsage::DynamicInput, SourceScript->GetUsage(), InputType);
 	if (ScratchScriptViewModel.IsValid())
 	{
-		SetDynamicInput(ScratchScriptViewModel->GetScript());
+		SetDynamicInput(ScratchScriptViewModel->GetOriginalScript());
+		GetSystemViewModel()->FocusTab(FNiagaraSystemToolkit::ScratchPadTabID);
+		GetSystemViewModel()->GetScriptScratchPadViewModel()->SetActiveScriptViewModel(ScratchScriptViewModel.ToSharedRef());
 		ScratchScriptViewModel->SetIsPendingRename(true);
 	}
-	GetSystemViewModel()->FocusTab(FNiagaraSystemToolkit::ScratchPadTabID);
 }
 
 TSharedPtr<const FStructOnScope> UNiagaraStackFunctionInput::GetLocalValueStruct()
@@ -1766,7 +1767,7 @@ void UNiagaraStackFunctionInput::SetValueFromClipboardFunctionInput(const UNiaga
 			{
 				UNiagaraScript* ClipboardScript = ClipboardFunctionInput.Dynamic->Script->IsAsset()
 					? ClipboardFunctionInput.Dynamic->Script
-					: GetSystemViewModel()->GetScriptScratchPadViewModel()->CreateNewScriptAsDuplicate(ClipboardFunctionInput.Dynamic->Script)->GetScript();
+					: GetSystemViewModel()->GetScriptScratchPadViewModel()->CreateNewScriptAsDuplicate(ClipboardFunctionInput.Dynamic->Script)->GetOriginalScript();
 				SetDynamicInput(ClipboardScript, ClipboardFunctionInput.Dynamic->FunctionName);
 
 				TArray<UNiagaraStackFunctionInputCollection*> DynamicInputCollections;
diff --git a/Engine/Plugins/FX/Niagara/Source/NiagaraEditor/Private/ViewModels/Stack/NiagaraStackScriptItemGroup.cpp b/Engine/Plugins/FX/Niagara/Source/NiagaraEditor/Private/ViewModels/Stack/NiagaraStackScriptItemGroup.cpp
--- a/Engine/Plugins/FX/Niagara/Source/NiagaraEditor/Private/ViewModels/Stack/NiagaraStackScriptItemGroup.cpp
+++ b/Engine/Plugins/FX/Niagara/Source/NiagaraEditor/Private/ViewModels/Stack/NiagaraStackScriptItemGroup.cpp
@@ -245,7 +245,7 @@ public:
 		// Generate actions for scratch pad modules
 		for (TSharedRef<FNiagaraScratchPadScriptViewModel> ScratchPadScriptViewModel : SystemViewModel.Pin()->GetScriptScratchPadViewModel()->GetScriptViewModels())
 		{
-			UNiagaraScript* ScratchPadScript = ScratchPadScriptViewModel->GetScript();
+			UNiagaraScript* ScratchPadScript = ScratchPadScriptViewModel->GetOriginalScript();
 			if (ScratchPadScript->Usage == ENiagaraScriptUsage::Module)
 			{
 				TArray<ENiagaraScriptUsage> SupportedUsages = ScratchPadScript->GetSupportedUsageContexts();
@@ -280,9 +280,10 @@ public:
 			TSharedPtr<FNiagaraScratchPadScriptViewModel> ScratchPadScriptViewModel = SystemViewModel.Pin()->GetScriptScratchPadViewModel()->CreateNewScript(ENiagaraScriptUsage::Module, OutputNode->GetUsage(), FNiagaraTypeDefinition());
 			if (ScratchPadScriptViewModel.IsValid())
 			{
-				ScratchPadScriptViewModel->SetIsPendingRename(true);
-				NewModuleNode = FNiagaraStackGraphUtilities::AddScriptModuleToStack(ScratchPadScriptViewModel->GetScript(), *OutputNode, TargetIndex);
+				NewModuleNode = FNiagaraStackGraphUtilities::AddScriptModuleToStack(ScratchPadScriptViewModel->GetOriginalScript(), *OutputNode, TargetIndex);
 				SystemViewModel.Pin()->FocusTab(FNiagaraSystemToolkit::ScratchPadTabID);
+				SystemViewModel.Pin()->GetScriptScratchPadViewModel()->SetActiveScriptViewModel(ScratchPadScriptViewModel.ToSharedRef());
+				ScratchPadScriptViewModel->SetIsPendingRename(true);
 			}
 		}
 
@@ -1108,7 +1109,7 @@ void UNiagaraStackScriptItemGroup::PasteModules(const UNiagaraClipboardContent*
 				{
 					UNiagaraScript* ClipboardScript = ClipboardFunction->Script->IsAsset()
 						? ClipboardFunction->Script
-						: GetSystemViewModel()->GetScriptScratchPadViewModel()->CreateNewScriptAsDuplicate(ClipboardFunction->Script)->GetScript();
+						: GetSystemViewModel()->GetScriptScratchPadViewModel()->CreateNewScriptAsDuplicate(ClipboardFunction->Script)->GetOriginalScript();
 					NewFunctionCallNode = FNiagaraStackGraphUtilities::AddScriptModuleToStack(ClipboardFunction->Script, *OutputNode, CurrentPasteIndex);
 				}
 				else
diff --git a/Engine/Plugins/FX/Niagara/Source/NiagaraEditor/Private/Widgets/SNiagaraScriptGraph.cpp b/Engine/Plugins/FX/Niagara/Source/NiagaraEditor/Private/Widgets/SNiagaraScriptGraph.cpp
--- a/Engine/Plugins/FX/Niagara/Source/NiagaraEditor/Private/Widgets/SNiagaraScriptGraph.cpp
+++ b/Engine/Plugins/FX/Niagara/Source/NiagaraEditor/Private/Widgets/SNiagaraScriptGraph.cpp
@@ -40,6 +40,7 @@ void SNiagaraScriptGraph::Construct(const FArguments& InArgs, TSharedRef<FNiagar
 	bUpdatingGraphSelectionFromViewModel = false;
 
 	GraphTitle = InArgs._GraphTitle;
+	ForegroundColor = InArgs._ForegroundColor;
 
 	GraphEditor = ConstructGraphEditor();
 	if (InArgs._ZoomToFitOnLoad)
diff --git a/Engine/Plugins/FX/Niagara/Source/NiagaraEditor/Public/NiagaraClipboard.h b/Engine/Plugins/FX/Niagara/Source/NiagaraEditor/Public/NiagaraClipboard.h
--- a/Engine/Plugins/FX/Niagara/Source/NiagaraEditor/Public/NiagaraClipboard.h
+++ b/Engine/Plugins/FX/Niagara/Source/NiagaraEditor/Public/NiagaraClipboard.h
@@ -121,6 +121,9 @@ public:
 
 	UPROPERTY()
 	TArray<const UNiagaraRendererProperties*> Renderers;
+
+	UPROPERTY()
+	TArray<const UNiagaraScript*> Scripts;
 };
 
 class FNiagaraClipboard
diff --git a/Engine/Plugins/FX/Niagara/Source/NiagaraEditor/Public/NiagaraEditorUtilities.h b/Engine/Plugins/FX/Niagara/Source/NiagaraEditor/Public/NiagaraEditorUtilities.h
--- a/Engine/Plugins/FX/Niagara/Source/NiagaraEditor/Public/NiagaraEditorUtilities.h
+++ b/Engine/Plugins/FX/Niagara/Source/NiagaraEditor/Public/NiagaraEditorUtilities.h
@@ -244,4 +244,12 @@ namespace FNiagaraEditorUtilities
 	NIAGARAEDITOR_API void WarnWithToastAndLog(FText WarningMessage);
 
 	void GetScriptRunAndExecutionIndexFromUsage(const ENiagaraScriptUsage& InUsage, int32& OutRunIndex, int32&OutExecutionIndex);
+
+	FName GetUniqueObjectName(UObject* Outer, UClass* ObjectClass, const FString& CandidateName);
+
+	template<typename T>
+	FName GetUniqueObjectName(UObject* Outer, const FString& CandidateName)
+	{
+		return GetUniqueObjectName(Outer, T::StaticClass(), CandidateName);
+	}
 };
diff --git a/Engine/Plugins/FX/Niagara/Source/NiagaraEditor/Public/ViewModels/NiagaraScratchPadScriptViewModel.h b/Engine/Plugins/FX/Niagara/Source/NiagaraEditor/Public/ViewModels/NiagaraScratchPadScriptViewModel.h
--- a/Engine/Plugins/FX/Niagara/Source/NiagaraEditor/Public/ViewModels/NiagaraScratchPadScriptViewModel.h
+++ b/Engine/Plugins/FX/Niagara/Source/NiagaraEditor/Public/ViewModels/NiagaraScratchPadScriptViewModel.h
@@ -4,17 +4,26 @@
 
 #include "ViewModels/NiagaraScriptViewModel.h"
 
-class NIAGARAEDITOR_API FNiagaraScratchPadScriptViewModel : public FNiagaraScriptViewModel
+class NIAGARAEDITOR_API FNiagaraScratchPadScriptViewModel : public FNiagaraScriptViewModel, public FGCObject
 {
 public:
 	DECLARE_MULTICAST_DELEGATE(FOnRenamed);
+	DECLARE_MULTICAST_DELEGATE(FOnPinnedChanged);
 
 public:
 	FNiagaraScratchPadScriptViewModel();
 
+	~FNiagaraScratchPadScriptViewModel();
+
 	void Initialize(UNiagaraScript* Script);
 
-	UNiagaraScript* GetScript() const;
+	virtual void AddReferencedObjects(FReferenceCollector& Collector) override;
+
+	UNiagaraScript* GetOriginalScript() const;
+
+	UNiagaraScript* GetEditScript() const;
+
+	FText GetToolTip() const;
 
 	bool GetIsPendingRename() const;
 
@@ -22,12 +31,44 @@ public:
 
 	void SetScriptName(FText InScriptName);
 
+	bool GetIsPinned() const;
+
+	void SetIsPinned(bool bInIsPinned);
+
+	float GetEditorHeight() const;
+
+	void SetEditorHeight(float InEditorHeight);
+
+	bool CanApplyChanges() const;
+
+	void ApplyChanges();
+
 	FOnRenamed& OnRenamed();
 
+	FOnPinnedChanged& OnPinnedChanged();
+
 private:
 	FText GetDisplayNameInternal() const;
 
+	void OnScriptGraphChanged(const FEdGraphEditAction &Action);
+
+	void OnScriptPropertyChanged(FPropertyChangedEvent& PropertyChangedEvent);
+
 private:
 	bool bIsPendingRename;
+
+	bool bIsPinned;
+
+	float EditorHeight;
+
+	UNiagaraScript* OriginalScript;
+
+	UNiagaraScript* EditScript;
+
+	bool bHasPendingChanges;
+
+	FDelegateHandle OnGraphNeedsRecompileHandle;
+
 	FOnRenamed OnRenamedDelegate;
+	FOnPinnedChanged OnPinnedChangedDelegate;
 };
\ No newline at end of file
diff --git a/Engine/Plugins/FX/Niagara/Source/NiagaraEditor/Public/ViewModels/NiagaraScratchPadViewModel.h b/Engine/Plugins/FX/Niagara/Source/NiagaraEditor/Public/ViewModels/NiagaraScratchPadViewModel.h
--- a/Engine/Plugins/FX/Niagara/Source/NiagaraEditor/Public/ViewModels/NiagaraScratchPadViewModel.h
+++ b/Engine/Plugins/FX/Niagara/Source/NiagaraEditor/Public/ViewModels/NiagaraScratchPadViewModel.h
@@ -31,6 +31,8 @@ public:
 
 	const TArray<TSharedRef<FNiagaraScratchPadScriptViewModel>>& GetScriptViewModels() const;
 
+	const TArray<TSharedRef<FNiagaraScratchPadScriptViewModel>>& GetEditScriptViewModels() const;
+
 	TSharedPtr<FNiagaraScratchPadScriptViewModel> GetViewModelForScript(UNiagaraScript* InScript);
 
 	const TArray<ENiagaraScriptUsage>& GetAvailableUsages() const;
@@ -39,18 +41,28 @@ public:
 
 	TSharedRef<FNiagaraObjectSelection> GetObjectSelection();
 
-	UNiagaraScript* GetActiveScript();
+	TSharedPtr<FNiagaraScratchPadScriptViewModel> GetActiveScriptViewModel();
+
+	void SetActiveScriptViewModel(TSharedRef<FNiagaraScratchPadScriptViewModel> InActiveScriptViewModel);
+
+	void ResetActiveScriptViewModel();
 
-	void SetActiveScript(UNiagaraScript* InActiveScript);
+	void CopyActiveScript();
+
+	bool CanPasteScript() const;
+
+	void PasteScript();
 
 	void DeleteActiveScript();
 
 	TSharedPtr<FNiagaraScratchPadScriptViewModel> CreateNewScript(ENiagaraScriptUsage InScriptUsage, ENiagaraScriptUsage InTargetSupportedUsage, FNiagaraTypeDefinition InOutputType);
 
-	TSharedPtr<FNiagaraScratchPadScriptViewModel> CreateNewScriptAsDuplicate(UNiagaraScript* ScriptToDuplicate);
+	TSharedPtr<FNiagaraScratchPadScriptViewModel> CreateNewScriptAsDuplicate(const UNiagaraScript* ScriptToDuplicate);
 
 	FOnScriptViewModelsChanged& OnScriptViewModelsChanged();
 
+	FOnScriptViewModelsChanged& OnEditScriptViewModelsChanged();
+
 	FOnActiveScriptChanged& OnActiveScriptChanged();
 
 	FOnScriptRenamed& OnScriptRenamed();
@@ -60,23 +72,33 @@ public:
 private:
 	TSharedRef<FNiagaraSystemViewModel> GetSystemViewModel();
 
+	void RefreshEditScriptViewModels();
+
 	void ScriptGraphNodeSelectionChanged(TWeakPtr<FNiagaraScratchPadScriptViewModel> InScriptViewModelWeak);
 
 	void ScriptViewModelScriptRenamed();
 
+	void ScriptViewModelPinnedChanged(TWeakPtr<FNiagaraScratchPadScriptViewModel> ScriptViewModelWeak);
+
 private:
 	TSharedPtr<FNiagaraObjectSelection> ObjectSelection;
 
-	UNiagaraScript* ActiveScript;
+	TSharedPtr<FNiagaraScratchPadScriptViewModel> ActiveScriptViewModel;
 
 	TWeakPtr<FNiagaraSystemViewModel> SystemViewModelWeak;
 
 	TArray<TSharedRef<FNiagaraScratchPadScriptViewModel>> ScriptViewModels;
 
+	TArray<TSharedRef<FNiagaraScratchPadScriptViewModel>> PinnedScriptViewModels;
+
+	TArray<TSharedRef<FNiagaraScratchPadScriptViewModel>> EditScriptViewModels;
+
 	TArray<ENiagaraScriptUsage> AvailableUsages;
 
 	FOnScriptViewModelsChanged OnScriptViewModelsChangedDelegate;
 
+	FOnScriptViewModelsChanged OnEditScriptViewModelsChangedDelegate;
+
 	FOnActiveScriptChanged OnActiveScriptChangedDelegate;
 
 	FOnScriptRenamed OnScriptRenamedDelegate;
diff --git a/Engine/Plugins/FX/Niagara/Source/NiagaraEditor/Public/Widgets/SItemSelector.h b/Engine/Plugins/FX/Niagara/Source/NiagaraEditor/Public/Widgets/SItemSelector.h
--- a/Engine/Plugins/FX/Niagara/Source/NiagaraEditor/Public/Widgets/SItemSelector.h
+++ b/Engine/Plugins/FX/Niagara/Source/NiagaraEditor/Public/Widgets/SItemSelector.h
@@ -5,6 +5,7 @@
 #include "Widgets/Views/STableRow.h"
 #include "Widgets/Input/SSearchBox.h"
 #include "Widgets/Views/STreeView.h"
+#include "Styling/CoreStyle.h"
 
 enum class EItemSelectorClickActivateMode
 {
@@ -32,6 +33,8 @@ public:
 	SLATE_BEGIN_ARGS(SItemSelector)
 		: _AllowMultiselect(false)
 		, _ClickActivateMode(EItemSelectorClickActivateMode::DoubleClick)
+		, _CategoryRowStyle(&FEditorStyle::Get().GetWidgetStyle<FTableRowStyle>("TableView.Row"))
+		, _ClearSelectionOnClick(true)
 	{}
 		/** The items to display in the item selector. */
 		SLATE_ARGUMENT(TArray<ItemType>, Items)
@@ -50,6 +53,12 @@ public:
 		/** Whether or not a single click activates an item. */
 		SLATE_ARGUMENT(EItemSelectorClickActivateMode, ClickActivateMode)
 
+		/** The style to use for category rows. */
+		SLATE_STYLE_ARGUMENT(FTableRowStyle, CategoryRowStyle)
+
+		/** Whether or not the selection should be cleared when an empty area is clicked. */
+		SLATE_ARGUMENT(bool, ClearSelectionOnClick)
+
 		/** An optional delegate to get an array of categories for the specified item. Each category in the returned array represents one level of nested categories. 
 		NOTE: The OnCompareCategoriesForEquality, and OnGenerateWidgetForCategory delegates must be bound if this delegate is bound. */
 		SLATE_EVENT(FOnGetCategoriesForItem, OnGetCategoriesForItem)
@@ -498,6 +507,7 @@ public:
 		}
 		DefaultCategoryPaths.Append(InArgs._DefaultCategoryPaths);
 		ClickActivateMode = InArgs._ClickActivateMode;
+		CategoryRowStyle = InArgs._CategoryRowStyle;
 		OnGetCategoriesForItem = InArgs._OnGetCategoriesForItem;
 		OnCompareCategoriesForEquality = InArgs._OnCompareCategoriesForEquality;
 		OnCompareCategoriesForSorting = InArgs._OnCompareCategoriesForSorting;
@@ -543,6 +553,7 @@ public:
 				.OnMouseButtonDoubleClick(this, &SItemSelector::OnMouseDoubleClick)
 				.OnSelectionChanged(this, &SItemSelector::OnTreeSelectionChanged)
 				.TreeItemsSource(ViewModel->GetRootItems())
+				.ClearSelectionOnClick(InArgs._ClearSelectionOnClick)
 			]
 		];
 
@@ -655,6 +666,7 @@ private:
 		{
 			TSharedRef<FItemSelectorItemCategoryViewModel> ItemCategoryViewModel = StaticCastSharedRef<FItemSelectorItemCategoryViewModel>(Item);
 			return SNew(STableRow<TSharedRef<FItemSelectorItemViewModel>>, OwnerTable)
+				.Style(CategoryRowStyle)
 				.ShowSelection(false)
 				[
 					OnGenerateWidgetForCategory.Execute(ItemCategoryViewModel->GetCategory())
@@ -724,6 +736,8 @@ private:
 
 	EItemSelectorClickActivateMode ClickActivateMode;
 
+	const FTableRowStyle* CategoryRowStyle;
+
 	FOnGetCategoriesForItem OnGetCategoriesForItem;
 	FOnCompareCategoriesForEquality OnCompareCategoriesForEquality;
 	FOnCompareCategoriesForSorting OnCompareCategoriesForSorting;
diff --git a/Engine/Plugins/FX/Niagara/Source/NiagaraEditor/Public/Widgets/SNiagaraScriptGraph.h b/Engine/Plugins/FX/Niagara/Source/NiagaraEditor/Public/Widgets/SNiagaraScriptGraph.h
--- a/Engine/Plugins/FX/Niagara/Source/NiagaraEditor/Public/Widgets/SNiagaraScriptGraph.h
+++ b/Engine/Plugins/FX/Niagara/Source/NiagaraEditor/Public/Widgets/SNiagaraScriptGraph.h
@@ -8,6 +8,7 @@
 #include "EdGraph/EdGraphNode.h"
 #include "Widgets/Input/SSearchBox.h"
 #include "NiagaraEditorCommon.h"
+#include "EditorStyleSet.h"
 
 class FNiagaraScriptGraphViewModel;
 
@@ -16,10 +17,12 @@ class SNiagaraScriptGraph : public SCompoundWidget
 {
 public:
 	SLATE_BEGIN_ARGS(SNiagaraScriptGraph)
-		: _ZoomToFitOnLoad(false)
+		: _ForegroundColor(FEditorStyle::GetColor("DefaultForeground"))
+		, _ZoomToFitOnLoad(false)
 	{}
 		/** The text displayed in the title bar of the graph. */
 		SLATE_ATTRIBUTE(FText, GraphTitle)
+		SLATE_ATTRIBUTE(FSlateColor, ForegroundColor)
 		SLATE_ARGUMENT(bool, ZoomToFitOnLoad)
 	SLATE_END_ARGS();
 
diff --git a/Engine/Plugins/FX/Niagara/Source/NiagaraEditorWidgets/Private/NiagaraEditorWidgetsStyle.cpp b/Engine/Plugins/FX/Niagara/Source/NiagaraEditorWidgets/Private/NiagaraEditorWidgetsStyle.cpp
--- a/Engine/Plugins/FX/Niagara/Source/NiagaraEditorWidgets/Private/NiagaraEditorWidgetsStyle.cpp
+++ b/Engine/Plugins/FX/Niagara/Source/NiagaraEditorWidgets/Private/NiagaraEditorWidgetsStyle.cpp
@@ -243,12 +243,18 @@ TSharedRef< FSlateStyleSet > FNiagaraEditorWidgetsStyle::Create()
 		.SetFont(ScratchPadLargeHeaderFont);
 	Style->Set("NiagaraEditor.ScratchPad.LargeHeaderText", ScratchPadLargeHeaderText);
 
-	FSlateFontInfo ScratchPadSmallHeaderFont = DEFAULT_FONT("Bold", 9);
+	FSlateFontInfo ScratchPadSmallHeaderFont = DEFAULT_FONT("Bold", 10);
 	FTextBlockStyle ScratchPadSmallHeaderText = FTextBlockStyle(NormalText)
 		.SetFont(ScratchPadSmallHeaderFont);
 	Style->Set("NiagaraEditor.ScratchPad.SmallHeaderText", ScratchPadSmallHeaderText);
 
-	Style->Set("NiagaraEditor.ScratchPad.HeaderColor", FLinearColor(FColor(48, 48, 48)));
+	FSlateBrush ScratchPadCategoryBrush = BOX_PLUGIN_BRUSH("Icons/CategoryRow", FMargin(2.0f / 8.0f), FLinearColor(FColor(48, 48, 48)));
+	FSlateBrush ScratchPadHoveredCategoryBrush = BOX_PLUGIN_BRUSH("Icons/CategoryRow", FMargin(2.0f / 8.0f), FLinearColor(FColor(38, 38, 38)));
+	Style->Set("NiagaraEditor.ScratchPad.CategoryRow", FTableRowStyle(NormalTableRowStyle)
+		.SetEvenRowBackgroundBrush(ScratchPadCategoryBrush)
+		.SetOddRowBackgroundBrush(ScratchPadCategoryBrush)
+		.SetEvenRowBackgroundHoveredBrush(ScratchPadHoveredCategoryBrush)
+		.SetOddRowBackgroundHoveredBrush(ScratchPadHoveredCategoryBrush));
 
 	Style->Set("NiagaraEditor.Scope.Engine", FLinearColor(FColor(230, 102, 102)));
 	Style->Set("NiagaraEditor.Scope.Owner", FLinearColor(FColor(210, 112, 112)));
diff --git a/Engine/Plugins/FX/Niagara/Source/NiagaraEditorWidgets/Private/NiagaraScratchPadCommandContext.cpp b/Engine/Plugins/FX/Niagara/Source/NiagaraEditorWidgets/Private/NiagaraScratchPadCommandContext.cpp
--- a/Engine/Plugins/FX/Niagara/Source/NiagaraEditorWidgets/Private/NiagaraScratchPadCommandContext.cpp
+++ b/Engine/Plugins/FX/Niagara/Source/NiagaraEditorWidgets/Private/NiagaraScratchPadCommandContext.cpp
@@ -2,6 +2,9 @@
 
 #include "NiagaraScratchPadCommandContext.h"
 #include "ViewModels/NiagaraScratchPadViewModel.h"
+#include "ViewModels/NiagaraScratchPadScriptViewModel.h"
+#include "NiagaraEditorModule.h"
+#include "NiagaraEditorCommands.h"
 
 #include "Framework/Notifications/NotificationManager.h"
 #include "Logging/LogMacros.h"
@@ -36,41 +39,29 @@ bool FNiagaraScratchPadCommandContext::ProcessCommandBindings(const FKeyEvent& I
 	return GetCommands()->ProcessCommandBindings(InKeyEvent);
 }
 
-bool FNiagaraScratchPadCommandContext::AddEditMenuItems(FMenuBuilder& MenuBuilder)
+void FNiagaraScratchPadCommandContext::AddMenuItems(FMenuBuilder& MenuBuilder)
 {
-	bool bSupportsCut = false;
-	bool bSupportsCopy = false;
-	bool bSupportsPaste = false;
-	bool bSupportsDelete = true;
-	if (bSupportsCut || bSupportsCopy || bSupportsPaste || bSupportsDelete)
+	MenuBuilder.BeginSection("ScriptEdit", LOCTEXT("ScriptEditActions", "Edit"));
 	{
-		MenuBuilder.BeginSection("EntryEdit", LOCTEXT("EntryEditActions", "Edit"));
-		{
-			if (bSupportsCut)
-			{
-				TAttribute<FText> CanCutToolTip = TAttribute<FText>::Create(TAttribute<FText>::FGetter::CreateSP(this, &FNiagaraScratchPadCommandContext::GetCanCutSelectedScriptsToolTip));
-				MenuBuilder.AddMenuEntry(FGenericCommands::Get().Cut, NAME_None, TAttribute<FText>(), CanCutToolTip);
-			}
-			if (bSupportsCopy)
-			{
-				TAttribute<FText> CanCopyToolTip = TAttribute<FText>::Create(TAttribute<FText>::FGetter::CreateSP(this, &FNiagaraScratchPadCommandContext::GetCanCopySelectedScriptsToolTip));
-				MenuBuilder.AddMenuEntry(FGenericCommands::Get().Copy, NAME_None, TAttribute<FText>(), CanCopyToolTip);
-			}
-			if (bSupportsPaste)
-			{
-				TAttribute<FText> CanPasteToolTip = TAttribute<FText>::Create(TAttribute<FText>::FGetter::CreateSP(this, &FNiagaraScratchPadCommandContext::GetCanPasteSelectedScriptsToolTip));
-				MenuBuilder.AddMenuEntry(FGenericCommands::Get().Paste, NAME_None, TAttribute<FText>(), CanPasteToolTip);
-			}
-			if (bSupportsDelete)
-			{
-				TAttribute<FText> CanDeleteToolTip = TAttribute<FText>::Create(TAttribute<FText>::FGetter::CreateSP(this, &FNiagaraScratchPadCommandContext::GetCanDeleteSelectedScriptsToolTip));
-				MenuBuilder.AddMenuEntry(FGenericCommands::Get().Delete, NAME_None, TAttribute<FText>(), CanDeleteToolTip);
-			}
-		}
-		MenuBuilder.EndSection();
-		return true;
+		TAttribute<FText> CanCutToolTip = TAttribute<FText>::Create(TAttribute<FText>::FGetter::CreateSP(this, &FNiagaraScratchPadCommandContext::GetCanCutSelectedScriptsToolTip));
+		MenuBuilder.AddMenuEntry(FGenericCommands::Get().Cut, NAME_None, TAttribute<FText>(), CanCutToolTip);
+
+		TAttribute<FText> CanCopyToolTip = TAttribute<FText>::Create(TAttribute<FText>::FGetter::CreateSP(this, &FNiagaraScratchPadCommandContext::GetCanCopySelectedScriptsToolTip));
+		MenuBuilder.AddMenuEntry(FGenericCommands::Get().Copy, NAME_None, TAttribute<FText>(), CanCopyToolTip);
+
+		TAttribute<FText> CanPasteToolTip = TAttribute<FText>::Create(TAttribute<FText>::FGetter::CreateSP(this, &FNiagaraScratchPadCommandContext::GetCanPasteSelectedScriptsToolTip));
+		MenuBuilder.AddMenuEntry(FGenericCommands::Get().Paste, NAME_None, TAttribute<FText>(), CanPasteToolTip);
+
+		TAttribute<FText> CanDeleteToolTip = TAttribute<FText>::Create(TAttribute<FText>::FGetter::CreateSP(this, &FNiagaraScratchPadCommandContext::GetCanDeleteSelectedScriptsToolTip));
+		MenuBuilder.AddMenuEntry(FGenericCommands::Get().Delete, NAME_None, TAttribute<FText>(), CanDeleteToolTip);
+	}
+	MenuBuilder.EndSection();
+
+	MenuBuilder.BeginSection("Script", LOCTEXT("ScriptActions", "Script"));
+	{
+		MenuBuilder.AddMenuEntry(FNiagaraEditorModule::Get().Commands().Apply);
 	}
-	return false;
+	MenuBuilder.EndSection();
 }
 
 void FNiagaraScratchPadCommandContext::SetupCommands()
@@ -87,53 +78,64 @@ void FNiagaraScratchPadCommandContext::SetupCommands()
 	Commands->MapAction(FGenericCommands::Get().Delete, FUIAction(
 		FExecuteAction::CreateSP(this, &FNiagaraScratchPadCommandContext::DeleteSelectedScripts),
 		FCanExecuteAction::CreateSP(this, &FNiagaraScratchPadCommandContext::CanDeleteSelectedScripts)));
+	Commands->MapAction(FNiagaraEditorModule::Get().Commands().Apply, FUIAction(
+		FExecuteAction::CreateSP(this, &FNiagaraScratchPadCommandContext::ApplyChangesToSelectedScripts),
+		FCanExecuteAction::CreateSP(this, &FNiagaraScratchPadCommandContext::CanApplyChangesToSelectedScripts)));
 }
 
 bool FNiagaraScratchPadCommandContext::CanCutSelectedScripts() const
 {
-	return false;
+	return ScratchPadViewModel->GetActiveScriptViewModel().IsValid();
 }
 
 FText FNiagaraScratchPadCommandContext::GetCanCutSelectedScriptsToolTip() const
 {
-	return FText();
+	return LOCTEXT("CutToolTip", "Cut the selected scratch pad script.");
 }
 
 void FNiagaraScratchPadCommandContext::CutSelectedScripts() const
 {
+	if (CanCutSelectedScripts())
+	{
+		FScopedTransaction Transaction(LOCTEXT("CutScratchPadScriptTransaction", "Cut the selected scratch pad script to the system clipboard."));
+		ScratchPadViewModel->CopyActiveScript();
+		ScratchPadViewModel->DeleteActiveScript();
+	}
 }
 
 bool FNiagaraScratchPadCommandContext::CanCopySelectedScripts() const
 {
-	return false;
+	return ScratchPadViewModel->GetActiveScriptViewModel().IsValid();
 }
 
 FText FNiagaraScratchPadCommandContext::GetCanCopySelectedScriptsToolTip() const
 {
-	return FText();
+	return LOCTEXT("CopyToolTip", "Copy the selected scratch pad script to the system clipboard.");
 }
 
 void FNiagaraScratchPadCommandContext::CopySelectedScripts() const
 {
+	ScratchPadViewModel->CopyActiveScript();
 }
 
 bool FNiagaraScratchPadCommandContext::CanPasteSelectedScripts() const
 {
-	return false;
+	return ScratchPadViewModel->CanPasteScript();
 }
 
 FText FNiagaraScratchPadCommandContext::GetCanPasteSelectedScriptsToolTip() const
 {
-	return FText();
+	return LOCTEXT("PasteToolTip", "Paste a scratch pad script from the system clipboard.");;
 }
 
 void FNiagaraScratchPadCommandContext::PasteSelectedScripts() const
 {
+	ScratchPadViewModel->PasteScript();
 }
 
 bool FNiagaraScratchPadCommandContext::CanDeleteSelectedScripts() const
 {
-	return true;
+	return ScratchPadViewModel->GetActiveScriptViewModel().IsValid();
 }
 
 FText FNiagaraScratchPadCommandContext::GetCanDeleteSelectedScriptsToolTip() const
@@ -146,5 +148,18 @@ void FNiagaraScratchPadCommandContext::DeleteSelectedScripts() const
 	ScratchPadViewModel->DeleteActiveScript();
 }
 
+bool FNiagaraScratchPadCommandContext::CanApplyChangesToSelectedScripts() const
+{
+	return ScratchPadViewModel->GetActiveScriptViewModel().IsValid() && ScratchPadViewModel->GetActiveScriptViewModel()->CanApplyChanges();
+}
+
+void FNiagaraScratchPadCommandContext::ApplyChangesToSelectedScripts() const
+{
+	if (CanApplyChangesToSelectedScripts())
+	{
+		ScratchPadViewModel->GetActiveScriptViewModel()->ApplyChanges();
+	}
+}
+
 #undef LOCTEXT_NAMESPACE
 
diff --git a/Engine/Plugins/FX/Niagara/Source/NiagaraEditorWidgets/Private/NiagaraScratchPadCommandContext.h b/Engine/Plugins/FX/Niagara/Source/NiagaraEditorWidgets/Private/NiagaraScratchPadCommandContext.h
--- a/Engine/Plugins/FX/Niagara/Source/NiagaraEditorWidgets/Private/NiagaraScratchPadCommandContext.h
+++ b/Engine/Plugins/FX/Niagara/Source/NiagaraEditorWidgets/Private/NiagaraScratchPadCommandContext.h
@@ -18,7 +18,7 @@ public:
 
 	bool ProcessCommandBindings(const FKeyEvent& InKeyEvent);
 
-	bool AddEditMenuItems(FMenuBuilder& MenuBuilder);
+	void AddMenuItems(FMenuBuilder& MenuBuilder);
 
 private:
 	void SetupCommands();
@@ -47,6 +47,10 @@ private:
 
 	void DeleteSelectedScripts() const;
 
+	bool CanApplyChangesToSelectedScripts() const;
+
+	void ApplyChangesToSelectedScripts() const;
+
 private:
 	TSharedRef<FUICommandList> Commands;
 
diff --git a/Engine/Plugins/FX/Niagara/Source/NiagaraEditorWidgets/Private/SNiagaraScratchPad.cpp b/Engine/Plugins/FX/Niagara/Source/NiagaraEditorWidgets/Private/SNiagaraScratchPad.cpp
--- a/Engine/Plugins/FX/Niagara/Source/NiagaraEditorWidgets/Private/SNiagaraScratchPad.cpp
+++ b/Engine/Plugins/FX/Niagara/Source/NiagaraEditorWidgets/Private/SNiagaraScratchPad.cpp
@@ -3,7 +3,6 @@
 #include "SNiagaraScratchPad.h"
 #include "ViewModels/NiagaraScratchPadViewModel.h"
 #include "ViewModels/NiagaraScratchPadScriptViewModel.h"
-#include "ViewModels/Stack/NiagaraStackFunctionInputCollection.h"
 #include "Widgets/SDynamicLayoutBox.h"
 #include "Widgets/SItemSelector.h"
 #include "Widgets/SNiagaraSelectedObjectsDetails.h"
@@ -11,6 +10,7 @@
 #include "Widgets/SNiagaraScriptGraph.h"
 #include "NiagaraObjectSelection.h"
 #include "NiagaraNodeFunctionCall.h"
+#include "NiagaraEditorStyle.h"
 #include "NiagaraEditorWidgetsStyle.h"
 #include "NiagaraScratchPadCommandContext.h"
 
@@ -20,6 +20,7 @@
 #include "Framework/MultiBox/MultiBoxBuilder.h"
 #include "Framework/Application/SlateApplication.h"
 #include "Layout/WidgetPath.h"
+#include "EditorFontGlyphs.h"
 
 #define LOCTEXT_NAMESPACE "NiagaraScratchPad"
 
@@ -29,6 +30,100 @@ FName SelectionEditorName = "SelectionEditor";
 FName WideLayoutName = "Wide";
 FName NarrowLayoutName = "Narrow";
 
+class SNiagaraPinButton : public SButton
+{
+public:
+	DECLARE_DELEGATE_OneParam(FOnPinnedChanged, bool /* bIsPinned */)
+
+public:
+	SLATE_BEGIN_ARGS(SNiagaraPinButton)
+		: _IsPinned(false) 
+		, _ShowWhenUnpinned(true)
+		, _PinTargetDisplayName(LOCTEXT("DefaultTargetDisplayName", "Target"))
+		, _PinItemDisplayName(LOCTEXT("DefaultItemDisplayName", "Item"))
+	{}
+		SLATE_ATTRIBUTE(bool, IsPinned)
+		SLATE_ATTRIBUTE(bool, ShowWhenUnpinned)
+		SLATE_ARGUMENT(FText, PinTargetDisplayName)
+		SLATE_ARGUMENT(FText, PinItemDisplayName)
+		SLATE_EVENT(FOnPinnedChanged, OnPinnedChanged);
+	SLATE_END_ARGS()
+
+	void Construct(const FArguments& InArgs)
+	{
+		IsPinned = InArgs._IsPinned;
+		ShowWhenUnpinned = InArgs._ShowWhenUnpinned;
+		OnPinnedChangedDelegate = InArgs._OnPinnedChanged;
+		PinnedToolTip = FText::Format(LOCTEXT("UnpinnedFormat", "Unpin this {0} from the {1}."), InArgs._PinItemDisplayName, InArgs._PinTargetDisplayName);
+		UnpinnedToolTip = FText::Format(LOCTEXT("PinnedFormat", "Pin this {0} to the {1}."), InArgs._PinItemDisplayName, InArgs._PinTargetDisplayName);
+
+		// Visibility and ToolTipText are base attributes so can't be set in the construct call below,
+		// so them them directly here since the base widget construct has already been run.
+		TAttribute<EVisibility> PinVisibility;
+		PinVisibility.Bind(this, &SNiagaraPinButton::GetVisibilityFromPinned);
+		SetVisibility(PinVisibility);
+		TAttribute<FText> PinToolTipText;
+		PinToolTipText.Bind(this, &SNiagaraPinButton::GetToolTipTextFromPinned);
+		SetToolTipText(PinToolTipText);
+
+		SButton::Construct(
+			SButton::FArguments()
+			.ButtonStyle(FEditorStyle::Get(), "HoverHintOnly")
+			.ForegroundColor(FSlateColor::UseForeground())
+			.OnClicked(this, &SNiagaraPinButton::OnButtonClicked)
+			.ContentPadding(FMargin(3, 2, 2, 2))
+			[
+				SNew(SBox)
+				.HAlign(HAlign_Center)
+				.VAlign(VAlign_Center)
+				.WidthOverride(16)
+				.HeightOverride(16)
+				.RenderTransform(this, &SNiagaraPinButton::GetPinGlyphRenderTransform)
+				.RenderTransformPivot(FVector2D(0.5f, 0.5f))
+				[
+					SNew(STextBlock)
+					.Font(FEditorStyle::Get().GetFontStyle("FontAwesome.10"))
+					.Text(FEditorFontGlyphs::Thumb_Tack)
+				]
+			]);
+	}
+
+private:
+	FReply OnButtonClicked()
+	{
+		OnPinnedChangedDelegate.ExecuteIfBound(!IsPinned.Get());
+		return FReply::Handled();
+	}
+
+	FText GetToolTipTextFromPinned() const
+	{
+		return IsPinned.Get(false) 
+			? PinnedToolTip 
+			: UnpinnedToolTip;
+	}
+
+	EVisibility GetVisibilityFromPinned() const
+	{
+		return IsPinned.Get(false) || ShowWhenUnpinned.Get(true)
+			? EVisibility::Visible
+			: EVisibility::Hidden;
+	}
+
+	TOptional<FSlateRenderTransform> GetPinGlyphRenderTransform() const
+	{
+		return IsPinned.Get(false) 
+			? TOptional<FSlateRenderTransform>()
+			: FSlateRenderTransform(FQuat2D(PI / 2));
+	}
+
+private:
+	TAttribute<bool> IsPinned;
+	TAttribute<bool> ShowWhenUnpinned;
+	FOnPinnedChanged OnPinnedChangedDelegate;
+	FText PinnedToolTip;
+	FText UnpinnedToolTip;
+};
+
 class SNiagaraScratchPadScriptRow : public SCompoundWidget
 {
 	SLATE_BEGIN_ARGS(SNiagaraScratchPadScriptRow) {}
@@ -48,10 +143,28 @@ class SNiagaraScratchPadScriptRow : public SCompoundWidget
 
 		ChildSlot
 		[
-			SAssignNew(NameEditableText, SInlineEditableTextBlock)
-			.Text(this, &SNiagaraScratchPadScriptRow::GetNameText)
-			.IsSelected(this, &SNiagaraScratchPadScriptRow::GetIsSelected)
-			.OnTextCommitted(this, &SNiagaraScratchPadScriptRow::OnNameTextCommitted)
+			SNew(SHorizontalBox)
+			+ SHorizontalBox::Slot()
+			.VAlign(VAlign_Center)
+			.Padding(3, 0, 0, 0)
+			[
+				SAssignNew(NameEditableText, SInlineEditableTextBlock)
+				.Text(this, &SNiagaraScratchPadScriptRow::GetNameText)
+				.ToolTipText(ScriptViewModel.ToSharedRef(), &FNiagaraScratchPadScriptViewModel::GetToolTip)
+				.IsSelected(this, &SNiagaraScratchPadScriptRow::GetIsSelected)
+				.OnTextCommitted(this, &SNiagaraScratchPadScriptRow::OnNameTextCommitted)
+			]
+			+ SHorizontalBox::Slot()
+			.AutoWidth()
+			.Padding(1)
+			[
+				SNew(SNiagaraPinButton)
+				.IsPinned(ScriptViewModel.ToSharedRef(), &FNiagaraScratchPadScriptViewModel::GetIsPinned)
+				.OnPinnedChanged(ScriptViewModel.ToSharedRef(), &FNiagaraScratchPadScriptViewModel::SetIsPinned)
+				.ShowWhenUnpinned(this, &SNiagaraScratchPadScriptRow::IsActive)
+				.PinItemDisplayName(LOCTEXT("PinItem", "script"))
+				.PinTargetDisplayName(LOCTEXT("PinTarget", "edit list"))
+			]
 		];
 	}
 
@@ -79,10 +192,10 @@ class SNiagaraScratchPadScriptRow : public SCompoundWidget
 		if (MouseEvent.GetEffectingButton() == EKeys::RightMouseButton)
 		{
 			// Set this script to be the active one.
-			ScratchPadViewModel->SetActiveScript(ScriptViewModel->GetScript());
+			ScratchPadViewModel->SetActiveScriptViewModel(ScriptViewModel.ToSharedRef());
 
 			FMenuBuilder MenuBuilder(true, CommandContext->GetCommands());
-			CommandContext->AddEditMenuItems(MenuBuilder);
+			CommandContext->AddMenuItems(MenuBuilder);
 
 			FWidgetPath WidgetPath = MouseEvent.GetEventPath() != nullptr ? *MouseEvent.GetEventPath() : FWidgetPath();
 			FSlateApplication::Get().PushMenu(AsShared(), WidgetPath, MenuBuilder.MakeWidget(), MouseEvent.GetScreenSpacePosition(), FPopupTransitionEffect(FPopupTransitionEffect::ContextMenu));
@@ -107,6 +220,11 @@ private:
 		ScriptViewModel->SetScriptName(InText);
 	}
 
+	bool IsActive() const
+	{
+		return IsHovered();
+	}
+
 private:
 	UNiagaraScratchPadViewModel* ScratchPadViewModel;
 	TSharedPtr<FNiagaraScratchPadScriptViewModel> ScriptViewModel;
@@ -132,6 +250,8 @@ class SNiagaraScratchPadScriptSelector : public SCompoundWidget
 		[
 			SAssignNew(ScriptSelector, SNiagaraScriptViewModelSelector)
 			.ClickActivateMode(EItemSelectorClickActivateMode::SingleClick)
+			.CategoryRowStyle(FNiagaraEditorWidgetsStyle::Get(), "NiagaraEditor.ScratchPad.CategoryRow")
+			.ClearSelectionOnClick(false)
 			.Items(ViewModel->GetScriptViewModels())
 			.DefaultCategories(ViewModel->GetAvailableUsages())
 			.OnGetCategoriesForItem(this, &SNiagaraScratchPadScriptSelector::OnGetCategoriesForItem)
@@ -145,6 +265,13 @@ class SNiagaraScratchPadScriptSelector : public SCompoundWidget
 			.OnItemActivated(this, &SNiagaraScratchPadScriptSelector::OnScriptActivated)
 			.OnSelectionChanged(this, &SNiagaraScratchPadScriptSelector::OnSelectionChanged)
 		];
+
+		if (ViewModel->GetActiveScriptViewModel().IsValid())
+		{
+			TArray<TSharedRef<FNiagaraScratchPadScriptViewModel>> SelectedViewModels;
+			SelectedViewModels.Add(ViewModel->GetActiveScriptViewModel().ToSharedRef());
+			ScriptSelector->SetSelectedItems(SelectedViewModels);
+		}
 	}
 
 	virtual FReply OnKeyDown(const FGeometry& MyGeometry, const FKeyEvent& InKeyEvent) override
@@ -156,6 +283,30 @@ class SNiagaraScratchPadScriptSelector : public SCompoundWidget
 		return SCompoundWidget::OnKeyDown(MyGeometry, InKeyEvent);
 	}
 
+	virtual FReply OnMouseButtonDown(const FGeometry& MyGeometry, const FPointerEvent& MouseEvent) override
+	{
+		if (MouseEvent.GetEffectingButton() == EKeys::RightMouseButton)
+		{
+			return FReply::Handled()
+				.CaptureMouse(SharedThis(this));
+		}
+		return SCompoundWidget::OnMouseButtonDown(MyGeometry, MouseEvent);
+	}
+
+	virtual FReply OnMouseButtonUp(const FGeometry& MyGeometry, const FPointerEvent& MouseEvent) override
+	{
+		if (MouseEvent.GetEffectingButton() == EKeys::RightMouseButton)
+		{
+			FMenuBuilder MenuBuilder(true, CommandContext->GetCommands());
+			CommandContext->AddMenuItems(MenuBuilder);
+
+			FWidgetPath WidgetPath = MouseEvent.GetEventPath() != nullptr ? *MouseEvent.GetEventPath() : FWidgetPath();
+			FSlateApplication::Get().PushMenu(AsShared(), WidgetPath, MenuBuilder.MakeWidget(), MouseEvent.GetScreenSpacePosition(), FPopupTransitionEffect(FPopupTransitionEffect::ContextMenu));
+			return FReply::Handled().ReleaseMouseCapture();
+		}
+		return SCompoundWidget::OnMouseButtonUp(MyGeometry, MouseEvent);
+	}
+
 private:
 	void ScriptViewModelsChanged()
 	{
@@ -170,7 +321,7 @@ private:
 		if (bIsUpdatingSelection == false)
 		{
 			TGuardValue<bool> UpdateGuard(bIsUpdatingSelection, true);
-			TSharedPtr<FNiagaraScratchPadScriptViewModel> ActiveScriptViewModel = ViewModel->GetViewModelForScript(ViewModel->GetActiveScript());
+			TSharedPtr<FNiagaraScratchPadScriptViewModel> ActiveScriptViewModel = ViewModel->GetActiveScriptViewModel();
 			if (ActiveScriptViewModel.IsValid())
 			{
 				TArray<TSharedRef<FNiagaraScratchPadScriptViewModel>> SelectedItems;
@@ -218,29 +369,27 @@ private:
 
 	TSharedRef<SWidget> OnGenerateWidgetForCategory(const ENiagaraScriptUsage& Category)
 	{
-		return SNew(SBorder)
-		.BorderImage(FEditorStyle::GetBrush("WhiteBrush"))
-		.BorderBackgroundColor(FNiagaraEditorWidgetsStyle::Get().GetColor("NiagaraEditor.ScratchPad.HeaderColor"))
+		return SNew(SHorizontalBox)
+		+ SHorizontalBox::Slot()
+		.VAlign(VAlign_Center)
+		.Padding(3, 0, 0, 0)
 		[
-			SNew(SHorizontalBox)
-			+ SHorizontalBox::Slot()
-			.VAlign(VAlign_Center)
-			.Padding(5, 0, 0, 0)
-			[
-				SNew(STextBlock)
-				.TextStyle(FNiagaraEditorWidgetsStyle::Get(), "NiagaraEditor.ScratchPad.SmallHeaderText")
-				.Text(ViewModel->GetDisplayNameForUsage(Category))
-			]
-			+ SHorizontalBox::Slot()
-			.AutoWidth()
+			SNew(STextBlock)
+			.TextStyle(FNiagaraEditorWidgetsStyle::Get(), "NiagaraEditor.ScratchPad.SmallHeaderText")
+			.Text(ViewModel->GetDisplayNameForUsage(Category))
+		]
+		+ SHorizontalBox::Slot()
+		.AutoWidth()
+		.Padding(0.0f, 4.0f, 3.0f, 4.0f)
+		[
+			SNew(SButton)
+			.ButtonStyle(FEditorStyle::Get(), "RoundButton")
+			.OnClicked(this, &SNiagaraScratchPadScriptSelector::ScriptSelectorAddButtonClicked, Category)
+			.ContentPadding(FMargin(3.0f, 2.0f, 2.0f, 2.0f))
+			.Content()
 			[
-				SNew(SButton)
-				.OnClicked(this, &SNiagaraScratchPadScriptSelector::ScriptSelectorAddButtonClicked, Category)
-				.Content()
-				[
-					SNew(STextBlock)
-					.Text(LOCTEXT("AddButtonText", "Add"))
-				]
+				SNew(SImage)
+				.Image(FEditorStyle::GetBrush("Plus"))
 			]
 		];
 	}
@@ -256,7 +405,7 @@ private:
 		if (bIsUpdatingSelection == false)
 		{
 			TGuardValue<bool> UpdateGuard(bIsUpdatingSelection, true);
-			ViewModel->SetActiveScript(ActivatedScript->GetScript());
+			ViewModel->SetActiveScriptViewModel(ActivatedScript);
 		}
 	}
 
@@ -268,11 +417,11 @@ private:
 			TArray<TSharedRef<FNiagaraScratchPadScriptViewModel>> SelectedScripts = ScriptSelector->GetSelectedItems();
 			if (SelectedScripts.Num() == 0)
 			{
-				ViewModel->SetActiveScript(nullptr);
+				ViewModel->ResetActiveScriptViewModel();
 			}
 			else if (SelectedScripts.Num())
 			{
-				ViewModel->SetActiveScript(SelectedScripts[0]->GetScript());
+				ViewModel->SetActiveScriptViewModel(SelectedScripts[0]);
 			}
 		}
 	}
@@ -282,6 +431,7 @@ private:
 		TSharedPtr<FNiagaraScratchPadScriptViewModel> NewScriptViewModel = ViewModel->CreateNewScript(Usage, ENiagaraScriptUsage::ParticleUpdateScript, FNiagaraTypeDefinition());
 		if (NewScriptViewModel.IsValid())
 		{
+			ViewModel->SetActiveScriptViewModel(NewScriptViewModel.ToSharedRef());
 			NewScriptViewModel->SetIsPendingRename(true);
 		}
 		return FReply::Handled();
@@ -290,7 +440,7 @@ private:
 	bool GetItemIsSelected(TWeakPtr<FNiagaraScratchPadScriptViewModel> ItemWeak) const
 	{
 		TSharedPtr<FNiagaraScratchPadScriptViewModel> Item = ItemWeak.Pin();
-		return Item.IsValid() && ViewModel->GetActiveScript() == Item->GetScripts()[0];
+		return Item.IsValid() && ViewModel->GetActiveScriptViewModel() == Item;
 	}
 
 private:
@@ -317,21 +467,54 @@ class SNiagaraScratchPadScriptEditor : public SCompoundWidget
 				SNew(SHorizontalBox)
 				+ SHorizontalBox::Slot()
 				.VAlign(VAlign_Center)
+				.Padding(0, 0, 2, 0)
+				.AutoWidth()
 				[
 					SNew(STextBlock)
 					.TextStyle(FNiagaraEditorWidgetsStyle::Get(), "NiagaraEditor.ScratchPad.LargeHeaderText")
 					.Text(this, &SNiagaraScratchPadScriptEditor::GetNameText)
+					.ToolTipText(ScriptViewModel.ToSharedRef(), &FNiagaraScratchPadScriptViewModel::GetToolTip)
 				]
 				+ SHorizontalBox::Slot()
 				.AutoWidth()
+				.Padding(1)
+				[
+					SNew(SNiagaraPinButton)
+					.IsPinned(ScriptViewModel.ToSharedRef(), &FNiagaraScratchPadScriptViewModel::GetIsPinned)
+					.OnPinnedChanged(ScriptViewModel.ToSharedRef(), &FNiagaraScratchPadScriptViewModel::SetIsPinned)
+					.PinItemDisplayName(LOCTEXT("PinItem", "script"))
+					.PinTargetDisplayName(LOCTEXT("PinTarget", "edit list"))
+				]
+				+ SHorizontalBox::Slot()
+				.HAlign(HAlign_Right)
+				.Padding(0.0f, 2.0f, 1.0f, 4.0f)
 				[
 					SNew(SButton)
-					.OnClicked(this, &SNiagaraScratchPadScriptEditor::OnRefreshButtonClicked)
-					.ToolTipText(LOCTEXT("RefreshButtonToolTip", "Refresh the inputs for this script in the selection stack."))
+					.ButtonStyle(FEditorStyle::Get(), "RoundButton")
+					.OnClicked(this, &SNiagaraScratchPadScriptEditor::OnApplyButtonClicked)
+					.ToolTipText(LOCTEXT("ApplyButtonToolTip", "Apply the current changes to this script.  This will update the selection stack UI and compile neccessary scripts."))
+					.IsEnabled(this, &SNiagaraScratchPadScriptEditor::GetApplyButtonIsEnabled)
+					.ForegroundColor(FSlateColor::UseForeground())
+					.ContentPadding(FMargin(0.0f))
 					.Content()
 					[
-						SNew(STextBlock)
-						.Text(LOCTEXT("RefreshButtonText", "Refresh"))
+						SNew(SHorizontalBox)
+						+ SHorizontalBox::Slot()
+						.AutoWidth()
+						.Padding(2.0f, 1.0f, 2.0f, 1.0f)
+						[
+							SNew(SImage)
+							.Image(FNiagaraEditorStyle::Get().GetBrush("NiagaraEditor.Apply.Small"))
+						]
+						+ SHorizontalBox::Slot()
+						.AutoWidth()
+						.VAlign(VAlign_Center)
+						.Padding(2.0f, 2.0f, 2.0f, 3.0f)
+						[
+							SNew(STextBlock)
+							.TextStyle(FNiagaraEditorWidgetsStyle::Get(), "NiagaraEditor.ScratchPad.SmallHeaderText")
+							.Text(LOCTEXT("ApplyButtonLabel", "Apply"))
+						]
 					]
 				]
 			]
@@ -349,33 +532,18 @@ private:
 		return ScriptViewModel->GetDisplayName();
 	}
 
-	FReply OnRefreshButtonClicked()
+	FReply OnApplyButtonClicked()
 	{
-		TArray<UNiagaraNodeFunctionCall*> FunctionCallNodesToRefresh;
-		for (TObjectIterator<UNiagaraNodeFunctionCall> It; It; ++It)
-		{
-			UNiagaraNodeFunctionCall* FunctionCallNode = *It;
-			if (FunctionCallNode->FunctionScript == ScriptViewModel->GetScript())
-			{
-				FunctionCallNodesToRefresh.Add(FunctionCallNode);
-			}
-		}
-
-		if (FunctionCallNodesToRefresh.Num())
-		{
-			for (TObjectIterator<UNiagaraStackFunctionInputCollection> It; It; ++It)
-			{
-				UNiagaraStackFunctionInputCollection* StackFunctionInputCollection = *It;
-				if (StackFunctionInputCollection->IsFinalized() == false && FunctionCallNodesToRefresh.Contains(StackFunctionInputCollection->GetInputFunctionCallNode()))
-				{
-					StackFunctionInputCollection->RefreshChildren();
-				}
-			}
-		}
+		ScriptViewModel->ApplyChanges();
 
 		return FReply::Handled();
 	}
 
+	bool GetApplyButtonIsEnabled() const
+	{
+		return ScriptViewModel->CanApplyChanges();
+	}
+
 private:
 	TSharedPtr<FNiagaraScratchPadScriptViewModel> ScriptViewModel;
 };
@@ -390,40 +558,94 @@ class SNiagaraScratchPadScriptEditorList : public SCompoundWidget
 	{
 		ViewModel = InViewModel;
 		ViewModel->OnScriptViewModelsChanged().AddSP(this, &SNiagaraScratchPadScriptEditorList::ScriptViewModelsChanged);
+		ViewModel->OnEditScriptViewModelsChanged().AddSP(this, &SNiagaraScratchPadScriptEditorList::EditScriptViewModelsChanged);
 		ViewModel->OnActiveScriptChanged().AddSP(this, &SNiagaraScratchPadScriptEditorList::ActiveScriptChanged);
 		bIsUpdatingSelection = false;
 
-		ChildSlot
-		[
-			SAssignNew(ScriptEditors, SListView<TSharedRef<FNiagaraScratchPadScriptViewModel>>)
-			.ListItemsSource(&ViewModel->GetScriptViewModels())
-			.OnGenerateRow(this, &SNiagaraScratchPadScriptEditorList::OnGenerateScriptEditorRow)
-			.OnSelectionChanged(this, &SNiagaraScratchPadScriptEditorList::OnSelectionChanged)
-		];
+		UpdateContentFromEditScriptViewModels();
 	}
 
 private:
 	void ScriptViewModelsChanged()
 	{
-		if (ScriptEditors.IsValid())
+		ScriptViewModelWidgetPairs.RemoveAll([](const FScriptViewModelWidgetPair& ScriptViewModelWidgetPair)
 		{
-			ScriptEditors->RequestListRefresh();
+			return ScriptViewModelWidgetPair.ViewModel.IsValid() == false || ScriptViewModelWidgetPair.Widget.IsValid() == false;
+		});
+	}
+
+	void EditScriptViewModelsChanged()
+	{
+		UpdateContentFromEditScriptViewModels();
+	}
+
+	TSharedRef<SWidget> FindOrAddScriptEditor(TSharedRef<FNiagaraScratchPadScriptViewModel> ScriptViewModel)
+	{
+		FScriptViewModelWidgetPair* ExistingPair = ScriptViewModelWidgetPairs.FindByPredicate([ScriptViewModel](FScriptViewModelWidgetPair& ScriptViewModelWidgetPair)
+		{ 
+			return ScriptViewModelWidgetPair.ViewModel == ScriptViewModel && ScriptViewModelWidgetPair.Widget.IsValid();
+		});
+
+		if (ExistingPair != nullptr)
+		{
+			return ExistingPair->Widget.ToSharedRef();
+		}
+		else
+		{
+			TSharedRef<SWidget> NewEditor = SNew(SNiagaraScratchPadScriptEditor, ScriptViewModel);
+			ScriptViewModelWidgetPairs.Add({ TWeakPtr<FNiagaraScratchPadScriptViewModel>(ScriptViewModel), NewEditor });
+			return NewEditor;
 		}
 	}
 
+	void UpdateContentFromEditScriptViewModels()
+	{
+		TSharedPtr<SWidget> NewContent;
+		if (ViewModel->GetEditScriptViewModels().Num() == 0)
+		{
+			NewContent = SNullWidget::NullWidget;
+			ScriptEditorList.Reset();
+		}
+		else if(ViewModel->GetEditScriptViewModels().Num() == 1)
+		{
+			NewContent = FindOrAddScriptEditor(ViewModel->GetEditScriptViewModels()[0]);
+			ScriptEditorList.Reset();
+		}
+		else 
+		{
+			if (ScriptEditorList.IsValid())
+			{
+				ScriptEditorList->RequestListRefresh();
+			}
+			else 
+			{
+				ScriptEditorList = SNew(SListView<TSharedRef<FNiagaraScratchPadScriptViewModel>>)
+					.ListItemsSource(&ViewModel->GetEditScriptViewModels())
+					.OnGenerateRow(this, &SNiagaraScratchPadScriptEditorList::OnGenerateScriptEditorRow)
+					.OnSelectionChanged(this, &SNiagaraScratchPadScriptEditorList::OnSelectionChanged);
+			}
+			NewContent = ScriptEditorList;
+		}
+
+		ChildSlot
+		[
+			NewContent.ToSharedRef()
+		];
+	}
+
 	void ActiveScriptChanged()
 	{
-		if (bIsUpdatingSelection == false)
+		if (ScriptEditorList.IsValid() && bIsUpdatingSelection == false)
 		{
 			TGuardValue<bool> UpdateGuard(bIsUpdatingSelection, true);
-			TSharedPtr<FNiagaraScratchPadScriptViewModel> ActiveScriptViewModel = ViewModel->GetViewModelForScript(ViewModel->GetActiveScript());
+			TSharedPtr<FNiagaraScratchPadScriptViewModel> ActiveScriptViewModel = ViewModel->GetActiveScriptViewModel();
 			if (ActiveScriptViewModel.IsValid())
 			{
-				ScriptEditors->SetSelection(ActiveScriptViewModel.ToSharedRef());
+				ScriptEditorList->SetSelection(ActiveScriptViewModel.ToSharedRef());
 			}
 			else
 			{
-				ScriptEditors->ClearSelection();
+				ScriptEditorList->ClearSelection();
 			}
 		}
 	}
@@ -433,9 +655,10 @@ private:
 		return SNew(STableRow<TSharedRef<FNiagaraScratchPadScriptViewModel>>, OwnerTable)
 		[
 			SNew(SVerticalResizeBox)
-			.ContentHeight(300)
+			.ContentHeight(Item, &FNiagaraScratchPadScriptViewModel::GetEditorHeight)
+			.ContentHeightChanged(Item, &FNiagaraScratchPadScriptViewModel::SetEditorHeight)
 			[
-				SNew(SNiagaraScratchPadScriptEditor, Item)
+				FindOrAddScriptEditor(Item)
 			]
 		];
 	}
@@ -446,21 +669,28 @@ private:
 		{
 			TGuardValue<bool> UpdateGuard(bIsUpdatingSelection, true);
 			TArray<TSharedRef<FNiagaraScratchPadScriptViewModel>> SelectedScripts;
-			ScriptEditors->GetSelectedItems(SelectedScripts);
+			ScriptEditorList->GetSelectedItems(SelectedScripts);
 			if (SelectedScripts.Num() == 0)
 			{
-				ViewModel->SetActiveScript(nullptr);
+				ViewModel->ResetActiveScriptViewModel();
 			}
 			else if (SelectedScripts.Num())
 			{
-				ViewModel->SetActiveScript(SelectedScripts[0]->GetScript());
+				ViewModel->SetActiveScriptViewModel(SelectedScripts[0]);
 			}
 		}
 	}
 
 private:
+	struct FScriptViewModelWidgetPair
+	{
+		TWeakPtr<FNiagaraScratchPadScriptViewModel> ViewModel;
+		TSharedPtr<SWidget> Widget;
+	};
+
 	UNiagaraScratchPadViewModel* ViewModel;
-	TSharedPtr<SListView<TSharedRef<FNiagaraScratchPadScriptViewModel>>> ScriptEditors;
+	TSharedPtr<SListView<TSharedRef<FNiagaraScratchPadScriptViewModel>>> ScriptEditorList;
+	TArray<FScriptViewModelWidgetPair> ScriptViewModelWidgetPairs;
 	bool bIsUpdatingSelection;
 };
 
@@ -498,6 +728,8 @@ void SNiagaraScratchPad::Construct(const FArguments& InArgs, UNiagaraScratchPadV
 			{
 				Layout = SNew(SSplitter)
 				.Orientation(Orient_Horizontal)
+				.PhysicalSplitterHandleSize(7.0f)
+				.HitDetectionSplitterHandleSize(7.0f)
 				+ SSplitter::Slot()
 				.Value(0.15f)
 				[
@@ -518,6 +750,8 @@ void SNiagaraScratchPad::Construct(const FArguments& InArgs, UNiagaraScratchPadV
 			{
 				Layout = SNew(SSplitter)
 				.Orientation(Orient_Horizontal)
+				.PhysicalSplitterHandleSize(7.0f)
+				.HitDetectionSplitterHandleSize(7.0f)
 				+ SSplitter::Slot()
 				.Value(0.3f)
 				[
@@ -565,6 +799,7 @@ TSharedRef<SWidget> SNiagaraScratchPad::ConstructScriptSelector()
 	return SNew(SVerticalBox)
 	+ SVerticalBox::Slot()
 	.AutoHeight()
+	.Padding(0.0f, 2.0f)
 	[
 		SNew(STextBlock)
 		.TextStyle(FNiagaraEditorWidgetsStyle::Get(), "NiagaraEditor.ScratchPad.LargeHeaderText")
