Fixed missing barriers when simulation stages are used.

We were missing transitions for the destination buffers, and source buffers were transitioned to read-write instead of just read. This change unifies the barrier handling code, so we set the same barriers regardless of stages being used or not.

Also fixed a missing barrier when clearing the neighborhood buffers in UNiagaraDataInterfaceNeighborGrid3D.

#rnx

diff --git a/Engine/Plugins/FX/Niagara/Source/Niagara/Classes/NiagaraEmitterInstanceBatcher.h b/Engine/Plugins/FX/Niagara/Source/Niagara/Classes/NiagaraEmitterInstanceBatcher.h
--- a/Engine/Plugins/FX/Niagara/Source/Niagara/Classes/NiagaraEmitterInstanceBatcher.h
+++ b/Engine/Plugins/FX/Niagara/Source/Niagara/Classes/NiagaraEmitterInstanceBatcher.h
@@ -155,9 +155,10 @@ private:
 	using FEmitterInstanceList = TArray<FNiagaraComputeInstanceData*>;
 
 	void ExecuteAll(FRHICommandList& RHICmdList, FRHIUniformBuffer* ViewUniformBuffer, ETickStage TickStage);
-	void ResizeBuffersAndGatherResources(FOverlappableTicks& OverlappableTick, FRHICommandList& RHICmdList, FNiagaraBufferArray& ReadBuffers, FNiagaraBufferArray& WriteBuffers, FNiagaraBufferArray& OutputGraphicsBuffers, FEmitterInstanceList& InstancesWithPersistentIDs);
-	void DispatchAllOnCompute(FOverlappableTicks& OverlappableTick, FRHICommandList& RHICmdList, FRHIUniformBuffer* ViewUniformBuffer, FNiagaraBufferArray& ReadBuffers, FNiagaraBufferArray& WriteBuffers);
-	void DispatchMultipleStages(const FNiagaraGPUSystemTick& Tick, FNiagaraComputeInstanceData *Instance, FRHICommandList &RHICmdList, FRHIUniformBuffer* ViewUniformBuffer, const FNiagaraShaderRef& ComputeShader);
+	void ResizeBuffersAndGatherResources(FOverlappableTicks& OverlappableTick, FRHICommandList& RHICmdList, FNiagaraBufferArray& OutputGraphicsBuffers, FEmitterInstanceList& InstancesWithPersistentIDs);
+	void TransitionBuffers(FRHICommandList& RHICmdList, const FNiagaraShaderRef& ComputeShader, FNiagaraComputeInstanceData* Instance, uint32 SimulationStageIndex, const FNiagaraDataBuffer*& LastSource, bool& bFreeIDTableTransitioned);
+	void DispatchAllOnCompute(FOverlappableTicks& OverlappableTick, FRHICommandList& RHICmdList, FRHIUniformBuffer* ViewUniformBuffer);
+	void DispatchMultipleStages(const FNiagaraGPUSystemTick& Tick, FNiagaraComputeInstanceData* Instance, FRHICommandList& RHICmdList, FRHIUniformBuffer* ViewUniformBuffer, const FNiagaraShaderRef& ComputeShader);
 
 	bool ShouldTickForStage(const FNiagaraGPUSystemTick& Tick, ETickStage TickStage) const;
 
diff --git a/Engine/Plugins/FX/Niagara/Source/Niagara/Private/NiagaraDataInterfaceNeighborGrid3D.cpp b/Engine/Plugins/FX/Niagara/Source/Niagara/Private/NiagaraDataInterfaceNeighborGrid3D.cpp
--- a/Engine/Plugins/FX/Niagara/Source/Niagara/Private/NiagaraDataInterfaceNeighborGrid3D.cpp
+++ b/Engine/Plugins/FX/Niagara/Source/Niagara/Private/NiagaraDataInterfaceNeighborGrid3D.cpp
@@ -513,6 +513,9 @@ void FNiagaraDataInterfaceProxyNeighborGrid3D::PreStage(FRHICommandList& RHICmdL
 
 		SCOPED_DRAW_EVENT(RHICmdList, NiagaraNeighborGrid3DClearNeighborInfo);
 		ERHIFeatureLevel::Type FeatureLevel = Context.Batcher->GetFeatureLevel();
+
+		FRHIUnorderedAccessView* TransitionBuffers[] = { ProxyData->NeighborhoodBuffer.UAV, ProxyData->NeighborhoodCountBuffer.UAV };
+		RHICmdList.TransitionResources(EResourceTransitionAccess::EWritable, EResourceTransitionPipeline::EComputeToCompute, TransitionBuffers, UE_ARRAY_COUNT(TransitionBuffers));
 		NiagaraFillGPUIntBuffer(RHICmdList, FeatureLevel, ProxyData->NeighborhoodBuffer, -1);
 		NiagaraFillGPUIntBuffer(RHICmdList, FeatureLevel, ProxyData->NeighborhoodCountBuffer, 0);
 	}
diff --git a/Engine/Plugins/FX/Niagara/Source/Niagara/Private/NiagaraEmitterInstanceBatcher.cpp b/Engine/Plugins/FX/Niagara/Source/Niagara/Private/NiagaraEmitterInstanceBatcher.cpp
--- a/Engine/Plugins/FX/Niagara/Source/Niagara/Private/NiagaraEmitterInstanceBatcher.cpp
+++ b/Engine/Plugins/FX/Niagara/Source/Niagara/Private/NiagaraEmitterInstanceBatcher.cpp
@@ -276,11 +276,65 @@ void NiagaraEmitterInstanceBatcher::PostStageInterface(const FNiagaraGPUSystemTi
 	}
 }
 
-void NiagaraEmitterInstanceBatcher::DispatchMultipleStages(const FNiagaraGPUSystemTick& Tick, FNiagaraComputeInstanceData *Instance, FRHICommandList &RHICmdList, FRHIUniformBuffer* ViewUniformBuffer, const FNiagaraShaderRef& ComputeShader)
+void NiagaraEmitterInstanceBatcher::TransitionBuffers(FRHICommandList& RHICmdList, const FNiagaraShaderRef& ComputeShader, FNiagaraComputeInstanceData* Instance, uint32 SimulationStageIndex, const FNiagaraDataBuffer*& LastSource, bool& bFreeIDTableTransitioned)
 {
 	FNiagaraComputeExecutionContext* Context = Instance->Context;
+	const bool bRequiresPersistentIDs = Context->MainDataSet->RequiresPersistentIDs();
+	FNiagaraSimStageData& SimStageData = Instance->SimStageData[SimulationStageIndex];
+
+	FNiagaraBufferArray ReadBuffers, WriteBuffers;
+
+	if (bRequiresPersistentIDs && !bFreeIDTableTransitioned)
+	{
+		// There's one free ID buffer at the emitter level, so we only need to make it writable once.
+		ReadBuffers.Add(Context->MainDataSet->GetGPUFreeIDs().UAV);
+		bFreeIDTableTransitioned = true;
+	}
+
+	if (LastSource != SimStageData.Source)
+	{
+		// If we're reading from a different buffer set than the previous stage, make the inputs readable.
+		LastSource = SimStageData.Source;
+		if (ComputeShader->FloatInputBufferParam.IsBound())
+		{
+			ReadBuffers.Add(SimStageData.Source->GetGPUBufferFloat().UAV);
+		}
+		if (ComputeShader->IntInputBufferParam.IsBound())
+		{
+			ReadBuffers.Add(SimStageData.Source->GetGPUBufferInt().UAV);
+		}
+	}
+
+	// Always transition output buffers.
+	if (ComputeShader->FloatOutputBufferParam.IsBound())
+	{
+		WriteBuffers.Add(SimStageData.Destination->GetGPUBufferFloat().UAV);
+	}
+	if (ComputeShader->IntOutputBufferParam.IsBound())
+	{
+		WriteBuffers.Add(SimStageData.Destination->GetGPUBufferInt().UAV);
+	}
+	if (bRequiresPersistentIDs)
+	{
+		// TODO: clear the ID to Index table for stages > 0. We may be able to
+		// skip this in some cases.
+		WriteBuffers.Add(SimStageData.Destination->GetGPUIDToIndexTable().UAV);
+	}
 
-	if (!ResetDataInterfaces(Tick, Instance, RHICmdList, ComputeShader)) return;
+	RHICmdList.TransitionResources(EResourceTransitionAccess::EReadable, EResourceTransitionPipeline::EComputeToCompute, ReadBuffers.GetData(), ReadBuffers.Num());
+	RHICmdList.TransitionResources(EResourceTransitionAccess::ERWBarrier, EResourceTransitionPipeline::EComputeToCompute, WriteBuffers.GetData(), WriteBuffers.Num());
+}
+
+void NiagaraEmitterInstanceBatcher::DispatchMultipleStages(const FNiagaraGPUSystemTick& Tick, FNiagaraComputeInstanceData* Instance, FRHICommandList& RHICmdList, FRHIUniformBuffer* ViewUniformBuffer, const FNiagaraShaderRef& ComputeShader)
+{
+	if (!ResetDataInterfaces(Tick, Instance, RHICmdList, ComputeShader))
+	{
+		return;
+	}
+
+	FNiagaraComputeExecutionContext* Context = Instance->Context;
+	const FNiagaraDataBuffer* LastSource = nullptr;
+	bool bFreeIDTableTransitioned = false;
 
 	if (Tick.NumInstancesWithSimStages > 0)
 	{
@@ -288,7 +342,6 @@ void NiagaraEmitterInstanceBatcher::DispatchMultipleStages(const FNiagaraGPUSyst
 		
 		const uint32 NumStages = Instance->Context->MaxUpdateIterations;
 		const uint32 DefaultSimulationStageIndex = Instance->Context->DefaultSimulationStageIndex;
-		bool bTransitionCurrentBuffer = false;
 
 		for (uint32 SimulationStageIndex = 0; SimulationStageIndex < NumStages; ++SimulationStageIndex)
 		{
@@ -307,50 +360,32 @@ void NiagaraEmitterInstanceBatcher::DispatchMultipleStages(const FNiagaraGPUSyst
 
 			PreStageInterface(Tick, Instance, RHICmdList, ComputeShader, SimulationStageIndex);
 
-			// If we are reading from current data we need to transition the resource if it was previously written to
-			if (bTransitionCurrentBuffer && (ComputeShader->FloatInputBufferParam.IsBound() || ComputeShader->IntInputBufferParam.IsBound()))
-			{
-				bTransitionCurrentBuffer = false;
-
-				TArray<FRHIUnorderedAccessView*, TInlineAllocator<2>> Resources;
-				if (Instance->SimStageData[SimulationStageIndex].Source->GetGPUBufferFloat().UAV.IsValid())
-				{
-					Resources.Add(Instance->SimStageData[SimulationStageIndex].Source->GetGPUBufferFloat().UAV);
-				}
-				if (Instance->SimStageData[SimulationStageIndex].Source->GetGPUBufferInt().UAV.IsValid())
-				{
-					Resources.Add(Instance->SimStageData[SimulationStageIndex].Source->GetGPUBufferInt().UAV);
-				}
-				if (Resources.Num() > 0)
-				{
-					RHICmdList.TransitionResources(EResourceTransitionAccess::ERWBarrier, EResourceTransitionPipeline::EComputeToCompute, Resources.GetData(), Resources.Num());
-				}
-			}
+			TransitionBuffers(RHICmdList, ComputeShader, Instance, SimulationStageIndex, LastSource, bFreeIDTableTransitioned);
 
 			if (!IterationInterface)
 			{
 				Run(Tick, Instance, 0, Instance->SimStageData[SimulationStageIndex].Destination->GetNumInstances(), ComputeShader, RHICmdList, ViewUniformBuffer, Instance->SpawnInfo, false, DefaultSimulationStageIndex, SimulationStageIndex,  nullptr, HasRunParticleStage);
 				HasRunParticleStage = true;
-				bTransitionCurrentBuffer = true;
 			}
 			else
 			{
 				// run with correct number of instances.  This will make curr data junk or empty
 				Run(Tick, Instance, 0, IterationInterface->ElementCount, ComputeShader, RHICmdList, ViewUniformBuffer, Instance->SpawnInfo, false, DefaultSimulationStageIndex, SimulationStageIndex, IterationInterface);
 			}
+
 			PostStageInterface(Tick, Instance, RHICmdList, ComputeShader, SimulationStageIndex);
 		}
-
 	}
 	else
 	{
 		// run shader, sim and spawn in a single dispatch
 		check(Instance->SimStageData.Num() > 0);
+		TransitionBuffers(RHICmdList, ComputeShader, Instance, 0, LastSource, bFreeIDTableTransitioned);
 		Run(Tick, Instance, 0, Instance->SimStageData[0].Destination->GetNumInstances(), ComputeShader, RHICmdList, ViewUniformBuffer, Instance->SpawnInfo);
 	}
 }
 
-void NiagaraEmitterInstanceBatcher::ResizeBuffersAndGatherResources(FOverlappableTicks& OverlappableTick, FRHICommandList& RHICmdList, FNiagaraBufferArray& ReadBuffers, FNiagaraBufferArray& WriteBuffers, FNiagaraBufferArray& OutputGraphicsBuffers, FEmitterInstanceList& InstancesWithPersistentIDs)
+void NiagaraEmitterInstanceBatcher::ResizeBuffersAndGatherResources(FOverlappableTicks& OverlappableTick, FRHICommandList& RHICmdList, FNiagaraBufferArray& OutputGraphicsBuffers, FEmitterInstanceList& InstancesWithPersistentIDs)
 {
 	SCOPE_CYCLE_COUNTER(STAT_NiagaraGPUDispatchSetup_RT);
 
@@ -403,7 +438,6 @@ void NiagaraEmitterInstanceBatcher::ResizeBuffersAndGatherResources(FOverlappabl
 			if (bRequiresPersistentIDs)
 			{
 				Context->MainDataSet->AllocateGPUFreeIDs(AllocatedInstances + 1, RHICmdList, FeatureLevel, Context->GetDebugSimName());
-				ReadBuffers.Add(Context->MainDataSet->GetGPUFreeIDs().UAV);
 				InstancesWithPersistentIDs.Add(&Instance);
 			}
 
@@ -416,30 +450,7 @@ void NiagaraEmitterInstanceBatcher::ResizeBuffersAndGatherResources(FOverlappabl
 			}
 			Instance.SimStageData[0].DestinationCountOffset = Instance.SimStageData[0].Destination->GetGPUInstanceCountBufferOffset();
 
-			if ( Shader->FloatInputBufferParam.IsBound() )
-			{
-				ReadBuffers.Add(CurrentData->GetGPUBufferFloat().UAV);
-			}
-			if ( Shader->IntInputBufferParam.IsBound() )
-			{
-				ReadBuffers.Add(CurrentData->GetGPUBufferInt().UAV);
-			}
-
-			if ( Shader->FloatOutputBufferParam.IsBound() )
-			{
-				WriteBuffers.Add(DestinationData->GetGPUBufferFloat().UAV);
-			}
-			if ( Shader->IntOutputBufferParam.IsBound() )
-			{
-				WriteBuffers.Add(DestinationData->GetGPUBufferInt().UAV);
-			}
-
 			//UE_LOG(LogScript, Warning, TEXT("ResizeBuffersAndGatherResources [%d][%d] Run ReqInst: %d Cur: %p Dest: %p "), Index, 0, RequiredInstances, Instance.SimStageData[0].Source, Instance.SimStageData[0].Destination);
-			if (bRequiresPersistentIDs)
-			{
-				WriteBuffers.Add(DestinationData->GetGPUIDToIndexTable().UAV);
-			}
-
 
 			Context->MainDataSet->EndSimulate();
 
@@ -553,7 +564,7 @@ void NiagaraEmitterInstanceBatcher::ResizeBuffersAndGatherResources(FOverlappabl
 	}
 }
 
-void NiagaraEmitterInstanceBatcher::DispatchAllOnCompute(FOverlappableTicks& OverlappableTick, FRHICommandList& RHICmdList, FRHIUniformBuffer* ViewUniformBuffer, FNiagaraBufferArray& ReadBuffers, FNiagaraBufferArray& WriteBuffers)
+void NiagaraEmitterInstanceBatcher::DispatchAllOnCompute(FOverlappableTicks& OverlappableTick, FRHICommandList& RHICmdList, FRHIUniformBuffer* ViewUniformBuffer)
 {
 	FRHICommandListImmediate& RHICmdListImmediate = FRHICommandListExecutor::GetImmediateCommandList();
 
@@ -581,10 +592,6 @@ void NiagaraEmitterInstanceBatcher::DispatchAllOnCompute(FOverlappableTicks& Ove
 	}
 #endif // WITH_EDITORONLY_DATA
 
-	//
-	RHICmdList.TransitionResources(EResourceTransitionAccess::EReadable, EResourceTransitionPipeline::EComputeToCompute, ReadBuffers.GetData(), ReadBuffers.Num());
-	RHICmdList.TransitionResources(EResourceTransitionAccess::ERWBarrier, EResourceTransitionPipeline::EComputeToCompute, WriteBuffers.GetData(), WriteBuffers.Num());
-
 	for (FNiagaraGPUSystemTick* Tick : OverlappableTick)
 	{
 		uint32 DispatchCount = Tick->Count;
@@ -701,6 +708,7 @@ void NiagaraEmitterInstanceBatcher::UpdateFreeIDBuffers(FRHICommandList& RHICmdL
 	FNiagaraBufferArray ReadBuffers, WriteBuffers;
 	for(FNiagaraComputeInstanceData* Instance : Instances)
 	{
+		// TODO: this is incorrect if we have simulation stages which can kill particles.
 		ReadBuffers.Add(Instance->SimStageData[0].Destination->GetGPUIDToIndexTable().UAV);
 		WriteBuffers.Add(Instance->Context->MainDataSet->GetGPUFreeIDs().UAV);
 	}
@@ -815,23 +823,21 @@ void NiagaraEmitterInstanceBatcher::ExecuteAll(FRHICommandList& RHICmdList, FRHI
 	RHICmdList.BeginUAVOverlap();
 
 	FEmitterInstanceList InstancesWithPersistentIDs;
-	FNiagaraBufferArray ReadBuffers, WriteBuffers, OutputGraphicsBuffers;
+	FNiagaraBufferArray OutputGraphicsBuffers;
 
 	for (int32 SimPassIdx = 0; SimPassIdx < SimPasses.Num(); ++SimPassIdx)
 	{
 		FOverlappableTicks& SimPass = SimPasses[SimPassIdx];
-		ReadBuffers.SetNum(0, false);
-		WriteBuffers.SetNum(0, false);
 		InstancesWithPersistentIDs.SetNum(0, false);
 
 		// This initial pass gathers all the buffers that are read from and written to so we can do batch resource transitions.
 		// It also ensures the GPU buffers are large enough to hold everything.
-		ResizeBuffersAndGatherResources(SimPass, RHICmdList, ReadBuffers, WriteBuffers, OutputGraphicsBuffers, InstancesWithPersistentIDs);
+		ResizeBuffersAndGatherResources(SimPass, RHICmdList, OutputGraphicsBuffers, InstancesWithPersistentIDs);
 
 		{
 			SCOPED_DRAW_EVENT(RHICmdList, NiagaraGPUSimulation);
 			SCOPED_GPU_STAT(RHICmdList, NiagaraGPUSimulation);
-			DispatchAllOnCompute(SimPass, RHICmdList, ViewUniformBuffer, ReadBuffers, WriteBuffers);
+			DispatchAllOnCompute(SimPass, RHICmdList, ViewUniformBuffer);
 		}
 
 		if (InstancesWithPersistentIDs.Num() == 0)
