Improved system counts stats.
Now tracking solo systems properly and showing them separately to help ID solo systems we should remove.

diff --git a/Engine/Plugins/FX/Niagara/Source/Niagara/Classes/NiagaraSystem.h b/Engine/Plugins/FX/Niagara/Source/Niagara/Classes/NiagaraSystem.h
--- a/Engine/Plugins/FX/Niagara/Source/Niagara/Classes/NiagaraSystem.h
+++ b/Engine/Plugins/FX/Niagara/Source/Niagara/Classes/NiagaraSystem.h
@@ -373,7 +373,7 @@ public:
 	uint32 bFixedBounds : 1;
 
 	TStatId GetStatID(bool bGameThread, bool bConcurrent)const;
-	void SetInstanceCountStat(int32 NumInstances)const;
+	void AddToInstanceCountStat(int32 NumInstances, bool bSolo)const;
 
 	UPROPERTY(EditAnywhere, Category = "Script Fast Path")
 	ENiagaraFastPathMode FastPathMode;
@@ -507,6 +507,7 @@ protected:
 	mutable TStatId StatID_RT_CNC;
 
 	mutable TStatId StatID_InstanceCount;
+	mutable TStatId StatID_InstanceCountSolo;
 #endif
 
 	FNiagaraSystemScalabilitySettings CurrentScalabilitySettings;
diff --git a/Engine/Plugins/FX/Niagara/Source/Niagara/Private/NiagaraComponent.cpp b/Engine/Plugins/FX/Niagara/Source/Niagara/Private/NiagaraComponent.cpp
--- a/Engine/Plugins/FX/Niagara/Source/Niagara/Private/NiagaraComponent.cpp
+++ b/Engine/Plugins/FX/Niagara/Source/Niagara/Private/NiagaraComponent.cpp
@@ -587,6 +587,10 @@ void UNiagaraComponent::TickComponent(float DeltaSeconds, enum ELevelTick TickTy
 	check(SystemInstance->IsSolo());
 	if (IsActive() && SystemInstance.Get() && !SystemInstance->IsComplete())
 	{
+		Asset->AddToInstanceCountStat(1, true);
+		INC_DWORD_STAT_BY(STAT_TotalNiagaraSystemInstances, 1);
+		INC_DWORD_STAT_BY(STAT_TotalNiagaraSystemInstancesSolo, 1);
+
 		// If the interfaces have changed in a meaningful way, we need to potentially rebind and update the values.
 		if (OverrideParameters.GetInterfacesDirty())
 		{
diff --git a/Engine/Plugins/FX/Niagara/Source/Niagara/Private/NiagaraStats.h b/Engine/Plugins/FX/Niagara/Source/Niagara/Private/NiagaraStats.h
--- a/Engine/Plugins/FX/Niagara/Source/Niagara/Private/NiagaraStats.h
+++ b/Engine/Plugins/FX/Niagara/Source/Niagara/Private/NiagaraStats.h
@@ -28,4 +28,7 @@ DECLARE_CYCLE_STAT(TEXT("RT Concurrent Total"), STAT_NiagaraOverview_RT_CNC, STA
 DECLARE_STATS_GROUP(TEXT("Niagara Systems"), STATGROUP_NiagaraSystems, STATCAT_Advanced);
 DECLARE_STATS_GROUP(TEXT("Niagara Emitters"), STATGROUP_NiagaraEmitters, STATCAT_Advanced);
 
-DECLARE_STATS_GROUP(TEXT("Niagara System Counts"), STATGROUP_NiagaraSystemCounts, STATCAT_Advanced);
\ No newline at end of file
+DECLARE_STATS_GROUP(TEXT("Niagara System Counts"), STATGROUP_NiagaraSystemCounts, STATCAT_Advanced);
+
+DECLARE_DWORD_COUNTER_STAT(TEXT("Total System Instances"), STAT_TotalNiagaraSystemInstances, STATGROUP_NiagaraSystemCounts);
+DECLARE_DWORD_COUNTER_STAT(TEXT("Total Solo System Instances"), STAT_TotalNiagaraSystemInstancesSolo, STATGROUP_NiagaraSystemCounts);
\ No newline at end of file
diff --git a/Engine/Plugins/FX/Niagara/Source/Niagara/Private/NiagaraSystem.cpp b/Engine/Plugins/FX/Niagara/Source/Niagara/Private/NiagaraSystem.cpp
--- a/Engine/Plugins/FX/Niagara/Source/Niagara/Private/NiagaraSystem.cpp
+++ b/Engine/Plugins/FX/Niagara/Source/Niagara/Private/NiagaraSystem.cpp
@@ -1438,7 +1438,7 @@ TStatId UNiagaraSystem::GetStatID(bool bGameThread, bool bConcurrent)const
 	return TStatId();
 }
 
-void UNiagaraSystem::SetInstanceCountStat(int32 NumInstances)const
+void UNiagaraSystem::AddToInstanceCountStat(int32 NumInstances, bool bSolo)const
 {
 #if STATS
 	if (!StatID_GT.IsValidStat())
@@ -1448,8 +1448,16 @@ void UNiagaraSystem::SetInstanceCountStat(int32 NumInstances)const
 
 	if (FThreadStats::IsCollectingData())
 	{
-		FThreadStats::AddMessage(StatID_InstanceCount.GetName(), EStatOperation::Set, int64(NumInstances));
-		TRACE_STAT_SET(StatID_InstanceCount.GetName(), int64(Value));
+		if (bSolo)
+		{
+			FThreadStats::AddMessage(StatID_InstanceCountSolo.GetName(), EStatOperation::Add, int64(NumInstances));
+			TRACE_STAT_ADD(StatID_InstanceCount.GetName(), int64(Value));
+		}
+		else
+		{
+			FThreadStats::AddMessage(StatID_InstanceCount.GetName(), EStatOperation::Add, int64(NumInstances));
+			TRACE_STAT_ADD(StatID_InstanceCount.GetName(), int64(Value));
+		}
 	}
 #endif
 }
@@ -1463,6 +1471,8 @@ void UNiagaraSystem::GenerateStatID()const
 	StatID_RT_CNC = FDynamicStats::CreateStatId<FStatGroup_STATGROUP_NiagaraSystems>(GetPathName() + TEXT(" [RT_CNC]"));
 
 	StatID_InstanceCount = FDynamicStats::CreateStatIdInt64<FStatGroup_STATGROUP_NiagaraSystemCounts>(GetPathName());
+	StatID_InstanceCountSolo = FDynamicStats::CreateStatIdInt64<FStatGroup_STATGROUP_NiagaraSystemCounts>(GetPathName() + TEXT(" [SOLO]"));
+
 #endif
 }
 
diff --git a/Engine/Plugins/FX/Niagara/Source/Niagara/Private/NiagaraSystemSimulation.cpp b/Engine/Plugins/FX/Niagara/Source/Niagara/Private/NiagaraSystemSimulation.cpp
--- a/Engine/Plugins/FX/Niagara/Source/Niagara/Private/NiagaraSystemSimulation.cpp
+++ b/Engine/Plugins/FX/Niagara/Source/Niagara/Private/NiagaraSystemSimulation.cpp
@@ -44,8 +44,6 @@ DECLARE_CYCLE_STAT(TEXT("ForcedWaitForAsync"), STAT_NiagaraSystemSim_ForceWaitFo
 DECLARE_CYCLE_STAT(TEXT("ForcedWait Fake Stall"), STAT_NiagaraSystemSim_ForceWaitFakeStall, STATGROUP_Niagara);
 
 
-DECLARE_DWORD_COUNTER_STAT(TEXT("Total System Instances"), STAT_TotalNiagaraSystemInstances, STATGROUP_NiagaraSystemCounts);
-
 static int32 GbDumpSystemData = 0;
 static FAutoConsoleVariableRef CVarNiagaraDumpSystemData(
 	TEXT("fx.DumpSystemData"),
@@ -795,8 +793,12 @@ void FNiagaraSystemSimulation::Tick_GameThread(float DeltaSeconds, const FGraphE
 	static const auto EffectsQualityCVar = IConsoleManager::Get().FindConsoleVariable(TEXT("sg.EffectsQuality"));
 	FNiagaraSystemSimulationTickContext Context(this, SystemInstances, MainDataSet, DeltaSeconds, SpawnNum, EffectsQualityCVar->GetInt(), MyCompletionGraphEvent);
 
-	System->SetInstanceCountStat(SystemInstances.Num());
-	INC_DWORD_STAT_BY(STAT_TotalNiagaraSystemInstances, SystemInstances.Num());
+	//Solo systems add their counts in their component tick.
+	if (GetIsSolo() == false)
+	{
+		System->AddToInstanceCountStat(SystemInstances.Num(), false);
+		INC_DWORD_STAT_BY(STAT_TotalNiagaraSystemInstances, SystemInstances.Num());
+	}
 
 	//Now kick of the concurrent tick.
 	if (Context.bTickAsync)
