Add postload pass for NiagaraGraph to fixup null script variables in existing assets.

diff --git a/Engine/Plugins/FX/Niagara/Source/NiagaraCore/Private/NiagaraCustomVersion.cpp b/Engine/Plugins/FX/Niagara/Source/NiagaraCore/Private/NiagaraCustomVersion.cpp
--- a/Engine/Plugins/FX/Niagara/Source/NiagaraCore/Private/NiagaraCustomVersion.cpp
+++ b/Engine/Plugins/FX/Niagara/Source/NiagaraCore/Private/NiagaraCustomVersion.cpp
@@ -8,4 +8,4 @@ const FGuid FNiagaraCustomVersion::GUID(0xFCF57AFA, 0x50764283, 0xB9A9E658, 0xFF
 // Register the custom version with core
 FCustomVersionRegistration GRegisterNiagaraCustomVersion(FNiagaraCustomVersion::GUID, FNiagaraCustomVersion::LatestVersion, TEXT("NiagaraVer"));
 
-const FGuid FNiagaraCustomVersion::LatestScriptCompileVersion(0x318C2E9D, 0xB5794FBD, 0xABE176B1, 0xC31904C9);
+const FGuid FNiagaraCustomVersion::LatestScriptCompileVersion(0xE599EF89, 0x9D574012, 0x99E29018, 0x4682F9EC);
diff --git a/Engine/Plugins/FX/Niagara/Source/NiagaraCore/Public/NiagaraCustomVersion.h b/Engine/Plugins/FX/Niagara/Source/NiagaraCore/Public/NiagaraCustomVersion.h
--- a/Engine/Plugins/FX/Niagara/Source/NiagaraCore/Public/NiagaraCustomVersion.h
+++ b/Engine/Plugins/FX/Niagara/Source/NiagaraCore/Public/NiagaraCustomVersion.h
@@ -144,6 +144,8 @@ struct FNiagaraCustomVersion
 
 		PrecompileNamespaceFixup, // Promote parameters used across script executions to the Dataset, and Demote unused parameters.
 
+		FixNullScriptVariables, // Postload fixup in UNiagaraGraph to fixup VariableToScriptVariable map entries being null. 
+
 		// DO NOT ADD A NEW VERSION UNLESS YOU HAVE TALKED TO THE NIAGARA LEAD. Mismanagement of these versions can lead to data loss if it is adjusted in multiple streams simultaneously.
 		// -----<new versions can be added above this line>  -------------------------------------------------
 		VersionPlusOne,
diff --git a/Engine/Plugins/FX/Niagara/Source/NiagaraEditor/Private/NiagaraGraph.cpp b/Engine/Plugins/FX/Niagara/Source/NiagaraEditor/Private/NiagaraGraph.cpp
--- a/Engine/Plugins/FX/Niagara/Source/NiagaraEditor/Private/NiagaraGraph.cpp
+++ b/Engine/Plugins/FX/Niagara/Source/NiagaraEditor/Private/NiagaraGraph.cpp
@@ -137,6 +137,24 @@ void UNiagaraGraph::PostLoad()
 {
 	Super::PostLoad();
 
+	const int32 NiagaraVer = GetLinkerCustomVersion(FNiagaraCustomVersion::GUID);
+
+	if (NiagaraVer < FNiagaraCustomVersion::FixNullScriptVariables)
+	{
+		for (auto It = VariableToScriptVariable.CreateIterator(); It; ++It)
+		{
+			FNiagaraVariable Var = It.Key();
+			UNiagaraScriptVariable* ScriptVar = It.Value();
+
+			if (ScriptVar == nullptr)
+			{
+				ScriptVar = NewObject<UNiagaraScriptVariable>(const_cast<UNiagaraGraph*>(this));
+				GenerateMetaDataForScriptVariable(ScriptVar);
+				UE_LOG(LogNiagaraEditor, Display, TEXT("Fixed null UNiagaraScriptVariable | variable %s | asset path %s"), *Var.GetName().ToString(), *GetPathName());
+			}
+		}
+	}
+
 	for (FNiagaraGraphScriptUsageInfo& CachedUsageInfoItem : CachedUsageInfo)
 	{
 		CachedUsageInfoItem.PostLoad(this);
@@ -219,10 +237,7 @@ void UNiagaraGraph::PostLoad()
 	{
 		SetFlags(RF_Transactional);
 	}
-
-	// Migrate input condition metadata
-	const int32 NiagaraVer = GetLinkerCustomVersion(FNiagaraCustomVersion::GUID);
-
+	
 	FString FullPathName = GetPathName();
 	FString PathName;
 	int ColonPos;
@@ -234,7 +249,7 @@ void UNiagaraGraph::PostLoad()
 	}
 	// UE_LOG(LogNiagaraEditor, Log, TEXT("PostLoad %s"), *FullPathName);
 
-
+	// Migrate input condition metadata
 	if (NiagaraVer < FNiagaraCustomVersion::MetaDataAndParametersUpdate)
 	{
 		int NumMigrated = 0;
