Niagara - Fix an ensure where property row generators, which are editor tickable objects, were destructed during the system view model's tick.

diff --git a/Engine/Plugins/FX/Niagara/Source/NiagaraEditor/Private/NiagaraEditorModule.cpp b/Engine/Plugins/FX/Niagara/Source/NiagaraEditor/Private/NiagaraEditorModule.cpp
--- a/Engine/Plugins/FX/Niagara/Source/NiagaraEditor/Private/NiagaraEditorModule.cpp
+++ b/Engine/Plugins/FX/Niagara/Source/NiagaraEditor/Private/NiagaraEditorModule.cpp
@@ -129,6 +129,7 @@
 #include "DeviceProfiles/DeviceProfileManager.h"
 #include "Interfaces/ITargetPlatform.h"
 #include "DeviceProfiles/DeviceProfile.h"
+#include "Containers/Ticker.h"
 
 IMPLEMENT_MODULE( FNiagaraEditorModule, NiagaraEditor );
 
@@ -1422,6 +1423,21 @@ void FNiagaraEditorModule::ReinitializeStyle()
 	FNiagaraEditorStyle::Initialize();
 }
 
+void FNiagaraEditorModule::EnqueueObjectForDeferredDestructionInternal(FDeferredDestructionContainerBase* InObjectToDestruct)
+{
+	if (EnqueuedForDeferredDestruction.Num() == 0)
+	{
+		FTicker::GetCoreTicker().AddTicker(FTickerDelegate::CreateRaw(this, &FNiagaraEditorModule::DeferredDestructObjects));
+	}
+	EnqueuedForDeferredDestruction.Add(MakeShareable<FDeferredDestructionContainerBase>(InObjectToDestruct));
+}
+
+bool FNiagaraEditorModule::DeferredDestructObjects(float InDeltaTime)
+{
+	EnqueuedForDeferredDestruction.Empty();
+	return false;
+}
+
 PRAGMA_ENABLE_OPTIMIZATION
 
 #undef LOCTEXT_NAMESPACE
diff --git a/Engine/Plugins/FX/Niagara/Source/NiagaraEditor/Private/ViewModels/Stack/NiagaraStackObject.cpp b/Engine/Plugins/FX/Niagara/Source/NiagaraEditor/Private/ViewModels/Stack/NiagaraStackObject.cpp
--- a/Engine/Plugins/FX/Niagara/Source/NiagaraEditor/Private/ViewModels/Stack/NiagaraStackObject.cpp
+++ b/Engine/Plugins/FX/Niagara/Source/NiagaraEditor/Private/ViewModels/Stack/NiagaraStackObject.cpp
@@ -3,6 +3,7 @@
 #include "ViewModels/Stack/NiagaraStackObject.h"
 #include "ViewModels/Stack/NiagaraStackPropertyRow.h"
 #include "NiagaraNode.h"
+#include "NiagaraEditorModule.h"
 
 #include "Modules/ModuleManager.h"
 #include "IPropertyRowGenerator.h"
@@ -65,6 +66,11 @@ void UNiagaraStackObject::FinalizeInternal()
 	if (PropertyRowGenerator.IsValid())
 	{
 		PropertyRowGenerator->OnRowsRefreshed().RemoveAll(this);
+		PropertyRowGenerator->SetObjects(TArray<UObject*>());
+
+		// Enqueue the row generator for destruction because stack entries might be finalized during the system view model tick
+		// and you can't destruct tickables while other tickables are being ticked.
+		FNiagaraEditorModule::Get().EnqueueObjectForDeferredDestruction(PropertyRowGenerator.ToSharedRef());
 		PropertyRowGenerator.Reset();
 	}
 	Super::FinalizeInternal();
diff --git a/Engine/Plugins/FX/Niagara/Source/NiagaraEditor/Public/NiagaraEditorModule.h b/Engine/Plugins/FX/Niagara/Source/NiagaraEditor/Public/NiagaraEditorModule.h
--- a/Engine/Plugins/FX/Niagara/Source/NiagaraEditor/Public/NiagaraEditorModule.h
+++ b/Engine/Plugins/FX/Niagara/Source/NiagaraEditor/Public/NiagaraEditorModule.h
@@ -121,7 +121,39 @@ public:
 
 	FNiagaraClipboard& GetClipboard() const;
 
+	template<typename T>
+	void EnqueueObjectForDeferredDestruction(TSharedRef<T> InObjectToDestruct)
+	{
+		TDeferredDestructionContainer<T>* ObjectInContainer = new TDeferredDestructionContainer<T>(InObjectToDestruct);
+		EnqueueObjectForDeferredDestructionInternal(ObjectInContainer);
+	}
+
 private:
+	class FDeferredDestructionContainerBase
+	{
+	public:
+		virtual ~FDeferredDestructionContainerBase()
+		{
+		}
+	};
+
+	template<typename T>
+	class TDeferredDestructionContainer : public FDeferredDestructionContainerBase
+	{
+	public:
+		TDeferredDestructionContainer(TSharedRef<const T> InObjectToDestruct)
+			: ObjectToDestuct(InObjectToDestruct)
+		{
+		}
+
+		virtual ~TDeferredDestructionContainer()
+		{
+			ObjectToDestuct.Reset();
+		}
+
+		TSharedPtr<const T> ObjectToDestuct;
+	};
+
 	void RegisterAssetTypeAction(IAssetTools& AssetTools, TSharedRef<IAssetTypeActions> Action);
 	void OnNiagaraSettingsChangedEvent(const FString& PropertyName, const UNiagaraSettings* Settings);
 	void OnPreGarbageCollection();
@@ -140,6 +172,10 @@ private:
 	void TestCompileScriptFromConsole(const TArray<FString>& Arguments);
 	void ReinitializeStyle();
 
+	void EnqueueObjectForDeferredDestructionInternal(FDeferredDestructionContainerBase* InObjectToDestruct);
+
+	bool DeferredDestructObjects(float InDeltaTime);
+
 private:
 	TSharedPtr<FExtensibilityManager> MenuExtensibilityManager;
 	TSharedPtr<FExtensibilityManager> ToolBarExtensibilityManager;
@@ -196,4 +232,6 @@ private:
 	IConsoleCommand* ReinitializeStyleCommand;
 
 	TMap<int32, TSharedPtr<FHlslNiagaraCompiler>> ActiveCompilations;
+
+	TArray<TSharedRef<const FDeferredDestructionContainerBase>> EnqueuedForDeferredDestruction;
 };
