UEdGraphPin BreakAllPinLinks and MakeLinkTo don't by default notify the owning node that changes have happened. In some cases, this means that we don't detect that the graph has changed out from underneath us and so we don't recompile.

diff --git a/Engine/Plugins/FX/Niagara/Source/NiagaraEditor/Private/ViewModels/Stack/NiagaraStackGraphUtilities.cpp b/Engine/Plugins/FX/Niagara/Source/NiagaraEditor/Private/ViewModels/Stack/NiagaraStackGraphUtilities.cpp
--- a/Engine/Plugins/FX/Niagara/Source/NiagaraEditor/Private/ViewModels/Stack/NiagaraStackGraphUtilities.cpp
+++ b/Engine/Plugins/FX/Niagara/Source/NiagaraEditor/Private/ViewModels/Stack/NiagaraStackGraphUtilities.cpp
@@ -36,6 +36,18 @@ DECLARE_CYCLE_STAT(TEXT("Niagara - StackGraphUtilities - RelayoutGraph"), STAT_N
 
 #define LOCTEXT_NAMESPACE "NiagaraStackGraphUtilities"
 
+void FNiagaraStackGraphUtilities::MakeLinkTo(UEdGraphPin* PinA, UEdGraphPin* PinB)
+{
+	PinA->MakeLinkTo(PinB);
+	PinA->GetOwningNode()->PinConnectionListChanged(PinA);
+	PinB->GetOwningNode()->PinConnectionListChanged(PinB);
+}
+
+void FNiagaraStackGraphUtilities::BreakAllPinLinks(UEdGraphPin* PinA)
+{
+	PinA->BreakAllPinLinks(true);
+}
+
 void FNiagaraStackGraphUtilities::RelayoutGraph(UEdGraph& Graph)
 {
 	SCOPE_CYCLE_COUNTER(STAT_NiagaraEditor_StackGraphUtilities_RelayoutGraph);
@@ -174,7 +186,7 @@ void FNiagaraStackGraphUtilities::ConnectPinToInputNode(UEdGraphPin& Pin, UNiaga
 	InputNode.GetOutputPins(InputPins);
 	if (InputPins.Num() == 1)
 	{
-		Pin.MakeLinkTo(InputPins[0]);
+		MakeLinkTo(&Pin, InputPins[0]);
 	}
 }
 
@@ -431,27 +443,28 @@ void FNiagaraStackGraphUtilities::GetStackNodeGroups(UNiagaraNode& StackNode, TA
 void FNiagaraStackGraphUtilities::DisconnectStackNodeGroup(const FStackNodeGroup& DisconnectGroup, const FStackNodeGroup& PreviousGroup, const FStackNodeGroup& NextGroup)
 {
 	UEdGraphPin* PreviousOutputPin = FNiagaraStackGraphUtilities::GetParameterMapOutputPin(*PreviousGroup.EndNode);
-	PreviousOutputPin->BreakAllPinLinks();
+	BreakAllPinLinks(PreviousOutputPin);
 
 	UEdGraphPin* DisconnectOutputPin = FNiagaraStackGraphUtilities::GetParameterMapOutputPin(*DisconnectGroup.EndNode);
-	DisconnectOutputPin->BreakAllPinLinks();
+	BreakAllPinLinks(DisconnectOutputPin);
 
 	for (UNiagaraNode* NextStartNode : NextGroup.StartNodes)
 	{
 		UEdGraphPin* NextStartInputPin = FNiagaraStackGraphUtilities::GetParameterMapInputPin(*NextStartNode);
-		PreviousOutputPin->MakeLinkTo(NextStartInputPin);
+		MakeLinkTo(PreviousOutputPin, NextStartInputPin);
 	}
 }
 
 void FNiagaraStackGraphUtilities::ConnectStackNodeGroup(const FStackNodeGroup& ConnectGroup, const FStackNodeGroup& NewPreviousGroup, const FStackNodeGroup& NewNextGroup)
 {
 	UEdGraphPin* NewPreviousOutputPin = FNiagaraStackGraphUtilities::GetParameterMapOutputPin(*NewPreviousGroup.EndNode);
-	NewPreviousOutputPin->BreakAllPinLinks();
+	BreakAllPinLinks(NewPreviousOutputPin);
 
 	for (UNiagaraNode* ConnectStartNode : ConnectGroup.StartNodes)
 	{
 		UEdGraphPin* ConnectInputPin = FNiagaraStackGraphUtilities::GetParameterMapInputPin(*ConnectStartNode);
-		NewPreviousOutputPin->MakeLinkTo(ConnectInputPin);
+		MakeLinkTo(NewPreviousOutputPin, ConnectInputPin);
+
 	}
 
 	UEdGraphPin* ConnectOutputPin = FNiagaraStackGraphUtilities::GetParameterMapOutputPin(*ConnectGroup.EndNode);
@@ -459,7 +472,7 @@ void FNiagaraStackGraphUtilities::ConnectStackNodeGroup(const FStackNodeGroup& C
 	for (UNiagaraNode* NewNextStartNode : NewNextGroup.StartNodes)
 	{
 		UEdGraphPin* NewNextStartInputPin = FNiagaraStackGraphUtilities::GetParameterMapInputPin(*NewNextStartNode);
-		ConnectOutputPin->MakeLinkTo(NewNextStartInputPin);
+		MakeLinkTo(ConnectOutputPin, NewNextStartInputPin);
 	}
 }
 
@@ -818,14 +831,14 @@ UNiagaraNodeParameterMapSet& FNiagaraStackGraphUtilities::GetOrCreateStackFuncti
 		UEdGraphPin* OwningFunctionCallInputPin = FNiagaraStackGraphUtilities::GetParameterMapInputPin(StackFunctionCall);
 		UEdGraphPin* PreviousStackNodeOutputPin = OwningFunctionCallInputPin->LinkedTo[0];
 
-		OwningFunctionCallInputPin->BreakAllPinLinks();
-		OwningFunctionCallInputPin->MakeLinkTo(OverrideNodeOutputPin);
+		BreakAllPinLinks(OwningFunctionCallInputPin);
+		MakeLinkTo(OwningFunctionCallInputPin, OverrideNodeOutputPin);
 		for (UEdGraphPin* PreviousStackNodeOutputLinkedPin : PreviousStackNodeOutputPin->LinkedTo)
 		{
-			PreviousStackNodeOutputLinkedPin->MakeLinkTo(OverrideNodeOutputPin);
+			MakeLinkTo(PreviousStackNodeOutputLinkedPin, OverrideNodeOutputPin);
 		}
-		PreviousStackNodeOutputPin->BreakAllPinLinks();
-		PreviousStackNodeOutputPin->MakeLinkTo(OverrideNodeInputPin);
+		BreakAllPinLinks(PreviousStackNodeOutputPin);
+		MakeLinkTo(PreviousStackNodeOutputPin, OverrideNodeInputPin);
 	}
 	return *OverrideNode;
 }
@@ -937,14 +950,14 @@ void FNiagaraStackGraphUtilities::RemoveNodesForStackFunctionInputOverridePin(UE
 							}
 
 							// Disconnect the override node and remove it.
-							InputPin->BreakAllPinLinks();
-							OutputPin->BreakAllPinLinks();
+							BreakAllPinLinks(InputPin);
+							BreakAllPinLinks(OutputPin);
 							Graph->RemoveNode(DynamicInputNodeOverrideNode);
 
 							// Reconnect the pins which were connected to the removed override node.
 							for (UEdGraphPin* LinkedOutputPin : LinkedOutputPins)
 							{
-								LinkedInputPin->MakeLinkTo(LinkedOutputPin);
+								MakeLinkTo(LinkedInputPin, LinkedOutputPin);
 							}
 						}
 					}
@@ -978,8 +991,8 @@ void FNiagaraStackGraphUtilities::SetLinkedValueHandleForFunctionInput(UEdGraphP
 	const UEdGraphSchema_Niagara* NiagaraSchema = GetDefault<UEdGraphSchema_Niagara>();
 	FNiagaraTypeDefinition InputType = NiagaraSchema->PinToTypeDefinition(&OverridePin);
 	UEdGraphPin* GetOutputPin = GetNode->RequestNewTypedPin(EGPD_Output, InputType, LinkedParameterHandle.GetParameterHandleString());
-	GetInputPin->MakeLinkTo(PreviousStackNodeOutputPin);
-	GetOutputPin->MakeLinkTo(&OverridePin);
+	MakeLinkTo(GetInputPin, PreviousStackNodeOutputPin);
+	MakeLinkTo(GetOutputPin, &OverridePin);
 
 	if (NewNodePersistentId.IsValid())
 	{
@@ -1042,12 +1055,12 @@ void FNiagaraStackGraphUtilities::SetDynamicInputForFunctionInput(UEdGraphPin& O
 	
 	if (FunctionCallInputPin != nullptr && PreviousStackNodeOutputPin != nullptr)
 	{
-		FunctionCallInputPin->MakeLinkTo(PreviousStackNodeOutputPin);
+		MakeLinkTo(FunctionCallInputPin, PreviousStackNodeOutputPin);
 	}
 	
 	if (FunctionCallOutputPins.Num() >= 1 && FunctionCallOutputPins[0] != nullptr)
 	{
-		FunctionCallOutputPins[0]->MakeLinkTo(&OverridePin);
+		MakeLinkTo(FunctionCallOutputPins[0], &OverridePin);
 	}
 
 	OutDynamicInputFunctionCall = FunctionCallNode;
@@ -1092,8 +1105,8 @@ void FNiagaraStackGraphUtilities::SetCustomExpressionForFunctionInput(UEdGraphPi
 	UEdGraphPin* PreviousStackNodeOutputPin = OverrideNodeInputPin->LinkedTo[0];
 	checkf(PreviousStackNodeOutputPin != nullptr, TEXT("Invalid Stack Graph - No previous stack node."));
 
-	FunctionCallInputPin->MakeLinkTo(PreviousStackNodeOutputPin);
-	FunctionCallOutputPins[0]->MakeLinkTo(&OverridePin);
+	MakeLinkTo(FunctionCallInputPin, PreviousStackNodeOutputPin);
+	MakeLinkTo(FunctionCallOutputPins[0], &OverridePin);
 
 	OutDynamicInputFunctionCall = FunctionCallNode;
 
@@ -1394,8 +1407,8 @@ UNiagaraNodeOutput* FNiagaraStackGraphUtilities::ResetGraphForOutput(UNiagaraGra
 	}
 
 	UEdGraphPin* InputNodeOutputPin = GetParameterMapOutputPin(*InputNode);
-	OutputNodeInputPin->BreakAllPinLinks();
-	OutputNodeInputPin->MakeLinkTo(InputNodeOutputPin);
+	BreakAllPinLinks(OutputNodeInputPin);
+	MakeLinkTo(OutputNodeInputPin, InputNodeOutputPin);
 
 	if (ScriptUsage == ENiagaraScriptUsage::SystemSpawnScript || ScriptUsage == ENiagaraScriptUsage::SystemUpdateScript)
 	{
@@ -1983,6 +1996,8 @@ bool FNiagaraStackGraphUtilities::GetStackIssuesRecursively(const UNiagaraStackE
 
 void FNiagaraStackGraphUtilities::MoveModule(UNiagaraScript& SourceScript, UNiagaraNodeFunctionCall& ModuleToMove, UNiagaraSystem& TargetSystem, FGuid TargetEmitterHandleId, ENiagaraScriptUsage TargetUsage, FGuid TargetUsageId, int32 TargetModuleIndex, bool bForceCopy, UNiagaraNodeFunctionCall*& OutMovedModule)
 {
+	UE_LOG(LogNiagaraEditor, Log, TEXT("Move module %s"), *ModuleToMove.GetPathName());
+
 	UNiagaraScript* TargetScript = FNiagaraEditorUtilities::GetScriptFromSystem(TargetSystem, TargetEmitterHandleId, TargetUsage, TargetUsageId);
 	checkf(TargetScript != nullptr, TEXT("Target script not found"));
 
@@ -2208,6 +2223,9 @@ void FNiagaraStackGraphUtilities::MoveModule(UNiagaraScript& SourceScript, UNiag
 	}
 
 	OutMovedModule = Cast<UNiagaraNodeFunctionCall>(TargetGroup.EndNode);
+
+
+	UE_LOG(LogNiagaraEditor, Log, TEXT("Finished moving!"));
 }
 
 bool FNiagaraStackGraphUtilities::ParameterAllowedInExecutionCategory(const FName InParameterName, const FName ExecutionCategory)
@@ -2259,7 +2277,7 @@ void FNiagaraStackGraphUtilities::RebuildEmitterNodes(UNiagaraSystem& System)
 
 		if (InPinLinkedPin != nullptr &&& OutPinLinkedPin != nullptr)
 		{
-			InPinLinkedPin->MakeLinkTo(OutPinLinkedPin);
+			MakeLinkTo(InPinLinkedPin, OutPinLinkedPin);
 		}
 	}
 
@@ -2314,7 +2332,7 @@ void FNiagaraStackGraphUtilities::RebuildEmitterNodes(UNiagaraSystem& System)
 
 			InputNodeCreator.Finalize();
 
-			InputNodes[i]->GetOutputPin(0)->MakeLinkTo(OutputNodes[i]->GetInputPin(0));
+			MakeLinkTo(InputNodes[i]->GetOutputPin(0), OutputNodes[i]->GetInputPin(0));
 		}
 	}
 
diff --git a/Engine/Plugins/FX/Niagara/Source/NiagaraEditor/Public/ViewModels/Stack/NiagaraStackGraphUtilities.h b/Engine/Plugins/FX/Niagara/Source/NiagaraEditor/Public/ViewModels/Stack/NiagaraStackGraphUtilities.h
--- a/Engine/Plugins/FX/Niagara/Source/NiagaraEditor/Public/ViewModels/Stack/NiagaraStackGraphUtilities.h
+++ b/Engine/Plugins/FX/Niagara/Source/NiagaraEditor/Public/ViewModels/Stack/NiagaraStackGraphUtilities.h
@@ -27,6 +27,9 @@ class FCompileConstantResolver;
 
 namespace FNiagaraStackGraphUtilities
 {
+	void MakeLinkTo(UEdGraphPin* PinA, UEdGraphPin* PinB);
+	void BreakAllPinLinks(UEdGraphPin* PinA);
+
 	void RelayoutGraph(UEdGraph& Graph);
 
 	void GetWrittenVariablesForGraph(UEdGraph& Graph, TArray<FNiagaraVariable>& OutWrittenVariables);
