fixing erroneous warning if GPU emitter doesn't have static bounds but the parent system does.

diff --git a/Engine/Plugins/FX/Niagara/Source/NiagaraEditor/Private/ViewModels/Stack/NiagaraStackEmitterSettingsGroup.cpp b/Engine/Plugins/FX/Niagara/Source/NiagaraEditor/Private/ViewModels/Stack/NiagaraStackEmitterSettingsGroup.cpp
--- a/Engine/Plugins/FX/Niagara/Source/NiagaraEditor/Private/ViewModels/Stack/NiagaraStackEmitterSettingsGroup.cpp
+++ b/Engine/Plugins/FX/Niagara/Source/NiagaraEditor/Private/ViewModels/Stack/NiagaraStackEmitterSettingsGroup.cpp
@@ -9,6 +9,7 @@
 #include "ViewModels/Stack/NiagaraStackGraphUtilities.h"
 #include "NiagaraScriptMergeManager.h"
 #include "NiagaraEmitterDetailsCustomization.h"
+#include "NiagaraSystem.h"
 
 #define LOCTEXT_NAMESPACE "UNiagaraStackEmitterItemGroup"
 
@@ -103,14 +104,26 @@ void UNiagaraStackEmitterPropertiesItem::RefreshIssues(TArray<FStackIssue>& NewI
 	UNiagaraEmitter* ActualEmitter = GetEmitterViewModel()->GetEmitter();
 	if (ActualEmitter && ActualEmitter->SimTarget == ENiagaraSimTarget::GPUComputeSim && ActualEmitter->bFixedBounds == false)
 	{
-		FStackIssue MissingRequiredFixedBoundsModuleError(
-			EStackIssueSeverity::Warning,
-			LOCTEXT("RequiredFixedBoundsWarningFormat", "The emitter is GPU but the fixed bounds checkbox is not set.\r\nPlease update the Emitter properties, otherwise existing value for fixed bounds will be used."),
-			LOCTEXT("MissingFixedBounds", "Missing fixed bounds."),
-			GetStackEditorDataKey(),
-			false);
-
-		NewIssues.Add(MissingRequiredFixedBoundsModuleError);
+		bool bAddError = true;
+		
+		UNiagaraSystem& Sys = GetSystemViewModel()->GetSystem();
+		if (Sys.bFixedBounds)
+		{
+			bAddError = false;
+		}			
+		
+
+		if (bAddError)
+		{
+			FStackIssue MissingRequiredFixedBoundsModuleError(
+				EStackIssueSeverity::Warning,
+				LOCTEXT("RequiredFixedBoundsWarningFormat", "The emitter is GPU but the fixed bounds checkbox is not set.\r\nPlease update the Emitter or System properties, otherwise existing value for fixed bounds will be used."),
+				LOCTEXT("MissingFixedBounds", "Missing fixed bounds."),
+				GetStackEditorDataKey(),
+				false);
+
+			NewIssues.Add(MissingRequiredFixedBoundsModuleError);
+		}
 	}
 }
 
