Fix crash clearing out a NPCInstance's collection.

diff --git a/Engine/Plugins/FX/Niagara/Source/NiagaraEditor/Private/ViewModels/NiagaraCollectionParameterViewModel.cpp b/Engine/Plugins/FX/Niagara/Source/NiagaraEditor/Private/ViewModels/NiagaraCollectionParameterViewModel.cpp
--- a/Engine/Plugins/FX/Niagara/Source/NiagaraEditor/Private/ViewModels/NiagaraCollectionParameterViewModel.cpp
+++ b/Engine/Plugins/FX/Niagara/Source/NiagaraEditor/Private/ViewModels/NiagaraCollectionParameterViewModel.cpp
@@ -28,7 +28,14 @@ void FNiagaraCollectionParameterViewModel::Reset()
 
 FName FNiagaraCollectionParameterViewModel::GetName() const
 {
-	return *CollectionInst->GetParent()->FriendlyNameFromParameterName(Parameter.GetName().ToString());
+	if (CollectionInst && CollectionInst->GetParent())
+	{
+		return *CollectionInst->GetParent()->FriendlyNameFromParameterName(Parameter.GetName().ToString());
+	}
+	else
+	{
+		return FName();
+	}
 }
 
 FText FNiagaraCollectionParameterViewModel::GetTypeDisplayName() const
