Fix postload failing to actually make new non-null entries.

diff --git a/Engine/Plugins/FX/Niagara/Source/NiagaraEditor/Private/NiagaraGraph.cpp b/Engine/Plugins/FX/Niagara/Source/NiagaraEditor/Private/NiagaraGraph.cpp
--- a/Engine/Plugins/FX/Niagara/Source/NiagaraEditor/Private/NiagaraGraph.cpp
+++ b/Engine/Plugins/FX/Niagara/Source/NiagaraEditor/Private/NiagaraGraph.cpp
@@ -144,7 +144,7 @@ void UNiagaraGraph::PostLoad()
 		for (auto It = VariableToScriptVariable.CreateIterator(); It; ++It)
 		{
 			FNiagaraVariable Var = It.Key();
-			UNiagaraScriptVariable* ScriptVar = It.Value();
+			UNiagaraScriptVariable*& ScriptVar = It.Value();
 
 			if (ScriptVar == nullptr)
 			{
