Avoiding shutdown path that erroneously triggers an error.

diff --git a/Engine/Plugins/FX/Niagara/Source/Niagara/Private/NiagaraSystemSimulation.cpp b/Engine/Plugins/FX/Niagara/Source/Niagara/Private/NiagaraSystemSimulation.cpp
--- a/Engine/Plugins/FX/Niagara/Source/Niagara/Private/NiagaraSystemSimulation.cpp
+++ b/Engine/Plugins/FX/Niagara/Source/Niagara/Private/NiagaraSystemSimulation.cpp
@@ -421,11 +421,29 @@ void FNiagaraSystemSimulation::Destroy()
 
 	while (SystemInstances.Num())
 	{
-		SystemInstances.Last()->Deactivate(true);
+		FNiagaraSystemInstance* Inst = SystemInstances.Last();
+		check(Inst);
+		if (ensure(Inst->GetComponent()))//Currently we have no cases whre there shouldn't be a component but maybe in future.
+		{
+			Inst->GetComponent()->DeactivateImmediate();
+		}
+		else
+		{
+			Inst->Deactivate(true);
+		}
 	}
 	while (PendingSystemInstances.Num())
 	{
-		PendingSystemInstances.Last()->Deactivate(true);
+		FNiagaraSystemInstance* Inst = PendingSystemInstances.Last();
+		check(Inst);
+		if (ensure(Inst->GetComponent()))//Currently we have no cases whre there shouldn't be a component but maybe in future.
+		{
+			Inst->GetComponent()->DeactivateImmediate();
+		}
+		else
+		{
+			Inst->Deactivate(true);
+		}
 	}
 	SystemInstances.Empty();
 	PendingSystemInstances.Empty();
