Fix for uninitialized field in FNiagaraEventGeneratorProperties::MaxEventsPerFrame, causing non-determinism.

diff --git a/Engine/Plugins/FX/Niagara/Source/Niagara/Classes/NiagaraEmitter.h b/Engine/Plugins/FX/Niagara/Source/Niagara/Classes/NiagaraEmitter.h
--- a/Engine/Plugins/FX/Niagara/Source/Niagara/Classes/NiagaraEmitter.h
+++ b/Engine/Plugins/FX/Niagara/Source/Niagara/Classes/NiagaraEmitter.h
@@ -65,11 +65,7 @@ struct FNiagaraEventGeneratorProperties
 {
 	GENERATED_BODY()
 
-	FNiagaraEventGeneratorProperties()
-	: MaxEventsPerFrame(64)
-	{
-
-	}
+	FNiagaraEventGeneratorProperties() = default;
 
 	FNiagaraEventGeneratorProperties(FNiagaraDataSetProperties &Props, FName InEventGenerator)
 		: ID(Props.ID.Name)
@@ -82,7 +78,7 @@ struct FNiagaraEventGeneratorProperties
 
 	/** Max Number of Events that can be generated per frame. */
 	UPROPERTY(EditAnywhere, Category = "Event Receiver")
-	int32 MaxEventsPerFrame; //TODO - More complex allocation so that we can grow dynamically if more space is needed ?
+	int32 MaxEventsPerFrame = 64; //TODO - More complex allocation so that we can grow dynamically if more space is needed ?
 
 	UPROPERTY()
 	FName ID;
