Buttons in Niagara module headers in the stack view are now overlaid above the name.

#rnx

diff --git a/Engine/Plugins/FX/Niagara/Source/NiagaraEditorWidgets/Private/Stack/SNiagaraStackModuleItem.cpp b/Engine/Plugins/FX/Niagara/Source/NiagaraEditorWidgets/Private/Stack/SNiagaraStackModuleItem.cpp
--- a/Engine/Plugins/FX/Niagara/Source/NiagaraEditorWidgets/Private/Stack/SNiagaraStackModuleItem.cpp
+++ b/Engine/Plugins/FX/Niagara/Source/NiagaraEditorWidgets/Private/Stack/SNiagaraStackModuleItem.cpp
@@ -48,10 +48,11 @@ void SNiagaraStackModuleItem::Construct(const FArguments& InArgs, UNiagaraStackM
 		.BackgroundColorHover(FNiagaraEditorWidgetsStyle::Get().GetColor("NiagaraEditor.Stack.DropTarget.BackgroundColorHover"))
 		.Content()
 		[
-			SNew(SHorizontalBox)
+			SNew(SOverlay)
 			// Name
-			+ SHorizontalBox::Slot()
+			+ SOverlay::Slot()
 			.VAlign(VAlign_Center)
+			.HAlign(HAlign_Fill)
 			.Padding(2, 0, 0, 0)
 			[
 				SAssignNew(DisplayNameWidget, SNiagaraStackDisplayName, InModuleItem, *InStackViewModel, "NiagaraEditor.Stack.ItemText")
@@ -59,74 +60,81 @@ void SNiagaraStackModuleItem::Construct(const FArguments& InArgs, UNiagaraStackM
 				.TypeNameStyle(FNiagaraEditorWidgetsStyle::Get(), "NiagaraEditor.Stack.TypeNameText")
 			]
 			// Raise Action Menu button
-			+ SHorizontalBox::Slot()
-			.AutoWidth()
+			+ SOverlay::Slot()
+			.HAlign(HAlign_Right)
 			[
-				SNew(SComboButton)
-				.HasDownArrow(false)
-				.ButtonStyle(FEditorStyle::Get(), "HoverHintOnly")
-				.ForegroundColor(FSlateColor::UseForeground())
-				.OnGetMenuContent(this, &SNiagaraStackModuleItem::RaiseActionMenuClicked)
-				.ContentPadding(FMargin(2))
-				.HAlign(HAlign_Center)
-				.VAlign(VAlign_Center)
-				.Visibility(this, &SNiagaraStackModuleItem::GetRaiseActionMenuVisibility)
-				.IsEnabled(this, &SNiagaraStackModuleItem::GetButtonsEnabled)
-				.ButtonContent()
+				SNew(SBorder)
+				.Padding(0)
+				.BorderImage(FEditorStyle::GetBrush("WhiteBrush"))
+				.BorderBackgroundColor(FNiagaraEditorWidgetsStyle::Get().GetColor("NiagaraEditor.Stack.Item.HeaderBackgroundColor"))
 				[
-					SNew(SImage)
-					.Image(FEditorStyle::Get().GetBrush("PropertyWindow.Button_AddToArray"))
+					SNew(SHorizontalBox)
+					+ SHorizontalBox::Slot()
+					.VAlign(VAlign_Center)
+					.AutoWidth()
+					[
+						SNew(SComboButton)
+						.HasDownArrow(true)
+						.ButtonStyle(FEditorStyle::Get(), "HoverHintOnly")
+						.ForegroundColor(FSlateColor::UseForeground())
+						.OnGetMenuContent(this, &SNiagaraStackModuleItem::RaiseActionMenuClicked)
+						.ContentPadding(FMargin(2))
+						.HAlign(HAlign_Center)
+						.VAlign(VAlign_Center)
+						.Visibility(this, &SNiagaraStackModuleItem::GetRaiseActionMenuVisibility)
+						.IsEnabled(this, &SNiagaraStackModuleItem::GetButtonsEnabled)
+					]
+					// Refresh button
+					+ SHorizontalBox::Slot()
+					.VAlign(VAlign_Center)
+					.AutoWidth()
+					[
+						SNew(SButton)
+						.ButtonStyle(FEditorStyle::Get(), "HoverHintOnly")
+						.IsFocusable(false)
+						.ForegroundColor(FNiagaraEditorWidgetsStyle::Get().GetColor("NiagaraEditor.Stack.FlatButtonColor"))
+						.ToolTipText(LOCTEXT("RefreshTooltip", "Refresh this module"))
+						.Visibility(this, &SNiagaraStackModuleItem::GetRefreshVisibility)
+						.IsEnabled(this, &SNiagaraStackModuleItem::GetButtonsEnabled)
+						.OnClicked(this, &SNiagaraStackModuleItem::RefreshClicked)
+						.Content()
+						[
+							SNew(STextBlock)
+							.Font(FEditorStyle::Get().GetFontStyle("FontAwesome.10"))
+							.Text(FEditorFontGlyphs::Refresh)
+						]
+					]
+					// Delete button
+					+ SHorizontalBox::Slot()
+					.VAlign(VAlign_Center)
+					.AutoWidth()
+					[
+						SNew(SButton)
+						.ButtonStyle(FEditorStyle::Get(), "HoverHintOnly")
+						.IsFocusable(false)
+						.ForegroundColor(FNiagaraEditorWidgetsStyle::Get().GetColor("NiagaraEditor.Stack.FlatButtonColor"))
+						.ToolTipText(this, &SNiagaraStackModuleItem::GetDeleteButtonToolTipText)
+						.IsEnabled(this, &SNiagaraStackModuleItem::GetDeleteButtonEnabled)
+						.OnClicked(this, &SNiagaraStackModuleItem::DeleteClicked)
+						.Content()
+						[
+							SNew(STextBlock)
+							.Font(FEditorStyle::Get().GetFontStyle("FontAwesome.10"))
+							.Text(FEditorFontGlyphs::Trash)
+						]
+					]
+					// Enabled checkbox
+					+ SHorizontalBox::Slot()
+					.Padding(2, 0, 0, 0)
+					.AutoWidth()
+					[
+						SNew(SCheckBox)
+						.IsChecked(this, &SNiagaraStackModuleItem::GetCheckState)
+						.OnCheckStateChanged(this, &SNiagaraStackModuleItem::OnCheckStateChanged)
+						.IsEnabled(this, &SNiagaraStackModuleItem::GetEnabledCheckBoxEnabled)
+					]
 				]
 			]
-			// Refresh button
-			+ SHorizontalBox::Slot()
-			.VAlign(VAlign_Center)
-			.AutoWidth()
-			[
-				SNew(SButton)
-				.ButtonStyle(FEditorStyle::Get(), "HoverHintOnly")
-				.IsFocusable(false)
-				.ForegroundColor(FNiagaraEditorWidgetsStyle::Get().GetColor("NiagaraEditor.Stack.FlatButtonColor"))
-				.ToolTipText(LOCTEXT("RefreshTooltip", "Refresh this module"))
-				.Visibility(this, &SNiagaraStackModuleItem::GetRefreshVisibility)
-				.IsEnabled(this, &SNiagaraStackModuleItem::GetButtonsEnabled)
-				.OnClicked(this, &SNiagaraStackModuleItem::RefreshClicked)
-				.Content()
-				[
-					SNew(STextBlock)
-					.Font(FEditorStyle::Get().GetFontStyle("FontAwesome.10"))
-					.Text(FEditorFontGlyphs::Refresh)
-				]
-			]
-			// Delete button
-			+ SHorizontalBox::Slot()
-			.VAlign(VAlign_Center)
-			.AutoWidth()
-			[
-				SNew(SButton)
-				.ButtonStyle(FEditorStyle::Get(), "HoverHintOnly")
-				.IsFocusable(false)
-				.ForegroundColor(FNiagaraEditorWidgetsStyle::Get().GetColor("NiagaraEditor.Stack.FlatButtonColor"))
-				.ToolTipText(this, &SNiagaraStackModuleItem::GetDeleteButtonToolTipText)
-				.IsEnabled(this, &SNiagaraStackModuleItem::GetDeleteButtonEnabled)
-				.OnClicked(this, &SNiagaraStackModuleItem::DeleteClicked)
-				.Content()
-				[
-					SNew(STextBlock)
-					.Font(FEditorStyle::Get().GetFontStyle("FontAwesome.10"))
-					.Text(FEditorFontGlyphs::Trash)
-				]
-			]
-			// Enabled checkbox
-			+ SHorizontalBox::Slot()
-			.Padding(2, 0, 0, 0)
-			.AutoWidth()
-			[
-				SNew(SCheckBox)
-				.IsChecked(this, &SNiagaraStackModuleItem::GetCheckState)
-				.OnCheckStateChanged(this, &SNiagaraStackModuleItem::OnCheckStateChanged)
-				.IsEnabled(this, &SNiagaraStackModuleItem::GetEnabledCheckBoxEnabled)
-			]
 		]
 	];
 }
