Updated module graph UI to better support parameter default modes.
Changes:
* "Value" mode variables do not allow connections to their input pins
* "Binding" mode variables are hidden completely
* Changing the default value in the details panel automatically updates all the graph nodes
* Existing graphs are validated and updated on PosLoad

diff --git a/Engine/Plugins/FX/Niagara/Source/Niagara/Public/NiagaraTypes.h b/Engine/Plugins/FX/Niagara/Source/Niagara/Public/NiagaraTypes.h
--- a/Engine/Plugins/FX/Niagara/Source/Niagara/Public/NiagaraTypes.h
+++ b/Engine/Plugins/FX/Niagara/Source/Niagara/Public/NiagaraTypes.h
@@ -1130,6 +1130,18 @@ struct FNiagaraVariable : public FNiagaraVariableBase
 		return !(*this == Other);
 	}
 
+	/** Checks if the types match and either both variables are uninitialized or both hold exactly the same data.*/
+	bool HoldsSameData(const FNiagaraVariable& Other) const
+	{
+		if (TypeDef != Other.TypeDef) {
+			return false;
+		}
+		if (!IsDataAllocated() && !Other.IsDataAllocated()) {
+			return true;
+		}
+		return IsDataAllocated() && Other.IsDataAllocated() && VarData.Num() == Other.VarData.Num() && FMemory::Memcmp(VarData.GetData(), Other.VarData.GetData(), VarData.Num()) == 0;
+	}
+
 	// Var data operations
 	void AllocateData()
 	{
diff --git a/Engine/Plugins/FX/Niagara/Source/NiagaraEditor/Private/Customizations/NiagaraScriptVariableCustomization.cpp b/Engine/Plugins/FX/Niagara/Source/NiagaraEditor/Private/Customizations/NiagaraScriptVariableCustomization.cpp
--- a/Engine/Plugins/FX/Niagara/Source/NiagaraEditor/Private/Customizations/NiagaraScriptVariableCustomization.cpp
+++ b/Engine/Plugins/FX/Niagara/Source/NiagaraEditor/Private/Customizations/NiagaraScriptVariableCustomization.cpp
@@ -41,7 +41,7 @@ FNiagaraScriptVariableDetails::~FNiagaraScriptVariableDetails()
 	GEditor->UnregisterForUndo(this);
 }
 
-UEdGraphPin* FNiagaraScriptVariableDetails::GetDefaultPin()
+UEdGraphPin* FNiagaraScriptVariableDetails::GetAnyDefaultPin()
 {
 	// TODO: We don't know the usage at this point, so we'll try each script type in order
 	//       This could probably be made much more robust, but works for now.
@@ -61,6 +61,15 @@ UEdGraphPin* FNiagaraScriptVariableDetails::GetDefaultPin()
 	return nullptr;
 }
 
+TArray<UEdGraphPin*> FNiagaraScriptVariableDetails::GetDefaultPins()
+{
+	if (UNiagaraGraph* Graph = Cast<UNiagaraGraph>(Variable->GetOuter()))
+	{
+		return Graph->FindParameterMapDefaultValuePins(Variable->Variable.GetName());
+	}
+	return TArray<UEdGraphPin*>();
+}
+
 void FNiagaraScriptVariableDetails::PostUndo(bool bSuccess)
 {
 	if (Variable == nullptr)
@@ -80,7 +89,7 @@ void FNiagaraScriptVariableDetails::PostUndo(bool bSuccess)
 	}
 	else
 	{
-		if (UEdGraphPin* Pin = GetDefaultPin())
+		if (UEdGraphPin* Pin = GetAnyDefaultPin())
 		{
 			if (TypeUtilityValue && ParameterEditorValue)
 			{
@@ -101,6 +110,11 @@ void FNiagaraScriptVariableDetails::CustomizeDetails(const TSharedPtr<IDetailLay
 
 void FNiagaraScriptVariableDetails::OnComboValueChanged()
 {
+	if (UNiagaraGraph* Graph = Cast<UNiagaraGraph>(Variable->GetOuter()))
+	{
+		Graph->ScriptVariableChanged(this->Variable->Variable);
+	}
+
 	IDetailLayoutBuilder* DetailBuilderPtr = nullptr;
 	if (TSharedPtr<IDetailLayoutBuilder> LockedDetailBuilder = CachedDetailBuilder.Pin())
 	{
@@ -190,7 +204,7 @@ void FNiagaraScriptVariableDetails::CustomizeDetails(IDetailLayoutBuilder& Detai
 		DetailBuilder.HideProperty(DefaultModeHandle);
 		DetailBuilder.AddPropertyToCategory(DefaultModeHandle);
 
-		if (UEdGraphPin* Pin = GetDefaultPin())
+		if (UEdGraphPin* Pin = GetAnyDefaultPin())
 		{
 			TypeUtilityValue = EditorModule.GetTypeUtilities(Variable->Variable.GetType());
 			if (TypeUtilityValue && TypeUtilityValue->CanCreateParameterEditor() && Variable->DefaultMode == ENiagaraDefaultMode::Value)
@@ -251,35 +265,38 @@ void FNiagaraScriptVariableDetails::CustomizeDetails(IDetailLayoutBuilder& Detai
  
 void FNiagaraScriptVariableDetails::OnValueChanged()
 {
-
-	if (TypeUtilityValue && ParameterEditorValue)
+	if (TypeUtilityValue && ParameterEditorValue && GetDefaultPins().Num() > 0)
 	{
-		if (UEdGraphPin* Pin = GetDefaultPin())
+		FString NewDefaultValue;
+		if (!ParameterEditorValue->CanChangeContinuously())
 		{
-			FString NewDefaultValue;
-			if (!ParameterEditorValue->CanChangeContinuously())
+			const FScopedTransaction Transaction(NSLOCTEXT("ScriptVariableCustomization", "ChangeValue", "Change Default Value"));
+			Variable->Modify();
+
+			TSharedPtr<FStructOnScope> ParameterValue = MakeShareable(new FStructOnScope(Variable->Variable.GetType().GetStruct()));
+			ParameterEditorValue->UpdateStructFromInternalValue(ParameterValue.ToSharedRef());
+			Variable->Variable.SetData(ParameterValue->GetStructMemory());
+			NewDefaultValue = TypeUtilityValue->GetPinDefaultStringFromValue(Variable->Variable);
+
+			for (UEdGraphPin* Pin : GetDefaultPins())
 			{
-				const FScopedTransaction Transaction( NSLOCTEXT("ScriptVariableCustomization", "ChangeValue", "Change Default Value" ) );
 				Pin->Modify();
-				Variable->Modify();
+				GetDefault<UEdGraphSchema_Niagara>()->TrySetDefaultValue(*Pin, NewDefaultValue, true);
+			}
+		}
+		else
+		{
+			TSharedPtr<FStructOnScope> ParameterValue = MakeShareable(new FStructOnScope(Variable->Variable.GetType().GetStruct()));
+			ParameterEditorValue->UpdateStructFromInternalValue(ParameterValue.ToSharedRef());
+			Variable->Variable.SetData(ParameterValue->GetStructMemory());
+			NewDefaultValue = TypeUtilityValue->GetPinDefaultStringFromValue(Variable->Variable);
 
-				TSharedPtr<FStructOnScope> ParameterValue = MakeShareable(new FStructOnScope(Variable->Variable.GetType().GetStruct()));
-				ParameterEditorValue->UpdateStructFromInternalValue(ParameterValue.ToSharedRef());
-				Variable->Variable.SetData(ParameterValue->GetStructMemory());
-				NewDefaultValue = TypeUtilityValue->GetPinDefaultStringFromValue(Variable->Variable);
-			}		
-			else
+			for (UEdGraphPin* Pin : GetDefaultPins())
 			{
-				TSharedPtr<FStructOnScope> ParameterValue = MakeShareable(new FStructOnScope(Variable->Variable.GetType().GetStruct()));
-				ParameterEditorValue->UpdateStructFromInternalValue(ParameterValue.ToSharedRef());
-				Variable->Variable.SetData(ParameterValue->GetStructMemory());
-				NewDefaultValue = TypeUtilityValue->GetPinDefaultStringFromValue(Variable->Variable);
+				GetDefault<UEdGraphSchema_Niagara>()->TrySetDefaultValue(*Pin, NewDefaultValue, true);
 			}
-
-			GetDefault<UEdGraphSchema_Niagara>()->TrySetDefaultValue(*Pin, NewDefaultValue, true);
-			
 		}
-	} 
+	}
 }
  
 
@@ -290,22 +307,21 @@ void FNiagaraScriptVariableDetails::OnBeginValueChanged()
 		return;
 	}
 
-	if (TypeUtilityValue && ParameterEditorValue)
+	if (TypeUtilityValue && ParameterEditorValue && GetDefaultPins().Num() > 0)
 	{
-		if (UEdGraphPin* Pin = GetDefaultPin())
+		GEditor->BeginTransaction(NSLOCTEXT("ScriptVariableCustomization", "ChangeValue", "Change Default Value"));
+		Variable->Modify();
+		TSharedPtr<FStructOnScope> ParameterValue = MakeShareable(new FStructOnScope(Variable->Variable.GetType().GetStruct()));
+		ParameterEditorValue->UpdateStructFromInternalValue(ParameterValue.ToSharedRef());
+		Variable->Variable.SetData(ParameterValue->GetStructMemory());
+
+		for (UEdGraphPin* Pin : GetDefaultPins())
 		{
-			GEditor->BeginTransaction(NSLOCTEXT("ScriptVariableCustomization", "ChangeValue", "Change Default Value"));
-			Variable->Modify();
 			Pin->Modify();
-			TSharedPtr<FStructOnScope> ParameterValue = MakeShareable(new FStructOnScope(Variable->Variable.GetType().GetStruct()));
-			ParameterEditorValue->UpdateStructFromInternalValue(ParameterValue.ToSharedRef());
-			Variable->Variable.SetData(ParameterValue->GetStructMemory());
-
 			FString NewDefaultValue = TypeUtilityValue->GetPinDefaultStringFromValue(Variable->Variable);
 			GetDefault<UEdGraphSchema_Niagara>()->TrySetDefaultValue(*Pin, NewDefaultValue, true);
-
 		}
-	} 
+	}
 }
 
 void FNiagaraScriptVariableDetails::OnEndValueChanged()
diff --git a/Engine/Plugins/FX/Niagara/Source/NiagaraEditor/Private/Customizations/NiagaraScriptVariableCustomization.h b/Engine/Plugins/FX/Niagara/Source/NiagaraEditor/Private/Customizations/NiagaraScriptVariableCustomization.h
--- a/Engine/Plugins/FX/Niagara/Source/NiagaraEditor/Private/Customizations/NiagaraScriptVariableCustomization.h
+++ b/Engine/Plugins/FX/Niagara/Source/NiagaraEditor/Private/Customizations/NiagaraScriptVariableCustomization.h
@@ -36,7 +36,9 @@ private:
 	void OnValueChanged();
 	void OnStaticSwitchValueChanged();
 
-	UEdGraphPin* GetDefaultPin();
+	UEdGraphPin* GetAnyDefaultPin();
+
+	TArray<UEdGraphPin*> GetDefaultPins();
 
 	TWeakPtr<class IDetailLayoutBuilder> CachedDetailBuilder;
 
diff --git a/Engine/Plugins/FX/Niagara/Source/NiagaraEditor/Private/NiagaraGraph.cpp b/Engine/Plugins/FX/Niagara/Source/NiagaraEditor/Private/NiagaraGraph.cpp
--- a/Engine/Plugins/FX/Niagara/Source/NiagaraEditor/Private/NiagaraGraph.cpp
+++ b/Engine/Plugins/FX/Niagara/Source/NiagaraEditor/Private/NiagaraGraph.cpp
@@ -269,6 +269,9 @@ void UNiagaraGraph::PostLoad()
 			//UE_LOG(LogNiagaraEditor, Log, TEXT("Migrated %d old metadata entries for \"%s\""), NumMigrated, *FullPathName);
 		}
 	}
+
+	// Fix inconsistencies in the default value declaration between graph and metadata
+	ValidateDefaultPins();
 	
 	{
 		// Create a UNiagaraScriptVariable instance for every entry in the parameter map for which there is no existing script variable. 
@@ -279,8 +282,6 @@ void UNiagaraGraph::PostLoad()
 			if (Variable == nullptr)
 			{
 				VarsToAdd.Add(ParameterToReferences.Key);
- 
- 
 			}
 		}
 		for (FNiagaraVariable& Var : VarsToAdd)
@@ -379,6 +380,110 @@ void UNiagaraGraph::PostLoad()
 	InvalidateCachedParameterData();
 }
 
+TArray<UEdGraphPin*> UNiagaraGraph::FindParameterMapDefaultValuePins(const FName VariableName) const
+{
+	TArray<UEdGraphPin*> DefaultPins;
+
+	for (UEdGraphNode* Node : Nodes)
+	{
+		UNiagaraNodeOutput* OutNode = Cast<UNiagaraNodeOutput>(Node);
+		if (!OutNode)
+		{
+			continue;
+		}
+		TArray<UNiagaraNode*> NodesTraversed;
+		BuildTraversal(NodesTraversed, OutNode);
+
+		for (UNiagaraNode* NiagaraNode : NodesTraversed)
+		{
+			UNiagaraNodeParameterMapGet* GetNode = Cast<UNiagaraNodeParameterMapGet>(NiagaraNode);
+			if (!GetNode)
+			{
+				continue;
+			}
+			TArray<UEdGraphPin*> OutputPins;
+			GetNode->GetOutputPins(OutputPins);
+			for (UEdGraphPin* OutputPin : OutputPins)
+			{
+				if (VariableName != OutputPin->PinName)
+				{
+					continue;
+				}
+				if (UEdGraphPin* Pin = GetNode->GetDefaultPin(OutputPin))
+				{
+					check(Pin->Direction == EEdGraphPinDirection::EGPD_Input);
+					DefaultPins.AddUnique(Pin);
+				}
+			}
+		}
+	}
+	return DefaultPins;
+}
+
+void UNiagaraGraph::ValidateDefaultPins()
+{
+	FNiagaraEditorModule& EditorModule = FNiagaraEditorModule::Get();
+	for (auto& MetaData : GetAllMetaData()) {
+		FNiagaraVariable Variable = MetaData.Key;
+		UNiagaraScriptVariable* ScriptVariable = MetaData.Value;
+		if (!ScriptVariable || ScriptVariable->DefaultMode == ENiagaraDefaultMode::Custom) {
+			// If the user selected custom mode they can basically do whatever they want
+			continue;
+		}
+		if (ScriptVariable->Metadata.bIsStaticSwitch) {
+			// We ignore static switch variables as they handle default values differently
+			continue;
+		}
+
+		// Determine if the values set in the graph (which are used for the compilation) are consistent and if necessary update the metadata value
+		TArray<UEdGraphPin*> Pins = FindParameterMapDefaultValuePins(Variable.GetName());
+		bool IsCustom = false;
+		bool IsConsistent = true;
+		bool DefaultInitialized = false;
+		FNiagaraVariable DefaultData;
+		auto TypeUtilityValue = EditorModule.GetTypeUtilities(Variable.GetType());
+		for (UEdGraphPin* Pin : Pins) {
+			if (Pin->LinkedTo.Num() > 0) {
+				IsCustom = true;
+			}
+			else if (!Pin->bHidden && TypeUtilityValue) {
+				FNiagaraVariable ComparisonData = Variable;
+				TypeUtilityValue->SetValueFromPinDefaultString(Pin->DefaultValue, ComparisonData);
+				if (!DefaultInitialized) {
+					DefaultData = ComparisonData;
+					DefaultInitialized = true;
+				}
+				else if (!ComparisonData.HoldsSameData(DefaultData)) {
+					IsConsistent = false;
+				}
+			}
+		}
+
+		if (!IsCustom && !IsConsistent) {
+			UE_LOG(LogNiagaraEditor, Error, TEXT("Niagara graph %s: The default value declaration for the variable '%s' is not consistent between the graph and the metadata.\n  Either change the default mode to 'custom' or check the input pins in the parameter map get node in the graph."),
+				*GetFullName(), *Variable.GetName().ToString());
+			continue;
+		}
+		if (IsCustom) {
+			ScriptVariable->DefaultMode = ENiagaraDefaultMode::Custom;
+		}
+		else {
+			for (UEdGraphPin* Pin : Pins) {
+				Pin->bNotConnectable = true;
+				Pin->bDefaultValueIsReadOnly = true;
+				if (ScriptVariable->DefaultMode == ENiagaraDefaultMode::Binding) {
+					Pin->bHidden = true;
+				}
+
+			}
+			if (DefaultInitialized) {
+				// set default value from pins
+				ScriptVariable->Variable = DefaultData;
+			}
+		}
+	}
+}
+
 void UNiagaraGraph::PostEditChangeProperty(FPropertyChangedEvent& PropertyChangedEvent)
 {
 	NotifyGraphChanged();
@@ -1116,11 +1221,8 @@ bool UNiagaraGraph::RenameParameter(const FNiagaraVariable& Parameter, FName New
 	SetPerScriptMetaData(NewParameter, OldMetaData);
 
 	// Fixup reference collection and pin names
-	FNiagaraGraphParameterReferenceCollection* ReferenceCollection = ParameterToReferencesMap.Find(Parameter);
-	int NumReferences = 0;
-	if (ReferenceCollection)
+	if (FNiagaraGraphParameterReferenceCollection* ReferenceCollection = ParameterToReferencesMap.Find(Parameter))
 	{
-		NumReferences = ReferenceCollection->ParameterReferences.Num();
 		const FText NewNameText = FText::FromName(NewName);
 		FNiagaraGraphParameterReferenceCollection NewReferences = *ReferenceCollection;
 		for (FNiagaraGraphParameterReference& Reference : NewReferences.ParameterReferences)
@@ -1147,6 +1249,42 @@ bool UNiagaraGraph::RenameParameter(const FNiagaraVariable& Parameter, FName New
 	return true;
 }
 
+void UNiagaraGraph::ScriptVariableChanged(FNiagaraVariable Variable)
+{
+	UNiagaraScriptVariable** ScriptVariable = GetAllMetaData().Find(Variable);
+	if (!ScriptVariable || !*ScriptVariable || (*ScriptVariable)->Metadata.bIsStaticSwitch) {
+		return;
+	}
+
+	FNiagaraEditorModule& EditorModule = FNiagaraEditorModule::Get();
+	auto TypeUtilityValue = EditorModule.GetTypeUtilities(Variable.GetType());
+
+	TArray<UEdGraphPin*> Pins = FindParameterMapDefaultValuePins(Variable.GetName());
+	for (UEdGraphPin* Pin : Pins) {
+		Pin->bHidden = (*ScriptVariable)->DefaultMode == ENiagaraDefaultMode::Binding;
+		if ((*ScriptVariable)->DefaultMode == ENiagaraDefaultMode::Custom) {
+			Pin->bNotConnectable = false;
+			Pin->bDefaultValueIsReadOnly = false;
+		}
+		else {
+			Pin->BreakAllPinLinks(true);
+			Pin->bNotConnectable = true;
+			Pin->bDefaultValueIsReadOnly = true;
+
+			if ((*ScriptVariable)->DefaultMode == ENiagaraDefaultMode::Value && !Variable.GetType().IsDataInterface()) {
+				if (!Variable.IsDataAllocated()) {
+					Variable.AllocateData();
+				}
+				FString NewDefaultValue = TypeUtilityValue->GetPinDefaultStringFromValue(Variable);
+				GetDefault<UEdGraphSchema_Niagara>()->TrySetDefaultValue(*Pin, NewDefaultValue, true);
+			}
+		}
+	}
+
+	ValidateDefaultPins();
+	NotifyGraphChanged();
+}
+
 int32 UNiagaraGraph::GetOutputNodeVariableIndex(const FNiagaraVariable& Variable)const
 {
 	TArray<FNiagaraVariable> Variables;
diff --git a/Engine/Plugins/FX/Niagara/Source/NiagaraEditor/Private/NiagaraNodeParameterMapGet.cpp b/Engine/Plugins/FX/Niagara/Source/NiagaraEditor/Private/NiagaraNodeParameterMapGet.cpp
--- a/Engine/Plugins/FX/Niagara/Source/NiagaraEditor/Private/NiagaraNodeParameterMapGet.cpp
+++ b/Engine/Plugins/FX/Niagara/Source/NiagaraEditor/Private/NiagaraNodeParameterMapGet.cpp
@@ -18,6 +18,7 @@
 #include "Widgets/Input/SEditableTextBox.h"
 #include "EdGraph/EdGraphNode.h"
 #include "NiagaraParameterCollection.h"
+#include "NiagaraScriptVariable.h"
 
 #define LOCTEXT_NAMESPACE "NiagaraNodeParameterMapGet"
 
@@ -87,6 +88,10 @@ UEdGraphPin* UNiagaraNodeParameterMapGet::CreateDefaultPin(UEdGraphPin* OutputPi
 
 	UEdGraphPin* DefaultPin = CreatePin(EEdGraphPinDirection::EGPD_Input, OutputPin->PinType, TEXT(""));
 
+	// we make the pin read only because the default value is set in the parameter panel unless the default mode is set to "custom" by the user
+	DefaultPin->bNotConnectable = true;
+	DefaultPin->bDefaultValueIsReadOnly = true;
+
 	const UEdGraphSchema_Niagara* Schema = GetDefault<UEdGraphSchema_Niagara>();
 	FNiagaraTypeDefinition NiagaraType = Schema->PinToTypeDefinition(OutputPin);
 	bool bNeedsValue = NiagaraType.IsDataInterface() == false;
@@ -97,6 +102,20 @@ UEdGraphPin* UNiagaraNodeParameterMapGet::CreateDefaultPin(UEdGraphPin* OutputPi
 	{
 		DefaultPin->DefaultValue = PinDefaultValue;
 	}
+
+	// If the variable of the new default pin is already in use in the graph we use the confiugured default value
+	if (UNiagaraGraph* Graph = Cast<UNiagaraGraph>(GetGraph())) {
+		if (UNiagaraScriptVariable* ScriptVar = Graph->GetScriptVariable(Var.GetName())) {
+			if (ScriptVar->DefaultMode == ENiagaraDefaultMode::Value) {
+				FNiagaraEditorModule& NiagaraEditorModule = FModuleManager::GetModuleChecked<FNiagaraEditorModule>("NiagaraEditor");
+				TSharedPtr<INiagaraEditorTypeUtilities, ESPMode::ThreadSafe> TypeEditorUtilities = NiagaraEditorModule.GetTypeUtilities(ScriptVar->Variable.GetType());
+				if (bNeedsValue && TypeEditorUtilities.IsValid() && TypeEditorUtilities->CanHandlePinDefaults())
+				{
+					DefaultPin->DefaultValue = TypeEditorUtilities->GetPinDefaultStringFromValue(ScriptVar->Variable);
+				}
+			}
+		}
+	}
 	
 	if (!OutputPin->PersistentGuid.IsValid())
 	{
@@ -108,10 +127,20 @@ UEdGraphPin* UNiagaraNodeParameterMapGet::CreateDefaultPin(UEdGraphPin* OutputPi
 	}
 	PinOutputToPinDefaultPersistentId.Add(OutputPin->PersistentGuid, DefaultPin->PersistentGuid);
 
-	SynchronizeDefaultInputPin(DefaultPin, OutputPin);
+	SynchronizeDefaultInputPin(DefaultPin, OutputPin, GetScriptVariable(OutputPin->PinName));
 	return DefaultPin;
 }
 
+UNiagaraScriptVariable* UNiagaraNodeParameterMapGet::GetScriptVariable(FName VariableName) const
+{
+	if (UNiagaraGraph* Graph = Cast<UNiagaraGraph>(GetGraph())) {
+		if (UNiagaraScriptVariable* ScriptVar = Graph->GetScriptVariable(VariableName)) {
+			return ScriptVar;
+		}
+	}
+	return nullptr;
+}
+
 void UNiagaraNodeParameterMapGet::OnPinRenamed(UEdGraphPin* RenamedPin, const FString& OldName)
 {
 	UNiagaraNodeParameterMapBase::OnPinRenamed(RenamedPin, OldName);
@@ -120,7 +149,7 @@ void UNiagaraNodeParameterMapGet::OnPinRenamed(UEdGraphPin* RenamedPin, const FS
 	if (DefaultPin)
 	{
 		DefaultPin->Modify();
-		SynchronizeDefaultInputPin(DefaultPin, RenamedPin);
+		SynchronizeDefaultInputPin(DefaultPin, RenamedPin, GetScriptVariable(RenamedPin->PinName));
 	}
 
 	MarkNodeRequiresSynchronization(__FUNCTION__, true);
@@ -244,7 +273,7 @@ void UNiagaraNodeParameterMapGet::PostLoad()
 		}
 		else
 		{
-			SynchronizeDefaultInputPin(InputPin, OutputPin);
+			SynchronizeDefaultInputPin(InputPin, OutputPin, GetScriptVariable(OutputPin->PinName));
 		}
 
 		OutputPin->PinType.PinSubCategory = UNiagaraNodeParameterMapBase::ParameterPinSubCategory;
@@ -252,7 +281,7 @@ void UNiagaraNodeParameterMapGet::PostLoad()
 }
 
 
-void UNiagaraNodeParameterMapGet::SynchronizeDefaultInputPin(UEdGraphPin* DefaultPin, UEdGraphPin* OutputPin)
+void UNiagaraNodeParameterMapGet::SynchronizeDefaultInputPin(UEdGraphPin* DefaultPin, UEdGraphPin* OutputPin, UNiagaraScriptVariable* ScriptVar)
 {
 	const UEdGraphSchema_Niagara* Schema = GetDefault<UEdGraphSchema_Niagara>();
 	if (!DefaultPin)
@@ -260,7 +289,8 @@ void UNiagaraNodeParameterMapGet::SynchronizeDefaultInputPin(UEdGraphPin* Defaul
 		return;
 	}
 
-	if (FNiagaraParameterMapHistory::IsEngineParameter(Schema->PinToNiagaraVariable(OutputPin)))
+	FNiagaraVariable Var = Schema->PinToNiagaraVariable(OutputPin);
+	if (FNiagaraParameterMapHistory::IsEngineParameter(Var))
 	{
 		DefaultPin->bDefaultValueIsIgnored = true;
 		DefaultPin->bNotConnectable = true;
@@ -269,11 +299,24 @@ void UNiagaraNodeParameterMapGet::SynchronizeDefaultInputPin(UEdGraphPin* Defaul
 	}
 	else
 	{
-		DefaultPin->bDefaultValueIsIgnored = false;
-		DefaultPin->bNotConnectable = false;
 		DefaultPin->bHidden = false;
 		DefaultPin->PinToolTip = FText::Format(LOCTEXT("DefaultValueTooltip_UnlessOverridden", "Default value for {0} if no other module has set it previously in the stack."), FText::FromName(OutputPin->PinName)).ToString();
 	}
+
+	// Sync pin visibility with the configured default mode
+	if (ScriptVar) {
+		if (ScriptVar->DefaultMode == ENiagaraDefaultMode::Value) {
+			DefaultPin->bNotConnectable = true;
+			DefaultPin->bDefaultValueIsReadOnly = true;
+		}
+		else if (ScriptVar->DefaultMode == ENiagaraDefaultMode::Custom) {
+			DefaultPin->bNotConnectable = false;
+			DefaultPin->bDefaultValueIsReadOnly = false;
+		}
+		else if (ScriptVar->DefaultMode == ENiagaraDefaultMode::Binding) {
+			DefaultPin->bHidden = true;
+		}
+	}
 }
 
 
diff --git a/Engine/Plugins/FX/Niagara/Source/NiagaraEditor/Private/NiagaraNodeParameterMapGet.h b/Engine/Plugins/FX/Niagara/Source/NiagaraEditor/Private/NiagaraNodeParameterMapGet.h
--- a/Engine/Plugins/FX/Niagara/Source/NiagaraEditor/Private/NiagaraNodeParameterMapGet.h
+++ b/Engine/Plugins/FX/Niagara/Source/NiagaraEditor/Private/NiagaraNodeParameterMapGet.h
@@ -51,7 +51,7 @@ protected:
 	virtual void RemoveDynamicPin(UEdGraphPin* Pin) override;
 
 	/** Make sure that descriptions match up as well as any other value that can be changed.*/
-	void SynchronizeDefaultInputPin(UEdGraphPin* DefaultPin, UEdGraphPin* OutputPin);
+	void SynchronizeDefaultInputPin(UEdGraphPin* DefaultPin, UEdGraphPin* OutputPin, UNiagaraScriptVariable* ScriptVar);
 	
 	/** Reverse the lookup from GetDefaultPin.*/
 	UEdGraphPin* GetOutputPinForDefault(const UEdGraphPin* DefaultPin) const;
@@ -62,6 +62,9 @@ protected:
 	/** Convenience method to determine whether this Node is a Map Get or Map Set when adding a parameter through the parameter panel. */
 	virtual EEdGraphPinDirection GetPinDirectionForNewParameters() { return EEdGraphPinDirection::EGPD_Output; };
 
+	/** Returns the metadata for the given variable or nullptr if no metadata is available */
+	UNiagaraScriptVariable* GetScriptVariable(FName VariableName) const;
+
 	UPROPERTY(meta = (SkipForCompileHash="true"))
 	TMap<FGuid, FGuid> PinOutputToPinDefaultPersistentId;
 };
diff --git a/Engine/Plugins/FX/Niagara/Source/NiagaraEditor/Public/NiagaraGraph.h b/Engine/Plugins/FX/Niagara/Source/NiagaraEditor/Public/NiagaraGraph.h
--- a/Engine/Plugins/FX/Niagara/Source/NiagaraEditor/Public/NiagaraGraph.h
+++ b/Engine/Plugins/FX/Niagara/Source/NiagaraEditor/Public/NiagaraGraph.h
@@ -254,6 +254,9 @@ class UNiagaraGraph : public UEdGraph
 	/** Walk through the graph for an ParameterMapGet nodes and see if any of them specify a default for VariableName.*/
 	UEdGraphPin* FindParameterMapDefaultValuePin(const FName VariableName, ENiagaraScriptUsage InUsage, ENiagaraScriptUsage InParentUsage) const;
 
+	/** Walk through the graph for an ParameterMapGet nodes and find all matching default pins for VariableName, irrespective of usage. */
+	TArray<UEdGraphPin*> FindParameterMapDefaultValuePins(const FName VariableName) const;
+
 	/** Gets the meta-data associated with this variable, if it exists.*/
 	TOptional<FNiagaraVariableMetaData> GetMetaData(const FNiagaraVariable& InVar) const;
 
@@ -340,6 +343,8 @@ class UNiagaraGraph : public UEdGraph
 
 	FDelegateHandle RegisterPinVisualWidgetProvider(FOnGetPinVisualWidget OnGetPinVisualWidget);
 
+	void ScriptVariableChanged(FNiagaraVariable Variable);
+
 protected:
 	void RebuildNumericCache();
 	bool bNeedNumericCacheRebuilt;
@@ -375,6 +380,9 @@ private:
 
 	/** Find all nodes in the graph that can be reached during compilation. */
 	TArray<UEdGraphNode*> FindReachbleNodes() const;
+
+	/** Compares the values on the default pins with the metadata and syncs the two if necessary */
+	void ValidateDefaultPins();
 	
 	/** The current change identifier for this graph overall. Used to sync status with UNiagaraScripts.*/
 	UPROPERTY()
