Change loctext key to fix uniqueness warning.

diff --git a/Engine/Plugins/FX/Niagara/Source/NiagaraEditor/Private/ViewModels/NiagaraParameterPanelViewModel.cpp b/Engine/Plugins/FX/Niagara/Source/NiagaraEditor/Private/ViewModels/NiagaraParameterPanelViewModel.cpp
--- a/Engine/Plugins/FX/Niagara/Source/NiagaraEditor/Private/ViewModels/NiagaraParameterPanelViewModel.cpp
+++ b/Engine/Plugins/FX/Niagara/Source/NiagaraEditor/Private/ViewModels/NiagaraParameterPanelViewModel.cpp
@@ -544,7 +544,7 @@ NiagaraParameterPanelSectionID::Type FNiagaraScriptToolkitParameterPanelViewMode
 
 void FNiagaraScriptToolkitParameterPanelViewModel::AddParameter(const FNiagaraVariable& InVariableToAdd, const FNiagaraVariableMetaData& InVariableMetaDataToAssign)
 {
-	FScopedTransaction AddParameter(LOCTEXT("AddParameter", "Add parameter and create new associated node"));
+	FScopedTransaction AddParameter(LOCTEXT("AddParameterFromParameterPanel", "Add Parameter"));
 	UNiagaraGraph::FAddParameterOptions AddParameterOptions = UNiagaraGraph::FAddParameterOptions();
 	AddParameterOptions.NewParameterScope = InVariableMetaDataToAssign.Scope;
 	AddParameterOptions.NewParameterUsage = InVariableMetaDataToAssign.Usage;
