Make sure that we never end up with a zero array size

diff --git a/Engine/Plugins/FX/Niagara/Source/Niagara/Private/NiagaraDataInterfaceParticleRead.cpp b/Engine/Plugins/FX/Niagara/Source/Niagara/Private/NiagaraDataInterfaceParticleRead.cpp
--- a/Engine/Plugins/FX/Niagara/Source/Niagara/Private/NiagaraDataInterfaceParticleRead.cpp
+++ b/Engine/Plugins/FX/Niagara/Source/Niagara/Private/NiagaraDataInterfaceParticleRead.cpp
@@ -1174,7 +1174,8 @@ void UNiagaraDataInterfaceParticleRead::GetParameterDefinitionHLSL(const FNiagar
 	// If we use an int array for the attribute indices, the shader compiler will actually use int4 due to the packing rules,
 	// and leave 3 elements unused. Besides being wasteful, this means that the array we send to the CS would need to be padded,
 	// which is a hassle. Instead, use int4 explicitly, and access individual components in the generated code.
-	const int AttributeInt4Count = FMath::DivideAndRoundUp(ParamInfo.GeneratedFunctions.Num(), 4);
+	// Note that we have to have at least one here because hlsl doesn't support arrays of size 0.
+	const int AttributeInt4Count = FMath::Max(1, FMath::DivideAndRoundUp(ParamInfo.GeneratedFunctions.Num(), 4));
 
 	TMap<FString, FStringFormatArg> ArgsDeclarations;
 	ArgsDeclarations.Add(TEXT("NumSpawnedParticlesName"), NumSpawnedParticlesBaseName + ParamInfo.DataInterfaceHLSLSymbol);
