Preventing a crash when querying for missing data.

diff --git a/Engine/Plugins/FX/Niagara/Source/Niagara/Private/NiagaraComponent.cpp b/Engine/Plugins/FX/Niagara/Source/Niagara/Private/NiagaraComponent.cpp
--- a/Engine/Plugins/FX/Niagara/Source/Niagara/Private/NiagaraComponent.cpp
+++ b/Engine/Plugins/FX/Niagara/Source/Niagara/Private/NiagaraComponent.cpp
@@ -1644,11 +1644,23 @@ TArray<FVector> UNiagaraComponent::GetNiagaraParticleValueVec3_DebugOnly(const F
 				int32 NumParticles = ParticleData.GetNumInstances();
 				Positions.SetNum(NumParticles);
 				FNiagaraDataSetAccessor<FVector> PosData(Sim->GetData(), FNiagaraVariable(FNiagaraTypeDefinition::GetVec3Def(), *InValueName));
-				for (int32 i = 0; i < NumParticles; ++i)
+
+				if (PosData.IsValidForRead())
 				{
-					FVector Position;
-					PosData.Get(i, Position);
-					Positions[i] = Position;
+					for (int32 i = 0; i < NumParticles; ++i)
+					{
+						FVector Position;
+						PosData.Get(i, Position);
+						Positions[i] = Position;
+					}
+				}
+				else
+				{
+					UE_LOG(LogNiagara, Warning, TEXT("Unabled to find variable %s on %s per-particle data. Returning zeroes."), *InValueName, *GetPathName());
+					for (int32 i = 0; i < NumParticles; ++i)
+					{
+						Positions[i] = FVector::ZeroVector;
+					}
 				}
 			}
 		}
