Fixed a bug that could nead to a crash when a variable default value was not initialized

diff --git a/Engine/Plugins/FX/Niagara/Source/NiagaraEditor/Private/NiagaraNodeParameterMapGet.cpp b/Engine/Plugins/FX/Niagara/Source/NiagaraEditor/Private/NiagaraNodeParameterMapGet.cpp
--- a/Engine/Plugins/FX/Niagara/Source/NiagaraEditor/Private/NiagaraNodeParameterMapGet.cpp
+++ b/Engine/Plugins/FX/Niagara/Source/NiagaraEditor/Private/NiagaraNodeParameterMapGet.cpp
@@ -103,10 +103,10 @@ UEdGraphPin* UNiagaraNodeParameterMapGet::CreateDefaultPin(UEdGraphPin* OutputPi
 		DefaultPin->DefaultValue = PinDefaultValue;
 	}
 
-	// If the variable of the new default pin is already in use in the graph we use the confiugured default value
+	// If the variable of the new default pin is already in use in the graph we use the configured default value
 	if (UNiagaraGraph* Graph = Cast<UNiagaraGraph>(GetGraph())) {
 		if (UNiagaraScriptVariable* ScriptVar = Graph->GetScriptVariable(Var.GetName())) {
-			if (ScriptVar->DefaultMode == ENiagaraDefaultMode::Value) {
+			if (ScriptVar->DefaultMode == ENiagaraDefaultMode::Value && ScriptVar->Variable.IsValid() && ScriptVar->Variable.IsDataAllocated()) {
 				FNiagaraEditorModule& NiagaraEditorModule = FModuleManager::GetModuleChecked<FNiagaraEditorModule>("NiagaraEditor");
 				TSharedPtr<INiagaraEditorTypeUtilities, ESPMode::ThreadSafe> TypeEditorUtilities = NiagaraEditorModule.GetTypeUtilities(ScriptVar->Variable.GetType());
 				if (bNeedsValue && TypeEditorUtilities.IsValid() && TypeEditorUtilities->CanHandlePinDefaults())
