Changed log category from error to warning for default validation message

diff --git a/Engine/Plugins/FX/Niagara/Source/NiagaraEditor/Private/NiagaraGraph.cpp b/Engine/Plugins/FX/Niagara/Source/NiagaraEditor/Private/NiagaraGraph.cpp
--- a/Engine/Plugins/FX/Niagara/Source/NiagaraEditor/Private/NiagaraGraph.cpp
+++ b/Engine/Plugins/FX/Niagara/Source/NiagaraEditor/Private/NiagaraGraph.cpp
@@ -460,7 +460,7 @@ void UNiagaraGraph::ValidateDefaultPins()
 		}
 
 		if (!IsCustom && !IsConsistent) {
-			UE_LOG(LogNiagaraEditor, Error, TEXT("Niagara graph %s: The default value declaration for the variable '%s' is not consistent between the graph and the metadata.\n  Either change the default mode to 'custom' or check the input pins in the parameter map get node in the graph."),
+			UE_LOG(LogNiagaraEditor, Warning, TEXT("Niagara graph %s: The default value declaration for the variable '%s' is not consistent between the graph and the metadata.\n  Either change the default mode to 'custom' or check the input pins in the parameter map get node in the graph."),
 				*GetFullName(), *Variable.GetName().ToString());
 			continue;
 		}
