Making it impossible to make numeric attributes as numeric is a generic conversion type, not a physical storage type.

diff --git a/Engine/Plugins/FX/Niagara/Source/NiagaraEditor/Private/NiagaraCompiler.cpp b/Engine/Plugins/FX/Niagara/Source/NiagaraEditor/Private/NiagaraCompiler.cpp
--- a/Engine/Plugins/FX/Niagara/Source/NiagaraEditor/Private/NiagaraCompiler.cpp
+++ b/Engine/Plugins/FX/Niagara/Source/NiagaraEditor/Private/NiagaraCompiler.cpp
@@ -460,7 +460,10 @@ void FNiagaraCompileRequestData::FinishPrecompile(UNiagaraScriptSource* ScriptSo
 				History.OriginatingScriptUsage = FoundOutputNode->GetUsage();
 				for (FNiagaraVariable& Var : History.Variables)
 				{
-					check(Var.GetType() != FNiagaraTypeDefinition::GetGenericNumericDef());
+					if (Var.GetType() == FNiagaraTypeDefinition::GetGenericNumericDef())
+					{
+						UE_LOG(LogNiagaraEditor, Log, TEXT("Invalid numeric parameter found! %s"), *Var.GetName().ToString())
+					}
 				}
 			}
 			if (FoundOutputNode->GetUsage() == ENiagaraScriptUsage::ParticleSimulationStageScript)
diff --git a/Engine/Plugins/FX/Niagara/Source/NiagaraEditor/Private/Widgets/SNiagaraParameterMapView.cpp b/Engine/Plugins/FX/Niagara/Source/NiagaraEditor/Private/Widgets/SNiagaraParameterMapView.cpp
--- a/Engine/Plugins/FX/Niagara/Source/NiagaraEditor/Private/Widgets/SNiagaraParameterMapView.cpp
+++ b/Engine/Plugins/FX/Niagara/Source/NiagaraEditor/Private/Widgets/SNiagaraParameterMapView.cpp
@@ -308,12 +308,16 @@ void SNiagaraParameterMapView::AddParameter(FNiagaraVariable NewVariable)
 		GraphActionMenu->OnRequestRenameOnActionNode();
 	}
 }
-
-bool SNiagaraParameterMapView::AllowMakeType(const FNiagaraTypeDefinition& InType) const
+bool SNiagaraParameterMapView::AllowMakeTypeGeneric(const FNiagaraTypeDefinition& InType) const
 {
 	return InType != FNiagaraTypeDefinition::GetParameterMapDef();
 }
 
+bool SNiagaraParameterMapView::AllowMakeTypeAttribute(const FNiagaraTypeDefinition& InType) const
+{
+	return InType != FNiagaraTypeDefinition::GetParameterMapDef() && InType != FNiagaraTypeDefinition::GetGenericNumericDef();
+}
+
 void SNiagaraParameterMapView::OnFilterTextChanged(const FText& InFilterText)
 {
 	GraphActionMenu->GenerateFilteredItems(false);
@@ -574,9 +578,11 @@ bool SNiagaraParameterMapView::SelectionHasContextMenu() const
 
 TSharedRef<SWidget> SNiagaraParameterMapView::OnGetParameterMenu(const NiagaraParameterMapSectionID::Type InSection)
 {
+	bool bTypeIsAttribute = true; 	// Leaving around the old generic path in case it is needed in the future.
+
 	TSharedRef<SNiagaraAddParameterMenu> MenuWidget = SNew(SNiagaraAddParameterMenu, Graphs)
 		.OnAddParameter(this, &SNiagaraParameterMapView::AddParameter)
-		.OnAllowMakeType(this, &SNiagaraParameterMapView::AllowMakeType)
+		.OnAllowMakeType(this, bTypeIsAttribute ? &SNiagaraParameterMapView::AllowMakeTypeAttribute : &SNiagaraParameterMapView::AllowMakeTypeGeneric)
 		.Section(InSection)
 		.ShowNamespaceCategory(false)
 		.ShowGraphParameters(false)
diff --git a/Engine/Plugins/FX/Niagara/Source/NiagaraEditor/Private/Widgets/SNiagaraParameterMapView.h b/Engine/Plugins/FX/Niagara/Source/NiagaraEditor/Private/Widgets/SNiagaraParameterMapView.h
--- a/Engine/Plugins/FX/Niagara/Source/NiagaraEditor/Private/Widgets/SNiagaraParameterMapView.h
+++ b/Engine/Plugins/FX/Niagara/Source/NiagaraEditor/Private/Widgets/SNiagaraParameterMapView.h
@@ -94,8 +94,11 @@ public:
 	static bool IsStaticSwitchParameter(const FNiagaraVariable& Variable, const TArray<TWeakObjectPtr<UNiagaraGraph>>& Graphs);
 
 private:
-	/** Function to bind to SNiagaraAddParameterMenus to filter types we allow creating */
-	bool AllowMakeType(const FNiagaraTypeDefinition& InType) const;
+	/** Function to bind to SNiagaraAddParameterMenus to filter types we allow creating in generic parameters*/
+	bool AllowMakeTypeGeneric(const FNiagaraTypeDefinition& InType) const;
+
+	/** Function to bind to SNiagaraAddParameterMenus to filter types we allow creating for attributes */
+	bool AllowMakeTypeAttribute(const FNiagaraTypeDefinition& InType) const;
 
 	/** Callback when the filter is changed, forces the action tree(s) to filter */
 	void OnFilterTextChanged(const FText& InFilterText);
diff --git a/Engine/Plugins/FX/Niagara/Source/NiagaraEditor/Private/Widgets/SNiagaraParameterPanel.cpp b/Engine/Plugins/FX/Niagara/Source/NiagaraEditor/Private/Widgets/SNiagaraParameterPanel.cpp
--- a/Engine/Plugins/FX/Niagara/Source/NiagaraEditor/Private/Widgets/SNiagaraParameterPanel.cpp
+++ b/Engine/Plugins/FX/Niagara/Source/NiagaraEditor/Private/Widgets/SNiagaraParameterPanel.cpp
@@ -108,7 +108,7 @@ void SNiagaraParameterPanel::AddParameter(FNiagaraVariable NewVariable, const Ni
 
 bool SNiagaraParameterPanel::AllowMakeType(const FNiagaraTypeDefinition& InType) const
 {
-	return InType != FNiagaraTypeDefinition::GetParameterMapDef();
+	return InType != FNiagaraTypeDefinition::GetParameterMapDef() && InType != FNiagaraTypeDefinition::GetGenericNumericDef();
 }
 
 TSharedRef<SWidget> SNiagaraParameterPanel::OnCreateWidgetForAction(struct FCreateWidgetForActionData* const InCreateData)
