Send a property change delegate when updating the niagara thumbnail.
#rnx

diff --git a/Engine/Plugins/FX/Niagara/Source/NiagaraEditor/Private/Toolkits/NiagaraSystemToolkit.cpp b/Engine/Plugins/FX/Niagara/Source/NiagaraEditor/Private/Toolkits/NiagaraSystemToolkit.cpp
--- a/Engine/Plugins/FX/Niagara/Source/NiagaraEditor/Private/Toolkits/NiagaraSystemToolkit.cpp
+++ b/Engine/Plugins/FX/Niagara/Source/NiagaraEditor/Private/Toolkits/NiagaraSystemToolkit.cpp
@@ -787,6 +787,9 @@ void FNiagaraSystemToolkit::OnThumbnailCaptured(UTexture2D* Thumbnail)
 	{
 		System->MarkPackageDirty();
 		System->ThumbnailImage = Thumbnail;
+		// Broadcast an object property changed event to update the content browser
+		FPropertyChangedEvent EmptyPropertyChangedEvent(nullptr);
+		FCoreUObjectDelegates::OnObjectPropertyChanged.Broadcast(System, EmptyPropertyChangedEvent);
 	}
 	else if (SystemToolkitMode == ESystemToolkitMode::Emitter) 
 	{
@@ -794,6 +797,9 @@ void FNiagaraSystemToolkit::OnThumbnailCaptured(UTexture2D* Thumbnail)
 		UNiagaraEmitter* EditableEmitter = EditableEmitterViewModel->GetEmitter();
 		EditableEmitter->ThumbnailImage = Thumbnail;
 		bEmitterThumbnailUpdated = true;
+		// Broadcast an object property changed event to update the content browser
+		FPropertyChangedEvent EmptyPropertyChangedEvent(nullptr);
+		FCoreUObjectDelegates::OnObjectPropertyChanged.Broadcast(EditableEmitter, EmptyPropertyChangedEvent);
 	}
 }
 
