Small change to harden the SkelMeshBoneSampling with Niagara
-SkinOnTheFly path no longer relies on accessing the Mesh (Mesh->RefSkeleton) for anything, instead using cached data for bone counts
-fixes off by one error if we were to try to reference the last bone (it would have resulted in a 0 vector)

diff --git a/Engine/Plugins/FX/Niagara/Source/Niagara/Classes/NDISkeletalMeshCommon.h b/Engine/Plugins/FX/Niagara/Source/Niagara/Classes/NDISkeletalMeshCommon.h
--- a/Engine/Plugins/FX/Niagara/Source/Niagara/Classes/NDISkeletalMeshCommon.h
+++ b/Engine/Plugins/FX/Niagara/Source/Niagara/Classes/NDISkeletalMeshCommon.h
@@ -46,6 +46,7 @@ void FSkeletalMeshAccessorHelper::Init<
 template<typename SkinningMode>
 struct FSkinnedPositionAccessorHelper
 {
+	FORCEINLINE int32 GetBoneCount(FSkeletalMeshAccessorHelper& Accessor, bool RequiresPrevious) const = delete;
 	FORCEINLINE void GetTriangleIndices(int32 Tri, int32& Idx0, int32& Idx1, int32& Idx2) = delete;
 	FORCEINLINE void GetSkinnedTrianglePositions(FSkeletalMeshAccessorHelper& Accessor, int32 Tri, FVector& OutPos0, FVector& OutPos1, FVector& OutPos2, FVector& OutPrev0, FVector& OutPrev1, FVector& OutPrev2, int32& Idx0, int32& Idx1, int32& Idx2) = delete;
 	FORCEINLINE void GetSkinnedTrianglePositions(FSkeletalMeshAccessorHelper& Accessor, int32 Tri, FVector& OutPos0, FVector& OutPos1, FVector& OutPos2, int32& Idx0, int32& Idx1, int32& Idx2) = delete;
@@ -60,6 +61,12 @@ struct FSkinnedPositionAccessorHelper
 template<>
 struct FSkinnedPositionAccessorHelper<TIntegralConstant<ENDISkeletalMesh_SkinningMode, ENDISkeletalMesh_SkinningMode::None>>
 {
+	FORCEINLINE int32 GetBoneCount(FSkeletalMeshAccessorHelper& Accessor, bool RequiresPrevious) const
+	{
+		const FReferenceSkeleton& RefSkel = Accessor.Mesh->RefSkeleton;
+		return RefSkel.GetNum();
+	}
+
 	FORCEINLINE void GetTriangleIndices(FSkeletalMeshAccessorHelper& Accessor, int32 Tri, int32& Idx0, int32& Idx1, int32& Idx2)
 	{
 		const int32 BaseIndex = Tri * 3;
@@ -117,6 +124,11 @@ struct FSkinnedPositionAccessorHelper<TIntegralConstant<ENDISkeletalMesh_Skinnin
 template<>
 struct FSkinnedPositionAccessorHelper<TIntegralConstant<ENDISkeletalMesh_SkinningMode, ENDISkeletalMesh_SkinningMode::SkinOnTheFly>>
 {
+	FORCEINLINE int32 GetBoneCount(FSkeletalMeshAccessorHelper& Accessor, bool RequiresPrevious) const
+	{
+		return Accessor.SkinningData->GetBoneCount(RequiresPrevious);
+	}
+
 	FORCEINLINE void GetTriangleIndices(FSkeletalMeshAccessorHelper& Accessor, int32 Tri, int32& Idx0, int32& Idx1, int32& Idx2)
 	{
 		const int32 BaseIndex = Tri * 3;
@@ -174,6 +186,11 @@ struct FSkinnedPositionAccessorHelper<TIntegralConstant<ENDISkeletalMesh_Skinnin
 template<>
 struct FSkinnedPositionAccessorHelper<TIntegralConstant<ENDISkeletalMesh_SkinningMode, ENDISkeletalMesh_SkinningMode::PreSkin>>
 {
+	FORCEINLINE int32 GetBoneCount(FSkeletalMeshAccessorHelper& Accessor, bool RequiresPrevious) const
+	{
+		return Accessor.SkinningData->GetBoneCount(RequiresPrevious);
+	}
+
 	FORCEINLINE void GetTriangleIndices(FSkeletalMeshAccessorHelper& Accessor, int32 Tri, int32& Idx0, int32& Idx1, int32& Idx2)
 	{
 		const int32 BaseIndex = Tri * 3;
diff --git a/Engine/Plugins/FX/Niagara/Source/Niagara/Classes/NiagaraDataInterfaceSkeletalMesh.h b/Engine/Plugins/FX/Niagara/Source/Niagara/Classes/NiagaraDataInterfaceSkeletalMesh.h
--- a/Engine/Plugins/FX/Niagara/Source/Niagara/Classes/NiagaraDataInterfaceSkeletalMesh.h
+++ b/Engine/Plugins/FX/Niagara/Source/Niagara/Classes/NiagaraDataInterfaceSkeletalMesh.h
@@ -95,6 +95,7 @@ struct FSkeletalMeshSkinningData
 		, DeltaSeconds(.0333f)
 		, CurrIndex(0)
 		, BoneMatrixUsers(0)
+		, bForceDataRefresh(false)
 	{}
 
 	void RegisterUser(FSkeletalMeshSkinningDataUsage Usage);
@@ -104,12 +105,24 @@ struct FSkeletalMeshSkinningData
 
 	bool Tick(float InDeltaSeconds, bool bRequirePreskin = true);
 
-	FORCEINLINE FVector GetPosition(int32 LODIndex, int32 VertexIndex)const
+	FORCEINLINE int32 GetBoneCount(bool RequiresPrevious) const
+	{
+		int32 BoneCount = FMath::Min(CurrBoneRefToLocals().Num(), CurrComponentTransforms().Num());
+
+		if (RequiresPrevious)
+		{
+			BoneCount = FMath::Min3(BoneCount, PrevBoneRefToLocals().Num(), PrevComponentTransforms().Num());
+		}
+
+		return BoneCount;
+	}
+
+	FORCEINLINE FVector GetPosition(int32 LODIndex, int32 VertexIndex) const
 	{
 		return LODData[LODIndex].SkinnedCPUPositions[CurrIndex][VertexIndex];
 	}
 
-	FORCEINLINE FVector GetPreviousPosition(int32 LODIndex, int32 VertexIndex)const
+	FORCEINLINE FVector GetPreviousPosition(int32 LODIndex, int32 VertexIndex) const
 	{
 		return LODData[LODIndex].SkinnedCPUPositions[CurrIndex ^ 1][VertexIndex];
 	}
@@ -129,21 +142,41 @@ struct FSkeletalMeshSkinningData
 		return BoneRefToLocals[CurrIndex];
 	}
 
+	FORCEINLINE const TArray<FMatrix>& CurrBoneRefToLocals() const
+	{
+		return BoneRefToLocals[CurrIndex];
+	}
+
 	FORCEINLINE TArray<FMatrix>& PrevBoneRefToLocals()
 	{
 		return BoneRefToLocals[CurrIndex ^ 1];
 	}
 
+	FORCEINLINE const TArray<FMatrix>& PrevBoneRefToLocals() const
+	{
+		return BoneRefToLocals[CurrIndex ^ 1];
+	}
+
 	FORCEINLINE TArray<FTransform>& CurrComponentTransforms()
 	{
 		return ComponentTransforms[CurrIndex];
 	}
 
+	FORCEINLINE const TArray<FTransform>& CurrComponentTransforms() const
+	{
+		return ComponentTransforms[CurrIndex];
+	}
+
 	FORCEINLINE TArray<FTransform>& PrevComponentTransforms()
 	{
 		return ComponentTransforms[CurrIndex ^ 1];
 	}
 
+	FORCEINLINE const TArray<FTransform>& PrevComponentTransforms() const
+	{
+		return ComponentTransforms[CurrIndex ^ 1];
+	}
+
 private:
 
 	void UpdateBoneTransforms();
@@ -404,8 +437,6 @@ struct FNDISkeletalMesh_InstanceData
 
 	//Cached ComponentToWorld from previous tick.
 	FMatrix PrevTransform;
-	//InverseTranspose of above for transforming normals/tangents.
-	FMatrix PrevTransformInverseTransposed;
 
 	/** Time separating Transform and PrevTransform. */
 	float DeltaSeconds;
diff --git a/Engine/Plugins/FX/Niagara/Source/Niagara/Classes/NiagaraDataInterfaceStaticMesh.h b/Engine/Plugins/FX/Niagara/Source/Niagara/Classes/NiagaraDataInterfaceStaticMesh.h
--- a/Engine/Plugins/FX/Niagara/Source/Niagara/Classes/NiagaraDataInterfaceStaticMesh.h
+++ b/Engine/Plugins/FX/Niagara/Source/Niagara/Classes/NiagaraDataInterfaceStaticMesh.h
@@ -117,8 +117,6 @@ struct FNDIStaticMesh_InstanceData
 
 	//Cached ComponentToWorld from previous tick.
 	FMatrix PrevTransform;
-	//InverseTranspose of above for transforming normals/tangents.
-	FMatrix PrevTransformInverseTransposed;
 
 	/** Time separating Transform and PrevTransform. */
 	float DeltaSeconds;
diff --git a/Engine/Plugins/FX/Niagara/Source/Niagara/Private/NDISkeletalMesh_BoneSampling.cpp b/Engine/Plugins/FX/Niagara/Source/Niagara/Private/NDISkeletalMesh_BoneSampling.cpp
--- a/Engine/Plugins/FX/Niagara/Source/Niagara/Private/NDISkeletalMesh_BoneSampling.cpp
+++ b/Engine/Plugins/FX/Niagara/Source/Niagara/Private/NDISkeletalMesh_BoneSampling.cpp
@@ -460,13 +460,10 @@ void UNiagaraDataInterfaceSkeletalMesh::GetSkinnedBoneData(FVectorVMContext& Con
 	FSkeletalMeshAccessorHelper Accessor;
 	Accessor.Init<TIntegralConstant<int32, 0>, TIntegralConstant<int32, 0>>(InstData);
 
-	const FReferenceSkeleton& RefSkel = Accessor.Mesh->RefSkeleton;
-
-	const int32 BoneMax = RefSkel.GetNum() - 1;
-	const int32 BoneAndSocketMax = BoneMax + InstData->SpecificSockets.Num();
+	const int32 BoneCount = SkinningHandler.GetBoneCount(Accessor, bInterpolated::Value);
+	const int32 BoneAndSocketCount = BoneCount + InstData->SpecificSockets.Num();
 	float InvDt = 1.0f / InstData->DeltaSeconds;
 
-	const int32 SpecificSocketBoneOffset = InstData->SpecificSocketBoneOffset;
 	const TArray<FTransform>& SpecificSocketCurrTransforms = InstData->GetSpecificSocketsCurrBuffer();
 	const TArray<FTransform>& SpecificSocketPrevTransforms = InstData->GetSpecificSocketsPrevBuffer();
 
@@ -476,17 +473,17 @@ void UNiagaraDataInterfaceSkeletalMesh::GetSkinnedBoneData(FVectorVMContext& Con
 
 		// Determine bone or socket
 		int32 RawBone = BoneParam.GetAndAdvance();
-		int32 Bone = FMath::Clamp(RawBone, 0, BoneAndSocketMax);
-		const bool bIsSocket = Bone > BoneMax;
-		const int32 Socket = Bone - SpecificSocketBoneOffset;
+		int32 Bone = FMath::Clamp(RawBone, 0, BoneAndSocketCount - 1);
+		const bool bIsSocket = Bone >= BoneCount;
+		const int32 Socket = Bone - BoneCount;
 
 		FVector Pos;
 		FVector Prev;
 		FVector Velocity;
 
 		// Handle edge cases first...
-		if ((!bIsSocket && Bone >= BoneMax) ||
-			(bIsSocket && (Socket >= SpecificSocketCurrTransforms.Num() || Socket < 0)))
+		if ((!bIsSocket && Bone >= BoneCount) ||
+			(bIsSocket && !SpecificSocketCurrTransforms.IsValidIndex(Socket)))
 		{
 			Pos = FVector::ZeroVector;
 			TransformHandler.TransformPosition(Pos, InstanceTransform);
@@ -607,10 +604,9 @@ void UNiagaraDataInterfaceSkeletalMesh::GetSpecificSocketBoneAt(FVectorVMContext
 	VectorVM::FUserPtrHandler<FNDISkeletalMesh_InstanceData> InstData(Context);
 
 	VectorVM::FExternalFuncRegisterHandler<int32> OutSocketBone(Context);
-	const TArray<FName>& SpecificSocketsArray = InstData->SpecificSockets;
 	const int32 SpecificSocketBoneOffset = InstData->SpecificSocketBoneOffset;
+	const int32 Max = SpecificSockets.Num() - 1;
 
-	int32 Max = SpecificSockets.Num() - 1;
 	if (Max != INDEX_NONE)
 	{
 		for (int32 i = 0; i < Context.NumInstances; ++i)
diff --git a/Engine/Plugins/FX/Niagara/Source/Niagara/Private/NiagaraDataInterfaceSkeletalMesh.cpp b/Engine/Plugins/FX/Niagara/Source/Niagara/Private/NiagaraDataInterfaceSkeletalMesh.cpp
--- a/Engine/Plugins/FX/Niagara/Source/Niagara/Private/NiagaraDataInterfaceSkeletalMesh.cpp
+++ b/Engine/Plugins/FX/Niagara/Source/Niagara/Private/NiagaraDataInterfaceSkeletalMesh.cpp
@@ -1232,7 +1232,6 @@ bool FNDISkeletalMesh_InstanceData::Init(UNiagaraDataInterfaceSkeletalMesh* Inte
 	Transform = FMatrix::Identity;
 	TransformInverseTransposed = FMatrix::Identity;
 	PrevTransform = FMatrix::Identity;
-	PrevTransformInverseTransposed = FMatrix::Identity;
 	DeltaSeconds = SystemInstance->GetComponent()->GetWorld()->GetDeltaSeconds();
 	ChangeId = Interface->ChangeId;
 	bIsGpuUniformlyDistributedSampling = false;
@@ -1268,7 +1267,6 @@ bool FNDISkeletalMesh_InstanceData::Init(UNiagaraDataInterfaceSkeletalMesh* Inte
 	Transform = Component->GetComponentToWorld().ToMatrixWithScale();
 	TransformInverseTransposed = Transform.InverseFast().GetTransposed();
 	PrevTransform = Transform;
-	PrevTransformInverseTransposed = TransformInverseTransposed;
 	
 	CachedAttachParent = Component->GetAttachParent();
 
@@ -1598,14 +1596,12 @@ bool FNDISkeletalMesh_InstanceData::Tick(UNiagaraDataInterfaceSkeletalMesh* Inte
 		if (Component.IsValid() && Mesh)
 		{
 			PrevTransform = Transform;
-			PrevTransformInverseTransposed = TransformInverseTransposed;
 			Transform = Component->GetComponentToWorld().ToMatrixWithScale();
 			TransformInverseTransposed = Transform.Inverse().GetTransposed();
 		}
 		else
 		{
 			PrevTransform = FMatrix::Identity;
-			PrevTransformInverseTransposed = FMatrix::Identity;
 			Transform = FMatrix::Identity;
 			TransformInverseTransposed = FMatrix::Identity;
 		}
diff --git a/Engine/Plugins/FX/Niagara/Source/Niagara/Private/NiagaraDataInterfaceStaticMesh.cpp b/Engine/Plugins/FX/Niagara/Source/Niagara/Private/NiagaraDataInterfaceStaticMesh.cpp
--- a/Engine/Plugins/FX/Niagara/Source/Niagara/Private/NiagaraDataInterfaceStaticMesh.cpp
+++ b/Engine/Plugins/FX/Niagara/Source/Niagara/Private/NiagaraDataInterfaceStaticMesh.cpp
@@ -172,7 +172,6 @@ bool FNDIStaticMesh_InstanceData::Init(UNiagaraDataInterfaceStaticMesh* Interfac
 	Transform = FMatrix::Identity;
 	TransformInverseTransposed = FMatrix::Identity;
 	PrevTransform = FMatrix::Identity;
-	PrevTransformInverseTransposed = FMatrix::Identity;
 	DeltaSeconds = 0.0f;
 	ChangeId = Interface->ChangeId;
 
@@ -253,7 +252,6 @@ bool FNDIStaticMesh_InstanceData::Init(UNiagaraDataInterfaceStaticMesh* Interfac
 	if (SafeComponent_GT.IsValid() && Mesh)
 	{
 		PrevTransform = Transform;
-		PrevTransformInverseTransposed = TransformInverseTransposed;
 		Transform = SafeComponent_GT->GetComponentToWorld().ToMatrixWithScale();
 		TransformInverseTransposed = Transform.InverseFast().GetTransposed();
 	}
@@ -362,14 +360,12 @@ bool FNDIStaticMesh_InstanceData::Tick(UNiagaraDataInterfaceStaticMesh* Interfac
 		if (SafeComponent_GT.IsValid() && SafeMesh_GT.IsValid())
 		{
 			PrevTransform = Transform;
-			PrevTransformInverseTransposed = TransformInverseTransposed;
 			Transform = SafeComponent_GT->GetComponentToWorld().ToMatrixWithScale();
 			TransformInverseTransposed = Transform.InverseFast().GetTransposed();
 		}
 		else
 		{
 			PrevTransform = FMatrix::Identity;
-			PrevTransformInverseTransposed = FMatrix::Identity;
 			Transform = FMatrix::Identity;
 			TransformInverseTransposed = FMatrix::Identity;
 		}
