Logging the SetEmitterEnable is currently unimplemented.

diff --git a/Engine/Plugins/FX/Niagara/Source/Niagara/Private/NiagaraSystemInstance.cpp b/Engine/Plugins/FX/Niagara/Source/Niagara/Private/NiagaraSystemInstance.cpp
--- a/Engine/Plugins/FX/Niagara/Source/Niagara/Private/NiagaraSystemInstance.cpp
+++ b/Engine/Plugins/FX/Niagara/Source/Niagara/Private/NiagaraSystemInstance.cpp
@@ -115,6 +115,11 @@ void FNiagaraSystemInstance::SetEmitterEnable(FName EmitterName, bool bNewEnable
 {
 	WaitForAsyncTickAndFinalize();
 
+
+	UE_LOG(LogNiagara, Warning, TEXT("SetEmitterEnable: Emitter \"%s\" is not currently implemented."), *EmitterName.ToString());
+	return;
+
+	/*
 	UNiagaraSystem* System = GetSystem();
 	if (System != nullptr)
 	{
@@ -155,7 +160,7 @@ void FNiagaraSystemInstance::SetEmitterEnable(FName EmitterName, bool bNewEnable
 		{
 			UE_LOG(LogNiagara, Log, TEXT("SetEmitterEnable: Emitter \"%s\" was not found in the system's list of emitters!"), *EmitterName.ToString());
 		}
-	}
+	}*/
 }
 
 
