- Fix issue where waiting on the system simulation task we execute the finalize which completes and removes the instance while we are attempting to remove the instance

diff --git a/Engine/Plugins/FX/Niagara/Source/Niagara/Private/NiagaraSystemSimulation.cpp b/Engine/Plugins/FX/Niagara/Source/Niagara/Private/NiagaraSystemSimulation.cpp
--- a/Engine/Plugins/FX/Niagara/Source/Niagara/Private/NiagaraSystemSimulation.cpp
+++ b/Engine/Plugins/FX/Niagara/Source/Niagara/Private/NiagaraSystemSimulation.cpp
@@ -1779,23 +1779,28 @@ void FNiagaraSystemSimulation::RemoveInstance(FNiagaraSystemInstance* Instance)
 		WaitForSystemTickComplete();
 		Instance->WaitForAsyncTickDoNotFinalize();
 
-		const int32 NumInstances = MainDataSet.GetCurrentDataChecked().GetNumInstances();
-		check(SystemInstances.Num() == NumInstances);
+		// There is a slim window where the finalize will have executed so we must ensure we have not been removed.
+		// This can happen where the async task is not complete, we start to wait and it posts the finalize task.  The TG will drain the GT queue which contains the finalize and we have been removed (via completion)
+		if ( Instance->SystemInstanceIndex != INDEX_NONE )
+		{
+			const int32 NumInstances = MainDataSet.GetCurrentDataChecked().GetNumInstances();
+			check(SystemInstances.Num() == NumInstances);
 
-		const int32 SystemIndex = Instance->SystemInstanceIndex;
-		check(Instance == SystemInstances[SystemIndex]);
-		check(SystemInstances.IsValidIndex(SystemIndex));
+			const int32 SystemIndex = Instance->SystemInstanceIndex;
+			check(Instance == SystemInstances[SystemIndex]);
+			check(SystemInstances.IsValidIndex(SystemIndex));
 
-		MainDataSet.GetCurrentDataChecked().KillInstance(SystemIndex);
-		SystemInstances.RemoveAtSwap(SystemIndex);
-		Instance->SystemInstanceIndex = INDEX_NONE;
-		if (SystemInstances.IsValidIndex(SystemIndex))
-		{
-			SystemInstances[SystemIndex]->SystemInstanceIndex = SystemIndex;
-		}
+			MainDataSet.GetCurrentDataChecked().KillInstance(SystemIndex);
+			SystemInstances.RemoveAtSwap(SystemIndex);
+			Instance->SystemInstanceIndex = INDEX_NONE;
+			if (SystemInstances.IsValidIndex(SystemIndex))
+			{
+				SystemInstances[SystemIndex]->SystemInstanceIndex = SystemIndex;
+			}
 
-		check(SystemInstances.Num() == MainDataSet.GetCurrentDataChecked().GetNumInstances());
-		check(PausedSystemInstances.Num() == PausedInstanceData.GetCurrentDataChecked().GetNumInstances());
+			check(SystemInstances.Num() == MainDataSet.GetCurrentDataChecked().GetNumInstances());
+			check(PausedSystemInstances.Num() == PausedInstanceData.GetCurrentDataChecked().GetNumInstances());
+		}
 	}
 
 #if NIAGARA_NAN_CHECKING
