Making it a compile error to have multiple event reads

diff --git a/Engine/Plugins/FX/Niagara/Source/NiagaraEditor/Private/NiagaraHlslTranslator.cpp b/Engine/Plugins/FX/Niagara/Source/NiagaraEditor/Private/NiagaraHlslTranslator.cpp
--- a/Engine/Plugins/FX/Niagara/Source/NiagaraEditor/Private/NiagaraHlslTranslator.cpp
+++ b/Engine/Plugins/FX/Niagara/Source/NiagaraEditor/Private/NiagaraHlslTranslator.cpp
@@ -212,7 +212,11 @@ FString FHlslNiagaraTranslator::GetCode(FNiagaraCodeChunk& Chunk)
 	{
 		if (DefinitionString.IsEmpty())
 		{
-			check(Chunk.bDecl);//Otherwise, we're doing nothing here.
+			if (!Chunk.bDecl)//Otherwise, we're doing nothing here.
+			{
+				Warning(LOCTEXT("MissingDeclForChunk", "Missing definition string."), nullptr, nullptr);
+			}
+
 			FinalString += GetStructHlslTypeName(Chunk.Type) + TEXT(" ") + Chunk.SymbolName + TEXT(";\n");
 		}
 		else
@@ -4932,6 +4936,11 @@ void FHlslNiagaraTranslator::ReadDataSet(const FNiagaraDataSetID DataSet, const
 	}
 
 	TMap<int32, FDataSetAccessInfo>& Reads = DataSetReadInfo[(int32)AccessMode].FindOrAdd(DataSet);
+	if (Reads.Num() != 0)
+	{
+		Error(FText::Format(LOCTEXT("TooManyDataSetReads", "Only one Event Read node per Event handler! Read data set node: \"{0}\""), FText::FromName(DataSet.Name)), nullptr, nullptr);
+	}
+
 	FDataSetAccessInfo* DataSetReadForInput = Reads.Find(InputChunk);
 	if (!DataSetReadForInput)
 	{
@@ -5571,7 +5580,7 @@ void FHlslNiagaraTranslator::HandleDataInterfaceCall(FNiagaraScriptDataInterface
 	{
 		Error(FText::Format(LOCTEXT("FunctionCallDataInterfaceGPUMissing", "Function call \"{0}\" does not work on GPU sims."), FText::FromName(InMatchingSignature.Name)), CurNode, nullptr);
 	}
-
+	
 	//UE_LOG(LogNiagaraEditor, Log, TEXT("HandleDataInterfaceCall %d %s %s %s"), ActiveStageIdx, *InMatchingSignature.Name.ToString(), InMatchingSignature.bWriteFunction ? TEXT("true") : TEXT("False"), *Info.Name.ToString());
 
 	if (InMatchingSignature.bWriteFunction && CompilationOutput.ScriptData.SimulationStageMetaData.Num() != 0 && TranslationStages[ActiveStageIdx].SourceSimStage != -1)
@@ -6628,14 +6637,18 @@ int32 FHlslNiagaraTranslator::CompileOutputPin(const UEdGraphPin* InPin)
 	return Ret;
 }
 
-void FHlslNiagaraTranslator::Error(FText ErrorText, const UNiagaraNode* Node, const UEdGraphPin* Pin)
+void FHlslNiagaraTranslator::Error(FText ErrorText, const UNiagaraNode* InNode, const UEdGraphPin* Pin)
 {
+
+	const UNiagaraNode* CurContextNode = ActiveHistoryForFunctionCalls.GetCallingContext();
+	const UNiagaraNode* TargetNode = InNode ? InNode : CurContextNode;
+
 	FString NodePinStr = TEXT("");
 	FString NodePinPrefix = TEXT(" - ");
 	FString NodePinSuffix = TEXT("");
-	if (Node)
+	if (TargetNode)
 	{
-		FString NodeTitle = Node->GetNodeTitle(ENodeTitleType::FullTitle).ToString();
+		FString NodeTitle = TargetNode->GetNodeTitle(ENodeTitleType::FullTitle).ToString();
 		if (NodeTitle.Len() > 0)
 		{
 			NodePinStr += TEXT("Node: ") + NodeTitle;
@@ -6643,7 +6656,7 @@ void FHlslNiagaraTranslator::Error(FText ErrorText, const UNiagaraNode* Node, co
 		}
 		else
 		{
-			FString NodeName = Node->GetName();
+			FString NodeName = TargetNode->GetName();
 			if (NodeName.Len() > 0)
 			{
 				NodePinStr += TEXT("Node: ") + NodeName;
@@ -6658,18 +6671,22 @@ void FHlslNiagaraTranslator::Error(FText ErrorText, const UNiagaraNode* Node, co
 	}
 
 	FString ErrorString = FString::Printf(TEXT("%s%s%s%s"), *ErrorText.ToString(), *NodePinPrefix, *NodePinStr, *NodePinSuffix);
-	TranslateResults.CompileEvents.Add(FNiagaraCompileEvent(FNiagaraCompileEventSeverity::Error, ErrorString, Node ? Node->NodeGuid : FGuid(), Pin ? Pin->PersistentGuid : FGuid(), GetCallstackGuids()));
+	TranslateResults.CompileEvents.Add(FNiagaraCompileEvent(FNiagaraCompileEventSeverity::Error, ErrorString, TargetNode ? TargetNode->NodeGuid : FGuid(), Pin ? Pin->PersistentGuid : FGuid(), GetCallstackGuids()));
 	TranslateResults.NumErrors++;
 }
 
-void FHlslNiagaraTranslator::Warning(FText WarningText, const UNiagaraNode* Node, const UEdGraphPin* Pin)
+void FHlslNiagaraTranslator::Warning(FText WarningText, const UNiagaraNode* InNode, const UEdGraphPin* Pin)
 {
+
+	const UNiagaraNode* CurContextNode = ActiveHistoryForFunctionCalls.GetCallingContext();
+	const UNiagaraNode* TargetNode = InNode ? InNode : CurContextNode;
+
 	FString NodePinStr = TEXT("");
 	FString NodePinPrefix = TEXT(" - ");
 	FString NodePinSuffix = TEXT("");
-	if (Node && Node->GetName().Len() > 0)
+	if (TargetNode && TargetNode->GetName().Len() > 0)
 	{
-		NodePinStr += TEXT("Node: ") + Node->GetName();
+		NodePinStr += TEXT("Node: ") + TargetNode->GetName();
 		NodePinSuffix = TEXT(" - ");
 	}
 	if (Pin && Pin->PinFriendlyName.ToString().Len() > 0)
@@ -6679,7 +6696,7 @@ void FHlslNiagaraTranslator::Warning(FText WarningText, const UNiagaraNode* Node
 	}
 
 	FString WarnString = FString::Printf(TEXT("%s%s%s%s"), *WarningText.ToString(), *NodePinPrefix, *NodePinStr, *NodePinSuffix);
-	TranslateResults.CompileEvents.Add(FNiagaraCompileEvent(FNiagaraCompileEventSeverity::Warning, WarnString, Node ? Node->NodeGuid : FGuid(), Pin ? Pin->PersistentGuid : FGuid(), GetCallstackGuids()));
+	TranslateResults.CompileEvents.Add(FNiagaraCompileEvent(FNiagaraCompileEventSeverity::Warning, WarnString, TargetNode ? TargetNode->NodeGuid : FGuid(), Pin ? Pin->PersistentGuid : FGuid(), GetCallstackGuids()));
 	TranslateResults.NumWarnings++;
 }
 
