Fixing compile warning when used in module/function/dynamic input script

diff --git a/Engine/Plugins/FX/Niagara/Source/NiagaraEditor/Private/NiagaraNodeSimTargetSelector.cpp b/Engine/Plugins/FX/Niagara/Source/NiagaraEditor/Private/NiagaraNodeSimTargetSelector.cpp
--- a/Engine/Plugins/FX/Niagara/Source/NiagaraEditor/Private/NiagaraNodeSimTargetSelector.cpp
+++ b/Engine/Plugins/FX/Niagara/Source/NiagaraEditor/Private/NiagaraNodeSimTargetSelector.cpp
@@ -77,6 +77,12 @@ void UNiagaraNodeSimTargetSelector::Compile(class FHlslNiagaraTranslator* Transl
 	bool bCPUSim = Translator->IsCompileOptionDefined(TEXT("CPUSim"));
 	bool bGPUSim = Translator->IsCompileOptionDefined(TEXT("GPUComputeSim"));
 
+	if (Translator->GetTargetUsage() >= ENiagaraScriptUsage::Function && Translator->GetTargetUsage() <= ENiagaraScriptUsage::DynamicInput)
+	{
+		// Functions through Dynamic inputs are missing the context, so just use CPU by default.
+		bCPUSim = true;
+	}
+
 
 	int32 VarIdx;
 	if (bCPUSim/*SimulationTarget == ENiagaraSimTarget::CPUSim*/)
