Adding sim stages to selection script

diff --git a/Engine/Plugins/FX/Niagara/Source/NiagaraEditor/Private/NiagaraNodeUsageSelector.cpp b/Engine/Plugins/FX/Niagara/Source/NiagaraEditor/Private/NiagaraNodeUsageSelector.cpp
--- a/Engine/Plugins/FX/Niagara/Source/NiagaraEditor/Private/NiagaraNodeUsageSelector.cpp
+++ b/Engine/Plugins/FX/Niagara/Source/NiagaraEditor/Private/NiagaraNodeUsageSelector.cpp
@@ -180,6 +180,7 @@ void UNiagaraNodeUsageSelector::AppendFunctionAliasForContext(const FNiagaraGrap
 	case ENiagaraScriptUsage::ParticleSpawnScript:
 	case ENiagaraScriptUsage::ParticleUpdateScript:
 	case ENiagaraScriptUsage::ParticleEventScript:
+	case ENiagaraScriptUsage::ParticleSimulationStageScript:
 	case ENiagaraScriptUsage::ParticleGPUComputeScript:
 		UsageString = "Particle";
 		break;
