Add legacy custom parameter name mode to parameter details panel to handle legacy custom namespaces. Fixup postload to move parameters with illegal namespaces to this mode.
Refactor FNiagaraVariableMetaData to implement Getters/Setters for members.

diff --git a/Engine/Plugins/FX/Niagara/Source/Niagara/Private/NiagaraTypes.cpp b/Engine/Plugins/FX/Niagara/Source/Niagara/Private/NiagaraTypes.cpp
new file mode 100644
--- /dev/null
+++ b/Engine/Plugins/FX/Niagara/Source/Niagara/Private/NiagaraTypes.cpp
@@ -0,0 +1,95 @@
+// Copyright Epic Games, Inc. All Rights Reserved.
+
+#include "NiagaraTypes.h"
+#include "NiagaraConstants.h"
+
+FString FNiagaraTypeUtilities::GetNamespaceStringForScriptParameterScope(const ENiagaraParameterScope& InScope)
+{
+	switch (InScope) {
+	case ENiagaraParameterScope::Engine:
+		return PARAM_MAP_ENGINE_STR;
+		break;
+	case ENiagaraParameterScope::Owner:
+		return PARAM_MAP_ENGINE_OWNER_STR;
+		break;
+	case ENiagaraParameterScope::User:
+		return PARAM_MAP_USER_STR;
+		break;
+	case ENiagaraParameterScope::System:
+		return PARAM_MAP_SYSTEM_STR;
+		break;
+	case ENiagaraParameterScope::Emitter:
+		return PARAM_MAP_EMITTER_STR;
+		break;
+	case ENiagaraParameterScope::Particles:
+		return PARAM_MAP_ATTRIBUTE_STR;
+		break;
+	case ENiagaraParameterScope::ScriptPersistent:
+		return PARAM_MAP_SCRIPT_PERSISTENT_STR;
+		break;
+	case ENiagaraParameterScope::ScriptTransient:
+		return PARAM_MAP_SCRIPT_TRANSIENT_STR;
+		break;
+	case ENiagaraParameterScope::Input:
+		return PARAM_MAP_MODULE_STR;
+		break;
+	case ENiagaraParameterScope::Local:
+		return PARAM_MAP_LOCAL_STR;
+		break;
+	default:
+		checkf(false, TEXT("Unhandled parameter scope encountered!"));
+		return FString();
+		break;
+	};
+}
+
+FString FNiagaraVariableMetaData::GetNamespaceString() const
+{
+	checkf(bUseLegacyNameString == false, TEXT("Tried to get namespace string for parameter using legacy name string edit mode!"));
+	if (IsInputOrLocalUsage())
+	{
+		FString NamespaceString = FNiagaraTypeUtilities::GetNamespaceStringForScriptParameterScope(Scope);
+		if (Usage == ENiagaraScriptParameterUsage::InitialValueInput)
+		{
+			NamespaceString.Append(PARAM_MAP_INITIAL_STR);
+		}
+		return NamespaceString;
+	}
+	return PARAM_MAP_OUTPUT_STR;
+}
+
+bool FNiagaraVariableMetaData::GetParameterName(FName& OutName) const
+{
+	if (bUseLegacyNameString)
+	{
+		return false;
+	}
+	OutName = CachedNamespacelessVariableName;
+	return true;
+}
+
+bool FNiagaraVariableMetaData::GetScope(ENiagaraParameterScope& OutScope) const
+{
+	if (bUseLegacyNameString)
+	{
+		OutScope = ENiagaraParameterScope::Custom;
+		return false;
+	}
+	OutScope = Scope;
+	return true;
+}
+
+void FNiagaraVariableMetaData::CopyPerScriptMetaData(const FNiagaraVariableMetaData& OtherMetaData)
+{
+	SetUsage(OtherMetaData.GetUsage());
+	ENiagaraParameterScope OtherMetaDataScope;
+	OtherMetaData.GetScope(OtherMetaDataScope);
+	SetScope(OtherMetaDataScope);
+	FName OtherMetaDataName;
+	if (OtherMetaData.GetParameterName(OtherMetaDataName))
+	{
+		SetCachedNamespacelessVariableName(OtherMetaDataName);
+	}
+	SetWasCreatedInSystemEditor(OtherMetaData.GetWasCreatedInSystemEditor());
+	SetIsUsingLegacyNameString(OtherMetaData.GetIsUsingLegacyNameString());
+}
diff --git a/Engine/Plugins/FX/Niagara/Source/Niagara/Public/NiagaraTypes.h b/Engine/Plugins/FX/Niagara/Source/Niagara/Public/NiagaraTypes.h
--- a/Engine/Plugins/FX/Niagara/Source/Niagara/Public/NiagaraTypes.h
+++ b/Engine/Plugins/FX/Niagara/Source/Niagara/Public/NiagaraTypes.h
@@ -13,6 +13,12 @@ class UNiagaraDataInterfaceBase;
 
 DECLARE_LOG_CATEGORY_EXTERN(LogNiagara, Log, Verbose);
 
+// helper methods for basic struct definitions
+struct NIAGARA_API FNiagaraTypeUtilities
+{
+	static FString GetNamespaceStringForScriptParameterScope(const ENiagaraParameterScope& InScope);
+};
+
 // basic type struct definitions
 
 USTRUCT(meta = (DisplayName = "float"))
@@ -463,6 +469,8 @@ enum class ENiagaraParameterScope : uint32
 
 	Local UMETA(Hidden), //Convenience markup for ScopeToString functions, only use in conjunction with ENiagaraScriptParameterUsage::Local.
 
+	Custom UMETA(Hidden), //Convenience markup for expressing parameters using legacy editor mode to freetype namespace and name.
+
 	// insert new scopes before
 	None UMETA(Hidden),
 
@@ -519,6 +527,7 @@ public:
 		, bAddedToNodeGraphDeepCopy(false)
 		, bOutputIsPersistent(false)
 		, bCreatedInSystemEditor(false)
+		, bUseLegacyNameString(false)
 	{};
 public:
 	UPROPERTY(EditAnywhere, Category = "Variable", meta = (MultiLine = true, SkipForCompileHash = "true"))
@@ -550,12 +559,53 @@ public:
 	UPROPERTY(EditAnywhere, Category = "Variable", DisplayName = "Property Metadata", meta = (ToolTip = "Property Metadata", SkipForCompileHash = "true"))
 	TMap<FName, FString> PropertyMetaData;
 
+public:
+	FString GetNamespaceString() const;
+
+	/** Gets the Scope and notifies if it does not apply due to an override being set.
+	 * @params OutScope		The Scope to return;
+	 * @return bool			Whether the returned scope is not overridden. Is false if bUseLegacyNameString is set.
+	 */
+	bool GetScope(ENiagaraParameterScope& OutScope) const;
+	void SetScope(ENiagaraParameterScope InScope) { Scope = InScope; };
+
+	ENiagaraScriptParameterUsage GetUsage() const { return Usage; };
+	void SetUsage(ENiagaraScriptParameterUsage InUsage) { Usage = InUsage; };
+
+	/** Gets the CachedNamespacelessParameterName and notifies if it cannot be returned due to an override being set.
+	 * @params OutName		The Name to return;
+	 * @return bool			Whether the CachedNamespacelessParameterName can be returned. Is false if bUseLegacyNameString is set.
+	 */
+	bool GetParameterName(FName& OutName) const;
+	void SetCachedNamespacelessVariableName(const FName& InVariableName) { CachedNamespacelessVariableName = InVariableName; };
+
+	bool GetWasCreatedInSystemEditor() const { return bCreatedInSystemEditor; };
+	void SetWasCreatedInSystemEditor(bool bWasCreatedInSystemEditor) { bCreatedInSystemEditor = bWasCreatedInSystemEditor; };
+
+	bool GetWasAddedToNodeGraphDeepCopy() const { return bAddedToNodeGraphDeepCopy; };
+	void SetWasAddedToNodeGraphDeepCopy(bool bWasAddedToNodeGraphDeepCopy) { bAddedToNodeGraphDeepCopy = bWasAddedToNodeGraphDeepCopy; };
+
+	bool GetIsStaticSwitch() const { return bIsStaticSwitch; };
+	void SetIsStaticSwitch(bool bInIsStaticSwitch) { bIsStaticSwitch = bInIsStaticSwitch; };
+
+	int32 GetStaticSwitchDefaultValue() const { return StaticSwitchDefaultValue; };
+	void SetStaticSwitchDefaultValue(int32 InStaticSwitchDefaultValue) { StaticSwitchDefaultValue = InStaticSwitchDefaultValue; };
+
+	bool GetOutputIsPersistent() const { return bOutputIsPersistent; };
+	void SetOutputIsPersistent(bool bInOutputIsPersistent) { bOutputIsPersistent = bInOutputIsPersistent; };
+
+	bool GetIsUsingLegacyNameString() const { return bUseLegacyNameString; };
+	void SetIsUsingLegacyNameString(bool bInUseLegacyNameString) { bUseLegacyNameString = bInUseLegacyNameString; };
+
+	void CopyPerScriptMetaData(const FNiagaraVariableMetaData& OtherMetaData);
+
+private:
 	/** Defines the scope of a variable that is an input to a script. */
-	UPROPERTY()
+	UPROPERTY(meta = (SkipForCompileHash = "true"))
 	ENiagaraParameterScope Scope;
 
 	/** Defines the usage of a variable as an argument or output relative to the script. */
-	UPROPERTY()
+	UPROPERTY(meta = (SkipForCompileHash = "true"))
 	ENiagaraScriptParameterUsage Usage;
 
 	/** This is a read-only variable that designates if the metadata is tied to a static switch or not. */
@@ -567,21 +617,24 @@ public:
 	int32 StaticSwitchDefaultValue;  // TODO: This should be moved to the UNiagaraScriptVariable in the future
 
 	/** Transient data to mark variables set in the node graph deep copy as having been derived from a module namespace parameter default. */
-	UPROPERTY(Transient)
+	UPROPERTY(Transient, meta = (SkipForCompileHash = "true"))
 	bool bAddedToNodeGraphDeepCopy;
 
 	/** Only valid if Usage is Output. Marks the associated FNiagaraVariable as Persistent across script runs and therefore should be retained in the Dataset during compilation/translation. */
-	UPROPERTY()
+	UPROPERTY(meta = (SkipForCompileHash = "true"))
 	bool bOutputIsPersistent;
 
 	/** Namespace-less name for associated FNiagaraVariable. Edited directly by user and then used to generate full Name of associated FNiagaraVariable. */
-	UPROPERTY()
+	UPROPERTY(meta = (SkipForCompileHash = "true"))
 	FName CachedNamespacelessVariableName;
 
 	/** Track if the associated parameter was created in the Emitter/System editor. Used to determine whether the associated parameter can be deleted from the Emitter/System editor. */
-	UPROPERTY()
+	UPROPERTY(meta = (SkipForCompileHash = "true"))
 	bool bCreatedInSystemEditor;
 
+	UPROPERTY(EditAnywhere, Category = "Custom Name", meta = (ToolTip = "Enable using a legacy custom name string.", SkipForCompileHash = "true"))
+	bool bUseLegacyNameString;
+
 public:
 	FORCEINLINE bool IsInputUsage() const { return Usage == ENiagaraScriptParameterUsage::Input || Usage == ENiagaraScriptParameterUsage::InputOutput; };
 	FORCEINLINE bool IsInputOrLocalUsage() const { return Usage == ENiagaraScriptParameterUsage::Input || Usage == ENiagaraScriptParameterUsage::InputOutput || Usage == ENiagaraScriptParameterUsage::InitialValueInput || Usage == ENiagaraScriptParameterUsage::Local; };
diff --git a/Engine/Plugins/FX/Niagara/Source/NiagaraEditor/Private/Customizations/NiagaraScriptVariableCustomization.cpp b/Engine/Plugins/FX/Niagara/Source/NiagaraEditor/Private/Customizations/NiagaraScriptVariableCustomization.cpp
--- a/Engine/Plugins/FX/Niagara/Source/NiagaraEditor/Private/Customizations/NiagaraScriptVariableCustomization.cpp
+++ b/Engine/Plugins/FX/Niagara/Source/NiagaraEditor/Private/Customizations/NiagaraScriptVariableCustomization.cpp
@@ -77,12 +77,12 @@ void FNiagaraScriptVariableDetails::PostUndo(bool bSuccess)
 		return;
 	}
 
-	if (Variable->Metadata.bIsStaticSwitch)
+	if (Variable->Metadata.GetIsStaticSwitch())
 	{
 		if (TypeUtilityStaticSwitchValue && ParameterEditorStaticSwitchValue)
 		{
 			TSharedPtr<FStructOnScope> ParameterValue = MakeShareable(new FStructOnScope(Variable->Variable.GetType().GetStruct()));
-			Variable->Variable.SetValue(Variable->Metadata.StaticSwitchDefaultValue);
+			Variable->Variable.SetValue(Variable->Metadata.GetStaticSwitchDefaultValue());
 			Variable->Variable.CopyTo(ParameterValue->GetStructMemory());
 			ParameterEditorStaticSwitchValue->UpdateInternalValueFromStruct(ParameterValue.ToSharedRef());
 		}
@@ -167,7 +167,7 @@ void FNiagaraScriptVariableDetails::CustomizeDetails(IDetailLayoutBuilder& Detai
 	DefaultModeHandle->SetOnPropertyValueChanged(FSimpleDelegate::CreateSP(this, &FNiagaraScriptVariableDetails::OnComboValueChanged));
 
 	FNiagaraEditorModule& EditorModule = FNiagaraEditorModule::Get();
-	if (Variable->Metadata.bIsStaticSwitch)
+	if (Variable->Metadata.GetIsStaticSwitch())
 	{
 		TypeUtilityStaticSwitchValue = EditorModule.GetTypeUtilities(Variable->Variable.GetType());
 		if (TypeUtilityStaticSwitchValue && TypeUtilityStaticSwitchValue->CanCreateParameterEditor())
@@ -175,7 +175,7 @@ void FNiagaraScriptVariableDetails::CustomizeDetails(IDetailLayoutBuilder& Detai
 			ParameterEditorStaticSwitchValue = TypeUtilityStaticSwitchValue->CreateParameterEditor(Variable->Variable.GetType());
 			
 			TSharedPtr<FStructOnScope> ParameterValue = MakeShareable(new FStructOnScope(Variable->Variable.GetType().GetStruct()));
-			Variable->Variable.SetValue(Variable->Metadata.StaticSwitchDefaultValue);
+			Variable->Variable.SetValue(Variable->Metadata.GetStaticSwitchDefaultValue());
 			Variable->Variable.CopyTo(ParameterValue->GetStructMemory());
 			ParameterEditorStaticSwitchValue->UpdateInternalValueFromStruct(ParameterValue.ToSharedRef());
 			ParameterEditorStaticSwitchValue->SetOnValueChanged(SNiagaraParameterEditor::FOnValueChange::CreateSP(this, &FNiagaraScriptVariableDetails::OnStaticSwitchValueChanged));
@@ -241,7 +241,7 @@ void FNiagaraScriptVariableDetails::CustomizeDetails(IDetailLayoutBuilder& Detai
 	}
 	
 
-	if (Variable->Metadata.bIsStaticSwitch) {
+	if (Variable->Metadata.GetIsStaticSwitch()) {
 		// Hide metadata UProperties that aren't useful for static switch variables
 		DetailBuilder.HideProperty(GET_MEMBER_NAME_CHECKED(UNiagaraScriptVariable, Metadata.EditCondition));
 		DetailBuilder.HideProperty(GET_MEMBER_NAME_CHECKED(UNiagaraScriptVariable, Metadata.VisibleCondition));
@@ -341,7 +341,7 @@ void FNiagaraScriptVariableDetails::OnStaticSwitchValueChanged()
 		TSharedPtr<FStructOnScope> ParameterValue = MakeShareable(new FStructOnScope(Variable->Variable.GetType().GetStruct()));
 		ParameterEditorStaticSwitchValue->UpdateStructFromInternalValue(ParameterValue.ToSharedRef());
 		Variable->Variable.SetData(ParameterValue->GetStructMemory());
-		Variable->Metadata.StaticSwitchDefaultValue = Variable->Variable.GetValue<int>();
+		Variable->Metadata.SetStaticSwitchDefaultValue(Variable->Variable.GetValue<int>());
 		
 #if WITH_EDITOR
 		if (UNiagaraGraph* Graph = Cast<UNiagaraGraph>(Variable->GetOuter()))
diff --git a/Engine/Plugins/FX/Niagara/Source/NiagaraEditor/Private/Customizations/NiagaraStaticSwitchNodeDetails.cpp b/Engine/Plugins/FX/Niagara/Source/NiagaraEditor/Private/Customizations/NiagaraStaticSwitchNodeDetails.cpp
--- a/Engine/Plugins/FX/Niagara/Source/NiagaraEditor/Private/Customizations/NiagaraStaticSwitchNodeDetails.cpp
+++ b/Engine/Plugins/FX/Niagara/Source/NiagaraEditor/Private/Customizations/NiagaraStaticSwitchNodeDetails.cpp
@@ -249,7 +249,7 @@ int32 FNiagaraStaticSwitchNodeDetails::GetDefaultWidgetIndex() const
 TOptional<int32> FNiagaraStaticSwitchNodeDetails::GetSwitchDefaultValue() const
 {
 	TOptional<FNiagaraVariableMetaData> MetaData = GetSwitchParameterMetadata();
-	return MetaData.IsSet() ? TOptional<int32>(MetaData->StaticSwitchDefaultValue) : TOptional<int32>();
+	return MetaData.IsSet() ? TOptional<int32>(MetaData->GetStaticSwitchDefaultValue()) : TOptional<int32>();
 }
 
 void FNiagaraStaticSwitchNodeDetails::DefaultIntValueCommitted(int32 Value, ETextCommit::Type CommitInfo)
@@ -257,7 +257,7 @@ void FNiagaraStaticSwitchNodeDetails::DefaultIntValueCommitted(int32 Value, ETex
 	TOptional<FNiagaraVariableMetaData> MetaData = GetSwitchParameterMetadata();
 	if (MetaData.IsSet())
 	{
-		MetaData->StaticSwitchDefaultValue = Value;
+		MetaData->SetStaticSwitchDefaultValue(Value);
 		SetSwitchParameterMetadata(MetaData.GetValue());
 	}
 }
@@ -267,7 +267,7 @@ void FNiagaraStaticSwitchNodeDetails::DefaultBoolValueCommitted(ECheckBoxState N
 	TOptional<FNiagaraVariableMetaData> MetaData = GetSwitchParameterMetadata();
 	if (MetaData.IsSet())
 	{
-		MetaData->StaticSwitchDefaultValue = (NewState == ECheckBoxState::Checked) ? 1 : 0;
+		MetaData->SetStaticSwitchDefaultValue((NewState == ECheckBoxState::Checked) ? 1 : 0);
 		SetSwitchParameterMetadata(MetaData.GetValue());
 	}
 }
@@ -314,7 +314,7 @@ void FNiagaraStaticSwitchNodeDetails::OnSelectionChanged(TSharedPtr<DefaultEnumO
 		return;
 	}
 
-	MetaData->StaticSwitchDefaultValue = SelectedDefaultValue->EnumIndex;
+	MetaData->SetStaticSwitchDefaultValue(SelectedDefaultValue->EnumIndex);
 	SetSwitchParameterMetadata(MetaData.GetValue());
 }
 
@@ -453,7 +453,7 @@ void FNiagaraStaticSwitchNodeDetails::RefreshDropdownValues()
 			FText DisplayName = Enum->GetDisplayNameTextByIndex(i);
 			DefaultEnumDropdownOptions.Add(MakeShared<DefaultEnumOption>(DisplayName, i));
 
-			if (MetaData.IsSet() && i == MetaData->StaticSwitchDefaultValue)
+			if (MetaData.IsSet() && i == MetaData->GetStaticSwitchDefaultValue())
 			{
 				SelectedDefaultValue = DefaultEnumDropdownOptions[i];
 			}
diff --git a/Engine/Plugins/FX/Niagara/Source/NiagaraEditor/Private/NiagaraGraph.cpp b/Engine/Plugins/FX/Niagara/Source/NiagaraEditor/Private/NiagaraGraph.cpp
--- a/Engine/Plugins/FX/Niagara/Source/NiagaraEditor/Private/NiagaraGraph.cpp
+++ b/Engine/Plugins/FX/Niagara/Source/NiagaraEditor/Private/NiagaraGraph.cpp
@@ -360,17 +360,19 @@ void UNiagaraGraph::PostLoad()
 				continue;
 			}
 
-			if (ScriptVar->Metadata.bIsStaticSwitch)
+			if (ScriptVar->Metadata.GetIsStaticSwitch())
 			{
 				// Do not need to assign scopes and usages to static switches, so assign the name and continue.
-				ScriptVar->Metadata.CachedNamespacelessVariableName = ScriptVar->Variable.GetName();
+				ScriptVar->Metadata.SetCachedNamespacelessVariableName(ScriptVar->Variable.GetName());
 				continue;
 			}
 
-			// Resolve the namespace of the variable to a scope, and assign the usage to local if "local" is in the name string.
-			FNiagaraStackGraphUtilities::SetParameterMetaDataFromName(Var.GetName(), ScriptVar->Metadata.Scope, ScriptVar->Metadata.Usage, ScriptVar->Metadata.CachedNamespacelessVariableName);
+			// Resolve the namespace of the variable to a scope, cache the scopeless name in the UI, and set the usage if the name string would force that to occur.
+			FNiagaraVariableMetaData OutMetaData;
+			FNiagaraStackGraphUtilities::GetParameterMetaDataFromName(Var.GetName(), OutMetaData);
+			ScriptVar->Metadata.CopyPerScriptMetaData(OutMetaData);
 
-			// If the usage is not local, select the usage based on the associated input/output pins.
+			// If the usage is not local or an initial value, select the usage based on the associated input/output pins.
 			SetScriptVariableUsageForPins(VarToPinsMap, ScriptVar);
 		}
 		
@@ -430,7 +432,7 @@ void UNiagaraGraph::ValidateDefaultPins()
 			// If the user selected custom mode they can basically do whatever they want
 			continue;
 		}
-		if (ScriptVariable->Metadata.bIsStaticSwitch) {
+		if (ScriptVariable->Metadata.GetIsStaticSwitch()) {
 			// We ignore static switch variables as they handle default values differently
 			continue;
 		}
@@ -916,7 +918,7 @@ bool UNiagaraGraph::UpdateUsageForScriptVariable(UNiagaraScriptVariable* ScriptV
 
 bool UNiagaraGraph::SetScriptVariableUsageForPins(const TMap<FNiagaraVariable, FInputPinsAndOutputPins>& VarToPinsMap, UNiagaraScriptVariable* ScriptVariable) const
 {
-	if (ScriptVariable->Metadata.Usage == ENiagaraScriptParameterUsage::Local || ScriptVariable->Metadata.Usage == ENiagaraScriptParameterUsage::InitialValueInput)
+	if (ScriptVariable->Metadata.GetUsage() == ENiagaraScriptParameterUsage::Local || ScriptVariable->Metadata.GetUsage() == ENiagaraScriptParameterUsage::InitialValueInput)
 	{
 		return true;
 	}
@@ -928,27 +930,27 @@ bool UNiagaraGraph::SetScriptVariableUsageForPins(const TMap<FNiagaraVariable, F
 		{
 			if (InOutPins->OutputPins.Num() > 0)
 			{
-				ScriptVariable->Metadata.Usage = ENiagaraScriptParameterUsage::InputOutput;
+				ScriptVariable->Metadata.SetUsage(ENiagaraScriptParameterUsage::InputOutput);
 				return true;
 			}
-			ScriptVariable->Metadata.Usage = ENiagaraScriptParameterUsage::Output;
+			ScriptVariable->Metadata.SetUsage(ENiagaraScriptParameterUsage::Output);
 			return true;
 		}
 		else if (InOutPins->OutputPins.Num() > 0)
 		{
-			ScriptVariable->Metadata.Usage = ENiagaraScriptParameterUsage::Input;
+			ScriptVariable->Metadata.SetUsage(ENiagaraScriptParameterUsage::Input);
 			return true;
 		}
 		else
 		{
 			ensureMsgf(false, TEXT("Var to pins map had entry with 0 pins!"));
-			ScriptVariable->Metadata.Usage = ENiagaraScriptParameterUsage::Local;
+			ScriptVariable->Metadata.SetUsage(ENiagaraScriptParameterUsage::Local);
 			return true;
 		}
 	}
 	else
 	{
-		ScriptVariable->Metadata.Usage = ENiagaraScriptParameterUsage::Local;
+		ScriptVariable->Metadata.SetUsage(ENiagaraScriptParameterUsage::Local);
 		return false;
 	}
 
@@ -1043,7 +1045,7 @@ void UNiagaraGraph::AddParameter(const FNiagaraVariable& Parameter, bool bIsStat
 	{
 		UNiagaraScriptVariable* NewScriptVariable = NewObject<UNiagaraScriptVariable>(this, FName(), RF_Transactional);
 		NewScriptVariable->Variable = Parameter;
-		NewScriptVariable->Metadata.bIsStaticSwitch = bIsStaticSwitch;
+		NewScriptVariable->Metadata.SetIsStaticSwitch(bIsStaticSwitch);
 		VariableToScriptVariable.Add(Parameter, NewScriptVariable);
 	}
 }
@@ -1073,15 +1075,15 @@ void UNiagaraGraph::AddParameter(FNiagaraVariable& Parameter, const FAddParamete
 
 		UNiagaraScriptVariable* NewScriptVariable = NewObject<UNiagaraScriptVariable>(this, FName(), RF_Transactional);
 		NewScriptVariable->Variable = Parameter;
-		NewScriptVariable->Metadata.bIsStaticSwitch = Options.bIsStaticSwitch;
-		NewScriptVariable->Metadata.bCreatedInSystemEditor = Options.bAddedFromSystemEditor;
+		NewScriptVariable->Metadata.SetIsStaticSwitch(Options.bIsStaticSwitch);
+		NewScriptVariable->Metadata.SetWasCreatedInSystemEditor(Options.bAddedFromSystemEditor);
 		VariableToScriptVariable.Add(Parameter, NewScriptVariable);
 
 		bool bSkipRefreshMetaDataScopeAndUsage = false;
 		if (Options.NewParameterUsage.IsSet())
 		{
 			bSkipRefreshMetaDataScopeAndUsage = true;
-			NewScriptVariable->Metadata.Usage = Options.NewParameterUsage.GetValue();
+			NewScriptVariable->Metadata.SetUsage(Options.NewParameterUsage.GetValue());
 		}
 
 		if (Options.NewParameterScope.IsSet())
@@ -1181,7 +1183,7 @@ bool UNiagaraGraph::RenameParameter(const FNiagaraVariable& Parameter, FName New
 	FNiagaraVariableMetaData OldMetaData;
 	if (OldScriptVariable)
 	{
-		if (!bFromStaticSwitch && (*OldScriptVariable)->Metadata.bIsStaticSwitch)
+		if (!bFromStaticSwitch && (*OldScriptVariable)->Metadata.GetIsStaticSwitch())
 		{
 			// We current disallow renaming static switch variables in the Parameters panel. 
 			bIsRenamingParameter = false;
@@ -1191,10 +1193,10 @@ bool UNiagaraGraph::RenameParameter(const FNiagaraVariable& Parameter, FName New
 	}
 
 	// Update the OldMetaData which will be applied to the new parameter to synchronize the cached namespace-less name and scope.
-	OldMetaData.CachedNamespacelessVariableName = FName(*FNiagaraStackGraphUtilities::GetNamespacelessVariableNameString(NewParameter.GetName()));
+	OldMetaData.SetCachedNamespacelessVariableName(FName(*FNiagaraStackGraphUtilities::GetNamespacelessVariableNameString(NewParameter.GetName())));
 	if (NewScope != ENiagaraParameterScope::None)
 	{
-		OldMetaData.Scope = NewScope;
+		OldMetaData.SetScope(NewScope);
 	}
 		
 	// Swap metadata to the new parameter; put the new parameter into VariableToScriptVariable
@@ -1252,7 +1254,7 @@ bool UNiagaraGraph::RenameParameter(const FNiagaraVariable& Parameter, FName New
 void UNiagaraGraph::ScriptVariableChanged(FNiagaraVariable Variable)
 {
 	UNiagaraScriptVariable** ScriptVariable = GetAllMetaData().Find(Variable);
-	if (!ScriptVariable || !*ScriptVariable || (*ScriptVariable)->Metadata.bIsStaticSwitch) {
+	if (!ScriptVariable || !*ScriptVariable || (*ScriptVariable)->Metadata.GetIsStaticSwitch()) {
 		return;
 	}
 
@@ -1712,40 +1714,6 @@ void UNiagaraGraph::CopyCachedReferencesMap(UNiagaraGraph* TargetGraph)
 	TargetGraph->ParameterToReferencesMap = ParameterToReferencesMap;
 }
 
-void UNiagaraGraph::FixupParameterUsageMetaData()
-{
-	// Collect all input and output pins to infer usages of variables they reference.
-	const TMap<FNiagaraVariable, FInputPinsAndOutputPins> VarToPinsMap = CollectVarsToInOutPinsMap();
-
-	// Now iterate all variables
-	for (auto It = VariableToScriptVariable.CreateIterator(); It; ++It)
-	{
-		FNiagaraVariable Var = It.Key();
-		UNiagaraScriptVariable* ScriptVar = It.Value();
-
-		if (Var.GetName() == FNiagaraConstants::InputPinName || Var.GetName() == FNiagaraConstants::OutputPinName)
-		{
-			//@todo Pins leaked into variable maps at some point, need to clean.
-			continue;
-		}
-
-		if (ScriptVar->Metadata.bIsStaticSwitch)
-		{
-			// Do not need to assign scopes and usages to static switches, so assign the name and continue.
-			ScriptVar->Metadata.CachedNamespacelessVariableName = ScriptVar->Variable.GetName();
-			continue;
-		}
-
-		// Resolve the namespace of the variable to a scope, and assign the usage to local if "local" is in the name string.
-		FNiagaraStackGraphUtilities::SetParameterMetaDataFromName(Var.GetName(), ScriptVar->Metadata.Scope, ScriptVar->Metadata.Usage, ScriptVar->Metadata.CachedNamespacelessVariableName);
-
-		// If the usage is not local, select the usage based on the associated input/output pins.
-		SetScriptVariableUsageForPins(VarToPinsMap, ScriptVar);
-	}
-
-	NotifyGraphNeedsRecompile();
-}
-
 bool UNiagaraGraph::IsPinVisualWidgetProviderRegistered() const
 {
 	return OnGetPinVisualWidgetDelegate.IsBound();
@@ -1979,10 +1947,7 @@ void UNiagaraGraph::SetPerScriptMetaData(const FNiagaraVariable& InVar, const FN
 		{
 			// Replace the old metadata scope, usage and cached name, but retain persistent metadata such as description.
 			FNiagaraVariableMetaData& OldMetaData = (*FoundMetaData)->Metadata;
-			OldMetaData.Usage = InMetaData.Usage;
-			OldMetaData.Scope = InMetaData.Scope;
-			OldMetaData.CachedNamespacelessVariableName = InMetaData.CachedNamespacelessVariableName;
-			OldMetaData.bCreatedInSystemEditor = InMetaData.bCreatedInSystemEditor;
+			OldMetaData.CopyPerScriptMetaData(InMetaData);
 		}
 	}
 	else
@@ -2138,7 +2103,7 @@ void UNiagaraGraph::RefreshParameterReferences() const
 		{
 			UNiagaraScriptVariable* NewScriptVariable = NewObject<UNiagaraScriptVariable>(const_cast<UNiagaraGraph*>(this));
 			NewScriptVariable->Variable = Variable;
-			NewScriptVariable->Metadata.bIsStaticSwitch = false;
+			NewScriptVariable->Metadata.SetIsStaticSwitch(false);
 			GenerateMetaDataForScriptVariable(NewScriptVariable);
 			VariableToScriptVariable.Add(Variable, NewScriptVariable);
 		}
@@ -2206,7 +2171,9 @@ void UNiagaraGraph::InvalidateCachedParameterData()
 
 void UNiagaraGraph::GenerateMetaDataForScriptVariable(UNiagaraScriptVariable* InScriptVariable) const
 {
-	FNiagaraStackGraphUtilities::SetParameterMetaDataFromName(InScriptVariable->Variable.GetName(), InScriptVariable->Metadata.Scope, InScriptVariable->Metadata.Usage, InScriptVariable->Metadata.CachedNamespacelessVariableName);
+	FNiagaraVariableMetaData OutMetaData;
+	FNiagaraStackGraphUtilities::GetParameterMetaDataFromName(InScriptVariable->Variable.GetName(), OutMetaData);
+	InScriptVariable->Metadata.CopyPerScriptMetaData(OutMetaData);
 	UpdateUsageForScriptVariable(InScriptVariable);
 }
 
diff --git a/Engine/Plugins/FX/Niagara/Source/NiagaraEditor/Private/NiagaraNodeFunctionCall.cpp b/Engine/Plugins/FX/Niagara/Source/NiagaraEditor/Private/NiagaraNodeFunctionCall.cpp
--- a/Engine/Plugins/FX/Niagara/Source/NiagaraEditor/Private/NiagaraNodeFunctionCall.cpp
+++ b/Engine/Plugins/FX/Niagara/Source/NiagaraEditor/Private/NiagaraNodeFunctionCall.cpp
@@ -237,7 +237,7 @@ void UNiagaraNodeFunctionCall::AllocateDefaultPins()
 			TOptional<FNiagaraVariableMetaData> MetaData = Graph->GetMetaData(Input);
 			if (MetaData.IsSet())
 			{
-				int32 DefaultValue = MetaData->StaticSwitchDefaultValue;
+				int32 DefaultValue = MetaData->GetStaticSwitchDefaultValue();
 				Input.AllocateData();
 				Input.SetValue<FNiagaraInt32>({ DefaultValue });
 				
diff --git a/Engine/Plugins/FX/Niagara/Source/NiagaraEditor/Private/NiagaraParameterMapHistory.cpp b/Engine/Plugins/FX/Niagara/Source/NiagaraEditor/Private/NiagaraParameterMapHistory.cpp
--- a/Engine/Plugins/FX/Niagara/Source/NiagaraEditor/Private/NiagaraParameterMapHistory.cpp
+++ b/Engine/Plugins/FX/Niagara/Source/NiagaraEditor/Private/NiagaraParameterMapHistory.cpp
@@ -1482,12 +1482,13 @@ void FNiagaraParameterMapHistoryWithMetaDataBuilder::FixupHistoryVariableNamespa
 			auto SortVars = [&InHandle, &VariablesToSkip](bool bGatherInputs)->TArray<FNiagaraHistoryVariable> {
 				TArray<FNiagaraHistoryVariable> OutHistoryVars = InHandle.HistoryVariables.FilterByPredicate([&VariablesToSkip, bGatherInputs](const FNiagaraHistoryVariable& HistoryVar)
 				{
-					ENiagaraParameterScope HistoryScope = HistoryVar.VarMetaData->Scope;
-					if (HistoryScope == ENiagaraParameterScope::User || HistoryScope == ENiagaraParameterScope::Engine || HistoryScope == ENiagaraParameterScope::ScriptTransient || HistoryScope == ENiagaraParameterScope::None)
+					ENiagaraParameterScope HistoryVarScope;
+					checkf(HistoryVar.VarMetaData->GetScope(HistoryVarScope), TEXT("Tried to process UNiagaraScriptVariable with override namespace set during precompiler namespace fixup!"));
+					if (HistoryVarScope == ENiagaraParameterScope::User || HistoryVarScope == ENiagaraParameterScope::Engine || HistoryVarScope == ENiagaraParameterScope::ScriptTransient || HistoryVarScope == ENiagaraParameterScope::None)
 					{
 						return false;
 					}
-					else if (HistoryVar.VarMetaData->Usage == ENiagaraScriptParameterUsage::Input && HistoryVar.VarMetaData->Scope == ENiagaraParameterScope::Input)
+					else if (HistoryVar.VarMetaData->GetUsage() == ENiagaraScriptParameterUsage::Input && HistoryVarScope == ENiagaraParameterScope::Input)
 					{
 						// Inputs are equivalent to module namespace, these should not change.
 						return false;
@@ -1497,7 +1498,7 @@ void FNiagaraParameterMapHistoryWithMetaDataBuilder::FixupHistoryVariableNamespa
 						return false;
 					}
 
-					if (HistoryVar.VarMetaData->bAddedToNodeGraphDeepCopy == false)
+					if (HistoryVar.VarMetaData->GetWasAddedToNodeGraphDeepCopy() == false)
 					{
 						FNiagaraParameterHandle VarHandle = FNiagaraParameterHandle(HistoryVar.VarWithOriginalAlias->GetName());
 						if (VarHandle.GetNamespace().ToString().Contains(TRANSLATOR_SET_VARIABLES_STR))
@@ -1508,11 +1509,11 @@ void FNiagaraParameterMapHistoryWithMetaDataBuilder::FixupHistoryVariableNamespa
 					
 					if (bGatherInputs)
 					{
-						return HistoryVar.VarMetaData->Usage == ENiagaraScriptParameterUsage::Input || HistoryVar.VarMetaData->Usage == ENiagaraScriptParameterUsage::InputOutput || HistoryVar.VarMetaData->bAddedToNodeGraphDeepCopy;
+						return HistoryVar.VarMetaData->GetUsage() == ENiagaraScriptParameterUsage::Input || HistoryVar.VarMetaData->GetUsage() == ENiagaraScriptParameterUsage::InputOutput || HistoryVar.VarMetaData->GetWasAddedToNodeGraphDeepCopy();
 					}
 					else
 					{
-						return HistoryVar.VarMetaData->Usage == ENiagaraScriptParameterUsage::Output || HistoryVar.VarMetaData->Usage == ENiagaraScriptParameterUsage::InputOutput;
+						return HistoryVar.VarMetaData->GetUsage() == ENiagaraScriptParameterUsage::Output || HistoryVar.VarMetaData->GetUsage() == ENiagaraScriptParameterUsage::InputOutput;
 					}
 					return false;
 				});
@@ -1681,9 +1682,9 @@ void FNiagaraParameterMapHistoryWithMetaDataBuilder::FixupHistoryVariableNamespa
 		TArray<FNiagaraHistoryVariable> InputAndOutputVars;
 		for (FSortedHistoryHandle* Handle : PerExecIdxSortedHandles)
 		{
-			ModuleHistoryVars.Append(Handle->OutputHistoryVars.FilterByPredicate([](const FNiagaraHistoryVariable& HistoryVar) {return HistoryVar.VarMetaData->bAddedToNodeGraphDeepCopy; }));
+			ModuleHistoryVars.Append(Handle->OutputHistoryVars.FilterByPredicate([](const FNiagaraHistoryVariable& HistoryVar) {return HistoryVar.VarMetaData->GetWasAddedToNodeGraphDeepCopy(); }));
 			InputAndOutputVars.Append(Handle->InputHistoryVars);
-			InputAndOutputVars.Append(Handle->OutputHistoryVars.FilterByPredicate([](const FNiagaraHistoryVariable& HistoryVar) {return HistoryVar.VarMetaData->bAddedToNodeGraphDeepCopy == false; }));
+			InputAndOutputVars.Append(Handle->OutputHistoryVars.FilterByPredicate([](const FNiagaraHistoryVariable& HistoryVar) {return HistoryVar.VarMetaData->GetWasAddedToNodeGraphDeepCopy() == false; }));
 		}
 
 		TMap<FName, uint32> ModuleCollisionNameCounts;
@@ -1748,14 +1749,16 @@ void FNiagaraParameterMapHistoryWithMetaDataBuilder::FixupHistoryVariableNamespa
 					auto FindUpstreamVarMatchIdx = [&Handle, &MatchIdx, &CurrentHistoryVar, &CurrentInputVarHandle, &UpstreamSortedHistoryHandle]()->bool {
 						const ENiagaraParameterScope UpstreamScope = UpstreamSortedHistoryHandle->GetScriptScope();
 						ENiagaraParameterScope TargetScope = ENiagaraParameterScope::None;
-						if (CurrentHistoryVar.VarMetaData->Scope == ENiagaraParameterScope::ScriptPersistent)
+						ENiagaraParameterScope CurrentHistoryVarScope;
+						checkf(CurrentHistoryVar.VarMetaData->GetScope(CurrentHistoryVarScope), TEXT("Tried to process UNiagaraScriptVariable with override namespace set during precompiler namespace fixup!"));
+						if (CurrentHistoryVarScope == ENiagaraParameterScope::ScriptPersistent)
 						{
 							// If the current history var is targeting the script alias scope then get the scope of the current history.
 							TargetScope = Handle.GetScriptScope();
 						}
 						else
 						{
-							TargetScope = CurrentHistoryVar.VarMetaData->Scope;
+							TargetScope = CurrentHistoryVarScope;
 						}
 
 						MatchIdx = UpstreamSortedHistoryHandle->OutputHistoryVars.IndexOfByPredicate(
@@ -1772,7 +1775,7 @@ void FNiagaraParameterMapHistoryWithMetaDataBuilder::FixupHistoryVariableNamespa
 					{
 						// Found a match upstream. Continue checking for upstream matches, recording those matches if they occur and fixing those matches' variable and pin names. Then fix the current history variable's name and pin names.
 						FString FixedNamespaceString = FString();
-						FixedNamespaceString.Append(FNiagaraStackGraphUtilities::GetNamespaceStringForScriptParameterScope(UpstreamSortedHistoryHandle->GetScriptScope()));
+						FixedNamespaceString.Append(FNiagaraTypeUtilities::GetNamespaceStringForScriptParameterScope(UpstreamSortedHistoryHandle->GetScriptScope()));
 						FixedNamespaceString.Append(CurrentInputVarHandle.GetName().ToString());
 						FName FixedNamespaceName = FName(*FixedNamespaceString);
 
@@ -1856,12 +1859,14 @@ void FNiagaraParameterMapHistoryWithMetaDataBuilder::FixupHistoryVariableNamespa
 		{
 			for (FNiagaraHistoryVariable& InputHistoryVar : Handle.InputHistoryVars)
 			{
-				if (InputHistoryVar.VarMetaData->Scope == ENiagaraParameterScope::None)
+				ENiagaraParameterScope InputHistoryVarScope;
+				checkf(InputHistoryVar.VarMetaData->GetScope(InputHistoryVarScope), TEXT("Tried to process UNiagaraScriptVariable with override namespace set during precompiler namespace fixup!"));
+				if (InputHistoryVarScope == ENiagaraParameterScope::None)
 				{
 					// Only need to fixup variables that are already in the dataset.
 					continue;
 				}
-				else if (InputHistoryVar.VarMetaData->Usage == ENiagaraScriptParameterUsage::InputOutput)
+				else if (InputHistoryVar.VarMetaData->GetUsage() == ENiagaraScriptParameterUsage::InputOutput)
 				{
 					//@todo these may be culled if the param map write history does not link the input and output pins in any way
 					continue;
@@ -1876,21 +1881,23 @@ void FNiagaraParameterMapHistoryWithMetaDataBuilder::FixupHistoryVariableNamespa
 
 			for (FNiagaraHistoryVariable& OutputHistoryVar : Handle.OutputHistoryVars)
 			{
+				ENiagaraParameterScope OutputHistoryVarScope;
+				checkf(OutputHistoryVar.VarMetaData->GetScope(OutputHistoryVarScope), TEXT("Tried to process UNiagaraScriptVariable with override namespace set during precompiler namespace fixup!"));
 				if (OutputHistoryVar.bDoNotDemote)
 				{
 					continue;
 				}
-				else if (OutputHistoryVar.VarMetaData->Scope == ENiagaraParameterScope::None)
+				else if (OutputHistoryVarScope == ENiagaraParameterScope::None)
 				{
 					// Only need to fixup variables that are already in the dataset.
 					continue;
 				}
-				else if (OutputHistoryVar.VarMetaData->Usage == ENiagaraScriptParameterUsage::InputOutput)
+				else if (OutputHistoryVar.VarMetaData->GetUsage() == ENiagaraScriptParameterUsage::InputOutput)
 				{
 					//@todo these may be culled if the param map write history does not link the input and output pins in any way
 					continue;
 				}
-				else if (OutputHistoryVar.VarMetaData->Usage == ENiagaraScriptParameterUsage::Output && OutputHistoryVar.VarMetaData->bOutputIsPersistent)
+				else if (OutputHistoryVar.VarMetaData->GetUsage() == ENiagaraScriptParameterUsage::Output && OutputHistoryVar.VarMetaData->GetOutputIsPersistent())
 				{
 					if (Handle.InputHistoryVars.ContainsByPredicate([&OutputHistoryVar](const FNiagaraHistoryVariable& InputHistoryVar) {return InputHistoryVar.VarWithOriginalAlias == OutputHistoryVar.VarWithOriginalAlias; }))
 					{
@@ -1917,7 +1924,7 @@ int32 FNiagaraParameterMapHistoryWithMetaDataBuilder::AddVariableToHistory(FNiag
 		if (CallingGraphContext.Num() == 1)
 		{
 			// Mark vars added to the node graph deep copy to fixup name during precompile, and allow linking downstream even if their scope is local.
-			MetaData->bAddedToNodeGraphDeepCopy = true;
+			MetaData->SetWasAddedToNodeGraphDeepCopy(true);
 		}
 		return History.AddVariable(InVar, InAliasedVar, InPin, MetaData);
 	}
@@ -1942,7 +1949,7 @@ FNiagaraParameterMapHistoryHandle::FNiagaraParameterMapHistoryHandle(FNiagaraPar
 	EmitterUniqueName = InEmitterUniqueName;
 
 	Scope = FNiagaraStackGraphUtilities::GetScopeForScriptUsage(InHistory.OriginatingScriptUsage);
-	const FString HistoryScopeString = FNiagaraStackGraphUtilities::GetNamespaceStringForScriptParameterScope(Scope);
+	const FString HistoryScopeString = FNiagaraTypeUtilities::GetNamespaceStringForScriptParameterScope(Scope);
 	FNiagaraEditorUtilities::GetScriptRunAndExecutionIndexFromUsage(InHistory.OriginatingScriptUsage, ScriptRunIndex, ScriptExecutionIndex);
 
 	for (int VarIdx = 0; VarIdx < InHistory.Variables.Num(); ++VarIdx)
@@ -1950,6 +1957,11 @@ FNiagaraParameterMapHistoryHandle::FNiagaraParameterMapHistoryHandle(FNiagaraPar
 		FNiagaraVariable* Var = &InHistory.Variables[VarIdx];
 		FNiagaraVariable* VarWithOriginalAliasIntact = &InHistory.VariablesWithOriginalAliasesIntact[VarIdx];
 		FNiagaraVariableMetaData* VarMetaData = &InHistory.VariableMetaData[VarIdx];
+		if (VarMetaData->GetIsUsingLegacyNameString() == true)
+		{
+			// Cannot fixup names of variables using legacy name string mode, skip this entry.
+			continue;
+		}
 		HistoryVariables.Add(FNiagaraHistoryVariable(Var, VarWithOriginalAliasIntact, VarMetaData));
 	}
 }
@@ -1962,7 +1974,7 @@ FNiagaraParameterMapHistoryHandle::FNiagaraParameterMapHistoryHandle(FNiagaraPar
 	EmitterUniqueName = InEmitterUniqueName;
 
 	Scope = FNiagaraStackGraphUtilities::GetScopeForScriptUsage(InSecondaryHistory.OriginatingScriptUsage);
-	const FString HistoryScopeString = FNiagaraStackGraphUtilities::GetNamespaceStringForScriptParameterScope(Scope);
+	const FString HistoryScopeString = FNiagaraTypeUtilities::GetNamespaceStringForScriptParameterScope(Scope);
 	FNiagaraEditorUtilities::GetScriptRunAndExecutionIndexFromUsage(InSecondaryHistory.OriginatingScriptUsage, ScriptRunIndex, ScriptExecutionIndex);
 
 	for (int32 SecondaryVarIdx = 0; SecondaryVarIdx < InSecondaryHistory.Variables.Num(); ++SecondaryVarIdx)
@@ -1976,6 +1988,11 @@ FNiagaraParameterMapHistoryHandle::FNiagaraParameterMapHistoryHandle(FNiagaraPar
 
 			// Assign the secondary history metadata as it is canonical in regards to parameter linking.
 			FNiagaraVariableMetaData* ScriptVarMetaData = &InSecondaryHistory.VariableMetaData[SecondaryVarIdx];
+			if (ScriptVarMetaData->GetIsUsingLegacyNameString() == true)
+			{
+				// Cannot fixup names of variables using legacy name string mode, skip this entry.
+				continue;
+			}
 			HistoryVariables.Add(FNiagaraHistoryVariable(PrimaryVar, PrimaryVarWithOriginalAliasIntact, ScriptVarMetaData));
 		}
 	}
diff --git a/Engine/Plugins/FX/Niagara/Source/NiagaraEditor/Private/ViewModels/NiagaraParameterNameViewModel.cpp b/Engine/Plugins/FX/Niagara/Source/NiagaraEditor/Private/ViewModels/NiagaraParameterNameViewModel.cpp
--- a/Engine/Plugins/FX/Niagara/Source/NiagaraEditor/Private/ViewModels/NiagaraParameterNameViewModel.cpp
+++ b/Engine/Plugins/FX/Niagara/Source/NiagaraEditor/Private/ViewModels/NiagaraParameterNameViewModel.cpp
@@ -26,20 +26,27 @@ FNiagaraParameterPanelEntryParameterNameViewModel::FNiagaraParameterPanelEntryPa
 
 TSharedRef<SWidget> FNiagaraParameterPanelEntryParameterNameViewModel::CreateScopeSlotWidget() const
 {
+	if (CachedScriptVarAndViewInfo.MetaData.GetIsUsingLegacyNameString())
+	{
+		return SNullWidget::NullWidget;
+	}
+
+	ENiagaraParameterScope CachedScope;
+	CachedScriptVarAndViewInfo.MetaData.GetScope(CachedScope);
 	UEnum* ParameterScopeEnum = FNiagaraTypeDefinition::GetParameterScopeEnum();
 	TSharedPtr< SWidget > ScopeComboBoxWidget;
 	FNiagaraEditorModule& NiagaraEditorModule = FModuleManager::LoadModuleChecked<FNiagaraEditorModule>("NiagaraEditor");
 
 	bool bEnableScopeSlotWidget = true;
-	if (CachedScriptVarAndViewInfo.MetaData.Usage != ENiagaraScriptParameterUsage::Input)
+	if (CachedScriptVarAndViewInfo.MetaData.GetUsage() != ENiagaraScriptParameterUsage::Input)
 	{
 		bEnableScopeSlotWidget = false;
 	}
-	else if (CachedScriptVarAndViewInfo.MetaData.Scope == ENiagaraParameterScope::Local)
+	else if (CachedScope == ENiagaraParameterScope::Local)
 	{
 		bEnableScopeSlotWidget = false;
 	}
-	else if (CachedScriptVarAndViewInfo.MetaData.bIsStaticSwitch)
+	else if (CachedScriptVarAndViewInfo.MetaData.GetIsStaticSwitch())
 	{
 		bEnableScopeSlotWidget = false;
 	}
@@ -82,7 +89,12 @@ TSharedRef<SWidget> FNiagaraParameterPanelEntryParameterNameViewModel::CreateTex
 
 int32 FNiagaraParameterPanelEntryParameterNameViewModel::GetScopeValue() const
 {
-	return (int32)CachedScriptVarAndViewInfo.MetaData.Scope;
+	ENiagaraParameterScope CachedScope;
+	if (ensureMsgf(CachedScriptVarAndViewInfo.MetaData.GetScope(CachedScope), TEXT("Failed to get scope value for param as override namespace is set! This method should not be bound!")))
+	{
+		return (int32)CachedScope;
+	}
+	return int32(ENiagaraParameterScope::Custom);
 }
 
 void FNiagaraParameterPanelEntryParameterNameViewModel::OnScopeValueChanged(int32 NewScopeValue, ESelectInfo::Type SelectionType) const
@@ -102,7 +114,13 @@ bool FNiagaraParameterPanelEntryParameterNameViewModel::GetScopeValueIsEnabled(i
 
 FText FNiagaraParameterPanelEntryParameterNameViewModel::GetParameterNameText() const
 {
-	return FText::FromName(CachedScriptVarAndViewInfo.MetaData.CachedNamespacelessVariableName);
+	if (CachedScriptVarAndViewInfo.MetaData.GetIsUsingLegacyNameString())
+	{
+		return FText::FromName(CachedScriptVarAndViewInfo.ScriptVariable.GetName());
+	}
+	FName ParameterName;
+	CachedScriptVarAndViewInfo.MetaData.GetParameterName(ParameterName);
+	return FText::FromName(ParameterName);
 }
 
 void FNiagaraParameterPanelEntryParameterNameViewModel::OnParameterRenamed(const FText& NewNameText, ETextCommit::Type TextCommitType) const
@@ -126,6 +144,12 @@ FNiagaraGraphPinParameterNameViewModel::FNiagaraGraphPinParameterNameViewModel(
 
 TSharedRef<SWidget> FNiagaraGraphPinParameterNameViewModel::CreateScopeSlotWidget() const
 {
+	// If using legacy name mode, skip the scope widget as we only need a text slot widget.
+	if (CachedScriptVarAndViewInfo.MetaData.GetIsUsingLegacyNameString())
+	{
+		return SNullWidget::NullWidget;
+	}
+
 	UEnum* ParameterScopeEnum = FNiagaraTypeDefinition::GetParameterScopeEnum();
 	TSharedPtr< SWidget > ScopeComboBoxWidget;
 	FNiagaraEditorModule& NiagaraEditorModule = FModuleManager::LoadModuleChecked<FNiagaraEditorModule>("NiagaraEditor");
@@ -166,15 +190,17 @@ TSharedRef<SWidget> FNiagaraGraphPinParameterNameViewModel::CreateTextSlotWidget
 
 int32 FNiagaraGraphPinParameterNameViewModel::GetScopeValue() const
 {
-	return (int32)CachedScriptVarAndViewInfo.MetaData.Scope;
+	ENiagaraParameterScope CachedScope;
+	if (ensureMsgf(CachedScriptVarAndViewInfo.MetaData.GetScope(CachedScope), TEXT("Failed to get scope value for param as override namespace is set! This method should not be bound!")))
+	{
+		return (int32)CachedScope;
+	}
+	return int32(ENiagaraParameterScope::Custom);
 }
 
 void FNiagaraGraphPinParameterNameViewModel::OnScopeValueChanged(int32 NewScopeValue, ESelectInfo::Type SelectionType) const
 {
 	ParameterPanelViewModel->ChangePinScope(OwningPin, ENiagaraParameterScope(NewScopeValue));
-
-	//UEdGraphPin* NonConstOwningPin = const_cast<UEdGraphPin*>(OwningPin); //@todo refactor
-	//Cast<UNiagaraNode>(OwningPin->GetOwningNode())->CommitEditablePinName(FText::FromName(NewPinName), NonConstOwningPin);
 }
 
 FText FNiagaraGraphPinParameterNameViewModel::GetToolTipForScopeValue(int32 ScopeValue) const
@@ -189,13 +215,16 @@ bool FNiagaraGraphPinParameterNameViewModel::GetScopeValueIsEnabled(int32 ScopeV
 
 FText FNiagaraGraphPinParameterNameViewModel::GetParameterNameText() const
 {
-	return FText::FromName(CachedScriptVarAndViewInfo.MetaData.CachedNamespacelessVariableName);
+	if (CachedScriptVarAndViewInfo.MetaData.GetIsUsingLegacyNameString())
+	{
+		return FText::FromName(CachedScriptVarAndViewInfo.ScriptVariable.GetName());
+	}
+	FName ParameterName;
+	CachedScriptVarAndViewInfo.MetaData.GetParameterName(ParameterName);
+	return FText::FromName(ParameterName);
 }
 
 void FNiagaraGraphPinParameterNameViewModel::OnParameterRenamed(const FText& NewNameText, ETextCommit::Type SelectionType) const
 {
 	ParameterPanelViewModel->RenamePin(OwningPin, NewNameText);
-
-	//UEdGraphPin* NonConstOwningPin = const_cast<UEdGraphPin*>(OwningPin); //@todo refactor
-	//Cast<UNiagaraNode>(OwningPin->GetOwningNode())->CommitEditablePinName(NewNameText, NonConstOwningPin);
 }
diff --git a/Engine/Plugins/FX/Niagara/Source/NiagaraEditor/Private/ViewModels/NiagaraParameterPanelViewModel.cpp b/Engine/Plugins/FX/Niagara/Source/NiagaraEditor/Private/ViewModels/NiagaraParameterPanelViewModel.cpp
--- a/Engine/Plugins/FX/Niagara/Source/NiagaraEditor/Private/ViewModels/NiagaraParameterPanelViewModel.cpp
+++ b/Engine/Plugins/FX/Niagara/Source/NiagaraEditor/Private/ViewModels/NiagaraParameterPanelViewModel.cpp
@@ -64,6 +64,9 @@ FText NiagaraParameterPanelSectionID::OnGetSectionTitle(const NiagaraParameterPa
 	case NiagaraParameterPanelSectionID::INITIALVALUES:
 		SectionTitle = LOCTEXT("InitialValuesParamSection", "Initial Values");
 		break;
+	case NiagaraParameterPanelSectionID::CUSTOM:
+		SectionTitle = LOCTEXT("CustomParamSection", "Custom");
+		break;
 // 	case NiagaraParameterPanelSectionID::PARAMETERCOLLECTION: //@TODO implement parameter collection handling
 // 		SectionTitle = NSLOCTEXT("GraphActionNode", "ParameterCollection", "Parameter Collection");
 // 		break;
@@ -91,6 +94,8 @@ const NiagaraParameterPanelSectionID::Type NiagaraParameterPanelSectionID::GetSe
 		return NiagaraParameterPanelSectionID::EMITTER;
 	case ENiagaraParameterScope::Particles:
 		return NiagaraParameterPanelSectionID::PARTICLES;
+	case ENiagaraParameterScope::Custom:
+		return NiagaraParameterPanelSectionID::CUSTOM;
 	case ENiagaraParameterScope::ScriptPersistent:
 	case ENiagaraParameterScope::ScriptTransient:
 		// This is a potential situation if a script alias param has not had its scope cached by compiling.
@@ -103,17 +108,21 @@ const NiagaraParameterPanelSectionID::Type NiagaraParameterPanelSectionID::GetSe
 
 const NiagaraParameterPanelSectionID::Type NiagaraParameterPanelSectionID::GetSectionForParameterMetaData(const FNiagaraVariableMetaData& MetaData)
 {
-	if (MetaData.Usage == ENiagaraScriptParameterUsage::Local)
+	ENiagaraParameterScope MetaDataScope;
+	MetaData.GetScope(MetaDataScope);
+	ENiagaraScriptParameterUsage MetaDataUsage = MetaData.GetUsage();
+
+	if (MetaDataUsage == ENiagaraScriptParameterUsage::Local)
 	{
 		return NiagaraParameterPanelSectionID::LOCALS;
 	}
-	else if (MetaData.Usage == ENiagaraScriptParameterUsage::InitialValueInput)
+	else if (MetaDataUsage == ENiagaraScriptParameterUsage::InitialValueInput)
 	{
 		return NiagaraParameterPanelSectionID::INITIALVALUES;
 	}
-	else if (MetaData.Usage == ENiagaraScriptParameterUsage::Input)
+	else if (MetaDataUsage == ENiagaraScriptParameterUsage::Input)
 	{
-		if (MetaData.Scope == ENiagaraParameterScope::Input)
+		if (MetaDataScope == ENiagaraParameterScope::Input)
 		{
 			return NiagaraParameterPanelSectionID::INPUTS;
 		}
@@ -122,15 +131,15 @@ const NiagaraParameterPanelSectionID::Type NiagaraParameterPanelSectionID::GetSe
 			return NiagaraParameterPanelSectionID::REFERENCES;
 		}
 	}
-	else if (MetaData.bIsStaticSwitch)
+	else if (MetaData.GetIsStaticSwitch())
 	{
 		return NiagaraParameterPanelSectionID::INPUTS;
 	}
-	else if (MetaData.Usage == ENiagaraScriptParameterUsage::Output)
+	else if (MetaDataUsage == ENiagaraScriptParameterUsage::Output)
 	{
 		return NiagaraParameterPanelSectionID::OUTPUTS;
 	}
-	else if (MetaData.Usage == ENiagaraScriptParameterUsage::InputOutput)
+	else if (MetaDataUsage == ENiagaraScriptParameterUsage::InputOutput)
 	{
 		ensureMsgf(false, TEXT("Encountered an InputOutput parameter usage when getting section ID for parameter panel!"));
 		return NiagaraParameterPanelSectionID::REFERENCES;
@@ -268,18 +277,21 @@ void FNiagaraSystemToolkitParameterPanelViewModel::CollectStaticSections(TArray<
 NiagaraParameterPanelSectionID::Type FNiagaraSystemToolkitParameterPanelViewModel::GetSectionForVarAndViewInfo(const FNiagaraScriptVariableAndViewInfo& VarAndViewInfo) const
 {
 	// Generic Emitter/System view, categorize by scope
-	if (VarAndViewInfo.MetaData.Usage == ENiagaraScriptParameterUsage::Input && VarAndViewInfo.MetaData.Scope != ENiagaraParameterScope::Input)
+
+	ENiagaraParameterScope ParameterScope = ENiagaraParameterScope::None;
+	VarAndViewInfo.MetaData.GetScope(ParameterScope);
+	if (VarAndViewInfo.MetaData.GetUsage() == ENiagaraScriptParameterUsage::Input && ParameterScope != ENiagaraParameterScope::Input)
 	{
-		return NiagaraParameterPanelSectionID::GetSectionForScope(VarAndViewInfo.MetaData.Scope);
+		return NiagaraParameterPanelSectionID::GetSectionForScope(ParameterScope);
 	}
 	else
 	{
 		//return NiagaraParameterPanelSectionID::GetSectionForScope(VarAndViewInfo.ScriptVariableMetaData.LastPrecompileScope); //@todo(ng) cache the known scope in the niagarasystem
-		ENiagaraParameterScope ReturnScope = ENiagaraParameterScope::None;
-		ENiagaraScriptParameterUsage EmptyUsage;
-		FName EmptyName;
-		FNiagaraStackGraphUtilities::SetParameterMetaDataFromName(VarAndViewInfo.ScriptVariable.GetName(), ReturnScope, EmptyUsage, EmptyName);
-		return NiagaraParameterPanelSectionID::GetSectionForScope(ReturnScope);
+		FNiagaraVariableMetaData OutMetaData;
+		FNiagaraStackGraphUtilities::GetParameterMetaDataFromName(VarAndViewInfo.ScriptVariable.GetName(), VarAndViewInfo.MetaData);
+		ENiagaraParameterScope OutMetaDataScope;
+		VarAndViewInfo.MetaData.GetScope(OutMetaDataScope);
+		return NiagaraParameterPanelSectionID::GetSectionForScope(OutMetaDataScope);
 	}
 }
 
@@ -288,18 +300,21 @@ void FNiagaraSystemToolkitParameterPanelViewModel::AddParameter(const FNiagaraVa
 	FScopedTransaction AddParameter(LOCTEXT("AddParameter", "Add Parameter"));
 	bool bSystemIsSelected = OverviewSelectionViewModel->GetSystemIsSelected();
 
+	ENiagaraParameterScope NewScope = ENiagaraParameterScope::None;
+	InVariableMetaDataToAssign.GetScope(NewScope);
+
 	UNiagaraGraph::FAddParameterOptions AddParameterOptions = UNiagaraGraph::FAddParameterOptions();
-	AddParameterOptions.NewParameterScope = InVariableMetaDataToAssign.Scope;
-	AddParameterOptions.NewParameterUsage = InVariableMetaDataToAssign.Usage;
+	AddParameterOptions.NewParameterScope = NewScope;
+	AddParameterOptions.NewParameterUsage = InVariableMetaDataToAssign.GetUsage();
 	AddParameterOptions.bAddedFromSystemEditor = true;
 	FNiagaraVariable DuplicateVar = InVariableToAdd; //@todo(ng) rewrite
 
-	if (InVariableMetaDataToAssign.Scope == ENiagaraParameterScope::User)
+	if (NewScope == ENiagaraParameterScope::User)
 	{
 		SystemViewModel->GetSystem().Modify();
 		SystemViewModel->GetSystem().GetExposedParameters().AddParameter(InVariableToAdd); //@Todo(ng) copy ParameterMapView impl
 	}
-	else if (InVariableMetaDataToAssign.Scope == ENiagaraParameterScope::System)
+	else if (NewScope == ENiagaraParameterScope::System)
 	{
 		UNiagaraGraph* Graph = SystemViewModel->GetSystemScriptViewModel()->GetGraphViewModel()->GetGraph();
 		Graph->Modify();
@@ -328,7 +343,11 @@ void FNiagaraSystemToolkitParameterPanelViewModel::AddParameter(const FNiagaraVa
 void FNiagaraSystemToolkitParameterPanelViewModel::RemoveParameter(const FNiagaraVariable& TargetVariableToRemove, const FNiagaraVariableMetaData& TargetVariableMetaData)
 {
 	FScopedTransaction RemoveParameter(LOCTEXT("RemoveParameter", "Remove Parameter"));
-	if (TargetVariableMetaData.IsInputUsage() && TargetVariableMetaData.Scope == ENiagaraParameterScope::User) //@todo(ng) verify
+
+	ENiagaraParameterScope TargetVariableScope = ENiagaraParameterScope::None;
+	TargetVariableMetaData.GetScope(TargetVariableScope);
+
+	if (TargetVariableMetaData.IsInputUsage() && TargetVariableScope == ENiagaraParameterScope::User) //@todo(ng) verify
 	{
 		SystemViewModel->GetSystem().Modify();
 		SystemViewModel->GetSystem().GetExposedParameters().RemoveParameter(TargetVariableToRemove);
@@ -349,20 +368,32 @@ void FNiagaraSystemToolkitParameterPanelViewModel::RemoveParameter(const FNiagar
 
 bool FNiagaraSystemToolkitParameterPanelViewModel::CanRemoveParameter(const FNiagaraVariable& TargetVariableToRemove, const FNiagaraVariableMetaData& TargetVariableMetaData) const
 {
-	return TargetVariableMetaData.bCreatedInSystemEditor;
+	return TargetVariableMetaData.GetWasCreatedInSystemEditor();
 }
 
 void FNiagaraSystemToolkitParameterPanelViewModel::RenameParameter(const FNiagaraVariable& TargetVariableToRename, const FNiagaraVariableMetaData& TargetVariableMetaData, const FText& NewVariableNameText) const
 {
 	FScopedTransaction RenameParameter(LOCTEXT("RenameParameter", "Rename Parameter"));
-	const FString ExistingVariableScopeString = FNiagaraStackGraphUtilities::GetNamespaceStringForScriptParameterScope(TargetVariableMetaData.Scope);
-	const FName NewVariableName = FName(*(ExistingVariableScopeString + NewVariableNameText.ToString()));
+
+	FName NewVariableName;
+	ENiagaraParameterScope TargetVariableScope;
+	FString ExistingVariableScopeString;
+	if (TargetVariableMetaData.GetScope(TargetVariableScope))
+	{
+		ExistingVariableScopeString = FNiagaraTypeUtilities::GetNamespaceStringForScriptParameterScope(TargetVariableScope);
+		NewVariableName = FName(*(ExistingVariableScopeString + NewVariableNameText.ToString()));
+	}
+	else
+	{
+		NewVariableName = FName(*NewVariableNameText.ToString());
+	}
+
 	for (const TWeakObjectPtr<UNiagaraGraph>& Graph : GetEditableEmitterScriptGraphs())
 	{
 		if (ensureMsgf(Graph.IsValid(), TEXT("Editable Emitter Script Graph was stale when renaming parameter!")))
 		{
 			Graph->Modify();
-			Graph->RenameParameter(TargetVariableToRename, NewVariableName, TargetVariableMetaData.bIsStaticSwitch, TargetVariableMetaData.Scope); //@todo(ng) handle renaming system params
+			Graph->RenameParameter(TargetVariableToRename, NewVariableName, TargetVariableMetaData.GetIsStaticSwitch(), TargetVariableScope); //@todo(ng) handle renaming system params
 		}
 	}
 	Refresh();
@@ -376,12 +407,12 @@ void FNiagaraSystemToolkitParameterPanelViewModel::ChangeParameterScope(const FN
 
 bool FNiagaraSystemToolkitParameterPanelViewModel::CanModifyParameter(const FNiagaraVariable& TargetVariableToModify, const FNiagaraVariableMetaData& TargetVariableMetaData) const
 {
-	return TargetVariableMetaData.bCreatedInSystemEditor == true;
+	return TargetVariableMetaData.GetWasCreatedInSystemEditor() == true;
 }
 
 bool FNiagaraSystemToolkitParameterPanelViewModel::CanRenameParameter(const FNiagaraVariable& TargetVariableToRename, const FNiagaraVariableMetaData& TargetVariableMetaData, const FText& NewVariableNameText) const
 {
-	return TargetVariableMetaData.bCreatedInSystemEditor == true;
+	return TargetVariableMetaData.GetWasCreatedInSystemEditor() == true;
 }
 
 FReply FNiagaraSystemToolkitParameterPanelViewModel::HandleActionDragged(const TSharedPtr<FEdGraphSchemaAction>& InAction, const FPointerEvent& MouseEvent) const
@@ -428,22 +459,24 @@ const TArray<FNiagaraScriptVariableAndViewInfo> FNiagaraSystemToolkitParameterPa
 			for (auto Iter = GraphVarToScriptVarMap.CreateConstIterator(); Iter; ++Iter)
 			{
 				const FNiagaraVariableMetaData& MetaData = Iter.Value()->Metadata;
+				ENiagaraParameterScope MetaDataScope;
+				MetaData.GetScope(MetaDataScope);
 
-				if (MetaData.Scope == ENiagaraParameterScope::None || MetaData.Usage == ENiagaraScriptParameterUsage::None)
+				if (MetaDataScope == ENiagaraParameterScope::None || MetaData.GetUsage() == ENiagaraScriptParameterUsage::None)
 				{
 					ensureMsgf(false, TEXT("Invalid MetaData found for graph variable: %s"), *Iter.Value()->Variable.GetName().ToString());
 					continue;
 				}
-				else if (MetaData.Scope == ENiagaraParameterScope::Local)
+				else if (MetaDataScope == ENiagaraParameterScope::Local)
 				{
 					// Note, the MetaData.Usage being local is fine to display for the System toolkit.
 					continue;
 				}
-				else if (MetaData.Scope == ENiagaraParameterScope::Input) // Do not expose inputs as configurable values
+				else if (MetaDataScope == ENiagaraParameterScope::Input) // Do not expose inputs as configurable values
 				{
 					continue;
 				}
-				else if (MetaData.Scope == ENiagaraParameterScope::ScriptPersistent || MetaData.Scope == ENiagaraParameterScope::ScriptTransient)
+				else if (MetaDataScope == ENiagaraParameterScope::ScriptPersistent || MetaDataScope == ENiagaraParameterScope::ScriptTransient)
 				{
 					//@todo(ng) Skip script alias parameters until we can resolve them!
 					continue;
@@ -546,26 +579,31 @@ void FNiagaraScriptToolkitParameterPanelViewModel::AddParameter(const FNiagaraVa
 {
 	FScopedTransaction AddParameter(LOCTEXT("AddParameterFromParameterPanel", "Add Parameter"));
 	UNiagaraGraph::FAddParameterOptions AddParameterOptions = UNiagaraGraph::FAddParameterOptions();
-	AddParameterOptions.NewParameterScope = InVariableMetaDataToAssign.Scope;
-	AddParameterOptions.NewParameterUsage = InVariableMetaDataToAssign.Usage;
+
+	ENiagaraParameterScope InVariableScope;
+	InVariableMetaDataToAssign.GetScope(InVariableScope);
+	ENiagaraScriptParameterUsage InVariableUsage = InVariableMetaDataToAssign.GetUsage();
+
+	AddParameterOptions.NewParameterScope = InVariableScope;
+	AddParameterOptions.NewParameterUsage = InVariableMetaDataToAssign.GetUsage();
 	FNiagaraVariable DuplicateVar = InVariableToAdd; //@todo(ng) rewrite
 
 	UNiagaraGraph* Graph = ScriptViewModel->GetGraphViewModel()->GetGraph();
 	Graph->Modify();
 	const FVector2D NewNodePos = Graph->GetGoodPlaceForNewNode();
-	if (InVariableMetaDataToAssign.Usage == ENiagaraScriptParameterUsage::Input)
+	if (InVariableUsage == ENiagaraScriptParameterUsage::Input)
 	{
 		UNiagaraNodeParameterMapGet* NewMapGet = FNiagaraSchemaAction_NewNode::SpawnNodeFromTemplate<UNiagaraNodeParameterMapGet>(Graph, NewObject<UNiagaraNodeParameterMapGet>(), NewNodePos);
 		NewMapGet->Modify();
 		NewMapGet->AddParameter(DuplicateVar, AddParameterOptions);
 	}
-	else if (InVariableMetaDataToAssign.Usage == ENiagaraScriptParameterUsage::Output)
+	else if (InVariableUsage == ENiagaraScriptParameterUsage::Output)
 	{
 		UNiagaraNodeParameterMapSet* NewMapSet = FNiagaraSchemaAction_NewNode::SpawnNodeFromTemplate<UNiagaraNodeParameterMapSet>(Graph, NewObject<UNiagaraNodeParameterMapSet>(), NewNodePos);
 		NewMapSet->Modify();
 		NewMapSet->AddParameter(DuplicateVar, AddParameterOptions);
 	}
-	else if (InVariableMetaDataToAssign.Usage == ENiagaraScriptParameterUsage::InputOutput)
+	else if (InVariableUsage == ENiagaraScriptParameterUsage::InputOutput)
 	{
 		ensureMsgf(false, TEXT("Unexpected usage encountered when adding parameter through parameter panel view model!"));
 	}
@@ -591,24 +629,38 @@ bool FNiagaraScriptToolkitParameterPanelViewModel::CanRemoveParameter(const FNia
 void FNiagaraScriptToolkitParameterPanelViewModel::RenameParameter(const FNiagaraVariable& TargetVariableToRename, const FNiagaraVariableMetaData& TargetVariableMetaData, const FText& NewVariableNameText) const
 {
 	FScopedTransaction RenameParameterAndReferencedPins(LOCTEXT("RenameParameterAndReferencedPins", "Rename parameter and referenced pins"));
-	const FString ExistingVariableScopeString = FNiagaraStackGraphUtilities::GetNamespaceStringForScriptParameterMetaData(TargetVariableMetaData);
-	const FName NewVariableName = FName(*(ExistingVariableScopeString + NewVariableNameText.ToString()));
+
+	FName NewVariableName;
+	if (TargetVariableMetaData.GetIsUsingLegacyNameString())
+	{
+		NewVariableName = FName(*NewVariableNameText.ToString());
+	}
+	else
+	{
+		const FString ExistingVariableScopeString = TargetVariableMetaData.GetNamespaceString();
+		NewVariableName = FName(*(ExistingVariableScopeString + NewVariableNameText.ToString()));
+	}
+
+	
 	UNiagaraGraph* Graph = ScriptViewModel->GetGraphViewModel()->GetGraph();
 	Graph->Modify();
-	Graph->RenameParameter(TargetVariableToRename, NewVariableName, TargetVariableMetaData.bIsStaticSwitch, TargetVariableMetaData.Scope);
+
+	ENiagaraParameterScope TargetVariableScope = ENiagaraParameterScope::None;
+	TargetVariableMetaData.GetScope(TargetVariableScope);
+	Graph->RenameParameter(TargetVariableToRename, NewVariableName, TargetVariableMetaData.GetIsStaticSwitch(), TargetVariableScope);
 }
 
 void FNiagaraScriptToolkitParameterPanelViewModel::ChangeParameterScope(const FNiagaraVariable& TargetVariableToModify, const FNiagaraVariableMetaData& TargetVariableMetaData, const ENiagaraParameterScope NewVariableScope) const
 {
-	if (ensureMsgf(TargetVariableMetaData.Usage != ENiagaraScriptParameterUsage::Output, TEXT("Tried to change scope of output parameter!")))
+	if (ensureMsgf(TargetVariableMetaData.GetUsage() != ENiagaraScriptParameterUsage::Output, TEXT("Tried to change scope of output parameter!")))
 	{
 		FScopedTransaction ChangeParameterScopeAndReferencedPins(LOCTEXT("ChangeParameterScopeAndReferencedPins", "Change parameter scope, Rename parameter and referenced pins"));
 
-		bool bIsInitialValue = TargetVariableMetaData.Usage == ENiagaraScriptParameterUsage::InitialValueInput;
+		bool bIsInitialValue = TargetVariableMetaData.GetUsage() == ENiagaraScriptParameterUsage::InitialValueInput;
 		const FName NewVariableName = FNiagaraStackGraphUtilities::GetVariableNameForScope(TargetVariableToModify.GetName(), NewVariableScope, bIsInitialValue);
 		UNiagaraGraph* Graph = ScriptViewModel->GetGraphViewModel()->GetGraph();
 		Graph->Modify();
-		Graph->RenameParameter(TargetVariableToModify, NewVariableName, TargetVariableMetaData.bIsStaticSwitch, NewVariableScope);
+		Graph->RenameParameter(TargetVariableToModify, NewVariableName, TargetVariableMetaData.GetIsStaticSwitch(), NewVariableScope);
 	}
 }
 
@@ -624,9 +676,13 @@ bool FNiagaraScriptToolkitParameterPanelViewModel::CanRenameParameter(const FNia
 	FName NewName = FName(*NewVariableNameText.ToString());
 	for (const FNiagaraScriptVariableAndViewInfo& ViewedParameter : CachedViewedParameters)
 	{
-		if (NewName == ViewedParameter.MetaData.CachedNamespacelessVariableName)
+		FName ParameterName;
+		if (ViewedParameter.MetaData.GetParameterName(ParameterName))
 		{
-			return false; //@todo(ng) wrap this into the verify logic
+			if (NewName == ParameterName)
+			{
+				return false; //@todo(ng) wrap this into the verify logic
+			}
 		}
 	}
 	return true;
@@ -699,28 +755,42 @@ TStaticArray<FScopeIsEnabledAndTooltip, (int32)ENiagaraParameterScope::Num> FNia
 	}
 
 	// Prevent setting enum values that would alias an existing parameter
-	const FName InName = InVarMetaData.CachedNamespacelessVariableName;
-	for (const FNiagaraScriptVariableAndViewInfo& ViewedParameter : CachedViewedParameters)
+	FName InName;
+	if (InVarMetaData.GetParameterName(InName))
 	{
-		if (InName == ViewedParameter.MetaData.CachedNamespacelessVariableName)
+		for (const FNiagaraScriptVariableAndViewInfo& ViewedParameter : CachedViewedParameters)
+		{
+			FName CachedParameterName;
+			if (ViewedParameter.MetaData.GetParameterName(CachedParameterName))
+			{
+				if (InName == CachedParameterName)
+				{
+					ENiagaraParameterScope ViewedParameterScope;
+					ViewedParameter.MetaData.GetScope(ViewedParameterScope);
+
+					PerScopeInfo[(int32)ViewedParameterScope].bEnabled = false;
+					PerScopeInfo[(int32)ViewedParameterScope].Tooltip = LOCTEXT("NiagaraInvalidScopeSelectionNameAlias", "Cannot select scope {0}: Parameter with same name already has this scope."); //@todo(ng) get scope
+				}
+			}
+		}
+
+		// Prevent making Module namespace parameters in function and dynamic input scripts
+		if (ScriptViewModel->GetStandaloneScript()->GetUsage() != ENiagaraScriptUsage::Module)
 		{
-			PerScopeInfo[(int32)ViewedParameter.MetaData.Scope].bEnabled = false;
-			PerScopeInfo[(int32)ViewedParameter.MetaData.Scope].Tooltip = LOCTEXT("NiagaraInvalidScopeSelectionNameAlias", "Cannot select scope {0}: Parameter with same name already has this scope."); //@todo(ng) get scope
+			PerScopeInfo[(int32)ENiagaraParameterScope::Input].bEnabled = false;
+			PerScopeInfo[(int32)ENiagaraParameterScope::Input].Tooltip = LOCTEXT("NiagaraInvalidScopeSelectionModule", "Cannot select scope {0}: Scope is only valid in Module Scripts."); //@todo(ng) get scope
 		}
-	}
 
-	// Prevent making Module namespace parameters in function and dynamic input scripts
-	if (ScriptViewModel->GetStandaloneScript()->GetUsage() != ENiagaraScriptUsage::Module)
-	{
-		PerScopeInfo[(int32)ENiagaraParameterScope::Input].bEnabled = false;
-		PerScopeInfo[(int32)ENiagaraParameterScope::Input].Tooltip = LOCTEXT("NiagaraInvalidScopeSelectionModule", "Cannot select scope {0}: Scope is only valid in Module Scripts."); //@todo(ng) get scope
+		const TArray<ENiagaraParameterScope> InvalidParameterScopes = ScriptViewModel->GetStandaloneScript()->GetUnsupportedParameterScopes();
+		for (ENiagaraParameterScope InvalidScope : InvalidParameterScopes)
+		{
+			PerScopeInfo[(int32)InvalidScope].bEnabled = false;
+			PerScopeInfo[(int32)InvalidScope].Tooltip = LOCTEXT("NiagaraInvalidScopeSelectionUsageBitmask", "Cannot select scope {0}: Script Usage flags do not support a usage with this scope."); //@todo(ng) rewrite
+		}
 	}
-
-	const TArray<ENiagaraParameterScope> InvalidParameterScopes = ScriptViewModel->GetStandaloneScript()->GetUnsupportedParameterScopes();
-	for (ENiagaraParameterScope InvalidScope : InvalidParameterScopes)
+	else
 	{
-		PerScopeInfo[(int32)InvalidScope].bEnabled = false;
-		PerScopeInfo[(int32)InvalidScope].Tooltip = LOCTEXT("NiagaraInvalidScopeSelectionUsageBitmask", "Cannot select scope {0}: Script Usage flags do not support a usage with this scope."); //@todo(ng) rewrite
+		// Failed to get parameter name as an override name is set, allow any scope.
 	}
 
 	return PerScopeInfo;
@@ -744,26 +814,28 @@ const TArray<FNiagaraScriptVariableAndViewInfo> FNiagaraScriptToolkitParameterPa
 			}
 
 			const FNiagaraVariableMetaData& MetaData = Iter.Value()->Metadata;
+			ENiagaraParameterScope MetaDataScope = ENiagaraParameterScope::None;
+			MetaData.GetScope(MetaDataScope);
 
-			if (MetaData.bIsStaticSwitch == false && (MetaData.Scope == ENiagaraParameterScope::None || MetaData.Usage == ENiagaraScriptParameterUsage::None) )
+			if (MetaData.GetIsStaticSwitch() == false && (MetaDataScope == ENiagaraParameterScope::None || MetaData.GetUsage() == ENiagaraScriptParameterUsage::None) )
 			{
 				// Parameters that are not static switches must have a scope and usage set.
 				ensureMsgf(false, TEXT("Invalid MetaData found for graph variable: %s"), *Iter.Value()->Variable.GetName().ToString());
 				continue;
 			}
-			else if (MetaData.Usage == ENiagaraScriptParameterUsage::InputOutput)
+			else if (MetaData.GetUsage() == ENiagaraScriptParameterUsage::InputOutput)
 			{
 				// Need two stack entries to represent this script parameter as both an input and output.
 				FNiagaraVariableMetaData InputVariableMetaData = MetaData;
 				FNiagaraVariableMetaData OutputVariableMetaData = MetaData;
-				InputVariableMetaData.Usage = ENiagaraScriptParameterUsage::Input;
-				OutputVariableMetaData.Usage = ENiagaraScriptParameterUsage::Output;
+				InputVariableMetaData.SetUsage( ENiagaraScriptParameterUsage::Input);
+				OutputVariableMetaData.SetUsage(ENiagaraScriptParameterUsage::Output);
 
 				const TStaticArray<FScopeIsEnabledAndTooltip, (int32)ENiagaraParameterScope::Num> PerScopeInfo = GetParameterScopesEnabledAndTooltips(Variable, MetaData);
 				ViewedParameters.Add(FNiagaraScriptVariableAndViewInfo(Variable, InputVariableMetaData, PerScopeInfo));
 				ViewedParameters.Add(FNiagaraScriptVariableAndViewInfo(Variable, OutputVariableMetaData));
 			}
-			else if(MetaData.Usage == ENiagaraScriptParameterUsage::Input || MetaData.Usage == ENiagaraScriptParameterUsage::InitialValueInput)
+			else if(MetaData.GetUsage() == ENiagaraScriptParameterUsage::Input || MetaData.GetUsage() == ENiagaraScriptParameterUsage::InitialValueInput)
 			{
 				const TStaticArray<FScopeIsEnabledAndTooltip, (int32)ENiagaraParameterScope::Num> PerScopeInfo = GetParameterScopesEnabledAndTooltips(Variable, MetaData);
 				ViewedParameters.Add(FNiagaraScriptVariableAndViewInfo(Variable, MetaData, PerScopeInfo));
diff --git a/Engine/Plugins/FX/Niagara/Source/NiagaraEditor/Private/ViewModels/Stack/NiagaraStackFunctionInput.cpp b/Engine/Plugins/FX/Niagara/Source/NiagaraEditor/Private/ViewModels/Stack/NiagaraStackFunctionInput.cpp
--- a/Engine/Plugins/FX/Niagara/Source/NiagaraEditor/Private/ViewModels/Stack/NiagaraStackFunctionInput.cpp
+++ b/Engine/Plugins/FX/Niagara/Source/NiagaraEditor/Private/ViewModels/Stack/NiagaraStackFunctionInput.cpp
@@ -2025,7 +2025,7 @@ void UNiagaraStackFunctionInput::UpdateValuesFromScriptDefaults(FInputValues& In
 		// Static switch parameters are always locally set values.
 		if (InputScriptVariable != nullptr)
 		{
-			TSharedPtr<FStructOnScope> StaticSwitchLocalStruct = FNiagaraEditorUtilities::StaticSwitchDefaultIntToStructOnScope(InputScriptVariable->Metadata.StaticSwitchDefaultValue, InputType);
+			TSharedPtr<FStructOnScope> StaticSwitchLocalStruct = FNiagaraEditorUtilities::StaticSwitchDefaultIntToStructOnScope(InputScriptVariable->Metadata.GetStaticSwitchDefaultValue(), InputType);
 			if(ensureMsgf(StaticSwitchLocalStruct.IsValid(), TEXT("Unsupported static struct default value.")))
 			{ 
 				InInputValues.Mode = EValueMode::Local;
diff --git a/Engine/Plugins/FX/Niagara/Source/NiagaraEditor/Private/ViewModels/Stack/NiagaraStackGraphUtilities.cpp b/Engine/Plugins/FX/Niagara/Source/NiagaraEditor/Private/ViewModels/Stack/NiagaraStackGraphUtilities.cpp
--- a/Engine/Plugins/FX/Niagara/Source/NiagaraEditor/Private/ViewModels/Stack/NiagaraStackGraphUtilities.cpp
+++ b/Engine/Plugins/FX/Niagara/Source/NiagaraEditor/Private/ViewModels/Stack/NiagaraStackGraphUtilities.cpp
@@ -2540,146 +2540,133 @@ void FNiagaraStackGraphUtilities::RenameReferencingParameters(UNiagaraSystem& Sy
 	}
 }
 
-FString FNiagaraStackGraphUtilities::GetNamespaceStringForScriptParameterScope(const ENiagaraParameterScope& InScope)
-{
-	switch (InScope) {
-		case ENiagaraParameterScope::Engine:
-			return PARAM_MAP_ENGINE_STR;
-			break;
-		case ENiagaraParameterScope::Owner:
-			return PARAM_MAP_ENGINE_OWNER_STR;
-			break;
-		case ENiagaraParameterScope::User:
-			return PARAM_MAP_USER_STR;
-			break;
-		case ENiagaraParameterScope::System:
-			return PARAM_MAP_SYSTEM_STR;
-			break;
-		case ENiagaraParameterScope::Emitter:
-			return PARAM_MAP_EMITTER_STR;
-			break;
-		case ENiagaraParameterScope::Particles:
-			return PARAM_MAP_ATTRIBUTE_STR;
-			break;
-		case ENiagaraParameterScope::ScriptPersistent:
-			return PARAM_MAP_SCRIPT_PERSISTENT_STR;
-			break;
-		case ENiagaraParameterScope::ScriptTransient:
-			return PARAM_MAP_SCRIPT_TRANSIENT_STR;
-			break;
-		case ENiagaraParameterScope::Input:
-			return PARAM_MAP_MODULE_STR;
-			break;
-		case ENiagaraParameterScope::Local:
-			return PARAM_MAP_LOCAL_STR;
-			break;
-		default:
-			checkf(false, TEXT("Unhandled parameter scope encountered!"));
-			return FString();
-			break;
-	};
-}
-
-FString FNiagaraStackGraphUtilities::GetNamespaceStringForScriptParameterMetaData(const FNiagaraVariableMetaData& InMetaData)
-{
-	if (InMetaData.IsInputOrLocalUsage())
-	{
-		FString NamespaceString = FNiagaraStackGraphUtilities::GetNamespaceStringForScriptParameterScope(InMetaData.Scope);
-		if (InMetaData.Usage == ENiagaraScriptParameterUsage::InitialValueInput)
-		{
-			NamespaceString.Append(PARAM_MAP_INITIAL_STR);
-		}
-		return NamespaceString;
-	}
-	return PARAM_MAP_OUTPUT_STR;
-}
-
-TArray<FName> FNiagaraStackGraphUtilities::DecomposeVariableNamespace(const FName& InVarName, FName& OutFriendlyName)
+TArray<FName> FNiagaraStackGraphUtilities::DecomposeVariableNamespace(const FName& InVarNameToken, FName& OutName)
 {
 	int32 DotIndex;
 	TArray<FName> OutNamespaces;
-	FString VarNameString = InVarName.ToString();
+	FString VarNameString = InVarNameToken.ToString();
 	while (VarNameString.FindChar(TEXT('.'), DotIndex))
 	{
 		OutNamespaces.Add(FName(*VarNameString.Left(DotIndex)));
 		VarNameString = *VarNameString.RightChop(DotIndex + 1);
 	}
-	OutFriendlyName = FName(*VarNameString);
+	OutName = FName(*VarNameString);
 	return OutNamespaces;
 }
 
-void FNiagaraStackGraphUtilities::SetParameterMetaDataFromName(const FName& InVarName, ENiagaraParameterScope& OutScope, ENiagaraScriptParameterUsage& OutUsage, FName& OutFriendlyName)
+void FNiagaraStackGraphUtilities::GetParameterMetaDataFromName(const FName& InVarNameToken, FNiagaraVariableMetaData& OutMetaData)
 {
-	TArray<FName> DecomposedNamespaces = FNiagaraStackGraphUtilities::DecomposeVariableNamespace(InVarName, OutFriendlyName);
-	if (DecomposedNamespaces.Num() == 0)
-	{
-		UE_LOG(LogNiagaraEditor, Display, TEXT("Unexpected parameter encountered without a namespace: %s"), *InVarName.ToString());
-	}
-
-	bool bCheckIsInitialValue = false;
-	for (int NamespaceIdx = 0; NamespaceIdx < DecomposedNamespaces.Num(); ++NamespaceIdx)
-	{
-		const FName& Namespace = DecomposedNamespaces[NamespaceIdx];
-
-		if (bCheckIsInitialValue)
-		{
-			if (Namespace == FNiagaraParameterHandle::InitialNamespace)
-			{
-				OutUsage = ENiagaraScriptParameterUsage::InitialValueInput;
-			}
-			return;
-		}
+	auto MarkAsLegacyCustomName = [&OutMetaData]() {
+		OutMetaData.SetScope(ENiagaraParameterScope::Custom);
+		OutMetaData.SetIsUsingLegacyNameString(true);
+	};
 
+	auto GetMetaDataForFirstNamespace = [MarkAsLegacyCustomName](const FName& Namespace, FNiagaraVariableMetaData& OutMetaData)-> bool /*bNextNamespaceCanBeValid*/ {
 		if (Namespace == FNiagaraParameterHandle::LocalNamespace)
 		{
-			OutScope = ENiagaraParameterScope::Local;
-			OutUsage = ENiagaraScriptParameterUsage::Local;
-			return;
+			OutMetaData.SetScope(ENiagaraParameterScope::Local);
+			OutMetaData.SetUsage(ENiagaraScriptParameterUsage::Local);
+			return true;
 		}
 		else if (Namespace == FNiagaraParameterHandle::ModuleNamespace)
 		{
-			OutScope = ENiagaraParameterScope::Input;
-			return;
+			OutMetaData.SetScope(ENiagaraParameterScope::Input);
+			return false;
 		}
 		else if (Namespace == FNiagaraParameterHandle::UserNamespace)
 		{
-			OutScope = ENiagaraParameterScope::User;
-			return;
+			OutMetaData.SetScope(ENiagaraParameterScope::User);
+			return false;
 		}
 		else if (Namespace == FNiagaraParameterHandle::EngineNamespace)
 		{
-			OutScope = ENiagaraParameterScope::Engine;
-			return;
+			OutMetaData.SetScope(ENiagaraParameterScope::Engine);
+			return true;
 		}
 		else if (Namespace == FNiagaraParameterHandle::SystemNamespace)
 		{
-			OutScope = ENiagaraParameterScope::System;
-			bCheckIsInitialValue = true;
+			OutMetaData.SetScope(ENiagaraParameterScope::System);
+			return true;
 		}
 		else if (Namespace == FNiagaraParameterHandle::EmitterNamespace)
 		{
-			if (Namespace == FNiagaraParameterHandle::OwnerNamespace)
-			{
-				OutScope = ENiagaraParameterScope::Owner;
-			}
-			else
-			{
-				OutScope = ENiagaraParameterScope::Emitter;
-			}
-			bCheckIsInitialValue = true;
+			OutMetaData.SetScope(ENiagaraParameterScope::Emitter);
+			return true;
 		}
 		else if (Namespace == FNiagaraParameterHandle::ParticleAttributeNamespace)
 		{
-			OutScope = ENiagaraParameterScope::Particles;
-			bCheckIsInitialValue = true;
+			OutMetaData.SetScope(ENiagaraParameterScope::Particles);
+			return true;
 		}
+
+		MarkAsLegacyCustomName();
+		return false;
+	};
+
+	auto GetMetaDataForInitialNamespace = [](const FName& Namespace, FNiagaraVariableMetaData& OutMetaData)-> bool /*bFoundInitialNamespace*/ {
+		if (Namespace == FNiagaraParameterHandle::InitialNamespace)
+		{
+			OutMetaData.SetUsage(ENiagaraScriptParameterUsage::InitialValueInput);
+			return true;
+		}
+		return false;
+	};
+
+	auto GetScopeCanHaveInitialNamespaceSuffix = [](ENiagaraParameterScope InScope)-> bool /*bCanHaveInitialNamespaceSuffix*/ {
+		switch (InScope) {
+		case ENiagaraParameterScope::System:
+		case ENiagaraParameterScope::Emitter:
+		case ENiagaraParameterScope::Particles:
+			return true;
+		default:
+			return false;
+		};
+		return false;
+	};
+
+	FName NamespacelessName;
+	TArray<FName> DecomposedNamespaces = FNiagaraStackGraphUtilities::DecomposeVariableNamespace(InVarNameToken, NamespacelessName);
+	OutMetaData.SetCachedNamespacelessVariableName(NamespacelessName);
+	if (DecomposedNamespaces.Num() == 0)
+	{
+		UE_LOG(LogNiagaraEditor, Display, TEXT("Unexpected parameter encountered without a namespace: %s"), *InVarNameToken.ToString());
+		MarkAsLegacyCustomName();
+		return;
 	}
-	
-	if (bCheckIsInitialValue == false)
+	else if (DecomposedNamespaces.Num() == 1)
 	{
-		// custom user namespace
-		OutScope = ENiagaraParameterScope::ScriptTransient;
+		GetMetaDataForFirstNamespace(DecomposedNamespaces[0], OutMetaData);
+		return;
 	}
+	else if (DecomposedNamespaces.Num() == 2)
+	{
+		bool bNextNamespaceCanBeValid = GetMetaDataForFirstNamespace(DecomposedNamespaces[0], OutMetaData);
+		if (bNextNamespaceCanBeValid)
+		{
+			ENiagaraParameterScope FirstNamespaceScope;
+			OutMetaData.GetScope(FirstNamespaceScope);
+			if (FirstNamespaceScope == ENiagaraParameterScope::Local)
+			{
+				// "local.module." namespaces may be handled as local scopes and do not need to be marked as legacy namespaces.
+				if (DecomposedNamespaces[1] == FNiagaraParameterHandle::ModuleNamespace)
+				{
+					return;
+				}
+			}
+			else if (GetScopeCanHaveInitialNamespaceSuffix(FirstNamespaceScope))
+			{
+				bool bFoundInitialNamespace = GetMetaDataForInitialNamespace(DecomposedNamespaces[1], OutMetaData);
+				if (bFoundInitialNamespace)
+				{
+					return;
+				}
+			}
+		}
+		MarkAsLegacyCustomName();
+		return;
+	}
+
+	MarkAsLegacyCustomName();
 }
 
 FString FNiagaraStackGraphUtilities::GetNamespacelessVariableNameString(const FName& InVarName)
@@ -2696,7 +2683,7 @@ FString FNiagaraStackGraphUtilities::GetNamespacelessVariableNameString(const FN
 
 FName FNiagaraStackGraphUtilities::GetVariableNameForScope(const FName& CurrentVariableName, const ENiagaraParameterScope NewVariableScope, bool bIsInitialValue)
 {
-	FString NewVariableScopeString = FNiagaraStackGraphUtilities::GetNamespaceStringForScriptParameterScope(NewVariableScope);
+	FString NewVariableScopeString = FNiagaraTypeUtilities::GetNamespaceStringForScriptParameterScope(NewVariableScope);
 	if (bIsInitialValue)
 	{
 		NewVariableScopeString.Append(PARAM_MAP_INITIAL_STR);
diff --git a/Engine/Plugins/FX/Niagara/Source/NiagaraEditor/Private/Widgets/SNiagaraParameterPanel.cpp b/Engine/Plugins/FX/Niagara/Source/NiagaraEditor/Private/Widgets/SNiagaraParameterPanel.cpp
--- a/Engine/Plugins/FX/Niagara/Source/NiagaraEditor/Private/Widgets/SNiagaraParameterPanel.cpp
+++ b/Engine/Plugins/FX/Niagara/Source/NiagaraEditor/Private/Widgets/SNiagaraParameterPanel.cpp
@@ -38,7 +38,7 @@ void SNiagaraParameterPanel::Construct(const FArguments& InArgs, const TSharedPt
 
 	ParameterPanelViewModel->GetOnRefreshed().BindRaw(this, &SNiagaraParameterPanel::Refresh);
 
-	AddParameterButtons.SetNum(NiagaraParameterPanelSectionID::PARTICLES + 1); //@Todo(ng) verify
+	AddParameterButtons.SetNum(NiagaraParameterPanelSectionID::CUSTOM + 1); //@Todo(ng) verify
 
 	// Register all commands for right click on action node
 	FNiagaraParameterPanelCommands::Register();
@@ -96,8 +96,8 @@ void SNiagaraParameterPanel::Tick(const FGeometry& AllottedGeometry, const doubl
 void SNiagaraParameterPanel::AddParameter(FNiagaraVariable NewVariable, const NiagaraParameterPanelSectionID::Type SectionID)
 {
 	FNiagaraVariableMetaData NewVariableMetaData = FNiagaraVariableMetaData();
-	NewVariableMetaData.Scope = NiagaraParameterPanelSectionID::GetScopeForNewParametersInSection(SectionID);
-	NewVariableMetaData.Usage = NiagaraParameterPanelSectionID::GetUsageForNewParametersInSection(SectionID);
+	NewVariableMetaData.SetScope(NiagaraParameterPanelSectionID::GetScopeForNewParametersInSection(SectionID));
+	NewVariableMetaData.SetUsage(NiagaraParameterPanelSectionID::GetUsageForNewParametersInSection(SectionID));
 
 	ParameterPanelViewModel->AddParameter(NewVariable, NewVariableMetaData);
 
@@ -121,8 +121,28 @@ void SNiagaraParameterPanel::CollectAllActions(FGraphActionListBuilderBase& OutA
 {
 	const FText TooltipFormat = LOCTEXT("Parameters", "Name: {0} \nType: {1}"); //@todo(ng) move to const static
 
+	auto SortParameterPanelEntries = [](const FNiagaraScriptVariableAndViewInfo& A, const FNiagaraScriptVariableAndViewInfo& B) {
+		if (A.MetaData.GetIsUsingLegacyNameString())
+		{
+			if (B.MetaData.GetIsUsingLegacyNameString())
+			{
+				return A.ScriptVariable.GetName().LexicalLess(B.ScriptVariable.GetName());
+			}
+			return false;
+		}
+		else if (B.MetaData.GetIsUsingLegacyNameString())
+		{
+			return true;
+		}
+		FName AName;
+		FName BName;
+		A.MetaData.GetParameterName(AName);
+		B.MetaData.GetParameterName(BName);
+		return AName.LexicalLess(BName);
+	};
+
 	TArray<FNiagaraScriptVariableAndViewInfo> VarAndViewInfos = ParameterPanelViewModel->GetViewedParameters();
-	VarAndViewInfos.Sort([](const FNiagaraScriptVariableAndViewInfo& A, const FNiagaraScriptVariableAndViewInfo& B) { return A.MetaData.CachedNamespacelessVariableName.LexicalLess(B.MetaData.CachedNamespacelessVariableName); });
+	VarAndViewInfos.Sort(SortParameterPanelEntries);
 	for (const FNiagaraScriptVariableAndViewInfo& VarAndViewInfo : VarAndViewInfos)
 	{
 		const FText Name = FText::FromName(VarAndViewInfo.ScriptVariable.GetName());
@@ -313,7 +333,7 @@ void SNiagaraAddParameterMenu2::Construct(const FArguments& InArgs, TArray<TWeak
 
 	Graphs = InGraphs;
 	NewParameterScope = InNewParameterScope;
-	NewParameterScopeText = FText::FromString(FNiagaraStackGraphUtilities::GetNamespaceStringForScriptParameterScope(NewParameterScope));
+	NewParameterScopeText = FText::FromString(FNiagaraTypeUtilities::GetNamespaceStringForScriptParameterScope(NewParameterScope));
 
 	ChildSlot
 	[
diff --git a/Engine/Plugins/FX/Niagara/Source/NiagaraEditor/Public/NiagaraEditorCommon.h b/Engine/Plugins/FX/Niagara/Source/NiagaraEditor/Public/NiagaraEditorCommon.h
--- a/Engine/Plugins/FX/Niagara/Source/NiagaraEditor/Public/NiagaraEditorCommon.h
+++ b/Engine/Plugins/FX/Niagara/Source/NiagaraEditor/Public/NiagaraEditorCommon.h
@@ -206,6 +206,7 @@ namespace NiagaraParameterPanelSectionID
 		SYSTEM,
 		EMITTER,
 		PARTICLES,
+		CUSTOM,
 	};
 
 	FText OnGetSectionTitle(const NiagaraParameterPanelSectionID::Type InSection);
diff --git a/Engine/Plugins/FX/Niagara/Source/NiagaraEditor/Public/NiagaraGraph.h b/Engine/Plugins/FX/Niagara/Source/NiagaraEditor/Public/NiagaraGraph.h
--- a/Engine/Plugins/FX/Niagara/Source/NiagaraEditor/Public/NiagaraGraph.h
+++ b/Engine/Plugins/FX/Niagara/Source/NiagaraEditor/Public/NiagaraGraph.h
@@ -332,8 +332,6 @@ class UNiagaraGraph : public UEdGraph
 
 	void CopyCachedReferencesMap(UNiagaraGraph* TargetGraph);
 
-	void FixupParameterUsageMetaData();
-
 	bool IsPinVisualWidgetProviderRegistered() const;
 
 	/** Wrapper to provide visual widgets for pins that are managed by external viewmodels.
diff --git a/Engine/Plugins/FX/Niagara/Source/NiagaraEditor/Public/ViewModels/Stack/NiagaraStackGraphUtilities.h b/Engine/Plugins/FX/Niagara/Source/NiagaraEditor/Public/ViewModels/Stack/NiagaraStackGraphUtilities.h
--- a/Engine/Plugins/FX/Niagara/Source/NiagaraEditor/Public/ViewModels/Stack/NiagaraStackGraphUtilities.h
+++ b/Engine/Plugins/FX/Niagara/Source/NiagaraEditor/Public/ViewModels/Stack/NiagaraStackGraphUtilities.h
@@ -192,12 +192,10 @@ namespace FNiagaraStackGraphUtilities
 
 	void GatherRenamedStackFunctionOutputVariableNames(UNiagaraEmitter* Emitter, UNiagaraNodeFunctionCall& FunctionCallNode, const FString& OldFunctionName, const FString& NewFunctionName, TMap<FName, FName>& OutOldToNewNameMap);
 	void GatherRenamedStackFunctionInputAndOutputVariableNames(UNiagaraEmitter* Emitter, UNiagaraNodeFunctionCall& FunctionCallNode, const FString& OldFunctionName, const FString& NewFunctionName, TMap<FName, FName>& OutOldToNewNameMap);
-	FString GetNamespaceStringForScriptParameterScope(const ENiagaraParameterScope& InScope);
-	FString GetNamespaceStringForScriptParameterMetaData(const FNiagaraVariableMetaData& InMetaData);
 
-	TArray<FName> DecomposeVariableNamespace(const FName& InVarName, FName& OutFriendlyName);
+	TArray<FName> DecomposeVariableNamespace(const FName& InVarNameToken, FName& OutName);
 
-	void SetParameterMetaDataFromName(const FName& InVarName, ENiagaraParameterScope& OutScope, ENiagaraScriptParameterUsage& OutUsage, FName& OutFriendlyName);
+	void GetParameterMetaDataFromName(const FName& InVarNameToken, FNiagaraVariableMetaData& OutMetaData);
 
 	FString GetNamespacelessVariableNameString(const FName& InVarName);
 
