Fix build error from submitting intermediate shelf.

diff --git a/Engine/Plugins/FX/Niagara/Source/NiagaraEditor/Private/ViewModels/NiagaraParameterPanelViewModel.cpp b/Engine/Plugins/FX/Niagara/Source/NiagaraEditor/Private/ViewModels/NiagaraParameterPanelViewModel.cpp
--- a/Engine/Plugins/FX/Niagara/Source/NiagaraEditor/Private/ViewModels/NiagaraParameterPanelViewModel.cpp
+++ b/Engine/Plugins/FX/Niagara/Source/NiagaraEditor/Private/ViewModels/NiagaraParameterPanelViewModel.cpp
@@ -288,9 +288,9 @@ NiagaraParameterPanelSectionID::Type FNiagaraSystemToolkitParameterPanelViewMode
 	{
 		//return NiagaraParameterPanelSectionID::GetSectionForScope(VarAndViewInfo.ScriptVariableMetaData.LastPrecompileScope); //@todo(ng) cache the known scope in the niagarasystem
 		FNiagaraVariableMetaData OutMetaData;
-		FNiagaraStackGraphUtilities::GetParameterMetaDataFromName(VarAndViewInfo.ScriptVariable.GetName(), VarAndViewInfo.MetaData);
+		FNiagaraStackGraphUtilities::GetParameterMetaDataFromName(VarAndViewInfo.ScriptVariable.GetName(), OutMetaData);
 		ENiagaraParameterScope OutMetaDataScope;
-		VarAndViewInfo.MetaData.GetScope(OutMetaDataScope);
+		OutMetaData.GetScope(OutMetaDataScope);
 		return NiagaraParameterPanelSectionID::GetSectionForScope(OutMetaDataScope);
 	}
 }
