Fixes for threading issues in Niagara Audio Listeners- RegisterSubmixBufferListener and UnregisterSubmixBufferListener dispatch a raw pointer to the audio thread, so we can't destroy the buffer listener without fencing on the audio thread.
This also has to happen in editor as well, even though the game thread doesn't exist, since we could have a stale pointer enqueued when the next call to FAudioDevice::Update is called.

diff --git a/Engine/Plugins/FX/Niagara/Source/Niagara/Classes/NiagaraDataInterfaceAudio.h b/Engine/Plugins/FX/Niagara/Source/Niagara/Classes/NiagaraDataInterfaceAudio.h
--- a/Engine/Plugins/FX/Niagara/Source/Niagara/Classes/NiagaraDataInterfaceAudio.h
+++ b/Engine/Plugins/FX/Niagara/Source/Niagara/Classes/NiagaraDataInterfaceAudio.h
@@ -7,6 +7,7 @@
 #include "AudioDeviceManager.h"
 #include "Sound/SoundSubmix.h"
 #include "DSP/MultithreadedPatching.h"
+#include "UObject/WeakObjectPtrTemplates.h"
 #include "NiagaraDataInterfaceAudio.generated.h"
 
 /** Class used to to capture the audio stream of an arbitrary submix. */
@@ -21,7 +22,7 @@ public:
 	 *                               data to be overwritten by threads that produce audio. High 
 	 *                               values require more memory.
 	 */
-	FNiagaraSubmixListener(Audio::FPatchMixer& InMixer, int32 InNumSamplesToBuffer);
+	FNiagaraSubmixListener(Audio::FPatchMixer& InMixer, int32 InNumSamplesToBuffer, Audio::FDeviceId InDeviceId, USoundSubmix* InSoundSubmix);
 
 	FNiagaraSubmixListener(const FNiagaraSubmixListener& Other)
 	{
@@ -33,6 +34,8 @@ public:
 	/** Move submix listener. */
 	FNiagaraSubmixListener(FNiagaraSubmixListener&& Other);
 
+	void RegisterToSubmix();
+
 	virtual ~FNiagaraSubmixListener();
 
 	/** Returns the current sample rate of the current submix. */
@@ -48,10 +51,17 @@ public:
 private:
 	FNiagaraSubmixListener();
 
+
+	void UnregisterFromSubmix();
+
 	TAtomic<int32> NumChannelsInSubmix;
 	TAtomic<int32> SubmixSampleRate;
 
 	Audio::FPatchInput MixerInput;
+
+	Audio::FDeviceId AudioDeviceId;
+	USoundSubmix* Submix;
+	bool bIsRegistered;
 };
 
 /** Niagara data interface proxy for audio submix listener. */
@@ -100,7 +110,7 @@ private:
 
 	// Map of audio devices to submix listeners. Needed for editor where multiple audio devices
 	// may exist.
-	TMap<Audio::FDeviceId, FNiagaraSubmixListener> SubmixListeners;
+	TMap<Audio::FDeviceId, TUniquePtr<FNiagaraSubmixListener>> SubmixListeners;
 
 	// Mixer for sending audio 
 	Audio::FPatchMixer PatchMixer;
diff --git a/Engine/Plugins/FX/Niagara/Source/Niagara/Classes/NiagaraDataInterfaceAudioOscilloscope.h b/Engine/Plugins/FX/Niagara/Source/Niagara/Classes/NiagaraDataInterfaceAudioOscilloscope.h
--- a/Engine/Plugins/FX/Niagara/Source/Niagara/Classes/NiagaraDataInterfaceAudioOscilloscope.h
+++ b/Engine/Plugins/FX/Niagara/Source/Niagara/Classes/NiagaraDataInterfaceAudioOscilloscope.h
@@ -59,9 +59,9 @@ private:
 	void OnNewDeviceCreated(Audio::FDeviceId InID);
 	void OnDeviceDestroyed(Audio::FDeviceId InID);
 
-	TMap<Audio::FDeviceId, FNiagaraSubmixListener> SubmixListeners;
+	TMap<Audio::FDeviceId, TUniquePtr<FNiagaraSubmixListener>> SubmixListeners;
 
-	// This 
+	// This mixer is patched into by all instances of FNiagaraSubmixListener in the SubmixListeners map, and is consumed by DownsampleAudioToBuffer().
 	Audio::FPatchMixer PatchMixer;
 
 	USoundSubmix* SubmixRegisteredTo;
diff --git a/Engine/Plugins/FX/Niagara/Source/Niagara/Private/NiagaraDataInterfaceAudio.cpp b/Engine/Plugins/FX/Niagara/Source/Niagara/Private/NiagaraDataInterfaceAudio.cpp
--- a/Engine/Plugins/FX/Niagara/Source/Niagara/Private/NiagaraDataInterfaceAudio.cpp
+++ b/Engine/Plugins/FX/Niagara/Source/Niagara/Private/NiagaraDataInterfaceAudio.cpp
@@ -2,22 +2,32 @@
 
 #include "NiagaraDataInterfaceAudio.h"
 
-FNiagaraSubmixListener::FNiagaraSubmixListener(Audio::FPatchMixer& InMixer, int32 InNumSamplesToBuffer)
+FNiagaraSubmixListener::FNiagaraSubmixListener(Audio::FPatchMixer& InMixer, int32 InNumSamplesToBuffer, Audio::FDeviceId InDeviceId, USoundSubmix* InSoundSubmix)
 	: NumChannelsInSubmix(0)
 	, SubmixSampleRate(0)
 	, MixerInput(InMixer.AddNewInput(InNumSamplesToBuffer, 1.0f))
+	, AudioDeviceId(InDeviceId)
+	, Submix(InSoundSubmix)
+	, bIsRegistered(false)
 {
-
 }
 
 FNiagaraSubmixListener::FNiagaraSubmixListener()
 	: NumChannelsInSubmix(0)
 	, SubmixSampleRate(0)
+	, MixerInput()
+	, AudioDeviceId(INDEX_NONE)
+	, Submix(nullptr)
+	, bIsRegistered(false)
 {
 }
 
 FNiagaraSubmixListener::FNiagaraSubmixListener(FNiagaraSubmixListener&& Other)
+	: FNiagaraSubmixListener()
 {
+	UnregisterFromSubmix();
+	Other.UnregisterFromSubmix();
+
 	NumChannelsInSubmix = Other.NumChannelsInSubmix.Load();
 	Other.NumChannelsInSubmix = 0;
 
@@ -25,10 +35,56 @@ FNiagaraSubmixListener::FNiagaraSubmixListener(FNiagaraSubmixListener&& Other)
 	Other.SubmixSampleRate = 0;
 
 	MixerInput = MoveTemp(Other.MixerInput);
+
+	Submix = Other.Submix;
+	Other.Submix = nullptr;
+
+	AudioDeviceId = Other.AudioDeviceId;
+	Other.AudioDeviceId = INDEX_NONE;
+
+	RegisterToSubmix();
 }
 
 FNiagaraSubmixListener::~FNiagaraSubmixListener()
 {
+	UnregisterFromSubmix();
+}
+
+void FNiagaraSubmixListener::RegisterToSubmix()
+{
+	FAudioDeviceHandle DeviceHandle = FAudioDeviceManager::Get()->GetAudioDevice(AudioDeviceId);
+	ensure(DeviceHandle);
+
+	if (DeviceHandle)
+	{
+		DeviceHandle->RegisterSubmixBufferListener(this, Submix);
+		bIsRegistered = true;
+
+		// RegisterSubmixBufferListener lazily enqueues the registration on the audio thread,
+		// so we have to wait for the audio thread to destroy.
+		FAudioCommandFence Fence;
+		Fence.BeginFence();
+		Fence.Wait();
+	}
+}
+
+void FNiagaraSubmixListener::UnregisterFromSubmix()
+{
+	if (bIsRegistered)
+	{
+		FAudioDeviceHandle DeviceHandle = FAudioDeviceManager::Get()->GetAudioDevice(AudioDeviceId);
+		USoundSubmix* RawSubmixPtr = Submix;
+
+		check(DeviceHandle);
+		DeviceHandle->UnregisterSubmixBufferListener(this, Submix);
+		bIsRegistered = false;
+
+		// UnregisterSubmixBufferListener lazily enqueues the unregistration on the audio thread,
+		// so we have to wait for the audio thread to destroy.
+		FAudioCommandFence Fence;
+		Fence.BeginFence();
+		Fence.Wait();
+	}
 }
 
 float FNiagaraSubmixListener::GetSampleRate() const
@@ -60,7 +116,6 @@ FNiagaraDataInterfaceProxySubmix::FNiagaraDataInterfaceProxySubmix(int32 InNumSa
 	check(NumSamplesToBuffer > 0);
 
 	DeviceCreatedHandle = FAudioDeviceManagerDelegates::OnAudioDeviceCreated.AddRaw(this, &FNiagaraDataInterfaceProxySubmix::OnNewDeviceCreated);
-
 	DeviceDestroyedHandle = FAudioDeviceManagerDelegates::OnAudioDeviceDestroyed.AddRaw(this, &FNiagaraDataInterfaceProxySubmix::OnDeviceDestroyed);
 }
 
@@ -127,23 +182,14 @@ void FNiagaraDataInterfaceProxySubmix::OnDeviceDestroyed(Audio::FDeviceId InDevi
 void FNiagaraDataInterfaceProxySubmix::AddSubmixListener(const Audio::FDeviceId InDeviceId)
 {
 	check(!SubmixListeners.Contains(InDeviceId));
-
-	FAudioDeviceHandle DeviceHandle = FAudioDeviceManager::Get()->GetAudioDevice(InDeviceId);
-	check(DeviceHandle);
-
-	FNiagaraSubmixListener& Listener = SubmixListeners.Emplace(InDeviceId, FNiagaraSubmixListener(PatchMixer, NumSamplesToBuffer));
-
-	DeviceHandle->RegisterSubmixBufferListener(&Listener, SubmixRegisteredTo);
+	SubmixListeners.Emplace(InDeviceId, new FNiagaraSubmixListener(PatchMixer, NumSamplesToBuffer, InDeviceId, SubmixRegisteredTo));
+	SubmixListeners[InDeviceId]->RegisterToSubmix();
 }
 
 void FNiagaraDataInterfaceProxySubmix::RemoveSubmixListener(Audio::FDeviceId InDeviceId)
 {
 	if (SubmixListeners.Contains(InDeviceId))
 	{
-		FAudioDeviceHandle DeviceHandle = FAudioDeviceManager::Get()->GetAudioDevice(InDeviceId);
-		check(DeviceHandle);
-
-		DeviceHandle->UnregisterSubmixBufferListener(&SubmixListeners[InDeviceId], SubmixRegisteredTo);
 		SubmixListeners.Remove(InDeviceId);
 	}
 }
@@ -154,7 +200,7 @@ int32 FNiagaraDataInterfaceProxySubmix::GetNumChannels() const
 	SubmixListeners.GetKeys(DeviceIds);
 	for (const Audio::FDeviceId DeviceId : DeviceIds)
 	{
-		int32 NumChannels = SubmixListeners[DeviceId].GetNumChannels();
+		int32 NumChannels = SubmixListeners[DeviceId]->GetNumChannels();
 
 		if (NumChannels != 0)
 		{
@@ -172,7 +218,7 @@ float FNiagaraDataInterfaceProxySubmix::GetSampleRate() const
 	SubmixListeners.GetKeys(DeviceIds);
 	for (const Audio::FDeviceId DeviceId : DeviceIds)
 	{
-		float SampleRate = SubmixListeners[DeviceId].GetSampleRate();
+		float SampleRate = SubmixListeners[DeviceId]->GetSampleRate();
 
 		if (SampleRate > 0.f)
 		{
diff --git a/Engine/Plugins/FX/Niagara/Source/Niagara/Private/NiagaraDataInterfaceAudioOscilloscope.cpp b/Engine/Plugins/FX/Niagara/Source/Niagara/Private/NiagaraDataInterfaceAudioOscilloscope.cpp
--- a/Engine/Plugins/FX/Niagara/Source/Niagara/Private/NiagaraDataInterfaceAudioOscilloscope.cpp
+++ b/Engine/Plugins/FX/Niagara/Source/Niagara/Private/NiagaraDataInterfaceAudioOscilloscope.cpp
@@ -55,8 +55,8 @@ void FNiagaraDataInterfaceProxyOscilloscope::RegisterToAllAudioDevices()
 		DeviceManager->IterateOverAllDevices([&](Audio::FDeviceId DeviceId, FAudioDevice* InDevice)
 		{
 			check(!SubmixListeners.Contains(DeviceId));
-			FNiagaraSubmixListener& Listener = SubmixListeners.Emplace(DeviceId, FNiagaraSubmixListener(PatchMixer, Resolution * 2));
-			InDevice->RegisterSubmixBufferListener(&Listener, SubmixRegisteredTo);
+			SubmixListeners.Emplace(DeviceId, new FNiagaraSubmixListener(PatchMixer, Resolution * 2, DeviceId, SubmixRegisteredTo));
+			SubmixListeners[DeviceId]->RegisterToSubmix();
 		});
 	}
 }
@@ -69,11 +69,12 @@ void FNiagaraDataInterfaceProxyOscilloscope::UnregisterFromAllAudioDevices(USoun
 		DeviceManager->IterateOverAllDevices([&](Audio::FDeviceId DeviceId, FAudioDevice* InDevice)
 		{
 			check(SubmixListeners.Contains(DeviceId));
-
-			InDevice->UnregisterSubmixBufferListener(&SubmixListeners[DeviceId], Submix);
 			SubmixListeners.Remove(DeviceId);
 		});
 	}
+
+	ensure(SubmixListeners.Num() == 0);
+	SubmixListeners.Reset();
 }
 
 void FNiagaraDataInterfaceProxyOscilloscope::OnUpdateSubmix(USoundSubmix* Submix)
@@ -94,11 +95,8 @@ void FNiagaraDataInterfaceProxyOscilloscope::OnNewDeviceCreated(Audio::FDeviceId
 	if (bIsSubmixListenerRegistered)
 	{
 		check(!SubmixListeners.Contains(InID));
-		FAudioDeviceHandle DeviceHandle = FAudioDeviceManager::Get()->GetAudioDevice(InID);
-		check(DeviceHandle);
-
-		FNiagaraSubmixListener& Listener = SubmixListeners.Emplace(InID, FNiagaraSubmixListener(PatchMixer, Resolution * 2));
-		DeviceHandle->RegisterSubmixBufferListener(&Listener, SubmixRegisteredTo);
+		SubmixListeners.Emplace(InID, new FNiagaraSubmixListener(PatchMixer, Resolution * 2, InID, SubmixRegisteredTo));
+		SubmixListeners[InID]->RegisterToSubmix();
 	}
 }
 
@@ -108,10 +106,6 @@ void FNiagaraDataInterfaceProxyOscilloscope::OnDeviceDestroyed(Audio::FDeviceId
 	{
 		if (SubmixListeners.Contains(InID))
 		{
-			FAudioDeviceHandle DeviceHandle = FAudioDeviceManager::Get()->GetAudioDevice(InID);
-			check(DeviceHandle);
-
-			DeviceHandle->UnregisterSubmixBufferListener(&SubmixListeners[InID], SubmixRegisteredTo);
 			SubmixListeners.Remove(InID);
 		}
 	}
@@ -475,8 +469,8 @@ void FNiagaraDataInterfaceProxyOscilloscope::DownsampleAudioToBuffer()
 
 	for (auto& Listener : SubmixListeners)
 	{
-		NumChannelsInDownsampledBuffer = Listener.Value.GetNumChannels();
-		SourceSampleRate = Listener.Value.GetSampleRate();
+		NumChannelsInDownsampledBuffer = Listener.Value->GetNumChannels();
+		SourceSampleRate = Listener.Value->GetSampleRate();
 
 		if (NumChannelsInDownsampledBuffer != 0)
 		{
diff --git a/Engine/Source/Runtime/AudioMixer/Private/AudioMixerSubmix.cpp b/Engine/Source/Runtime/AudioMixer/Private/AudioMixerSubmix.cpp
--- a/Engine/Source/Runtime/AudioMixer/Private/AudioMixerSubmix.cpp
+++ b/Engine/Source/Runtime/AudioMixer/Private/AudioMixerSubmix.cpp
@@ -970,6 +970,7 @@ namespace Audio
 		// Now loop through any buffer listeners and feed the listeners the result of this audio callback
 		if(const USoundSubmix* SoundSubmix = Cast<const USoundSubmix>(OwningSubmixObject))
 		{
+			CSV_SCOPED_TIMING_STAT(Audio, SubmixBufferListeners);
 			double AudioClock = MixerDevice->GetAudioTime();
 			float SampleRate = MixerDevice->GetSampleRate();
 			FScopeLock Lock(&BufferListenerCriticalSection);
@@ -1408,6 +1409,7 @@ namespace Audio
 		check(BufferListener);
 		BufferListeners.Remove(BufferListener);
 	}
+
 	void FMixerSubmix::StartEnvelopeFollowing(int32 AttackTime, int32 ReleaseTime)
 	{
 		if (!bIsEnvelopeFollowing)
