Merging 11849741 from main - Niagara - Fix some crash issues in the stack UI.
+ Use the struct default when the niagara editor utilities doesn't generate a valid default value for a pin.  Fixes crashes when setting structs like matrix in the stack.
+ Defer the stack refresh after compile until the next frame because certain data operations can result in a system reset which checks the compile status and will cause the stack to refresh and finalize entries when they're still being used.
+ Fix the assignment node input renaming so that it broadcasts the input changed delegate which regenerates the inputs immediately instead of on the next refresh.  Also fix up the renaming code in the function input to order things correctly and update things safely.

diff --git a/Engine/Plugins/FX/Niagara/Source/NiagaraEditor/Private/NiagaraNodeAssignment.cpp b/Engine/Plugins/FX/Niagara/Source/NiagaraEditor/Private/NiagaraNodeAssignment.cpp
--- a/Engine/Plugins/FX/Niagara/Source/NiagaraEditor/Private/NiagaraNodeAssignment.cpp
+++ b/Engine/Plugins/FX/Niagara/Source/NiagaraEditor/Private/NiagaraNodeAssignment.cpp
@@ -71,6 +71,7 @@ bool UNiagaraNodeAssignment::RefreshFromExternalChanges()
 	GenerateScript();
 	ReallocatePins();
 	RefreshTitle();
+	OnInputsChangedDelegate.Broadcast();
 	return true;
 }
 
@@ -333,7 +334,6 @@ void UNiagaraNodeAssignment::AddParameter(FNiagaraVariable InVar, FString InDefa
 
 	RefreshFromExternalChanges();
 	MarkNodeRequiresSynchronization(__FUNCTION__, true);
-	OnInputsChangedDelegate.Broadcast();
 
 	RefreshTitle();
 }
@@ -363,7 +363,6 @@ void UNiagaraNodeAssignment::RemoveParameter(const FNiagaraVariable& InVar)
 
 	RefreshFromExternalChanges();
 	MarkNodeRequiresSynchronization(__FUNCTION__, true);
-	OnInputsChangedDelegate.Broadcast();
 
 	RefreshTitle();
 }
@@ -660,18 +659,24 @@ bool UNiagaraNodeAssignment::SetAssignmentTarget(int32 Idx, const FNiagaraVariab
 	return bRetValue;
 }
 
-bool UNiagaraNodeAssignment::SetAssignmentTargetName(int32 Idx, const FName& InName)
+bool UNiagaraNodeAssignment::RenameAssignmentTarget(FName OldName, FName NewName)
 {
-	check(Idx < AssignmentTargets.Num());
-	if (AssignmentTargets[Idx].GetName() != InName)
+	for (FNiagaraVariable& AssignmentTarget : AssignmentTargets)
 	{
-		AssignmentTargets[Idx].SetName(InName);
-		MarkNodeRequiresSynchronization(__FUNCTION__, true);
-		RefreshTitle();
-		return true;
+		if (AssignmentTarget.GetName() == OldName)
+		{
+			Modify();
+			if (FunctionScript != nullptr)
+			{
+				FunctionScript->Modify();
+				FunctionScript->GetSource()->Modify();
+			}
+
+			AssignmentTarget.SetName(NewName);
+			return true;
+		}
 	}
 	return false;
 }
 
-
 #undef LOCTEXT_NAMESPACE
diff --git a/Engine/Plugins/FX/Niagara/Source/NiagaraEditor/Private/ViewModels/Stack/NiagaraStackFunctionInput.cpp b/Engine/Plugins/FX/Niagara/Source/NiagaraEditor/Private/ViewModels/Stack/NiagaraStackFunctionInput.cpp
--- a/Engine/Plugins/FX/Niagara/Source/NiagaraEditor/Private/ViewModels/Stack/NiagaraStackFunctionInput.cpp
+++ b/Engine/Plugins/FX/Niagara/Source/NiagaraEditor/Private/ViewModels/Stack/NiagaraStackFunctionInput.cpp
@@ -1478,67 +1478,46 @@ bool UNiagaraStackFunctionInput::SupportsRename() const
 void UNiagaraStackFunctionInput::OnRenamed(FText NewNameText)
 {
 	FName NewName(*NewNameText.ToString());
-	if (OwningAssignmentNode.IsValid() && InputParameterHandlePath.Num() == 1 && InputParameterHandle.GetName() != NewName)
+	if (InputParameterHandle.GetName() != NewName && OwningAssignmentNode.IsValid())
 	{
-		FScopedTransaction ScopedTransaction(LOCTEXT("RenameInput", "Rename this function's input."));
-
-		FInputValues OldInputValues = InputValues;
-		UEdGraphPin* OriginalOverridePin = GetOverridePin();
-
-		// We'll be making changes, so go ahead and keep track of the override pointer if it exists.
-		if (OriginalOverridePin != nullptr)
-		{
-			OriginalOverridePin->GetOwningNode()->Modify();
-		}
+		bool bIsCurrentlyExpanded = GetStackEditorData().GetStackEntryIsExpanded(FNiagaraStackGraphUtilities::GenerateStackModuleEditorDataKey(*OwningAssignmentNode.Get()), false);
 
-		bool bIsCurrentlyExpanded = GetStackEditorData().GetStackEntryIsExpanded(FNiagaraStackGraphUtilities::GenerateStackModuleEditorDataKey(*OwningAssignmentNode), false);
-
-
-		int32 FoundIdx = OwningAssignmentNode->FindAssignmentTarget(InputParameterHandle.GetName());
-		check(FoundIdx != INDEX_NONE);
-		FNiagaraParameterHandle TargetHandle(OwningAssignmentNode->GetAssignmentTargetName(FoundIdx));
-
-		OwningAssignmentNode->Modify();
-		if (OwningAssignmentNode->FunctionScript != nullptr)
-		{
-			OwningAssignmentNode->FunctionScript->Modify();
-			OwningAssignmentNode->FunctionScript->GetSource()->Modify();
-		}
-		if (OwningAssignmentNode->SetAssignmentTargetName(FoundIdx, NewName))
-		{
-			OwningAssignmentNode->RefreshFromExternalChanges();
-		}
-
-		InputParameterHandle = FNiagaraParameterHandle(InputParameterHandle.GetNamespace(), NewName);
-		InputParameterHandlePath.Empty(1);
-		InputParameterHandlePath.Add(InputParameterHandle);
-		AliasedInputParameterHandle = FNiagaraParameterHandle::CreateAliasedModuleParameterHandle(InputParameterHandle, OwningAssignmentNode.Get());
-		DisplayName = FText::FromName(InputParameterHandle.GetName());
-
-		if (IsRapidIterationCandidate())
+		FScopedTransaction ScopedTransaction(LOCTEXT("RenameInput", "Rename this function's input."));
+		if (ensureMsgf(OwningAssignmentNode->RenameAssignmentTarget(InputParameterHandle.GetName(), NewName), TEXT("Failed to rename assignment node input.")))
 		{
-			FNiagaraVariable OldRapidIterationParameter = RapidIterationParameter;
-			RapidIterationParameter = CreateRapidIterationVariable(AliasedInputParameterHandle.GetParameterHandleString());
-
-			for (TWeakObjectPtr<UNiagaraScript> Script : AffectedScripts)
+			// Fixing up the stack graph and rapid iteration parameters must happen first so that when the stack is refreshed the UI is correct.
+			FNiagaraParameterHandle NewInputParameterHandle = FNiagaraParameterHandle(InputParameterHandle.GetNamespace(), NewName);
+			FNiagaraParameterHandle NewAliasedInputParameterHandle = FNiagaraParameterHandle::CreateAliasedModuleParameterHandle(NewInputParameterHandle, OwningAssignmentNode.Get());
+			UEdGraphPin* OverridePin = GetOverridePin();
+			if (OverridePin != nullptr)
 			{
-				Script->Modify();
-				Script->RapidIterationParameters.RenameParameter(OldRapidIterationParameter, *RapidIterationParameter.GetName().ToString());
+				// If there is an override pin then the only thing that needs to happen is that it's name needs to be updated so that the value it
+				// holds or is linked to stays intact.
+				OverridePin->Modify();
+				OverridePin->PinName = NewAliasedInputParameterHandle.GetParameterHandleString();
+			}
+			else if (IsRapidIterationCandidate())
+			{
+				// Otherwise if this is a valid rapid iteration parameter the values in the affected scripts need to be updated.
+				FNiagaraVariable NewRapidIterationParameter = CreateRapidIterationVariable(NewAliasedInputParameterHandle.GetParameterHandleString());
+				for (TWeakObjectPtr<UNiagaraScript> AffectedScript : AffectedScripts)
+				{
+					if (AffectedScript.IsValid())
+					{
+						AffectedScript->Modify();
+						AffectedScript->RapidIterationParameters.RenameParameter(RapidIterationParameter, *NewRapidIterationParameter.GetName().ToString());
+					}
+				}
 			}
 
-			UE_LOG(LogNiagaraEditor, Log, TEXT("Renaming %s to %s"), *OldRapidIterationParameter.GetName().ToString(), *RapidIterationParameter.GetName().ToString());
-		}
+			// Restore the expanded state with the new editor data key.
+			FString NewStackEditorDataKey = FNiagaraStackGraphUtilities::GenerateStackFunctionInputEditorDataKey(*OwningAssignmentNode.Get(), NewInputParameterHandle);
+			GetStackEditorData().SetStackEntryIsExpanded(NewStackEditorDataKey, bIsCurrentlyExpanded);
 
-		// Go ahead and have the override pin point to the new name instead of the old name..
-		if (OriginalOverridePin != nullptr)
-		{
-			OriginalOverridePin->PinName = AliasedInputParameterHandle.GetParameterHandleString();
+			// This refresh call must come last because it will finalize this input entry which would cause earlier fixup to fail.
+			OwningAssignmentNode->RefreshFromExternalChanges();
+			ensureMsgf(IsFinalized(), TEXT("Input not finalized when renamed."));
 		}
-		
-		StackEditorDataKey = FNiagaraStackGraphUtilities::GenerateStackFunctionInputEditorDataKey(*OwningFunctionCallNode.Get(), InputParameterHandle);
-		GetStackEditorData().SetStackEntryIsExpanded(FNiagaraStackGraphUtilities::GenerateStackModuleEditorDataKey(*OwningAssignmentNode), bIsCurrentlyExpanded);
-
-		CastChecked<UNiagaraGraph>(OwningAssignmentNode->GetGraph())->NotifyGraphNeedsRecompile();
 	}
 }
 
@@ -1925,11 +1904,15 @@ void UNiagaraStackFunctionInput::GetDefaultLocalValueFromDefaultPin(UEdGraphPin*
 	InInputValues.Mode = EValueMode::Local;
 	const UEdGraphSchema_Niagara* NiagaraSchema = GetDefault<UEdGraphSchema_Niagara>();
 	FNiagaraVariable LocalValueVariable = NiagaraSchema->PinToNiagaraVariable(DefaultPin);
+	InInputValues.LocalStruct = MakeShared<FStructOnScope>(InputType.GetStruct());
 	if (LocalValueVariable.IsDataAllocated())
 	{
-		InInputValues.LocalStruct = MakeShared<FStructOnScope>(InputType.GetStruct());
 		FMemory::Memcpy(InInputValues.LocalStruct->GetStructMemory(), LocalValueVariable.GetData(), InputType.GetSize());
 	}
+	else
+	{
+		InputType.GetStruct()->InitializeStruct(InInputValues.LocalStruct->GetStructMemory());
+	}
 }
 
 struct FLinkedHandleOrFunctionNode
diff --git a/Engine/Plugins/FX/Niagara/Source/NiagaraEditor/Private/ViewModels/Stack/NiagaraStackViewModel.cpp b/Engine/Plugins/FX/Niagara/Source/NiagaraEditor/Private/ViewModels/Stack/NiagaraStackViewModel.cpp
--- a/Engine/Plugins/FX/Niagara/Source/NiagaraEditor/Private/ViewModels/Stack/NiagaraStackViewModel.cpp
+++ b/Engine/Plugins/FX/Niagara/Source/NiagaraEditor/Private/ViewModels/Stack/NiagaraStackViewModel.cpp
@@ -351,15 +351,14 @@ void UNiagaraStackViewModel::GetPathForEntry(UNiagaraStackEntry* Entry, TArray<U
 
 void UNiagaraStackViewModel::OnSystemCompiled()
 {
-	RootEntry->RefreshChildren();
-	// when the entries are refreshed, make sure search results are still valid, the stack counts on them
-	OnSearchTextChanged(CurrentSearchText);
+	// Queue a refresh for the next tick because forcing a refresh now can cause entries to be finalized while they're still being used.
+	bRefreshPending = true;
 }
 
 void UNiagaraStackViewModel::OnEmitterCompiled()
 {
-	RootEntry->RefreshChildren();
-	OnSearchTextChanged(CurrentSearchText);
+	// Queue a refresh for the next tick because forcing a refresh now can cause entries to be finalized while they're still being used.
+	bRefreshPending = true;
 }
 
 void UNiagaraStackViewModel::EmitterParentRemoved()
diff --git a/Engine/Plugins/FX/Niagara/Source/NiagaraEditor/Public/NiagaraNodeAssignment.h b/Engine/Plugins/FX/Niagara/Source/NiagaraEditor/Public/NiagaraNodeAssignment.h
--- a/Engine/Plugins/FX/Niagara/Source/NiagaraEditor/Public/NiagaraNodeAssignment.h
+++ b/Engine/Plugins/FX/Niagara/Source/NiagaraEditor/Public/NiagaraNodeAssignment.h
@@ -26,8 +26,9 @@ public:
 	bool SetAssignmentTarget(int32 Idx, const FNiagaraVariable& InVar, const FString* InDefaultValue = nullptr);
 	int32 AddAssignmentTarget(const FNiagaraVariable& InVar, const FString* InDefaultValue = nullptr);
 
-	FName GetAssignmentTargetName(int32 Idx) const { return AssignmentTargets[Idx].GetName(); }
-	bool SetAssignmentTargetName(int32 Idx, const FName& InName);
+	/** This will rename an existing input but will not refresh the underlying script.  RefreshFromExternalChanges must be called if this function returns true. */
+	bool RenameAssignmentTarget(FName OldName, FName NewName);
+
 	const FNiagaraVariable& GetAssignmentTarget(int32 Idx) const { return AssignmentTargets[Idx]; }
 
 	const TArray<FNiagaraVariable>& GetAssignmentTargets() const { return AssignmentTargets; }
