Niagara - Fix the matrix type utilities to set the default for new matricies to the identity matrix rather then uninitialized data and also update the default value lookup in stack function inputs to use the utility function for the default when the pin to niagara variable doesn't set the data.

diff --git a/Engine/Plugins/FX/Niagara/Source/NiagaraEditor/Private/TypeEditorUtilities/NiagaraMatrixTypeEditorUtilities.cpp b/Engine/Plugins/FX/Niagara/Source/NiagaraEditor/Private/TypeEditorUtilities/NiagaraMatrixTypeEditorUtilities.cpp
--- a/Engine/Plugins/FX/Niagara/Source/NiagaraEditor/Private/TypeEditorUtilities/NiagaraMatrixTypeEditorUtilities.cpp
+++ b/Engine/Plugins/FX/Niagara/Source/NiagaraEditor/Private/TypeEditorUtilities/NiagaraMatrixTypeEditorUtilities.cpp
@@ -8,6 +8,6 @@ void FNiagaraEditorMatrixTypeUtilities::UpdateVariableWithDefaultValue(FNiagaraV
 	// This is implemented directly here to avoid alignment issues with the default initialization path for structs which
 	// uses pointer assignment rather than a mem-copy.
 	checkf(Variable.GetType().GetStruct() == FNiagaraTypeDefinition::GetMatrix4Struct(), TEXT("Struct type not supported."));
-	FMatrix DefaultMatrix;
+	FMatrix DefaultMatrix = FMatrix::Identity;
 	Variable.SetData((uint8*)&DefaultMatrix);
 }
\ No newline at end of file
diff --git a/Engine/Plugins/FX/Niagara/Source/NiagaraEditor/Private/ViewModels/Stack/NiagaraStackFunctionInput.cpp b/Engine/Plugins/FX/Niagara/Source/NiagaraEditor/Private/ViewModels/Stack/NiagaraStackFunctionInput.cpp
--- a/Engine/Plugins/FX/Niagara/Source/NiagaraEditor/Private/ViewModels/Stack/NiagaraStackFunctionInput.cpp
+++ b/Engine/Plugins/FX/Niagara/Source/NiagaraEditor/Private/ViewModels/Stack/NiagaraStackFunctionInput.cpp
@@ -1905,14 +1905,15 @@ void UNiagaraStackFunctionInput::GetDefaultLocalValueFromDefaultPin(UEdGraphPin*
 	const UEdGraphSchema_Niagara* NiagaraSchema = GetDefault<UEdGraphSchema_Niagara>();
 	FNiagaraVariable LocalValueVariable = NiagaraSchema->PinToNiagaraVariable(DefaultPin);
 	InInputValues.LocalStruct = MakeShared<FStructOnScope>(InputType.GetStruct());
-	if (LocalValueVariable.IsDataAllocated())
+	if (LocalValueVariable.IsDataAllocated() == false)
 	{
-		FMemory::Memcpy(InInputValues.LocalStruct->GetStructMemory(), LocalValueVariable.GetData(), InputType.GetSize());
+		FNiagaraEditorUtilities::ResetVariableToDefaultValue(LocalValueVariable);
 	}
-	else
+	if (ensureMsgf(LocalValueVariable.IsDataAllocated(), TEXT("Neither PinToNiagaraVariable or ResetVariableToDefaultValue generated a value.  Allocating with 0s.")) == false)
 	{
-		InputType.GetStruct()->InitializeStruct(InInputValues.LocalStruct->GetStructMemory());
+		LocalValueVariable.AllocateData();
 	}
+	FMemory::Memcpy(InInputValues.LocalStruct->GetStructMemory(), LocalValueVariable.GetData(), InputType.GetSize());
 }
 
 struct FLinkedHandleOrFunctionNode
