Niagara - Prevent ticking the system view model after it's been cleaned up.  Fixes crashes when closing the niagara system/emitter editor.

diff --git a/Engine/Plugins/FX/Niagara/Source/NiagaraEditor/Private/ViewModels/NiagaraSystemViewModel.cpp b/Engine/Plugins/FX/Niagara/Source/NiagaraEditor/Private/ViewModels/NiagaraSystemViewModel.cpp
--- a/Engine/Plugins/FX/Niagara/Source/NiagaraEditor/Private/ViewModels/NiagaraSystemViewModel.cpp
+++ b/Engine/Plugins/FX/Niagara/Source/NiagaraEditor/Private/ViewModels/NiagaraSystemViewModel.cpp
@@ -194,6 +194,8 @@ void FNiagaraSystemViewModel::Cleanup()
 		ScriptScratchPadViewModel->Finalize();
 		ScriptScratchPadViewModel = nullptr;
 	}
+
+	System = nullptr;
 }
 
 FNiagaraSystemViewModel::~FNiagaraSystemViewModel()
@@ -510,6 +512,13 @@ void FNiagaraSystemViewModel::PostUndo(bool bSuccess)
 
 void FNiagaraSystemViewModel::Tick(float DeltaTime)
 {
+	if (System == nullptr)
+	{
+		// If the system pointer is no longer valid, this system view model has been cleaned up and is invalid
+		// and is about to be destroyed so don't run tick logic.
+		return;
+	}
+
 	if (bCompilePendingCompletion && GetSystem().HasOutstandingCompilationRequests() == false)
 	{
 		bCompilePendingCompletion = false;
@@ -1891,7 +1900,7 @@ void FNiagaraSystemViewModel::AddSystemEventHandlers()
 
 void FNiagaraSystemViewModel::RemoveSystemEventHandlers()
 {
-	if (GetSystem().IsValid())
+	if (System != nullptr && System->IsValid())
 	{
 		TArray<UNiagaraScript*> Scripts;
 		Scripts.Add(GetSystem().GetSystemSpawnScript());
