Don't activate Niagara FX in dedicated server worlds.
While generaly a good idea anyway this avoids some issues with a recent chagne to not tick the world manger in dedicated server worlds.

diff --git a/Engine/Plugins/FX/Niagara/Source/Niagara/Private/NiagaraComponent.cpp b/Engine/Plugins/FX/Niagara/Source/Niagara/Private/NiagaraComponent.cpp
--- a/Engine/Plugins/FX/Niagara/Source/Niagara/Private/NiagaraComponent.cpp
+++ b/Engine/Plugins/FX/Niagara/Source/Niagara/Private/NiagaraComponent.cpp
@@ -839,8 +839,9 @@ void UNiagaraComponent::ActivateInternal(bool bReset /* = false */, bool bIsScal
 		return;
 	}
 	
+	UWorld* World = GetWorld();
 	// If the particle system can't ever render (ie on dedicated server or in a commandlet) than do not activate...
-	if (!FApp::CanEverRender())
+	if (!FApp::CanEverRender() || !World || World->IsNetMode(NM_DedicatedServer))
 	{
 		return;
 	}
