Small change to harden the SkelMeshBoneSampling with Niagara
-SkinOnTheFly path no longer relies on accessing the Mesh (Mesh->RefSkeleton) for anything, instead using cached data for bone counts
-fixes off by one error if we were to try to reference the last bone (it would have resulted in a 0 vector)

#localization none
#tests automated EngineTests, validating bone sampling still works as expected

diff --git a/Engine/Plugins/FX/Niagara/Source/Niagara/Classes/NDISkeletalMeshCommon.h b/Engine/Plugins/FX/Niagara/Source/Niagara/Classes/NDISkeletalMeshCommon.h
--- a/Engine/Plugins/FX/Niagara/Source/Niagara/Classes/NDISkeletalMeshCommon.h
+++ b/Engine/Plugins/FX/Niagara/Source/Niagara/Classes/NDISkeletalMeshCommon.h
@@ -63,8 +63,12 @@ struct FSkinnedPositionAccessorHelper<TIntegralConstant<ENDISkeletalMesh_Skinnin
 {
 	FORCEINLINE int32 GetBoneCount(FSkeletalMeshAccessorHelper& Accessor, bool RequiresPrevious) const
 	{
-		const FReferenceSkeleton& RefSkel = Accessor.Mesh->RefSkeleton;
-		return RefSkel.GetNum();
+		if (const USkeletalMesh* Mesh = Accessor.Mesh)
+		{
+			return Mesh->RefSkeleton.GetNum();
+		}
+
+		return 0;
 	}
 
 	FORCEINLINE void GetTriangleIndices(FSkeletalMeshAccessorHelper& Accessor, int32 Tri, int32& Idx0, int32& Idx1, int32& Idx2)
@@ -126,7 +130,12 @@ struct FSkinnedPositionAccessorHelper<TIntegralConstant<ENDISkeletalMesh_Skinnin
 {
 	FORCEINLINE int32 GetBoneCount(FSkeletalMeshAccessorHelper& Accessor, bool RequiresPrevious) const
 	{
-		return Accessor.SkinningData->GetBoneCount(RequiresPrevious);
+		if (const FSkeletalMeshSkinningData* SkinningData = Accessor.SkinningData)
+		{
+			return SkinningData->GetBoneCount(RequiresPrevious);
+		}
+
+		return 0;
 	}
 
 	FORCEINLINE void GetTriangleIndices(FSkeletalMeshAccessorHelper& Accessor, int32 Tri, int32& Idx0, int32& Idx1, int32& Idx2)
@@ -188,7 +197,12 @@ struct FSkinnedPositionAccessorHelper<TIntegralConstant<ENDISkeletalMesh_Skinnin
 {
 	FORCEINLINE int32 GetBoneCount(FSkeletalMeshAccessorHelper& Accessor, bool RequiresPrevious) const
 	{
-		return Accessor.SkinningData->GetBoneCount(RequiresPrevious);
+		if (const FSkeletalMeshSkinningData* SkinningData = Accessor.SkinningData)
+		{
+			return SkinningData->GetBoneCount(RequiresPrevious);
+		}
+
+		return 0;
 	}
 
 	FORCEINLINE void GetTriangleIndices(FSkeletalMeshAccessorHelper& Accessor, int32 Tri, int32& Idx0, int32& Idx1, int32& Idx2)
diff --git a/Engine/Plugins/FX/Niagara/Source/Niagara/Private/NDISkeletalMesh_BoneSampling.cpp b/Engine/Plugins/FX/Niagara/Source/Niagara/Private/NDISkeletalMesh_BoneSampling.cpp
--- a/Engine/Plugins/FX/Niagara/Source/Niagara/Private/NDISkeletalMesh_BoneSampling.cpp
+++ b/Engine/Plugins/FX/Niagara/Source/Niagara/Private/NDISkeletalMesh_BoneSampling.cpp
@@ -472,8 +472,8 @@ void UNiagaraDataInterfaceSkeletalMesh::GetSkinnedBoneData(FVectorVMContext& Con
 		const float Interp = bInterpolated::Value ? InterpParam.GetAndAdvance() : 1.0f;
 
 		// Determine bone or socket
-		int32 RawBone = BoneParam.GetAndAdvance();
-		int32 Bone = FMath::Clamp(RawBone, 0, BoneAndSocketCount - 1);
+		const int32 RawBone = BoneParam.GetAndAdvance();
+		const int32 Bone = FMath::Clamp(RawBone, 0, BoneAndSocketCount - 1);
 		const bool bIsSocket = Bone >= BoneCount;
 		const int32 Socket = Bone - BoneCount;
 
