Collection of changes to help reduce the amount of non-deterministic cooking resulting from Niagara

Fix for redundant niagara GPUComputeShaders being saved into cooked builds because BeginCacheForCookedPlatformData is called multiple (non-deterministic) times during a cook
Allocate space in ParameterStore for missing parameters to fix potential memory overwrite
Initialize variables to fix deterministic cooking

diff --git a/Engine/Plugins/FX/Niagara/Source/Niagara/Public/NiagaraTypes.h b/Engine/Plugins/FX/Niagara/Source/Niagara/Public/NiagaraTypes.h
--- a/Engine/Plugins/FX/Niagara/Source/Niagara/Public/NiagaraTypes.h
+++ b/Engine/Plugins/FX/Niagara/Source/Niagara/Public/NiagaraTypes.h
@@ -67,7 +67,7 @@ struct FNiagaraBool
 
 private:
 	UPROPERTY(EditAnywhere, Category = Parameters)// Must be either FNiagaraBool::True or FNiagaraBool::False.
-	int32 Value;
+	int32 Value = FNiagaraBool::False;
 };
 
 USTRUCT()
@@ -89,10 +89,10 @@ struct FNiagaraTestStructInner
 	GENERATED_USTRUCT_BODY()
 
 	UPROPERTY(EditAnywhere, Category = TestStruct)
-	FVector InnerVector1;
+	FVector InnerVector1 = FVector::ZeroVector;
 
 	UPROPERTY(EditAnywhere, Category = TestStruct)
-	FVector InnerVector2;
+	FVector InnerVector2 = FVector::ZeroVector;
 };
 
 USTRUCT()
@@ -101,10 +101,10 @@ struct FNiagaraTestStruct
 	GENERATED_USTRUCT_BODY()
 
 	UPROPERTY(EditAnywhere, Category = TestStruct)
-	FVector Vector1;
+	FVector Vector1 = FVector::ZeroVector;
 
 	UPROPERTY(EditAnywhere, Category = TestStruct)
-	FVector Vector2;
+	FVector Vector2 = FVector::ZeroVector;
 
 	UPROPERTY(EditAnywhere, Category = TestStruct)
 	FNiagaraTestStructInner InnerStruct1;
@@ -119,16 +119,16 @@ struct FNiagaraMatrix
 	GENERATED_USTRUCT_BODY()
 
 	UPROPERTY(EditAnywhere, Category=NiagaraMatrix)
-	FVector4 Row0;
+	FVector4 Row0 = FVector4(EForceInit::ForceInitToZero);
 
 	UPROPERTY(EditAnywhere, Category = NiagaraMatrix)
-	FVector4 Row1;
+	FVector4 Row1 = FVector4(EForceInit::ForceInitToZero);
 
 	UPROPERTY(EditAnywhere, Category = NiagaraMatrix)
-	FVector4 Row2;
+	FVector4 Row2 = FVector4(EForceInit::ForceInitToZero);
 
 	UPROPERTY(EditAnywhere, Category = NiagaraMatrix)
-	FVector4 Row3;
+	FVector4 Row3 = FVector4(EForceInit::ForceInitToZero);
 };
 
 /** Data controlling the spawning of particles */
@@ -137,29 +137,22 @@ struct FNiagaraSpawnInfo
 {
 	GENERATED_USTRUCT_BODY();
 	
-	FNiagaraSpawnInfo()
-		: Count(0)
-		, InterpStartDt(0.0f)
-		, IntervalDt(1.0f)
-		, SpawnGroup(0)
-	{}
-
 	/** How many particles to spawn. */
 	UPROPERTY(EditAnywhere, Category = SpawnInfo)
-	int32 Count;
+	int32 Count = 0;
 	/** The sub frame delta time at which to spawn the first particle. */
 	UPROPERTY(EditAnywhere, Category = SpawnInfo)
-	float InterpStartDt;
+	float InterpStartDt = 0.0f;
 	/** The sub frame delta time between each particle. */
 	UPROPERTY(EditAnywhere, Category = SpawnInfo)
-	float IntervalDt;
+	float IntervalDt = 1.0f;
 	/**
 	 * An integer used to identify this spawn info.
 	 * Typically this is unused.
 	 * An example usage is when using multiple spawn modules to spawn from multiple discreet locations.
 	 */
 	UPROPERTY(EditAnywhere, Category = SpawnInfo)
-	int32 SpawnGroup;
+	int32 SpawnGroup = 0;
 };
 
 USTRUCT(Blueprintable, meta = (DisplayName = "Niagara ID"))
@@ -172,14 +165,14 @@ struct FNiagaraID
 	Is always unique among currently living particles but will be reused after the particle dies.
 	*/
 	UPROPERTY(EditAnywhere, Category = ID)
-	int32 Index;
+	int32 Index = 0;
 
 	/** 
 	A unique tag for when this ID was acquired. 
 	Allows us to differentiate between particles when one dies and another reuses it's Index.
 	*/
 	UPROPERTY(EditAnywhere, Category = ID)
-	int32 AcquireTag;
+	int32 AcquireTag = 0;
 
 	bool operator==(const FNiagaraID& Other)const { return Index == Other.Index && AcquireTag == Other.AcquireTag; }
 	bool operator!=(const FNiagaraID& Other)const { return !(*this == Other); }
diff --git a/Engine/Plugins/FX/Niagara/Source/NiagaraEditor/Private/TypeEditorUtilities/NiagaraMatrixTypeEditorUtilities.cpp b/Engine/Plugins/FX/Niagara/Source/NiagaraEditor/Private/TypeEditorUtilities/NiagaraMatrixTypeEditorUtilities.cpp
--- a/Engine/Plugins/FX/Niagara/Source/NiagaraEditor/Private/TypeEditorUtilities/NiagaraMatrixTypeEditorUtilities.cpp
+++ b/Engine/Plugins/FX/Niagara/Source/NiagaraEditor/Private/TypeEditorUtilities/NiagaraMatrixTypeEditorUtilities.cpp
@@ -8,6 +8,6 @@ void FNiagaraEditorMatrixTypeUtilities::UpdateVariableWithDefaultValue(FNiagaraV
 	// This is implemented directly here to avoid alignment issues with the default initialization path for structs which
 	// uses pointer assignment rather than a mem-copy.
 	checkf(Variable.GetType().GetStruct() == FNiagaraTypeDefinition::GetMatrix4Struct(), TEXT("Struct type not supported."));
-	FMatrix DefaultMatrix = FMatrix::Identity;
+	FMatrix DefaultMatrix(FMatrix::Identity);
 	Variable.SetData((uint8*)&DefaultMatrix);
 }
\ No newline at end of file
