Fix recompileshaders running on target with COTF server

diff --git a/Engine/Plugins/FX/Niagara/Source/NiagaraShader/Private/NiagaraShader.cpp b/Engine/Plugins/FX/Niagara/Source/NiagaraShader/Private/NiagaraShader.cpp
--- a/Engine/Plugins/FX/Niagara/Source/NiagaraShader/Private/NiagaraShader.cpp
+++ b/Engine/Plugins/FX/Niagara/Source/NiagaraShader/Private/NiagaraShader.cpp
@@ -652,191 +652,6 @@ void FNiagaraShaderMap::RestoreShadersFromMemory(const TArray<uint8>& ShaderData
 	//RegisterSerializedShaders(false);
 }
 
-void FNiagaraShaderMap::SaveForRemoteRecompile(FArchive& Ar, const TMap<FString, TArray<TRefCountPtr<FNiagaraShaderMap> > >& CompiledShaderMaps, const TArray<FShaderResourceId>& ClientResourceIds)
-{
-	UE_LOG(LogShaders, Display, TEXT("Niagara shader map looking for unique resources, %d were on client"), ClientResourceIds.Num());
-	check(false);
-#if 0
-	// first, we look for the unique shader resources
-	TArray<FShaderResource*> UniqueResources;
-	int32 NumSkippedResources = 0;
-
-	for (TMap<FString, TArray<TRefCountPtr<FNiagaraShaderMap> > >::TConstIterator It(CompiledShaderMaps); It; ++It)
-	{
-		const TArray<TRefCountPtr<FNiagaraShaderMap> >& ShaderMapArray = It.Value();
-
-		for (int32 Index = 0; Index < ShaderMapArray.Num(); Index++)
-		{
-			FNiagaraShaderMap* ShaderMap = ShaderMapArray[Index];
-
-			if (ShaderMap)
-			{
-				// get all shaders in the shader map
-				TMap<FShaderId, TShaderRef<FShader>> ShaderList;
-				ShaderMap->GetShaderList(ShaderList);
-
-				// get the resources from the shaders
-				for (auto& KeyValue : ShaderList)
-				{
-					const TShaderRef<FShader> Shader = KeyValue.Value;
-					FShaderResourceId ShaderId = Shader.GetResourceId();
-
-					// skip this shader if the Id was already on the client (ie, it didn't change)
-					if (ClientResourceIds.Contains(ShaderId) == false)
-					{
-						// lookup the resource by ID
-						// add it if it's unique
-						UniqueResources.AddUnique(Shader.GetResourceChecked());
-					}
-					else
-					{
-						NumSkippedResources++;
-					}
-				}
-			}
-		}
-	}
-
-	UE_LOG(LogShaders, Display, TEXT("Sending %d new Niagara shader resources, skipped %d existing"), UniqueResources.Num(), NumSkippedResources);
-
-	// now serialize them
-	int32 NumUniqueResources = UniqueResources.Num();
-	Ar << NumUniqueResources;
-
-	for (int32 Index = 0; Index < NumUniqueResources; Index++)
-	{
-		UniqueResources[Index]->SaveCodeToArchive(Ar);
-	}
-
-	// now we serialize a map (for each script)
-	int32 MapSize = CompiledShaderMaps.Num();
-	Ar << MapSize;
-
-	for (TMap<FString, TArray<TRefCountPtr<FNiagaraShaderMap> > >::TConstIterator It(CompiledShaderMaps); It; ++It)
-	{
-		const TArray<TRefCountPtr<FNiagaraShaderMap> >& ShaderMapArray = It.Value();
-
-		FString ScriptName = It.Key();
-		Ar << ScriptName;
-
-		int32 NumShaderMaps = ShaderMapArray.Num();
-		Ar << NumShaderMaps;
-
-		for (int32 Index = 0; Index < ShaderMapArray.Num(); Index++)
-		{
-			FNiagaraShaderMap* ShaderMap = ShaderMapArray[Index];
-
-			if (ShaderMap && NumUniqueResources > 0)
-			{
-				uint8 bIsValid = 1;
-				Ar << bIsValid;
-				ShaderMap->Serialize(Ar, false);
-			}
-			else
-			{
-				uint8 bIsValid = 0;
-				Ar << bIsValid;
-			}
-		}
-	}
-#endif
-}
-
-void FNiagaraShaderMap::LoadForRemoteRecompile(FArchive& Ar, EShaderPlatform ShaderPlatform, const TArray<FString>& ScriptsForShaderMaps)
-{
-	check(false);	// TODO: implement
-	/*
-	int32 NumResources;
-	Ar << NumResources;
-
-	// KeepAliveReferences keeps resources alive until we are finished serializing in this function
-	TArray<TRefCountPtr<FShaderResource> > KeepAliveReferences;
-
-	// load and register the resources
-	for (int32 Index = 0; Index < NumResources; Index++)
-	{
-		// Load the inlined shader resource
-		FShaderResource* Resource = new FShaderResource();
-		Resource->Serialize(Ar);
-
-		// if this Id is already in memory, that means that this is a repeated resource and so we skip it
-		if (FShaderResource::FindShaderResourceById(Resource->GetId()) != NULL)
-		{
-			delete Resource;
-		}
-		// otherwise, it's a new resource, so we register it for the maps to find below
-		else
-		{
-			Resource->Register();
-
-			// Keep this guy alive until we finish serializing all the FShaders in
-			// The FShaders which are discarded may cause these resources to be discarded 
-			KeepAliveReferences.Add( Resource );
-		}
-	}
-
-	int32 MapSize;
-	Ar << MapSize;
-
-	for (int32 MaterialIndex = 0; MaterialIndex < MapSize; MaterialIndex++)
-	{
-		FString MaterialName;
-		Ar << MaterialName;
-
-		UMaterialInterface* MatchingMaterial = FindObjectChecked<UMaterialInterface>(NULL, *MaterialName);
-
-		int32 NumShaderMaps = 0;
-		Ar << NumShaderMaps;
-
-		TArray<TRefCountPtr<FNiagaraShaderMap> > LoadedShaderMaps;
-
-		for (int32 ShaderMapIndex = 0; ShaderMapIndex < NumShaderMaps; ShaderMapIndex++)
-		{
-			uint8 bIsValid = 0;
-			Ar << bIsValid;
-
-			if (bIsValid)
-			{
-				FNiagaraShaderMap* ShaderMap = new FNiagaraShaderMap;
-
-				// serialize the id and the material shader map
-				ShaderMap->Serialize(Ar, false);
-
-				// Register in the global map
-				ShaderMap->Register(ShaderPlatform);
-
-				LoadedShaderMaps.Add(ShaderMap);
-			}
-		}
-
-		// Assign in two passes: first pass for shader maps with unspecified quality levels,
-		// Second pass for shader maps with a specific quality level
-			for (int32 ShaderMapIndex = 0; ShaderMapIndex < LoadedShaderMaps.Num(); ShaderMapIndex++)
-			{
-				FNiagaraShaderMap* LoadedShaderMap = LoadedShaderMaps[ShaderMapIndex];
-
-				if (LoadedShaderMap->GetShaderPlatform() == ShaderPlatform 
-					&& LoadedShaderMap->GetShaderMapId().FeatureLevel == GetMaxSupportedFeatureLevel(ShaderPlatform))
-				{
-					FMaterialResource* MaterialResource = MatchingMaterial->GetMaterialResource(GetMaxSupportedFeatureLevel(ShaderPlatform), (EMaterialQualityLevel::Type)QualityLevelIndex);
-					MaterialResource->SetGameThreadShaderMap(LoadedShaderMap);
-					MaterialResource->RegisterInlineShaderMap();
-
-					ENQUEUE_RENDER_COMMAND(FSetShaderMapOnMaterialResources)(
-						[MaterialResource,LoadedShaderMap](FRHICommandListImmediate& RHICmdList)
-						{
-							MaterialResource->SetRenderingThreadShaderMap(LoadedShaderMap);
-						});
-				}
-			}
-		}
-	}
-	*/
-}
-
-
-
-
 /**
 * Compiles the shaders for a script and caches them in this shader map.
 * @param script - The script to compile shaders for.
diff --git a/Engine/Plugins/FX/Niagara/Source/NiagaraShader/Public/NiagaraShared.h b/Engine/Plugins/FX/Niagara/Source/NiagaraShader/Public/NiagaraShared.h
--- a/Engine/Plugins/FX/Niagara/Source/NiagaraShader/Public/NiagaraShared.h
+++ b/Engine/Plugins/FX/Niagara/Source/NiagaraShader/Public/NiagaraShared.h
@@ -472,10 +472,6 @@ public:
 	/** Recreates FShaders from the passed in memory, handling shader key changes. */
 	void RestoreShadersFromMemory(const TArray<uint8>& ShaderData);
 
-	/** Serializes a shader map to an archive (used with recompiling shaders for a remote console) */
-	NIAGARASHADER_API  static void SaveForRemoteRecompile(FArchive& Ar, const TMap<FString, TArray<TRefCountPtr<FNiagaraShaderMap> > >& CompiledShaderMaps, const TArray<FShaderResourceId>& ClientResourceIds);
-	NIAGARASHADER_API  static void LoadForRemoteRecompile(FArchive& Ar, EShaderPlatform ShaderPlatform, const TArray<FString>& ScriptsForShaderMaps);
-
 	/** Returns the maximum number of texture samplers used by any shader in this shader map. */
 	uint32 GetMaxTextureSamplers() const;
 
diff --git a/Engine/Source/Editor/UnrealEd/Private/CookOnTheFlyServer.cpp b/Engine/Source/Editor/UnrealEd/Private/CookOnTheFlyServer.cpp
--- a/Engine/Source/Editor/UnrealEd/Private/CookOnTheFlyServer.cpp
+++ b/Engine/Source/Editor/UnrealEd/Private/CookOnTheFlyServer.cpp
@@ -6501,7 +6501,6 @@ void UCookOnTheFlyServer::SaveGlobalShaderMapFiles(const TArrayView<const ITarge
 			RecompileData.ShaderPlatform == -1 ? SP_NumPlatforms : (EShaderPlatform)RecompileData.ShaderPlatform, //-V547
 			OutputDir, 
 			RecompileData.MaterialsToLoad, 
-			RecompileData.SerializedShaderResources, 
 			RecompileData.MeshMaterialMaps, 
 			RecompileData.ModifiedFiles);
 	}
@@ -8202,7 +8201,6 @@ void UCookOnTheFlyServer::HandleNetworkFileServerRecompileShaders(const FShaderR
 		RecompileData.ShaderPlatform == -1 ? SP_NumPlatforms : (EShaderPlatform)RecompileData.ShaderPlatform,
 		OutputDir, 
 		RecompileData.MaterialsToLoad, 
-		RecompileData.SerializedShaderResources, 
 		RecompileData.MeshMaterialMaps, 
 		RecompileData.ModifiedFiles,
 		RecompileData.bCompileChangedShaders);
diff --git a/Engine/Source/Runtime/Core/Private/Misc/PreLoadFile.cpp b/Engine/Source/Runtime/Core/Private/Misc/PreLoadFile.cpp
--- a/Engine/Source/Runtime/Core/Private/Misc/PreLoadFile.cpp
+++ b/Engine/Source/Runtime/Core/Private/Misc/PreLoadFile.cpp
@@ -97,8 +97,13 @@ void FPreLoadFile::KickOffRead()
 
 void* FPreLoadFile::TakeOwnershipOfLoadedData(int64* OutFileSize)
 {
-	check(CompletionEvent != nullptr);
+	if (!CompletionEvent)
+	{
+		// may need to read again, if (re-)requesting data after the initial boot sequence
+		KickOffRead();
+	}
 
+	check(CompletionEvent);
 	if (CompletionEvent->Wait(0) == false)
 	{
 		printf("PreLoadFile %s wasn't ready...\n", TCHAR_TO_ANSI(*Path));
@@ -107,6 +112,7 @@ void* FPreLoadFile::TakeOwnershipOfLoadedData(int64* OutFileSize)
 		CompletionEvent->Wait();
 	}
 	FPlatformProcess::ReturnSynchEventToPool(CompletionEvent);
+	CompletionEvent = nullptr;
 
 	// return the size of Data if needed
 	if (OutFileSize)
diff --git a/Engine/Source/Runtime/Engine/Private/Materials/MaterialShader.cpp b/Engine/Source/Runtime/Engine/Private/Materials/MaterialShader.cpp
--- a/Engine/Source/Runtime/Engine/Private/Materials/MaterialShader.cpp
+++ b/Engine/Source/Runtime/Engine/Private/Materials/MaterialShader.cpp
@@ -1152,61 +1152,8 @@ void FMaterialShaderMap::RestoreShadersFromMemory(const TArray<uint8>& ShaderDat
 #endif
 }
 
-void FMaterialShaderMap::SaveForRemoteRecompile(FArchive& Ar, const TMap<FString, TArray<TRefCountPtr<FMaterialShaderMap> > >& CompiledShaderMaps, const TArray<FShaderResourceId>& ClientResourceIds)
+void FMaterialShaderMap::SaveForRemoteRecompile(FArchive& Ar, const TMap<FString, TArray<TRefCountPtr<FMaterialShaderMap> > >& CompiledShaderMaps)
 {
-	UE_LOG(LogMaterial, Display, TEXT("Looking for unique resources, %d were on client"), ClientResourceIds.Num());	
-	check(false);
-#if 0
-	// first, we look for the unique shader resources
-	TArray<TRefCountPtr<FShaderResource>> UniqueResources;
-	int32 NumSkippedResources = 0;
-
-	for (TMap<FString, TArray<TRefCountPtr<FMaterialShaderMap> > >::TConstIterator It(CompiledShaderMaps); It; ++It)
-	{
-		const TArray<TRefCountPtr<FMaterialShaderMap> >& ShaderMapArray = It.Value();
-
-		for (int32 Index = 0; Index < ShaderMapArray.Num(); Index++)
-		{
-			FMaterialShaderMap* ShaderMap = ShaderMapArray[Index];
-
-			if (ShaderMap)
-			{
-				// get all shaders in the shader map
-				TMap<FShaderId, TShaderRef<FShader>> ShaderList;
-				ShaderMap->GetShaderList(ShaderList);
-
-				// get the resources from the shaders
-				for (auto& KeyValue : ShaderList)
-				{
-					const TShaderRef<FShader>& Shader = KeyValue.Value;
-					FShaderResourceId ShaderId = Shader.GetResourceId();
-
-					// skip this shader if the Id was already on the client (ie, it didn't change)
-					if (ClientResourceIds.Contains(ShaderId) == false)
-					{
-						// lookup the resource by ID and add it if it's unique
-						UniqueResources.AddUnique(Shader.GetResourceChecked());
-					}
-					else
-					{
-						NumSkippedResources++;
-					}
-				}
-			}
-		}
-	}
-
-	UE_LOG(LogMaterial, Display, TEXT("Sending %d new shader resources, skipped %d existing"), UniqueResources.Num(), NumSkippedResources);
-
-	// now serialize them
-	int32 NumUniqueResources = UniqueResources.Num();
-	Ar << NumUniqueResources;
-
-	for (int32 Index = 0; Index < NumUniqueResources; Index++)
-	{
-		UniqueResources[Index]->SaveCodeToArchive(Ar);
-	}
-
 	// now we serialize a map (for each material), but without inline the resources, since they are above
 	int32 MapSize = CompiledShaderMaps.Num();
 	Ar << MapSize;
@@ -1224,8 +1171,7 @@ void FMaterialShaderMap::SaveForRemoteRecompile(FArchive& Ar, const TMap<FString
 		for (int32 Index = 0; Index < ShaderMapArray.Num(); Index++)
 		{
 			FMaterialShaderMap* ShaderMap = ShaderMapArray[Index];
-
-			if (ShaderMap && NumUniqueResources > 0)
+			if (ShaderMap)
 			{
 				uint8 bIsValid = 1;
 				Ar << bIsValid;
@@ -1238,27 +1184,10 @@ void FMaterialShaderMap::SaveForRemoteRecompile(FArchive& Ar, const TMap<FString
 			}
 		}
 	}
-#endif
 }
 
 void FMaterialShaderMap::LoadForRemoteRecompile(FArchive& Ar, EShaderPlatform ShaderPlatform, const TArray<FString>& MaterialsForShaderMaps)
 {
-	check(false);
-#if 0// WITH_EDITOR
-	int32 NumResources;
-	Ar << NumResources;
-
-	// KeepAliveReferences keeps resources alive until we are finished serializing in this function
-	TArray<TRefCountPtr<FShaderResource>> KeepAliveReferences;
-	KeepAliveReferences.SetNum(NumResources);
-
-	// load and register the resources
-	for (int32 Index = 0; Index < NumResources; Index++)
-	{
-		// Load the inlined shader resource
-		KeepAliveReferences[Index] = FShaderResource_InlineCode::FindOrLoad(Ar);
-	}
-
 	int32 MapSize;
 	Ar << MapSize;
 
@@ -1323,7 +1252,6 @@ void FMaterialShaderMap::LoadForRemoteRecompile(FArchive& Ar, EShaderPlatform Sh
 			}
 		}
 	}
-#endif // WITH_EDITOR
 }
 
 void FMaterialShaderMap::FinalizeContent()
@@ -2300,9 +2228,13 @@ uint32 FMaterialShaderMap::GetShaderNum() const
 void FMaterialShaderMap::Register(EShaderPlatform InShaderPlatform)
 {
 	extern int32 GCreateShadersOnLoad;
-	if (GCreateShadersOnLoad && GetShaderPlatform() == InShaderPlatform && GetResource())
+	if (GCreateShadersOnLoad && GetShaderPlatform() == InShaderPlatform)
 	{
-		// TODO
+		FShaderMapResource* ShaderResource = GetResource();
+		if (ShaderResource)
+		{
+			ShaderResource->BeginCreateAllShaders();
+		}
 	}
 
 	if (!bRegistered)
diff --git a/Engine/Source/Runtime/Engine/Private/Materials/MaterialShared.cpp b/Engine/Source/Runtime/Engine/Private/Materials/MaterialShared.cpp
--- a/Engine/Source/Runtime/Engine/Private/Materials/MaterialShared.cpp
+++ b/Engine/Source/Runtime/Engine/Private/Materials/MaterialShared.cpp
@@ -216,23 +216,6 @@ bool ReloadMaterialResource(
 }
 #endif
 
-static inline SIZE_T AddShaderSize(const TShaderRef<FShader>& Shader, TSet<FShaderResourceId>& UniqueShaderResourceIds)
-{
-	SIZE_T ResourceSize = 0;
-	// @todo fnhack
-//	check(false);
-#if 0
-	FShaderResourceId ResourceId = Shader.GetResourceId();
-	bool bCountedResource = false;
-	UniqueShaderResourceIds.Add(ResourceId, &bCountedResource);
-	if (!bCountedResource)
-	{
-		ResourceSize += Shader.GetResource()->GetSizeBytes();
-	}
-#endif
-	return ResourceSize;
-}
-
 int32 FMaterialCompiler::Errorf(const TCHAR* Format,...)
 {
 	TCHAR	ErrorText[2048];
@@ -1428,10 +1411,6 @@ void FMaterialResource::NotifyCompilationFinished()
 void FMaterialResource::GetResourceSizeEx(FResourceSizeEx& CumulativeResourceSize)
 {
 	TSet<const FMaterialShaderMap*> UniqueShaderMaps;
-	TMap<FShaderId, TShaderRef<FShader>> UniqueShaders;
-	TArray<FShaderPipelineRef> ShaderPipelines;
-	TSet<FShaderResourceId> UniqueShaderResourceIds;
-
 	UniqueShaderMaps.Add(GetGameThreadShaderMap());
 
 	for (TSet<const FMaterialShaderMap*>::TConstIterator It(UniqueShaderMaps); It; ++It)
@@ -1441,37 +1420,15 @@ void FMaterialResource::GetResourceSizeEx(FResourceSizeEx& CumulativeResourceSiz
 		{
 			CumulativeResourceSize.AddDedicatedSystemMemoryBytes(MaterialShaderMap->GetFrozenContentSize());
 
-			// Shaders are shared, so only count them in total mode
-			if (CumulativeResourceSize.GetResourceSizeMode() == EResourceSizeMode::EstimatedTotal)
+			const FShaderMapResource* Resource = MaterialShaderMap->GetResource();
+			if (Resource)
 			{
-				MaterialShaderMap->GetShaderList(UniqueShaders);
-				MaterialShaderMap->GetShaderPipelineList(ShaderPipelines);
+				CumulativeResourceSize.AddDedicatedSystemMemoryBytes(Resource->GetSizeBytes());
 			}
-		}
-	}
 
-	for (auto& KeyValue : UniqueShaders)
-	{
-		const TShaderRef<FShader>& Shader = KeyValue.Value;
-		if (Shader.IsValid())
-		{
-			CumulativeResourceSize.AddDedicatedSystemMemoryBytes(AddShaderSize(Shader, UniqueShaderResourceIds));
+			// TODO - account for shader code size somehow, either inline or shared code
 		}
 	}
-
-	for (const FShaderPipelineRef& Pipeline : ShaderPipelines)
-	{
-		if (Pipeline.IsValid())
-		{
-			for (const TShaderRef<FShader>& Shader : Pipeline.GetShaders())
-			{
-				if (Shader.IsValid())
-				{
-					CumulativeResourceSize.AddDedicatedSystemMemoryBytes(AddShaderSize(Shader, UniqueShaderResourceIds));
-				}
-			}
-		}
-	}	
 }
 
 /**
diff --git a/Engine/Source/Runtime/Engine/Private/ShaderCompiler/ShaderCompiler.cpp b/Engine/Source/Runtime/Engine/Private/ShaderCompiler/ShaderCompiler.cpp
--- a/Engine/Source/Runtime/Engine/Private/ShaderCompiler/ShaderCompiler.cpp
+++ b/Engine/Source/Runtime/Engine/Private/ShaderCompiler/ShaderCompiler.cpp
@@ -3883,16 +3883,6 @@ public:
 		Payload << MaterialsToLoad;
 		uint32 ShaderPlatform = ( uint32 )GMaxRHIShaderPlatform;
 		Payload << ShaderPlatform;
-		check(false); // TODO
-		// tell the other side the Ids we have so it doesn't send back duplicates
-		// (need to serialize this into a TArray since FShaderResourceId isn't known in the file server)
-		//TArray<FShaderResourceId> AllIds;
-		//FShaderResource_InlineCode::GetAllShaderResourceId( AllIds );
-
-		TArray<uint8> SerializedBytes;
-		FMemoryWriter Ar( SerializedBytes );
-		//Ar << AllIds;
-		Payload << SerializedBytes;
 		Payload << bCompileChangedShaders;
 	}
 
@@ -3907,7 +3897,11 @@ public:
 		FlushRenderingCommands();
 
 		// reload the global shaders
-		CompileGlobalShaderMap(true);
+		{
+			extern int32 GCreateShadersOnLoad; // Some platforms rely on global shaders to be created to implement basic RHI functionality
+			TGuardValue<int32> Guard(GCreateShadersOnLoad, 1);
+			CompileGlobalShaderMap(true);
+		}
 
 		// load all the mesh material shaders if any were sent back
 		if (MeshMaterialMaps.Num() > 0)
@@ -4798,7 +4792,6 @@ void RecompileShadersForRemote(
 	EShaderPlatform ShaderPlatformToCompile,
 	const FString& OutputDirectory,
 	const TArray<FString>& MaterialsToLoad,
-	const TArray<uint8>& SerializedShaderResources,
 	TArray<uint8>* MeshMaterialMaps,
 	TArray<FString>* ModifiedFiles,
 	bool bCompileChangedShaders)
@@ -4874,15 +4867,10 @@ void RecompileShadersForRemote(
 					// write the shader compilation info to memory, converting fnames to strings
 					FMemoryWriter MemWriter(*MeshMaterialMaps, true);
 					FNameAsStringProxyArchive Ar(MemWriter);
-
-					// pull the serialized resource ids into an array of resources
-					TArray<FShaderResourceId> ClientResourceIds;
-					check(false); // TODO
-					//FMemoryReader MemReader(SerializedShaderResources, true);
-					//MemReader << ClientResourceIds;
+					Ar.SetCookingTarget(TargetPlatform);
 
 					// save out the shader map to the byte array
-					FMaterialShaderMap::SaveForRemoteRecompile(Ar, CompiledShaderMaps, ClientResourceIds);
+					FMaterialShaderMap::SaveForRemoteRecompile(Ar, CompiledShaderMaps);
 				}
 
 				// save it out so the client can get it (and it's up to date next time)
diff --git a/Engine/Source/Runtime/Engine/Public/MaterialShared.h b/Engine/Source/Runtime/Engine/Public/MaterialShared.h
--- a/Engine/Source/Runtime/Engine/Public/MaterialShared.h
+++ b/Engine/Source/Runtime/Engine/Public/MaterialShared.h
@@ -1210,7 +1210,7 @@ public:
 	void RestoreShadersFromMemory(const TArray<uint8>& ShaderData);
 
 	/** Serializes a shader map to an archive (used with recompiling shaders for a remote console) */
-	ENGINE_API static void SaveForRemoteRecompile(FArchive& Ar, const TMap<FString, TArray<TRefCountPtr<FMaterialShaderMap> > >& CompiledShaderMaps, const TArray<FShaderResourceId>& ClientResourceIds);
+	ENGINE_API static void SaveForRemoteRecompile(FArchive& Ar, const TMap<FString, TArray<TRefCountPtr<FMaterialShaderMap> > >& CompiledShaderMaps);
 	ENGINE_API static void LoadForRemoteRecompile(FArchive& Ar, EShaderPlatform ShaderPlatform, const TArray<FString>& MaterialsForShaderMaps);
 
 #if WITH_EDITOR
diff --git a/Engine/Source/Runtime/Engine/Public/ShaderCompiler.h b/Engine/Source/Runtime/Engine/Public/ShaderCompiler.h
--- a/Engine/Source/Runtime/Engine/Public/ShaderCompiler.h
+++ b/Engine/Source/Runtime/Engine/Public/ShaderCompiler.h
@@ -715,7 +715,6 @@ extern ENGINE_API FString SaveGlobalShaderFile(EShaderPlatform Platform, FString
 * @param PlatformName					Name of the Platform the shaders are compiled for
 * @param OutputDirectory				The directory the compiled data will be stored to
 * @param MaterialsToLoad				List of Materials that need to be loaded and compiled
-* @param SerializedShaderResources		Serialized shader resources
 * @param MeshMaterialMaps				Mesh material maps
 * @param ModifiedFiles					Returns the list of modified files if not NULL
 * @param bCompileChangedShaders		Whether to compile all changed shaders or the specific material that is passed
@@ -725,7 +724,6 @@ extern ENGINE_API void RecompileShadersForRemote(
 	EShaderPlatform ShaderPlatform,
 	const FString& OutputDirectory,
 	const TArray<FString>& MaterialsToLoad,
-	const TArray<uint8>& SerializedShaderResources,
 	TArray<uint8>* MeshMaterialMaps,
 	TArray<FString>* ModifiedFiles,
 	bool bCompileChangedShaders = true);
diff --git a/Engine/Source/Runtime/NetworkFileSystem/Private/NetworkFileServerConnection.cpp b/Engine/Source/Runtime/NetworkFileSystem/Private/NetworkFileServerConnection.cpp
--- a/Engine/Source/Runtime/NetworkFileSystem/Private/NetworkFileServerConnection.cpp
+++ b/Engine/Source/Runtime/NetworkFileSystem/Private/NetworkFileServerConnection.cpp
@@ -1153,7 +1153,6 @@ void FNetworkFileServerClientConnection::ProcessRecompileShaders( FArchive& In,
 	// tell other side all the materials to load, by pathname
 	In << RecompileData.MaterialsToLoad;
 	In << RecompileData.ShaderPlatform;
-	In << RecompileData.SerializedShaderResources;
 	In << RecompileData.bCompileChangedShaders;
 
 	NetworkFileDelegates->RecompileShadersDelegate.ExecuteIfBound(RecompileData);
diff --git a/Engine/Source/Runtime/NetworkFileSystem/Public/INetworkFileSystemModule.h b/Engine/Source/Runtime/NetworkFileSystem/Public/INetworkFileSystemModule.h
--- a/Engine/Source/Runtime/NetworkFileSystem/Public/INetworkFileSystemModule.h
+++ b/Engine/Source/Runtime/NetworkFileSystem/Public/INetworkFileSystemModule.h
@@ -23,7 +23,6 @@ struct FShaderRecompileData
 	TArray<FString>* ModifiedFiles;
 	TArray<uint8>* MeshMaterialMaps;
 	TArray<FString> MaterialsToLoad;
-	TArray<uint8> SerializedShaderResources;
 	bool bCompileChangedShaders;
 
 	FShaderRecompileData() :
@@ -38,7 +37,6 @@ struct FShaderRecompileData
 		ModifiedFiles = Other.ModifiedFiles;
 		MeshMaterialMaps = Other.MeshMaterialMaps;
 		MaterialsToLoad = Other.MaterialsToLoad;
-		SerializedShaderResources = Other.SerializedShaderResources;
 		bCompileChangedShaders = Other.bCompileChangedShaders;
 
 		return *this;
diff --git a/Engine/Source/Runtime/RenderCore/Private/ShaderMap.cpp b/Engine/Source/Runtime/RenderCore/Private/ShaderMap.cpp
--- a/Engine/Source/Runtime/RenderCore/Private/ShaderMap.cpp
+++ b/Engine/Source/Runtime/RenderCore/Private/ShaderMap.cpp
@@ -143,6 +143,7 @@ void FShaderMapBase::Serialize(FArchive& Ar, bool bInlineShaderResources, bool b
 
 		void* SaveFrozenContent = MemoryImageResult.Bytes.GetData();
 		uint32 SaveFrozenContentSize = MemoryImageResult.Bytes.Num();
+		check(SaveFrozenContentSize > 0u);
 		Ar << SaveFrozenContentSize;
 		Ar.Serialize(SaveFrozenContent, SaveFrozenContentSize);
 		MemoryImageResult.SaveToArchive(Ar);
@@ -190,6 +191,9 @@ void FShaderMapBase::Serialize(FArchive& Ar, bool bInlineShaderResources, bool b
 		PointerTable = CreatePointerTable();
 
 		Ar << FrozenContentSize;
+		// ensure frozen content is at least as big as our FShaderMapContent-derived class
+		checkf(FrozenContentSize >= ContentTypeLayout.Size, TEXT("Invalid FrozenContentSize for %s, got %d, expected at least %d"), ContentTypeLayout.Name, FrozenContentSize, ContentTypeLayout.Size);
+
 		void* ContentMemory = FMemory::Malloc(FrozenContentSize);
 		Ar.Serialize(ContentMemory, FrozenContentSize);
 		Content = static_cast<FShaderMapContent*>(ContentMemory);
diff --git a/Engine/Source/Runtime/RenderCore/Public/Shader.h b/Engine/Source/Runtime/RenderCore/Public/Shader.h
--- a/Engine/Source/Runtime/RenderCore/Public/Shader.h
+++ b/Engine/Source/Runtime/RenderCore/Public/Shader.h
@@ -95,11 +95,6 @@ public:
 	}
 };
 
-// TODO - remove me
-class FShaderResourceId
-{
-};
-
 class FShaderParameterInfo
 {
 	DECLARE_TYPE_LAYOUT(FShaderParameterInfo, NonVirtual);
