Fixing bogus warning that treated a lifetime counter as a per-frame counter.

diff --git a/Engine/Plugins/FX/Niagara/Source/Niagara/Private/NiagaraEmitterInstance.cpp b/Engine/Plugins/FX/Niagara/Source/Niagara/Private/NiagaraEmitterInstance.cpp
--- a/Engine/Plugins/FX/Niagara/Source/Niagara/Private/NiagaraEmitterInstance.cpp
+++ b/Engine/Plugins/FX/Niagara/Source/Niagara/Private/NiagaraEmitterInstance.cpp
@@ -1177,15 +1177,18 @@ void FNiagaraEmitterInstance::Tick(float DeltaSeconds)
 			GpuSpawnInfo.MaxParticleCount = AllocationSize;
 
 			int NumSpawnInfos = 0;
+			int32 NumSpawnedOnGPUThisFrame = 0;
 			if (ExecutionState == ENiagaraExecutionState::Active)
 			{
-				for (FNiagaraSpawnInfo& Info : SpawnInfos)
+				for (int32 SpawnInfoIdx = 0; SpawnInfoIdx < SpawnInfos.Num(); SpawnInfoIdx++)
 				{
+					FNiagaraSpawnInfo& Info = SpawnInfos[SpawnInfoIdx];
 					if (Info.Count > 0 && (NumSpawnInfos < NIAGARA_MAX_GPU_SPAWN_INFOS))
 					{
-						if ((TotalSpawnedParticles + Info.Count) > GMaxNiagaraGPUParticlesSpawnPerFrame)
+						NumSpawnedOnGPUThisFrame += Info.Count;
+						if (NumSpawnedOnGPUThisFrame > GMaxNiagaraGPUParticlesSpawnPerFrame)
 						{
-							UE_LOG(LogNiagara, Warning, TEXT("%s has attempted to execeed max GPU per frame spawn! | Max: %d | Requested: %d"), *CachedEmitter->GetUniqueEmitterName(), GMaxNiagaraGPUParticlesSpawnPerFrame, TotalSpawnedParticles + Info.Count);
+							UE_LOG(LogNiagara, Warning, TEXT("%s has attempted to execeed max GPU per frame spawn! | Max: %d | Requested: %d | SpawnInfoEntry: %d"), *CachedEmitter->GetUniqueEmitterName(), GMaxNiagaraGPUParticlesSpawnPerFrame, NumSpawnedOnGPUThisFrame, SpawnInfoIdx);
 							break;
 						}
 
