Since it is impossible to set a reference to a world within the system editor, make all the properties referencing AActors or UActorComponent derived classes disabled in the system editor.

diff --git a/Engine/Plugins/FX/Niagara/Source/Niagara/Classes/NiagaraDataInterfaceSkeletalMesh.h b/Engine/Plugins/FX/Niagara/Source/Niagara/Classes/NiagaraDataInterfaceSkeletalMesh.h
--- a/Engine/Plugins/FX/Niagara/Source/Niagara/Classes/NiagaraDataInterfaceSkeletalMesh.h
+++ b/Engine/Plugins/FX/Niagara/Source/Niagara/Classes/NiagaraDataInterfaceSkeletalMesh.h
@@ -519,7 +519,7 @@ public:
 	USkeletalMesh* DefaultMesh;
 #endif
 
-	/** The source actor from which to sample. Takes precedence over the direct mesh. */
+	/** The source actor from which to sample. Takes precedence over the direct mesh. Note that this can only be set when used as a user variable on a component in the world.*/
 	UPROPERTY(EditAnywhere, Category = "Mesh")
 	AActor* Source;
 
diff --git a/Engine/Plugins/FX/Niagara/Source/Niagara/Classes/NiagaraDataInterfaceSpline.h b/Engine/Plugins/FX/Niagara/Source/Niagara/Classes/NiagaraDataInterfaceSpline.h
--- a/Engine/Plugins/FX/Niagara/Source/Niagara/Classes/NiagaraDataInterfaceSpline.h
+++ b/Engine/Plugins/FX/Niagara/Source/Niagara/Classes/NiagaraDataInterfaceSpline.h
@@ -27,7 +27,7 @@ class NIAGARA_API UNiagaraDataInterfaceSpline : public UNiagaraDataInterface
 
 public:
 	
-	/** The source actor from which to sample.  */
+	/** The source actor from which to sample.  Note that this can only be set when used as a user variable on a component in the world.*/
 	UPROPERTY(EditAnywhere, Category = "Spline")
 	AActor* Source;
 	
diff --git a/Engine/Plugins/FX/Niagara/Source/Niagara/Classes/NiagaraDataInterfaceStaticMesh.h b/Engine/Plugins/FX/Niagara/Source/Niagara/Classes/NiagaraDataInterfaceStaticMesh.h
--- a/Engine/Plugins/FX/Niagara/Source/Niagara/Classes/NiagaraDataInterfaceStaticMesh.h
+++ b/Engine/Plugins/FX/Niagara/Source/Niagara/Classes/NiagaraDataInterfaceStaticMesh.h
@@ -171,7 +171,7 @@ public:
 	UPROPERTY(EditAnywhere, Category = "Mesh")
 	UStaticMesh* DefaultMesh;
 
-	/** The source actor from which to sample. Takes precedence over the direct mesh. */
+	/** The source actor from which to sample. Takes precedence over the direct mesh. Note that this can only be set when used as a user variable on a component in the world. */
 	UPROPERTY(EditAnywhere, Category = "Mesh")
 	AActor* Source;
 	
diff --git a/Engine/Plugins/FX/Niagara/Source/NiagaraEditor/Private/ViewModels/Stack/NiagaraStackPropertyRow.cpp b/Engine/Plugins/FX/Niagara/Source/NiagaraEditor/Private/ViewModels/Stack/NiagaraStackPropertyRow.cpp
--- a/Engine/Plugins/FX/Niagara/Source/NiagaraEditor/Private/ViewModels/Stack/NiagaraStackPropertyRow.cpp
+++ b/Engine/Plugins/FX/Niagara/Source/NiagaraEditor/Private/ViewModels/Stack/NiagaraStackPropertyRow.cpp
@@ -18,6 +18,16 @@ void UNiagaraStackPropertyRow::Initialize(FRequiredEntryData InRequiredEntryData
 	RowStyle = DetailTreeNode->GetNodeType() == EDetailNodeType::Category
 		? EStackRowStyle::ItemCategory
 		: EStackRowStyle::ItemContent;
+	bCannotEditInThisContext = false;
+	if (PropertyHandle.IsValid() && PropertyHandle->GetProperty())
+	{
+		FProperty* Prop = PropertyHandle->GetProperty();
+		FObjectPropertyBase* ObjProp = CastField<FObjectPropertyBase>(Prop);
+		if (ObjProp && ObjProp->PropertyClass && (ObjProp->PropertyClass->IsChildOf(AActor::StaticClass()) || ObjProp->PropertyClass->IsChildOf(UActorComponent::StaticClass())))
+		{
+			bCannotEditInThisContext = true;
+		}
+	}
 }
 
 TSharedRef<IDetailTreeNode> UNiagaraStackPropertyRow::GetDetailTreeNode() const
@@ -27,6 +37,8 @@ TSharedRef<IDetailTreeNode> UNiagaraStackPropertyRow::GetDetailTreeNode() const
 
 bool UNiagaraStackPropertyRow::GetIsEnabled() const
 {
+	if (bCannotEditInThisContext) 
+		return false;
 	return OwningNiagaraNode == nullptr || OwningNiagaraNode->GetDesiredEnabledState() == ENodeEnabledState::Enabled;
 }
 
diff --git a/Engine/Plugins/FX/Niagara/Source/NiagaraEditor/Public/ViewModels/Stack/NiagaraStackPropertyRow.h b/Engine/Plugins/FX/Niagara/Source/NiagaraEditor/Public/ViewModels/Stack/NiagaraStackPropertyRow.h
--- a/Engine/Plugins/FX/Niagara/Source/NiagaraEditor/Public/ViewModels/Stack/NiagaraStackPropertyRow.h
+++ b/Engine/Plugins/FX/Niagara/Source/NiagaraEditor/Public/ViewModels/Stack/NiagaraStackPropertyRow.h
@@ -31,4 +31,5 @@ private:
 	TSharedPtr<IDetailTreeNode> DetailTreeNode;
 	UNiagaraNode* OwningNiagaraNode;
 	EStackRowStyle RowStyle;
+	bool bCannotEditInThisContext;
 };
\ No newline at end of file
