Fixing event errors if same event used multiple times (which is ok)

diff --git a/Engine/Plugins/FX/Niagara/Source/NiagaraEditor/Private/NiagaraHlslTranslator.cpp b/Engine/Plugins/FX/Niagara/Source/NiagaraEditor/Private/NiagaraHlslTranslator.cpp
--- a/Engine/Plugins/FX/Niagara/Source/NiagaraEditor/Private/NiagaraHlslTranslator.cpp
+++ b/Engine/Plugins/FX/Niagara/Source/NiagaraEditor/Private/NiagaraHlslTranslator.cpp
@@ -2566,6 +2566,7 @@ void FHlslNiagaraTranslator::DefineMain(FString &OutHlslOutput,
 			{
 				FString CopyStr = TEXT("\tContext.") + TranslationStages[StageIdx + 1].PassNamespace + TEXT(".Particles = Context.") + TranslationStages[StageIdx].PassNamespace + TEXT(".Particles;\n");
 				OutHlslOutput += CopyStr;
+				
 			}
 			OutHlslOutput += TEXT("\t//End Transfer of Attributes!\n\n");
 		}
@@ -4936,14 +4937,15 @@ void FHlslNiagaraTranslator::ReadDataSet(const FNiagaraDataSetID DataSet, const
 	}
 
 	TMap<int32, FDataSetAccessInfo>& Reads = DataSetReadInfo[(int32)AccessMode].FindOrAdd(DataSet);
-	if (Reads.Num() != 0)
-	{
-		Error(FText::Format(LOCTEXT("TooManyDataSetReads", "Only one Event Read node per Event handler! Read data set node: \"{0}\""), FText::FromName(DataSet.Name)), nullptr, nullptr);
-	}
 
 	FDataSetAccessInfo* DataSetReadForInput = Reads.Find(InputChunk);
 	if (!DataSetReadForInput)
 	{
+		if (Reads.Num() != 0) // If it is the same event within the graph that is ok, but we don't get here unless it is new.
+		{
+			Error(FText::Format(LOCTEXT("TooManyDataSetReads", "Only one Event Read node per Event handler! Read data set node: \"{0}\""), FText::FromName(DataSet.Name)), nullptr, nullptr);
+		}
+
 		DataSetReadForInput = &Reads.Add(InputChunk);
 
 		DataSetReadForInput->Variables = Variables;
