Fix for random crash encountered during testing

diff --git a/Engine/Plugins/FX/Niagara/Source/Niagara/Classes/NiagaraDataSet.h b/Engine/Plugins/FX/Niagara/Source/Niagara/Classes/NiagaraDataSet.h
--- a/Engine/Plugins/FX/Niagara/Source/Niagara/Classes/NiagaraDataSet.h
+++ b/Engine/Plugins/FX/Niagara/Source/Niagara/Classes/NiagaraDataSet.h
@@ -1357,10 +1357,20 @@ struct FNiagaraDataSetAccessor<FNiagaraSpawnInfo> : public FNiagaraDataSetAccess
 	FORCEINLINE void Get(int32 Index, FNiagaraSpawnInfo& OutValue)const
 	{
 		checkSlow(IsValidForRead());
-		OutValue.Count = SrcCountBase[Index];
-		OutValue.InterpStartDt = SrcInterpStartDtBase[Index];
-		OutValue.IntervalDt = SrcIntervalDtBase[Index];
-		OutValue.SpawnGroup = SrcGroupBase[Index];
+		if (SrcCountBase)
+		{
+			OutValue.Count = SrcCountBase[Index];
+			OutValue.InterpStartDt = SrcInterpStartDtBase[Index];
+			OutValue.IntervalDt = SrcIntervalDtBase[Index];
+			OutValue.SpawnGroup = SrcGroupBase[Index];
+		}
+		else
+		{
+			OutValue.Count = 0;
+			OutValue.InterpStartDt = 0;
+			OutValue.IntervalDt = 0;
+			OutValue.SpawnGroup = 0;
+		}
 	}
 
 	FORCEINLINE void Set(int32 Index, const FNiagaraSpawnInfo& InValue)
