Allowing rejection samping

diff --git a/Engine/Plugins/FX/Niagara/Source/NiagaraCore/Private/NiagaraCustomVersion.cpp b/Engine/Plugins/FX/Niagara/Source/NiagaraCore/Private/NiagaraCustomVersion.cpp
--- a/Engine/Plugins/FX/Niagara/Source/NiagaraCore/Private/NiagaraCustomVersion.cpp
+++ b/Engine/Plugins/FX/Niagara/Source/NiagaraCore/Private/NiagaraCustomVersion.cpp
@@ -8,4 +8,4 @@ const FGuid FNiagaraCustomVersion::GUID(0xFCF57AFA, 0x50764283, 0xB9A9E658, 0xFF
 // Register the custom version with core
 FCustomVersionRegistration GRegisterNiagaraCustomVersion(FNiagaraCustomVersion::GUID, FNiagaraCustomVersion::LatestVersion, TEXT("NiagaraVer"));
 
-const FGuid FNiagaraCustomVersion::LatestScriptCompileVersion(0xE599EF89, 0x9D574012, 0x99E29018, 0x4682F9EC);
+const FGuid FNiagaraCustomVersion::LatestScriptCompileVersion(0x0C1CF412, 0xA3B240BB, 0x900EBA38, 0x0CE36DDA );
diff --git a/Engine/Plugins/FX/Niagara/Source/NiagaraEditor/Private/NiagaraHlslTranslator.cpp b/Engine/Plugins/FX/Niagara/Source/NiagaraEditor/Private/NiagaraHlslTranslator.cpp
--- a/Engine/Plugins/FX/Niagara/Source/NiagaraEditor/Private/NiagaraHlslTranslator.cpp
+++ b/Engine/Plugins/FX/Niagara/Source/NiagaraEditor/Private/NiagaraHlslTranslator.cpp
@@ -2212,7 +2212,7 @@ void FHlslNiagaraTranslator::DefineMainGPUFunctions(
 				if (ParamMapDefinedAttributesToNamespaceVars.Num() != 0)
 				{
 					HlslOutput += TEXT("\t\tContext.") + TranslationStages[i].PassNamespace + TEXT(".Particles = Context.") + TranslationStages[i - 1].PassNamespace + TEXT(".Particles;\n");
-					if (bUsesAlive)
+					if (TranslationStages[i - 1].bWritesAlive)
 					{
 						HlslOutput += TEXT("\t\tContext.") + TranslationStages[i].PassNamespace + TEXT(".DataInstance = Context.") + TranslationStages[i - 1].PassNamespace + TEXT(".DataInstance;\n");
 					}
@@ -2564,8 +2564,11 @@ void FHlslNiagaraTranslator::DefineMain(FString &OutHlslOutput,
 			OutHlslOutput += TEXT("\t//Begin Transfer of Attributes!\n");
 			if (ParamMapDefinedAttributesToNamespaceVars.Num() != 0)
 			{
-				FString CopyStr = TEXT("\tContext.") + TranslationStages[StageIdx + 1].PassNamespace + TEXT(".Particles = Context.") + TranslationStages[StageIdx].PassNamespace + TEXT(".Particles;\n");
-				OutHlslOutput += CopyStr;
+				OutHlslOutput += TEXT("\tContext.") + TranslationStages[StageIdx + 1].PassNamespace + TEXT(".Particles = Context.") + TranslationStages[StageIdx].PassNamespace + TEXT(".Particles;\n");
+				if (TranslationStages[StageIdx].bWritesAlive)
+				{
+					OutHlslOutput += TEXT("\t\tContext.") + TranslationStages[StageIdx + 1].PassNamespace + TEXT(".DataInstance = Context.") + TranslationStages[StageIdx].PassNamespace + TEXT(".DataInstance;\n");
+				}
 				
 			}
 			OutHlslOutput += TEXT("\t//End Transfer of Attributes!\n\n");
@@ -4003,6 +4006,10 @@ void FHlslNiagaraTranslator::ParameterMapSet(UNiagaraNodeParameterMapSet* SetNod
 			}
 			else
 			{
+				if (Var == FNiagaraVariable(FNiagaraTypeDefinition::GetBoolDef(), TEXT("DataInstance.Alive")))
+				{
+					TranslationStages[ActiveStageIdx].bWritesAlive = true;
+				}
 				AddBodyChunk(ParameterMapInstanceName + TEXT(".") + GetSanitizedSymbolName(Var.GetName().ToString()), TEXT("{0}"), Var.GetType(), Input, false);
 			}
 		}
@@ -4590,6 +4597,13 @@ void FHlslNiagaraTranslator::HandleParameterRead(int32 ParamMapHistoryIdx, const
 		bIsPerInstanceAttribute = true;
 	}
 
+	// Make sure to leave IsAlive alone if copying over previous stage params.
+	if (Var == FNiagaraVariable(FNiagaraTypeDefinition::GetBoolDef(), TEXT("DataInstance.Alive")) && ActiveStageIdx > 0 && TranslationStages[ActiveStageIdx - 1].bCopyPreviousParams 
+		&& TranslationStages[ActiveStageIdx - 1].bWritesAlive)
+	{
+		bIsPerInstanceAttribute = true; 
+	}
+
 	int32 LastSetChunkIdx = INDEX_NONE;
 	if (ParamMapHistoryIdx < ParamMapHistories.Num())
 	{
diff --git a/Engine/Plugins/FX/Niagara/Source/NiagaraEditor/Private/NiagaraHlslTranslator.h b/Engine/Plugins/FX/Niagara/Source/NiagaraEditor/Private/NiagaraHlslTranslator.h
--- a/Engine/Plugins/FX/Niagara/Source/NiagaraEditor/Private/NiagaraHlslTranslator.h
+++ b/Engine/Plugins/FX/Niagara/Source/NiagaraEditor/Private/NiagaraHlslTranslator.h
@@ -319,6 +319,7 @@ public:
 	int32 SourceSimStage = -1;
 	bool bSpawnOnly;
 	bool bUsesAlive;
+	bool bWritesAlive = false;
 };
 
 class NIAGARAEDITOR_API FHlslNiagaraTranslator
