Fix for generating compile hashes for TMap<FName,FName> (currently used for particle read DI).  Resulted in incorrect hashes and NDC cooking as it was including random memory.

diff --git a/Engine/Plugins/FX/Niagara/Source/NiagaraEditor/Private/NiagaraEditorUtilities.cpp b/Engine/Plugins/FX/Niagara/Source/NiagaraEditor/Private/NiagaraEditorUtilities.cpp
--- a/Engine/Plugins/FX/Niagara/Source/NiagaraEditor/Private/NiagaraEditorUtilities.cpp
+++ b/Engine/Plugins/FX/Niagara/Source/NiagaraEditor/Private/NiagaraEditorUtilities.cpp
@@ -503,9 +503,9 @@ bool FNiagaraEditorUtilities::NestedPropertiesAppendCompileHash(const void* Cont
 						for (int32 ArrayIdx = 0; ArrayIdx < MapHelper.Num(); ArrayIdx++)
 						{
 							InVisitor->UpdateString(*FString::Printf(TEXT("Key[%d]"), ArrayIdx), Names[ArrayIdx].ToString());
-							if (!PODPropertyAppendCompileHash(MapHelper.GetValuePtr(ArrayIdx), MapHelper.GetValueProperty(), FString::Printf(TEXT("Value[%d]"), ArrayIdx), InVisitor))
+							if (!PODPropertyAppendCompileHash(MapHelper.GetPairPtr(ArrayIdx), MapHelper.GetValueProperty(), FString::Printf(TEXT("Value[%d]"), ArrayIdx), InVisitor))
 							{
-								UE_LOG(LogNiagaraEditor, Warning, TEXT("Skipping %s because it is an map value property of unsupported underlying type, please add \"meta = (SkipForCompileHash=\"true\")\" to avoid this warning in the future or handle it yourself in NestedPropertiesAppendCompileHash!"), *Property->GetName());
+								UE_LOG(LogNiagaraEditor, Warning, TEXT("Skipping %s because it is an map value property of unsupported underlying type, please add \"meta = (SkipForCompileHash=\"true\")\" to avoid this warning in the future or handle it yourself in PODPropertyAppendCompileHash!"), *Property->GetName());
 								bPassed = false;
 								continue;
 							}
