Mark UNiagaraScripts/UNiagaraEmitters as editor only if they are not owned by a UNiagaraSystem
-works around the fact that parent emitters can be marked for export because they are referenced by our RF_Public data interfaces

diff --git a/Engine/Plugins/FX/Niagara/Source/Niagara/Classes/NiagaraEmitter.h b/Engine/Plugins/FX/Niagara/Source/Niagara/Classes/NiagaraEmitter.h
--- a/Engine/Plugins/FX/Niagara/Source/Niagara/Classes/NiagaraEmitter.h
+++ b/Engine/Plugins/FX/Niagara/Source/Niagara/Classes/NiagaraEmitter.h
@@ -243,6 +243,7 @@ public:
 	void Serialize(FArchive& Ar)override;
 	virtual void PostInitProperties() override;
 	virtual void PostLoad() override;
+	virtual bool IsEditorOnly() const override;
 	//End UObject Interface
 
 	/** Toggles whether or not the particles within this emitter are relative to the emitter origin or in global space. */ 
diff --git a/Engine/Plugins/FX/Niagara/Source/Niagara/Classes/NiagaraScript.h b/Engine/Plugins/FX/Niagara/Source/Niagara/Classes/NiagaraScript.h
--- a/Engine/Plugins/FX/Niagara/Source/Niagara/Classes/NiagaraScript.h
+++ b/Engine/Plugins/FX/Niagara/Source/Niagara/Classes/NiagaraScript.h
@@ -558,6 +558,8 @@ public:
 	virtual void PostEditChangeProperty(FPropertyChangedEvent& PropertyChangedEvent) override;
 #endif
 	virtual void GetAssetRegistryTags(TArray<FAssetRegistryTag>& OutTags) const override;
+
+	virtual bool IsEditorOnly() const override;
 	//~ End UObject interface
 
 	// Infrastructure for GPU compute Shaders
diff --git a/Engine/Plugins/FX/Niagara/Source/Niagara/Private/NiagaraEmitter.cpp b/Engine/Plugins/FX/Niagara/Source/Niagara/Private/NiagaraEmitter.cpp
--- a/Engine/Plugins/FX/Niagara/Source/Niagara/Private/NiagaraEmitter.cpp
+++ b/Engine/Plugins/FX/Niagara/Source/Niagara/Private/NiagaraEmitter.cpp
@@ -563,6 +563,32 @@ void UNiagaraEmitter::PostLoad()
 	ResolveScalabilitySettings();
 }
 
+bool UNiagaraEmitter::IsEditorOnly() const
+{
+#if WITH_EDITOR
+	if (HasAnyFlags(RF_ClassDefaultObject))
+	{
+		return false;
+	}
+
+	// if the emitter does not have a system as it's outer than it is likely a standalone emitter/parent emitter and so is editor only
+	if (const UNiagaraSystem* SystemOwner = Cast<const UNiagaraSystem>(GetOuter()))
+	{
+		for (const auto& SystemEmitterHandle : SystemOwner->GetEmitterHandles())
+		{
+			if (SystemEmitterHandle.GetInstance() == this)
+			{
+				return false;
+			}
+		}
+	}
+
+	return true;
+#else
+	return Super::IsEditorOnly();
+#endif
+}
+
 #if WITH_EDITOR
 /** Creates a new emitter with the supplied emitter as a parent emitter and the supplied system as it's owner. */
 UNiagaraEmitter* UNiagaraEmitter::CreateWithParentAndOwner(UNiagaraEmitter& InParentEmitter, UObject* InOwner, FName InName, EObjectFlags FlagMask)
diff --git a/Engine/Plugins/FX/Niagara/Source/Niagara/Private/NiagaraScript.cpp b/Engine/Plugins/FX/Niagara/Source/Niagara/Private/NiagaraScript.cpp
--- a/Engine/Plugins/FX/Niagara/Source/Niagara/Private/NiagaraScript.cpp
+++ b/Engine/Plugins/FX/Niagara/Source/Niagara/Private/NiagaraScript.cpp
@@ -1586,6 +1586,31 @@ void UNiagaraScript::GetAssetRegistryTags(TArray<FAssetRegistryTag>& OutTags) co
 #endif
 }
 
+bool UNiagaraScript::IsEditorOnly() const
+{
+#if WITH_EDITOR
+	if (HasAnyFlags(RF_ClassDefaultObject))
+	{
+		return false;
+	}
+
+	if (const UNiagaraEmitter* EmitterOwner = Cast<const UNiagaraEmitter>(GetOuter()))
+	{
+		// we want to only cook scripts that are referenced by systems (as opposed to standalone scripts that may be getting
+		// referenced via an emitter's parent, this will also take care of GPUScripts that are created for CPU emitters
+		TArray<UNiagaraScript*> OwnerScripts;
+
+		EmitterOwner->GetScripts(OwnerScripts, false);
+
+		if (!OwnerScripts.Contains(this))
+		{
+			return true;
+		}
+	}
+#endif
+	return Super::IsEditorOnly();
+}
+
 #if WITH_EDITOR
 
 void UNiagaraScript::BeginCacheForCookedPlatformData(const ITargetPlatform *TargetPlatform)
