Bind the dummy UInt buffer to MeshTriangleBuffer binding slot when not doing uniformly distributed sampling on the GPU.

diff --git a/Engine/Plugins/FX/Niagara/Source/Niagara/Private/NiagaraDataInterfaceStaticMesh.cpp b/Engine/Plugins/FX/Niagara/Source/Niagara/Private/NiagaraDataInterfaceStaticMesh.cpp
--- a/Engine/Plugins/FX/Niagara/Source/Niagara/Private/NiagaraDataInterfaceStaticMesh.cpp
+++ b/Engine/Plugins/FX/Niagara/Source/Niagara/Private/NiagaraDataInterfaceStaticMesh.cpp
@@ -550,7 +550,7 @@ public:
 				}
 				else
 				{
-					SetSRVParameter(RHICmdList, ComputeShaderRHI, MeshTriangleBuffer, FNiagaraRenderer::GetDummyFloatBuffer());
+					SetSRVParameter(RHICmdList, ComputeShaderRHI, MeshTriangleBuffer, FNiagaraRenderer::GetDummyUIntBuffer());
 				}
 
 				const float InvDeltaTime = Data->DeltaSeconds > 0.0f ? 1.0f / Data->DeltaSeconds : 0.0f;
