Change the NiagaraParameterStore to use an array instead of a TMap for it's bindings for improved memory/performance

diff --git a/Engine/Plugins/FX/Niagara/Source/Niagara/Private/NiagaraParameterStore.cpp b/Engine/Plugins/FX/Niagara/Source/Niagara/Private/NiagaraParameterStore.cpp
--- a/Engine/Plugins/FX/Niagara/Source/Niagara/Private/NiagaraParameterStore.cpp
+++ b/Engine/Plugins/FX/Niagara/Source/Niagara/Private/NiagaraParameterStore.cpp
@@ -163,14 +163,13 @@ void FNiagaraParameterStore::Bind(FNiagaraParameterStore* DestStore, const FNiag
 {
 	check(DestStore);
 	SCOPE_CYCLE_COUNTER(STAT_NiagaraParameterStoreBind);
-	if (!Bindings.Contains(DestStore))
+
+	if (!Algo::FindBy(Bindings, DestStore, &BindingPair::Key))
 	{
-		// Bind the parameter stores only if they have variables in common.
 		FNiagaraParameterStoreBinding HeapBinding;
-		if  (HeapBinding.Initialize(DestStore, this, BoundParameters))
+		if (HeapBinding.Initialize(DestStore, this, BoundParameters))
 		{
-			FNiagaraParameterStoreBinding& Binding = Bindings.FindOrAdd(DestStore);
-			FMemory::Memswap(&Binding, &HeapBinding, sizeof(FNiagaraParameterStoreBinding));
+			Bindings.Emplace(DestStore, HeapBinding);
 		}
 	}
 }
@@ -302,11 +301,15 @@ bool FNiagaraParameterStoreBinding::BindParameters(FNiagaraParameterStore* DestS
 
 void FNiagaraParameterStore::Unbind(FNiagaraParameterStore* DestStore)
 {
-	FNiagaraParameterStoreBinding* Binding = Bindings.Find(DestStore);
-	if (Binding)
+	const int32 BindingIndex = Bindings.IndexOfByPredicate([DestStore](const BindingPair& Binding)
+	{
+		return Binding.Key == DestStore;
+	});
+
+	if (BindingIndex != INDEX_NONE)
 	{
-		Binding->Empty(DestStore, this);
-		Bindings.Remove(DestStore);
+		Bindings[BindingIndex].Value.Empty(DestStore, this);
+		Bindings.RemoveAtSwap(BindingIndex);
 	}
 }
 
@@ -332,10 +335,10 @@ void FNiagaraParameterStore::TransferBindings(FNiagaraParameterStore& OtherStore
 bool FNiagaraParameterStore::VerifyBinding(const FNiagaraParameterStore* DestStore)const
 {
 #if WITH_EDITORONLY_DATA
-	const FNiagaraParameterStoreBinding* Binding = Bindings.Find(DestStore);
+	const BindingPair* Binding = Algo::FindBy(Bindings, DestStore, &BindingPair::Key);
 	if (Binding)
 	{
-		return Binding->VerifyBinding(DestStore, this);
+		return Binding->Value.VerifyBinding(DestStore, this);
 	}
 	else
 	{
diff --git a/Engine/Plugins/FX/Niagara/Source/Niagara/Public/NiagaraParameterStore.h b/Engine/Plugins/FX/Niagara/Source/Niagara/Public/NiagaraParameterStore.h
--- a/Engine/Plugins/FX/Niagara/Source/Niagara/Public/NiagaraParameterStore.h
+++ b/Engine/Plugins/FX/Niagara/Source/Niagara/Public/NiagaraParameterStore.h
@@ -164,7 +164,8 @@ private:
 	TArray<UObject*> UObjects;
 
 	/** Bindings between this parameter store and others we push data into when we tick. */
-	TMap<FNiagaraParameterStore*, FNiagaraParameterStoreBinding> Bindings;
+	typedef TPair<FNiagaraParameterStore*, FNiagaraParameterStoreBinding> BindingPair;
+	TArray<BindingPair> Bindings;
 
 	/** Parameter stores we've been bound to and are feeding data into us. */
 	TArray<FNiagaraParameterStore*> SourceStores;
