Add emitter's full name to the compile hash for niagara scripts
Emitters from different systems could have their hash collide and we're not prepared to share the results between the two.  Additionally the collision would result in NDC issues as there'd be a race on which shader would be generated and there were elements in there (like stat scopes) that would differ between systems.

diff --git a/Engine/Plugins/FX/Niagara/Source/Niagara/Private/NiagaraScript.cpp b/Engine/Plugins/FX/Niagara/Source/Niagara/Private/NiagaraScript.cpp
--- a/Engine/Plugins/FX/Niagara/Source/Niagara/Private/NiagaraScript.cpp
+++ b/Engine/Plugins/FX/Niagara/Source/Niagara/Private/NiagaraScript.cpp
@@ -378,6 +378,11 @@ void UNiagaraScript::ComputeVMCompilationId(FNiagaraVMExecutableDataId& Id) cons
 		{
 			Id.AdditionalDefines.Add(TEXT("Emitter.UseOldShaderStages"));
 		}
+
+		// in order to generate deterministic cooks we need to differentiate between two emitters that might
+		// generate the same hash but slightly different shaders.  In particular the full path of the emitter
+		// is used to generate the stat scopes that won't directly change the shader but will alter the name map
+		Id.AdditionalDefines.Add(Emitter->GetFullName());
 	}
 
 	if (UNiagaraSystem* System = Cast<UNiagaraSystem>(Obj))
