Move FNiagaraVariableMetaData from storing target scope directly to holding a registered scope info name. Scope infos are now registered in the editor module and represent the scope to display a parameter in and the actual namespace string to use when generating the hlsl token name.

Fixup postload to move subnamespaced parameters to registered scope infos.

Consolidated some stack graph utilities methods to editor utilities.

diff --git a/Engine/Plugins/FX/Niagara/Source/Niagara/Classes/NiagaraConstants.h b/Engine/Plugins/FX/Niagara/Source/Niagara/Classes/NiagaraConstants.h
--- a/Engine/Plugins/FX/Niagara/Source/Niagara/Classes/NiagaraConstants.h
+++ b/Engine/Plugins/FX/Niagara/Source/Niagara/Classes/NiagaraConstants.h
@@ -9,6 +9,7 @@
 #define PARAM_MAP_ENGINE_OWNER_STR TEXT("Engine.Owner.")
 #define PARAM_MAP_ENGINE_SYSTEM_STR TEXT("Engine.System.")
 #define PARAM_MAP_ENGINE_EMITTER_STR TEXT("Engine.Emitter.")
+#define PARAM_MAP_LOCAL_MODULE_STR TEXT("Local.Module.")
 #define PARAM_MAP_USER_STR TEXT("User.")
 #define PARAM_MAP_SYSTEM_STR TEXT("System.")
 #define PARAM_MAP_EMITTER_STR TEXT("Emitter.")
@@ -144,6 +145,30 @@ struct NIAGARA_API FNiagaraConstants
 	static const FName InputPinName;
 	static const FName OutputPinName;
 
+
+	/** Reserved Namespace Names and Scope Names */
+	static const FName UserNamespace;
+	static const FName EngineNamespace;
+	static const FName SystemNamespace;
+	static const FName EmitterNamespace;
+	static const FName ParticleAttributeNamespace;
+	static const FName ModuleNamespace;
+	static const FName ParameterCollectionNamespace;
+	static const FString InitialPrefix;
+	static const FName LocalNamespace;
+	static const FName InitialNamespace;
+	static const FName OwnerNamespace;
+
+	static const FName EngineOwnerScopeName;
+	static const FName EngineSystemScopeName;
+	static const FName EngineEmitterScopeName;
+
+	static const FName ScriptTransientScopeName;
+	static const FName ScriptPersistentScopeName;
+	static const FName InputScopeName;
+	static const FName OutputScopeName;
+	static const FName CustomScopeName;
+
 private:
 	static TArray<FNiagaraVariable> SystemParameters;
 	static TArray<FNiagaraVariable> TranslatorParameters;
diff --git a/Engine/Plugins/FX/Niagara/Source/Niagara/Private/NiagaraConstants.cpp b/Engine/Plugins/FX/Niagara/Source/Niagara/Private/NiagaraConstants.cpp
--- a/Engine/Plugins/FX/Niagara/Source/Niagara/Private/NiagaraConstants.cpp
+++ b/Engine/Plugins/FX/Niagara/Source/Niagara/Private/NiagaraConstants.cpp
@@ -23,6 +23,28 @@ TArray<FNiagaraVariable> FNiagaraConstants::EngineManagedAttributes;
 const FName FNiagaraConstants::InputPinName("InputPin");
 const FName FNiagaraConstants::OutputPinName("OutputPin");
 
+const FName FNiagaraConstants::UserNamespace(TEXT("User"));
+const FName FNiagaraConstants::EngineNamespace(TEXT("Engine"));
+const FName FNiagaraConstants::SystemNamespace(TEXT("System"));
+const FName FNiagaraConstants::EmitterNamespace(TEXT("Emitter"));
+const FName FNiagaraConstants::ParticleAttributeNamespace(TEXT("Particles"));
+const FName FNiagaraConstants::ModuleNamespace(TEXT("Module"));
+const FName FNiagaraConstants::ParameterCollectionNamespace(TEXT("NPC"));
+const FString FNiagaraConstants::InitialPrefix(TEXT("Initial"));
+const FName FNiagaraConstants::LocalNamespace(TEXT("Local"));
+const FName FNiagaraConstants::InitialNamespace(TEXT("Initial"));
+const FName FNiagaraConstants::OwnerNamespace(TEXT("Owner"));
+
+const FName FNiagaraConstants::EngineOwnerScopeName(TEXT("EngineOwner"));
+const FName FNiagaraConstants::EngineSystemScopeName(TEXT("EngineSystem"));
+const FName FNiagaraConstants::EngineEmitterScopeName(TEXT("EngineEmitter"));
+const FName FNiagaraConstants::CustomScopeName(TEXT("Custom"));
+
+const FName FNiagaraConstants::InputScopeName(TEXT("Input"));
+const FName FNiagaraConstants::OutputScopeName(TEXT("Output"));
+const FName FNiagaraConstants::ScriptTransientScopeName(TEXT("ScriptTransient"));
+const FName FNiagaraConstants::ScriptPersistentScopeName(TEXT("ScriptPersistent"));
+
 void FNiagaraConstants::Init()
 {
 	if (SystemParameters.Num() == 0)
diff --git a/Engine/Plugins/FX/Niagara/Source/Niagara/Private/NiagaraTypes.cpp b/Engine/Plugins/FX/Niagara/Source/Niagara/Private/NiagaraTypes.cpp
--- a/Engine/Plugins/FX/Niagara/Source/Niagara/Private/NiagaraTypes.cpp
+++ b/Engine/Plugins/FX/Niagara/Source/Niagara/Private/NiagaraTypes.cpp
@@ -43,21 +43,6 @@ FString FNiagaraTypeUtilities::GetNamespaceStringForScriptParameterScope(const E
 	};
 }
 
-FString FNiagaraVariableMetaData::GetNamespaceString() const
-{
-	checkf(bUseLegacyNameString == false, TEXT("Tried to get namespace string for parameter using legacy name string edit mode!"));
-	if (IsInputOrLocalUsage())
-	{
-		FString NamespaceString = FNiagaraTypeUtilities::GetNamespaceStringForScriptParameterScope(Scope);
-		if (Usage == ENiagaraScriptParameterUsage::InitialValueInput)
-		{
-			NamespaceString.Append(PARAM_MAP_INITIAL_STR);
-		}
-		return NamespaceString;
-	}
-	return PARAM_MAP_OUTPUT_STR;
-}
-
 bool FNiagaraVariableMetaData::GetParameterName(FName& OutName) const
 {
 	if (bUseLegacyNameString)
@@ -68,23 +53,10 @@ bool FNiagaraVariableMetaData::GetParameterName(FName& OutName) const
 	return true;
 }
 
-bool FNiagaraVariableMetaData::GetScope(ENiagaraParameterScope& OutScope) const
-{
-	if (bUseLegacyNameString)
-	{
-		OutScope = ENiagaraParameterScope::Custom;
-		return false;
-	}
-	OutScope = Scope;
-	return true;
-}
-
 void FNiagaraVariableMetaData::CopyPerScriptMetaData(const FNiagaraVariableMetaData& OtherMetaData)
 {
 	SetUsage(OtherMetaData.GetUsage());
-	ENiagaraParameterScope OtherMetaDataScope;
-	OtherMetaData.GetScope(OtherMetaDataScope);
-	SetScope(OtherMetaDataScope);
+	SetScopeName(OtherMetaData.GetScopeName());
 	FName OtherMetaDataName;
 	if (OtherMetaData.GetParameterName(OtherMetaDataName))
 	{
diff --git a/Engine/Plugins/FX/Niagara/Source/Niagara/Public/NiagaraTypes.h b/Engine/Plugins/FX/Niagara/Source/Niagara/Public/NiagaraTypes.h
--- a/Engine/Plugins/FX/Niagara/Source/Niagara/Public/NiagaraTypes.h
+++ b/Engine/Plugins/FX/Niagara/Source/Niagara/Public/NiagaraTypes.h
@@ -464,6 +464,7 @@ enum class ENiagaraParameterScope : uint32
 
 	Custom UMETA(Hidden), //Convenience markup for expressing parameters using legacy editor mode to freetype namespace and name.
 
+	DISPLAY_ONLY_StaticSwitch UMETA(DisplayName="Static Switch", Hidden), //Only use for display string in SEnumComboBoxes; does not have implementation for classes that interact with ENiagaraParameterScope.
 	// insert new scopes before
 	None UMETA(Hidden),
 
@@ -504,6 +505,38 @@ public:
 	TArray<FString> TargetValues;
 };
 
+USTRUCT()
+struct NIAGARA_API FNiagaraParameterScopeInfo
+{
+	GENERATED_BODY();
+
+public:
+	FNiagaraParameterScopeInfo()
+		: Scope(ENiagaraParameterScope::None)
+		, NamespaceString()
+	{};
+
+	FNiagaraParameterScopeInfo(const ENiagaraParameterScope InScope, const FString& InNamespaceString)
+		: Scope(InScope)
+		, NamespaceString(InNamespaceString)
+	{};
+
+	bool operator == (const FNiagaraParameterScopeInfo& Other) const
+	{
+		return Scope == Other.Scope && NamespaceString == Other.NamespaceString;
+	}
+
+	ENiagaraParameterScope GetScope() const { return Scope; };
+	const FString& GetNamespaceString() const { return NamespaceString; };
+
+private:
+	UPROPERTY()
+		ENiagaraParameterScope Scope;
+
+	UPROPERTY()
+		FString NamespaceString;
+};
+
 USTRUCT()
 struct NIAGARA_API FNiagaraVariableMetaData
 {
@@ -513,7 +546,7 @@ public:
 		: bAdvancedDisplay(false)
 		, EditorSortPriority(0)
 		, bInlineEditConditionToggle(false)
-		, Scope(ENiagaraParameterScope::None)
+		, ScopeName()
 		, Usage(ENiagaraScriptParameterUsage::None)
 		, bIsStaticSwitch(false)
 		, StaticSwitchDefaultValue(0)
@@ -553,14 +586,9 @@ public:
 	TMap<FName, FString> PropertyMetaData;
 
 public:
-	FString GetNamespaceString() const;
 
-	/** Gets the Scope and notifies if it does not apply due to an override being set.
-	 * @params OutScope		The Scope to return;
-	 * @return bool			Whether the returned scope is not overridden. Is false if bUseLegacyNameString is set.
-	 */
-	bool GetScope(ENiagaraParameterScope& OutScope) const;
-	void SetScope(ENiagaraParameterScope InScope) { Scope = InScope; };
+	const FName& GetScopeName() const { return ScopeName; };
+	void SetScopeName(const FName& InScopeName) { ScopeName = InScopeName; };
 
 	ENiagaraScriptParameterUsage GetUsage() const { return Usage; };
 	void SetUsage(ENiagaraScriptParameterUsage InUsage) { Usage = InUsage; };
@@ -593,9 +621,9 @@ public:
 	void CopyPerScriptMetaData(const FNiagaraVariableMetaData& OtherMetaData);
 
 private:
-	/** Defines the scope of a variable that is an input to a script. */
+	/** Defines the scope of a variable that is an input to a script. Used to lookup registered scope infos and resolve the actual ENiagaraParameterScope and Namespace string to use. */
 	UPROPERTY(meta = (SkipForCompileHash = "true"))
-	ENiagaraParameterScope Scope;
+	FName ScopeName;
 
 	/** Defines the usage of a variable as an argument or output relative to the script. */
 	UPROPERTY(meta = (SkipForCompileHash = "true"))
diff --git a/Engine/Plugins/FX/Niagara/Source/NiagaraEditor/Private/NiagaraEditorModule.cpp b/Engine/Plugins/FX/Niagara/Source/NiagaraEditor/Private/NiagaraEditorModule.cpp
--- a/Engine/Plugins/FX/Niagara/Source/NiagaraEditor/Private/NiagaraEditorModule.cpp
+++ b/Engine/Plugins/FX/Niagara/Source/NiagaraEditor/Private/NiagaraEditorModule.cpp
@@ -130,6 +130,7 @@
 #include "Interfaces/ITargetPlatform.h"
 #include "DeviceProfiles/DeviceProfile.h"
 #include "Containers/Ticker.h"
+#include "NiagaraConstants.h"
 
 IMPLEMENT_MODULE( FNiagaraEditorModule, NiagaraEditor );
 
@@ -139,6 +140,7 @@ PRAGMA_DISABLE_OPTIMIZATION
 
 const FName FNiagaraEditorModule::NiagaraEditorAppIdentifier( TEXT( "NiagaraEditorApp" ) );
 const FLinearColor FNiagaraEditorModule::WorldCentricTabColorScale(0.0f, 0.0f, 0.2f, 0.5f);
+TArray<TPair<FName, FNiagaraParameterScopeInfo>> FNiagaraEditorModule::RegisteredParameterScopeInfos;
 
 EAssetTypeCategories::Type FNiagaraEditorModule::NiagaraAssetCategory;
 
@@ -150,6 +152,16 @@ static FAutoConsoleVariableRef CVarShowFastPathOptions(
 	ECVF_Default
 );
 
+const FNiagaraParameterScopeInfo* FNiagaraEditorModule::FindParameterScopeInfo(const FName& ParameterScopeInfoName)
+{
+	auto FindPredicate = [ParameterScopeInfoName](const TPair<FName, FNiagaraParameterScopeInfo>& ScopeInfoPair) {return ScopeInfoPair.Key == ParameterScopeInfoName; };
+	if (const TPair<FName, FNiagaraParameterScopeInfo>* ScopeInfoPair = RegisteredParameterScopeInfos.FindByPredicate(FindPredicate))
+	{
+		return &ScopeInfoPair->Value;
+	}
+	return nullptr;
+}
+
 //////////////////////////////////////////////////////////////////////////
 
 class FNiagaraEditorOnlyDataUtilities : public INiagaraEditorOnlyDataUtilities
@@ -1009,6 +1021,20 @@ void FNiagaraEditorModule::StartupModule()
 		TEXT("Dumps data relevant to generating the compile id for an asset."),
 		FConsoleCommandWithArgsDelegate::CreateStatic(&DumpCompileIdDataForAsset));
 
+	RegisterParameterScopeInfo(FNiagaraConstants::UserNamespace, FNiagaraParameterScopeInfo(ENiagaraParameterScope::User, PARAM_MAP_USER_STR));
+	RegisterParameterScopeInfo(FNiagaraConstants::EngineNamespace, FNiagaraParameterScopeInfo(ENiagaraParameterScope::Engine, PARAM_MAP_ENGINE_STR));
+	RegisterParameterScopeInfo(FNiagaraConstants::EngineOwnerScopeName, FNiagaraParameterScopeInfo(ENiagaraParameterScope::Owner, PARAM_MAP_ENGINE_OWNER_STR));
+	RegisterParameterScopeInfo(FNiagaraConstants::EngineSystemScopeName, FNiagaraParameterScopeInfo(ENiagaraParameterScope::Engine, PARAM_MAP_ENGINE_SYSTEM_STR));
+	RegisterParameterScopeInfo(FNiagaraConstants::EngineEmitterScopeName, FNiagaraParameterScopeInfo(ENiagaraParameterScope::Engine, PARAM_MAP_ENGINE_EMITTER_STR));
+	RegisterParameterScopeInfo(FNiagaraConstants::SystemNamespace, FNiagaraParameterScopeInfo(ENiagaraParameterScope::System, PARAM_MAP_SYSTEM_STR));
+	RegisterParameterScopeInfo(FNiagaraConstants::EmitterNamespace, FNiagaraParameterScopeInfo(ENiagaraParameterScope::Emitter, PARAM_MAP_EMITTER_STR));
+	RegisterParameterScopeInfo(FNiagaraConstants::ParticleAttributeNamespace, FNiagaraParameterScopeInfo(ENiagaraParameterScope::Particles, PARAM_MAP_ATTRIBUTE_STR));
+	
+	RegisterParameterScopeInfo(FNiagaraConstants::InputScopeName, FNiagaraParameterScopeInfo(ENiagaraParameterScope::Input, PARAM_MAP_MODULE_STR));
+	RegisterParameterScopeInfo(FNiagaraConstants::LocalNamespace, FNiagaraParameterScopeInfo(ENiagaraParameterScope::Local, PARAM_MAP_LOCAL_MODULE_STR));
+	RegisterParameterScopeInfo(FNiagaraConstants::ScriptPersistentScopeName, FNiagaraParameterScopeInfo(ENiagaraParameterScope::ScriptPersistent, PARAM_MAP_SCRIPT_PERSISTENT_STR));
+	RegisterParameterScopeInfo(FNiagaraConstants::ScriptTransientScopeName, FNiagaraParameterScopeInfo(ENiagaraParameterScope::ScriptTransient, PARAM_MAP_SCRIPT_TRANSIENT_STR));
+	RegisterParameterScopeInfo(FNiagaraConstants::CustomScopeName, FNiagaraParameterScopeInfo(ENiagaraParameterScope::Custom, FString()));
 }
 
 
@@ -1327,6 +1353,11 @@ FNiagaraClipboard& FNiagaraEditorModule::GetClipboard() const
 	return Clipboard.Get();
 }
 
+void FNiagaraEditorModule::RegisterParameterScopeInfo(const FName& ParameterScopeInfoName, const FNiagaraParameterScopeInfo& ParameterScopeInfo)
+{
+	RegisteredParameterScopeInfos.AddUnique(TPair<FName, FNiagaraParameterScopeInfo>(ParameterScopeInfoName, ParameterScopeInfo));
+}
+
 void FNiagaraEditorModule::RegisterAssetTypeAction(IAssetTools& AssetTools, TSharedRef<IAssetTypeActions> Action)
 {
 	AssetTools.RegisterAssetTypeActions(Action);
diff --git a/Engine/Plugins/FX/Niagara/Source/NiagaraEditor/Private/NiagaraEditorUtilities.cpp b/Engine/Plugins/FX/Niagara/Source/NiagaraEditor/Private/NiagaraEditorUtilities.cpp
--- a/Engine/Plugins/FX/Niagara/Source/NiagaraEditor/Private/NiagaraEditorUtilities.cpp
+++ b/Engine/Plugins/FX/Niagara/Source/NiagaraEditor/Private/NiagaraEditorUtilities.cpp
@@ -2065,4 +2065,276 @@ FName FNiagaraEditorUtilities::GetUniqueObjectName(UObject* Outer, UClass* Objec
 	}
 }
 
+bool FNiagaraEditorUtilities::GetVariableMetaDataScope(const FNiagaraVariableMetaData& MetaData, ENiagaraParameterScope& OutScope)
+{
+	if (MetaData.GetIsUsingLegacyNameString())
+	{
+		OutScope = ENiagaraParameterScope::Custom;
+		return false;
+	}
+
+	const FName& MetaDataScopeName = MetaData.GetScopeName();
+	const FNiagaraParameterScopeInfo* ScopeInfo = FNiagaraEditorModule::FindParameterScopeInfo(MetaDataScopeName);
+	if (ensureMsgf(ScopeInfo != nullptr, TEXT("Failed to find registered parameter scope info for scope name %s!"), *MetaDataScopeName.ToString()))
+	{
+		OutScope = ScopeInfo->GetScope();
+		return true;
+	}
+	OutScope = ENiagaraParameterScope::Custom;
+	return false;
+}
+
+bool FNiagaraEditorUtilities::GetVariableMetaDataNamespaceString(const FNiagaraVariableMetaData& MetaData, FString& OutNamespaceString)
+{
+	if (MetaData.GetIsUsingLegacyNameString())
+	{
+		return false;
+	}
+
+	if (MetaData.IsInputOrLocalUsage())
+	{
+		const FName& MetaDataScopeName = MetaData.GetScopeName();
+		const FNiagaraParameterScopeInfo* ScopeInfo = FNiagaraEditorModule::FindParameterScopeInfo(MetaDataScopeName);
+		if (ScopeInfo == nullptr)
+		{
+			ensureMsgf(false, TEXT("Failed to find registered parameter scope info for scope name %s!"), *MetaDataScopeName.ToString());
+			return false;
+		}
+
+		FString NamespaceString = ScopeInfo->GetNamespaceString();
+		if (MetaData.GetUsage() == ENiagaraScriptParameterUsage::InitialValueInput)
+		{
+			NamespaceString.Append(PARAM_MAP_INITIAL_STR);
+		}
+		OutNamespaceString = NamespaceString;
+		return true;
+	}
+	OutNamespaceString = PARAM_MAP_OUTPUT_STR;
+	return true;
+}
+
+bool FNiagaraEditorUtilities::GetVariableMetaDataNamespaceStringForNewScope(const FNiagaraVariableMetaData& MetaData, const FName& NewScopeName, FString& OutNamespaceString)
+{
+	if (MetaData.GetIsUsingLegacyNameString())
+	{
+		return false;
+	}
+
+	const FNiagaraParameterScopeInfo* ScopeInfo = FNiagaraEditorModule::FindParameterScopeInfo(NewScopeName);
+	if (ScopeInfo == nullptr)
+	{
+		ensureMsgf(false, TEXT("Failed to find registered parameter scope info for scope name %s!"), *NewScopeName.ToString());
+		return false;
+	}
+
+	FString NamespaceString = ScopeInfo->GetNamespaceString();
+	if (MetaData.GetUsage() == ENiagaraScriptParameterUsage::InitialValueInput)
+	{
+		NamespaceString.Append(PARAM_MAP_INITIAL_STR);
+	}
+	OutNamespaceString = NamespaceString;
+
+	return true;
+}
+
+FName FNiagaraEditorUtilities::GetScopeNameForParameterScope(ENiagaraParameterScope InScope)
+{
+	switch (InScope) {
+	case ENiagaraParameterScope::User:
+		return FNiagaraConstants::UserNamespace;
+	case ENiagaraParameterScope::Engine:
+		return FNiagaraConstants::EngineNamespace;
+	case ENiagaraParameterScope::Owner:
+		return FNiagaraConstants::EngineOwnerScopeName;
+	case ENiagaraParameterScope::System:
+		return FNiagaraConstants::SystemNamespace;
+	case ENiagaraParameterScope::Emitter:
+		return FNiagaraConstants::EmitterNamespace;
+	case ENiagaraParameterScope::Particles:
+		return FNiagaraConstants::ParticleAttributeNamespace;
+	case ENiagaraParameterScope::Local:
+		return FNiagaraConstants::LocalNamespace;
+	case ENiagaraParameterScope::Input:
+		return FNiagaraConstants::InputScopeName;
+	case ENiagaraParameterScope::Custom:
+		return FNiagaraConstants::CustomScopeName;
+	case ENiagaraParameterScope::ScriptPersistent:
+		return FNiagaraConstants::ScriptPersistentScopeName;
+	case ENiagaraParameterScope::ScriptTransient:
+		return FNiagaraConstants::ScriptTransientScopeName;
+	default:
+		ensureMsgf(false, TEXT("Tried to get scope name for unknown parameter scope!"));
+	}
+	return FNiagaraConstants::ParticleAttributeNamespace;
+}
+
+
+TArray<FName> FNiagaraEditorUtilities::DecomposeVariableNamespace(const FName& InVarNameToken, FName& OutName)
+{
+	int32 DotIndex;
+	TArray<FName> OutNamespaces;
+	FString VarNameString = InVarNameToken.ToString();
+	while (VarNameString.FindChar(TEXT('.'), DotIndex))
+	{
+		OutNamespaces.Add(FName(*VarNameString.Left(DotIndex)));
+		VarNameString = *VarNameString.RightChop(DotIndex + 1);
+	}
+	OutName = FName(*VarNameString);
+	return OutNamespaces;
+}
+
+void FNiagaraEditorUtilities::GetParameterMetaDataFromName(const FName& InVarNameToken, FNiagaraVariableMetaData& OutMetaData)
+{
+	auto MarkAsLegacyCustomName = [&OutMetaData]() {
+		OutMetaData.SetScopeName(FNiagaraConstants::CustomScopeName);
+		OutMetaData.SetIsUsingLegacyNameString(true);
+	};
+
+	auto GetMetaDataForFirstNamespace = [MarkAsLegacyCustomName](const FName& Namespace, FNiagaraVariableMetaData& OutMetaData)-> bool /*bNextNamespaceCanBeValid*/ {
+		if (Namespace == FNiagaraConstants::LocalNamespace)
+		{
+			OutMetaData.SetScopeName(Namespace);
+			OutMetaData.SetUsage(ENiagaraScriptParameterUsage::Local);
+			return true;
+		}
+		else if (Namespace == FNiagaraConstants::ModuleNamespace)
+		{
+			OutMetaData.SetScopeName(FNiagaraConstants::InputScopeName);
+			return false;
+		}
+		else if (Namespace == FNiagaraConstants::UserNamespace)
+		{
+			OutMetaData.SetScopeName(Namespace);
+			return false;
+		}
+		else if (Namespace == FNiagaraConstants::EngineNamespace)
+		{
+			OutMetaData.SetScopeName(Namespace);
+			return true;
+		}
+		else if (Namespace == FNiagaraConstants::SystemNamespace)
+		{
+			OutMetaData.SetScopeName(Namespace);
+			return true;
+		}
+		else if (Namespace == FNiagaraConstants::EmitterNamespace)
+		{
+			OutMetaData.SetScopeName(Namespace);
+			return true;
+		}
+		else if (Namespace == FNiagaraConstants::ParticleAttributeNamespace)
+		{
+			OutMetaData.SetScopeName(Namespace);
+			return true;
+		}
+
+		MarkAsLegacyCustomName();
+		return false;
+	};
+
+	auto GetMetaDataForInitialNamespace = [](const FName& Namespace, FNiagaraVariableMetaData& OutMetaData)-> bool /*bFoundInitialNamespace*/ {
+		if (Namespace == FNiagaraConstants::InitialNamespace)
+		{
+			OutMetaData.SetUsage(ENiagaraScriptParameterUsage::InitialValueInput);
+			return true;
+		}
+		return false;
+	};
+
+	auto GetScopeCanHaveInitialNamespaceSuffix = [](ENiagaraParameterScope InScope)-> bool /*bCanHaveInitialNamespaceSuffix*/ {
+		switch (InScope) {
+		case ENiagaraParameterScope::System:
+		case ENiagaraParameterScope::Emitter:
+		case ENiagaraParameterScope::Particles:
+			return true;
+		default:
+			return false;
+		};
+		return false;
+	};
+
+	auto GetMetaDataForEngineSubNamespace = [](const FName& Namespace, FNiagaraVariableMetaData& OutMetaData)-> bool /*bFoundValidEngineSubNamespace*/ {
+		if (Namespace == FNiagaraConstants::OwnerNamespace)
+		{
+			OutMetaData.SetScopeName(FNiagaraConstants::EngineOwnerScopeName);
+			return true;
+		}
+		else if (Namespace == FNiagaraConstants::SystemNamespace)
+		{
+			OutMetaData.SetScopeName(FNiagaraConstants::EngineSystemScopeName);
+			return true;
+		}
+		else if (Namespace == FNiagaraConstants::EmitterNamespace)
+		{
+			OutMetaData.SetScopeName(FNiagaraConstants::EngineEmitterScopeName);
+			return true;
+		}
+		return false;
+	};
+
+	FName NamespacelessName;
+	TArray<FName> DecomposedNamespaces = FNiagaraEditorUtilities::DecomposeVariableNamespace(InVarNameToken, NamespacelessName);
+	OutMetaData.SetCachedNamespacelessVariableName(NamespacelessName);
+	if (DecomposedNamespaces.Num() == 0)
+	{
+		UE_LOG(LogNiagaraEditor, Display, TEXT("Unexpected parameter encountered without a namespace: %s"), *InVarNameToken.ToString());
+		MarkAsLegacyCustomName();
+		return;
+	}
+	else if (DecomposedNamespaces.Num() == 1)
+	{
+		GetMetaDataForFirstNamespace(DecomposedNamespaces[0], OutMetaData);
+		return;
+	}
+	else if (DecomposedNamespaces.Num() == 2)
+	{
+		bool bNextNamespaceCanBeValid = GetMetaDataForFirstNamespace(DecomposedNamespaces[0], OutMetaData);
+		if (bNextNamespaceCanBeValid)
+		{
+			ENiagaraParameterScope FirstNamespaceScope;
+			FNiagaraEditorUtilities::GetVariableMetaDataScope(OutMetaData, FirstNamespaceScope);
+			if (FirstNamespaceScope == ENiagaraParameterScope::Local)
+			{
+				// "local.module." namespaces may be handled as local scopes and do not need to be marked as legacy namespaces.
+				if (DecomposedNamespaces[1] == FNiagaraConstants::ModuleNamespace)
+				{
+					return;
+				}
+			}
+			else if (FirstNamespaceScope == ENiagaraParameterScope::Engine)
+			{
+				bool bFoundEngineSubNamespace = GetMetaDataForEngineSubNamespace(DecomposedNamespaces[1], OutMetaData);
+				if (bFoundEngineSubNamespace)
+				{
+					return;
+				}
+			}
+			else if (GetScopeCanHaveInitialNamespaceSuffix(FirstNamespaceScope))
+			{
+				bool bFoundInitialNamespace = GetMetaDataForInitialNamespace(DecomposedNamespaces[1], OutMetaData);
+				if (bFoundInitialNamespace)
+				{
+					return;
+				}
+			}
+		}
+		MarkAsLegacyCustomName();
+		return;
+	}
+
+	MarkAsLegacyCustomName();
+}
+
+FString FNiagaraEditorUtilities::GetNamespacelessVariableNameString(const FName& InVarName)
+{
+	int32 DotIndex;
+	FString VarNameString = InVarName.ToString();
+	if (VarNameString.FindLastChar(TEXT('.'), DotIndex))
+	{
+		return VarNameString.RightChop(DotIndex + 1);
+	}
+	// No dot index, must be a namespaceless variable name (e.g. static switch name) just return the name to string.
+	return VarNameString;
+}
+
 #undef LOCTEXT_NAMESPACE
diff --git a/Engine/Plugins/FX/Niagara/Source/NiagaraEditor/Private/NiagaraGraph.cpp b/Engine/Plugins/FX/Niagara/Source/NiagaraEditor/Private/NiagaraGraph.cpp
--- a/Engine/Plugins/FX/Niagara/Source/NiagaraEditor/Private/NiagaraGraph.cpp
+++ b/Engine/Plugins/FX/Niagara/Source/NiagaraEditor/Private/NiagaraGraph.cpp
@@ -369,7 +369,7 @@ void UNiagaraGraph::PostLoad()
 
 			// Resolve the namespace of the variable to a scope, cache the scopeless name in the UI, and set the usage if the name string would force that to occur.
 			FNiagaraVariableMetaData OutMetaData;
-			FNiagaraStackGraphUtilities::GetParameterMetaDataFromName(Var.GetName(), OutMetaData);
+			FNiagaraEditorUtilities::GetParameterMetaDataFromName(Var.GetName(), OutMetaData);
 			ScriptVar->Metadata.CopyPerScriptMetaData(OutMetaData);
 
 			// If the usage is not local or an initial value, select the usage based on the associated input/output pins.
@@ -1086,18 +1086,17 @@ void UNiagaraGraph::AddParameter(FNiagaraVariable& Parameter, const FAddParamete
 			NewScriptVariable->Metadata.SetUsage(Options.NewParameterUsage.GetValue());
 		}
 
-		if (Options.NewParameterScope.IsSet())
+		if (Options.NewParameterScopeName.IsSet())
 		{
 			bSkipRefreshMetaDataScopeAndUsage = true;
-
-			bool bIsInitialValue = false;
-			if (Options.NewParameterUsage.IsSet())
-			{
-				bIsInitialValue = Options.NewParameterUsage == ENiagaraScriptParameterUsage::InitialValueInput;
-			}
+			NewScriptVariable->Metadata.SetScopeName(Options.NewParameterScopeName.GetValue());
 			
-			const FName NewScriptVariableName = FNiagaraStackGraphUtilities::GetVariableNameForScope(NewScriptVariable->Variable.GetName(), Options.NewParameterScope.GetValue(), bIsInitialValue);
-			RenameParameter(NewScriptVariable->Variable, NewScriptVariableName, Options.bIsStaticSwitch, Options.NewParameterScope.GetValue());
+			FString NamespaceString;
+			checkf(FNiagaraEditorUtilities::GetVariableMetaDataNamespaceString(NewScriptVariable->Metadata, NamespaceString), TEXT("Tried to get a name for a new parameter set to use legacy custom name!"));
+			FString NameString = NamespaceString;
+			NameString.Append(FNiagaraEditorUtilities::GetNamespacelessVariableNameString(Parameter.GetName()));
+			const FName NewScriptVariableName = FName(*NameString);
+			RenameParameter(NewScriptVariable->Variable, NewScriptVariableName, Options.bIsStaticSwitch, Options.NewParameterScopeName.GetValue());
 		}
 
 		if (bSkipRefreshMetaDataScopeAndUsage == false && Options.bRefreshMetaDataScopeAndUsage)
@@ -1166,7 +1165,7 @@ void UNiagaraGraph::RemoveParameter(const FNiagaraVariable& Parameter)
 	}
 }
 
-bool UNiagaraGraph::RenameParameter(const FNiagaraVariable& Parameter, FName NewName, bool bFromStaticSwitch, ENiagaraParameterScope NewScope)
+bool UNiagaraGraph::RenameParameter(const FNiagaraVariable& Parameter, FName NewName, bool bFromStaticSwitch, FName NewScopeName)
 {
 	// Block rename when already renaming. This prevents recursion when CommitEditablePinName is called on referenced nodes. 
 	if (bIsRenamingParameter)
@@ -1193,10 +1192,10 @@ bool UNiagaraGraph::RenameParameter(const FNiagaraVariable& Parameter, FName New
 	}
 
 	// Update the OldMetaData which will be applied to the new parameter to synchronize the cached namespace-less name and scope.
-	OldMetaData.SetCachedNamespacelessVariableName(FName(*FNiagaraStackGraphUtilities::GetNamespacelessVariableNameString(NewParameter.GetName())));
-	if (NewScope != ENiagaraParameterScope::None)
+	OldMetaData.SetCachedNamespacelessVariableName(FName(*FNiagaraEditorUtilities::GetNamespacelessVariableNameString(NewParameter.GetName())));
+	if (NewScopeName.IsNone() == false)
 	{
-		OldMetaData.SetScope(NewScope);
+		OldMetaData.SetScopeName(NewScopeName);
 	}
 		
 	// Swap metadata to the new parameter; put the new parameter into VariableToScriptVariable
@@ -2172,7 +2171,7 @@ void UNiagaraGraph::InvalidateCachedParameterData()
 void UNiagaraGraph::GenerateMetaDataForScriptVariable(UNiagaraScriptVariable* InScriptVariable) const
 {
 	FNiagaraVariableMetaData OutMetaData;
-	FNiagaraStackGraphUtilities::GetParameterMetaDataFromName(InScriptVariable->Variable.GetName(), OutMetaData);
+	FNiagaraEditorUtilities::GetParameterMetaDataFromName(InScriptVariable->Variable.GetName(), OutMetaData);
 	InScriptVariable->Metadata.CopyPerScriptMetaData(OutMetaData);
 	UpdateUsageForScriptVariable(InScriptVariable);
 }
diff --git a/Engine/Plugins/FX/Niagara/Source/NiagaraEditor/Private/NiagaraNodeWithDynamicPins.cpp b/Engine/Plugins/FX/Niagara/Source/NiagaraEditor/Private/NiagaraNodeWithDynamicPins.cpp
--- a/Engine/Plugins/FX/Niagara/Source/NiagaraEditor/Private/NiagaraNodeWithDynamicPins.cpp
+++ b/Engine/Plugins/FX/Niagara/Source/NiagaraEditor/Private/NiagaraNodeWithDynamicPins.cpp
@@ -15,6 +15,7 @@
 #include "Framework/Application/SlateApplication.h"
 #include "NiagaraNodeParameterMapBase.h"
 #include "NiagaraScriptVariable.h"
+#include "NiagaraConstants.h"
 
 #define LOCTEXT_NAMESPACE "NiagaraNodeWithDynamicPins"
 
@@ -282,7 +283,7 @@ void UNiagaraNodeWithDynamicPins::AddParameter(FNiagaraVariable Parameter, const
 		AddParameterOptions.NewParameterUsage = ENiagaraScriptParameterUsage::Input;
 	}
 	// Default scope of new parameters on a pin to Particles, by convention. //@todo(ng) choose based on usage
-	AddParameterOptions.NewParameterScope = ENiagaraParameterScope::Particles;
+	AddParameterOptions.NewParameterScopeName = FNiagaraConstants::ParticleAttributeNamespace;
 	//AddParameterOptions.bMakeParameterNameUnique //@todo(ng) consider
 	
 	UNiagaraGraph* Graph = GetNiagaraGraph();
diff --git a/Engine/Plugins/FX/Niagara/Source/NiagaraEditor/Private/NiagaraParameterMapHistory.cpp b/Engine/Plugins/FX/Niagara/Source/NiagaraEditor/Private/NiagaraParameterMapHistory.cpp
--- a/Engine/Plugins/FX/Niagara/Source/NiagaraEditor/Private/NiagaraParameterMapHistory.cpp
+++ b/Engine/Plugins/FX/Niagara/Source/NiagaraEditor/Private/NiagaraParameterMapHistory.cpp
@@ -1482,8 +1482,7 @@ void FNiagaraParameterMapHistoryWithMetaDataBuilder::FixupHistoryVariableNamespa
 			auto SortVars = [&InHandle, &VariablesToSkip](bool bGatherInputs)->TArray<FNiagaraHistoryVariable> {
 				TArray<FNiagaraHistoryVariable> OutHistoryVars = InHandle.HistoryVariables.FilterByPredicate([&VariablesToSkip, bGatherInputs](const FNiagaraHistoryVariable& HistoryVar)
 				{
-					ENiagaraParameterScope HistoryVarScope;
-					checkf(HistoryVar.VarMetaData->GetScope(HistoryVarScope), TEXT("Tried to process UNiagaraScriptVariable with override namespace set during precompiler namespace fixup!"));
+					const ENiagaraParameterScope HistoryVarScope = HistoryVar.StaticVarMetaData.GetScope();
 					if (HistoryVarScope == ENiagaraParameterScope::User || HistoryVarScope == ENiagaraParameterScope::Engine || HistoryVarScope == ENiagaraParameterScope::ScriptTransient || HistoryVarScope == ENiagaraParameterScope::None)
 					{
 						return false;
@@ -1749,8 +1748,7 @@ void FNiagaraParameterMapHistoryWithMetaDataBuilder::FixupHistoryVariableNamespa
 					auto FindUpstreamVarMatchIdx = [&Handle, &MatchIdx, &CurrentHistoryVar, &CurrentInputVarHandle, &UpstreamSortedHistoryHandle]()->bool {
 						const ENiagaraParameterScope UpstreamScope = UpstreamSortedHistoryHandle->GetScriptScope();
 						ENiagaraParameterScope TargetScope = ENiagaraParameterScope::None;
-						ENiagaraParameterScope CurrentHistoryVarScope;
-						checkf(CurrentHistoryVar.VarMetaData->GetScope(CurrentHistoryVarScope), TEXT("Tried to process UNiagaraScriptVariable with override namespace set during precompiler namespace fixup!"));
+						const ENiagaraParameterScope CurrentHistoryVarScope = CurrentHistoryVar.StaticVarMetaData.GetScope();
 						if (CurrentHistoryVarScope == ENiagaraParameterScope::ScriptPersistent)
 						{
 							// If the current history var is targeting the script alias scope then get the scope of the current history.
@@ -1859,8 +1857,7 @@ void FNiagaraParameterMapHistoryWithMetaDataBuilder::FixupHistoryVariableNamespa
 		{
 			for (FNiagaraHistoryVariable& InputHistoryVar : Handle.InputHistoryVars)
 			{
-				ENiagaraParameterScope InputHistoryVarScope;
-				checkf(InputHistoryVar.VarMetaData->GetScope(InputHistoryVarScope), TEXT("Tried to process UNiagaraScriptVariable with override namespace set during precompiler namespace fixup!"));
+				const ENiagaraParameterScope InputHistoryVarScope = InputHistoryVar.StaticVarMetaData.GetScope();
 				if (InputHistoryVarScope == ENiagaraParameterScope::None)
 				{
 					// Only need to fixup variables that are already in the dataset.
@@ -1881,8 +1878,7 @@ void FNiagaraParameterMapHistoryWithMetaDataBuilder::FixupHistoryVariableNamespa
 
 			for (FNiagaraHistoryVariable& OutputHistoryVar : Handle.OutputHistoryVars)
 			{
-				ENiagaraParameterScope OutputHistoryVarScope;
-				checkf(OutputHistoryVar.VarMetaData->GetScope(OutputHistoryVarScope), TEXT("Tried to process UNiagaraScriptVariable with override namespace set during precompiler namespace fixup!"));
+				const ENiagaraParameterScope OutputHistoryVarScope = OutputHistoryVar.StaticVarMetaData.GetScope();
 				if (OutputHistoryVar.bDoNotDemote)
 				{
 					continue;
@@ -1956,8 +1952,8 @@ FNiagaraParameterMapHistoryHandle::FNiagaraParameterMapHistoryHandle(FNiagaraPar
 	{
 		FNiagaraVariable* Var = &InHistory.Variables[VarIdx];
 		FNiagaraVariable* VarWithOriginalAliasIntact = &InHistory.VariablesWithOriginalAliasesIntact[VarIdx];
-		FNiagaraVariableMetaData* VarMetaData = &InHistory.VariableMetaData[VarIdx];
-		if (VarMetaData->GetIsUsingLegacyNameString() == true)
+		FNiagaraVariableMetaData& VarMetaData = InHistory.VariableMetaData[VarIdx];
+		if (VarMetaData.GetIsUsingLegacyNameString() == true)
 		{
 			// Cannot fixup names of variables using legacy name string mode, skip this entry.
 			continue;
@@ -1987,8 +1983,8 @@ FNiagaraParameterMapHistoryHandle::FNiagaraParameterMapHistoryHandle(FNiagaraPar
 			FNiagaraVariable* PrimaryVarWithOriginalAliasIntact = &InPrimaryHistory.VariablesWithOriginalAliasesIntact[FoundPrimaryVarIdx];
 
 			// Assign the secondary history metadata as it is canonical in regards to parameter linking.
-			FNiagaraVariableMetaData* ScriptVarMetaData = &InSecondaryHistory.VariableMetaData[SecondaryVarIdx];
-			if (ScriptVarMetaData->GetIsUsingLegacyNameString() == true)
+			FNiagaraVariableMetaData& ScriptVarMetaData = InSecondaryHistory.VariableMetaData[SecondaryVarIdx];
+			if (ScriptVarMetaData.GetIsUsingLegacyNameString() == true)
 			{
 				// Cannot fixup names of variables using legacy name string mode, skip this entry.
 				continue;
@@ -2037,3 +2033,9 @@ void FNiagaraHistoryVariable::SetVarWithOriginalAliasName(const FName& NewName)
 	FNiagaraVariable* NonConstVarWithOriginalAlias = const_cast<FNiagaraVariable*>(VarWithOriginalAlias);
 	NonConstVarWithOriginalAlias->SetName(NewName);
 }
+
+FNiagaraStaticVariableMetaData::FNiagaraStaticVariableMetaData(const FNiagaraVariableMetaData& ViewMetaData)
+{
+	ensureMsgf(FNiagaraEditorUtilities::GetVariableMetaDataScope(ViewMetaData, Scope), TEXT("Tried to get scope for parameter with override name enabled or invalid scope!"));
+	FNiagaraEditorUtilities::GetVariableMetaDataNamespaceString(ViewMetaData, NamespaceString);
+}
diff --git a/Engine/Plugins/FX/Niagara/Source/NiagaraEditor/Private/NiagaraScriptMergeManager.cpp b/Engine/Plugins/FX/Niagara/Source/NiagaraEditor/Private/NiagaraScriptMergeManager.cpp
--- a/Engine/Plugins/FX/Niagara/Source/NiagaraEditor/Private/NiagaraScriptMergeManager.cpp
+++ b/Engine/Plugins/FX/Niagara/Source/NiagaraEditor/Private/NiagaraScriptMergeManager.cpp
@@ -27,6 +27,7 @@
 #include "EdGraph/EdGraphNode.h"
 #include "EdGraph/EdGraphPin.h"
 #include "Modules/ModuleManager.h"
+#include "NiagaraConstants.h"
 
 #define LOCTEXT_NAMESPACE "NiagaraScriptMergeManager"
 
@@ -2019,7 +2020,7 @@ FNiagaraScriptMergeManager::FApplyDiffResults FNiagaraScriptMergeManager::AddInp
 	UNiagaraNodeAssignment* AssignmentNode = Cast<UNiagaraNodeAssignment>(&TargetFunctionCall);
 	if (AssignmentNode)
 	{
-		FNiagaraParameterHandle FunctionInputHandle(FNiagaraParameterHandle::ModuleNamespace, *OverrideToAdd->GetInputName());
+		FNiagaraParameterHandle FunctionInputHandle(FNiagaraConstants::ModuleNamespace, *OverrideToAdd->GetInputName());
 		UNiagaraNodeAssignment* PreviousVersionAssignmentNode = Cast<UNiagaraNodeAssignment>(OverrideToAdd->GetOwningFunctionCall());
 		bool bAnyAdded = false;
 		for (int32 i = 0; i < PreviousVersionAssignmentNode->NumTargets(); i++)
@@ -2040,7 +2041,7 @@ FNiagaraScriptMergeManager::FApplyDiffResults FNiagaraScriptMergeManager::AddInp
 		}
 	}
 
-	FNiagaraParameterHandle FunctionInputHandle(FNiagaraParameterHandle::ModuleNamespace, *OverrideToAdd->GetInputName());
+	FNiagaraParameterHandle FunctionInputHandle(FNiagaraConstants::ModuleNamespace, *OverrideToAdd->GetInputName());
 	FNiagaraParameterHandle AliasedFunctionInputHandle = FNiagaraParameterHandle::CreateAliasedModuleParameterHandle(FunctionInputHandle, &TargetFunctionCall);
 
 	if (OverrideToAdd->GetOverridePin() != nullptr)
diff --git a/Engine/Plugins/FX/Niagara/Source/NiagaraEditor/Private/ViewModels/NiagaraParameterNameViewModel.cpp b/Engine/Plugins/FX/Niagara/Source/NiagaraEditor/Private/ViewModels/NiagaraParameterNameViewModel.cpp
--- a/Engine/Plugins/FX/Niagara/Source/NiagaraEditor/Private/ViewModels/NiagaraParameterNameViewModel.cpp
+++ b/Engine/Plugins/FX/Niagara/Source/NiagaraEditor/Private/ViewModels/NiagaraParameterNameViewModel.cpp
@@ -31,26 +31,26 @@ TSharedRef<SWidget> FNiagaraParameterPanelEntryParameterNameViewModel::CreateSco
 		return SNullWidget::NullWidget;
 	}
 
-	ENiagaraParameterScope CachedScope;
-	CachedScriptVarAndViewInfo.MetaData.GetScope(CachedScope);
-	UEnum* ParameterScopeEnum = FNiagaraTypeDefinition::GetParameterScopeEnum();
-	TSharedPtr< SWidget > ScopeComboBoxWidget;
-	FNiagaraEditorModule& NiagaraEditorModule = FModuleManager::LoadModuleChecked<FNiagaraEditorModule>("NiagaraEditor");
-
 	bool bEnableScopeSlotWidget = true;
-	if (CachedScriptVarAndViewInfo.MetaData.GetUsage() != ENiagaraScriptParameterUsage::Input)
-	{
-		bEnableScopeSlotWidget = false;
-	}
-	else if (CachedScope == ENiagaraParameterScope::Local)
+	if (CachedScriptVarAndViewInfo.MetaData.GetIsStaticSwitch() == false)
 	{
-		bEnableScopeSlotWidget = false;
-	}
-	else if (CachedScriptVarAndViewInfo.MetaData.GetIsStaticSwitch())
-	{
-		bEnableScopeSlotWidget = false;
+		if (CachedScriptVarAndViewInfo.MetaData.GetUsage() != ENiagaraScriptParameterUsage::Input)
+		{
+			bEnableScopeSlotWidget = false;
+		}
+
+		ENiagaraParameterScope CachedScope;
+		FNiagaraEditorUtilities::GetVariableMetaDataScope(CachedScriptVarAndViewInfo.MetaData, CachedScope);
+		if (CachedScope == ENiagaraParameterScope::Local)
+		{
+			bEnableScopeSlotWidget = false;
+		}
 	}
 
+	UEnum* ParameterScopeEnum = FNiagaraTypeDefinition::GetParameterScopeEnum();
+	TSharedPtr< SWidget > ScopeComboBoxWidget;
+	FNiagaraEditorModule& NiagaraEditorModule = FModuleManager::LoadModuleChecked<FNiagaraEditorModule>("NiagaraEditor");
+
 	SAssignNew(ScopeComboBoxWidget, SBox)
 		.MinDesiredWidth(80.0f) //@todo(ng) tune and make const static
 		.Content()
@@ -89,8 +89,12 @@ TSharedRef<SWidget> FNiagaraParameterPanelEntryParameterNameViewModel::CreateTex
 
 int32 FNiagaraParameterPanelEntryParameterNameViewModel::GetScopeValue() const
 {
+	if (CachedScriptVarAndViewInfo.MetaData.GetIsStaticSwitch())
+	{
+		return int32(ENiagaraParameterScope::DISPLAY_ONLY_StaticSwitch);
+	}
 	ENiagaraParameterScope CachedScope;
-	if (ensureMsgf(CachedScriptVarAndViewInfo.MetaData.GetScope(CachedScope), TEXT("Failed to get scope value for param as override namespace is set! This method should not be bound!")))
+	if (ensureMsgf(FNiagaraEditorUtilities::GetVariableMetaDataScope(CachedScriptVarAndViewInfo.MetaData, CachedScope), TEXT("Failed to get scope value for param as override namespace is set! This method should not be bound!")))
 	{
 		return (int32)CachedScope;
 	}
@@ -191,7 +195,7 @@ TSharedRef<SWidget> FNiagaraGraphPinParameterNameViewModel::CreateTextSlotWidget
 int32 FNiagaraGraphPinParameterNameViewModel::GetScopeValue() const
 {
 	ENiagaraParameterScope CachedScope;
-	if (ensureMsgf(CachedScriptVarAndViewInfo.MetaData.GetScope(CachedScope), TEXT("Failed to get scope value for param as override namespace is set! This method should not be bound!")))
+	if (ensureMsgf(FNiagaraEditorUtilities::GetVariableMetaDataScope(CachedScriptVarAndViewInfo.MetaData, CachedScope), TEXT("Failed to get scope value for param as override namespace is set! This method should not be bound!")))
 	{
 		return (int32)CachedScope;
 	}
diff --git a/Engine/Plugins/FX/Niagara/Source/NiagaraEditor/Private/ViewModels/NiagaraParameterPanelViewModel.cpp b/Engine/Plugins/FX/Niagara/Source/NiagaraEditor/Private/ViewModels/NiagaraParameterPanelViewModel.cpp
--- a/Engine/Plugins/FX/Niagara/Source/NiagaraEditor/Private/ViewModels/NiagaraParameterPanelViewModel.cpp
+++ b/Engine/Plugins/FX/Niagara/Source/NiagaraEditor/Private/ViewModels/NiagaraParameterPanelViewModel.cpp
@@ -46,6 +46,9 @@ FText NiagaraParameterPanelSectionID::OnGetSectionTitle(const NiagaraParameterPa
 	case NiagaraParameterPanelSectionID::EMITTER:
 		SectionTitle = LOCTEXT("EmitterParamSection", "Emitter");
 		break;
+	case NiagaraParameterPanelSectionID::OWNER:
+		SectionTitle = LOCTEXT("OwnerParamSection", "Owner");
+		break;
 	case NiagaraParameterPanelSectionID::PARTICLES:
 		SectionTitle = LOCTEXT("ParticlesParamSection", "Particles");
 		break;
@@ -99,6 +102,7 @@ const NiagaraParameterPanelSectionID::Type NiagaraParameterPanelSectionID::GetSe
 	case ENiagaraParameterScope::ScriptPersistent:
 	case ENiagaraParameterScope::ScriptTransient:
 		// This is a potential situation if a script alias param has not had its scope cached by compiling.
+		ensureMsgf(false, TEXT("Tried to get section ID for script transient or persistent scope parameter!"));
 		return NiagaraParameterPanelSectionID::NONE;
 	default:
 		ensureMsgf(false, TEXT("Failed to find matching section ID for script parameter scope!"));
@@ -108,8 +112,13 @@ const NiagaraParameterPanelSectionID::Type NiagaraParameterPanelSectionID::GetSe
 
 const NiagaraParameterPanelSectionID::Type NiagaraParameterPanelSectionID::GetSectionForParameterMetaData(const FNiagaraVariableMetaData& MetaData)
 {
+	if (MetaData.GetIsStaticSwitch())
+	{
+		return NiagaraParameterPanelSectionID::INPUTS;
+	}
+
 	ENiagaraParameterScope MetaDataScope;
-	MetaData.GetScope(MetaDataScope);
+	FNiagaraEditorUtilities::GetVariableMetaDataScope(MetaData, MetaDataScope);
 	ENiagaraScriptParameterUsage MetaDataUsage = MetaData.GetUsage();
 
 	if (MetaDataUsage == ENiagaraScriptParameterUsage::Local)
@@ -131,10 +140,6 @@ const NiagaraParameterPanelSectionID::Type NiagaraParameterPanelSectionID::GetSe
 			return NiagaraParameterPanelSectionID::REFERENCES;
 		}
 	}
-	else if (MetaData.GetIsStaticSwitch())
-	{
-		return NiagaraParameterPanelSectionID::INPUTS;
-	}
 	else if (MetaDataUsage == ENiagaraScriptParameterUsage::Output)
 	{
 		return NiagaraParameterPanelSectionID::OUTPUTS;
@@ -170,9 +175,10 @@ ENiagaraParameterScope NiagaraParameterPanelSectionID::GetScopeForNewParametersI
 		return ENiagaraParameterScope::Local;
 	case NiagaraParameterPanelSectionID::INITIALVALUES:
 		return ENiagaraParameterScope::Particles;
+	case NiagaraParameterPanelSectionID::INPUTS:
+		return ENiagaraParameterScope::Input;
 
 	// Default to Particles scope if coming from section IDs that are not directly associated with scope.
-	case NiagaraParameterPanelSectionID::INPUTS:
 	case NiagaraParameterPanelSectionID::REFERENCES:
 	case NiagaraParameterPanelSectionID::OUTPUTS:
 		return ENiagaraParameterScope::Particles;
@@ -269,6 +275,7 @@ void FNiagaraSystemToolkitParameterPanelViewModel::CollectStaticSections(TArray<
 	// Generic Emitter/System view, categorize by scope
 	StaticSectionIDs.Add(NiagaraParameterPanelSectionID::USER);
 	StaticSectionIDs.Add(NiagaraParameterPanelSectionID::ENGINE);
+	StaticSectionIDs.Add(NiagaraParameterPanelSectionID::OWNER);
 	StaticSectionIDs.Add(NiagaraParameterPanelSectionID::SYSTEM);
 	StaticSectionIDs.Add(NiagaraParameterPanelSectionID::EMITTER);
 	StaticSectionIDs.Add(NiagaraParameterPanelSectionID::PARTICLES);
@@ -279,7 +286,7 @@ NiagaraParameterPanelSectionID::Type FNiagaraSystemToolkitParameterPanelViewMode
 	// Generic Emitter/System view, categorize by scope
 
 	ENiagaraParameterScope ParameterScope = ENiagaraParameterScope::None;
-	VarAndViewInfo.MetaData.GetScope(ParameterScope);
+	FNiagaraEditorUtilities::GetVariableMetaDataScope(VarAndViewInfo.MetaData, ParameterScope);
 	if (VarAndViewInfo.MetaData.GetUsage() == ENiagaraScriptParameterUsage::Input && ParameterScope != ENiagaraParameterScope::Input)
 	{
 		return NiagaraParameterPanelSectionID::GetSectionForScope(ParameterScope);
@@ -288,9 +295,9 @@ NiagaraParameterPanelSectionID::Type FNiagaraSystemToolkitParameterPanelViewMode
 	{
 		//return NiagaraParameterPanelSectionID::GetSectionForScope(VarAndViewInfo.ScriptVariableMetaData.LastPrecompileScope); //@todo(ng) cache the known scope in the niagarasystem
 		FNiagaraVariableMetaData OutMetaData;
-		FNiagaraStackGraphUtilities::GetParameterMetaDataFromName(VarAndViewInfo.ScriptVariable.GetName(), OutMetaData);
+		FNiagaraEditorUtilities::GetParameterMetaDataFromName(VarAndViewInfo.ScriptVariable.GetName(), OutMetaData);
 		ENiagaraParameterScope OutMetaDataScope;
-		OutMetaData.GetScope(OutMetaDataScope);
+		FNiagaraEditorUtilities::GetVariableMetaDataScope(OutMetaData, OutMetaDataScope);
 		return NiagaraParameterPanelSectionID::GetSectionForScope(OutMetaDataScope);
 	}
 }
@@ -300,15 +307,14 @@ void FNiagaraSystemToolkitParameterPanelViewModel::AddParameter(const FNiagaraVa
 	FScopedTransaction AddParameter(LOCTEXT("AddParameter", "Add Parameter"));
 	bool bSystemIsSelected = OverviewSelectionViewModel->GetSystemIsSelected();
 
-	ENiagaraParameterScope NewScope = ENiagaraParameterScope::None;
-	InVariableMetaDataToAssign.GetScope(NewScope);
-
 	UNiagaraGraph::FAddParameterOptions AddParameterOptions = UNiagaraGraph::FAddParameterOptions();
-	AddParameterOptions.NewParameterScope = NewScope;
+	AddParameterOptions.NewParameterScopeName = InVariableMetaDataToAssign.GetScopeName();
 	AddParameterOptions.NewParameterUsage = InVariableMetaDataToAssign.GetUsage();
 	AddParameterOptions.bAddedFromSystemEditor = true;
 	FNiagaraVariable DuplicateVar = InVariableToAdd; //@todo(ng) rewrite
 
+	ENiagaraParameterScope NewScope;
+	FNiagaraEditorUtilities::GetVariableMetaDataScope(InVariableMetaDataToAssign, NewScope);
 	if (NewScope == ENiagaraParameterScope::User)
 	{
 		SystemViewModel->GetSystem().Modify();
@@ -345,7 +351,7 @@ void FNiagaraSystemToolkitParameterPanelViewModel::RemoveParameter(const FNiagar
 	FScopedTransaction RemoveParameter(LOCTEXT("RemoveParameter", "Remove Parameter"));
 
 	ENiagaraParameterScope TargetVariableScope = ENiagaraParameterScope::None;
-	TargetVariableMetaData.GetScope(TargetVariableScope);
+	FNiagaraEditorUtilities::GetVariableMetaDataScope(TargetVariableMetaData, TargetVariableScope);
 
 	if (TargetVariableMetaData.IsInputUsage() && TargetVariableScope == ENiagaraParameterScope::User) //@todo(ng) verify
 	{
@@ -376,12 +382,11 @@ void FNiagaraSystemToolkitParameterPanelViewModel::RenameParameter(const FNiagar
 	FScopedTransaction RenameParameter(LOCTEXT("RenameParameter", "Rename Parameter"));
 
 	FName NewVariableName;
-	ENiagaraParameterScope TargetVariableScope;
-	FString ExistingVariableScopeString;
-	if (TargetVariableMetaData.GetScope(TargetVariableScope))
+	FString TargetVariableNamespaceString;
+	
+	if (FNiagaraEditorUtilities::GetVariableMetaDataNamespaceString(TargetVariableMetaData, TargetVariableNamespaceString))
 	{
-		ExistingVariableScopeString = FNiagaraTypeUtilities::GetNamespaceStringForScriptParameterScope(TargetVariableScope);
-		NewVariableName = FName(*(ExistingVariableScopeString + NewVariableNameText.ToString()));
+		NewVariableName = FName(*(TargetVariableNamespaceString + NewVariableNameText.ToString()));
 	}
 	else
 	{
@@ -393,7 +398,7 @@ void FNiagaraSystemToolkitParameterPanelViewModel::RenameParameter(const FNiagar
 		if (ensureMsgf(Graph.IsValid(), TEXT("Editable Emitter Script Graph was stale when renaming parameter!")))
 		{
 			Graph->Modify();
-			Graph->RenameParameter(TargetVariableToRename, NewVariableName, TargetVariableMetaData.GetIsStaticSwitch(), TargetVariableScope); //@todo(ng) handle renaming system params
+			Graph->RenameParameter(TargetVariableToRename, NewVariableName, TargetVariableMetaData.GetIsStaticSwitch(), TargetVariableMetaData.GetScopeName()); //@todo(ng) handle renaming system params
 		}
 	}
 	Refresh();
@@ -460,7 +465,7 @@ const TArray<FNiagaraScriptVariableAndViewInfo> FNiagaraSystemToolkitParameterPa
 			{
 				const FNiagaraVariableMetaData& MetaData = Iter.Value()->Metadata;
 				ENiagaraParameterScope MetaDataScope;
-				MetaData.GetScope(MetaDataScope);
+				FNiagaraEditorUtilities::GetVariableMetaDataScope(MetaData, MetaDataScope);
 
 				if (MetaDataScope == ENiagaraParameterScope::None || MetaData.GetUsage() == ENiagaraScriptParameterUsage::None)
 				{
@@ -580,17 +585,14 @@ void FNiagaraScriptToolkitParameterPanelViewModel::AddParameter(const FNiagaraVa
 	FScopedTransaction AddParameter(LOCTEXT("AddParameterFromParameterPanel", "Add Parameter"));
 	UNiagaraGraph::FAddParameterOptions AddParameterOptions = UNiagaraGraph::FAddParameterOptions();
 
-	ENiagaraParameterScope InVariableScope;
-	InVariableMetaDataToAssign.GetScope(InVariableScope);
-	ENiagaraScriptParameterUsage InVariableUsage = InVariableMetaDataToAssign.GetUsage();
-
-	AddParameterOptions.NewParameterScope = InVariableScope;
+	AddParameterOptions.NewParameterScopeName = InVariableMetaDataToAssign.GetScopeName();
 	AddParameterOptions.NewParameterUsage = InVariableMetaDataToAssign.GetUsage();
 	FNiagaraVariable DuplicateVar = InVariableToAdd; //@todo(ng) rewrite
 
 	UNiagaraGraph* Graph = ScriptViewModel->GetGraphViewModel()->GetGraph();
 	Graph->Modify();
 	const FVector2D NewNodePos = Graph->GetGoodPlaceForNewNode();
+	ENiagaraScriptParameterUsage InVariableUsage = InVariableMetaDataToAssign.GetUsage();
 	if (InVariableUsage == ENiagaraScriptParameterUsage::Input)
 	{
 		UNiagaraNodeParameterMapGet* NewMapGet = FNiagaraSchemaAction_NewNode::SpawnNodeFromTemplate<UNiagaraNodeParameterMapGet>(Graph, NewObject<UNiagaraNodeParameterMapGet>(), NewNodePos);
@@ -637,17 +639,16 @@ void FNiagaraScriptToolkitParameterPanelViewModel::RenameParameter(const FNiagar
 	}
 	else
 	{
-		const FString ExistingVariableScopeString = TargetVariableMetaData.GetNamespaceString();
-		NewVariableName = FName(*(ExistingVariableScopeString + NewVariableNameText.ToString()));
+		FString TargetVariableNamespaceString;
+		checkf(FNiagaraEditorUtilities::GetVariableMetaDataNamespaceString(TargetVariableMetaData, TargetVariableNamespaceString), TEXT("Tried to get namespace string for parameter using legacy name string edit mode!"));
+		NewVariableName = FName(*(TargetVariableNamespaceString + NewVariableNameText.ToString()));
 	}
 
 	
 	UNiagaraGraph* Graph = ScriptViewModel->GetGraphViewModel()->GetGraph();
 	Graph->Modify();
 
-	ENiagaraParameterScope TargetVariableScope = ENiagaraParameterScope::None;
-	TargetVariableMetaData.GetScope(TargetVariableScope);
-	Graph->RenameParameter(TargetVariableToRename, NewVariableName, TargetVariableMetaData.GetIsStaticSwitch(), TargetVariableScope);
+	Graph->RenameParameter(TargetVariableToRename, NewVariableName, TargetVariableMetaData.GetIsStaticSwitch(), TargetVariableMetaData.GetScopeName());
 }
 
 void FNiagaraScriptToolkitParameterPanelViewModel::ChangeParameterScope(const FNiagaraVariable& TargetVariableToModify, const FNiagaraVariableMetaData& TargetVariableMetaData, const ENiagaraParameterScope NewVariableScope) const
@@ -656,11 +657,18 @@ void FNiagaraScriptToolkitParameterPanelViewModel::ChangeParameterScope(const FN
 	{
 		FScopedTransaction ChangeParameterScopeAndReferencedPins(LOCTEXT("ChangeParameterScopeAndReferencedPins", "Change parameter scope, Rename parameter and referenced pins"));
 
-		bool bIsInitialValue = TargetVariableMetaData.GetUsage() == ENiagaraScriptParameterUsage::InitialValueInput;
-		const FName NewVariableName = FNiagaraStackGraphUtilities::GetVariableNameForScope(TargetVariableToModify.GetName(), NewVariableScope, bIsInitialValue);
-		UNiagaraGraph* Graph = ScriptViewModel->GetGraphViewModel()->GetGraph();
-		Graph->Modify();
-		Graph->RenameParameter(TargetVariableToModify, NewVariableName, TargetVariableMetaData.GetIsStaticSwitch(), NewVariableScope);
+		const FName TargetScopeName = FNiagaraEditorUtilities::GetScopeNameForParameterScope(NewVariableScope);
+		FString NewNamespaceString;
+		if (ensureMsgf(FNiagaraEditorUtilities::GetVariableMetaDataNamespaceStringForNewScope(TargetVariableMetaData, TargetScopeName, NewNamespaceString), TEXT("Tried to change scope of parameter with override name mode enabled!")))
+		{
+			FString NewNameString = NewNamespaceString;
+			NewNamespaceString.Append(FNiagaraEditorUtilities::GetNamespacelessVariableNameString(TargetVariableToModify.GetName()));
+			const FName NewVariableHLSLTokenName = FName(*NewNamespaceString);
+
+			UNiagaraGraph* Graph = ScriptViewModel->GetGraphViewModel()->GetGraph();
+			Graph->Modify();
+			Graph->RenameParameter(TargetVariableToModify, NewVariableHLSLTokenName, TargetVariableMetaData.GetIsStaticSwitch(), TargetScopeName);
+		}
 	}
 }
 
@@ -766,7 +774,7 @@ TStaticArray<FScopeIsEnabledAndTooltip, (int32)ENiagaraParameterScope::Num> FNia
 				if (InName == CachedParameterName)
 				{
 					ENiagaraParameterScope ViewedParameterScope;
-					ViewedParameter.MetaData.GetScope(ViewedParameterScope);
+					FNiagaraEditorUtilities::GetVariableMetaDataScope(ViewedParameter.MetaData, ViewedParameterScope);
 
 					PerScopeInfo[(int32)ViewedParameterScope].bEnabled = false;
 					PerScopeInfo[(int32)ViewedParameterScope].Tooltip = LOCTEXT("NiagaraInvalidScopeSelectionNameAlias", "Cannot select scope {0}: Parameter with same name already has this scope."); //@todo(ng) get scope
@@ -815,15 +823,18 @@ const TArray<FNiagaraScriptVariableAndViewInfo> FNiagaraScriptToolkitParameterPa
 
 			const FNiagaraVariableMetaData& MetaData = Iter.Value()->Metadata;
 			ENiagaraParameterScope MetaDataScope = ENiagaraParameterScope::None;
-			MetaData.GetScope(MetaDataScope);
-
-			if (MetaData.GetIsStaticSwitch() == false && (MetaDataScope == ENiagaraParameterScope::None || MetaData.GetUsage() == ENiagaraScriptParameterUsage::None) )
+			if (MetaData.GetIsStaticSwitch() == false)
 			{
-				// Parameters that are not static switches must have a scope and usage set.
-				ensureMsgf(false, TEXT("Invalid MetaData found for graph variable: %s"), *Iter.Value()->Variable.GetName().ToString());
-				continue;
+				FNiagaraEditorUtilities::GetVariableMetaDataScope(MetaData, MetaDataScope);
+				if (MetaDataScope == ENiagaraParameterScope::None || MetaData.GetUsage() == ENiagaraScriptParameterUsage::None)
+				{
+					// Parameters that are not static switches must have a scope and usage set.
+					ensureMsgf(false, TEXT("Invalid MetaData found for graph variable: %s"), *Iter.Value()->Variable.GetName().ToString());
+					continue;
+				}
 			}
-			else if (MetaData.GetUsage() == ENiagaraScriptParameterUsage::InputOutput)
+
+			if (MetaData.GetUsage() == ENiagaraScriptParameterUsage::InputOutput)
 			{
 				// Need two stack entries to represent this script parameter as both an input and output.
 				FNiagaraVariableMetaData InputVariableMetaData = MetaData;
diff --git a/Engine/Plugins/FX/Niagara/Source/NiagaraEditor/Private/ViewModels/Stack/NiagaraParameterHandle.cpp b/Engine/Plugins/FX/Niagara/Source/NiagaraEditor/Private/ViewModels/Stack/NiagaraParameterHandle.cpp
--- a/Engine/Plugins/FX/Niagara/Source/NiagaraEditor/Private/ViewModels/Stack/NiagaraParameterHandle.cpp
+++ b/Engine/Plugins/FX/Niagara/Source/NiagaraEditor/Private/ViewModels/Stack/NiagaraParameterHandle.cpp
@@ -4,18 +4,8 @@
 #include "NiagaraNodeFunctionCall.h"
 #include "NiagaraScript.h"
 #include "NiagaraTypes.h"
-
-const FName FNiagaraParameterHandle::UserNamespace(TEXT("User"));
-const FName FNiagaraParameterHandle::EngineNamespace(TEXT("Engine"));
-const FName FNiagaraParameterHandle::SystemNamespace(TEXT("System"));
-const FName FNiagaraParameterHandle::EmitterNamespace(TEXT("Emitter"));
-const FName FNiagaraParameterHandle::ParticleAttributeNamespace(TEXT("Particles"));
-const FName FNiagaraParameterHandle::ModuleNamespace(TEXT("Module"));
-const FName FNiagaraParameterHandle::ParameterCollectionNamespace(TEXT("NPC"));
-const FString FNiagaraParameterHandle::InitialPrefix(TEXT("Initial"));
-const FName FNiagaraParameterHandle::LocalNamespace(TEXT("Local"));
-const FName FNiagaraParameterHandle::InitialNamespace(TEXT("Initial"));
-const FName FNiagaraParameterHandle::OwnerNamespace(TEXT("Owner"));
+#include "NiagaraEditorModule.h"
+#include "NiagaraConstants.h"
 
 FNiagaraParameterHandle::FNiagaraParameterHandle() 
 {
@@ -64,22 +54,22 @@ FNiagaraParameterHandle FNiagaraParameterHandle::CreateEngineParameterHandle(con
 FNiagaraParameterHandle FNiagaraParameterHandle::CreateEmitterParameterHandle(const FNiagaraVariable& EmitterVariable)
 {
 
-	return FNiagaraParameterHandle(EmitterNamespace, EmitterVariable.GetName());
+	return FNiagaraParameterHandle(FNiagaraConstants::EmitterNamespace, EmitterVariable.GetName());
 }
 
 FNiagaraParameterHandle FNiagaraParameterHandle::CreateParticleAttributeParameterHandle(const FName InName)
 {
-	return FNiagaraParameterHandle(ParticleAttributeNamespace, InName);
+	return FNiagaraParameterHandle(FNiagaraConstants::ParticleAttributeNamespace, InName);
 }
 
 FNiagaraParameterHandle FNiagaraParameterHandle::CreateModuleParameterHandle(const FName InName)
 {
-	return FNiagaraParameterHandle(ModuleNamespace, InName);
+	return FNiagaraParameterHandle(FNiagaraConstants::ModuleNamespace, InName);
 }
 
 FNiagaraParameterHandle FNiagaraParameterHandle::CreateInitialParameterHandle(const FNiagaraParameterHandle& Handle)
 {
-	return FNiagaraParameterHandle(Handle.GetNamespace(), *FString::Printf(TEXT("%s.%s"), *InitialPrefix, *Handle.GetName().ToString()));
+	return FNiagaraParameterHandle(Handle.GetNamespace(), *FString::Printf(TEXT("%s.%s"), *FNiagaraConstants::InitialPrefix, *Handle.GetName().ToString()));
 }
 
 bool FNiagaraParameterHandle::IsValid() const 
@@ -104,56 +94,56 @@ const FName FNiagaraParameterHandle::GetNamespace() const
 
 bool FNiagaraParameterHandle::IsUserHandle() const
 {
-	return Namespace == UserNamespace;
+	return Namespace == FNiagaraConstants::UserNamespace;
 }
 
 bool FNiagaraParameterHandle::IsEngineHandle() const
 {
-	return Namespace == EngineNamespace;
+	return Namespace == FNiagaraConstants::EngineNamespace;
 }
 
 bool FNiagaraParameterHandle::IsSystemHandle() const
 {
-	return Namespace == SystemNamespace;
+	return Namespace == FNiagaraConstants::SystemNamespace;
 }
 
 bool FNiagaraParameterHandle::IsEmitterHandle() const
 {
-	return Namespace == EmitterNamespace;
+	return Namespace == FNiagaraConstants::EmitterNamespace;
 }
 
 bool FNiagaraParameterHandle::IsParticleAttributeHandle() const
 {
-	return Namespace == ParticleAttributeNamespace;
+	return Namespace == FNiagaraConstants::ParticleAttributeNamespace;
 }
 
 bool FNiagaraParameterHandle::IsModuleHandle() const
 {
-	return Namespace == ModuleNamespace;
+	return Namespace == FNiagaraConstants::ModuleNamespace;
 }
 
 bool FNiagaraParameterHandle::IsParameterCollectionHandle() const
 {
-	return Namespace == ParameterCollectionNamespace;
+	return Namespace == FNiagaraConstants::ParameterCollectionNamespace;
 }
 
 bool FNiagaraParameterHandle::IsReadOnlyHandle() const
 {
 	return
-		Namespace == UserNamespace ||
-		Namespace == EngineNamespace ||
-		Namespace == ParameterCollectionNamespace ||
-		Namespace == ModuleNamespace;
+		Namespace == FNiagaraConstants::UserNamespace ||
+		Namespace == FNiagaraConstants::EngineNamespace ||
+		Namespace == FNiagaraConstants::ParameterCollectionNamespace ||
+		Namespace == FNiagaraConstants::ModuleNamespace;
 }
 
 bool FNiagaraParameterHandle::IsTransientHandle() const
 {
 	return
-		Namespace != UserNamespace &&
-		Namespace != EngineNamespace &&
-		Namespace != SystemNamespace &&
-		Namespace != EmitterNamespace &&
-		Namespace != ParticleAttributeNamespace &&
-		Namespace != ModuleNamespace &&
-		Namespace != ParameterCollectionNamespace;
+		Namespace != FNiagaraConstants::UserNamespace &&
+		Namespace != FNiagaraConstants::EngineNamespace &&
+		Namespace != FNiagaraConstants::SystemNamespace &&
+		Namespace != FNiagaraConstants::EmitterNamespace &&
+		Namespace != FNiagaraConstants::ParticleAttributeNamespace &&
+		Namespace != FNiagaraConstants::ModuleNamespace &&
+		Namespace != FNiagaraConstants::ParameterCollectionNamespace;
 }
\ No newline at end of file
diff --git a/Engine/Plugins/FX/Niagara/Source/NiagaraEditor/Private/ViewModels/Stack/NiagaraStackFunctionInput.cpp b/Engine/Plugins/FX/Niagara/Source/NiagaraEditor/Private/ViewModels/Stack/NiagaraStackFunctionInput.cpp
--- a/Engine/Plugins/FX/Niagara/Source/NiagaraEditor/Private/ViewModels/Stack/NiagaraStackFunctionInput.cpp
+++ b/Engine/Plugins/FX/Niagara/Source/NiagaraEditor/Private/ViewModels/Stack/NiagaraStackFunctionInput.cpp
@@ -876,13 +876,13 @@ FName GetNamespaceForUsage(ENiagaraScriptUsage Usage)
 	case ENiagaraScriptUsage::ParticleUpdateScript:
 	case ENiagaraScriptUsage::ParticleEventScript:
 	case ENiagaraScriptUsage::ParticleSimulationStageScript:
-		return FNiagaraParameterHandle::ParticleAttributeNamespace;
+		return FNiagaraConstants::ParticleAttributeNamespace;
 	case ENiagaraScriptUsage::EmitterSpawnScript:
 	case ENiagaraScriptUsage::EmitterUpdateScript:
-		return FNiagaraParameterHandle::EmitterNamespace;
+		return FNiagaraConstants::EmitterNamespace;
 	case ENiagaraScriptUsage::SystemSpawnScript:
 	case ENiagaraScriptUsage::SystemUpdateScript:
-		return FNiagaraParameterHandle::SystemNamespace;
+		return FNiagaraConstants::SystemNamespace;
 	default:
 		return NAME_None;
 	}
@@ -1549,27 +1549,27 @@ void UNiagaraStackFunctionInput::GetNamespacesForNewParameters(TArray<FName>& Ou
 
 	if (OutputNode->GetUsage() == ENiagaraScriptUsage::ParticleSpawnScript || OutputNode->GetUsage() == ENiagaraScriptUsage::ParticleUpdateScript)
 	{
-		OutNamespacesForNewParameters.Add(FNiagaraParameterHandle::ParticleAttributeNamespace);
-		OutNamespacesForNewParameters.Add(FNiagaraParameterHandle::EmitterNamespace);
+		OutNamespacesForNewParameters.Add(FNiagaraConstants::ParticleAttributeNamespace);
+		OutNamespacesForNewParameters.Add(FNiagaraConstants::EmitterNamespace);
 		if (bIsEditingSystem)
 		{
-			OutNamespacesForNewParameters.Add(FNiagaraParameterHandle::SystemNamespace);
-			OutNamespacesForNewParameters.Add(FNiagaraParameterHandle::UserNamespace);
+			OutNamespacesForNewParameters.Add(FNiagaraConstants::SystemNamespace);
+			OutNamespacesForNewParameters.Add(FNiagaraConstants::UserNamespace);
 		}
 	}
 	else if (OutputNode->GetUsage() == ENiagaraScriptUsage::EmitterSpawnScript || OutputNode->GetUsage() == ENiagaraScriptUsage::EmitterUpdateScript)
 	{
-		OutNamespacesForNewParameters.Add(FNiagaraParameterHandle::EmitterNamespace);
+		OutNamespacesForNewParameters.Add(FNiagaraConstants::EmitterNamespace);
 		if (bIsEditingSystem)
 		{
-			OutNamespacesForNewParameters.Add(FNiagaraParameterHandle::SystemNamespace);
-			OutNamespacesForNewParameters.Add(FNiagaraParameterHandle::UserNamespace);
+			OutNamespacesForNewParameters.Add(FNiagaraConstants::SystemNamespace);
+			OutNamespacesForNewParameters.Add(FNiagaraConstants::UserNamespace);
 		}
 	}
 	else if ((OutputNode->GetUsage() == ENiagaraScriptUsage::SystemSpawnScript || OutputNode->GetUsage() == ENiagaraScriptUsage::SystemUpdateScript) && bIsEditingSystem)
 	{
-		OutNamespacesForNewParameters.Add(FNiagaraParameterHandle::SystemNamespace);
-		OutNamespacesForNewParameters.Add(FNiagaraParameterHandle::UserNamespace);
+		OutNamespacesForNewParameters.Add(FNiagaraConstants::SystemNamespace);
+		OutNamespacesForNewParameters.Add(FNiagaraConstants::UserNamespace);
 	}
 }
 
diff --git a/Engine/Plugins/FX/Niagara/Source/NiagaraEditor/Private/ViewModels/Stack/NiagaraStackGraphUtilities.cpp b/Engine/Plugins/FX/Niagara/Source/NiagaraEditor/Private/ViewModels/Stack/NiagaraStackGraphUtilities.cpp
--- a/Engine/Plugins/FX/Niagara/Source/NiagaraEditor/Private/ViewModels/Stack/NiagaraStackGraphUtilities.cpp
+++ b/Engine/Plugins/FX/Niagara/Source/NiagaraEditor/Private/ViewModels/Stack/NiagaraStackGraphUtilities.cpp
@@ -1519,7 +1519,7 @@ void FNiagaraStackGraphUtilities::GetNewParameterAvailableTypes(TArray<FNiagaraT
 	for (const FNiagaraTypeDefinition& RegisteredParameterType : FNiagaraTypeRegistry::GetRegisteredParameterTypes())
 	{
 		//Object types only allowed in user namespace at the moment.
-		if (RegisteredParameterType.IsUObject() && RegisteredParameterType.IsDataInterface() == false && Namespace != FNiagaraParameterHandle::UserNamespace)
+		if (RegisteredParameterType.IsUObject() && RegisteredParameterType.IsDataInterface() == false && Namespace != FNiagaraConstants::UserNamespace)
 		{
 			continue;
 		}
@@ -1608,13 +1608,13 @@ TOptional<FName> FNiagaraStackGraphUtilities::GetNamespaceForScriptUsage(ENiagar
 	case ENiagaraScriptUsage::ParticleUpdateScript:
 	case ENiagaraScriptUsage::ParticleEventScript:
 	case ENiagaraScriptUsage::ParticleSimulationStageScript:
-		return FNiagaraParameterHandle::ParticleAttributeNamespace;
+		return FNiagaraConstants::ParticleAttributeNamespace;
 	case ENiagaraScriptUsage::EmitterSpawnScript:
 	case ENiagaraScriptUsage::EmitterUpdateScript:
-		return FNiagaraParameterHandle::EmitterNamespace;
+		return FNiagaraConstants::EmitterNamespace;
 	case ENiagaraScriptUsage::SystemSpawnScript:
 	case ENiagaraScriptUsage::SystemUpdateScript:
-		return FNiagaraParameterHandle::SystemNamespace;
+		return FNiagaraConstants::SystemNamespace;
 	default:
 		return TOptional<FName>();
 	}
@@ -2540,157 +2540,4 @@ void FNiagaraStackGraphUtilities::RenameReferencingParameters(UNiagaraSystem& Sy
 	}
 }
 
-TArray<FName> FNiagaraStackGraphUtilities::DecomposeVariableNamespace(const FName& InVarNameToken, FName& OutName)
-{
-	int32 DotIndex;
-	TArray<FName> OutNamespaces;
-	FString VarNameString = InVarNameToken.ToString();
-	while (VarNameString.FindChar(TEXT('.'), DotIndex))
-	{
-		OutNamespaces.Add(FName(*VarNameString.Left(DotIndex)));
-		VarNameString = *VarNameString.RightChop(DotIndex + 1);
-	}
-	OutName = FName(*VarNameString);
-	return OutNamespaces;
-}
-
-void FNiagaraStackGraphUtilities::GetParameterMetaDataFromName(const FName& InVarNameToken, FNiagaraVariableMetaData& OutMetaData)
-{
-	auto MarkAsLegacyCustomName = [&OutMetaData]() {
-		OutMetaData.SetScope(ENiagaraParameterScope::Custom);
-		OutMetaData.SetIsUsingLegacyNameString(true);
-	};
-
-	auto GetMetaDataForFirstNamespace = [MarkAsLegacyCustomName](const FName& Namespace, FNiagaraVariableMetaData& OutMetaData)-> bool /*bNextNamespaceCanBeValid*/ {
-		if (Namespace == FNiagaraParameterHandle::LocalNamespace)
-		{
-			OutMetaData.SetScope(ENiagaraParameterScope::Local);
-			OutMetaData.SetUsage(ENiagaraScriptParameterUsage::Local);
-			return true;
-		}
-		else if (Namespace == FNiagaraParameterHandle::ModuleNamespace)
-		{
-			OutMetaData.SetScope(ENiagaraParameterScope::Input);
-			return false;
-		}
-		else if (Namespace == FNiagaraParameterHandle::UserNamespace)
-		{
-			OutMetaData.SetScope(ENiagaraParameterScope::User);
-			return false;
-		}
-		else if (Namespace == FNiagaraParameterHandle::EngineNamespace)
-		{
-			OutMetaData.SetScope(ENiagaraParameterScope::Engine);
-			return true;
-		}
-		else if (Namespace == FNiagaraParameterHandle::SystemNamespace)
-		{
-			OutMetaData.SetScope(ENiagaraParameterScope::System);
-			return true;
-		}
-		else if (Namespace == FNiagaraParameterHandle::EmitterNamespace)
-		{
-			OutMetaData.SetScope(ENiagaraParameterScope::Emitter);
-			return true;
-		}
-		else if (Namespace == FNiagaraParameterHandle::ParticleAttributeNamespace)
-		{
-			OutMetaData.SetScope(ENiagaraParameterScope::Particles);
-			return true;
-		}
-
-		MarkAsLegacyCustomName();
-		return false;
-	};
-
-	auto GetMetaDataForInitialNamespace = [](const FName& Namespace, FNiagaraVariableMetaData& OutMetaData)-> bool /*bFoundInitialNamespace*/ {
-		if (Namespace == FNiagaraParameterHandle::InitialNamespace)
-		{
-			OutMetaData.SetUsage(ENiagaraScriptParameterUsage::InitialValueInput);
-			return true;
-		}
-		return false;
-	};
-
-	auto GetScopeCanHaveInitialNamespaceSuffix = [](ENiagaraParameterScope InScope)-> bool /*bCanHaveInitialNamespaceSuffix*/ {
-		switch (InScope) {
-		case ENiagaraParameterScope::System:
-		case ENiagaraParameterScope::Emitter:
-		case ENiagaraParameterScope::Particles:
-			return true;
-		default:
-			return false;
-		};
-		return false;
-	};
-
-	FName NamespacelessName;
-	TArray<FName> DecomposedNamespaces = FNiagaraStackGraphUtilities::DecomposeVariableNamespace(InVarNameToken, NamespacelessName);
-	OutMetaData.SetCachedNamespacelessVariableName(NamespacelessName);
-	if (DecomposedNamespaces.Num() == 0)
-	{
-		UE_LOG(LogNiagaraEditor, Display, TEXT("Unexpected parameter encountered without a namespace: %s"), *InVarNameToken.ToString());
-		MarkAsLegacyCustomName();
-		return;
-	}
-	else if (DecomposedNamespaces.Num() == 1)
-	{
-		GetMetaDataForFirstNamespace(DecomposedNamespaces[0], OutMetaData);
-		return;
-	}
-	else if (DecomposedNamespaces.Num() == 2)
-	{
-		bool bNextNamespaceCanBeValid = GetMetaDataForFirstNamespace(DecomposedNamespaces[0], OutMetaData);
-		if (bNextNamespaceCanBeValid)
-		{
-			ENiagaraParameterScope FirstNamespaceScope;
-			OutMetaData.GetScope(FirstNamespaceScope);
-			if (FirstNamespaceScope == ENiagaraParameterScope::Local)
-			{
-				// "local.module." namespaces may be handled as local scopes and do not need to be marked as legacy namespaces.
-				if (DecomposedNamespaces[1] == FNiagaraParameterHandle::ModuleNamespace)
-				{
-					return;
-				}
-			}
-			else if (GetScopeCanHaveInitialNamespaceSuffix(FirstNamespaceScope))
-			{
-				bool bFoundInitialNamespace = GetMetaDataForInitialNamespace(DecomposedNamespaces[1], OutMetaData);
-				if (bFoundInitialNamespace)
-				{
-					return;
-				}
-			}
-		}
-		MarkAsLegacyCustomName();
-		return;
-	}
-
-	MarkAsLegacyCustomName();
-}
-
-FString FNiagaraStackGraphUtilities::GetNamespacelessVariableNameString(const FName& InVarName)
-{
-	int32 DotIndex;
-	FString VarNameString = InVarName.ToString();
-	if (VarNameString.FindLastChar(TEXT('.'), DotIndex))
-	{
-		return VarNameString.RightChop(DotIndex + 1);
-	}
-	// No dot index, must be a namespaceless variable name (e.g. static switch name) just return the name to string.
-	return VarNameString;
-}
-
-FName FNiagaraStackGraphUtilities::GetVariableNameForScope(const FName& CurrentVariableName, const ENiagaraParameterScope NewVariableScope, bool bIsInitialValue)
-{
-	FString NewVariableScopeString = FNiagaraTypeUtilities::GetNamespaceStringForScriptParameterScope(NewVariableScope);
-	if (bIsInitialValue)
-	{
-		NewVariableScopeString.Append(PARAM_MAP_INITIAL_STR);
-	}
-	const FString ExistingVariableNameString = FNiagaraStackGraphUtilities::GetNamespacelessVariableNameString(CurrentVariableName);
-	FName NewVariableName = FName(*(NewVariableScopeString + ExistingVariableNameString));
-	return NewVariableName;
-}
-
 #undef LOCTEXT_NAMESPACE
diff --git a/Engine/Plugins/FX/Niagara/Source/NiagaraEditor/Private/ViewModels/Stack/NiagaraStackParameterStoreEntry.cpp b/Engine/Plugins/FX/Niagara/Source/NiagaraEditor/Private/ViewModels/Stack/NiagaraStackParameterStoreEntry.cpp
--- a/Engine/Plugins/FX/Niagara/Source/NiagaraEditor/Private/ViewModels/Stack/NiagaraStackParameterStoreEntry.cpp
+++ b/Engine/Plugins/FX/Niagara/Source/NiagaraEditor/Private/ViewModels/Stack/NiagaraStackParameterStoreEntry.cpp
@@ -23,6 +23,7 @@
 #include "AssetRegistryModule.h"
 #include "ARFilter.h"
 #include "EdGraph/EdGraphPin.h"
+#include "NiagaraConstants.h"
 
 
 #define LOCTEXT_NAMESPACE "UNiagaraStackParameterStoreEntry"
@@ -218,14 +219,14 @@ TArray<UEdGraphPin*> UNiagaraStackParameterStoreEntry::GetOwningPins()
 void UNiagaraStackParameterStoreEntry::OnRenamed(FText NewName)
 {
 	FString ActualNameString = NewName.ToString();
-	FString NamespacePrefix = FNiagaraParameterHandle::UserNamespace.ToString() + ".";
+	FString NamespacePrefix = FNiagaraConstants::UserNamespace.ToString() + ".";
 	if (ActualNameString.Contains(NamespacePrefix))
 	{
 		ActualNameString = ActualNameString.Replace(*NamespacePrefix, TEXT(""));
 	}
 	FName ActualName = FName(*ActualNameString);
 	// what if it's not user namespace? dehardcode.
-	FNiagaraParameterHandle ParameterHandle(FNiagaraParameterHandle::UserNamespace, ActualName); 
+	FNiagaraParameterHandle ParameterHandle(FNiagaraConstants::UserNamespace, ActualName); 
 	FName VariableName = ParameterHandle.GetParameterHandleString();
 	if (VariableName != ParameterName)
 	{
@@ -387,7 +388,7 @@ UObject* UNiagaraStackParameterStoreEntry::GetCurrentValueObject()
 
 bool UNiagaraStackParameterStoreEntry::IsUniqueName(FString NewName)
 {
-	FString NamespacePrefix = FNiagaraParameterHandle::UserNamespace.ToString() + "."; // correcting name of variable for comparison, all user variables start with "User."
+	FString NamespacePrefix = FNiagaraConstants::UserNamespace.ToString() + "."; // correcting name of variable for comparison, all user variables start with "User."
 	if (!NewName.Contains(NamespacePrefix))
 	{
 		NewName = NamespacePrefix + NewName;
diff --git a/Engine/Plugins/FX/Niagara/Source/NiagaraEditor/Private/ViewModels/Stack/NiagaraStackSystemSettingsGroup.cpp b/Engine/Plugins/FX/Niagara/Source/NiagaraEditor/Private/ViewModels/Stack/NiagaraStackSystemSettingsGroup.cpp
--- a/Engine/Plugins/FX/Niagara/Source/NiagaraEditor/Private/ViewModels/Stack/NiagaraStackSystemSettingsGroup.cpp
+++ b/Engine/Plugins/FX/Niagara/Source/NiagaraEditor/Private/ViewModels/Stack/NiagaraStackSystemSettingsGroup.cpp
@@ -15,6 +15,7 @@
 #include "ScopedTransaction.h"
 #include "ViewModels/Stack/NiagaraStackGraphUtilities.h"
 #include "ViewModels/Stack/NiagaraStackSystemPropertiesItem.h"
+#include "NiagaraConstants.h"
 
 
 #define LOCTEXT_NAMESPACE "UNiagaraStackParameterStoreGroup"
@@ -72,10 +73,10 @@ public:
 	virtual void GenerateAddActions(TArray<TSharedRef<INiagaraStackItemGroupAddAction>>& OutAddActions, const FNiagaraStackItemGroupAddOptions& AddProperties) const override
 	{
 		TArray<FNiagaraTypeDefinition> AvailableTypes;
-		FNiagaraStackGraphUtilities::GetNewParameterAvailableTypes(AvailableTypes, FNiagaraParameterHandle::UserNamespace);
+		FNiagaraStackGraphUtilities::GetNewParameterAvailableTypes(AvailableTypes, FNiagaraConstants::UserNamespace);
 		for (const FNiagaraTypeDefinition& AvailableType : AvailableTypes)
 		{
-			FNiagaraParameterHandle NewParameterHandle(FNiagaraParameterHandle::UserNamespace, *(TEXT("New") + AvailableType.GetName()));
+			FNiagaraParameterHandle NewParameterHandle(FNiagaraConstants::UserNamespace, *(TEXT("New") + AvailableType.GetName()));
 			FNiagaraVariable NewParameterVariable(AvailableType, NewParameterHandle.GetParameterHandleString());
 			OutAddActions.Add(MakeShared<FParameterStoreGroupAddAction>(NewParameterVariable));
 		}
diff --git a/Engine/Plugins/FX/Niagara/Source/NiagaraEditor/Private/Widgets/SNiagaraParameterPanel.cpp b/Engine/Plugins/FX/Niagara/Source/NiagaraEditor/Private/Widgets/SNiagaraParameterPanel.cpp
--- a/Engine/Plugins/FX/Niagara/Source/NiagaraEditor/Private/Widgets/SNiagaraParameterPanel.cpp
+++ b/Engine/Plugins/FX/Niagara/Source/NiagaraEditor/Private/Widgets/SNiagaraParameterPanel.cpp
@@ -96,7 +96,7 @@ void SNiagaraParameterPanel::Tick(const FGeometry& AllottedGeometry, const doubl
 void SNiagaraParameterPanel::AddParameter(FNiagaraVariable NewVariable, const NiagaraParameterPanelSectionID::Type SectionID)
 {
 	FNiagaraVariableMetaData NewVariableMetaData = FNiagaraVariableMetaData();
-	NewVariableMetaData.SetScope(NiagaraParameterPanelSectionID::GetScopeForNewParametersInSection(SectionID));
+	NewVariableMetaData.SetScopeName(FNiagaraEditorUtilities::GetScopeNameForParameterScope(NiagaraParameterPanelSectionID::GetScopeForNewParametersInSection(SectionID)));
 	NewVariableMetaData.SetUsage(NiagaraParameterPanelSectionID::GetUsageForNewParametersInSection(SectionID));
 
 	ParameterPanelViewModel->AddParameter(NewVariable, NewVariableMetaData);
diff --git a/Engine/Plugins/FX/Niagara/Source/NiagaraEditor/Public/NiagaraEditorModule.h b/Engine/Plugins/FX/Niagara/Source/NiagaraEditor/Public/NiagaraEditorModule.h
--- a/Engine/Plugins/FX/Niagara/Source/NiagaraEditor/Public/NiagaraEditorModule.h
+++ b/Engine/Plugins/FX/Niagara/Source/NiagaraEditor/Public/NiagaraEditorModule.h
@@ -128,6 +128,9 @@ public:
 		EnqueueObjectForDeferredDestructionInternal(ObjectInContainer);
 	}
 
+	/** Lookup a parameter scope info by name. Returns nullptr if the parameter scope info registered name cannot be found. */
+	static const FNiagaraParameterScopeInfo* FindParameterScopeInfo(const FName& ParameterScopeInfoName);
+
 private:
 	class FDeferredDestructionContainerBase
 	{
@@ -176,6 +179,10 @@ private:
 
 	bool DeferredDestructObjects(float InDeltaTime);
 
+	/** Register a parameter scope info to lookup by name. */
+	static void RegisterParameterScopeInfo(const FName& ParameterScopeInfoName, const FNiagaraParameterScopeInfo& ParameterScopeInfo);
+
+
 private:
 	TSharedPtr<FExtensibilityManager> MenuExtensibilityManager;
 	TSharedPtr<FExtensibilityManager> ToolBarExtensibilityManager;
@@ -234,4 +241,6 @@ private:
 	TMap<int32, TSharedPtr<FHlslNiagaraCompiler>> ActiveCompilations;
 
 	TArray<TSharedRef<const FDeferredDestructionContainerBase>> EnqueuedForDeferredDestruction;
+
+	static TArray<TPair<FName, FNiagaraParameterScopeInfo>> RegisteredParameterScopeInfos;
 };
diff --git a/Engine/Plugins/FX/Niagara/Source/NiagaraEditor/Public/NiagaraEditorUtilities.h b/Engine/Plugins/FX/Niagara/Source/NiagaraEditor/Public/NiagaraEditorUtilities.h
--- a/Engine/Plugins/FX/Niagara/Source/NiagaraEditor/Public/NiagaraEditorUtilities.h
+++ b/Engine/Plugins/FX/Niagara/Source/NiagaraEditor/Public/NiagaraEditorUtilities.h
@@ -252,4 +252,35 @@ namespace FNiagaraEditorUtilities
 	{
 		return GetUniqueObjectName(Outer, T::StaticClass(), CandidateName);
 	}
+
+	/** Gets the Scope and notifies if it does not apply due to an override being set.
+	 * @params MetaData				The MetaData to get the namespace string for.
+	 * @params OutScope		The Scope to return.
+	 * @return bool			Whether the returned scope is not overridden. Is false if bUseLegacyNameString is set.
+	 */
+	bool GetVariableMetaDataScope(const FNiagaraVariableMetaData& MetaData, ENiagaraParameterScope& OutScope);
+
+	/** Gets the Namespace string and notifies if it does not apply due to an override being set.
+	 * @params MetaData				The MetaData to get the namespace string for.
+	 * @params OutNamespaceString	The Namespace string to return.
+	 * @return bool					Whether the returned Namespace string is valid. Is false if bUseLegacyNameString is set.
+	 */
+	bool GetVariableMetaDataNamespaceString(const FNiagaraVariableMetaData& MetaData, FString& OutNamespaceString);
+
+	/** Gets the Namespace string and notifies if it does not apply due to an override being set.
+	 * @params MetaData				The MetaData to get the namespace string for.
+	 * @params NewScopeName			The NewScopeName to consider when getting the namespace string.
+	 * @params OutNamespaceString	The Namespace string to return.
+	 * @return bool					Whether the returned Namespace string is valid. Is false if bUseLegacyNameString is set.
+	 */
+	bool GetVariableMetaDataNamespaceStringForNewScope(const FNiagaraVariableMetaData& MetaData, const FName& NewScopeName, FString& OutNamespaceString);
+
+	FName GetScopeNameForParameterScope(ENiagaraParameterScope InScope);
+
+	TArray<FName> DecomposeVariableNamespace(const FName& InVarNameToken, FName& OutName);
+
+	void GetParameterMetaDataFromName(const FName& InVarNameToken, FNiagaraVariableMetaData& OutMetaData);
+
+	FString GetNamespacelessVariableNameString(const FName& InVarName);
+
 };
diff --git a/Engine/Plugins/FX/Niagara/Source/NiagaraEditor/Public/NiagaraGraph.h b/Engine/Plugins/FX/Niagara/Source/NiagaraEditor/Public/NiagaraGraph.h
--- a/Engine/Plugins/FX/Niagara/Source/NiagaraEditor/Public/NiagaraGraph.h
+++ b/Engine/Plugins/FX/Niagara/Source/NiagaraEditor/Public/NiagaraGraph.h
@@ -169,7 +169,7 @@ class UNiagaraGraph : public UEdGraph
 	struct FAddParameterOptions
 	{
 		FAddParameterOptions()
-			: NewParameterScope(TOptional<ENiagaraParameterScope>())
+			: NewParameterScopeName(TOptional<FName>())
 			, NewParameterUsage(TOptional<ENiagaraScriptParameterUsage>())
 			, bRefreshMetaDataScopeAndUsage(false)
 			, bIsStaticSwitch(false)
@@ -178,7 +178,7 @@ class UNiagaraGraph : public UEdGraph
 		{};
 
 		/** Optional scope to assign to the new parameter. New scope is force assigned and will ignore any scope that would be assigned if bRefreshMetaDataScopeAndUsage is true. */
-		TOptional<ENiagaraParameterScope> NewParameterScope;
+		TOptional<FName> NewParameterScopeName;
 		/** Optional usage to assign to the new parameter. New usage is force assigned and will ignore any usage that would be assigned if bRefreshMetaDataScopeAndUsage is true.*/
 		TOptional<ENiagaraScriptParameterUsage> NewParameterUsage;
 		/** Whether the new or already existing parameter should refresh its scope and usage. If no scope is assigned, assign a scope from the variable name. Set a new usage depending on associated in/out pins in the graph. */
@@ -292,10 +292,10 @@ class UNiagaraGraph : public UEdGraph
 	* @param Parameter				The target FNiagaraVariable key to lookup the canonical UNiagaraScriptVariable to rename.
 	* @param NewName				The new name to apply.
 	* @param bFromStaticSwitch		Whether the target parameter is from a static switch. Used to determine whether to fixup binding strings.
-	* @param NewScope				The new scope from the FNiagaraVariable's associated FNiagaraVariableMetaData. Used if changing the scope triggered a rename. Applied to the canonical UNiagaraScriptVariable's Metadata.
+	* @param NewScopeName			The new scope name from the FNiagaraVariable's associated FNiagaraVariableMetaData. Used if changing the scope triggered a rename. Applied to the canonical UNiagaraScriptVariable's Metadata.
 	* @return						true if the new name was applied. 
 	*/
-	bool RenameParameter(const FNiagaraVariable& Parameter, FName NewName, bool bFromStaticSwitch = false, ENiagaraParameterScope NewScope = ENiagaraParameterScope::None);
+	bool RenameParameter(const FNiagaraVariable& Parameter, FName NewName, bool bFromStaticSwitch = false, FName NewScopeName = FName());
 
 	/** 
 	 * Sets the ENiagaraScriptParameterUsage on the Metadata of a ScriptVariable depending on its pin usage in the node graph (used on input/output pins). 
diff --git a/Engine/Plugins/FX/Niagara/Source/NiagaraEditor/Public/NiagaraParameterMapHistory.h b/Engine/Plugins/FX/Niagara/Source/NiagaraEditor/Public/NiagaraParameterMapHistory.h
--- a/Engine/Plugins/FX/Niagara/Source/NiagaraEditor/Public/NiagaraParameterMapHistory.h
+++ b/Engine/Plugins/FX/Niagara/Source/NiagaraEditor/Public/NiagaraParameterMapHistory.h
@@ -223,19 +223,34 @@ public:
 	bool ShouldIgnoreVariableDefault(const FNiagaraVariable& Var)const;
 };
 
+// view of FNiagaraVariableMetaData that caches configurable types to avoid verification checks on getters.
+struct FNiagaraStaticVariableMetaData
+{
+	FNiagaraStaticVariableMetaData(const FNiagaraVariableMetaData& ViewMetaData);
+
+	ENiagaraParameterScope GetScope() const { return Scope; };
+	const FString& GetNamespaceString() const { return NamespaceString; };
+
+private:
+	ENiagaraParameterScope Scope;
+	FString NamespaceString;
+};
+
 // helper struct for handling tracked variable metadata and aliases from histories.
 struct FNiagaraHistoryVariable
 {
-	FNiagaraHistoryVariable(FNiagaraVariable* InVar, const FNiagaraVariable* InVarWithOriginalAlias, const FNiagaraVariableMetaData* InVarMetaData)
+	FNiagaraHistoryVariable(FNiagaraVariable* InVar, const FNiagaraVariable* InVarWithOriginalAlias, const FNiagaraVariableMetaData& InVarMetaData)
 		: Var(InVar)
 		, VarWithOriginalAlias(InVarWithOriginalAlias)
-		, VarMetaData(InVarMetaData)
+		, VarMetaData(&InVarMetaData)
+		, StaticVarMetaData(InVarMetaData)
 		, bDoNotDemote(false)
 	{};
 
 	FNiagaraVariable* Var;
 	const FNiagaraVariable* VarWithOriginalAlias;
 	const FNiagaraVariableMetaData* VarMetaData;
+	const FNiagaraStaticVariableMetaData StaticVarMetaData;
 
 	// Used to track whether this history variable should not be demoted during pre-compilation.
 	bool bDoNotDemote;
diff --git a/Engine/Plugins/FX/Niagara/Source/NiagaraEditor/Public/ViewModels/Stack/NiagaraParameterHandle.h b/Engine/Plugins/FX/Niagara/Source/NiagaraEditor/Public/ViewModels/Stack/NiagaraParameterHandle.h
--- a/Engine/Plugins/FX/Niagara/Source/NiagaraEditor/Public/ViewModels/Stack/NiagaraParameterHandle.h
+++ b/Engine/Plugins/FX/Niagara/Source/NiagaraEditor/Public/ViewModels/Stack/NiagaraParameterHandle.h
@@ -59,19 +59,6 @@ public:
 
 	bool IsTransientHandle() const;
 
-public:
-	static const FName UserNamespace;
-	static const FName EngineNamespace;
-	static const FName SystemNamespace;
-	static const FName EmitterNamespace;
-	static const FName ParticleAttributeNamespace;
-	static const FName ModuleNamespace;
-	static const FName ParameterCollectionNamespace;
-	static const FString InitialPrefix;
-	static const FName LocalNamespace;
-	static const FName InitialNamespace;
-	static const FName OwnerNamespace;
-
 private:
 	FName ParameterHandleName;
 	FName Name;
diff --git a/Engine/Plugins/FX/Niagara/Source/NiagaraEditor/Public/ViewModels/Stack/NiagaraStackGraphUtilities.h b/Engine/Plugins/FX/Niagara/Source/NiagaraEditor/Public/ViewModels/Stack/NiagaraStackGraphUtilities.h
--- a/Engine/Plugins/FX/Niagara/Source/NiagaraEditor/Public/ViewModels/Stack/NiagaraStackGraphUtilities.h
+++ b/Engine/Plugins/FX/Niagara/Source/NiagaraEditor/Public/ViewModels/Stack/NiagaraStackGraphUtilities.h
@@ -192,12 +192,4 @@ namespace FNiagaraStackGraphUtilities
 
 	void GatherRenamedStackFunctionOutputVariableNames(UNiagaraEmitter* Emitter, UNiagaraNodeFunctionCall& FunctionCallNode, const FString& OldFunctionName, const FString& NewFunctionName, TMap<FName, FName>& OutOldToNewNameMap);
 	void GatherRenamedStackFunctionInputAndOutputVariableNames(UNiagaraEmitter* Emitter, UNiagaraNodeFunctionCall& FunctionCallNode, const FString& OldFunctionName, const FString& NewFunctionName, TMap<FName, FName>& OutOldToNewNameMap);
-
-	TArray<FName> DecomposeVariableNamespace(const FName& InVarNameToken, FName& OutName);
-
-	void GetParameterMetaDataFromName(const FName& InVarNameToken, FNiagaraVariableMetaData& OutMetaData);
-
-	FString GetNamespacelessVariableNameString(const FName& InVarName);
-
-	FName GetVariableNameForScope(const FName& CurrentVariableName, const ENiagaraParameterScope NewVariableScope, bool bIsInitialValue);
 }
