Static analysis warning coverage

diff --git a/Engine/Plugins/FX/Niagara/Source/NiagaraEditor/Private/ViewModels/Stack/NiagaraStackPropertyRow.cpp b/Engine/Plugins/FX/Niagara/Source/NiagaraEditor/Private/ViewModels/Stack/NiagaraStackPropertyRow.cpp
--- a/Engine/Plugins/FX/Niagara/Source/NiagaraEditor/Private/ViewModels/Stack/NiagaraStackPropertyRow.cpp
+++ b/Engine/Plugins/FX/Niagara/Source/NiagaraEditor/Private/ViewModels/Stack/NiagaraStackPropertyRow.cpp
@@ -10,7 +10,7 @@
 void UNiagaraStackPropertyRow::Initialize(FRequiredEntryData InRequiredEntryData, TSharedRef<IDetailTreeNode> InDetailTreeNode, FString InOwnerStackItemEditorDataKey, FString InOwnerStackEditorDataKey, UNiagaraNode* InOwningNiagaraNode)
 {
 	TSharedPtr<IPropertyHandle> PropertyHandle = InDetailTreeNode->CreatePropertyHandle();
-	bool bRowIsAdvanced = PropertyHandle.IsValid() && PropertyHandle->GetProperty()->HasAnyPropertyFlags(CPF_AdvancedDisplay);
+	bool bRowIsAdvanced = PropertyHandle.IsValid() && PropertyHandle->GetProperty() && PropertyHandle->GetProperty()->HasAnyPropertyFlags(CPF_AdvancedDisplay);
 	FString RowStackEditorDataKey = FString::Printf(TEXT("%s-%s"), *InOwnerStackEditorDataKey, *InDetailTreeNode->GetNodeName().ToString());
 	Super::Initialize(InRequiredEntryData, bRowIsAdvanced, InOwnerStackItemEditorDataKey, RowStackEditorDataKey);
 	DetailTreeNode = InDetailTreeNode;
@@ -19,7 +19,7 @@ void UNiagaraStackPropertyRow::Initialize(FRequiredEntryData InRequiredEntryData
 		? EStackRowStyle::ItemCategory
 		: EStackRowStyle::ItemContent;
 	bCannotEditInThisContext = false;
-	if (PropertyHandle.IsValid() && PropertyHandle->GetProperty())
+	if (PropertyHandle.IsValid() && PropertyHandle.Get() && PropertyHandle->GetProperty())
 	{
 		FProperty* Prop = PropertyHandle->GetProperty();
 		FObjectPropertyBase* ObjProp = CastField<FObjectPropertyBase>(Prop);
