Niagara - Fix several system/emitter editor usability issues.
+ Fixed an issue where deleting renderers from the selection stack view with the trash can could not be undone.
+ Fixed an issue where text in the stack wouldn't clip correctly which would make it overlap with the buttons.
+ Fixed an issue where the delete trash can was showing up for all groups and not just groups that support delete.
+ Fixed an issue where renderers don't show the rename option that modules have.
+ Fixed an issue where renamed modules and renderers were not visually distinct in the overview.
+ Fix a crash when renaming and undoing.

diff --git a/Engine/Plugins/FX/Niagara/Source/NiagaraEditor/Private/ViewModels/Stack/NiagaraStackEntry.cpp b/Engine/Plugins/FX/Niagara/Source/NiagaraEditor/Private/ViewModels/Stack/NiagaraStackEntry.cpp
--- a/Engine/Plugins/FX/Niagara/Source/NiagaraEditor/Private/ViewModels/Stack/NiagaraStackEntry.cpp
+++ b/Engine/Plugins/FX/Niagara/Source/NiagaraEditor/Private/ViewModels/Stack/NiagaraStackEntry.cpp
@@ -153,6 +153,11 @@ void UNiagaraStackEntry::Initialize(FRequiredEntryData InRequiredEntryData, FStr
 
 void UNiagaraStackEntry::Finalize()
 {
+	if (ensureMsgf(IsFinalized() == false, TEXT("Can not finalize a stack entry more than once.")) == false)
+	{
+		return;
+	}
+
 	FinalizeInternal();
 	checkf(bIsFinalized, TEXT("Parent FinalizeInternal not called from overriden FinalizeInternal"));
 
@@ -162,7 +167,10 @@ void UNiagaraStackEntry::Finalize()
 
 	for (UNiagaraStackEntry* Child : Children)
 	{
-		Child->Finalize();
+		if(Child->GetOuter() == this)
+		{
+			Child->Finalize();
+		}
 	}
 	Children.Empty();
 
@@ -183,9 +191,9 @@ FText UNiagaraStackEntry::GetDisplayName() const
 	return FText();
 }
 
-FText UNiagaraStackEntry::GetOriginalName() const
+TOptional<FText> UNiagaraStackEntry::GetAlternateDisplayName() const
 {
-	return FText();
+	return AlternateDisplayName;
 }
 
 UObject* UNiagaraStackEntry::GetDisplayedObject() const
@@ -367,6 +375,11 @@ UNiagaraStackEntry::FOnRequestFullRefresh& UNiagaraStackEntry::OnRequestFullRefr
 	return RequestFullRefreshDeferredDelegate;
 }
 
+UNiagaraStackEntry::FOnAlternateDisplayNameChanged& UNiagaraStackEntry::OnAlternateDisplayNameChanged()
+{
+	return AlternateDisplayNameChangedDelegate;
+}
+
 int32 UNiagaraStackEntry::GetIndentLevel() const
 {
 	return IndentLevel;
@@ -617,6 +630,21 @@ void UNiagaraStackEntry::RefreshChildren()
 		ErrorChild->OnIssueModified().AddUObject(this, &UNiagaraStackEntry::IssueModified);
 	}
 
+	const FText* NewAlternateName = StackEditorData->GetStackEntryDisplayName(StackEditorDataKey);
+	if (NewAlternateName != nullptr && NewAlternateName->IsEmptyOrWhitespace() == false)
+	{
+		if (AlternateDisplayName.IsSet() == false || NewAlternateName->IdenticalTo(AlternateDisplayName.GetValue()) == false)
+		{
+			AlternateDisplayName = *NewAlternateName;
+			AlternateDisplayNameChangedDelegate.Broadcast();
+		}
+	}
+	else if(AlternateDisplayName.IsSet())
+	{
+		AlternateDisplayName.Reset();
+		AlternateDisplayNameChangedDelegate.Broadcast();
+	}
+
 	PostRefreshChildrenInternal();
 
 	StructureChangedDelegate.Broadcast();
@@ -771,6 +799,16 @@ void UNiagaraStackEntry::OnRenamed(FText NewName)
 
 			GetStackEditorData().Modify();
 			GetStackEditorData().SetStackEntryDisplayName(GetStackEditorDataKey(), NewName);
+
+			if (NewName.IsEmptyOrWhitespace())
+			{
+				AlternateDisplayName.Reset();
+			}
+			else
+			{
+				AlternateDisplayName = NewName;
+			}
+			AlternateDisplayNameChangedDelegate.Broadcast();
 		}
 	}
 }
diff --git a/Engine/Plugins/FX/Niagara/Source/NiagaraEditor/Private/ViewModels/Stack/NiagaraStackModuleItem.cpp b/Engine/Plugins/FX/Niagara/Source/NiagaraEditor/Private/ViewModels/Stack/NiagaraStackModuleItem.cpp
--- a/Engine/Plugins/FX/Niagara/Source/NiagaraEditor/Private/ViewModels/Stack/NiagaraStackModuleItem.cpp
+++ b/Engine/Plugins/FX/Niagara/Source/NiagaraEditor/Private/ViewModels/Stack/NiagaraStackModuleItem.cpp
@@ -136,47 +136,11 @@ void UNiagaraStackModuleItem::Initialize(FRequiredEntryData InRequiredEntryData,
 		AddChildFilter(FOnFilterChild::CreateUObject(this, &UNiagaraStackModuleItem::FilterOutputCollection));
 		AddChildFilter(FOnFilterChild::CreateUObject(this, &UNiagaraStackModuleItem::FilterLinkedInputCollection));
 	}
-
-	OnEditorDataChanged();
-	GetStackEditorData().OnPersistentDataChanged().AddUObject(this, &UNiagaraStackModuleItem::OnEditorDataChanged);
-}
-
-void UNiagaraStackModuleItem::OnEditorDataChanged()
-{
-	const FText* NewDisplayName = GetStackEditorData().GetStackEntryDisplayName(GetStackEditorDataKey());
-	if (NewDisplayName == nullptr)
-	{
-		CustomDisplayName.Reset();
-	}
-	else
-	{
-		CustomDisplayName = *NewDisplayName;
-	}
 }
 
 FText UNiagaraStackModuleItem::GetDisplayName() const
 {
-	if (IsFinalized())
-	{
-		return FText::FromName(NAME_None);
-	}
-
-	if (CustomDisplayName.IsSet())
-	{
-		return CustomDisplayName.GetValue();
-	}
-	
-	return GetOriginalName();
-}
-
-FText UNiagaraStackModuleItem::GetOriginalName() const
-{
-	if (FunctionCallNode != nullptr)
-	{
-		return FunctionCallNode->GetNodeTitle(ENodeTitleType::ListView);
-	}
-
-	return FText::FromName(NAME_None);
+	return FunctionCallNode->GetNodeTitle(ENodeTitleType::ListView);
 }
 
 UObject* UNiagaraStackModuleItem::GetDisplayedObject() const
@@ -1166,7 +1130,6 @@ void UNiagaraStackModuleItem::Delete()
 				GetSystemViewModel()->NotifyDataObjectChanged(InputNode->GetDataInterface());
 			}
 		}
-		Finalize();
 		ModifiedGroupItemsDelegate.Broadcast();
 	}
 }
diff --git a/Engine/Plugins/FX/Niagara/Source/NiagaraEditor/Private/ViewModels/Stack/NiagaraStackRendererItem.cpp b/Engine/Plugins/FX/Niagara/Source/NiagaraEditor/Private/ViewModels/Stack/NiagaraStackRendererItem.cpp
--- a/Engine/Plugins/FX/Niagara/Source/NiagaraEditor/Private/ViewModels/Stack/NiagaraStackRendererItem.cpp
+++ b/Engine/Plugins/FX/Niagara/Source/NiagaraEditor/Private/ViewModels/Stack/NiagaraStackRendererItem.cpp
@@ -244,7 +244,6 @@ void UNiagaraStackRendererItem::Delete()
 	Emitter->RemoveRenderer(RendererProperties.Get());
 
 	OnDataObjectModified().Broadcast(RendererProperties.Get());
-	Finalize();
 }
 
 bool UNiagaraStackRendererItem::HasBaseRenderer() const
diff --git a/Engine/Plugins/FX/Niagara/Source/NiagaraEditor/Private/ViewModels/Stack/NiagaraStackViewModel.cpp b/Engine/Plugins/FX/Niagara/Source/NiagaraEditor/Private/ViewModels/Stack/NiagaraStackViewModel.cpp
--- a/Engine/Plugins/FX/Niagara/Source/NiagaraEditor/Private/ViewModels/Stack/NiagaraStackViewModel.cpp
+++ b/Engine/Plugins/FX/Niagara/Source/NiagaraEditor/Private/ViewModels/Stack/NiagaraStackViewModel.cpp
@@ -129,6 +129,8 @@ void UNiagaraStackViewModel::InitializeWithRootEntry(UNiagaraStackEntry* InRootE
 	RootEntry->OnRequestFullRefreshDeferred().AddUObject(this, &UNiagaraStackViewModel::EntryRequestFullRefreshDeferred);
 	RootEntries.Add(RootEntry);
 
+	bExternalRootEntry = true;
+
 	GEditor->RegisterForUndo(this);
 	StructureChangedDelegate.Broadcast();
 }
diff --git a/Engine/Plugins/FX/Niagara/Source/NiagaraEditor/Public/ViewModels/Stack/NiagaraStackEntry.h b/Engine/Plugins/FX/Niagara/Source/NiagaraEditor/Public/ViewModels/Stack/NiagaraStackEntry.h
--- a/Engine/Plugins/FX/Niagara/Source/NiagaraEditor/Public/ViewModels/Stack/NiagaraStackEntry.h
+++ b/Engine/Plugins/FX/Niagara/Source/NiagaraEditor/Public/ViewModels/Stack/NiagaraStackEntry.h
@@ -94,6 +94,7 @@ public:
 	DECLARE_DELEGATE_RetVal_TwoParams(TOptional<FDropRequestResponse>, FOnRequestDrop, const UNiagaraStackEntry& /*TargetEntry*/, const FDropRequest& /*DropRequest*/);
 	DECLARE_DELEGATE_RetVal_OneParam(bool, FOnFilterChild, const UNiagaraStackEntry&);
 	DECLARE_DELEGATE(FStackIssueFixDelegate);
+	DECLARE_MULTICAST_DELEGATE(FOnAlternateDisplayNameChanged);
 
 public:
 	struct NIAGARAEDITOR_API FExecutionCategoryNames
@@ -225,7 +226,7 @@ public:
 
 	virtual FText GetDisplayName() const;
 
-	virtual FText GetOriginalName() const;
+	TOptional<FText> GetAlternateDisplayName() const;
 
 	virtual UObject* GetDisplayedObject() const;
 
@@ -292,6 +293,8 @@ public:
 
 	FOnRequestFullRefresh& OnRequestFullRefreshDeferred();
 
+	FOnAlternateDisplayNameChanged& OnAlternateDisplayNameChanged();
+
 	void RefreshChildren();
 
 	void RefreshChildrenDeferred();
@@ -447,6 +450,8 @@ private:
 
 	FOnRequestFullRefresh RequestFullRefreshDeferredDelegate;
 
+	FOnAlternateDisplayNameChanged AlternateDisplayNameChangedDelegate;
+
 	TArray<FOnFilterChild> ChildFilters;
 
 	UPROPERTY()
@@ -477,6 +482,8 @@ private:
 
 	bool bOwnerIsEnabled;
 
+	TOptional<FText> AlternateDisplayName;
+
 	int32 TotalNumberOfInfoIssues;
 	int32 TotalNumberOfWarningIssues;
 	int32 TotalNumberOfErrorIssues;
diff --git a/Engine/Plugins/FX/Niagara/Source/NiagaraEditor/Public/ViewModels/Stack/NiagaraStackModuleItem.h b/Engine/Plugins/FX/Niagara/Source/NiagaraEditor/Public/ViewModels/Stack/NiagaraStackModuleItem.h
--- a/Engine/Plugins/FX/Niagara/Source/NiagaraEditor/Public/ViewModels/Stack/NiagaraStackModuleItem.h
+++ b/Engine/Plugins/FX/Niagara/Source/NiagaraEditor/Public/ViewModels/Stack/NiagaraStackModuleItem.h
@@ -29,7 +29,6 @@ public:
 
 	virtual FText GetDisplayName() const override;
 	virtual UObject* GetDisplayedObject() const override;
-	virtual FText GetOriginalName() const override;
 	virtual FText GetTooltipText() const override;
 
 	INiagaraStackItemGroupAddUtilities* GetGroupAddUtilities();
@@ -113,8 +112,6 @@ private:
 	bool bIsEnabled;
 	bool bCanRefresh;
 
-	TOptional<FText> CustomDisplayName;
-
 	UPROPERTY()
 	UNiagaraStackModuleItemLinkedInputCollection* LinkedInputCollection;
 
diff --git a/Engine/Plugins/FX/Niagara/Source/NiagaraEditor/Public/ViewModels/Stack/NiagaraStackRendererItem.h b/Engine/Plugins/FX/Niagara/Source/NiagaraEditor/Public/ViewModels/Stack/NiagaraStackRendererItem.h
--- a/Engine/Plugins/FX/Niagara/Source/NiagaraEditor/Public/ViewModels/Stack/NiagaraStackRendererItem.h
+++ b/Engine/Plugins/FX/Niagara/Source/NiagaraEditor/Public/ViewModels/Stack/NiagaraStackRendererItem.h
@@ -44,6 +44,8 @@ public:
 	virtual FText GetDeleteTransactionText() const override;
 	virtual void Delete();
 
+	virtual bool SupportsRename() const override { return true; }
+
 	bool HasBaseRenderer() const;
 
 	virtual bool SupportsChangeEnabled() const override { return true; }
diff --git a/Engine/Plugins/FX/Niagara/Source/NiagaraEditorWidgets/Private/NiagaraEditorWidgetsStyle.cpp b/Engine/Plugins/FX/Niagara/Source/NiagaraEditorWidgets/Private/NiagaraEditorWidgetsStyle.cpp
--- a/Engine/Plugins/FX/Niagara/Source/NiagaraEditorWidgets/Private/NiagaraEditorWidgetsStyle.cpp
+++ b/Engine/Plugins/FX/Niagara/Source/NiagaraEditorWidgets/Private/NiagaraEditorWidgetsStyle.cpp
@@ -129,6 +129,11 @@ TSharedRef< FSlateStyleSet > FNiagaraEditorWidgetsStyle::Create()
 		.SetFont(SystemOverviewItemFont);
 	Style->Set("NiagaraEditor.SystemOverview.ItemText", SystemOverviewItemText);
 
+	FSlateFontInfo SystemOverviewAlternateItemFont = DEFAULT_FONT("Italic", 9);
+	FTextBlockStyle SystemOverviewAlternateItemText = FTextBlockStyle(NormalText)
+		.SetFont(SystemOverviewAlternateItemFont);
+	Style->Set("NiagaraEditor.SystemOverview.AlternateItemText", SystemOverviewAlternateItemText);
+
 	Style->Set("NiagaraEditor.SystemOverview.Item.BackgroundColor", FLinearColor(FColor(62, 62, 62)));
 	Style->Set("NiagaraEditor.SystemOverview.Group.BackgroundColor", FLinearColor::Transparent);
 	Style->Set("NiagaraEditor.SystemOverview.CheckBoxColor", FLinearColor(FColor(160, 160, 160)));
diff --git a/Engine/Plugins/FX/Niagara/Source/NiagaraEditorWidgets/Private/SNiagaraOverviewStack.cpp b/Engine/Plugins/FX/Niagara/Source/NiagaraEditorWidgets/Private/SNiagaraOverviewStack.cpp
--- a/Engine/Plugins/FX/Niagara/Source/NiagaraEditorWidgets/Private/SNiagaraOverviewStack.cpp
+++ b/Engine/Plugins/FX/Niagara/Source/NiagaraEditorWidgets/Private/SNiagaraOverviewStack.cpp
@@ -257,6 +257,94 @@ private:
 	FOnCheckedChanged OnCheckedChanged;
 };
 
+class SNiagaraSystemOverviewItemName : public SCompoundWidget
+{
+	SLATE_BEGIN_ARGS(SNiagaraSystemOverviewItemName) {}
+	SLATE_END_ARGS();
+
+	void Construct(const FArguments& InArgs, UNiagaraStackItem* InStackItem)
+	{
+		StackItem = InStackItem;
+		StackItem->OnAlternateDisplayNameChanged().AddSP(this, &SNiagaraSystemOverviewItemName::UpdateFromAlternateDisplayName);
+
+		ChildSlot
+		[
+			SAssignNew(NameTextBlock, STextBlock)
+			.Text(this, &SNiagaraSystemOverviewItemName::GetItemDisplayName)
+			.ToolTipText(this, &SNiagaraSystemOverviewItemName::GetItemToolTip)
+			.IsEnabled_UObject(StackItem.Get(), &UNiagaraStackEntry::GetIsEnabledAndOwnerIsEnabled)
+		];
+		UpdateFromAlternateDisplayName();
+	}
+
+	~SNiagaraSystemOverviewItemName()
+	{
+		if (StackItem.IsValid())
+		{
+			StackItem->OnAlternateDisplayNameChanged().RemoveAll(this);
+		}
+	}
+
+private:
+	FText GetItemDisplayName() const
+	{
+		if (StackItem.IsValid() && StackItem->IsFinalized() == false)
+		{
+			return StackItem->GetAlternateDisplayName().IsSet() ? StackItem->GetAlternateDisplayName().GetValue() : StackItem->GetDisplayName();
+		}
+		return FText::GetEmpty();
+	}
+
+	FText GetItemToolTip() const
+	{
+		if (StackItem.IsValid() && StackItem->IsFinalized() == false)
+		{
+			FText CurrentToolTip = StackItem->GetTooltipText();
+			if (ToolTipCache.IsSet() == false || LastToolTipCache.IdenticalTo(CurrentToolTip) == false)
+			{
+				if(StackItem->GetAlternateDisplayName().IsSet())
+				{
+					FText AlternateNameAndOriginalName = FText::Format(LOCTEXT("AlternateNameAndOriginalNameFormat", "{0} ({1})"), StackItem->GetAlternateDisplayName().GetValue(), StackItem->GetDisplayName());
+					if (CurrentToolTip.IsEmptyOrWhitespace())
+					{
+						ToolTipCache = AlternateNameAndOriginalName;
+					}
+					else
+					{
+						ToolTipCache = FText::Format(LOCTEXT("AlternateDisplayNameToolTipFormat", "{0}\n\n{1}"), AlternateNameAndOriginalName, StackItem->GetTooltipText());
+					}
+				}
+				else
+				{
+					ToolTipCache = CurrentToolTip;
+				}
+			}
+			LastToolTipCache = CurrentToolTip;
+			return ToolTipCache.GetValue();
+		}
+		return FText::GetEmpty();
+	}
+
+	void UpdateFromAlternateDisplayName()
+	{
+		if (StackItem.IsValid() && StackItem->IsFinalized() == false && StackItem->GetAlternateDisplayName().IsSet())
+		{
+			NameTextBlock->SetTextStyle(&FNiagaraEditorWidgetsStyle::Get().GetWidgetStyle<FTextBlockStyle>("NiagaraEditor.SystemOverview.AlternateItemText"));
+		}
+		else
+		{
+			NameTextBlock->SetTextStyle(&FNiagaraEditorWidgetsStyle::Get().GetWidgetStyle<FTextBlockStyle>("NiagaraEditor.SystemOverview.ItemText"));
+		}
+		ToolTipCache.Reset();
+	}
+
+private:
+	mutable FText LastToolTipCache;
+	mutable TOptional<FText> ToolTipCache;
+	TSharedPtr<STextBlock> NameTextBlock;
+	TWeakObjectPtr<UNiagaraStackItem> StackItem;
+};
+
 void SNiagaraOverviewStack::Construct(const FArguments& InArgs, UNiagaraStackViewModel& InStackViewModel, UNiagaraSystemSelectionViewModel& InOverviewSelectionViewModel)
 {
 	StackCommandContext = MakeShared<FNiagaraStackCommandContext>();
@@ -492,10 +580,7 @@ TSharedRef<ITableRow> SNiagaraOverviewStack::OnGenerateRowForEntry(UNiagaraStack
 			.VAlign(VAlign_Center)
 			.Padding(3, 2, 0, 2)
 			[
-				SNew(STextBlock)
-				.TextStyle(FNiagaraEditorWidgetsStyle::Get(), "NiagaraEditor.SystemOverview.ItemText")
-				.Text_UObject(Item, &UNiagaraStackEntry::GetDisplayName)
-				.IsEnabled_UObject(Item, &UNiagaraStackEntry::GetIsEnabledAndOwnerIsEnabled)
+				SNew(SNiagaraSystemOverviewItemName, StackItem)
 			]
 			// Enabled checkbox
 			+ SHorizontalBox::Slot()
diff --git a/Engine/Plugins/FX/Niagara/Source/NiagaraEditorWidgets/Private/SNiagaraStackEntryWidget.cpp b/Engine/Plugins/FX/Niagara/Source/NiagaraEditorWidgets/Private/SNiagaraStackEntryWidget.cpp
--- a/Engine/Plugins/FX/Niagara/Source/NiagaraEditorWidgets/Private/SNiagaraStackEntryWidget.cpp
+++ b/Engine/Plugins/FX/Niagara/Source/NiagaraEditorWidgets/Private/SNiagaraStackEntryWidget.cpp
@@ -17,9 +17,17 @@ void SNiagaraStackDisplayName::Construct(const FArguments& InArgs, UNiagaraStack
 	TextStyleName = InTextStyleName;
 
 	TypeNameStyle = InArgs._TypeNameStyle;
-	OnRenameCommitted = InArgs._OnRenameCommitted;
+
+	TAttribute<FText> EntryToolTipText;
+	EntryToolTipText.Bind(this, &SNiagaraStackDisplayName::GetEntryToolTipText);
+	SetToolTip(FSlateApplication::Get().MakeToolTip(EntryToolTipText));
+
+	TAttribute<bool> EntryIsEnabled;
+	EntryIsEnabled.Bind(this, &SNiagaraStackDisplayName::GetEntryIsEnabled);
+	SetEnabled(EntryIsEnabled);
 
 	StackViewModel->OnStructureChanged().AddSP(this, &SNiagaraStackDisplayName::StackViewModelStructureChanged);
+	StackEntryItem->OnAlternateDisplayNameChanged().AddSP(this, &SNiagaraStackDisplayName::StackEntryItemAlternateNameChanged);
 
 	ChildSlot
 	[
@@ -33,57 +41,16 @@ void SNiagaraStackDisplayName::Construct(const FArguments& InArgs, UNiagaraStack
 SNiagaraStackDisplayName::~SNiagaraStackDisplayName()
 {
 	StackViewModel->OnStructureChanged().RemoveAll(this);
+	StackEntryItem->OnAlternateDisplayNameChanged().RemoveAll(this);
 }
 
 TSharedRef<SWidget> SNiagaraStackDisplayName::ConstructChildren()
 {
+	EditableTextBlock.Reset();
 	const FInlineEditableTextBlockStyle& EditableStyle = FNiagaraEditorWidgetsStyle::Get().GetWidgetStyle<FInlineEditableTextBlockStyle>(TextStyleName);
+	TArray<TSharedRef<SWidget>> NameWidgets;
 
-	TSharedPtr<SWidget> BaseNameWidget;
-	if (StackEntryItem->SupportsRename())
-	{
-		BaseNameWidget = SNew(SHorizontalBox)
-			+ SHorizontalBox::Slot()
-			.AutoWidth()
-			.HAlign(HAlign_Left)
-			.VAlign(VAlign_Bottom)
-			[
-				SAssignNew(EditableTextBlock, SInlineEditableTextBlock)
-					.Style(&EditableStyle)
-					.ToolTipText_UObject(StackEntryItem, &UNiagaraStackEntry::GetTooltipText)
-					.Text_UObject(StackEntryItem, &UNiagaraStackEntry::GetDisplayName)
-					.HighlightText_UObject(StackViewModel, &UNiagaraStackViewModel::GetCurrentSearchText)
-					.ColorAndOpacity(this, &SNiagaraStackEntryWidget::GetTextColorForSearch, FSlateColor::UseForeground())
-					.OnTextCommitted(OnRenameCommitted)
-					.IsEnabled_UObject(StackEntryItem, &UNiagaraStackEntry::GetIsEnabledAndOwnerIsEnabled)
-			]
-			+ SHorizontalBox::Slot()
-			.AutoWidth()
-			.Padding(5,0,0,0)
-			.HAlign(HAlign_Left)
-			.VAlign(VAlign_Bottom)
-			[
-				SNew(STextBlock)
-				.TextStyle(TypeNameStyle)
-				.ToolTipText_UObject(StackEntryItem, &UNiagaraStackEntry::GetTooltipText)
-				.Text(this, &SNiagaraStackDisplayName::GetOriginalName)
-				.HighlightText_UObject(StackViewModel, &UNiagaraStackViewModel::GetCurrentSearchText)
-				.ColorAndOpacity(this, &SNiagaraStackDisplayName::GetTextColorForSearch, FSlateColor::UseSubduedForeground())
-				.IsEnabled_UObject(StackEntryItem, &UNiagaraStackEntry::GetIsEnabledAndOwnerIsEnabled)
-				.Visibility(this, &SNiagaraStackDisplayName::ShouldShowOriginalName)
-			];
-	}
-	else
-	{
-		BaseNameWidget = SNew(STextBlock)
-			.TextStyle(&EditableStyle.TextStyle)
-			.ToolTipText_UObject(StackEntryItem, &UNiagaraStackEntry::GetTooltipText)
-			.Text_UObject(StackEntryItem, &UNiagaraStackEntry::GetDisplayName)
-			.HighlightText_UObject(StackViewModel, &UNiagaraStackViewModel::GetCurrentSearchText)
-			.ColorAndOpacity(this, &SNiagaraStackDisplayName::GetTextColorForSearch, FSlateColor::UseForeground())
-			.IsEnabled_UObject(StackEntryItem, &UNiagaraStackEntry::GetIsEnabledAndOwnerIsEnabled);
-	}
-
+	// First check to see if we need to insert the emitter name.
 	int32 NumTopLevelEmitters = 0;
 	for (const TSharedRef<UNiagaraStackViewModel::FTopLevelViewModel>& TopLevelViewModel : StackViewModel->GetTopLevelViewModels())
 	{
@@ -93,36 +60,69 @@ TSharedRef<SWidget> SNiagaraStackDisplayName::ConstructChildren()
 		}
 	}
 
-	if (NumTopLevelEmitters <= 1)
+	if (NumTopLevelEmitters > 1)
 	{
-		TopLevelViewModelCountAtLastConstruction = 1;
-		return BaseNameWidget.ToSharedRef();
+		TSharedPtr<UNiagaraStackViewModel::FTopLevelViewModel> TopLevelViewModel = StackViewModel->GetTopLevelViewModelForEntry(*StackEntryItem);
+		NameWidgets.Add(SNew(STextBlock)
+			.TextStyle(FNiagaraEditorWidgetsStyle::Get(), TextStyleName)
+			.Text(this, &SNiagaraStackDisplayName::GetTopLevelDisplayName, TWeakPtr<UNiagaraStackViewModel::FTopLevelViewModel>(TopLevelViewModel))
+			.HighlightText_UObject(StackViewModel, &UNiagaraStackViewModel::GetCurrentSearchText)
+			.ColorAndOpacity(this, &SNiagaraStackDisplayName::GetTextColorForSearch, FSlateColor::UseForeground()));
+	}
+	TopLevelViewModelCountAtLastConstruction = NumTopLevelEmitters;
+
+	// Next add the main name widget which will be the alternate name if it's available, otherwise it's the regular display name.
+	if (StackEntryItem->SupportsRename())
+	{
+		// If the entry can be renamed we need an editable text block.
+		NameWidgets.Add(SAssignNew(EditableTextBlock, SInlineEditableTextBlock)
+			.Style(&EditableStyle)
+			.Text(this, &SNiagaraStackDisplayName::GetEntryDisplayName)
+			.HighlightText_UObject(StackViewModel, &UNiagaraStackViewModel::GetCurrentSearchText)
+			.ColorAndOpacity(this, &SNiagaraStackEntryWidget::GetTextColorForSearch, FSlateColor::UseForeground())
+			.OnTextCommitted(this, &SNiagaraStackDisplayName::EntryNameTextCommitted));
 	}
 	else
 	{
-		TopLevelViewModelCountAtLastConstruction = StackViewModel->GetTopLevelViewModels().Num();
-		TSharedPtr<UNiagaraStackViewModel::FTopLevelViewModel> TopLevelViewModel = StackViewModel->GetTopLevelViewModelForEntry(*StackEntryItem);
-		if(TopLevelViewModel.IsValid())
+		// Otherwise add a regular text block.
+		NameWidgets.Add(SNew(STextBlock)
+			.TextStyle(&EditableStyle.TextStyle)
+			.Text(this, &SNiagaraStackDisplayName::GetEntryDisplayName)
+			.HighlightText_UObject(StackViewModel, &UNiagaraStackViewModel::GetCurrentSearchText)
+			.ColorAndOpacity(this, &SNiagaraStackDisplayName::GetTextColorForSearch, FSlateColor::UseForeground()));
+	}
+
+	// Finally add a subdued box for the regular display name if we're showing an alternate name.
+	if(StackEntryItem->GetAlternateDisplayName().IsSet())
+	{
+		NameWidgets.Add(SNew(STextBlock)
+			.TextStyle(TypeNameStyle)
+			.Text(this, &SNiagaraStackDisplayName::GetOriginalName)
+			.HighlightText_UObject(StackViewModel, &UNiagaraStackViewModel::GetCurrentSearchText)
+			.ColorAndOpacity(this, &SNiagaraStackDisplayName::GetTextColorForSearch, FSlateColor::UseSubduedForeground()));
+	}
+
+	// If there is more than one name, put them in a wrap box so that they flow correctly when the rows are narrow.
+	if(NameWidgets.Num() > 1)
+	{
+		TSharedRef<SWrapBox> NamesWrapBox = SNew(SWrapBox)
+			.UseAllottedWidth(true);
+		for (TSharedRef<SWidget> NameWidget : NameWidgets)
 		{
-			return SNew(SWrapBox)
-				.Clipping(EWidgetClipping::ClipToBoundsAlways)
-				.UseAllottedWidth(true)
-				+ SWrapBox::Slot()
-				[
-					SNew(STextBlock)
-					.TextStyle(FNiagaraEditorWidgetsStyle::Get(), TextStyleName)
-					.ToolTipText_UObject(StackEntryItem, &UNiagaraStackEntry::GetTooltipText)
-					.Text(this, &SNiagaraStackDisplayName::GetTopLevelDisplayName, TWeakPtr<UNiagaraStackViewModel::FTopLevelViewModel>(TopLevelViewModel))
-					.HighlightText_UObject(StackViewModel, &UNiagaraStackViewModel::GetCurrentSearchText)
-					.ColorAndOpacity(this, &SNiagaraStackDisplayName::GetTextColorForSearch, FSlateColor::UseForeground())
-					.IsEnabled_UObject(StackEntryItem, &UNiagaraStackEntry::GetIsEnabledAndOwnerIsEnabled)
-				]
-				+ SWrapBox::Slot()
+			NamesWrapBox->AddSlot()
+				.VAlign(VAlign_Center)
+				.Padding(FMargin(0, 0, 5, 0))
 				[
-					BaseNameWidget.ToSharedRef()
+					NameWidget
 				];
 		}
+		return NamesWrapBox;
 	}
+	else if (NameWidgets.Num() == 1)
+	{
+		return NameWidgets[0];
+	}
+
 	TopLevelViewModelCountAtLastConstruction = -1;
 	return SNullWidget::NullWidget;
 }
@@ -135,7 +135,7 @@ FText SNiagaraStackDisplayName::GetTopLevelDisplayName(TWeakPtr<UNiagaraStackVie
 		if (TopLevelViewModel->GetDisplayName().IdenticalTo(TopLevelDisplayNameCache) == false)
 		{
 			TopLevelDisplayNameCache = TopLevelViewModel->GetDisplayName();
-			TopLevelDisplayNameFormattedCache = FText::Format(LOCTEXT("TopLevelDisplayNameFormat", "{0} - "), TopLevelDisplayNameCache);
+			TopLevelDisplayNameFormattedCache = FText::Format(LOCTEXT("TopLevelDisplayNameFormat", "{0} -"), TopLevelDisplayNameCache);
 		}
 	}
 	else
@@ -153,19 +153,52 @@ void SNiagaraStackDisplayName::StackViewModelStructureChanged()
 	}
 }
 
+void SNiagaraStackDisplayName::StackEntryItemAlternateNameChanged()
+{
+	if (StackEntryItem->IsFinalized() == false)
+	{
+		Container->SetContent(ConstructChildren());
+	}
+}
+
+FText SNiagaraStackDisplayName::GetEntryDisplayName() const 
+{
+	return StackEntryItem->GetAlternateDisplayName().IsSet() ? StackEntryItem->GetAlternateDisplayName().GetValue() : StackEntryItem->GetDisplayName();
+}
+
 FText SNiagaraStackDisplayName::GetOriginalName() const
 {
 	if (StackEntryItem->IsFinalized() == false)
 	{
-		return FText::Format(FTextFormat::FromString(TEXT("({0})")), StackEntryItem->GetOriginalName());
+		return FText::Format(FTextFormat::FromString(TEXT("({0})")), StackEntryItem->GetDisplayName());
+	}
+	return FText::GetEmpty();
+}
+
+FText SNiagaraStackDisplayName::GetEntryToolTipText() const
+{
+	if (StackEntryItem->IsFinalized() == false)
+	{
+		return StackEntryItem->GetTooltipText();
 	}
 	return FText::GetEmpty();
 }
 
-EVisibility SNiagaraStackDisplayName::ShouldShowOriginalName() const
+bool SNiagaraStackDisplayName::GetEntryIsEnabled() const
+{
+	if (StackEntryItem->IsFinalized() == false)
+	{
+		return StackEntryItem->GetIsEnabledAndOwnerIsEnabled();
+	}
+	return false;
+}
+
+void SNiagaraStackDisplayName::EntryNameTextCommitted(const FText& InText, ETextCommit::Type CommitInfo)
 {
-	const FText* CurrentDisplayName = StackEntryItem->GetStackEditorData().GetStackEntryDisplayName(StackEntryItem->GetStackEditorDataKey());
-	return CurrentDisplayName != nullptr ? EVisibility::Visible : EVisibility::Collapsed;
+	if (StackEntryItem->IsFinalized() == false && CommitInfo != ETextCommit::OnCleared)
+	{
+		StackEntryItem->OnRenamed(InText);
+	}
 }
 
 void SNiagaraStackDisplayName::StartRename()
diff --git a/Engine/Plugins/FX/Niagara/Source/NiagaraEditorWidgets/Private/SNiagaraStackEntryWidget.h b/Engine/Plugins/FX/Niagara/Source/NiagaraEditorWidgets/Private/SNiagaraStackEntryWidget.h
--- a/Engine/Plugins/FX/Niagara/Source/NiagaraEditorWidgets/Private/SNiagaraStackEntryWidget.h
+++ b/Engine/Plugins/FX/Niagara/Source/NiagaraEditorWidgets/Private/SNiagaraStackEntryWidget.h
@@ -30,8 +30,10 @@ protected:
 class SNiagaraStackDisplayName : public SNiagaraStackEntryWidget
 {
 public:
-	SLATE_BEGIN_ARGS(SNiagaraStackDisplayName) { }
-		SLATE_EVENT(FOnTextCommitted, OnRenameCommitted);
+	SLATE_BEGIN_ARGS(SNiagaraStackDisplayName)
+	{
+		_Clipping = EWidgetClipping::OnDemand;
+	}
 		SLATE_STYLE_ARGUMENT(FTextBlockStyle, TypeNameStyle)
 	SLATE_END_ARGS();
 
@@ -50,12 +52,21 @@ private:
 
 	void StackViewModelStructureChanged();
 
+	void StackEntryItemAlternateNameChanged();
+
+	FText GetEntryDisplayName() const;
+
 	FText GetOriginalName() const;
 
-	EVisibility ShouldShowOriginalName() const;
+	FText GetEntryToolTipText() const;
+
+	bool GetEntryIsEnabled() const;
+
+	void EntryNameTextCommitted(const FText& InText, ETextCommit::Type CommitInfo);
 
 private:
 	FName TextStyleName;
+
 	const FTextBlockStyle* TypeNameStyle;
 
 	TSharedPtr<SBox> Container;
diff --git a/Engine/Plugins/FX/Niagara/Source/NiagaraEditorWidgets/Private/Stack/SNiagaraStackItem.cpp b/Engine/Plugins/FX/Niagara/Source/NiagaraEditorWidgets/Private/Stack/SNiagaraStackItem.cpp
--- a/Engine/Plugins/FX/Niagara/Source/NiagaraEditorWidgets/Private/Stack/SNiagaraStackItem.cpp
+++ b/Engine/Plugins/FX/Niagara/Source/NiagaraEditorWidgets/Private/Stack/SNiagaraStackItem.cpp
@@ -13,7 +13,7 @@
 #include "Widgets/Text/STextBlock.h"
 #include "Widgets/Images/SImage.h"
 #include "Widgets/Layout/SBox.h"
-#include "SNiagaraStackIssueIcon.h"
+#include "ViewModels/Stack/NiagaraStackClipboardUtilities.h"
 
 #define LOCTEXT_NAMESPACE "NiagaraStackItem"
 
@@ -40,11 +40,15 @@ void SNiagaraStackItem::Construct(const FArguments& InArgs, UNiagaraStackItem& I
 
 	// Display name
 	RowBox->AddSlot()
-	.Padding(2, 0, 2, 0)
-	.VAlign(VAlign_Center)
-	[
-		SNew(SNiagaraStackDisplayName, InItem, *InStackViewModel, "NiagaraEditor.Stack.ItemText")
-	];
+		.Padding(2, 0, 2, 0)
+		.VAlign(VAlign_Center)
+		[
+			SAssignNew(DisplayNameWidget, SNiagaraStackDisplayName, InItem, *InStackViewModel, "NiagaraEditor.Stack.ItemText")
+			.TypeNameStyle(FNiagaraEditorWidgetsStyle::Get(), "NiagaraEditor.Stack.TypeNameText")
+		];
+
+	// Allow derived classes to add additional widgets.
+	AddCustomRowWidgets(RowBox);
 
 	// Reset to base button
 	if (Item->SupportsResetToBase())
@@ -104,15 +108,35 @@ void SNiagaraStackItem::Construct(const FArguments& InArgs, UNiagaraStackItem& I
 			SNew(SCheckBox)
 			.IsChecked(this, &SNiagaraStackItem::CheckEnabledStatus)
 			.OnCheckStateChanged(this, &SNiagaraStackItem::OnCheckStateChanged)
+			.IsEnabled(this, &SNiagaraStackItem::GetEnabledCheckBoxEnabled)
 		];
 	}
 
 	ChildSlot
 	[
-		RowBox
+		// Allow derived classes add a container for the row widgets, e.g. a drop target.
+		AddContainerForRowWidgets(RowBox)
 	];
 }
 
+void SNiagaraStackItem::Tick(const FGeometry& AllottedGeometry, const double InCurrentTime, const float InDeltaTime)
+{
+	if (Item->GetIsRenamePending())
+	{
+		if (DisplayNameWidget.IsValid())
+		{
+			DisplayNameWidget->StartRename();
+		}
+		Item->SetIsRenamePending(false);
+	}
+	SNiagaraStackEntryWidget::Tick(AllottedGeometry, InCurrentTime, InDeltaTime);
+}
+
+TSharedRef<SWidget> SNiagaraStackItem::AddContainerForRowWidgets(TSharedRef<SWidget> RowWidgets)
+{
+	return RowWidgets;
+}
+
 EVisibility SNiagaraStackItem::GetResetToBaseButtonVisibility() const
 {
 	FText Unused;
@@ -147,7 +171,9 @@ bool SNiagaraStackItem::GetDeleteButtonEnabled() const
 
 FReply SNiagaraStackItem::DeleteClicked()
 {
-	Item->Delete();
+	TArray<UNiagaraStackEntry*> EntriesToDelete;
+	EntriesToDelete.Add(Item);
+	FNiagaraStackClipboardUtilities::DeleteSelection(EntriesToDelete);
 	return FReply::Handled();
 }
 
@@ -161,4 +187,9 @@ ECheckBoxState SNiagaraStackItem::CheckEnabledStatus() const
 	return Item->GetIsEnabled() ? ECheckBoxState::Checked : ECheckBoxState::Unchecked;
 }
 
+bool SNiagaraStackItem::GetEnabledCheckBoxEnabled() const
+{
+	return Item->GetOwnerIsEnabled();
+}
+
 #undef LOCTEXT_NAMESPACE
\ No newline at end of file
diff --git a/Engine/Plugins/FX/Niagara/Source/NiagaraEditorWidgets/Private/Stack/SNiagaraStackItem.h b/Engine/Plugins/FX/Niagara/Source/NiagaraEditorWidgets/Private/Stack/SNiagaraStackItem.h
--- a/Engine/Plugins/FX/Niagara/Source/NiagaraEditorWidgets/Private/Stack/SNiagaraStackItem.h
+++ b/Engine/Plugins/FX/Niagara/Source/NiagaraEditorWidgets/Private/Stack/SNiagaraStackItem.h
@@ -16,7 +16,14 @@ public:
 	SLATE_END_ARGS();
 
 	void Construct(const FArguments& InArgs, UNiagaraStackItem& InItem, UNiagaraStackViewModel* InStackViewModel);
-	
+
+	virtual void Tick(const FGeometry& AllottedGeometry, const double InCurrentTime, const float InDeltaTime) override;
+
+protected:
+	virtual void AddCustomRowWidgets(TSharedRef<SHorizontalBox> HorizontalBox) { }
+
+	virtual TSharedRef<SWidget> AddContainerForRowWidgets(TSharedRef<SWidget> RowWidgets);
+
 private:
 	EVisibility GetResetToBaseButtonVisibility() const;
 
@@ -36,6 +43,9 @@ private:
 
 	ECheckBoxState CheckEnabledStatus() const;
 
+	bool GetEnabledCheckBoxEnabled() const;
+
 private:
 	UNiagaraStackItem* Item;
+	TSharedPtr<SNiagaraStackDisplayName> DisplayNameWidget;
 };
\ No newline at end of file
diff --git a/Engine/Plugins/FX/Niagara/Source/NiagaraEditorWidgets/Private/Stack/SNiagaraStackItemGroup.cpp b/Engine/Plugins/FX/Niagara/Source/NiagaraEditorWidgets/Private/Stack/SNiagaraStackItemGroup.cpp
--- a/Engine/Plugins/FX/Niagara/Source/NiagaraEditorWidgets/Private/Stack/SNiagaraStackItemGroup.cpp
+++ b/Engine/Plugins/FX/Niagara/Source/NiagaraEditorWidgets/Private/Stack/SNiagaraStackItemGroup.cpp
@@ -41,6 +41,7 @@ void SNiagaraStackItemGroup::Construct(const FArguments& InArgs, UNiagaraStackIt
 			.ToolTipText(this, &SNiagaraStackItemGroup::GetDeleteButtonToolTip)
 			.OnClicked(this, &SNiagaraStackItemGroup::DeleteClicked)
 			.IsEnabled(this, &SNiagaraStackItemGroup::GetDeleteButtonIsEnabled)
+			.Visibility(this, &SNiagaraStackItemGroup::GetDeleteButtonVisibility)
 			.Content()
 			[
 				SNew(STextBlock)
@@ -52,6 +53,7 @@ void SNiagaraStackItemGroup::Construct(const FArguments& InArgs, UNiagaraStackIt
 		+ SHorizontalBox::Slot()
 		.AutoWidth()
 		.HAlign(HAlign_Right)
+		.Padding(2, 0, 0, 0)
 		[
 			ConstructAddButton()
 		]
@@ -81,6 +83,11 @@ bool SNiagaraStackItemGroup::GetDeleteButtonIsEnabled() const
 	return Group->TestCanDeleteWithMessage(UnusedMessage);
 }
 
+EVisibility SNiagaraStackItemGroup::GetDeleteButtonVisibility() const
+{
+	return Group->SupportsDelete() ? EVisibility::Visible : EVisibility::Collapsed;
+}
+
 FReply SNiagaraStackItemGroup::DeleteClicked()
 {
 	Group->Delete();
diff --git a/Engine/Plugins/FX/Niagara/Source/NiagaraEditorWidgets/Private/Stack/SNiagaraStackItemGroup.h b/Engine/Plugins/FX/Niagara/Source/NiagaraEditorWidgets/Private/Stack/SNiagaraStackItemGroup.h
--- a/Engine/Plugins/FX/Niagara/Source/NiagaraEditorWidgets/Private/Stack/SNiagaraStackItemGroup.h
+++ b/Engine/Plugins/FX/Niagara/Source/NiagaraEditorWidgets/Private/Stack/SNiagaraStackItemGroup.h
@@ -23,6 +23,8 @@ private:
 
 	bool GetDeleteButtonIsEnabled() const;
 
+	EVisibility GetDeleteButtonVisibility() const;
+
 	FReply DeleteClicked();
 
 private:
diff --git a/Engine/Plugins/FX/Niagara/Source/NiagaraEditorWidgets/Private/Stack/SNiagaraStackModuleItem.cpp b/Engine/Plugins/FX/Niagara/Source/NiagaraEditorWidgets/Private/Stack/SNiagaraStackModuleItem.cpp
--- a/Engine/Plugins/FX/Niagara/Source/NiagaraEditorWidgets/Private/Stack/SNiagaraStackModuleItem.cpp
+++ b/Engine/Plugins/FX/Niagara/Source/NiagaraEditorWidgets/Private/Stack/SNiagaraStackModuleItem.cpp
@@ -34,124 +34,7 @@
 void SNiagaraStackModuleItem::Construct(const FArguments& InArgs, UNiagaraStackModuleItem& InModuleItem, UNiagaraStackViewModel* InStackViewModel)
 {
 	ModuleItem = &InModuleItem;
-	StackEntryItem = ModuleItem;
-	StackViewModel = InStackViewModel;
-
-	ChildSlot
-	[
-		SNew(SDropTarget)
-		.OnAllowDrop(this, &SNiagaraStackModuleItem::OnModuleItemAllowDrop)
-		.OnDrop(this, &SNiagaraStackModuleItem::OnModuleItemDrop)
-		.HorizontalImage(FNiagaraEditorWidgetsStyle::Get().GetBrush("NiagaraEditor.Stack.DropTarget.BorderHorizontal"))
-		.VerticalImage(FNiagaraEditorWidgetsStyle::Get().GetBrush("NiagaraEditor.Stack.DropTarget.BorderVertical"))
-		.BackgroundColor(FNiagaraEditorWidgetsStyle::Get().GetColor("NiagaraEditor.Stack.DropTarget.BackgroundColor"))
-		.BackgroundColorHover(FNiagaraEditorWidgetsStyle::Get().GetColor("NiagaraEditor.Stack.DropTarget.BackgroundColorHover"))
-		.Content()
-		[
-			SNew(SOverlay)
-			// Name
-			+ SOverlay::Slot()
-			.VAlign(VAlign_Center)
-			.HAlign(HAlign_Fill)
-			.Padding(2, 0, 0, 0)
-			[
-				SAssignNew(DisplayNameWidget, SNiagaraStackDisplayName, InModuleItem, *InStackViewModel, "NiagaraEditor.Stack.ItemText")
-				.OnRenameCommitted(this, &SNiagaraStackModuleItem::OnRenameCommitted)
-				.TypeNameStyle(FNiagaraEditorWidgetsStyle::Get(), "NiagaraEditor.Stack.TypeNameText")
-			]
-			// Raise Action Menu button
-			+ SOverlay::Slot()
-			.HAlign(HAlign_Right)
-			[
-				SNew(SBorder)
-				.Padding(0)
-				.BorderImage(FEditorStyle::GetBrush("WhiteBrush"))
-				.BorderBackgroundColor(FNiagaraEditorWidgetsStyle::Get().GetColor("NiagaraEditor.Stack.Item.HeaderBackgroundColor"))
-				[
-					SNew(SHorizontalBox)
-					+ SHorizontalBox::Slot()
-					.VAlign(VAlign_Center)
-					.AutoWidth()
-					[
-						SNew(SComboButton)
-						.HasDownArrow(false)
-						.ButtonStyle(FEditorStyle::Get(), "HoverHintOnly")
-						.ForegroundColor(FSlateColor::UseForeground())
-						.OnGetMenuContent(this, &SNiagaraStackModuleItem::RaiseActionMenuClicked)
-						.ContentPadding(FMargin(2))
-						.HAlign(HAlign_Center)
-						.VAlign(VAlign_Center)
-						.Visibility(this, &SNiagaraStackModuleItem::GetRaiseActionMenuVisibility)
-						.IsEnabled(this, &SNiagaraStackModuleItem::GetButtonsEnabled)
-						.ButtonContent()
-						[
-							SNew(SImage)
-							.Image(FEditorStyle::Get().GetBrush("PropertyWindow.Button_AddToArray"))
-						]
-					]
-					// Refresh button
-					+ SHorizontalBox::Slot()
-					.VAlign(VAlign_Center)
-					.AutoWidth()
-					[
-						SNew(SButton)
-						.ButtonStyle(FEditorStyle::Get(), "HoverHintOnly")
-						.IsFocusable(false)
-						.ForegroundColor(FNiagaraEditorWidgetsStyle::Get().GetColor("NiagaraEditor.Stack.FlatButtonColor"))
-						.ToolTipText(LOCTEXT("RefreshTooltip", "Refresh this module"))
-						.Visibility(this, &SNiagaraStackModuleItem::GetRefreshVisibility)
-						.IsEnabled(this, &SNiagaraStackModuleItem::GetButtonsEnabled)
-						.OnClicked(this, &SNiagaraStackModuleItem::RefreshClicked)
-						.Content()
-						[
-							SNew(STextBlock)
-							.Font(FEditorStyle::Get().GetFontStyle("FontAwesome.10"))
-							.Text(FEditorFontGlyphs::Refresh)
-						]
-					]
-					// Delete button
-					+ SHorizontalBox::Slot()
-					.VAlign(VAlign_Center)
-					.AutoWidth()
-					[
-						SNew(SButton)
-						.ButtonStyle(FEditorStyle::Get(), "HoverHintOnly")
-						.IsFocusable(false)
-						.ForegroundColor(FNiagaraEditorWidgetsStyle::Get().GetColor("NiagaraEditor.Stack.FlatButtonColor"))
-						.ToolTipText(this, &SNiagaraStackModuleItem::GetDeleteButtonToolTipText)
-						.IsEnabled(this, &SNiagaraStackModuleItem::GetDeleteButtonEnabled)
-						.OnClicked(this, &SNiagaraStackModuleItem::DeleteClicked)
-						.Content()
-						[
-							SNew(STextBlock)
-							.Font(FEditorStyle::Get().GetFontStyle("FontAwesome.10"))
-							.Text(FEditorFontGlyphs::Trash)
-						]
-					]
-					// Enabled checkbox
-					+ SHorizontalBox::Slot()
-					.Padding(2, 0, 0, 0)
-					.AutoWidth()
-					[
-						SNew(SCheckBox)
-						.IsChecked(this, &SNiagaraStackModuleItem::GetCheckState)
-						.OnCheckStateChanged(this, &SNiagaraStackModuleItem::OnCheckStateChanged)
-						.IsEnabled(this, &SNiagaraStackModuleItem::GetEnabledCheckBoxEnabled)
-					]
-				]
-			]
-		]
-	];
-}
-
-void SNiagaraStackModuleItem::SetEnabled(bool bInIsEnabled)
-{
-	ModuleItem->SetIsEnabled(bInIsEnabled);
-}
-
-bool SNiagaraStackModuleItem::CheckEnabledStatus(bool bIsEnabled)
-{
-	return ModuleItem->GetIsEnabled() == bIsEnabled;
+	SNiagaraStackItem::Construct(SNiagaraStackItem::FArguments(), InModuleItem, InStackViewModel);
 }
 
 void SNiagaraStackModuleItem::FillRowContextMenu(FMenuBuilder& MenuBuilder)
@@ -178,26 +61,68 @@ void SNiagaraStackModuleItem::Tick(const FGeometry& AllottedGeometry, const doub
 		ModuleItem->SetIsModuleScriptReassignmentPending(false);
 		ShowReassignModuleScriptMenu();
 	}
-
-	if (StackEntryItem->GetIsRenamePending())
-	{
-		if (DisplayNameWidget.IsValid())
-		{
-			DisplayNameWidget->StartRename();
-		}
-		StackEntryItem->SetIsRenamePending(false);
-	}
-	SNiagaraStackEntryWidget::Tick(AllottedGeometry, InCurrentTime, InDeltaTime);
+	SNiagaraStackItem::Tick(AllottedGeometry, InCurrentTime, InDeltaTime);
 }
 
-ECheckBoxState SNiagaraStackModuleItem::GetCheckState() const
+void SNiagaraStackModuleItem::AddCustomRowWidgets(TSharedRef<SHorizontalBox> HorizontalBox)
 {
-	return ModuleItem->GetIsEnabled() ? ECheckBoxState::Checked : ECheckBoxState::Unchecked;
+	// Add menu.
+	HorizontalBox->AddSlot()
+	.VAlign(VAlign_Center)
+	.AutoWidth()
+	[
+		SNew(SComboButton)
+		.HasDownArrow(false)
+		.ButtonStyle(FEditorStyle::Get(), "HoverHintOnly")
+		.ForegroundColor(FSlateColor::UseForeground())
+		.OnGetMenuContent(this, &SNiagaraStackModuleItem::RaiseActionMenuClicked)
+		.ContentPadding(FMargin(2))
+		.HAlign(HAlign_Center)
+		.VAlign(VAlign_Center)
+		.Visibility(this, &SNiagaraStackModuleItem::GetRaiseActionMenuVisibility)
+		.IsEnabled(this, &SNiagaraStackModuleItem::GetButtonsEnabled)
+		.ButtonContent()
+		[
+			SNew(SImage)
+			.Image(FEditorStyle::Get().GetBrush("PropertyWindow.Button_AddToArray"))
+		]
+	];
+
+	// Refresh button
+	HorizontalBox->AddSlot()
+	.VAlign(VAlign_Center)
+	.AutoWidth()
+	[
+		SNew(SButton)
+		.ButtonStyle(FEditorStyle::Get(), "HoverHintOnly")
+		.IsFocusable(false)
+		.ForegroundColor(FNiagaraEditorWidgetsStyle::Get().GetColor("NiagaraEditor.Stack.FlatButtonColor"))
+		.ToolTipText(LOCTEXT("RefreshTooltip", "Refresh this module"))
+		.Visibility(this, &SNiagaraStackModuleItem::GetRefreshVisibility)
+		.IsEnabled(this, &SNiagaraStackModuleItem::GetButtonsEnabled)
+		.OnClicked(this, &SNiagaraStackModuleItem::RefreshClicked)
+		.Content()
+		[
+			SNew(STextBlock)
+			.Font(FEditorStyle::Get().GetFontStyle("FontAwesome.10"))
+		.Text(FEditorFontGlyphs::Refresh)
+		]
+	];
 }
 
-void SNiagaraStackModuleItem::OnCheckStateChanged(ECheckBoxState InCheckState)
+TSharedRef<SWidget> SNiagaraStackModuleItem::AddContainerForRowWidgets(TSharedRef<SWidget> RowWidgets)
 {
-	ModuleItem->SetIsEnabled(InCheckState == ECheckBoxState::Checked);
+	return SNew(SDropTarget)
+	.OnAllowDrop(this, &SNiagaraStackModuleItem::OnModuleItemAllowDrop)
+	.OnDrop(this, &SNiagaraStackModuleItem::OnModuleItemDrop)
+	.HorizontalImage(FNiagaraEditorWidgetsStyle::Get().GetBrush("NiagaraEditor.Stack.DropTarget.BorderHorizontal"))
+	.VerticalImage(FNiagaraEditorWidgetsStyle::Get().GetBrush("NiagaraEditor.Stack.DropTarget.BorderVertical"))
+	.BackgroundColor(FNiagaraEditorWidgetsStyle::Get().GetColor("NiagaraEditor.Stack.DropTarget.BackgroundColor"))
+	.BackgroundColorHover(FNiagaraEditorWidgetsStyle::Get().GetColor("NiagaraEditor.Stack.DropTarget.BackgroundColorHover"))
+	.Content()
+	[
+		RowWidgets
+	];
 }
 
 bool SNiagaraStackModuleItem::GetButtonsEnabled() const
@@ -205,24 +130,6 @@ bool SNiagaraStackModuleItem::GetButtonsEnabled() const
 	return ModuleItem->GetOwnerIsEnabled() && ModuleItem->GetIsEnabled();
 }
 
-FText SNiagaraStackModuleItem::GetDeleteButtonToolTipText() const
-{
-	FText CanDeleteMessage;
-	ModuleItem->TestCanDeleteWithMessage(CanDeleteMessage);
-	return CanDeleteMessage;
-}
-
-bool SNiagaraStackModuleItem::GetDeleteButtonEnabled() const
-{
-	FText CanDeleteMessage;
-	return ModuleItem->TestCanDeleteWithMessage(CanDeleteMessage);
-}
-
-bool SNiagaraStackModuleItem::GetEnabledCheckBoxEnabled() const
-{
-	return ModuleItem->GetOwnerIsEnabled();
-}
-
 EVisibility SNiagaraStackModuleItem::GetRaiseActionMenuVisibility() const
 {
 	return CanRaiseActionMenu() ? EVisibility::Visible : EVisibility::Collapsed;
@@ -233,14 +140,6 @@ EVisibility SNiagaraStackModuleItem::GetRefreshVisibility() const
 	return ModuleItem->CanRefresh() ? EVisibility::Visible : EVisibility::Collapsed;
 }
 
-FReply SNiagaraStackModuleItem::DeleteClicked()
-{
-	TArray<UNiagaraStackEntry*> EntriesToDelete;
-	EntriesToDelete.Add(ModuleItem);
-	FNiagaraStackClipboardUtilities::DeleteSelection(EntriesToDelete);
-	return FReply::Handled();
-}
-
 TSharedRef<SWidget> SNiagaraStackModuleItem::RaiseActionMenuClicked()
 {
 	if (CanRaiseActionMenu())
@@ -252,11 +151,6 @@ TSharedRef<SWidget> SNiagaraStackModuleItem::RaiseActionMenuClicked()
 			if (OutputNode)
 			{
 				FMenuBuilder MenuBuilder(true, nullptr);
-				/*MenuBuilder.AddMenuEntry(
-					LOCTEXT("MergeLabel", "Merge Up"),
-					LOCTEXT("MergeToolTip", "If a Set Variables node precedes this one in the stack, merge this node (and all variable binding logic) into that stack."),
-					FSlateIcon(),
-					FUIAction(FExecuteAction::CreateUObject(AssignmentNode, &UNiagaraNodeAssignment::MergeUp)));*/
 				MenuBuilder.AddSubMenu(LOCTEXT("AddVariables", "Add Variable"), 
 					LOCTEXT("AddVariablesTooltip", "Add another variable to the end of the list"),
 					FNewMenuDelegate::CreateLambda([OutputNode, AssignmentNode](FMenuBuilder& SubMenuBuilder)
@@ -395,12 +289,4 @@ void SNiagaraStackModuleItem::ShowReassignModuleScriptMenu()
 	FSlateApplication::Get().PushMenu(AsShared(), FWidgetPath(), MenuWidget, MenuPosition, FPopupTransitionEffect::ContextMenu);
 }
 
-void SNiagaraStackModuleItem::OnRenameCommitted(const FText& NewName, ETextCommit::Type CommitType)
-{
-	if (CommitType != ETextCommit::OnCleared)
-	{
-		StackEntryItem->OnRenamed(NewName);
-	}
-}
-
 #undef LOCTEXT_NAMESPACE
\ No newline at end of file
diff --git a/Engine/Plugins/FX/Niagara/Source/NiagaraEditorWidgets/Private/Stack/SNiagaraStackModuleItem.h b/Engine/Plugins/FX/Niagara/Source/NiagaraEditorWidgets/Private/Stack/SNiagaraStackModuleItem.h
--- a/Engine/Plugins/FX/Niagara/Source/NiagaraEditorWidgets/Private/Stack/SNiagaraStackModuleItem.h
+++ b/Engine/Plugins/FX/Niagara/Source/NiagaraEditorWidgets/Private/Stack/SNiagaraStackModuleItem.h
@@ -2,7 +2,7 @@
 
 #pragma once
 
-#include "SNiagaraStackEntryWidget.h"
+#include "SNiagaraStackItem.h"
 #include "Styling/SlateTypes.h"
 #include "Layout/Visibility.h"
 #include "Types/SlateEnums.h"
@@ -11,7 +11,7 @@ class UNiagaraStackModuleItem;
 class UNiagaraStackViewModel;
 class SNiagaraStackDisplayName;
 
-class SNiagaraStackModuleItem : public SNiagaraStackEntryWidget
+class SNiagaraStackModuleItem : public SNiagaraStackItem
 {
 public:
 	SLATE_BEGIN_ARGS(SNiagaraStackModuleItem) { }
@@ -19,35 +19,24 @@ public:
 
 	void Construct(const FArguments& InArgs, UNiagaraStackModuleItem& InModuleItem, UNiagaraStackViewModel* InStackViewModel);
 
-	void SetEnabled(bool bInIsEnabled);
-
-	bool CheckEnabledStatus(bool bIsEnabled);
-
 	void FillRowContextMenu(class FMenuBuilder& MenuBuilder);
 
 	//~ SWidget interface
 	virtual FReply OnMouseButtonDoubleClick(const FGeometry& InMyGeometry, const FPointerEvent& InMouseEvent) override;
 	virtual void Tick(const FGeometry& AllottedGeometry, const double InCurrentTime, const float InDeltaTime) override;
 
-private:
-	ECheckBoxState GetCheckState() const;
+protected:
+	virtual void AddCustomRowWidgets(TSharedRef<SHorizontalBox> HorizontalBox) override;
 
-	void OnCheckStateChanged(ECheckBoxState InCheckState);
+	virtual TSharedRef<SWidget> AddContainerForRowWidgets(TSharedRef<SWidget> RowWidgets) override;
 
+private:
 	bool GetButtonsEnabled() const;
 
-	FText GetDeleteButtonToolTipText() const;
-
-	bool GetDeleteButtonEnabled() const;
-
-	bool GetEnabledCheckBoxEnabled() const;
-
 	EVisibility GetRaiseActionMenuVisibility() const;
 
 	EVisibility GetRefreshVisibility() const;
 
-	FReply DeleteClicked();
-	
 	TSharedRef<SWidget> RaiseActionMenuClicked();
 
 	bool CanRaiseActionMenu() const;
@@ -60,9 +49,6 @@ private:
 
 	void ShowReassignModuleScriptMenu();
 
-	void OnRenameCommitted(const FText& NewName, ETextCommit::Type CommitType);
-
 private:
 	UNiagaraStackModuleItem* ModuleItem;
-	TSharedPtr<SNiagaraStackDisplayName> DisplayNameWidget;
 };
\ No newline at end of file
diff --git a/Engine/Plugins/FX/Niagara/Source/NiagaraEditorWidgets/Private/Stack/SNiagaraStackTableRow.cpp b/Engine/Plugins/FX/Niagara/Source/NiagaraEditorWidgets/Private/Stack/SNiagaraStackTableRow.cpp
--- a/Engine/Plugins/FX/Niagara/Source/NiagaraEditorWidgets/Private/Stack/SNiagaraStackTableRow.cpp
+++ b/Engine/Plugins/FX/Niagara/Source/NiagaraEditorWidgets/Private/Stack/SNiagaraStackTableRow.cpp
@@ -132,7 +132,7 @@ void SNiagaraStackTableRow::SetNameAndValueContent(TSharedRef<SWidget> InNameWid
 	// Execution sub-category icon
 	+ SHorizontalBox::Slot()
 	.AutoWidth()
-	.Padding(FMargin(1, 1, 4, 1))
+	.Padding(FMargin(1, 1, 2, 1))
 	.VAlign(EVerticalAlignment::VAlign_Center)
 	.HAlign(EHorizontalAlignment::HAlign_Center)
 	[
