Niagara - First round of UI perf improvements.

diff --git a/Engine/Plugins/FX/Niagara/Source/NiagaraEditor/Private/ViewModels/NiagaraSystemViewModel.cpp b/Engine/Plugins/FX/Niagara/Source/NiagaraEditor/Private/ViewModels/NiagaraSystemViewModel.cpp
--- a/Engine/Plugins/FX/Niagara/Source/NiagaraEditor/Private/ViewModels/NiagaraSystemViewModel.cpp
+++ b/Engine/Plugins/FX/Niagara/Source/NiagaraEditor/Private/ViewModels/NiagaraSystemViewModel.cpp
@@ -251,12 +251,14 @@ void FNiagaraSystemViewModel::CompileSystem(bool bForce)
 	check(SystemScriptViewModel.IsValid());
 	SystemScriptViewModel->CompileSystem(bForce);
 	bCompilePendingCompletion = true;
+	InvalidateCachedCompileStatus();
 }
 
 ENiagaraScriptCompileStatus FNiagaraSystemViewModel::GetLatestCompileStatus() const
 {
-	check(SystemScriptViewModel.IsValid());
-	return SystemScriptViewModel->GetLatestCompileStatus();
+	return LatestCompileStatusCache.IsSet() 
+		? LatestCompileStatusCache.GetValue() 
+		: ENiagaraScriptCompileStatus::NCS_Unknown;
 }
 
 UNiagaraSystemEditorData& FNiagaraSystemViewModel::GetEditorData() const
@@ -531,41 +533,19 @@ void FNiagaraSystemViewModel::Tick(float DeltaTime)
 		return;
 	}
 
-	if (bForceAutoCompileOnce || (GetDefault<UNiagaraEditorSettings>()->GetAutoCompile() && bCanAutoCompile))
-	{
-		bool bRecompile = false;
+	TickCompileStatus();
 
-		check(SystemScriptViewModel.IsValid());
-		if (SystemScriptViewModel->GetLatestCompileStatus() == ENiagaraScriptCompileStatus::NCS_Dirty)
-		{
-			//SystemScriptViewModel->CompileSystem();
-			//UE_LOG(LogNiagaraEditor, Log, TEXT("Compiling %s due to dirty scripts."), *GetSystem().GetName());
-			bRecompile |= true;
-		}
+	bool bAutoCompileThisFrame = GetDefault<UNiagaraEditorSettings>()->GetAutoCompile() && bCanAutoCompile && GetLatestCompileStatus() == ENiagaraScriptCompileStatus::NCS_Dirty;
+	if ((bForceAutoCompileOnce || bAutoCompileThisFrame) && GetSystem().HasOutstandingCompilationRequests() == false)
 
-		for (TSharedRef<FNiagaraEmitterHandleViewModel> EmitterHandleViewModel : EmitterHandleViewModels)
-		{
-			if (EmitterHandleViewModel->GetIsEnabled() && EmitterHandleViewModel->GetEmitterViewModel()->GetLatestCompileStatus() == ENiagaraScriptCompileStatus::NCS_Dirty)
-			{
-				bRecompile |= true;
-				//EmitterHandleViewModel->GetEmitterViewModel()->CompileScripts();
-				//UE_LOG(LogNiagaraEditor, Log, TEXT("Compiling %s - %s due to dirty scripts."), *GetSystem().GetName(), *EmitterHandleViewModel->GetName().ToString());
-			}
-		}
-
-		if (GetSystem().HasOutstandingCompilationRequests() == false)
-		{
-			if (bRecompile || bForceAutoCompileOnce)
-			{
-				CompileSystem(false);
-				bForceAutoCompileOnce = false;
-			}
+	{
+		CompileSystem(false);
+		bForceAutoCompileOnce = false;
+	}
 
-			if (bResetRequestPending)
-			{
-				ResetSystem(ETimeResetMode::AllowResetTime, EMultiResetMode::AllowResetAllInstances, EReinitMode::ReinitializeSystem);
-			}
-		}
+	if (bResetRequestPending)
+	{
+		ResetSystem(ETimeResetMode::AllowResetTime, EMultiResetMode::AllowResetAllInstances, EReinitMode::ReinitializeSystem);
 	}
 
 	if (EmitterIdsRequiringSequencerTrackUpdate.Num() > 0)
@@ -787,6 +767,102 @@ void FNiagaraSystemViewModel::SendLastCompileMessageJobs() const
 	FNiagaraMessageManager::Get()->QueueMessageJobBatch(JobBatchToQueue, SystemMessageLogGuidKey.GetValue());
 }
 
+void FNiagaraSystemViewModel::InvalidateCachedCompileStatus()
+{
+	LatestCompileStatusCache.Reset();
+	ScriptsToCheckForStatus.Empty();
+	ScriptCompileStatuses.Empty();
+}
+
+void FNiagaraSystemViewModel::TickCompileStatus()
+{
+	// Checking the compile status is expensive which is why it's updated on tick, and multiple scripts
+	// are time sliced across multiple frames.
+	if (System->HasOutstandingCompilationRequests())
+	{
+		// When compiling the compile status is always unknown.
+		InvalidateCachedCompileStatus();
+	}
+	else
+	{
+		if (LatestCompileStatusCache.IsSet() == false)
+		{
+			if (ScriptsToCheckForStatus.Num() > 0)
+			{
+				// If there are still scripts to check for status do that...
+				UNiagaraScript* ScriptToCheck = ScriptsToCheckForStatus.Last();
+				ScriptsToCheckForStatus.RemoveAt(ScriptsToCheckForStatus.Num() - 1);
+
+				if (ScriptToCheck->AreScriptAndSourceSynchronized() == false)
+				{
+					// If any script is not synchronized the entire system is considered dirty and we no longer need
+					// to check other script statuses.
+					LatestCompileStatusCache = ENiagaraScriptCompileStatus::NCS_Dirty;
+					ScriptsToCheckForStatus.Empty();
+					ScriptCompileStatuses.Empty();
+				}
+				else
+				{
+					ENiagaraScriptCompileStatus LastScirptCompileStatus = ScriptToCheck->GetLastCompileStatus();
+					if (LastScirptCompileStatus == ENiagaraScriptCompileStatus::NCS_Unknown ||
+						LastScirptCompileStatus == ENiagaraScriptCompileStatus::NCS_Dirty ||
+						LastScirptCompileStatus == ENiagaraScriptCompileStatus::NCS_BeingCreated)
+					{
+						// If the script doesn't have a vaild last compile status assume that it's dirty and by extension
+						// so is the system and we can ignore the other statuses.
+						LatestCompileStatusCache = ENiagaraScriptCompileStatus::NCS_Dirty;
+						ScriptsToCheckForStatus.Empty();
+						ScriptCompileStatuses.Empty();
+					}
+					else
+					{
+						// Otherwise save it for processing once all scripts are done.
+						ScriptCompileStatuses.Add(LastScirptCompileStatus);
+					}
+				}
+			}
+			else
+			{
+				// Otherwise check for status results.
+				if (ScriptCompileStatuses.Num() > 0)
+				{
+					// If there are statues from scripts figure out the system status.
+					if (ScriptCompileStatuses.Contains(ENiagaraScriptCompileStatus::NCS_Error))
+					{
+						// If any individual script has an error the system is considered to have an error.
+						LatestCompileStatusCache = ENiagaraScriptCompileStatus::NCS_Error;
+					}
+					else if (ScriptCompileStatuses.Contains(ENiagaraScriptCompileStatus::NCS_UpToDateWithWarnings) ||
+						ScriptCompileStatuses.Contains(ENiagaraScriptCompileStatus::NCS_ComputeUpToDateWithWarnings))
+					{
+						// If there were no errors and any individual script has a warning the system is considered to have a warning.
+						LatestCompileStatusCache = ENiagaraScriptCompileStatus::NCS_UpToDateWithWarnings;
+					}
+					else
+					{
+						LatestCompileStatusCache = ENiagaraScriptCompileStatus::NCS_UpToDate;
+					}
+					ScriptCompileStatuses.Empty();
+				}
+				else
+				{
+					// If there is no current status, no scripts to check, and no status results then
+					// collect up the scripts to check deferred.
+					ScriptsToCheckForStatus.Add(System->GetSystemSpawnScript());
+					ScriptsToCheckForStatus.Add(System->GetSystemUpdateScript());
+					for (const FNiagaraEmitterHandle& EmitterHandle : System->GetEmitterHandles())
+					{
+						if (EmitterHandle.GetIsEnabled())
+						{
+							EmitterHandle.GetInstance()->GetScripts(ScriptsToCheckForStatus, true);
+						}
+					}
+				}
+			}
+		}
+	}
+}
+
 void FNiagaraSystemViewModel::SetupPreviewComponentAndInstance()
 {
 	if (bCanSimulate)
@@ -1432,11 +1508,13 @@ void FNiagaraSystemViewModel::EmitterScriptGraphChanged(const FEdGraphEditAction
 	}
 	// Remove from cache on graph change 
 	GuidToCachedStackModuleData.Remove(OwningEmitterHandleId);
+	InvalidateCachedCompileStatus();
 }
 
 void FNiagaraSystemViewModel::SystemScriptGraphChanged(const FEdGraphEditAction& InAction)
 {
 	GuidToCachedStackModuleData.Empty();
+	InvalidateCachedCompileStatus();
 }
 
 void FNiagaraSystemViewModel::EmitterParameterStoreChanged(const FNiagaraParameterStore& ChangedParameterStore, const UNiagaraScript& OwningScript)
diff --git a/Engine/Plugins/FX/Niagara/Source/NiagaraEditor/Public/ViewModels/NiagaraSystemViewModel.h b/Engine/Plugins/FX/Niagara/Source/NiagaraEditor/Public/ViewModels/NiagaraSystemViewModel.h
--- a/Engine/Plugins/FX/Niagara/Source/NiagaraEditor/Public/ViewModels/NiagaraSystemViewModel.h
+++ b/Engine/Plugins/FX/Niagara/Source/NiagaraEditor/Public/ViewModels/NiagaraSystemViewModel.h
@@ -314,6 +314,10 @@ private:
 	/** Sends message jobs to FNiagaraMessageManager for all compile events from the last compile. */
 	void SendLastCompileMessageJobs() const;
 
+	void InvalidateCachedCompileStatus();
+
+	void TickCompileStatus();
+
 	/** Sets up the preview component and System instance. */
 	void SetupPreviewComponentAndInstance();
 
@@ -569,4 +573,9 @@ private:
 	UNiagaraSystemSelectionViewModel* SelectionViewModel;
 
 	UNiagaraScratchPadViewModel* ScriptScratchPadViewModel;
+
+	TArray<UNiagaraScript*> ScriptsToCheckForStatus;
+	TArray<ENiagaraScriptCompileStatus> ScriptCompileStatuses;
+
+	TOptional<ENiagaraScriptCompileStatus> LatestCompileStatusCache;
 };
\ No newline at end of file
