Culling Emitters on cook.

- UNiagaraEmitters that are not enabled on a platform are filtered from the cook.
- UNiagaraSystems with no enabled emitters are also filtered.
- Fixed issue with application of EffectsQuality on cooked platforms.
- Fixed up various places reading from now potentially null emitters and systems.
- Fixed issue where dummy device profiles were being made that was interfering the the PlatformSet code.
- Fixed effect quality change resettig all systems. Now only resent active systems.

diff --git a/Engine/Plugins/FX/Niagara/Source/Niagara/Classes/NiagaraEmitter.h b/Engine/Plugins/FX/Niagara/Source/Niagara/Classes/NiagaraEmitter.h
--- a/Engine/Plugins/FX/Niagara/Source/Niagara/Classes/NiagaraEmitter.h
+++ b/Engine/Plugins/FX/Niagara/Source/Niagara/Classes/NiagaraEmitter.h
@@ -238,14 +238,16 @@ public:
 	virtual void PostEditChangeProperty(struct FPropertyChangedEvent& PropertyChangedEvent) override;
 	NIAGARA_API FOnPropertiesChanged& OnPropertiesChanged();
 	NIAGARA_API FOnRenderersChanged& OnRenderersChanged();
-	bool IsEnabledOnPlatform(const FString& PlatformName);
 #endif
+	virtual bool NeedsLoadForTargetPlatform(const ITargetPlatform* TargetPlatform) const override;
 	void Serialize(FArchive& Ar)override;
 	virtual void PostInitProperties() override;
 	virtual void PostLoad() override;
 	virtual bool IsEditorOnly() const override;
 	//End UObject Interface
 
+	bool IsEnabledOnPlatform(const FString& PlatformName)const;
+
 	/** Toggles whether or not the particles within this emitter are relative to the emitter origin or in global space. */ 
 	UPROPERTY(EditAnywhere, Category = "Emitter")
 	bool bLocalSpace;
diff --git a/Engine/Plugins/FX/Niagara/Source/Niagara/Classes/NiagaraPlatformSet.h b/Engine/Plugins/FX/Niagara/Source/Niagara/Classes/NiagaraPlatformSet.h
--- a/Engine/Plugins/FX/Niagara/Source/Niagara/Classes/NiagaraPlatformSet.h
+++ b/Engine/Plugins/FX/Niagara/Source/Niagara/Classes/NiagaraPlatformSet.h
@@ -123,6 +123,9 @@ public:
 	/** Fill OutProfiles with all device proifles that have been overriden at the passed EffectsQuality. */
 	void GetOverridenDeviceProfiles(int32 EffectsQuality, TArray<UDeviceProfile*>& OutEnabledProfiles, TArray<UDeviceProfile*>& OutDisabledProfiles)const;
 
+	/** Returns true if this set is enabled for any profiles on the specified platform. */
+	bool IsEnabledForPlatform(const FString& PlatformName)const;
+
 	/** Returns true if the current platform can modify it's niagara scalability settings at runtime. */
 	static bool CanChangeScalabilityAtRuntime();
 
@@ -131,11 +134,11 @@ public:
 
 	static int32 GetEffectQualityMaskForDeviceProfile(const UDeviceProfile* Profile);
 
+	/** Returns true if the passed platform should prune emitters on cook. */
+	static bool ShouldPruneEmittersOnCook(const FString& PlatformName);
+
 	//Editor only public API
 #if WITH_EDITOR
-	/** Returns true if this set is enabled for any profiles on the specified platform. */
-	bool IsEnabledForPlatform(const FString& PlatformName);
-
 	/**
 	Does the enabled set of profiles in this set conflict with those in another set.
 	In some cases this is not allowed.
@@ -158,9 +161,6 @@ public:
 	/** Returns true if the passed platform can modify it's niagara scalability settings at runtime. */
 	static bool CanChangeScalabilityAtRuntime(const FString& PlatformName);
 
-	/** Returns true if the passed platform should prune emitters on cook. */
-	static bool ShouldPruneEmittersOnCook(const FString& PlatformName);
-
 	/** Inspects the passed sets and generates an array of all conflicts between these sets. Used to keep arrays of platform sets orthogonal. */
 	static bool GatherConflicts(const TArray<const FNiagaraPlatformSet*>& PlatformSets, TArray<FNiagaraPlatformSetConflictInfo>& OutConflicts);
 #endif
diff --git a/Engine/Plugins/FX/Niagara/Source/Niagara/Classes/NiagaraSystem.h b/Engine/Plugins/FX/Niagara/Source/Niagara/Classes/NiagaraSystem.h
--- a/Engine/Plugins/FX/Niagara/Source/Niagara/Classes/NiagaraSystem.h
+++ b/Engine/Plugins/FX/Niagara/Source/Niagara/Classes/NiagaraSystem.h
@@ -187,7 +187,8 @@ public:
 	void Serialize(FArchive& Ar)override;
 	virtual void PostLoad() override; 
 	virtual void BeginDestroy() override;
-	virtual void PreSave(const class ITargetPlatform * TargetPlatform) override;
+	virtual void PreSave(const class ITargetPlatform* TargetPlatform) override;
+	virtual bool NeedsLoadForTargetPlatform(const ITargetPlatform* TargetPlatform) const override;
 #if WITH_EDITOR
 	virtual void PreEditChange(FProperty* PropertyThatWillChange)override;
 	virtual void PostEditChangeProperty(struct FPropertyChangedEvent& PropertyChangedEvent) override; 
diff --git a/Engine/Plugins/FX/Niagara/Source/Niagara/Private/NiagaraCommon.cpp b/Engine/Plugins/FX/Niagara/Source/Niagara/Private/NiagaraCommon.cpp
--- a/Engine/Plugins/FX/Niagara/Source/Niagara/Private/NiagaraCommon.cpp
+++ b/Engine/Plugins/FX/Niagara/Source/Niagara/Private/NiagaraCommon.cpp
@@ -134,7 +134,11 @@ void FNiagaraSystemUpdateContext::AddAll(bool bReInit)
 	{
 		UNiagaraComponent* Comp = *It;
 		check(Comp);
-		AddInternal(Comp, bReInit);
+		bool bIsActive = Comp->IsActive() || Comp->IsRegisteredWithScalabilityManager();
+		if (bIsActive && bOnlyActive)
+		{
+			AddInternal(Comp, bReInit);
+		}
 	}
 }
 
@@ -146,7 +150,11 @@ void FNiagaraSystemUpdateContext::Add(const UNiagaraSystem* System, bool bReInit
 		check(Comp);
 		if (Comp->GetAsset() == System)
 		{
-			AddInternal(Comp, bReInit);
+			bool bIsActive = Comp->IsActive() || Comp->IsRegisteredWithScalabilityManager();
+			if (bIsActive && bOnlyActive)
+			{
+				AddInternal(Comp, bReInit);
+			}
 		}
 	}
 }
@@ -161,8 +169,12 @@ void FNiagaraSystemUpdateContext::Add(const UNiagaraEmitter* Emitter, bool bReIn
 		FNiagaraSystemInstance* SystemInst = Comp->GetSystemInstance();
 		if (SystemInst && SystemInst->UsesEmitter(Emitter))
 		{
-			AddInternal(Comp, bReInit);
-		}		
+			bool bIsActive = Comp->IsActive() || Comp->IsRegisteredWithScalabilityManager();
+			if (bIsActive && bOnlyActive)
+			{
+				AddInternal(Comp, bReInit);
+			}
+		}
 	}
 }
 
@@ -175,7 +187,11 @@ void FNiagaraSystemUpdateContext::Add(const UNiagaraScript* Script, bool bReInit
 		UNiagaraSystem* System = Comp->GetAsset();
 		if (System && System->UsesScript(Script))
 		{
-			AddInternal(Comp, bReInit);
+			bool bIsActive = Comp->IsActive() || Comp->IsRegisteredWithScalabilityManager();
+			if (bIsActive && bOnlyActive)
+			{
+				AddInternal(Comp, bReInit);
+			}
 		}
 	}
 }
@@ -205,7 +221,11 @@ void FNiagaraSystemUpdateContext::Add(const UNiagaraParameterCollection* Collect
 		FNiagaraSystemInstance* SystemInst = Comp->GetSystemInstance();
 		if (SystemInst && SystemInst->UsesCollection(Collection))
 		{
-			AddInternal(Comp, bReInit);
+			bool bIsActive = Comp->IsActive() || Comp->IsRegisteredWithScalabilityManager();
+			if (bIsActive && bOnlyActive)
+			{
+				AddInternal(Comp, bReInit);
+			}
 		}
 	}
 }
diff --git a/Engine/Plugins/FX/Niagara/Source/Niagara/Private/NiagaraComponent.cpp b/Engine/Plugins/FX/Niagara/Source/Niagara/Private/NiagaraComponent.cpp
--- a/Engine/Plugins/FX/Niagara/Source/Niagara/Private/NiagaraComponent.cpp
+++ b/Engine/Plugins/FX/Niagara/Source/Niagara/Private/NiagaraComponent.cpp
@@ -1317,6 +1317,11 @@ void UNiagaraComponent::SendRenderDynamicData_Concurrent()
 			FNiagaraEmitterInstance* EmitterInst = &SystemInstance->GetEmitters()[i].Get();
 			UNiagaraEmitter* Emitter = EmitterInst->GetCachedEmitter();
 
+			if(Emitter == nullptr)
+			{
+				continue;
+			}
+
 #if STATS
 			TStatId EmitterStatID = Emitter->GetStatID(true, true);
 			FScopeCycleCounter EmitterStatCounter(EmitterStatID);
diff --git a/Engine/Plugins/FX/Niagara/Source/Niagara/Private/NiagaraEmitter.cpp b/Engine/Plugins/FX/Niagara/Source/Niagara/Private/NiagaraEmitter.cpp
--- a/Engine/Plugins/FX/Niagara/Source/Niagara/Private/NiagaraEmitter.cpp
+++ b/Engine/Plugins/FX/Niagara/Source/Niagara/Private/NiagaraEmitter.cpp
@@ -589,6 +589,21 @@ bool UNiagaraEmitter::IsEditorOnly() const
 #endif
 }
 
+bool UNiagaraEmitter::NeedsLoadForTargetPlatform(const ITargetPlatform* TargetPlatform)const
+{
+	if (!FNiagaraPlatformSet::ShouldPruneEmittersOnCook(TargetPlatform->IniPlatformName()))
+	{
+		return true;
+	}
+
+	bool bIsEnabled = IsEnabledOnPlatform(TargetPlatform->IniPlatformName());
+	if(!bIsEnabled)
+	{
+		UE_LOG(LogNiagara, Verbose, TEXT("Pruned emitter %s for platform %s"), *GetFullName(), *TargetPlatform->DisplayName().ToString())
+	}
+	return bIsEnabled;
+}
+
 #if WITH_EDITOR
 /** Creates a new emitter with the supplied emitter as a parent emitter and the supplied system as it's owner. */
 UNiagaraEmitter* UNiagaraEmitter::CreateWithParentAndOwner(UNiagaraEmitter& InParentEmitter, UObject* InOwner, FName InName, EObjectFlags FlagMask)
@@ -744,14 +759,12 @@ UNiagaraEmitter::FOnPropertiesChanged& UNiagaraEmitter::OnRenderersChanged()
 {
 	return OnRenderersChangedDelegate;
 }
+#endif
 
-bool UNiagaraEmitter::IsEnabledOnPlatform(const FString& PlatformName)
+bool UNiagaraEmitter::IsEnabledOnPlatform(const FString& PlatformName)const
 {
-	bool bCanPrune = FNiagaraPlatformSet::ShouldPruneEmittersOnCook(PlatformName);
-	return bCanPrune && Platforms.IsEnabledForPlatform(PlatformName);
+	return Platforms.IsEnabledForPlatform(PlatformName);
 }
-#endif
-
 
 bool UNiagaraEmitter::IsValid()const
 {
diff --git a/Engine/Plugins/FX/Niagara/Source/Niagara/Private/NiagaraEmitterInstance.cpp b/Engine/Plugins/FX/Niagara/Source/Niagara/Private/NiagaraEmitterInstance.cpp
--- a/Engine/Plugins/FX/Niagara/Source/Niagara/Private/NiagaraEmitterInstance.cpp
+++ b/Engine/Plugins/FX/Niagara/Source/Niagara/Private/NiagaraEmitterInstance.cpp
@@ -185,7 +185,7 @@ FBox FNiagaraEmitterInstance::GetBounds()
 
 bool FNiagaraEmitterInstance::IsReadyToRun() const
 {
-	if (!CachedEmitter->IsReadyToRun())
+	if (!IsDisabled() && !CachedEmitter->IsReadyToRun())
 	{
 		return false;
 	}
@@ -588,6 +588,11 @@ void FNiagaraEmitterInstance::CheckForErrors()
 
 void FNiagaraEmitterInstance::DirtyDataInterfaces()
 {
+	if (IsDisabled())
+	{
+		return;
+	}
+
 	// Make sure that our function tables need to be regenerated...
 	SpawnExecContext.DirtyDataInterfaces();
 	UpdateExecContext.DirtyDataInterfaces();
diff --git a/Engine/Plugins/FX/Niagara/Source/Niagara/Private/NiagaraModule.cpp b/Engine/Plugins/FX/Niagara/Source/Niagara/Private/NiagaraModule.cpp
--- a/Engine/Plugins/FX/Niagara/Source/Niagara/Private/NiagaraModule.cpp
+++ b/Engine/Plugins/FX/Niagara/Source/Niagara/Private/NiagaraModule.cpp
@@ -40,7 +40,7 @@ static FAutoConsoleVariableRef CVarLogCompileIdGeneration(
 
 float INiagaraModule::EngineGlobalSpawnCountScale = 1.0f;
 float INiagaraModule::EngineGlobalSystemCountScale = 1.0f;
-int32 INiagaraModule::EngineEffectsQuality = 4;
+int32 INiagaraModule::EngineEffectsQuality = INDEX_NONE;
 
 int32 GEnableVerboseNiagaraChangeIdLogging = 0;
 static FAutoConsoleVariableRef CVarEnableVerboseNiagaraChangeIdLogging(
@@ -150,6 +150,10 @@ void INiagaraModule::StartupModule()
 
 	FNiagaraWorldManager::OnStartup();
 
+	//Force update on module load.
+	static const auto CVarEQ = IConsoleManager::Get().FindTConsoleVariableDataInt(TEXT("sg.EffectsQuality"));
+	OnEffectsQualityChanged(CVarEQ->GetValueOnGameThread(), true);
+
 #if WITH_EDITOR	
 	// Loading uncooked data in a game environment, we still need to get some functionality from the NiagaraEditor module.
 	// This includes the ability to compile scripts and load WITH_EDITOR_ONLY data.
@@ -390,19 +394,22 @@ void INiagaraModule::UnregisterPrecompiler(FDelegateHandle DelegateHandle)
 
 #endif
 
-void INiagaraModule::OnEffectsQualityChanged(int32 NewEffectsQuality)
+void INiagaraModule::OnEffectsQualityChanged(int32 NewEffectsQuality, bool bForce)
 {
-	if (NewEffectsQuality != EngineEffectsQuality && FNiagaraPlatformSet::CanChangeScalabilityAtRuntime())
+	if (FNiagaraPlatformSet::CanChangeScalabilityAtRuntime() || bForce)
 	{
-		EngineEffectsQuality = NewEffectsQuality;
+		if (NewEffectsQuality != EngineEffectsQuality)
+		{
+			EngineEffectsQuality = NewEffectsQuality;
 
-		FNiagaraPlatformSet::InvalidateCachedData();
+			FNiagaraPlatformSet::InvalidateCachedData();
 
-		for (TObjectIterator<UNiagaraSystem> It; It; ++It)
-		{
-			UNiagaraSystem* System = *It;
-			check(System);
-			System->OnEffectsQualityChanged();
+			for (TObjectIterator<UNiagaraSystem> It; It; ++It)
+			{
+				UNiagaraSystem* System = *It;
+				check(System);
+				System->OnEffectsQualityChanged();
+			}
 		}
 	}
 }
@@ -947,10 +954,8 @@ void INiagaraModule::ProcessShaderCompilationQueue()
 
 void CVarSinkFunc()
 {
-	//This Cvar sink can happen before the one which primes the cached scalability cvars so we must grab this ourselves.
-	IConsoleManager& ConsoleMan = IConsoleManager::Get();
-	static const auto CVarEQ = ConsoleMan.FindTConsoleVariableDataInt(TEXT("sg.EffectsQuality"));
-	INiagaraModule::OnEffectsQualityChanged(CVarEQ->GetValueOnGameThread());
+	static const auto CVarEQ = IConsoleManager::Get().FindTConsoleVariableDataInt(TEXT("sg.EffectsQuality"));
+	INiagaraModule::OnEffectsQualityChanged(CVarEQ->GetValueOnGameThread(), false);
 }
 
 static FAutoConsoleVariableSink CVarSink(FConsoleCommandDelegate::CreateStatic(&CVarSinkFunc));
diff --git a/Engine/Plugins/FX/Niagara/Source/Niagara/Private/NiagaraPlatformSet.cpp b/Engine/Plugins/FX/Niagara/Source/Niagara/Private/NiagaraPlatformSet.cpp
--- a/Engine/Plugins/FX/Niagara/Source/Niagara/Private/NiagaraPlatformSet.cpp
+++ b/Engine/Plugins/FX/Niagara/Source/Niagara/Private/NiagaraPlatformSet.cpp
@@ -6,6 +6,11 @@
 #include "DeviceProfiles/DeviceProfileManager.h"
 #include "Scalability.h"
 #include "Misc/ConfigCacheIni.h"
+#include "Interfaces/ITargetPlatform.h"
+
+#if WITH_EDITOR
+#include "PlatformInfo.h"
+#endif
 
 #define LOCTEXT_NAMESPACE "NiagaraPlatformSet"
 
@@ -22,6 +27,15 @@ static FAutoConsoleVariableRef CVarNiagaraAllowRuntimeScalabilityChanges(
 	ECVF_ReadOnly
 );
 
+const TCHAR* PruneEmittersOnCookName = TEXT("fx.Niagara.PruneEmittersOnCook");
+int32 GbPruneEmittersOnCook = 1;
+static FAutoConsoleVariableRef CVarPruneEmittersOnCook(
+	PruneEmittersOnCookName,
+	GbPruneEmittersOnCook,
+	TEXT("If > 0 this platform will prune disabled emitters during cook. \n"),
+	ECVF_ReadOnly
+);
+
 uint32 FNiagaraPlatformSet::LastDirtiedFrame = 0;
 #if WITH_EDITOR
 TMap<const UDeviceProfile*, int32> FNiagaraPlatformSet::CachedEQMasksPerDeviceProfile;
@@ -227,9 +241,7 @@ void FNiagaraPlatformSet::InvalidateCachedData()
 	LastDirtiedFrame = GFrameNumber;
 }
 
-#if WITH_EDITOR
-
-bool FNiagaraPlatformSet::IsEnabledForPlatform(const FString& PlatformName)
+bool FNiagaraPlatformSet::IsEnabledForPlatform(const FString& PlatformName)const
 {
 	for (const UObject* ProfileObj : UDeviceProfileManager::Get().Profiles)
 	{
@@ -249,6 +261,18 @@ bool FNiagaraPlatformSet::IsEnabledForPlatform(const FString& PlatformName)
 	return false;
 }
 
+bool FNiagaraPlatformSet::ShouldPruneEmittersOnCook(const FString& PlatformName)
+{
+#if WITH_EDITOR
+	FPlatformIniSettings Settings = GetPlatformIniSettings(PlatformName);
+	return Settings.bPruneEmittersOnCook != 0;
+#else
+	return GbPruneEmittersOnCook != 0;
+#endif
+}
+
+#if WITH_EDITOR
+
 // bool FNiagaraPlatformSet::Conflicts(const FNiagaraPlatformSet& Other, TArray<const UDeviceProfile*>& OutConflictingProfiles, bool bIncludeProfilesWithVariableEffectsQuality)const
 // {
 // 	for (const UObject* ProfileObj : UDeviceProfileManager::Get().Profiles)
@@ -359,12 +383,6 @@ bool FNiagaraPlatformSet::CanChangeScalabilityAtRuntime(const UDeviceProfile* De
 	return true;//Assuming true if we fail to find the platform seems safest.
 }
 
-bool FNiagaraPlatformSet::ShouldPruneEmittersOnCook(const FString& PlatformName)
-{
-	FPlatformIniSettings Settings = GetPlatformIniSettings(PlatformName);
-	return Settings.bPruneEmittersOnCook != 0;
-}
-
 bool FNiagaraPlatformSet::GatherConflicts(const TArray<const FNiagaraPlatformSet*>& PlatformSets, TArray<FNiagaraPlatformSetConflictInfo>& OutConflicts)
 {
 	Scalability::FQualityLevels NumLevels = Scalability::GetQualityLevelCounts();
@@ -456,7 +474,7 @@ FNiagaraPlatformSet::FPlatformIniSettings FNiagaraPlatformSet::GetPlatformIniSet
 	}
 
 	int32 PruneEmittersOnCook = 0;
-	FindCVarValue(TEXT("SystemSettings"), TEXT("fx.PruneEmittersOnCook"), PruneEmittersOnCook);
+	FindCVarValue(TEXT("SystemSettings"), PruneEmittersOnCookName, PruneEmittersOnCook);
 
 	FPlatformIniSettings& NewSetting = CachedPlatformIniSettings.Add(*PlatformName);
 	NewSetting = FPlatformIniSettings(CanChangeEffectQuality, PruneEmittersOnCook, EQMask);
@@ -470,4 +488,4 @@ int32 FNiagaraPlatformSet::GetEffectQualityMaskForPlatform(const FString& Platfo
 }
 #endif
 
-#undef LOCTEXT_NAMESPACE
\ No newline at end of file
+#undef LOCTEXT_NAMESPACE
diff --git a/Engine/Plugins/FX/Niagara/Source/Niagara/Private/NiagaraSystem.cpp b/Engine/Plugins/FX/Niagara/Source/Niagara/Private/NiagaraSystem.cpp
--- a/Engine/Plugins/FX/Niagara/Source/Niagara/Private/NiagaraSystem.cpp
+++ b/Engine/Plugins/FX/Niagara/Source/Niagara/Private/NiagaraSystem.cpp
@@ -43,14 +43,6 @@ static FAutoConsoleVariableRef CVarNiagaraForceSystemsToCookOutRapidIterationOnL
 	ECVF_Default
 );
 
-static int GNiagaraVerboseCookOutput = 0;
-static FAutoConsoleVariableRef CVarNiagaraVerboseCookOutput(
-	TEXT("fx.NiagaraVerboseCookOutput"),
-	GNiagaraVerboseCookOutput,
-	TEXT("When enabled UNiagaraSystems generate verbose output on cook."),
-	ECVF_Default
-);
-
 //////////////////////////////////////////////////////////////////////////
 
 UNiagaraSystem::UNiagaraSystem(const FObjectInitializer& ObjectInitializer)
@@ -99,6 +91,21 @@ void UNiagaraSystem::PreSave(const class ITargetPlatform * TargetPlatform)
 #endif
 }
 
+bool UNiagaraSystem::NeedsLoadForTargetPlatform(const ITargetPlatform* TargetPlatform)const
+{
+	bool bHasAnyEnabledEmitters = false;
+	for (const FNiagaraEmitterHandle& EmitterHandle : GetEmitterHandles())
+	{
+		if (EmitterHandle.GetIsEnabled() && EmitterHandle.GetInstance()->Platforms.IsEnabledForPlatform(TargetPlatform->IniPlatformName()))
+		{
+			bHasAnyEnabledEmitters = true;
+			break;
+		}
+	}
+
+	return bHasAnyEnabledEmitters;
+}
+
 #if WITH_EDITOR
 void UNiagaraSystem::BeginCacheForCookedPlatformData(const ITargetPlatform *TargetPlatform)
 {
@@ -210,67 +217,8 @@ void UNiagaraSystem::RequestCompileForEmitter(UNiagaraEmitter* InEmitter)
 
 void UNiagaraSystem::Serialize(FArchive& Ar)
 {
-	//TODO: Cook out emitters. This requires some wrangling on the compilation side for system scripts.
-// 	//When cooking we clear out any emitters that are disabled entirely or just for the target platform.
-// 	TArray<FNiagaraEmitterHandle> OldEmitters;
-// 	if (Ar.IsCooking())
-// 	{
-// 		OldEmitters = EmitterHandles;
-// 		const ITargetPlatform* Target = Ar.CookingTarget();
-// 		for (FNiagaraEmitterHandle& Handle : EmitterHandles)
-// 		{
-// 			if (Handle.GetIsEnabled() && Handle.GetInstance()->IsEnabledOnPlatform(Target->PlatformName()))
-// 			{
-// 				Handle.ClearEmitter();
-// 			}
-// 		}
-// 	}
-
-	//Some log spam to display emitters we should be cooking out.
-#if WITH_EDITOR && 0
-	if (Ar.IsCooking() && GNiagaraVerboseCookOutput)
-	{
-		TBitArray<TInlineAllocator<8>> EnabledEmitters;
-		EnabledEmitters.Reserve(EmitterHandles.Num());
-		int32 NumEnabled = 0;
-		const ITargetPlatform* Target = Ar.CookingTarget();
-		for (FNiagaraEmitterHandle& Handle : EmitterHandles)
-		{
-			bool bEnabled = Handle.GetInstance()->IsEnabledOnPlatform(Target->PlatformName());
-			EnabledEmitters.Add(bEnabled);
-			if (bEnabled)
-			{
-				NumEnabled++;
-			}
-		}
-
-		if (NumEnabled != EmitterHandles.Num())
-		{
-			UE_LOG(LogNiagara, Log, TEXT("==============================================="));
-			UE_LOG(LogNiagara, Log, TEXT("System has emitters that could be cooked out!"), *GetName());
-
-			for (int32 EmitterIdx = 0; EmitterIdx < EmitterHandles.Num(); ++EmitterIdx)
-			{
-				if (EnabledEmitters[EmitterIdx] == false)
-				{
-					UE_LOG(LogNiagara, Log, TEXT("%s"), *EmitterHandles[EmitterIdx].GetName().ToString());
-				}
-			}
-			UE_LOG(LogNiagara, Log, TEXT("==============================================="));
-		}
-	}
-
-#endif
-
 	Super::Serialize(Ar);
 
-	//TODO: Cook out emitters. This requires some wrangling on the compilation side for system scripts.
-// 	//Restore old emitters post cook pruning.
-// 	if (Ar.IsCooking())
-// 	{
-// 		EmitterHandles = OldEmitters;
-// 	}
-
 	Ar.UsingCustomVersion(FNiagaraCustomVersion::GUID);
 
 	if (Ar.CustomVer(FNiagaraCustomVersion::GUID) >= FNiagaraCustomVersion::ChangeEmitterCompiledDataToSharedRefs)
@@ -1534,10 +1482,15 @@ void UNiagaraSystem::OnEffectsQualityChanged()
 
 	for (FNiagaraEmitterHandle& Handle : EmitterHandles)
 	{
-		Handle.GetInstance()->OnEffectsQualityChanged();
+		if (Handle.GetInstance())
+		{
+			Handle.GetInstance()->OnEffectsQualityChanged();
+		}
 	}
 
 	FNiagaraSystemUpdateContext UpdateCtx;
+	UpdateCtx.SetDestroyOnAdd(true);
+	UpdateCtx.SetOnlyActive(true);
 	UpdateCtx.Add(this, true);
 }
 
diff --git a/Engine/Plugins/FX/Niagara/Source/Niagara/Private/NiagaraSystemInstance.cpp b/Engine/Plugins/FX/Niagara/Source/Niagara/Private/NiagaraSystemInstance.cpp
--- a/Engine/Plugins/FX/Niagara/Source/Niagara/Private/NiagaraSystemInstance.cpp
+++ b/Engine/Plugins/FX/Niagara/Source/Niagara/Private/NiagaraSystemInstance.cpp
@@ -1356,6 +1356,11 @@ void FNiagaraSystemInstance::InitDataInterfaces()
 	for (TSharedRef<FNiagaraEmitterInstance, ESPMode::ThreadSafe> Simulation : Emitters)
 	{
 		FNiagaraEmitterInstance& Sim = Simulation.Get();
+		if (Sim.IsDisabled())
+		{
+			continue;
+		}
+
 		CalcInstDataSize(Sim.GetSpawnExecutionContext().GetDataInterfaces());
 		CalcInstDataSize(Sim.GetUpdateExecutionContext().GetDataInterfaces());
 		for (int32 i = 0; i < Sim.GetEventExecutionContexts().Num(); i++)
diff --git a/Engine/Plugins/FX/Niagara/Source/Niagara/Private/NiagaraSystemSimulation.cpp b/Engine/Plugins/FX/Niagara/Source/Niagara/Private/NiagaraSystemSimulation.cpp
--- a/Engine/Plugins/FX/Niagara/Source/Niagara/Private/NiagaraSystemSimulation.cpp
+++ b/Engine/Plugins/FX/Niagara/Source/Niagara/Private/NiagaraSystemSimulation.cpp
@@ -1961,29 +1961,32 @@ void FNiagaraSystemSimulation::InitParameterDataSetBindings(FNiagaraSystemInstan
 		for (int32 EmitterIdx = 0; EmitterIdx < EmitterCount; ++EmitterIdx)
 		{
 			FNiagaraEmitterInstance& EmitterInst = Emitters[EmitterIdx].Get();
-			const FString EmitterName = EmitterInst.GetCachedEmitter()->GetUniqueEmitterName();
+			if (!EmitterInst.IsDisabled())
+			{
+				const FString EmitterName = EmitterInst.GetCachedEmitter()->GetUniqueEmitterName();
 
-			FNiagaraScriptExecutionContext& SpawnContext = EmitterInst.GetSpawnExecutionContext();
-			DataSetToEmitterSpawnParameters[EmitterIdx].Init(MainDataSet, SpawnContext.Parameters);
+				FNiagaraScriptExecutionContext& SpawnContext = EmitterInst.GetSpawnExecutionContext();
+				DataSetToEmitterSpawnParameters[EmitterIdx].Init(MainDataSet, SpawnContext.Parameters);
 
-			FNiagaraScriptExecutionContext& UpdateContext = EmitterInst.GetUpdateExecutionContext();
-			DataSetToEmitterUpdateParameters[EmitterIdx].Init(MainDataSet, UpdateContext.Parameters);
+				FNiagaraScriptExecutionContext& UpdateContext = EmitterInst.GetUpdateExecutionContext();
+				DataSetToEmitterUpdateParameters[EmitterIdx].Init(MainDataSet, UpdateContext.Parameters);
 
 
-			FNiagaraComputeExecutionContext* GPUContext = EmitterInst.GetGPUContext();
-			if (GPUContext)
-			{
-				DataSetToEmitterGPUParameters[EmitterIdx].Init(MainDataSet, GPUContext->CombinedParamStore);
-			}
+				FNiagaraComputeExecutionContext* GPUContext = EmitterInst.GetGPUContext();
+				if (GPUContext)
+				{
+					DataSetToEmitterGPUParameters[EmitterIdx].Init(MainDataSet, GPUContext->CombinedParamStore);
+				}
 
-			TArray<FNiagaraScriptExecutionContext>& EventContexts = EmitterInst.GetEventExecutionContexts();
-			const int32 EventCount = EventContexts.Num();
-			DataSetToEmitterEventParameters[EmitterIdx].SetNum(EventCount);
+				TArray<FNiagaraScriptExecutionContext>& EventContexts = EmitterInst.GetEventExecutionContexts();
+				const int32 EventCount = EventContexts.Num();
+				DataSetToEmitterEventParameters[EmitterIdx].SetNum(EventCount);
 
-			for (int32 EventIdx = 0; EventIdx < EventCount; ++EventIdx)
-			{
-				FNiagaraScriptExecutionContext& EventContext = EventContexts[EventIdx];
-				DataSetToEmitterEventParameters[EmitterIdx][EventIdx].Init(MainDataSet, EventContext.Parameters);
+				for (int32 EventIdx = 0; EventIdx < EventCount; ++EventIdx)
+				{
+					FNiagaraScriptExecutionContext& EventContext = EventContexts[EventIdx];
+					DataSetToEmitterEventParameters[EmitterIdx][EventIdx].Init(MainDataSet, EventContext.Parameters);
+				}
 			}
 		}
 	}
diff --git a/Engine/Plugins/FX/Niagara/Source/Niagara/Public/NiagaraCommon.h b/Engine/Plugins/FX/Niagara/Source/Niagara/Public/NiagaraCommon.h
--- a/Engine/Plugins/FX/Niagara/Source/Niagara/Public/NiagaraCommon.h
+++ b/Engine/Plugins/FX/Niagara/Source/Niagara/Public/NiagaraCommon.h
@@ -547,18 +547,19 @@ struct NIAGARA_API FNiagaraSystemUpdateContext
 {
 	GENERATED_BODY()
 
-	FNiagaraSystemUpdateContext(const UNiagaraSystem* System, bool bReInit, bool bInDestroyOnAdd = false) :bDestroyOnAdd(bInDestroyOnAdd) { Add(System, bReInit); }
+	FNiagaraSystemUpdateContext(const UNiagaraSystem* System, bool bReInit) :bDestroyOnAdd(false), bOnlyActive(false) { Add(System, bReInit); }
 #if WITH_EDITORONLY_DATA
-	FNiagaraSystemUpdateContext(const UNiagaraEmitter* Emitter, bool bReInit, bool bInDestroyOnAdd = false) : bDestroyOnAdd(bInDestroyOnAdd) { Add(Emitter, bReInit); }
-	FNiagaraSystemUpdateContext(const UNiagaraScript* Script, bool bReInit, bool bInDestroyOnAdd = false) :bDestroyOnAdd(bInDestroyOnAdd) { Add(Script, bReInit); }
+	FNiagaraSystemUpdateContext(const UNiagaraEmitter* Emitter, bool bReInit) : bDestroyOnAdd(false), bOnlyActive(false) { Add(Emitter, bReInit); }
+	FNiagaraSystemUpdateContext(const UNiagaraScript* Script, bool bReInit) :bDestroyOnAdd(false), bOnlyActive(false) { Add(Script, bReInit); }
 	//FNiagaraSystemUpdateContext(UNiagaraDataInterface* Interface, bool bReinit) : Add(Interface, bReinit) {}
-	FNiagaraSystemUpdateContext(const UNiagaraParameterCollection* Collection, bool bReInit, bool bInDestroyOnAdd = false) :bDestroyOnAdd(bInDestroyOnAdd) { Add(Collection, bReInit); }
+	FNiagaraSystemUpdateContext(const UNiagaraParameterCollection* Collection, bool bReInit) :bDestroyOnAdd(false), bOnlyActive(false) { Add(Collection, bReInit); }
 #endif
 	FNiagaraSystemUpdateContext():bDestroyOnAdd(false){ }
 
 	~FNiagaraSystemUpdateContext();
 
 	void SetDestroyOnAdd(bool bInDestroyOnAdd) { bDestroyOnAdd = bInDestroyOnAdd; }
+	void SetOnlyActive(bool bInOnlyActive) { bOnlyActive = bInOnlyActive; }
 
 	void Add(const UNiagaraSystem* System, bool bReInit);
 #if WITH_EDITORONLY_DATA
@@ -586,6 +587,7 @@ private:
 	TArray<UNiagaraSystem*> SystemSimsToDestroy;
 
 	bool bDestroyOnAdd;
+	bool bOnlyActive;
 	//TODO: When we allow component less systems we'll also want to find and reset those.
 };
 
diff --git a/Engine/Plugins/FX/Niagara/Source/Niagara/Public/NiagaraComponent.h b/Engine/Plugins/FX/Niagara/Source/Niagara/Public/NiagaraComponent.h
--- a/Engine/Plugins/FX/Niagara/Source/Niagara/Public/NiagaraComponent.h
+++ b/Engine/Plugins/FX/Niagara/Source/Niagara/Public/NiagaraComponent.h
@@ -521,6 +521,7 @@ public:
 	UFUNCTION(BlueprintCallable, Category = Scalability, meta = (Keywords = "LOD scalability"))
 	void SetAllowScalability(bool bAllow);
 
+	FORCEINLINE bool IsRegisteredWithScalabilityManager()const { return ScalabilityManagerHandle != INDEX_NONE; }
 	FORCEINLINE int32 GetScalabilityManagerHandle()const { return ScalabilityManagerHandle; }
 private:
 	/** Did we try and activate but fail due to the asset being not yet ready. Keep looping.*/
diff --git a/Engine/Plugins/FX/Niagara/Source/Niagara/Public/NiagaraModule.h b/Engine/Plugins/FX/Niagara/Source/Niagara/Public/NiagaraModule.h
--- a/Engine/Plugins/FX/Niagara/Source/Niagara/Public/NiagaraModule.h
+++ b/Engine/Plugins/FX/Niagara/Source/Niagara/Public/NiagaraModule.h
@@ -165,7 +165,7 @@ public:
 	FORCEINLINE static const FNiagaraVariable&  GetVar_DataInstance_Alive() { return DataInstance_Alive; }
 	FORCEINLINE static const FNiagaraVariable&  GetVar_BeginDefaults() { return Translator_BeginDefaults; }
 
-	static void OnEffectsQualityChanged(int32 NewEffectsQuality);
+	static void OnEffectsQualityChanged(int32 NewEffectsQuality, bool bForce);
 
 	FOnProcessQueue OnProcessQueue;
 
diff --git a/Engine/Source/Runtime/Engine/Private/DeviceProfiles/DeviceProfileManager.cpp b/Engine/Source/Runtime/Engine/Private/DeviceProfiles/DeviceProfileManager.cpp
--- a/Engine/Source/Runtime/Engine/Private/DeviceProfiles/DeviceProfileManager.cpp
+++ b/Engine/Source/Runtime/Engine/Private/DeviceProfiles/DeviceProfileManager.cpp
@@ -559,7 +559,7 @@ void UDeviceProfileManager::LoadProfiles()
 				ITargetPlatform* Platform = TargetPlatforms[PlatformIndex];
 
 				// Set TextureLODSettings
-				const UTextureLODSettings* TextureLODSettingsObj = FindProfile(*Platform->GetPlatformInfo().VanillaPlatformName.ToString());
+				const UTextureLODSettings* TextureLODSettingsObj = FindProfile(*Platform->GetPlatformInfo().VanillaPlatformName.ToString(), false);
 				Platform->RegisterTextureLODSettings(TextureLODSettingsObj);
 			}
 		}
