Bump Niagara custom version to fix resaved script assets failing to postload their UNiagaraGraphs and asserting.
Fixup adding parameters for Array namespace to go through normal AddParameter() flow to fixup missing metadata on these entries.
Skip adding parameters to ParameterMapHistoryHandles if they do not have a valid Scope Name.

diff --git a/Engine/Plugins/FX/Niagara/Source/NiagaraCore/Public/NiagaraCustomVersion.h b/Engine/Plugins/FX/Niagara/Source/NiagaraCore/Public/NiagaraCustomVersion.h
--- a/Engine/Plugins/FX/Niagara/Source/NiagaraCore/Public/NiagaraCustomVersion.h
+++ b/Engine/Plugins/FX/Niagara/Source/NiagaraCore/Public/NiagaraCustomVersion.h
@@ -146,6 +146,8 @@ struct FNiagaraCustomVersion
 
 		FixNullScriptVariables, // Postload fixup in UNiagaraGraph to fixup VariableToScriptVariable map entries being null. 
 
+		PrecompileNamespaceFixup2, // Move FNiagaraVariableMetaData from storing scope enum to storing registered scope name.
+
 		// DO NOT ADD A NEW VERSION UNLESS YOU HAVE TALKED TO THE NIAGARA LEAD. Mismanagement of these versions can lead to data loss if it is adjusted in multiple streams simultaneously.
 		// -----<new versions can be added above this line>  -------------------------------------------------
 		VersionPlusOne,
diff --git a/Engine/Plugins/FX/Niagara/Source/NiagaraEditor/Private/NiagaraGraph.cpp b/Engine/Plugins/FX/Niagara/Source/NiagaraEditor/Private/NiagaraGraph.cpp
--- a/Engine/Plugins/FX/Niagara/Source/NiagaraEditor/Private/NiagaraGraph.cpp
+++ b/Engine/Plugins/FX/Niagara/Source/NiagaraEditor/Private/NiagaraGraph.cpp
@@ -110,6 +110,11 @@ void UNiagaraGraph::RemoveOnGraphNeedsRecompileHandler(FDelegateHandle Handle)
 	OnGraphNeedsRecompile.Remove(Handle);
 }
 
+void UNiagaraGraph::MutableAddParameter(FNiagaraVariable& Parameter, const FAddParameterOptions Options) const
+{
+	const_cast<UNiagaraGraph*>(this)->AddParameter(Parameter, Options);
+}
+
 void UNiagaraGraph::NotifyGraphChanged(const FEdGraphEditAction& InAction)
 {
 	InvalidateCachedParameterData();
@@ -343,7 +348,7 @@ void UNiagaraGraph::PostLoad()
 		}
 	}
 
-	if (NiagaraVer < FNiagaraCustomVersion::PrecompileNamespaceFixup)
+	if (NiagaraVer < FNiagaraCustomVersion::PrecompileNamespaceFixup2)
 	{
 		// Collect all input and output pins to infer usages of variables they reference.
 		const TMap<FNiagaraVariable, FInputPinsAndOutputPins> VarToPinsMap = CollectVarsToInOutPinsMap();
@@ -2152,11 +2157,10 @@ void UNiagaraGraph::RefreshParameterReferences() const
 		}
 		for (const FString& RegisterName : RegisterNames)
 		{
+			FAddParameterOptions AddParameterOptions = FAddParameterOptions();
+			AddParameterOptions.bRefreshMetaDataScopeAndUsage = true;
 			FNiagaraVariable Parameter = FNiagaraVariable(FNiagaraTypeDefinition::GetIntDef(), *RegisterName);
-			if (ParameterToReferencesMap.Find(Parameter) == nullptr)
-			{
-				ParameterToReferencesMap.Add(Parameter, false);
-			}
+			MutableAddParameter(Parameter, AddParameterOptions);
 		}
 	}
 
diff --git a/Engine/Plugins/FX/Niagara/Source/NiagaraEditor/Private/NiagaraParameterMapHistory.cpp b/Engine/Plugins/FX/Niagara/Source/NiagaraEditor/Private/NiagaraParameterMapHistory.cpp
--- a/Engine/Plugins/FX/Niagara/Source/NiagaraEditor/Private/NiagaraParameterMapHistory.cpp
+++ b/Engine/Plugins/FX/Niagara/Source/NiagaraEditor/Private/NiagaraParameterMapHistory.cpp
@@ -1958,6 +1958,12 @@ FNiagaraParameterMapHistoryHandle::FNiagaraParameterMapHistoryHandle(FNiagaraPar
 			// Cannot fixup names of variables using legacy name string mode, skip this entry.
 			continue;
 		}
+		else if (VarMetaData.GetScopeName().IsNone())
+		{
+			// cannot make history variables for variables with metadata without a scope name. Log this for tracking and skip this entry.
+			UE_LOG(LogNiagaraEditor, Display, TEXT("Variable MetaData added to ParameterMapHistoryHandle had empty Scope Name. Variable: %s"), *Var->GetName().ToString());
+			continue;
+		}
 		HistoryVariables.Add(FNiagaraHistoryVariable(Var, VarWithOriginalAliasIntact, VarMetaData));
 	}
 }
@@ -1989,6 +1995,12 @@ FNiagaraParameterMapHistoryHandle::FNiagaraParameterMapHistoryHandle(FNiagaraPar
 				// Cannot fixup names of variables using legacy name string mode, skip this entry.
 				continue;
 			}
+			else if (ScriptVarMetaData.GetScopeName().IsNone())
+			{
+				// cannot make history variables for variables with metadata without a scope name. Log this for tracking and skip this entry.
+				UE_LOG(LogNiagaraEditor, Display, TEXT("Variable MetaData added to ParameterMapHistoryHandle had empty Scope Name. Variable: %s"), *SecondaryVar.GetName().ToString());
+				continue;
+			}
 			HistoryVariables.Add(FNiagaraHistoryVariable(PrimaryVar, PrimaryVarWithOriginalAliasIntact, ScriptVarMetaData));
 		}
 	}
diff --git a/Engine/Plugins/FX/Niagara/Source/NiagaraEditor/Public/NiagaraGraph.h b/Engine/Plugins/FX/Niagara/Source/NiagaraEditor/Public/NiagaraGraph.h
--- a/Engine/Plugins/FX/Niagara/Source/NiagaraEditor/Public/NiagaraGraph.h
+++ b/Engine/Plugins/FX/Niagara/Source/NiagaraEditor/Public/NiagaraGraph.h
@@ -351,6 +351,10 @@ protected:
 	bool AppendCompileHash(FNiagaraCompileHashVisitor* InVisitor, const TArray<UNiagaraNode*>& InTraversal) const;
 
 private:
+
+	/** The const is a lie! Public implementation AddParameter interacts with mutable members only. Private implementation is provided to handle adding parameters in RefreshParameterReferences(). */
+	void MutableAddParameter(FNiagaraVariable& Parameter, const FAddParameterOptions Options) const;
+
 	virtual void NotifyGraphChanged(const FEdGraphEditAction& InAction) override;
 
 	/** Find parameters in the graph. */
