Niagara - Fix a bug with copy/paste where inputs behind static switches could get missed when pasting modules or multiple input values.  Also fix a sefety issue with the stack function input binder which could cause crash on undo.

diff --git a/Engine/Plugins/FX/Niagara/Source/NiagaraEditor/Private/NiagaraStackFunctionInputBinder.cpp b/Engine/Plugins/FX/Niagara/Source/NiagaraEditor/Private/NiagaraStackFunctionInputBinder.cpp
--- a/Engine/Plugins/FX/Niagara/Source/NiagaraEditor/Private/NiagaraStackFunctionInputBinder.cpp
+++ b/Engine/Plugins/FX/Niagara/Source/NiagaraEditor/Private/NiagaraStackFunctionInputBinder.cpp
@@ -291,6 +291,13 @@ FGuid FNiagaraStackFunctionInputBinder::GetChangeIdFromFunctionScript() const
 
 void FNiagaraStackFunctionInputBinder::RefreshGraphPins() const
 {
+	if (FunctionCallNode->FunctionScript == nullptr || FunctionCallNode->FunctionScript->GetSource() == nullptr)
+	{
+		OverridePin = nullptr;
+		DefaultPin = nullptr;
+		return;
+	}
+
 	OverridePin = FNiagaraStackGraphUtilities::GetStackFunctionInputOverridePin(*FunctionCallNode, AliasedParameterHandle);
 	DefaultPin = FunctionCallNode->FindParameterMapDefaultValuePin(InputName, Script->GetUsage());
 	
diff --git a/Engine/Plugins/FX/Niagara/Source/NiagaraEditor/Private/ViewModels/Stack/NiagaraStackFunctionInputCollection.cpp b/Engine/Plugins/FX/Niagara/Source/NiagaraEditor/Private/ViewModels/Stack/NiagaraStackFunctionInputCollection.cpp
--- a/Engine/Plugins/FX/Niagara/Source/NiagaraEditor/Private/ViewModels/Stack/NiagaraStackFunctionInputCollection.cpp
+++ b/Engine/Plugins/FX/Niagara/Source/NiagaraEditor/Private/ViewModels/Stack/NiagaraStackFunctionInputCollection.cpp
@@ -90,11 +90,22 @@ void UNiagaraStackFunctionInputCollection::ToClipboardFunctionInputs(UObject* In
 
 void UNiagaraStackFunctionInputCollection::SetValuesFromClipboardFunctionInputs(const TArray<const UNiagaraClipboardFunctionInput*>& ClipboardFunctionInputs)
 {
+	TArray<const UNiagaraClipboardFunctionInput*> StaticSwitchInputs;
+	TArray<const UNiagaraClipboardFunctionInput*> StandardInputs;
+
 	TArray<UNiagaraStackInputCategory*> ChildCategories;
 	GetUnfilteredChildrenOfType(ChildCategories);
+
+	// Set static switches first so that other inputs will be available to set.
+	for (UNiagaraStackInputCategory* ChildCategory : ChildCategories)
+	{
+		ChildCategory->SetStaticSwitchValuesFromClipboardFunctionInputs(ClipboardFunctionInputs);
+	}
+
+	RefreshChildren();
 	for (UNiagaraStackInputCategory* ChildCategory : ChildCategories)
 	{
-		ChildCategory->SetValuesFromClipboardFunctionInputs(ClipboardFunctionInputs);
+		ChildCategory->SetStandardValuesFromClipboardFunctionInputs(ClipboardFunctionInputs);
 	}
 }
 
diff --git a/Engine/Plugins/FX/Niagara/Source/NiagaraEditor/Private/ViewModels/Stack/NiagaraStackInputCategory.cpp b/Engine/Plugins/FX/Niagara/Source/NiagaraEditor/Private/ViewModels/Stack/NiagaraStackInputCategory.cpp
--- a/Engine/Plugins/FX/Niagara/Source/NiagaraEditor/Private/ViewModels/Stack/NiagaraStackInputCategory.cpp
+++ b/Engine/Plugins/FX/Niagara/Source/NiagaraEditor/Private/ViewModels/Stack/NiagaraStackInputCategory.cpp
@@ -107,15 +107,16 @@ void UNiagaraStackInputCategory::ToClipboardFunctionInputs(UObject* InOuter, TAr
 	}
 }
 
-void  UNiagaraStackInputCategory::SetValuesFromClipboardFunctionInputs(const TArray<const UNiagaraClipboardFunctionInput*>& ClipboardFunctionInputs)
+template<typename Predicate>
+void SetValuesFromFunctionInputsInternal(UNiagaraStackInputCategory* Category, const TArray<const UNiagaraClipboardFunctionInput*>& ClipboardFunctionInputs, Predicate InputMatchesFilter)
 {
 	TArray<UNiagaraStackFunctionInput*> ChildInputs;
-	GetUnfilteredChildrenOfType(ChildInputs);
+	Category->GetUnfilteredChildrenOfType(ChildInputs);
 	for (UNiagaraStackFunctionInput* ChildInput : ChildInputs)
 	{
 		for (const UNiagaraClipboardFunctionInput* ClipboardFunctionInput : ClipboardFunctionInputs)
 		{
-			if (ChildInput->GetInputParameterHandle().GetName() == ClipboardFunctionInput->InputName && ChildInput->GetInputType() == ClipboardFunctionInput->InputType)
+			if (InputMatchesFilter(ChildInput) && ChildInput->GetInputParameterHandle().GetName() == ClipboardFunctionInput->InputName && ChildInput->GetInputType() == ClipboardFunctionInput->InputType)
 			{
 				ChildInput->SetValueFromClipboardFunctionInput(*ClipboardFunctionInput);
 			}
@@ -123,6 +124,16 @@ void  UNiagaraStackInputCategory::SetValuesFromClipboardFunctionInputs(const TAr
 	}
 }
 
+void  UNiagaraStackInputCategory::SetStaticSwitchValuesFromClipboardFunctionInputs(const TArray<const UNiagaraClipboardFunctionInput*>& ClipboardFunctionInputs)
+{
+	SetValuesFromFunctionInputsInternal(this, ClipboardFunctionInputs, [](UNiagaraStackFunctionInput* ChildInput) { return ChildInput->IsStaticParameter(); });
+}
+
+void  UNiagaraStackInputCategory::SetStandardValuesFromClipboardFunctionInputs(const TArray<const UNiagaraClipboardFunctionInput*>& ClipboardFunctionInputs)
+{
+	SetValuesFromFunctionInputsInternal(this, ClipboardFunctionInputs, [](UNiagaraStackFunctionInput* ChildInput) { return ChildInput->IsStaticParameter() == false; });
+}
+
 bool UNiagaraStackInputCategory::FilterForVisibleCondition(const UNiagaraStackEntry& Child) const
 {
 	const UNiagaraStackFunctionInput* StackFunctionInputChild = Cast<UNiagaraStackFunctionInput>(&Child);
diff --git a/Engine/Plugins/FX/Niagara/Source/NiagaraEditor/Public/ViewModels/Stack/NiagaraStackInputCategory.h b/Engine/Plugins/FX/Niagara/Source/NiagaraEditor/Public/ViewModels/Stack/NiagaraStackInputCategory.h
--- a/Engine/Plugins/FX/Niagara/Source/NiagaraEditor/Public/ViewModels/Stack/NiagaraStackInputCategory.h
+++ b/Engine/Plugins/FX/Niagara/Source/NiagaraEditor/Public/ViewModels/Stack/NiagaraStackInputCategory.h
@@ -45,9 +45,12 @@ public:
 
 	void ToClipboardFunctionInputs(UObject* InOuter, TArray<const UNiagaraClipboardFunctionInput*>& OutClipboardFunctionInputs) const;
 
-	void SetValuesFromClipboardFunctionInputs(const TArray<const UNiagaraClipboardFunctionInput*>& ClipboardFunctionInputs);
+	void SetStaticSwitchValuesFromClipboardFunctionInputs(const TArray<const UNiagaraClipboardFunctionInput*>& ClipboardFunctionInputs);
 
+	void SetStandardValuesFromClipboardFunctionInputs(const TArray<const UNiagaraClipboardFunctionInput*>& ClipboardFunctionInputs);
 protected:
+
+
 	//~ UNiagaraStackEntry interface
 	virtual void RefreshChildrenInternal(const TArray<UNiagaraStackEntry*>& CurrentChildren, TArray<UNiagaraStackEntry*>& NewChildren, TArray<FStackIssue>& NewIssues) override;
 
