Fix custom hlsl nodes failing to resolve pin names.
Prevent renaming parameters by renaming their pin names directly. Context menu still provides option to rename pins.

diff --git a/Engine/Plugins/FX/Niagara/Source/NiagaraEditor/Private/NiagaraNodeWithDynamicPins.cpp b/Engine/Plugins/FX/Niagara/Source/NiagaraEditor/Private/NiagaraNodeWithDynamicPins.cpp
--- a/Engine/Plugins/FX/Niagara/Source/NiagaraEditor/Private/NiagaraNodeWithDynamicPins.cpp
+++ b/Engine/Plugins/FX/Niagara/Source/NiagaraEditor/Private/NiagaraNodeWithDynamicPins.cpp
@@ -272,47 +272,63 @@ void UNiagaraNodeWithDynamicPins::CollectAddPinActions(FGraphActionListBuilderBa
 
 void UNiagaraNodeWithDynamicPins::AddParameter(FNiagaraVariable Parameter, const UEdGraphPin* AddPin)
 {
-	FScopedTransaction AddNewPinTransaction(LOCTEXT("AddNewPinTransaction", "Add pin to node"));
-	UNiagaraGraph::FAddParameterOptions AddParameterOptions = UNiagaraGraph::FAddParameterOptions();
-	if (AddPin->Direction == EEdGraphPinDirection::EGPD_Input)
+	if (this->IsA<UNiagaraNodeParameterMapBase>())
 	{
-		AddParameterOptions.NewParameterUsage = ENiagaraScriptParameterUsage::Output;
+		// Parameter map type nodes create new parameters when adding pins.
+		FScopedTransaction AddNewPinTransaction(LOCTEXT("AddNewPinTransaction", "Add pin to node"));
+		UNiagaraGraph::FAddParameterOptions AddParameterOptions = UNiagaraGraph::FAddParameterOptions();
+		if (AddPin->Direction == EEdGraphPinDirection::EGPD_Input)
+		{
+			AddParameterOptions.NewParameterUsage = ENiagaraScriptParameterUsage::Output;
+		}
+		else
+		{
+			AddParameterOptions.NewParameterUsage = ENiagaraScriptParameterUsage::Input;
+		}
+		// Default scope of new parameters on a pin to Particles, by convention. //@todo(ng) choose based on usage
+		AddParameterOptions.NewParameterScopeName = FNiagaraConstants::ParticleAttributeNamespace;
+		//AddParameterOptions.bMakeParameterNameUnique //@todo(ng) consider
+
+		UNiagaraGraph* Graph = GetNiagaraGraph();
+		checkf(Graph != nullptr, TEXT("Failed to get niagara graph when adding pin!"));
+
+		Graph->Modify();
+		Graph->AddParameter(Parameter, AddParameterOptions);
+
+		Modify();
+		UEdGraphPin* Pin = this->RequestNewTypedPin(AddPin->Direction, Parameter.GetType(), Parameter.GetName());
+		FNiagaraGraphParameterReference NewParameterReference = FNiagaraGraphParameterReference(Pin->PersistentGuid, this);
+		Graph->AddParameterReference(Parameter, NewParameterReference);
 	}
 	else
 	{
-		AddParameterOptions.NewParameterUsage = ENiagaraScriptParameterUsage::Input;
+		RequestNewTypedPin(AddPin->Direction, Parameter.GetType(), Parameter.GetName());
 	}
-	// Default scope of new parameters on a pin to Particles, by convention. //@todo(ng) choose based on usage
-	AddParameterOptions.NewParameterScopeName = FNiagaraConstants::ParticleAttributeNamespace;
-	//AddParameterOptions.bMakeParameterNameUnique //@todo(ng) consider
-	
-	UNiagaraGraph* Graph = GetNiagaraGraph();
-	checkf(Graph != nullptr, TEXT("Failed to get niagara graph when adding pin!"));
-	
-	Graph->Modify();
-	Graph->AddParameter(Parameter, AddParameterOptions);
-
-	Modify();
-	UEdGraphPin* Pin = this->RequestNewTypedPin(AddPin->Direction, Parameter.GetType(), Parameter.GetName());
-	FNiagaraGraphParameterReference NewParameterReference = FNiagaraGraphParameterReference(Pin->PersistentGuid, this);
-	Graph->AddParameterReference(Parameter, NewParameterReference);
 }
 
 void UNiagaraNodeWithDynamicPins::AddParameter(FNiagaraVariable Parameter, const struct UNiagaraGraph::FAddParameterOptions AddParameterOptions)
 {
-	UNiagaraGraph* Graph = GetNiagaraGraph();
-	checkf(Graph != nullptr, TEXT("Failed to get niagara graph when adding pin!"));
-
-	Graph->Modify();
-	Graph->AddParameter(Parameter, AddParameterOptions);
-
 	EEdGraphPinDirection NewPinDirection = GetPinDirectionForNewParameters();
 	checkf(NewPinDirection != EGPD_MAX, TEXT("Could not determine direction for new pin! Did you derive a new node from UNiagaraNodeWithDynamicPins?"));
-	
-	Modify();
-	UEdGraphPin* Pin = this->RequestNewTypedPin(NewPinDirection, Parameter.GetType(), Parameter.GetName());
-	FNiagaraGraphParameterReference NewParameterReference = FNiagaraGraphParameterReference(Pin->PersistentGuid, this);
-	Graph->AddParameterReference(Parameter, NewParameterReference);
+
+	if (this->IsA<UNiagaraNodeParameterMapBase>())
+	{
+		// Parameter map type nodes create new parameters when adding pins.
+		UNiagaraGraph* Graph = GetNiagaraGraph();
+		checkf(Graph != nullptr, TEXT("Failed to get niagara graph when adding pin!"));
+
+		Graph->Modify();
+		Graph->AddParameter(Parameter, AddParameterOptions);
+
+		Modify();
+		UEdGraphPin* Pin = RequestNewTypedPin(NewPinDirection, Parameter.GetType(), Parameter.GetName());
+		FNiagaraGraphParameterReference NewParameterReference = FNiagaraGraphParameterReference(Pin->PersistentGuid, this);
+		Graph->AddParameterReference(Parameter, NewParameterReference);
+	}
+	else
+	{
+		RequestNewTypedPin(NewPinDirection, Parameter.GetType(), Parameter.GetName());
+	}
 }
 
 void UNiagaraNodeWithDynamicPins::RemoveDynamicPin(UEdGraphPin* Pin)
@@ -320,17 +336,21 @@ void UNiagaraNodeWithDynamicPins::RemoveDynamicPin(UEdGraphPin* Pin)
 	RemovePin(Pin);
 	MarkNodeRequiresSynchronization(__FUNCTION__, true);
 
-	if (UNiagaraGraph* Graph = GetNiagaraGraph())
+	if (this->IsA<UNiagaraNodeParameterMapBase>())
 	{
-		const UEdGraphSchema_Niagara* Schema = GetDefault<UEdGraphSchema_Niagara>();
-		FNiagaraVariable PinVariable = Schema->PinToNiagaraVariable(Pin, false);
-		UNiagaraScriptVariable** PinAssociatedScriptVariable = Graph->GetAllMetaData().Find(PinVariable);
-		if (PinAssociatedScriptVariable != nullptr)
+		// Synchronize parameters if deleting a pin off of a parameter map type node.
+		if (UNiagaraGraph* Graph = GetNiagaraGraph())
 		{
-			bool bShouldRemoveScriptVariable = !Graph->UpdateUsageForScriptVariable(*PinAssociatedScriptVariable);
-			if (bShouldRemoveScriptVariable)
+			const UEdGraphSchema_Niagara* Schema = GetDefault<UEdGraphSchema_Niagara>();
+			FNiagaraVariable PinVariable = Schema->PinToNiagaraVariable(Pin, false);
+			UNiagaraScriptVariable** PinAssociatedScriptVariable = Graph->GetAllMetaData().Find(PinVariable);
+			if (PinAssociatedScriptVariable != nullptr)
 			{
-				Graph->RemoveParameter((*PinAssociatedScriptVariable)->Variable);
+				bool bShouldRemoveScriptVariable = !Graph->UpdateUsageForScriptVariable(*PinAssociatedScriptVariable);
+				if (bShouldRemoveScriptVariable)
+				{
+					Graph->RemoveParameter((*PinAssociatedScriptVariable)->Variable);
+				}
 			}
 		}
 	}
diff --git a/Engine/Plugins/FX/Niagara/Source/NiagaraEditor/Private/TNiagaraGraphPinEditableName.h b/Engine/Plugins/FX/Niagara/Source/NiagaraEditor/Private/TNiagaraGraphPinEditableName.h
--- a/Engine/Plugins/FX/Niagara/Source/NiagaraEditor/Private/TNiagaraGraphPinEditableName.h
+++ b/Engine/Plugins/FX/Niagara/Source/NiagaraEditor/Private/TNiagaraGraphPinEditableName.h
@@ -79,30 +79,57 @@ protected:
 	virtual TSharedRef<SWidget> GetLabelWidget(const FName& InLabelStyle) override
 	{
 		UNiagaraNode* ParentNode = Cast<UNiagaraNode>(this->GraphPinObj->GetOwningNode());
+
+		auto CreateLabelTextBlock = [&]()->TSharedRef<SWidget> {
+			return SNew(SInlineEditableTextBlock)
+				.Style(&FEditorStyle::Get().GetWidgetStyle<FInlineEditableTextBlockStyle>("Graph.Node.InlineEditablePinName"))
+				.Text(this, &TNiagaraGraphPinEditableName<BaseClass>::GetParentPinLabel)
+				.Visibility(this, &TNiagaraGraphPinEditableName<BaseClass>::GetParentPinVisibility)
+				.ColorAndOpacity(this, &TNiagaraGraphPinEditableName<BaseClass>::GetParentPinTextColor)
+				.IsReadOnly(true);
+		};
+
+		auto CreateRenamableLabelTextBlock = [&]()->TSharedRef<SWidget> {
+			if (ParentNode->IsPinNameEditableUponCreation(this->GraphPinObj))
+			{
+				bPendingRename = true;
+			}
+			
+			return SNew(SInlineEditableTextBlock)
+				.Style(&FEditorStyle::Get().GetWidgetStyle<FInlineEditableTextBlockStyle>("Graph.Node.InlineEditablePinName"))
+				.Text(this, &TNiagaraGraphPinEditableName<BaseClass>::GetParentPinLabel)
+				.Visibility(this, &TNiagaraGraphPinEditableName<BaseClass>::GetParentPinVisibility)
+				.ColorAndOpacity(this, &TNiagaraGraphPinEditableName<BaseClass>::GetParentPinTextColor)
+				.OnVerifyTextChanged(this, &TNiagaraGraphPinEditableName<BaseClass>::OnVerifyTextChanged)
+				.OnTextCommitted(this, &TNiagaraGraphPinEditableName<BaseClass>::OnTextCommitted);
+		};
+
 		if (ParentNode && ParentNode->IsPinNameEditable(this->GraphPinObj))
 		{
 			UNiagaraGraph* NiagaraGraph = ParentNode->GetNiagaraGraph();
-			if (FNiagaraEditorCommonCVar::GNiagaraEnableParameterPanel2 > 0 && NiagaraGraph->IsPinVisualWidgetProviderRegistered())
+			if (FNiagaraEditorCommonCVar::GNiagaraEnableParameterPanel2 > 0)
 			{
-				return NiagaraGraph->GetPinVisualWidget(this->GraphPinObj);
+				if (ParentNode->IsA<UNiagaraNodeParameterMapBase>())
+				{
+					if (NiagaraGraph->IsPinVisualWidgetProviderRegistered())
+					{
+						return NiagaraGraph->GetPinVisualWidget(this->GraphPinObj);
+					}
+					return CreateLabelTextBlock();
+				}
+				else if (ParentNode->IsA<UNiagaraNodeCustomHlsl>())
+				{
+					return CreateRenamableLabelTextBlock();
+				}
+				else
+				{
+					ensureMsgf(false, TEXT("Tried to create a pin widget for unhandled node class!"));
+					return CreateRenamableLabelTextBlock();
+				}
 			}
 			else
 			{
-				CreatedTextBlock = SNew(SInlineEditableTextBlock)
-					.Style(&FEditorStyle::Get().GetWidgetStyle<FInlineEditableTextBlockStyle>("Graph.Node.InlineEditablePinName"))
-					.Text(this, &TNiagaraGraphPinEditableName<BaseClass>::GetParentPinLabel)
-					.Visibility(this, &TNiagaraGraphPinEditableName<BaseClass>::GetParentPinVisibility)
-					.ColorAndOpacity(this, &TNiagaraGraphPinEditableName<BaseClass>::GetParentPinTextColor)
-					.OnVerifyTextChanged(this, &TNiagaraGraphPinEditableName<BaseClass>::OnVerifyTextChanged)
-					.OnTextCommitted(this, &TNiagaraGraphPinEditableName<BaseClass>::OnTextCommitted);
-
-				TSharedRef<SWidget> EditableTextBlock = CreatedTextBlock.ToSharedRef();
-
-				if (ParentNode->IsPinNameEditableUponCreation(this->GraphPinObj))
-				{
-					bPendingRename = true;
-				}
-				return EditableTextBlock;
+				return CreateRenamableLabelTextBlock();
 			}
 		}
 		else
diff --git a/Engine/Plugins/FX/Niagara/Source/NiagaraEditor/Private/Widgets/SNiagaraGraphPinAdd.cpp b/Engine/Plugins/FX/Niagara/Source/NiagaraEditor/Private/Widgets/SNiagaraGraphPinAdd.cpp
--- a/Engine/Plugins/FX/Niagara/Source/NiagaraEditor/Private/Widgets/SNiagaraGraphPinAdd.cpp
+++ b/Engine/Plugins/FX/Niagara/Source/NiagaraEditor/Private/Widgets/SNiagaraGraphPinAdd.cpp
@@ -13,6 +13,8 @@
 #include "Widgets/Layout/SBox.h"
 #include "Framework/Application/SlateApplication.h"
 #include "SNiagaraParameterPanel.h"
+#include "NiagaraNodeParameterMapBase.h"
+#include "NiagaraNodeCustomHlsl.h"
 
 #define LOCTEXT_NAMESPACE "NiagaraGraphPinAdd"
 
@@ -66,18 +68,42 @@ TSharedRef<SWidget> SNiagaraGraphPinAdd::OnGetAddButtonMenuContent()
 {
 	if (OwningNode != nullptr)
 	{
-		TArray<TWeakObjectPtr<UNiagaraGraph>> Graphs;
-		Graphs.Add(OwningNode->GetNiagaraGraph());
+		if (OwningNode->IsA<UNiagaraNodeParameterMapBase>())
+		{
+			TArray<TWeakObjectPtr<UNiagaraGraph>> Graphs;
+			Graphs.Add(OwningNode->GetNiagaraGraph());
+
+			// Creating a new parameter from a graph pin always defaults to Particles scope.
+			TSharedRef<SNiagaraAddParameterMenu2> MenuWidget = SNew(SNiagaraAddParameterMenu2, Graphs)
+				.OnAddParameter_UObject(OwningNode, &UNiagaraNodeWithDynamicPins::AddParameter, (const UEdGraphPin*)GetPinObj())
+				.OnCollectCustomActions_UObject(OwningNode, &UNiagaraNodeWithDynamicPins::CollectAddPinActions, GetPinObj())
+				.OnAllowMakeType_UObject(OwningNode, &UNiagaraNodeWithDynamicPins::AllowNiagaraTypeForAddPin)
+				.IsParameterRead(GraphPinObj ? GraphPinObj->Direction == EGPD_Output : true);
+
+			AddButton->SetMenuContentWidgetToFocus(MenuWidget->GetSearchBox());
+			return MenuWidget;
+		}
+		else if (OwningNode->IsA<UNiagaraNodeCustomHlsl>())
+		{
+			// Skip binding the add parameter delegate for custom HLSL as it does not interact with graph's UNiagaraScriptVariables.
 
-		// Creating a new parameter from a graph pin always defaults to Particles scope.
-		TSharedRef<SNiagaraAddParameterMenu2> MenuWidget = SNew(SNiagaraAddParameterMenu2, Graphs, ENiagaraParameterScope::Particles)
-			.OnAddParameter_UObject(OwningNode, &UNiagaraNodeWithDynamicPins::AddParameter, (const UEdGraphPin*)GetPinObj())
-			.OnCollectCustomActions_UObject(OwningNode, &UNiagaraNodeWithDynamicPins::CollectAddPinActions, GetPinObj())
-			.OnAllowMakeType_UObject(OwningNode, &UNiagaraNodeWithDynamicPins::AllowNiagaraTypeForAddPin)
-			.IsParameterRead(GraphPinObj ? GraphPinObj->Direction == EGPD_Output : true);
+			TArray<TWeakObjectPtr<UNiagaraGraph>> Graphs;
+			Graphs.Add(OwningNode->GetNiagaraGraph());
 
-		AddButton->SetMenuContentWidgetToFocus(MenuWidget->GetSearchBox());
-		return MenuWidget;
+			// Creating a new parameter from a graph pin always defaults to Particles scope.
+			TSharedRef<SNiagaraAddParameterMenu2> MenuWidget = SNew(SNiagaraAddParameterMenu2, Graphs)
+				.OnCollectCustomActions_UObject(OwningNode, &UNiagaraNodeWithDynamicPins::CollectAddPinActions, GetPinObj())
+				.OnAllowMakeType_UObject(OwningNode, &UNiagaraNodeWithDynamicPins::AllowNiagaraTypeForAddPin)
+				.IsParameterRead(GraphPinObj ? GraphPinObj->Direction == EGPD_Output : true);
+
+			AddButton->SetMenuContentWidgetToFocus(MenuWidget->GetSearchBox());
+			return MenuWidget;
+		}
+		else
+		{
+			ensureMsgf(false, TEXT("Tried to create a graph pin add menu for unhandled node class!"));
+			return SNullWidget::NullWidget;
+		}
 	}
 	else
 	{
diff --git a/Engine/Plugins/FX/Niagara/Source/NiagaraEditor/Private/Widgets/SNiagaraParameterPanel.cpp b/Engine/Plugins/FX/Niagara/Source/NiagaraEditor/Private/Widgets/SNiagaraParameterPanel.cpp
--- a/Engine/Plugins/FX/Niagara/Source/NiagaraEditor/Private/Widgets/SNiagaraParameterPanel.cpp
+++ b/Engine/Plugins/FX/Niagara/Source/NiagaraEditor/Private/Widgets/SNiagaraParameterPanel.cpp
@@ -232,13 +232,14 @@ TSharedRef<SWidget> SNiagaraParameterPanel::CreateAddToSectionButton(const Niaga
 
 TSharedRef<SWidget> SNiagaraParameterPanel::OnGetParameterMenu(const NiagaraParameterPanelSectionID::Type InSection)
 {
-	ENiagaraParameterScope NewParameterScope = NiagaraParameterPanelSectionID::GetScopeForNewParametersInSection(InSection);
+	ENiagaraParameterScope NewParameterScopeForSection = NiagaraParameterPanelSectionID::GetScopeForNewParametersInSection(InSection);
 
-	TSharedRef<SNiagaraAddParameterMenu2> MenuWidget = SNew(SNiagaraAddParameterMenu2, ParameterPanelViewModel->GetEditableGraphs(), NewParameterScope)
+	TSharedRef<SNiagaraAddParameterMenu2> MenuWidget = SNew(SNiagaraAddParameterMenu2, ParameterPanelViewModel->GetEditableGraphs())
 		.OnAddParameter(this, &SNiagaraParameterPanel::AddParameter, InSection)
 		.OnAllowMakeType(this, &SNiagaraParameterPanel::AllowMakeType)
 		.ShowGraphParameters(false)
-		.AutoExpandMenu(true);
+		.AutoExpandMenu(true)
+		.NewParameterScope(NewParameterScopeForSection);
 
 	AddParameterButtons[(int32)InSection]->SetMenuContentWidgetToFocus(MenuWidget->GetSearchBox()->AsShared());
 	return MenuWidget;
@@ -321,7 +322,7 @@ void SNiagaraParameterPanel::Refresh()
 	/************************************************************************/
 	/* SNiagaraAddParameterMenu2                                             */
 	/************************************************************************/
-void SNiagaraAddParameterMenu2::Construct(const FArguments& InArgs, TArray<TWeakObjectPtr<UNiagaraGraph>> InGraphs, ENiagaraParameterScope InNewParameterScope)
+void SNiagaraAddParameterMenu2::Construct(const FArguments& InArgs, TArray<TWeakObjectPtr<UNiagaraGraph>> InGraphs)
 {
 	this->OnAddParameter = InArgs._OnAddParameter;
 	this->OnCollectCustomActions = InArgs._OnCollectCustomActions;
@@ -330,9 +331,9 @@ void SNiagaraAddParameterMenu2::Construct(const FArguments& InArgs, TArray<TWeak
 	this->ShowGraphParameters = InArgs._ShowGraphParameters;
 	this->AutoExpandMenu = InArgs._AutoExpandMenu;
 	this->IsParameterRead = InArgs._IsParameterRead;
+	this->NewParameterScope = InArgs._NewParameterScope;
 
 	Graphs = InGraphs;
-	NewParameterScope = InNewParameterScope;
 	NewParameterScopeText = FText::FromString(FNiagaraTypeUtilities::GetNamespaceStringForScriptParameterScope(NewParameterScope));
 
 	ChildSlot
diff --git a/Engine/Plugins/FX/Niagara/Source/NiagaraEditor/Private/Widgets/SNiagaraParameterPanel.h b/Engine/Plugins/FX/Niagara/Source/NiagaraEditor/Private/Widgets/SNiagaraParameterPanel.h
--- a/Engine/Plugins/FX/Niagara/Source/NiagaraEditor/Private/Widgets/SNiagaraParameterPanel.h
+++ b/Engine/Plugins/FX/Niagara/Source/NiagaraEditor/Private/Widgets/SNiagaraParameterPanel.h
@@ -138,7 +138,8 @@ public:
 		: _AllowCreatingNew(true)
 		, _ShowGraphParameters(true)
 		, _AutoExpandMenu(false)
-		, _IsParameterRead(true) {}
+		, _IsParameterRead(true)
+		, _NewParameterScope(ENiagaraParameterScope::Particles) {}
 		SLATE_EVENT(FOnAddParameter, OnAddParameter)
 		SLATE_EVENT(FOnCollectCustomActions, OnCollectCustomActions)
 		SLATE_EVENT(FOnAllowMakeType, OnAllowMakeType)
@@ -146,9 +147,10 @@ public:
 		SLATE_ATTRIBUTE(bool, ShowGraphParameters)
 		SLATE_ATTRIBUTE(bool, AutoExpandMenu)
 		SLATE_ATTRIBUTE(bool, IsParameterRead)
+		SLATE_ARGUMENT(ENiagaraParameterScope, NewParameterScope)
 	SLATE_END_ARGS();
 
-	void Construct(const FArguments& InArgs, TArray<TWeakObjectPtr<UNiagaraGraph>> InGraphs, ENiagaraParameterScope InNewParameterScope);
+	void Construct(const FArguments& InArgs, TArray<TWeakObjectPtr<UNiagaraGraph>> InGraphs);
 
 	TSharedRef<SEditableTextBox> GetSearchBox();
 
