Removed UNiagaraDataInterfaceStaticMesh logs are required not anymore, when sampling static mesh that have missing color stream.

diff --git a/Engine/Plugins/FX/Niagara/Source/Niagara/Private/NiagaraDataInterfaceStaticMesh.cpp b/Engine/Plugins/FX/Niagara/Source/Niagara/Private/NiagaraDataInterfaceStaticMesh.cpp
--- a/Engine/Plugins/FX/Niagara/Source/Niagara/Private/NiagaraDataInterfaceStaticMesh.cpp
+++ b/Engine/Plugins/FX/Niagara/Source/Niagara/Private/NiagaraDataInterfaceStaticMesh.cpp
@@ -60,13 +60,13 @@ float FStaticMeshFilteredAreaWeightedSectionSampler::GetWeights(TArray<float>& O
 
 	if (Owner->Mesh->bSupportUniformlyDistributedSampling && LODRes.AreaWeightedSectionSamplers.Num() > 0)
 	{
-		for (int32 i = 0; i < Owner->GetValidSections().Num(); ++i)
-		{
-			int32 SecIdx = Owner->GetValidSections()[i];
-			float T = LODRes.AreaWeightedSectionSamplers[SecIdx].GetTotalWeight();
-			OutWeights.Add(T);
-			Total += T;
-		}
+	for (int32 i = 0; i < Owner->GetValidSections().Num(); ++i)
+	{
+		int32 SecIdx = Owner->GetValidSections()[i];
+		float T = LODRes.AreaWeightedSectionSamplers[SecIdx].GetTotalWeight();
+		OutWeights.Add(T);
+		Total += T;
+	}
 	}
 	else
 	{
@@ -1017,10 +1017,6 @@ void UNiagaraDataInterfaceStaticMesh::GetVMExternalFunction(const FVMExternalFun
 	FNDIStaticMesh_InstanceData* InstData = (FNDIStaticMesh_InstanceData*)InstanceData;
 	check(InstData && InstData->Mesh && InstData->SafeComponent_GT.IsValid());
 
-	bool bNeedsVertexPositions = false;
-	bool bNeedsVertexColors = false;
-	bool bNeedsVertMain = true;//Assuming we always need this?
-	
 	if (BindingInfo.Name == StaticMeshHelpers::RandomSectionName)
 	{
 		check(BindingInfo.GetNumInputs() == 1 && BindingInfo.GetNumOutputs() == 1);
@@ -1046,56 +1042,46 @@ void UNiagaraDataInterfaceStaticMesh::GetVMExternalFunction(const FVMExternalFun
 	else if (BindingInfo.Name == StaticMeshHelpers::GetTriPositionName)
 	{
 		check(BindingInfo.GetNumInputs() == 5 && BindingInfo.GetNumOutputs() == 3);
-		bNeedsVertexPositions = true;
 		TNDIExplicitBinder<FNDITransformHandlerNoop, NDI_FUNC_BINDER(UNiagaraDataInterfaceStaticMesh, GetTriCoordPosition)>::Bind(this, BindingInfo, InstanceData, OutFunc);
 	}
 	else if (BindingInfo.Name == StaticMeshHelpers::GetTriPositionWSName)
 	{
 		check(BindingInfo.GetNumInputs() == 5 && BindingInfo.GetNumOutputs() == 3);
-		bNeedsVertexPositions = true;
 		TNDIExplicitBinder<FNDITransformHandler, NDI_FUNC_BINDER(UNiagaraDataInterfaceStaticMesh, GetTriCoordPosition)>::Bind(this, BindingInfo, InstanceData, OutFunc);
 	}
 	else if (BindingInfo.Name == StaticMeshHelpers::GetTriNormalName)
 	{
 		check(BindingInfo.GetNumInputs() == 5 && BindingInfo.GetNumOutputs() == 3);
-		bNeedsVertMain = true;
 		TNDIExplicitBinder<FNDITransformHandlerNoop, NDI_FUNC_BINDER(UNiagaraDataInterfaceStaticMesh, GetTriCoordNormal)>::Bind(this, BindingInfo, InstanceData, OutFunc);
 	}
 	else if (BindingInfo.Name == StaticMeshHelpers::GetTriNormalWSName)
 	{
 		check(BindingInfo.GetNumInputs() == 5 && BindingInfo.GetNumOutputs() == 3);
-		bNeedsVertMain = true;
 		TNDIExplicitBinder<FNDITransformHandler, NDI_FUNC_BINDER(UNiagaraDataInterfaceStaticMesh, GetTriCoordNormal)>::Bind(this, BindingInfo, InstanceData, OutFunc);
 	}
 	else if (BindingInfo.Name == StaticMeshHelpers::GetTriTangentsName)
 	{
 		check(BindingInfo.GetNumInputs() == 5 && BindingInfo.GetNumOutputs() == 9);
-		bNeedsVertMain = true;
 		TTypedMeshAccessorBinder<TNDIExplicitBinder<FNDITransformHandlerNoop, NDI_FUNC_BINDER(UNiagaraDataInterfaceStaticMesh, GetTriCoordTangents)>>::Bind(this, BindingInfo, InstanceData, OutFunc);
 	}
 	else if (BindingInfo.Name == StaticMeshHelpers::GetTriTangentsWSName)
 	{
 		check(BindingInfo.GetNumInputs() == 5 && BindingInfo.GetNumOutputs() == 9);
-		bNeedsVertMain = true;
 		TTypedMeshAccessorBinder<TNDIExplicitBinder<FNDITransformHandler, NDI_FUNC_BINDER(UNiagaraDataInterfaceStaticMesh, GetTriCoordTangents)>>::Bind(this, BindingInfo, InstanceData, OutFunc);
 	}
 	else if (BindingInfo.Name == StaticMeshHelpers::GetTriColorName)
 	{
 		check(BindingInfo.GetNumInputs() == 5 && BindingInfo.GetNumOutputs() == 4);
-		bNeedsVertexColors = true;
 		NDI_FUNC_BINDER(UNiagaraDataInterfaceStaticMesh, GetTriCoordColor)::Bind(this, OutFunc);
 	}
 	else if (BindingInfo.Name == StaticMeshHelpers::GetTriUVName)
 	{
 		check(BindingInfo.GetNumInputs() == 6 && BindingInfo.GetNumOutputs() == 2);
-		bNeedsVertMain = true;
 		TTypedMeshAccessorBinder<NDI_FUNC_BINDER(UNiagaraDataInterfaceStaticMesh, GetTriCoordUV)>::Bind(this, BindingInfo, InstanceData, OutFunc);
 	}
 	else if (BindingInfo.Name == StaticMeshHelpers::GetTriPositionAndVelocityName)
 	{
 		check(BindingInfo.GetNumInputs() == 5 && BindingInfo.GetNumOutputs() == 6);
-		bNeedsVertMain = true;
-		bNeedsVertexPositions = true;  
 		NDI_FUNC_BINDER(UNiagaraDataInterfaceStaticMesh, GetTriCoordPositionAndVelocity)::Bind(this, OutFunc);
 	}
 	else if (BindingInfo.Name == StaticMeshHelpers::GetMeshLocalToWorldName)
@@ -1116,28 +1102,13 @@ void UNiagaraDataInterfaceStaticMesh::GetVMExternalFunction(const FVMExternalFun
 	else if (BindingInfo.Name == StaticMeshHelpers::GetVertexPositionName)
 	{
 		check(BindingInfo.GetNumInputs() == 2 && BindingInfo.GetNumOutputs() == 3);
-		bNeedsVertexPositions = true;
 		TNDIExplicitBinder<FNDITransformHandlerNoop, NDI_FUNC_BINDER(UNiagaraDataInterfaceStaticMesh, GetVertexPosition)>::Bind(this, BindingInfo, InstanceData, OutFunc);
 	}
 	else if (BindingInfo.Name == StaticMeshHelpers::GetVertexPositionWSName)
 	{
 		check(BindingInfo.GetNumInputs() == 2 && BindingInfo.GetNumOutputs() == 3);
-		bNeedsVertexPositions = true;
 		TNDIExplicitBinder<FNDITransformHandler, NDI_FUNC_BINDER(UNiagaraDataInterfaceStaticMesh, GetVertexPosition)>::Bind(this, BindingInfo, InstanceData, OutFunc);
 	}
-
-	if (bNeedsVertexPositions && !InstData->MeshHasPositions())
-	{
-		UE_LOG(LogNiagara, Log, TEXT("Static Mesh data interface is cannot run as it's reading position data on a mesh that does not provide it. - Mesh:%s  "), *InstData->Mesh->GetFullName());
-	}
-	if (bNeedsVertexColors && !InstData->MeshHasColors())
-	{
-		UE_LOG(LogNiagara, Log, TEXT("Static Mesh data interface is cannot run as it's reading color data on a mesh that does not provide it. - Mesh:%s  "), *InstData->Mesh->GetFullName());
-	}
-	if (bNeedsVertMain && !InstData->MeshHasVerts())
-	{
-		UE_LOG(LogNiagara, Log, TEXT("Static Mesh data interface is cannot run as it's reading vertex data on a mesh with no vertex data. - Mesh:%s  "), *InstData->Mesh->GetFullName());
-	}
 }
 
 bool UNiagaraDataInterfaceStaticMesh::CopyToInternal(UNiagaraDataInterface* Destination) const
@@ -1273,9 +1244,9 @@ FORCEINLINE int32 UNiagaraDataInterfaceStaticMesh::RandomSection<TIntegralConsta
 	checkSlow(InstData->GetValidSections().Num() > 0);
 	if (InstData->GetAreaWeigtedSampler().GetNumEntries() > 0)
 	{
-		int32 Idx = InstData->GetAreaWeigtedSampler().GetEntryIndex(RandStream.GetFraction(), RandStream.GetFraction());
-		return InstData->GetValidSections()[Idx];
-	}
+	int32 Idx = InstData->GetAreaWeigtedSampler().GetEntryIndex(RandStream.GetFraction(), RandStream.GetFraction());
+	return InstData->GetValidSections()[Idx];
+}
 	else
 	{
 		return 0;
@@ -1287,8 +1258,8 @@ FORCEINLINE int32 UNiagaraDataInterfaceStaticMesh::RandomSection<TIntegralConsta
 {
 	if (Res.AreaWeightedSampler.GetNumEntries() > 0)
 	{
-		return Res.AreaWeightedSampler.GetEntryIndex(RandStream.GetFraction(), RandStream.GetFraction());
-	}
+	return Res.AreaWeightedSampler.GetEntryIndex(RandStream.GetFraction(), RandStream.GetFraction());
+}
 	else
 	{
 		return 0;
@@ -1331,12 +1302,12 @@ FORCEINLINE int32 UNiagaraDataInterfaceStaticMesh::RandomTriIndex<TIntegralConst
 	int32 SecIdx = RandomSection<TIntegralConstant<bool, true>, true>(RandStream, Res, InstData);
 	if (SecIdx < Res.Sections.Num() && SecIdx < Res.AreaWeightedSectionSamplers.Num())
 	{
-		FStaticMeshSection& Sec = Res.Sections[SecIdx];
+	FStaticMeshSection&  Sec = Res.Sections[SecIdx];
 		if (Res.AreaWeightedSectionSamplers[SecIdx].GetNumEntries() > 0)
 		{
-			int32 Tri = Res.AreaWeightedSectionSamplers[SecIdx].GetEntryIndex(RandStream.GetFraction(), RandStream.GetFraction());
-			return (Sec.FirstIndex / 3) + Tri;
-		}
+	int32 Tri = Res.AreaWeightedSectionSamplers[SecIdx].GetEntryIndex(RandStream.GetFraction(), RandStream.GetFraction());
+	return (Sec.FirstIndex / 3) + Tri;
+}
 		return (Sec.FirstIndex / 3);
 	}
 	return 0;
@@ -1348,12 +1319,12 @@ FORCEINLINE int32 UNiagaraDataInterfaceStaticMesh::RandomTriIndex<TIntegralConst
 	int32 SecIdx = RandomSection<TIntegralConstant<bool, true>, false>(RandStream, Res, InstData);
 	if (SecIdx < Res.Sections.Num() && SecIdx < Res.AreaWeightedSectionSamplers.Num())
 	{
-		FStaticMeshSection&  Sec = Res.Sections[SecIdx];
+	FStaticMeshSection&  Sec = Res.Sections[SecIdx];
 		if (Res.AreaWeightedSectionSamplers[SecIdx].GetNumEntries() > 0)
 		{
-			int32 Tri = Res.AreaWeightedSectionSamplers[SecIdx].GetEntryIndex(RandStream.GetFraction(), RandStream.GetFraction());
-			return (Sec.FirstIndex / 3) + Tri;
-		}
+	int32 Tri = Res.AreaWeightedSectionSamplers[SecIdx].GetEntryIndex(RandStream.GetFraction(), RandStream.GetFraction());
+	return (Sec.FirstIndex / 3) + Tri;
+}
 		return (Sec.FirstIndex / 3);
 	}
 	return 0;
@@ -1365,10 +1336,10 @@ FORCEINLINE int32 UNiagaraDataInterfaceStaticMesh::RandomTriIndex<TIntegralConst
 	int32 SecIdx = RandomSection<TIntegralConstant<bool, false>, true>(RandStream, Res, InstData);
 	if (SecIdx < Res.Sections.Num())
 	{
-		FStaticMeshSection& Sec = Res.Sections[SecIdx];
-		int32 Tri = RandStream.RandRange(0, Sec.NumTriangles - 1);
-		return (Sec.FirstIndex / 3) + Tri;
-	}
+	FStaticMeshSection&  Sec = Res.Sections[SecIdx];
+	int32 Tri = RandStream.RandRange(0, Sec.NumTriangles - 1);
+	return (Sec.FirstIndex / 3) + Tri;
+}
 	return 0;
 }
 
@@ -1378,10 +1349,10 @@ FORCEINLINE int32 UNiagaraDataInterfaceStaticMesh::RandomTriIndex<TIntegralConst
 	int32 SecIdx = RandomSection<TIntegralConstant<bool, false>, false>(RandStream, Res, InstData);
 	if (SecIdx < Res.Sections.Num())
 	{
-		FStaticMeshSection& Sec = Res.Sections[SecIdx];
-		int32 Tri = RandStream.RandRange(0, Sec.NumTriangles - 1);
-		return (Sec.FirstIndex / 3) + Tri;
-	}
+	FStaticMeshSection&  Sec = Res.Sections[SecIdx];
+	int32 Tri = RandStream.RandRange(0, Sec.NumTriangles - 1);
+	return (Sec.FirstIndex / 3) + Tri;
+}
 	return 0;
 }
 
