Non-Unity fix.

diff --git a/Engine/Plugins/FX/Niagara/Source/NiagaraEditor/Private/TNiagaraGraphPinEditableName.h b/Engine/Plugins/FX/Niagara/Source/NiagaraEditor/Private/TNiagaraGraphPinEditableName.h
--- a/Engine/Plugins/FX/Niagara/Source/NiagaraEditor/Private/TNiagaraGraphPinEditableName.h
+++ b/Engine/Plugins/FX/Niagara/Source/NiagaraEditor/Private/TNiagaraGraphPinEditableName.h
@@ -4,6 +4,7 @@
 
 #include "SGraphPin.h"
 #include "NiagaraNode.h"
+#include "NiagaraNodeParameterMapBase.h"
 #include "Widgets/Text/SInlineEditableTextBlock.h"
 
 /** A graph pin widget for allowing a pin to have an editable name for a pin. */
