Non-Unity fixes.

diff --git a/Engine/Plugins/FX/Niagara/Source/NiagaraEditor/Private/TNiagaraGraphPinEditableName.h b/Engine/Plugins/FX/Niagara/Source/NiagaraEditor/Private/TNiagaraGraphPinEditableName.h
--- a/Engine/Plugins/FX/Niagara/Source/NiagaraEditor/Private/TNiagaraGraphPinEditableName.h
+++ b/Engine/Plugins/FX/Niagara/Source/NiagaraEditor/Private/TNiagaraGraphPinEditableName.h
@@ -6,6 +6,7 @@
 #include "NiagaraNode.h"
 #include "NiagaraNodeParameterMapBase.h"
 #include "Widgets/Text/SInlineEditableTextBlock.h"
+#include "NiagaraNodeCustomHlsl.h"
 
 /** A graph pin widget for allowing a pin to have an editable name for a pin. */
 template< class BaseClass >
