Fixed CustomSideVector aligment giving the same result as Custom.

diff --git a/Engine/Plugins/FX/Niagara/Source/Niagara/Private/NiagaraRendererRibbons.cpp b/Engine/Plugins/FX/Niagara/Source/Niagara/Private/NiagaraRendererRibbons.cpp
--- a/Engine/Plugins/FX/Niagara/Source/Niagara/Private/NiagaraRendererRibbons.cpp
+++ b/Engine/Plugins/FX/Niagara/Source/Niagara/Private/NiagaraRendererRibbons.cpp
@@ -799,7 +799,7 @@ void FNiagaraRendererRibbons::SetupMeshBatchAndCollectorResourceForView(
 	VFLooseParams.NiagaraParticleDataFloat = CollectorResources.VertexFactory.GetParticleDataFloatSRV();
 	VFLooseParams.NiagaraFloatDataStride = CollectorResources.VertexFactory.GetFloatDataStride();
 	VFLooseParams.SortedIndicesOffset = CollectorResources.VertexFactory.GetSortedIndicesOffset();
-	VFLooseParams.FacingMode = CollectorResources.VertexFactory.GetFacingMode();
+	VFLooseParams.FacingMode = static_cast<uint32>(FacingMode);
 
 	// Collector.AllocateOneFrameResource uses default ctor, initialize the vertex factory
 	CollectorResources.VertexFactory.SetParticleFactoryType(NVFT_Ribbon);
