- Better handling of GPU triangle / bone sampling when no component or we have a master pose component
- Fix an issue with triangle sampling and meshes with > 4 influences

diff --git a/Engine/Plugins/FX/Niagara/Shaders/Private/NiagaraDataInterfaceSkeletalMesh.ush b/Engine/Plugins/FX/Niagara/Shaders/Private/NiagaraDataInterfaceSkeletalMesh.ush
--- a/Engine/Plugins/FX/Niagara/Shaders/Private/NiagaraDataInterfaceSkeletalMesh.ush
+++ b/Engine/Plugins/FX/Niagara/Shaders/Private/NiagaraDataInterfaceSkeletalMesh.ush
@@ -150,23 +150,22 @@ FDISkelMeshSkinnedVertex DISkelMesh_GetSkinnedVertex(in FDISkelMeshContext DICon
 		else
 #endif
 		{
-#if DISKELMESH_BONE_INFLUENCES > DISKELMESH_BONE_INFLUENCES_DEFAULT
 			uint WeightOffset = DIContext.MeshNumWeights == 8 ? 2 : 1;
-
+#if DISKELMESH_BONE_INFLUENCES > DISKELMESH_BONE_INFLUENCES_DEFAULT
 			int4 BlendIndices = DISkelMesh_UnpackIndices4(DIContext.MeshSkinWeightBuffer[VertexIndex * DIContext.MeshWeightStride]);
 			float4 BlendWeights = DISkelMesh_UnpackWeights4(DIContext.MeshSkinWeightBuffer[VertexIndex * DIContext.MeshWeightStride + WeightOffset]);
 			CurrBoneMatrix = DISkelMesh_GetCurrSkinningMatrix(DIContext, VertexIndex, BlendIndices, BlendWeights);
 			PrevBoneMatrix = DISkelMesh_GetPrevSkinningMatrix(DIContext, VertexIndex, BlendIndices, BlendWeights);
 			if (DIContext.MeshNumWeights == 8)
 			{
-			    BlendIndices = DISkelMesh_UnpackIndices4(DIContext.MeshSkinWeightBuffer[VertexIndex * DIContext.MeshWeightStride + 2]);
+            BlendIndices = DISkelMesh_UnpackIndices4(DIContext.MeshSkinWeightBuffer[VertexIndex * DIContext.MeshWeightStride + 1]);
 			    BlendWeights = DISkelMesh_UnpackWeights4(DIContext.MeshSkinWeightBuffer[VertexIndex * DIContext.MeshWeightStride + WeightOffset + 1]);
 				CurrBoneMatrix += DISkelMesh_GetCurrSkinningMatrix(DIContext, VertexIndex, BlendIndices, BlendWeights);
 				PrevBoneMatrix += DISkelMesh_GetPrevSkinningMatrix(DIContext, VertexIndex, BlendIndices, BlendWeights);
 			}
 #else
 			int4 BlendIndices = DISkelMesh_UnpackIndices4(DIContext.MeshSkinWeightBuffer[VertexIndex * DIContext.MeshWeightStride]);
-			float4 BlendWeights = DISkelMesh_UnpackWeights4(DIContext.MeshSkinWeightBuffer[VertexIndex * DIContext.MeshWeightStride + 1]);
+        float4 BlendWeights = DISkelMesh_UnpackWeights4(DIContext.MeshSkinWeightBuffer[VertexIndex * DIContext.MeshWeightStride + WeightOffset]);
 			CurrBoneMatrix = DISkelMesh_GetCurrSkinningMatrix(DIContext, VertexIndex, BlendIndices, BlendWeights);
 			PrevBoneMatrix = DISkelMesh_GetPrevSkinningMatrix(DIContext, VertexIndex, BlendIndices, BlendWeights);
 #endif
diff --git a/Engine/Plugins/FX/Niagara/Source/Niagara/Private/NiagaraDataInterfaceSkeletalMesh.cpp b/Engine/Plugins/FX/Niagara/Source/Niagara/Private/NiagaraDataInterfaceSkeletalMesh.cpp
--- a/Engine/Plugins/FX/Niagara/Source/Niagara/Private/NiagaraDataInterfaceSkeletalMesh.cpp
+++ b/Engine/Plugins/FX/Niagara/Source/Niagara/Private/NiagaraDataInterfaceSkeletalMesh.cpp
@@ -649,97 +649,133 @@ void FSkeletalMeshGpuDynamicBufferProxy::NewFrame(const FNDISkeletalMesh_Instanc
 		return;
 	}
 
-	// Create AllSectionsRefToLocalMatrices
-	TArray<FVector4> AllSectionsRefToLocalMatrices;
 	static_assert(sizeof(FVector4) == 4 * sizeof(float), "FVector4 should match 4 * floats");
-	{
-		TArray<FMatrix> RefToLocalMatrices;
-		if (SkelComp)
-		{
-			//-TODO: We need to handle MasterPoseComponent here
-			SkelComp->CacheRefToLocalMatrices(RefToLocalMatrices);
-		}
-		else
-		{
-			RefToLocalMatrices = SkelMesh->RefBasesInvMatrix;
-		}
 
-		TIndirectArray<FSkeletalMeshLODRenderData>& LODRenderDataArray = SkelMesh->GetResourceForRendering()->LODRenderData;
-		check(0 <= LODIndex && LODIndex < LODRenderDataArray.Num());
-		FSkeletalMeshLODRenderData& LODRenderData = LODRenderDataArray[LODIndex];
-		TArray<FSkelMeshRenderSection>& Sections = LODRenderData.RenderSections;
-		uint32 SectionCount = Sections.Num();
+	TArray<FVector4> AllSectionsRefToLocalMatrices;
+	TArray<FVector4> BoneSamplingData;
 
-		// Count number of matrices we want before appending all of them according to the per section mapping from BoneMap
-		uint32 Float4Count = 0;
-		for (const FSkelMeshRenderSection& Section : Sections)
+	auto FillBuffers =
+		[&](const TArray<FTransform>& BoneTransforms)
 		{
-			Float4Count += Section.BoneMap.Num() * 3;
-		}
-		check(Float4Count == 3 * SectionBoneCount);
-		AllSectionsRefToLocalMatrices.AddUninitialized(Float4Count);
+			check(BoneTransforms.Num() == SkelMesh->RefBasesInvMatrix.Num());
 
-		Float4Count = 0;
-		for (const FSkelMeshRenderSection& Section : Sections)
-		{
-			const uint32 MatrixCount = Section.BoneMap.Num();
-			for (uint32 m = 0; m < MatrixCount; ++m)
+			// Fill AllSectionsRefToLocalMatrices
+			TIndirectArray<FSkeletalMeshLODRenderData>& LODRenderDataArray = SkelMesh->GetResourceForRendering()->LODRenderData;
+			check(0 <= LODIndex && LODIndex < LODRenderDataArray.Num());
+			FSkeletalMeshLODRenderData& LODRenderData = LODRenderDataArray[LODIndex];
+			TArray<FSkelMeshRenderSection>& Sections = LODRenderData.RenderSections;
+			uint32 SectionCount = Sections.Num();
+
+			// Count number of matrices we want before appending all of them according to the per section mapping from BoneMap
+			uint32 Float4Count = 0;
+			for (const FSkelMeshRenderSection& Section : Sections)
 			{
-				RefToLocalMatrices[Section.BoneMap[m]].To3x4MatrixTranspose(&AllSectionsRefToLocalMatrices[Float4Count].X);
-				Float4Count += 3;
+				Float4Count += Section.BoneMap.Num() * 3;
 			}
-		}
-	}
+			check(Float4Count == 3 * SectionBoneCount);
+			AllSectionsRefToLocalMatrices.AddUninitialized(Float4Count);
 
-	// Generate information for bone sampling
-	TArray<FVector4> BoneSamplingData;
-	{
-		BoneSamplingData.Reserve((SamplingBoneCount + SamplingSocketCount) * 2);
+			Float4Count = 0;
+			for (const FSkelMeshRenderSection& Section : Sections)
+			{
+				const uint32 MatrixCount = Section.BoneMap.Num();
+				for (uint32 m=0; m < MatrixCount; ++m)
+				{
+					const int32 BoneIndex = Section.BoneMap[m];
+					const FTransform& BoneTransform = BoneTransforms[BoneIndex];
+					const FMatrix BoneMatrix = SkelMesh->RefBasesInvMatrix[BoneIndex] * BoneTransform.ToMatrixWithScale();
+					BoneMatrix.To3x4MatrixTranspose(&AllSectionsRefToLocalMatrices[Float4Count].X);
+					Float4Count += 3;
+				}
+			}
 
-		// Append bones
-		if (SkelComp != nullptr)
-		{
-			//-TODO: We need to handle MasterPoseComponent here
-			const TArray<FTransform>& ComponentTransforms = SkelComp->GetComponentSpaceTransforms();
-			check(ComponentTransforms.Num() == SamplingBoneCount);
+			// Fill BoneSamplingData
+			BoneSamplingData.Reserve((SamplingBoneCount + SamplingSocketCount) * 2);
+			check(BoneTransforms.Num() == SamplingBoneCount);
 
-			for (const FTransform& BoneTransform : ComponentTransforms)
+			for (const FTransform& BoneTransform : BoneTransforms)
 			{
 				const FQuat Rotation = BoneTransform.GetRotation();
 				BoneSamplingData.Add(BoneTransform.GetLocation());
 				BoneSamplingData.Add(FVector4(Rotation.X, Rotation.Y, Rotation.Z, Rotation.W));
 			}
-		}
-		else
+
+			// Append sockets
+			for (const FTransform& SocketTransform : InstanceData->GetSpecificSocketsCurrBuffer())
+			{
+				const FQuat Rotation = SocketTransform.GetRotation();
+				BoneSamplingData.Add(SocketTransform.GetLocation());
+				BoneSamplingData.Add(FVector4(Rotation.X, Rotation.Y, Rotation.Z, Rotation.W));
+			}
+		};
+
+	// If we have a component pull transforms from component otherwise grab from skel mesh
+	if (SkelComp)
+	{
+		if (USkinnedMeshComponent* MasterComponent = SkelComp->MasterPoseComponent.Get())
 		{
-			//-TODO: Opt and combine with MaterPostComponent
+			const TArray<int32>& MasterBoneMap = SkelComp->GetMasterBoneMap();
+			const int32 NumBones = MasterBoneMap.Num();
+
 			TArray<FTransform> TempBoneTransforms;
 			TempBoneTransforms.Reserve(SamplingBoneCount);
 
-			const TArray<FTransform>& RefTransforms = SkelMesh->RefSkeleton.GetRefBonePose();
-			for (int32 i=0; i < RefTransforms.Num(); ++i)
+			if (NumBones == 0)
 			{
-				FTransform BoneTransform = RefTransforms[i];
-				const int32 ParentIndex = SkelMesh->RefSkeleton.GetParentIndex(i);
-				if (TempBoneTransforms.IsValidIndex(ParentIndex))
+				// This case indicates an invalid master pose component (e.g. no skeletal mesh)
+				TempBoneTransforms.AddDefaulted(SamplingBoneCount);
+			}
+			else
+			{
+				const TArray<FTransform>& MasterTransforms = MasterComponent->GetComponentSpaceTransforms();
+				for (int32 BoneIndex=0; BoneIndex < NumBones; ++BoneIndex)
 				{
-					BoneTransform = BoneTransform * TempBoneTransforms[ParentIndex];
-				}
-				TempBoneTransforms.Add(BoneTransform);
+					if (MasterBoneMap.IsValidIndex(BoneIndex))
+					{
+						const int32 MasterIndex = MasterBoneMap[BoneIndex];
+						if (MasterIndex != INDEX_NONE && MasterIndex < MasterTransforms.Num())
+						{
+							TempBoneTransforms.Add(MasterTransforms[MasterIndex]);
+							continue;
+						}
+					}
 
-				const FQuat Rotation = BoneTransform.GetRotation();
-				BoneSamplingData.Add(BoneTransform.GetLocation());
-				BoneSamplingData.Add(FVector4(Rotation.X, Rotation.Y, Rotation.Z, Rotation.W));
+					const int32 ParentIndex = SkelMesh->RefSkeleton.GetParentIndex(BoneIndex);
+					FTransform BoneTransform =SkelMesh->RefSkeleton.GetRefBonePose()[BoneIndex];
+					if (TempBoneTransforms.IsValidIndex(ParentIndex))
+					{
+						BoneTransform = BoneTransform * TempBoneTransforms[ParentIndex];
+					}
+					TempBoneTransforms.Add(BoneTransform);
+				}
 			}
+			FillBuffers(TempBoneTransforms);
+		}
+		else
+		{
+			const TArray<FTransform>& ComponentTransforms = SkelComp->GetComponentSpaceTransforms();
+			FillBuffers(ComponentTransforms);
 		}
+	}
+	else
+	{
+		//-TODO: Opt and combine with MaterPoseComponent
+		TArray<FTransform> TempBoneTransforms;
+		TempBoneTransforms.Reserve(SamplingBoneCount);
 
-		// Append sockets
-		for (const FTransform& SocketTransform : InstanceData->GetSpecificSocketsCurrBuffer())
+		const TArray<FTransform>& RefTransforms = SkelMesh->RefSkeleton.GetRefBonePose();
+		for (int32 i=0; i < RefTransforms.Num(); ++i)
 		{
-			const FQuat Rotation = SocketTransform.GetRotation();
-			BoneSamplingData.Add(SocketTransform.GetLocation());
-			BoneSamplingData.Add(FVector4(Rotation.X, Rotation.Y, Rotation.Z, Rotation.W));
+			FTransform BoneTransform = RefTransforms[i];
+			const int32 ParentIndex = SkelMesh->RefSkeleton.GetParentIndex(i);
+			if (TempBoneTransforms.IsValidIndex(ParentIndex))
+			{
+				BoneTransform = BoneTransform * TempBoneTransforms[ParentIndex];
+			}
+			TempBoneTransforms.Add(BoneTransform);
 		}
+
+		FillBuffers(TempBoneTransforms);
 	}
 
 	FSkeletalMeshGpuDynamicBufferProxy* ThisProxy = this;
diff --git a/Engine/Plugins/FX/Niagara/Source/NiagaraCore/Private/NiagaraCustomVersion.cpp b/Engine/Plugins/FX/Niagara/Source/NiagaraCore/Private/NiagaraCustomVersion.cpp
--- a/Engine/Plugins/FX/Niagara/Source/NiagaraCore/Private/NiagaraCustomVersion.cpp
+++ b/Engine/Plugins/FX/Niagara/Source/NiagaraCore/Private/NiagaraCustomVersion.cpp
@@ -8,4 +8,4 @@ const FGuid FNiagaraCustomVersion::GUID(0xFCF57AFA, 0x50764283, 0xB9A9E658, 0xFF
 // Register the custom version with core
 FCustomVersionRegistration GRegisterNiagaraCustomVersion(FNiagaraCustomVersion::GUID, FNiagaraCustomVersion::LatestVersion, TEXT("NiagaraVer"));
 
-const FGuid FNiagaraCustomVersion::LatestScriptCompileVersion(0x0C1CF412, 0xA3B240BB, 0x900EBA38, 0x0CE36DDA );
+const FGuid FNiagaraCustomVersion::LatestScriptCompileVersion(0x1BDBCD02, 0xA1B55E47, 0x802D1B66, 0x5C5FCC12);
