Added code to customize the Fix Issue around deprecated scripts. This allows us to better handle the EmitterState deprecation, which had several complicated parameters set up.

diff --git a/Engine/Plugins/FX/Niagara/Source/Niagara/Classes/NiagaraScript.h b/Engine/Plugins/FX/Niagara/Source/Niagara/Classes/NiagaraScript.h
--- a/Engine/Plugins/FX/Niagara/Source/Niagara/Classes/NiagaraScript.h
+++ b/Engine/Plugins/FX/Niagara/Source/Niagara/Classes/NiagaraScript.h
@@ -19,6 +19,7 @@
 
 class UNiagaraDataInterface;
 class FNiagaraCompileRequestDataBase;
+class UNiagaraConvertInPlaceUtilityBase;
 
 void SerializeNiagaraShaderMaps(const TMap<const ITargetPlatform*, TArray<FNiagaraShaderScript*>>* PlatformScriptResourcesToSave, FArchive& Ar, TArray<FNiagaraShaderScript>& OutLoadedResources);
 void ProcessSerializedShaderMaps(UNiagaraScript* Owner, TArray<FNiagaraShaderScript>& LoadedResources, FNiagaraShaderScript& OutResourceForCurrentPlatform, FNiagaraShaderScript* (&OutScriptResourcesLoaded)[ERHIFeatureLevel::Num]);
@@ -424,6 +425,10 @@ public:
 	UPROPERTY(EditAnywhere, Category = Script, meta = (EditCondition = "bDeprecated"))
 	UNiagaraScript* DeprecationRecommendation;
 
+	/* Custom logic to convert the contents of an existing script assignment to this script.*/
+	UPROPERTY(EditAnywhere, Category = Script)
+	TSubclassOf<UNiagaraConvertInPlaceUtilityBase> ConversionUtility;
+
 	/** Is this script experimental and less supported? */
 	UPROPERTY(EditAnywhere, Category = Script)
 	uint32 bExperimental : 1;
@@ -600,6 +605,7 @@ public:
 		return ScriptShader;
 	}
 
+
 	void SetFeatureLevel(ERHIFeatureLevel::Type InFeatureLevel)		{ FeatureLevel = InFeatureLevel; }
 
 	NIAGARA_API void GenerateStatScopeIDs();
diff --git a/Engine/Plugins/FX/Niagara/Source/Niagara/Public/NiagaraConvertInPlaceUtilityBase.h b/Engine/Plugins/FX/Niagara/Source/Niagara/Public/NiagaraConvertInPlaceUtilityBase.h
new file mode 100644
--- /dev/null
+++ b/Engine/Plugins/FX/Niagara/Source/Niagara/Public/NiagaraConvertInPlaceUtilityBase.h
@@ -0,0 +1,22 @@
+// Copyright Epic Games, Inc. All Rights Reserved.
+
+#pragma once
+
+#include "CoreMinimal.h"
+#include "UObject/ObjectMacros.h"
+#include "UObject/Object.h"
+#include "Internationalization/Text.h"
+#include "NiagaraConvertInPlaceUtilityBase.generated.h"
+
+class UNiagaraScript;
+class UNiagaraNodeFunctionCall;
+class UNiagaraClipboardContent;
+class UNiagaraStackFunctionInputCollection;
+
+UCLASS(Abstract)
+class NIAGARA_API UNiagaraConvertInPlaceUtilityBase : public UObject
+{
+	GENERATED_BODY()
+public:
+	virtual bool Convert(UNiagaraScript* InOldScript, UNiagaraClipboardContent* InOldClipboardContent, UNiagaraScript* InNewScript, UNiagaraStackFunctionInputCollection* InInputCollection, UNiagaraClipboardContent* InNewClipboardContent, UNiagaraNodeFunctionCall* InCallingNode, FText& OutMessage) { return true; };
+};
\ No newline at end of file
diff --git a/Engine/Plugins/FX/Niagara/Source/NiagaraEditor/Private/NiagaraClipboard.cpp b/Engine/Plugins/FX/Niagara/Source/NiagaraEditor/Private/NiagaraClipboard.cpp
--- a/Engine/Plugins/FX/Niagara/Source/NiagaraEditor/Private/NiagaraClipboard.cpp
+++ b/Engine/Plugins/FX/Niagara/Source/NiagaraEditor/Private/NiagaraClipboard.cpp
@@ -86,6 +86,27 @@ const UNiagaraClipboardFunctionInput* UNiagaraClipboardFunctionInput::CreateDyna
 	return NewInput;
 }
 
+
+bool UNiagaraClipboardFunctionInput::CopyValuesFrom(const UNiagaraClipboardFunctionInput* InOther)
+{
+	if (InputType != InOther->InputType)
+		return false;
+
+	ValueMode = InOther->ValueMode;
+	Local = InOther->Local;
+	Linked = InOther->Linked;
+
+	Data = nullptr;
+	if (InOther->Data)
+		Data = Cast<UNiagaraDataInterface>(StaticDuplicateObject(InOther->Data, this));
+	Expression = InOther->Expression;
+	Dynamic = nullptr;
+	if (InOther->Dynamic)
+		Dynamic = Cast<UNiagaraClipboardFunction>(StaticDuplicateObject(InOther->Dynamic, this));
+
+	return true;
+}
+
 UNiagaraClipboardFunction* UNiagaraClipboardFunction::CreateScriptFunction(UObject* InOuter, FString InFunctionName, UNiagaraScript* InScript)
 {
 	UNiagaraClipboardFunction* NewFunction = Cast<UNiagaraClipboardFunction>(NewObject<UNiagaraClipboardFunction>(InOuter));
@@ -140,5 +161,173 @@ const UNiagaraClipboardContent* FNiagaraClipboard::GetClipboardContent() const
 		return ClipboardContentFactory.ClipboardContent;
 	}
 
+	return nullptr;
+}
+
+void UNiagaraClipboardEditorScriptingUtilities::TryGetInputByName(const TArray<UNiagaraClipboardFunctionInput*>& InInputs, FName InInputName, bool& bOutSucceeded, UNiagaraClipboardFunctionInput*& OutInput)
+{
+	for (UNiagaraClipboardFunctionInput* Input : InInputs)
+	{
+		if (Input->InputName == InInputName)
+		{
+			bOutSucceeded = true;
+			OutInput = const_cast<UNiagaraClipboardFunctionInput*>(Input);
+		}
+	}
+	bOutSucceeded = false;
+	OutInput = nullptr;
+}
+
+void UNiagaraClipboardEditorScriptingUtilities::TryGetLocalValueAsFloat(const UNiagaraClipboardFunctionInput* InInput, bool& bOutSucceeded, float& OutValue)
+{
+	if (InInput->ValueMode == ENiagaraClipboardFunctionInputValueMode::Local && InInput->InputType == FNiagaraTypeDefinition::GetFloatDef() && InInput->Local.Num() == InInput->InputType.GetSize())
+	{
+		bOutSucceeded = true;
+		OutValue = *((float*)InInput->Local.GetData());
+	}
+	else
+	{
+		bOutSucceeded = false;
+		OutValue = 0.0f;
+	}
+}
+
+void UNiagaraClipboardEditorScriptingUtilities::TryGetLocalValueAsInt(const UNiagaraClipboardFunctionInput* InInput, bool& bOutSucceeded, int32& OutValue)
+{
+	bool bCompatibleType = InInput->InputType == FNiagaraTypeDefinition::GetIntDef() || (InInput->InputType.IsEnum() && InInput->InputType.GetSize() == sizeof(int32)) ||
+		InInput->InputType == FNiagaraTypeDefinition::GetBoolDef();
+	if (InInput->ValueMode == ENiagaraClipboardFunctionInputValueMode::Local && bCompatibleType && InInput->Local.Num() == InInput->InputType.GetSize())
+	{
+		bOutSucceeded = true;
+		OutValue = *((int32*)InInput->Local.GetData());
+	}
+	else
+	{
+		bOutSucceeded = false;
+		OutValue = 0;
+	}
+}
+
+void UNiagaraClipboardEditorScriptingUtilities::TrySetLocalValueAsInt(UNiagaraClipboardFunctionInput* InInput, bool& bOutSucceeded, int32 InValue,bool bLooseTyping)
+{
+	bool bCompatibleType = InInput->InputType == FNiagaraTypeDefinition::GetIntDef() || (bLooseTyping && InInput->InputType.IsEnum() && InInput->InputType.GetSize() == sizeof(int32)) ||
+		(bLooseTyping && InInput->InputType == FNiagaraTypeDefinition::GetBoolDef() && InInput->InputType.GetSize() == sizeof(int32));
+	if (InInput->ValueMode == ENiagaraClipboardFunctionInputValueMode::Local && bCompatibleType && InInput->Local.Num() == InInput->InputType.GetSize())
+	{
+		bOutSucceeded = true;
+		*((int32*)InInput->Local.GetData()) = InValue;
+	}
+	else
+	{
+		bOutSucceeded = false;
+	}
+}
+
+FName UNiagaraClipboardEditorScriptingUtilities::GetTypeName(const UNiagaraClipboardFunctionInput* InInput)
+{
+	return InInput->InputType.GetFName();
+}
+
+const UNiagaraClipboardFunctionInput* CreateLocalValue(UObject* InOuter, FName InInputName, FNiagaraTypeDefinition InInputType, bool bInHasEditCondition, bool bInEditConditionValue, TArray<uint8>& InLocalValueData)
+{
+	return UNiagaraClipboardFunctionInput::CreateLocalValue(
+		InOuter != nullptr ? InOuter : GetTransientPackage(),
+		InInputName,
+		InInputType,
+		bInHasEditCondition ? TOptional<bool>(bInEditConditionValue) : TOptional<bool>(),
+		InLocalValueData);
+}
+
+FNiagaraTypeDefinition GetTypeByName(FName InTypeName)
+{
+	for (const FNiagaraTypeDefinition& RegisteredType : FNiagaraTypeRegistry::GetRegisteredTypes())
+	{
+		if (RegisteredType.GetFName() == InTypeName)
+		{
+			return RegisteredType;
+		}
+	}
+	return FNiagaraTypeDefinition();
+}
+
+UNiagaraClipboardFunctionInput* UNiagaraClipboardEditorScriptingUtilities::CreateFloatLocalValueInput(UObject* InOuter, FName InInputName, bool bInHasEditCondition, bool bInEditConditionValue, float InFloatValue)
+{
+	FNiagaraTypeDefinition InputType = FNiagaraTypeDefinition::GetFloatDef();
+	TArray<uint8> FloatValue;
+	FloatValue.AddUninitialized(InputType.GetSize());
+	FMemory::Memcpy(FloatValue.GetData(), &InFloatValue, InputType.GetSize());
+
+	return const_cast<UNiagaraClipboardFunctionInput*>(CreateLocalValue(InOuter, InInputName, InputType, bInHasEditCondition, bInEditConditionValue, FloatValue));
+}
+
+
+UNiagaraClipboardFunctionInput* UNiagaraClipboardEditorScriptingUtilities::CreateIntLocalValueInput(UObject* InOuter, FName InInputName, bool bInHasEditCondition, bool bInEditConditionValue, int32 InIntValue)
+{
+	FNiagaraTypeDefinition InputType = FNiagaraTypeDefinition::GetIntDef();
+	TArray<uint8> IntValue;
+	IntValue.AddUninitialized(InputType.GetSize());
+	FMemory::Memcpy(IntValue.GetData(), &InIntValue, InputType.GetSize());
+
+	return const_cast<UNiagaraClipboardFunctionInput*>(CreateLocalValue(InOuter, InInputName, InputType, bInHasEditCondition, bInEditConditionValue, IntValue));
+}
+
+UNiagaraClipboardFunctionInput* UNiagaraClipboardEditorScriptingUtilities::CreateLinkedValueInput(UObject* InOuter, FName InInputName, FName InInputTypeName, bool bInHasEditCondition, bool bInEditConditionValue, FName InLinkedValue)
+{
+	FNiagaraTypeDefinition InputType = GetTypeByName(InInputName);
+	if (InputType.IsValid())
+	{
+		return const_cast<UNiagaraClipboardFunctionInput*>(UNiagaraClipboardFunctionInput::CreateLinkedValue(
+			InOuter != nullptr ? InOuter : GetTransientPackage(),
+			InInputName,
+			InputType,
+			bInHasEditCondition ? TOptional<bool>(bInEditConditionValue) : TOptional<bool>(),
+			InLinkedValue));
+	}
+	return nullptr;
+}
+
+UNiagaraClipboardFunctionInput* UNiagaraClipboardEditorScriptingUtilities::CreateDataValueInput(UObject* InOuter, FName InInputName, FName InInputTypeName, bool bInHasEditCondition, bool bInEditConditionValue, UNiagaraDataInterface* InDataValue)
+{
+	FNiagaraTypeDefinition InputType = GetTypeByName(InInputName);
+	if (InputType.IsValid())
+	{
+		return const_cast<UNiagaraClipboardFunctionInput*>(UNiagaraClipboardFunctionInput::CreateDataValue(
+			InOuter != nullptr ? InOuter : GetTransientPackage(),
+			InInputName,
+			InputType,
+			bInHasEditCondition ? TOptional<bool>(bInEditConditionValue) : TOptional<bool>(),
+			InDataValue));
+	}
+	return nullptr;
+}
+
+UNiagaraClipboardFunctionInput* UNiagaraClipboardEditorScriptingUtilities::CreateExpressionValueInput(UObject* InOuter, FName InInputName, FName InInputTypeName, bool bInHasEditCondition, bool bInEditConditionValue, const FString& InExpressionValue)
+{
+	FNiagaraTypeDefinition InputType = GetTypeByName(InInputName);
+	if (InputType.IsValid())
+	{
+		return const_cast<UNiagaraClipboardFunctionInput*>(UNiagaraClipboardFunctionInput::CreateExpressionValue(
+			InOuter != nullptr ? InOuter : GetTransientPackage(),
+			InInputName,
+			InputType,
+			bInHasEditCondition ? TOptional<bool>(bInEditConditionValue) : TOptional<bool>(),
+			InExpressionValue));
+	}
+	return nullptr;
+}
+
+UNiagaraClipboardFunctionInput* UNiagaraClipboardEditorScriptingUtilities::CreateDynamicValueInput(UObject* InOuter, FName InInputName, FName InInputTypeName, bool bInHasEditCondition, bool bInEditConditionValue, FString InDynamicValueName, UNiagaraScript* InDynamicValue)
+{
+	FNiagaraTypeDefinition InputType = GetTypeByName(InInputName);
+	if (InputType.IsValid())
+	{
+		return const_cast<UNiagaraClipboardFunctionInput*>(UNiagaraClipboardFunctionInput::CreateDynamicValue(
+			InOuter != nullptr ? InOuter : GetTransientPackage(),
+			InInputName,
+			InputType,
+			bInHasEditCondition ? TOptional<bool>(bInEditConditionValue) : TOptional<bool>(),
+			InDynamicValueName,
+			InDynamicValue));
+	}
 	return nullptr;
 }
\ No newline at end of file
diff --git a/Engine/Plugins/FX/Niagara/Source/NiagaraEditor/Private/NiagaraConvertInPlaceEmitterAndSystemState.cpp b/Engine/Plugins/FX/Niagara/Source/NiagaraEditor/Private/NiagaraConvertInPlaceEmitterAndSystemState.cpp
new file mode 100644
--- /dev/null
+++ b/Engine/Plugins/FX/Niagara/Source/NiagaraEditor/Private/NiagaraConvertInPlaceEmitterAndSystemState.cpp
@@ -0,0 +1,229 @@
+// Copyright Epic Games, Inc. All Rights Reserved.
+
+#include "NiagaraConvertInPlaceEmitterAndSystemState.h"
+#include "NiagaraClipboard.h"
+#include "ViewModels/Stack/NiagaraStackFunctionInputCollection.h"
+
+bool UNiagaraConvertInPlaceEmitterAndSystemState::Convert(UNiagaraScript* InOldScript, UNiagaraClipboardContent* InOldClipboardContent, UNiagaraScript* InNewScript, UNiagaraStackFunctionInputCollection* InInputCollection, UNiagaraClipboardContent* InNewClipboardContent, UNiagaraNodeFunctionCall* InCallingNode, FText& OutMessage)
+{
+	UE_LOG(LogNiagaraEditor, Log, TEXT("%s to %s"), *InOldScript->GetPathName(), *InNewScript->GetPathName());
+
+	if (InOldScript->GetPathName() == TEXT("/Niagara/Modules/Emitter/EmitterLifeCycle.EmitterLifeCycle") && InNewScript->GetPathName() == TEXT("/Niagara/Modules/Emitter/EmitterState.EmitterState"))
+	{
+		if (InOldClipboardContent->Functions.Num() == 1 && InNewClipboardContent->Functions.Num() == 1 &&
+			InOldClipboardContent->Functions[0] && InNewClipboardContent->Functions[0])
+		{
+			int32 OldMaxLoopCount = 0;
+			const UNiagaraClipboardFunctionInput* InputOldLoopDuration = nullptr;
+			const UNiagaraClipboardFunctionInput* InputOldMaxLoopCount = nullptr;
+			const UNiagaraClipboardFunctionInput* InputOldDurationRecalcEachLoop = nullptr;
+			const UNiagaraClipboardFunctionInput* InputOldNextLoopDelay = nullptr;
+			const UNiagaraClipboardFunctionInput* InputOldDelayFirstLoopOnly = nullptr;
+			int32 AutoComplete = 0;
+			int32 CompleteOnInactive = 0;
+			bool bAutoCompleteWasLocal = false;
+			bool bCompleteOnInactiveWasLocal = false;
+
+			for (const UNiagaraClipboardFunctionInput* Input : InOldClipboardContent->Functions[0]->Inputs)
+			{
+				if (Input)
+				{
+					UE_LOG(LogNiagaraEditor, Log, TEXT("OldInput \"%s\""), *Input->InputName.ToString());
+
+					bool bSuccess = false;
+					if (Input->InputName == TEXT("MaxLoopCount"))
+					{
+						InputOldMaxLoopCount = Input;
+						UNiagaraClipboardEditorScriptingUtilities::TryGetLocalValueAsInt(Input, bSuccess, OldMaxLoopCount);
+					}
+					else if (Input->InputName == TEXT("NextLoopDuration"))
+					{
+						InputOldLoopDuration = Input;
+					}
+					else if (Input->InputName == TEXT("DurationRecalcEachLoop"))
+					{
+						InputOldDurationRecalcEachLoop = Input;
+					}
+					else if (Input->InputName == TEXT("NextLoopDelay"))
+					{
+						InputOldNextLoopDelay = Input;
+					}
+					else if (Input->InputName == TEXT("AutoComplete"))
+					{
+						UNiagaraClipboardEditorScriptingUtilities::TryGetLocalValueAsInt(Input, bAutoCompleteWasLocal, AutoComplete);
+					}
+					else if (Input->InputName == TEXT("CompleteOnInactive"))
+					{
+						UNiagaraClipboardEditorScriptingUtilities::TryGetLocalValueAsInt(Input, bCompleteOnInactiveWasLocal, CompleteOnInactive);
+					}
+					else if (Input->InputName == TEXT("DelayFirstLoopOnly"))
+					{
+						InputOldDelayFirstLoopOnly = Input;						 
+					}
+				}
+			}
+
+			TArray<const UNiagaraClipboardFunctionInput*> NewInputs;
+			bool bMultipleLoopsFound = false;
+			bool bUseLoopDelay = false;
+			TArray<FName> ConvertedItems;
+			
+			for (const UNiagaraClipboardFunctionInput* Input : InNewClipboardContent->Functions[0]->Inputs)
+			{
+				if (Input)
+				{
+					UE_LOG(LogNiagaraEditor, Log, TEXT("New Input \"%s\""), *Input->InputName.ToString());
+					bool bSuccess = false;
+					if (Input->InputName == TEXT("Life Cycle Mode"))
+					{
+						// When converting old Emitter Life Cycle, always assume that new type is Self
+						UNiagaraClipboardFunctionInput* NewInput = Cast< UNiagaraClipboardFunctionInput>(StaticDuplicateObject(Input, GetTransientPackage()));
+						UNiagaraClipboardEditorScriptingUtilities::TrySetLocalValueAsInt(NewInput, bSuccess, 1);
+						NewInputs.Add(NewInput);
+						if (bSuccess)
+							ConvertedItems.Add(Input->InputName);
+					}
+					else if (Input->InputName == TEXT("Loop Behavior"))
+					{
+						UNiagaraClipboardFunctionInput* NewInput = Cast< UNiagaraClipboardFunctionInput>(StaticDuplicateObject(Input, GetTransientPackage()));
+						int32 NewValue = -1;
+						if (InputOldMaxLoopCount && InputOldMaxLoopCount->ValueMode == ENiagaraClipboardFunctionInputValueMode::Local)
+						{
+							if (OldMaxLoopCount == 0)
+								NewValue = 0; // Infinite
+							else if (OldMaxLoopCount == 1)
+								NewValue = 1; // Once
+							else if (OldMaxLoopCount > 1)
+								NewValue = 2; // Multiple
+						}
+						
+						// Force to multiple if anything other than local
+						if (NewValue < 0)
+							NewValue = 2; // Multiple
+
+						if (NewValue == 2)
+							bMultipleLoopsFound = true;
+
+						UNiagaraClipboardEditorScriptingUtilities::TrySetLocalValueAsInt(NewInput, bSuccess, NewValue);
+						NewInputs.Add(NewInput);
+						if (bSuccess)
+							ConvertedItems.Add(Input->InputName);
+					}
+					else if (Input->InputName == TEXT("Inactive Response"))
+					{
+
+						UNiagaraClipboardFunctionInput* NewInput = Cast< UNiagaraClipboardFunctionInput>(StaticDuplicateObject(Input, GetTransientPackage()));
+						NewInputs.Add(NewInput);
+						
+						int32 NewValue = -1;
+						if (AutoComplete && bAutoCompleteWasLocal)
+							NewValue = 0; // Complete
+						else if (CompleteOnInactive && bCompleteOnInactiveWasLocal)
+							NewValue = 1; // Kill
+
+						// Force to complete if anything else
+						if (NewValue < 0)
+							NewValue = 0; // Complete
+
+						UNiagaraClipboardEditorScriptingUtilities::TrySetLocalValueAsInt(NewInput, bSuccess, NewValue);
+						if (bSuccess)
+							ConvertedItems.Add(Input->InputName);
+					}
+					else if (Input->InputName == TEXT("Loop Duration"))
+					{
+						if (InputOldLoopDuration)
+						{
+							UNiagaraClipboardFunctionInput* NewInput = Cast< UNiagaraClipboardFunctionInput>(StaticDuplicateObject(InputOldLoopDuration, GetTransientPackage()));
+							NewInput->InputName = Input->InputName;
+							NewInputs.Add(NewInput);
+							ConvertedItems.Add(Input->InputName);
+						}
+					}
+					else if (Input->InputName == TEXT("Recalculate Duration Each Loop"))
+					{
+						if (InputOldDurationRecalcEachLoop)
+						{
+							UNiagaraClipboardFunctionInput* NewInput = Cast< UNiagaraClipboardFunctionInput>(StaticDuplicateObject(InputOldDurationRecalcEachLoop, GetTransientPackage()));
+							NewInput->InputName = Input->InputName;
+							NewInputs.Add(NewInput);
+							ConvertedItems.Add(Input->InputName);
+						}
+					}
+					else if (Input->InputName == TEXT("Loop Delay"))
+					{
+						if (InputOldNextLoopDelay)
+						{
+							UNiagaraClipboardFunctionInput* NewInput = Cast< UNiagaraClipboardFunctionInput>(StaticDuplicateObject(InputOldNextLoopDelay, GetTransientPackage()));
+							NewInput->InputName = Input->InputName;
+							NewInputs.Add(NewInput);
+							ConvertedItems.Add(Input->InputName);
+						}
+					}
+					else if (Input->InputName == TEXT("UseLoopDelay"))
+					{
+						if (InputOldNextLoopDelay)
+						{
+							float OutFloat = 0.0f;
+							bool bLoopDelayFound = false;
+							UNiagaraClipboardEditorScriptingUtilities::TryGetLocalValueAsFloat(InputOldNextLoopDelay, bLoopDelayFound, OutFloat);
+							if ((bLoopDelayFound && OutFloat != 0.0f) || InputOldNextLoopDelay->ValueMode != ENiagaraClipboardFunctionInputValueMode::Local)
+							{
+								bUseLoopDelay = true;
+								int32 NewValue = -1;
+								UNiagaraClipboardFunctionInput* NewInput = Cast< UNiagaraClipboardFunctionInput>(StaticDuplicateObject(Input, GetTransientPackage())); UNiagaraClipboardEditorScriptingUtilities::TrySetLocalValueAsInt(NewInput, bSuccess, NewValue);
+								NewInputs.Add(NewInput);
+								if (bSuccess)
+									ConvertedItems.Add(Input->InputName);
+							}
+						}
+					}
+				}
+			}
+
+			InInputCollection->SetValuesFromClipboardFunctionInputs(NewInputs);
+			//InCallingNode->RefreshFromExternalChanges();
+
+			// We have to do this as a second pass as Loop Count is behind a static switch set above.
+			NewInputs.Empty();
+			if (InputOldMaxLoopCount && bMultipleLoopsFound)
+			{
+				UNiagaraClipboardFunctionInput* NewInput = Cast< UNiagaraClipboardFunctionInput>(StaticDuplicateObject(InputOldMaxLoopCount, GetTransientPackage()));
+				NewInput->InputName = TEXT("Loop Count");
+				NewInputs.Add(NewInput);
+				ConvertedItems.Add(NewInput->InputName);
+
+				InInputCollection->SetValuesFromClipboardFunctionInputs(NewInputs);
+			}
+
+			// Use first frame delay is behind a static switch set above.
+			NewInputs.Empty();
+			if (bUseLoopDelay && InputOldDelayFirstLoopOnly)
+			{
+				UNiagaraClipboardFunctionInput* NewInput = Cast< UNiagaraClipboardFunctionInput>(StaticDuplicateObject(InputOldDelayFirstLoopOnly, GetTransientPackage()));
+				NewInput->InputName = TEXT("Delay First Loop Only");
+				NewInputs.Add(NewInput);
+				ConvertedItems.Add(NewInput->InputName);
+
+				InInputCollection->SetValuesFromClipboardFunctionInputs(NewInputs);
+			}
+
+			FString ConcatString = FString::Printf(TEXT("Copied over %d properties\nPlease double-check for correctness\n"), ConvertedItems.Num());
+			for (int32 i = 0; i < ConvertedItems.Num(); i++)
+			{
+				ConcatString += ConvertedItems[i].ToString();
+
+				if (i < ConvertedItems.Num() - 1)
+					ConcatString += TEXT(",\n");
+				else
+					ConcatString += TEXT("\n");
+
+			}
+			OutMessage = FText::FromString(ConcatString);
+			if (ConvertedItems.Num() > 0)
+				return true;
+			else
+				return false;
+		}
+	}
+	return false;
+};
\ No newline at end of file
diff --git a/Engine/Plugins/FX/Niagara/Source/NiagaraEditor/Private/NiagaraConvertInPlaceEmitterAndSystemState.h b/Engine/Plugins/FX/Niagara/Source/NiagaraEditor/Private/NiagaraConvertInPlaceEmitterAndSystemState.h
new file mode 100644
--- /dev/null
+++ b/Engine/Plugins/FX/Niagara/Source/NiagaraEditor/Private/NiagaraConvertInPlaceEmitterAndSystemState.h
@@ -0,0 +1,16 @@
+// Copyright Epic Games, Inc. All Rights Reserved.
+
+#pragma once
+
+#include "NiagaraConvertInPlaceUtilityBase.h"
+#include "NiagaraScript.h"
+#include "NiagaraNodeFunctionCall.h"
+#include "NiagaraConvertInPlaceEmitterAndSystemState.generated.h"
+
+UCLASS()
+class NIAGARAEDITOR_API  UNiagaraConvertInPlaceEmitterAndSystemState : public UNiagaraConvertInPlaceUtilityBase
+{
+	GENERATED_BODY()
+public:
+	virtual bool Convert(UNiagaraScript* InOldScript, UNiagaraClipboardContent* InOldClipboardContent, UNiagaraScript* InNewScript, UNiagaraStackFunctionInputCollection* InInputCollection, UNiagaraClipboardContent* InNewClipboardContent, UNiagaraNodeFunctionCall* InCallingNode, FText& OutMessage) override;
+};
\ No newline at end of file
diff --git a/Engine/Plugins/FX/Niagara/Source/NiagaraEditor/Private/NiagaraNodeFunctionCall.cpp b/Engine/Plugins/FX/Niagara/Source/NiagaraEditor/Private/NiagaraNodeFunctionCall.cpp
--- a/Engine/Plugins/FX/Niagara/Source/NiagaraEditor/Private/NiagaraNodeFunctionCall.cpp
+++ b/Engine/Plugins/FX/Niagara/Source/NiagaraEditor/Private/NiagaraNodeFunctionCall.cpp
@@ -930,9 +930,9 @@ UEdGraphPin* UNiagaraNodeFunctionCall::FindStaticSwitchInputPin(const FName& Var
 	return nullptr;
 }
 
-void UNiagaraNodeFunctionCall::SuggestName(FString SuggestedName)
+void UNiagaraNodeFunctionCall::SuggestName(FString SuggestedName, bool bForceSuggestion)
 {
-	ComputeNodeName(SuggestedName);
+	ComputeNodeName(SuggestedName, bForceSuggestion);
 }
 
 UNiagaraNodeFunctionCall::FOnInputsChanged& UNiagaraNodeFunctionCall::OnInputsChanged()
@@ -979,7 +979,7 @@ void UNiagaraNodeFunctionCall::AutowireNewNode(UEdGraphPin* FromPin)
 	ComputeNodeName();
 }
 
-void UNiagaraNodeFunctionCall::ComputeNodeName(FString SuggestedName)
+void UNiagaraNodeFunctionCall::ComputeNodeName(FString SuggestedName, bool bForceSuggestion)
 {
 	FString FunctionName = FunctionScript ? FunctionScript->GetName() : Signature.GetName();
 	FName ProposedName;
@@ -987,7 +987,7 @@ void UNiagaraNodeFunctionCall::ComputeNodeName(FString SuggestedName)
 	{ 
 		// If we have a suggested name and and either there is no function name, or it is a permutation of the function name
 		// it can be used as the proposed name.
-		if (FunctionName.IsEmpty() || SuggestedName == FunctionName || (SuggestedName.StartsWith(FunctionName) && SuggestedName.RightChop(FunctionName.Len()).IsNumeric()))
+		if (bForceSuggestion || FunctionName.IsEmpty() || SuggestedName == FunctionName || (SuggestedName.StartsWith(FunctionName) && SuggestedName.RightChop(FunctionName.Len()).IsNumeric()))
 		{
 			ProposedName = *SuggestedName;
 		}
diff --git a/Engine/Plugins/FX/Niagara/Source/NiagaraEditor/Private/NiagaraStackFunctionInputBinder.cpp b/Engine/Plugins/FX/Niagara/Source/NiagaraEditor/Private/NiagaraStackFunctionInputBinder.cpp
--- a/Engine/Plugins/FX/Niagara/Source/NiagaraEditor/Private/NiagaraStackFunctionInputBinder.cpp
+++ b/Engine/Plugins/FX/Niagara/Source/NiagaraEditor/Private/NiagaraStackFunctionInputBinder.cpp
@@ -291,7 +291,7 @@ FGuid FNiagaraStackFunctionInputBinder::GetChangeIdFromFunctionScript() const
 
 void FNiagaraStackFunctionInputBinder::RefreshGraphPins() const
 {
-	if (FunctionCallNode->FunctionScript == nullptr || FunctionCallNode->FunctionScript->GetSource() == nullptr)
+	if (FunctionCallNode.IsValid() == false || FunctionCallNode->FunctionScript == nullptr || FunctionCallNode->FunctionScript->GetSource() == nullptr)
 	{
 		OverridePin = nullptr;
 		DefaultPin = nullptr;
diff --git a/Engine/Plugins/FX/Niagara/Source/NiagaraEditor/Private/ViewModels/Stack/NiagaraStackFunctionInput.cpp b/Engine/Plugins/FX/Niagara/Source/NiagaraEditor/Private/ViewModels/Stack/NiagaraStackFunctionInput.cpp
--- a/Engine/Plugins/FX/Niagara/Source/NiagaraEditor/Private/ViewModels/Stack/NiagaraStackFunctionInput.cpp
+++ b/Engine/Plugins/FX/Niagara/Source/NiagaraEditor/Private/ViewModels/Stack/NiagaraStackFunctionInput.cpp
@@ -44,6 +44,9 @@
 #include "EdGraph/EdGraphPin.h"
 #include "INiagaraEditorTypeUtilities.h"
 #include "ViewModels/Stack/NiagaraStackInputCategory.h"
+#include "Framework/Notifications/NotificationManager.h"
+#include "Widgets/Notifications/SNotificationList.h"
+#include "NiagaraConvertInPlaceUtilityBase.h"
 
 #include "NiagaraScriptVariable.h"
 
@@ -1636,6 +1639,15 @@ void UNiagaraStackFunctionInput::ReassignDynamicInputScript(UNiagaraScript* Dyna
 		const FString OldName = InputValues.DynamicNode->GetFunctionName();
 
 		InputValues.DynamicNode->Modify();
+
+		UNiagaraClipboardContent* OldClipboardContent = nullptr;
+		UNiagaraScript* OldScript = InputValues.DynamicNode->FunctionScript;
+		if (DynamicInputScript->ConversionUtility != nullptr)
+		{
+			OldClipboardContent = UNiagaraClipboardContent::Create();
+			Copy(OldClipboardContent);
+		}
+
 		InputValues.DynamicNode->FunctionScript = DynamicInputScript;
 
 		// intermediate refresh to purge any rapid iteration parameters that have been removed in the new script
@@ -1649,8 +1661,31 @@ void UNiagaraStackFunctionInput::ReassignDynamicInputScript(UNiagaraScript* Dyna
 		FNiagaraStackGraphUtilities::RenameReferencingParameters(System, Emitter, *InputValues.DynamicNode.Get(), OldName, NewName);
 
 		InputValues.DynamicNode->RefreshFromExternalChanges();
+
 		InputValues.DynamicNode->MarkNodeRequiresSynchronization(TEXT("Dynamic input script reassigned."), true);
 		RefreshChildren();
+		
+		if (DynamicInputScript->ConversionUtility != nullptr && OldClipboardContent != nullptr)
+		{
+			UNiagaraConvertInPlaceUtilityBase* ConversionUtility = NewObject< UNiagaraConvertInPlaceUtilityBase>(GetTransientPackage(), DynamicInputScript->ConversionUtility);
+
+			UNiagaraClipboardContent* NewClipboardContent = UNiagaraClipboardContent::Create();
+			Copy(NewClipboardContent);
+			TArray<UNiagaraStackFunctionInputCollection*> DynamicInputCollections;
+			GetUnfilteredChildrenOfType(DynamicInputCollections);
+
+			FText ConvertMessage;
+			if (ConversionUtility && DynamicInputCollections.Num() == 0 && !ConversionUtility->Convert(OldScript, OldClipboardContent, DynamicInputScript, DynamicInputCollections[0], NewClipboardContent, InputValues.DynamicNode.Get(), ConvertMessage) || !ConvertMessage.IsEmptyOrWhitespace())
+			{
+				// Notify the end-user about the convert message, but continue the process as they could always undo.
+				FNotificationInfo Msg(FText::Format(LOCTEXT("FixConvertInPlace", "Conversion Note: {0}"), ConvertMessage));
+				Msg.ExpireDuration = 5.0f;
+				Msg.bFireAndForget = true;
+				Msg.Image = FCoreStyle::Get().GetBrush(TEXT("MessageLog.Info"));
+				FSlateNotificationManager::Get().AddNotification(Msg);
+			}
+		}
+
 	}
 }
 
diff --git a/Engine/Plugins/FX/Niagara/Source/NiagaraEditor/Private/ViewModels/Stack/NiagaraStackModuleItem.cpp b/Engine/Plugins/FX/Niagara/Source/NiagaraEditor/Private/ViewModels/Stack/NiagaraStackModuleItem.cpp
--- a/Engine/Plugins/FX/Niagara/Source/NiagaraEditor/Private/ViewModels/Stack/NiagaraStackModuleItem.cpp
+++ b/Engine/Plugins/FX/Niagara/Source/NiagaraEditor/Private/ViewModels/Stack/NiagaraStackModuleItem.cpp
@@ -32,6 +32,8 @@
 #include "Widgets/SWidget.h"
 #include "NiagaraActions.h"
 #include "NiagaraClipboard.h"
+#include "NiagaraConvertInPlaceUtilityBase.h"
+#include "Framework/Notifications/NotificationManager.h"
 
 #include "ScopedTransaction.h"
 #include "NiagaraScriptVariable.h"
@@ -344,8 +346,14 @@ void UNiagaraStackModuleItem::RefreshIssues(TArray<FStackIssue>& NewIssues)
 					{
 						NewIssues[AddIdx].InsertFix(0,
 							FStackIssueFix(
-							LOCTEXT("SelectNewModuleScriptFixUseRecommended", "Use recommended replacement"),
-							FStackIssueFixDelegate::CreateLambda([this]() { ReassignModuleScript(FunctionCallNode->FunctionScript->DeprecationRecommendation); })));
+							LOCTEXT("SelectNewModuleScriptFixUseRecommended", "Use recommended replacement and keep a disabled backup"),
+							FStackIssueFixDelegate::CreateLambda([this]() 
+							{
+									if (DeprecationDelegate.IsBound())
+									{
+										DeprecationDelegate.Execute(this);
+									}
+							})));
 					}
 				}
 				else
@@ -353,7 +361,7 @@ void UNiagaraStackModuleItem::RefreshIssues(TArray<FStackIssue>& NewIssues)
 					NewIssues.Add(FStackIssue(
 						EStackIssueSeverity::Warning,
 						ModuleScriptDeprecationShort,
-						FText::Format(LOCTEXT("ModuleScriptDeprecationFixParentLong", "The script asset for the assigned module {0} has been deprecated.  This module is inherited and this issue must be fixed in the parent emitter."),
+						FText::Format(LOCTEXT("ModuleScriptDeprecationFixParentLong", "The script asset for the assigned module {0} has been deprecated.\nThis module is inherited and this issue must be fixed in the parent emitter.\nYou will need to touch up this instance once that is done."),
 						FText::FromString(FunctionCallNode->GetFunctionName())),
 						GetStackEditorDataKey(),
 						false));
@@ -948,9 +956,15 @@ void UNiagaraStackModuleItem::ReassignModuleScript(UNiagaraScript* ModuleScript)
 		FScopedTransaction ScopedTransaction(LOCTEXT("ReassignModuleTransaction", "Reassign module script"));
 
 		const FString OldName = FunctionCallNode->GetFunctionName();
-		const UNiagaraScript* OldScript = FunctionCallNode->FunctionScript;
+		UNiagaraScript* OldScript = FunctionCallNode->FunctionScript;
 
 		FunctionCallNode->Modify();
+		UNiagaraClipboardContent* OldClipboardContent = nullptr;
+		if (ModuleScript->ConversionUtility != nullptr)
+		{
+			OldClipboardContent = UNiagaraClipboardContent::Create();
+			Copy(OldClipboardContent);
+		}
 		FunctionCallNode->FunctionScript = ModuleScript;
 		
 		// intermediate refresh to purge any rapid iteration parameters that have been removed in the new script
@@ -962,10 +976,29 @@ void UNiagaraStackModuleItem::ReassignModuleScript(UNiagaraScript* ModuleScript)
 		UNiagaraSystem& System = GetSystemViewModel()->GetSystem();
 		UNiagaraEmitter* Emitter = GetEmitterViewModel().IsValid() ? GetEmitterViewModel()->GetEmitter() : nullptr;
 		FNiagaraStackGraphUtilities::RenameReferencingParameters(System, Emitter, *FunctionCallNode, OldName, NewName);
-
 		FunctionCallNode->RefreshFromExternalChanges();
 		FunctionCallNode->MarkNodeRequiresSynchronization(TEXT("Module script reassigned."), true);
 		RefreshChildren();
+		
+		
+		if (ModuleScript->ConversionUtility != nullptr && OldClipboardContent != nullptr)
+		{
+			UNiagaraConvertInPlaceUtilityBase* ConversionUtility = NewObject< UNiagaraConvertInPlaceUtilityBase>(GetTransientPackage(), ModuleScript->ConversionUtility);
+			FText ConvertMessage;
+
+			UNiagaraClipboardContent* NewClipboardContent = UNiagaraClipboardContent::Create();
+			Copy(NewClipboardContent);
+
+			if (ConversionUtility && !ConversionUtility->Convert(OldScript, OldClipboardContent, ModuleScript, InputCollection, NewClipboardContent, FunctionCallNode, ConvertMessage) || !ConvertMessage.IsEmptyOrWhitespace())
+			{
+				// Notify the end-user about the convert message, but continue the process as they could always undo.
+				FNotificationInfo Msg(FText::Format(LOCTEXT("FixConvertInPlace", "Conversion Note: {0}"), ConvertMessage));
+				Msg.ExpireDuration = 5.0f;
+				Msg.bFireAndForget = true;
+				Msg.Image = FCoreStyle::Get().GetBrush(TEXT("MessageLog.Note"));
+				FSlateNotificationManager::Get().AddNotification(Msg);
+			}
+		}
 	}
 }
 
diff --git a/Engine/Plugins/FX/Niagara/Source/NiagaraEditor/Private/ViewModels/Stack/NiagaraStackScriptItemGroup.cpp b/Engine/Plugins/FX/Niagara/Source/NiagaraEditor/Private/ViewModels/Stack/NiagaraStackScriptItemGroup.cpp
--- a/Engine/Plugins/FX/Niagara/Source/NiagaraEditor/Private/ViewModels/Stack/NiagaraStackScriptItemGroup.cpp
+++ b/Engine/Plugins/FX/Niagara/Source/NiagaraEditor/Private/ViewModels/Stack/NiagaraStackScriptItemGroup.cpp
@@ -462,6 +462,7 @@ void UNiagaraStackScriptItemGroup::RefreshChildrenInternal(const TArray<UNiagara
 				ModuleItem->OnModifiedGroupItems().AddUObject(this, &UNiagaraStackScriptItemGroup::ChildModifiedGroupItems);
 				ModuleItem->SetOnRequestCanPaste(UNiagaraStackModuleItem::FOnRequestCanPaste::CreateUObject(this, &UNiagaraStackScriptItemGroup::ChildRequestCanPaste));
 				ModuleItem->SetOnRequestPaste(UNiagaraStackModuleItem::FOnRequestPaste::CreateUObject(this, &UNiagaraStackScriptItemGroup::ChildRequestPaste));
+				ModuleItem->SetOnRequestDeprecationRecommended(UNiagaraStackModuleItem::FOnRequestDeprecationRecommended::CreateUObject(this, &UNiagaraStackScriptItemGroup::ChildRequestDeprecatedRecommendation));
 			}
 
 			NewChildren.Add(ModuleItem);
@@ -1014,6 +1015,44 @@ void UNiagaraStackScriptItemGroup::ChildRequestPaste(const UNiagaraClipboardCont
 	PasteModules(ClipboardContent, PasteIndex, OutPasteWarning);
 }
 
+void UNiagaraStackScriptItemGroup::ChildRequestDeprecatedRecommendation(UNiagaraStackModuleItem* TargetChild)
+{
+
+	FScopedTransaction ScopedTransaction(LOCTEXT("UpdateDeprecatedScript", "Update module script to new version"));
+
+	UNiagaraClipboardContent* ClipboardContent = UNiagaraClipboardContent::Create();
+	TargetChild->GetModuleNode().GetNiagaraGraph()->Modify();
+	TargetChild->GetModuleNode().Modify();
+	Modify();
+	TargetChild->Modify();
+
+	// Step 1: Make a copy of the existing node so that we can perform surgery on it to update.
+	UNiagaraScript* TargetScript = TargetChild->GetModuleNode().FunctionScript->DeprecationRecommendation;
+	TargetChild->Copy(ClipboardContent);
+
+	// Step 2: Disable the old master so that end users can use it as reference and inheritance isn't wiped out.
+	//TargetChild->Rename();
+	//TargetChild->GetModuleNode().SuggestName(TEXT("Deprecated Original " + TargetChild->GetModuleNode().GetFunctionName()), true);
+	TargetChild->OnRenamed(FText::Format(LOCTEXT("NewDeprecationName", "Deprecated Original {0}"), TargetChild->GetDisplayName()));
+	TargetChild->SetEnabled(false);
+
+	// Step 3: Paste in the copy
+	FText PasteWarning;
+	int32 NewChildIdx = TargetChild->GetModuleIndex() + 1;
+	PasteModules(ClipboardContent, TargetChild->GetModuleIndex() + 1, PasteWarning);
+	RefreshChildren();
+
+	// Step 4: convert new copy in place to suggested script.
+	TArray<UNiagaraStackModuleItem*> ModuleItems;
+	GetUnfilteredChildrenOfType(ModuleItems);
+
+	if (ModuleItems.Num() > NewChildIdx)
+	{
+		ModuleItems[NewChildIdx]->ReassignModuleScript(TargetScript);
+	}
+	//
+}
+
 void UNiagaraStackScriptItemGroup::OnScriptGraphChanged(const struct FEdGraphEditAction& InAction)
 {
 	if (InAction.Action == GRAPHACTION_RemoveNode)
diff --git a/Engine/Plugins/FX/Niagara/Source/NiagaraEditor/Public/NiagaraClipboard.h b/Engine/Plugins/FX/Niagara/Source/NiagaraEditor/Public/NiagaraClipboard.h
--- a/Engine/Plugins/FX/Niagara/Source/NiagaraEditor/Public/NiagaraClipboard.h
+++ b/Engine/Plugins/FX/Niagara/Source/NiagaraEditor/Public/NiagaraClipboard.h
@@ -67,6 +67,8 @@ public:
 
 	UPROPERTY()
 	UNiagaraClipboardFunction* Dynamic;
+
+	bool CopyValuesFrom(const UNiagaraClipboardFunctionInput* InOther);
 };
 
 UENUM()
@@ -134,4 +136,44 @@ public:
 	void SetClipboardContent(UNiagaraClipboardContent* ClipboardContent);
 
 	const UNiagaraClipboardContent* GetClipboardContent() const;
+};
+
+UCLASS()
+class UNiagaraClipboardEditorScriptingUtilities : public UObject
+{
+	GENERATED_BODY()
+
+public:
+	UFUNCTION(BlueprintPure, Category = "Input")
+	static void TryGetInputByName(const TArray<UNiagaraClipboardFunctionInput*>& InInputs, FName InInputName, bool& bOutSucceeded, UNiagaraClipboardFunctionInput*& OutInput);
+
+	UFUNCTION(BlueprintPure, Category = "Input")
+	static void TryGetLocalValueAsFloat(const UNiagaraClipboardFunctionInput* InInput, bool& bOutSucceeded, float& OutValue);
+
+	UFUNCTION(BlueprintPure, Category = "Input")
+	static void TryGetLocalValueAsInt(const UNiagaraClipboardFunctionInput* InInput, bool& bOutSucceeded, int32& OutValue);
+
+	UFUNCTION(BlueprintPure, Category = "Input")
+	static void TrySetLocalValueAsInt(UNiagaraClipboardFunctionInput* InInput, bool& bOutSucceeded, int32 InValue, bool bLooseTyping = true);
+
+	UFUNCTION(BlueprintPure, Category = "Input")
+	static FName GetTypeName(const UNiagaraClipboardFunctionInput* InInput);
+
+	UFUNCTION(BlueprintPure, Category = "Input")
+	static UNiagaraClipboardFunctionInput* CreateFloatLocalValueInput(UObject* InOuter, FName InInputName, bool bInHasEditCondition, bool bInEditConditionValue, float InLocalValue);
+
+	UFUNCTION(BlueprintPure, Category = "Input")
+	static UNiagaraClipboardFunctionInput* CreateIntLocalValueInput(UObject* InOuter, FName InInputName, bool bInHasEditCondition, bool bInEditConditionValue, int32 InLocalValue);
+
+	UFUNCTION(BlueprintPure, Category = "Input")
+	static UNiagaraClipboardFunctionInput* CreateLinkedValueInput(UObject* InOuter, FName InInputName, FName InInputTypeName, bool bInHasEditCondition, bool bInEditConditionValue, FName InLinkedValue);
+
+	UFUNCTION(BlueprintPure, Category = "Input")
+	static UNiagaraClipboardFunctionInput* CreateDataValueInput(UObject* InOuter, FName InInputName, FName InInputTypeName, bool bInHasEditCondition, bool bInEditConditionValue, UNiagaraDataInterface* InDataValue);
+
+	UFUNCTION(BlueprintPure, Category = "Input")
+	static UNiagaraClipboardFunctionInput* CreateExpressionValueInput(UObject* InOuter, FName InInputName, FName InInputTypeName, bool bInHasEditCondition, bool bInEditConditionValue, const FString& InExpressionValue);
+
+	UFUNCTION(BlueprintPure, Category = "Input")
+	static UNiagaraClipboardFunctionInput* CreateDynamicValueInput(UObject* InOuter, FName InInputName, FName InInputTypeName, bool bInHasEditCondition, bool bInEditConditionValue, FString InDynamicValueName, UNiagaraScript* InDynamicValue);
 };
\ No newline at end of file
diff --git a/Engine/Plugins/FX/Niagara/Source/NiagaraEditor/Public/NiagaraNodeFunctionCall.h b/Engine/Plugins/FX/Niagara/Source/NiagaraEditor/Public/NiagaraNodeFunctionCall.h
--- a/Engine/Plugins/FX/Niagara/Source/NiagaraEditor/Public/NiagaraNodeFunctionCall.h
+++ b/Engine/Plugins/FX/Niagara/Source/NiagaraEditor/Public/NiagaraNodeFunctionCall.h
@@ -122,7 +122,7 @@ public:
 
 	/** Tries to rename this function call to a new name.  The actual name that gets applied might be different due to conflicts with existing
 		nodes with the same name. */
-	void SuggestName(FString SuggestedName);
+	void SuggestName(FString SuggestedName, bool bForceSuggestion = false);
 
 	FOnInputsChanged& OnInputsChanged();
 
@@ -145,7 +145,7 @@ protected:
 	virtual bool CanMovePin(const UEdGraphPin* Pin) const override { return false; }
 
 	/** Resets the node name based on the referenced script or signature. Guaranteed unique within a given graph instance.*/
-	void ComputeNodeName(FString SuggestedName = FString());
+	void ComputeNodeName(FString SuggestedName = FString(), bool bForceSuggestion = false);
 
 	void SetPinAutoGeneratedDefaultValue(UEdGraphPin& FunctionInputPin, UNiagaraNodeInput& FunctionScriptInputNode);
 
diff --git a/Engine/Plugins/FX/Niagara/Source/NiagaraEditor/Public/ViewModels/Stack/NiagaraStackModuleItem.h b/Engine/Plugins/FX/Niagara/Source/NiagaraEditor/Public/ViewModels/Stack/NiagaraStackModuleItem.h
--- a/Engine/Plugins/FX/Niagara/Source/NiagaraEditor/Public/ViewModels/Stack/NiagaraStackModuleItem.h
+++ b/Engine/Plugins/FX/Niagara/Source/NiagaraEditor/Public/ViewModels/Stack/NiagaraStackModuleItem.h
@@ -21,6 +21,8 @@ class NIAGARAEDITOR_API UNiagaraStackModuleItem : public UNiagaraStackItem
 	GENERATED_BODY()
 
 public:
+	DECLARE_DELEGATE_OneParam(FOnRequestDeprecationRecommended, UNiagaraStackModuleItem*);
+
 	UNiagaraStackModuleItem();
 
 	UNiagaraNodeFunctionCall& GetModuleNode() const;
@@ -89,7 +91,19 @@ public:
 	virtual FText GetDeleteTransactionText() const override;
 	virtual void Delete() override;
 
+	void SetOnRequestDeprecationRecommended(FOnRequestDeprecationRecommended InOnRequest)
+	{
+		DeprecationDelegate = InOnRequest;
+	}
+
+	void SetEnabled(bool bEnabled)
+	{
+		SetIsEnabledInternal(bEnabled);
+	}
+
 protected:
+	FOnRequestDeprecationRecommended DeprecationDelegate;
+
 	virtual void RefreshChildrenInternal(const TArray<UNiagaraStackEntry*>& CurrentChildren, TArray<UNiagaraStackEntry*>& NewChildren, TArray<FStackIssue>& NewIssues) override;
 	virtual void SetIsEnabledInternal(bool bInIsEnabled) override;
 
diff --git a/Engine/Plugins/FX/Niagara/Source/NiagaraEditor/Public/ViewModels/Stack/NiagaraStackScriptItemGroup.h b/Engine/Plugins/FX/Niagara/Source/NiagaraEditor/Public/ViewModels/Stack/NiagaraStackScriptItemGroup.h
--- a/Engine/Plugins/FX/Niagara/Source/NiagaraEditor/Public/ViewModels/Stack/NiagaraStackScriptItemGroup.h
+++ b/Engine/Plugins/FX/Niagara/Source/NiagaraEditor/Public/ViewModels/Stack/NiagaraStackScriptItemGroup.h
@@ -58,6 +58,8 @@ private:
 
 	void ChildRequestPaste(const UNiagaraClipboardContent* ClipboardContent, int32 PasteIndex, FText& OutPasteWarning);
 
+	void ChildRequestDeprecatedRecommendation(class UNiagaraStackModuleItem* TargetChild);
+
 	void OnScriptGraphChanged(const struct FEdGraphEditAction& InAction);
 
 	TOptional<FDropRequestResponse> CanDropOnTarget(const UNiagaraStackEntry& TargetEntry, const FDropRequest& DropRequest);
