Cancel pending compile jobs after Undo

diff --git a/Engine/Plugins/FX/Niagara/Source/Niagara/Classes/NiagaraSystem.h b/Engine/Plugins/FX/Niagara/Source/Niagara/Classes/NiagaraSystem.h
--- a/Engine/Plugins/FX/Niagara/Source/Niagara/Classes/NiagaraSystem.h
+++ b/Engine/Plugins/FX/Niagara/Source/Niagara/Classes/NiagaraSystem.h
@@ -157,6 +157,11 @@ USTRUCT()
 struct FNiagaraSystemCompileRequest
 {
 	GENERATED_USTRUCT_BODY()
+
+	FNiagaraSystemCompileRequest()
+		: bIsValid(true)
+	{
+	}
 		
 	double StartTime;
 
@@ -166,6 +171,8 @@ struct FNiagaraSystemCompileRequest
 	TArray<FEmitterCompiledScriptPair> EmitterCompiledScriptPairs;
 	
 	TMap<UNiagaraScript*, TSharedPtr<FNiagaraCompileRequestDataBase, ESPMode::ThreadSafe> > MappedData;
+
+	bool bIsValid;
 };
 
 /** Container for multiple emitters that combine together to create a particle system effect.*/
@@ -277,6 +284,9 @@ public:
 	/** Blocks until all active compile jobs have finished */
 	void WaitForCompilationComplete();
 
+	/** Invalidates any active compilation requests which will ignore their results. */
+	void InvalidateActiveCompiles();
+
 	/** Delegate called when the system's dependencies have all been compiled.*/
 	FOnSystemCompiled& OnSystemCompiled();
 
diff --git a/Engine/Plugins/FX/Niagara/Source/Niagara/Private/NiagaraSystem.cpp b/Engine/Plugins/FX/Niagara/Source/Niagara/Private/NiagaraSystem.cpp
--- a/Engine/Plugins/FX/Niagara/Source/Niagara/Private/NiagaraSystem.cpp
+++ b/Engine/Plugins/FX/Niagara/Source/Niagara/Private/NiagaraSystem.cpp
@@ -813,6 +813,15 @@ void UNiagaraSystem::WaitForCompilationComplete()
 	{
 		QueryCompileComplete(true, ActiveCompilations.Num() == 1);
 	}
+
+}
+
+void UNiagaraSystem::InvalidateActiveCompiles()
+{
+	for (FNiagaraSystemCompileRequest& ActiveCompilation : ActiveCompilations)
+	{
+		ActiveCompilation.bIsValid = false;
+	}
 }
 
 bool UNiagaraSystem::PollForCompilationComplete()
@@ -855,7 +864,7 @@ bool UNiagaraSystem::QueryCompileComplete(bool bWait, bool bDoPost, bool bDoNotA
 		}
 
 		// In the world of do not apply, we're exiting the system completely so let's just kill any active compilations altogether.
-		if (bDoNotApply)
+		if (bDoNotApply || ActiveCompilations[ActiveCompileIdx].bIsValid == false)
 		{
 			ActiveCompilations[ActiveCompileIdx].RootObjects.Empty();
 			ActiveCompilations.RemoveAt(ActiveCompileIdx);
diff --git a/Engine/Plugins/FX/Niagara/Source/NiagaraEditor/Private/ViewModels/NiagaraSystemViewModel.cpp b/Engine/Plugins/FX/Niagara/Source/NiagaraEditor/Private/ViewModels/NiagaraSystemViewModel.cpp
--- a/Engine/Plugins/FX/Niagara/Source/NiagaraEditor/Private/ViewModels/NiagaraSystemViewModel.cpp
+++ b/Engine/Plugins/FX/Niagara/Source/NiagaraEditor/Private/ViewModels/NiagaraSystemViewModel.cpp
@@ -510,6 +510,8 @@ void FNiagaraSystemViewModel::PostUndo(bool bSuccess)
 	ResetEmitterHandleViewModelsAndTracks();
 	RefreshAll();
 	GetDefault<UEdGraphSchema_NiagaraSystemOverview>()->ForceVisualizationCacheClear();
+	System->InvalidateActiveCompiles();
+	System->RequestCompile(false);
 }
 
 void FNiagaraSystemViewModel::Tick(float DeltaTime)
@@ -1233,7 +1235,6 @@ void FNiagaraSystemViewModel::ResetSystem(ETimeResetMode TimeResetMode, EMultiRe
 				Component->SetDesiredAge(0);
 			}
 		}
-
 		else if (ReinitMode == EReinitMode::ReinitializeSystem)
 		{
 			Component->ReinitializeSystem();
