Make an error when using curves without LUTs on GPU sims.

diff --git a/Engine/Plugins/FX/Niagara/Source/Niagara/Classes/NiagaraDataInterfaceCurveBase.h b/Engine/Plugins/FX/Niagara/Source/Niagara/Classes/NiagaraDataInterfaceCurveBase.h
--- a/Engine/Plugins/FX/Niagara/Source/Niagara/Classes/NiagaraDataInterfaceCurveBase.h
+++ b/Engine/Plugins/FX/Niagara/Source/Niagara/Classes/NiagaraDataInterfaceCurveBase.h
@@ -43,20 +43,20 @@ public:
 protected:
 	GENERATED_BODY()
 
-	UPROPERTY()
-	TArray<float> ShaderLUT;
+		UPROPERTY()
+		TArray<float> ShaderLUT;
 
 	UPROPERTY()
-	float LUTMinTime;
+		float LUTMinTime;
 
 	UPROPERTY()
-	float LUTMaxTime;
+		float LUTMaxTime;
 
 	UPROPERTY()
-	float LUTInvTimeRange;
+		float LUTInvTimeRange;
 
 	UPROPERTY()
-	float LUTNumSamplesMinusOne;
+		float LUTNumSamplesMinusOne;
 
 	/** Remap a sample time for this curve to 0 to 1 between first and last keys for LUT access.*/
 	FORCEINLINE float NormalizeTime(float T) const
@@ -74,24 +74,30 @@ public:
 	DECLARE_NIAGARA_DI_PARAMETER();
 
 	UPROPERTY(EditAnywhere, AdvancedDisplay, Category = "Curve")
-	uint32 bUseLUT : 1;
+		uint32 bUseLUT : 1;
 
 #if WITH_EDITORONLY_DATA
 	/** Do we optimize the LUT, this saves memory but may introduce errors.  Errors can be reduced modifying the threshold. */
 	UPROPERTY(EditAnywhere, AdvancedDisplay, Category = "Curve")
-	uint32 bOptimizeLUT : 1;
+		uint32 bOptimizeLUT : 1;
 
 	UPROPERTY(EditAnywhere, AdvancedDisplay, Category = "Curve")
-	uint32 bOverrideOptimizeThreshold : 1;
+		uint32 bOverrideOptimizeThreshold : 1;
 
 	/** Threshold used to optimize the LUT. */
-	UPROPERTY(EditAnywhere, AdvancedDisplay, Category = "Curve", meta=(EditCondition="bOverrideOptimizeThreshold"))
-	float OptimizeThreshold;
+	UPROPERTY(EditAnywhere, AdvancedDisplay, Category = "Curve", meta = (EditCondition = "bOverrideOptimizeThreshold"))
+		float OptimizeThreshold;
 
 	UPROPERTY(EditAnywhere, Transient, Category = "Curve")
-	bool ShowInCurveEditor;
+		bool ShowInCurveEditor;
 #endif
 
+#if WITH_EDITOR	
+	/** Refreshes and returns the errors detected with the corresponding data, if any.*/
+	virtual TArray<FNiagaraDataInterfaceError> GetErrors() override;
+#endif
+
+
 public:
 	UNiagaraDataInterfaceCurveBase()
 		: LUTMinTime(0.0f)
diff --git a/Engine/Plugins/FX/Niagara/Source/Niagara/Private/NiagaraDataInterfaceCurveBase.cpp b/Engine/Plugins/FX/Niagara/Source/Niagara/Private/NiagaraDataInterfaceCurveBase.cpp
--- a/Engine/Plugins/FX/Niagara/Source/Niagara/Private/NiagaraDataInterfaceCurveBase.cpp
+++ b/Engine/Plugins/FX/Niagara/Source/Niagara/Private/NiagaraDataInterfaceCurveBase.cpp
@@ -8,6 +8,9 @@
 #include "ShaderParameterUtils.h"
 #include "NiagaraShader.h"
 #include "NiagaraCustomVersion.h"
+#include "NiagaraEmitter.h"
+#include "Internationalization/Internationalization.h"
+#define LOCTEXT_NAMESPACE "NiagaraDataInterfaceCurveBase"
 
 IMPLEMENT_TYPE_LAYOUT(FNiagaraDataInterfaceParametersCS_Curve);
 
@@ -337,3 +340,37 @@ void UNiagaraDataInterfaceCurveBase::PushToRenderThread()
 		RHIUnlockVertexBuffer(RT_Proxy->CurveLUT.Buffer);
 	});
 }
+
+
+#if WITH_EDITOR	
+/** Refreshes and returns the errors detected with the corresponding data, if any.*/
+TArray<FNiagaraDataInterfaceError> UNiagaraDataInterfaceCurveBase::GetErrors()
+{
+	// Trace down the root emitter (if there is one)
+	TArray<FNiagaraDataInterfaceError> Errors;
+	UObject* Obj = GetOuter();
+	UNiagaraEmitter* Emitter = nullptr;
+	while (Obj && !Obj->IsA(UPackage::StaticClass()))
+	{
+		Emitter = Cast<UNiagaraEmitter>(Obj);
+		if (!Emitter)
+			Obj = Obj->GetOuter();
+		else
+			break;
+	}
+
+	// If there is a root emitter, assume that we are in its particle stack and point out that we need bUseLUT true for GPU sims.
+	if (Emitter && Emitter->SimTarget == ENiagaraSimTarget::GPUComputeSim && !bUseLUT)
+	{
+		FNiagaraDataInterfaceError LUTsNeededForGPUSimsError(LOCTEXT("LUTsNeededForGPUSims", "This Data Interface must have bUseLUT set to true for GPU sims."),
+			LOCTEXT("LUTsNeededForGPUSimsSummary", "bUseLUT Required"),
+			FNiagaraDataInterfaceFix());
+
+		Errors.Add(LUTsNeededForGPUSimsError);
+	}
+	return Errors;
+}
+#endif
+
+
+#undef LOCTEXT_NAMESPACE
\ No newline at end of file
