- Null skinning data in move constructor
- Add reader / writer lock to ensure we do not modify skinning data while VM is using the data

#rnx

diff --git a/Engine/Plugins/FX/Niagara/Source/Niagara/Classes/NDISkeletalMeshCommon.h b/Engine/Plugins/FX/Niagara/Source/Niagara/Classes/NDISkeletalMeshCommon.h
--- a/Engine/Plugins/FX/Niagara/Source/Niagara/Classes/NDISkeletalMeshCommon.h
+++ b/Engine/Plugins/FX/Niagara/Source/Niagara/Classes/NDISkeletalMeshCommon.h
@@ -6,6 +6,22 @@
 
 struct FSkeletalMeshAccessorHelper
 {
+	FSkeletalMeshAccessorHelper()
+	{
+	}
+	~FSkeletalMeshAccessorHelper()
+	{
+		if (SkinningData != nullptr)
+		{
+			SkinningData->ExitRead();
+		}
+	}
+
+	FSkeletalMeshAccessorHelper(const FSkeletalMeshAccessorHelper&) = delete;
+	FSkeletalMeshAccessorHelper(FSkeletalMeshAccessorHelper&&) = delete;
+	FSkeletalMeshAccessorHelper operator=(const FSkeletalMeshAccessorHelper&) = delete;
+	FSkeletalMeshAccessorHelper operator=(FSkeletalMeshAccessorHelper&&) = delete;
+
 	template<typename FilterMode, typename AreaWeightingMode>
 	FORCEINLINE void Init(FNDISkeletalMesh_InstanceData* InstData)
 	{
@@ -15,6 +31,11 @@ struct FSkeletalMeshAccessorHelper
 		IndexBuffer = LODData->MultiSizeIndexContainer.GetIndexBuffer();
 		SkinningData = InstData->SkinningData.SkinningData.Get();
 		Usage = InstData->SkinningData.Usage;
+
+		if (SkinningData != nullptr)
+		{
+			SkinningData->EnterRead();
+		}
 	}
 
 	USkeletalMeshComponent* Comp = nullptr;
diff --git a/Engine/Plugins/FX/Niagara/Source/Niagara/Classes/NiagaraDataInterfaceSkeletalMesh.h b/Engine/Plugins/FX/Niagara/Source/Niagara/Classes/NiagaraDataInterfaceSkeletalMesh.h
--- a/Engine/Plugins/FX/Niagara/Source/Niagara/Classes/NiagaraDataInterfaceSkeletalMesh.h
+++ b/Engine/Plugins/FX/Niagara/Source/Niagara/Classes/NiagaraDataInterfaceSkeletalMesh.h
@@ -53,39 +53,15 @@ struct FSkeletalMeshSkinningDataHandle
 {
 	FSkeletalMeshSkinningDataHandle();
 	FSkeletalMeshSkinningDataHandle(FSkeletalMeshSkinningDataUsage InUsage, TSharedPtr<struct FSkeletalMeshSkinningData> InSkinningData);
-
-	FSkeletalMeshSkinningDataHandle& operator=(FSkeletalMeshSkinningDataHandle&& Other)
-	{
-		if (this != &Other)
-		{
-			Usage = Other.Usage;
-			SkinningData = Other.SkinningData;
-			Other.SkinningData = nullptr;
-		}
-		return *this;
-	}
-
-	FSkeletalMeshSkinningDataHandle(FSkeletalMeshSkinningDataHandle&& Other)
-	{
-		Usage = Other.Usage;
-		SkinningData = Other.SkinningData;
-	}
-
+	FSkeletalMeshSkinningDataHandle(const FSkeletalMeshSkinningDataHandle& Other) = delete;
+	FSkeletalMeshSkinningDataHandle(FSkeletalMeshSkinningDataHandle&& Other);
 	~FSkeletalMeshSkinningDataHandle();
 
+	FSkeletalMeshSkinningDataHandle& operator=(const FSkeletalMeshSkinningDataHandle& Other) = delete;
+	FSkeletalMeshSkinningDataHandle& operator=(FSkeletalMeshSkinningDataHandle&& Other);
+
 	FSkeletalMeshSkinningDataUsage Usage;
 	TSharedPtr<FSkeletalMeshSkinningData> SkinningData;
-
-private:
-
-	FSkeletalMeshSkinningDataHandle& operator=(FSkeletalMeshSkinningDataHandle& Other)
-	{
-		return *this;
-	}
-
-	FSkeletalMeshSkinningDataHandle(FSkeletalMeshSkinningDataHandle& Other)
-	{
-	}
 };
 
 struct FSkeletalMeshSkinningData
@@ -105,6 +81,16 @@ struct FSkeletalMeshSkinningData
 
 	bool Tick(float InDeltaSeconds, bool bRequirePreskin = true);
 
+	FORCEINLINE void EnterRead()
+	{
+		RWGuard.ReadLock();
+	}
+
+	FORCEINLINE void ExitRead()
+	{
+		RWGuard.ReadUnlock();
+	}
+
 	FORCEINLINE int32 GetBoneCount(bool RequiresPrevious) const
 	{
 		int32 BoneCount = FMath::Min(CurrBoneRefToLocals().Num(), CurrComponentTransforms().Num());
@@ -181,7 +167,7 @@ private:
 
 	void UpdateBoneTransforms();
 
-	FCriticalSection CriticalSection; 
+	FRWLock RWGuard;
 	
 	TWeakObjectPtr<USkeletalMeshComponent> MeshComp;
 
diff --git a/Engine/Plugins/FX/Niagara/Source/Niagara/Private/NiagaraDataInterfaceSkeletalMesh.cpp b/Engine/Plugins/FX/Niagara/Source/Niagara/Private/NiagaraDataInterfaceSkeletalMesh.cpp
--- a/Engine/Plugins/FX/Niagara/Source/Niagara/Private/NiagaraDataInterfaceSkeletalMesh.cpp
+++ b/Engine/Plugins/FX/Niagara/Source/Niagara/Private/NiagaraDataInterfaceSkeletalMesh.cpp
@@ -85,16 +85,34 @@ FSkeletalMeshSkinningDataHandle::~FSkeletalMeshSkinningDataHandle()
 	}
 }
 
+FSkeletalMeshSkinningDataHandle::FSkeletalMeshSkinningDataHandle(FSkeletalMeshSkinningDataHandle&& Other)
+{
+	Usage = Other.Usage;
+	SkinningData = Other.SkinningData;
+	Other.SkinningData = nullptr;
+}
+
+FSkeletalMeshSkinningDataHandle& FSkeletalMeshSkinningDataHandle::operator=(FSkeletalMeshSkinningDataHandle&& Other)
+{
+	if (this != &Other)
+	{
+		Usage = Other.Usage;
+		SkinningData = Other.SkinningData;
+		Other.SkinningData = nullptr;
+	}
+	return *this;
+}
+
 //////////////////////////////////////////////////////////////////////////
 void FSkeletalMeshSkinningData::ForceDataRefresh()
 {
-	FScopeLock Lock(&CriticalSection);
+	FRWScopeLock Lock(RWGuard, SLT_Write);
 	bForceDataRefresh = true;
 }
 
 void FSkeletalMeshSkinningData::RegisterUser(FSkeletalMeshSkinningDataUsage Usage)
 {
-	FScopeLock Lock(&CriticalSection);
+	FRWScopeLock Lock(RWGuard, SLT_Write);
 	USkeletalMeshComponent* SkelComp = MeshComp.Get();
 
 	int32 LODIndex = Usage.GetLODIndex();
@@ -150,7 +168,7 @@ void FSkeletalMeshSkinningData::RegisterUser(FSkeletalMeshSkinningDataUsage Usag
 
 void FSkeletalMeshSkinningData::UnregisterUser(FSkeletalMeshSkinningDataUsage Usage)
 {
-	FScopeLock Lock(&CriticalSection);
+	FRWScopeLock Lock(RWGuard, SLT_Write);
 	check(LODData.IsValidIndex(Usage.GetLODIndex()));
 
 	if (Usage.NeedBoneMatrices())
@@ -165,7 +183,7 @@ void FSkeletalMeshSkinningData::UnregisterUser(FSkeletalMeshSkinningDataUsage Us
 	}
 }
 
-bool FSkeletalMeshSkinningData::IsUsed()const
+bool FSkeletalMeshSkinningData::IsUsed() const
 {
 	if (BoneMatrixUsers > 0)
 	{
@@ -256,6 +274,8 @@ void FSkeletalMeshSkinningData::UpdateBoneTransforms()
 
 bool FSkeletalMeshSkinningData::Tick(float InDeltaSeconds, bool bRequirePreskin)
 {
+	FRWScopeLock Lock(RWGuard, SLT_Write);
+
 	USkeletalMeshComponent* SkelComp = MeshComp.Get();
 	check(SkelComp);
 	DeltaSeconds = InDeltaSeconds;
@@ -314,6 +334,11 @@ FSkeletalMeshSkinningDataHandle FNDI_SkeletalMesh_GeneratedData::GetCachedSkinni
 		if ( CachedSkinningDataAndUsage* Existing = CachedSkinningData.Find(Component) )
 		{
 			check(Existing->SkinningData.IsValid());
+			ensure(Existing->Usage.NeedBoneMatrices() == Usage.NeedBoneMatrices());
+			ensure(Existing->Usage.NeedPreSkinnedVerts() == Usage.NeedPreSkinnedVerts());
+			ensure(Existing->Usage.NeedsDataImmediately() == Usage.NeedsDataImmediately());
+			ensure(Existing->Usage.GetLODIndex() == Usage.GetLODIndex());
+
 			return FSkeletalMeshSkinningDataHandle(Existing->Usage, Existing->SkinningData);
 		}
 	}
@@ -2351,6 +2376,11 @@ void FSkeletalMeshAccessorHelper::Init<
 	const FSkeletalMeshSamplingInfo& SamplingInfo = InstData->Mesh->GetSamplingInfo();
 	SamplingRegion = &SamplingInfo.GetRegion(InstData->SamplingRegionIndices[0]);
 	SamplingRegionBuiltData = &SamplingInfo.GetRegionBuiltData(InstData->SamplingRegionIndices[0]);
+
+	if (SkinningData != nullptr)
+	{
+		SkinningData->EnterRead();
+	}
 }
 
 template<>
@@ -2369,6 +2399,11 @@ void FSkeletalMeshAccessorHelper::Init<
 	const FSkeletalMeshSamplingInfo& SamplingInfo = InstData->Mesh->GetSamplingInfo();
 	SamplingRegion = &SamplingInfo.GetRegion(InstData->SamplingRegionIndices[0]);
 	SamplingRegionBuiltData = &SamplingInfo.GetRegionBuiltData(InstData->SamplingRegionIndices[0]);
+
+	if (SkinningData != nullptr)
+	{
+		SkinningData->EnterRead();
+	}
 }
 
 #undef LOCTEXT_NAMESPACE
\ No newline at end of file
