Adding ability for renderers to provide warnings/cues to end users.

diff --git a/Engine/Plugins/FX/Niagara/Source/Niagara/Private/NiagaraLightRendererProperties.cpp b/Engine/Plugins/FX/Niagara/Source/Niagara/Private/NiagaraLightRendererProperties.cpp
--- a/Engine/Plugins/FX/Niagara/Source/Niagara/Private/NiagaraLightRendererProperties.cpp
+++ b/Engine/Plugins/FX/Niagara/Source/Niagara/Private/NiagaraLightRendererProperties.cpp
@@ -126,6 +126,11 @@ void UNiagaraLightRendererProperties::GetRendererTooltipWidgets(const FNiagaraEm
 	OutWidgets.Add(LightTooltip);
 }
 
+void UNiagaraLightRendererProperties::GetRendererFeedback(const UNiagaraEmitter* InEmitter, TArray<FText>& OutErrors, TArray<FText>& OutWarnings, TArray<FText>& OutInfo) const
+{
+	Super::GetRendererFeedback(InEmitter, OutErrors, OutWarnings, OutInfo);
+}
+
 bool UNiagaraLightRendererProperties::IsMaterialValidForRenderer(UMaterial* Material, FText& InvalidMessage)
 {
 	return true;
diff --git a/Engine/Plugins/FX/Niagara/Source/Niagara/Private/NiagaraMeshRendererProperties.cpp b/Engine/Plugins/FX/Niagara/Source/Niagara/Private/NiagaraMeshRendererProperties.cpp
--- a/Engine/Plugins/FX/Niagara/Source/Niagara/Private/NiagaraMeshRendererProperties.cpp
+++ b/Engine/Plugins/FX/Niagara/Source/Niagara/Private/NiagaraMeshRendererProperties.cpp
@@ -366,6 +366,12 @@ void UNiagaraMeshRendererProperties::GetRendererTooltipWidgets(const FNiagaraEmi
 	}
 }
 
+
+void UNiagaraMeshRendererProperties::GetRendererFeedback(const UNiagaraEmitter* InEmitter, TArray<FText>& OutErrors, TArray<FText>& OutWarnings, TArray<FText>& OutInfo) const
+{
+	Super::GetRendererFeedback(InEmitter, OutErrors, OutWarnings, OutInfo);
+}
+
 void UNiagaraMeshRendererProperties::BeginDestroy()
 {
 	Super::BeginDestroy();
diff --git a/Engine/Plugins/FX/Niagara/Source/Niagara/Private/NiagaraRibbonRendererProperties.cpp b/Engine/Plugins/FX/Niagara/Source/Niagara/Private/NiagaraRibbonRendererProperties.cpp
--- a/Engine/Plugins/FX/Niagara/Source/Niagara/Private/NiagaraRibbonRendererProperties.cpp
+++ b/Engine/Plugins/FX/Niagara/Source/Niagara/Private/NiagaraRibbonRendererProperties.cpp
@@ -221,6 +221,12 @@ void UNiagaraRibbonRendererProperties::GetRendererTooltipWidgets(const FNiagaraE
 }
 
 
+void UNiagaraRibbonRendererProperties::GetRendererFeedback(const UNiagaraEmitter* InEmitter, TArray<FText>& OutErrors, TArray<FText>& OutWarnings, TArray<FText>& OutInfo) const
+{
+	Super::GetRendererFeedback(InEmitter, OutErrors, OutWarnings, OutInfo);
+}
+
+
 bool UNiagaraRibbonRendererProperties::IsMaterialValidForRenderer(UMaterial* InMaterial, FText& InvalidMessage)
 {
 	if (InMaterial->bUsedWithNiagaraRibbons == false)
diff --git a/Engine/Plugins/FX/Niagara/Source/Niagara/Private/NiagaraSpriteRendererProperties.cpp b/Engine/Plugins/FX/Niagara/Source/Niagara/Private/NiagaraSpriteRendererProperties.cpp
--- a/Engine/Plugins/FX/Niagara/Source/Niagara/Private/NiagaraSpriteRendererProperties.cpp
+++ b/Engine/Plugins/FX/Niagara/Source/Niagara/Private/NiagaraSpriteRendererProperties.cpp
@@ -213,6 +213,7 @@ void UNiagaraSpriteRendererProperties::PostEditChangeProperty(struct FPropertyCh
 	}
 	
 	Super::PostEditChangeProperty(PropertyChangedEvent);
+
 }
 
 const TArray<FNiagaraVariable>& UNiagaraSpriteRendererProperties::GetBoundAttributes()
@@ -319,6 +320,22 @@ void UNiagaraSpriteRendererProperties::GetRendererWidgets(const FNiagaraEmitterI
 	}
 }
 
+
+void UNiagaraSpriteRendererProperties::GetRendererFeedback(const UNiagaraEmitter* InEmitter, TArray<FText>& OutErrors, TArray<FText>& OutWarnings, TArray<FText>& OutInfo) const
+{
+	Super::GetRendererFeedback(InEmitter, OutErrors, OutWarnings, OutInfo);
+	if (InEmitter->SpawnScriptProps.Script->GetVMExecutableData().IsValid())
+	{
+		if (bUseMaterialCutoutTexture || CutoutTexture)
+		{
+			if (InEmitter->SpawnScriptProps.Script->GetVMExecutableData().Attributes.Contains(UVScaleBinding.DataSetVariable))
+			{
+				OutInfo.Add(LOCTEXT("SpriteRendererUVScaleWithCutout", "Cutouts will not be sized dynamically with UVScale variable. If scaling above 1.0, geometry may clip."));
+			}			
+		}
+	}
+}
+
 void UNiagaraSpriteRendererProperties::GetRendererTooltipWidgets(const FNiagaraEmitterInstance* InEmitter, TArray<TSharedPtr<SWidget>>& OutWidgets, TSharedPtr<FAssetThumbnailPool> InThumbnailPool) const
 {
 	TArray<UMaterialInterface*> Materials;
diff --git a/Engine/Plugins/FX/Niagara/Source/Niagara/Public/NiagaraLightRendererProperties.h b/Engine/Plugins/FX/Niagara/Source/Niagara/Public/NiagaraLightRendererProperties.h
--- a/Engine/Plugins/FX/Niagara/Source/Niagara/Public/NiagaraLightRendererProperties.h
+++ b/Engine/Plugins/FX/Niagara/Source/Niagara/Public/NiagaraLightRendererProperties.h
@@ -38,6 +38,7 @@ public:
 	virtual const TArray<FNiagaraVariable>& GetOptionalAttributes() override;
 	virtual void GetRendererWidgets(const FNiagaraEmitterInstance* InEmitter, TArray<TSharedPtr<SWidget>>& OutWidgets, TSharedPtr<FAssetThumbnailPool> InThumbnailPool) const override;
 	virtual void GetRendererTooltipWidgets(const FNiagaraEmitterInstance* InEmitter, TArray<TSharedPtr<SWidget>>& OutWidgets, TSharedPtr<FAssetThumbnailPool> InThumbnailPool) const override;
+	virtual void GetRendererFeedback(const UNiagaraEmitter* InEmitter, TArray<FText>& OutErrors, TArray<FText>& OutWarnings, TArray<FText>& OutInfo) const override;
 #endif // WITH_EDITORONLY_DATA
 
 	/** Whether to use physically based inverse squared falloff from the light.  If unchecked, the value from the LightExponent binding will be used instead. */
diff --git a/Engine/Plugins/FX/Niagara/Source/Niagara/Public/NiagaraMeshRendererProperties.h b/Engine/Plugins/FX/Niagara/Source/Niagara/Public/NiagaraMeshRendererProperties.h
--- a/Engine/Plugins/FX/Niagara/Source/Niagara/Public/NiagaraMeshRendererProperties.h
+++ b/Engine/Plugins/FX/Niagara/Source/Niagara/Public/NiagaraMeshRendererProperties.h
@@ -100,6 +100,7 @@ public:
 	virtual const TArray<FNiagaraVariable>& GetOptionalAttributes() override;
 	virtual	void GetRendererWidgets(const FNiagaraEmitterInstance* InEmitter, TArray<TSharedPtr<SWidget>>& OutWidgets, TSharedPtr<FAssetThumbnailPool> InThumbnailPool) const override;
 	virtual	void GetRendererTooltipWidgets(const FNiagaraEmitterInstance* InEmitter, TArray<TSharedPtr<SWidget>>& OutWidgets, TSharedPtr<FAssetThumbnailPool> InThumbnailPool) const override;
+	virtual void GetRendererFeedback(const UNiagaraEmitter* InEmitter, TArray<FText>& OutErrors, TArray<FText>& OutWarnings, TArray<FText>& OutInfo) const override;
 	void OnMeshChanged();
 	void CheckMaterialUsage();
 #endif // WITH_EDITORONLY_DATA
diff --git a/Engine/Plugins/FX/Niagara/Source/Niagara/Public/NiagaraRendererProperties.h b/Engine/Plugins/FX/Niagara/Source/Niagara/Public/NiagaraRendererProperties.h
--- a/Engine/Plugins/FX/Niagara/Source/Niagara/Public/NiagaraRendererProperties.h
+++ b/Engine/Plugins/FX/Niagara/Source/Niagara/Public/NiagaraRendererProperties.h
@@ -46,6 +46,7 @@ public:
 	virtual bool IsSimTargetSupported(ENiagaraSimTarget InSimTarget) const { return false; };
 
 #if WITH_EDITORONLY_DATA
+
 	virtual bool IsMaterialValidForRenderer(UMaterial* Material, FText& InvalidMessage) { return true; }
 
 	virtual void FixMaterial(UMaterial* Material) { }
@@ -61,6 +62,8 @@ public:
 
 	virtual void GetRendererWidgets(const FNiagaraEmitterInstance* InEmitter, TArray<TSharedPtr<SWidget>>& OutWidgets, TSharedPtr<FAssetThumbnailPool> InThumbnailPool) const PURE_VIRTUAL(UNiagaraRendererProperties::GetRendererWidgets, );
 	virtual void GetRendererTooltipWidgets(const FNiagaraEmitterInstance* InEmitter, TArray<TSharedPtr<SWidget>>& OutWidgets, TSharedPtr<FAssetThumbnailPool> InThumbnailPool) const PURE_VIRTUAL(UNiagaraRendererProperties::GetRendererTooltipWidgets, );
+	virtual void GetRendererFeedback(const UNiagaraEmitter* InEmitter, TArray<FText>& OutErrors, TArray<FText>& OutWarnings, TArray<FText>& OutInfo) const {};
+
 #endif // WITH_EDITORONLY_DATA
 
 
@@ -85,4 +88,5 @@ public:
 protected:
 	// Copy of variables in the attribute binding, updated when GetRequiredAttributes() is called.
 	TArray<FNiagaraVariable> CurrentAttributeBindings;
+
 };
diff --git a/Engine/Plugins/FX/Niagara/Source/Niagara/Public/NiagaraRibbonRendererProperties.h b/Engine/Plugins/FX/Niagara/Source/Niagara/Public/NiagaraRibbonRendererProperties.h
--- a/Engine/Plugins/FX/Niagara/Source/Niagara/Public/NiagaraRibbonRendererProperties.h
+++ b/Engine/Plugins/FX/Niagara/Source/Niagara/Public/NiagaraRibbonRendererProperties.h
@@ -88,6 +88,8 @@ public:
 	virtual const TArray<FNiagaraVariable>& GetOptionalAttributes() override;
 	virtual void GetRendererWidgets(const FNiagaraEmitterInstance* InEmitter, TArray<TSharedPtr<SWidget>>& OutWidgets, TSharedPtr<FAssetThumbnailPool> InThumbnailPool) const override;
 	virtual void GetRendererTooltipWidgets(const FNiagaraEmitterInstance* InEmitter, TArray<TSharedPtr<SWidget>>& OutWidgets, TSharedPtr<FAssetThumbnailPool> InThumbnailPool) const override;
+	virtual void GetRendererFeedback(const UNiagaraEmitter* InEmitter, TArray<FText>& OutErrors, TArray<FText>& OutWarnings, TArray<FText>& OutInfo) const override;
+
 #endif
 	//UNiagaraRendererProperties Interface END
 
diff --git a/Engine/Plugins/FX/Niagara/Source/Niagara/Public/NiagaraSpriteRendererProperties.h b/Engine/Plugins/FX/Niagara/Source/Niagara/Public/NiagaraSpriteRendererProperties.h
--- a/Engine/Plugins/FX/Niagara/Source/Niagara/Public/NiagaraSpriteRendererProperties.h
+++ b/Engine/Plugins/FX/Niagara/Source/Niagara/Public/NiagaraSpriteRendererProperties.h
@@ -73,6 +73,7 @@ public:
 	virtual const TArray<FNiagaraVariable>& GetOptionalAttributes() override;
 	virtual void GetRendererWidgets(const FNiagaraEmitterInstance* InEmitter, TArray<TSharedPtr<SWidget>>& OutWidgets, TSharedPtr<FAssetThumbnailPool> InThumbnailPool) const override;
 	virtual void GetRendererTooltipWidgets(const FNiagaraEmitterInstance* InEmitter, TArray<TSharedPtr<SWidget>>& OutWidgets, TSharedPtr<FAssetThumbnailPool> InThumbnailPool) const override;
+	virtual void GetRendererFeedback(const UNiagaraEmitter* InEmitter, TArray<FText>& OutErrors, TArray<FText>& OutWarnings, TArray<FText>& OutInfo) const override;
 #endif
 	//UNiagaraMaterialRendererProperties interface END
 
@@ -119,11 +120,11 @@ public:
 	UPROPERTY(EditAnywhere, Category = "Sorting")
 	uint32 bSortOnlyWhenTranslucent : 1;
 
-	/** The distance at which FacingCameraDistanceBlend	is fully FacingCamera */
+	/** When FacingMode is FacingCameraDistanceBlend, the distance at which the sprite is fully facing the camera plane. */
 	UPROPERTY(EditAnywhere, Category = "Sprite Rendering", meta = (UIMin = "0"))
 	float MinFacingCameraBlendDistance;
 
-	/** The distance at which FacingCameraDistanceBlend	is fully FacingCameraPosition */
+	/** When FacingMode is FacingCameraDistanceBlend, the distance at which the sprite is fully facing the camera position */
 	UPROPERTY(EditAnywhere, Category = "Sprite Rendering", meta = (UIMin = "0"))
 	float MaxFacingCameraBlendDistance;
 
diff --git a/Engine/Plugins/FX/Niagara/Source/NiagaraEditor/Private/ViewModels/Stack/NiagaraStackEntry.cpp b/Engine/Plugins/FX/Niagara/Source/NiagaraEditor/Private/ViewModels/Stack/NiagaraStackEntry.cpp
--- a/Engine/Plugins/FX/Niagara/Source/NiagaraEditor/Private/ViewModels/Stack/NiagaraStackEntry.cpp
+++ b/Engine/Plugins/FX/Niagara/Source/NiagaraEditor/Private/ViewModels/Stack/NiagaraStackEntry.cpp
@@ -74,10 +74,10 @@ UNiagaraStackEntry::FStackIssue::FStackIssue(EStackIssueSeverity InSeverity, FTe
 {
 	checkf(ShortDescription.IsEmptyOrWhitespace() == false, TEXT("Short description can not be empty."));
 	//checkf(LongDescription.IsEmptyOrWhitespace() == false, TEXT("Long description can not be empty."));
-	if (LongDescription.IsEmptyOrWhitespace())
-	{
-		LongDescription = ShortDescription;
-	}
+	//if (LongDescription.IsEmptyOrWhitespace())
+	//{
+	//	LongDescription = ShortDescription;
+	//}
 	checkf(InStackEditorDataKey.IsEmpty() == false, TEXT("Stack editor data key can not be empty."));
 }
 
diff --git a/Engine/Plugins/FX/Niagara/Source/NiagaraEditor/Private/ViewModels/Stack/NiagaraStackErrorItem.cpp b/Engine/Plugins/FX/Niagara/Source/NiagaraEditor/Private/ViewModels/Stack/NiagaraStackErrorItem.cpp
--- a/Engine/Plugins/FX/Niagara/Source/NiagaraEditor/Private/ViewModels/Stack/NiagaraStackErrorItem.cpp
+++ b/Engine/Plugins/FX/Niagara/Source/NiagaraEditor/Private/ViewModels/Stack/NiagaraStackErrorItem.cpp
@@ -51,15 +51,21 @@ void UNiagaraStackErrorItem::RefreshChildrenInternal(const TArray<UNiagaraStackE
 			FixChild->OnIssueFixed().RemoveAll(this);
 		}
 	}
-	// long description
-	UNiagaraStackErrorItemLongDescription* ErrorEntryLongDescription = FindCurrentChildOfTypeByPredicate<UNiagaraStackErrorItemLongDescription>(CurrentChildren,
-		[&](UNiagaraStackErrorItemLongDescription* CurrentChild) { return true; });
-	if (ErrorEntryLongDescription == nullptr)
+	
+	if (!StackIssue.GetLongDescription().IsEmptyOrWhitespace())
 	{
-		ErrorEntryLongDescription = NewObject<UNiagaraStackErrorItemLongDescription>(this);
-		ErrorEntryLongDescription->Initialize(CreateDefaultChildRequiredData(), StackIssue, GetStackEditorDataKey());
+		// long description
+		UNiagaraStackErrorItemLongDescription* ErrorEntryLongDescription = FindCurrentChildOfTypeByPredicate<UNiagaraStackErrorItemLongDescription>(CurrentChildren,
+			[&](UNiagaraStackErrorItemLongDescription* CurrentChild) { return true; });
+		if (ErrorEntryLongDescription == nullptr)
+		{
+			ErrorEntryLongDescription = NewObject<UNiagaraStackErrorItemLongDescription>(this);
+			ErrorEntryLongDescription->Initialize(CreateDefaultChildRequiredData(), StackIssue, GetStackEditorDataKey());
+		}
+
+
+		NewChildren.Add(ErrorEntryLongDescription);
 	}
-	NewChildren.Add(ErrorEntryLongDescription);
 
 	// fixes
 	for (int i = 0; i < StackIssue.GetFixes().Num(); i++)
diff --git a/Engine/Plugins/FX/Niagara/Source/NiagaraEditor/Private/ViewModels/Stack/NiagaraStackRendererItem.cpp b/Engine/Plugins/FX/Niagara/Source/NiagaraEditor/Private/ViewModels/Stack/NiagaraStackRendererItem.cpp
--- a/Engine/Plugins/FX/Niagara/Source/NiagaraEditor/Private/ViewModels/Stack/NiagaraStackRendererItem.cpp
+++ b/Engine/Plugins/FX/Niagara/Source/NiagaraEditor/Private/ViewModels/Stack/NiagaraStackRendererItem.cpp
@@ -384,11 +384,57 @@ void UNiagaraStackRendererItem::RefreshIssues(TArray<FStackIssue>& NewIssues)
 		NewIssues.Add(TargetSupportError);
 	}
 
+	if (RendererProperties->GetIsEnabled())
+	{
+		TArray<FText> Errors;
+		TArray<FText> Warnings;
+		TArray<FText> Infos;
+
+		RendererProperties->GetRendererFeedback(GetEmitterViewModel()->GetEmitter(), Errors, Warnings, Infos);
+
+		for (const FText& Item : Errors)
+		{
+			FStackIssue TargetSupportError(
+				EStackIssueSeverity::Error,
+				Item,
+				FText(),
+				GetStackEditorDataKey(),
+				false);
+
+			NewIssues.Add(TargetSupportError);
+		}
+
+		for (const FText& Item : Warnings)
+		{
+			FStackIssue TargetSupportError(
+				EStackIssueSeverity::Warning,
+				Item,
+				FText(),
+				GetStackEditorDataKey(),
+				false);
+
+			NewIssues.Add(TargetSupportError);
+		}
+
+		for (const FText& Item : Infos)
+		{
+			FStackIssue TargetSupportError(
+				EStackIssueSeverity::Info,
+				Item,
+				FText(),
+				GetStackEditorDataKey(),
+				false);
+
+			NewIssues.Add(TargetSupportError);
+		}
+	}
+
 }
 
 void UNiagaraStackRendererItem::RendererChanged()
 {
 	bCanResetToBaseCache.Reset();
+	RefreshChildren();
 }
 
 #undef LOCTEXT_NAMESPACE
