Runtime memory improvements
-reordered entries in FNiagaraDataSetCompiledData
-reenabled sharing of FNiagaraDataSetCompiledData between FNiagaraDataSet

diff --git a/Engine/Plugins/FX/Niagara/Source/Niagara/Classes/NiagaraDataSet.h b/Engine/Plugins/FX/Niagara/Source/Niagara/Classes/NiagaraDataSet.h
--- a/Engine/Plugins/FX/Niagara/Source/Niagara/Classes/NiagaraDataSet.h
+++ b/Engine/Plugins/FX/Niagara/Source/Niagara/Classes/NiagaraDataSet.h
@@ -238,6 +238,10 @@ struct NIAGARA_API FNiagaraDataSetCompiledData
 	UPROPERTY()
 	TArray<FNiagaraVariableLayoutInfo> VariableLayouts;
 
+	/** Unique ID for this DataSet. Used to allow referencing from other emitters and Systems. */
+	UPROPERTY()
+	FNiagaraDataSetID ID;
+
 	/** Total number of components of each type in the data set. */
 	UPROPERTY()
 	uint32 TotalFloatComponents;
@@ -249,10 +253,6 @@ struct NIAGARA_API FNiagaraDataSetCompiledData
 	UPROPERTY()
 	uint32 bRequiresPersistentIDs : 1;
 
-	/** Unique ID for this DataSet. Used to allow referencing from other emitters and Systems. */
-	UPROPERTY()
-	FNiagaraDataSetID ID;
-
 	/** Sim target this DataSet is targeting (CPU/GPU). */
 	UPROPERTY()
 	ENiagaraSimTarget SimTarget;
@@ -281,9 +281,7 @@ public:
 
 	FORCEINLINE void Init(const FNiagaraDataSetCompiledData* InDataSetCompiledData)
 	{
-		//CompiledData = InDataSetCompiledData != nullptr ? InDataSetCompiledData : &FNiagaraDataSetCompiledData::DummyCompiledData;
-		//Temporarily taking a copy of the compiled data to avoid lifetime issues in some cases.
-		CompiledData = InDataSetCompiledData != nullptr ? *InDataSetCompiledData : FNiagaraDataSetCompiledData::DummyCompiledData;
+		CompiledData = InDataSetCompiledData != nullptr ? InDataSetCompiledData : &FNiagaraDataSetCompiledData::DummyCompiledData;
 		bInitialized = true;
 		Reset();
 	}
@@ -304,9 +302,9 @@ public:
 	uint32 GetSizeBytes()const;
 
 	FORCEINLINE bool IsInitialized() const { return bInitialized; }
-	FORCEINLINE ENiagaraSimTarget GetSimTarget() const { return CompiledData.SimTarget; }
-	FORCEINLINE FNiagaraDataSetID GetID() const { return CompiledData.ID; }	
-	FORCEINLINE bool RequiresPersistentIDs() const { return CompiledData.bRequiresPersistentIDs; }
+	FORCEINLINE ENiagaraSimTarget GetSimTarget() const { return CompiledData->SimTarget; }
+	FORCEINLINE FNiagaraDataSetID GetID() const { return CompiledData->ID; }	
+	FORCEINLINE bool RequiresPersistentIDs() const { return CompiledData->bRequiresPersistentIDs; }
 
 	FORCEINLINE TArray<int32>& GetFreeIDTable() { return FreeIDsTable; }
 	FORCEINLINE TArray<int32>& GetSpawnedIDsTable() { return SpawnedIDsTable; }
@@ -317,35 +315,35 @@ public:
 	FORCEINLINE FRWBuffer& GetGPUFreeIDs() { return GPUFreeIDs; }
 	FORCEINLINE uint32 GetGPUNumAllocatedIDs() const { return GPUNumAllocatedIDs; }
 
-	FORCEINLINE const TArray<FNiagaraVariable>& GetVariables()const { return CompiledData.Variables; }
-	FORCEINLINE uint32 GetNumVariables()const { return CompiledData.Variables.Num(); }
-	FORCEINLINE bool HasVariable(const FNiagaraVariable& Var)const { return CompiledData.Variables.Contains(Var); }
-	FORCEINLINE uint32 GetNumFloatComponents()const { return CompiledData.TotalFloatComponents; }
-	FORCEINLINE uint32 GetNumInt32Components()const { return CompiledData.TotalInt32Components; }
+	FORCEINLINE const TArray<FNiagaraVariable>& GetVariables() const { return CompiledData->Variables; }
+	FORCEINLINE uint32 GetNumVariables() const { return CompiledData->Variables.Num(); }
+	FORCEINLINE bool HasVariable(const FNiagaraVariable& Var) const { return CompiledData->Variables.Contains(Var); }
+	FORCEINLINE uint32 GetNumFloatComponents() const { return CompiledData->TotalFloatComponents; }
+	FORCEINLINE uint32 GetNumInt32Components() const { return CompiledData->TotalInt32Components; }
 
-	const TArray<FNiagaraVariableLayoutInfo>& GetVariableLayouts()const { return CompiledData.VariableLayouts; }
-	const FNiagaraVariableLayoutInfo* GetVariableLayout(const FNiagaraVariable& Var)const;
+	const TArray<FNiagaraVariableLayoutInfo>& GetVariableLayouts() const { return CompiledData->VariableLayouts; }
+	const FNiagaraVariableLayoutInfo* GetVariableLayout(const FNiagaraVariable& Var) const;
 	bool GetVariableComponentOffsets(const FNiagaraVariable& Var, int32 &FloatStart, int32 &IntStart) const;
 
 	void CopyTo(FNiagaraDataSet& Other, int32 StartIdx = 0, int32 NumInstances = INDEX_NONE, bool bResetOther=true)const;
 
 	void CopyFromGPUReadback(float* GPUReadBackFloat, int* GPUReadBackInt, int32 StartIdx = 0, int32 NumInstances = INDEX_NONE, uint32 FloatStride = 0, uint32 IntStride = 0);
 
-	void CheckForNaNs()const;
+	void CheckForNaNs() const;
 
-	void Dump(int32 StartIndex, int32 NumInstances, const FString& Label)const;
+	void Dump(int32 StartIndex, int32 NumInstances, const FString& Label) const;
 
 	FORCEINLINE bool IsCurrentDataValid()const { return CurrentData != nullptr; }
 	FORCEINLINE FNiagaraDataBuffer* GetCurrentData()const {	return CurrentData; }
 	FORCEINLINE FNiagaraDataBuffer* GetDestinationData()const { return DestinationData; }
 
-	FORCEINLINE FNiagaraDataBuffer& GetCurrentDataChecked()const
+	FORCEINLINE FNiagaraDataBuffer& GetCurrentDataChecked() const
 	{
 		check(CurrentData);
 		return *CurrentData;
 	}
 
-	FORCEINLINE FNiagaraDataBuffer& GetDestinationDataChecked()const
+	FORCEINLINE FNiagaraDataBuffer& GetDestinationDataChecked() const
 	{
 		check(DestinationData);
 		return *DestinationData;
@@ -377,9 +375,7 @@ private:
 #endif
 	}
 
-	//const FNiagaraDataSetCompiledData* CompiledData;
-	//For safety we're temporarily taking a copy of the compiled data. In certain cases the lifetime of the CompiledData ptr cannot be guaranteed. 
-	FNiagaraDataSetCompiledData CompiledData;
+	const FNiagaraDataSetCompiledData* CompiledData;
 
 	/** Table of free IDs available to allocate next tick. */
 	TArray<int32> FreeIDsTable;
diff --git a/Engine/Plugins/FX/Niagara/Source/Niagara/Private/NiagaraDataSet.cpp b/Engine/Plugins/FX/Niagara/Source/Niagara/Private/NiagaraDataSet.cpp
--- a/Engine/Plugins/FX/Niagara/Source/Niagara/Private/NiagaraDataSet.cpp
+++ b/Engine/Plugins/FX/Niagara/Source/Niagara/Private/NiagaraDataSet.cpp
@@ -76,7 +76,7 @@ static FAutoConsoleVariableRef CVarNiagaraGPUDataBufferBufferSlack(
 );
 
 FNiagaraDataSet::FNiagaraDataSet()
-	: CompiledData(FNiagaraDataSetCompiledData::DummyCompiledData)
+	: CompiledData(&FNiagaraDataSetCompiledData::DummyCompiledData)
 	, NumFreeIDs(0)
 	, MaxUsedID(0)
 	, IDAcquireTag(0)
@@ -398,7 +398,7 @@ void FNiagaraDataSet::AllocateGPUFreeIDs(uint32 InNumInstances, FRHICommandList&
 const FNiagaraVariableLayoutInfo* FNiagaraDataSet::GetVariableLayout(const FNiagaraVariable& Var)const
 {
 	int32 VarLayoutIndex = GetVariables().IndexOfByKey(Var);
-	return VarLayoutIndex != INDEX_NONE ? &CompiledData.VariableLayouts[VarLayoutIndex] : nullptr;
+	return VarLayoutIndex != INDEX_NONE ? &CompiledData->VariableLayouts[VarLayoutIndex] : nullptr;
 }
 
 bool FNiagaraDataSet::GetVariableComponentOffsets(const FNiagaraVariable& Var, int32 &FloatStart, int32 &IntStart) const
