Remove unused DataSetMap & accessors from Niagara system/emitter instances

diff --git a/Engine/Plugins/FX/Niagara/Source/Niagara/Classes/NiagaraEmitterInstance.h b/Engine/Plugins/FX/Niagara/Source/Niagara/Classes/NiagaraEmitterInstance.h
--- a/Engine/Plugins/FX/Niagara/Source/Niagara/Classes/NiagaraEmitterInstance.h
+++ b/Engine/Plugins/FX/Niagara/Source/Niagara/Classes/NiagaraEmitterInstance.h
@@ -94,8 +94,6 @@ public:
 	ENiagaraExecutionState NIAGARA_API GetExecutionState() { return ExecutionState; }
 	void NIAGARA_API SetExecutionState(ENiagaraExecutionState InState);
 
-	FNiagaraDataSet* GetDataSet(FNiagaraDataSetID SetID);
-
 	FBox GetBounds();
 
 	FNiagaraScriptExecutionContext& GetSpawnExecutionContext() { return SpawnExecContext; }
@@ -200,7 +198,6 @@ private:
 
 	TArray<FNiagaraDataSet*> UpdateScriptEventDataSets;
 	TArray<FNiagaraDataSet*> SpawnScriptEventDataSets;
-	TMap<FNiagaraDataSetID, FNiagaraDataSet*> DataSetMap;
 
 	TArray<bool> UpdateEventGeneratorIsSharedByIndex;
 	TArray<bool> SpawnEventGeneratorIsSharedByIndex;
diff --git a/Engine/Plugins/FX/Niagara/Source/Niagara/Private/NiagaraEmitterInstance.cpp b/Engine/Plugins/FX/Niagara/Source/Niagara/Private/NiagaraEmitterInstance.cpp
--- a/Engine/Plugins/FX/Niagara/Source/Niagara/Private/NiagaraEmitterInstance.cpp
+++ b/Engine/Plugins/FX/Niagara/Source/Niagara/Private/NiagaraEmitterInstance.cpp
@@ -282,11 +282,6 @@ void FNiagaraEmitterInstance::Init(int32 InEmitterIdx, FNiagaraSystemInstanceID
 
 		ResetSimulation();
 
-		DataSetMap.Empty();
-
-		//Add the particle data to the data set map.
-		//Currently just used for the tick loop but will also allow access directly to the particle data from other emitters.
-		DataSetMap.Add(Data.GetID(), &Data);
 		//Warn the user if there are any attributes used in the update script that are not initialized in the spawn script.
 		//TODO: We need some window in the System editor and possibly the graph editor for warnings and errors.
 
@@ -730,22 +725,6 @@ void FNiagaraEmitterInstance::BindParameters(bool bExternalOnly)
 	}
 }
 
-FNiagaraDataSet* FNiagaraEmitterInstance::GetDataSet(FNiagaraDataSetID SetID)
-{
-	FNiagaraDataSet** SetPtr = DataSetMap.Find(SetID);
-	FNiagaraDataSet* Ret = NULL;
-	if (SetPtr)
-	{
-		Ret = *SetPtr;
-	}
-	else
-	{
-		// TODO: keep track of data sets generated by the scripts (event writers) and find here
-	}
-
-	return Ret;
-}
-
 const FNiagaraEmitterHandle& FNiagaraEmitterInstance::GetEmitterHandle() const
 {
 	UNiagaraSystem* Sys = ParentSystemInstance->GetSystem();
diff --git a/Engine/Plugins/FX/Niagara/Source/Niagara/Private/NiagaraSystemInstance.cpp b/Engine/Plugins/FX/Niagara/Source/Niagara/Private/NiagaraSystemInstance.cpp
--- a/Engine/Plugins/FX/Niagara/Source/Niagara/Private/NiagaraSystemInstance.cpp
+++ b/Engine/Plugins/FX/Niagara/Source/Niagara/Private/NiagaraSystemInstance.cpp
@@ -2224,30 +2224,6 @@ FNiagaraEmitterInstance* FNiagaraSystemInstance::GetEmitterByID(FGuid InID)
 	return nullptr;
 }
 
-FNiagaraDataSet* FNiagaraSystemInstance::GetDataSet(FNiagaraDataSetID SetID, FName EmitterName)
-{
-	if (EmitterName == NAME_None)
-	{
-		if (FNiagaraDataSet* ExternalSet = ExternalEvents.Find(SetID))
-		{
-			return ExternalSet;
-		}
-	}
-	for (TSharedPtr<FNiagaraEmitterInstance, ESPMode::ThreadSafe> Emitter : Emitters)
-	{
-		check(Emitter.IsValid());
-		if (!Emitter->IsComplete())
-		{
-			if (Emitter->GetCachedIDName() == EmitterName)
-			{
-				return Emitter->GetDataSet(SetID);
-			}
-		}
-	}
-
-	return NULL;
-}
-
 FNiagaraSystemInstance::FOnInitialized& FNiagaraSystemInstance::OnInitialized()
 {
 	return OnInitializedDelegate;
diff --git a/Engine/Plugins/FX/Niagara/Source/Niagara/Public/NiagaraSystemInstance.h b/Engine/Plugins/FX/Niagara/Source/Niagara/Public/NiagaraSystemInstance.h
--- a/Engine/Plugins/FX/Niagara/Source/Niagara/Public/NiagaraSystemInstance.h
+++ b/Engine/Plugins/FX/Niagara/Source/Niagara/Public/NiagaraSystemInstance.h
@@ -171,9 +171,6 @@ public:
 
 	FORCEINLINE bool IsSolo()const { return bSolo; }
 
-	/** Gets a data set either from another emitter or one owned by the System itself. */
-	FNiagaraDataSet* GetDataSet(FNiagaraDataSetID SetID, FName EmitterName = NAME_None);
-
 	/** Gets a multicast delegate which is called whenever this instance is initialized with an System asset. */
 	FOnInitialized& OnInitialized();
 
@@ -341,8 +338,6 @@ private:
 	float LODDistance;
 	float MaxLODDistance;
 
-	TMap<FNiagaraDataSetID, FNiagaraDataSet> ExternalEvents;
-
 	TArray< TSharedRef<FNiagaraEmitterInstance, ESPMode::ThreadSafe> > Emitters;
 
 	FOnInitialized OnInitializedDelegate;
