Moved Niaga's UserDefinedTypeRegistry to WITH_EDITORONLY to prevent a number of unnecessary allocations/asset loads in game

diff --git a/Engine/Plugins/FX/Niagara/Source/Niagara/Private/NiagaraModule.cpp b/Engine/Plugins/FX/Niagara/Source/Niagara/Private/NiagaraModule.cpp
--- a/Engine/Plugins/FX/Niagara/Source/Niagara/Private/NiagaraModule.cpp
+++ b/Engine/Plugins/FX/Niagara/Source/Niagara/Private/NiagaraModule.cpp
@@ -558,7 +558,9 @@ void FNiagaraTypeDefinition::Init()
 
 	ParameterScopeEnum = StaticEnum<ENiagaraParameterScope>();
 	
+#if WITH_EDITOR
 	RecreateUserDefinedTypeRegistry();
+#endif
 }
 
 bool FNiagaraTypeDefinition::IsValidNumericInput(const FNiagaraTypeDefinition& TypeDef)
@@ -635,6 +637,7 @@ bool FNiagaraTypeDefinition::AppendCompileHash(FNiagaraCompileHashVisitor* InVis
 #endif
 }
 
+#if WITH_EDITOR
 void FNiagaraTypeDefinition::RecreateUserDefinedTypeRegistry()
 {
 	static auto* CoreUObjectPkg = FindObjectChecked<UPackage>(nullptr, TEXT("/Script/CoreUObject"));
@@ -749,6 +752,7 @@ void FNiagaraTypeDefinition::RecreateUserDefinedTypeRegistry()
 
 	FNiagaraTypeRegistry::Register(StaticEnum<ENiagaraLegacyTrailWidthMode>(), true, true, false);
 }
+#endif
 
 bool FNiagaraTypeDefinition::IsScalarDefinition(const FNiagaraTypeDefinition& Type)
 {
diff --git a/Engine/Plugins/FX/Niagara/Source/Niagara/Public/NiagaraTypes.h b/Engine/Plugins/FX/Niagara/Source/Niagara/Public/NiagaraTypes.h
--- a/Engine/Plugins/FX/Niagara/Source/Niagara/Public/NiagaraTypes.h
+++ b/Engine/Plugins/FX/Niagara/Source/Niagara/Public/NiagaraTypes.h
@@ -870,7 +870,9 @@ private:
 
 public:
 	static void Init();
+#if WITH_EDITOR
 	static void RecreateUserDefinedTypeRegistry();
+#endif
 	static const FNiagaraTypeDefinition& GetFloatDef() { return FloatDef; }
 	static const FNiagaraTypeDefinition& GetBoolDef() { return BoolDef; }
 	static const FNiagaraTypeDefinition& GetIntDef() { return IntDef; }
