CIS fixes

diff --git a/Engine/Plugins/FX/Niagara/Source/NiagaraEditor/Private/ViewModels/Stack/NiagaraStackFunctionInput.cpp b/Engine/Plugins/FX/Niagara/Source/NiagaraEditor/Private/ViewModels/Stack/NiagaraStackFunctionInput.cpp
--- a/Engine/Plugins/FX/Niagara/Source/NiagaraEditor/Private/ViewModels/Stack/NiagaraStackFunctionInput.cpp
+++ b/Engine/Plugins/FX/Niagara/Source/NiagaraEditor/Private/ViewModels/Stack/NiagaraStackFunctionInput.cpp
@@ -1675,14 +1675,18 @@ void UNiagaraStackFunctionInput::ReassignDynamicInputScript(UNiagaraScript* Dyna
 			GetUnfilteredChildrenOfType(DynamicInputCollections);
 
 			FText ConvertMessage;
-			if (ConversionUtility && DynamicInputCollections.Num() == 0 && !ConversionUtility->Convert(OldScript, OldClipboardContent, DynamicInputScript, DynamicInputCollections[0], NewClipboardContent, InputValues.DynamicNode.Get(), ConvertMessage) || !ConvertMessage.IsEmptyOrWhitespace())
+			if (ConversionUtility && DynamicInputCollections.Num() == 0)
 			{
-				// Notify the end-user about the convert message, but continue the process as they could always undo.
-				FNotificationInfo Msg(FText::Format(LOCTEXT("FixConvertInPlace", "Conversion Note: {0}"), ConvertMessage));
-				Msg.ExpireDuration = 5.0f;
-				Msg.bFireAndForget = true;
-				Msg.Image = FCoreStyle::Get().GetBrush(TEXT("MessageLog.Info"));
-				FSlateNotificationManager::Get().AddNotification(Msg);
+				bool bConverted = ConversionUtility->Convert(OldScript, OldClipboardContent, DynamicInputScript, DynamicInputCollections[0], NewClipboardContent, InputValues.DynamicNode.Get(), ConvertMessage);
+				if (!ConvertMessage.IsEmptyOrWhitespace())
+				{
+					// Notify the end-user about the convert message, but continue the process as they could always undo.
+					FNotificationInfo Msg(FText::Format(LOCTEXT("FixConvertInPlace", "Conversion Note: {0}"), ConvertMessage));
+					Msg.ExpireDuration = 5.0f;
+					Msg.bFireAndForget = true;
+					Msg.Image = FCoreStyle::Get().GetBrush(TEXT("MessageLog.Note"));
+					FSlateNotificationManager::Get().AddNotification(Msg);
+				}
 			}
 		}
 
diff --git a/Engine/Plugins/FX/Niagara/Source/NiagaraEditor/Private/ViewModels/Stack/NiagaraStackModuleItem.cpp b/Engine/Plugins/FX/Niagara/Source/NiagaraEditor/Private/ViewModels/Stack/NiagaraStackModuleItem.cpp
--- a/Engine/Plugins/FX/Niagara/Source/NiagaraEditor/Private/ViewModels/Stack/NiagaraStackModuleItem.cpp
+++ b/Engine/Plugins/FX/Niagara/Source/NiagaraEditor/Private/ViewModels/Stack/NiagaraStackModuleItem.cpp
@@ -989,14 +989,18 @@ void UNiagaraStackModuleItem::ReassignModuleScript(UNiagaraScript* ModuleScript)
 			UNiagaraClipboardContent* NewClipboardContent = UNiagaraClipboardContent::Create();
 			Copy(NewClipboardContent);
 
-			if (ConversionUtility && !ConversionUtility->Convert(OldScript, OldClipboardContent, ModuleScript, InputCollection, NewClipboardContent, FunctionCallNode, ConvertMessage) || !ConvertMessage.IsEmptyOrWhitespace())
+			if (ConversionUtility )
 			{
-				// Notify the end-user about the convert message, but continue the process as they could always undo.
-				FNotificationInfo Msg(FText::Format(LOCTEXT("FixConvertInPlace", "Conversion Note: {0}"), ConvertMessage));
-				Msg.ExpireDuration = 5.0f;
-				Msg.bFireAndForget = true;
-				Msg.Image = FCoreStyle::Get().GetBrush(TEXT("MessageLog.Note"));
-				FSlateNotificationManager::Get().AddNotification(Msg);
+				bool bConverted = ConversionUtility->Convert(OldScript, OldClipboardContent, ModuleScript, InputCollection, NewClipboardContent, FunctionCallNode, ConvertMessage);
+				if (!ConvertMessage.IsEmptyOrWhitespace())
+				{
+					// Notify the end-user about the convert message, but continue the process as they could always undo.
+					FNotificationInfo Msg(FText::Format(LOCTEXT("FixConvertInPlace", "Conversion Note: {0}"), ConvertMessage));
+					Msg.ExpireDuration = 5.0f;
+					Msg.bFireAndForget = true;
+					Msg.Image = FCoreStyle::Get().GetBrush(TEXT("MessageLog.Note"));
+					FSlateNotificationManager::Get().AddNotification(Msg);
+				}
 			}
 		}
 	}
