Fixed out of bounds accesses in shader when spawn-only sim stages are used on an emitter which doesn't spawn particles.

The unused slots in EmitterSpawnInfoOffsets and EmitterSpawnInfoParams were only cleared if we had at least one spawn info, so SetupExecIndexAndSpawnInfoForGPU() ended up reading bogus values and producing an invalid execution index. Also cleaned up the code related to these parameters:
 * EmitterSpawnInfoOffsets is an int array, which gets rid of some float to int conversions in the shader
 * EmitterSpawnInfoParams is a mix of floats and ints on CPU, which gets rid of some ugly casts (and strict aliasing violations)

diff --git a/Engine/Plugins/FX/Niagara/Shaders/Private/NiagaraEmitterInstanceShader.usf b/Engine/Plugins/FX/Niagara/Shaders/Private/NiagaraEmitterInstanceShader.usf
--- a/Engine/Plugins/FX/Niagara/Shaders/Private/NiagaraEmitterInstanceShader.usf
+++ b/Engine/Plugins/FX/Niagara/Shaders/Private/NiagaraEmitterInstanceShader.usf
@@ -101,7 +101,7 @@ bool NiagaraAll(bool4 b) { return b.x && b.y && b.z && b.w; }
 	// NIAGARA_MAX_GPU_SPAWN_INFOS is set from the shader compiler
 	#define NIAGARA_MAX_GPU_SPAWN_INFOS_V4	((NIAGARA_MAX_GPU_SPAWN_INFOS + 3) / 4)
 
-	float4	EmitterSpawnInfoOffsets[NIAGARA_MAX_GPU_SPAWN_INFOS_V4];
+	int4	EmitterSpawnInfoOffsets[NIAGARA_MAX_GPU_SPAWN_INFOS_V4];
 	float4	EmitterSpawnInfoParams[NIAGARA_MAX_GPU_SPAWN_INFOS];			// Packed data where x = IntervalDt, y = InterpStartDt, z = Group, w = Start Particle Index
 
 	static int GInterpSpawnIndex;
@@ -128,9 +128,13 @@ bool NiagaraAll(bool4 b) { return b.x && b.y && b.z && b.w; }
 
 		int SpawnInfoIndex = 0;
 		UNROLL
-		for ( int i=0; i < NIAGARA_MAX_GPU_SPAWN_INFOS_V4; ++i )
+		for (int i = 0; i < NIAGARA_MAX_GPU_SPAWN_INFOS_V4; ++i)
 		{
-			SpawnInfoIndex += (int)dot(1.0f, GGPUExecIndex >= EmitterSpawnInfoOffsets[i]);
+			// This returns 0xffffffff for each component when the comparison is true, so we'll do a
+			// bitwise and with 1 to get increment amounts for each spawn info.
+			int4 CompareResults = GGPUExecIndex >= EmitterSpawnInfoOffsets[i];
+			CompareResults = CompareResults & int4(1, 1, 1, 1);
+			SpawnInfoIndex += CompareResults.x + CompareResults.y + CompareResults.z + CompareResults.w;
 		}
 
 		Emitter_SpawnInterval		= EmitterSpawnInfoParams[SpawnInfoIndex].x;
diff --git a/Engine/Plugins/FX/Niagara/Source/Niagara/Classes/NiagaraScriptExecutionContext.h b/Engine/Plugins/FX/Niagara/Source/Niagara/Classes/NiagaraScriptExecutionContext.h
--- a/Engine/Plugins/FX/Niagara/Source/Niagara/Classes/NiagaraScriptExecutionContext.h
+++ b/Engine/Plugins/FX/Niagara/Source/Niagara/Classes/NiagaraScriptExecutionContext.h
@@ -185,13 +185,21 @@ struct FNiagaraScriptExecutionContext
 	bool CanExecute()const;
 };
 
+struct FNiagaraGpuSpawnInfoParams
+{
+	float IntervalDt;
+	float InterpStartDt;
+	int32 SpawnGroup;
+	int32 GroupSpawnStartIndex;
+};
+
 struct FNiagaraGpuSpawnInfo
 {
-	uint32		EventSpawnTotal = 0;
-	uint32		SpawnRateInstances = 0;
-	uint32		MaxParticleCount = 0;
-	FVector4	SpawnInfoStartOffsets[NIAGARA_MAX_GPU_SPAWN_INFOS_V4];
-	FVector4	SpawnInfoParams[NIAGARA_MAX_GPU_SPAWN_INFOS];
+	uint32 EventSpawnTotal = 0;
+	uint32 SpawnRateInstances = 0;
+	uint32 MaxParticleCount = 0;
+	int32 SpawnInfoStartOffsets[NIAGARA_MAX_GPU_SPAWN_INFOS];
+	FNiagaraGpuSpawnInfoParams SpawnInfoParams[NIAGARA_MAX_GPU_SPAWN_INFOS];
 };
 
 struct FNiagaraComputeExecutionContext
diff --git a/Engine/Plugins/FX/Niagara/Source/Niagara/Private/NiagaraEmitterInstance.cpp b/Engine/Plugins/FX/Niagara/Source/Niagara/Private/NiagaraEmitterInstance.cpp
--- a/Engine/Plugins/FX/Niagara/Source/Niagara/Private/NiagaraEmitterInstance.cpp
+++ b/Engine/Plugins/FX/Niagara/Source/Niagara/Private/NiagaraEmitterInstance.cpp
@@ -1153,7 +1153,7 @@ void FNiagaraEmitterInstance::Tick(float DeltaSeconds)
 
 		// Calculate spawn information to pass to the RT
 		{
-			static_assert((NIAGARA_MAX_GPU_SPAWN_INFOS == NIAGARA_MAX_GPU_SPAWN_INFOS_V4 * 4) && (NIAGARA_MAX_GPU_SPAWN_INFOS > 0), "NIAGARA_MAX_GPU_SPAWN_INFOS should be greater than zero and a multiple of 4");
+			static_assert(((NIAGARA_MAX_GPU_SPAWN_INFOS % 4) == 0) && (NIAGARA_MAX_GPU_SPAWN_INFOS > 0), "NIAGARA_MAX_GPU_SPAWN_INFOS should be greater than zero and a multiple of 4");
 
 			FNiagaraGpuSpawnInfo& GpuSpawnInfo = GPUExecContext->GpuSpawnInfo_GT;
 			GpuSpawnInfo.EventSpawnTotal = EventSpawnTotal;
@@ -1176,13 +1176,13 @@ void FNiagaraEmitterInstance::Tick(float DeltaSeconds)
 							break;
 						}
 
-						GpuSpawnInfo.SpawnInfoParams[NumSpawnInfos].X = Info.IntervalDt;
-						GpuSpawnInfo.SpawnInfoParams[NumSpawnInfos].Y = Info.InterpStartDt;
-						reinterpret_cast<int32&>(GpuSpawnInfo.SpawnInfoParams[NumSpawnInfos].Z) = Info.SpawnGroup;
-						reinterpret_cast<int32&>(GpuSpawnInfo.SpawnInfoParams[NumSpawnInfos].W) = GpuSpawnInfo.SpawnRateInstances;
+						GpuSpawnInfo.SpawnInfoParams[NumSpawnInfos].IntervalDt = Info.IntervalDt;
+						GpuSpawnInfo.SpawnInfoParams[NumSpawnInfos].InterpStartDt = Info.InterpStartDt;
+						GpuSpawnInfo.SpawnInfoParams[NumSpawnInfos].SpawnGroup = Info.SpawnGroup;
+						GpuSpawnInfo.SpawnInfoParams[NumSpawnInfos].GroupSpawnStartIndex = (int32)GpuSpawnInfo.SpawnRateInstances;
 
 						GpuSpawnInfo.SpawnRateInstances += Info.Count;
-						reinterpret_cast<float*>(GpuSpawnInfo.SpawnInfoStartOffsets)[NumSpawnInfos] = (float)GpuSpawnInfo.SpawnRateInstances;
+						GpuSpawnInfo.SpawnInfoStartOffsets[NumSpawnInfos] = (int32)GpuSpawnInfo.SpawnRateInstances;
 
 						++NumSpawnInfos;
 					}
@@ -1197,19 +1197,15 @@ void FNiagaraEmitterInstance::Tick(float DeltaSeconds)
 				}
 			}
 
-			// If we have spawning make sure we clear out the remaining data and leave the end slot as MAX to avoid reading off end of the array on the GPU
-			if ( GpuSpawnInfo.EventSpawnTotal + GpuSpawnInfo.SpawnRateInstances >  0 )
+			// Clear out the remaining data and leave the end slot as MAX to avoid reading off end of the array on the GPU
+			while (NumSpawnInfos < NIAGARA_MAX_GPU_SPAWN_INFOS)
 			{
-				while (NumSpawnInfos < NIAGARA_MAX_GPU_SPAWN_INFOS)
-				{
-					GpuSpawnInfo.SpawnInfoParams[NumSpawnInfos].X = 0.0f;
-					GpuSpawnInfo.SpawnInfoParams[NumSpawnInfos].Y = 0.0f;
-					reinterpret_cast<int32&>(GpuSpawnInfo.SpawnInfoParams[NumSpawnInfos].Z) = 0;
-					reinterpret_cast<int32&>(GpuSpawnInfo.SpawnInfoParams[NumSpawnInfos].W) = GpuSpawnInfo.SpawnRateInstances;
-					reinterpret_cast<float*>(GpuSpawnInfo.SpawnInfoStartOffsets)[NumSpawnInfos] = (float)MAX_FLT;
-					++NumSpawnInfos;
-				}
-				reinterpret_cast<float*>(GpuSpawnInfo.SpawnInfoStartOffsets)[NIAGARA_MAX_GPU_SPAWN_INFOS - 1] = MAX_FLT;
+				GpuSpawnInfo.SpawnInfoParams[NumSpawnInfos].IntervalDt = 0.0f;
+				GpuSpawnInfo.SpawnInfoParams[NumSpawnInfos].InterpStartDt = 0.0f;
+				GpuSpawnInfo.SpawnInfoParams[NumSpawnInfos].SpawnGroup = 0;
+				GpuSpawnInfo.SpawnInfoParams[NumSpawnInfos].GroupSpawnStartIndex = (int32)GpuSpawnInfo.SpawnRateInstances;
+				GpuSpawnInfo.SpawnInfoStartOffsets[NumSpawnInfos] = INT32_MAX;
+				++NumSpawnInfos;
 			}
 		}
 
diff --git a/Engine/Plugins/FX/Niagara/Source/Niagara/Private/NiagaraEmitterInstanceBatcher.cpp b/Engine/Plugins/FX/Niagara/Source/Niagara/Private/NiagaraEmitterInstanceBatcher.cpp
--- a/Engine/Plugins/FX/Niagara/Source/Niagara/Private/NiagaraEmitterInstanceBatcher.cpp
+++ b/Engine/Plugins/FX/Niagara/Source/Niagara/Private/NiagaraEmitterInstanceBatcher.cpp
@@ -1333,8 +1333,10 @@ void NiagaraEmitterInstanceBatcher::Run(const FNiagaraGPUSystemTick& Tick, const
 	//
 	static_assert((sizeof(SpawnInfo.SpawnInfoStartOffsets) % SHADER_PARAMETER_ARRAY_ELEMENT_ALIGNMENT) == 0, "sizeof SpawnInfoStartOffsets should be a multiple of SHADER_PARAMETER_ARRAY_ELEMENT_ALIGNMENT");
 	static_assert((sizeof(SpawnInfo.SpawnInfoParams) % SHADER_PARAMETER_ARRAY_ELEMENT_ALIGNMENT) == 0, "sizeof SpawnInfoParams should be a multiple of SHADER_PARAMETER_ARRAY_ELEMENT_ALIGNMENT");
-	SetShaderValueArray(RHICmdList, ComputeShader, Shader->EmitterSpawnInfoOffsetsParam, SpawnInfo.SpawnInfoStartOffsets, NIAGARA_MAX_GPU_SPAWN_INFOS_V4);
-	SetShaderValueArray(RHICmdList, ComputeShader, Shader->EmitterSpawnInfoParamsParam, SpawnInfo.SpawnInfoParams, NIAGARA_MAX_GPU_SPAWN_INFOS);
+	SetShaderValueArray(RHICmdList, ComputeShader, Shader->EmitterSpawnInfoOffsetsParam, SpawnInfo.SpawnInfoStartOffsets, NIAGARA_MAX_GPU_SPAWN_INFOS);
+	// This parameter is an array of structs with 2 floats and 2 ints on CPU, but a float4 array on GPU. The shader uses asint() to cast the integer values. To set the parameter, 
+	// we pass the structure array as a float* to SetShaderValueArray() and specify the number of floats (not float vectors).
+	SetShaderValueArray(RHICmdList, ComputeShader, Shader->EmitterSpawnInfoParamsParam, &SpawnInfo.SpawnInfoParams[0].IntervalDt, 4*NIAGARA_MAX_GPU_SPAWN_INFOS);
 
 	SetShaderValue(RHICmdList, ComputeShader, Shader->UpdateStartInstanceParam, UpdateStartInstance);					// 0, except for event handler runs
 	SetShaderValue(RHICmdList, ComputeShader, Shader->NumSpawnedInstancesParam, InstancesToSpawnThisFrame);				// number of instances in the spawn run
diff --git a/Engine/Plugins/FX/Niagara/Source/Niagara/Public/NiagaraCommon.h b/Engine/Plugins/FX/Niagara/Source/Niagara/Public/NiagaraCommon.h
--- a/Engine/Plugins/FX/Niagara/Source/Niagara/Public/NiagaraCommon.h
+++ b/Engine/Plugins/FX/Niagara/Source/Niagara/Public/NiagaraCommon.h
@@ -29,7 +29,6 @@ const uint32 NIAGARA_MAX_COMPUTE_THREADGROUPS = 65535;
 
 /** The maximum number of spawn infos we can run on the GPU, modifying this will require a version update as it is used in the shader compiler  */
 constexpr uint32 NIAGARA_MAX_GPU_SPAWN_INFOS = 8;
-constexpr uint32 NIAGARA_MAX_GPU_SPAWN_INFOS_V4 = (NIAGARA_MAX_GPU_SPAWN_INFOS + 3) / 4;
 
 /** TickGroup information for Niagara.  */
 constexpr ETickingGroup NiagaraFirstTickGroup = TG_PrePhysics;
diff --git a/Engine/Plugins/FX/Niagara/Source/NiagaraCore/Private/NiagaraCustomVersion.cpp b/Engine/Plugins/FX/Niagara/Source/NiagaraCore/Private/NiagaraCustomVersion.cpp
--- a/Engine/Plugins/FX/Niagara/Source/NiagaraCore/Private/NiagaraCustomVersion.cpp
+++ b/Engine/Plugins/FX/Niagara/Source/NiagaraCore/Private/NiagaraCustomVersion.cpp
@@ -8,4 +8,4 @@ const FGuid FNiagaraCustomVersion::GUID(0xFCF57AFA, 0x50764283, 0xB9A9E658, 0xFF
 // Register the custom version with core
 FCustomVersionRegistration GRegisterNiagaraCustomVersion(FNiagaraCustomVersion::GUID, FNiagaraCustomVersion::LatestVersion, TEXT("NiagaraVer"));
 
-const FGuid FNiagaraCustomVersion::LatestScriptCompileVersion(0x1BDBCD02, 0xA1B55E47, 0x802D1B66, 0x5C5FCC12);
+const FGuid FNiagaraCustomVersion::LatestScriptCompileVersion(0x3F7E65A5, 0xF86140EC, 0x8FAB94B7, 0xE9F031FD);
