- Integrate 11984198: Additional fix for handling virtual bones and component-less preview in Niagara Editor

diff --git a/Engine/Plugins/FX/Niagara/Source/Niagara/Private/NiagaraDataInterfaceSkeletalMesh.cpp b/Engine/Plugins/FX/Niagara/Source/Niagara/Private/NiagaraDataInterfaceSkeletalMesh.cpp
--- a/Engine/Plugins/FX/Niagara/Source/Niagara/Private/NiagaraDataInterfaceSkeletalMesh.cpp
+++ b/Engine/Plugins/FX/Niagara/Source/Niagara/Private/NiagaraDataInterfaceSkeletalMesh.cpp
@@ -682,8 +682,6 @@ void FSkeletalMeshGpuDynamicBufferProxy::NewFrame(const FNDISkeletalMesh_Instanc
 	auto FillBuffers =
 		[&](const TArray<FTransform>& BoneTransforms)
 		{
-			check(BoneTransforms.Num() == SkelMesh->RefBasesInvMatrix.Num());
-
 			// Fill AllSectionsRefToLocalMatrices
 			TIndirectArray<FSkeletalMeshLODRenderData>& LODRenderDataArray = SkelMesh->GetResourceForRendering()->LODRenderData;
 			check(0 <= LODIndex && LODIndex < LODRenderDataArray.Num());
@@ -708,7 +706,7 @@ void FSkeletalMeshGpuDynamicBufferProxy::NewFrame(const FNDISkeletalMesh_Instanc
 				{
 					const int32 BoneIndex = Section.BoneMap[m];
 					const FTransform& BoneTransform = BoneTransforms[BoneIndex];
-					const FMatrix BoneMatrix = SkelMesh->RefBasesInvMatrix[BoneIndex] * BoneTransform.ToMatrixWithScale();
+					const FMatrix BoneMatrix = SkelMesh->RefBasesInvMatrix.IsValidIndex(BoneIndex) ? SkelMesh->RefBasesInvMatrix[BoneIndex] * BoneTransform.ToMatrixWithScale() : BoneTransform.ToMatrixWithScale();
 					BoneMatrix.To3x4MatrixTranspose(&AllSectionsRefToLocalMatrices[Float4Count].X);
 					Float4Count += 3;
 				}
