Fixed crash when recompiling Niagara systems

FNiagaraShaderMapId objects were not hashed correctly because of a typo introduced with the loading optimization changes, so GIdToNiagaraShaderMap was broken. Using GetTypeHash() on FSHAHash::Hash is wrong, because that's just a uint8[20], so the function call resolves to the void* overload and it returns different values for the same hash. FSHAHash has its own GetTypeHash() overload, so we need to use that.

Also discovered that BaseCompileHash was always all-zeroes because GPU scripts don't have a graph representation. Made UNiagaraGraph::RebuildCachedCompileIds() use the spawn script hash in that case, so we have something to use as a hash for the keys in GIdToNiagaraShaderMap.

diff --git a/Engine/Plugins/FX/Niagara/Source/NiagaraEditor/Private/NiagaraGraph.cpp b/Engine/Plugins/FX/Niagara/Source/NiagaraEditor/Private/NiagaraGraph.cpp
--- a/Engine/Plugins/FX/Niagara/Source/NiagaraEditor/Private/NiagaraGraph.cpp
+++ b/Engine/Plugins/FX/Niagara/Source/NiagaraEditor/Private/NiagaraGraph.cpp
@@ -1690,11 +1690,10 @@ void UNiagaraGraph::RebuildCachedCompileIds(bool bForce)
 			GpuUsageInfo.BaseId = OldGpuInfo->BaseId;
 		}
 
-		// The gpu script has no graph representation so the hash is empty.  Changes to the gpu script are handled by collecting the compile hashes of it's dependencies.
-		TArray<uint8> DataHash;
-		DataHash.AddZeroed(20);
-		GpuUsageInfo.CompileHash = FNiagaraCompileHash(DataHash);
-		GpuUsageInfo.CompileHashFromGraph = FNiagaraCompileHash(DataHash);
+		// The GPU script has no graph representation, but we still need to fill in the hash, because it's used in the shader map ID.
+		// Just copy the hash from the spawn script.
+		GpuUsageInfo.CompileHash = ParticleSpawnUsageInfo->CompileHash;
+		GpuUsageInfo.CompileHashFromGraph = ParticleSpawnUsageInfo->CompileHashFromGraph;
 
 		NewUsageCache.Add(GpuUsageInfo);
 	}
diff --git a/Engine/Plugins/FX/Niagara/Source/NiagaraShader/Public/NiagaraShared.h b/Engine/Plugins/FX/Niagara/Source/NiagaraShader/Public/NiagaraShared.h
--- a/Engine/Plugins/FX/Niagara/Source/NiagaraShader/Public/NiagaraShared.h
+++ b/Engine/Plugins/FX/Niagara/Source/NiagaraShader/Public/NiagaraShared.h
@@ -245,7 +245,7 @@ public:
 
 	friend uint32 GetTypeHash(const FNiagaraShaderMapId& Ref)
 	{
-		return GetTypeHash(Ref.BaseCompileHash.Hash);
+		return GetTypeHash(Ref.BaseCompileHash);
 	}
 
 	SIZE_T GetSizeBytes() const
