Niagara - Prevent a crash by making the base data interface customization more safe when handling change events.

diff --git a/Engine/Plugins/FX/Niagara/Source/NiagaraEditorWidgets/Private/DetailCustomizations/NiagaraDataInterfaceDetails.cpp b/Engine/Plugins/FX/Niagara/Source/NiagaraEditorWidgets/Private/DetailCustomizations/NiagaraDataInterfaceDetails.cpp
--- a/Engine/Plugins/FX/Niagara/Source/NiagaraEditorWidgets/Private/DetailCustomizations/NiagaraDataInterfaceDetails.cpp
+++ b/Engine/Plugins/FX/Niagara/Source/NiagaraEditorWidgets/Private/DetailCustomizations/NiagaraDataInterfaceDetails.cpp
@@ -197,6 +197,7 @@ private:
 void FNiagaraDataInterfaceDetailsBase::CustomizeDetails(IDetailLayoutBuilder& DetailBuilder)
 {
 	Builder = &DetailBuilder;
+	PropertyUtilitiesWeak = DetailBuilder.GetPropertyUtilities();
 	TArray<TWeakObjectPtr<UObject>> SelectedObjects;
 	DetailBuilder.GetObjectsBeingCustomized(SelectedObjects);
 	check(SelectedObjects.Num() == 1);
@@ -213,10 +214,17 @@ void FNiagaraDataInterfaceDetailsBase::CustomizeDetails(IDetailLayoutBuilder& De
 
 void FNiagaraDataInterfaceDetailsBase::OnDataChanged() // need to only refresh errors, and all will be good
 {
-	if (Builder != nullptr)
+	TSharedPtr<IPropertyUtilities> PropertyUtilities = PropertyUtilitiesWeak.Pin();
+	bool bStillValid =
+		DataInterface.IsValid() &&
+		PropertyUtilities.IsValid() &&
+		PropertyUtilities->GetSelectedObjects().Num() == 1 &&
+		PropertyUtilities->GetSelectedObjects()[0].IsValid() &&
+		PropertyUtilities->GetSelectedObjects()[0].Get() == DataInterface.Get();
+
+	if (bStillValid)
 	{
 		int CurrentErrorCount = DataInterface->GetErrors().Num();
-		
 		if (CurrentErrorCount == 0)
 		{
 			ErrorsCategoryBuilder->SetCategoryVisibility(false);
diff --git a/Engine/Plugins/FX/Niagara/Source/NiagaraEditorWidgets/Private/DetailCustomizations/NiagaraDataInterfaceDetails.h b/Engine/Plugins/FX/Niagara/Source/NiagaraEditorWidgets/Private/DetailCustomizations/NiagaraDataInterfaceDetails.h
--- a/Engine/Plugins/FX/Niagara/Source/NiagaraEditorWidgets/Private/DetailCustomizations/NiagaraDataInterfaceDetails.h
+++ b/Engine/Plugins/FX/Niagara/Source/NiagaraEditorWidgets/Private/DetailCustomizations/NiagaraDataInterfaceDetails.h
@@ -10,6 +10,7 @@ class IDetailLayoutBuilder;
 class UNiagaraDataInterface;
 class IDetailCategoryBuilder;
 class FNiagaraDataInterfaceCustomNodeBuilder;
+class IPropertyUtilities;
 
 /** Base details customization for Niagara data interfaces. */
 class FNiagaraDataInterfaceDetailsBase : public IDetailCustomization
@@ -27,4 +28,5 @@ private:
 	TSharedPtr<FNiagaraDataInterfaceCustomNodeBuilder> CustomBuilder;
 	IDetailCategoryBuilder* ErrorsCategoryBuilder;
 	IDetailLayoutBuilder* Builder;
+	TWeakPtr<IPropertyUtilities> PropertyUtilitiesWeak;
 };
