Niagara - Fix style warnings.

diff --git a/Engine/Plugins/FX/Niagara/Source/NiagaraEditorWidgets/Private/NiagaraEditorWidgetsStyle.cpp b/Engine/Plugins/FX/Niagara/Source/NiagaraEditorWidgets/Private/NiagaraEditorWidgetsStyle.cpp
--- a/Engine/Plugins/FX/Niagara/Source/NiagaraEditorWidgets/Private/NiagaraEditorWidgetsStyle.cpp
+++ b/Engine/Plugins/FX/Niagara/Source/NiagaraEditorWidgets/Private/NiagaraEditorWidgetsStyle.cpp
@@ -67,19 +67,21 @@ TSharedRef< FSlateStyleSet > FNiagaraEditorWidgetsStyle::Create()
 	// Stack
 	Style->Set("NiagaraEditor.Stack.IconSize", FVector2D(18.0f, 18.0f));
 
-	FSlateFontInfo StackGroupFont = DEFAULT_FONT("Bold", 12);
+	FSlateFontInfo StackGroupFont = DEFAULT_FONT("Bold", 10);
 	FTextBlockStyle StackGroupText = FTextBlockStyle(NormalText)
 		.SetFont(StackGroupFont)
 		.SetColorAndOpacity(FLinearColor(1.0f, 1.0f, 1.0f, 1.0f));
+
+	Style->Set("NiagaraEditor.Stack.GroupText", StackGroupText);
 	
-	FEditableTextBoxStyle EditableGroupText = FEditableTextBoxStyle(NormalEditableTextBox)
+	FEditableTextBoxStyle StackEditableGroupText = FEditableTextBoxStyle(NormalEditableTextBox)
 		.SetFont(StackGroupFont);
 
-	FInlineEditableTextBlockStyle StackEditableGroupText = FInlineEditableTextBlockStyle()
-		.SetEditableTextBoxStyle(EditableGroupText)
+	FInlineEditableTextBlockStyle StackInlineEditableGroupText = FInlineEditableTextBlockStyle()
+		.SetEditableTextBoxStyle(StackEditableGroupText)
 		.SetTextStyle(StackGroupText); 
 
-	Style->Set("NiagaraEditor.Stack.GroupText", StackEditableGroupText);
+	Style->Set("NiagaraEditor.Stack.EditableGroupText", StackInlineEditableGroupText);
 
 	FSlateFontInfo StackDefaultFont = DEFAULT_FONT("Regular", 10);
 	FTextBlockStyle StackDefaultText = FTextBlockStyle(NormalText)
@@ -105,19 +107,20 @@ TSharedRef< FSlateStyleSet > FNiagaraEditorWidgetsStyle::Create()
 	FSlateFontInfo StackItemFont = DEFAULT_FONT("Regular", 11);
 	FTextBlockStyle StackItemText = FTextBlockStyle(NormalText)
 		.SetFont(StackItemFont);
+	Style->Set("NiagaraEditor.Stack.ItemText", StackItemText);
 	
-	FEditableTextBoxStyle EditableItemText = FEditableTextBoxStyle(NormalEditableTextBox)
+	FEditableTextBoxStyle StackEditableItemText = FEditableTextBoxStyle(NormalEditableTextBox)
 		.SetFont(StackItemFont);
 
-	FInlineEditableTextBlockStyle StackEditableItemText = FInlineEditableTextBlockStyle()
+	FInlineEditableTextBlockStyle StackInlineEditableItemText = FInlineEditableTextBlockStyle()
 		.SetTextStyle(StackItemText)
-		.SetEditableTextBoxStyle(EditableItemText);
-	Style->Set("NiagaraEditor.Stack.ItemText", StackEditableItemText);
+		.SetEditableTextBoxStyle(StackEditableItemText);
+	Style->Set("NiagaraEditor.Stack.EditableItemText", StackInlineEditableItemText);
 
-	FSlateFontInfo TypeNameFont = DEFAULT_FONT("Regular", 9);
-	FTextBlockStyle TypeNameText = FTextBlockStyle(NormalText)
-		.SetFont(TypeNameFont);
-	Style->Set("NiagaraEditor.Stack.TypeNameText", TypeNameText);
+	FSlateFontInfo StackSubduedItemFont = DEFAULT_FONT("Regular", 9);
+	FTextBlockStyle StackSubduedItemText = FTextBlockStyle(NormalText)
+		.SetFont(StackSubduedItemFont);
+	Style->Set("NiagaraEditor.Stack.SubduedItemText", StackSubduedItemText);
 
 	FSlateFontInfo SystemOverviewListHeaderFont = DEFAULT_FONT("Bold", 12);
 	FTextBlockStyle SystemOverviewListHeaderText = FTextBlockStyle(NormalText)
diff --git a/Engine/Plugins/FX/Niagara/Source/NiagaraEditorWidgets/Private/SNiagaraStackEntryWidget.cpp b/Engine/Plugins/FX/Niagara/Source/NiagaraEditorWidgets/Private/SNiagaraStackEntryWidget.cpp
--- a/Engine/Plugins/FX/Niagara/Source/NiagaraEditorWidgets/Private/SNiagaraStackEntryWidget.cpp
+++ b/Engine/Plugins/FX/Niagara/Source/NiagaraEditorWidgets/Private/SNiagaraStackEntryWidget.cpp
@@ -11,13 +11,14 @@
 
 #define LOCTEXT_NAMESPACE "SNiagaraStackEntryWidget"
 
-void SNiagaraStackDisplayName::Construct(const FArguments& InArgs, UNiagaraStackEntry& InStackEntry, UNiagaraStackViewModel& InStackViewModel, FName InTextStyleName)
+void SNiagaraStackDisplayName::Construct(const FArguments& InArgs, UNiagaraStackEntry& InStackEntry, UNiagaraStackViewModel& InStackViewModel)
 {
 	StackEntryItem = &InStackEntry;
 	StackViewModel = &InStackViewModel;
-	TextStyleName = InTextStyleName;
 
-	TypeNameStyle = InArgs._TypeNameStyle;
+	NameStyle = InArgs._NameStyle;
+	EditableNameStyle = InArgs._EditableNameStyle;
+	OriginalNameStyle = InArgs._OriginalNameStyle;
 
 	TAttribute<FText> EntryToolTipText;
 	EntryToolTipText.Bind(this, &SNiagaraStackDisplayName::GetEntryToolTipText);
@@ -48,7 +49,6 @@ SNiagaraStackDisplayName::~SNiagaraStackDisplayName()
 TSharedRef<SWidget> SNiagaraStackDisplayName::ConstructChildren()
 {
 	EditableTextBlock.Reset();
-	const FInlineEditableTextBlockStyle& EditableStyle = FNiagaraEditorWidgetsStyle::Get().GetWidgetStyle<FInlineEditableTextBlockStyle>(TextStyleName);
 	TArray<TSharedRef<SWidget>> NameWidgets;
 
 	// First check to see if we need to insert the emitter name.
@@ -65,7 +65,7 @@ TSharedRef<SWidget> SNiagaraStackDisplayName::ConstructChildren()
 	{
 		TSharedPtr<UNiagaraStackViewModel::FTopLevelViewModel> TopLevelViewModel = StackViewModel->GetTopLevelViewModelForEntry(*StackEntryItem);
 		NameWidgets.Add(SNew(STextBlock)
-			.TextStyle(FNiagaraEditorWidgetsStyle::Get(), TextStyleName)
+			.TextStyle(NameStyle)
 			.Text(this, &SNiagaraStackDisplayName::GetTopLevelDisplayName, TWeakPtr<UNiagaraStackViewModel::FTopLevelViewModel>(TopLevelViewModel))
 			.HighlightText_UObject(StackViewModel, &UNiagaraStackViewModel::GetCurrentSearchText)
 			.ColorAndOpacity(this, &SNiagaraStackDisplayName::GetTextColorForSearch, FSlateColor::UseForeground()));
@@ -77,7 +77,7 @@ TSharedRef<SWidget> SNiagaraStackDisplayName::ConstructChildren()
 	{
 		// If the entry can be renamed we need an editable text block.
 		NameWidgets.Add(SAssignNew(EditableTextBlock, SInlineEditableTextBlock)
-			.Style(&EditableStyle)
+			.Style(EditableNameStyle)
 			.Text(this, &SNiagaraStackDisplayName::GetEntryDisplayName)
 			.HighlightText_UObject(StackViewModel, &UNiagaraStackViewModel::GetCurrentSearchText)
 			.ColorAndOpacity(this, &SNiagaraStackEntryWidget::GetTextColorForSearch, FSlateColor::UseForeground())
@@ -87,7 +87,7 @@ TSharedRef<SWidget> SNiagaraStackDisplayName::ConstructChildren()
 	{
 		// Otherwise add a regular text block.
 		NameWidgets.Add(SNew(STextBlock)
-			.TextStyle(&EditableStyle.TextStyle)
+			.TextStyle(NameStyle)
 			.Text(this, &SNiagaraStackDisplayName::GetEntryDisplayName)
 			.HighlightText_UObject(StackViewModel, &UNiagaraStackViewModel::GetCurrentSearchText)
 			.ColorAndOpacity(this, &SNiagaraStackDisplayName::GetTextColorForSearch, FSlateColor::UseForeground()));
@@ -97,7 +97,7 @@ TSharedRef<SWidget> SNiagaraStackDisplayName::ConstructChildren()
 	if(StackEntryItem->GetAlternateDisplayName().IsSet())
 	{
 		NameWidgets.Add(SNew(STextBlock)
-			.TextStyle(TypeNameStyle)
+			.TextStyle(OriginalNameStyle)
 			.Text(this, &SNiagaraStackDisplayName::GetOriginalName)
 			.HighlightText_UObject(StackViewModel, &UNiagaraStackViewModel::GetCurrentSearchText)
 			.ColorAndOpacity(this, &SNiagaraStackDisplayName::GetTextColorForSearch, FSlateColor::UseSubduedForeground()));
diff --git a/Engine/Plugins/FX/Niagara/Source/NiagaraEditorWidgets/Private/SNiagaraStackEntryWidget.h b/Engine/Plugins/FX/Niagara/Source/NiagaraEditorWidgets/Private/SNiagaraStackEntryWidget.h
--- a/Engine/Plugins/FX/Niagara/Source/NiagaraEditorWidgets/Private/SNiagaraStackEntryWidget.h
+++ b/Engine/Plugins/FX/Niagara/Source/NiagaraEditorWidgets/Private/SNiagaraStackEntryWidget.h
@@ -6,6 +6,7 @@
 #include "Widgets/DeclarativeSyntaxSupport.h"
 #include "ViewModels/Stack/NiagaraStackViewModel.h"
 #include "Framework/SlateDelegates.h"
+#include "NiagaraEditorWidgetsStyle.h"
 
 class UNiagaraStackViewModel;
 class UNiagaraStackEntry;
@@ -31,13 +32,18 @@ class SNiagaraStackDisplayName : public SNiagaraStackEntryWidget
 {
 public:
 	SLATE_BEGIN_ARGS(SNiagaraStackDisplayName)
+		: _NameStyle(&FNiagaraEditorWidgetsStyle::Get().GetWidgetStyle<FTextBlockStyle>("NiagaraEditor.Stack.ItemText"))
+		, _EditableNameStyle(&FNiagaraEditorWidgetsStyle::Get().GetWidgetStyle<FInlineEditableTextBlockStyle>("NiagaraEditor.Stack.EditableItemText"))
+		, _OriginalNameStyle(&FNiagaraEditorWidgetsStyle::Get().GetWidgetStyle<FTextBlockStyle>("NiagaraEditor.Stack.SubduedItemText"))
 	{
 		_Clipping = EWidgetClipping::OnDemand;
 	}
-		SLATE_STYLE_ARGUMENT(FTextBlockStyle, TypeNameStyle)
+		SLATE_STYLE_ARGUMENT(FTextBlockStyle, NameStyle)
+		SLATE_STYLE_ARGUMENT(FInlineEditableTextBlockStyle, EditableNameStyle)
+		SLATE_STYLE_ARGUMENT(FTextBlockStyle, OriginalNameStyle)
 	SLATE_END_ARGS();
 
-	void Construct(const FArguments& InArgs, UNiagaraStackEntry& InStackEntry, UNiagaraStackViewModel& InStackViewModel, FName InTextStyleName);
+	void Construct(const FArguments& InArgs, UNiagaraStackEntry& InStackEntry, UNiagaraStackViewModel& InStackViewModel);
 
 	~SNiagaraStackDisplayName();
 
@@ -65,9 +71,9 @@ private:
 	void EntryNameTextCommitted(const FText& InText, ETextCommit::Type CommitInfo);
 
 private:
-	FName TextStyleName;
-
-	const FTextBlockStyle* TypeNameStyle;
+	const FTextBlockStyle* NameStyle;
+	const FInlineEditableTextBlockStyle* EditableNameStyle;
+	const FTextBlockStyle* OriginalNameStyle;
 
 	TSharedPtr<SBox> Container;
 
diff --git a/Engine/Plugins/FX/Niagara/Source/NiagaraEditorWidgets/Private/Stack/SNiagaraStackItem.cpp b/Engine/Plugins/FX/Niagara/Source/NiagaraEditorWidgets/Private/Stack/SNiagaraStackItem.cpp
--- a/Engine/Plugins/FX/Niagara/Source/NiagaraEditorWidgets/Private/Stack/SNiagaraStackItem.cpp
+++ b/Engine/Plugins/FX/Niagara/Source/NiagaraEditorWidgets/Private/Stack/SNiagaraStackItem.cpp
@@ -43,8 +43,7 @@ void SNiagaraStackItem::Construct(const FArguments& InArgs, UNiagaraStackItem& I
 		.Padding(2, 0, 2, 0)
 		.VAlign(VAlign_Center)
 		[
-			SAssignNew(DisplayNameWidget, SNiagaraStackDisplayName, InItem, *InStackViewModel, "NiagaraEditor.Stack.ItemText")
-			.TypeNameStyle(FNiagaraEditorWidgetsStyle::Get(), "NiagaraEditor.Stack.TypeNameText")
+			SAssignNew(DisplayNameWidget, SNiagaraStackDisplayName, InItem, *InStackViewModel)
 		];
 
 	// Allow derived classes to add additional widgets.
diff --git a/Engine/Plugins/FX/Niagara/Source/NiagaraEditorWidgets/Private/Stack/SNiagaraStackItemGroup.cpp b/Engine/Plugins/FX/Niagara/Source/NiagaraEditorWidgets/Private/Stack/SNiagaraStackItemGroup.cpp
--- a/Engine/Plugins/FX/Niagara/Source/NiagaraEditorWidgets/Private/Stack/SNiagaraStackItemGroup.cpp
+++ b/Engine/Plugins/FX/Niagara/Source/NiagaraEditorWidgets/Private/Stack/SNiagaraStackItemGroup.cpp
@@ -28,7 +28,8 @@ void SNiagaraStackItemGroup::Construct(const FArguments& InArgs, UNiagaraStackIt
 		.VAlign(VAlign_Center)
 		.Padding(2, 0, 0, 0)
 		[
-			SNew(SNiagaraStackDisplayName, InGroup, *InStackViewModel, "NiagaraEditor.Stack.GroupText")
+			SNew(SNiagaraStackDisplayName, InGroup, *InStackViewModel)
+			.NameStyle(FNiagaraEditorWidgetsStyle::Get(), "NiagaraEditor.Stack.GroupText")
 		]
 		// Delete group button
 		+ SHorizontalBox::Slot()
