Updating NDI Audio Spectrum

diff --git a/Engine/Plugins/FX/Niagara/Source/Niagara/Classes/NiagaraDataInterfaceAudio.h b/Engine/Plugins/FX/Niagara/Source/Niagara/Classes/NiagaraDataInterfaceAudio.h
--- a/Engine/Plugins/FX/Niagara/Source/Niagara/Classes/NiagaraDataInterfaceAudio.h
+++ b/Engine/Plugins/FX/Niagara/Source/Niagara/Classes/NiagaraDataInterfaceAudio.h
@@ -67,6 +67,9 @@ private:
 /** Niagara data interface proxy for audio submix listener. */
 struct FNiagaraDataInterfaceProxySubmix : public FNiagaraDataInterfaceProxy 
 {
+	FNiagaraDataInterfaceProxySubmix(const FNiagaraDataInterfaceProxySubmix& Other) = delete;
+	FNiagaraDataInterfaceProxySubmix& operator=(const FNiagaraDataInterfaceProxySubmix& Other) = delete;
+
 	/** Construct a FNiagaraDataInterfaceProxySubmix
 	 *
 	 *  @params InNumSamplesToBuffer - Number of audio sampels to buffer internally. 
@@ -99,11 +102,12 @@ struct FNiagaraDataInterfaceProxySubmix : public FNiagaraDataInterfaceProxy
 	virtual int32 PerInstanceDataPassedToRenderThreadSize() const override;
 
 private:
+
 	void AddSubmixListener(const Audio::FDeviceId InDeviceId);
 	void RemoveSubmixListener(const Audio::FDeviceId InDeviceId);
 
 	void RegisterToAllAudioDevices();
-	void UnregisterFromAllAudioDevices(USoundSubmix* Submix);
+	void UnregisterFromAllAudioDevices();
 
 	void OnNewDeviceCreated(Audio::FDeviceId InDeviceId);
 	void OnDeviceDestroyed(Audio::FDeviceId InDeviceId);
@@ -136,6 +140,8 @@ public:
 	USoundSubmix* Submix;
 
 
+	virtual bool Equals(const UNiagaraDataInterface* Other) const override;
+
 #if WITH_EDITOR
 	virtual void PostEditChangeProperty(struct FPropertyChangedEvent& PropertyChangedEvent) override;
 #endif // WITH_EDITOR
diff --git a/Engine/Plugins/FX/Niagara/Source/Niagara/Classes/NiagaraDataInterfaceAudioSpectrum.h b/Engine/Plugins/FX/Niagara/Source/Niagara/Classes/NiagaraDataInterfaceAudioSpectrum.h
--- a/Engine/Plugins/FX/Niagara/Source/Niagara/Classes/NiagaraDataInterfaceAudioSpectrum.h
+++ b/Engine/Plugins/FX/Niagara/Source/Niagara/Classes/NiagaraDataInterfaceAudioSpectrum.h
@@ -42,6 +42,8 @@ struct FNiagaraDataInterfaceProxySpectrum : public FNiagaraDataInterfaceProxySub
 	 */
 	float GetSpectrumValue(float InNormalizedPositionInSpectrum, int32 InChannelIndex);
 
+	int32 GetNumBands() const;
+
 	/** Updates the minimum and maximum frequency of the CQT on the render thread. */
 	void UpdateCQT(float InMinimumFrequency, float InMaximumFrequency, int32 InNumBands);
 
@@ -93,7 +95,7 @@ private:
 
 	float MinimumFrequency;
 	float MaximumFrequency;
-	int32 NumBands;
+	TAtomic<int32> NumBands;
 	float NoiseFloorDb;
 
 	int32 NumChannels;
@@ -138,6 +140,7 @@ public:
 	// Global variable prefixes
 	static const FString GetSpectrumName;
 	static const FString NumChannelsName;
+	static const FString ResolutionName;
 	
 	/** The number of spectrum samples to pass to the GPU */
 	UPROPERTY(EditAnywhere, Category = "Spectrum", meta = (ClampMin = "16", ClampMax = "1024") )
@@ -177,10 +180,13 @@ public:
 
 	virtual void GetParameterDefinitionHLSL(const FNiagaraDataInterfaceGPUParamInfo& ParamInfo, FString& OutHLSL) override;
 
+	virtual bool Equals(const UNiagaraDataInterface* Other) const override;
+
 #if WITH_EDITOR
 	virtual void PostEditChangeProperty(struct FPropertyChangedEvent& PropertyChangedEvent) override;
 #endif // WITH_EDITOR
 
+
 	virtual void PostInitProperties() override;
 	virtual void PostLoad() override;
 
diff --git a/Engine/Plugins/FX/Niagara/Source/Niagara/Private/NiagaraDataInterfaceAudio.cpp b/Engine/Plugins/FX/Niagara/Source/Niagara/Private/NiagaraDataInterfaceAudio.cpp
--- a/Engine/Plugins/FX/Niagara/Source/Niagara/Private/NiagaraDataInterfaceAudio.cpp
+++ b/Engine/Plugins/FX/Niagara/Source/Niagara/Private/NiagaraDataInterfaceAudio.cpp
@@ -137,7 +137,7 @@ void FNiagaraDataInterfaceProxySubmix::RegisterToAllAudioDevices()
 	}
 }
 
-void FNiagaraDataInterfaceProxySubmix::UnregisterFromAllAudioDevices(USoundSubmix* Submix)
+void FNiagaraDataInterfaceProxySubmix::UnregisterFromAllAudioDevices()
 {
 	if (FAudioDeviceManager* DeviceManager = FAudioDeviceManager::Get())
 	{
@@ -153,7 +153,7 @@ void FNiagaraDataInterfaceProxySubmix::OnUpdateSubmix(USoundSubmix* Submix)
 {
 	if (bIsSubmixListenerRegistered)
 	{
-		UnregisterFromAllAudioDevices(Submix);
+		UnregisterFromAllAudioDevices();
 	}
 
 	SubmixRegisteredTo = Submix;
@@ -234,7 +234,7 @@ void FNiagaraDataInterfaceProxySubmix::OnBeginDestroy()
 {
 	if (bIsSubmixListenerRegistered)
 	{
-		UnregisterFromAllAudioDevices(SubmixRegisteredTo);
+		UnregisterFromAllAudioDevices();
 	}
 }
 
@@ -273,6 +273,16 @@ UNiagaraDataInterfaceAudioSubmix::UNiagaraDataInterfaceAudioSubmix(FObjectInitia
 	Proxy = MakeUnique<FNiagaraDataInterfaceProxySubmix>(DefaultNumSamplesToBuffer);
 }
 
+bool UNiagaraDataInterfaceAudioSubmix::Equals(const UNiagaraDataInterface* Other) const
+{
+	bool bIsEqual = Super::Equals(Other);
+
+	const UNiagaraDataInterfaceAudioSubmix* OtherSubmix = CastChecked<const UNiagaraDataInterfaceAudioSubmix>(Other);
+
+	bIsEqual &= OtherSubmix->Submix == Submix;
+
+	return bIsEqual;
+}
 
 #if WITH_EDITOR
 void UNiagaraDataInterfaceAudioSubmix::PostEditChangeProperty(struct FPropertyChangedEvent& PropertyChangedEvent)
@@ -323,6 +333,7 @@ bool UNiagaraDataInterfaceAudioSubmix::CopyToInternal(UNiagaraDataInterface* Des
 	if (CastedDestination)
 	{
 		CastedDestination->Submix = Submix;
+		CastedDestination->GetProxyAs<FNiagaraDataInterfaceProxySubmix>()->OnUpdateSubmix(Submix);
 	}
 	
 	return true;
diff --git a/Engine/Plugins/FX/Niagara/Source/Niagara/Private/NiagaraDataInterfaceAudioSpectrum.cpp b/Engine/Plugins/FX/Niagara/Source/Niagara/Private/NiagaraDataInterfaceAudioSpectrum.cpp
--- a/Engine/Plugins/FX/Niagara/Source/Niagara/Private/NiagaraDataInterfaceAudioSpectrum.cpp
+++ b/Engine/Plugins/FX/Niagara/Source/Niagara/Private/NiagaraDataInterfaceAudioSpectrum.cpp
@@ -76,6 +76,12 @@ float FNiagaraDataInterfaceProxySpectrum::GetSpectrumValue(float InNormalizedPos
 	return FMath::Lerp<float>(Buffer[LowerIndex], Buffer[UpperIndex], Alpha);
 }
 
+int32 FNiagaraDataInterfaceProxySpectrum::GetNumBands() const
+{
+	check(IsInRenderingThread());
+	return NumBands;
+}
+
 
 void FNiagaraDataInterfaceProxySpectrum::UpdateCQT(float InMinimumFrequency, float InMaximumFrequency, int32 InNumBands)
 {
@@ -587,6 +593,8 @@ void UNiagaraDataInterfaceAudioSpectrum::GetFunctions(TArray<FNiagaraFunctionSig
 		NumChannelsSignature.bRequiresContext = false;
 		OutFunctions.Add(NumChannelsSignature);
 	}
+
+
 }
 
 DEFINE_NDI_DIRECT_FUNC_BINDER(UNiagaraDataInterfaceAudioSpectrum, GetSpectrumValue);
@@ -648,7 +656,7 @@ bool UNiagaraDataInterfaceAudioSpectrum::GetFunctionHLSL(const FNiagaraDataInter
 		TMap<FString, FStringFormatArg> ArgsBounds = {
 			{TEXT("FunctionName"), FStringFormatArg(FunctionInfo.InstanceName)},
 			{TEXT("ChannelCount"), FStringFormatArg(NumChannelsName + ParamInfo.DataInterfaceHLSLSymbol)},
-			{TEXT("SpectrumResolution"), FStringFormatArg(Resolution)},
+			{TEXT("SpectrumResolution"), FStringFormatArg(ResolutionName + ParamInfo.DataInterfaceHLSLSymbol)},
 			{TEXT("SpectrumBuffer"), FStringFormatArg(GetSpectrumName + ParamInfo.DataInterfaceHLSLSymbol)},
 		};
 		OutHLSL += FString::Format(FormatBounds, ArgsBounds);
@@ -676,6 +684,20 @@ bool UNiagaraDataInterfaceAudioSpectrum::GetFunctionHLSL(const FNiagaraDataInter
 	}
 }
 
+bool UNiagaraDataInterfaceAudioSpectrum::Equals(const UNiagaraDataInterface* Other) const
+{
+	bool bIsEqual = Super::Equals(Other);
+
+	const UNiagaraDataInterfaceAudioSpectrum* OtherSpectrum = CastChecked<const UNiagaraDataInterfaceAudioSpectrum>(Other);
+
+	bIsEqual &= OtherSpectrum->Resolution == Resolution;
+	bIsEqual &= OtherSpectrum->MinimumFrequency == MinimumFrequency;
+	bIsEqual &= OtherSpectrum->MaximumFrequency == MaximumFrequency;
+	bIsEqual &= OtherSpectrum->NoiseFloorDb == NoiseFloorDb;
+
+	return bIsEqual;
+}
+
 void UNiagaraDataInterfaceAudioSpectrum::GetParameterDefinitionHLSL(const FNiagaraDataInterfaceGPUParamInfo& ParamInfo, FString& OutHLSL)
 {
 	Super::GetParameterDefinitionHLSL(ParamInfo, OutHLSL);
@@ -683,11 +705,13 @@ void UNiagaraDataInterfaceAudioSpectrum::GetParameterDefinitionHLSL(const FNiaga
 	static const TCHAR *FormatDeclarations = TEXT(R"(				
 		Buffer<float> {GetSpectrumName};
 		int {NumChannelsName};
+		int {ResolutionName};
 
 	)");
 	TMap<FString, FStringFormatArg> ArgsDeclarations = {
 		{ TEXT("GetSpectrumName"),    FStringFormatArg(GetSpectrumName + ParamInfo.DataInterfaceHLSLSymbol) },
 		{ TEXT("NumChannelsName"),    FStringFormatArg(NumChannelsName + ParamInfo.DataInterfaceHLSLSymbol) },
+		{ TEXT("ResolutionName"),     FStringFormatArg(ResolutionName  + ParamInfo.DataInterfaceHLSLSymbol) },
 	};
 	OutHLSL += FString::Format(FormatDeclarations, ArgsDeclarations);
 }
@@ -699,6 +723,9 @@ struct FNiagaraDataInterfaceParametersCS_AudioSpectrum : public FNiagaraDataInte
 	void Bind(const FNiagaraDataInterfaceGPUParamInfo& ParameterInfo, const class FShaderParameterMap& ParameterMap)
 	{
 		NumChannels.Bind(ParameterMap, *(UNiagaraDataInterfaceAudioSpectrum::NumChannelsName + ParameterInfo.DataInterfaceHLSLSymbol));
+
+		Resolution.Bind(ParameterMap, *(UNiagaraDataInterfaceAudioSpectrum::ResolutionName + ParameterInfo.DataInterfaceHLSLSymbol));
+
 		SpectrumBuffer.Bind(ParameterMap, *(UNiagaraDataInterfaceAudioSpectrum::GetSpectrumName + ParameterInfo.DataInterfaceHLSLSymbol));
 	}
 
@@ -712,10 +739,12 @@ struct FNiagaraDataInterfaceParametersCS_AudioSpectrum : public FNiagaraDataInte
 		FReadBuffer& SpectrumBufferSRV = NDI->ComputeAndPostSRV();
 
 		SetShaderValue(RHICmdList, ComputeShaderRHI, NumChannels, NDI->GetNumChannels());
+		SetShaderValue(RHICmdList, ComputeShaderRHI, Resolution, NDI->GetNumBands());
 		RHICmdList.SetShaderResourceViewParameter(ComputeShaderRHI, SpectrumBuffer.GetBaseIndex(), SpectrumBufferSRV.SRV);
 	}
 
 	LAYOUT_FIELD(FShaderParameter, NumChannels);
+	LAYOUT_FIELD(FShaderParameter, Resolution);
 	LAYOUT_FIELD(FShaderResourceParameter,SpectrumBuffer);
 };
 
@@ -778,7 +807,10 @@ void UNiagaraDataInterfaceAudioSpectrum::PostLoad()
 
 bool UNiagaraDataInterfaceAudioSpectrum::CopyToInternal(UNiagaraDataInterface* Destination) const
 {
-	Super::CopyToInternal(Destination);
+	if (!Super::CopyToInternal(Destination))
+	{
+		return false;
+	}
 
 	UNiagaraDataInterfaceAudioSpectrum* CastedDestination = Cast<UNiagaraDataInterfaceAudioSpectrum>(Destination);
 
@@ -788,6 +820,12 @@ bool UNiagaraDataInterfaceAudioSpectrum::CopyToInternal(UNiagaraDataInterface* D
 		CastedDestination->MinimumFrequency = MinimumFrequency;
 		CastedDestination->MaximumFrequency = MaximumFrequency;
 		CastedDestination->NoiseFloorDb = NoiseFloorDb;
+
+		if (!CastedDestination->HasAnyFlags(RF_ClassDefaultObject))
+		{
+			CastedDestination->GetProxyAs<FNiagaraDataInterfaceProxySpectrum>()->UpdateCQT(MinimumFrequency, MaximumFrequency, Resolution);
+			CastedDestination->GetProxyAs<FNiagaraDataInterfaceProxySpectrum>()->UpdateNoiseFloor(NoiseFloorDb);
+		}
 	}
 	
 	return true;
@@ -801,5 +839,7 @@ const FName UNiagaraDataInterfaceAudioSpectrum::GetNumChannelsFunctionName("GetN
 // Global variable prefixes, used in HLSL parameter declarations.
 const FString UNiagaraDataInterfaceAudioSpectrum::GetSpectrumName(TEXT("SpectrumBuffer_"));
 const FString UNiagaraDataInterfaceAudioSpectrum::NumChannelsName(TEXT("NumChannels_"));
+const FString UNiagaraDataInterfaceAudioSpectrum::ResolutionName(TEXT("Resolution_"));
+
 
 #undef LOCTEXT_NAMESPACE
diff --git a/Engine/Source/Runtime/SignalProcessing/Private/MultithreadedPatching.cpp b/Engine/Source/Runtime/SignalProcessing/Private/MultithreadedPatching.cpp
--- a/Engine/Source/Runtime/SignalProcessing/Private/MultithreadedPatching.cpp
+++ b/Engine/Source/Runtime/SignalProcessing/Private/MultithreadedPatching.cpp
@@ -240,6 +240,8 @@ namespace Audio
 	int32 FPatchMixer::MaxNumberOfSamplesThatCanBePopped()
 	{
 		FScopeLock ScopeLock(&CurrentPatchesCriticalSection);
+
+		CleanUpDisconnectedPatches();
 		ConnectNewPatches();
 
 		// Iterate through our inputs and see which input has the least audio buffered.
@@ -289,6 +291,18 @@ namespace Audio
 	{
 		FScopeLock PendingInputDeletionScopeLock(&InputDeletionCriticalSection);
 
+		 // Callers of this function must have CurrentPatchesCritialSection locked so that 
+		 // this is not causing a race condition.
+		for (const FPatchOutputStrongPtr& Patch : CurrentInputs)
+		{
+			check(Patch.IsValid());
+
+			if (Patch->IsInputStale())
+			{
+				DisconnectedInputs.Add(Patch->PatchID);
+			}
+		}
+
 		// Iterate through all of the PatchIDs we need to clean up.
 		for (const int32& PatchID : DisconnectedInputs)
 		{
diff --git a/Engine/Source/Runtime/SignalProcessing/Public/DSP/MultithreadedPatching.h b/Engine/Source/Runtime/SignalProcessing/Public/DSP/MultithreadedPatching.h
--- a/Engine/Source/Runtime/SignalProcessing/Public/DSP/MultithreadedPatching.h
+++ b/Engine/Source/Runtime/SignalProcessing/Public/DSP/MultithreadedPatching.h
@@ -144,7 +144,8 @@ namespace Audio
 		/** Called within PopAudio. Flushes the PendingNewPatches array into CurrentPatches. During this function, AddNewPatch is blocked. */
 		void ConnectNewPatches();
 
-		/** Called within PopAudio. Removes PendingTapsToDelete from CurrentPatches and ConnectNewPatches. During this function, RemoveTap and AddNewPatch are blocked. */
+		/** Called within PopAudio and MaxNumberOfSamplesThatCanBePopped. Removes PendingTapsToDelete from CurrentPatches and ConnectNewPatches. 
+		 * During this function, RemoveTap and AddNewPatch are blocked. Callers of this function must have CurrentPatchesCritialSection locked. */
 		void CleanUpDisconnectedPatches();
 
 		/** New taps are added here in AddNewPatch, and then are moved to CurrentPatches in ConnectNewPatches. */
