Niagara crash reporter integration.

diff --git a/Engine/Plugins/FX/Niagara/Source/Niagara/Private/NiagaraCrashReporterHandler.cpp b/Engine/Plugins/FX/Niagara/Source/Niagara/Private/NiagaraCrashReporterHandler.cpp
new file mode 100644
--- /dev/null
+++ b/Engine/Plugins/FX/Niagara/Source/Niagara/Private/NiagaraCrashReporterHandler.cpp
@@ -0,0 +1,158 @@
+// Copyright Epic Games, Inc. All Rights Reserved.
+
+#include "NiagaraCrashReporterHandler.h"
+
+#include "CoreMinimal.h"
+#include "HAL\ThreadManager.h"
+#include "GenericPlatform\GenericPlatformCrashContext.h"
+#include "Misc\ScopeLock.h"
+#include "NiagaraSystemInstance.h"
+#include "NiagaraSystemSimulation.h"
+#include "NiagaraComponent.h"
+
+static int32 GbEnableNiagaraCRHandler = 0;
+static FAutoConsoleVariableRef CVarEnableNiagaraCRHandler(
+	TEXT("fx.EnableNiagaraCRHandler"),
+	GbEnableNiagaraCRHandler,
+	TEXT("If > 0 Niagara will push some state into the crash reporter. This is not free so should not be used unless actively tracking a crash in the wild. Even then it should only be enabled on the platforms needed etc. \n"),
+	FConsoleVariableDelegate::CreateStatic(&FNiagaraCrashReporterHandler::OnEnabledChanged),
+	ECVF_Default);
+
+FCriticalSection FNiagaraCrashReporterHandler::CritSec;
+TArray<FNiagaraCrashReporterHandler*> FNiagaraCrashReporterHandler::AllHandlers;
+
+FNiagaraCrashReporterHandler::FNiagaraCrashReporterHandler()
+{
+	FString ThreadName;
+	if (ThreadId == GGameThreadId)
+	{
+		ThreadName = TEXT("GameThread");
+	}
+	else if (ThreadId == GRenderThreadId)
+	{
+		ThreadName = TEXT("RenderThread");
+	}
+	else
+	{
+		ThreadName = FThreadManager::Get().GetThreadName(ThreadId);
+	}
+
+	Name = FString::Printf(TEXT("NiagaraCRInfo | %s(%u)"), *ThreadName, ThreadId);
+	
+	NoneInfoString = FString::Printf(TEXT("%s | None |"), *Name);
+
+	FScopeLock Lock(&CritSec);
+	AllHandlers.AddUnique(this);
+}
+
+FNiagaraCrashReporterHandler::~FNiagaraCrashReporterHandler()
+{
+	FScopeLock Lock(&CritSec);
+	AllHandlers.Remove(this);
+}
+
+void FNiagaraCrashReporterHandler::Clear()
+{
+	ScopeInfoStack.Empty();
+
+	FScopeLock ScopeLock(&CritSec);
+	FString Empty;
+	UE_LOG(LogNiagara, Log, TEXT("SetEngineData: %s"), *Empty);
+	FGenericCrashContext::SetEngineData(Name, Empty);
+}
+
+void FNiagaraCrashReporterHandler::OnEnabledChanged(IConsoleVariable* InVariable)
+{
+	if (InVariable->GetInt() == 0)
+	{
+		for (FNiagaraCrashReporterHandler* Handler : AllHandlers)
+		{
+			Handler->Clear();
+		}
+	}
+}
+
+void FNiagaraCrashReporterHandler::SetInfo(const FString& Info)
+{
+	FScopeLock ScopeLock(&CritSec);
+	//UE_LOG(LogNiagara, Log, TEXT("SetEngineData: %s"), *Info);
+	FGenericCrashContext::SetEngineData(Name, Info);
+}
+
+void FNiagaraCrashReporterHandler::PushInfo(FNiagaraSystemInstance* Inst)
+{
+	check(Inst && Inst->GetComponent() && Inst->GetComponent()->GetAsset());
+	UNiagaraComponent* Comp = Inst->GetComponent();
+	USceneComponent* AttachParent = Comp->GetAttachParent();
+
+	PushInfo(FString::Printf(TEXT("%s | SystemInstance | System: %s | bSolo: %s | Component: %s | AttachedTo: %s |"), 
+	*Name, *Comp->GetAsset()->GetFullName(), Inst->IsSolo() ? TEXT("true") : TEXT("false"), *Comp->GetFullName(), AttachParent ? *AttachParent->GetFullName() : TEXT("nullptr")));
+
+	//TODO: More info?
+}
+
+void FNiagaraCrashReporterHandler::PushInfo(FNiagaraSystemSimulation* SystemSim)
+{
+	check(SystemSim && SystemSim->GetSystem());
+	
+	UNiagaraSystem* System = SystemSim->GetSystem();
+
+	PushInfo(FString::Printf(TEXT("%s | SystemSimulation | System: %s | bSolo: %s |"), *Name, *System->GetFullName(), SystemSim->GetIsSolo() ? TEXT("true") : TEXT("false")));
+
+	//TODO: More info?
+}
+
+void FNiagaraCrashReporterHandler::PushInfo(const FString& Info)
+{
+	ScopeInfoStack.Push(Info);
+	SetInfo(Info);
+}
+
+void FNiagaraCrashReporterHandler::PopInfo()
+{
+	ScopeInfoStack.Pop();
+	if (ScopeInfoStack.Num())
+	{
+		SetInfo(ScopeInfoStack.Top());
+	}
+	else
+	{
+		SetInfo(NoneInfoString);
+	}
+}
+
+//////////////////////////////////////////////////////////////////////////
+
+FNiagaraCrashReporterScope::FNiagaraCrashReporterScope(FNiagaraSystemInstance* Inst)
+{
+	bWasEnabled = GbEnableNiagaraCRHandler != 0;
+	if (GbEnableNiagaraCRHandler)
+	{
+		FNiagaraCrashReporterHandler::Get().PushInfo(Inst);
+	}
+}
+
+FNiagaraCrashReporterScope::FNiagaraCrashReporterScope(FNiagaraSystemSimulation* Sim)
+{
+	bWasEnabled = GbEnableNiagaraCRHandler != 0;
+	if (GbEnableNiagaraCRHandler)
+	{
+		FNiagaraCrashReporterHandler::Get().PushInfo(Sim);
+	}
+}
+
+FNiagaraCrashReporterScope::~FNiagaraCrashReporterScope()
+{
+	if(GbEnableNiagaraCRHandler)
+	{
+		FNiagaraCrashReporterHandler::Get().PopInfo();
+	}
+	else
+	{
+		//Cache this and clear in case GbEnableNiagaraCRHandler changes mid scope.
+		if(bWasEnabled)
+		{
+			FNiagaraCrashReporterHandler::Get().Clear();
+		}
+	}
+}
diff --git a/Engine/Plugins/FX/Niagara/Source/Niagara/Private/NiagaraEmitterInstance.cpp b/Engine/Plugins/FX/Niagara/Source/Niagara/Private/NiagaraEmitterInstance.cpp
--- a/Engine/Plugins/FX/Niagara/Source/Niagara/Private/NiagaraEmitterInstance.cpp
+++ b/Engine/Plugins/FX/Niagara/Source/Niagara/Private/NiagaraEmitterInstance.cpp
@@ -1034,6 +1034,19 @@ void FNiagaraEmitterInstance::SetSystemFixedBoundsOverride(FBox SystemFixedBound
 	CachedSystemFixedBounds = SystemFixedBounds;
 }
 
+static int32 GbTriggerCrash = 0;
+static FAutoConsoleVariableRef CVarTriggerCrash(
+	TEXT("fx.TriggerDebugCrash"),
+	GbTriggerCrash,
+	TEXT("If > 0 we deliberately crash to test Crash Reporter integration."),
+	ECVF_Default
+);
+
+FORCENOINLINE void NiagaraTestCrash()
+{
+	check(0);
+}
+
 void FNiagaraEmitterInstance::Tick(float DeltaSeconds)
 {
 	SCOPE_CYCLE_COUNTER(STAT_NiagaraTick);
@@ -1048,6 +1061,15 @@ void FNiagaraEmitterInstance::Tick(float DeltaSeconds)
 		return;
 	}
 
+	//Test crash allowing us to test CR functionality.
+#if !UE_BUILD_SHIPPING
+	if (GbTriggerCrash)
+	{
+		GbTriggerCrash = 0;
+		NiagaraTestCrash();
+	}
+#endif
+
 	checkSlow(ParticleDataSet);
 	FNiagaraDataSet& Data = *ParticleDataSet;
 	EmitterAge += DeltaSeconds;
diff --git a/Engine/Plugins/FX/Niagara/Source/Niagara/Private/NiagaraSystemInstance.cpp b/Engine/Plugins/FX/Niagara/Source/Niagara/Private/NiagaraSystemInstance.cpp
--- a/Engine/Plugins/FX/Niagara/Source/Niagara/Private/NiagaraSystemInstance.cpp
+++ b/Engine/Plugins/FX/Niagara/Source/Niagara/Private/NiagaraSystemInstance.cpp
@@ -13,6 +13,7 @@
 #include "Templates/AlignmentTemplates.h"
 #include "NiagaraEmitterInstanceBatcher.h"
 #include "GameFramework/PlayerController.h"
+#include "NiagaraCrashReporterHandler.h"
 
 
 DECLARE_CYCLE_STAT(TEXT("System Activate [GT]"), STAT_NiagaraSystemActivate, STATGROUP_Niagara);
@@ -1971,6 +1972,8 @@ void FNiagaraSystemInstance::Tick_GameThread(float DeltaSeconds)
 	CSV_SCOPED_TIMING_STAT_EXCLUSIVE(Niagara);
 	LLM_SCOPE(ELLMTag::Niagara);
 
+	FNiagaraCrashReporterScope CRScope(this);
+
 	UNiagaraSystem* System = GetSystem();
 	FScopeCycleCounter SystemStat(System->GetStatID(true, false));
 
@@ -2000,6 +2003,8 @@ void FNiagaraSystemInstance::Tick_Concurrent()
 	LLM_SCOPE(ELLMTag::Niagara);
 	FScopeCycleCounterUObject AdditionalScope(GetSystem(), GET_STATID(STAT_NiagaraOverview_GT_CNC));
 
+	FNiagaraCrashReporterScope CRScope(this);
+
 	// Reset values that will be accumulated during emitter tick.
 	TotalGPUParamSize = 0;
 	ActiveGPUEmitterCount = 0;
@@ -2103,6 +2108,8 @@ void FNiagaraSystemInstance::FinalizeTick_GameThread()
 {
 	if (bNeedsFinalize)//We can come in here twice in one tick if the GT calls WaitForAsync() while there is a GT finalize task in the queue.
 	{
+		FNiagaraCrashReporterScope CRScope(this);
+
 		SCOPE_CYCLE_COUNTER(STAT_NiagaraOverview_GT);
 		SCOPE_CYCLE_COUNTER(STAT_NiagaraSystemInst_FinalizeGT);
 		CSV_SCOPED_TIMING_STAT_EXCLUSIVE(Niagara);
diff --git a/Engine/Plugins/FX/Niagara/Source/Niagara/Private/NiagaraSystemSimulation.cpp b/Engine/Plugins/FX/Niagara/Source/Niagara/Private/NiagaraSystemSimulation.cpp
--- a/Engine/Plugins/FX/Niagara/Source/Niagara/Private/NiagaraSystemSimulation.cpp
+++ b/Engine/Plugins/FX/Niagara/Source/Niagara/Private/NiagaraSystemSimulation.cpp
@@ -14,6 +14,7 @@
 #include "NiagaraComponent.h"
 #include "NiagaraWorldManager.h"
 #include "NiagaraEmitterInstanceBatcher.h"
+#include "NiagaraCrashReporterHandler.h"
 
 // Niagara simulations async will block the tick task from completion until all async work is finished
 // If simulations are allowed to tick async we will create a FNiagaraSystemSimulationTickTask task to run on any thread
@@ -636,6 +637,8 @@ void FNiagaraSystemSimulation::Tick_GameThread(float DeltaSeconds, const FGraphE
 	check(IsInGameThread());
 	check(bInSpawnPhase == false);
 
+	FNiagaraCrashReporterScope CRScope(this);
+
 	WaitForSystemTickComplete(true);
 
 	SCOPE_CYCLE_COUNTER(STAT_NiagaraOverview_GT);
@@ -929,6 +932,8 @@ void FNiagaraSystemSimulation::Spawn_GameThread(float DeltaSeconds)
 	UNiagaraSystem* System = WeakSystem.Get();
 	FScopeCycleCounterUObject AdditionalScope(System, GET_STATID(STAT_NiagaraOverview_GT_CNC));
 
+	FNiagaraCrashReporterScope CRScope(this);
+
 	WaitForSystemTickComplete(true);
 
 	bInSpawnPhase = true;
@@ -1060,6 +1065,8 @@ void FNiagaraSystemSimulation::Tick_Concurrent(FNiagaraSystemSimulationTickConte
 	FNiagaraScopedRuntimeCycleCounter RuntimeScope(Context.System, true, true);
 	FNiagaraSystemInstance* SoloSystemInstance = bIsSolo && Context.Instances.Num() == 1 ? Context.Instances[0] : nullptr;
 
+	FNiagaraCrashReporterScope CRScope(this);
+
 	if (bCanExecute && Context.Instances.Num())
 	{
 		if (GbDumpSystemData || Context.System->bDumpDebugSystemInfo)
diff --git a/Engine/Plugins/FX/Niagara/Source/Niagara/Public/NiagaraCrashReporterHandler.h b/Engine/Plugins/FX/Niagara/Source/Niagara/Public/NiagaraCrashReporterHandler.h
new file mode 100644
--- /dev/null
+++ b/Engine/Plugins/FX/Niagara/Source/Niagara/Public/NiagaraCrashReporterHandler.h
@@ -0,0 +1,60 @@
+// Copyright Epic Games, Inc. All Rights Reserved.
+
+#pragma once
+
+#include "HAL\ThreadSingleton.h"
+#include "HAL\CriticalSection.h"
+
+class FNiagaraSystemInstance;
+class FNiagaraSystemSimulation;
+
+/** Manages Niagara's integration with the Crash Reporter. Allows Niagara code to add information to a crash report which can help track down rare bugs. */
+class FNiagaraCrashReporterHandler : public TThreadSingleton<FNiagaraCrashReporterHandler>
+{
+	//Each thread we execute on can add it's own named data to the CR.
+	FString Name;
+
+	/** String to clear to when we're not inside the scope of a system instance. */
+	FString NoneInfoString;
+
+	/** Stack of scope infos, allows multi level scope tracking. */
+	TArray<FString> ScopeInfoStack;
+
+	//Allows us to clear out all data if we disable reporting.
+	static TArray<FNiagaraCrashReporterHandler*> AllHandlers;
+
+	/** Critical section ensuring thread-safe access to shared Crash Reporter state. */
+	static FCriticalSection CritSec;
+
+	void SetInfo(const FString& Info);
+public:
+
+	FNiagaraCrashReporterHandler();
+	~FNiagaraCrashReporterHandler();
+
+	/** Clears out Niagara's CR entry for this thread. */
+	void Clear();
+
+	/** Push info on a system instance to the info stack and set as Niagara's current CR info for this thread. */
+	void PushInfo(FNiagaraSystemInstance* Inst);
+	/** Push info on a system simulation to the info stack and set as Niagara's current CR info for this thread. */
+	void PushInfo(FNiagaraSystemSimulation* SystemSim);
+	/** Push an arbitrary string onto info stack and set as Niagara's current CR info for this thread. */
+	void PushInfo(const FString& Info);
+	/** Pop the current info from the stack and set the previous info as Niagara's current CR info for this thread. */
+	void PopInfo();
+
+	/** Handle a change in the enabled CVar for Niagara's CR info. */
+	static void OnEnabledChanged(IConsoleVariable* InVariable);
+};
+
+/** Helper object allowing easy tracking of Niagara code in it's crash reporter integration.  */
+class FNiagaraCrashReporterScope
+{
+private:
+	bool bWasEnabled = false;
+public:
+	FNiagaraCrashReporterScope(FNiagaraSystemInstance* Inst);
+	FNiagaraCrashReporterScope(FNiagaraSystemSimulation* Sim);
+	~FNiagaraCrashReporterScope();
+};
\ No newline at end of file
