Fix for NiagaraCrashReporter compiling on max/ios

diff --git a/Engine/Plugins/FX/Niagara/Source/Niagara/Public/NiagaraCrashReporterHandler.h b/Engine/Plugins/FX/Niagara/Source/Niagara/Public/NiagaraCrashReporterHandler.h
--- a/Engine/Plugins/FX/Niagara/Source/Niagara/Public/NiagaraCrashReporterHandler.h
+++ b/Engine/Plugins/FX/Niagara/Source/Niagara/Public/NiagaraCrashReporterHandler.h
@@ -2,8 +2,8 @@
 
 #pragma once
 
-#include "HAL\ThreadSingleton.h"
-#include "HAL\CriticalSection.h"
+#include "HAL/ThreadSingleton.h"
+#include "HAL/CriticalSection.h"
 
 class FNiagaraSystemInstance;
 class FNiagaraSystemSimulation;
