Fix for NiagaraCrashReporter compiling on max/ios

diff --git a/Engine/Plugins/FX/Niagara/Source/Niagara/Private/NiagaraCrashReporterHandler.cpp b/Engine/Plugins/FX/Niagara/Source/Niagara/Private/NiagaraCrashReporterHandler.cpp
--- a/Engine/Plugins/FX/Niagara/Source/Niagara/Private/NiagaraCrashReporterHandler.cpp
+++ b/Engine/Plugins/FX/Niagara/Source/Niagara/Private/NiagaraCrashReporterHandler.cpp
@@ -3,9 +3,9 @@
 #include "NiagaraCrashReporterHandler.h"
 
 #include "CoreMinimal.h"
-#include "HAL\ThreadManager.h"
-#include "GenericPlatform\GenericPlatformCrashContext.h"
-#include "Misc\ScopeLock.h"
+#include "HAL/ThreadManager.h"
+#include "GenericPlatform/GenericPlatformCrashContext.h"
+#include "Misc/ScopeLock.h"
 #include "NiagaraSystemInstance.h"
 #include "NiagaraSystemSimulation.h"
 #include "NiagaraComponent.h"
