Non unity compile fix

diff --git a/Engine/Plugins/FX/Niagara/Source/Niagara/Public/NiagaraCrashReporterHandler.h b/Engine/Plugins/FX/Niagara/Source/Niagara/Public/NiagaraCrashReporterHandler.h
--- a/Engine/Plugins/FX/Niagara/Source/Niagara/Public/NiagaraCrashReporterHandler.h
+++ b/Engine/Plugins/FX/Niagara/Source/Niagara/Public/NiagaraCrashReporterHandler.h
@@ -2,6 +2,7 @@
 
 #pragma once
 
+#include "CoreMinimal.h"
 #include "HAL/ThreadSingleton.h"
 #include "HAL/CriticalSection.h"
 
