Platform Set menu now stays open when pressing the buttons on the right-hand side. Clicking the name will close the menu.

Moved NiagaraPlatformSetCustomization to its own file as it had become very large.

#rnx

diff --git a/Engine/Plugins/FX/Niagara/Source/NiagaraEditor/Private/Customizations/NiagaraPlatformSetCustomization.cpp b/Engine/Plugins/FX/Niagara/Source/NiagaraEditor/Private/Customizations/NiagaraPlatformSetCustomization.cpp
new file mode 100644
--- /dev/null
+++ b/Engine/Plugins/FX/Niagara/Source/NiagaraEditor/Private/Customizations/NiagaraPlatformSetCustomization.cpp
@@ -0,0 +1,991 @@
+// Copyright Epic Games, Inc. All Rights Reserved.
+
+#include "NiagaraPlatformSetCustomization.h"
+#include "DetailWidgetRow.h"
+#include "IDetailChildrenBuilder.h"
+#include "NiagaraEditorStyle.h"
+#include "NiagaraEmitter.h"
+#include "NiagaraPlatformSet.h"
+#include "NiagaraSystem.h"
+#include "NiagaraTypes.h"
+#include "PlatformInfo.h"
+#include "PropertyHandle.h"
+#include "Scalability.h"
+#include "DeviceProfiles/DeviceProfile.h"
+#include "DeviceProfiles/DeviceProfileManager.h"
+#include "Layout/Visibility.h"
+#include "Widgets/Images/SImage.h"
+#include "Widgets/Input/SButton.h"
+#include "Widgets/Input/SCheckBox.h"
+#include "Widgets/Input/SMenuAnchor.h"
+#include "Widgets/Layout/SWrapBox.h"
+#include "Widgets/Text/STextBlock.h"
+
+#define LOCTEXT_NAMESPACE "FNiagaraPlatformSetCustomization"
+
+void FNiagaraPlatformSetCustomization::CustomizeHeader(TSharedRef<IPropertyHandle> InPropertyHandle, FDetailWidgetRow& HeaderRow, IPropertyTypeCustomizationUtils& CustomizationUtils)
+{
+	PropertyHandle = InPropertyHandle;
+	TArray<UObject*> Objects;
+	PropertyHandle->GetOuterObjects(Objects);
+	TargetPlatformSet = (FNiagaraPlatformSet*)PropertyHandle->GetValueBaseAddress((uint8*)Objects[0]);
+
+	if (PlatformSelectionStates.Num() == 0)
+	{
+		PlatformSelectionStates.Add(MakeShared<ENiagaraPlatformSelectionState>(ENiagaraPlatformSelectionState::Enabled));
+		PlatformSelectionStates.Add(MakeShared<ENiagaraPlatformSelectionState>(ENiagaraPlatformSelectionState::Disabled));
+		PlatformSelectionStates.Add(MakeShared<ENiagaraPlatformSelectionState>(ENiagaraPlatformSelectionState::Default));
+	}
+
+	//Look for outer types for which we need to ensure there are no conflicting settings.
+	SystemScalabilitySettings = nullptr;
+	EmitterScalabilitySettings = nullptr;
+	SystemScalabilityOverrides = nullptr;
+	EmitterScalabilityOverrides = nullptr;
+	PlatformSetArray.Reset();
+	PlatformSetArrayIndex = INDEX_NONE;
+
+	//Look whether this platform set belongs to a class which must keep orthogonal platform sets.
+	//We then interrogate the other sets via these ptrs to look for conflicts.
+	TSharedPtr<IPropertyHandle> CurrHandle = PropertyHandle->GetParentHandle();
+	while (CurrHandle)
+	{
+		int32 ThisIndex = CurrHandle->GetIndexInArray();
+		if (ThisIndex != INDEX_NONE)
+		{
+			PlatformSetArray = CurrHandle->GetParentHandle()->AsArray();
+			PlatformSetArrayIndex = ThisIndex;
+		}
+
+		if (FProperty* CurrProperty = CurrHandle->GetProperty())
+		{
+			if (UStruct* CurrStruct = CurrProperty->GetOwnerStruct())
+			{
+				TSharedPtr<IPropertyHandle> ParentHandle = CurrHandle->GetParentHandle();
+				if (CurrStruct == FNiagaraSystemScalabilitySettingsArray::StaticStruct())
+				{
+					SystemScalabilitySettings = (FNiagaraSystemScalabilitySettingsArray*)ParentHandle->GetValueBaseAddress((uint8*)Objects[0]);
+					break;
+				}
+				else if (CurrStruct == FNiagaraEmitterScalabilitySettingsArray::StaticStruct())
+				{
+					EmitterScalabilitySettings = (FNiagaraEmitterScalabilitySettingsArray*)ParentHandle->GetValueBaseAddress((uint8*)Objects[0]);
+					break;
+				}
+				else if (CurrStruct == FNiagaraSystemScalabilityOverrides::StaticStruct())
+				{
+					SystemScalabilityOverrides = (FNiagaraSystemScalabilityOverrides*)ParentHandle->GetValueBaseAddress((uint8*)Objects[0]);
+					break;
+				}
+				else if (CurrStruct == FNiagaraEmitterScalabilityOverrides::StaticStruct())
+				{
+					EmitterScalabilityOverrides = (FNiagaraEmitterScalabilityOverrides*)ParentHandle->GetValueBaseAddress((uint8*)Objects[0]);
+					break;
+				}
+			}
+		}
+
+		CurrHandle = CurrHandle->GetParentHandle();
+	}
+
+	UpdateCachedConflicts();
+	
+	BaseSystem = Objects[0]->GetTypedOuter<UNiagaraSystem>();
+	PropertyHandle->SetOnPropertyValueChanged(FSimpleDelegate::CreateSP(this, &FNiagaraPlatformSetCustomization::OnPropertyValueChanged));
+
+	HeaderRow
+		.WholeRowContent()
+		[ 	
+			SAssignNew(EffectsQualityWidgetBox, SWrapBox)
+			.UseAllottedWidth(true)
+		];
+
+	GenerateEffectsQualitySelectionWidgets();
+}
+
+void FNiagaraPlatformSetCustomization::CustomizeChildren(TSharedRef<IPropertyHandle> StructPropertyHandle, class IDetailChildrenBuilder& ChildBuilder, IPropertyTypeCustomizationUtils& StructCustomizationUtils)
+{
+}
+
+FText FNiagaraPlatformSetCustomization::GetCurrentText() const
+{
+	return LOCTEXT("Platforms", "Platforms");
+}
+
+FText FNiagaraPlatformSetCustomization::GetTooltipText() const
+{
+	return LOCTEXT("Platforms", "Platforms");
+}
+
+void FNiagaraPlatformSetCustomization::GenerateEffectsQualitySelectionWidgets()
+{
+	EffectsQualityWidgetBox->ClearChildren();
+	
+	Scalability::FQualityLevels Counts = Scalability::GetQualityLevelCounts();
+	EffectsQualityMenuAnchors.SetNum(Counts.EffectsQuality);
+	EffectsQualityMenuContents.SetNum(Counts.EffectsQuality);
+
+	for (int32 EffectsQuality = 0; EffectsQuality < Counts.EffectsQuality; ++EffectsQuality)
+	{
+		bool First = EffectsQuality == 0;
+		bool Last = EffectsQuality == (Counts.EffectsQuality - 1);
+
+		if (!EffectsQualityMenuAnchors[EffectsQuality].IsValid())
+		{
+			EffectsQualityMenuAnchors[EffectsQuality] = SNew(SMenuAnchor)
+				.ToolTipText(LOCTEXT("AddPlatformOverride", "Add an override for a specific platform."))
+				.OnGetMenuContent(this, &FNiagaraPlatformSetCustomization::GenerateDeviceProfileTreeWidget, EffectsQuality)
+				[
+					SNew(SButton)
+					.Visibility_Lambda([this, EffectsQuality]()
+					{
+						return EffectsQualityWidgetBox->GetChildren()->GetChildAt(EffectsQuality)->IsHovered() ? EVisibility::Visible : EVisibility::Hidden;
+					})
+					.ButtonStyle(FEditorStyle::Get(), "HoverHintOnly")
+					.ForegroundColor(FSlateColor::UseForeground())
+					.OnClicked(this, &FNiagaraPlatformSetCustomization::ToggleMenuOpenForEffectsQuality, EffectsQuality)
+					[
+						SNew(SBox)
+						.WidthOverride(8)
+						.HeightOverride(8)
+						.VAlign(VAlign_Center)
+						[
+							SNew(SImage)
+							.Image(FNiagaraEditorStyle::Get().GetBrush("NiagaraEditor.PlatformSet.DropdownButton"))
+						]
+					]
+				];
+		}
+
+		EffectsQualityWidgetBox->AddSlot()
+			.Padding(0, 0, 1, 0)
+			[
+				SNew(SBox)
+				.WidthOverride(100)
+				.VAlign(VAlign_Top)
+				[
+					SNew(SVerticalBox)
+					+ SVerticalBox::Slot()
+					.AutoHeight()
+					[
+						SNew(SCheckBox)
+						.Style(FNiagaraEditorStyle::Get(), First ? "NiagaraEditor.PlatformSet.StartButton" : 
+							(Last ? "NiagaraEditor.PlatformSet.EndButton" : "NiagaraEditor.PlatformSet.MiddleButton"))
+						.IsChecked(this, &FNiagaraPlatformSetCustomization::IsEQChecked, EffectsQuality)
+						.OnCheckStateChanged(this, &FNiagaraPlatformSetCustomization::EQCheckStateChanged, EffectsQuality)
+						[
+							SNew(SHorizontalBox)
+							+ SHorizontalBox::Slot()
+							.HAlign(HAlign_Fill)
+							.Padding(6,2,2,4)
+							[
+								SNew(STextBlock)
+								.TextStyle(FNiagaraEditorStyle::Get(), "NiagaraEditor.PlatformSet.ButtonText")
+								.Text(FNiagaraPlatformSet::GetEffectsQualityText(EffectsQuality))
+								.ColorAndOpacity(this, &FNiagaraPlatformSetCustomization::GetEffectsQualityButtonTextColor, EffectsQuality)
+								.ShadowOffset(FVector2D(1, 1))
+							]
+							// error icon
+							+ SHorizontalBox::Slot()
+							.AutoWidth()
+							.HAlign(HAlign_Right)
+							.VAlign(VAlign_Center)
+							.Padding(0,0,2,0)
+							[
+								SNew(SBox)
+								.WidthOverride(12)
+								.HeightOverride(12)
+								[
+									SNew(SImage)
+									.ToolTipText(this, &FNiagaraPlatformSetCustomization::GetEffectsQualityErrorToolTip, EffectsQuality)
+									.Visibility(this, &FNiagaraPlatformSetCustomization::GetEffectsQualityErrorVisibility, EffectsQuality)
+									.Image(FEditorStyle::GetBrush("Icons.Error"))
+								]
+							]
+							// dropdown button
+							+ SHorizontalBox::Slot()
+							.AutoWidth()
+							.HAlign(HAlign_Right)
+							.VAlign(VAlign_Center)
+							.Padding(0,0,2,0)
+							[
+								EffectsQualityMenuAnchors[EffectsQuality].ToSharedRef()
+							]
+						]
+					]
+					+ SVerticalBox::Slot()
+					.AutoHeight()
+					[
+						GenerateAdditionalDevicesWidgetForEQ(EffectsQuality)
+					]
+				]
+			];
+	}
+}
+
+void FNiagaraPlatformSetCustomization::UpdateCachedConflicts()
+{
+	TArray<const FNiagaraPlatformSet*> PlatformSets;
+	if (SystemScalabilityOverrides != nullptr)
+	{
+		for (FNiagaraSystemScalabilityOverride& Override : SystemScalabilityOverrides->Overrides)
+		{
+			PlatformSets.Add(&Override.Platforms);
+		}
+	}
+	else if (EmitterScalabilityOverrides != nullptr)
+	{
+		for (FNiagaraEmitterScalabilityOverride& Override : EmitterScalabilityOverrides->Overrides)
+		{
+			PlatformSets.Add(&Override.Platforms);
+		}
+	}
+	else if (SystemScalabilitySettings != nullptr)
+	{
+		for (FNiagaraSystemScalabilitySettings& Settings : SystemScalabilitySettings->Settings)
+		{
+			PlatformSets.Add(&Settings.Platforms);
+		}
+	}
+	else if (EmitterScalabilitySettings != nullptr)
+	{
+		for (FNiagaraEmitterScalabilitySettings& Settings : EmitterScalabilitySettings->Settings)
+		{
+			PlatformSets.Add(&Settings.Platforms);
+		}
+	}
+
+	CachedConflicts.Reset();
+	FNiagaraPlatformSet::GatherConflicts(PlatformSets, CachedConflicts);
+}
+
+static TSharedPtr<IPropertyHandle> FindChildPlatformSet(TSharedPtr<IPropertyHandle> PropertyHandle)
+{
+	if (FStructProperty* Property = CastField<FStructProperty>(PropertyHandle->GetProperty()))
+	{
+		if (Property->Struct == FNiagaraPlatformSet::StaticStruct())
+		{
+			return PropertyHandle;
+		}
+	}
+
+	// recurse
+	uint32 NumChildren = 0;
+	PropertyHandle->GetNumChildren(NumChildren);
+
+	for (int32 Idx = 0; Idx < (int32) NumChildren; ++Idx)
+	{
+		TSharedPtr<IPropertyHandle> Child = PropertyHandle->GetChildHandle(Idx);
+
+		TSharedPtr<IPropertyHandle> ChildResult = FindChildPlatformSet(Child);
+		if (ChildResult.IsValid())
+		{
+			return ChildResult;
+		}
+	}
+
+	return TSharedPtr<IPropertyHandle>();
+}
+
+void FNiagaraPlatformSetCustomization::InvalidateSiblingConflicts() const
+{
+	if (!PlatformSetArray.IsValid())
+	{
+		return;
+	}
+
+	uint32 ArrayCount = 0;
+	PlatformSetArray->GetNumElements(ArrayCount);
+	for (int32 Idx = 0; Idx < (int32) ArrayCount; ++Idx)
+	{
+		if (Idx == PlatformSetArrayIndex)
+		{
+			continue; // skip self
+		}
+
+		TSharedRef<IPropertyHandle> Sibling = PlatformSetArray->GetElement(Idx);
+
+		TSharedPtr<IPropertyHandle> SiblingPlatformSet = FindChildPlatformSet(Sibling);
+		if (SiblingPlatformSet.IsValid())
+		{
+			SiblingPlatformSet->NotifyPostChange();
+		}
+	}
+}
+
+EVisibility FNiagaraPlatformSetCustomization::GetEffectsQualityErrorVisibility(int32 EffectsQuality) const
+{
+	// not part of an array
+	if (PlatformSetArrayIndex == INDEX_NONE)
+	{
+		return EVisibility::Collapsed;
+	}
+
+	for (const FNiagaraPlatformSetConflictInfo& ConflictInfo : CachedConflicts)
+	{
+		if (ConflictInfo.SetAIndex == PlatformSetArrayIndex || 
+			ConflictInfo.SetBIndex == PlatformSetArrayIndex)
+		{
+			// this conflict applies to this platform set, check if it applies to this effects quality button
+			const int32 EQMask = FNiagaraPlatformSet::CreateEQMask(EffectsQuality);
+
+			for (const FNiagaraPlatformSetConflictEntry& Conflict : ConflictInfo.Conflicts)
+			{
+				if ((EQMask & Conflict.EffectsQualityMask) != 0)
+				{
+					return EVisibility::Visible;
+				}
+			}
+		}
+	}
+
+	return EVisibility::Collapsed;
+}
+
+FText FNiagaraPlatformSetCustomization::GetEffectsQualityErrorToolTip(int32 EffectsQuality) const
+{
+	if (PlatformSetArrayIndex == INDEX_NONE)
+	{
+		return FText::GetEmpty();
+	}
+
+	for (const FNiagaraPlatformSetConflictInfo& ConflictInfo : CachedConflicts)
+	{
+		int32 OtherIndex = INDEX_NONE;
+
+		if (ConflictInfo.SetAIndex == PlatformSetArrayIndex)
+		{
+			OtherIndex = ConflictInfo.SetBIndex;
+		}
+
+		if (ConflictInfo.SetBIndex == PlatformSetArrayIndex)
+		{
+			OtherIndex = ConflictInfo.SetAIndex;
+		}
+
+		if (OtherIndex != INDEX_NONE)
+		{
+			const int32 EQMask = FNiagaraPlatformSet::CreateEQMask(EffectsQuality);
+
+			for (const FNiagaraPlatformSetConflictEntry& Conflict : ConflictInfo.Conflicts)
+			{
+				if ((EQMask & Conflict.EffectsQualityMask) != 0)
+				{
+					FText FormatString = LOCTEXT("EffectsQualityConflictToolTip", "This effect quality conflicts with the set at index {Index} in this array."); 
+
+					FFormatNamedArguments Args;
+					Args.Add(TEXT("Index"), OtherIndex);
+
+					return FText::Format(FormatString, Args);
+				}
+			}
+		}
+	}
+
+	return FText::GetEmpty();
+}
+
+FSlateColor FNiagaraPlatformSetCustomization::GetEffectsQualityButtonTextColor(int32 EffectsQuality) const
+{
+	return TargetPlatformSet->IsEffectQualityEnabled(EffectsQuality) ? 
+		FSlateColor(FLinearColor(0.95f, 0.95f, 0.95f)) :
+		FSlateColor::UseForeground();
+}
+
+TSharedRef<SWidget> FNiagaraPlatformSetCustomization::GenerateAdditionalDevicesWidgetForEQ(int32 EffectsQuality)
+{
+	TSharedRef<SVerticalBox> Container = SNew(SVerticalBox);
+
+	auto AddDeviceProfileOverrideWidget = [&](UDeviceProfile* Profile, bool bEnabled)
+	{
+		TSharedPtr<SHorizontalBox> DeviceBox;
+
+		Container->AddSlot()
+			.AutoHeight()
+			[
+				SAssignNew(DeviceBox, SHorizontalBox)
+			];
+
+		const FText DeviceNameText = FText::FromName(Profile->GetFName());
+
+		DeviceBox->AddSlot()
+			.AutoWidth()
+			.HAlign(HAlign_Left)
+			.VAlign(VAlign_Center)
+			.Padding(3, 0, 3, 0)
+			[
+				SNew(SBox)
+				.WidthOverride(8)
+				.HeightOverride(8)
+				.ToolTipText(bEnabled ?
+					LOCTEXT("DeviceIncludedToolTip", "This device is included.") :
+					LOCTEXT("DeviceExcludedToolTip", "This device is excluded."))
+				[
+					SNew(SImage)
+					.Image(bEnabled ?
+						FNiagaraEditorStyle::Get().GetBrush("NiagaraEditor.PlatformSet.Include") :
+						FNiagaraEditorStyle::Get().GetBrush("NiagaraEditor.PlatformSet.Exclude"))
+					.ColorAndOpacity(bEnabled ?
+						FLinearColor(0, 1, 0) :
+						FLinearColor(1, 0, 0))
+					]
+				];
+
+		DeviceBox->AddSlot()
+			.HAlign(HAlign_Fill)
+			.Padding(0, 0, 2, 0)
+			[
+				SNew(STextBlock)
+				.Text(DeviceNameText)
+				.ToolTipText(DeviceNameText)
+			];
+
+		DeviceBox->AddSlot()
+			.AutoWidth()
+			.HAlign(HAlign_Right)
+			.VAlign(VAlign_Center)
+			.Padding(0, 0, 2, 0)
+			[
+				SNew(SBox)
+				.WidthOverride(12)
+				.HeightOverride(12)
+				[
+					SNew(SImage)
+					.ToolTipText(this, &FNiagaraPlatformSetCustomization::GetDeviceProfileErrorToolTip, Profile, EffectsQuality)
+					.Visibility(this, &FNiagaraPlatformSetCustomization::GetDeviceProfileErrorVisibility, Profile, EffectsQuality)
+					.Image(FEditorStyle::GetBrush("Icons.Error"))
+				]
+			];
+
+		DeviceBox->AddSlot()
+			.AutoWidth()
+			.HAlign(HAlign_Right)
+			.VAlign(VAlign_Center)
+			.Padding(0, 0, 2, 0)
+			[
+				SNew(SButton)
+				.ButtonStyle(FEditorStyle::Get(), "HoverHintOnly")
+				.ForegroundColor(FSlateColor::UseForeground())
+				.OnClicked(this, &FNiagaraPlatformSetCustomization::RemoveDeviceProfile, Profile, EffectsQuality)
+				.ToolTipText(LOCTEXT("RemoveDevice", "Remove this device override."))
+				[
+					SNew(SBox)
+					.WidthOverride(8)
+					.HeightOverride(8)
+					.VAlign(VAlign_Center)
+					[
+						SNew(SImage)
+						.Image(FNiagaraEditorStyle::Get().GetBrush("NiagaraEditor.PlatformSet.Remove"))
+					]
+				]
+			];
+		};
+
+	TArray<UDeviceProfile*> EnabledProfiles;
+	TArray<UDeviceProfile*> DisabledProfiles;
+	TargetPlatformSet->GetOverridenDeviceProfiles(EffectsQuality, EnabledProfiles, DisabledProfiles);
+
+	for (UDeviceProfile* Profile : EnabledProfiles)
+	{
+		AddDeviceProfileOverrideWidget(Profile, true);
+	}
+
+	for (UDeviceProfile* Profile : DisabledProfiles)
+	{
+		AddDeviceProfileOverrideWidget(Profile, false);
+	}
+
+	return Container;
+}
+
+FReply FNiagaraPlatformSetCustomization::RemoveDeviceProfile(UDeviceProfile* Profile, int32 EffectsQuality)
+{
+	PropertyHandle->NotifyPreChange();
+	TargetPlatformSet->SetDeviceProfileState(Profile, EffectsQuality, ENiagaraPlatformSelectionState::Default);
+	PropertyHandle->NotifyPostChange(EPropertyChangeType::ValueSet);
+
+	InvalidateSiblingConflicts();
+
+	return FReply::Handled();
+}
+
+EVisibility FNiagaraPlatformSetCustomization::GetDeviceProfileErrorVisibility(UDeviceProfile* Profile, int32 EffectsQuality) const
+{
+	// not part of an array
+	if (PlatformSetArrayIndex == INDEX_NONE)
+	{
+		return EVisibility::Collapsed;
+	}
+
+	for (const FNiagaraPlatformSetConflictInfo& ConflictInfo : CachedConflicts)
+	{
+		if (ConflictInfo.SetAIndex == PlatformSetArrayIndex || 
+			ConflictInfo.SetBIndex == PlatformSetArrayIndex)
+		{
+			const int32 EQMask = FNiagaraPlatformSet::CreateEQMask(EffectsQuality);
+
+			for (const FNiagaraPlatformSetConflictEntry& Entry : ConflictInfo.Conflicts)
+			{
+				if ((Entry.EffectsQualityMask & EQMask) != 0 &&
+					Entry.ProfileName == Profile->GetFName())
+				{
+					return EVisibility::Visible;
+				}
+			}
+		}
+	}
+
+	return EVisibility::Collapsed;
+}
+
+FText FNiagaraPlatformSetCustomization::GetDeviceProfileErrorToolTip(UDeviceProfile* Profile, int32 EffectsQuality) const
+{
+	for (const FNiagaraPlatformSetConflictInfo& ConflictInfo : CachedConflicts)
+	{
+		int32 OtherIndex = INDEX_NONE;
+
+		if (ConflictInfo.SetAIndex == PlatformSetArrayIndex)
+		{
+			OtherIndex = ConflictInfo.SetBIndex;
+		}
+
+		if (ConflictInfo.SetBIndex == PlatformSetArrayIndex)
+		{
+			OtherIndex = ConflictInfo.SetAIndex;
+		}
+
+		if (OtherIndex != INDEX_NONE)
+		{
+			const int32 EQMask = FNiagaraPlatformSet::CreateEQMask(EffectsQuality);
+
+			for (const FNiagaraPlatformSetConflictEntry& Entry : ConflictInfo.Conflicts)
+			{
+				if ((Entry.EffectsQualityMask & EQMask) != 0 &&
+					Entry.ProfileName == Profile->GetFName())
+				{
+					FText FormatString = LOCTEXT("PlatformOverrideConflictToolTip", "This platform override conflicts with the set at index {Index} in this array."); 
+
+					FFormatNamedArguments Args;
+					Args.Add(TEXT("Index"), OtherIndex);
+
+					return FText::Format(FormatString, Args);
+				}
+			}
+		}
+	}
+
+	return FText::GetEmpty();
+}
+
+FReply FNiagaraPlatformSetCustomization::ToggleMenuOpenForEffectsQuality(int32 EffectsQuality)
+{
+	check(EffectsQualityMenuAnchors.IsValidIndex(EffectsQuality));
+	
+	TSharedPtr<SMenuAnchor> MenuAnchor = EffectsQualityMenuAnchors[EffectsQuality];
+	MenuAnchor->SetIsOpen(!MenuAnchor->IsOpen());
+
+	return FReply::Handled();
+}
+
+// Is or does a viewmodel contain any children active at the given quality?
+bool FNiagaraPlatformSetCustomization::IsTreeActiveForEQ(const TSharedPtr<FNiagaraDeviceProfileViewModel>& Tree, int32 EffectsQualityMask) const
+{
+	int32 Mask = TargetPlatformSet->GetEffectQualityMaskForDeviceProfile(Tree->Profile);
+	if ((Mask & EffectsQualityMask) != 0)
+	{
+		return true;
+	}
+
+	for (const TSharedPtr<FNiagaraDeviceProfileViewModel>& Child : Tree->Children)
+	{
+		if (IsTreeActiveForEQ(Child, EffectsQualityMask))
+		{
+			return true;
+		}
+	}
+
+	return false;
+}
+
+void FNiagaraPlatformSetCustomization::FilterTreeForEQ(const TSharedPtr<FNiagaraDeviceProfileViewModel>& SourceTree, TSharedPtr<FNiagaraDeviceProfileViewModel>& FilteredTree, int32 EffectsQualityMask)
+{
+	for (const TSharedPtr<FNiagaraDeviceProfileViewModel>& SourceChild : SourceTree->Children)
+	{
+		if (IsTreeActiveForEQ(SourceChild, EffectsQualityMask))
+		{
+			TSharedPtr<FNiagaraDeviceProfileViewModel>& FilteredChild = FilteredTree->Children.Add_GetRef(MakeShared<FNiagaraDeviceProfileViewModel>());
+			FilteredChild->Profile = SourceChild->Profile;
+
+			FilterTreeForEQ(SourceChild, FilteredChild, EffectsQualityMask);
+		}
+	}
+}
+
+void FNiagaraPlatformSetCustomization::CreateDeviceProfileTree()
+{
+	//Pull device profiles out by their hierarchy depth.
+	TArray<TArray<UDeviceProfile*>> DeviceProfilesByHierarchyLevel;
+	for (UObject* Profile : UDeviceProfileManager::Get().Profiles)
+	{
+		UDeviceProfile* DeviceProfile = CastChecked<UDeviceProfile>(Profile);
+
+		TFunction<void(int32&, UDeviceProfile*)> FindDepth;
+		FindDepth = [&](int32& Depth, UDeviceProfile* CurrProfile)
+		{
+			if (CurrProfile->Parent)
+			{
+				FindDepth(++Depth, Cast<UDeviceProfile>(CurrProfile->Parent));
+			}
+		};
+
+		int32 ProfileDepth = 0;
+		FindDepth(ProfileDepth, DeviceProfile);
+		DeviceProfilesByHierarchyLevel.SetNum(FMath::Max(ProfileDepth+1, DeviceProfilesByHierarchyLevel.Num()));
+		DeviceProfilesByHierarchyLevel[ProfileDepth].Add(DeviceProfile);
+	}
+	
+	FullDeviceProfileTree.Reset(DeviceProfilesByHierarchyLevel[0].Num());
+	for (int32 RootProfileIdx = 0; RootProfileIdx < DeviceProfilesByHierarchyLevel[0].Num(); ++RootProfileIdx)
+	{
+		TFunction<void(FNiagaraDeviceProfileViewModel*, int32)> BuildProfileTree;
+		BuildProfileTree = [&](FNiagaraDeviceProfileViewModel* CurrRoot, int32 CurrLevel)
+		{
+			int32 NextLevel = CurrLevel + 1;
+			if (NextLevel < DeviceProfilesByHierarchyLevel.Num())
+			{
+				for (UDeviceProfile* PossibleChild : DeviceProfilesByHierarchyLevel[NextLevel])
+				{
+					if (PossibleChild->Parent == CurrRoot->Profile)
+					{
+						TSharedPtr<FNiagaraDeviceProfileViewModel>& NewChild = CurrRoot->Children.Add_GetRef(MakeShared<FNiagaraDeviceProfileViewModel>());
+						NewChild->Profile = PossibleChild;
+						BuildProfileTree(NewChild.Get(), NextLevel);
+					}
+				}
+			}
+		};
+
+		//Add all root nodes and build their trees.
+		TSharedPtr<FNiagaraDeviceProfileViewModel> CurrRoot = MakeShared<FNiagaraDeviceProfileViewModel>();
+		CurrRoot->Profile = DeviceProfilesByHierarchyLevel[0][RootProfileIdx];
+		BuildProfileTree(CurrRoot.Get(), 0);
+		FullDeviceProfileTree.Add(CurrRoot);
+	}
+
+	Scalability::FQualityLevels Counts = Scalability::GetQualityLevelCounts();
+	FilteredDeviceProfileTrees.SetNum(Counts.EffectsQuality);
+	
+	for (TSharedPtr<FNiagaraDeviceProfileViewModel>& FullDeviceRoot : FullDeviceProfileTree)
+	{
+		for (int32 EffectsQuality = 0; EffectsQuality < Counts.EffectsQuality; ++EffectsQuality)
+		{
+			int32 EffectsQualityMask = FNiagaraPlatformSet::CreateEQMask(EffectsQuality);
+
+			if (IsTreeActiveForEQ(FullDeviceRoot, EffectsQualityMask))
+			{
+				TArray<TSharedPtr<FNiagaraDeviceProfileViewModel>>& FilteredRoots = FilteredDeviceProfileTrees[EffectsQuality];
+
+				TSharedPtr<FNiagaraDeviceProfileViewModel>& FilteredRoot = FilteredRoots.Add_GetRef(MakeShared<FNiagaraDeviceProfileViewModel>());
+				FilteredRoot->Profile = FullDeviceRoot->Profile;
+
+				FilterTreeForEQ(FullDeviceRoot, FilteredRoot, EffectsQualityMask);
+			}
+		}
+	}
+}
+
+TSharedRef<SWidget> FNiagaraPlatformSetCustomization::GenerateDeviceProfileTreeWidget(int32 EffectsQuality)
+{
+	if (FullDeviceProfileTree.Num() == 0)
+	{
+		CreateDeviceProfileTree();	
+	}
+
+	if (EffectsQualityMenuContents[EffectsQuality].IsValid())
+	{
+		return EffectsQualityMenuContents[EffectsQuality].ToSharedRef();
+	}
+
+	TArray<TSharedPtr<FNiagaraDeviceProfileViewModel>>* TreeToUse = &FullDeviceProfileTree;
+	if (EffectsQuality != INDEX_NONE)
+	{
+		check(EffectsQuality < FilteredDeviceProfileTrees.Num());
+		TreeToUse = &FilteredDeviceProfileTrees[EffectsQuality];
+	}
+
+	return SAssignNew(EffectsQualityMenuContents[EffectsQuality], SBorder)
+		.BorderImage(FEditorStyle::Get().GetBrush("Menu.Background"))
+		[
+			SAssignNew(DeviceProfileTreeWidget, STreeView<TSharedPtr<FNiagaraDeviceProfileViewModel>>)
+			.TreeItemsSource(TreeToUse)
+			.OnGenerateRow(this, &FNiagaraPlatformSetCustomization::OnGenerateDeviceProfileTreeRow, EffectsQuality)
+			.OnGetChildren(this, &FNiagaraPlatformSetCustomization::OnGetDeviceProfileTreeChildren, EffectsQuality)
+			.SelectionMode(ESelectionMode::None)
+		];
+}
+
+TSharedRef<ITableRow> FNiagaraPlatformSetCustomization::OnGenerateDeviceProfileTreeRow(TSharedPtr<FNiagaraDeviceProfileViewModel> InItem, const TSharedRef<STableViewBase>& OwnerTable, int32 EffectsQuality)
+{
+	TSharedPtr<SHorizontalBox> RowContainer;
+	SAssignNew(RowContainer, SHorizontalBox);
+
+	int32 ProfileMask = TargetPlatformSet->GetEffectQualityMaskForDeviceProfile(InItem->Profile);
+	FText NameTooltip = FText::Format(LOCTEXT("ProfileEQTooltipFmt", "Effects Quality: {0}"), FNiagaraPlatformSet::GetEffectsQualityMaskText(ProfileMask));
+	
+	//Top level profile. Look for a platform icon.
+	if (InItem->Profile->Parent == nullptr)
+	{
+		if (const PlatformInfo::FPlatformInfo* Info = PlatformInfo::FindPlatformInfo(*InItem->Profile->DeviceType))
+		{
+			const FSlateBrush* DeviceProfileTypeIcon = FEditorStyle::GetBrush(Info->GetIconStyleName(PlatformInfo::EPlatformIconSize::Normal));
+			if (DeviceProfileTypeIcon != FEditorStyle::Get().GetDefaultBrush())
+			{
+				RowContainer->AddSlot()
+					.AutoWidth()
+					.VAlign(VAlign_Center)
+					.HAlign(HAlign_Left)
+					.Padding(4, 0, 0, 0)
+					[
+						SNew(SBox)
+						.WidthOverride(16)
+						.HeightOverride(16)
+						[
+							SNew(SImage)
+							.Image(DeviceProfileTypeIcon)
+						]
+					];
+			}
+		}
+	}
+
+	FName TextStyleName("NormalText");
+	FSlateColor TextColor(FSlateColor::UseForeground());
+	ENiagaraPlatformSelectionState CurrentState = TargetPlatformSet->GetDeviceProfileState(InItem->Profile, EffectsQuality);
+	
+	if (CurrentState == ENiagaraPlatformSelectionState::Enabled)
+	{
+		TextStyleName = "RichTextBlock.Bold";
+		TextColor = FSlateColor(FLinearColor(FVector4(0,1,0,1)));
+	}
+	else if (CurrentState == ENiagaraPlatformSelectionState::Disabled)
+	{
+		TextStyleName = "RichTextBlock.Italic";
+		TextColor = FSlateColor(FLinearColor(FVector4(1,0,0,1)));
+	}
+
+
+	int32 EffectsQualityMask = FNiagaraPlatformSet::CreateEQMask(EffectsQuality);
+	if ((ProfileMask & EffectsQualityMask) == 0)
+	{
+		TextColor = FSlateColor::UseSubduedForeground();
+	}
+
+	RowContainer->AddSlot()
+		.Padding(4, 2, 0, 2)
+		.HAlign(HAlign_Fill)
+		.VAlign(VAlign_Center)
+		[
+			SNew(SButton)
+			.ButtonStyle(FEditorStyle::Get(), "NoBorder")
+			.OnClicked(this, &FNiagaraPlatformSetCustomization::OnProfileMenuButtonClicked, InItem, EffectsQuality, false)
+			.IsEnabled(this, &FNiagaraPlatformSetCustomization::GetProfileMenuItemEnabled, InItem, EffectsQuality)
+			.ForegroundColor(TextColor)
+			.ToolTipText(NameTooltip)
+			[
+				SNew(STextBlock)
+				.TextStyle(FEditorStyle::Get(), TextStyleName)
+				.Text(FText::FromString(InItem->Profile->GetName()))
+			]
+		];
+
+	RowContainer->AddSlot()
+		.AutoWidth()
+		.Padding(12, 2, 4, 4)
+		.HAlign(HAlign_Right)
+		.VAlign(VAlign_Center)
+		[
+			SNew(SButton)
+			.ButtonStyle(FEditorStyle::Get(), "HoverHintOnly")
+			.HAlign(HAlign_Center)
+			.VAlign(VAlign_Fill)
+			.Visibility(this, &FNiagaraPlatformSetCustomization::GetProfileMenuButtonVisibility, InItem, EffectsQuality)
+			.OnClicked(this, &FNiagaraPlatformSetCustomization::OnProfileMenuButtonClicked, InItem, EffectsQuality, true)
+			.ToolTipText(this, &FNiagaraPlatformSetCustomization::GetProfileMenuButtonToolTip, InItem, EffectsQuality)
+			[
+				SNew(SBox)
+				.WidthOverride(8)
+				.HeightOverride(8)
+				.VAlign(VAlign_Center)
+				.HAlign(HAlign_Center)
+				[
+					SNew(SImage)
+					.Image(this, &FNiagaraPlatformSetCustomization::GetProfileMenuButtonImage, InItem, EffectsQuality)
+				]
+			]
+		];
+
+	return SNew(STableRow<TSharedPtr<FNiagaraDeviceProfileViewModel>>, OwnerTable)
+		.Style(FNiagaraEditorStyle::Get(), "NiagaraEditor.PlatformSet.TreeView")
+		[
+			RowContainer.ToSharedRef()
+		];
+}
+
+enum class EProfileButtonMode
+{
+	None,
+	Include,
+	Exclude,
+	Remove
+};
+
+static EProfileButtonMode GetProfileMenuButtonMode(FNiagaraPlatformSet* PlatformSet, TSharedPtr<FNiagaraDeviceProfileViewModel> Item, int32 EffectsQuality)
+{
+	int32 Mask = PlatformSet->GetEffectQualityMaskForDeviceProfile(Item->Profile);
+	int32 EQMask = FNiagaraPlatformSet::CreateEQMask(EffectsQuality);
+
+	if ((Mask & EQMask) == 0)
+	{
+		return EProfileButtonMode::None;
+	}
+
+	ENiagaraPlatformSelectionState CurrentState = PlatformSet->GetDeviceProfileState(Item->Profile, EffectsQuality);
+
+	bool bQualityEnabled = PlatformSet->IsEffectQualityEnabled(EffectsQuality);
+	bool bIsDefault = CurrentState == ENiagaraPlatformSelectionState::Default;
+
+	if (bIsDefault && bQualityEnabled)
+	{
+		return EProfileButtonMode::Exclude;
+	}
+
+	if (bIsDefault && !bQualityEnabled)
+	{
+		return EProfileButtonMode::Include;
+	}
+
+	return EProfileButtonMode::Remove;
+}
+
+
+FText FNiagaraPlatformSetCustomization::GetProfileMenuButtonToolTip(TSharedPtr<FNiagaraDeviceProfileViewModel> Item, int32 EffectsQuality) const
+{
+	EProfileButtonMode Mode = GetProfileMenuButtonMode(TargetPlatformSet, Item, EffectsQuality);
+	switch(Mode)
+	{
+		case EProfileButtonMode::Include:
+			return LOCTEXT("IncludePlatform", "Include this platform.");
+		case EProfileButtonMode::Exclude:
+			return LOCTEXT("ExcludePlatform", "Exclude this platform.");
+		case EProfileButtonMode::Remove:
+			return LOCTEXT("RemovePlatform", "Remove this platform override.");
+	}
+
+	return FText::GetEmpty();
+}
+
+EVisibility FNiagaraPlatformSetCustomization::GetProfileMenuButtonVisibility(TSharedPtr<FNiagaraDeviceProfileViewModel> Item, int32 EffectsQuality) const
+{
+	EProfileButtonMode Mode = GetProfileMenuButtonMode(TargetPlatformSet, Item, EffectsQuality);
+	if (Mode == EProfileButtonMode::None)
+	{
+		return EVisibility::Collapsed;
+	}
+
+	return EVisibility::Visible;
+}
+
+bool FNiagaraPlatformSetCustomization::GetProfileMenuItemEnabled(TSharedPtr<FNiagaraDeviceProfileViewModel> Item, int32 EffectsQuality) const
+{
+	EProfileButtonMode Mode = GetProfileMenuButtonMode(TargetPlatformSet, Item, EffectsQuality);
+	if (Mode == EProfileButtonMode::None)
+	{
+		return false;
+	}
+
+	return true;
+}
+
+const FSlateBrush* FNiagaraPlatformSetCustomization::GetProfileMenuButtonImage(TSharedPtr<FNiagaraDeviceProfileViewModel> Item, int32 EffectsQuality) const
+{
+	EProfileButtonMode Mode = GetProfileMenuButtonMode(TargetPlatformSet, Item, EffectsQuality);
+	switch(Mode)
+	{
+		case EProfileButtonMode::Include:
+			return FNiagaraEditorStyle::Get().GetBrush("NiagaraEditor.PlatformSet.Include");
+		case EProfileButtonMode::Exclude:
+			return FNiagaraEditorStyle::Get().GetBrush("NiagaraEditor.PlatformSet.Exclude");
+		case EProfileButtonMode::Remove:
+			return FNiagaraEditorStyle::Get().GetBrush("NiagaraEditor.PlatformSet.Remove");
+	}
+
+	return FEditorStyle::GetBrush("NoBrush");
+}
+
+FReply FNiagaraPlatformSetCustomization::OnProfileMenuButtonClicked(TSharedPtr<FNiagaraDeviceProfileViewModel> Item, int32 EffectsQuality, bool bReopenMenu)
+{
+	ENiagaraPlatformSelectionState TargetState;
+
+	EProfileButtonMode Mode = GetProfileMenuButtonMode(TargetPlatformSet, Item, EffectsQuality);
+	switch(Mode)
+	{
+		case EProfileButtonMode::Include:
+			TargetState = ENiagaraPlatformSelectionState::Enabled;
+			break;
+		case EProfileButtonMode::Exclude:
+			TargetState = ENiagaraPlatformSelectionState::Disabled;
+			break;
+		case EProfileButtonMode::Remove:
+			TargetState = ENiagaraPlatformSelectionState::Default;
+			break;
+		default:
+			return FReply::Handled(); // shouldn't happen, button should be collapsed
+	}
+
+	PropertyHandle->NotifyPreChange();
+	TargetPlatformSet->SetDeviceProfileState(Item->Profile, EffectsQuality, TargetState);
+	PropertyHandle->NotifyPostChange(EPropertyChangeType::ValueSet);
+
+	InvalidateSiblingConflicts();
+
+	if (!bReopenMenu)
+	{
+		EffectsQualityMenuAnchors[EffectsQuality]->SetIsOpen(false);
+	}
+
+	return FReply::Handled();
+}
+
+void FNiagaraPlatformSetCustomization::OnGetDeviceProfileTreeChildren(TSharedPtr<FNiagaraDeviceProfileViewModel> InItem, TArray< TSharedPtr<FNiagaraDeviceProfileViewModel> >& OutChildren, int32 EffectsQuality)
+{
+	if (TargetPlatformSet->GetDeviceProfileState(InItem->Profile, EffectsQuality) == ENiagaraPlatformSelectionState::Default)
+	{
+		OutChildren = InItem->Children;
+	}
+}
+
+ECheckBoxState FNiagaraPlatformSetCustomization::IsEQChecked(int32 EffectsQuality)const
+{
+	if (TargetPlatformSet)
+	{
+		return TargetPlatformSet->IsEffectQualityEnabled(EffectsQuality) ? ECheckBoxState::Checked : ECheckBoxState::Unchecked;
+	}
+	return ECheckBoxState::Undetermined;
+}
+
+void FNiagaraPlatformSetCustomization::EQCheckStateChanged(ECheckBoxState CheckState, int32 EffectsQuality)
+{
+	PropertyHandle->NotifyPreChange();
+	TargetPlatformSet->SetEnabledForEffectQuality(EffectsQuality, CheckState == ECheckBoxState::Checked);
+	PropertyHandle->NotifyPostChange(EPropertyChangeType::ValueSet);
+
+	InvalidateSiblingConflicts();
+}
+
+void FNiagaraPlatformSetCustomization::OnPropertyValueChanged()
+{
+	GenerateEffectsQualitySelectionWidgets();
+	UpdateCachedConflicts();
+}
+
+#undef LOCTEXT_NAMESPACE
\ No newline at end of file
diff --git a/Engine/Plugins/FX/Niagara/Source/NiagaraEditor/Private/Customizations/NiagaraPlatformSetCustomization.h b/Engine/Plugins/FX/Niagara/Source/NiagaraEditor/Private/Customizations/NiagaraPlatformSetCustomization.h
new file mode 100644
--- /dev/null
+++ b/Engine/Plugins/FX/Niagara/Source/NiagaraEditor/Private/Customizations/NiagaraPlatformSetCustomization.h
@@ -0,0 +1,107 @@
+// Copyright Epic Games, Inc. All Rights Reserved.
+
+#pragma once
+
+#include "CoreMinimal.h"
+#include "IPropertyTypeCustomization.h"
+#include "NiagaraPlatformSet.h"
+#include "Layout/Visibility.h"
+#include "Widgets/Views/STreeView.h"
+
+class FDetailWidgetRow;
+class IPropertyHandleArray;
+class ITableRow;
+class SMenuAnchor;
+class STableViewBase;
+class SWrapBox;
+
+enum class ECheckBoxState : uint8;
+
+struct FNiagaraDeviceProfileViewModel
+{
+	class UDeviceProfile* Profile;
+	TArray<TSharedPtr<FNiagaraDeviceProfileViewModel>> Children;
+};
+
+class FNiagaraPlatformSetCustomization : public IPropertyTypeCustomization
+{
+public:
+	/** @return A new instance of this class */
+	static TSharedRef<IPropertyTypeCustomization> MakeInstance()
+	{
+		return MakeShared<FNiagaraPlatformSetCustomization>();
+	}
+
+	/** IPropertyTypeCustomization interface begin */
+	virtual void CustomizeHeader(TSharedRef<IPropertyHandle> StructPropertyHandle, class FDetailWidgetRow& HeaderRow, IPropertyTypeCustomizationUtils& StructCustomizationUtils) override;
+	virtual void CustomizeChildren(TSharedRef<IPropertyHandle> StructPropertyHandle, class IDetailChildrenBuilder& ChildBuilder, IPropertyTypeCustomizationUtils& StructCustomizationUtils) override;
+	/** IPropertyTypeCustomization interface end */
+private:
+
+	void CreateDeviceProfileTree();
+
+	TSharedRef<SWidget> GenerateDeviceProfileTreeWidget(int32 EffectsQuality);
+	TSharedRef<ITableRow> OnGenerateDeviceProfileTreeRow(TSharedPtr<FNiagaraDeviceProfileViewModel> InItem, const TSharedRef<STableViewBase>& OwnerTable, int32 EffectsQuality);
+	void OnGetDeviceProfileTreeChildren(TSharedPtr<FNiagaraDeviceProfileViewModel> InItem, TArray< TSharedPtr<FNiagaraDeviceProfileViewModel> >& OutChildren, int32 EffectsQuality);
+	FReply RemoveDeviceProfile(UDeviceProfile* Profile, int32 EffectsQuality);
+	TSharedRef<SWidget> GetEffectsQualityMenuContents(int32 EffectsQuality) const;
+
+	EVisibility GetDeviceProfileErrorVisibility(UDeviceProfile* Profile, int32 EffectsQuality) const;
+	FText GetDeviceProfileErrorToolTip(UDeviceProfile* Profile, int32 EffectsQuality) const;
+
+	FSlateColor GetEffectsQualityButtonTextColor(int32 EffectsQuality) const;
+	EVisibility GetEffectsQualityErrorVisibility(int32 EffectsQuality) const;
+	FText GetEffectsQualityErrorToolTip(int32 EffectsQuality) const;
+
+	const FSlateBrush* GetProfileMenuButtonImage(TSharedPtr<FNiagaraDeviceProfileViewModel> Item, int32 EffectsQuality) const;
+	EVisibility GetProfileMenuButtonVisibility(TSharedPtr<FNiagaraDeviceProfileViewModel> Item, int32 EffectsQuality) const;
+	bool GetProfileMenuItemEnabled(TSharedPtr<FNiagaraDeviceProfileViewModel> Item, int32 EffectsQuality) const;
+	FText GetProfileMenuButtonToolTip(TSharedPtr<FNiagaraDeviceProfileViewModel> Item, int32 EffectsQuality) const;
+	FReply OnProfileMenuButtonClicked(TSharedPtr<FNiagaraDeviceProfileViewModel> InItem, int32 EffectsQuality, bool bReopenMenu);
+
+	TSharedRef<SWidget> GetCurrentDeviceProfileSelectionWidget(TSharedPtr<FNiagaraDeviceProfileViewModel> ProfileView);
+	TSharedRef<SWidget> OnGenerateDeviceProfileSelectionWidget(TSharedPtr<ENiagaraPlatformSelectionState> InItem);
+
+	FText GetCurrentText() const;
+	FText GetTooltipText() const;
+
+	void GenerateEffectsQualitySelectionWidgets();
+	TSharedRef<SWidget> GenerateAdditionalDevicesWidgetForEQ(int32 EffectsQuality);
+
+	bool IsTreeActiveForEQ(const TSharedPtr<FNiagaraDeviceProfileViewModel>& Tree, int32 EffectsQualityMask) const;
+	void FilterTreeForEQ(const TSharedPtr<FNiagaraDeviceProfileViewModel>& SourceTree, TSharedPtr<FNiagaraDeviceProfileViewModel>& FilteredTree, int32 EffectsQualityMask);
+
+	ECheckBoxState IsEQChecked(int32 EffectsQuality) const;
+	void EQCheckStateChanged(ECheckBoxState CheckState, int32 EffectsQuality);
+
+	FReply ToggleMenuOpenForEffectsQuality(int32 EffectsQuality);
+
+	void UpdateCachedConflicts();
+	void InvalidateSiblingConflicts() const;
+	void OnPropertyValueChanged();
+
+private:
+	TArray<FNiagaraPlatformSetConflictInfo> CachedConflicts;
+	TSharedPtr<IPropertyHandleArray> PlatformSetArray;
+	int32 PlatformSetArrayIndex;
+
+	TArray<TSharedPtr<SMenuAnchor>> EffectsQualityMenuAnchors;
+	TArray<TSharedPtr<SWidget>> EffectsQualityMenuContents;
+	TSharedPtr<SWrapBox> EffectsQualityWidgetBox;
+
+	TSharedPtr<IPropertyHandle> PropertyHandle;
+	struct FNiagaraPlatformSet* TargetPlatformSet;
+	class UNiagaraSystem* BaseSystem;
+
+	TArray<TSharedPtr<FNiagaraDeviceProfileViewModel>> FullDeviceProfileTree;
+	TArray<TArray<TSharedPtr<FNiagaraDeviceProfileViewModel>>> FilteredDeviceProfileTrees;
+
+	TArray<TSharedPtr<ENiagaraPlatformSelectionState>> PlatformSelectionStates;
+
+	TSharedPtr<STreeView<TSharedPtr<FNiagaraDeviceProfileViewModel>>> DeviceProfileTreeWidget;
+
+	struct FNiagaraSystemScalabilitySettingsArray* SystemScalabilitySettings;
+	struct FNiagaraEmitterScalabilitySettingsArray* EmitterScalabilitySettings;
+	struct FNiagaraSystemScalabilityOverrides* SystemScalabilityOverrides;
+	struct FNiagaraEmitterScalabilityOverrides* EmitterScalabilityOverrides;
+};
diff --git a/Engine/Plugins/FX/Niagara/Source/NiagaraEditor/Private/Customizations/NiagaraTypeCustomizations.cpp b/Engine/Plugins/FX/Niagara/Source/NiagaraEditor/Private/Customizations/NiagaraTypeCustomizations.cpp
--- a/Engine/Plugins/FX/Niagara/Source/NiagaraEditor/Private/Customizations/NiagaraTypeCustomizations.cpp
+++ b/Engine/Plugins/FX/Niagara/Source/NiagaraEditor/Private/Customizations/NiagaraTypeCustomizations.cpp
@@ -1,33 +1,27 @@
 // Copyright Epic Games, Inc. All Rights Reserved.
 
 #include "NiagaraTypeCustomizations.h"
-#include "CoreMinimal.h"
 #include "DetailLayoutBuilder.h"
 #include "DetailWidgetRow.h"
 #include "IDetailChildrenBuilder.h"
 #include "NiagaraConstants.h"
-#include "NiagaraConstants.h"
 #include "NiagaraEditorStyle.h"
 #include "NiagaraEmitter.h"
 #include "NiagaraNodeOutput.h"
 #include "NiagaraNodeParameterMapBase.h"
 #include "NiagaraParameterMapHistory.h"
-#include "NiagaraPlatformSet.h"
 #include "NiagaraRendererProperties.h"
 #include "NiagaraScriptSource.h"
 #include "NiagaraScriptVariable.h"
 #include "NiagaraSystem.h"
 #include "NiagaraTypes.h"
-#include "PlatformInfo.h"
 #include "PropertyHandle.h"
 #include "SGraphActionMenu.h"
-#include "Scalability.h"
 #include "ScopedTransaction.h"
-#include "DeviceProfiles/DeviceProfile.h"
-#include "DeviceProfiles/DeviceProfileManager.h"
 #include "Framework/Application/SlateApplication.h"
 #include "Widgets/Input/SButton.h"
 #include "Widgets/Input/SCheckBox.h"
+#include "Widgets/Input/SComboButton.h"
 #include "Widgets/Input/STextComboBox.h"
 #include "Widgets/Layout/SWrapBox.h"
 #include "NiagaraSimulationStageBase.h"
@@ -36,7 +30,6 @@
 
 #define LOCTEXT_NAMESPACE "FNiagaraVariableAttributeBindingCustomization"
 
-
 void FNiagaraNumericCustomization::CustomizeHeader(TSharedRef<IPropertyHandle> PropertyHandle, FDetailWidgetRow& HeaderRow, IPropertyTypeCustomizationUtils& CustomizationUtils)
 {
 	TSharedPtr<IPropertyHandle> ValueHandle = PropertyHandle->GetChildHandle(TEXT("Value"));
@@ -916,939 +909,5 @@ void FNiagaraScriptVariableBindingCustomization::CustomizeHeader(TSharedRef<IPro
 			];
 	}
 }
-//////////////////////////////////////////////////////////////////////////
-
-void FNiagaraPlatformSetTypeCustomization::CustomizeHeader(TSharedRef<IPropertyHandle> InPropertyHandle, FDetailWidgetRow& HeaderRow, IPropertyTypeCustomizationUtils& CustomizationUtils)
-{
-	PropertyHandle = InPropertyHandle;
-	TArray<UObject*> Objects;
-	PropertyHandle->GetOuterObjects(Objects);
-	TargetPlatformSet = (FNiagaraPlatformSet*)PropertyHandle->GetValueBaseAddress((uint8*)Objects[0]);
-
-	if (PlatformSelectionStates.Num() == 0)
-	{
-		PlatformSelectionStates.Add(MakeShared<ENiagaraPlatformSelectionState>(ENiagaraPlatformSelectionState::Enabled));
-		PlatformSelectionStates.Add(MakeShared<ENiagaraPlatformSelectionState>(ENiagaraPlatformSelectionState::Disabled));
-		PlatformSelectionStates.Add(MakeShared<ENiagaraPlatformSelectionState>(ENiagaraPlatformSelectionState::Default));
-	}
-
-	//Look for outer types for which we need to ensure there are no conflicting settings.
-	SystemScalabilitySettings = nullptr;
-	EmitterScalabilitySettings = nullptr;
-	SystemScalabilityOverrides = nullptr;
-	EmitterScalabilityOverrides = nullptr;
-	PlatformSetArray.Reset();
-	PlatformSetArrayIndex = INDEX_NONE;
-
-	//Look whether this platform set belongs to a class which must keep orthogonal platform sets.
-	//We then interrogate the other sets via these ptrs to look for conflicts.
-	TSharedPtr<IPropertyHandle> CurrHandle = PropertyHandle->GetParentHandle();
-	while (CurrHandle)
-	{
-		int32 ThisIndex = CurrHandle->GetIndexInArray();
-		if (ThisIndex != INDEX_NONE)
-		{
-			PlatformSetArray = CurrHandle->GetParentHandle()->AsArray();
-			PlatformSetArrayIndex = ThisIndex;
-		}
-
-		if (FProperty* CurrProperty = CurrHandle->GetProperty())
-		{
-			if (UStruct* CurrStruct = CurrProperty->GetOwnerStruct())
-			{
-				TSharedPtr<IPropertyHandle> ParentHandle = CurrHandle->GetParentHandle();
-				if (CurrStruct == FNiagaraSystemScalabilitySettingsArray::StaticStruct())
-				{
-					SystemScalabilitySettings = (FNiagaraSystemScalabilitySettingsArray*)ParentHandle->GetValueBaseAddress((uint8*)Objects[0]);
-					break;
-				}
-				else if (CurrStruct == FNiagaraEmitterScalabilitySettingsArray::StaticStruct())
-				{
-					EmitterScalabilitySettings = (FNiagaraEmitterScalabilitySettingsArray*)ParentHandle->GetValueBaseAddress((uint8*)Objects[0]);
-					break;
-				}
-				else if (CurrStruct == FNiagaraSystemScalabilityOverrides::StaticStruct())
-				{
-					SystemScalabilityOverrides = (FNiagaraSystemScalabilityOverrides*)ParentHandle->GetValueBaseAddress((uint8*)Objects[0]);
-					break;
-				}
-				else if (CurrStruct == FNiagaraEmitterScalabilityOverrides::StaticStruct())
-				{
-					EmitterScalabilityOverrides = (FNiagaraEmitterScalabilityOverrides*)ParentHandle->GetValueBaseAddress((uint8*)Objects[0]);
-					break;
-				}
-			}
-		}
-
-		CurrHandle = CurrHandle->GetParentHandle();
-	}
-
-	UpdateCachedConflicts();
-	
-	BaseSystem = Objects[0]->GetTypedOuter<UNiagaraSystem>();
-	PropertyHandle->SetOnPropertyValueChanged(FSimpleDelegate::CreateSP(this, &FNiagaraPlatformSetTypeCustomization::OnPropertyValueChanged));
-
-	HeaderRow
-		.WholeRowContent()
-		[ 	
-			SAssignNew(EffectsQualityWidgetBox, SWrapBox)
-			.UseAllottedWidth(true)
-		];
-
-	GenerateEffectsQualitySelectionWidgets();
-}
-
-void FNiagaraPlatformSetTypeCustomization::CustomizeChildren(TSharedRef<IPropertyHandle> StructPropertyHandle, class IDetailChildrenBuilder& ChildBuilder, IPropertyTypeCustomizationUtils& StructCustomizationUtils)
-{
-}
-
-FText FNiagaraPlatformSetTypeCustomization::GetCurrentText() const
-{
-	return LOCTEXT("Platforms", "Platforms");
-}
-
-FText FNiagaraPlatformSetTypeCustomization::GetTooltipText() const
-{
-	return LOCTEXT("Platforms", "Platforms");
-}
-
-void FNiagaraPlatformSetTypeCustomization::GenerateEffectsQualitySelectionWidgets()
-{
-	EffectsQualityWidgetBox->ClearChildren();
-
-	Scalability::FQualityLevels Counts = Scalability::GetQualityLevelCounts();
-	EffectsQualityMenus.SetNum(Counts.EffectsQuality);
-
-	for (int32 EffectsQuality = 0; EffectsQuality < Counts.EffectsQuality; ++EffectsQuality)
-	{
-		bool First = EffectsQuality == 0;
-		bool Last = EffectsQuality == (Counts.EffectsQuality - 1);
-
-		EffectsQualityWidgetBox->AddSlot()
-			.Padding(0, 0, 1, 0)
-			[
-				SNew(SBox)
-				.WidthOverride(100)
-				.VAlign(VAlign_Top)
-				[
-					SNew(SVerticalBox)
-					+ SVerticalBox::Slot()
-					.AutoHeight()
-					[
-						SNew(SCheckBox)
-						.Style(FNiagaraEditorStyle::Get(), First ? "NiagaraEditor.PlatformSet.StartButton" : 
-							(Last ? "NiagaraEditor.PlatformSet.EndButton" : "NiagaraEditor.PlatformSet.MiddleButton"))
-						.IsChecked(this, &FNiagaraPlatformSetTypeCustomization::IsEQChecked, EffectsQuality)
-						.OnCheckStateChanged(this, &FNiagaraPlatformSetTypeCustomization::EQCheckStateChanged, EffectsQuality)
-						[
-							SNew(SHorizontalBox)
-							+ SHorizontalBox::Slot()
-							.HAlign(HAlign_Fill)
-							.Padding(6,2,2,4)
-							[
-								SNew(STextBlock)
-								.TextStyle(FNiagaraEditorStyle::Get(), "NiagaraEditor.PlatformSet.ButtonText")
-								.Text(FNiagaraPlatformSet::GetEffectsQualityText(EffectsQuality))
-								.ColorAndOpacity(this, &FNiagaraPlatformSetTypeCustomization::GetEffectsQualityButtonTextColor, EffectsQuality)
-								.ShadowOffset(FVector2D(1, 1))
-							]
-							// error icon
-							+ SHorizontalBox::Slot()
-							.AutoWidth()
-							.HAlign(HAlign_Right)
-							.VAlign(VAlign_Center)
-							.Padding(0,0,2,0)
-							[
-								SNew(SBox)
-								.WidthOverride(12)
-								.HeightOverride(12)
-								[
-									SNew(SImage)
-									.ToolTipText(this, &FNiagaraPlatformSetTypeCustomization::GetEffectsQualityErrorToolTip, EffectsQuality)
-									.Visibility(this, &FNiagaraPlatformSetTypeCustomization::GetEffectsQualityErrorVisibility, EffectsQuality)
-									.Image(FEditorStyle::GetBrush("Icons.Error"))
-								]
-							]
-							// dropdown button
-							+ SHorizontalBox::Slot()
-							.AutoWidth()
-							.HAlign(HAlign_Right)
-							.VAlign(VAlign_Center)
-							.Padding(0,0,2,0)
-							[
-								// dropdown button
-								SAssignNew(EffectsQualityMenus[EffectsQuality], SMenuAnchor)
-								.ToolTipText(LOCTEXT("AddPlatformOverride", "Add an override for a specific platform."))
-								.OnGetMenuContent(this, &FNiagaraPlatformSetTypeCustomization::GenerateDeviceProfileTreeWidget, EffectsQuality)
-								[
-									SNew(SButton)
-									.Visibility_Lambda([this, EffectsQuality]()
-									{
-										return EffectsQualityWidgetBox->GetChildren()->GetChildAt(EffectsQuality)->IsHovered() ? EVisibility::Visible : EVisibility::Hidden;
-									})
-									.ButtonStyle(FEditorStyle::Get(), "HoverHintOnly")
-									.ForegroundColor(FSlateColor::UseForeground())
-									.OnClicked(this, &FNiagaraPlatformSetTypeCustomization::ToggleMenuOpenForEffectsQuality, EffectsQuality)
-									[
-										SNew(SBox)
-										.WidthOverride(8)
-										.HeightOverride(8)
-										.VAlign(VAlign_Center)
-										[
-											SNew(SImage)
-											.Image(FNiagaraEditorStyle::Get().GetBrush("NiagaraEditor.PlatformSet.DropdownButton"))
-										]
-									]
-								]
-							]
-						]
-					]
-					+ SVerticalBox::Slot()
-					.AutoHeight()
-					[
-						GenerateAdditionalDevicesWidgetForEQ(EffectsQuality)
-					]
-				]
-			];
-	}
-}
-
-void FNiagaraPlatformSetTypeCustomization::UpdateCachedConflicts()
-{
-	TArray<const FNiagaraPlatformSet*> PlatformSets;
-	if (SystemScalabilityOverrides != nullptr)
-	{
-		for (FNiagaraSystemScalabilityOverride& Override : SystemScalabilityOverrides->Overrides)
-		{
-			PlatformSets.Add(&Override.Platforms);
-		}
-	}
-	else if (EmitterScalabilityOverrides != nullptr)
-	{
-		for (FNiagaraEmitterScalabilityOverride& Override : EmitterScalabilityOverrides->Overrides)
-		{
-			PlatformSets.Add(&Override.Platforms);
-		}
-	}
-	else if (SystemScalabilitySettings != nullptr)
-	{
-		for (FNiagaraSystemScalabilitySettings& Settings : SystemScalabilitySettings->Settings)
-		{
-			PlatformSets.Add(&Settings.Platforms);
-		}
-	}
-	else if (EmitterScalabilitySettings != nullptr)
-	{
-		for (FNiagaraEmitterScalabilitySettings& Settings : EmitterScalabilitySettings->Settings)
-		{
-			PlatformSets.Add(&Settings.Platforms);
-		}
-	}
-
-	CachedConflicts.Reset();
-	FNiagaraPlatformSet::GatherConflicts(PlatformSets, CachedConflicts);
-}
-
-static TSharedPtr<IPropertyHandle> FindChildPlatformSet(TSharedPtr<IPropertyHandle> PropertyHandle)
-{
-	if (FStructProperty* Property = CastField<FStructProperty>(PropertyHandle->GetProperty()))
-	{
-		if (Property->Struct == FNiagaraPlatformSet::StaticStruct())
-		{
-			return PropertyHandle;
-		}
-	}
-
-	// recurse
-	uint32 NumChildren = 0;
-	PropertyHandle->GetNumChildren(NumChildren);
-
-	for (int32 Idx = 0; Idx < (int32) NumChildren; ++Idx)
-	{
-		TSharedPtr<IPropertyHandle> Child = PropertyHandle->GetChildHandle(Idx);
-
-		TSharedPtr<IPropertyHandle> ChildResult = FindChildPlatformSet(Child);
-		if (ChildResult.IsValid())
-		{
-			return ChildResult;
-		}
-	}
-
-	return TSharedPtr<IPropertyHandle>();
-}
-
-void FNiagaraPlatformSetTypeCustomization::InvalidateSiblingConflicts() const
-{
-	if (!PlatformSetArray.IsValid())
-	{
-		return;
-	}
-
-	uint32 ArrayCount = 0;
-	PlatformSetArray->GetNumElements(ArrayCount);
-	for (int32 Idx = 0; Idx < (int32) ArrayCount; ++Idx)
-	{
-		if (Idx == PlatformSetArrayIndex)
-		{
-			continue; // skip self
-		}
-
-		TSharedRef<IPropertyHandle> Sibling = PlatformSetArray->GetElement(Idx);
-
-		TSharedPtr<IPropertyHandle> SiblingPlatformSet = FindChildPlatformSet(Sibling);
-		if (SiblingPlatformSet.IsValid())
-		{
-			SiblingPlatformSet->NotifyPostChange();
-		}
-	}
-}
-
-EVisibility FNiagaraPlatformSetTypeCustomization::GetEffectsQualityErrorVisibility(int32 EffectsQuality) const
-{
-	// not part of an array
-	if (PlatformSetArrayIndex == INDEX_NONE)
-	{
-		return EVisibility::Collapsed;
-	}
-
-	for (const FNiagaraPlatformSetConflictInfo& ConflictInfo : CachedConflicts)
-	{
-		if (ConflictInfo.SetAIndex == PlatformSetArrayIndex || 
-			ConflictInfo.SetBIndex == PlatformSetArrayIndex)
-		{
-			// this conflict applies to this platform set, check if it applies to this effects quality button
-			const int32 EQMask = FNiagaraPlatformSet::CreateEQMask(EffectsQuality);
-
-			for (const FNiagaraPlatformSetConflictEntry& Conflict : ConflictInfo.Conflicts)
-			{
-				if ((EQMask & Conflict.EffectsQualityMask) != 0)
-				{
-					return EVisibility::Visible;
-				}
-			}
-		}
-	}
-
-	return EVisibility::Collapsed;
-}
-
-FText FNiagaraPlatformSetTypeCustomization::GetEffectsQualityErrorToolTip(int32 EffectsQuality) const
-{
-	if (PlatformSetArrayIndex == INDEX_NONE)
-	{
-		return FText::GetEmpty();
-	}
-
-	for (const FNiagaraPlatformSetConflictInfo& ConflictInfo : CachedConflicts)
-	{
-		int32 OtherIndex = INDEX_NONE;
-
-		if (ConflictInfo.SetAIndex == PlatformSetArrayIndex)
-		{
-			OtherIndex = ConflictInfo.SetBIndex;
-		}
-
-		if (ConflictInfo.SetBIndex == PlatformSetArrayIndex)
-		{
-			OtherIndex = ConflictInfo.SetAIndex;
-		}
-
-		if (OtherIndex != INDEX_NONE)
-		{
-			const int32 EQMask = FNiagaraPlatformSet::CreateEQMask(EffectsQuality);
-
-			for (const FNiagaraPlatformSetConflictEntry& Conflict : ConflictInfo.Conflicts)
-			{
-				if ((EQMask & Conflict.EffectsQualityMask) != 0)
-				{
-					FText FormatString = LOCTEXT("EffectsQualityConflictToolTip", "This effect quality conflicts with the set at index {Index} in this array."); 
-
-					FFormatNamedArguments Args;
-					Args.Add(TEXT("Index"), OtherIndex);
-
-					return FText::Format(FormatString, Args);
-				}
-			}
-		}
-	}
-
-	return FText::GetEmpty();
-}
-
-FSlateColor FNiagaraPlatformSetTypeCustomization::GetEffectsQualityButtonTextColor(int32 EffectsQuality) const
-{
-	return TargetPlatformSet->IsEffectQualityEnabled(EffectsQuality) ? 
-		FSlateColor(FLinearColor(0.95f, 0.95f, 0.95f)) :
-		FSlateColor::UseForeground();
-}
-
-TSharedRef<SWidget> FNiagaraPlatformSetTypeCustomization::GenerateAdditionalDevicesWidgetForEQ(int32 EffectsQuality)
-{
-	TSharedRef<SVerticalBox> Container = SNew(SVerticalBox);
-
-	auto AddDeviceProfileOverrideWidget = [&](UDeviceProfile* Profile, bool bEnabled)
-	{
-		TSharedPtr<SHorizontalBox> DeviceBox;
-
-		Container->AddSlot()
-			.AutoHeight()
-			[
-				SAssignNew(DeviceBox, SHorizontalBox)
-			];
-
-		const FText DeviceNameText = FText::FromName(Profile->GetFName());
-
-		DeviceBox->AddSlot()
-			.AutoWidth()
-			.HAlign(HAlign_Left)
-			.VAlign(VAlign_Center)
-			.Padding(3, 0, 3, 0)
-			[
-				SNew(SBox)
-				.WidthOverride(8)
-				.HeightOverride(8)
-				.ToolTipText(bEnabled ?
-					LOCTEXT("DeviceIncludedToolTip", "This device is included.") :
-					LOCTEXT("DeviceExcludedToolTip", "This device is excluded."))
-				[
-					SNew(SImage)
-					.Image(bEnabled ?
-						FNiagaraEditorStyle::Get().GetBrush("NiagaraEditor.PlatformSet.Include") :
-						FNiagaraEditorStyle::Get().GetBrush("NiagaraEditor.PlatformSet.Exclude"))
-					.ColorAndOpacity(bEnabled ?
-						FLinearColor(0, 1, 0) :
-						FLinearColor(1, 0, 0))
-					]
-				];
-
-		DeviceBox->AddSlot()
-			.HAlign(HAlign_Fill)
-			.Padding(0, 0, 2, 0)
-			[
-				SNew(STextBlock)
-				.Text(DeviceNameText)
-				.ToolTipText(DeviceNameText)
-			];
-
-		DeviceBox->AddSlot()
-			.AutoWidth()
-			.HAlign(HAlign_Right)
-			.VAlign(VAlign_Center)
-			.Padding(0, 0, 2, 0)
-			[
-				SNew(SBox)
-				.WidthOverride(12)
-				.HeightOverride(12)
-				[
-					SNew(SImage)
-					.ToolTipText(this, &FNiagaraPlatformSetTypeCustomization::GetDeviceProfileErrorToolTip, Profile, EffectsQuality)
-					.Visibility(this, &FNiagaraPlatformSetTypeCustomization::GetDeviceProfileErrorVisibility, Profile, EffectsQuality)
-					.Image(FEditorStyle::GetBrush("Icons.Error"))
-				]
-			];
-
-		DeviceBox->AddSlot()
-			.AutoWidth()
-			.HAlign(HAlign_Right)
-			.VAlign(VAlign_Center)
-			.Padding(0, 0, 2, 0)
-			[
-				SNew(SButton)
-				.ButtonStyle(FEditorStyle::Get(), "HoverHintOnly")
-				.ForegroundColor(FSlateColor::UseForeground())
-				.OnClicked(this, &FNiagaraPlatformSetTypeCustomization::RemoveDeviceProfile, Profile, EffectsQuality)
-				.ToolTipText(LOCTEXT("RemoveDevice", "Remove this device override."))
-				[
-					SNew(SBox)
-					.WidthOverride(8)
-					.HeightOverride(8)
-					.VAlign(VAlign_Center)
-					[
-						SNew(SImage)
-						.Image(FNiagaraEditorStyle::Get().GetBrush("NiagaraEditor.PlatformSet.Remove"))
-					]
-				]
-			];
-		};
-
-	TArray<UDeviceProfile*> EnabledProfiles;
-	TArray<UDeviceProfile*> DisabledProfiles;
-	TargetPlatformSet->GetOverridenDeviceProfiles(EffectsQuality, EnabledProfiles, DisabledProfiles);
-
-	for (UDeviceProfile* Profile : EnabledProfiles)
-	{
-		AddDeviceProfileOverrideWidget(Profile, true);
-	}
-
-	for (UDeviceProfile* Profile : DisabledProfiles)
-	{
-		AddDeviceProfileOverrideWidget(Profile, false);
-	}
-
-	return Container;
-}
-
-FReply FNiagaraPlatformSetTypeCustomization::RemoveDeviceProfile(UDeviceProfile* Profile, int32 EffectsQuality)
-{
-	PropertyHandle->NotifyPreChange();
-	TargetPlatformSet->SetDeviceProfileState(Profile, EffectsQuality, ENiagaraPlatformSelectionState::Default);
-	PropertyHandle->NotifyPostChange(EPropertyChangeType::ValueSet);
-
-	InvalidateSiblingConflicts();
-
-	return FReply::Handled();
-}
-
-EVisibility FNiagaraPlatformSetTypeCustomization::GetDeviceProfileErrorVisibility(UDeviceProfile* Profile, int32 EffectsQuality) const
-{
-	// not part of an array
-	if (PlatformSetArrayIndex == INDEX_NONE)
-	{
-		return EVisibility::Collapsed;
-	}
-
-	for (const FNiagaraPlatformSetConflictInfo& ConflictInfo : CachedConflicts)
-	{
-		if (ConflictInfo.SetAIndex == PlatformSetArrayIndex || 
-			ConflictInfo.SetBIndex == PlatformSetArrayIndex)
-		{
-			const int32 EQMask = FNiagaraPlatformSet::CreateEQMask(EffectsQuality);
-
-			for (const FNiagaraPlatformSetConflictEntry& Entry : ConflictInfo.Conflicts)
-			{
-				if ((Entry.EffectsQualityMask & EQMask) != 0 &&
-					Entry.ProfileName == Profile->GetFName())
-				{
-					return EVisibility::Visible;
-				}
-			}
-		}
-	}
-
-	return EVisibility::Collapsed;
-}
-
-FText FNiagaraPlatformSetTypeCustomization::GetDeviceProfileErrorToolTip(UDeviceProfile* Profile, int32 EffectsQuality) const
-{
-	for (const FNiagaraPlatformSetConflictInfo& ConflictInfo : CachedConflicts)
-	{
-		int32 OtherIndex = INDEX_NONE;
-
-		if (ConflictInfo.SetAIndex == PlatformSetArrayIndex)
-		{
-			OtherIndex = ConflictInfo.SetBIndex;
-		}
-
-		if (ConflictInfo.SetBIndex == PlatformSetArrayIndex)
-		{
-			OtherIndex = ConflictInfo.SetAIndex;
-		}
-
-		if (OtherIndex != INDEX_NONE)
-		{
-			const int32 EQMask = FNiagaraPlatformSet::CreateEQMask(EffectsQuality);
-
-			for (const FNiagaraPlatformSetConflictEntry& Entry : ConflictInfo.Conflicts)
-			{
-				if ((Entry.EffectsQualityMask & EQMask) != 0 &&
-					Entry.ProfileName == Profile->GetFName())
-				{
-					FText FormatString = LOCTEXT("PlatformOverrideConflictToolTip", "This platform override conflicts with the set at index {Index} in this array."); 
-
-					FFormatNamedArguments Args;
-					Args.Add(TEXT("Index"), OtherIndex);
-
-					return FText::Format(FormatString, Args);
-				}
-			}
-		}
-	}
-
-	return FText::GetEmpty();
-}
-
-FReply FNiagaraPlatformSetTypeCustomization::ToggleMenuOpenForEffectsQuality(int32 EffectsQuality)
-{
-	check(EffectsQualityMenus.IsValidIndex(EffectsQuality));
-	
-	TSharedPtr<SMenuAnchor> MenuAnchor = EffectsQualityMenus[EffectsQuality];
-	MenuAnchor->SetIsOpen(!MenuAnchor->IsOpen());
-
-	return FReply::Handled();
-}
-
-// Is or does a viewmodel contain any children active at the given quality?
-bool FNiagaraPlatformSetTypeCustomization::IsTreeActiveForEQ(const TSharedPtr<FNiagaraDeviceProfileViewModel>& Tree, int32 EffectsQualityMask) const
-{
-	int32 Mask = TargetPlatformSet->GetEffectQualityMaskForDeviceProfile(Tree->Profile);
-	if ((Mask & EffectsQualityMask) != 0)
-	{
-		return true;
-	}
-
-	for (const TSharedPtr<FNiagaraDeviceProfileViewModel>& Child : Tree->Children)
-	{
-		if (IsTreeActiveForEQ(Child, EffectsQualityMask))
-		{
-			return true;
-		}
-	}
-
-	return false;
-}
-
-void FNiagaraPlatformSetTypeCustomization::FilterTreeForEQ(const TSharedPtr<FNiagaraDeviceProfileViewModel>& SourceTree, TSharedPtr<FNiagaraDeviceProfileViewModel>& FilteredTree, int32 EffectsQualityMask)
-{
-	for (const TSharedPtr<FNiagaraDeviceProfileViewModel>& SourceChild : SourceTree->Children)
-	{
-		if (IsTreeActiveForEQ(SourceChild, EffectsQualityMask))
-		{
-			TSharedPtr<FNiagaraDeviceProfileViewModel>& FilteredChild = FilteredTree->Children.Add_GetRef(MakeShared<FNiagaraDeviceProfileViewModel>());
-			FilteredChild->Profile = SourceChild->Profile;
-
-			FilterTreeForEQ(SourceChild, FilteredChild, EffectsQualityMask);
-		}
-	}
-}
-
-void FNiagaraPlatformSetTypeCustomization::CreateDeviceProfileTree()
-{
-	//Pull device profiles out by their hierarchy depth.
-	TArray<TArray<UDeviceProfile*>> DeviceProfilesByHierarchyLevel;
-	for (UObject* Profile : UDeviceProfileManager::Get().Profiles)
-	{
-		UDeviceProfile* DeviceProfile = CastChecked<UDeviceProfile>(Profile);
-
-		TFunction<void(int32&, UDeviceProfile*)> FindDepth;
-		FindDepth = [&](int32& Depth, UDeviceProfile* CurrProfile)
-		{
-			if (CurrProfile->Parent)
-			{
-				FindDepth(++Depth, Cast<UDeviceProfile>(CurrProfile->Parent));
-			}
-		};
-
-		int32 ProfileDepth = 0;
-		FindDepth(ProfileDepth, DeviceProfile);
-		DeviceProfilesByHierarchyLevel.SetNum(FMath::Max(ProfileDepth+1, DeviceProfilesByHierarchyLevel.Num()));
-		DeviceProfilesByHierarchyLevel[ProfileDepth].Add(DeviceProfile);
-	}
-	
-	FullDeviceProfileTree.Reset(DeviceProfilesByHierarchyLevel[0].Num());
-	for (int32 RootProfileIdx = 0; RootProfileIdx < DeviceProfilesByHierarchyLevel[0].Num(); ++RootProfileIdx)
-	{
-		TFunction<void(FNiagaraDeviceProfileViewModel*, int32)> BuildProfileTree;
-		BuildProfileTree = [&](FNiagaraDeviceProfileViewModel* CurrRoot, int32 CurrLevel)
-		{
-			int32 NextLevel = CurrLevel + 1;
-			if (NextLevel < DeviceProfilesByHierarchyLevel.Num())
-			{
-				for (UDeviceProfile* PossibleChild : DeviceProfilesByHierarchyLevel[NextLevel])
-				{
-					if (PossibleChild->Parent == CurrRoot->Profile)
-					{
-						TSharedPtr<FNiagaraDeviceProfileViewModel>& NewChild = CurrRoot->Children.Add_GetRef(MakeShared<FNiagaraDeviceProfileViewModel>());
-						NewChild->Profile = PossibleChild;
-						BuildProfileTree(NewChild.Get(), NextLevel);
-					}
-				}
-			}
-		};
-
-		//Add all root nodes and build their trees.
-		TSharedPtr<FNiagaraDeviceProfileViewModel> CurrRoot = MakeShared<FNiagaraDeviceProfileViewModel>();
-		CurrRoot->Profile = DeviceProfilesByHierarchyLevel[0][RootProfileIdx];
-		BuildProfileTree(CurrRoot.Get(), 0);
-		FullDeviceProfileTree.Add(CurrRoot);
-	}
-
-	Scalability::FQualityLevels Counts = Scalability::GetQualityLevelCounts();
-	FilteredDeviceProfileTrees.SetNum(Counts.EffectsQuality);
-	
-	for (TSharedPtr<FNiagaraDeviceProfileViewModel>& FullDeviceRoot : FullDeviceProfileTree)
-	{
-		for (int32 EffectsQuality = 0; EffectsQuality < Counts.EffectsQuality; ++EffectsQuality)
-		{
-			int32 EffectsQualityMask = FNiagaraPlatformSet::CreateEQMask(EffectsQuality);
-
-			if (IsTreeActiveForEQ(FullDeviceRoot, EffectsQualityMask))
-			{
-				TArray<TSharedPtr<FNiagaraDeviceProfileViewModel>>& FilteredRoots = FilteredDeviceProfileTrees[EffectsQuality];
-
-				TSharedPtr<FNiagaraDeviceProfileViewModel>& FilteredRoot = FilteredRoots.Add_GetRef(MakeShared<FNiagaraDeviceProfileViewModel>());
-				FilteredRoot->Profile = FullDeviceRoot->Profile;
-
-				FilterTreeForEQ(FullDeviceRoot, FilteredRoot, EffectsQualityMask);
-			}
-		}
-	}
-}
-
-TSharedRef<SWidget> FNiagaraPlatformSetTypeCustomization::GenerateDeviceProfileTreeWidget(int32 EffectsQuality)
-{
-	if (FullDeviceProfileTree.Num() == 0)
-	{
-		CreateDeviceProfileTree();	
-	}
-
-	TArray<TSharedPtr<FNiagaraDeviceProfileViewModel>>* TreeToUse = &FullDeviceProfileTree;
-	if (EffectsQuality != INDEX_NONE)
-	{
-		check(EffectsQuality < FilteredDeviceProfileTrees.Num());
-		TreeToUse = &FilteredDeviceProfileTrees[EffectsQuality];
-	}
-
-	return SNew(SBorder)
-		.BorderImage(FEditorStyle::Get().GetBrush("Menu.Background"))
-		[
-			SAssignNew(DeviceProfileTreeWidget, STreeView<TSharedPtr<FNiagaraDeviceProfileViewModel>>)
-			.TreeItemsSource(TreeToUse)
-			.OnGenerateRow(this, &FNiagaraPlatformSetTypeCustomization::OnGenerateDeviceProfileTreeRow, EffectsQuality)
-			.OnGetChildren(this, &FNiagaraPlatformSetTypeCustomization::OnGetDeviceProfileTreeChildren, EffectsQuality)
-			.SelectionMode(ESelectionMode::None)
-		];
-}
-
-TSharedRef<ITableRow> FNiagaraPlatformSetTypeCustomization::OnGenerateDeviceProfileTreeRow(TSharedPtr<FNiagaraDeviceProfileViewModel> InItem, const TSharedRef<STableViewBase>& OwnerTable, int32 EffectsQuality)
-{
-	TSharedPtr<SHorizontalBox> RowContainer;
-	SAssignNew(RowContainer, SHorizontalBox);
-
-	int32 ProfileMask = TargetPlatformSet->GetEffectQualityMaskForDeviceProfile(InItem->Profile);
-	FText NameTooltip = FText::Format(LOCTEXT("ProfileEQTooltipFmt", "Effects Quality: {0}"), FNiagaraPlatformSet::GetEffectsQualityMaskText(ProfileMask));
-	
-	//Top level profile. Look for a platform icon.
-	if (InItem->Profile->Parent == nullptr)
-	{
-		if (const PlatformInfo::FPlatformInfo* Info = PlatformInfo::FindPlatformInfo(*InItem->Profile->DeviceType))
-		{
-			const FSlateBrush* DeviceProfileTypeIcon = FEditorStyle::GetBrush(Info->GetIconStyleName(PlatformInfo::EPlatformIconSize::Normal));
-			if (DeviceProfileTypeIcon != FEditorStyle::Get().GetDefaultBrush())
-			{
-				RowContainer->AddSlot()
-					.AutoWidth()
-					.VAlign(VAlign_Center)
-					.HAlign(HAlign_Left)
-					.Padding(4, 0, 0, 0)
-					[
-						SNew(SBox)
-						.WidthOverride(16)
-						.HeightOverride(16)
-						[
-							SNew(SImage)
-							.Image(DeviceProfileTypeIcon)
-						]
-					];
-			}
-		}
-	}
-
-	FName TextStyleName("NormalText");
-	FSlateColor TextColor(FSlateColor::UseForeground());
-	ENiagaraPlatformSelectionState CurrentState = TargetPlatformSet->GetDeviceProfileState(InItem->Profile, EffectsQuality);
-	
-	if (CurrentState == ENiagaraPlatformSelectionState::Enabled)
-	{
-		TextStyleName = "RichTextBlock.Bold";
-		TextColor = FSlateColor(FLinearColor(FVector4(0,1,0,1)));
-	}
-	else if (CurrentState == ENiagaraPlatformSelectionState::Disabled)
-	{
-		TextStyleName = "RichTextBlock.Italic";
-		TextColor = FSlateColor(FLinearColor(FVector4(1,0,0,1)));
-	}
-
-
-	int32 EffectsQualityMask = FNiagaraPlatformSet::CreateEQMask(EffectsQuality);
-	if ((ProfileMask & EffectsQualityMask) == 0)
-	{
-		TextColor = FSlateColor::UseSubduedForeground();
-	}
-
-	RowContainer->AddSlot()
-		.Padding(4, 2, 0, 2)
-		.HAlign(HAlign_Fill)
-		.VAlign(VAlign_Center)
-		[
-			SNew(STextBlock)
-			.TextStyle(FEditorStyle::Get(), TextStyleName)
-			.Text(FText::FromString(InItem->Profile->GetName()))
-			.ToolTipText(NameTooltip)
-			.ColorAndOpacity(TextColor)
-		];
-
-	RowContainer->AddSlot()
-		.AutoWidth()
-		.Padding(12, 2, 4, 4)
-		.HAlign(HAlign_Right)
-		.VAlign(VAlign_Center)
-		[
-			SNew(SButton)
-			.ButtonStyle(FEditorStyle::Get(), "HoverHintOnly")
-			.HAlign(HAlign_Center)
-			.VAlign(VAlign_Fill)
-			.Visibility(this, &FNiagaraPlatformSetTypeCustomization::GetProfileMenuButtonVisibility, InItem, EffectsQuality)
-			.OnClicked(this, &FNiagaraPlatformSetTypeCustomization::OnProfileMenuButtonClicked, InItem, EffectsQuality)
-			.ToolTipText(this, &FNiagaraPlatformSetTypeCustomization::GetProfileMenuButtonToolTip, InItem, EffectsQuality)
-			[
-				SNew(SBox)
-				.WidthOverride(8)
-				.HeightOverride(8)
-				.VAlign(VAlign_Center)
-				.HAlign(HAlign_Center)
-				[
-					SNew(SImage)
-					.Image(this, &FNiagaraPlatformSetTypeCustomization::GetProfileMenuButtonImage, InItem, EffectsQuality)
-				]
-			]
-		];
-
-	return SNew(STableRow<TSharedPtr<FNiagaraDeviceProfileViewModel>>, OwnerTable)
-		.Style(FNiagaraEditorStyle::Get(), "NiagaraEditor.PlatformSet.TreeView")
-		[
-			RowContainer.ToSharedRef()
-		];
-}
-
-enum class EProfileButtonMode
-{
-	None,
-	Include,
-	Exclude,
-	Remove
-};
-
-static EProfileButtonMode GetProfileMenuButtonMode(FNiagaraPlatformSet* PlatformSet, TSharedPtr<FNiagaraDeviceProfileViewModel> Item, int32 EffectsQuality)
-{
-	int32 Mask = PlatformSet->GetEffectQualityMaskForDeviceProfile(Item->Profile);
-	int32 EQMask = FNiagaraPlatformSet::CreateEQMask(EffectsQuality);
-
-	if ((Mask & EQMask) == 0)
-	{
-		return EProfileButtonMode::None;
-	}
-
-	ENiagaraPlatformSelectionState CurrentState = PlatformSet->GetDeviceProfileState(Item->Profile, EffectsQuality);
-
-	bool bQualityEnabled = PlatformSet->IsEffectQualityEnabled(EffectsQuality);
-	bool bIsDefault = CurrentState == ENiagaraPlatformSelectionState::Default;
-
-	if (bIsDefault && bQualityEnabled)
-	{
-		return EProfileButtonMode::Exclude;
-	}
-
-	if (bIsDefault && !bQualityEnabled)
-	{
-		return EProfileButtonMode::Include;
-	}
-
-	return EProfileButtonMode::Remove;
-}
-
-
-FText FNiagaraPlatformSetTypeCustomization::GetProfileMenuButtonToolTip(TSharedPtr<FNiagaraDeviceProfileViewModel> Item, int32 EffectsQuality) const
-{
-	EProfileButtonMode Mode = GetProfileMenuButtonMode(TargetPlatformSet, Item, EffectsQuality);
-	switch(Mode)
-	{
-		case EProfileButtonMode::Include:
-			return LOCTEXT("IncludePlatform", "Include this platform.");
-		case EProfileButtonMode::Exclude:
-			return LOCTEXT("ExcludePlatform", "Exclude this platform.");
-		case EProfileButtonMode::Remove:
-			return LOCTEXT("RemovePlatform", "Remove this platform override.");
-	}
-
-	return FText::GetEmpty();
-}
-
-EVisibility FNiagaraPlatformSetTypeCustomization::GetProfileMenuButtonVisibility(TSharedPtr<FNiagaraDeviceProfileViewModel> Item, int32 EffectsQuality) const
-{
-	EProfileButtonMode Mode = GetProfileMenuButtonMode(TargetPlatformSet, Item, EffectsQuality);
-	if (Mode == EProfileButtonMode::None)
-	{
-		return EVisibility::Collapsed;
-	}
-
-	return EVisibility::Visible;
-}
-
-const FSlateBrush* FNiagaraPlatformSetTypeCustomization::GetProfileMenuButtonImage(TSharedPtr<FNiagaraDeviceProfileViewModel> Item, int32 EffectsQuality) const
-{
-	EProfileButtonMode Mode = GetProfileMenuButtonMode(TargetPlatformSet, Item, EffectsQuality);
-	switch(Mode)
-	{
-		case EProfileButtonMode::Include:
-			return FNiagaraEditorStyle::Get().GetBrush("NiagaraEditor.PlatformSet.Include");
-		case EProfileButtonMode::Exclude:
-			return FNiagaraEditorStyle::Get().GetBrush("NiagaraEditor.PlatformSet.Exclude");
-		case EProfileButtonMode::Remove:
-			return FNiagaraEditorStyle::Get().GetBrush("NiagaraEditor.PlatformSet.Remove");
-	}
-
-	return FEditorStyle::GetBrush("NoBrush");
-}
-
-FReply FNiagaraPlatformSetTypeCustomization::OnProfileMenuButtonClicked(TSharedPtr<FNiagaraDeviceProfileViewModel> Item, int32 EffectsQuality)
-{
-	ENiagaraPlatformSelectionState TargetState;
-
-	EProfileButtonMode Mode = GetProfileMenuButtonMode(TargetPlatformSet, Item, EffectsQuality);
-	switch(Mode)
-	{
-		case EProfileButtonMode::Include:
-			TargetState = ENiagaraPlatformSelectionState::Enabled;
-			break;
-		case EProfileButtonMode::Exclude:
-			TargetState = ENiagaraPlatformSelectionState::Disabled;
-			break;
-		case EProfileButtonMode::Remove:
-			TargetState = ENiagaraPlatformSelectionState::Default;
-			break;
-		default:
-			return FReply::Handled(); // shouldn't happen, button should be collapsed
-	}
-
-	PropertyHandle->NotifyPreChange();
-	TargetPlatformSet->SetDeviceProfileState(Item->Profile, EffectsQuality, TargetState);
-	PropertyHandle->NotifyPostChange(EPropertyChangeType::ValueSet);
-
-	InvalidateSiblingConflicts();
-
-	return FReply::Handled();
-}
-
-void FNiagaraPlatformSetTypeCustomization::OnGetDeviceProfileTreeChildren(TSharedPtr<FNiagaraDeviceProfileViewModel> InItem, TArray< TSharedPtr<FNiagaraDeviceProfileViewModel> >& OutChildren, int32 EffectsQuality)
-{
-	if (TargetPlatformSet->GetDeviceProfileState(InItem->Profile, EffectsQuality) == ENiagaraPlatformSelectionState::Default)
-	{
-		OutChildren = InItem->Children;
-	}
-}
-
-ECheckBoxState FNiagaraPlatformSetTypeCustomization::IsEQChecked(int32 EffectsQuality)const
-{
-	if (TargetPlatformSet)
-	{
-		return TargetPlatformSet->IsEffectQualityEnabled(EffectsQuality) ? ECheckBoxState::Checked : ECheckBoxState::Unchecked;
-	}
-	return ECheckBoxState::Undetermined;
-}
-
-void FNiagaraPlatformSetTypeCustomization::EQCheckStateChanged(ECheckBoxState CheckState, int32 EffectsQuality)
-{
-	PropertyHandle->NotifyPreChange();
-	TargetPlatformSet->SetEnabledForEffectQuality(EffectsQuality, CheckState == ECheckBoxState::Checked);
-	PropertyHandle->NotifyPostChange(EPropertyChangeType::ValueSet);
-
-	InvalidateSiblingConflicts();
-}
-
-void FNiagaraPlatformSetTypeCustomization::OnPropertyValueChanged()
-{
-	GenerateEffectsQualitySelectionWidgets();
-	UpdateCachedConflicts();
-}
 
 #undef LOCTEXT_NAMESPACE
diff --git a/Engine/Plugins/FX/Niagara/Source/NiagaraEditor/Private/Customizations/NiagaraTypeCustomizations.h b/Engine/Plugins/FX/Niagara/Source/NiagaraEditor/Private/Customizations/NiagaraTypeCustomizations.h
--- a/Engine/Plugins/FX/Niagara/Source/NiagaraEditor/Private/Customizations/NiagaraTypeCustomizations.h
+++ b/Engine/Plugins/FX/Niagara/Source/NiagaraEditor/Private/Customizations/NiagaraTypeCustomizations.h
@@ -5,10 +5,7 @@
 #include "CoreMinimal.h"
 #include "IDetailCustomization.h"
 #include "IPropertyTypeCustomization.h"
-#include "NiagaraPlatformSet.h"
 #include "EdGraph/EdGraphSchema.h"
-#include "Styling/SlateTypes.h"
-#include "Widgets/Views/STreeView.h"
 
 #include "NiagaraTypeCustomizations.generated.h"
 
@@ -16,9 +13,8 @@ class FDetailWidgetRow;
 class FNiagaraScriptViewModel;
 class IPropertyHandle;
 class IPropertyHandleArray;
-class IPropertyTypeCustomization;
-class SMenuAnchor;
-class SWrapBox;
+
+enum class ECheckBoxState : uint8;
 
 class FNiagaraNumericCustomization : public IPropertyTypeCustomization
 {
@@ -224,89 +220,3 @@ private:
 	class UNiagaraScriptVariable* BaseScriptVariable;
 	struct FNiagaraScriptVariableBinding* TargetVariableBinding;
 };
-
-struct FNiagaraDeviceProfileViewModel
-{
-	class UDeviceProfile* Profile;
-	TArray<TSharedPtr<FNiagaraDeviceProfileViewModel>> Children;
-};
-
-class FNiagaraPlatformSetTypeCustomization : public IPropertyTypeCustomization
-{
-public:
-	/** @return A new instance of this class */
-	static TSharedRef<IPropertyTypeCustomization> MakeInstance()
-	{
-		return MakeShared<FNiagaraPlatformSetTypeCustomization>();
-	}
-
-	/** IPropertyTypeCustomization interface begin */
-	virtual void CustomizeHeader(TSharedRef<IPropertyHandle> StructPropertyHandle, class FDetailWidgetRow& HeaderRow, IPropertyTypeCustomizationUtils& StructCustomizationUtils) override;
-	virtual void CustomizeChildren(TSharedRef<IPropertyHandle> StructPropertyHandle, class IDetailChildrenBuilder& ChildBuilder, IPropertyTypeCustomizationUtils& StructCustomizationUtils) override;
-	/** IPropertyTypeCustomization interface end */
-private:
-
-	void CreateDeviceProfileTree();
-
-	TSharedRef<SWidget> GenerateDeviceProfileTreeWidget(int32 EffectsQuality);
-	TSharedRef<ITableRow> OnGenerateDeviceProfileTreeRow(TSharedPtr<FNiagaraDeviceProfileViewModel> InItem, const TSharedRef<STableViewBase>& OwnerTable, int32 EffectsQuality);
-	void OnGetDeviceProfileTreeChildren(TSharedPtr<FNiagaraDeviceProfileViewModel> InItem, TArray< TSharedPtr<FNiagaraDeviceProfileViewModel> >& OutChildren, int32 EffectsQuality);
-	FReply RemoveDeviceProfile(UDeviceProfile* Profile, int32 EffectsQuality);
-
-	EVisibility GetDeviceProfileErrorVisibility(UDeviceProfile* Profile, int32 EffectsQuality) const;
-	FText GetDeviceProfileErrorToolTip(UDeviceProfile* Profile, int32 EffectsQuality) const;
-
-	FSlateColor GetEffectsQualityButtonTextColor(int32 EffectsQuality) const;
-	EVisibility GetEffectsQualityErrorVisibility(int32 EffectsQuality) const;
-	FText GetEffectsQualityErrorToolTip(int32 EffectsQuality) const;
-
-	const FSlateBrush* GetProfileMenuButtonImage(TSharedPtr<FNiagaraDeviceProfileViewModel> Item, int32 EffectsQuality) const;
-	EVisibility GetProfileMenuButtonVisibility(TSharedPtr<FNiagaraDeviceProfileViewModel> Item, int32 EffectsQuality) const;
-	FText GetProfileMenuButtonToolTip(TSharedPtr<FNiagaraDeviceProfileViewModel> Item, int32 EffectsQuality) const;
-	FReply OnProfileMenuButtonClicked(TSharedPtr<FNiagaraDeviceProfileViewModel> InItem, int32 EffectsQuality);
-
-	TSharedRef<SWidget> GetCurrentDeviceProfileSelectionWidget(TSharedPtr<FNiagaraDeviceProfileViewModel> ProfileView);
-	TSharedRef<SWidget> OnGenerateDeviceProfileSelectionWidget(TSharedPtr<ENiagaraPlatformSelectionState> InItem);
-
-	FText GetCurrentText() const;
-	FText GetTooltipText() const;
-
-	void GenerateEffectsQualitySelectionWidgets();
-	TSharedRef<SWidget> GenerateAdditionalDevicesWidgetForEQ(int32 EffectsQuality);
-
-	bool IsTreeActiveForEQ(const TSharedPtr<FNiagaraDeviceProfileViewModel>& Tree, int32 EffectsQualityMask) const;
-	void FilterTreeForEQ(const TSharedPtr<FNiagaraDeviceProfileViewModel>& SourceTree, TSharedPtr<FNiagaraDeviceProfileViewModel>& FilteredTree, int32 EffectsQualityMask);
-
-	ECheckBoxState IsEQChecked(int32 EffectsQuality) const;
-	void EQCheckStateChanged(ECheckBoxState CheckState, int32 EffectsQuality);
-
-	FReply ToggleMenuOpenForEffectsQuality(int32 EffectsQuality);
-
-	void UpdateCachedConflicts();
-	void InvalidateSiblingConflicts() const;
-	void OnPropertyValueChanged();
-
-private:
-	TArray<FNiagaraPlatformSetConflictInfo> CachedConflicts;
-	TSharedPtr<IPropertyHandleArray> PlatformSetArray;
-	int32 PlatformSetArrayIndex;
-
-	TArray<TSharedPtr<SMenuAnchor>> EffectsQualityMenus;
-	TSharedPtr<SWrapBox> EffectsQualityWidgetBox;
-
-	TSharedPtr<IPropertyHandle> PropertyHandle;
-	struct FNiagaraPlatformSet* TargetPlatformSet;
-	class UNiagaraSystem* BaseSystem;
-
-	TArray<TSharedPtr<FNiagaraDeviceProfileViewModel>> FullDeviceProfileTree;
-	TArray<TArray<TSharedPtr<FNiagaraDeviceProfileViewModel>>> FilteredDeviceProfileTrees;
-
-	TArray<TSharedPtr<ENiagaraPlatformSelectionState>> PlatformSelectionStates;
-
-	TSharedPtr<STreeView<TSharedPtr<FNiagaraDeviceProfileViewModel>>> DeviceProfileTreeWidget;
-
-	struct FNiagaraSystemScalabilitySettingsArray* SystemScalabilitySettings;
-	struct FNiagaraEmitterScalabilitySettingsArray* EmitterScalabilitySettings;
-	struct FNiagaraSystemScalabilityOverrides* SystemScalabilityOverrides;
-	struct FNiagaraEmitterScalabilityOverrides* EmitterScalabilityOverrides;
-};
diff --git a/Engine/Plugins/FX/Niagara/Source/NiagaraEditor/Private/NiagaraEditorModule.cpp b/Engine/Plugins/FX/Niagara/Source/NiagaraEditor/Private/NiagaraEditorModule.cpp
--- a/Engine/Plugins/FX/Niagara/Source/NiagaraEditor/Private/NiagaraEditorModule.cpp
+++ b/Engine/Plugins/FX/Niagara/Source/NiagaraEditor/Private/NiagaraEditorModule.cpp
@@ -37,6 +37,7 @@
 #include "NiagaraNodeConvert.h"
 #include "NiagaraNodeAssignment.h"
 #include "EdGraphSchema_Niagara.h"
+
 #include "TypeEditorUtilities/NiagaraFloatTypeEditorUtilities.h"
 #include "TypeEditorUtilities/NiagaraIntegerTypeEditorUtilities.h"
 #include "TypeEditorUtilities/NiagaraEnumTypeEditorUtilities.h"
@@ -46,10 +47,9 @@
 #include "TypeEditorUtilities/NiagaraColorTypeEditorUtilities.h"
 #include "TypeEditorUtilities/NiagaraMatrixTypeEditorUtilities.h"
 #include "TypeEditorUtilities/NiagaraDataInterfaceCurveTypeEditorUtilities.h"
+
 #include "NiagaraEditorStyle.h"
 #include "NiagaraEditorCommands.h"
-#include "Sequencer/NiagaraSequence/NiagaraEmitterTrackEditor.h"
-#include "Sequencer/LevelSequence/NiagaraSystemTrackEditor.h"
 #include "PropertyEditorModule.h"
 #include "NiagaraSettings.h"
 #include "NiagaraModule.h"
@@ -64,7 +64,6 @@
 #include "ViewModels/NiagaraSystemViewModel.h"
 #include "ViewModels/NiagaraEmitterViewModel.h"
 #include "TNiagaraGraphPinEditableName.h"
-#include "Sequencer/NiagaraSequence/Sections/MovieSceneNiagaraEmitterSection.h"
 #include "UObject/Class.h"
 #include "NiagaraScriptMergeManager.h"
 #include "NiagaraEmitter.h"
@@ -87,15 +86,22 @@
 #include "Sections/MovieSceneVectorSection.h"
 #include "Sections/MovieSceneColorSection.h"
 
+#include "Sequencer/NiagaraSequence/Sections/MovieSceneNiagaraEmitterSection.h"
+#include "Sequencer/NiagaraSequence/NiagaraEmitterTrackEditor.h"
+#include "Sequencer/LevelSequence/NiagaraSystemTrackEditor.h"
+
 #include "ISequencerSection.h"
 #include "Sections/BoolPropertySection.h"
 #include "Sections/ColorPropertySection.h"
 
 #include "Customizations/NiagaraComponentDetails.h"
-#include "Customizations/NiagaraTypeCustomizations.h"
+#include "Customizations/NiagaraFunctionCallNodeDetails.h"
 #include "Customizations/NiagaraEventScriptPropertiesCustomization.h"
+#include "Customizations/NiagaraPlatformSetCustomization.h"
 #include "Customizations/NiagaraScriptVariableCustomization.h"
 #include "Customizations/NiagaraScriptDetails.h"
+#include "Customizations/NiagaraStaticSwitchNodeDetails.h"
+#include "Customizations/NiagaraTypeCustomizations.h"
 
 #include "NiagaraComponent.h"
 #include "NiagaraNodeStaticSwitch.h"
@@ -108,8 +114,6 @@
 #include "NiagaraHlslTranslator.h"
 #include "NiagaraThumbnailRenderer.h"
 #include "Misc/FeedbackContext.h"
-#include "Customizations/NiagaraStaticSwitchNodeDetails.h"
-#include "Customizations/NiagaraFunctionCallNodeDetails.h"
 #include "NiagaraNodeFunctionCall.h"
 #include "Engine/Selection.h"
 #include "NiagaraActor.h"
@@ -829,7 +833,7 @@ void FNiagaraEditorModule::StartupModule()
 
 	PropertyModule.RegisterCustomPropertyTypeLayout(
 		FNiagaraPlatformSet::StaticStruct()->GetFName(),
-		FOnGetPropertyTypeCustomizationInstance::CreateStatic(&FNiagaraPlatformSetTypeCustomization::MakeInstance)
+		FOnGetPropertyTypeCustomizationInstance::CreateStatic(&FNiagaraPlatformSetCustomization::MakeInstance)
 	);
 	
 	PropertyModule.RegisterCustomPropertyTypeLayout(
diff --git a/Engine/Source/Editor/EditorStyle/Private/SlateEditorStyle.cpp b/Engine/Source/Editor/EditorStyle/Private/SlateEditorStyle.cpp
--- a/Engine/Source/Editor/EditorStyle/Private/SlateEditorStyle.cpp
+++ b/Engine/Source/Editor/EditorStyle/Private/SlateEditorStyle.cpp
@@ -260,13 +260,18 @@ void FSlateEditorStyle::FStyle::SetupGeneralStyles()
 
 	// Normal Text
 	{
-		Set( "RichTextBlock.TextHighlight", FTextBlockStyle(NormalText)
-			.SetColorAndOpacity( FLinearColor( 1.0f, 1.0f, 1.0f ) ) );
-		Set( "RichTextBlock.Bold", FTextBlockStyle(NormalText)
-			.SetFont( DEFAULT_FONT("Bold", FCoreStyle::RegularTextSize )) );
-		Set( "RichTextBlock.BoldHighlight", FTextBlockStyle(NormalText)
-			.SetFont( DEFAULT_FONT("Bold", FCoreStyle::RegularTextSize ))
-			.SetColorAndOpacity( FLinearColor( 1.0f, 1.0f, 1.0f ) ) );
+		Set("RichTextBlock.TextHighlight", FTextBlockStyle(NormalText)
+			.SetColorAndOpacity(FLinearColor( 1.0f, 1.0f, 1.0f)));
+		Set("RichTextBlock.Bold", FTextBlockStyle(NormalText)
+			.SetFont(DEFAULT_FONT("Bold", FCoreStyle::RegularTextSize)));
+		Set("RichTextBlock.BoldHighlight", FTextBlockStyle(NormalText)
+			.SetFont(DEFAULT_FONT("Bold", FCoreStyle::RegularTextSize))
+			.SetColorAndOpacity(FLinearColor(1.0f, 1.0f, 1.0f)));
+		Set("RichTextBlock.Italic", FTextBlockStyle(NormalText)
+			.SetFont(DEFAULT_FONT("Italic", FCoreStyle::RegularTextSize)));
+		Set("RichTextBlock.ItalicHighlight", FTextBlockStyle(NormalText)
+			.SetFont(DEFAULT_FONT("Italic", FCoreStyle::RegularTextSize))
+			.SetColorAndOpacity(FLinearColor(1.0f, 1.0f, 1.0f)));
 
 		Set( "TextBlock.HighlightShape",  new BOX_BRUSH( "Common/TextBlockHighlightShape", FMargin(3.f/8.f) ));
 		Set( "TextBlock.HighlighColor", FLinearColor( 0.02f, 0.3f, 0.0f ) );
