MN - Fix Niagara SkelmeshSampling where the SRV was copied before it was properly initzialized, luckily a proxy object already existed.

diff --git a/Engine/Plugins/FX/Niagara/Source/Niagara/Classes/NiagaraDataInterfaceSkeletalMesh.h b/Engine/Plugins/FX/Niagara/Source/Niagara/Classes/NiagaraDataInterfaceSkeletalMesh.h
--- a/Engine/Plugins/FX/Niagara/Source/Niagara/Classes/NiagaraDataInterfaceSkeletalMesh.h
+++ b/Engine/Plugins/FX/Niagara/Source/Niagara/Classes/NiagaraDataInterfaceSkeletalMesh.h
@@ -447,8 +447,8 @@ struct FNDISkeletalMesh_InstanceData
 
 	/** True if the mesh we're using is to be rendered in unlimited bone influences mode. */
 	uint32 bUnlimitedBoneInfluences : 1;
-	FRHIShaderResourceView* MeshSkinWeightBufferSrv;
-	FRHIShaderResourceView* MeshSkinWeightLookupBufferSrv;
+	const FSkinWeightDataVertexBuffer* MeshSkinWeightBuffer;
+	const FSkinWeightLookupVertexBuffer* MeshSkinWeightLookupBuffer;
 	uint32 MeshWeightStrideByte;
 	uint32 MeshSkinWeightIndexSizeByte;
 
@@ -778,8 +778,8 @@ struct FNiagaraDISkeletalMeshPassedDataToRT
 {
 	FSkeletalMeshGpuSpawnStaticBuffers* StaticBuffers;
 	FSkeletalMeshGpuDynamicBufferProxy* DynamicBuffer;
-	FRHIShaderResourceView* MeshSkinWeightBufferSrv;
-	FRHIShaderResourceView* MeshSkinWeightLookupBufferSrv;
+	const FSkinWeightDataVertexBuffer* MeshSkinWeightBuffer;
+	const FSkinWeightLookupVertexBuffer* MeshSkinWeightLookupBuffer;
 
 	bool bIsGpuUniformlyDistributedSampling;
 
diff --git a/Engine/Plugins/FX/Niagara/Source/Niagara/Private/NiagaraDataInterfaceSkeletalMesh.cpp b/Engine/Plugins/FX/Niagara/Source/Niagara/Private/NiagaraDataInterfaceSkeletalMesh.cpp
--- a/Engine/Plugins/FX/Niagara/Source/Niagara/Private/NiagaraDataInterfaceSkeletalMesh.cpp
+++ b/Engine/Plugins/FX/Niagara/Source/Niagara/Private/NiagaraDataInterfaceSkeletalMesh.cpp
@@ -978,8 +978,8 @@ public:
 				SetSRVParameter(RHICmdList, ComputeShaderRHI, MeshTriangleSamplerAliasBuffer, FNiagaraRenderer::GetDummyUIntBuffer());
 			}
 
-			SetSRVParameter(RHICmdList, ComputeShaderRHI, MeshSkinWeightBuffer, InstanceData->MeshSkinWeightBufferSrv);
-			SetSRVParameter(RHICmdList, ComputeShaderRHI, MeshSkinWeightLookupBuffer, InstanceData->MeshSkinWeightLookupBufferSrv);
+			SetSRVParameter(RHICmdList, ComputeShaderRHI, MeshSkinWeightBuffer, InstanceData->MeshSkinWeightBuffer->GetSRV());
+			SetSRVParameter(RHICmdList, ComputeShaderRHI, MeshSkinWeightLookupBuffer, InstanceData->MeshSkinWeightLookupBuffer->GetSRV());
 
 			SetShaderValue(RHICmdList, ComputeShaderRHI, MeshWeightStride, InstanceData->MeshWeightStrideByte/4);
 			SetShaderValue(RHICmdList, ComputeShaderRHI, MeshSkinWeightIndexSize, InstanceData->MeshSkinWeightIndexSizeByte);
@@ -1118,9 +1118,8 @@ void FNiagaraDataInterfaceProxySkeletalMesh::ConsumePerInstanceDataFromGameThrea
 	Data.StaticBuffers = SourceData->StaticBuffers;
 	Data.Transform = SourceData->Transform;
 
-	// @todo-threadsafety race here. Need to hold a ref to this buffer on the RT
-	Data.MeshSkinWeightBufferSrv = SourceData->MeshSkinWeightBufferSrv;
-	Data.MeshSkinWeightLookupBufferSrv = SourceData->MeshSkinWeightLookupBufferSrv;
+	Data.MeshSkinWeightBuffer = SourceData->MeshSkinWeightBuffer;
+	Data.MeshSkinWeightLookupBuffer = SourceData->MeshSkinWeightLookupBuffer;
 }
 
 //////////////////////////////////////////////////////////////////////////
@@ -1141,9 +1140,8 @@ void UNiagaraDataInterfaceSkeletalMesh::ProvidePerInstanceDataForRenderThread(vo
 	Data->StaticBuffers = SourceData->MeshGpuSpawnStaticBuffers;
 	Data->Transform = SourceData->Transform;
 
-	// @todo-threadsafety race here. Need to hold a ref to this buffer on the RT
-	Data->MeshSkinWeightBufferSrv = SourceData->MeshSkinWeightBufferSrv;
-	Data->MeshSkinWeightLookupBufferSrv = SourceData->MeshSkinWeightLookupBufferSrv;
+	Data->MeshSkinWeightBuffer = SourceData->MeshSkinWeightBuffer;
+	Data->MeshSkinWeightLookupBuffer = SourceData->MeshSkinWeightLookupBuffer;
 }
 
 USkeletalMesh* UNiagaraDataInterfaceSkeletalMesh::GetSkeletalMesh(UNiagaraComponent* OwningComponent, TWeakObjectPtr<USceneComponent>& SceneComponent, USkeletalMeshComponent*& FoundSkelComp, FNDISkeletalMesh_InstanceData* InstData)
@@ -1552,9 +1550,9 @@ bool FNDISkeletalMesh_InstanceData::Init(UNiagaraDataInterfaceSkeletalMesh* Inte
 		bUnlimitedBoneInfluences = (BoneInfluenceType == GPUSkinBoneInfluenceType::UnlimitedBoneInfluence);
 		MeshWeightStrideByte = SkinWeightBuffer->GetConstantInfluencesVertexStride();
 		MeshSkinWeightIndexSizeByte = SkinWeightBuffer->GetBoneIndexByteSize();
-		MeshSkinWeightBufferSrv = SkinWeightBuffer->GetDataVertexBuffer()->GetSRV();
+		MeshSkinWeightBuffer = SkinWeightBuffer->GetDataVertexBuffer();
 		//check(MeshSkinWeightBufferSrv->IsValid()); // not available in this stream
-		MeshSkinWeightLookupBufferSrv = SkinWeightBuffer->GetLookupVertexBuffer()->GetSRV();
+		MeshSkinWeightLookupBuffer = SkinWeightBuffer->GetLookupVertexBuffer();
 
 		FSkeletalMeshLODInfo* LODInfo = Mesh->GetLODInfo(LODIndex);
 		bIsGpuUniformlyDistributedSampling = LODInfo->bSupportUniformlyDistributedSampling && bAllRegionsAreAreaWeighting;
