Niagara "Set Variables" modules now display the name of the variable that they are setting, and the number of additional variables changed in parentheses.

Set Specific Parameters menu entries now display the actual parameter name instead of adding spaces for no reason.

diff --git a/Engine/Plugins/FX/Niagara/Source/NiagaraEditor/Private/NiagaraNodeAssignment.cpp b/Engine/Plugins/FX/Niagara/Source/NiagaraEditor/Private/NiagaraNodeAssignment.cpp
--- a/Engine/Plugins/FX/Niagara/Source/NiagaraEditor/Private/NiagaraNodeAssignment.cpp
+++ b/Engine/Plugins/FX/Niagara/Source/NiagaraEditor/Private/NiagaraNodeAssignment.cpp
@@ -663,6 +663,7 @@ bool UNiagaraNodeAssignment::RenameAssignmentTarget(FName OldName, FName NewName
 {
 	for (FNiagaraVariable& AssignmentTarget : AssignmentTargets)
 	{
+
 		if (AssignmentTarget.GetName() == OldName)
 		{
 			Modify();
@@ -673,9 +674,11 @@ bool UNiagaraNodeAssignment::RenameAssignmentTarget(FName OldName, FName NewName
 			}
 
 			AssignmentTarget.SetName(NewName);
+			RefreshTitle();
 			return true;
 		}
 	}
+
 	return false;
 }
 
