Niagara - When rellocating pins, keep around pins which have non default values or links and mark them as orphaned pins.  Update the translator to output a compile error when orphaned pins are encountered.  Among other usability issues, this fixes an issue where the entire script stack becomes invalid when force deleting a module.

diff --git a/Engine/Plugins/FX/Niagara/Source/NiagaraEditor/Private/EdGraphSchema_Niagara.cpp b/Engine/Plugins/FX/Niagara/Source/NiagaraEditor/Private/EdGraphSchema_Niagara.cpp
--- a/Engine/Plugins/FX/Niagara/Source/NiagaraEditor/Private/EdGraphSchema_Niagara.cpp
+++ b/Engine/Plugins/FX/Niagara/Source/NiagaraEditor/Private/EdGraphSchema_Niagara.cpp
@@ -1050,6 +1050,8 @@ FConnectionDrawingPolicy* UEdGraphSchema_Niagara::CreateConnectionDrawingPolicy(
 
 void UEdGraphSchema_Niagara::ResetPinToAutogeneratedDefaultValue(UEdGraphPin* Pin, bool bCallModifyCallbacks) const
 {
+	const FScopedTransaction Transaction(LOCTEXT("ResetPinToDefault", "Reset pin to default."));
+	Pin->Modify();
 	Pin->DefaultValue = Pin->AutogeneratedDefaultValue;
 	if (bCallModifyCallbacks)
 	{
diff --git a/Engine/Plugins/FX/Niagara/Source/NiagaraEditor/Private/NiagaraHlslTranslator.cpp b/Engine/Plugins/FX/Niagara/Source/NiagaraEditor/Private/NiagaraHlslTranslator.cpp
--- a/Engine/Plugins/FX/Niagara/Source/NiagaraEditor/Private/NiagaraHlslTranslator.cpp
+++ b/Engine/Plugins/FX/Niagara/Source/NiagaraEditor/Private/NiagaraHlslTranslator.cpp
@@ -263,6 +263,11 @@ bool FHlslNiagaraTranslator::ValidateTypePins(UNiagaraNode* NodeToValidate)
 				bPinsAreValid = false;
 			}
 		}
+
+		if (Pin->bOrphanedPin)
+		{
+			Error(LOCTEXT("OrphanedPinError", "Node pin is no longer valid.  This pin must be disconnected or reset to default so it can be removed."), NodeToValidate, Pin);
+		}
 	}
 	return bPinsAreValid;
 }
diff --git a/Engine/Plugins/FX/Niagara/Source/NiagaraEditor/Private/NiagaraNode.cpp b/Engine/Plugins/FX/Niagara/Source/NiagaraEditor/Private/NiagaraNode.cpp
--- a/Engine/Plugins/FX/Niagara/Source/NiagaraEditor/Private/NiagaraNode.cpp
+++ b/Engine/Plugins/FX/Niagara/Source/NiagaraEditor/Private/NiagaraNode.cpp
@@ -230,8 +230,25 @@ bool UNiagaraNode::ReallocatePins(bool bMarkNeedsResynchronizeOnChange)
 		else
 		{
 			bAllSame = false;
+			if (OldPin->DefaultValue != OldPin->AutogeneratedDefaultValue || OldPin->LinkedTo.Num() > 0)
+			{
+				// If an old pin is not reused mark it as orphaned so that it's value and connections aren't lost.
+				OldPin->bOrphanedPin = true;
+			}
+		}
+	}
+
+	for (UEdGraphPin* OldPin : OldPins)
+	{
+		if (OldPin->bOrphanedPin)
+		{
+			// Add orphaned pins back in at the end so that the user can fix them.
+			Pins.Add(OldPin);
+		}
+		else
+		{
+			OldPin->MarkPendingKill();
 		}
-		OldPin->MarkPendingKill();
 	}
 
 	//GetGraph()->NotifyGraphChanged();
@@ -433,6 +450,10 @@ void UNiagaraNode::ForceChangeId(const FGuid& InId, bool bRaiseGraphNeedsRecompi
 
 void UNiagaraNode::PinDefaultValueChanged(UEdGraphPin* Pin) 
 {
+	if (Pin->bOrphanedPin && Pin->DefaultValue == Pin->AutogeneratedDefaultValue)
+	{
+		RemovePin(Pin);
+	}
 	MarkNodeRequiresSynchronization(__FUNCTION__, true);
 	Super::PinDefaultValueChanged(Pin);
 }
@@ -451,6 +472,13 @@ void UNiagaraNode::OnPinRemoved(UEdGraphPin* InRemovedPin)
 
 void UNiagaraNode::NodeConnectionListChanged()
 {
+	for (UEdGraphPin* Pin : Pins)
+	{
+		if (Pin->bOrphanedPin && Pin->LinkedTo.Num() == 0)
+		{
+			RemovePin(Pin);
+		}
+	}
 	MarkNodeRequiresSynchronization(__FUNCTION__, true);
 	Super::NodeConnectionListChanged();
 }
@@ -458,6 +486,10 @@ void UNiagaraNode::NodeConnectionListChanged()
 
 void UNiagaraNode::PinConnectionListChanged(UEdGraphPin* Pin) 
 {
+	if (Pin->bOrphanedPin && Pin->LinkedTo.Num() == 0)
+	{
+		RemovePin(Pin);
+	}
 	MarkNodeRequiresSynchronization(__FUNCTION__, true);
 	Super::PinConnectionListChanged(Pin);
 }
