Niagara - Fix some crashes related to undo/redo.

diff --git a/Engine/Plugins/FX/Niagara/Source/NiagaraEditor/Private/EdGraphSchema_Niagara.cpp b/Engine/Plugins/FX/Niagara/Source/NiagaraEditor/Private/EdGraphSchema_Niagara.cpp
--- a/Engine/Plugins/FX/Niagara/Source/NiagaraEditor/Private/EdGraphSchema_Niagara.cpp
+++ b/Engine/Plugins/FX/Niagara/Source/NiagaraEditor/Private/EdGraphSchema_Niagara.cpp
@@ -1050,7 +1050,7 @@ FConnectionDrawingPolicy* UEdGraphSchema_Niagara::CreateConnectionDrawingPolicy(
 
 void UEdGraphSchema_Niagara::ResetPinToAutogeneratedDefaultValue(UEdGraphPin* Pin, bool bCallModifyCallbacks) const
 {
-	const FScopedTransaction Transaction(LOCTEXT("ResetPinToDefault", "Reset pin to default."));
+	const FScopedTransaction Transaction(LOCTEXT("ResetPinToDefault", "Reset pin to default."), GIsTransacting == false);
 	Pin->Modify();
 	Pin->DefaultValue = Pin->AutogeneratedDefaultValue;
 	if (bCallModifyCallbacks)
diff --git a/Engine/Plugins/FX/Niagara/Source/NiagaraEditor/Private/ViewModels/Stack/NiagaraStackEmitterSettingsGroup.cpp b/Engine/Plugins/FX/Niagara/Source/NiagaraEditor/Private/ViewModels/Stack/NiagaraStackEmitterSettingsGroup.cpp
--- a/Engine/Plugins/FX/Niagara/Source/NiagaraEditor/Private/ViewModels/Stack/NiagaraStackEmitterSettingsGroup.cpp
+++ b/Engine/Plugins/FX/Niagara/Source/NiagaraEditor/Private/ViewModels/Stack/NiagaraStackEmitterSettingsGroup.cpp
@@ -129,8 +129,13 @@ void UNiagaraStackEmitterPropertiesItem::RefreshIssues(TArray<FStackIssue>& NewI
 
 void UNiagaraStackEmitterPropertiesItem::EmitterPropertiesChanged()
 {
-	bCanResetToBaseCache.Reset();
-	RefreshChildren();
+	if (IsFinalized() == false)
+	{
+		// Undo/redo can cause objects to disappear and reappear which can prevent safe removal of delegates
+		// so guard against receiving an event when finalized here.
+		bCanResetToBaseCache.Reset();
+		RefreshChildren();
+	}
 }
 
 UNiagaraStackEmitterSettingsGroup::UNiagaraStackEmitterSettingsGroup()
diff --git a/Engine/Plugins/FX/Niagara/Source/NiagaraEditor/Private/ViewModels/Stack/NiagaraStackRendererItem.cpp b/Engine/Plugins/FX/Niagara/Source/NiagaraEditor/Private/ViewModels/Stack/NiagaraStackRendererItem.cpp
--- a/Engine/Plugins/FX/Niagara/Source/NiagaraEditor/Private/ViewModels/Stack/NiagaraStackRendererItem.cpp
+++ b/Engine/Plugins/FX/Niagara/Source/NiagaraEditor/Private/ViewModels/Stack/NiagaraStackRendererItem.cpp
@@ -433,8 +433,13 @@ void UNiagaraStackRendererItem::RefreshIssues(TArray<FStackIssue>& NewIssues)
 
 void UNiagaraStackRendererItem::RendererChanged()
 {
-	bCanResetToBaseCache.Reset();
-	RefreshChildren();
+	if (IsFinalized() == false)
+	{
+		// Undo/redo can cause objects to disappear and reappear which can prevent safe removal of delegates
+		// so guard against receiving an event when finalized here.
+		bCanResetToBaseCache.Reset();
+		RefreshChildren();
+	}
 }
 
 #undef LOCTEXT_NAMESPACE
