Niagara - Fix a bunch of engine test content which started failing due to the orpahaned pin change.  Also add stack issues and fixes for static switches on stack nodes which are no longer valid.

diff --git a/Engine/Plugins/FX/Niagara/Source/NiagaraEditor/Private/ViewModels/Stack/NiagaraStackFunctionInputCollection.cpp b/Engine/Plugins/FX/Niagara/Source/NiagaraEditor/Private/ViewModels/Stack/NiagaraStackFunctionInputCollection.cpp
--- a/Engine/Plugins/FX/Niagara/Source/NiagaraEditor/Private/ViewModels/Stack/NiagaraStackFunctionInputCollection.cpp
+++ b/Engine/Plugins/FX/Niagara/Source/NiagaraEditor/Private/ViewModels/Stack/NiagaraStackFunctionInputCollection.cpp
@@ -288,6 +288,18 @@ UNiagaraStackEntry::FStackIssueFix UNiagaraStackFunctionInputCollection::GetNode
 	}));
 }
 
+UNiagaraStackEntry::FStackIssueFix UNiagaraStackFunctionInputCollection::GetResetPinFix(UEdGraphPin* PinToReset, FText FixDescription)
+{
+	return FStackIssueFix(
+		FixDescription,
+		UNiagaraStackEntry::FStackIssueFixDelegate::CreateLambda([=]()
+	{
+		FScopedTransaction ScopedTransaction(FixDescription);
+		const UEdGraphSchema_Niagara* NiagaraSchema = GetDefault<UEdGraphSchema_Niagara>();
+		NiagaraSchema->ResetPinToAutogeneratedDefaultValue(PinToReset);
+	}));
+}
+
 void UNiagaraStackFunctionInputCollection::RefreshIssues(TArray<FName> DuplicateInputNames, TArray<FName> ValidAliasedInputNames, TArray<const UEdGraphPin*> PinsWithInvalidTypes, TMap<FName, UEdGraphPin*> StaticSwitchInputs, TArray<FStackIssue>& NewIssues)
 {
 	if (!GetIsEnabled())
@@ -318,8 +330,8 @@ void UNiagaraStackFunctionInputCollection::RefreshIssues(TArray<FName> Duplicate
 			{
 				FStackIssue InvalidInputOverrideError(
 					EStackIssueSeverity::Error,
-					FText::Format(LOCTEXT("InvalidInputSummaryFormat", "Invalid Input Override: {0}"), FText::FromString(OverridePin->PinName.ToString())),
-					FText::Format(LOCTEXT("InvalidInputFormat", "The input {0} was previously overriden but is no longer exposed by the function {1}.\nPress the fix button to remove this unused override data,\nor check the function definition to see why this input is no longer exposed."),
+					FText::Format(LOCTEXT("InvalidInputOverrideSummaryFormat", "Invalid Input Override: {0}"), FText::FromString(OverridePin->PinName.ToString())),
+					FText::Format(LOCTEXT("InvalidInputOverrideFormat", "The input {0} was previously overriden but is no longer exposed by the function {1}.\nPress the fix button to remove this unused override data,\nor check the function definition to see why this input is no longer exposed."),
 						FText::FromString(OverridePin->PinName.ToString()), FText::FromString(InputFunctionCallNode->GetFunctionName())),
 					GetStackEditorDataKey(),
 					false,
@@ -406,6 +418,23 @@ void UNiagaraStackFunctionInputCollection::RefreshIssues(TArray<FName> Duplicate
 			false);
 		NewIssues.Add(InputWithInvalidTypeError);
 	}
+
+	// Generate issues for orphaned input pins from static switches which are no longer valid.
+	for (UEdGraphPin* InputFunctionCallNodePin : InputFunctionCallNode->Pins)
+	{
+		if (InputFunctionCallNodePin->Direction == EEdGraphPinDirection::EGPD_Input && InputFunctionCallNodePin->bOrphanedPin)
+		{
+			FStackIssue InvalidInputError(
+				EStackIssueSeverity::Error,
+				FText::Format(LOCTEXT("InvalidInputSummaryFormat", "Invalid Input: {0}"), FText::FromString(InputFunctionCallNodePin->PinName.ToString())),
+				FText::Format(LOCTEXT("InvalidInputFormat", "The input {0} was previously set but is no longer exposed by the function {1}.\nPress the fix button to remove this unused input data,\nor check the function definition to see why this input is no longer exposed."),
+					FText::FromString(InputFunctionCallNodePin->PinName.ToString()), FText::FromString(InputFunctionCallNode->GetFunctionName())),
+				GetStackEditorDataKey(),
+				false,
+				GetResetPinFix(InputFunctionCallNodePin, LOCTEXT("RemoveInvalidInputPinFix", "Remove invalid input.")));
+			NewIssues.Add(InvalidInputError);
+		}
+	}
 }
 
 void UNiagaraStackFunctionInputCollection::OnFunctionInputsChanged()
diff --git a/Engine/Plugins/FX/Niagara/Source/NiagaraEditor/Public/ViewModels/Stack/NiagaraStackFunctionInputCollection.h b/Engine/Plugins/FX/Niagara/Source/NiagaraEditor/Public/ViewModels/Stack/NiagaraStackFunctionInputCollection.h
--- a/Engine/Plugins/FX/Niagara/Source/NiagaraEditor/Public/ViewModels/Stack/NiagaraStackFunctionInputCollection.h
+++ b/Engine/Plugins/FX/Niagara/Source/NiagaraEditor/Public/ViewModels/Stack/NiagaraStackFunctionInputCollection.h
@@ -52,6 +52,8 @@ private:
 
 	UNiagaraStackEntry::FStackIssueFix GetNodeRemovalFix(UEdGraphPin* PinToRemove, FText FixDescription);
 
+	UNiagaraStackEntry::FStackIssueFix GetResetPinFix(UEdGraphPin* PinToReset, FText FixDescription);
+
 	struct FInputData
 	{
 		const UEdGraphPin* Pin;
