Descriptions for UI

diff --git a/Engine/Plugins/FX/Niagara/Source/Niagara/Public/NiagaraTypes.h b/Engine/Plugins/FX/Niagara/Source/Niagara/Public/NiagaraTypes.h
--- a/Engine/Plugins/FX/Niagara/Source/Niagara/Public/NiagaraTypes.h
+++ b/Engine/Plugins/FX/Niagara/Source/Niagara/Public/NiagaraTypes.h
@@ -442,22 +442,31 @@ public:
 UENUM()
 enum class ENiagaraParameterScope : uint32
 {
+	/** Parameter that is an input argument into this graph.*/
 	Input,
 
+	/** Parameter that is exposed to the owning component for editing.*/
 	User,
 
+	/** Parameter provided by the engine. These are explicitly defined by the engine codebase.*/
 	Engine,
 
+	/** Parameter provided by the engine focused on the owning component. These are explicitly defined by the engine codebase.*/
 	Owner,
 
+	/** Parameter is an attribute of the owning system payload. It is persistent across frames.*/
 	System,
 
+	/** Parameter is an attribute of the owning emitter payload. It is persistent across frames.*/
 	Emitter,
 
+	/** Parameter is an attribute of the owning particle payload. It is persistent across frames.*/
 	Particles,
 
+	/** Parameter is initialized in the appropriate spawn section for the stack. It is persistent from frame to frame. For example, if used consistently in an Emitter stack, this parameter will turn into an emitter attribute. Similarly, if used in a Particle stack, it will turn into a particle attribute.*/
 	ScriptPersistent UMETA(Hidden), //@todo(ng) hiding until autotest verification is made.
 
+	/** Parameter is initialized at the start of this script phase and can be shared amongst other modules, but is not persistent across script frames or from stack section to stack section.*/
 	ScriptTransient,
 
 	Local UMETA(Hidden), //Convenience markup for ScopeToString functions, only use in conjunction with ENiagaraScriptParameterUsage::Local.
