Fixed Niagara static mesh color sampling when asset doesn't have vertex colors.

diff --git a/Engine/Plugins/FX/Niagara/Source/Niagara/Classes/NiagaraDataInterfaceStaticMesh.h b/Engine/Plugins/FX/Niagara/Source/Niagara/Classes/NiagaraDataInterfaceStaticMesh.h
--- a/Engine/Plugins/FX/Niagara/Source/Niagara/Classes/NiagaraDataInterfaceStaticMesh.h
+++ b/Engine/Plugins/FX/Niagara/Source/Niagara/Classes/NiagaraDataInterfaceStaticMesh.h
@@ -232,6 +232,7 @@ public:
 	static const FString InstanceWorldVelocityName;
 	static const FString AreaWeightedSamplingName;
 	static const FString NumTexCoordName;
+	static const FString UseColorBufferName;
 
 public:
 	void GetNumTriangles(FVectorVMContext& Context);
diff --git a/Engine/Plugins/FX/Niagara/Source/Niagara/Private/NiagaraDataInterfaceStaticMesh.cpp b/Engine/Plugins/FX/Niagara/Source/Niagara/Private/NiagaraDataInterfaceStaticMesh.cpp
--- a/Engine/Plugins/FX/Niagara/Source/Niagara/Private/NiagaraDataInterfaceStaticMesh.cpp
+++ b/Engine/Plugins/FX/Niagara/Source/Niagara/Private/NiagaraDataInterfaceStaticMesh.cpp
@@ -28,6 +28,7 @@ const FString UNiagaraDataInterfaceStaticMesh::InstanceInvDeltaTimeName(TEXT("In
 const FString UNiagaraDataInterfaceStaticMesh::InstanceWorldVelocityName(TEXT("InstanceWorldVelocity_"));
 const FString UNiagaraDataInterfaceStaticMesh::AreaWeightedSamplingName(TEXT("AreaWeightedSamplingName_"));
 const FString UNiagaraDataInterfaceStaticMesh::NumTexCoordName(TEXT("NumTexCoordName_"));
+const FString UNiagaraDataInterfaceStaticMesh::UseColorBufferName(TEXT("UseColorBuffer_"));
 
 static int32 GNiagaraFailStaticMeshDataInterface = 0;
 static FAutoConsoleVariableRef CVarNiagaraFailStaticMeshDataInterface(
@@ -457,6 +458,7 @@ struct FNDIStaticMeshParametersName
 	FString InstanceWorldVelocityName;
 	FString AreaWeightedSamplingName;
 	FString NumTexCoordName;
+	FString UseColorBufferName;
 };
 
 static void GetNiagaraDataInterfaceParametersName(FNDIStaticMeshParametersName& Names, const FString& Suffix)
@@ -476,6 +478,7 @@ static void GetNiagaraDataInterfaceParametersName(FNDIStaticMeshParametersName&
 	Names.InstanceWorldVelocityName = UNiagaraDataInterfaceStaticMesh::InstanceWorldVelocityName + Suffix;
 	Names.AreaWeightedSamplingName = UNiagaraDataInterfaceStaticMesh::AreaWeightedSamplingName + Suffix;
 	Names.NumTexCoordName = UNiagaraDataInterfaceStaticMesh::NumTexCoordName + Suffix;
+	Names.UseColorBufferName = UNiagaraDataInterfaceStaticMesh::UseColorBufferName + Suffix;
 }
 
 struct FNiagaraDataInterfaceParametersCS_StaticMesh : public FNiagaraDataInterfaceParametersCS
@@ -502,6 +505,7 @@ public:
 		InstanceWorldVelocity.Bind(ParameterMap, *ParamNames.InstanceWorldVelocityName);
 		AreaWeightedSampling.Bind(ParameterMap, *ParamNames.AreaWeightedSamplingName);
 		NumTexCoord.Bind(ParameterMap, *ParamNames.NumTexCoordName);
+		UseColorBuffer.Bind(ParameterMap, *ParamNames.UseColorBufferName);
 	}
 
 	void Set(FRHICommandList& RHICmdList, const FNiagaraDataInterfaceSetArgs& Context) const
@@ -529,16 +533,18 @@ public:
 				}
 				else
 				{
-					SetSRVParameter(RHICmdList, ComputeShaderRHI, MeshTexCoordBuffer, FNiagaraRenderer::GetDummyFloatBuffer());
+					SetSRVParameter(RHICmdList, ComputeShaderRHI, MeshTexCoordBuffer, FNiagaraRenderer::GetDummyFloat2Buffer());
 				}
 
 				if(SpawnBuffer->GetBufferColorSRV())
 				{
 					SetSRVParameter(RHICmdList, ComputeShaderRHI, MeshColorBuffer, SpawnBuffer->GetBufferColorSRV());
+					SetShaderValue(RHICmdList, ComputeShaderRHI, UseColorBuffer, 1);
 				}
 				else
 				{
-					SetSRVParameter(RHICmdList, ComputeShaderRHI, MeshColorBuffer, FNiagaraRenderer::GetDummyFloatBuffer());
+					SetSRVParameter(RHICmdList, ComputeShaderRHI, MeshColorBuffer, FNiagaraRenderer::GetDummyWhiteColorBuffer());
+					SetShaderValue(RHICmdList, ComputeShaderRHI, UseColorBuffer, 0);
 				}
 
 				check(SpawnBuffer); // should always be allocated, we always need the GPU buffer for a GpuSimulation.
@@ -566,16 +572,17 @@ public:
 			else
 			{
 				SetSRVParameter(RHICmdList, ComputeShaderRHI, MeshVertexBuffer, FNiagaraRenderer::GetDummyFloatBuffer());
-				SetSRVParameter(RHICmdList, ComputeShaderRHI, MeshTangentBuffer, FNiagaraRenderer::GetDummyFloatBuffer());
+				SetSRVParameter(RHICmdList, ComputeShaderRHI, MeshTangentBuffer, FNiagaraRenderer::GetDummyFloat4Buffer());
 				SetSRVParameter(RHICmdList, ComputeShaderRHI, MeshIndexBuffer, FNiagaraRenderer::GetDummyUIntBuffer());
 
 				SetShaderValue(RHICmdList, ComputeShaderRHI, NumTexCoord, 0);
-				SetSRVParameter(RHICmdList, ComputeShaderRHI, MeshTexCoordBuffer, FNiagaraRenderer::GetDummyFloatBuffer());
-				SetSRVParameter(RHICmdList, ComputeShaderRHI, MeshColorBuffer, FNiagaraRenderer::GetDummyFloatBuffer());
+				SetSRVParameter(RHICmdList, ComputeShaderRHI, MeshTexCoordBuffer, FNiagaraRenderer::GetDummyFloat2Buffer());
+				SetSRVParameter(RHICmdList, ComputeShaderRHI, MeshColorBuffer, FNiagaraRenderer::GetDummyWhiteColorBuffer());
+				SetShaderValue(RHICmdList, ComputeShaderRHI, UseColorBuffer, 0);
 
 				SetShaderValue(RHICmdList, ComputeShaderRHI, SectionCount, 0);
-				SetSRVParameter(RHICmdList, ComputeShaderRHI, MeshSectionBuffer, FNiagaraRenderer::GetDummyUIntBuffer());
-				SetSRVParameter(RHICmdList, ComputeShaderRHI, MeshTriangleBuffer, FNiagaraRenderer::GetDummyUIntBuffer());
+				SetSRVParameter(RHICmdList, ComputeShaderRHI, MeshSectionBuffer, FNiagaraRenderer::GetDummyUInt4Buffer());
+				SetSRVParameter(RHICmdList, ComputeShaderRHI, MeshTriangleBuffer, FNiagaraRenderer::GetDummyUInt4Buffer());
 
 				SetShaderValue(RHICmdList, ComputeShaderRHI, InstanceTransform, FMatrix::Identity);
 				SetShaderValue(RHICmdList, ComputeShaderRHI, InstanceTransformInverseTransposed, FMatrix::Identity);
@@ -603,6 +610,7 @@ private:
 	LAYOUT_FIELD(FShaderParameter, InstanceWorldVelocity);
 	LAYOUT_FIELD(FShaderParameter, AreaWeightedSampling);
 	LAYOUT_FIELD(FShaderParameter, NumTexCoord);
+	LAYOUT_FIELD(FShaderParameter, UseColorBuffer);
 };
 
 IMPLEMENT_TYPE_LAYOUT(FNiagaraDataInterfaceParametersCS_StaticMesh);
@@ -1954,6 +1962,7 @@ bool UNiagaraDataInterfaceStaticMesh::GetFunctionHLSL(const FNiagaraDataInterfac
 		{TEXT("InstanceWorldVelocity"), ParamNames.InstanceWorldVelocityName},
 		{TEXT("AreaWeightedSamplingName"), ParamNames.AreaWeightedSamplingName},
 		{TEXT("NumTexCoordName"), ParamNames.NumTexCoordName},
+		{TEXT("UseColorBufferName"), ParamNames.UseColorBufferName},
 	};
 
 	if (FunctionInfo.DefinitionName == StaticMeshHelpers::RandomSectionName)
@@ -2225,16 +2234,20 @@ bool UNiagaraDataInterfaceStaticMesh::GetFunctionHLSL(const FNiagaraDataInterfac
 		static const TCHAR *FormatSample = TEXT(R"(
 			void {InstanceFunctionName} (in {MeshTriCoordinateStructName} In_Coord, out float4 Out_Color)
 			{
-				uint TriangleIndex = In_Coord.Tri * 3;
-				uint VertexIndex0 = {MeshIndexBufferName}[TriangleIndex  ];
-				uint VertexIndex1 = {MeshIndexBufferName}[TriangleIndex+1];
-				uint VertexIndex2 = {MeshIndexBufferName}[TriangleIndex+2];
+				Out_Color = float4(1, 1, 1, 1);
+				if ({UseColorBufferName})
+				{
+					uint TriangleIndex = In_Coord.Tri * 3;
+					uint VertexIndex0 = {MeshIndexBufferName}[TriangleIndex  ];
+					uint VertexIndex1 = {MeshIndexBufferName}[TriangleIndex+1];
+					uint VertexIndex2 = {MeshIndexBufferName}[TriangleIndex+2];
 
-				float4 Color0 = {MeshColorBufferName}[VertexIndex0] FMANUALFETCH_COLOR_COMPONENT_SWIZZLE;
-				float4 Color1 = {MeshColorBufferName}[VertexIndex1] FMANUALFETCH_COLOR_COMPONENT_SWIZZLE;
-				float4 Color2 = {MeshColorBufferName}[VertexIndex2] FMANUALFETCH_COLOR_COMPONENT_SWIZZLE;
+					float4 Color0 = {MeshColorBufferName}[VertexIndex0] FMANUALFETCH_COLOR_COMPONENT_SWIZZLE;
+					float4 Color1 = {MeshColorBufferName}[VertexIndex1] FMANUALFETCH_COLOR_COMPONENT_SWIZZLE;
+					float4 Color2 = {MeshColorBufferName}[VertexIndex2] FMANUALFETCH_COLOR_COMPONENT_SWIZZLE;
 
-				Out_Color = Color0 * In_Coord.BaryCoord.x + Color1 * In_Coord.BaryCoord.y + Color2 * In_Coord.BaryCoord.z;
+					Out_Color = Color0 * In_Coord.BaryCoord.x + Color1 * In_Coord.BaryCoord.y + Color2 * In_Coord.BaryCoord.z;
+				}
 			}
 			)");
 		OutHLSL += FString::Format(FormatSample, ArgsSample);
@@ -2368,6 +2381,7 @@ void UNiagaraDataInterfaceStaticMesh::GetParameterDefinitionHLSL(const FNiagaraD
 	OutHLSL += TEXT("Buffer<float4> ") + ParamNames.MeshColorBufferName + TEXT(";\n");
 	OutHLSL += TEXT("Buffer<uint4> ") + ParamNames.MeshSectionBufferName + TEXT(";\n");
 	OutHLSL += TEXT("Buffer<uint4> ") + ParamNames.MeshTriangleBufferName + TEXT(";\n");
+	OutHLSL += TEXT("uint ") + ParamNames.UseColorBufferName + TEXT(";\n");
 	OutHLSL += TEXT("uint ") + ParamNames.SectionCountName + TEXT(";\n");
 	OutHLSL += TEXT("float4x4 ") + ParamNames.InstanceTransformName + TEXT(";\n");
 	OutHLSL += TEXT("float4x4 ") + ParamNames.InstanceTransformInverseTransposedName + TEXT(";\n");
diff --git a/Engine/Plugins/FX/Niagara/Source/Niagara/Private/NiagaraRenderer.cpp b/Engine/Plugins/FX/Niagara/Source/Niagara/Private/NiagaraRenderer.cpp
--- a/Engine/Plugins/FX/Niagara/Source/Niagara/Private/NiagaraRenderer.cpp
+++ b/Engine/Plugins/FX/Niagara/Source/Niagara/Private/NiagaraRenderer.cpp
@@ -29,11 +29,13 @@ static FAutoConsoleVariableRef CVarNiagaraRadixSortThreshold(
 class FNiagaraEmptyBufferSRV : public FRenderResource
 {
 public:
-	FNiagaraEmptyBufferSRV(EPixelFormat InPixelFormat, const FString& InDebugName) : PixelFormat(InPixelFormat), DebugName(InDebugName) {}
+	FNiagaraEmptyBufferSRV(EPixelFormat InPixelFormat, const FString& InDebugName, uint32 InDefaultValue = 0) : PixelFormat(InPixelFormat), DebugName(InDebugName), DefaultValue(InDefaultValue) {}
 	EPixelFormat PixelFormat;
 	FString DebugName;
 	FVertexBufferRHIRef Buffer;
 	FShaderResourceViewRHIRef SRV;
+	uint32 DefaultValue = 0;
+
 
 	virtual void InitRHI() override
 	{
@@ -46,6 +48,12 @@ public:
 		// Zero the buffer memory.
 		void* Data = RHILockVertexBuffer(Buffer, 0, NumBytes, RLM_WriteOnly);
 		FMemory::Memset(Data, 0, NumBytes);
+		
+		if (PixelFormat == PF_R8G8B8A8)
+		{
+			*reinterpret_cast<uint32*>(Data) = DefaultValue;
+		}
+
 		RHIUnlockVertexBuffer(Buffer);
 
 		SRV = RHICreateShaderResourceView(Buffer, NumBytes, PixelFormat);
@@ -97,6 +105,14 @@ FRHIShaderResourceView* FNiagaraRenderer::GetDummyFloatBuffer()
 	return DummyFloatBuffer.SRV;
 }
 
+FRHIShaderResourceView* FNiagaraRenderer::GetDummyFloat2Buffer()
+{
+	check(IsInRenderingThread());
+	static TGlobalResource<FNiagaraEmptyBufferSRV> DummyFloat2Buffer(PF_G16R16F, TEXT("NiagaraRenderer::DummyFloat2"));
+	return DummyFloat2Buffer.SRV;
+}
+
+
 FRHIShaderResourceView* FNiagaraRenderer::GetDummyFloat4Buffer()
 {
 	check(IsInRenderingThread());
@@ -104,6 +120,13 @@ FRHIShaderResourceView* FNiagaraRenderer::GetDummyFloat4Buffer()
 	return DummyFloat4Buffer.SRV;
 }
 
+FRHIShaderResourceView* FNiagaraRenderer::GetDummyWhiteColorBuffer()
+{
+	check(IsInRenderingThread());
+	static TGlobalResource<FNiagaraEmptyBufferSRV> DummyWhiteColorBuffer(PF_R8G8B8A8, TEXT("NiagaraRenderer::DummyWhiteColorBuffer"), FColor::White.ToPackedRGBA());
+	return DummyWhiteColorBuffer.SRV;
+}
+
 FRHIShaderResourceView* FNiagaraRenderer::GetDummyIntBuffer()
 {
 	check(IsInRenderingThread());
@@ -118,6 +141,13 @@ FRHIShaderResourceView* FNiagaraRenderer::GetDummyUIntBuffer()
 	return DummyUIntBuffer.SRV;
 }
 
+FRHIShaderResourceView* FNiagaraRenderer::GetDummyUInt4Buffer()
+{
+	check(IsInRenderingThread());
+	static TGlobalResource<FNiagaraEmptyBufferSRV> DummyUInt4Buffer(PF_R32G32B32A32_UINT, TEXT("NiagaraRenderer::DummyUInt4"));
+	return DummyUInt4Buffer.SRV;
+}
+
 FRHIShaderResourceView* FNiagaraRenderer::GetDummyTextureReadBuffer2D()
 {
 	check(IsInRenderingThread());
diff --git a/Engine/Plugins/FX/Niagara/Source/Niagara/Public/NiagaraRenderer.h b/Engine/Plugins/FX/Niagara/Source/Niagara/Public/NiagaraRenderer.h
--- a/Engine/Plugins/FX/Niagara/Source/Niagara/Public/NiagaraRenderer.h
+++ b/Engine/Plugins/FX/Niagara/Source/Niagara/Public/NiagaraRenderer.h
@@ -116,9 +116,12 @@ public:
 #endif
 
 	NIAGARA_API static FRHIShaderResourceView* GetDummyFloatBuffer();
+	NIAGARA_API static FRHIShaderResourceView* GetDummyFloat2Buffer();
 	NIAGARA_API static FRHIShaderResourceView* GetDummyFloat4Buffer();
+	NIAGARA_API static FRHIShaderResourceView* GetDummyWhiteColorBuffer();
 	NIAGARA_API static FRHIShaderResourceView* GetDummyIntBuffer();
 	NIAGARA_API static FRHIShaderResourceView* GetDummyUIntBuffer();
+	NIAGARA_API static FRHIShaderResourceView* GetDummyUInt4Buffer();
 	NIAGARA_API static FRHIShaderResourceView* GetDummyTextureReadBuffer2D();
 
 	FORCEINLINE ENiagaraSimTarget GetSimTarget() const { return SimTarget; }
diff --git a/Engine/Plugins/FX/Niagara/Source/NiagaraCore/Private/NiagaraCustomVersion.cpp b/Engine/Plugins/FX/Niagara/Source/NiagaraCore/Private/NiagaraCustomVersion.cpp
--- a/Engine/Plugins/FX/Niagara/Source/NiagaraCore/Private/NiagaraCustomVersion.cpp
+++ b/Engine/Plugins/FX/Niagara/Source/NiagaraCore/Private/NiagaraCustomVersion.cpp
@@ -8,4 +8,4 @@ const FGuid FNiagaraCustomVersion::GUID(0xFCF57AFA, 0x50764283, 0xB9A9E658, 0xFF
 // Register the custom version with core
 FCustomVersionRegistration GRegisterNiagaraCustomVersion(FNiagaraCustomVersion::GUID, FNiagaraCustomVersion::LatestVersion, TEXT("NiagaraVer"));
 
-const FGuid FNiagaraCustomVersion::LatestScriptCompileVersion(0x3F7E65A5, 0xF86140EC, 0x8FAB94B7, 0xE9F031FD);
+const FGuid FNiagaraCustomVersion::LatestScriptCompileVersion(0x6D52A92E, 0xDC5C4791, 0x922FD12B, 0x82A7D6B7);
\ No newline at end of file
