Unshelved from pending changelist '12049653':

Niagara - When compile errors occur show them in modules and dynamic inputs in the stack where possible.

diff --git a/Engine/Plugins/FX/Niagara/Source/NiagaraEditor/Private/NiagaraMessageManager.cpp b/Engine/Plugins/FX/Niagara/Source/NiagaraEditor/Private/NiagaraMessageManager.cpp
--- a/Engine/Plugins/FX/Niagara/Source/NiagaraEditor/Private/NiagaraMessageManager.cpp
+++ b/Engine/Plugins/FX/Niagara/Source/NiagaraEditor/Private/NiagaraMessageManager.cpp
@@ -52,13 +52,13 @@ void FNiagaraMessageManager::RefreshMessagesForAssetKey(const FGuid& InAssetKey)
 
 void FNiagaraMessageManager::RefreshMessagesForAssetKeyAndMessageJobType(const FGuid& InAssetKey, const ENiagaraMessageJobType& InMessageJobType)
 {
-	TMap<const ENiagaraMessageJobType, TArray<TSharedRef<const INiagaraMessage>>>* MessageJobTypeToMessageMap = ObjectToTypeMappedMessagesMap.Find(InAssetKey);
+	TMap<const ENiagaraMessageJobType, FMessageData>* MessageJobTypeToMessageMap = ObjectToTypeMappedMessagesMap.Find(InAssetKey);
 	if (MessageJobTypeToMessageMap != nullptr)
 	{
-		TArray<TSharedRef<const INiagaraMessage>>* MessageArr = MessageJobTypeToMessageMap->Find(InMessageJobType);
-		if (MessageArr != nullptr)
+		FMessageData* MessageData = MessageJobTypeToMessageMap->Find(InMessageJobType);
+		if (MessageData != nullptr)
 		{
-			MessageArr->Reset();
+			MessageData->Reset();
 		}
 	} 
 
@@ -118,12 +118,26 @@ const TOptional<const FString> FNiagaraMessageManager::GetStringForScriptUsageIn
 const TArray<TSharedRef<const INiagaraMessage>> FNiagaraMessageManager::GetMessagesForAssetKey(const FGuid& InAssetKey) const
 {
  	TArray<TSharedRef<const INiagaraMessage>> FoundMessages;
-	const TMap <const ENiagaraMessageJobType, TArray<TSharedRef<const INiagaraMessage>>>* MessageJobTypeToMessagesMapPtr = ObjectToTypeMappedMessagesMap.Find(InAssetKey);
+	const TMap <const ENiagaraMessageJobType, FMessageData>* MessageJobTypeToMessagesMapPtr = ObjectToTypeMappedMessagesMap.Find(InAssetKey);
 	if (MessageJobTypeToMessagesMapPtr != nullptr)
 	{
-		for (const TTuple<const ENiagaraMessageJobType, TArray<TSharedRef<const INiagaraMessage>>>& MessageJobTypeToMessagesPair : *MessageJobTypeToMessagesMapPtr)
+		for (const TTuple<const ENiagaraMessageJobType, FMessageData>& MessageJobTypeToMessagesPair : *MessageJobTypeToMessagesMapPtr)
 		{
-			FoundMessages.Append(MessageJobTypeToMessagesPair.Value);
+			FoundMessages.Append(MessageJobTypeToMessagesPair.Value.GetMessages());
+		}
+	}
+	return FoundMessages;
+}
+
+const TArray<TSharedRef<const INiagaraMessage>> FNiagaraMessageManager::GetMessagesForAssetKeyAndObjectKey(const FGuid& InAssetKey, const FObjectKey& InObjectKey)
+{
+	TArray<TSharedRef<const INiagaraMessage>> FoundMessages;
+	const TMap <const ENiagaraMessageJobType, FMessageData>* MessageJobTypeToMessagesMapPtr = ObjectToTypeMappedMessagesMap.Find(InAssetKey);
+	if (MessageJobTypeToMessagesMapPtr != nullptr)
+	{
+		for (const TTuple<const ENiagaraMessageJobType, FMessageData>& MessageJobTypeToMessagesPair : *MessageJobTypeToMessagesMapPtr)
+		{
+			MessageJobTypeToMessagesPair.Value.GetMessagesForObjectKey(InObjectKey, FoundMessages);
 		}
 	}
 	return FoundMessages;
@@ -150,7 +164,7 @@ void FNiagaraMessageManager::DoMessageJobsInQueueTick()
 				TSharedPtr<const INiagaraMessageJob> CurrentMessageJob = CurrentMessageJobBatch.MessageJobs.Pop(false);
 				TSharedRef<const INiagaraMessage> GeneratedMessage = CurrentMessageJob->GenerateNiagaraMessage();
 				GeneratedMessagesForCurrentMessageJobBatch.Add(GeneratedMessage);
-				ObjectToTypeMappedMessagesMap.FindOrAdd(CurrentMessageJobBatch.MessageJobsAssetKey).FindOrAdd(CurrentMessageJob->GetMessageJobType()).Add(GeneratedMessage);
+				ObjectToTypeMappedMessagesMap.FindOrAdd(CurrentMessageJobBatch.MessageJobsAssetKey).FindOrAdd(CurrentMessageJob->GetMessageJobType()).AddMessage(GeneratedMessage);
 				CurrentWorkLoopTime = FPlatformTime::Seconds();
 
 				if (CurrentWorkLoopTime - WorkStartTime > MaxJobWorkTime)
@@ -170,11 +184,13 @@ FNiagaraMessageCompileEvent::FNiagaraMessageCompileEvent(
 	, TArray<FNiagaraScriptNameAndAssetPath>& InContextScriptNamesAndAssetPaths
 	, TOptional<const FText>& InOwningScriptNameAndUsageText
 	, TOptional<const FNiagaraScriptNameAndAssetPath>& InCompiledScriptNameAndAssetPath
+	, const TArray<FObjectKey>& InAssociatedObjectKeys
 	)
 	: CompileEvent(InCompileEvent)
 	, ContextScriptNamesAndAssetPaths(InContextScriptNamesAndAssetPaths)
 	, OwningScriptNameAndUsageText(InOwningScriptNameAndUsageText)
 	, CompiledScriptNameAndAssetPath(InCompiledScriptNameAndAssetPath)
+	, AssociatedObjectKeys(InAssociatedObjectKeys)
 {
 }
 
@@ -269,6 +285,59 @@ TSharedRef<FTokenizedMessage> FNiagaraMessageCompileEvent::GenerateTokenizedMess
 	return Message;
 }
 
+FText FNiagaraMessageCompileEvent::GenerateMessageText() const
+{
+	if (OwningScriptNameAndUsageText.IsSet())
+	{
+		return FText::Format(LOCTEXT("MessageAndOwnerFormat", "{0}{1}"), FText::FromString(CompileEvent.Message), OwningScriptNameAndUsageText.GetValue());
+	}
+	else
+	{
+		return FText::FromString(CompileEvent.Message);
+	}
+}
+
+void FNiagaraMessageCompileEvent::GenerateLinks(TArray<FText>& OutLinkDisplayNames, TArray<FSimpleDelegate>& OutLinkNavigationActions) const
+{
+	// TODO: Do this without instantiating these tokens.
+	if ((CompiledScriptNameAndAssetPath.IsSet() || ContextScriptNamesAndAssetPaths.Num() > 0) &&
+		(CompileEvent.NodeGuid.IsValid() || CompileEvent.PinGuid.IsValid()))
+	{
+		FString AssetPath = CompiledScriptNameAndAssetPath.IsSet()
+			? CompiledScriptNameAndAssetPath.GetValue().ScriptAssetPathString
+			: ContextScriptNamesAndAssetPaths.Last().ScriptAssetPathString;
+
+		TSharedPtr<FNiagaraCompileEventToken> NavigateToSourceToken;
+		if (CompileEvent.PinGuid.IsValid())
+		{
+			NavigateToSourceToken = FNiagaraCompileEventToken::Create(AssetPath, FText::FromString(CompileEvent.Message), CompileEvent.NodeGuid, CompileEvent.PinGuid);
+		}
+		else if (CompileEvent.NodeGuid.IsValid())
+		{
+			NavigateToSourceToken = FNiagaraCompileEventToken::Create(AssetPath, FText::FromString(CompileEvent.Message), CompileEvent.NodeGuid);
+		}
+
+		if (NavigateToSourceToken.IsValid())
+		{
+			OutLinkDisplayNames.Add(LOCTEXT("NavigateToMessageSource", "Navigate to message source"));
+			OutLinkNavigationActions.Add(FSimpleDelegate::CreateLambda([NavigateToSourceToken]()
+			{
+				NavigateToSourceToken->GetOnMessageTokenActivated().Execute(NavigateToSourceToken.ToSharedRef());
+			}));
+		}
+	}
+
+	for (const FNiagaraScriptNameAndAssetPath& ScriptNameAndPath : ContextScriptNamesAndAssetPaths)
+	{
+		OutLinkDisplayNames.Add(FText::Format(LOCTEXT("NavigateToAssetFormat", "Navigate to asset: {0}"), FText::FromString(*ScriptNameAndPath.ScriptNameString)));
+		TSharedRef<FNiagaraCompileEventToken> NavigateToAssetToken = FNiagaraCompileEventToken::Create(ScriptNameAndPath.ScriptAssetPathString, FText::FromString(*ScriptNameAndPath.ScriptNameString));
+		OutLinkNavigationActions.Add(FSimpleDelegate::CreateLambda([NavigateToAssetToken]()
+		{
+			NavigateToAssetToken->GetOnMessageTokenActivated().Execute(NavigateToAssetToken);
+		}));
+	}
+}
+
 FNiagaraMessageJobCompileEvent::FNiagaraMessageJobCompileEvent(
 	  const FNiagaraCompileEvent& InCompileEvent
 	, const TWeakObjectPtr<const UNiagaraScript>& InOriginatingScriptWeakObjPtr
@@ -300,7 +369,8 @@ TSharedRef<const INiagaraMessage> FNiagaraMessageJobCompileEvent::GenerateNiagar
 		TOptional<const FText> OwningScriptNameAndUsageText;
 		TOptional<const FText> ScriptNameAndPathsGetterFailureReason;
 		TOptional<const FString> OwningScriptNameStringCopy = OwningScriptNameString;
-		bool bSuccessfullyFoundScripts = RecursiveGetScriptNamesAndAssetPathsFromContextStack(ContextStackGuids, ScriptGraph, ContextScriptNamesAndPaths, OwningScriptNameStringCopy, ScriptNameAndPathsGetterFailureReason);
+		TArray<FObjectKey> ContextObjectKeys;
+		bool bSuccessfullyFoundScripts = RecursiveGetScriptNamesAndAssetPathsFromContextStack(ContextStackGuids, ScriptGraph, ContextScriptNamesAndPaths, OwningScriptNameStringCopy, ScriptNameAndPathsGetterFailureReason, ContextObjectKeys);
 
 		if (bSuccessfullyFoundScripts == false)
 		{
@@ -326,7 +396,7 @@ TSharedRef<const INiagaraMessage> FNiagaraMessageJobCompileEvent::GenerateNiagar
 			//If this compile event is from a script, set the compiled script name and asset path so the user can navigate to errors locally.
 			CompiledScriptNameAndPath = TOptional<const FNiagaraScriptNameAndAssetPath>(FNiagaraScriptNameAndAssetPath(OriginatingScriptWeakObjPtr->GetName(), SourceScriptAssetPath.GetValue()));
 		}
-		return MakeShared<FNiagaraMessageCompileEvent>(CompileEvent, ContextScriptNamesAndPaths, OwningScriptNameAndUsageText, CompiledScriptNameAndPath);
+		return MakeShared<FNiagaraMessageCompileEvent>(CompileEvent, ContextScriptNamesAndPaths, OwningScriptNameAndUsageText, CompiledScriptNameAndPath, ContextObjectKeys);
 	}
 	//The originating script weak ptr is no longer valid, send an error message asking for recompile.
 	FText MessageText = LOCTEXT("CompileEventMessageJobFail", "Cached info for compile event is out of date, recompile to get full info. Event: {0}");
@@ -341,6 +411,7 @@ bool FNiagaraMessageJobCompileEvent::RecursiveGetScriptNamesAndAssetPathsFromCon
 	, TArray<FNiagaraScriptNameAndAssetPath>& OutContextScriptNamesAndAssetPaths
 	, TOptional<const FString>& OutEmitterName
 	, TOptional<const FText>& OutFailureReason
+	, TArray<FObjectKey>& OutContextNodeObjectKeys
 ) const
 {
 	checkf(InGraphToSearch, TEXT("Failed to get a node graph to search!"));
@@ -371,6 +442,7 @@ bool FNiagaraMessageJobCompileEvent::RecursiveGetScriptNamesAndAssetPathsFromCon
 	if (ContextNode.IsSet())
 	{
 		// found a node in the current graph that has a GUID in the context list
+		OutContextNodeObjectKeys.Add(FObjectKey(ContextNode.GetValue()));
 
 		UNiagaraNodeFunctionCall* FunctionCallNode = Cast<UNiagaraNodeFunctionCall>(ContextNode.GetValue());
 		if (FunctionCallNode)
@@ -401,7 +473,7 @@ bool FNiagaraMessageJobCompileEvent::RecursiveGetScriptNamesAndAssetPathsFromCon
 				return false;
 			}
 			OutContextScriptNamesAndAssetPaths.Add(FNiagaraScriptNameAndAssetPath(FunctionCallNodeAssignedScript->GetName(), FunctionCallNodeAssignedScript->GetPathName()));
-			return RecursiveGetScriptNamesAndAssetPathsFromContextStack(InContextStackNodeGuids, FunctionScriptGraph, OutContextScriptNamesAndAssetPaths, OutEmitterName, OutFailureReason);
+			return RecursiveGetScriptNamesAndAssetPathsFromContextStack(InContextStackNodeGuids, FunctionScriptGraph, OutContextScriptNamesAndAssetPaths, OutEmitterName, OutFailureReason, OutContextNodeObjectKeys);
 		}
 
 		UNiagaraNodeEmitter* EmitterNode = Cast<UNiagaraNodeEmitter>(ContextNode.GetValue());
@@ -414,7 +486,7 @@ bool FNiagaraMessageJobCompileEvent::RecursiveGetScriptNamesAndAssetPathsFromCon
 			checkf(EmitterScriptGraph, TEXT("Emitter Script Source does not have a UNiagaraGraph!"));
 
 			OutEmitterName = EmitterNode->GetEmitterUniqueName();
-			return RecursiveGetScriptNamesAndAssetPathsFromContextStack(InContextStackNodeGuids, EmitterScriptGraph, OutContextScriptNamesAndAssetPaths, OutEmitterName, OutFailureReason);
+			return RecursiveGetScriptNamesAndAssetPathsFromContextStack(InContextStackNodeGuids, EmitterScriptGraph, OutContextScriptNamesAndAssetPaths, OutEmitterName, OutFailureReason, OutContextNodeObjectKeys);
 		}
 		checkf(false, TEXT("Matching node is not a function call or emitter node!"));
 	}
@@ -462,11 +534,21 @@ TSharedRef<FTokenizedMessage> FNiagaraMessageNeedRecompile::GenerateTokenizedMes
 	return FTokenizedMessage::Create(EMessageSeverity::Error, NeedRecompileMessage);
 }
 
+FText FNiagaraMessageNeedRecompile::GenerateMessageText() const
+{
+	return NeedRecompileMessage;
+}
+
 TSharedRef<FTokenizedMessage> FNiagaraMessagePostCompileSummary::GenerateTokenizedMessage() const
 {
 	return FTokenizedMessage::Create(MessageSeverity, PostCompileSummaryText);
 }
 
+FText FNiagaraMessagePostCompileSummary::GenerateMessageText() const
+{
+	return PostCompileSummaryText;
+}
+
 FNiagaraCompileEventToken::FNiagaraCompileEventToken(
 	  const FString& InScriptAssetPath
 	, const FText& InMessage
@@ -533,4 +615,35 @@ TSharedRef<FNiagaraCompileEventToken> FNiagaraCompileEventToken::Create(
 	return MakeShareable(new FNiagaraCompileEventToken(InScriptAssetPath, InMessage, InNodeGUID, InPinGUID));
 }
 
+const TArray<TSharedRef<const INiagaraMessage>>& FNiagaraMessageManager::FMessageData::GetMessages() const
+{
+	return Messages;
+}
+
+const void FNiagaraMessageManager::FMessageData::GetMessagesForObjectKey(FObjectKey InObjectKey, TArray<TSharedRef<const INiagaraMessage>>& OutMessages) const
+{
+	const TArray<TSharedRef<const INiagaraMessage>>* FoundMessages = ObjectKeyToMessagesMap.Find(InObjectKey);
+	if (FoundMessages != nullptr)
+	{
+		OutMessages.Append(*FoundMessages);
+	}
+}
+
+void FNiagaraMessageManager::FMessageData::AddMessage(TSharedRef<const INiagaraMessage> InMessage)
+{
+	Messages.Add(InMessage);
+	TArray<FObjectKey> AssociatedObjectKeys;
+	InMessage->GetAssociatedObjectKeys(AssociatedObjectKeys);
+	for(const FObjectKey& AssociatedObjectKey : AssociatedObjectKeys)
+	{
+		ObjectKeyToMessagesMap.FindOrAdd(AssociatedObjectKey).Add(InMessage);
+	}
+}
+
+void FNiagaraMessageManager::FMessageData::Reset()
+{
+	Messages.Reset();
+	ObjectKeyToMessagesMap.Reset();
+}
+
 #undef LOCTEXT_NAMESPACE //NiagaraMessageManager
diff --git a/Engine/Plugins/FX/Niagara/Source/NiagaraEditor/Private/ViewModels/NiagaraSystemViewModel.cpp b/Engine/Plugins/FX/Niagara/Source/NiagaraEditor/Private/ViewModels/NiagaraSystemViewModel.cpp
--- a/Engine/Plugins/FX/Niagara/Source/NiagaraEditor/Private/ViewModels/NiagaraSystemViewModel.cpp
+++ b/Engine/Plugins/FX/Niagara/Source/NiagaraEditor/Private/ViewModels/NiagaraSystemViewModel.cpp
@@ -419,6 +419,11 @@ void FNiagaraSystemViewModel::DuplicateEmitters(TArray<FEmitterHandleToDuplicate
 	bForceAutoCompileOnce = true;
 }
 
+FGuid FNiagaraSystemViewModel::GetMessageLogGuid() const
+{
+	return SystemMessageLogGuidKey.IsSet() ? SystemMessageLogGuidKey.GetValue() : FGuid();
+}
+
 void FNiagaraSystemViewModel::DeleteEmitters(TSet<FGuid> EmitterHandleIdsToDelete)
 {
 	if (EmitterHandleIdsToDelete.Num() > 0)
diff --git a/Engine/Plugins/FX/Niagara/Source/NiagaraEditor/Private/ViewModels/Stack/NiagaraStackEntry.cpp b/Engine/Plugins/FX/Niagara/Source/NiagaraEditor/Private/ViewModels/Stack/NiagaraStackEntry.cpp
--- a/Engine/Plugins/FX/Niagara/Source/NiagaraEditor/Private/ViewModels/Stack/NiagaraStackEntry.cpp
+++ b/Engine/Plugins/FX/Niagara/Source/NiagaraEditor/Private/ViewModels/Stack/NiagaraStackEntry.cpp
@@ -26,9 +26,10 @@ UNiagaraStackEntry::FStackIssueFix::FStackIssueFix()
 {
 }
 
-UNiagaraStackEntry::FStackIssueFix::FStackIssueFix(FText InDescription, FStackIssueFixDelegate InFixDelegate)
+UNiagaraStackEntry::FStackIssueFix::FStackIssueFix(FText InDescription, FStackIssueFixDelegate InFixDelegate, EStackIssueFixStyle InFixStyle)
 	: Description(InDescription)
 	, FixDelegate(InFixDelegate)
+	, Style(InFixStyle)
 	, UniqueIdentifier(FMD5::HashAnsiString(*FString::Printf(TEXT("%s"), *InDescription.ToString())))
 {
 	checkf(Description.IsEmptyOrWhitespace() == false, TEXT("Description can not be empty."));
@@ -60,6 +61,11 @@ const UNiagaraStackEntry::FStackIssueFixDelegate& UNiagaraStackEntry::FStackIssu
 	return FixDelegate;
 }
 
+UNiagaraStackEntry::EStackIssueFixStyle UNiagaraStackEntry::FStackIssueFix::GetStyle() const
+{
+	return Style;
+}
+
 UNiagaraStackEntry::FStackIssue::FStackIssue()
 {
 }
diff --git a/Engine/Plugins/FX/Niagara/Source/NiagaraEditor/Private/ViewModels/Stack/NiagaraStackFunctionInput.cpp b/Engine/Plugins/FX/Niagara/Source/NiagaraEditor/Private/ViewModels/Stack/NiagaraStackFunctionInput.cpp
--- a/Engine/Plugins/FX/Niagara/Source/NiagaraEditor/Private/ViewModels/Stack/NiagaraStackFunctionInput.cpp
+++ b/Engine/Plugins/FX/Niagara/Source/NiagaraEditor/Private/ViewModels/Stack/NiagaraStackFunctionInput.cpp
@@ -33,6 +33,7 @@
 #include "NiagaraEmitter.h"
 #include "NiagaraClipboard.h"
 #include "Toolkits/NiagaraSystemToolkit.h"
+#include "NiagaraMessageManager.h"
 
 #include "Materials/Material.h"
 #include "Materials/MaterialExpressionDynamicParameter.h"
@@ -183,6 +184,12 @@ void UNiagaraStackFunctionInput::FinalizeInternal()
 		SourceScript->GetSource()->OnChanged().RemoveAll(this);
 	}
 
+	if (MessageManagerRefreshHandle.IsValid())
+	{
+		FNiagaraMessageManager::Get()->GetOnRequestRefresh().Remove(MessageManagerRefreshHandle);
+		MessageManagerRefreshHandle.Reset();
+	}
+
 	Super::FinalizeInternal();
 }
 
@@ -389,6 +396,22 @@ void UNiagaraStackFunctionInput::RefreshChildrenInternal(const TArray<UNiagaraSt
 	RefreshFromMetaData();
 	RefreshValues();
 
+	if (InputValues.DynamicNode.IsValid())
+	{
+		if (MessageManagerRefreshHandle.IsValid() == false)
+		{
+			MessageManagerRefreshHandle = FNiagaraMessageManager::Get()->GetOnRequestRefresh().AddUObject(this, &UNiagaraStackFunctionInput::OnMessageManagerRefresh);
+		}
+	}
+	else
+	{
+		if (MessageManagerRefreshHandle.IsValid())
+		{
+			FNiagaraMessageManager::Get()->GetOnRequestRefresh().Remove(MessageManagerRefreshHandle);
+			MessageManagerRefreshHandle.Reset();
+		}
+	}
+
 	if (InputValues.Mode == EValueMode::Dynamic && InputValues.DynamicNode.IsValid())
 	{
 		if (InputValues.DynamicNode->FunctionScript != nullptr)
@@ -542,6 +565,16 @@ void UNiagaraStackFunctionInput::RefreshChildrenInternal(const TArray<UNiagaraSt
 			DisplayNameOverride = FText::FromString(DisplayNameStr);
 		}	
 	}
+
+	if (InputValues.DynamicNode.IsValid())
+	{
+		TArray<TSharedRef<const INiagaraMessage>> Messages = FNiagaraMessageManager::Get()->GetMessagesForAssetKeyAndObjectKey(
+			GetSystemViewModel()->GetMessageLogGuid(), FObjectKey(InputValues.DynamicNode.Get()));
+		for (TSharedRef<const INiagaraMessage> Message : Messages)
+		{
+			NewIssues.Add(FNiagaraStackGraphUtilities::MessageManagerMessageToStackIssue(Message, GetStackEditorDataKey()));
+		}
+	}
 }
 
 class UNiagaraStackFunctionInputUtilities
@@ -1471,6 +1504,17 @@ FNiagaraVariable UNiagaraStackFunctionInput::CreateRapidIterationVariable(const
 	return FNiagaraStackGraphUtilities::CreateRapidIterationParameter(UniqueEmitterName, OutputNode->GetUsage(), InName, InputType);
 }
 
+void UNiagaraStackFunctionInput::OnMessageManagerRefresh(const FGuid& MessageJobBatchAssetKey, const TArray<TSharedRef<const INiagaraMessage>> NewMessages)
+{
+	if (InputValues.DynamicNode.IsValid() && GetSystemViewModel()->GetMessageLogGuid() == MessageJobBatchAssetKey)
+	{
+		if (FNiagaraMessageManager::Get()->GetMessagesForAssetKeyAndObjectKey(MessageJobBatchAssetKey, FObjectKey(InputValues.DynamicNode.Get())).Num() > 0)
+		{
+			RefreshChildren();
+		}
+	}
+}
+
 bool UNiagaraStackFunctionInput::SupportsRename() const
 {
 	// Only module level assignment node inputs can be renamed.
diff --git a/Engine/Plugins/FX/Niagara/Source/NiagaraEditor/Private/ViewModels/Stack/NiagaraStackGraphUtilities.cpp b/Engine/Plugins/FX/Niagara/Source/NiagaraEditor/Private/ViewModels/Stack/NiagaraStackGraphUtilities.cpp
--- a/Engine/Plugins/FX/Niagara/Source/NiagaraEditor/Private/ViewModels/Stack/NiagaraStackGraphUtilities.cpp
+++ b/Engine/Plugins/FX/Niagara/Source/NiagaraEditor/Private/ViewModels/Stack/NiagaraStackGraphUtilities.cpp
@@ -31,6 +31,7 @@
 #include "ViewModels/NiagaraEmitterViewModel.h"
 #include "AssetRegistryModule.h"
 #include "ObjectTools.h"
+#include "NiagaraMessageManager.h"
 
 DECLARE_CYCLE_STAT(TEXT("Niagara - StackGraphUtilities - RelayoutGraph"), STAT_NiagaraEditor_StackGraphUtilities_RelayoutGraph, STATGROUP_NiagaraEditor);
 
@@ -2540,4 +2541,60 @@ void FNiagaraStackGraphUtilities::RenameReferencingParameters(UNiagaraSystem& Sy
 	}
 }
 
+UNiagaraStackEntry::FStackIssue FNiagaraStackGraphUtilities::MessageManagerMessageToStackIssue(TSharedRef<const INiagaraMessage> InMessage, FString InStackEditorDataKey)
+{
+	TSharedRef<FTokenizedMessage> TokenizedMessage = InMessage->GenerateTokenizedMessage();
+	EStackIssueSeverity StackIssueSeverity;
+	switch (TokenizedMessage->GetSeverity())
+	{
+	case EMessageSeverity::CriticalError:
+	case EMessageSeverity::Error:
+		StackIssueSeverity = EStackIssueSeverity::Error;
+		break;
+	case EMessageSeverity::PerformanceWarning:
+	case EMessageSeverity::Warning:
+		StackIssueSeverity = EStackIssueSeverity::Warning;
+		break;
+	case EMessageSeverity::Info:
+		StackIssueSeverity = EStackIssueSeverity::Info;
+		break;
+	default:
+		StackIssueSeverity = EStackIssueSeverity::Info;
+		break;
+	}
+
+	FText ShortDescription;
+	switch (InMessage->GetMessageType())
+	{
+	case ENiagaraMessageType::CompileEventMessage:
+		ShortDescription = LOCTEXT("CompileErrorShortDescription", "Compile Error");
+		break;
+	default:
+		ShortDescription = LOCTEXT("UnspecifiedErrorShortDescription", "Unspecified Error");
+		break;
+	}
+
+	TArray<UNiagaraStackEntry::FStackIssueFix> FixLinks;
+	TArray<FText> LinkMessages;
+	TArray<FSimpleDelegate> LinkNavigateActions;
+	InMessage->GenerateLinks(LinkMessages, LinkNavigateActions);
+	for (int32 LinkIndex = 0; LinkIndex < LinkMessages.Num(); LinkIndex++)
+	{
+		const FText& LinkMessage = LinkMessages[LinkIndex];
+		const FSimpleDelegate& LinkNavigateAction = LinkNavigateActions[LinkIndex];
+		FixLinks.Add(UNiagaraStackEntry::FStackIssueFix(
+			LinkMessage,
+			UNiagaraStackEntry::FStackIssueFixDelegate::CreateLambda([LinkNavigateAction]() { LinkNavigateAction.Execute(); }),
+			UNiagaraStackEntry::EStackIssueFixStyle::Link));
+	}
+
+	return UNiagaraStackEntry::FStackIssue(
+		StackIssueSeverity,
+		ShortDescription,
+		InMessage->GenerateMessageText(),
+		InStackEditorDataKey,
+		false,
+		FixLinks);
+}
+
 #undef LOCTEXT_NAMESPACE
diff --git a/Engine/Plugins/FX/Niagara/Source/NiagaraEditor/Private/ViewModels/Stack/NiagaraStackModuleItem.cpp b/Engine/Plugins/FX/Niagara/Source/NiagaraEditor/Private/ViewModels/Stack/NiagaraStackModuleItem.cpp
--- a/Engine/Plugins/FX/Niagara/Source/NiagaraEditor/Private/ViewModels/Stack/NiagaraStackModuleItem.cpp
+++ b/Engine/Plugins/FX/Niagara/Source/NiagaraEditor/Private/ViewModels/Stack/NiagaraStackModuleItem.cpp
@@ -34,6 +34,7 @@
 #include "NiagaraClipboard.h"
 #include "NiagaraConvertInPlaceUtilityBase.h"
 #include "Framework/Notifications/NotificationManager.h"
+#include "NiagaraMessageManager.h"
 
 #include "ScopedTransaction.h"
 #include "NiagaraScriptVariable.h"
@@ -138,6 +139,8 @@ void UNiagaraStackModuleItem::Initialize(FRequiredEntryData InRequiredEntryData,
 		AddChildFilter(FOnFilterChild::CreateUObject(this, &UNiagaraStackModuleItem::FilterOutputCollection));
 		AddChildFilter(FOnFilterChild::CreateUObject(this, &UNiagaraStackModuleItem::FilterLinkedInputCollection));
 	}
+
+	FNiagaraMessageManager::Get()->GetOnRequestRefresh().AddUObject(this, &UNiagaraStackModuleItem::OnMessageManagerRefresh);
 }
 
 FText UNiagaraStackModuleItem::GetDisplayName() const
@@ -167,6 +170,12 @@ INiagaraStackItemGroupAddUtilities* UNiagaraStackModuleItem::GetGroupAddUtilitie
 	return GroupAddUtilities;
 }
 
+void UNiagaraStackModuleItem::FinalizeInternal()
+{
+	FNiagaraMessageManager::Get()->GetOnRequestRefresh().RemoveAll(this);
+	Super::FinalizeInternal();
+}
+
 void UNiagaraStackModuleItem::RefreshChildrenInternal(const TArray<UNiagaraStackEntry*>& CurrentChildren, TArray<UNiagaraStackEntry*>& NewChildren, TArray<FStackIssue>& NewIssues)
 {
 	bCanRefresh = false;
@@ -392,6 +401,13 @@ void UNiagaraStackModuleItem::RefreshIssues(TArray<FStackIssue>& NewIssues)
 			}
 		}
 
+		TArray<TSharedRef<const INiagaraMessage>> Messages = FNiagaraMessageManager::Get()->GetMessagesForAssetKeyAndObjectKey(
+			GetSystemViewModel()->GetMessageLogGuid(), FObjectKey(FunctionCallNode));
+		for (TSharedRef<const INiagaraMessage> Message : Messages)
+		{
+			NewIssues.Add(FNiagaraStackGraphUtilities::MessageManagerMessageToStackIssue(Message, GetStackEditorDataKey()));
+		}
+
 		if (FunctionCallNode->FunctionScript == nullptr && FunctionCallNode->GetClass() == UNiagaraNodeFunctionCall::StaticClass())
 		{
 			FText ModuleScriptMissingShort = LOCTEXT("ModuleScriptMissingShort", "Missing module script");
@@ -831,6 +847,17 @@ void UNiagaraStackModuleItem::RefreshIsEnabled()
 	}
 }
 
+void UNiagaraStackModuleItem::OnMessageManagerRefresh(const FGuid& MessageJobBatchAssetKey, const TArray<TSharedRef<const INiagaraMessage>> NewMessages)
+{
+	if (GetSystemViewModel()->GetMessageLogGuid() == MessageJobBatchAssetKey)
+	{
+		if (FNiagaraMessageManager::Get()->GetMessagesForAssetKeyAndObjectKey(MessageJobBatchAssetKey, FObjectKey(FunctionCallNode)).Num() > 0)
+		{
+			RefreshChildren();
+		}
+	}
+}
+
 bool UNiagaraStackModuleItem::CanMoveAndDelete() const
 {
 	if (bCanMoveAndDeleteCache.IsSet() == false)
diff --git a/Engine/Plugins/FX/Niagara/Source/NiagaraEditor/Public/NiagaraMessageManager.h b/Engine/Plugins/FX/Niagara/Source/NiagaraEditor/Public/NiagaraMessageManager.h
--- a/Engine/Plugins/FX/Niagara/Source/NiagaraEditor/Public/NiagaraMessageManager.h
+++ b/Engine/Plugins/FX/Niagara/Source/NiagaraEditor/Public/NiagaraMessageManager.h
@@ -7,6 +7,7 @@
 #include "Logging/TokenizedMessage.h"
 #include "NiagaraGraph.h"
 #include "TickableEditorObject.h"
+#include "UObject/ObjectKey.h"
 
 class FNiagaraScriptToolkit;
 
@@ -46,6 +47,12 @@ public:
 
 	virtual const ENiagaraMessageType GetMessageType() const = 0;
 
+	// TODO: Add a better API for generating stack issues.
+	virtual FText GenerateMessageText() const = 0;
+	virtual void GenerateLinks(TArray<FText>& OutLinkDisplayNames, TArray<FSimpleDelegate>& OutLinkNavigationActions) const = 0;
+
+	virtual void GetAssociatedObjectKeys(TArray<FObjectKey>& OutAssociatedObjectKeys) const = 0;
+
 	virtual ~INiagaraMessage() {};
 };
 
@@ -84,6 +91,7 @@ private:
 		, TArray<FNiagaraScriptNameAndAssetPath>& OutContextScriptNamesAndAssetPaths
 		, TOptional<const FString>& OutEmitterName
 		, TOptional<const FText>& OutFailureReason
+		, TArray<FObjectKey>& OutContextNodeObjectKeys
 	) const;
 
 	const FNiagaraCompileEvent CompileEvent;
@@ -123,18 +131,25 @@ public:
 		, TArray<FNiagaraScriptNameAndAssetPath>& InContextScriptNamesAndAssetPaths
 		, TOptional<const FText>& InOwningScriptNameAndUsageText
 		, TOptional<const FNiagaraScriptNameAndAssetPath>& InCompiledScriptNameAndAssetPath
+		, const TArray<FObjectKey>& InAssociatedObjectKeys
 	);
 
 	virtual TSharedRef<FTokenizedMessage> GenerateTokenizedMessage() const override;
-	//@todo(message manager) make stack specific message type generator here
 
 	virtual const ENiagaraMessageType GetMessageType() const override { return ENiagaraMessageType::CompileEventMessage; };
 
+	virtual FText GenerateMessageText() const override;
+
+	virtual void GenerateLinks(TArray<FText>& OutLinkDisplayNames, TArray<FSimpleDelegate>& OutLinkNavigationActions) const override;
+
+	virtual void GetAssociatedObjectKeys(TArray<FObjectKey>& OutAssociatedObjectKeys) const override { OutAssociatedObjectKeys.Append(AssociatedObjectKeys); }
+
 private:
 	const FNiagaraCompileEvent CompileEvent;
 	const TArray<FNiagaraScriptNameAndAssetPath> ContextScriptNamesAndAssetPaths;
 	const TOptional<const FText> OwningScriptNameAndUsageText;
 	const TOptional<const FNiagaraScriptNameAndAssetPath> CompiledScriptNameAndAssetPath;
+	const TArray<FObjectKey> AssociatedObjectKeys;
 };
 
 class FNiagaraMessageNeedRecompile : public INiagaraMessage
@@ -149,6 +164,12 @@ public:
 
 	virtual const ENiagaraMessageType GetMessageType() const override { return ENiagaraMessageType::NeedRecompileMessage; };
 
+	virtual FText GenerateMessageText() const override;
+
+	virtual void GenerateLinks(TArray<FText>& OutLinkDisplayNames, TArray<FSimpleDelegate>& OutLinkNavigationActions) const override { }
+
+	virtual void GetAssociatedObjectKeys(TArray<FObjectKey>& OutAssociatedObjectKeys) const override { }
+
 private:
 	const FText NeedRecompileMessage;
 };
@@ -167,6 +188,12 @@ public:
 	// we alias the message type as post compile summary is always generated with compile events
 	virtual const ENiagaraMessageType GetMessageType() const override { return  ENiagaraMessageType::CompileEventMessage; };
 
+	virtual FText GenerateMessageText() const override;
+
+	virtual void GenerateLinks(TArray<FText>& OutLinkDisplayNames, TArray<FSimpleDelegate>& OutLinkNavigationActions) const override { }
+
+	virtual void GetAssociatedObjectKeys(TArray<FObjectKey>& OutAssociatedObjectKeys) const override { }
+
 private:
 	const FText PostCompileSummaryText;
 	const EMessageSeverity::Type MessageSeverity;
@@ -253,6 +280,8 @@ public:
 
 	const TArray<TSharedRef<const INiagaraMessage>> GetMessagesForAssetKey(const FGuid& InAssetKey) const;
 
+	const TArray<TSharedRef<const INiagaraMessage>> GetMessagesForAssetKeyAndObjectKey(const FGuid& InAssetKey, const FObjectKey& InObjectKey);
+
 	FOnRequestRefresh& GetOnRequestRefresh() { return OnRequestRefresh; };
 
 	//~ Begin FTickableEditorObject Interface.
@@ -264,6 +293,19 @@ public:
 	//~ End FTickableEditorObject Interface
 
 private:
+	class FMessageData
+	{
+	public:
+		const TArray<TSharedRef<const INiagaraMessage>>& GetMessages() const;
+		const void GetMessagesForObjectKey(FObjectKey InObjectKey, TArray<TSharedRef<const INiagaraMessage>>& OutMessages) const;
+		void AddMessage(TSharedRef<const INiagaraMessage> InMessage);
+		void Reset();
+
+	private:
+		TArray<TSharedRef<const INiagaraMessage>> Messages;
+		TMap<FObjectKey, TArray<TSharedRef<const INiagaraMessage>>> ObjectKeyToMessagesMap;
+	};
+
 	FNiagaraMessageManager();
 
 	FOnRequestRefresh OnRequestRefresh;
@@ -272,6 +314,6 @@ private:
 	static const double MaxJobWorkTime;
 	TArray<FNiagaraMessageJobBatch> MessageJobBatchArr;
 	TArray<TSharedRef<const INiagaraMessage>> GeneratedMessagesForCurrentMessageJobBatch;
-	TMap<FGuid, TMap<const ENiagaraMessageJobType, TArray<TSharedRef<const INiagaraMessage>>>> ObjectToTypeMappedMessagesMap;
+	TMap<FGuid, TMap<const ENiagaraMessageJobType, FMessageData>> ObjectToTypeMappedMessagesMap;
 	static FNiagaraMessageManager* Singleton;
 };
diff --git a/Engine/Plugins/FX/Niagara/Source/NiagaraEditor/Public/ViewModels/NiagaraSystemViewModel.h b/Engine/Plugins/FX/Niagara/Source/NiagaraEditor/Public/ViewModels/NiagaraSystemViewModel.h
--- a/Engine/Plugins/FX/Niagara/Source/NiagaraEditor/Public/ViewModels/NiagaraSystemViewModel.h
+++ b/Engine/Plugins/FX/Niagara/Source/NiagaraEditor/Public/ViewModels/NiagaraSystemViewModel.h
@@ -309,6 +309,8 @@ public:
 	/** Duplicates a set of emitters and refreshes everything.*/
 	void DuplicateEmitters(TArray<FEmitterHandleToDuplicate> EmitterHandlesToDuplicate);
 
+	FGuid GetMessageLogGuid() const;
+
 private:
 
 	/** Sends message jobs to FNiagaraMessageManager for all compile events from the last compile. */
diff --git a/Engine/Plugins/FX/Niagara/Source/NiagaraEditor/Public/ViewModels/Stack/NiagaraStackEntry.h b/Engine/Plugins/FX/Niagara/Source/NiagaraEditor/Public/ViewModels/Stack/NiagaraStackEntry.h
--- a/Engine/Plugins/FX/Niagara/Source/NiagaraEditor/Public/ViewModels/Stack/NiagaraStackEntry.h
+++ b/Engine/Plugins/FX/Niagara/Source/NiagaraEditor/Public/ViewModels/Stack/NiagaraStackEntry.h
@@ -158,11 +158,18 @@ public:
 	};
 
 	// stack issue stuff
-	struct FStackIssueFix
+
+	enum class EStackIssueFixStyle
+	{
+		Fix,
+		Link
+	};
+
+	struct NIAGARAEDITOR_API FStackIssueFix
 	{
 		FStackIssueFix();
 
-		FStackIssueFix(FText InDescription, FStackIssueFixDelegate InFixDelegate);
+		FStackIssueFix(FText InDescription, FStackIssueFixDelegate InFixDelegate, EStackIssueFixStyle FixStyle = EStackIssueFixStyle::Fix);
 
 		bool IsValid() const;
 
@@ -172,11 +179,14 @@ public:
 
 		const FStackIssueFixDelegate& GetFixDelegate() const;
 
+		EStackIssueFixStyle GetStyle() const;
+
 		const FString& GetUniqueIdentifier() const;
 
 	private:
 		FText Description;
 		FStackIssueFixDelegate FixDelegate;
+		EStackIssueFixStyle Style;
 		FString UniqueIdentifier;
 	};
 
diff --git a/Engine/Plugins/FX/Niagara/Source/NiagaraEditor/Public/ViewModels/Stack/NiagaraStackFunctionInput.h b/Engine/Plugins/FX/Niagara/Source/NiagaraEditor/Public/ViewModels/Stack/NiagaraStackFunctionInput.h
--- a/Engine/Plugins/FX/Niagara/Source/NiagaraEditor/Public/ViewModels/Stack/NiagaraStackFunctionInput.h
+++ b/Engine/Plugins/FX/Niagara/Source/NiagaraEditor/Public/ViewModels/Stack/NiagaraStackFunctionInput.h
@@ -330,6 +330,9 @@ private:
 
 	FNiagaraVariable CreateRapidIterationVariable(const FName& InName);
 
+	/** Handles the message manager refreshing messages. */
+	void OnMessageManagerRefresh(const FGuid& MessageJobBatchAssetKey, const TArray<TSharedRef<const INiagaraMessage>> NewMessages);
+
 private:
 	/** The module function call which owns this input entry. NOTE: This input might not be an input to the module function
 		call, it may be an input to a dynamic input function call which is owned by the module. */
@@ -435,4 +438,7 @@ private:
 
 	/** Whether or not the dynamic input for this input has a function script reassignment pending due to a request to fix a missing script. */
 	bool bIsDynamicInputScriptReassignmentPending;
+
+	/** A handled to the refresh delegate on the message manager if it's bound. */
+	FDelegateHandle MessageManagerRefreshHandle;
 };
diff --git a/Engine/Plugins/FX/Niagara/Source/NiagaraEditor/Public/ViewModels/Stack/NiagaraStackGraphUtilities.h b/Engine/Plugins/FX/Niagara/Source/NiagaraEditor/Public/ViewModels/Stack/NiagaraStackGraphUtilities.h
--- a/Engine/Plugins/FX/Niagara/Source/NiagaraEditor/Public/ViewModels/Stack/NiagaraStackGraphUtilities.h
+++ b/Engine/Plugins/FX/Niagara/Source/NiagaraEditor/Public/ViewModels/Stack/NiagaraStackGraphUtilities.h
@@ -5,6 +5,7 @@
 #include "NiagaraTypes.h"
 #include "NiagaraCommon.h"
 #include "ViewModels/Stack/NiagaraParameterHandle.h"
+#include "ViewModels/Stack/NiagaraStackEntry.h"
 #include "AssetData.h"
 
 class UEdGraph;
@@ -21,9 +22,9 @@ class FNiagaraSystemViewModel;
 class UNiagaraEmitter;
 class FNiagaraEmitterViewModel;
 class UNiagaraStackEditorData;
-class UNiagaraStackEntry;
 class UNiagaraStackErrorItem;
 class FCompileConstantResolver;
+class INiagaraMessage;
 
 namespace FNiagaraStackGraphUtilities
 {
@@ -192,4 +193,6 @@ namespace FNiagaraStackGraphUtilities
 
 	void GatherRenamedStackFunctionOutputVariableNames(UNiagaraEmitter* Emitter, UNiagaraNodeFunctionCall& FunctionCallNode, const FString& OldFunctionName, const FString& NewFunctionName, TMap<FName, FName>& OutOldToNewNameMap);
 	void GatherRenamedStackFunctionInputAndOutputVariableNames(UNiagaraEmitter* Emitter, UNiagaraNodeFunctionCall& FunctionCallNode, const FString& OldFunctionName, const FString& NewFunctionName, TMap<FName, FName>& OutOldToNewNameMap);
+
+	UNiagaraStackEntry::FStackIssue MessageManagerMessageToStackIssue(TSharedRef<const INiagaraMessage> InMessage, FString InStackEditorDataKey);
 }
diff --git a/Engine/Plugins/FX/Niagara/Source/NiagaraEditor/Public/ViewModels/Stack/NiagaraStackModuleItem.h b/Engine/Plugins/FX/Niagara/Source/NiagaraEditor/Public/ViewModels/Stack/NiagaraStackModuleItem.h
--- a/Engine/Plugins/FX/Niagara/Source/NiagaraEditor/Public/ViewModels/Stack/NiagaraStackModuleItem.h
+++ b/Engine/Plugins/FX/Niagara/Source/NiagaraEditor/Public/ViewModels/Stack/NiagaraStackModuleItem.h
@@ -14,6 +14,7 @@ class UNiagaraScript;
 class INiagaraStackItemGroupAddUtilities;
 struct FAssetData;
 class UNiagaraClipboardFunctionInput;
+class INiagaraMessage;
 
 UCLASS()
 class NIAGARAEDITOR_API UNiagaraStackModuleItem : public UNiagaraStackItem
@@ -104,6 +105,7 @@ public:
 protected:
 	FOnRequestDeprecationRecommended DeprecationDelegate;
 
+	virtual void FinalizeInternal() override;
 	virtual void RefreshChildrenInternal(const TArray<UNiagaraStackEntry*>& CurrentChildren, TArray<UNiagaraStackEntry*>& NewChildren, TArray<FStackIssue>& NewIssues) override;
 	virtual void SetIsEnabledInternal(bool bInIsEnabled) override;
 
@@ -116,8 +118,8 @@ private:
 	bool FilterLinkedInputCollection(const UNiagaraStackEntry& Child) const;
 	bool FilterLinkedInputCollectionChild(const UNiagaraStackEntry& Child) const;
 	void RefreshIssues(TArray<FStackIssue>& NewIssues);
-	void OnEditorDataChanged();
 	void RefreshIsEnabled();
+	void OnMessageManagerRefresh(const FGuid& MessageJobBatchAssetKey, const TArray<TSharedRef<const INiagaraMessage>> NewMessages);
 
 private:
 	UNiagaraNodeOutput* OutputNode;
diff --git a/Engine/Plugins/FX/Niagara/Source/NiagaraEditorWidgets/Private/Stack/SNiagaraStackErrorItem.cpp b/Engine/Plugins/FX/Niagara/Source/NiagaraEditorWidgets/Private/Stack/SNiagaraStackErrorItem.cpp
--- a/Engine/Plugins/FX/Niagara/Source/NiagaraEditorWidgets/Private/Stack/SNiagaraStackErrorItem.cpp
+++ b/Engine/Plugins/FX/Niagara/Source/NiagaraEditorWidgets/Private/Stack/SNiagaraStackErrorItem.cpp
@@ -10,6 +10,7 @@
 #include "Widgets/Images/SImage.h"
 #include "Widgets/Input/SButton.h"
 #include "Widgets/Text/STextBlock.h"
+#include "Widgets/Input/SHyperlink.h"
 
 #define LOCTEXT_NAMESPACE "SNiagaraStackErrorItem"
 
@@ -42,6 +43,7 @@ void SNiagaraStackErrorItem::Construct(const FArguments& InArgs, UNiagaraStackEr
 		.VAlign(VAlign_Center)
 		.HAlign(HAlign_Left)
 		.AutoWidth()
+		.Padding(0.0f, 0.0f, 4.0f, 0.0f)
 		[
 			SNew(SImage)
 			.Image(IconBrush)
@@ -64,9 +66,11 @@ void SNiagaraStackErrorItemFix::Construct(const FArguments& InArgs, UNiagaraStac
 	ErrorItem = InErrorItem;
 	StackViewModel = InStackViewModel;
 
-	ChildSlot
-	[
-		SNew(SHorizontalBox)
+	TSharedPtr<SWidget> FixWidget;
+
+	if (ErrorItem->GetStackIssueFix().GetStyle() == UNiagaraStackEntry::EStackIssueFixStyle::Fix)
+	{
+		FixWidget = SNew(SHorizontalBox)
 		+ SHorizontalBox::Slot()
 		.Padding(0, 4, 0, 0)
 		.VAlign(VAlign_Center)
@@ -87,9 +91,33 @@ void SNiagaraStackErrorItemFix::Construct(const FArguments& InArgs, UNiagaraStac
 			.TextStyle(FNiagaraEditorStyle::Get(), "NiagaraEditor.ParameterText")
 			.Text_UObject(ErrorItem, &UNiagaraStackErrorItemFix::GetFixButtonText)
 			.OnClicked_UObject(ErrorItem, &UNiagaraStackErrorItemFix::OnTryFixError)
-		]
+		];
+	}
+	else if(ErrorItem->GetStackIssueFix().GetStyle() == UNiagaraStackEntry::EStackIssueFixStyle::Link)
+	{
+		FixWidget = SNew(SBox)
+		.HAlign(HAlign_Left)
+		[
+			SNew(SHyperlink)
+			.Text_UObject(ErrorItem, &UNiagaraStackErrorItemFix::GetDisplayName)
+			.OnNavigate(this, &SNiagaraStackErrorItemFix::LinkNavigate)
+		];
+	}
+	else
+	{
+		FixWidget = SNullWidget::NullWidget;
+	}
+
+	ChildSlot
+	[
+		FixWidget.ToSharedRef()
 	];
 }
 
+void SNiagaraStackErrorItemFix::LinkNavigate()
+{
+	ErrorItem->OnTryFixError();
+}
+
 #undef LOCTEXT_NAMESPACE //"SNiagaraStackErrorItem"
 
diff --git a/Engine/Plugins/FX/Niagara/Source/NiagaraEditorWidgets/Private/Stack/SNiagaraStackErrorItem.h b/Engine/Plugins/FX/Niagara/Source/NiagaraEditorWidgets/Private/Stack/SNiagaraStackErrorItem.h
--- a/Engine/Plugins/FX/Niagara/Source/NiagaraEditorWidgets/Private/Stack/SNiagaraStackErrorItem.h
+++ b/Engine/Plugins/FX/Niagara/Source/NiagaraEditorWidgets/Private/Stack/SNiagaraStackErrorItem.h
@@ -29,6 +29,9 @@ public:
 
 	void Construct(const FArguments& InArgs, UNiagaraStackErrorItemFix* InErrorItem, UNiagaraStackViewModel* InStackViewModel);
 
+private:
+	void LinkNavigate();
+
 private:
 	UNiagaraStackErrorItemFix* ErrorItem;
 };
\ No newline at end of file
