Burst namespace refactor updates. Working through as many UX/workflow issues as I can at once.

EmitterState's static switches weren't showing up properly in the UI as static switches due to stale parameter meta-data.
Fixed long-standing issue where if you used a custom enum to set up a static switch and that enum wasn't in the master list moving forwards, we would not show it and it would list as invalid.
Fixed dragging from a non-namespaced output pin onto a parameter map set's input pin resulting in invalid pin UI.
Ability to add engine vars
Now limited in scope changes to just the ones that are valid to actually do.

diff --git a/Engine/Plugins/FX/Niagara/Source/Niagara/Public/NiagaraTypes.h b/Engine/Plugins/FX/Niagara/Source/Niagara/Public/NiagaraTypes.h
--- a/Engine/Plugins/FX/Niagara/Source/Niagara/Public/NiagaraTypes.h
+++ b/Engine/Plugins/FX/Niagara/Source/Niagara/Public/NiagaraTypes.h
@@ -449,10 +449,10 @@ enum class ENiagaraParameterScope : uint32
 	User,
 
 	/** Parameter provided by the engine. These are explicitly defined by the engine codebase.*/
-	Engine,
+	Engine UMETA(Hidden),
 
 	/** Parameter provided by the engine focused on the owning component. These are explicitly defined by the engine codebase.*/
-	Owner,
+	Owner UMETA(Hidden),
 
 	/** Parameter is an attribute of the owning system payload. It is persistent across frames.*/
 	System,
diff --git a/Engine/Plugins/FX/Niagara/Source/NiagaraEditor/Private/Customizations/NiagaraStaticSwitchNodeDetails.cpp b/Engine/Plugins/FX/Niagara/Source/NiagaraEditor/Private/Customizations/NiagaraStaticSwitchNodeDetails.cpp
--- a/Engine/Plugins/FX/Niagara/Source/NiagaraEditor/Private/Customizations/NiagaraStaticSwitchNodeDetails.cpp
+++ b/Engine/Plugins/FX/Niagara/Source/NiagaraEditor/Private/Customizations/NiagaraStaticSwitchNodeDetails.cpp
@@ -50,6 +50,22 @@ void FNiagaraStaticSwitchNodeDetails::CustomizeDetails(IDetailLayoutBuilder& Det
 	if (ObjectsCustomized.Num() == 1 && ObjectsCustomized[0]->IsA<UNiagaraNodeStaticSwitch>())
 	{
 		Node = CastChecked<UNiagaraNodeStaticSwitch>(ObjectsCustomized[0].Get());
+
+		FNiagaraTypeDefinition TypeDefOnNode = Node->GetInputType();
+
+		// Just in case the actual enum type isn't shown, add to the list.
+		if (TypeDefOnNode.IsEnum())
+		{
+			bool bFoundAlready = false;
+			for (TSharedPtr<SwitchDropdownOption>& Option : DropdownOptions)
+			{
+				if (Option->Enum == TypeDefOnNode.GetEnum())
+					bFoundAlready = true;
+			}
+
+			if (!bFoundAlready)
+				DropdownOptions.Add(MakeShareable(new SwitchDropdownOption(TypeDefOnNode.GetEnum()->GetName(), TypeDefOnNode.GetEnum())));
+		}
 		UpdateSelectionFromNode();
 		
 		IDetailCategoryBuilder& CategoryBuilder = DetailBuilder.EditCategory(SwitchCategoryName);		
diff --git a/Engine/Plugins/FX/Niagara/Source/NiagaraEditor/Private/NiagaraActions.cpp b/Engine/Plugins/FX/Niagara/Source/NiagaraEditor/Private/NiagaraActions.cpp
--- a/Engine/Plugins/FX/Niagara/Source/NiagaraEditor/Private/NiagaraActions.cpp
+++ b/Engine/Plugins/FX/Niagara/Source/NiagaraEditor/Private/NiagaraActions.cpp
@@ -107,10 +107,14 @@ FReply FNiagaraParameterGraphDragOperation::DroppedOnNode(FVector2D ScreenPositi
 		const FNiagaraVariable& Parameter = ParameterAction->GetParameter();
 		if (UNiagaraNodeParameterMapGet* GetMapNode = Cast<UNiagaraNodeParameterMapGet>(GetHoveredNode()))
 		{
+			FScopedTransaction AddNewPinTransaction(LOCTEXT("Drop Onto Get Pin", "Drop parameter onto Get node"));
+			GetMapNode->Modify();
 			GetMapNode->RequestNewTypedPin(EGPD_Output, Parameter.GetType(), Parameter.GetName());
 		} 
 		else if (UNiagaraNodeParameterMapSet* SetMapNode = Cast<UNiagaraNodeParameterMapSet>(GetHoveredNode()))
 		{
+			FScopedTransaction AddNewPinTransaction(LOCTEXT("Drop Onto Set Pin", "Drop parameter onto Set node"));
+			SetMapNode->Modify();
 			SetMapNode->RequestNewTypedPin(EGPD_Input, Parameter.GetType(), Parameter.GetName());
 		}
 	}
@@ -193,6 +197,8 @@ FReply FNiagaraParameterGraphDragOperation::DroppedOnPanel(const TSharedRef<SWid
 
 void FNiagaraParameterGraphDragOperation::MakeGetMap(FNiagaraParameterNodeConstructionParams InParams)
 {
+
+	FScopedTransaction AddNewPinTransaction(LOCTEXT("MakeGetMap", "Make Get Node For Variable"));
 	check(InParams.Graph);
 	InParams.Graph->Modify();
 	FGraphNodeCreator<UNiagaraNodeParameterMapGet> GetNodeCreator(*InParams.Graph);
@@ -205,6 +211,7 @@ void FNiagaraParameterGraphDragOperation::MakeGetMap(FNiagaraParameterNodeConstr
 
 void FNiagaraParameterGraphDragOperation::MakeSetMap(FNiagaraParameterNodeConstructionParams InParams)
 {
+	FScopedTransaction AddNewPinTransaction(LOCTEXT("MakeGetMap", "Make Set Node For Variable"));
 	check(InParams.Graph);
 	InParams.Graph->Modify();
 	FGraphNodeCreator<UNiagaraNodeParameterMapSet> SetNodeCreator(*InParams.Graph);
diff --git a/Engine/Plugins/FX/Niagara/Source/NiagaraEditor/Private/NiagaraEditorUtilities.cpp b/Engine/Plugins/FX/Niagara/Source/NiagaraEditor/Private/NiagaraEditorUtilities.cpp
--- a/Engine/Plugins/FX/Niagara/Source/NiagaraEditor/Private/NiagaraEditorUtilities.cpp
+++ b/Engine/Plugins/FX/Niagara/Source/NiagaraEditor/Private/NiagaraEditorUtilities.cpp
@@ -2183,6 +2183,36 @@ TArray<FName> FNiagaraEditorUtilities::DecomposeVariableNamespace(const FName& I
 	return OutNamespaces;
 }
 
+void  FNiagaraEditorUtilities::RecomposeVariableNamespace(const FName& InVarNameToken, const TArray<FName>& InParentNamespaces, FName& OutName)
+{
+	FString VarNameString;
+	for (FName Name : InParentNamespaces)
+	{
+		VarNameString += Name.ToString() + TEXT(".");
+	}
+	VarNameString += InVarNameToken.ToString();
+	OutName = FName(*VarNameString);
+}
+
+bool FNiagaraEditorUtilities::IsScopeEditable(const FName& InScopeName)
+{
+	if (InScopeName == FNiagaraConstants::EngineNamespace || InScopeName == FNiagaraConstants::EngineOwnerScopeName || InScopeName == FNiagaraConstants::EngineSystemScopeName || InScopeName == FNiagaraConstants::EngineEmitterScopeName)
+	{
+		return false;
+	}
+	return true;
+}
+
+bool FNiagaraEditorUtilities::IsScopeUserAssignable(const FName& InScopeName)
+{
+	if (InScopeName == FNiagaraConstants::EngineNamespace || InScopeName == FNiagaraConstants::EngineOwnerScopeName || InScopeName == FNiagaraConstants::EngineSystemScopeName || InScopeName == FNiagaraConstants::EngineEmitterScopeName)
+	{
+		return false;
+	}
+	return true;
+}
+
+
 void FNiagaraEditorUtilities::GetParameterMetaDataFromName(const FName& InVarNameToken, FNiagaraVariableMetaData& OutMetaData)
 {
 	auto MarkAsLegacyCustomName = [&OutMetaData]() {
@@ -2200,31 +2230,37 @@ void FNiagaraEditorUtilities::GetParameterMetaDataFromName(const FName& InVarNam
 		else if (Namespace == FNiagaraConstants::ModuleNamespace)
 		{
 			OutMetaData.SetScopeName(FNiagaraConstants::InputScopeName);
+			OutMetaData.SetUsage(ENiagaraScriptParameterUsage::Input);
 			return false;
 		}
 		else if (Namespace == FNiagaraConstants::UserNamespace)
 		{
 			OutMetaData.SetScopeName(Namespace);
+			OutMetaData.SetUsage(ENiagaraScriptParameterUsage::Input);
 			return false;
 		}
 		else if (Namespace == FNiagaraConstants::EngineNamespace)
 		{
 			OutMetaData.SetScopeName(Namespace);
+			OutMetaData.SetUsage(ENiagaraScriptParameterUsage::Input);
 			return true;
 		}
 		else if (Namespace == FNiagaraConstants::SystemNamespace)
 		{
 			OutMetaData.SetScopeName(Namespace);
+			OutMetaData.SetUsage(ENiagaraScriptParameterUsage::Input);
 			return true;
 		}
 		else if (Namespace == FNiagaraConstants::EmitterNamespace)
 		{
 			OutMetaData.SetScopeName(Namespace);
+			OutMetaData.SetUsage(ENiagaraScriptParameterUsage::Input);
 			return true;
 		}
 		else if (Namespace == FNiagaraConstants::ParticleAttributeNamespace)
 		{
 			OutMetaData.SetScopeName(Namespace);
+			OutMetaData.SetUsage(ENiagaraScriptParameterUsage::Input);
 			return true;
 		}
 
diff --git a/Engine/Plugins/FX/Niagara/Source/NiagaraEditor/Private/NiagaraGraph.cpp b/Engine/Plugins/FX/Niagara/Source/NiagaraEditor/Private/NiagaraGraph.cpp
--- a/Engine/Plugins/FX/Niagara/Source/NiagaraEditor/Private/NiagaraGraph.cpp
+++ b/Engine/Plugins/FX/Niagara/Source/NiagaraEditor/Private/NiagaraGraph.cpp
@@ -494,6 +494,7 @@ void UNiagaraGraph::ValidateDefaultPins()
 void UNiagaraGraph::PostEditChangeProperty(FPropertyChangedEvent& PropertyChangedEvent)
 {
 	NotifyGraphChanged();
+	RefreshParameterReferences();
 }
 
 class UNiagaraScriptSource* UNiagaraGraph::GetSource() const
@@ -1035,7 +1036,7 @@ UNiagaraScriptVariable* UNiagaraGraph::GetScriptVariable(FName ParameterName) co
 	return nullptr;
 }
 
-void UNiagaraGraph::AddParameter(const FNiagaraVariable& Parameter, bool bIsStaticSwitch /*= false*/)
+UNiagaraScriptVariable* UNiagaraGraph::AddParameter(const FNiagaraVariable& Parameter, bool bIsStaticSwitch /*= false*/)
 {
 	FNiagaraGraphParameterReferenceCollection* FoundParameterReferenceCollection = ParameterToReferencesMap.Find(Parameter);
 	if (!FoundParameterReferenceCollection)
@@ -1048,15 +1049,65 @@ void UNiagaraGraph::AddParameter(const FNiagaraVariable& Parameter, bool bIsStat
 	UNiagaraScriptVariable** FoundScriptVariable = VariableToScriptVariable.Find(Parameter);
 	if (!FoundScriptVariable)
 	{
+		Modify();
 		UNiagaraScriptVariable* NewScriptVariable = NewObject<UNiagaraScriptVariable>(this, FName(), RF_Transactional);
 		NewScriptVariable->Variable = Parameter;
 		NewScriptVariable->Metadata.SetIsStaticSwitch(bIsStaticSwitch);
 		VariableToScriptVariable.Add(Parameter, NewScriptVariable);
+		return NewScriptVariable;
 	}
+	return *FoundScriptVariable;
 }
 
-void UNiagaraGraph::AddParameter(FNiagaraVariable& Parameter, const FAddParameterOptions Options)
+FName UNiagaraGraph::MakeUniqueParameterName(const FName& InName)
 {
+	TArray<TWeakObjectPtr<UNiagaraGraph>> Graphs;
+	Graphs.Emplace(this);
+	return MakeUniqueParameterNameAcrossGraphs(InName, Graphs);
+}
+
+ FName UNiagaraGraph::MakeUniqueParameterNameAcrossGraphs(const FName& InName, TArray<TWeakObjectPtr<UNiagaraGraph>>& InGraphs)
+{
+	 FName PinNameWithoutNamespace;
+	 TArray<FName> NamespacesInVar = FNiagaraEditorUtilities::DecomposeVariableNamespace(InName, PinNameWithoutNamespace);
+
+	 TSet<FName> Names;
+	 for (TWeakObjectPtr<UNiagaraGraph> Graph : InGraphs)
+	 {
+		 if (Graph.IsValid())
+		 {
+			 for (const auto& ParameterElement : Graph->ParameterToReferencesMap)
+			 {
+				 FName GraphParamName = ParameterElement.Key.GetName();
+				 FName PinNameWithoutNamespaceGraph;
+				 FNiagaraEditorUtilities::DecomposeVariableNamespace(GraphParamName, PinNameWithoutNamespaceGraph);
+				 
+				 Names.Add(PinNameWithoutNamespaceGraph);
+			 }
+		 }
+	 }
+	 FName Final;
+	 FName NewNameLeaf = (FNiagaraUtilities::GetUniqueName(PinNameWithoutNamespace, Names));
+	 FNiagaraEditorUtilities::RecomposeVariableNamespace(NewNameLeaf, NamespacesInVar, Final);
+	 return Final;
+}
+
+
+UNiagaraScriptVariable* UNiagaraGraph::AddParameter(FNiagaraVariable& Parameter, const FAddParameterOptions Options)
+{
+	Modify();
+
+	// First we need to determine the new name
+	UNiagaraScriptVariable** FoundScriptVariable = VariableToScriptVariable.Find(Parameter);
+	if (!FoundScriptVariable)
+	{
+		if (Options.bMakeParameterNameUnique)
+		{
+			Parameter.SetName(MakeUniqueParameterName(Parameter.GetName()));
+		}
+	}
+
+	// Now add reference map if needed
 	FNiagaraGraphParameterReferenceCollection* FoundParameterReferenceCollection = ParameterToReferencesMap.Find(Parameter);
 	if (!FoundParameterReferenceCollection)
 	{
@@ -1065,19 +1116,10 @@ void UNiagaraGraph::AddParameter(FNiagaraVariable& Parameter, const FAddParamete
 		ParameterToReferencesMap.Add(Parameter, NewReferenceCollection);
 	}
 
-	UNiagaraScriptVariable** FoundScriptVariable = VariableToScriptVariable.Find(Parameter);
+	// Fill in the rest...
 	if (!FoundScriptVariable)
 	{
-		if (Options.bMakeParameterNameUnique)
-		{
-			TSet<FName> Names;
-			for (const auto& ParameterElement : GetParameterReferenceMap())
-			{
-				Names.Add(ParameterElement.Key.GetName());
-			}
-			Parameter.SetName(FNiagaraUtilities::GetUniqueName(Parameter.GetName(), Names));
-		}
-
+		Modify();
 		UNiagaraScriptVariable* NewScriptVariable = NewObject<UNiagaraScriptVariable>(this, FName(), RF_Transactional);
 		NewScriptVariable->Variable = Parameter;
 		NewScriptVariable->Metadata.SetIsStaticSwitch(Options.bIsStaticSwitch);
@@ -1108,12 +1150,14 @@ void UNiagaraGraph::AddParameter(FNiagaraVariable& Parameter, const FAddParamete
 		{
 			GenerateMetaDataForScriptVariable(NewScriptVariable);
 		}
+		return NewScriptVariable;
 	}
 	else if (Options.bRefreshMetaDataScopeAndUsage)
 	{
 		// The script variable already exists so do not add it, but refresh the existing metadata scope and usage as the option is set.
 		UpdateUsageForScriptVariable(*FoundScriptVariable);
 	}
+	return *FoundScriptVariable;
 }
 
 void UNiagaraGraph::AddParameterReference(const FNiagaraVariable& Parameter, FNiagaraGraphParameterReference& NewParameterReference)
@@ -1125,7 +1169,7 @@ void UNiagaraGraph::AddParameterReference(const FNiagaraVariable& Parameter, FNi
 	}
 }
 
-void UNiagaraGraph::RemoveParameter(const FNiagaraVariable& Parameter)
+void UNiagaraGraph::RemoveParameter(const FNiagaraVariable& Parameter, bool bAllowDeleteStaticSwitch /*= false*/)
 {
 	FNiagaraGraphParameterReferenceCollection* ReferenceCollection = ParameterToReferencesMap.Find(Parameter);
 	if (ReferenceCollection)
@@ -1133,10 +1177,10 @@ void UNiagaraGraph::RemoveParameter(const FNiagaraVariable& Parameter)
 		for (int32 Index = 0; Index < ReferenceCollection->ParameterReferences.Num(); Index++)
 		{
 			const FNiagaraGraphParameterReference& Reference = ReferenceCollection->ParameterReferences[Index];
-			UNiagaraNode* Node = Reference.Value.Get();
+			UNiagaraNode* Node = Cast<UNiagaraNode>(Reference.Value.Get());
 			if (Node && Node->GetGraph() == this)
 			{
-				if (Node->IsA(UNiagaraNodeStaticSwitch::StaticClass()))
+				if (Node->IsA(UNiagaraNodeStaticSwitch::StaticClass()) && bAllowDeleteStaticSwitch == false)
 				{
 					// Static switch parameters are automatically populated from the graph nodes and cannot be manually deleted
 					NotifyGraphChanged();
@@ -1159,7 +1203,7 @@ void UNiagaraGraph::RemoveParameter(const FNiagaraVariable& Parameter)
 	TArray<FNiagaraVariable> VarsToRemove;
 	for (auto It : VariableToScriptVariable)
 	{
-		if (It.Key.GetName() == Parameter.GetName())
+		if (It.Key == Parameter)
 		{
 			VarsToRemove.Add(It.Key);
 		}
@@ -1170,7 +1214,7 @@ void UNiagaraGraph::RemoveParameter(const FNiagaraVariable& Parameter)
 	}
 }
 
-bool UNiagaraGraph::RenameParameter(const FNiagaraVariable& Parameter, FName NewName, bool bFromStaticSwitch, FName NewScopeName)
+bool UNiagaraGraph::RenameParameter(const FNiagaraVariable& Parameter, FName NewName, bool bRenameRequestedFromStaticSwitch, FName NewScopeName)
 {
 	// Block rename when already renaming. This prevents recursion when CommitEditablePinName is called on referenced nodes. 
 	if (bIsRenamingParameter)
@@ -1187,7 +1231,7 @@ bool UNiagaraGraph::RenameParameter(const FNiagaraVariable& Parameter, FName New
 	FNiagaraVariableMetaData OldMetaData;
 	if (OldScriptVariable)
 	{
-		if (!bFromStaticSwitch && (*OldScriptVariable)->Metadata.GetIsStaticSwitch())
+		if (!bRenameRequestedFromStaticSwitch && (*OldScriptVariable)->Metadata.GetIsStaticSwitch())
 		{
 			// We current disallow renaming static switch variables in the Parameters panel. 
 			bIsRenamingParameter = false;
@@ -1206,6 +1250,7 @@ bool UNiagaraGraph::RenameParameter(const FNiagaraVariable& Parameter, FName New
 	// Swap metadata to the new parameter; put the new parameter into VariableToScriptVariable
 	if (OldScriptVariable)
 	{
+		Modify();
 		// Rename all the bindings that point to the old parameter 
 		for (auto It : VariableToScriptVariable)
 		{
@@ -1233,7 +1278,7 @@ bool UNiagaraGraph::RenameParameter(const FNiagaraVariable& Parameter, FName New
 		FNiagaraGraphParameterReferenceCollection NewReferences = *ReferenceCollection;
 		for (FNiagaraGraphParameterReference& Reference : NewReferences.ParameterReferences)
 		{
-			UNiagaraNode* Node = Reference.Value.Get();
+			UNiagaraNode* Node = Cast<UNiagaraNode>(Reference.Value.Get());
 			if (Node && Node->GetGraph() == this)
 			{
 				Node->Modify();
@@ -1931,11 +1976,13 @@ void UNiagaraGraph::SetMetaData(const FNiagaraVariable& InVar, const FNiagaraVar
 		if (*FoundMetaData)
 		{
 			// Replace the old metadata..
+			(*FoundMetaData)->Modify();
 			(*FoundMetaData)->Metadata = InMetaData;
 		} 
 	}
 	else 
 	{
+		Modify();
 		UNiagaraScriptVariable*& NewScriptVariable = VariableToScriptVariable.Add(InVar, NewObject<UNiagaraScriptVariable>(this, FName(), RF_Transactional));
 		NewScriptVariable->Variable = InVar;
 		NewScriptVariable->Metadata = InMetaData;
@@ -1948,13 +1995,16 @@ void UNiagaraGraph::SetPerScriptMetaData(const FNiagaraVariable& InVar, const FN
 	{
 		if (*FoundMetaData)
 		{
+			(*FoundMetaData)->Modify();
 			// Replace the old metadata scope, usage and cached name, but retain persistent metadata such as description.
 			FNiagaraVariableMetaData& OldMetaData = (*FoundMetaData)->Metadata;
 			OldMetaData.CopyPerScriptMetaData(InMetaData);
+			OldMetaData.SetIsStaticSwitch(InMetaData.GetIsStaticSwitch());
 		}
 	}
 	else
 	{
+		Modify();
 		UNiagaraScriptVariable*& NewScriptVariable = VariableToScriptVariable.Add(InVar, NewObject<UNiagaraScriptVariable>(this, FName(), RF_Transactional));
 		NewScriptVariable->Variable = InVar;
 		NewScriptVariable->Metadata = InMetaData;
diff --git a/Engine/Plugins/FX/Niagara/Source/NiagaraEditor/Private/NiagaraNodeParameterMapSet.cpp b/Engine/Plugins/FX/Niagara/Source/NiagaraEditor/Private/NiagaraNodeParameterMapSet.cpp
--- a/Engine/Plugins/FX/Niagara/Source/NiagaraEditor/Private/NiagaraNodeParameterMapSet.cpp
+++ b/Engine/Plugins/FX/Niagara/Source/NiagaraEditor/Private/NiagaraNodeParameterMapSet.cpp
@@ -12,6 +12,7 @@
 #include "Widgets/Layout/SBox.h"
 #include "Widgets/Input/SEditableTextBox.h"
 #include "EdGraph/EdGraphNode.h"
+#include "NiagaraScriptVariable.h"
 
 #include "ScopedTransaction.h"
 
@@ -84,6 +85,16 @@ void UNiagaraNodeParameterMapSet::OnNewTypedPinAdded(UEdGraphPin* NewPin)
 	{
 		TArray<UEdGraphPin*> InputPins;
 		GetInputPins(InputPins);
+		
+		// Determine if this is already namespaced or not. We need to do things differently below if not...
+		FName NewPinName = NewPin->GetFName();
+		bool bCreatedNamespace = false;
+		FName PinNameWithoutNamespace;
+		if (FNiagaraEditorUtilities::DecomposeVariableNamespace(NewPinName, PinNameWithoutNamespace).Num() == 0)
+		{
+			NewPinName = *(PARAM_MAP_LOCAL_MODULE_STR +  NewPinName.ToString());
+			bCreatedNamespace = true;
+		}
 
 		TSet<FName> Names;
 		for (const UEdGraphPin* Pin : InputPins)
@@ -93,9 +104,29 @@ void UNiagaraNodeParameterMapSet::OnNewTypedPinAdded(UEdGraphPin* NewPin)
 				Names.Add(Pin->GetFName());
 			}
 		}
-		const FName NewUniqueName = FNiagaraUtilities::GetUniqueName(*NewPin->GetName(), Names);
+		const FName NewUniqueName = FNiagaraUtilities::GetUniqueName(NewPinName, Names);
+
+		//GetDefault<UEdGraphSchema_Niagara>()->PinToNiagaraVariable()
 		NewPin->PinName = NewUniqueName;
 		NewPin->PinType.PinSubCategory = UNiagaraNodeParameterMapBase::ParameterPinSubCategory;
+		
+		// If dragging from a function or other non-namespaced parent node, we should 
+		// make a local variable to contain the value by default.
+		if (bCreatedNamespace && GetNiagaraGraph())
+		{
+			const UEdGraphSchema_Niagara* Schema = GetDefault<UEdGraphSchema_Niagara>();
+			FNiagaraVariable PinVariable = Schema->PinToNiagaraVariable(NewPin, false);
+			UNiagaraGraph::FAddParameterOptions AddParameterOptions = UNiagaraGraph::FAddParameterOptions();
+
+			//FNiagaraEditorUtilities::GetParameterMetaDataFromName(NewUniqueName, MetaDataGuess);
+			//AddParameterOptions.NewParameterScopeName = FNiagaraConstants::LocalNamespace;
+			//AddParameterOptions.NewParameterUsage = ENiagaraScriptParameterUsage::Local;
+			AddParameterOptions.bMakeParameterNameUnique = true;
+			AddParameterOptions.bRefreshMetaDataScopeAndUsage = true;
+
+			UNiagaraScriptVariable* Var = GetNiagaraGraph()->AddParameter(PinVariable, AddParameterOptions);
+			NewPin->PinName = Var->Variable.GetName();
+		}
 		UpdateAddedPinMetaData(NewPin);
 	}
 
diff --git a/Engine/Plugins/FX/Niagara/Source/NiagaraEditor/Private/NiagaraNodeStaticSwitch.cpp b/Engine/Plugins/FX/Niagara/Source/NiagaraEditor/Private/NiagaraNodeStaticSwitch.cpp
--- a/Engine/Plugins/FX/Niagara/Source/NiagaraEditor/Private/NiagaraNodeStaticSwitch.cpp
+++ b/Engine/Plugins/FX/Niagara/Source/NiagaraEditor/Private/NiagaraNodeStaticSwitch.cpp
@@ -4,6 +4,7 @@
 #include "NiagaraEditorUtilities.h"
 #include "NiagaraHlslTranslator.h"
 #include "NiagaraConstants.h"
+#include "NiagaraScriptVariable.h"
 
 #define LOCTEXT_NAMESPACE "NiagaraNodeStaticSwitch"
 
@@ -105,7 +106,8 @@ void UNiagaraNodeStaticSwitch::RemoveUnusedGraphParameter(const FNiagaraVariable
 	int Index = GraphVariables.Find(OldParameter);
 	if (Index == INDEX_NONE)
 	{
-		GetNiagaraGraph()->RemoveParameter(OldParameter);
+		// Force delete the old static switch parameter.
+		GetNiagaraGraph()->RemoveParameter(OldParameter, true);
 	}
 	else
 	{
@@ -391,6 +393,24 @@ bool UNiagaraNodeStaticSwitch::SubstituteCompiledPin(FHlslNiagaraTranslator* Tra
 	return false;
 }
 
+
+void UNiagaraNodeStaticSwitch::PostLoad()
+{
+	Super::PostLoad();
+
+	// Make sure that we are added to the static switch list.
+	if (GetInputType().IsValid() && InputParameterName.IsValid())
+	{
+		UNiagaraScriptVariable* Var = GetNiagaraGraph()->GetScriptVariable(InputParameterName);
+		if (Var != nullptr && Var->Variable.GetType() == GetInputType() && Var->Metadata.GetIsStaticSwitch() == false)
+		{
+			UE_LOG(LogNiagaraEditor, Log, TEXT("Static switch constant \"%s\" in \"%s\" didn't have static switch meta-data conversion set properly. Fixing now."), *InputParameterName.ToString(), *GetPathName())
+			Var->Metadata.SetIsStaticSwitch(true);
+			MarkNodeRequiresSynchronization(TEXT("Static switch metadata updated"), true);
+		}
+	}
+}
+
 UEdGraphPin* UNiagaraNodeStaticSwitch::GetTracedOutputPin(UEdGraphPin* LocallyOwnedOutputPin) const
 {
 	return GetTracedOutputPin(LocallyOwnedOutputPin, true);
diff --git a/Engine/Plugins/FX/Niagara/Source/NiagaraEditor/Private/NiagaraNodeStaticSwitch.h b/Engine/Plugins/FX/Niagara/Source/NiagaraEditor/Private/NiagaraNodeStaticSwitch.h
--- a/Engine/Plugins/FX/Niagara/Source/NiagaraEditor/Private/NiagaraNodeStaticSwitch.h
+++ b/Engine/Plugins/FX/Niagara/Source/NiagaraEditor/Private/NiagaraNodeStaticSwitch.h
@@ -70,6 +70,7 @@ public:
 	virtual FText GetNodeTitle(ENodeTitleType::Type TitleType) const override;
 	virtual FLinearColor GetNodeTitleColor() const override;
 	virtual bool IsCompilerRelevant() const override { return false; }
+	virtual void PostLoad() override;
 	//~ End EdGraphNode Interface
 
 	//~ Begin UNiagaraNode Interface
diff --git a/Engine/Plugins/FX/Niagara/Source/NiagaraEditor/Private/NiagaraNodeWithDynamicPins.cpp b/Engine/Plugins/FX/Niagara/Source/NiagaraEditor/Private/NiagaraNodeWithDynamicPins.cpp
--- a/Engine/Plugins/FX/Niagara/Source/NiagaraEditor/Private/NiagaraNodeWithDynamicPins.cpp
+++ b/Engine/Plugins/FX/Niagara/Source/NiagaraEditor/Private/NiagaraNodeWithDynamicPins.cpp
@@ -79,6 +79,7 @@ UEdGraphPin* UNiagaraNodeWithDynamicPins::RequestNewTypedPin(EEdGraphPinDirectio
 
 UEdGraphPin* UNiagaraNodeWithDynamicPins::RequestNewTypedPin(EEdGraphPinDirection Direction, const FNiagaraTypeDefinition& Type, const FName InName)
 {
+	Modify();
 	const UEdGraphSchema_Niagara* Schema = GetDefault<UEdGraphSchema_Niagara>();
 	UEdGraphPin* AddPin = GetAddPin(GetAllPins(), Direction);
 	checkf(AddPin != nullptr, TEXT("Add pin is missing"));
@@ -108,10 +109,12 @@ void UNiagaraNodeWithDynamicPins::UpdateAddedPinMetaData(const UEdGraphPin* Adde
 	{
 		const UEdGraphSchema_Niagara* Schema = GetDefault<UEdGraphSchema_Niagara>();
 		FNiagaraVariable PinVariable = Schema->PinToNiagaraVariable(AddedPin, false);
+		
 		if (UNiagaraScriptVariable** ScriptVariable = Graph->GetAllMetaData().Find(PinVariable))
 		{
 			Graph->UpdateUsageForScriptVariable(*ScriptVariable);
 		}
+		
 	}
 }
 
@@ -277,17 +280,17 @@ void UNiagaraNodeWithDynamicPins::AddParameter(FNiagaraVariable Parameter, const
 		// Parameter map type nodes create new parameters when adding pins.
 		FScopedTransaction AddNewPinTransaction(LOCTEXT("AddNewPinTransaction", "Add pin to node"));
 		UNiagaraGraph::FAddParameterOptions AddParameterOptions = UNiagaraGraph::FAddParameterOptions();
-		if (AddPin->Direction == EEdGraphPinDirection::EGPD_Input)
-		{
-			AddParameterOptions.NewParameterUsage = ENiagaraScriptParameterUsage::Output;
-		}
-		else
-		{
-			AddParameterOptions.NewParameterUsage = ENiagaraScriptParameterUsage::Input;
-		}
-		// Default scope of new parameters on a pin to Particles, by convention. //@todo(ng) choose based on usage
-		AddParameterOptions.NewParameterScopeName = FNiagaraConstants::ParticleAttributeNamespace;
-		//AddParameterOptions.bMakeParameterNameUnique //@todo(ng) consider
+		
+		FNiagaraVariableMetaData GuessedMetaData;
+		FNiagaraEditorUtilities::GetParameterMetaDataFromName(Parameter.GetName(), GuessedMetaData);
+
+		AddParameterOptions.NewParameterUsage = GuessedMetaData.GetUsage();
+		AddParameterOptions.NewParameterScopeName = GuessedMetaData.GetScopeName();
+
+		if (FNiagaraConstants::FindEngineConstant(Parameter) == nullptr)
+			AddParameterOptions.bMakeParameterNameUnique = true;
+	
+		AddParameterOptions.bRefreshMetaDataScopeAndUsage = true;
 
 		UNiagaraGraph* Graph = GetNiagaraGraph();
 		checkf(Graph != nullptr, TEXT("Failed to get niagara graph when adding pin!"));
diff --git a/Engine/Plugins/FX/Niagara/Source/NiagaraEditor/Private/ViewModels/NiagaraParameterNameViewModel.cpp b/Engine/Plugins/FX/Niagara/Source/NiagaraEditor/Private/ViewModels/NiagaraParameterNameViewModel.cpp
--- a/Engine/Plugins/FX/Niagara/Source/NiagaraEditor/Private/ViewModels/NiagaraParameterNameViewModel.cpp
+++ b/Engine/Plugins/FX/Niagara/Source/NiagaraEditor/Private/ViewModels/NiagaraParameterNameViewModel.cpp
@@ -14,6 +14,7 @@
 #include "Modules/ModuleManager.h"
 #include "EdGraph/EdGraphPin.h"
 #include "NiagaraEditorCommon.h"
+#include "NiagaraConstants.h"
 
 ///////////////////////////////////////////////////////////////////////////////
 /// Parameter Panel Entry Parameter Name ViewModel							///
@@ -32,19 +33,13 @@ TSharedRef<SWidget> FNiagaraParameterPanelEntryParameterNameViewModel::CreateSco
 	}
 
 	bool bEnableScopeSlotWidget = true;
-	if (CachedScriptVarAndViewInfo.MetaData.GetIsStaticSwitch() == false)
+	if (CachedScriptVarAndViewInfo.MetaData.GetIsStaticSwitch())
+	{
+		bEnableScopeSlotWidget = false;
+	}
+	else if (!FNiagaraEditorUtilities::IsScopeEditable(CachedScriptVarAndViewInfo.MetaData.GetScopeName()))
 	{
-		if (CachedScriptVarAndViewInfo.MetaData.GetUsage() != ENiagaraScriptParameterUsage::Input)
-		{
-			bEnableScopeSlotWidget = false;
-		}
-
-		ENiagaraParameterScope CachedScope;
-		FNiagaraEditorUtilities::GetVariableMetaDataScope(CachedScriptVarAndViewInfo.MetaData, CachedScope);
-		if (CachedScope == ENiagaraParameterScope::Local)
-		{
-			bEnableScopeSlotWidget = false;
-		}
+		bEnableScopeSlotWidget = false;
 	}
 
 	UEnum* ParameterScopeEnum = FNiagaraTypeDefinition::GetParameterScopeEnum();
@@ -158,7 +153,7 @@ TSharedRef<SWidget> FNiagaraGraphPinParameterNameViewModel::CreateScopeSlotWidge
 	TSharedPtr< SWidget > ScopeComboBoxWidget;
 	FNiagaraEditorModule& NiagaraEditorModule = FModuleManager::LoadModuleChecked<FNiagaraEditorModule>("NiagaraEditor");
 
-	bool bEnableScopeSlotWidget = OwningPin->Direction == EEdGraphPinDirection::EGPD_Output;
+	bool bEnableScopeSlotWidget = false;// OwningPin->Direction == EEdGraphPinDirection::EGPD_Output;
 
 	SAssignNew(ScopeComboBoxWidget, SBox)
 		.MinDesiredWidth(80.0f) //@todo(ng) tune and make const static
diff --git a/Engine/Plugins/FX/Niagara/Source/NiagaraEditor/Private/ViewModels/NiagaraParameterPanelViewModel.cpp b/Engine/Plugins/FX/Niagara/Source/NiagaraEditor/Private/ViewModels/NiagaraParameterPanelViewModel.cpp
--- a/Engine/Plugins/FX/Niagara/Source/NiagaraEditor/Private/ViewModels/NiagaraParameterPanelViewModel.cpp
+++ b/Engine/Plugins/FX/Niagara/Source/NiagaraEditor/Private/ViewModels/NiagaraParameterPanelViewModel.cpp
@@ -422,13 +422,15 @@ bool FNiagaraSystemToolkitParameterPanelViewModel::CanRenameParameter(const FNia
 
 FReply FNiagaraSystemToolkitParameterPanelViewModel::HandleActionDragged(const TSharedPtr<FEdGraphSchemaAction>& InAction, const FPointerEvent& MouseEvent) const
 {
-	const FText TooltipFormat = LOCTEXT("Parameters", "Name: {0} \nType: {1}"); //@todo(ng) move to static define
+	const FText TooltipFormat = LOCTEXT("Parameters", "Name: {0} \nType: {1}\nScope: {2}\nUsage: {3}"); //@todo(ng) move to static define
 	const FNiagaraScriptVarAndViewInfoAction* ScriptVarAction = static_cast<FNiagaraScriptVarAndViewInfoAction*>(InAction.Get());
 	const FNiagaraScriptVariableAndViewInfo& ScriptVarAndViewInfo = ScriptVarAction->ScriptVariableAndViewInfo;
 	NiagaraParameterPanelSectionID::Type Section = NiagaraParameterPanelSectionID::GetSectionForParameterMetaData(ScriptVarAndViewInfo.MetaData);
 	const FNiagaraVariable& Var = ScriptVarAndViewInfo.ScriptVariable;
 	const FText Name = FText::FromName(Var.GetName());
-	const FText Tooltip = FText::Format(TooltipFormat, FText::FromName(Var.GetName()), Var.GetType().GetNameText());
+	const FText ScopeText = FText::FromName(ScriptVarAndViewInfo.MetaData.GetScopeName());
+	const FText UsageText = StaticEnum<ENiagaraScriptParameterUsage>()->GetDisplayNameTextByValue((int64)ScriptVarAndViewInfo.MetaData.GetUsage());
+	const FText Tooltip = FText::Format(TooltipFormat, FText::FromName(Var.GetName()), Var.GetType().GetNameText(), ScopeText, UsageText);
 
 	TSharedPtr<FNiagaraParameterAction> ParameterAction = MakeShared<FNiagaraParameterAction>(Var, FText::GetEmpty(), Name, Tooltip, 0, FText(), Section);
 	TSharedRef<FNiagaraParameterDragOperation> DragOperation = MakeShared<FNiagaraParameterDragOperation>(ParameterAction);
@@ -591,25 +593,7 @@ void FNiagaraScriptToolkitParameterPanelViewModel::AddParameter(const FNiagaraVa
 
 	UNiagaraGraph* Graph = ScriptViewModel->GetGraphViewModel()->GetGraph();
 	Graph->Modify();
-	const FVector2D NewNodePos = Graph->GetGoodPlaceForNewNode();
-	ENiagaraScriptParameterUsage InVariableUsage = InVariableMetaDataToAssign.GetUsage();
-	if (InVariableUsage == ENiagaraScriptParameterUsage::Input)
-	{
-		UNiagaraNodeParameterMapGet* NewMapGet = FNiagaraSchemaAction_NewNode::SpawnNodeFromTemplate<UNiagaraNodeParameterMapGet>(Graph, NewObject<UNiagaraNodeParameterMapGet>(), NewNodePos);
-		NewMapGet->Modify();
-		NewMapGet->AddParameter(DuplicateVar, AddParameterOptions);
-	}
-	else if (InVariableUsage == ENiagaraScriptParameterUsage::Output)
-	{
-		UNiagaraNodeParameterMapSet* NewMapSet = FNiagaraSchemaAction_NewNode::SpawnNodeFromTemplate<UNiagaraNodeParameterMapSet>(Graph, NewObject<UNiagaraNodeParameterMapSet>(), NewNodePos);
-		NewMapSet->Modify();
-		NewMapSet->AddParameter(DuplicateVar, AddParameterOptions);
-	}
-	else if (InVariableUsage == ENiagaraScriptParameterUsage::InputOutput)
-	{
-		ensureMsgf(false, TEXT("Unexpected usage encountered when adding parameter through parameter panel view model!"));
-	}
-	else
+	
 	{
 		Graph->AddParameter(DuplicateVar, AddParameterOptions);
 	}
@@ -653,6 +637,12 @@ void FNiagaraScriptToolkitParameterPanelViewModel::RenameParameter(const FNiagar
 
 void FNiagaraScriptToolkitParameterPanelViewModel::ChangeParameterScope(const FNiagaraVariable& TargetVariableToModify, const FNiagaraVariableMetaData& TargetVariableMetaData, const ENiagaraParameterScope NewVariableScope) const
 {
+	if (!FNiagaraEditorUtilities::IsScopeUserAssignable(TargetVariableMetaData.GetScopeName()))
+	{
+		FNiagaraEditorUtilities::WarnWithToastAndLog(FText::Format(LOCTEXT("ScopeNotUserAssignable","The selected scope {0} cannot be assigned by a user"), FText::FromName(TargetVariableMetaData.GetScopeName())));
+		return;
+	}
+
 	if (ensureMsgf(TargetVariableMetaData.GetUsage() != ENiagaraScriptParameterUsage::Output, TEXT("Tried to change scope of output parameter!")))
 	{
 		FScopedTransaction ChangeParameterScopeAndReferencedPins(LOCTEXT("ChangeParameterScopeAndReferencedPins", "Change parameter scope, Rename parameter and referenced pins"));
@@ -777,7 +767,7 @@ TStaticArray<FScopeIsEnabledAndTooltip, (int32)ENiagaraParameterScope::Num> FNia
 					FNiagaraEditorUtilities::GetVariableMetaDataScope(ViewedParameter.MetaData, ViewedParameterScope);
 
 					PerScopeInfo[(int32)ViewedParameterScope].bEnabled = false;
-					PerScopeInfo[(int32)ViewedParameterScope].Tooltip = LOCTEXT("NiagaraInvalidScopeSelectionNameAlias", "Cannot select scope {0}: Parameter with same name already has this scope."); //@todo(ng) get scope
+					PerScopeInfo[(int32)ViewedParameterScope].Tooltip = LOCTEXT("NiagaraScopeSelectionNameAlias", "Cannot select scope: Parameter with same name already has this scope."); //@todo(ng) get scope
 				}
 			}
 		}
@@ -786,14 +776,14 @@ TStaticArray<FScopeIsEnabledAndTooltip, (int32)ENiagaraParameterScope::Num> FNia
 		if (ScriptViewModel->GetStandaloneScript()->GetUsage() != ENiagaraScriptUsage::Module)
 		{
 			PerScopeInfo[(int32)ENiagaraParameterScope::Input].bEnabled = false;
-			PerScopeInfo[(int32)ENiagaraParameterScope::Input].Tooltip = LOCTEXT("NiagaraInvalidScopeSelectionModule", "Cannot select scope {0}: Scope is only valid in Module Scripts."); //@todo(ng) get scope
+			PerScopeInfo[(int32)ENiagaraParameterScope::Input].Tooltip = LOCTEXT("NiagaraInvalidScopeSelectionModule", "Cannot select scope: Scope is only valid in Module Scripts."); //@todo(ng) get scope
 		}
 
 		const TArray<ENiagaraParameterScope> InvalidParameterScopes = ScriptViewModel->GetStandaloneScript()->GetUnsupportedParameterScopes();
 		for (ENiagaraParameterScope InvalidScope : InvalidParameterScopes)
 		{
 			PerScopeInfo[(int32)InvalidScope].bEnabled = false;
-			PerScopeInfo[(int32)InvalidScope].Tooltip = LOCTEXT("NiagaraInvalidScopeSelectionUsageBitmask", "Cannot select scope {0}: Script Usage flags do not support a usage with this scope."); //@todo(ng) rewrite
+			PerScopeInfo[(int32)InvalidScope].Tooltip = LOCTEXT("NiagaraInvalidScopeSelectionUsageBitmask", "Cannot select scope: Script Usage flags do not support a usage with this scope."); //@todo(ng) rewrite
 		}
 	}
 	else
diff --git a/Engine/Plugins/FX/Niagara/Source/NiagaraEditor/Private/Widgets/SNiagaraGraphPinAdd.cpp b/Engine/Plugins/FX/Niagara/Source/NiagaraEditor/Private/Widgets/SNiagaraGraphPinAdd.cpp
--- a/Engine/Plugins/FX/Niagara/Source/NiagaraEditor/Private/Widgets/SNiagaraGraphPinAdd.cpp
+++ b/Engine/Plugins/FX/Niagara/Source/NiagaraEditor/Private/Widgets/SNiagaraGraphPinAdd.cpp
@@ -15,6 +15,7 @@
 #include "SNiagaraParameterPanel.h"
 #include "NiagaraNodeParameterMapBase.h"
 #include "NiagaraNodeCustomHlsl.h"
+#include "NiagaraConstants.h"
 
 #define LOCTEXT_NAMESPACE "NiagaraGraphPinAdd"
 
@@ -75,6 +76,9 @@ TSharedRef<SWidget> SNiagaraGraphPinAdd::OnGetAddButtonMenuContent()
 
 			// Creating a new parameter from a graph pin always defaults to Particles scope.
 			TSharedRef<SNiagaraAddParameterMenu2> MenuWidget = SNew(SNiagaraAddParameterMenu2, Graphs)
+				.NewParameterScope(ENiagaraParameterScope::Local)
+				.NewParameterNamespace(PARAM_MAP_LOCAL_MODULE_STR)
+				.ShowKnownConstantParametersFilter(NiagaraParameterPanelSectionID::Type::REFERENCES)
 				.OnAddParameter_UObject(OwningNode, &UNiagaraNodeWithDynamicPins::AddParameter, (const UEdGraphPin*)GetPinObj())
 				.OnCollectCustomActions_UObject(OwningNode, &UNiagaraNodeWithDynamicPins::CollectAddPinActions, GetPinObj())
 				.OnAllowMakeType_UObject(OwningNode, &UNiagaraNodeWithDynamicPins::AllowNiagaraTypeForAddPin)
@@ -83,7 +87,7 @@ TSharedRef<SWidget> SNiagaraGraphPinAdd::OnGetAddButtonMenuContent()
 			AddButton->SetMenuContentWidgetToFocus(MenuWidget->GetSearchBox());
 			return MenuWidget;
 		}
-		else if (OwningNode->IsA<UNiagaraNodeCustomHlsl>())
+		else 
 		{
 			// Skip binding the add parameter delegate for custom HLSL as it does not interact with graph's UNiagaraScriptVariables.
 
@@ -99,11 +103,6 @@ TSharedRef<SWidget> SNiagaraGraphPinAdd::OnGetAddButtonMenuContent()
 			AddButton->SetMenuContentWidgetToFocus(MenuWidget->GetSearchBox());
 			return MenuWidget;
 		}
-		else
-		{
-			ensureMsgf(false, TEXT("Tried to create a graph pin add menu for unhandled node class!"));
-			return SNullWidget::NullWidget;
-		}
 	}
 	else
 	{
diff --git a/Engine/Plugins/FX/Niagara/Source/NiagaraEditor/Private/Widgets/SNiagaraParameterPanel.cpp b/Engine/Plugins/FX/Niagara/Source/NiagaraEditor/Private/Widgets/SNiagaraParameterPanel.cpp
--- a/Engine/Plugins/FX/Niagara/Source/NiagaraEditor/Private/Widgets/SNiagaraParameterPanel.cpp
+++ b/Engine/Plugins/FX/Niagara/Source/NiagaraEditor/Private/Widgets/SNiagaraParameterPanel.cpp
@@ -14,6 +14,7 @@
 #include "SNiagaraParameterPanelPaletteItem.h"
 #include "SNiagaraGraphActionWidget.h"
 #include "NiagaraScriptVariable.h"
+#include "NiagaraConstants.h"
 
 #define LOCTEXT_NAMESPACE "NiagaraParameterPanel"
 
@@ -95,11 +96,15 @@ void SNiagaraParameterPanel::Tick(const FGeometry& AllottedGeometry, const doubl
 
 void SNiagaraParameterPanel::AddParameter(FNiagaraVariable NewVariable, const NiagaraParameterPanelSectionID::Type SectionID)
 {
-	FNiagaraVariableMetaData NewVariableMetaData = FNiagaraVariableMetaData();
-	NewVariableMetaData.SetScopeName(FNiagaraEditorUtilities::GetScopeNameForParameterScope(NiagaraParameterPanelSectionID::GetScopeForNewParametersInSection(SectionID)));
-	NewVariableMetaData.SetUsage(NiagaraParameterPanelSectionID::GetUsageForNewParametersInSection(SectionID));
+	//FNiagaraVariableMetaData NewVariableMetaData = FNiagaraVariableMetaData();
 
-	ParameterPanelViewModel->AddParameter(NewVariable, NewVariableMetaData);
+	FNiagaraVariableMetaData GuessedMetaData;
+	FNiagaraEditorUtilities::GetParameterMetaDataFromName(NewVariable.GetName(), GuessedMetaData);
+
+	//NewVariableMetaData.SetScopeName(FNiagaraEditorUtilities::GetScopeNameForParameterScope(NiagaraParameterPanelSectionID::GetScopeForNewParametersInSection(SectionID)));
+	//NewVariableMetaData.SetUsage(NiagaraParameterPanelSectionID::GetUsageForNewParametersInSection(SectionID));
+
+	ParameterPanelViewModel->AddParameter(NewVariable, GuessedMetaData);
 
 	GraphActionMenu->RefreshAllActions(true);
 	GraphActionMenu->SelectItemByName(NewVariable.GetName());
@@ -143,6 +148,7 @@ void SNiagaraParameterPanel::CollectAllActions(FGraphActionListBuilderBase& OutA
 
 	TArray<FNiagaraScriptVariableAndViewInfo> VarAndViewInfos = ParameterPanelViewModel->GetViewedParameters();
 	VarAndViewInfos.Sort(SortParameterPanelEntries);
+	//UE_LOG(LogNiagaraEditor, Log, TEXT("Adding Vars--------------------------------------------------"));
 	for (const FNiagaraScriptVariableAndViewInfo& VarAndViewInfo : VarAndViewInfos)
 	{
 		const FText Name = FText::FromName(VarAndViewInfo.ScriptVariable.GetName());
@@ -150,6 +156,9 @@ void SNiagaraParameterPanel::CollectAllActions(FGraphActionListBuilderBase& OutA
 		const NiagaraParameterPanelSectionID::Type Section = ParameterPanelViewModel->GetSectionForVarAndViewInfo(VarAndViewInfo);
 		TSharedPtr<FNiagaraScriptVarAndViewInfoAction> ScriptVarAndViewInfoAction(new FNiagaraScriptVarAndViewInfoAction(VarAndViewInfo, FText::GetEmpty(), Name, Tooltip, 0, FText(), Section /*= 0*/));
 		OutAllActions.AddAction(ScriptVarAndViewInfoAction);
+
+		//UE_LOG(LogNiagaraEditor, Log, TEXT("%s"), *Name.ToString());
+
 	}
 }
  
@@ -212,17 +221,6 @@ TSharedRef<SWidget> SNiagaraParameterPanel::CreateAddToSectionButton(const Niaga
 			.Image(FEditorStyle::GetBrush("Plus"))
 		]
 
-		+ SHorizontalBox::Slot()
-		.VAlign(VAlign_Center)
-		.AutoWidth()
-		.Padding(FMargin(2, 0, 0, 0))
-		[
-			SNew(STextBlock)
-			.Font(IDetailLayoutBuilder::GetDetailFontBold())
-			.Text(AddNewText)
-			//.Visibility(this, &SNiagaraParameterMapView::OnAddButtonTextVisibility, WeakRowWidget, InSection) //@todo(ng) impl
-			.ShadowOffset(FVector2D(1, 1))
-		]
 	];
 
 	AddParameterButtons[(int32)InSection] = Button;
@@ -238,7 +236,9 @@ TSharedRef<SWidget> SNiagaraParameterPanel::OnGetParameterMenu(const NiagaraPara
 		.OnAddParameter(this, &SNiagaraParameterPanel::AddParameter, InSection)
 		.OnAllowMakeType(this, &SNiagaraParameterPanel::AllowMakeType)
 		.ShowGraphParameters(false)
-		.AutoExpandMenu(true)
+		.ShowKnownConstantParametersFilter(InSection)
+		.AutoExpandMenu(InSection == NiagaraParameterPanelSectionID::Type::LOCALS || InSection == NiagaraParameterPanelSectionID::Type::INPUTS ||
+			InSection == NiagaraParameterPanelSectionID::Type::OUTPUTS)
 		.NewParameterScope(NewParameterScopeForSection);
 
 	AddParameterButtons[(int32)InSection]->SetMenuContentWidgetToFocus(MenuWidget->GetSearchBox()->AsShared());
@@ -329,12 +329,21 @@ void SNiagaraAddParameterMenu2::Construct(const FArguments& InArgs, TArray<TWeak
 	this->OnAllowMakeType = InArgs._OnAllowMakeType;
 	this->AllowCreatingNew = InArgs._AllowCreatingNew;
 	this->ShowGraphParameters = InArgs._ShowGraphParameters;
+	this->ShowKnownConstantParametersFilter = InArgs._ShowKnownConstantParametersFilter;
 	this->AutoExpandMenu = InArgs._AutoExpandMenu;
 	this->IsParameterRead = InArgs._IsParameterRead;
 	this->NewParameterScope = InArgs._NewParameterScope;
+	this->NewParameterNamespace = InArgs._NewParameterNamespace;
 
 	Graphs = InGraphs;
-	NewParameterScopeText = FText::FromString(FNiagaraTypeUtilities::GetNamespaceStringForScriptParameterScope(NewParameterScope));
+	if (NewParameterNamespace.IsSet())
+	{
+		NewParameterScopeText = FText::FromString(NewParameterNamespace.Get());
+	}
+	else
+	{
+		NewParameterScopeText = FText::FromString(FNiagaraTypeUtilities::GetNamespaceStringForScriptParameterScope(NewParameterScope));
+	}
 
 	ChildSlot
 	[
@@ -442,7 +451,9 @@ void SNiagaraAddParameterMenu2::CollectAllActions(FGraphActionListBuilderBase& O
 					NewMenuActionInfo.TooltipText = VariableStruct->GetToolTipText(true);
 				}
 
-				const FName NewVariableName = FName(*(NewParameterScopeText.ToString() + RegisteredTypeNameText.ToString()));
+				FName NewVariableName = FName(*(NewParameterScopeText.ToString() + TEXT("New ") + RegisteredTypeNameText.ToString()));
+				NewVariableName = UNiagaraGraph::MakeUniqueParameterNameAcrossGraphs(NewVariableName, Graphs);
+			
 				FNiagaraVariable NewVar = FNiagaraVariable(RegisteredType, NewVariableName);
 				FNiagaraEditorUtilities::ResetVariableToDefaultValue(NewVar);
 				NewMenuActionInfo.NewVariable = NewVar;
@@ -500,6 +511,50 @@ void SNiagaraAddParameterMenu2::CollectAllActions(FGraphActionListBuilderBase& O
 			}
 		}
 	}
+
+	if (ShowKnownConstantParametersFilter.IsSet())
+	{
+		NiagaraParameterPanelSectionID::Type Type = ShowKnownConstantParametersFilter.Get();
+		if (Type != NiagaraParameterPanelSectionID::Type::NONE)
+		{
+			if (Type == NiagaraParameterPanelSectionID::Type::ENGINE ||
+				Type == NiagaraParameterPanelSectionID::Type::OWNER ||
+				Type == NiagaraParameterPanelSectionID::Type::REFERENCES ||
+				Type == NiagaraParameterPanelSectionID::Type::SYSTEM ||
+				Type == NiagaraParameterPanelSectionID::Type::EMITTER)
+			{
+				const TArray<FNiagaraVariable>& EngineConstants = FNiagaraConstants::GetEngineConstants();
+				for (const FNiagaraVariable& Var : EngineConstants)
+				{
+					const FText Category = LOCTEXT("NiagaraAddEngineParameterMenu", "Engine Constant");
+					const FText DisplayName = FText::FromName(Var.GetName());
+					const FText Tooltip = FNiagaraConstants::GetEngineConstantDescription(Var);
+
+
+					FNiagaraVariableMetaData GuessedMetaData;
+					FNiagaraEditorUtilities::GetParameterMetaDataFromName(Var.GetName(), GuessedMetaData);
+					bool bAdd = Type == NiagaraParameterPanelSectionID::Type::REFERENCES; // References has all of these!
+					if (Type == NiagaraParameterPanelSectionID::Type::ENGINE && GuessedMetaData.GetScopeName() == FNiagaraConstants::EngineNamespace)
+						bAdd = true;
+					else if (Type == NiagaraParameterPanelSectionID::Type::OWNER && GuessedMetaData.GetScopeName() == FNiagaraConstants::EngineOwnerScopeName)
+						bAdd = true;
+					else if (Type == NiagaraParameterPanelSectionID::Type::SYSTEM && GuessedMetaData.GetScopeName() == FNiagaraConstants::EngineSystemScopeName)
+						bAdd = true;
+					else if (Type == NiagaraParameterPanelSectionID::Type::EMITTER && GuessedMetaData.GetScopeName() == FNiagaraConstants::EngineEmitterScopeName)
+						bAdd = true;
+
+					if (!bAdd)
+						continue;
+
+					TSharedPtr<FNiagaraMenuAction> Action(new FNiagaraMenuAction(
+						Category, DisplayName, Tooltip, 0, FText::GetEmpty(),
+						FNiagaraMenuAction::FOnExecuteStackAction::CreateSP(this, &SNiagaraAddParameterMenu2::AddParameterSelected, Var)));
+
+					OutAllActions.AddAction(Action);
+				}
+			}
+		}
+	}
 }
 
 void SNiagaraAddParameterMenu2::AddParameterSelected(FNiagaraVariable NewVariable)
diff --git a/Engine/Plugins/FX/Niagara/Source/NiagaraEditor/Private/Widgets/SNiagaraParameterPanel.h b/Engine/Plugins/FX/Niagara/Source/NiagaraEditor/Private/Widgets/SNiagaraParameterPanel.h
--- a/Engine/Plugins/FX/Niagara/Source/NiagaraEditor/Private/Widgets/SNiagaraParameterPanel.h
+++ b/Engine/Plugins/FX/Niagara/Source/NiagaraEditor/Private/Widgets/SNiagaraParameterPanel.h
@@ -137,6 +137,7 @@ public:
 	SLATE_BEGIN_ARGS(SNiagaraAddParameterMenu2)
 		: _AllowCreatingNew(true)
 		, _ShowGraphParameters(true)
+		, _ShowKnownConstantParametersFilter(NiagaraParameterPanelSectionID::Type::NONE)
 		, _AutoExpandMenu(false)
 		, _IsParameterRead(true)
 		, _NewParameterScope(ENiagaraParameterScope::Particles) {}
@@ -145,8 +146,10 @@ public:
 		SLATE_EVENT(FOnAllowMakeType, OnAllowMakeType)
 		SLATE_ATTRIBUTE(bool, AllowCreatingNew)
 		SLATE_ATTRIBUTE(bool, ShowGraphParameters)
+		SLATE_ATTRIBUTE(NiagaraParameterPanelSectionID::Type, ShowKnownConstantParametersFilter)
 		SLATE_ATTRIBUTE(bool, AutoExpandMenu)
 		SLATE_ATTRIBUTE(bool, IsParameterRead)
+		SLATE_ATTRIBUTE(FString, NewParameterNamespace)
 		SLATE_ARGUMENT(ENiagaraParameterScope, NewParameterScope)
 	SLATE_END_ARGS();
 
@@ -171,8 +174,10 @@ private:
 
 	TAttribute<bool> AllowCreatingNew;
 	TAttribute<bool> ShowGraphParameters;
+	TAttribute<NiagaraParameterPanelSectionID::Type> ShowKnownConstantParametersFilter;
 	TAttribute<bool> AutoExpandMenu;
 	TAttribute<bool> IsParameterRead;
+	TAttribute<FString> NewParameterNamespace;
 
 	/** Default scope to give to new parameters created through this menu. Used when generating new variable names so that their namespace is correct. */
 	ENiagaraParameterScope NewParameterScope;
diff --git a/Engine/Plugins/FX/Niagara/Source/NiagaraEditor/Public/NiagaraEditorUtilities.h b/Engine/Plugins/FX/Niagara/Source/NiagaraEditor/Public/NiagaraEditorUtilities.h
--- a/Engine/Plugins/FX/Niagara/Source/NiagaraEditor/Public/NiagaraEditorUtilities.h
+++ b/Engine/Plugins/FX/Niagara/Source/NiagaraEditor/Public/NiagaraEditorUtilities.h
@@ -277,8 +277,13 @@ namespace FNiagaraEditorUtilities
 
 	FName GetScopeNameForParameterScope(ENiagaraParameterScope InScope);
 
+	bool IsScopeEditable(const FName& InScopeName);
+	bool IsScopeUserAssignable(const FName& InScopeName);
+
 	TArray<FName> DecomposeVariableNamespace(const FName& InVarNameToken, FName& OutName);
 
+	void  RecomposeVariableNamespace(const FName& InVarNameToken, const TArray<FName>& InParentNamespaces, FName& OutName);
+
 	void GetParameterMetaDataFromName(const FName& InVarNameToken, FNiagaraVariableMetaData& OutMetaData);
 
 	FString GetNamespacelessVariableNameString(const FName& InVarName);
diff --git a/Engine/Plugins/FX/Niagara/Source/NiagaraEditor/Public/NiagaraGraph.h b/Engine/Plugins/FX/Niagara/Source/NiagaraEditor/Public/NiagaraGraph.h
--- a/Engine/Plugins/FX/Niagara/Source/NiagaraEditor/Public/NiagaraGraph.h
+++ b/Engine/Plugins/FX/Niagara/Source/NiagaraEditor/Public/NiagaraGraph.h
@@ -28,16 +28,40 @@ struct FInputPinsAndOutputPins
 class UNiagaraNode;
 class UNiagaraGraph;
 
-typedef TPair<FGuid, TWeakObjectPtr<UNiagaraNode>> FNiagaraGraphParameterReference;
+USTRUCT()
+struct FNiagaraGraphParameterReference
+{
+	GENERATED_USTRUCT_BODY()
+public:
+	FNiagaraGraphParameterReference() {}
+	FNiagaraGraphParameterReference(const FGuid& InKey, UObject* InValue): Key(InKey), Value(InValue){}
+
+
+	UPROPERTY()
+	FGuid Key;
+
+	UPROPERTY()
+	TWeakObjectPtr<UObject> Value;
+
+	FORCEINLINE bool operator==(const FNiagaraGraphParameterReference& Other)const
+	{
+		return Other.Key == Key && Other.Value == Value;
+	}
+};
 
+USTRUCT()
 struct FNiagaraGraphParameterReferenceCollection
 {
+
+	GENERATED_USTRUCT_BODY()
 public:
 	FNiagaraGraphParameterReferenceCollection(const bool bInCreated = false);
 
 	/** All the references in the graph. */
+	UPROPERTY()
 	TArray<FNiagaraGraphParameterReference> ParameterReferences;
 
+	UPROPERTY()
 	const UNiagaraGraph* Graph;
 
 	/** Returns true if this parameter was initially created by the user. */
@@ -45,6 +69,7 @@ public:
 
 private:
 	/** Whether this parameter was initially created by the user. */
+	UPROPERTY()
 	bool bCreated;
 };
 
@@ -276,15 +301,15 @@ class UNiagaraGraph : public UEdGraph
 	UNiagaraScriptVariable* GetScriptVariable(FName ParameterName) const;
 
 	/** Adds parameter to parameters map setting it as created by the user.*/
-	void AddParameter(const FNiagaraVariable& Parameter, bool bIsStaticSwitch = false);
+	UNiagaraScriptVariable* AddParameter(const FNiagaraVariable& Parameter, bool bIsStaticSwitch = false);
 
-	void AddParameter(FNiagaraVariable& Parameter, const FAddParameterOptions Options);
+	UNiagaraScriptVariable* AddParameter(FNiagaraVariable& Parameter, const FAddParameterOptions Options);
 
 	/** Adds an FNiagaraGraphParameterReference to the ParameterToReferenceMap. */
 	void AddParameterReference(const FNiagaraVariable& Parameter, FNiagaraGraphParameterReference& NewParameterReference);
 
 	/** Remove parameter from map and all the pins associated. */
-	void RemoveParameter(const FNiagaraVariable& Parameter);
+	void RemoveParameter(const FNiagaraVariable& Parameter, bool bAllowDeleteStaticSwitch = false);
 
 	/**
 	* Rename parameter from map and all the pins associated.
@@ -295,7 +320,7 @@ class UNiagaraGraph : public UEdGraph
 	* @param NewScopeName			The new scope name from the FNiagaraVariable's associated FNiagaraVariableMetaData. Used if changing the scope triggered a rename. Applied to the canonical UNiagaraScriptVariable's Metadata.
 	* @return						true if the new name was applied. 
 	*/
-	bool RenameParameter(const FNiagaraVariable& Parameter, FName NewName, bool bFromStaticSwitch = false, FName NewScopeName = FName());
+	bool RenameParameter(const FNiagaraVariable& Parameter, FName NewName, bool bRenameRequestedFromStaticSwitch = false, FName NewScopeName = FName());
 
 	/** 
 	 * Sets the ENiagaraScriptParameterUsage on the Metadata of a ScriptVariable depending on its pin usage in the node graph (used on input/output pins). 
@@ -343,6 +368,12 @@ class UNiagaraGraph : public UEdGraph
 
 	void ScriptVariableChanged(FNiagaraVariable Variable);
 
+	/** Go through all known parameter names in this graph and generate a new unique one.*/
+	FName MakeUniqueParameterName(const FName& InName);
+
+
+	static FName MakeUniqueParameterNameAcrossGraphs(const FName& InName, TArray<TWeakObjectPtr<UNiagaraGraph>>& InGraphs);
+
 protected:
 	void RebuildNumericCache();
 	bool bNeedNumericCacheRebuilt;
@@ -409,6 +440,7 @@ private:
 	mutable TMap<FNiagaraVariable, UNiagaraScriptVariable*> VariableToScriptVariable;
 	
 	/** A map of parameters in the graph to their referencers. */
+	UPROPERTY(Transient)
 	mutable TMap<FNiagaraVariable, FNiagaraGraphParameterReferenceCollection> ParameterToReferencesMap;
 
 	FOnDataInterfaceChanged OnDataInterfaceChangedDelegate;
