Removing invalid error

diff --git a/Engine/Plugins/FX/Niagara/Source/NiagaraEditor/Private/ViewModels/NiagaraParameterPanelViewModel.cpp b/Engine/Plugins/FX/Niagara/Source/NiagaraEditor/Private/ViewModels/NiagaraParameterPanelViewModel.cpp
--- a/Engine/Plugins/FX/Niagara/Source/NiagaraEditor/Private/ViewModels/NiagaraParameterPanelViewModel.cpp
+++ b/Engine/Plugins/FX/Niagara/Source/NiagaraEditor/Private/ViewModels/NiagaraParameterPanelViewModel.cpp
@@ -756,7 +756,7 @@ TStaticArray<FScopeIsEnabledAndTooltip, (int32)ENiagaraParameterScope::Num> FNia
 	FName InName;
 	if (InVarMetaData.GetParameterName(InName))
 	{
-		for (const FNiagaraScriptVariableAndViewInfo& ViewedParameter : CachedViewedParameters)
+		/*for (const FNiagaraScriptVariableAndViewInfo& ViewedParameter : CachedViewedParameters)
 		{
 			FName CachedParameterName;
 			if (ViewedParameter.MetaData.GetParameterName(CachedParameterName))
@@ -770,7 +770,7 @@ TStaticArray<FScopeIsEnabledAndTooltip, (int32)ENiagaraParameterScope::Num> FNia
 					PerScopeInfo[(int32)ViewedParameterScope].Tooltip = LOCTEXT("NiagaraScopeSelectionNameAlias", "Cannot select scope: Parameter with same name already has this scope."); //@todo(ng) get scope
 				}
 			}
-		}
+		}*/
 
 		// Prevent making Module namespace parameters in function and dynamic input scripts
 		if (ScriptViewModel->GetStandaloneScript()->GetUsage() != ENiagaraScriptUsage::Module)
