Fixing missing include for nonunity CIS

diff --git a/Engine/Plugins/FX/Niagara/Source/NiagaraEditor/Private/NiagaraActions.cpp b/Engine/Plugins/FX/Niagara/Source/NiagaraEditor/Private/NiagaraActions.cpp
--- a/Engine/Plugins/FX/Niagara/Source/NiagaraEditor/Private/NiagaraActions.cpp
+++ b/Engine/Plugins/FX/Niagara/Source/NiagaraEditor/Private/NiagaraActions.cpp
@@ -9,6 +9,7 @@
 #include "Framework/Application/MenuStack.h"
 #include "Framework/Application/SlateApplication.h"
 #include "Layout/WidgetPath.h"
+#include "ScopedTransaction.h"
 #include "NiagaraParameterPanelViewModel.h"
 
 #define LOCTEXT_NAMESPACE "NiagaraActions"
