Niagara - Fix emitter dependencies not refreshing correct in a system.

diff --git a/Engine/Plugins/FX/Niagara/Source/NiagaraEditor/Private/ViewModels/NiagaraSystemViewModel.cpp b/Engine/Plugins/FX/Niagara/Source/NiagaraEditor/Private/ViewModels/NiagaraSystemViewModel.cpp
--- a/Engine/Plugins/FX/Niagara/Source/NiagaraEditor/Private/ViewModels/NiagaraSystemViewModel.cpp
+++ b/Engine/Plugins/FX/Niagara/Source/NiagaraEditor/Private/ViewModels/NiagaraSystemViewModel.cpp
@@ -1960,54 +1960,61 @@ void FNiagaraSystemViewModel::UpdateEmitterFixedBounds()
 
 void FNiagaraSystemViewModel::AddSystemEventHandlers()
 {
-	if (GetSystem().IsValid())
+	if (System != nullptr)
 	{
 		TArray<UNiagaraScript*> Scripts;
-		Scripts.Add(GetSystem().GetSystemSpawnScript());
-		Scripts.Add(GetSystem().GetSystemUpdateScript());
+		Scripts.Add(System->GetSystemSpawnScript());
+		Scripts.Add(System->GetSystemUpdateScript());
 		
 		for (UNiagaraScript* Script : Scripts)
 		{
-			FDelegateHandle OnParameterStoreChangedHandle = Script->RapidIterationParameters.AddOnChangedHandler(
-				FNiagaraParameterStore::FOnChanged::FDelegate::CreateRaw<FNiagaraSystemViewModel, const FNiagaraParameterStore&, const UNiagaraScript*>(
-					this, &FNiagaraSystemViewModel::SystemParameterStoreChanged, Script->RapidIterationParameters, Script));
-			ScriptToOnParameterStoreChangedHandleMap.Add(FObjectKey(Script), OnParameterStoreChangedHandle);
+			if (Script != nullptr)
+			{
+				FDelegateHandle OnParameterStoreChangedHandle = Script->RapidIterationParameters.AddOnChangedHandler(
+					FNiagaraParameterStore::FOnChanged::FDelegate::CreateSP<FNiagaraSystemViewModel, const FNiagaraParameterStore&, const UNiagaraScript*>(
+						this->AsShared(), &FNiagaraSystemViewModel::SystemParameterStoreChanged, Script->RapidIterationParameters, Script));
+				ScriptToOnParameterStoreChangedHandleMap.Add(FObjectKey(Script), OnParameterStoreChangedHandle);
+			}
 		}
 
-		UserParameterStoreChangedHandle = GetSystem().GetExposedParameters().AddOnChangedHandler(
-			FNiagaraParameterStore::FOnChanged::FDelegate::CreateRaw<FNiagaraSystemViewModel, const FNiagaraParameterStore&, const UNiagaraScript*>(
-				this, &FNiagaraSystemViewModel::SystemParameterStoreChanged, GetSystem().GetExposedParameters(), nullptr));
+		UserParameterStoreChangedHandle = System->GetExposedParameters().AddOnChangedHandler(
+			FNiagaraParameterStore::FOnChanged::FDelegate::CreateSP<FNiagaraSystemViewModel, const FNiagaraParameterStore&, const UNiagaraScript*>(
+				this->AsShared(), &FNiagaraSystemViewModel::SystemParameterStoreChanged, System->GetExposedParameters(), nullptr));
 
-		SystemScriptGraphChangedHandler = SystemScriptViewModel->GetGraphViewModel()->GetGraph()->AddOnGraphChangedHandler(FOnGraphChanged::FDelegate::CreateRaw(this, &FNiagaraSystemViewModel::SystemScriptGraphChanged));
+		SystemScriptGraphChangedHandle = SystemScriptViewModel->GetGraphViewModel()->GetGraph()->AddOnGraphChangedHandler(FOnGraphChanged::FDelegate::CreateSP(this->AsShared(), &FNiagaraSystemViewModel::SystemScriptGraphChanged));
 	}
 }
 
 void FNiagaraSystemViewModel::RemoveSystemEventHandlers()
 {
-	if (System != nullptr && System->IsValid())
+	if (System != nullptr)
 	{
 		TArray<UNiagaraScript*> Scripts;
-		Scripts.Add(GetSystem().GetSystemSpawnScript());
-		Scripts.Add(GetSystem().GetSystemUpdateScript());
+		Scripts.Add(System->GetSystemSpawnScript());
+		Scripts.Add(System->GetSystemUpdateScript());
 
 		for (UNiagaraScript* Script : Scripts)
 		{
-			FDelegateHandle* OnParameterStoreChangedHandle = ScriptToOnParameterStoreChangedHandleMap.Find(FObjectKey(Script));
-			if (OnParameterStoreChangedHandle != nullptr)
+			if (Script != nullptr)
 			{
-				Script->RapidIterationParameters.RemoveOnChangedHandler(*OnParameterStoreChangedHandle);
+				FDelegateHandle* OnParameterStoreChangedHandle = ScriptToOnParameterStoreChangedHandleMap.Find(FObjectKey(Script));
+				if (OnParameterStoreChangedHandle != nullptr)
+				{
+					Script->RapidIterationParameters.RemoveOnChangedHandler(*OnParameterStoreChangedHandle);
+				}
 			}
 		}
 
-		GetSystem().GetExposedParameters().RemoveOnChangedHandler(UserParameterStoreChangedHandle);
+		System->GetExposedParameters().RemoveOnChangedHandler(UserParameterStoreChangedHandle);
 		if (SystemScriptViewModel.IsValid())
 		{
-			SystemScriptViewModel->GetGraphViewModel()->GetGraph()->RemoveOnGraphChangedHandler(SystemScriptGraphChangedHandler);
+			SystemScriptViewModel->GetGraphViewModel()->GetGraph()->RemoveOnGraphChangedHandler(SystemScriptGraphChangedHandle);
 		}
 	}
 
 	ScriptToOnParameterStoreChangedHandleMap.Empty();
 	UserParameterStoreChangedHandle.Reset();
+	SystemScriptGraphChangedHandle.Reset();
 }
 
 void FNiagaraSystemViewModel::NotifyPinnedCurvesChanged()
diff --git a/Engine/Plugins/FX/Niagara/Source/NiagaraEditor/Public/ViewModels/NiagaraSystemViewModel.h b/Engine/Plugins/FX/Niagara/Source/NiagaraEditor/Public/ViewModels/NiagaraSystemViewModel.h
--- a/Engine/Plugins/FX/Niagara/Source/NiagaraEditor/Public/ViewModels/NiagaraSystemViewModel.h
+++ b/Engine/Plugins/FX/Niagara/Source/NiagaraEditor/Public/ViewModels/NiagaraSystemViewModel.h
@@ -559,7 +559,7 @@ private:
 	TMap<FGuid, TArray<FNiagaraStackModuleData>> GuidToCachedStackModuleData;
 	
 	/** A handle to the on graph changed delegate for the system script. */
-	FDelegateHandle SystemScriptGraphChangedHandler;
+	FDelegateHandle SystemScriptGraphChangedHandle;
 
 	/** An array of emitter handle ids which need their sequencer tracks refreshed next frame. */
 	TArray<FGuid> EmitterIdsRequiringSequencerTrackUpdate;
