Fix Niagara first sim stage update

diff --git a/Engine/Plugins/FX/Niagara/Shaders/Private/NiagaraEmitterInstanceShader.usf b/Engine/Plugins/FX/Niagara/Shaders/Private/NiagaraEmitterInstanceShader.usf
--- a/Engine/Plugins/FX/Niagara/Shaders/Private/NiagaraEmitterInstanceShader.usf
+++ b/Engine/Plugins/FX/Niagara/Shaders/Private/NiagaraEmitterInstanceShader.usf
@@ -682,7 +682,7 @@ int rand_int(int x, int Seed1, int Seed2, int Seed3)
 			// that manipulates the RWInstanceCounts UAV. The generated code will copy the source count
 			// at the end of the shader.
 			BRANCH
-			if (SimulationStageIndex > 0 && !bStageKillsParticles)
+			if (!bStageKillsParticles)
 			{
 				return GDispatchThreadId.x;
 			}
diff --git a/Engine/Plugins/FX/Niagara/Shaders/Private/NiagaraStrandsPhysics.ush b/Engine/Plugins/FX/Niagara/Shaders/Private/NiagaraStrandsPhysics.ush
--- a/Engine/Plugins/FX/Niagara/Shaders/Private/NiagaraStrandsPhysics.ush
+++ b/Engine/Plugins/FX/Niagara/Shaders/Private/NiagaraStrandsPhysics.ush
@@ -510,7 +510,7 @@ void SolveBendTwistMaterial(const float InMaterialCompliance, const float InMate
 	if (!IsMobile)
 	{
 		const float BendingLimit = 2*InBending*InRestLength;
-		for (int i = GGroupThreadId.x+1, end = GGroupThreadId.x+InBlockSize; i < end; ++i)
+		for (uint i = GGroupThreadId.x+1, end = GGroupThreadId.x+InBlockSize; i < end; ++i)
 		{
 			if (i > GGroupThreadId.x+1)
 			{
diff --git a/Engine/Plugins/FX/Niagara/Source/NiagaraEditor/Private/NiagaraHlslTranslator.cpp b/Engine/Plugins/FX/Niagara/Source/NiagaraEditor/Private/NiagaraHlslTranslator.cpp
--- a/Engine/Plugins/FX/Niagara/Source/NiagaraEditor/Private/NiagaraHlslTranslator.cpp
+++ b/Engine/Plugins/FX/Niagara/Source/NiagaraEditor/Private/NiagaraHlslTranslator.cpp
@@ -2467,9 +2467,9 @@ void FHlslNiagaraTranslator::DefineMainGPUFunctions(
 		HlslOutput += TEXT(
 			"	// If a stage doesn't kill particles, StoreUpdateVariables() never calls AcquireIndex(), so the\n"
 			"   // count isn't updated. In that case we must manually copy the original count here.\n"
-			"	if (SimulationStageIndex > 0 && !GStageUsesAlive && GDispatchThreadId.x == 0) \n"
+			"	if (!GStageUsesAlive && GDispatchThreadId.x == 0) \n"
 			"	{ \n"
-			"		RWInstanceCounts[WriteInstanceCountOffset] = RWInstanceCounts[ReadInstanceCountOffset]; \n"
+			"		RWInstanceCounts[WriteInstanceCountOffset] = GSpawnStartInstance + SpawnedInstances; \n"
 			"	} \n"
 		);
 	}
