Niagara - Fix infinite compiles.

diff --git a/Engine/Plugins/FX/Niagara/Source/Niagara/Private/NiagaraCommon.cpp b/Engine/Plugins/FX/Niagara/Source/Niagara/Private/NiagaraCommon.cpp
--- a/Engine/Plugins/FX/Niagara/Source/Niagara/Private/NiagaraCommon.cpp
+++ b/Engine/Plugins/FX/Niagara/Source/Niagara/Private/NiagaraCommon.cpp
@@ -135,7 +135,7 @@ void FNiagaraSystemUpdateContext::AddAll(bool bReInit)
 		UNiagaraComponent* Comp = *It;
 		check(Comp);
 		bool bIsActive = Comp->IsActive() || Comp->IsRegisteredWithScalabilityManager();
-		if (bIsActive && bOnlyActive)
+		if (bIsActive || bOnlyActive == false)
 		{
 			AddInternal(Comp, bReInit);
 		}
@@ -151,7 +151,7 @@ void FNiagaraSystemUpdateContext::Add(const UNiagaraSystem* System, bool bReInit
 		if (Comp->GetAsset() == System)
 		{
 			bool bIsActive = Comp->IsActive() || Comp->IsRegisteredWithScalabilityManager();
-			if (bIsActive && bOnlyActive)
+			if (bIsActive || bOnlyActive == false)
 			{
 				AddInternal(Comp, bReInit);
 			}
@@ -170,7 +170,7 @@ void FNiagaraSystemUpdateContext::Add(const UNiagaraEmitter* Emitter, bool bReIn
 		if (SystemInst && SystemInst->UsesEmitter(Emitter))
 		{
 			bool bIsActive = Comp->IsActive() || Comp->IsRegisteredWithScalabilityManager();
-			if (bIsActive && bOnlyActive)
+			if (bIsActive || bOnlyActive == false)
 			{
 				AddInternal(Comp, bReInit);
 			}
@@ -188,7 +188,7 @@ void FNiagaraSystemUpdateContext::Add(const UNiagaraScript* Script, bool bReInit
 		if (System && System->UsesScript(Script))
 		{
 			bool bIsActive = Comp->IsActive() || Comp->IsRegisteredWithScalabilityManager();
-			if (bIsActive && bOnlyActive)
+			if (bIsActive || bOnlyActive == false)
 			{
 				AddInternal(Comp, bReInit);
 			}
@@ -222,7 +222,7 @@ void FNiagaraSystemUpdateContext::Add(const UNiagaraParameterCollection* Collect
 		if (SystemInst && SystemInst->UsesCollection(Collection))
 		{
 			bool bIsActive = Comp->IsActive() || Comp->IsRegisteredWithScalabilityManager();
-			if (bIsActive && bOnlyActive)
+			if (bIsActive || bOnlyActive == false)
 			{
 				AddInternal(Comp, bReInit);
 			}
diff --git a/Engine/Plugins/FX/Niagara/Source/Niagara/Public/NiagaraCommon.h b/Engine/Plugins/FX/Niagara/Source/Niagara/Public/NiagaraCommon.h
--- a/Engine/Plugins/FX/Niagara/Source/Niagara/Public/NiagaraCommon.h
+++ b/Engine/Plugins/FX/Niagara/Source/Niagara/Public/NiagaraCommon.h
@@ -546,14 +546,14 @@ struct NIAGARA_API FNiagaraSystemUpdateContext
 {
 	GENERATED_BODY()
 
-	FNiagaraSystemUpdateContext(const UNiagaraSystem* System, bool bReInit) :bDestroyOnAdd(false), bOnlyActive(false) { Add(System, bReInit); }
+	FNiagaraSystemUpdateContext(const UNiagaraSystem* System, bool bReInit) : bDestroyOnAdd(false), bOnlyActive(false) { Add(System, bReInit); }
 #if WITH_EDITORONLY_DATA
 	FNiagaraSystemUpdateContext(const UNiagaraEmitter* Emitter, bool bReInit) : bDestroyOnAdd(false), bOnlyActive(false) { Add(Emitter, bReInit); }
 	FNiagaraSystemUpdateContext(const UNiagaraScript* Script, bool bReInit) :bDestroyOnAdd(false), bOnlyActive(false) { Add(Script, bReInit); }
 	//FNiagaraSystemUpdateContext(UNiagaraDataInterface* Interface, bool bReinit) : Add(Interface, bReinit) {}
 	FNiagaraSystemUpdateContext(const UNiagaraParameterCollection* Collection, bool bReInit) :bDestroyOnAdd(false), bOnlyActive(false) { Add(Collection, bReInit); }
 #endif
-	FNiagaraSystemUpdateContext():bDestroyOnAdd(false){ }
+	FNiagaraSystemUpdateContext():bDestroyOnAdd(false), bOnlyActive(false) { }
 
 	~FNiagaraSystemUpdateContext();
 
