Added support for uint32 vertex indices when using Niagara ribbons.

diff --git a/Engine/Plugins/FX/Niagara/Source/Niagara/Private/NiagaraRendererRibbons.cpp b/Engine/Plugins/FX/Niagara/Source/Niagara/Private/NiagaraRendererRibbons.cpp
--- a/Engine/Plugins/FX/Niagara/Source/Niagara/Private/NiagaraRendererRibbons.cpp
+++ b/Engine/Plugins/FX/Niagara/Source/Niagara/Private/NiagaraRendererRibbons.cpp
@@ -8,6 +8,7 @@
 #include "RayTracingDefinitions.h"
 #include "RayTracingDynamicGeometryCollection.h"
 #include "RayTracingInstance.h"
+#include "Math/NumericLimits.h"
 
 DECLARE_CYCLE_STAT(TEXT("Generate Ribbon Vertex Data [GT]"), STAT_NiagaraGenRibbonVertexData, STATGROUP_Niagara);
 DECLARE_CYCLE_STAT(TEXT("Render Ribbons [RT]"), STAT_NiagaraRenderRibbons, STATGROUP_Niagara);
@@ -241,23 +242,26 @@ void FNiagaraRendererRibbons::CreateRenderThreadResources(NiagaraEmitterInstance
 #endif
 }
 
-void FNiagaraRendererRibbons::GenerateIndexBuffer(uint16* OutIndices, uint16& outVertexCount, const TArray<int32>& SegmentData, int32 InterpCount, bool bInvertOrder)
+template <typename TValue>
+void FNiagaraRendererRibbons::GenerateIndexBuffer(FGlobalDynamicIndexBuffer::FAllocationEx& InOutIndexAllocation, const TArray<int32>& SegmentData, int32 InterpCount, bool bInvertOrder)
 {
-	outVertexCount = 0;
-	auto AddTriangleIndices = [&outVertexCount, &OutIndices, InterpCount](int32 SegmentIndex)
+	TValue* OutIndices = (TValue*)InOutIndexAllocation.Buffer;
+	TValue MaxIndex = 0;
+
+	auto AddTriangleIndices = [&MaxIndex, &OutIndices, InterpCount](int32 SegmentIndex)
 	{
 		for (int32 SubSegmentIndex = 0; SubSegmentIndex < InterpCount; ++SubSegmentIndex)
 		{
-			const uint16 BaseVertexIndex = (int16)(SegmentIndex * InterpCount + SubSegmentIndex) * 2;
+			const TValue BaseVertexIndex = (TValue)(SegmentIndex * InterpCount + SubSegmentIndex) * 2;
 			OutIndices[0] = BaseVertexIndex + 0;
 			OutIndices[1] = BaseVertexIndex + 1;
 			OutIndices[2] = BaseVertexIndex + 2;
 			OutIndices[3] = BaseVertexIndex + 1;
 			OutIndices[4] = BaseVertexIndex + 3;
 			OutIndices[5] = BaseVertexIndex + 2;
-
-			outVertexCount = FMath::Max<uint16>(outVertexCount, BaseVertexIndex + 4u);
 			OutIndices += 6;
+
+			MaxIndex = FMath::Max<TValue>(MaxIndex, BaseVertexIndex + 4);
 		}
 	};
 
@@ -276,6 +280,8 @@ void FNiagaraRendererRibbons::GenerateIndexBuffer(uint16* OutIndices, uint16& ou
 			AddTriangleIndices(SegmentData[SegmentDataIndex]);
 		}
 	}
+
+	InOutIndexAllocation.MaxUsedIndex = MaxIndex;
 }
 
 void FNiagaraRendererRibbons::GetDynamicMeshElements(const TArray<const FSceneView*>& Views, const FSceneViewFamily& ViewFamily, uint32 VisibilityMap, FMeshElementCollector& Collector, const FNiagaraSceneProxy *SceneProxy) const
@@ -315,17 +321,12 @@ void FNiagaraRendererRibbons::GetDynamicMeshElements(const TArray<const FSceneVi
 
 			FMeshBatch& MeshBatch = Collector.AllocateMesh();
 
-			FGlobalDynamicIndexBuffer::FAllocation DynamicIndexAllocation;
+			FGlobalDynamicIndexBuffer::FAllocationEx DynamicIndexAllocation;
 			FNiagaraMeshCollectorResourcesRibbon& CollectorResources = Collector.AllocateOneFrameResource<FNiagaraMeshCollectorResourcesRibbon>();
 
-			uint16 VertexCount;
-			uint32 NumPrimitives;
-			CreatePerViewResources(View, ViewFamily, SceneProxy, Collector,
-				VertexCount, NumPrimitives, CollectorResources.UniformBuffer, DynamicIndexAllocation);
+			CreatePerViewResources(View, ViewFamily, SceneProxy, Collector, CollectorResources.UniformBuffer, DynamicIndexAllocation);
 
-			SetupMeshBatchAndCollectorResourceForView(
-				View, ViewFamily, SceneProxy, Collector,
-				DynamicDataRibbon, NumPrimitives, DynamicIndexAllocation, MeshBatch, CollectorResources);
+			SetupMeshBatchAndCollectorResourceForView(View, ViewFamily, SceneProxy, Collector, DynamicDataRibbon, DynamicIndexAllocation, MeshBatch, CollectorResources);
 
 			Collector.AddMesh(ViewIndex, MeshBatch);
 		}
@@ -741,8 +742,7 @@ void FNiagaraRendererRibbons::SetupMeshBatchAndCollectorResourceForView(
 	const FNiagaraSceneProxy* SceneProxy,
 	FMeshElementCollector& Collector,
 	struct FNiagaraDynamicDataRibbon* DynamicDataRibbon,
-	uint32 NumPrimitives,
-	const FGlobalDynamicIndexBuffer::FAllocation& IndexAllocation,
+	const FGlobalDynamicIndexBuffer::FAllocationEx& IndexAllocation,
 	FMeshBatch& MeshBatch,
 	class FNiagaraMeshCollectorResourcesRibbon& CollectorResources) const
 {
@@ -835,7 +835,7 @@ void FNiagaraRendererRibbons::SetupMeshBatchAndCollectorResourceForView(
 	FMeshBatchElement& MeshElement = MeshBatch.Elements[0];
 	MeshElement.IndexBuffer = IndexAllocation.IndexBuffer;
 	MeshElement.FirstIndex = IndexAllocation.FirstIndex;
-	MeshElement.NumPrimitives = NumPrimitives;
+	MeshElement.NumPrimitives = IndexAllocation.NumIndices / 3; // 3 indices per triangle
 	check(MeshElement.NumPrimitives > 0);
 	MeshElement.NumInstances = 1;
 	MeshElement.MinVertexIndex = 0;
@@ -872,8 +872,12 @@ FNiagaraRendererRibbons::FCPUSimParticleDataAllocation FNiagaraRendererRibbons::
 }
 
 void FNiagaraRendererRibbons::CreatePerViewResources(
-	const FSceneView* View, const FSceneViewFamily& ViewFamily, const FNiagaraSceneProxy* SceneProxy, FMeshElementCollector& Collector,
-	uint16& OutVertexCount, uint32& OutNumberOfPrimitives, FNiagaraRibbonUniformBufferRef& OutUniformBuffer, FGlobalDynamicIndexBuffer::FAllocation& InOutIndexAllocation) const
+	const FSceneView* View, 
+	const FSceneViewFamily& ViewFamily, 
+	const FNiagaraSceneProxy* SceneProxy, 
+	FMeshElementCollector& Collector,
+	FNiagaraRibbonUniformBufferRef& OutUniformBuffer, 
+	FGlobalDynamicIndexBuffer::FAllocationEx& InOutIndexAllocation) const
 {
 	FNiagaraDynamicDataRibbon* DynamicDataRibbon = static_cast<FNiagaraDynamicDataRibbon*>(DynamicDataRender);
 	FNiagaraDataBuffer* SourceParticleData = DynamicDataRibbon->GetParticleDataToRender();
@@ -935,9 +939,6 @@ void FNiagaraRendererRibbons::CreatePerViewResources(
 		NumSegments *= SegmentTessellation;
 	}
 
-
-	OutNumberOfPrimitives = NumSegments * 2;
-
 	// Figure out whether start is closer to the view plane than end
 	// TODO : This doesn't work with multi-ribbons.
 	const float StartDist = FVector::DotProduct(View->GetViewDirection(), DynamicDataRibbon->StartPos - ViewOriginForDistanceCulling);
@@ -948,12 +949,19 @@ void FNiagaraRendererRibbons::CreatePerViewResources(
 	// Copy the index data over.
 	FGlobalDynamicIndexBuffer& DynamicIndexBuffer = Collector.GetDynamicIndexBuffer();
 
-	const uint32 IndicesPerPrimitive = 3;
-
-	// the extra plus 3 is for safe padding, the actual size is `NumberOfPrimitives * IndicesPerPrimitive`
-	InOutIndexAllocation = DynamicIndexBuffer.Allocate(OutNumberOfPrimitives * IndicesPerPrimitive, sizeof(uint16));
-
-	GenerateIndexBuffer((uint16*)InOutIndexAllocation.Buffer, OutVertexCount, DynamicDataRibbon->SegmentData, SegmentTessellation, bInvertOrder);
+	// *2 because each position get mapped on each side, *3 for triangles.
+	const uint32 MaxTessellationIndex = SourceParticleData->GetNumInstances() * SegmentTessellation * 2;
+	const uint32 NumIndices = NumSegments * 2 * 3;
+	if (MaxTessellationIndex < MAX_uint16)
+	{
+		InOutIndexAllocation = DynamicIndexBuffer.Allocate<uint16>(NumIndices);
+		GenerateIndexBuffer<uint16>(InOutIndexAllocation, DynamicDataRibbon->SegmentData, SegmentTessellation, bInvertOrder);
+	}
+	else
+	{
+		InOutIndexAllocation = DynamicIndexBuffer.Allocate<uint32>(NumIndices);
+		GenerateIndexBuffer<uint32>(InOutIndexAllocation, DynamicDataRibbon->SegmentData, SegmentTessellation, bInvertOrder);
+	}
 
 	FNiagaraRibbonUniformParameters PerViewUniformParameters;
 	FMemory::Memzero(&PerViewUniformParameters,sizeof(PerViewUniformParameters)); // Clear unset bytes
@@ -1035,25 +1043,19 @@ void FNiagaraRendererRibbons::GetDynamicRayTracingInstances(FRayTracingMaterialG
 		// Setup material for our ray tracing instance
 		FNiagaraMeshCollectorResourcesRibbon& CollectorResources = Context.RayTracingMeshResourceCollector.AllocateOneFrameResource<FNiagaraMeshCollectorResourcesRibbon>();
 
-		FGlobalDynamicIndexBuffer::FAllocation DynamicIndexAllocation;
-		uint16 VertexCount;
-
-		uint32 NumPrimitives;
-		CreatePerViewResources(Context.ReferenceView, Context.ReferenceViewFamily, SceneProxy, Context.RayTracingMeshResourceCollector,
-			VertexCount, NumPrimitives, CollectorResources.UniformBuffer, DynamicIndexAllocation);
+		FGlobalDynamicIndexBuffer::FAllocationEx DynamicIndexAllocation;
+		CreatePerViewResources(Context.ReferenceView, Context.ReferenceViewFamily, SceneProxy, Context.RayTracingMeshResourceCollector, CollectorResources.UniformBuffer, DynamicIndexAllocation);
 
 		RayTracingGeometry.Initializer.IndexBuffer = DynamicIndexAllocation.IndexBuffer->IndexBufferRHI;
-		RayTracingGeometry.Initializer.IndexBufferOffset = DynamicIndexAllocation.FirstIndex * sizeof(uint16);
-
+		RayTracingGeometry.Initializer.IndexBufferOffset = DynamicIndexAllocation.FirstIndex * DynamicIndexAllocation.IndexStride;
 
 		FMeshBatch MeshBatch;
 
-		SetupMeshBatchAndCollectorResourceForView(
-			Context.ReferenceView, Context.ReferenceViewFamily, SceneProxy, Context.RayTracingMeshResourceCollector,
-			DynamicDataRibbon, NumPrimitives, DynamicIndexAllocation, MeshBatch, CollectorResources);
+		SetupMeshBatchAndCollectorResourceForView(Context.ReferenceView, Context.ReferenceViewFamily, SceneProxy, Context.RayTracingMeshResourceCollector, DynamicDataRibbon, DynamicIndexAllocation, MeshBatch, CollectorResources);
 
 		RayTracingInstance.Materials.Add(MeshBatch);
 
+		const uint32 VertexCount = DynamicIndexAllocation.MaxUsedIndex;
 		Context.DynamicRayTracingGeometriesToUpdate.Add(
 			FRayTracingDynamicGeometryUpdateParams
 			{
diff --git a/Engine/Plugins/FX/Niagara/Source/Niagara/Public/NiagaraRendererRibbons.h b/Engine/Plugins/FX/Niagara/Source/Niagara/Public/NiagaraRendererRibbons.h
--- a/Engine/Plugins/FX/Niagara/Source/Niagara/Public/NiagaraRendererRibbons.h
+++ b/Engine/Plugins/FX/Niagara/Source/Niagara/Public/NiagaraRendererRibbons.h
@@ -35,7 +35,8 @@ public:
 
 	FORCEINLINE void AddDynamicParam(TArray<FNiagaraRibbonVertexDynamicParameter>& ParamData, const FVector4& DynamicParam);
 protected:
-	static void GenerateIndexBuffer(uint16* OutIndices, uint16& outVertexCount, const TArray<int32>& SegmentData, int32 InterpCount, bool bInvertOrder);
+	template <typename TValue>
+	static void GenerateIndexBuffer(FGlobalDynamicIndexBuffer::FAllocationEx& InOutIndexAllocation, const TArray<int32>& SegmentData, int32 InterpCount, bool bInvertOrder);
 
 private:
 	struct FCPUSimParticleDataAllocation
@@ -50,14 +51,17 @@ private:
 		const FNiagaraSceneProxy* SceneProxy,
 		FMeshElementCollector& Collector,
 		struct FNiagaraDynamicDataRibbon* DynamicData,
-		uint32 NumPrimitives,
-		const FGlobalDynamicIndexBuffer::FAllocation& IndexAllocation,
+		const FGlobalDynamicIndexBuffer::FAllocationEx& IndexAllocation,
 		FMeshBatch& OutMeshBatch,
 		class FNiagaraMeshCollectorResourcesRibbon& OutCollectorResources) const;
 
 	void CreatePerViewResources(
-		const FSceneView* View, const FSceneViewFamily& ViewFamily, const FNiagaraSceneProxy* SceneProxy, FMeshElementCollector& Collector,
-		uint16& outVertexCount, uint32& OutNumSegments, FNiagaraRibbonUniformBufferRef& OutUniformBuffer, FGlobalDynamicIndexBuffer::FAllocation& InOutIndexAllocation) const;
+		const FSceneView* View,
+		const FSceneViewFamily& ViewFamily,
+		const FNiagaraSceneProxy* SceneProxy,
+		FMeshElementCollector& Collector,
+		FNiagaraRibbonUniformBufferRef& OutUniformBuffer,
+		FGlobalDynamicIndexBuffer::FAllocationEx& InOutIndexAllocation) const;
 
 	FCPUSimParticleDataAllocation AllocateParticleDataIfCPUSim(struct FNiagaraDynamicDataRibbon* DynamicDataRibbon, FGlobalDynamicReadBuffer& DynamicReadBuffer) const;
 
diff --git a/Engine/Source/Runtime/RenderCore/Public/RenderResource.h b/Engine/Source/Runtime/RenderCore/Public/RenderResource.h
--- a/Engine/Source/Runtime/RenderCore/Public/RenderResource.h
+++ b/Engine/Source/Runtime/RenderCore/Public/RenderResource.h
@@ -761,19 +761,11 @@ public:
 	struct FAllocation
 	{
 		/** The location of the buffer in main memory. */
-		uint8* Buffer;
+		uint8* Buffer = nullptr;
 		/** The vertex buffer to bind for draw calls. */
-		FIndexBuffer* IndexBuffer;
+		FIndexBuffer* IndexBuffer = nullptr;
 		/** The offset in to the index buffer. */
-		uint32 FirstIndex;
-
-		/** Default constructor. */
-		FAllocation()
-			: Buffer(NULL)
-			, IndexBuffer(NULL)
-			, FirstIndex(0)
-		{
-		}
+		uint32 FirstIndex = 0;
 
 		/** Returns true if the allocation is valid. */
 		FORCEINLINE bool IsValid() const
@@ -782,6 +774,25 @@ public:
 		}
 	};
 
+	/** Information data with usage details to avoid passing around parameters. */
+	struct FAllocationEx : public FAllocation
+	{
+		FAllocationEx() = default;
+
+		FAllocationEx(const FAllocation& InRef, uint32 InNumIndices, uint32 InIndexStride) 
+			: FAllocation(InRef)
+			, NumIndices(InNumIndices)
+			, IndexStride(InIndexStride) 
+		{}
+
+		/** The number of indices allocated. */
+		uint32 NumIndices = 0;
+		/** The allocation stride (2 or 4 bytes). */
+		uint32 IndexStride = 0;
+		/** The maximum value of the indices used. */
+		uint32 MaxUsedIndex = 0;
+	};
+
 	/** Default constructor. */
 	FGlobalDynamicIndexBuffer();
 
@@ -802,9 +813,9 @@ public:
 	 * @returns an FAllocation with information regarding the allocated memory.
 	 */
 	template <typename IndexType>
-	inline FAllocation Allocate(uint32 NumIndices)
+	FORCEINLINE FAllocationEx Allocate(uint32 NumIndices)
 	{
-		return Allocate(NumIndices, sizeof(IndexType));
+		return FAllocationEx(Allocate(NumIndices, sizeof(IndexType)), NumIndices, sizeof(IndexType));
 	}
 
 	/**
