Niagara - Fix an issue where script assets could be missing from menus due to usage bitfield fixup that was only happening for loaded assets.

diff --git a/Engine/Plugins/FX/Niagara/Source/Niagara/Classes/NiagaraScript.h b/Engine/Plugins/FX/Niagara/Source/Niagara/Classes/NiagaraScript.h
--- a/Engine/Plugins/FX/Niagara/Source/Niagara/Classes/NiagaraScript.h
+++ b/Engine/Plugins/FX/Niagara/Source/Niagara/Classes/NiagaraScript.h
@@ -463,6 +463,8 @@ public:
 
 	UPROPERTY(EditAnywhere, Category = Script, DisplayName = "Script Metadata", meta = (ToolTip = "Script Metadata"))
 	TMap<FName, FString> ScriptMetaData;
+
+	NIAGARA_API static const FName NiagaraCustomVersionTagName;
 #endif
 
 	NIAGARA_API void ComputeVMCompilationId(FNiagaraVMExecutableDataId& Id) const;
diff --git a/Engine/Plugins/FX/Niagara/Source/Niagara/Private/NiagaraScript.cpp b/Engine/Plugins/FX/Niagara/Source/Niagara/Private/NiagaraScript.cpp
--- a/Engine/Plugins/FX/Niagara/Source/Niagara/Private/NiagaraScript.cpp
+++ b/Engine/Plugins/FX/Niagara/Source/Niagara/Private/NiagaraScript.cpp
@@ -246,6 +246,8 @@ void FNiagaraVMExecutableDataId::AppendKeyString(FString& KeyString, const FStri
 
 #endif
 
+const FName UNiagaraScript::NiagaraCustomVersionTagName("NiagaraCustomVersion");
+
 UNiagaraScript::UNiagaraScript(const FObjectInitializer& ObjectInitializer)
 	: Super(ObjectInitializer)
 	, Usage(ENiagaraScriptUsage::Function)
@@ -1588,6 +1590,11 @@ void UNiagaraScript::GetAssetRegistryTags(TArray<FAssetRegistryTag>& OutTags) co
 
 		OutTags.Add(FAssetRegistryTag(HighlightsName, *HighlightsTags, UObject::FAssetRegistryTag::TT_Hidden));
 	}
+
+	// Add the current custom version to the tags so that tags can be fixed up in the future without having to load
+	// the whole asset.
+	const int32 NiagaraVer = GetLinkerCustomVersion(FNiagaraCustomVersion::GUID);
+	OutTags.Add(FAssetRegistryTag(NiagaraCustomVersionTagName, FString::FromInt(NiagaraVer), UObject::FAssetRegistryTag::TT_Hidden));
 #endif
 }
 
diff --git a/Engine/Plugins/FX/Niagara/Source/NiagaraEditor/Private/NiagaraEditorUtilities.cpp b/Engine/Plugins/FX/Niagara/Source/NiagaraEditor/Private/NiagaraEditorUtilities.cpp
--- a/Engine/Plugins/FX/Niagara/Source/NiagaraEditor/Private/NiagaraEditorUtilities.cpp
+++ b/Engine/Plugins/FX/Niagara/Source/NiagaraEditor/Private/NiagaraEditorUtilities.cpp
@@ -1282,6 +1282,10 @@ void FNiagaraEditorUtilities::GetFilteredScriptAssets(FGetFilteredScriptAssetsOp
 
 	for (int i = 0; i < FilteredScriptAssets.Num(); ++i)
 	{
+		// Get the custom version the asset was saved with so it can be used below.
+		int32 NiagaraVersion = INDEX_NONE;
+		FilteredScriptAssets[i].GetTagValue(UNiagaraScript::NiagaraCustomVersionTagName, NiagaraVersion);
+
 		// Check if the script is deprecated
 		if (InFilter.bIncludeDeprecatedScripts == false)
 		{
@@ -1307,11 +1311,26 @@ void FNiagaraEditorUtilities::GetFilteredScriptAssets(FGetFilteredScriptAssetsOp
 		// Check if usage bitmask matches
 		if (InFilter.TargetUsageToMatch.IsSet())
 		{
-			FString BitfieldTagValue;
-			int32 BitfieldValue, TargetBit;
-			BitfieldTagValue = FilteredScriptAssets[i].GetTagValueRef<FString>(GET_MEMBER_NAME_CHECKED(UNiagaraScript, ModuleUsageBitmask));
-			BitfieldValue = FCString::Atoi(*BitfieldTagValue);
-			TargetBit = (BitfieldValue >> (int32)InFilter.TargetUsageToMatch.GetValue()) & 1;
+			int32 BitfieldValue = 0;
+			if (NiagaraVersion == INDEX_NONE || NiagaraVersion < FNiagaraCustomVersion::AddSimulationStageUsageEnum)
+			{
+				// If there is no custom version, or it's less than the simulation stage enum fix up, we need to load the 
+				// asset to get the correct bitmask since the shader stage enum broke the old ones.
+				UNiagaraScript* AssetScript = Cast<UNiagaraScript>(FilteredScriptAssets[i].GetAsset());
+				if (AssetScript != nullptr)
+				{
+					BitfieldValue = AssetScript->ModuleUsageBitmask;
+				}
+			}
+			else
+			{
+				// Otherwise the asset is new enough to have a valid bitmask.
+				FString BitfieldTagValue;
+				BitfieldTagValue = FilteredScriptAssets[i].GetTagValueRef<FString>(GET_MEMBER_NAME_CHECKED(UNiagaraScript, ModuleUsageBitmask));
+				BitfieldValue = FCString::Atoi(*BitfieldTagValue);
+			}
+
+			int32 TargetBit = (BitfieldValue >> (int32)InFilter.TargetUsageToMatch.GetValue()) & 1;
 			if (TargetBit != 1)
 			{
 				continue;
