Niagara - Prevent difficult to reproduce shut down crash.

diff --git a/Engine/Plugins/FX/Niagara/Source/NiagaraEditor/Private/ViewModels/NiagaraScriptViewModel.cpp b/Engine/Plugins/FX/Niagara/Source/NiagaraEditor/Private/ViewModels/NiagaraScriptViewModel.cpp
--- a/Engine/Plugins/FX/Niagara/Source/NiagaraEditor/Private/ViewModels/NiagaraScriptViewModel.cpp
+++ b/Engine/Plugins/FX/Niagara/Source/NiagaraEditor/Private/ViewModels/NiagaraScriptViewModel.cpp
@@ -110,7 +110,10 @@ FNiagaraScriptViewModel::~FNiagaraScriptViewModel()
 		}
 	}
 
-	GEditor->UnregisterForUndo(this);
+	if (GEditor != nullptr)
+	{
+		GEditor->UnregisterForUndo(this);
+	}
 
 	for (TNiagaraViewModelManager<UNiagaraScript, FNiagaraScriptViewModel>::Handle RegisteredHandle : RegisteredHandles)
 	{
