Fix Niagara non-editor compile error

diff --git a/Engine/Plugins/FX/Niagara/Source/Niagara/Private/NiagaraScript.cpp b/Engine/Plugins/FX/Niagara/Source/Niagara/Private/NiagaraScript.cpp
--- a/Engine/Plugins/FX/Niagara/Source/Niagara/Private/NiagaraScript.cpp
+++ b/Engine/Plugins/FX/Niagara/Source/Niagara/Private/NiagaraScript.cpp
@@ -246,7 +246,9 @@ void FNiagaraVMExecutableDataId::AppendKeyString(FString& KeyString, const FStri
 
 #endif
 
+#if WITH_EDITORONLY_DATA
 const FName UNiagaraScript::NiagaraCustomVersionTagName("NiagaraCustomVersion");
+#endif
 
 UNiagaraScript::UNiagaraScript(const FObjectInitializer& ObjectInitializer)
 	: Super(ObjectInitializer)
