Fixing localization duplicate error.

diff --git a/Engine/Plugins/FX/Niagara/Source/NiagaraEditor/Private/NiagaraActions.cpp b/Engine/Plugins/FX/Niagara/Source/NiagaraEditor/Private/NiagaraActions.cpp
--- a/Engine/Plugins/FX/Niagara/Source/NiagaraEditor/Private/NiagaraActions.cpp
+++ b/Engine/Plugins/FX/Niagara/Source/NiagaraEditor/Private/NiagaraActions.cpp
@@ -212,7 +212,7 @@ void FNiagaraParameterGraphDragOperation::MakeGetMap(FNiagaraParameterNodeConstr
 
 void FNiagaraParameterGraphDragOperation::MakeSetMap(FNiagaraParameterNodeConstructionParams InParams)
 {
-	FScopedTransaction AddNewPinTransaction(LOCTEXT("MakeGetMap", "Make Set Node For Variable"));
+	FScopedTransaction AddNewPinTransaction(LOCTEXT("MakeSetMap", "Make Set Node For Variable"));
 	check(InParams.Graph);
 	InParams.Graph->Modify();
 	FGraphNodeCreator<UNiagaraNodeParameterMapSet> SetNodeCreator(*InParams.Graph);
