Cleaning up tooltips

diff --git a/Engine/Plugins/FX/Niagara/Source/NiagaraEditor/Private/ViewModels/Stack/NiagaraStackEventHandlerGroup.cpp b/Engine/Plugins/FX/Niagara/Source/NiagaraEditor/Private/ViewModels/Stack/NiagaraStackEventHandlerGroup.cpp
--- a/Engine/Plugins/FX/Niagara/Source/NiagaraEditor/Private/ViewModels/Stack/NiagaraStackEventHandlerGroup.cpp
+++ b/Engine/Plugins/FX/Niagara/Source/NiagaraEditor/Private/ViewModels/Stack/NiagaraStackEventHandlerGroup.cpp
@@ -70,7 +70,7 @@ private:
 void UNiagaraStackEventHandlerGroup::Initialize(FRequiredEntryData InRequiredEntryData)
 {
 	FText DisplayName = LOCTEXT("EventGroupName", "Add Event Handler");
-	FText ToolTip = LOCTEXT("EventGroupTooltip", "Determines how this Emitter responds to incoming events. There can be more than one event handler script stack per Emitter.");
+	FText ToolTip = LOCTEXT("EventGroupTooltip", "Determines how this Emitter responds to incoming events. There can be more than one event handler stage per Emitter.");
 	AddUtilities = MakeShared<FEventHandlerGroupAddUtilities>(InRequiredEntryData.EmitterViewModel.ToSharedRef(), 
 		TNiagaraStackItemGroupAddUtilities<FNiagaraEventScriptProperties>::FOnItemAdded::CreateUObject(this, &UNiagaraStackEventHandlerGroup::ItemAddedFromUtilties));
 	Super::Initialize(InRequiredEntryData, DisplayName, ToolTip, AddUtilities.Get());
diff --git a/Engine/Plugins/FX/Niagara/Source/NiagaraEditor/Private/ViewModels/Stack/NiagaraStackEventScriptItemGroup.cpp b/Engine/Plugins/FX/Niagara/Source/NiagaraEditor/Private/ViewModels/Stack/NiagaraStackEventScriptItemGroup.cpp
--- a/Engine/Plugins/FX/Niagara/Source/NiagaraEditor/Private/ViewModels/Stack/NiagaraStackEventScriptItemGroup.cpp
+++ b/Engine/Plugins/FX/Niagara/Source/NiagaraEditor/Private/ViewModels/Stack/NiagaraStackEventScriptItemGroup.cpp
@@ -202,7 +202,7 @@ void UNiagaraStackEventScriptItemGroup::Initialize(
 	ENiagaraScriptUsage InScriptUsage,
 	FGuid InScriptUsageId)
 {
-	FText ToolTip = LOCTEXT("EventGroupTooltip", "Determines how this Emitter responds to incoming events. There can be more than one event handler script stack per Emitter.");
+	FText ToolTip = LOCTEXT("EventGroupTooltip", "Determines how this Emitter responds to incoming events. There can be more than one event handler stage per Emitter.");
 	FText TempDisplayName = FText::Format(LOCTEXT("TempDisplayNameFormat", "Event Handler - {0}"), FText::FromString(InScriptUsageId.ToString(EGuidFormats::DigitsWithHyphens)));
 	Super::Initialize(InRequiredEntryData, TempDisplayName, ToolTip, InScriptViewModel, InScriptUsage, InScriptUsageId);
 }
diff --git a/Engine/Plugins/FX/Niagara/Source/NiagaraEditor/Private/ViewModels/Stack/NiagaraStackRoot.cpp b/Engine/Plugins/FX/Niagara/Source/NiagaraEditor/Private/ViewModels/Stack/NiagaraStackRoot.cpp
--- a/Engine/Plugins/FX/Niagara/Source/NiagaraEditor/Private/ViewModels/Stack/NiagaraStackRoot.cpp
+++ b/Engine/Plugins/FX/Niagara/Source/NiagaraEditor/Private/ViewModels/Stack/NiagaraStackRoot.cpp
@@ -77,7 +77,7 @@ void UNiagaraStackRoot::RefreshChildrenInternal(const TArray<UNiagaraStackEntry*
 			FExecutionCategoryNames::System, FExecutionSubcategoryNames::Spawn,
 			GetSystemViewModel()->GetEditorData().GetStackEditorData());
 		FText DisplayName = LOCTEXT("SystemSpawnGroupName", "System Spawn");
-		FText ToolTip = LOCTEXT("SystemSpawnGroupToolTip", "Occurs once at System creation on the CPU. Modules in this section should initialize defaults and/or do initial setup.\r\nModules are executed in order from top to bottom of the stack.");
+		FText ToolTip = LOCTEXT("SystemSpawnGroupToolTip", "Occurs once at System creation on the CPU. Modules in this stage should initialize defaults and/or do initial setup.\r\nModules are executed in order from top to bottom of the stack.");
 		SystemSpawnGroup->Initialize(RequiredEntryData, DisplayName, ToolTip, GetSystemViewModel()->GetSystemScriptViewModel().ToSharedRef(), ENiagaraScriptUsage::SystemSpawnScript);
 	}
 
@@ -88,7 +88,7 @@ void UNiagaraStackRoot::RefreshChildrenInternal(const TArray<UNiagaraStackEntry*
 			FExecutionCategoryNames::System, FExecutionSubcategoryNames::Update,
 			GetSystemViewModel()->GetEditorData().GetStackEditorData());
 		FText DisplayName = LOCTEXT("SystemUpdateGroupName", "System Update");
-		FText ToolTip = LOCTEXT("SystemUpdateGroupToolTip", "Occurs every Emitter tick on the CPU.Modules in this section should compute values for parameters for emitter or particle update or spawning this frame.\r\nModules are executed in order from top to bottom of the stack.");
+		FText ToolTip = LOCTEXT("SystemUpdateGroupToolTip", "Occurs every Emitter tick on the CPU.Modules in this stage should compute values for parameters for emitter or particle update or spawning this frame.\r\nModules are executed in order from top to bottom of the stack.");
 		SystemUpdateGroup->Initialize(RequiredEntryData, DisplayName, ToolTip, GetSystemViewModel()->GetSystemScriptViewModel().ToSharedRef(), ENiagaraScriptUsage::SystemUpdateScript);
 	}
 
@@ -110,7 +110,7 @@ void UNiagaraStackRoot::RefreshChildrenInternal(const TArray<UNiagaraStackEntry*
 			FExecutionCategoryNames::Emitter, FExecutionSubcategoryNames::Spawn,
 			GetEmitterViewModel()->GetOrCreateEditorData().GetStackEditorData());
 		FText DisplayName = LOCTEXT("EmitterSpawnGroupName", "Emitter Spawn");
-		FText ToolTip = LOCTEXT("EmitterSpawnGroupTooltip", "Occurs once at Emitter creation on the CPU. Modules in this section should initialize defaults and/or do initial setup.\r\nModules are executed in order from top to bottom of the stack.");
+		FText ToolTip = LOCTEXT("EmitterSpawnGroupTooltip", "Occurs once at Emitter creation on the CPU. Modules in this stage should initialize defaults and/or do initial setup.\r\nModules are executed in order from top to bottom of the stack.");
 		EmitterSpawnGroup->Initialize(RequiredEntryData, DisplayName, ToolTip, GetEmitterViewModel()->GetSharedScriptViewModel(), ENiagaraScriptUsage::EmitterSpawnScript);
 	}
 
@@ -121,7 +121,7 @@ void UNiagaraStackRoot::RefreshChildrenInternal(const TArray<UNiagaraStackEntry*
 			FExecutionCategoryNames::Emitter, FExecutionSubcategoryNames::Update,
 			GetEmitterViewModel()->GetOrCreateEditorData().GetStackEditorData());
 		FText DisplayName = LOCTEXT("EmitterUpdateGroupName", "Emitter Update");
-		FText ToolTip = LOCTEXT("EmitterUpdateGroupTooltip", "Occurs every Emitter tick on the CPU. Modules in this section should compute values for parameters for Particle Update or Spawning this frame.\r\nModules are executed in order from top to bottom of the stack.");
+		FText ToolTip = LOCTEXT("EmitterUpdateGroupTooltip", "Occurs every Emitter tick on the CPU. Modules in this stage should compute values for parameters for Particle Update or Spawning this frame.\r\nModules are executed in order from top to bottom of the stack.");
 		EmitterUpdateGroup->Initialize(RequiredEntryData, DisplayName, ToolTip, GetEmitterViewModel()->GetSharedScriptViewModel(), ENiagaraScriptUsage::EmitterUpdateScript);
 	}
 
@@ -132,7 +132,7 @@ void UNiagaraStackRoot::RefreshChildrenInternal(const TArray<UNiagaraStackEntry*
 			FExecutionCategoryNames::Particle, FExecutionSubcategoryNames::Spawn,
 			GetEmitterViewModel()->GetOrCreateEditorData().GetStackEditorData());
 		FText DisplayName = LOCTEXT("ParticleSpawnGroupName", "Particle Spawn");
-		FText ToolTip = LOCTEXT("ParticleSpawnGroupTooltip", "Called once per created particle. Modules in this section should set up initial values for each particle.\r\nIf \"Use Interpolated Spawning\" is set, we will also run the Particle Update script after the Particle Spawn script.\r\nModules are executed in order from top to bottom of the stack.");
+		FText ToolTip = LOCTEXT("ParticleSpawnGroupTooltip", "Called once per created particle. Modules in this stage should set up initial values for each particle.\r\nIf \"Use Interpolated Spawning\" is set, we will also run the Particle Update stage after the Particle Spawn stage.\r\nModules are executed in order from top to bottom of the stack.");
 		ParticleSpawnGroup->Initialize(RequiredEntryData, DisplayName, ToolTip, GetEmitterViewModel()->GetSharedScriptViewModel(), ENiagaraScriptUsage::ParticleSpawnScript);
 	}
 
@@ -143,7 +143,7 @@ void UNiagaraStackRoot::RefreshChildrenInternal(const TArray<UNiagaraStackEntry*
 			FExecutionCategoryNames::Particle, FExecutionSubcategoryNames::Update,
 			GetEmitterViewModel()->GetOrCreateEditorData().GetStackEditorData());
 		FText DisplayName = LOCTEXT("ParticleUpdateGroupName", "Particle Update");
-		FText ToolTip = LOCTEXT("ParticleUpdateGroupTooltip", "Called every frame per particle. Modules in this section should update new values for this frame.\r\nModules are executed in order from top to bottom of the stack.");
+		FText ToolTip = LOCTEXT("ParticleUpdateGroupTooltip", "Called every frame per particle. Modules in this stage should update new values for this frame.\r\nModules are executed in order from top to bottom of the stack.");
 		ParticleUpdateGroup->Initialize(RequiredEntryData, DisplayName, ToolTip, GetEmitterViewModel()->GetSharedScriptViewModel(), ENiagaraScriptUsage::ParticleUpdateScript);
 	}
 
