Adding output scope

diff --git a/Engine/Plugins/FX/Niagara/Source/Niagara/Classes/NiagaraConstants.h b/Engine/Plugins/FX/Niagara/Source/Niagara/Classes/NiagaraConstants.h
--- a/Engine/Plugins/FX/Niagara/Source/Niagara/Classes/NiagaraConstants.h
+++ b/Engine/Plugins/FX/Niagara/Source/Niagara/Classes/NiagaraConstants.h
@@ -10,6 +10,7 @@
 #define PARAM_MAP_ENGINE_SYSTEM_STR TEXT("Engine.System.")
 #define PARAM_MAP_ENGINE_EMITTER_STR TEXT("Engine.Emitter.")
 #define PARAM_MAP_LOCAL_MODULE_STR TEXT("Local.Module.")
+#define PARAM_MAP_OUTPUT_MODULE_STR TEXT("Output.Module.")
 #define PARAM_MAP_USER_STR TEXT("User.")
 #define PARAM_MAP_SYSTEM_STR TEXT("System.")
 #define PARAM_MAP_EMITTER_STR TEXT("Emitter.")
diff --git a/Engine/Plugins/FX/Niagara/Source/Niagara/Private/NiagaraTypes.cpp b/Engine/Plugins/FX/Niagara/Source/Niagara/Private/NiagaraTypes.cpp
--- a/Engine/Plugins/FX/Niagara/Source/Niagara/Private/NiagaraTypes.cpp
+++ b/Engine/Plugins/FX/Niagara/Source/Niagara/Private/NiagaraTypes.cpp
@@ -33,6 +33,9 @@ FString FNiagaraTypeUtilities::GetNamespaceStringForScriptParameterScope(const E
 	case ENiagaraParameterScope::Input:
 		return PARAM_MAP_MODULE_STR;
 		break;
+	case ENiagaraParameterScope::Output:
+		return PARAM_MAP_OUTPUT_MODULE_STR;
+		break;
 	case ENiagaraParameterScope::Local:
 		return PARAM_MAP_LOCAL_STR;
 		break;
diff --git a/Engine/Plugins/FX/Niagara/Source/Niagara/Public/NiagaraTypes.h b/Engine/Plugins/FX/Niagara/Source/Niagara/Public/NiagaraTypes.h
--- a/Engine/Plugins/FX/Niagara/Source/Niagara/Public/NiagaraTypes.h
+++ b/Engine/Plugins/FX/Niagara/Source/Niagara/Public/NiagaraTypes.h
@@ -443,37 +443,42 @@ UENUM()
 enum class ENiagaraParameterScope : uint32
 {
 	/** Parameter that is an input argument into this graph.*/
-	Input,
-
-	/** Parameter that is exposed to the owning component for editing.*/
-	User,
+	Input UMETA(DisplayName = "Input"),
 
-	/** Parameter provided by the engine. These are explicitly defined by the engine codebase.*/
-	Engine UMETA(Hidden),
+	/** Parameter that is exposed to the owning component for editing and are read-only when used in the graph*/
+	User UMETA(DisplayName = "User"),
 
-	/** Parameter provided by the engine focused on the owning component. These are explicitly defined by the engine codebase.*/
-	Owner UMETA(Hidden),
+	/** Parameter provided by the engine. These are explicitly defined by the engine codebase and read-only. */
+	Engine UMETA(DisplayName = "Engine (Generic)", Hidden),
 
-	/** Parameter is an attribute of the owning system payload. It is persistent across frames.*/
-	System,
+	/** Parameter provided by the engine focused on the owning component. These are explicitly defined by the engine codebase and read-only.*/
+	Owner UMETA(DisplayName = "Engine (Owner)", Hidden),
 
-	/** Parameter is an attribute of the owning emitter payload. It is persistent across frames.*/
-	Emitter,
+	/** Parameter is an attribute of the owning system payload. It is persistent across frames and initialized in the System Spawn stage of the stack.*/
+	System  UMETA(DisplayName = "System"),
 
-	/** Parameter is an attribute of the owning particle payload. It is persistent across frames.*/
-	Particles,
+	/** Parameter is an attribute of the owning emitter payload. It is persistent across frames and initialized in the Emitter Spawn stage of the stack.*/
+	Emitter   UMETA(DisplayName = "Emitter"),
 
-	/** Parameter is initialized in the appropriate spawn section for the stack. It is persistent from frame to frame. For example, if used consistently in an Emitter stack, this parameter will turn into an emitter attribute. Similarly, if used in a Particle stack, it will turn into a particle attribute.*/
-	ScriptPersistent UMETA(Hidden), //@todo(ng) hiding until autotest verification is made.
+	/** Parameter is an attribute of the owning particle payload. It is persistent across frames and initialized in the Particle Spawn stage of the stack.*/
+	Particles  UMETA(DisplayName = "Particles"),
 
-	/** Parameter is initialized at the start of this script phase and can be shared amongst other modules, but is not persistent across script frames or from stack section to stack section.*/
-	ScriptTransient,
+	/** Parameter is initialized in the appropriate spawn stage for the stack. It is persistent from frame to frame. For example, if used consistently in an Emitter stage, this parameter will turn into an emitter attribute. Similarly, if used in a Particle stage, it will turn into a particle attribute.*/
+	ScriptPersistent UMETA(DisplayName = "Stage (Persistent)", Hidden), //@todo(ng) hiding until autotest verification is made.
 
-	Local UMETA(Hidden), //Convenience markup for ScopeToString functions, only use in conjunction with ENiagaraScriptParameterUsage::Local.
+	/** Parameter is initialized at the start of this stage and can be shared amongst other modules within this stack stage, but is not persistent across runs or from stack stage to stack stage.*/
+	ScriptTransient UMETA(DisplayName = "Stage (Transient)"),
+	
+	/** Parameter is initialized at the start of this script and is only used within the context of this script. It is invisible to the parent stage stack.*/
+	Local UMETA(DisplayName = "Local"), //Convenience markup for ScopeToString functions, only use in conjunction with ENiagaraScriptParameterUsage::Local.
 
 	Custom UMETA(Hidden), //Convenience markup for expressing parameters using legacy editor mode to freetype namespace and name.
 
 	DISPLAY_ONLY_StaticSwitch UMETA(DisplayName="Static Switch", Hidden), //Only use for display string in SEnumComboBoxes; does not have implementation for classes that interact with ENiagaraParameterScope.
+	
+	/** Parameter is output to the owning stack stage from this script, but is only meaningful if bound elsewhere in the stage.*/
+	Output UMETA(DisplayName = "Output"),
+
 	// insert new scopes before
 	None UMETA(Hidden),
 
diff --git a/Engine/Plugins/FX/Niagara/Source/NiagaraEditor/Private/NiagaraEditorModule.cpp b/Engine/Plugins/FX/Niagara/Source/NiagaraEditor/Private/NiagaraEditorModule.cpp
--- a/Engine/Plugins/FX/Niagara/Source/NiagaraEditor/Private/NiagaraEditorModule.cpp
+++ b/Engine/Plugins/FX/Niagara/Source/NiagaraEditor/Private/NiagaraEditorModule.cpp
@@ -1038,6 +1038,7 @@ void FNiagaraEditorModule::StartupModule()
 	RegisterParameterScopeInfo(FNiagaraConstants::LocalNamespace, FNiagaraParameterScopeInfo(ENiagaraParameterScope::Local, PARAM_MAP_LOCAL_MODULE_STR));
 	RegisterParameterScopeInfo(FNiagaraConstants::ScriptPersistentScopeName, FNiagaraParameterScopeInfo(ENiagaraParameterScope::ScriptPersistent, PARAM_MAP_SCRIPT_PERSISTENT_STR));
 	RegisterParameterScopeInfo(FNiagaraConstants::ScriptTransientScopeName, FNiagaraParameterScopeInfo(ENiagaraParameterScope::ScriptTransient, PARAM_MAP_SCRIPT_TRANSIENT_STR));
+	RegisterParameterScopeInfo(FNiagaraConstants::OutputScopeName, FNiagaraParameterScopeInfo(ENiagaraParameterScope::Output, PARAM_MAP_OUTPUT_MODULE_STR));
 	RegisterParameterScopeInfo(FNiagaraConstants::CustomScopeName, FNiagaraParameterScopeInfo(ENiagaraParameterScope::Custom, FString()));
 }
 
diff --git a/Engine/Plugins/FX/Niagara/Source/NiagaraEditor/Private/NiagaraEditorUtilities.cpp b/Engine/Plugins/FX/Niagara/Source/NiagaraEditor/Private/NiagaraEditorUtilities.cpp
--- a/Engine/Plugins/FX/Niagara/Source/NiagaraEditor/Private/NiagaraEditorUtilities.cpp
+++ b/Engine/Plugins/FX/Niagara/Source/NiagaraEditor/Private/NiagaraEditorUtilities.cpp
@@ -2180,7 +2180,9 @@ FName FNiagaraEditorUtilities::GetScopeNameForParameterScope(ENiagaraParameterSc
 	case ENiagaraParameterScope::ScriptPersistent:
 		return FNiagaraConstants::ScriptPersistentScopeName;
 	case ENiagaraParameterScope::ScriptTransient:
-		return FNiagaraConstants::ScriptTransientScopeName;
+		return FNiagaraConstants::ScriptTransientScopeName; 
+	case ENiagaraParameterScope::Output:
+		return FNiagaraConstants::OutputScopeName;
 	default:
 		ensureMsgf(false, TEXT("Tried to get scope name for unknown parameter scope!"));
 	}
