Niagara - When refreshing emitter nodes from external changes and the enabled state changes, mark the node as needing sychrnonization and the graph as needs compile.  Fix inconsistency issues when enabling and disabling emitters.

diff --git a/Engine/Plugins/FX/Niagara/Source/NiagaraEditor/Private/NiagaraNodeEmitter.cpp b/Engine/Plugins/FX/Niagara/Source/NiagaraEditor/Private/NiagaraNodeEmitter.cpp
--- a/Engine/Plugins/FX/Niagara/Source/NiagaraEditor/Private/NiagaraNodeEmitter.cpp
+++ b/Engine/Plugins/FX/Niagara/Source/NiagaraEditor/Private/NiagaraNodeEmitter.cpp
@@ -255,7 +255,12 @@ UNiagaraGraph* UNiagaraNodeEmitter::GetCalledGraph() const
 bool UNiagaraNodeEmitter::RefreshFromExternalChanges()
 {
 	DisplayName = GetNameFromEmitter();
+	ENodeEnabledState OldEnabledState = GetDesiredEnabledState();
 	SyncEnabledState();
+	if (OldEnabledState != GetDesiredEnabledState())
+	{
+		MarkNodeRequiresSynchronization(TEXT("Emitter Node Enabled Changed"), true);
+	}
 	return true;
 }
 
