Adding particle common parameters

diff --git a/Engine/Plugins/FX/Niagara/Source/NiagaraEditor/Private/Widgets/SNiagaraParameterPanel.cpp b/Engine/Plugins/FX/Niagara/Source/NiagaraEditor/Private/Widgets/SNiagaraParameterPanel.cpp
--- a/Engine/Plugins/FX/Niagara/Source/NiagaraEditor/Private/Widgets/SNiagaraParameterPanel.cpp
+++ b/Engine/Plugins/FX/Niagara/Source/NiagaraEditor/Private/Widgets/SNiagaraParameterPanel.cpp
@@ -353,6 +353,7 @@ void SNiagaraAddParameterMenu2::Construct(const FArguments& InArgs, TArray<TWeak
 		[
 			SNew(SBox)
 			.MinDesiredWidth(300)
+			.MaxDesiredHeight(700) // Set max desired height to prevent flickering bug for menu larger than screen
 			[
 				SAssignNew(GraphMenu, SGraphActionMenu)
 				.OnActionSelected(this, &SNiagaraAddParameterMenu2::OnActionSelected)
@@ -517,6 +518,7 @@ void SNiagaraAddParameterMenu2::CollectAllActions(FGraphActionListBuilderBase& O
 		NiagaraParameterPanelSectionID::Type Type = ShowKnownConstantParametersFilter.Get();
 		if (Type != NiagaraParameterPanelSectionID::Type::NONE)
 		{
+			// Handle engine parameters...
 			if (Type == NiagaraParameterPanelSectionID::Type::ENGINE ||
 				Type == NiagaraParameterPanelSectionID::Type::OWNER ||
 				Type == NiagaraParameterPanelSectionID::Type::REFERENCES ||
@@ -553,6 +555,28 @@ void SNiagaraAddParameterMenu2::CollectAllActions(FGraphActionListBuilderBase& O
 					OutAllActions.AddAction(Action);
 				}
 			}
+
+			// Handle particles 
+			if (Type == NiagaraParameterPanelSectionID::Type::PARTICLES ||
+				Type == NiagaraParameterPanelSectionID::Type::REFERENCES )
+			{
+				const TArray<FNiagaraVariable>& Vars = FNiagaraConstants::GetCommonParticleAttributes();
+
+				for (const FNiagaraVariable& Var : Vars)
+				{
+					const FText Category = LOCTEXT("NiagaraAddCommParticleParameterMenu", "Common Particle Attributes");
+					const FText DisplayName = FText::FromName(Var.GetName());
+					const FText Tooltip = FNiagaraConstants::GetAttributeDescription(Var);
+					
+					TSharedPtr<FNiagaraMenuAction> Action(new FNiagaraMenuAction(
+						Category, DisplayName, Tooltip, 0, FText::GetEmpty(),
+						FNiagaraMenuAction::FOnExecuteStackAction::CreateSP(this, &SNiagaraAddParameterMenu2::AddParameterSelected, Var)));
+
+					OutAllActions.AddAction(Action);
+
+				}
+			}
+
 		}
 	}
 }
