Adding tooltip cue for defaults

diff --git a/Engine/Plugins/FX/Niagara/Source/NiagaraEditor/Private/NiagaraNodeParameterMapGet.cpp b/Engine/Plugins/FX/Niagara/Source/NiagaraEditor/Private/NiagaraNodeParameterMapGet.cpp
--- a/Engine/Plugins/FX/Niagara/Source/NiagaraEditor/Private/NiagaraNodeParameterMapGet.cpp
+++ b/Engine/Plugins/FX/Niagara/Source/NiagaraEditor/Private/NiagaraNodeParameterMapGet.cpp
@@ -500,7 +500,7 @@ void UNiagaraNodeParameterMapGet::GetPinHoverText(const UEdGraphPin& Pin, FStrin
 				const UEdGraphPin* OutputPin = GetOutputPinForDefault(&Pin);
 				if (OutputPin)
 				{
-					FText Desc = FText::Format(LOCTEXT("DefaultValueTooltip", "Default value for \"{0}\" if no other module has set it previously in the stack."), FText::FromName(OutputPin->PinName));
+					FText Desc = FText::Format(LOCTEXT("DefaultValueTooltip", "Default value for \"{0}\" if no other module has set it previously in the stack.\nPlease edit this value by selecting in the parameters panel and editing in the details panel."), FText::FromName(OutputPin->PinName));
 					HoverTextOut = Desc.ToString();
 					return;
 				}
@@ -511,7 +511,11 @@ void UNiagaraNodeParameterMapGet::GetPinHoverText(const UEdGraphPin& Pin, FStrin
 				TOptional<FNiagaraVariableMetaData> Metadata = NiagaraGraph->GetMetaData(Var);
 
 				FText Description = Metadata.IsSet() ? Metadata->Description : FText::GetEmpty();
-				FText ToolTipText = FNiagaraEditorUtilities::FormatVariableDescription(Description, FText::FromName(Pin.PinName), TypeDef.GetNameText());
+				const FText TooltipFormat = LOCTEXT("Parameters", "Name: {0} \nType: {1}\nDescription: {2}\nScope: {3}\nUsage: {4}"); 
+				const FText Name = FText::FromName(Var.GetName());
+				const FText ScopeText = FText::FromName(Metadata->GetScopeName());
+				const FText UsageText = StaticEnum<ENiagaraScriptParameterUsage>()->GetDisplayNameTextByValue((int64)Metadata->GetUsage());
+				const FText ToolTipText = FText::Format(TooltipFormat, FText::FromName(Var.GetName()), Var.GetType().GetNameText(), Description, ScopeText, UsageText);
 				HoverTextOut = ToolTipText.ToString();
 			}
 		}
