Niagara - Unify some of the "Library Only" widgets, and add the option to the module and dynamic input reassignment menus.

diff --git a/Engine/Plugins/FX/Niagara/Source/NiagaraEditor/Private/Widgets/SNiagaraLibraryOnlyToggleHeader.cpp b/Engine/Plugins/FX/Niagara/Source/NiagaraEditor/Private/Widgets/SNiagaraLibraryOnlyToggleHeader.cpp
new file mode 100644
--- /dev/null
+++ b/Engine/Plugins/FX/Niagara/Source/NiagaraEditor/Private/Widgets/SNiagaraLibraryOnlyToggleHeader.cpp
@@ -0,0 +1,60 @@
+// Copyright Epic Games, Inc. All Rights Reserved.
+
+#include "Widgets/SNiagaraLibraryOnlyToggleHeader.h"
+#include "SGraphActionMenu.h"
+#include "Widgets/Input/SCheckBox.h"
+#include "Widgets/Text/STextBlock.h"
+#include "Widgets/SBoxPanel.h"
+
+void SNiagaraLibraryOnlyToggleHeader::Construct(const FArguments& InArgs)
+{
+	LibraryOnly = InArgs._LibraryOnly;
+	LibraryOnlyChanged = InArgs._LibraryOnlyChanged;
+	ChildSlot
+	[
+		SNew(SHorizontalBox)
+
+		// Search context description
+		+SHorizontalBox::Slot()
+		.VAlign(VAlign_Center)
+		.AutoWidth()
+		[
+			SNew(STextBlock)
+			.Text(InArgs._HeaderLabelText)
+		]
+
+		// Library Only Toggle
+		+SHorizontalBox::Slot()
+		.HAlign(HAlign_Right)
+		.VAlign(VAlign_Center)
+		[
+			SNew(SCheckBox)
+			.OnCheckStateChanged(this, &SNiagaraLibraryOnlyToggleHeader::OnCheckStateChanged)
+			.IsChecked(this, &SNiagaraLibraryOnlyToggleHeader::GetCheckState)
+			[
+				SNew(STextBlock)
+				.Text(NSLOCTEXT("LibraryOnlyToggle", "LibraryOnly", "Library Only"))
+			]
+		]
+	];
+}
+
+void SNiagaraLibraryOnlyToggleHeader::SetActionMenu(TSharedRef<SGraphActionMenu> InActionMenu)
+{
+	ActionMenuWeak = InActionMenu;
+}
+
+void SNiagaraLibraryOnlyToggleHeader::OnCheckStateChanged(ECheckBoxState InCheckState)
+{
+	LibraryOnlyChanged.ExecuteIfBound(InCheckState == ECheckBoxState::Checked);
+	TSharedPtr<SGraphActionMenu> ActionMenu = ActionMenuWeak.Pin();
+	if (ActionMenu.IsValid())
+	{
+		ActionMenu->RefreshAllActions(true, false);
+	}
+}
+
+ECheckBoxState SNiagaraLibraryOnlyToggleHeader::GetCheckState() const
+{
+	return LibraryOnly.Get() ? ECheckBoxState::Checked : ECheckBoxState::Unchecked;
+}
\ No newline at end of file
diff --git a/Engine/Plugins/FX/Niagara/Source/NiagaraEditor/Public/Widgets/SNiagaraLibraryOnlyToggleHeader.h b/Engine/Plugins/FX/Niagara/Source/NiagaraEditor/Public/Widgets/SNiagaraLibraryOnlyToggleHeader.h
new file mode 100644
--- /dev/null
+++ b/Engine/Plugins/FX/Niagara/Source/NiagaraEditor/Public/Widgets/SNiagaraLibraryOnlyToggleHeader.h
@@ -0,0 +1,38 @@
+// Copyright Epic Games, Inc. All Rights Reserved.
+
+#pragma once
+
+#include "Widgets/SCompoundWidget.h"
+#include "Widgets/DeclarativeSyntaxSupport.h"
+
+class SGraphActionMenu;
+
+class SNiagaraLibraryOnlyToggleHeader : public SCompoundWidget
+{
+public:
+	DECLARE_DELEGATE_OneParam(FOnLibraryOnlyChanged, bool);
+
+public:
+	SLATE_BEGIN_ARGS(SNiagaraLibraryOnlyToggleHeader)
+		: _LibraryOnly(true)
+		, _HeaderLabelText(NSLOCTEXT("LibraryOnlyToggle", "DefaultHeader", "Select Item"))
+		{}
+		SLATE_ATTRIBUTE(bool, LibraryOnly)
+		SLATE_EVENT(FOnLibraryOnlyChanged, LibraryOnlyChanged)
+		SLATE_ARGUMENT(FText, HeaderLabelText)
+	SLATE_END_ARGS();
+
+	NIAGARAEDITOR_API void Construct(const FArguments& InArgs);
+
+	NIAGARAEDITOR_API void SetActionMenu(TSharedRef<SGraphActionMenu> InActionMenu);
+
+private:
+	void OnCheckStateChanged(ECheckBoxState InCheckState);
+
+	ECheckBoxState GetCheckState() const;
+
+private:
+	TAttribute<bool> LibraryOnly;
+	FOnLibraryOnlyChanged LibraryOnlyChanged;
+	TWeakPtr<SGraphActionMenu> ActionMenuWeak;
+};
\ No newline at end of file
diff --git a/Engine/Plugins/FX/Niagara/Source/NiagaraEditorWidgets/Private/Stack/SNiagaraStackFunctionInputValue.cpp b/Engine/Plugins/FX/Niagara/Source/NiagaraEditorWidgets/Private/Stack/SNiagaraStackFunctionInputValue.cpp
--- a/Engine/Plugins/FX/Niagara/Source/NiagaraEditorWidgets/Private/Stack/SNiagaraStackFunctionInputValue.cpp
+++ b/Engine/Plugins/FX/Niagara/Source/NiagaraEditorWidgets/Private/Stack/SNiagaraStackFunctionInputValue.cpp
@@ -31,12 +31,13 @@
 #include "SNiagaraGraphActionWidget.h"
 #include "Subsystems/AssetEditorSubsystem.h"
 #include "EditorFontGlyphs.h"
+#include "Widgets/SNiagaraLibraryOnlyToggleHeader.h"
 
 #define LOCTEXT_NAMESPACE "NiagaraStackFunctionInputValue"
 
 const float TextIconSize = 16;
 
-bool SNiagaraStackFunctionInputValue::bIncludeNonLibraryInputs = false;
+bool SNiagaraStackFunctionInputValue::bLibraryOnly = true;
 
 void SNiagaraStackFunctionInputValue::Construct(const FArguments& InArgs, UNiagaraStackFunctionInput* InFunctionInput)
 {
@@ -434,6 +435,9 @@ TSharedRef<SExpanderArrow> SNiagaraStackFunctionInputValue::CreateCustomNiagaraF
 
 TSharedRef<SWidget> SNiagaraStackFunctionInputValue::OnGetAvailableHandleMenu()
 {
+	TSharedPtr<SNiagaraLibraryOnlyToggleHeader> LibraryOnlyToggle;
+	TSharedPtr<SGraphActionMenu> GraphActionMenu;
+
 	TSharedRef<SBorder> MenuWidget = SNew(SBorder)
 	.BorderImage(FEditorStyle::GetBrush("Menu.Background"))
 	.Padding(5)
@@ -446,35 +450,15 @@ TSharedRef<SWidget> SNiagaraStackFunctionInputValue::OnGetAvailableHandleMenu()
 			+SVerticalBox::Slot()
 			.Padding(1.0f)
 			[
-				SNew(SHorizontalBox)
-
-				// Search context description
-				+SHorizontalBox::Slot()
-				.VAlign(VAlign_Center)
-				.AutoWidth()
-				[
-					SNew(STextBlock)
-					.Text(LOCTEXT("FunctionInputValueTitle", "Edit value"))
-				]
-
-				// Library Only Toggle
-				+SHorizontalBox::Slot()
-				.HAlign(HAlign_Right)
-				.VAlign(VAlign_Center)
-				[
-					SNew(SCheckBox)
-					.OnCheckStateChanged(this, &SNiagaraStackFunctionInputValue::OnLibraryToggleChanged)
-					.IsChecked(this, &SNiagaraStackFunctionInputValue::LibraryToggleIsChecked)
-					[
-						SNew(STextBlock)
-						.Text(LOCTEXT("LibraryOnly", "Library Only"))
-					]
-				]
+				SAssignNew(LibraryOnlyToggle, SNiagaraLibraryOnlyToggleHeader)
+				.HeaderLabelText(LOCTEXT("FunctionInputValueTitle", "Edit value"))
+				.LibraryOnly(this, &SNiagaraStackFunctionInputValue::GetLibraryOnly)
+				.LibraryOnlyChanged(this, &SNiagaraStackFunctionInputValue::SetLibraryOnly)
 			]
 			+SVerticalBox::Slot()
 			.FillHeight(15)
 			[
-				SAssignNew(SelectInputFunctionMenu, SGraphActionMenu)
+				SAssignNew(GraphActionMenu, SGraphActionMenu)
 				.OnActionSelected(this, &SNiagaraStackFunctionInputValue::OnActionSelected)
 				.OnCollectAllActions(this, &SNiagaraStackFunctionInputValue::CollectAllActions)
 				.AutoExpandActionMenu(false)
@@ -488,7 +472,8 @@ TSharedRef<SWidget> SNiagaraStackFunctionInputValue::OnGetAvailableHandleMenu()
 		]
 	];
 
-	SetFunctionInputButton->SetMenuContentWidgetToFocus(SelectInputFunctionMenu->GetFilterTextBox()->AsShared());
+	LibraryOnlyToggle->SetActionMenu(GraphActionMenu.ToSharedRef());
+	SetFunctionInputButton->SetMenuContentWidgetToFocus(GraphActionMenu->GetFilterTextBox()->AsShared());
 	return MenuWidget;
 }
 
@@ -532,7 +517,7 @@ void SNiagaraStackFunctionInputValue::CollectAllActions(FGraphActionListBuilderB
 	{
 		const FText CategoryName = LOCTEXT("DynamicInputValueCategory", "Dynamic Inputs");
 		TArray<UNiagaraScript*> DynamicInputScripts;
-		FunctionInput->GetAvailableDynamicInputs(DynamicInputScripts, bIncludeNonLibraryInputs);
+		FunctionInput->GetAvailableDynamicInputs(DynamicInputScripts, bLibraryOnly == false);
 		for (UNiagaraScript* DynamicInputScript : DynamicInputScripts)
 		{
 			const FText DynamicInputText = FNiagaraEditorUtilities::FormatScriptName(DynamicInputScript->GetFName(), DynamicInputScript->bExposeToLibrary);
@@ -859,11 +844,11 @@ void ReassignDynamicInputScript(UNiagaraStackFunctionInput* FunctionInput, UNiag
 	FunctionInput->ReassignDynamicInputScript(NewDynamicInputScript);
 }
 
-void CollectDynamicInputActions(FGraphActionListBuilderBase& DynamicInputActions, UNiagaraStackFunctionInput* FunctionInput)
+void SNiagaraStackFunctionInputValue::CollectDynamicInputActionsForReassign(FGraphActionListBuilderBase& DynamicInputActions) const
 {
 	const FText CategoryName = LOCTEXT("DynamicInputValueCategory", "Dynamic Inputs");
 	TArray<UNiagaraScript*> DynamicInputScripts;
-	FunctionInput->GetAvailableDynamicInputs(DynamicInputScripts);
+	FunctionInput->GetAvailableDynamicInputs(DynamicInputScripts, bLibraryOnly == false);
 	for (UNiagaraScript* DynamicInputScript : DynamicInputScripts)
 	{
 		const FText DynamicInputText = FNiagaraEditorUtilities::FormatScriptName(DynamicInputScript->GetFName(), DynamicInputScript->bExposeToLibrary);
@@ -876,6 +861,9 @@ void CollectDynamicInputActions(FGraphActionListBuilderBase& DynamicInputActions
 
 void SNiagaraStackFunctionInputValue::ShowReassignDynamicInputScriptMenu()
 {
+	TSharedPtr<SNiagaraLibraryOnlyToggleHeader> LibraryOnlyToggle;
+	TSharedPtr<SGraphActionMenu> GraphActionMenu;
+
 	TSharedRef<SBorder> MenuWidget = SNew(SBorder)
 		.BorderImage(FEditorStyle::GetBrush("Menu.Background"))
 		.Padding(5)
@@ -884,30 +872,44 @@ void SNiagaraStackFunctionInputValue::ShowReassignDynamicInputScriptMenu()
 			.WidthOverride(300)
 			.HeightOverride(400)
 			[
-				SNew(SGraphActionMenu)
-				.OnActionSelected(this, &SNiagaraStackFunctionInputValue::OnActionSelected)
-				.OnCollectAllActions_Static(CollectDynamicInputActions, FunctionInput)
-				.AutoExpandActionMenu(true)
-				.ShowFilterTextBox(true)
-				.OnCreateCustomRowExpander_Static(&CreateCustomNiagaraFunctionInputActionExpander)
+				SNew(SVerticalBox)
+				+SVerticalBox::Slot()
+				.Padding(1.0f)
+				[
+					SAssignNew(LibraryOnlyToggle, SNiagaraLibraryOnlyToggleHeader)
+					.HeaderLabelText(LOCTEXT("ReassignDynamicInputLabel", "Select a new dynamic input"))
+					.LibraryOnly(this, &SNiagaraStackFunctionInputValue::GetLibraryOnly)
+					.LibraryOnlyChanged(this, &SNiagaraStackFunctionInputValue::SetLibraryOnly)
+				]
+				+SVerticalBox::Slot()
+				.FillHeight(15)
+				[
+					SAssignNew(GraphActionMenu, SGraphActionMenu)
+					.OnActionSelected(this, &SNiagaraStackFunctionInputValue::OnActionSelected)
+					.OnCollectAllActions(this, &SNiagaraStackFunctionInputValue::CollectDynamicInputActionsForReassign)
+					.AutoExpandActionMenu(true)
+					.ShowFilterTextBox(true)
+					.OnCreateCustomRowExpander_Static(&CreateCustomNiagaraFunctionInputActionExpander)
+				]
 			]
 		];
 
+	LibraryOnlyToggle->SetActionMenu(GraphActionMenu.ToSharedRef());
+
 	FGeometry ThisGeometry = GetCachedGeometry();
 	bool bAutoAdjustForDpiScale = false; // Don't adjust for dpi scale because the push menu command is expecting an unscaled position.
 	FVector2D MenuPosition = FSlateApplication::Get().CalculatePopupWindowPosition(ThisGeometry.GetLayoutBoundingRect(), MenuWidget->GetDesiredSize(), bAutoAdjustForDpiScale);
 	FSlateApplication::Get().PushMenu(AsShared(), FWidgetPath(), MenuWidget, MenuPosition, FPopupTransitionEffect::ContextMenu);
 }
 
-void SNiagaraStackFunctionInputValue::OnLibraryToggleChanged(ECheckBoxState CheckState)
+bool SNiagaraStackFunctionInputValue::GetLibraryOnly() const
 {
-	SNiagaraStackFunctionInputValue::bIncludeNonLibraryInputs = CheckState == ECheckBoxState::Unchecked;
-	SelectInputFunctionMenu->RefreshAllActions(true, false);
+	return bLibraryOnly;
 }
 
-ECheckBoxState SNiagaraStackFunctionInputValue::LibraryToggleIsChecked() const
+void SNiagaraStackFunctionInputValue::SetLibraryOnly(bool bInIsLibraryOnly)
 {
-	return SNiagaraStackFunctionInputValue::bIncludeNonLibraryInputs ? ECheckBoxState::Unchecked : ECheckBoxState::Checked;
+	bLibraryOnly = bInIsLibraryOnly;
 }
 
 #undef LOCTEXT_NAMESPACE
diff --git a/Engine/Plugins/FX/Niagara/Source/NiagaraEditorWidgets/Private/Stack/SNiagaraStackFunctionInputValue.h b/Engine/Plugins/FX/Niagara/Source/NiagaraEditorWidgets/Private/Stack/SNiagaraStackFunctionInputValue.h
--- a/Engine/Plugins/FX/Niagara/Source/NiagaraEditorWidgets/Private/Stack/SNiagaraStackFunctionInputValue.h
+++ b/Engine/Plugins/FX/Niagara/Source/NiagaraEditorWidgets/Private/Stack/SNiagaraStackFunctionInputValue.h
@@ -14,6 +14,7 @@ class SNiagaraParameterEditor;
 class SBox;
 class IStructureDetailsView;
 class SComboButton;
+struct FGraphActionListBuilderBase;
 
 class SNiagaraStackFunctionInputValue: public SCompoundWidget
 {
@@ -132,11 +133,13 @@ private:
 
 	bool OnFunctionInputAllowDrop(TSharedPtr<FDragDropOperation> DragDropOperation);
 
+	void CollectDynamicInputActionsForReassign(FGraphActionListBuilderBase& DynamicInputActions) const;
+
 	void ShowReassignDynamicInputScriptMenu();
 
-	void OnLibraryToggleChanged(ECheckBoxState CheckState);
+	bool GetLibraryOnly() const;
 
-	ECheckBoxState LibraryToggleIsChecked() const;
+	void SetLibraryOnly(bool bInIsLibraryOnly);
 
 private:
 	UNiagaraStackFunctionInput* FunctionInput;
@@ -149,7 +152,6 @@ private:
 	TSharedPtr<IStructureDetailsView> LocalValueStructDetailsView;
 
 	TSharedPtr<SComboButton> SetFunctionInputButton;
-	TSharedPtr<SGraphActionMenu> SelectInputFunctionMenu;
 
-	static bool bIncludeNonLibraryInputs;
+	static bool bLibraryOnly;
 };
diff --git a/Engine/Plugins/FX/Niagara/Source/NiagaraEditorWidgets/Private/Stack/SNiagaraStackItemGroupAddMenu.cpp b/Engine/Plugins/FX/Niagara/Source/NiagaraEditorWidgets/Private/Stack/SNiagaraStackItemGroupAddMenu.cpp
--- a/Engine/Plugins/FX/Niagara/Source/NiagaraEditorWidgets/Private/Stack/SNiagaraStackItemGroupAddMenu.cpp
+++ b/Engine/Plugins/FX/Niagara/Source/NiagaraEditorWidgets/Private/Stack/SNiagaraStackItemGroupAddMenu.cpp
@@ -9,13 +9,16 @@
 #include "Framework/Application/SlateApplication.h"
 #include "Widgets/Input/SCheckBox.h"
 #include "Styling/SlateTypes.h"
+#include "Widgets/SNiagaraLibraryOnlyToggleHeader.h"
 
 #define LOCTEXT_NAMESPACE "NiagaraStackItemGroupAddMenu"
 
-bool SNiagaraStackItemGroupAddMenu::bIncludeNonLibraryScripts = false;
+bool SNiagaraStackItemGroupAddMenu::bLibraryOnly = true;
 
 void SNiagaraStackItemGroupAddMenu::Construct(const FArguments& InArgs, INiagaraStackItemGroupAddUtilities* InAddUtilities, int32 InInsertIndex)
 {
+	TSharedPtr<SNiagaraLibraryOnlyToggleHeader> LibraryOnlyToggle;
+
 	AddUtilities = InAddUtilities;
 	InsertIndex = InInsertIndex;
 	bSetFocusOnNextTick = true;
@@ -33,30 +36,10 @@ void SNiagaraStackItemGroupAddMenu::Construct(const FArguments& InArgs, INiagara
 				+SVerticalBox::Slot()
 				.Padding(1.0f)
 				[
-					SNew(SHorizontalBox)
-
-					// Search context description
-					+SHorizontalBox::Slot()
-					.VAlign(VAlign_Center)
-					.AutoWidth()
-					[
-						SNew(STextBlock)
-						.Text(FText::Format(LOCTEXT("AddToGroupFormatTitle", "Add new {0}"), AddUtilities->GetAddItemName()))
-					]
-
-					// Library Only Toggle
-					+SHorizontalBox::Slot()
-					.HAlign(HAlign_Right)
-					.VAlign(VAlign_Center)
-					[
-						SNew(SCheckBox)
-						.OnCheckStateChanged(this, &SNiagaraStackItemGroupAddMenu::OnLibraryToggleChanged)
-						.IsChecked(this, &SNiagaraStackItemGroupAddMenu::LibraryToggleIsChecked)
-						[
-							SNew(STextBlock)
-							.Text(LOCTEXT("LibraryOnly", "Library Only"))
-						]
-					]
+					SAssignNew(LibraryOnlyToggle, SNiagaraLibraryOnlyToggleHeader)
+					.HeaderLabelText(FText::Format(LOCTEXT("AddToGroupFormatTitle", "Add new {0}"), AddUtilities->GetAddItemName()))
+					.LibraryOnly(this, &SNiagaraStackItemGroupAddMenu::GetLibraryOnly)
+					.LibraryOnlyChanged(this, &SNiagaraStackItemGroupAddMenu::SetLibraryOnly)
 				]
 				+SVerticalBox::Slot()
 				.FillHeight(15)
@@ -70,6 +53,8 @@ void SNiagaraStackItemGroupAddMenu::Construct(const FArguments& InArgs, INiagara
 			]
 		]
 	];
+
+	LibraryOnlyToggle->SetActionMenu(AddMenu.ToSharedRef());
 }
 
 void SNiagaraStackItemGroupAddMenu::Tick(const FGeometry& AllottedGeometry, const double InCurrentTime, const float InDeltaTime)
@@ -95,7 +80,7 @@ void SNiagaraStackItemGroupAddMenu::CollectAllAddActions(FGraphActionListBuilder
 
 	FNiagaraStackItemGroupAddOptions AddOptions;
 	AddOptions.bIncludeDeprecated = false;
-	AddOptions.bIncludeNonLibrary = bIncludeNonLibraryScripts;
+	AddOptions.bIncludeNonLibrary = bLibraryOnly == false;
 
 	TArray<TSharedRef<INiagaraStackItemGroupAddAction>> AddActions;
 	AddUtilities->GenerateAddActions(AddActions, AddOptions);
@@ -126,15 +111,14 @@ void SNiagaraStackItemGroupAddMenu::OnActionSelected(const TArray< TSharedPtr<FE
 	}
 }
 
-void SNiagaraStackItemGroupAddMenu::OnLibraryToggleChanged(ECheckBoxState CheckState)
+bool SNiagaraStackItemGroupAddMenu::GetLibraryOnly() const
 {
-	SNiagaraStackItemGroupAddMenu::bIncludeNonLibraryScripts = CheckState == ECheckBoxState::Unchecked;
-	AddMenu->RefreshAllActions(true, false);
+	return bLibraryOnly;
 }
 
-ECheckBoxState SNiagaraStackItemGroupAddMenu::LibraryToggleIsChecked() const
+void SNiagaraStackItemGroupAddMenu::SetLibraryOnly(bool bInLibraryOnly)
 {
-	return SNiagaraStackItemGroupAddMenu::bIncludeNonLibraryScripts ? ECheckBoxState::Unchecked : ECheckBoxState::Checked;
+	bLibraryOnly = bInLibraryOnly;
 }
 
 #undef LOCTEXT_NAMESPACE
\ No newline at end of file
diff --git a/Engine/Plugins/FX/Niagara/Source/NiagaraEditorWidgets/Private/Stack/SNiagaraStackItemGroupAddMenu.h b/Engine/Plugins/FX/Niagara/Source/NiagaraEditorWidgets/Private/Stack/SNiagaraStackItemGroupAddMenu.h
--- a/Engine/Plugins/FX/Niagara/Source/NiagaraEditorWidgets/Private/Stack/SNiagaraStackItemGroupAddMenu.h
+++ b/Engine/Plugins/FX/Niagara/Source/NiagaraEditorWidgets/Private/Stack/SNiagaraStackItemGroupAddMenu.h
@@ -26,9 +26,9 @@ private:
 
 	void OnActionSelected(const TArray< TSharedPtr<FEdGraphSchemaAction> >& SelectedActions, ESelectInfo::Type InSelectionType);
 
-	void OnLibraryToggleChanged(ECheckBoxState CheckState);
+	bool GetLibraryOnly() const;
 
-	ECheckBoxState LibraryToggleIsChecked() const;
+	void SetLibraryOnly(bool bInLibraryOnly);
 
 private:
 	INiagaraStackItemGroupAddUtilities* AddUtilities;
@@ -39,5 +39,5 @@ private:
 
 	bool bSetFocusOnNextTick;
 	
-	static bool bIncludeNonLibraryScripts;
+	static bool bLibraryOnly;
 };
\ No newline at end of file
diff --git a/Engine/Plugins/FX/Niagara/Source/NiagaraEditorWidgets/Private/Stack/SNiagaraStackModuleItem.cpp b/Engine/Plugins/FX/Niagara/Source/NiagaraEditorWidgets/Private/Stack/SNiagaraStackModuleItem.cpp
--- a/Engine/Plugins/FX/Niagara/Source/NiagaraEditorWidgets/Private/Stack/SNiagaraStackModuleItem.cpp
+++ b/Engine/Plugins/FX/Niagara/Source/NiagaraEditorWidgets/Private/Stack/SNiagaraStackModuleItem.cpp
@@ -28,9 +28,12 @@
 #include "Editor.h"
 #include "EditorFontGlyphs.h"
 #include "ViewModels/Stack/NiagaraStackClipboardUtilities.h"
+#include "Widgets/SNiagaraLibraryOnlyToggleHeader.h"
 
 #define LOCTEXT_NAMESPACE "NiagaraStackModuleItem"
 
+bool SNiagaraStackModuleItem::bLibraryOnly = true;
+
 void SNiagaraStackModuleItem::Construct(const FArguments& InArgs, UNiagaraStackModuleItem& InModuleItem, UNiagaraStackViewModel* InStackViewModel)
 {
 	ModuleItem = &InModuleItem;
@@ -233,12 +236,13 @@ void OnActionSelected(const TArray<TSharedPtr<FEdGraphSchemaAction>>& SelectedAc
 	}
 }
 
-void CollectModuleActions(FGraphActionListBuilderBase& ModuleActions, UNiagaraStackModuleItem* ModuleItem)
+void SNiagaraStackModuleItem::CollectModuleActions(FGraphActionListBuilderBase& ModuleActions)
 {
 	TArray<FAssetData> ModuleAssets;
 	FNiagaraEditorUtilities::FGetFilteredScriptAssetsOptions ModuleScriptFilterOptions;
 	ModuleScriptFilterOptions.ScriptUsageToInclude = ENiagaraScriptUsage::Module;
 	ModuleScriptFilterOptions.TargetUsageToMatch = ModuleItem->GetOutputNode()->GetUsage();
+	ModuleScriptFilterOptions.bIncludeNonLibraryScripts = bLibraryOnly == false;
 	FNiagaraEditorUtilities::GetFilteredScriptAssets(ModuleScriptFilterOptions, ModuleAssets);
 	for (const FAssetData& ModuleAsset : ModuleAssets)
 	{
@@ -268,6 +272,9 @@ void CollectModuleActions(FGraphActionListBuilderBase& ModuleActions, UNiagaraSt
 
 void SNiagaraStackModuleItem::ShowReassignModuleScriptMenu()
 {
+	TSharedPtr<SNiagaraLibraryOnlyToggleHeader> LibraryOnlyToggle;
+	TSharedPtr<SGraphActionMenu> GraphActionMenu;
+
 	TSharedRef<SBorder> MenuWidget = SNew(SBorder)
 		.BorderImage(FEditorStyle::GetBrush("Menu.Background"))
 		.Padding(5)
@@ -276,17 +283,42 @@ void SNiagaraStackModuleItem::ShowReassignModuleScriptMenu()
 			.WidthOverride(300)
 			.HeightOverride(400)
 			[
-				SNew(SGraphActionMenu)
-				.OnActionSelected_Static(OnActionSelected)
-				.OnCollectAllActions_Static(CollectModuleActions, ModuleItem)
-				.ShowFilterTextBox(true)
+				SNew(SVerticalBox)
+				+SVerticalBox::Slot()
+				.Padding(1.0f)
+				[
+					SAssignNew(LibraryOnlyToggle, SNiagaraLibraryOnlyToggleHeader)
+					.HeaderLabelText(LOCTEXT("ReassignModuleLabel", "Select a new module"))
+					.LibraryOnly(this, &SNiagaraStackModuleItem::GetLibraryOnly)
+					.LibraryOnlyChanged(this, &SNiagaraStackModuleItem::SetLibraryOnly)
+				]
+				+SVerticalBox::Slot()
+				.FillHeight(15)
+				[
+					SAssignNew(GraphActionMenu, SGraphActionMenu)
+					.OnActionSelected_Static(OnActionSelected)
+					.OnCollectAllActions(this, &SNiagaraStackModuleItem::CollectModuleActions)
+					.ShowFilterTextBox(true)
+				]
 			]
 		];
 
+	LibraryOnlyToggle->SetActionMenu(GraphActionMenu.ToSharedRef());
+
 	FGeometry ThisGeometry = GetCachedGeometry();
 	bool bAutoAdjustForDpiScale = false; // Don't adjust for dpi scale because the push menu command is expecting an unscaled position.
 	FVector2D MenuPosition = FSlateApplication::Get().CalculatePopupWindowPosition(ThisGeometry.GetLayoutBoundingRect(), MenuWidget->GetDesiredSize(), bAutoAdjustForDpiScale);
 	FSlateApplication::Get().PushMenu(AsShared(), FWidgetPath(), MenuWidget, MenuPosition, FPopupTransitionEffect::ContextMenu);
 }
 
+bool SNiagaraStackModuleItem::GetLibraryOnly() const
+{
+	return bLibraryOnly;
+}
+
+void SNiagaraStackModuleItem::SetLibraryOnly(bool bInLibraryOnly)
+{
+	bLibraryOnly = bInLibraryOnly;
+}
+
 #undef LOCTEXT_NAMESPACE
\ No newline at end of file
diff --git a/Engine/Plugins/FX/Niagara/Source/NiagaraEditorWidgets/Private/Stack/SNiagaraStackModuleItem.h b/Engine/Plugins/FX/Niagara/Source/NiagaraEditorWidgets/Private/Stack/SNiagaraStackModuleItem.h
--- a/Engine/Plugins/FX/Niagara/Source/NiagaraEditorWidgets/Private/Stack/SNiagaraStackModuleItem.h
+++ b/Engine/Plugins/FX/Niagara/Source/NiagaraEditorWidgets/Private/Stack/SNiagaraStackModuleItem.h
@@ -10,6 +10,7 @@
 class UNiagaraStackModuleItem;
 class UNiagaraStackViewModel;
 class SNiagaraStackDisplayName;
+struct FGraphActionListBuilderBase;
 
 class SNiagaraStackModuleItem : public SNiagaraStackItem
 {
@@ -47,8 +48,16 @@ private:
 
 	bool OnModuleItemAllowDrop(TSharedPtr<class FDragDropOperation> DragDropOperation);
 
+	void CollectModuleActions(FGraphActionListBuilderBase& ModuleActions);
+
 	void ShowReassignModuleScriptMenu();
 
+	bool GetLibraryOnly() const;
+
+	void SetLibraryOnly(bool bInLibraryOnly);
+
 private:
 	UNiagaraStackModuleItem* ModuleItem;
+
+	static bool bLibraryOnly;
 };
\ No newline at end of file
