FNM - Added Compression option for Niagara Systems. Thge Option will try to compress some of the parameters to use half precission, which can be a performance benefit on the GPU and sometimes also on the CPU.

diff --git a/Engine/Plugins/FX/Niagara/Shaders/Private/NiagaraEmitterInstanceShader.usf b/Engine/Plugins/FX/Niagara/Shaders/Private/NiagaraEmitterInstanceShader.usf
--- a/Engine/Plugins/FX/Niagara/Shaders/Private/NiagaraEmitterInstanceShader.usf
+++ b/Engine/Plugins/FX/Niagara/Shaders/Private/NiagaraEmitterInstanceShader.usf
@@ -419,6 +419,7 @@ int rand_int(int x, int Seed1, int Seed2, int Seed3)
 	float InputDataFloat(int DataSetIndex, int RegisterIdx);  //DataSetIndex is 0 for main dataset
 	int InputDataInt(int DataSetIndex, int RegisterIdx);
 	bool InputDataBool(int DataSetIndex, int RegisterIdx);
+	float InputDataHalf(int DataSetIndex, int RegisterIdx);
 
 	float InputDataNoadvanceFloat(int DataSetIndex, int RegisterIdx);  //DataSetIndex is 0 for main dataset
 	int InputDataNoadvanceInt(int DataSetIndex, int RegisterIdx);
@@ -427,6 +428,7 @@ int rand_int(int x, int Seed1, int Seed2, int Seed3)
 	void OutputDataFloat(int DataSetIndex, int RegisterIndex, int InstanceIndex, float Value);
 	void OutputDataInt(int DataSetIndex, int RegisterIndex, int InstanceIndex, int Value);
 	void OutputDataBool(int DataSetIndex, int RegisterIndex, int InstanceIndex, bool Value);
+	void OutputDataHalf(int DataSetIndex, int RegisterIndex, int InstanceIndex, float Value);
 #endif
 
 
@@ -458,8 +460,10 @@ int rand_int(int x, int Seed1, int Seed2, int Seed3)
 	 */
 	Buffer<float> InputFloat;
 	Buffer<int> InputInt;
+	Buffer<half> InputHalf;
 	RWBuffer<int> RWOutputInt;
 	RWBuffer<float> RWOutputFloat;
+	RWBuffer<half> RWOutputHalf;
 
 	RWBuffer<uint> RWInstanceCounts;
 	uint ReadInstanceCountOffset;
@@ -646,6 +650,10 @@ int rand_int(int x, int Seed1, int Seed2, int Seed3)
 		return InputInt[RegisterIdx*ComponentBufferSizeRead + InstanceIdx] == -1;
 	}
 
+	float InputDataHalf(int DataSetIndex, int RegisterIdx, int InstanceIdx)
+	{
+		return InputHalf[RegisterIdx*ComponentBufferSizeRead + InstanceIdx];
+	}
 
 	/* ---------------------------------------------------------------------
 	 * OutputData operations 
@@ -666,6 +674,11 @@ int rand_int(int x, int Seed1, int Seed2, int Seed3)
 		RWOutputInt[RegisterIndex*ComponentBufferSizeWrite + InstanceIndex + UpdateStartInstance] = Value ? -1 : 0;
 	}
 
+	void OutputDataHalf(int DataSetIndex, int RegisterIndex, int InstanceIndex, float Value)
+	{
+		RWOutputHalf[RegisterIndex*ComponentBufferSizeWrite + InstanceIndex + UpdateStartInstance] = Value;
+	}
+
 	void EnterStatScope(int ID)	{}
 	void ExitStatScope()	{}
 #endif // GPU_SIMULATION
diff --git a/Engine/Plugins/FX/Niagara/Shaders/Private/NiagaraGPURayTracingTransforms.usf b/Engine/Plugins/FX/Niagara/Shaders/Private/NiagaraGPURayTracingTransforms.usf
--- a/Engine/Plugins/FX/Niagara/Shaders/Private/NiagaraGPURayTracingTransforms.usf
+++ b/Engine/Plugins/FX/Niagara/Shaders/Private/NiagaraGPURayTracingTransforms.usf
@@ -6,12 +6,24 @@
 
 #include "/Engine/Private/Common.ush"
 
-#define NIAGARA_KEYGEN 1
-#include "NiagaraVFParticleAccess.usf"
-
-#if NIAGARA_LOOSEPARAM_BINDING
 Buffer<float> NiagaraParticleDataFloat;
-#endif
+uint NiagaraFloatDataStride;
+
+float GetFloat(int RegisterIdx, uint InstanceID)
+{
+	RegisterIdx &= (~(1 << 31));
+	return NiagaraParticleDataFloat[(RegisterIdx * NiagaraFloatDataStride + InstanceID)];
+}
+
+float3 GetVec3(int RegisterIndex, uint InstanceID)
+{
+	return float3(GetFloat(RegisterIndex, InstanceID), GetFloat(RegisterIndex+1, InstanceID), GetFloat(RegisterIndex+2, InstanceID));
+}
+
+float4 GetVec4(int RegisterIndex, uint InstanceID)
+{
+	return float4(GetFloat(RegisterIndex, InstanceID), GetFloat(RegisterIndex+1, InstanceID), GetFloat(RegisterIndex+2, InstanceID), GetFloat(RegisterIndex+3, InstanceID));
+}
 
 // Offsets: (Position, Transform, Scale, bLocalSpace)
 uint4 NiagaraOffsets;
diff --git a/Engine/Plugins/FX/Niagara/Shaders/Private/NiagaraSortKeyGen.usf b/Engine/Plugins/FX/Niagara/Shaders/Private/NiagaraSortKeyGen.usf
--- a/Engine/Plugins/FX/Niagara/Shaders/Private/NiagaraSortKeyGen.usf
+++ b/Engine/Plugins/FX/Niagara/Shaders/Private/NiagaraSortKeyGen.usf
@@ -15,12 +15,19 @@
 
 #include "/Engine/Private/Common.ush"
 
-#define NIAGARA_KEYGEN 1
-#include "NiagaraVFParticleAccess.usf"
-
-#if NIAGARA_LOOSEPARAM_BINDING
 Buffer<float> NiagaraParticleDataFloat;
-#endif
+uint NiagaraFloatDataStride;
+
+float GetFloat(int RegisterIdx, uint InstanceID)
+{
+	RegisterIdx &= (~(1 << 31));
+	return NiagaraParticleDataFloat[(RegisterIdx * NiagaraFloatDataStride + InstanceID)];
+}
+
+float3 GetVec3(int RegisterIndex, uint InstanceID)
+{
+	return float3(GetFloat(RegisterIndex, InstanceID), GetFloat(RegisterIndex+1, InstanceID), GetFloat(RegisterIndex+2, InstanceID));
+}
 
 Buffer<uint> GPUParticleCountBuffer;
 
diff --git a/Engine/Plugins/FX/Niagara/Shaders/Private/NiagaraVFParticleAccess.usf b/Engine/Plugins/FX/Niagara/Shaders/Private/NiagaraVFParticleAccess.usf
--- a/Engine/Plugins/FX/Niagara/Shaders/Private/NiagaraVFParticleAccess.usf
+++ b/Engine/Plugins/FX/Niagara/Shaders/Private/NiagaraVFParticleAccess.usf
@@ -8,29 +8,22 @@
 //uint NiagaraIntDataOffset;
 //uint NiagaraIntDataStride;
 
-#ifndef NIAGARA_KEYGEN
-#define NIAGARA_KEYGEN 0
-#endif
-
-#if (!RAYHITGROUPSHADER && !COMPUTESHADER) || NIAGARA_KEYGEN
-
-Buffer<float> NiagaraParticleDataFloat;
-uint NiagaraFloatDataStride;
-
-float GetFloat(int RegisterIdx, uint InstanceID)
-{
-	return NiagaraParticleDataFloat[(RegisterIdx * NiagaraFloatDataStride + InstanceID)];
-}
-
-#else
 
 float GetFloat(int RegisterIdx, uint InstanceID)
 {
-	return NiagaraVFLooseParameters.NiagaraParticleDataFloat[(RegisterIdx * NiagaraVFLooseParameters.NiagaraFloatDataStride + InstanceID)];
+	bool IsHalf = RegisterIdx & (1 << 31);
+	RegisterIdx &= (~(1 << 31));
+
+	if(IsHalf)
+	{
+		return NiagaraVFLooseParameters.NiagaraParticleDataHalf[(RegisterIdx * NiagaraVFLooseParameters.NiagaraFloatDataStride + InstanceID)];
+	}
+	else
+	{
+		return NiagaraVFLooseParameters.NiagaraParticleDataFloat[(RegisterIdx * NiagaraVFLooseParameters.NiagaraFloatDataStride + InstanceID)];
+	}
 }
 
-#endif
-
 float2 GetVec2(int RegisterIndex, uint InstanceID)
 {
 	return float2(GetFloat(RegisterIndex, InstanceID), GetFloat(RegisterIndex+1, InstanceID));
diff --git a/Engine/Plugins/FX/Niagara/Source/Niagara/Classes/NiagaraDataInterfaceParticleRead.h b/Engine/Plugins/FX/Niagara/Source/Niagara/Classes/NiagaraDataInterfaceParticleRead.h
--- a/Engine/Plugins/FX/Niagara/Source/Niagara/Classes/NiagaraDataInterfaceParticleRead.h
+++ b/Engine/Plugins/FX/Niagara/Source/Niagara/Classes/NiagaraDataInterfaceParticleRead.h
@@ -51,6 +51,7 @@ protected:
 	virtual bool CopyToInternal(UNiagaraDataInterface* Destination) const override;
 
 private:
-	template<typename T>
-	T RetrieveValueWithCheck(FNiagaraEmitterInstance* EmitterInstance, const FNiagaraTypeDefinition& Type, const FName& Attr, const FNiagaraID& ParticleID, bool& bValid);
+	template<typename T, typename AccessorType>
+	T RetrieveValueWithCheck(FNiagaraEmitterInstance* EmitterInstance, const FName& Attr, const FNiagaraID& ParticleID, bool& bValid);
+
 };
\ No newline at end of file
diff --git a/Engine/Plugins/FX/Niagara/Source/Niagara/Classes/NiagaraDataSet.h b/Engine/Plugins/FX/Niagara/Source/Niagara/Classes/NiagaraDataSet.h
--- a/Engine/Plugins/FX/Niagara/Source/Niagara/Classes/NiagaraDataSet.h
+++ b/Engine/Plugins/FX/Niagara/Source/Niagara/Classes/NiagaraDataSet.h
@@ -23,8 +23,13 @@ struct FNiagaraVariableLayoutInfo
 	UPROPERTY()
 	uint32 Int32ComponentStart;
 
+	/** Start index for the half components in the main buffer. */
+	UPROPERTY()
+	uint32 HalfComponentStart;
+
 	uint32 GetNumFloatComponents()const { return LayoutInfo.FloatComponentByteOffsets.Num(); }
 	uint32 GetNumInt32Components()const { return LayoutInfo.Int32ComponentByteOffsets.Num(); }
+	uint32 GetNumHalfComponents()const { return LayoutInfo.HalfComponentByteOffsets.Num(); }
 
 	/** This variable's type layout info. */
 	UPROPERTY()
@@ -107,28 +112,40 @@ public:
 	void SwapInstances(uint32 OldIndex, uint32 NewIndex);
 	void KillInstance(uint32 InstanceIdx);
 	void CopyTo(FNiagaraDataBuffer& DestBuffer, int32 SrcStartIdx, int32 DestStartIdx, int32 NumInstances)const;
-	void GPUCopyFrom(float* GPUReadBackFloat, int* GPUReadBackInt, int32 StartIdx, int32 NumInstances, uint32 InSrcFloatStride, uint32 InSrcIntStride);
+	void GPUCopyFrom(const float* GPUReadBackFloat, const int* GPUReadBackInt, const FFloat16* GPUReadBackHalf, int32 StartIdx, int32 NumInstances, uint32 InSrcFloatStride, uint32 InSrcIntStride, uint32 InSrcHalfStride);
 	void Dump(int32 StartIndex, int32 NumInstances, const FString& Label)const;
 
-	FORCEINLINE TArray<uint8*>& GetRegisterTable() { return RegisterTable; }
+	FORCEINLINE TArrayView<uint8 const* RESTRICT const> GetRegisterTable() const { return TArrayView<uint8 const* RESTRICT const>(RegisterTable); }
+	
+	typedef uint32 RegisterTypeOffsetType[3];
+	FORCEINLINE const RegisterTypeOffsetType& GetRegisterTypeOffsets() { return RegisterTypeOffsets; }
 
-	FORCEINLINE const TArray<uint8>& GetFloatBuffer()const { return FloatData; }
-	FORCEINLINE const TArray<uint8>& GetInt32Buffer()const { return Int32Data; }
+	FORCEINLINE const TArray<uint8>& GetFloatBuffer() const { return FloatData; }
+	FORCEINLINE const TArray<uint8>& GetInt32Buffer() const { return Int32Data; }
+	FORCEINLINE const TArray<uint8>& GetHalfBuffer() const { return HalfData; }
 
-	FORCEINLINE const uint8* GetComponentPtrFloat(uint32 ComponentIdx)const	{ return FloatData.GetData() + FloatStride * ComponentIdx; }
-	FORCEINLINE const uint8* GetComponentPtrInt32(uint32 ComponentIdx)const	{ return Int32Data.GetData() + Int32Stride * ComponentIdx; }
-	FORCEINLINE uint8* GetComponentPtrFloat(uint32 ComponentIdx) { return FloatData.GetData() + FloatStride * ComponentIdx;	}
+	FORCEINLINE const uint8* GetComponentPtrFloat(uint32 ComponentIdx) const { return FloatData.GetData() + FloatStride * ComponentIdx; }
+	FORCEINLINE const uint8* GetComponentPtrInt32(uint32 ComponentIdx) const { return Int32Data.GetData() + Int32Stride * ComponentIdx; }
+	FORCEINLINE const uint8* GetComponentPtrHalf(uint32 ComponentIdx) const { return HalfData.GetData() + HalfStride * ComponentIdx; }
+	FORCEINLINE uint8* GetComponentPtrFloat(uint32 ComponentIdx) { return FloatData.GetData() + FloatStride * ComponentIdx; }
 	FORCEINLINE uint8* GetComponentPtrInt32(uint32 ComponentIdx) { return Int32Data.GetData() + Int32Stride * ComponentIdx;	}
+	FORCEINLINE uint8* GetComponentPtrHalf(uint32 ComponentIdx) { return HalfData.GetData() + HalfStride * ComponentIdx; }
 
-	FORCEINLINE float* GetInstancePtrFloat(uint32 ComponentIdx, uint32 InstanceIdx)	{ return (float*)GetComponentPtrFloat(ComponentIdx) + InstanceIdx; }
-	FORCEINLINE int32* GetInstancePtrInt32(uint32 ComponentIdx, uint32 InstanceIdx)	{ return (int32*)GetComponentPtrInt32(ComponentIdx) + InstanceIdx; }
-	FORCEINLINE float* GetInstancePtrFloat(uint32 ComponentIdx, uint32 InstanceIdx)const { return (float*)GetComponentPtrFloat(ComponentIdx) + InstanceIdx; }
-	FORCEINLINE int32* GetInstancePtrInt32(uint32 ComponentIdx, uint32 InstanceIdx)const { return (int32*)GetComponentPtrInt32(ComponentIdx) + InstanceIdx;	}
+	FORCEINLINE float* GetInstancePtrFloat(uint32 ComponentIdx, uint32 InstanceIdx)	{ return (float*)(GetComponentPtrFloat(ComponentIdx)) + InstanceIdx; }
+	FORCEINLINE int32* GetInstancePtrInt32(uint32 ComponentIdx, uint32 InstanceIdx)	{ return (int32*)(GetComponentPtrInt32(ComponentIdx)) + InstanceIdx; }
+	FORCEINLINE FFloat16* GetInstancePtrHalf(uint32 ComponentIdx, uint32 InstanceIdx) { return (FFloat16*)(GetComponentPtrHalf(ComponentIdx)) + InstanceIdx; }
+
+	FORCEINLINE const float* GetInstancePtrFloat(uint32 ComponentIdx, uint32 InstanceIdx) const { return (float*)(GetComponentPtrFloat(ComponentIdx)) + InstanceIdx; }
+	FORCEINLINE const int32* GetInstancePtrInt32(uint32 ComponentIdx, uint32 InstanceIdx) const { return (int32*)(GetComponentPtrInt32(ComponentIdx)) + InstanceIdx; }
+	FORCEINLINE const FFloat16* GetInstancePtrHalf(uint32 ComponentIdx, uint32 InstanceIdx) const { return (FFloat16*)(GetComponentPtrHalf(ComponentIdx)) + InstanceIdx; }
 
 	FORCEINLINE uint8* GetComponentPtrFloat(float* BasePtr, uint32 ComponentIdx) const { return (uint8*)BasePtr + FloatStride * ComponentIdx; }
 	FORCEINLINE uint8* GetComponentPtrInt32(int* BasePtr, uint32 ComponentIdx) const { return (uint8*)BasePtr + Int32Stride * ComponentIdx; }
+	FORCEINLINE uint8* GetComponentPtrHalf(FFloat16* BasePtr, uint32 ComponentIdx) const { return (uint8*)(BasePtr) + HalfStride * ComponentIdx; }
+
 	FORCEINLINE float* GetInstancePtrFloat(float* BasePtr, uint32 ComponentIdx, uint32 InstanceIdx)const { return (float*)GetComponentPtrFloat(BasePtr, ComponentIdx) + InstanceIdx; }
 	FORCEINLINE int32* GetInstancePtrInt32(int* BasePtr, uint32 ComponentIdx, uint32 InstanceIdx)const { return (int32*)GetComponentPtrInt32(BasePtr, ComponentIdx) + InstanceIdx; }
+	FORCEINLINE FFloat16* GetInstancePtrHalf(FFloat16 *BasePtr, uint32 ComponentIdx, uint32 InstanceIdx)const { return (FFloat16*)GetComponentPtrHalf(BasePtr, ComponentIdx) + InstanceIdx; }
 
 	FORCEINLINE uint32 GetNumInstances()const { return NumInstances; }
 	FORCEINLINE uint32 GetNumInstancesAllocated()const { return NumInstancesAllocated; }
@@ -138,7 +155,8 @@ public:
 	FORCEINLINE void SetNumSpawnedInstances(uint32 InNumSpawnedInstances) { NumSpawnedInstances = InNumSpawnedInstances; }
 	FORCEINLINE uint32 GetSizeBytes()const { return FloatData.Num() + Int32Data.Num(); }
 	FORCEINLINE FRWBuffer& GetGPUBufferFloat() { return GPUBufferFloat; }
-	FORCEINLINE FRWBuffer& GetGPUBufferInt() { return GPUBufferInt;	}
+	FORCEINLINE FRWBuffer& GetGPUBufferInt() { return GPUBufferInt; }
+	FORCEINLINE FRWBuffer& GetGPUBufferHalf() { return GPUBufferHalf; }
 	FORCEINLINE uint32 GetGPUInstanceCountBufferOffset() const { return GPUInstanceCountBufferOffset; }
 	FORCEINLINE void ClearGPUInstanceCountBufferOffset() { GPUInstanceCountBufferOffset = INDEX_NONE; }
 	FORCEINLINE FRWBuffer& GetGPUIDToIndexTable() { return GPUIDToIndexTable; }
@@ -146,6 +164,7 @@ public:
 	FORCEINLINE int32 GetSafeComponentBufferSize() const { return GetSafeComponentBufferSize(GetNumInstancesAllocated()); }
 	FORCEINLINE uint32 GetFloatStride() const { return FloatStride; }
 	FORCEINLINE uint32 GetInt32Stride() const { return Int32Stride; }
+	FORCEINLINE uint32 GetHalfStride() const { return HalfStride; }
 
 	FORCEINLINE uint32 GetIDAcquireTag() const { return IDAcquireTag; }
 	FORCEINLINE void SetIDAcquireTag(uint32 InTag) { IDAcquireTag = InTag; }
@@ -164,8 +183,6 @@ public:
 
 	void ReleaseGPUInstanceCount(FNiagaraGPUInstanceCountManager& GPUInstanceCountManager);
 
-	FORCEINLINE const TArray<uint8*>& GetRegisterTable()const { return RegisterTable; }
-
 	void BuildRegisterTable();
 private:
 
@@ -187,6 +204,8 @@ private:
 	TArray<uint8> FloatData;
 	/** Int32 components of simulation data. */
 	TArray<uint8> Int32Data;
+	/** Half components of simulation data. */
+	TArray<uint8> HalfData;
 
 	/** Table of IDs to real buffer indices. */
 	TArray<int32> IDToIndexTable;
@@ -200,10 +219,12 @@ private:
 	uint32 NumChunksAllocatedForGPU;
 	/** GPU Buffer containing floating point values for GPU simulations. */
 	FRWBuffer GPUBufferFloat;
-	/** GPU Buffer containing floating point values for GPU simulations. */
+	/** GPU Buffer containing integer values for GPU simulations. */
 	FRWBuffer GPUBufferInt;
 	/** GPU table which maps particle ID to index. */
 	FRWBuffer GPUIDToIndexTable;
+	/** GPU Buffer containing half values for GPU simulations. */
+	FRWBuffer GPUBufferHalf;
 	//////////////////////////////////////////////////////////////////////////
 
 	/** Number of instances in data. */
@@ -214,6 +235,8 @@ private:
 	uint32 FloatStride;
 	/** Stride between components in the int32 buffer. */
 	uint32 Int32Stride;
+	/** Stride between components in the half buffer. */
+	uint32 HalfStride;
 	/** Number of instances spawned in the last tick. */
 	uint32 NumSpawnedInstances;
 	/** ID acquire tag used in the last tick. */
@@ -221,6 +244,7 @@ private:
 
 	/** Table containing current base locations for all registers in this dataset. */
 	TArray<uint8*> RegisterTable;//TODO: Should make inline? Feels like a useful size to keep local would be too big.
+	RegisterTypeOffsetType RegisterTypeOffsets;
 };
 
 //////////////////////////////////////////////////////////////////////////
@@ -249,6 +273,9 @@ struct NIAGARA_API FNiagaraDataSetCompiledData
 	UPROPERTY()
 	uint32 TotalInt32Components;
 
+	UPROPERTY()
+	uint32 TotalHalfComponents;
+
 	/** Whether or not this dataset require persistent IDs. */
 	UPROPERTY()
 	uint32 bRequiresPersistentIDs : 1;
@@ -315,19 +342,21 @@ public:
 	FORCEINLINE FRWBuffer& GetGPUFreeIDs() { return GPUFreeIDs; }
 	FORCEINLINE uint32 GetGPUNumAllocatedIDs() const { return GPUNumAllocatedIDs; }
 
-	FORCEINLINE const TArray<FNiagaraVariable>& GetVariables() const { return CompiledData->Variables; }
-	FORCEINLINE uint32 GetNumVariables() const { return CompiledData->Variables.Num(); }
-	FORCEINLINE bool HasVariable(const FNiagaraVariable& Var) const { return CompiledData->Variables.Contains(Var); }
-	FORCEINLINE uint32 GetNumFloatComponents() const { return CompiledData->TotalFloatComponents; }
-	FORCEINLINE uint32 GetNumInt32Components() const { return CompiledData->TotalInt32Components; }
+	FORCEINLINE const TArray<FNiagaraVariable>& GetVariables()const { return CompiledData->Variables; }
+	FORCEINLINE uint32 GetNumVariables()const { return CompiledData->Variables.Num(); }
+	FORCEINLINE bool HasVariable(const FNiagaraVariable& Var)const { return CompiledData->Variables.Contains(Var); }
+	FORCEINLINE bool HasVariable(const FName& Name)const { return CompiledData->Variables.FindByPredicate([&](const FNiagaraVariable& VarInfo) { return VarInfo.GetName() == Name; }) != nullptr; }
+	FORCEINLINE uint32 GetNumFloatComponents()const { return CompiledData->TotalFloatComponents; }
+	FORCEINLINE uint32 GetNumInt32Components()const { return CompiledData->TotalInt32Components; }
+	FORCEINLINE uint32 GetNumHalfComponents()const { return CompiledData->TotalHalfComponents; }
 
 	const TArray<FNiagaraVariableLayoutInfo>& GetVariableLayouts() const { return CompiledData->VariableLayouts; }
 	const FNiagaraVariableLayoutInfo* GetVariableLayout(const FNiagaraVariable& Var) const;
-	bool GetVariableComponentOffsets(const FNiagaraVariable& Var, int32 &FloatStart, int32 &IntStart) const;
+	bool GetVariableComponentOffsets(const FNiagaraVariable& Var, int32 &FloatStart, int32 &IntStart, int32& HalfStart) const;
 
 	void CopyTo(FNiagaraDataSet& Other, int32 StartIdx = 0, int32 NumInstances = INDEX_NONE, bool bResetOther=true)const;
 
-	void CopyFromGPUReadback(float* GPUReadBackFloat, int* GPUReadBackInt, int32 StartIdx = 0, int32 NumInstances = INDEX_NONE, uint32 FloatStride = 0, uint32 IntStride = 0);
+	void CopyFromGPUReadback(const float* GPUReadBackFloat, const int* GPUReadBackInt, const FFloat16* GPUReadBackHalf, int32 StartIdx = 0, int32 NumInstances = INDEX_NONE, uint32 FloatStride = 0, uint32 IntStride = 0, uint32 HalfStride = 0);
 
 	void CheckForNaNs() const;
 
@@ -429,6 +458,7 @@ struct FNiagaraDataSetAccessorBase
 	FNiagaraDataSetAccessorBase(FNiagaraDataSet* InDataSet, FNiagaraVariable InVar)
 		: DataSet(InDataSet)
 	{
+		checkSlow(DataSet->HasVariable(InVar));
 		Var = InVar;
 		VarLayout = DataSet->GetVariableLayout(InVar);
 	}
@@ -463,78 +493,7 @@ protected:
 };
 
 template<typename T>
-struct FNiagaraDataSetAccessor : public FNiagaraDataSetAccessorBase
-{
-	FNiagaraDataSetAccessor<T>() {}
-	FNiagaraDataSetAccessor(FNiagaraDataSet& InDataSet, FNiagaraVariable InVar)
-		: FNiagaraDataSetAccessorBase(&InDataSet, InVar)
-	{
-		check(sizeof(T) == InVar.GetType().GetSize());
-		checkf(false, TEXT("You must provide a fast runtime specialization for this type."));// Allow this slow generic version?
-	}
-
-	FORCEINLINE T operator[](int32 Index)const
-	{
-		return Get(Index);
-	}
-
-	FORCEINLINE T Get(int32 Index)const
-	{
-		T Ret;
-		Get(Index, Ret);
-		return Ret;
-	}
-
-	FORCEINLINE void Get(int32 Index, T& OutValue)const
-	{
-		checkSlow(DataSet);
-		uint8* ValuePtr = (uint8*)&OutValue;
-
-		FNiagaraDataBuffer* DataBuffer = DataSet->GetCurrentData();
-		checkSlow(DataBuffer);
-
-		for (uint32 CompIdx = 0; CompIdx < VarLayout->GetNumFloatComponents(); ++CompIdx)
-		{
-			uint32 CompBufferOffset = VarLayout->FloatComponentStart + CompIdx;
-			float* Src = DataBuffer->GetInstancePtrFloat(CompBufferOffset, Index);
-			float* Dst = (float*)(ValuePtr + VarLayout->LayoutInfo.FloatComponentByteOffsets[CompIdx]);
-			*Dst = *Src;
-		}
-
-		for (uint32 CompIdx = 0; CompIdx < VarLayout->GetNumInt32Components(); ++CompIdx)
-		{
-			uint32 CompBufferOffset = VarLayout->Int32ComponentStart + CompIdx;
-			int32* Src = DataBuffer->GetInstancePtrInt32(CompBufferOffset, Index);
-			int32* Dst = (int32*)(ValuePtr + VarLayout->LayoutInfo.Int32ComponentByteOffsets[CompIdx]);
-			*Dst = *Src;
-		}
-	}
-
-	FORCEINLINE void Set(int32 Index, const T& InValue)
-	{
-		checkSlow(DataSet);
-		uint8* ValuePtr = (uint8*)&InValue;
-
-		FNiagaraDataBuffer* DataBuffer = DataSet->GetDestinationData();
-		checkSlow(DataBuffer);
-
-		for (uint32 CompIdx = 0; CompIdx < VarLayout->GetNumFloatComponents(); ++CompIdx)
-		{
-			uint32 CompBufferOffset = VarLayout->FloatComponentStart + CompIdx;
-			float* Dst = DataBuffer->GetInstancePtrFloat(CompBufferOffset, Index);
-			float* Src = (float*)(ValuePtr + VarLayout->LayoutInfo.FloatComponentByteOffsets[CompIdx]);
-			*Dst = *Src;
-		}
-
-		for (uint32 CompIdx = 0; CompIdx < VarLayout->GetNumInt32Components(); ++CompIdx)
-		{
-			uint32 CompBufferOffset = VarLayout->Int32ComponentStart + CompIdx;
-			int32* Dst = DataBuffer->GetInstancePtrInt32(CompBufferOffset, Index);
-			int32* Src = (int32*)(ValuePtr + VarLayout->LayoutInfo.Int32ComponentByteOffsets[CompIdx]);
-			*Dst = *Src;
-		}
-	}
-};
+struct FNiagaraDataSetAccessor;
 
 template<>
 struct FNiagaraDataSetAccessor<FNiagaraBool> : public FNiagaraDataSetAccessorBase
@@ -604,8 +563,8 @@ struct FNiagaraDataSetAccessor<FNiagaraBool> : public FNiagaraDataSetAccessorBas
 
 private:
 
-	int32* SrcBase;
-	int32* DestBase;
+	int32* SrcBase = nullptr;
+	int32* DestBase = nullptr;
 };
 
 template<>
@@ -679,7 +638,7 @@ struct FNiagaraDataSetAccessor<int32> : public FNiagaraDataSetAccessorBase
 		OutValue = SrcBase[Index];
 	}
 
-	FORCEINLINE void Set(int32 Index, const int32& InValue)
+	FORCEINLINE void Set(int32 Index, int32 InValue)
 	{
 		checkSlow(IsValidForWrite());
 		DestBase[Index] = InValue;
@@ -687,8 +646,14 @@ struct FNiagaraDataSetAccessor<int32> : public FNiagaraDataSetAccessorBase
 
 private:
 
-	int32* SrcBase;
-	int32* DestBase;
+	int32* SrcBase = nullptr;
+	int32* DestBase = nullptr;
+};
+
+struct FNiagaraMinMaxStruct
+{
+	float Min;
+	float Max;
 };
 
 template<>
@@ -759,17 +724,41 @@ struct FNiagaraDataSetAccessor<float> : public FNiagaraDataSetAccessorBase
 		OutValue = SrcBase[Index];
 	}
 
-	FORCEINLINE void Set(int32 Index, const float& InValue)
+	FORCEINLINE void Set(int32 Index, float InValue)
 	{
 		checkSlow(IsValidForWrite());
 		DestBase[Index] = InValue;
 	}
 
-	FORCEINLINE const float* const GetX() const { return SrcBase; }
+	FORCEINLINE float GetMaxElement(uint32 ElementIndex) const 
+	{ 
+		checkSlow(ElementIndex == 0);
+		float MaxValue = FLT_MIN;
+		const int NumInstances = DataSet->GetCurrentData()->GetNumInstances();
+		for (int Index = 0; Index < NumInstances; Index++)
+		{
+			MaxValue = FMath::Max(MaxValue, SrcBase[Index]);
+		}
+		return MaxValue;
+	}
+
+	FORCEINLINE FNiagaraMinMaxStruct GetMinMaxElement(uint32 ElementIndex) const
+	{
+		checkSlow(ElementIndex == 0);
+		float MinValue = FLT_MAX;
+		float MaxValue = FLT_MIN;
+		const int NumInstances = DataSet->GetCurrentData()->GetNumInstances();
+		for (int Index = 0; Index < NumInstances; Index++)
+		{
+			MaxValue = FMath::Max(MaxValue, SrcBase[Index]);
+			MinValue = FMath::Min(MinValue, SrcBase[Index]);
+		}
+		return FNiagaraMinMaxStruct{ MinValue , MaxValue };
+	}
 
 private:
-	float* SrcBase;
-	float* DestBase;
+	float* SrcBase = nullptr;
+	float* DestBase = nullptr;
 };
 
 template<>
@@ -852,15 +841,40 @@ struct FNiagaraDataSetAccessor<FVector2D> : public FNiagaraDataSetAccessorBase
 		DestYBase[Index] = InValue.Y;
 	}
 
-	FORCEINLINE const float* const GetX() const { return SrcXBase; }
-	FORCEINLINE const float* const GetY() const { return SrcYBase; }
+	FORCEINLINE float GetMaxElement(uint32 ElementIndex) const
+	{
+		checkSlow(ElementIndex == 0 || ElementIndex == 1);
+		float* SrcBase = ElementIndex == 0 ? SrcXBase : SrcYBase;
+		float MaxValue = FLT_MIN;
+		const int NumInstances = DataSet->GetCurrentData()->GetNumInstances();
+		for (int Index = 0; Index < NumInstances; Index++)
+		{
+			MaxValue = FMath::Max(MaxValue, SrcBase[Index]);
+		}
+		return MaxValue;
+	}
+
+	FORCEINLINE FNiagaraMinMaxStruct GetMinMaxElement(uint32 ElementIndex) const
+	{
+		checkSlow(ElementIndex == 0 || ElementIndex == 1);
+		float* SrcBase = ElementIndex == 0 ? SrcXBase : SrcYBase;
+		float MinValue = FLT_MAX;
+		float MaxValue = FLT_MIN;
+		const int NumInstances = DataSet->GetCurrentData()->GetNumInstances();
+		for (int Index = 0; Index < NumInstances; Index++)
+		{
+			MaxValue = FMath::Max(MaxValue, SrcBase[Index]);
+			MinValue = FMath::Min(MinValue, SrcBase[Index]);
+		}
+		return FNiagaraMinMaxStruct{ MinValue , MaxValue };
+	}
 
 private:
 
-	float* SrcXBase;
-	float* SrcYBase;
-	float* DestXBase;
-	float* DestYBase;
+	float* SrcXBase = nullptr;
+	float* SrcYBase = nullptr;
+	float* DestXBase = nullptr;
+	float* DestYBase = nullptr;
 };
 
 template<>
@@ -949,18 +963,42 @@ struct FNiagaraDataSetAccessor<FVector> : public FNiagaraDataSetAccessorBase
 		DestZBase[Index] = InValue.Z;
 	}
 
-	FORCEINLINE const float* const GetX() const { return SrcXBase; }
-	FORCEINLINE const float* const GetY() const { return SrcYBase; }
-	FORCEINLINE const float* const GetZ() const { return SrcZBase; }
+	FORCEINLINE float GetMaxElement(uint32 ElementIndex) const
+	{
+		checkSlow(ElementIndex == 0 || ElementIndex == 1 || ElementIndex == 2);
+		float* SrcBase = ElementIndex == 0 ? SrcXBase : (ElementIndex == 1 ? SrcYBase : SrcZBase);
+		float MaxValue = FLT_MIN;
+		const int NumInstances = DataSet->GetCurrentData()->GetNumInstances();
+		for (int Index = 0; Index < NumInstances; Index++)
+		{
+			MaxValue = FMath::Max(MaxValue, SrcBase[Index]);
+		}
+		return MaxValue;
+	}
+
+	FORCEINLINE FNiagaraMinMaxStruct GetMinMaxElement(uint32 ElementIndex) const
+	{
+		checkSlow(ElementIndex == 0 || ElementIndex == 1 || ElementIndex == 2);
+		float* SrcBase = ElementIndex == 0 ? SrcXBase : (ElementIndex == 1 ? SrcYBase : SrcZBase);
+		float MinValue = FLT_MAX;
+		float MaxValue = FLT_MIN;
+		const int NumInstances = DataSet->GetCurrentData()->GetNumInstances();
+		for (int Index = 0; Index < NumInstances; Index++)
+		{
+			MaxValue = FMath::Max(MaxValue, SrcBase[Index]);
+			MinValue = FMath::Min(MinValue, SrcBase[Index]);
+		}
+		return FNiagaraMinMaxStruct{ MinValue , MaxValue };
+	}
 
 private:
 
-	float* SrcXBase;
-	float* SrcYBase;
-	float* SrcZBase;
-	float* DestXBase;
-	float* DestYBase;
-	float* DestZBase;
+	float* SrcXBase = nullptr;
+	float* SrcYBase = nullptr;
+	float* SrcZBase = nullptr;
+	float* DestXBase = nullptr;
+	float* DestYBase = nullptr;
+	float* DestZBase = nullptr;
 };
 
 template<>
@@ -1057,14 +1095,14 @@ struct FNiagaraDataSetAccessor<FVector4> : public FNiagaraDataSetAccessorBase
 
 private:
 
-	float* SrcXBase;
-	float* SrcYBase;
-	float* SrcZBase;
-	float* SrcWBase;
-	float* DestXBase;
-	float* DestYBase;
-	float* DestZBase;
-	float* DestWBase;
+	float* SrcXBase = nullptr;
+	float* SrcYBase = nullptr;
+	float* SrcZBase = nullptr;
+	float* SrcWBase = nullptr;
+	float* DestXBase = nullptr;
+	float* DestYBase = nullptr;
+	float* DestZBase = nullptr;
+	float* DestWBase = nullptr;
 };
 
 
@@ -1162,14 +1200,14 @@ struct FNiagaraDataSetAccessor<FQuat> : public FNiagaraDataSetAccessorBase
 
 private:
 
-	float* SrcXBase;
-	float* SrcYBase;
-	float* SrcZBase;
-	float* SrcWBase;
-	float* DestXBase;
-	float* DestYBase;
-	float* DestZBase;
-	float* DestWBase;
+	float* SrcXBase = nullptr;
+	float* SrcYBase = nullptr;
+	float* SrcZBase = nullptr;
+	float* SrcWBase = nullptr;
+	float* DestXBase = nullptr;
+	float* DestYBase = nullptr;
+	float* DestZBase = nullptr;
+	float* DestWBase = nullptr;
 };
 
 template<>
@@ -1266,14 +1304,14 @@ struct FNiagaraDataSetAccessor<FLinearColor> : public FNiagaraDataSetAccessorBas
 
 private:
 
-	float* SrcRBase;
-	float* SrcGBase;
-	float* SrcBBase;
-	float* SrcABase;
-	float* DestRBase;
-	float* DestGBase;
-	float* DestBBase;
-	float* DestABase;
+	float* SrcRBase = nullptr;
+	float* SrcGBase = nullptr;
+	float* SrcBBase = nullptr;
+	float* SrcABase = nullptr;
+	float* DestRBase = nullptr;
+	float* DestGBase = nullptr;
+	float* DestBBase = nullptr;
+	float* DestABase = nullptr;
 };
 
 template<>
@@ -1380,15 +1418,15 @@ struct FNiagaraDataSetAccessor<FNiagaraSpawnInfo> : public FNiagaraDataSetAccess
 
 private:
 
-	int32* SrcCountBase;
-	float* SrcInterpStartDtBase;
-	float* SrcIntervalDtBase;
-	int32* SrcGroupBase;
+	int32* SrcCountBase = nullptr;
+	float* SrcInterpStartDtBase = nullptr;
+	float* SrcIntervalDtBase = nullptr;
+	int32* SrcGroupBase = nullptr;
 
-	int32* DestCountBase;
-	float* DestInterpStartDtBase;
-	float* DestIntervalDtBase;
-	int32* DestGroupBase;
+	int32* DestCountBase = nullptr;
+	float* DestInterpStartDtBase = nullptr;
+	float* DestIntervalDtBase = nullptr;
+	int32* DestGroupBase = nullptr;
 };
 
 template<>
@@ -1472,12 +1510,558 @@ struct FNiagaraDataSetAccessor<FNiagaraID> : public FNiagaraDataSetAccessorBase
 
 private:
 
-	int32* SrcIndexBase;
-	int32* SrcTagBase;
-	int32* DestIndexBase;
-	int32* DestTagBase;
+	int32* SrcIndexBase = nullptr;
+	int32* SrcTagBase = nullptr;
+	int32* DestIndexBase = nullptr;
+	int32* DestTagBase = nullptr;
+};
+
+template<int ElementCount>
+struct FNiagaraHalfDataSetAccessor : public FNiagaraDataSetAccessorBase
+{
+	using DataType = TStaticArray<FFloat16, ElementCount>;
+
+	FNiagaraHalfDataSetAccessor() {}
+	FNiagaraHalfDataSetAccessor(FNiagaraDataSet& InDataSet, FNiagaraVariable InVar)
+		: FNiagaraDataSetAccessorBase(&InDataSet, InVar)
+	{
+		check(sizeof(DataType) == InVar.GetType().GetSize());
+		InitForAccess();
+	}
+
+	void InitForAccess()
+	{
+		for (int32 ElementIt = 0; ElementIt < ElementCount; ++ElementIt)
+		{
+			SrcBase[ElementIt] = nullptr;
+			DestBase[ElementIt] = nullptr;
+		}
+		const FNiagaraDataBuffer* SrcBuffer = DataSet->GetCurrentData();
+		if (IsValid() && SrcBuffer)
+		{
+			for (int32 ElementIt = 0; ElementIt < ElementCount; ++ElementIt)
+			{
+				SrcBase[ElementIt] = reinterpret_cast<const FFloat16*>(SrcBuffer->GetComponentPtrHalf(VarLayout->HalfComponentStart + ElementIt));
+
+				//Writes are only valid if we're during a simulation pass.
+				FNiagaraDataBuffer* DestBuffer = DataSet->GetDestinationData();
+				if (DestBuffer)
+				{
+					DestBase[ElementIt] = reinterpret_cast<FFloat16*>(DestBuffer->GetComponentPtrHalf(VarLayout->HalfComponentStart + ElementIt));
+				}
+			}
+		}
+	}
+
+	FORCEINLINE bool IsValidForRead() const
+	{
+		for (int32 ElementIt = 0; ElementIt < ElementCount; ++ElementIt)
+		{
+			if (!SrcBase[ElementIt])
+			{
+				return false;
+			}
+		}
+		return true;
+	}
+	FORCEINLINE bool IsValidForWrite() const
+	{
+		for (int32 ElementIt = 0; ElementIt < ElementCount; ++ElementIt)
+		{
+			if (!DestBase[ElementIt])
+			{
+				return false;
+			}
+		}
+		return true;
+	}
+
+	FORCEINLINE DataType operator[](int32 Index)const
+	{
+		return Get(Index);
+	}
+
+	FORCEINLINE DataType GetSafe(int32 Index, DataType Default = DataType(0.0f))const
+	{
+		if (IsValidForRead())
+		{
+			FNiagaraDataBuffer* DataBuffer = DataSet->GetCurrentData();
+			if (DataBuffer && Index >= 0 && (uint32)Index < DataBuffer->GetNumInstances())
+			{
+				return Get(Index);
+			}
+			else
+			{
+				ensure(DataBuffer && Index >= 0 && (uint32)Index < DataBuffer->GetNumInstances()); // just to capture the badness in the logs or debugger if attached.
+			}
+		}
+
+		return Default;
+	}
+
+	FORCEINLINE DataType Get(int32 Index)const
+	{
+		DataType Ret;
+		Get(Index, Ret);
+		return Ret;
+	}
+
+	FORCEINLINE void Get(int32 Index, DataType& OutValue)const
+	{
+		checkSlow(IsValidForRead());
+		for (int32 ElementIt = 0; ElementIt < ElementCount; ++ElementIt)
+		{
+			OutValue[ElementIt] = SrcBase[ElementIt][Index];
+		}
+	}
+
+	FORCEINLINE void Set(int32 Index, DataType InValue)
+	{
+		checkSlow(IsValidForWrite());
+		for (int32 ElementIt = 0; ElementIt < ElementCount; ++ElementIt)
+		{
+			DestBase[ElementIt][Index] = InValue;
+		}
+	}
+
+	FORCEINLINE float GetMaxElement(uint32 ElementIndex) const
+	{
+		checkSlow(ElementIndex < ElementCount);
+		int16 MaxValue = static_cast<int16>(0xfc00u);
+		const int NumInstances = DataSet->GetCurrentData()->GetNumInstances();
+		for (int Index = 0; Index < NumInstances; Index++)
+		{
+			MaxValue = FMath::Max(MaxValue, ((int16*)SrcBase[ElementIndex])[Index]);
+		}
+		float FloatMaxValue = FPlatformMath::LoadHalf((uint16*)&MaxValue);
+		return FloatMaxValue;
+	}
+
+	FORCEINLINE FNiagaraMinMaxStruct GetMinMaxElement(uint32 ElementIndex) const
+	{
+		checkSlow(ElementIndex < ElementCount);
+		int16 MinValue = static_cast<int16>(0x7c00u);
+		int16 MaxValue = static_cast<int16>(0xfc00u);
+		const int NumInstances = DataSet->GetCurrentData()->GetNumInstances();
+		for (int Index = 0; Index < NumInstances; Index++)
+		{
+			MaxValue = FMath::Max(MaxValue, ((int16*)SrcBase[ElementIndex])[Index]);
+			MinValue = FMath::Min(MinValue, ((int16*)SrcBase[ElementIndex])[Index]);
+		}
+		float FloatMinValue = FPlatformMath::LoadHalf((uint16*)&MinValue);
+		float FloatMaxValue = FPlatformMath::LoadHalf((uint16*)&MaxValue);
+		return FNiagaraMinMaxStruct{ FloatMinValue , FloatMaxValue };
+	}
+private:
+	const FFloat16* SrcBase[ElementCount] = {};
+	FFloat16* DestBase[ElementCount] = {};
+};
+
+template<>
+struct FNiagaraDataSetAccessor<TStaticArray<FFloat16, 1>> : public FNiagaraHalfDataSetAccessor<1>
+{
+	using FNiagaraHalfDataSetAccessor::FNiagaraHalfDataSetAccessor;
+};
+
+template<>
+struct FNiagaraDataSetAccessor<TStaticArray<FFloat16, 2>> : public FNiagaraHalfDataSetAccessor<2>
+{
+	using FNiagaraHalfDataSetAccessor::FNiagaraHalfDataSetAccessor;
+};
+
+template<>
+struct FNiagaraDataSetAccessor<TStaticArray<FFloat16, 3>> : public FNiagaraHalfDataSetAccessor<3>
+{
+	using FNiagaraHalfDataSetAccessor::FNiagaraHalfDataSetAccessor;
+};
+
+template<>
+struct FNiagaraDataSetAccessor<TStaticArray<FFloat16, 4>> : public FNiagaraHalfDataSetAccessor<4>
+{
+	using FNiagaraHalfDataSetAccessor::FNiagaraHalfDataSetAccessor;
+};
+
+template<typename TagType>
+struct FNiagaraEitherDataSetAccessor
+{
+	typedef typename TagType::EitherA EitherA;
+	typedef typename TagType::EitherB EitherB;
+
+	FNiagaraEitherDataSetAccessor() {}
+	FNiagaraEitherDataSetAccessor(FNiagaraDataSet& InDataSet, const FName& InName)
+	{
+		if (InDataSet.HasVariable(FNiagaraVariable(TagType::GetTypeA(), InName)))
+		{
+			FloatDataSet = FNiagaraDataSetAccessor<EitherA>(InDataSet, FNiagaraVariable(TagType::GetTypeA(), InName));
+		}
+		else
+		{
+			HalfDataSet = FNiagaraDataSetAccessor<EitherB>(InDataSet, FNiagaraVariable(TagType::GetTypeB(), InName));
+		}
+		InitForAccess();
+	}
+
+	void InitForAccess()
+	{
+		if (FloatDataSet.IsValid())
+		{
+			FloatDataSet.InitForAccess();
+		}
+		else if (HalfDataSet.IsValid())
+		{
+			HalfDataSet.InitForAccess();
+		}
+	}
+
+	bool IsValid()const
+	{
+		return FloatDataSet.IsValid() || HalfDataSet.IsValid();
+	}
+
+	FORCEINLINE bool IsValidForRead()const
+	{
+		return  (FloatDataSet.IsValid() && FloatDataSet.IsValidForRead()) || (HalfDataSet.IsValid() && HalfDataSet.IsValidForRead());
+	}
+	FORCEINLINE bool IsValidForWrite()const
+	{
+		return  (FloatDataSet.IsValid() && FloatDataSet.IsValidForWrite()) || (HalfDataSet.IsValid() && HalfDataSet.IsValidForWrite());
+	}
+
+	FORCEINLINE EitherA operator[](int32 Index)const
+	{
+		if (FloatDataSet.IsValid())
+		{
+			return FloatDataSet[Index];
+		}
+		else
+		{
+			return TagType::ConvertB(HalfDataSet[Index]);
+		}
+	}
+
+	FORCEINLINE EitherA GetSafe(int32 Index, EitherA Default = EitherA())const
+	{
+		if (FloatDataSet.IsValid())
+		{
+			return FloatDataSet.GetSafe(Index, Default);
+		}
+		else
+		{
+			EitherB Val = TagType::ConvertA(Default);
+			return TagType::ConvertB(HalfDataSet.GetSafe(Index, Val));
+		}
+	}
+
+	FORCEINLINE EitherA Get(int32 Index)const
+	{
+		if (FloatDataSet.IsValid())
+		{
+			return FloatDataSet.Get(Index);
+		}
+		else
+		{
+			return TagType::ConvertB(HalfDataSet.Get(Index));
+		}
+	}
+
+	FORCEINLINE void Get(int32 Index, EitherA& OutValue)const
+	{
+		if (FloatDataSet.IsValid())
+		{
+			FloatDataSet.Get(Index, OutValue);
+		}
+		else
+		{
+			EitherB Val;
+			HalfDataSet.Get(Index, Val);
+			OutValue = TagType::ConvertB(Val);
+		}
+	}
+
+	FORCEINLINE void Set(int32 Index, EitherA InValue)
+	{
+		if (FloatDataSet.IsValid())
+		{
+			FloatDataSet.Set(Index, InValue);
+		}
+		else
+		{
+			HalfDataSet.Set(Index, TagType::ConvertA(InValue));
+		}
+	}
+
+	FORCEINLINE float GetMaxElement(uint32 ElementIndex) const
+	{
+		if (FloatDataSet.IsValid())
+		{
+			return FloatDataSet.GetMaxElement(ElementIndex);
+		}
+		else
+		{
+			return HalfDataSet.GetMaxElement(ElementIndex);
+		}
+	}
+
+	FORCEINLINE FNiagaraMinMaxStruct GetMinMaxElement(uint32 ElementIndex) const
+	{
+		if (FloatDataSet.IsValid())
+		{
+			return FloatDataSet.GetMinMaxElement(ElementIndex);
+		}
+		else
+		{
+			return HalfDataSet.GetMinMaxElement(ElementIndex);
+		}
+	}
+
+private:
+	FNiagaraDataSetAccessor<EitherA> FloatDataSet;
+	FNiagaraDataSetAccessor<EitherB> HalfDataSet;
+};
+
+struct FNiagaraDataConversions
+{
+	struct FHalfOrFloat
+	{
+		using EitherA = float;
+		using EitherB = TStaticArray<FFloat16, 1>;
+
+		static TStaticArray<FFloat16, 1> ConvertA(float In)
+		{
+			TStaticArray<FFloat16, 1> Ret;
+			FPlatformMath::StoreHalf((uint16*)&Ret[0], In);
+			return Ret;
+		}
+		static float ConvertB(TStaticArray<FFloat16, 1> In)
+		{
+			return FPlatformMath::LoadHalf((uint16*)&In[0]);
+		}
+
+		static FNiagaraTypeDefinition GetTypeA()
+		{
+			return FNiagaraTypeDefinition::GetFloatDef();
+		}
+
+		static FNiagaraTypeDefinition GetTypeB()
+		{
+			return FNiagaraTypeDefinition::GetHalfDef();
+		}
+	};
+
+	struct FHalf2OrFloat2
+	{
+		using EitherA = FVector2D;
+		using EitherB = TStaticArray<FFloat16, 2>;
+
+		static TStaticArray<FFloat16, 2> ConvertA(FVector2D In)
+		{
+			TStaticArray<FFloat16, 2> Ret;
+			FPlatformMath::StoreHalf((uint16*)&Ret[0], In.X);
+			FPlatformMath::StoreHalf((uint16*)&Ret[1], In.Y);
+			return Ret;
+		}
+		static FVector2D ConvertB(TStaticArray<FFloat16, 2> In)
+		{
+			return FVector2D
+			(
+				FPlatformMath::LoadHalf((uint16*)&In[0]),
+				FPlatformMath::LoadHalf((uint16*)&In[1])
+			);
+		}
+
+		static FNiagaraTypeDefinition GetTypeA()
+		{
+			return FNiagaraTypeDefinition::GetVec2Def();
+		}
+
+		static FNiagaraTypeDefinition GetTypeB()
+		{
+			return FNiagaraTypeDefinition::GetHalfVec2Def();
+		}
+	};
+
+	struct FHalf3OrFloat3
+	{
+		using EitherA = FVector;
+		using EitherB = TStaticArray<FFloat16, 3>;
+
+		static TStaticArray<FFloat16, 3> ConvertA(FVector In)
+		{
+			TStaticArray<FFloat16, 3> Ret;
+			FPlatformMath::StoreHalf((uint16*)&Ret[0], In.X);
+			FPlatformMath::StoreHalf((uint16*)&Ret[1], In.Y);
+			FPlatformMath::StoreHalf((uint16*)&Ret[2], In.Z);
+			return Ret;
+		}
+		static FVector ConvertB(TStaticArray<FFloat16, 3> In)
+		{
+			return FVector
+			(
+				FPlatformMath::LoadHalf((uint16*)&In[0]), 
+				FPlatformMath::LoadHalf((uint16*)&In[1]), 
+				FPlatformMath::LoadHalf((uint16*)&In[2])
+			);
+		}
+
+		static FNiagaraTypeDefinition GetTypeA()
+		{
+			return FNiagaraTypeDefinition::GetVec3Def();
+		}
+
+		static FNiagaraTypeDefinition GetTypeB()
+		{
+			return FNiagaraTypeDefinition::GetHalfVec3Def();
+		}
+	};
+
+	struct FHalf4OrFloat4
+	{
+		using EitherA = FVector4;
+		using EitherB = TStaticArray<FFloat16, 4>;
+
+		static TStaticArray<FFloat16, 4> ConvertA(FVector4 In)
+		{
+			TStaticArray<FFloat16, 4> Ret;
+			FPlatformMath::VectorStoreHalf((uint16*)&Ret[0], (float*)&In);
+			return Ret;
+		}
+		static FVector4 ConvertB(TStaticArray<FFloat16, 4> In)
+		{
+			FVector4 Ret;
+			FPlatformMath::VectorLoadHalf((float*)&Ret, (uint16*)&In[0]);
+			return Ret;
+		}
+
+		static FNiagaraTypeDefinition GetTypeA()
+		{
+			return FNiagaraTypeDefinition::GetVec4Def();
+		}
+
+		static FNiagaraTypeDefinition GetTypeB()
+		{
+			return FNiagaraTypeDefinition::GetHalfVec4Def();
+		}
+	};
+
+	struct FHalf4OrColor4
+	{
+		using EitherA = FLinearColor;
+		using EitherB = TStaticArray<FFloat16, 4>;
+
+		static TStaticArray<FFloat16, 4> ConvertA(FLinearColor In)
+		{
+			TStaticArray<FFloat16, 4> Ret;
+			FPlatformMath::VectorStoreHalf((uint16*)&Ret[0], (float*)&In);
+			return Ret;
+		}
+		static FLinearColor ConvertB(TStaticArray<FFloat16, 4> In)
+		{
+			FLinearColor Ret;
+			FPlatformMath::VectorLoadHalf((float*)&Ret, (uint16*)&In[0]);
+			return Ret;
+		}
+
+		static FNiagaraTypeDefinition GetTypeA()
+		{
+			return FNiagaraTypeDefinition::GetColorDef();
+		}
+
+		static FNiagaraTypeDefinition GetTypeB()
+		{
+			return FNiagaraTypeDefinition::GetHalfVec4Def();
+		}
+	};
+
+	struct FHalf4OrQuat4
+	{
+		using EitherA = FQuat;
+		using EitherB = TStaticArray<FFloat16, 4>;
+
+		static TStaticArray<FFloat16, 4> ConvertA(FQuat In)
+		{
+			TStaticArray<FFloat16, 4> Ret;
+			FPlatformMath::VectorStoreHalf((uint16*)&Ret[0], (float*)&In);
+			return Ret;
+		}
+		static FQuat ConvertB(TStaticArray<FFloat16, 4> In)
+		{
+			FQuat Ret;
+			FPlatformMath::VectorLoadHalf((float*)&Ret, (uint16*)&In[0]);
+			return Ret;
+		}
+
+		static FNiagaraTypeDefinition GetTypeA()
+		{
+			return FNiagaraTypeDefinition::GetQuatDef();
+		}
+
+		static FNiagaraTypeDefinition GetTypeB()
+		{
+			return FNiagaraTypeDefinition::GetHalfVec4Def();
+		}
+	};
+
+	struct FNiagaraInt {};
+
+	struct FNiagaraBool {};
+};
+
+template<>
+struct FNiagaraDataSetAccessor<FNiagaraDataConversions::FHalfOrFloat> : public FNiagaraEitherDataSetAccessor<FNiagaraDataConversions::FHalfOrFloat>
+{
+	using FNiagaraEitherDataSetAccessor<FNiagaraDataConversions::FHalfOrFloat>::FNiagaraEitherDataSetAccessor;
+};
+
+
+template<>
+struct FNiagaraDataSetAccessor<FNiagaraDataConversions::FHalf2OrFloat2> : public FNiagaraEitherDataSetAccessor<FNiagaraDataConversions::FHalf2OrFloat2>
+{
+	using FNiagaraEitherDataSetAccessor<FNiagaraDataConversions::FHalf2OrFloat2>::FNiagaraEitherDataSetAccessor;
+};
+
+template<>
+struct FNiagaraDataSetAccessor<FNiagaraDataConversions::FHalf3OrFloat3> : public FNiagaraEitherDataSetAccessor<FNiagaraDataConversions::FHalf3OrFloat3>
+{
+	using FNiagaraEitherDataSetAccessor<FNiagaraDataConversions::FHalf3OrFloat3>::FNiagaraEitherDataSetAccessor;
+};
+
+
+template<>
+struct FNiagaraDataSetAccessor<FNiagaraDataConversions::FHalf4OrFloat4> : public FNiagaraEitherDataSetAccessor<FNiagaraDataConversions::FHalf4OrFloat4>
+{
+	using FNiagaraEitherDataSetAccessor<FNiagaraDataConversions::FHalf4OrFloat4>::FNiagaraEitherDataSetAccessor;
+};
+
+
+template<>
+struct FNiagaraDataSetAccessor<FNiagaraDataConversions::FHalf4OrColor4> : public FNiagaraEitherDataSetAccessor<FNiagaraDataConversions::FHalf4OrColor4>
+{
+	using FNiagaraEitherDataSetAccessor<FNiagaraDataConversions::FHalf4OrColor4>::FNiagaraEitherDataSetAccessor;
+};
+
+template<>
+struct FNiagaraDataSetAccessor<FNiagaraDataConversions::FHalf4OrQuat4> : public FNiagaraEitherDataSetAccessor<FNiagaraDataConversions::FHalf4OrQuat4>
+{
+	using FNiagaraEitherDataSetAccessor<FNiagaraDataConversions::FHalf4OrQuat4>::FNiagaraEitherDataSetAccessor;
+};
+
+template<>
+struct FNiagaraDataSetAccessor<FNiagaraDataConversions::FNiagaraInt> : public FNiagaraDataSetAccessor<int>
+{
+	FNiagaraDataSetAccessor<FNiagaraDataConversions::FNiagaraInt>(FNiagaraDataSet& InDataSet, const FName& InName) : FNiagaraDataSetAccessor<int>(InDataSet, FNiagaraVariable(FNiagaraTypeDefinition::GetIntDef(), InName))
+	{
+	}
+};
+
+template<>
+struct FNiagaraDataSetAccessor<FNiagaraDataConversions::FNiagaraBool> : public FNiagaraDataSetAccessor<FNiagaraBool>
+{
+	FNiagaraDataSetAccessor<FNiagaraDataConversions::FNiagaraBool>(FNiagaraDataSet& InDataSet, const FName& InName) : FNiagaraDataSetAccessor<FNiagaraBool>(InDataSet, FNiagaraVariable(FNiagaraTypeDefinition::GetBoolDef(), InName))
+	{
+	}
 };
 
+
 /**
 Iterator that will pull or push data between a NiagaraDataBuffer and some FNiagaraVariables it contains.
 Super slow. Don't use at runtime.
@@ -1504,7 +2088,7 @@ struct FNiagaraDataVariableIterator
 			for (uint32 CompIdx = 0; CompIdx < Layout.GetNumFloatComponents(); ++CompIdx)
 			{
 				uint32 CompBufferOffset = Layout.FloatComponentStart + CompIdx;
-				float* Src = Data->GetInstancePtrFloat(CompBufferOffset, CurrIdx);
+				const float* Src = Data->GetInstancePtrFloat(CompBufferOffset, CurrIdx);
 				float* Dst = (float*)(ValuePtr + Layout.LayoutInfo.FloatComponentByteOffsets[CompIdx]);
 				*Dst = *Src;
 			}
@@ -1512,10 +2096,18 @@ struct FNiagaraDataVariableIterator
 			for (uint32 CompIdx = 0; CompIdx < Layout.GetNumInt32Components(); ++CompIdx)
 			{
 				uint32 CompBufferOffset = Layout.Int32ComponentStart + CompIdx;
-				int32* Src = Data->GetInstancePtrInt32(CompBufferOffset, CurrIdx);
+				const int32* Src = Data->GetInstancePtrInt32(CompBufferOffset, CurrIdx);
 				int32* Dst = (int32*)(ValuePtr + Layout.LayoutInfo.Int32ComponentByteOffsets[CompIdx]);
 				*Dst = *Src;
 			}
+
+			for (uint32 CompIdx = 0; CompIdx < Layout.GetNumHalfComponents(); ++CompIdx)
+			{
+				uint32 CompBufferOffset = Layout.HalfComponentStart + CompIdx;
+				const FFloat16* Src = Data->GetInstancePtrHalf(CompBufferOffset, CurrIdx);
+				FFloat16* Dst = (FFloat16*)(ValuePtr + Layout.LayoutInfo.HalfComponentByteOffsets[CompIdx]);
+				*Dst = *Src;
+			}
 		}
 	}
 
diff --git a/Engine/Plugins/FX/Niagara/Source/Niagara/Classes/NiagaraGPUSortInfo.h b/Engine/Plugins/FX/Niagara/Source/Niagara/Classes/NiagaraGPUSortInfo.h
--- a/Engine/Plugins/FX/Niagara/Source/Niagara/Classes/NiagaraGPUSortInfo.h
+++ b/Engine/Plugins/FX/Niagara/Source/Niagara/Classes/NiagaraGPUSortInfo.h
@@ -41,8 +41,8 @@ struct FNiagaraGPUSortInfo
 	FShaderResourceViewRHIRef GPUParticleCountSRV;
 	uint32 GPUParticleCountOffset = INDEX_NONE;
 	// View data.
-	FVector ViewOrigin;
-	FVector ViewDirection;
+	FVector ViewOrigin = FVector(0, 0, 0);
+	FVector ViewDirection = FVector(0, 0, 1);
 
 	// The GPUSortManager bindings for this sort task.
 	FGPUSortManager::FAllocationInfo AllocationInfo;
diff --git a/Engine/Plugins/FX/Niagara/Source/Niagara/Classes/NiagaraScriptExecutionContext.h b/Engine/Plugins/FX/Niagara/Source/Niagara/Classes/NiagaraScriptExecutionContext.h
--- a/Engine/Plugins/FX/Niagara/Source/Niagara/Classes/NiagaraScriptExecutionContext.h
+++ b/Engine/Plugins/FX/Niagara/Source/Niagara/Classes/NiagaraScriptExecutionContext.h
@@ -287,11 +287,14 @@ public:
 #if WITH_EDITORONLY_DATA
 	mutable FRHIGPUMemoryReadback *GPUDebugDataReadbackFloat;
 	mutable FRHIGPUMemoryReadback *GPUDebugDataReadbackInt;
+	mutable FRHIGPUMemoryReadback *GPUDebugDataReadbackHalf;
 	mutable FRHIGPUMemoryReadback *GPUDebugDataReadbackCounts;
 	mutable uint32 GPUDebugDataFloatSize;
 	mutable uint32 GPUDebugDataIntSize;
+	mutable uint32 GPUDebugDataHalfSize;
 	mutable uint32 GPUDebugDataFloatStride;
 	mutable uint32 GPUDebugDataIntStride;
+	mutable uint32 GPUDebugDataHalfStride;
 	mutable uint32 GPUDebugDataCountOffset;
 	mutable TSharedPtr<struct FNiagaraScriptDebuggerInfo, ESPMode::ThreadSafe> DebugInfo;
 #endif
diff --git a/Engine/Plugins/FX/Niagara/Source/Niagara/Private/NiagaraBoundsCalculatorHelper.h b/Engine/Plugins/FX/Niagara/Source/Niagara/Private/NiagaraBoundsCalculatorHelper.h
--- a/Engine/Plugins/FX/Niagara/Source/Niagara/Private/NiagaraBoundsCalculatorHelper.h
+++ b/Engine/Plugins/FX/Niagara/Source/Niagara/Private/NiagaraBoundsCalculatorHelper.h
@@ -22,33 +22,40 @@ public:
 		static const FName ScaleName(TEXT("Scale"));
 		static const FName RibbonWidthName(TEXT("RibbonWidth"));
 
-		PositionAccessor = FNiagaraDataSetAccessor<FVector>(DataSet, FNiagaraVariable(FNiagaraTypeDefinition::GetVec3Def(), PositionName));
+		if (DataSet.HasVariable(PositionName))
+		{
+			PositionAccessor = FNiagaraDataSetAccessor<FNiagaraDataConversions::FHalf3OrFloat3>(DataSet, PositionName);
+		}
+		else
+		{
+			PositionAccessor = FNiagaraDataSetAccessor<FNiagaraDataConversions::FHalf3OrFloat3>();
+		}
 
-		if (bUsedWithSprites)
+		if (bUsedWithSprites && DataSet.HasVariable(SpriteSizeName))
 		{
-			SpriteSizeAccessor = FNiagaraDataSetAccessor<FVector2D>(DataSet, FNiagaraVariable(FNiagaraTypeDefinition::GetVec2Def(), SpriteSizeName));
+			SpriteSizeAccessor = FNiagaraDataSetAccessor<FNiagaraDataConversions::FHalf2OrFloat2>(DataSet, SpriteSizeName);
 		}
 		else
 		{
-			SpriteSizeAccessor = FNiagaraDataSetAccessor<FVector2D>();
+			SpriteSizeAccessor = FNiagaraDataSetAccessor<FNiagaraDataConversions::FHalf2OrFloat2>();
 		}
 
-		if (bUsedWithMeshes)
+		if (bUsedWithMeshes && DataSet.HasVariable(ScaleName))
 		{
-			ScaleAccessor = FNiagaraDataSetAccessor<FVector>(DataSet, FNiagaraVariable(FNiagaraTypeDefinition::GetVec3Def(), ScaleName));
+			ScaleAccessor = FNiagaraDataSetAccessor<FNiagaraDataConversions::FHalf3OrFloat3>(DataSet, ScaleName);
 		}
 		else
 		{
-			ScaleAccessor = FNiagaraDataSetAccessor<FVector>();
+			ScaleAccessor = FNiagaraDataSetAccessor<FNiagaraDataConversions::FHalf3OrFloat3>();
 		}
 
-		if (bUsedWithRibbons)
+		if (bUsedWithRibbons && DataSet.HasVariable(RibbonWidthName))
 		{
-			RibbonWidthAccessor = FNiagaraDataSetAccessor<float>(DataSet, FNiagaraVariable(FNiagaraTypeDefinition::GetFloatDef(), RibbonWidthName));
+			RibbonWidthAccessor = FNiagaraDataSetAccessor<FNiagaraDataConversions::FHalfOrFloat>(DataSet, RibbonWidthName);
 		}
 		else
 		{
-			RibbonWidthAccessor = FNiagaraDataSetAccessor<float>();
+			RibbonWidthAccessor = FNiagaraDataSetAccessor<FNiagaraDataConversions::FHalfOrFloat>();
 		}
 	}
 
@@ -79,15 +86,13 @@ public:
 
 		constexpr float kDefaultSize = 50.0f;
 
-		const FVector PositionMax = FVector(
-			GetFloatMax(MakeArrayView(PositionAccessor.GetX(), NumInstances)),
-			GetFloatMax(MakeArrayView(PositionAccessor.GetY(), NumInstances)),
-			GetFloatMax(MakeArrayView(PositionAccessor.GetZ(), NumInstances)));
 
-		const FVector PositionMin = FVector(
-			GetFloatMin(MakeArrayView(PositionAccessor.GetX(), NumInstances)),
-			GetFloatMin(MakeArrayView(PositionAccessor.GetY(), NumInstances)),
-			GetFloatMin(MakeArrayView(PositionAccessor.GetZ(), NumInstances)));
+		FNiagaraMinMaxStruct MinMaxXs = PositionAccessor.GetMinMaxElement(0);
+		FNiagaraMinMaxStruct MinMaxYs = PositionAccessor.GetMinMaxElement(1);
+		FNiagaraMinMaxStruct MinMaxZs = PositionAccessor.GetMinMaxElement(2);
+
+		const FVector PositionMax = FVector(MinMaxXs.Max, MinMaxYs.Max, MinMaxZs.Max);
+		const FVector PositionMin = FVector(MinMaxXs.Min, MinMaxYs.Min, MinMaxZs.Min);
 
 		float MaxSize = KINDA_SMALL_NUMBER;
 		if (bUsedWithMeshes)
@@ -95,10 +100,12 @@ public:
 			FVector MaxScale = FVector(kDefaultSize, kDefaultSize, kDefaultSize);
 			if (ScaleAccessor.IsValid())
 			{
-				MaxScale = FVector(
-					GetFloatMax(MakeArrayView(ScaleAccessor.GetX(), NumInstances)),
-					GetFloatMax(MakeArrayView(ScaleAccessor.GetY(), NumInstances)),
-					GetFloatMax(MakeArrayView(ScaleAccessor.GetZ(), NumInstances)));
+				MaxScale = FVector
+				(
+					ScaleAccessor.GetMaxElement(0),
+					ScaleAccessor.GetMaxElement(1),
+					ScaleAccessor.GetMaxElement(2)
+				);
 			}
 
 			const FVector ScaledExtents = MeshExtents * (MaxScale.IsNearlyZero() ? FVector::OneVector : MaxScale);
@@ -111,9 +118,7 @@ public:
 
 			if (SpriteSizeAccessor.IsValid())
 			{
-				MaxSpriteSize = FMath::Max(
-					GetFloatMax(MakeArrayView(SpriteSizeAccessor.GetX(), NumInstances)),
-					GetFloatMax(MakeArrayView(SpriteSizeAccessor.GetY(), NumInstances)));
+				MaxSpriteSize = FMath::Max(SpriteSizeAccessor.GetMaxElement(0), SpriteSizeAccessor.GetMaxElement(1));
 			}
 			MaxSize = FMath::Max(MaxSize, FMath::IsNearlyZero(MaxSpriteSize) ? 1.0f : MaxSpriteSize);
 		}
@@ -123,7 +128,7 @@ public:
 			float MaxRibbonWidth = kDefaultSize;
 			if (RibbonWidthAccessor.IsValid())
 			{
-				MaxRibbonWidth = GetFloatMax(MakeArrayView(RibbonWidthAccessor.GetX(), NumInstances));
+				MaxRibbonWidth = RibbonWidthAccessor.GetMaxElement(0);
 			}
 
 			MaxSize = FMath::Max(MaxSize, FMath::IsNearlyZero(MaxRibbonWidth) ? 1.0f : MaxRibbonWidth);
@@ -132,9 +137,9 @@ public:
 		return FBox(PositionMin, PositionMax).ExpandBy(MaxSize);
 	}
 
-	FNiagaraDataSetAccessor<FVector> PositionAccessor;
-	FNiagaraDataSetAccessor<FVector2D> SpriteSizeAccessor;
-	FNiagaraDataSetAccessor<FVector> ScaleAccessor;
-	FNiagaraDataSetAccessor<float> RibbonWidthAccessor;
+	FNiagaraDataSetAccessor<FNiagaraDataConversions::FHalf3OrFloat3> PositionAccessor;
+	FNiagaraDataSetAccessor<FNiagaraDataConversions::FHalf2OrFloat2> SpriteSizeAccessor;
+	FNiagaraDataSetAccessor<FNiagaraDataConversions::FHalf3OrFloat3> ScaleAccessor;
+	FNiagaraDataSetAccessor<FNiagaraDataConversions::FHalfOrFloat> RibbonWidthAccessor;
 	const FVector MeshExtents = FVector::OneVector;
 };
diff --git a/Engine/Plugins/FX/Niagara/Source/Niagara/Private/NiagaraComponent.cpp b/Engine/Plugins/FX/Niagara/Source/Niagara/Private/NiagaraComponent.cpp
--- a/Engine/Plugins/FX/Niagara/Source/Niagara/Private/NiagaraComponent.cpp
+++ b/Engine/Plugins/FX/Niagara/Source/Niagara/Private/NiagaraComponent.cpp
@@ -1688,7 +1688,7 @@ TArray<float> UNiagaraComponent::GetNiagaraParticleValues_DebugOnly(const FStrin
 				FNiagaraDataBuffer& ParticleData = Sim->GetData().GetCurrentDataChecked();
 				int32 NumParticles = ParticleData.GetNumInstances();
 				Values.SetNum(NumParticles);
-				FNiagaraDataSetAccessor<float> ValueData(Sim->GetData(), FNiagaraVariable(FNiagaraTypeDefinition::GetFloatDef(), *InValueName));
+				FNiagaraDataSetAccessor<FNiagaraDataConversions::FHalfOrFloat> ValueData(Sim->GetData(), *InValueName);
 				for (int32 i = 0; i < NumParticles; ++i)
 				{
 					float Value;
diff --git a/Engine/Plugins/FX/Niagara/Source/Niagara/Private/NiagaraDataInterfaceParticleRead.cpp b/Engine/Plugins/FX/Niagara/Source/Niagara/Private/NiagaraDataInterfaceParticleRead.cpp
--- a/Engine/Plugins/FX/Niagara/Source/Niagara/Private/NiagaraDataInterfaceParticleRead.cpp
+++ b/Engine/Plugins/FX/Niagara/Source/Niagara/Private/NiagaraDataInterfaceParticleRead.cpp
@@ -27,9 +27,12 @@ static const FString SpawnedIDsBufferBaseName(TEXT("SpawnedIDsBuffer_"));
 static const FString IDToIndexTableBaseName(TEXT("IDToIndexTable_"));
 static const FString InputFloatBufferBaseName(TEXT("InputFloatBuffer_"));
 static const FString InputIntBufferBaseName(TEXT("InputIntBuffer_"));
+static const FString InputHalfBufferBaseName(TEXT("InputHalfBuffer_"));
 static const FString ParticleStrideFloatBaseName(TEXT("ParticleStrideFloat_"));
 static const FString ParticleStrideIntBaseName(TEXT("ParticleStrideInt_"));
+static const FString ParticleStrideHalfBaseName(TEXT("ParticleStrideHalf_"));
 static const FString AttributeIndicesBaseName(TEXT("AttributeIndices_"));
+static const FString AttributeCompressedBaseName(TEXT("AttributeCompressed_"));
 static const FString AcquireTagRegisterIndexBaseName(TEXT("AcquireTagRegisterIndex_"));
 
 enum class ENiagaraParticleDataComponentType : uint8
@@ -99,6 +102,7 @@ struct FNDIParticleRead_RenderInstanceData
 	FString SourceEmitterName;
 	const FNiagaraDataSet* CachedDataSet;
 	TArray<int32> AttributeIndices;
+	TArray<int32> AttributeCompressed;
 	int32 AcquireTagRegisterIndex;
 	bool bSourceEmitterNotGPUErrorShown;
 };
@@ -179,9 +183,12 @@ struct FNiagaraDataInterfaceParametersCS_ParticleRead : public FNiagaraDataInter
 		IDToIndexTableParam.Bind(ParameterMap, *(IDToIndexTableBaseName + ParameterInfo.DataInterfaceHLSLSymbol));
 		InputFloatBufferParam.Bind(ParameterMap, *(InputFloatBufferBaseName + ParameterInfo.DataInterfaceHLSLSymbol));
 		InputIntBufferParam.Bind(ParameterMap, *(InputIntBufferBaseName + ParameterInfo.DataInterfaceHLSLSymbol));
+		InputHalfBufferParam.Bind(ParameterMap, *(InputHalfBufferBaseName + ParameterInfo.DataInterfaceHLSLSymbol));
 		ParticleStrideFloatParam.Bind(ParameterMap, *(ParticleStrideFloatBaseName + ParameterInfo.DataInterfaceHLSLSymbol));
 		ParticleStrideIntParam.Bind(ParameterMap, *(ParticleStrideIntBaseName + ParameterInfo.DataInterfaceHLSLSymbol));
+		ParticleStrideHalfParam.Bind(ParameterMap, *(ParticleStrideHalfBaseName + ParameterInfo.DataInterfaceHLSLSymbol));
 		AttributeIndicesParam.Bind(ParameterMap, *(AttributeIndicesBaseName + ParameterInfo.DataInterfaceHLSLSymbol));
+		AttributeCompressedParam.Bind(ParameterMap, *(AttributeCompressedBaseName + ParameterInfo.DataInterfaceHLSLSymbol));
 		AcquireTagRegisterIndexParam.Bind(ParameterMap, *(AcquireTagRegisterIndexBaseName + ParameterInfo.DataInterfaceHLSLSymbol));
 
 		int32 NumFuncs = ParameterInfo.GeneratedFunctions.Num();
@@ -211,11 +218,16 @@ struct FNiagaraDataInterfaceParametersCS_ParticleRead : public FNiagaraDataInter
 	{
 		// Set all the indices to -1, so that any reads return false.
 		TArray<int32, TInlineAllocator<32>> AttributeIndices;
+		TArray<int32, TInlineAllocator<32>> AttributeCompressed;
+
 		int NumAttrIndices = Align(AttributeNames.Num(), 4);
 		AttributeIndices.SetNumUninitialized(NumAttrIndices);
+		AttributeCompressed.SetNumUninitialized(NumAttrIndices);
+
 		for (int AttrIdx = 0; AttrIdx < AttributeIndices.Num(); ++AttrIdx)
 		{
 			AttributeIndices[AttrIdx] = -1;
+			AttributeCompressed[AttrIdx] = 0;
 		}
 
 		int32 AcquireTagRegisterIndex = -1;
@@ -231,9 +243,12 @@ struct FNiagaraDataInterfaceParametersCS_ParticleRead : public FNiagaraDataInter
 		SetSRVParameter(RHICmdList, ComputeShader, IDToIndexTableParam, FNiagaraRenderer::GetDummyIntBuffer());
 		SetSRVParameter(RHICmdList, ComputeShader, InputFloatBufferParam, FNiagaraRenderer::GetDummyFloatBuffer());
 		SetSRVParameter(RHICmdList, ComputeShader, InputIntBufferParam, FNiagaraRenderer::GetDummyIntBuffer());
+		SetSRVParameter(RHICmdList, ComputeShader, InputHalfBufferParam, FNiagaraRenderer::GetDummyHalfBuffer());
 		SetShaderValue(RHICmdList, ComputeShader, ParticleStrideFloatParam, 0);
 		SetShaderValue(RHICmdList, ComputeShader, ParticleStrideIntParam, 0);
+		SetShaderValue(RHICmdList, ComputeShader, ParticleStrideHalfParam, 0);
 		SetShaderValueArray(RHICmdList, ComputeShader, AttributeIndicesParam, AttributeIndices.GetData(), AttributeIndices.Num());
+		SetShaderValueArray(RHICmdList, ComputeShader, AttributeCompressedParam, AttributeCompressed.GetData(), AttributeCompressed.Num());
 		SetShaderValue(RHICmdList, ComputeShader, AcquireTagRegisterIndexParam, AcquireTagRegisterIndex);
 	}
 
@@ -253,12 +268,29 @@ struct FNiagaraDataInterfaceParametersCS_ParticleRead : public FNiagaraDataInter
 		}
 	}
 
+	bool CheckHalfVariableType(const FNiagaraTypeDefinition& VarType, ENiagaraParticleDataValueType AttributeType) const
+	{
+		switch (AttributeType)
+		{
+			case ENiagaraParticleDataValueType::Int: return VarType == FNiagaraTypeDefinition::GetIntDef();
+			case ENiagaraParticleDataValueType::Float: return VarType == FNiagaraTypeDefinition::GetHalfDef();
+			case ENiagaraParticleDataValueType::Vec2: return VarType == FNiagaraTypeDefinition::GetHalfVec2Def();
+			case ENiagaraParticleDataValueType::Vec3: return VarType == FNiagaraTypeDefinition::GetHalfVec3Def();
+			case ENiagaraParticleDataValueType::Vec4: return VarType == FNiagaraTypeDefinition::GetHalfVec4Def();
+			case ENiagaraParticleDataValueType::Bool: return VarType == FNiagaraTypeDefinition::GetBoolDef();
+			case ENiagaraParticleDataValueType::Color: return VarType == FNiagaraTypeDefinition::GetHalfVec4Def();
+			case ENiagaraParticleDataValueType::Quat: return VarType == FNiagaraTypeDefinition::GetHalfVec4Def();
+			default: return false;
+		}
+	}
+
 	void FindAttributeIndices(FNDIParticleRead_RenderInstanceData* InstanceData, const FNiagaraDataSet* SourceDataSet) const
 	{
 		check(AttributeNames.Num() == AttributeTypes.Num());
 
 		int NumAttrIndices = Align(AttributeNames.Num(), 4);
 		InstanceData->AttributeIndices.SetNumUninitialized(NumAttrIndices);
+		InstanceData->AttributeCompressed.SetNumUninitialized(NumAttrIndices);
 
 		// Find the register index for each named attribute in the source emitter.
 		const TArray<FNiagaraVariable>& SourceEmitterVariables = SourceDataSet->GetVariables();
@@ -283,6 +315,13 @@ struct FNiagaraDataInterfaceParametersCS_ParticleRead : public FNiagaraDataInter
 					{
 						const FNiagaraVariableLayoutInfo& Layout = SourceEmitterVariableLayouts[VarIdx];
 						InstanceData->AttributeIndices[AttrNameIdx] = (AttributeType == ENiagaraParticleDataValueType::Int || AttributeType == ENiagaraParticleDataValueType::Bool) ? Layout.Int32ComponentStart : Layout.FloatComponentStart;
+						InstanceData->AttributeCompressed[AttrNameIdx] = 0;
+					}
+					else if (CheckHalfVariableType(Var.GetType(), AttributeType))
+					{
+						const FNiagaraVariableLayoutInfo& Layout = SourceEmitterVariableLayouts[VarIdx];
+						InstanceData->AttributeIndices[AttrNameIdx] = (AttributeType == ENiagaraParticleDataValueType::Int || AttributeType == ENiagaraParticleDataValueType::Bool) ? Layout.Int32ComponentStart : Layout.HalfComponentStart;
+						InstanceData->AttributeCompressed[AttrNameIdx] = 1;
 					}
 					else
 					{
@@ -290,6 +329,7 @@ struct FNiagaraDataInterfaceParametersCS_ParticleRead : public FNiagaraDataInter
 							*Var.GetName().ToString(), *InstanceData->SourceEmitterName, *Var.GetType().GetName(), NiagaraParticleDataValueTypeName(AttributeType)
 						);
 						InstanceData->AttributeIndices[AttrNameIdx] = -1;
+						InstanceData->AttributeCompressed[AttrNameIdx] = 0;
 					}
 					FoundVariable = true;
 					break;
@@ -300,6 +340,7 @@ struct FNiagaraDataInterfaceParametersCS_ParticleRead : public FNiagaraDataInter
 			{
 				UE_LOG(LogNiagara, Error, TEXT("Particle read DI is trying to access inexistent variable '%s' in emitter '%s'."), *AttrName.ToString(), *InstanceData->SourceEmitterName);
 				InstanceData->AttributeIndices[AttrNameIdx] = -1;
+				InstanceData->AttributeCompressed[AttrNameIdx] = 0;
 			}
 		}
 
@@ -326,6 +367,7 @@ struct FNiagaraDataInterfaceParametersCS_ParticleRead : public FNiagaraDataInter
 		for (int AttrIdx = AttributeNames.Num(); AttrIdx < InstanceData->AttributeIndices.Num(); ++AttrIdx)
 		{
 			InstanceData->AttributeIndices[AttrIdx] = -1;
+			InstanceData->AttributeCompressed[AttrIdx] = 0;
 		}
 	}
 
@@ -441,13 +483,17 @@ struct FNiagaraDataInterfaceParametersCS_ParticleRead : public FNiagaraDataInter
 
 		const uint32 ParticleStrideFloat = SourceData->GetFloatStride() / sizeof(float);
 		const uint32 ParticleStrideInt = SourceData->GetInt32Stride() / sizeof(int32);
+		const uint32 ParticleStrideHalf = SourceData->GetHalfStride() / sizeof(FFloat16);
 
 		SetSRVParameter(RHICmdList, ComputeShader, IDToIndexTableParam, GetIntSRVWithFallback(SourceData->GetGPUIDToIndexTable()));
 		SetSRVParameter(RHICmdList, ComputeShader, InputFloatBufferParam, GetFloatSRVWithFallback(SourceData->GetGPUBufferFloat()));
 		SetSRVParameter(RHICmdList, ComputeShader, InputIntBufferParam, GetIntSRVWithFallback(SourceData->GetGPUBufferInt()));
+		SetSRVParameter(RHICmdList, ComputeShader, InputHalfBufferParam, GetFloatSRVWithFallback(SourceData->GetGPUBufferHalf()));
 		SetShaderValue(RHICmdList, ComputeShader, ParticleStrideFloatParam, ParticleStrideFloat);
 		SetShaderValue(RHICmdList, ComputeShader, ParticleStrideIntParam, ParticleStrideInt);
+		SetShaderValue(RHICmdList, ComputeShader, ParticleStrideHalfParam, ParticleStrideHalf);
 		SetShaderValueArray(RHICmdList, ComputeShader, AttributeIndicesParam, InstanceData->AttributeIndices.GetData(), InstanceData->AttributeIndices.Num());
+		SetShaderValueArray(RHICmdList, ComputeShader, AttributeCompressedParam, InstanceData->AttributeCompressed.GetData(), InstanceData->AttributeCompressed.Num());
 		SetShaderValue(RHICmdList, ComputeShader, AcquireTagRegisterIndexParam, InstanceData->AcquireTagRegisterIndex);
 	}
 	
@@ -459,9 +505,12 @@ private:
 	LAYOUT_FIELD(FShaderResourceParameter, IDToIndexTableParam);
 	LAYOUT_FIELD(FShaderResourceParameter, InputFloatBufferParam);
 	LAYOUT_FIELD(FShaderResourceParameter, InputIntBufferParam);
+	LAYOUT_FIELD(FShaderResourceParameter, InputHalfBufferParam);
 	LAYOUT_FIELD(FShaderParameter, ParticleStrideFloatParam);
+	LAYOUT_FIELD(FShaderParameter, ParticleStrideHalfParam);
 	LAYOUT_FIELD(FShaderParameter, ParticleStrideIntParam);
 	LAYOUT_FIELD(FShaderParameter, AttributeIndicesParam);
+	LAYOUT_FIELD(FShaderParameter, AttributeCompressedParam);
 	LAYOUT_FIELD(FShaderParameter, AcquireTagRegisterIndexParam);
 	LAYOUT_FIELD(TMemoryImageArray<FName>, AttributeNames);
 	LAYOUT_FIELD(TMemoryImageArray<ENiagaraParticleDataValueType>, AttributeTypes);
@@ -703,6 +752,11 @@ DEFINE_NDI_DIRECT_FUNC_BINDER_WITH_PAYLOAD(UNiagaraDataInterfaceParticleRead, Re
 DEFINE_NDI_DIRECT_FUNC_BINDER_WITH_PAYLOAD(UNiagaraDataInterfaceParticleRead, ReadColor);
 DEFINE_NDI_DIRECT_FUNC_BINDER_WITH_PAYLOAD(UNiagaraDataInterfaceParticleRead, ReadQuat);
 
+static bool HasMatchingVariable(TArrayView<const FNiagaraVariable> Variables, FName AttributeName, const FNiagaraTypeDefinition& ValidType)
+{
+	return Variables.Find(FNiagaraVariable(ValidType, AttributeName)) != INDEX_NONE;
+}
+
 void UNiagaraDataInterfaceParticleRead::GetVMExternalFunction(const FVMExternalFunctionBindingInfo& BindingInfo, void* InstanceData, FVMExternalFunction& OutFunc)
 {
 	if (BindingInfo.Name == GetNumSpawnedParticlesFunctionName)
@@ -734,11 +788,13 @@ void UNiagaraDataInterfaceParticleRead::GetVMExternalFunction(const FVMExternalF
 		return;
 	}
 
+	TArrayView<const FNiagaraVariable> EmitterVariables = PIData->EmitterInstance->GetData().GetVariables();
+
 	const FName AttributeToRead = FunctionSpecifier->Value;
 	if (BindingInfo.Name == GetFloatAttributeFunctionName)
 	{
-		FNiagaraVariable VariableToRead(FNiagaraTypeDefinition::GetFloatDef(), AttributeToRead);
-		if (PIData->EmitterInstance->GetData().GetVariables().Find(VariableToRead) != INDEX_NONE)
+		if (HasMatchingVariable(EmitterVariables, AttributeToRead, FNiagaraTypeDefinition::GetFloatDef())
+			|| HasMatchingVariable(EmitterVariables, AttributeToRead, FNiagaraTypeDefinition::GetHalfDef()))
 		{
 			NDI_FUNC_BINDER(UNiagaraDataInterfaceParticleRead, ReadFloat)::Bind(this, OutFunc, AttributeToRead);
 			bBindSuccessful = true;
@@ -746,8 +802,8 @@ void UNiagaraDataInterfaceParticleRead::GetVMExternalFunction(const FVMExternalF
 	}
 	else if (BindingInfo.Name == GetVec2AttributeFunctionName)
 	{
-		FNiagaraVariable VariableToRead(FNiagaraTypeDefinition::GetVec2Def(), AttributeToRead);
-		if (PIData->EmitterInstance->GetData().GetVariables().Find(VariableToRead) != INDEX_NONE)
+		if (HasMatchingVariable(EmitterVariables, AttributeToRead, FNiagaraTypeDefinition::GetVec2Def())
+			|| HasMatchingVariable(EmitterVariables, AttributeToRead, FNiagaraTypeDefinition::GetHalfVec2Def()))
 		{
 			NDI_FUNC_BINDER(UNiagaraDataInterfaceParticleRead, ReadVector2)::Bind(this, OutFunc, AttributeToRead);
 			bBindSuccessful = true;
@@ -755,8 +811,8 @@ void UNiagaraDataInterfaceParticleRead::GetVMExternalFunction(const FVMExternalF
 	}
 	else if (BindingInfo.Name == GetVec3AttributeFunctionName)
 	{
-		FNiagaraVariable VariableToRead(FNiagaraTypeDefinition::GetVec3Def(), AttributeToRead);
-		if (PIData->EmitterInstance->GetData().GetVariables().Find(VariableToRead) != INDEX_NONE)
+		if (HasMatchingVariable(EmitterVariables, AttributeToRead, FNiagaraTypeDefinition::GetVec3Def())
+			|| HasMatchingVariable(EmitterVariables, AttributeToRead, FNiagaraTypeDefinition::GetHalfVec3Def()))
 		{
 			NDI_FUNC_BINDER(UNiagaraDataInterfaceParticleRead, ReadVector3)::Bind(this, OutFunc, AttributeToRead);
 			bBindSuccessful = true;
@@ -764,8 +820,8 @@ void UNiagaraDataInterfaceParticleRead::GetVMExternalFunction(const FVMExternalF
 	}
 	else if (BindingInfo.Name == GetVec4AttributeFunctionName)
 	{
-		FNiagaraVariable VariableToRead(FNiagaraTypeDefinition::GetVec4Def(), AttributeToRead);
-		if (PIData->EmitterInstance->GetData().GetVariables().Find(VariableToRead) != INDEX_NONE)
+		if (HasMatchingVariable(EmitterVariables, AttributeToRead, FNiagaraTypeDefinition::GetVec4Def())
+			|| HasMatchingVariable(EmitterVariables, AttributeToRead, FNiagaraTypeDefinition::GetHalfVec4Def()))
 		{
 			NDI_FUNC_BINDER(UNiagaraDataInterfaceParticleRead, ReadVector4)::Bind(this, OutFunc, AttributeToRead);
 			bBindSuccessful = true;
@@ -773,8 +829,7 @@ void UNiagaraDataInterfaceParticleRead::GetVMExternalFunction(const FVMExternalF
 	}
 	else if (BindingInfo.Name == GetIntAttributeFunctionName)
 	{
-		FNiagaraVariable VariableToRead(FNiagaraTypeDefinition::GetIntDef(), AttributeToRead);
-		if (PIData->EmitterInstance->GetData().GetVariables().Find(VariableToRead) != INDEX_NONE)
+		if (HasMatchingVariable(EmitterVariables, AttributeToRead, FNiagaraTypeDefinition::GetIntDef()))
 		{
 			NDI_FUNC_BINDER(UNiagaraDataInterfaceParticleRead, ReadInt)::Bind(this, OutFunc, AttributeToRead);
 			bBindSuccessful = true;
@@ -782,8 +837,7 @@ void UNiagaraDataInterfaceParticleRead::GetVMExternalFunction(const FVMExternalF
 	}
 	else if (BindingInfo.Name == GetBoolAttributeFunctionName)
 	{
-		FNiagaraVariable VariableToRead(FNiagaraTypeDefinition::GetBoolDef(), AttributeToRead);
-		if (PIData->EmitterInstance->GetData().GetVariables().Find(VariableToRead) != INDEX_NONE)
+		if (HasMatchingVariable(EmitterVariables, AttributeToRead, FNiagaraTypeDefinition::GetBoolDef()))
 		{
 			NDI_FUNC_BINDER(UNiagaraDataInterfaceParticleRead, ReadBool)::Bind(this, OutFunc, AttributeToRead);
 			bBindSuccessful = true;
@@ -791,8 +845,8 @@ void UNiagaraDataInterfaceParticleRead::GetVMExternalFunction(const FVMExternalF
 	}
 	else if (BindingInfo.Name == GetColorAttributeFunctionName)
 	{
-		FNiagaraVariable VariableToRead(FNiagaraTypeDefinition::GetColorDef(), AttributeToRead);
-		if (PIData->EmitterInstance->GetData().GetVariables().Find(VariableToRead) != INDEX_NONE)
+		if (HasMatchingVariable(EmitterVariables, AttributeToRead, FNiagaraTypeDefinition::GetColorDef())
+			|| HasMatchingVariable(EmitterVariables, AttributeToRead, FNiagaraTypeDefinition::GetHalfVec4Def()))
 		{
 			NDI_FUNC_BINDER(UNiagaraDataInterfaceParticleRead, ReadColor)::Bind(this, OutFunc, AttributeToRead);
 			bBindSuccessful = true;
@@ -800,8 +854,8 @@ void UNiagaraDataInterfaceParticleRead::GetVMExternalFunction(const FVMExternalF
 	}
 	else if (BindingInfo.Name == GetQuatAttributeFunctionName)
 	{
-		FNiagaraVariable VariableToRead(FNiagaraTypeDefinition::GetQuatDef(), AttributeToRead);
-		if (PIData->EmitterInstance->GetData().GetVariables().Find(VariableToRead) != INDEX_NONE)
+		if (HasMatchingVariable(EmitterVariables, AttributeToRead, FNiagaraTypeDefinition::GetQuatDef())
+			|| HasMatchingVariable(EmitterVariables, AttributeToRead, FNiagaraTypeDefinition::GetHalfVec4Def()))
 		{
 			NDI_FUNC_BINDER(UNiagaraDataInterfaceParticleRead, ReadQuat)::Bind(this, OutFunc, AttributeToRead);
 			bBindSuccessful = true;
@@ -898,7 +952,7 @@ void UNiagaraDataInterfaceParticleRead::ReadFloat(FVectorVMContext& Context, FNa
 	{
 		FNiagaraID ParticleID = { ParticleIDIndexParam.GetAndAdvance(), ParticleIDAcquireTagParam.GetAndAdvance() };
 		bool bValid;
-		float Value = RetrieveValueWithCheck<float>(InstanceData->EmitterInstance, FNiagaraTypeDefinition::GetFloatDef(), AttributeToRead, ParticleID, bValid);
+		float Value = RetrieveValueWithCheck<float, FNiagaraDataConversions::FHalfOrFloat>(InstanceData->EmitterInstance, AttributeToRead, ParticleID, bValid);
 		FNiagaraBool ValidValue;
 		ValidValue.SetValue(bValid);
 		*OutValid.GetDestAndAdvance() = ValidValue;
@@ -922,7 +976,7 @@ void UNiagaraDataInterfaceParticleRead::ReadVector2(FVectorVMContext& Context, F
 		FNiagaraID ParticleID = { ParticleIDIndexParam.GetAndAdvance(), ParticleIDAcquireTagParam.GetAndAdvance() };
 		bool bValid;
 
-		FVector2D Value = RetrieveValueWithCheck<FVector2D>(InstanceData->EmitterInstance, FNiagaraTypeDefinition::GetVec2Def(), AttributeToRead, ParticleID, bValid);
+		FVector2D Value = RetrieveValueWithCheck<FVector2D, FNiagaraDataConversions::FHalf2OrFloat2>(InstanceData->EmitterInstance, AttributeToRead, ParticleID, bValid);
 
 		FNiagaraBool ValidValue;
 		ValidValue.SetValue(bValid);
@@ -949,7 +1003,7 @@ void UNiagaraDataInterfaceParticleRead::ReadVector3(FVectorVMContext& Context, F
 		FNiagaraID ParticleID = { ParticleIDIndexParam.GetAndAdvance(), ParticleIDAcquireTagParam.GetAndAdvance() };
 		bool bValid;
 
-		FVector Value = RetrieveValueWithCheck<FVector>(InstanceData->EmitterInstance, FNiagaraTypeDefinition::GetVec3Def(), AttributeToRead, ParticleID, bValid);
+		FVector Value = RetrieveValueWithCheck<FVector, FNiagaraDataConversions::FHalf3OrFloat3>(InstanceData->EmitterInstance, AttributeToRead, ParticleID, bValid);
 
 		FNiagaraBool ValidValue;
 		ValidValue.SetValue(bValid);
@@ -978,7 +1032,7 @@ void UNiagaraDataInterfaceParticleRead::ReadVector4(FVectorVMContext& Context, F
 		FNiagaraID ParticleID = { ParticleIDIndexParam.GetAndAdvance(), ParticleIDAcquireTagParam.GetAndAdvance() };
 		bool bValid;
 
-		FVector4 Value = RetrieveValueWithCheck<FVector4>(InstanceData->EmitterInstance, FNiagaraTypeDefinition::GetVec4Def(), AttributeToRead, ParticleID, bValid);
+		FVector4 Value = RetrieveValueWithCheck<FVector4, FNiagaraDataConversions::FHalf4OrFloat4>(InstanceData->EmitterInstance, AttributeToRead, ParticleID, bValid);
 
 		FNiagaraBool ValidValue;
 		ValidValue.SetValue(bValid);
@@ -1005,7 +1059,7 @@ void UNiagaraDataInterfaceParticleRead::ReadInt(FVectorVMContext& Context, FName
 		FNiagaraID ParticleID = { ParticleIDIndexParam.GetAndAdvance(), ParticleIDAcquireTagParam.GetAndAdvance() };
 		bool bValid;
 
-		int32 Value = RetrieveValueWithCheck<int32>(InstanceData->EmitterInstance, FNiagaraTypeDefinition::GetIntDef(), AttributeToRead, ParticleID, bValid);
+		int32 Value = RetrieveValueWithCheck<int, FNiagaraDataConversions::FNiagaraInt>(InstanceData->EmitterInstance, AttributeToRead, ParticleID, bValid);
 
 		FNiagaraBool ValidValue;
 		ValidValue.SetValue(bValid);
@@ -1029,7 +1083,7 @@ void UNiagaraDataInterfaceParticleRead::ReadBool(FVectorVMContext& Context, FNam
 		FNiagaraID ParticleID = { ParticleIDIndexParam.GetAndAdvance(), ParticleIDAcquireTagParam.GetAndAdvance() };
 		bool bValid;
 
-		FNiagaraBool Value = RetrieveValueWithCheck<FNiagaraBool>(InstanceData->EmitterInstance, FNiagaraTypeDefinition::GetBoolDef(), AttributeToRead, ParticleID, bValid);
+		FNiagaraBool Value = RetrieveValueWithCheck<FNiagaraBool, FNiagaraDataConversions::FNiagaraBool>(InstanceData->EmitterInstance, AttributeToRead, ParticleID, bValid);
 
 		FNiagaraBool ValidValue;
 		ValidValue.SetValue(bValid);
@@ -1056,7 +1110,7 @@ void UNiagaraDataInterfaceParticleRead::ReadColor(FVectorVMContext& Context, FNa
 		FNiagaraID ParticleID = { ParticleIDIndexParam.GetAndAdvance(), ParticleIDAcquireTagParam.GetAndAdvance() };
 		bool bValid;
 
-		FLinearColor Value = RetrieveValueWithCheck<FLinearColor>(InstanceData->EmitterInstance, FNiagaraTypeDefinition::GetColorDef(), AttributeToRead, ParticleID, bValid);
+		FLinearColor Value = RetrieveValueWithCheck<FLinearColor, FNiagaraDataConversions::FHalf4OrColor4>(InstanceData->EmitterInstance, AttributeToRead, ParticleID, bValid);
 
 		FNiagaraBool ValidValue;
 		ValidValue.SetValue(bValid);
@@ -1086,7 +1140,7 @@ void UNiagaraDataInterfaceParticleRead::ReadQuat(FVectorVMContext& Context, FNam
 		FNiagaraID ParticleID = { ParticleIDIndexParam.GetAndAdvance(), ParticleIDAcquireTagParam.GetAndAdvance() };
 		bool bValid;
 
-		FQuat Value = RetrieveValueWithCheck<FQuat>(InstanceData->EmitterInstance, FNiagaraTypeDefinition::GetQuatDef(), AttributeToRead, ParticleID, bValid);
+		FQuat Value = RetrieveValueWithCheck<FQuat, FNiagaraDataConversions::FHalf4OrQuat4>(InstanceData->EmitterInstance, AttributeToRead, ParticleID, bValid);
 
 		FNiagaraBool ValidValue;
 		ValidValue.SetValue(bValid);
@@ -1098,10 +1152,10 @@ void UNiagaraDataInterfaceParticleRead::ReadQuat(FVectorVMContext& Context, FNam
 	}
 }
 
-template <typename T>
-T UNiagaraDataInterfaceParticleRead::RetrieveValueWithCheck(FNiagaraEmitterInstance* EmitterInstance, const FNiagaraTypeDefinition& Type, const FName& Attr, const FNiagaraID& ParticleID, bool& bValid)
+template <typename Type, typename AccessorType>
+Type UNiagaraDataInterfaceParticleRead::RetrieveValueWithCheck(FNiagaraEmitterInstance* EmitterInstance, const FName& Attr, const FNiagaraID& ParticleID, bool& bValid)
 {
-	T Value = T();
+	Type Value = Type();
 	bValid = false;
 
 	const FNiagaraDataBuffer* CurrentData = EmitterInstance->GetData().GetCurrentData();
@@ -1116,8 +1170,7 @@ T UNiagaraDataInterfaceParticleRead::RetrieveValueWithCheck(FNiagaraEmitterInsta
 		return Value;
 	}
 
-	FNiagaraVariable ReadVar(Type, Attr);
-	FNiagaraDataSetAccessor<T> ValueData(EmitterInstance->GetData(), ReadVar);
+	FNiagaraDataSetAccessor<AccessorType> ValueData(EmitterInstance->GetData(), Attr);
 
 	FNiagaraVariable IDVar(FNiagaraTypeDefinition::GetIDDef(), "ID");
 	FNiagaraDataSetAccessor<FNiagaraID> IDData(EmitterInstance->GetData(), IDVar);
@@ -1128,7 +1181,7 @@ T UNiagaraDataInterfaceParticleRead::RetrieveValueWithCheck(FNiagaraEmitterInsta
 		FNiagaraID ActualID = IDData.GetSafe(ParticleIndex, NIAGARA_INVALID_ID);
 		if (ActualID == ParticleID)
 		{
-			Value = ValueData.GetSafe(ParticleIndex, T());
+			Value = ValueData.GetSafe(ParticleIndex, Type());
 			bValid = true;
 		}
 	}
@@ -1163,12 +1216,15 @@ void UNiagaraDataInterfaceParticleRead::GetParameterDefinitionHLSL(const FNiagar
 		"int {NumParticlesName};\n"
 		"uint {ParticleStrideFloatName};\n"
 		"uint {ParticleStrideIntName};\n"
+		"uint {ParticleStrideHalfName};\n"
 		"int {AcquireTagRegisterIndexName};\n"
 		"Buffer<int> {SpawnedIDsBufferName};\n"
 		"Buffer<int> {IDToIndexTableName};\n"
 		"Buffer<float> {InputFloatBufferName};\n"
 		"Buffer<int> {InputIntBufferName};\n"
-		"int4 {AttributeIndicesName}[{AttributeInt4Count}];\n\n"
+		"Buffer<float> {InputHalfBufferName};\n"
+		"int4 {AttributeIndicesName}[{AttributeInt4Count}];\n"
+		"int4 {AttributeCompressedName}[{AttributeInt4Count}];\n\n"
 	);
 
 	// If we use an int array for the attribute indices, the shader compiler will actually use int4 due to the packing rules,
@@ -1183,12 +1239,15 @@ void UNiagaraDataInterfaceParticleRead::GetParameterDefinitionHLSL(const FNiagar
 	ArgsDeclarations.Add(TEXT("NumParticlesName"), NumParticlesBaseName + ParamInfo.DataInterfaceHLSLSymbol);
 	ArgsDeclarations.Add(TEXT("ParticleStrideFloatName"), ParticleStrideFloatBaseName + ParamInfo.DataInterfaceHLSLSymbol);
 	ArgsDeclarations.Add(TEXT("ParticleStrideIntName"), ParticleStrideIntBaseName + ParamInfo.DataInterfaceHLSLSymbol);
+	ArgsDeclarations.Add(TEXT("ParticleStrideHalfName"), ParticleStrideHalfBaseName + ParamInfo.DataInterfaceHLSLSymbol);
 	ArgsDeclarations.Add(TEXT("AcquireTagRegisterIndexName"), AcquireTagRegisterIndexBaseName + ParamInfo.DataInterfaceHLSLSymbol);
 	ArgsDeclarations.Add(TEXT("SpawnedIDsBufferName"), SpawnedIDsBufferBaseName + ParamInfo.DataInterfaceHLSLSymbol);
 	ArgsDeclarations.Add(TEXT("IDToIndexTableName"), IDToIndexTableBaseName + ParamInfo.DataInterfaceHLSLSymbol);
 	ArgsDeclarations.Add(TEXT("InputFloatBufferName"), InputFloatBufferBaseName + ParamInfo.DataInterfaceHLSLSymbol);
 	ArgsDeclarations.Add(TEXT("InputIntBufferName"), InputIntBufferBaseName + ParamInfo.DataInterfaceHLSLSymbol);
+	ArgsDeclarations.Add(TEXT("InputHalfBufferName"), InputHalfBufferBaseName + ParamInfo.DataInterfaceHLSLSymbol);
 	ArgsDeclarations.Add(TEXT("AttributeIndicesName"), AttributeIndicesBaseName + ParamInfo.DataInterfaceHLSLSymbol);
+	ArgsDeclarations.Add(TEXT("AttributeCompressedName"), AttributeCompressedBaseName + ParamInfo.DataInterfaceHLSLSymbol);
 	ArgsDeclarations.Add(TEXT("AttributeInt4Count"), AttributeInt4Count);
 
 	OutHLSL += FString::Format(FormatDeclarations, ArgsDeclarations);
@@ -1205,7 +1264,16 @@ static bool GenerateGetFunctionHLSL(const FNiagaraDataInterfaceGPUParamInfo& Par
 		"    if(ParticleIndex != -1 && In_ParticleID.AcquireTag == AcquireTag)\n"
 		"    {\n"
 		"        Out_Valid = true;\n"
+		"\n"
+		"        BRANCH\n"
+		"        if (!{AttributeCompressedName}[{AttributeIndexGroup}]{AttributeIndexComponent})\n"
+		"        {\n"
 		"{FetchValueCode}"
+		"        }\n"
+		"        else\n"
+		"        {\n"
+		"{FetchCompressedValueCode}"
+		"        }\n"
 		"    }\n"
 		"    else\n"
 		"    {\n"
@@ -1219,8 +1287,10 @@ static bool GenerateGetFunctionHLSL(const FNiagaraDataInterfaceGPUParamInfo& Par
 
 	const FString ParticleStrideFloatName = ParticleStrideFloatBaseName + ParamInfo.DataInterfaceHLSLSymbol;
 	const FString ParticleStrideIntName = ParticleStrideIntBaseName + ParamInfo.DataInterfaceHLSLSymbol;
+	const FString ParticleStrideHalfName = ParticleStrideHalfBaseName + ParamInfo.DataInterfaceHLSLSymbol;
 	const FString InputFloatBufferName = InputFloatBufferBaseName + ParamInfo.DataInterfaceHLSLSymbol;
 	const FString InputIntBufferName = InputIntBufferBaseName + ParamInfo.DataInterfaceHLSLSymbol;
+	const FString InputHalfBufferName = InputHalfBufferBaseName + ParamInfo.DataInterfaceHLSLSymbol;
 
 	const TCHAR* ComponentTypeName;
 	const TCHAR* InputBufferName;
@@ -1229,8 +1299,8 @@ static bool GenerateGetFunctionHLSL(const FNiagaraDataInterfaceGPUParamInfo& Par
 	{
 		case ENiagaraParticleDataComponentType::Float:
 			ComponentTypeName = TEXT("float");
-			InputBufferName = *InputFloatBufferName;
-			InputBufferStrideName = *ParticleStrideFloatName;
+			InputBufferName = nullptr; // unused
+			InputBufferStrideName = nullptr; // unused
 			break;
 		case ENiagaraParticleDataComponentType::Int:
 			ComponentTypeName = TEXT("int");
@@ -1244,7 +1314,7 @@ static bool GenerateGetFunctionHLSL(const FNiagaraDataInterfaceGPUParamInfo& Par
 			break;
 		default:
 			UE_LOG(LogNiagara, Error, TEXT("Unknown component type %d while generating function %s"), ComponentType, *FunctionInfo.InstanceName);
-	return false;
+			return false;
 	}
 
 	FString ExtraDefaultValues;
@@ -1254,11 +1324,35 @@ static bool GenerateGetFunctionHLSL(const FNiagaraDataInterfaceGPUParamInfo& Par
 	}
 
 	FString FetchValueCode;
-	for (int ComponentIndex = 0; ComponentIndex < NumComponents; ++ComponentIndex)
+	FString FetchCompressedValueCode;
+
+	// deal with floats differently because the data may be in the buffer corresponding to our half precision SRV
+	if (ComponentType == ENiagaraParticleDataComponentType::Float)
 	{
-		const TCHAR* ComponentName = (NumComponents > 1) ? ComponentNames[ComponentIndex] : TEXT("");
-		const FString FetchComponentCode = FString::Printf(TEXT("        Out_Value%s = %s(%s[(RegisterIndex + %d)*%s + ParticleIndex]);\n"), ComponentName, ComponentTypeName, InputBufferName, ComponentIndex, InputBufferStrideName);
-		FetchValueCode += FetchComponentCode;
+		for (int ComponentIndex = 0; ComponentIndex < NumComponents; ++ComponentIndex)
+		{
+			const TCHAR* ComponentName = (NumComponents > 1) ? ComponentNames[ComponentIndex] : TEXT("");
+			const FString FetchComponentCode = FString::Printf(TEXT("            Out_Value%s = %s(%s[(RegisterIndex + %d)*%s + ParticleIndex]);\n"), ComponentName, ComponentTypeName, *InputFloatBufferName, ComponentIndex, *ParticleStrideFloatName);
+			FetchValueCode += FetchComponentCode;
+		}
+
+		for (int ComponentIndex = 0; ComponentIndex < NumComponents; ++ComponentIndex)
+		{
+			const TCHAR* ComponentName = (NumComponents > 1) ? ComponentNames[ComponentIndex] : TEXT("");
+			const FString FetchComponentCode = FString::Printf(TEXT("            Out_Value%s = %s(%s[(RegisterIndex + %d)*%s + ParticleIndex]);\n"), ComponentName, ComponentTypeName, *InputHalfBufferName, ComponentIndex, *ParticleStrideHalfName);
+			FetchCompressedValueCode += FetchComponentCode;
+		}
+	}
+	else
+	{
+		for (int ComponentIndex = 0; ComponentIndex < NumComponents; ++ComponentIndex)
+		{
+			const TCHAR* ComponentName = (NumComponents > 1) ? ComponentNames[ComponentIndex] : TEXT("");
+			const FString FetchComponentCode = FString::Printf(TEXT("        Out_Value%s = %s(%s[(RegisterIndex + %d)*%s + ParticleIndex]);\n"), ComponentName, ComponentTypeName, InputBufferName, ComponentIndex, InputBufferStrideName);
+			FetchValueCode += FetchComponentCode;
+		}
+
+		FetchCompressedValueCode = FetchValueCode;
 	}
 
 	FString ValueTypeName = (NumComponents > 1) ? FString::Printf(TEXT("%s%d"), ComponentTypeName, NumComponents) : FString(ComponentTypeName);
@@ -1267,6 +1361,7 @@ static bool GenerateGetFunctionHLSL(const FNiagaraDataInterfaceGPUParamInfo& Par
 	FuncTemplateArgs.Add(TEXT("FunctionName"), FunctionInfo.InstanceName);
 	FuncTemplateArgs.Add(TEXT("ValueType"), ValueTypeName);
 	FuncTemplateArgs.Add(TEXT("AttributeIndicesName"), AttributeIndicesBaseName + ParamInfo.DataInterfaceHLSLSymbol);
+	FuncTemplateArgs.Add(TEXT("AttributeCompressedName"), AttributeCompressedBaseName + ParamInfo.DataInterfaceHLSLSymbol);
 	FuncTemplateArgs.Add(TEXT("AttributeIndexGroup"), FunctionInstanceIndex / 4);
 	FuncTemplateArgs.Add(TEXT("AttributeIndexComponent"), ComponentNames[FunctionInstanceIndex % 4]);
 	FuncTemplateArgs.Add(TEXT("IDToIndexTableName"), IDToIndexTableBaseName + ParamInfo.DataInterfaceHLSLSymbol);
@@ -1274,6 +1369,7 @@ static bool GenerateGetFunctionHLSL(const FNiagaraDataInterfaceGPUParamInfo& Par
 	FuncTemplateArgs.Add(TEXT("AcquireTagRegisterIndexName"), AcquireTagRegisterIndexBaseName + ParamInfo.DataInterfaceHLSLSymbol);
 	FuncTemplateArgs.Add(TEXT("ParticleStrideIntName"), ParticleStrideIntName);
 	FuncTemplateArgs.Add(TEXT("FetchValueCode"), FetchValueCode);
+	FuncTemplateArgs.Add(TEXT("FetchCompressedValueCode"), FetchCompressedValueCode);
 	FuncTemplateArgs.Add(TEXT("ExtraDefaultValues"), ExtraDefaultValues);
 	
 	OutHLSL += FString::Format(FuncTemplate, FuncTemplateArgs);
