FNM - Added Compression option for Niagara Systems. Thge Option will try to compress some of the parameters to use half precission, which can be a performance benefit on the GPU and sometimes also on 

diff --git a/Engine/Plugins/FX/Niagara/Source/Niagara/Private/NiagaraDataSet.cpp b/Engine/Plugins/FX/Niagara/Source/Niagara/Private/NiagaraDataSet.cpp
--- a/Engine/Plugins/FX/Niagara/Source/Niagara/Private/NiagaraDataSet.cpp
+++ b/Engine/Plugins/FX/Niagara/Source/Niagara/Private/NiagaraDataSet.cpp
@@ -401,18 +401,20 @@ const FNiagaraVariableLayoutInfo* FNiagaraDataSet::GetVariableLayout(const FNiag
 	return VarLayoutIndex != INDEX_NONE ? &CompiledData->VariableLayouts[VarLayoutIndex] : nullptr;
 }
 
-bool FNiagaraDataSet::GetVariableComponentOffsets(const FNiagaraVariable& Var, int32 &FloatStart, int32 &IntStart) const
+bool FNiagaraDataSet::GetVariableComponentOffsets(const FNiagaraVariable& Var, int32 &FloatStart, int32 &IntStart, int32& HalfStart) const
 {
 	const FNiagaraVariableLayoutInfo *Info = GetVariableLayout(Var);
 	if (Info)
 	{
 		FloatStart = Info->FloatComponentStart;
 		IntStart = Info->Int32ComponentStart;
+		HalfStart = Info->HalfComponentStart;
 		return true;
 	}
 
-	FloatStart = -1;
-	IntStart = -1;
+	FloatStart = INDEX_NONE;
+	IntStart = INDEX_NONE;
+	HalfStart = INDEX_NONE;
 	return false;
 }
 
@@ -471,13 +473,13 @@ void FNiagaraDataSet::CopyTo(FNiagaraDataSet& Other, int32 StartIdx, int32 NumIn
 	}
 }
 
-void FNiagaraDataSet::CopyFromGPUReadback(float* GPUReadBackFloat, int* GPUReadBackInt, int32 StartIdx /* = 0 */, int32 NumInstances /* = INDEX_NONE */, uint32 FloatStride, uint32 IntStride)
+void FNiagaraDataSet::CopyFromGPUReadback(const float* GPUReadBackFloat, const int* GPUReadBackInt, const FFloat16* GPUReadBackHalf, int32 StartIdx /* = 0 */, int32 NumInstances /* = INDEX_NONE */, uint32 FloatStride, uint32 IntStride, uint32 HalfStride)
 {
 	check(IsInRenderingThread());
 	check(IsInitialized());//We should be finalized with proper layout information already.
 
 	FNiagaraDataBuffer& DestBuffer = BeginSimulate();
-	DestBuffer.GPUCopyFrom(GPUReadBackFloat, GPUReadBackInt, StartIdx, NumInstances, FloatStride, IntStride);
+	DestBuffer.GPUCopyFrom(GPUReadBackFloat, GPUReadBackInt, GPUReadBackHalf, StartIdx, NumInstances, FloatStride, IntStride, HalfStride);
 	EndSimulate();
 }
  
@@ -491,6 +493,7 @@ FNiagaraDataBuffer::FNiagaraDataBuffer(FNiagaraDataSet* InOwner)
 	, NumInstancesAllocated(0)
 	, FloatStride(0)
 	, Int32Stride(0)
+	, HalfStride(0)
 	, NumSpawnedInstances(0)
 	, IDAcquireTag(0)
 {
@@ -582,30 +585,36 @@ void FNiagaraDataBuffer::Allocate(uint32 InNumInstances, bool bMaintainExisting)
 	uint32 NewInt32Stride = GetSafeComponentBufferSize(InNumInstances * sizeof(int32));
 	int32 NewInt32Num = NewInt32Stride * Owner->GetNumInt32Components();
 
+	uint32 NewHalfStride = GetSafeComponentBufferSize(InNumInstances * sizeof(FFloat16));
+	int32 NewHalfNum = NewHalfStride * Owner->GetNumHalfComponents();
+
 	// Do sizes match?
-	if (NewFloatNum != FloatData.Num() || NewInt32Num != Int32Data.Num())
+	if (NewFloatNum != FloatData.Num() || NewInt32Num != Int32Data.Num() || NewHalfNum != HalfData.Num())
 	{
 		// Do we need to grow or shrink?
-		const bool bGrowData = NewFloatNum > FloatData.Num() || NewInt32Num > Int32Data.Num();
+		const bool bGrowData = NewFloatNum > FloatData.Num() || NewInt32Num > Int32Data.Num() || NewHalfNum > HalfData.Num();
 		const bool bShrinkFloatData = !bGrowData && (GNiagaraDataBufferShrinkFactor * FMath::Max(GNiagaraDataBufferMinSize, NewFloatNum) < FloatData.Max() || !NewFloatNum);
 		const bool bShrinkIntData = !bGrowData && (GNiagaraDataBufferShrinkFactor * FMath::Max(GNiagaraDataBufferMinSize, NewInt32Num) < Int32Data.Max() || !NewInt32Num);
-		if ( bGrowData || bShrinkFloatData || bShrinkIntData )
+		const bool bShrinkHalfData = !bGrowData && (GNiagaraDataBufferShrinkFactor * FMath::Max(GNiagaraDataBufferMinSize, NewHalfNum) < HalfData.Max() || !NewHalfNum);
+		if (bGrowData || bShrinkFloatData || bShrinkIntData || bShrinkHalfData)
 		{
 			NumInstancesAllocated = InNumInstances;
 
-			DEC_MEMORY_STAT_BY(STAT_NiagaraParticleMemory, FloatData.GetAllocatedSize() + Int32Data.GetAllocatedSize());
+			DEC_MEMORY_STAT_BY(STAT_NiagaraParticleMemory, FloatData.GetAllocatedSize() + Int32Data.GetAllocatedSize() + HalfData.GetAllocatedSize());
 			if (bMaintainExisting)
 			{
 				TArray<uint8> NewFloatData;
 				TArray<uint8> NewInt32Data;
+				TArray<uint8> NewHalfData;
 				NewFloatData.SetNum(NewFloatNum);
 				NewInt32Data.SetNum(NewInt32Num);
+				NewHalfData.SetNum(NewHalfNum);
 
 				if (NewFloatStride > 0 && FloatStride > 0)
 				{
 					const int32 FloatComponents = Owner->GetNumFloatComponents();
 					const uint32 BytesToCopy = FMath::Min(NewFloatStride, FloatStride);
-					for (int32 CompIdx=FloatComponents-1; CompIdx >= 0; --CompIdx)
+					for (int32 CompIdx = FloatComponents - 1; CompIdx >= 0; --CompIdx)
 					{
 						const uint8* Src = FloatData.GetData() + FloatStride * CompIdx;
 						uint8* Dst = NewFloatData.GetData() + NewFloatStride * CompIdx;
@@ -617,7 +626,7 @@ void FNiagaraDataBuffer::Allocate(uint32 InNumInstances, bool bMaintainExisting)
 				{
 					const int32 IntComponents = Owner->GetNumInt32Components();
 					const uint32 BytesToCopy = FMath::Min(NewInt32Stride, Int32Stride);
-					for (int32 CompIdx=IntComponents - 1; CompIdx >= 0; --CompIdx)
+					for (int32 CompIdx = IntComponents - 1; CompIdx >= 0; --CompIdx)
 					{
 						const uint8* Src = Int32Data.GetData() + Int32Stride * CompIdx;
 						uint8* Dst = NewInt32Data.GetData() + NewInt32Stride * CompIdx;
@@ -625,15 +634,29 @@ void FNiagaraDataBuffer::Allocate(uint32 InNumInstances, bool bMaintainExisting)
 					}
 				}
 
+				if (NewHalfStride > 0 && HalfStride > 0)
+				{
+					const int32 HalfComponents = Owner->GetNumHalfComponents();
+					const uint32 BytesToCopy = FMath::Min(NewHalfStride, HalfStride);
+					for (int32 CompIdx = HalfComponents - 1; CompIdx >= 0; --CompIdx)
+					{
+						const uint8* Src = HalfData.GetData() + HalfStride * CompIdx;
+						uint8* Dst = NewHalfData.GetData() + NewHalfStride * CompIdx;
+						FMemory::Memcpy(Dst, Src, BytesToCopy);
+					}
+				}
+
 				FloatData = MoveTemp(NewFloatData);
 				Int32Data = MoveTemp(NewInt32Data);
+				HalfData = MoveTemp(NewHalfData);
 			}
 			else
 			{
 				FloatData.SetNum(NewFloatNum, bShrinkFloatData);
 				Int32Data.SetNum(NewInt32Num, bShrinkIntData);
+				HalfData.SetNum(NewHalfNum, bShrinkHalfData);
 			}
-			INC_MEMORY_STAT_BY(STAT_NiagaraParticleMemory, FloatData.GetAllocatedSize() + Int32Data.GetAllocatedSize());
+			INC_MEMORY_STAT_BY(STAT_NiagaraParticleMemory, FloatData.GetAllocatedSize() + Int32Data.GetAllocatedSize() + HalfData.GetAllocatedSize());
 		}
 		// Calculate strides based upon max of instance counts
 		// This allows us to skip building the register table when shrinking
@@ -642,6 +665,7 @@ void FNiagaraDataBuffer::Allocate(uint32 InNumInstances, bool bMaintainExisting)
 			NumInstancesAllocated = FMath::Max(NumInstancesAllocated, InNumInstances);
 			NewFloatStride = GetSafeComponentBufferSize(NumInstancesAllocated * sizeof(float));
 			NewInt32Stride = GetSafeComponentBufferSize(NumInstancesAllocated * sizeof(int32));
+			NewHalfStride = GetSafeComponentBufferSize(NumInstancesAllocated * sizeof(FFloat16));
 		}
 	}
 	else
@@ -649,10 +673,11 @@ void FNiagaraDataBuffer::Allocate(uint32 InNumInstances, bool bMaintainExisting)
 		NumInstancesAllocated = InNumInstances;
 	}
 
-	if ( (NewFloatStride != FloatStride) || (NewInt32Stride != Int32Stride) )
+	if ((NewFloatStride != FloatStride) || (NewInt32Stride != Int32Stride) || (NewHalfStride != HalfStride))
 	{
 		FloatStride = NewFloatStride;
 		Int32Stride = NewInt32Stride;
+		HalfStride = NewHalfStride;
 		BuildRegisterTable();
 	}
 }
@@ -682,6 +707,7 @@ void FNiagaraDataBuffer::AllocateGPU(uint32 InNumInstances, FNiagaraGPUInstanceC
 	// Pack the data so that the space between elements is the padded thread group size
 	FloatStride = PaddedNumInstances * sizeof(float);
 	Int32Stride = PaddedNumInstances * sizeof(int32);
+	HalfStride = PaddedNumInstances * sizeof(FFloat16);
 
 	// This never seems to happen because of the +1 in NiagaraEmitterInstanceBatcher::ResizeBuffersAndGatherResources()
 	if (PaddedNumInstances == 0)
@@ -694,6 +720,11 @@ void FNiagaraDataBuffer::AllocateGPU(uint32 InNumInstances, FNiagaraGPUInstanceC
 		{
 			GPUBufferInt.Release();
 		}
+		if (GPUBufferHalf.Buffer)
+		{
+			GPUBufferHalf.Release();
+		}
+
 		if (GPUIDToIndexTable.Buffer)
 		{
 			GPUIDToIndexTable.Release();
@@ -731,6 +762,14 @@ void FNiagaraDataBuffer::AllocateGPU(uint32 InNumInstances, FNiagaraGPUInstanceC
 				}
 				GPUBufferInt.Initialize(sizeof(int32), NumElementsToAlloc * Owner->GetNumInt32Components(), EPixelFormat::PF_R32_SINT, DataBufferFlags, TEXT("NiagaraIntDataBuffer"));
 			}
+			if (Owner->GetNumHalfComponents())
+			{
+				if (GPUBufferHalf.Buffer)
+				{
+					GPUBufferHalf.Release();
+				}
+				GPUBufferHalf.Initialize(sizeof(FFloat16), NumElementsToAlloc * Owner->GetNumHalfComponents(), EPixelFormat::PF_R16F, BUF_Static);
+			}
 		}
 
 		if (Owner->RequiresPersistentIDs())
@@ -773,6 +812,13 @@ void FNiagaraDataBuffer::SwapInstances(uint32 OldIndex, uint32 NewIndex)
 		*Dst = *Src;
 		*Src = Temp;
 	}
+	uint32 HalfComponents = Owner->GetNumHalfComponents();
+	for (uint32 CompIdx = 0; CompIdx < HalfComponents; ++CompIdx)
+	{
+		FFloat16* Src = GetInstancePtrHalf(CompIdx, OldIndex);
+		FFloat16* Dst = GetInstancePtrHalf(CompIdx, NewIndex);
+		Swap(*Dst, *Src);
+	}
 }
 
 void FNiagaraDataBuffer::KillInstance(uint32 InstanceIdx)
@@ -795,6 +841,13 @@ void FNiagaraDataBuffer::KillInstance(uint32 InstanceIdx)
 		int32* Dst = GetInstancePtrInt32(CompIdx, InstanceIdx);
 		*Dst = *Src;
 	}
+	uint32 HalfComponents = Owner->GetNumHalfComponents();
+	for (uint32 CompIdx = 0; CompIdx < HalfComponents; ++CompIdx)
+	{
+		const FFloat16* Src = GetInstancePtrHalf(CompIdx, NumInstances);
+		FFloat16* Dst = GetInstancePtrHalf(CompIdx, InstanceIdx);
+		*Dst = *Src;
+	}
 
 #if NIAGARA_NAN_CHECKING
 	CheckForNaNs();
@@ -854,6 +907,23 @@ void FNiagaraDataBuffer::CopyTo(FNiagaraDataBuffer& DestBuffer, int32 StartIdx,
 			size_t Count = SrcEnd - SrcStart;
 			FMemory::Memcpy(Dst, SrcStart, Count * sizeof(int32));
 
+			if (Count > 0)
+			{
+				for (size_t i = 0; i < Count; i++)
+				{
+					checkSlow(SrcStart[i] == Dst[i]);
+				}
+			}
+		}
+		uint32 HalfComponents = Owner->GetNumHalfComponents();
+		for (uint32 CompIdx = 0; CompIdx < HalfComponents; ++CompIdx)
+		{
+			const FFloat16* SrcStart = GetInstancePtrHalf(CompIdx, StartIdx);
+			const FFloat16* SrcEnd = GetInstancePtrHalf(CompIdx, StartIdx + InstancesToCopy);
+			FFloat16* Dst = DestBuffer.GetInstancePtrHalf(CompIdx, DestStartIdx);
+			size_t Count = SrcEnd - SrcStart;
+			FMemory::Memcpy(Dst, SrcStart, Count * sizeof(FFloat16));
+
 			if (Count > 0)
 			{
 				for (size_t i = 0; i < Count; i++)
@@ -865,7 +935,7 @@ void FNiagaraDataBuffer::CopyTo(FNiagaraDataBuffer& DestBuffer, int32 StartIdx,
 	}
 }
 
-void FNiagaraDataBuffer::GPUCopyFrom(float* GPUReadBackFloat, int* GPUReadBackInt, int32 InStartIdx, int32 InNumInstances, uint32 InSrcFloatStride, uint32 InSrcIntStride)
+void FNiagaraDataBuffer::GPUCopyFrom(const float* GPUReadBackFloat, const int* GPUReadBackInt, const FFloat16* GPUReadBackHalf, int32 InStartIdx, int32 InNumInstances, uint32 InSrcFloatStride, uint32 InSrcIntStride, uint32 InSrcHalfStride)
 {
 	//CheckUsage(false); //Have to disable this as in this specific case we write to a "CPUSim" from the RT.
 
@@ -912,6 +982,28 @@ void FNiagaraDataBuffer::GPUCopyFrom(float* GPUReadBackFloat, int* GPUReadBackIn
 			size_t Count = SrcEnd - SrcStart;
 			FMemory::Memcpy(Dst, SrcStart, Count * sizeof(int32));
 
+			if (Count > 0)
+			{
+				for (size_t i = 0; i < Count; i++)
+				{
+					check(SrcStart[i] == Dst[i]);
+				}
+			}
+		}
+	}
+	if (GPUReadBackHalf)
+	{
+		uint32 HalfComponents = Owner->GetNumHalfComponents();
+		for (uint32 CompIdx = 0; CompIdx < HalfComponents; ++CompIdx)
+		{
+			// We have to reimplement the logic from GetInstancePtrInt here because the incoming stride may be different than this 
+			// data buffer's stride.
+			const FFloat16* SrcStart = (const FFloat16*)((uint8*)GPUReadBackHalf + InSrcHalfStride * CompIdx) + InStartIdx;
+			const FFloat16* SrcEnd = (const FFloat16*)((uint8*)GPUReadBackHalf + InSrcHalfStride * CompIdx) + InStartIdx + InNumInstances;
+			FFloat16* Dst = GetInstancePtrHalf(CompIdx, 0);
+			size_t Count = SrcEnd - SrcStart;
+			FMemory::Memcpy(Dst, SrcStart, Count * sizeof(FFloat16));
+
 			if (Count > 0)
 			{
 				for (size_t i = 0; i < Count; i++)
@@ -993,14 +1085,16 @@ void FNiagaraDataBuffer::SetShaderParams(FNiagaraShader *Shader, FRHICommandList
 	{
 		const bool InstancesAllocated = GetNumInstancesAllocated() > 0;
 
-		SetSRVParameter(CommandList, ComputeShader, Shader->FloatInputBufferParam, InstancesAllocated ? GetGPUBufferFloat().SRV.GetReference() : FNiagaraRenderer::GetDummyFloatBuffer());
-		SetSRVParameter(CommandList, ComputeShader, Shader->IntInputBufferParam, InstancesAllocated ? GetGPUBufferInt().SRV.GetReference() : FNiagaraRenderer::GetDummyIntBuffer());
+		SetSRVParameter(CommandList, ComputeShader, Shader->FloatInputBufferParam, InstancesAllocated ? GetGPUBufferFloat().SRV : FNiagaraRenderer::GetDummyFloatBuffer());
+		SetSRVParameter(CommandList, ComputeShader, Shader->IntInputBufferParam, InstancesAllocated ? GetGPUBufferInt().SRV : FNiagaraRenderer::GetDummyIntBuffer());
+		SetSRVParameter(CommandList, ComputeShader, Shader->HalfInputBufferParam, InstancesAllocated ? GetGPUBufferHalf().SRV : FNiagaraRenderer::GetDummyHalfBuffer());
 		SetShaderValue(CommandList, ComputeShader, Shader->ComponentBufferSizeReadParam, SafeBufferSize);
 	}
 	else
 	{
 		Shader->FloatOutputBufferParam.SetBuffer(CommandList, ComputeShader, GetGPUBufferFloat());
 		Shader->IntOutputBufferParam.SetBuffer(CommandList, ComputeShader, GetGPUBufferInt());
+		Shader->HalfOutputBufferParam.SetBuffer(CommandList, ComputeShader, GetGPUBufferHalf());
 		SetShaderValue(CommandList, ComputeShader, Shader->ComponentBufferSizeWriteParam, SafeBufferSize);
 		if (Shader->IDToIndexBufferParam.IsUAVBound())
 		{
@@ -1023,6 +1117,14 @@ void FNiagaraDataBuffer::UnsetShaderParams(FNiagaraShader *Shader, FRHICommandLi
 		//RHICmdList.TransitionResource(EResourceTransitionAccess::EReadable, EResourceTransitionPipeline::EGfxToGfx, CurrDataRender().GetGPUBufferFloat()->UAV);
 	}
 
+	if (Shader->HalfOutputBufferParam.IsUAVBound())
+	{
+#if !PLATFORM_PS4
+		Shader->HalfOutputBufferParam.UnsetUAV(RHICmdList, ShaderRHI);
+#endif
+		//RHICmdList.TransitionResource(EResourceTransitionAccess::EReadable, EResourceTransitionPipeline::EGfxToGfx, CurrDataRender().GetGPUBufferFloat()->UAV);
+	}
+
 	if (Shader->IntOutputBufferParam.IsUAVBound())
 	{
 #if !PLATFORM_PS4
@@ -1099,7 +1201,18 @@ void FScopedNiagaraDataSetGPUReadback::ReadbackData(NiagaraEmitterInstanceBatche
 				FMemory::Memcpy(DataBuffer->Int32Data.GetData(), CPUIntBuffer, GPUIntBuffer.NumBytes);
 				RHICmdList.UnlockVertexBuffer(GPUIntBuffer.Buffer);
 			}
-		}
+	
+			// Read half data
+			const FRWBuffer& GPUHalfBuffer = DataBuffer->GetGPUBufferHalf();
+			if (GPUHalfBuffer.Buffer.IsValid())
+			{
+				DataBuffer->HalfData.AddUninitialized(GPUHalfBuffer.NumBytes);
+
+				void* CPUHalfBuffer = RHICmdList.LockVertexBuffer(GPUHalfBuffer.Buffer, 0, GPUHalfBuffer.NumBytes, RLM_ReadOnly);
+				FMemory::Memcpy(DataBuffer->HalfData.GetData(), CPUHalfBuffer, GPUHalfBuffer.NumBytes);
+				RHICmdList.UnlockVertexBuffer(GPUHalfBuffer.Buffer);
+			}
+	}
 	);
 	FlushRenderingCommands();
 }
@@ -1109,27 +1222,56 @@ void FScopedNiagaraDataSetGPUReadback::ReadbackData(NiagaraEmitterInstanceBatche
 
 void FNiagaraDataBuffer::BuildRegisterTable()
 {
-	int32 TotalRegisters = Owner->GetNumFloatComponents() + Owner->GetNumInt32Components();
+	int32 TotalRegisters = Owner->GetNumFloatComponents() + Owner->GetNumInt32Components() + Owner->GetNumHalfComponents();
 	RegisterTable.Reset();
 	RegisterTable.SetNumUninitialized(TotalRegisters);
 	int32 NumRegisters = 0;
+
+	memset(RegisterTypeOffsets, 0, sizeof(RegisterTypeOffsets));
 	for (const FNiagaraVariableLayoutInfo& VarLayout : Owner->GetVariableLayouts())
 	{
-		int32 NumFloats = VarLayout.GetNumFloatComponents();
-		int32 NumInts = VarLayout.GetNumInt32Components();
+		const int32 NumFloats = VarLayout.GetNumFloatComponents();
 		for (int32 CompIdx = 0; CompIdx < NumFloats; ++CompIdx)
 		{
 			uint32 CompBufferOffset = VarLayout.FloatComponentStart + CompIdx;
 			uint32 CompRegisterOffset = VarLayout.LayoutInfo.FloatComponentRegisterOffsets[CompIdx];
 			RegisterTable[NumRegisters + CompRegisterOffset] = (uint8*)GetComponentPtrFloat(CompBufferOffset);
+			check(RegisterTable[NumRegisters + CompRegisterOffset]);
 		}
+		NumRegisters += NumFloats;
+	}
+
+	RegisterTypeOffsets[1] = NumRegisters;
+	for (const FNiagaraVariableLayoutInfo& VarLayout : Owner->GetVariableLayouts())
+	{
+		const int32 NumInts = VarLayout.GetNumInt32Components();
 		for (int32 CompIdx = 0; CompIdx < NumInts; ++CompIdx)
 		{
 			uint32 CompBufferOffset = VarLayout.Int32ComponentStart + CompIdx;
 			uint32 CompRegisterOffset = VarLayout.LayoutInfo.Int32ComponentRegisterOffsets[CompIdx];
 			RegisterTable[NumRegisters + CompRegisterOffset] = (uint8*)GetComponentPtrInt32(CompBufferOffset);
+			check(RegisterTable[NumRegisters + CompRegisterOffset]);
 		}
-		NumRegisters += NumFloats + NumInts;
+		NumRegisters += NumInts;
+	}
+		
+	RegisterTypeOffsets[2] = NumRegisters;
+	for (const FNiagaraVariableLayoutInfo& VarLayout : Owner->GetVariableLayouts())
+	{
+		const int32 NumHalfs = VarLayout.GetNumHalfComponents();
+		for (int32 CompIdx = 0; CompIdx < NumHalfs; ++CompIdx)
+		{
+			uint32 CompBufferOffset = VarLayout.HalfComponentStart + CompIdx;
+			uint32 CompRegisterOffset = VarLayout.LayoutInfo.HalfComponentRegisterOffsets[CompIdx];
+			RegisterTable[NumRegisters + CompRegisterOffset] = (uint8*)GetComponentPtrHalf(CompBufferOffset);
+			check(RegisterTable[NumRegisters + CompRegisterOffset]);
+		}
+		NumRegisters += NumHalfs;
+	}
+
+	for (void* Register : RegisterTable)
+	{
+		check(Register);
 	}
 }
 
@@ -1143,6 +1285,7 @@ void FNiagaraDataSetCompiledData::BuildLayout()
 	VariableLayouts.Empty();
 	TotalFloatComponents = 0;
 	TotalInt32Components = 0;
+	TotalHalfComponents = 0;
 
 	VariableLayouts.Reserve(Variables.Num());
 	for (FNiagaraVariable& Var : Variables)
@@ -1151,8 +1294,10 @@ void FNiagaraDataSetCompiledData::BuildLayout()
 		FNiagaraTypeLayoutInfo::GenerateLayoutInfo(VarInfo.LayoutInfo, Var.GetType().GetScriptStruct());
 		VarInfo.FloatComponentStart = TotalFloatComponents;
 		VarInfo.Int32ComponentStart = TotalInt32Components;
+		VarInfo.HalfComponentStart = TotalHalfComponents;
 		TotalFloatComponents += VarInfo.GetNumFloatComponents();
 		TotalInt32Components += VarInfo.GetNumInt32Components();
+		TotalHalfComponents += VarInfo.GetNumHalfComponents();
 	}
 }
 
@@ -1166,6 +1311,7 @@ void FNiagaraDataSetCompiledData::Empty()
 	bRequiresPersistentIDs = 0;
 	TotalFloatComponents = 0;
 	TotalInt32Components = 0;
+	TotalHalfComponents = 0;
 	Variables.Empty();
 	VariableLayouts.Empty();
 	ID = FNiagaraDataSetID();
diff --git a/Engine/Plugins/FX/Niagara/Source/Niagara/Private/NiagaraEmitterInstanceBatcher.cpp b/Engine/Plugins/FX/Niagara/Source/Niagara/Private/NiagaraEmitterInstanceBatcher.cpp
--- a/Engine/Plugins/FX/Niagara/Source/Niagara/Private/NiagaraEmitterInstanceBatcher.cpp
+++ b/Engine/Plugins/FX/Niagara/Source/Niagara/Private/NiagaraEmitterInstanceBatcher.cpp
@@ -299,6 +299,10 @@ void NiagaraEmitterInstanceBatcher::TransitionBuffers(FRHICommandList& RHICmdLis
 		{
 			ReadBuffers.Add(SimStageData.Source->GetGPUBufferFloat().UAV);
 		}
+		if (ComputeShader->HalfInputBufferParam.IsBound())
+		{
+			ReadBuffers.Add(SimStageData.Source->GetGPUBufferHalf().UAV);
+		}
 		if (ComputeShader->IntInputBufferParam.IsBound())
 		{
 			ReadBuffers.Add(SimStageData.Source->GetGPUBufferInt().UAV);
@@ -310,6 +314,10 @@ void NiagaraEmitterInstanceBatcher::TransitionBuffers(FRHICommandList& RHICmdLis
 	{
 		WriteBuffers.Add(SimStageData.Destination->GetGPUBufferFloat().UAV);
 	}
+	if (ComputeShader->HalfOutputBufferParam.IsBound())
+	{
+		WriteBuffers.Add(SimStageData.Destination->GetGPUBufferHalf().UAV);
+	}
 	if (ComputeShader->IntOutputBufferParam.IsBound())
 	{
 		WriteBuffers.Add(SimStageData.Destination->GetGPUBufferInt().UAV);
@@ -532,6 +540,7 @@ void NiagaraEmitterInstanceBatcher::ResizeBuffersAndGatherResources(FOverlappabl
 
 				Context->SetDataToRender(CurrentData);
 				OutputGraphicsBuffers.Add(CurrentData->GetGPUBufferFloat().UAV);
+				OutputGraphicsBuffers.Add(CurrentData->GetGPUBufferHalf().UAV);
 				OutputGraphicsBuffers.Add(CurrentData->GetGPUBufferInt().UAV);
 			}
 		}
@@ -1091,13 +1100,14 @@ void NiagaraEmitterInstanceBatcher::ProcessDebugInfo(FRHICommandList &RHICmdList
 	if (Context && Context->DebugInfo.IsValid())
 	{
 		// Fire off the readback if not already doing so
-		if (!Context->GPUDebugDataReadbackFloat && !Context->GPUDebugDataReadbackInt && !Context->GPUDebugDataReadbackCounts)
+		if (!Context->GPUDebugDataReadbackFloat && !Context->GPUDebugDataReadbackInt && !Context->GPUDebugDataReadbackHalf && !Context->GPUDebugDataReadbackCounts)
 		{
 			// Do nothing.., handled in Run
 		}
 		// We may not have floats or ints, but we should have at least one of the two
 		else if ((Context->GPUDebugDataReadbackFloat == nullptr || Context->GPUDebugDataReadbackFloat->IsReady()) 
 				&& (Context->GPUDebugDataReadbackInt == nullptr || Context->GPUDebugDataReadbackInt->IsReady())
+				&& (Context->GPUDebugDataReadbackHalf == nullptr || Context->GPUDebugDataReadbackHalf->IsReady())
 				&& Context->GPUDebugDataReadbackCounts->IsReady()
 			)
 		{
@@ -1115,8 +1125,13 @@ void NiagaraEmitterInstanceBatcher::ProcessDebugInfo(FRHICommandList &RHICmdList
 				{
 					IntDataBuffer = static_cast<int*>(Context->GPUDebugDataReadbackInt->Lock(Context->GPUDebugDataIntSize));
 				}
+				FFloat16* HalfDataBuffer = nullptr;
+				if (Context->GPUDebugDataReadbackHalf)
+				{
+					HalfDataBuffer = static_cast<FFloat16*>(Context->GPUDebugDataReadbackHalf->Lock(Context->GPUDebugDataHalfSize));
+				}
 
-				Context->DebugInfo->Frame.CopyFromGPUReadback(FloatDataBuffer, IntDataBuffer, 0, NewExistingDataCount, Context->GPUDebugDataFloatStride, Context->GPUDebugDataIntStride);
+				Context->DebugInfo->Frame.CopyFromGPUReadback(FloatDataBuffer, IntDataBuffer, HalfDataBuffer, NewExistingDataCount, Context->GPUDebugDataFloatStride, Context->GPUDebugDataIntStride, Context->GPUDebugDataHalfStride);
 
 				Context->DebugInfo->bWritten = true;
 
@@ -1128,6 +1143,10 @@ void NiagaraEmitterInstanceBatcher::ProcessDebugInfo(FRHICommandList &RHICmdList
 				{
 					Context->GPUDebugDataReadbackInt->Unlock();
 				}
+				if (Context->GPUDebugDataReadbackHalf)
+				{
+					Context->GPUDebugDataReadbackHalf->Unlock();
+				}
 				Context->GPUDebugDataReadbackCounts->Unlock();
 			}
 			{
@@ -1143,12 +1162,18 @@ void NiagaraEmitterInstanceBatcher::ProcessDebugInfo(FRHICommandList &RHICmdList
 					delete Context->GPUDebugDataReadbackInt;
 					Context->GPUDebugDataReadbackInt = nullptr;
 				}
+				if (Context->GPUDebugDataReadbackInt)
+				{
+					delete Context->GPUDebugDataReadbackHalf;
+					Context->GPUDebugDataReadbackHalf = nullptr;
+				}
 				delete Context->GPUDebugDataReadbackCounts;
 				Context->GPUDebugDataReadbackCounts = nullptr;	
 				Context->GPUDebugDataFloatSize = 0;
 				Context->GPUDebugDataIntSize = 0;
 				Context->GPUDebugDataFloatStride = 0;
 				Context->GPUDebugDataIntStride = 0;
+				Context->GPUDebugDataHalfStride = 0;
 				Context->GPUDebugDataCountOffset = INDEX_NONE;
 			}
 
@@ -1406,7 +1431,7 @@ void NiagaraEmitterInstanceBatcher::Run(const FNiagaraGPUSystemTick& Tick, const
 	{
 		//UE_LOG(LogNiagara, Warning, TEXT("Queued up!"));
 
-		if (!Context->GPUDebugDataReadbackFloat && !Context->GPUDebugDataReadbackInt && !Context->GPUDebugDataReadbackCounts && DestinationData.GetGPUInstanceCountBufferOffset() != INDEX_NONE && SimulationStageIndex == Context->MaxUpdateIterations - 1)
+		if (!Context->GPUDebugDataReadbackFloat && !Context->GPUDebugDataReadbackInt && !Context->GPUDebugDataReadbackHalf && !Context->GPUDebugDataReadbackCounts && DestinationData.GetGPUInstanceCountBufferOffset() != INDEX_NONE && SimulationStageIndex == Context->MaxUpdateIterations - 1)
 		{
 			Context->GPUDebugDataFloatSize = 0;
 			Context->GPUDebugDataIntSize = 0;
@@ -1431,6 +1456,15 @@ void NiagaraEmitterInstanceBatcher::Run(const FNiagaraGPUSystemTick& Tick, const
 				Context->GPUDebugDataIntStride = DestinationData.GetInt32Stride();
 			}
 
+			if (DestinationData.GetGPUBufferHalf().NumBytes > 0)
+			{
+				static const FName ReadbackHalfName(TEXT("Niagara GPU Debug Info Half Emitter Readback"));
+				Context->GPUDebugDataReadbackHalf = new FRHIGPUBufferReadback(ReadbackHalfName);
+				Context->GPUDebugDataReadbackHalf->EnqueueCopy(RHICmdList, DestinationData.GetGPUBufferHalf().Buffer);
+				Context->GPUDebugDataHalfSize = DestinationData.GetGPUBufferHalf().NumBytes;
+				Context->GPUDebugDataHalfStride = DestinationData.GetHalfStride();
+			}
+
 			static const FName ReadbackCountsName(TEXT("Niagara GPU Emitter Readback"));
 			Context->GPUDebugDataReadbackCounts = new FRHIGPUBufferReadback(ReadbackCountsName);
 			Context->GPUDebugDataReadbackCounts->EnqueueCopy(RHICmdList, GPUInstanceCounterManager.GetInstanceCountBuffer().Buffer);
diff --git a/Engine/Plugins/FX/Niagara/Source/Niagara/Private/NiagaraModule.cpp b/Engine/Plugins/FX/Niagara/Source/Niagara/Private/NiagaraModule.cpp
--- a/Engine/Plugins/FX/Niagara/Source/Niagara/Private/NiagaraModule.cpp
+++ b/Engine/Plugins/FX/Niagara/Source/Niagara/Private/NiagaraModule.cpp
@@ -430,6 +430,11 @@ UScriptStruct* FNiagaraTypeDefinition::Vec2Struct;
 UScriptStruct* FNiagaraTypeDefinition::ColorStruct;
 UScriptStruct* FNiagaraTypeDefinition::QuatStruct;
 
+UScriptStruct* FNiagaraTypeDefinition::HalfStruct;
+UScriptStruct* FNiagaraTypeDefinition::HalfVec2Struct;
+UScriptStruct* FNiagaraTypeDefinition::HalfVec3Struct;
+UScriptStruct* FNiagaraTypeDefinition::HalfVec4Struct;
+
 UClass* FNiagaraTypeDefinition::UObjectClass;
 UClass* FNiagaraTypeDefinition::UMaterialClass;
 
@@ -453,6 +458,11 @@ FNiagaraTypeDefinition FNiagaraTypeDefinition::Vec2Def;
 FNiagaraTypeDefinition FNiagaraTypeDefinition::ColorDef;
 FNiagaraTypeDefinition FNiagaraTypeDefinition::QuatDef;
 
+FNiagaraTypeDefinition FNiagaraTypeDefinition::HalfDef;
+FNiagaraTypeDefinition FNiagaraTypeDefinition::HalfVec2Def;
+FNiagaraTypeDefinition FNiagaraTypeDefinition::HalfVec3Def;
+FNiagaraTypeDefinition FNiagaraTypeDefinition::HalfVec4Def;
+
 FNiagaraTypeDefinition FNiagaraTypeDefinition::UObjectDef;
 FNiagaraTypeDefinition FNiagaraTypeDefinition::UMaterialDef;
 
@@ -492,6 +502,11 @@ void FNiagaraTypeDefinition::Init()
 	FNiagaraTypeDefinition::IntStruct = FindObjectChecked<UScriptStruct>(NiagaraPkg, TEXT("NiagaraInt32"));
 	FNiagaraTypeDefinition::Matrix4Struct = FindObjectChecked<UScriptStruct>(NiagaraPkg, TEXT("NiagaraMatrix"));
 
+	FNiagaraTypeDefinition::HalfStruct = FindObjectChecked<UScriptStruct>(NiagaraPkg, TEXT("NiagaraHalf"));
+	FNiagaraTypeDefinition::HalfVec2Struct = FindObjectChecked<UScriptStruct>(NiagaraPkg, TEXT("NiagaraHalfVector2"));
+	FNiagaraTypeDefinition::HalfVec3Struct = FindObjectChecked<UScriptStruct>(NiagaraPkg, TEXT("NiagaraHalfVector3"));
+	FNiagaraTypeDefinition::HalfVec4Struct = FindObjectChecked<UScriptStruct>(NiagaraPkg, TEXT("NiagaraHalfVector4"));
+
 	FNiagaraTypeDefinition::Vec2Struct = FindObjectChecked<UScriptStruct>(CoreUObjectPkg, TEXT("Vector2D"));
 	FNiagaraTypeDefinition::Vec3Struct = FindObjectChecked<UScriptStruct>(CoreUObjectPkg, TEXT("Vector"));
 	FNiagaraTypeDefinition::Vec4Struct = FindObjectChecked<UScriptStruct>(CoreUObjectPkg, TEXT("Vector4"));
@@ -514,6 +529,11 @@ void FNiagaraTypeDefinition::Init()
 	QuatDef = FNiagaraTypeDefinition(QuatStruct);
 	Matrix4Def = FNiagaraTypeDefinition(Matrix4Struct);
 
+	HalfDef = FNiagaraTypeDefinition(HalfStruct);
+	HalfVec2Def = FNiagaraTypeDefinition(HalfVec2Struct);
+	HalfVec3Def = FNiagaraTypeDefinition(HalfVec3Struct);
+	HalfVec4Def = FNiagaraTypeDefinition(HalfVec4Struct);
+
 	UObjectDef = FNiagaraTypeDefinition(UObjectClass);
 	UMaterialDef = FNiagaraTypeDefinition(UMaterialClass);
 
@@ -551,6 +571,7 @@ void FNiagaraTypeDefinition::Init()
 	ScalarStructs.Add(BoolStruct);
 	ScalarStructs.Add(IntStruct);
 	ScalarStructs.Add(FloatStruct);
+	ScalarStructs.Add(HalfStruct);
 
 	ExecutionStateEnum = StaticEnum<ENiagaraExecutionState>();
 	ExecutionStateSourceEnum = StaticEnum<ENiagaraExecutionStateSource>();
@@ -652,6 +673,7 @@ void FNiagaraTypeDefinition::RecreateUserDefinedTypeRegistry()
 	FNiagaraTypeRegistry::Register(IDDef, true, true, false);
 	FNiagaraTypeRegistry::Register(NumericDef, true, false, false);
 	FNiagaraTypeRegistry::Register(FloatDef, true, true, false);
+	FNiagaraTypeRegistry::Register(HalfDef, true, true, false);
 	FNiagaraTypeRegistry::Register(IntDef, true, true, false);
 	FNiagaraTypeRegistry::Register(BoolDef, true, true, false);
 	FNiagaraTypeRegistry::Register(Vec2Def, true, true, false);
diff --git a/Engine/Plugins/FX/Niagara/Source/Niagara/Private/NiagaraParameterStore.cpp b/Engine/Plugins/FX/Niagara/Source/Niagara/Private/NiagaraParameterStore.cpp
--- a/Engine/Plugins/FX/Niagara/Source/Niagara/Private/NiagaraParameterStore.cpp
+++ b/Engine/Plugins/FX/Niagara/Source/Niagara/Private/NiagaraParameterStore.cpp
@@ -26,6 +26,8 @@ static FAutoConsoleVariableRef CVarNiagaraDumpParticleParameterStores(
 
 struct FNiagaraVariableSearch
 {
+	typedef int32(*VariableCompareFunction)(const FNiagaraVariableBase&, const FNiagaraVariableBase&);
+
 	static FORCEINLINE int32 Compare(const FNiagaraVariableBase& A, const FNiagaraVariableBase& B)
 	{
 #if NIAGARA_VARIABLE_LEXICAL_SORTING
@@ -47,7 +49,18 @@ struct FNiagaraVariableSearch
 		}
 	}
 
-	static bool Find(const FNiagaraVariableWithOffset* Variables, const FNiagaraVariableBase& Ref, int32 Start, int32 Num, int32& CheckIndex)
+	static FORCEINLINE int32 CompareIgnoreType(const FNiagaraVariableBase& A, const FNiagaraVariableBase& B)
+	{
+#if NIAGARA_VARIABLE_LEXICAL_SORTING
+		int32 ComparisonDiff = A.GetName().Compare(B.GetName());
+#else
+		int32 ComparisonDiff = A.GetName().CompareIndexes(B.GetName());
+#endif
+
+		return ComparisonDiff;
+	}
+
+	static bool FindInternal(VariableCompareFunction CompareFn, const FNiagaraVariableWithOffset* Variables, const FNiagaraVariableBase& Ref, int32 Start, int32 Num, int32& CheckIndex)
 	{
 		while (Num)
 		{
@@ -57,7 +70,7 @@ struct FNiagaraVariableSearch
 			CheckIndex = Start + Num;
 			const int32 StartIfLess = CheckIndex + LeftoverSize;
 
-			const int32 ComparisonDiff = Compare(Variables[CheckIndex], Ref);
+			const int32 ComparisonDiff = CompareFn(Variables[CheckIndex], Ref);
 			if (ComparisonDiff < 0)
 			{
 				Start = CheckIndex + 1;
@@ -71,6 +84,15 @@ struct FNiagaraVariableSearch
 		CheckIndex = Start;
 		return false;
 	}
+
+	static FORCEINLINE bool Find(const FNiagaraVariableWithOffset* Variables, const FNiagaraVariableBase& Ref, int32 Start, int32 Num, bool IgnoreType, int32& CheckIndex)
+	{
+		if (IgnoreType)
+		{
+			return FindInternal(CompareIgnoreType, Variables, Ref, Start, Num, CheckIndex);
+		}
+		return FindInternal(Compare, Variables, Ref, Start, Num, CheckIndex);
+	}
 };
 
 //////////////////////////////////////////////////////////////////////////
@@ -475,7 +497,7 @@ bool FNiagaraParameterStore::AddParameter(const FNiagaraVariable& Param, bool bI
 	int32 InsertPos = 0;
 	if (SortedParameterOffsets.Num())
 	{
-		const bool bAlreadyIn = FNiagaraVariableSearch::Find(SortedParameterOffsets.GetData(), ParamNoData, 0, SortedParameterOffsets.Num(), InsertPos);
+		const bool bAlreadyIn = FNiagaraVariableSearch::Find(SortedParameterOffsets.GetData(), ParamNoData, 0, SortedParameterOffsets.Num(), false /* IgnoreType */, InsertPos);
 		if (bAlreadyIn)
 		{
 			if (OutOffset)
@@ -926,7 +948,7 @@ const FNiagaraVariableBase* FNiagaraParameterStore::FindVariable(UNiagaraDataInt
 	return nullptr;
 }
 
-const int32* FNiagaraParameterStore::FindParameterOffset(const FNiagaraVariable& Parameter) const
+const int32* FNiagaraParameterStore::FindParameterOffset(const FNiagaraVariable& Parameter, bool IgnoreType) const
 {
 #if WITH_EDITORONLY_DATA
 	check(!ParameterOffsets.Num()); // Migration to SortedParameterOffsets
@@ -935,7 +957,7 @@ const int32* FNiagaraParameterStore::FindParameterOffset(const FNiagaraVariable&
 	if (SortedParameterOffsets.Num())
 	{
 		int32 MatchingIndex = 0;
-		if (FNiagaraVariableSearch::Find(SortedParameterOffsets.GetData(), Parameter, 0, SortedParameterOffsets.Num(), MatchingIndex))
+		if (FNiagaraVariableSearch::Find(SortedParameterOffsets.GetData(), Parameter, 0, SortedParameterOffsets.Num(), IgnoreType, MatchingIndex))
 		{
 			return &SortedParameterOffsets[MatchingIndex].Offset;
 		}
diff --git a/Engine/Plugins/FX/Niagara/Source/Niagara/Private/NiagaraRenderer.cpp b/Engine/Plugins/FX/Niagara/Source/Niagara/Private/NiagaraRenderer.cpp
--- a/Engine/Plugins/FX/Niagara/Source/Niagara/Private/NiagaraRenderer.cpp
+++ b/Engine/Plugins/FX/Niagara/Source/Niagara/Private/NiagaraRenderer.cpp
@@ -155,18 +155,47 @@ FRHIShaderResourceView* FNiagaraRenderer::GetDummyTextureReadBuffer2D()
 	return DummyTextureReadBuffer2D.SRV;
 }
 
+FRHIShaderResourceView* FNiagaraRenderer::GetDummyHalfBuffer()
+{
+	check(IsInRenderingThread());
+	static TGlobalResource<FNiagaraEmptyBufferSRV> DummyHalfBuffer(PF_R16F, TEXT("NiagaraRenderer::DummyHalf"));
+	return DummyHalfBuffer.SRV;
+}
+
+
 bool FNiagaraRenderer::SetVertexFactoryVariable(const FNiagaraDataSet& DataSet, const FNiagaraVariable& Var, int32 VFVarOffset)
 {
+	// use the DataSetVariable to figure out the information about the data that we'll be sending to the renderer
+	const auto& DataSetVariables = DataSet.GetVariables();
+	int32 VariableMatch = FNiagaraVariable::SearchArrayForPartialNameMatch(DataSetVariables, Var.GetName());
+	if (VariableMatch == INDEX_NONE)
+	{
+		VFVariables[VFVarOffset] = FNiagaraRendererVariableInfo(INDEX_NONE, INDEX_NONE, 0, false, false);
+		return false;
+	}
+	const auto& DataSetVariable = DataSetVariables[VariableMatch];
+	const FNiagaraTypeDefinition& VarType = DataSetVariable.GetType();
+
+	const bool HalfVariable = VarType == FNiagaraTypeDefinition::GetHalfDef()
+		|| VarType == FNiagaraTypeDefinition::GetHalfVec2Def()
+		|| VarType == FNiagaraTypeDefinition::GetHalfVec3Def()
+		|| VarType == FNiagaraTypeDefinition::GetHalfVec4Def();
+
+
 	int32 FloatOffset;
+	int32 HalfOffset;
 	int32 IntOffset;//TODO: No VF uses ints atm but it should be trivial to copy the float path if some vf should need to.
-	DataSet.GetVariableComponentOffsets(Var, FloatOffset, IntOffset);
-	int32 NumComponents = Var.GetSizeInBytes() / sizeof(float);
+	DataSet.GetVariableComponentOffsets(DataSetVariable, FloatOffset, IntOffset, HalfOffset);
+	int32 VarSize = HalfVariable ? sizeof(FFloat16) : sizeof(float);
+	int32 NumComponents = DataSetVariable.GetSizeInBytes() / VarSize;
+	int32 Offset = HalfVariable ? HalfOffset : FloatOffset;
+	int32& TotalVFComponents = HalfVariable ? TotalVFHalfComponents : TotalVFFloatComponents;
 
 	int32 GPULocation = INDEX_NONE;
 	bool bUpload = true;
-	if (FloatOffset != INDEX_NONE)
+	if (Offset != INDEX_NONE)
 	{
-		if (FNiagaraRendererVariableInfo* ExistingVarInfo = VFVariables.FindByPredicate([&](const FNiagaraRendererVariableInfo& VarInfo) { return VarInfo.DatasetOffset == FloatOffset; }))
+		if (FNiagaraRendererVariableInfo* ExistingVarInfo = VFVariables.FindByPredicate([&](const FNiagaraRendererVariableInfo& VarInfo) { return VarInfo.DatasetOffset == Offset && VarInfo.bHalfType == HalfVariable; }))
 		{
 			//Don't need to upload this var again if it's already been uploaded for another var info. Just point to that.
 			//E.g. when custom sorting uses age.
@@ -176,35 +205,54 @@ bool FNiagaraRenderer::SetVertexFactoryVariable(const FNiagaraDataSet& DataSet,
 		else
 		{
 			//For CPU Sims we pack just the required data tightly in a GPU buffer we upload. For GPU sims the data is there already so we just provide the real data location.
-			GPULocation = SimTarget == ENiagaraSimTarget::CPUSim ? TotalVFComponents : FloatOffset;
+			GPULocation = SimTarget == ENiagaraSimTarget::CPUSim ? TotalVFComponents : Offset;
 			TotalVFComponents += NumComponents;
 		}
 	}
 
-	VFVariables[VFVarOffset] = FNiagaraRendererVariableInfo(FloatOffset, GPULocation, NumComponents, bUpload);
+	VFVariables[VFVarOffset] = FNiagaraRendererVariableInfo(Offset, GPULocation, NumComponents, bUpload, HalfVariable);
 
-	return FloatOffset != INDEX_NONE;
+	return Offset != INDEX_NONE;
 }
 
-FGlobalDynamicReadBuffer::FAllocation FNiagaraRenderer::TransferDataToGPU(FGlobalDynamicReadBuffer& DynamicReadBuffer, FNiagaraDataBuffer* SrcData)const
+FParticleRenderData FNiagaraRenderer::TransferDataToGPU(FGlobalDynamicReadBuffer& DynamicReadBuffer, FNiagaraDataBuffer* SrcData)const
 {
-	int32 TotalFloatSize = TotalVFComponents * SrcData->GetNumInstances();
-	int32 ComponentStrideDest = SrcData->GetNumInstances() * sizeof(float);
-	FGlobalDynamicReadBuffer::FAllocation Allocation = DynamicReadBuffer.AllocateFloat(TotalFloatSize);
+	int32 TotalFloatSize = TotalVFFloatComponents * SrcData->GetNumInstances();
+	int32 TotalHalfSize = TotalVFHalfComponents * SrcData->GetNumInstances();
+
+	int32 ComponentHalfStrideDest = SrcData->GetNumInstances() * sizeof(FFloat16);
+	int32 ComponentFloatStrideDest = SrcData->GetNumInstances() * sizeof(float);
+
+	FGlobalDynamicReadBuffer::FAllocation FloatAllocation = TotalFloatSize ? DynamicReadBuffer.AllocateFloat(TotalFloatSize) : FGlobalDynamicReadBuffer::FAllocation();
+	FGlobalDynamicReadBuffer::FAllocation HalfAllocation = TotalHalfSize ? DynamicReadBuffer.AllocateHalf(TotalHalfSize) : FGlobalDynamicReadBuffer::FAllocation();
+
 	for (const FNiagaraRendererVariableInfo& VarInfo : VFVariables)
 	{
 		int32 GpuOffset = VarInfo.GetGPUOffset();
 		if (GpuOffset != INDEX_NONE && VarInfo.bUpload)
 		{
-			for (int32 CompIdx = 0; CompIdx < VarInfo.NumComponents; ++CompIdx)
+			if (VarInfo.bHalfType)
 			{
-				float* SrcComponent = (float*)SrcData->GetComponentPtrFloat(VarInfo.DatasetOffset + CompIdx);
-				void* Dest = Allocation.Buffer + ComponentStrideDest * (GpuOffset + CompIdx);
-				FMemory::Memcpy(Dest, SrcComponent, ComponentStrideDest);
+				GpuOffset &= ~(1 << 31);
+				for (int32 CompIdx = 0; CompIdx < VarInfo.NumComponents; ++CompIdx)
+				{
+					FFloat16* SrcComponent = (FFloat16*)SrcData->GetComponentPtrHalf(VarInfo.DatasetOffset + CompIdx);
+					void* Dest = HalfAllocation.Buffer + ComponentHalfStrideDest * (GpuOffset + CompIdx);
+					FMemory::Memcpy(Dest, SrcComponent, ComponentHalfStrideDest);
+				}
+			}
+			else
+			{
+				for (int32 CompIdx = 0; CompIdx < VarInfo.NumComponents; ++CompIdx)
+				{
+					float* SrcComponent = (float*)SrcData->GetComponentPtrFloat(VarInfo.DatasetOffset + CompIdx);
+					void* Dest = FloatAllocation.Buffer + ComponentFloatStrideDest * (GpuOffset + CompIdx);
+					FMemory::Memcpy(Dest, SrcComponent, ComponentFloatStrideDest);
+				}
 			}
 		}
 	}
-	return Allocation;
+	return FParticleRenderData{ FloatAllocation, HalfAllocation };
 }
 
 //////////////////////////////////////////////////////////////////////////
@@ -271,7 +319,8 @@ FNiagaraRenderer::FNiagaraRenderer(ERHIFeatureLevel::Type InFeatureLevel, const
 	, SimTarget(Emitter->GetCachedEmitter()->SimTarget)
 	, NumIndicesPerInstance(InProps ? InProps->GetNumIndicesPerInstance() : 0)
 	, FeatureLevel(InFeatureLevel)
-	, TotalVFComponents(0)
+	, TotalVFHalfComponents(0)
+	, TotalVFFloatComponents(0)
 {
 #if STATS
 	EmitterStatID = Emitter->GetCachedEmitter()->GetStatID(false, false);
@@ -363,6 +412,17 @@ struct FParticleOrderAsUint
 
 		Index = InIndex;
 	}
+
+	template <bool bStrictlyPositive, bool bAscending>
+	FORCEINLINE_DEBUGGABLE void SetAsUint(int32 InIndex, FFloat16 InOrder) 
+	{
+		const uint32 SortKeySignBit = 0x8000;
+		uint32 InOrderAsUint = InOrder.Encoded;
+		InOrderAsUint = (bStrictlyPositive || InOrder.Components.Sign != 0) ? (InOrderAsUint | SortKeySignBit) : ~InOrderAsUint;
+		OrderAsUint = bAscending ? InOrderAsUint : ~InOrderAsUint;
+		OrderAsUint &= 0xFFFF;
+		Index = InIndex;
+	}
 		
 	FORCEINLINE operator uint32() const { return OrderAsUint; }
 };
@@ -378,6 +438,7 @@ void FNiagaraRenderer::SortIndices(const FNiagaraGPUSortInfo& SortInfo, int32 So
 	check(SortInfo.SortAttributeOffset != INDEX_NONE);
 
 	const bool bUseRadixSort = GNiagaraRadixSortThreshold != -1 && (int32)NumInstances  > GNiagaraRadixSortThreshold;
+	const bool bSortVarIsHalf = VFVariables[SortVarIdx].bHalfType;
 
 	int32* RESTRICT IndexBuffer = (int32*)(OutIndices.Buffer);
 
@@ -386,46 +447,95 @@ void FNiagaraRenderer::SortIndices(const FNiagaraGPUSortInfo& SortInfo, int32 So
 
 	if (SortInfo.SortMode == ENiagaraSortMode::ViewDepth || SortInfo.SortMode == ENiagaraSortMode::ViewDistance)
 	{
-		int32 BaseCompOffset = VFVariables[SortVarIdx].DatasetOffset;
-		float* RESTRICT PositionX = (float*)Buffer.GetComponentPtrFloat(BaseCompOffset);
-		float* RESTRICT PositionY = (float*)Buffer.GetComponentPtrFloat(BaseCompOffset + 1);
-		float* RESTRICT PositionZ = (float*)Buffer.GetComponentPtrFloat(BaseCompOffset + 2);
-		auto GetPos = [&PositionX, &PositionY, &PositionZ](int32 Idx)
+		if (bSortVarIsHalf)
 		{
-			return FVector(PositionX[Idx], PositionY[Idx], PositionZ[Idx]);
-		};
+			const int32 BaseCompOffset = VFVariables[SortVarIdx].DatasetOffset;
+			FFloat16* RESTRICT PositionX = (FFloat16*)Buffer.GetComponentPtrHalf(BaseCompOffset);
+			FFloat16* RESTRICT PositionY = (FFloat16*)Buffer.GetComponentPtrHalf(BaseCompOffset + 1);
+			FFloat16* RESTRICT PositionZ = (FFloat16*)Buffer.GetComponentPtrHalf(BaseCompOffset + 2);
+			auto GetPos = [&PositionX, &PositionY, &PositionZ](int32 Idx)
+			{
+				return FVector(PositionX[Idx], PositionY[Idx], PositionZ[Idx]);
+			};
 
-		if (SortInfo.SortMode == ENiagaraSortMode::ViewDepth)
-		{
-			for (uint32 i = 0; i < NumInstances; ++i)
+			if (SortInfo.SortMode == ENiagaraSortMode::ViewDepth)
 			{
-				ParticleOrder[i].SetAsUint<true, false>(i, FVector::DotProduct(GetPos(i) - SortInfo.ViewOrigin, SortInfo.ViewDirection));
+				for (uint32 i = 0; i < NumInstances; ++i)
+				{
+					ParticleOrder[i].SetAsUint<true, false>(i, FVector::DotProduct(GetPos(i) - SortInfo.ViewOrigin, SortInfo.ViewDirection));
+				}
+			}
+			else
+			{
+				for (uint32 i = 0; i < NumInstances; ++i)
+				{
+					ParticleOrder[i].SetAsUint<true, false>(i, (GetPos(i) - SortInfo.ViewOrigin).SizeSquared());
+				}
 			}
 		}
 		else
 		{
-			for (uint32 i = 0; i < NumInstances; ++i)
+			const int32 BaseCompOffset = VFVariables[SortVarIdx].DatasetOffset;
+			float* RESTRICT PositionX = (float*)Buffer.GetComponentPtrFloat(BaseCompOffset);
+			float* RESTRICT PositionY = (float*)Buffer.GetComponentPtrFloat(BaseCompOffset + 1);
+			float* RESTRICT PositionZ = (float*)Buffer.GetComponentPtrFloat(BaseCompOffset + 2);
+			auto GetPos = [&PositionX, &PositionY, &PositionZ](int32 Idx)
+			{
+				return FVector(PositionX[Idx], PositionY[Idx], PositionZ[Idx]);
+			};
+
+			if (SortInfo.SortMode == ENiagaraSortMode::ViewDepth)
 			{
-				ParticleOrder[i].SetAsUint<true, false>(i, (GetPos(i) - SortInfo.ViewOrigin).SizeSquared());
+				for (uint32 i = 0; i < NumInstances; ++i)
+				{
+					ParticleOrder[i].SetAsUint<true, false>(i, FVector::DotProduct(GetPos(i) - SortInfo.ViewOrigin, SortInfo.ViewDirection));
+				}
+			}
+			else
+			{
+				for (uint32 i = 0; i < NumInstances; ++i)
+				{
+					ParticleOrder[i].SetAsUint<true, false>(i, (GetPos(i) - SortInfo.ViewOrigin).SizeSquared());
+				}
 			}
 		}
 	}
 	else
 	{
-		float* RESTRICT CustomSorting = (float*)Buffer.GetComponentPtrFloat(VFVariables[SortVarIdx].DatasetOffset);
-
-		if (SortInfo.SortMode == ENiagaraSortMode::CustomAscending)
+		if (bSortVarIsHalf)
 		{
-			for (uint32 i = 0; i < NumInstances; ++i)
+			FFloat16* RESTRICT CustomSorting = (FFloat16*)Buffer.GetComponentPtrHalf(VFVariables[SortVarIdx].DatasetOffset);
+			if (SortInfo.SortMode == ENiagaraSortMode::CustomAscending)
+			{
+				for (uint32 i = 0; i < NumInstances; ++i)
+				{
+					ParticleOrder[i].SetAsUint<false, true>(i, CustomSorting[i]);
+				}
+			}
+			else // ENiagaraSortMode::CustomDecending
 			{
-				ParticleOrder[i].SetAsUint<false, true>(i, CustomSorting[i]);
+				for (uint32 i = 0; i < NumInstances; ++i)
+				{
+					ParticleOrder[i].SetAsUint<false, false>(i, CustomSorting[i]);
+				}
 			}
 		}
-		else // ENiagaraSortMode::CustomDecending
+		else
 		{
-			for (uint32 i = 0; i < NumInstances; ++i)
+			float* RESTRICT CustomSorting = (float*)Buffer.GetComponentPtrFloat(VFVariables[SortVarIdx].DatasetOffset);
+			if (SortInfo.SortMode == ENiagaraSortMode::CustomAscending)
+			{
+				for (uint32 i = 0; i < NumInstances; ++i)
+				{
+					ParticleOrder[i].SetAsUint<false, true>(i, CustomSorting[i]);
+				}
+			}
+			else // ENiagaraSortMode::CustomDecending
 			{
-				ParticleOrder[i].SetAsUint<false, false>(i, CustomSorting[i]);
+				for (uint32 i = 0; i < NumInstances; ++i)
+				{
+					ParticleOrder[i].SetAsUint<false, false>(i, CustomSorting[i]);
+				}
 			}
 		}
 	}
diff --git a/Engine/Plugins/FX/Niagara/Source/Niagara/Private/NiagaraRendererLights.cpp b/Engine/Plugins/FX/Niagara/Source/Niagara/Private/NiagaraRendererLights.cpp
--- a/Engine/Plugins/FX/Niagara/Source/Niagara/Private/NiagaraRendererLights.cpp
+++ b/Engine/Plugins/FX/Niagara/Source/Niagara/Private/NiagaraRendererLights.cpp
@@ -61,11 +61,31 @@ FNiagaraDynamicDataBase* FNiagaraRendererLights::GenerateDynamicData(const FNiag
 	//I'm not a great fan of pulling scalar components out to a structured vert buffer like this.
 	//TODO: Experiment with a new VF that reads the data directly from the scalar layout.
 	FNiagaraDataSetAccessor<FVector> PosAccessor(Data, Properties->PositionBinding.DataSetVariable);
-	FNiagaraDataSetAccessor<FLinearColor> ColAccessor(Data, Properties->ColorBinding.DataSetVariable);
-	FNiagaraDataSetAccessor<float> RadiusAccessor(Data, Properties->RadiusBinding.DataSetVariable);
-	FNiagaraDataSetAccessor<float> ExponentAccessor(Data, Properties->LightExponentBinding.DataSetVariable);
-	FNiagaraDataSetAccessor<float> ScatteringAccessor(Data, Properties->VolumetricScatteringBinding.DataSetVariable);
-	FNiagaraDataSetAccessor<int32> EnabledAccessor(Data, Properties->LightRenderingEnabledBinding.DataSetVariable);
+	FNiagaraDataSetAccessor<FNiagaraDataConversions::FHalf4OrColor4> ColAccessor;
+	if (Data.HasVariable(Properties->ColorBinding.DataSetVariable.GetName()))
+	{
+		ColAccessor = FNiagaraDataSetAccessor<FNiagaraDataConversions::FHalf4OrColor4>(Data, Properties->ColorBinding.DataSetVariable.GetName());
+	}
+	FNiagaraDataSetAccessor<FNiagaraDataConversions::FHalfOrFloat> RadiusAccessor;
+	if (Data.HasVariable(Properties->RadiusBinding.DataSetVariable.GetName()))
+	{
+		RadiusAccessor = FNiagaraDataSetAccessor<FNiagaraDataConversions::FHalfOrFloat>(Data, Properties->RadiusBinding.DataSetVariable.GetName());
+	}
+	FNiagaraDataSetAccessor<FNiagaraDataConversions::FHalfOrFloat> ExponentAccessor;
+	if (Data.HasVariable(Properties->LightExponentBinding.DataSetVariable.GetName()))
+	{
+		ExponentAccessor = FNiagaraDataSetAccessor<FNiagaraDataConversions::FHalfOrFloat>(Data, Properties->LightExponentBinding.DataSetVariable.GetName());
+	}
+	FNiagaraDataSetAccessor<FNiagaraDataConversions::FHalfOrFloat> ScatteringAccessor;
+	if (Data.HasVariable(Properties->VolumetricScatteringBinding.DataSetVariable.GetName()))
+	{
+		ScatteringAccessor = FNiagaraDataSetAccessor<FNiagaraDataConversions::FHalfOrFloat>(Data, Properties->VolumetricScatteringBinding.DataSetVariable.GetName());
+	}
+	FNiagaraDataSetAccessor<int32> EnabledAccessor;
+	if (Data.HasVariable(Properties->LightRenderingEnabledBinding.DataSetVariable.GetName()))
+	{
+		EnabledAccessor = FNiagaraDataSetAccessor<int32>(Data, Properties->LightRenderingEnabledBinding.DataSetVariable);
+	}
 
 	const FMatrix& LocalToWorldMatrix = Proxy->GetLocalToWorld();
 	FVector DefaultColor = FVector(Properties->ColorBinding.DefaultValueIfNonExistent.GetValue<FLinearColor>());
diff --git a/Engine/Plugins/FX/Niagara/Source/Niagara/Private/NiagaraRendererMeshes.cpp b/Engine/Plugins/FX/Niagara/Source/Niagara/Private/NiagaraRendererMeshes.cpp
--- a/Engine/Plugins/FX/Niagara/Source/Niagara/Private/NiagaraRendererMeshes.cpp
+++ b/Engine/Plugins/FX/Niagara/Source/Niagara/Private/NiagaraRendererMeshes.cpp
@@ -61,7 +61,6 @@ public:
 	{
 		if (VertexFactory)
 		{
-			VertexFactory->SetParticleData(nullptr, 0);
 			VertexFactory->SetSortedIndices(nullptr, 0xFFFFFFFF);
 		}
 	}
@@ -123,7 +122,8 @@ FNiagaraRendererMeshes::FNiagaraRendererMeshes(ERHIFeatureLevel::Type FeatureLev
 	const FNiagaraDataSet& Data = Emitter->GetData();
 
 	MaterialParamValidMask = 0;
-	TotalVFComponents = 0;
+	TotalVFHalfComponents = 0;
+	TotalVFFloatComponents = 0;
 	VFVariables.SetNum(ENiagaraMeshVFLayout::Num);
 	SetVertexFactoryVariable(Data, Properties->PositionBinding.DataSetVariable, ENiagaraMeshVFLayout::Position);
 	SetVertexFactoryVariable(Data, Properties->VelocityBinding.DataSetVariable, ENiagaraMeshVFLayout::Velocity);
@@ -249,7 +249,7 @@ void FNiagaraRendererMeshes::GetDynamicMeshElements(const TArray<const FSceneVie
 	int32 NumInstances = SourceParticleData->GetNumInstances();
 
 	FGlobalDynamicReadBuffer& DynamicReadBuffer = Collector.GetDynamicReadBuffer();
-	FGlobalDynamicReadBuffer::FAllocation ParticleData;
+	FParticleRenderData ParticleData;
 
 	//Grab the material proxies we'll be using for each section and check them for translucency.
 	bool bHasTranslucentMaterials = false;
@@ -260,8 +260,8 @@ void FNiagaraRendererMeshes::GetDynamicMeshElements(const TArray<const FSceneVie
 		bHasTranslucentMaterials |= IsTranslucentBlendMode(BlendMode);
 	}
 
-	bool bShouldSort = SortMode != ENiagaraSortMode::None && (bHasTranslucentMaterials || !bSortOnlyWhenTranslucent);
-	bool bCustomSorting = SortMode == ENiagaraSortMode::CustomAscending || SortMode == ENiagaraSortMode::CustomDecending;
+	const bool bShouldSort = SortMode != ENiagaraSortMode::None && (bHasTranslucentMaterials || !bSortOnlyWhenTranslucent);
+	const bool bCustomSorting = SortMode == ENiagaraSortMode::CustomAscending || SortMode == ENiagaraSortMode::CustomDecending;
 	//Disable the upload of sorting data if we're using a material that doesn't need it.
 	//TODO: we can probably reinit the GPU layout info entirely to remove custom sorting from the buffer but for now just skip the upload if it's not needed.
 	VFVariables[ENiagaraMeshVFLayout::CustomSorting].bUpload &= bCustomSorting;
@@ -277,8 +277,11 @@ void FNiagaraRendererMeshes::GetDynamicMeshElements(const TArray<const FSceneVie
 		{
 			SCOPE_CYCLE_COUNTER(STAT_NiagaraRenderMeshes_AllocateGPUData);
 			int32 TotalFloatSize = SourceParticleData->GetFloatBuffer().Num() / sizeof(float);
-			ParticleData = DynamicReadBuffer.AllocateFloat(TotalFloatSize);
-			FMemory::Memcpy(ParticleData.Buffer, SourceParticleData->GetFloatBuffer().GetData(), SourceParticleData->GetFloatBuffer().Num());
+			ParticleData.FloatData = DynamicReadBuffer.AllocateFloat(TotalFloatSize);
+			FMemory::Memcpy(ParticleData.FloatData.Buffer, SourceParticleData->GetFloatBuffer().GetData(), SourceParticleData->GetFloatBuffer().Num());
+			int32 TotalHalfSize = SourceParticleData->GetHalfBuffer().Num() / sizeof(FFloat16);
+			ParticleData.HalfData = DynamicReadBuffer.AllocateHalf(TotalFloatSize);
+			FMemory::Memcpy(ParticleData.HalfData.Buffer, SourceParticleData->GetHalfBuffer().GetData(), SourceParticleData->GetHalfBuffer().Num());
 		}
 	}
 	
@@ -372,24 +375,16 @@ void FNiagaraRendererMeshes::GetDynamicMeshElements(const TArray<const FSceneVie
 				CollectorResources.VertexFactory->SetSortedIndices(nullptr, 0xFFFFFFFF);
 
 				FNiagaraGPUSortInfo SortInfo;
-				SortInfo.SortAttributeOffset = INDEX_NONE;
-				int32 SortVarIdx = INDEX_NONE;
+				const int32 SortVarIdx = bCustomSorting ? ENiagaraMeshVFLayout::CustomSorting : ENiagaraMeshVFLayout::Position;
+				const bool bSortVarIsHalf = bShouldSort && VFVariables[SortVarIdx].bHalfType;
 				if (bShouldSort)
 				{
 					SortInfo.ParticleCount = NumInstances;
 					SortInfo.SortMode = SortMode;
+					SortInfo.SortAttributeOffset = (VFVariables[SortVarIdx].GetGPUOffset() & ~(1 << 31));
 					SortInfo.SetSortFlags(GNiagaraGPUSortingUseMaxPrecision != 0, bHasTranslucentMaterials); 
-					if (bCustomSorting)
+					if (!bCustomSorting)
 					{
-						SortVarIdx = ENiagaraMeshVFLayout::CustomSorting;
-						SortInfo.SortAttributeOffset = VFVariables[ENiagaraMeshVFLayout::CustomSorting].GetGPUOffset();
-						SortInfo.ViewOrigin.Set(0, 0, 0);
-						SortInfo.ViewDirection.Set(0, 0, 1);
-					}
-					else
-					{
-						SortVarIdx = ENiagaraMeshVFLayout::Position;
-						SortInfo.SortAttributeOffset = VFVariables[ENiagaraMeshVFLayout::Position].GetGPUOffset();
 						SortInfo.ViewOrigin = View->ViewMatrices.GetViewOrigin();
 						SortInfo.ViewDirection = View->GetViewDirection();
 						if (bLocalSpace)
@@ -402,7 +397,11 @@ void FNiagaraRendererMeshes::GetDynamicMeshElements(const TArray<const FSceneVie
 
 				if (SimTarget == ENiagaraSimTarget::CPUSim)//TODO: Compute shader for sorting gpu sims and larger cpu sims.
 				{
-					check(ParticleData.IsValid());
+					FRHIShaderResourceView* FloatSRV = ParticleData.FloatData.IsValid() ? ParticleData.FloatData.SRV : (FRHIShaderResourceView*)FNiagaraRenderer::GetDummyFloatBuffer();
+					FRHIShaderResourceView* HalfSRV = ParticleData.HalfData.IsValid() ? ParticleData.HalfData.SRV : (FRHIShaderResourceView*)FNiagaraRenderer::GetDummyHalfBuffer();
+					const uint32 FloatParticleDataStride = GbEnableMinimalGPUBuffers ? SourceParticleData->GetNumInstances() : (SourceParticleData->GetFloatStride() / sizeof(float));
+					check(GbEnableMinimalGPUBuffers || FloatParticleDataStride == SourceParticleData->GetHalfStride() / sizeof(FFloat16));
+
 					if (SortInfo.SortMode != ENiagaraSortMode::None && SortInfo.SortAttributeOffset != INDEX_NONE)
 					{
 						if (GNiagaraGPUSortingCPUToGPUThreshold >= 0 &&
@@ -410,8 +409,8 @@ void FNiagaraRendererMeshes::GetDynamicMeshElements(const TArray<const FSceneVie
 							FNiagaraUtilities::AllowComputeShaders(Batcher->GetShaderPlatform()))
 						{
 							SortInfo.ParticleCount = NumInstances;
-							SortInfo.ParticleDataFloatSRV = ParticleData.SRV;
-							SortInfo.FloatDataStride = SourceParticleData->GetFloatStride() / sizeof(float);
+							SortInfo.ParticleDataFloatSRV = bSortVarIsHalf ? HalfSRV : FloatSRV;
+							SortInfo.FloatDataStride = FloatParticleDataStride; // Same if Halfs.
 							if (Batcher->AddSortedGPUSimulation(SortInfo))
 							{
 								CollectorResources.VertexFactory->SetSortedIndices(SortInfo.AllocationInfo.BufferSRV, SortInfo.AllocationInfo.BufferOffset);
@@ -425,17 +424,24 @@ void FNiagaraRendererMeshes::GetDynamicMeshElements(const TArray<const FSceneVie
 							CollectorResources.VertexFactory->SetSortedIndices(SortedIndices.SRV, 0);
 						}
 					}
-					int32 ParticleDataStride = GbEnableMinimalGPUBuffers ? SourceParticleData->GetNumInstances() : SourceParticleData->GetFloatStride() / sizeof(float);
-					CollectorResources.VertexFactory->SetParticleData(ParticleData.SRV, ParticleDataStride);
+
+					PerViewUniformParameters.NiagaraFloatDataStride = FloatParticleDataStride;
+					PerViewUniformParameters.NiagaraParticleDataFloat = FloatSRV;
+					PerViewUniformParameters.NiagaraParticleDataHalf = HalfSRV;
 				}
 				else
 				{
+					FRHIShaderResourceView* FloatSRV = SourceParticleData->GetGPUBufferFloat().SRV.IsValid() ? (FRHIShaderResourceView*)SourceParticleData->GetGPUBufferFloat().SRV : (FRHIShaderResourceView*)FNiagaraRenderer::GetDummyFloatBuffer();
+					FRHIShaderResourceView* HalfSRV = SourceParticleData->GetGPUBufferHalf().SRV.IsValid() ? (FRHIShaderResourceView*)SourceParticleData->GetGPUBufferHalf().SRV : (FRHIShaderResourceView*)FNiagaraRenderer::GetDummyHalfBuffer();
+					const uint32 FloatParticleDataStride = SourceParticleData->GetFloatStride() / sizeof(float);
+					check(FloatParticleDataStride == SourceParticleData->GetHalfStride() / sizeof(FFloat16));
+
 					if (SortInfo.SortMode != ENiagaraSortMode::None && SortInfo.SortAttributeOffset != INDEX_NONE)
 					{
 						// Here we need to be conservative about the InstanceCount, since the final value is only known on the GPU after the simulation.
 						SortInfo.ParticleCount = SourceParticleData->GetNumInstances();
-						SortInfo.ParticleDataFloatSRV = SourceParticleData->GetGPUBufferFloat().SRV;
-						SortInfo.FloatDataStride = SourceParticleData->GetFloatStride() / sizeof(float);
+						SortInfo.ParticleDataFloatSRV = bSortVarIsHalf ? HalfSRV : FloatSRV;
+						SortInfo.FloatDataStride = FloatParticleDataStride; // Same if Halfs.
 						SortInfo.GPUParticleCountSRV = Batcher->GetGPUInstanceCounterManager().GetInstanceCountBuffer().SRV;
 						SortInfo.GPUParticleCountOffset = SourceParticleData->GetGPUInstanceCountBufferOffset();
 						if (Batcher->AddSortedGPUSimulation(SortInfo))
@@ -443,14 +449,10 @@ void FNiagaraRendererMeshes::GetDynamicMeshElements(const TArray<const FSceneVie
 							CollectorResources.VertexFactory->SetSortedIndices(SortInfo.AllocationInfo.BufferSRV, SortInfo.AllocationInfo.BufferOffset);
 						}
 					}
-					if (SourceParticleData->GetGPUBufferFloat().SRV.IsValid())
-					{
-						CollectorResources.VertexFactory->SetParticleData(SourceParticleData->GetGPUBufferFloat().SRV, SourceParticleData->GetFloatStride() / sizeof(float));
-					}
-					else
-					{
-						CollectorResources.VertexFactory->SetParticleData(FNiagaraRenderer::GetDummyFloatBuffer(), 0);
-					}
+					
+					PerViewUniformParameters.NiagaraFloatDataStride = FloatParticleDataStride;
+					PerViewUniformParameters.NiagaraParticleDataFloat = FloatSRV;
+					PerViewUniformParameters.NiagaraParticleDataHalf = HalfSRV;
 				}
 
 				// Collector.AllocateOneFrameResource uses default ctor, initialize the vertex factory
@@ -621,7 +623,7 @@ void FNiagaraRendererMeshes::GetDynamicRayTracingInstances(FRayTracingMaterialGa
 	}
 
 	int32 NumInstances = SourceParticleData->GetNumInstances();
-	int32 TotalFloatSize = TotalVFComponents * SourceParticleData->GetNumInstances();
+	int32 TotalFloatSize = TotalVFFloatComponents * SourceParticleData->GetNumInstances();
 	int32 ComponentStrideDest = SourceParticleData->GetNumInstances() * sizeof(float);
 
 	//ENiagaraMeshVFLayout::Transform just contains a Quat, not the whole transform
diff --git a/Engine/Plugins/FX/Niagara/Source/Niagara/Private/NiagaraRendererRibbons.cpp b/Engine/Plugins/FX/Niagara/Source/Niagara/Private/NiagaraRendererRibbons.cpp
--- a/Engine/Plugins/FX/Niagara/Source/Niagara/Private/NiagaraRendererRibbons.cpp
+++ b/Engine/Plugins/FX/Niagara/Source/Niagara/Private/NiagaraRendererRibbons.cpp
@@ -184,7 +184,8 @@ FNiagaraRendererRibbons::FNiagaraRendererRibbons(ERHIFeatureLevel::Type FeatureL
 	CustomTessellationMinAngle *= PI / 180.f;
 	bCustomUseScreenSpace = Properties->bScreenSpaceTessellation;
 
-	TotalVFComponents = 0;
+	TotalVFHalfComponents = 0;
+	TotalVFFloatComponents = 0;
 	VFVariables.SetNum(ENiagaraRibbonVFLayout::Num);
 	// required attributes
 	SetVertexFactoryVariable(Data, Properties->PositionBinding.DataSetVariable, ENiagaraRibbonVFLayout::Position);
@@ -350,7 +351,7 @@ int FNiagaraRendererRibbons::GetDynamicDataSize()const
 }
 
 void CalculateUVScaleAndOffsets(
-	const FNiagaraDataSetAccessor<float>& SortKeyData, const TArray<int32>& RibbonIndices,
+	const FNiagaraDataSetAccessor<FNiagaraDataConversions::FHalfOrFloat>& SortKeyData, const TArray<int32>& RibbonIndices,
 	bool bSortKeyIsAge, int32 StartIndex, int32 EndIndex, int32 NumSegments,
 	float InUTilingDistance, float InUScale, float InUOffset, ENiagaraRibbonAgeOffsetMode InAgeOffsetMode,
 	float& OutUScale, float& OutUOffset)
@@ -416,21 +417,49 @@ FNiagaraDynamicDataBase* FNiagaraRendererRibbons::GenerateDynamicData(const FNia
 	const UNiagaraRibbonRendererProperties* Properties = CastChecked<const UNiagaraRibbonRendererProperties>(InProperties);
 
 	bool bSortKeyIsAge = false;
-	FNiagaraDataSetAccessor<float> SortKeyData(Data, Properties->RibbonLinkOrderBinding.DataSetVariable);
-	if (SortKeyData.IsValid() == false)
+	FNiagaraDataSetAccessor<FNiagaraDataConversions::FHalfOrFloat> SortKeyData;
+	if (Data.HasVariable(Properties->RibbonLinkOrderBinding.DataSetVariable.GetName()))
 	{
-		SortKeyData = FNiagaraDataSetAccessor<float>(Data, Properties->NormalizedAgeBinding.DataSetVariable);
+		SortKeyData = FNiagaraDataSetAccessor<FNiagaraDataConversions::FHalfOrFloat>(Data, Properties->RibbonLinkOrderBinding.DataSetVariable.GetName());
+	}
+	else
+	{
+		SortKeyData = FNiagaraDataSetAccessor<FNiagaraDataConversions::FHalfOrFloat>(Data, Properties->NormalizedAgeBinding.DataSetVariable.GetName());
 		bSortKeyIsAge = true;
 	}
 
-	FNiagaraDataSetAccessor<FVector> PosData(Data, Properties->PositionBinding.DataSetVariable);
-	FNiagaraDataSetAccessor<float> SizeData(Data, Properties->RibbonWidthBinding.DataSetVariable);
-	FNiagaraDataSetAccessor<float> TwistData(Data, Properties->RibbonTwistBinding.DataSetVariable);
-	FNiagaraDataSetAccessor<FVector> FacingData(Data, Properties->RibbonFacingBinding.DataSetVariable);
-	FNiagaraDataSetAccessor<FVector4> MaterialParamData(Data, Properties->DynamicMaterialBinding.DataSetVariable);
-	FNiagaraDataSetAccessor<FVector4> MaterialParam1Data(Data, Properties->DynamicMaterial1Binding.DataSetVariable);
-	FNiagaraDataSetAccessor<FVector4> MaterialParam2Data(Data, Properties->DynamicMaterial2Binding.DataSetVariable);
-	FNiagaraDataSetAccessor<FVector4> MaterialParam3Data(Data, Properties->DynamicMaterial3Binding.DataSetVariable);
+	FNiagaraDataSetAccessor<FNiagaraDataConversions::FHalf3OrFloat3> PosData(Data, Properties->PositionBinding.DataSetVariable.GetName());
+	FNiagaraDataSetAccessor<FNiagaraDataConversions::FHalfOrFloat> SizeData(Data, Properties->RibbonWidthBinding.DataSetVariable.GetName());
+	FNiagaraDataSetAccessor<FNiagaraDataConversions::FHalfOrFloat> TwistData;
+	if (Data.HasVariable(Properties->RibbonTwistBinding.DataSetVariable.GetName()))
+	{
+		TwistData = FNiagaraDataSetAccessor<FNiagaraDataConversions::FHalfOrFloat>(Data, Properties->RibbonTwistBinding.DataSetVariable.GetName());
+	}
+	FNiagaraDataSetAccessor<FNiagaraDataConversions::FHalf3OrFloat3> FacingData;
+	if (Data.HasVariable(Properties->RibbonFacingBinding.DataSetVariable.GetName()))
+	{
+		FacingData = FNiagaraDataSetAccessor<FNiagaraDataConversions::FHalf3OrFloat3>(Data, Properties->RibbonFacingBinding.DataSetVariable.GetName());
+	}
+	FNiagaraDataSetAccessor<FNiagaraDataConversions::FHalf4OrFloat4> MaterialParamData;
+	if (Data.HasVariable(Properties->DynamicMaterialBinding.DataSetVariable.GetName()))
+	{
+		MaterialParamData = FNiagaraDataSetAccessor<FNiagaraDataConversions::FHalf4OrFloat4>(Data, Properties->DynamicMaterialBinding.DataSetVariable.GetName());
+	}
+	FNiagaraDataSetAccessor<FNiagaraDataConversions::FHalf4OrFloat4> MaterialParam1Data;
+	if (Data.HasVariable(Properties->DynamicMaterial1Binding.DataSetVariable.GetName()))
+	{
+		MaterialParam1Data = FNiagaraDataSetAccessor<FNiagaraDataConversions::FHalf4OrFloat4>(Data, Properties->DynamicMaterial1Binding.DataSetVariable.GetName());
+	}
+	FNiagaraDataSetAccessor<FNiagaraDataConversions::FHalf4OrFloat4> MaterialParam2Data;
+	if (Data.HasVariable(Properties->DynamicMaterial2Binding.DataSetVariable.GetName()))
+	{
+		MaterialParam2Data = FNiagaraDataSetAccessor<FNiagaraDataConversions::FHalf4OrFloat4>(Data, Properties->DynamicMaterial2Binding.DataSetVariable.GetName());
+	}
+	FNiagaraDataSetAccessor<FNiagaraDataConversions::FHalf4OrFloat4> MaterialParam3Data;
+	if (Data.HasVariable(Properties->DynamicMaterial3Binding.DataSetVariable.GetName()))
+	{
+		MaterialParam3Data = FNiagaraDataSetAccessor<FNiagaraDataConversions::FHalf4OrFloat4>(Data, Properties->DynamicMaterial3Binding.DataSetVariable.GetName());
+	}
 
 	FNiagaraDataSetAccessor<int32> RibbonIdData;
 	FNiagaraDataSetAccessor<FNiagaraID> RibbonFullIDData;
@@ -735,7 +764,6 @@ bool FNiagaraRendererRibbons::IsMaterialValid(UMaterialInterface* Mat)const
 	return Mat && Mat->CheckMaterialUsage_Concurrent(MATUSAGE_NiagaraRibbons);
 }
 
-
 void FNiagaraRendererRibbons::SetupMeshBatchAndCollectorResourceForView(
 	const FSceneView* View,
 	const FSceneViewFamily& ViewFamily,
@@ -757,8 +785,10 @@ void FNiagaraRendererRibbons::SetupMeshBatchAndCollectorResourceForView(
 	FCPUSimParticleDataAllocation CPUSimParticleDataAllocation = AllocateParticleDataIfCPUSim(DynamicDataRibbon, Collector.GetDynamicReadBuffer());
 	auto& ParticleData = CPUSimParticleDataAllocation.ParticleData;
 
-	int32 ParticleDataStride = GbEnableMinimalGPUBuffers ? SourceParticleData->GetNumInstances() : SourceParticleData->GetFloatStride() / sizeof(float);
-	CollectorResources.VertexFactory.SetParticleData(ParticleData.SRV, ParticleDataStride);
+	int32 ParticleDataFloatStride = GbEnableMinimalGPUBuffers ? SourceParticleData->GetNumInstances() : SourceParticleData->GetFloatStride() / sizeof(float);
+	int32 ParticleDataHalfStride = GbEnableMinimalGPUBuffers ? SourceParticleData->GetNumInstances() : SourceParticleData->GetHalfStride() / sizeof(FFloat16);
+
+	check(ParticleDataFloatStride == ParticleDataHalfStride);
 
 	// TODO: need to make these two a global alloc buffer as well, not recreate
 				// pass in the sorted indices so the VS can fetch the particle data in order
@@ -791,13 +821,17 @@ void FNiagaraRendererRibbons::SetupMeshBatchAndCollectorResourceForView(
 	RHIUnlockVertexBuffer(PackedPerRibbonDataByIndexBuffer.Buffer);
 	CollectorResources.VertexFactory.SetPackedPerRibbonDataByIndexSRV(PackedPerRibbonDataByIndexBuffer.Buffer, PackedPerRibbonDataByIndexBuffer.SRV);
 
+	FRHIShaderResourceView* FloatSRV = ParticleData.FloatData.IsValid() ? ParticleData.FloatData.SRV : (FRHIShaderResourceView*)FNiagaraRenderer::GetDummyFloatBuffer();
+	FRHIShaderResourceView* HalfSRV = ParticleData.HalfData.IsValid() ? ParticleData.HalfData.SRV : (FRHIShaderResourceView*)FNiagaraRenderer::GetDummyHalfBuffer();
+
 	FNiagaraRibbonVFLooseParameters VFLooseParams;
 	VFLooseParams.SortedIndices = SortedIndicesBuffer.SRV;
 	VFLooseParams.TangentsAndDistances = TangentsAndDistancesBuffer.SRV;
 	VFLooseParams.MultiRibbonIndices = MultiRibbonIndicesBuffer.SRV;
 	VFLooseParams.PackedPerRibbonDataByIndex = PackedPerRibbonDataByIndexBuffer.SRV;
-	VFLooseParams.NiagaraParticleDataFloat = CollectorResources.VertexFactory.GetParticleDataFloatSRV();
-	VFLooseParams.NiagaraFloatDataStride = CollectorResources.VertexFactory.GetFloatDataStride();
+	VFLooseParams.NiagaraParticleDataFloat = FloatSRV;
+	VFLooseParams.NiagaraParticleDataHalf = HalfSRV;
+	VFLooseParams.NiagaraFloatDataStride = ParticleDataFloatStride;
 	VFLooseParams.SortedIndicesOffset = CollectorResources.VertexFactory.GetSortedIndicesOffset();
 	VFLooseParams.FacingMode = static_cast<uint32>(FacingMode);
 
@@ -847,7 +881,6 @@ FNiagaraRendererRibbons::FCPUSimParticleDataAllocation FNiagaraRendererRibbons::
 {
 	FNiagaraDataBuffer* SourceParticleData = DynamicDataRibbon->GetParticleDataToRender();
 	check(SourceParticleData);//Can be null but should be checked before here.
-	int32 TotalFloatSize = SourceParticleData->GetFloatBuffer().Num() / sizeof(float);
 
 	FCPUSimParticleDataAllocation CPUSimParticleDataAllocation{ DynamicReadBuffer };
 
@@ -863,8 +896,12 @@ FNiagaraRendererRibbons::FCPUSimParticleDataAllocation FNiagaraRendererRibbons::
 		}
 		else
 		{
-			CPUSimParticleDataAllocation.ParticleData = DynamicReadBuffer.AllocateFloat(TotalFloatSize);
-			FMemory::Memcpy(CPUSimParticleDataAllocation.ParticleData.Buffer, SourceParticleData->GetFloatBuffer().GetData(), SourceParticleData->GetFloatBuffer().Num());
+			int32 TotalFloatSize = SourceParticleData->GetFloatBuffer().Num() / sizeof(float);
+			CPUSimParticleDataAllocation.ParticleData.FloatData = DynamicReadBuffer.AllocateFloat(TotalFloatSize);
+			FMemory::Memcpy(CPUSimParticleDataAllocation.ParticleData.FloatData.Buffer, SourceParticleData->GetFloatBuffer().GetData(), SourceParticleData->GetFloatBuffer().Num());
+			int32 TotalHalfSize = SourceParticleData->GetHalfBuffer().Num() / sizeof(FFloat16);
+			CPUSimParticleDataAllocation.ParticleData.HalfData = DynamicReadBuffer.AllocateHalf(TotalFloatSize);
+			FMemory::Memcpy(CPUSimParticleDataAllocation.ParticleData.HalfData.Buffer, SourceParticleData->GetHalfBuffer().GetData(), SourceParticleData->GetHalfBuffer().Num());
 		}
 	}
 
diff --git a/Engine/Plugins/FX/Niagara/Source/Niagara/Private/NiagaraRendererSprites.cpp b/Engine/Plugins/FX/Niagara/Source/Niagara/Private/NiagaraRendererSprites.cpp
--- a/Engine/Plugins/FX/Niagara/Source/Niagara/Private/NiagaraRendererSprites.cpp
+++ b/Engine/Plugins/FX/Niagara/Source/Niagara/Private/NiagaraRendererSprites.cpp
@@ -128,7 +128,8 @@ FNiagaraRendererSprites::FNiagaraRendererSprites(ERHIFeatureLevel::Type FeatureL
 
 	const FNiagaraDataSet& Data = Emitter->GetData();
 
-	TotalVFComponents = 0;
+	TotalVFFloatComponents = 0;
+	TotalVFHalfComponents = 0;
 	VFVariables.SetNum(ENiagaraSpriteVFLayout::Num);
 	SetVertexFactoryVariable(Data, Properties->PositionBinding.DataSetVariable, ENiagaraSpriteVFLayout::Position);
 	SetVertexFactoryVariable(Data, Properties->VelocityBinding.DataSetVariable, ENiagaraSpriteVFLayout::Velocity);
@@ -206,10 +207,13 @@ FNiagaraRendererSprites::FCPUSimParticleDataAllocation FNiagaraRendererSprites::
 		}
 		else
 		{
+			SCOPE_CYCLE_COUNTER(STAT_NiagaraRenderSpritesCPUSimMemCopy);
 			int32 TotalFloatSize = SourceParticleData->GetFloatBuffer().Num() / sizeof(float);
-			CPUSimParticleDataAllocation.ParticleData = DynamicReadBuffer.AllocateFloat(TotalFloatSize);
-			SCOPE_CYCLE_COUNTER(STAT_NiagaraRenderSpritesCPUSimMemCopy);			
-			FMemory::Memcpy(CPUSimParticleDataAllocation.ParticleData.Buffer, SourceParticleData->GetFloatBuffer().GetData(), SourceParticleData->GetFloatBuffer().Num());
+			CPUSimParticleDataAllocation.ParticleData.FloatData = DynamicReadBuffer.AllocateFloat(TotalFloatSize);
+			FMemory::Memcpy(CPUSimParticleDataAllocation.ParticleData.FloatData.Buffer, SourceParticleData->GetFloatBuffer().GetData(), SourceParticleData->GetFloatBuffer().Num());
+			int32 TotalHalfSize = SourceParticleData->GetHalfBuffer().Num() / sizeof(FFloat16);
+			CPUSimParticleDataAllocation.ParticleData.HalfData = DynamicReadBuffer.AllocateHalf(TotalFloatSize);
+			FMemory::Memcpy(CPUSimParticleDataAllocation.ParticleData.HalfData.Buffer, SourceParticleData->GetHalfBuffer().GetData(), SourceParticleData->GetHalfBuffer().Num());
 		}
 	}
 
@@ -296,6 +300,7 @@ void FNiagaraRendererSprites::SetVertexFactoryParticleData(
 	FNiagaraDynamicDataSprites* DynamicDataSprites, 
 	FCPUSimParticleDataAllocation& CPUSimParticleDataAllocation,
 	const FSceneView* View,
+	FNiagaraSpriteVFLooseParameters& VFLooseParams,
 	const FNiagaraSceneProxy *SceneProxy) const
 {
 	NiagaraEmitterInstanceBatcher* Batcher = SceneProxy->GetBatcher();
@@ -330,26 +335,21 @@ void FNiagaraRendererSprites::SetVertexFactoryParticleData(
 		EBlendMode BlendMode = MaterialRenderProxy->GetMaterial(FeatureLevel)->GetBlendMode();
 		OutVertexFactory.SetSortedIndices(nullptr, 0xFFFFFFFF);
 
-		int32 NumInstances = SourceParticleData->GetNumInstances();
+		const int32 NumInstances = SourceParticleData->GetNumInstances();
 		FNiagaraGPUSortInfo SortInfo;
-		SortInfo.SortAttributeOffset = INDEX_NONE;
-		int32 SortingVarIdx = INDEX_NONE;
 		const bool bHasTranslucentMaterials = IsTranslucentBlendMode(BlendMode);
-		bool bShouldSort = SortMode != ENiagaraSortMode::None && (bHasTranslucentMaterials || !bSortOnlyWhenTranslucent);
+		const bool bShouldSort = SortMode != ENiagaraSortMode::None && (bHasTranslucentMaterials || !bSortOnlyWhenTranslucent);
+		const bool bCustomSorting = SortMode == ENiagaraSortMode::CustomAscending || SortMode == ENiagaraSortMode::CustomDecending;
+		const int32 SortVarIdx = bCustomSorting ? ENiagaraSpriteVFLayout::CustomSorting : ENiagaraSpriteVFLayout::Position;
+		const bool bSortVarIsHalf = bShouldSort && VFVariables[SortVarIdx].bHalfType;
 		if (bShouldSort)
 		{
 			SortInfo.ParticleCount = NumInstances;
 			SortInfo.SortMode = SortMode;
+			SortInfo.SortAttributeOffset = (VFVariables[SortVarIdx].GetGPUOffset() & ~(1 << 31));
 			SortInfo.SetSortFlags(GNiagaraGPUSortingUseMaxPrecision != 0, bHasTranslucentMaterials); 
-			if (SortInfo.SortMode == ENiagaraSortMode::CustomAscending || SortInfo.SortMode == ENiagaraSortMode::CustomDecending)
-			{
-				SortingVarIdx = ENiagaraSpriteVFLayout::CustomSorting;	
-				SortInfo.ViewOrigin.Set(0, 0, 0);
-				SortInfo.ViewDirection.Set(0, 0, 1);
-			}
-			else
+			if (!bCustomSorting)
 			{
-				SortingVarIdx = ENiagaraSpriteVFLayout::Position;
 				SortInfo.ViewOrigin = View->ViewMatrices.GetViewOrigin();
 				SortInfo.ViewDirection = View->GetViewDirection();
 				if (bLocalSpace)
@@ -358,13 +358,15 @@ void FNiagaraRendererSprites::SetVertexFactoryParticleData(
 					SortInfo.ViewDirection = SceneProxy->GetLocalToWorld().GetTransposed().TransformVector(SortInfo.ViewDirection);
 				}
 			}
-			SortInfo.SortAttributeOffset = VFVariables[SortingVarIdx].GetGPUOffset();
 		}
 
 		if (SimTarget == ENiagaraSimTarget::CPUSim)
 		{
-			int32 ParticleStrideInFloats = GbEnableMinimalGPUBuffers ? SourceParticleData->GetNumInstances() : SourceParticleData->GetFloatStride() / sizeof(float);
-			check(CPUSimParticleDataAllocation.ParticleData.IsValid());
+			FRHIShaderResourceView* FloatSRV = CPUSimParticleDataAllocation.ParticleData.FloatData.IsValid() ? CPUSimParticleDataAllocation.ParticleData.FloatData.SRV : (FRHIShaderResourceView*)FNiagaraRenderer::GetDummyFloatBuffer();
+			FRHIShaderResourceView* HalfSRV = CPUSimParticleDataAllocation.ParticleData.HalfData.IsValid() ? CPUSimParticleDataAllocation.ParticleData.HalfData.SRV : (FRHIShaderResourceView*)FNiagaraRenderer::GetDummyHalfBuffer();
+			const uint32 ParticleStrideInFloats = GbEnableMinimalGPUBuffers ? SourceParticleData->GetNumInstances() : (SourceParticleData->GetFloatStride() / sizeof(float));
+			check(GbEnableMinimalGPUBuffers || ParticleStrideInFloats == SourceParticleData->GetHalfStride() / sizeof(FFloat16));
+
 			if (SortInfo.SortMode != ENiagaraSortMode::None && SortInfo.SortAttributeOffset != INDEX_NONE)
 			{
 				if (GNiagaraGPUSortingCPUToGPUThreshold >= 0 &&
@@ -372,8 +374,8 @@ void FNiagaraRendererSprites::SetVertexFactoryParticleData(
 					FNiagaraUtilities::AllowComputeShaders(Batcher->GetShaderPlatform()))
 				{
 					SortInfo.ParticleCount = NumInstances;
-					SortInfo.ParticleDataFloatSRV = CPUSimParticleDataAllocation.ParticleData.SRV;
-					SortInfo.FloatDataStride = ParticleStrideInFloats;
+					SortInfo.ParticleDataFloatSRV = bSortVarIsHalf ? HalfSRV : FloatSRV;
+					SortInfo.FloatDataStride = ParticleStrideInFloats; // Same if Halfs.
 					if (Batcher->AddSortedGPUSimulation(SortInfo))
 					{
 						OutVertexFactory.SetSortedIndices(SortInfo.AllocationInfo.BufferSRV, SortInfo.AllocationInfo.BufferOffset);
@@ -385,21 +387,30 @@ void FNiagaraRendererSprites::SetVertexFactoryParticleData(
 
 					FGlobalDynamicReadBuffer::FAllocation SortedIndices;
 					SortedIndices = CPUSimParticleDataAllocation.DynamicReadBuffer.AllocateInt32(NumInstances);
-					SortIndices(SortInfo, SortingVarIdx, *SourceParticleData, SortedIndices);
+					SortIndices(SortInfo, SortVarIdx, *SourceParticleData, SortedIndices);
 					OutVertexFactory.SetSortedIndices(SortedIndices.SRV, 0);
 				}
 			}
-			OutVertexFactory.SetParticleData(CPUSimParticleDataAllocation.ParticleData.SRV, ParticleStrideInFloats);
+			auto& ParticleData = CPUSimParticleDataAllocation.ParticleData;
+
+			VFLooseParams.NiagaraFloatDataStride = ParticleStrideInFloats;
+			VFLooseParams.NiagaraParticleDataFloat = FloatSRV;
+			VFLooseParams.NiagaraParticleDataHalf = HalfSRV;
 		}
 		else // ENiagaraSimTarget::GPUSim
 		{
+			FRHIShaderResourceView* FloatSRV = SourceParticleData->GetGPUBufferFloat().SRV.IsValid() ? (FRHIShaderResourceView*)SourceParticleData->GetGPUBufferFloat().SRV : (FRHIShaderResourceView*)FNiagaraRenderer::GetDummyFloatBuffer();
+			FRHIShaderResourceView* HalfSRV = SourceParticleData->GetGPUBufferHalf().SRV.IsValid() ? (FRHIShaderResourceView*)SourceParticleData->GetGPUBufferHalf().SRV : (FRHIShaderResourceView*)FNiagaraRenderer::GetDummyHalfBuffer();
+			const uint32 ParticleStrideInFloats = SourceParticleData->GetFloatStride() / sizeof(float);
+			check(ParticleStrideInFloats == SourceParticleData->GetHalfStride() / sizeof(FFloat16));
+
 			if (SortInfo.SortMode != ENiagaraSortMode::None && SortInfo.SortAttributeOffset != INDEX_NONE)
 			{
 				// Here we need to be conservative about the InstanceCount, since the final value is only known on the GPU after the simulation.
 				SortInfo.ParticleCount = SourceParticleData->GetNumInstances();
 
-				SortInfo.ParticleDataFloatSRV = SourceParticleData->GetGPUBufferFloat().SRV;
-				SortInfo.FloatDataStride = SourceParticleData->GetFloatStride() / sizeof(float);
+				SortInfo.ParticleDataFloatSRV = bSortVarIsHalf ? HalfSRV : FloatSRV;
+				SortInfo.FloatDataStride = ParticleStrideInFloats;
 				SortInfo.GPUParticleCountSRV = Batcher->GetGPUInstanceCounterManager().GetInstanceCountBuffer().SRV;
 				SortInfo.GPUParticleCountOffset = SourceParticleData->GetGPUInstanceCountBufferOffset();
 				if (Batcher->AddSortedGPUSimulation(SortInfo))
@@ -408,14 +419,9 @@ void FNiagaraRendererSprites::SetVertexFactoryParticleData(
 				}
 			}
 
-			if ( SourceParticleData->GetGPUBufferFloat().SRV.IsValid() )
-			{
-				OutVertexFactory.SetParticleData(SourceParticleData->GetGPUBufferFloat().SRV, SourceParticleData->GetFloatStride() / sizeof(float));
-			}
-			else
-			{
-				OutVertexFactory.SetParticleData(FNiagaraRenderer::GetDummyFloatBuffer(), 0);
-			}
+			VFLooseParams.NiagaraFloatDataStride = ParticleStrideInFloats;
+			VFLooseParams.NiagaraParticleDataFloat = FloatSRV;
+			VFLooseParams.NiagaraParticleDataHalf = HalfSRV;
 		}
 	}
 }
@@ -427,6 +433,7 @@ void FNiagaraRendererSprites::CreateMeshBatchForView(
 	FNiagaraDynamicDataSprites *DynamicDataSprites,
 	uint32 IndirectArgsOffset,
 	FMeshBatch& MeshBatch,
+	FNiagaraSpriteVFLooseParameters& VFLooseParams,
 	FNiagaraMeshCollectorResourcesSprite& CollectorResources) const
 {
 	FNiagaraDataBuffer* SourceParticleData = DynamicDataSprites->GetParticleDataToRender();
@@ -457,11 +464,8 @@ void FNiagaraRendererSprites::CreateMeshBatchForView(
 	CollectorResources.VertexFactory.InitResource();
 	CollectorResources.VertexFactory.SetSpriteUniformBuffer(CollectorResources.UniformBuffer);
 
-	FNiagaraSpriteVFLooseParameters VFLooseParams;
 	VFLooseParams.NumCutoutVerticesPerFrame = CollectorResources.VertexFactory.GetNumCutoutVerticesPerFrame();
 	VFLooseParams.CutoutGeometry = CollectorResources.VertexFactory.GetCutoutGeometrySRV() ? CollectorResources.VertexFactory.GetCutoutGeometrySRV() : GFNiagaraNullCutoutVertexBuffer.VertexBufferSRV.GetReference();
-	VFLooseParams.NiagaraParticleDataFloat = CollectorResources.VertexFactory.GetParticleDataFloatSRV();
-	VFLooseParams.NiagaraFloatDataStride = CollectorResources.VertexFactory.GetFloatDataStride();
 	VFLooseParams.ParticleAlignmentMode = CollectorResources.VertexFactory.GetAlignmentMode();
 	VFLooseParams.ParticleFacingMode = CollectorResources.VertexFactory.GetFacingMode();
 	VFLooseParams.SortedIndices = CollectorResources.VertexFactory.GetSortedIndicesSRV() ? CollectorResources.VertexFactory.GetSortedIndicesSRV() : GFNiagaraNullSortedIndicesVertexBuffer.VertexBufferSRV.GetReference();
@@ -580,10 +584,13 @@ void FNiagaraRendererSprites::GetDynamicMeshElements(const TArray<const FSceneVi
 			const FSceneView* View = Views[ViewIndex];
 
 			FNiagaraMeshCollectorResourcesSprite& CollectorResources = Collector.AllocateOneFrameResource<FNiagaraMeshCollectorResourcesSprite>();
-			SetVertexFactoryParticleData(CollectorResources.VertexFactory, DynamicDataSprites, CPUSimParticleDataAllocation, View, SceneProxy);
+			FNiagaraSpriteVFLooseParameters VFLooseParams;
+			SetVertexFactoryParticleData(CollectorResources.VertexFactory, DynamicDataSprites, CPUSimParticleDataAllocation, View, VFLooseParams, SceneProxy);
 			CollectorResources.UniformBuffer = CreatePerViewUniformBuffer(View, ViewFamily, SceneProxy);
 			FMeshBatch& MeshBatch = Collector.AllocateMesh();
-			CreateMeshBatchForView(View, ViewFamily, SceneProxy, DynamicDataSprites, IndirectArgsOffset, MeshBatch, CollectorResources);
+
+			CreateMeshBatchForView(View, ViewFamily, SceneProxy, DynamicDataSprites, IndirectArgsOffset, MeshBatch, VFLooseParams, CollectorResources);
+
 			Collector.AddMesh(ViewIndex, MeshBatch);
 		}
 	}
@@ -628,10 +635,11 @@ void FNiagaraRendererSprites::GetDynamicRayTracingInstances(FRayTracingMaterialG
 		// Setup material for our ray tracing instance
 		FCPUSimParticleDataAllocation CPUSimParticleDataAllocation = ConditionalAllocateCPUSimParticleData(DynamicDataSprites, Context.RayTracingMeshResourceCollector.GetDynamicReadBuffer());
 		FNiagaraMeshCollectorResourcesSprite& CollectorResources = Context.RayTracingMeshResourceCollector.AllocateOneFrameResource<FNiagaraMeshCollectorResourcesSprite>();
-		SetVertexFactoryParticleData(CollectorResources.VertexFactory, DynamicDataSprites, CPUSimParticleDataAllocation, Context.ReferenceView, SceneProxy);
+		FNiagaraSpriteVFLooseParameters VFLooseParams;
+		SetVertexFactoryParticleData(CollectorResources.VertexFactory, DynamicDataSprites, CPUSimParticleDataAllocation, Context.ReferenceView, VFLooseParams, SceneProxy);
 		CollectorResources.UniformBuffer = CreatePerViewUniformBuffer(Context.ReferenceView, Context.ReferenceViewFamily, SceneProxy);
 		FMeshBatch MeshBatch;
-		CreateMeshBatchForView(Context.ReferenceView, Context.ReferenceViewFamily, SceneProxy, DynamicDataSprites, IndirectArgsOffset, MeshBatch, CollectorResources);
+		CreateMeshBatchForView(Context.ReferenceView, Context.ReferenceViewFamily, SceneProxy, DynamicDataSprites, IndirectArgsOffset, MeshBatch, VFLooseParams, CollectorResources);
 
 		ensureMsgf(MeshBatch.Elements[0].IndexBuffer != &GSixTriangleParticleIndexBuffer, TEXT("Cutout geometry is not supported by ray tracing"));
 
@@ -696,4 +704,4 @@ int FNiagaraRendererSprites::GetDynamicDataSize()const
 bool FNiagaraRendererSprites::IsMaterialValid(UMaterialInterface* Mat)const
 {
 	return Mat && Mat->CheckMaterialUsage_Concurrent(MATUSAGE_NiagaraSprites);
-}
\ No newline at end of file
+}
diff --git a/Engine/Plugins/FX/Niagara/Source/Niagara/Private/NiagaraScript.cpp b/Engine/Plugins/FX/Niagara/Source/Niagara/Private/NiagaraScript.cpp
--- a/Engine/Plugins/FX/Niagara/Source/Niagara/Private/NiagaraScript.cpp
+++ b/Engine/Plugins/FX/Niagara/Source/Niagara/Private/NiagaraScript.cpp
@@ -322,6 +322,10 @@ void UNiagaraScript::ComputeVMCompilationId(FNiagaraVMExecutableDataId& Id) cons
 		{
 			Id.bUsesRapidIterationParams = false;
 		}
+		if (EmitterOwner && EmitterOwner->bCompressAttributes)
+		{
+			Id.AdditionalDefines.Add(TEXT("CompressAttributes"));
+		}
 
 		if ((Emitter->bInterpolatedSpawning && Usage == ENiagaraScriptUsage::ParticleGPUComputeScript) || 
 			(Emitter->bInterpolatedSpawning && Usage == ENiagaraScriptUsage::ParticleSpawnScript) ||
@@ -395,6 +399,10 @@ void UNiagaraScript::ComputeVMCompilationId(FNiagaraVMExecutableDataId& Id) cons
 		{
 			Id.bUsesRapidIterationParams = false;
 		}
+		if (System->bCompressAttributes)
+		{
+			Id.AdditionalDefines.Add(TEXT("CompressAttributes"));
+		}
 
 		for (const FNiagaraEmitterHandle& EmitterHandle: System->GetEmitterHandles())
 		{
diff --git a/Engine/Plugins/FX/Niagara/Source/Niagara/Private/NiagaraScriptExecutionContext.cpp b/Engine/Plugins/FX/Niagara/Source/Niagara/Private/NiagaraScriptExecutionContext.cpp
--- a/Engine/Plugins/FX/Niagara/Source/Niagara/Private/NiagaraScriptExecutionContext.cpp
+++ b/Engine/Plugins/FX/Niagara/Source/Niagara/Private/NiagaraScriptExecutionContext.cpp
@@ -158,11 +158,24 @@ void FNiagaraScriptExecutionContext::BindData(int32 Index, FNiagaraDataSet& Data
 	DataSetInfo[Index].Init(&DataSet, Input, Output, StartInstance, bUpdateInstanceCounts);
 
 	//Would be nice to roll this and DataSetInfo into one but currently the VM being in it's own Engine module prevents this. Possibly should move the VM into Niagara itself.
-	uint8** InputRegisters = Input ? Input->GetRegisterTable().GetData() : nullptr;
-	uint8** OutputRegisters = Output ? Output->GetRegisterTable().GetData() : nullptr;
+	TArrayView<uint8 const* RESTRICT const> InputRegisters = Input ? Input->GetRegisterTable() : TArrayView<uint8 const* RESTRICT const>();
+	TArrayView<uint8 const* RESTRICT const> OutputRegisters = Output ? Output->GetRegisterTable() : TArrayView<uint8 const* RESTRICT const>();
+
 	DataSetMetaTable.SetNum(FMath::Max(DataSetMetaTable.Num(), Index + 1));
 	DataSetMetaTable[Index].Init(InputRegisters, OutputRegisters, StartInstance,
 		Output ? &Output->GetIDTable() : nullptr, &DataSet.GetFreeIDTable(), &DataSet.GetNumFreeIDs(), &DataSet.GetMaxUsedID(), DataSet.GetIDAcquireTag(), &DataSet.GetSpawnedIDsTable());
+
+	if (InputRegisters.Num() > 0)
+	{
+		static_assert(sizeof(DataSetMetaTable[Index].InputRegisterTypeOffsets) == sizeof(FNiagaraDataBuffer::RegisterTypeOffsetType), "ArraySizes do not match");
+		memcpy(DataSetMetaTable[Index].InputRegisterTypeOffsets, Input->GetRegisterTypeOffsets(), sizeof(FNiagaraDataBuffer::RegisterTypeOffsetType));
+	}
+
+	if (OutputRegisters.Num() > 0)
+	{
+		static_assert(sizeof(DataSetMetaTable[Index].OutputRegisterTypeOffsets) == sizeof(FNiagaraDataBuffer::RegisterTypeOffsetType), "ArraySizes do not match");
+		memcpy(DataSetMetaTable[Index].OutputRegisterTypeOffsets, Output->GetRegisterTypeOffsets(), sizeof(FNiagaraDataBuffer::RegisterTypeOffsetType));
+	}
 }
 
 void FNiagaraScriptExecutionContext::BindData(int32 Index, FNiagaraDataBuffer* Input, int32 StartInstance, bool bUpdateInstanceCounts)
@@ -172,8 +185,16 @@ void FNiagaraScriptExecutionContext::BindData(int32 Index, FNiagaraDataBuffer* I
 	FNiagaraDataSet* DataSet = Input->GetOwner();
 	DataSetInfo[Index].Init(DataSet, Input, nullptr, StartInstance, bUpdateInstanceCounts);
 
+	TArrayView<uint8 const* RESTRICT const> InputRegisters = Input ? Input->GetRegisterTable() : TArrayView<uint8 const* RESTRICT const>();
+
 	DataSetMetaTable.SetNum(FMath::Max(DataSetMetaTable.Num(), Index + 1));
-	DataSetMetaTable[Index].Init(Input->GetRegisterTable().GetData(), nullptr, StartInstance, nullptr, nullptr, &DataSet->GetNumFreeIDs(), &DataSet->GetMaxUsedID(), DataSet->GetIDAcquireTag(), &DataSet->GetSpawnedIDsTable());
+	DataSetMetaTable[Index].Init(InputRegisters, TArrayView<uint8 const* RESTRICT const>(), StartInstance, nullptr, nullptr, &DataSet->GetNumFreeIDs(), &DataSet->GetMaxUsedID(), DataSet->GetIDAcquireTag(), &DataSet->GetSpawnedIDsTable());
+
+	if (InputRegisters.Num() > 0)
+	{
+		static_assert(sizeof(DataSetMetaTable[Index].InputRegisterTypeOffsets) == sizeof(FNiagaraDataBuffer::RegisterTypeOffsetType), "ArraySizes do not match");
+		memcpy(DataSetMetaTable[Index].InputRegisterTypeOffsets, Input->GetRegisterTypeOffsets(), sizeof(FNiagaraDataBuffer::RegisterTypeOffsetType));
+	}
 }
 
 bool FNiagaraScriptExecutionContext::Execute(uint32 NumInstances, const FScriptExecutionConstantBufferTable& ConstantBufferTable)
diff --git a/Engine/Plugins/FX/Niagara/Source/Niagara/Private/NiagaraScriptExecutionParameterStore.cpp b/Engine/Plugins/FX/Niagara/Source/Niagara/Private/NiagaraScriptExecutionParameterStore.cpp
--- a/Engine/Plugins/FX/Niagara/Source/Niagara/Private/NiagaraScriptExecutionParameterStore.cpp
+++ b/Engine/Plugins/FX/Niagara/Source/Niagara/Private/NiagaraScriptExecutionParameterStore.cpp
@@ -83,6 +83,34 @@ uint32 FNiagaraScriptExecutionParameterStore::GenerateLayoutInfoInternal(TArray<
 		InSrcOffset += sizeof(FMatrix);
 		NextMemberOffset = Members[Members.Num() - 1].DestOffset + Members[Members.Num() - 1].DestSize;
 	}
+	else if (InSrcStruct == FNiagaraTypeDefinition::GetHalfStruct())
+	{
+		uint32 HalfSize = (TShaderParameterTypeInfo<FFloat16>::NumRows * TShaderParameterTypeInfo<FFloat16>::NumColumns) * sizeof(FFloat16);
+		Members.Emplace(InSrcOffset, Align(NextMemberOffset, TShaderParameterTypeInfo<FFloat16>::Alignment), HalfSize, HalfSize);
+		InSrcOffset += sizeof(FFloat16);
+		NextMemberOffset = Members[Members.Num() - 1].DestOffset + Members[Members.Num() - 1].DestSize;
+	}
+	else if (InSrcStruct == FNiagaraTypeDefinition::GetHalfVec2Struct())
+	{
+		uint32 StructFinalSize = (TShaderParameterTypeInfo<FFloat16[2]>::NumRows * TShaderParameterTypeInfo<FFloat16>::NumColumns) * sizeof(FFloat16);
+		Members.Emplace(InSrcOffset, OffsetAlign(NextMemberOffset, VectorPaddedSize), StructFinalSize, VectorPaddedSize);
+		InSrcOffset += sizeof(FFloat16[2]);
+		NextMemberOffset = Members[Members.Num() - 1].DestOffset + Members[Members.Num() - 1].DestSize;
+	}
+	else if (InSrcStruct == FNiagaraTypeDefinition::GetHalfVec3Struct())
+	{
+		uint32 StructFinalSize = (TShaderParameterTypeInfo<FFloat16[3]>::NumRows * TShaderParameterTypeInfo<FFloat16>::NumColumns) * sizeof(FFloat16);
+		Members.Emplace(InSrcOffset, OffsetAlign(NextMemberOffset, VectorPaddedSize), StructFinalSize, VectorPaddedSize);
+		InSrcOffset += sizeof(FFloat16[3]);
+		NextMemberOffset = Members[Members.Num() - 1].DestOffset + Members[Members.Num() - 1].DestSize;
+	}
+	else if (InSrcStruct == FNiagaraTypeDefinition::GetHalfVec4Struct())
+	{
+		uint32 StructFinalSize = (TShaderParameterTypeInfo<FFloat16[4]>::NumRows * TShaderParameterTypeInfo<FFloat16>::NumColumns) * sizeof(FFloat16);
+		Members.Emplace(InSrcOffset, OffsetAlign(NextMemberOffset, VectorPaddedSize), StructFinalSize, VectorPaddedSize);
+		InSrcOffset += sizeof(FFloat16[4]);
+		NextMemberOffset = Members[Members.Num() - 1].DestOffset + Members[Members.Num() - 1].DestSize;
+	}
 	else
 	{
 		NextMemberOffset = Align(NextMemberOffset, SHADER_PARAMETER_STRUCT_ALIGNMENT); // New structs should be aligned to the head..
diff --git a/Engine/Plugins/FX/Niagara/Source/Niagara/Public/NiagaraCommon.h b/Engine/Plugins/FX/Niagara/Source/Niagara/Public/NiagaraCommon.h
--- a/Engine/Plugins/FX/Niagara/Source/Niagara/Public/NiagaraCommon.h
+++ b/Engine/Plugins/FX/Niagara/Source/Niagara/Public/NiagaraCommon.h
@@ -51,6 +51,7 @@ enum class ENiagaraTickBehavior : uint8
 
 enum ENiagaraBaseTypes
 {
+	NBT_Half,
 	NBT_Float,
 	NBT_Int32,
 	NBT_Bool,
diff --git a/Engine/Plugins/FX/Niagara/Source/Niagara/Public/NiagaraParameterStore.h b/Engine/Plugins/FX/Niagara/Source/Niagara/Public/NiagaraParameterStore.h
--- a/Engine/Plugins/FX/Niagara/Source/Niagara/Public/NiagaraParameterStore.h
+++ b/Engine/Plugins/FX/Niagara/Source/Niagara/Public/NiagaraParameterStore.h
@@ -355,7 +355,7 @@ public:
 	/** Returns the associated FNiagaraVariable for the passed data interface if it exists in the store. Null if not.*/
 	const FNiagaraVariableBase* FindVariable(UNiagaraDataInterface* Interface)const;
 
-	virtual const int32* FindParameterOffset(const FNiagaraVariable& Parameter) const;
+	virtual const int32* FindParameterOffset(const FNiagaraVariable& Parameter, bool IgnoreType = false) const;
 
 	void PostLoad();
 	void SortParameters();
diff --git a/Engine/Plugins/FX/Niagara/Source/Niagara/Public/NiagaraRenderer.h b/Engine/Plugins/FX/Niagara/Source/Niagara/Public/NiagaraRenderer.h
--- a/Engine/Plugins/FX/Niagara/Source/Niagara/Public/NiagaraRenderer.h
+++ b/Engine/Plugins/FX/Niagara/Source/Niagara/Public/NiagaraRenderer.h
@@ -61,21 +61,36 @@ extern int32 GbEnableMinimalGPUBuffers;
 struct FNiagaraRendererVariableInfo
 {
 	FNiagaraRendererVariableInfo()
-		: DatasetOffset(INDEX_NONE), GPUBufferOffset(INDEX_NONE), NumComponents(0), bUpload(true)
+		: DatasetOffset(INDEX_NONE), GPUBufferOffset(INDEX_NONE), NumComponents(0), bUpload(true), bHalfType(false)
 	{
 	}
 	
-	FNiagaraRendererVariableInfo(int32 InDatasetOffset, int32 InGPUBufferOffset, int32 InNumComponents, bool InbUpload)
-		: DatasetOffset(InDatasetOffset), GPUBufferOffset(InGPUBufferOffset), NumComponents(InNumComponents), bUpload(InbUpload)
+	FNiagaraRendererVariableInfo(int32 InDatasetOffset, int32 InGPUBufferOffset, int32 InNumComponents, bool InbUpload, bool bInHalfType)
+		: DatasetOffset(InDatasetOffset), GPUBufferOffset(InGPUBufferOffset), NumComponents(InNumComponents), bUpload(InbUpload), bHalfType(bInHalfType)
 	{
 	}
 
-	FORCEINLINE int32 GetGPUOffset()const { return GbEnableMinimalGPUBuffers ? GPUBufferOffset : DatasetOffset; }
+	FORCEINLINE int32 GetGPUOffset()const 
+	{ 
+		int32 Offset = GbEnableMinimalGPUBuffers ? GPUBufferOffset : DatasetOffset;
+		if (bHalfType)
+		{
+			Offset |= 1 << 31;
+		}
+		return Offset;
+	}
 
 	int32 DatasetOffset;
 	int32 GPUBufferOffset;
 	int32 NumComponents;
 	bool bUpload;
+	bool bHalfType;
+};
+
+struct FParticleRenderData
+{
+	FGlobalDynamicReadBuffer::FAllocation FloatData;
+	FGlobalDynamicReadBuffer::FAllocation HalfData;
 };
 
 /**
@@ -123,6 +138,7 @@ public:
 	NIAGARA_API static FRHIShaderResourceView* GetDummyUIntBuffer();
 	NIAGARA_API static FRHIShaderResourceView* GetDummyUInt4Buffer();
 	NIAGARA_API static FRHIShaderResourceView* GetDummyTextureReadBuffer2D();
+	NIAGARA_API static FRHIShaderResourceView* GetDummyHalfBuffer();
 
 	FORCEINLINE ENiagaraSimTarget GetSimTarget() const { return SimTarget; }
 
@@ -150,10 +166,11 @@ protected:
 	virtual int32 GetMaxIndirectArgs() const { return 1; }
 
 	bool SetVertexFactoryVariable(const FNiagaraDataSet& DataSet, const FNiagaraVariable& Var, int32 VFVarOffset);
-	FGlobalDynamicReadBuffer::FAllocation TransferDataToGPU(FGlobalDynamicReadBuffer& DynamicReadBuffer, FNiagaraDataBuffer* SrcData)const;
+	FParticleRenderData TransferDataToGPU(FGlobalDynamicReadBuffer& DynamicReadBuffer, FNiagaraDataBuffer* SrcData)const;
 	
 	mutable TArray<FNiagaraRendererVariableInfo, TInlineAllocator<16>> VFVariables;
-	int32 TotalVFComponents;
+	int32 TotalVFHalfComponents;
+	int32 TotalVFFloatComponents;
 
 	/** Cached array of materials used from the properties data. Validated with usage flags etc. */
 	TArray<UMaterialInterface*> BaseMaterials_GT;
diff --git a/Engine/Plugins/FX/Niagara/Source/Niagara/Public/NiagaraRendererRibbons.h b/Engine/Plugins/FX/Niagara/Source/Niagara/Public/NiagaraRendererRibbons.h
--- a/Engine/Plugins/FX/Niagara/Source/Niagara/Public/NiagaraRendererRibbons.h
+++ b/Engine/Plugins/FX/Niagara/Source/Niagara/Public/NiagaraRendererRibbons.h
@@ -42,7 +42,7 @@ private:
 	struct FCPUSimParticleDataAllocation
 	{
 		FGlobalDynamicReadBuffer& DynamicReadBuffer;
-		FGlobalDynamicReadBuffer::FAllocation ParticleData;
+		FParticleRenderData ParticleData;
 	};
 
 	void SetupMeshBatchAndCollectorResourceForView(
diff --git a/Engine/Plugins/FX/Niagara/Source/Niagara/Public/NiagaraRendererSprites.h b/Engine/Plugins/FX/Niagara/Source/Niagara/Public/NiagaraRendererSprites.h
--- a/Engine/Plugins/FX/Niagara/Source/Niagara/Public/NiagaraRendererSprites.h
+++ b/Engine/Plugins/FX/Niagara/Source/Niagara/Public/NiagaraRendererSprites.h
@@ -37,7 +37,7 @@ private:
 	struct FCPUSimParticleDataAllocation
 	{
 		FGlobalDynamicReadBuffer& DynamicReadBuffer;
-		FGlobalDynamicReadBuffer::FAllocation ParticleData;
+		FParticleRenderData ParticleData;
 	};
 
 	FCPUSimParticleDataAllocation ConditionalAllocateCPUSimParticleData(FNiagaraDynamicDataSprites *DynamicDataSprites, FGlobalDynamicReadBuffer& DynamicReadBuffer) const;
@@ -47,6 +47,7 @@ private:
 		FNiagaraDynamicDataSprites *DynamicDataSprites,
 		FCPUSimParticleDataAllocation& CPUSimParticleDataAllocation,
 		const FSceneView* View,
+		class FNiagaraSpriteVFLooseParameters& VFLooseParams,
 		const FNiagaraSceneProxy *SceneProxy) const;
 	void CreateMeshBatchForView(
 		const FSceneView* View,
@@ -55,6 +56,7 @@ private:
 		FNiagaraDynamicDataSprites *DynamicDataSprites,
 		uint32 IndirectArgsOffset,
 		FMeshBatch& OutMeshBatch,
+		class FNiagaraSpriteVFLooseParameters& VFLooseParams,
 		class FNiagaraMeshCollectorResourcesSprite& OutCollectorResources) const;
 
 	//Cached data from the properties struct.
