FNM - Added Compression option for Niagara Systems. Thge Option will try to compress some of the parameters to use half precission, which can be a performance benefit on the GPU and sometimes also on 

diff --git a/Engine/Plugins/FX/Niagara/Source/Niagara/Public/NiagaraSystemSimulation.h b/Engine/Plugins/FX/Niagara/Source/Niagara/Public/NiagaraSystemSimulation.h
--- a/Engine/Plugins/FX/Niagara/Source/Niagara/Public/NiagaraSystemSimulation.h
+++ b/Engine/Plugins/FX/Niagara/Source/Niagara/Public/NiagaraSystemSimulation.h
@@ -32,13 +32,21 @@ struct FNiagaraParameterStoreToDataSetBinding
 		int32 DataSetComponentOffset;
 		FDataOffsets(int32 InParamOffset, int32 InDataSetComponentOffset) : ParameterOffset(InParamOffset), DataSetComponentOffset(InDataSetComponentOffset) {}
 	};
+	struct FHalfDataOffsets : public FDataOffsets
+	{
+		bool ApplyAsFloat;
+		FHalfDataOffsets(int32 InParamOffset, int32 InDataSetComponentOffset, bool InApplyAsFloat) : FDataOffsets(InParamOffset, InDataSetComponentOffset), ApplyAsFloat(InApplyAsFloat) {}
+	};
+
 	TArray<FDataOffsets> FloatOffsets;
 	TArray<FDataOffsets> Int32Offsets;
+	TArray<FHalfDataOffsets> HalfOffsets;
 
 	void Empty()
 	{
 		FloatOffsets.Empty();
 		Int32Offsets.Empty();
+		HalfOffsets.Empty();
 	}
 
 	void Init(FNiagaraDataSet& DataSet, const FNiagaraParameterStore& ParameterStore)
@@ -49,9 +57,10 @@ struct FNiagaraParameterStoreToDataSetBinding
 		for (const FNiagaraVariable& Var : DataSet.GetVariables())
 		{
 			const FNiagaraVariableLayoutInfo* Layout = DataSet.GetVariableLayout(Var);
-			const int32* ParameterOffsetPtr = ParameterStore.FindParameterOffset(Var);
+			const int32* ParameterOffsetPtr = ParameterStore.FindParameterOffset(Var, true);
 			int32 NumFloats = 0;
 			int32 NumInts = 0;
+			int32 NumHalfs = 0;
 			if (ParameterOffsetPtr && Layout)
 			{
 				int32 ParameterOffset = *ParameterOffsetPtr;
@@ -67,6 +76,26 @@ struct FNiagaraParameterStoreToDataSetBinding
 					int32 DataSetOffset = Layout->Int32ComponentStart + NumInts++;
 					Int32Offsets.Emplace(ParamOffset, DataSetOffset);
 				}
+				for (uint32 CompIdx = 0; CompIdx < Layout->GetNumHalfComponents(); ++CompIdx)
+				{
+					constexpr bool ParameterSetsSupportHalf = false;
+
+					if (ParameterSetsSupportHalf)
+					{
+						int32 ParamOffset = ParameterOffset + Layout->LayoutInfo.HalfComponentByteOffsets[CompIdx];
+						int32 DataSetOffset = Layout->HalfComponentStart + NumHalfs++;
+						HalfOffsets.Emplace(ParamOffset, DataSetOffset, !ParameterSetsSupportHalf);
+					}
+					else
+					{
+						// if parameter sets don't support half, then we need to write in floats into the parameter set, and
+						// for that we need to adjust the offset based on the difference in stride between float & half
+						// In reality 
+						int32 ParamOffset = ParameterOffset + sizeof(float) * Layout->LayoutInfo.HalfComponentByteOffsets[CompIdx] / sizeof(FFloat16);
+						int32 DataSetOffset = Layout->HalfComponentStart + NumHalfs++;
+						HalfOffsets.Emplace(ParamOffset, DataSetOffset, !ParameterSetsSupportHalf);
+					}
+				}
 			}
 		}
 	}
@@ -89,6 +118,18 @@ struct FNiagaraParameterStoreToDataSetBinding
 			int32* DataSetPtr = CurrBuffer->GetInstancePtrInt32(DataOffsets.DataSetComponentOffset, DataSetInstanceIndex);
 			ParameterStore.SetParameterByOffset(DataOffsets.ParameterOffset, *DataSetPtr);
 		}
+		for (const FHalfDataOffsets& DataOffsets : HalfOffsets)
+		{
+			FFloat16* DataSetPtr = CurrBuffer->GetInstancePtrHalf(DataOffsets.DataSetComponentOffset, DataSetInstanceIndex);
+			if (DataOffsets.ApplyAsFloat)
+			{
+				ParameterStore.SetParameterByOffset(DataOffsets.ParameterOffset, DataSetPtr->GetFloat());
+			}
+			else
+			{
+				ParameterStore.SetParameterByOffset(DataOffsets.ParameterOffset, *DataSetPtr);
+			}
+		}
 
 #if NIAGARA_NAN_CHECKING
 		ParameterStore.CheckForNaNs();
@@ -118,6 +159,21 @@ struct FNiagaraParameterStoreToDataSetBinding
 			int32* DataSetPtr = CurrBuffer.GetInstancePtrInt32(DataOffsets.DataSetComponentOffset, DataSetInstanceIndex);
 			*DataSetPtr = *ParamPtr;
 		}
+		for (const FHalfDataOffsets& DataOffsets : HalfOffsets)
+		{
+			FFloat16* DataSetPtr = CurrBuffer.GetInstancePtrHalf(DataOffsets.DataSetComponentOffset, DataSetInstanceIndex);
+
+			if (DataOffsets.ApplyAsFloat)
+			{
+				float* ParamPtr = (float*)(ParameterData + DataOffsets.ParameterOffset);
+				*DataSetPtr = *ParamPtr;
+			}
+			else
+			{
+				FFloat16* ParamPtr = (FFloat16*)(ParameterData + DataOffsets.ParameterOffset);
+				*DataSetPtr = *ParamPtr;
+			}
+		}
 
 #if NIAGARA_NAN_CHECKING
 		DataSet.CheckForNaNs();
diff --git a/Engine/Plugins/FX/Niagara/Source/Niagara/Public/NiagaraTypes.h b/Engine/Plugins/FX/Niagara/Source/Niagara/Public/NiagaraTypes.h
--- a/Engine/Plugins/FX/Niagara/Source/Niagara/Public/NiagaraTypes.h
+++ b/Engine/Plugins/FX/Niagara/Source/Niagara/Public/NiagaraTypes.h
@@ -70,6 +70,60 @@ private:
 	int32 Value = FNiagaraBool::False;
 };
 
+USTRUCT(meta = (DisplayName = "Half"))
+struct FNiagaraHalf
+{
+	GENERATED_USTRUCT_BODY()
+
+	UPROPERTY(EditAnywhere, Category = Parameters)
+	uint16 Value = 0;
+};
+
+USTRUCT(meta = (DisplayName = "Half Vector2"))
+struct FNiagaraHalfVector2
+{
+	GENERATED_USTRUCT_BODY()
+
+	UPROPERTY(EditAnywhere, Category = Parameters)
+	uint16 x = 0;
+
+	UPROPERTY(EditAnywhere, Category = Parameters)
+	uint16 y = 0;
+};
+
+USTRUCT(meta = (DisplayName = "Half Vector3"))
+struct FNiagaraHalfVector3
+{
+	GENERATED_USTRUCT_BODY()
+
+	UPROPERTY(EditAnywhere, Category = Parameters)
+	uint16 x = 0;
+
+	UPROPERTY(EditAnywhere, Category = Parameters)
+	uint16 y = 0;
+
+	UPROPERTY(EditAnywhere, Category = Parameters)
+	uint16 z = 0;
+};
+
+USTRUCT(meta = (DisplayName = "Half Vector4"))
+struct FNiagaraHalfVector4
+{
+	GENERATED_USTRUCT_BODY()
+
+	UPROPERTY(EditAnywhere, Category = Parameters)
+	uint16 x = 0;
+
+	UPROPERTY(EditAnywhere, Category = Parameters)
+	uint16 y = 0;
+
+	UPROPERTY(EditAnywhere, Category = Parameters)
+	uint16 z = 0;
+
+	UPROPERTY(EditAnywhere, Category = Parameters)
+	uint16 w = 0;
+};
+
 USTRUCT()
 struct FNiagaraNumeric
 {
@@ -211,7 +265,13 @@ struct FNiagaraTypeLayoutInfo
 	UPROPERTY()
 	TArray<uint32> Int32ComponentRegisterOffsets;
 
-	FORCEINLINE uint32 GetNumComponents()const { return FloatComponentByteOffsets.Num() + Int32ComponentByteOffsets.Num(); }
+	/** Byte offset of each half component in a structured layout. */
+	UPROPERTY()
+	TArray<uint32> HalfComponentByteOffsets;
+
+	/** Offset into register table for each half component. */
+	UPROPERTY()
+	TArray<uint32> HalfComponentRegisterOffsets;
 
 	static void GenerateLayoutInfo(FNiagaraTypeLayoutInfo& Layout, const UScriptStruct* Struct)
 	{
@@ -219,6 +279,8 @@ struct FNiagaraTypeLayoutInfo
 		Layout.FloatComponentRegisterOffsets.Empty();
 		Layout.Int32ComponentByteOffsets.Empty();
 		Layout.Int32ComponentRegisterOffsets.Empty();
+		Layout.HalfComponentByteOffsets.Empty();
+		Layout.HalfComponentRegisterOffsets.Empty();
 		GenerateLayoutInfoInternal(Layout, Struct);
 	}
 
@@ -231,12 +293,17 @@ private:
 			int32 PropOffset = BaseOffest + Property->GetOffset_ForInternal();
 			if (Property->IsA(FFloatProperty::StaticClass()))
 			{
-				Layout.FloatComponentRegisterOffsets.Add(Layout.GetNumComponents());
+				Layout.FloatComponentRegisterOffsets.Add(Layout.FloatComponentByteOffsets.Num());
 				Layout.FloatComponentByteOffsets.Add(PropOffset);
 			}
+			else if (Property->IsA(FUInt16Property::StaticClass()))
+			{
+				Layout.HalfComponentRegisterOffsets.Add(Layout.HalfComponentByteOffsets.Num());
+				Layout.HalfComponentByteOffsets.Add(PropOffset);
+			}
 			else if (Property->IsA(FIntProperty::StaticClass()) || Property->IsA(FBoolProperty::StaticClass()))
 			{
-				Layout.Int32ComponentRegisterOffsets.Add(Layout.GetNumComponents());
+				Layout.Int32ComponentRegisterOffsets.Add(Layout.Int32ComponentByteOffsets.Num());
 				Layout.Int32ComponentByteOffsets.Add(PropOffset);
 			}
 			//Should be able to support double easily enough
@@ -902,6 +969,11 @@ public:
 	static const FNiagaraTypeDefinition& GetUObjectDef() { return UObjectDef; }
 	static const FNiagaraTypeDefinition& GetUMaterialDef() { return UMaterialDef; }
 
+	static const FNiagaraTypeDefinition& GetHalfDef() { return HalfDef; }
+	static const FNiagaraTypeDefinition& GetHalfVec2Def() { return HalfVec2Def; }
+	static const FNiagaraTypeDefinition& GetHalfVec3Def() { return HalfVec3Def; }
+	static const FNiagaraTypeDefinition& GetHalfVec4Def() { return HalfVec4Def; }
+
 	static UScriptStruct* GetFloatStruct() { return FloatStruct; }
 	static UScriptStruct* GetBoolStruct() { return BoolStruct; }
 	static UScriptStruct* GetIntStruct() { return IntStruct; }
@@ -915,6 +987,11 @@ public:
 	static UScriptStruct* GetParameterMapStruct() { return ParameterMapStruct; }
 	static UScriptStruct* GetIDStruct() { return IDStruct; }
 
+	static UScriptStruct* GetHalfStruct() { return HalfStruct; }
+	static UScriptStruct* GetHalfVec2Struct() { return HalfVec2Struct; }
+	static UScriptStruct* GetHalfVec3Struct() { return HalfVec3Struct; }
+	static UScriptStruct* GetHalfVec4Struct() { return HalfVec4Struct; }
+
 	static UEnum* GetExecutionStateEnum() { return ExecutionStateEnum; }
 	static UEnum* GetExecutionStateSouceEnum() { return ExecutionStateSourceEnum; }
 	static UEnum* GetSimulationTargetEnum() { return SimulationTargetEnum; }
@@ -959,6 +1036,11 @@ private:
 	static FNiagaraTypeDefinition UObjectDef;
 	static FNiagaraTypeDefinition UMaterialDef;
 
+	static FNiagaraTypeDefinition HalfDef;
+	static FNiagaraTypeDefinition HalfVec2Def;
+	static FNiagaraTypeDefinition HalfVec3Def;
+	static FNiagaraTypeDefinition HalfVec4Def;
+
 	static UScriptStruct* FloatStruct;
 	static UScriptStruct* BoolStruct;
 	static UScriptStruct* IntStruct;
@@ -970,6 +1052,11 @@ private:
 	static UScriptStruct* Matrix4Struct;
 	static UScriptStruct* NumericStruct;
 
+	static UScriptStruct* HalfStruct;
+	static UScriptStruct* HalfVec2Struct;
+	static UScriptStruct* HalfVec3Struct;
+	static UScriptStruct* HalfVec4Struct;
+
 	static UClass* UObjectClass;
 	static UClass* UMaterialClass;
 
diff --git a/Engine/Plugins/FX/Niagara/Source/Niagara/Public/NiagaraUserRedirectionParameterStore.h b/Engine/Plugins/FX/Niagara/Source/Niagara/Public/NiagaraUserRedirectionParameterStore.h
--- a/Engine/Plugins/FX/Niagara/Source/Niagara/Public/NiagaraUserRedirectionParameterStore.h
+++ b/Engine/Plugins/FX/Niagara/Source/Niagara/Public/NiagaraUserRedirectionParameterStore.h
@@ -29,7 +29,7 @@ public:
 	FORCEINLINE void GetUserParameters(TArray<FNiagaraVariable>& OutParameters) const { return UserParameterRedirects.GenerateKeyArray(OutParameters); }
 
 	// ~ Begin FNiagaraParameterStore overrides
-	FORCEINLINE_DEBUGGABLE virtual const int32* FindParameterOffset(const FNiagaraVariable& Parameter) const override
+	FORCEINLINE_DEBUGGABLE virtual const int32* FindParameterOffset(const FNiagaraVariable& Parameter, bool IgnoreType = false) const override
 	{
 		const FNiagaraVariable* Redirection = UserParameterRedirects.Find(Parameter);
 		return FNiagaraParameterStore::FindParameterOffset(Redirection ? *Redirection : Parameter);
diff --git a/Engine/Plugins/FX/Niagara/Source/NiagaraCore/Private/NiagaraCustomVersion.cpp b/Engine/Plugins/FX/Niagara/Source/NiagaraCore/Private/NiagaraCustomVersion.cpp
--- a/Engine/Plugins/FX/Niagara/Source/NiagaraCore/Private/NiagaraCustomVersion.cpp
+++ b/Engine/Plugins/FX/Niagara/Source/NiagaraCore/Private/NiagaraCustomVersion.cpp
@@ -8,4 +8,4 @@ const FGuid FNiagaraCustomVersion::GUID(0xFCF57AFA, 0x50764283, 0xB9A9E658, 0xFF
 // Register the custom version with core
 FCustomVersionRegistration GRegisterNiagaraCustomVersion(FNiagaraCustomVersion::GUID, FNiagaraCustomVersion::LatestVersion, TEXT("NiagaraVer"));
 
-const FGuid FNiagaraCustomVersion::LatestScriptCompileVersion(0x6D52A92E, 0xDC5C4791, 0x922FD12B, 0x82A7D6B7);
\ No newline at end of file
+const FGuid FNiagaraCustomVersion::LatestScriptCompileVersion(0x429FAAA4, 0xB12B4D31, 0x922FD12B, 0x82A7D6B7);
diff --git a/Engine/Plugins/FX/Niagara/Source/NiagaraEditor/Private/NiagaraHlslTranslator.cpp b/Engine/Plugins/FX/Niagara/Source/NiagaraEditor/Private/NiagaraHlslTranslator.cpp
--- a/Engine/Plugins/FX/Niagara/Source/NiagaraEditor/Private/NiagaraHlslTranslator.cpp
+++ b/Engine/Plugins/FX/Niagara/Source/NiagaraEditor/Private/NiagaraHlslTranslator.cpp
@@ -810,6 +810,59 @@ void FHlslNiagaraTranslator::BuildConstantBuffer(ENiagaraCodeChunkMode ChunkMode
 	}
 }
 
+static void ConvertFloatToHalf(const FNiagaraCompileOptions& InCompileOptions, TArray<FNiagaraVariable>& Attributes)
+{
+	if (InCompileOptions.AdditionalDefines.Contains(TEXT("CompressAttributes")))// && UNiagaraScript::IsParticleScript(InCompileOptions.TargetUsage))
+	{
+		static FNiagaraTypeDefinition ConvertMapping[][2] =
+		{
+			{ FNiagaraTypeDefinition::GetFloatDef(), FNiagaraTypeDefinition::GetHalfDef() },
+			{ FNiagaraTypeDefinition::GetVec2Def(), FNiagaraTypeDefinition::GetHalfVec2Def() },
+			{ FNiagaraTypeDefinition::GetVec3Def(), FNiagaraTypeDefinition::GetHalfVec3Def() },
+			{ FNiagaraTypeDefinition::GetVec4Def(), FNiagaraTypeDefinition::GetHalfVec4Def() },
+			{ FNiagaraTypeDefinition::GetColorDef(), FNiagaraTypeDefinition::GetHalfVec4Def() },
+			{ FNiagaraTypeDefinition::GetQuatDef(), FNiagaraTypeDefinition::GetHalfVec4Def() },
+		};
+
+		TArray<FNiagaraVariable> ConvertExceptions =
+		{
+			SYS_PARAM_ENGINE_POSITION,
+			SYS_PARAM_ENGINE_INV_DELTA_TIME,
+			SYS_PARAM_ENGINE_TIME,
+			SYS_PARAM_ENGINE_REAL_TIME,
+			SYS_PARAM_ENGINE_SYSTEM_AGE,
+			SYS_PARAM_ENGINE_SYSTEM_NUM_EMITTERS_ALIVE,
+			SYS_PARAM_ENGINE_SYSTEM_NUM_EMITTERS,
+			SYS_PARAM_ENGINE_NUM_SYSTEM_INSTANCES,
+			SYS_PARAM_ENGINE_EMITTER_NUM_PARTICLES,
+			SYS_PARAM_ENGINE_EMITTER_TOTAL_SPAWNED_PARTICLES,
+			SYS_PARAM_PARTICLES_UNIQUE_ID,
+			SYS_PARAM_PARTICLES_ID,
+			SYS_PARAM_EMITTER_AGE,
+			SYS_PARAM_PARTICLES_POSITION,
+			SYS_PARAM_PARTICLES_LIFETIME,
+		};
+
+		for (FNiagaraVariable& Attribute : Attributes)
+		{
+			// check if the variable matches an exception that we don't want to convert
+			if (FNiagaraVariable::SearchArrayForPartialNameMatch(ConvertExceptions, Attribute.GetName()) != INDEX_NONE)
+			{
+				continue;
+			}
+
+			for (int32 ConvertIt = 0; ConvertIt < UE_ARRAY_COUNT(ConvertMapping); ++ConvertIt)
+			{
+				if (Attribute.GetType() == ConvertMapping[ConvertIt][0])
+				{
+					Attribute.SetType(ConvertMapping[ConvertIt][1]);
+					break;
+				}
+			}
+		}
+	}
+}
+
 const FNiagaraTranslateResults &FHlslNiagaraTranslator::Translate(const FNiagaraCompileRequestData* InCompileData, const FNiagaraCompileOptions& InCompileOptions, FHlslNiagaraTranslatorOptions InTranslateOptions)
 {
 	SCOPE_CYCLE_COUNTER(STAT_NiagaraEditor_HlslTranslator_Translate);
@@ -1392,6 +1445,8 @@ const FNiagaraTranslateResults &FHlslNiagaraTranslator::Translate(const FNiagara
 		// We sort the variables so that they end up in the same ordering between Spawn & Update...
 		Algo::SortBy(BasicAttributes, &FNiagaraVariable::GetName, FNameLexicalLess());
 
+		ConvertFloatToHalf(CompileOptions, BasicAttributes);
+
 		DataSetVariables[InstanceReadVarsIndex] = BasicAttributes;
 		DataSetVariables[InstanceWriteVarsIndex] = BasicAttributes;
 
@@ -1553,7 +1608,7 @@ const FNiagaraTranslateResults &FHlslNiagaraTranslator::Translate(const FNiagara
 }
 
 
-void FHlslNiagaraTranslator::GatherVariableForDataSetAccess(const FNiagaraVariable& Var, FString Format, int32& IntCounter, int32 &FloatCounter, int32 DataSetIndex, FString InstanceIdxSymbol, FString &HlslOutputString)
+void FHlslNiagaraTranslator::GatherVariableForDataSetAccess(const FNiagaraVariable& Var, FString Format, int32& IntCounter, int32 &FloatCounter, int32& HalfCounter, int32 DataSetIndex, FString InstanceIdxSymbol, FString &HlslOutputString)
 {
 	TArray<FString> Components;
 	UScriptStruct* Struct = Var.GetType().GetScriptStruct();
@@ -1576,12 +1631,12 @@ void FHlslNiagaraTranslator::GatherVariableForDataSetAccess(const FNiagaraVariab
 	{
 		FormatArgs.Add(DataSetIndex);
 	}
-	int32 RegIdx = FormatArgs.Add(0);
+	const int32 RegIdx = FormatArgs.Add(0);
 	if (!InstanceIdxSymbol.IsEmpty())
 	{
 		FormatArgs.Add(InstanceIdxSymbol);
 	}
-	int32 DefaultIdx = FormatArgs.Add(0);
+	const int32 DefaultIdx = FormatArgs.Add(0);
 
 	check(Components.Num() == Types.Num());
 	for (int32 CompIdx = 0; CompIdx < Components.Num(); ++CompIdx)
@@ -1592,25 +1647,24 @@ void FHlslNiagaraTranslator::GatherVariableForDataSetAccess(const FNiagaraVariab
 			FormatArgs[DefaultIdx] = TEXT("0.0f");
 			FormatArgs[RegIdx] = FloatCounter++;
 		}
+		else if (Types[CompIdx] == NBT_Half)
+		{
+			FormatArgs[1] = TEXT("Half");
+			FormatArgs[DefaultIdx] = TEXT("0.0f");
+			FormatArgs[RegIdx] = HalfCounter++;
+		}
 		else if (Types[CompIdx] == NBT_Int32)
 		{
 			FormatArgs[1] = TEXT("Int");
 			FormatArgs[DefaultIdx] = TEXT("0");
-			if (CompilationTarget == ENiagaraSimTarget::GPUComputeSim)
-			{
-				FormatArgs[RegIdx] = IntCounter++;
-			}
-			else
-			{
-				FormatArgs[RegIdx] = FloatCounter++;
-			}
+			FormatArgs[RegIdx] = IntCounter++;
 		}
 		else
 		{
 			check(Types[CompIdx] == NBT_Bool);
 			FormatArgs[1] = TEXT("Bool");
 			FormatArgs[DefaultIdx] = TEXT("false");
-			FormatArgs[RegIdx] = CompilationTarget == ENiagaraSimTarget::GPUComputeSim ? IntCounter++ : FloatCounter++;
+			FormatArgs[RegIdx] = IntCounter++;
 		}
 		FormatArgs[0] = Components[CompIdx];
 		HlslOutputString += FString::Format(*Format, FormatArgs);
@@ -1634,6 +1688,12 @@ void FHlslNiagaraTranslator::GatherComponentsForDataSetAccess(UScriptStruct* Str
 		Components.Add(VariableSymbol);
 		return;
 	}
+	else if (FNiagaraTypeDefinition(Struct) == FNiagaraTypeDefinition::GetHalfDef())
+	{
+		Types.Add(ENiagaraBaseTypes::NBT_Half);
+		Components.Add(VariableSymbol);
+		return;
+	}
 
 	for (TFieldIterator<const FProperty> PropertyIt(Struct, EFieldIteratorFlags::IncludeSuper); PropertyIt; ++PropertyIt)
 	{
@@ -1685,6 +1745,11 @@ void FHlslNiagaraTranslator::GatherComponentsForDataSetAccess(UScriptStruct* Str
 				Types.Add(ENiagaraBaseTypes::NBT_Bool);
 				Components.Add(VarName);
 			}
+			else if (Property->IsA(FUInt16Property::StaticClass()))
+			{
+				Types.Add(ENiagaraBaseTypes::NBT_Half);
+				Components.Add(VarName);
+			}
 		}
 	}
 }
@@ -1802,8 +1867,7 @@ void FHlslNiagaraTranslator::DefineDataSetReadFunction(FString &HlslOutputString
 		FNiagaraDataSetID DataSet = DataSetInfoPair.Key;
 		int32 OffsetCounterInt = 0;
 		int32 OffsetCounterFloat = 0;
-		int32 &FloatCounter = OffsetCounterFloat;
-		int32 &IntCounter = CompilationTarget == ENiagaraSimTarget::GPUComputeSim ? OffsetCounterInt : OffsetCounterFloat;
+		int32 OffsetCounterHalf = 0;
 		int32 DataSetIndex = 1;
 		for (TPair<int32, FDataSetAccessInfo>& IndexInfoPair : DataSetInfoPair.Value)
 		{
@@ -1816,7 +1880,7 @@ void FHlslNiagaraTranslator::DefineDataSetReadFunction(FString &HlslOutputString
 				{
 					// TODO: temp = should really generate output functions for each set
 					FString Fmt = Symbol + Var.GetName().ToString() + FString(TEXT("{0} = ReadDataSet{1}")) + SetIdx + TEXT("[{2}*") + DataSetComponentBufferSize + " + SetInstanceIndex];\n";
-					GatherVariableForDataSetAccess(Var, Fmt, IntCounter, FloatCounter, -1, TEXT(""), HlslOutputString);
+					GatherVariableForDataSetAccess(Var, Fmt, OffsetCounterInt, OffsetCounterFloat, OffsetCounterHalf, -1, TEXT(""), HlslOutputString);
 				}
 			}
 			else
@@ -1825,7 +1889,7 @@ void FHlslNiagaraTranslator::DefineDataSetReadFunction(FString &HlslOutputString
 				{
 					// TODO: currently always emitting a non-advancing read, needs to be changed for some of the use cases
 					FString Fmt = TEXT("\tContext.") + DataSet.Name.ToString() + "Read." + Var.GetName().ToString() + TEXT("{0} = InputDataNoadvance{1}({2}, {3});\n");
-					GatherVariableForDataSetAccess(Var, Fmt, IntCounter, FloatCounter, DataSetIndex, TEXT(""), HlslOutputString);
+					GatherVariableForDataSetAccess(Var, Fmt, OffsetCounterInt, OffsetCounterFloat, OffsetCounterHalf, DataSetIndex, TEXT(""), HlslOutputString);
 				}
 			}
 		}
@@ -1858,8 +1922,7 @@ void FHlslNiagaraTranslator::DefineDataSetWriteFunction(FString &HlslOutputStrin
 		}
 		int32 WriteOffsetInt = 0;
 		int32 WriteOffsetFloat = 0;
-		int32 &FloatCounter = WriteOffsetFloat;
-		int32 &IntCounter = CompilationTarget == ENiagaraSimTarget::GPUComputeSim ? WriteOffsetInt : WriteOffsetFloat;
+		int32 WriteOffsetHalf = 0;
 
 		// grab the current ouput index; currently pass true, but should use an arbitrary bool to determine whether write should happen or not
 
@@ -1880,7 +1943,7 @@ void FHlslNiagaraTranslator::DefineDataSetWriteFunction(FString &HlslOutputStrin
 				{
 					// TODO: temp = should really generate output functions for each set
 					FString Fmt = FString(TEXT("\t\tRWWriteDataSet{1}")) + SetIdx + TEXT("[{2}*") + DataSetComponentBufferSize + TEXT(" + {3}] = ") + Symbol + TEXT(".") + Var.GetName().ToString() + TEXT("{0};\n");
-					GatherVariableForDataSetAccess(Var, Fmt, IntCounter, FloatCounter, -1, TEXT("TmpWriteIndex"), HlslOutputString);
+					GatherVariableForDataSetAccess(Var, Fmt, WriteOffsetInt, WriteOffsetFloat, WriteOffsetHalf, -1, TEXT("TmpWriteIndex"), HlslOutputString);
 				}
 			}
 			else
@@ -1889,7 +1952,7 @@ void FHlslNiagaraTranslator::DefineDataSetWriteFunction(FString &HlslOutputStrin
 				{
 					// TODO: data set index is always 1; need to increase each set
 					FString Fmt = TEXT("\t\tOutputData{1}(") + FString::FromInt(DataSetIndex) + (", {2}, {3}, ") + Symbol + "." + Var.GetName().ToString() + TEXT("{0});\n");
-					GatherVariableForDataSetAccess(Var, Fmt, FloatCounter, FloatCounter, -1, TEXT("TmpWriteIndex"), HlslOutputString);
+					GatherVariableForDataSetAccess(Var, Fmt, WriteOffsetInt, WriteOffsetFloat, WriteOffsetHalf, -1, TEXT("TmpWriteIndex"), HlslOutputString);
 				}
 			}
 		}
@@ -2074,14 +2137,14 @@ void FHlslNiagaraTranslator::DefineMainGPUFunctions(
 		}
 
 		//The VM register binding assumes the same inputs as outputs which is obviously not always the case.
-		for (int32 DataSetIndex = 0, IntCounter = 0, FloatCounter = 0; DataSetIndex < DataSetReads.Num(); ++DataSetIndex)
+		for (int32 DataSetIndex = 0, IntCounter = 0, FloatCounter = 0, HalfCounter = 0; DataSetIndex < DataSetReads.Num(); ++DataSetIndex)
 		{
 			const FNiagaraDataSetID DataSetID = ReadDataSetIDs[DataSetIndex];
 			const TArray<FNiagaraVariable>& NiagaraVariables = DataSetVariables[DataSetReads[DataSetID]];
 			for (const FNiagaraVariable& Var : NiagaraVariables)
 			{
 				FString VarFmt = ContextName + GetSanitizedSymbolName(Var.GetName().ToString()) + TEXT("{0} = {4};\n");
-				GatherVariableForDataSetAccess(Var, VarFmt, IntCounter, FloatCounter, DataSetIndex, TEXT(""), HlslOutput);
+				GatherVariableForDataSetAccess(Var, VarFmt, IntCounter, FloatCounter, HalfCounter, DataSetIndex, TEXT(""), HlslOutput);
 			}
 		}
 
@@ -2131,7 +2194,7 @@ void FHlslNiagaraTranslator::DefineMainGPUFunctions(
 				ContextName = FString::Printf(TEXT("\t\tContext.%s."), *TranslationStages[i].PassNamespace);
 			}
 
-			for (int32 DataSetIndex = 0, IntCounter = 0, FloatCounter = 0; DataSetIndex < DataSetReads.Num(); ++DataSetIndex)
+			for (int32 DataSetIndex = 0, IntCounter = 0, FloatCounter = 0, HalfCounter = 0; DataSetIndex < DataSetReads.Num(); ++DataSetIndex)
 			{
 				const FNiagaraDataSetID DataSetID = ReadDataSetIDs[DataSetIndex];
 				const TArray<FNiagaraVariable>& NiagaraVariables = DataSetVariables[DataSetReads[DataSetID]];
@@ -2161,7 +2224,7 @@ void FHlslNiagaraTranslator::DefineMainGPUFunctions(
 							}
 						}
 					}
-					GatherVariableForDataSetAccess(Var, VarFmt, IntCounter, FloatCounter, DataSetIndex, TEXT(""), HlslOutput);
+					GatherVariableForDataSetAccess(Var, VarFmt, IntCounter, FloatCounter, HalfCounter, DataSetIndex, TEXT(""), HlslOutput);
 				}
 			}
 			if (bUsesAlive)
@@ -2280,7 +2343,7 @@ void FHlslNiagaraTranslator::DefineMainGPUFunctions(
 				HlslOutput += FString::Printf(TEXT("\t\t\tUpdateID(0, %sParticles.ID.Index, WriteIndex);\n"), *ContextName);
 			}
 
-			for (int32 DataSetIndex = 0, IntCounter = 0, FloatCounter = 0; DataSetIndex < DataSetWrites.Num(); ++DataSetIndex)
+			for (int32 DataSetIndex = 0, IntCounter = 0, FloatCounter = 0, HalfCounter = 0; DataSetIndex < DataSetWrites.Num(); ++DataSetIndex)
 			{
 				const FNiagaraDataSetID DataSetID = ReadDataSetIDs[DataSetIndex];
 				const TArray<FNiagaraVariable>& NiagaraVariables = DataSetVariables[DataSetWrites[DataSetID]];
@@ -2288,7 +2351,7 @@ void FHlslNiagaraTranslator::DefineMainGPUFunctions(
 				{
 					// If coming from a parameter map, use the one on the context, otherwise use the output.
 					FString VarFmt = TEXT("\t\t\tOutputData{1}(0, {2}, {3}, ") + ContextName + GetSanitizedSymbolName(Var.GetName().ToString()) + TEXT("{0});\n");
-					GatherVariableForDataSetAccess(Var, VarFmt, IntCounter, FloatCounter, -1, TEXT("WriteIndex"), HlslOutput);
+					GatherVariableForDataSetAccess(Var, VarFmt, IntCounter, FloatCounter, HalfCounter, -1, TEXT("WriteIndex"), HlslOutput);
 				}
 			}
 
@@ -2675,8 +2738,7 @@ void FHlslNiagaraTranslator::DefineDataSetVariableWrites(FString &OutHlslOutput,
 
 	int32 WriteOffsetInt = 0;
 	int32 WriteOffsetFloat = 0;
-	int32 &FloatCounter = WriteOffsetFloat;
-	int32 &IntCounter = WriteOffsetFloat;
+	int32 WriteOffsetHalf = 0;
 	for (const FNiagaraVariable &Var : WriteVars)
 	{
 		// If coming from a parameter map, use the one on the context, otherwise use the output.
@@ -2689,7 +2751,7 @@ void FHlslNiagaraTranslator::DefineDataSetVariableWrites(FString &OutHlslOutput,
 		{
 			Fmt = TEXT("\tOutputData{1}(0, {2}, {3}, Context.Map.") + GetSanitizedSymbolName(Var.GetName().ToString()) + TEXT("{0});\n");
 		}
-		GatherVariableForDataSetAccess(Var, Fmt, IntCounter, FloatCounter, -1, TEXT("TmpWriteIndex"), OutHlslOutput);
+		GatherVariableForDataSetAccess(Var, Fmt, WriteOffsetInt, WriteOffsetFloat, WriteOffsetHalf, -1, TEXT("TmpWriteIndex"), OutHlslOutput);
 	}
 	OutHlslOutput += "\t}\n";
 }
@@ -2700,8 +2762,7 @@ void FHlslNiagaraTranslator::DefineDataSetVariableReads(FString &OutHlslOutput,
 
 	int32 ReadOffsetInt = 0;
 	int32 ReadOffsetFloat = 0;
-	int32 &FloatCounter = ReadOffsetFloat;
-	int32 &IntCounter = ReadOffsetFloat;
+	int32 ReadOffsetHalf = 0;
 
 	FString DataSetName = Id.Name.ToString();
 	FString Fmt;
@@ -2735,7 +2796,7 @@ void FHlslNiagaraTranslator::DefineDataSetVariableReads(FString &OutHlslOutput,
 		for (const FNiagaraVariable &Var : ReadVars)
 		{
 			Fmt = ContextName + GetSanitizedSymbolName(Var.GetName().ToString()) + TEXT("{0} = {4};\n");
-			GatherVariableForDataSetAccess(Var, Fmt, IntCounter, FloatCounter, DataSetIndex, TEXT(""), VarReads);
+			GatherVariableForDataSetAccess(Var, Fmt, ReadOffsetInt, ReadOffsetFloat, ReadOffsetHalf, DataSetIndex, TEXT(""), VarReads);
 		}
 
 		OutHlslOutput += VarReads;
@@ -2756,6 +2817,7 @@ void FHlslNiagaraTranslator::DefineDataSetVariableReads(FString &OutHlslOutput,
 		{
 			ReadOffsetInt = 0;
 			ReadOffsetFloat = 0;
+			ReadOffsetHalf = 0;
 		}
 
 		FString VarReads;
@@ -2779,12 +2841,12 @@ void FHlslNiagaraTranslator::DefineDataSetVariableReads(FString &OutHlslOutput,
 					{
 						FString RegisterName = VariableName;
 						RegisterName.ReplaceInline(PARAM_MAP_ATTRIBUTE_STR, PARAM_MAP_INDICES_STR);
-						const int32 RegisterValue = Var.GetType().IsFloatPrimitive() ? FloatCounter : IntCounter;
-						VarReads += RegisterName + FString::Printf(TEXT(" = %d;\n"), RegisterValue);
+
+						Fmt += RegisterName + TEXT(" = {3};\n");
 					}
 				}
 			}
-			GatherVariableForDataSetAccess(Var, Fmt, IntCounter, FloatCounter, DataSetIndex, TEXT(""), VarReads);
+			GatherVariableForDataSetAccess(Var, Fmt, ReadOffsetInt, ReadOffsetFloat, ReadOffsetHalf, DataSetIndex, TEXT(""), VarReads);
 		}
 
 		OutHlslOutput += VarReads;
diff --git a/Engine/Plugins/FX/Niagara/Source/NiagaraEditor/Private/NiagaraHlslTranslator.h b/Engine/Plugins/FX/Niagara/Source/NiagaraEditor/Private/NiagaraHlslTranslator.h
--- a/Engine/Plugins/FX/Niagara/Source/NiagaraEditor/Private/NiagaraHlslTranslator.h
+++ b/Engine/Plugins/FX/Niagara/Source/NiagaraEditor/Private/NiagaraHlslTranslator.h
@@ -503,7 +503,7 @@ public:
 	TArray<FNiagaraDataInterfaceGPUParamInfo>& GetDataInterfaceParameters() { return DIParamInfo; }
 
 	// Format string should have up to 5 entries, {{0} = Computed Variable Suffix, {1} = Float or Int, {2} = Data Set Index, {3} = Register Index, {4} Default value for that type.
-	void GatherVariableForDataSetAccess(const FNiagaraVariable& Variable, FString Format, int32& RegisterIdxInt, int32& RegisterIdxFloat, int32 DataSetIndex, FString InstanceIdxSymbol, FString &HlslOutput);
+	void GatherVariableForDataSetAccess(const FNiagaraVariable& Variable, FString Format, int32& RegisterIdxInt, int32& RegisterIdxFloat, int32& RegisterIdxHalf, int32 DataSetIndex, FString InstanceIdxSymbol, FString &HlslOutput);
 	void GatherComponentsForDataSetAccess(UScriptStruct* Struct, FString VariableSymbol, bool bMatrixRoot, TArray<FString>& Components, TArray<ENiagaraBaseTypes>& Types);
 
 	FString CompileDataInterfaceFunction(UNiagaraDataInterface* DataInterface, FNiagaraFunctionSignature& Signature);
diff --git a/Engine/Plugins/FX/Niagara/Source/NiagaraEditor/Private/Widgets/SNiagaraSpreadsheetView.cpp b/Engine/Plugins/FX/Niagara/Source/NiagaraEditor/Private/Widgets/SNiagaraSpreadsheetView.cpp
--- a/Engine/Plugins/FX/Niagara/Source/NiagaraEditor/Private/Widgets/SNiagaraSpreadsheetView.cpp
+++ b/Engine/Plugins/FX/Niagara/Source/NiagaraEditor/Private/Widgets/SNiagaraSpreadsheetView.cpp
@@ -95,23 +95,48 @@ TSharedRef< SWidget > SNiagaraSpreadsheetRow::GenerateWidgetForColumn(const FNam
 		
 		if (FieldInfo->bFloat)
 		{
-			float* Src = nullptr;
-			if (UseGlobalOffsets && ParameterStore->GetParameterDataArray().GetData() != nullptr)
-			{
-				const uint32 CompBufferOffset = FieldInfo->GlobalStartOffset;
-				const uint8* SrcU8 = ParameterStore->GetParameterDataArray().GetData();
-				Src = SrcU8 ? (float*)(SrcU8 + CompBufferOffset) : nullptr;
-			}
-			else if (DataSet != nullptr && DataSet->GetCurrentDataChecked().GetNumInstances() > 0)
+			
+			float Value = 0.0f;
+
+			if (FieldInfo->bHalf)
 			{
-				uint32 CompBufferOffset = FieldInfo->FloatStartOffset;
-				Src = DataSet->GetCurrentData()->GetInstancePtrFloat(CompBufferOffset, RealRowIdx);
-			}
+				FFloat16* Src = nullptr;
+				if (UseGlobalOffsets && ParameterStore->GetParameterDataArray().GetData() != nullptr)
+				{
+					const uint32 CompBufferOffset = FieldInfo->GlobalStartOffset;
+					const uint8* SrcU8 = ParameterStore->GetParameterDataArray().GetData();
+					Src = SrcU8 ? (FFloat16*)(SrcU8 + CompBufferOffset) : nullptr;
+				}
+				else if (DataSet != nullptr && DataSet->GetCurrentDataChecked().GetNumInstances() > 0)
+				{
+					uint32 CompBufferOffset = FieldInfo->FloatStartOffset;
+					Src = DataSet->GetCurrentData()->GetInstancePtrHalf(CompBufferOffset, RealRowIdx);
+				}
 
-			float Value = 0.0f;
-			if (Src != nullptr)
+				if (Src != nullptr)
+				{
+					Value = Src[0];
+				}
+			}
+			else
 			{
-				Value = Src[0];
+				float* Src = nullptr;
+				if (UseGlobalOffsets && ParameterStore->GetParameterDataArray().GetData() != nullptr)
+				{
+					const uint32 CompBufferOffset = FieldInfo->GlobalStartOffset;
+					const uint8* SrcU8 = ParameterStore->GetParameterDataArray().GetData();
+					Src = SrcU8 ? (float*)(SrcU8 + CompBufferOffset) : nullptr;
+				}
+				else if (DataSet != nullptr && DataSet->GetCurrentDataChecked().GetNumInstances() > 0)
+				{
+					uint32 CompBufferOffset = FieldInfo->FloatStartOffset;
+					Src = DataSet->GetCurrentData()->GetInstancePtrFloat(CompBufferOffset, RealRowIdx);
+				}
+	
+				if (Src != nullptr)
+				{
+					Value = Src[0];
+				}
 			}
 
 			EntryWidget = SNew(STextBlock)
@@ -964,9 +989,18 @@ FReply SNiagaraSpreadsheetView::OnCSVOutputPressed()
 				{
 					if (FieldInfo->bFloat)
 					{
-						uint32 CompBufferOffset = FieldInfo->FloatStartOffset;
-						float* Src = CaptureData[(int32)TabState].DataSet.GetCurrentDataChecked().GetInstancePtrFloat(CompBufferOffset, RowIndex);
-						CSVOutput += FString::Printf(TEXT("%3.9f"), Src[0]);
+						if (FieldInfo->bHalf)
+						{
+							uint32 CompBufferOffset = FieldInfo->FloatStartOffset;
+							FFloat16* Src = CaptureData[(int32)TabState].DataSet.GetCurrentDataChecked().GetInstancePtrHalf(CompBufferOffset, RowIndex);
+							CSVOutput += FString::Printf(TEXT("%3.9f"), float(Src[0]));
+						}
+						else
+						{
+							uint32 CompBufferOffset = FieldInfo->FloatStartOffset;
+							float* Src = CaptureData[(int32)TabState].DataSet.GetCurrentDataChecked().GetInstancePtrFloat(CompBufferOffset, RowIndex);
+							CSVOutput += FString::Printf(TEXT("%3.9f"), Src[0]);
+						}
 					}
 					else
 					{
@@ -1143,13 +1177,15 @@ void SNiagaraSpreadsheetView::GenerateLayoutInfo(FNiagaraTypeLayoutInfo& Layout,
 		{
 			SNiagaraSpreadsheetView::FieldInfo Info;
 			Info.bFloat = true;
+			Info.bHalf = false;
+			Info.bBoolean = false;
 			Info.FloatStartOffset = Layout.FloatComponentRegisterOffsets.Num();
 			Info.IntStartOffset = UINT_MAX;
 			Info.GlobalStartOffset = sizeof(float) * Layout.FloatComponentRegisterOffsets.Num() + sizeof(int32) * Layout.Int32ComponentByteOffsets.Num();
 			Info.Enum = nullptr;
 			FieldInfo.Add(Info);
 				
-			Layout.FloatComponentRegisterOffsets.Add(Layout.GetNumComponents());
+			Layout.FloatComponentRegisterOffsets.Add(Layout.FloatComponentByteOffsets.Num());
 			Layout.FloatComponentByteOffsets.Add(Property->GetOffset_ForInternal());
 			PropertyNames.Add(PropertyName);
 		}
@@ -1157,6 +1193,7 @@ void SNiagaraSpreadsheetView::GenerateLayoutInfo(FNiagaraTypeLayoutInfo& Layout,
 		{
 			SNiagaraSpreadsheetView::FieldInfo Info;
 			Info.bFloat = false;
+			Info.bHalf = false;
 			Info.bBoolean = Property->IsA(FBoolProperty::StaticClass());
 			Info.FloatStartOffset = UINT_MAX;
 			Info.IntStartOffset = Layout.Int32ComponentRegisterOffsets.Num();
@@ -1164,10 +1201,26 @@ void SNiagaraSpreadsheetView::GenerateLayoutInfo(FNiagaraTypeLayoutInfo& Layout,
 			Info.Enum = Enum;
 			FieldInfo.Add(Info);
 
-			Layout.Int32ComponentRegisterOffsets.Add(Layout.GetNumComponents());
+			Layout.Int32ComponentRegisterOffsets.Add(Layout.Int32ComponentByteOffsets.Num());
 			Layout.Int32ComponentByteOffsets.Add(Property->GetOffset_ForInternal());
 			PropertyNames.Add(PropertyName);
 		}
+		else if (Property->IsA(FUInt16Property::StaticClass()))
+		{
+			SNiagaraSpreadsheetView::FieldInfo Info;
+			Info.bHalf = true;
+			Info.bFloat = true;
+			Info.bBoolean = false;
+			Info.FloatStartOffset = Layout.HalfComponentRegisterOffsets.Num();
+			Info.IntStartOffset = UINT_MAX;
+			Info.GlobalStartOffset = sizeof(float) * Layout.FloatComponentRegisterOffsets.Num() + sizeof(int32) * Layout.Int32ComponentRegisterOffsets.Num() + sizeof(FFloat16) * Layout.HalfComponentByteOffsets.Num();
+			Info.Enum = Enum;
+			FieldInfo.Add(Info);
+
+			Layout.HalfComponentRegisterOffsets.Add(Layout.HalfComponentByteOffsets.Num());
+			Layout.HalfComponentByteOffsets.Add(Property->GetOffset_ForInternal());
+			PropertyNames.Add(PropertyName);
+		}
 		else if (Property->IsA(FEnumProperty::StaticClass()))
 		{
 			FEnumProperty* EnumProp = CastFieldChecked<FEnumProperty>(Property);
@@ -1202,6 +1255,7 @@ void SNiagaraSpreadsheetView::ResetColumns(EUITab Tab)
 			CaptureData[(int32)i].OutputFieldInfoMap = MakeShared<TMap<FName, FieldInfo> >();
 			uint32 TotalFloatComponents = 0;
 			uint32 TotalInt32Components = 0;
+			uint32 TotalHalfComponents = 0;
 
 			TArray<FNiagaraVariable> Variables = CaptureData[(int32)i].DataSet.GetVariables();
 
@@ -1231,6 +1285,7 @@ void SNiagaraSpreadsheetView::ResetColumns(EUITab Tab)
 
 				uint32 TotalFloatComponentsBeforeStruct = TotalFloatComponents;
 				uint32 TotalInt32ComponentsBeforeStruct = TotalInt32Components;
+				uint32 TotalHalfComponentsBeforeStruct = TotalHalfComponents;
 
 				GenerateLayoutInfo(Layout, Struct, Enum, Var.GetName(), PropertyNames, FieldInfos);
 
@@ -1239,8 +1294,16 @@ void SNiagaraSpreadsheetView::ResetColumns(EUITab Tab)
 					const FName PropertyName = PropertyNames[VarIdx];
 					if (FieldInfos[VarIdx].bFloat)
 					{
-						FieldInfos[VarIdx].FloatStartOffset += TotalFloatComponentsBeforeStruct;
-						TotalFloatComponents++;
+						if (FieldInfos[VarIdx].bHalf)
+						{
+							FieldInfos[VarIdx].FloatStartOffset += TotalFloatComponentsBeforeStruct;
+							TotalFloatComponents++;
+						}
+						else
+						{
+							FieldInfos[VarIdx].FloatStartOffset += TotalHalfComponentsBeforeStruct;
+							TotalHalfComponents++;
+						}
 					}
 					else
 					{
diff --git a/Engine/Plugins/FX/Niagara/Source/NiagaraEditor/Private/Widgets/SNiagaraSpreadsheetView.h b/Engine/Plugins/FX/Niagara/Source/NiagaraEditor/Private/Widgets/SNiagaraSpreadsheetView.h
--- a/Engine/Plugins/FX/Niagara/Source/NiagaraEditor/Private/Widgets/SNiagaraSpreadsheetView.h
+++ b/Engine/Plugins/FX/Niagara/Source/NiagaraEditor/Private/Widgets/SNiagaraSpreadsheetView.h
@@ -40,6 +40,7 @@ public:
 		uint32 IntStartOffset;
 		uint32 GlobalStartOffset;
 		bool bFloat;
+		bool bHalf;
 		bool bBoolean;
 		TWeakObjectPtr<const UEnum> Enum;
 	};
diff --git a/Engine/Plugins/FX/Niagara/Source/NiagaraShader/Private/NiagaraShader.cpp b/Engine/Plugins/FX/Niagara/Source/NiagaraShader/Private/NiagaraShader.cpp
--- a/Engine/Plugins/FX/Niagara/Source/NiagaraShader/Private/NiagaraShader.cpp
+++ b/Engine/Plugins/FX/Niagara/Source/NiagaraShader/Private/NiagaraShader.cpp
@@ -1106,8 +1106,10 @@ void FNiagaraShader::BindParams(const TArray<FNiagaraDataInterfaceGPUParamInfo>&
 {
 	FloatInputBufferParam.Bind(ParameterMap, TEXT("InputFloat"));
 	IntInputBufferParam.Bind(ParameterMap, TEXT("InputInt"));
+	HalfInputBufferParam.Bind(ParameterMap, TEXT("InputHalf"));
 	FloatOutputBufferParam.Bind(ParameterMap, TEXT("OutputFloat"));
 	IntOutputBufferParam.Bind(ParameterMap, TEXT("OutputInt"));
+	HalfOutputBufferParam.Bind(ParameterMap, TEXT("OutputHalf"));
 
 	InstanceCountsParam.Bind(ParameterMap, TEXT("InstanceCounts"));
 	ReadInstanceCountOffsetParam.Bind(ParameterMap, TEXT("ReadInstanceCountOffset"));
@@ -1179,7 +1181,7 @@ void FNiagaraShader::BindParams(const TArray<FNiagaraDataInterfaceGPUParamInfo>&
 		ParamRef.Bind(DIParamInfo, ParameterMap);
 	}
 
-	ensure(FloatOutputBufferParam.IsBound() || IntOutputBufferParam.IsBound());	// we should have at least one output buffer we're writing to
+	ensure(HalfOutputBufferParam.IsBound() || FloatOutputBufferParam.IsBound() || IntOutputBufferParam.IsBound());	// we should have at least one output buffer we're writing to
 	ensure(InstanceCountsParam.IsBound());
 	ensure(UpdateStartInstanceParam.IsBound());
 	ensure(NumSpawnedInstancesParam.IsBound());
diff --git a/Engine/Plugins/FX/Niagara/Source/NiagaraShader/Public/NiagaraShader.h b/Engine/Plugins/FX/Niagara/Source/NiagaraShader/Public/NiagaraShader.h
--- a/Engine/Plugins/FX/Niagara/Source/NiagaraShader/Public/NiagaraShader.h
+++ b/Engine/Plugins/FX/Niagara/Source/NiagaraShader/Public/NiagaraShader.h
@@ -93,8 +93,10 @@ public:
 
 	LAYOUT_FIELD(FShaderResourceParameter, FloatInputBufferParam);
 	LAYOUT_FIELD(FShaderResourceParameter, IntInputBufferParam);
+	LAYOUT_FIELD(FShaderResourceParameter, HalfInputBufferParam);
 	LAYOUT_FIELD(FRWShaderParameter, FloatOutputBufferParam);
 	LAYOUT_FIELD(FRWShaderParameter, IntOutputBufferParam);
+	LAYOUT_FIELD(FRWShaderParameter, HalfOutputBufferParam);
 	LAYOUT_FIELD(FRWShaderParameter, InstanceCountsParam);
 	LAYOUT_FIELD(FShaderParameter, ReadInstanceCountOffsetParam);
 	LAYOUT_FIELD(FShaderParameter, WriteInstanceCountOffsetParam);
diff --git a/Engine/Plugins/FX/Niagara/Source/NiagaraVertexFactories/Private/NiagaraMeshVertexFactory.cpp b/Engine/Plugins/FX/Niagara/Source/NiagaraVertexFactories/Private/NiagaraMeshVertexFactory.cpp
--- a/Engine/Plugins/FX/Niagara/Source/NiagaraVertexFactories/Private/NiagaraMeshVertexFactory.cpp
+++ b/Engine/Plugins/FX/Niagara/Source/NiagaraVertexFactories/Private/NiagaraMeshVertexFactory.cpp
@@ -19,8 +19,6 @@ public:
 	void Bind(const FShaderParameterMap& ParameterMap)
 	{
 		//PrevTransformBuffer.Bind(ParameterMap, TEXT("PrevTransformBuffer"));
-		NiagaraParticleDataFloat.Bind(ParameterMap, TEXT("NiagaraParticleDataFloat"));
-		FloatDataStride.Bind(ParameterMap, TEXT("NiagaraFloatDataStride"));
 
 		// 		NiagaraParticleDataInt.Bind(ParameterMap, TEXT("NiagaraParticleDataInt"));
 		// 		FloatDataOffset.Bind(ParameterMap, TEXT("NiagaraInt32DataOffset"));
@@ -44,9 +42,6 @@ public:
 		FNiagaraMeshVertexFactory* NiagaraMeshVF = (FNiagaraMeshVertexFactory*)VertexFactory;
 		ShaderBindings.Add(Shader->GetUniformBufferParameter<FNiagaraMeshUniformParameters>(), NiagaraMeshVF->GetUniformBuffer());
 
-		ShaderBindings.Add(NiagaraParticleDataFloat, NiagaraMeshVF->GetParticleDataFloatSRV());
-		ShaderBindings.Add(FloatDataStride, NiagaraMeshVF->GetFloatDataStride());
-
 		FRHIShaderResourceView* SortedSRV = NiagaraMeshVF->GetSortedIndicesSRV();
 		ShaderBindings.Add(SortedIndices, SortedSRV != nullptr ? SortedSRV : GFNiagaraNullSortedIndicesVertexBuffer.VertexBufferSRV.GetReference());
 		ShaderBindings.Add(SortedIndicesOffset, NiagaraMeshVF->GetSortedIndicesOffset());
@@ -56,8 +51,6 @@ private:
 
 
 		//LAYOUT_FIELD(FShaderResourceParameter, PrevTransformBuffer)
-		LAYOUT_FIELD(FShaderResourceParameter, NiagaraParticleDataFloat);
-		LAYOUT_FIELD(FShaderParameter, FloatDataStride);
 
 		// 	LAYOUT_FIELD(FShaderResourceParameter, NiagaraParticleDataInt)
 		// 	LAYOUT_FIELD(FShaderParameter, Int32DataOffset)
diff --git a/Engine/Plugins/FX/Niagara/Source/NiagaraVertexFactories/Private/NiagaraRibbonVertexFactory.cpp b/Engine/Plugins/FX/Niagara/Source/NiagaraVertexFactories/Private/NiagaraRibbonVertexFactory.cpp
--- a/Engine/Plugins/FX/Niagara/Source/NiagaraVertexFactories/Private/NiagaraRibbonVertexFactory.cpp
+++ b/Engine/Plugins/FX/Niagara/Source/NiagaraVertexFactories/Private/NiagaraRibbonVertexFactory.cpp
@@ -30,8 +30,6 @@ class FNiagaraRibbonVertexFactoryShaderParametersVS : public FNiagaraRibbonVerte
 public:
 	void Bind(const FShaderParameterMap& ParameterMap)
 	{
-		NiagaraParticleDataFloat.Bind(ParameterMap, TEXT("NiagaraParticleDataFloat"));
-		FloatDataStride.Bind(ParameterMap, TEXT("NiagaraFloatDataStride"));
 	}
 
 	void GetElementShaderBindings(
@@ -48,17 +46,12 @@ public:
 		FNiagaraRibbonVertexFactory* RibbonVF = (FNiagaraRibbonVertexFactory*)VertexFactory;
 		ShaderBindings.Add(Shader->GetUniformBufferParameter<FNiagaraRibbonUniformParameters>(), RibbonVF->GetRibbonUniformBuffer());
 		ShaderBindings.Add(Shader->GetUniformBufferParameter<FNiagaraRibbonVFLooseParameters>(), RibbonVF->LooseParameterUniformBuffer);
-		ShaderBindings.Add(NiagaraParticleDataFloat, RibbonVF->GetParticleDataFloatSRV());
-		ShaderBindings.Add(FloatDataStride, RibbonVF->GetFloatDataStride());
 	}
 
 private:
-	LAYOUT_FIELD(FShaderResourceParameter, NiagaraParticleDataFloat);
-	LAYOUT_FIELD(FShaderParameter, FloatDataStride);
 };
 
 
-
 /**
 * Shader parameters for the beam/trail vertex factory.
 */
diff --git a/Engine/Plugins/FX/Niagara/Source/NiagaraVertexFactories/Private/NiagaraSpriteVertexFactory.cpp b/Engine/Plugins/FX/Niagara/Source/NiagaraVertexFactories/Private/NiagaraSpriteVertexFactory.cpp
--- a/Engine/Plugins/FX/Niagara/Source/NiagaraVertexFactories/Private/NiagaraSpriteVertexFactory.cpp
+++ b/Engine/Plugins/FX/Niagara/Source/NiagaraVertexFactories/Private/NiagaraSpriteVertexFactory.cpp
@@ -36,9 +36,6 @@ public:
 		NumCutoutVerticesPerFrame.Bind(ParameterMap, TEXT("NumCutoutVerticesPerFrame"));
 		CutoutGeometry.Bind(ParameterMap, TEXT("CutoutGeometry"));
 
-		NiagaraParticleDataFloat.Bind(ParameterMap, TEXT("NiagaraParticleDataFloat"));
-		FloatDataStride.Bind(ParameterMap, TEXT("NiagaraFloatDataStride"));
-
 //  		NiagaraParticleDataInt.Bind(ParameterMap, TEXT("NiagaraParticleDataInt"));
 //  		Int32DataOffset.Bind(ParameterMap, TEXT("NiagaraInt32DataOffset"));
 //  		Int32DataStride.Bind(ParameterMap, TEXT("NiagaraInt3DataStride"));
@@ -73,9 +70,6 @@ public:
 		ShaderBindings.Add(ParticleAlignmentMode, SpriteVF->GetAlignmentMode());
 		ShaderBindings.Add(ParticleFacingMode, SpriteVF->GetFacingMode());
 
-		ShaderBindings.Add(NiagaraParticleDataFloat, SpriteVF->GetParticleDataFloatSRV());
-		ShaderBindings.Add(FloatDataStride, SpriteVF->GetFloatDataStride());
-
 		ShaderBindings.Add(SortedIndices, SpriteVF->GetSortedIndicesSRV() ? SpriteVF->GetSortedIndicesSRV() : GFNiagaraNullSortedIndicesVertexBuffer.VertexBufferSRV.GetReference());
 		ShaderBindings.Add(SortedIndicesOffset, SpriteVF->GetSortedIndicesOffset());
 	}
@@ -89,8 +83,6 @@ private:
 
 		LAYOUT_FIELD(FShaderResourceParameter, CutoutGeometry);
 
-	LAYOUT_FIELD(FShaderResourceParameter, NiagaraParticleDataFloat);
-	LAYOUT_FIELD(FShaderParameter, FloatDataStride);
 
 	//  	LAYOUT_FIELD(FShaderResourceParameter, NiagaraParticleDataInt);
 	//  	LAYOUT_FIELD(FShaderParameter, Int32DataOffset);
diff --git a/Engine/Plugins/FX/Niagara/Source/NiagaraVertexFactories/Public/NiagaraMeshVertexFactory.h b/Engine/Plugins/FX/Niagara/Source/NiagaraVertexFactories/Public/NiagaraMeshVertexFactory.h
--- a/Engine/Plugins/FX/Niagara/Source/NiagaraVertexFactories/Public/NiagaraMeshVertexFactory.h
+++ b/Engine/Plugins/FX/Niagara/Source/NiagaraVertexFactories/Public/NiagaraMeshVertexFactory.h
@@ -55,6 +55,9 @@ BEGIN_GLOBAL_SHADER_PARAMETER_STRUCT(FNiagaraMeshUniformParameters, NIAGARAVERTE
 	SHADER_PARAMETER(uint32, bLockedAxisEnable)
 	SHADER_PARAMETER(FVector, LockedAxis)
 	SHADER_PARAMETER(uint32, LockedAxisSpace)
+	SHADER_PARAMETER(uint32, NiagaraFloatDataStride)
+	SHADER_PARAMETER_SRV(Buffer<float>, NiagaraParticleDataFloat)
+	SHADER_PARAMETER_SRV(Buffer<float>, NiagaraParticleDataHalf)
 END_GLOBAL_SHADER_PARAMETER_STRUCT()
 
 typedef TUniformBufferRef<FNiagaraMeshUniformParameters> FNiagaraMeshUniformBufferRef;
@@ -104,12 +107,7 @@ public:
 		// Set a define so we can tell in MaterialTemplate.usf when we are compiling a mesh particle vertex factory
 		OutEnvironment.SetDefine(TEXT("NIAGARA_MESH_FACTORY"), TEXT("1"));
 		OutEnvironment.SetDefine(TEXT("NIAGARA_MESH_INSTANCED"), TEXT("1"));
-	}
-
-	void SetParticleData(const FShaderResourceViewRHIRef& InParticleDataFloatSRV, uint32 InFloatDataStride)
-	{
-		ParticleDataFloatSRV = InParticleDataFloatSRV;
-		FloatDataStride = InFloatDataStride;
+		OutEnvironment.SetDefine(TEXT("NiagaraVFLooseParameters"), TEXT("NiagaraMeshVF"));
 	}
 
 	void SetSortedIndices(const FShaderResourceViewRHIRef& InSortedIndicesSRV, uint32 InSortedIndicesOffset)
@@ -118,16 +116,6 @@ public:
 		SortedIndicesOffset = InSortedIndicesOffset;
 	}
 
-	FORCEINLINE FRHIShaderResourceView* GetParticleDataFloatSRV()
-	{
-		return ParticleDataFloatSRV;
-	}
-
-	FORCEINLINE int32 GetFloatDataStride()
-	{
-		return FloatDataStride;
-	}
-
 	FORCEINLINE FRHIShaderResourceView* GetSortedIndicesSRV()
 	{
 		return SortedIndicesSRV;
@@ -190,6 +178,9 @@ protected:
 	FShaderResourceViewRHIRef ParticleDataFloatSRV;
 	uint32 FloatDataStride;
 
+	FShaderResourceViewRHIRef ParticleDataHalfSRV;
+	uint32 HalfDataStride;
+
 	FShaderResourceViewRHIRef SortedIndicesSRV;
 	uint32 SortedIndicesOffset;
 };
diff --git a/Engine/Plugins/FX/Niagara/Source/NiagaraVertexFactories/Public/NiagaraRibbonVertexFactory.h b/Engine/Plugins/FX/Niagara/Source/NiagaraVertexFactories/Public/NiagaraRibbonVertexFactory.h
--- a/Engine/Plugins/FX/Niagara/Source/NiagaraVertexFactories/Public/NiagaraRibbonVertexFactory.h
+++ b/Engine/Plugins/FX/Niagara/Source/NiagaraVertexFactories/Public/NiagaraRibbonVertexFactory.h
@@ -62,6 +62,7 @@ BEGIN_GLOBAL_SHADER_PARAMETER_STRUCT(FNiagaraRibbonVFLooseParameters, NIAGARAVER
 	SHADER_PARAMETER_SRV(Buffer<uint>, MultiRibbonIndices)
 	SHADER_PARAMETER_SRV(Buffer<float>, PackedPerRibbonDataByIndex)
 	SHADER_PARAMETER_SRV(Buffer<float>, NiagaraParticleDataFloat)
+	SHADER_PARAMETER_SRV(Buffer<float>, NiagaraParticleDataHalf)
 END_GLOBAL_SHADER_PARAMETER_STRUCT()
 typedef TUniformBufferRef<FNiagaraRibbonVFLooseParameters> FNiagaraRibbonVFLooseParametersRef;
 
@@ -130,13 +131,6 @@ public:
 	*/
 	void SetDynamicParameterBuffer(const FVertexBuffer* InDynamicParameterBuffer, int32 ParameterIndex, uint32 StreamOffset, uint32 Stride);
 
-
-	void SetParticleData(const FShaderResourceViewRHIRef& InParticleDataFloatSRV, uint32 InFloatDataStride)
-	{
-		ParticleDataFloatSRV = InParticleDataFloatSRV;
-		FloatDataStride = InFloatDataStride;
-	}
-
 	void SetSortedIndices(const FVertexBufferRHIRef& InSortedIndicesBuffer, const FShaderResourceViewRHIRef& InSortedIndicesSRV, uint32 InSortedIndicesOffset)
 	{
 		SortedIndicesBuffer = InSortedIndicesBuffer;
@@ -168,16 +162,6 @@ public:
 		FacingMode = InFacingMode;
 	}
 
-	FORCEINLINE FRHIShaderResourceView* GetParticleDataFloatSRV()
-	{
-		return ParticleDataFloatSRV;
-	}
-
-	FORCEINLINE int32 GetFloatDataStride()
-	{
-		return FloatDataStride;
-	}
-
 	FORCEINLINE FRHIShaderResourceView* GetSortedIndicesSRV()
 	{
 		return SortedIndicesSRV;
@@ -237,9 +221,6 @@ private:
 
 	const FNiagaraDataSet *DataSet;
 
-	FShaderResourceViewRHIRef ParticleDataFloatSRV;
-	uint32 FloatDataStride;
-
 	FVertexBufferRHIRef SortedIndicesBuffer;
 	FVertexBufferRHIRef TangentAndDistancesBuffer;
 	FVertexBufferRHIRef MultiRibbonIndicesBuffer;
diff --git a/Engine/Plugins/FX/Niagara/Source/NiagaraVertexFactories/Public/NiagaraSortingGPU.h b/Engine/Plugins/FX/Niagara/Source/NiagaraVertexFactories/Public/NiagaraSortingGPU.h
--- a/Engine/Plugins/FX/Niagara/Source/NiagaraVertexFactories/Public/NiagaraSortingGPU.h
+++ b/Engine/Plugins/FX/Niagara/Source/NiagaraVertexFactories/Public/NiagaraSortingGPU.h
@@ -64,6 +64,7 @@ private:
 
 	LAYOUT_FIELD(FShaderResourceParameter, NiagaraParticleDataFloat)
 	LAYOUT_FIELD(FShaderParameter, FloatDataStride)
+	LAYOUT_FIELD(FShaderResourceParameter, NiagaraParticleDataHalf)
 	LAYOUT_FIELD(FShaderResourceParameter, GPUParticleCountBuffer)
 	LAYOUT_FIELD(FShaderParameter, ParticleCountParams)
 	LAYOUT_FIELD(FShaderParameter, SortParams)
diff --git a/Engine/Plugins/FX/Niagara/Source/NiagaraVertexFactories/Public/NiagaraSpriteVertexFactory.h b/Engine/Plugins/FX/Niagara/Source/NiagaraVertexFactories/Public/NiagaraSpriteVertexFactory.h
--- a/Engine/Plugins/FX/Niagara/Source/NiagaraVertexFactories/Public/NiagaraSpriteVertexFactory.h
+++ b/Engine/Plugins/FX/Niagara/Source/NiagaraVertexFactories/Public/NiagaraSpriteVertexFactory.h
@@ -64,6 +64,7 @@ BEGIN_GLOBAL_SHADER_PARAMETER_STRUCT(FNiagaraSpriteVFLooseParameters, NIAGARAVER
 	SHADER_PARAMETER(uint32, IndirectArgsOffset)
 	SHADER_PARAMETER_SRV(Buffer<float2>, CutoutGeometry)
 	SHADER_PARAMETER_SRV(Buffer<float>, NiagaraParticleDataFloat)
+	SHADER_PARAMETER_SRV(Buffer<float>, NiagaraParticleDataHalf)
 	SHADER_PARAMETER_SRV(Buffer<int>, SortedIndices)
 	SHADER_PARAMETER_SRV(Buffer<uint>, IndirectArgsBuffer)
 END_GLOBAL_SHADER_PARAMETER_STRUCT()
@@ -87,7 +88,6 @@ public:
 		CutoutGeometrySRV(nullptr),
 		AlignmentMode(0),
 		FacingMode(0),
-		FloatDataStride(0),
 		SortedIndicesOffset(0)
 	{}
 
@@ -98,7 +98,6 @@ public:
 		CutoutGeometrySRV(nullptr),
 		AlignmentMode(0),
 		FacingMode(0),
-		FloatDataStride(0),
 		SortedIndicesOffset(0)
 		
 	{}
@@ -145,28 +144,12 @@ public:
 	inline int32 GetNumCutoutVerticesPerFrame() const { return NumCutoutVerticesPerFrame; }
 	inline FRHIShaderResourceView* GetCutoutGeometrySRV() const { return CutoutGeometrySRV; }
 
-	void SetParticleData(const FShaderResourceViewRHIRef& InParticleDataFloatSRV, uint32 InFloatDataStride)
-	{
-		ParticleDataFloatSRV = InParticleDataFloatSRV;
-		FloatDataStride = InFloatDataStride;
-	}
-
 	void SetSortedIndices(const FShaderResourceViewRHIRef& InSortedIndicesSRV, uint32 InSortedIndicesOffset)
 	{
 		SortedIndicesSRV = InSortedIndicesSRV;
 		SortedIndicesOffset = InSortedIndicesOffset;
 	}
 
-	FORCEINLINE FRHIShaderResourceView* GetParticleDataFloatSRV()
-	{
-		return ParticleDataFloatSRV;
-	}
-
-	FORCEINLINE int32 GetFloatDataStride()
-	{
-		return FloatDataStride;
-	}
-
 	FORCEINLINE FRHIShaderResourceView* GetSortedIndicesSRV()
 	{
 		return SortedIndicesSRV;
@@ -219,10 +202,6 @@ private:
 	FShaderResourceViewRHIRef CutoutGeometrySRV;
 	uint32 AlignmentMode;
 	uint32 FacingMode;
-	
-	
-	FShaderResourceViewRHIRef ParticleDataFloatSRV;
-	uint32 FloatDataStride;
 
 	FShaderResourceViewRHIRef SortedIndicesSRV;
 	uint32 SortedIndicesOffset;
diff --git a/Engine/Source/Developer/ShaderFormatVectorVM/Private/ir_vm_gen_bytecode_visitor.cpp b/Engine/Source/Developer/ShaderFormatVectorVM/Private/ir_vm_gen_bytecode_visitor.cpp
--- a/Engine/Source/Developer/ShaderFormatVectorVM/Private/ir_vm_gen_bytecode_visitor.cpp
+++ b/Engine/Source/Developer/ShaderFormatVectorVM/Private/ir_vm_gen_bytecode_visitor.cpp
@@ -362,17 +362,53 @@ EVectorVMOp get_special_vm_opcode(ir_function_signature* signature)
 	{
 		return EVectorVMOp::noise;
 	}
-	else if (strncmp(signature->function_name(), "InputDataNoadvance", strlen("InputDataNoadvance")) == 0)
+	else if (strncmp(signature->function_name(), "InputDataNoadvanceFloat", strlen("InputDataNoadvanceFloat")) == 0)
 	{
-		return EVectorVMOp::inputdata_noadvance_32bit;
+		return EVectorVMOp::inputdata_noadvance_float;
 	}
-	else if (strncmp(signature->function_name(), "InputData", strlen("InputData")) == 0)
+	else if (strncmp(signature->function_name(), "InputDataNoadvanceInt", strlen("InputDataNoadvanceInt")) == 0)
 	{
-		return EVectorVMOp::inputdata_32bit;
+		return EVectorVMOp::inputdata_noadvance_int32;
 	}
-	else if (strncmp(signature->function_name(), "OutputData", strlen("OutputData")) == 0)
+	else if (strncmp(signature->function_name(), "InputDataNoadvanceBool", strlen("InputDataNoadvanceBool")) == 0)
 	{
-		return EVectorVMOp::outputdata_32bit;
+		return EVectorVMOp::inputdata_noadvance_int32;
+	}
+	else if (strncmp(signature->function_name(), "InputDataNoadvanceHalf", strlen("InputDataNoadvanceHalf")) == 0)
+	{
+		return EVectorVMOp::inputdata_noadvance_half;
+	}
+	else if (strncmp(signature->function_name(), "InputDataHalf", strlen("InputDataHalf")) == 0)
+	{
+		return EVectorVMOp::inputdata_half;
+	}
+	else if (strncmp(signature->function_name(), "OutputDataHalf", strlen("OutputDataHalf")) == 0)
+	{
+		return EVectorVMOp::outputdata_half;
+	}
+	else if (strncmp(signature->function_name(), "InputDataInt", strlen("InputDataInt")) == 0)
+	{
+		return EVectorVMOp::inputdata_int32;
+	}
+	else if (strncmp(signature->function_name(), "OutputDataInt", strlen("OutputDataInt")) == 0)
+	{
+		return EVectorVMOp::outputdata_int32;
+	}
+	else if (strncmp(signature->function_name(), "InputDataBool", strlen("InputDataBool")) == 0)
+	{
+		return EVectorVMOp::inputdata_int32;
+	}
+	else if (strncmp(signature->function_name(), "OutputDataBool", strlen("OutputDataBool")) == 0)
+	{
+		return EVectorVMOp::outputdata_int32;
+	}
+	else if (strncmp(signature->function_name(), "InputDataFloat", strlen("InputDataFloat")) == 0)
+	{
+		return EVectorVMOp::inputdata_float;
+	}
+	else if (strncmp(signature->function_name(), "OutputDataFloat", strlen("OutputDataFloat")) == 0)
+	{
+		return EVectorVMOp::outputdata_float;
 	}
 	else if (strcmp(signature->function_name(), "AcquireIndex") == 0)
 	{
@@ -862,6 +898,72 @@ struct op_input : public op_base
 	}
 };
 
+struct op_input_half : public op_input
+{
+	op_input_half(): op_input(){}
+
+	virtual FString to_string()
+	{
+		if (instance_idx_component)
+		{
+			return FString::Printf(TEXT("[R][%d] = InputDataHalf(%d, %d, [%s][%d]);\n"),
+				dest_component->offset,
+				dataset_idx,
+				register_idx,
+				get_location_string(instance_idx_component->owner->location),
+				instance_idx_component->offset);
+		}
+		else
+		{
+			return FString::Printf(TEXT("[R][%d] = InputDataHalf(%d, %d);\n"), dest_component->offset, dataset_idx, register_idx);
+		}
+	}
+};
+
+
+struct op_input_float : public op_input
+{
+	op_input_float() : op_input() {}
+
+	virtual FString to_string()
+	{
+		if (instance_idx_component)
+		{
+			return FString::Printf(TEXT("[R][%d] = InputDataFloat(%d, %d, [%s][%d]);\n"),
+				dest_component->offset,
+				dataset_idx,
+				register_idx,
+				get_location_string(instance_idx_component->owner->location),
+				instance_idx_component->offset);
+		}
+		else
+		{
+			return FString::Printf(TEXT("[R][%d] = InputDataFloat(%d, %d);\n"), dest_component->offset, dataset_idx, register_idx);
+		}
+	}
+};
+
+struct op_input_int : public op_input
+{
+	op_input_int() : op_input() {}
+
+	virtual FString to_string()
+	{
+		if (instance_idx_component)
+		{
+			return FString::Printf(TEXT("[R][%d] = InputDataInt(%d, %d, [%s][%d]);\n"),
+				dest_component->offset,
+				dataset_idx,
+				register_idx,
+				get_location_string(instance_idx_component->owner->location),
+				instance_idx_component->offset);
+		}
+		else
+		{
+			return FString::Printf(TEXT("[R][%d] = InputDataInt(%d, %d);\n"), dest_component->offset, dataset_idx, register_idx);
+		}
+	}
+};
 
 // TODO: make this a proper explicitly indexing op
 struct op_input_noadvance : public op_input
@@ -910,6 +1012,68 @@ struct op_input_noadvance : public op_input
 	}
 };
 
+struct op_input_noadvance_int : public op_input_noadvance
+{
+	op_input_noadvance_int() : op_input_noadvance() {}
+	virtual FString to_string()
+	{
+		if (instance_idx_component)
+		{
+			return FString::Printf(TEXT("[R][%d] = InputDataNoadvanceInt(%d, %d, [%s][%d]);\n"),
+				dest_component->offset,
+				dataset_idx,
+				register_idx,
+				get_location_string(instance_idx_component->owner->location),
+				instance_idx_component->offset);
+		}
+		else
+		{
+			return FString::Printf(TEXT("[R][%d] = InputDataNoadvanceInt(%d, %d);\n"), dest_component->offset, dataset_idx, register_idx);
+		}
+	}
+};
+
+struct op_input_noadvance_float : public op_input_noadvance
+{
+	op_input_noadvance_float() : op_input_noadvance() {}
+	virtual FString to_string()
+	{
+		if (instance_idx_component)
+		{
+			return FString::Printf(TEXT("[R][%d] = InputDataNoadvanceFloat(%d, %d, [%s][%d]);\n"),
+				dest_component->offset,
+				dataset_idx,
+				register_idx,
+				get_location_string(instance_idx_component->owner->location),
+				instance_idx_component->offset);
+		}
+		else
+		{
+			return FString::Printf(TEXT("[R][%d] = InputDataNoadvanceFloat(%d, %d);\n"), dest_component->offset, dataset_idx, register_idx);
+		}
+	}
+};
+
+struct op_input_noadvance_half : public op_input_noadvance
+{
+	op_input_noadvance_half() : op_input_noadvance() {}
+	virtual FString to_string()
+	{
+		if (instance_idx_component)
+		{
+			return FString::Printf(TEXT("[R][%d] = InputDataNoadvanceHalf(%d, %d, [%s][%d]);\n"),
+				dest_component->offset,
+				dataset_idx,
+				register_idx,
+				get_location_string(instance_idx_component->owner->location),
+				instance_idx_component->offset);
+		}
+		else
+		{
+			return FString::Printf(TEXT("[R][%d] = InputDataNoadvanceHalf(%d, %d);\n"), dest_component->offset, dataset_idx, register_idx);
+		}
+	}
+};
 
 struct op_output : public op_base
 {
@@ -957,6 +1121,63 @@ struct op_output : public op_base
 	}
 };
 
+struct op_output_half : public op_output
+{
+	op_output_half(): op_output(){}
+
+	virtual FString to_string()
+	{
+		check(instance_idx_component);
+		check(value_component);
+
+		return FString::Printf(TEXT("OutputDataHalf(%d, %d, [%s][%d], [%s][%d]);\n"),
+			dataset_index,
+			register_idx,
+			get_location_string(instance_idx_component->owner->location),
+			instance_idx_component->offset,
+			get_location_string(value_component->owner->location),
+			value_component->offset);
+	}
+};
+
+struct op_output_float : public op_output
+{
+	op_output_float() : op_output() {}
+
+	virtual FString to_string()
+	{
+		check(instance_idx_component);
+		check(value_component);
+
+		return FString::Printf(TEXT("OutputDataFloat(%d, %d, [%s][%d], [%s][%d]);\n"),
+			dataset_index,
+			register_idx,
+			get_location_string(instance_idx_component->owner->location),
+			instance_idx_component->offset,
+			get_location_string(value_component->owner->location),
+			value_component->offset);
+	}
+};
+
+struct op_output_int : public op_output
+{
+	op_output_int() : op_output() {}
+
+	virtual FString to_string()
+	{
+		check(instance_idx_component);
+		check(value_component);
+
+		return FString::Printf(TEXT("OutputDataInt(%d, %d, [%s][%d], [%s][%d]);\n"),
+			dataset_index,
+			register_idx,
+			get_location_string(instance_idx_component->owner->location),
+			instance_idx_component->offset,
+			get_location_string(value_component->owner->location),
+			value_component->offset);
+	}
+};
+
 struct op_index_acquire : public op_base
 {
 	op_index_acquire()
@@ -1758,9 +1979,15 @@ class ir_gen_vvm_visitor : public ir_hierarchical_visitor
 
 		switch (VVMOpCode)
 		{
-		case EVectorVMOp::inputdata_32bit: allocate<op_input>(); break;
-		case EVectorVMOp::inputdata_noadvance_32bit: allocate<op_input_noadvance>(); break;
-		case EVectorVMOp::outputdata_32bit: allocate<op_output>(); break;
+		case EVectorVMOp::inputdata_half: allocate<op_input_half>(); break;
+		case EVectorVMOp::inputdata_int32: allocate<op_input_int>(); break;
+		case EVectorVMOp::inputdata_float: allocate<op_input_float>(); break;
+		case EVectorVMOp::inputdata_noadvance_float: allocate<op_input_noadvance_float>(); break;
+		case EVectorVMOp::inputdata_noadvance_int32: allocate<op_input_noadvance_int>(); break;
+		case EVectorVMOp::inputdata_noadvance_half: allocate<op_input_noadvance_half>(); break;
+		case EVectorVMOp::outputdata_half: allocate<op_output_half>(); break;
+		case EVectorVMOp::outputdata_int32: allocate<op_output_int>(); break;
+		case EVectorVMOp::outputdata_float: allocate<op_output_float>(); break;
 		case EVectorVMOp::acquireindex: allocate<op_index_acquire>(); break;
 		case EVectorVMOp::acquire_id: allocate<op_id_acquire>(); break;
 		case EVectorVMOp::update_id: allocate<op_id_update>(); break;
diff --git a/Engine/Source/Runtime/Core/Public/GenericPlatform/GenericPlatformMath.h b/Engine/Source/Runtime/Core/Public/GenericPlatform/GenericPlatformMath.h
--- a/Engine/Source/Runtime/Core/Public/GenericPlatform/GenericPlatformMath.h
+++ b/Engine/Source/Runtime/Core/Public/GenericPlatform/GenericPlatformMath.h
@@ -16,6 +16,122 @@
  */
 struct FGenericPlatformMath
 {
+	//https://gist.github.com/rygorous/2156668
+	static FORCEINLINE float LoadHalf(const uint16* Ptr)
+	{
+		uint16 FP16 = *Ptr;
+		uint32 shifted_exp = 0x7c00 << 13;			// exponent mask after shift
+		union FP32T
+		{
+			uint32 u;
+			float f;		
+		} FP32, magic = { 113 << 23 };
+
+		FP32.u = (FP16 & 0x7fff) << 13;				// exponent/mantissa bits
+		uint32 exp = shifted_exp & FP32.u;			// just the exponent
+		FP32.u += uint32(127 - 15) << 23;			// exponent adjust
+
+		// handle exponent special cases
+		if (exp == shifted_exp)						// Inf/NaN?
+		{
+			FP32.u += uint32(128 - 16) << 23;		// extra exp adjust
+		}
+		else if (exp == 0)							// Zero/Denormal?
+		{
+			FP32.u += 1 << 23;						// extra exp adjust
+			FP32.f -= magic.f;						// renormalize
+		}
+
+		FP32.u |= (FP16 & 0x8000) << 16;			// sign bit
+		return FP32.f;
+	}
+
+	//https://gist.github.com/rygorous/2156668
+	static FORCEINLINE void StoreHalf(uint16* Ptr, float Value)
+	{
+		union FP32T
+		{
+			uint32 u;
+			float f;
+		} FP32 = {};
+		uint16 FP16 = {};
+
+		FP32.f = Value;
+
+		FP32T f32infty = { uint32(255 << 23) };
+		FP32T f16max = { uint32(127 + 16) << 23 };
+		FP32T denorm_magic = { (uint32(127 - 15) + uint32(23 - 10) + 1) << 23 };
+		uint32 sign_mask = 0x80000000u;
+
+		uint32 sign = FP32.u & sign_mask;
+		FP32.u ^= sign;
+
+		// NOTE all the integer compares in this function can be safely
+		// compiled into signed compares since all operands are below
+		// 0x80000000. Important if you want fast straight SSE2 code
+		// (since there's no unsigned PCMPGTD).
+
+		if (FP32.u >= f16max.u) // result is Inf or NaN (all exponent bits set)
+		{
+			FP16 = (FP32.u > f32infty.u) ? 0x7e00 : 0x7c00; // NaN->qNaN and Inf->Inf
+		}
+		else // (De)normalized number or zero
+		{
+			if (FP32.u < uint32(113 << 23)) // resulting FP16 is subnormal or zero
+			{
+				// use a magic value to align our 10 mantissa bits at the bottom of
+				// the float. as long as FP addition is round-to-nearest-even this
+				// just works.
+				FP32.f += denorm_magic.f;
+
+				// and one integer subtract of the bias later, we have our final float!
+				FP16 = uint16(FP32.u - denorm_magic.u);
+			}
+			else
+			{
+				uint32 mant_odd = (FP32.u >> 13) & 1; // resulting mantissa is odd
+
+				// update exponent, rounding bias part 1
+				FP32.u += (uint32(15 - 127) << 23) + 0xfff;
+				// rounding bias part 2
+				FP32.u += mant_odd;
+				// take the bits!
+				FP16 = uint16(FP32.u >> 13);
+			}
+		}
+
+		FP16 |= sign >> 16;
+		*Ptr = FP16;
+	}
+
+	static FORCEINLINE void VectorLoadHalf(float* RESTRICT Dst, const uint16* RESTRICT Src)
+	{
+		Dst[0] = LoadHalf(&Src[0]);
+		Dst[1] = LoadHalf(&Src[1]);
+		Dst[2] = LoadHalf(&Src[2]);
+		Dst[3] = LoadHalf(&Src[3]);
+	}
+
+	static FORCEINLINE void VectorStoreHalf(uint16* RESTRICT Dst, const float* RESTRICT Src)
+	{
+		StoreHalf(&Dst[0], Src[0]);
+		StoreHalf(&Dst[1], Src[1]);
+		StoreHalf(&Dst[2], Src[2]);
+		StoreHalf(&Dst[3], Src[3]);
+	}
+
+	static FORCEINLINE void WideVectorLoadHalf(float* RESTRICT Dst, const uint16* RESTRICT Src)
+	{
+		VectorLoadHalf(Dst, Src);
+		VectorLoadHalf(Dst + 4, Src + 4);
+	}
+
+	static FORCEINLINE void WideVectorStoreHalf(uint16* RESTRICT Dst, const float* RESTRICT Src)
+	{
+		VectorStoreHalf(Dst, Src);
+		VectorStoreHalf(Dst + 4, Src + 4);
+	}
+
 	/**
 	 * Converts a float to an integer with truncation towards zero.
 	 * @param F		Floating point value to convert
diff --git a/Engine/Source/Runtime/Core/Public/Windows/WindowsPlatformMath.h b/Engine/Source/Runtime/Core/Public/Windows/WindowsPlatformMath.h
--- a/Engine/Source/Runtime/Core/Public/Windows/WindowsPlatformMath.h
+++ b/Engine/Source/Runtime/Core/Public/Windows/WindowsPlatformMath.h
@@ -18,6 +18,40 @@
 **/
 struct FWindowsPlatformMath : public TUnrealPlatformMathSSE4Base<FGenericPlatformMath>
 {
+#if PLATFORM_XBOXONE
+	static FORCEINLINE void WideVectorLoadHalf(float* RESTRICT Dst, const uint16* RESTRICT Src)
+	{
+		_mm256_storeu_ps(Dst, _mm256_cvtph_ps(_mm_loadu_si128((__m128i*)Src)));
+	}
+
+	static FORCEINLINE void WideVectorStoreHalf(uint16* RESTRICT Dst, const float* RESTRICT Src)
+	{
+		_mm_storeu_si128((__m128i*)Dst, _mm256_cvtps_ph(_mm256_loadu_ps(Src), _MM_FROUND_TO_NEAREST_INT));
+	}
+
+	static FORCEINLINE void VectorStoreHalf(uint16* RESTRICT Dst, const float* RESTRICT Src)
+	{
+		_mm_storeu_si64((__m128i*)Dst, _mm_cvtps_ph(_mm_loadu_ps(Src), _MM_FROUND_TO_NEAREST_INT));
+	}
+
+	static FORCEINLINE void VectorLoadHalf(float* RESTRICT Dst, const uint16* RESTRICT Src)
+	{
+		_mm_storeu_ps(Dst, _mm_cvtph_ps(_mm_loadu_si64((__m128i*)Src)));
+	}
+
+	static FORCEINLINE float LoadHalf(const uint16* Ptr)
+	{
+		float Value;
+		_mm_store_ss(&Value, _mm_cvtph_ps(_mm_loadu_si16(Ptr)));
+		return Value;
+	}
+
+	static FORCEINLINE void StoreHalf(uint16* Ptr, float Value)
+	{
+		_mm_storeu_si16(Ptr, _mm_cvtps_ph(_mm_load_ss(&Value), _MM_FROUND_TO_NEAREST_INT));
+	}
+#endif
+
 #if PLATFORM_ENABLE_VECTORINTRINSICS
 	static FORCEINLINE bool IsNaN( float A ) { return _isnan(A) != 0; }
 	static FORCEINLINE bool IsNaN(double A) { return _isnan(A) != 0; }
diff --git a/Engine/Source/Runtime/RenderCore/Private/DynamicBufferAllocator.cpp b/Engine/Source/Runtime/RenderCore/Private/DynamicBufferAllocator.cpp
--- a/Engine/Source/Runtime/RenderCore/Private/DynamicBufferAllocator.cpp
+++ b/Engine/Source/Runtime/RenderCore/Private/DynamicBufferAllocator.cpp
@@ -63,6 +63,7 @@ FGlobalDynamicReadBuffer::FGlobalDynamicReadBuffer()
 {
 	FloatBufferPool = new FDynamicReadBufferPool();
 	Int32BufferPool = new FDynamicReadBufferPool();
+	HalfBufferPool = new FDynamicReadBufferPool();
 }
 
 FGlobalDynamicReadBuffer::~FGlobalDynamicReadBuffer()
@@ -83,6 +84,12 @@ void FGlobalDynamicReadBuffer::Cleanup()
 		delete Int32BufferPool;
 		Int32BufferPool = nullptr;
 	}
+
+	if (HalfBufferPool)
+	{
+		delete HalfBufferPool;
+		HalfBufferPool = nullptr;
+	}
 }
 void FGlobalDynamicReadBuffer::InitRHI()
 {
@@ -156,6 +163,67 @@ FGlobalDynamicReadBuffer::FAllocation FGlobalDynamicReadBuffer::AllocateFloat(ui
 	return Allocation;
 }
 
+FGlobalDynamicReadBuffer::FAllocation FGlobalDynamicReadBuffer::AllocateHalf(uint32 Num)
+{
+	FScopeLock ScopeLock(&HalfBufferPool->CriticalSection);
+	FAllocation Allocation;
+
+	// The codepath using FShaderResourceViewInitializer, requires the SRV to be aligned on 16 bytes on some platforms.
+	Num = Align(Num, RHIGetMinimumAlignmentForBufferBackedSRV(PF_R16F) / sizeof(FFloat16));
+
+	TotalAllocatedSinceLastCommit += Num;
+	if (IsRenderAlarmLoggingEnabled())
+	{
+		UE_LOG(LogRendererCore, Warning, TEXT("FGlobalReadBuffer::AllocateHalf(%u), will have allocated %u total this frame"), Num, TotalAllocatedSinceLastCommit);
+	}
+	uint32 SizeInBytes = sizeof(FFloat16) * Num;
+	FDynamicAllocReadBuffer* Buffer = HalfBufferPool->CurrentBuffer;
+	if (Buffer == NULL || Buffer->AllocatedByteCount + SizeInBytes > Buffer->NumBytes)
+	{
+		// Find a buffer in the pool big enough to service the request.
+		Buffer = NULL;
+		for (int32 BufferIndex = 0, NumBuffers = HalfBufferPool->Buffers.Num(); BufferIndex < NumBuffers; ++BufferIndex)
+		{
+			FDynamicAllocReadBuffer& BufferToCheck = HalfBufferPool->Buffers[BufferIndex];
+			if (BufferToCheck.AllocatedByteCount + SizeInBytes <= BufferToCheck.NumBytes)
+			{
+				Buffer = &BufferToCheck;
+				break;
+			}
+		}
+
+		// Create a new vertex buffer if needed.
+		if (Buffer == NULL)
+		{
+			const uint32 AlignedNum = FMath::DivideAndRoundUp(Num, (uint32)GAlignReadBufferRenderingBufferSize) * GAlignReadBufferRenderingBufferSize;
+			const uint32 NewBufferSize = FMath::Max(AlignedNum, (uint32)GMinReadBufferRenderingBufferSize);
+			Buffer = new FDynamicAllocReadBuffer();
+			HalfBufferPool->Buffers.Add(Buffer);
+			Buffer->Initialize(sizeof(FFloat16), NewBufferSize, PF_R16F, BUF_Volatile);
+		}
+
+		// Lock the buffer if needed.
+		if (Buffer->MappedBuffer == NULL)
+		{
+			Buffer->Lock();
+		}
+
+		// Remember this buffer, we'll try to allocate out of it in the future.
+		HalfBufferPool->CurrentBuffer = Buffer;
+	}
+
+	check(Buffer != NULL);
+	checkf(Buffer->AllocatedByteCount + SizeInBytes <= Buffer->NumBytes, TEXT("Global dynamic read buffer half buffer allocation failed: BufferSize=%d AllocatedByteCount=%d SizeInBytes=%d"), Buffer->NumBytes, Buffer->AllocatedByteCount, SizeInBytes);
+	Allocation.Buffer = Buffer->MappedBuffer + Buffer->AllocatedByteCount;
+	Allocation.ReadBuffer = Buffer;
+	Buffer->SubAllocations.Emplace(RHICreateShaderResourceView(FShaderResourceViewInitializer(Buffer->Buffer, PF_R16F, Buffer->AllocatedByteCount / sizeof(FFloat16), Num)));
+	Allocation.SRV = Buffer->SubAllocations.Last();
+	Buffer->AllocatedByteCount += SizeInBytes;
+
+	return Allocation;
+}
+
+
 FGlobalDynamicReadBuffer::FAllocation FGlobalDynamicReadBuffer::AllocateInt32(uint32 Num)
 {
 	FScopeLock ScopeLock(&Int32BufferPool->CriticalSection);
@@ -229,7 +297,7 @@ void FGlobalDynamicReadBuffer::Commit()
 	for (int32 BufferIndex = 0, NumBuffers = FloatBufferPool->Buffers.Num(); BufferIndex < NumBuffers; ++BufferIndex)
 	{
 		FDynamicAllocReadBuffer& Buffer = FloatBufferPool->Buffers[BufferIndex];
-		if (Buffer.MappedBuffer != NULL)
+		if (Buffer.MappedBuffer != nullptr)
 		{
 			Buffer.Unlock();
 		}
@@ -246,12 +314,12 @@ void FGlobalDynamicReadBuffer::Commit()
 			}
 		}
 	}
-	FloatBufferPool->CurrentBuffer = NULL;
+	FloatBufferPool->CurrentBuffer = nullptr;
 
 	for (int32 BufferIndex = 0, NumBuffers = Int32BufferPool->Buffers.Num(); BufferIndex < NumBuffers; ++BufferIndex)
 	{
 		FDynamicAllocReadBuffer& Buffer = Int32BufferPool->Buffers[BufferIndex];
-		if (Buffer.MappedBuffer != NULL)
+		if (Buffer.MappedBuffer != nullptr)
 		{
 			Buffer.Unlock();
 		}
@@ -268,6 +336,30 @@ void FGlobalDynamicReadBuffer::Commit()
 			}
 		}
 	}
-	Int32BufferPool->CurrentBuffer = NULL;
+	Int32BufferPool->CurrentBuffer = nullptr;
+
+	for (int32 BufferIndex = 0, NumBuffers = HalfBufferPool->Buffers.Num(); BufferIndex < NumBuffers; ++BufferIndex)
+	{
+		FDynamicAllocReadBuffer& Buffer = HalfBufferPool->Buffers[BufferIndex];
+		if (Buffer.MappedBuffer != nullptr)
+		{
+			Buffer.Unlock();
+		}
+		else if (GGlobalBufferNumFramesUnusedThresold && !Buffer.AllocatedByteCount)
+		{
+			++Buffer.NumFramesUnused;
+			if (Buffer.NumFramesUnused >= GGlobalBufferNumFramesUnusedThresold)
+			{
+				// Remove the buffer, assumes they are unordered.
+				Buffer.Release();
+				HalfBufferPool->Buffers.RemoveAtSwap(BufferIndex);
+				--BufferIndex;
+				--NumBuffers;
+			}
+		}
+	}
+	HalfBufferPool->CurrentBuffer = nullptr;
+
+
 	TotalAllocatedSinceLastCommit = 0;
 }
diff --git a/Engine/Source/Runtime/RenderCore/Public/DynamicBufferAllocator.h b/Engine/Source/Runtime/RenderCore/Public/DynamicBufferAllocator.h
--- a/Engine/Source/Runtime/RenderCore/Public/DynamicBufferAllocator.h
+++ b/Engine/Source/Runtime/RenderCore/Public/DynamicBufferAllocator.h
@@ -73,6 +73,7 @@ public:
 	FGlobalDynamicReadBuffer();
 	~FGlobalDynamicReadBuffer();
 	
+	FAllocation AllocateHalf(uint32 Num);
 	FAllocation AllocateFloat(uint32 Num);
 	FAllocation AllocateInt32(uint32 Num);
 
@@ -93,6 +94,7 @@ protected:
 	void Cleanup();
 
 	/** The pools of read buffers from which allocations are made. */
+	FDynamicReadBufferPool* HalfBufferPool;
 	FDynamicReadBufferPool* FloatBufferPool;
 	FDynamicReadBufferPool* Int32BufferPool;
 
