Removing redundant entries in the parameters list

diff --git a/Engine/Plugins/FX/Niagara/Source/NiagaraEditor/Private/Widgets/SNiagaraParameterPanel.cpp b/Engine/Plugins/FX/Niagara/Source/NiagaraEditor/Private/Widgets/SNiagaraParameterPanel.cpp
--- a/Engine/Plugins/FX/Niagara/Source/NiagaraEditor/Private/Widgets/SNiagaraParameterPanel.cpp
+++ b/Engine/Plugins/FX/Niagara/Source/NiagaraEditor/Private/Widgets/SNiagaraParameterPanel.cpp
@@ -149,15 +149,21 @@ void SNiagaraParameterPanel::CollectAllActions(FGraphActionListBuilderBase& OutA
 	TArray<FNiagaraScriptVariableAndViewInfo> VarAndViewInfos = ParameterPanelViewModel->GetViewedParameters();
 	VarAndViewInfos.Sort(SortParameterPanelEntries);
 	//UE_LOG(LogNiagaraEditor, Log, TEXT("Adding Vars--------------------------------------------------"));
+	TArray<FNiagaraVariable> AddedVars;
 	for (const FNiagaraScriptVariableAndViewInfo& VarAndViewInfo : VarAndViewInfos)
 	{
-		const FText Name = FText::FromName(VarAndViewInfo.ScriptVariable.GetName());
-		const FText Tooltip = FText::Format(TooltipFormat, FText::FromName(VarAndViewInfo.ScriptVariable.GetName()), VarAndViewInfo.ScriptVariable.GetType().GetNameText());
-		const NiagaraParameterPanelSectionID::Type Section = ParameterPanelViewModel->GetSectionForVarAndViewInfo(VarAndViewInfo);
-		TSharedPtr<FNiagaraScriptVarAndViewInfoAction> ScriptVarAndViewInfoAction(new FNiagaraScriptVarAndViewInfoAction(VarAndViewInfo, FText::GetEmpty(), Name, Tooltip, 0, FText(), Section /*= 0*/));
-		OutAllActions.AddAction(ScriptVarAndViewInfoAction);
-
-		//UE_LOG(LogNiagaraEditor, Log, TEXT("%s"), *Name.ToString());
+		int32 FoundVarIndex = AddedVars.Find(VarAndViewInfo.ScriptVariable);
+		if (FoundVarIndex < 0)
+		{
+			const FText Name = FText::FromName(VarAndViewInfo.ScriptVariable.GetName());
+			const FText Tooltip = FText::Format(TooltipFormat, FText::FromName(VarAndViewInfo.ScriptVariable.GetName()), VarAndViewInfo.ScriptVariable.GetType().GetNameText());
+			const NiagaraParameterPanelSectionID::Type Section = ParameterPanelViewModel->GetSectionForVarAndViewInfo(VarAndViewInfo);
+			TSharedPtr<FNiagaraScriptVarAndViewInfoAction> ScriptVarAndViewInfoAction(new FNiagaraScriptVarAndViewInfoAction(VarAndViewInfo, FText::GetEmpty(), Name, Tooltip, 0, FText(), Section /*= 0*/));
+			OutAllActions.AddAction(ScriptVarAndViewInfoAction);
+			AddedVars.Add(VarAndViewInfo.ScriptVariable);
+
+			//UE_LOG(LogNiagaraEditor, Log, TEXT("%s"), *Name.ToString());
+		}
 
 	}
 }
