Niagara ribbons are now preserving multi-ribbon ordering to prevent random flickering when camera moves around.
Disabled multi-ribbon ordering when using opaque materials.

diff --git a/Engine/Plugins/FX/Niagara/Source/Niagara/Private/NiagaraRendererRibbons.cpp b/Engine/Plugins/FX/Niagara/Source/Niagara/Private/NiagaraRendererRibbons.cpp
--- a/Engine/Plugins/FX/Niagara/Source/Niagara/Private/NiagaraRendererRibbons.cpp
+++ b/Engine/Plugins/FX/Niagara/Source/Niagara/Private/NiagaraRendererRibbons.cpp
@@ -93,14 +93,34 @@ struct FNiagaraDynamicDataRibbon : public FNiagaraDynamicDataBase
 	// The list of all segments, each one connecting SortedIndices[SegmentId] to SortedIndices[SegmentId + 1].
 	// We use this format because the final index buffer gets generated based on view sorting and InterpCount.
 	TArray<int32> SegmentData;
+	/** The list of all particle (instance) indices. Converts raw indices to particles indices. Ordered along each ribbons, from head to tail. */
 	TArray<int32> SortedIndices;
+	/** The tangent and distance between segments, for each raw index (raw VS particle indices). */
 	TArray<FVector4> TangentAndDistances;
+	/** The multi ribbon index, for each raw index. (raw VS particle indices). */
 	TArray<uint32> MultiRibbonIndices;
+	/** Data for each multi ribbon. There are several entries per ribbon. */
 	TArray<float> PackedPerRibbonDataByIndex;
 
-	// start and end world space position of the ribbon, to figure out draw direction
-	FVector StartPos;
-	FVector EndPos;
+	struct FMultiRibbonInfo
+	{
+		/** start and end world space position of the ribbon, to figure out draw direction */
+		FVector StartPos;
+		FVector EndPos;
+		int32 BaseSegmentDataIndex = 0;
+		int32 NumSegmentDataIndices = 0;
+
+		FORCEINLINE bool UseInvertOrder(const FVector& ViewDirection, const FVector& ViewOriginForDistanceCulling, ENiagaraRibbonDrawDirection DrawDirection) const
+		{
+			const float StartDist = FVector::DotProduct(ViewDirection, StartPos - ViewOriginForDistanceCulling);
+			const float EndDist = FVector::DotProduct(ViewDirection, EndPos - ViewOriginForDistanceCulling);
+			return ((StartDist >= EndDist) && DrawDirection == ENiagaraRibbonDrawDirection::BackToFront)
+				|| ((StartDist < EndDist) && DrawDirection == ENiagaraRibbonDrawDirection::FrontToBack);
+		}
+	};
+
+	/** Ribbon perperties required for sorting. */
+	TArray<FMultiRibbonInfo> MultiRibbonInfos;
 
 	void PackPerRibbonData(float U0Scale, float U0Offset, float U1Scale, float U1Offset, uint32 NumSegments, uint32 FirstParticleId)
 	{
@@ -244,9 +264,8 @@ void FNiagaraRendererRibbons::CreateRenderThreadResources(NiagaraEmitterInstance
 }
 
 template <typename TValue>
-void FNiagaraRendererRibbons::GenerateIndexBuffer(FGlobalDynamicIndexBuffer::FAllocationEx& InOutIndexAllocation, const TArray<int32>& SegmentData, int32 InterpCount, bool bInvertOrder)
+TValue* FNiagaraRendererRibbons::AppendToIndexBuffer(TValue* OutIndices, uint32& OutMaxUsedIndex, const TArrayView<int32>& SegmentData, int32 InterpCount, bool bInvertOrder)
 {
-	TValue* OutIndices = (TValue*)InOutIndexAllocation.Buffer;
 	TValue MaxIndex = 0;
 
 	auto AddTriangleIndices = [&MaxIndex, &OutIndices, InterpCount](int32 SegmentIndex)
@@ -282,7 +301,38 @@ void FNiagaraRendererRibbons::GenerateIndexBuffer(FGlobalDynamicIndexBuffer::FAl
 		}
 	}
 
-	InOutIndexAllocation.MaxUsedIndex = MaxIndex;
+	OutMaxUsedIndex = MaxIndex;
+	return OutIndices;
+}
+
+template <typename TValue>
+void FNiagaraRendererRibbons::GenerateIndexBuffer(
+	FGlobalDynamicIndexBuffer::FAllocationEx& InOutIndexAllocation, 
+	int32 InterpCount, 
+	const FVector& ViewDirection, 
+	const FVector& ViewOriginForDistanceCulling, 
+	FNiagaraDynamicDataRibbon* DynamicData) const
+{
+	check(DynamicData);
+
+	FMaterialRenderProxy* MaterialRenderProxy = DynamicData->Material;
+	check(MaterialRenderProxy);
+	const EBlendMode BlendMode = MaterialRenderProxy->GetMaterial(FeatureLevel)->GetBlendMode();
+
+	TValue* CurrentIndexBuffer = (TValue*)InOutIndexAllocation.Buffer;
+	if (IsTranslucentBlendMode(BlendMode) && DynamicData->MultiRibbonInfos.Num())
+	{
+		for (const FNiagaraDynamicDataRibbon::FMultiRibbonInfo& MultiRibbonInfo : DynamicData->MultiRibbonInfos)
+		{
+			TArrayView<int32> CurrentSegmentData(DynamicData->SegmentData.GetData() + MultiRibbonInfo.BaseSegmentDataIndex, MultiRibbonInfo.NumSegmentDataIndices);
+			CurrentIndexBuffer = AppendToIndexBuffer(CurrentIndexBuffer, InOutIndexAllocation.MaxUsedIndex, CurrentSegmentData, InterpCount, MultiRibbonInfo.UseInvertOrder(ViewDirection, ViewOriginForDistanceCulling, DrawDirection));
+		}
+	}
+	else // Otherwise ignore multi ribbon ordering.
+	{
+		TArrayView<int32> CurrentSegmentData(DynamicData->SegmentData.GetData(), DynamicData->SegmentData.Num());
+		CurrentIndexBuffer = AppendToIndexBuffer(CurrentIndexBuffer, InOutIndexAllocation.MaxUsedIndex, CurrentSegmentData, InterpCount, false);
+	}
 }
 
 void FNiagaraRendererRibbons::GetDynamicMeshElements(const TArray<const FSceneView*>& Views, const FSceneViewFamily& ViewFamily, uint32 VisibilityMap, FMeshElementCollector& Collector, const FNiagaraSceneProxy *SceneProxy) const
@@ -619,6 +669,12 @@ FNiagaraDynamicDataBase* FNiagaraRendererRibbons::GenerateDynamicData(const FNia
 
 		if (NumSegments > 0)
 		{
+			FNiagaraDynamicDataRibbon::FMultiRibbonInfo& MultiRibbonInfo = DynamicData->MultiRibbonInfos[RibbonIndex];
+			MultiRibbonInfo.StartPos = PosData[RibbonIndices[0]];
+			MultiRibbonInfo.EndPos = PosData[RibbonIndices.Last()];
+			MultiRibbonInfo.BaseSegmentDataIndex = SegmentData.Num();
+			MultiRibbonInfo.NumSegmentDataIndices = NumSegments;
+
 			// Update the tangents for the first and last vertex, apply a reflect vector logic so that the initial and final curvature is continuous.
 			if (NumSegments > 1)
 			{
@@ -648,9 +704,7 @@ FNiagaraDynamicDataBase* FNiagaraRendererRibbons::GenerateDynamicData(const FNia
 		}
 	};
 
-	// store the start and end positions for the ribbon for draw distance flipping 
-	DynamicData->StartPos = PosData[0];
-	DynamicData->EndPos = PosData[Data.GetCurrentDataChecked().GetNumInstances() - 1];
+	DynamicData->MultiRibbonInfos.Reset();
 
 	//TODO: Move sorting to share code with sprite and mesh sorting and support the custom sorting key.
 	int32 TotalIndices = Data.GetCurrentDataChecked().GetNumInstances();
@@ -662,6 +716,7 @@ FNiagaraDynamicDataBase* FNiagaraRendererRibbons::GenerateDynamicData(const FNia
 		{
 			SortedIndices.Add(i);
 		}
+		DynamicData->MultiRibbonInfos.AddZeroed(1);
 
 		SortedIndices.Sort([&SortKeyData](const int32& A, const int32& B) {	return (SortKeyData[A] < SortKeyData[B]); });
 
@@ -678,6 +733,7 @@ FNiagaraDynamicDataBase* FNiagaraRendererRibbons::GenerateDynamicData(const FNia
 				TArray<int32>& Indices = MultiRibbonSortedIndices.FindOrAdd(RibbonFullIDData[i]);
 				Indices.Add(i);
 			}
+			DynamicData->MultiRibbonInfos.AddZeroed(MultiRibbonSortedIndices.Num());
 
 			uint32 RibbonIndex = 0;
 			for (TPair<FNiagaraID, TArray<int32>>& Pair : MultiRibbonSortedIndices)
@@ -701,6 +757,7 @@ FNiagaraDynamicDataBase* FNiagaraRendererRibbons::GenerateDynamicData(const FNia
 				TArray<int32>& Indices = MultiRibbonSortedIndices.FindOrAdd(RibbonIdData[i]);
 				Indices.Add(i);
 			}
+			DynamicData->MultiRibbonInfos.AddZeroed(MultiRibbonSortedIndices.Num());
 
 			uint32 RibbonIndex = 0;
 			for (TPair<int32, TArray<int32>>& Pair : MultiRibbonSortedIndices)
@@ -976,13 +1033,6 @@ void FNiagaraRendererRibbons::CreatePerViewResources(
 		NumSegments *= SegmentTessellation;
 	}
 
-	// Figure out whether start is closer to the view plane than end
-	// TODO : This doesn't work with multi-ribbons.
-	const float StartDist = FVector::DotProduct(View->GetViewDirection(), DynamicDataRibbon->StartPos - ViewOriginForDistanceCulling);
-	const float EndDist = FVector::DotProduct(View->GetViewDirection(), DynamicDataRibbon->EndPos - ViewOriginForDistanceCulling);
-	const bool bInvertOrder = ((StartDist > EndDist) && DrawDirection == ENiagaraRibbonDrawDirection::BackToFront)
-		|| ((StartDist < EndDist) && DrawDirection == ENiagaraRibbonDrawDirection::FrontToBack);
-
 	// Copy the index data over.
 	FGlobalDynamicIndexBuffer& DynamicIndexBuffer = Collector.GetDynamicIndexBuffer();
 
@@ -992,12 +1042,12 @@ void FNiagaraRendererRibbons::CreatePerViewResources(
 	if (MaxTessellationIndex < MAX_uint16)
 	{
 		InOutIndexAllocation = DynamicIndexBuffer.Allocate<uint16>(NumIndices);
-		GenerateIndexBuffer<uint16>(InOutIndexAllocation, DynamicDataRibbon->SegmentData, SegmentTessellation, bInvertOrder);
+		GenerateIndexBuffer<uint16>(InOutIndexAllocation, SegmentTessellation, View->GetViewDirection(), ViewOriginForDistanceCulling, DynamicDataRibbon);
 	}
 	else
 	{
 		InOutIndexAllocation = DynamicIndexBuffer.Allocate<uint32>(NumIndices);
-		GenerateIndexBuffer<uint32>(InOutIndexAllocation, DynamicDataRibbon->SegmentData, SegmentTessellation, bInvertOrder);
+		GenerateIndexBuffer<uint32>(InOutIndexAllocation, SegmentTessellation, View->GetViewDirection(), ViewOriginForDistanceCulling, DynamicDataRibbon);
 	}
 
 	FNiagaraRibbonUniformParameters PerViewUniformParameters;
diff --git a/Engine/Plugins/FX/Niagara/Source/Niagara/Public/NiagaraRendererRibbons.h b/Engine/Plugins/FX/Niagara/Source/Niagara/Public/NiagaraRendererRibbons.h
--- a/Engine/Plugins/FX/Niagara/Source/Niagara/Public/NiagaraRendererRibbons.h
+++ b/Engine/Plugins/FX/Niagara/Source/Niagara/Public/NiagaraRendererRibbons.h
@@ -35,8 +35,18 @@ public:
 
 	FORCEINLINE void AddDynamicParam(TArray<FNiagaraRibbonVertexDynamicParameter>& ParamData, const FVector4& DynamicParam);
 protected:
+
+	template <typename TValue>
+	static TValue* AppendToIndexBuffer(TValue* OutIndices, uint32& OutMaxUsedIndex, const TArrayView<int32>& SegmentData, int32 InterpCount, bool bInvertOrder);
+
+	/** Generate the raw index buffer preserving multi ribbon ordering. */
 	template <typename TValue>
-	static void GenerateIndexBuffer(FGlobalDynamicIndexBuffer::FAllocationEx& InOutIndexAllocation, const TArray<int32>& SegmentData, int32 InterpCount, bool bInvertOrder);
+	void GenerateIndexBuffer(
+		FGlobalDynamicIndexBuffer::FAllocationEx& InOutIndexAllocation, 
+		int32 InterpCount, 
+		const FVector& ViewDirection, 
+		const FVector& ViewOriginForDistanceCulling, 
+		struct FNiagaraDynamicDataRibbon* DynamicData) const;
 
 private:
 	struct FCPUSimParticleDataAllocation
