Hiding default value pin but keeping custom UI for displaying the default value.

diff --git a/Engine/Plugins/FX/Niagara/Source/NiagaraEditor/Private/Widgets/SNiagaraGraphParameterMapGetNode.cpp b/Engine/Plugins/FX/Niagara/Source/NiagaraEditor/Private/Widgets/SNiagaraGraphParameterMapGetNode.cpp
--- a/Engine/Plugins/FX/Niagara/Source/NiagaraEditor/Private/Widgets/SNiagaraGraphParameterMapGetNode.cpp
+++ b/Engine/Plugins/FX/Niagara/Source/NiagaraEditor/Private/Widgets/SNiagaraGraphParameterMapGetNode.cpp
@@ -28,6 +28,7 @@ void SNiagaraGraphParameterMapGetNode::AddPin(const TSharedRef<SGraphPin>& PinTo
 
 	const UEdGraphPin* PinObj = PinToAdd->GetPinObj();
 	const bool bAdvancedParameter = (PinObj != nullptr) && PinObj->bAdvancedView;
+	const bool bInvisiblePin = (PinObj != nullptr) && PinObj->bDefaultValueIsReadOnly;
 	if (bAdvancedParameter)
 	{
 		PinToAdd->SetVisibility(TAttribute<EVisibility>(PinToAdd, &SGraphPin::IsPinVisibleAsAdvanced));
@@ -36,6 +37,11 @@ void SNiagaraGraphParameterMapGetNode::AddPin(const TSharedRef<SGraphPin>& PinTo
 	// Save the UI building for later...
 	if (PinToAdd->GetDirection() == EEdGraphPinDirection::EGPD_Input)
 	{
+		if (bInvisiblePin)
+		{
+			//PinToAdd->SetOnlyShowDefaultValue(true);
+			PinToAdd->SetPinColorModifier(FLinearColor(0.0f, 0.0f, 0.0f, 0.0f));
+		}
 		InputPins.Add(PinToAdd);
 	}
 	else // Direction == EEdGraphPinDirection::EGPD_Output
@@ -63,6 +69,7 @@ void SNiagaraGraphParameterMapGetNode::CreatePinWidgets()
 	
 	UNiagaraNodeParameterMapGet* GetNode = Cast<UNiagaraNodeParameterMapGet>(GraphNode);
 
+
 	// Deferred pin adding to line up input/output pins by name.
 	for (int32 i = 0; i < OutputPins.Num() + 1; i++)
 	{
