Avoid possibility of nullptr use / checks inside crash rerport handler.

diff --git a/Engine/Plugins/FX/Niagara/Source/Niagara/Private/NiagaraCrashReporterHandler.cpp b/Engine/Plugins/FX/Niagara/Source/Niagara/Private/NiagaraCrashReporterHandler.cpp
--- a/Engine/Plugins/FX/Niagara/Source/Niagara/Private/NiagaraCrashReporterHandler.cpp
+++ b/Engine/Plugins/FX/Niagara/Source/Niagara/Private/NiagaraCrashReporterHandler.cpp
@@ -37,9 +37,9 @@ FNiagaraCrashReporterHandler::FNiagaraCrashReporterHandler()
 		ThreadName = FThreadManager::Get().GetThreadName(ThreadId);
 	}
 
-	Name = FString::Printf(TEXT("NiagaraCRInfo | %s(%u)"), *ThreadName, ThreadId);
-	
-	NoneInfoString = FString::Printf(TEXT("%s | None |"), *Name);
+	Name = FString::Printf(TEXT("NiagaraCRInfo_%s_%u"), *ThreadName, ThreadId);
+	Name = Name.Replace(TEXT(" "), TEXT("_"));
+	NoneInfoString = TEXT("None");
 
 	FScopeLock Lock(&CritSec);
 	AllHandlers.AddUnique(this);
@@ -53,11 +53,11 @@ FNiagaraCrashReporterHandler::~FNiagaraCrashReporterHandler()
 
 void FNiagaraCrashReporterHandler::Clear()
 {
-	ScopeInfoStack.Empty();
-
 	FScopeLock ScopeLock(&CritSec);
+
+	ScopeInfoStack.Empty();
 	FString Empty;
-	UE_LOG(LogNiagara, Log, TEXT("SetEngineData: %s"), *Empty);
+	//UE_LOG(LogNiagara, Log, TEXT("SetEngineData: %s"), *Empty);
 	FGenericCrashContext::SetEngineData(Name, Empty);
 }
 
@@ -74,35 +74,38 @@ void FNiagaraCrashReporterHandler::OnEnabledChanged(IConsoleVariable* InVariable
 
 void FNiagaraCrashReporterHandler::SetInfo(const FString& Info)
 {
-	FScopeLock ScopeLock(&CritSec);
 	//UE_LOG(LogNiagara, Log, TEXT("SetEngineData: %s"), *Info);
 	FGenericCrashContext::SetEngineData(Name, Info);
 }
 
+static const FString NullStr(TEXT("nullptr"));
 void FNiagaraCrashReporterHandler::PushInfo(FNiagaraSystemInstance* Inst)
 {
-	check(Inst && Inst->GetComponent() && Inst->GetComponent()->GetAsset());
-	UNiagaraComponent* Comp = Inst->GetComponent();
-	USceneComponent* AttachParent = Comp->GetAttachParent();
-
-	PushInfo(FString::Printf(TEXT("%s | SystemInstance | System: %s | bSolo: %s | Component: %s | AttachedTo: %s |"), 
-	*Name, *Comp->GetAsset()->GetFullName(), Inst->IsSolo() ? TEXT("true") : TEXT("false"), *Comp->GetFullName(), AttachParent ? *AttachParent->GetFullName() : TEXT("nullptr")));
-
-	//TODO: More info?
+	FScopeLock ScopeLock(&CritSec);
+	if (Inst)
+	{
+		PushInfo_Internal(Inst->GetCrashReporterTag());
+	}
+	else
+	{
+		PushInfo_Internal(NullStr);
+	}
 }
 
 void FNiagaraCrashReporterHandler::PushInfo(FNiagaraSystemSimulation* SystemSim)
 {
-	check(SystemSim && SystemSim->GetSystem());
-	
-	UNiagaraSystem* System = SystemSim->GetSystem();
-
-	PushInfo(FString::Printf(TEXT("%s | SystemSimulation | System: %s | bSolo: %s |"), *Name, *System->GetFullName(), SystemSim->GetIsSolo() ? TEXT("true") : TEXT("false")));
-
-	//TODO: More info?
+	FScopeLock ScopeLock(&CritSec);
+	if (SystemSim)
+	{
+		PushInfo_Internal(SystemSim->GetCrashReporterTag());
+	}
+	else
+	{
+		PushInfo_Internal(NullStr);
+	}
 }
 
-void FNiagaraCrashReporterHandler::PushInfo(const FString& Info)
+void FNiagaraCrashReporterHandler::PushInfo_Internal(const FString& Info)
 {
 	ScopeInfoStack.Push(Info);
 	SetInfo(Info);
@@ -110,6 +113,7 @@ void FNiagaraCrashReporterHandler::PushInfo(const FString& Info)
 
 void FNiagaraCrashReporterHandler::PopInfo()
 {
+	FScopeLock ScopeLock(&CritSec);
 	ScopeInfoStack.Pop();
 	if (ScopeInfoStack.Num())
 	{
diff --git a/Engine/Plugins/FX/Niagara/Source/Niagara/Private/NiagaraSystemInstance.cpp b/Engine/Plugins/FX/Niagara/Source/Niagara/Private/NiagaraSystemInstance.cpp
--- a/Engine/Plugins/FX/Niagara/Source/Niagara/Private/NiagaraSystemInstance.cpp
+++ b/Engine/Plugins/FX/Niagara/Source/Niagara/Private/NiagaraSystemInstance.cpp
@@ -2252,3 +2252,19 @@ FNiagaraSystemInstance::FOnDestroyed& FNiagaraSystemInstance::OnDestroyed()
 	return OnDestroyedDelegate;
 }
 #endif
+
+const FString& FNiagaraSystemInstance::GetCrashReporterTag()const
+{
+	if(CrashReporterTag.IsEmpty())
+	{
+		UNiagaraSystem* Sys = Component ? Component->GetAsset() : nullptr;
+		USceneComponent* AttachParent = Component ? Component->GetAttachParent() : nullptr;
+
+		const FString& CompName = Component ? Component->GetFullName() : TEXT("nullptr");
+		const FString& AssetName = Sys ? Sys->GetFullName() : TEXT("nullptr");
+		const FString& AttachName = AttachParent ? AttachParent->GetFullName() : TEXT("nullptr");
+
+		CrashReporterTag = FString::Printf(TEXT("SystemInstance | System: %s | bSolo: %s | Component: %s | AttachedTo: %s |"), *AssetName, IsSolo() ? TEXT("true") : TEXT("false"), *CompName, *AttachName);
+	}
+	return CrashReporterTag;
+}
\ No newline at end of file
diff --git a/Engine/Plugins/FX/Niagara/Source/Niagara/Private/NiagaraSystemSimulation.cpp b/Engine/Plugins/FX/Niagara/Source/Niagara/Private/NiagaraSystemSimulation.cpp
--- a/Engine/Plugins/FX/Niagara/Source/Niagara/Private/NiagaraSystemSimulation.cpp
+++ b/Engine/Plugins/FX/Niagara/Source/Niagara/Private/NiagaraSystemSimulation.cpp
@@ -1999,6 +1999,18 @@ void FNiagaraSystemSimulation::InitParameterDataSetBindings(FNiagaraSystemInstan
 	}
 }
 
+const FString& FNiagaraSystemSimulation::GetCrashReporterTag()const
+{
+	if (CrashReporterTag.IsEmpty())
+	{
+		UNiagaraSystem* Sys = GetSystem();
+		const FString& AssetName = Sys ? Sys->GetFullName() : TEXT("nullptr");
+
+		CrashReporterTag = FString::Printf(TEXT("SystemSimulation | System: %s | bSolo: %s |"), *AssetName, bIsSolo ? TEXT("true") : TEXT("false"));
+	}
+	return CrashReporterTag;
+}
+
 void FNiagaraConstantBufferToDataSetBinding::Init(const FNiagaraSystemCompiledData& CompiledData)
 {
 	// for now we'll copy the data to our local structure so that we don't have to worry about the lifetime of the compiled data
diff --git a/Engine/Plugins/FX/Niagara/Source/Niagara/Public/NiagaraCrashReporterHandler.h b/Engine/Plugins/FX/Niagara/Source/Niagara/Public/NiagaraCrashReporterHandler.h
--- a/Engine/Plugins/FX/Niagara/Source/Niagara/Public/NiagaraCrashReporterHandler.h
+++ b/Engine/Plugins/FX/Niagara/Source/Niagara/Public/NiagaraCrashReporterHandler.h
@@ -28,6 +28,9 @@ class FNiagaraCrashReporterHandler : public TThreadSingleton<FNiagaraCrashReport
 	static FCriticalSection CritSec;
 
 	void SetInfo(const FString& Info);
+
+	/** Push an arbitrary string onto info stack and set as Niagara's current CR info for this thread. */
+	void PushInfo_Internal(const FString& Info);
 public:
 
 	FNiagaraCrashReporterHandler();
@@ -40,8 +43,6 @@ public:
 	void PushInfo(FNiagaraSystemInstance* Inst);
 	/** Push info on a system simulation to the info stack and set as Niagara's current CR info for this thread. */
 	void PushInfo(FNiagaraSystemSimulation* SystemSim);
-	/** Push an arbitrary string onto info stack and set as Niagara's current CR info for this thread. */
-	void PushInfo(const FString& Info);
 	/** Pop the current info from the stack and set the previous info as Niagara's current CR info for this thread. */
 	void PopInfo();
 
diff --git a/Engine/Plugins/FX/Niagara/Source/Niagara/Public/NiagaraSystemInstance.h b/Engine/Plugins/FX/Niagara/Source/Niagara/Public/NiagaraSystemInstance.h
--- a/Engine/Plugins/FX/Niagara/Source/Niagara/Public/NiagaraSystemInstance.h
+++ b/Engine/Plugins/FX/Niagara/Source/Niagara/Public/NiagaraSystemInstance.h
@@ -284,6 +284,7 @@ public:
 
 	FORCEINLINE void SetLODDistance(float InLODDistance, float InMaxLODDistance);
 
+	const FString& GetCrashReporterTag()const;
 private:
 
 	void DestroyDataInterfaceInstanceData();
@@ -431,6 +432,9 @@ private:
 	/** Array of emitter indices sorted by execution priority. The emitters will be ticked in this order. */
 	TArray<int32> EmitterExecutionOrder;
 
+	/** Tag we feed into crash reporter for this instance. */
+	mutable FString CrashReporterTag;
+
 public:
 	// Transient data that is accumulated during tick.
 	uint32 TotalGPUParamSize = 0;
diff --git a/Engine/Plugins/FX/Niagara/Source/Niagara/Public/NiagaraSystemSimulation.h b/Engine/Plugins/FX/Niagara/Source/Niagara/Public/NiagaraSystemSimulation.h
--- a/Engine/Plugins/FX/Niagara/Source/Niagara/Public/NiagaraSystemSimulation.h
+++ b/Engine/Plugins/FX/Niagara/Source/Niagara/Public/NiagaraSystemSimulation.h
@@ -290,6 +290,7 @@ public:
 
 	void AddTickGroupPromotion(FNiagaraSystemInstance* Instance);
 
+	const FString& GetCrashReporterTag()const;
 protected:
 	/** Sets constant parameter values */
 	void SetupParameters_GameThread(float DeltaSeconds);
@@ -408,4 +409,6 @@ protected:
 
 	/** Current task that is executing */
 	FGraphEventRef SystemTickGraphEvent;
+
+	mutable FString CrashReporterTag;
 };
