Removed fast path from Niagara

diff --git a/Engine/Plugins/FX/Niagara/Source/Niagara/Classes/NiagaraEmitter.h b/Engine/Plugins/FX/Niagara/Source/Niagara/Classes/NiagaraEmitter.h
--- a/Engine/Plugins/FX/Niagara/Source/Niagara/Classes/NiagaraEmitter.h
+++ b/Engine/Plugins/FX/Niagara/Source/Niagara/Classes/NiagaraEmitter.h
@@ -9,7 +9,6 @@
 #include "NiagaraScript.h"
 #include "NiagaraCollision.h"
 #include "INiagaraMergeManager.h"
-#include "NiagaraSystemFastPath.h"
 #include "NiagaraEffectType.h"
 #include "NiagaraEmitter.generated.h"
 
@@ -506,29 +505,6 @@ public:
 
 	void OnEffectsQualityChanged();
 
-	void InitFastPathAttributeNames();
-
-	UPROPERTY(EditAnywhere, Category = "Script Fast Path")
-	FNiagaraFastPath_Module_EmitterScalability EmitterScalability;
-
-	UPROPERTY(EditAnywhere, Category = "Script Fast Path")
-	FNiagaraFastPath_Module_EmitterLifeCycle EmitterLifeCycle;
-
-	UPROPERTY(EditAnywhere, Category = "Script Fast Path")
-	TArray<FNiagaraFastPath_Module_SpawnRate> SpawnRate;
-
-	UPROPERTY(EditAnywhere, Category = "Script Fast Path")
-	TArray<FNiagaraFastPath_Module_SpawnPerUnit> SpawnPerUnit;
-
-	UPROPERTY(EditAnywhere, Category = "Script Fast Path")
-	TArray<FNiagaraFastPath_Module_SpawnBurstInstantaneous> SpawnBurstInstantaneous;
-
-	UPROPERTY()
-	FNiagaraFastPathAttributeNames SpawnFastPathAttributeNames;
-
-	UPROPERTY()
-	FNiagaraFastPathAttributeNames UpdateFastPathAttributeNames;
-
 protected:
 	virtual void BeginDestroy() override;
 
diff --git a/Engine/Plugins/FX/Niagara/Source/Niagara/Classes/NiagaraEmitterInstance.h b/Engine/Plugins/FX/Niagara/Source/Niagara/Classes/NiagaraEmitterInstance.h
--- a/Engine/Plugins/FX/Niagara/Source/Niagara/Classes/NiagaraEmitterInstance.h
+++ b/Engine/Plugins/FX/Niagara/Source/Niagara/Classes/NiagaraEmitterInstance.h
@@ -15,7 +15,6 @@ NiagaraEmitterInstance.h: Niagara emitter simulation class
 #include "NiagaraEmitter.h"
 #include "NiagaraScriptExecutionContext.h"
 #include "NiagaraBoundsCalculator.h"
-#include "NiagaraSystemFastPath.h"
 
 class FNiagaraSystemInstance;
 struct FNiagaraEmitterHandle;
@@ -120,35 +119,16 @@ public:
 
 	bool FindBinding(const FNiagaraUserParameterBinding& InBinding, TArray<UMaterialInterface*>& OutMaterials) const;
 
-	void InitFastPathAttributeBindings();
-
-	void TickFastPathAttributeBindings();
-
-	FNiagaraEmitterFastPath::FParamMap0& GetFastPathMap() { return FastPathMap; }
 private:
 	void CheckForErrors();
 
-	void InitFastPathParameterBindingsInternal(const FNiagaraFastPathAttributeNames& SourceParameterNames, FNiagaraParameterStore& TargetParameterStore);
 	void BuildConstantBufferTable(
 		const FNiagaraScriptExecutionContext& ExecContext,
 		FScriptExecutionConstantBufferTable& ConstantBufferTable) const;
 
 	/** Generate emitter bounds */
 	FBox InternalCalculateDynamicBounds(int32 ParticleCount) const;
-
-	template<typename TBindingType, typename TVariableType>
-	static void AddBinding(FName ParameterName, TVariableType ParameterType, TBindingType* SourceValuePtr, FNiagaraParameterStore& TargetParameterStore, TArray<TNiagaraFastPathAttributeBinding<TBindingType>>& TargetBindings)
-	{
-		TNiagaraFastPathAttributeBinding<TBindingType> Binding;
-		FNiagaraVariable ParameterVariable = FNiagaraVariable(ParameterType, ParameterName);
-		Binding.ParameterBinding.Init(TargetParameterStore, ParameterVariable);
-		if (Binding.ParameterBinding.ValuePtr != nullptr)
-		{
-			Binding.ParameterValue = SourceValuePtr;
-			TargetBindings.Add(Binding);
-		}
-	}
-
+	
 	/** The index of our emitter in our parent system instance. */
 	int32 EmitterIdx;
 
@@ -224,10 +204,5 @@ private:
 	/** Optional list of bounds calculators. */
 	TArray<TUniquePtr<FNiagaraBoundsCalculator>, TInlineAllocator<1>> BoundsCalculators;
 
-	FNiagaraEmitterFastPath::FParamMap0 FastPathMap;
-
-	TArray<TNiagaraFastPathAttributeBinding<int32>> FastPathIntAttributeBindings;
-	TArray<TNiagaraFastPathAttributeBinding<float>> FastPathFloatAttributeBindings;
-
 	TSharedPtr<const FNiagaraEmitterCompiledData> CachedEmitterCompiledData;
 };
diff --git a/Engine/Plugins/FX/Niagara/Source/Niagara/Classes/NiagaraSystem.h b/Engine/Plugins/FX/Niagara/Source/Niagara/Classes/NiagaraSystem.h
--- a/Engine/Plugins/FX/Niagara/Source/Niagara/Classes/NiagaraSystem.h
+++ b/Engine/Plugins/FX/Niagara/Source/Niagara/Classes/NiagaraSystem.h
@@ -13,7 +13,6 @@
 #include "NiagaraEmitterHandle.h"
 #include "NiagaraParameterCollection.h"
 #include "NiagaraUserRedirectionParameterStore.h"
-#include "NiagaraSystemFastPath.h"
 #include "NiagaraEffectType.h"
 
 #include "NiagaraSystem.generated.h"
@@ -391,15 +390,6 @@ public:
 	TStatId GetStatID(bool bGameThread, bool bConcurrent)const;
 	void AddToInstanceCountStat(int32 NumInstances, bool bSolo)const;
 
-	UPROPERTY(EditAnywhere, Category = "Script Fast Path")
-	ENiagaraFastPathMode FastPathMode;
-
-	UPROPERTY(EditAnywhere, Category = "Script Fast Path")
-	FNiagaraFastPath_Module_SystemScalability SystemScalability;
-
-	UPROPERTY(EditAnywhere, Category = "Script Fast Path")
-	FNiagaraFastPath_Module_SystemLifeCycle SystemLifeCycle;
-
 private:
 #if WITH_EDITORONLY_DATA
 	/** Checks the ddc for vm execution data for the given script. Return true if the data was loaded from the ddc, false otherwise. */
diff --git a/Engine/Plugins/FX/Niagara/Source/Niagara/Private/NiagaraEmitter.cpp b/Engine/Plugins/FX/Niagara/Source/Niagara/Private/NiagaraEmitter.cpp
--- a/Engine/Plugins/FX/Niagara/Source/Niagara/Private/NiagaraEmitter.cpp
+++ b/Engine/Plugins/FX/Niagara/Source/Niagara/Private/NiagaraEmitter.cpp
@@ -1049,8 +1049,6 @@ void UNiagaraEmitter::OnPostCompile()
 	RuntimeEstimation = MemoryRuntimeEstimation();
 
 	OnEmitterVMCompiled().Broadcast(this);
-
-	InitFastPathAttributeNames();
 }
 
 UNiagaraEmitter* UNiagaraEmitter::MakeRecursiveDeepCopy(UObject* DestOuter) const
@@ -1661,34 +1659,3 @@ void UNiagaraEmitter::OnEffectsQualityChanged()
 {
 	ResolveScalabilitySettings();
 }
-
-void UNiagaraEmitter::InitFastPathAttributeNames()
-{
-	auto InitParameters = [](const FNiagaraParameters& Parameters, const FString& EmitterName, FNiagaraFastPathAttributeNames& FastPathParameterNames)
-	{
-		FastPathParameterNames.System.Empty();
-		FastPathParameterNames.SystemFullNames.Empty();
-		FastPathParameterNames.Emitter.Empty();
-		FastPathParameterNames.EmitterFullNames.Empty();
-
-		FString SystemPrefix = TEXT("System.");
-		FString EmitterPrefix = EmitterName + TEXT(".");
-		for (const FNiagaraVariable& Parameter : Parameters.Parameters)
-		{
-			FString ParameterNameString = Parameter.GetName().ToString();
-			if (ParameterNameString.StartsWith(SystemPrefix))
-			{
-				FastPathParameterNames.System.Add(*(Parameter.GetName().ToString().RightChop(SystemPrefix.Len())));
-				FastPathParameterNames.SystemFullNames.Add(Parameter.GetName());
-			}
-			else if (ParameterNameString.StartsWith(EmitterPrefix))
-			{
-				FastPathParameterNames.Emitter.Add(*(Parameter.GetName().ToString().RightChop(EmitterPrefix.Len())));
-				FastPathParameterNames.EmitterFullNames.Add(Parameter.GetName());
-			}
-		}
-	};
-
-	InitParameters(SpawnScriptProps.Script->GetVMExecutableData().Parameters, UniqueEmitterName, SpawnFastPathAttributeNames);
-	InitParameters(UpdateScriptProps.Script->GetVMExecutableData().Parameters, UniqueEmitterName, UpdateFastPathAttributeNames);
-}
diff --git a/Engine/Plugins/FX/Niagara/Source/Niagara/Private/NiagaraEmitterInstance.cpp b/Engine/Plugins/FX/Niagara/Source/Niagara/Private/NiagaraEmitterInstance.cpp
--- a/Engine/Plugins/FX/Niagara/Source/Niagara/Private/NiagaraEmitterInstance.cpp
+++ b/Engine/Plugins/FX/Niagara/Source/Niagara/Private/NiagaraEmitterInstance.cpp
@@ -267,14 +267,7 @@ void FNiagaraEmitterInstance::Init(int32 InEmitterIdx, FNiagaraSystemInstanceID
 
 	//Init the spawn infos to the correct number for this system.
 	int32 NumEvents = CachedEmitter->GetEventHandlers().Num();
-	if (ParentSystemInstance->GetSystem()->FastPathMode != ENiagaraFastPathMode::ScriptVMOnly)
-	{
-		SpawnInfos.SetNum(CachedEmitter->SpawnRate.Num() + CachedEmitter->SpawnPerUnit.Num() + CachedEmitter->SpawnBurstInstantaneous.Num());
-	}
-	else
-	{
-		SpawnInfos.SetNum(CachedEmitterCompiledData->SpawnAttributes.Num());
-	}
+	SpawnInfos.SetNum(CachedEmitterCompiledData->SpawnAttributes.Num());
 
 	{
 		SCOPE_CYCLE_COUNTER(STAT_NiagaraEmitterInit_Datasets);
@@ -487,10 +480,6 @@ void FNiagaraEmitterInstance::ResetSimulation(bool bKillExisting /*= true*/)
 		{
 			GPUExecContext->Reset(Batcher);
 		}
-
-		FastPathMap = FNiagaraEmitterFastPath::FParamMap0();
-		FastPathIntAttributeBindings.Empty();
-		FastPathFloatAttributeBindings.Empty();
 	}
 }
 
@@ -1812,30 +1801,6 @@ bool FNiagaraEmitterInstance::FindBinding(const FNiagaraUserParameterBinding& In
 	return false;
 }
 
-void FNiagaraEmitterInstance::InitFastPathAttributeBindings()
-{
-	FastPathIntAttributeBindings.Empty();
-	FastPathFloatAttributeBindings.Empty();
-
-	FNiagaraEmitterFastPath::InitFastPathAttributeBindings(CachedEmitter->SpawnFastPathAttributeNames, SpawnExecContext.Parameters, 
-		GetParentSystemInstance()->FastPathMap, FastPathMap, FastPathIntAttributeBindings, FastPathFloatAttributeBindings);
-	FNiagaraEmitterFastPath::InitFastPathAttributeBindings(CachedEmitter->UpdateFastPathAttributeNames, UpdateExecContext.Parameters, 
-		GetParentSystemInstance()->FastPathMap, FastPathMap, FastPathIntAttributeBindings, FastPathFloatAttributeBindings);
-}
-
-void FNiagaraEmitterInstance::TickFastPathAttributeBindings()
-{
-	for (TNiagaraFastPathAttributeBinding<int32>& IntBinding : FastPathIntAttributeBindings)
-	{
-		IntBinding.Tick();
-	}
-
-	for (TNiagaraFastPathAttributeBinding<float>& FloatBinding : FastPathFloatAttributeBindings)
-	{
-		FloatBinding.Tick();
-	}
-}
-
 void FNiagaraEmitterInstance::BuildConstantBufferTable(
 	const FNiagaraScriptExecutionContext& ExecContext,
 	FScriptExecutionConstantBufferTable& ConstantBufferTable) const
diff --git a/Engine/Plugins/FX/Niagara/Source/Niagara/Private/NiagaraSystem.cpp b/Engine/Plugins/FX/Niagara/Source/Niagara/Private/NiagaraSystem.cpp
--- a/Engine/Plugins/FX/Niagara/Source/Niagara/Private/NiagaraSystem.cpp
+++ b/Engine/Plugins/FX/Niagara/Source/Niagara/Private/NiagaraSystem.cpp
@@ -48,7 +48,6 @@ static FAutoConsoleVariableRef CVarNiagaraForceSystemsToCookOutRapidIterationOnL
 UNiagaraSystem::UNiagaraSystem(const FObjectInitializer& ObjectInitializer)
 : Super(ObjectInitializer)
 , bFixedBounds(false)
-, FastPathMode(ENiagaraFastPathMode::ScriptVMOnly)
 #if WITH_EDITORONLY_DATA
 , bIsolateEnabled(false)
 #endif
diff --git a/Engine/Plugins/FX/Niagara/Source/Niagara/Private/NiagaraSystemFastPath.cpp b/Engine/Plugins/FX/Niagara/Source/Niagara/Private/NiagaraSystemFastPath.cpp
deleted file mode 100644
diff --git a/Engine/Plugins/FX/Niagara/Source/Niagara/Private/NiagaraSystemInstance.cpp b/Engine/Plugins/FX/Niagara/Source/Niagara/Private/NiagaraSystemInstance.cpp
--- a/Engine/Plugins/FX/Niagara/Source/Niagara/Private/NiagaraSystemInstance.cpp
+++ b/Engine/Plugins/FX/Niagara/Source/Niagara/Private/NiagaraSystemInstance.cpp
@@ -1006,11 +1006,6 @@ void FNiagaraSystemInstance::ReInitInternal()
 
 	bAlreadyBound = false;
 
-	FastPathIntUserParameterInputBindings.Empty();
-	FastPathFloatUserParameterInputBindings.Empty();
-	FastPathIntUpdateRangedInputBindings.Empty();
-	FastPathFloatUpdateRangedInputBindings.Empty();
-
 	UNiagaraSystem* System = GetSystem();
 	if (System == nullptr || Component == nullptr)
 	{
@@ -1721,37 +1716,6 @@ void FNiagaraSystemInstance::TickInstanceParameters_Concurrent()
 	InstanceParameters.MarkParametersDirty();
 }
 
-void FNiagaraSystemInstance::TickFastPathBindings()
-{
-	for (TNiagaraFastPathUserParameterInputBinding<int32>& IntUserParameterInputBinding : FastPathIntUserParameterInputBindings)
-	{
-		IntUserParameterInputBinding.Tick();
-	}
-
-	for (TNiagaraFastPathUserParameterInputBinding<float>& FloatUserParameterInputBinding : FastPathFloatUserParameterInputBindings)
-	{
-		FloatUserParameterInputBinding.Tick();
-	}
-
-	for (TNiagaraFastPathRangedInputBinding<int32>& IntUpdateRangedInputBinding : FastPathIntUpdateRangedInputBindings)
-	{
-		IntUpdateRangedInputBinding.Tick();
-	}
-
-	for (TNiagaraFastPathRangedInputBinding<float>& FloatUpdateRangedInputBinding : FastPathFloatUpdateRangedInputBindings)
-	{
-		FloatUpdateRangedInputBinding.Tick();
-	}
-}
-
-void FNiagaraSystemInstance::ResetFastPathBindings()
-{
-	FastPathIntUserParameterInputBindings.Empty();
-	FastPathFloatUserParameterInputBindings.Empty();
-	FastPathIntUpdateRangedInputBindings.Empty();
-	FastPathFloatUpdateRangedInputBindings.Empty();
-}
-
 void
 FNiagaraSystemInstance::ClearEventDataSets()
 {
diff --git a/Engine/Plugins/FX/Niagara/Source/Niagara/Private/NiagaraSystemSimulation.cpp b/Engine/Plugins/FX/Niagara/Source/Niagara/Private/NiagaraSystemSimulation.cpp
--- a/Engine/Plugins/FX/Niagara/Source/Niagara/Private/NiagaraSystemSimulation.cpp
+++ b/Engine/Plugins/FX/Niagara/Source/Niagara/Private/NiagaraSystemSimulation.cpp
@@ -32,7 +32,6 @@ DECLARE_CYCLE_STAT(TEXT("System Sim Update [CNC]"), STAT_NiagaraSystemSim_Update
 DECLARE_CYCLE_STAT(TEXT("System Sim Spawn [CNC]"), STAT_NiagaraSystemSim_SpawnCNC, STATGROUP_Niagara);
 DECLARE_CYCLE_STAT(TEXT("System Sim Transfer Results [CNC]"), STAT_NiagaraSystemSim_TransferResultsCNC, STATGROUP_Niagara);
 DECLARE_CYCLE_STAT(TEXT("System Sim Init [GT]"), STAT_NiagaraSystemSim_Init, STATGROUP_Niagara);
-DECLARE_CYCLE_STAT(TEXT("System Sim FastPath [CNC]"), STAT_NiagaraSystemSim_FastPathCNC, STATGROUP_Niagara);
 
 DECLARE_CYCLE_STAT(TEXT("System Sim Init (DataSets) [GT]"), STAT_NiagaraSystemSim_Init_DataSets, STATGROUP_Niagara);
 DECLARE_CYCLE_STAT(TEXT("System Sim Init (ExecContexts) [GT]"), STAT_NiagaraSystemSim_Init_ExecContexts, STATGROUP_Niagara);
@@ -1083,24 +1082,16 @@ void FNiagaraSystemSimulation::Tick_Concurrent(FNiagaraSystemSimulationTickConte
 			SystemInstance->TickInstanceParameters_Concurrent();
 		}
 
-		if (Context.Owner->GetSystem()->FastPathMode != ENiagaraFastPathMode::ScriptVMOnly)
-		{
-			TickFastPath(Context);
-		}
+		PrepareForSystemSimulate(Context);
 
-		if(Context.Owner->GetSystem()->FastPathMode != ENiagaraFastPathMode::FastPathOnly)
+		if (Context.SpawnNum > 0)
 		{
-			PrepareForSystemSimulate(Context);
-
-			if (Context.SpawnNum > 0)
-			{
-				SpawnSystemInstances(Context);
-			}
+			SpawnSystemInstances(Context);
+		}
 
-			UpdateSystemInstances(Context);
+		UpdateSystemInstances(Context);
 
-			TransferSystemSimResults(Context);
-		}
+		TransferSystemSimResults(Context);
 
 		for (FNiagaraSystemInstance* Instance : Context.Instances)
 		{
@@ -1140,227 +1131,6 @@ void FNiagaraSystemSimulation::Tick_Concurrent(FNiagaraSystemSimulationTickConte
 	}
 }
 
-void FNiagaraSystemSimulation::TickFastPath(FNiagaraSystemSimulationTickContext& Context)
-{
-	SCOPE_CYCLE_COUNTER(STAT_NiagaraSystemSim_FastPathCNC);
-
-	// PrepareForSystemSimulate
-	for (FNiagaraSystemInstance* SystemInstance : Context.Instances)
-	{
-		const auto& SystemParameters = SystemInstance->GetSystemParameters();
-
-		FNiagaraSystemFastPath::FParamMap0& SystemMap = SystemInstance->GetFastPathMap();
-		SystemMap.Engine.Owner.ExecutionState = SystemInstance->GetRequestedExecutionState();
-		SystemMap.Engine.Owner.TimeSinceRendered = SystemInstance->GetSystemParameters().EngineTimeSinceRendered;
-
-		TArray<TSharedRef<FNiagaraEmitterInstance, ESPMode::ThreadSafe>>& EmitterInstances = SystemInstance->GetEmitters();
-		for (int32 EmitterIndex = 0; EmitterIndex < EmitterInstances.Num(); ++EmitterIndex)
-		{
-			TSharedRef<FNiagaraEmitterInstance, ESPMode::ThreadSafe>& EmitterInstance = EmitterInstances[EmitterIndex];
-			if (EmitterInstance->GetExecutionState() == ENiagaraExecutionState::Disabled)
-			{
-				continue;
-			}
-			FNiagaraEmitterFastPath::FParamMap0& EmitterMap = EmitterInstance->GetFastPathMap();
-			EmitterMap.Engine.Owner.LODDistance = SystemParameters.EngineLodDistance;
-			EmitterMap.Engine.Owner.MaxLODDistance = SystemParameters.EngineLodDistance / SystemParameters.EngineLodDistanceFraction;
-			EmitterMap.Engine.DeltaTime = Context.DeltaSeconds;
-			EmitterMap.Engine.Emitter.NumParticles = SystemInstance->GetEmitterParameters(EmitterIndex).EmitterNumParticles;
-			EmitterMap.Engine.Owner.Velocity = SystemInstance->GetOwnerParameters().EngineVelocity;
-			EmitterMap.Emitter.ExecutionState = EmitterInstance->GetExecutionState();
-			EmitterMap.Engine.GlobalSpawnCountScale = INiagaraModule::GetGlobalSpawnCountScale();
-
-			if (UNiagaraEmitter* CachedEmitter = EmitterInstance->GetCachedEmitter())
-			{
-				const FNiagaraEmitterScalabilitySettings& ScalabilitySettings = CachedEmitter->GetScalabilitySettings();
-				EmitterMap.Emitter.SpawnCountScale = ScalabilitySettings.SpawnCountScale;
-			}
-		}
-	}
-
-	// SpawnSystemInstances
-	int32 CurrentNumInstances = Context.Instances.Num() - Context.SpawnNum;
-	if (Context.SpawnNum > 0)
-	{
-		// There is no simulation work to do here, but the number of instances in the data set needs to be correct for other codepaths.
-		FNiagaraDataSet& SimulationDataSet = Context.DataSet;
-		SimulationDataSet.BeginSimulate();
-		SimulationDataSet.Allocate(Context.Instances.Num());
-		SimulationDataSet.GetDestinationDataChecked().SetNumInstances(Context.Instances.Num());
-		SimulationDataSet.EndSimulate();
-
-		for(int32 SpawnIndex = CurrentNumInstances; SpawnIndex < Context.Instances.Num(); ++SpawnIndex)
-		{
-			FNiagaraSystemInstance* SystemInstance = Context.Instances[SpawnIndex];
-			SystemInstance->ResetFastPathBindings();
-
-			FNiagaraSystemFastPath::SetSpawnMapDefaults(SystemInstance->FastPathMap);
-			TArray<TSharedRef<FNiagaraEmitterInstance, ESPMode::ThreadSafe>>& EmitterInstances = SystemInstance->GetEmitters();
-			for (TSharedRef<FNiagaraEmitterInstance, ESPMode::ThreadSafe>& EmitterInstance : EmitterInstances)
-			{
-				const UNiagaraEmitter* Emitter = EmitterInstance->GetCachedEmitter();
-				if (EmitterInstance->GetExecutionState() == ENiagaraExecutionState::Disabled)
-				{
-					continue;
-				}
-
-				FNiagaraEmitterFastPath::FParamMap0& EmitterMap = EmitterInstance->GetFastPathMap();
-				FNiagaraEmitterFastPath::SetSpawnMapDefaults(*Emitter, EmitterMap);
-				EmitterInstance->InitFastPathAttributeBindings();
-
-				EmitterMap.SpawnRate.SetNumUninitialized(Emitter->SpawnRate.Num());
-				for (int32 SpawnRateIndex = 0; SpawnRateIndex < Emitter->SpawnRate.Num(); ++SpawnRateIndex)
-				{
-					EmitterMap.SpawnRate[SpawnRateIndex].Init(
-						Emitter, SpawnRateIndex, SystemInstance->GetInstanceParameters(),
-						SystemInstance->FastPathIntUserParameterInputBindings, SystemInstance->FastPathFloatUserParameterInputBindings,
-						SystemInstance->FastPathIntUpdateRangedInputBindings, SystemInstance->FastPathFloatUpdateRangedInputBindings);
-				}
-
-				EmitterMap.SpawnPerUnit.SetNumUninitialized(Emitter->SpawnPerUnit.Num());
-				for (int32 SpawnPerUnitIndex = 0; SpawnPerUnitIndex < Emitter->SpawnPerUnit.Num(); ++SpawnPerUnitIndex)
-				{
-					EmitterMap.SpawnPerUnit[SpawnPerUnitIndex].Init(
-						Emitter, SpawnPerUnitIndex, SystemInstance->GetInstanceParameters(),
-						SystemInstance->FastPathIntUserParameterInputBindings, SystemInstance->FastPathFloatUserParameterInputBindings,
-						SystemInstance->FastPathIntUpdateRangedInputBindings, SystemInstance->FastPathFloatUpdateRangedInputBindings);
-				}
-
-				EmitterMap.SpawnBurst_Instantaneous.SetNumUninitialized(Emitter->SpawnBurstInstantaneous.Num());
-				for (int32 SpawnBurstInstantaneousIndex = 0; SpawnBurstInstantaneousIndex < Emitter->SpawnBurstInstantaneous.Num(); ++SpawnBurstInstantaneousIndex)
-				{
-					EmitterMap.SpawnBurst_Instantaneous[SpawnBurstInstantaneousIndex].Init(
-						Emitter, SpawnBurstInstantaneousIndex, SystemInstance->GetInstanceParameters(),
-						SystemInstance->FastPathIntUserParameterInputBindings, SystemInstance->FastPathFloatUserParameterInputBindings,
-						SystemInstance->FastPathIntUpdateRangedInputBindings, SystemInstance->FastPathFloatUpdateRangedInputBindings);
-				}
-			}
-		}
-	}
-
-	// UpdateSystemInstances
-	const UNiagaraSystem* System = GetSystem();
-	auto UpdateSystemInstance = [System](FNiagaraSystemInstance* SystemInstance)
-	{
-		FNiagaraSystemFastPath::FParamMap0& SystemMap = SystemInstance->FastPathMap;
-		FNiagaraSystemFastPath::SetUpdateMapDefaults(SystemMap);
-		SystemInstance->TickFastPathBindings();
-
-		if (System->SystemScalability.bUseSystemScalability)
-		{
-			FNiagaraSystemFastPath::Module_SystemScalability(System->SystemScalability, SystemMap);
-		}
-
-		FNiagaraSystemFastPath::Module_SystemLifeCycle(System->SystemLifeCycle, SystemMap);
-
-		TArray<TSharedRef<FNiagaraEmitterInstance, ESPMode::ThreadSafe>>& EmitterInstances = SystemInstance->GetEmitters();
-		for (TSharedRef<FNiagaraEmitterInstance, ESPMode::ThreadSafe>& EmitterInstance : EmitterInstances)
-		{
-			if (EmitterInstance->GetExecutionState() == ENiagaraExecutionState::Disabled)
-			{
-				continue;
-			}
-
-			FNiagaraEmitterFastPath::FParamMap0& EmitterMap = EmitterInstance->GetFastPathMap();
-			FNiagaraEmitterFastPath::SetUpdateMapDefaults(EmitterMap);
-			EmitterMap.System = &SystemMap.System;
-
-			const UNiagaraEmitter* Emitter = EmitterInstance->GetCachedEmitter();
-			if (Emitter->EmitterScalability.bUseEmitterScalability)
-			{
-				FNiagaraEmitterFastPath::Module_EmitterScalability(Emitter->EmitterScalability, EmitterMap);
-			}
-
-			FNiagaraEmitterFastPath::Module_EmitterLifeCycle(Emitter->EmitterLifeCycle, EmitterMap);
-
-			if (Emitter->SpawnRate.Num() > 0)
-			{
-				FNiagaraEmitterFastPath::Module_SpawnRate(EmitterMap);
-			}
-
-			if (Emitter->SpawnPerUnit.Num() > 0)
-			{
-				FNiagaraEmitterFastPath::Module_SpawnPerUnit(EmitterMap);
-			}
-
-			if (Emitter->SpawnBurstInstantaneous.Num() > 0)
-			{
-				FNiagaraEmitterFastPath::Module_SpawnBurstInstantaneous(EmitterMap);
-			}
-		}
-	};
-
-	// Run update on current instances.
-	for (int32 CurrentIndex = 0; CurrentIndex < CurrentNumInstances; ++CurrentIndex)
-	{
-		FNiagaraSystemInstance* CurrentInstance = Context.Instances[CurrentIndex];
-		UpdateSystemInstance(CurrentInstance);
-	}
-
-	// Run updated on spawned instances with a small delta time.
-	for (int32 SpawnIndex = CurrentNumInstances; SpawnIndex < Context.Instances.Num(); ++SpawnIndex)
-	{
-		FNiagaraSystemInstance* SpawnedInstance = Context.Instances[SpawnIndex];
-		TArray<TSharedRef<FNiagaraEmitterInstance, ESPMode::ThreadSafe>>& EmitterInstances = SpawnedInstance->GetEmitters();
-		for (int32 EmitterIndex = 0; EmitterIndex < EmitterInstances.Num(); ++EmitterIndex)
-		{
-			TSharedRef<FNiagaraEmitterInstance, ESPMode::ThreadSafe>& EmitterInstance = EmitterInstances[EmitterIndex];
-			if (EmitterInstance->GetExecutionState() == ENiagaraExecutionState::Disabled)
-			{
-				continue;
-			}
-
-			EmitterInstance->GetFastPathMap().Engine.DeltaTime = 0.0001f;
-		}
-		UpdateSystemInstance(SpawnedInstance);
-	}
-
-	// TransferSystemSimResults
-	for (FNiagaraSystemInstance* SystemInstance : Context.Instances)
-	{
-		SystemInstance->SetActualExecutionState(SystemInstance->FastPathMap.System.ExecutionState);
-		if (SystemInstance->IsDisabled() == false)
-		{
-			TArray<TSharedRef<FNiagaraEmitterInstance, ESPMode::ThreadSafe>>& EmitterInstances = SystemInstance->GetEmitters();
-			for (TSharedRef<FNiagaraEmitterInstance, ESPMode::ThreadSafe>& EmitterInstance : EmitterInstances)
-			{
-				if (EmitterInstance->GetExecutionState() == ENiagaraExecutionState::Disabled)
-				{
-					continue;
-				}
-
-				FNiagaraEmitterFastPath::FParamMap0& EmitterMap = EmitterInstance->GetFastPathMap();
-
-				EmitterInstance->SetExecutionState(EmitterMap.Emitter.ExecutionState);
-
-				TArray<FNiagaraSpawnInfo>& EmitterInstSpawnInfos = EmitterInstance->GetSpawnInfo();
-
-				int32 SpawnInfoIndex = 0;
-
-				for (const FNiagaraEmitterFastPath::FParamMap0_Emitter_SpawnRate& SpawnRateOutput : EmitterMap.Emitter.SpawnRate)
-				{
-					EmitterInstSpawnInfos[SpawnInfoIndex] = SpawnRateOutput.SpawnOutputInfo;
-					SpawnInfoIndex++;
-				}
-
-				for (const FNiagaraEmitterFastPath::FParamMap0_Emitter_SpawnPerUnit& SpawnPerUnitOutput : EmitterMap.Emitter.SpawnPerUnit)
-				{
-					EmitterInstSpawnInfos[SpawnInfoIndex] = SpawnPerUnitOutput.SpawnOutputInfo;
-					SpawnInfoIndex++;
-				}
-
-				for (const FNiagaraEmitterFastPath::FParamMap0_Emitter_SpawnBurst_Instantaneous& SpawnBurstInstantaneousOutput : EmitterMap.Emitter.SpawnBurst_Instantaneous)
-				{
-					EmitterInstSpawnInfos[SpawnInfoIndex] = SpawnBurstInstantaneousOutput.SpawnBurst;
-					SpawnInfoIndex++;
-				}
-
-				EmitterInstance->TickFastPathAttributeBindings();
-			}
-		}
-	}
-}
-
 void FNiagaraSystemSimulation::SetupParameters_GameThread(float DeltaSeconds)
 {
 	check(IsInGameThread());
@@ -1390,13 +1160,9 @@ void FNiagaraSystemSimulation::PrepareForSystemSimulate(FNiagaraSystemSimulation
 	SpawnInstanceParameterDataSet.Allocate(NumInstances);
 	UpdateInstanceParameterDataSet.Allocate(NumInstances);
 
-	bool bUsingFastPath = GetSystem()->FastPathMode != ENiagaraFastPathMode::ScriptVMOnly;
-	if (bUsingFastPath == false)
+	for (int32 EmitterIdx = 0; EmitterIdx < Context.System->GetNumEmitters(); ++EmitterIdx)
 	{
-		for (int32 EmitterIdx = 0; EmitterIdx < Context.System->GetNumEmitters(); ++EmitterIdx)
-		{
-			EmitterExecutionStateAccessors[EmitterIdx].InitForAccess();
-		}
+		EmitterExecutionStateAccessors[EmitterIdx].InitForAccess();
 	}
 
 	//Tick instance parameters and transfer any needed into the system simulation dataset.
@@ -1416,16 +1182,13 @@ void FNiagaraSystemSimulation::PrepareForSystemSimulate(FNiagaraSystemSimulation
 		//TODO: Find good way to check that we're not using any instance parameter data interfaces in the system scripts here.
 		//In that case we need to solo and will never get here.
 
-		if (bUsingFastPath == false)
+		TArray<TSharedRef<FNiagaraEmitterInstance, ESPMode::ThreadSafe>>& Emitters = Inst->GetEmitters();
+		for (int32 EmitterIdx = 0; EmitterIdx < Emitters.Num(); ++EmitterIdx)
 		{
-			TArray<TSharedRef<FNiagaraEmitterInstance, ESPMode::ThreadSafe>>& Emitters = Inst->GetEmitters();
-			for (int32 EmitterIdx = 0; EmitterIdx < Emitters.Num(); ++EmitterIdx)
+			FNiagaraEmitterInstance& EmitterInst = Emitters[EmitterIdx].Get();
+			if (EmitterExecutionStateAccessors.Num() > EmitterIdx&& EmitterExecutionStateAccessors[EmitterIdx].IsValidForWrite())
 			{
-				FNiagaraEmitterInstance& EmitterInst = Emitters[EmitterIdx].Get();
-				if (EmitterExecutionStateAccessors.Num() > EmitterIdx && EmitterExecutionStateAccessors[EmitterIdx].IsValidForWrite())
-				{
-					EmitterExecutionStateAccessors[EmitterIdx].Set(SystemIndex, (int32)EmitterInst.GetExecutionState());
-				}
+				EmitterExecutionStateAccessors[EmitterIdx].Set(SystemIndex, (int32)EmitterInst.GetExecutionState());
 			}
 		}
 	};
@@ -1604,20 +1367,16 @@ void FNiagaraSystemSimulation::TransferSystemSimResults(FNiagaraSystemSimulation
 		return;
 	}
 
-	bool bIsUsingFastPath = GetSystem()->FastPathMode != ENiagaraFastPathMode::ScriptVMOnly;
-	if (bIsUsingFastPath == false)
+	SystemExecutionStateAccessor.SetDataSet(Context.DataSet);
+	SystemExecutionStateAccessor.InitForAccess();
+	for (int32 EmitterIdx = 0; EmitterIdx < Context.System->GetNumEmitters(); ++EmitterIdx)
 	{
-		SystemExecutionStateAccessor.SetDataSet(Context.DataSet);
-		SystemExecutionStateAccessor.InitForAccess();
-		for (int32 EmitterIdx = 0; EmitterIdx < Context.System->GetNumEmitters(); ++EmitterIdx)
+		EmitterExecutionStateAccessors[EmitterIdx].SetDataSet(Context.DataSet);
+		EmitterExecutionStateAccessors[EmitterIdx].InitForAccess();
+		for (int32 SpawnInfoIdx = 0; SpawnInfoIdx < EmitterSpawnInfoAccessors[EmitterIdx].Num(); ++SpawnInfoIdx)
 		{
-			EmitterExecutionStateAccessors[EmitterIdx].SetDataSet(Context.DataSet);
-			EmitterExecutionStateAccessors[EmitterIdx].InitForAccess();
-			for (int32 SpawnInfoIdx = 0; SpawnInfoIdx < EmitterSpawnInfoAccessors[EmitterIdx].Num(); ++SpawnInfoIdx)
-			{
-				EmitterSpawnInfoAccessors[EmitterIdx][SpawnInfoIdx].SetDataSet(Context.DataSet);
-				EmitterSpawnInfoAccessors[EmitterIdx][SpawnInfoIdx].InitForAccess();
-			}
+			EmitterSpawnInfoAccessors[EmitterIdx][SpawnInfoIdx].SetDataSet(Context.DataSet);
+			EmitterSpawnInfoAccessors[EmitterIdx][SpawnInfoIdx].InitForAccess();
 		}
 	}
 
@@ -1625,12 +1384,9 @@ void FNiagaraSystemSimulation::TransferSystemSimResults(FNiagaraSystemSimulation
 	{
 		FNiagaraSystemInstance* SystemInst = Context.Instances[SystemIndex];
 
-		if (bIsUsingFastPath == false)
-		{
-			//Apply the systems requested execution state to it's actual execution state.
-			ENiagaraExecutionState ExecutionState = (ENiagaraExecutionState)SystemExecutionStateAccessor.GetSafe(SystemIndex, (int32)ENiagaraExecutionState::Disabled);
-			SystemInst->SetActualExecutionState(ExecutionState);
-		}
+		//Apply the systems requested execution state to it's actual execution state.
+		ENiagaraExecutionState ExecutionState = (ENiagaraExecutionState)SystemExecutionStateAccessor.GetSafe(SystemIndex, (int32)ENiagaraExecutionState::Disabled);
+		SystemInst->SetActualExecutionState(ExecutionState);
 
 		if (!SystemInst->IsDisabled())
 		{
@@ -1648,22 +1404,19 @@ void FNiagaraSystemSimulation::TransferSystemSimResults(FNiagaraSystemSimulation
 
 				check(Emitters.Num() > EmitterIdx);
 
-				if (bIsUsingFastPath == false)
-				{
-					ENiagaraExecutionState State = (ENiagaraExecutionState)EmitterExecutionStateAccessors[EmitterIdx].GetSafe(SystemIndex, (int32)ENiagaraExecutionState::Disabled);
-					EmitterInst.SetExecutionState(State);
+				ENiagaraExecutionState State = (ENiagaraExecutionState)EmitterExecutionStateAccessors[EmitterIdx].GetSafe(SystemIndex, (int32)ENiagaraExecutionState::Disabled);
+				EmitterInst.SetExecutionState(State);
 
-					TArray<FNiagaraSpawnInfo>& EmitterInstSpawnInfos = EmitterInst.GetSpawnInfo();
-					for (int32 SpawnInfoIdx = 0; SpawnInfoIdx < EmitterSpawnInfoAccessors[EmitterIdx].Num(); ++SpawnInfoIdx)
+				TArray<FNiagaraSpawnInfo>& EmitterInstSpawnInfos = EmitterInst.GetSpawnInfo();
+				for (int32 SpawnInfoIdx = 0; SpawnInfoIdx < EmitterSpawnInfoAccessors[EmitterIdx].Num(); ++SpawnInfoIdx)
+				{
+					if (SpawnInfoIdx < EmitterInstSpawnInfos.Num())
+					{
+						EmitterInstSpawnInfos[SpawnInfoIdx] = EmitterSpawnInfoAccessors[EmitterIdx][SpawnInfoIdx].Get(SystemIndex);
+					}
+					else
 					{
-						if (SpawnInfoIdx < EmitterInstSpawnInfos.Num())
-						{
-							EmitterInstSpawnInfos[SpawnInfoIdx] = EmitterSpawnInfoAccessors[EmitterIdx][SpawnInfoIdx].Get(SystemIndex);
-						}
-						else
-						{
-							ensure(SpawnInfoIdx < EmitterInstSpawnInfos.Num());
-						}
+						ensure(SpawnInfoIdx < EmitterInstSpawnInfos.Num());
 					}
 				}
 
diff --git a/Engine/Plugins/FX/Niagara/Source/Niagara/Public/NiagaraSystemFastPath.h b/Engine/Plugins/FX/Niagara/Source/Niagara/Public/NiagaraSystemFastPath.h
deleted file mode 100644
diff --git a/Engine/Plugins/FX/Niagara/Source/Niagara/Public/NiagaraSystemInstance.h b/Engine/Plugins/FX/Niagara/Source/Niagara/Public/NiagaraSystemInstance.h
--- a/Engine/Plugins/FX/Niagara/Source/Niagara/Public/NiagaraSystemInstance.h
+++ b/Engine/Plugins/FX/Niagara/Source/Niagara/Public/NiagaraSystemInstance.h
@@ -10,7 +10,6 @@
 #include "Templates/UniquePtr.h"
 #include "NiagaraCommon.h"
 #include "NiagaraDataInterface.h"
-#include "NiagaraSystemFastPath.h"
 
 class FNiagaraWorldManager;
 class UNiagaraComponent;
@@ -272,16 +271,10 @@ public:
 
 	void TickInstanceParameters_Concurrent();
 
-	void TickFastPathBindings();
-
-	void ResetFastPathBindings();
-
 	FNiagaraDataSet* CreateEventDataSet(FName EmitterName, FName EventName);
 	FNiagaraDataSet* GetEventDataSet(FName EmitterName, FName EventName) const;
 	void ClearEventDataSets();
 
-	FNiagaraSystemFastPath::FParamMap0& GetFastPathMap() { return FastPathMap; }
-
 	FORCEINLINE void SetLODDistance(float InLODDistance, float InMaxLODDistance);
 
 	const FString& GetCrashReporterTag()const;
@@ -471,13 +464,6 @@ public:
 	};
 
 	FInstanceParameters GatheredInstanceParameters;
-
-	FNiagaraSystemFastPath::FParamMap0 FastPathMap;
-
-	TArray<TNiagaraFastPathRangedInputBinding<int32>> FastPathIntUpdateRangedInputBindings;
-	TArray<TNiagaraFastPathRangedInputBinding<float>> FastPathFloatUpdateRangedInputBindings;
-	TArray<TNiagaraFastPathUserParameterInputBinding<int32>> FastPathIntUserParameterInputBindings;
-	TArray<TNiagaraFastPathUserParameterInputBinding<float>> FastPathFloatUserParameterInputBindings;
 };
 
 FORCEINLINE void FNiagaraSystemInstance::SetLODDistance(float InLODDistance, float InMaxLODDistance)
diff --git a/Engine/Plugins/FX/Niagara/Source/Niagara/Public/NiagaraSystemSimulation.h b/Engine/Plugins/FX/Niagara/Source/Niagara/Public/NiagaraSystemSimulation.h
--- a/Engine/Plugins/FX/Niagara/Source/Niagara/Public/NiagaraSystemSimulation.h
+++ b/Engine/Plugins/FX/Niagara/Source/Niagara/Public/NiagaraSystemSimulation.h
@@ -249,8 +249,6 @@ public:
 	/** Second phase of system sim tick that can run on any thread. */
 	void Tick_Concurrent(FNiagaraSystemSimulationTickContext& Context);
 
-	void TickFastPath(FNiagaraSystemSimulationTickContext& Context);
-
 	/** Update TickGroups for pending instances and execute tick group promotions. */
 	void UpdateTickGroups_GameThread();
 	/** Spawn any pending instances, assumes that you have update tick groups ahead of time. */
diff --git a/Engine/Plugins/FX/Niagara/Source/NiagaraEditor/Private/Customizations/NiagaraSystemDetails.cpp b/Engine/Plugins/FX/Niagara/Source/NiagaraEditor/Private/Customizations/NiagaraSystemDetails.cpp
deleted file mode 100644
diff --git a/Engine/Plugins/FX/Niagara/Source/NiagaraEditor/Private/Customizations/NiagaraSystemDetails.h b/Engine/Plugins/FX/Niagara/Source/NiagaraEditor/Private/Customizations/NiagaraSystemDetails.h
deleted file mode 100644
diff --git a/Engine/Plugins/FX/Niagara/Source/NiagaraEditor/Private/EdGraphSchema_Niagara.cpp b/Engine/Plugins/FX/Niagara/Source/NiagaraEditor/Private/EdGraphSchema_Niagara.cpp
--- a/Engine/Plugins/FX/Niagara/Source/NiagaraEditor/Private/EdGraphSchema_Niagara.cpp
+++ b/Engine/Plugins/FX/Niagara/Source/NiagaraEditor/Private/EdGraphSchema_Niagara.cpp
@@ -426,25 +426,6 @@ TArray<TSharedPtr<FNiagaraSchemaAction_NewNode> > UEdGraphSchema_Niagara::GetGra
 		{
 			AddScriptFunctionAction(LOCTEXT("Function Menu Title", "Functions"), FunctionScriptAsset);
 		}
-
-		// Insert custom opcodes here...
-		const TArray<FNiagaraFunctionSignature>& VMFastPathOps = UNiagaraFunctionLibrary::GetVectorVMFastPathOps();
-		for (const FNiagaraFunctionSignature& Sig : VMFastPathOps)
-		{
-			UNiagaraNodeFunctionCall* FunctionCallNode = NewObject<UNiagaraNodeFunctionCall>(OwnerOfTemporaries);
-			
-			FunctionCallNode->Signature = Sig;
-			FText AssetDesc;
-			FText Keywords;
-			
-			FString DisplayNameString = FName::NameToDisplayString(Sig.Name.ToString(), false);
-
-			const FText MenuDesc = FText::FromString(DisplayNameString);
-			const FText TooltipDesc = MenuDesc;
-
-			TSharedPtr<FNiagaraSchemaAction_NewNode> FunctionCallAction = AddNewNodeAction(NewActions, LOCTEXT("Function Menu Title", "Functions"), MenuDesc, *DisplayNameString, TooltipDesc, Keywords);
-			FunctionCallAction->NodeTemplate = FunctionCallNode;
-		}
 	}
 
 	//Add modules
diff --git a/Engine/Plugins/FX/Niagara/Source/NiagaraEditor/Private/NiagaraEditorModule.cpp b/Engine/Plugins/FX/Niagara/Source/NiagaraEditor/Private/NiagaraEditorModule.cpp
--- a/Engine/Plugins/FX/Niagara/Source/NiagaraEditor/Private/NiagaraEditorModule.cpp
+++ b/Engine/Plugins/FX/Niagara/Source/NiagaraEditor/Private/NiagaraEditorModule.cpp
@@ -148,14 +148,6 @@ TArray<TPair<FName, FNiagaraParameterScopeInfo>> FNiagaraEditorModule::Registere
 
 EAssetTypeCategories::Type FNiagaraEditorModule::NiagaraAssetCategory;
 
-int32 GbShowFastPathOptions = 0;
-static FAutoConsoleVariableRef CVarShowFastPathOptions(
-	TEXT("fx.Niagara.ShowFastPathOptions"),
-	GbShowFastPathOptions,
-	TEXT("If > 0 the experimental fast path options will be shown in the system and emitter properties in the niagara system editor.\n"),
-	ECVF_Default
-);
-
 const FNiagaraParameterScopeInfo* FNiagaraEditorModule::FindParameterScopeInfo(const FName& ParameterScopeInfoName)
 {
 	auto FindPredicate = [ParameterScopeInfoName](const TPair<FName, FNiagaraParameterScopeInfo>& ScopeInfoPair) {return ScopeInfoPair.Key == ParameterScopeInfoName; };
diff --git a/Engine/Plugins/FX/Niagara/Source/NiagaraEditor/Private/NiagaraEmitterDetailsCustomization.cpp b/Engine/Plugins/FX/Niagara/Source/NiagaraEditor/Private/NiagaraEmitterDetailsCustomization.cpp
--- a/Engine/Plugins/FX/Niagara/Source/NiagaraEditor/Private/NiagaraEmitterDetailsCustomization.cpp
+++ b/Engine/Plugins/FX/Niagara/Source/NiagaraEditor/Private/NiagaraEmitterDetailsCustomization.cpp
@@ -16,10 +16,5 @@ void FNiagaraEmitterDetails::CustomizeDetails(IDetailLayoutBuilder& InDetailLayo
 {
 	TSharedPtr<IPropertyHandle> EventHandlersPropertyHandle = InDetailLayout.GetProperty(UNiagaraEmitter::PrivateMemberNames::EventHandlerScriptProps);
 	EventHandlersPropertyHandle->MarkHiddenByCustomization();
-
-	if (GbShowFastPathOptions <= 0)
-	{
-		InDetailLayout.EditCategory("Script Fast Path").SetCategoryVisibility(false);
-	}
 }
 
diff --git a/Engine/Plugins/FX/Niagara/Source/NiagaraEditor/Private/ViewModels/Stack/NiagaraStackSystemPropertiesItem.cpp b/Engine/Plugins/FX/Niagara/Source/NiagaraEditor/Private/ViewModels/Stack/NiagaraStackSystemPropertiesItem.cpp
--- a/Engine/Plugins/FX/Niagara/Source/NiagaraEditor/Private/ViewModels/Stack/NiagaraStackSystemPropertiesItem.cpp
+++ b/Engine/Plugins/FX/Niagara/Source/NiagaraEditor/Private/ViewModels/Stack/NiagaraStackSystemPropertiesItem.cpp
@@ -4,7 +4,6 @@
 #include "ViewModels/Stack/NiagaraStackObject.h"
 #include "NiagaraSystem.h"
 #include "ViewModels/NiagaraSystemViewModel.h"
-#include "Customizations/NiagaraSystemDetails.h"
 
 #define LOCTEXT_NAMESPACE "UNiagaraStackSystemItemGroup"
 
@@ -36,7 +35,6 @@ void UNiagaraStackSystemPropertiesItem::RefreshChildrenInternal(const TArray<UNi
 		SystemObject = NewObject<UNiagaraStackObject>(this);
 		FRequiredEntryData RequiredEntryData(GetSystemViewModel(), GetEmitterViewModel(), FExecutionCategoryNames::System, NAME_None, GetStackEditorData());
 		SystemObject->Initialize(RequiredEntryData, System.Get(), GetStackEditorDataKey());
-		SystemObject->RegisterInstancedCustomPropertyLayout(UNiagaraSystem::StaticClass(), FOnGetDetailCustomizationInstance::CreateStatic(&FNiagaraSystemDetails::MakeInstance));
 	}
 
 	NewChildren.Add(SystemObject);
diff --git a/Engine/Plugins/FX/Niagara/Source/NiagaraEditor/Public/NiagaraEditorModule.h b/Engine/Plugins/FX/Niagara/Source/NiagaraEditor/Public/NiagaraEditorModule.h
--- a/Engine/Plugins/FX/Niagara/Source/NiagaraEditor/Public/NiagaraEditorModule.h
+++ b/Engine/Plugins/FX/Niagara/Source/NiagaraEditor/Public/NiagaraEditorModule.h
@@ -44,8 +44,6 @@ public:
 	virtual FLinearColor GetColorForParameterScope(ENiagaraParameterScope ParameterScope) const = 0;
 };
 
-extern int32 GbShowFastPathOptions;
-
 /** Niagara Editor module */
 class FNiagaraEditorModule : public IModuleInterface,
 	public IHasMenuExtensibility, public IHasToolBarExtensibility, public FGCObject
