Memory savings allowing execution instances to share the variable definitions used within FNiagaraParameterStores
-Splits FNiagaraScriptExecutionParmeterStore into a new class FNiagaraScriptInstanceParameterStore which handles referencing the data from the source parameter store on the script
-Accessing the variables is now done through a ArrayView (which is overriden by child classes) instead of using the Array directly
-Saves about 25% of our memory when dealing with large instance counts

diff --git a/Engine/Plugins/FX/Niagara/Source/Niagara/Classes/NiagaraDataSet.h b/Engine/Plugins/FX/Niagara/Source/Niagara/Classes/NiagaraDataSet.h
--- a/Engine/Plugins/FX/Niagara/Source/Niagara/Classes/NiagaraDataSet.h
+++ b/Engine/Plugins/FX/Niagara/Source/Niagara/Classes/NiagaraDataSet.h
@@ -308,7 +308,7 @@ public:
 
 	FORCEINLINE void Init(const FNiagaraDataSetCompiledData* InDataSetCompiledData)
 	{
-		CompiledData = InDataSetCompiledData != nullptr ? InDataSetCompiledData : &FNiagaraDataSetCompiledData::DummyCompiledData;
+		CompiledData.Init(InDataSetCompiledData != nullptr ? InDataSetCompiledData : &FNiagaraDataSetCompiledData::DummyCompiledData);
 		bInitialized = true;
 		Reset();
 	}
@@ -404,7 +404,7 @@ private:
 #endif
 	}
 
-	const FNiagaraDataSetCompiledData* CompiledData;
+	FNiagaraCompiledDataReference<FNiagaraDataSetCompiledData> CompiledData;
 
 	/** Table of free IDs available to allocate next tick. */
 	TArray<int32> FreeIDsTable;
@@ -489,7 +489,7 @@ protected:
 
 	FNiagaraDataSet* DataSet;
 	const FNiagaraVariableLayoutInfo* VarLayout;
-	FNiagaraVariable Var;
+	FNiagaraVariableBase Var;
 };
 
 template<typename T>
diff --git a/Engine/Plugins/FX/Niagara/Source/Niagara/Classes/NiagaraScript.h b/Engine/Plugins/FX/Niagara/Source/Niagara/Classes/NiagaraScript.h
--- a/Engine/Plugins/FX/Niagara/Source/Niagara/Classes/NiagaraScript.h
+++ b/Engine/Plugins/FX/Niagara/Source/Niagara/Classes/NiagaraScript.h
@@ -689,7 +689,7 @@ public:
 	
 	virtual ~UNiagaraScript();
 
-	FNiagaraScriptExecutionParameterStore* GetExecutionReadyParameterStore(ENiagaraSimTarget SimTarget);
+	const FNiagaraScriptExecutionParameterStore* GetExecutionReadyParameterStore(ENiagaraSimTarget SimTarget);
 	void InvalidateExecutionReadyParameterStores();
 
 private:
@@ -782,4 +782,4 @@ private:
 	TArray<FNiagaraScriptDataInterfaceInfo> CachedDefaultDataInterfaces;
 
 	static UNiagaraDataInterface* CopyDataInterface(UNiagaraDataInterface* Src, UObject* Owner);
-};
+};
\ No newline at end of file
diff --git a/Engine/Plugins/FX/Niagara/Source/Niagara/Classes/NiagaraScriptExecutionContext.h b/Engine/Plugins/FX/Niagara/Source/Niagara/Classes/NiagaraScriptExecutionContext.h
--- a/Engine/Plugins/FX/Niagara/Source/Niagara/Classes/NiagaraScriptExecutionContext.h
+++ b/Engine/Plugins/FX/Niagara/Source/Niagara/Classes/NiagaraScriptExecutionContext.h
@@ -156,7 +156,7 @@ struct FNiagaraScriptExecutionContext
 	TArray<void*> DataInterfaceInstDataTable;
 
 	/** Parameter store. Contains all data interfaces and a parameter buffer that can be used directly by the VM or GPU. */
-	FNiagaraScriptExecutionParameterStore Parameters;
+	FNiagaraScriptInstanceParameterStore Parameters;
 
 	TArray<FDataSetMeta, TInlineAllocator<4>> DataSetMetaTable;
 
@@ -255,7 +255,7 @@ public:
 
 	//Dynamic state updated either from GT via RT commands or from the RT side sim code itself.
 	//TArray<uint8, TAlignedHeapAllocator<16>> ParamData_RT;		// RT side copy of the parameter data
-	FNiagaraScriptExecutionParameterStore CombinedParamStore;
+	FNiagaraScriptInstanceParameterStore CombinedParamStore;
 #if DO_CHECK
 	TArray< FString >  DIClassNames;
 #endif
diff --git a/Engine/Plugins/FX/Niagara/Source/Niagara/Private/NiagaraCommon.cpp b/Engine/Plugins/FX/Niagara/Source/Niagara/Private/NiagaraCommon.cpp
--- a/Engine/Plugins/FX/Niagara/Source/Niagara/Private/NiagaraCommon.cpp
+++ b/Engine/Plugins/FX/Niagara/Source/Niagara/Private/NiagaraCommon.cpp
@@ -439,13 +439,13 @@ void FNiagaraUtilities::PrepareRapidIterationParameters(const TArray<UNiagaraScr
 		FNiagaraParameterStore& PreparedParameterStore = It.Value();
 
 		bool bOverwriteParameters = false;
-		if (Script->RapidIterationParameters.GetNumParameters() != PreparedParameterStore.GetNumParameters())
+		if (Script->RapidIterationParameters.ParameterVariables.Num() != PreparedParameterStore.ParameterVariables.Num())
 		{
 			bOverwriteParameters = true;
 		}
 		else
 		{
-			for (const FNiagaraVariableWithOffset& ParamWithOffset : Script->RapidIterationParameters.GetSortedParameterOffsets())
+			for (const FNiagaraVariableWithOffset& ParamWithOffset : Script->RapidIterationParameters.ParameterVariables)
 			{
 				const FNiagaraVariable& SourceParameter = ParamWithOffset;
 				const int32 SourceOffset = ParamWithOffset.Offset;
diff --git a/Engine/Plugins/FX/Niagara/Source/Niagara/Private/NiagaraDataSet.cpp b/Engine/Plugins/FX/Niagara/Source/Niagara/Private/NiagaraDataSet.cpp
--- a/Engine/Plugins/FX/Niagara/Source/Niagara/Private/NiagaraDataSet.cpp
+++ b/Engine/Plugins/FX/Niagara/Source/Niagara/Private/NiagaraDataSet.cpp
@@ -76,8 +76,7 @@ static FAutoConsoleVariableRef CVarNiagaraGPUDataBufferBufferSlack(
 );
 
 FNiagaraDataSet::FNiagaraDataSet()
-	: CompiledData(&FNiagaraDataSetCompiledData::DummyCompiledData)
-	, NumFreeIDs(0)
+	: NumFreeIDs(0)
 	, MaxUsedID(0)
 	, IDAcquireTag(0)
 	, GPUNumAllocatedIDs(0)
@@ -85,6 +84,7 @@ FNiagaraDataSet::FNiagaraDataSet()
 	, DestinationData(nullptr)
 	, bInitialized(false)
 {
+	CompiledData.Init(&FNiagaraDataSetCompiledData::DummyCompiledData);
 }
 
 FNiagaraDataSet::~FNiagaraDataSet()
diff --git a/Engine/Plugins/FX/Niagara/Source/Niagara/Private/NiagaraParameterStore.cpp b/Engine/Plugins/FX/Niagara/Source/Niagara/Private/NiagaraParameterStore.cpp
--- a/Engine/Plugins/FX/Niagara/Source/Niagara/Private/NiagaraParameterStore.cpp
+++ b/Engine/Plugins/FX/Niagara/Source/Niagara/Private/NiagaraParameterStore.cpp
@@ -105,20 +105,21 @@ static FAutoConsoleVariableRef CVarNiagaraAllowQuickSortedParameterOffetsCopy(
 	ECVF_Scalability
 );
 
-namespace
+void FNiagaraParameterStore::CopySortedParameterOffsets(TArrayView<FNiagaraVariableWithOffset> Src)
 {
-	void CopySortedParameterOffsets(TArray<FNiagaraVariableWithOffset>& Dest, const TArray<FNiagaraVariableWithOffset>& Src)
+	if (GNiagaraAllowQuickSortedParameterOffetsCopy)
 	{
-		if (GNiagaraAllowQuickSortedParameterOffetsCopy)
-		{
-			Dest.SetNumUninitialized(Src.Num());
-			FMemory::Memcpy(Dest.GetData(), Src.GetData(), Dest.GetTypeSize() * Dest.Num());
-		}
-		else
-		{
-			Dest = Src;
-		}
+		const int32 VariableCount = Src.Num();
+
+		SortedParameterOffsets.SetNumUninitialized(VariableCount);
+		FMemory::Memcpy(SortedParameterOffsets.GetData(), Src.GetData(), SortedParameterOffsets.GetTypeSize() * VariableCount);
+	}
+	else
+	{
+		SortedParameterOffsets = TArray<FNiagaraVariableWithOffset>(Src.GetData(), Src.Num());
 	}
+
+	InternalStorageChanged();
 }
 
 //////////////////////////////////////////////////////////////////////////
@@ -153,7 +154,7 @@ FNiagaraParameterStore& FNiagaraParameterStore::operator=(const FNiagaraParamete
 #if WITH_EDITORONLY_DATA
 	ParameterOffsets = Other.ParameterOffsets;
 #endif // WITH_EDITORONLY_DATA
-	CopySortedParameterOffsets(SortedParameterOffsets, Other.SortedParameterOffsets);
+	CopySortedParameterOffsets(Other.ParameterVariables);
 	DEC_MEMORY_STAT_BY(STAT_NiagaraParamStoreMemory, ParameterData.GetAllocatedSize());
 	ParameterData = Other.ParameterData;
 	INC_MEMORY_STAT_BY(STAT_NiagaraParamStoreMemory, ParameterData.GetAllocatedSize());
@@ -199,20 +200,20 @@ void FNiagaraParameterStore::Bind(FNiagaraParameterStore* DestStore, const FNiag
 template <typename TVisitor>
 void FNiagaraParameterStoreBinding::MatchParameters(FNiagaraParameterStore* DestStore, FNiagaraParameterStore* SrcStore, TVisitor Visitor)
 {
-	const int32 SrcNum = SrcStore->GetSortedParameterOffsets().Num();
-	const int32 DestNum = DestStore->GetSortedParameterOffsets().Num();
+	TArrayView<FNiagaraVariableWithOffset> SrcParamWithOffsets = SrcStore->ParameterVariables;
+	TArrayView<FNiagaraVariableWithOffset> DestParamWithOffsets = DestStore->ParameterVariables;
+
+	const int32 SrcNum = SrcParamWithOffsets.Num();
+	const int32 DestNum = DestParamWithOffsets.Num();
 	const int32 BinarySearchComplexity = FMath::Min<int32>(SrcNum, DestNum) * FMath::RoundToInt(FMath::Log2((float)FMath::Max<int32>(SrcNum, DestNum)));
 	if (BinarySearchComplexity >= SrcNum + DestNum)
 	{
-		const TArray<FNiagaraVariableWithOffset>& SrcParamWithOffsets = SrcStore->GetSortedParameterOffsets();
-		const TArray<FNiagaraVariableWithOffset>& DestParamWithOffsets = DestStore->GetSortedParameterOffsets();
-
 		int32 SrcIndex = 0;
 		int32 DestIndex = 0;
 		while (SrcIndex < SrcNum && DestIndex < DestNum)
 		{
-			const FNiagaraVariableWithOffset& SrcParamWithOffset = SrcStore->GetSortedParameterOffsets()[SrcIndex];
-			const FNiagaraVariableWithOffset& DestParamWithOffset = DestStore->GetSortedParameterOffsets()[DestIndex];
+			const FNiagaraVariableWithOffset& SrcParamWithOffset = SrcParamWithOffsets[SrcIndex];
+			const FNiagaraVariableWithOffset& DestParamWithOffset = DestParamWithOffsets[DestIndex];
 
 			const int32 CompValue = FNiagaraVariableSearch::Compare(SrcParamWithOffset, DestParamWithOffset);
 			if (CompValue < 0)
@@ -234,16 +235,16 @@ void FNiagaraParameterStoreBinding::MatchParameters(FNiagaraParameterStore* Dest
 
 	}
 	// Process the smaller parameter store the get the least amount of iterations when it is small (often empty).
-	else if (DestStore->GetNumParameters() <= SrcStore->GetNumParameters())
+	else if (DestNum <= SrcNum)
 	{
-		for (const FNiagaraVariableWithOffset& ParamWithOffset : DestStore->GetSortedParameterOffsets())
+		for (const FNiagaraVariableWithOffset& ParamWithOffset : DestParamWithOffsets)
 		{
 			Visitor(ParamWithOffset, SrcStore->IndexOf(ParamWithOffset), ParamWithOffset.Offset);
 		}
 	}
 	else
 	{
-		for (const FNiagaraVariableWithOffset& ParamWithOffset : SrcStore->GetSortedParameterOffsets())
+		for (const FNiagaraVariableWithOffset& ParamWithOffset : SrcParamWithOffsets)
 		{
 			Visitor(ParamWithOffset, ParamWithOffset.Offset, DestStore->IndexOf(ParamWithOffset));
 		}
@@ -375,7 +376,7 @@ bool FNiagaraParameterStore::VerifyBinding(const FNiagaraParameterStore* DestSto
 
 void FNiagaraParameterStore::CheckForNaNs()const
 {
-	for (const FNiagaraVariableWithOffset& Var : SortedParameterOffsets)
+	for (const FNiagaraVariableWithOffset& Var : ParameterVariables)
 	{
 		const int32 Offset = Var.Offset;
 
@@ -438,11 +439,10 @@ void FNiagaraParameterStore::UnbindFromSourceStores()
 
 void FNiagaraParameterStore::DumpParameters(bool bDumpBindings)const
 {
-	TArray<FNiagaraVariable> Vars;
-	GetParameters(Vars);
-	for (FNiagaraVariable& Var : Vars)
+	for (const FNiagaraVariableWithOffset& VariableBase : ParameterVariables)
 	{
-		Var.SetData(GetParameterData_Internal(IndexOf(Var))); // Need to keep the data in sync
+		FNiagaraVariable Var(VariableBase);
+		Var.SetData(GetParameterData(VariableBase.Offset));
 		UE_LOG(LogNiagara, Log, TEXT("Param: %s Offset: %d Type : %s"), *Var.ToString(), IndexOf(Var), *Var.GetType().GetName());
 	}
 
@@ -458,13 +458,14 @@ void FNiagaraParameterStore::DumpParameters(bool bDumpBindings)const
 FString FNiagaraParameterStore::ToString() const
 {
 	FString Value;
-	TArray<FNiagaraVariable> Vars;
-	GetParameters(Vars);
-	for (FNiagaraVariable& Var : Vars)
+
+	for (const FNiagaraVariableWithOffset& VariableBase : ParameterVariables)
 	{
-		Var.SetData(GetParameterData_Internal(IndexOf(Var))); // Need to keep the data in sync
+		FNiagaraVariable Var(VariableBase);
+		Var.SetData(GetParameterData(VariableBase.Offset));
 		Value += FString::Printf(TEXT("Param: %s Offset: %d Type : %s\n"), *Var.ToString(), IndexOf(Var), *Var.GetType().GetName());
 	}
+
 	return Value;
 }
 
@@ -492,23 +493,22 @@ bool FNiagaraParameterStore::AddParameter(const FNiagaraVariable& Param, bool bI
 	check(!ParameterOffsets.Num()); // Migration to SortedParameterOffsets
 #endif
 
-	FNiagaraVariable ParamNoData(Param.GetType(), Param.GetName());
-
 	int32 InsertPos = 0;
-	if (SortedParameterOffsets.Num())
+	if (ParameterVariables.Num())
 	{
-		const bool bAlreadyIn = FNiagaraVariableSearch::Find(SortedParameterOffsets.GetData(), ParamNoData, 0, SortedParameterOffsets.Num(), false /* IgnoreType */, InsertPos);
+		const bool bAlreadyIn = FNiagaraVariableSearch::Find(ParameterVariables.GetData(), Param, 0, ParameterVariables.Num(), false /* IgnoreType */, InsertPos);
 		if (bAlreadyIn)
 		{
 			if (OutOffset)
 			{
-				*OutOffset = SortedParameterOffsets[InsertPos].Offset;
+				*OutOffset = ParameterVariables[InsertPos].Offset;
 			}
 			return false;
 		}
 	}
-	const FNiagaraVariableWithOffset ParamWithOffset(ParamNoData, (int32)INDEX_NONE);
-	int32& Offset = SortedParameterOffsets.EmplaceAt_GetRef(InsertPos, ParamNoData, (int32)INDEX_NONE).Offset;
+
+	int32& Offset = SortedParameterOffsets.EmplaceAt_GetRef(InsertPos, Param, (int32)INDEX_NONE).Offset;
+	InternalStorageChanged();
 
 	if (Param.GetType().IsDataInterface())
 	{
@@ -518,7 +518,6 @@ bool FNiagaraParameterStore::AddParameter(const FNiagaraVariable& Param, bool bI
 	else if (Param.GetType().IsUObject())
 	{
 		Offset = UObjects.AddDefaulted();
-		//UObjects[Offset] = nullptr;
 	}
 	else
 	{
@@ -528,22 +527,21 @@ bool FNiagaraParameterStore::AddParameter(const FNiagaraVariable& Param, bool bI
 		int32 ParamAlignment = Param.GetAlignment();
 		//int32 Offset = AlignArbitrary(ParameterData.Num(), ParamAlignment);//TODO: We need to handle alignment better here. Need to both satisfy CPU and GPU alignment concerns. VM doesn't care but the VM complier needs to be aware. Probably best to have everything adhere to GPU alignment rules.
 		Offset = ParameterData.Num();
-		// Memory must be initialized in order to have deterministic cooking. 
-		// This is because some system parameters never get initialized otherwise (particle count, owner rotation, ...)
-		ParameterData.AddZeroed(ParamSize);
 				
-		INC_MEMORY_STAT_BY(STAT_NiagaraParamStoreMemory, ParameterData.GetAllocatedSize());
-
 		//Temporary to init param data from FNiagaraVariable storage. This will be removed when we change the UNiagaraScript to use a parameter store too.
 		if (Param.IsDataAllocated())
 		{
-			uint8* Dest = GetParameterData_Internal(Offset);
-			const uint8* Src = Param.GetData();
-			if (Dest != Src)
-			{
-				FMemory::Memcpy(Dest, Src, ParamSize);
-			}
+			ParameterData.AddUninitialized(ParamSize);
+			FMemory::Memcpy(GetParameterData_Internal(Offset), Param.GetData(), ParamSize);
 		}
+		else
+		{
+			// Memory must be initialized in order to have deterministic cooking. 
+			// This is because some system parameters never get initialized otherwise (particle count, owner rotation, ...)
+			ParameterData.AddZeroed(ParamSize);
+		}
+
+		INC_MEMORY_STAT_BY(STAT_NiagaraParamStoreMemory, ParameterData.GetAllocatedSize());
 	}
 
 	if (bTriggerRebind)
@@ -562,7 +560,7 @@ bool FNiagaraParameterStore::AddParameter(const FNiagaraVariable& Param, bool bI
 	return true;
 }
 
-bool FNiagaraParameterStore::RemoveParameter(const FNiagaraVariable& ToRemove)
+bool FNiagaraParameterStore::RemoveParameter(const FNiagaraVariableBase& ToRemove)
 {
 #if WITH_EDITORONLY_DATA
 	check(!ParameterOffsets.Num()); // Migration to SortedParameterOffsets
@@ -576,7 +574,7 @@ bool FNiagaraParameterStore::RemoveParameter(const FNiagaraVariable& ToRemove)
 		TArray<uint8> NewData;
 		TArray<UNiagaraDataInterface*> NewInterfaces;
 		TArray<UObject*> NewUObjects;
-		for (const FNiagaraVariableWithOffset& Existing : SortedParameterOffsets)
+		for (const FNiagaraVariableWithOffset& Existing : ParameterVariables)
 		{
 			const FNiagaraVariable& ExistingVar = Existing;
 			const int32 ExistingOffset = Existing.Offset;
@@ -609,7 +607,7 @@ bool FNiagaraParameterStore::RemoveParameter(const FNiagaraVariable& ToRemove)
 
 		DEC_MEMORY_STAT_BY(STAT_NiagaraParamStoreMemory, ParameterData.GetAllocatedSize());
 
-		CopySortedParameterOffsets(SortedParameterOffsets, NewOffsets);
+		CopySortedParameterOffsets(MakeArrayView(NewOffsets));
 		ParameterData = NewData;
 		DataInterfaces = NewInterfaces;
 		UObjects = NewUObjects;
@@ -623,7 +621,7 @@ bool FNiagaraParameterStore::RemoveParameter(const FNiagaraVariable& ToRemove)
 	return false;
 }
 
-void FNiagaraParameterStore::RenameParameter(const FNiagaraVariable& Param, FName NewName)
+void FNiagaraParameterStore::RenameParameter(const FNiagaraVariableBase& Param, FName NewName)
 {
 #if WITH_EDITORONLY_DATA
 	check(!ParameterOffsets.Num()); // Migration to SortedParameterOffsets
@@ -673,12 +671,9 @@ void FNiagaraParameterStore::SanityCheckData(bool bInitInterfaces)
 
 	int32 ParameterDataSize = 0;
 
-	TArray<FNiagaraVariableWithOffset>::TConstIterator It = SortedParameterOffsets.CreateConstIterator();
-	while (It)
+	for (const FNiagaraVariableWithOffset& Parameter : ParameterVariables)
 	{
-		FNiagaraVariable Parameter = *It;
-		int32 SrcIndex = It->Offset;
-		++It;
+		const int32 SrcIndex = Parameter.Offset;
 
 		if (Parameter.IsValid())
 		{
@@ -738,14 +733,11 @@ void FNiagaraParameterStore::SanityCheckData(bool bInitInterfaces)
 	}
 }
 
-void FNiagaraParameterStore::CopyParametersTo(FNiagaraParameterStore& DestStore, bool bOnlyAdd, EDataInterfaceCopyMethod DataInterfaceCopyMethod)
+void FNiagaraParameterStore::CopyParametersTo(FNiagaraParameterStore& DestStore, bool bOnlyAdd, EDataInterfaceCopyMethod DataInterfaceCopyMethod) const
 {
-	TArray<FNiagaraVariableWithOffset>::TConstIterator It = SortedParameterOffsets.CreateConstIterator();
-	while (It)
+	for (const FNiagaraVariableWithOffset& Parameter : ParameterVariables)
 	{
-		FNiagaraVariable Parameter = *It;
-		int32 SrcIndex = It->Offset;
-		++It;
+		int32 SrcIndex = Parameter.Offset;
 
 		if (Parameter.IsValid() == false)
 		{
@@ -812,7 +804,7 @@ void FNiagaraParameterStore::CopyParametersTo(FNiagaraParameterStore& DestStore,
 			{
 				if (ParameterData.Num() != 0)
 				{
-					DestStore.SetParameterData(GetParameterData_Internal(SrcIndex), DestIndex, Parameter.GetSizeInBytes());
+					DestStore.SetParameterData(GetParameterData(SrcIndex), DestIndex, Parameter.GetSizeInBytes());
 				}
 			}
 		}
@@ -821,13 +813,36 @@ void FNiagaraParameterStore::CopyParametersTo(FNiagaraParameterStore& DestStore,
 }
 
 
-FORCEINLINE void FNiagaraParameterStore::SetParameterDataArray(const TArray<uint8>& InParameterDataArray)
+void FNiagaraParameterStore::SetParameterDataArray(const TArray<uint8>& InParameterDataArray, bool bNotifyAsDirty)
 {
 	DEC_MEMORY_STAT_BY(STAT_NiagaraParamStoreMemory, ParameterData.GetAllocatedSize());
 	ParameterData = InParameterDataArray;
 	INC_MEMORY_STAT_BY(STAT_NiagaraParamStoreMemory, ParameterData.GetAllocatedSize());
 
-	OnParameterChange();
+	if (bNotifyAsDirty)
+	{
+		MarkParametersDirty();
+	}
+}
+
+void FNiagaraParameterStore::SetDataInterfaces(const TArray<UNiagaraDataInterface*>& InDataInterfaces, bool bNotifyAsDirty)
+{
+	DataInterfaces = InDataInterfaces;
+
+	if (bNotifyAsDirty)
+	{
+		MarkInterfacesDirty();
+	}
+}
+
+void FNiagaraParameterStore::SetUObjects(const TArray<UObject*>& InUObjects, bool bNotifyAsDirty)
+{
+	UObjects = InUObjects;
+
+	if (bNotifyAsDirty)
+	{
+		MarkUObjectsDirty();
+	}
 }
 
 void FNiagaraParameterStore::InitFromSource(const FNiagaraParameterStore* SrcStore, bool bNotifyAsDirty)
@@ -841,7 +856,7 @@ void FNiagaraParameterStore::InitFromSource(const FNiagaraParameterStore* SrcSto
 #if WITH_EDITORONLY_DATA
 	ParameterOffsets = SrcStore->ParameterOffsets;
 #endif // WITH_EDITORONLY_DATA
-	CopySortedParameterOffsets(SortedParameterOffsets, SrcStore->SortedParameterOffsets);
+	CopySortedParameterOffsets(SrcStore->ParameterVariables);
 	DEC_MEMORY_STAT_BY(STAT_NiagaraParamStoreMemory, ParameterData.GetAllocatedSize());
 	ParameterData = SrcStore->ParameterData;
 	INC_MEMORY_STAT_BY(STAT_NiagaraParamStoreMemory, ParameterData.GetAllocatedSize());
@@ -861,12 +876,9 @@ void FNiagaraParameterStore::InitFromSource(const FNiagaraParameterStore* SrcSto
 
 void FNiagaraParameterStore::RemoveParameters(FNiagaraParameterStore& DestStore)
 {
-	TArray<FNiagaraVariableWithOffset>::TConstIterator It = SortedParameterOffsets.CreateConstIterator();
-	while (It)
+	for (const FNiagaraVariableWithOffset& Parameter : ParameterVariables)
 	{
-		FNiagaraVariable Parameter = *It;
 		DestStore.RemoveParameter(Parameter);
-		++It;
 	}
 }
 
@@ -875,7 +887,10 @@ void FNiagaraParameterStore::Empty(bool bClearBindings)
 #if WITH_EDITORONLY_DATA
 	ParameterOffsets.Empty();
 #endif // WITH_EDITORONLY_DATA
+
 	SortedParameterOffsets.Empty();
+	InternalStorageChanged();
+
 	DEC_MEMORY_STAT_BY(STAT_NiagaraParamStoreMemory, ParameterData.GetAllocatedSize());
 	ParameterData.Empty();
 	DEC_MEMORY_STAT_BY(STAT_NiagaraParamStoreMemory, ParameterData.GetAllocatedSize());
@@ -896,7 +911,10 @@ void FNiagaraParameterStore::Reset(bool bClearBindings)
 #if WITH_EDITORONLY_DATA
 	ParameterOffsets.Reset();
 #endif // WITH_EDITORONLY_DATA
-	SortedParameterOffsets.Empty();
+
+	SortedParameterOffsets.Reset();
+	InternalStorageChanged();
+
 	DEC_MEMORY_STAT_BY(STAT_NiagaraParamStoreMemory, ParameterData.GetAllocatedSize());
 	ParameterData.Reset();
 	DEC_MEMORY_STAT_BY(STAT_NiagaraParamStoreMemory, ParameterData.GetAllocatedSize());
@@ -931,13 +949,18 @@ void FNiagaraParameterStore::OnLayoutChange()
 #endif
 }
 
+void FNiagaraParameterStore::InternalStorageChanged()
+{
+	ParameterVariables = MakeArrayView(SortedParameterOffsets);
+}
+
 const FNiagaraVariableBase* FNiagaraParameterStore::FindVariable(UNiagaraDataInterface* Interface)const
 {
 	SCOPE_CYCLE_COUNTER(STAT_NiagaraParameterStoreFindVar);
 	int32 Idx = DataInterfaces.IndexOfByKey(Interface);
 	if (Idx != INDEX_NONE)
 	{
-		for (const FNiagaraVariableWithOffset& ParamWithOffset : SortedParameterOffsets)
+		for (const FNiagaraVariableWithOffset& ParamWithOffset : ParameterVariables)
 		{
 			if (ParamWithOffset.Offset == Idx && ParamWithOffset.GetType().GetClass() == Interface->GetClass())
 			{
@@ -948,18 +971,18 @@ const FNiagaraVariableBase* FNiagaraParameterStore::FindVariable(UNiagaraDataInt
 	return nullptr;
 }
 
-const int32* FNiagaraParameterStore::FindParameterOffset(const FNiagaraVariable& Parameter, bool IgnoreType) const
+const int32* FNiagaraParameterStore::FindParameterOffset(const FNiagaraVariableBase& Parameter, bool IgnoreType) const
 {
 #if WITH_EDITORONLY_DATA
 	check(!ParameterOffsets.Num()); // Migration to SortedParameterOffsets
 #endif
 
-	if (SortedParameterOffsets.Num())
+	if (ParameterVariables.Num())
 	{
 		int32 MatchingIndex = 0;
-		if (FNiagaraVariableSearch::Find(SortedParameterOffsets.GetData(), Parameter, 0, SortedParameterOffsets.Num(), IgnoreType, MatchingIndex))
+		if (FNiagaraVariableSearch::Find(ParameterVariables.GetData(), Parameter, 0, ParameterVariables.Num(), IgnoreType, MatchingIndex))
 		{
-			return &SortedParameterOffsets[MatchingIndex].Offset;
+			return &ParameterVariables[MatchingIndex].Offset;
 		}
 	}
 	return nullptr;
@@ -973,9 +996,11 @@ void FNiagaraParameterStore::PostLoad()
 	{
 		for (const TPair<FNiagaraVariable, int32>& ParamOffsetPair : ParameterOffsets)
 		{
-			SortedParameterOffsets.Add(FNiagaraVariableWithOffset(ParamOffsetPair.Key, ParamOffsetPair.Value));
+			SortedParameterOffsets.Emplace(ParamOffsetPair.Key, ParamOffsetPair.Value);
 		}
 		ParameterOffsets.Empty();
+
+		InternalStorageChanged();
 	}
 #endif
 
@@ -989,6 +1014,7 @@ void FNiagaraParameterStore::SortParameters()
 	{
 		return FNiagaraVariableSearch::Compare(Lhs, Rhs) < 0;
 	});
+	InternalStorageChanged();
 }
 
 #if WITH_EDITOR
diff --git a/Engine/Plugins/FX/Niagara/Source/Niagara/Private/NiagaraScript.cpp b/Engine/Plugins/FX/Niagara/Source/Niagara/Private/NiagaraScript.cpp
--- a/Engine/Plugins/FX/Niagara/Source/Niagara/Private/NiagaraScript.cpp
+++ b/Engine/Plugins/FX/Niagara/Source/Niagara/Private/NiagaraScript.cpp
@@ -72,7 +72,6 @@ FNiagaraScriptDebuggerInfo::FNiagaraScriptDebuggerInfo(FName InName, ENiagaraScr
 UNiagaraScriptSourceBase::UNiagaraScriptSourceBase(const FObjectInitializer& ObjectInitializer)
 	: Super(ObjectInitializer)
 {
-	
 }
 
 
@@ -530,12 +529,12 @@ bool UNiagaraScript::ContainsUsage(ENiagaraScriptUsage InUsage) const
 	return false;
 }
 
-FNiagaraScriptExecutionParameterStore* UNiagaraScript::GetExecutionReadyParameterStore(ENiagaraSimTarget SimTarget)
+const FNiagaraScriptExecutionParameterStore* UNiagaraScript::GetExecutionReadyParameterStore(ENiagaraSimTarget SimTarget)
 {
 #if WITH_EDITORONLY_DATA
 	if (SimTarget == ENiagaraSimTarget::CPUSim && IsReadyToRun(ENiagaraSimTarget::CPUSim))
 	{
-		if (ScriptExecutionParamStoreCPU.IsInitialized() == false)
+		if (!ScriptExecutionParamStoreCPU.bInitialized)
 		{
 			ScriptExecutionParamStoreCPU.InitFromOwningScript(this, SimTarget, false);
 		}
@@ -543,7 +542,7 @@ FNiagaraScriptExecutionParameterStore* UNiagaraScript::GetExecutionReadyParamete
 	}
 	else if (SimTarget == ENiagaraSimTarget::GPUComputeSim)
 	{
-		if (ScriptExecutionParamStoreGPU.IsInitialized() == false)
+		if (!ScriptExecutionParamStoreGPU.bInitialized)
 		{
 			ScriptExecutionParamStoreGPU.InitFromOwningScript(this, SimTarget, false);
 		}
@@ -754,11 +753,8 @@ void UNiagaraScript::Serialize(FArchive& Ar)
 			TemporaryStore = RapidIterationParameters;
 
 			// Get the active parameters
-			TArray<FNiagaraVariable> Vars;
-			RapidIterationParameters.GetParameters(Vars);
-
 			// Remove all parameters that aren't data interfaces or uobjects
-			for (const FNiagaraVariable& Var : Vars)
+			for (const FNiagaraVariableBase& Var : TemporaryStore.ParameterVariables)
 			{
 				if (Var.IsDataInterface() || Var.IsUObject())
 					continue;
@@ -766,7 +762,7 @@ void UNiagaraScript::Serialize(FArchive& Ar)
 				NumRemoved++;
 			}
 
-			UE_LOG(LogNiagara, Verbose, TEXT("Pruned %d/%d parameters from script %s"), NumRemoved, TemporaryStore.GetNumParameters(), *GetFullName());
+			UE_LOG(LogNiagara, Verbose, TEXT("Pruned %d/%d parameters from script %s"), NumRemoved, TemporaryStore.ParameterVariables.Num(), *GetFullName());
 		}
 	}
 
@@ -893,7 +889,7 @@ void UNiagaraScript::PostLoad()
 	{
 		ScriptExecutionParamStore.PostLoad();
 		RapidIterationParameters.Bind(&ScriptExecutionParamStore, &ScriptExecutionBoundParameters);
-		ScriptExecutionParamStore.SetAsInitialized();
+		ScriptExecutionParamStore.bInitialized = true;
 		ScriptExecutionBoundParameters.Empty();
 	}
 
diff --git a/Engine/Plugins/FX/Niagara/Source/Niagara/Private/NiagaraScriptExecutionParameterStore.cpp b/Engine/Plugins/FX/Niagara/Source/Niagara/Private/NiagaraScriptExecutionParameterStore.cpp
--- a/Engine/Plugins/FX/Niagara/Source/Niagara/Private/NiagaraScriptExecutionParameterStore.cpp
+++ b/Engine/Plugins/FX/Niagara/Source/Niagara/Private/NiagaraScriptExecutionParameterStore.cpp
@@ -7,10 +7,11 @@
 #include "NiagaraSystemInstance.h"
 
 FNiagaraScriptExecutionParameterStore::FNiagaraScriptExecutionParameterStore()
-	: FNiagaraParameterStore(), ParameterSize(0), PaddedParameterSize(0), bInitialized(false)
-{
-
-}
+	: FNiagaraParameterStore()
+	, ParameterSize(0)
+	, PaddedParameterSize(0)
+	, bInitialized(false)
+{}
 
 FNiagaraScriptExecutionParameterStore::FNiagaraScriptExecutionParameterStore(const FNiagaraParameterStore& Other)
 {
@@ -23,6 +24,7 @@ FNiagaraScriptExecutionParameterStore& FNiagaraScriptExecutionParameterStore::op
 	return *this;
 }
 
+#if WITH_EDITORONLY_DATA
 uint32 OffsetAlign(uint32 SrcOffset, uint32 Size)
 {
 	uint32 OffsetRemaining = SHADER_PARAMETER_STRUCT_ALIGNMENT - (SrcOffset % SHADER_PARAMETER_STRUCT_ALIGNMENT);
@@ -215,38 +217,6 @@ void FNiagaraScriptExecutionParameterStore::InitFromOwningScript(UNiagaraScript*
 	bInitialized = true;
 }
 
-void FNiagaraScriptExecutionParameterStore::InitFromOwningContext(UNiagaraScript* Script, ENiagaraSimTarget SimTarget, bool bNotifyAsDirty)
-{
-	Empty();
-	PaddingInfo.Empty();
-
-#if WITH_EDITORONLY_DATA
-	if (Script != nullptr)
-	{
-		DebugName = FString::Printf(TEXT("ScriptExecParamStore %s %p"), *Script->GetFullName(), this);
-	}
-#endif
-
-	const FNiagaraScriptExecutionParameterStore* SrcStore = Script->GetExecutionReadyParameterStore(SimTarget);
-	if (SrcStore)
-	{
-		InitFromSource(SrcStore, false);
-		ParameterSize = SrcStore->ParameterSize;
-		PaddedParameterSize = SrcStore->PaddedParameterSize;
-		PaddingInfo = SrcStore->PaddingInfo;
-
-		if (bNotifyAsDirty)
-		{
-			MarkParametersDirty();
-			MarkInterfacesDirty();
-			OnLayoutChange();
-		}
-	}
-
-
-	bInitialized = true;
-}
-
 void FNiagaraScriptExecutionParameterStore::AddScriptParams(UNiagaraScript* Script, ENiagaraSimTarget SimTarget, bool bTriggerRebind)
 {
 	if (Script == nullptr)
@@ -262,10 +232,8 @@ void FNiagaraScriptExecutionParameterStore::AddScriptParams(UNiagaraScript* Scri
 		bAdded |= AddParameter(Param, false, false);
 	}
 
-#if WITH_EDITORONLY_DATA
 	DebugName = FString::Printf(TEXT("ScriptExecParamStore %s %p"), *Script->GetFullName(), this);
-#endif
-	
+
 	ParameterSize = GetParameterDataArray().Num();
 
 	//Add previous frame values if we're interpolated spawn.
@@ -281,12 +249,6 @@ void FNiagaraScriptExecutionParameterStore::AddScriptParams(UNiagaraScript* Scri
 			bAdded |= AddParameter(PrevParam, false, false);
 		}
 	}
-	
-	if (bIsInterpolatedSpawn)
-	{
-		CopyCurrToPrev();
-		bAdded = true;
-	}
 
 	//Internal constants - only needed for non-GPU sim
 	if (SimTarget != ENiagaraSimTarget::GPUComputeSim)
@@ -323,26 +285,105 @@ void FNiagaraScriptExecutionParameterStore::AddScriptParams(UNiagaraScript* Scri
 		OnLayoutChange();
 	}
 }
+#endif
+
+//////////////////////////////////////////////////////////////////////////
+/// FNiagaraScriptInstanceParameterStore
+//////////////////////////////////////////////////////////////////////////
+
+FNiagaraScriptInstanceParameterStore::FNiagaraScriptInstanceParameterStore()
+	: FNiagaraParameterStore()
+	, bInitialized(false)
+{}
+
+void FNiagaraScriptInstanceParameterStore::InitFromOwningContext(UNiagaraScript* Script, ENiagaraSimTarget SimTarget, bool bNotifyAsDirty)
+{
+	const FNiagaraScriptExecutionParameterStore* SrcStore = Script ? Script->GetExecutionReadyParameterStore(SimTarget) : nullptr;
+
+#if WITH_EDITORONLY_DATA
+	DebugName = Script ? FString::Printf(TEXT("ScriptExecParamStore %s %p"), *Script->GetFullName(), this) : TEXT("");
+#endif
+
+	if (SrcStore)
+	{
+		Empty(false);
+
+		SetParameterDataArray(SrcStore->GetParameterDataArray(), false);
+		SetDataInterfaces(SrcStore->GetDataInterfaces(), false);
+		SetUObjects(SrcStore->GetUObjects(), false);
+
+		if (bNotifyAsDirty)
+		{
+			MarkParametersDirty();
+			MarkInterfacesDirty();
+			MarkUObjectsDirty();
+			OnLayoutChange();
+		}
+
+		ScriptParameterStore.Init(SrcStore);
+		InternalStorageChanged();
+	}
+	else
+	{
+		Empty();
+	}
+
+	bInitialized = true;
+}
 
-void FNiagaraScriptExecutionParameterStore::CopyCurrToPrev()
+void FNiagaraScriptInstanceParameterStore::CopyCurrToPrev()
 {
-	int32 ParamStart = ParameterSize;
-	checkSlow(FMath::Frac((float)ParameterSize) == 0.0f);
+	if (const auto* ScriptStore = ScriptParameterStore.Get())
+	{
+		const int32 ParamStart = ScriptStore->ParameterSize;
+		FMemory::Memcpy(GetParameterData_Internal(ParamStart), GetParameterData(0), ParamStart);
+	}
+}
 
-	FMemory::Memcpy(GetParameterData_Internal(ParamStart), GetParameterData(0), ParamStart);
+uint32 FNiagaraScriptInstanceParameterStore::GetExternalParameterSize() const
+{
+	if (const auto* ScriptStore = ScriptParameterStore.Get())
+	{
+		return ScriptStore->ParameterSize;
+	}
+	return 0;
 }
 
+uint32 FNiagaraScriptInstanceParameterStore::GetPaddedParameterSizeInBytes() const
+{
+	if (const auto* ScriptStore = ScriptParameterStore.Get())
+	{
+		return ScriptStore->PaddedParameterSize;
+	}
+	return 0;
+}
 
-void FNiagaraScriptExecutionParameterStore::CopyParameterDataToPaddedBuffer(uint8* InTargetBuffer, uint32 InTargetBufferSizeInBytes)
+void FNiagaraScriptInstanceParameterStore::CopyParameterDataToPaddedBuffer(uint8* InTargetBuffer, uint32 InTargetBufferSizeInBytes) const
 {
-	check((uint32)ParameterSize <= PaddedParameterSize);
-	check(InTargetBufferSizeInBytes >= PaddedParameterSize);
-	FMemory::Memzero(InTargetBuffer, InTargetBufferSizeInBytes);
-	const uint8* SrcData = GetParameterDataArray().GetData();
-	for (int32 i = 0; i < PaddingInfo.Num(); i++)
+	if (const auto* ScriptStore = ScriptParameterStore.Get())
 	{
-		check(uint32(PaddingInfo[i].DestOffset + PaddingInfo[i].DestSize) <= InTargetBufferSizeInBytes);
-		check(uint32(PaddingInfo[i].SrcOffset + PaddingInfo[i].SrcSize) <= (uint32)GetParameterDataArray().Num());
-		FMemory::Memcpy(InTargetBuffer + PaddingInfo[i].DestOffset, SrcData + PaddingInfo[i].SrcOffset, PaddingInfo[i].SrcSize);
+		check((uint32)ScriptStore->ParameterSize <= ScriptStore->PaddedParameterSize);
+		check(InTargetBufferSizeInBytes >= ScriptStore->PaddedParameterSize);
+		FMemory::Memzero(InTargetBuffer, InTargetBufferSizeInBytes);
+		const auto& ParameterDataArray = GetParameterDataArray();
+		const uint8* SrcData = ParameterDataArray.GetData();
+		for (const auto& Padding : ScriptStore->PaddingInfo)
+		{
+			check(uint32(Padding.DestOffset + Padding.DestSize) <= InTargetBufferSizeInBytes);
+			check(uint32(Padding.SrcOffset + Padding.SrcSize) <= (uint32)ParameterDataArray.Num());
+			FMemory::Memcpy(InTargetBuffer + Padding.DestOffset, SrcData + Padding.SrcOffset, Padding.SrcSize);
+		}
 	}
 }
+
+void FNiagaraScriptInstanceParameterStore::InternalStorageChanged()
+{
+	if (const auto* ScriptStore = ScriptParameterStore.Get())
+	{
+		ParameterVariables = ScriptStore->ParameterVariables;
+	}
+	else
+	{
+		ParameterVariables = TArrayView<FNiagaraVariableWithOffset>();
+	}
+}
\ No newline at end of file
diff --git a/Engine/Plugins/FX/Niagara/Source/Niagara/Private/NiagaraSystem.cpp b/Engine/Plugins/FX/Niagara/Source/Niagara/Private/NiagaraSystem.cpp
--- a/Engine/Plugins/FX/Niagara/Source/Niagara/Private/NiagaraSystem.cpp
+++ b/Engine/Plugins/FX/Niagara/Source/Niagara/Private/NiagaraSystem.cpp
@@ -457,7 +457,7 @@ void UNiagaraSystem::PostLoad()
 			InitEmitterCompiledData();
 		}
 
-		if (SystemCompiledData.InstanceParamStore.GetNumParameters() == 0 ||SystemCompiledData.DataSetCompiledData.Variables.Num() == 0)
+		if (SystemCompiledData.InstanceParamStore.ParameterVariables.Num() == 0 ||SystemCompiledData.DataSetCompiledData.Variables.Num() == 0)
 		{
 			InitSystemCompiledData();
 		}
diff --git a/Engine/Plugins/FX/Niagara/Source/Niagara/Private/NiagaraUserRedirectionParameterStore.cpp b/Engine/Plugins/FX/Niagara/Source/Niagara/Private/NiagaraUserRedirectionParameterStore.cpp
--- a/Engine/Plugins/FX/Niagara/Source/Niagara/Private/NiagaraUserRedirectionParameterStore.cpp
+++ b/Engine/Plugins/FX/Niagara/Source/Niagara/Private/NiagaraUserRedirectionParameterStore.cpp
@@ -44,7 +44,7 @@ void FNiagaraUserRedirectionParameterStore::RecreateRedirections()
 {
 	UserParameterRedirects.Reset();
 
-	for (const FNiagaraVariable& Var : GetSortedParameterOffsets())
+	for (const FNiagaraVariable& Var : ParameterVariables)
 	{
 		if (IsUserParameter(Var))
 		{
@@ -62,7 +62,7 @@ bool FNiagaraUserRedirectionParameterStore::AddParameter(const FNiagaraVariable&
 	return Super::AddParameter(Param, bInitialize, bTriggerRebind, OutOffset);
 }
 
-bool FNiagaraUserRedirectionParameterStore::RemoveParameter(const FNiagaraVariable& InVar)
+bool FNiagaraUserRedirectionParameterStore::RemoveParameter(const FNiagaraVariableBase& InVar)
 {
 	const FNiagaraVariable* Redirection = UserParameterRedirects.Find(InVar);
 	const FNiagaraVariable& ToRemove = Redirection ? *Redirection : InVar;
diff --git a/Engine/Plugins/FX/Niagara/Source/Niagara/Public/NiagaraCommon.h b/Engine/Plugins/FX/Niagara/Source/Niagara/Public/NiagaraCommon.h
--- a/Engine/Plugins/FX/Niagara/Source/Niagara/Public/NiagaraCommon.h
+++ b/Engine/Plugins/FX/Niagara/Source/Niagara/Public/NiagaraCommon.h
@@ -830,3 +830,60 @@ enum class ENiagaraLegacyTrailWidthMode : uint32
 	FromFirst,
 	FromSecond,
 };
+
+
+//////////////////////////////////////////////////////////////////////////
+/// FNiagaraCompiledDataReference
+//////////////////////////////////////////////////////////////////////////
+
+// Simple struct intended to hide the details of passing a reference to compiled data.  In
+// particular for editor builds it will actually make a copy of the data to help try to avoid the many
+// edge cases of recompiling/editing, while regular builds can reap the wins of just referencing the data
+// and saving memory.
+template<typename T>
+struct FNiagaraCompiledDataReference
+{
+public:
+	void Init(const T* SourceValue)
+	{
+#if WITH_EDITOR
+		OptionalStructValue = *SourceValue;
+#else
+		StructPtr = SourceValue;
+#endif
+	}
+
+	const T* operator->() const
+	{
+		return Get();
+	}
+
+	const T* Get() const
+	{
+#if WITH_EDITOR
+		if (OptionalStructValue.IsSet())
+		{
+			return &OptionalStructValue.GetValue();
+		}
+		return nullptr;
+#else
+		return StructPtr;
+#endif
+	}
+
+	void Reset()
+	{
+#if WITH_EDITOR
+		OptionalStructValue.Reset();
+#else
+		StructPtr = nullptr;
+#endif
+	}
+
+private:
+#if WITH_EDITOR
+	TOptional<T> OptionalStructValue;
+#else
+	const T* StructPtr = nullptr;
+#endif
+};
diff --git a/Engine/Plugins/FX/Niagara/Source/Niagara/Public/NiagaraParameterStore.h b/Engine/Plugins/FX/Niagara/Source/Niagara/Public/NiagaraParameterStore.h
--- a/Engine/Plugins/FX/Niagara/Source/Niagara/Public/NiagaraParameterStore.h
+++ b/Engine/Plugins/FX/Niagara/Source/Niagara/Public/NiagaraParameterStore.h
@@ -137,6 +137,11 @@ struct NIAGARA_API FNiagaraParameterStore
 
 	GENERATED_USTRUCT_BODY()
 
+	/** The View of the set of variables represented by this ParameterStore.  By default it will be the variables defined by the
+	SortedParameterOffsets, but child classes can override it so that we can share the list of Variables across different
+	instances as it's generally read only data in game */
+	TArrayView<FNiagaraVariableWithOffset> ParameterVariables;
+
 private:
 	/** Owner of this store. Used to provide an outer to data interfaces in this store. */
 	UPROPERTY(Transient)
@@ -148,6 +153,8 @@ private:
 	TMap<FNiagaraVariable, int32> ParameterOffsets;
 #endif // WITH_EDITORONLY_DATA
 
+	/** Storage for the set of variables that are represented by this ParameterStore.  Can be empty, and changes to it should
+	also trigger InternalStorageChanged() so that the public ParameterVariables view can be updated. */
 	UPROPERTY()
 	TArray<FNiagaraVariableWithOffset> SortedParameterOffsets;
 
@@ -248,21 +255,21 @@ public:
 #endif
 
 	/** Removes the passed parameter if it exists in the store. */
-	virtual bool RemoveParameter(const FNiagaraVariable& Param);
+	virtual bool RemoveParameter(const FNiagaraVariableBase& Param);
 
 	/** Renames the passed parameter. */
-	void RenameParameter(const FNiagaraVariable& Param, FName NewName);
+	virtual void RenameParameter(const FNiagaraVariableBase& Param, FName NewName);
 
 	/** Removes all parameters from this store and releases any data. */
 	virtual void Empty(bool bClearBindings = true);
 
-	/** Removes all parameters from this store but does't change memory allocations. */
+	/** Removes all parameters from this store but doesn't change memory allocations. */
 	virtual void Reset(bool bClearBindings = true);
 
-	FORCEINLINE void GetParameters(TArray<FNiagaraVariable>& OutParameters) const 
-	{ 
-		OutParameters.Reserve(SortedParameterOffsets.Num());
-		for (const FNiagaraVariableWithOffset& ParamWithOffset : SortedParameterOffsets)
+	FORCEINLINE void GetParameters(TArray<FNiagaraVariable>& OutParameters) const
+	{
+		OutParameters.Reserve(ParameterVariables.Num());
+		for (const FNiagaraVariableWithOffset& ParamWithOffset : ParameterVariables)
 		{
 			OutParameters.Add(ParamWithOffset);
 		}
@@ -270,16 +277,10 @@ public:
 
 	FORCEINLINE TArray<FNiagaraParameterStore*>& GetSourceParameterStores() { return SourceStores; }
 
-	FORCEINLINE const TArray<FNiagaraVariableWithOffset>& GetSortedParameterOffsets() const { return SortedParameterOffsets; }
-
-	FORCEINLINE int32 GetNumParameters() const { return SortedParameterOffsets.Num(); }
-
 	FORCEINLINE const TArray<UObject*>& GetUObjects()const { return UObjects; }
 	FORCEINLINE const TArray<UNiagaraDataInterface*>& GetDataInterfaces()const { return DataInterfaces; }
 	FORCEINLINE const TArray<uint8>& GetParameterDataArray()const { return ParameterData; }
 
-	FORCEINLINE void SetParameterDataArray(const TArray<uint8>& InParameterDataArray);
-
 	void SanityCheckData(bool bInitInterfaces = true);
 
 	// Called to initially set up the parameter store to *exactly* match the input store (other than any bindings and the internal name of it).
@@ -355,7 +356,7 @@ public:
 	/** Returns the associated FNiagaraVariable for the passed data interface if it exists in the store. Null if not.*/
 	const FNiagaraVariableBase* FindVariable(UNiagaraDataInterface* Interface)const;
 
-	virtual const int32* FindParameterOffset(const FNiagaraVariable& Parameter, bool IgnoreType = false) const;
+	virtual const int32* FindParameterOffset(const FNiagaraVariableBase& Parameter, bool IgnoreType = false) const;
 
 	void PostLoad();
 	void SortParameters();
@@ -398,7 +399,7 @@ public:
 			}
 			else
 			{
-				DestStore.SetParameterData(GetParameterData_Internal(SrcIndex), DestIndex, Parameter.GetSizeInBytes());
+				DestStore.SetParameterData(GetParameterData(SrcIndex), DestIndex, Parameter.GetSizeInBytes());
 			}
 		}
 	}
@@ -415,7 +416,7 @@ public:
 	};
 
 	/** Copies all parameters from this parameter store into another.*/
-	void CopyParametersTo(FNiagaraParameterStore& DestStore, bool bOnlyAdd, EDataInterfaceCopyMethod DataInterfaceCopyMethod);
+	void CopyParametersTo(FNiagaraParameterStore& DestStore, bool bOnlyAdd, EDataInterfaceCopyMethod DataInterfaceCopyMethod = EDataInterfaceCopyMethod::None) const;
 
 	/** Remove all parameters from this parameter store from another.*/
 	void RemoveParameters(FNiagaraParameterStore& DestStore);
@@ -512,26 +513,6 @@ public:
 		return false;
 	}
 
-	/**
-	* Sets the parameter using the internally stored data in the passed FNiagaraVariable.
-	* TODO: Remove this. IMO FNiagaraVariable should be just for data definition and all storage should be done via this class.
-	*/
-	FORCEINLINE_DEBUGGABLE void SetParameter(const FNiagaraVariable& Param)
-	{
-		checkSlow(Param.IsDataAllocated());
-		int32 Offset = IndexOf(Param);
-		if (Offset != INDEX_NONE)
-		{
-			uint8* Dest = GetParameterData_Internal(Offset);
-			const uint8* Src = Param.GetData();
-			if (Dest != Src)
-			{
-				FMemory::Memcpy(Dest, Src, Param.GetSizeInBytes());
-			}
-			OnParameterChange();
-		}
-	}
-
 	FORCEINLINE_DEBUGGABLE void SetDataInterface(UNiagaraDataInterface* InInterface, int32 Offset)
 	{
 		DataInterfaces[Offset] = InInterface;
@@ -599,19 +580,14 @@ public:
 protected:
 	void TickBindings();
 	void OnLayoutChange();
+	void CopySortedParameterOffsets(TArrayView<FNiagaraVariableWithOffset> Src);
 
 	/** Returns the parameter data at the passed offset. */
-	FORCEINLINE const uint8* GetParameterData_Internal(int32 Offset) const
-	{
-		return ParameterData.GetData() + Offset;
-	}
-
 	FORCEINLINE uint8* GetParameterData_Internal(int32 Offset) 
 	{
 		return ParameterData.GetData() + Offset;
 	}
 
-
 	template<typename ParamType>
 	void SetParameterByOffset(uint32 ParamOffset, const ParamType& Param)
 	{
@@ -620,6 +596,12 @@ protected:
 		//SetParameterData((const uint8*)&Param, ParamOffset, sizeof(ParamType)); // TODO why aren't we using this path instead of SetParametersByOffset?
 	}
 
+	void SetParameterDataArray(const TArray<uint8>& InParameterDataArray, bool bNotifyAsDirty = true);
+	void SetDataInterfaces(const TArray<UNiagaraDataInterface*>& InDataInterfaces, bool bNotifyAsDirty = true);
+	void SetUObjects(const TArray<UObject*>& InUObjects, bool bNotifyAsDirty = true);
+
+	virtual void InternalStorageChanged();
+
 	friend struct FNiagaraParameterStoreToDataSetBinding;    // this should be the only class calling SetParameterByOffset
 };
 
@@ -672,7 +654,7 @@ FORCEINLINE_DEBUGGABLE bool FNiagaraParameterStoreBinding::VerifyBinding(const F
 	bool bBindingValid = true;
 #if WITH_EDITORONLY_DATA
 	TArray<FName, TInlineAllocator<32>> MissingParameterNames;
-	for (const FNiagaraVariableWithOffset& ParamWithOffset : DestStore->GetSortedParameterOffsets())
+	for (const FNiagaraVariableWithOffset& ParamWithOffset : DestStore->ParameterVariables)
 	{
 		const FNiagaraVariable& Parameter = ParamWithOffset;
 		int32 DestOffset = ParamWithOffset.Offset;
@@ -761,7 +743,7 @@ FORCEINLINE_DEBUGGABLE void FNiagaraParameterStoreBinding::Dump(const FNiagaraPa
 		ensure(Binding.DestOffset != -1);
 		FNiagaraVariable Param;
 		bool bFound = false;
-		for (const FNiagaraVariableWithOffset& ParamWithOffset : DestStore->GetSortedParameterOffsets())
+		for (const FNiagaraVariableWithOffset& ParamWithOffset : DestStore->ParameterVariables)
 		{
 			if (ParamWithOffset.Offset == Binding.DestOffset && !ParamWithOffset.IsDataInterface())
 			{
@@ -793,7 +775,7 @@ FORCEINLINE_DEBUGGABLE void FNiagaraParameterStoreBinding::Dump(const FNiagaraPa
 		ensure(Binding.DestOffset != -1);
 		FNiagaraVariable Param;
 		bool bFound = false;
-		for (const FNiagaraVariableWithOffset& ParamWithOffset : DestStore->GetSortedParameterOffsets())
+		for (const FNiagaraVariableWithOffset& ParamWithOffset : DestStore->ParameterVariables)
 		{
 			if (ParamWithOffset.Offset == Binding.DestOffset && ParamWithOffset.IsDataInterface())
 			{
@@ -825,7 +807,7 @@ FORCEINLINE_DEBUGGABLE void FNiagaraParameterStoreBinding::Dump(const FNiagaraPa
 		ensure(Binding.DestOffset != -1);
 		FNiagaraVariable Param;
 		bool bFound = false;
-		for (const FNiagaraVariableWithOffset& ParamWithOffset : DestStore->GetSortedParameterOffsets())
+		for (const FNiagaraVariableWithOffset& ParamWithOffset : DestStore->ParameterVariables)
 		{
 			if (ParamWithOffset.Offset == Binding.DestOffset && ParamWithOffset.IsUObject())
 			{
diff --git a/Engine/Plugins/FX/Niagara/Source/Niagara/Public/NiagaraScriptExecutionParameterStore.h b/Engine/Plugins/FX/Niagara/Source/Niagara/Public/NiagaraScriptExecutionParameterStore.h
--- a/Engine/Plugins/FX/Niagara/Source/Niagara/Public/NiagaraScriptExecutionParameterStore.h
+++ b/Engine/Plugins/FX/Niagara/Source/Niagara/Public/NiagaraScriptExecutionParameterStore.h
@@ -42,15 +42,15 @@ public:
 
 	//TODO: These function can probably go away entirely when we replace the FNiagaraParameters and DataInterface info in the script with an FNiagaraParameterStore.
 	//Special care with prev params and internal params will have to be taken there.
+#if WITH_EDITORONLY_DATA
 	/** Call this init function if you are using a Niagara parameter store within a UNiagaraScript.*/
 	void InitFromOwningScript(UNiagaraScript* Script, ENiagaraSimTarget SimTarget, bool bNotifyAsDirty);
-	/** Call this init function if you are using a Niagara parameter store within an FNiagaraScriptExecutionContext.*/
-	void InitFromOwningContext(UNiagaraScript* Script, ENiagaraSimTarget SimTarget, bool bNotifyAsDirty);
 	void AddScriptParams(UNiagaraScript* Script, ENiagaraSimTarget SimTarget, bool bTriggerRebind);
-	void CopyCurrToPrev();
+#endif
 
 	virtual bool AddParameter(const FNiagaraVariable& Param, bool bInitInterfaces = true, bool bTriggerRebind = true, int32* OutOffset = nullptr) override
 	{
+#if WITH_EDITORONLY_DATA
 		int32 NewParamOffset = INDEX_NONE;
 		const bool bAdded = FNiagaraParameterStore::AddParameter(Param, bInitInterfaces, bTriggerRebind, &NewParamOffset);
 		if (bAdded)
@@ -62,57 +62,126 @@ public:
 			*OutOffset = NewParamOffset;
 		}
 		return bAdded;
+#else
+		check(0);
+		return false;
+#endif
 	}
 
-	virtual bool RemoveParameter(const FNiagaraVariable& Param) override
+	virtual bool RemoveParameter(const FNiagaraVariableBase& Param) override
 	{
 		check(0);//Not allowed to remove parameters from an execution store as it will adjust the table layout mess up the 
 		return false;
 	}
 
-	void RenameParameter(FNiagaraVariable& Param, FName NewName)
+	virtual void RenameParameter(const FNiagaraVariableBase& Param, FName NewName) override
 	{
 		check(0);//Can't rename parameters for an execution store.
 	}
 
-	virtual void Empty(bool bClearBindings=true) override
+	virtual void Empty(bool bClearBindings = true) override
 	{
 		FNiagaraParameterStore::Empty(bClearBindings);
+		ParameterSize = 0;
+		PaddedParameterSize = 0;
 		PaddingInfo.Empty();
+		bInitialized = false;
+	}
+
+	virtual void Reset(bool bClearBindings = true) override
+	{
+		FNiagaraParameterStore::Reset(bClearBindings);
+		ParameterSize = 0;
 		PaddedParameterSize = 0;
+		PaddingInfo.Empty();
 		bInitialized = false;
 	}
 
-	// Just the external parameters, not previous or internal...
-	uint32 GetExternalParameterSize() const { return ParameterSize; }
+	/** Size of the parameter data not including prev frame values or internal constants. Allows copying into previous parameter values for interpolated spawn scripts. */
+	UPROPERTY()
+	int32 ParameterSize;
 
-	// The entire buffer padded out by the required alignment of the types..
-	uint32 GetPaddedParameterSizeInBytes() const { return PaddedParameterSize; }
+	UPROPERTY()
+	uint32 PaddedParameterSize;
 
-	// Helper that converts the data from the base type array internally into the padded out renderer-ready format.
-	void CopyParameterDataToPaddedBuffer(uint8* InTargetBuffer, uint32 InTargetBufferSizeInBytes);
+	UPROPERTY()
+	TArray<FNiagaraScriptExecutionPaddingInfo> PaddingInfo;
 
+	UPROPERTY()
+	uint8 bInitialized : 1;
 
-	bool IsInitialized() const {return bInitialized;}
-	void SetAsInitialized() { bInitialized = true; }
 protected:
 	void AddPaddedParamSize(const FNiagaraTypeDefinition& InParamType, uint32 InOffset);
 	void AddAlignmentPadding();
 
 private:
 
-	/** Size of the parameter data not including prev frame values or internal constants. Allows copying into previous parameter values for interpolated spawn scripts. */
-	UPROPERTY()
-	int32 ParameterSize;
+#if WITH_EDITORONLY_DATA
+	static uint32 GenerateLayoutInfoInternal(TArray<FNiagaraScriptExecutionPaddingInfo>& Members, uint32& NextMemberOffset, const UStruct* InSrcStruct, uint32 InSrcOffset);
+#endif
+};
 
-	UPROPERTY()
-	uint32 PaddedParameterSize;
+//////////////////////////////////////////////////////////////////////////
+/// FNiagaraScriptInstanceParameterStore
+//////////////////////////////////////////////////////////////////////////
 
-	static uint32 GenerateLayoutInfoInternal(TArray<FNiagaraScriptExecutionPaddingInfo>& Members, uint32& NextMemberOffset, const UStruct* InSrcStruct, uint32 InSrcOffset);
+USTRUCT()
+struct FNiagaraScriptInstanceParameterStore : public FNiagaraParameterStore
+{
+	GENERATED_USTRUCT_BODY()
 
-	UPROPERTY()
-	TArray<FNiagaraScriptExecutionPaddingInfo> PaddingInfo;
+public:
+	FNiagaraScriptInstanceParameterStore();
+	virtual ~FNiagaraScriptInstanceParameterStore() = default;
+
+	/** Call this init function if you are using a Niagara parameter store within an FNiagaraScriptExecutionContext.*/
+	void InitFromOwningContext(UNiagaraScript* Script, ENiagaraSimTarget SimTarget, bool bNotifyAsDirty);
+	void CopyCurrToPrev();
+
+	virtual bool AddParameter(const FNiagaraVariable& Param, bool bInitInterfaces = true, bool bTriggerRebind = true, int32* OutOffset = nullptr) override
+	{
+		check(0);
+		return false;
+	}
+
+	virtual bool RemoveParameter(const FNiagaraVariableBase& Param) override
+	{
+		check(0);//Not allowed to remove parameters from an execution store as it will adjust the table layout mess up the 
+		return false;
+	}
+
+	virtual void RenameParameter(const FNiagaraVariableBase& Param, FName NewName) override
+	{
+		check(0);//Can't rename parameters for an execution store.
+	}
+
+	virtual void Empty(bool bClearBindings = true) override
+	{
+		FNiagaraParameterStore::Empty(bClearBindings);
+		ScriptParameterStore.Reset();
+		bInitialized = false;
+	}
+
+	virtual void Reset(bool bClearBindings = true) override
+	{
+		FNiagaraParameterStore::Reset(bClearBindings);
+		ScriptParameterStore.Reset();
+		bInitialized = false;
+	}
+
+	// Just the external parameters, not previous or internal...
+	uint32 GetExternalParameterSize() const;
+
+	uint32 GetPaddedParameterSizeInBytes() const;
+
+	// Helper that converts the data from the base type array internally into the padded out renderer-ready format.
+	void CopyParameterDataToPaddedBuffer(uint8* InTargetBuffer, uint32 InTargetBufferSizeInBytes) const;
+
+protected:
+	virtual void InternalStorageChanged() override;
+
+private:
+	FNiagaraCompiledDataReference<FNiagaraScriptExecutionParameterStore> ScriptParameterStore;
 
-	UPROPERTY()
 	uint8 bInitialized : 1;
-};
\ No newline at end of file
+};
diff --git a/Engine/Plugins/FX/Niagara/Source/Niagara/Public/NiagaraTypes.h b/Engine/Plugins/FX/Niagara/Source/Niagara/Public/NiagaraTypes.h
--- a/Engine/Plugins/FX/Niagara/Source/Niagara/Public/NiagaraTypes.h
+++ b/Engine/Plugins/FX/Niagara/Source/Niagara/Public/NiagaraTypes.h
@@ -1448,7 +1448,7 @@ private:
 	TArray<uint8> VarData;
 };
 
-FORCEINLINE uint32 GetTypeHash(const FNiagaraVariable& Var)
+FORCEINLINE uint32 GetTypeHash(const FNiagaraVariableBase& Var)
 {
 	return HashCombine(GetTypeHash(Var.GetType()), GetTypeHash(Var.GetName()));
 }
diff --git a/Engine/Plugins/FX/Niagara/Source/Niagara/Public/NiagaraUserRedirectionParameterStore.h b/Engine/Plugins/FX/Niagara/Source/Niagara/Public/NiagaraUserRedirectionParameterStore.h
--- a/Engine/Plugins/FX/Niagara/Source/Niagara/Public/NiagaraUserRedirectionParameterStore.h
+++ b/Engine/Plugins/FX/Niagara/Source/Niagara/Public/NiagaraUserRedirectionParameterStore.h
@@ -29,13 +29,13 @@ public:
 	FORCEINLINE void GetUserParameters(TArray<FNiagaraVariable>& OutParameters) const { return UserParameterRedirects.GenerateKeyArray(OutParameters); }
 
 	// ~ Begin FNiagaraParameterStore overrides
-	FORCEINLINE_DEBUGGABLE virtual const int32* FindParameterOffset(const FNiagaraVariable& Parameter, bool IgnoreType = false) const override
+	FORCEINLINE_DEBUGGABLE virtual const int32* FindParameterOffset(const FNiagaraVariableBase& Parameter, bool IgnoreType = false) const override
 	{
 		const FNiagaraVariable* Redirection = UserParameterRedirects.Find(Parameter);
 		return FNiagaraParameterStore::FindParameterOffset(Redirection ? *Redirection : Parameter);
 	}
 	virtual bool AddParameter(const FNiagaraVariable& Param, bool bInitialize = true, bool bTriggerRebind = true, int32* OutOffset = nullptr) override;
-	virtual bool RemoveParameter(const FNiagaraVariable& InVar) override;
+	virtual bool RemoveParameter(const FNiagaraVariableBase& InVar) override;
 	virtual void InitFromSource(const FNiagaraParameterStore* SrcStore, bool bNotifyAsDirty) override;
 	virtual void Empty(bool bClearBindings = true) override;
 	virtual void Reset(bool bClearBindings = true) override;
diff --git a/Engine/Plugins/FX/Niagara/Source/NiagaraEditor/Private/Customizations/NiagaraTypeCustomizations.cpp b/Engine/Plugins/FX/Niagara/Source/NiagaraEditor/Private/Customizations/NiagaraTypeCustomizations.cpp
--- a/Engine/Plugins/FX/Niagara/Source/NiagaraEditor/Private/Customizations/NiagaraTypeCustomizations.cpp
+++ b/Engine/Plugins/FX/Niagara/Source/NiagaraEditor/Private/Customizations/NiagaraTypeCustomizations.cpp
@@ -351,7 +351,7 @@ TArray<FName> FNiagaraUserParameterBindingCustomization::GetNames() const
 
 	if (BaseSystem && TargetUserParameterBinding)
 	{
-		for (const FNiagaraVariable& Var : BaseSystem->GetExposedParameters().GetSortedParameterOffsets())
+		for (const FNiagaraVariable& Var : BaseSystem->GetExposedParameters().ParameterVariables)
 		{
 			if (FNiagaraParameterMapHistory::IsUserParameter(Var) && Var.GetType() == TargetUserParameterBinding->Parameter.GetType())
 			{
diff --git a/Engine/Plugins/FX/Niagara/Source/NiagaraEditor/Private/ViewModels/Stack/NiagaraStackSystemSettingsGroup.cpp b/Engine/Plugins/FX/Niagara/Source/NiagaraEditor/Private/ViewModels/Stack/NiagaraStackSystemSettingsGroup.cpp
--- a/Engine/Plugins/FX/Niagara/Source/NiagaraEditor/Private/ViewModels/Stack/NiagaraStackSystemSettingsGroup.cpp
+++ b/Engine/Plugins/FX/Niagara/Source/NiagaraEditor/Private/ViewModels/Stack/NiagaraStackSystemSettingsGroup.cpp
@@ -171,7 +171,7 @@ FText UNiagaraStackParameterStoreItem::GetTooltipText() const
 
 void UNiagaraStackParameterStoreItem::RefreshChildrenInternal(const TArray<UNiagaraStackEntry*>& CurrentChildren, TArray<UNiagaraStackEntry*>& NewChildren, TArray<FStackIssue>& NewIssues)
 {
-	if (ParameterStore != nullptr && ParameterStore->GetNumParameters() > 0)
+	if (ParameterStore != nullptr && ParameterStore->ParameterVariables.Num() > 0)
 	{
 		TArray<FNiagaraVariable> Variables;
 		ParameterStore->GetParameters(Variables);
diff --git a/Engine/Plugins/FX/Niagara/Source/NiagaraEditor/Private/Widgets/SNiagaraSpreadsheetView.cpp b/Engine/Plugins/FX/Niagara/Source/NiagaraEditor/Private/Widgets/SNiagaraSpreadsheetView.cpp
--- a/Engine/Plugins/FX/Niagara/Source/NiagaraEditor/Private/Widgets/SNiagaraSpreadsheetView.cpp
+++ b/Engine/Plugins/FX/Niagara/Source/NiagaraEditor/Private/Widgets/SNiagaraSpreadsheetView.cpp
@@ -1142,7 +1142,7 @@ void SNiagaraSpreadsheetView::ResetEntries(EUITab Tab)
 		}
 	
 		{
-			int32 NumInstances = CaptureData[(int32)Tab].InputParams.GetNumParameters();
+			int32 NumInstances = CaptureData[(int32)Tab].InputParams.ParameterVariables.Num();
 			if (!CaptureData[(int32)Tab].bInputColumnsAreAttributes && CaptureData[(int32)Tab].SupportedInputFields.IsValid())
 			{
 				NumInstances = CaptureData[(int32)Tab].SupportedInputFields->Num();
