Clean up allocations surrounding the generation of the optimimzed VM bytecode for niagara (provide an estimate for the starting size, and shrink the array afterwards)

diff --git a/Engine/Plugins/FX/Niagara/Source/Niagara/Private/NiagaraScript.cpp b/Engine/Plugins/FX/Niagara/Source/Niagara/Private/NiagaraScript.cpp
--- a/Engine/Plugins/FX/Niagara/Source/Niagara/Private/NiagaraScript.cpp
+++ b/Engine/Plugins/FX/Niagara/Source/Niagara/Private/NiagaraScript.cpp
@@ -627,11 +627,16 @@ void UNiagaraScript::AsyncOptimizeByteCode()
 	static const IConsoleVariable* CVarFreeUnoptimizedByteCode = IConsoleManager::Get().FindConsoleVariable(TEXT("vm.FreeUnoptimizedByteCode"));
 	if ( FPlatformProperties::RequiresCookedData() && CVarFreeUnoptimizedByteCode && (CVarFreeUnoptimizedByteCode->GetInt() != 0) )
 	{
+		// use the current size of the byte code as a starting point for the allocator
+		CachedScriptVM.OptimizedByteCode.Reserve(CachedScriptVM.ByteCode.Num());
+
 		VectorVM::OptimizeByteCode(CachedScriptVM.ByteCode.GetData(), CachedScriptVM.OptimizedByteCode, MakeArrayView(ExternalFunctionRegisterCounts));
 		if (CachedScriptVM.OptimizedByteCode.Num() > 0)
 		{
 			CachedScriptVM.ByteCode.Empty();
 		}
+
+		CachedScriptVM.OptimizedByteCode.Shrink();
 	}
 	else
 	{
@@ -642,6 +647,7 @@ void UNiagaraScript::AsyncOptimizeByteCode()
 			{
 				// Generate optimized byte code on any thread
 				TArray<uint8> OptimizedByteCode;
+				OptimizedByteCode.Reserve(InByteCode.Num());
 				VectorVM::OptimizeByteCode(InByteCode.GetData(), OptimizedByteCode, MakeArrayView(InExternalFunctionRegisterCounts));
 
 				// Kick off task to set optimized byte code on game thread
@@ -653,6 +659,7 @@ void UNiagaraScript::AsyncOptimizeByteCode()
 						if ( (NiagaraScript != nullptr) && (NiagaraScript->CachedScriptVMId == InCachedScriptVMId) )
 						{
 							NiagaraScript->CachedScriptVM.OptimizedByteCode = MoveTemp(InOptimizedByteCode);
+							NiagaraScript->CachedScriptVM.OptimizedByteCode.Shrink();
 						}
 					}
 				);
