Fix crash adding EventHandler caused by system instances being active and the event handler modifying the cached data

diff --git a/Engine/Plugins/FX/Niagara/Source/NiagaraEditor/Private/ViewModels/Stack/NiagaraStackEventHandlerGroup.cpp b/Engine/Plugins/FX/Niagara/Source/NiagaraEditor/Private/ViewModels/Stack/NiagaraStackEventHandlerGroup.cpp
--- a/Engine/Plugins/FX/Niagara/Source/NiagaraEditor/Private/ViewModels/Stack/NiagaraStackEventHandlerGroup.cpp
+++ b/Engine/Plugins/FX/Niagara/Source/NiagaraEditor/Private/ViewModels/Stack/NiagaraStackEventHandlerGroup.cpp
@@ -9,6 +9,7 @@
 #include "NiagaraEmitter.h"
 #include "NiagaraGraph.h"
 #include "NiagaraScriptSource.h"
+#include "NiagaraSystem.h"
 #include "ViewModels/Stack/INiagaraStackItemGroupAddUtilities.h"
 #include "ViewModels/Stack/NiagaraStackGraphUtilities.h"
 
@@ -40,6 +41,13 @@ public:
 		// The stack should not have been created if any of these are null, so bail out if it happens somehow rather than try to handle all of these cases.
 		checkf(Emitter != nullptr && Source != nullptr && Graph != nullptr, TEXT("Stack created for invalid emitter or graph."));
 
+		// since this is potentially modifying live data we need to kill off any existing instances that might be in flight before we make
+		// the change.
+		if (const UNiagaraSystem* EmitterSystem = Cast<UNiagaraSystem>(Emitter->GetOuter()))
+		{
+			FNiagaraEditorUtilities::KillSystemInstances(*EmitterSystem);
+		}
+
 		FScopedTransaction ScopedTransaction(LOCTEXT("AddNewEventHandlerTransaction", "Add new event handler"));
 
 		Emitter->Modify();
