Caching some event handling info to avoid costly work every frame.
Saves ~30us per event handler per emitter per frame.

Also exiting tick earlier if there is no work to do.
Shaves ~2.5us of empty ticking emitters.

diff --git a/Engine/Plugins/FX/Niagara/Source/Niagara/Private/NiagaraEmitterInstance.cpp b/Engine/Plugins/FX/Niagara/Source/Niagara/Private/NiagaraEmitterInstance.cpp
--- a/Engine/Plugins/FX/Niagara/Source/Niagara/Private/NiagaraEmitterInstance.cpp
+++ b/Engine/Plugins/FX/Niagara/Source/Niagara/Private/NiagaraEmitterInstance.cpp
@@ -439,6 +439,20 @@ void FNiagaraEmitterInstance::Init(int32 InEmitterIdx, FNiagaraSystemInstanceID
 			EventContext.Parameters.UnbindFromSourceStores();
 		}
 	}	
+
+	const int32 NumEventHandlers = CachedEmitter->GetEventHandlers().Num();
+	EventHandlingInfo.Reset();
+	EventHandlingInfo.SetNum(NumEventHandlers);
+	for (int32 i = 0; i < NumEventHandlers; i++)
+	{
+		const FNiagaraEventScriptProperties& EventHandlerProps = CachedEmitter->GetEventHandlers()[i];
+		FNiagaraEventHandlingInfo& Info = EventHandlingInfo[i];
+		Info.SourceEmitterGuid = EventHandlerProps.SourceEmitterID;
+		Info.SourceEmitterName = Info.SourceEmitterGuid.IsValid() ? *Info.SourceEmitterGuid.ToString() : CachedIDName;
+		Info.SpawnCounts.Reset();
+		Info.TotalSpawnCount = 0;
+		Info.EventData = nullptr;
+	}
 }
 
 void FNiagaraEmitterInstance::ResetSimulation(bool bKillExisting /*= true*/)
@@ -831,7 +845,11 @@ void FNiagaraEmitterInstance::PostTick()
 
 	checkSlow(CachedEmitter);
 
-	EventHandlingInfo.Reset();
+	//Clear refs to event data buffers.
+	for (FNiagaraEventHandlingInfo& Info : EventHandlingInfo)
+	{
+		Info.SetEventData(nullptr);
+	}
 
 	CachedBounds.Init();
 	if (CachedSystemFixedBounds.IsSet())
@@ -976,27 +994,22 @@ void FNiagaraEmitterInstance::PreTick()
 	//Gather events we're going to be reading from / handling this frame.
 	//We must do this in pre-tick so we can gather (and mark in use) all sets from other emitters.
 	const int32 NumEventHandlers = CachedEmitter->GetEventHandlers().Num();
-	if (NumEventHandlers > 0)
+	EventSpawnTotal = 0;
+	for (int32 i = 0; i < NumEventHandlers; i++)
 	{
-		EventHandlingInfo.Reset();
-		EventHandlingInfo.SetNum(NumEventHandlers);
-		EventSpawnTotal = 0;
-		for (int32 i = 0; i < NumEventHandlers; i++)
-		{
-			const FNiagaraEventScriptProperties &EventHandlerProps = CachedEmitter->GetEventHandlers()[i];
-			FNiagaraEventHandlingInfo& Info = EventHandlingInfo[i];
-			Info.SourceEmitterGuid = EventHandlerProps.SourceEmitterID;
-			Info.SourceEmitterName = Info.SourceEmitterGuid.IsValid() ? *Info.SourceEmitterGuid.ToString() : CachedIDName;
-			Info.SpawnCounts.Reset();
-			Info.TotalSpawnCount = 0;
-			Info.EventData = nullptr;
-			if (FNiagaraDataSet* EventSet = ParentSystemInstance->GetEventDataSet(Info.SourceEmitterName, EventHandlerProps.SourceEventName))
-			{
-				Info.SetEventData(&EventSet->GetCurrentDataChecked());
-				uint32 EventSpawnNum = CalculateEventSpawnCount(EventHandlerProps, Info.SpawnCounts, EventSet);
-				Info.TotalSpawnCount += EventSpawnNum;
-				EventSpawnTotal += EventSpawnNum;
-			}
+		const FNiagaraEventScriptProperties& EventHandlerProps = CachedEmitter->GetEventHandlers()[i];
+		FNiagaraEventHandlingInfo& Info = EventHandlingInfo[i];
+
+		Info.TotalSpawnCount = 0;//This was being done every frame but should be done in init?
+		Info.SpawnCounts.Reset();
+
+		//TODO: We can move this lookup into the init and just store a ptr to the other set?
+		if (FNiagaraDataSet* EventSet = ParentSystemInstance->GetEventDataSet(Info.SourceEmitterName, EventHandlerProps.SourceEventName))
+		{
+			Info.SetEventData(&EventSet->GetCurrentDataChecked());
+			uint32 EventSpawnNum = CalculateEventSpawnCount(EventHandlerProps, Info.SpawnCounts, EventSet);
+			Info.TotalSpawnCount += EventSpawnNum;
+			EventSpawnTotal += EventSpawnNum;
 		}
 	}
 
@@ -1087,21 +1100,11 @@ void FNiagaraEmitterInstance::Tick(float DeltaSeconds)
 	}
 
 
-	check(Data.GetNumVariables() > 0);
-	check(CachedEmitter->SpawnScriptProps.Script);
-	check(CachedEmitter->UpdateScriptProps.Script);
+	checkSlow(Data.GetNumVariables() > 0);
+	checkSlow(CachedEmitter->SpawnScriptProps.Script);
+	checkSlow(CachedEmitter->UpdateScriptProps.Script);
 	
 	//TickEvents(DeltaSeconds);
-
-	// add system constants
-	{
-		SCOPE_CYCLE_COUNTER(STAT_NiagaraConstants);
-
-		auto& EmitterParameters = ParentSystemInstance->EditEmitterParameters(EmitterIdx);
-		EmitterParameters.EmitterTotalSpawnedParticles = TotalSpawnedParticles;
-		EmitterParameters.EmitterAge = EmitterAge;
-		EmitterParameters.EmitterRandomSeed = CachedEmitter->RandomSeed;
-	}
 	
 	// Calculate number of new particles from regular spawning 
 	uint32 SpawnTotal = 0;
@@ -1118,6 +1121,13 @@ void FNiagaraEmitterInstance::Tick(float DeltaSeconds)
 	int32 OrigNumParticles = GetNumParticles();
 	int32 AllocationEstimate = CachedEmitter->GetMaxParticleCountEstimate();
 	int32 RequiredSize = OrigNumParticles + SpawnTotal + EventSpawnTotal;
+
+	if (RequiredSize == 0)
+	{
+		//Early out if we have no particles to process.
+		return;
+	}
+
 	int32 AllocationSize = FMath::Max<int32>(AllocationEstimate, RequiredSize);
 	if (AllocationSize > MaxAllocationCount)
 	{
@@ -1137,6 +1147,16 @@ void FNiagaraEmitterInstance::Tick(float DeltaSeconds)
 		MinOverallocation = Overallocation;
 	}
 
+	// add system constants
+	{
+		SCOPE_CYCLE_COUNTER(STAT_NiagaraConstants);
+
+		auto& EmitterParameters = ParentSystemInstance->EditEmitterParameters(EmitterIdx);
+		EmitterParameters.EmitterTotalSpawnedParticles = TotalSpawnedParticles;
+		EmitterParameters.EmitterAge = EmitterAge;
+		EmitterParameters.EmitterRandomSeed = CachedEmitter->RandomSeed;
+	}
+
 	/* GPU simulation -  we just create an FNiagaraComputeExecutionContext, queue it, and let the batcher take care of the rest
 	 */
 	if (CachedEmitter->SimTarget == ENiagaraSimTarget::GPUComputeSim && GPUExecContext != nullptr)
