Maybe PVS is confused on the previous line and producing a false positive because of the unnecesary conditional.

diff --git a/Engine/Plugins/FX/Niagara/Source/Niagara/Private/NiagaraScriptExecutionContext.cpp b/Engine/Plugins/FX/Niagara/Source/Niagara/Private/NiagaraScriptExecutionContext.cpp
--- a/Engine/Plugins/FX/Niagara/Source/Niagara/Private/NiagaraScriptExecutionContext.cpp
+++ b/Engine/Plugins/FX/Niagara/Source/Niagara/Private/NiagaraScriptExecutionContext.cpp
@@ -158,7 +158,7 @@ void FNiagaraScriptExecutionContext::BindData(int32 Index, FNiagaraDataSet& Data
 	DataSetInfo[Index].Init(&DataSet, Input, Output, StartInstance, bUpdateInstanceCounts);
 
 	//Would be nice to roll this and DataSetInfo into one but currently the VM being in it's own Engine module prevents this. Possibly should move the VM into Niagara itself.
-	TArrayView<uint8 const* RESTRICT const> InputRegisters = Input ? Input->GetRegisterTable() : TArrayView<uint8 const* RESTRICT const>();
+	TArrayView<uint8 const* RESTRICT const> InputRegisters = Input->GetRegisterTable();
 	TArrayView<uint8 const* RESTRICT const> OutputRegisters = Output ? Output->GetRegisterTable() : TArrayView<uint8 const* RESTRICT const>();
 
 	DataSetMetaTable.SetNum(FMath::Max(DataSetMetaTable.Num(), Index + 1));
