Making user parameters show up when added directly. Also refreshing panel post compile.

diff --git a/Engine/Plugins/FX/Niagara/Source/NiagaraEditor/Private/ViewModels/NiagaraParameterPanelViewModel.cpp b/Engine/Plugins/FX/Niagara/Source/NiagaraEditor/Private/ViewModels/NiagaraParameterPanelViewModel.cpp
--- a/Engine/Plugins/FX/Niagara/Source/NiagaraEditor/Private/ViewModels/NiagaraParameterPanelViewModel.cpp
+++ b/Engine/Plugins/FX/Niagara/Source/NiagaraEditor/Private/ViewModels/NiagaraParameterPanelViewModel.cpp
@@ -25,6 +25,7 @@
 #include "EdGraph/EdGraphSchema.h"
 #include "Editor.h"
 #include "ScopedTransaction.h"
+#include "NiagaraSystemEditorData.h"
 
 #define LOCTEXT_NAMESPACE "NiagaraParameterPanelViewModel"
 
@@ -263,6 +264,7 @@ void FNiagaraSystemToolkitParameterPanelViewModel::InitBindings()
 		SystemViewModel->GetSelectionViewModel()->OnEmitterHandleIdSelectionChanged().AddSP(this, &FNiagaraSystemToolkitParameterPanelViewModel::RefreshSelectedEmitterScriptGraphs);
 		SystemViewModel->OnEmitterHandleViewModelsChanged().AddSP(this, &FNiagaraSystemToolkitParameterPanelViewModel::RefreshSelectedEmitterScriptGraphs);
 	}
+	SystemViewModel->OnSystemCompiled().AddSP(this, &FNiagaraSystemToolkitParameterPanelViewModel::Refresh);
 }
 
 void FNiagaraSystemToolkitParameterPanelViewModel::Refresh() const
@@ -318,7 +320,10 @@ void FNiagaraSystemToolkitParameterPanelViewModel::AddParameter(const FNiagaraVa
 	if (NewScope == ENiagaraParameterScope::User)
 	{
 		SystemViewModel->GetSystem().Modify();
-		SystemViewModel->GetSystem().GetExposedParameters().AddParameter(InVariableToAdd); //@Todo(ng) copy ParameterMapView impl
+		UNiagaraSystem* System = &SystemViewModel->GetSystem();
+		UNiagaraSystemEditorData* SystemEditorData = CastChecked<UNiagaraSystemEditorData>(System->GetEditorData(), ECastCheckedType::NullChecked);
+		SystemEditorData->Modify();
+		bool bSuccess = FNiagaraEditorUtilities::AddParameter(DuplicateVar, System->GetExposedParameters(), *System, SystemEditorData->GetStackEditorData());
 	}
 	else if (NewScope == ENiagaraParameterScope::System)
 	{
@@ -492,6 +497,15 @@ const TArray<FNiagaraScriptVariableAndViewInfo> FNiagaraSystemToolkitParameterPa
 			}
 		}
 	}
+
+	TArray<FNiagaraVariable> UserVars;
+	SystemViewModel->GetSystem().GetExposedParameters().GetParameters(UserVars);
+	for (const FNiagaraVariable& Var : UserVars)
+	{
+		FNiagaraVariableMetaData MetaData;
+		FNiagaraEditorUtilities::GetParameterMetaDataFromName(Var.GetName(), MetaData);
+		ViewedParameters.Add(FNiagaraScriptVariableAndViewInfo(Var, MetaData));
+	}
 	CachedViewedParameters = ViewedParameters;
 	return ViewedParameters;
 }
