Fixed: the instance count offset shader parameters weren't being set when a stage used an iteration interface, which led to bogus values for the execution index.

#rnx

diff --git a/Engine/Plugins/FX/Niagara/Source/Niagara/Private/NiagaraEmitterInstanceBatcher.cpp b/Engine/Plugins/FX/Niagara/Source/Niagara/Private/NiagaraEmitterInstanceBatcher.cpp
--- a/Engine/Plugins/FX/Niagara/Source/Niagara/Private/NiagaraEmitterInstanceBatcher.cpp
+++ b/Engine/Plugins/FX/Niagara/Source/Niagara/Private/NiagaraEmitterInstanceBatcher.cpp
@@ -1337,6 +1337,8 @@ void NiagaraEmitterInstanceBatcher::Run(const FNiagaraGPUSystemTick& Tick, const
 		if (IterationInterface)
 		{
 			RHICmdList.SetUAVParameter(ComputeShader, Shader->InstanceCountsParam.GetUAVIndex(), GetEmptyRWBufferFromPool(RHICmdList, PF_R32_UINT));
+			SetShaderValue(RHICmdList, ComputeShader, Shader->ReadInstanceCountOffsetParam, -1);
+			SetShaderValue(RHICmdList, ComputeShader, Shader->WriteInstanceCountOffsetParam, -1);
 		}
 		else
 		{
diff --git a/Engine/Plugins/FX/Niagara/Source/NiagaraCore/Private/NiagaraCustomVersion.cpp b/Engine/Plugins/FX/Niagara/Source/NiagaraCore/Private/NiagaraCustomVersion.cpp
--- a/Engine/Plugins/FX/Niagara/Source/NiagaraCore/Private/NiagaraCustomVersion.cpp
+++ b/Engine/Plugins/FX/Niagara/Source/NiagaraCore/Private/NiagaraCustomVersion.cpp
@@ -8,4 +8,4 @@ const FGuid FNiagaraCustomVersion::GUID(0xFCF57AFA, 0x50764283, 0xB9A9E658, 0xFF
 // Register the custom version with core
 FCustomVersionRegistration GRegisterNiagaraCustomVersion(FNiagaraCustomVersion::GUID, FNiagaraCustomVersion::LatestVersion, TEXT("NiagaraVer"));
 
-const FGuid FNiagaraCustomVersion::LatestScriptCompileVersion(0x429FAAA4, 0xB12B4D31, 0x922FD12B, 0x82A7D6B7);
+const FGuid FNiagaraCustomVersion::LatestScriptCompileVersion(0x15C033EF, 0x03B849AE, 0xB9235F51, 0x0C187CC4);
diff --git a/Engine/Plugins/FX/Niagara/Source/NiagaraEditor/Private/NiagaraHlslTranslator.cpp b/Engine/Plugins/FX/Niagara/Source/NiagaraEditor/Private/NiagaraHlslTranslator.cpp
--- a/Engine/Plugins/FX/Niagara/Source/NiagaraEditor/Private/NiagaraHlslTranslator.cpp
+++ b/Engine/Plugins/FX/Niagara/Source/NiagaraEditor/Private/NiagaraHlslTranslator.cpp
@@ -2462,7 +2462,7 @@ void FHlslNiagaraTranslator::DefineMainGPUFunctions(
 			"	const uint InstanceID = UpdateStartInstance + DispatchThreadId.x; \n"
 			"	if (ReadInstanceCountOffset == 0xFFFFFFFF)\n"
 			"	{\n"
-			"			GSpawnStartInstance = 0; \n"
+			"		GSpawnStartInstance = 0; \n"
 			"	}\n"
 			"	else\n"
 			"	{\n"
@@ -2530,7 +2530,7 @@ void FHlslNiagaraTranslator::DefineMainGPUFunctions(
 		HlslOutput += TEXT(
 			"	// If a stage doesn't kill particles, StoreUpdateVariables() never calls AcquireIndex(), so the\n"
 			"   // count isn't updated. In that case we must manually copy the original count here.\n"
-			"	if (!GStageUsesAlive && GDispatchThreadId.x == 0) \n"
+			"	if (!GStageUsesAlive && WriteInstanceCountOffset != 0xFFFFFFFF && GDispatchThreadId.x == 0) \n"
 			"	{ \n"
 			"		RWInstanceCounts[WriteInstanceCountOffset] = GSpawnStartInstance + SpawnedInstances; \n"
 			"	} \n"
