Setting up auto-expansion for several categories and limiting make new in others.

diff --git a/Engine/Plugins/FX/Niagara/Source/NiagaraEditor/Private/Widgets/SNiagaraParameterPanel.cpp b/Engine/Plugins/FX/Niagara/Source/NiagaraEditor/Private/Widgets/SNiagaraParameterPanel.cpp
--- a/Engine/Plugins/FX/Niagara/Source/NiagaraEditor/Private/Widgets/SNiagaraParameterPanel.cpp
+++ b/Engine/Plugins/FX/Niagara/Source/NiagaraEditor/Private/Widgets/SNiagaraParameterPanel.cpp
@@ -238,13 +238,17 @@ TSharedRef<SWidget> SNiagaraParameterPanel::OnGetParameterMenu(const NiagaraPara
 {
 	ENiagaraParameterScope NewParameterScopeForSection = NiagaraParameterPanelSectionID::GetScopeForNewParametersInSection(InSection);
 
+	const bool bCanCreateNew = InSection != NiagaraParameterPanelSectionID::Type::ENGINE && InSection != NiagaraParameterPanelSectionID::Type::OWNER;
+	const bool bAutoExpand = InSection == NiagaraParameterPanelSectionID::Type::LOCALS || InSection == NiagaraParameterPanelSectionID::Type::INPUTS ||
+		InSection == NiagaraParameterPanelSectionID::Type::OUTPUTS || InSection == NiagaraParameterPanelSectionID::Type::USER || 
+		InSection == NiagaraParameterPanelSectionID::Type::ENGINE && InSection == NiagaraParameterPanelSectionID::Type::OWNER;
 	TSharedRef<SNiagaraAddParameterMenu2> MenuWidget = SNew(SNiagaraAddParameterMenu2, ParameterPanelViewModel->GetEditableGraphs())
 		.OnAddParameter(this, &SNiagaraParameterPanel::AddParameter, InSection)
 		.OnAllowMakeType(this, &SNiagaraParameterPanel::AllowMakeType)
 		.ShowGraphParameters(false)
 		.ShowKnownConstantParametersFilter(InSection)
-		.AutoExpandMenu(InSection == NiagaraParameterPanelSectionID::Type::LOCALS || InSection == NiagaraParameterPanelSectionID::Type::INPUTS ||
-			InSection == NiagaraParameterPanelSectionID::Type::OUTPUTS)
+		.AllowCreatingNew(bCanCreateNew)
+		.AutoExpandMenu(bAutoExpand)
 		.NewParameterScope(NewParameterScopeForSection);
 
 	AddParameterButtons[(int32)InSection]->SetMenuContentWidgetToFocus(MenuWidget->GetSearchBox()->AsShared());
