Fixing linux build

diff --git a/Engine/Plugins/FX/Niagara/Source/NiagaraEditor/Private/Widgets/SNiagaraParameterPanel.cpp b/Engine/Plugins/FX/Niagara/Source/NiagaraEditor/Private/Widgets/SNiagaraParameterPanel.cpp
--- a/Engine/Plugins/FX/Niagara/Source/NiagaraEditor/Private/Widgets/SNiagaraParameterPanel.cpp
+++ b/Engine/Plugins/FX/Niagara/Source/NiagaraEditor/Private/Widgets/SNiagaraParameterPanel.cpp
@@ -241,7 +241,7 @@ TSharedRef<SWidget> SNiagaraParameterPanel::OnGetParameterMenu(const NiagaraPara
 	const bool bCanCreateNew = InSection != NiagaraParameterPanelSectionID::Type::ENGINE && InSection != NiagaraParameterPanelSectionID::Type::OWNER;
 	const bool bAutoExpand = InSection == NiagaraParameterPanelSectionID::Type::LOCALS || InSection == NiagaraParameterPanelSectionID::Type::INPUTS ||
 		InSection == NiagaraParameterPanelSectionID::Type::OUTPUTS || InSection == NiagaraParameterPanelSectionID::Type::USER || 
-		InSection == NiagaraParameterPanelSectionID::Type::ENGINE && InSection == NiagaraParameterPanelSectionID::Type::OWNER;
+		InSection == NiagaraParameterPanelSectionID::Type::ENGINE || InSection == NiagaraParameterPanelSectionID::Type::OWNER;
 	TSharedRef<SNiagaraAddParameterMenu2> MenuWidget = SNew(SNiagaraAddParameterMenu2, ParameterPanelViewModel->GetEditableGraphs())
 		.OnAddParameter(this, &SNiagaraParameterPanel::AddParameter, InSection)
 		.OnAllowMakeType(this, &SNiagaraParameterPanel::AllowMakeType)
