Fixing compile issue with HLSL when generating the Array members. We weren't properly considering that they might be part of other strings, so the string replacement was too naive.

diff --git a/Engine/Plugins/FX/Niagara/Source/NiagaraCore/Private/NiagaraCustomVersion.cpp b/Engine/Plugins/FX/Niagara/Source/NiagaraCore/Private/NiagaraCustomVersion.cpp
--- a/Engine/Plugins/FX/Niagara/Source/NiagaraCore/Private/NiagaraCustomVersion.cpp
+++ b/Engine/Plugins/FX/Niagara/Source/NiagaraCore/Private/NiagaraCustomVersion.cpp
@@ -8,4 +8,4 @@ const FGuid FNiagaraCustomVersion::GUID(0xFCF57AFA, 0x50764283, 0xB9A9E658, 0xFF
 // Register the custom version with core
 FCustomVersionRegistration GRegisterNiagaraCustomVersion(FNiagaraCustomVersion::GUID, FNiagaraCustomVersion::LatestVersion, TEXT("NiagaraVer"));
 
-const FGuid FNiagaraCustomVersion::LatestScriptCompileVersion(0x15C033EF, 0x03B849AE, 0xB9235F51, 0x0C187CC4);
+const FGuid FNiagaraCustomVersion::LatestScriptCompileVersion(0x0A1C1D5F, 0xBE674CC0, 0x91115FF8, 0x07E209F3 );
diff --git a/Engine/Plugins/FX/Niagara/Source/NiagaraEditor/Private/NiagaraHlslTranslator.cpp b/Engine/Plugins/FX/Niagara/Source/NiagaraEditor/Private/NiagaraHlslTranslator.cpp
--- a/Engine/Plugins/FX/Niagara/Source/NiagaraEditor/Private/NiagaraHlslTranslator.cpp
+++ b/Engine/Plugins/FX/Niagara/Source/NiagaraEditor/Private/NiagaraHlslTranslator.cpp
@@ -2218,7 +2218,13 @@ void FHlslNiagaraTranslator::DefineMainGPUFunctions(
 							if (FNiagaraParameterMapHistory::IsAttribute(Var))
 							{
 								FString RegisterName = VarName;
-								RegisterName.ReplaceInline(PARAM_MAP_ATTRIBUTE_STR, PARAM_MAP_INDICES_STR);
+								const FString AttribStr = PARAM_MAP_ATTRIBUTE_STR;
+								int32 ParticlesIdx = RegisterName.Find(AttribStr);
+								if (ParticlesIdx != -1 && (ParticlesIdx == 0 || (ParticlesIdx > 0 && RegisterName[ParticlesIdx - 1] == '.')))
+								{
+									RegisterName.RemoveAt(ParticlesIdx, AttribStr.Len());
+									RegisterName.InsertAt(ParticlesIdx, PARAM_MAP_INDICES_STR);
+								}
 								const int32 RegisterValue = Var.GetType().IsFloatPrimitive() ? FloatCounter : IntCounter;
 								HlslOutput += RegisterName + FString::Printf(TEXT(" = %d;\n"), RegisterValue);
 							}
