- Integrate 12115914 - Fix for having an RenderTarget still bound that we are going to use as an SRV during Niagara's compute pass

diff --git a/Engine/Plugins/FX/Niagara/Source/Niagara/Private/NiagaraEmitterInstanceBatcher.cpp b/Engine/Plugins/FX/Niagara/Source/Niagara/Private/NiagaraEmitterInstanceBatcher.cpp
--- a/Engine/Plugins/FX/Niagara/Source/Niagara/Private/NiagaraEmitterInstanceBatcher.cpp
+++ b/Engine/Plugins/FX/Niagara/Source/Niagara/Private/NiagaraEmitterInstanceBatcher.cpp
@@ -751,7 +751,7 @@ void NiagaraEmitterInstanceBatcher::ExecuteAll(FRHICommandList& RHICmdList, FRHI
 		return;
 	}
 
-	SCOPED_DRAW_EVENT(RHICmdList, NiagaraEmitterInstanceBatcher_ExecuteAll);
+	RHICmdList.BeginRenderPass(FRHIRenderPassInfo(), TEXT("NiagaraEmitterInstanceBatcher_ExecuteAll"));
 
 	FMemMark Mark(FMemStack::Get());
 	TArray<FOverlappableTicks, TMemStackAllocator<> > SimPasses;
@@ -881,6 +881,8 @@ void NiagaraEmitterInstanceBatcher::ExecuteAll(FRHICommandList& RHICmdList, FRHI
 	RHICmdList.TransitionResources(EResourceTransitionAccess::EReadable, EResourceTransitionPipeline::EComputeToGfx, OutputGraphicsBuffers.GetData(), OutputGraphicsBuffers.Num());
 
 	RHICmdList.EndUAVOverlap();
+
+	RHICmdList.EndRenderPass();
 }
 
 void NiagaraEmitterInstanceBatcher::PreInitViews(FRHICommandListImmediate& RHICmdList, bool bAllowGPUParticleUpdate)
