- Integrate 12115921 - Fix races in Volume / Texture data interfaces, we were reading from data set by the RT when determining if the texture is valid

diff --git a/Engine/Plugins/FX/Niagara/Source/Niagara/Classes/NiagaraDataInterfaceTexture.h b/Engine/Plugins/FX/Niagara/Source/Niagara/Classes/NiagaraDataInterfaceTexture.h
--- a/Engine/Plugins/FX/Niagara/Source/Niagara/Classes/NiagaraDataInterfaceTexture.h
+++ b/Engine/Plugins/FX/Niagara/Source/Niagara/Classes/NiagaraDataInterfaceTexture.h
@@ -65,8 +65,8 @@ protected:
 struct FNiagaraDataInterfaceProxyTexture : public FNiagaraDataInterfaceProxy
 {
 	FSamplerStateRHIRef SamplerStateRHI;
-	FTextureReferenceRHIRef	TextureReferenceRHI;
-	float TexDims[2];
+	FTextureRHIRef	TextureRHI;
+	FVector2D TexDims;
 
 	virtual void ConsumePerInstanceDataFromGameThread(void* PerInstanceData, const FNiagaraSystemInstanceID& Instance) override { check(false); }
 	virtual int32 PerInstanceDataPassedToRenderThreadSize() const override
diff --git a/Engine/Plugins/FX/Niagara/Source/Niagara/Classes/NiagaraDataInterfaceVolumeTexture.h b/Engine/Plugins/FX/Niagara/Source/Niagara/Classes/NiagaraDataInterfaceVolumeTexture.h
--- a/Engine/Plugins/FX/Niagara/Source/Niagara/Classes/NiagaraDataInterfaceVolumeTexture.h
+++ b/Engine/Plugins/FX/Niagara/Source/Niagara/Classes/NiagaraDataInterfaceVolumeTexture.h
@@ -63,7 +63,7 @@ protected:
 struct FNiagaraDataInterfaceProxyVolumeTexture : public FNiagaraDataInterfaceProxy
 {
 	FSamplerStateRHIRef SamplerStateRHI;
-	FTextureReferenceRHIRef TextureReferenceRHI;
+	FTextureRHIRef TextureRHI;
 	FVector TexDims;
 
 	virtual void ConsumePerInstanceDataFromGameThread(void* PerInstanceData, const FNiagaraSystemInstanceID& Instance) override { check(false); }
diff --git a/Engine/Plugins/FX/Niagara/Source/Niagara/Private/NiagaraDataInterfaceTexture.cpp b/Engine/Plugins/FX/Niagara/Source/Niagara/Private/NiagaraDataInterfaceTexture.cpp
--- a/Engine/Plugins/FX/Niagara/Source/Niagara/Private/NiagaraDataInterfaceTexture.cpp
+++ b/Engine/Plugins/FX/Niagara/Source/Niagara/Private/NiagaraDataInterfaceTexture.cpp
@@ -384,8 +384,7 @@ public:
 		FRHIComputeShader* ComputeShaderRHI = Context.Shader.GetComputeShader();
 		FNiagaraDataInterfaceProxyTexture* TextureDI = static_cast<FNiagaraDataInterfaceProxyTexture*>(Context.DataInterface);
 
-		FRHITexture* TextureRHI = (TextureDI && TextureDI->TextureReferenceRHI) ? TextureDI->TextureReferenceRHI->GetReferencedTexture() : nullptr;
-		if (TextureRHI)
+		if (TextureDI && TextureDI->TextureRHI)
 		{
 			FRHISamplerState* SamplerStateRHI = TextureDI->SamplerStateRHI;
 			if (!SamplerStateRHI)
@@ -400,7 +399,7 @@ public:
 				TextureParam,
 				SamplerParam,
 				SamplerStateRHI,
-				TextureRHI
+				TextureDI->TextureRHI
 			);
 			SetShaderValue(RHICmdList, ComputeShaderRHI, Dimensions, TextureDI->TexDims);
 		}
@@ -414,7 +413,7 @@ public:
 				GBlackTexture->SamplerStateRHI,
 				GBlackTexture->TextureRHI
 			);
-			float TexDims[] = { 0.0f, 0.0f };
+			FVector2D TexDims(EForceInit::ForceInitToZero);
 			SetShaderValue(RHICmdList, ComputeShaderRHI, Dimensions, TexDims);
 		}
 	}
@@ -433,32 +432,21 @@ void UNiagaraDataInterfaceTexture::PushToRenderThread()
 {
 	FNiagaraDataInterfaceProxyTexture* RT_Proxy = GetProxyAs<FNiagaraDataInterfaceProxyTexture>();
 
-	float DimX = 0.0f;
-	float DimY = 0.0f;
-	FTextureReferenceRHIRef  RT_TextureReference;
-	FSamplerStateRHIRef RT_SamplerState;
+	FVector2D RT_TexDims(EForceInit::ForceInitToZero);
 
-	if (Texture && Texture->TextureReference.TextureReferenceRHI.IsValid())
+	if (Texture)
 	{
-		if (Texture->TextureReference.TextureReferenceRHI)
-		{
-			RT_TextureReference = Texture->TextureReference.TextureReferenceRHI;
-		}
-		if (Texture->Resource)
-		{
-			RT_SamplerState = Texture->Resource->SamplerStateRHI;
-		}
-		DimX = Texture->GetSurfaceWidth();
-		DimY = Texture->GetSurfaceHeight();
+		RT_TexDims.X = Texture->GetSurfaceWidth();
+		RT_TexDims.Y = Texture->GetSurfaceHeight();
 	}
 
-	ENQUEUE_RENDER_COMMAND(FPushDITextureToRT) (
-		[RT_Proxy, RT_TextureReference, RT_SamplerState, DimX, DimY](FRHICommandListImmediate& RHICmdList)
+	ENQUEUE_RENDER_COMMAND(FPushDITextureToRT)
+	(
+		[RT_Proxy, RT_Resource=Texture ? Texture->Resource : nullptr, RT_TexDims](FRHICommandListImmediate& RHICmdList)
 		{
-			RT_Proxy->TextureReferenceRHI = RT_TextureReference;
-			RT_Proxy->SamplerStateRHI = RT_SamplerState;
-			RT_Proxy->TexDims[0] = DimX;
-			RT_Proxy->TexDims[1] = DimY;
+			RT_Proxy->TextureRHI = RT_Resource ? RT_Resource->TextureRHI : nullptr;
+			RT_Proxy->SamplerStateRHI = RT_Resource ? RT_Resource->SamplerStateRHI : nullptr;
+			RT_Proxy->TexDims = RT_TexDims;
 		}
 	);
 }
diff --git a/Engine/Plugins/FX/Niagara/Source/Niagara/Private/NiagaraDataInterfaceVolumeTexture.cpp b/Engine/Plugins/FX/Niagara/Source/Niagara/Private/NiagaraDataInterfaceVolumeTexture.cpp
--- a/Engine/Plugins/FX/Niagara/Source/Niagara/Private/NiagaraDataInterfaceVolumeTexture.cpp
+++ b/Engine/Plugins/FX/Niagara/Source/Niagara/Private/NiagaraDataInterfaceVolumeTexture.cpp
@@ -241,8 +241,7 @@ public:
 		FRHIComputeShader* ComputeShaderRHI = Context.Shader.GetComputeShader();
 		FNiagaraDataInterfaceProxyVolumeTexture* TextureDI = static_cast<FNiagaraDataInterfaceProxyVolumeTexture*>(Context.DataInterface);
 
-		FRHITexture* TextureRHI = (TextureDI && TextureDI->TextureReferenceRHI) ? TextureDI->TextureReferenceRHI->GetReferencedTexture() : nullptr;
-		if (TextureRHI)
+		if (TextureDI && TextureDI->TextureRHI)
 		{
 			FRHISamplerState* SamplerStateRHI = TextureDI->SamplerStateRHI;
 			if (!SamplerStateRHI)
@@ -257,7 +256,7 @@ public:
 				TextureParam,
 				SamplerParam,
 				SamplerStateRHI,
-				TextureRHI
+				TextureDI->TextureRHI
 			);
 			SetShaderValue(RHICmdList, ComputeShaderRHI, Dimensions, TextureDI->TexDims);
 		}
@@ -271,8 +270,7 @@ public:
 				GBlackTexture->SamplerStateRHI,
 				GBlackTexture->TextureRHI
 			);
-			FVector TexDims(0.0f, 0.0f, 0.0f);
-			SetShaderValue(RHICmdList, ComputeShaderRHI, Dimensions, TexDims);
+			SetShaderValue(RHICmdList, ComputeShaderRHI, Dimensions, FVector::ZeroVector);
 		}
 	}
 private:
@@ -287,36 +285,22 @@ void UNiagaraDataInterfaceVolumeTexture::PushToRenderThread()
 {
 	FNiagaraDataInterfaceProxyVolumeTexture* RT_Proxy = GetProxyAs<FNiagaraDataInterfaceProxyVolumeTexture>();
 
-	float DimX = 0.0f;
-	float DimY = 0.0f;
-	float DimZ = 0.0f;
-	FTextureReferenceRHIRef RT_TextureReference;
-	FSamplerStateRHIRef RT_SamplerState;
-
-	if (Texture && Texture->TextureReference.TextureReferenceRHI.IsValid())
+	FVector RT_TexDims(EForceInit::ForceInitToZero);
+	if (Texture)
 	{
-		if (Texture->TextureReference.TextureReferenceRHI)
-		{
-			RT_TextureReference = Texture->TextureReference.TextureReferenceRHI;
-		}
-		if (Texture->Resource)
-		{
-			RT_SamplerState = Texture->Resource->SamplerStateRHI;
-		}
-		DimX = Texture->GetSizeX();
-		DimY = Texture->GetSizeY();
-		DimZ = Texture->GetSizeZ();
+		RT_TexDims.X = Texture->GetSizeX();
+		RT_TexDims.Y = Texture->GetSizeY();
+		RT_TexDims.Z = Texture->GetSizeZ();
 	}
 
-	ENQUEUE_RENDER_COMMAND(FPushDITextureToRT) (
-		[RT_Proxy, RT_TextureReference, RT_SamplerState, DimX, DimY, DimZ](FRHICommandListImmediate& RHICmdList)
-	{
-		RT_Proxy->TextureReferenceRHI = RT_TextureReference;
-		RT_Proxy->SamplerStateRHI = RT_SamplerState;
-		RT_Proxy->TexDims.X = DimX;
-		RT_Proxy->TexDims.Y = DimY;
-		RT_Proxy->TexDims.Z = DimZ;
-	}
+	ENQUEUE_RENDER_COMMAND(FPushDITextureToRT)
+	(
+		[RT_Proxy, RT_Resource=Texture ? Texture->Resource : nullptr, RT_TexDims](FRHICommandListImmediate& RHICmdList)
+		{
+			RT_Proxy->TextureRHI = RT_Resource ? RT_Resource->TextureRHI : nullptr;
+			RT_Proxy->SamplerStateRHI = RT_Resource ? RT_Resource->SamplerStateRHI : nullptr;
+			RT_Proxy->TexDims = RT_TexDims;
+		}
 	);
 }
 
