Niagara - Fix a bug with cooked niagara sequence data and update test data.

UE-89846

diff --git a/Engine/Plugins/FX/Niagara/Source/Niagara/Private/MovieScene/MovieSceneNiagaraSystemTrackTemplate.h b/Engine/Plugins/FX/Niagara/Source/Niagara/Private/MovieScene/MovieSceneNiagaraSystemTrackTemplate.h
--- a/Engine/Plugins/FX/Niagara/Source/Niagara/Private/MovieScene/MovieSceneNiagaraSystemTrackTemplate.h
+++ b/Engine/Plugins/FX/Niagara/Source/Niagara/Private/MovieScene/MovieSceneNiagaraSystemTrackTemplate.h
@@ -47,8 +47,12 @@ private:
 	FFrameNumber SpawnSectionStartFrame;
 	UPROPERTY()
 	FFrameNumber SpawnSectionEndFrame;
+	UPROPERTY()
 	ENiagaraSystemSpawnSectionStartBehavior SpawnSectionStartBehavior;
+	UPROPERTY()
 	ENiagaraSystemSpawnSectionEvaluateBehavior SpawnSectionEvaluateBehavior;
+	UPROPERTY()
 	ENiagaraSystemSpawnSectionEndBehavior SpawnSectionEndBehavior;
+	UPROPERTY()
 	ENiagaraAgeUpdateMode AgeUpdateMode;
 };
\ No newline at end of file
