- Fix issue with hair that relies on copying a texture inside it's DI and we can't do that within a RenderPass
- Remove submit commands hint (no longer required)

diff --git a/Engine/Plugins/FX/Niagara/Source/Niagara/Private/NiagaraEmitterInstanceBatcher.cpp b/Engine/Plugins/FX/Niagara/Source/Niagara/Private/NiagaraEmitterInstanceBatcher.cpp
--- a/Engine/Plugins/FX/Niagara/Source/Niagara/Private/NiagaraEmitterInstanceBatcher.cpp
+++ b/Engine/Plugins/FX/Niagara/Source/Niagara/Private/NiagaraEmitterInstanceBatcher.cpp
@@ -52,14 +52,6 @@ static FAutoConsoleVariableRef CVarNiagaraUseAsyncCompute(
 	ECVF_Default
 );
 
-int32 GNiagaraSubmitCommands = 0;
-static FAutoConsoleVariableRef CVarNiagaraSubmitCommands(
-	TEXT("fx.NiagaraSubmitCommands"),
-	GNiagaraSubmitCommands,
-	TEXT("1 - (Default) Submit commands to the GPU once we have finished dispatching.\n"),
-	ECVF_Default
-);
-
 // @todo REMOVE THIS HACK
 int32 GNiagaraGpuMaxQueuedRenderFrames = 10;
 static FAutoConsoleVariableRef CVarNiagaraGpuMaxQueuedRenderFrames(
@@ -618,11 +610,6 @@ void NiagaraEmitterInstanceBatcher::DispatchAllOnCompute(FOverlappableTicks& Ove
 			}
 		}
 	}
-
-	if (GNiagaraSubmitCommands)
-	{
-		RHICmdList.SubmitCommandsHint();
-	}
 }
 
 void NiagaraEmitterInstanceBatcher::PostRenderOpaque(FRHICommandListImmediate& RHICmdList, FRHIUniformBuffer* ViewUniformBuffer, const class FShaderParametersMetadata* SceneTexturesUniformBufferStruct, FRHIUniformBuffer* SceneTexturesUniformBuffer, bool bAllowGPUParticleUpdate)
@@ -751,7 +738,7 @@ void NiagaraEmitterInstanceBatcher::ExecuteAll(FRHICommandList& RHICmdList, FRHI
 		return;
 	}
 
-	RHICmdList.BeginRenderPass(FRHIRenderPassInfo(), TEXT("NiagaraEmitterInstanceBatcher_ExecuteAll"));
+	SCOPED_DRAW_EVENTF(RHICmdList, NiagaraEmitterInstanceBatcher_ExecuteAll, TEXT("NiagaraEmitterInstanceBatcher_ExecuteAll - TickStage(%d)"), TickStage);
 
 	FMemMark Mark(FMemStack::Get());
 	TArray<FOverlappableTicks, TMemStackAllocator<> > SimPasses;
@@ -829,6 +816,11 @@ void NiagaraEmitterInstanceBatcher::ExecuteAll(FRHICommandList& RHICmdList, FRHI
 		}
 	}
 
+	// Clear any RT bindings that we may be using
+	// Note: We can not encapsulate the whole Niagara pass as some DI's may not be compatable (i.e. use CopyTexture function), we need to fix this with future RDG conversion
+	RHICmdList.BeginRenderPass(FRHIRenderPassInfo(), TEXT("NiagaraClearRTs"));
+	RHICmdList.EndRenderPass();
+
 	RHICmdList.BeginUAVOverlap();
 
 	FEmitterInstanceList InstancesWithPersistentIDs;
@@ -881,8 +873,6 @@ void NiagaraEmitterInstanceBatcher::ExecuteAll(FRHICommandList& RHICmdList, FRHI
 	RHICmdList.TransitionResources(EResourceTransitionAccess::EReadable, EResourceTransitionPipeline::EComputeToGfx, OutputGraphicsBuffers.GetData(), OutputGraphicsBuffers.Num());
 
 	RHICmdList.EndUAVOverlap();
-
-	RHICmdList.EndRenderPass();
 }
 
 void NiagaraEmitterInstanceBatcher::PreInitViews(FRHICommandListImmediate& RHICmdList, bool bAllowGPUParticleUpdate)
