Fix for gracefully handling sorting niagara sprites/meshes when the sort key is not available in the particle data set

diff --git a/Engine/Plugins/FX/Niagara/Source/Niagara/Private/NiagaraRendererMeshes.cpp b/Engine/Plugins/FX/Niagara/Source/Niagara/Private/NiagaraRendererMeshes.cpp
--- a/Engine/Plugins/FX/Niagara/Source/Niagara/Private/NiagaraRendererMeshes.cpp
+++ b/Engine/Plugins/FX/Niagara/Source/Niagara/Private/NiagaraRendererMeshes.cpp
@@ -377,7 +377,7 @@ void FNiagaraRendererMeshes::GetDynamicMeshElements(const TArray<const FSceneVie
 				FNiagaraGPUSortInfo SortInfo;
 				const int32 SortVarIdx = bCustomSorting ? ENiagaraMeshVFLayout::CustomSorting : ENiagaraMeshVFLayout::Position;
 				const bool bSortVarIsHalf = bShouldSort && VFVariables[SortVarIdx].bHalfType;
-				if (bShouldSort)
+				if (bShouldSort && VFVariables[SortVarIdx].GetGPUOffset() != INDEX_NONE)
 				{
 					SortInfo.ParticleCount = NumInstances;
 					SortInfo.SortMode = SortMode;
diff --git a/Engine/Plugins/FX/Niagara/Source/Niagara/Private/NiagaraRendererSprites.cpp b/Engine/Plugins/FX/Niagara/Source/Niagara/Private/NiagaraRendererSprites.cpp
--- a/Engine/Plugins/FX/Niagara/Source/Niagara/Private/NiagaraRendererSprites.cpp
+++ b/Engine/Plugins/FX/Niagara/Source/Niagara/Private/NiagaraRendererSprites.cpp
@@ -342,7 +342,7 @@ void FNiagaraRendererSprites::SetVertexFactoryParticleData(
 		const bool bCustomSorting = SortMode == ENiagaraSortMode::CustomAscending || SortMode == ENiagaraSortMode::CustomDecending;
 		const int32 SortVarIdx = bCustomSorting ? ENiagaraSpriteVFLayout::CustomSorting : ENiagaraSpriteVFLayout::Position;
 		const bool bSortVarIsHalf = bShouldSort && VFVariables[SortVarIdx].bHalfType;
-		if (bShouldSort)
+		if (bShouldSort && VFVariables[SortVarIdx].GetGPUOffset() != INDEX_NONE)
 		{
 			SortInfo.ParticleCount = NumInstances;
 			SortInfo.SortMode = SortMode;
