Generates function bindings for external DI functions that don't require user poitners so that we can share them across instances
-changes the function table to be an array of function pointers to allow for the redirection without copies/reallocations
-eliminates 2nd most common allocation in cases with large instance counts in niagara

diff --git a/Engine/Plugins/FX/Niagara/Source/Niagara/Classes/NiagaraScript.h b/Engine/Plugins/FX/Niagara/Source/Niagara/Classes/NiagaraScript.h
--- a/Engine/Plugins/FX/Niagara/Source/Niagara/Classes/NiagaraScript.h
+++ b/Engine/Plugins/FX/Niagara/Source/Niagara/Classes/NiagaraScript.h
@@ -309,6 +309,8 @@ public:
 	UPROPERTY()
 	TArray<FVMExternalFunctionBindingInfo> CalledVMExternalFunctions;
 
+	TArray<FVMExternalFunction> CalledVMExternalFunctionBindings;
+
 	UPROPERTY()
 	TArray<FNiagaraDataSetID> ReadDataSets;
 
@@ -702,6 +704,9 @@ private:
 	/** Kicks off an async job to convert the ByteCode into an optimized version for the platform we are running on. */
 	void AsyncOptimizeByteCode();
 
+	/** Generates all of the function bindings for DI that don't require user data */
+	void GenerateDefaultFunctionBindings();
+
 #if WITH_EDITORONLY_DATA
 	UPROPERTY(Transient)
 	FNiagaraScriptExecutionParameterStore ScriptExecutionParamStoreCPU;
diff --git a/Engine/Plugins/FX/Niagara/Source/Niagara/Classes/NiagaraScriptExecutionContext.h b/Engine/Plugins/FX/Niagara/Source/Niagara/Classes/NiagaraScriptExecutionContext.h
--- a/Engine/Plugins/FX/Niagara/Source/Niagara/Classes/NiagaraScriptExecutionContext.h
+++ b/Engine/Plugins/FX/Niagara/Source/Niagara/Classes/NiagaraScriptExecutionContext.h
@@ -149,9 +149,14 @@ struct FNiagaraScriptExecutionContext
 {
 	UNiagaraScript* Script;
 
-	/** Table of external function delegates called from the VM. */
-	TArray<FVMExternalFunction> FunctionTable;
+	/** Table of external function delegate handles called from the VM. */
+	TArray<const FVMExternalFunction*> FunctionTable;
 
+private:
+	/** Table of external function delegates unique to the instance */
+	TArray<FVMExternalFunction> LocalFunctionTable;
+
+public:
 	/** Table of instance data for data interfaces that require it. */
 	TArray<void*> DataInterfaceInstDataTable;
 
diff --git a/Engine/Plugins/FX/Niagara/Source/Niagara/Private/NiagaraScript.cpp b/Engine/Plugins/FX/Niagara/Source/Niagara/Private/NiagaraScript.cpp
--- a/Engine/Plugins/FX/Niagara/Source/Niagara/Private/NiagaraScript.cpp
+++ b/Engine/Plugins/FX/Niagara/Source/Niagara/Private/NiagaraScript.cpp
@@ -537,6 +537,9 @@ const FNiagaraScriptExecutionParameterStore* UNiagaraScript::GetExecutionReadyPa
 		if (!ScriptExecutionParamStoreCPU.bInitialized)
 		{
 			ScriptExecutionParamStoreCPU.InitFromOwningScript(this, SimTarget, false);
+
+			// generate the function bindings for those external functions where there's no user (per-instance) data required
+			GenerateDefaultFunctionBindings();
 		}
 		return &ScriptExecutionParamStoreCPU;
 	}
@@ -668,6 +671,40 @@ void UNiagaraScript::AsyncOptimizeByteCode()
 	}
 }
 
+void UNiagaraScript::GenerateDefaultFunctionBindings()
+{
+	// generate the function bindings for those external functions where there's no user (per-instance) data required
+	auto SimTarget = GetSimTarget();
+	const int32 ExternalFunctionCount = CachedScriptVM.CalledVMExternalFunctions.Num();
+
+	if (SimTarget.IsSet() && ExternalFunctionCount)
+	{
+		CachedScriptVM.CalledVMExternalFunctionBindings.Empty(ExternalFunctionCount);
+
+		const FNiagaraScriptExecutionParameterStore* ScriptParameterStore = GetExecutionReadyParameterStore(*SimTarget);
+		const auto& ScriptDataInterfaces = ScriptParameterStore->GetDataInterfaces();
+
+		const int32 DataInterfaceCount = FMath::Min(CachedScriptVM.DataInterfaceInfo.Num(), ScriptDataInterfaces.Num());
+		check(DataInterfaceCount == CachedScriptVM.DataInterfaceInfo.Num());
+		check(DataInterfaceCount == ScriptDataInterfaces.Num());
+
+		for (const FVMExternalFunctionBindingInfo& BindingInfo : CachedScriptVM.CalledVMExternalFunctions)
+		{
+			FVMExternalFunction& FuncBind = CachedScriptVM.CalledVMExternalFunctionBindings.AddDefaulted_GetRef();
+
+			for (int32 DataInterfaceIt = 0; DataInterfaceIt < DataInterfaceCount; ++DataInterfaceIt)
+			{
+				const FNiagaraScriptDataInterfaceCompileInfo& ScriptInfo = CachedScriptVM.DataInterfaceInfo[DataInterfaceIt];
+
+				if (ScriptInfo.UserPtrIdx == INDEX_NONE && ScriptInfo.Name == BindingInfo.OwnerName)
+				{
+					ScriptDataInterfaces[DataInterfaceIt]->GetVMExternalFunction(BindingInfo, nullptr, FuncBind);
+				}
+			}
+		}
+	}
+}
+
 void UNiagaraScript::PreSave(const class ITargetPlatform* TargetPlatform)
 {
 	Super::PreSave(TargetPlatform);
@@ -898,6 +935,10 @@ void UNiagaraScript::PostLoad()
 		RapidIterationParameters.Bind(&ScriptExecutionParamStore, &ScriptExecutionBoundParameters);
 		ScriptExecutionParamStore.bInitialized = true;
 		ScriptExecutionBoundParameters.Empty();
+
+		// generate the function bindings for those external functions where there's no user (per-instance) data required
+		GenerateDefaultFunctionBindings();
+
 	}
 
 	bool bNeedsRecompile = false;
@@ -995,8 +1036,6 @@ void UNiagaraScript::PostLoad()
 
 	// Optimize the VM script for runtime usage
 	AsyncOptimizeByteCode();
-
-	//FNiagaraUtilities::DumpHLSLText(RapidIterationParameters.ToString(), *GetPathName());
 }
 
 bool UNiagaraScript::IsReadyToRun(ENiagaraSimTarget SimTarget) const
diff --git a/Engine/Plugins/FX/Niagara/Source/Niagara/Private/NiagaraScriptExecutionContext.cpp b/Engine/Plugins/FX/Niagara/Source/Niagara/Private/NiagaraScriptExecutionContext.cpp
--- a/Engine/Plugins/FX/Niagara/Source/Niagara/Private/NiagaraScriptExecutionContext.cpp
+++ b/Engine/Plugins/FX/Niagara/Source/Niagara/Private/NiagaraScriptExecutionContext.cpp
@@ -58,13 +58,14 @@ bool FNiagaraScriptExecutionContext::Tick(FNiagaraSystemInstance* ParentSystemIn
 		SCOPE_CYCLE_COUNTER(STAT_NiagaraScriptExecContextTick);
 		if (Script && Script->IsReadyToRun(ENiagaraSimTarget::CPUSim))//TODO: Remove. Script can only be null for system instances that currently don't have their script exec context set up correctly.
 		{
+			const FNiagaraVMExecutableData& ScriptExecutableData = Script->GetVMExecutableData();
 			const TArray<UNiagaraDataInterface*>& DataInterfaces = GetDataInterfaces();
 
 			SCOPE_CYCLE_COUNTER(STAT_NiagaraRebindDataInterfaceFunctionTable);
 			// UE_LOG(LogNiagara, Log, TEXT("Updating data interfaces for script %s"), *Script->GetFullName());
 
 			// We must make sure that the data interfaces match up between the original script values and our overrides...
-			if (Script->GetVMExecutableData().DataInterfaceInfo.Num() != DataInterfaces.Num())
+			if (ScriptExecutableData.DataInterfaceInfo.Num() != DataInterfaces.Num())
 			{
 				UE_LOG(LogNiagara, Warning, TEXT("Mismatch between Niagara Exectuion Context data interfaces and those in it's script!"));
 				return false;
@@ -73,12 +74,12 @@ bool FNiagaraScriptExecutionContext::Tick(FNiagaraSystemInstance* ParentSystemIn
 			//Fill the instance data table.
 			if (ParentSystemInstance)
 			{
-				DataInterfaceInstDataTable.SetNumZeroed(Script->GetVMExecutableData().NumUserPtrs, false);
+				DataInterfaceInstDataTable.SetNumZeroed(ScriptExecutableData.NumUserPtrs, false);
 				for (int32 i = 0; i < DataInterfaces.Num(); i++)
 				{
 					UNiagaraDataInterface* Interface = DataInterfaces[i];
 
-					int32 UserPtrIdx = Script->GetVMExecutableData().DataInterfaceInfo[i].UserPtrIdx;
+					int32 UserPtrIdx = ScriptExecutableData.DataInterfaceInfo[i].UserPtrIdx;
 					if (UserPtrIdx != INDEX_NONE)
 					{
 						void* InstData = ParentSystemInstance->FindDataInterfaceInstanceData(Interface);
@@ -88,44 +89,77 @@ bool FNiagaraScriptExecutionContext::Tick(FNiagaraSystemInstance* ParentSystemIn
 			}
 			else
 			{
-				check(Script->GetVMExecutableData().NumUserPtrs == 0);//Can't have user ptrs if we have no parent instance.
+				check(ScriptExecutableData.NumUserPtrs == 0);//Can't have user ptrs if we have no parent instance.
 			}
 
-			FunctionTable.Reset(Script->GetVMExecutableData().CalledVMExternalFunctions.Num());
+			const int32 FunctionCount = ScriptExecutableData.CalledVMExternalFunctions.Num();
+			FunctionTable.Reset(FunctionCount);
+			FunctionTable.AddZeroed(FunctionCount);
+			TArray<int32> LocalFunctionTableIndices;
+			LocalFunctionTableIndices.Reserve(FunctionCount);
+
+			const auto& ScriptDataInterfaces = Script->GetExecutionReadyParameterStore(SimTarget)->GetDataInterfaces();
 
 			bool bSuccessfullyMapped = true;
-			for (FVMExternalFunctionBindingInfo& BindingInfo : Script->GetVMExecutableData().CalledVMExternalFunctions)
+
+			for (int32 FunctionIt = 0; FunctionIt < FunctionCount; ++FunctionIt)
 			{
+				const FVMExternalFunctionBindingInfo& BindingInfo = ScriptExecutableData.CalledVMExternalFunctions[FunctionIt];
+
 				// First check to see if we can pull from the fast path library..
 				FVMExternalFunction FuncBind;
 				if (UNiagaraFunctionLibrary::GetVectorVMFastPathExternalFunction(BindingInfo, FuncBind) && FuncBind.IsBound())
 				{
-					FunctionTable.Add(FuncBind);
+					LocalFunctionTable.Add(FuncBind);
+					LocalFunctionTableIndices.Add(FunctionIt);
 					continue;
 				}
 
-				for (int32 i = 0; i < Script->GetVMExecutableData().DataInterfaceInfo.Num(); i++)
+				for (int32 i = 0; i < ScriptExecutableData.DataInterfaceInfo.Num(); i++)
 				{
-					FNiagaraScriptDataInterfaceCompileInfo& ScriptInfo = Script->GetVMExecutableData().DataInterfaceInfo[i];
+					const FNiagaraScriptDataInterfaceCompileInfo& ScriptInfo = ScriptExecutableData.DataInterfaceInfo[i];
 					UNiagaraDataInterface* ExternalInterface = DataInterfaces[i];
 					if (ScriptInfo.Name == BindingInfo.OwnerName)
 					{
+						// first check to see if we should just use the one from the script
+						if (ScriptExecutableData.CalledVMExternalFunctionBindings.IsValidIndex(FunctionIt)
+							&& ExternalInterface == ScriptDataInterfaces[i])
+						{
+							const FVMExternalFunction& ScriptFuncBind = ScriptExecutableData.CalledVMExternalFunctionBindings[FunctionIt];
+							if (ScriptFuncBind.IsBound())
+							{
+								FunctionTable[FunctionIt] = &ScriptFuncBind;
+
+								check(ScriptInfo.UserPtrIdx == INDEX_NONE);
+								break;
+							}
+						}
+
 						void* InstData = ScriptInfo.UserPtrIdx == INDEX_NONE ? nullptr : DataInterfaceInstDataTable[ScriptInfo.UserPtrIdx];
-						int32 AddedIdx = FunctionTable.Add(FVMExternalFunction());
+						FVMExternalFunction& LocalFunction = LocalFunctionTable.AddDefaulted_GetRef();
+						LocalFunctionTableIndices.Add(FunctionIt);
+
 						if (ExternalInterface != nullptr)
 						{
-							ExternalInterface->GetVMExternalFunction(BindingInfo, InstData, FunctionTable[AddedIdx]);
+							ExternalInterface->GetVMExternalFunction(BindingInfo, InstData, LocalFunction);
 						}
 
-						if (AddedIdx != INDEX_NONE && !FunctionTable[AddedIdx].IsBound())
+						if (!LocalFunction.IsBound())
 						{
 							UE_LOG(LogNiagara, Error, TEXT("Could not Get VMExternalFunction '%s'.. emitter will not run!"), *BindingInfo.Name.ToString());
 							bSuccessfullyMapped = false;
 						}
+						break;
 					}
 				}
 			}
 
+			const int32 LocalFunctionCount = LocalFunctionTableIndices.Num();
+			for (int32 LocalFunctionIt = 0; LocalFunctionIt < LocalFunctionCount; ++LocalFunctionIt)
+			{
+				FunctionTable[LocalFunctionTableIndices[LocalFunctionIt]] = &LocalFunctionTable[LocalFunctionIt];
+			}
+
 			if (!bSuccessfullyMapped)
 			{
 				UE_LOG(LogNiagara, Warning, TEXT("Error building data interface function table!"));
diff --git a/Engine/Source/Runtime/VectorVM/Private/VectorVM.cpp b/Engine/Source/Runtime/VectorVM/Private/VectorVM.cpp
--- a/Engine/Source/Runtime/VectorVM/Private/VectorVM.cpp
+++ b/Engine/Source/Runtime/VectorVM/Private/VectorVM.cpp
@@ -427,7 +427,7 @@ void FVectorVMContext::PrepareForExec(
 	int32 InConstantTableCount,
 	const uint8* const* InConstantTable,
 	const int32* InConstantTableSizes,
-	FVMExternalFunction* InExternalFunctionTable,
+	const FVMExternalFunction* const* InExternalFunctionTable,
 	void** InUserPtrTable,
 	TArrayView<FDataSetMeta> InDataSetMetaTable,
 	int32 MaxNumInstances,
@@ -1951,7 +1951,13 @@ struct FKernelExternalFunctionCall
 	static void Exec(FVectorVMContext& Context)
 	{
 		const uint32 ExternalFuncIdx = Context.DecodeU8();
-		Context.ExternalFunctionTable[ExternalFuncIdx].Execute(Context);
+		const FVMExternalFunction* ExternalFunction = Context.ExternalFunctionTable[ExternalFuncIdx];
+		check(ExternalFunction);
+
+		if (ExternalFunction)
+		{
+			ExternalFunction->Execute(Context);
+		}
 	}
 };
 
@@ -2403,7 +2409,7 @@ void VectorVM::Exec(
 	const uint8* const* ConstantTable,
 	const int32* ConstantTableSizes,
 	TArrayView<FDataSetMeta> DataSetMetaTable,
-	FVMExternalFunction* ExternalFunctionTable,
+	const FVMExternalFunction* const* ExternalFunctionTable,
 	void** UserPtrTable,
 	int32 NumInstances
 #if STATS
diff --git a/Engine/Source/Runtime/VectorVM/Public/VectorVM.h b/Engine/Source/Runtime/VectorVM/Public/VectorVM.h
--- a/Engine/Source/Runtime/VectorVM/Public/VectorVM.h
+++ b/Engine/Source/Runtime/VectorVM/Public/VectorVM.h
@@ -286,7 +286,7 @@ public:
 	int32 NumTempRegisters;
 
 	/** Pointer to the shared data table. */
-	FVMExternalFunction* RESTRICT ExternalFunctionTable;
+	const FVMExternalFunction* const* RESTRICT ExternalFunctionTable;
 	/** Table of user pointers.*/
 	void** UserPtrTable;
 
@@ -330,7 +330,7 @@ public:
 		int32 ConstantTableCount,
 		const uint8* const* InConstantTables,
 		const int32* InConstantTableSizes,
-		FVMExternalFunction* InExternalFunctionTable,
+		const FVMExternalFunction* const* InExternalFunctionTable,
 		void** InUserPtrTable,
 		TArrayView<FDataSetMeta> InDataSetMetaTable,
 		int32 MaxNumInstances,
@@ -456,7 +456,7 @@ namespace VectorVM
 		const uint8* const* ConstantTable,
 		const int32* ConstantTableSizes,
 		TArrayView<FDataSetMeta> DataSetMetaTable,
-		FVMExternalFunction* ExternalFunctionTable,
+		const FVMExternalFunction* const* ExternalFunctionTable,
 		void** UserPtrTable,
 		int32 NumInstances
 #if STATS
