Reduce inline allocator to take care of the general (non-event) case for emitters

diff --git a/Engine/Plugins/FX/Niagara/Source/Niagara/Classes/NiagaraScriptExecutionContext.h b/Engine/Plugins/FX/Niagara/Source/Niagara/Classes/NiagaraScriptExecutionContext.h
--- a/Engine/Plugins/FX/Niagara/Source/Niagara/Classes/NiagaraScriptExecutionContext.h
+++ b/Engine/Plugins/FX/Niagara/Source/Niagara/Classes/NiagaraScriptExecutionContext.h
@@ -163,9 +163,9 @@ public:
 	/** Parameter store. Contains all data interfaces and a parameter buffer that can be used directly by the VM or GPU. */
 	FNiagaraScriptInstanceParameterStore Parameters;
 
-	TArray<FDataSetMeta, TInlineAllocator<4>> DataSetMetaTable;
+	TArray<FDataSetMeta, TInlineAllocator<2>> DataSetMetaTable;
 
-	TArray<FNiagaraDataSetExecutionInfo, TInlineAllocator<4>> DataSetInfo;
+	TArray<FNiagaraDataSetExecutionInfo, TInlineAllocator<2>> DataSetInfo;
 
 	static uint32 TickCounter;
 
