Caching off socket transform and bone index to avoid costly lookup each frame.
Saves 12-15us GT time from every Skel mesh DI instance tick using specific sockets.

diff --git a/Engine/Plugins/FX/Niagara/Source/Niagara/Classes/NiagaraDataInterfaceSkeletalMesh.h b/Engine/Plugins/FX/Niagara/Source/Niagara/Classes/NiagaraDataInterfaceSkeletalMesh.h
--- a/Engine/Plugins/FX/Niagara/Source/Niagara/Classes/NiagaraDataInterfaceSkeletalMesh.h
+++ b/Engine/Plugins/FX/Niagara/Source/Niagara/Classes/NiagaraDataInterfaceSkeletalMesh.h
@@ -431,7 +431,15 @@ struct FNDISkeletalMesh_InstanceData
 	TArray<int32> SpecificBones;
 
 	/** Name of all the sockets we use. */
-	TArray<FName> SpecificSockets;
+	//TArray<FName> SpecificSockets;
+	struct FCachedSocketInfo
+	{
+		FCachedSocketInfo(): Transform(ForceInit), BoneIdx(INDEX_NONE){}
+		FTransform Transform;
+		int32 BoneIdx;
+	};
+	TArray<FCachedSocketInfo> SpecificSocketInfo;
+
 	/** Bone index of the first socket, sockets are appended to the end of the bone array */
 	int32 SpecificSocketBoneOffset = 0;
 
diff --git a/Engine/Plugins/FX/Niagara/Source/Niagara/Private/NDISkeletalMesh_BoneSampling.cpp b/Engine/Plugins/FX/Niagara/Source/Niagara/Private/NDISkeletalMesh_BoneSampling.cpp
--- a/Engine/Plugins/FX/Niagara/Source/Niagara/Private/NDISkeletalMesh_BoneSampling.cpp
+++ b/Engine/Plugins/FX/Niagara/Source/Niagara/Private/NDISkeletalMesh_BoneSampling.cpp
@@ -461,7 +461,7 @@ void UNiagaraDataInterfaceSkeletalMesh::GetSkinnedBoneData(FVectorVMContext& Con
 	Accessor.Init<TIntegralConstant<int32, 0>, TIntegralConstant<int32, 0>>(InstData);
 
 	const int32 BoneCount = SkinningHandler.GetBoneCount(Accessor, bInterpolated::Value);
-	const int32 BoneAndSocketCount = BoneCount + InstData->SpecificSockets.Num();
+	const int32 BoneAndSocketCount = BoneCount + InstData->SpecificSocketInfo.Num();
 	float InvDt = 1.0f / InstData->DeltaSeconds;
 
 	const TArray<FTransform>& SpecificSocketCurrTransforms = InstData->GetSpecificSocketsCurrBuffer();
@@ -589,7 +589,7 @@ void UNiagaraDataInterfaceSkeletalMesh::GetSpecificSocketCount(FVectorVMContext&
 
 	VectorVM::FExternalFuncRegisterHandler<int32> OutCount(Context);
 
-	const int32 Num = InstData->SpecificSockets.Num();
+	const int32 Num = InstData->SpecificSocketInfo.Num();
 	for (int32 i = 0; i < Context.NumInstances; ++i)
 	{
 		*OutCount.GetDestAndAdvance() = Num;
@@ -697,7 +697,6 @@ void UNiagaraDataInterfaceSkeletalMesh::RandomSpecificSocketBone(FVectorVMContex
 	VectorVM::FUserPtrHandler<FNDISkeletalMesh_InstanceData> InstData(Context);
 
 	VectorVM::FExternalFuncRegisterHandler<int32> OutSocketBone(Context);
-	const TArray<FName>& SpecificSocketsArray = InstData->SpecificSockets;
 	const int32 SpecificSocketBoneOffset = InstData->SpecificSocketBoneOffset;
 
 	int32 Max = SpecificSockets.Num() - 1;
diff --git a/Engine/Plugins/FX/Niagara/Source/Niagara/Private/NiagaraDataInterfaceSkeletalMesh.cpp b/Engine/Plugins/FX/Niagara/Source/Niagara/Private/NiagaraDataInterfaceSkeletalMesh.cpp
--- a/Engine/Plugins/FX/Niagara/Source/Niagara/Private/NiagaraDataInterfaceSkeletalMesh.cpp
+++ b/Engine/Plugins/FX/Niagara/Source/Niagara/Private/NiagaraDataInterfaceSkeletalMesh.cpp
@@ -495,7 +495,7 @@ void FSkeletalMeshGpuSpawnStaticBuffers::Initialise(FNDISkeletalMesh_InstanceDat
 			}
 		}
 
-		NumSpecificSockets = InstData->SpecificSockets.Num();
+		NumSpecificSockets = InstData->SpecificSocketInfo.Num();
 		SpecificSocketBoneOffset = InstData->SpecificSocketBoneOffset;
 	}
 }
@@ -1511,7 +1511,34 @@ bool FNDISkeletalMesh_InstanceData::Init(UNiagaraDataInterfaceSkeletalMesh* Inte
 
 	// Gather specific socket information
 	{
-		SpecificSockets = Interface->SpecificSockets;
+		TArray<FName>& SpecificSockets = Interface->SpecificSockets;
+		SpecificSocketInfo.SetNum(SpecificSockets.Num());
+
+		//-TODO: We may need to handle skinning mode changes here
+		if (NewSkelComp != nullptr)
+		{
+			for (int32 i = 0; i < SpecificSocketInfo.Num(); ++i)
+			{
+				NewSkelComp->GetSocketInfoByName(SpecificSockets[i], SpecificSocketInfo[i].Transform, SpecificSocketInfo[i].BoneIdx);
+			}
+		}
+		else if (Mesh != nullptr)
+		{
+			for (int32 i = 0; i < SpecificSocketInfo.Num(); ++i)
+			{
+				SpecificSocketInfo[i].Transform = FTransform(Mesh->GetComposedRefPoseMatrix(SpecificSockets[i]));
+				SpecificSocketInfo[i].BoneIdx = INDEX_NONE;
+			}
+		}
+		else
+		{
+			for (int32 i = 0; i < SpecificSocketInfo.Num(); ++i)
+			{
+				SpecificSocketInfo[i].Transform = FTransform::Identity;
+				SpecificSocketInfo[i].BoneIdx = INDEX_NONE;
+			}
+		}
+
 		SpecificSocketBoneOffset = Mesh->RefSkeleton.GetNum();
 
 		SpecificSocketTransformsIndex = 0;
@@ -1571,7 +1598,7 @@ bool FNDISkeletalMesh_InstanceData::Init(UNiagaraDataInterfaceSkeletalMesh* Inte
 		BeginInitResource(MeshGpuSpawnStaticBuffers);
 
 		MeshGpuSpawnDynamicBuffers = new FSkeletalMeshGpuDynamicBufferProxy();
-		MeshGpuSpawnDynamicBuffers->Initialise(RefSkel, LODData, SpecificSockets.Num());
+		MeshGpuSpawnDynamicBuffers->Initialise(RefSkel, LODData, SpecificSocketInfo.Num());
 		BeginInitResource(MeshGpuSpawnDynamicBuffers);
 	}
 
@@ -1688,30 +1715,11 @@ void FNDISkeletalMesh_InstanceData::UpdateSpecificSocketTransforms()
 	USkeletalMeshComponent* SkelComp = Cast<USkeletalMeshComponent>(Component.Get());
 	TArray<FTransform>& WriteBuffer = GetSpecificSocketsWriteBuffer();
 
-	//-TODO: We may need to handle skinning mode changes here
-	if (SkelComp != nullptr)
+	for (int32 i = 0; i < SpecificSocketInfo.Num(); ++i)
 	{
-		for (int32 i = 0; i < SpecificSockets.Num(); ++i)
-		{
-			FTransform SocketTransform;
-			int32 BoneIdx;
-			SkelComp->GetSocketInfoByName(SpecificSockets[i], SocketTransform, BoneIdx);
-			WriteBuffer[i] = SocketTransform * SkelComp->GetBoneTransform(BoneIdx, FTransform::Identity);
-		}
-	}
-	else if ( Mesh != nullptr )
-	{
-		for (int32 i = 0; i < SpecificSockets.Num(); ++i)
-		{
-			WriteBuffer[i] = FTransform(Mesh->GetComposedRefPoseMatrix(SpecificSockets[i]));
-		}
-	}
-	else
-	{
-		for (int32 i = 0; i < SpecificSockets.Num(); ++i)
-		{
-			WriteBuffer[i] = FTransform::Identity;
-		}
+		const FCachedSocketInfo& SocketInfo = SpecificSocketInfo[i];
+		const FTransform& BoneTransform = SocketInfo.BoneIdx != INDEX_NONE ? SkelComp->GetBoneTransform(SocketInfo.BoneIdx, FTransform::Identity) : FTransform::Identity;
+		WriteBuffer[i] = SocketInfo.Transform * BoneTransform;
 	}
 }
 
