Fix for cooked GPU emitters grabbing an invalid parameter store resulting in a crash

diff --git a/Engine/Plugins/FX/Niagara/Source/Niagara/Classes/NiagaraScriptExecutionContext.h b/Engine/Plugins/FX/Niagara/Source/Niagara/Classes/NiagaraScriptExecutionContext.h
--- a/Engine/Plugins/FX/Niagara/Source/Niagara/Classes/NiagaraScriptExecutionContext.h
+++ b/Engine/Plugins/FX/Niagara/Source/Niagara/Classes/NiagaraScriptExecutionContext.h
@@ -176,7 +176,7 @@ public:
 
 	bool Init(UNiagaraScript* InScript, ENiagaraSimTarget InTarget);
 	
-	bool Tick(class FNiagaraSystemInstance* Instance, ENiagaraSimTarget SimTarget = ENiagaraSimTarget::CPUSim);
+	bool Tick(class FNiagaraSystemInstance* Instance, ENiagaraSimTarget SimTarget);
 	void PostTick();
 
 	void BindData(int32 Index, FNiagaraDataSet& DataSet, int32 StartInstance, bool bUpdateInstanceCounts);
diff --git a/Engine/Plugins/FX/Niagara/Source/Niagara/Private/NiagaraEmitterInstance.cpp b/Engine/Plugins/FX/Niagara/Source/Niagara/Private/NiagaraEmitterInstance.cpp
--- a/Engine/Plugins/FX/Niagara/Source/Niagara/Private/NiagaraEmitterInstance.cpp
+++ b/Engine/Plugins/FX/Niagara/Source/Niagara/Private/NiagaraEmitterInstance.cpp
@@ -937,8 +937,8 @@ void FNiagaraEmitterInstance::PreTick()
 
 
 	bool bOk = true;
-	bOk &= SpawnExecContext.Tick(ParentSystemInstance);
-	bOk &= UpdateExecContext.Tick(ParentSystemInstance);
+	bOk &= SpawnExecContext.Tick(ParentSystemInstance, CachedEmitter->SimTarget);
+	bOk &= UpdateExecContext.Tick(ParentSystemInstance, CachedEmitter->SimTarget);
 
 	// @todo THREADSAFETY We should not tick GPU contexts on the game thread!
 	if (CachedEmitter->SimTarget == ENiagaraSimTarget::GPUComputeSim && GPUExecContext != nullptr)
@@ -947,7 +947,7 @@ void FNiagaraEmitterInstance::PreTick()
 	}
 	for (FNiagaraScriptExecutionContext& EventContext : EventExecContexts)
 	{
-		bOk &= EventContext.Tick(ParentSystemInstance);
+		bOk &= EventContext.Tick(ParentSystemInstance, CachedEmitter->SimTarget);
 	}
 
 	if (!bOk)
diff --git a/Engine/Plugins/FX/Niagara/Source/Niagara/Private/NiagaraSystemSimulation.cpp b/Engine/Plugins/FX/Niagara/Source/Niagara/Private/NiagaraSystemSimulation.cpp
--- a/Engine/Plugins/FX/Niagara/Source/Niagara/Private/NiagaraSystemSimulation.cpp
+++ b/Engine/Plugins/FX/Niagara/Source/Niagara/Private/NiagaraSystemSimulation.cpp
@@ -1224,7 +1224,7 @@ void FNiagaraSystemSimulation::SpawnSystemInstances(FNiagaraSystemSimulationTick
 	Context.DataSet.GetDestinationDataChecked().SetNumInstances(NumInstances);
 
 	// Run Spawn
-	if (SpawnExecContext.Tick(SoloSystemInstance) == false)
+	if (SpawnExecContext.Tick(SoloSystemInstance, ENiagaraSimTarget::CPUSim) == false)
 	{
 		for (FNiagaraSystemInstance* SystemInst : Context.Instances)
 		{
@@ -1284,7 +1284,7 @@ void FNiagaraSystemSimulation::UpdateSystemInstances(FNiagaraSystemSimulationTic
 		DestinationData.SetNumInstances(NumInstances);
 
 		// Tick UpdateExecContext, this can fail to bind VM functions if this happens we become invalid so mark all instances as disabled
-		if (UpdateExecContext.Tick(Context.Instances[0]) == false)
+		if (UpdateExecContext.Tick(Context.Instances[0], ENiagaraSimTarget::CPUSim) == false)
 		{
 			for (FNiagaraSystemInstance* SystemInst : Context.Instances)
 			{
