- Change Begin / End Render Pass to use compute version, this will still clear bindings and is compatable with all RHIs

diff --git a/Engine/Plugins/FX/Niagara/Source/Niagara/Private/NiagaraEmitterInstanceBatcher.cpp b/Engine/Plugins/FX/Niagara/Source/Niagara/Private/NiagaraEmitterInstanceBatcher.cpp
--- a/Engine/Plugins/FX/Niagara/Source/Niagara/Private/NiagaraEmitterInstanceBatcher.cpp
+++ b/Engine/Plugins/FX/Niagara/Source/Niagara/Private/NiagaraEmitterInstanceBatcher.cpp
@@ -818,61 +818,63 @@ void NiagaraEmitterInstanceBatcher::ExecuteAll(FRHICommandList& RHICmdList, FRHI
 
 	// Clear any RT bindings that we may be using
 	// Note: We can not encapsulate the whole Niagara pass as some DI's may not be compatable (i.e. use CopyTexture function), we need to fix this with future RDG conversion
-	RHICmdList.BeginRenderPass(FRHIRenderPassInfo(), TEXT("NiagaraClearRTs"));
-	RHICmdList.EndRenderPass();
-
-	RHICmdList.BeginUAVOverlap();
-
-	FEmitterInstanceList InstancesWithPersistentIDs;
-	FNiagaraBufferArray OutputGraphicsBuffers;
-
-	for (int32 SimPassIdx = 0; SimPassIdx < SimPasses.Num(); ++SimPassIdx)
+	if (SimPasses.Num() > 0)
 	{
-		FOverlappableTicks& SimPass = SimPasses[SimPassIdx];
-		InstancesWithPersistentIDs.SetNum(0, false);
+		RHICmdList.BeginComputePass(TEXT("NiagaraCompute"));
+		RHICmdList.BeginUAVOverlap();
 
-		// This initial pass gathers all the buffers that are read from and written to so we can do batch resource transitions.
-		// It also ensures the GPU buffers are large enough to hold everything.
-		ResizeBuffersAndGatherResources(SimPass, RHICmdList, OutputGraphicsBuffers, InstancesWithPersistentIDs);
+		FEmitterInstanceList InstancesWithPersistentIDs;
+		FNiagaraBufferArray OutputGraphicsBuffers;
 
+		for (int32 SimPassIdx = 0; SimPassIdx < SimPasses.Num(); ++SimPassIdx)
 		{
-			SCOPED_DRAW_EVENT(RHICmdList, NiagaraGPUSimulation);
-			SCOPED_GPU_STAT(RHICmdList, NiagaraGPUSimulation);
-			DispatchAllOnCompute(SimPass, RHICmdList, ViewUniformBuffer);
-		}
+			FOverlappableTicks& SimPass = SimPasses[SimPassIdx];
+			InstancesWithPersistentIDs.SetNum(0, false);
 
-		if (InstancesWithPersistentIDs.Num() == 0)
-		{
-			continue;
-		}
+			// This initial pass gathers all the buffers that are read from and written to so we can do batch resource transitions.
+			// It also ensures the GPU buffers are large enough to hold everything.
+			ResizeBuffersAndGatherResources(SimPass, RHICmdList, OutputGraphicsBuffers, InstancesWithPersistentIDs);
 
-		// If we're doing multiple ticks (e.g. when scrubbing the timeline in the editor), we must update the free ID buffers before running
-		// the next tick, which will cause stalls (because the ID to index buffer is written by DispatchAllOnCompute and read by UpdateFreeIDBuffers).
-		// However, when we're at the last tick, we can postpone the update until later in the frame and avoid the stall. This will be the case when
-		// running normally, with one tick per frame.
-		if (SimPassIdx < SimPasses.Num() - 1)
-		{
-			ResizeFreeIDsListSizesBuffer(InstancesWithPersistentIDs.Num());
-			ClearFreeIDsListSizesBuffer(RHICmdList);
-			UpdateFreeIDBuffers(RHICmdList, InstancesWithPersistentIDs);
-		}
-		else
-		{
-			DeferredIDBufferUpdates.Append(InstancesWithPersistentIDs);
-			ResizeFreeIDsListSizesBuffer(DeferredIDBufferUpdates.Num());
-
-			// Speculatively clear the list sizes buffer here. Under normal circumstances, this happens in the first stage which finds instances with persistent IDs
-			// (usually PreInitViews) and it's finished by the time the deferred updates need to be processed. If a subsequent tick stage runs multiple time ticks,
-			// the first step will find the buffer already cleared and will not clear again. The only time when this clear is superfluous is when a following stage
-			// reallocates the buffer, but that's unlikely (and amortized) because we allocate in chunks.
-			ClearFreeIDsListSizesBuffer(RHICmdList);
+			{
+				SCOPED_DRAW_EVENT(RHICmdList, NiagaraGPUSimulation);
+				SCOPED_GPU_STAT(RHICmdList, NiagaraGPUSimulation);
+				DispatchAllOnCompute(SimPass, RHICmdList, ViewUniformBuffer);
+			}
+
+			if (InstancesWithPersistentIDs.Num() == 0)
+			{
+				continue;
+			}
+
+			// If we're doing multiple ticks (e.g. when scrubbing the timeline in the editor), we must update the free ID buffers before running
+			// the next tick, which will cause stalls (because the ID to index buffer is written by DispatchAllOnCompute and read by UpdateFreeIDBuffers).
+			// However, when we're at the last tick, we can postpone the update until later in the frame and avoid the stall. This will be the case when
+			// running normally, with one tick per frame.
+			if (SimPassIdx < SimPasses.Num() - 1)
+			{
+				ResizeFreeIDsListSizesBuffer(InstancesWithPersistentIDs.Num());
+				ClearFreeIDsListSizesBuffer(RHICmdList);
+				UpdateFreeIDBuffers(RHICmdList, InstancesWithPersistentIDs);
+			}
+			else
+			{
+				DeferredIDBufferUpdates.Append(InstancesWithPersistentIDs);
+				ResizeFreeIDsListSizesBuffer(DeferredIDBufferUpdates.Num());
+
+				// Speculatively clear the list sizes buffer here. Under normal circumstances, this happens in the first stage which finds instances with persistent IDs
+				// (usually PreInitViews) and it's finished by the time the deferred updates need to be processed. If a subsequent tick stage runs multiple time ticks,
+				// the first step will find the buffer already cleared and will not clear again. The only time when this clear is superfluous is when a following stage
+				// reallocates the buffer, but that's unlikely (and amortized) because we allocate in chunks.
+				ClearFreeIDsListSizesBuffer(RHICmdList);
+			}
 		}
-	}
 
-	OutputGraphicsBuffers.Add(GPUInstanceCounterManager.GetInstanceCountBuffer().UAV);
-	RHICmdList.TransitionResources(EResourceTransitionAccess::EReadable, EResourceTransitionPipeline::EComputeToGfx, OutputGraphicsBuffers.GetData(), OutputGraphicsBuffers.Num());
+		OutputGraphicsBuffers.Add(GPUInstanceCounterManager.GetInstanceCountBuffer().UAV);
+		RHICmdList.TransitionResources(EResourceTransitionAccess::EReadable, EResourceTransitionPipeline::EComputeToGfx, OutputGraphicsBuffers.GetData(), OutputGraphicsBuffers.Num());
 
-	RHICmdList.EndUAVOverlap();
+		RHICmdList.EndUAVOverlap();
+		RHICmdList.EndComputePass();
+	}
 }
 
 void NiagaraEmitterInstanceBatcher::PreInitViews(FRHICommandListImmediate& RHICmdList, bool bAllowGPUParticleUpdate)
