Now bubbling up selection whenever working with Parameter Map Get nodes.

diff --git a/Engine/Plugins/FX/Niagara/Source/NiagaraEditor/Private/NiagaraGraph.cpp b/Engine/Plugins/FX/Niagara/Source/NiagaraEditor/Private/NiagaraGraph.cpp
--- a/Engine/Plugins/FX/Niagara/Source/NiagaraEditor/Private/NiagaraGraph.cpp
+++ b/Engine/Plugins/FX/Niagara/Source/NiagaraEditor/Private/NiagaraGraph.cpp
@@ -1779,6 +1779,12 @@ FDelegateHandle UNiagaraGraph::RegisterPinVisualWidgetProvider(FOnGetPinVisualWi
 	return OnGetPinVisualWidgetDelegate.GetHandle();
 }
 
+void UNiagaraGraph::UnregisterPinVisualWidgetProvider(const FDelegateHandle& InHandle)
+{
+	ensure(OnGetPinVisualWidgetDelegate.GetHandle() == InHandle);
+	OnGetPinVisualWidgetDelegate.Unbind();
+}
+
 const class UEdGraphSchema_Niagara* UNiagaraGraph::GetNiagaraSchema() const
 {
 	return Cast<UEdGraphSchema_Niagara>(GetSchema());
@@ -2016,6 +2022,11 @@ UNiagaraGraph::FOnDataInterfaceChanged& UNiagaraGraph::OnDataInterfaceChanged()
 	return OnDataInterfaceChangedDelegate;
 }
 
+UNiagaraGraph::FOnSubObjectSelectionChanged& UNiagaraGraph::OnSubObjectSelectionChanged()
+{
+	return OnSelectedSubObjectChanged;
+}
+
 void UNiagaraGraph::RefreshParameterReferences() const 
 {
 	// A set of variables to track which parameters are used so that unused parameters can be removed after the reference tracking.
diff --git a/Engine/Plugins/FX/Niagara/Source/NiagaraEditor/Private/NiagaraNodeParameterMapBase.cpp b/Engine/Plugins/FX/Niagara/Source/NiagaraEditor/Private/NiagaraNodeParameterMapBase.cpp
--- a/Engine/Plugins/FX/Niagara/Source/NiagaraEditor/Private/NiagaraNodeParameterMapBase.cpp
+++ b/Engine/Plugins/FX/Niagara/Source/NiagaraEditor/Private/NiagaraNodeParameterMapBase.cpp
@@ -17,6 +17,8 @@
 #include "NiagaraParameterCollection.h"
 #include "Widgets/SNiagaraParameterMapView.h"
 #include "Widgets/Input/SEditableTextBox.h"
+#include "ToolMenus.h"
+#include "NiagaraScriptVariable.h"
 
 #include "IAssetTools.h"
 #include "AssetRegistryModule.h"
@@ -226,4 +228,53 @@ void UNiagaraNodeParameterMapBase::OnPinRenamed(UEdGraphPin* RenamedPin, const F
 
 }
 
+void UNiagaraNodeParameterMapBase::GetNodeContextMenuActions(UToolMenu* Menu, UGraphNodeContextMenuContext* Context) const
+{
+	Super::GetNodeContextMenuActions(Menu, Context);
+
+	UEdGraphPin* Pin = const_cast<UEdGraphPin*>(Context->Pin);
+	if (Pin && !IsAddPin(Pin))
+	{
+
+		FNiagaraVariable Var = CastChecked<UEdGraphSchema_Niagara>(GetSchema())->PinToNiagaraVariable(Pin);
+		const UNiagaraGraph* Graph = GetNiagaraGraph();
+
+		//if (!FNiagaraConstants::IsNiagaraConstant(Var))
+		{
+			FToolMenuSection& Section = Menu->AddSection("EdGraphSchema_NiagaraParamAction", LOCTEXT("EditPinMenuHeader", "Parameters"));
+			Section.AddMenuEntry(
+				"SelectParameter",
+				LOCTEXT("SelectParameterPin", "Select parameter"),
+				LOCTEXT("SelectParameterPinToolTip", "Select this parameter in the paramter panel"),
+				FSlateIcon(),
+				FUIAction(FExecuteAction::CreateUObject(const_cast<UNiagaraNodeParameterMapBase*>(this), &UNiagaraNodeParameterMapBase::SelectParameterFromPin, const_cast<UEdGraphPin*>(Context->Pin))));
+		}
+	}
+}
+
+void UNiagaraNodeParameterMapBase::SelectParameterFromPin(UEdGraphPin* InPin) 
+{
+	UNiagaraGraph* NiagaraGraph = GetNiagaraGraph();
+	if (NiagaraGraph && InPin)
+	{
+		const UEdGraphSchema_Niagara* Schema = Cast<UEdGraphSchema_Niagara>(NiagaraGraph->GetSchema());
+		if (Schema)
+		{
+			if (IsAddPin(InPin))
+			{
+				return;
+			}
+
+			FNiagaraTypeDefinition TypeDef = Schema->PinToTypeDefinition(InPin);
+			FNiagaraVariable PinVariable = FNiagaraVariable(TypeDef, InPin->PinName);
+			UNiagaraScriptVariable** PinAssociatedScriptVariable = NiagaraGraph->GetAllMetaData().Find(PinVariable);
+			if (PinAssociatedScriptVariable != nullptr)
+			{
+				NiagaraGraph->OnSubObjectSelectionChanged().Broadcast(*PinAssociatedScriptVariable);
+			}
+		}
+	}
+}
+
+
 #undef LOCTEXT_NAMESPACE
\ No newline at end of file
diff --git a/Engine/Plugins/FX/Niagara/Source/NiagaraEditor/Private/NiagaraNodeParameterMapBase.h b/Engine/Plugins/FX/Niagara/Source/NiagaraEditor/Private/NiagaraNodeParameterMapBase.h
--- a/Engine/Plugins/FX/Niagara/Source/NiagaraEditor/Private/NiagaraNodeParameterMapBase.h
+++ b/Engine/Plugins/FX/Niagara/Source/NiagaraEditor/Private/NiagaraNodeParameterMapBase.h
@@ -32,6 +32,7 @@ public:
 	/** Called when a pin's description text is committed. */
 	void PinDescriptionTextCommitted(const FText& Text, ETextCommit::Type CommitType, UEdGraphPin* Pin);
 
+	virtual void GetNodeContextMenuActions(class UToolMenu* Menu, class UGraphNodeContextMenuContext* Context) const override;
 	virtual void CollectAddPinActions(FGraphActionListBuilderBase& OutActions, bool& bOutCreateRemainingActions, UEdGraphPin* Pin) override;
 	virtual void GetPinHoverText(const UEdGraphPin& Pin, FString& HoverTextOut) const override;
 
@@ -40,8 +41,10 @@ public:
 	void SetPinName(UEdGraphPin* InPin, const FName& InName);
 
 protected:
+	virtual void SelectParameterFromPin(UEdGraphPin* InPin);
 	virtual void OnPinRenamed(UEdGraphPin* RenamedPin, const FString& OldName) override;
 
+
 	UEdGraphPin* PinPendingRename;
 
 public:
diff --git a/Engine/Plugins/FX/Niagara/Source/NiagaraEditor/Private/NiagaraNodeParameterMapGet.cpp b/Engine/Plugins/FX/Niagara/Source/NiagaraEditor/Private/NiagaraNodeParameterMapGet.cpp
--- a/Engine/Plugins/FX/Niagara/Source/NiagaraEditor/Private/NiagaraNodeParameterMapGet.cpp
+++ b/Engine/Plugins/FX/Niagara/Source/NiagaraEditor/Private/NiagaraNodeParameterMapGet.cpp
@@ -526,28 +526,6 @@ void UNiagaraNodeParameterMapGet::GetNodeContextMenuActions(UToolMenu* Menu, UGr
 {
 	Super::GetNodeContextMenuActions(Menu, Context);
 
-	UEdGraphPin* Pin = const_cast<UEdGraphPin*>(Context->Pin);
-	if (Pin && Pin->Direction == EEdGraphPinDirection::EGPD_Output)
-	{
-		
-		FNiagaraVariable Var = CastChecked<UEdGraphSchema_Niagara>(GetSchema())->PinToNiagaraVariable(Pin);
-		const UNiagaraGraph* Graph = GetNiagaraGraph();
-
-		if (!FNiagaraConstants::IsNiagaraConstant(Var))
-		{
-			FToolMenuSection& Section = Menu->AddSection("EdGraphSchema_NiagaraMetaDataActions", LOCTEXT("EditPinMenuHeader", "Meta-Data"));
-			TSharedRef<SWidget> RenameWidget =
-				SNew(SBox)
-				.WidthOverride(100)
-				.Padding(FMargin(5, 0, 0, 0))
-				[
-					SNew(SEditableTextBox)
-					.Text_UObject(this, &UNiagaraNodeParameterMapBase::GetPinDescriptionText, Pin)
-					.OnTextCommitted_UObject(const_cast<UNiagaraNodeParameterMapGet*>(this), &UNiagaraNodeParameterMapBase::PinDescriptionTextCommitted, Pin)
-				];
-			Section.AddEntry(FToolMenuEntry::InitWidget("RenameWidget", RenameWidget, LOCTEXT("DescMenuItem", "Description")));
-		}
-	}
 }
 
 #undef LOCTEXT_NAMESPACE
diff --git a/Engine/Plugins/FX/Niagara/Source/NiagaraEditor/Private/NiagaraNodeParameterMapSet.cpp b/Engine/Plugins/FX/Niagara/Source/NiagaraEditor/Private/NiagaraNodeParameterMapSet.cpp
--- a/Engine/Plugins/FX/Niagara/Source/NiagaraEditor/Private/NiagaraNodeParameterMapSet.cpp
+++ b/Engine/Plugins/FX/Niagara/Source/NiagaraEditor/Private/NiagaraNodeParameterMapSet.cpp
@@ -290,27 +290,7 @@ void UNiagaraNodeParameterMapSet::GetNodeContextMenuActions(UToolMenu* Menu, UGr
 {
 	Super::GetNodeContextMenuActions(Menu, Context);
 
-	UEdGraphPin* Pin = const_cast<UEdGraphPin*>(Context->Pin);
-	if (Pin && Pin->Direction == EEdGraphPinDirection::EGPD_Input)
-	{
-		FNiagaraVariable Var = CastChecked<UEdGraphSchema_Niagara>(GetSchema())->PinToNiagaraVariable(Pin);
-		const UNiagaraGraph* Graph = GetNiagaraGraph();
 
-		if (!FNiagaraConstants::IsNiagaraConstant(Var))
-		{
-			FToolMenuSection& Section = Menu->AddSection("EdGraphSchema_NiagaraMetaDataActions", LOCTEXT("EditPinMenuHeader", "Meta-Data"));
-			TSharedRef<SWidget> RenameWidget =
-				SNew(SBox)
-				.WidthOverride(100)
-				.Padding(FMargin(5, 0, 0, 0))
-				[
-					SNew(SEditableTextBox)
-					.Text_UObject(this, &UNiagaraNodeParameterMapBase::GetPinDescriptionText, Pin)
-					.OnTextCommitted_UObject(const_cast<UNiagaraNodeParameterMapSet*>(this), &UNiagaraNodeParameterMapBase::PinDescriptionTextCommitted, Pin)
-				];
-			Section.AddEntry(FToolMenuEntry::InitWidget("RenameWidget", RenameWidget, LOCTEXT("DescMenuItem", "Description")));
-		}
-	}
 }
 
 void UNiagaraNodeParameterMapSet::PostLoad()
diff --git a/Engine/Plugins/FX/Niagara/Source/NiagaraEditor/Private/TNiagaraGraphPinEditableName.h b/Engine/Plugins/FX/Niagara/Source/NiagaraEditor/Private/TNiagaraGraphPinEditableName.h
--- a/Engine/Plugins/FX/Niagara/Source/NiagaraEditor/Private/TNiagaraGraphPinEditableName.h
+++ b/Engine/Plugins/FX/Niagara/Source/NiagaraEditor/Private/TNiagaraGraphPinEditableName.h
@@ -119,15 +119,10 @@ protected:
 					}
 					return CreateLabelTextBlock();
 				}
-				else if (ParentNode->IsA<UNiagaraNodeCustomHlsl>())
+				else 
 				{
 					return CreateRenamableLabelTextBlock();
 				}
-				else
-				{
-					ensureMsgf(false, TEXT("Tried to create a pin widget for unhandled node class!"));
-					return CreateRenamableLabelTextBlock();
-				}
 			}
 			else
 			{
@@ -136,6 +131,7 @@ protected:
 		}
 		else
 		{
+			
 			return BaseClass::GetLabelWidget(InLabelStyle);
 		}
 	}
diff --git a/Engine/Plugins/FX/Niagara/Source/NiagaraEditor/Private/ViewModels/NiagaraParameterPanelViewModel.cpp b/Engine/Plugins/FX/Niagara/Source/NiagaraEditor/Private/ViewModels/NiagaraParameterPanelViewModel.cpp
--- a/Engine/Plugins/FX/Niagara/Source/NiagaraEditor/Private/ViewModels/NiagaraParameterPanelViewModel.cpp
+++ b/Engine/Plugins/FX/Niagara/Source/NiagaraEditor/Private/ViewModels/NiagaraParameterPanelViewModel.cpp
@@ -565,6 +565,9 @@ FNiagaraScriptToolkitParameterPanelViewModel::~FNiagaraScriptToolkitParameterPan
 	UNiagaraGraph* NiagaraGraph = static_cast<UNiagaraGraph*>(ScriptViewModel->GetGraphViewModel()->GetGraph());
 	NiagaraGraph->RemoveOnGraphChangedHandler(OnGraphChangedHandle);
 	NiagaraGraph->RemoveOnGraphNeedsRecompileHandler(OnGraphNeedsRecompileHandle);
+	NiagaraGraph->RegisterPinVisualWidgetProvider(nullptr);
+	NiagaraGraph->OnSubObjectSelectionChanged().Remove(OnSubObjectSelectionHandle);
+
 }
 
 void FNiagaraScriptToolkitParameterPanelViewModel::InitBindings()
@@ -574,7 +577,8 @@ void FNiagaraScriptToolkitParameterPanelViewModel::InitBindings()
 		FOnGraphChanged::FDelegate::CreateRaw(this, &FNiagaraScriptToolkitParameterPanelViewModel::HandleOnGraphChanged));
 	OnGraphNeedsRecompileHandle = NiagaraGraph->AddOnGraphNeedsRecompileHandler(
 		FOnGraphChanged::FDelegate::CreateRaw(this, &FNiagaraScriptToolkitParameterPanelViewModel::HandleOnGraphChanged));
-	ScriptViewModel->GetGraphViewModel()->GetGraph()->RegisterPinVisualWidgetProvider(UNiagaraGraph::FOnGetPinVisualWidget::CreateSP(this, &FNiagaraScriptToolkitParameterPanelViewModel::GetScriptParameterVisualWidget));
+	ScriptVisualPinHandle = NiagaraGraph->RegisterPinVisualWidgetProvider(UNiagaraGraph::FOnGetPinVisualWidget::CreateSP(this, &FNiagaraScriptToolkitParameterPanelViewModel::GetScriptParameterVisualWidget));
+	OnSubObjectSelectionHandle = NiagaraGraph->OnSubObjectSelectionChanged().AddSP(this, &FNiagaraScriptToolkitParameterPanelViewModel::HandleGraphSubObjectSelectionChanged);
 }
 
 void FNiagaraScriptToolkitParameterPanelViewModel::Refresh() const
@@ -582,6 +586,11 @@ void FNiagaraScriptToolkitParameterPanelViewModel::Refresh() const
 	OnParameterPanelViewModelRefreshed.ExecuteIfBound();
 }
 
+void FNiagaraScriptToolkitParameterPanelViewModel::HandleGraphSubObjectSelectionChanged(const UObject* Obj)
+{
+	OnParameterPanelViewModelExternalSelectionChanged.Broadcast(Obj);
+}
+
 void FNiagaraScriptToolkitParameterPanelViewModel::CollectStaticSections(TArray<int32>& StaticSectionIDs) const
 {
 	StaticSectionIDs.Add(NiagaraParameterPanelSectionID::INPUTS);
diff --git a/Engine/Plugins/FX/Niagara/Source/NiagaraEditor/Private/ViewModels/NiagaraParameterPanelViewModel.h b/Engine/Plugins/FX/Niagara/Source/NiagaraEditor/Private/ViewModels/NiagaraParameterPanelViewModel.h
--- a/Engine/Plugins/FX/Niagara/Source/NiagaraEditor/Private/ViewModels/NiagaraParameterPanelViewModel.h
+++ b/Engine/Plugins/FX/Niagara/Source/NiagaraEditor/Private/ViewModels/NiagaraParameterPanelViewModel.h
@@ -26,6 +26,7 @@ class INiagaraParameterPanelViewModel : public TSharedFromThis<INiagaraParameter
 public:
 	/** Delegate to signal the view model's state has changed. */
 	DECLARE_DELEGATE(FOnParameterPanelViewModelRefreshed);
+	DECLARE_MULTICAST_DELEGATE_OneParam(FOnParameterPanelViewModelExternalSelectionChanged, const UObject*);
 
 	virtual ~INiagaraParameterPanelViewModel() { }
 
@@ -77,9 +78,10 @@ public:
 
 	FOnParameterPanelViewModelRefreshed& GetOnRefreshed() { return OnParameterPanelViewModelRefreshed; };
 
+	FOnParameterPanelViewModelExternalSelectionChanged& GetExternalSelectionChanged() { return OnParameterPanelViewModelExternalSelectionChanged; }
 protected:
 	FOnParameterPanelViewModelRefreshed OnParameterPanelViewModelRefreshed;
-
+	FOnParameterPanelViewModelExternalSelectionChanged OnParameterPanelViewModelExternalSelectionChanged;
 	/** Cached list of parameters sent to SNiagarParameterPanel, updated whenever GetViewedParameters is called. */
 	TArray<FNiagaraScriptVariableAndViewInfo> CachedViewedParameters;
 
@@ -174,6 +176,8 @@ public:
 	virtual bool CanRenameParameter(const FNiagaraVariable& TargetVariableToRename, const FNiagaraVariableMetaData& TargetVariableMetaData, const FText& NewVariableNameText) const override;
 
 	virtual void HandleActionSelected(const TSharedPtr<FEdGraphSchemaAction>& InAction, ESelectInfo::Type InSelectionType) override;
+	
+	void HandleGraphSubObjectSelectionChanged(const UObject* Obj);
 
 	virtual FReply HandleActionDragged(const TSharedPtr<FEdGraphSchemaAction>& InAction, const FPointerEvent& MouseEvent) const override;
 
@@ -190,6 +194,8 @@ private:
 
 	FDelegateHandle OnGraphChangedHandle;
 	FDelegateHandle OnGraphNeedsRecompileHandle;
+	FDelegateHandle ScriptVisualPinHandle;
+	FDelegateHandle OnSubObjectSelectionHandle;
 
 	void HandleOnGraphChanged(const struct FEdGraphEditAction& InAction);
 
diff --git a/Engine/Plugins/FX/Niagara/Source/NiagaraEditor/Private/Widgets/SNiagaraGraphParameterMapGetNode.cpp b/Engine/Plugins/FX/Niagara/Source/NiagaraEditor/Private/Widgets/SNiagaraGraphParameterMapGetNode.cpp
--- a/Engine/Plugins/FX/Niagara/Source/NiagaraEditor/Private/Widgets/SNiagaraGraphParameterMapGetNode.cpp
+++ b/Engine/Plugins/FX/Niagara/Source/NiagaraEditor/Private/Widgets/SNiagaraGraphParameterMapGetNode.cpp
@@ -7,6 +7,8 @@
 #include "Rendering/DrawElements.h"
 #include "Widgets/SBoxPanel.h"
 #include "SGraphPin.h"
+#include "EdGraphSchema_Niagara.h"
+#include "NiagaraScriptVariable.h"
 
 #define LOCTEXT_NAMESPACE "SNiagaraGraphParameterMapGetNode"
 
@@ -139,6 +141,7 @@ void SNiagaraGraphParameterMapGetNode::CreatePinWidgets()
 			.HAlign(HAlign_Fill)
 			.VAlign(VAlign_Fill)
 			.Padding(FMargin(0, 3))
+			.OnMouseButtonDown(this, &SNiagaraGraphParameterMapGetNode::OnBorderMouseButtonDown, i)
 			[
 				Widget.ToSharedRef()
 			];
@@ -153,4 +156,33 @@ const FSlateBrush* SNiagaraGraphParameterMapGetNode::GetBackgroundBrush(TSharedP
 {
 	return Border->IsHovered() ? BackgroundHoveredBrush	: BackgroundBrush;
 }
+
+
+FReply SNiagaraGraphParameterMapGetNode::OnBorderMouseButtonDown(const FGeometry& MyGeometry, const FPointerEvent& MouseEvent, int32 InWhichPin)
+{
+	if (InWhichPin >= 0 && InWhichPin < OutputPins.Num() + 1)
+	{
+		UNiagaraNodeParameterMapGet* GetNode = Cast<UNiagaraNodeParameterMapGet>(GraphNode);
+		if (GetNode)
+		{
+			UNiagaraGraph* Graph = GetNode->GetNiagaraGraph();
+			if (Graph && InWhichPin > 0)
+			{
+				const UEdGraphSchema_Niagara* Schema = Graph->GetNiagaraSchema();
+				if (Schema)
+				{
+					FNiagaraVariable Var = Schema->PinToNiagaraVariable(OutputPins[InWhichPin-1]->GetPinObj());
+					UNiagaraScriptVariable** PinAssociatedScriptVariable = Graph->GetAllMetaData().Find(Var);
+					if (PinAssociatedScriptVariable != nullptr)
+					{
+						Graph->OnSubObjectSelectionChanged().Broadcast(*PinAssociatedScriptVariable);
+					}
+				}
+			}
+		}
+
+	}
+	return FReply::Handled();
+}
+
 #undef LOCTEXT_NAMESPACE
diff --git a/Engine/Plugins/FX/Niagara/Source/NiagaraEditor/Private/Widgets/SNiagaraGraphParameterMapGetNode.h b/Engine/Plugins/FX/Niagara/Source/NiagaraEditor/Private/Widgets/SNiagaraGraphParameterMapGetNode.h
--- a/Engine/Plugins/FX/Niagara/Source/NiagaraEditor/Private/Widgets/SNiagaraGraphParameterMapGetNode.h
+++ b/Engine/Plugins/FX/Niagara/Source/NiagaraEditor/Private/Widgets/SNiagaraGraphParameterMapGetNode.h
@@ -16,6 +16,7 @@ public:
 	virtual void CreatePinWidgets() override;
 
 protected:
+	FReply OnBorderMouseButtonDown(const FGeometry& MyGeometry, const FPointerEvent& MouseEvent, int32 InWhichPin);
 	const FSlateBrush* GetBackgroundBrush(TSharedPtr<SWidget> Border) const;
 	TSharedPtr<SVerticalBox> PinContainerRoot;
 
diff --git a/Engine/Plugins/FX/Niagara/Source/NiagaraEditor/Private/Widgets/SNiagaraParameterPanel.cpp b/Engine/Plugins/FX/Niagara/Source/NiagaraEditor/Private/Widgets/SNiagaraParameterPanel.cpp
--- a/Engine/Plugins/FX/Niagara/Source/NiagaraEditor/Private/Widgets/SNiagaraParameterPanel.cpp
+++ b/Engine/Plugins/FX/Niagara/Source/NiagaraEditor/Private/Widgets/SNiagaraParameterPanel.cpp
@@ -38,6 +38,7 @@ void SNiagaraParameterPanel::Construct(const FArguments& InArgs, const TSharedPt
 	ToolkitCommands = InToolkitCommands;
 
 	ParameterPanelViewModel->GetOnRefreshed().BindRaw(this, &SNiagaraParameterPanel::Refresh);
+	ParameterPanelViewModel->GetExternalSelectionChanged().AddRaw(this, &SNiagaraParameterPanel::HandleExternalSelectionChanged);
 
 	AddParameterButtons.SetNum(NiagaraParameterPanelSectionID::CUSTOM + 1); //@Todo(ng) verify
 
@@ -94,6 +95,18 @@ void SNiagaraParameterPanel::Tick(const FGeometry& AllottedGeometry, const doubl
 	}
 }
 
+void SNiagaraParameterPanel::HandleExternalSelectionChanged(const UObject* Obj)
+{
+	if (Obj && Obj->IsA< UNiagaraScriptVariable>())
+	{
+		const UNiagaraScriptVariable* Var = Cast< UNiagaraScriptVariable>(Obj);
+		if (Var)
+		{
+			GraphActionMenu->SelectItemByName(Var->Variable.GetName());
+		}
+	}
+}
+
 void SNiagaraParameterPanel::AddParameter(FNiagaraVariable NewVariable, const NiagaraParameterPanelSectionID::Type SectionID)
 {
 	//FNiagaraVariableMetaData NewVariableMetaData = FNiagaraVariableMetaData();
diff --git a/Engine/Plugins/FX/Niagara/Source/NiagaraEditor/Private/Widgets/SNiagaraParameterPanel.h b/Engine/Plugins/FX/Niagara/Source/NiagaraEditor/Private/Widgets/SNiagaraParameterPanel.h
--- a/Engine/Plugins/FX/Niagara/Source/NiagaraEditor/Private/Widgets/SNiagaraParameterPanel.h
+++ b/Engine/Plugins/FX/Niagara/Source/NiagaraEditor/Private/Widgets/SNiagaraParameterPanel.h
@@ -103,6 +103,7 @@ private:
 	bool CanRequestRenameOnActionNode() const;
 	void HandlePaletteItemParameterRenamed(const FText& InText, const FNiagaraScriptVarAndViewInfoAction& InAction);
 	void HandlePaletteItemScopeComboSelectionChanged(ENiagaraParameterScope InScope, const FNiagaraScriptVarAndViewInfoAction& InAction);
+	void HandleExternalSelectionChanged(const UObject* Obj);
 
 	/** Delegate handler used to match an FName to an action in the list, used for renaming keys */
 	bool HandleActionMatchesName(struct FEdGraphSchemaAction* InAction, const FName& InName) const;
diff --git a/Engine/Plugins/FX/Niagara/Source/NiagaraEditor/Public/NiagaraGraph.h b/Engine/Plugins/FX/Niagara/Source/NiagaraEditor/Public/NiagaraGraph.h
--- a/Engine/Plugins/FX/Niagara/Source/NiagaraEditor/Public/NiagaraGraph.h
+++ b/Engine/Plugins/FX/Niagara/Source/NiagaraEditor/Public/NiagaraGraph.h
@@ -134,6 +134,7 @@ class UNiagaraGraph : public UEdGraph
 
 	DECLARE_MULTICAST_DELEGATE(FOnDataInterfaceChanged);
 	DECLARE_DELEGATE_RetVal_OneParam(TSharedRef<SWidget>, FOnGetPinVisualWidget, const UEdGraphPin*);
+	DECLARE_MULTICAST_DELEGATE_OneParam(FOnSubObjectSelectionChanged, const UObject*);
 
 	//~ Begin UObject Interface
 	virtual void PostLoad() override;
@@ -333,6 +334,9 @@ class UNiagaraGraph : public UEdGraph
 	/** Gets a delegate which is called whenever a contained data interfaces changes. */
 	FOnDataInterfaceChanged& OnDataInterfaceChanged();
 
+	/** Gets a delegate which is called whenever a custom subobject in the graph is selected*/
+	FOnSubObjectSelectionChanged& OnSubObjectSelectionChanged();
+
 	void ForceGraphToRecompileOnNextCheck();
 
 	/** Add a listener for OnGraphNeedsRecompile events */
@@ -365,6 +369,7 @@ class UNiagaraGraph : public UEdGraph
 	TSharedRef<SWidget> GetPinVisualWidget(const UEdGraphPin* Pin) const;
 
 	FDelegateHandle RegisterPinVisualWidgetProvider(FOnGetPinVisualWidget OnGetPinVisualWidget);
+	void UnregisterPinVisualWidgetProvider(const FDelegateHandle& InHandle);
 
 	void ScriptVariableChanged(FNiagaraVariable Variable);
 
@@ -444,6 +449,7 @@ private:
 	mutable TMap<FNiagaraVariable, FNiagaraGraphParameterReferenceCollection> ParameterToReferencesMap;
 
 	FOnDataInterfaceChanged OnDataInterfaceChangedDelegate;
+	FOnSubObjectSelectionChanged OnSelectedSubObjectChanged;
 
 	/** Whether currently renaming a parameter to prevent recursion. */
 	bool bIsRenamingParameter;
