Fix 12.20 build

diff --git a/Engine/Plugins/FX/Niagara/Source/Niagara/Classes/NiagaraDataInterfaceSkeletalMesh.h b/Engine/Plugins/FX/Niagara/Source/Niagara/Classes/NiagaraDataInterfaceSkeletalMesh.h
--- a/Engine/Plugins/FX/Niagara/Source/Niagara/Classes/NiagaraDataInterfaceSkeletalMesh.h
+++ b/Engine/Plugins/FX/Niagara/Source/Niagara/Classes/NiagaraDataInterfaceSkeletalMesh.h
@@ -434,7 +434,7 @@ struct FNDISkeletalMesh_InstanceData
 	//TArray<FName> SpecificSockets;
 	struct FCachedSocketInfo
 	{
-		FCachedSocketInfo(): Transform(ForceInit), BoneIdx(INDEX_NONE){}
+		FCachedSocketInfo() : BoneIdx(INDEX_NONE){}
 		FTransform Transform;
 		int32 BoneIdx;
 	};
