- Integrate 11455813 - Fix uninitialized variable BaseVMScript

diff --git a/Engine/Plugins/FX/Niagara/Source/NiagaraShader/Public/NiagaraShared.h b/Engine/Plugins/FX/Niagara/Source/NiagaraShader/Public/NiagaraShared.h
--- a/Engine/Plugins/FX/Niagara/Source/NiagaraShader/Public/NiagaraShared.h
+++ b/Engine/Plugins/FX/Niagara/Source/NiagaraShader/Public/NiagaraShared.h
@@ -563,12 +563,13 @@ public:
 	/**
 	 * Minimal initialization constructor.
 	 */
-	FNiagaraShaderScript() :
-		GameThreadShaderMap(NULL),
-		RenderingThreadShaderMap(NULL),
-		FeatureLevel(GMaxRHIFeatureLevel),
-		bLoadedCookedShaderMapId(false),
-		bLoadedFromCookedMaterial(false)
+	FNiagaraShaderScript()
+		: BaseVMScript(nullptr)
+		, GameThreadShaderMap(nullptr)
+		, RenderingThreadShaderMap(nullptr)
+		, FeatureLevel(GMaxRHIFeatureLevel)
+		, bLoadedCookedShaderMapId(false)
+		, bLoadedFromCookedMaterial(false)
 	{}
 
 	/**
