- Integrate 11451196 - Fix mesh renderer stats so we don't increment per section

diff --git a/Engine/Plugins/FX/Niagara/Source/Niagara/Private/NiagaraRendererMeshes.cpp b/Engine/Plugins/FX/Niagara/Source/Niagara/Private/NiagaraRendererMeshes.cpp
--- a/Engine/Plugins/FX/Niagara/Source/Niagara/Private/NiagaraRendererMeshes.cpp
+++ b/Engine/Plugins/FX/Niagara/Source/Niagara/Private/NiagaraRendererMeshes.cpp
@@ -459,6 +459,10 @@ void FNiagaraRendererMeshes::GetDynamicMeshElements(const TArray<const FSceneVie
 				CollectorResources.UniformBuffer = FNiagaraMeshUniformBufferRef::CreateUniformBufferImmediate(PerViewUniformParameters, UniformBuffer_SingleFrame);
 				CollectorResources.VertexFactory->SetUniformBuffer(CollectorResources.UniformBuffer);
 
+				// Increment stats
+				INC_DWORD_STAT_BY(STAT_NiagaraNumMeshVerts, NumInstances * LODModel.GetNumVertices());
+				INC_DWORD_STAT_BY(STAT_NiagaraNumMeshes, NumInstances);
+
 				// GPU mesh rendering currently only supports one mesh section.
 				// TODO: Add proper support for multiple mesh sections for GPU mesh particles.
 				const int32 SectionCount = LODModel.Sections.Num();
@@ -535,9 +539,6 @@ void FNiagaraRendererMeshes::GetDynamicMeshElements(const TArray<const FSceneVie
 					Mesh.bUseWireframeSelectionColoring = SceneProxy->IsSelected();
 
 					Collector.AddMesh(ViewIndex, Mesh);
-
-					INC_DWORD_STAT_BY(STAT_NiagaraNumMeshVerts, NumInstances * LODModel.GetNumVertices());
-					INC_DWORD_STAT_BY(STAT_NiagaraNumMeshes, NumInstances);
 				}
 			}
 		}
