- Integrate 11448519 - Fix an issue if a GPU emitter was disabled on the first frame and never ticked

diff --git a/Engine/Plugins/FX/Niagara/Source/Niagara/Classes/NiagaraEmitterInstance.h b/Engine/Plugins/FX/Niagara/Source/Niagara/Classes/NiagaraEmitterInstance.h
--- a/Engine/Plugins/FX/Niagara/Source/Niagara/Classes/NiagaraEmitterInstance.h
+++ b/Engine/Plugins/FX/Niagara/Source/Niagara/Classes/NiagaraEmitterInstance.h
@@ -119,6 +119,9 @@ public:
 
 	bool FindBinding(const FNiagaraUserParameterBinding& InBinding, TArray<UMaterialInterface*>& OutMaterials) const;
 
+
+	bool HasTicked() const { return TickCount > 0;  }
+
 private:
 	void CheckForErrors();
 
diff --git a/Engine/Plugins/FX/Niagara/Source/Niagara/Private/NiagaraScriptExecutionContext.cpp b/Engine/Plugins/FX/Niagara/Source/Niagara/Private/NiagaraScriptExecutionContext.cpp
--- a/Engine/Plugins/FX/Niagara/Source/Niagara/Private/NiagaraScriptExecutionContext.cpp
+++ b/Engine/Plugins/FX/Niagara/Source/Niagara/Private/NiagaraScriptExecutionContext.cpp
@@ -380,8 +380,10 @@ void FNiagaraGPUSystemTick::Init(FNiagaraSystemInstance* InSystemInstance)
 	{
 		FNiagaraEmitterInstance* Emitter = &InSystemInstance->GetEmitters()[EmitterIdx].Get();
 
-		if (Emitter && Emitter->GetCachedEmitter()->SimTarget == ENiagaraSimTarget::GPUComputeSim && Emitter->GetGPUContext() != nullptr && Emitter->GetExecutionState() != ENiagaraExecutionState::Complete)
+		if (Emitter && Emitter->GetCachedEmitter()->SimTarget == ENiagaraSimTarget::GPUComputeSim && Emitter->GetGPUContext() != nullptr && !Emitter->IsComplete())
 		{
+			check(Emitter->HasTicked() == true);
+			
 			FNiagaraComputeExecutionContext* GPUContext = Emitter->GetGPUContext();
 
 			FNiagaraComputeInstanceData* InstanceData = new (&Instances[InstanceIndex]) FNiagaraComputeInstanceData;
diff --git a/Engine/Plugins/FX/Niagara/Source/Niagara/Private/NiagaraSystemInstance.cpp b/Engine/Plugins/FX/Niagara/Source/Niagara/Private/NiagaraSystemInstance.cpp
--- a/Engine/Plugins/FX/Niagara/Source/Niagara/Private/NiagaraSystemInstance.cpp
+++ b/Engine/Plugins/FX/Niagara/Source/Niagara/Private/NiagaraSystemInstance.cpp
@@ -2020,7 +2020,7 @@ void FNiagaraSystemInstance::Tick_Concurrent()
 			Inst.Tick(CachedDeltaSeconds);
 		}
 
-		if (Inst.GetCachedEmitter() && Inst.GetCachedEmitter()->SimTarget == ENiagaraSimTarget::GPUComputeSim && Inst.GetGPUContext() != nullptr && (Inst.GetExecutionState() != ENiagaraExecutionState::Complete))
+		if (Inst.GetCachedEmitter() && Inst.GetCachedEmitter()->SimTarget == ENiagaraSimTarget::GPUComputeSim && Inst.GetGPUContext() != nullptr && !Inst.IsComplete())
 		{
 			if (FirstGpuEmitter)
 			{
