- Integrate 11411820 - Track GPU particle memory

diff --git a/Engine/Plugins/FX/Niagara/Source/Niagara/Private/NiagaraDataSet.cpp b/Engine/Plugins/FX/Niagara/Source/Niagara/Private/NiagaraDataSet.cpp
--- a/Engine/Plugins/FX/Niagara/Source/Niagara/Private/NiagaraDataSet.cpp
+++ b/Engine/Plugins/FX/Niagara/Source/Niagara/Private/NiagaraDataSet.cpp
@@ -507,6 +507,8 @@ FNiagaraDataBuffer::~FNiagaraDataBuffer()
 	// in which case this would have been released on the game thread and not from the batcher DataSetsToDestroy_RT.
 	check(!IsInRenderingThread() || GPUInstanceCountBufferOffset == INDEX_NONE);
 	DEC_MEMORY_STAT_BY(STAT_NiagaraParticleMemory, FloatData.GetAllocatedSize() + Int32Data.GetAllocatedSize());
+
+	DEC_MEMORY_STAT_BY(STAT_NiagaraGPUParticleMemory, GPUBufferFloat.NumBytes + GPUBufferInt.NumBytes + GPUIDToIndexTable.NumBytes);
 }
 
 int32 FNiagaraDataBuffer::TransferInstance(FNiagaraDataBuffer& SourceBuffer, int32 InstanceIndex, bool bRemoveFromSource)
@@ -709,6 +711,8 @@ void FNiagaraDataBuffer::AllocateGPU(uint32 InNumInstances, FNiagaraGPUInstanceC
 	Int32Stride = PaddedNumInstances * sizeof(int32);
 	HalfStride = PaddedNumInstances * sizeof(FFloat16);
 
+	DEC_MEMORY_STAT_BY(STAT_NiagaraGPUParticleMemory, GPUBufferFloat.NumBytes + GPUBufferInt.NumBytes + GPUIDToIndexTable.NumBytes);
+
 	// This never seems to happen because of the +1 in NiagaraEmitterInstanceBatcher::ResizeBuffersAndGatherResources()
 	if (PaddedNumInstances == 0)
 	{
@@ -788,6 +792,7 @@ void FNiagaraDataBuffer::AllocateGPU(uint32 InNumInstances, FNiagaraGPUInstanceC
 			}
 		}
 	}
+	INC_MEMORY_STAT_BY(STAT_NiagaraGPUParticleMemory, GPUBufferFloat.NumBytes + GPUBufferInt.NumBytes + GPUIDToIndexTable.NumBytes);
 }
 
 void FNiagaraDataBuffer::SwapInstances(uint32 OldIndex, uint32 NewIndex) 
diff --git a/Engine/Plugins/FX/Niagara/Source/Niagara/Private/NiagaraStats.h b/Engine/Plugins/FX/Niagara/Source/Niagara/Private/NiagaraStats.h
--- a/Engine/Plugins/FX/Niagara/Source/Niagara/Private/NiagaraStats.h
+++ b/Engine/Plugins/FX/Niagara/Source/Niagara/Private/NiagaraStats.h
@@ -15,6 +15,7 @@ DECLARE_CYCLE_STAT(TEXT("Constant Setup [CNC]"), STAT_NiagaraConstants, STATGROU
 DECLARE_CYCLE_STAT(TEXT("Emitter Tick [CNC]"), STAT_NiagaraTick, STATGROUP_Niagara);
 DECLARE_DWORD_COUNTER_STAT(TEXT("NumSystems"), STAT_NiagaraNumSystems, STATGROUP_Niagara);
 DECLARE_MEMORY_STAT(TEXT("Niagara particle data memory"), STAT_NiagaraParticleMemory, STATGROUP_Niagara);
+DECLARE_MEMORY_STAT(TEXT("Niagara GPU particle data memory"), STAT_NiagaraGPUParticleMemory, STATGROUP_Niagara);
 DECLARE_MEMORY_STAT(TEXT("Niagara vertex buffer memory"), STAT_NiagaraVBMemory, STATGROUP_Niagara);
 
 
