- Integrate 11432922 & 11430508 - Use FriendlyName with Niagara Shader Compile Jobs

diff --git a/Engine/Plugins/FX/Niagara/Source/Niagara/Classes/NiagaraScript.h b/Engine/Plugins/FX/Niagara/Source/Niagara/Classes/NiagaraScript.h
--- a/Engine/Plugins/FX/Niagara/Source/Niagara/Classes/NiagaraScript.h
+++ b/Engine/Plugins/FX/Niagara/Source/Niagara/Classes/NiagaraScript.h
@@ -668,8 +668,9 @@ public:
 		false if not.*/
 	NIAGARA_API bool SynchronizeExecutablesWithMaster(const UNiagaraScript* Script, const TMap<FString, FString>& RenameMap);
 
-	NIAGARA_API void SyncAliases(const TMap<FString, FString>& RenameMap);
+	NIAGARA_API FString GetFriendlyName() const;
 
+	NIAGARA_API void SyncAliases(const TMap<FString, FString>& RenameMap);
 #endif
 	
 	UFUNCTION()
diff --git a/Engine/Plugins/FX/Niagara/Source/Niagara/Private/NiagaraScript.cpp b/Engine/Plugins/FX/Niagara/Source/Niagara/Private/NiagaraScript.cpp
--- a/Engine/Plugins/FX/Niagara/Source/Niagara/Private/NiagaraScript.cpp
+++ b/Engine/Plugins/FX/Niagara/Source/Niagara/Private/NiagaraScript.cpp
@@ -1783,7 +1783,7 @@ void UNiagaraScript::CacheResourceShadersForCooking(EShaderPlatform ShaderPlatfo
 
 			NewResource->SetScript(this, (ERHIFeatureLevel::Type)TargetFeatureLevel, CachedScriptVMId.CompilerVersionID, CachedScriptVMId.AdditionalDefines,
 				CachedScriptVMId.BaseScriptCompileHash,	CachedScriptVMId.ReferencedCompileHashes, 
-				CachedScriptVMId.bUsesRapidIterationParams, GetFullName());
+				CachedScriptVMId.bUsesRapidIterationParams, GetFriendlyName());
 			ResourceToCache = NewResource;
 
 			check(ResourceToCache);
@@ -1850,7 +1850,7 @@ void UNiagaraScript::CacheResourceShadersForRendering(bool bRegenerateId, bool b
 			ERHIFeatureLevel::Type CacheFeatureLevel = GMaxRHIFeatureLevel;
 			ScriptResource.SetScript(this, FeatureLevel, CachedScriptVMId.CompilerVersionID, CachedScriptVMId.AdditionalDefines,
 				CachedScriptVMId.BaseScriptCompileHash, CachedScriptVMId.ReferencedCompileHashes, 
-				CachedScriptVMId.bUsesRapidIterationParams, GetFullName());
+				CachedScriptVMId.bUsesRapidIterationParams, GetFriendlyName());
 
 			//if (ScriptResourcesByFeatureLevel[FeatureLevel])
 			{
@@ -1870,6 +1870,21 @@ void UNiagaraScript::CacheResourceShadersForRendering(bool bRegenerateId, bool b
 	}
 }
 
+FString UNiagaraScript::GetFriendlyName() const
+{
+	UEnum* ENiagaraScriptUsageEnum = StaticEnum<ENiagaraScriptUsage>();
+
+	UNiagaraEmitter* EmitterObject = GetTypedOuter<UNiagaraEmitter>();
+	UObject* SystemObject = EmitterObject != nullptr ? EmitterObject->GetOuter() : nullptr;
+	FString FriendlyName = FString::Printf(TEXT("%s/%s/%s"),
+		SystemObject ? *FPaths::MakeValidFileName(SystemObject->GetName()) : TEXT("UnknownSystem"),
+		EmitterObject ? *FPaths::MakeValidFileName(EmitterObject->GetUniqueEmitterName()) : TEXT("UnknownEmitter"),
+		ENiagaraScriptUsageEnum ? *FPaths::MakeValidFileName(ENiagaraScriptUsageEnum->GetNameStringByValue((int64)Usage)) : TEXT("UnknownEnum")
+	);
+
+	return FriendlyName;
+}
+
 void UNiagaraScript::SyncAliases(const TMap<FString, FString>& RenameMap)
 {
 	// First handle any rapid iteration parameters...
diff --git a/Engine/Plugins/FX/Niagara/Source/NiagaraShader/Private/NiagaraShader.cpp b/Engine/Plugins/FX/Niagara/Source/NiagaraShader/Private/NiagaraShader.cpp
--- a/Engine/Plugins/FX/Niagara/Source/NiagaraShader/Private/NiagaraShader.cpp
+++ b/Engine/Plugins/FX/Niagara/Source/NiagaraShader/Private/NiagaraShader.cpp
@@ -400,9 +400,8 @@ TSharedRef<FShaderCommonCompileJob, ESPMode::ThreadSafe> FNiagaraShaderType::Beg
 	SetupCompileEnvironment(Platform, Script, ShaderEnvironment);
 
 	TSharedRef<FShaderCommonCompileJob, ESPMode::ThreadSafe> ShaderJob(NewJob);
-	FString DebugGroupName = Script->SourceName;
 	::GlobalBeginCompileShader(
-		DebugGroupName,
+		Script->GetFriendlyName(),
 		nullptr,
 		this,
 		nullptr,//ShaderPipeline,
