Early exit for system simulations with no work.

diff --git a/Engine/Plugins/FX/Niagara/Source/Niagara/Private/NiagaraSystemSimulation.cpp b/Engine/Plugins/FX/Niagara/Source/Niagara/Private/NiagaraSystemSimulation.cpp
--- a/Engine/Plugins/FX/Niagara/Source/Niagara/Private/NiagaraSystemSimulation.cpp
+++ b/Engine/Plugins/FX/Niagara/Source/Niagara/Private/NiagaraSystemSimulation.cpp
@@ -633,6 +633,11 @@ void FNiagaraSystemSimulation::FlushTickBatch(FNiagaraSystemSimulationTickContex
 /** First phase of system sim tick. Must run on GameThread. */
 void FNiagaraSystemSimulation::Tick_GameThread(float DeltaSeconds, const FGraphEventRef& MyCompletionGraphEvent)
 {
+	if (PendingSystemInstances.Num() == 0 && SystemInstances.Num() == 0)
+	{
+		return;
+	}
+
 	check(IsInGameThread());
 	check(bInSpawnPhase == false);
 
