- Remove early out as it skip PostTick which will not clear events resulting in continual spawning

diff --git a/Engine/Plugins/FX/Niagara/Source/Niagara/Private/NiagaraEmitterInstance.cpp b/Engine/Plugins/FX/Niagara/Source/Niagara/Private/NiagaraEmitterInstance.cpp
--- a/Engine/Plugins/FX/Niagara/Source/Niagara/Private/NiagaraEmitterInstance.cpp
+++ b/Engine/Plugins/FX/Niagara/Source/Niagara/Private/NiagaraEmitterInstance.cpp
@@ -1121,13 +1121,6 @@ void FNiagaraEmitterInstance::Tick(float DeltaSeconds)
 	int32 OrigNumParticles = GetNumParticles();
 	int32 AllocationEstimate = CachedEmitter->GetMaxParticleCountEstimate();
 	int32 RequiredSize = OrigNumParticles + SpawnTotal + EventSpawnTotal;
-
-	if (RequiredSize == 0)
-	{
-		//Early out if we have no particles to process.
-		return;
-	}
-
 	int32 AllocationSize = FMath::Max<int32>(AllocationEstimate, RequiredSize);
 	if (AllocationSize > MaxAllocationCount)
 	{
