- Calculate max number of instances requires across all sim passes
- This reduces memory pressure when sequencer is scrubbing the timeline and forcing many simulation ticks at once as we will not progressively grow our buffers but instead do it once up front
- In my test it changed 149 GPU allocation into 2
- Recreated from CL 11411824 & 11880012

diff --git a/Engine/Plugins/FX/Niagara/Source/Niagara/Classes/NiagaraScriptExecutionContext.h b/Engine/Plugins/FX/Niagara/Source/Niagara/Classes/NiagaraScriptExecutionContext.h
--- a/Engine/Plugins/FX/Niagara/Source/Niagara/Classes/NiagaraScriptExecutionContext.h
+++ b/Engine/Plugins/FX/Niagara/Source/Niagara/Classes/NiagaraScriptExecutionContext.h
@@ -281,6 +281,8 @@ public:
 
 	/** Temp data used in NiagaraEmitterInstanceBatcher::ExecuteAll() to avoid creating a map per FNiagaraComputeExecutionContext */
 	mutable int32 ScratchIndex = INDEX_NONE;
+	mutable uint32 ScratchNumInstances = 0;
+	mutable uint32 ScratchMaxInstances = 0;
 
 	TArray < FSimulationStageMetaData> SimStageInfo;
 
diff --git a/Engine/Plugins/FX/Niagara/Source/Niagara/Private/NiagaraEmitterInstanceBatcher.cpp b/Engine/Plugins/FX/Niagara/Source/Niagara/Private/NiagaraEmitterInstanceBatcher.cpp
--- a/Engine/Plugins/FX/Niagara/Source/Niagara/Private/NiagaraEmitterInstanceBatcher.cpp
+++ b/Engine/Plugins/FX/Niagara/Source/Niagara/Private/NiagaraEmitterInstanceBatcher.cpp
@@ -433,7 +433,7 @@ void NiagaraEmitterInstanceBatcher::ResizeBuffersAndGatherResources(FOverlappabl
 			//We must assume all particles survive when allocating here. 
 			//If this is not true, the read back in ResolveDatasetWrites will shrink the buffers.
 			const uint32 RequiredInstances = FMath::Max(PrevNumInstances, NewNumInstances);
-			const uint32 AllocatedInstances = FMath::Max(RequiredInstances, Instance.SpawnInfo.MaxParticleCount);
+			const uint32 AllocatedInstances = FMath::Max(RequiredInstances, Context->ScratchMaxInstances);
 
 			if (bRequiresPersistentIDs)
 			{
@@ -748,21 +748,47 @@ void NiagaraEmitterInstanceBatcher::ExecuteAll(FRHICommandList& RHICmdList, FRHI
 		for (FNiagaraGPUSystemTick& Tick : Ticks_RT)
 		{
 			FNiagaraComputeInstanceData* Data = Tick.GetInstanceData();
-			FNiagaraComputeExecutionContext* Context = Data->Context;
+			FNiagaraComputeExecutionContext* SharedContext = Data->Context;
 			// This assumes all emitters fallback to the same FNiagaraShaderScript*.
-			FNiagaraShaderRef ComputeShader = Context->GPUScript_RT->GetShader();
+			FNiagaraShaderRef ComputeShader = SharedContext->GPUScript_RT->GetShader();
 			if (ComputeShader.IsNull() || !ShouldTickForStage(Tick, TickStage))
 			{
 				continue;
 			}
 
 			Tick.bIsFinalTick = false; // @todo : this is true sometimes, needs investigation
-			if (Context->ScratchIndex == INDEX_NONE)
+
+			const bool bResetCounts = SharedContext->ScratchIndex == INDEX_NONE;
+			if (bResetCounts)
 			{
-				RelevantContexts.Add(Context);
+				RelevantContexts.Add(SharedContext);
 			}
 			// Here scratch index represent the index of the last tick
-			Context->ScratchIndex = RelevantTicks.Add(&Tick);
+			SharedContext->ScratchIndex = RelevantTicks.Add(&Tick);
+
+			// Allows us to count total required instances across all frames
+			for (uint32 i = 0; i < Tick.Count; ++i)
+			{
+				FNiagaraComputeInstanceData& InstanceData = Tick.GetInstanceData()[i];
+				FNiagaraComputeExecutionContext* ExecContext = InstanceData.Context;
+				if (ExecContext == nullptr)
+				{
+					continue;
+				}
+
+				uint32 PrevNumInstances = 0;
+				if (bResetCounts)
+				{
+					ExecContext->ScratchMaxInstances = InstanceData.SpawnInfo.MaxParticleCount;
+					PrevNumInstances = Tick.bNeedsReset ? 0 : ExecContext->MainDataSet->GetCurrentData()->GetNumInstances();
+				}
+				else
+				{
+					PrevNumInstances = Tick.bNeedsReset ? 0 : ExecContext->ScratchNumInstances;
+				}
+				ExecContext->ScratchNumInstances = InstanceData.SpawnInfo.SpawnRateInstances + InstanceData.SpawnInfo.EventSpawnTotal + PrevNumInstances;
+				ExecContext->ScratchMaxInstances = FMath::Max(ExecContext->ScratchMaxInstances, ExecContext->ScratchNumInstances);
+			}
 		}
 
 		// Set bIsFinalTick for the last tick of each context and reset the scratch index.
