Niagara - Fix crash when adding scratch dynamic inputs.

diff --git a/Engine/Plugins/FX/Niagara/Source/NiagaraEditor/Private/NiagaraHlslTranslator.cpp b/Engine/Plugins/FX/Niagara/Source/NiagaraEditor/Private/NiagaraHlslTranslator.cpp
--- a/Engine/Plugins/FX/Niagara/Source/NiagaraEditor/Private/NiagaraHlslTranslator.cpp
+++ b/Engine/Plugins/FX/Niagara/Source/NiagaraEditor/Private/NiagaraHlslTranslator.cpp
@@ -3819,7 +3819,15 @@ void FHlslNiagaraTranslator::Output(UNiagaraNodeOutput* OutputNode, const TArray
 
 	// Build up the attribute list. We don't auto-expand parameter maps here.
 	TArray<FNiagaraVariable> Outputs = OutputNode->GetOutputs();
-	check(ComputedInputs.Num() == Outputs.Num());
+	int32 NumberOfValidComputedInputs = 0;
+	for (int32 ComputedInput : ComputedInputs)
+	{
+		if (ComputedInput != INDEX_NONE)
+		{
+			NumberOfValidComputedInputs++;
+		}
+	}
+	check(NumberOfValidComputedInputs == Outputs.Num());
 	for (int32 PinIdx = 0; PinIdx < Outputs.Num(); PinIdx++)
 	{
 		Attributes.Add(Outputs[PinIdx]);
diff --git a/Engine/Plugins/FX/Niagara/Source/NiagaraEditor/Private/NiagaraNode.cpp b/Engine/Plugins/FX/Niagara/Source/NiagaraEditor/Private/NiagaraNode.cpp
--- a/Engine/Plugins/FX/Niagara/Source/NiagaraEditor/Private/NiagaraNode.cpp
+++ b/Engine/Plugins/FX/Niagara/Source/NiagaraEditor/Private/NiagaraNode.cpp
@@ -271,6 +271,11 @@ int32 UNiagaraNode::CompileInputPin(FHlslNiagaraTranslator *Translator, UEdGraph
 	return Translator->CompilePin(Pin);
 }
 
+bool UNiagaraNode::IsValidPinToCompile(UEdGraphPin* Pin) const 
+{ 
+	return Pin->bOrphanedPin == false;
+}
+
 bool UNiagaraNode::CompileInputPins(FHlslNiagaraTranslator *Translator, TArray<int32>& OutCompiledInputs)
 {
 	bool bError = false;
diff --git a/Engine/Plugins/FX/Niagara/Source/NiagaraEditor/Private/NiagaraNodeFunctionCall.cpp b/Engine/Plugins/FX/Niagara/Source/NiagaraEditor/Private/NiagaraNodeFunctionCall.cpp
--- a/Engine/Plugins/FX/Niagara/Source/NiagaraEditor/Private/NiagaraNodeFunctionCall.cpp
+++ b/Engine/Plugins/FX/Niagara/Source/NiagaraEditor/Private/NiagaraNodeFunctionCall.cpp
@@ -616,11 +616,6 @@ void UNiagaraNodeFunctionCall::Compile(class FHlslNiagaraTranslator* Translator,
 	}
 }
 
-bool UNiagaraNodeFunctionCall::IsValidPinToCompile(UEdGraphPin* Pin) const
-{
-	return !IsAddPin(Pin);
-}
-
 UObject*  UNiagaraNodeFunctionCall::GetReferencedAsset() const
 {
 	if (FunctionScript && FunctionScript->GetOutermost() != GetOutermost())
diff --git a/Engine/Plugins/FX/Niagara/Source/NiagaraEditor/Private/NiagaraNodeWithDynamicPins.cpp b/Engine/Plugins/FX/Niagara/Source/NiagaraEditor/Private/NiagaraNodeWithDynamicPins.cpp
--- a/Engine/Plugins/FX/Niagara/Source/NiagaraEditor/Private/NiagaraNodeWithDynamicPins.cpp
+++ b/Engine/Plugins/FX/Niagara/Source/NiagaraEditor/Private/NiagaraNodeWithDynamicPins.cpp
@@ -179,6 +179,11 @@ void UNiagaraNodeWithDynamicPins::MoveDynamicPin(UEdGraphPin* Pin, int32 Directi
 	}
 }
 
+bool UNiagaraNodeWithDynamicPins::IsValidPinToCompile(UEdGraphPin* Pin) const
+{
+	return !IsAddPin(Pin) && Super::IsValidPinToCompile(Pin);
+}
+
 void UNiagaraNodeWithDynamicPins::GetNodeContextMenuActions(UToolMenu* Menu, UGraphNodeContextMenuContext* Context) const
 {
 	Super::GetNodeContextMenuActions(Menu, Context);
diff --git a/Engine/Plugins/FX/Niagara/Source/NiagaraEditor/Private/ViewModels/NiagaraScratchPadViewModel.cpp b/Engine/Plugins/FX/Niagara/Source/NiagaraEditor/Private/ViewModels/NiagaraScratchPadViewModel.cpp
--- a/Engine/Plugins/FX/Niagara/Source/NiagaraEditor/Private/ViewModels/NiagaraScratchPadViewModel.cpp
+++ b/Engine/Plugins/FX/Niagara/Source/NiagaraEditor/Private/ViewModels/NiagaraScratchPadViewModel.cpp
@@ -290,6 +290,11 @@ TSharedPtr<FNiagaraScratchPadScriptViewModel> UNiagaraScratchPadViewModel::Creat
 				if (InOutputType.IsValid())
 				{
 					UNiagaraNodeOutput* DynamicInputOutputNode = OutputNodes[0];
+					
+					// Break pin lins before changing outputs to prevent old linked inputs from being retained as orphaned.
+					DynamicInputOutputNode->BreakAllNodeLinks();
+
+					// Add the new output and then refresh by notifying that the outputs property changed.
 					DynamicInputOutputNode->Outputs.Empty();
 					DynamicInputOutputNode->Outputs.Add(FNiagaraVariable(InOutputType, "Output"));
 					FPropertyChangedEvent OutputsChangedEvent(FindFieldChecked<FProperty>(UNiagaraNodeOutput::StaticClass(), GET_MEMBER_NAME_CHECKED(UNiagaraNodeOutput, Outputs)));
diff --git a/Engine/Plugins/FX/Niagara/Source/NiagaraEditor/Public/NiagaraNode.h b/Engine/Plugins/FX/Niagara/Source/NiagaraEditor/Public/NiagaraNode.h
--- a/Engine/Plugins/FX/Niagara/Source/NiagaraEditor/Public/NiagaraNode.h
+++ b/Engine/Plugins/FX/Niagara/Source/NiagaraEditor/Public/NiagaraNode.h
@@ -151,7 +151,7 @@ protected:
 	virtual bool PODPropertyAppendCompileHash(const void* Container, FProperty* Property, const FString& PropertyName, FNiagaraCompileHashVisitor* InVisitor) const;
 
 	virtual int32 CompileInputPin(class FHlslNiagaraTranslator *Translator, UEdGraphPin* Pin);
-	virtual bool IsValidPinToCompile(UEdGraphPin* Pin) const { return true; }; 
+	virtual bool IsValidPinToCompile(UEdGraphPin* Pin) const;
 
 	void NumericResolutionByPins(const UEdGraphSchema_Niagara* Schema, TArray<UEdGraphPin*>& InputPins, TArray<UEdGraphPin*>& OutputPins,
 		bool bFixInline, TMap<TPair<FGuid, UEdGraphNode*>, FNiagaraTypeDefinition>* PinCache);
diff --git a/Engine/Plugins/FX/Niagara/Source/NiagaraEditor/Public/NiagaraNodeFunctionCall.h b/Engine/Plugins/FX/Niagara/Source/NiagaraEditor/Public/NiagaraNodeFunctionCall.h
--- a/Engine/Plugins/FX/Niagara/Source/NiagaraEditor/Public/NiagaraNodeFunctionCall.h
+++ b/Engine/Plugins/FX/Niagara/Source/NiagaraEditor/Public/NiagaraNodeFunctionCall.h
@@ -136,7 +136,6 @@ public:
 	virtual TSharedPtr<SGraphNode> CreateVisualWidget() override;
 protected:
 
-	virtual bool IsValidPinToCompile(UEdGraphPin* Pin) const;
 	virtual bool GetValidateDataInterfaces() const { return true; };
 
 	virtual bool AllowDynamicPins() const override { return false; }
diff --git a/Engine/Plugins/FX/Niagara/Source/NiagaraEditor/Public/NiagaraNodeWithDynamicPins.h b/Engine/Plugins/FX/Niagara/Source/NiagaraEditor/Public/NiagaraNodeWithDynamicPins.h
--- a/Engine/Plugins/FX/Niagara/Source/NiagaraEditor/Public/NiagaraNodeWithDynamicPins.h
+++ b/Engine/Plugins/FX/Niagara/Source/NiagaraEditor/Public/NiagaraNodeWithDynamicPins.h
@@ -83,6 +83,8 @@ protected:
 	/** Convenience method to determine whether this Node is a Map Get or Map Set when adding a parameter through the parameter panel. */
 	virtual EEdGraphPinDirection GetPinDirectionForNewParameters() { return EEdGraphPinDirection::EGPD_MAX; };
 
+	bool IsValidPinToCompile(UEdGraphPin* Pin) const override;
+
 private:
 
 	/** Gets the display text for a pin. */
