Fix tooltips for invalid scope combobox selections.
Fix Output usage parameters displaying underlying scope instead of "Output"
Add context menu to parameter panel entries for requesting rename and deleting.
Focus text box to rename new parameters on creation through parameters panel.
Fix newly created parameters changing usage to "local" when added to a map get or set node.

diff --git a/Engine/Plugins/FX/Niagara/Source/Niagara/Classes/NiagaraConstants.h b/Engine/Plugins/FX/Niagara/Source/Niagara/Classes/NiagaraConstants.h
--- a/Engine/Plugins/FX/Niagara/Source/Niagara/Classes/NiagaraConstants.h
+++ b/Engine/Plugins/FX/Niagara/Source/Niagara/Classes/NiagaraConstants.h
@@ -168,6 +168,7 @@ struct NIAGARA_API FNiagaraConstants
 	static const FName ScriptPersistentScopeName;
 	static const FName InputScopeName;
 	static const FName OutputScopeName;
+	static const FName UniqueOutputScopeName;
 	static const FName CustomScopeName;
 
 private:
diff --git a/Engine/Plugins/FX/Niagara/Source/Niagara/Private/NiagaraConstants.cpp b/Engine/Plugins/FX/Niagara/Source/Niagara/Private/NiagaraConstants.cpp
--- a/Engine/Plugins/FX/Niagara/Source/Niagara/Private/NiagaraConstants.cpp
+++ b/Engine/Plugins/FX/Niagara/Source/Niagara/Private/NiagaraConstants.cpp
@@ -42,6 +42,7 @@ const FName FNiagaraConstants::CustomScopeName(TEXT("Custom"));
 
 const FName FNiagaraConstants::InputScopeName(TEXT("Input"));
 const FName FNiagaraConstants::OutputScopeName(TEXT("Output"));
+const FName FNiagaraConstants::UniqueOutputScopeName(TEXT("OutputModule"));
 const FName FNiagaraConstants::ScriptTransientScopeName(TEXT("ScriptTransient"));
 const FName FNiagaraConstants::ScriptPersistentScopeName(TEXT("ScriptPersistent"));
 
diff --git a/Engine/Plugins/FX/Niagara/Source/Niagara/Private/NiagaraTypes.cpp b/Engine/Plugins/FX/Niagara/Source/Niagara/Private/NiagaraTypes.cpp
--- a/Engine/Plugins/FX/Niagara/Source/Niagara/Private/NiagaraTypes.cpp
+++ b/Engine/Plugins/FX/Niagara/Source/Niagara/Private/NiagaraTypes.cpp
@@ -37,7 +37,7 @@ FString FNiagaraTypeUtilities::GetNamespaceStringForScriptParameterScope(const E
 		return PARAM_MAP_OUTPUT_MODULE_STR;
 		break;
 	case ENiagaraParameterScope::Local:
-		return PARAM_MAP_LOCAL_STR;
+		return PARAM_MAP_LOCAL_MODULE_STR;
 		break;
 	default:
 		checkf(false, TEXT("Unhandled parameter scope encountered!"));
diff --git a/Engine/Plugins/FX/Niagara/Source/NiagaraEditor/Private/NiagaraEditorModule.cpp b/Engine/Plugins/FX/Niagara/Source/NiagaraEditor/Private/NiagaraEditorModule.cpp
--- a/Engine/Plugins/FX/Niagara/Source/NiagaraEditor/Private/NiagaraEditorModule.cpp
+++ b/Engine/Plugins/FX/Niagara/Source/NiagaraEditor/Private/NiagaraEditorModule.cpp
@@ -1030,7 +1030,8 @@ void FNiagaraEditorModule::StartupModule()
 	RegisterParameterScopeInfo(FNiagaraConstants::LocalNamespace, FNiagaraParameterScopeInfo(ENiagaraParameterScope::Local, PARAM_MAP_LOCAL_MODULE_STR));
 	RegisterParameterScopeInfo(FNiagaraConstants::ScriptPersistentScopeName, FNiagaraParameterScopeInfo(ENiagaraParameterScope::ScriptPersistent, PARAM_MAP_SCRIPT_PERSISTENT_STR));
 	RegisterParameterScopeInfo(FNiagaraConstants::ScriptTransientScopeName, FNiagaraParameterScopeInfo(ENiagaraParameterScope::ScriptTransient, PARAM_MAP_SCRIPT_TRANSIENT_STR));
-	RegisterParameterScopeInfo(FNiagaraConstants::OutputScopeName, FNiagaraParameterScopeInfo(ENiagaraParameterScope::Output, PARAM_MAP_OUTPUT_MODULE_STR));
+	RegisterParameterScopeInfo(FNiagaraConstants::OutputScopeName, FNiagaraParameterScopeInfo(ENiagaraParameterScope::Output, PARAM_MAP_OUTPUT_STR));
+	RegisterParameterScopeInfo(FNiagaraConstants::UniqueOutputScopeName, FNiagaraParameterScopeInfo(ENiagaraParameterScope::Output, PARAM_MAP_OUTPUT_MODULE_STR));
 	RegisterParameterScopeInfo(FNiagaraConstants::CustomScopeName, FNiagaraParameterScopeInfo(ENiagaraParameterScope::Custom, FString()));
 }
 
diff --git a/Engine/Plugins/FX/Niagara/Source/NiagaraEditor/Private/NiagaraEditorUtilities.cpp b/Engine/Plugins/FX/Niagara/Source/NiagaraEditor/Private/NiagaraEditorUtilities.cpp
--- a/Engine/Plugins/FX/Niagara/Source/NiagaraEditor/Private/NiagaraEditorUtilities.cpp
+++ b/Engine/Plugins/FX/Niagara/Source/NiagaraEditor/Private/NiagaraEditorUtilities.cpp
@@ -1859,7 +1859,7 @@ void FNiagaraEditorUtilities::CreateAssetFromEmitter(TSharedRef<FNiagaraEmitterH
 
 void FNiagaraEditorUtilities::GetScriptRunAndExecutionIndexFromUsage(const ENiagaraScriptUsage& InUsage, int32& OutRunIndex, int32&OutExecutionIndex)
 {
-	switch (InUsage) //@todo(ng) fix numeric scheme Execution Index needing to continuously increment
+	switch (InUsage)
 	{
 	case ENiagaraScriptUsage::SystemSpawnScript:
 		OutRunIndex = 0;
@@ -2110,25 +2110,20 @@ bool FNiagaraEditorUtilities::GetVariableMetaDataNamespaceString(const FNiagaraV
 		return false;
 	}
 
-	if (MetaData.IsInputOrLocalUsage())
+	const FName& MetaDataScopeName = MetaData.GetScopeName();
+	const FNiagaraParameterScopeInfo* ScopeInfo = FNiagaraEditorModule::FindParameterScopeInfo(MetaDataScopeName);
+	if (ScopeInfo == nullptr)
 	{
-		const FName& MetaDataScopeName = MetaData.GetScopeName();
-		const FNiagaraParameterScopeInfo* ScopeInfo = FNiagaraEditorModule::FindParameterScopeInfo(MetaDataScopeName);
-		if (ScopeInfo == nullptr)
-		{
-			ensureMsgf(false, TEXT("Failed to find registered parameter scope info for scope name %s!"), *MetaDataScopeName.ToString());
-			return false;
-		}
+		ensureMsgf(false, TEXT("Failed to find registered parameter scope info for scope name %s!"), *MetaDataScopeName.ToString());
+		return false;
+	}
 
-		FString NamespaceString = ScopeInfo->GetNamespaceString();
-		if (MetaData.GetUsage() == ENiagaraScriptParameterUsage::InitialValueInput)
-		{
-			NamespaceString.Append(PARAM_MAP_INITIAL_STR);
-		}
-		OutNamespaceString = NamespaceString;
-		return true;
+	FString NamespaceString = ScopeInfo->GetNamespaceString();
+	if (MetaData.GetUsage() == ENiagaraScriptParameterUsage::InitialValueInput)
+	{
+		NamespaceString.Append(PARAM_MAP_INITIAL_STR);
 	}
-	OutNamespaceString = PARAM_MAP_OUTPUT_STR;
+	OutNamespaceString = NamespaceString;
 	return true;
 }
 
@@ -2284,6 +2279,12 @@ void FNiagaraEditorUtilities::GetParameterMetaDataFromName(const FName& InVarNam
 			OutMetaData.SetUsage(ENiagaraScriptParameterUsage::Input);
 			return true;
 		}
+		else if (Namespace == FNiagaraConstants::OutputScopeName)
+		{
+			OutMetaData.SetScopeName(Namespace);
+			OutMetaData.SetUsage(ENiagaraScriptParameterUsage::Output);
+			return true;
+		}
 
 		MarkAsLegacyCustomName();
 		return false;
@@ -2366,6 +2367,14 @@ void FNiagaraEditorUtilities::GetParameterMetaDataFromName(const FName& InVarNam
 					return;
 				}
 			}
+			else if (DecomposedNamespaces[0] == FNiagaraConstants::OutputScopeName)
+			{
+				if (DecomposedNamespaces[1] == FNiagaraConstants::ModuleNamespace)
+				{
+					OutMetaData.SetScopeName(FNiagaraConstants::UniqueOutputScopeName);
+					return;
+				}
+			}
 			else if (GetScopeCanHaveInitialNamespaceSuffix(FirstNamespaceScope))
 			{
 				bool bFoundInitialNamespace = GetMetaDataForInitialNamespace(DecomposedNamespaces[1], OutMetaData);
diff --git a/Engine/Plugins/FX/Niagara/Source/NiagaraEditor/Private/NiagaraNodeWithDynamicPins.cpp b/Engine/Plugins/FX/Niagara/Source/NiagaraEditor/Private/NiagaraNodeWithDynamicPins.cpp
--- a/Engine/Plugins/FX/Niagara/Source/NiagaraEditor/Private/NiagaraNodeWithDynamicPins.cpp
+++ b/Engine/Plugins/FX/Niagara/Source/NiagaraEditor/Private/NiagaraNodeWithDynamicPins.cpp
@@ -292,19 +292,26 @@ void UNiagaraNodeWithDynamicPins::AddParameter(FNiagaraVariable Parameter, const
 		AddParameterOptions.NewParameterUsage = GuessedMetaData.GetUsage();
 		AddParameterOptions.NewParameterScopeName = GuessedMetaData.GetScopeName();
 
-		if (FNiagaraConstants::FindEngineConstant(Parameter) == nullptr)
-			AddParameterOptions.bMakeParameterNameUnique = true;
-	
-		AddParameterOptions.bRefreshMetaDataScopeAndUsage = true;
-
 		UNiagaraGraph* Graph = GetNiagaraGraph();
 		checkf(Graph != nullptr, TEXT("Failed to get niagara graph when adding pin!"));
 
-		Graph->Modify();
-		Graph->AddParameter(Parameter, AddParameterOptions);
+		// Resolve the unique parameter name before adding to the graph as the pin needs to be created first to resolve the parameter metadata usage.
+		if (FNiagaraConstants::FindEngineConstant(Parameter) == nullptr)
+		{
+			UNiagaraScriptVariable** FoundScriptVariable = Graph->GetAllMetaData().Find(Parameter);
+			if (!FoundScriptVariable)
+			{
+				Parameter.SetName(Graph->MakeUniqueParameterName(Parameter.GetName()));
+			}
+		}
 
 		Modify();
 		UEdGraphPin* Pin = this->RequestNewTypedPin(AddPin->Direction, Parameter.GetType(), Parameter.GetName());
+
+		Graph->Modify();
+		Graph->AddParameter(Parameter, AddParameterOptions);
+
+		// Add the Parameter Reference 
 		FNiagaraGraphParameterReference NewParameterReference = FNiagaraGraphParameterReference(Pin->PersistentGuid, this);
 		Graph->AddParameterReference(Parameter, NewParameterReference);
 	}
diff --git a/Engine/Plugins/FX/Niagara/Source/NiagaraEditor/Private/ViewModels/NiagaraParameterNameViewModel.cpp b/Engine/Plugins/FX/Niagara/Source/NiagaraEditor/Private/ViewModels/NiagaraParameterNameViewModel.cpp
--- a/Engine/Plugins/FX/Niagara/Source/NiagaraEditor/Private/ViewModels/NiagaraParameterNameViewModel.cpp
+++ b/Engine/Plugins/FX/Niagara/Source/NiagaraEditor/Private/ViewModels/NiagaraParameterNameViewModel.cpp
@@ -73,11 +73,11 @@ TSharedRef<SWidget> FNiagaraParameterPanelEntryParameterNameViewModel::CreateTex
 		.Font(NameFont)
 		.HighlightText(CreateData->HighlightText)
 		.OnTextCommitted(this, &FNiagaraParameterPanelEntryParameterNameViewModel::OnParameterRenamed)
-		//.OnVerifyTextChanged(InArgs._OnVerifyParameterNameChanged) //@todo(ng) impl
+		.OnVerifyTextChanged(this, &FNiagaraParameterPanelEntryParameterNameViewModel::VerifyParameterNameChanged)
 		.IsSelected(CreateData->IsRowSelectedDelegate)
 		.IsReadOnly(bIsReadOnly);
 
-	//CreateData->OnRenameRequest->BindSP(static_cast<SInlineEditableTextBlock*>(DisplayWidget.Get()), &SInlineEditableTextBlock::EnterEditingMode); //@Todo(ng) figure out a way to make this work
+	CreateData->OnRenameRequest->BindSP(static_cast<SInlineEditableTextBlock*>(DisplayWidget.Get()), &SInlineEditableTextBlock::EnterEditingMode);
 
 	return DisplayWidget.ToSharedRef();
 }
@@ -88,6 +88,11 @@ int32 FNiagaraParameterPanelEntryParameterNameViewModel::GetScopeValue() const
 	{
 		return int32(ENiagaraParameterScope::DISPLAY_ONLY_StaticSwitch);
 	}
+	else if (CachedScriptVarAndViewInfo.MetaData.GetUsage() == ENiagaraScriptParameterUsage::Output)
+	{
+		return int32(ENiagaraParameterScope::Output);
+	}
+
 	ENiagaraParameterScope CachedScope;
 	if (ensureMsgf(FNiagaraEditorUtilities::GetVariableMetaDataScope(CachedScriptVarAndViewInfo.MetaData, CachedScope), TEXT("Failed to get scope value for param as override namespace is set! This method should not be bound!")))
 	{
@@ -122,6 +127,16 @@ FText FNiagaraParameterPanelEntryParameterNameViewModel::GetParameterNameText()
 	return FText::FromName(ParameterName);
 }
 
+bool FNiagaraParameterPanelEntryParameterNameViewModel::VerifyParameterNameChanged(const FText& NewNameText, FText& OutErrorText) const
+{
+	if (OnVerifyParameterRenamedDelegate.IsBound())
+	{
+		return OnVerifyParameterRenamedDelegate.Execute(CachedScriptVarAndViewInfo.ScriptVariable, CachedScriptVarAndViewInfo.MetaData, NewNameText, OutErrorText);
+	}
+	ensureMsgf(false, TEXT("Verify parameter renamed delegate was not bound!"));
+	return true;
+}
+
 void FNiagaraParameterPanelEntryParameterNameViewModel::OnParameterRenamed(const FText& NewNameText, ETextCommit::Type TextCommitType) const
 {
 	OnParameterRenamedDelegate.ExecuteIfBound(CachedScriptVarAndViewInfo.ScriptVariable, CachedScriptVarAndViewInfo.MetaData, NewNameText);
@@ -181,7 +196,7 @@ TSharedRef<SWidget> FNiagaraGraphPinParameterNameViewModel::CreateTextSlotWidget
 		.Text(this, &FNiagaraGraphPinParameterNameViewModel::GetParameterNameText)
 		.Font(NameFont)
 		.OnTextCommitted(this, &FNiagaraGraphPinParameterNameViewModel::OnParameterRenamed)
-		//.OnVerifyTextChanged(InArgs._OnVerifyParameterNameChanged) //@todo(ng) impl
+		.OnVerifyTextChanged(this, &FNiagaraGraphPinParameterNameViewModel::VerifyParameterNameChanged)
 		.IsReadOnly(bIsReadOnly);
 
 	return DisplayWidget.ToSharedRef();
@@ -189,6 +204,15 @@ TSharedRef<SWidget> FNiagaraGraphPinParameterNameViewModel::CreateTextSlotWidget
 
 int32 FNiagaraGraphPinParameterNameViewModel::GetScopeValue() const
 {
+	if (CachedScriptVarAndViewInfo.MetaData.GetIsStaticSwitch())
+	{
+		return int32(ENiagaraParameterScope::DISPLAY_ONLY_StaticSwitch);
+	}
+	else if (CachedScriptVarAndViewInfo.MetaData.GetUsage() == ENiagaraScriptParameterUsage::Output)
+	{
+		return int32(ENiagaraParameterScope::Output);
+	}
+
 	ENiagaraParameterScope CachedScope;
 	if (ensureMsgf(FNiagaraEditorUtilities::GetVariableMetaDataScope(CachedScriptVarAndViewInfo.MetaData, CachedScope), TEXT("Failed to get scope value for param as override namespace is set! This method should not be bound!")))
 	{
@@ -223,6 +247,11 @@ FText FNiagaraGraphPinParameterNameViewModel::GetParameterNameText() const
 	return FText::FromName(ParameterName);
 }
 
+bool FNiagaraGraphPinParameterNameViewModel::VerifyParameterNameChanged(const FText& NewNameText, FText& OutErrorText) const
+{
+	return ParameterPanelViewModel->GetCanRenameParameterAndToolTip(CachedScriptVarAndViewInfo.ScriptVariable, CachedScriptVarAndViewInfo.MetaData, NewNameText, OutErrorText);
+}
+
 void FNiagaraGraphPinParameterNameViewModel::OnParameterRenamed(const FText& NewNameText, ETextCommit::Type SelectionType) const
 {
 	ParameterPanelViewModel->RenamePin(OwningPin, NewNameText);
diff --git a/Engine/Plugins/FX/Niagara/Source/NiagaraEditor/Private/ViewModels/NiagaraParameterNameViewModel.h b/Engine/Plugins/FX/Niagara/Source/NiagaraEditor/Private/ViewModels/NiagaraParameterNameViewModel.h
--- a/Engine/Plugins/FX/Niagara/Source/NiagaraEditor/Private/ViewModels/NiagaraParameterNameViewModel.h
+++ b/Engine/Plugins/FX/Niagara/Source/NiagaraEditor/Private/ViewModels/NiagaraParameterNameViewModel.h
@@ -28,7 +28,7 @@ public:
 	virtual bool GetScopeValueIsEnabled(int32 ScopeValue) const = 0;
 
 	virtual FText GetParameterNameText() const = 0;
-	//bool VerifyParameterNameChanged(const FText& NewNameText, FText&) //@todo(ng) impl
+	virtual bool VerifyParameterNameChanged(const FText& NewNameText, FText& OutErrorText) const = 0;
 	virtual void OnParameterRenamed(const FText& NewNameText, ETextCommit::Type SelectionType) const = 0;
 };
 
@@ -37,6 +37,7 @@ class FNiagaraParameterPanelEntryParameterNameViewModel : public INiagaraParamet
 public:
 	DECLARE_DELEGATE_ThreeParams(FOnScopeSelectionChanged, const FNiagaraVariable&, const FNiagaraVariableMetaData&, const ENiagaraParameterScope)
 	DECLARE_DELEGATE_ThreeParams(FOnParameterRenamed, const FNiagaraVariable&, const FNiagaraVariableMetaData&, const FText&)
+	DECLARE_DELEGATE_RetVal_FourParams(bool, FOnVerifyParameterRenamed, const FNiagaraVariable&, const FNiagaraVariableMetaData&, TOptional<const FText> , FText&)
 
 	FNiagaraParameterPanelEntryParameterNameViewModel(FCreateWidgetForActionData* const InCreateData, const FNiagaraScriptVariableAndViewInfo& InScriptVarAndViewInfo);
 
@@ -50,16 +51,18 @@ public:
 	virtual bool GetScopeValueIsEnabled(int32 ScopeValue) const override;
 
 	virtual FText GetParameterNameText() const override;
-	//bool VerifyParameterNameChanged(const FText& NewNameText, FText&) //@todo(ng) impl
+	virtual bool VerifyParameterNameChanged(const FText& NewNameText, FText& OutErrorText) const override;
 	virtual void OnParameterRenamed(const FText& NewNameText, ETextCommit::Type SelectionType) const override;
 	/** End INiagaraParameterNameViewModel Interface */
 
 	FOnScopeSelectionChanged& GetOnScopeSelectionChangedDelegate() { return OnScopeSelectionChangedDelegate; };
 	FOnParameterRenamed& GetOnParameterRenamedDelegate() { return OnParameterRenamedDelegate; };
+	FOnVerifyParameterRenamed& GetOnVerifyParameterRenamedDelegate() { return OnVerifyParameterRenamedDelegate; };
 
 private:
 	FOnScopeSelectionChanged OnScopeSelectionChangedDelegate;
 	FOnParameterRenamed OnParameterRenamedDelegate;
+	FOnVerifyParameterRenamed OnVerifyParameterRenamedDelegate;
 
 	FCreateWidgetForActionData* const CreateData;
 	const FNiagaraScriptVariableAndViewInfo CachedScriptVarAndViewInfo;
@@ -86,7 +89,7 @@ public:
 	virtual bool GetScopeValueIsEnabled(int32 ScopeValue) const override;
 
 	virtual FText GetParameterNameText() const override;
-	//bool VerifyParameterNameChanged(const FText& NewNameText, FText&) //@todo(ng) impl
+	virtual bool VerifyParameterNameChanged(const FText& NewNameText, FText& OutErrorText) const override;
 	virtual void OnParameterRenamed(const FText& NewNameText, ETextCommit::Type SelectionType) const override;
 	/** End INiagaraParameterNameViewModel Interface */
 
diff --git a/Engine/Plugins/FX/Niagara/Source/NiagaraEditor/Private/ViewModels/NiagaraParameterPanelViewModel.cpp b/Engine/Plugins/FX/Niagara/Source/NiagaraEditor/Private/ViewModels/NiagaraParameterPanelViewModel.cpp
--- a/Engine/Plugins/FX/Niagara/Source/NiagaraEditor/Private/ViewModels/NiagaraParameterPanelViewModel.cpp
+++ b/Engine/Plugins/FX/Niagara/Source/NiagaraEditor/Private/ViewModels/NiagaraParameterPanelViewModel.cpp
@@ -178,10 +178,11 @@ ENiagaraParameterScope NiagaraParameterPanelSectionID::GetScopeForNewParametersI
 		return ENiagaraParameterScope::Particles;
 	case NiagaraParameterPanelSectionID::INPUTS:
 		return ENiagaraParameterScope::Input;
+	case NiagaraParameterPanelSectionID::OUTPUTS:
+		return ENiagaraParameterScope::Output;
 
 	// Default to Particles scope if coming from section IDs that are not directly associated with scope.
 	case NiagaraParameterPanelSectionID::REFERENCES:
-	case NiagaraParameterPanelSectionID::OUTPUTS:
 		return ENiagaraParameterScope::Particles;
 
 	case NiagaraParameterPanelSectionID::NONE:
@@ -228,6 +229,7 @@ TSharedRef<SWidget> INiagaraParameterPanelViewModel::GetScriptParameterVisualWid
 	TSharedPtr<FNiagaraParameterPanelEntryParameterNameViewModel> ParameterNameViewModel = MakeShared<FNiagaraParameterPanelEntryParameterNameViewModel>(InCreateData, ScriptVarAndViewInfo);
 	ParameterNameViewModel->GetOnParameterRenamedDelegate().BindSP(this, &INiagaraParameterPanelViewModel::RenameParameter);
 	ParameterNameViewModel->GetOnScopeSelectionChangedDelegate().BindSP(this, &INiagaraParameterPanelViewModel::ChangeParameterScope);
+	ParameterNameViewModel->GetOnVerifyParameterRenamedDelegate().BindSP(this, &INiagaraParameterPanelViewModel::GetCanRenameParameterAndToolTip);
 
 	TSharedPtr<SWidget> ScriptParameterVisualWidget = SNew(SNiagaraParameterNameView, ParameterNameViewModel);
 	return ScriptParameterVisualWidget->AsShared();
@@ -304,7 +306,7 @@ NiagaraParameterPanelSectionID::Type FNiagaraSystemToolkitParameterPanelViewMode
 	}
 }
 
-void FNiagaraSystemToolkitParameterPanelViewModel::AddParameter(const FNiagaraVariable& InVariableToAdd, const FNiagaraVariableMetaData& InVariableMetaDataToAssign)
+const UNiagaraScriptVariable* FNiagaraSystemToolkitParameterPanelViewModel::AddParameter(FNiagaraVariable& VariableToAdd, const FNiagaraVariableMetaData& InVariableMetaDataToAssign)
 {
 	FScopedTransaction AddParameter(LOCTEXT("AddParameter", "Add Parameter"));
 	bool bSystemIsSelected = OverviewSelectionViewModel->GetSystemIsSelected();
@@ -312,8 +314,8 @@ void FNiagaraSystemToolkitParameterPanelViewModel::AddParameter(const FNiagaraVa
 	UNiagaraGraph::FAddParameterOptions AddParameterOptions = UNiagaraGraph::FAddParameterOptions();
 	AddParameterOptions.NewParameterScopeName = InVariableMetaDataToAssign.GetScopeName();
 	AddParameterOptions.NewParameterUsage = InVariableMetaDataToAssign.GetUsage();
+	AddParameterOptions.bMakeParameterNameUnique = true;
 	AddParameterOptions.bAddedFromSystemEditor = true;
-	FNiagaraVariable DuplicateVar = InVariableToAdd; //@todo(ng) rewrite
 
 	ENiagaraParameterScope NewScope;
 	FNiagaraEditorUtilities::GetVariableMetaDataScope(InVariableMetaDataToAssign, NewScope);
@@ -323,13 +325,15 @@ void FNiagaraSystemToolkitParameterPanelViewModel::AddParameter(const FNiagaraVa
 		UNiagaraSystem* System = &SystemViewModel->GetSystem();
 		UNiagaraSystemEditorData* SystemEditorData = CastChecked<UNiagaraSystemEditorData>(System->GetEditorData(), ECastCheckedType::NullChecked);
 		SystemEditorData->Modify();
-		bool bSuccess = FNiagaraEditorUtilities::AddParameter(DuplicateVar, System->GetExposedParameters(), *System, SystemEditorData->GetStackEditorData());
+		bool bSuccess = FNiagaraEditorUtilities::AddParameter(VariableToAdd, System->GetExposedParameters(), *System, SystemEditorData->GetStackEditorData());
 	}
 	else if (NewScope == ENiagaraParameterScope::System)
 	{
 		UNiagaraGraph* Graph = SystemViewModel->GetSystemScriptViewModel()->GetGraphViewModel()->GetGraph();
 		Graph->Modify();
-		Graph->AddParameter(DuplicateVar, AddParameterOptions); //@todo(ng) verify
+		const UNiagaraScriptVariable* NewScriptVar = Graph->AddParameter(VariableToAdd, AddParameterOptions); //@todo(ng) verify
+		Refresh();
+		return NewScriptVar;
 	}
 	else
 	{
@@ -338,7 +342,9 @@ void FNiagaraSystemToolkitParameterPanelViewModel::AddParameter(const FNiagaraVa
 			if (ensureMsgf(Graph.IsValid(), TEXT("Editable Emitter Script Graph was stale when adding parameter!")))
 			{
 				Graph->Modify();
-				Graph->AddParameter(DuplicateVar, AddParameterOptions);
+				const UNiagaraScriptVariable* NewScriptVar = Graph->AddParameter(VariableToAdd, AddParameterOptions);
+				Refresh();
+				return NewScriptVar;
 			}
 		}
 	}
@@ -348,10 +354,10 @@ void FNiagaraSystemToolkitParameterPanelViewModel::AddParameter(const FNiagaraVa
 		//SystemViewModel-> //@todo(ng) handle user params
 	}
 	
-	Refresh();
+	return nullptr;
 }
 
-void FNiagaraSystemToolkitParameterPanelViewModel::RemoveParameter(const FNiagaraVariable& TargetVariableToRemove, const FNiagaraVariableMetaData& TargetVariableMetaData)
+void FNiagaraSystemToolkitParameterPanelViewModel::DeleteParameter(const FNiagaraVariable& TargetVariableToRemove, const FNiagaraVariableMetaData& TargetVariableMetaData)
 {
 	FScopedTransaction RemoveParameter(LOCTEXT("RemoveParameter", "Remove Parameter"));
 
@@ -377,9 +383,15 @@ void FNiagaraSystemToolkitParameterPanelViewModel::RemoveParameter(const FNiagar
 	Refresh();
 }
 
-bool FNiagaraSystemToolkitParameterPanelViewModel::CanRemoveParameter(const FNiagaraVariable& TargetVariableToRemove, const FNiagaraVariableMetaData& TargetVariableMetaData) const
+bool FNiagaraSystemToolkitParameterPanelViewModel::GetCanDeleteParameterAndToolTip(const FNiagaraVariable& TargetVariableToRemove, const FNiagaraVariableMetaData& TargetVariableMetaData, FText& OutCanDeleteParameterToolTip) const
 {
-	return TargetVariableMetaData.GetWasCreatedInSystemEditor();
+	if (TargetVariableMetaData.GetWasCreatedInSystemEditor())
+	{
+		OutCanDeleteParameterToolTip = LOCTEXT("SystemToolkitParameterPanelViewModel_DeleteParameter_CreatedInSystem", "Delete this Parameter and remove any usages from the System and Emitters.");
+		return true;
+	}
+	OutCanDeleteParameterToolTip = LOCTEXT("SystemToolkitParameterPanelViewModel_DeleteParameter_NotCreatedInSystem", "Cannot Delete this Parameter: Only Parameters created in the System Editor can be deleted from the System Editor.");
+	return false;
 }
 
 void FNiagaraSystemToolkitParameterPanelViewModel::RenameParameter(const FNiagaraVariable& TargetVariableToRename, const FNiagaraVariableMetaData& TargetVariableMetaData, const FText& NewVariableNameText) const
@@ -403,7 +415,7 @@ void FNiagaraSystemToolkitParameterPanelViewModel::RenameParameter(const FNiagar
 		if (ensureMsgf(Graph.IsValid(), TEXT("Editable Emitter Script Graph was stale when renaming parameter!")))
 		{
 			Graph->Modify();
-			Graph->RenameParameter(TargetVariableToRename, NewVariableName, TargetVariableMetaData.GetIsStaticSwitch(), TargetVariableMetaData.GetScopeName()); //@todo(ng) handle renaming system params
+			Graph->RenameParameter(TargetVariableToRename, NewVariableName, false, TargetVariableMetaData.GetScopeName()); //@todo(ng) handle renaming system params
 		}
 	}
 	Refresh();
@@ -420,9 +432,21 @@ bool FNiagaraSystemToolkitParameterPanelViewModel::CanModifyParameter(const FNia
 	return TargetVariableMetaData.GetWasCreatedInSystemEditor() == true;
 }
 
-bool FNiagaraSystemToolkitParameterPanelViewModel::CanRenameParameter(const FNiagaraVariable& TargetVariableToRename, const FNiagaraVariableMetaData& TargetVariableMetaData, const FText& NewVariableNameText) const
+bool FNiagaraSystemToolkitParameterPanelViewModel::GetCanRenameParameterAndToolTip(const FNiagaraVariable& TargetVariableToRename, const FNiagaraVariableMetaData& TargetVariableMetaData, TOptional<const FText> NewVariableNameText, FText& OutCanRenameParameterToolTip) const
 {
-	return TargetVariableMetaData.GetWasCreatedInSystemEditor() == true;
+	if (TargetVariableMetaData.GetWasCreatedInSystemEditor())
+	{
+		if(NewVariableNameText.IsSet() && NewVariableNameText->IsEmptyOrWhitespace())
+		{
+			OutCanRenameParameterToolTip = LOCTEXT("SystemToolkitParameterPanelViewModel_RenameParameter_NameNone", "Parameter must have a name.");
+			return false;
+		}
+
+		OutCanRenameParameterToolTip = LOCTEXT("SystemToolkitParameterPanelViewModel_RenameParameter_CreatedInSystem", "Rename this Parameter and all usages in the System and Emitters.");
+		return true;
+	}
+	OutCanRenameParameterToolTip = LOCTEXT("SystemToolkitParameterPanelViewModel_RenameParameter_NotCreatedInSystem", "Cannot rename Parameter: Only Parameters created in the System Editor can be renamed from the System Editor.");
+	return false;
 }
 
 FReply FNiagaraSystemToolkitParameterPanelViewModel::HandleActionDragged(const TSharedPtr<FEdGraphSchemaAction>& InAction, const FPointerEvent& MouseEvent) const
@@ -488,7 +512,7 @@ const TArray<FNiagaraScriptVariableAndViewInfo> FNiagaraSystemToolkitParameterPa
 				{
 					continue;
 				}
-				else if (MetaDataScope == ENiagaraParameterScope::ScriptPersistent || MetaDataScope == ENiagaraParameterScope::ScriptTransient)
+				else if (MetaDataScope == ENiagaraParameterScope::ScriptPersistent || MetaDataScope == ENiagaraParameterScope::ScriptTransient || MetaDataScope == ENiagaraParameterScope::Output)
 				{
 					//@todo(ng) Skip script alias parameters until we can resolve them!
 					continue;
@@ -605,24 +629,30 @@ NiagaraParameterPanelSectionID::Type FNiagaraScriptToolkitParameterPanelViewMode
 	return NiagaraParameterPanelSectionID::GetSectionForParameterMetaData(VarAndViewInfo.MetaData);
 }
 
-void FNiagaraScriptToolkitParameterPanelViewModel::AddParameter(const FNiagaraVariable& InVariableToAdd, const FNiagaraVariableMetaData& InVariableMetaDataToAssign)
+const UNiagaraScriptVariable* FNiagaraScriptToolkitParameterPanelViewModel::AddParameter(FNiagaraVariable& VariableToAdd, const FNiagaraVariableMetaData& InVariableMetaDataToAssign)
 {
-	FScopedTransaction AddParameter(LOCTEXT("AddParameterFromParameterPanel", "Add Parameter"));
-	UNiagaraGraph::FAddParameterOptions AddParameterOptions = UNiagaraGraph::FAddParameterOptions();
+	ENiagaraScriptParameterUsage InVariableUsage = InVariableMetaDataToAssign.GetUsage();
+	if (InVariableUsage == ENiagaraScriptParameterUsage::InputOutput)
+	{
+		ensureMsgf(false, TEXT("Unexpected usage encountered when adding parameter through parameter panel view model!"));
+	}
+	else
+	{
+		FScopedTransaction AddParameter(LOCTEXT("AddParameterFromParameterPanel", "Add Parameter"));
+		UNiagaraGraph::FAddParameterOptions AddParameterOptions = UNiagaraGraph::FAddParameterOptions();
 
-	AddParameterOptions.NewParameterScopeName = InVariableMetaDataToAssign.GetScopeName();
-	AddParameterOptions.NewParameterUsage = InVariableMetaDataToAssign.GetUsage();
-	FNiagaraVariable DuplicateVar = InVariableToAdd; //@todo(ng) rewrite
+		AddParameterOptions.NewParameterScopeName = InVariableMetaDataToAssign.GetScopeName();
+		AddParameterOptions.NewParameterUsage = InVariableMetaDataToAssign.GetUsage();
+		AddParameterOptions.bMakeParameterNameUnique = true;
 
-	UNiagaraGraph* Graph = ScriptViewModel->GetGraphViewModel()->GetGraph();
-	Graph->Modify();
-	
-	{
-		Graph->AddParameter(DuplicateVar, AddParameterOptions);
+		UNiagaraGraph* Graph = ScriptViewModel->GetGraphViewModel()->GetGraph();
+		Graph->Modify();
+		return Graph->AddParameter(VariableToAdd, AddParameterOptions);
 	}
+	return nullptr;
 }
 
-void FNiagaraScriptToolkitParameterPanelViewModel::RemoveParameter(const FNiagaraVariable& TargetVariableToRemove, const FNiagaraVariableMetaData& TargetVariableMetaData)
+void FNiagaraScriptToolkitParameterPanelViewModel::DeleteParameter(const FNiagaraVariable& TargetVariableToRemove, const FNiagaraVariableMetaData& TargetVariableMetaData)
 {
 	FScopedTransaction RemoveParametersWithPins(LOCTEXT("RemoveParametersWithPins", "Remove parameter and referenced pins"));
 	UNiagaraGraph* Graph = ScriptViewModel->GetGraphViewModel()->GetGraph();
@@ -630,8 +660,9 @@ void FNiagaraScriptToolkitParameterPanelViewModel::RemoveParameter(const FNiagar
 	Graph->RemoveParameter(TargetVariableToRemove);
 }
 
-bool FNiagaraScriptToolkitParameterPanelViewModel::CanRemoveParameter(const FNiagaraVariable& TargetVariableToRemove, const FNiagaraVariableMetaData& TargetVariableMetaData) const
+bool FNiagaraScriptToolkitParameterPanelViewModel::GetCanDeleteParameterAndToolTip(const FNiagaraVariable& TargetVariableToRemove, const FNiagaraVariableMetaData& TargetVariableMetaData, FText& OutCanDeleteParameterToolTip) const
 {
+	OutCanDeleteParameterToolTip = LOCTEXT("ScriptToolkitParameterPanelViewModel_DeleteParameter", "Delete this Parameter and remove any usages from the graph.");
 	return true;
 }
 
@@ -655,7 +686,7 @@ void FNiagaraScriptToolkitParameterPanelViewModel::RenameParameter(const FNiagar
 	UNiagaraGraph* Graph = ScriptViewModel->GetGraphViewModel()->GetGraph();
 	Graph->Modify();
 
-	Graph->RenameParameter(TargetVariableToRename, NewVariableName, TargetVariableMetaData.GetIsStaticSwitch(), TargetVariableMetaData.GetScopeName());
+	Graph->RenameParameter(TargetVariableToRename, NewVariableName, false, TargetVariableMetaData.GetScopeName());
 }
 
 void FNiagaraScriptToolkitParameterPanelViewModel::ChangeParameterScope(const FNiagaraVariable& TargetVariableToModify, const FNiagaraVariableMetaData& TargetVariableMetaData, const ENiagaraParameterScope NewVariableScope) const
@@ -680,7 +711,7 @@ void FNiagaraScriptToolkitParameterPanelViewModel::ChangeParameterScope(const FN
 
 			UNiagaraGraph* Graph = ScriptViewModel->GetGraphViewModel()->GetGraph();
 			Graph->Modify();
-			Graph->RenameParameter(TargetVariableToModify, NewVariableHLSLTokenName, TargetVariableMetaData.GetIsStaticSwitch(), TargetScopeName);
+			Graph->RenameParameter(TargetVariableToModify, NewVariableHLSLTokenName, false, TargetScopeName);
 		}
 	}
 }
@@ -691,21 +722,93 @@ bool FNiagaraScriptToolkitParameterPanelViewModel::CanModifyParameter(const FNia
 	return true;
 }
 
-bool FNiagaraScriptToolkitParameterPanelViewModel::CanRenameParameter(const FNiagaraVariable& TargetVariableToRename, const FNiagaraVariableMetaData& TargetVariableMetaData, const FText& NewVariableNameText) const
+bool FNiagaraScriptToolkitParameterPanelViewModel::GetCanRenameParameterAndToolTip(const FNiagaraVariable& TargetVariableToRename, const FNiagaraVariableMetaData& TargetVariableMetaData, TOptional<const FText> NewVariableNameText, FText& OutCanRenameParameterToolTip) const
 {
-	// Prevent name values that would alias an existing parameter
-	FName NewName = FName(*NewVariableNameText.ToString());
-	for (const FNiagaraScriptVariableAndViewInfo& ViewedParameter : CachedViewedParameters)
+	if (TargetVariableMetaData.GetIsStaticSwitch())
+	{
+		OutCanRenameParameterToolTip = LOCTEXT("ScriptToolkitParameterPanelViewModel_RenameParameter_StaticSwitch", "Cannot rename static switches through the parameter panel. Rename the static switch through the associated node's selected details panel.");
+		return false;
+	}
+
+	if (NewVariableNameText.IsSet())
 	{
-		FName ParameterName;
-		if (ViewedParameter.MetaData.GetParameterName(ParameterName))
+		if (NewVariableNameText->IsEmptyOrWhitespace())
+		{
+			OutCanRenameParameterToolTip = LOCTEXT("ScriptToolkitParameterPanelViewModel_RenameParameter_NameNone", "Parameter must have a name.");
+			return false;
+		}
+
+		FName NewName = FName(*NewVariableNameText->ToString());
+		TArray<FName> AliasScopeNames;
+		auto FindParameterNameAlias = [this](const TArray<FName>& AliasScopeNames, const FName& NewName)->const FNiagaraScriptVariableAndViewInfo* /* OutMatchedParameter*/ {
+			for (const FNiagaraScriptVariableAndViewInfo& ViewedParameter : CachedViewedParameters)
+			{
+				FName ParameterName;
+				if (AliasScopeNames.Contains(ViewedParameter.MetaData.GetScopeName()))
+				{
+					if (ViewedParameter.MetaData.GetParameterName(ParameterName))
+					{
+						if (NewName == ParameterName)
+						{
+							return &ViewedParameter;
+						}
+					}
+				}
+			}
+			return nullptr;
+		};
+
+		// Prevent name values that would alias an existing parameter with relevant usage
+		FText GenericNameAliasToolTip = LOCTEXT("ScriptToolkitParameterPanelViewModel_RenameParameter_Alias", "Cannot rename Parameter: A Parameter already exists with name {0}.");
+		FName ExistingName;
+		ensureMsgf(TargetVariableMetaData.GetParameterName(ExistingName), TEXT("Failed to get namespaceless parameter name from metadata on rename verify!"));
+		if (NewName != ExistingName)
 		{
-			if (NewName == ParameterName)
+			if (TargetVariableMetaData.GetUsage() == ENiagaraScriptParameterUsage::Input && TargetVariableMetaData.GetScopeName() == FNiagaraConstants::InputScopeName || TargetVariableMetaData.GetUsage() == ENiagaraScriptParameterUsage::Output)
 			{
-				return false; //@todo(ng) wrap this into the verify logic
+				// Inputs may not name alias outputs and vice versa
+				AliasScopeNames.Add(FNiagaraConstants::InputScopeName);
+				AliasScopeNames.Add(FNiagaraConstants::OutputScopeName);
+				AliasScopeNames.Add(FNiagaraConstants::UniqueOutputScopeName);
+				const FNiagaraScriptVariableAndViewInfo* NameAliasedVariableInfo = FindParameterNameAlias(AliasScopeNames, NewName);
+				if (NameAliasedVariableInfo != nullptr)
+				{
+					ENiagaraScriptParameterUsage TargetVarUsage = TargetVariableMetaData.GetUsage();
+					ENiagaraScriptParameterUsage AliasVarUsage = NameAliasedVariableInfo->MetaData.GetUsage();
+					if (TargetVarUsage == AliasVarUsage)
+					{
+						OutCanRenameParameterToolTip = FText::Format(GenericNameAliasToolTip, NewVariableNameText.GetValue());
+					}
+					else if (TargetVarUsage == ENiagaraScriptParameterUsage::Input && AliasVarUsage == ENiagaraScriptParameterUsage::Output)
+					{
+						OutCanRenameParameterToolTip = FText::Format(LOCTEXT("ScriptToolkitParameterPanelViewModel_RenameParameter_InOutAlias", "Cannot rename Input Parameter: An Output Parameter already exists with name {0}."), NewVariableNameText.GetValue());
+					}
+					else if (TargetVarUsage == ENiagaraScriptParameterUsage::Output && AliasVarUsage == ENiagaraScriptParameterUsage::Input)
+					{
+						OutCanRenameParameterToolTip = FText::Format(LOCTEXT("ScriptToolkitParameterPanelViewModel_RenameParameter_OutInAlias", "Cannot rename Output Parameter: An Input Parameter already exists with name {0}."), NewVariableNameText.GetValue());
+					}
+					else
+					{
+						ensureMsgf(false, TEXT("Unexpected usages encountered when checking name aliases for input/output rename!"));
+						OutCanRenameParameterToolTip = FText::Format(GenericNameAliasToolTip, NewVariableNameText.GetValue());
+					}
+					return false;
+				}
+			}
+			else
+			{
+				AliasScopeNames.Add(TargetVariableMetaData.GetScopeName());
+				const FNiagaraScriptVariableAndViewInfo* NameAliasedVariableInfo = FindParameterNameAlias(AliasScopeNames, NewName);
+				if (NameAliasedVariableInfo != nullptr)
+				{
+					OutCanRenameParameterToolTip = FText::Format(GenericNameAliasToolTip, NewVariableNameText.GetValue());
+					return false;
+				}
 			}
 		}
 	}
+
+	OutCanRenameParameterToolTip = LOCTEXT("ScriptToolkitParameterPanelViewModel_RenameParameter", "Rename this Parameter and all usages in the graph.");
 	return true;
 }
 
@@ -756,8 +859,15 @@ TSharedRef<SWidget> FNiagaraScriptToolkitParameterPanelViewModel::GetScriptParam
 	}
 
 	// Cannot resolve the parameter from the pin, put an error widget in.
-	TSharedPtr<SWidget> ErrorTextBlock = SNew(STextBlock).Text(FText::FromString("Could not resolve parameter!")); //@todo(ng) make error item method
-	return ErrorTextBlock->AsShared();
+	FNiagaraVariable StandInParameter = Schema->PinToNiagaraVariable(Pin, false);
+	FNiagaraVariableMetaData StandInMetaData;
+	FNiagaraEditorUtilities::GetParameterMetaDataFromName(StandInParameter.GetName(), StandInMetaData);
+	StandInMetaData.SetUsage(ENiagaraScriptParameterUsage::Local);
+	
+	FNiagaraScriptVariableAndViewInfo StandInScriptVarAndViewInfo = FNiagaraScriptVariableAndViewInfo(StandInParameter, StandInMetaData);
+	TSharedPtr<FNiagaraGraphPinParameterNameViewModel> ParameterNameViewModel = MakeShared<FNiagaraGraphPinParameterNameViewModel>(Pin, StandInScriptVarAndViewInfo, this);
+	TSharedPtr<SWidget> ScriptParameterVisualWidget = SNew(SNiagaraParameterNameView, ParameterNameViewModel);
+	return ScriptParameterVisualWidget->AsShared();
 }
 
 TStaticArray<FScopeIsEnabledAndTooltip, (int32)ENiagaraParameterScope::Num> FNiagaraScriptToolkitParameterPanelViewModel::GetParameterScopesEnabledAndTooltips(const FNiagaraVariable& InVar, const FNiagaraVariableMetaData& InVarMetaData) const
@@ -790,7 +900,8 @@ TStaticArray<FScopeIsEnabledAndTooltip, (int32)ENiagaraParameterScope::Num> FNia
 					FNiagaraEditorUtilities::GetVariableMetaDataScope(ViewedParameter.MetaData, ViewedParameterScope);
 
 					PerScopeInfo[(int32)ViewedParameterScope].bEnabled = false;
-					PerScopeInfo[(int32)ViewedParameterScope].Tooltip = LOCTEXT("NiagaraScopeSelectionNameAlias", "Cannot select scope: Parameter with same name already has this scope."); //@todo(ng) get scope
+					PerScopeInfo[(int32)ViewedParameterScope].Tooltip = FText::Format(LOCTEXT("NiagaraInvalidScopeSelectionNameAlias", "Cannot select scope '{0}': Parameter with same name already has this scope.")
+						, FText::FromName(FNiagaraEditorUtilities::GetScopeNameForParameterScope(ViewedParameterScope)));
 				}
 			}
 		}*/
@@ -799,14 +910,16 @@ TStaticArray<FScopeIsEnabledAndTooltip, (int32)ENiagaraParameterScope::Num> FNia
 		if (ScriptViewModel->GetStandaloneScript()->GetUsage() != ENiagaraScriptUsage::Module)
 		{
 			PerScopeInfo[(int32)ENiagaraParameterScope::Input].bEnabled = false;
-			PerScopeInfo[(int32)ENiagaraParameterScope::Input].Tooltip = LOCTEXT("NiagaraInvalidScopeSelectionModule", "Cannot select scope: Scope is only valid in Module Scripts."); //@todo(ng) get scope
+			PerScopeInfo[(int32)ENiagaraParameterScope::Input].Tooltip = FText::Format(LOCTEXT("NiagaraInvalidScopeSelectionModule", "Cannot select scope '{0}': Scope is only valid in Module Scripts.")
+				, FText::FromName(FNiagaraEditorUtilities::GetScopeNameForParameterScope(ENiagaraParameterScope::Input)));
 		}
 
 		const TArray<ENiagaraParameterScope> InvalidParameterScopes = ScriptViewModel->GetStandaloneScript()->GetUnsupportedParameterScopes();
 		for (ENiagaraParameterScope InvalidScope : InvalidParameterScopes)
 		{
 			PerScopeInfo[(int32)InvalidScope].bEnabled = false;
-			PerScopeInfo[(int32)InvalidScope].Tooltip = LOCTEXT("NiagaraInvalidScopeSelectionUsageBitmask", "Cannot select scope: Script Usage flags do not support a usage with this scope."); //@todo(ng) rewrite
+			PerScopeInfo[(int32)InvalidScope].Tooltip = FText::Format(LOCTEXT("NiagaraInvalidScopeSelectionUsageBitmask", "Cannot select scope '{0}': Script Usage flags do not support a usage with this scope.")
+				, FText::FromName(FNiagaraEditorUtilities::GetScopeNameForParameterScope(InvalidScope)));
 		}
 	}
 	else
diff --git a/Engine/Plugins/FX/Niagara/Source/NiagaraEditor/Private/ViewModels/NiagaraParameterPanelViewModel.h b/Engine/Plugins/FX/Niagara/Source/NiagaraEditor/Private/ViewModels/NiagaraParameterPanelViewModel.h
--- a/Engine/Plugins/FX/Niagara/Source/NiagaraEditor/Private/ViewModels/NiagaraParameterPanelViewModel.h
+++ b/Engine/Plugins/FX/Niagara/Source/NiagaraEditor/Private/ViewModels/NiagaraParameterPanelViewModel.h
@@ -18,6 +18,7 @@ class FNiagaraStandaloneScriptViewModel;
 class FDelegateHandle;
 struct FCreateWidgetForActionData;
 class FNiagaraObjectSelection;
+class UNiagaraScriptVariable;
 
 /** Interface for view models to the parameter panel. */
 class INiagaraParameterPanelViewModel : public TSharedFromThis<INiagaraParameterPanelViewModel>
@@ -39,11 +40,11 @@ public:
 
 	virtual NiagaraParameterPanelSectionID::Type GetSectionForVarAndViewInfo(const FNiagaraScriptVariableAndViewInfo& VarAndViewInfo) const = 0;
 	
-	virtual void AddParameter(const FNiagaraVariable& InVariableToAdd, const FNiagaraVariableMetaData& InVariableMetaDataToAssign) = 0;
+	virtual const UNiagaraScriptVariable* AddParameter(FNiagaraVariable& InVariableToAdd, const FNiagaraVariableMetaData& InVariableMetaDataToAssign) = 0;
 
-	virtual void RemoveParameter(const FNiagaraVariable& TargetVariableToRemove, const FNiagaraVariableMetaData& TargetVariableMetaData) = 0;
+	virtual void DeleteParameter(const FNiagaraVariable& TargetVariableToRemove, const FNiagaraVariableMetaData& TargetVariableMetaData) = 0;
 
-	virtual bool CanRemoveParameter(const FNiagaraVariable& TargetVariableToRemove, const FNiagaraVariableMetaData& TargetVariableMetaData) const = 0;
+	virtual bool GetCanDeleteParameterAndToolTip(const FNiagaraVariable& TargetVariableToRemove, const FNiagaraVariableMetaData& TargetVariableMetaData, FText& OutCanDeleteParameterToolTip) const = 0;
 
 	virtual void RenameParameter(const FNiagaraVariable& TargetVariableToRename, const FNiagaraVariableMetaData& TargetVariableMetaData, const FText& NewVariableNameText) const = 0;
 
@@ -51,7 +52,7 @@ public:
 
 	virtual bool CanModifyParameter(const FNiagaraVariable& TargetVariableToModify, const FNiagaraVariableMetaData& TargetVariableMetaData) const = 0;
 
-	virtual bool CanRenameParameter(const FNiagaraVariable& TargetVariableToRename, const FNiagaraVariableMetaData& TargetVariableMetaData, const FText& NewVariableNameText) const = 0;
+	virtual bool GetCanRenameParameterAndToolTip(const FNiagaraVariable& TargetVariableToRename, const FNiagaraVariableMetaData& TargetVariableMetaData, TOptional<const FText> NewVariableNameText, FText& OutCanRenameParameterToolTip) const = 0;
 
 	virtual void HandleActionSelected(const TSharedPtr<FEdGraphSchemaAction>& InAction, ESelectInfo::Type InSelectionType) {};
 
@@ -104,11 +105,11 @@ public:
 
 	virtual NiagaraParameterPanelSectionID::Type GetSectionForVarAndViewInfo(const FNiagaraScriptVariableAndViewInfo& VarAndViewInfo) const override;
 
-	virtual void AddParameter(const FNiagaraVariable& VariableToAdd, const FNiagaraVariableMetaData& VariableMetaDataToAssign) override;
+	virtual const UNiagaraScriptVariable* AddParameter(FNiagaraVariable& VariableToAdd, const FNiagaraVariableMetaData& VariableMetaDataToAssign) override;
 
-	virtual void RemoveParameter(const FNiagaraVariable& TargetVariableToRemove, const FNiagaraVariableMetaData& TargetVariableMetaData) override;
+	virtual void DeleteParameter(const FNiagaraVariable& TargetVariableToRemove, const FNiagaraVariableMetaData& TargetVariableMetaData) override;
 
-	virtual bool CanRemoveParameter(const FNiagaraVariable& TargetVariableToRemove, const FNiagaraVariableMetaData& TargetVariableMetaData) const override;
+	virtual bool GetCanDeleteParameterAndToolTip(const FNiagaraVariable& TargetVariableToRemove, const FNiagaraVariableMetaData& TargetVariableMetaData, FText& OutCanDeleteParameterToolTip) const override;
 
 	virtual void RenameParameter(const FNiagaraVariable& TargetVariableToRename, const FNiagaraVariableMetaData& TargetVariableMetaData, const FText& NewVariableNameText) const override;
 
@@ -116,7 +117,7 @@ public:
 
 	virtual bool CanModifyParameter(const FNiagaraVariable& TargetVariableToModify, const FNiagaraVariableMetaData& TargetVariableMetaData) const override;
 
-	virtual bool CanRenameParameter(const FNiagaraVariable& TargetVariableToRename, const FNiagaraVariableMetaData& TargetVariableMetaData, const FText& NewVariableNameText) const override;
+	virtual bool GetCanRenameParameterAndToolTip(const FNiagaraVariable& TargetVariableToRename, const FNiagaraVariableMetaData& TargetVariableMetaData, TOptional<const FText> NewVariableNameText, FText& OutCanRenameParameterToolTip) const override;
 
 	virtual FReply HandleActionDragged(const TSharedPtr<FEdGraphSchemaAction>& InAction, const FPointerEvent& MouseEvent) const override;
 
@@ -161,11 +162,11 @@ public:
 
 	virtual NiagaraParameterPanelSectionID::Type GetSectionForVarAndViewInfo(const FNiagaraScriptVariableAndViewInfo& VarAndViewInfo) const override;
 
-	virtual void AddParameter(const FNiagaraVariable& VariableToAdd, const FNiagaraVariableMetaData& VariableMetaDataToAssign) override;
+	virtual const UNiagaraScriptVariable* AddParameter(FNiagaraVariable& VariableToAdd, const FNiagaraVariableMetaData& VariableMetaDataToAssign) override;
 
-	virtual void RemoveParameter(const FNiagaraVariable& TargetVariableToRemove, const FNiagaraVariableMetaData& TargetVariableMetaData) override;
+	virtual void DeleteParameter(const FNiagaraVariable& TargetVariableToRemove, const FNiagaraVariableMetaData& TargetVariableMetaData) override;
 
-	virtual bool CanRemoveParameter(const FNiagaraVariable& TargetVariableToRemove, const FNiagaraVariableMetaData& TargetVariableMetaData) const override;
+	virtual bool GetCanDeleteParameterAndToolTip(const FNiagaraVariable& TargetVariableToRemove, const FNiagaraVariableMetaData& TargetVariableMetaData, FText& OutCanDeleteParameterToolTip) const override;
 
 	virtual void RenameParameter(const FNiagaraVariable& TargetVariableToRename, const FNiagaraVariableMetaData& TargetVariableMetaData, const FText& NewVariableNameText) const override;
 
@@ -173,7 +174,7 @@ public:
 
 	virtual bool CanModifyParameter(const FNiagaraVariable& TargetVariableToModify, const FNiagaraVariableMetaData& TargetVariableMetaData) const override;
 
-	virtual bool CanRenameParameter(const FNiagaraVariable& TargetVariableToRename, const FNiagaraVariableMetaData& TargetVariableMetaData, const FText& NewVariableNameText) const override;
+	virtual bool GetCanRenameParameterAndToolTip(const FNiagaraVariable& TargetVariableToRename, const FNiagaraVariableMetaData& TargetVariableMetaData, TOptional<const FText> NewVariableNameText, FText& OutCanRenameParameterToolTip) const override;
 
 	virtual void HandleActionSelected(const TSharedPtr<FEdGraphSchemaAction>& InAction, ESelectInfo::Type InSelectionType) override;
 	
diff --git a/Engine/Plugins/FX/Niagara/Source/NiagaraEditor/Private/Widgets/SNiagaraParameterPanel.cpp b/Engine/Plugins/FX/Niagara/Source/NiagaraEditor/Private/Widgets/SNiagaraParameterPanel.cpp
--- a/Engine/Plugins/FX/Niagara/Source/NiagaraEditor/Private/Widgets/SNiagaraParameterPanel.cpp
+++ b/Engine/Plugins/FX/Niagara/Source/NiagaraEditor/Private/Widgets/SNiagaraParameterPanel.cpp
@@ -14,6 +14,7 @@
 #include "SNiagaraParameterPanelPaletteItem.h"
 #include "SNiagaraGraphActionWidget.h"
 #include "NiagaraScriptVariable.h"
+#include "Editor/GraphEditor/Private/GraphActionNode.h"
 #include "NiagaraConstants.h"
 
 #define LOCTEXT_NAMESPACE "NiagaraParameterPanel"
@@ -34,6 +35,9 @@ void FNiagaraParameterPanelCommands::RegisterCommands()
 
 void SNiagaraParameterPanel::Construct(const FArguments& InArgs, const TSharedPtr<INiagaraParameterPanelViewModel>& InParameterPanelViewModel, const TSharedPtr<FUICommandList>& InToolkitCommands)
 {
+	bNeedsRefresh = false;
+	bGraphActionPendingRename = false;
+
 	ParameterPanelViewModel = InParameterPanelViewModel;
 	ToolkitCommands = InToolkitCommands;
 
@@ -60,8 +64,8 @@ void SNiagaraParameterPanel::Construct(const FArguments& InArgs, const TSharedPt
  		.OnActionDragged(this, &SNiagaraParameterPanel::OnActionDragged)
 		.OnActionSelected(this, &SNiagaraParameterPanel::OnActionSelected)
 // 		.OnActionDoubleClicked(this, &SNiagaraParameterMapView::OnActionDoubleClicked) //@todo(ng) impl
-// 		.OnContextMenuOpening(this, &SNiagaraParameterMapView::OnContextMenuOpening)
-// 		.OnCanRenameSelectedAction(this, &SNiagaraParameterMapView::CanRequestRenameOnActionNode)
+		.OnContextMenuOpening(this, &SNiagaraParameterPanel::OnContextMenuOpening)
+		.OnCanRenameSelectedAction(this, &SNiagaraParameterPanel::CanRequestRenameOnActionNode)
  		.OnGetSectionTitle(this, &SNiagaraParameterPanel::OnGetSectionTitle)
  		.OnGetSectionWidget(this, &SNiagaraParameterPanel::OnGetSectionWidget)
  		.OnCreateCustomRowExpander_Static(&SNiagaraParameterPanel::CreateCustomActionExpander)
@@ -93,6 +97,15 @@ void SNiagaraParameterPanel::Tick(const FGeometry& AllottedGeometry, const doubl
 		GraphActionMenu->RefreshAllActions(true);
 		bNeedsRefresh = false;
 	}
+
+	if (bGraphActionPendingRename)
+	{
+		if (GraphActionMenu->CanRequestRenameOnActionNode())
+		{
+			GraphActionMenu->OnRequestRenameOnActionNode();
+			bGraphActionPendingRename = false;
+		}
+	}
 }
 
 void SNiagaraParameterPanel::HandleExternalSelectionChanged(const UObject* Obj)
@@ -109,19 +122,19 @@ void SNiagaraParameterPanel::HandleExternalSelectionChanged(const UObject* Obj)
 
 void SNiagaraParameterPanel::AddParameter(FNiagaraVariable NewVariable, const NiagaraParameterPanelSectionID::Type SectionID)
 {
-	//FNiagaraVariableMetaData NewVariableMetaData = FNiagaraVariableMetaData();
 
 	FNiagaraVariableMetaData GuessedMetaData;
 	FNiagaraEditorUtilities::GetParameterMetaDataFromName(NewVariable.GetName(), GuessedMetaData);
+	const UNiagaraScriptVariable* NewScriptVar = ParameterPanelViewModel->AddParameter(NewVariable, GuessedMetaData);
+	
+	if (NewScriptVar != nullptr)
+	{
+		GraphActionMenu->RefreshAllActions(true);
+		GraphActionMenu->SelectItemByName(NewScriptVar->Variable.GetName());
 
-	//NewVariableMetaData.SetScopeName(FNiagaraEditorUtilities::GetScopeNameForParameterScope(NiagaraParameterPanelSectionID::GetScopeForNewParametersInSection(SectionID)));
-	//NewVariableMetaData.SetUsage(NiagaraParameterPanelSectionID::GetUsageForNewParametersInSection(SectionID));
-
-	ParameterPanelViewModel->AddParameter(NewVariable, GuessedMetaData);
-
-	GraphActionMenu->RefreshAllActions(true);
-	GraphActionMenu->SelectItemByName(NewVariable.GetName());
-	GraphActionMenu->OnRequestRenameOnActionNode();
+		// Delay calling the rename delegate one tick as the widget to which it is bound is not yet constructed.
+		bGraphActionPendingRename = true;
+	}
 }
 
 bool SNiagaraParameterPanel::AllowMakeType(const FNiagaraTypeDefinition& InType) const
@@ -203,6 +216,59 @@ void SNiagaraParameterPanel::OnActionSelected(const TArray< TSharedPtr<FEdGraphS
 	}
 }
 
+TSharedPtr<SWidget> SNiagaraParameterPanel::OnContextMenuOpening()
+{
+	// Get the current selected action
+	TArray<TSharedPtr<FEdGraphSchemaAction> > SelectedActions;
+	GraphActionMenu->GetSelectedActions(SelectedActions);
+	ensureMsgf(SelectedActions.Num() == 1 || SelectedActions.Num() == 0, TEXT("Unexpected number of selected actions encountered when getting current selected action for parameter panel!"));
+	if (SelectedActions.Num() == 1 && SelectedActions[0].IsValid())
+	{
+		TSharedPtr<FNiagaraScriptVarAndViewInfoAction> CurrentAction = StaticCastSharedPtr<FNiagaraScriptVarAndViewInfoAction>(SelectedActions[0]);
+		if (ensureMsgf(CurrentAction.Get(), TEXT("Action pointer was null when getting selected action for parameter panel!")))
+		{
+			const bool bShouldCloseWindowAfterMenuSelection = true;
+			FMenuBuilder MenuBuilder(bShouldCloseWindowAfterMenuSelection, ToolkitCommands);
+			MenuBuilder.BeginSection("BasicOperations");
+
+			const FNiagaraVariable& CurrentParameter = CurrentAction->ScriptVariableAndViewInfo.ScriptVariable;
+			const FNiagaraVariableMetaData& CurrentParameterMetaData = CurrentAction->ScriptVariableAndViewInfo.MetaData;
+
+			FText CanRenameParameterToolTip;
+			bool bCanRequestRenameParameter = ParameterPanelViewModel->GetCanRenameParameterAndToolTip(CurrentParameter, CurrentParameterMetaData, TOptional<const FText>(), CanRenameParameterToolTip);
+			MenuBuilder.AddMenuEntry(
+				LOCTEXT("ParameterPanelRenameParameterContextMenu", "Rename Parameter"),
+				CanRenameParameterToolTip,
+				FSlateIcon(),
+				FUIAction(
+					FExecuteAction::CreateLambda([this]()
+						{
+							OnRequestRenameOnActionNode();
+						}),
+					FCanExecuteAction::CreateLambda([bCanRequestRenameParameter]() {return bCanRequestRenameParameter; }))
+			);
+
+			FText CanDeleteParameterToolTip;
+			bool bCanDeleteParameter = ParameterPanelViewModel->GetCanDeleteParameterAndToolTip(CurrentParameter, CurrentParameterMetaData, CanDeleteParameterToolTip);
+			MenuBuilder.AddMenuEntry(
+				LOCTEXT("ParameterPanelDeleteParameterContextMenu", "Delete Parameter"),
+				CanDeleteParameterToolTip,
+				FSlateIcon(),
+				FUIAction(
+					FExecuteAction::CreateLambda([this, CurrentParameter, CurrentParameterMetaData]()
+						{
+							ParameterPanelViewModel->DeleteParameter(CurrentParameter, CurrentParameterMetaData);
+						}),
+					FCanExecuteAction::CreateLambda([bCanDeleteParameter]() {return bCanDeleteParameter; }))
+			);
+
+			MenuBuilder.EndSection();
+			return MenuBuilder.MakeWidget();
+		}
+	}
+	return SNullWidget::NullWidget;
+}
+
 FText SNiagaraParameterPanel::OnGetSectionTitle(int32 InSectionID)
 {
 	return NiagaraParameterPanelSectionID::OnGetSectionTitle((NiagaraParameterPanelSectionID::Type)InSectionID);
@@ -254,7 +320,7 @@ TSharedRef<SWidget> SNiagaraParameterPanel::OnGetParameterMenu(const NiagaraPara
 	const bool bCanCreateNew = InSection != NiagaraParameterPanelSectionID::Type::ENGINE && InSection != NiagaraParameterPanelSectionID::Type::OWNER;
 	const bool bAutoExpand = InSection == NiagaraParameterPanelSectionID::Type::LOCALS || InSection == NiagaraParameterPanelSectionID::Type::INPUTS ||
 		InSection == NiagaraParameterPanelSectionID::Type::OUTPUTS || InSection == NiagaraParameterPanelSectionID::Type::USER || 
-		InSection == NiagaraParameterPanelSectionID::Type::ENGINE || InSection == NiagaraParameterPanelSectionID::Type::OWNER;
+		InSection == NiagaraParameterPanelSectionID::Type::ENGINE && InSection == NiagaraParameterPanelSectionID::Type::OWNER;
 	TSharedRef<SNiagaraAddParameterMenu2> MenuWidget = SNew(SNiagaraAddParameterMenu2, ParameterPanelViewModel->GetEditableGraphs())
 		.OnAddParameter(this, &SNiagaraParameterPanel::AddParameter, InSection)
 		.OnAllowMakeType(this, &SNiagaraParameterPanel::AllowMakeType)
@@ -280,9 +346,10 @@ void SNiagaraParameterPanel::TryDeleteEntries()
 		{
 			const FNiagaraVariable& Parameter = ScriptVarAndViewInfoAction->ScriptVariableAndViewInfo.ScriptVariable;
 			const FNiagaraVariableMetaData& ParameterMetaData = ScriptVarAndViewInfoAction->ScriptVariableAndViewInfo.MetaData;
-			if (ParameterPanelViewModel->CanRemoveParameter(Parameter, ParameterMetaData))
+			FText ToolTipText;
+			if (ParameterPanelViewModel->GetCanDeleteParameterAndToolTip(Parameter, ParameterMetaData, ToolTipText))
 			{
-				ParameterPanelViewModel->RemoveParameter(Parameter, ParameterMetaData);
+				ParameterPanelViewModel->DeleteParameter(Parameter, ParameterMetaData);
 			}
 		}
 	}
@@ -300,7 +367,8 @@ bool SNiagaraParameterPanel::CanTryDeleteEntries() const
 		{
 			const FNiagaraVariable& Parameter = ScriptVarAndViewInfoAction->ScriptVariableAndViewInfo.ScriptVariable;
 			const FNiagaraVariableMetaData& ParameterMetaData = ScriptVarAndViewInfoAction->ScriptVariableAndViewInfo.MetaData;
-			if (ParameterPanelViewModel->CanRemoveParameter(Parameter, ParameterMetaData))
+			FText ToolTipText;
+			if (ParameterPanelViewModel->GetCanDeleteParameterAndToolTip(Parameter, ParameterMetaData, ToolTipText))
 			{
 				return true;
 			}
@@ -314,9 +382,21 @@ void SNiagaraParameterPanel::OnRequestRenameOnActionNode()
 	GraphActionMenu->OnRequestRenameOnActionNode();
 }
 
-bool SNiagaraParameterPanel::CanRequestRenameOnActionNode(TWeakPtr<struct FGraphActionNode> InSelectedNode) const
+bool SNiagaraParameterPanel::CanRequestRenameOnActionNode(TWeakPtr<FGraphActionNode> InSelectedNode) const
 {
-	return true; //@todo(ng) delegate to viewmodel
+	if (InSelectedNode.IsValid() && InSelectedNode.Pin()->Actions.Num() == 1 && InSelectedNode.Pin()->Actions[0].IsValid())
+	{
+		const FNiagaraScriptVarAndViewInfoAction* SelectedAction = static_cast<FNiagaraScriptVarAndViewInfoAction*>(InSelectedNode.Pin()->Actions[0].Get());
+		FText BlankText = FText();
+		return ParameterPanelViewModel->GetCanRenameParameterAndToolTip(
+			SelectedAction->ScriptVariableAndViewInfo.ScriptVariable,
+			SelectedAction->ScriptVariableAndViewInfo.MetaData,
+			TOptional<const FText>(),
+			BlankText
+		);
+	}
+	
+	return false;
 }
 
 bool SNiagaraParameterPanel::CanRequestRenameOnActionNode() const
@@ -649,9 +729,4 @@ TSharedRef<SExpanderArrow> SNiagaraParameterPanel::CreateCustomActionExpander(co
 	return SNew(SNiagaraActionMenuExpander2, ActionMenuData);
 }
 
-void SNiagaraParameterPanel::SetVariablesObjectSelection(const TSharedRef<FNiagaraObjectSelection>& InVariablesObjectSelection)
-{
-	SelectedVariableObjects = InVariablesObjectSelection;
-}
-
 #undef LOCTEXT_NAMESPACE // "NiagaraParameterPanel"
diff --git a/Engine/Plugins/FX/Niagara/Source/NiagaraEditor/Private/Widgets/SNiagaraParameterPanel.h b/Engine/Plugins/FX/Niagara/Source/NiagaraEditor/Private/Widgets/SNiagaraParameterPanel.h
--- a/Engine/Plugins/FX/Niagara/Source/NiagaraEditor/Private/Widgets/SNiagaraParameterPanel.h
+++ b/Engine/Plugins/FX/Niagara/Source/NiagaraEditor/Private/Widgets/SNiagaraParameterPanel.h
@@ -67,8 +67,6 @@ public:
 
 	static TSharedRef<SExpanderArrow> CreateCustomActionExpander(const struct FCustomExpanderData& ActionMenuData);
 
-	void SetVariablesObjectSelection(const TSharedRef<FNiagaraObjectSelection>& InVariablesObjectSelection);
-
 private:
 	/** Function to bind to SNiagaraAddParameterMenus to filter types we allow creating */
 	bool AllowMakeType(const FNiagaraTypeDefinition& InType) const;
@@ -80,7 +78,7 @@ private:
 	FReply OnActionDragged(const TArray<TSharedPtr<FEdGraphSchemaAction>>& InActions, const FPointerEvent& MouseEvent);
 	void OnActionSelected(const TArray<TSharedPtr<FEdGraphSchemaAction>>& InActions, ESelectInfo::Type InSelectionType);
 // 	void OnActionDoubleClicked(const TArray<TSharedPtr<FEdGraphSchemaAction>>& InActions); //@todo(ng) impl
-// 	TSharedPtr<SWidget> OnContextMenuOpening();
+	TSharedPtr<SWidget> OnContextMenuOpening();
  	FText OnGetSectionTitle(int32 InSectionID);
 	TSharedRef<SWidget> OnGetSectionWidget(TSharedRef<SWidget> RowWidget, int32 InSectionID);
 	TSharedRef<SWidget> CreateAddToSectionButton(const NiagaraParameterPanelSectionID::Type InSection, TWeakPtr<SWidget> WeakRowWidget, FText AddNewText, FName MetaDataTag);
@@ -123,6 +121,7 @@ private:
 	TSharedPtr<FUICommandList> ToolkitCommands; //@todo(ng) add Find And Rename Parameter command
 
 	bool bNeedsRefresh;
+	bool bGraphActionPendingRename;
 
 	TSharedPtr<INiagaraParameterPanelViewModel> ParameterPanelViewModel;
 };
