Niagara - Fix CIS warnings for content and uninitialized variables.  Also fix up the view models to recompile properly when making the asset fixes.

diff --git a/Engine/Plugins/FX/Niagara/Source/Niagara/Private/MovieScene/MovieSceneNiagaraSystemTrackTemplate.cpp b/Engine/Plugins/FX/Niagara/Source/Niagara/Private/MovieScene/MovieSceneNiagaraSystemTrackTemplate.cpp
--- a/Engine/Plugins/FX/Niagara/Source/Niagara/Private/MovieScene/MovieSceneNiagaraSystemTrackTemplate.cpp
+++ b/Engine/Plugins/FX/Niagara/Source/Niagara/Private/MovieScene/MovieSceneNiagaraSystemTrackTemplate.cpp
@@ -215,6 +215,16 @@ FMovieSceneNiagaraSystemTrackImplementation::FMovieSceneNiagaraSystemTrackImplem
 {
 }
 
+FMovieSceneNiagaraSystemTrackImplementation::FMovieSceneNiagaraSystemTrackImplementation()
+	: SpawnSectionStartFrame(FFrameNumber())
+	, SpawnSectionEndFrame(FFrameNumber())
+	, SpawnSectionStartBehavior(ENiagaraSystemSpawnSectionStartBehavior::Activate)
+	, SpawnSectionEvaluateBehavior(ENiagaraSystemSpawnSectionEvaluateBehavior::None)
+	, SpawnSectionEndBehavior(ENiagaraSystemSpawnSectionEndBehavior::SetSystemInactive)
+	, AgeUpdateMode(ENiagaraAgeUpdateMode::TickDeltaTime)
+{
+}
+
 void FMovieSceneNiagaraSystemTrackImplementation::Evaluate(const FMovieSceneEvaluationTrack& Track, FMovieSceneSegmentIdentifier SegmentID, const FMovieSceneEvaluationOperand& Operand, const FMovieSceneContext& Context, const FPersistentEvaluationData& PersistentData, FMovieSceneExecutionTokens& ExecutionTokens) const
 {
 	ExecutionTokens.SetContext(Context);
diff --git a/Engine/Plugins/FX/Niagara/Source/Niagara/Private/MovieScene/MovieSceneNiagaraSystemTrackTemplate.h b/Engine/Plugins/FX/Niagara/Source/Niagara/Private/MovieScene/MovieSceneNiagaraSystemTrackTemplate.h
--- a/Engine/Plugins/FX/Niagara/Source/Niagara/Private/MovieScene/MovieSceneNiagaraSystemTrackTemplate.h
+++ b/Engine/Plugins/FX/Niagara/Source/Niagara/Private/MovieScene/MovieSceneNiagaraSystemTrackTemplate.h
@@ -26,8 +26,7 @@ struct FMovieSceneNiagaraSystemTrackImplementation : public FMovieSceneTrackImpl
 	GENERATED_BODY()
 
 public:
-	FMovieSceneNiagaraSystemTrackImplementation() { }
-	FMovieSceneNiagaraSystemTrackImplementation(FFrameNumber InSpawnSectionStartFrame, FFrameNumber InSpawnSectionEndFrame);
+	FMovieSceneNiagaraSystemTrackImplementation();
 	FMovieSceneNiagaraSystemTrackImplementation(
 		FFrameNumber InSpawnSectionStartFrame, FFrameNumber InSpawnSectionEndFrame,
 		ENiagaraSystemSpawnSectionStartBehavior InSectionStartBehavior, ENiagaraSystemSpawnSectionEvaluateBehavior InSectionEvaluateBehavior,
diff --git a/Engine/Plugins/FX/Niagara/Source/NiagaraEditor/Private/ViewModels/NiagaraSystemViewModel.cpp b/Engine/Plugins/FX/Niagara/Source/NiagaraEditor/Private/ViewModels/NiagaraSystemViewModel.cpp
--- a/Engine/Plugins/FX/Niagara/Source/NiagaraEditor/Private/ViewModels/NiagaraSystemViewModel.cpp
+++ b/Engine/Plugins/FX/Niagara/Source/NiagaraEditor/Private/ViewModels/NiagaraSystemViewModel.cpp
@@ -1981,7 +1981,10 @@ void FNiagaraSystemViewModel::AddSystemEventHandlers()
 			FNiagaraParameterStore::FOnChanged::FDelegate::CreateSP<FNiagaraSystemViewModel, const FNiagaraParameterStore&, const UNiagaraScript*>(
 				this->AsShared(), &FNiagaraSystemViewModel::SystemParameterStoreChanged, System->GetExposedParameters(), nullptr));
 
-		SystemScriptGraphChangedHandle = SystemScriptViewModel->GetGraphViewModel()->GetGraph()->AddOnGraphChangedHandler(FOnGraphChanged::FDelegate::CreateSP(this->AsShared(), &FNiagaraSystemViewModel::SystemScriptGraphChanged));
+		SystemScriptGraphChangedHandle = SystemScriptViewModel->GetGraphViewModel()->GetGraph()->AddOnGraphChangedHandler(
+			FOnGraphChanged::FDelegate::CreateSP(this->AsShared(), &FNiagaraSystemViewModel::SystemScriptGraphChanged));
+		SystemScriptGraphNeedsRecompileHandle = SystemScriptViewModel->GetGraphViewModel()->GetGraph()->AddOnGraphNeedsRecompileHandler(
+			FOnGraphChanged::FDelegate::CreateSP(this->AsShared(), &FNiagaraSystemViewModel::SystemScriptGraphChanged));
 	}
 }
 
@@ -2009,12 +2012,14 @@ void FNiagaraSystemViewModel::RemoveSystemEventHandlers()
 		if (SystemScriptViewModel.IsValid())
 		{
 			SystemScriptViewModel->GetGraphViewModel()->GetGraph()->RemoveOnGraphChangedHandler(SystemScriptGraphChangedHandle);
+			SystemScriptViewModel->GetGraphViewModel()->GetGraph()->RemoveOnGraphNeedsRecompileHandler(SystemScriptGraphNeedsRecompileHandle);
 		}
 	}
 
 	ScriptToOnParameterStoreChangedHandleMap.Empty();
 	UserParameterStoreChangedHandle.Reset();
 	SystemScriptGraphChangedHandle.Reset();
+	SystemScriptGraphNeedsRecompileHandle.Reset();
 }
 
 void FNiagaraSystemViewModel::NotifyPinnedCurvesChanged()
diff --git a/Engine/Plugins/FX/Niagara/Source/NiagaraEditor/Private/ViewModels/Stack/NiagaraStackFunctionInputCollection.cpp b/Engine/Plugins/FX/Niagara/Source/NiagaraEditor/Private/ViewModels/Stack/NiagaraStackFunctionInputCollection.cpp
--- a/Engine/Plugins/FX/Niagara/Source/NiagaraEditor/Private/ViewModels/Stack/NiagaraStackFunctionInputCollection.cpp
+++ b/Engine/Plugins/FX/Niagara/Source/NiagaraEditor/Private/ViewModels/Stack/NiagaraStackFunctionInputCollection.cpp
@@ -296,7 +296,12 @@ UNiagaraStackEntry::FStackIssueFix UNiagaraStackFunctionInputCollection::GetRese
 	{
 		FScopedTransaction ScopedTransaction(FixDescription);
 		const UEdGraphSchema_Niagara* NiagaraSchema = GetDefault<UEdGraphSchema_Niagara>();
+		UNiagaraNode* OwningNiagaraNode = Cast<UNiagaraNode>(PinToReset->GetOwningNode());
 		NiagaraSchema->ResetPinToAutogeneratedDefaultValue(PinToReset);
+		if (OwningNiagaraNode != nullptr)
+		{
+			OwningNiagaraNode->MarkNodeRequiresSynchronization("Pin reset to default value.", true);
+		}
 	}));
 }
 
@@ -329,7 +334,7 @@ void UNiagaraStackFunctionInputCollection::RefreshIssues(TArray<FName> Duplicate
 				ValidAliasedInputNames.Contains(OverridePin->PinName) == false)
 			{
 				FStackIssue InvalidInputOverrideError(
-					EStackIssueSeverity::Error,
+					EStackIssueSeverity::Warning,
 					FText::Format(LOCTEXT("InvalidInputOverrideSummaryFormat", "Invalid Input Override: {0}"), FText::FromString(OverridePin->PinName.ToString())),
 					FText::Format(LOCTEXT("InvalidInputOverrideFormat", "The input {0} was previously overriden but is no longer exposed by the function {1}.\nPress the fix button to remove this unused override data,\nor check the function definition to see why this input is no longer exposed."),
 						FText::FromString(OverridePin->PinName.ToString()), FText::FromString(InputFunctionCallNode->GetFunctionName())),
diff --git a/Engine/Plugins/FX/Niagara/Source/NiagaraEditor/Public/ViewModels/NiagaraSystemViewModel.h b/Engine/Plugins/FX/Niagara/Source/NiagaraEditor/Public/ViewModels/NiagaraSystemViewModel.h
--- a/Engine/Plugins/FX/Niagara/Source/NiagaraEditor/Public/ViewModels/NiagaraSystemViewModel.h
+++ b/Engine/Plugins/FX/Niagara/Source/NiagaraEditor/Public/ViewModels/NiagaraSystemViewModel.h
@@ -561,6 +561,9 @@ private:
 	/** A handle to the on graph changed delegate for the system script. */
 	FDelegateHandle SystemScriptGraphChangedHandle;
 
+	/** A handle to the on graph needs recompile delegate for the system script. */
+	FDelegateHandle SystemScriptGraphNeedsRecompileHandle;
+
 	/** An array of emitter handle ids which need their sequencer tracks refreshed next frame. */
 	TArray<FGuid> EmitterIdsRequiringSequencerTrackUpdate;
 
