Parameter map get and set nodes now get a drop target highlight when hovered over.

diff --git a/Engine/Plugins/FX/Niagara/Source/NiagaraEditor/Private/NiagaraActions.cpp b/Engine/Plugins/FX/Niagara/Source/NiagaraEditor/Private/NiagaraActions.cpp
--- a/Engine/Plugins/FX/Niagara/Source/NiagaraEditor/Private/NiagaraActions.cpp
+++ b/Engine/Plugins/FX/Niagara/Source/NiagaraEditor/Private/NiagaraActions.cpp
@@ -196,6 +196,11 @@ FReply FNiagaraParameterGraphDragOperation::DroppedOnPanel(const TSharedRef<SWid
 }
 
 
+bool FNiagaraParameterGraphDragOperation::IsCurrentlyHoveringNode(const UEdGraphNode* TestNode) const
+{
+	return TestNode == GetHoveredNode();
+}
+
 void FNiagaraParameterGraphDragOperation::MakeGetMap(FNiagaraParameterNodeConstructionParams InParams)
 {
 
diff --git a/Engine/Plugins/FX/Niagara/Source/NiagaraEditor/Private/NiagaraEditorStyle.cpp b/Engine/Plugins/FX/Niagara/Source/NiagaraEditor/Private/NiagaraEditorStyle.cpp
--- a/Engine/Plugins/FX/Niagara/Source/NiagaraEditor/Private/NiagaraEditorStyle.cpp
+++ b/Engine/Plugins/FX/Niagara/Source/NiagaraEditor/Private/NiagaraEditorStyle.cpp
@@ -326,6 +326,11 @@ TSharedRef< FSlateStyleSet > FNiagaraEditorStyle::Create()
 		.SetInactiveBrush(IMAGE_CORE_BRUSH("Common/Selection", Icon8x8, SelectionColor_Inactive))
 		.SetInactiveHoveredBrush(IMAGE_CORE_BRUSH("Common/Selection", Icon8x8, SelectionColor_Inactive)));
 
+	Style->Set("NiagaraEditor.DropTarget.BackgroundColor", FLinearColor(1.0f, 1.0f, 1.0f, 0.25f));
+	Style->Set("NiagaraEditor.DropTarget.BackgroundColorHover", FLinearColor(1.0f, 1.0f, 1.0f, 0.1f));
+	Style->Set("NiagaraEditor.DropTarget.BorderVertical", new IMAGE_PLUGIN_BRUSH("Icons/StackDropTargetBorder_Vertical", FVector2D(2, 8), FLinearColor::White, ESlateBrushTileType::Vertical));
+	Style->Set("NiagaraEditor.DropTarget.BorderHorizontal", new IMAGE_PLUGIN_BRUSH("Icons/StackDropTargetBorder_Horizontal", FVector2D(8, 2), FLinearColor::White, ESlateBrushTileType::Horizontal));
+
 	return Style;
 }
 
diff --git a/Engine/Plugins/FX/Niagara/Source/NiagaraEditor/Private/NiagaraNodeParameterMapBase.cpp b/Engine/Plugins/FX/Niagara/Source/NiagaraEditor/Private/NiagaraNodeParameterMapBase.cpp
--- a/Engine/Plugins/FX/Niagara/Source/NiagaraEditor/Private/NiagaraNodeParameterMapBase.cpp
+++ b/Engine/Plugins/FX/Niagara/Source/NiagaraEditor/Private/NiagaraNodeParameterMapBase.cpp
@@ -191,6 +191,11 @@ void UNiagaraNodeParameterMapBase::SetPinName(UEdGraphPin* InPin, const FName& I
 	OnPinRenamed(InPin, OldName.ToString());
 }
 
+bool UNiagaraNodeParameterMapBase::OnAllowDrop(TSharedPtr<FDragDropOperation> DragDropOperation)
+{
+	return true;
+}
+
 void UNiagaraNodeParameterMapBase::OnPinRenamed(UEdGraphPin* RenamedPin, const FString& OldName)
 {
 	RenamedPin->PinFriendlyName = FText::FromName(RenamedPin->PinName);
diff --git a/Engine/Plugins/FX/Niagara/Source/NiagaraEditor/Private/NiagaraNodeParameterMapBase.h b/Engine/Plugins/FX/Niagara/Source/NiagaraEditor/Private/NiagaraNodeParameterMapBase.h
--- a/Engine/Plugins/FX/Niagara/Source/NiagaraEditor/Private/NiagaraNodeParameterMapBase.h
+++ b/Engine/Plugins/FX/Niagara/Source/NiagaraEditor/Private/NiagaraNodeParameterMapBase.h
@@ -40,6 +40,7 @@ public:
 
 	void SetPinName(UEdGraphPin* InPin, const FName& InName);
 
+	bool OnAllowDrop(TSharedPtr<FDragDropOperation> DragDropOperation);
 protected:
 	virtual void SelectParameterFromPin(UEdGraphPin* InPin);
 	virtual void OnPinRenamed(UEdGraphPin* RenamedPin, const FString& OldName) override;
diff --git a/Engine/Plugins/FX/Niagara/Source/NiagaraEditor/Private/NiagaraNodeParameterMapSet.cpp b/Engine/Plugins/FX/Niagara/Source/NiagaraEditor/Private/NiagaraNodeParameterMapSet.cpp
--- a/Engine/Plugins/FX/Niagara/Source/NiagaraEditor/Private/NiagaraNodeParameterMapSet.cpp
+++ b/Engine/Plugins/FX/Niagara/Source/NiagaraEditor/Private/NiagaraNodeParameterMapSet.cpp
@@ -15,6 +15,7 @@
 #include "NiagaraScriptVariable.h"
 
 #include "ScopedTransaction.h"
+#include "SNiagaraGraphParameterMapSetNode.h"
 
 #define LOCTEXT_NAMESPACE "NiagaraNodeParameterMapSet"
 
@@ -32,6 +33,11 @@ void UNiagaraNodeParameterMapSet::AllocateDefaultPins()
 	CreateAddPin(EGPD_Input);
 }
 
+TSharedPtr<SGraphNode> UNiagaraNodeParameterMapSet::CreateVisualWidget()
+{
+	return SNew(SNiagaraGraphParameterMapSetNode, this);
+}
+
 bool UNiagaraNodeParameterMapSet::IsPinNameEditable(const UEdGraphPin* GraphPinObj) const
 {
 	const UEdGraphSchema_Niagara* Schema = GetDefault<UEdGraphSchema_Niagara>();
diff --git a/Engine/Plugins/FX/Niagara/Source/NiagaraEditor/Private/NiagaraNodeParameterMapSet.h b/Engine/Plugins/FX/Niagara/Source/NiagaraEditor/Private/NiagaraNodeParameterMapSet.h
--- a/Engine/Plugins/FX/Niagara/Source/NiagaraEditor/Private/NiagaraNodeParameterMapSet.h
+++ b/Engine/Plugins/FX/Niagara/Source/NiagaraEditor/Private/NiagaraNodeParameterMapSet.h
@@ -20,6 +20,7 @@ public:
 
 	virtual void AllocateDefaultPins() override;
 
+	virtual TSharedPtr<SGraphNode> CreateVisualWidget() override;
 	virtual bool IsPinNameEditable(const UEdGraphPin* GraphPinObj) const override;
 	virtual bool IsPinNameEditableUponCreation(const UEdGraphPin* GraphPinObj) const override;
 	virtual void RemoveDynamicPin(UEdGraphPin* Pin) override;
diff --git a/Engine/Plugins/FX/Niagara/Source/NiagaraEditor/Private/Widgets/SNiagaraGraphParameterMapGetNode.cpp b/Engine/Plugins/FX/Niagara/Source/NiagaraEditor/Private/Widgets/SNiagaraGraphParameterMapGetNode.cpp
--- a/Engine/Plugins/FX/Niagara/Source/NiagaraEditor/Private/Widgets/SNiagaraGraphParameterMapGetNode.cpp
+++ b/Engine/Plugins/FX/Niagara/Source/NiagaraEditor/Private/Widgets/SNiagaraGraphParameterMapGetNode.cpp
@@ -9,6 +9,9 @@
 #include "SGraphPin.h"
 #include "EdGraphSchema_Niagara.h"
 #include "NiagaraScriptVariable.h"
+#include "SDropTarget.h"
+#include "NiagaraEditorStyle.h"
+
 
 #define LOCTEXT_NAMESPACE "SNiagaraGraphParameterMapGetNode"
 
@@ -55,13 +58,23 @@ void SNiagaraGraphParameterMapGetNode::AddPin(const TSharedRef<SGraphPin>& PinTo
 TSharedRef<SWidget> SNiagaraGraphParameterMapGetNode::CreateNodeContentArea()
 {
 	// NODE CONTENT AREA
-	return SNew(SBorder)
-		.BorderImage(FEditorStyle::GetBrush("NoBorder"))
-		.HAlign(HAlign_Fill)
-		.VAlign(VAlign_Fill)
-		.Padding(FMargin(0, 3))
+	return 	SNew(SDropTarget)
+		.OnDrop(this, &SNiagaraGraphParameterMapGetNode::OnDroppedOnTarget)
+		.OnAllowDrop(this, &SNiagaraGraphParameterMapGetNode::OnAllowDrop)
+		.HorizontalImage(FNiagaraEditorStyle::Get().GetBrush("NiagaraEditor.DropTarget.BorderHorizontal"))
+		.VerticalImage(FNiagaraEditorStyle::Get().GetBrush("NiagaraEditor.DropTarget.BorderVertical"))
+		.BackgroundColor(FNiagaraEditorStyle::Get().GetColor("NiagaraEditor.DropTarget.BackgroundColor"))
+		.BackgroundColorHover(FNiagaraEditorStyle::Get().GetColor("NiagaraEditor.DropTarget.BackgroundColorHover"))
+		.Content()
 		[
-			SAssignNew(PinContainerRoot, SVerticalBox)
+			SNew(SBorder)
+			.BorderImage(FEditorStyle::GetBrush("NoBorder"))
+			.HAlign(HAlign_Fill)
+			.VAlign(VAlign_Fill)
+			.Padding(FMargin(0, 3))
+			[
+				SAssignNew(PinContainerRoot, SVerticalBox)
+			]
 		];
 }
 
@@ -185,4 +198,20 @@ FReply SNiagaraGraphParameterMapGetNode::OnBorderMouseButtonDown(const FGeometry
 	return FReply::Handled();
 }
 
+FReply SNiagaraGraphParameterMapGetNode::OnDroppedOnTarget(TSharedPtr<FDragDropOperation> DropOperation)
+{
+	return FReply::Unhandled();
+}
+
+bool SNiagaraGraphParameterMapGetNode::OnAllowDrop(TSharedPtr<FDragDropOperation> DragDropOperation)
+{
+	UNiagaraNodeParameterMapBase* MapNode = Cast<UNiagaraNodeParameterMapBase>(GraphNode);
+	if (MapNode
+		&& DragDropOperation->IsOfType<FNiagaraParameterGraphDragOperation>()
+		&& StaticCastSharedPtr<FNiagaraParameterGraphDragOperation>(DragDropOperation)->IsCurrentlyHoveringNode(GraphNode))
+	{
+		return MapNode->OnAllowDrop(DragDropOperation);
+	}
+	return false;
+}
 #undef LOCTEXT_NAMESPACE
diff --git a/Engine/Plugins/FX/Niagara/Source/NiagaraEditor/Private/Widgets/SNiagaraGraphParameterMapGetNode.h b/Engine/Plugins/FX/Niagara/Source/NiagaraEditor/Private/Widgets/SNiagaraGraphParameterMapGetNode.h
--- a/Engine/Plugins/FX/Niagara/Source/NiagaraEditor/Private/Widgets/SNiagaraGraphParameterMapGetNode.h
+++ b/Engine/Plugins/FX/Niagara/Source/NiagaraEditor/Private/Widgets/SNiagaraGraphParameterMapGetNode.h
@@ -18,6 +18,8 @@ public:
 protected:
 	FReply OnBorderMouseButtonDown(const FGeometry& MyGeometry, const FPointerEvent& MouseEvent, int32 InWhichPin);
 	const FSlateBrush* GetBackgroundBrush(TSharedPtr<SWidget> Border) const;
+	FReply OnDroppedOnTarget(TSharedPtr<FDragDropOperation> DropOperation);
+	bool OnAllowDrop(TSharedPtr<FDragDropOperation> DragDropOperation);
 	TSharedPtr<SVerticalBox> PinContainerRoot;
 
 	const FSlateBrush* BackgroundBrush;
diff --git a/Engine/Plugins/FX/Niagara/Source/NiagaraEditor/Private/Widgets/SNiagaraGraphParameterMapSetNode.cpp b/Engine/Plugins/FX/Niagara/Source/NiagaraEditor/Private/Widgets/SNiagaraGraphParameterMapSetNode.cpp
new file mode 100644
--- /dev/null
+++ b/Engine/Plugins/FX/Niagara/Source/NiagaraEditor/Private/Widgets/SNiagaraGraphParameterMapSetNode.cpp
@@ -0,0 +1,68 @@
+// Copyright Epic Games, Inc. All Rights Reserved.
+
+#include "SNiagaraGraphParameterMapSetNode.h"
+#include "NiagaraNodeParameterMapSet.h"
+#include "Widgets/Input/SButton.h"
+#include "GraphEditorSettings.h"
+#include "Rendering/DrawElements.h"
+#include "Widgets/SBoxPanel.h"
+#include "SGraphPin.h"
+#include "EdGraphSchema_Niagara.h"
+#include "NiagaraScriptVariable.h"
+#include "SDropTarget.h"
+#include "NiagaraEditorStyle.h"
+
+
+#define LOCTEXT_NAMESPACE "SNiagaraGraphParameterMapSetNode"
+
+
+void SNiagaraGraphParameterMapSetNode::Construct(const FArguments& InArgs, UEdGraphNode* InGraphNode)
+{
+
+	GraphNode = InGraphNode; 
+	RegisterNiagaraGraphNode(InGraphNode);
+
+	UpdateGraphNode();
+}
+
+TSharedRef<SWidget> SNiagaraGraphParameterMapSetNode::CreateNodeContentArea()
+{
+	TSharedRef<SWidget> ExistingContent = SNiagaraGraphNode::CreateNodeContentArea();
+	// NODE CONTENT AREA
+	return 	SNew(SDropTarget)
+		.OnDrop(this, &SNiagaraGraphParameterMapSetNode::OnDroppedOnTarget)
+		.OnAllowDrop(this, &SNiagaraGraphParameterMapSetNode::OnAllowDrop)
+		.HorizontalImage(FNiagaraEditorStyle::Get().GetBrush("NiagaraEditor.DropTarget.BorderHorizontal"))
+		.VerticalImage(FNiagaraEditorStyle::Get().GetBrush("NiagaraEditor.DropTarget.BorderVertical"))
+		.BackgroundColor(FNiagaraEditorStyle::Get().GetColor("NiagaraEditor.DropTarget.BackgroundColor"))
+		.BackgroundColorHover(FNiagaraEditorStyle::Get().GetColor("NiagaraEditor.DropTarget.BackgroundColorHover"))
+		.Content()
+		[
+			SNew(SBorder)
+			.BorderImage(FEditorStyle::GetBrush("NoBorder"))
+			.HAlign(HAlign_Fill)
+			.VAlign(VAlign_Fill)
+			.Padding(FMargin(0, 3))
+			[
+				ExistingContent
+			]
+		];
+}
+
+FReply SNiagaraGraphParameterMapSetNode::OnDroppedOnTarget(TSharedPtr<FDragDropOperation> DropOperation)
+{
+	return FReply::Unhandled();
+}
+
+bool SNiagaraGraphParameterMapSetNode::OnAllowDrop(TSharedPtr<FDragDropOperation> DragDropOperation)
+{
+	UNiagaraNodeParameterMapBase* MapNode = Cast<UNiagaraNodeParameterMapBase>(GraphNode);
+	if (MapNode
+		&& DragDropOperation->IsOfType<FNiagaraParameterGraphDragOperation>()
+		&& StaticCastSharedPtr<FNiagaraParameterGraphDragOperation>(DragDropOperation)->IsCurrentlyHoveringNode(GraphNode))
+	{
+		return MapNode->OnAllowDrop(DragDropOperation);
+	}
+	return false;
+}
+#undef LOCTEXT_NAMESPACE
diff --git a/Engine/Plugins/FX/Niagara/Source/NiagaraEditor/Private/Widgets/SNiagaraGraphParameterMapSetNode.h b/Engine/Plugins/FX/Niagara/Source/NiagaraEditor/Private/Widgets/SNiagaraGraphParameterMapSetNode.h
new file mode 100644
--- /dev/null
+++ b/Engine/Plugins/FX/Niagara/Source/NiagaraEditor/Private/Widgets/SNiagaraGraphParameterMapSetNode.h
@@ -0,0 +1,19 @@
+// Copyright Epic Games, Inc. All Rights Reserved.
+
+#pragma once
+#include "SNiagaraGraphNode.h"
+
+class SNiagaraGraphParameterMapSetNode : public SNiagaraGraphNode
+{
+public:
+	SLATE_BEGIN_ARGS(SNiagaraGraphParameterMapSetNode) {}
+	SLATE_END_ARGS();
+
+	void Construct(const FArguments& InArgs, UEdGraphNode* InGraphNode);
+	virtual TSharedRef<SWidget> CreateNodeContentArea();
+
+protected:
+	FReply OnDroppedOnTarget(TSharedPtr<FDragDropOperation> DropOperation);
+	bool OnAllowDrop(TSharedPtr<FDragDropOperation> DragDropOperation);
+
+};
\ No newline at end of file
diff --git a/Engine/Plugins/FX/Niagara/Source/NiagaraEditor/Public/NiagaraActions.h b/Engine/Plugins/FX/Niagara/Source/NiagaraEditor/Public/NiagaraActions.h
--- a/Engine/Plugins/FX/Niagara/Source/NiagaraEditor/Public/NiagaraActions.h
+++ b/Engine/Plugins/FX/Niagara/Source/NiagaraEditor/Public/NiagaraActions.h
@@ -94,6 +94,9 @@ public:
 	/** Set if operation is modified by the ctrl key */
 	void SetCtrlDrag(bool InIsCtrlDrag) { bControlDrag = InIsCtrlDrag; }
 
+	/** Returns true if the drag operation is currently hovering over the supplied node */
+	bool IsCurrentlyHoveringNode(const UEdGraphNode* TestNode) const;
+
 protected:
 	/** Constructor */
 	FNiagaraParameterGraphDragOperation();
