Fixing Linux CIS

diff --git a/Engine/Plugins/FX/Niagara/Source/NiagaraEditor/Private/ViewModels/NiagaraParameterPanelViewModel.cpp b/Engine/Plugins/FX/Niagara/Source/NiagaraEditor/Private/ViewModels/NiagaraParameterPanelViewModel.cpp
--- a/Engine/Plugins/FX/Niagara/Source/NiagaraEditor/Private/ViewModels/NiagaraParameterPanelViewModel.cpp
+++ b/Engine/Plugins/FX/Niagara/Source/NiagaraEditor/Private/ViewModels/NiagaraParameterPanelViewModel.cpp
@@ -764,7 +764,7 @@ bool FNiagaraScriptToolkitParameterPanelViewModel::GetCanRenameParameterAndToolT
 		ensureMsgf(TargetVariableMetaData.GetParameterName(ExistingName), TEXT("Failed to get namespaceless parameter name from metadata on rename verify!"));
 		if (NewName != ExistingName)
 		{
-			if (TargetVariableMetaData.GetUsage() == ENiagaraScriptParameterUsage::Input && TargetVariableMetaData.GetScopeName() == FNiagaraConstants::InputScopeName || TargetVariableMetaData.GetUsage() == ENiagaraScriptParameterUsage::Output)
+			if ((TargetVariableMetaData.GetUsage() == ENiagaraScriptParameterUsage::Input && TargetVariableMetaData.GetScopeName() == FNiagaraConstants::InputScopeName) || TargetVariableMetaData.GetUsage() == ENiagaraScriptParameterUsage::Output)
 			{
 				// Inputs may not name alias outputs and vice versa
 				AliasScopeNames.Add(FNiagaraConstants::InputScopeName);
