The platform include and exclude buttons now correctly update the tree view's expansion buttons in the platform set menu.

diff --git a/Engine/Plugins/FX/Niagara/Source/NiagaraEditor/Private/Customizations/NiagaraPlatformSetCustomization.cpp b/Engine/Plugins/FX/Niagara/Source/NiagaraEditor/Private/Customizations/NiagaraPlatformSetCustomization.cpp
--- a/Engine/Plugins/FX/Niagara/Source/NiagaraEditor/Private/Customizations/NiagaraPlatformSetCustomization.cpp
+++ b/Engine/Plugins/FX/Niagara/Source/NiagaraEditor/Private/Customizations/NiagaraPlatformSetCustomization.cpp
@@ -953,6 +953,8 @@ FReply FNiagaraPlatformSetCustomization::OnProfileMenuButtonClicked(TSharedPtr<F
 		EffectsQualityMenuAnchors[EffectsQuality]->SetIsOpen(false);
 	}
 
+	DeviceProfileTreeWidget->RequestTreeRefresh();
+
 	return FReply::Handled();
 }
 
