Making inline pin renaming work in Niagara graphs

diff --git a/Engine/Plugins/FX/Niagara/Source/Niagara/Private/NiagaraTypes.cpp b/Engine/Plugins/FX/Niagara/Source/Niagara/Private/NiagaraTypes.cpp
--- a/Engine/Plugins/FX/Niagara/Source/Niagara/Private/NiagaraTypes.cpp
+++ b/Engine/Plugins/FX/Niagara/Source/Niagara/Private/NiagaraTypes.cpp
@@ -68,3 +68,13 @@ void FNiagaraVariableMetaData::CopyPerScriptMetaData(const FNiagaraVariableMetaD
 	SetWasCreatedInSystemEditor(OtherMetaData.GetWasCreatedInSystemEditor());
 	SetIsUsingLegacyNameString(OtherMetaData.GetIsUsingLegacyNameString());
 }
+
+
+void FNiagaraVariableMetaData::SetCachedNamespacelessVariableName(const FName& InVariableName)
+{
+	/*if (InVariableName == NAME_None || InVariableName == TEXT("None"))
+	{
+	}*/
+	//UE_LOG(LogNiagara, Log, TEXT("SetCachedNamespacelessVariableName %s!"), *InVariableName.ToString());
+	CachedNamespacelessVariableName = InVariableName;
+};
diff --git a/Engine/Plugins/FX/Niagara/Source/Niagara/Public/NiagaraTypes.h b/Engine/Plugins/FX/Niagara/Source/Niagara/Public/NiagaraTypes.h
--- a/Engine/Plugins/FX/Niagara/Source/Niagara/Public/NiagaraTypes.h
+++ b/Engine/Plugins/FX/Niagara/Source/Niagara/Public/NiagaraTypes.h
@@ -679,7 +679,7 @@ public:
 	 * @return bool			Whether the CachedNamespacelessParameterName can be returned. Is false if bUseLegacyNameString is set.
 	 */
 	bool GetParameterName(FName& OutName) const;
-	void SetCachedNamespacelessVariableName(const FName& InVariableName) { CachedNamespacelessVariableName = InVariableName; };
+	void SetCachedNamespacelessVariableName(const FName& InVariableName);;
 
 	bool GetWasCreatedInSystemEditor() const { return bCreatedInSystemEditor; };
 	void SetWasCreatedInSystemEditor(bool bWasCreatedInSystemEditor) { bCreatedInSystemEditor = bWasCreatedInSystemEditor; };
@@ -727,7 +727,7 @@ private:
 	bool bOutputIsPersistent;
 
 	/** Namespace-less name for associated FNiagaraVariable. Edited directly by user and then used to generate full Name of associated FNiagaraVariable. */
-	UPROPERTY(meta = (SkipForCompileHash = "true"))
+	UPROPERTY(VisibleAnywhere, Category = "Variable", DisplayName = "Property Name", meta = (SkipForCompileHash = "true"))
 	FName CachedNamespacelessVariableName;
 
 	/** Track if the associated parameter was created in the Emitter/System editor. Used to determine whether the associated parameter can be deleted from the Emitter/System editor. */
diff --git a/Engine/Plugins/FX/Niagara/Source/NiagaraEditor/Private/NiagaraActions.cpp b/Engine/Plugins/FX/Niagara/Source/NiagaraEditor/Private/NiagaraActions.cpp
--- a/Engine/Plugins/FX/Niagara/Source/NiagaraEditor/Private/NiagaraActions.cpp
+++ b/Engine/Plugins/FX/Niagara/Source/NiagaraEditor/Private/NiagaraActions.cpp
@@ -110,13 +110,15 @@ FReply FNiagaraParameterGraphDragOperation::DroppedOnNode(FVector2D ScreenPositi
 		{
 			FScopedTransaction AddNewPinTransaction(LOCTEXT("Drop Onto Get Pin", "Drop parameter onto Get node"));
 			GetMapNode->Modify();
-			GetMapNode->RequestNewTypedPin(EGPD_Output, Parameter.GetType(), Parameter.GetName());
+			UEdGraphPin* Pin = GetMapNode->RequestNewTypedPin(EGPD_Output, Parameter.GetType(), Parameter.GetName());
+			GetMapNode->CancelEditablePinName(FText::GetEmpty(), Pin);
 		} 
 		else if (UNiagaraNodeParameterMapSet* SetMapNode = Cast<UNiagaraNodeParameterMapSet>(GetHoveredNode()))
 		{
 			FScopedTransaction AddNewPinTransaction(LOCTEXT("Drop Onto Set Pin", "Drop parameter onto Set node"));
 			SetMapNode->Modify();
-			SetMapNode->RequestNewTypedPin(EGPD_Input, Parameter.GetType(), Parameter.GetName());
+			UEdGraphPin* Pin = SetMapNode->RequestNewTypedPin(EGPD_Input, Parameter.GetType(), Parameter.GetName());
+			SetMapNode->CancelEditablePinName(FText::GetEmpty(), Pin);
 		}
 	}
 
diff --git a/Engine/Plugins/FX/Niagara/Source/NiagaraEditor/Private/NiagaraEditorUtilities.cpp b/Engine/Plugins/FX/Niagara/Source/NiagaraEditor/Private/NiagaraEditorUtilities.cpp
--- a/Engine/Plugins/FX/Niagara/Source/NiagaraEditor/Private/NiagaraEditorUtilities.cpp
+++ b/Engine/Plugins/FX/Niagara/Source/NiagaraEditor/Private/NiagaraEditorUtilities.cpp
@@ -2052,6 +2052,17 @@ void FNiagaraEditorUtilities::WarnWithToastAndLog(FText WarningMessage)
 	UE_LOG(LogNiagaraEditor, Warning, TEXT("%s"), *WarningMessage.ToString());
 }
 
+void FNiagaraEditorUtilities::InfoWithToastAndLog(FText InfoMessage)
+{
+	FNotificationInfo WarningNotification(InfoMessage);
+	WarningNotification.ExpireDuration = 5.0f;
+	WarningNotification.bFireAndForget = true;
+	WarningNotification.bUseLargeFont = false;
+	WarningNotification.Image = FCoreStyle::Get().GetBrush(TEXT("MessageLog.Note"));
+	FSlateNotificationManager::Get().AddNotification(WarningNotification);
+	UE_LOG(LogNiagaraEditor, Warning, TEXT("%s"), *InfoMessage.ToString());
+}
+
 FName FNiagaraEditorUtilities::GetUniqueObjectName(UObject* Outer, UClass* ObjectClass, const FString& CandidateName)
 {
 	if (StaticFindObject(ObjectClass, Outer, *CandidateName) == nullptr)
diff --git a/Engine/Plugins/FX/Niagara/Source/NiagaraEditor/Private/NiagaraGraph.cpp b/Engine/Plugins/FX/Niagara/Source/NiagaraEditor/Private/NiagaraGraph.cpp
--- a/Engine/Plugins/FX/Niagara/Source/NiagaraEditor/Private/NiagaraGraph.cpp
+++ b/Engine/Plugins/FX/Niagara/Source/NiagaraEditor/Private/NiagaraGraph.cpp
@@ -1126,6 +1126,10 @@ UNiagaraScriptVariable* UNiagaraGraph::AddParameter(FNiagaraVariable& Parameter,
 		NewScriptVariable->Metadata.SetWasCreatedInSystemEditor(Options.bAddedFromSystemEditor);
 		VariableToScriptVariable.Add(Parameter, NewScriptVariable);
 
+		FName NewParamName;
+		FNiagaraEditorUtilities::DecomposeVariableNamespace(Parameter.GetName(), NewParamName);
+		NewScriptVariable->Metadata.SetCachedNamespacelessVariableName(NewParamName);
+
 		bool bSkipRefreshMetaDataScopeAndUsage = false;
 		if (Options.NewParameterUsage.IsSet())
 		{
@@ -1214,8 +1218,129 @@ void UNiagaraGraph::RemoveParameter(const FNiagaraVariable& Parameter, bool bAll
 	}
 }
 
-bool UNiagaraGraph::RenameParameter(const FNiagaraVariable& Parameter, FName NewName, bool bRenameRequestedFromStaticSwitch, FName NewScopeName)
+bool UNiagaraGraph::RenameParameterFromPin(const FNiagaraVariable& Parameter, FName NewName, UEdGraphPin* InPin)
 {
+	if (Parameter.GetName() == NewName)
+		return true;
+
+	Modify();
+	if (FNiagaraGraphParameterReferenceCollection* ReferenceCollection = ParameterToReferencesMap.Find(Parameter))
+	{
+		FNiagaraGraphParameterReferenceCollection NewReferences = *ReferenceCollection;
+		if (NewReferences.ParameterReferences.Num() == 1 && NewReferences.ParameterReferences[0].Key == InPin->PersistentGuid)
+		{
+			bool bRenameRequestedFromStaticSwitch = false;
+			FName NewScopeName = FName();
+			bool bMerged = false;
+
+			if (RenameParameter(Parameter, NewName, bRenameRequestedFromStaticSwitch, NewScopeName, &bMerged))
+			{
+				if (!bMerged)
+				{
+					FName ParamName = NewName;
+					FNiagaraEditorUtilities::DecomposeVariableNamespace(NewName, ParamName);
+					FNiagaraEditorUtilities::InfoWithToastAndLog(FText::Format(LOCTEXT("RenamedVarInGraphForAll", "\"{0}\" has been fully renamed as it was only used on this node."), FText::FromName(ParamName)));
+				}
+				return true;
+			}
+			return false;
+		}
+	}
+
+	// Create the new parameter
+	FNiagaraVariable NewParameter = Parameter;
+	NewParameter.SetName(NewName);
+
+	FNiagaraVariableMetaData MetaData;
+	UNiagaraScriptVariable** OldScriptVariable = VariableToScriptVariable.Find(Parameter);
+	UNiagaraScriptVariable** NewScriptVariableFound = VariableToScriptVariable.Find(NewParameter);
+
+	if (OldScriptVariable != nullptr)
+	{
+		MetaData = (*OldScriptVariable)->Metadata;
+
+
+		//FName CachedName;
+		//MetaData.GetParameterName(CachedName);
+		//UE_LOG(LogNiagaraEditor, Log, TEXT("RenameParameterFromPin %s %s!"), *Parameter.GetName().ToString(), *CachedName.ToString());
+	}
+
+	// Update the OldMetaData which will be applied to the new parameter to synchronize the cached namespace-less name and scope.
+	bool bMerged = false;
+	MetaData.SetCachedNamespacelessVariableName(FName(*FNiagaraEditorUtilities::GetNamespacelessVariableNameString(NewParameter.GetName())));
+	FAddParameterOptions Options;
+	AddParameter(NewParameter, Options);
+
+	if (bIsRenamingParameter)
+	{
+		return false;
+	}
+
+	// Swap metadata to the new parameter; put the new parameter into VariableToScriptVariable
+	if (OldScriptVariable)
+	{
+		// Rename all the bindings that point to the old parameter 
+		for (auto It : VariableToScriptVariable)
+		{
+			UNiagaraScriptVariable* Variable = It.Value;
+			if (Variable && Variable->DefaultBinding.GetName() == Parameter.GetName())
+			{
+				Variable->DefaultBinding.SetName(NewParameter.GetName());
+			}
+		}
+
+		// Only create a new variable if needed.
+		if (!NewScriptVariableFound)
+		{
+			// Replace the script variable data
+			UNiagaraScriptVariable* NewScriptVariable = NewObject<UNiagaraScriptVariable>(this, FName(), RF_Transactional);
+			NewScriptVariable->Variable = NewParameter;
+			NewScriptVariable->DefaultMode = (*OldScriptVariable)->DefaultMode;
+			NewScriptVariable->DefaultBinding = (*OldScriptVariable)->DefaultBinding;
+			VariableToScriptVariable.Add(NewParameter, NewScriptVariable);
+		}
+		VariableToScriptVariable.Remove(Parameter);
+	}
+	// Either set the new meta-data or use the existing meta-data.
+	if (!NewScriptVariableFound)
+	{
+		SetPerScriptMetaData(NewParameter, MetaData);
+	}
+	else
+	{
+		bMerged = true;
+		FName NewParamName = NewName;
+		FNiagaraEditorUtilities::DecomposeVariableNamespace(NewName, NewParamName);
+		FName OldParamName = Parameter.GetName();
+		FNiagaraEditorUtilities::DecomposeVariableNamespace(Parameter.GetName(), OldParamName);
+		FNiagaraEditorUtilities::InfoWithToastAndLog(FText::Format(LOCTEXT("MergedVar", "\"{0}\" has been merged with parameter \"{1}\".\nAll of \"{1}\"'s meta-data will be used, overwriting \"{0}\"'s meta-data."), FText::FromName(OldParamName), FText::FromName(NewParamName)));
+	}
+	RefreshParameterReferences();
+
+	if (!bMerged)
+	{
+		FName NewParamName = NewName;
+		FNiagaraEditorUtilities::DecomposeVariableNamespace(NewName, NewParamName);
+		FName OldParamName = Parameter.GetName();
+		FNiagaraEditorUtilities::DecomposeVariableNamespace(Parameter.GetName(), OldParamName);
+		FNiagaraEditorUtilities::InfoWithToastAndLog(FText::Format(LOCTEXT("RenamedVarInGraphForNode", "\"{0}\" has been duplicated as \"{1}\" as it is used in multiple locations.\nPlease edit the Parameters Panel version to change in all locations."), FText::FromName(OldParamName), FText::FromName(NewParamName)));
+	}
+
+	return false;
+}
+
+
+bool UNiagaraGraph::RenameParameter(const FNiagaraVariable& Parameter, FName NewName, bool bRenameRequestedFromStaticSwitch, FName NewScopeName, bool* bMerged)
+{
+	// Initialize the merger state if requested
+	if (bMerged)
+		*bMerged = false;
+
+	if (Parameter.GetName() == NewName)
+		return true;
+
+	
+
 	// Block rename when already renaming. This prevents recursion when CommitEditablePinName is called on referenced nodes. 
 	if (bIsRenamingParameter)
 	{
@@ -1238,6 +1363,11 @@ bool UNiagaraGraph::RenameParameter(const FNiagaraVariable& Parameter, FName New
 			return false;
 		}
 		OldMetaData = (*OldScriptVariable)->Metadata;
+
+
+		//FName CachedName;
+		//OldMetaData.GetParameterName(CachedName);
+		//UE_LOG(LogNiagaraEditor, Log, TEXT("RenameParameter %s %s!"), *Parameter.GetName().ToString(), *CachedName.ToString());
 	}
 
 	// Update the OldMetaData which will be applied to the new parameter to synchronize the cached namespace-less name and scope.
@@ -1247,6 +1377,8 @@ bool UNiagaraGraph::RenameParameter(const FNiagaraVariable& Parameter, FName New
 		OldMetaData.SetScopeName(NewScopeName);
 	}
 		
+	UNiagaraScriptVariable** NewScriptVariableFound = VariableToScriptVariable.Find(NewParameter);
+
 	// Swap metadata to the new parameter; put the new parameter into VariableToScriptVariable
 	if (OldScriptVariable)
 	{
@@ -1261,15 +1393,34 @@ bool UNiagaraGraph::RenameParameter(const FNiagaraVariable& Parameter, FName New
 			}
 		}
 
-		// Replace the script variable data
-		UNiagaraScriptVariable* NewScriptVariable = NewObject<UNiagaraScriptVariable>(this, FName(), RF_Transactional);
-		NewScriptVariable->Variable = NewParameter;
-		NewScriptVariable->DefaultMode = (*OldScriptVariable)->DefaultMode;
-		NewScriptVariable->DefaultBinding = (*OldScriptVariable)->DefaultBinding;
+		// Only create a new variable if needed.
+		if (!NewScriptVariableFound)
+		{
+			// Replace the script variable data
+			UNiagaraScriptVariable* NewScriptVariable = NewObject<UNiagaraScriptVariable>(this, FName(), RF_Transactional);
+			NewScriptVariable->Variable = NewParameter;
+			NewScriptVariable->DefaultMode = (*OldScriptVariable)->DefaultMode;
+			NewScriptVariable->DefaultBinding = (*OldScriptVariable)->DefaultBinding;
+			VariableToScriptVariable.Add(NewParameter, NewScriptVariable);
+		}
 		VariableToScriptVariable.Remove(Parameter);
-		VariableToScriptVariable.Add(NewParameter, NewScriptVariable);
 	}
-	SetPerScriptMetaData(NewParameter, OldMetaData);
+
+	// Either set the new meta-data or use the existing meta-data.
+	if (!NewScriptVariableFound)
+	{
+		SetPerScriptMetaData(NewParameter, OldMetaData);
+	}
+	else
+	{
+		if (bMerged)
+			*bMerged = true;
+		FName NewParamName = NewName;
+		FNiagaraEditorUtilities::DecomposeVariableNamespace(NewName, NewParamName);
+		FName OldParamName = Parameter.GetName();
+		FNiagaraEditorUtilities::DecomposeVariableNamespace(Parameter.GetName(), OldParamName);
+		FNiagaraEditorUtilities::InfoWithToastAndLog(FText::Format(LOCTEXT("MergedVar", "\"{0}\" has been merged with parameter \"{1}\".\nAll of \"{1}\"'s meta-data will be used, overwriting \"{0}\"'s meta-data."), FText::FromName(OldParamName), FText::FromName(NewParamName)));
+	}
 
 	// Fixup reference collection and pin names
 	if (FNiagaraGraphParameterReferenceCollection* ReferenceCollection = ParameterToReferencesMap.Find(Parameter))
@@ -1285,13 +1436,21 @@ bool UNiagaraGraph::RenameParameter(const FNiagaraVariable& Parameter, FName New
 				UEdGraphPin* Pin = Node->GetPinByPersistentGuid(Reference.Key);
 				if (Pin)
 				{
-					Node->CommitEditablePinName(NewNameText, Pin);
+					Pin->Modify();
+					Node->CommitEditablePinName(NewNameText, Pin, false);
 				}
 			}
 		}
 
 		ParameterToReferencesMap.Remove(Parameter);
-		ParameterToReferencesMap.Add(NewParameter, NewReferences);
+		if (!bMerged)
+		{
+			ParameterToReferencesMap.Add(NewParameter, NewReferences);
+		}
+		else
+		{
+			RefreshParameterReferences();
+		}
 	}
 
 	bIsRenamingParameter = false;
@@ -1977,6 +2136,11 @@ void UNiagaraGraph::SetMetaData(const FNiagaraVariable& InVar, const FNiagaraVar
 {
 	ensure(FNiagaraConstants::IsNiagaraConstant(InVar) == false);
 
+
+	//FName CachedName;
+	//InMetaData.GetParameterName(CachedName);
+	//UE_LOG(LogNiagaraEditor, Log, TEXT("SetMetaData %s %s!"), *InVar.GetName().ToString(), *CachedName.ToString());
+
 	if (UNiagaraScriptVariable** FoundMetaData = VariableToScriptVariable.Find(InVar))
 	{
 		if (*FoundMetaData)
@@ -2006,6 +2170,8 @@ void UNiagaraGraph::SetPerScriptMetaData(const FNiagaraVariable& InVar, const FN
 			FNiagaraVariableMetaData& OldMetaData = (*FoundMetaData)->Metadata;
 			OldMetaData.CopyPerScriptMetaData(InMetaData);
 			OldMetaData.SetIsStaticSwitch(InMetaData.GetIsStaticSwitch());
+
+
 		}
 	}
 	else
@@ -2015,6 +2181,10 @@ void UNiagaraGraph::SetPerScriptMetaData(const FNiagaraVariable& InVar, const FN
 		NewScriptVariable->Variable = InVar;
 		NewScriptVariable->Metadata = InMetaData;
 	}
+
+	//FName CachedName;
+	//InMetaData.GetParameterName(CachedName);
+	//UE_LOG(LogNiagaraEditor, Log, TEXT("SetPerScriptMetaData %s %s!"), *InVar.GetName().ToString(), *CachedName.ToString());
 }
 
 UNiagaraGraph::FOnDataInterfaceChanged& UNiagaraGraph::OnDataInterfaceChanged()
diff --git a/Engine/Plugins/FX/Niagara/Source/NiagaraEditor/Private/NiagaraNodeCustomHlsl.cpp b/Engine/Plugins/FX/Niagara/Source/NiagaraEditor/Private/NiagaraNodeCustomHlsl.cpp
--- a/Engine/Plugins/FX/Niagara/Source/NiagaraEditor/Private/NiagaraNodeCustomHlsl.cpp
+++ b/Engine/Plugins/FX/Niagara/Source/NiagaraEditor/Private/NiagaraNodeCustomHlsl.cpp
@@ -234,7 +234,7 @@ bool UNiagaraNodeCustomHlsl::VerifyEditablePinName(const FText& InName, FText& O
 	return true;
 }
 
-bool UNiagaraNodeCustomHlsl::CommitEditablePinName(const FText& InName, UEdGraphPin* InGraphPinObj)
+bool UNiagaraNodeCustomHlsl::CommitEditablePinName(const FText& InName, UEdGraphPin* InGraphPinObj, bool bSuppressEvents)
 {
 	if (Pins.Contains(InGraphPinObj))
 	{
@@ -244,7 +244,8 @@ bool UNiagaraNodeCustomHlsl::CommitEditablePinName(const FText& InName, UEdGraph
 
 		FString OldPinName = InGraphPinObj->PinName.ToString();
 		InGraphPinObj->PinName = *InName.ToString();
-		OnPinRenamed(InGraphPinObj, OldPinName);
+		if (bSuppressEvents == false)
+			OnPinRenamed(InGraphPinObj, OldPinName);
 
 		return true;
 	}
diff --git a/Engine/Plugins/FX/Niagara/Source/NiagaraEditor/Private/NiagaraNodeEmitter.cpp b/Engine/Plugins/FX/Niagara/Source/NiagaraEditor/Private/NiagaraNodeEmitter.cpp
--- a/Engine/Plugins/FX/Niagara/Source/NiagaraEditor/Private/NiagaraNodeEmitter.cpp
+++ b/Engine/Plugins/FX/Niagara/Source/NiagaraEditor/Private/NiagaraNodeEmitter.cpp
@@ -82,7 +82,7 @@ bool UNiagaraNodeEmitter::VerifyEditablePinName(const FText& InName, FText& OutE
 	return true;
 }
 
-bool UNiagaraNodeEmitter::CommitEditablePinName(const FText& InName, UEdGraphPin* InGraphPinObj)
+bool UNiagaraNodeEmitter::CommitEditablePinName(const FText& InName, UEdGraphPin* InGraphPinObj, bool bSuppressEvents)
 {
 	return false;
 }
diff --git a/Engine/Plugins/FX/Niagara/Source/NiagaraEditor/Private/NiagaraNodeEmitter.h b/Engine/Plugins/FX/Niagara/Source/NiagaraEditor/Private/NiagaraNodeEmitter.h
--- a/Engine/Plugins/FX/Niagara/Source/NiagaraEditor/Private/NiagaraNodeEmitter.h
+++ b/Engine/Plugins/FX/Niagara/Source/NiagaraEditor/Private/NiagaraNodeEmitter.h
@@ -43,7 +43,7 @@ public:
 	virtual bool IsPinNameEditable(const UEdGraphPin* GraphPinObj) const override;
 	virtual bool IsPinNameEditableUponCreation(const UEdGraphPin* GraphPinObj) const override;
 	virtual bool VerifyEditablePinName(const FText& InName, FText& OutErrorMessage, const UEdGraphPin* InGraphPinObj) const override;
-	virtual bool CommitEditablePinName(const FText& InName, UEdGraphPin* InGraphPinObj)  override;
+	virtual bool CommitEditablePinName(const FText& InName, UEdGraphPin* InGraphPinObj, bool bSuppressEvents = false)  override;
 
 	virtual void BuildParameterMapHistory(FNiagaraParameterMapHistoryBuilder& OutHistory, bool bRecursive = true, bool bFilterForCompilation = true) const override;
 
diff --git a/Engine/Plugins/FX/Niagara/Source/NiagaraEditor/Private/NiagaraNodeParameterMapBase.cpp b/Engine/Plugins/FX/Niagara/Source/NiagaraEditor/Private/NiagaraNodeParameterMapBase.cpp
--- a/Engine/Plugins/FX/Niagara/Source/NiagaraEditor/Private/NiagaraNodeParameterMapBase.cpp
+++ b/Engine/Plugins/FX/Niagara/Source/NiagaraEditor/Private/NiagaraNodeParameterMapBase.cpp
@@ -177,8 +177,16 @@ void UNiagaraNodeParameterMapBase::GetPinHoverText(const UEdGraphPin& Pin, FStri
 
 			FNiagaraVariable Var = FNiagaraVariable(TypeDef, Pin.PinName);
 			TOptional<FNiagaraVariableMetaData> Metadata = NiagaraGraph->GetMetaData(Var);
+
 			FText Description = Metadata.IsSet() ? Metadata->Description : FText::GetEmpty();
-			FText ToolTipText = FNiagaraEditorUtilities::FormatVariableDescription(Description, FText::FromName(Pin.PinName), TypeDef.GetNameText());
+			const FText TooltipFormat = LOCTEXT("Parameters", "Name: {0} \nType: {1}\nDescription: {2}\nScope: {3}\nUser Editable: {4}\nUsage: {5}");
+			const FText Name = FText::FromName(Var.GetName());
+			FName CachedParamName;
+			Metadata->GetParameterName(CachedParamName);
+			const FText ScopeText = FText::FromName(Metadata->GetScopeName());
+			const FText UserEditableText = FText::FromName(CachedParamName);
+			const FText UsageText = StaticEnum<ENiagaraScriptParameterUsage>()->GetDisplayNameTextByValue((int64)Metadata->GetUsage());
+			const FText ToolTipText = FText::Format(TooltipFormat, FText::FromName(Var.GetName()), Var.GetType().GetNameText(), Description, ScopeText, UserEditableText, UsageText);
 			HoverTextOut = ToolTipText.ToString();
 		}
 	}
@@ -224,7 +232,7 @@ void UNiagaraNodeParameterMapBase::OnPinRenamed(UEdGraphPin* RenamedPin, const F
 	FNiagaraVariable Var(VarType, *OldName);
 
 	UNiagaraGraph* Graph = GetNiagaraGraph();
-	Graph->RenameParameter(Var, NewUniqueName);
+	Graph->RenameParameterFromPin(Var, NewUniqueName, RenamedPin);
 
 	if (RenamedPin == PinPendingRename)
 	{
diff --git a/Engine/Plugins/FX/Niagara/Source/NiagaraEditor/Private/NiagaraNodeParameterMapBase.h b/Engine/Plugins/FX/Niagara/Source/NiagaraEditor/Private/NiagaraNodeParameterMapBase.h
--- a/Engine/Plugins/FX/Niagara/Source/NiagaraEditor/Private/NiagaraNodeParameterMapBase.h
+++ b/Engine/Plugins/FX/Niagara/Source/NiagaraEditor/Private/NiagaraNodeParameterMapBase.h
@@ -41,6 +41,9 @@ public:
 	void SetPinName(UEdGraphPin* InPin, const FName& InName);
 
 	bool OnAllowDrop(TSharedPtr<FDragDropOperation> DragDropOperation);
+
+	virtual bool CanRenamePinFromContextMenu(const UEdGraphPin* Pin) const { return false; }
+
 protected:
 	virtual void SelectParameterFromPin(UEdGraphPin* InPin);
 	virtual void OnPinRenamed(UEdGraphPin* RenamedPin, const FString& OldName) override;
diff --git a/Engine/Plugins/FX/Niagara/Source/NiagaraEditor/Private/NiagaraNodeParameterMapGet.cpp b/Engine/Plugins/FX/Niagara/Source/NiagaraEditor/Private/NiagaraNodeParameterMapGet.cpp
--- a/Engine/Plugins/FX/Niagara/Source/NiagaraEditor/Private/NiagaraNodeParameterMapGet.cpp
+++ b/Engine/Plugins/FX/Niagara/Source/NiagaraEditor/Private/NiagaraNodeParameterMapGet.cpp
@@ -19,6 +19,7 @@
 #include "EdGraph/EdGraphNode.h"
 #include "NiagaraParameterCollection.h"
 #include "NiagaraScriptVariable.h"
+#include "NiagaraConstants.h"
 
 #define LOCTEXT_NAMESPACE "NiagaraNodeParameterMapGet"
 
@@ -60,7 +61,12 @@ bool UNiagaraNodeParameterMapGet::IsPinNameEditableUponCreation(const UEdGraphPi
 {
 	if (GraphPinObj == PinPendingRename && GraphPinObj->Direction == EEdGraphPinDirection::EGPD_Output)
 	{
-		return true;
+		const UEdGraphSchema_Niagara* Schema = GetDefault<UEdGraphSchema_Niagara>();
+		FNiagaraVariable Var = Schema->PinToNiagaraVariable(GraphPinObj);
+		if (!FNiagaraConstants::IsNiagaraConstant(Var))
+			return true;
+		else
+			return false;
 	}
 	else
 	{
@@ -430,18 +436,27 @@ bool UNiagaraNodeParameterMapGet::CancelEditablePinName(const FText& InName, UEd
 	return true;
 }
 
-bool UNiagaraNodeParameterMapGet::CommitEditablePinName(const FText& InName, UEdGraphPin* InGraphPinObj) 
+bool UNiagaraNodeParameterMapGet::CommitEditablePinName(const FText& InName, UEdGraphPin* InGraphPinObj, bool bSuppressEvents)
 {
 	if (Pins.Contains(InGraphPinObj) && InGraphPinObj->Direction == EEdGraphPinDirection::EGPD_Output)
 	{
+		FString OldPinName = InGraphPinObj->PinName.ToString();
+		FString NewPinName = InName.ToString();
+
+		// Early out if the same!
+		if (OldPinName == NewPinName)
+		{
+			return true;
+		}
+
 		FScopedTransaction AddNewPinTransaction(LOCTEXT("Rename Pin", "Renamed pin"));
 		Modify();
 		InGraphPinObj->Modify();
 
-		FString OldPinName = InGraphPinObj->PinName.ToString();
 		InGraphPinObj->PinName = *InName.ToString();
 
-		OnPinRenamed(InGraphPinObj, OldPinName);
+		if (!bSuppressEvents)
+			OnPinRenamed(InGraphPinObj, OldPinName);
 
 		return true;
 	}
@@ -511,11 +526,14 @@ void UNiagaraNodeParameterMapGet::GetPinHoverText(const UEdGraphPin& Pin, FStrin
 				TOptional<FNiagaraVariableMetaData> Metadata = NiagaraGraph->GetMetaData(Var);
 
 				FText Description = Metadata.IsSet() ? Metadata->Description : FText::GetEmpty();
-				const FText TooltipFormat = LOCTEXT("Parameters", "Name: {0} \nType: {1}\nDescription: {2}\nScope: {3}\nUsage: {4}"); 
+				const FText TooltipFormat = LOCTEXT("Parameters", "Name: {0} \nType: {1}\nDescription: {2}\nScope: {3}\nUser Editable: {4}\nUsage: {5}"); 
 				const FText Name = FText::FromName(Var.GetName());
+				FName CachedParamName;
+				Metadata->GetParameterName(CachedParamName);
 				const FText ScopeText = FText::FromName(Metadata->GetScopeName());
+				const FText UserEditableText = FText::FromName(CachedParamName);
 				const FText UsageText = StaticEnum<ENiagaraScriptParameterUsage>()->GetDisplayNameTextByValue((int64)Metadata->GetUsage());
-				const FText ToolTipText = FText::Format(TooltipFormat, FText::FromName(Var.GetName()), Var.GetType().GetNameText(), Description, ScopeText, UsageText);
+				const FText ToolTipText = FText::Format(TooltipFormat, FText::FromName(Var.GetName()), Var.GetType().GetNameText(), Description, ScopeText, UserEditableText, UsageText);
 				HoverTextOut = ToolTipText.ToString();
 			}
 		}
diff --git a/Engine/Plugins/FX/Niagara/Source/NiagaraEditor/Private/NiagaraNodeParameterMapGet.h b/Engine/Plugins/FX/Niagara/Source/NiagaraEditor/Private/NiagaraNodeParameterMapGet.h
--- a/Engine/Plugins/FX/Niagara/Source/NiagaraEditor/Private/NiagaraNodeParameterMapGet.h
+++ b/Engine/Plugins/FX/Niagara/Source/NiagaraEditor/Private/NiagaraNodeParameterMapGet.h
@@ -26,7 +26,7 @@ public:
 	virtual bool IsPinNameEditable(const UEdGraphPin* GraphPinObj) const override;
 	virtual bool IsPinNameEditableUponCreation(const UEdGraphPin* GraphPinObj) const override;
 	virtual bool VerifyEditablePinName(const FText& InName, FText& OutErrorMessage, const UEdGraphPin* InGraphPinObj) const override;
-	virtual bool CommitEditablePinName(const FText& InName, UEdGraphPin* InGraphPinObj)  override;
+	virtual bool CommitEditablePinName(const FText& InName, UEdGraphPin* InGraphPinObj, bool bSuppressEvents)  override;
 	virtual bool CancelEditablePinName(const FText& InName, UEdGraphPin* InGraphPinObj) override;
 
 	virtual void Compile(class FHlslNiagaraTranslator* Translator, TArray<int32>& Outputs);
diff --git a/Engine/Plugins/FX/Niagara/Source/NiagaraEditor/Private/NiagaraNodeParameterMapSet.cpp b/Engine/Plugins/FX/Niagara/Source/NiagaraEditor/Private/NiagaraNodeParameterMapSet.cpp
--- a/Engine/Plugins/FX/Niagara/Source/NiagaraEditor/Private/NiagaraNodeParameterMapSet.cpp
+++ b/Engine/Plugins/FX/Niagara/Source/NiagaraEditor/Private/NiagaraNodeParameterMapSet.cpp
@@ -13,6 +13,7 @@
 #include "Widgets/Input/SEditableTextBox.h"
 #include "EdGraph/EdGraphNode.h"
 #include "NiagaraScriptVariable.h"
+#include "NiagaraConstants.h"
 
 #include "ScopedTransaction.h"
 #include "SNiagaraGraphParameterMapSetNode.h"
@@ -56,6 +57,7 @@ bool UNiagaraNodeParameterMapSet::IsPinNameEditableUponCreation(const UEdGraphPi
 {
 	if (GraphPinObj == PinPendingRename)
 	{
+
 		return true;
 	}
 	else
@@ -141,7 +143,13 @@ void UNiagaraNodeParameterMapSet::OnNewTypedPinAdded(UEdGraphPin* NewPin)
 		NewPin->PersistentGuid = FGuid::NewGuid();
 	}
 
-	PinPendingRename = NewPin;
+	const UEdGraphSchema_Niagara* Schema = GetDefault<UEdGraphSchema_Niagara>();
+	FNiagaraVariable Var = Schema->PinToNiagaraVariable(NewPin);
+	if (!FNiagaraConstants::IsNiagaraConstant(Var))
+		PinPendingRename = NewPin;
+	else
+		PinPendingRename = nullptr;
+	
 }
 
 void UNiagaraNodeParameterMapSet::OnPinRenamed(UEdGraphPin* RenamedPin, const FString& OldName)
@@ -159,17 +167,27 @@ bool UNiagaraNodeParameterMapSet::CancelEditablePinName(const FText& InName, UEd
 	return true;
 }
 
-bool UNiagaraNodeParameterMapSet::CommitEditablePinName(const FText& InName, UEdGraphPin* InGraphPinObj) 
+bool UNiagaraNodeParameterMapSet::CommitEditablePinName(const FText& InName, UEdGraphPin* InGraphPinObj, bool bSuppressEvents)
 {
 	if (Pins.Contains(InGraphPinObj))
 	{
+		
+		FString OldPinName = InGraphPinObj->PinName.ToString();
+		FString NewPinName = InName.ToString();
+
+		// Early out if the same!
+		if (OldPinName == NewPinName)
+		{
+			return true;
+		}
+
 		FScopedTransaction AddNewPinTransaction(LOCTEXT("Rename Pin", "Renamed pin"));
 		Modify();
 		InGraphPinObj->Modify();
-		
-		FString OldPinName = InGraphPinObj->PinName.ToString();
-		InGraphPinObj->PinName = *InName.ToString();
-		OnPinRenamed(InGraphPinObj, OldPinName);
+
+		InGraphPinObj->PinName = *NewPinName;
+		if (bSuppressEvents == false)	
+			OnPinRenamed(InGraphPinObj, OldPinName);
 
 		return true;
 	}
diff --git a/Engine/Plugins/FX/Niagara/Source/NiagaraEditor/Private/NiagaraNodeParameterMapSet.h b/Engine/Plugins/FX/Niagara/Source/NiagaraEditor/Private/NiagaraNodeParameterMapSet.h
--- a/Engine/Plugins/FX/Niagara/Source/NiagaraEditor/Private/NiagaraNodeParameterMapSet.h
+++ b/Engine/Plugins/FX/Niagara/Source/NiagaraEditor/Private/NiagaraNodeParameterMapSet.h
@@ -25,7 +25,7 @@ public:
 	virtual bool IsPinNameEditableUponCreation(const UEdGraphPin* GraphPinObj) const override;
 	virtual void RemoveDynamicPin(UEdGraphPin* Pin) override;
 	virtual bool VerifyEditablePinName(const FText& InName, FText& OutErrorMessage, const UEdGraphPin* InGraphPinObj) const override;
-	virtual bool CommitEditablePinName(const FText& InName, UEdGraphPin* InGraphPinObj)  override;
+	virtual bool CommitEditablePinName(const FText& InName, UEdGraphPin* InGraphPinObj, bool bSuppressEvents = false)  override;
 	virtual bool CancelEditablePinName(const FText& InName, UEdGraphPin* InGraphPinObj) override;
 
 	virtual void BuildParameterMapHistory(FNiagaraParameterMapHistoryBuilder& OutHistory, bool bRecursive = true, bool bFilterForCompilation = true) const override;
diff --git a/Engine/Plugins/FX/Niagara/Source/NiagaraEditor/Private/NiagaraNodeWithDynamicPins.cpp b/Engine/Plugins/FX/Niagara/Source/NiagaraEditor/Private/NiagaraNodeWithDynamicPins.cpp
--- a/Engine/Plugins/FX/Niagara/Source/NiagaraEditor/Private/NiagaraNodeWithDynamicPins.cpp
+++ b/Engine/Plugins/FX/Niagara/Source/NiagaraEditor/Private/NiagaraNodeWithDynamicPins.cpp
@@ -190,7 +190,7 @@ void UNiagaraNodeWithDynamicPins::GetNodeContextMenuActions(UToolMenu* Menu, UGr
 	if (Context->Pin != nullptr)
 	{
 		FToolMenuSection& Section = Menu->AddSection("EdGraphSchema_NiagaraPinActions", LOCTEXT("EditPinMenuHeader", "Edit Pin"));
-		if (CanRenamePin(Context->Pin))
+		if (CanRenamePinFromContextMenu(Context->Pin))
 		{
 			UEdGraphPin* Pin = const_cast<UEdGraphPin*>(Context->Pin);
 			TSharedRef<SWidget> RenameWidget =
diff --git a/Engine/Plugins/FX/Niagara/Source/NiagaraEditor/Private/TNiagaraGraphPinEditableName.h b/Engine/Plugins/FX/Niagara/Source/NiagaraEditor/Private/TNiagaraGraphPinEditableName.h
--- a/Engine/Plugins/FX/Niagara/Source/NiagaraEditor/Private/TNiagaraGraphPinEditableName.h
+++ b/Engine/Plugins/FX/Niagara/Source/NiagaraEditor/Private/TNiagaraGraphPinEditableName.h
@@ -83,12 +83,13 @@ protected:
 		UNiagaraNode* ParentNode = Cast<UNiagaraNode>(this->GraphPinObj->GetOwningNode());
 
 		auto CreateLabelTextBlock = [&]()->TSharedRef<SWidget> {
-			return SNew(SInlineEditableTextBlock)
+			CreatedTextBlock = SNew(SInlineEditableTextBlock)
 				.Style(&FEditorStyle::Get().GetWidgetStyle<FInlineEditableTextBlockStyle>("Graph.Node.InlineEditablePinName"))
 				.Text(this, &TNiagaraGraphPinEditableName<BaseClass>::GetParentPinLabel)
 				.Visibility(this, &TNiagaraGraphPinEditableName<BaseClass>::GetParentPinVisibility)
 				.ColorAndOpacity(this, &TNiagaraGraphPinEditableName<BaseClass>::GetParentPinTextColor)
 				.IsReadOnly(true);
+			return CreatedTextBlock.ToSharedRef();
 		};
 
 		auto CreateRenamableLabelTextBlock = [&]()->TSharedRef<SWidget> {
@@ -97,13 +98,14 @@ protected:
 				bPendingRename = true;
 			}
 			
-			return SNew(SInlineEditableTextBlock)
+			CreatedTextBlock = SNew(SInlineEditableTextBlock)
 				.Style(&FEditorStyle::Get().GetWidgetStyle<FInlineEditableTextBlockStyle>("Graph.Node.InlineEditablePinName"))
 				.Text(this, &TNiagaraGraphPinEditableName<BaseClass>::GetParentPinLabel)
 				.Visibility(this, &TNiagaraGraphPinEditableName<BaseClass>::GetParentPinVisibility)
 				.ColorAndOpacity(this, &TNiagaraGraphPinEditableName<BaseClass>::GetParentPinTextColor)
 				.OnVerifyTextChanged(this, &TNiagaraGraphPinEditableName<BaseClass>::OnVerifyTextChanged)
 				.OnTextCommitted(this, &TNiagaraGraphPinEditableName<BaseClass>::OnTextCommitted);
+			return CreatedTextBlock.ToSharedRef();
 		};
 
 		if (ParentNode && ParentNode->IsPinNameEditable(this->GraphPinObj))
diff --git a/Engine/Plugins/FX/Niagara/Source/NiagaraEditor/Private/ViewModels/NiagaraParameterNameViewModel.cpp b/Engine/Plugins/FX/Niagara/Source/NiagaraEditor/Private/ViewModels/NiagaraParameterNameViewModel.cpp
--- a/Engine/Plugins/FX/Niagara/Source/NiagaraEditor/Private/ViewModels/NiagaraParameterNameViewModel.cpp
+++ b/Engine/Plugins/FX/Niagara/Source/NiagaraEditor/Private/ViewModels/NiagaraParameterNameViewModel.cpp
@@ -62,9 +62,9 @@ TSharedRef<SWidget> FNiagaraParameterPanelEntryParameterNameViewModel::CreateSco
 	return ScopeComboBoxWidget.ToSharedRef();
 }
 
-TSharedRef<SWidget> FNiagaraParameterPanelEntryParameterNameViewModel::CreateTextSlotWidget() const
+TSharedRef<SInlineEditableTextBlock> FNiagaraParameterPanelEntryParameterNameViewModel::CreateTextSlotWidget() const
 {
-	TSharedPtr< SWidget > DisplayWidget;
+	TSharedPtr< SInlineEditableTextBlock > DisplayWidget;
 	const FName FontType = FName("Italic");
 	FSlateFontInfo NameFont = FCoreStyle::GetDefaultFontStyle(FontType, 10); //@todo(ng) get correct font
 
@@ -147,7 +147,7 @@ void FNiagaraParameterPanelEntryParameterNameViewModel::OnParameterRenamed(const
 ///////////////////////////////////////////////////////////////////////////////
 
 FNiagaraGraphPinParameterNameViewModel::FNiagaraGraphPinParameterNameViewModel(
-	const UEdGraphPin* InOwningPin
+	UEdGraphPin* InOwningPin
 	, const FNiagaraScriptVariableAndViewInfo& InScriptVarAndViewInfo
 	, const FNiagaraScriptToolkitParameterPanelViewModel* InParameterPanelViewModel
 )
@@ -186,9 +186,9 @@ TSharedRef<SWidget> FNiagaraGraphPinParameterNameViewModel::CreateScopeSlotWidge
 	return ScopeComboBoxWidget.ToSharedRef();
 }
 
-TSharedRef<SWidget> FNiagaraGraphPinParameterNameViewModel::CreateTextSlotWidget() const
+TSharedRef<SInlineEditableTextBlock> FNiagaraGraphPinParameterNameViewModel::CreateTextSlotWidget() const
 {
-	TSharedPtr< SWidget > DisplayWidget;
+	TSharedPtr< SInlineEditableTextBlock > DisplayWidget;
 	const FName FontType = FName("Italic");
 	FSlateFontInfo NameFont = FCoreStyle::GetDefaultFontStyle(FontType, 10); //@todo(ng) get correct style
 
@@ -254,5 +254,28 @@ bool FNiagaraGraphPinParameterNameViewModel::VerifyParameterNameChanged(const FT
 
 void FNiagaraGraphPinParameterNameViewModel::OnParameterRenamed(const FText& NewNameText, ETextCommit::Type SelectionType) const
 {
-	ParameterPanelViewModel->RenamePin(OwningPin, NewNameText);
+	UNiagaraNode* Node = Cast<UNiagaraNode>(OwningPin->GetOwningNode());
+	if (Node)
+	{
+		if (CachedScriptVarAndViewInfo.MetaData.GetIsUsingLegacyNameString())
+		{
+			Node->CommitEditablePinName(NewNameText, OwningPin);
+		}
+		else
+		{
+			FName OldFullName = CachedScriptVarAndViewInfo.ScriptVariable.GetName();
+			FName OldParameterName;
+			CachedScriptVarAndViewInfo.MetaData.GetParameterName(OldParameterName);
+
+			FString OldParamNameStr = OldParameterName.ToString();
+			FString NewFullNameStr = OldFullName.ToString();
+			NewFullNameStr.RemoveFromEnd(OldParamNameStr);
+			NewFullNameStr += NewNameText.ToString();
+
+			FText NewFullNameText = FText::FromString(NewFullNameStr);
+
+			Node->CommitEditablePinName(NewFullNameText, OwningPin);
+		}
+	}
+	//ParameterPanelViewModel->RenamePin(OwningPin, NewNameText);
 }
diff --git a/Engine/Plugins/FX/Niagara/Source/NiagaraEditor/Private/ViewModels/NiagaraParameterNameViewModel.h b/Engine/Plugins/FX/Niagara/Source/NiagaraEditor/Private/ViewModels/NiagaraParameterNameViewModel.h
--- a/Engine/Plugins/FX/Niagara/Source/NiagaraEditor/Private/ViewModels/NiagaraParameterNameViewModel.h
+++ b/Engine/Plugins/FX/Niagara/Source/NiagaraEditor/Private/ViewModels/NiagaraParameterNameViewModel.h
@@ -20,7 +20,7 @@ public:
 	virtual ~INiagaraParameterNameViewModel() { }
 
 	virtual TSharedRef<SWidget> CreateScopeSlotWidget() const = 0;
-	virtual TSharedRef<SWidget> CreateTextSlotWidget() const = 0;
+	virtual TSharedRef<SInlineEditableTextBlock> CreateTextSlotWidget() const = 0;
 
 	virtual int32 GetScopeValue() const = 0;
 	virtual void OnScopeValueChanged(int32 NewScopeValue, ESelectInfo::Type SelectionType) const = 0;
@@ -43,7 +43,7 @@ public:
 
 	/** Begin INiagaraParameterNameViewModel Interface */
 	virtual TSharedRef<SWidget> CreateScopeSlotWidget() const override;
-	virtual TSharedRef<SWidget> CreateTextSlotWidget() const override;
+	virtual TSharedRef<SInlineEditableTextBlock> CreateTextSlotWidget() const override;
 
 	virtual int32 GetScopeValue() const override;
 	virtual void OnScopeValueChanged(int32 NewScopeValue, ESelectInfo::Type SelectionType) const override;
@@ -74,14 +74,14 @@ class FNiagaraGraphPinParameterNameViewModel : public INiagaraParameterNameViewM
 {
 public:
 	FNiagaraGraphPinParameterNameViewModel(
-		  const UEdGraphPin* InOwningPin
+		  UEdGraphPin* InOwningPin
 		, const FNiagaraScriptVariableAndViewInfo& InScriptVarAndViewInfo
 		, const FNiagaraScriptToolkitParameterPanelViewModel* InParameterPanelViewModel
 	);
 
 	/** Begin INiagaraParameterNameViewModel Interface */
 	virtual TSharedRef<SWidget> CreateScopeSlotWidget() const override;
-	virtual TSharedRef<SWidget> CreateTextSlotWidget() const override;
+	virtual TSharedRef<SInlineEditableTextBlock> CreateTextSlotWidget() const override;
 
 	virtual int32 GetScopeValue() const override;
 	virtual void OnScopeValueChanged(int32 NewScopeValue, ESelectInfo::Type SelectionType) const override;
@@ -94,7 +94,7 @@ public:
 	/** End INiagaraParameterNameViewModel Interface */
 
 private:
-	const UEdGraphPin* OwningPin;
+	UEdGraphPin* OwningPin;
 	const FNiagaraScriptVariableAndViewInfo CachedScriptVarAndViewInfo;
 	const FNiagaraScriptToolkitParameterPanelViewModel* ParameterPanelViewModel;
 
diff --git a/Engine/Plugins/FX/Niagara/Source/NiagaraEditor/Private/ViewModels/NiagaraParameterPanelViewModel.cpp b/Engine/Plugins/FX/Niagara/Source/NiagaraEditor/Private/ViewModels/NiagaraParameterPanelViewModel.cpp
--- a/Engine/Plugins/FX/Niagara/Source/NiagaraEditor/Private/ViewModels/NiagaraParameterPanelViewModel.cpp
+++ b/Engine/Plugins/FX/Niagara/Source/NiagaraEditor/Private/ViewModels/NiagaraParameterPanelViewModel.cpp
@@ -26,6 +26,7 @@
 #include "Editor.h"
 #include "ScopedTransaction.h"
 #include "NiagaraSystemEditorData.h"
+#include "Widgets/Text/SInlineEditableTextBlock.h"
 
 #define LOCTEXT_NAMESPACE "NiagaraParameterPanelViewModel"
 
@@ -759,7 +760,7 @@ bool FNiagaraScriptToolkitParameterPanelViewModel::GetCanRenameParameterAndToolT
 		};
 
 		// Prevent name values that would alias an existing parameter with relevant usage
-		FText GenericNameAliasToolTip = LOCTEXT("ScriptToolkitParameterPanelViewModel_RenameParameter_Alias", "Cannot rename Parameter: A Parameter already exists with name {0}.");
+		FText GenericNameAliasToolTip = LOCTEXT("ScriptToolkitParameterPanelViewModel_RenameParameter_Alias", "Cannot rename Parameter: A Parameter already exists with name {0} but different type.");
 		FName ExistingName;
 		ensureMsgf(TargetVariableMetaData.GetParameterName(ExistingName), TEXT("Failed to get namespaceless parameter name from metadata on rename verify!"));
 		if (NewName != ExistingName)
@@ -801,8 +802,12 @@ bool FNiagaraScriptToolkitParameterPanelViewModel::GetCanRenameParameterAndToolT
 				const FNiagaraScriptVariableAndViewInfo* NameAliasedVariableInfo = FindParameterNameAlias(AliasScopeNames, NewName);
 				if (NameAliasedVariableInfo != nullptr)
 				{
-					OutCanRenameParameterToolTip = FText::Format(GenericNameAliasToolTip, NewVariableNameText.GetValue());
-					return false;
+					if (TargetVariableToRename.GetType() != NameAliasedVariableInfo->ScriptVariable.GetType())
+					{ 
+						OutCanRenameParameterToolTip = FText::Format(GenericNameAliasToolTip, NewVariableNameText.GetValue());
+						return false;
+					}
+					return true;
 				}
 			}
 		}
@@ -835,11 +840,23 @@ TSharedRef<SWidget> FNiagaraScriptToolkitParameterPanelViewModel::GetScriptParam
 	const UEdGraphSchema_Niagara* Schema = GetDefault<UEdGraphSchema_Niagara>();
 	const FNiagaraVariable PinVar = Schema->PinToNiagaraVariable(Pin);
 
+	bool bRenameImmediate = false;
+	UNiagaraNode* OwnerNode = Cast<UNiagaraNode>(Pin->GetOwningNode());
+	if (OwnerNode && OwnerNode->IsPinNameEditableUponCreation(Pin))
+	{
+		bRenameImmediate = true;
+	}
+	else if (FNiagaraConstants::IsNiagaraConstant(PinVar))
+	{
+		bRenameImmediate = false;
+	}
+
 	const FNiagaraScriptVariableAndViewInfo* ScriptVarAndViewInfo = CachedViewedParameters.FindByPredicate([PinVar](const FNiagaraScriptVariableAndViewInfo& Entry) {return Entry.ScriptVariable == PinVar; });
 	if (ScriptVarAndViewInfo != nullptr)
 	{
-		TSharedPtr<FNiagaraGraphPinParameterNameViewModel> ParameterNameViewModel = MakeShared<FNiagaraGraphPinParameterNameViewModel>(Pin, *ScriptVarAndViewInfo, this);
-		TSharedPtr<SWidget> ScriptParameterVisualWidget = SNew(SNiagaraParameterNameView, ParameterNameViewModel);
+		TSharedPtr<FNiagaraGraphPinParameterNameViewModel> ParameterNameViewModel = MakeShared<FNiagaraGraphPinParameterNameViewModel>(const_cast<UEdGraphPin*>(Pin), *ScriptVarAndViewInfo, this);
+		TSharedPtr<SNiagaraParameterNameView> ScriptParameterVisualWidget = SNew(SNiagaraParameterNameView, ParameterNameViewModel);
+		ScriptParameterVisualWidget->SetPendingRename(bRenameImmediate);
 		return ScriptParameterVisualWidget->AsShared();
 	}
 	else
@@ -852,8 +869,9 @@ TSharedRef<SWidget> FNiagaraScriptToolkitParameterPanelViewModel::GetScriptParam
 			const TStaticArray<FScopeIsEnabledAndTooltip, (int32)ENiagaraParameterScope::Num> PerScopeInfo = GetParameterScopesEnabledAndTooltips(ScriptVar->Variable, ScriptVar->Metadata);
 			FNiagaraScriptVariableAndViewInfo NewScriptVarAndViewInfo = FNiagaraScriptVariableAndViewInfo(ScriptVar->Variable, ScriptVar->Metadata, PerScopeInfo);
 
-			TSharedPtr<FNiagaraGraphPinParameterNameViewModel> ParameterNameViewModel = MakeShared<FNiagaraGraphPinParameterNameViewModel>(Pin, NewScriptVarAndViewInfo, this);
-			TSharedPtr<SWidget> ScriptParameterVisualWidget = SNew(SNiagaraParameterNameView, ParameterNameViewModel);
+			TSharedPtr<FNiagaraGraphPinParameterNameViewModel> ParameterNameViewModel = MakeShared<FNiagaraGraphPinParameterNameViewModel>(const_cast<UEdGraphPin*>(Pin), NewScriptVarAndViewInfo, this);
+			TSharedPtr<SNiagaraParameterNameView> ScriptParameterVisualWidget = SNew(SNiagaraParameterNameView, ParameterNameViewModel);
+			ScriptParameterVisualWidget->SetPendingRename(bRenameImmediate);
 			return ScriptParameterVisualWidget->AsShared();
 		}
 	}
@@ -865,8 +883,9 @@ TSharedRef<SWidget> FNiagaraScriptToolkitParameterPanelViewModel::GetScriptParam
 	StandInMetaData.SetUsage(ENiagaraScriptParameterUsage::Local);
 	
 	FNiagaraScriptVariableAndViewInfo StandInScriptVarAndViewInfo = FNiagaraScriptVariableAndViewInfo(StandInParameter, StandInMetaData);
-	TSharedPtr<FNiagaraGraphPinParameterNameViewModel> ParameterNameViewModel = MakeShared<FNiagaraGraphPinParameterNameViewModel>(Pin, StandInScriptVarAndViewInfo, this);
-	TSharedPtr<SWidget> ScriptParameterVisualWidget = SNew(SNiagaraParameterNameView, ParameterNameViewModel);
+	TSharedPtr<FNiagaraGraphPinParameterNameViewModel> ParameterNameViewModel = MakeShared<FNiagaraGraphPinParameterNameViewModel>(const_cast<UEdGraphPin*>(Pin), StandInScriptVarAndViewInfo, this);
+	TSharedPtr<SNiagaraParameterNameView> ScriptParameterVisualWidget = SNew(SNiagaraParameterNameView, ParameterNameViewModel);
+	ScriptParameterVisualWidget->SetPendingRename(bRenameImmediate);
 	return ScriptParameterVisualWidget->AsShared();
 }
 
diff --git a/Engine/Plugins/FX/Niagara/Source/NiagaraEditor/Private/Widgets/SNiagaraParameterNameView.cpp b/Engine/Plugins/FX/Niagara/Source/NiagaraEditor/Private/Widgets/SNiagaraParameterNameView.cpp
--- a/Engine/Plugins/FX/Niagara/Source/NiagaraEditor/Private/Widgets/SNiagaraParameterNameView.cpp
+++ b/Engine/Plugins/FX/Niagara/Source/NiagaraEditor/Private/Widgets/SNiagaraParameterNameView.cpp
@@ -3,6 +3,7 @@
 #include "SNiagaraParameterNameView.h"
 #include "ViewModels/NiagaraParameterNameViewModel.h"
 #include "Widgets/SBoxPanel.h"
+#include "Widgets/Text/SInlineEditableTextBlock.h"
 
 class SWidget;
 
@@ -10,8 +11,10 @@ void SNiagaraParameterNameView::Construct(const FArguments& InArgs, const TShare
 {
 	ParameterNameViewModel = InParameterNameViewModel;
 
-	TSharedRef<SWidget> ScopeSlotWidget = ParameterNameViewModel->CreateScopeSlotWidget();
-	TSharedRef<SWidget> NameSlotWidget = ParameterNameViewModel->CreateTextSlotWidget();
+	ScopeSlotWidget = ParameterNameViewModel->CreateScopeSlotWidget();
+	NameSlotWidget = ParameterNameViewModel->CreateTextSlotWidget();
+	TSharedRef<SWidget> NameSlotWidgetMade = NameSlotWidget.ToSharedRef();
+	TSharedRef<SWidget> ScopeSlotWidgetMade = ScopeSlotWidget.ToSharedRef();
 
 	ChildSlot
 	[
@@ -21,13 +24,23 @@ void SNiagaraParameterNameView::Construct(const FArguments& InArgs, const TShare
 		.VAlign(VAlign_Center)
 		.Padding(3, 0)
 		[
-			ScopeSlotWidget
+			ScopeSlotWidgetMade
 		]
 		+SHorizontalBox::Slot()
 		.VAlign(VAlign_Center)
 		.Padding(3, 0)
 		[
-			NameSlotWidget
+			NameSlotWidgetMade
 		]
 	];
 }
+
+void SNiagaraParameterNameView::Tick(const FGeometry& AllottedGeometry, const double InCurrentTime, const float InDeltaTime) 
+{
+	if (bPendingRename)
+	{
+		NameSlotWidget.Get()->EnterEditingMode();
+		bPendingRename = false;
+	}
+	SCompoundWidget::Tick(AllottedGeometry, InCurrentTime, InDeltaTime);
+}
diff --git a/Engine/Plugins/FX/Niagara/Source/NiagaraEditor/Private/Widgets/SNiagaraParameterNameView.h b/Engine/Plugins/FX/Niagara/Source/NiagaraEditor/Private/Widgets/SNiagaraParameterNameView.h
--- a/Engine/Plugins/FX/Niagara/Source/NiagaraEditor/Private/Widgets/SNiagaraParameterNameView.h
+++ b/Engine/Plugins/FX/Niagara/Source/NiagaraEditor/Private/Widgets/SNiagaraParameterNameView.h
@@ -67,6 +67,8 @@ public:
 	SLATE_END_ARGS();
 
 	void Construct(const FArguments& InArgs, const TSharedPtr<class INiagaraParameterNameViewModel>& InParameterNameViewModel);
+	void SetPendingRename(bool bInPendingRename) { bPendingRename = bInPendingRename; }
+	virtual void Tick(const FGeometry& AllottedGeometry, const double InCurrentTime, const float InDeltaTime) override;
 
 protected:
 	TSharedPtr<INiagaraParameterNameViewModel> ParameterNameViewModel;
@@ -86,4 +88,8 @@ protected:
 	FText GetTooltipForScopeValue(int32 Value) const;
 	bool GetScopeValueIsEnabled(int32 Value) const;
 	FText GetParameterNameText() const;
+
+	TSharedPtr<SWidget> ScopeSlotWidget;
+	TSharedPtr<SInlineEditableTextBlock> NameSlotWidget;
+	bool bPendingRename = false;
 };
diff --git a/Engine/Plugins/FX/Niagara/Source/NiagaraEditor/Private/Widgets/SNiagaraParameterPanel.cpp b/Engine/Plugins/FX/Niagara/Source/NiagaraEditor/Private/Widgets/SNiagaraParameterPanel.cpp
--- a/Engine/Plugins/FX/Niagara/Source/NiagaraEditor/Private/Widgets/SNiagaraParameterPanel.cpp
+++ b/Engine/Plugins/FX/Niagara/Source/NiagaraEditor/Private/Widgets/SNiagaraParameterPanel.cpp
@@ -134,7 +134,8 @@ void SNiagaraParameterPanel::AddParameter(FNiagaraVariable NewVariable, const Ni
 		GraphActionMenu->SelectItemByName(NewScriptVar->Variable.GetName());
 
 		// Delay calling the rename delegate one tick as the widget to which it is bound is not yet constructed.
-		bGraphActionPendingRename = true;
+		if (FNiagaraConstants::IsNiagaraConstant(NewVariable) == false)
+			bGraphActionPendingRename = true;
 	}
 }
 
@@ -622,12 +623,13 @@ void SNiagaraAddParameterMenu2::CollectAllActions(FGraphActionListBuilderBase& O
 		NiagaraParameterPanelSectionID::Type Type = ShowKnownConstantParametersFilter.Get();
 		if (Type != NiagaraParameterPanelSectionID::Type::NONE)
 		{
-			// Handle engine parameters...
-			if (Type == NiagaraParameterPanelSectionID::Type::ENGINE ||
+			// Handle engine parameters, but only in read-only contexts.
+			if ((!IsParameterRead.IsSet() || (IsParameterRead.IsSet() && IsParameterRead.Get())) &&
+				(Type == NiagaraParameterPanelSectionID::Type::ENGINE ||
 				Type == NiagaraParameterPanelSectionID::Type::OWNER ||
 				Type == NiagaraParameterPanelSectionID::Type::REFERENCES ||
 				Type == NiagaraParameterPanelSectionID::Type::SYSTEM ||
-				Type == NiagaraParameterPanelSectionID::Type::EMITTER)
+				Type == NiagaraParameterPanelSectionID::Type::EMITTER))
 			{
 				const TArray<FNiagaraVariable>& EngineConstants = FNiagaraConstants::GetEngineConstants();
 				for (const FNiagaraVariable& Var : EngineConstants)
diff --git a/Engine/Plugins/FX/Niagara/Source/NiagaraEditor/Public/NiagaraEditorUtilities.h b/Engine/Plugins/FX/Niagara/Source/NiagaraEditor/Public/NiagaraEditorUtilities.h
--- a/Engine/Plugins/FX/Niagara/Source/NiagaraEditor/Public/NiagaraEditorUtilities.h
+++ b/Engine/Plugins/FX/Niagara/Source/NiagaraEditor/Public/NiagaraEditorUtilities.h
@@ -242,6 +242,7 @@ namespace FNiagaraEditorUtilities
 	void CreateAssetFromEmitter(TSharedRef<FNiagaraEmitterHandleViewModel> EmitterHandleViewModel);
 
 	NIAGARAEDITOR_API void WarnWithToastAndLog(FText WarningMessage);
+	NIAGARAEDITOR_API void InfoWithToastAndLog(FText WarningMessage);
 
 	void GetScriptRunAndExecutionIndexFromUsage(const ENiagaraScriptUsage& InUsage, int32& OutRunIndex, int32&OutExecutionIndex);
 
diff --git a/Engine/Plugins/FX/Niagara/Source/NiagaraEditor/Public/NiagaraGraph.h b/Engine/Plugins/FX/Niagara/Source/NiagaraEditor/Public/NiagaraGraph.h
--- a/Engine/Plugins/FX/Niagara/Source/NiagaraEditor/Public/NiagaraGraph.h
+++ b/Engine/Plugins/FX/Niagara/Source/NiagaraEditor/Public/NiagaraGraph.h
@@ -319,9 +319,14 @@ class UNiagaraGraph : public UEdGraph
 	* @param NewName				The new name to apply.
 	* @param bFromStaticSwitch		Whether the target parameter is from a static switch. Used to determine whether to fixup binding strings.
 	* @param NewScopeName			The new scope name from the FNiagaraVariable's associated FNiagaraVariableMetaData. Used if changing the scope triggered a rename. Applied to the canonical UNiagaraScriptVariable's Metadata.
+	* @param bMerged				Whether or not the rename ended up merging with a different parameter because the names are the same.
 	* @return						true if the new name was applied. 
 	*/
-	bool RenameParameter(const FNiagaraVariable& Parameter, FName NewName, bool bRenameRequestedFromStaticSwitch = false, FName NewScopeName = FName());
+	bool RenameParameter(const FNiagaraVariable& Parameter, FName NewName, bool bRenameRequestedFromStaticSwitch = false, FName NewScopeName = FName(), bool* bMerged = nullptr);
+
+	/** Rename a pin inline in a graph. If this is the only instance used in the graph, then rename them all, otherwise make a duplicate. */
+	bool RenameParameterFromPin(const FNiagaraVariable& Parameter, FName NewName, UEdGraphPin* InPin);
+
 
 	/** 
 	 * Sets the ENiagaraScriptParameterUsage on the Metadata of a ScriptVariable depending on its pin usage in the node graph (used on input/output pins). 
diff --git a/Engine/Plugins/FX/Niagara/Source/NiagaraEditor/Public/NiagaraNode.h b/Engine/Plugins/FX/Niagara/Source/NiagaraEditor/Public/NiagaraNode.h
--- a/Engine/Plugins/FX/Niagara/Source/NiagaraEditor/Public/NiagaraNode.h
+++ b/Engine/Plugins/FX/Niagara/Source/NiagaraEditor/Public/NiagaraNode.h
@@ -96,7 +96,7 @@ public:
 	virtual bool VerifyEditablePinName(const FText& InName, FText& OutErrorMessage, const UEdGraphPin* InGraphPinObj) const { return false; }
 
 	/** Verify that the potential rename has produced acceptable results for a pin.*/
-	virtual bool CommitEditablePinName(const FText& InName,  UEdGraphPin* InGraphPinObj) { return false; }
+	virtual bool CommitEditablePinName(const FText& InName,  UEdGraphPin* InGraphPinObj, bool bSuppressEvents = false) { return false; }
 	/** Notify the rename was cancelled.*/
 	virtual bool CancelEditablePinName(const FText& InName, UEdGraphPin* InGraphPinObj) { return false; }
 
diff --git a/Engine/Plugins/FX/Niagara/Source/NiagaraEditor/Public/NiagaraNodeCustomHlsl.h b/Engine/Plugins/FX/Niagara/Source/NiagaraEditor/Public/NiagaraNodeCustomHlsl.h
--- a/Engine/Plugins/FX/Niagara/Source/NiagaraEditor/Public/NiagaraNodeCustomHlsl.h
+++ b/Engine/Plugins/FX/Niagara/Source/NiagaraEditor/Public/NiagaraNodeCustomHlsl.h
@@ -51,7 +51,7 @@ protected:
 	virtual bool VerifyEditablePinName(const FText& InName, FText& OutErrorMessage, const UEdGraphPin* InGraphPinObj) const override;
 	virtual bool IsPinNameEditableUponCreation(const UEdGraphPin* Pin) const override;
 	virtual bool IsPinNameEditable(const UEdGraphPin* Pin) const override;
-	virtual bool CommitEditablePinName(const FText& InName, UEdGraphPin* InGraphPinObj) override;
+	virtual bool CommitEditablePinName(const FText& InName, UEdGraphPin* InGraphPinObj, bool bSuppressEvents = false) override;
 	virtual bool CancelEditablePinName(const FText& InName, UEdGraphPin* InGraphPinObj) override;
 	
 	virtual bool CanRenamePin(const UEdGraphPin* Pin) const override { return UNiagaraNodeWithDynamicPins::CanRenamePin(Pin); }
diff --git a/Engine/Plugins/FX/Niagara/Source/NiagaraEditor/Public/NiagaraNodeWithDynamicPins.h b/Engine/Plugins/FX/Niagara/Source/NiagaraEditor/Public/NiagaraNodeWithDynamicPins.h
--- a/Engine/Plugins/FX/Niagara/Source/NiagaraEditor/Public/NiagaraNodeWithDynamicPins.h
+++ b/Engine/Plugins/FX/Niagara/Source/NiagaraEditor/Public/NiagaraNodeWithDynamicPins.h
@@ -64,6 +64,8 @@ protected:
 	/** Called when a pin is renamed. */
 	virtual void OnPinRenamed(UEdGraphPin* RenamedPin, const FString& OldPinName) { }
 
+	virtual bool CanRenamePinFromContextMenu(const UEdGraphPin* Pin) const { return CanRenamePin(Pin); }
+
 	/** Called to determine if a pin can be renamed by the user. */
 	virtual bool CanRenamePin(const UEdGraphPin* Pin) const;
 
diff --git a/Engine/Plugins/FX/Niagara/Source/NiagaraEditorWidgets/Private/Stack/SNiagaraStackParameterStoreEntryName.cpp b/Engine/Plugins/FX/Niagara/Source/NiagaraEditorWidgets/Private/Stack/SNiagaraStackParameterStoreEntryName.cpp
--- a/Engine/Plugins/FX/Niagara/Source/NiagaraEditorWidgets/Private/Stack/SNiagaraStackParameterStoreEntryName.cpp
+++ b/Engine/Plugins/FX/Niagara/Source/NiagaraEditorWidgets/Private/Stack/SNiagaraStackParameterStoreEntryName.cpp
@@ -81,7 +81,7 @@ void SNiagaraStackParameterStoreEntryName::OnNameTextCommitted(const FText& InTe
 		FNotificationInfo Info(FText::Format(LOCTEXT("NiagaraRenamedUserParameter", "System exposed parameter was renamed.\n{0}\n(All links to inner variables were updated in the process.)"), StackEntry->GetDisplayName()));
 		Info.ExpireDuration = 5.0f;
 		Info.bFireAndForget = true;
-		Info.Image = FCoreStyle::Get().GetBrush(TEXT("MessageLog.Info"));
+		Info.Image = FCoreStyle::Get().GetBrush(TEXT("MessageLog.Note"));
 		FSlateNotificationManager::Get().AddNotification(Info);
 	}
 }
diff --git a/Engine/Plugins/FX/Niagara/Source/NiagaraEditorWidgets/Private/Stack/SNiagaraStackParameterStoreEntryValue.cpp b/Engine/Plugins/FX/Niagara/Source/NiagaraEditorWidgets/Private/Stack/SNiagaraStackParameterStoreEntryValue.cpp
--- a/Engine/Plugins/FX/Niagara/Source/NiagaraEditorWidgets/Private/Stack/SNiagaraStackParameterStoreEntryValue.cpp
+++ b/Engine/Plugins/FX/Niagara/Source/NiagaraEditorWidgets/Private/Stack/SNiagaraStackParameterStoreEntryValue.cpp
@@ -110,7 +110,7 @@ FReply SNiagaraStackParameterStoreEntryValue::DeleteClicked()
 	FNotificationInfo Info(FText::Format(LOCTEXT("NiagaraDeletedUserParameter", "System exposed parameter was deleted.\n{0}\n(All links to inner variables were invalidated in the process.)"), StackEntry->GetDisplayName()));
 	Info.ExpireDuration = 5.0f;
 	Info.bFireAndForget = true;
-	Info.Image = FCoreStyle::Get().GetBrush(TEXT("MessageLog.Info"));
+	Info.Image = FCoreStyle::Get().GetBrush(TEXT("MessageLog.Note"));
 	FSlateNotificationManager::Get().AddNotification(Info);
 
 	// Delete after the notification is posted to prevent the entry from becoming invalidated before generating the message.
