Give clues about setting defaults in the details panel. Also update the panel when they add to the graph.

diff --git a/Engine/Plugins/FX/Niagara/Source/NiagaraEditor/Private/Customizations/NiagaraScriptVariableCustomization.cpp b/Engine/Plugins/FX/Niagara/Source/NiagaraEditor/Private/Customizations/NiagaraScriptVariableCustomization.cpp
--- a/Engine/Plugins/FX/Niagara/Source/NiagaraEditor/Private/Customizations/NiagaraScriptVariableCustomization.cpp
+++ b/Engine/Plugins/FX/Niagara/Source/NiagaraEditor/Private/Customizations/NiagaraScriptVariableCustomization.cpp
@@ -238,6 +238,20 @@ void FNiagaraScriptVariableDetails::CustomizeDetails(IDetailLayoutBuilder& Detai
 				TypeUtilityValue = nullptr;
 			}
 		}
+		else
+		{
+			if (Variable->DefaultMode == ENiagaraDefaultMode::Value)
+			{
+				FDetailWidgetRow& DefaultValueWidget = CategoryBuilder.AddCustomRow(LOCTEXT("DefaultValueFilterText", "Default Value"));
+				DefaultValueWidget.WholeRowContent()
+					.HAlign(HAlign_Fill)
+					[
+						SNew(STextBlock)
+						.Font(FNiagaraEditorStyle::Get().GetFontStyle("NiagaraEditor.ParameterFont"))
+					.Text(NSLOCTEXT("ScriptVariableCustomization", "MissingDefaults", "To set default, add to a Map Get node."))
+					];
+			}
+		}
 	}
 	
 
diff --git a/Engine/Plugins/FX/Niagara/Source/NiagaraEditor/Private/Toolkits/NiagaraScriptToolkit.cpp b/Engine/Plugins/FX/Niagara/Source/NiagaraEditor/Private/Toolkits/NiagaraScriptToolkit.cpp
--- a/Engine/Plugins/FX/Niagara/Source/NiagaraEditor/Private/Toolkits/NiagaraScriptToolkit.cpp
+++ b/Engine/Plugins/FX/Niagara/Source/NiagaraEditor/Private/Toolkits/NiagaraScriptToolkit.cpp
@@ -334,12 +334,16 @@ void FNiagaraScriptToolkit::OnEditedScriptPropertyFinishedChanging(const FProper
 			}
 		}
 	}
+
 	bEditedScriptHasPendingChanges = true;
 }
 
 void FNiagaraScriptToolkit::OnVMScriptCompiled(UNiagaraScript* InScript)
 {
 	UpdateModuleStats();
+	ParameterPanelViewModel->Refresh();
+	if (SelectedDetailsWidget.IsValid())
+		SelectedDetailsWidget->SelectedObjectsChanged();
 }
 
 TSharedRef<SDockTab> FNiagaraScriptToolkit::SpawnTabScriptDetails(const FSpawnTabArgs& Args)
@@ -373,7 +377,7 @@ TSharedRef<SDockTab> FNiagaraScriptToolkit::SpawnTabSelectedDetails(const FSpawn
 		.Label(LOCTEXT("SelectedDetailsTabLabel", "Selected Details"))
 		.TabColorScale(GetTabColorScale())
 		[
-			SNew(SNiagaraSelectedObjectsDetails, ScriptViewModel->GetGraphViewModel()->GetNodeSelection(), ScriptViewModel->GetVariableSelection())
+			SAssignNew(SelectedDetailsWidget, SNiagaraSelectedObjectsDetails, ScriptViewModel->GetGraphViewModel()->GetNodeSelection(), ScriptViewModel->GetVariableSelection())
 		];
 }
 
@@ -524,6 +528,25 @@ FSlateIcon FNiagaraScriptToolkit::GetCompileStatusImage() const
 		return FSlateIcon(FNiagaraEditorStyle::GetStyleSetName(), "Niagara.CompileStatus.Warning");
 	}
 }
+void FNiagaraScriptToolkit::Tick(float DeltaTime)
+{
+	if (bRefreshSelected)
+	{
+		if (SelectedDetailsWidget.IsValid())
+			SelectedDetailsWidget->SelectedObjectsChanged();
+		bRefreshSelected = false;
+	}
+}
+
+bool FNiagaraScriptToolkit::IsTickable() const
+{
+	return true;
+}
+
+TStatId FNiagaraScriptToolkit::GetStatId() const
+{
+	RETURN_QUICK_DECLARE_CYCLE_STAT(FNiagaraScriptToolkit, STATGROUP_Tickables);
+}
 
 FText FNiagaraScriptToolkit::GetCompileStatusTooltip() const
 {
@@ -770,6 +793,8 @@ bool FNiagaraScriptToolkit::OnRequestClose()
 void FNiagaraScriptToolkit::OnEditedScriptGraphChanged(const FEdGraphEditAction& InAction)
 {
 	bEditedScriptHasPendingChanges = true;
+	bRefreshSelected = true;
+
 }
 
 void FNiagaraScriptToolkit::FocusGraphElementIfSameScriptID(const FNiagaraScriptIDAndGraphFocusInfo* FocusInfo)
diff --git a/Engine/Plugins/FX/Niagara/Source/NiagaraEditor/Private/Toolkits/NiagaraScriptToolkit.h b/Engine/Plugins/FX/Niagara/Source/NiagaraEditor/Private/Toolkits/NiagaraScriptToolkit.h
--- a/Engine/Plugins/FX/Niagara/Source/NiagaraEditor/Private/Toolkits/NiagaraScriptToolkit.h
+++ b/Engine/Plugins/FX/Niagara/Source/NiagaraEditor/Private/Toolkits/NiagaraScriptToolkit.h
@@ -10,6 +10,7 @@
 #include "Toolkits/AssetEditorToolkit.h"
 #include "UObject/GCObject.h"
 
+#include "TickableEditorObject.h"
 #include "NiagaraEditorCommon.h"
 
 class IDetailsView;
@@ -23,9 +24,10 @@ struct FEdGraphEditAction;
 class FNiagaraMessageLogViewModel;
 class FNiagaraStandaloneScriptViewModel;
 class FNiagaraScriptToolkitParameterPanelViewModel;
+class SNiagaraSelectedObjectsDetails;
 
 /** Viewer/editor for a DataTable */
-class FNiagaraScriptToolkit : public FAssetEditorToolkit, public FGCObject
+class FNiagaraScriptToolkit : public FAssetEditorToolkit, public FGCObject, public FTickableEditorObject
 {
 public:
 	FNiagaraScriptToolkit();
@@ -60,6 +62,10 @@ public:
 	*/
 	void UpdateModuleStats();
 
+	//~ FTickableEditorObject interface
+	virtual void Tick(float DeltaTime) override;
+	virtual bool IsTickable() const override;
+	virtual TStatId GetStatId() const override;
 protected:
 	//~ FAssetEditorToolkit interface
 	virtual void GetSaveableObjects(TArray<UObject*>& OutObjects) const override;
@@ -86,6 +92,8 @@ private:
 	TSharedRef<SDockTab> SpawnTabScriptParameters2(const FSpawnTabArgs& Args);
 	TSharedRef<SDockTab> SpawnTabStats(const FSpawnTabArgs& Args);
 
+	TSharedPtr< SNiagaraSelectedObjectsDetails> SelectedDetailsWidget;
+
 	TSharedRef<SDockTab> SpawnTabMessageLog(const FSpawnTabArgs& Args);
 
 	/** Sets up commands for the toolkit toolbar. */
@@ -154,6 +162,7 @@ private:
 
 	bool bEditedScriptHasPendingChanges;
 	bool bChangesDiscarded;
+	bool bRefreshSelected = false;
 
 	TSharedPtr<class SNiagaraScriptGraph> NiagaraScriptGraphWidget;
 };
