Fixed a bug where an actor with a Niagara components on autodestroy could be destroyed before the system was finished

diff --git a/Engine/Plugins/FX/Niagara/Source/Niagara/Private/NiagaraActor.cpp b/Engine/Plugins/FX/Niagara/Source/Niagara/Private/NiagaraActor.cpp
--- a/Engine/Plugins/FX/Niagara/Source/Niagara/Private/NiagaraActor.cpp
+++ b/Engine/Plugins/FX/Niagara/Source/Niagara/Private/NiagaraActor.cpp
@@ -65,6 +65,30 @@ ANiagaraActor::ANiagaraActor(const FObjectInitializer& ObjectInitializer)
 #endif // WITH_EDITORONLY_DATA
 }
 
+void ANiagaraActor::PostRegisterAllComponents()
+{
+	Super::PostRegisterAllComponents();
+
+	// Set Notification Delegate
+	if (NiagaraComponent)
+	{
+		NiagaraComponent->OnSystemFinished.AddUniqueDynamic(this, &ANiagaraActor::OnNiagaraSystemFinished);
+	}
+}
+
+void ANiagaraActor::SetDestroyOnSystemFinish(bool bShouldDestroyOnSystemFinish)
+{
+	bDestroyOnSystemFinish = bShouldDestroyOnSystemFinish ? 1 : 0;  
+};
+
+void ANiagaraActor::OnNiagaraSystemFinished(UNiagaraComponent* FinishedComponent)
+{
+	if (bDestroyOnSystemFinish)
+	{
+		SetLifeSpan(0.0001f);
+	}
+}
+
 #if WITH_EDITOR
 bool ANiagaraActor::GetReferencedContentObjects(TArray<UObject*>& Objects) const
 {
@@ -87,4 +111,3 @@ void ANiagaraActor::ResetInLevel()
 	}
 }
 #endif // WITH_EDITOR
-
diff --git a/Engine/Plugins/FX/Niagara/Source/Niagara/Public/NiagaraActor.h b/Engine/Plugins/FX/Niagara/Source/Niagara/Public/NiagaraActor.h
--- a/Engine/Plugins/FX/Niagara/Source/Niagara/Public/NiagaraActor.h
+++ b/Engine/Plugins/FX/Niagara/Source/Niagara/Public/NiagaraActor.h
@@ -12,6 +12,14 @@ class ANiagaraActor : public AActor
 {
 	GENERATED_UCLASS_BODY()
 
+public:
+
+	virtual void PostRegisterAllComponents() override;
+
+	/** Set true for this actor to self-destruct when the Niagara system finishes, false otherwise */
+	UFUNCTION(BlueprintCallable, Category=NiagaraActor)
+	NIAGARA_API void SetDestroyOnSystemFinish(bool bShouldDestroyOnSystemFinish);
+
 private:
 	/** Pointer to System component */
 	UPROPERTY(VisibleAnywhere, BlueprintReadOnly, Category=NiagaraActor, meta = (AllowPrivateAccess = "true"))
@@ -28,6 +36,14 @@ private:
 
 #endif
 
+	/** True for this actor to self-destruct when the Niagara system finishes, false otherwise */
+	UPROPERTY()
+	uint32 bDestroyOnSystemFinish : 1;
+
+	/** Callback when Niagara system finishes. */
+	UFUNCTION(CallInEditor)
+	void OnNiagaraSystemFinished(UNiagaraComponent* FinishedComponent);
+
 public:
 	/** Returns NiagaraComponent subobject **/
 	NIAGARA_API class UNiagaraComponent* GetNiagaraComponent() const { return NiagaraComponent; }
