Remove InternalParameters (script literal values) from the runtime for niagara scripts, replaced with just the blob data that we supply to the VM
-brings the cost of the per-constant cost down to 4 bytes instead of ~56

diff --git a/Engine/Plugins/FX/Niagara/Source/Niagara/Classes/NiagaraScript.h b/Engine/Plugins/FX/Niagara/Source/Niagara/Classes/NiagaraScript.h
--- a/Engine/Plugins/FX/Niagara/Source/Niagara/Classes/NiagaraScript.h
+++ b/Engine/Plugins/FX/Niagara/Source/Niagara/Classes/NiagaraScript.h
@@ -282,8 +282,14 @@ public:
 	UPROPERTY()
 	FNiagaraParameters Parameters;
 
+#if WITH_EDITORONLY_DATA
+	/** All the data for using external constants in the script, laid out in the order they are expected in the uniform table.*/
 	UPROPERTY()
 	FNiagaraParameters InternalParameters;
+#endif
+
+	UPROPERTY()
+	TArray<uint8> ScriptLiterals;
 
 	UPROPERTY()
 	TMap<FName, FNiagaraParameters> DataSetToParameters;
@@ -371,6 +377,10 @@ public:
 	bool IsValid() const;
 
 	void Reset();
+
+#if WITH_EDITORONLY_DATA
+	void BakeScriptLiterals(TArray<uint8>& OutLiterals) const;
+#endif
 };
 
 /** Runtime script for a Niagara system */
diff --git a/Engine/Plugins/FX/Niagara/Source/Niagara/Classes/NiagaraScriptExecutionContext.h b/Engine/Plugins/FX/Niagara/Source/Niagara/Classes/NiagaraScriptExecutionContext.h
--- a/Engine/Plugins/FX/Niagara/Source/Niagara/Classes/NiagaraScriptExecutionContext.h
+++ b/Engine/Plugins/FX/Niagara/Source/Niagara/Classes/NiagaraScriptExecutionContext.h
@@ -188,6 +188,8 @@ public:
 	void DirtyDataInterfaces();
 
 	bool CanExecute()const;
+
+	TArrayView<const uint8> GetScriptLiterals() const;
 };
 
 struct FNiagaraGpuSpawnInfoParams
diff --git a/Engine/Plugins/FX/Niagara/Source/Niagara/Private/NiagaraEmitterInstance.cpp b/Engine/Plugins/FX/Niagara/Source/Niagara/Private/NiagaraEmitterInstance.cpp
--- a/Engine/Plugins/FX/Niagara/Source/Niagara/Private/NiagaraEmitterInstance.cpp
+++ b/Engine/Plugins/FX/Niagara/Source/Niagara/Private/NiagaraEmitterInstance.cpp
@@ -1818,12 +1818,11 @@ void FNiagaraEmitterInstance::BuildConstantBufferTable(
 	const FNiagaraScriptExecutionContext& ExecContext,
 	FScriptExecutionConstantBufferTable& ConstantBufferTable) const
 {
+	const auto ScriptLiterals = ExecContext.GetScriptLiterals();
 	const auto& ExternalParameterData = ExecContext.Parameters.GetParameterDataArray();
 	uint8* ExternalParameterBuffer = const_cast<uint8*>(ExternalParameterData.GetData());
 
 	const uint32 ExternalParameterSize = ExecContext.Parameters.GetExternalParameterSize();
-	const uint32 LiteralConstantOffset = (ExecContext.HasInterpolationParameters ? 2 : 1) * ExternalParameterSize;
-	const uint32 LiteralConstantSize = ExternalParameterData.Num() - LiteralConstantOffset;
 
 	const uint32 TableCount = 5 * (ExecContext.HasInterpolationParameters ? 2 : 1) + 1;
 	ConstantBufferTable.Reset(TableCount);
@@ -1843,5 +1842,5 @@ void FNiagaraEmitterInstance::BuildConstantBufferTable(
 		ConstantBufferTable.AddRawBuffer(ExternalParameterBuffer + ExternalParameterSize, ExternalParameterSize);
 	}
 
-	ConstantBufferTable.AddRawBuffer(ExternalParameterBuffer + LiteralConstantOffset, LiteralConstantSize);
+	ConstantBufferTable.AddRawBuffer(ScriptLiterals.GetData(), ScriptLiterals.Num());
 }
diff --git a/Engine/Plugins/FX/Niagara/Source/Niagara/Private/NiagaraScript.cpp b/Engine/Plugins/FX/Niagara/Source/Niagara/Private/NiagaraScript.cpp
--- a/Engine/Plugins/FX/Niagara/Source/Niagara/Private/NiagaraScript.cpp
+++ b/Engine/Plugins/FX/Niagara/Source/Niagara/Private/NiagaraScript.cpp
@@ -105,6 +105,35 @@ void FNiagaraVMExecutableData::SerializeData(FArchive& Ar, bool bDDCData)
 	FNiagaraVMExecutableDataType->SerializeTaggedProperties(Ar, (uint8*)this, FNiagaraVMExecutableDataType, nullptr);
 }
 
+#if WITH_EDITORONLY_DATA
+void FNiagaraVMExecutableData::BakeScriptLiterals(TArray<uint8>& OutLiterals) const
+{
+	OutLiterals.Empty();
+
+	const auto& Variables = InternalParameters.Parameters;
+	const int32 VariableCount = Variables.Num();
+
+	int32 TotalSize = 0;
+	for (int32 Index = 0; Index < VariableCount; ++Index)
+	{
+		TotalSize += Variables[Index].GetAllocatedSizeInBytes();
+	}
+
+	OutLiterals.AddZeroed(TotalSize);
+
+	uint8* LiteralData = OutLiterals.GetData();
+
+	for (int32 Index = 0; Index < VariableCount; ++Index)
+	{
+		const FNiagaraVariable& Variable = Variables[Index];
+		const int32 VariableSize = Variable.GetAllocatedSizeInBytes();
+		
+		FMemory::Memcpy(LiteralData, Variable.GetData(), VariableSize);
+		LiteralData += VariableSize;
+	}
+}
+#endif
+
 bool FNiagaraVMExecutableDataId::IsValid() const
 {
 	return CompilerVersionID.IsValid();
@@ -810,6 +839,22 @@ void UNiagaraScript::Serialize(FArchive& Ar)
 		}
 	}
 
+#if WITH_EDITOR
+	if (Ar.IsCooking() && Ar.IsSaving())
+	{
+		auto& ExecutableData = GetVMExecutableData();
+
+		if (Usage != ENiagaraScriptUsage::ParticleGPUComputeScript)
+		{
+			ExecutableData.BakeScriptLiterals(ExecutableData.ScriptLiterals);
+		}
+		else
+		{
+			ExecutableData.ScriptLiterals.Empty();
+		}
+	}
+#endif
+
 	Super::Serialize(Ar);
 
 	// Restore after serialize
diff --git a/Engine/Plugins/FX/Niagara/Source/Niagara/Private/NiagaraScriptExecutionContext.cpp b/Engine/Plugins/FX/Niagara/Source/Niagara/Private/NiagaraScriptExecutionContext.cpp
--- a/Engine/Plugins/FX/Niagara/Source/Niagara/Private/NiagaraScriptExecutionContext.cpp
+++ b/Engine/Plugins/FX/Niagara/Source/Niagara/Private/NiagaraScriptExecutionContext.cpp
@@ -296,6 +296,15 @@ bool FNiagaraScriptExecutionContext::CanExecute()const
 	return Script && Script->GetVMExecutableData().IsValid() && Script->GetVMExecutableData().ByteCode.Num() > 0;
 }
 
+TArrayView<const uint8> FNiagaraScriptExecutionContext::GetScriptLiterals() const
+{
+#if WITH_EDITORONLY_DATA
+	return Parameters.GetScriptLiterals();
+#else
+	return MakeArrayView(Script->GetVMExecutableData().ScriptLiterals);
+#endif
+}
+
 void FNiagaraGPUSystemTick::Init(FNiagaraSystemInstance* InSystemInstance)
 {
 	check(IsInGameThread());
diff --git a/Engine/Plugins/FX/Niagara/Source/Niagara/Private/NiagaraScriptExecutionParameterStore.cpp b/Engine/Plugins/FX/Niagara/Source/Niagara/Private/NiagaraScriptExecutionParameterStore.cpp
--- a/Engine/Plugins/FX/Niagara/Source/Niagara/Private/NiagaraScriptExecutionParameterStore.cpp
+++ b/Engine/Plugins/FX/Niagara/Source/Niagara/Private/NiagaraScriptExecutionParameterStore.cpp
@@ -225,9 +225,11 @@ void FNiagaraScriptExecutionParameterStore::AddScriptParams(UNiagaraScript* Scri
 	}
 	PaddingInfo.Empty();
 
+	const FNiagaraVMExecutableData& ExecutableData = Script->GetVMExecutableData();
+
 	//Here we add the current frame parameters.
 	bool bAdded = false;
-	for (FNiagaraVariable& Param : Script->GetVMExecutableData().Parameters.Parameters)
+	for (const FNiagaraVariable& Param : ExecutableData.Parameters.Parameters)
 	{
 		bAdded |= AddParameter(Param, false, false);
 	}
@@ -243,20 +245,17 @@ void FNiagaraScriptExecutionParameterStore::AddScriptParams(UNiagaraScript* Scri
 	{
 		AddAlignmentPadding();
 
-		for (FNiagaraVariable& Param : Script->GetVMExecutableData().Parameters.Parameters)
+		for (const FNiagaraVariable& Param : ExecutableData.Parameters.Parameters)
 		{
 			FNiagaraVariable PrevParam(Param.GetType(), FName(*(INTERPOLATED_PARAMETER_PREFIX + Param.GetName().ToString())));
 			bAdded |= AddParameter(PrevParam, false, false);
 		}
 	}
 
-	//Internal constants - only needed for non-GPU sim
+	// for VM scripts we need to build the script literals; in cooked builds this is already in the cached VM data
 	if (SimTarget != ENiagaraSimTarget::GPUComputeSim)
 	{
-		for (FNiagaraVariable& InternalVar : Script->GetVMExecutableData().InternalParameters.Parameters)
-		{
-			bAdded |= AddParameter(InternalVar, false, false);
-		}
+		ExecutableData.BakeScriptLiterals(CachedScriptLiterals);
 	}
 
 	check(Script->GetVMExecutableData().DataInterfaceInfo.Num() == Script->GetCachedDefaultDataInterfaces().Num());
@@ -386,4 +385,15 @@ void FNiagaraScriptInstanceParameterStore::InternalStorageChanged()
 	{
 		ParameterVariables = TArrayView<FNiagaraVariableWithOffset>();
 	}
-}
\ No newline at end of file
+}
+
+#if WITH_EDITORONLY_DATA
+TArrayView<const uint8> FNiagaraScriptInstanceParameterStore::GetScriptLiterals() const
+{
+	if (const auto* ScriptStore = ScriptParameterStore.Get())
+	{
+		return MakeArrayView(ScriptStore->CachedScriptLiterals);
+	}
+	return MakeArrayView<const uint8>(nullptr, 0);
+}
+#endif
\ No newline at end of file
diff --git a/Engine/Plugins/FX/Niagara/Source/Niagara/Private/NiagaraSystemSimulation.cpp b/Engine/Plugins/FX/Niagara/Source/Niagara/Private/NiagaraSystemSimulation.cpp
--- a/Engine/Plugins/FX/Niagara/Source/Niagara/Private/NiagaraSystemSimulation.cpp
+++ b/Engine/Plugins/FX/Niagara/Source/Niagara/Private/NiagaraSystemSimulation.cpp
@@ -1846,17 +1846,16 @@ void FNiagaraSystemSimulation::BuildConstantBufferTable(
 	FNiagaraScriptExecutionContext& ExecContext,
 	FScriptExecutionConstantBufferTable& ConstantBufferTable) const
 {
+	const auto ScriptLiterals = ExecContext.GetScriptLiterals();
+
 	check(!ExecContext.HasInterpolationParameters);
 
 	const auto& ExternalParameterData = ExecContext.Parameters.GetParameterDataArray();
 	uint8* ExternalParameterBuffer = const_cast<uint8*>(ExternalParameterData.GetData());
-
 	const uint32 ExternalParameterSize = ExecContext.Parameters.GetExternalParameterSize();
-	const uint32 LiteralConstantOffset = ExternalParameterSize;
-	const uint32 LiteralConstantSize = ExternalParameterData.Num() - LiteralConstantOffset;
 
 	ConstantBufferTable.Reset(3);
 	ConstantBufferTable.AddTypedBuffer(GlobalParameters);
 	ConstantBufferTable.AddRawBuffer(ExternalParameterBuffer, ExternalParameterSize);
-	ConstantBufferTable.AddRawBuffer(ExternalParameterBuffer + LiteralConstantOffset, LiteralConstantSize);
+	ConstantBufferTable.AddRawBuffer(ScriptLiterals.GetData(), ScriptLiterals.Num());
 }
diff --git a/Engine/Plugins/FX/Niagara/Source/Niagara/Public/NiagaraScriptExecutionParameterStore.h b/Engine/Plugins/FX/Niagara/Source/Niagara/Public/NiagaraScriptExecutionParameterStore.h
--- a/Engine/Plugins/FX/Niagara/Source/Niagara/Public/NiagaraScriptExecutionParameterStore.h
+++ b/Engine/Plugins/FX/Niagara/Source/Niagara/Public/NiagaraScriptExecutionParameterStore.h
@@ -110,6 +110,10 @@ public:
 	UPROPERTY()
 	uint8 bInitialized : 1;
 
+#if WITH_EDITORONLY_DATA
+	TArray<uint8> CachedScriptLiterals;
+#endif
+
 protected:
 	void AddPaddedParamSize(const FNiagaraTypeDefinition& InParamType, uint32 InOffset);
 	void AddAlignmentPadding();
@@ -177,6 +181,10 @@ public:
 	// Helper that converts the data from the base type array internally into the padded out renderer-ready format.
 	void CopyParameterDataToPaddedBuffer(uint8* InTargetBuffer, uint32 InTargetBufferSizeInBytes) const;
 
+#if WITH_EDITORONLY_DATA
+	TArrayView<const uint8> GetScriptLiterals() const;
+#endif
+
 protected:
 	virtual void InternalStorageChanged() override;
 
