Move the Niagara's script compiled Parmeters array to editor only

diff --git a/Engine/Plugins/FX/Niagara/Source/Niagara/Classes/NiagaraEmitter.h b/Engine/Plugins/FX/Niagara/Source/Niagara/Classes/NiagaraEmitter.h
--- a/Engine/Plugins/FX/Niagara/Source/Niagara/Classes/NiagaraEmitter.h
+++ b/Engine/Plugins/FX/Niagara/Source/Niagara/Classes/NiagaraEmitter.h
@@ -501,6 +501,8 @@ public:
 	NIAGARA_API void SetParent(UNiagaraEmitter& InParent);
 
 	NIAGARA_API	void Reparent(UNiagaraEmitter& InParent);
+
+	void InitFastPathAttributeNames();
 #endif
 
 	void OnEffectsQualityChanged();
diff --git a/Engine/Plugins/FX/Niagara/Source/Niagara/Private/NiagaraEmitter.cpp b/Engine/Plugins/FX/Niagara/Source/Niagara/Private/NiagaraEmitter.cpp
--- a/Engine/Plugins/FX/Niagara/Source/Niagara/Private/NiagaraEmitter.cpp
+++ b/Engine/Plugins/FX/Niagara/Source/Niagara/Private/NiagaraEmitter.cpp
@@ -5,7 +5,7 @@
 #include "NiagaraScript.h"
 #include "NiagaraScriptSourceBase.h"
 #include "NiagaraRendererProperties.h"
-#include "NiagaraSimulationStageBase.h"
+#include "NiagaraShaderStageBase.h"
 #include "NiagaraCustomVersion.h"
 #include "UObject/Package.h"
 #include "UObject/Linker.h"
@@ -264,34 +264,6 @@ UNiagaraEmitter* UNiagaraEmitter::DuplicateWithoutMerging(UObject* InOuter)
 
 void UNiagaraEmitter::Serialize(FArchive& Ar)
 {
-
-#if WITH_EDITORONLY_DATA
-	for (UNiagaraSimulationStageBase* Stage : SimulationStages)
-	{
-		if (Stage)
-		{
-			if (Stage->Script)
-			{
-				if (!HasAnyFlags(RF_Transient))
-				{
-					if (Stage->Script->HasAnyFlags(RF_Transient))
-					{
-						UE_LOG(LogNiagara, Error, TEXT("Emitter \"%s\" has a simulation stage with a Transient script and the emitter itself isn't transient!"), *GetPathName());
-					}
-				}
-			}
-			else
-			{
-				UE_LOG(LogNiagara, Error, TEXT("Emitter \"%s\" has a simulation stage with a null Script entry!"), *GetPathName());
-			}
-
-		}
-		else
-		{
-			UE_LOG(LogNiagara, Error, TEXT("Emitter \"%s\" has a simulation stage with a null entry!"), *GetPathName());
-		}
-	}
-#endif
 	Super::Serialize(Ar);
 
 	Ar.UsingCustomVersion(FNiagaraCustomVersion::GUID);
@@ -314,12 +286,11 @@ void UNiagaraEmitter::PostLoad()
 		}
 	}
 
-
-	for (int32 SimulationStageIndex = SimulationStages.Num() - 1; SimulationStageIndex >= 0; --SimulationStageIndex)
+	for (int32 ShaderStageIndex = ShaderStages.Num() - 1; ShaderStageIndex >= 0; --ShaderStageIndex)
 	{
-		if (ensureMsgf(SimulationStages[SimulationStageIndex] != nullptr && SimulationStages[SimulationStageIndex]->Script != nullptr, TEXT("Null simulation stage, or simulation stage with a null script found in %s at index %i, removing it to prevent crashes."), *GetPathName(), SimulationStageIndex) == false)
+		if (ensureMsgf(ShaderStages[ShaderStageIndex] != nullptr && ShaderStages[ShaderStageIndex]->Script != nullptr, TEXT("Null shader stage, or shader stage with a null script found in %s at index %i, removing it to prevent crashes."), *GetPathName(), ShaderStageIndex) == false)
 		{
-			SimulationStages.RemoveAt(SimulationStageIndex);
+			ShaderStages.RemoveAt(ShaderStageIndex);
 		}
 	}
 
@@ -541,10 +512,10 @@ void UNiagaraEmitter::PostLoad()
 				FNiagaraParameterStore::FOnChanged::FDelegate::CreateUObject(this, &UNiagaraEmitter::ScriptRapidIterationParameterChanged));
 		}
 
-		for (UNiagaraSimulationStageBase* SimulationStage : SimulationStages)
+		for (UNiagaraShaderStageBase* ShaderStage : ShaderStages)
 		{
-			SimulationStage->OnChanged().AddUObject(this, &UNiagaraEmitter::SimulationStageChanged);
-			SimulationStage->Script->RapidIterationParameters.AddOnChangedHandler(
+			ShaderStage->OnChanged().AddUObject(this, &UNiagaraEmitter::ShaderStageChanged);
+			ShaderStage->Script->RapidIterationParameters.AddOnChangedHandler(
 				FNiagaraParameterStore::FOnChanged::FDelegate::CreateUObject(this, &UNiagaraEmitter::ScriptRapidIterationParameterChanged));
 		}
 
@@ -715,24 +686,9 @@ void UNiagaraEmitter::PostEditChangeProperty(struct FPropertyChangedEvent& Prope
 			GraphSource->MarkNotSynchronized(TEXT("Emitter Determinism changed."));
 		}
 
-		bNeedsRecompile = true;
-	}
-	if (PropertyName == GET_MEMBER_NAME_CHECKED(UNiagaraEmitter, bSimulationStagesEnabled))
-	{
-		if (GraphSource != nullptr)
-		{
-			GraphSource->MarkNotSynchronized(TEXT("SimulationStagesEnabled changed."));
-		}
-		bNeedsRecompile = true;
-	}
-	if (PropertyName == GET_MEMBER_NAME_CHECKED(UNiagaraEmitter, bDeprecatedShaderStagesEnabled))
-	{
-		if (GraphSource != nullptr)
-		{
-			GraphSource->MarkNotSynchronized(TEXT("DeprecatedShaderStagesEnabled changed."));
-		}
-		bNeedsRecompile = true;
-
+#if WITH_EDITORONLY_DATA
+		UNiagaraSystem::RequestCompileForEmitter(this);
+#endif
 	}
 
 	ResolveScalabilitySettings();
@@ -868,11 +824,11 @@ void UNiagaraEmitter::GetScripts(TArray<UNiagaraScript*>& OutScripts, bool bComp
 
 	if (!bCompilableOnly)
 	{
-		for (int32 i = 0; i < SimulationStages.Num(); i++)
+		for (int32 i = 0; i < ShaderStages.Num(); i++)
 		{
-			if (SimulationStages[i] && SimulationStages[i]->Script)
+			if (ShaderStages[i] && ShaderStages[i]->Script)
 			{
-				OutScripts.Add(SimulationStages[i]->Script);
+				OutScripts.Add(ShaderStages[i]->Script);
 			}
 		}
 	}
@@ -964,9 +920,9 @@ bool UNiagaraEmitter::AreAllScriptAndSourcesSynchronized() const
 		}
 	}
 
-	for (int32 i = 0; i < SimulationStages.Num(); i++)
+	for (int32 i = 0; i < ShaderStages.Num(); i++)
 	{
-		if (SimulationStages[i] && SimulationStages[i]->Script  && SimulationStages[i]->Script->IsCompilable() && !SimulationStages[i]->Script->AreScriptAndSourceSynchronized())
+		if (ShaderStages[i] && ShaderStages[i]->Script  && ShaderStages[i]->Script->IsCompilable() && !ShaderStages[i]->Script->AreScriptAndSourceSynchronized())
 		{
 			return false;
 		}
@@ -1049,6 +1005,8 @@ void UNiagaraEmitter::OnPostCompile()
 	RuntimeEstimation = MemoryRuntimeEstimation();
 
 	OnEmitterVMCompiled().Broadcast(this);
+
+	InitFastPathAttributeNames();
 }
 
 UNiagaraEmitter* UNiagaraEmitter::MakeRecursiveDeepCopy(UObject* DestOuter) const
@@ -1227,19 +1185,19 @@ void UNiagaraEmitter::UpdateFromMergedCopy(const INiagaraMergeManager& MergeMana
 	}
 
 	// Copy shader stages
-	for (UNiagaraSimulationStageBase*& SimulationStage : SimulationStages)
+	for (UNiagaraShaderStageBase*& ShaderStage : ShaderStages)
 	{
-		SimulationStage->OnChanged().RemoveAll(this);
-		SimulationStage->Script->RapidIterationParameters.RemoveAllOnChangedHandlers(this);
+		ShaderStage->OnChanged().RemoveAll(this);
+		ShaderStage->Script->RapidIterationParameters.RemoveAllOnChangedHandlers(this);
 	}
-	SimulationStages.Empty();
+	ShaderStages.Empty();
 
-	for (UNiagaraSimulationStageBase* MergedSimulationStage : MergedEmitter->SimulationStages)
+	for (UNiagaraShaderStageBase* MergedShaderStage : MergedEmitter->ShaderStages)
 	{
-		ReouterMergedObject(this, MergedSimulationStage);
-		SimulationStages.Add(MergedSimulationStage);
-		MergedSimulationStage->OnChanged().AddUObject(this, &UNiagaraEmitter::SimulationStageChanged);
-		MergedSimulationStage->Script->RapidIterationParameters.AddOnChangedHandler(
+		ReouterMergedObject(this, MergedShaderStage);
+		ShaderStages.Add(MergedShaderStage);
+		MergedShaderStage->OnChanged().AddUObject(this, &UNiagaraEmitter::ShaderStageChanged);
+		MergedShaderStage->Script->RapidIterationParameters.AddOnChangedHandler(
 			FNiagaraParameterStore::FOnChanged::FDelegate::CreateUObject(this, &UNiagaraEmitter::ScriptRapidIterationParameterChanged));
 	}
 
@@ -1388,52 +1346,52 @@ void UNiagaraEmitter::RemoveEventHandlerByUsageId(FGuid EventHandlerUsageId)
 #endif
 }
 
-UNiagaraSimulationStageBase* UNiagaraEmitter::GetSimulationStageById(FGuid ScriptUsageId) const
+UNiagaraShaderStageBase* UNiagaraEmitter::GetShaderStageById(FGuid ScriptUsageId) const
 {
-	UNiagaraSimulationStageBase*const* FoundSimulationStagePtr = SimulationStages.FindByPredicate([&ScriptUsageId](UNiagaraSimulationStageBase* SimulationStage) { return SimulationStage->Script->GetUsageId() == ScriptUsageId; });
-	return FoundSimulationStagePtr != nullptr ? *FoundSimulationStagePtr : nullptr;
+	UNiagaraShaderStageBase*const* FoundShaderStagePtr = ShaderStages.FindByPredicate([&ScriptUsageId](UNiagaraShaderStageBase* ShaderStage) { return ShaderStage->Script->GetUsageId() == ScriptUsageId; });
+	return FoundShaderStagePtr != nullptr ? *FoundShaderStagePtr : nullptr;
 }
 
-void UNiagaraEmitter::AddSimulationStage(UNiagaraSimulationStageBase* SimulationStage)
+void UNiagaraEmitter::AddShaderStage(UNiagaraShaderStageBase* ShaderStage)
 {
 	Modify();
-	SimulationStages.Add(SimulationStage);
+	ShaderStages.Add(ShaderStage);
 #if WITH_EDITOR
-	SimulationStage->OnChanged().AddUObject(this, &UNiagaraEmitter::SimulationStageChanged);
-	SimulationStage->Script->RapidIterationParameters.AddOnChangedHandler(
+	ShaderStage->OnChanged().AddUObject(this, &UNiagaraEmitter::ShaderStageChanged);
+	ShaderStage->Script->RapidIterationParameters.AddOnChangedHandler(
 		FNiagaraParameterStore::FOnChanged::FDelegate::CreateUObject(this, &UNiagaraEmitter::ScriptRapidIterationParameterChanged));
 	UpdateChangeId(TEXT("Shader stage added"));
 #endif
 }
 
-void UNiagaraEmitter::RemoveSimulationStage(UNiagaraSimulationStageBase* SimulationStage)
+void UNiagaraEmitter::RemoveShaderStage(UNiagaraShaderStageBase* ShaderStage)
 {
 	Modify();
-	bool bRemoved = SimulationStages.Remove(SimulationStage) != 0;
+	bool bRemoved = ShaderStages.Remove(ShaderStage) != 0;
 #if WITH_EDITOR
 	if (bRemoved)
 	{
-		SimulationStage->OnChanged().RemoveAll(this);
-		SimulationStage->Script->RapidIterationParameters.RemoveAllOnChangedHandlers(this);
-		UpdateChangeId(TEXT("Simulation stage removed"));
+		ShaderStage->OnChanged().RemoveAll(this);
+		ShaderStage->Script->RapidIterationParameters.RemoveAllOnChangedHandlers(this);
+		UpdateChangeId(TEXT("Shader stage removed"));
 	}
 #endif
 }
 
-void UNiagaraEmitter::MoveSimulationStageToIndex(UNiagaraSimulationStageBase* SimulationStageToMove, int32 TargetIndex)
+void UNiagaraEmitter::MoveShaderStageToIndex(UNiagaraShaderStageBase* ShaderStageToMove, int32 TargetIndex)
 {
-	int32 CurrentIndex = SimulationStages.IndexOfByKey(SimulationStageToMove);
-	checkf(CurrentIndex != INDEX_NONE, TEXT("Simulation stage could not be moved because it is not owned by this emitter."));
+	int32 CurrentIndex = ShaderStages.IndexOfByKey(ShaderStageToMove);
+	checkf(CurrentIndex != INDEX_NONE, TEXT("Shader stage could not be moved because it is not owned by this emitter."));
 	if (TargetIndex != CurrentIndex)
 	{
 		int32 AdjustedTargetIndex = CurrentIndex < TargetIndex
 			? TargetIndex - 1 // If the current index is less than the target index, the target index needs to be decreased to make up for the item being removed.
 			: TargetIndex;
 
-		SimulationStages.Remove(SimulationStageToMove);
-		SimulationStages.Insert(SimulationStageToMove, AdjustedTargetIndex);
+		ShaderStages.Remove(ShaderStageToMove);
+		ShaderStages.Insert(ShaderStageToMove, AdjustedTargetIndex);
 #if WITH_EDITOR
-		UpdateChangeId("Simulation stage moved.");
+		UpdateChangeId("Shader stage moved.");
 #endif
 	}
 }
@@ -1479,9 +1437,9 @@ void UNiagaraEmitter::ScriptRapidIterationParameterChanged()
 	UpdateChangeId(TEXT("Script rapid iteration parameter changed."));
 }
 
-void UNiagaraEmitter::SimulationStageChanged()
+void UNiagaraEmitter::ShaderStageChanged()
 {
-	UpdateChangeId(TEXT("Simulation Stage Changed"));
+	UpdateChangeId(TEXT("Shader Stage Changed"));
 }
 
 void UNiagaraEmitter::RendererChanged()
