- Skeletal Mesh Niagara DI: Rename "DefaultMesh" to "PreviewMesh"
- Static Mesh Niagara DI: Add "PreviewMesh" fallback mesh that won't get cooked out for viewing in the editor

diff --git a/Engine/Plugins/FX/Niagara/Source/Niagara/Classes/NiagaraDataInterfaceSkeletalMesh.h b/Engine/Plugins/FX/Niagara/Source/Niagara/Classes/NiagaraDataInterfaceSkeletalMesh.h
--- a/Engine/Plugins/FX/Niagara/Source/Niagara/Classes/NiagaraDataInterfaceSkeletalMesh.h
+++ b/Engine/Plugins/FX/Niagara/Source/Niagara/Classes/NiagaraDataInterfaceSkeletalMesh.h
@@ -497,7 +497,7 @@ struct FNDISkeletalMesh_InstanceData
 	bool HasColorData();
 };
 
-/** Data Interface allowing sampling of static meshes. */
+/** Data Interface allowing sampling of skeletal meshes. */
 UCLASS(EditInlineNew, Category = "Meshes", meta = (DisplayName = "Skeletal Mesh"))
 class NIAGARA_API UNiagaraDataInterfaceSkeletalMesh : public UNiagaraDataInterface
 {
@@ -510,9 +510,9 @@ public:
 #if WITH_EDITORONLY_DATA
 	/** Mesh used to sample from when not overridden by a source actor from the scene. Only available in editor for previewing. This is removed in cooked builds. */
 	UPROPERTY(EditAnywhere, Category = "Mesh")
-	USkeletalMesh* DefaultMesh;
+	USkeletalMesh* PreviewMesh;
 #endif
-
+	
 	/** The source actor from which to sample. Takes precedence over the direct mesh. Note that this can only be set when used as a user variable on a component in the world.*/
 	UPROPERTY(EditAnywhere, Category = "Mesh")
 	AActor* Source;
diff --git a/Engine/Plugins/FX/Niagara/Source/Niagara/Classes/NiagaraDataInterfaceStaticMesh.h b/Engine/Plugins/FX/Niagara/Source/Niagara/Classes/NiagaraDataInterfaceStaticMesh.h
--- a/Engine/Plugins/FX/Niagara/Source/Niagara/Classes/NiagaraDataInterfaceStaticMesh.h
+++ b/Engine/Plugins/FX/Niagara/Source/Niagara/Classes/NiagaraDataInterfaceStaticMesh.h
@@ -167,7 +167,13 @@ public:
 
 	DECLARE_NIAGARA_DI_PARAMETER();
 	
-	/** Mesh used to sample from when not overridden by a source actor from the scene. Also useful for previewing in the editor. */
+#if WITH_EDITORONLY_DATA
+	/** Mesh used to sample from when not overridden by a source actor from the scene. Only available in editor for previewing. This is removed in cooked builds. */
+	UPROPERTY(EditAnywhere, Category = "Mesh")
+	UStaticMesh* PreviewMesh;
+#endif
+
+	/** Mesh used to sample from when not overridden by a source actor from the scene. This mesh is NOT removed from cooked builds. */
 	UPROPERTY(EditAnywhere, Category = "Mesh")
 	UStaticMesh* DefaultMesh;
 
diff --git a/Engine/Plugins/FX/Niagara/Source/Niagara/Private/NiagaraDataInterfaceSkeletalMesh.cpp b/Engine/Plugins/FX/Niagara/Source/Niagara/Private/NiagaraDataInterfaceSkeletalMesh.cpp
--- a/Engine/Plugins/FX/Niagara/Source/Niagara/Private/NiagaraDataInterfaceSkeletalMesh.cpp
+++ b/Engine/Plugins/FX/Niagara/Source/Niagara/Private/NiagaraDataInterfaceSkeletalMesh.cpp
@@ -1267,9 +1267,9 @@ USkeletalMesh* UNiagaraDataInterfaceSkeletalMesh::GetSkeletalMesh(UNiagaraCompon
 	}
 
 #if WITH_EDITORONLY_DATA
-	if (!Mesh && DefaultMesh)
+	if (!Mesh && PreviewMesh)
 	{
-		Mesh = DefaultMesh;
+		Mesh = PreviewMesh;
 	}
 #endif
 
@@ -1651,7 +1651,7 @@ bool FNDISkeletalMesh_InstanceData::ResetRequired(UNiagaraDataInterfaceSkeletalM
 	else
 	{
 #if WITH_EDITORONLY_DATA
-		if (!Interface->DefaultMesh)
+		if (!Interface->PreviewMesh)
 		{
 			return true;
 		}
@@ -1766,7 +1766,7 @@ void FNDISkeletalMesh_InstanceData::Release()
 UNiagaraDataInterfaceSkeletalMesh::UNiagaraDataInterfaceSkeletalMesh(FObjectInitializer const& ObjectInitializer)
 	: Super(ObjectInitializer)
 #if WITH_EDITORONLY_DATA
-	, DefaultMesh(nullptr)
+	, PreviewMesh(nullptr)
 #endif
 	, Source(nullptr)
 	, SkinningMode(ENDISkeletalMesh_SkinningMode::SkinOnTheFly)
@@ -1812,7 +1812,6 @@ void UNiagaraDataInterfaceSkeletalMesh::PostEditChangeProperty(FPropertyChangedE
 
 #endif //WITH_EDITOR
 
-
 void UNiagaraDataInterfaceSkeletalMesh::GetFunctions(TArray<FNiagaraFunctionSignature>& OutFunctions)
 {
 	GetTriangleSamplingFunctions(OutFunctions);
@@ -1890,7 +1889,7 @@ bool UNiagaraDataInterfaceSkeletalMesh::CopyToInternal(UNiagaraDataInterface* De
 	OtherTyped->SpecificBones = SpecificBones;
 	OtherTyped->SpecificSockets = SpecificSockets;
 #if WITH_EDITORONLY_DATA
-	OtherTyped->DefaultMesh = DefaultMesh;
+	OtherTyped->PreviewMesh = PreviewMesh;
 	OtherTyped->bRequiresCPUAccess = bRequiresCPUAccess;
 #endif
 	return true;
@@ -1905,7 +1904,7 @@ bool UNiagaraDataInterfaceSkeletalMesh::Equals(const UNiagaraDataInterface* Othe
 	const UNiagaraDataInterfaceSkeletalMesh* OtherTyped = CastChecked<const UNiagaraDataInterfaceSkeletalMesh>(Other);
 	return OtherTyped->Source == Source &&
 #if WITH_EDITORONLY_DATA
-		OtherTyped->DefaultMesh == DefaultMesh &&
+		OtherTyped->PreviewMesh == PreviewMesh &&
 #endif
 		OtherTyped->MeshUserParameter == MeshUserParameter &&
 		OtherTyped->SkinningMode == SkinningMode &&
@@ -1964,11 +1963,11 @@ TArray<FNiagaraDataInterfaceError> UNiagaraDataInterfaceSkeletalMesh::GetErrors(
 	
 	// Collect Errors
 #if WITH_EDITORONLY_DATA
-	if (DefaultMesh != nullptr)
+	if (PreviewMesh != nullptr)
 	{
 		if (bRequiresCPUAccess)
 		{
-			for (const FSkeletalMeshLODInfo& LODInfo : DefaultMesh->GetLODInfoArray())
+			for (const FSkeletalMeshLODInfo& LODInfo : PreviewMesh->GetLODInfoArray())
 			{
 				if (!LODInfo.bAllowCPUAccess)
 				{
@@ -1986,12 +1985,12 @@ TArray<FNiagaraDataInterfaceError> UNiagaraDataInterfaceSkeletalMesh::GetErrors(
 	// Report Errors
 	if (Source == nullptr && bHasCPUAccessError)
 	{
-		FNiagaraDataInterfaceError CPUAccessNotAllowedError(FText::Format(LOCTEXT("CPUAccessNotAllowedError", "This mesh needs CPU access in order to be used properly.({0})"), FText::FromString(DefaultMesh->GetName())),
+		FNiagaraDataInterfaceError CPUAccessNotAllowedError(FText::Format(LOCTEXT("CPUAccessNotAllowedError", "This mesh needs CPU access in order to be used properly.({0})"), FText::FromString(PreviewMesh->GetName())),
 			LOCTEXT("CPUAccessNotAllowedErrorSummary", "CPU access error"),
 			FNiagaraDataInterfaceFix::CreateLambda([=]()
 		{
-			DefaultMesh->Modify();
-			for (FSkeletalMeshLODInfo& LODInfo : DefaultMesh->GetLODInfoArray())
+			PreviewMesh->Modify();
+			for (FSkeletalMeshLODInfo& LODInfo : PreviewMesh->GetLODInfoArray())
 			{
 				LODInfo.bAllowCPUAccess = true;
 			}
diff --git a/Engine/Plugins/FX/Niagara/Source/Niagara/Private/NiagaraDataInterfaceStaticMesh.cpp b/Engine/Plugins/FX/Niagara/Source/Niagara/Private/NiagaraDataInterfaceStaticMesh.cpp
--- a/Engine/Plugins/FX/Niagara/Source/Niagara/Private/NiagaraDataInterfaceStaticMesh.cpp
+++ b/Engine/Plugins/FX/Niagara/Source/Niagara/Private/NiagaraDataInterfaceStaticMesh.cpp
@@ -250,6 +250,13 @@ bool FNDIStaticMesh_InstanceData::Init(UNiagaraDataInterfaceStaticMesh* Interfac
 		Mesh = Interface->DefaultMesh;
 	}
 
+#if WITH_EDITORONLY_DATA
+	if (!Mesh && Interface->PreviewMesh)
+	{
+		Mesh = Interface->PreviewMesh;
+	}
+#endif
+
 	if (SafeComponent_GT.IsValid() && Mesh)
 	{
 		PrevTransform = Transform;
@@ -662,6 +669,9 @@ void FNiagaraDataInterfaceProxyStaticMesh::ConsumePerInstanceDataFromGameThread(
 
 UNiagaraDataInterfaceStaticMesh::UNiagaraDataInterfaceStaticMesh(FObjectInitializer const& ObjectInitializer)
 	: Super(ObjectInitializer)
+#if WITH_EDITORONLY_DATA
+	, PreviewMesh(nullptr)
+#endif
 	, DefaultMesh(nullptr)
 	, Source(nullptr)	
 	, ChangeId(0)
@@ -1024,7 +1034,7 @@ void UNiagaraDataInterfaceStaticMesh::GetVMExternalFunction(const FVMExternalFun
 {
 	FNDIStaticMesh_InstanceData* InstData = (FNDIStaticMesh_InstanceData*)InstanceData;
 	check(InstData && InstData->Mesh && InstData->SafeComponent_GT.IsValid());
-
+	
 	if (BindingInfo.Name == StaticMeshHelpers::RandomSectionName)
 	{
 		check(BindingInfo.GetNumInputs() == 1 && BindingInfo.GetNumOutputs() == 1);
@@ -1129,6 +1139,9 @@ bool UNiagaraDataInterfaceStaticMesh::CopyToInternal(UNiagaraDataInterface* Dest
 	UNiagaraDataInterfaceStaticMesh* OtherTyped = CastChecked<UNiagaraDataInterfaceStaticMesh>(Destination);
 	OtherTyped->Source = Source;
 	OtherTyped->DefaultMesh = DefaultMesh;
+#if WITH_EDITORONLY_DATA
+	OtherTyped->PreviewMesh = PreviewMesh;
+#endif
 	//OtherTyped->bEnableVertexColorRangeSorting = bEnableVertexColorRangeSorting;//TODO: Vertex color filtering needs more work.
 	OtherTyped->SectionFilter = SectionFilter;
 	return true;
@@ -1240,6 +1253,23 @@ TArray<FNiagaraDataInterfaceError> UNiagaraDataInterfaceStaticMesh::GetErrors()
 
 		Errors.Add(CPUAccessNotAllowedError);
 	}
+
+#if WITH_EDITORONLY_DATA
+	if (Source == nullptr && PreviewMesh != nullptr && !PreviewMesh->bAllowCPUAccess)
+	{
+		FNiagaraDataInterfaceError CPUAccessNotAllowedError(FText::Format(LOCTEXT("CPUAccessNotAllowedError", "This mesh needs CPU access in order to be used properly.({0})"), FText::FromString(PreviewMesh->GetName())),
+			LOCTEXT("CPUAccessNotAllowedErrorSummary", "CPU access error"),
+			FNiagaraDataInterfaceFix::CreateLambda([=]()
+		{
+			PreviewMesh->Modify();
+			PreviewMesh->bAllowCPUAccess = true;
+			return true;
+		}));
+
+		Errors.Add(CPUAccessNotAllowedError);
+	}
+#endif
+
 	return Errors;
 }
 #endif
