- Fix static analysis warnings

#rnx

diff --git a/Engine/Plugins/FX/Niagara/Source/Niagara/Private/NiagaraComponent.cpp b/Engine/Plugins/FX/Niagara/Source/Niagara/Private/NiagaraComponent.cpp
--- a/Engine/Plugins/FX/Niagara/Source/Niagara/Private/NiagaraComponent.cpp
+++ b/Engine/Plugins/FX/Niagara/Source/Niagara/Private/NiagaraComponent.cpp
@@ -587,6 +587,7 @@ void UNiagaraComponent::TickComponent(float DeltaSeconds, enum ELevelTick TickTy
 	check(SystemInstance->IsSolo());
 	if (IsActive() && SystemInstance.Get() && !SystemInstance->IsComplete())
 	{
+		check(Asset != nullptr);
 		Asset->AddToInstanceCountStat(1, true);
 		INC_DWORD_STAT_BY(STAT_TotalNiagaraSystemInstances, 1);
 		INC_DWORD_STAT_BY(STAT_TotalNiagaraSystemInstancesSolo, 1);
