Niagara - Fix issues with scratch pad copy/paste and add navigation options.
+ Fixed bugs with copy/paste including crash when pasting to a different project.
+ Added an icon to the tab, and used that same icon on a button which appears on scratchpad modules and inputs which can be used to navigate to the scratch pad script.
+ Double clicking on a scratch input or scratch module in the selection stack selects the scratch pad script.
+ Added a discard changes option to the right click menu.
+ Added a "Crerate Asset" option to the right click menu.
+ Added a "Find Next Usage" option to the right click menu which will select usages of the scratch pad script in the seleciton stack.

diff --git a/Engine/Plugins/FX/Niagara/Source/NiagaraEditor/Private/NiagaraClipboard.cpp b/Engine/Plugins/FX/Niagara/Source/NiagaraEditor/Private/NiagaraClipboard.cpp
--- a/Engine/Plugins/FX/Niagara/Source/NiagaraEditor/Private/NiagaraClipboard.cpp
+++ b/Engine/Plugins/FX/Niagara/Source/NiagaraEditor/Private/NiagaraClipboard.cpp
@@ -2,6 +2,7 @@
 
 #include "NiagaraClipboard.h"
 #include "NiagaraDataInterface.h"
+#include "NiagaraScript.h"
 
 #include "Factories.h"
 #include "UObject/UObjectMarks.h"
diff --git a/Engine/Plugins/FX/Niagara/Source/NiagaraEditor/Private/NiagaraEditorCommands.cpp b/Engine/Plugins/FX/Niagara/Source/NiagaraEditor/Private/NiagaraEditorCommands.cpp
--- a/Engine/Plugins/FX/Niagara/Source/NiagaraEditor/Private/NiagaraEditorCommands.cpp
+++ b/Engine/Plugins/FX/Niagara/Source/NiagaraEditor/Private/NiagaraEditorCommands.cpp
@@ -8,9 +8,13 @@
 
 void FNiagaraEditorCommands::RegisterCommands()
 {
-	UI_COMMAND(Apply, "Apply", "Apply unsaved changes to the current asset.", EUserInterfaceActionType::Button, FInputChord());
+	UI_COMMAND(Apply, "Apply", "Apply unsaved changes to the current object.", EUserInterfaceActionType::Button, FInputChord());
+	UI_COMMAND(Discard, "Discard", "Discard unsaved changes from the current object.", EUserInterfaceActionType::Button, FInputChord());
 	UI_COMMAND(Compile, "Compile", "Compile the current scripts", EUserInterfaceActionType::Button, FInputChord());
 	UI_COMMAND(RefreshNodes, "Refresh", "Refreshes the current graph nodes, and updates pins due to external changes.", EUserInterfaceActionType::Button, FInputChord());
+	UI_COMMAND(SelectNextUsage, "Select Next Usage", "Selects the next usage of the selected item.", EUserInterfaceActionType::Button, FInputChord());
+	UI_COMMAND(CreateAssetFromSelection, "Create Asset...", "Creates an asset from the current selection.", EUserInterfaceActionType::Button, FInputChord());
+
 	UI_COMMAND(ResetSimulation, "Reset", "Resets the current simulation", EUserInterfaceActionType::Button, FInputChord(EKeys::SpaceBar));
 	UI_COMMAND(TogglePreviewGrid, "Grid", "Toggles the preview pane's grid.", EUserInterfaceActionType::ToggleButton, FInputChord() );
 	UI_COMMAND(ToggleInstructionCounts, "InstructionCounts", "Display Instruction Counts", EUserInterfaceActionType::ToggleButton, FInputChord());
diff --git a/Engine/Plugins/FX/Niagara/Source/NiagaraEditor/Private/NiagaraEditorStyle.cpp b/Engine/Plugins/FX/Niagara/Source/NiagaraEditor/Private/NiagaraEditorStyle.cpp
--- a/Engine/Plugins/FX/Niagara/Source/NiagaraEditor/Private/NiagaraEditorStyle.cpp
+++ b/Engine/Plugins/FX/Niagara/Source/NiagaraEditor/Private/NiagaraEditorStyle.cpp
@@ -237,6 +237,8 @@ TSharedRef< FSlateStyleSet > FNiagaraEditorStyle::Create()
 
 	// Icons
 	Style->Set("NiagaraEditor.Isolate", new IMAGE_PLUGIN_BRUSH("Icons/Isolate", Icon16x16));
+
+	Style->Set("NiagaraEditor.Scratch", new IMAGE_PLUGIN_BRUSH("Icons/Scratch", Icon16x16, FLinearColor::Yellow));
 	
 	// Emitter details customization
 	Style->Set("NiagaraEditor.MaterialWarningBorder", new BOX_CORE_BRUSH("Common/GroupBorderLight", FMargin(4.0f / 16.0f)));
diff --git a/Engine/Plugins/FX/Niagara/Source/NiagaraEditor/Private/NiagaraEditorUtilities.cpp b/Engine/Plugins/FX/Niagara/Source/NiagaraEditor/Private/NiagaraEditorUtilities.cpp
--- a/Engine/Plugins/FX/Niagara/Source/NiagaraEditor/Private/NiagaraEditorUtilities.cpp
+++ b/Engine/Plugins/FX/Niagara/Source/NiagaraEditor/Private/NiagaraEditorUtilities.cpp
@@ -2414,4 +2414,16 @@ FString FNiagaraEditorUtilities::GetNamespacelessVariableNameString(const FName&
 	return VarNameString;
 }
 
+void FNiagaraEditorUtilities::GetReferencingFunctionCallNodes(UNiagaraScript* Script, TArray<UNiagaraNodeFunctionCall*>& OutReferencingFunctionCallNodes)
+{
+	for (TObjectIterator<UNiagaraNodeFunctionCall> It; It; ++It)
+	{
+		UNiagaraNodeFunctionCall* FunctionCallNode = *It;
+		if (FunctionCallNode->FunctionScript == Script)
+		{
+			OutReferencingFunctionCallNodes.Add(FunctionCallNode);
+		}
+	}
+}
+
 #undef LOCTEXT_NAMESPACE
diff --git a/Engine/Plugins/FX/Niagara/Source/NiagaraEditor/Private/Toolkits/NiagaraSystemToolkit.cpp b/Engine/Plugins/FX/Niagara/Source/NiagaraEditor/Private/Toolkits/NiagaraSystemToolkit.cpp
--- a/Engine/Plugins/FX/Niagara/Source/NiagaraEditor/Private/Toolkits/NiagaraSystemToolkit.cpp
+++ b/Engine/Plugins/FX/Niagara/Source/NiagaraEditor/Private/Toolkits/NiagaraSystemToolkit.cpp
@@ -696,6 +696,7 @@ TSharedRef<SDockTab> FNiagaraSystemToolkit::SpawnTab_ScratchPad(const FSpawnTabA
 {
 	TSharedRef<SDockTab> SpawnedTab = SNew(SDockTab)
 		.Label(LOCTEXT("ScratchPadTabLabel", "Scratch Pad"))
+		.Icon(FNiagaraEditorStyle::Get().GetBrush("NiagaraEditor.Scratch"))
 		[
 			FNiagaraEditorModule::Get().GetWidgetProvider()->CreateScriptScratchPad(*SystemViewModel->GetScriptScratchPadViewModel())
 		];
@@ -1379,7 +1380,7 @@ bool FNiagaraSystemToolkit::OnRequestClose()
 	bool bHasUnappliedScratchPadChanges = false;
 	for (TSharedRef<FNiagaraScratchPadScriptViewModel> ScratchPadViewModel : SystemViewModel->GetScriptScratchPadViewModel()->GetScriptViewModels())
 	{
-		if (ScratchPadViewModel->CanApplyChanges())
+		if (ScratchPadViewModel->HasUnappliedChanges())
 		{
 			bHasUnappliedScratchPadChanges = true;
 			break;
diff --git a/Engine/Plugins/FX/Niagara/Source/NiagaraEditor/Private/ViewModels/NiagaraScratchPadScriptViewModel.cpp b/Engine/Plugins/FX/Niagara/Source/NiagaraEditor/Private/ViewModels/NiagaraScratchPadScriptViewModel.cpp
--- a/Engine/Plugins/FX/Niagara/Source/NiagaraEditor/Private/ViewModels/NiagaraScratchPadScriptViewModel.cpp
+++ b/Engine/Plugins/FX/Niagara/Source/NiagaraEditor/Private/ViewModels/NiagaraScratchPadScriptViewModel.cpp
@@ -78,18 +78,6 @@ void FNiagaraScratchPadScriptViewModel::SetIsPendingRename(bool bInIsPendingRena
 	bIsPendingRename = bInIsPendingRename;
 }
 
-void GetReferencingFunctionCallNodes(UNiagaraScript* Script, TArray<UNiagaraNodeFunctionCall*>& OutReferencingFunctionCallNodes)
-{
-	for (TObjectIterator<UNiagaraNodeFunctionCall> It; It; ++It)
-	{
-		UNiagaraNodeFunctionCall* FunctionCallNode = *It;
-		if (FunctionCallNode->FunctionScript == Script)
-		{
-			OutReferencingFunctionCallNodes.Add(FunctionCallNode);
-		}
-	}
-}
-
 void FNiagaraScratchPadScriptViewModel::SetScriptName(FText InScriptName)
 {
 	FString NewName = ObjectTools::SanitizeObjectName(InScriptName.ToString());
@@ -102,7 +90,7 @@ void FNiagaraScratchPadScriptViewModel::SetScriptName(FText InScriptName)
 		OriginalScript->Rename(*NewUniqueName.ToString(), nullptr, REN_DontCreateRedirectors);
 
 		TArray<UNiagaraNodeFunctionCall*> ReferencingFunctionCallNodes;
-		GetReferencingFunctionCallNodes(OriginalScript, ReferencingFunctionCallNodes);
+		FNiagaraEditorUtilities::GetReferencingFunctionCallNodes(OriginalScript, ReferencingFunctionCallNodes);
 		for(UNiagaraNodeFunctionCall* ReferencingFunctionCallNode : ReferencingFunctionCallNodes)
 		{
 			ReferencingFunctionCallNode->Modify();
@@ -145,7 +133,7 @@ void FNiagaraScratchPadScriptViewModel::SetEditorHeight(float InEditorHeight)
 	EditorHeight = InEditorHeight;
 }
 
-bool FNiagaraScratchPadScriptViewModel::CanApplyChanges() const
+bool FNiagaraScratchPadScriptViewModel::HasUnappliedChanges() const
 {
 	return bHasPendingChanges;
 }
@@ -161,7 +149,7 @@ void FNiagaraScratchPadScriptViewModel::ApplyChanges()
 	bHasPendingChanges = false;
 
 	TArray<UNiagaraNodeFunctionCall*> FunctionCallNodesToRefresh;
-	GetReferencingFunctionCallNodes(OriginalScript, FunctionCallNodesToRefresh);
+	FNiagaraEditorUtilities::GetReferencingFunctionCallNodes(OriginalScript, FunctionCallNodesToRefresh);
 
 	TArray<UNiagaraStackFunctionInputCollection*> InputCollectionsToRefresh;
 	if (FunctionCallNodesToRefresh.Num())
@@ -188,6 +176,11 @@ void FNiagaraScratchPadScriptViewModel::ApplyChanges()
 	}
 }
 
+void FNiagaraScratchPadScriptViewModel::DiscardChanges()
+{
+	OnRequestDiscardChangesDelegate.ExecuteIfBound();
+}
+
 FNiagaraScratchPadScriptViewModel::FOnRenamed& FNiagaraScratchPadScriptViewModel::OnRenamed()
 {
 	return OnRenamedDelegate;
@@ -198,6 +191,11 @@ FNiagaraScratchPadScriptViewModel::FOnPinnedChanged& FNiagaraScratchPadScriptVie
 	return OnPinnedChangedDelegate;
 }
 
+FSimpleDelegate& FNiagaraScratchPadScriptViewModel::OnRequestDiscardChanges()
+{
+	return OnRequestDiscardChangesDelegate;
+}
+
 FText FNiagaraScratchPadScriptViewModel::GetDisplayNameInternal() const
 {
 	return FText::Format(LOCTEXT("DisplayNameFormat", "{0}{1}"), FText::FromString(OriginalScript->GetName()), bHasPendingChanges ? LOCTEXT("HasPendingChanges", "*") : FText());
diff --git a/Engine/Plugins/FX/Niagara/Source/NiagaraEditor/Private/ViewModels/NiagaraScratchPadViewModel.cpp b/Engine/Plugins/FX/Niagara/Source/NiagaraEditor/Private/ViewModels/NiagaraScratchPadViewModel.cpp
--- a/Engine/Plugins/FX/Niagara/Source/NiagaraEditor/Private/ViewModels/NiagaraScratchPadViewModel.cpp
+++ b/Engine/Plugins/FX/Niagara/Source/NiagaraEditor/Private/ViewModels/NiagaraScratchPadViewModel.cpp
@@ -2,8 +2,11 @@
 
 #include "ViewModels/NiagaraScratchPadViewModel.h"
 #include "ViewModels/NiagaraSystemViewModel.h"
+#include "ViewModels/NiagaraSystemSelectionViewModel.h"
 #include "ViewModels/NiagaraScratchPadScriptViewModel.h"
 #include "ViewModels/NiagaraScriptGraphViewModel.h"
+#include "ViewModels/Stack/NiagaraStackModuleItem.h"
+#include "ViewModels/Stack/NiagaraStackFunctionInput.h"
 #include "NiagaraSystem.h"
 #include "NiagaraScriptSource.h"
 #include "NiagaraNodeOutput.h"
@@ -12,8 +15,16 @@
 #include "NiagaraNodeFunctionCall.h"
 #include "NiagaraEditorModule.h"
 #include "NiagaraClipboard.h"
+#include "Toolkits/NiagaraSystemToolkit.h"
+#include "NiagaraEditorUtilities.h"
 
 #include "ScopedTransaction.h"
+#include "Modules/ModuleManager.h"
+#include "AssetToolsModule.h"
+#include "Misc/PackageName.h"
+#include "Subsystems/AssetEditorSubsystem.h"
+#include "Misc/MessageDialog.h"
+#include "GenericPlatform/GenericPlatformMisc.h"
 
 #define LOCTEXT_NAMESPACE "NiagaraScratchPadViewModel"
 
@@ -82,11 +93,7 @@ void UNiagaraScratchPadViewModel::RefreshScriptViewModels()
 			}
 			else
 			{
-				ScriptViewModel = MakeShared<FNiagaraScratchPadScriptViewModel>();
-				ScriptViewModel->Initialize(ScratchPadScript);
-				ScriptViewModel->GetGraphViewModel()->GetNodeSelection()->OnSelectedObjectsChanged().AddUObject(this, &UNiagaraScratchPadViewModel::ScriptGraphNodeSelectionChanged, TWeakPtr<FNiagaraScratchPadScriptViewModel>(ScriptViewModel));
-				ScriptViewModel->OnRenamed().AddUObject(this, &UNiagaraScratchPadViewModel::ScriptViewModelScriptRenamed);
-				ScriptViewModel->OnPinnedChanged().AddUObject(this, &UNiagaraScratchPadViewModel::ScriptViewModelPinnedChanged, TWeakPtr<FNiagaraScratchPadScriptViewModel>(ScriptViewModel));
+				ScriptViewModel = CreateAndSetupScriptviewModel(ScratchPadScript);
 				bViewModelsChanged = true;
 			}
 			ScriptViewModels.Add(ScriptViewModel.ToSharedRef());
@@ -197,6 +204,15 @@ void UNiagaraScratchPadViewModel::SetActiveScriptViewModel(TSharedRef<FNiagaraSc
 	}
 }
 
+void UNiagaraScratchPadViewModel::FocusScratchPadScriptViewModel(TSharedRef<FNiagaraScratchPadScriptViewModel> InScriptViewModel)
+{
+	if (ensureMsgf(ScriptViewModels.Contains(InScriptViewModel), TEXT("Can only focus a view model from this scratch pad view model.")))
+	{
+		SetActiveScriptViewModel(InScriptViewModel);
+		GetSystemViewModel()->FocusTab(FNiagaraSystemToolkit::ScratchPadTabID);
+	}
+}
+
 void UNiagaraScratchPadViewModel::ResetActiveScriptViewModel()
 {
 	if (ActiveScriptViewModel.IsValid())
@@ -342,6 +358,107 @@ TSharedPtr<FNiagaraScratchPadScriptViewModel> UNiagaraScratchPadViewModel::Creat
 	return GetViewModelForScript(NewScript);
 }
 
+void UNiagaraScratchPadViewModel::CreateAssetFromActiveScript()
+{
+	if (ActiveScriptViewModel.IsValid())
+	{
+		if (ActiveScriptViewModel->HasUnappliedChanges())
+		{
+			FText Title = LOCTEXT("ScriptHasUnappliedchangesLabel", "Apply Changes?");
+			EAppReturnType::Type DialogResult = FMessageDialog::Open(EAppMsgType::YesNoCancel, EAppReturnType::Cancel,
+				LOCTEXT("ScriptHasUnappliedChangesMessage", "The selected scratch pad script has unapplied changes.\nWould you like to apply the changes before saving?\n"),
+				&Title);
+			if (DialogResult == EAppReturnType::Cancel)
+			{
+				return;
+			}
+			else if(DialogResult == EAppReturnType::Yes)
+			{
+				ActiveScriptViewModel->ApplyChanges();
+			}
+		}
+
+		UNiagaraScript* ScriptToCopy = ActiveScriptViewModel->GetOriginalScript();
+		const FString StartingPath = FPackageName::GetLongPackagePath(ScriptToCopy->GetOutermost()->GetName());
+		FAssetToolsModule& AssetToolsModule = FModuleManager::LoadModuleChecked<FAssetToolsModule>("AssetTools");
+		UNiagaraScript* NewAssetScript = Cast<UNiagaraScript>(AssetToolsModule.Get().DuplicateAssetWithDialogAndTitle(
+			ScriptToCopy->GetName(), StartingPath, ScriptToCopy, LOCTEXT("CreateScriptAssetTitle", "Create Script As")));
+		if (NewAssetScript != nullptr)
+		{
+			NewAssetScript->SetFlags(RF_Standalone | RF_Public);
+			GEditor->GetEditorSubsystem<UAssetEditorSubsystem>()->OpenEditorForAsset(NewAssetScript);
+		}
+	}
+}
+
+bool UNiagaraScratchPadViewModel::CanSelectNextUsageForActiveScript()
+{
+	if (ActiveScriptViewModel.IsValid())
+	{
+		TArray<UNiagaraNodeFunctionCall*> ReferencingFunctionCallNodes;
+		FNiagaraEditorUtilities::GetReferencingFunctionCallNodes(ActiveScriptViewModel->GetOriginalScript(), ReferencingFunctionCallNodes);
+		return ReferencingFunctionCallNodes.Num() > 0;
+	}
+	return false;
+}
+
+void UNiagaraScratchPadViewModel::SelectNextUsageForActiveScript()
+{
+	if (ActiveScriptViewModel.IsValid())
+	{
+		UNiagaraScript* OriginalScript = ActiveScriptViewModel->GetOriginalScript();
+		TArray<UNiagaraStackEntry*> ReferencingSelectableEntries;
+		if (OriginalScript->GetUsage() == ENiagaraScriptUsage::Module)
+		{
+			for (TObjectIterator<UNiagaraStackModuleItem> It; It; ++It)
+			{
+				UNiagaraStackModuleItem* ModuleItem = *It;
+				if (ModuleItem->GetModuleNode().FunctionScript == OriginalScript)
+				{
+					ReferencingSelectableEntries.Add(ModuleItem);
+				}
+			}
+		}
+		else if (OriginalScript->GetUsage() == ENiagaraScriptUsage::DynamicInput)
+		{
+			for (TObjectIterator<UNiagaraStackFunctionInput> It; It; ++It)
+			{
+				UNiagaraStackFunctionInput* FunctionInput = *It;
+				if (FunctionInput->GetValueMode() == UNiagaraStackFunctionInput::EValueMode::Dynamic &&
+					FunctionInput->GetDynamicInputNode()->FunctionScript == OriginalScript)
+				{
+					UNiagaraStackModuleItem* OwningModuleItem = FunctionInput->GetTypedOuter<UNiagaraStackModuleItem>();
+					if (OwningModuleItem != nullptr)
+					{
+						ReferencingSelectableEntries.Add(OwningModuleItem);
+					}
+				}
+			}
+		}
+
+		if (ReferencingSelectableEntries.Num() > 0)
+		{
+			int32 CurrentIndex = INDEX_NONE;
+			TArray<UNiagaraStackEntry*> SelectedEntries;
+			GetSystemViewModel()->GetSelectionViewModel()->GetSelectedEntries(SelectedEntries);
+			if (SelectedEntries.Num() == 1)
+			{
+				CurrentIndex = ReferencingSelectableEntries.IndexOfByKey(SelectedEntries[0]);
+			}
+
+			CurrentIndex++;
+			if (CurrentIndex >= ReferencingSelectableEntries.Num())
+			{
+				CurrentIndex = 0;
+			}
+
+			TArray<UNiagaraStackEntry*> NewSelectedEntries;
+			NewSelectedEntries.Add(ReferencingSelectableEntries[CurrentIndex]);
+			GetSystemViewModel()->GetSelectionViewModel()->UpdateSelectedEntries(NewSelectedEntries, SelectedEntries, true);
+		}
+	}
+}
+
 UNiagaraScratchPadViewModel::FOnScriptViewModelsChanged& UNiagaraScratchPadViewModel::OnScriptViewModelsChanged()
 {
 	return OnScriptViewModelsChangedDelegate;
@@ -374,6 +491,17 @@ TSharedRef<FNiagaraSystemViewModel> UNiagaraScratchPadViewModel::GetSystemViewMo
 	return SystemViewModel.ToSharedRef();
 }
 
+TSharedRef<FNiagaraScratchPadScriptViewModel> UNiagaraScratchPadViewModel::CreateAndSetupScriptviewModel(UNiagaraScript* ScratchPadScript)
+{
+	TSharedRef<FNiagaraScratchPadScriptViewModel> ScriptViewModel = MakeShared<FNiagaraScratchPadScriptViewModel>();
+	ScriptViewModel->Initialize(ScratchPadScript);
+	ScriptViewModel->GetGraphViewModel()->GetNodeSelection()->OnSelectedObjectsChanged().AddUObject(this, &UNiagaraScratchPadViewModel::ScriptGraphNodeSelectionChanged, TWeakPtr<FNiagaraScratchPadScriptViewModel>(ScriptViewModel));
+	ScriptViewModel->OnRenamed().AddUObject(this, &UNiagaraScratchPadViewModel::ScriptViewModelScriptRenamed);
+	ScriptViewModel->OnPinnedChanged().AddUObject(this, &UNiagaraScratchPadViewModel::ScriptViewModelPinnedChanged, TWeakPtr<FNiagaraScratchPadScriptViewModel>(ScriptViewModel));
+	ScriptViewModel->OnRequestDiscardChanges().BindUObject(this, &UNiagaraScratchPadViewModel::ScriptViewModelRequestDiscardChanges, TWeakPtr<FNiagaraScratchPadScriptViewModel>(ScriptViewModel));
+	return ScriptViewModel;
+}
+
 void UNiagaraScratchPadViewModel::RefreshEditScriptViewModels()
 {
 	EditScriptViewModels.Empty();
@@ -455,4 +583,35 @@ void UNiagaraScratchPadViewModel::ScriptViewModelPinnedChanged(TWeakPtr<FNiagara
 	}
 }
 
+void UNiagaraScratchPadViewModel::ScriptViewModelRequestDiscardChanges(TWeakPtr<FNiagaraScratchPadScriptViewModel> ScriptViewModelWeak)
+{
+	if (ActiveScriptViewModel.IsValid() && ActiveScriptViewModel->HasUnappliedChanges())
+	{
+		FText Title = LOCTEXT("DiscardChangesTitle", "Discard Changes?");
+		EAppReturnType::Type DialogResult = FMessageDialog::Open(EAppMsgType::YesNo, EAppReturnType::No,
+			LOCTEXT("DiscardChangesMessage", "Are you sure you want to discard changes?\nThis operation can not be undone."),
+			&Title);
+		if (DialogResult == EAppReturnType::No)
+		{
+			return;
+		}
+
+		int32 ViewModelIndex = ScriptViewModels.IndexOfByKey(ActiveScriptViewModel.ToSharedRef());
+		if (ensureMsgf(ViewModelIndex != INDEX_NONE, TEXT("Active script view model wasn't in script view model collection!")))
+		{
+			TSharedRef<FNiagaraScratchPadScriptViewModel> NewScriptViewModel = CreateAndSetupScriptviewModel(ActiveScriptViewModel->GetOriginalScript());
+			ScriptViewModels[ViewModelIndex] = NewScriptViewModel;
+
+			int32 PinnedViewModelIndex = PinnedScriptViewModels.IndexOfByKey(ActiveScriptViewModel.ToSharedRef());
+			if (PinnedViewModelIndex != INDEX_NONE)
+			{
+				PinnedScriptViewModels[PinnedViewModelIndex] = NewScriptViewModel;
+			}
+
+			OnScriptViewModelsChangedDelegate.Broadcast();
+			SetActiveScriptViewModel(NewScriptViewModel);
+		}
+	}
+}
+
 #undef LOCTEXT_NAMESPACE
\ No newline at end of file
diff --git a/Engine/Plugins/FX/Niagara/Source/NiagaraEditor/Private/ViewModels/Stack/NiagaraStackFunctionInput.cpp b/Engine/Plugins/FX/Niagara/Source/NiagaraEditor/Private/ViewModels/Stack/NiagaraStackFunctionInput.cpp
--- a/Engine/Plugins/FX/Niagara/Source/NiagaraEditor/Private/ViewModels/Stack/NiagaraStackFunctionInput.cpp
+++ b/Engine/Plugins/FX/Niagara/Source/NiagaraEditor/Private/ViewModels/Stack/NiagaraStackFunctionInput.cpp
@@ -293,10 +293,19 @@ bool UNiagaraStackFunctionInput::TestCanPasteWithMessage(const UNiagaraClipboard
 	}
 	else if (ClipboardContent->FunctionInputs.Num() == 1)
 	{
-		if (ClipboardContent->FunctionInputs[0] != nullptr)
+		const UNiagaraClipboardFunctionInput* ClipboardFunctionInput = ClipboardContent->FunctionInputs[0];
+		if (ClipboardFunctionInput != nullptr)
 		{
-			if (ClipboardContent->FunctionInputs[0]->InputType == InputType)
+			if (ClipboardFunctionInput->InputType == InputType)
 			{
+				if (ClipboardFunctionInput->ValueMode == ENiagaraClipboardFunctionInputValueMode::Dynamic)
+				{
+					if (ClipboardFunctionInput->Dynamic->Script.IsValid() == false)
+					{
+						OutMessage = LOCTEXT("CantPasteInvalidDynamicInputScript", "Can not paste the dynamic input because its script is no longer valid.");
+						return false;
+					}
+				}
 				OutMessage = LOCTEXT("PastMessage", "Paste the input from the clipboard here.");
 				return true;
 			}
@@ -755,6 +764,7 @@ void UNiagaraStackFunctionInput::RefreshValues()
 	bCanResetCache.Reset();
 	bCanResetToBaseCache.Reset();
 	ValueToolTipCache.Reset();
+	bIsScratchDynamicInputCache.Reset();
 	ValueChangedDelegate.Broadcast();
 }
 
@@ -1164,8 +1174,7 @@ void UNiagaraStackFunctionInput::SetScratch()
 	if (ScratchScriptViewModel.IsValid())
 	{
 		SetDynamicInput(ScratchScriptViewModel->GetOriginalScript());
-		GetSystemViewModel()->FocusTab(FNiagaraSystemToolkit::ScratchPadTabID);
-		GetSystemViewModel()->GetScriptScratchPadViewModel()->SetActiveScriptViewModel(ScratchScriptViewModel.ToSharedRef());
+		GetSystemViewModel()->GetScriptScratchPadViewModel()->FocusScratchPadScriptViewModel(ScratchScriptViewModel.ToSharedRef());
 		ScratchScriptViewModel->SetIsPendingRename(true);
 	}
 }
@@ -1825,18 +1834,31 @@ void UNiagaraStackFunctionInput::SetValueFromClipboardFunctionInput(const UNiaga
 			break;
 		case ENiagaraClipboardFunctionInputValueMode::Dynamic:
 			if (ensureMsgf(ClipboardFunctionInput.Dynamic->ScriptMode == ENiagaraClipboardFunctionScriptMode::ScriptAsset,
-				TEXT("Can not set dynamic input value from clipboard, only script asset funcitons can be set.")))
+				TEXT("Dynamic input values can only be set from script asset clipboard functions.")))
 			{
-				UNiagaraScript* ClipboardScript = ClipboardFunctionInput.Dynamic->Script->IsAsset()
-					? ClipboardFunctionInput.Dynamic->Script
-					: GetSystemViewModel()->GetScriptScratchPadViewModel()->CreateNewScriptAsDuplicate(ClipboardFunctionInput.Dynamic->Script)->GetOriginalScript();
-				SetDynamicInput(ClipboardScript, ClipboardFunctionInput.Dynamic->FunctionName);
-
-				TArray<UNiagaraStackFunctionInputCollection*> DynamicInputCollections;
-				GetUnfilteredChildrenOfType(DynamicInputCollections);
-				for (UNiagaraStackFunctionInputCollection* DynamicInputCollection : DynamicInputCollections)
+				UNiagaraScript* ClipboardFunctionScript = ClipboardFunctionInput.Dynamic->Script.Get();
+				if (ClipboardFunctionScript != nullptr)
 				{
-					DynamicInputCollection->SetValuesFromClipboardFunctionInputs(ClipboardFunctionInput.Dynamic->Inputs);
+					UNiagaraScript* NewDynamicInputScript;
+					if (ClipboardFunctionScript->IsAsset() ||
+						GetSystemViewModel()->GetScriptScratchPadViewModel()->GetViewModelForScript(ClipboardFunctionScript).IsValid())
+					{
+						// If the clipboard script is an asset, or it's in the scratch pad of the current asset, it can be used directly.
+						NewDynamicInputScript = ClipboardFunctionScript;
+					}
+					else
+					{
+						// Otherwise it's a scratch pad script from another asset so we need to add a duplicate scratch pad script to this asset.
+						NewDynamicInputScript = GetSystemViewModel()->GetScriptScratchPadViewModel()->CreateNewScriptAsDuplicate(ClipboardFunctionScript)->GetOriginalScript();
+					}
+					SetDynamicInput(NewDynamicInputScript, ClipboardFunctionInput.Dynamic->FunctionName);
+
+					TArray<UNiagaraStackFunctionInputCollection*> DynamicInputCollections;
+					GetUnfilteredChildrenOfType(DynamicInputCollections);
+					for (UNiagaraStackFunctionInputCollection* DynamicInputCollection : DynamicInputCollections)
+					{
+						DynamicInputCollection->SetValuesFromClipboardFunctionInputs(ClipboardFunctionInput.Dynamic->Inputs);
+					}
 				}
 			}
 			break;
@@ -1852,6 +1874,18 @@ void UNiagaraStackFunctionInput::SetValueFromClipboardFunctionInput(const UNiaga
 	}
 }
 
+bool UNiagaraStackFunctionInput::IsScratchDynamicInput() const
+{
+	if (bIsScratchDynamicInputCache.IsSet() == false)
+	{
+		bIsScratchDynamicInputCache = 
+			InputValues.Mode == EValueMode::Dynamic &&
+			InputValues.DynamicNode.IsValid() &&
+			GetSystemViewModel()->GetScriptScratchPadViewModel()->GetViewModelForScript(InputValues.DynamicNode->FunctionScript).IsValid();
+	}
+	return bIsScratchDynamicInputCache.GetValue();
+}
+
 void UNiagaraStackFunctionInput::GetSearchItems(TArray<FStackSearchItem>& SearchItems) const
 {
 	if (GetShouldPassFilterForVisibleCondition() && GetIsInlineEditConditionToggle() == false)
diff --git a/Engine/Plugins/FX/Niagara/Source/NiagaraEditor/Private/ViewModels/Stack/NiagaraStackModuleItem.cpp b/Engine/Plugins/FX/Niagara/Source/NiagaraEditor/Private/ViewModels/Stack/NiagaraStackModuleItem.cpp
--- a/Engine/Plugins/FX/Niagara/Source/NiagaraEditor/Private/ViewModels/Stack/NiagaraStackModuleItem.cpp
+++ b/Engine/Plugins/FX/Niagara/Source/NiagaraEditor/Private/ViewModels/Stack/NiagaraStackModuleItem.cpp
@@ -10,6 +10,8 @@
 #include "ViewModels/Stack/NiagaraStackFunctionInput.h"
 #include "ViewModels/Stack/NiagaraStackInputCategory.h"
 #include "ViewModels/Stack/NiagaraStackModuleItemOutputCollection.h"
+#include "ViewModels/NiagaraScratchPadViewModel.h"
+#include "ViewModels/NiagaraScratchPadScriptViewModel.h"
 #include "NiagaraNodeFunctionCall.h"
 #include "NiagaraNodeParameterMapSet.h"
 #include "NiagaraNodeOutput.h"
@@ -180,6 +182,7 @@ void UNiagaraStackModuleItem::RefreshChildrenInternal(const TArray<UNiagaraStack
 {
 	bCanRefresh = false;
 	bCanMoveAndDeleteCache.Reset();
+	bIsScratchModuleCache.Reset();
 
 	if (FunctionCallNode != nullptr && FunctionCallNode->ScriptIsValid())
 	{
@@ -1198,6 +1201,15 @@ void UNiagaraStackModuleItem::Delete()
 	}
 }
 
+bool UNiagaraStackModuleItem::IsScratchModule() const
+{
+	if (bIsScratchModuleCache.IsSet() == false)
+	{
+		bIsScratchModuleCache = GetSystemViewModel()->GetScriptScratchPadViewModel()->GetViewModelForScript(FunctionCallNode->FunctionScript).IsValid();
+	}
+	return bIsScratchModuleCache.GetValue();
+}
+
 UObject* UNiagaraStackModuleItem::GetExternalAsset() const
 {
 	if (GetModuleNode().FunctionScript != nullptr && GetModuleNode().FunctionScript->IsAsset())
diff --git a/Engine/Plugins/FX/Niagara/Source/NiagaraEditor/Private/ViewModels/Stack/NiagaraStackScriptItemGroup.cpp b/Engine/Plugins/FX/Niagara/Source/NiagaraEditor/Private/ViewModels/Stack/NiagaraStackScriptItemGroup.cpp
--- a/Engine/Plugins/FX/Niagara/Source/NiagaraEditor/Private/ViewModels/Stack/NiagaraStackScriptItemGroup.cpp
+++ b/Engine/Plugins/FX/Niagara/Source/NiagaraEditor/Private/ViewModels/Stack/NiagaraStackScriptItemGroup.cpp
@@ -281,8 +281,7 @@ public:
 			if (ScratchPadScriptViewModel.IsValid())
 			{
 				NewModuleNode = FNiagaraStackGraphUtilities::AddScriptModuleToStack(ScratchPadScriptViewModel->GetOriginalScript(), *OutputNode, TargetIndex);
-				SystemViewModel.Pin()->FocusTab(FNiagaraSystemToolkit::ScratchPadTabID);
-				SystemViewModel.Pin()->GetScriptScratchPadViewModel()->SetActiveScriptViewModel(ScratchPadScriptViewModel.ToSharedRef());
+				SystemViewModel.Pin()->GetScriptScratchPadViewModel()->FocusScratchPadScriptViewModel(ScratchPadScriptViewModel.ToSharedRef());
 				ScratchPadScriptViewModel->SetIsPendingRename(true);
 			}
 		}
@@ -379,19 +378,30 @@ bool UNiagaraStackScriptItemGroup::TestCanPasteWithMessage(const UNiagaraClipboa
 {
 	if (ClipboardContent->Functions.Num() > 0)
 	{
+		bool bValidFunction = false;
 		bool bValidUsage = false;
 		UNiagaraNodeOutput* OutputNode = GetScriptOutputNode();
 		for (const UNiagaraClipboardFunction* Function : ClipboardContent->Functions)
 		{
-			if (Function->ScriptMode == ENiagaraClipboardFunctionScriptMode::ScriptAsset)
+			if (Function == nullptr)
 			{
-				if (Function != nullptr && Function->Script != nullptr && Function->Script->GetSupportedUsageContexts().Contains(OutputNode->GetUsage()))
+				continue;
+			}
+			else if (Function->ScriptMode == ENiagaraClipboardFunctionScriptMode::ScriptAsset)
+			{
+				UNiagaraScript* ClipboardFunctionScript = Function->Script.Get();
+				if (ClipboardFunctionScript != nullptr)
 				{
-					bValidUsage = true;
+					bValidFunction = true;
+					if (ClipboardFunctionScript->GetSupportedUsageContexts().Contains(OutputNode->GetUsage()))
+					{
+						bValidUsage = true;
+					}
 				}
 			}
 			else if (Function->ScriptMode == ENiagaraClipboardFunctionScriptMode::Assignment)
 			{
+				bValidFunction = true;
 				for (const FNiagaraVariable& AssignmentTarget : Function->AssignmentTargets)
 				{
 					if (FNiagaraStackGraphUtilities::CanWriteParameterFromUsage(AssignmentTarget, OutputNode->GetUsage()))
@@ -402,14 +412,22 @@ bool UNiagaraStackScriptItemGroup::TestCanPasteWithMessage(const UNiagaraClipboa
 				}
 			}
 		}
-		if (bValidUsage)
+
+		if (bValidFunction && bValidUsage)
 		{
-			OutMessage = LOCTEXT("PasteModules", "Paste modules from the clipboard which have a valid usage.");
+			OutMessage = LOCTEXT("PasteModules", "Paste modules from the clipboard which have valid function scripts and valid usage.");
 			return true;
 		}
 		else
 		{
-			OutMessage = LOCTEXT("CantPasteModulesForUsage", "Can't paste the copied modules because they don't support the correct usage.");
+			if (bValidFunction == false)
+			{
+				OutMessage = LOCTEXT("CantPasteModulesInvalidFunction", "Can't paste the copied modules because they don't references scripts which are still valid.");
+			}
+			else
+			{
+				OutMessage = LOCTEXT("CantPasteModulesForUsage", "Can't paste the copied modules because they don't support the correct usage.");
+			}
 			return false;
 		}
 	}
@@ -1093,7 +1111,7 @@ void RenameInputsFromClipboard(TMap<FName, FName> OldModuleOutputNameToNewModule
 		}
 		else if (SourceInput->ValueMode == ENiagaraClipboardFunctionInputValueMode::Dynamic)
 		{
-			const UNiagaraClipboardFunctionInput* RenamedDynamicInput = UNiagaraClipboardFunctionInput::CreateDynamicValue(InOuter, SourceInput->InputName, SourceInput->InputType, bEditConditionValue, SourceInput->Dynamic->FunctionName, SourceInput->Dynamic->Script);
+			const UNiagaraClipboardFunctionInput* RenamedDynamicInput = UNiagaraClipboardFunctionInput::CreateDynamicValue(InOuter, SourceInput->InputName, SourceInput->InputType, bEditConditionValue, SourceInput->Dynamic->FunctionName, SourceInput->Dynamic->Script.Get());
 			RenameInputsFromClipboard(OldModuleOutputNameToNewModuleOutputNameMap, RenamedDynamicInput->Dynamic, SourceInput->Dynamic->Inputs, RenamedDynamicInput->Dynamic->Inputs);
 			OutRenamedInputs.Add(RenamedDynamicInput);
 		}
@@ -1144,12 +1162,22 @@ void UNiagaraStackScriptItemGroup::PasteModules(const UNiagaraClipboardContent*
 			}
 			case ENiagaraClipboardFunctionScriptMode::ScriptAsset:
 			{
-				if (ClipboardFunction->Script->GetSupportedUsageContexts().Contains(OutputNode->GetUsage()))
+				UNiagaraScript* ClipboardFunctionScript = ClipboardFunction->Script.Get();
+				if (ClipboardFunctionScript != nullptr && ClipboardFunctionScript->GetSupportedUsageContexts().Contains(OutputNode->GetUsage()))
 				{
-					UNiagaraScript* ClipboardScript = ClipboardFunction->Script->IsAsset()
-						? ClipboardFunction->Script
-						: GetSystemViewModel()->GetScriptScratchPadViewModel()->CreateNewScriptAsDuplicate(ClipboardFunction->Script)->GetOriginalScript();
-					NewFunctionCallNode = FNiagaraStackGraphUtilities::AddScriptModuleToStack(ClipboardFunction->Script, *OutputNode, CurrentPasteIndex);
+					UNiagaraScript* NewFunctionScript;
+					if (ClipboardFunctionScript->IsAsset() ||
+						GetSystemViewModel()->GetScriptScratchPadViewModel()->GetViewModelForScript(ClipboardFunctionScript).IsValid())
+					{
+						// If the clipboard script is an asset, or it's in the scratch pad of the current asset, it can be used directly.
+						NewFunctionScript = ClipboardFunctionScript;
+					}
+					else
+					{
+						// Otherwise it's a scratch pad script from another asset so we need to add a duplicate scratch pad script to this asset.
+						NewFunctionScript = GetSystemViewModel()->GetScriptScratchPadViewModel()->CreateNewScriptAsDuplicate(ClipboardFunctionScript)->GetOriginalScript();
+					}
+					NewFunctionCallNode = FNiagaraStackGraphUtilities::AddScriptModuleToStack(NewFunctionScript, *OutputNode, CurrentPasteIndex);
 				}
 				else
 				{
diff --git a/Engine/Plugins/FX/Niagara/Source/NiagaraEditor/Public/NiagaraClipboard.h b/Engine/Plugins/FX/Niagara/Source/NiagaraEditor/Public/NiagaraClipboard.h
--- a/Engine/Plugins/FX/Niagara/Source/NiagaraEditor/Public/NiagaraClipboard.h
+++ b/Engine/Plugins/FX/Niagara/Source/NiagaraEditor/Public/NiagaraClipboard.h
@@ -4,6 +4,7 @@
 
 #include "NiagaraTypes.h"
 #include "UObject/GCObject.h"
+#include "UObject/SoftObjectPtr.h"
 #include "NiagaraClipboard.generated.h"
 
 class UNiagaraDataInterface;
@@ -95,7 +96,7 @@ public:
 	ENiagaraClipboardFunctionScriptMode ScriptMode;
 
 	UPROPERTY()
-	UNiagaraScript* Script;
+	TSoftObjectPtr<UNiagaraScript> Script;
 
 	UPROPERTY()
 	TArray<FNiagaraVariable> AssignmentTargets;
diff --git a/Engine/Plugins/FX/Niagara/Source/NiagaraEditor/Public/NiagaraEditorCommands.h b/Engine/Plugins/FX/Niagara/Source/NiagaraEditor/Public/NiagaraEditorCommands.h
--- a/Engine/Plugins/FX/Niagara/Source/NiagaraEditor/Public/NiagaraEditorCommands.h
+++ b/Engine/Plugins/FX/Niagara/Source/NiagaraEditor/Public/NiagaraEditorCommands.h
@@ -24,9 +24,12 @@ public:
 	virtual void RegisterCommands() override;
 
 	TSharedPtr<FUICommandInfo> Apply;
+	TSharedPtr<FUICommandInfo> Discard;
 	TSharedPtr<FUICommandInfo> Compile;
 	TSharedPtr<FUICommandInfo> RefreshNodes;
 	TSharedPtr<FUICommandInfo> ResetSimulation;
+	TSharedPtr<FUICommandInfo> SelectNextUsage;
+	TSharedPtr<FUICommandInfo> CreateAssetFromSelection;
 
 	/** Toggles the preview pane's grid */
 	TSharedPtr<FUICommandInfo> TogglePreviewGrid;
diff --git a/Engine/Plugins/FX/Niagara/Source/NiagaraEditor/Public/NiagaraEditorUtilities.h b/Engine/Plugins/FX/Niagara/Source/NiagaraEditor/Public/NiagaraEditorUtilities.h
--- a/Engine/Plugins/FX/Niagara/Source/NiagaraEditor/Public/NiagaraEditorUtilities.h
+++ b/Engine/Plugins/FX/Niagara/Source/NiagaraEditor/Public/NiagaraEditorUtilities.h
@@ -11,6 +11,7 @@
 
 class UNiagaraNodeInput;
 class UNiagaraNodeOutput;
+class UNiagaraNodeFunctionCall;
 struct FNiagaraVariable;
 struct FNiagaraTypeDefinition;
 class UNiagaraGraph;
@@ -289,4 +290,5 @@ namespace FNiagaraEditorUtilities
 
 	FString GetNamespacelessVariableNameString(const FName& InVarName);
 
+	void GetReferencingFunctionCallNodes(UNiagaraScript* Script, TArray<UNiagaraNodeFunctionCall*>& OutReferencingFunctionCallNodes);
 };
diff --git a/Engine/Plugins/FX/Niagara/Source/NiagaraEditor/Public/ViewModels/NiagaraScratchPadScriptViewModel.h b/Engine/Plugins/FX/Niagara/Source/NiagaraEditor/Public/ViewModels/NiagaraScratchPadScriptViewModel.h
--- a/Engine/Plugins/FX/Niagara/Source/NiagaraEditor/Public/ViewModels/NiagaraScratchPadScriptViewModel.h
+++ b/Engine/Plugins/FX/Niagara/Source/NiagaraEditor/Public/ViewModels/NiagaraScratchPadScriptViewModel.h
@@ -40,14 +40,18 @@ public:
 
 	void SetEditorHeight(float InEditorHeight);
 
-	bool CanApplyChanges() const;
+	bool HasUnappliedChanges() const;
 
 	void ApplyChanges();
 
+	void DiscardChanges();
+
 	FOnRenamed& OnRenamed();
 
 	FOnPinnedChanged& OnPinnedChanged();
 
+	FSimpleDelegate& OnRequestDiscardChanges();
+
 private:
 	FText GetDisplayNameInternal() const;
 
@@ -72,4 +76,5 @@ private:
 
 	FOnRenamed OnRenamedDelegate;
 	FOnPinnedChanged OnPinnedChangedDelegate;
+	FSimpleDelegate OnRequestDiscardChangesDelegate;
 };
\ No newline at end of file
diff --git a/Engine/Plugins/FX/Niagara/Source/NiagaraEditor/Public/ViewModels/NiagaraScratchPadViewModel.h b/Engine/Plugins/FX/Niagara/Source/NiagaraEditor/Public/ViewModels/NiagaraScratchPadViewModel.h
--- a/Engine/Plugins/FX/Niagara/Source/NiagaraEditor/Public/ViewModels/NiagaraScratchPadViewModel.h
+++ b/Engine/Plugins/FX/Niagara/Source/NiagaraEditor/Public/ViewModels/NiagaraScratchPadViewModel.h
@@ -45,6 +45,8 @@ public:
 
 	void SetActiveScriptViewModel(TSharedRef<FNiagaraScratchPadScriptViewModel> InActiveScriptViewModel);
 
+	void FocusScratchPadScriptViewModel(TSharedRef<FNiagaraScratchPadScriptViewModel> InScriptViewModel);
+
 	void ResetActiveScriptViewModel();
 
 	void CopyActiveScript();
@@ -59,6 +61,12 @@ public:
 
 	TSharedPtr<FNiagaraScratchPadScriptViewModel> CreateNewScriptAsDuplicate(const UNiagaraScript* ScriptToDuplicate);
 
+	void CreateAssetFromActiveScript();
+
+	bool CanSelectNextUsageForActiveScript();
+
+	void SelectNextUsageForActiveScript();
+
 	FOnScriptViewModelsChanged& OnScriptViewModelsChanged();
 
 	FOnScriptViewModelsChanged& OnEditScriptViewModelsChanged();
@@ -72,6 +80,8 @@ public:
 private:
 	TSharedRef<FNiagaraSystemViewModel> GetSystemViewModel();
 
+	TSharedRef<FNiagaraScratchPadScriptViewModel> CreateAndSetupScriptviewModel(UNiagaraScript* ScratchPadScript);
+
 	void RefreshEditScriptViewModels();
 
 	void ScriptGraphNodeSelectionChanged(TWeakPtr<FNiagaraScratchPadScriptViewModel> InScriptViewModelWeak);
@@ -80,6 +90,8 @@ private:
 
 	void ScriptViewModelPinnedChanged(TWeakPtr<FNiagaraScratchPadScriptViewModel> ScriptViewModelWeak);
 
+	void ScriptViewModelRequestDiscardChanges(TWeakPtr<FNiagaraScratchPadScriptViewModel> ScriptViewModelWeak);
+
 private:
 	TSharedPtr<FNiagaraObjectSelection> ObjectSelection;
 
diff --git a/Engine/Plugins/FX/Niagara/Source/NiagaraEditor/Public/ViewModels/NiagaraSystemViewModel.h b/Engine/Plugins/FX/Niagara/Source/NiagaraEditor/Public/ViewModels/NiagaraSystemViewModel.h
--- a/Engine/Plugins/FX/Niagara/Source/NiagaraEditor/Public/ViewModels/NiagaraSystemViewModel.h
+++ b/Engine/Plugins/FX/Niagara/Source/NiagaraEditor/Public/ViewModels/NiagaraSystemViewModel.h
@@ -281,7 +281,7 @@ public:
 
 	FOnRequestFocusTab& OnRequestFocusTab();
 
-	void FocusTab(FName TabName);
+	NIAGARAEDITOR_API void FocusTab(FName TabName);
 	
 	/** Gets the system toolkit command list. */
 	NIAGARAEDITOR_API TSharedPtr<FUICommandList> GetToolkitCommands();
diff --git a/Engine/Plugins/FX/Niagara/Source/NiagaraEditor/Public/ViewModels/Stack/NiagaraStackFunctionInput.h b/Engine/Plugins/FX/Niagara/Source/NiagaraEditor/Public/ViewModels/Stack/NiagaraStackFunctionInput.h
--- a/Engine/Plugins/FX/Niagara/Source/NiagaraEditor/Public/ViewModels/Stack/NiagaraStackFunctionInput.h
+++ b/Engine/Plugins/FX/Niagara/Source/NiagaraEditor/Public/ViewModels/Stack/NiagaraStackFunctionInput.h
@@ -241,6 +241,8 @@ public:
 
 	void SetValueFromClipboardFunctionInput(const UNiagaraClipboardFunctionInput& ClipboardFunctionInput);
 
+	bool IsScratchDynamicInput() const;
+
 public:
 	//~ UNiagaraStackEntry interface
 	virtual void GetSearchItems(TArray<FStackSearchItem>& SearchItems) const override;
@@ -401,6 +403,8 @@ private:
 	/** A tooltip to show for the value of this input. */
 	mutable TOptional<FText> ValueToolTipCache;
 
+	mutable TOptional<bool> bIsScratchDynamicInputCache;
+
 	/** A flag to prevent handling graph changes when it's being updated directly by this object. */
 	bool bUpdatingGraphDirectly;
 
diff --git a/Engine/Plugins/FX/Niagara/Source/NiagaraEditor/Public/ViewModels/Stack/NiagaraStackModuleItem.h b/Engine/Plugins/FX/Niagara/Source/NiagaraEditor/Public/ViewModels/Stack/NiagaraStackModuleItem.h
--- a/Engine/Plugins/FX/Niagara/Source/NiagaraEditor/Public/ViewModels/Stack/NiagaraStackModuleItem.h
+++ b/Engine/Plugins/FX/Niagara/Source/NiagaraEditor/Public/ViewModels/Stack/NiagaraStackModuleItem.h
@@ -92,6 +92,8 @@ public:
 	virtual FText GetDeleteTransactionText() const override;
 	virtual void Delete() override;
 
+	bool IsScratchModule() const;
+
 	void SetOnRequestDeprecationRecommended(FOnRequestDeprecationRecommended InOnRequest)
 	{
 		DeprecationDelegate = InOnRequest;
@@ -139,5 +141,7 @@ private:
 
 	INiagaraStackItemGroupAddUtilities* GroupAddUtilities;
 
+	mutable TOptional<bool> bIsScratchModuleCache;
+
 	bool bIsModuleScriptReassignmentPending;
 };
diff --git a/Engine/Plugins/FX/Niagara/Source/NiagaraEditorWidgets/Private/NiagaraEditorWidgetsStyle.cpp b/Engine/Plugins/FX/Niagara/Source/NiagaraEditorWidgets/Private/NiagaraEditorWidgetsStyle.cpp
--- a/Engine/Plugins/FX/Niagara/Source/NiagaraEditorWidgets/Private/NiagaraEditorWidgetsStyle.cpp
+++ b/Engine/Plugins/FX/Niagara/Source/NiagaraEditorWidgets/Private/NiagaraEditorWidgetsStyle.cpp
@@ -251,7 +251,7 @@ TSharedRef< FSlateStyleSet > FNiagaraEditorWidgetsStyle::Create()
 		.SetFont(ScratchPadLargeHeaderFont);
 	Style->Set("NiagaraEditor.ScratchPad.LargeHeaderText", ScratchPadLargeHeaderText);
 
-	FSlateFontInfo ScratchPadSmallHeaderFont = DEFAULT_FONT("Bold", 10);
+	FSlateFontInfo ScratchPadSmallHeaderFont = DEFAULT_FONT("Bold", 11);
 	FTextBlockStyle ScratchPadSmallHeaderText = FTextBlockStyle(NormalText)
 		.SetFont(ScratchPadSmallHeaderFont);
 	Style->Set("NiagaraEditor.ScratchPad.SmallHeaderText", ScratchPadSmallHeaderText);
diff --git a/Engine/Plugins/FX/Niagara/Source/NiagaraEditorWidgets/Private/NiagaraScratchPadCommandContext.cpp b/Engine/Plugins/FX/Niagara/Source/NiagaraEditorWidgets/Private/NiagaraScratchPadCommandContext.cpp
--- a/Engine/Plugins/FX/Niagara/Source/NiagaraEditorWidgets/Private/NiagaraScratchPadCommandContext.cpp
+++ b/Engine/Plugins/FX/Niagara/Source/NiagaraEditorWidgets/Private/NiagaraScratchPadCommandContext.cpp
@@ -60,7 +60,10 @@ void FNiagaraScratchPadCommandContext::AddMenuItems(FMenuBuilder& MenuBuilder)
 
 	MenuBuilder.BeginSection("Script", LOCTEXT("ScriptActions", "Script"));
 	{
-		MenuBuilder.AddMenuEntry(FNiagaraEditorModule::Get().Commands().Apply);
+		MenuBuilder.AddMenuEntry(FNiagaraEditorModule::Get().Commands().Apply, NAME_None, LOCTEXT("ApplyChangesLabel", "Apply Changes"));
+		MenuBuilder.AddMenuEntry(FNiagaraEditorModule::Get().Commands().Discard, NAME_None, LOCTEXT("DiscardChangesLabel", "Discard Changes"));
+		MenuBuilder.AddMenuEntry(FNiagaraEditorModule::Get().Commands().SelectNextUsage, NAME_None, TAttribute<FText>(), LOCTEXT("SelectNextToolTip", "Select the next usage of this script in the selection stack."));
+		MenuBuilder.AddMenuEntry(FNiagaraEditorModule::Get().Commands().CreateAssetFromSelection);
 	}
 	MenuBuilder.EndSection();
 }
@@ -82,6 +85,15 @@ void FNiagaraScratchPadCommandContext::SetupCommands()
 	Commands->MapAction(FNiagaraEditorModule::Get().Commands().Apply, FUIAction(
 		FExecuteAction::CreateSP(this, &FNiagaraScratchPadCommandContext::ApplyChangesToSelectedScripts),
 		FCanExecuteAction::CreateSP(this, &FNiagaraScratchPadCommandContext::CanApplyChangesToSelectedScripts)));
+	Commands->MapAction(FNiagaraEditorModule::Get().Commands().Discard, FUIAction(
+		FExecuteAction::CreateSP(this, &FNiagaraScratchPadCommandContext::DiscardChangesFromSelectedScripts),
+		FCanExecuteAction::CreateSP(this, &FNiagaraScratchPadCommandContext::CanDiscardChangesFromSelectedScripts)));
+	Commands->MapAction(FNiagaraEditorModule::Get().Commands().SelectNextUsage, FUIAction(
+		FExecuteAction::CreateSP(this, &FNiagaraScratchPadCommandContext::SelectNextUsageForSelectedScript),
+		FCanExecuteAction::CreateSP(this, &FNiagaraScratchPadCommandContext::CanSelectNextUsageForSelectedScript)));
+	Commands->MapAction(FNiagaraEditorModule::Get().Commands().CreateAssetFromSelection, FUIAction(
+		FExecuteAction::CreateSP(this, &FNiagaraScratchPadCommandContext::CreateAssetFromSelectedScript),
+		FCanExecuteAction::CreateSP(this, &FNiagaraScratchPadCommandContext::CanCreateAssetFromSelectedScript)));
 }
 
 bool FNiagaraScratchPadCommandContext::CanCutSelectedScripts() const
@@ -151,7 +163,7 @@ void FNiagaraScratchPadCommandContext::DeleteSelectedScripts() const
 
 bool FNiagaraScratchPadCommandContext::CanApplyChangesToSelectedScripts() const
 {
-	return ScratchPadViewModel->GetActiveScriptViewModel().IsValid() && ScratchPadViewModel->GetActiveScriptViewModel()->CanApplyChanges();
+	return ScratchPadViewModel->GetActiveScriptViewModel().IsValid() && ScratchPadViewModel->GetActiveScriptViewModel()->HasUnappliedChanges();
 }
 
 void FNiagaraScratchPadCommandContext::ApplyChangesToSelectedScripts() const
@@ -162,5 +174,38 @@ void FNiagaraScratchPadCommandContext::ApplyChangesToSelectedScripts() const
 	}
 }
 
+bool FNiagaraScratchPadCommandContext::CanDiscardChangesFromSelectedScripts() const
+{
+	return ScratchPadViewModel->GetActiveScriptViewModel().IsValid() && ScratchPadViewModel->GetActiveScriptViewModel()->HasUnappliedChanges();
+}
+
+void FNiagaraScratchPadCommandContext::DiscardChangesFromSelectedScripts() const
+{
+	if (CanDiscardChangesFromSelectedScripts())
+	{
+		ScratchPadViewModel->GetActiveScriptViewModel()->DiscardChanges();
+	}
+}
+
+bool FNiagaraScratchPadCommandContext::CanSelectNextUsageForSelectedScript() const
+{
+	return ScratchPadViewModel->CanSelectNextUsageForActiveScript();
+}
+
+void FNiagaraScratchPadCommandContext::SelectNextUsageForSelectedScript() const
+{
+	ScratchPadViewModel->SelectNextUsageForActiveScript();
+}
+
+bool FNiagaraScratchPadCommandContext::CanCreateAssetFromSelectedScript() const
+{
+	return ScratchPadViewModel->GetActiveScriptViewModel().IsValid();
+}
+
+void FNiagaraScratchPadCommandContext::CreateAssetFromSelectedScript() const
+{
+	ScratchPadViewModel->CreateAssetFromActiveScript();
+}
+
 #undef LOCTEXT_NAMESPACE
 
diff --git a/Engine/Plugins/FX/Niagara/Source/NiagaraEditorWidgets/Private/NiagaraScratchPadCommandContext.h b/Engine/Plugins/FX/Niagara/Source/NiagaraEditorWidgets/Private/NiagaraScratchPadCommandContext.h
--- a/Engine/Plugins/FX/Niagara/Source/NiagaraEditorWidgets/Private/NiagaraScratchPadCommandContext.h
+++ b/Engine/Plugins/FX/Niagara/Source/NiagaraEditorWidgets/Private/NiagaraScratchPadCommandContext.h
@@ -51,6 +51,18 @@ private:
 
 	void ApplyChangesToSelectedScripts() const;
 
+	bool CanDiscardChangesFromSelectedScripts() const;
+
+	void DiscardChangesFromSelectedScripts() const;
+
+	bool CanSelectNextUsageForSelectedScript() const;
+
+	void SelectNextUsageForSelectedScript() const;
+
+	bool CanCreateAssetFromSelectedScript() const;
+
+	void CreateAssetFromSelectedScript() const;
+
 private:
 	TSharedRef<FUICommandList> Commands;
 
diff --git a/Engine/Plugins/FX/Niagara/Source/NiagaraEditorWidgets/Private/SNiagaraScratchPad.cpp b/Engine/Plugins/FX/Niagara/Source/NiagaraEditorWidgets/Private/SNiagaraScratchPad.cpp
--- a/Engine/Plugins/FX/Niagara/Source/NiagaraEditorWidgets/Private/SNiagaraScratchPad.cpp
+++ b/Engine/Plugins/FX/Niagara/Source/NiagaraEditorWidgets/Private/SNiagaraScratchPad.cpp
@@ -490,18 +490,16 @@ class SNiagaraScratchPadScriptEditor : public SCompoundWidget
 				.Padding(0.0f, 2.0f, 1.0f, 4.0f)
 				[
 					SNew(SButton)
-					.ButtonStyle(FEditorStyle::Get(), "RoundButton")
 					.OnClicked(this, &SNiagaraScratchPadScriptEditor::OnApplyButtonClicked)
 					.ToolTipText(LOCTEXT("ApplyButtonToolTip", "Apply the current changes to this script.  This will update the selection stack UI and compile neccessary scripts."))
 					.IsEnabled(this, &SNiagaraScratchPadScriptEditor::GetApplyButtonIsEnabled)
-					.ForegroundColor(FSlateColor::UseForeground())
-					.ContentPadding(FMargin(0.0f))
+					.ContentPadding(FMargin(3.0f, 0.0f))
 					.Content()
 					[
 						SNew(SHorizontalBox)
 						+ SHorizontalBox::Slot()
 						.AutoWidth()
-						.Padding(2.0f, 1.0f, 2.0f, 1.0f)
+						.Padding(2.0f, 2.0f, 2.0f, 2.0f)
 						[
 							SNew(SImage)
 							.Image(FNiagaraEditorStyle::Get().GetBrush("NiagaraEditor.Apply.Small"))
@@ -541,7 +539,7 @@ private:
 
 	bool GetApplyButtonIsEnabled() const
 	{
-		return ScriptViewModel->CanApplyChanges();
+		return ScriptViewModel->HasUnappliedChanges();
 	}
 
 private:
diff --git a/Engine/Plugins/FX/Niagara/Source/NiagaraEditorWidgets/Private/Stack/SNiagaraStackFunctionInputValue.cpp b/Engine/Plugins/FX/Niagara/Source/NiagaraEditorWidgets/Private/Stack/SNiagaraStackFunctionInputValue.cpp
--- a/Engine/Plugins/FX/Niagara/Source/NiagaraEditorWidgets/Private/Stack/SNiagaraStackFunctionInputValue.cpp
+++ b/Engine/Plugins/FX/Niagara/Source/NiagaraEditorWidgets/Private/Stack/SNiagaraStackFunctionInputValue.cpp
@@ -2,6 +2,8 @@
 
 #include "Stack/SNiagaraStackFunctionInputValue.h"
 #include "ViewModels/Stack/NiagaraStackFunctionInput.h"
+#include "ViewModels/NiagaraScratchPadViewModel.h"
+#include "ViewModels/NiagaraScratchPadScriptViewModel.h"
 #include "NiagaraEditorModule.h"
 #include "INiagaraEditorTypeUtilities.h"
 #include "NiagaraEditorUtilities.h"
@@ -188,10 +190,37 @@ TSharedRef<SWidget> SNiagaraStackFunctionInputValue::ConstructValueWidgets()
 	}
 	case UNiagaraStackFunctionInput::EValueMode::Dynamic:
 	{
-		return SNew(STextBlock)
+		TSharedRef<SWidget> DynamicInputText = SNew(STextBlock)
 			.TextStyle(FNiagaraEditorStyle::Get(), "NiagaraEditor.ParameterText")
 			.Text(this, &SNiagaraStackFunctionInputValue::GetDynamicValueText)
 			.OnDoubleClicked(this, &SNiagaraStackFunctionInputValue::DynamicInputTextDoubleClicked);
+		if (FunctionInput->IsScratchDynamicInput())
+		{
+			return SNew(SHorizontalBox)
+				+ SHorizontalBox::Slot()
+				.VAlign(VAlign_Center)
+				[
+					DynamicInputText
+				]
+				+ SHorizontalBox::Slot()
+				.AutoWidth()
+				[
+					SNew(SButton)
+					.ButtonStyle(FEditorStyle::Get(), "RoundButton")
+					.OnClicked(this, &SNiagaraStackFunctionInputValue::ScratchButtonPressed)
+					.ToolTipText(LOCTEXT("OpenInScratchToolTip", "Open this dynamic input in the scratch pad."))
+					.ContentPadding(FMargin(1.0f, 0.0f))
+					.Content()
+					[
+						SNew(SImage)
+						.Image(FNiagaraEditorStyle::Get().GetBrush("NiagaraEditor.Scratch"))
+					]
+				];
+		}
+		else
+		{
+			return DynamicInputText;
+		}
 	}
 	case UNiagaraStackFunctionInput::EValueMode::DefaultFunction:
 	{
@@ -388,10 +417,22 @@ void SNiagaraStackFunctionInputValue::OnExpressionTextCommitted(const FText& Nam
 FReply SNiagaraStackFunctionInputValue::DynamicInputTextDoubleClicked(const FGeometry& MyGeometry, const FPointerEvent& PointerEvent)
 {
 	UNiagaraNodeFunctionCall* DynamicInputNode = FunctionInput->GetDynamicInputNode();
-	if (DynamicInputNode->FunctionScript != nullptr && DynamicInputNode->FunctionScript->IsAsset())
+	if (DynamicInputNode->FunctionScript != nullptr)
 	{
-		GEditor->GetEditorSubsystem<UAssetEditorSubsystem>()->OpenEditorForAsset(DynamicInputNode->FunctionScript);
-		return FReply::Handled();
+		if (DynamicInputNode->FunctionScript->IsAsset())
+		{
+			GEditor->GetEditorSubsystem<UAssetEditorSubsystem>()->OpenEditorForAsset(DynamicInputNode->FunctionScript);
+			return FReply::Handled();
+		}
+		else
+		{
+			TSharedPtr<FNiagaraScratchPadScriptViewModel> ScratchPadScriptViewModel = FunctionInput->GetSystemViewModel()->GetScriptScratchPadViewModel()->GetViewModelForScript(DynamicInputNode->FunctionScript);
+			if(ScratchPadScriptViewModel.IsValid())
+			{
+				FunctionInput->GetSystemViewModel()->GetScriptScratchPadViewModel()->FocusScratchPadScriptViewModel(ScratchPadScriptViewModel.ToSharedRef());
+				return FReply::Handled();
+			}
+		}
 	}
 	return FReply::Unhandled();
 }
@@ -912,4 +953,16 @@ void SNiagaraStackFunctionInputValue::SetLibraryOnly(bool bInIsLibraryOnly)
 	bLibraryOnly = bInIsLibraryOnly;
 }
 
+FReply SNiagaraStackFunctionInputValue::ScratchButtonPressed() const
+{
+	TSharedPtr<FNiagaraScratchPadScriptViewModel> ScratchDynamicInputViewModel = 
+		FunctionInput->GetSystemViewModel()->GetScriptScratchPadViewModel()->GetViewModelForScript(FunctionInput->GetDynamicInputNode()->FunctionScript);
+	if (ScratchDynamicInputViewModel.IsValid())
+	{
+		FunctionInput->GetSystemViewModel()->GetScriptScratchPadViewModel()->FocusScratchPadScriptViewModel(ScratchDynamicInputViewModel.ToSharedRef());
+		return FReply::Handled();
+	}
+	return FReply::Unhandled();
+}
+
 #undef LOCTEXT_NAMESPACE
diff --git a/Engine/Plugins/FX/Niagara/Source/NiagaraEditorWidgets/Private/Stack/SNiagaraStackFunctionInputValue.h b/Engine/Plugins/FX/Niagara/Source/NiagaraEditorWidgets/Private/Stack/SNiagaraStackFunctionInputValue.h
--- a/Engine/Plugins/FX/Niagara/Source/NiagaraEditorWidgets/Private/Stack/SNiagaraStackFunctionInputValue.h
+++ b/Engine/Plugins/FX/Niagara/Source/NiagaraEditorWidgets/Private/Stack/SNiagaraStackFunctionInputValue.h
@@ -141,6 +141,8 @@ private:
 
 	void SetLibraryOnly(bool bInIsLibraryOnly);
 
+	FReply ScratchButtonPressed() const;
+
 private:
 	UNiagaraStackFunctionInput* FunctionInput;
 
diff --git a/Engine/Plugins/FX/Niagara/Source/NiagaraEditorWidgets/Private/Stack/SNiagaraStackModuleItem.cpp b/Engine/Plugins/FX/Niagara/Source/NiagaraEditorWidgets/Private/Stack/SNiagaraStackModuleItem.cpp
--- a/Engine/Plugins/FX/Niagara/Source/NiagaraEditorWidgets/Private/Stack/SNiagaraStackModuleItem.cpp
+++ b/Engine/Plugins/FX/Niagara/Source/NiagaraEditorWidgets/Private/Stack/SNiagaraStackModuleItem.cpp
@@ -12,6 +12,9 @@
 #include "ViewModels/Stack/NiagaraStackModuleItem.h"
 #include "ViewModels/Stack/NiagaraStackViewModel.h"
 #include "ViewModels/Stack/NiagaraStackGraphUtilities.h"
+#include "ViewModels/NiagaraSystemViewModel.h"
+#include "ViewModels/NiagaraScratchPadViewModel.h"
+#include "ViewModels/NiagaraScratchPadScriptViewModel.h"
 #include "Widgets/Layout/SBox.h"
 #include "Widgets/SBoxPanel.h"
 #include "Widgets/Input/SButton.h"
@@ -49,10 +52,22 @@ void SNiagaraStackModuleItem::FillRowContextMenu(FMenuBuilder& MenuBuilder)
 FReply SNiagaraStackModuleItem::OnMouseButtonDoubleClick(const FGeometry& InMyGeometry, const FPointerEvent& InMouseEvent)
 {
 	const UNiagaraNodeFunctionCall& ModuleFunctionCall = ModuleItem->GetModuleNode();
-	if (ModuleFunctionCall.FunctionScript != nullptr && ModuleFunctionCall.FunctionScript->IsAsset())
+	if (ModuleFunctionCall.FunctionScript != nullptr)
 	{
-		GEditor->GetEditorSubsystem<UAssetEditorSubsystem>()->OpenEditorForAsset(const_cast<UNiagaraScript*>(ModuleFunctionCall.FunctionScript));
-		return FReply::Handled();
+		if (ModuleFunctionCall.FunctionScript->IsAsset())
+		{
+			GEditor->GetEditorSubsystem<UAssetEditorSubsystem>()->OpenEditorForAsset(const_cast<UNiagaraScript*>(ModuleFunctionCall.FunctionScript));
+			return FReply::Handled();
+		}
+		else if (ModuleItem->IsScratchModule())
+		{
+			TSharedPtr<FNiagaraScratchPadScriptViewModel> ScratchPadScriptViewModel = ModuleItem->GetSystemViewModel()->GetScriptScratchPadViewModel()->GetViewModelForScript(ModuleFunctionCall.FunctionScript);
+			if (ScratchPadScriptViewModel.IsValid())
+			{
+				ModuleItem->GetSystemViewModel()->GetScriptScratchPadViewModel()->FocusScratchPadScriptViewModel(ScratchPadScriptViewModel.ToSharedRef());
+				return FReply::Handled();
+			}
+		}
 	}
 	return FReply::Unhandled();
 }
@@ -69,6 +84,24 @@ void SNiagaraStackModuleItem::Tick(const FGeometry& AllottedGeometry, const doub
 
 void SNiagaraStackModuleItem::AddCustomRowWidgets(TSharedRef<SHorizontalBox> HorizontalBox)
 {
+	// Scratch navigation
+	if(ModuleItem->IsScratchModule())
+	{
+		HorizontalBox->AddSlot()
+		.AutoWidth()
+		[
+			SNew(SButton)
+			.ButtonStyle(FEditorStyle::Get(), "RoundButton")
+			.OnClicked(this, &SNiagaraStackModuleItem::ScratchButtonPressed)
+			.ToolTipText(LOCTEXT("OpenInScratchToolTip", "Open this dynamic input in the scratch pad."))
+			.ContentPadding(FMargin(1.0f, 0.0f))
+			.Content()
+			[
+				SNew(SImage)
+				.Image(FNiagaraEditorStyle::Get().GetBrush("NiagaraEditor.Scratch"))
+			]
+		];
+	}
 	// Add menu.
 	HorizontalBox->AddSlot()
 	.VAlign(VAlign_Center)
@@ -143,6 +176,18 @@ EVisibility SNiagaraStackModuleItem::GetRefreshVisibility() const
 	return ModuleItem->CanRefresh() ? EVisibility::Visible : EVisibility::Collapsed;
 }
 
+FReply SNiagaraStackModuleItem::ScratchButtonPressed() const
+{
+	TSharedPtr<FNiagaraScratchPadScriptViewModel> ScratchModuleViewModel =
+		ModuleItem->GetSystemViewModel()->GetScriptScratchPadViewModel()->GetViewModelForScript(ModuleItem->GetModuleNode().FunctionScript);
+	if (ScratchModuleViewModel.IsValid())
+	{
+		ModuleItem->GetSystemViewModel()->GetScriptScratchPadViewModel()->FocusScratchPadScriptViewModel(ScratchModuleViewModel.ToSharedRef());
+		return FReply::Handled();
+	}
+	return FReply::Unhandled();
+}
+
 TSharedRef<SWidget> SNiagaraStackModuleItem::RaiseActionMenuClicked()
 {
 	if (CanRaiseActionMenu())
diff --git a/Engine/Plugins/FX/Niagara/Source/NiagaraEditorWidgets/Private/Stack/SNiagaraStackModuleItem.h b/Engine/Plugins/FX/Niagara/Source/NiagaraEditorWidgets/Private/Stack/SNiagaraStackModuleItem.h
--- a/Engine/Plugins/FX/Niagara/Source/NiagaraEditorWidgets/Private/Stack/SNiagaraStackModuleItem.h
+++ b/Engine/Plugins/FX/Niagara/Source/NiagaraEditorWidgets/Private/Stack/SNiagaraStackModuleItem.h
@@ -38,6 +38,8 @@ private:
 
 	EVisibility GetRefreshVisibility() const;
 
+	FReply ScratchButtonPressed() const;
+	
 	TSharedRef<SWidget> RaiseActionMenuClicked();
 
 	bool CanRaiseActionMenu() const;
