Make ShaderScriptResource a TUniquePtr on the UNiagaraScript so that we don't have to pay for the cost on non-GPU scripts

diff --git a/Engine/Plugins/FX/Niagara/Source/Niagara/Classes/NiagaraScript.h b/Engine/Plugins/FX/Niagara/Source/Niagara/Classes/NiagaraScript.h
--- a/Engine/Plugins/FX/Niagara/Source/Niagara/Classes/NiagaraScript.h
+++ b/Engine/Plugins/FX/Niagara/Source/Niagara/Classes/NiagaraScript.h
@@ -21,9 +21,6 @@ class UNiagaraDataInterface;
 class FNiagaraCompileRequestDataBase;
 class UNiagaraConvertInPlaceUtilityBase;
 
-void SerializeNiagaraShaderMaps(const TMap<const ITargetPlatform*, TArray<FNiagaraShaderScript*>>* PlatformScriptResourcesToSave, FArchive& Ar, TArray<FNiagaraShaderScript>& OutLoadedResources);
-void ProcessSerializedShaderMaps(UNiagaraScript* Owner, TArray<FNiagaraShaderScript>& LoadedResources, FNiagaraShaderScript& OutResourceForCurrentPlatform, FNiagaraShaderScript* (&OutScriptResourcesLoaded)[ERHIFeatureLevel::Num]);
-
 #define NIAGARA_INVALID_MEMORY (0xBA)
 
 #define NIAGARA_SCRIPT_COMPILE_LOGGING_MEDIUM
@@ -593,34 +590,33 @@ public:
 	FNiagaraShaderScript* AllocateResource();
 	FNiagaraShaderScript* GetRenderThreadScript()
 	{
-		return &ScriptResource;
+		return ScriptResource.Get();
 	}
 
 	const FNiagaraShaderScript* GetRenderThreadScript() const
 	{
-		return &ScriptResource;
+		return ScriptResource.Get();
 	}
+#if WITH_EDITORONLY_DATA
 
 	FComputeShaderRHIRef GetScriptShader() 
 	{
-		if (!ScriptShader)
+		if (!ScriptShader && ScriptResource.IsValid())
 		{
-			ScriptShader = ScriptResource.GetShader().GetComputeShader();	// NIAGARATODO: need to put this caching somewhere else, as it wont' know when we update the resource
+			ScriptShader = ScriptResource->GetShader().GetComputeShader();	// NIAGARATODO: need to put this caching somewhere else, as it wont' know when we update the resource
 		}
 		return ScriptShader; 
 	}
 
 	FComputeShaderRHIRef GetScriptShaderGameThread()
 	{
-		if (!ScriptShader)
+		if (!ScriptShader && ScriptResource.IsValid())
 		{
-			ScriptShader = ScriptResource.GetShaderGameThread().GetComputeShader();	// NIAGARATODO: need to put this caching somewhere else, as it wont' know when we update the resource
+			ScriptShader = ScriptResource->GetShaderGameThread().GetComputeShader();	// NIAGARATODO: need to put this caching somewhere else, as it wont' know when we update the resource
 		}
 		return ScriptShader;
 	}
-
-
-	void SetFeatureLevel(ERHIFeatureLevel::Type InFeatureLevel)		{ FeatureLevel = InFeatureLevel; }
+#endif
 
 	NIAGARA_API void GenerateStatScopeIDs();
 
@@ -709,6 +705,9 @@ private:
 	bool OwnerCanBeRunOnGpu() const;
 	bool LegacyCanBeRunOnGpu()const;
 
+	void ProcessSerializedShaderMaps();
+	void SerializeNiagaraShaderMaps(FArchive& Ar, int32 NiagaraVer, bool IsValidShaderScript);
+
 	/** Return the expected SimTarget for this script. Only returns a valid target if there is valid data to run with. */
 	TOptional<ENiagaraSimTarget> GetSimTarget() const;
 
@@ -762,14 +761,12 @@ private:
 	mutable FNiagaraVMExecutableDataId LastGeneratedVMId;
 #endif
 
-	TArray<FNiagaraShaderScript> LoadedScriptResources;
-
-	FNiagaraShaderScript ScriptResource;
+	TUniquePtr<FNiagaraShaderScript> ScriptResource;
 
+#if WITH_EDITORONLY_DATA
+	TArray<FNiagaraShaderScript> LoadedScriptResources;
 	FNiagaraShaderScript* ScriptResourcesByFeatureLevel[ERHIFeatureLevel::Num];
-
-	/** Feature level that the shader map is going to be compiled for. */
-	ERHIFeatureLevel::Type FeatureLevel;
+#endif
 
 	/** Compute shader compiled for this script */
 	FComputeShaderRHIRef ScriptShader;
diff --git a/Engine/Plugins/FX/Niagara/Source/Niagara/Private/NiagaraScript.cpp b/Engine/Plugins/FX/Niagara/Source/Niagara/Private/NiagaraScript.cpp
--- a/Engine/Plugins/FX/Niagara/Source/Niagara/Private/NiagaraScript.cpp
+++ b/Engine/Plugins/FX/Niagara/Source/Niagara/Private/NiagaraScript.cpp
@@ -288,7 +288,8 @@ UNiagaraScript::UNiagaraScript(const FObjectInitializer& ObjectInitializer)
 #endif
 {
 #if WITH_EDITORONLY_DATA
-	ScriptResource.OnCompilationComplete().AddUniqueDynamic(this, &UNiagaraScript::RaiseOnGPUCompilationComplete);
+	ScriptResource = MakeUnique<FNiagaraShaderScript>();
+	ScriptResource->OnCompilationComplete().AddUniqueDynamic(this, &UNiagaraScript::RaiseOnGPUCompilationComplete);
 
 	RapidIterationParameters.DebugName = *GetFullName();
 #endif	
@@ -896,16 +897,7 @@ void UNiagaraScript::Serialize(FArchive& Ar)
 		}
 	}
 
-	if ( (!Ar.IsLoading() && IsValidShaderScript)		// saving shader maps only for particle sim and spawn scripts
-		|| (Ar.IsLoading() && NiagaraVer >= FNiagaraCustomVersion::NiagaraShaderMaps && (NiagaraVer < FNiagaraCustomVersion::NiagaraShaderMapCooking || IsValidShaderScript))  // load only if we know shader map is presen
-		)
-	{
-#if WITH_EDITOR
-		SerializeNiagaraShaderMaps(&CachedScriptResourcesForCooking, Ar, LoadedScriptResources);
-#else
-		SerializeNiagaraShaderMaps(nullptr, Ar, LoadedScriptResources);
-#endif
-	}
+	SerializeNiagaraShaderMaps(Ar, NiagaraVer, IsValidShaderScript);
 }
 
 /** Is usage A dependent on Usage B?*/
@@ -1064,17 +1056,8 @@ void UNiagaraScript::PostLoad()
 	}
 #endif
 	
-	// Resources can be processed / registered now that we're back on the main thread
-	ProcessSerializedShaderMaps(this, LoadedScriptResources, ScriptResource, ScriptResourcesByFeatureLevel);
+	ProcessSerializedShaderMaps();
 
-	// for now, force recompile until we can be sure everything is working
-	//bNeedsRecompile = true;
-#if WITH_EDITORONLY_DATA
-	if (CachedScriptVMId.BaseScriptCompileHash.IsValid())
-	{
-		CacheResourceShadersForRendering(false, bNeedsRecompile);
-	}
-#endif
 #if STATS
 	GenerateStatScopeIDs();
 #endif
@@ -1442,7 +1425,7 @@ void UNiagaraScript::SetVMCompilationResults(const FNiagaraVMExecutableDataId& I
 	CachedScriptVM = InScriptVM;
 	CachedParameterCollectionReferences.Empty();
 	// Proactively clear out the script resource, because it might be stale now.
-	ScriptResource.Invalidate();
+	ScriptResource->Invalidate();
 	
 	if (CachedScriptVM.LastCompileStatus == ENiagaraScriptCompileStatus::NCS_Error)
 	{
@@ -1893,24 +1876,24 @@ void UNiagaraScript::CacheResourceShadersForRendering(bool bRegenerateId, bool b
 		{
 			FNiagaraShaderScript* ResourceToCache;
 			ERHIFeatureLevel::Type CacheFeatureLevel = GMaxRHIFeatureLevel;
-			ScriptResource.SetScript(this, FeatureLevel, CachedScriptVMId.CompilerVersionID, CachedScriptVMId.AdditionalDefines,
+			ScriptResource->SetScript(this, CacheFeatureLevel, CachedScriptVMId.CompilerVersionID, CachedScriptVMId.AdditionalDefines,
 				CachedScriptVMId.BaseScriptCompileHash, CachedScriptVMId.ReferencedCompileHashes, 
 				CachedScriptVMId.bUsesRapidIterationParams, GetFriendlyName());
 
-			//if (ScriptResourcesByFeatureLevel[FeatureLevel])
+			//if (ScriptResourcesByFeatureLevel[CacheFeatureLevel])
 			{
 				const EShaderPlatform ShaderPlatform = GShaderPlatformForFeatureLevel[CacheFeatureLevel];
 				if (FNiagaraUtilities::SupportsGPUParticles(ShaderPlatform))
 				{
 					ResourceToCache = ScriptResourcesByFeatureLevel[CacheFeatureLevel];
-					CacheShadersForResources(ShaderPlatform, &ScriptResource, true);
-					ScriptResourcesByFeatureLevel[CacheFeatureLevel] = &ScriptResource;
+					CacheShadersForResources(ShaderPlatform, ScriptResource.Get(), true);
+					ScriptResourcesByFeatureLevel[CacheFeatureLevel] = ScriptResource.Get();
 				}
 			}
 		}
 		else
 		{
-			ScriptResource.Invalidate();
+			ScriptResource->Invalidate();
 		}
 	}
 }
@@ -2020,7 +2003,7 @@ bool UNiagaraScript::SynchronizeExecutablesWithMaster(const UNiagaraScript* Scri
 	if (Id == Script->GetVMExecutableDataCompilationId())
 	{
 		CachedScriptVM.Reset();
-		ScriptResource.Invalidate();
+		ScriptResource->Invalidate();
 
 		CachedScriptVM = Script->CachedScriptVM;
 		CachedScriptVMId = Script->CachedScriptVMId;
@@ -2056,7 +2039,7 @@ void UNiagaraScript::InvalidateCompileResults(const FString& Reason)
 {
 	UE_LOG(LogNiagara, Verbose, TEXT("InvalidateCompileResults Script:%s Reason:%s"), *GetPathName(), *Reason);
 	CachedScriptVM.Reset();
-	ScriptResource.Invalidate();
+	ScriptResource->Invalidate();
 	CachedScriptVMId.Invalidate();
 	LastGeneratedVMId.Invalidate();
 	CachedDefaultDataInterfaces.Reset();
@@ -2085,12 +2068,12 @@ NIAGARA_API bool UNiagaraScript::IsScriptCompilationPending(bool bGPUScript) con
 {
 	if (bGPUScript)
 	{
-		FNiagaraShaderRef Shader = ScriptResource.GetShaderGameThread();
+		FNiagaraShaderRef Shader = ScriptResource->GetShaderGameThread();
 		if (Shader.IsValid())
 		{
 			return false;
 		}
-		return !ScriptResource.IsCompilationFinished();
+		return !ScriptResource->IsCompilationFinished();
 	}
 	else
 	{
@@ -2106,13 +2089,13 @@ NIAGARA_API bool UNiagaraScript::DidScriptCompilationSucceed(bool bGPUScript) co
 {
 	if (bGPUScript)
 	{
-		FNiagaraShaderRef Shader = ScriptResource.GetShaderGameThread();
+		FNiagaraShaderRef Shader = ScriptResource->GetShaderGameThread();
 		if (Shader.IsValid())
 		{
 			return true;
 		}
 
-		if (ScriptResource.IsCompilationFinished())
+		if (ScriptResource->IsCompilationFinished())
 		{
 			// If we failed compilation, it would be finished and Shader would be null.
 			return false;
@@ -2129,73 +2112,101 @@ NIAGARA_API bool UNiagaraScript::DidScriptCompilationSucceed(bool bGPUScript) co
 	return false;
 }
 
-void SerializeNiagaraShaderMaps(const TMap<const ITargetPlatform*, TArray<FNiagaraShaderScript*>>* PlatformScriptResourcesToSave, FArchive& Ar, TArray<FNiagaraShaderScript>& OutLoadedResources)
+void UNiagaraScript::SerializeNiagaraShaderMaps(FArchive& Ar, int32 NiagaraVer, bool IsValidShaderScript)
 {
-	Ar.UsingCustomVersion(FFortniteMainBranchObjectVersion::GUID);
-	Ar.UsingCustomVersion(FRenderingObjectVersion::GUID);
-	Ar.UsingCustomVersion(FEditorObjectVersion::GUID);
-	Ar.UsingCustomVersion(FReleaseObjectVersion::GUID);
-
-//	SCOPED_LOADTIMER(SerializeInlineShaderMaps);
-	if (Ar.IsSaving())
+#if WITH_EDITOR
+	if ((Ar.IsSaving() && IsValidShaderScript)		// saving shader maps only for particle sim and spawn scripts
+		|| (Ar.IsLoading() && NiagaraVer >= FNiagaraCustomVersion::NiagaraShaderMaps && (NiagaraVer < FNiagaraCustomVersion::NiagaraShaderMapCooking || IsValidShaderScript))  // load only if we know shader map is presen
+		)
 	{
-		int32 NumResourcesToSave = 0;
-		const TArray<FNiagaraShaderScript*>* ScriptResourcesToSavePtr = nullptr;
+		Ar.UsingCustomVersion(FFortniteMainBranchObjectVersion::GUID);
+		Ar.UsingCustomVersion(FRenderingObjectVersion::GUID);
+		Ar.UsingCustomVersion(FEditorObjectVersion::GUID);
+		Ar.UsingCustomVersion(FReleaseObjectVersion::GUID);
 
-		if (Ar.IsCooking())
+		//	SCOPED_LOADTIMER(SerializeInlineShaderMaps);
+		if (Ar.IsSaving())
 		{
-			checkf(PlatformScriptResourcesToSave != nullptr, TEXT("PlatformScriptResourcesToSave must be supplied when cooking"));
-			ScriptResourcesToSavePtr = PlatformScriptResourcesToSave->Find(Ar.CookingTarget());
-			if (ScriptResourcesToSavePtr != nullptr)
+			int32 NumResourcesToSave = 0;
+			const TArray<FNiagaraShaderScript*>* ScriptResourcesToSavePtr = nullptr;
+
+			if (Ar.IsCooking())
 			{
-				NumResourcesToSave = ScriptResourcesToSavePtr->Num();
+				ScriptResourcesToSavePtr = CachedScriptResourcesForCooking.Find(Ar.CookingTarget());
+				if (ScriptResourcesToSavePtr != nullptr)
+				{
+					NumResourcesToSave = ScriptResourcesToSavePtr->Num();
+				}
 			}
-		}
 
-		Ar << NumResourcesToSave;
+			Ar << NumResourcesToSave;
 
-		if (ScriptResourcesToSavePtr != nullptr)
+			if (ScriptResourcesToSavePtr != nullptr)
+			{
+				for (FNiagaraShaderScript* ScriptResourceToSave : (*ScriptResourcesToSavePtr))
+				{
+					checkf(ScriptResourceToSave != nullptr, TEXT("Invalid script resource was cached"));
+					ScriptResourceToSave->SerializeShaderMap(Ar);
+				}
+			}
+		}
+		else if (Ar.IsLoading())
 		{
-			for (FNiagaraShaderScript* ScriptResourceToSave : (*ScriptResourcesToSavePtr))
+			int32 NumLoadedResources = 0;
+			Ar << NumLoadedResources;
+			for (int32 i = 0; i < NumLoadedResources; i++)
 			{
-				checkf(ScriptResourceToSave != nullptr, TEXT("Invalid script resource was cached"));
-				ScriptResourceToSave->SerializeShaderMap(Ar);
+				FNiagaraShaderScript LoadedResource;
+				LoadedResource.SerializeShaderMap(Ar);
+				LoadedScriptResources.Add(LoadedResource);
 			}
 		}
 	}
-	else if (Ar.IsLoading())
+#else
+	if (Ar.IsLoading() && IsValidShaderScript)
 	{
-		int32 NumLoadedResources = 0;
-		Ar << NumLoadedResources;
-		for (int32 i = 0; i < NumLoadedResources; i++)
+		check(NiagaraVer >= FNiagaraCustomVersion::NiagaraShaderMaps);
+		int32 ResourceCount = 0;
+		Ar << ResourceCount;
+
+		for (int32 ResourceIt = 0; ResourceIt < ResourceCount; ++ResourceIt)
 		{
-			FNiagaraShaderScript LoadedResource;
-			LoadedResource.SerializeShaderMap(Ar);
-			OutLoadedResources.Add(LoadedResource);
+			FNiagaraShaderScript Resource;
+			Resource.SerializeShaderMap(Ar);
+
+			if (!ScriptResource && GMaxRHIShaderPlatform == Resource.GetGameThreadShaderMap()->GetShaderPlatform())
+			{
+				ScriptResource = MakeUnique<FNiagaraShaderScript>(Resource);
+			}
 		}
 	}
+#endif
 }
 
-void ProcessSerializedShaderMaps(UNiagaraScript* Owner, TArray<FNiagaraShaderScript>& LoadedResources, FNiagaraShaderScript& OutResourceForCurrentPlatform, FNiagaraShaderScript* (&OutScriptResourcesLoaded)[ERHIFeatureLevel::Num])
+void UNiagaraScript::ProcessSerializedShaderMaps()
 {
 	check(IsInGameThread());
 
-	for (FNiagaraShaderScript& LoadedResource : LoadedResources)
+	bool HasScriptResource = false;
+
+#if WITH_EDITORONLY_DATA
+	for (FNiagaraShaderScript& LoadedResource : LoadedScriptResources)
 	{
+		LoadedResource.RegisterShaderMap();
+
 		FNiagaraShaderMap* LoadedShaderMap = LoadedResource.GetGameThreadShaderMap();
 		if (LoadedShaderMap && LoadedShaderMap->GetShaderPlatform() == GMaxRHIShaderPlatform)
 		{
-			OutResourceForCurrentPlatform = LoadedResource;
+			HasScriptResource = true;
+			ScriptResource = MakeUnique<FNiagaraShaderScript>(LoadedResource);
 
 			ERHIFeatureLevel::Type LoadedFeatureLevel = LoadedShaderMap->GetShaderMapId().FeatureLevel;
-			if (!OutScriptResourcesLoaded[LoadedFeatureLevel])
+			if (!ScriptResourcesByFeatureLevel[LoadedFeatureLevel])
 			{
-				OutScriptResourcesLoaded[LoadedFeatureLevel] = Owner->AllocateResource();
+				ScriptResourcesByFeatureLevel[LoadedFeatureLevel] = AllocateResource();
 			}
 
-			OutScriptResourcesLoaded[LoadedFeatureLevel]->SetShaderMap(LoadedShaderMap);
-			OutResourceForCurrentPlatform.SetDataInterfaceParamInfo(Owner->GetVMExecutableData().DIParamInfo);
-
+			ScriptResourcesByFeatureLevel[LoadedFeatureLevel]->SetShaderMap(LoadedShaderMap);
 			break;
 		}
 		else
@@ -2203,6 +2214,18 @@ void ProcessSerializedShaderMaps(UNiagaraScript* Owner, TArray<FNiagaraShaderScr
 			LoadedResource.DiscardShaderMap();
 		}
 	}
+#else
+	if (ScriptResource.IsValid())
+	{
+		ScriptResource->RegisterShaderMap();
+		HasScriptResource = true;
+	}
+#endif
+
+	if (HasScriptResource)
+	{
+		ScriptResource->SetDataInterfaceParamInfo(Owner->GetVMExecutableData().DIParamInfo);
+	}
 }
 
 FNiagaraShaderScript* UNiagaraScript::AllocateResource()
