Fixing rename pin showing up when changing default modes.

diff --git a/Engine/Plugins/FX/Niagara/Source/NiagaraEditor/Private/NiagaraNodeParameterMapGet.cpp b/Engine/Plugins/FX/Niagara/Source/NiagaraEditor/Private/NiagaraNodeParameterMapGet.cpp
--- a/Engine/Plugins/FX/Niagara/Source/NiagaraEditor/Private/NiagaraNodeParameterMapGet.cpp
+++ b/Engine/Plugins/FX/Niagara/Source/NiagaraEditor/Private/NiagaraNodeParameterMapGet.cpp
@@ -438,6 +438,11 @@ bool UNiagaraNodeParameterMapGet::CancelEditablePinName(const FText& InName, UEd
 
 bool UNiagaraNodeParameterMapGet::CommitEditablePinName(const FText& InName, UEdGraphPin* InGraphPinObj, bool bSuppressEvents)
 {
+	if (InGraphPinObj == PinPendingRename)
+	{
+		PinPendingRename = nullptr;
+	}
+
 	if (Pins.Contains(InGraphPinObj) && InGraphPinObj->Direction == EEdGraphPinDirection::EGPD_Output)
 	{
 		FString OldPinName = InGraphPinObj->PinName.ToString();
diff --git a/Engine/Plugins/FX/Niagara/Source/NiagaraEditor/Private/NiagaraNodeParameterMapSet.cpp b/Engine/Plugins/FX/Niagara/Source/NiagaraEditor/Private/NiagaraNodeParameterMapSet.cpp
--- a/Engine/Plugins/FX/Niagara/Source/NiagaraEditor/Private/NiagaraNodeParameterMapSet.cpp
+++ b/Engine/Plugins/FX/Niagara/Source/NiagaraEditor/Private/NiagaraNodeParameterMapSet.cpp
@@ -169,6 +169,11 @@ bool UNiagaraNodeParameterMapSet::CancelEditablePinName(const FText& InName, UEd
 
 bool UNiagaraNodeParameterMapSet::CommitEditablePinName(const FText& InName, UEdGraphPin* InGraphPinObj, bool bSuppressEvents)
 {
+	if (InGraphPinObj == PinPendingRename)
+	{
+		PinPendingRename = nullptr;
+	}
+
 	if (Pins.Contains(InGraphPinObj))
 	{
 		
