FIx for bad merge with niagara shader script change

diff --git a/Engine/Plugins/FX/Niagara/Source/Niagara/Private/NiagaraScript.cpp b/Engine/Plugins/FX/Niagara/Source/Niagara/Private/NiagaraScript.cpp
--- a/Engine/Plugins/FX/Niagara/Source/Niagara/Private/NiagaraScript.cpp
+++ b/Engine/Plugins/FX/Niagara/Source/Niagara/Private/NiagaraScript.cpp
@@ -2192,8 +2192,6 @@ void UNiagaraScript::ProcessSerializedShaderMaps()
 #if WITH_EDITORONLY_DATA
 	for (FNiagaraShaderScript& LoadedResource : LoadedScriptResources)
 	{
-		LoadedResource.RegisterShaderMap();
-
 		FNiagaraShaderMap* LoadedShaderMap = LoadedResource.GetGameThreadShaderMap();
 		if (LoadedShaderMap && LoadedShaderMap->GetShaderPlatform() == GMaxRHIShaderPlatform)
 		{
@@ -2215,16 +2213,12 @@ void UNiagaraScript::ProcessSerializedShaderMaps()
 		}
 	}
 #else
-	if (ScriptResource.IsValid())
-	{
-		ScriptResource->RegisterShaderMap();
-		HasScriptResource = true;
-	}
+	HasScriptResource = ScriptResource.IsValid();
 #endif
 
 	if (HasScriptResource)
 	{
-		ScriptResource->SetDataInterfaceParamInfo(Owner->GetVMExecutableData().DIParamInfo);
+		ScriptResource->SetDataInterfaceParamInfo(CachedScriptVM.DIParamInfo);
 	}
 }
 
