- Fix for stomping over override parameter data when they are set inside a constuctor script

diff --git a/Engine/Plugins/FX/Niagara/Source/Niagara/Private/NiagaraComponent.cpp b/Engine/Plugins/FX/Niagara/Source/Niagara/Private/NiagaraComponent.cpp
--- a/Engine/Plugins/FX/Niagara/Source/Niagara/Private/NiagaraComponent.cpp
+++ b/Engine/Plugins/FX/Niagara/Source/Niagara/Private/NiagaraComponent.cpp
@@ -1509,99 +1509,119 @@ FNiagaraSystemInstance* UNiagaraComponent::GetSystemInstance() const
 
 void UNiagaraComponent::SetVariableLinearColor(FName InVariableName, const FLinearColor& InValue)
 {
-	OverrideParameters.SetParameterValue(InValue, FNiagaraVariable(FNiagaraTypeDefinition::GetColorDef(), InVariableName), true);
+	const FNiagaraVariable VariableDesc(FNiagaraTypeDefinition::GetColorDef(), InVariableName);
+	OverrideParameters.SetParameterValue(InValue, VariableDesc, true);
+#if WITH_EDITOR
+	SetParameterValueOverriddenLocally(VariableDesc, true, false);
+#endif
 }
 
 void UNiagaraComponent::SetNiagaraVariableLinearColor(const FString& InVariableName, const FLinearColor& InValue)
 {
-	FName VarName = FName(*InVariableName);
-
-	OverrideParameters.SetParameterValue(InValue, FNiagaraVariable(FNiagaraTypeDefinition::GetColorDef(), VarName), true);
+	SetVariableLinearColor(FName(*InVariableName), InValue);
 }
 
-
 void UNiagaraComponent::SetVariableQuat(FName InVariableName, const FQuat& InValue)
 {
-	OverrideParameters.SetParameterValue(InValue, FNiagaraVariable(FNiagaraTypeDefinition::GetQuatDef(), InVariableName), true);
+	const FNiagaraVariable VariableDesc(FNiagaraTypeDefinition::GetQuatDef(), InVariableName);
+	OverrideParameters.SetParameterValue(InValue, VariableDesc, true);
+#if WITH_EDITOR
+	SetParameterValueOverriddenLocally(VariableDesc, true, false);
+#endif
 }
 
 void UNiagaraComponent::SetNiagaraVariableQuat(const FString& InVariableName, const FQuat& InValue)
 {
-	FName VarName = FName(*InVariableName);
-
-	OverrideParameters.SetParameterValue(InValue, FNiagaraVariable(FNiagaraTypeDefinition::GetQuatDef(), VarName), true);
+	SetVariableQuat(FName(*InVariableName), InValue);
 }
 
 void UNiagaraComponent::SetVariableVec4(FName InVariableName, const FVector4& InValue)
 {
-	OverrideParameters.SetParameterValue(InValue, FNiagaraVariable(FNiagaraTypeDefinition::GetVec4Def(), InVariableName), true);
+	const FNiagaraVariable VariableDesc(FNiagaraTypeDefinition::GetVec4Def(), InVariableName);
+	OverrideParameters.SetParameterValue(InValue, VariableDesc, true);
+#if WITH_EDITOR
+	SetParameterValueOverriddenLocally(VariableDesc, true, false);
+#endif
 }
 
 void UNiagaraComponent::SetNiagaraVariableVec4(const FString& InVariableName, const FVector4& InValue)
 {
-	FName VarName = FName(*InVariableName);
-
-	OverrideParameters.SetParameterValue(InValue, FNiagaraVariable(FNiagaraTypeDefinition::GetVec4Def(), VarName), true);
+	SetVariableVec4(FName(*InVariableName), InValue);
 }
 
 void UNiagaraComponent::SetVariableVec3(FName InVariableName, FVector InValue)
 {
-	OverrideParameters.SetParameterValue(InValue, FNiagaraVariable(FNiagaraTypeDefinition::GetVec3Def(), InVariableName), true);
+	const FNiagaraVariable VariableDesc(FNiagaraTypeDefinition::GetVec3Def(), InVariableName);
+	OverrideParameters.SetParameterValue(InValue, VariableDesc, true);
+#if WITH_EDITOR
+	SetParameterValueOverriddenLocally(VariableDesc, true, false);
+#endif
 }
 
 void UNiagaraComponent::SetNiagaraVariableVec3(const FString& InVariableName, FVector InValue)
 {
-	FName VarName = FName(*InVariableName);
-	
-	OverrideParameters.SetParameterValue(InValue, FNiagaraVariable(FNiagaraTypeDefinition::GetVec3Def(), VarName), true);
+	SetVariableVec3(FName(*InVariableName), InValue);
 }
 
 void UNiagaraComponent::SetVariableVec2(FName InVariableName, FVector2D InValue)
 {
-	OverrideParameters.SetParameterValue(InValue, FNiagaraVariable(FNiagaraTypeDefinition::GetVec2Def(),InVariableName), true);
+	const FNiagaraVariable VariableDesc(FNiagaraTypeDefinition::GetVec2Def(), InVariableName);
+	OverrideParameters.SetParameterValue(InValue, VariableDesc, true);
+#if WITH_EDITOR
+	SetParameterValueOverriddenLocally(VariableDesc, true, false);
+#endif
 }
 
 void UNiagaraComponent::SetNiagaraVariableVec2(const FString& InVariableName, FVector2D InValue)
 {
-	FName VarName = FName(*InVariableName);
-
-	OverrideParameters.SetParameterValue(InValue, FNiagaraVariable(FNiagaraTypeDefinition::GetVec2Def(),VarName), true);
+	SetVariableVec2(FName(*InVariableName), InValue);
 }
 
 void UNiagaraComponent::SetVariableFloat(FName InVariableName, float InValue)
 {
-	OverrideParameters.SetParameterValue(InValue, FNiagaraVariable(FNiagaraTypeDefinition::GetFloatDef(), InVariableName), true);
+	const FNiagaraVariable VariableDesc(FNiagaraTypeDefinition::GetFloatDef(), InVariableName);
+	OverrideParameters.SetParameterValue(InValue, VariableDesc, true);
+#if WITH_EDITOR
+	SetParameterValueOverriddenLocally(VariableDesc, true, false);
+#endif
 }
 
 void UNiagaraComponent::SetNiagaraVariableFloat(const FString& InVariableName, float InValue)
 {
-	FName VarName = FName(*InVariableName);
-
-	OverrideParameters.SetParameterValue(InValue, FNiagaraVariable(FNiagaraTypeDefinition::GetFloatDef(), VarName), true);
+	SetVariableFloat(FName(*InVariableName), InValue);
 }
 
 void UNiagaraComponent::SetVariableInt(FName InVariableName, int32 InValue)
 {
-	OverrideParameters.SetParameterValue(InValue, FNiagaraVariable(FNiagaraTypeDefinition::GetIntDef(), InVariableName), true);
+	const FNiagaraVariable VariableDesc(FNiagaraTypeDefinition::GetIntDef(), InVariableName);
+	OverrideParameters.SetParameterValue(InValue, VariableDesc, true);
+#if WITH_EDITOR
+	SetParameterValueOverriddenLocally(VariableDesc, true, false);
+#endif
 }
 
 void UNiagaraComponent::SetNiagaraVariableInt(const FString& InVariableName, int32 InValue)
 {
-	FName VarName = FName(*InVariableName);
-
-	OverrideParameters.SetParameterValue(InValue, FNiagaraVariable(FNiagaraTypeDefinition::GetIntDef(), VarName), true);
+	SetVariableInt(FName(*InVariableName), InValue);
 }
 
 void UNiagaraComponent::SetVariableBool(FName InVariableName, bool InValue)
 {
-	OverrideParameters.SetParameterValue(InValue ? FNiagaraBool::True : FNiagaraBool::False, FNiagaraVariable(FNiagaraTypeDefinition::GetBoolDef(), InVariableName), true);
+	const FNiagaraVariable VariableDesc(FNiagaraTypeDefinition::GetBoolDef(), InVariableName);
+	OverrideParameters.SetParameterValue(InValue ? FNiagaraBool::True : FNiagaraBool::False, VariableDesc, true);
+#if WITH_EDITOR
+	SetParameterValueOverriddenLocally(VariableDesc, true, false);
+#endif
 }
 
 void UNiagaraComponent::SetNiagaraVariableBool(const FString& InVariableName, bool InValue)
 {
-	FName VarName = FName(*InVariableName);
+	SetVariableBool(FName(*InVariableName), InValue);
+}
 
-	OverrideParameters.SetParameterValue(InValue ? FNiagaraBool::True : FNiagaraBool::False, FNiagaraVariable(FNiagaraTypeDefinition::GetBoolDef(), VarName), true);
+void UNiagaraComponent::SetVariableActor(FName InVariableName, AActor* InValue)
+{
+	SetVariableObject(InVariableName, InValue);
 }
 
 void UNiagaraComponent::SetNiagaraVariableActor(const FString& InVariableName, AActor* InValue)
@@ -1609,26 +1629,27 @@ void UNiagaraComponent::SetNiagaraVariableActor(const FString& InVariableName, A
 	SetNiagaraVariableObject(InVariableName, InValue);
 }
 
-void UNiagaraComponent::SetVariableActor(FName InVariableName, AActor* InValue)
+void UNiagaraComponent::SetVariableObject(FName InVariableName, UObject* InValue)
 {
-	SetVariableObject(InVariableName, InValue);
+	const FNiagaraVariable VariableDesc(FNiagaraTypeDefinition::GetUObjectDef(), InVariableName);
+	OverrideParameters.SetUObject(InValue, VariableDesc);
+#if WITH_EDITOR
+	SetParameterValueOverriddenLocally(VariableDesc, true, false);
+#endif
 }
 
 void UNiagaraComponent::SetNiagaraVariableObject(const FString& InVariableName, UObject* InValue)
 {
-	FName VarName = FName(*InVariableName);
-
-	OverrideParameters.SetUObject(InValue, FNiagaraVariable(FNiagaraTypeDefinition::GetUObjectDef(), VarName));
-}
-
-void UNiagaraComponent::SetVariableObject(FName InVariableName, UObject* InValue)
-{
-	OverrideParameters.SetUObject(InValue, FNiagaraVariable(FNiagaraTypeDefinition::GetUObjectDef(), InVariableName));
+	SetVariableObject(FName(*InVariableName), InValue);
 }
 
 void UNiagaraComponent::SetVariableMaterial(FName InVariableName, UMaterialInterface* InValue)
 {
-	OverrideParameters.SetUObject(InValue, FNiagaraVariable(FNiagaraTypeDefinition::GetUMaterialDef(), InVariableName));
+	const FNiagaraVariable VariableDesc(FNiagaraTypeDefinition::GetUMaterialDef(), InVariableName);
+	OverrideParameters.SetUObject(InValue, VariableDesc);
+#if WITH_EDITOR
+	SetParameterValueOverriddenLocally(VariableDesc, true, false);
+#endif
 	MarkRenderStateDirty(); // Materials might be using this on the system, so invalidate the render state to re-gather them.
 }
 
