- Normalize parameter names in SetParameterValueOverriddenLocally as the new scheme can make them come through with "User." at the start

diff --git a/Engine/Plugins/FX/Niagara/Source/Niagara/Private/NiagaraComponent.cpp b/Engine/Plugins/FX/Niagara/Source/Niagara/Private/NiagaraComponent.cpp
--- a/Engine/Plugins/FX/Niagara/Source/Niagara/Private/NiagaraComponent.cpp
+++ b/Engine/Plugins/FX/Niagara/Source/Niagara/Private/NiagaraComponent.cpp
@@ -1977,8 +1977,16 @@ bool UNiagaraComponent::IsParameterValueOverriddenLocally(const FName& InParamNa
 
 void UNiagaraComponent::SetParameterValueOverriddenLocally(const FNiagaraVariable& InParam, bool bInOverriden, bool bRequiresSystemInstanceReset)
 {
-	bool* FoundVar = EditorOverridesValue.Find(InParam.GetName());
+	FName ParameterName = InParam.GetName();
+	{
+		FString ParameterString = ParameterName.ToString();
+		if (ParameterString.RemoveFromStart(TEXT("User.")))
+		{
+			ParameterName = FName(*ParameterString);
+		}
+	}
 
+	bool* FoundVar = EditorOverridesValue.Find(InParam.GetName());
 	if (FoundVar != nullptr && bInOverriden) 
 	{
 		*(FoundVar) = bInOverriden;
