Fixing functions that have parameter map inputs not mapping properly.

diff --git a/Engine/Plugins/FX/Niagara/Source/NiagaraEditor/Private/NiagaraHlslTranslator.cpp b/Engine/Plugins/FX/Niagara/Source/NiagaraEditor/Private/NiagaraHlslTranslator.cpp
--- a/Engine/Plugins/FX/Niagara/Source/NiagaraEditor/Private/NiagaraHlslTranslator.cpp
+++ b/Engine/Plugins/FX/Niagara/Source/NiagaraEditor/Private/NiagaraHlslTranslator.cpp
@@ -6164,9 +6164,18 @@ void FHlslNiagaraTranslator::GenerateFunctionCall(ENiagaraScriptUsage ScriptUsag
 			bool bSkip = false;
 			if (FunctionSignature.Outputs[i].GetType() == FNiagaraTypeDefinition::GetParameterMapDef())
 			{
-				if (i < FunctionSignature.Inputs.Num() && FunctionSignature.Inputs[i].GetType() == FNiagaraTypeDefinition::GetParameterMapDef())
+				int32 FoundInputParamMapIdx = INDEX_NONE;
+				for (int32 j = 0; j < FunctionSignature.Inputs.Num(); j++)
 				{
-					Output = Inputs[i];
+					if (FunctionSignature.Inputs[j].GetType() == FNiagaraTypeDefinition::GetParameterMapDef())
+					{
+						FoundInputParamMapIdx = j;
+						break;
+					}
+				}
+				if (FoundInputParamMapIdx < Inputs.Num() && FoundInputParamMapIdx != INDEX_NONE)
+				{
+					Output = Inputs[FoundInputParamMapIdx];
 				}
 				bSkip = true;
 			}
@@ -6661,7 +6670,11 @@ int32 FHlslNiagaraTranslator::CompileOutputPin(const UEdGraphPin* InPin)
 	// The incoming pin to compile may be pointing to a reroute node. If so, we just jump over it
 	// to where it really came from.
 	UEdGraphPin* Pin = UNiagaraNode::TraceOutputPin(const_cast<UEdGraphPin*>(InPin));
-	check(Pin->Direction == EGPD_Output);
+	if (Pin == nullptr || Pin->Direction != EGPD_Output)
+	{
+		Error(LOCTEXT("TraceOutputPinFailed", "Failed to trace pin to an output!"), Cast<UNiagaraNode>(InPin->GetOwningNode()), InPin);
+		return INDEX_NONE;
+	}
 
 	// The node can also replace our pin with another pin (e.g. in the case of static switches), so we need to make sure we don't run into a circular dependency
 	TSet<UEdGraphPin*> SeenPins;
