Copy-pasted modules will now retain their display name if one is set.

diff --git a/Engine/Plugins/FX/Niagara/Source/NiagaraEditor/Private/ViewModels/Stack/NiagaraStackModuleItem.cpp b/Engine/Plugins/FX/Niagara/Source/NiagaraEditor/Private/ViewModels/Stack/NiagaraStackModuleItem.cpp
--- a/Engine/Plugins/FX/Niagara/Source/NiagaraEditor/Private/ViewModels/Stack/NiagaraStackModuleItem.cpp
+++ b/Engine/Plugins/FX/Niagara/Source/NiagaraEditor/Private/ViewModels/Stack/NiagaraStackModuleItem.cpp
@@ -1104,6 +1104,8 @@ void UNiagaraStackModuleItem::Copy(UNiagaraClipboardContent* ClipboardContent) c
 		ClipboardFunction = UNiagaraClipboardFunction::CreateScriptFunction(ClipboardContent, FunctionCallNode->GetFunctionName(), FunctionCallNode->FunctionScript);
 	}
 
+	ClipboardFunction->DisplayName = CustomDisplayName.Get(FText::GetEmpty());
+
 	InputCollection->ToClipboardFunctionInputs(ClipboardFunction, ClipboardFunction->Inputs);
 	ClipboardContent->Functions.Add(ClipboardFunction);
 }
diff --git a/Engine/Plugins/FX/Niagara/Source/NiagaraEditor/Private/ViewModels/Stack/NiagaraStackScriptItemGroup.cpp b/Engine/Plugins/FX/Niagara/Source/NiagaraEditor/Private/ViewModels/Stack/NiagaraStackScriptItemGroup.cpp
--- a/Engine/Plugins/FX/Niagara/Source/NiagaraEditor/Private/ViewModels/Stack/NiagaraStackScriptItemGroup.cpp
+++ b/Engine/Plugins/FX/Niagara/Source/NiagaraEditor/Private/ViewModels/Stack/NiagaraStackScriptItemGroup.cpp
@@ -1225,6 +1225,11 @@ void UNiagaraStackScriptItemGroup::PasteModules(const UNiagaraClipboardContent*
 				FunctionInputs = ClipboardFunction->Inputs;
 			}
 			(*NewModuleItemPtr)->SetInputValuesFromClipboardFunctionInputs(FunctionInputs);
+
+			if (!ClipboardFunction->DisplayName.IsEmptyOrWhitespace())
+			{
+				(*NewModuleItemPtr)->OnRenamed(ClipboardFunction->DisplayName);
+			}
 		}
 	}
 
diff --git a/Engine/Plugins/FX/Niagara/Source/NiagaraEditor/Public/NiagaraClipboard.h b/Engine/Plugins/FX/Niagara/Source/NiagaraEditor/Public/NiagaraClipboard.h
--- a/Engine/Plugins/FX/Niagara/Source/NiagaraEditor/Public/NiagaraClipboard.h
+++ b/Engine/Plugins/FX/Niagara/Source/NiagaraEditor/Public/NiagaraClipboard.h
@@ -92,6 +92,9 @@ public:
 	UPROPERTY()
 	FString FunctionName;
 
+	UPROPERTY()
+	FText DisplayName;
+
 	UPROPERTY()
 	ENiagaraClipboardFunctionScriptMode ScriptMode;
 
