Fixed compile issue due to refactor.

diff --git a/Engine/Plugins/FX/Niagara/Source/NiagaraEditor/Private/ViewModels/Stack/NiagaraStackModuleItem.cpp b/Engine/Plugins/FX/Niagara/Source/NiagaraEditor/Private/ViewModels/Stack/NiagaraStackModuleItem.cpp
--- a/Engine/Plugins/FX/Niagara/Source/NiagaraEditor/Private/ViewModels/Stack/NiagaraStackModuleItem.cpp
+++ b/Engine/Plugins/FX/Niagara/Source/NiagaraEditor/Private/ViewModels/Stack/NiagaraStackModuleItem.cpp
@@ -1104,7 +1104,7 @@ void UNiagaraStackModuleItem::Copy(UNiagaraClipboardContent* ClipboardContent) c
 		ClipboardFunction = UNiagaraClipboardFunction::CreateScriptFunction(ClipboardContent, FunctionCallNode->GetFunctionName(), FunctionCallNode->FunctionScript);
 	}
 
-	ClipboardFunction->DisplayName = CustomDisplayName.Get(FText::GetEmpty());
+	ClipboardFunction->DisplayName = GetAlternateDisplayName().Get(FText::GetEmpty());
 
 	InputCollection->ToClipboardFunctionInputs(ClipboardFunction, ClipboardFunction->Inputs);
 	ClipboardContent->Functions.Add(ClipboardFunction);
