Merging my recent chagnes to Release-4.25.

diff --git a/Engine/Plugins/FX/Niagara/Source/Niagara/Private/NiagaraEmitterInstance.cpp b/Engine/Plugins/FX/Niagara/Source/Niagara/Private/NiagaraEmitterInstance.cpp
--- a/Engine/Plugins/FX/Niagara/Source/Niagara/Private/NiagaraEmitterInstance.cpp
+++ b/Engine/Plugins/FX/Niagara/Source/Niagara/Private/NiagaraEmitterInstance.cpp
@@ -1121,6 +1121,13 @@ void FNiagaraEmitterInstance::Tick(float DeltaSeconds)
 	int32 OrigNumParticles = GetNumParticles();
 	int32 AllocationEstimate = CachedEmitter->GetMaxParticleCountEstimate();
 	int32 RequiredSize = OrigNumParticles + SpawnTotal + EventSpawnTotal;
+
+	if (RequiredSize == 0)
+	{
+		//Early out if we have no particles to process.
+		//return;
+	}
+
 	int32 AllocationSize = FMath::Max<int32>(AllocationEstimate, RequiredSize);
 	if (AllocationSize > MaxAllocationCount)
 	{
diff --git a/Engine/Plugins/FX/Niagara/Source/Niagara/Public/NiagaraCommon.h b/Engine/Plugins/FX/Niagara/Source/Niagara/Public/NiagaraCommon.h
--- a/Engine/Plugins/FX/Niagara/Source/Niagara/Public/NiagaraCommon.h
+++ b/Engine/Plugins/FX/Niagara/Source/Niagara/Public/NiagaraCommon.h
@@ -547,14 +547,14 @@ struct NIAGARA_API FNiagaraSystemUpdateContext
 {
 	GENERATED_BODY()
 
-	FNiagaraSystemUpdateContext(const UNiagaraSystem* System, bool bReInit) : bDestroyOnAdd(false), bOnlyActive(false) { Add(System, bReInit); }
+	FNiagaraSystemUpdateContext(const UNiagaraSystem* System, bool bReInit, bool bInDestroyOnAdd = false, bool bInOnlyActive = false) :bDestroyOnAdd(bInDestroyOnAdd), bOnlyActive(bInOnlyActive) { Add(System, bReInit); }
 #if WITH_EDITORONLY_DATA
-	FNiagaraSystemUpdateContext(const UNiagaraEmitter* Emitter, bool bReInit) : bDestroyOnAdd(false), bOnlyActive(false) { Add(Emitter, bReInit); }
-	FNiagaraSystemUpdateContext(const UNiagaraScript* Script, bool bReInit) :bDestroyOnAdd(false), bOnlyActive(false) { Add(Script, bReInit); }
+	FNiagaraSystemUpdateContext(const UNiagaraEmitter* Emitter, bool bReInit, bool bInDestroyOnAdd = false, bool bInOnlyActive = false) :bDestroyOnAdd(bInDestroyOnAdd), bOnlyActive(bInOnlyActive)  { Add(Emitter, bReInit); }
+	FNiagaraSystemUpdateContext(const UNiagaraScript* Script, bool bReInit, bool bInDestroyOnAdd = false, bool bInOnlyActive = false) :bDestroyOnAdd(bInDestroyOnAdd), bOnlyActive(bInOnlyActive)  { Add(Script, bReInit); }
 	//FNiagaraSystemUpdateContext(UNiagaraDataInterface* Interface, bool bReinit) : Add(Interface, bReinit) {}
-	FNiagaraSystemUpdateContext(const UNiagaraParameterCollection* Collection, bool bReInit) :bDestroyOnAdd(false), bOnlyActive(false) { Add(Collection, bReInit); }
+	FNiagaraSystemUpdateContext(const UNiagaraParameterCollection* Collection, bool bReInit, bool bInDestroyOnAdd = false, bool bInOnlyActive = false) :bDestroyOnAdd(bInDestroyOnAdd), bOnlyActive(bInOnlyActive) { Add(Collection, bReInit); }
 #endif
-	FNiagaraSystemUpdateContext():bDestroyOnAdd(false), bOnlyActive(false) { }
+	FNiagaraSystemUpdateContext():bDestroyOnAdd(false), bOnlyActive(false){ }
 
 	~FNiagaraSystemUpdateContext();
 
