- When setting a texture object ensure the editor will not stomp over the data if the system is duplicated by letting the component know that it was overriden

diff --git a/Engine/Plugins/FX/Niagara/Source/Niagara/Private/NiagaraFunctionLibrary.cpp b/Engine/Plugins/FX/Niagara/Source/Niagara/Private/NiagaraFunctionLibrary.cpp
--- a/Engine/Plugins/FX/Niagara/Source/Niagara/Private/NiagaraFunctionLibrary.cpp
+++ b/Engine/Plugins/FX/Niagara/Source/Niagara/Private/NiagaraFunctionLibrary.cpp
@@ -422,6 +422,9 @@ void UNiagaraFunctionLibrary::SetTextureObject(UNiagaraComponent* NiagaraSystem,
 		return;
 	}
 
+#if WITH_EDITOR
+	NiagaraSystem->SetParameterValueOverriddenLocally(Variable, true, false);
+#endif
 	TexureDI->SetTexture(Texture);
 }
 
@@ -457,6 +460,9 @@ void UNiagaraFunctionLibrary::SetVolumeTextureObject(UNiagaraComponent* NiagaraS
 		return;
 	}
 
+#if WITH_EDITOR
+	NiagaraSystem->SetParameterValueOverriddenLocally(Variable, true, false);
+#endif
 	TexureDI->SetTexture(Texture);
 }
 
