Fixing issue where dragging particles variables onto a set node caused the UI to be wrong on the scope selector.

diff --git a/Engine/Plugins/FX/Niagara/Source/NiagaraEditor/Private/ViewModels/NiagaraParameterNameViewModel.cpp b/Engine/Plugins/FX/Niagara/Source/NiagaraEditor/Private/ViewModels/NiagaraParameterNameViewModel.cpp
--- a/Engine/Plugins/FX/Niagara/Source/NiagaraEditor/Private/ViewModels/NiagaraParameterNameViewModel.cpp
+++ b/Engine/Plugins/FX/Niagara/Source/NiagaraEditor/Private/ViewModels/NiagaraParameterNameViewModel.cpp
@@ -88,10 +88,6 @@ int32 FNiagaraParameterPanelEntryParameterNameViewModel::GetScopeValue() const
 	{
 		return int32(ENiagaraParameterScope::DISPLAY_ONLY_StaticSwitch);
 	}
-	else if (CachedScriptVarAndViewInfo.MetaData.GetUsage() == ENiagaraScriptParameterUsage::Output)
-	{
-		return int32(ENiagaraParameterScope::Output);
-	}
 
 	ENiagaraParameterScope CachedScope;
 	if (ensureMsgf(FNiagaraEditorUtilities::GetVariableMetaDataScope(CachedScriptVarAndViewInfo.MetaData, CachedScope), TEXT("Failed to get scope value for param as override namespace is set! This method should not be bound!")))
@@ -208,10 +204,6 @@ int32 FNiagaraGraphPinParameterNameViewModel::GetScopeValue() const
 	{
 		return int32(ENiagaraParameterScope::DISPLAY_ONLY_StaticSwitch);
 	}
-	else if (CachedScriptVarAndViewInfo.MetaData.GetUsage() == ENiagaraScriptParameterUsage::Output)
-	{
-		return int32(ENiagaraParameterScope::Output);
-	}
 
 	ENiagaraParameterScope CachedScope;
 	if (ensureMsgf(FNiagaraEditorUtilities::GetVariableMetaDataScope(CachedScriptVarAndViewInfo.MetaData, CachedScope), TEXT("Failed to get scope value for param as override namespace is set! This method should not be bound!")))
