Makign get map nodes default to input params and set nodes default to output params.

diff --git a/Engine/Plugins/FX/Niagara/Source/NiagaraEditor/Private/NiagaraNodeParameterMapBase.h b/Engine/Plugins/FX/Niagara/Source/NiagaraEditor/Private/NiagaraNodeParameterMapBase.h
--- a/Engine/Plugins/FX/Niagara/Source/NiagaraEditor/Private/NiagaraNodeParameterMapBase.h
+++ b/Engine/Plugins/FX/Niagara/Source/NiagaraEditor/Private/NiagaraNodeParameterMapBase.h
@@ -4,6 +4,7 @@
 
 #include "NiagaraNodeWithDynamicPins.h"
 #include "SGraphPin.h"
+#include "NiagaraConstants.h"
 #include "NiagaraNodeParameterMapBase.generated.h"
 
 class UEdGraphPin;
@@ -44,6 +45,8 @@ public:
 
 	virtual bool CanRenamePinFromContextMenu(const UEdGraphPin* Pin) const { return false; }
 
+	virtual FName GetNewPinDefaultNamespace() const { return PARAM_MAP_LOCAL_MODULE_STR; }
+
 protected:
 	virtual void SelectParameterFromPin(UEdGraphPin* InPin);
 	virtual void OnPinRenamed(UEdGraphPin* RenamedPin, const FString& OldName) override;
diff --git a/Engine/Plugins/FX/Niagara/Source/NiagaraEditor/Private/NiagaraNodeParameterMapGet.h b/Engine/Plugins/FX/Niagara/Source/NiagaraEditor/Private/NiagaraNodeParameterMapGet.h
--- a/Engine/Plugins/FX/Niagara/Source/NiagaraEditor/Private/NiagaraNodeParameterMapGet.h
+++ b/Engine/Plugins/FX/Niagara/Source/NiagaraEditor/Private/NiagaraNodeParameterMapGet.h
@@ -43,6 +43,8 @@ public:
 	void GatherExternalDependencyData(ENiagaraScriptUsage InMasterUsage, const FGuid& InMasterUsageId, TArray<FNiagaraCompileHash>& InReferencedCompileHashes, TArray<FString>& InReferencedObjs) const override;
 	virtual void GetPinHoverText(const UEdGraphPin& Pin, FString& HoverTextOut) const override;
 
+	virtual FName GetNewPinDefaultNamespace() const { return PARAM_MAP_MODULE_STR; }
+
 protected:
 	virtual void OnNewTypedPinAdded(UEdGraphPin* NewPin) override;
 	virtual void OnPinRenamed(UEdGraphPin* RenamedPin, const FString& OldName) override;
diff --git a/Engine/Plugins/FX/Niagara/Source/NiagaraEditor/Private/NiagaraNodeParameterMapSet.cpp b/Engine/Plugins/FX/Niagara/Source/NiagaraEditor/Private/NiagaraNodeParameterMapSet.cpp
--- a/Engine/Plugins/FX/Niagara/Source/NiagaraEditor/Private/NiagaraNodeParameterMapSet.cpp
+++ b/Engine/Plugins/FX/Niagara/Source/NiagaraEditor/Private/NiagaraNodeParameterMapSet.cpp
@@ -100,7 +100,7 @@ void UNiagaraNodeParameterMapSet::OnNewTypedPinAdded(UEdGraphPin* NewPin)
 		FName PinNameWithoutNamespace;
 		if (FNiagaraEditorUtilities::DecomposeVariableNamespace(NewPinName, PinNameWithoutNamespace).Num() == 0)
 		{
-			NewPinName = *(PARAM_MAP_LOCAL_MODULE_STR +  NewPinName.ToString());
+			NewPinName = *(PARAM_MAP_OUTPUT_MODULE_STR +  NewPinName.ToString());
 			bCreatedNamespace = true;
 		}
 
diff --git a/Engine/Plugins/FX/Niagara/Source/NiagaraEditor/Private/NiagaraNodeParameterMapSet.h b/Engine/Plugins/FX/Niagara/Source/NiagaraEditor/Private/NiagaraNodeParameterMapSet.h
--- a/Engine/Plugins/FX/Niagara/Source/NiagaraEditor/Private/NiagaraNodeParameterMapSet.h
+++ b/Engine/Plugins/FX/Niagara/Source/NiagaraEditor/Private/NiagaraNodeParameterMapSet.h
@@ -37,6 +37,7 @@ public:
 	virtual void GetNodeContextMenuActions(class UToolMenu* Menu, class UGraphNodeContextMenuContext* Context) const override;
 	virtual bool IncludeParentNodeContextMenu() const { return true; }
 	virtual void PostLoad() override;
+	virtual FName GetNewPinDefaultNamespace() const { return PARAM_MAP_OUTPUT_MODULE_STR; }
 
 protected:
 	virtual void OnNewTypedPinAdded(UEdGraphPin* NewPin); 
diff --git a/Engine/Plugins/FX/Niagara/Source/NiagaraEditor/Private/Widgets/SNiagaraGraphPinAdd.cpp b/Engine/Plugins/FX/Niagara/Source/NiagaraEditor/Private/Widgets/SNiagaraGraphPinAdd.cpp
--- a/Engine/Plugins/FX/Niagara/Source/NiagaraEditor/Private/Widgets/SNiagaraGraphPinAdd.cpp
+++ b/Engine/Plugins/FX/Niagara/Source/NiagaraEditor/Private/Widgets/SNiagaraGraphPinAdd.cpp
@@ -74,10 +74,11 @@ TSharedRef<SWidget> SNiagaraGraphPinAdd::OnGetAddButtonMenuContent()
 			TArray<TWeakObjectPtr<UNiagaraGraph>> Graphs;
 			Graphs.Add(OwningNode->GetNiagaraGraph());
 
+
 			// Creating a new parameter from a graph pin always defaults to Particles scope.
 			TSharedRef<SNiagaraAddParameterMenu2> MenuWidget = SNew(SNiagaraAddParameterMenu2, Graphs)
 				.NewParameterScope(ENiagaraParameterScope::Local)
-				.NewParameterNamespace(PARAM_MAP_LOCAL_MODULE_STR)
+				.NewParameterNamespace(Cast<UNiagaraNodeParameterMapBase>(OwningNode)->GetNewPinDefaultNamespace().ToString())
 				.ShowKnownConstantParametersFilter(NiagaraParameterPanelSectionID::Type::REFERENCES)
 				.OnAddParameter_UObject(OwningNode, &UNiagaraNodeWithDynamicPins::AddParameter, (const UEdGraphPin*)GetPinObj())
 				.OnCollectCustomActions_UObject(OwningNode, &UNiagaraNodeWithDynamicPins::CollectAddPinActions, GetPinObj())
