- Integrate CL 12234004: Fix crash when editing a comment

#rnx

diff --git a/Engine/Plugins/FX/Niagara/Source/NiagaraEditorWidgets/Private/SNiagaraOverviewGraph.cpp b/Engine/Plugins/FX/Niagara/Source/NiagaraEditorWidgets/Private/SNiagaraOverviewGraph.cpp
--- a/Engine/Plugins/FX/Niagara/Source/NiagaraEditorWidgets/Private/SNiagaraOverviewGraph.cpp
+++ b/Engine/Plugins/FX/Niagara/Source/NiagaraEditorWidgets/Private/SNiagaraOverviewGraph.cpp
@@ -255,11 +255,13 @@ void SNiagaraOverviewGraph::OnClearIsolated()
 bool SNiagaraOverviewGraph::OnVerifyNodeTitle(const FText& NewText, UEdGraphNode* Node, FText& OutErrorMessage) const
 {
 	UNiagaraOverviewNode* NiagaraNode = Cast<UNiagaraOverviewNode>(Node);
-
-	TSharedPtr<FNiagaraEmitterHandleViewModel> NodeEmitterHandleViewModel = ViewModel->GetSystemViewModel()->GetEmitterHandleViewModelById(NiagaraNode->GetEmitterHandleGuid());
-	if (ensureMsgf(NodeEmitterHandleViewModel.IsValid(), TEXT("Failed to find EmitterHandleViewModel with matching Emitter GUID to Overview Node!")))
+	if (NiagaraNode != nullptr)
 	{
-		return NodeEmitterHandleViewModel->VerifyNameTextChanged(NewText, OutErrorMessage);
+		TSharedPtr<FNiagaraEmitterHandleViewModel> NodeEmitterHandleViewModel = ViewModel->GetSystemViewModel()->GetEmitterHandleViewModelById(NiagaraNode->GetEmitterHandleGuid());
+		if (ensureMsgf(NodeEmitterHandleViewModel.IsValid(), TEXT("Failed to find EmitterHandleViewModel with matching Emitter GUID to Overview Node!")))
+		{
+			return NodeEmitterHandleViewModel->VerifyNameTextChanged(NewText, OutErrorMessage);
+		}
 	}
 
 	return true;
