Reorganize FNiagaraEmitterInstance layout, including moving event data into it's own struct so that we can avoid the cost unless we're using events

diff --git a/Engine/Plugins/FX/Niagara/Source/Niagara/Classes/NiagaraEmitterInstance.h b/Engine/Plugins/FX/Niagara/Source/Niagara/Classes/NiagaraEmitterInstance.h
--- a/Engine/Plugins/FX/Niagara/Source/Niagara/Classes/NiagaraEmitterInstance.h
+++ b/Engine/Plugins/FX/Niagara/Source/Niagara/Classes/NiagaraEmitterInstance.h
@@ -28,9 +28,25 @@ struct FNiagaraEmitterCompiledData;
 */
 class FNiagaraEmitterInstance
 {
+private:
+	struct FEventInstanceData
+	{
+		TArray<FNiagaraScriptExecutionContext> EventExecContexts;
+		TArray<FNiagaraParameterDirectBinding<int32>> EventExecCountBindings;
+
+		TArray<FNiagaraDataSet*> UpdateScriptEventDataSets;
+		TArray<FNiagaraDataSet*> SpawnScriptEventDataSets;
+
+		TArray<bool> UpdateEventGeneratorIsSharedByIndex;
+		TArray<bool> SpawnEventGeneratorIsSharedByIndex;
+
+		/** Data required for handling events. */
+		TArray<FNiagaraEventHandlingInfo> EventHandlingInfo;
+		int32 EventSpawnTotal = 0;
+	};
+
 public:
 	explicit FNiagaraEmitterInstance(FNiagaraSystemInstance* InParentSystemInstance);
-	bool bDumpAfterEvent;
 	virtual ~FNiagaraEmitterInstance();
 
 	void Init(int32 InEmitterIdx, FNiagaraSystemInstanceID SystemInstanceID);
@@ -39,7 +55,7 @@ public:
 
 	void DirtyDataInterfaces();
 
-	/** Replaces the binding for a single parameter colleciton instance. If for example the component begins to override the global instance. */
+	/** Replaces the binding for a single parameter collection instance. If for example the component begins to override the global instance. */
 	//void RebindParameterCollection(UNiagaraParameterCollectionInstance* OldInstance, UNiagaraParameterCollectionInstance* NewInstance);
 	void BindParameters(bool bExternalOnly);
 	void UnbindParameters(bool bExternalOnly);
@@ -97,7 +113,7 @@ public:
 
 	FNiagaraScriptExecutionContext& GetSpawnExecutionContext() { return SpawnExecContext; }
 	FNiagaraScriptExecutionContext& GetUpdateExecutionContext() { return UpdateExecContext; }
-	TArray<FNiagaraScriptExecutionContext>& GetEventExecutionContexts() { return EventExecContexts; }
+	TArrayView<FNiagaraScriptExecutionContext> GetEventExecutionContexts();
 
 	FORCEINLINE FName GetCachedIDName()const { return CachedIDName; }
 	FORCEINLINE UNiagaraEmitter* GetCachedEmitter()const { return CachedEmitter; }
@@ -131,36 +147,13 @@ private:
 
 	/** Generate emitter bounds */
 	FBox InternalCalculateDynamicBounds(int32 ParticleCount) const;
-	
-	/** The index of our emitter in our parent system instance. */
-	int32 EmitterIdx;
-
-	/* The age of the emitter*/
-	float EmitterAge;
-
-	int32 TickCount;
-
-	int32 TotalSpawnedParticles;
-	
-	/** Typical resets must be deferred until the tick as the RT could still be using the current buffer. */
-	uint32 bResetPending : 1;
-
-	/* Cycles taken to process the tick. */
-	uint32 CPUTimeCycles;
-	/* Emitter tick state */
-	ENiagaraExecutionState ExecutionState;
-	/* Emitter bounds */
-	FBox CachedBounds;
-
-	uint32 MaxRuntimeAllocation;
 
 	/** Array of all spawn info driven by our owning emitter script. */
 	TArray<FNiagaraSpawnInfo> SpawnInfos;
 
 	FNiagaraScriptExecutionContext SpawnExecContext;
 	FNiagaraScriptExecutionContext UpdateExecContext;
-	FNiagaraComputeExecutionContext* GPUExecContext;
-	TArray<FNiagaraScriptExecutionContext> EventExecContexts;
+	FNiagaraComputeExecutionContext* GPUExecContext = nullptr;
 
 	FNiagaraParameterDirectBinding<float> SpawnIntervalBinding;
 	FNiagaraParameterDirectBinding<float> InterpSpawnStartBinding;
@@ -168,44 +161,59 @@ private:
 
 	FNiagaraParameterDirectBinding<int32> SpawnExecCountBinding;
 	FNiagaraParameterDirectBinding<int32> UpdateExecCountBinding;
-	TArray<FNiagaraParameterDirectBinding<int32>> EventExecCountBindings;
-	
-	/** particle simulation data. Must be a shared ref as various things on the RT can have direct ref to it. */
-	FNiagaraDataSet* ParticleDataSet;
 
-	FNiagaraSystemInstance *ParentSystemInstance;
+	/** Optional list of bounds calculators. */
+	TArray<TUniquePtr<FNiagaraBoundsCalculator>, TInlineAllocator<1>> BoundsCalculators;
 
-	/** Raw pointer to the emitter that we're instanced from. Raw ptr should be safe here as we check for the validity of the system and it's emitters higher up before any ticking. */
-	UNiagaraEmitter* CachedEmitter;
-	FName CachedIDName;
+	TSharedPtr<const FNiagaraEmitterCompiledData> CachedEmitterCompiledData;
 
-	TArray<FNiagaraDataSet*> UpdateScriptEventDataSets;
-	TArray<FNiagaraDataSet*> SpawnScriptEventDataSets;
+	TUniquePtr<FEventInstanceData> EventInstanceData;
 
-	TArray<bool> UpdateEventGeneratorIsSharedByIndex;
-	TArray<bool> SpawnEventGeneratorIsSharedByIndex;
+	/** A parameter store which contains the data interfaces parameters which were defined by the scripts. */
+	FNiagaraParameterStore ScriptDefinedDataInterfaceParameters;
 
-	FNiagaraSystemInstanceID OwnerSystemInstanceID;
+	/* Emitter bounds */
+	FBox CachedBounds;
 
 	/** Cached fixed bounds of the parent system which override this Emitter Instances bounds if set. Whenever we initialize the owning SystemInstance we will reconstruct this
 	 ** EmitterInstance and the cached bounds will be unset. */
 	TOptional<FBox> CachedSystemFixedBounds;
 
-	/** A parameter store which contains the data interfaces parameters which were defined by the scripts. */
-	FNiagaraParameterStore ScriptDefinedDataInterfaceParameters;
+	FNiagaraSystemInstanceID OwnerSystemInstanceID;
 
 	NiagaraEmitterInstanceBatcher* Batcher = nullptr;
 
-	/** Data required for handling events. */
-	TArray<FNiagaraEventHandlingInfo> EventHandlingInfo;
-	int32 EventSpawnTotal;
+	/** particle simulation data. Must be a shared ref as various things on the RT can have direct ref to it. */
+	FNiagaraDataSet* ParticleDataSet = nullptr;
+
+	FNiagaraSystemInstance *ParentSystemInstance = nullptr;
+
+	/** Raw pointer to the emitter that we're instanced from. Raw ptr should be safe here as we check for the validity of the system and it's emitters higher up before any ticking. */
+	UNiagaraEmitter* CachedEmitter = nullptr;
+	FName CachedIDName;
+
+	/** The index of our emitter in our parent system instance. */
+	int32 EmitterIdx = INDEX_NONE;
+
+	/* The age of the emitter*/
+	float EmitterAge = 0.0f;
+
+	int32 TickCount = 0;
+
+	int32 TotalSpawnedParticles = 0;
+	
+	/* Cycles taken to process the tick. */
+	uint32 CPUTimeCycles = 0;
+
+	uint32 MaxRuntimeAllocation = 0;
 
 	int32 MaxAllocationCount = 0;
 	int32 MinOverallocation = -1;
 	int32 ReallocationCount = 0;
 
-	/** Optional list of bounds calculators. */
-	TArray<TUniquePtr<FNiagaraBoundsCalculator>, TInlineAllocator<1>> BoundsCalculators;
+	/* Emitter tick state */
+	ENiagaraExecutionState ExecutionState = ENiagaraExecutionState::Inactive;
 
-	TSharedPtr<const FNiagaraEmitterCompiledData> CachedEmitterCompiledData;
+	/** Typical resets must be deferred until the tick as the RT could still be using the current buffer. */
+	uint32 bResetPending : 1;
 };
diff --git a/Engine/Plugins/FX/Niagara/Source/Niagara/Private/NiagaraEmitterInstance.cpp b/Engine/Plugins/FX/Niagara/Source/Niagara/Private/NiagaraEmitterInstance.cpp
--- a/Engine/Plugins/FX/Niagara/Source/Niagara/Private/NiagaraEmitterInstance.cpp
+++ b/Engine/Plugins/FX/Niagara/Source/Niagara/Private/NiagaraEmitterInstance.cpp
@@ -29,6 +29,7 @@ DECLARE_CYCLE_STAT(TEXT("Init Emitters [GT]"), STAT_NiagaraEmitterInit, STATGROU
 DECLARE_CYCLE_STAT(TEXT("Init Emitters (DataSets) [GT]"), STAT_NiagaraEmitterInit_Datasets, STATGROUP_Niagara);
 DECLARE_CYCLE_STAT(TEXT("Init Emitters (ExecContexts) [GT]"), STAT_NiagaraEmitterInit_ExecContexts, STATGROUP_Niagara);
 DECLARE_CYCLE_STAT(TEXT("Init Emitters (DirectBindings) [GT]"), STAT_NiagaraEmitterInit_DirectBindings, STATGROUP_Niagara);
+DECLARE_CYCLE_STAT(TEXT("Init Emitters (Events) [GT]"), STAT_NiagaraEmitterInit_Events, STATGROUP_Niagara);
 DECLARE_CYCLE_STAT(TEXT("Init Emitters (DIDefaultsAndBoundCalcs) [GT]"), STAT_NiagaraEmitterInit_DIDefaultsAndBoundsCalcs, STATGROUP_Niagara);
 
 
@@ -118,16 +119,10 @@ struct FNiagaraEditorOnlyCycleTimer
 //////////////////////////////////////////////////////////////////////////
 
 FNiagaraEmitterInstance::FNiagaraEmitterInstance(FNiagaraSystemInstance* InParentSystemInstance)
-: CPUTimeCycles(0)
-, ExecutionState(ENiagaraExecutionState::Inactive)
-, CachedBounds(ForceInit)
-, GPUExecContext(nullptr)
-, ParentSystemInstance(InParentSystemInstance)
-, CachedEmitter(nullptr)
+: CachedBounds(ForceInit)
 , CachedSystemFixedBounds()
-, EventSpawnTotal(0)
+, ParentSystemInstance(InParentSystemInstance)
 {
-	bDumpAfterEvent = false;
 	ParticleDataSet = new FNiagaraDataSet();
 
 	Batcher = ParentSystemInstance ? ParentSystemInstance->GetBatcher() : nullptr;
@@ -183,6 +178,15 @@ FBox FNiagaraEmitterInstance::GetBounds()
 	return CachedBounds;
 }
 
+TArrayView<FNiagaraScriptExecutionContext> FNiagaraEmitterInstance::GetEventExecutionContexts()
+{
+	if (EventInstanceData.IsValid())
+	{
+		return MakeArrayView(EventInstanceData->EventExecContexts);
+	}
+	return MakeArrayView<FNiagaraScriptExecutionContext>(nullptr, 0);
+}
+
 bool FNiagaraEmitterInstance::IsReadyToRun() const
 {
 	if (!IsDisabled() && !CachedEmitter->IsReadyToRun())
@@ -266,7 +270,6 @@ void FNiagaraEmitterInstance::Init(int32 InEmitterIdx, FNiagaraSystemInstanceID
 	}
 
 	//Init the spawn infos to the correct number for this system.
-	int32 NumEvents = CachedEmitter->GetEventHandlers().Num();
 	SpawnInfos.SetNum(CachedEmitterCompiledData->SpawnAttributes.Num());
 
 	{
@@ -300,34 +303,61 @@ void FNiagaraEmitterInstance::Init(int32 InEmitterIdx, FNiagaraSystemInstanceID
 			}
 		}
 #endif
+	}
 
+	{
 		ensure(CachedEmitter->UpdateScriptProps.DataSetAccessSynchronized());
-		UpdateScriptEventDataSets.Empty();
-		UpdateEventGeneratorIsSharedByIndex.SetNumZeroed(CachedEmitter->UpdateScriptProps.EventGenerators.Num());
-		int32 UpdateEventGeneratorIndex = 0;
-		for (const FNiagaraEventGeneratorProperties &GeneratorProps : CachedEmitter->UpdateScriptProps.EventGenerators)
-		{
-			FNiagaraDataSet *Set = ParentSystemInstance->CreateEventDataSet(EmitterHandle.GetIdName(), GeneratorProps.ID);
-			Set->Init(&GeneratorProps.DataSetCompiledData);
-
-			UpdateScriptEventDataSets.Add(Set);
-			UpdateEventGeneratorIsSharedByIndex[UpdateEventGeneratorIndex] = CachedEmitter->IsEventGeneratorShared(GeneratorProps.ID);
-			UpdateEventGeneratorIndex++;
-		}
+		const int32 UpdateEventGeneratorCount = CachedEmitter->UpdateScriptProps.EventGenerators.Num();
 
 		ensure(CachedEmitter->SpawnScriptProps.DataSetAccessSynchronized());
-		SpawnScriptEventDataSets.Empty();
-		SpawnEventGeneratorIsSharedByIndex.SetNumZeroed(CachedEmitter->SpawnScriptProps.EventGenerators.Num());
-		int32 SpawnEventGeneratorIndex = 0;
-		for (const FNiagaraEventGeneratorProperties &GeneratorProps : CachedEmitter->SpawnScriptProps.EventGenerators)
+		const int32 SpawnEventGeneratorCount = CachedEmitter->SpawnScriptProps.EventGenerators.Num();
+
+		const int32 NumEvents = CachedEmitter->GetEventHandlers().Num();
+
+		if (UpdateEventGeneratorCount || SpawnEventGeneratorCount || NumEvents)
 		{
-			FNiagaraDataSet *Set = ParentSystemInstance->CreateEventDataSet(EmitterHandle.GetIdName(), GeneratorProps.ID);
+			SCOPE_CYCLE_COUNTER(STAT_NiagaraEmitterInit_Events);
+			
+			EventInstanceData = MakeUnique<FEventInstanceData>();
+			EventInstanceData->UpdateScriptEventDataSets.Empty(UpdateEventGeneratorCount);
+			EventInstanceData->UpdateEventGeneratorIsSharedByIndex.SetNumZeroed(UpdateEventGeneratorCount);
+			int32 UpdateEventGeneratorIndex = 0;
+			for (const FNiagaraEventGeneratorProperties &GeneratorProps : CachedEmitter->UpdateScriptProps.EventGenerators)
+			{
+				FNiagaraDataSet *Set = ParentSystemInstance->CreateEventDataSet(EmitterHandle.GetIdName(), GeneratorProps.ID);
+				Set->Init(&GeneratorProps.DataSetCompiledData);
 
-			Set->Init(&GeneratorProps.DataSetCompiledData);
+				EventInstanceData->UpdateScriptEventDataSets.Add(Set);
+				EventInstanceData->UpdateEventGeneratorIsSharedByIndex[UpdateEventGeneratorIndex] = CachedEmitter->IsEventGeneratorShared(GeneratorProps.ID);
+				++UpdateEventGeneratorIndex;
+			}
 
-			SpawnScriptEventDataSets.Add(Set);
-			SpawnEventGeneratorIsSharedByIndex[SpawnEventGeneratorIndex] = CachedEmitter->IsEventGeneratorShared(GeneratorProps.ID);
-			SpawnEventGeneratorIndex++;
+			EventInstanceData->SpawnScriptEventDataSets.Empty(SpawnEventGeneratorCount);
+			EventInstanceData->SpawnEventGeneratorIsSharedByIndex.SetNumZeroed(SpawnEventGeneratorCount);
+			int32 SpawnEventGeneratorIndex = 0;
+			for (const FNiagaraEventGeneratorProperties &GeneratorProps : CachedEmitter->SpawnScriptProps.EventGenerators)
+			{
+				FNiagaraDataSet *Set = ParentSystemInstance->CreateEventDataSet(EmitterHandle.GetIdName(), GeneratorProps.ID);
+				Set->Init(&GeneratorProps.DataSetCompiledData);
+
+				EventInstanceData->SpawnScriptEventDataSets.Add(Set);
+				EventInstanceData->SpawnEventGeneratorIsSharedByIndex[SpawnEventGeneratorIndex] = CachedEmitter->IsEventGeneratorShared(GeneratorProps.ID);
+				++SpawnEventGeneratorIndex;
+			}
+
+			EventInstanceData->EventExecContexts.SetNum(NumEvents);
+			EventInstanceData->EventExecCountBindings.SetNum(NumEvents);
+
+			for (int32 i = 0; i < NumEvents; i++)
+			{
+				ensure(CachedEmitter->GetEventHandlers()[i].DataSetAccessSynchronized());
+
+				UNiagaraScript* EventScript = CachedEmitter->GetEventHandlers()[i].Script;
+
+				//This is cpu explicitly? Are we doing event handlers on GPU?
+				EventInstanceData->EventExecContexts[i].Init(EventScript, ENiagaraSimTarget::CPUSim);
+				EventInstanceData->EventExecCountBindings[i].Init(EventInstanceData->EventExecContexts[i].Parameters, SYS_PARAM_ENGINE_EXEC_COUNT);
+			}
 		}
 	}
 
@@ -354,17 +384,7 @@ void FNiagaraEmitterInstance::Init(int32 InEmitterIdx, FNiagaraSystemInstanceID
 				CachedEmitter->GetSimulationStages()[i]->Script->RapidIterationParameters.Bind(&GPUExecContext->CombinedParamStore);
 			}
 		}
-	
-		EventExecContexts.SetNum(CachedEmitter->GetEventHandlers().Num());
-		for (int32 i = 0; i < NumEvents; i++)
-		{
-			ensure(CachedEmitter->GetEventHandlers()[i].DataSetAccessSynchronized());
-
-			UNiagaraScript* EventScript = CachedEmitter->GetEventHandlers()[i].Script;
 
-			//This is cpu explicitly? Are we doing event handlers on GPU?
-			EventExecContexts[i].Init(EventScript, ENiagaraSimTarget::CPUSim);
-		}
 	}
 
 
@@ -380,11 +400,6 @@ void FNiagaraEmitterInstance::Init(int32 InEmitterIdx, FNiagaraSystemInstanceID
 		// Initialize the exec count
 		SpawnExecCountBinding.Init(SpawnExecContext.Parameters, SYS_PARAM_ENGINE_EXEC_COUNT);
 		UpdateExecCountBinding.Init(UpdateExecContext.Parameters, SYS_PARAM_ENGINE_EXEC_COUNT);
-		EventExecCountBindings.SetNum(NumEvents);
-		for (int32 i = 0; i < NumEvents; i++)
-		{
-			EventExecCountBindings[i].Init(EventExecContexts[i].Parameters, SYS_PARAM_ENGINE_EXEC_COUNT);
-		}
 	}
 
 	{
@@ -432,27 +447,30 @@ void FNiagaraEmitterInstance::Init(int32 InEmitterIdx, FNiagaraSystemInstanceID
 		ScriptDefinedDataInterfaceParameters.Bind(&UpdateExecContext.Parameters);
 		UpdateExecContext.Parameters.UnbindFromSourceStores();
 
-		for (FNiagaraScriptExecutionContext& EventContext : EventExecContexts)
+		if (EventInstanceData.IsValid())
 		{
-			SystemScriptDefinedDataInterfaceParameters.Bind(&EventContext.Parameters);
-			ScriptDefinedDataInterfaceParameters.Bind(&EventContext.Parameters);
-			EventContext.Parameters.UnbindFromSourceStores();
+			for (FNiagaraScriptExecutionContext& EventContext : EventInstanceData->EventExecContexts)
+			{
+				SystemScriptDefinedDataInterfaceParameters.Bind(&EventContext.Parameters);
+				ScriptDefinedDataInterfaceParameters.Bind(&EventContext.Parameters);
+				EventContext.Parameters.UnbindFromSourceStores();
+			}
+
+			const int32 NumEventHandlers = CachedEmitter->GetEventHandlers().Num();
+			EventInstanceData->EventHandlingInfo.Reset();
+			EventInstanceData->EventHandlingInfo.SetNum(NumEventHandlers);
+			for (int32 i = 0; i < NumEventHandlers; i++)
+			{
+				const FNiagaraEventScriptProperties& EventHandlerProps = CachedEmitter->GetEventHandlers()[i];
+				FNiagaraEventHandlingInfo& Info = EventInstanceData->EventHandlingInfo[i];
+				Info.SourceEmitterGuid = EventHandlerProps.SourceEmitterID;
+				Info.SourceEmitterName = Info.SourceEmitterGuid.IsValid() ? *Info.SourceEmitterGuid.ToString() : CachedIDName;
+				Info.SpawnCounts.Reset();
+				Info.TotalSpawnCount = 0;
+				Info.EventData = nullptr;
+			}
 		}
 	}	
-
-	const int32 NumEventHandlers = CachedEmitter->GetEventHandlers().Num();
-	EventHandlingInfo.Reset();
-	EventHandlingInfo.SetNum(NumEventHandlers);
-	for (int32 i = 0; i < NumEventHandlers; i++)
-	{
-		const FNiagaraEventScriptProperties& EventHandlerProps = CachedEmitter->GetEventHandlers()[i];
-		FNiagaraEventHandlingInfo& Info = EventHandlingInfo[i];
-		Info.SourceEmitterGuid = EventHandlerProps.SourceEmitterID;
-		Info.SourceEmitterName = Info.SourceEmitterGuid.IsValid() ? *Info.SourceEmitterGuid.ToString() : CachedIDName;
-		Info.SpawnCounts.Reset();
-		Info.TotalSpawnCount = 0;
-		Info.EventData = nullptr;
-	}
 }
 
 void FNiagaraEmitterInstance::ResetSimulation(bool bKillExisting /*= true*/)
@@ -481,13 +499,17 @@ void FNiagaraEmitterInstance::ResetSimulation(bool bKillExisting /*= true*/)
 		TotalSpawnedParticles = 0;
 
 		ParticleDataSet->ResetBuffers();
-		for (FNiagaraDataSet* SpawnScriptEventDataSet : SpawnScriptEventDataSets)
-		{
-			SpawnScriptEventDataSet->ResetBuffers();
-		}
-		for (FNiagaraDataSet* UpdateScriptEventDataSet : UpdateScriptEventDataSets)
+
+		if (EventInstanceData.IsValid())
 		{
-			UpdateScriptEventDataSet->ResetBuffers();
+			for (FNiagaraDataSet* SpawnScriptEventDataSet : EventInstanceData->SpawnScriptEventDataSets)
+			{
+				SpawnScriptEventDataSet->ResetBuffers();
+			}
+			for (FNiagaraDataSet* UpdateScriptEventDataSet : EventInstanceData->UpdateScriptEventDataSets)
+			{
+				UpdateScriptEventDataSet->ResetBuffers();
+			}
 		}
 
 		if (CachedEmitter->SimTarget == ENiagaraSimTarget::GPUComputeSim && GPUExecContext != nullptr)
@@ -600,7 +622,7 @@ void FNiagaraEmitterInstance::DirtyDataInterfaces()
 		GPUExecContext->DirtyDataInterfaces();
 	}
 
-	for (FNiagaraScriptExecutionContext& EventContext : EventExecContexts)
+	for (FNiagaraScriptExecutionContext& EventContext : GetEventExecutionContexts())
 	{
 		EventContext.DirtyDataInterfaces();
 	}
@@ -641,13 +663,13 @@ void FNiagaraEmitterInstance::UnbindParameters(bool bExternalOnly)
 			}
 		}
 
-		for (int32 EventIdx = 0; EventIdx < EventExecContexts.Num(); ++EventIdx)
+		for (FNiagaraScriptExecutionContext& EventContext : GetEventExecutionContexts())
 		{
-			for (UNiagaraParameterCollection* Collection : EventExecContexts[EventIdx].Script->GetCachedParameterCollectionReferences())
+			for (UNiagaraParameterCollection* Collection : EventContext.Script->GetCachedParameterCollectionReferences())
 			{
 				if (UNiagaraParameterCollectionInstance* NPCInst = ParentSystemInstance->GetParameterCollectionInstance(Collection))
 				{
-					NPCInst->GetParameterStore().Unbind(&EventExecContexts[EventIdx].Parameters);
+					NPCInst->GetParameterStore().Unbind(&EventContext.Parameters);
 				}
 			}
 		}
@@ -661,9 +683,9 @@ void FNiagaraEmitterInstance::UnbindParameters(bool bExternalOnly)
 			GPUExecContext->CombinedParamStore.UnbindFromSourceStores();
 		}
 
-		for (int32 EventIdx = 0; EventIdx < EventExecContexts.Num(); ++EventIdx)
+		for (FNiagaraScriptExecutionContext& EventContext : GetEventExecutionContexts())
 		{
-			EventExecContexts[EventIdx].Parameters.UnbindFromSourceStores();
+			EventContext.Parameters.UnbindFromSourceStores();
 		}
 	}
 }
@@ -684,34 +706,38 @@ void FNiagaraEmitterInstance::BindParameters(bool bExternalOnly)
 		ParentSystemInstance->GetParameterCollectionInstance(Collection)->GetParameterStore().Bind(&UpdateExecContext.Parameters);
 	}
 
-	for (int32 EventIdx = 0; EventIdx < EventExecContexts.Num(); ++EventIdx)
+	for (FNiagaraScriptExecutionContext& EventContext : GetEventExecutionContexts())
 	{
-		for (UNiagaraParameterCollection* Collection : EventExecContexts[EventIdx].Script->GetCachedParameterCollectionReferences())
+		for (UNiagaraParameterCollection* Collection : EventContext.Script->GetCachedParameterCollectionReferences())
 		{
-			ParentSystemInstance->GetParameterCollectionInstance(Collection)->GetParameterStore().Bind(&EventExecContexts[EventIdx].Parameters);
+			ParentSystemInstance->GetParameterCollectionInstance(Collection)->GetParameterStore().Bind(&EventContext.Parameters);
 		}
 	}
 
 	if (!bExternalOnly)
 	{
 		//Now bind parameters from the component and system.
-		FNiagaraParameterStore& InstanceParams = ParentSystemInstance->GetInstanceParameters();		
-	
+		FNiagaraParameterStore& InstanceParams = ParentSystemInstance->GetInstanceParameters();
+
 		InstanceParams.Bind(&SpawnExecContext.Parameters);
 		InstanceParams.Bind(&UpdateExecContext.Parameters);
 
-		for (FNiagaraScriptExecutionContext& EventContext : EventExecContexts)
+		for (FNiagaraScriptExecutionContext& EventContext : GetEventExecutionContexts())
 		{
 			InstanceParams.Bind(&EventContext.Parameters);
 		}
 
-	#if WITH_EDITORONLY_DATA
+#if WITH_EDITORONLY_DATA
 		CachedEmitter->SpawnScriptProps.Script->RapidIterationParameters.Bind(&SpawnExecContext.Parameters);
 		CachedEmitter->UpdateScriptProps.Script->RapidIterationParameters.Bind(&UpdateExecContext.Parameters);
-		ensure(CachedEmitter->GetEventHandlers().Num() == EventExecContexts.Num());
-		for (int32 i = 0; i < CachedEmitter->GetEventHandlers().Num(); i++)
+
+		if (EventInstanceData.IsValid())
 		{
-			CachedEmitter->GetEventHandlers()[i].Script->RapidIterationParameters.Bind(&EventExecContexts[i].Parameters);
+			ensure(CachedEmitter->GetEventHandlers().Num() == EventInstanceData->EventExecContexts.Num());
+			for (int32 i = 0; i < CachedEmitter->GetEventHandlers().Num(); i++)
+			{
+				CachedEmitter->GetEventHandlers()[i].Script->RapidIterationParameters.Bind(&EventInstanceData->EventExecContexts[i].Parameters);
+			}
 		}
 	#endif
 
@@ -845,10 +871,13 @@ void FNiagaraEmitterInstance::PostTick()
 
 	checkSlow(CachedEmitter);
 
-	//Clear refs to event data buffers.
-	for (FNiagaraEventHandlingInfo& Info : EventHandlingInfo)
+	if (EventInstanceData.IsValid())
 	{
-		Info.SetEventData(nullptr);
+		//Clear refs to event data buffers.
+		for (FNiagaraEventHandlingInfo& Info : EventInstanceData->EventHandlingInfo)
+		{
+			Info.SetEventData(nullptr);
+		}
 	}
 
 	CachedBounds.Init();
@@ -896,6 +925,18 @@ bool FNiagaraEmitterInstance::HandleCompletion(bool bForce)
 	if (IsComplete())
 	{
 		ParticleDataSet->ResetBuffers();
+		if (EventInstanceData.IsValid())
+		{
+			for (FNiagaraDataSet* EventDataSet : EventInstanceData->UpdateScriptEventDataSets)
+			{
+				EventDataSet->ResetBuffers();
+			}
+
+			for (FNiagaraDataSet* EventDataSet : EventInstanceData->SpawnScriptEventDataSets)
+			{
+				EventDataSet->ResetBuffers();
+			}
+		}
 		return true;
 	}
 
@@ -928,14 +969,16 @@ void FNiagaraEmitterInstance::PreTick()
 #if WITH_EDITOR
 	CachedEmitter->SpawnScriptProps.Script->RapidIterationParameters.Tick();
 	CachedEmitter->UpdateScriptProps.Script->RapidIterationParameters.Tick();
-	ensure(CachedEmitter->GetEventHandlers().Num() == EventExecContexts.Num());
-	for (int32 i = 0; i < CachedEmitter->GetEventHandlers().Num(); i++)
+	if (EventInstanceData.IsValid())
 	{
-		CachedEmitter->GetEventHandlers()[i].Script->RapidIterationParameters.Tick();
+		ensure(CachedEmitter->GetEventHandlers().Num() == EventInstanceData->EventExecContexts.Num());
+		for (int32 i = 0; i < CachedEmitter->GetEventHandlers().Num(); i++)
+		{
+			CachedEmitter->GetEventHandlers()[i].Script->RapidIterationParameters.Tick();
+		}
 	}
 #endif
 
-
 	bool bOk = true;
 	bOk &= SpawnExecContext.Tick(ParentSystemInstance, CachedEmitter->SimTarget);
 	bOk &= UpdateExecContext.Tick(ParentSystemInstance, CachedEmitter->SimTarget);
@@ -945,7 +988,8 @@ void FNiagaraEmitterInstance::PreTick()
 	{
 		bOk &= GPUExecContext->Tick(ParentSystemInstance);
 	}
-	for (FNiagaraScriptExecutionContext& EventContext : EventExecContexts)
+
+	for (FNiagaraScriptExecutionContext& EventContext : GetEventExecutionContexts())
 	{
 		bOk &= EventContext.Tick(ParentSystemInstance, CachedEmitter->SimTarget);
 	}
@@ -967,7 +1011,8 @@ void FNiagaraEmitterInstance::PreTick()
 			//We PostTick the GPUExecContext here to prime crucial PREV parameters (such as PREV_Engine.Owner.Position). This PostTick call is necessary as the GPUExecContext has not been sent to the batcher yet.
 			GPUExecContext->PostTick();
 		}
-		for (FNiagaraScriptExecutionContext& EventContext : EventExecContexts)
+
+		for (FNiagaraScriptExecutionContext& EventContext : GetEventExecutionContexts())
 		{
 			EventContext.PostTick();
 		}
@@ -980,36 +1025,44 @@ void FNiagaraEmitterInstance::PreTick()
 	if (bResetPending)
 	{
 		Data.ResetBuffers();
-		for (FNiagaraDataSet* SpawnScriptEventDataSet : SpawnScriptEventDataSets)
-		{
-			SpawnScriptEventDataSet->ResetBuffers();
-		}
-		for (FNiagaraDataSet* UpdateScriptEventDataSet : UpdateScriptEventDataSets)
+
+		if (EventInstanceData.IsValid())
 		{
-			UpdateScriptEventDataSet->ResetBuffers();
+			for (FNiagaraDataSet* SpawnScriptEventDataSet : EventInstanceData->SpawnScriptEventDataSets)
+			{
+				SpawnScriptEventDataSet->ResetBuffers();
+			}
+			for (FNiagaraDataSet* UpdateScriptEventDataSet : EventInstanceData->UpdateScriptEventDataSets)
+			{
+				UpdateScriptEventDataSet->ResetBuffers();
+			}
 		}
+
 		bResetPending = false;
 	}
 
-	//Gather events we're going to be reading from / handling this frame.
-	//We must do this in pre-tick so we can gather (and mark in use) all sets from other emitters.
-	const int32 NumEventHandlers = CachedEmitter->GetEventHandlers().Num();
-	EventSpawnTotal = 0;
-	for (int32 i = 0; i < NumEventHandlers; i++)
+	if (EventInstanceData.IsValid())
 	{
-		const FNiagaraEventScriptProperties& EventHandlerProps = CachedEmitter->GetEventHandlers()[i];
-		FNiagaraEventHandlingInfo& Info = EventHandlingInfo[i];
+		//Gather events we're going to be reading from / handling this frame.
+		//We must do this in pre-tick so we can gather (and mark in use) all sets from other emitters.
+		const int32 NumEventHandlers = CachedEmitter->GetEventHandlers().Num();
+		EventInstanceData->EventSpawnTotal = 0;
+		for (int32 i = 0; i < NumEventHandlers; i++)
+		{
+			const FNiagaraEventScriptProperties& EventHandlerProps = CachedEmitter->GetEventHandlers()[i];
+			FNiagaraEventHandlingInfo& Info = EventInstanceData->EventHandlingInfo[i];
 
-		Info.TotalSpawnCount = 0;//This was being done every frame but should be done in init?
-		Info.SpawnCounts.Reset();
+			Info.TotalSpawnCount = 0;//This was being done every frame but should be done in init?
+			Info.SpawnCounts.Reset();
 
-		//TODO: We can move this lookup into the init and just store a ptr to the other set?
-		if (FNiagaraDataSet* EventSet = ParentSystemInstance->GetEventDataSet(Info.SourceEmitterName, EventHandlerProps.SourceEventName))
-		{
-			Info.SetEventData(&EventSet->GetCurrentDataChecked());
-			uint32 EventSpawnNum = CalculateEventSpawnCount(EventHandlerProps, Info.SpawnCounts, EventSet);
-			Info.TotalSpawnCount += EventSpawnNum;
-			EventSpawnTotal += EventSpawnNum;
+			//TODO: We can move this lookup into the init and just store a ptr to the other set?
+			if (FNiagaraDataSet* EventSet = ParentSystemInstance->GetEventDataSet(Info.SourceEmitterName, EventHandlerProps.SourceEventName))
+			{
+				Info.SetEventData(&EventSet->GetCurrentDataChecked());
+				uint32 EventSpawnNum = CalculateEventSpawnCount(EventHandlerProps, Info.SpawnCounts, EventSet);
+				Info.TotalSpawnCount += EventSpawnNum;
+				EventInstanceData->EventSpawnTotal += EventSpawnNum;
+			}
 		}
 	}
 
@@ -1118,6 +1171,8 @@ void FNiagaraEmitterInstance::Tick(float DeltaSeconds)
 			}
 		}
 	}
+
+	int32 EventSpawnTotal = (EventInstanceData.IsValid() ? EventInstanceData->EventSpawnTotal : 0);
 	int32 OrigNumParticles = GetNumParticles();
 	int32 AllocationEstimate = CachedEmitter->GetMaxParticleCountEstimate();
 	int32 RequiredSize = OrigNumParticles + SpawnTotal + EventSpawnTotal;
@@ -1264,8 +1319,9 @@ void FNiagaraEmitterInstance::Tick(float DeltaSeconds)
 		// Need to call post-tick, which calls the copy to previous for interpolated spawning
 		SpawnExecContext.PostTick();
 		UpdateExecContext.PostTick();
+
 		// At this stage GPU execution is being handled by the batcher so we do not need to call PostTick() for it
-		for (FNiagaraScriptExecutionContext& EventContext : EventExecContexts)
+		for (FNiagaraScriptExecutionContext& EventContext : GetEventExecutionContexts())
 		{
 			EventContext.PostTick();
 		}
@@ -1289,38 +1345,46 @@ void FNiagaraEmitterInstance::Tick(float DeltaSeconds)
 		//For now we completely bail out of spawning new particles. Possibly should improve this in future.
 		AllocationSize = OrigNumParticles;
 		SpawnTotal = 0;
-		EventSpawnTotal = 0;
+
+		if (EventInstanceData.IsValid())
+		{
+			EventSpawnTotal = 0;
+			EventInstanceData->EventSpawnTotal = 0;
+		}
 	}
 
 	Data.BeginSimulate();
 	Data.Allocate(AllocationSize);
 
-	int32 SpawnEventGeneratorIndex = 0;
-	for (FNiagaraDataSet* SpawnEventDataSet : SpawnScriptEventDataSets)
+	if (EventInstanceData.IsValid())
 	{
-		int32 NumToAllocate = SpawnTotal + EventSpawnTotal;
-		if (SpawnEventGeneratorIsSharedByIndex[SpawnEventGeneratorIndex])
+		int32 SpawnEventGeneratorIndex = 0;
+		for (FNiagaraDataSet* SpawnEventDataSet : EventInstanceData->SpawnScriptEventDataSets)
 		{
-			// For shared event data sets we need to allocate storage for the current particles since
-			// the same data set will be used in the update execution.
-			NumToAllocate += OrigNumParticles;
+			int32 NumToAllocate = SpawnTotal + EventSpawnTotal;
+			if (EventInstanceData->SpawnEventGeneratorIsSharedByIndex[SpawnEventGeneratorIndex])
+			{
+				// For shared event data sets we need to allocate storage for the current particles since
+				// the same data set will be used in the update execution.
+				NumToAllocate += OrigNumParticles;
+			}
+			SpawnEventDataSet->BeginSimulate();
+			SpawnEventDataSet->Allocate(NumToAllocate);
+			SpawnEventGeneratorIndex++;
 		}
-		SpawnEventDataSet->BeginSimulate();
-		SpawnEventDataSet->Allocate(NumToAllocate);
-		SpawnEventGeneratorIndex++;
-	}
 
-	int32 UpdateEventGeneratorIndex = 0;
-	for (FNiagaraDataSet* UpdateEventDataSet : UpdateScriptEventDataSets)
-	{
-		if (UpdateEventGeneratorIsSharedByIndex[UpdateEventGeneratorIndex] == false)
+		int32 UpdateEventGeneratorIndex = 0;
+		for (FNiagaraDataSet* UpdateEventDataSet : EventInstanceData->UpdateScriptEventDataSets)
 		{
-			// We only allocate update event data sets if they're not shared, because shared event datasets will have already
-			// been allocated as part of the spawn event data set handling.
-			UpdateEventDataSet->BeginSimulate();
-			UpdateEventDataSet->Allocate(OrigNumParticles);
+			if (EventInstanceData->UpdateEventGeneratorIsSharedByIndex[UpdateEventGeneratorIndex] == false)
+			{
+				// We only allocate update event data sets if they're not shared, because shared event datasets will have already
+				// been allocated as part of the spawn event data set handling.
+				UpdateEventDataSet->BeginSimulate();
+				UpdateEventDataSet->Allocate(OrigNumParticles);
+			}
+			UpdateEventGeneratorIndex++;
 		}
-		UpdateEventGeneratorIndex++;
 	}
 
 	// Simulate existing particles forward by DeltaSeconds.
@@ -1331,12 +1395,16 @@ void FNiagaraEmitterInstance::Tick(float DeltaSeconds)
 
 		UpdateExecCountBinding.SetValue(OrigNumParticles);
 		UpdateExecContext.BindData(0, Data, 0, true);
-		int32 EventDataSetIdx = 1;
-		for (FNiagaraDataSet* EventDataSet : UpdateScriptEventDataSets)
+
+		if (EventInstanceData.IsValid())
 		{
-			check(EventDataSet);
-			EventDataSet->GetDestinationDataChecked().SetNumInstances(OrigNumParticles);
-			UpdateExecContext.BindData(EventDataSetIdx++, *EventDataSet, 0, true);
+			int32 EventDataSetIdx = 1;
+			for (FNiagaraDataSet* EventDataSet : EventInstanceData->UpdateScriptEventDataSets)
+			{
+				check(EventDataSet);
+				EventDataSet->GetDestinationDataChecked().SetNumInstances(OrigNumParticles);
+				UpdateExecContext.BindData(EventDataSetIdx++, *EventDataSet, 0, true);
+			}
 		}
 
 		FScriptExecutionConstantBufferTable UpdateConstantBufferTable;
@@ -1350,12 +1418,16 @@ void FNiagaraEmitterInstance::Tick(float DeltaSeconds)
 		if (GbDumpParticleData || System->bDumpDebugEmitterInfo)
 		{
 			Data.GetDestinationDataChecked().Dump(0, OrigNumParticles, FString::Printf(TEXT("=== Updated %d Particles (%d Died) ==="), OrigNumParticles, -DeltaParticles));
-			for (int32 EventIdx = 0; EventIdx < UpdateScriptEventDataSets.Num(); ++EventIdx)
+
+			if (EventInstanceData.IsValid())
 			{
-				FNiagaraDataSet* EventDataSet = UpdateScriptEventDataSets[EventIdx];
-				if (EventDataSet && EventDataSet->GetDestinationDataChecked().GetNumInstances() > 0)
+				for (int32 EventIdx = 0; EventIdx < EventInstanceData->UpdateScriptEventDataSets.Num(); ++EventIdx)
 				{
-					EventDataSet->GetDestinationDataChecked().Dump(0, INDEX_NONE, FString::Printf(TEXT("Update Script Event %d"), EventIdx));
+					FNiagaraDataSet* EventDataSet = EventInstanceData->UpdateScriptEventDataSets[EventIdx];
+					if (EventDataSet && EventDataSet->GetDestinationDataChecked().GetNumInstances() > 0)
+					{
+						EventDataSet->GetDestinationDataChecked().Dump(0, INDEX_NONE, FString::Printf(TEXT("Update Script Event %d"), EventIdx));
+					}
 				}
 			}
 		//	UE_LOG(LogNiagara, Log, TEXT("=== Update Parameters ===") );
@@ -1397,14 +1469,17 @@ void FNiagaraEmitterInstance::Tick(float DeltaSeconds)
 			SpawnExecCountBinding.SetValue(Num);
 			SpawnExecContext.BindData(0, Data, OrigNum, true);
 
-			//UE_LOG(LogNiagara, Log, TEXT("SpawnScriptEventDataSets: %d"), SpawnScriptEventDataSets.Num());
-			int32 EventDataSetIdx = 1;
-			for (FNiagaraDataSet* EventDataSet : SpawnScriptEventDataSets)
+			if (EventInstanceData.IsValid())
 			{
-				//UE_LOG(LogNiagara, Log, TEXT("SpawnScriptEventDataSets.. %d"), EventDataSet->GetNumVariables());
-				int32 EventOrigNum = EventDataSet->GetDestinationDataChecked().GetNumInstances();
-				EventDataSet->GetDestinationDataChecked().SetNumInstances(EventOrigNum + Num);
-				SpawnExecContext.BindData(EventDataSetIdx++, *EventDataSet, EventOrigNum, true);
+				//UE_LOG(LogNiagara, Log, TEXT("SpawnScriptEventDataSets: %d"), SpawnScriptEventDataSets.Num());
+				int32 EventDataSetIdx = 1;
+				for (FNiagaraDataSet* EventDataSet : EventInstanceData->SpawnScriptEventDataSets)
+				{
+					//UE_LOG(LogNiagara, Log, TEXT("SpawnScriptEventDataSets.. %d"), EventDataSet->GetNumVariables());
+					int32 EventOrigNum = EventDataSet->GetDestinationDataChecked().GetNumInstances();
+					EventDataSet->GetDestinationDataChecked().SetNumInstances(EventOrigNum + Num);
+					SpawnExecContext.BindData(EventDataSetIdx++, *EventDataSet, EventOrigNum, true);
+				}
 			}
 
 			SpawnExecContext.Execute(Num, SpawnConstantBufferTable);
@@ -1412,14 +1487,19 @@ void FNiagaraEmitterInstance::Tick(float DeltaSeconds)
 			if (GbDumpParticleData || System->bDumpDebugEmitterInfo)
 			{
 				Data.GetDestinationDataChecked().Dump(OrigNum, Num, FString::Printf(TEXT("===  %s Spawned %d Particles==="), DumpLabel, Num));
-				for (int32 EventIdx = 0; EventIdx < SpawnScriptEventDataSets.Num(); ++EventIdx)
+
+				if (EventInstanceData.IsValid())
 				{
-					FNiagaraDataSet* EventDataSet = SpawnScriptEventDataSets[EventIdx];
-					if (EventDataSet && EventDataSet->GetDestinationDataChecked().GetNumInstances() > 0)
+					for (int32 EventIdx = 0; EventIdx < EventInstanceData->SpawnScriptEventDataSets.Num(); ++EventIdx)
 					{
-						EventDataSet->GetDestinationDataChecked().Dump(0, INDEX_NONE, FString::Printf(TEXT("Spawn Script Event %d"), EventIdx));
+						FNiagaraDataSet* EventDataSet = EventInstanceData->SpawnScriptEventDataSets[EventIdx];
+						if (EventDataSet && EventDataSet->GetDestinationDataChecked().GetNumInstances() > 0)
+						{
+							EventDataSet->GetDestinationDataChecked().Dump(0, INDEX_NONE, FString::Printf(TEXT("Spawn Script Event %d"), EventIdx));
+						}
 					}
 				}
+
 				//UE_LOG(LogNiagara, Log, TEXT("=== %s Spawn Parameters ==="), *DumpLabel);
 				SpawnExecContext.Parameters.Dump();
 			}
@@ -1440,28 +1520,31 @@ void FNiagaraEmitterInstance::Tick(float DeltaSeconds)
 
 		EventSpawnStart = Data.GetDestinationDataChecked().GetNumInstances();
 
-		for (int32 EventScriptIdx = 0; EventScriptIdx < CachedEmitter->GetEventHandlers().Num(); EventScriptIdx++)
+		if (EventInstanceData.IsValid())
 		{
-			FNiagaraEventHandlingInfo& Info = EventHandlingInfo[EventScriptIdx];
-			//Spawn particles coming from events.
-			for (int32 i = 0; i < Info.SpawnCounts.Num(); i++)
+			for (int32 EventScriptIdx = 0; EventScriptIdx < CachedEmitter->GetEventHandlers().Num(); EventScriptIdx++)
 			{
-				const int32 EventNumToSpawn = Info.SpawnCounts[i];
-				if (EventNumToSpawn > 0)
+				FNiagaraEventHandlingInfo& Info = EventInstanceData->EventHandlingInfo[EventScriptIdx];
+				//Spawn particles coming from events.
+				for (int32 i = 0; i < Info.SpawnCounts.Num(); i++)
 				{
-					const int32 CurrNumParticles = Data.GetDestinationDataChecked().GetNumInstances();
-
-					//Event spawns are instantaneous at the middle of the frame?
-					auto& EmitterParameters = ParentSystemInstance->EditEmitterParameters(EmitterIdx);
-					SpawnIntervalBinding.SetValue(0.0f);
-					InterpSpawnStartBinding.SetValue(DeltaSeconds * 0.5f);
-					SpawnGroupBinding.SetValue(0);
-					SpawnParticles(EventNumToSpawn, TEXT("Event Spawn"));
-
-					//Update EventSpawnCounts to the number actually spawned.
-					int32 NumActuallySpawned = Data.GetDestinationDataChecked().GetNumInstances() - CurrNumParticles;
-					TotalActualEventSpawns += NumActuallySpawned;
-					Info.SpawnCounts[i] = NumActuallySpawned;
+					const int32 EventNumToSpawn = Info.SpawnCounts[i];
+					if (EventNumToSpawn > 0)
+					{
+						const int32 CurrNumParticles = Data.GetDestinationDataChecked().GetNumInstances();
+
+						//Event spawns are instantaneous at the middle of the frame?
+						auto& EmitterParameters = ParentSystemInstance->EditEmitterParameters(EmitterIdx);
+						SpawnIntervalBinding.SetValue(0.0f);
+						InterpSpawnStartBinding.SetValue(DeltaSeconds * 0.5f);
+						SpawnGroupBinding.SetValue(0);
+						SpawnParticles(EventNumToSpawn, TEXT("Event Spawn"));
+
+						//Update EventSpawnCounts to the number actually spawned.
+						int32 NumActuallySpawned = Data.GetDestinationDataChecked().GetNumInstances() - CurrNumParticles;
+						TotalActualEventSpawns += NumActuallySpawned;
+						Info.SpawnCounts[i] = NumActuallySpawned;
+					}
 				}
 			}
 		}
@@ -1475,19 +1558,23 @@ void FNiagaraEmitterInstance::Tick(float DeltaSeconds)
 
 	//We're done with this simulation pass.
 	Data.EndSimulate();
-	for (FNiagaraDataSet* SpawnEventDataSet : SpawnScriptEventDataSets)
+
+	if (EventInstanceData.IsValid())
 	{
-		if (SpawnEventDataSet && SpawnEventDataSet->GetDestinationData())
+		for (FNiagaraDataSet* SpawnEventDataSet : EventInstanceData->SpawnScriptEventDataSets)
 		{
-			SpawnEventDataSet->EndSimulate();
+			if (SpawnEventDataSet && SpawnEventDataSet->GetDestinationData())
+			{
+				SpawnEventDataSet->EndSimulate();
+			}
 		}
-	}
 
-	for (FNiagaraDataSet* UpdateEventDataSet : UpdateScriptEventDataSets)
-	{
-		if (UpdateEventDataSet && UpdateEventDataSet->GetDestinationData())
+		for (FNiagaraDataSet* UpdateEventDataSet : EventInstanceData->UpdateScriptEventDataSets)
 		{
-			UpdateEventDataSet->EndSimulate();
+			if (UpdateEventDataSet && UpdateEventDataSet->GetDestinationData())
+			{
+				UpdateEventDataSet->EndSimulate();
+			}
 		}
 	}
 
@@ -1518,137 +1605,140 @@ void FNiagaraEmitterInstance::Tick(float DeltaSeconds)
 		UE_LOG(LogNiagara, Log, TEXT("Skipping spawning due to execution state! %d"), (uint32)ExecutionState)
 	}*/
 
-	if (TotalActualEventSpawns > 0)
+	if (EventInstanceData.IsValid())
 	{
-		if (GbDumpParticleData || System->bDumpDebugEmitterInfo)
+		if (TotalActualEventSpawns > 0)
 		{
-			Data.Dump(0, INDEX_NONE, TEXT("Existing Data - Pre Event Alloc"));
+			if (GbDumpParticleData || System->bDumpDebugEmitterInfo)
+			{
+				Data.Dump(0, INDEX_NONE, TEXT("Existing Data - Pre Event Alloc"));
+			}
+			//Allocate a new dest buffer to write spawn event handler results into.
+			//Can just do one allocate here for all spawn event handlers.
+			//Though this requires us to copy the contents of the instances we're not writing to in this pass over from the previous buffer.
+			FNiagaraDataBuffer& DestBuffer = Data.BeginSimulate();
+			Data.Allocate(Data.GetCurrentDataChecked().GetNumInstances(), true);
+			DestBuffer.SetNumInstances(Data.GetCurrentDataChecked().GetNumInstances());
+
+			//if (GbDumpParticleData || System->bDumpDebugEmitterInfo)
+			//{
+			//	DestBuffer.Dump(0, INDEX_NONE, TEXT("Existing Data - Post Event Alloc, Pre Events"));
+			//}
 		}
-		//Allocate a new dest buffer to write spawn event handler results into.
-		//Can just do one allocate here for all spawn event handlers.
-		//Though this requires us to copy the contents of the instances we're not writing to in this pass over from the previous buffer.
-		FNiagaraDataBuffer& DestBuffer = Data.BeginSimulate();
-		Data.Allocate(Data.GetCurrentDataChecked().GetNumInstances(), true);
-		DestBuffer.SetNumInstances(Data.GetCurrentDataChecked().GetNumInstances());
-
-		//if (GbDumpParticleData || System->bDumpDebugEmitterInfo)
-		//{
-		//	DestBuffer.Dump(0, INDEX_NONE, TEXT("Existing Data - Post Event Alloc, Pre Events"));
-		//}
-	}
-
-	int32 SpawnEventScriptStartIndex = EventSpawnStart;
-	for (int32 EventScriptIdx = 0; EventScriptIdx < CachedEmitter->GetEventHandlers().Num(); EventScriptIdx++)
-	{
-		const FNiagaraEventScriptProperties &EventHandlerProps = CachedEmitter->GetEventHandlers()[EventScriptIdx];
-		FNiagaraEventHandlingInfo& Info = EventHandlingInfo[EventScriptIdx];
 
-		FScriptExecutionConstantBufferTable EventConstantBufferTable;
-		BuildConstantBufferTable(EventExecContexts[EventScriptIdx], EventConstantBufferTable);
-
-		if (Info.EventData && Info.SpawnCounts.Num() > 0)
+		int32 SpawnEventScriptStartIndex = EventSpawnStart;
+		for (int32 EventScriptIdx = 0; EventScriptIdx < CachedEmitter->GetEventHandlers().Num(); EventScriptIdx++)
 		{
-			SCOPE_CYCLE_COUNTER(STAT_NiagaraEventHandle);
+			const FNiagaraEventScriptProperties &EventHandlerProps = CachedEmitter->GetEventHandlers()[EventScriptIdx];
+			FNiagaraEventHandlingInfo& Info = EventInstanceData->EventHandlingInfo[EventScriptIdx];
+
+			FScriptExecutionConstantBufferTable EventConstantBufferTable;
+			BuildConstantBufferTable(EventInstanceData->EventExecContexts[EventScriptIdx], EventConstantBufferTable);
 
-			for (int32 i = 0; i < Info.SpawnCounts.Num(); i++)
+			if (Info.EventData && Info.SpawnCounts.Num() > 0)
 			{
-				int32 EventNumToSpawn = Info.SpawnCounts[i];
-				if (EventNumToSpawn > 0)
+				SCOPE_CYCLE_COUNTER(STAT_NiagaraEventHandle);
+
+				for (int32 i = 0; i < Info.SpawnCounts.Num(); i++)
 				{
-					EventExecCountBindings[EventScriptIdx].SetValue(EventNumToSpawn);
+					int32 EventNumToSpawn = Info.SpawnCounts[i];
+					if (EventNumToSpawn > 0)
+					{
+						EventInstanceData->EventExecCountBindings[EventScriptIdx].SetValue(EventNumToSpawn);
 
-					EventExecContexts[EventScriptIdx].BindData(0, Data, EventSpawnStart, false);
-					EventExecContexts[EventScriptIdx].BindData(1, EventHandlingInfo[EventScriptIdx].EventData, i, false);
+						EventInstanceData->EventExecContexts[EventScriptIdx].BindData(0, Data, EventSpawnStart, false);
+						EventInstanceData->EventExecContexts[EventScriptIdx].BindData(1, EventInstanceData->EventHandlingInfo[EventScriptIdx].EventData, i, false);
 
-					EventExecContexts[EventScriptIdx].Execute(EventNumToSpawn, EventConstantBufferTable);
+						EventInstanceData->EventExecContexts[EventScriptIdx].Execute(EventNumToSpawn, EventConstantBufferTable);
 
-					if (GbDumpParticleData || System->bDumpDebugEmitterInfo)
-					{
-						EventHandlingInfo[EventScriptIdx].EventData->Dump(i, 1, FString::Printf(TEXT("=== Event Data %d [%d] ==="), EventScriptIdx, i));
-						Data.GetDestinationDataChecked().Dump(EventSpawnStart, EventNumToSpawn, FString::Printf(TEXT("=== Event %d %d Particles ==="), EventScriptIdx, EventNumToSpawn));
-						//UE_LOG(LogNiagara, Log, TEXT("=== Event %d Parameters ==="), EventScriptIdx);
-						EventExecContexts[EventScriptIdx].Parameters.Dump();
-					}
+						if (GbDumpParticleData || System->bDumpDebugEmitterInfo)
+						{
+							EventInstanceData->EventHandlingInfo[EventScriptIdx].EventData->Dump(i, 1, FString::Printf(TEXT("=== Event Data %d [%d] ==="), EventScriptIdx, i));
+							Data.GetDestinationDataChecked().Dump(EventSpawnStart, EventNumToSpawn, FString::Printf(TEXT("=== Event %d %d Particles ==="), EventScriptIdx, EventNumToSpawn));
+							//UE_LOG(LogNiagara, Log, TEXT("=== Event %d Parameters ==="), EventScriptIdx);
+							EventInstanceData->EventExecContexts[EventScriptIdx].Parameters.Dump();
+						}
 
-#if WITH_EDITORONLY_DATA
-					if (ParentSystemInstance->ShouldCaptureThisFrame())
-					{
-						FGuid EventGuid = EventExecContexts[EventScriptIdx].Script->GetUsageId();
-						TSharedPtr<struct FNiagaraScriptDebuggerInfo, ESPMode::ThreadSafe> DebugInfo = ParentSystemInstance->GetActiveCaptureWrite(CachedIDName, ENiagaraScriptUsage::ParticleEventScript, EventGuid);
-						if (DebugInfo)
+	#if WITH_EDITORONLY_DATA
+						if (ParentSystemInstance->ShouldCaptureThisFrame())
 						{
-							Data.CopyTo(DebugInfo->Frame, EventSpawnStart, EventNumToSpawn);
-							DebugInfo->Parameters = EventExecContexts[EventScriptIdx].Parameters;
-							DebugInfo->bWritten = true;
+							FGuid EventGuid = EventInstanceData->EventExecContexts[EventScriptIdx].Script->GetUsageId();
+							TSharedPtr<struct FNiagaraScriptDebuggerInfo, ESPMode::ThreadSafe> DebugInfo = ParentSystemInstance->GetActiveCaptureWrite(CachedIDName, ENiagaraScriptUsage::ParticleEventScript, EventGuid);
+							if (DebugInfo)
+							{
+								Data.CopyTo(DebugInfo->Frame, EventSpawnStart, EventNumToSpawn);
+								DebugInfo->Parameters = EventInstanceData->EventExecContexts[EventScriptIdx].Parameters;
+								DebugInfo->bWritten = true;
+							}
 						}
+	#endif
+						EventSpawnStart += EventNumToSpawn;
 					}
-#endif
-					EventSpawnStart += EventNumToSpawn;
 				}
 			}
 		}
-	}
 
-	//If we processed any events we need to end simulate to update the current sim state.
-	if (Data.GetDestinationData())
-	{
-		Data.EndSimulate();
-	}
-
-	// Update events need a copy per event so that the previous event's data can be used.
-	for (int32 EventScriptIdx = 0; EventScriptIdx < CachedEmitter->GetEventHandlers().Num(); EventScriptIdx++)
-	{
-		const FNiagaraEventScriptProperties &EventHandlerProps = CachedEmitter->GetEventHandlers()[EventScriptIdx];
-		FNiagaraDataBuffer* EventData = EventHandlingInfo[EventScriptIdx].EventData;
-
-		FScriptExecutionConstantBufferTable EventConstantBufferTable;
-		BuildConstantBufferTable(EventExecContexts[EventScriptIdx], EventConstantBufferTable);
+		//If we processed any events we need to end simulate to update the current sim state.
+		if (Data.GetDestinationData())
+		{
+			Data.EndSimulate();
+		}
 
-		// handle all-particle events
-		if (EventHandlerProps.Script && EventHandlerProps.ExecutionMode == EScriptExecutionMode::EveryParticle && EventData)
+		// Update events need a copy per event so that the previous event's data can be used.
+		for (int32 EventScriptIdx = 0; EventScriptIdx < CachedEmitter->GetEventHandlers().Num(); EventScriptIdx++)
 		{
-			uint32 NumParticles = Data.GetCurrentDataChecked().GetNumInstances();
+			const FNiagaraEventScriptProperties &EventHandlerProps = CachedEmitter->GetEventHandlers()[EventScriptIdx];
+			FNiagaraDataBuffer* EventData = EventInstanceData->EventHandlingInfo[EventScriptIdx].EventData;
+
+			FScriptExecutionConstantBufferTable EventConstantBufferTable;
+			BuildConstantBufferTable(EventInstanceData->EventExecContexts[EventScriptIdx], EventConstantBufferTable);
 
-			if (EventData->GetNumInstances())
+			// handle all-particle events
+			if (EventHandlerProps.Script && EventHandlerProps.ExecutionMode == EScriptExecutionMode::EveryParticle && EventData)
 			{
-				SCOPE_CYCLE_COUNTER(STAT_NiagaraEventHandle);
-				
-				for (uint32 i = 0; i < EventData->GetNumInstances(); i++)
+				uint32 NumParticles = Data.GetCurrentDataChecked().GetNumInstances();
+
+				if (EventData->GetNumInstances())
 				{
-					Data.BeginSimulate();
-					Data.Allocate(NumParticles);
+					SCOPE_CYCLE_COUNTER(STAT_NiagaraEventHandle);
 
-					uint32 NumInstancesPrev = Data.GetCurrentDataChecked().GetNumInstances();
-					EventExecCountBindings[EventScriptIdx].SetValue(NumInstancesPrev);
+					for (uint32 i = 0; i < EventData->GetNumInstances(); i++)
+					{
+						Data.BeginSimulate();
+						Data.Allocate(NumParticles);
 
-					EventExecContexts[EventScriptIdx].BindData(0, Data, 0, true);
-					EventExecContexts[EventScriptIdx].BindData(1, EventData, i, false);
+						uint32 NumInstancesPrev = Data.GetCurrentDataChecked().GetNumInstances();
+						EventInstanceData->EventExecCountBindings[EventScriptIdx].SetValue(NumInstancesPrev);
 
-					EventExecContexts[EventScriptIdx].Execute(NumInstancesPrev, EventConstantBufferTable);
+						EventInstanceData->EventExecContexts[EventScriptIdx].BindData(0, Data, 0, true);
+						EventInstanceData->EventExecContexts[EventScriptIdx].BindData(1, EventData, i, false);
 
-					Data.EndSimulate();
+						EventInstanceData->EventExecContexts[EventScriptIdx].Execute(NumInstancesPrev, EventConstantBufferTable);
 
-					if (GbDumpParticleData || System->bDumpDebugEmitterInfo)
-					{
-						EventData->Dump(i, 1, FString::Printf(TEXT("=== Event Data %d [%d] ==="), EventScriptIdx, i));
-						Data.GetCurrentDataChecked().Dump(0, NumInstancesPrev, FString::Printf(TEXT("=== Event %d %d Particles ==="), EventScriptIdx, NumInstancesPrev));
-						EventExecContexts[EventScriptIdx].Parameters.Dump();
-					}
+						Data.EndSimulate();
+
+						if (GbDumpParticleData || System->bDumpDebugEmitterInfo)
+						{
+							EventData->Dump(i, 1, FString::Printf(TEXT("=== Event Data %d [%d] ==="), EventScriptIdx, i));
+							Data.GetCurrentDataChecked().Dump(0, NumInstancesPrev, FString::Printf(TEXT("=== Event %d %d Particles ==="), EventScriptIdx, NumInstancesPrev));
+							EventInstanceData->EventExecContexts[EventScriptIdx].Parameters.Dump();
+						}
 
 #if WITH_EDITORONLY_DATA
-					if (ParentSystemInstance->ShouldCaptureThisFrame())
-					{
-						FGuid EventGuid = EventExecContexts[EventScriptIdx].Script->GetUsageId();
-						TSharedPtr<struct FNiagaraScriptDebuggerInfo, ESPMode::ThreadSafe> DebugInfo = ParentSystemInstance->GetActiveCaptureWrite(CachedIDName, ENiagaraScriptUsage::ParticleEventScript, EventGuid);
-						if (DebugInfo)
+						if (ParentSystemInstance->ShouldCaptureThisFrame())
 						{
-							Data.CopyTo(DebugInfo->Frame, 0, NumInstancesPrev);
-							DebugInfo->Parameters = EventExecContexts[EventScriptIdx].Parameters;
-							DebugInfo->bWritten = true;
+							FGuid EventGuid = EventInstanceData->EventExecContexts[EventScriptIdx].Script->GetUsageId();
+							TSharedPtr<struct FNiagaraScriptDebuggerInfo, ESPMode::ThreadSafe> DebugInfo = ParentSystemInstance->GetActiveCaptureWrite(CachedIDName, ENiagaraScriptUsage::ParticleEventScript, EventGuid);
+							if (DebugInfo)
+							{
+								Data.CopyTo(DebugInfo->Frame, 0, NumInstancesPrev);
+								DebugInfo->Parameters = EventInstanceData->EventExecContexts[EventScriptIdx].Parameters;
+								DebugInfo->bWritten = true;
+							}
 						}
-					}
 #endif
-					ensure(NumParticles == Data.GetCurrentDataChecked().GetNumInstances());
+						ensure(NumParticles == Data.GetCurrentDataChecked().GetNumInstances());
+					}
 				}
 			}
 		}
@@ -1717,7 +1807,8 @@ void FNiagaraEmitterInstance::Tick(float DeltaSeconds)
 	SpawnExecContext.PostTick();
 	UpdateExecContext.PostTick();
 	// At this stage GPU execution is being handled by the batcher so we do not need to call PostTick() for it
-	for (FNiagaraScriptExecutionContext& EventContext : EventExecContexts)
+
+	for (FNiagaraScriptExecutionContext& EventContext : GetEventExecutionContexts())
 	{
 		EventContext.PostTick();
 	}
diff --git a/Engine/Plugins/FX/Niagara/Source/Niagara/Private/NiagaraSystemInstance.cpp b/Engine/Plugins/FX/Niagara/Source/Niagara/Private/NiagaraSystemInstance.cpp
--- a/Engine/Plugins/FX/Niagara/Source/Niagara/Private/NiagaraSystemInstance.cpp
+++ b/Engine/Plugins/FX/Niagara/Source/Niagara/Private/NiagaraSystemInstance.cpp
@@ -183,7 +183,9 @@ void FNiagaraSystemInstance::Init(bool bInForceSolo)
 #if WITH_EDITORONLY_DATA
 	InstanceParameters.DebugName = *FString::Printf(TEXT("SystemInstance %p"), this);
 #endif
+#if WITH_EDITOR
 	OnInitializedDelegate.Broadcast();
+#endif
 }
 
 void FNiagaraSystemInstance::SetRequestedExecutionState(ENiagaraExecutionState InState)
@@ -569,7 +571,9 @@ void FNiagaraSystemInstance::Complete()
 		// We've already notified once, no need to do so again.
 		bNotifyOnCompletion = false;
 
+#if WITH_EDITOR
 		OnCompleteDelegate.Broadcast(this);
+#endif
 
 		if (Component)
 		{
@@ -2195,6 +2199,7 @@ FNiagaraEmitterInstance* FNiagaraSystemInstance::GetEmitterByID(FGuid InID)
 	return nullptr;
 }
 
+#if WITH_EDITOR
 FNiagaraSystemInstance::FOnInitialized& FNiagaraSystemInstance::OnInitialized()
 {
 	return OnInitializedDelegate;
@@ -2205,7 +2210,6 @@ FNiagaraSystemInstance::FOnComplete& FNiagaraSystemInstance::OnComplete()
 	return OnCompleteDelegate;
 }
 
-#if WITH_EDITOR
 FNiagaraSystemInstance::FOnReset& FNiagaraSystemInstance::OnReset()
 {
 	return OnResetDelegate;
diff --git a/Engine/Plugins/FX/Niagara/Source/Niagara/Private/NiagaraSystemSimulation.cpp b/Engine/Plugins/FX/Niagara/Source/Niagara/Private/NiagaraSystemSimulation.cpp
--- a/Engine/Plugins/FX/Niagara/Source/Niagara/Private/NiagaraSystemSimulation.cpp
+++ b/Engine/Plugins/FX/Niagara/Source/Niagara/Private/NiagaraSystemSimulation.cpp
@@ -1439,7 +1439,7 @@ void FNiagaraSystemSimulation::TransferSystemSimResults(FNiagaraSystemSimulation
 					DataSetToEmitterGPUParameters[EmitterIdx].DataSetToParameterStore(GPUContext->CombinedParamStore, Context.DataSet, SystemIndex);
 				}
 
-				TArray<FNiagaraScriptExecutionContext>& EventContexts = EmitterInst.GetEventExecutionContexts();
+				TArrayView<FNiagaraScriptExecutionContext> EventContexts = EmitterInst.GetEventExecutionContexts();
 				for (int32 EventIdx = 0; EventIdx < EventContexts.Num(); ++EventIdx)
 				{
 					FNiagaraScriptExecutionContext& EventContext = EventContexts[EventIdx];
@@ -1743,7 +1743,7 @@ void FNiagaraSystemSimulation::InitParameterDataSetBindings(FNiagaraSystemInstan
 					DataSetToEmitterGPUParameters[EmitterIdx].Init(MainDataSet, GPUContext->CombinedParamStore);
 				}
 
-				TArray<FNiagaraScriptExecutionContext>& EventContexts = EmitterInst.GetEventExecutionContexts();
+				TArrayView<FNiagaraScriptExecutionContext> EventContexts = EmitterInst.GetEventExecutionContexts();
 				const int32 EventCount = EventContexts.Num();
 				DataSetToEmitterEventParameters[EmitterIdx].SetNum(EventCount);
 
