Fixed check failure when hovering a Map Get pin for which there's no metadata.

#rnx

diff --git a/Engine/Plugins/FX/Niagara/Source/NiagaraEditor/Private/NiagaraNodeParameterMapGet.cpp b/Engine/Plugins/FX/Niagara/Source/NiagaraEditor/Private/NiagaraNodeParameterMapGet.cpp
--- a/Engine/Plugins/FX/Niagara/Source/NiagaraEditor/Private/NiagaraNodeParameterMapGet.cpp
+++ b/Engine/Plugins/FX/Niagara/Source/NiagaraEditor/Private/NiagaraNodeParameterMapGet.cpp
@@ -530,16 +530,23 @@ void UNiagaraNodeParameterMapGet::GetPinHoverText(const UEdGraphPin& Pin, FStrin
 				FNiagaraVariable Var = FNiagaraVariable(TypeDef, Pin.PinName);
 				TOptional<FNiagaraVariableMetaData> Metadata = NiagaraGraph->GetMetaData(Var);
 
-				FText Description = Metadata.IsSet() ? Metadata->Description : FText::GetEmpty();
-				const FText TooltipFormat = LOCTEXT("Parameters", "Name: {0} \nType: {1}\nDescription: {2}\nScope: {3}\nUser Editable: {4}\nUsage: {5}"); 
-				const FText Name = FText::FromName(Var.GetName());
-				FName CachedParamName;
-				Metadata->GetParameterName(CachedParamName);
-				const FText ScopeText = FText::FromName(Metadata->GetScopeName());
-				const FText UserEditableText = FText::FromName(CachedParamName);
-				const FText UsageText = StaticEnum<ENiagaraScriptParameterUsage>()->GetDisplayNameTextByValue((int64)Metadata->GetUsage());
-				const FText ToolTipText = FText::Format(TooltipFormat, FText::FromName(Var.GetName()), Var.GetType().GetNameText(), Description, ScopeText, UserEditableText, UsageText);
-				HoverTextOut = ToolTipText.ToString();
+				if (Metadata.IsSet())
+				{
+					FText Description = Metadata->Description;
+					const FText TooltipFormat = LOCTEXT("Parameters", "Name: {0} \nType: {1}\nDescription: {2}\nScope: {3}\nUser Editable: {4}\nUsage: {5}");
+					const FText Name = FText::FromName(Var.GetName());
+					FName CachedParamName;
+					Metadata->GetParameterName(CachedParamName);
+					const FText ScopeText = FText::FromName(Metadata->GetScopeName());
+					const FText UserEditableText = FText::FromName(CachedParamName);
+					const FText UsageText = StaticEnum<ENiagaraScriptParameterUsage>()->GetDisplayNameTextByValue((int64)Metadata->GetUsage());
+					const FText ToolTipText = FText::Format(TooltipFormat, FText::FromName(Var.GetName()), Var.GetType().GetNameText(), Description, ScopeText, UserEditableText, UsageText);
+					HoverTextOut = ToolTipText.ToString();
+				}
+				else
+				{
+					HoverTextOut = TEXT("");
+				}
 			}
 		}
 	}
