Renaming a Niagara module to its default name now clears the display name.

#rnx

diff --git a/Engine/Plugins/FX/Niagara/Source/NiagaraEditor/Private/ViewModels/Stack/NiagaraStackEntry.cpp b/Engine/Plugins/FX/Niagara/Source/NiagaraEditor/Private/ViewModels/Stack/NiagaraStackEntry.cpp
--- a/Engine/Plugins/FX/Niagara/Source/NiagaraEditor/Private/ViewModels/Stack/NiagaraStackEntry.cpp
+++ b/Engine/Plugins/FX/Niagara/Source/NiagaraEditor/Private/ViewModels/Stack/NiagaraStackEntry.cpp
@@ -799,14 +799,14 @@ void UNiagaraStackEntry::OnRenamed(FText NewName)
 {
 	if (SupportsRename())
 	{
-		if (!NewName.EqualTo(GetDisplayName()))
+		if (!NewName.EqualTo(AlternateDisplayName.Get(FText::GetEmpty())))
 		{
 			FScopedTransaction ScopedTransaction(NSLOCTEXT("NiagaraStackEntry", "RenameModule", "Rename Module"));
 
 			GetStackEditorData().Modify();
 			GetStackEditorData().SetStackEntryDisplayName(GetStackEditorDataKey(), NewName);
 
-			if (NewName.IsEmptyOrWhitespace())
+			if (NewName.IsEmptyOrWhitespace() || NewName.EqualTo(GetDisplayName()))
 			{
 				AlternateDisplayName.Reset();
 			}
