Fix for crash with server accessing invalid GPUComputeScript resource

diff --git a/Engine/Plugins/FX/Niagara/Source/Niagara/Private/NiagaraScript.cpp b/Engine/Plugins/FX/Niagara/Source/Niagara/Private/NiagaraScript.cpp
--- a/Engine/Plugins/FX/Niagara/Source/Niagara/Private/NiagaraScript.cpp
+++ b/Engine/Plugins/FX/Niagara/Source/Niagara/Private/NiagaraScript.cpp
@@ -2068,45 +2068,45 @@ NIAGARA_API bool UNiagaraScript::IsScriptCompilationPending(bool bGPUScript) con
 {
 	if (bGPUScript)
 	{
-		FNiagaraShaderRef Shader = ScriptResource->GetShaderGameThread();
-		if (Shader.IsValid())
+		if (ScriptResource.IsValid())
 		{
-			return false;
+			FNiagaraShaderRef Shader = ScriptResource->GetShaderGameThread();
+			if (Shader.IsValid())
+			{
+				return false;
+			}
+			return !ScriptResource->IsCompilationFinished();
 		}
-		return !ScriptResource->IsCompilationFinished();
 	}
-	else
+	else if (CachedScriptVM.IsValid())
 	{
-		if (CachedScriptVM.IsValid())
-		{
-			return (CachedScriptVM.ByteCode.Num() == 0) && (CachedScriptVM.OptimizedByteCode.Num() == 0) && (CachedScriptVM.LastCompileStatus == ENiagaraScriptCompileStatus::NCS_BeingCreated || CachedScriptVM.LastCompileStatus == ENiagaraScriptCompileStatus::NCS_Unknown);
-		}
-		return false;
+		return (CachedScriptVM.ByteCode.Num() == 0) && (CachedScriptVM.OptimizedByteCode.Num() == 0) && (CachedScriptVM.LastCompileStatus == ENiagaraScriptCompileStatus::NCS_BeingCreated || CachedScriptVM.LastCompileStatus == ENiagaraScriptCompileStatus::NCS_Unknown);
 	}
+	return false;
 }
 
 NIAGARA_API bool UNiagaraScript::DidScriptCompilationSucceed(bool bGPUScript) const
 {
 	if (bGPUScript)
 	{
-		FNiagaraShaderRef Shader = ScriptResource->GetShaderGameThread();
-		if (Shader.IsValid())
+		if (ScriptResource.IsValid())
 		{
-			return true;
-		}
+			FNiagaraShaderRef Shader = ScriptResource->GetShaderGameThread();
+			if (Shader.IsValid())
+			{
+				return true;
+			}
 
-		if (ScriptResource->IsCompilationFinished())
-		{
-			// If we failed compilation, it would be finished and Shader would be null.
-			return false;
+			if (ScriptResource->IsCompilationFinished())
+			{
+				// If we failed compilation, it would be finished and Shader would be null.
+				return false;
+			}
 		}
 	}
-	else
+	else if (CachedScriptVM.IsValid())
 	{
-		if (CachedScriptVM.IsValid())
-		{
-			return (CachedScriptVM.ByteCode.Num() != 0) || (CachedScriptVM.OptimizedByteCode.Num() != 0);
-		}
+		return (CachedScriptVM.ByteCode.Num() != 0) || (CachedScriptVM.OptimizedByteCode.Num() != 0);
 	}
 
 	return false;
