Adding icons to the Emitter Properties row to indicate CPU or GPU sim target

diff --git a/Engine/Plugins/FX/Niagara/Source/NiagaraEditor/Private/NiagaraEditorStyle.cpp b/Engine/Plugins/FX/Niagara/Source/NiagaraEditor/Private/NiagaraEditorStyle.cpp
--- a/Engine/Plugins/FX/Niagara/Source/NiagaraEditor/Private/NiagaraEditorStyle.cpp
+++ b/Engine/Plugins/FX/Niagara/Source/NiagaraEditor/Private/NiagaraEditorStyle.cpp
@@ -262,6 +262,11 @@ TSharedRef< FSlateStyleSet > FNiagaraEditorStyle::Create()
 	Style->Set("ClassIcon.NiagaraLightRendererProperties", new IMAGE_PLUGIN_BRUSH("Icons/Renderers/renderer_light", Icon16x16));
 	Style->Set("ClassIcon.NiagaraRendererProperties", new IMAGE_PLUGIN_BRUSH("Icons/Renderers/renderer_default", Icon16x16));
 
+	//GPU/CPU icons
+	Style->Set("NiagaraEditor.Stack.GPUIcon", new IMAGE_PLUGIN_BRUSH("Icons/Simulate_GPU_x40", Icon16x16));
+	Style->Set("NiagaraEditor.Stack.CPUIcon", new IMAGE_PLUGIN_BRUSH("Icons/Simulate_CPU_x40", Icon16x16));
+
+
 	// Niagara sequence
 	Style->Set("NiagaraEditor.NiagaraSequence.DefaultTrackColor", FLinearColor(0, .25f, 0));
 
diff --git a/Engine/Plugins/FX/Niagara/Source/NiagaraEditor/Private/ViewModels/Stack/NiagaraStackEmitterSettingsGroup.cpp b/Engine/Plugins/FX/Niagara/Source/NiagaraEditor/Private/ViewModels/Stack/NiagaraStackEmitterSettingsGroup.cpp
--- a/Engine/Plugins/FX/Niagara/Source/NiagaraEditor/Private/ViewModels/Stack/NiagaraStackEmitterSettingsGroup.cpp
+++ b/Engine/Plugins/FX/Niagara/Source/NiagaraEditor/Private/ViewModels/Stack/NiagaraStackEmitterSettingsGroup.cpp
@@ -10,6 +10,7 @@
 #include "NiagaraScriptMergeManager.h"
 #include "NiagaraEmitterDetailsCustomization.h"
 #include "NiagaraSystem.h"
+#include "NiagaraEditorStyle.h"
 
 #define LOCTEXT_NAMESPACE "UNiagaraStackEmitterItemGroup"
 
@@ -82,6 +83,19 @@ bool UNiagaraStackEmitterPropertiesItem::IsExpandedByDefault() const
 	return false;
 }
 
+const FSlateBrush* UNiagaraStackEmitterPropertiesItem::GetIconBrush() const
+{
+	if (Emitter->SimTarget == ENiagaraSimTarget::CPUSim)
+	{
+		return FNiagaraEditorStyle::Get().GetBrush("NiagaraEditor.Stack.CPUIcon");
+	}
+	if (Emitter->SimTarget == ENiagaraSimTarget::GPUComputeSim)
+	{
+		return FNiagaraEditorStyle::Get().GetBrush("NiagaraEditor.Stack.GPUIcon");
+	}
+	return FEditorStyle::GetBrush("NoBrush");
+}
+
 void UNiagaraStackEmitterPropertiesItem::RefreshChildrenInternal(const TArray<UNiagaraStackEntry*>& CurrentChildren, TArray<UNiagaraStackEntry*>& NewChildren, TArray<FStackIssue>& NewIssues)
 {
 	if (EmitterObject == nullptr)
diff --git a/Engine/Plugins/FX/Niagara/Source/NiagaraEditor/Public/ViewModels/Stack/NiagaraStackEmitterSettingsGroup.h b/Engine/Plugins/FX/Niagara/Source/NiagaraEditor/Public/ViewModels/Stack/NiagaraStackEmitterSettingsGroup.h
--- a/Engine/Plugins/FX/Niagara/Source/NiagaraEditor/Public/ViewModels/Stack/NiagaraStackEmitterSettingsGroup.h
+++ b/Engine/Plugins/FX/Niagara/Source/NiagaraEditor/Public/ViewModels/Stack/NiagaraStackEmitterSettingsGroup.h
@@ -27,6 +27,8 @@ public:
 	virtual void ResetToBase() override;
 
 	virtual bool IsExpandedByDefault() const override;
+	virtual bool SupportsIcon() const { return true; }
+	virtual const FSlateBrush* GetIconBrush() const override;
 
 protected:
 	virtual void FinalizeInternal() override;
diff --git a/Engine/Plugins/FX/Niagara/Source/NiagaraEditorWidgets/Private/SNiagaraOverviewStack.cpp b/Engine/Plugins/FX/Niagara/Source/NiagaraEditorWidgets/Private/SNiagaraOverviewStack.cpp
--- a/Engine/Plugins/FX/Niagara/Source/NiagaraEditorWidgets/Private/SNiagaraOverviewStack.cpp
+++ b/Engine/Plugins/FX/Niagara/Source/NiagaraEditorWidgets/Private/SNiagaraOverviewStack.cpp
@@ -556,7 +556,7 @@ TSharedRef<ITableRow> SNiagaraOverviewStack::OnGenerateRowForEntry(UNiagaraStack
 				.VAlign(VAlign_Center)
 				[
 					SNew(SImage)
-					.Image(StackItem->GetIconBrush())
+					.Image_UObject(StackItem, &UNiagaraStackItem::GetIconBrush)
 				];
 		}
 		else
diff --git a/Engine/Plugins/FX/Niagara/Source/NiagaraEditorWidgets/Private/Stack/SNiagaraStackItem.cpp b/Engine/Plugins/FX/Niagara/Source/NiagaraEditorWidgets/Private/Stack/SNiagaraStackItem.cpp
--- a/Engine/Plugins/FX/Niagara/Source/NiagaraEditorWidgets/Private/Stack/SNiagaraStackItem.cpp
+++ b/Engine/Plugins/FX/Niagara/Source/NiagaraEditorWidgets/Private/Stack/SNiagaraStackItem.cpp
@@ -34,7 +34,7 @@ void SNiagaraStackItem::Construct(const FArguments& InArgs, UNiagaraStackItem& I
 		.VAlign(VAlign_Center)
 		[
 			SNew(SImage)
-			.Image(Item->GetIconBrush())
+			.Image_UObject(Item, &UNiagaraStackItem::GetIconBrush)
 		];
 	}
 
