Additional fixes for migrating ScriptResource to be optional, failed with nullrhi

diff --git a/Engine/Plugins/FX/Niagara/Source/Niagara/Private/NiagaraScript.cpp b/Engine/Plugins/FX/Niagara/Source/Niagara/Private/NiagaraScript.cpp
--- a/Engine/Plugins/FX/Niagara/Source/Niagara/Private/NiagaraScript.cpp
+++ b/Engine/Plugins/FX/Niagara/Source/Niagara/Private/NiagaraScript.cpp
@@ -2174,9 +2174,15 @@ void UNiagaraScript::SerializeNiagaraShaderMaps(FArchive& Ar, int32 NiagaraVer,
 			FNiagaraShaderScript Resource;
 			Resource.SerializeShaderMap(Ar);
 
-			if (!ScriptResource && GMaxRHIShaderPlatform == Resource.GetGameThreadShaderMap()->GetShaderPlatform())
+			if (!ScriptResource)
 			{
-				ScriptResource = MakeUnique<FNiagaraShaderScript>(Resource);
+				if (FNiagaraShaderMap* ShaderMap = Resource.GetGameThreadShaderMap())
+				{
+					if (GMaxRHIShaderPlatform == ShaderMap->GetShaderPlatform())
+					{
+						ScriptResource = MakeUnique<FNiagaraShaderScript>(Resource);
+					}
+				}
 			}
 		}
 	}
