- Remove erroneous checkSlow, we can be complete here due to disabling if the spawn / update funcitons fail to bind to VM functions

#rnx

diff --git a/Engine/Plugins/FX/Niagara/Source/Niagara/Private/NiagaraSystemSimulation.cpp b/Engine/Plugins/FX/Niagara/Source/Niagara/Private/NiagaraSystemSimulation.cpp
--- a/Engine/Plugins/FX/Niagara/Source/Niagara/Private/NiagaraSystemSimulation.cpp
+++ b/Engine/Plugins/FX/Niagara/Source/Niagara/Private/NiagaraSystemSimulation.cpp
@@ -1112,7 +1112,6 @@ void FNiagaraSystemSimulation::Tick_Concurrent(FNiagaraSystemSimulationTickConte
 			while (SystemInstIndex < Context.Instances.Num())
 			{
 				FNiagaraSystemInstance* Inst = Context.Instances[SystemInstIndex];
-				checkSlow(!Inst->IsComplete());
 				Inst->FinalizeTick_GameThread();
 
 				// If the system completes during finalize it will be removed from the instances, therefore we do not need to increment our system index;
