Adding icons to the Emitter Properties row to indicate CPU or GPU sim target
#rnx

diff --git a/Engine/Plugins/FX/Niagara/Source/NiagaraEditor/Private/ViewModels/Stack/NiagaraStackEmitterSettingsGroup.cpp b/Engine/Plugins/FX/Niagara/Source/NiagaraEditor/Private/ViewModels/Stack/NiagaraStackEmitterSettingsGroup.cpp
--- a/Engine/Plugins/FX/Niagara/Source/NiagaraEditor/Private/ViewModels/Stack/NiagaraStackEmitterSettingsGroup.cpp
+++ b/Engine/Plugins/FX/Niagara/Source/NiagaraEditor/Private/ViewModels/Stack/NiagaraStackEmitterSettingsGroup.cpp
@@ -9,6 +9,7 @@
 #include "ViewModels/Stack/NiagaraStackGraphUtilities.h"
 #include "NiagaraScriptMergeManager.h"
 #include "NiagaraEmitterDetailsCustomization.h"
+#include "NiagaraEditorStyle.h"
 #include "NiagaraSystem.h"
 #include "NiagaraEditorStyle.h"
 
