Fix for ParameterVariables view of parameter stores to use a virtual rather than the cached ArrayView to eliminate the risk of it getting out of date

diff --git a/Engine/Plugins/FX/Niagara/Source/Niagara/Private/NiagaraCommon.cpp b/Engine/Plugins/FX/Niagara/Source/Niagara/Private/NiagaraCommon.cpp
--- a/Engine/Plugins/FX/Niagara/Source/Niagara/Private/NiagaraCommon.cpp
+++ b/Engine/Plugins/FX/Niagara/Source/Niagara/Private/NiagaraCommon.cpp
@@ -438,14 +438,16 @@ void FNiagaraUtilities::PrepareRapidIterationParameters(const TArray<UNiagaraScr
 		UNiagaraScript* Script = It.Key();
 		FNiagaraParameterStore& PreparedParameterStore = It.Value();
 
+		auto RapidIterationParameters = Script->RapidIterationParameters.ReadParameterVariables();
+
 		bool bOverwriteParameters = false;
-		if (Script->RapidIterationParameters.ParameterVariables.Num() != PreparedParameterStore.ParameterVariables.Num())
+		if (RapidIterationParameters.Num() != PreparedParameterStore.ReadParameterVariables().Num())
 		{
 			bOverwriteParameters = true;
 		}
 		else
 		{
-			for (const FNiagaraVariableWithOffset& ParamWithOffset : Script->RapidIterationParameters.ParameterVariables)
+			for (const FNiagaraVariableWithOffset& ParamWithOffset : RapidIterationParameters)
 			{
 				const FNiagaraVariable& SourceParameter = ParamWithOffset;
 				const int32 SourceOffset = ParamWithOffset.Offset;
diff --git a/Engine/Plugins/FX/Niagara/Source/Niagara/Private/NiagaraParameterStore.cpp b/Engine/Plugins/FX/Niagara/Source/Niagara/Private/NiagaraParameterStore.cpp
--- a/Engine/Plugins/FX/Niagara/Source/Niagara/Private/NiagaraParameterStore.cpp
+++ b/Engine/Plugins/FX/Niagara/Source/Niagara/Private/NiagaraParameterStore.cpp
@@ -105,7 +105,7 @@ static FAutoConsoleVariableRef CVarNiagaraAllowQuickSortedParameterOffetsCopy(
 	ECVF_Scalability
 );
 
-void FNiagaraParameterStore::CopySortedParameterOffsets(TArrayView<FNiagaraVariableWithOffset> Src)
+void FNiagaraParameterStore::CopySortedParameterOffsets(TArrayView<const FNiagaraVariableWithOffset> Src)
 {
 	if (GNiagaraAllowQuickSortedParameterOffetsCopy)
 	{
@@ -118,8 +118,6 @@ void FNiagaraParameterStore::CopySortedParameterOffsets(TArrayView<FNiagaraVaria
 	{
 		SortedParameterOffsets = TArray<FNiagaraVariableWithOffset>(Src.GetData(), Src.Num());
 	}
-
-	InternalStorageChanged();
 }
 
 //////////////////////////////////////////////////////////////////////////
@@ -154,7 +152,7 @@ FNiagaraParameterStore& FNiagaraParameterStore::operator=(const FNiagaraParamete
 #if WITH_EDITORONLY_DATA
 	ParameterOffsets = Other.ParameterOffsets;
 #endif // WITH_EDITORONLY_DATA
-	CopySortedParameterOffsets(Other.ParameterVariables);
+	CopySortedParameterOffsets(Other.ReadParameterVariables());
 	DEC_MEMORY_STAT_BY(STAT_NiagaraParamStoreMemory, ParameterData.GetAllocatedSize());
 	ParameterData = Other.ParameterData;
 	INC_MEMORY_STAT_BY(STAT_NiagaraParamStoreMemory, ParameterData.GetAllocatedSize());
@@ -200,8 +198,8 @@ void FNiagaraParameterStore::Bind(FNiagaraParameterStore* DestStore, const FNiag
 template <typename TVisitor>
 void FNiagaraParameterStoreBinding::MatchParameters(FNiagaraParameterStore* DestStore, FNiagaraParameterStore* SrcStore, TVisitor Visitor)
 {
-	TArrayView<FNiagaraVariableWithOffset> SrcParamWithOffsets = SrcStore->ParameterVariables;
-	TArrayView<FNiagaraVariableWithOffset> DestParamWithOffsets = DestStore->ParameterVariables;
+	TArrayView<const FNiagaraVariableWithOffset> SrcParamWithOffsets = SrcStore->ReadParameterVariables();
+	TArrayView<const FNiagaraVariableWithOffset> DestParamWithOffsets = DestStore->ReadParameterVariables();
 
 	const int32 SrcNum = SrcParamWithOffsets.Num();
 	const int32 DestNum = DestParamWithOffsets.Num();
@@ -376,7 +374,7 @@ bool FNiagaraParameterStore::VerifyBinding(const FNiagaraParameterStore* DestSto
 
 void FNiagaraParameterStore::CheckForNaNs()const
 {
-	for (const FNiagaraVariableWithOffset& Var : ParameterVariables)
+	for (const FNiagaraVariableWithOffset& Var : ReadParameterVariables())
 	{
 		const int32 Offset = Var.Offset;
 
@@ -439,7 +437,7 @@ void FNiagaraParameterStore::UnbindFromSourceStores()
 
 void FNiagaraParameterStore::DumpParameters(bool bDumpBindings)const
 {
-	for (const FNiagaraVariableWithOffset& VariableBase : ParameterVariables)
+	for (const FNiagaraVariableWithOffset& VariableBase : ReadParameterVariables())
 	{
 		FNiagaraVariable Var(VariableBase);
 		Var.SetData(GetParameterData(VariableBase.Offset));
@@ -459,7 +457,7 @@ FString FNiagaraParameterStore::ToString() const
 {
 	FString Value;
 
-	for (const FNiagaraVariableWithOffset& VariableBase : ParameterVariables)
+	for (const FNiagaraVariableWithOffset& VariableBase : ReadParameterVariables())
 	{
 		FNiagaraVariable Var(VariableBase);
 		Var.SetData(GetParameterData(VariableBase.Offset));
@@ -493,6 +491,8 @@ bool FNiagaraParameterStore::AddParameter(const FNiagaraVariable& Param, bool bI
 	check(!ParameterOffsets.Num()); // Migration to SortedParameterOffsets
 #endif
 
+	auto ParameterVariables = ReadParameterVariables();
+
 	int32 InsertPos = 0;
 	if (ParameterVariables.Num())
 	{
@@ -508,7 +508,6 @@ bool FNiagaraParameterStore::AddParameter(const FNiagaraVariable& Param, bool bI
 	}
 
 	int32& Offset = SortedParameterOffsets.EmplaceAt_GetRef(InsertPos, Param, (int32)INDEX_NONE).Offset;
-	InternalStorageChanged();
 
 	if (Param.GetType().IsDataInterface())
 	{
@@ -574,7 +573,7 @@ bool FNiagaraParameterStore::RemoveParameter(const FNiagaraVariableBase& ToRemov
 		TArray<uint8> NewData;
 		TArray<UNiagaraDataInterface*> NewInterfaces;
 		TArray<UObject*> NewUObjects;
-		for (const FNiagaraVariableWithOffset& Existing : ParameterVariables)
+		for (const FNiagaraVariableWithOffset& Existing : ReadParameterVariables())
 		{
 			const FNiagaraVariable& ExistingVar = Existing;
 			const int32 ExistingOffset = Existing.Offset;
@@ -671,7 +670,7 @@ void FNiagaraParameterStore::SanityCheckData(bool bInitInterfaces)
 
 	int32 ParameterDataSize = 0;
 
-	for (const FNiagaraVariableWithOffset& Parameter : ParameterVariables)
+	for (const FNiagaraVariableWithOffset& Parameter : ReadParameterVariables())
 	{
 		const int32 SrcIndex = Parameter.Offset;
 
@@ -735,7 +734,7 @@ void FNiagaraParameterStore::SanityCheckData(bool bInitInterfaces)
 
 void FNiagaraParameterStore::CopyParametersTo(FNiagaraParameterStore& DestStore, bool bOnlyAdd, EDataInterfaceCopyMethod DataInterfaceCopyMethod) const
 {
-	for (const FNiagaraVariableWithOffset& Parameter : ParameterVariables)
+	for (const FNiagaraVariableWithOffset& Parameter : ReadParameterVariables())
 	{
 		int32 SrcIndex = Parameter.Offset;
 
@@ -856,7 +855,7 @@ void FNiagaraParameterStore::InitFromSource(const FNiagaraParameterStore* SrcSto
 #if WITH_EDITORONLY_DATA
 	ParameterOffsets = SrcStore->ParameterOffsets;
 #endif // WITH_EDITORONLY_DATA
-	CopySortedParameterOffsets(SrcStore->ParameterVariables);
+	CopySortedParameterOffsets(SrcStore->ReadParameterVariables());
 	DEC_MEMORY_STAT_BY(STAT_NiagaraParamStoreMemory, ParameterData.GetAllocatedSize());
 	ParameterData = SrcStore->ParameterData;
 	INC_MEMORY_STAT_BY(STAT_NiagaraParamStoreMemory, ParameterData.GetAllocatedSize());
@@ -876,7 +875,7 @@ void FNiagaraParameterStore::InitFromSource(const FNiagaraParameterStore* SrcSto
 
 void FNiagaraParameterStore::RemoveParameters(FNiagaraParameterStore& DestStore)
 {
-	for (const FNiagaraVariableWithOffset& Parameter : ParameterVariables)
+	for (const FNiagaraVariableWithOffset& Parameter : ReadParameterVariables())
 	{
 		DestStore.RemoveParameter(Parameter);
 	}
@@ -889,7 +888,6 @@ void FNiagaraParameterStore::Empty(bool bClearBindings)
 #endif // WITH_EDITORONLY_DATA
 
 	SortedParameterOffsets.Empty();
-	InternalStorageChanged();
 
 	DEC_MEMORY_STAT_BY(STAT_NiagaraParamStoreMemory, ParameterData.GetAllocatedSize());
 	ParameterData.Empty();
@@ -913,7 +911,6 @@ void FNiagaraParameterStore::Reset(bool bClearBindings)
 #endif // WITH_EDITORONLY_DATA
 
 	SortedParameterOffsets.Reset();
-	InternalStorageChanged();
 
 	DEC_MEMORY_STAT_BY(STAT_NiagaraParamStoreMemory, ParameterData.GetAllocatedSize());
 	ParameterData.Reset();
@@ -949,18 +946,13 @@ void FNiagaraParameterStore::OnLayoutChange()
 #endif
 }
 
-void FNiagaraParameterStore::InternalStorageChanged()
-{
-	ParameterVariables = MakeArrayView(SortedParameterOffsets);
-}
-
 const FNiagaraVariableBase* FNiagaraParameterStore::FindVariable(UNiagaraDataInterface* Interface)const
 {
 	SCOPE_CYCLE_COUNTER(STAT_NiagaraParameterStoreFindVar);
 	int32 Idx = DataInterfaces.IndexOfByKey(Interface);
 	if (Idx != INDEX_NONE)
 	{
-		for (const FNiagaraVariableWithOffset& ParamWithOffset : ParameterVariables)
+		for (const FNiagaraVariableWithOffset& ParamWithOffset : ReadParameterVariables())
 		{
 			if (ParamWithOffset.Offset == Idx && ParamWithOffset.GetType().GetClass() == Interface->GetClass())
 			{
@@ -977,6 +969,7 @@ const int32* FNiagaraParameterStore::FindParameterOffset(const FNiagaraVariableB
 	check(!ParameterOffsets.Num()); // Migration to SortedParameterOffsets
 #endif
 
+	auto ParameterVariables = ReadParameterVariables();
 	if (ParameterVariables.Num())
 	{
 		int32 MatchingIndex = 0;
@@ -999,8 +992,6 @@ void FNiagaraParameterStore::PostLoad()
 			SortedParameterOffsets.Emplace(ParamOffsetPair.Key, ParamOffsetPair.Value);
 		}
 		ParameterOffsets.Empty();
-
-		InternalStorageChanged();
 	}
 #endif
 
@@ -1014,7 +1005,6 @@ void FNiagaraParameterStore::SortParameters()
 	{
 		return FNiagaraVariableSearch::Compare(Lhs, Rhs) < 0;
 	});
-	InternalStorageChanged();
 }
 
 #if WITH_EDITOR
diff --git a/Engine/Plugins/FX/Niagara/Source/Niagara/Private/NiagaraScript.cpp b/Engine/Plugins/FX/Niagara/Source/Niagara/Private/NiagaraScript.cpp
--- a/Engine/Plugins/FX/Niagara/Source/Niagara/Private/NiagaraScript.cpp
+++ b/Engine/Plugins/FX/Niagara/Source/Niagara/Private/NiagaraScript.cpp
@@ -826,9 +826,11 @@ void UNiagaraScript::Serialize(FArchive& Ar)
 			// Copy off the parameter store for now..
 			TemporaryStore = RapidIterationParameters;
 
+			auto ParameterVariables = TemporaryStore.ReadParameterVariables();
+
 			// Get the active parameters
 			// Remove all parameters that aren't data interfaces or uobjects
-			for (const FNiagaraVariableBase& Var : TemporaryStore.ParameterVariables)
+			for (const FNiagaraVariableBase& Var : ParameterVariables)
 			{
 				if (Var.IsDataInterface() || Var.IsUObject())
 					continue;
@@ -836,7 +838,7 @@ void UNiagaraScript::Serialize(FArchive& Ar)
 				NumRemoved++;
 			}
 
-			UE_LOG(LogNiagara, Verbose, TEXT("Pruned %d/%d parameters from script %s"), NumRemoved, TemporaryStore.ParameterVariables.Num(), *GetFullName());
+			UE_LOG(LogNiagara, Verbose, TEXT("Pruned %d/%d parameters from script %s"), NumRemoved, ParameterVariables.Num(), *GetFullName());
 		}
 	}
 
diff --git a/Engine/Plugins/FX/Niagara/Source/Niagara/Private/NiagaraScriptExecutionParameterStore.cpp b/Engine/Plugins/FX/Niagara/Source/Niagara/Private/NiagaraScriptExecutionParameterStore.cpp
--- a/Engine/Plugins/FX/Niagara/Source/Niagara/Private/NiagaraScriptExecutionParameterStore.cpp
+++ b/Engine/Plugins/FX/Niagara/Source/Niagara/Private/NiagaraScriptExecutionParameterStore.cpp
@@ -320,7 +320,6 @@ void FNiagaraScriptInstanceParameterStore::InitFromOwningContext(UNiagaraScript*
 		}
 
 		ScriptParameterStore.Init(SrcStore);
-		InternalStorageChanged();
 	}
 	else
 	{
@@ -375,15 +374,15 @@ void FNiagaraScriptInstanceParameterStore::CopyParameterDataToPaddedBuffer(uint8
 	}
 }
 
-void FNiagaraScriptInstanceParameterStore::InternalStorageChanged()
+TArrayView<const FNiagaraVariableWithOffset> FNiagaraScriptInstanceParameterStore::ReadParameterVariables() const
 {
 	if (const auto* ScriptStore = ScriptParameterStore.Get())
 	{
-		ParameterVariables = ScriptStore->ParameterVariables;
+		return ScriptStore->ReadParameterVariables();
 	}
 	else
 	{
-		ParameterVariables = TArrayView<FNiagaraVariableWithOffset>();
+		return TArrayView<FNiagaraVariableWithOffset>();
 	}
 }
 
diff --git a/Engine/Plugins/FX/Niagara/Source/Niagara/Private/NiagaraSystem.cpp b/Engine/Plugins/FX/Niagara/Source/Niagara/Private/NiagaraSystem.cpp
--- a/Engine/Plugins/FX/Niagara/Source/Niagara/Private/NiagaraSystem.cpp
+++ b/Engine/Plugins/FX/Niagara/Source/Niagara/Private/NiagaraSystem.cpp
@@ -457,7 +457,7 @@ void UNiagaraSystem::PostLoad()
 			InitEmitterCompiledData();
 		}
 
-		if (SystemCompiledData.InstanceParamStore.ParameterVariables.Num() == 0 ||SystemCompiledData.DataSetCompiledData.Variables.Num() == 0)
+		if (SystemCompiledData.InstanceParamStore.ReadParameterVariables().Num() == 0 ||SystemCompiledData.DataSetCompiledData.Variables.Num() == 0)
 		{
 			InitSystemCompiledData();
 		}
diff --git a/Engine/Plugins/FX/Niagara/Source/Niagara/Private/NiagaraUserRedirectionParameterStore.cpp b/Engine/Plugins/FX/Niagara/Source/Niagara/Private/NiagaraUserRedirectionParameterStore.cpp
--- a/Engine/Plugins/FX/Niagara/Source/Niagara/Private/NiagaraUserRedirectionParameterStore.cpp
+++ b/Engine/Plugins/FX/Niagara/Source/Niagara/Private/NiagaraUserRedirectionParameterStore.cpp
@@ -44,7 +44,7 @@ void FNiagaraUserRedirectionParameterStore::RecreateRedirections()
 {
 	UserParameterRedirects.Reset();
 
-	for (const FNiagaraVariable& Var : ParameterVariables)
+	for (const FNiagaraVariable& Var : ReadParameterVariables())
 	{
 		if (IsUserParameter(Var))
 		{
diff --git a/Engine/Plugins/FX/Niagara/Source/Niagara/Public/NiagaraParameterStore.h b/Engine/Plugins/FX/Niagara/Source/Niagara/Public/NiagaraParameterStore.h
--- a/Engine/Plugins/FX/Niagara/Source/Niagara/Public/NiagaraParameterStore.h
+++ b/Engine/Plugins/FX/Niagara/Source/Niagara/Public/NiagaraParameterStore.h
@@ -140,7 +140,7 @@ struct NIAGARA_API FNiagaraParameterStore
 	/** The View of the set of variables represented by this ParameterStore.  By default it will be the variables defined by the
 	SortedParameterOffsets, but child classes can override it so that we can share the list of Variables across different
 	instances as it's generally read only data in game */
-	TArrayView<FNiagaraVariableWithOffset> ParameterVariables;
+	virtual TArrayView<const FNiagaraVariableWithOffset> ReadParameterVariables() const { return MakeArrayView(SortedParameterOffsets); }
 
 private:
 	/** Owner of this store. Used to provide an outer to data interfaces in this store. */
@@ -153,8 +153,8 @@ private:
 	TMap<FNiagaraVariable, int32> ParameterOffsets;
 #endif // WITH_EDITORONLY_DATA
 
-	/** Storage for the set of variables that are represented by this ParameterStore.  Can be empty, and changes to it should
-	also trigger InternalStorageChanged() so that the public ParameterVariables view can be updated. */
+	/** Storage for the set of variables that are represented by this ParameterStore.  Shouldn't be accecssed directly, instead use
+	ReadParameterVariables() */
 	UPROPERTY()
 	TArray<FNiagaraVariableWithOffset> SortedParameterOffsets;
 
@@ -268,6 +268,8 @@ public:
 
 	FORCEINLINE void GetParameters(TArray<FNiagaraVariable>& OutParameters) const
 	{
+		auto ParameterVariables = ReadParameterVariables();
+
 		OutParameters.Reserve(ParameterVariables.Num());
 		for (const FNiagaraVariableWithOffset& ParamWithOffset : ParameterVariables)
 		{
@@ -580,7 +582,7 @@ public:
 protected:
 	void TickBindings();
 	void OnLayoutChange();
-	void CopySortedParameterOffsets(TArrayView<FNiagaraVariableWithOffset> Src);
+	void CopySortedParameterOffsets(TArrayView<const FNiagaraVariableWithOffset> Src);
 
 	/** Returns the parameter data at the passed offset. */
 	FORCEINLINE uint8* GetParameterData_Internal(int32 Offset) 
@@ -600,8 +602,6 @@ protected:
 	void SetDataInterfaces(const TArray<UNiagaraDataInterface*>& InDataInterfaces, bool bNotifyAsDirty = true);
 	void SetUObjects(const TArray<UObject*>& InUObjects, bool bNotifyAsDirty = true);
 
-	virtual void InternalStorageChanged();
-
 	friend struct FNiagaraParameterStoreToDataSetBinding;    // this should be the only class calling SetParameterByOffset
 };
 
@@ -654,7 +654,7 @@ FORCEINLINE_DEBUGGABLE bool FNiagaraParameterStoreBinding::VerifyBinding(const F
 	bool bBindingValid = true;
 #if WITH_EDITORONLY_DATA
 	TArray<FName, TInlineAllocator<32>> MissingParameterNames;
-	for (const FNiagaraVariableWithOffset& ParamWithOffset : DestStore->ParameterVariables)
+	for (const FNiagaraVariableWithOffset& ParamWithOffset : DestStore->ReadParameterVariables())
 	{
 		const FNiagaraVariable& Parameter = ParamWithOffset;
 		int32 DestOffset = ParamWithOffset.Offset;
@@ -743,7 +743,7 @@ FORCEINLINE_DEBUGGABLE void FNiagaraParameterStoreBinding::Dump(const FNiagaraPa
 		ensure(Binding.DestOffset != -1);
 		FNiagaraVariable Param;
 		bool bFound = false;
-		for (const FNiagaraVariableWithOffset& ParamWithOffset : DestStore->ParameterVariables)
+		for (const FNiagaraVariableWithOffset& ParamWithOffset : DestStore->ReadParameterVariables())
 		{
 			if (ParamWithOffset.Offset == Binding.DestOffset && !ParamWithOffset.IsDataInterface())
 			{
@@ -775,7 +775,7 @@ FORCEINLINE_DEBUGGABLE void FNiagaraParameterStoreBinding::Dump(const FNiagaraPa
 		ensure(Binding.DestOffset != -1);
 		FNiagaraVariable Param;
 		bool bFound = false;
-		for (const FNiagaraVariableWithOffset& ParamWithOffset : DestStore->ParameterVariables)
+		for (const FNiagaraVariableWithOffset& ParamWithOffset : DestStore->ReadParameterVariables())
 		{
 			if (ParamWithOffset.Offset == Binding.DestOffset && ParamWithOffset.IsDataInterface())
 			{
@@ -807,7 +807,7 @@ FORCEINLINE_DEBUGGABLE void FNiagaraParameterStoreBinding::Dump(const FNiagaraPa
 		ensure(Binding.DestOffset != -1);
 		FNiagaraVariable Param;
 		bool bFound = false;
-		for (const FNiagaraVariableWithOffset& ParamWithOffset : DestStore->ParameterVariables)
+		for (const FNiagaraVariableWithOffset& ParamWithOffset : DestStore->ReadParameterVariables())
 		{
 			if (ParamWithOffset.Offset == Binding.DestOffset && ParamWithOffset.IsUObject())
 			{
diff --git a/Engine/Plugins/FX/Niagara/Source/Niagara/Public/NiagaraScriptExecutionParameterStore.h b/Engine/Plugins/FX/Niagara/Source/Niagara/Public/NiagaraScriptExecutionParameterStore.h
--- a/Engine/Plugins/FX/Niagara/Source/Niagara/Public/NiagaraScriptExecutionParameterStore.h
+++ b/Engine/Plugins/FX/Niagara/Source/Niagara/Public/NiagaraScriptExecutionParameterStore.h
@@ -185,8 +185,7 @@ public:
 	TArrayView<const uint8> GetScriptLiterals() const;
 #endif
 
-protected:
-	virtual void InternalStorageChanged() override;
+	virtual TArrayView<const FNiagaraVariableWithOffset> ReadParameterVariables() const override;
 
 private:
 	FNiagaraCompiledDataReference<FNiagaraScriptExecutionParameterStore> ScriptParameterStore;
diff --git a/Engine/Plugins/FX/Niagara/Source/NiagaraEditor/Private/Customizations/NiagaraTypeCustomizations.cpp b/Engine/Plugins/FX/Niagara/Source/NiagaraEditor/Private/Customizations/NiagaraTypeCustomizations.cpp
--- a/Engine/Plugins/FX/Niagara/Source/NiagaraEditor/Private/Customizations/NiagaraTypeCustomizations.cpp
+++ b/Engine/Plugins/FX/Niagara/Source/NiagaraEditor/Private/Customizations/NiagaraTypeCustomizations.cpp
@@ -351,7 +351,7 @@ TArray<FName> FNiagaraUserParameterBindingCustomization::GetNames() const
 
 	if (BaseSystem && TargetUserParameterBinding)
 	{
-		for (const FNiagaraVariable& Var : BaseSystem->GetExposedParameters().ParameterVariables)
+		for (const FNiagaraVariable& Var : BaseSystem->GetExposedParameters().ReadParameterVariables())
 		{
 			if (FNiagaraParameterMapHistory::IsUserParameter(Var) && Var.GetType() == TargetUserParameterBinding->Parameter.GetType())
 			{
diff --git a/Engine/Plugins/FX/Niagara/Source/NiagaraEditor/Private/ViewModels/Stack/NiagaraStackSystemSettingsGroup.cpp b/Engine/Plugins/FX/Niagara/Source/NiagaraEditor/Private/ViewModels/Stack/NiagaraStackSystemSettingsGroup.cpp
--- a/Engine/Plugins/FX/Niagara/Source/NiagaraEditor/Private/ViewModels/Stack/NiagaraStackSystemSettingsGroup.cpp
+++ b/Engine/Plugins/FX/Niagara/Source/NiagaraEditor/Private/ViewModels/Stack/NiagaraStackSystemSettingsGroup.cpp
@@ -171,7 +171,7 @@ FText UNiagaraStackParameterStoreItem::GetTooltipText() const
 
 void UNiagaraStackParameterStoreItem::RefreshChildrenInternal(const TArray<UNiagaraStackEntry*>& CurrentChildren, TArray<UNiagaraStackEntry*>& NewChildren, TArray<FStackIssue>& NewIssues)
 {
-	if (ParameterStore != nullptr && ParameterStore->ParameterVariables.Num() > 0)
+	if (ParameterStore != nullptr && ParameterStore->ReadParameterVariables().Num() > 0)
 	{
 		TArray<FNiagaraVariable> Variables;
 		ParameterStore->GetParameters(Variables);
diff --git a/Engine/Plugins/FX/Niagara/Source/NiagaraEditor/Private/Widgets/SNiagaraSpreadsheetView.cpp b/Engine/Plugins/FX/Niagara/Source/NiagaraEditor/Private/Widgets/SNiagaraSpreadsheetView.cpp
--- a/Engine/Plugins/FX/Niagara/Source/NiagaraEditor/Private/Widgets/SNiagaraSpreadsheetView.cpp
+++ b/Engine/Plugins/FX/Niagara/Source/NiagaraEditor/Private/Widgets/SNiagaraSpreadsheetView.cpp
@@ -1142,7 +1142,7 @@ void SNiagaraSpreadsheetView::ResetEntries(EUITab Tab)
 		}
 	
 		{
-			int32 NumInstances = CaptureData[(int32)Tab].InputParams.ParameterVariables.Num();
+			int32 NumInstances = CaptureData[(int32)Tab].InputParams.ReadParameterVariables().Num();
 			if (!CaptureData[(int32)Tab].bInputColumnsAreAttributes && CaptureData[(int32)Tab].SupportedInputFields.IsValid())
 			{
 				NumInstances = CaptureData[(int32)Tab].SupportedInputFields->Num();
