- Fix static analysis warning

#rnx

diff --git a/Engine/Plugins/FX/Niagara/Source/NiagaraEditor/Private/NiagaraHlslTranslator.cpp b/Engine/Plugins/FX/Niagara/Source/NiagaraEditor/Private/NiagaraHlslTranslator.cpp
--- a/Engine/Plugins/FX/Niagara/Source/NiagaraEditor/Private/NiagaraHlslTranslator.cpp
+++ b/Engine/Plugins/FX/Niagara/Source/NiagaraEditor/Private/NiagaraHlslTranslator.cpp
@@ -6662,11 +6662,14 @@ int32 FHlslNiagaraTranslator::CompileOutputPin(const UEdGraphPin* InPin)
 {
 	SCOPE_CYCLE_COUNTER(STAT_NiagaraEditor_HlslTranslator_CompileOutputPin);
 
-	if (InPin)
+	if (InPin == nullptr)
 	{
-		UpdateStaticSwitchConstants(InPin->GetOwningNode());
+		Error(LOCTEXT("InPinNullptr", "Input pin is nullptr!"), nullptr, nullptr);
+		return INDEX_NONE;
 	}
 
+	UpdateStaticSwitchConstants(InPin->GetOwningNode());
+
 	// The incoming pin to compile may be pointing to a reroute node. If so, we just jump over it
 	// to where it really came from.
 	UEdGraphPin* Pin = UNiagaraNode::TraceOutputPin(const_cast<UEdGraphPin*>(InPin));
