- Fix static analysis warning, Mesh is guaranteed to be valid or we early out

#rnx

diff --git a/Engine/Plugins/FX/Niagara/Source/Niagara/Private/NiagaraDataInterfaceSkeletalMesh.cpp b/Engine/Plugins/FX/Niagara/Source/Niagara/Private/NiagaraDataInterfaceSkeletalMesh.cpp
--- a/Engine/Plugins/FX/Niagara/Source/Niagara/Private/NiagaraDataInterfaceSkeletalMesh.cpp
+++ b/Engine/Plugins/FX/Niagara/Source/Niagara/Private/NiagaraDataInterfaceSkeletalMesh.cpp
@@ -1522,19 +1522,11 @@ bool FNDISkeletalMesh_InstanceData::Init(UNiagaraDataInterfaceSkeletalMesh* Inte
 				NewSkelComp->GetSocketInfoByName(SpecificSockets[i], SpecificSocketInfo[i].Transform, SpecificSocketInfo[i].BoneIdx);
 			}
 		}
-		else if (Mesh != nullptr)
-		{
-			for (int32 i = 0; i < SpecificSocketInfo.Num(); ++i)
-			{
-				SpecificSocketInfo[i].Transform = FTransform(Mesh->GetComposedRefPoseMatrix(SpecificSockets[i]));
-				SpecificSocketInfo[i].BoneIdx = INDEX_NONE;
-			}
-		}
 		else
 		{
 			for (int32 i = 0; i < SpecificSocketInfo.Num(); ++i)
 			{
-				SpecificSocketInfo[i].Transform = FTransform::Identity;
+				SpecificSocketInfo[i].Transform = FTransform(Mesh->GetComposedRefPoseMatrix(SpecificSockets[i]));
 				SpecificSocketInfo[i].BoneIdx = INDEX_NONE;
 			}
 		}
