Niagara - Fix crashes when loading broken data.

diff --git a/Engine/Plugins/FX/Niagara/Source/Niagara/Private/NiagaraScript.cpp b/Engine/Plugins/FX/Niagara/Source/Niagara/Private/NiagaraScript.cpp
--- a/Engine/Plugins/FX/Niagara/Source/Niagara/Private/NiagaraScript.cpp
+++ b/Engine/Plugins/FX/Niagara/Source/Niagara/Private/NiagaraScript.cpp
@@ -1551,6 +1551,12 @@ void UNiagaraScript::RequestCompile(bool bForceCompile)
 		INiagaraModule& NiagaraModule = FModuleManager::Get().LoadModuleChecked<INiagaraModule>(TEXT("Niagara"));
 		TSharedPtr<FNiagaraCompileRequestDataBase, ESPMode::ThreadSafe> RequestData = NiagaraModule.Precompile(this);
 
+		if (RequestData.IsValid() == false)
+		{
+			UE_LOG(LogNiagara, Error, TEXT("Failed to precompile %s.  This is due to unexpected invalid or broken data.  Additional details should be in the log."), *GetPathName());
+			return;
+		}
+
 		// check the ddc first
 		if (!GetDerivedDataCacheRef().GetSynchronous(*GetNiagaraDDCKeyString(), OutData, GetPathName()))
 		{
diff --git a/Engine/Plugins/FX/Niagara/Source/Niagara/Private/NiagaraSystem.cpp b/Engine/Plugins/FX/Niagara/Source/Niagara/Private/NiagaraSystem.cpp
--- a/Engine/Plugins/FX/Niagara/Source/Niagara/Private/NiagaraSystem.cpp
+++ b/Engine/Plugins/FX/Niagara/Source/Niagara/Private/NiagaraSystem.cpp
@@ -1124,6 +1124,12 @@ bool UNiagaraSystem::RequestCompile(bool bForce)
 	INiagaraModule& NiagaraModule = FModuleManager::Get().LoadModuleChecked<INiagaraModule>(TEXT("Niagara"));
 	TSharedPtr<FNiagaraCompileRequestDataBase, ESPMode::ThreadSafe> SystemPrecompiledData = NiagaraModule.Precompile(this);
 
+	if (SystemPrecompiledData.IsValid() == false)
+	{
+		UE_LOG(LogNiagara, Error, TEXT("Failed to precompile %s.  This is due to unexpected invalid or broken data.  Additional details should be in the log."), *GetPathName());
+		return false;
+	}
+
 	SystemPrecompiledData->GetReferencedObjects(ActiveCompilations[ActiveCompileIdx].RootObjects);
 
 	//Compile all emitters
diff --git a/Engine/Plugins/FX/Niagara/Source/NiagaraEditor/Private/NiagaraCompiler.cpp b/Engine/Plugins/FX/Niagara/Source/NiagaraEditor/Private/NiagaraCompiler.cpp
--- a/Engine/Plugins/FX/Niagara/Source/NiagaraEditor/Private/NiagaraCompiler.cpp
+++ b/Engine/Plugins/FX/Niagara/Source/NiagaraEditor/Private/NiagaraCompiler.cpp
@@ -702,7 +702,13 @@ TSharedPtr<FNiagaraCompileRequestDataBase, ESPMode::ThreadSafe> FNiagaraEditorMo
 						return false;
 					});
 
-					checkf(SameUsageSystemHistory != nullptr, TEXT("Failed to find System script with same run index as Emitter script!"));
+					if (SameUsageSystemHistory == nullptr)
+					{
+						UE_LOG(LogNiagaraEditor, Error, TEXT("Precompile failed for system: %s.  Failed to generate valid parameter history for emitter: %s."),
+							*System->GetPathName(), *EmitterRequest->EmitterUniqueName);
+						return TSharedPtr<FNiagaraCompileRequestDataBase, ESPMode::ThreadSafe>();
+					}
+
 					EmitterAndParticleScriptHistoryHandles.Add(FNiagaraParameterMapHistoryHandle(*SameUsageSystemHistory, EmitterScriptHistories[i], EmitterRequest->RequiredRendererVariables, EmitterUniqueName));
 				}
 			}
diff --git a/Engine/Plugins/FX/Niagara/Source/NiagaraEditor/Private/ViewModels/Stack/NiagaraStackGraphUtilities.cpp b/Engine/Plugins/FX/Niagara/Source/NiagaraEditor/Private/ViewModels/Stack/NiagaraStackGraphUtilities.cpp
--- a/Engine/Plugins/FX/Niagara/Source/NiagaraEditor/Private/ViewModels/Stack/NiagaraStackGraphUtilities.cpp
+++ b/Engine/Plugins/FX/Niagara/Source/NiagaraEditor/Private/ViewModels/Stack/NiagaraStackGraphUtilities.cpp
@@ -2357,13 +2357,16 @@ void FNiagaraStackGraphUtilities::RebuildEmitterNodes(UNiagaraSystem& System)
 			TArray<FNiagaraStackGraphUtilities::FStackNodeGroup> StackNodeGroups;
 			FNiagaraStackGraphUtilities::GetStackNodeGroups(*OutputNodes[i], StackNodeGroups);
 
-			FNiagaraStackGraphUtilities::FStackNodeGroup EmitterGroup;
-			EmitterGroup.StartNodes.Add(EmitterNode);
-			EmitterGroup.EndNode = EmitterNode;
+			if (StackNodeGroups.Num() >= 2)
+			{
+				FNiagaraStackGraphUtilities::FStackNodeGroup EmitterGroup;
+				EmitterGroup.StartNodes.Add(EmitterNode);
+				EmitterGroup.EndNode = EmitterNode;
 
-			FNiagaraStackGraphUtilities::FStackNodeGroup& OutputGroup = StackNodeGroups[StackNodeGroups.Num() - 1];
-			FNiagaraStackGraphUtilities::FStackNodeGroup& OutputGroupPrevious = StackNodeGroups[StackNodeGroups.Num() - 2];
-			FNiagaraStackGraphUtilities::ConnectStackNodeGroup(EmitterGroup, OutputGroupPrevious, OutputGroup);
+				FNiagaraStackGraphUtilities::FStackNodeGroup& OutputGroup = StackNodeGroups[StackNodeGroups.Num() - 1];
+				FNiagaraStackGraphUtilities::FStackNodeGroup& OutputGroupPrevious = StackNodeGroups[StackNodeGroups.Num() - 2];
+				FNiagaraStackGraphUtilities::ConnectStackNodeGroup(EmitterGroup, OutputGroupPrevious, OutputGroup);
+			}
 		}
 	}
 
diff --git a/Engine/Plugins/FX/Niagara/Source/NiagaraEditorWidgets/Private/SNiagaraOverviewStackNode.cpp b/Engine/Plugins/FX/Niagara/Source/NiagaraEditorWidgets/Private/SNiagaraOverviewStackNode.cpp
--- a/Engine/Plugins/FX/Niagara/Source/NiagaraEditorWidgets/Private/SNiagaraOverviewStackNode.cpp
+++ b/Engine/Plugins/FX/Niagara/Source/NiagaraEditorWidgets/Private/SNiagaraOverviewStackNode.cpp
@@ -112,6 +112,10 @@ TSharedRef<SWidget> SNiagaraOverviewStackNode::CreateTitleWidget(TSharedPtr<SNod
 
 TSharedRef<SWidget> SNiagaraOverviewStackNode::CreateTitleRightWidget()
 {
+	if (StackViewModel == nullptr)
+	{
+		return SNullWidget::NullWidget;
+	}
 	return SNew(SHorizontalBox)
 		+ SHorizontalBox::Slot()
 		.AutoWidth()
