Coalesce parameter store bindings to merge entries that have no gaps in between them; saves memory and memcpys in the GPU emitter case

diff --git a/Engine/Plugins/FX/Niagara/Source/Niagara/Private/NiagaraScript.cpp b/Engine/Plugins/FX/Niagara/Source/Niagara/Private/NiagaraScript.cpp
--- a/Engine/Plugins/FX/Niagara/Source/Niagara/Private/NiagaraScript.cpp
+++ b/Engine/Plugins/FX/Niagara/Source/Niagara/Private/NiagaraScript.cpp
@@ -855,6 +855,8 @@ void UNiagaraScript::Serialize(FArchive& Ar)
 		{
 			ExecutableData.ScriptLiterals.Empty();
 		}
+
+		ScriptExecutionParamStore.CoalescePaddingInfo();
 	}
 #endif
 
diff --git a/Engine/Plugins/FX/Niagara/Source/Niagara/Private/NiagaraScriptExecutionParameterStore.cpp b/Engine/Plugins/FX/Niagara/Source/Niagara/Private/NiagaraScriptExecutionParameterStore.cpp
--- a/Engine/Plugins/FX/Niagara/Source/Niagara/Private/NiagaraScriptExecutionParameterStore.cpp
+++ b/Engine/Plugins/FX/Niagara/Source/Niagara/Private/NiagaraScriptExecutionParameterStore.cpp
@@ -152,6 +152,30 @@ uint32 FNiagaraScriptExecutionParameterStore::GenerateLayoutInfoInternal(TArray<
 	return InSrcOffset;
 }
 
+void FNiagaraScriptExecutionParameterStore::CoalescePaddingInfo()
+{
+	int32 PaddingIt = 1;
+
+	while (PaddingIt < PaddingInfo.Num())
+	{
+		FNiagaraScriptExecutionPaddingInfo& PreviousEntry = PaddingInfo[PaddingIt - 1];
+		const FNiagaraScriptExecutionPaddingInfo& CurrentEntry = PaddingInfo[PaddingIt];
+
+		if (((PreviousEntry.SrcOffset + PreviousEntry.SrcSize) == CurrentEntry.SrcOffset)
+			&& ((PreviousEntry.DestOffset + PreviousEntry.DestSize) == CurrentEntry.DestOffset)
+			&& ((TNumericLimits<uint16>::Max() - PreviousEntry.SrcSize) >= CurrentEntry.SrcSize))
+		{
+			PreviousEntry.SrcSize += CurrentEntry.SrcSize;
+			PreviousEntry.DestSize += CurrentEntry.DestSize;
+			PaddingInfo.RemoveAt(PaddingIt);
+		}
+		else
+		{
+			++PaddingIt;
+		}
+	}
+}
+
 void FNiagaraScriptExecutionParameterStore::AddPaddedParamSize(const FNiagaraTypeDefinition& InParamType, uint32 InOffset)
 {
 	if (InParamType.IsDataInterface())
diff --git a/Engine/Plugins/FX/Niagara/Source/Niagara/Public/NiagaraScriptExecutionParameterStore.h b/Engine/Plugins/FX/Niagara/Source/Niagara/Public/NiagaraScriptExecutionParameterStore.h
--- a/Engine/Plugins/FX/Niagara/Source/Niagara/Public/NiagaraScriptExecutionParameterStore.h
+++ b/Engine/Plugins/FX/Niagara/Source/Niagara/Public/NiagaraScriptExecutionParameterStore.h
@@ -46,6 +46,7 @@ public:
 	/** Call this init function if you are using a Niagara parameter store within a UNiagaraScript.*/
 	void InitFromOwningScript(UNiagaraScript* Script, ENiagaraSimTarget SimTarget, bool bNotifyAsDirty);
 	void AddScriptParams(UNiagaraScript* Script, ENiagaraSimTarget SimTarget, bool bTriggerRebind);
+	void CoalescePaddingInfo();
 #endif
 
 	virtual bool AddParameter(const FNiagaraVariable& Param, bool bInitInterfaces = true, bool bTriggerRebind = true, int32* OutOffset = nullptr) override
