Prevent GPUComputeScripts from being created unless we're in the editor

diff --git a/Engine/Plugins/FX/Niagara/Source/Niagara/Private/NiagaraEmitter.cpp b/Engine/Plugins/FX/Niagara/Source/Niagara/Private/NiagaraEmitter.cpp
--- a/Engine/Plugins/FX/Niagara/Source/Niagara/Private/NiagaraEmitter.cpp
+++ b/Engine/Plugins/FX/Niagara/Source/Niagara/Private/NiagaraEmitter.cpp
@@ -372,21 +372,17 @@ void UNiagaraEmitter::PostLoad()
 		}
 	}
 
+#if WITH_EDITORONLY_DATA
 	if (!GPUComputeScript)
 	{
 		GPUComputeScript = NewObject<UNiagaraScript>(this, "GPUComputeScript", EObjectFlags::RF_Transactional);
 		GPUComputeScript->SetUsage(ENiagaraScriptUsage::ParticleGPUComputeScript);
-#if WITH_EDITORONLY_DATA
 		GPUComputeScript->SetSource(SpawnScriptProps.Script ? SpawnScriptProps.Script->GetSource() : nullptr);
-#endif
 	}
 
-
-#if WITH_EDITORONLY_DATA
-	if (GPUComputeScript)
-	{
-		GPUComputeScript->OnGPUScriptCompiled().AddUObject(this, &UNiagaraEmitter::RaiseOnEmitterGPUCompiled);
-	}
+	GPUComputeScript->OnGPUScriptCompiled().AddUObject(this, &UNiagaraEmitter::RaiseOnEmitterGPUCompiled);
+#else
+	check(GPUComputeScript == nullptr || SimTarget == ENiagaraSimTarget::GPUComputeSim);
 #endif
 
 	if (EmitterSpawnScriptProps.Script == nullptr || EmitterUpdateScriptProps.Script == nullptr)
