Clean up event creation to use NAME_None instead of emptry string literal

diff --git a/Engine/Plugins/FX/Niagara/Source/Niagara/Private/NiagaraEmitter.cpp b/Engine/Plugins/FX/Niagara/Source/Niagara/Private/NiagaraEmitter.cpp
--- a/Engine/Plugins/FX/Niagara/Source/Niagara/Private/NiagaraEmitter.cpp
+++ b/Engine/Plugins/FX/Niagara/Source/Niagara/Private/NiagaraEmitter.cpp
@@ -76,12 +76,12 @@ void FNiagaraEmitterScriptProperties::InitDataSetAccess()
 		//
 		for (FNiagaraDataSetID &ReadID : Script->GetVMExecutableData().ReadDataSets)
 		{
-			EventReceivers.Add( FNiagaraEventReceiverProperties(ReadID.Name, "", "") );
+			EventReceivers.Add( FNiagaraEventReceiverProperties(ReadID.Name, NAME_None, NAME_None) );
 		}
 
 		for (FNiagaraDataSetProperties &WriteID : Script->GetVMExecutableData().WriteDataSets)
 		{
-			FNiagaraEventGeneratorProperties Props(WriteID, "");
+			FNiagaraEventGeneratorProperties Props(WriteID, NAME_None);
 			EventGenerators.Add(Props);
 		}
 	}
