Fix changing submix/scope parameters.

diff --git a/Engine/Plugins/FX/Niagara/Source/Niagara/Classes/NiagaraDataInterfaceAudioOscilloscope.h b/Engine/Plugins/FX/Niagara/Source/Niagara/Classes/NiagaraDataInterfaceAudioOscilloscope.h
--- a/Engine/Plugins/FX/Niagara/Source/Niagara/Classes/NiagaraDataInterfaceAudioOscilloscope.h
+++ b/Engine/Plugins/FX/Niagara/Source/Niagara/Classes/NiagaraDataInterfaceAudioOscilloscope.h
@@ -126,6 +126,8 @@ public:
 	virtual bool GetFunctionHLSL(const FNiagaraDataInterfaceGPUParamInfo& ParamInfo, const FNiagaraDataInterfaceGeneratedFunction& FunctionInfo, int FunctionInstanceIndex, FString& OutHLSL) override;
 	virtual void GetParameterDefinitionHLSL(const FNiagaraDataInterfaceGPUParamInfo& ParamInfo, FString& OutHLSL) override;
 
+	virtual bool Equals(const UNiagaraDataInterface* Other) const override;
+
 #if WITH_EDITOR
 	virtual void PostEditChangeProperty(struct FPropertyChangedEvent& PropertyChangedEvent) override;
 #endif // WITH_EDITOR
@@ -133,6 +135,7 @@ public:
 	virtual void PostInitProperties() override;
 	virtual void BeginDestroy() override;
 	virtual void PostLoad() override;
+	
 
 protected:
 	virtual bool CopyToInternal(UNiagaraDataInterface* Destination) const override;
diff --git a/Engine/Plugins/FX/Niagara/Source/Niagara/Private/NiagaraDataInterfaceAudioOscilloscope.cpp b/Engine/Plugins/FX/Niagara/Source/Niagara/Private/NiagaraDataInterfaceAudioOscilloscope.cpp
--- a/Engine/Plugins/FX/Niagara/Source/Niagara/Private/NiagaraDataInterfaceAudioOscilloscope.cpp
+++ b/Engine/Plugins/FX/Niagara/Source/Niagara/Private/NiagaraDataInterfaceAudioOscilloscope.cpp
@@ -55,7 +55,8 @@ void FNiagaraDataInterfaceProxyOscilloscope::RegisterToAllAudioDevices()
 		DeviceManager->IterateOverAllDevices([&](Audio::FDeviceId DeviceId, FAudioDevice* InDevice)
 		{
 			check(!SubmixListeners.Contains(DeviceId));
-			SubmixListeners.Emplace(DeviceId, new FNiagaraSubmixListener(PatchMixer, Resolution * 2, DeviceId, SubmixRegisteredTo));
+			const int32 NumSamplesToPop = Align(FMath::FloorToInt(ScopeInMilliseconds / 1000.0f * InDevice->GetSampleRate()) * AUDIO_MIXER_MAX_OUTPUT_CHANNELS, 4);
+			SubmixListeners.Emplace(DeviceId, new FNiagaraSubmixListener(PatchMixer, NumSamplesToPop, DeviceId, SubmixRegisteredTo));
 			SubmixListeners[DeviceId]->RegisterToSubmix();
 		});
 	}
@@ -419,6 +420,15 @@ void UNiagaraDataInterfaceAudioOscilloscope::PostLoad()
 	GetProxyAs<FNiagaraDataInterfaceProxyOscilloscope>()->OnUpdateSubmix(Submix);
 }
 
+bool UNiagaraDataInterfaceAudioOscilloscope::Equals(const UNiagaraDataInterface* Other) const
+{
+	const UNiagaraDataInterfaceAudioOscilloscope* CastedOther = Cast<const UNiagaraDataInterfaceAudioOscilloscope>(Other);
+	return Super::Equals(Other)
+		&& (CastedOther->Submix == Submix)
+		&& (CastedOther->Resolution == Resolution)
+		&& FMath::IsNearlyEqual(CastedOther->ScopeInMilliseconds, ScopeInMilliseconds);
+}
+
 bool UNiagaraDataInterfaceAudioOscilloscope::CopyToInternal(UNiagaraDataInterface* Destination) const
 {
 	Super::CopyToInternal(Destination);
@@ -430,6 +440,9 @@ bool UNiagaraDataInterfaceAudioOscilloscope::CopyToInternal(UNiagaraDataInterfac
 		CastedDestination->Submix = Submix;
 		CastedDestination->Resolution = Resolution;
 		CastedDestination->ScopeInMilliseconds = ScopeInMilliseconds;
+
+		CastedDestination->GetProxyAs<FNiagaraDataInterfaceProxyOscilloscope>()->OnUpdateResampling(Resolution, ScopeInMilliseconds);
+		CastedDestination->GetProxyAs<FNiagaraDataInterfaceProxyOscilloscope>()->OnUpdateSubmix(Submix);
 	}
 	
 	return true;
@@ -486,19 +499,20 @@ void FNiagaraDataInterfaceProxyOscilloscope::DownsampleAudioToBuffer()
 	check(NumChannelsInDownsampledBuffer > 0);
 
 	// Get the number of frames of audio in the original sample rate to show.
-	int32 NumFramesToPop = FMath::FloorToInt(SourceSampleRate * (ScopeInMilliseconds / 1000.0f));
+	int32 NumFramesToPop = Align(FMath::FloorToInt(SourceSampleRate * (ScopeInMilliseconds / 1000.0f)), 4);
 
 	// If we have enough frames buffered to pop, update the PopBuffer. Otherwise, use the previous frames buffer.
 	int32 NumSamplesToPop = NumFramesToPop * NumChannelsInDownsampledBuffer;
-	if (PatchMixer.MaxNumberOfSamplesThatCanBePopped() >= NumFramesToPop)
+
+	if (PopBuffer.Num() != NumSamplesToPop)
 	{
 		PopBuffer.Reset();
 		PopBuffer.AddZeroed(NumSamplesToPop);
-		PatchMixer.PopAudio(PopBuffer.GetData(), NumSamplesToPop, true);
 	}
-	else
+
+	if (PatchMixer.MaxNumberOfSamplesThatCanBePopped() >= NumSamplesToPop)
 	{
-		PopBuffer.SetNumZeroed(NumSamplesToPop);
+		PatchMixer.PopAudio(PopBuffer.GetData(), NumSamplesToPop, true);
 	}
 
 	// Downsample buffer in place.
