Niagara - Add a parameters panel to the scratchpad to make it usable again after the namespace refactor.

diff --git a/Engine/Plugins/FX/Niagara/Source/NiagaraEditor/Private/NiagaraActions.cpp b/Engine/Plugins/FX/Niagara/Source/NiagaraEditor/Private/NiagaraActions.cpp
--- a/Engine/Plugins/FX/Niagara/Source/NiagaraEditor/Private/NiagaraActions.cpp
+++ b/Engine/Plugins/FX/Niagara/Source/NiagaraEditor/Private/NiagaraActions.cpp
@@ -10,7 +10,7 @@
 #include "Framework/Application/SlateApplication.h"
 #include "Layout/WidgetPath.h"
 #include "ScopedTransaction.h"
-#include "NiagaraParameterPanelViewModel.h"
+#include "ViewModels/NiagaraParameterPanelViewModel.h"
 
 #define LOCTEXT_NAMESPACE "NiagaraActions"
 
diff --git a/Engine/Plugins/FX/Niagara/Source/NiagaraEditor/Private/Toolkits/NiagaraScriptToolkit.cpp b/Engine/Plugins/FX/Niagara/Source/NiagaraEditor/Private/Toolkits/NiagaraScriptToolkit.cpp
--- a/Engine/Plugins/FX/Niagara/Source/NiagaraEditor/Private/Toolkits/NiagaraScriptToolkit.cpp
+++ b/Engine/Plugins/FX/Niagara/Source/NiagaraEditor/Private/Toolkits/NiagaraScriptToolkit.cpp
@@ -15,7 +15,7 @@
 #include "SNiagaraParameterCollection.h"
 #include "UObject/Package.h"
 #include "SNiagaraParameterMapView.h"
-#include "SNiagaraParameterPanel.h"
+#include "Widgets/SNiagaraParameterPanel.h"
 #include "NiagaraEditorStyle.h"
 #include "NiagaraSystem.h"
 #include "BusyCursor.h"
@@ -48,7 +48,7 @@
 #include "NiagaraMessageManager.h"
 #include "NiagaraStandaloneScriptViewModel.h"
 #include "NiagaraMessageLogViewModel.h"
-#include "NiagaraParameterPanelViewModel.h"
+#include "ViewModels/NiagaraParameterPanelViewModel.h"
 
 #define LOCTEXT_NAMESPACE "NiagaraScriptToolkit"
 
diff --git a/Engine/Plugins/FX/Niagara/Source/NiagaraEditor/Private/Toolkits/NiagaraSystemToolkit.cpp b/Engine/Plugins/FX/Niagara/Source/NiagaraEditor/Private/Toolkits/NiagaraSystemToolkit.cpp
--- a/Engine/Plugins/FX/Niagara/Source/NiagaraEditor/Private/Toolkits/NiagaraSystemToolkit.cpp
+++ b/Engine/Plugins/FX/Niagara/Source/NiagaraEditor/Private/Toolkits/NiagaraSystemToolkit.cpp
@@ -56,7 +56,7 @@
 #include "ViewModels/NiagaraOverviewGraphViewModel.h"
 #include "NiagaraScriptSourceBase.h"
 #include "Widgets/Input/SCheckBox.h"
-#include "NiagaraParameterPanelViewModel.h"
+#include "ViewModels/NiagaraParameterPanelViewModel.h"
 
 #define LOCTEXT_NAMESPACE "NiagaraSystemEditor"
 
diff --git a/Engine/Plugins/FX/Niagara/Source/NiagaraEditor/Private/ViewModels/NiagaraParameterNameViewModel.cpp b/Engine/Plugins/FX/Niagara/Source/NiagaraEditor/Private/ViewModels/NiagaraParameterNameViewModel.cpp
--- a/Engine/Plugins/FX/Niagara/Source/NiagaraEditor/Private/ViewModels/NiagaraParameterNameViewModel.cpp
+++ b/Engine/Plugins/FX/Niagara/Source/NiagaraEditor/Private/ViewModels/NiagaraParameterNameViewModel.cpp
@@ -7,7 +7,7 @@
 #include "NiagaraActions.h"
 #include "ViewModels/Stack/NiagaraStackGraphUtilities.h"
 #include "NiagaraNode.h"
-#include "NiagaraParameterPanelViewModel.h"
+#include "ViewModels/NiagaraParameterPanelViewModel.h"
 #include "NiagaraEditorUtilities.h"
 #include "NiagaraEditorModule.h"
 #include "CoreMinimal.h"
diff --git a/Engine/Plugins/FX/Niagara/Source/NiagaraEditor/Private/ViewModels/NiagaraParameterPanelViewModel.cpp b/Engine/Plugins/FX/Niagara/Source/NiagaraEditor/Private/ViewModels/NiagaraParameterPanelViewModel.cpp
--- a/Engine/Plugins/FX/Niagara/Source/NiagaraEditor/Private/ViewModels/NiagaraParameterPanelViewModel.cpp
+++ b/Engine/Plugins/FX/Niagara/Source/NiagaraEditor/Private/ViewModels/NiagaraParameterPanelViewModel.cpp
@@ -1,6 +1,6 @@
 // Copyright Epic Games, Inc. All Rights Reserved.
 
-#include "NiagaraParameterPanelViewModel.h"
+#include "ViewModels/NiagaraParameterPanelViewModel.h"
 #include "NiagaraScriptVariable.h"
 #include "NiagaraGraph.h"
 #include "NiagaraActions.h"
@@ -579,7 +579,7 @@ TArray<TWeakObjectPtr<UNiagaraGraph>> FNiagaraSystemToolkitParameterPanelViewMod
 /// Script Toolkit Parameter Panel View Model								///
 ///////////////////////////////////////////////////////////////////////////////
 
-FNiagaraScriptToolkitParameterPanelViewModel::FNiagaraScriptToolkitParameterPanelViewModel(TSharedPtr<FNiagaraStandaloneScriptViewModel> InScriptViewModel)
+FNiagaraScriptToolkitParameterPanelViewModel::FNiagaraScriptToolkitParameterPanelViewModel(TSharedPtr<FNiagaraScriptViewModel> InScriptViewModel)
 {
 	ScriptViewModel = InScriptViewModel;
 	VariableObjectSelection = ScriptViewModel->GetVariableSelection();
diff --git a/Engine/Plugins/FX/Niagara/Source/NiagaraEditor/Private/ViewModels/NiagaraParameterPanelViewModel.h b/Engine/Plugins/FX/Niagara/Source/NiagaraEditor/Private/ViewModels/NiagaraParameterPanelViewModel.h
deleted file mode 100644
diff --git a/Engine/Plugins/FX/Niagara/Source/NiagaraEditor/Private/ViewModels/NiagaraScratchPadScriptViewModel.cpp b/Engine/Plugins/FX/Niagara/Source/NiagaraEditor/Private/ViewModels/NiagaraScratchPadScriptViewModel.cpp
--- a/Engine/Plugins/FX/Niagara/Source/NiagaraEditor/Private/ViewModels/NiagaraScratchPadScriptViewModel.cpp
+++ b/Engine/Plugins/FX/Niagara/Source/NiagaraEditor/Private/ViewModels/NiagaraScratchPadScriptViewModel.cpp
@@ -3,6 +3,7 @@
 #include "ViewModels/NiagaraScratchPadScriptViewModel.h"
 #include "ViewModels/Stack/NiagaraStackGraphUtilities.h"
 #include "ViewModels/Stack/NiagaraStackFunctionInputCollection.h"
+#include "ViewModels/NiagaraParameterPanelViewModel.h"
 #include "NiagaraNodeFunctionCall.h"
 #include "NiagaraSystem.h"
 #include "NiagaraEmitter.h"
@@ -10,6 +11,7 @@
 
 #include "ObjectTools.h"
 #include "ScopedTransaction.h"
+#include "Framework/Commands/UICommandList.h"
 
 #define LOCTEXT_NAMESPACE "NiagaraScratchPadScriptViewModel"
 
@@ -44,6 +46,18 @@ void FNiagaraScratchPadScriptViewModel::Initialize(UNiagaraScript* Script)
 	UNiagaraScriptSource* EditScriptSource = CastChecked<UNiagaraScriptSource>(EditScript->GetSource());
 	OnGraphNeedsRecompileHandle = EditScriptSource->NodeGraph->AddOnGraphNeedsRecompileHandler(FOnGraphChanged::FDelegate::CreateSP(this, &FNiagaraScratchPadScriptViewModel::OnScriptGraphChanged));
 	EditScript->OnPropertyChanged().AddSP(this, &FNiagaraScratchPadScriptViewModel::OnScriptPropertyChanged);
+	ParameterPanelCommands = MakeShared<FUICommandList>();
+	ParameterPaneViewModel = MakeShared<FNiagaraScriptToolkitParameterPanelViewModel>(this->AsShared());
+	ParameterPaneViewModel->InitBindings();
+}
+
+void FNiagaraScratchPadScriptViewModel::Finalize()
+{
+	// This pointer needs to be reset manually here because there is a shared ref cycle.
+	if (ParameterPaneViewModel.IsValid())
+	{
+		ParameterPaneViewModel.Reset();
+	}
 }
 
 void FNiagaraScratchPadScriptViewModel::AddReferencedObjects(FReferenceCollector& Collector)
@@ -61,6 +75,16 @@ UNiagaraScript* FNiagaraScratchPadScriptViewModel::GetEditScript() const
 	return EditScript;
 }
 
+TSharedPtr<INiagaraParameterPanelViewModel> FNiagaraScratchPadScriptViewModel::GetParameterPanelViewModel() const
+{
+	return ParameterPaneViewModel;
+}
+
+TSharedPtr<FUICommandList> FNiagaraScratchPadScriptViewModel::GetParameterPanelCommands() const
+{
+	return ParameterPanelCommands;
+}
+
 FText FNiagaraScratchPadScriptViewModel::GetToolTip() const
 {
 	return FText::Format(LOCTEXT("ScratchPadScriptToolTipFormat", "Description: {0}{1}"),
diff --git a/Engine/Plugins/FX/Niagara/Source/NiagaraEditor/Private/ViewModels/NiagaraScratchPadViewModel.cpp b/Engine/Plugins/FX/Niagara/Source/NiagaraEditor/Private/ViewModels/NiagaraScratchPadViewModel.cpp
--- a/Engine/Plugins/FX/Niagara/Source/NiagaraEditor/Private/ViewModels/NiagaraScratchPadViewModel.cpp
+++ b/Engine/Plugins/FX/Niagara/Source/NiagaraEditor/Private/ViewModels/NiagaraScratchPadViewModel.cpp
@@ -5,6 +5,7 @@
 #include "ViewModels/NiagaraSystemSelectionViewModel.h"
 #include "ViewModels/NiagaraScratchPadScriptViewModel.h"
 #include "ViewModels/NiagaraScriptGraphViewModel.h"
+#include "ViewModels/NiagaraParameterPanelViewModel.h"
 #include "ViewModels/Stack/NiagaraStackModuleItem.h"
 #include "ViewModels/Stack/NiagaraStackFunctionInput.h"
 #include "NiagaraSystem.h"
@@ -42,6 +43,14 @@ void UNiagaraScratchPadViewModel::Initialize(TSharedRef<FNiagaraSystemViewModel>
 
 void UNiagaraScratchPadViewModel::Finalize()
 {
+	for (TSharedRef<FNiagaraScratchPadScriptViewModel> ScriptViewModel : ScriptViewModels)
+	{
+		TearDownScriptViewModel(ScriptViewModel);
+	}
+	ScriptViewModels.Empty();
+	PinnedScriptViewModels.Empty();
+	EditScriptViewModels.Empty();
+	ActiveScriptViewModel.Reset();
 }
 
 void GetOuterAndTargetScripts(TSharedRef<FNiagaraSystemViewModel> SystemViewModel, UObject*& OutOuter, TArray<UNiagaraScript*>*& OutTargetScripts)
@@ -104,8 +113,7 @@ void UNiagaraScratchPadViewModel::RefreshScriptViewModels()
 	{
 		for (TSharedRef<FNiagaraScratchPadScriptViewModel> OldScriptViewModel : OldScriptViewModels)
 		{
-			OldScriptViewModel->OnRenamed().RemoveAll(this);
-			OldScriptViewModel->OnPinnedChanged().RemoveAll(this);
+			TearDownScriptViewModel(OldScriptViewModel);
 		}
 		bViewModelsChanged = true;
 	}
@@ -298,7 +306,7 @@ TSharedPtr<FNiagaraScratchPadScriptViewModel> UNiagaraScratchPadViewModel::Creat
 		UNiagaraScript* DefaultDynamicInput = Cast<UNiagaraScript>(GetDefault<UNiagaraEditorSettings>()->DefaultDynamicInputScript.TryLoad());
 		if (DefaultDynamicInput != nullptr)
 		{
-			NewScript = CastChecked<UNiagaraScript>(StaticDuplicateObject(DefaultDynamicInput, ScriptOuter, GetUniqueScriptName(ScriptOuter, TEXT("NewScratchDynamicInput"))));
+			NewScript = CastChecked<UNiagaraScript>(StaticDuplicateObject(DefaultDynamicInput, ScriptOuter, GetUniqueScriptName(ScriptOuter, TEXT("ScratchDynamicInput"))));
 			TArray<UNiagaraNodeOutput*> OutputNodes;
 			CastChecked<UNiagaraScriptSource>(NewScript->GetSource())->NodeGraph->GetNodesOfClass(OutputNodes);
 			if (OutputNodes.Num() == 1)
@@ -325,7 +333,7 @@ TSharedPtr<FNiagaraScratchPadScriptViewModel> UNiagaraScratchPadViewModel::Creat
 		UNiagaraScript* DefaultModule = Cast<UNiagaraScript>(GetDefault<UNiagaraEditorSettings>()->DefaultModuleScript.TryLoad());
 		if (DefaultModule != nullptr)
 		{
-			NewScript = CastChecked<UNiagaraScript>(StaticDuplicateObject(DefaultModule, ScriptOuter, GetUniqueScriptName(ScriptOuter, TEXT("NewScratchModule"))));
+			NewScript = CastChecked<UNiagaraScript>(StaticDuplicateObject(DefaultModule, ScriptOuter, GetUniqueScriptName(ScriptOuter, TEXT("ScratchModule"))));
 		}
 		break;
 	}
@@ -499,9 +507,20 @@ TSharedRef<FNiagaraScratchPadScriptViewModel> UNiagaraScratchPadViewModel::Creat
 	ScriptViewModel->OnRenamed().AddUObject(this, &UNiagaraScratchPadViewModel::ScriptViewModelScriptRenamed);
 	ScriptViewModel->OnPinnedChanged().AddUObject(this, &UNiagaraScratchPadViewModel::ScriptViewModelPinnedChanged, TWeakPtr<FNiagaraScratchPadScriptViewModel>(ScriptViewModel));
 	ScriptViewModel->OnRequestDiscardChanges().BindUObject(this, &UNiagaraScratchPadViewModel::ScriptViewModelRequestDiscardChanges, TWeakPtr<FNiagaraScratchPadScriptViewModel>(ScriptViewModel));
+	ScriptViewModel->GetVariableSelection()->OnSelectedObjectsChanged().AddUObject(this, &UNiagaraScratchPadViewModel::ScriptViewModelVariableSelectionChanged, TWeakPtr<FNiagaraScratchPadScriptViewModel>(ScriptViewModel));
 	return ScriptViewModel;
 }
 
+void UNiagaraScratchPadViewModel::TearDownScriptViewModel(TSharedRef<FNiagaraScratchPadScriptViewModel> InScriptViewModel)
+{
+	InScriptViewModel->GetGraphViewModel()->GetNodeSelection()->OnSelectedObjectsChanged().RemoveAll(this);
+	InScriptViewModel->OnRenamed().RemoveAll(this);
+	InScriptViewModel->OnPinnedChanged().RemoveAll(this);
+	InScriptViewModel->OnRequestDiscardChanges().Unbind();
+	InScriptViewModel->GetVariableSelection()->OnSelectedObjectsChanged().RemoveAll(this);
+	InScriptViewModel->Finalize();
+}
+
 void UNiagaraScratchPadViewModel::RefreshEditScriptViewModels()
 {
 	EditScriptViewModels.Empty();
@@ -585,7 +604,8 @@ void UNiagaraScratchPadViewModel::ScriptViewModelPinnedChanged(TWeakPtr<FNiagara
 
 void UNiagaraScratchPadViewModel::ScriptViewModelRequestDiscardChanges(TWeakPtr<FNiagaraScratchPadScriptViewModel> ScriptViewModelWeak)
 {
-	if (ActiveScriptViewModel.IsValid() && ActiveScriptViewModel->HasUnappliedChanges())
+	TSharedPtr<FNiagaraScratchPadScriptViewModel> ScriptViewModel = ScriptViewModelWeak.Pin();
+	if (ScriptViewModel.IsValid() && ScriptViewModel->HasUnappliedChanges())
 	{
 		FText Title = LOCTEXT("DiscardChangesTitle", "Discard Changes?");
 		EAppReturnType::Type DialogResult = FMessageDialog::Open(EAppMsgType::YesNo, EAppReturnType::No,
@@ -596,13 +616,13 @@ void UNiagaraScratchPadViewModel::ScriptViewModelRequestDiscardChanges(TWeakPtr<
 			return;
 		}
 
-		int32 ViewModelIndex = ScriptViewModels.IndexOfByKey(ActiveScriptViewModel.ToSharedRef());
+		int32 ViewModelIndex = ScriptViewModels.IndexOfByKey(ScriptViewModel.ToSharedRef());
 		if (ensureMsgf(ViewModelIndex != INDEX_NONE, TEXT("Active script view model wasn't in script view model collection!")))
 		{
-			TSharedRef<FNiagaraScratchPadScriptViewModel> NewScriptViewModel = CreateAndSetupScriptviewModel(ActiveScriptViewModel->GetOriginalScript());
+			TSharedRef<FNiagaraScratchPadScriptViewModel> NewScriptViewModel = CreateAndSetupScriptviewModel(ScriptViewModel->GetOriginalScript());
 			ScriptViewModels[ViewModelIndex] = NewScriptViewModel;
 
-			int32 PinnedViewModelIndex = PinnedScriptViewModels.IndexOfByKey(ActiveScriptViewModel.ToSharedRef());
+			int32 PinnedViewModelIndex = PinnedScriptViewModels.IndexOfByKey(ScriptViewModel.ToSharedRef());
 			if (PinnedViewModelIndex != INDEX_NONE)
 			{
 				PinnedScriptViewModels[PinnedViewModelIndex] = NewScriptViewModel;
@@ -614,4 +634,13 @@ void UNiagaraScratchPadViewModel::ScriptViewModelRequestDiscardChanges(TWeakPtr<
 	}
 }
 
+void UNiagaraScratchPadViewModel::ScriptViewModelVariableSelectionChanged(TWeakPtr<FNiagaraScratchPadScriptViewModel> ScriptViewModelWeak)
+{
+	TSharedPtr<FNiagaraScratchPadScriptViewModel> ScriptViewModel = ScriptViewModelWeak.Pin();
+	if (ScriptViewModel.IsValid())
+	{
+		ObjectSelection->SetSelectedObjects(ScriptViewModel->GetVariableSelection()->GetSelectedObjects().Array());
+	}
+}
+
 #undef LOCTEXT_NAMESPACE
\ No newline at end of file
diff --git a/Engine/Plugins/FX/Niagara/Source/NiagaraEditor/Private/Widgets/SNiagaraGraphPinAdd.cpp b/Engine/Plugins/FX/Niagara/Source/NiagaraEditor/Private/Widgets/SNiagaraGraphPinAdd.cpp
--- a/Engine/Plugins/FX/Niagara/Source/NiagaraEditor/Private/Widgets/SNiagaraGraphPinAdd.cpp
+++ b/Engine/Plugins/FX/Niagara/Source/NiagaraEditor/Private/Widgets/SNiagaraGraphPinAdd.cpp
@@ -12,7 +12,7 @@
 #include "Widgets/Text/SInlineEditableTextBlock.h"
 #include "Widgets/Layout/SBox.h"
 #include "Framework/Application/SlateApplication.h"
-#include "SNiagaraParameterPanel.h"
+#include "Widgets/SNiagaraParameterPanel.h"
 #include "NiagaraNodeParameterMapBase.h"
 #include "NiagaraNodeCustomHlsl.h"
 #include "NiagaraConstants.h"
diff --git a/Engine/Plugins/FX/Niagara/Source/NiagaraEditor/Private/Widgets/SNiagaraParameterPanel.cpp b/Engine/Plugins/FX/Niagara/Source/NiagaraEditor/Private/Widgets/SNiagaraParameterPanel.cpp
--- a/Engine/Plugins/FX/Niagara/Source/NiagaraEditor/Private/Widgets/SNiagaraParameterPanel.cpp
+++ b/Engine/Plugins/FX/Niagara/Source/NiagaraEditor/Private/Widgets/SNiagaraParameterPanel.cpp
@@ -1,8 +1,8 @@
 // Copyright Epic Games, Inc. All Rights Reserved.
 
-#include "SNiagaraParameterPanel.h"
+#include "Widgets/SNiagaraParameterPanel.h"
 #include "NiagaraEditorCommon.h"
-#include "NiagaraParameterPanelViewModel.h"
+#include "ViewModels/NiagaraParameterPanelViewModel.h"
 #include "Framework/Commands/UICommandList.h"
 #include "SGraphActionMenu.h"
 #include "NiagaraActions.h"
diff --git a/Engine/Plugins/FX/Niagara/Source/NiagaraEditor/Private/Widgets/SNiagaraParameterPanel.h b/Engine/Plugins/FX/Niagara/Source/NiagaraEditor/Private/Widgets/SNiagaraParameterPanel.h
deleted file mode 100644
diff --git a/Engine/Plugins/FX/Niagara/Source/NiagaraEditor/Public/ViewModels/NiagaraParameterPanelViewModel.h b/Engine/Plugins/FX/Niagara/Source/NiagaraEditor/Public/ViewModels/NiagaraParameterPanelViewModel.h
new file mode 100644
--- /dev/null
+++ b/Engine/Plugins/FX/Niagara/Source/NiagaraEditor/Public/ViewModels/NiagaraParameterPanelViewModel.h
@@ -0,0 +1,208 @@
+// Copyright Epic Games, Inc. All Rights Reserved.
+
+#pragma once
+
+#include "NiagaraEditorCommon.h"
+#include "Types/SlateEnums.h"
+#include "EditorUndoClient.h"
+#include "EdGraph/EdGraphSchema.h"
+#include "Input/Events.h"
+#include "Input/Reply.h"
+#include "Templates/SharedPointer.h"
+
+class UNiagaraSystem;
+class FNiagaraSystemViewModel;
+class UNiagaraSystemSelectionViewModel;
+class UNiagaraGraph;
+class FNiagaraScriptViewModel;
+class FDelegateHandle;
+struct FCreateWidgetForActionData;
+class FNiagaraObjectSelection;
+class UNiagaraScriptVariable;
+
+/** Interface for view models to the parameter panel. */
+class INiagaraParameterPanelViewModel : public TSharedFromThis<INiagaraParameterPanelViewModel>
+{
+
+public:
+	/** Delegate to signal the view model's state has changed. */
+	DECLARE_DELEGATE(FOnParameterPanelViewModelRefreshed);
+	DECLARE_MULTICAST_DELEGATE_OneParam(FOnParameterPanelViewModelExternalSelectionChanged, const UObject*);
+
+	virtual ~INiagaraParameterPanelViewModel() { }
+
+	/** Separate method to call after ctor to create delegate bindings as we must be fully constructed to do so. */
+	virtual void InitBindings() = 0;
+
+	virtual void Refresh() const = 0;
+
+	virtual void CollectStaticSections(TArray<int32>& StaticSectionIDs) const = 0;
+
+	virtual NiagaraParameterPanelSectionID::Type GetSectionForVarAndViewInfo(const FNiagaraScriptVariableAndViewInfo& VarAndViewInfo) const = 0;
+	
+	virtual const UNiagaraScriptVariable* AddParameter(FNiagaraVariable& InVariableToAdd, const FNiagaraVariableMetaData& InVariableMetaDataToAssign) = 0;
+
+	virtual void DeleteParameter(const FNiagaraVariable& TargetVariableToRemove, const FNiagaraVariableMetaData& TargetVariableMetaData) = 0;
+
+	virtual bool GetCanDeleteParameterAndToolTip(const FNiagaraVariable& TargetVariableToRemove, const FNiagaraVariableMetaData& TargetVariableMetaData, FText& OutCanDeleteParameterToolTip) const = 0;
+
+	virtual void RenameParameter(const FNiagaraVariable& TargetVariableToRename, const FNiagaraVariableMetaData& TargetVariableMetaData, const FText& NewVariableNameText) const = 0;
+
+	virtual void ChangeParameterScope(const FNiagaraVariable& TargetVariableToModify, const FNiagaraVariableMetaData& TargetVariableMetaData, const ENiagaraParameterScope NewVariableScope) const = 0;
+
+	virtual bool CanModifyParameter(const FNiagaraVariable& TargetVariableToModify, const FNiagaraVariableMetaData& TargetVariableMetaData) const = 0;
+
+	virtual bool GetCanRenameParameterAndToolTip(const FNiagaraVariable& TargetVariableToRename, const FNiagaraVariableMetaData& TargetVariableMetaData, TOptional<const FText> NewVariableNameText, FText& OutCanRenameParameterToolTip) const = 0;
+
+	virtual void HandleActionSelected(const TSharedPtr<FEdGraphSchemaAction>& InAction, ESelectInfo::Type InSelectionType) {};
+
+	virtual FReply HandleActionDragged(const TSharedPtr<FEdGraphSchemaAction>& InAction, const FPointerEvent& MouseEvent) const = 0;
+
+	//virtual void HandleParameterDropped() {}; //@todo(ng) impl
+
+	virtual bool CanDropParameter(const FNiagaraVariable& InTargetScriptVariableToDrop) const { return false; };
+
+	virtual bool CanDropParameters(const TArray<FNiagaraVariable>& InTargetScriptVariablesToDrop) const { return false; };
+
+	virtual bool CanDragParameter(const FNiagaraVariable& InTargetScriptVariableToDrag) const { return false; };
+
+	virtual bool CanDragParameters(const TArray<FNiagaraVariable>& InTargetScriptVariablesToDrag) const { return false; };
+
+	/** Returns a list of Graphs that are valid for operations to edit their variables and/or metadata. 
+	 *Should collect all Graphs that are currently selected, but also Graphs that are implicitly selected, e.g. the node graph for the script toolkit.
+	 */
+	virtual TArray<TWeakObjectPtr<UNiagaraGraph>> GetEditableGraphs() const = 0;
+
+	virtual const TArray<FNiagaraScriptVariableAndViewInfo> GetViewedParameters() = 0;
+
+	TSharedRef<class SWidget> GetScriptParameterVisualWidget(FCreateWidgetForActionData* const InCreateData) const;
+
+	FOnParameterPanelViewModelRefreshed& GetOnRefreshed() { return OnParameterPanelViewModelRefreshed; };
+
+	FOnParameterPanelViewModelExternalSelectionChanged& GetExternalSelectionChanged() { return OnParameterPanelViewModelExternalSelectionChanged; }
+protected:
+	FOnParameterPanelViewModelRefreshed OnParameterPanelViewModelRefreshed;
+	FOnParameterPanelViewModelExternalSelectionChanged OnParameterPanelViewModelExternalSelectionChanged;
+	/** Cached list of parameters sent to SNiagarParameterPanel, updated whenever GetViewedParameters is called. */
+	TArray<FNiagaraScriptVariableAndViewInfo> CachedViewedParameters;
+
+};
+
+class FNiagaraSystemToolkitParameterPanelViewModel : public INiagaraParameterPanelViewModel, public FEditorUndoClient
+{
+public:
+	/** Construct a SystemToolkit Parameter Panel View Model from a System View Model. */
+	FNiagaraSystemToolkitParameterPanelViewModel(TSharedPtr<FNiagaraSystemViewModel> InSystemViewModel);
+
+	~FNiagaraSystemToolkitParameterPanelViewModel();
+
+	/** Begin INiagaraParameterPanelViewModel interface. */
+	virtual void InitBindings() override;
+
+	virtual void Refresh() const override;
+
+	virtual void CollectStaticSections(TArray<int32>& StaticSectionIDs) const override;
+
+	virtual NiagaraParameterPanelSectionID::Type GetSectionForVarAndViewInfo(const FNiagaraScriptVariableAndViewInfo& VarAndViewInfo) const override;
+
+	virtual const UNiagaraScriptVariable* AddParameter(FNiagaraVariable& VariableToAdd, const FNiagaraVariableMetaData& VariableMetaDataToAssign) override;
+
+	virtual void DeleteParameter(const FNiagaraVariable& TargetVariableToRemove, const FNiagaraVariableMetaData& TargetVariableMetaData) override;
+
+	virtual bool GetCanDeleteParameterAndToolTip(const FNiagaraVariable& TargetVariableToRemove, const FNiagaraVariableMetaData& TargetVariableMetaData, FText& OutCanDeleteParameterToolTip) const override;
+
+	virtual void RenameParameter(const FNiagaraVariable& TargetVariableToRename, const FNiagaraVariableMetaData& TargetVariableMetaData, const FText& NewVariableNameText) const override;
+
+	virtual void ChangeParameterScope(const FNiagaraVariable& TargetVariableToModify, const FNiagaraVariableMetaData& TargetVariableMetaData, const ENiagaraParameterScope NewVariableScope) const override;
+
+	virtual bool CanModifyParameter(const FNiagaraVariable& TargetVariableToModify, const FNiagaraVariableMetaData& TargetVariableMetaData) const override;
+
+	virtual bool GetCanRenameParameterAndToolTip(const FNiagaraVariable& TargetVariableToRename, const FNiagaraVariableMetaData& TargetVariableMetaData, TOptional<const FText> NewVariableNameText, FText& OutCanRenameParameterToolTip) const override;
+
+	virtual FReply HandleActionDragged(const TSharedPtr<FEdGraphSchemaAction>& InAction, const FPointerEvent& MouseEvent) const override;
+
+	virtual TArray<TWeakObjectPtr<UNiagaraGraph>> GetEditableGraphs() const override;
+
+	virtual const TArray<FNiagaraScriptVariableAndViewInfo> GetViewedParameters() override;
+	/** End INiagaraParameterPanelViewModel interface. */
+
+	/** Begin FEditorUndoClient Interface */
+	virtual void PostUndo(bool bSuccess) override;
+	virtual void PostRedo(bool bSuccess) override { PostUndo(bSuccess); }
+	/** End FEditorUndoClient Interface */
+
+private:
+	//** Updates the SelectedEmitterScriptGraphs array and then calls Refresh(). Use if the Emitter Script graphs change (e.g. Emitter deleted from System). */
+	void RefreshSelectedEmitterScriptGraphs();
+
+	TArray<TWeakObjectPtr<UNiagaraGraph>> GetEditableEmitterScriptGraphs() const;
+
+	// Graphs viewed to gather UNiagaraScriptVariables that are displayed by the Parameter Panel.
+	TWeakObjectPtr<UNiagaraGraph> SystemScriptGraph;
+	TArray<TWeakObjectPtr<UNiagaraGraph>>  SelectedEmitterScriptGraphs;
+
+	TSharedPtr<FNiagaraSystemViewModel> SystemViewModel;
+	UNiagaraSystemSelectionViewModel* OverviewSelectionViewModel;
+};
+
+class FNiagaraScriptToolkitParameterPanelViewModel : public INiagaraParameterPanelViewModel
+{
+public:
+	/** Construct a ScriptToolkit Parameter Panel View Model from a Script View Model. */
+	FNiagaraScriptToolkitParameterPanelViewModel(TSharedPtr<FNiagaraScriptViewModel> InScriptViewModel);
+
+	~FNiagaraScriptToolkitParameterPanelViewModel();
+
+	/** Begin INiagaraParameterPanelViewModel interface. */
+	virtual void InitBindings() override;
+
+	virtual void Refresh() const override;
+
+	virtual void CollectStaticSections(TArray<int32>& StaticSectionIDs) const override;
+
+	virtual NiagaraParameterPanelSectionID::Type GetSectionForVarAndViewInfo(const FNiagaraScriptVariableAndViewInfo& VarAndViewInfo) const override;
+
+	virtual const UNiagaraScriptVariable* AddParameter(FNiagaraVariable& VariableToAdd, const FNiagaraVariableMetaData& VariableMetaDataToAssign) override;
+
+	virtual void DeleteParameter(const FNiagaraVariable& TargetVariableToRemove, const FNiagaraVariableMetaData& TargetVariableMetaData) override;
+
+	virtual bool GetCanDeleteParameterAndToolTip(const FNiagaraVariable& TargetVariableToRemove, const FNiagaraVariableMetaData& TargetVariableMetaData, FText& OutCanDeleteParameterToolTip) const override;
+
+	virtual void RenameParameter(const FNiagaraVariable& TargetVariableToRename, const FNiagaraVariableMetaData& TargetVariableMetaData, const FText& NewVariableNameText) const override;
+
+	virtual void ChangeParameterScope(const FNiagaraVariable& TargetVariableToModify, const FNiagaraVariableMetaData& TargetVariableMetaData, const ENiagaraParameterScope NewVariableScope) const override;
+
+	virtual bool CanModifyParameter(const FNiagaraVariable& TargetVariableToModify, const FNiagaraVariableMetaData& TargetVariableMetaData) const override;
+
+	virtual bool GetCanRenameParameterAndToolTip(const FNiagaraVariable& TargetVariableToRename, const FNiagaraVariableMetaData& TargetVariableMetaData, TOptional<const FText> NewVariableNameText, FText& OutCanRenameParameterToolTip) const override;
+
+	virtual void HandleActionSelected(const TSharedPtr<FEdGraphSchemaAction>& InAction, ESelectInfo::Type InSelectionType) override;
+	
+	void HandleGraphSubObjectSelectionChanged(const UObject* Obj);
+
+	virtual FReply HandleActionDragged(const TSharedPtr<FEdGraphSchemaAction>& InAction, const FPointerEvent& MouseEvent) const override;
+
+	virtual TArray<TWeakObjectPtr<UNiagaraGraph>> GetEditableGraphs() const override;
+
+	virtual const TArray<FNiagaraScriptVariableAndViewInfo> GetViewedParameters() override;
+	/** End INiagaraParameterPanelViewModel interface. */
+
+	void RenamePin(const UEdGraphPin* TargetPinToRename, const FText& NewNameText) const;
+	void ChangePinScope(const UEdGraphPin* TargetPin, const ENiagaraParameterScope NewScope) const;
+
+private:
+	TSharedPtr<FNiagaraScriptViewModel> ScriptViewModel;
+
+	FDelegateHandle OnGraphChangedHandle;
+	FDelegateHandle OnGraphNeedsRecompileHandle;
+	FDelegateHandle ScriptVisualPinHandle;
+	FDelegateHandle OnSubObjectSelectionHandle;
+
+	void HandleOnGraphChanged(const struct FEdGraphEditAction& InAction);
+
+	TStaticArray<FScopeIsEnabledAndTooltip, (int32)ENiagaraParameterScope::Num> GetParameterScopesEnabledAndTooltips(const FNiagaraVariable& InVar, const FNiagaraVariableMetaData& InVarMetaData) const;
+
+	TSharedRef<class SWidget> GetScriptParameterVisualWidget(const UEdGraphPin* Pin) const;
+
+	TSharedPtr<FNiagaraObjectSelection> VariableObjectSelection;
+};
diff --git a/Engine/Plugins/FX/Niagara/Source/NiagaraEditor/Public/ViewModels/NiagaraScratchPadScriptViewModel.h b/Engine/Plugins/FX/Niagara/Source/NiagaraEditor/Public/ViewModels/NiagaraScratchPadScriptViewModel.h
--- a/Engine/Plugins/FX/Niagara/Source/NiagaraEditor/Public/ViewModels/NiagaraScratchPadScriptViewModel.h
+++ b/Engine/Plugins/FX/Niagara/Source/NiagaraEditor/Public/ViewModels/NiagaraScratchPadScriptViewModel.h
@@ -5,6 +5,9 @@
 #include "ViewModels/NiagaraScriptViewModel.h"
 #include "GraphEditAction.h"
 
+class INiagaraParameterPanelViewModel;
+class FUICommandList;
+
 class NIAGARAEDITOR_API FNiagaraScratchPadScriptViewModel : public FNiagaraScriptViewModel, public FGCObject
 {
 public:
@@ -18,12 +21,18 @@ public:
 
 	void Initialize(UNiagaraScript* Script);
 
+	void Finalize();
+
 	virtual void AddReferencedObjects(FReferenceCollector& Collector) override;
 
 	UNiagaraScript* GetOriginalScript() const;
 
 	UNiagaraScript* GetEditScript() const;
 
+	TSharedPtr<INiagaraParameterPanelViewModel> GetParameterPanelViewModel() const;
+
+	TSharedPtr<FUICommandList> GetParameterPanelCommands() const;
+
 	FText GetToolTip() const;
 
 	bool GetIsPendingRename() const;
@@ -72,6 +81,9 @@ private:
 
 	bool bHasPendingChanges;
 
+	TSharedPtr<FUICommandList> ParameterPanelCommands;
+	TSharedPtr<INiagaraParameterPanelViewModel> ParameterPaneViewModel;
+
 	FDelegateHandle OnGraphNeedsRecompileHandle;
 
 	FOnRenamed OnRenamedDelegate;
diff --git a/Engine/Plugins/FX/Niagara/Source/NiagaraEditor/Public/ViewModels/NiagaraScratchPadViewModel.h b/Engine/Plugins/FX/Niagara/Source/NiagaraEditor/Public/ViewModels/NiagaraScratchPadViewModel.h
--- a/Engine/Plugins/FX/Niagara/Source/NiagaraEditor/Public/ViewModels/NiagaraScratchPadViewModel.h
+++ b/Engine/Plugins/FX/Niagara/Source/NiagaraEditor/Public/ViewModels/NiagaraScratchPadViewModel.h
@@ -82,6 +82,8 @@ private:
 
 	TSharedRef<FNiagaraScratchPadScriptViewModel> CreateAndSetupScriptviewModel(UNiagaraScript* ScratchPadScript);
 
+	void TearDownScriptViewModel(TSharedRef<FNiagaraScratchPadScriptViewModel> InScriptViewModel);
+
 	void RefreshEditScriptViewModels();
 
 	void ScriptGraphNodeSelectionChanged(TWeakPtr<FNiagaraScratchPadScriptViewModel> InScriptViewModelWeak);
@@ -92,6 +94,8 @@ private:
 
 	void ScriptViewModelRequestDiscardChanges(TWeakPtr<FNiagaraScratchPadScriptViewModel> ScriptViewModelWeak);
 
+	void ScriptViewModelVariableSelectionChanged(TWeakPtr<FNiagaraScratchPadScriptViewModel> ScriptViewModelWeak);
+
 private:
 	TSharedPtr<FNiagaraObjectSelection> ObjectSelection;
 
diff --git a/Engine/Plugins/FX/Niagara/Source/NiagaraEditor/Public/Widgets/SNiagaraParameterPanel.h b/Engine/Plugins/FX/Niagara/Source/NiagaraEditor/Public/Widgets/SNiagaraParameterPanel.h
new file mode 100644
--- /dev/null
+++ b/Engine/Plugins/FX/Niagara/Source/NiagaraEditor/Public/Widgets/SNiagaraParameterPanel.h
@@ -0,0 +1,185 @@
+// Copyright Epic Games, Inc. All Rights Reserved.
+
+#pragma once
+
+#include "Widgets/SCompoundWidget.h"
+#include "Internationalization/Text.h"
+#include "NiagaraTypes.h"
+#include "NiagaraActions.h"
+#include "Framework/Commands/Commands.h"
+
+class SGraphActionMenu;
+class SEditableTextBox;
+class SExpanderArrow;
+class SSearchBox;
+class SComboButton;
+class SNiagaraGraphPinAdd;
+class FNiagaraObjectSelection;
+class UNiagaraGraph;
+struct FEdGraphSchemaAction; //@todo(ng) cleanup
+class FUICommandList;
+class INiagaraParameterPanelViewModel;
+struct FNiagaraVariableMetaData;
+
+class FNiagaraParameterPanelCommands : public TCommands<FNiagaraParameterPanelCommands>
+{
+public:
+	/** Constructor */
+	FNiagaraParameterPanelCommands()
+		: TCommands<FNiagaraParameterPanelCommands>(TEXT("NiagaraParameterMapViewCommands"), NSLOCTEXT("Contexts", "NiagaraParameterPanel", "NiagaraParameterPanel"), NAME_None, FEditorStyle::GetStyleSetName())
+	{
+	}
+
+	// Basic operations
+	TSharedPtr<FUICommandInfo> DeleteEntry;
+
+	/** Initialize commands */
+	virtual void RegisterCommands() override;
+};
+
+/** A widget for viewing and editing UNiagaraScriptVariables provided by an INiagaraParameterPanelViewModel */
+class SNiagaraParameterPanel : public SCompoundWidget
+{
+public:
+	enum EToolkitType
+	{
+		SCRIPT,
+		SYSTEM,
+	};
+
+	SLATE_BEGIN_ARGS(SNiagaraParameterPanel)
+	{}
+	SLATE_END_ARGS();
+
+	NIAGARAEDITOR_API ~SNiagaraParameterPanel();
+
+	NIAGARAEDITOR_API void Construct(const FArguments& InArgs, const TSharedPtr<INiagaraParameterPanelViewModel>& InParameterPanelViewModel, const TSharedPtr<FUICommandList>& InToolkitCommands);
+
+	NIAGARAEDITOR_API virtual void Tick(const FGeometry& AllottedGeometry, const double InCurrentTime, const float InDeltaTime) override;
+
+	/** Whether the add parameter button should be enabled. */
+// 	bool ParameterAddEnabled() const; //@todo(ng) impl
+
+	void AddParameter(FNiagaraVariable NewVariable, const NiagaraParameterPanelSectionID::Type SectionID);
+
+ 	/** Refreshes the palette items for the graph action menu on the next tick. */
+ 	void Refresh();
+
+	static TSharedRef<SExpanderArrow> CreateCustomActionExpander(const struct FCustomExpanderData& ActionMenuData);
+
+private:
+	/** Function to bind to SNiagaraAddParameterMenus to filter types we allow creating */
+	bool AllowMakeType(const FNiagaraTypeDefinition& InType) const;
+
+	// SGraphActionMenu delegates
+	TSharedRef<SWidget> OnCreateWidgetForAction(struct FCreateWidgetForActionData* const InCreateData);
+	void CollectAllActions(FGraphActionListBuilderBase& OutAllActions);
+	void CollectStaticSections(TArray<int32>& StaticSectionIDs) const;
+	FReply OnActionDragged(const TArray<TSharedPtr<FEdGraphSchemaAction>>& InActions, const FPointerEvent& MouseEvent);
+	void OnActionSelected(const TArray<TSharedPtr<FEdGraphSchemaAction>>& InActions, ESelectInfo::Type InSelectionType);
+// 	void OnActionDoubleClicked(const TArray<TSharedPtr<FEdGraphSchemaAction>>& InActions); //@todo(ng) impl
+	TSharedPtr<SWidget> OnContextMenuOpening();
+ 	FText OnGetSectionTitle(int32 InSectionID);
+	TSharedRef<SWidget> OnGetSectionWidget(TSharedRef<SWidget> RowWidget, int32 InSectionID);
+	TSharedRef<SWidget> CreateAddToSectionButton(const NiagaraParameterPanelSectionID::Type InSection, TWeakPtr<SWidget> WeakRowWidget, FText AddNewText, FName MetaDataTag);
+// 
+// 	/** Checks if the selected action has context menu */
+// 	bool SelectionHasContextMenu() const; //@todo(ng) impl
+// 
+	TSharedRef<SWidget> OnGetParameterMenu(const NiagaraParameterPanelSectionID::Type InSection);
+	
+	//EVisibility OnAddButtonTextVisibility(TWeakPtr<SWidget> RowWidget, const NiagaraParameterMapSectionID::Type InSection) const; //@todo(ng) impl
+
+
+// 	//Callbacks
+	/** Try to delete all entries that are selected. If only some selected entries pass CanDeleteEntry(), then only those selected entries will be deleted. */
+	void TryDeleteEntries();
+	/** Test if ALL selected entries cannot be deleted. If only one selected entry passes INiagaraParameterPanelViewModel::CanRemoveParameter(), return true. Per entry deletion handling is done in TryDeleteEntries(). */
+	bool CanTryDeleteEntries() const;
+	void OnRequestRenameOnActionNode();
+	bool CanRequestRenameOnActionNode(TWeakPtr<struct FGraphActionNode> InSelectedNode) const;
+	bool CanRequestRenameOnActionNode() const;
+	void HandlePaletteItemParameterRenamed(const FText& InText, const FNiagaraScriptVarAndViewInfoAction& InAction);
+	void HandlePaletteItemScopeComboSelectionChanged(ENiagaraParameterScope InScope, const FNiagaraScriptVarAndViewInfoAction& InAction);
+	void HandleExternalSelectionChanged(const UObject* Obj);
+
+	/** Delegate handler used to match an FName to an action in the list, used for renaming keys */
+	bool HandleActionMatchesName(struct FEdGraphSchemaAction* InAction, const FName& InName) const;
+
+private:
+	/** Graph Action Menu for displaying all our variables and functions */
+	TSharedPtr<SGraphActionMenu> GraphActionMenu;
+
+	/** The filter box that handles filtering for both graph action menus. */
+	TSharedPtr<SSearchBox> FilterBox;
+ 
+	/** Add parameter buttons for all sections. */
+	TArray<TSharedPtr<SComboButton>> AddParameterButtons;
+
+	TSharedPtr<FNiagaraObjectSelection> SelectedVariableObjects;
+
+	TSharedPtr<FUICommandList> ToolkitCommands; //@todo(ng) add Find And Rename Parameter command
+
+	bool bNeedsRefresh;
+	bool bGraphActionPendingRename;
+
+	TSharedPtr<INiagaraParameterPanelViewModel> ParameterPanelViewModel;
+};
+
+class SNiagaraAddParameterMenu2 : public SCompoundWidget
+{
+public:
+	/** Delegate that can be used to create a widget for a particular action */
+	DECLARE_DELEGATE_OneParam(FOnAddParameter, FNiagaraVariable);
+	DECLARE_DELEGATE_TwoParams(FOnCollectCustomActions, FGraphActionListBuilderBase&, bool&);
+	DECLARE_DELEGATE_RetVal_OneParam(bool, FOnAllowMakeType, const FNiagaraTypeDefinition&);
+
+	SLATE_BEGIN_ARGS(SNiagaraAddParameterMenu2)
+		: _AllowCreatingNew(true)
+		, _ShowGraphParameters(true)
+		, _ShowKnownConstantParametersFilter(NiagaraParameterPanelSectionID::Type::NONE)
+		, _AutoExpandMenu(false)
+		, _IsParameterRead(true)
+		, _NewParameterScope(ENiagaraParameterScope::Particles) {}
+		SLATE_EVENT(FOnAddParameter, OnAddParameter)
+		SLATE_EVENT(FOnCollectCustomActions, OnCollectCustomActions)
+		SLATE_EVENT(FOnAllowMakeType, OnAllowMakeType)
+		SLATE_ATTRIBUTE(bool, AllowCreatingNew)
+		SLATE_ATTRIBUTE(bool, ShowGraphParameters)
+		SLATE_ATTRIBUTE(NiagaraParameterPanelSectionID::Type, ShowKnownConstantParametersFilter)
+		SLATE_ATTRIBUTE(bool, AutoExpandMenu)
+		SLATE_ATTRIBUTE(bool, IsParameterRead)
+		SLATE_ATTRIBUTE(FString, NewParameterNamespace)
+		SLATE_ARGUMENT(ENiagaraParameterScope, NewParameterScope)
+	SLATE_END_ARGS();
+
+	void Construct(const FArguments& InArgs, TArray<TWeakObjectPtr<UNiagaraGraph>> InGraphs);
+
+	TSharedRef<SEditableTextBox> GetSearchBox();
+
+private:
+	void OnActionSelected(const TArray<TSharedPtr<FEdGraphSchemaAction>>& SelectedActions, ESelectInfo::Type InSelectionType);
+	void CollectAllActions(FGraphActionListBuilderBase& OutAllActions);
+	void AddParameterSelected(FNiagaraVariable NewVariable);
+	void AddParameterAndMetaDataSelected(FNiagaraVariable NewVariable, const FNiagaraVariableMetaData NewVariableMetaData);
+
+	TSharedPtr<SGraphActionMenu> GraphMenu;
+
+	/** Delegate that gets fired when a parameter was added. */
+	FOnAddParameter OnAddParameter;
+	FOnCollectCustomActions OnCollectCustomActions;
+	FOnAllowMakeType OnAllowMakeType;
+
+	TArray<TWeakObjectPtr<UNiagaraGraph>> Graphs;
+
+	TAttribute<bool> AllowCreatingNew;
+	TAttribute<bool> ShowGraphParameters;
+	TAttribute<NiagaraParameterPanelSectionID::Type> ShowKnownConstantParametersFilter;
+	TAttribute<bool> AutoExpandMenu;
+	TAttribute<bool> IsParameterRead;
+	TAttribute<FString> NewParameterNamespace;
+
+	/** Default scope to give to new parameters created through this menu. Used when generating new variable names so that their namespace is correct. */
+	ENiagaraParameterScope NewParameterScope;
+	FText NewParameterScopeText;
+};
diff --git a/Engine/Plugins/FX/Niagara/Source/NiagaraEditorWidgets/Private/NiagaraEditorWidgetsStyle.cpp b/Engine/Plugins/FX/Niagara/Source/NiagaraEditorWidgets/Private/NiagaraEditorWidgetsStyle.cpp
--- a/Engine/Plugins/FX/Niagara/Source/NiagaraEditorWidgets/Private/NiagaraEditorWidgetsStyle.cpp
+++ b/Engine/Plugins/FX/Niagara/Source/NiagaraEditorWidgets/Private/NiagaraEditorWidgetsStyle.cpp
@@ -246,15 +246,10 @@ TSharedRef< FSlateStyleSet > FNiagaraEditorWidgetsStyle::Create()
 	Style->Set("NiagaraEditor.ShowInCurveEditorIcon", new IMAGE_PLUGIN_BRUSH("Icons/ShowInCurveEditor", Icon16x16, FLinearColor::White));
 
 	// Scratch pad
-	FSlateFontInfo ScratchPadLargeHeaderFont = DEFAULT_FONT("Bold", 11);
-	FTextBlockStyle ScratchPadLargeHeaderText = FTextBlockStyle(NormalText)
-		.SetFont(ScratchPadLargeHeaderFont);
-	Style->Set("NiagaraEditor.ScratchPad.LargeHeaderText", ScratchPadLargeHeaderText);
-
-	FSlateFontInfo ScratchPadSmallHeaderFont = DEFAULT_FONT("Bold", 11);
-	FTextBlockStyle ScratchPadSmallHeaderText = FTextBlockStyle(NormalText)
-		.SetFont(ScratchPadSmallHeaderFont);
-	Style->Set("NiagaraEditor.ScratchPad.SmallHeaderText", ScratchPadSmallHeaderText);
+	FSlateFontInfo ScratchPadEditorHeaderFont = DEFAULT_FONT("Bold", 11);
+	FTextBlockStyle ScratchPadEditorHeaderText = FTextBlockStyle(NormalText)
+		.SetFont(ScratchPadEditorHeaderFont);
+	Style->Set("NiagaraEditor.ScratchPad.EditorHeaderText", ScratchPadEditorHeaderText);
 
 	FSlateBrush ScratchPadCategoryBrush = BOX_PLUGIN_BRUSH("Icons/CategoryRow", FMargin(2.0f / 8.0f), FLinearColor(FColor(48, 48, 48)));
 	FSlateBrush ScratchPadHoveredCategoryBrush = BOX_PLUGIN_BRUSH("Icons/CategoryRow", FMargin(2.0f / 8.0f), FLinearColor(FColor(38, 38, 38)));
diff --git a/Engine/Plugins/FX/Niagara/Source/NiagaraEditorWidgets/Private/SNiagaraScratchPad.cpp b/Engine/Plugins/FX/Niagara/Source/NiagaraEditorWidgets/Private/SNiagaraScratchPad.cpp
--- a/Engine/Plugins/FX/Niagara/Source/NiagaraEditorWidgets/Private/SNiagaraScratchPad.cpp
+++ b/Engine/Plugins/FX/Niagara/Source/NiagaraEditorWidgets/Private/SNiagaraScratchPad.cpp
@@ -13,6 +13,7 @@
 #include "NiagaraEditorStyle.h"
 #include "NiagaraEditorWidgetsStyle.h"
 #include "NiagaraScratchPadCommandContext.h"
+#include "Widgets/SNiagaraParameterPanel.h"
 
 #include "Widgets/Text/SInlineEditableTextBlock.h"
 #include "Widgets/Input/SButton.h"
@@ -25,6 +26,7 @@
 #define LOCTEXT_NAMESPACE "NiagaraScratchPad"
 
 FName ScriptSelectorName = "ScriptSelector";
+FName ScriptParameterPanelName = "ScriptParameterPanel";
 FName ScriptEditorName = "ScriptEditor";
 FName SelectionEditorName = "SelectionEditor";
 FName WideLayoutName = "Wide";
@@ -375,7 +377,7 @@ private:
 		.Padding(3, 0, 0, 0)
 		[
 			SNew(STextBlock)
-			.TextStyle(FNiagaraEditorWidgetsStyle::Get(), "NiagaraEditor.ScratchPad.SmallHeaderText")
+			.TextStyle(FEditorStyle::Get(), "DetailsView.CategoryTextStyle")
 			.Text(ViewModel->GetDisplayNameForUsage(Category))
 		]
 		+ SHorizontalBox::Slot()
@@ -467,11 +469,11 @@ class SNiagaraScratchPadScriptEditor : public SCompoundWidget
 				SNew(SHorizontalBox)
 				+ SHorizontalBox::Slot()
 				.VAlign(VAlign_Center)
-				.Padding(0, 0, 2, 0)
+				.Padding(2, 0, 2, 0)
 				.AutoWidth()
 				[
 					SNew(STextBlock)
-					.TextStyle(FNiagaraEditorWidgetsStyle::Get(), "NiagaraEditor.ScratchPad.LargeHeaderText")
+					.TextStyle(FNiagaraEditorWidgetsStyle::Get(), "NiagaraEditor.ScratchPad.EditorHeaderText")
 					.Text(this, &SNiagaraScratchPadScriptEditor::GetNameText)
 					.ToolTipText(ScriptViewModel.ToSharedRef(), &FNiagaraScratchPadScriptViewModel::GetToolTip)
 				]
@@ -510,7 +512,7 @@ class SNiagaraScratchPadScriptEditor : public SCompoundWidget
 						.Padding(2.0f, 2.0f, 2.0f, 3.0f)
 						[
 							SNew(STextBlock)
-							.TextStyle(FNiagaraEditorWidgetsStyle::Get(), "NiagaraEditor.ScratchPad.SmallHeaderText")
+							.TextStyle(FNiagaraEditorWidgetsStyle::Get(), "NiagaraEditor.ScratchPad.EditorHeaderText")
 							.Text(LOCTEXT("ApplyButtonLabel", "Apply"))
 						]
 					]
@@ -692,6 +694,89 @@ private:
 	bool bIsUpdatingSelection;
 };
 
+class SNiagaraScratchPadParameterPanel : public SCompoundWidget
+{
+	SLATE_BEGIN_ARGS(SNiagaraScratchPadScriptEditor) {}
+	SLATE_END_ARGS();
+
+	void Construct(const FArguments& InArgs, UNiagaraScratchPadViewModel* InViewModel)
+	{
+		ViewModel = InViewModel;
+		ViewModel->OnActiveScriptChanged().AddSP(this, &SNiagaraScratchPadParameterPanel::ActiveScriptChanged);
+		UpdateContent();
+	}
+
+	virtual FReply OnKeyDown(const FGeometry& MyGeometry, const FKeyEvent& InKeyEvent) override
+	{
+		TSharedPtr<FNiagaraScratchPadScriptViewModel> ScriptViewModel = ScriptViewModelWeak.Pin();
+		if (ScriptViewModel.IsValid() && ScriptViewModel->GetParameterPanelCommands()->ProcessCommandBindings(InKeyEvent))
+		{
+			return FReply::Handled();
+		}
+		return SCompoundWidget::OnKeyDown(MyGeometry, InKeyEvent);
+	}
+
+private:
+	void UpdateContent()
+	{
+		TSharedPtr<FNiagaraScratchPadScriptViewModel> ScriptViewModel = ScriptViewModelWeak.Pin();
+		TSharedPtr<FNiagaraScratchPadScriptViewModel> NewScriptViewModel = ViewModel->GetActiveScriptViewModel();
+		if (NewScriptViewModel != ScriptViewModel)
+		{
+			if (NewScriptViewModel.IsValid())
+			{
+				ChildSlot
+				[
+					SNew(SNiagaraParameterPanel, NewScriptViewModel->GetParameterPanelViewModel(), NewScriptViewModel->GetParameterPanelCommands())
+				];
+				ScriptViewModelWeak = NewScriptViewModel;
+			}
+			else
+			{
+				ChildSlot[SNullWidget::NullWidget];
+			}
+		}
+	}
+
+	void ActiveScriptChanged()
+	{
+		UpdateContent();
+	}
+
+private:
+	UNiagaraScratchPadViewModel* ViewModel;
+	TWeakPtr<FNiagaraScratchPadScriptViewModel> ScriptViewModelWeak;
+};
+
+class SNiagaraScratchPadSectionBox : public SCompoundWidget
+{
+	SLATE_BEGIN_ARGS(SNiagaraScratchPadSectionBox) {}
+		SLATE_ATTRIBUTE(FText, HeaderText)
+		SLATE_DEFAULT_SLOT(FArguments, Content);
+	SLATE_END_ARGS();
+
+	void Construct(const FArguments& InArgs)
+	{
+		ChildSlot
+		[
+			SNew(SVerticalBox)
+			+ SVerticalBox::Slot()
+			.AutoHeight()
+			.Padding(1.0f, 5.0f, 0.0f, 4.0f)
+			[
+				SNew(STextBlock)
+				.TextStyle(FEditorStyle::Get(), "Docking.TabFont")
+				.Text(InArgs._HeaderText)
+			]
+			+ SVerticalBox::Slot()
+			.Padding(5.0f, 0.0f, 3.0f, 5.0f)
+			[
+				InArgs._Content.Widget
+			]
+		];
+	}
+};
+
 void SNiagaraScratchPad::Construct(const FArguments& InArgs, UNiagaraScratchPadViewModel* InViewModel)
 {
 	ViewModel = InViewModel;
@@ -706,6 +791,10 @@ void SNiagaraScratchPad::Construct(const FArguments& InArgs, UNiagaraScratchPadV
 			{
 				return ConstructScriptSelector();
 			}
+			else if (InWidgetName == ScriptParameterPanelName)
+			{
+				return ConstructParameterPanel();
+			}
 			else if (InWidgetName == ScriptEditorName)
 			{
 				return ConstructScriptEditor();
@@ -726,17 +815,35 @@ void SNiagaraScratchPad::Construct(const FArguments& InArgs, UNiagaraScratchPadV
 			{
 				Layout = SNew(SSplitter)
 				.Orientation(Orient_Horizontal)
-				.PhysicalSplitterHandleSize(7.0f)
-				.HitDetectionSplitterHandleSize(7.0f)
+				.PhysicalSplitterHandleSize(4.0f)
+				.HitDetectionSplitterHandleSize(6.0f)
 				+ SSplitter::Slot()
 				.Value(0.15f)
 				[
-					InNamedWidgetProvider.GetNamedWidget(ScriptSelectorName)
+					SNew(SSplitter)
+					.Style(FEditorStyle::Get(), "SplitterDark")
+					.Orientation(Orient_Vertical)
+					.PhysicalSplitterHandleSize(4.0f)
+					.HitDetectionSplitterHandleSize(6.0f)
+					+ SSplitter::Slot()
+					.Value(0.5f)
+					[
+						InNamedWidgetProvider.GetNamedWidget(ScriptSelectorName)
+					]
+					+ SSplitter::Slot()
+					.Value(0.5f)
+					[
+						InNamedWidgetProvider.GetNamedWidget(ScriptParameterPanelName)
+					]
 				]
 				+ SSplitter::Slot()
 				.Value(0.6f)
 				[
-					InNamedWidgetProvider.GetNamedWidget(ScriptEditorName)
+					SNew(SBox)
+					.Padding(FMargin(3.0f, 0.0f))
+					[
+						InNamedWidgetProvider.GetNamedWidget(ScriptEditorName)
+					]
 				]
 				+ SSplitter::Slot()
 				.Value(0.25f)
@@ -748,20 +855,28 @@ void SNiagaraScratchPad::Construct(const FArguments& InArgs, UNiagaraScratchPadV
 			{
 				Layout = SNew(SSplitter)
 				.Orientation(Orient_Horizontal)
-				.PhysicalSplitterHandleSize(7.0f)
-				.HitDetectionSplitterHandleSize(7.0f)
+				.PhysicalSplitterHandleSize(4.0f)
+				.HitDetectionSplitterHandleSize(6.0f)
 				+ SSplitter::Slot()
 				.Value(0.3f)
 				[
 					SNew(SSplitter)
+					.Style(FEditorStyle::Get(), "SplitterDark")
 					.Orientation(Orient_Vertical)
+					.PhysicalSplitterHandleSize(4.0f)
+					.HitDetectionSplitterHandleSize(6.0f)
 					+ SSplitter::Slot()
-					.Value(0.5f)
+					.Value(0.3f)
 					[
 						InNamedWidgetProvider.GetNamedWidget(ScriptSelectorName)
 					]
 					+ SSplitter::Slot()
-					.Value(0.5f)
+					.Value(0.3f)
+					[
+						InNamedWidgetProvider.GetNamedWidget(ScriptParameterPanelName)
+					]
+					+ SSplitter::Slot()
+					.Value(0.4f)
 					[
 						InNamedWidgetProvider.GetNamedWidget(SelectionEditorName)
 					]
@@ -769,7 +884,11 @@ void SNiagaraScratchPad::Construct(const FArguments& InArgs, UNiagaraScratchPadV
 				+ SSplitter::Slot()
 				.Value(0.7f)
 				[
-					InNamedWidgetProvider.GetNamedWidget(ScriptEditorName)
+					SNew(SBox)
+					.Padding(FMargin(3.0f, 0.0f, 0.0f, 0.0f))
+					[
+						InNamedWidgetProvider.GetNamedWidget(ScriptEditorName)
+					]
 				];
 			}
 			else
@@ -794,21 +913,22 @@ void SNiagaraScratchPad::Construct(const FArguments& InArgs, UNiagaraScratchPadV
 
 TSharedRef<SWidget> SNiagaraScratchPad::ConstructScriptSelector()
 {
-	return SNew(SVerticalBox)
-	+ SVerticalBox::Slot()
-	.AutoHeight()
-	.Padding(0.0f, 2.0f)
-	[
-		SNew(STextBlock)
-		.TextStyle(FNiagaraEditorWidgetsStyle::Get(), "NiagaraEditor.ScratchPad.LargeHeaderText")
-		.Text(LOCTEXT("ScriptSelector", "Scratch Script Selector"))
-	]
-	+ SVerticalBox::Slot()
+	return SNew(SNiagaraScratchPadSectionBox)
+	.HeaderText(LOCTEXT("ScriptSelector", "Scratch Script Selector"))
 	[
 		SNew(SNiagaraScratchPadScriptSelector, ViewModel.Get(), CommandContext)
 	];
 }
 
+TSharedRef<SWidget> SNiagaraScratchPad::ConstructParameterPanel()
+{
+	return SNew(SNiagaraScratchPadSectionBox)
+	.HeaderText(LOCTEXT("ScratchScriptParameters", "Scratch Script Parameters"))
+	[
+		SNew(SNiagaraScratchPadParameterPanel, ViewModel.Get())
+	];
+}
+
 TSharedRef<SWidget> SNiagaraScratchPad::ConstructScriptEditor()
 {
 	return SNew(SNiagaraScratchPadScriptEditorList, ViewModel.Get());
@@ -816,15 +936,8 @@ TSharedRef<SWidget> SNiagaraScratchPad::ConstructScriptEditor()
 
 TSharedRef<SWidget> SNiagaraScratchPad::ConstructSelectionEditor()
 {
-	return SNew(SVerticalBox)
-	+ SVerticalBox::Slot()
-	.AutoHeight()
-	[
-		SNew(STextBlock)
-		.TextStyle(FNiagaraEditorWidgetsStyle::Get(), "NiagaraEditor.ScratchPad.LargeHeaderText")
-		.Text(LOCTEXT("ScratchPadSelection", "Scratch Pad Selection"))
-	]
-	+ SVerticalBox::Slot()
+	return SNew(SNiagaraScratchPadSectionBox)
+	.HeaderText(LOCTEXT("ScratchPadSelection", "Scratch Pad Selection"))
 	[
 		SNew(SNiagaraSelectedObjectsDetails, ViewModel->GetObjectSelection())
 	];
diff --git a/Engine/Plugins/FX/Niagara/Source/NiagaraEditorWidgets/Private/SNiagaraScratchPad.h b/Engine/Plugins/FX/Niagara/Source/NiagaraEditorWidgets/Private/SNiagaraScratchPad.h
--- a/Engine/Plugins/FX/Niagara/Source/NiagaraEditorWidgets/Private/SNiagaraScratchPad.h
+++ b/Engine/Plugins/FX/Niagara/Source/NiagaraEditorWidgets/Private/SNiagaraScratchPad.h
@@ -22,6 +22,8 @@ class SNiagaraScratchPad : public SCompoundWidget
 private:
 	TSharedRef<SWidget> ConstructScriptSelector();
 
+	TSharedRef<SWidget> ConstructParameterPanel();
+
 	TSharedRef<SWidget> ConstructScriptEditor();
 
 	TSharedRef<SWidget> ConstructSelectionEditor();
