Niagara - Prevent a crash when hovering over a parameter map get or set node that's missing it's metadata.

diff --git a/Engine/Plugins/FX/Niagara/Source/NiagaraEditor/Private/NiagaraNodeParameterMapBase.cpp b/Engine/Plugins/FX/Niagara/Source/NiagaraEditor/Private/NiagaraNodeParameterMapBase.cpp
--- a/Engine/Plugins/FX/Niagara/Source/NiagaraEditor/Private/NiagaraNodeParameterMapBase.cpp
+++ b/Engine/Plugins/FX/Niagara/Source/NiagaraEditor/Private/NiagaraNodeParameterMapBase.cpp
@@ -181,8 +181,11 @@ void UNiagaraNodeParameterMapBase::GetPinHoverText(const UEdGraphPin& Pin, FStri
 			FText Description = Metadata.IsSet() ? Metadata->Description : FText::GetEmpty();
 			const FText TooltipFormat = LOCTEXT("Parameters", "Name: {0} \nType: {1}\nDescription: {2}\nScope: {3}\nUser Editable: {4}\nUsage: {5}");
 			const FText Name = FText::FromName(Var.GetName());
-			FName CachedParamName;
-			Metadata->GetParameterName(CachedParamName);
+			FName CachedParamName = NAME_None;
+			if (Metadata.IsSet())
+			{
+				Metadata->GetParameterName(CachedParamName);
+			}
 			const FText ScopeText = FText::FromName(Metadata->GetScopeName());
 			const FText UserEditableText = FText::FromName(CachedParamName);
 			const FText UsageText = StaticEnum<ENiagaraScriptParameterUsage>()->GetDisplayNameTextByValue((int64)Metadata->GetUsage());
