Removed the loading and storing of particle data for sim stages which use an iteration interface, since those stages cannot access the data using the execution index (as it doesn't represent a particle index). This should improve performance, and also removes the need to bind dummy output buffers during those stages.

Made it possible for the spawn phase to kill particles.

#rnx

diff --git a/Engine/Plugins/FX/Niagara/Shaders/Private/NiagaraEmitterInstanceShader.usf b/Engine/Plugins/FX/Niagara/Shaders/Private/NiagaraEmitterInstanceShader.usf
--- a/Engine/Plugins/FX/Niagara/Shaders/Private/NiagaraEmitterInstanceShader.usf
+++ b/Engine/Plugins/FX/Niagara/Shaders/Private/NiagaraEmitterInstanceShader.usf
@@ -58,7 +58,7 @@ bool NiagaraAll(bool4 b) { return b.x && b.y && b.z && b.w; }
 	static uint GEmitterTickCounter;
 	static uint GSimStart;
 	static uint GRandomSeedOffset = 0;
-	static bool GStageUsesAlive = false;
+	static bool GStageWritesAlive = false;
 	
 	const static uint GSpawnPhase = 0;
 	const static uint GUpdatePhase = 1;
diff --git a/Engine/Plugins/FX/Niagara/Source/NiagaraCore/Private/NiagaraCustomVersion.cpp b/Engine/Plugins/FX/Niagara/Source/NiagaraCore/Private/NiagaraCustomVersion.cpp
--- a/Engine/Plugins/FX/Niagara/Source/NiagaraCore/Private/NiagaraCustomVersion.cpp
+++ b/Engine/Plugins/FX/Niagara/Source/NiagaraCore/Private/NiagaraCustomVersion.cpp
@@ -8,4 +8,4 @@ const FGuid FNiagaraCustomVersion::GUID(0xFCF57AFA, 0x50764283, 0xB9A9E658, 0xFF
 // Register the custom version with core
 FCustomVersionRegistration GRegisterNiagaraCustomVersion(FNiagaraCustomVersion::GUID, FNiagaraCustomVersion::LatestVersion, TEXT("NiagaraVer"));
 
-const FGuid FNiagaraCustomVersion::LatestScriptCompileVersion(0xC9457034, 0x12D44E17, 0xBB38C2BB, 0x9FF38441);
\ No newline at end of file
+const FGuid FNiagaraCustomVersion::LatestScriptCompileVersion(0x25C5FF75, 0x14574F55, 0xA8A70C22, 0x98D1C114);
diff --git a/Engine/Plugins/FX/Niagara/Source/NiagaraEditor/Private/NiagaraHlslTranslator.cpp b/Engine/Plugins/FX/Niagara/Source/NiagaraEditor/Private/NiagaraHlslTranslator.cpp
--- a/Engine/Plugins/FX/Niagara/Source/NiagaraEditor/Private/NiagaraHlslTranslator.cpp
+++ b/Engine/Plugins/FX/Niagara/Source/NiagaraEditor/Private/NiagaraHlslTranslator.cpp
@@ -2169,6 +2169,12 @@ void FHlslNiagaraTranslator::DefineMainGPUFunctions(
 		int32 StartIdx = 1;
 		for (int32 i = StartIdx; i < TranslationStages.Num(); i++)
 		{
+			// No need to load particle data for stages with an iteration source, since those do not run one thread per particle.
+			if (TranslationStages[i].IterationSource != NAME_None)
+			{
+				continue;
+			}
+
 			if (TranslationStages.Num() > 2)
 			{
 				if (StartIdx == i)
@@ -2309,6 +2315,12 @@ void FHlslNiagaraTranslator::DefineMainGPUFunctions(
 		int32 StartIdx = 1;
 		for (int32 i = StartIdx; i < TranslationStages.Num(); i++)
 		{
+			// No need to store particle data for stages with an iteration source, since those do not run one thread per particle.
+			if (TranslationStages[i].IterationSource != NAME_None)
+			{
+				continue;
+			}
+
 			if (TranslationStages.Num() > 2)
 			{
 				if (StartIdx == i)
@@ -2324,14 +2336,20 @@ void FHlslNiagaraTranslator::DefineMainGPUFunctions(
 					"	{\n"), TranslationStages[i].SimulationStageIndexMin, TranslationStages[i].SimulationStageIndexMax);
 			}
 
-			if (TranslationStages[i].bUsesAlive)
+			if (TranslationStages[i].bWritesAlive || (i == 1 && TranslationStages[0].bWritesAlive))
 			{
-				HlslOutput += TEXT("\t\tGStageUsesAlive = true;\n");
+				// This stage kills particles, so we must skip the dead ones when writing out the data.
+				// It's also possible that this is the update phase, and it doesn't kill particles, but the spawn phase does. It would be nice
+				// if we could only do this for newly spawned particles, but unfortunately that would mean placing thread sync operations
+				// under dynamic flow control, which is not allowed. Therefore, we must always use the more expensive path when the spawn phase
+				// can kill particles.
+				HlslOutput += TEXT("\t\tGStageWritesAlive = true;\n");
 				HlslOutput += TEXT("\t\tconst bool bValid = Context.") + TranslationStages[i].PassNamespace + TEXT(".DataInstance.Alive;\n");
 				HlslOutput += TEXT("\t\tconst int WriteIndex = OutputIndex(0, true, bValid);\n");
 			}
 			else
 			{
+				// The stage doesn't kill particles, we can take the simpler path which doesn't need to manage the particle count.
 				HlslOutput += TEXT("\t\tconst bool bValid = GCurrentPhase != -1;\n");
 				HlslOutput += TEXT("\t\tconst int WriteIndex = OutputIndex(0, false, bValid);\n");
 			}
@@ -2536,7 +2554,7 @@ void FHlslNiagaraTranslator::DefineMainGPUFunctions(
 		HlslOutput += TEXT(
 			"	// If a stage doesn't kill particles, StoreUpdateVariables() never calls AcquireIndex(), so the\n"
 			"   // count isn't updated. In that case we must manually copy the original count here.\n"
-			"	if (!GStageUsesAlive && WriteInstanceCountOffset != 0xFFFFFFFF && GDispatchThreadId.x == 0) \n"
+			"	if (!GStageWritesAlive && WriteInstanceCountOffset != 0xFFFFFFFF && GDispatchThreadId.x == 0) \n"
 			"	{ \n"
 			"		RWInstanceCounts[WriteInstanceCountOffset] = GSpawnStartInstance + SpawnedInstances; \n"
 			"	} \n"
