The GPU version of the GetNumParticles function in the direct read DI uses the value from the GPU-side instance count buffer, instead of the value which the batcher has.

The value known to the batcher can be one or more frames behind the actual number. This approach returns the accurate value.

Also cleaned up the UnsetShaderParams calls from NiagaraDataSet. It's no longer necessary to unset the buffers when iterated stages are used, since the translator was modified previously to no longer emit code which writes to those buffers for iterated stages.

#rnx

diff --git a/Engine/Plugins/FX/Niagara/Source/Niagara/Classes/NiagaraDataSet.h b/Engine/Plugins/FX/Niagara/Source/Niagara/Classes/NiagaraDataSet.h
--- a/Engine/Plugins/FX/Niagara/Source/Niagara/Classes/NiagaraDataSet.h
+++ b/Engine/Plugins/FX/Niagara/Source/Niagara/Classes/NiagaraDataSet.h
@@ -39,6 +39,7 @@ struct FNiagaraVariableLayoutInfo
 class FNiagaraDataSet;
 class FNiagaraShader;
 class FNiagaraGPUInstanceCountManager;
+class NiagaraEmitterInstanceBatcher;
 struct FNiagaraComputeExecutionContext;
 
 //Base class for objects in Niagara that are owned by one object but are then passed for reading to other objects, potentially on other threads.
@@ -178,8 +179,8 @@ public:
 	FORCEINLINE TArray<int32>& GetIDTable() { return IDToIndexTable; }
 	FORCEINLINE const TArray<int32>& GetIDTable() const { return IDToIndexTable; }
 
-	void SetShaderParams(class FNiagaraShader *Shader, FRHICommandList &CommandList, bool bInput);
-	void UnsetShaderParams(class FNiagaraShader *Shader, FRHICommandList &CommandList);
+	void SetShaderParams(class FNiagaraShader* Shader, FRHICommandList& CommandList, bool bInput);
+	void UnsetShaderParams(class FNiagaraShader* Shader, FRHICommandList& CommandList);
 
 	void ReleaseGPUInstanceCount(FNiagaraGPUInstanceCountManager& GPUInstanceCountManager);
 
@@ -2142,7 +2143,7 @@ public:
 		}
 	}
 
-	void ReadbackData(class NiagaraEmitterInstanceBatcher* Batcher, FNiagaraDataSet* InDataSet);
+	void ReadbackData(NiagaraEmitterInstanceBatcher* Batcher, FNiagaraDataSet* InDataSet);
 	uint32 GetNumInstances() const { check(DataSet != nullptr); return NumInstances; }
 
 private:
diff --git a/Engine/Plugins/FX/Niagara/Source/Niagara/Private/NiagaraDataInterfaceParticleRead.cpp b/Engine/Plugins/FX/Niagara/Source/Niagara/Private/NiagaraDataInterfaceParticleRead.cpp
--- a/Engine/Plugins/FX/Niagara/Source/Niagara/Private/NiagaraDataInterfaceParticleRead.cpp
+++ b/Engine/Plugins/FX/Niagara/Source/Niagara/Private/NiagaraDataInterfaceParticleRead.cpp
@@ -5,6 +5,7 @@
 #include "NiagaraSystemInstance.h"
 #include "ShaderParameterUtils.h"
 #include "NiagaraRenderer.h"
+#include "NiagaraEmitterInstanceBatcher.h"
 
 #define LOCTEXT_NAMESPACE "NiagaraDataInterfaceParticleRead"
 
@@ -22,7 +23,7 @@ static const FName GetQuatAttributeFunctionName("Get Quaternion Attribute");
 
 static const FString NumSpawnedParticlesBaseName(TEXT("NumSpawnedParticles_"));
 static const FString SpawnedParticlesAcquireTagBaseName(TEXT("SpawnedParticlesAcquireTag_"));
-static const FString NumParticlesBaseName(TEXT("NumParticles_"));
+static const FString InstanceCountOffsetBaseName(TEXT("InstanceCountOffset_"));
 static const FString SpawnedIDsBufferBaseName(TEXT("SpawnedIDsBuffer_"));
 static const FString IDToIndexTableBaseName(TEXT("IDToIndexTable_"));
 static const FString InputFloatBufferBaseName(TEXT("InputFloatBuffer_"));
@@ -178,7 +179,7 @@ struct FNiagaraDataInterfaceParametersCS_ParticleRead : public FNiagaraDataInter
 	{
 		NumSpawnedParticlesParam.Bind(ParameterMap, *(NumSpawnedParticlesBaseName + ParameterInfo.DataInterfaceHLSLSymbol));
 		SpawnedParticlesAcquireTagParam.Bind(ParameterMap, *(SpawnedParticlesAcquireTagBaseName + ParameterInfo.DataInterfaceHLSLSymbol));
-		NumParticlesParam.Bind(ParameterMap, *(NumParticlesBaseName + ParameterInfo.DataInterfaceHLSLSymbol));
+		InstanceCountOffsetParam.Bind(ParameterMap, *(InstanceCountOffsetBaseName + ParameterInfo.DataInterfaceHLSLSymbol));
 		SpawnedIDsBufferParam.Bind(ParameterMap, *(SpawnedIDsBufferBaseName + ParameterInfo.DataInterfaceHLSLSymbol));
 		IDToIndexTableParam.Bind(ParameterMap, *(IDToIndexTableBaseName + ParameterInfo.DataInterfaceHLSLSymbol));
 		InputFloatBufferParam.Bind(ParameterMap, *(InputFloatBufferBaseName + ParameterInfo.DataInterfaceHLSLSymbol));
@@ -237,9 +238,9 @@ struct FNiagaraDataInterfaceParametersCS_ParticleRead : public FNiagaraDataInter
 			SetShaderValue(RHICmdList, ComputeShader, NumSpawnedParticlesParam, 0);
 			SetShaderValue(RHICmdList, ComputeShader, SpawnedParticlesAcquireTagParam, 0);
 			SetSRVParameter(RHICmdList, ComputeShader, SpawnedIDsBufferParam, FNiagaraRenderer::GetDummyIntBuffer());
-			SetShaderValue(RHICmdList, ComputeShader, NumParticlesParam, 0);
 		}
 
+		SetShaderValue(RHICmdList, ComputeShader, InstanceCountOffsetParam, -1);
 		SetSRVParameter(RHICmdList, ComputeShader, IDToIndexTableParam, FNiagaraRenderer::GetDummyIntBuffer());
 		SetSRVParameter(RHICmdList, ComputeShader, InputFloatBufferParam, FNiagaraRenderer::GetDummyFloatBuffer());
 		SetSRVParameter(RHICmdList, ComputeShader, InputIntBufferParam, FNiagaraRenderer::GetDummyIntBuffer());
@@ -452,7 +453,6 @@ struct FNiagaraDataInterfaceParametersCS_ParticleRead : public FNiagaraDataInter
 		SetShaderValue(RHICmdList, ComputeShader, NumSpawnedParticlesParam, NumSpawnedInstances);
 		SetShaderValue(RHICmdList, ComputeShader, SpawnedParticlesAcquireTagParam, IDAcquireTag);
 		SetSRVParameter(RHICmdList, ComputeShader, SpawnedIDsBufferParam, GetIntSRVWithFallback(SourceDataSet->GetGPUFreeIDs()));
-		SetShaderValue(RHICmdList, ComputeShader, NumParticlesParam, SourceData ? SourceData->GetNumInstances() : 0);
 
 		if (!SourceData)
 		{
@@ -466,25 +466,27 @@ struct FNiagaraDataInterfaceParametersCS_ParticleRead : public FNiagaraDataInter
 			InstanceData->CachedDataSet = SourceDataSet;
 		}
 
-		if (!SourceData->GetGPUIDToIndexTable().Buffer)
-		{
-			// This can happen in the first frame, when there's no previous data yet. The DI shouldn't be
-			// queried in this case, because there's no way to have particle IDs (since there are no particles),
-			// but if it is it will just return failure and default values.
-			SetErrorParams(RHICmdList, ComputeShader, true);
-			return;
-		}
-
 		if (!bReadingOwnEmitter)
 		{
 			FRHIUnorderedAccessView* InputBuffers[] = { SourceData->GetGPUIDToIndexTable().UAV, SourceData->GetGPUBufferFloat().UAV, SourceData->GetGPUBufferInt().UAV };
 			RHICmdList.TransitionResources(EResourceTransitionAccess::EReadable, EResourceTransitionPipeline::EComputeToCompute, InputBuffers, UE_ARRAY_COUNT(InputBuffers));
+
+			if (InstanceCountOffsetParam.IsBound())
+			{
+				// If we're reading the instance count from another emitter, we must insert a barrier on the instance count buffer, to make sure the
+				// previous dispatch finished writing to it. For D3D11, we need to insert an end overlap / begin overlap pair to break up the current
+				// overlap group.
+				RHICmdList.TransitionResource(EResourceTransitionAccess::ERWBarrier, EResourceTransitionPipeline::EComputeToCompute, const_cast<NiagaraEmitterInstanceBatcher*>(Context.Batcher)->GetGPUInstanceCounterManager().GetInstanceCountBuffer().UAV);
+				RHICmdList.EndUAVOverlap();
+				RHICmdList.BeginUAVOverlap();
+			}
 		}
 
 		const uint32 ParticleStrideFloat = SourceData->GetFloatStride() / sizeof(float);
 		const uint32 ParticleStrideInt = SourceData->GetInt32Stride() / sizeof(int32);
 		const uint32 ParticleStrideHalf = SourceData->GetHalfStride() / sizeof(FFloat16);
 
+		SetShaderValue(RHICmdList, ComputeShader, InstanceCountOffsetParam, SourceData->GetGPUInstanceCountBufferOffset());
 		SetSRVParameter(RHICmdList, ComputeShader, IDToIndexTableParam, GetIntSRVWithFallback(SourceData->GetGPUIDToIndexTable()));
 		SetSRVParameter(RHICmdList, ComputeShader, InputFloatBufferParam, GetFloatSRVWithFallback(SourceData->GetGPUBufferFloat()));
 		SetSRVParameter(RHICmdList, ComputeShader, InputIntBufferParam, GetIntSRVWithFallback(SourceData->GetGPUBufferInt()));
@@ -500,7 +502,7 @@ struct FNiagaraDataInterfaceParametersCS_ParticleRead : public FNiagaraDataInter
 private:
 	LAYOUT_FIELD(FShaderParameter, NumSpawnedParticlesParam);
 	LAYOUT_FIELD(FShaderParameter, SpawnedParticlesAcquireTagParam);
-	LAYOUT_FIELD(FShaderParameter, NumParticlesParam);
+	LAYOUT_FIELD(FShaderParameter, InstanceCountOffsetParam);
 	LAYOUT_FIELD(FShaderResourceParameter, SpawnedIDsBufferParam);
 	LAYOUT_FIELD(FShaderResourceParameter, IDToIndexTableParam);
 	LAYOUT_FIELD(FShaderResourceParameter, InputFloatBufferParam);
@@ -1213,7 +1215,7 @@ void UNiagaraDataInterfaceParticleRead::GetParameterDefinitionHLSL(const FNiagar
 	static const TCHAR *FormatDeclarations = TEXT(
 		"int {NumSpawnedParticlesName};\n"
 		"int {SpawnedParticlesAcquireTagName};\n"
-		"int {NumParticlesName};\n"
+		"uint {InstanceCountOffsetName};\n"
 		"uint {ParticleStrideFloatName};\n"
 		"uint {ParticleStrideIntName};\n"
 		"uint {ParticleStrideHalfName};\n"
@@ -1236,7 +1238,7 @@ void UNiagaraDataInterfaceParticleRead::GetParameterDefinitionHLSL(const FNiagar
 	TMap<FString, FStringFormatArg> ArgsDeclarations;
 	ArgsDeclarations.Add(TEXT("NumSpawnedParticlesName"), NumSpawnedParticlesBaseName + ParamInfo.DataInterfaceHLSLSymbol);
 	ArgsDeclarations.Add(TEXT("SpawnedParticlesAcquireTagName"), SpawnedParticlesAcquireTagBaseName + ParamInfo.DataInterfaceHLSLSymbol);
-	ArgsDeclarations.Add(TEXT("NumParticlesName"), NumParticlesBaseName + ParamInfo.DataInterfaceHLSLSymbol);
+	ArgsDeclarations.Add(TEXT("InstanceCountOffsetName"), InstanceCountOffsetBaseName + ParamInfo.DataInterfaceHLSLSymbol);
 	ArgsDeclarations.Add(TEXT("ParticleStrideFloatName"), ParticleStrideFloatBaseName + ParamInfo.DataInterfaceHLSLSymbol);
 	ArgsDeclarations.Add(TEXT("ParticleStrideIntName"), ParticleStrideIntBaseName + ParamInfo.DataInterfaceHLSLSymbol);
 	ArgsDeclarations.Add(TEXT("ParticleStrideHalfName"), ParticleStrideHalfBaseName + ParamInfo.DataInterfaceHLSLSymbol);
@@ -1431,13 +1433,20 @@ bool UNiagaraDataInterfaceParticleRead::GetFunctionHLSL(const FNiagaraDataInterf
 		static const TCHAR* FuncTemplate = TEXT(
 			"void {FunctionName}(out int Out_NumParticles)\n"
 			"{\n"
-			"    Out_NumParticles = {NumParticlesName};\n"
+			"    if({InstanceCountOffsetName} != 0xffffffff)\n"
+			"    {\n"
+			"        Out_NumParticles = RWInstanceCounts[{InstanceCountOffsetName}];\n"
+			"    }\n"
+			"    else\n"
+			"    {\n"
+			"        Out_NumParticles = 0;\n"
+			"    }\n"
 			"}\n\n"
 		);
 
 		TMap<FString, FStringFormatArg> FuncTemplateArgs;
 		FuncTemplateArgs.Add(TEXT("FunctionName"), FunctionInfo.InstanceName);
-		FuncTemplateArgs.Add(TEXT("NumParticlesName"), NumParticlesBaseName + ParamInfo.DataInterfaceHLSLSymbol);
+		FuncTemplateArgs.Add(TEXT("InstanceCountOffsetName"), InstanceCountOffsetBaseName + ParamInfo.DataInterfaceHLSLSymbol);
 
 		OutHLSL += FString::Format(FuncTemplate, FuncTemplateArgs);
 		return true;
diff --git a/Engine/Plugins/FX/Niagara/Source/Niagara/Private/NiagaraDataSet.cpp b/Engine/Plugins/FX/Niagara/Source/Niagara/Private/NiagaraDataSet.cpp
--- a/Engine/Plugins/FX/Niagara/Source/Niagara/Private/NiagaraDataSet.cpp
+++ b/Engine/Plugins/FX/Niagara/Source/Niagara/Private/NiagaraDataSet.cpp
@@ -1079,7 +1079,7 @@ void FNiagaraDataBuffer::Dump(int32 StartIndex, int32 InNumInstances, const FStr
 
 /////////////////////////////////////////////////////////////////////////
 
-void FNiagaraDataBuffer::SetShaderParams(FNiagaraShader *Shader, FRHICommandList &CommandList, bool bInput)
+void FNiagaraDataBuffer::SetShaderParams(FNiagaraShader* Shader, FRHICommandList& CommandList, bool bInput)
 {
 	check(IsInRenderingThread());
 
@@ -1109,17 +1109,14 @@ void FNiagaraDataBuffer::SetShaderParams(FNiagaraShader *Shader, FRHICommandList
 	}
 }
 
-void FNiagaraDataBuffer::UnsetShaderParams(FNiagaraShader *Shader, FRHICommandList &RHICmdList)
+void FNiagaraDataBuffer::UnsetShaderParams(FNiagaraShader* Shader, FRHICommandList& RHICmdList)
 {
 	check(IsInRenderingThread());
 	FRHIComputeShader* ShaderRHI = RHICmdList.GetBoundComputeShader();
 
 	if (Shader->FloatOutputBufferParam.IsUAVBound())
 	{
-#if !PLATFORM_PS4
 		Shader->FloatOutputBufferParam.UnsetUAV(RHICmdList, ShaderRHI);
-#endif
-		//RHICmdList.TransitionResource(EResourceTransitionAccess::EReadable, EResourceTransitionPipeline::EGfxToGfx, CurrDataRender().GetGPUBufferFloat()->UAV);
 	}
 
 	if (Shader->HalfOutputBufferParam.IsUAVBound())
@@ -1132,17 +1129,12 @@ void FNiagaraDataBuffer::UnsetShaderParams(FNiagaraShader *Shader, FRHICommandLi
 
 	if (Shader->IntOutputBufferParam.IsUAVBound())
 	{
-#if !PLATFORM_PS4
 		Shader->IntOutputBufferParam.UnsetUAV(RHICmdList, ShaderRHI);
-#endif
-		//RHICmdList.TransitionResource(EResourceTransitionAccess::EReadable, EResourceTransitionPipeline::EGfxToGfx, CurrDataRender().GetGPUBufferInt()->UAV);
 	}
 
 	if (Shader->IDToIndexBufferParam.IsUAVBound())
 	{
-#if !PLATFORM_PS4
 		Shader->IDToIndexBufferParam.UnsetUAV(RHICmdList, RHICmdList.GetBoundComputeShader());
-#endif
 	}
 }
 
diff --git a/Engine/Plugins/FX/Niagara/Source/Niagara/Private/NiagaraEmitterInstanceBatcher.cpp b/Engine/Plugins/FX/Niagara/Source/Niagara/Private/NiagaraEmitterInstanceBatcher.cpp
--- a/Engine/Plugins/FX/Niagara/Source/Niagara/Private/NiagaraEmitterInstanceBatcher.cpp
+++ b/Engine/Plugins/FX/Niagara/Source/Niagara/Private/NiagaraEmitterInstanceBatcher.cpp
@@ -1354,16 +1354,16 @@ void NiagaraEmitterInstanceBatcher::Run(const FNiagaraGPUSystemTick& Tick, const
 	//
 	if (Shader->InstanceCountsParam.IsBound())
 	{
+		RHICmdList.TransitionResource(EResourceTransitionAccess::ERWNoBarrier, EResourceTransitionPipeline::EComputeToCompute, GPUInstanceCounterManager.GetInstanceCountBuffer().UAV);
+		Shader->InstanceCountsParam.SetBuffer(RHICmdList, ComputeShader, GPUInstanceCounterManager.GetInstanceCountBuffer());
+
 		if (IterationInterface)
 		{
-			RHICmdList.SetUAVParameter(ComputeShader, Shader->InstanceCountsParam.GetUAVIndex(), GetEmptyRWBufferFromPool(RHICmdList, PF_R32_UINT));
 			SetShaderValue(RHICmdList, ComputeShader, Shader->ReadInstanceCountOffsetParam, -1);
 			SetShaderValue(RHICmdList, ComputeShader, Shader->WriteInstanceCountOffsetParam, -1);
 		}
 		else
 		{
-			RHICmdList.TransitionResource(EResourceTransitionAccess::ERWNoBarrier, EResourceTransitionPipeline::EComputeToCompute, GPUInstanceCounterManager.GetInstanceCountBuffer().UAV);
-			Shader->InstanceCountsParam.SetBuffer(RHICmdList, ComputeShader, GPUInstanceCounterManager.GetInstanceCountBuffer());
 			const uint32 ReadOffset = (Tick.bNeedsReset && SimulationStageIndex == 0) ? INDEX_NONE : Instance->SimStageData[SimulationStageIndex].SourceCountOffset;
 			const uint32 WriteOffset = Instance->SimStageData[SimulationStageIndex].DestinationCountOffset;
 			//UE_LOG(LogNiagara, Log, TEXT("Instance count setup R: %d W: %d reset? %s %d"), ReadOffset, WriteOffset, Tick.bNeedsReset ? TEXT("T") : TEXT("F"), CurrentData.GetGPUInstanceCountBufferOffset());
@@ -1422,16 +1422,6 @@ void NiagaraEmitterInstanceBatcher::Run(const FNiagaraGPUSystemTick& Tick, const
 		SetConstantBuffer(RHICmdList, ComputeShader, Shader->ExternalConstantBufferParam[1], Context->ExternalCBufferLayout, Instance->ExternalParamData + Context->ExternalCBufferLayout.ConstantBufferSize);
 	}
 
-	// setup script parameters
-
-	// #todo(dmp): temporary hack -- unbind UAVs if we have a valid iteration DI.  This way, when we are outputting with a different iteration count, we don't
-	// mess up particle state
-	if (IterationInterface)
-	{
-		CurrentData.UnsetShaderParams(Shader.GetShader(), RHICmdList);
-		DestinationData.UnsetShaderParams(Shader.GetShader(), RHICmdList);
-	}
-
 	//UE_LOG(LogNiagara, Log, TEXT("Num Instance : %d | Num Group : %d | Spawned Istance : %d | Start Instance : %d | Num Indices : %d | Stage Index : %d"), 
 		//TotalNumInstances, NumThreadGroups, InstancesToSpawnThisFrame, UpdateStartInstance, Context->NumIndicesPerInstance, SimulationStageIndex);
 
diff --git a/Engine/Plugins/FX/Niagara/Source/NiagaraCore/Private/NiagaraCustomVersion.cpp b/Engine/Plugins/FX/Niagara/Source/NiagaraCore/Private/NiagaraCustomVersion.cpp
--- a/Engine/Plugins/FX/Niagara/Source/NiagaraCore/Private/NiagaraCustomVersion.cpp
+++ b/Engine/Plugins/FX/Niagara/Source/NiagaraCore/Private/NiagaraCustomVersion.cpp
@@ -8,4 +8,4 @@ const FGuid FNiagaraCustomVersion::GUID(0xFCF57AFA, 0x50764283, 0xB9A9E658, 0xFF
 // Register the custom version with core
 FCustomVersionRegistration GRegisterNiagaraCustomVersion(FNiagaraCustomVersion::GUID, FNiagaraCustomVersion::LatestVersion, TEXT("NiagaraVer"));
 
-const FGuid FNiagaraCustomVersion::LatestScriptCompileVersion(0x25C5FF75, 0x14574F55, 0xA8A70C22, 0x98D1C114);
+const FGuid FNiagaraCustomVersion::LatestScriptCompileVersion(0x3DA843BA, 0x016647B5, 0xBD440BD8, 0x1B55EBCB);
