Niagara - Fix gpu only functions not working in new shader stages.

diff --git a/Engine/Plugins/FX/Niagara/Source/Niagara/Private/NiagaraScript.cpp b/Engine/Plugins/FX/Niagara/Source/Niagara/Private/NiagaraScript.cpp
--- a/Engine/Plugins/FX/Niagara/Source/Niagara/Private/NiagaraScript.cpp
+++ b/Engine/Plugins/FX/Niagara/Source/Niagara/Private/NiagaraScript.cpp
@@ -342,10 +342,11 @@ void UNiagaraScript::ComputeVMCompilationId(FNiagaraVMExecutableDataId& Id) cons
 	
 	ENiagaraSimTarget SimTargetToBuild = ENiagaraSimTarget::CPUSim;
 	// Ideally we wouldn't want to do this but rather than push the data down
-	// from the emitter.
-	UObject* Obj = GetOuter();
-	if (UNiagaraEmitter* Emitter = Cast<UNiagaraEmitter>(Obj))
+	// from the emitter.  Checking all outers here to pick up simulation stages too.
+	UNiagaraEmitter* OuterEmitter = GetTypedOuter<UNiagaraEmitter>();
+	if (OuterEmitter != nullptr)
 	{
+		UNiagaraEmitter* Emitter = OuterEmitter;
 		UNiagaraSystem* EmitterOwner = Cast<UNiagaraSystem>(Emitter->GetOuter());
 		if (EmitterOwner && EmitterOwner->bBakeOutRapidIteration)
 		{
@@ -422,6 +423,7 @@ void UNiagaraScript::ComputeVMCompilationId(FNiagaraVMExecutableDataId& Id) cons
 		Id.AdditionalDefines.Add(Emitter->GetFullName());
 	}
 
+	UObject* Obj = GetOuter();
 	if (UNiagaraSystem* System = Cast<UNiagaraSystem>(Obj))
 	{
 		if (System->bBakeOutRapidIteration)
