Fix for coalesced binding info, which was generating broken copy offsets for GPU compute scripts.  Also cleared out the values in non-GPU scripts as they are unneeded

diff --git a/Engine/Plugins/FX/Niagara/Source/Niagara/Private/NiagaraScript.cpp b/Engine/Plugins/FX/Niagara/Source/Niagara/Private/NiagaraScript.cpp
--- a/Engine/Plugins/FX/Niagara/Source/Niagara/Private/NiagaraScript.cpp
+++ b/Engine/Plugins/FX/Niagara/Source/Niagara/Private/NiagaraScript.cpp
@@ -852,13 +852,17 @@ void UNiagaraScript::Serialize(FArchive& Ar)
 		if (Usage != ENiagaraScriptUsage::ParticleGPUComputeScript)
 		{
 			ExecutableData.BakeScriptLiterals(ExecutableData.ScriptLiterals);
+
+			// we only need the padding info for when we're dealing with GPU scripts (for
+			// FNiagaraScriptInstanceParameterStore::CopyParameterDataToPaddedBuffer())
+			ScriptExecutionParamStore.PaddingInfo.Empty();
 		}
 		else
 		{
 			ExecutableData.ScriptLiterals.Empty();
+			ScriptExecutionParamStore.CoalescePaddingInfo();
 		}
 
-		ScriptExecutionParamStore.CoalescePaddingInfo();
 	}
 #endif
 
diff --git a/Engine/Plugins/FX/Niagara/Source/Niagara/Private/NiagaraScriptExecutionParameterStore.cpp b/Engine/Plugins/FX/Niagara/Source/Niagara/Private/NiagaraScriptExecutionParameterStore.cpp
--- a/Engine/Plugins/FX/Niagara/Source/Niagara/Private/NiagaraScriptExecutionParameterStore.cpp
+++ b/Engine/Plugins/FX/Niagara/Source/Niagara/Private/NiagaraScriptExecutionParameterStore.cpp
@@ -162,7 +162,7 @@ void FNiagaraScriptExecutionParameterStore::CoalescePaddingInfo()
 		const FNiagaraScriptExecutionPaddingInfo& CurrentEntry = PaddingInfo[PaddingIt];
 
 		if (((PreviousEntry.SrcOffset + PreviousEntry.SrcSize) == CurrentEntry.SrcOffset)
-			&& ((PreviousEntry.DestOffset + PreviousEntry.DestSize) == CurrentEntry.DestOffset)
+			&& ((PreviousEntry.DestOffset + PreviousEntry.SrcSize) == CurrentEntry.DestOffset)
 			&& ((TNumericLimits<uint16>::Max() - PreviousEntry.SrcSize) >= CurrentEntry.SrcSize))
 		{
 			PreviousEntry.SrcSize += CurrentEntry.SrcSize;
