Niagara: Disabled namespace refactor precompile name matching by default, now controlled by cvar fx.NiagaraEnablePrecompilerNamespaceFixup.
Bumped niagara compiler version.
Removed old gNiagaraEnableParameterPanel2 cvar.

diff --git a/Engine/Plugins/FX/Niagara/Source/Niagara/Classes/NiagaraSystem.h b/Engine/Plugins/FX/Niagara/Source/Niagara/Classes/NiagaraSystem.h
--- a/Engine/Plugins/FX/Niagara/Source/Niagara/Classes/NiagaraSystem.h
+++ b/Engine/Plugins/FX/Niagara/Source/Niagara/Classes/NiagaraSystem.h
@@ -342,15 +342,6 @@ public:
 	UPROPERTY(EditAnywhere, AdvancedDisplay, Category = "Emitter")
 	uint32 bBakeOutRapidIteration : 1;
 
-	/** Toggles whether or not emitters within this system will try and compress their particle attributes. 
-	    In some cases, this precision change can lead to perceivable differences, but memory costs and or performance (especially true for GPU emitters) can improve. */
-	UPROPERTY(EditAnywhere, AdvancedDisplay, Category = "Emitter")
-	uint32 bCompressAttributes : 1;
-
-	/** Experimental feature to depromote parameters in the dataset that are not accessed during the precompile process. */
-	UPROPERTY(EditAnywhere, AdvancedDisplay, Category = "Emitter")
-	uint32 bCullDatasetParameters : 1;
-
 #endif
 
 	FORCEINLINE UNiagaraParameterCollectionInstance* GetParameterCollectionOverride(UNiagaraParameterCollection* Collection)
diff --git a/Engine/Plugins/FX/Niagara/Source/NiagaraCore/Private/NiagaraCustomVersion.cpp b/Engine/Plugins/FX/Niagara/Source/NiagaraCore/Private/NiagaraCustomVersion.cpp
--- a/Engine/Plugins/FX/Niagara/Source/NiagaraCore/Private/NiagaraCustomVersion.cpp
+++ b/Engine/Plugins/FX/Niagara/Source/NiagaraCore/Private/NiagaraCustomVersion.cpp
@@ -8,4 +8,4 @@ const FGuid FNiagaraCustomVersion::GUID(0xFCF57AFA, 0x50764283, 0xB9A9E658, 0xFF
 // Register the custom version with core
 FCustomVersionRegistration GRegisterNiagaraCustomVersion(FNiagaraCustomVersion::GUID, FNiagaraCustomVersion::LatestVersion, TEXT("NiagaraVer"));
 
-const FGuid FNiagaraCustomVersion::LatestScriptCompileVersion(0x3DA843BA, 0x016647B5, 0xBD440BD8, 0x1B55EBCB);
+const FGuid FNiagaraCustomVersion::LatestScriptCompileVersion(0x76B3E9E4, 0x883C48CE, 0x9508A68C, 0xD4BA993D );
diff --git a/Engine/Plugins/FX/Niagara/Source/NiagaraEditor/Private/NiagaraCompiler.cpp b/Engine/Plugins/FX/Niagara/Source/NiagaraEditor/Private/NiagaraCompiler.cpp
--- a/Engine/Plugins/FX/Niagara/Source/NiagaraEditor/Private/NiagaraCompiler.cpp
+++ b/Engine/Plugins/FX/Niagara/Source/NiagaraEditor/Private/NiagaraCompiler.cpp
@@ -72,6 +72,15 @@ static FAutoConsoleVariableRef CVarForceNiagaraVMBinaryDump(
 	ECVF_Default
 );
 
+static int32 GNiagaraEnablePrecompilerNamespaceFixup = 0;
+static FAutoConsoleVariableRef CVarNiagaraEnablePrecompilerNamespaceFixup(
+	TEXT("fx.NiagaraEnablePrecompilerNamespaceFixup"),
+	GNiagaraEnablePrecompilerNamespaceFixup,
+	TEXT("Enable a precompiler stage to discover parameter name matches and convert matched parameter hlsl name tokens to appropriate namespaces. \n"),
+	ECVF_Default
+);
+
+
 static FCriticalSection TranslationCritSec;
 
 void DumpHLSLText(const FString& SourceCode, const FString& DebugName)
@@ -638,87 +647,91 @@ TSharedPtr<FNiagaraCompileRequestDataBase, ESPMode::ThreadSafe> FNiagaraEditorMo
 			}
 		}
 
-		// Build parameter map histories just for System scripts in order to separate out System scope parameters from Emitter scope for the history handles.
-		TArray<FNiagaraParameterMapHistory> SystemOnlyHistories;
-		FNiagaraParameterMapHistoryWithMetaDataBuilder Builder;
-		Builder.ConstantResolver = EmptyResolver;
-		Builder.bShouldBuildSubHistories = false;
-		FString TranslationName = TEXT("Emitter");
-		Builder.AddGraphToCallingGraphContextStack(BasePtr->GetPrecomputedNodeGraph());
-		Builder.BeginTranslation(TranslationName);
-
-		TArray<UNiagaraNodeOutput*> SystemOutputNodes;
-		BasePtr->NodeGraphDeepCopy->FindOutputNodes(SystemOutputNodes);
-		for (UNiagaraNodeOutput* FoundOutputNode : SystemOutputNodes)
+		// Namespace fixup stage: Name match parameters across scripts and promote their namespaces to the Dataset if matched. 
+		if (GNiagaraEnablePrecompilerNamespaceFixup > 0)
 		{
-			Builder.EnableScriptWhitelist(true, FoundOutputNode->GetUsage());
-			Builder.BuildParameterMaps(FoundOutputNode, true);
-			
-			for (FNiagaraParameterMapHistory& History : Builder.Histories)
+			// Build parameter map histories just for System scripts in order to separate out System scope parameters from Emitter scope for the history handles.
+			TArray<FNiagaraParameterMapHistory> SystemOnlyHistories;
+			FNiagaraParameterMapHistoryWithMetaDataBuilder Builder;
+			Builder.ConstantResolver = EmptyResolver;
+			Builder.bShouldBuildSubHistories = false;
+			FString TranslationName = TEXT("Emitter");
+			Builder.AddGraphToCallingGraphContextStack(BasePtr->GetPrecomputedNodeGraph());
+			Builder.BeginTranslation(TranslationName);
+
+			TArray<UNiagaraNodeOutput*> SystemOutputNodes;
+			BasePtr->NodeGraphDeepCopy->FindOutputNodes(SystemOutputNodes);
+			for (UNiagaraNodeOutput* FoundOutputNode : SystemOutputNodes)
 			{
-				// FinishPrecompile() has not been called on this history so set the OriginatingScriptUsage.
-				History.OriginatingScriptUsage = FoundOutputNode->GetUsage();
-				SystemOnlyHistories.Add(History);
+				Builder.EnableScriptWhitelist(true, FoundOutputNode->GetUsage());
+				Builder.BuildParameterMaps(FoundOutputNode, true);
+
+				for (FNiagaraParameterMapHistory& History : Builder.Histories)
+				{
+					// FinishPrecompile() has not been called on this history so set the OriginatingScriptUsage.
+					History.OriginatingScriptUsage = FoundOutputNode->GetUsage();
+					SystemOnlyHistories.Add(History);
+				}
+				Builder.Histories.Empty();
 			}
-			Builder.Histories.Empty();
-		}
 
-		// Build history handles for System scripts.
-		TArray<FNiagaraParameterMapHistory>& SystemScriptHistories = BasePtr->GetPrecomputedHistories();
-		TArray<FNiagaraParameterMapHistoryHandle> SystemScriptHistoryHandles;
-		for (int i = 0; i < SystemScriptHistories.Num(); ++i)
-		{
-			SystemScriptHistoryHandles.Add(FNiagaraParameterMapHistoryHandle(SystemScriptHistories[i], SystemOnlyHistories[i], BasePtr->RequiredRendererVariables));
-		}
+			// Build history handles for System scripts.
+			TArray<FNiagaraParameterMapHistory>& SystemScriptHistories = BasePtr->GetPrecomputedHistories();
+			TArray<FNiagaraParameterMapHistoryHandle> SystemScriptHistoryHandles;
+			for (int i = 0; i < SystemScriptHistories.Num(); ++i)
+			{
+				SystemScriptHistoryHandles.Add(FNiagaraParameterMapHistoryHandle(SystemScriptHistories[i], SystemOnlyHistories[i], BasePtr->RequiredRendererVariables));
+			}
 
-		// Build history handles for Emitter and Particle scripts.
-		TArray<FNiagaraParameterMapHistoryHandle> EmitterAndParticleScriptHistoryHandles;
-		for (TSharedPtr<FNiagaraCompileRequestData, ESPMode::ThreadSafe> EmitterRequest : BasePtr->EmitterData)
-		{
-			FName EmitterUniqueName = FName(*EmitterRequest->GetUniqueEmitterName());
-			TArray<FNiagaraParameterMapHistory>& EmitterScriptHistories = EmitterRequest->GetPrecomputedHistories();
-			for (int i = 0; i < EmitterScriptHistories.Num(); ++i)
+			// Build history handles for Emitter and Particle scripts.
+			TArray<FNiagaraParameterMapHistoryHandle> EmitterAndParticleScriptHistoryHandles;
+			for (TSharedPtr<FNiagaraCompileRequestData, ESPMode::ThreadSafe> EmitterRequest : BasePtr->EmitterData)
 			{
-				if (FNiagaraStackGraphUtilities::GetScopeForScriptUsage(EmitterScriptHistories[i].OriginatingScriptUsage) == ENiagaraParameterScope::Particles)
+				FName EmitterUniqueName = FName(*EmitterRequest->GetUniqueEmitterName());
+				TArray<FNiagaraParameterMapHistory>& EmitterScriptHistories = EmitterRequest->GetPrecomputedHistories();
+				for (int i = 0; i < EmitterScriptHistories.Num(); ++i)
 				{
-					EmitterAndParticleScriptHistoryHandles.Add(FNiagaraParameterMapHistoryHandle(EmitterScriptHistories[i], EmitterRequest->RequiredRendererVariables, EmitterUniqueName));
-				}
-				else
-				{
-					FNiagaraParameterMapHistory* SameUsageSystemHistory = SystemScriptHistories.FindByPredicate([&EmitterScriptHistories, &i](const FNiagaraParameterMapHistory& SystemScriptHistory){
-						ENiagaraScriptUsage EmitterOriginatingScriptUsage = EmitterScriptHistories[i].OriginatingScriptUsage;
-						if (EmitterOriginatingScriptUsage == ENiagaraScriptUsage::EmitterSpawnScript)
-						{
-							return SystemScriptHistory.OriginatingScriptUsage == ENiagaraScriptUsage::SystemSpawnScript;
-						}
-						else if (EmitterOriginatingScriptUsage == ENiagaraScriptUsage::EmitterUpdateScript)
-						{
-							return SystemScriptHistory.OriginatingScriptUsage == ENiagaraScriptUsage::SystemUpdateScript;
-						}
-						else
+					if (FNiagaraStackGraphUtilities::GetScopeForScriptUsage(EmitterScriptHistories[i].OriginatingScriptUsage) == ENiagaraParameterScope::Particles)
+					{
+						EmitterAndParticleScriptHistoryHandles.Add(FNiagaraParameterMapHistoryHandle(EmitterScriptHistories[i], EmitterRequest->RequiredRendererVariables, EmitterUniqueName));
+					}
+					else
+					{
+						FNiagaraParameterMapHistory* SameUsageSystemHistory = SystemScriptHistories.FindByPredicate([&EmitterScriptHistories, &i](const FNiagaraParameterMapHistory& SystemScriptHistory) {
+							ENiagaraScriptUsage EmitterOriginatingScriptUsage = EmitterScriptHistories[i].OriginatingScriptUsage;
+							if (EmitterOriginatingScriptUsage == ENiagaraScriptUsage::EmitterSpawnScript)
+							{
+								return SystemScriptHistory.OriginatingScriptUsage == ENiagaraScriptUsage::SystemSpawnScript;
+							}
+							else if (EmitterOriginatingScriptUsage == ENiagaraScriptUsage::EmitterUpdateScript)
+							{
+								return SystemScriptHistory.OriginatingScriptUsage == ENiagaraScriptUsage::SystemUpdateScript;
+							}
+							else
+							{
+								checkf(false, TEXT("Encountered unexpected script usage for history!"));
+							}
+							return false;
+							});
+
+						if (SameUsageSystemHistory == nullptr)
 						{
-							checkf(false, TEXT("Encountered unexpected script usage for history!"));
+							UE_LOG(LogNiagaraEditor, Error, TEXT("Precompile failed for system: %s.  Failed to generate valid parameter history for emitter: %s."),
+								*System->GetPathName(), *EmitterRequest->EmitterUniqueName);
+							return TSharedPtr<FNiagaraCompileRequestDataBase, ESPMode::ThreadSafe>();
 						}
-						return false;
-					});
 
-					if (SameUsageSystemHistory == nullptr)
-					{
-						UE_LOG(LogNiagaraEditor, Error, TEXT("Precompile failed for system: %s.  Failed to generate valid parameter history for emitter: %s."),
-							*System->GetPathName(), *EmitterRequest->EmitterUniqueName);
-						return TSharedPtr<FNiagaraCompileRequestDataBase, ESPMode::ThreadSafe>();
+						EmitterAndParticleScriptHistoryHandles.Add(FNiagaraParameterMapHistoryHandle(*SameUsageSystemHistory, EmitterScriptHistories[i], EmitterRequest->RequiredRendererVariables, EmitterUniqueName));
 					}
-
-					EmitterAndParticleScriptHistoryHandles.Add(FNiagaraParameterMapHistoryHandle(*SameUsageSystemHistory, EmitterScriptHistories[i], EmitterRequest->RequiredRendererVariables, EmitterUniqueName));
 				}
 			}
-		}
 
-		// Fixup namespaces of parameters that will be used in the histories for translation.
-		TArray<FNiagaraParameterMapHistoryHandle> AllHistoryHandles;
-		AllHistoryHandles.Append(SystemScriptHistoryHandles);
-		AllHistoryHandles.Append(EmitterAndParticleScriptHistoryHandles);
-		FNiagaraParameterMapHistoryWithMetaDataBuilder::FixupHistoryVariableNamespaces(AllHistoryHandles, System->bCullDatasetParameters);
+			// Fixup namespaces of parameters that will be used in the histories for translation.
+			TArray<FNiagaraParameterMapHistoryHandle> AllHistoryHandles;
+			AllHistoryHandles.Append(SystemScriptHistoryHandles);
+			AllHistoryHandles.Append(EmitterAndParticleScriptHistoryHandles);
+			FNiagaraParameterMapHistoryWithMetaDataBuilder::FixupHistoryVariableNamespaces(AllHistoryHandles, 0);
+		}
 	}
 
 	UE_LOG(LogNiagaraEditor, Verbose, TEXT("'%s' Precompile took %f sec."), *InObj->GetOutermost()->GetName(),
diff --git a/Engine/Plugins/FX/Niagara/Source/NiagaraEditor/Private/NiagaraEditorCommon.cpp b/Engine/Plugins/FX/Niagara/Source/NiagaraEditor/Private/NiagaraEditorCommon.cpp
--- a/Engine/Plugins/FX/Niagara/Source/NiagaraEditor/Private/NiagaraEditorCommon.cpp
+++ b/Engine/Plugins/FX/Niagara/Source/NiagaraEditor/Private/NiagaraEditorCommon.cpp
@@ -16,13 +16,6 @@ DEFINE_LOG_CATEGORY(LogNiagaraEditor);
 TMap<FName, int32> FNiagaraOpInfo::OpInfoMap;
 TArray<FNiagaraOpInfo> FNiagaraOpInfo::OpInfos;
 
-int32 FNiagaraEditorCommonCVar::GNiagaraEnableParameterPanel2 = 1;
-FAutoConsoleVariableRef FNiagaraEditorCommonCVar::CVarEnableParameterPanel2(
-	TEXT("fx.NiagaraEnableParameterPanel2"),
-	FNiagaraEditorCommonCVar::GNiagaraEnableParameterPanel2,
-	TEXT("If true, enable enhanced parameter panel features and dynamic parameter linking in precompiler."),
-	ECVF_Default);
-
 const FNiagaraOpInfo* FNiagaraOpInfo::GetOpInfo(FName OpName)
 {
 	const int32* Idx = OpInfoMap.Find(OpName);
diff --git a/Engine/Plugins/FX/Niagara/Source/NiagaraEditor/Private/NiagaraParameterMapHistory.cpp b/Engine/Plugins/FX/Niagara/Source/NiagaraEditor/Private/NiagaraParameterMapHistory.cpp
--- a/Engine/Plugins/FX/Niagara/Source/NiagaraEditor/Private/NiagaraParameterMapHistory.cpp
+++ b/Engine/Plugins/FX/Niagara/Source/NiagaraEditor/Private/NiagaraParameterMapHistory.cpp
@@ -37,9 +37,9 @@ static FAutoConsoleVariableRef CVarNiagaraLogNamespaceFixup(
 
 static int32 GNiagaraForcePrecompilerCullDataset = 0;
 static FAutoConsoleVariableRef CVarNiagaraForcePrecompilerCullDataset(
-	TEXT("fx.NiagaraForcePrecompilerCullDataset"),
+	TEXT("fx.NiagaraEnablePrecompilerNamespaceDatasetCulling"),
 	GNiagaraForcePrecompilerCullDataset,
-	TEXT("Force the precompiler process to cull unused Dataset parameters. \n"),
+	TEXT("Force the namespace fixup precompiler process to cull unused Dataset parameters. Only enabled if fx.NiagaraEnablePrecompilerNamespaceFixup is also enabled. \n"),
 	ECVF_Default
 );
 
@@ -1444,11 +1444,6 @@ bool FCompileConstantResolver::ResolveConstant(FNiagaraVariable& OutConstant) co
 
 void FNiagaraParameterMapHistoryWithMetaDataBuilder::FixupHistoryVariableNamespaces(TArray<FNiagaraParameterMapHistoryHandle>& InHistoryHandles, bool bCullDatasetOutputs)
 {
-	if (FNiagaraEditorCommonCVar::GNiagaraEnableParameterPanel2 == 0)
-	{
-		return;
-	}
-
 	struct FMatchedVariableInfo
 	{
 	public:
diff --git a/Engine/Plugins/FX/Niagara/Source/NiagaraEditor/Private/TNiagaraGraphPinEditableName.h b/Engine/Plugins/FX/Niagara/Source/NiagaraEditor/Private/TNiagaraGraphPinEditableName.h
--- a/Engine/Plugins/FX/Niagara/Source/NiagaraEditor/Private/TNiagaraGraphPinEditableName.h
+++ b/Engine/Plugins/FX/Niagara/Source/NiagaraEditor/Private/TNiagaraGraphPinEditableName.h
@@ -111,20 +111,13 @@ protected:
 		if (ParentNode && ParentNode->IsPinNameEditable(this->GraphPinObj))
 		{
 			UNiagaraGraph* NiagaraGraph = ParentNode->GetNiagaraGraph();
-			if (FNiagaraEditorCommonCVar::GNiagaraEnableParameterPanel2 > 0)
+			if (ParentNode->IsA<UNiagaraNodeParameterMapBase>())
 			{
-				if (ParentNode->IsA<UNiagaraNodeParameterMapBase>())
+				if (NiagaraGraph->IsPinVisualWidgetProviderRegistered())
 				{
-					if (NiagaraGraph->IsPinVisualWidgetProviderRegistered())
-					{
-						return NiagaraGraph->GetPinVisualWidget(this->GraphPinObj);
-					}
-					return CreateLabelTextBlock();
-				}
-				else 
-				{
-					return CreateRenamableLabelTextBlock();
+					return NiagaraGraph->GetPinVisualWidget(this->GraphPinObj);
 				}
+				return CreateLabelTextBlock();
 			}
 			else
 			{
@@ -132,8 +125,7 @@ protected:
 			}
 		}
 		else
-		{
-			
+		{	
 			return BaseClass::GetLabelWidget(InLabelStyle);
 		}
 	}
diff --git a/Engine/Plugins/FX/Niagara/Source/NiagaraEditor/Private/Toolkits/NiagaraScriptToolkit.cpp b/Engine/Plugins/FX/Niagara/Source/NiagaraEditor/Private/Toolkits/NiagaraScriptToolkit.cpp
--- a/Engine/Plugins/FX/Niagara/Source/NiagaraEditor/Private/Toolkits/NiagaraScriptToolkit.cpp
+++ b/Engine/Plugins/FX/Niagara/Source/NiagaraEditor/Private/Toolkits/NiagaraScriptToolkit.cpp
@@ -101,12 +101,9 @@ void FNiagaraScriptToolkit::RegisterTabSpawners(const TSharedRef<class FTabManag
 		.SetDisplayName(LOCTEXT("ParametersTab", "Parameters"))
 		.SetGroup(WorkspaceMenuCategoryRef);
 
-	if (FNiagaraEditorCommonCVar::GNiagaraEnableParameterPanel2 > 0)
-	{
-		InTabManager->RegisterTabSpawner(ParametersTabId2, FOnSpawnTab::CreateSP(this, &FNiagaraScriptToolkit::SpawnTabScriptParameters2))
-			.SetDisplayName(LOCTEXT("ParametersTab2", "Parameters2"))
-			.SetGroup(WorkspaceMenuCategoryRef);
-	}
+	InTabManager->RegisterTabSpawner(ParametersTabId2, FOnSpawnTab::CreateSP(this, &FNiagaraScriptToolkit::SpawnTabScriptParameters2))
+		.SetDisplayName(LOCTEXT("ParametersTab2", "Parameters2"))
+		.SetGroup(WorkspaceMenuCategoryRef);
 		
 	InTabManager->RegisterTabSpawner(StatsTabId, FOnSpawnTab::CreateSP(this, &FNiagaraScriptToolkit::SpawnTabStats))
 		.SetDisplayName(LOCTEXT("StatsTab", "Stats"))
diff --git a/Engine/Plugins/FX/Niagara/Source/NiagaraEditor/Private/Toolkits/NiagaraSystemToolkit.cpp b/Engine/Plugins/FX/Niagara/Source/NiagaraEditor/Private/Toolkits/NiagaraSystemToolkit.cpp
--- a/Engine/Plugins/FX/Niagara/Source/NiagaraEditor/Private/Toolkits/NiagaraSystemToolkit.cpp
+++ b/Engine/Plugins/FX/Niagara/Source/NiagaraEditor/Private/Toolkits/NiagaraSystemToolkit.cpp
@@ -125,12 +125,9 @@ void FNiagaraSystemToolkit::RegisterTabSpawners(const TSharedRef<class FTabManag
 		.SetDisplayName(LOCTEXT("SystemParameters", "Parameters"))
 		.SetGroup(WorkspaceMenuCategory.ToSharedRef());
 
-	if (FNiagaraEditorCommonCVar::GNiagaraEnableParameterPanel2 > 0)
-	{
-		InTabManager->RegisterTabSpawner(SystemParametersTabID2, FOnSpawnTab::CreateSP(this, &FNiagaraSystemToolkit::SpawnTab_SystemParameters2))
-			.SetDisplayName(LOCTEXT("SystemParameters2", "Parameters2"))
-			.SetGroup(WorkspaceMenuCategory.ToSharedRef());
-	}
+	InTabManager->RegisterTabSpawner(SystemParametersTabID2, FOnSpawnTab::CreateSP(this, &FNiagaraSystemToolkit::SpawnTab_SystemParameters2))
+		.SetDisplayName(LOCTEXT("SystemParameters2", "Parameters2"))
+		.SetGroup(WorkspaceMenuCategory.ToSharedRef());
 
 	InTabManager->RegisterTabSpawner(SelectedEmitterStackTabID, FOnSpawnTab::CreateSP(this, &FNiagaraSystemToolkit::SpawnTab_SelectedEmitterStack))
 		.SetDisplayName(LOCTEXT("SelectedEmitterStacks", "Selected Emitters"))
diff --git a/Engine/Plugins/FX/Niagara/Source/NiagaraEditor/Public/NiagaraEditorCommon.h b/Engine/Plugins/FX/Niagara/Source/NiagaraEditor/Public/NiagaraEditorCommon.h
--- a/Engine/Plugins/FX/Niagara/Source/NiagaraEditor/Public/NiagaraEditorCommon.h
+++ b/Engine/Plugins/FX/Niagara/Source/NiagaraEditor/Public/NiagaraEditorCommon.h
@@ -216,10 +216,3 @@ namespace NiagaraParameterPanelSectionID
 	ENiagaraScriptParameterUsage GetUsageForNewParametersInSection(NiagaraParameterPanelSectionID::Type InSection);
 
 }
-
-class FNiagaraEditorCommonCVar
-{
-public:
-	static int32 GNiagaraEnableParameterPanel2;
-	static FAutoConsoleVariableRef CVarEnableParameterPanel2;
-};
