Fix bad merge to NiagaraSystem.h.

diff --git a/Engine/Plugins/FX/Niagara/Source/Niagara/Classes/NiagaraSystem.h b/Engine/Plugins/FX/Niagara/Source/Niagara/Classes/NiagaraSystem.h
--- a/Engine/Plugins/FX/Niagara/Source/Niagara/Classes/NiagaraSystem.h
+++ b/Engine/Plugins/FX/Niagara/Source/Niagara/Classes/NiagaraSystem.h
@@ -342,6 +342,11 @@ public:
 	UPROPERTY(EditAnywhere, AdvancedDisplay, Category = "Emitter")
 	uint32 bBakeOutRapidIteration : 1;
 
+	/** Toggles whether or not emitters within this system will try and compress their particle attributes.
+	In some cases, this precision change can lead to perceivable differences, but memory costs and or performance (especially true for GPU emitters) can improve. */
+	UPROPERTY(EditAnywhere, AdvancedDisplay, Category = "Emitter")
+	uint32 bCompressAttributes : 1;
+
 #endif
 
 	FORCEINLINE UNiagaraParameterCollectionInstance* GetParameterCollectionOverride(UNiagaraParameterCollection* Collection)
