RL - WHen AutomationTest are running than use synchronous Niagara Compiles, because they are accessed very late via Blueprint and having them bail until later can cause timing differences.

diff --git a/Engine/Plugins/FX/Niagara/Source/Niagara/Private/NiagaraScript.cpp b/Engine/Plugins/FX/Niagara/Source/Niagara/Private/NiagaraScript.cpp
--- a/Engine/Plugins/FX/Niagara/Source/Niagara/Private/NiagaraScript.cpp
+++ b/Engine/Plugins/FX/Niagara/Source/Niagara/Private/NiagaraScript.cpp
@@ -1849,7 +1849,7 @@ void UNiagaraScript::CacheShadersForResources(EShaderPlatform ShaderPlatform, FN
 	if (CanBeRunOnGpu())
 	{
 		// When not running in the editor, the shaders are created in-sync (in the postload) to avoid update issues.
-		const bool bSuccess = ResourceToCache->CacheShaders(ShaderPlatform, bApplyCompletedShaderMapForRendering, bForceRecompile, bCooking || !GIsEditor);
+		const bool bSuccess = ResourceToCache->CacheShaders(ShaderPlatform, bApplyCompletedShaderMapForRendering, bForceRecompile, bCooking || !GIsEditor || GIsAutomationTesting);
 
 #if defined(NIAGARA_SCRIPT_COMPILE_LOGGING_MEDIUM)
 		if (!bSuccess)
diff --git a/Engine/Plugins/FX/Niagara/Source/NiagaraEditor/Private/NiagaraHlslTranslator.cpp b/Engine/Plugins/FX/Niagara/Source/NiagaraEditor/Private/NiagaraHlslTranslator.cpp
--- a/Engine/Plugins/FX/Niagara/Source/NiagaraEditor/Private/NiagaraHlslTranslator.cpp
+++ b/Engine/Plugins/FX/Niagara/Source/NiagaraEditor/Private/NiagaraHlslTranslator.cpp
@@ -138,7 +138,7 @@ void FNiagaraShaderQueueTickable::ProcessQueue()
 
 			FString ShaderCode = CompilableScript->GetVMExecutableData().LastHlslTranslationGPU;
 			// When not running in the editor, the shaders are created in-sync in the postload.
-			const bool bSynchronousCompile = !GIsEditor;
+			const bool bSynchronousCompile = !GIsEditor || GIsAutomationTesting;
 
 			// Compile the shaders for the script.
 			NewShaderMap->Compile(ShaderScript, Item.ShaderMapId, CompilerEnvironment, NewCompilationOutput, Item.Platform, bSynchronousCompile, Item.bApply);
