- Add support for upgrading data interface function calls
- Add support Skel Mesh IsValidBone & BoneCount on the GPU
- Rename specific to filtered in skel mesh DI
- Fix a issues with getting random bones on the GPU (Filter Socket / Bone was incorrect)

diff --git a/Engine/Plugins/FX/Niagara/Config/BaseNiagara.ini b/Engine/Plugins/FX/Niagara/Config/BaseNiagara.ini
--- a/Engine/Plugins/FX/Niagara/Config/BaseNiagara.ini
+++ b/Engine/Plugins/FX/Niagara/Config/BaseNiagara.ini
@@ -20,6 +20,9 @@
 
 ;+EnumRedirects=(OldName="ENiagaraSortMode",ValueChanges=(("ENiagaraSortMode::SortNone","ENiagaraSortMode::None"), ("ENiagaraSortMode::SortViewDepth","ENiagaraSortMode::ViewDepth"), ("ENiagaraSortMode::SortViewDistance","ENiagaraSortMode::ViewDistance")) )
 
++PropertyRedirects=(OldName="NiagaraDataInterfaceSkeletalMesh.SpecificBones", NewName="NiagaraDataInterfaceSkeletalMesh.FilteredBones")
++PropertyRedirects=(OldName="NiagaraDataInterfaceSkeletalMesh.SpecificSockets", NewName="NiagaraDataInterfaceSkeletalMesh.FilteredSockets")
+
 [/Script/Niagara.NiagaraSettings]
 +AdditionalParameterEnums=/Niagara/Enums/ENiagaraCoordinateSpace.ENiagaraCoordinateSpace
 +AdditionalParameterEnums=/Niagara/Enums/ENiagaraOrientationAxis.ENiagaraOrientationAxis
diff --git a/Engine/Plugins/FX/Niagara/Shaders/Private/NiagaraDataInterfaceSkeletalMesh.ush b/Engine/Plugins/FX/Niagara/Shaders/Private/NiagaraDataInterfaceSkeletalMesh.ush
--- a/Engine/Plugins/FX/Niagara/Shaders/Private/NiagaraDataInterfaceSkeletalMesh.ush
+++ b/Engine/Plugins/FX/Niagara/Shaders/Private/NiagaraDataInterfaceSkeletalMesh.ush
@@ -32,10 +32,11 @@ struct FDISkelMeshContext
 	uint			MeshSkinWeightIndexSize;
 	uint			MeshNumTexCoord;
 	uint			MeshNumWeights;
-	uint			NumSpecificBones;
-	Buffer<uint>	SpecificBones;
-	uint			NumSpecificSockets;
-	uint			SpecificSocketBoneOffset;
+	int				NumBones;
+	int				NumFilteredBones;
+	Buffer<uint>	FilteredBones;
+	int				NumFilteredSockets;
+	int				FilteredSocketBoneOffset;
 	float4x4		InstanceTransform;
 	float4x4		InstancePrevTransform;
 	float			InstanceInvDeltaTime;
@@ -403,9 +404,9 @@ void DISkelMesh_GetTriColor(in FDISkelMeshContext DIContext, in uint TriangleInd
     OutColor = (DIContext.MeshColorBuffer[VertexIndex0] * BaryCoord.x) + (DIContext.MeshColorBuffer[VertexIndex1] * BaryCoord.y) + (DIContext.MeshColorBuffer[VertexIndex2] * BaryCoord.z);
 }
 
-void DISkelMesh_GetRandomVertex(in FDISkelMeshContext DIContext, out int OutVertex)
+void DISkelMesh_GetRandomVertex(in FDISkelMeshContext DIContext, uint Seed1, uint Seed2, uint Seed3, out int OutVertex)
 {
-    float RandT0 = NiagaraInternalNoise(1, 2, 3);
+	float RandT0 = NiagaraInternalNoise(Seed1, Seed2, Seed3);
     OutVertex = min(uint(RandT0 * float(DIContext.MeshVertexCount)), DIContext.MeshVertexCount - 1); // avoid % by using mul/min to Tri = MeshVertexCount
 }
 
@@ -449,38 +450,55 @@ void DISkelMesh_GetVertexUV(in FDISkelMeshContext DIContext, in int VertexIndex,
 }
 
 ////////////////////////////////////////////////////////////////////////////////////////////////////////////////////////////////////////////////////////////////
-// Specific Bones / Sockets
-void DISkelMesh_GetSpecificBoneCount(in FDISkelMeshContext DIContext, out int Count)
+// Filtered Bones / Sockets
+void DISkelMesh_IsValidBone(in FDISkelMeshContext DIContext, in int BoneIndex, out bool IsValid)
 {
-    Count = DIContext.NumSpecificBones;
+	IsValid = BoneIndex >= 0 && BoneIndex < DIContext.NumBones;
 }
 
-void DISkelMesh_GetSpecificBoneAt(in FDISkelMeshContext DIContext, in int InBoneIndex, out int Bone)
+void DISkelMesh_RandomBone(in FDISkelMeshContext DIContext, uint Seed1, uint Seed2, uint Seed3, out int Bone)
 {
-    int BoneIndex = min(InBoneIndex, int(DIContext.NumSpecificBones - 1));
-    Bone = (int)DIContext.SpecificBones[BoneIndex];
+	float RandT0 = NiagaraInternalNoise(Seed1, Seed2, Seed3);
+    Bone = min(int(RandT0 * float(DIContext.NumBones)), DIContext.NumBones - 1); // avoid % by using mul/min to Tri = MeshVertexCount
 }
 
-void DISkelMesh_RandomSpecificBone(in FDISkelMeshContext DIContext, out int Bone)
+void DISkelMesh_GetBoneCount(in FDISkelMeshContext DIContext, out int Count)
 {
-    float RandT0 = NiagaraInternalNoise(1, 2, 3);
-    Bone = min(uint(RandT0 * float(DIContext.NumSpecificBones)), DIContext.NumSpecificBones - 1); // avoid % by using mul/min to Tri = MeshVertexCount
+	Count = DIContext.NumBones;
 }
 
-void DISkelMesh_GetSpecificSocketCount(in FDISkelMeshContext DIContext, out int Count)
+void DISkelMesh_GetFilteredBoneCount(in FDISkelMeshContext DIContext, out int Count)
 {
-    Count = DIContext.NumSpecificSockets;
+    Count = DIContext.NumFilteredBones;
 }
 
-void DISkelMesh_GetSpecificSocketBoneAt(in FDISkelMeshContext DIContext, in int InSocketIndex, out int Bone)
+void DISkelMesh_GetFilteredBoneAt(in FDISkelMeshContext DIContext, in int InBoneIndex, out int Bone)
 {
-    Bone = (int)DIContext.SpecificSocketBoneOffset + clamp(InSocketIndex, 0, (int)DIContext.NumSpecificSockets);
+    int BoneIndex = min(InBoneIndex, DIContext.NumFilteredBones - 1);
+    Bone = (int)DIContext.FilteredBones[BoneIndex];
 }
 
-void DISkelMesh_RandomSpecificSocketBone(in FDISkelMeshContext DIContext, out int SocketBone)
+void DISkelMesh_RandomFilteredBone(in FDISkelMeshContext DIContext, uint Seed1, uint Seed2, uint Seed3, out int Bone)
 {
-    float RandT0 = NiagaraInternalNoise(1, 2, 3);
-    SocketBone = min(uint(RandT0 * float(DIContext.NumSpecificSockets)), DIContext.NumSpecificSockets - 1); // avoid % by using mul/min to Tri = MeshVertexCount
+	float RandT0 = NiagaraInternalNoise(Seed1, Seed2, Seed3);
+    int FilteredBone = min(int(RandT0 * float(DIContext.NumFilteredBones)), DIContext.NumFilteredBones - 1); // avoid % by using mul/min to Tri = MeshVertexCount
+    Bone = (int)DIContext.FilteredBones[FilteredBone];
+}
+
+void DISkelMesh_GetFilteredSocketCount(in FDISkelMeshContext DIContext, out int Count)
+{
+    Count = DIContext.NumFilteredSockets;
+}
+
+void DISkelMesh_GetFilteredSocketBoneAt(in FDISkelMeshContext DIContext, in int InSocketIndex, out int Bone)
+{
+    Bone = DIContext.FilteredSocketBoneOffset + clamp(InSocketIndex, 0, DIContext.NumFilteredSockets);
+}
+
+void DISkelMesh_RandomFilteredSocketBone(in FDISkelMeshContext DIContext, uint Seed1, uint Seed2, uint Seed3, out int SocketBone)
+{
+	float RandT0 = NiagaraInternalNoise(Seed1, Seed2, Seed3);
+    SocketBone = DIContext.FilteredSocketBoneOffset + min(int(RandT0 * float(DIContext.NumFilteredSockets)), DIContext.NumFilteredSockets - 1); // avoid % by using mul/min to Tri = MeshVertexCount
 }
 
 ////////////////////////////////////////////////////////////////////////////////////////////////////////////////////////////////////////////////////////////////
@@ -506,10 +524,11 @@ void DISkelMesh_RandomSpecificSocketBone(in FDISkelMeshContext DIContext, out in
 	uint			MeshSkinWeightIndexSize_ ## NAME; \
 	uint			MeshNumTexCoord_ ## NAME; \
 	uint			MeshNumWeights_ ## NAME; \
-	uint			NumSpecificBones_ ## NAME; \
-	Buffer<uint>	SpecificBones_ ## NAME; \
-	uint			NumSpecificSockets_ ## NAME; \
-	uint			SpecificSocketBoneOffset_ ## NAME; \
+	int				NumBones_ ## NAME; \
+	int				NumFilteredBones_ ## NAME; \
+	Buffer<uint>	FilteredBones_ ## NAME; \
+	int				NumFilteredSockets_ ## NAME; \
+	int				FilteredSocketBoneOffset_ ## NAME; \
 	float4x4		InstanceTransform_ ## NAME; \
 	float4x4		InstancePrevTransform_ ## NAME; \
 	float			InstanceInvDeltaTime_ ## NAME; \
@@ -535,10 +554,11 @@ void DISkelMesh_RandomSpecificSocketBone(in FDISkelMeshContext DIContext, out in
 	DIContext.MeshVertexCount = MeshVertexCount_ ## NAME; \
 	DIContext.MeshNumTexCoord = MeshNumTexCoord_ ## NAME; \
 	DIContext.MeshNumWeights = MeshNumWeights_ ## NAME; \
-	DIContext.NumSpecificBones	= NumSpecificBones_ ## NAME; \
-	DIContext.SpecificBones	= SpecificBones_ ## NAME; \
-	DIContext.NumSpecificSockets = NumSpecificSockets_ ## NAME; \
-	DIContext.SpecificSocketBoneOffset = SpecificSocketBoneOffset_ ## NAME; \
+	DIContext.NumBones = NumBones_ ## NAME; \
+	DIContext.NumFilteredBones	= NumFilteredBones_ ## NAME; \
+	DIContext.FilteredBones	= FilteredBones_ ## NAME; \
+	DIContext.NumFilteredSockets = NumFilteredSockets_ ## NAME; \
+	DIContext.FilteredSocketBoneOffset = FilteredSocketBoneOffset_ ## NAME; \
 	DIContext.MeshWeightStride = MeshWeightStride_ ## NAME; \
 	DIContext.MeshSkinWeightIndexSize = MeshSkinWeightIndexSize_ ## NAME; \
 	DIContext.InstanceTransform = InstanceTransform_ ## NAME; \
diff --git a/Engine/Plugins/FX/Niagara/Source/Niagara/Classes/NiagaraDataInterface.h b/Engine/Plugins/FX/Niagara/Source/Niagara/Classes/NiagaraDataInterface.h
--- a/Engine/Plugins/FX/Niagara/Source/Niagara/Classes/NiagaraDataInterface.h
+++ b/Engine/Plugins/FX/Niagara/Source/Niagara/Classes/NiagaraDataInterface.h
@@ -291,6 +291,17 @@ public:
 		return false;
 	}
 
+#if WITH_EDITORONLY_DATA
+	/**
+	Allows data interfaces the opportunity to rename / change the function signature and perform an upgrade.
+	Return true if the signature was modified and we need to refresh the pins / name, etc.
+	*/
+	virtual bool UpgradeFunctionCall(FNiagaraFunctionSignature& FunctionSignature)
+	{
+		return false;
+	}
+#endif
+
 	virtual void PostExecute() {}
 
 #if WITH_EDITOR	
diff --git a/Engine/Plugins/FX/Niagara/Source/Niagara/Classes/NiagaraDataInterfaceSkeletalMesh.h b/Engine/Plugins/FX/Niagara/Source/Niagara/Classes/NiagaraDataInterfaceSkeletalMesh.h
--- a/Engine/Plugins/FX/Niagara/Source/Niagara/Classes/NiagaraDataInterfaceSkeletalMesh.h
+++ b/Engine/Plugins/FX/Niagara/Source/Niagara/Classes/NiagaraDataInterfaceSkeletalMesh.h
@@ -301,11 +301,11 @@ public:
 	uint32 GetNumTexCoord() const { return NumTexCoord; }
 	uint32 GetNumWeights() const { return NumWeights; }
 
-	uint32 GetNumSpecificBones() const { return NumSpecificBones; }
-	FRHIShaderResourceView* GetSpecificBonesSRV() const { return SpecificBonesSRV; }
+	int32 GetNumFilteredBones() const { return NumFilteredBones; }
+	FRHIShaderResourceView* GetFilteredBonesSRV() const { return FilteredBonesSRV; }
 
-	uint32 GetNumSpecificSockets() const { return NumSpecificSockets; }
-	uint32 GetSpecificSocketBoneOffset() const { return SpecificSocketBoneOffset; }
+	int32 GetNumFilteredSockets() const { return NumFilteredSockets; }
+	int32 GetFilteredSocketBoneOffset() const { return FilteredSocketBoneOffset; }
 
 protected:
 
@@ -316,13 +316,13 @@ protected:
 	FVertexBufferRHIRef BufferTriangleMatricesOffsetRHI = nullptr;
 	FShaderResourceViewRHIRef BufferTriangleMatricesOffsetSRV = nullptr;
 
-	uint32 NumSpecificBones = 0;
-	TResourceArray<uint16> SpecificBonesArray;
-	FVertexBufferRHIRef SpecificBonesBuffer;
-	FShaderResourceViewRHIRef SpecificBonesSRV;
+	int32 NumFilteredBones = 0;
+	TResourceArray<uint16> FilteredBonesArray;
+	FVertexBufferRHIRef FilteredBonesBuffer;
+	FShaderResourceViewRHIRef FilteredBonesSRV;
 
-	uint32 NumSpecificSockets = 0;
-	uint32 SpecificSocketBoneOffset = 0;
+	int32 NumFilteredSockets = 0;
+	int32 FilteredSocketBoneOffset = 0;
 
 	/** Cached SRV to gpu buffers of the mesh we spawn from */
 	FRHIShaderResourceView* MeshVertexBufferSrv;
@@ -364,6 +364,8 @@ public:
 
 	bool DoesBoneDataExist() const { return bBoneGpuBufferValid;}
 
+	int32 GetNumBones() const { return (int32)SamplingBoneCount;  }
+
 	/** Encapsulates a GPU read / CPU write buffer for bone data */
 	struct FSkeletalBuffer
 	{
@@ -427,26 +429,25 @@ struct FNDISkeletalMesh_InstanceData
 	/** Time separating Transform and PrevTransform. */
 	float DeltaSeconds;
 
-	/** Indices of the bones specifically referenced by the interface. */
-	TArray<int32> SpecificBones;
+	/** Indices of the bones filtered by the user. */
+	TArray<int32> FilteredBones;
 
 	/** Name of all the sockets we use. */
-	//TArray<FName> SpecificSockets;
 	struct FCachedSocketInfo
 	{
 		FCachedSocketInfo() : BoneIdx(INDEX_NONE){}
 		FTransform Transform;
 		int32 BoneIdx;
 	};
-	TArray<FCachedSocketInfo> SpecificSocketInfo;
+	TArray<FCachedSocketInfo> FilteredSocketInfo;
 
 	/** Bone index of the first socket, sockets are appended to the end of the bone array */
-	int32 SpecificSocketBoneOffset = 0;
+	int32 FilteredSocketBoneOffset = 0;
 
 	/** Index into which socket transforms to use.  */
-	uint32 SpecificSocketTransformsIndex = 0;
+	uint32 FilteredSocketTransformsIndex = 0;
 	/** Transforms for sockets. */
-	TStaticArray<TArray<FTransform>, 2> SpecificSocketTransforms;
+	TStaticArray<TArray<FTransform>, 2> FilteredSocketTransforms;
 
 	uint32 ChangeId;
 
@@ -489,10 +490,19 @@ struct FNDISkeletalMesh_InstanceData
 		return Ret;
 	}
 
-	void UpdateSpecificSocketTransforms();
-	TArray<FTransform>& GetSpecificSocketsWriteBuffer() { return SpecificSocketTransforms[SpecificSocketTransformsIndex]; }
-	const TArray<FTransform>& GetSpecificSocketsCurrBuffer() const { return SpecificSocketTransforms[SpecificSocketTransformsIndex]; }
-	const TArray<FTransform>& GetSpecificSocketsPrevBuffer() const { return SpecificSocketTransforms[(SpecificSocketTransformsIndex + 1) % SpecificSocketTransforms.Num()]; }
+	FORCEINLINE_DEBUGGABLE FSkinWeightVertexBuffer* GetSkinWeights()
+	{
+		if (USkeletalMeshComponent* SkelComp = Cast<USkeletalMeshComponent>(Component.Get()))
+		{
+			return SkelComp->GetSkinWeightBuffer(GetLODIndex());
+		}
+		return &Mesh->GetResourceForRendering()->LODRenderData[GetLODIndex()].SkinWeightVertexBuffer;
+	}
+
+	void UpdateFilteredSocketTransforms();
+	TArray<FTransform>& GetFilteredSocketsWriteBuffer() { return FilteredSocketTransforms[FilteredSocketTransformsIndex]; }
+	const TArray<FTransform>& GetFilteredSocketsCurrBuffer() const { return FilteredSocketTransforms[FilteredSocketTransformsIndex]; }
+	const TArray<FTransform>& GetFilteredSocketsPrevBuffer() const { return FilteredSocketTransforms[(FilteredSocketTransformsIndex + 1) % FilteredSocketTransforms.Num()]; }
 
 	bool HasColorData();
 };
@@ -536,13 +546,13 @@ public:
 	UPROPERTY(EditAnywhere, Category="Mesh")
 	int32 WholeMeshLOD;
 
-	/** Set of specific bones that can be used for sampling. Select from these with GetSpecificBoneAt and RandomSpecificBone. */
+	/** Set of filtered bones that can be used for sampling. Select from these with GetFilteredBoneAt and RandomFilteredBone. */
 	UPROPERTY(EditAnywhere, Category = "Skeleton")
-	TArray<FName> SpecificBones;
+	TArray<FName> FilteredBones;
 
-	/** Set of specific sockets that can be used for sampling. Select from these with GetSpecificSocketAt and RandomSpecificSocket. */
+	/** Set of filtered sockets that can be used for sampling. Select from these with GetFilteredSocketAt and RandomFilteredSocket. */
 	UPROPERTY(EditAnywhere, Category = "Skeleton")
-	TArray<FName> SpecificSockets;
+	TArray<FName> FilteredSockets;
 
 #if WITH_EDITORONLY_DATA
 	/** Do we require CPU access to the data, this is set during GetVMExternalFunction. */
@@ -585,6 +595,9 @@ public:
 	virtual void GetCommonHLSL(FString& OutHLSL) override;
 	virtual void GetParameterDefinitionHLSL(const FNiagaraDataInterfaceGPUParamInfo& ParamInfo, FString& OutHLSL) override;
 	virtual bool GetFunctionHLSL(const FNiagaraDataInterfaceGPUParamInfo& ParamInfo, const FNiagaraDataInterfaceGeneratedFunction& FunctionInfo, int FunctionInstanceIndex, FString& OutHLSL) override;
+#if WITH_EDITORONLY_DATA
+	virtual bool UpgradeFunctionCall(FNiagaraFunctionSignature& FunctionSignature) override;
+#endif
 
 	virtual void ProvidePerInstanceDataForRenderThread(void* DataForRenderThread, void* PerInstanceData, const FNiagaraSystemInstanceID& SystemInstance) override;
 
@@ -608,10 +621,11 @@ public:
 	static const FString MeshSkinWeightIndexSizeName;
 	static const FString MeshNumTexCoordName;
 	static const FString MeshNumWeightsName;
-	static const FString NumSpecificBonesName;
-	static const FString SpecificBonesName;
-	static const FString NumSpecificSocketsName;
-	static const FString SpecificSocketBoneOffsetName;
+	static const FString NumBonesName;
+	static const FString NumFilteredBonesName;
+	static const FString FilteredBonesName;
+	static const FString NumFilteredSocketsName;
+	static const FString FilteredSocketBoneOffsetName;
 	static const FString InstanceTransformName;
 	static const FString InstancePrevTransformName;
 	static const FString InstanceInvDeltaTimeName;
@@ -658,10 +672,10 @@ private:
 	FORCEINLINE int32 RandomTriIndex(FNDIRandomHelper& RandHelper, FSkeletalMeshAccessorHelper& Accessor, FNDISkeletalMesh_InstanceData* InstData, int32 InstanceIndex);
 
 	template<typename FilterMode, typename AreaWeightingMode>
-	FORCEINLINE int32 GetSpecificTriangleCount(FSkeletalMeshAccessorHelper& Accessor, FNDISkeletalMesh_InstanceData* InstData);
+	FORCEINLINE int32 GetFilteredTriangleCount(FSkeletalMeshAccessorHelper& Accessor, FNDISkeletalMesh_InstanceData* InstData);
 
 	template<typename FilterMode, typename AreaWeightingMode>
-	FORCEINLINE int32 GetSpecificTriangleAt(FSkeletalMeshAccessorHelper& Accessor, FNDISkeletalMesh_InstanceData* InstData, int32 FilteredIdx);
+	FORCEINLINE int32 GetFilteredTriangleAt(FSkeletalMeshAccessorHelper& Accessor, FNDISkeletalMesh_InstanceData* InstData, int32 FilteredIdx);
 	//End of Mesh Sampling
 	//////////////////////////////////////////////////////////////////////////
 
@@ -697,13 +711,13 @@ public:
 
 private:
 	template<typename FilterMode>
-	FORCEINLINE int32 RandomVertIndex(FRandomStream& RandStream, FSkeletalMeshAccessorHelper& Accessor, FNDISkeletalMesh_InstanceData* InstData);
+	FORCEINLINE int32 RandomVertIndex(FNDIRandomHelper& RandHelper, int32 Instance, FSkeletalMeshAccessorHelper& Accessor, FNDISkeletalMesh_InstanceData* InstData);
 
 	template<typename FilterMode>
-	FORCEINLINE int32 GetSpecificVertexCount(FSkeletalMeshAccessorHelper& Accessor, FNDISkeletalMesh_InstanceData* InstData);
+	FORCEINLINE int32 GetFilteredVertexCount(FSkeletalMeshAccessorHelper& Accessor, FNDISkeletalMesh_InstanceData* InstData);
 
 	template<typename FilterMode>
-	FORCEINLINE int32 GetSpecificVertexAt(FSkeletalMeshAccessorHelper& Accessor, FNDISkeletalMesh_InstanceData* InstData, int32 FilteredIdx);
+	FORCEINLINE int32 GetFilteredVertexAt(FSkeletalMeshAccessorHelper& Accessor, FNDISkeletalMesh_InstanceData* InstData, int32 FilteredIdx);
 
 	//End of Vertex Sampling
 	//////////////////////////////////////////////////////////////////////////
@@ -715,24 +729,22 @@ public:
 	void GetSkeletonSamplingFunctions(TArray<FNiagaraFunctionSignature>& OutFunctions);
 	void BindSkeletonSamplingFunction(const FVMExternalFunctionBindingInfo& BindingInfo, FNDISkeletalMesh_InstanceData* InstData, FVMExternalFunction &OutFunc);
 
-	void IsValidBone(FVectorVMContext& Context);
-
-	void GetSpecificBoneCount(FVectorVMContext& Context);
-
-	void GetSpecificBoneAt(FVectorVMContext& Context);
-
-	void RandomSpecificBone(FVectorVMContext& Context);
-
 	template<typename SkinningHandlerType, typename TransformHandlerType, typename bInterpolated>
 	void GetSkinnedBoneData(FVectorVMContext& Context);
-	
-	void GetSpecificSocketCount(FVectorVMContext& Context);
 
-	void GetSpecificSocketBoneAt(FVectorVMContext& Context);
+	void IsValidBone(FVectorVMContext& Context);
+	void RandomBone(FVectorVMContext& Context);
+	void GetBoneCount(FVectorVMContext& Context);
+
+	void GetFilteredBoneCount(FVectorVMContext& Context);
+	void GetFilteredBoneAt(FVectorVMContext& Context);
+	void RandomFilteredBone(FVectorVMContext& Context);
 
-	void GetSpecificSocketTransform(FVectorVMContext& Context);
 
-	void RandomSpecificSocketBone(FVectorVMContext& Context);
+	void GetFilteredSocketCount(FVectorVMContext& Context);
+	void GetFilteredSocketBoneAt(FVectorVMContext& Context);
+	void GetFilteredSocketTransform(FVectorVMContext& Context);
+	void RandomFilteredSocketBone(FVectorVMContext& Context);
 		
 	// End of Direct Bone + Socket Sampling
 	//////////////////////////////////////////////////////////////////////////
@@ -762,14 +774,16 @@ public:
 	static const FName GetSkinnedBoneDataWSName;
 	static const FName GetSkinnedBoneDataInterpolatedName;
 	static const FName GetSkinnedBoneDataWSInterpolatedName;
-	static const FName RandomSpecificBoneName;
 	static const FName IsValidBoneName;
-	static const FName GetSpecificBoneCountName;
-	static const FName GetSpecificBoneAtName;
-	static const FName RandomSpecificSocketBoneName;
-	static const FName GetSpecificSocketCountName;
-	static const FName GetSpecificSocketBoneAtName;
-	static const FName GetSpecificSocketTransformName;
+	static const FName RandomBoneName;
+	static const FName GetBoneCountName;
+	static const FName RandomFilteredBoneName;
+	static const FName GetFilteredBoneCountName;
+	static const FName GetFilteredBoneAtName;
+	static const FName RandomFilteredSocketBoneName;
+	static const FName GetFilteredSocketCountName;
+	static const FName GetFilteredSocketBoneAtName;
+	static const FName GetFilteredSocketTransformName;
 
 	// Vertex Sampling
 	static const FName IsValidVertexName;
diff --git a/Engine/Plugins/FX/Niagara/Source/Niagara/Private/NDISkeletalMesh_BoneSampling.cpp b/Engine/Plugins/FX/Niagara/Source/Niagara/Private/NDISkeletalMesh_BoneSampling.cpp
--- a/Engine/Plugins/FX/Niagara/Source/Niagara/Private/NDISkeletalMesh_BoneSampling.cpp
+++ b/Engine/Plugins/FX/Niagara/Source/Niagara/Private/NDISkeletalMesh_BoneSampling.cpp
@@ -16,143 +16,171 @@ DECLARE_CYCLE_STAT(TEXT("Skel Mesh Skeleton Sampling"), STAT_NiagaraSkel_Bone_Sa
 //Final binders for all static mesh interface functions.
 DEFINE_NDI_FUNC_BINDER(UNiagaraDataInterfaceSkeletalMesh, GetSkinnedBoneData)
 DEFINE_NDI_DIRECT_FUNC_BINDER(UNiagaraDataInterfaceSkeletalMesh, IsValidBone)
-DEFINE_NDI_DIRECT_FUNC_BINDER(UNiagaraDataInterfaceSkeletalMesh, GetSpecificBoneAt)
+DEFINE_NDI_DIRECT_FUNC_BINDER(UNiagaraDataInterfaceSkeletalMesh, GetFilteredBoneAt)
 
-DEFINE_NDI_DIRECT_FUNC_BINDER(UNiagaraDataInterfaceSkeletalMesh, GetSpecificSocketBoneAt)
-DEFINE_NDI_DIRECT_FUNC_BINDER(UNiagaraDataInterfaceSkeletalMesh, GetSpecificSocketTransform)
+DEFINE_NDI_DIRECT_FUNC_BINDER(UNiagaraDataInterfaceSkeletalMesh, GetFilteredSocketBoneAt)
+DEFINE_NDI_DIRECT_FUNC_BINDER(UNiagaraDataInterfaceSkeletalMesh, GetFilteredSocketTransform)
 
 const FName FSkeletalMeshInterfaceHelper::GetSkinnedBoneDataName("GetSkinnedBoneData");
 const FName FSkeletalMeshInterfaceHelper::GetSkinnedBoneDataWSName("GetSkinnedBoneDataWS");
 const FName FSkeletalMeshInterfaceHelper::GetSkinnedBoneDataInterpolatedName("GetSkinnedBoneDataInterpolated");
 const FName FSkeletalMeshInterfaceHelper::GetSkinnedBoneDataWSInterpolatedName("GetSkinnedBoneDataWSInterpolated");
-const FName FSkeletalMeshInterfaceHelper::RandomSpecificBoneName("RandomSpecificBone");
-const FName FSkeletalMeshInterfaceHelper::IsValidBoneName("IsValidBoneName");
-const FName FSkeletalMeshInterfaceHelper::GetSpecificBoneCountName("GetSpecificBoneCount");
-const FName FSkeletalMeshInterfaceHelper::GetSpecificBoneAtName("GetSpecificBone");
-const FName FSkeletalMeshInterfaceHelper::RandomSpecificSocketBoneName("RandomSpecificSocketBone");
-const FName FSkeletalMeshInterfaceHelper::GetSpecificSocketCountName("GetSpecificSocketCount");
-const FName FSkeletalMeshInterfaceHelper::GetSpecificSocketTransformName("GetSpecificSocketTransform");
-const FName FSkeletalMeshInterfaceHelper::GetSpecificSocketBoneAtName("GetSpecificSocketBone");
+
+const FName FSkeletalMeshInterfaceHelper::IsValidBoneName("IsValidBone");
+const FName FSkeletalMeshInterfaceHelper::RandomBoneName("RandomBone");
+const FName FSkeletalMeshInterfaceHelper::GetBoneCountName("GetBoneCount");
+
+const FName FSkeletalMeshInterfaceHelper::RandomFilteredBoneName("RandomFilteredBone");
+const FName FSkeletalMeshInterfaceHelper::GetFilteredBoneCountName("GetFilteredBoneCount");
+const FName FSkeletalMeshInterfaceHelper::GetFilteredBoneAtName("GetFilteredBone");
+
+const FName FSkeletalMeshInterfaceHelper::RandomFilteredSocketBoneName("RandomFilteredSocketBone");
+const FName FSkeletalMeshInterfaceHelper::GetFilteredSocketCountName("GetFilteredSocketCount");
+const FName FSkeletalMeshInterfaceHelper::GetFilteredSocketTransformName("GetFilteredSocketTransform");
+const FName FSkeletalMeshInterfaceHelper::GetFilteredSocketBoneAtName("GetFilteredSocketBone");
 
 void UNiagaraDataInterfaceSkeletalMesh::GetSkeletonSamplingFunctions(TArray<FNiagaraFunctionSignature>& OutFunctions)
 {
 	//////////////////////////////////////////////////////////////////////////
 	// Bone functions.
-
 	{
 		FNiagaraFunctionSignature Sig;
-		Sig.Name = FSkeletalMeshInterfaceHelper::RandomSpecificBoneName;
+		Sig.Name = FSkeletalMeshInterfaceHelper::GetSkinnedBoneDataName;
 		Sig.Inputs.Add(FNiagaraVariable(FNiagaraTypeDefinition(GetClass()), TEXT("SkeletalMesh")));
-		Sig.Outputs.Add(FNiagaraVariable(FNiagaraTypeDefinition::GetIntDef() , TEXT("Bone")));
+		Sig.Inputs.Add(FNiagaraVariable(FNiagaraTypeDefinition::GetIntDef(), TEXT("Bone")));
+		Sig.Outputs.Add(FNiagaraVariable(FNiagaraTypeDefinition::GetVec3Def(), TEXT("Position")));
+		Sig.Outputs.Add(FNiagaraVariable(FNiagaraTypeDefinition::GetQuatDef(), TEXT("Rotation")));
+		Sig.Outputs.Add(FNiagaraVariable(FNiagaraTypeDefinition::GetVec3Def(), TEXT("Velocity")));
 		Sig.bMemberFunction = true;
 		Sig.bRequiresContext = false;
+#if WITH_EDITORONLY_DATA
+		Sig.Description = LOCTEXT("GetOptionalSkinnedBoneDataDesc", "Returns skinning dependant data for the pased bone in local space. All outputs are optional and you will incur zero to minimal cost if they are not connected.");
+#endif
 		OutFunctions.Add(Sig);
 	}
 
 	{
 		FNiagaraFunctionSignature Sig;
-		Sig.Name = FSkeletalMeshInterfaceHelper::IsValidBoneName;
+		Sig.Name = FSkeletalMeshInterfaceHelper::GetSkinnedBoneDataWSName;
 		Sig.Inputs.Add(FNiagaraVariable(FNiagaraTypeDefinition(GetClass()), TEXT("SkeletalMesh")));
 		Sig.Inputs.Add(FNiagaraVariable(FNiagaraTypeDefinition::GetIntDef(), TEXT("Bone")));
-		Sig.Outputs.Add(FNiagaraVariable(FNiagaraTypeDefinition::GetBoolDef(), TEXT("IsValid")));
+		Sig.Outputs.Add(FNiagaraVariable(FNiagaraTypeDefinition::GetVec3Def(), TEXT("Position")));
+		Sig.Outputs.Add(FNiagaraVariable(FNiagaraTypeDefinition::GetQuatDef(), TEXT("Rotation")));
+		Sig.Outputs.Add(FNiagaraVariable(FNiagaraTypeDefinition::GetVec3Def(), TEXT("Velocity")));
 		Sig.bMemberFunction = true;
 		Sig.bRequiresContext = false;
 #if WITH_EDITORONLY_DATA
-		Sig.Description = LOCTEXT("IsValidBoneDesc", "Determine if this bone index is valid for this mesh's skeleton.");
+		Sig.Description = LOCTEXT("GetOptionalSkinnedBoneDataWSDesc", "Returns skinning dependant data for the pased bone in world space. All outputs are optional and you will incur zero to minimal cost if they are not connected.");
 #endif
 		OutFunctions.Add(Sig);
 	}
 
+
 	{
 		FNiagaraFunctionSignature Sig;
-		Sig.Name = FSkeletalMeshInterfaceHelper::GetSkinnedBoneDataName;
+		Sig.Name = FSkeletalMeshInterfaceHelper::GetSkinnedBoneDataInterpolatedName;
 		Sig.Inputs.Add(FNiagaraVariable(FNiagaraTypeDefinition(GetClass()), TEXT("SkeletalMesh")));
 		Sig.Inputs.Add(FNiagaraVariable(FNiagaraTypeDefinition::GetIntDef(), TEXT("Bone")));
+		Sig.Inputs.Add(FNiagaraVariable(FNiagaraTypeDefinition::GetFloatDef(), TEXT("Interpolation")));
 		Sig.Outputs.Add(FNiagaraVariable(FNiagaraTypeDefinition::GetVec3Def(), TEXT("Position")));
 		Sig.Outputs.Add(FNiagaraVariable(FNiagaraTypeDefinition::GetQuatDef(), TEXT("Rotation")));
 		Sig.Outputs.Add(FNiagaraVariable(FNiagaraTypeDefinition::GetVec3Def(), TEXT("Velocity")));
 		Sig.bMemberFunction = true;
 		Sig.bRequiresContext = false;
 #if WITH_EDITORONLY_DATA
-		Sig.Description = LOCTEXT("GetOptionalSkinnedBoneDataDesc", "Returns skinning dependant data for the pased bone in local space. All outputs are optional and you will incur zero to minimal cost if they are not connected.");
+		Sig.Description = LOCTEXT("GetSkinnedBoneDataDesc", "Returns skinning dependant data for the pased bone in local space. Interpolated between this frame and the previous based on passed interpolation factor. All outputs are optional and you will incur zero to minimal cost if they are not connected.");
 #endif
 		OutFunctions.Add(Sig);
 	}
 
 	{
 		FNiagaraFunctionSignature Sig;
-		Sig.Name = FSkeletalMeshInterfaceHelper::GetSkinnedBoneDataWSName;
+		Sig.Name = FSkeletalMeshInterfaceHelper::GetSkinnedBoneDataWSInterpolatedName;
 		Sig.Inputs.Add(FNiagaraVariable(FNiagaraTypeDefinition(GetClass()), TEXT("SkeletalMesh")));
 		Sig.Inputs.Add(FNiagaraVariable(FNiagaraTypeDefinition::GetIntDef(), TEXT("Bone")));
+		Sig.Inputs.Add(FNiagaraVariable(FNiagaraTypeDefinition::GetFloatDef(), TEXT("Interpolation")));
 		Sig.Outputs.Add(FNiagaraVariable(FNiagaraTypeDefinition::GetVec3Def(), TEXT("Position")));
 		Sig.Outputs.Add(FNiagaraVariable(FNiagaraTypeDefinition::GetQuatDef(), TEXT("Rotation")));
 		Sig.Outputs.Add(FNiagaraVariable(FNiagaraTypeDefinition::GetVec3Def(), TEXT("Velocity")));
 		Sig.bMemberFunction = true;
 		Sig.bRequiresContext = false;
 #if WITH_EDITORONLY_DATA
-		Sig.Description = LOCTEXT("GetOptionalSkinnedBoneDataWSDesc", "Returns skinning dependant data for the pased bone in world space. All outputs are optional and you will incur zero to minimal cost if they are not connected.");
+		Sig.Description = LOCTEXT("GetSkinnedBoneDataWSDesc", "Returns skinning dependant data for the pased bone in world space. Interpolated between this frame and the previous based on passed interpolation factor. All outputs are optional and you will incur zero to minimal cost if they are not connected.");
 #endif
 		OutFunctions.Add(Sig);
 	}
 
-
 	{
 		FNiagaraFunctionSignature Sig;
-		Sig.Name = FSkeletalMeshInterfaceHelper::GetSkinnedBoneDataInterpolatedName;
+		Sig.Name = FSkeletalMeshInterfaceHelper::IsValidBoneName;
 		Sig.Inputs.Add(FNiagaraVariable(FNiagaraTypeDefinition(GetClass()), TEXT("SkeletalMesh")));
 		Sig.Inputs.Add(FNiagaraVariable(FNiagaraTypeDefinition::GetIntDef(), TEXT("Bone")));
-		Sig.Inputs.Add(FNiagaraVariable(FNiagaraTypeDefinition::GetFloatDef(), TEXT("Interpolation")));
-		Sig.Outputs.Add(FNiagaraVariable(FNiagaraTypeDefinition::GetVec3Def(), TEXT("Position")));
-		Sig.Outputs.Add(FNiagaraVariable(FNiagaraTypeDefinition::GetQuatDef(), TEXT("Rotation")));
-		Sig.Outputs.Add(FNiagaraVariable(FNiagaraTypeDefinition::GetVec3Def(), TEXT("Velocity")));
+		Sig.Outputs.Add(FNiagaraVariable(FNiagaraTypeDefinition::GetBoolDef(), TEXT("IsValid")));
 		Sig.bMemberFunction = true;
 		Sig.bRequiresContext = false;
 #if WITH_EDITORONLY_DATA
-		Sig.Description = LOCTEXT("GetSkinnedBoneDataDesc", "Returns skinning dependant data for the pased bone in local space. Interpolated between this frame and the previous based on passed interpolation factor. All outputs are optional and you will incur zero to minimal cost if they are not connected.");
+		Sig.Description = LOCTEXT("IsValidBoneDesc", "Determine if this bone index is valid for this mesh's skeleton.");
 #endif
 		OutFunctions.Add(Sig);
 	}
 
 	{
 		FNiagaraFunctionSignature Sig;
-		Sig.Name = FSkeletalMeshInterfaceHelper::GetSkinnedBoneDataWSInterpolatedName;
+		Sig.Name = FSkeletalMeshInterfaceHelper::RandomBoneName;
 		Sig.Inputs.Add(FNiagaraVariable(FNiagaraTypeDefinition(GetClass()), TEXT("SkeletalMesh")));
-		Sig.Inputs.Add(FNiagaraVariable(FNiagaraTypeDefinition::GetIntDef(), TEXT("Bone")));
-		Sig.Inputs.Add(FNiagaraVariable(FNiagaraTypeDefinition::GetFloatDef(), TEXT("Interpolation")));
-		Sig.Outputs.Add(FNiagaraVariable(FNiagaraTypeDefinition::GetVec3Def(), TEXT("Position")));
-		Sig.Outputs.Add(FNiagaraVariable(FNiagaraTypeDefinition::GetQuatDef(), TEXT("Rotation")));
-		Sig.Outputs.Add(FNiagaraVariable(FNiagaraTypeDefinition::GetVec3Def(), TEXT("Velocity")));
+		Sig.Inputs.Add(FNiagaraVariable(FNiagaraTypeDefinition(FNiagaraRandInfo::StaticStruct()), TEXT("RandomInfo")));
+		Sig.Outputs.Add(FNiagaraVariable(FNiagaraTypeDefinition::GetIntDef(), TEXT("Bone")));
+		Sig.bMemberFunction = true;
+		Sig.bRequiresContext = false;
+		OutFunctions.Add(Sig);
+	}
+	{
+		FNiagaraFunctionSignature Sig;
+		Sig.Name = FSkeletalMeshInterfaceHelper::GetBoneCountName;
+		Sig.Inputs.Add(FNiagaraVariable(FNiagaraTypeDefinition(GetClass()), TEXT("SkeletalMesh")));
+		Sig.Outputs.Add(FNiagaraVariable(FNiagaraTypeDefinition::GetIntDef(), TEXT("Count")));
 		Sig.bMemberFunction = true;
 		Sig.bRequiresContext = false;
 #if WITH_EDITORONLY_DATA
-		Sig.Description = LOCTEXT("GetSkinnedBoneDataWSDesc", "Returns skinning dependant data for the pased bone in world space. Interpolated between this frame and the previous based on passed interpolation factor. All outputs are optional and you will incur zero to minimal cost if they are not connected.");
+		Sig.Description = LOCTEXT("GetBoneCountDesc", "Returns the number of bones in the skeletal mesh.");
 #endif
 		OutFunctions.Add(Sig);
 	}
 
 	{
 		FNiagaraFunctionSignature Sig;
-		Sig.Name = FSkeletalMeshInterfaceHelper::GetSpecificBoneCountName;
+		Sig.Name = FSkeletalMeshInterfaceHelper::RandomFilteredBoneName;
+		Sig.Inputs.Add(FNiagaraVariable(FNiagaraTypeDefinition(GetClass()), TEXT("SkeletalMesh")));
+		Sig.Inputs.Add(FNiagaraVariable(FNiagaraTypeDefinition(FNiagaraRandInfo::StaticStruct()), TEXT("RandomInfo")));
+		Sig.Outputs.Add(FNiagaraVariable(FNiagaraTypeDefinition::GetIntDef(), TEXT("Bone")));
+		Sig.bMemberFunction = true;
+		Sig.bRequiresContext = false;
+		OutFunctions.Add(Sig);
+	}
+
+	{
+		FNiagaraFunctionSignature Sig;
+		Sig.Name = FSkeletalMeshInterfaceHelper::GetFilteredBoneCountName;
 		Sig.Inputs.Add(FNiagaraVariable(FNiagaraTypeDefinition(GetClass()), TEXT("SkeletalMesh")));
 		Sig.Outputs.Add(FNiagaraVariable(FNiagaraTypeDefinition::GetIntDef(), TEXT("Count")));
 		Sig.bMemberFunction = true;
 		Sig.bRequiresContext = false;
 #if WITH_EDITORONLY_DATA
-		Sig.Description = LOCTEXT("GetSpecificBoneCountDesc", "Returns the number of specific bones in the DI list.");
+		Sig.Description = LOCTEXT("GetFilteredBoneCountDesc", "Returns the number of filtered bones in the DI list.");
 #endif
 		OutFunctions.Add(Sig);
 	}
 
 	{
 		FNiagaraFunctionSignature Sig;
-		Sig.Name = FSkeletalMeshInterfaceHelper::GetSpecificBoneAtName;
+		Sig.Name = FSkeletalMeshInterfaceHelper::GetFilteredBoneAtName;
 		Sig.Inputs.Add(FNiagaraVariable(FNiagaraTypeDefinition(GetClass()), TEXT("SkeletalMesh")));
 		Sig.Inputs.Add(FNiagaraVariable(FNiagaraTypeDefinition::GetIntDef(), TEXT("Bone Index")));
 		Sig.Outputs.Add(FNiagaraVariable(FNiagaraTypeDefinition::GetIntDef(), TEXT("Bone")));
 		Sig.bMemberFunction = true;
 		Sig.bRequiresContext = false;
 #if WITH_EDITORONLY_DATA
-		Sig.Description = LOCTEXT("GetSpecificBoneAtDesc", "Gets the bone at the passed index in the DI's specfic bones list.");
+		Sig.Description = LOCTEXT("GetFilteredBoneAtDesc", "Gets the bone at the passed index in the DI's filter bones list.");
 #endif
 		OutFunctions.Add(Sig);
 	}
@@ -162,47 +190,48 @@ void UNiagaraDataInterfaceSkeletalMesh::GetSkeletonSamplingFunctions(TArray<FNia
 
 	{
 		FNiagaraFunctionSignature Sig;
-		Sig.Name = FSkeletalMeshInterfaceHelper::RandomSpecificSocketBoneName;
+		Sig.Name = FSkeletalMeshInterfaceHelper::RandomFilteredSocketBoneName;
 		Sig.Inputs.Add(FNiagaraVariable(FNiagaraTypeDefinition(GetClass()), TEXT("SkeletalMesh")));
+		Sig.Inputs.Add(FNiagaraVariable(FNiagaraTypeDefinition(FNiagaraRandInfo::StaticStruct()), TEXT("RandomInfo")));
 		Sig.Outputs.Add(FNiagaraVariable(FNiagaraTypeDefinition::GetIntDef(), TEXT("Socket Bone")));
 		Sig.bMemberFunction = true;
 		Sig.bRequiresContext = false;
 #if WITH_EDITORONLY_DATA
-		Sig.Description = LOCTEXT("RandomSpecificSocketBoneDesc", "Gets the bone for a random socket in the DI's specific socket list.");
+		Sig.Description = LOCTEXT("RandomFilteredSocketBoneDesc", "Gets the bone for a random socket in the DI's filtered socket list.");
 #endif
 		OutFunctions.Add(Sig);
 	}
 	
 	{
 		FNiagaraFunctionSignature Sig;
-		Sig.Name = FSkeletalMeshInterfaceHelper::GetSpecificSocketCountName;
+		Sig.Name = FSkeletalMeshInterfaceHelper::GetFilteredSocketCountName;
 		Sig.Inputs.Add(FNiagaraVariable(FNiagaraTypeDefinition(GetClass()), TEXT("SkeletalMesh")));
 		Sig.Outputs.Add(FNiagaraVariable(FNiagaraTypeDefinition::GetIntDef(), TEXT("Count")));
 		Sig.bMemberFunction = true;
 		Sig.bRequiresContext = false;
 #if WITH_EDITORONLY_DATA
-		Sig.Description = LOCTEXT("GetSpecificSocketCountDesc", "Returns the number of specific Sockets in the DI list.");
+		Sig.Description = LOCTEXT("GetFilteredSocketCountDesc", "Returns the number of filtered Sockets in the DI list.");
 #endif
 		OutFunctions.Add(Sig);
 	}
 
 	{
 		FNiagaraFunctionSignature Sig;
-		Sig.Name = FSkeletalMeshInterfaceHelper::GetSpecificSocketBoneAtName;
+		Sig.Name = FSkeletalMeshInterfaceHelper::GetFilteredSocketBoneAtName;
 		Sig.Inputs.Add(FNiagaraVariable(FNiagaraTypeDefinition(GetClass()), TEXT("SkeletalMesh")));
 		Sig.Inputs.Add(FNiagaraVariable(FNiagaraTypeDefinition::GetIntDef(), TEXT("Socket Index")));
 		Sig.Outputs.Add(FNiagaraVariable(FNiagaraTypeDefinition::GetIntDef(), TEXT("Socket Bone")));
 		Sig.bMemberFunction = true;
 		Sig.bRequiresContext = false;
 #if WITH_EDITORONLY_DATA
-		Sig.Description = LOCTEXT("GetSpecificSocketBoneAtDesc", "Gets the bone for the socket at the passed index in the DI's specfic socket list.");
+		Sig.Description = LOCTEXT("GetFilteredSocketBoneAtDesc", "Gets the bone for the socket at the passed index in the DI's filtered socket list.");
 #endif
 		OutFunctions.Add(Sig);
 	}
 
 	{
 		FNiagaraFunctionSignature Sig;
-		Sig.Name = FSkeletalMeshInterfaceHelper::GetSpecificSocketTransformName;
+		Sig.Name = FSkeletalMeshInterfaceHelper::GetFilteredSocketTransformName;
 		Sig.Inputs.Add(FNiagaraVariable(FNiagaraTypeDefinition(GetClass()), TEXT("SkeletalMesh")));
 		Sig.Inputs.Add(FNiagaraVariable(FNiagaraTypeDefinition::GetIntDef(), TEXT("Socket Index")));
 		Sig.Inputs.Add(FNiagaraVariable(FNiagaraTypeDefinition::GetBoolDef(), TEXT("Apply World Transform")));
@@ -212,7 +241,7 @@ void UNiagaraDataInterfaceSkeletalMesh::GetSkeletonSamplingFunctions(TArray<FNia
 		Sig.bMemberFunction = true;
 		Sig.bRequiresContext = false;
 #if WITH_EDITORONLY_DATA
-		Sig.Description = LOCTEXT("GetSpecificSocketTransformDesc", "Gets the transform for the socket at the passed index in the DI's specfic socket list. If the Source component is set it will respect the Relative Transform Space as well..");
+		Sig.Description = LOCTEXT("GetFilteredSocketTransformDesc", "Gets the transform for the socket at the passed index in the DI's filtered socket list. If the Source component is set it will respect the Relative Transform Space as well..");
 #endif
 		OutFunctions.Add(Sig);
 	}
@@ -221,18 +250,7 @@ void UNiagaraDataInterfaceSkeletalMesh::GetSkeletonSamplingFunctions(TArray<FNia
 void UNiagaraDataInterfaceSkeletalMesh::BindSkeletonSamplingFunction(const FVMExternalFunctionBindingInfo& BindingInfo, FNDISkeletalMesh_InstanceData* InstanceData, FVMExternalFunction &OutFunc)
 {
 	//Bone Functions
-	if (BindingInfo.Name == FSkeletalMeshInterfaceHelper::RandomSpecificBoneName)
-	{
-		check(BindingInfo.GetNumInputs() == 1 && BindingInfo.GetNumOutputs() == 1);
-		auto Lambda = [this](FVectorVMContext& Context) { this->RandomSpecificBone(Context); }; 
-		OutFunc = FVMExternalFunction::CreateLambda(Lambda); 
-	}
-	else if (BindingInfo.Name == FSkeletalMeshInterfaceHelper::IsValidBoneName)
-	{
-		check(BindingInfo.GetNumInputs() == 2 && BindingInfo.GetNumOutputs() == 1);
-		NDI_FUNC_BINDER(UNiagaraDataInterfaceSkeletalMesh, IsValidBone)::Bind(this, OutFunc);
-	}
-	else if (BindingInfo.Name == FSkeletalMeshInterfaceHelper::GetSkinnedBoneDataName)
+	if (BindingInfo.Name == FSkeletalMeshInterfaceHelper::GetSkinnedBoneDataName)
 	{
 		ensure(BindingInfo.GetNumInputs() == 2 && BindingInfo.GetNumOutputs() == 10);
 		TSkinningModeBinder<TNDIExplicitBinder<FNDITransformHandlerNoop, TNDIExplicitBinder<TIntegralConstant<bool, false>, NDI_FUNC_BINDER(UNiagaraDataInterfaceSkeletalMesh, GetSkinnedBoneData)>>>::Bind(this, BindingInfo, InstanceData, OutFunc);
@@ -252,39 +270,56 @@ void UNiagaraDataInterfaceSkeletalMesh::BindSkeletonSamplingFunction(const FVMEx
 		ensure(BindingInfo.GetNumInputs() == 3 && BindingInfo.GetNumOutputs() == 10);
 		TSkinningModeBinder<TNDIExplicitBinder<FNDITransformHandler, TNDIExplicitBinder<TIntegralConstant<bool, true>, NDI_FUNC_BINDER(UNiagaraDataInterfaceSkeletalMesh, GetSkinnedBoneData)>>>::Bind(this, BindingInfo, InstanceData, OutFunc);
 	}
-	else if (BindingInfo.Name == FSkeletalMeshInterfaceHelper::GetSpecificBoneCountName)
+	else if (BindingInfo.Name == FSkeletalMeshInterfaceHelper::IsValidBoneName)
+	{
+		check(BindingInfo.GetNumInputs() == 2 && BindingInfo.GetNumOutputs() == 1);
+		NDI_FUNC_BINDER(UNiagaraDataInterfaceSkeletalMesh, IsValidBone)::Bind(this, OutFunc);
+	}
+	else if (BindingInfo.Name == FSkeletalMeshInterfaceHelper::RandomBoneName)
+	{
+		check(BindingInfo.GetNumInputs() == 4 && BindingInfo.GetNumOutputs() == 1);
+		OutFunc = FVMExternalFunction::CreateLambda([this](FVectorVMContext& Context) { this->RandomBone(Context); });
+	}
+	else if (BindingInfo.Name == FSkeletalMeshInterfaceHelper::GetBoneCountName)
 	{
 		check(BindingInfo.GetNumInputs() == 1 && BindingInfo.GetNumOutputs() == 1);
-		auto Lambda = [this](FVectorVMContext& Context) { this->GetSpecificBoneCount(Context); };
-		OutFunc = FVMExternalFunction::CreateLambda(Lambda);
+		OutFunc = FVMExternalFunction::CreateLambda([this](FVectorVMContext& Context) { this->GetBoneCount(Context); });
+	}
+	else if (BindingInfo.Name == FSkeletalMeshInterfaceHelper::RandomFilteredBoneName)
+	{
+		check(BindingInfo.GetNumInputs() == 4 && BindingInfo.GetNumOutputs() == 1);
+		OutFunc = FVMExternalFunction::CreateLambda([this](FVectorVMContext& Context) { this->RandomFilteredBone(Context); });
 	}
-	else if (BindingInfo.Name == FSkeletalMeshInterfaceHelper::GetSpecificBoneAtName)
+	else if (BindingInfo.Name == FSkeletalMeshInterfaceHelper::GetFilteredBoneCountName)
+	{
+		check(BindingInfo.GetNumInputs() == 1 && BindingInfo.GetNumOutputs() == 1);
+		OutFunc = FVMExternalFunction::CreateLambda([this](FVectorVMContext& Context) { this->GetFilteredBoneCount(Context); });
+	}
+	else if (BindingInfo.Name == FSkeletalMeshInterfaceHelper::GetFilteredBoneAtName)
 	{
 		check(BindingInfo.GetNumInputs() == 2 && BindingInfo.GetNumOutputs() == 1);
-		NDI_FUNC_BINDER(UNiagaraDataInterfaceSkeletalMesh, GetSpecificBoneAt)::Bind(this, OutFunc);
+		NDI_FUNC_BINDER(UNiagaraDataInterfaceSkeletalMesh, GetFilteredBoneAt)::Bind(this, OutFunc);
 	}
 	//Socket Functions
-	else if (BindingInfo.Name == FSkeletalMeshInterfaceHelper::RandomSpecificSocketBoneName)
+	else if (BindingInfo.Name == FSkeletalMeshInterfaceHelper::RandomFilteredSocketBoneName)
 	{
-		check(BindingInfo.GetNumInputs() == 1 && BindingInfo.GetNumOutputs() == 1);
-		auto Lambda = [this](FVectorVMContext& Context) { this->RandomSpecificSocketBone(Context); };
-		OutFunc = FVMExternalFunction::CreateLambda(Lambda);
+		check(BindingInfo.GetNumInputs() == 4 && BindingInfo.GetNumOutputs() == 1);
+		OutFunc = FVMExternalFunction::CreateLambda([this](FVectorVMContext& Context) { this->RandomFilteredSocketBone(Context); });
 	}
-	else if (BindingInfo.Name == FSkeletalMeshInterfaceHelper::GetSpecificSocketCountName)
+	else if (BindingInfo.Name == FSkeletalMeshInterfaceHelper::GetFilteredSocketCountName)
 	{
 		check(BindingInfo.GetNumInputs() == 1 && BindingInfo.GetNumOutputs() == 1);
-		auto Lambda = [this](FVectorVMContext& Context) { this->GetSpecificSocketCount(Context); };
-		OutFunc = FVMExternalFunction::CreateLambda(Lambda);
+		OutFunc = FVMExternalFunction::CreateLambda([this](FVectorVMContext& Context) { this->GetFilteredSocketCount(Context); });
 	}
-	else if (BindingInfo.Name == FSkeletalMeshInterfaceHelper::GetSpecificSocketBoneAtName)
+	else if (BindingInfo.Name == FSkeletalMeshInterfaceHelper::GetFilteredSocketBoneAtName)
 	{
 		check(BindingInfo.GetNumInputs() == 2 && BindingInfo.GetNumOutputs() == 1);
-		NDI_FUNC_BINDER(UNiagaraDataInterfaceSkeletalMesh, GetSpecificSocketBoneAt)::Bind(this, OutFunc);
+		NDI_FUNC_BINDER(UNiagaraDataInterfaceSkeletalMesh, GetFilteredSocketBoneAt)::Bind(this, OutFunc);
 	}
-	else if (BindingInfo.Name == FSkeletalMeshInterfaceHelper::GetSpecificSocketTransformName)
+	else if (BindingInfo.Name == FSkeletalMeshInterfaceHelper::GetFilteredSocketTransformName)
 	{
 		check(BindingInfo.GetNumInputs() == 3 && BindingInfo.GetNumOutputs() == 10);
-		NDI_FUNC_BINDER(UNiagaraDataInterfaceSkeletalMesh, GetSpecificSocketTransform)::Bind(this, OutFunc);
+		NDI_FUNC_BINDER(UNiagaraDataInterfaceSkeletalMesh, GetFilteredSocketTransform)::Bind(this, OutFunc);
 	}
 }
 
@@ -292,37 +327,35 @@ void UNiagaraDataInterfaceSkeletalMesh::BindSkeletonSamplingFunction(const FVMEx
 //////////////////////////////////////////////////////////////////////////
 // Direct sampling from listed sockets and bones.
 
-void UNiagaraDataInterfaceSkeletalMesh::GetSpecificBoneCount(FVectorVMContext& Context)
+void UNiagaraDataInterfaceSkeletalMesh::GetFilteredBoneCount(FVectorVMContext& Context)
 {
 	SCOPE_CYCLE_COUNTER(STAT_NiagaraSkel_Bone_Sample);
 	VectorVM::FUserPtrHandler<FNDISkeletalMesh_InstanceData> InstData(Context);
 
 	VectorVM::FExternalFuncRegisterHandler<int32> OutCount(Context);
 
-	int32 Num = SpecificBones.Num();
+	int32 Num = FilteredBones.Num();
 	for (int32 i = 0; i < Context.NumInstances; ++i)
 	{
 		*OutCount.GetDestAndAdvance() = Num;
 	}
 }
 
-void UNiagaraDataInterfaceSkeletalMesh::GetSpecificBoneAt(FVectorVMContext& Context)
+void UNiagaraDataInterfaceSkeletalMesh::GetFilteredBoneAt(FVectorVMContext& Context)
 {
-	SCOPE_CYCLE_COUNTER(STAT_NiagaraSkel_Bone_Sample);
-
 	VectorVM::FExternalFuncInputHandler<int32> BoneParam(Context);
 	VectorVM::FUserPtrHandler<FNDISkeletalMesh_InstanceData> InstData(Context);
 
 	VectorVM::FExternalFuncRegisterHandler<int32> OutBone(Context);
-	const TArray<int32>& SpecificBonesArray = InstData->SpecificBones;
+	const TArray<int32>& FilteredBonesArray = InstData->FilteredBones;
 
-	int32 Max = SpecificBones.Num() - 1;
+	int32 Max = FilteredBones.Num() - 1;
 	if (Max != INDEX_NONE)
 	{
 		for (int32 i = 0; i < Context.NumInstances; ++i)
 		{
 			int32 BoneIndex = FMath::Clamp(BoneParam.GetAndAdvance(), 0, Max);
-			*OutBone.GetDestAndAdvance() = SpecificBonesArray[BoneIndex];
+			*OutBone.GetDestAndAdvance() = FilteredBonesArray[BoneIndex];
 		}
 	}
 	else
@@ -331,22 +364,22 @@ void UNiagaraDataInterfaceSkeletalMesh::GetSpecificBoneAt(FVectorVMContext& Cont
 	}
 }
 
-void UNiagaraDataInterfaceSkeletalMesh::RandomSpecificBone(FVectorVMContext& Context)
+void UNiagaraDataInterfaceSkeletalMesh::RandomFilteredBone(FVectorVMContext& Context)
 {
-	SCOPE_CYCLE_COUNTER(STAT_NiagaraSkel_Bone_Sample);
-
+	FNDIRandomHelper RandHelper(Context);
 	VectorVM::FUserPtrHandler<FNDISkeletalMesh_InstanceData> InstData(Context);
 
 	VectorVM::FExternalFuncRegisterHandler<int32> OutBone(Context);
-	const TArray<int32>& SpecificBonesArray = InstData->SpecificBones;
+	const TArray<int32>& FilteredBonesArray = InstData->FilteredBones;
 
-	int32 Max = SpecificBones.Num() - 1;
+	int32 Max = FilteredBones.Num() - 1;
 	if (Max != INDEX_NONE)
 	{
 		for (int32 i = 0; i < Context.NumInstances; ++i)
 		{
-			int32 BoneIndex = Context.RandStream.RandRange(0, Max);
-			*OutBone.GetDestAndAdvance() = SpecificBonesArray[BoneIndex];
+			RandHelper.GetAndAdvance();
+			const int32 BoneIndex = RandHelper.RandRange(i, 0, Max);
+			*OutBone.GetDestAndAdvance() = FilteredBonesArray[BoneIndex];
 		}
 	}
 	else
@@ -357,8 +390,6 @@ void UNiagaraDataInterfaceSkeletalMesh::RandomSpecificBone(FVectorVMContext& Con
 
 void UNiagaraDataInterfaceSkeletalMesh::IsValidBone(FVectorVMContext& Context)
 {
-	SCOPE_CYCLE_COUNTER(STAT_NiagaraSkel_Bone_Sample);
-
 	VectorVM::FExternalFuncInputHandler<int32> BoneParam(Context);
 
 	VectorVM::FUserPtrHandler<FNDISkeletalMesh_InstanceData> InstData(Context);
@@ -381,6 +412,47 @@ void UNiagaraDataInterfaceSkeletalMesh::IsValidBone(FVectorVMContext& Context)
 	}
 }
 
+void UNiagaraDataInterfaceSkeletalMesh::RandomBone(FVectorVMContext& Context)
+{
+	FNDIRandomHelper RandHelper(Context);
+	VectorVM::FUserPtrHandler<FNDISkeletalMesh_InstanceData> InstData(Context);
+	VectorVM::FExternalFuncRegisterHandler<int32> OutBone(Context);
+
+	FSkeletalMeshAccessorHelper MeshAccessor;
+	MeshAccessor.Init<TIntegralConstant<int32, 0>, TIntegralConstant<int32, 0>>(InstData);
+	const FReferenceSkeleton& RefSkeleton = MeshAccessor.Mesh->RefSkeleton;
+	const int32 NumBones = RefSkeleton.GetNum() - 1;
+	if (NumBones >= 0)
+	{
+		for (int32 i = 0; i < Context.NumInstances; ++i)
+		{
+			RandHelper.GetAndAdvance();
+			*OutBone.GetDestAndAdvance() = RandHelper.RandRange(i, 0, NumBones);
+		}
+	}
+	else
+	{
+		FMemory::Memset(OutBone.GetDest(), 0xFF, Context.NumInstances * sizeof(int32));
+	}
+}
+
+void UNiagaraDataInterfaceSkeletalMesh::GetBoneCount(FVectorVMContext& Context)
+{
+	SCOPE_CYCLE_COUNTER(STAT_NiagaraSkel_Bone_Sample);
+
+	VectorVM::FUserPtrHandler<FNDISkeletalMesh_InstanceData> InstData(Context);
+	VectorVM::FExternalFuncRegisterHandler<int32> OutCount(Context);
+
+	FSkeletalMeshAccessorHelper MeshAccessor;
+	MeshAccessor.Init<TIntegralConstant<int32, 0>, TIntegralConstant<int32, 0>>(InstData);
+	const FReferenceSkeleton& RefSkeleton = MeshAccessor.Mesh->RefSkeleton;
+	const int32 NumBones = RefSkeleton.GetNum();
+	for (int32 i = 0; i < Context.NumInstances; ++i)
+	{
+		*OutCount.GetDestAndAdvance() = NumBones;
+	}
+}
+
 struct FBoneSocketSkinnedDataOutputHandler
 {
 	FBoneSocketSkinnedDataOutputHandler(FVectorVMContext& Context)
@@ -429,7 +501,6 @@ struct FBoneSocketSkinnedDataOutputHandler
 template<typename SkinningHandlerType, typename TransformHandlerType, typename bInterpolated>
 void UNiagaraDataInterfaceSkeletalMesh::GetSkinnedBoneData(FVectorVMContext& Context)
 {
-	SCOPE_CYCLE_COUNTER(STAT_NiagaraSkel_Bone_Sample);
 	SkinningHandlerType SkinningHandler;
 	TransformHandlerType TransformHandler;
 	VectorVM::FExternalFuncInputHandler<int32> BoneParam(Context);
@@ -451,21 +522,20 @@ void UNiagaraDataInterfaceSkeletalMesh::GetSkinnedBoneData(FVectorVMContext& Con
 	//Also need to pull in a per particle interpolation factor.
 	const FMatrix& InstanceTransform = InstData->Transform;
 	const FMatrix& PrevInstanceTransform = InstData->PrevTransform;
-	const FQuat InstanceRotation = Output.bNeedsRotation ? InstData->Transform.GetMatrixWithoutScale().ToQuat() : FQuat::Identity;
-	const FQuat PrevInstanceRotation = Output.bNeedsRotation ? InstData->Transform.GetMatrixWithoutScale().ToQuat() : FQuat::Identity;
+	const FQuat InstanceRotation = Output.bNeedsRotation ? InstanceTransform.ToQuat() : FQuat::Identity;
+	const FQuat PrevInstanceRotation = Output.bNeedsRotation ? PrevInstanceTransform.ToQuat() : FQuat::Identity;
 
-	FSkinWeightVertexBuffer* SkinWeightBuffer;
-	FSkeletalMeshLODRenderData& LODData = InstData->GetLODRenderDataAndSkinWeights(SkinWeightBuffer);
+	FSkinWeightVertexBuffer* SkinWeightBuffer = InstData->GetSkinWeights();
 
 	FSkeletalMeshAccessorHelper Accessor;
 	Accessor.Init<TIntegralConstant<int32, 0>, TIntegralConstant<int32, 0>>(InstData);
 
 	const int32 BoneCount = SkinningHandler.GetBoneCount(Accessor, bInterpolated::Value);
-	const int32 BoneAndSocketCount = BoneCount + InstData->SpecificSocketInfo.Num();
+	const int32 BoneAndSocketCount = BoneCount + InstData->FilteredSocketInfo.Num();
 	float InvDt = 1.0f / InstData->DeltaSeconds;
 
-	const TArray<FTransform>& SpecificSocketCurrTransforms = InstData->GetSpecificSocketsCurrBuffer();
-	const TArray<FTransform>& SpecificSocketPrevTransforms = InstData->GetSpecificSocketsPrevBuffer();
+	const TArray<FTransform>& FilteredSocketCurrTransforms = InstData->GetFilteredSocketsCurrBuffer();
+	const TArray<FTransform>& FilteredSocketPrevTransforms = InstData->GetFilteredSocketsPrevBuffer();
 
 	for (int32 i = 0; i < Context.NumInstances; ++i)
 	{
@@ -483,7 +553,7 @@ void UNiagaraDataInterfaceSkeletalMesh::GetSkinnedBoneData(FVectorVMContext& Con
 
 		// Handle edge cases first...
 		if ((!bIsSocket && Bone >= BoneCount) ||
-			(bIsSocket && !SpecificSocketCurrTransforms.IsValidIndex(Socket)))
+			(bIsSocket && !FilteredSocketCurrTransforms.IsValidIndex(Socket)))
 		{
 			Pos = FVector::ZeroVector;
 			TransformHandler.TransformPosition(Pos, InstanceTransform);
@@ -507,8 +577,8 @@ void UNiagaraDataInterfaceSkeletalMesh::GetSkinnedBoneData(FVectorVMContext& Con
 		}
 		else if ( bIsSocket )
 		{
-			FTransform CurrSocketTransform = SpecificSocketCurrTransforms[Socket];
-			FTransform PrevSocketTransform = SpecificSocketPrevTransforms[Socket];
+			FTransform CurrSocketTransform = FilteredSocketCurrTransforms[Socket];
+			FTransform PrevSocketTransform = FilteredSocketPrevTransforms[Socket];
 
 			Pos = CurrSocketTransform.GetLocation();
 			TransformHandler.TransformPosition(Pos, InstanceTransform);
@@ -582,37 +652,34 @@ void UNiagaraDataInterfaceSkeletalMesh::GetSkinnedBoneData(FVectorVMContext& Con
 //////////////////////////////////////////////////////////////////////////
 // Sockets
 
-void UNiagaraDataInterfaceSkeletalMesh::GetSpecificSocketCount(FVectorVMContext& Context)
+void UNiagaraDataInterfaceSkeletalMesh::GetFilteredSocketCount(FVectorVMContext& Context)
 {
-	SCOPE_CYCLE_COUNTER(STAT_NiagaraSkel_Bone_Sample);
 	VectorVM::FUserPtrHandler<FNDISkeletalMesh_InstanceData> InstData(Context);
 
 	VectorVM::FExternalFuncRegisterHandler<int32> OutCount(Context);
 
-	const int32 Num = InstData->SpecificSocketInfo.Num();
+	const int32 Num = InstData->FilteredSocketInfo.Num();
 	for (int32 i = 0; i < Context.NumInstances; ++i)
 	{
 		*OutCount.GetDestAndAdvance() = Num;
 	}
 }
 
-void UNiagaraDataInterfaceSkeletalMesh::GetSpecificSocketBoneAt(FVectorVMContext& Context)
+void UNiagaraDataInterfaceSkeletalMesh::GetFilteredSocketBoneAt(FVectorVMContext& Context)
 {
-	SCOPE_CYCLE_COUNTER(STAT_NiagaraSkel_Bone_Sample);
-
 	VectorVM::FExternalFuncInputHandler<int32> SocketParam(Context);
 	VectorVM::FUserPtrHandler<FNDISkeletalMesh_InstanceData> InstData(Context);
 
 	VectorVM::FExternalFuncRegisterHandler<int32> OutSocketBone(Context);
-	const int32 SpecificSocketBoneOffset = InstData->SpecificSocketBoneOffset;
-	const int32 Max = SpecificSockets.Num() - 1;
+	const int32 FilteredSocketBoneOffset = InstData->FilteredSocketBoneOffset;
+	const int32 Max = FilteredSockets.Num() - 1;
 
 	if (Max != INDEX_NONE)
 	{
 		for (int32 i = 0; i < Context.NumInstances; ++i)
 		{
 			const int32 SocketIndex = FMath::Clamp(SocketParam.GetAndAdvance(), 0, Max);
-			*OutSocketBone.GetDestAndAdvance() = SpecificSocketBoneOffset + SocketIndex;
+			*OutSocketBone.GetDestAndAdvance() = FilteredSocketBoneOffset + SocketIndex;
 		}
 	}
 	else
@@ -621,10 +688,8 @@ void UNiagaraDataInterfaceSkeletalMesh::GetSpecificSocketBoneAt(FVectorVMContext
 	}
 }
 
-void UNiagaraDataInterfaceSkeletalMesh::GetSpecificSocketTransform(FVectorVMContext& Context)
+void UNiagaraDataInterfaceSkeletalMesh::GetFilteredSocketTransform(FVectorVMContext& Context)
 {
-	SCOPE_CYCLE_COUNTER(STAT_NiagaraSkel_Bone_Sample);
-
 	VectorVM::FExternalFuncInputHandler<int32> SocketParam(Context);
 	VectorVM::FExternalFuncInputHandler<int32> ApplyWorldTransform(Context);
 	
@@ -652,14 +717,14 @@ void UNiagaraDataInterfaceSkeletalMesh::GetSpecificSocketTransform(FVectorVMCont
 		ComponentToWorld = SkelComp->GetComponentToWorld();
 	}
 
-	const TArray<FTransform>& CurrentSpecificSockets = InstData->GetSpecificSocketsCurrBuffer();
+	const TArray<FTransform>& CurrentFilteredSockets = InstData->GetFilteredSocketsCurrBuffer();
 
 	for (int32 i = 0; i < Context.NumInstances; ++i)
 	{//
-		const int32 SocketIndex = FMath::Clamp(SocketParam.GetAndAdvance(), 0, CurrentSpecificSockets.Num()-1);
+		const int32 SocketIndex = FMath::Clamp(SocketParam.GetAndAdvance(), 0, CurrentFilteredSockets.Num()-1);
 		const int32 ShouldApplyWorldTransform = ApplyWorldTransform.GetAndAdvance();
 		
-		FTransform CurrentTransform = CurrentSpecificSockets[SocketIndex];
+		FTransform CurrentTransform = CurrentFilteredSockets[SocketIndex];
 		if (ShouldApplyWorldTransform == -1)
 		{
 			CurrentTransform = CurrentTransform * ComponentToWorld;
@@ -690,22 +755,22 @@ void UNiagaraDataInterfaceSkeletalMesh::GetSpecificSocketTransform(FVectorVMCont
 	}
 }
 
-void UNiagaraDataInterfaceSkeletalMesh::RandomSpecificSocketBone(FVectorVMContext& Context)
+void UNiagaraDataInterfaceSkeletalMesh::RandomFilteredSocketBone(FVectorVMContext& Context)
 {
-	SCOPE_CYCLE_COUNTER(STAT_NiagaraSkel_Bone_Sample);
-
+	FNDIRandomHelper RandHelper(Context);
 	VectorVM::FUserPtrHandler<FNDISkeletalMesh_InstanceData> InstData(Context);
 
 	VectorVM::FExternalFuncRegisterHandler<int32> OutSocketBone(Context);
-	const int32 SpecificSocketBoneOffset = InstData->SpecificSocketBoneOffset;
+	const int32 FilteredSocketBoneOffset = InstData->FilteredSocketBoneOffset;
 
-	int32 Max = SpecificSockets.Num() - 1;
+	int32 Max = FilteredSockets.Num() - 1;
 	if (Max != INDEX_NONE)
 	{
 		for (int32 i = 0; i < Context.NumInstances; ++i)
 		{
-			const int32 SocketIndex = Context.RandStream.RandRange(0, Max);
-			*OutSocketBone.GetDestAndAdvance() = SpecificSocketBoneOffset + SocketIndex;
+			RandHelper.GetAndAdvance();
+			const int32 SocketIndex = RandHelper.RandRange(i, 0, Max);
+			*OutSocketBone.GetDestAndAdvance() = FilteredSocketBoneOffset + SocketIndex;
 		}
 	}
 	else
diff --git a/Engine/Plugins/FX/Niagara/Source/Niagara/Private/NDISkeletalMesh_TriangleSampling.cpp b/Engine/Plugins/FX/Niagara/Source/Niagara/Private/NDISkeletalMesh_TriangleSampling.cpp
--- a/Engine/Plugins/FX/Niagara/Source/Niagara/Private/NDISkeletalMesh_TriangleSampling.cpp
+++ b/Engine/Plugins/FX/Niagara/Source/Niagara/Private/NDISkeletalMesh_TriangleSampling.cpp
@@ -202,7 +202,6 @@ void UNiagaraDataInterfaceSkeletalMesh::GetTriangleSamplingFunctions(TArray<FNia
 
 void UNiagaraDataInterfaceSkeletalMesh::BindTriangleSamplingFunction(const FVMExternalFunctionBindingInfo& BindingInfo, FNDISkeletalMesh_InstanceData* InstanceData, FVMExternalFunction &OutFunc)
 {
-
 	if (BindingInfo.Name == FSkeletalMeshInterfaceHelper::RandomTriCoordName)
 	{
 		check(BindingInfo.GetNumInputs() == 4 && BindingInfo.GetNumOutputs() == 4);
@@ -413,14 +412,14 @@ void UNiagaraDataInterfaceSkeletalMesh::IsValidTriCoord(FVectorVMContext& Contex
 //////////////////////////////////////////////////////////////////////////
 
 template<typename FilterMode, typename AreaWeightingMode>
-FORCEINLINE int32 UNiagaraDataInterfaceSkeletalMesh::GetSpecificTriangleCount(FSkeletalMeshAccessorHelper& Accessor, FNDISkeletalMesh_InstanceData* InstData)
+FORCEINLINE int32 UNiagaraDataInterfaceSkeletalMesh::GetFilteredTriangleCount(FSkeletalMeshAccessorHelper& Accessor, FNDISkeletalMesh_InstanceData* InstData)
 {
-	checkf(false, TEXT("Invalid template call for GetSpecificTriangleCount. Bug in Filter binding or Area Weighting binding. Contact code team."));
+	checkf(false, TEXT("Invalid template call for GetFilteredTriangleCount. Bug in Filter binding or Area Weighting binding. Contact code team."));
 	return 0;
 }
 
 template<>
-FORCEINLINE int32 UNiagaraDataInterfaceSkeletalMesh::GetSpecificTriangleCount<
+FORCEINLINE int32 UNiagaraDataInterfaceSkeletalMesh::GetFilteredTriangleCount<
 	TIntegralConstant<ENDISkeletalMesh_FilterMode, ENDISkeletalMesh_FilterMode::None>,
 	TIntegralConstant<ENDISkelMesh_AreaWeightingMode, ENDISkelMesh_AreaWeightingMode::None>>
 	(FSkeletalMeshAccessorHelper& Accessor, FNDISkeletalMesh_InstanceData* InstData)
@@ -434,7 +433,7 @@ FORCEINLINE int32 UNiagaraDataInterfaceSkeletalMesh::GetSpecificTriangleCount<
 }
 
 template<>
-FORCEINLINE int32 UNiagaraDataInterfaceSkeletalMesh::GetSpecificTriangleCount<
+FORCEINLINE int32 UNiagaraDataInterfaceSkeletalMesh::GetFilteredTriangleCount<
 	TIntegralConstant<ENDISkeletalMesh_FilterMode, ENDISkeletalMesh_FilterMode::None>,
 	TIntegralConstant<ENDISkelMesh_AreaWeightingMode, ENDISkelMesh_AreaWeightingMode::AreaWeighted>>
 	(FSkeletalMeshAccessorHelper& Accessor, FNDISkeletalMesh_InstanceData* InstData)
@@ -445,7 +444,7 @@ FORCEINLINE int32 UNiagaraDataInterfaceSkeletalMesh::GetSpecificTriangleCount<
 }
 
 template<>
-FORCEINLINE int32 UNiagaraDataInterfaceSkeletalMesh::GetSpecificTriangleCount<
+FORCEINLINE int32 UNiagaraDataInterfaceSkeletalMesh::GetFilteredTriangleCount<
 	TIntegralConstant<ENDISkeletalMesh_FilterMode, ENDISkeletalMesh_FilterMode::SingleRegion>,
 	TIntegralConstant<ENDISkelMesh_AreaWeightingMode, ENDISkelMesh_AreaWeightingMode::None>>
 	(FSkeletalMeshAccessorHelper& Accessor, FNDISkeletalMesh_InstanceData* InstData)
@@ -454,7 +453,7 @@ FORCEINLINE int32 UNiagaraDataInterfaceSkeletalMesh::GetSpecificTriangleCount<
 }
 
 template<>
-FORCEINLINE int32 UNiagaraDataInterfaceSkeletalMesh::GetSpecificTriangleCount<
+FORCEINLINE int32 UNiagaraDataInterfaceSkeletalMesh::GetFilteredTriangleCount<
 	TIntegralConstant<ENDISkeletalMesh_FilterMode, ENDISkeletalMesh_FilterMode::SingleRegion>,
 	TIntegralConstant<ENDISkelMesh_AreaWeightingMode, ENDISkelMesh_AreaWeightingMode::AreaWeighted>>
 	(FSkeletalMeshAccessorHelper& Accessor, FNDISkeletalMesh_InstanceData* InstData)
@@ -463,7 +462,7 @@ FORCEINLINE int32 UNiagaraDataInterfaceSkeletalMesh::GetSpecificTriangleCount<
 }
 
 template<>
-FORCEINLINE int32 UNiagaraDataInterfaceSkeletalMesh::GetSpecificTriangleCount<
+FORCEINLINE int32 UNiagaraDataInterfaceSkeletalMesh::GetFilteredTriangleCount<
 	TIntegralConstant<ENDISkeletalMesh_FilterMode, ENDISkeletalMesh_FilterMode::MultiRegion>,
 	TIntegralConstant<ENDISkelMesh_AreaWeightingMode, ENDISkelMesh_AreaWeightingMode::None>>
 	(FSkeletalMeshAccessorHelper& Accessor, FNDISkeletalMesh_InstanceData* InstData)
@@ -481,7 +480,7 @@ FORCEINLINE int32 UNiagaraDataInterfaceSkeletalMesh::GetSpecificTriangleCount<
 }
 
 template<>
-FORCEINLINE int32 UNiagaraDataInterfaceSkeletalMesh::GetSpecificTriangleCount<
+FORCEINLINE int32 UNiagaraDataInterfaceSkeletalMesh::GetFilteredTriangleCount<
 	TIntegralConstant<ENDISkeletalMesh_FilterMode, ENDISkeletalMesh_FilterMode::MultiRegion>,
 	TIntegralConstant<ENDISkelMesh_AreaWeightingMode, ENDISkelMesh_AreaWeightingMode::AreaWeighted>>
 	(FSkeletalMeshAccessorHelper& Accessor, FNDISkeletalMesh_InstanceData* InstData)
@@ -511,7 +510,7 @@ void UNiagaraDataInterfaceSkeletalMesh::GetFilteredTriangleCount(FVectorVMContex
 	FSkeletalMeshAccessorHelper MeshAccessor;
 	MeshAccessor.Init<FilterMode, AreaWeightingMode>(InstData);
 
-	int32 Count = GetSpecificTriangleCount<FilterMode, AreaWeightingMode>(MeshAccessor, InstData);
+	int32 Count = GetFilteredTriangleCount<FilterMode, AreaWeightingMode>(MeshAccessor, InstData);
 	for (int32 i = 0; i < Context.NumInstances; ++i)
 	{
 		*OutTri.GetDest() = Count;
@@ -523,14 +522,14 @@ void UNiagaraDataInterfaceSkeletalMesh::GetFilteredTriangleCount(FVectorVMContex
 //////////////////////////////////////////////////////////////////////////
 
 template<typename FilterMode, typename AreaWeightingMode>
-FORCEINLINE int32 UNiagaraDataInterfaceSkeletalMesh::GetSpecificTriangleAt(FSkeletalMeshAccessorHelper& Accessor, FNDISkeletalMesh_InstanceData* InstData, int32 FilteredIndex)
+FORCEINLINE int32 UNiagaraDataInterfaceSkeletalMesh::GetFilteredTriangleAt(FSkeletalMeshAccessorHelper& Accessor, FNDISkeletalMesh_InstanceData* InstData, int32 FilteredIndex)
 {
-	checkf(false, TEXT("Invalid template call for GetSpecificTriangleAt. Bug in Filter binding or Area Weighting binding. Contact code team."));
+	checkf(false, TEXT("Invalid template call for GetFilteredTriangleAt. Bug in Filter binding or Area Weighting binding. Contact code team."));
 	return 0;
 }
 
 template<>
-FORCEINLINE int32 UNiagaraDataInterfaceSkeletalMesh::GetSpecificTriangleAt<
+FORCEINLINE int32 UNiagaraDataInterfaceSkeletalMesh::GetFilteredTriangleAt<
 	TIntegralConstant<ENDISkeletalMesh_FilterMode, ENDISkeletalMesh_FilterMode::None>,
 	TIntegralConstant<ENDISkelMesh_AreaWeightingMode, ENDISkelMesh_AreaWeightingMode::None>>
 	(FSkeletalMeshAccessorHelper& Accessor, FNDISkeletalMesh_InstanceData* InstData, int32 FilteredIndex)
@@ -548,7 +547,7 @@ FORCEINLINE int32 UNiagaraDataInterfaceSkeletalMesh::GetSpecificTriangleAt<
 }
 
 template<>
-FORCEINLINE int32 UNiagaraDataInterfaceSkeletalMesh::GetSpecificTriangleAt<
+FORCEINLINE int32 UNiagaraDataInterfaceSkeletalMesh::GetFilteredTriangleAt<
 	TIntegralConstant<ENDISkeletalMesh_FilterMode, ENDISkeletalMesh_FilterMode::None>,
 	TIntegralConstant<ENDISkelMesh_AreaWeightingMode, ENDISkelMesh_AreaWeightingMode::AreaWeighted>>
 	(FSkeletalMeshAccessorHelper& Accessor, FNDISkeletalMesh_InstanceData* InstData, int32 FilteredIndex)
@@ -558,7 +557,7 @@ FORCEINLINE int32 UNiagaraDataInterfaceSkeletalMesh::GetSpecificTriangleAt<
 }
 
 template<>
-FORCEINLINE int32 UNiagaraDataInterfaceSkeletalMesh::GetSpecificTriangleAt<
+FORCEINLINE int32 UNiagaraDataInterfaceSkeletalMesh::GetFilteredTriangleAt<
 	TIntegralConstant<ENDISkeletalMesh_FilterMode, ENDISkeletalMesh_FilterMode::SingleRegion>,
 	TIntegralConstant<ENDISkelMesh_AreaWeightingMode, ENDISkelMesh_AreaWeightingMode::None>>
 	(FSkeletalMeshAccessorHelper& Accessor, FNDISkeletalMesh_InstanceData* InstData, int32 FilteredIndex)
@@ -569,7 +568,7 @@ FORCEINLINE int32 UNiagaraDataInterfaceSkeletalMesh::GetSpecificTriangleAt<
 }
 
 template<>
-FORCEINLINE int32 UNiagaraDataInterfaceSkeletalMesh::GetSpecificTriangleAt<
+FORCEINLINE int32 UNiagaraDataInterfaceSkeletalMesh::GetFilteredTriangleAt<
 	TIntegralConstant<ENDISkeletalMesh_FilterMode, ENDISkeletalMesh_FilterMode::SingleRegion>,
 	TIntegralConstant<ENDISkelMesh_AreaWeightingMode, ENDISkelMesh_AreaWeightingMode::AreaWeighted>>
 	(FSkeletalMeshAccessorHelper& Accessor, FNDISkeletalMesh_InstanceData* InstData, int32 FilteredIndex)
@@ -582,7 +581,7 @@ FORCEINLINE int32 UNiagaraDataInterfaceSkeletalMesh::GetSpecificTriangleAt<
 }
 
 template<>
-FORCEINLINE int32 UNiagaraDataInterfaceSkeletalMesh::GetSpecificTriangleAt<
+FORCEINLINE int32 UNiagaraDataInterfaceSkeletalMesh::GetFilteredTriangleAt<
 	TIntegralConstant<ENDISkeletalMesh_FilterMode, ENDISkeletalMesh_FilterMode::MultiRegion>,
 	TIntegralConstant<ENDISkelMesh_AreaWeightingMode, ENDISkelMesh_AreaWeightingMode::None>>
 	(FSkeletalMeshAccessorHelper& Accessor, FNDISkeletalMesh_InstanceData* InstData, int32 FilteredIndex)
@@ -603,7 +602,7 @@ FORCEINLINE int32 UNiagaraDataInterfaceSkeletalMesh::GetSpecificTriangleAt<
 }
 
 template<>
-FORCEINLINE int32 UNiagaraDataInterfaceSkeletalMesh::GetSpecificTriangleAt<
+FORCEINLINE int32 UNiagaraDataInterfaceSkeletalMesh::GetFilteredTriangleAt<
 	TIntegralConstant<ENDISkeletalMesh_FilterMode, ENDISkeletalMesh_FilterMode::MultiRegion>,
 	TIntegralConstant<ENDISkelMesh_AreaWeightingMode, ENDISkelMesh_AreaWeightingMode::AreaWeighted>>
 	(FSkeletalMeshAccessorHelper& Accessor, FNDISkeletalMesh_InstanceData* InstData, int32 FilteredIndex)
@@ -641,7 +640,7 @@ void UNiagaraDataInterfaceSkeletalMesh::GetFilteredTriangleAt(FVectorVMContext&
 	{
 		int32 Tri = TriParam.Get();
 		int32 RealIdx = 0;
-		RealIdx = GetSpecificTriangleAt<FilterMode, AreaWeightingMode>(Accessor, InstData, Tri);
+		RealIdx = GetFilteredTriangleAt<FilterMode, AreaWeightingMode>(Accessor, InstData, Tri);
 
 		const int32 TriMax = (Accessor.IndexBuffer->Num() / 3) - 1;
 		RealIdx = FMath::Clamp(RealIdx, 0, TriMax);
diff --git a/Engine/Plugins/FX/Niagara/Source/Niagara/Private/NDISkeletalMesh_VertexSampling.cpp b/Engine/Plugins/FX/Niagara/Source/Niagara/Private/NDISkeletalMesh_VertexSampling.cpp
--- a/Engine/Plugins/FX/Niagara/Source/Niagara/Private/NDISkeletalMesh_VertexSampling.cpp
+++ b/Engine/Plugins/FX/Niagara/Source/Niagara/Private/NDISkeletalMesh_VertexSampling.cpp
@@ -38,6 +38,7 @@ void UNiagaraDataInterfaceSkeletalMesh::GetVertexSamplingFunctions(TArray<FNiaga
 		FNiagaraFunctionSignature Sig;
 		Sig.Name = FSkeletalMeshInterfaceHelper::RandomVertexName;
 		Sig.Inputs.Add(FNiagaraVariable(FNiagaraTypeDefinition(GetClass()), TEXT("SkeletalMesh")));
+		Sig.Inputs.Add(FNiagaraVariable(FNiagaraTypeDefinition(FNiagaraRandInfo::StaticStruct()), TEXT("RandomInfo")));
 		Sig.Outputs.Add(FNiagaraVariable(FNiagaraTypeDefinition::GetIntDef(), TEXT("Vertex")));
 		Sig.bMemberFunction = true;
 		Sig.bRequiresContext = false;
@@ -188,7 +189,7 @@ void UNiagaraDataInterfaceSkeletalMesh::BindVertexSamplingFunction(const FVMExte
 }
 
 template<typename FilterMode>
-FORCEINLINE int32 UNiagaraDataInterfaceSkeletalMesh::RandomVertIndex(FRandomStream& RandStream, FSkeletalMeshAccessorHelper& Accessor, FNDISkeletalMesh_InstanceData* InstData)
+FORCEINLINE int32 UNiagaraDataInterfaceSkeletalMesh::RandomVertIndex(FNDIRandomHelper& RandHelper, int32 Instance, FSkeletalMeshAccessorHelper& Accessor, FNDISkeletalMesh_InstanceData* InstData)
 {
 	checkf(false, TEXT("Invalid template call for RandomVertIndex. Bug in Filter binding or Area Weighting binding. Contact code team."));
 	return 0;
@@ -196,28 +197,28 @@ FORCEINLINE int32 UNiagaraDataInterfaceSkeletalMesh::RandomVertIndex(FRandomStre
 
 template<>
 FORCEINLINE int32 UNiagaraDataInterfaceSkeletalMesh::RandomVertIndex<TIntegralConstant<ENDISkeletalMesh_FilterMode, ENDISkeletalMesh_FilterMode::None>>
-	(FRandomStream& RandStream, FSkeletalMeshAccessorHelper& Accessor, FNDISkeletalMesh_InstanceData* InstData)
+	(FNDIRandomHelper& RandHelper, int32 Instance, FSkeletalMeshAccessorHelper& Accessor, FNDISkeletalMesh_InstanceData* InstData)
 {
-	return RandStream.RandRange(0, Accessor.LODData->GetNumVertices() - 1);
+	return RandHelper.RandRange(Instance, 0, Accessor.LODData->GetNumVertices() - 1);
 }
 
 template<>
 FORCEINLINE int32 UNiagaraDataInterfaceSkeletalMesh::RandomVertIndex<TIntegralConstant<ENDISkeletalMesh_FilterMode, ENDISkeletalMesh_FilterMode::SingleRegion>>
-	(FRandomStream& RandStream, FSkeletalMeshAccessorHelper& Accessor, FNDISkeletalMesh_InstanceData* InstData)
+	(FNDIRandomHelper& RandHelper, int32 Instance, FSkeletalMeshAccessorHelper& Accessor, FNDISkeletalMesh_InstanceData* InstData)
 {
-	int32 Idx = RandStream.RandRange(0, Accessor.SamplingRegionBuiltData->Vertices.Num() - 1);
+	int32 Idx = RandHelper.RandRange(Instance, 0, Accessor.SamplingRegionBuiltData->Vertices.Num() - 1);
 	return Accessor.SamplingRegionBuiltData->Vertices[Idx];
 }
 
 template<>
 FORCEINLINE int32 UNiagaraDataInterfaceSkeletalMesh::RandomVertIndex<TIntegralConstant<ENDISkeletalMesh_FilterMode, ENDISkeletalMesh_FilterMode::MultiRegion>>
-	(FRandomStream& RandStream, FSkeletalMeshAccessorHelper& Accessor, FNDISkeletalMesh_InstanceData* InstData)
+	(FNDIRandomHelper& RandHelper, int32 Instance, FSkeletalMeshAccessorHelper& Accessor, FNDISkeletalMesh_InstanceData* InstData)
 {
-	int32 RegionIdx = RandStream.RandRange(0, InstData->SamplingRegionIndices.Num() - 1);
+	int32 RegionIdx = RandHelper.RandRange(Instance, 0, InstData->SamplingRegionIndices.Num() - 1);
 	const FSkeletalMeshSamplingInfo& SamplingInfo = InstData->Mesh->GetSamplingInfo();
 	const FSkeletalMeshSamplingRegion& Region = SamplingInfo.GetRegion(InstData->SamplingRegionIndices[RegionIdx]);
 	const FSkeletalMeshSamplingRegionBuiltData& RegionBuiltData = SamplingInfo.GetRegionBuiltData(InstData->SamplingRegionIndices[RegionIdx]);
-	int32 Idx = RandStream.RandRange(0, RegionBuiltData.Vertices.Num() - 1);
+	int32 Idx = RandHelper.RandRange(Instance, 0, RegionBuiltData.Vertices.Num() - 1);
 	return RegionBuiltData.Vertices[Idx];
 }
 
@@ -225,6 +226,8 @@ template<typename FilterMode>
 void UNiagaraDataInterfaceSkeletalMesh::RandomVertex(FVectorVMContext& Context)
 {
 	SCOPE_CYCLE_COUNTER(STAT_NiagaraSkel_Vertex_Sample);
+
+	FNDIRandomHelper RandHelper(Context);
 	VectorVM::FUserPtrHandler<FNDISkeletalMesh_InstanceData> InstData(Context);
 	checkfSlow(InstData.Get(), TEXT("Skeletal Mesh Interface has invalid instance data. %s"), *GetPathName());
 	checkfSlow(InstData->Mesh, TEXT("Skeletal Mesh Interface has invalid mesh. %s"), *GetPathName());
@@ -236,7 +239,8 @@ void UNiagaraDataInterfaceSkeletalMesh::RandomVertex(FVectorVMContext& Context)
 
 	for (int32 i = 0; i < Context.NumInstances; ++i)
 	{
-		*OutVert.GetDestAndAdvance() = RandomVertIndex<FilterMode>(Context.RandStream, MeshAccessor, InstData);
+		RandHelper.GetAndAdvance();
+		*OutVert.GetDestAndAdvance() = RandomVertIndex<FilterMode>(RandHelper, i, MeshAccessor, InstData);
 	}
 }
 
@@ -272,28 +276,28 @@ void UNiagaraDataInterfaceSkeletalMesh::IsValidVertex(FVectorVMContext& Context)
 //////////////////////////////////////////////////////////////////////////
 
 template<typename FilterMode>
-FORCEINLINE int32 UNiagaraDataInterfaceSkeletalMesh::GetSpecificVertexCount(FSkeletalMeshAccessorHelper& Accessor, FNDISkeletalMesh_InstanceData* InstData)
+FORCEINLINE int32 UNiagaraDataInterfaceSkeletalMesh::GetFilteredVertexCount(FSkeletalMeshAccessorHelper& Accessor, FNDISkeletalMesh_InstanceData* InstData)
 {
-	checkf(false, TEXT("Invalid template call for GetSpecificVertexCount. Bug in Filter binding or Area Weighting binding. Contact code team."));
+	checkf(false, TEXT("Invalid template call for GetFilteredVertexCount. Bug in Filter binding or Area Weighting binding. Contact code team."));
 	return 0;
 }
 
 template<>
-FORCEINLINE_DEBUGGABLE int32 UNiagaraDataInterfaceSkeletalMesh::GetSpecificVertexCount<TIntegralConstant<ENDISkeletalMesh_FilterMode, ENDISkeletalMesh_FilterMode::None>>
+FORCEINLINE_DEBUGGABLE int32 UNiagaraDataInterfaceSkeletalMesh::GetFilteredVertexCount<TIntegralConstant<ENDISkeletalMesh_FilterMode, ENDISkeletalMesh_FilterMode::None>>
 	(FSkeletalMeshAccessorHelper& Accessor, FNDISkeletalMesh_InstanceData* InstData)
 {
 	return  Accessor.LODData->GetNumVertices();;
 }
 
 template<>
-FORCEINLINE_DEBUGGABLE int32 UNiagaraDataInterfaceSkeletalMesh::GetSpecificVertexCount<TIntegralConstant<ENDISkeletalMesh_FilterMode, ENDISkeletalMesh_FilterMode::SingleRegion>>
+FORCEINLINE_DEBUGGABLE int32 UNiagaraDataInterfaceSkeletalMesh::GetFilteredVertexCount<TIntegralConstant<ENDISkeletalMesh_FilterMode, ENDISkeletalMesh_FilterMode::SingleRegion>>
 	(FSkeletalMeshAccessorHelper& Accessor, FNDISkeletalMesh_InstanceData* InstData)
 {
 	return Accessor.SamplingRegionBuiltData->Vertices.Num();
 }
 
 template<>
-FORCEINLINE_DEBUGGABLE int32 UNiagaraDataInterfaceSkeletalMesh::GetSpecificVertexCount<TIntegralConstant<ENDISkeletalMesh_FilterMode, ENDISkeletalMesh_FilterMode::MultiRegion>>
+FORCEINLINE_DEBUGGABLE int32 UNiagaraDataInterfaceSkeletalMesh::GetFilteredVertexCount<TIntegralConstant<ENDISkeletalMesh_FilterMode, ENDISkeletalMesh_FilterMode::MultiRegion>>
 	(FSkeletalMeshAccessorHelper& Accessor, FNDISkeletalMesh_InstanceData* InstData)
 {
 	int32 NumVerts = 0;
@@ -320,7 +324,7 @@ void UNiagaraDataInterfaceSkeletalMesh::GetFilteredVertexCount(FVectorVMContext&
 	FSkeletalMeshAccessorHelper MeshAccessor;
 	MeshAccessor.Init<FilterMode, TIntegralConstant<ENDISkelMesh_AreaWeightingMode, ENDISkelMesh_AreaWeightingMode::None>>(InstData);
 
-	int32 Count = GetSpecificVertexCount<FilterMode>(MeshAccessor, InstData);
+	int32 Count = GetFilteredVertexCount<FilterMode>(MeshAccessor, InstData);
 	for (int32 i = 0; i < Context.NumInstances; ++i)
 	{
 		*OutVert.GetDestAndAdvance() = Count;
@@ -331,21 +335,21 @@ void UNiagaraDataInterfaceSkeletalMesh::GetFilteredVertexCount(FVectorVMContext&
 //////////////////////////////////////////////////////////////////////////
 
 template<typename FilterMode>
-FORCEINLINE int32 UNiagaraDataInterfaceSkeletalMesh::GetSpecificVertexAt(FSkeletalMeshAccessorHelper& Accessor, FNDISkeletalMesh_InstanceData* InstData, int32 FilteredIndex)
+FORCEINLINE int32 UNiagaraDataInterfaceSkeletalMesh::GetFilteredVertexAt(FSkeletalMeshAccessorHelper& Accessor, FNDISkeletalMesh_InstanceData* InstData, int32 FilteredIndex)
 {
-	checkf(false, TEXT("Invalid template call for GetSpecificVertexAt. Bug in Filter binding. Contact code team."));
+	checkf(false, TEXT("Invalid template call for GetFilteredVertexAt. Bug in Filter binding. Contact code team."));
 	return 0;
 }
 
 template<>
-FORCEINLINE_DEBUGGABLE int32 UNiagaraDataInterfaceSkeletalMesh::GetSpecificVertexAt<TIntegralConstant<ENDISkeletalMesh_FilterMode, ENDISkeletalMesh_FilterMode::None>>
+FORCEINLINE_DEBUGGABLE int32 UNiagaraDataInterfaceSkeletalMesh::GetFilteredVertexAt<TIntegralConstant<ENDISkeletalMesh_FilterMode, ENDISkeletalMesh_FilterMode::None>>
 	(FSkeletalMeshAccessorHelper& Accessor, FNDISkeletalMesh_InstanceData* InstData, int32 FilteredIndex)
 {
 	return FilteredIndex;
 }
 
 template<>
-FORCEINLINE_DEBUGGABLE int32 UNiagaraDataInterfaceSkeletalMesh::GetSpecificVertexAt<TIntegralConstant<ENDISkeletalMesh_FilterMode, ENDISkeletalMesh_FilterMode::SingleRegion>>
+FORCEINLINE_DEBUGGABLE int32 UNiagaraDataInterfaceSkeletalMesh::GetFilteredVertexAt<TIntegralConstant<ENDISkeletalMesh_FilterMode, ENDISkeletalMesh_FilterMode::SingleRegion>>
 	(FSkeletalMeshAccessorHelper& Accessor, FNDISkeletalMesh_InstanceData* InstData, int32 FilteredIndex)
 {
 	int32 MaxIdx = Accessor.SamplingRegionBuiltData->Vertices.Num() - 1;
@@ -354,7 +358,7 @@ FORCEINLINE_DEBUGGABLE int32 UNiagaraDataInterfaceSkeletalMesh::GetSpecificVerte
 }
 
 template<>
-FORCEINLINE_DEBUGGABLE int32 UNiagaraDataInterfaceSkeletalMesh::GetSpecificVertexAt<TIntegralConstant<ENDISkeletalMesh_FilterMode, ENDISkeletalMesh_FilterMode::MultiRegion>>
+FORCEINLINE_DEBUGGABLE int32 UNiagaraDataInterfaceSkeletalMesh::GetFilteredVertexAt<TIntegralConstant<ENDISkeletalMesh_FilterMode, ENDISkeletalMesh_FilterMode::MultiRegion>>
 	(FSkeletalMeshAccessorHelper& Accessor, FNDISkeletalMesh_InstanceData* InstData, int32 FilteredIndex)
 {
 	for (int32 RegionIdx = 0; RegionIdx < InstData->SamplingRegionIndices.Num(); RegionIdx++)
@@ -391,7 +395,7 @@ void UNiagaraDataInterfaceSkeletalMesh::GetFilteredVertexAt(FVectorVMContext& Co
 	{
 		int32 FilteredVert = FilteredVertexParam.GetAndAdvance();
 		int32 RealIdx = 0;
-		RealIdx = GetSpecificVertexAt<FilterMode>(Accessor, InstData, FilteredVert);
+		RealIdx = GetFilteredVertexAt<FilterMode>(Accessor, InstData, FilteredVert);
 		RealIdx = FMath::Min(RealIdx, VertMax);
 
 		*OutVert.GetDestAndAdvance() = RealIdx;
diff --git a/Engine/Plugins/FX/Niagara/Source/Niagara/Private/NiagaraDataInterfaceSkeletalMesh.cpp b/Engine/Plugins/FX/Niagara/Source/Niagara/Private/NiagaraDataInterfaceSkeletalMesh.cpp
--- a/Engine/Plugins/FX/Niagara/Source/Niagara/Private/NiagaraDataInterfaceSkeletalMesh.cpp
+++ b/Engine/Plugins/FX/Niagara/Source/Niagara/Private/NiagaraDataInterfaceSkeletalMesh.cpp
@@ -22,6 +22,18 @@
 
 DECLARE_CYCLE_STAT(TEXT("PreSkin"), STAT_NiagaraSkel_PreSkin, STATGROUP_Niagara);
 
+struct FNiagaraSkelMeshDIFunctionVersion
+{
+	enum Type
+	{
+		InitialVersion = 0,
+		AddedRandomInfo = 1,
+
+		VersionPlusOne,
+		LatestVersion = VersionPlusOne - 1
+	};
+};
+
 //////////////////////////////////////////////////////////////////////////
 
 FSkeletalMeshSamplingRegionAreaWeightedSampler::FSkeletalMeshSamplingRegionAreaWeightedSampler()
@@ -458,10 +470,10 @@ void FSkeletalMeshGpuSpawnStaticBuffers::Initialise(FNDISkeletalMesh_InstanceDat
 		TriangleCount = 0;
 		VertexCount = 0;
 
-		NumSpecificBones = 0;
-		SpecificBonesArray.Empty();
-		NumSpecificSockets = 0;
-		SpecificSocketBoneOffset = 0;
+		NumFilteredBones = 0;
+		FilteredBonesArray.Empty();
+		NumFilteredSockets = 0;
+		FilteredSocketBoneOffset = 0;
 		return;
 	}
 	else
@@ -482,21 +494,21 @@ void FSkeletalMeshGpuSpawnStaticBuffers::Initialise(FNDISkeletalMesh_InstanceDat
 			UE_LOG(LogNiagara, Warning, TEXT("FSkeletalMeshGpuSpawnStaticBuffers> Vertex count is invalid %d"), VertexCount, (InstData && InstData->Mesh) ? *InstData->Mesh->GetFullName() : TEXT("Unknown Mesh"));
 		}
 
-		// Copy Specific Bones / Socket data into arrays that the renderer will use to create read buffers
+		// Copy filtered Bones / Socket data into arrays that the renderer will use to create read buffers
 		//-TODO: Exclude setting up these arrays if we don't sample from them
-		NumSpecificBones = InstData->SpecificBones.Num();
-		if (NumSpecificBones > 0)
+		NumFilteredBones = InstData->FilteredBones.Num();
+		if (NumFilteredBones > 0)
 		{
-			SpecificBonesArray.Reserve(NumSpecificBones);
-			for (int32 v : InstData->SpecificBones)
+			FilteredBonesArray.Reserve(NumFilteredBones);
+			for (int32 v : InstData->FilteredBones)
 			{
 				check(v <= 65535);
-				SpecificBonesArray.Add(v);
+				FilteredBonesArray.Add(v);
 			}
 		}
 
-		NumSpecificSockets = InstData->SpecificSocketInfo.Num();
-		SpecificSocketBoneOffset = InstData->SpecificSocketBoneOffset;
+		NumFilteredSockets = InstData->FilteredSocketInfo.Num();
+		FilteredSocketBoneOffset = InstData->FilteredSocketBoneOffset;
 	}
 }
 
@@ -572,21 +584,21 @@ void FSkeletalMeshGpuSpawnStaticBuffers::InitRHI()
 		BufferTriangleMatricesOffsetSRV = RHICreateShaderResourceView(BufferTriangleMatricesOffsetRHI, sizeof(uint32), PF_R32_UINT);
 	}
 
-	// Create arrays for specific bones / sockets
-	if ( NumSpecificBones > 0 )
+	// Create arrays for filtered bones / sockets
+	if ( NumFilteredBones > 0 )
 	{
 		FRHIResourceCreateInfo CreateInfo;
-		CreateInfo.ResourceArray = &SpecificBonesArray;
+		CreateInfo.ResourceArray = &FilteredBonesArray;
 
-		SpecificBonesBuffer = RHICreateVertexBuffer(NumSpecificBones * sizeof(uint16), BUF_Static | BUF_ShaderResource, CreateInfo);
-		SpecificBonesSRV = RHICreateShaderResourceView(SpecificBonesBuffer, sizeof(uint16), PF_R16_UINT);
+		FilteredBonesBuffer = RHICreateVertexBuffer(NumFilteredBones * sizeof(uint16), BUF_Static | BUF_ShaderResource, CreateInfo);
+		FilteredBonesSRV = RHICreateShaderResourceView(FilteredBonesBuffer, sizeof(uint16), PF_R16_UINT);
 	}
 }
 
 void FSkeletalMeshGpuSpawnStaticBuffers::ReleaseRHI()
 {
-	SpecificBonesBuffer.SafeRelease();
-	SpecificBonesSRV.SafeRelease();
+	FilteredBonesBuffer.SafeRelease();
+	FilteredBonesSRV.SafeRelease();
 
 	BufferTriangleUniformSamplerProbaRHI.SafeRelease();
 	BufferTriangleUniformSamplerProbaSRV.SafeRelease();
@@ -724,7 +736,7 @@ void FSkeletalMeshGpuDynamicBufferProxy::NewFrame(const FNDISkeletalMesh_Instanc
 			}
 
 			// Append sockets
-			for (const FTransform& SocketTransform : InstanceData->GetSpecificSocketsCurrBuffer())
+			for (const FTransform& SocketTransform : InstanceData->GetFilteredSocketsCurrBuffer())
 			{
 				const FQuat Rotation = SocketTransform.GetRotation();
 				BoneSamplingData.Add(SocketTransform.GetLocation());
@@ -852,10 +864,11 @@ struct FNDISkeletalMeshParametersName
 	FString MeshSkinWeightIndexSizeName;
 	FString MeshNumTexCoordName;
 	FString MeshNumWeightsName;
-	FString NumSpecificBonesName;
-	FString SpecificBonesName;
-	FString NumSpecificSocketsName;
-	FString SpecificSocketBoneOffsetName;
+	FString NumBonesName;
+	FString NumFilteredBonesName;
+	FString FilteredBonesName;
+	FString NumFilteredSocketsName;
+	FString FilteredSocketBoneOffsetName;
 	FString InstanceTransformName;
 	FString InstancePrevTransformName;
 	FString InstanceInvDeltaTimeName;
@@ -884,10 +897,11 @@ static void GetNiagaraDataInterfaceParametersName(FNDISkeletalMeshParametersName
 	Names.MeshSkinWeightIndexSizeName = UNiagaraDataInterfaceSkeletalMesh::MeshSkinWeightIndexSizeName + Suffix;
 	Names.MeshNumTexCoordName = UNiagaraDataInterfaceSkeletalMesh::MeshNumTexCoordName + Suffix;
 	Names.MeshNumWeightsName = UNiagaraDataInterfaceSkeletalMesh::MeshNumWeightsName + Suffix;
-	Names.NumSpecificBonesName = UNiagaraDataInterfaceSkeletalMesh::NumSpecificBonesName + Suffix;
-	Names.SpecificBonesName = UNiagaraDataInterfaceSkeletalMesh::SpecificBonesName + Suffix;
-	Names.NumSpecificSocketsName = UNiagaraDataInterfaceSkeletalMesh::NumSpecificSocketsName + Suffix;
-	Names.SpecificSocketBoneOffsetName = UNiagaraDataInterfaceSkeletalMesh::SpecificSocketBoneOffsetName + Suffix;
+	Names.NumBonesName = UNiagaraDataInterfaceSkeletalMesh::NumBonesName + Suffix;
+	Names.NumFilteredBonesName = UNiagaraDataInterfaceSkeletalMesh::NumFilteredBonesName + Suffix;
+	Names.FilteredBonesName = UNiagaraDataInterfaceSkeletalMesh::FilteredBonesName + Suffix;
+	Names.NumFilteredSocketsName = UNiagaraDataInterfaceSkeletalMesh::NumFilteredSocketsName + Suffix;
+	Names.FilteredSocketBoneOffsetName = UNiagaraDataInterfaceSkeletalMesh::FilteredSocketBoneOffsetName + Suffix;
 	Names.InstanceTransformName = UNiagaraDataInterfaceSkeletalMesh::InstanceTransformName + Suffix;
 	Names.InstancePrevTransformName = UNiagaraDataInterfaceSkeletalMesh::InstancePrevTransformName + Suffix;
 	Names.InstanceInvDeltaTimeName = UNiagaraDataInterfaceSkeletalMesh::InstanceInvDeltaTimeName + Suffix;
@@ -923,10 +937,11 @@ public:
 		MeshSkinWeightIndexSize.Bind(ParameterMap, *ParamNames.MeshSkinWeightIndexSizeName);
 		MeshNumTexCoord.Bind(ParameterMap, *ParamNames.MeshNumTexCoordName);
 		MeshNumWeights.Bind(ParameterMap, *ParamNames.MeshNumWeightsName);
-		NumSpecificBones.Bind(ParameterMap, *ParamNames.NumSpecificBonesName);
-		SpecificBones.Bind(ParameterMap, *ParamNames.SpecificBonesName);
-		NumSpecificSockets.Bind(ParameterMap, *ParamNames.NumSpecificSocketsName);
-		SpecificSocketBoneOffset.Bind(ParameterMap, *ParamNames.SpecificSocketBoneOffsetName);
+		NumBones.Bind(ParameterMap, *ParamNames.NumBonesName);
+		NumFilteredBones.Bind(ParameterMap, *ParamNames.NumFilteredBonesName);
+		FilteredBones.Bind(ParameterMap, *ParamNames.FilteredBonesName);
+		NumFilteredSockets.Bind(ParameterMap, *ParamNames.NumFilteredSocketsName);
+		FilteredSocketBoneOffset.Bind(ParameterMap, *ParamNames.FilteredSocketBoneOffsetName);
 		InstanceTransform.Bind(ParameterMap, *ParamNames.InstanceTransformName);
 		InstancePrevTransform.Bind(ParameterMap, *ParamNames.InstancePrevTransformName);
 		InstanceInvDeltaTime.Bind(ParameterMap, *ParamNames.InstanceInvDeltaTimeName);
@@ -1009,12 +1024,14 @@ public:
 				SetSRVParameter(RHICmdList, ComputeShaderRHI, MeshTriangleMatricesOffsetBuffer, FNiagaraRenderer::GetDummyUIntBuffer());
 			}
 
-			FRHIShaderResourceView* SpecificBonesSRV = StaticBuffers->GetNumSpecificBones() > 0 ? StaticBuffers->GetSpecificBonesSRV() : FNiagaraRenderer::GetDummyUIntBuffer();
-			SetShaderValue(RHICmdList, ComputeShaderRHI, NumSpecificBones, StaticBuffers->GetNumSpecificBones());
-			SetSRVParameter(RHICmdList, ComputeShaderRHI, SpecificBones, SpecificBonesSRV);
+			SetShaderValue(RHICmdList, ComputeShaderRHI, NumBones, DynamicBuffers->GetNumBones());
 
-			SetShaderValue(RHICmdList, ComputeShaderRHI, NumSpecificSockets, StaticBuffers->GetNumSpecificSockets());
-			SetShaderValue(RHICmdList, ComputeShaderRHI, SpecificSocketBoneOffset, StaticBuffers->GetSpecificSocketBoneOffset());
+			FRHIShaderResourceView* FilteredBonesSRV = StaticBuffers->GetNumFilteredBones() > 0 ? StaticBuffers->GetFilteredBonesSRV() : FNiagaraRenderer::GetDummyUIntBuffer();
+			SetShaderValue(RHICmdList, ComputeShaderRHI, NumFilteredBones, StaticBuffers->GetNumFilteredBones());
+			SetSRVParameter(RHICmdList, ComputeShaderRHI, FilteredBones, FilteredBonesSRV);
+
+			SetShaderValue(RHICmdList, ComputeShaderRHI, NumFilteredSockets, StaticBuffers->GetNumFilteredSockets());
+			SetShaderValue(RHICmdList, ComputeShaderRHI, FilteredSocketBoneOffset, StaticBuffers->GetFilteredSocketBoneOffset());
 
 			SetShaderValue(RHICmdList, ComputeShaderRHI, InstanceTransform, InstanceData->Transform);
 			SetShaderValue(RHICmdList, ComputeShaderRHI, InstancePrevTransform, InstanceData->PrevTransform);
@@ -1052,10 +1069,11 @@ public:
 			SetSRVParameter(RHICmdList, ComputeShaderRHI, MeshPrevSamplingBonesBuffer, FNiagaraRenderer::GetDummyFloat4Buffer());
 			SetSRVParameter(RHICmdList, ComputeShaderRHI, MeshTriangleMatricesOffsetBuffer, FNiagaraRenderer::GetDummyUIntBuffer());
 
-			SetShaderValue(RHICmdList, ComputeShaderRHI, NumSpecificBones, 0);
-			SetSRVParameter(RHICmdList, ComputeShaderRHI, SpecificBones, FNiagaraRenderer::GetDummyUIntBuffer());
-			SetShaderValue(RHICmdList, ComputeShaderRHI, NumSpecificSockets, 0);
-			SetShaderValue(RHICmdList, ComputeShaderRHI, SpecificSocketBoneOffset, 0);
+			SetShaderValue(RHICmdList, ComputeShaderRHI, NumBones, 0);
+			SetShaderValue(RHICmdList, ComputeShaderRHI, NumFilteredBones, 0);
+			SetSRVParameter(RHICmdList, ComputeShaderRHI, FilteredBones, FNiagaraRenderer::GetDummyUIntBuffer());
+			SetShaderValue(RHICmdList, ComputeShaderRHI, NumFilteredSockets, 0);
+			SetShaderValue(RHICmdList, ComputeShaderRHI, FilteredSocketBoneOffset, 0);
 
 			SetShaderValue(RHICmdList, ComputeShaderRHI, InstanceTransform, FMatrix::Identity);
 			SetShaderValue(RHICmdList, ComputeShaderRHI, InstancePrevTransform, FMatrix::Identity);
@@ -1086,10 +1104,11 @@ private:
 	LAYOUT_FIELD(FShaderParameter, MeshSkinWeightIndexSize);
 	LAYOUT_FIELD(FShaderParameter, MeshNumTexCoord);
 	LAYOUT_FIELD(FShaderParameter, MeshNumWeights);
-	LAYOUT_FIELD(FShaderParameter, NumSpecificBones);
-	LAYOUT_FIELD(FShaderResourceParameter, SpecificBones);
-	LAYOUT_FIELD(FShaderParameter, NumSpecificSockets);
-	LAYOUT_FIELD(FShaderParameter, SpecificSocketBoneOffset);
+	LAYOUT_FIELD(FShaderParameter, NumBones);
+	LAYOUT_FIELD(FShaderParameter, NumFilteredBones);
+	LAYOUT_FIELD(FShaderResourceParameter, FilteredBones);
+	LAYOUT_FIELD(FShaderParameter, NumFilteredSockets);
+	LAYOUT_FIELD(FShaderParameter, FilteredSocketBoneOffset);
 	LAYOUT_FIELD(FShaderParameter, InstanceTransform);
 	LAYOUT_FIELD(FShaderParameter, InstancePrevTransform);
 	LAYOUT_FIELD(FShaderParameter, InstanceInvDeltaTime);
@@ -1479,23 +1498,23 @@ bool FNDISkeletalMesh_InstanceData::Init(UNiagaraDataInterfaceSkeletalMesh* Inte
 		}
 	}
 
-	// Gather specific bones information
+	// Gather filtered bones information
 	FReferenceSkeleton& RefSkel = Mesh->RefSkeleton;
 	{
-		SpecificBones.SetNumUninitialized(Interface->SpecificBones.Num());
+		FilteredBones.SetNumUninitialized(Interface->FilteredBones.Num());
 		TArray<FName, TInlineAllocator<16>> MissingBones;
-		for (int32 BoneIdx = 0; BoneIdx < SpecificBones.Num(); ++BoneIdx)
+		for (int32 BoneIdx = 0; BoneIdx < FilteredBones.Num(); ++BoneIdx)
 		{
-			FName BoneName = Interface->SpecificBones[BoneIdx];
+			FName BoneName = Interface->FilteredBones[BoneIdx];
 			int32 Bone = RefSkel.FindBoneIndex(BoneName);
 			if (Bone == INDEX_NONE)
 			{
 				MissingBones.Add(BoneName);
-				SpecificBones[BoneIdx] = 0;
+				FilteredBones[BoneIdx] = 0;
 			}
 			else
 			{
-				SpecificBones[BoneIdx] = Bone;
+				FilteredBones[BoneIdx] = Bone;
 			}
 		}
 
@@ -1509,42 +1528,42 @@ bool FNDISkeletalMesh_InstanceData::Init(UNiagaraDataInterfaceSkeletalMesh* Inte
 		}
 	}
 
-	// Gather specific socket information
+	// Gather filtered socket information
 	{
-		TArray<FName>& SpecificSockets = Interface->SpecificSockets;
-		SpecificSocketInfo.SetNum(SpecificSockets.Num());
+		TArray<FName>& FilteredSockets = Interface->FilteredSockets;
+		FilteredSocketInfo.SetNum(FilteredSockets.Num());
 
 		//-TODO: We may need to handle skinning mode changes here
 		if (NewSkelComp != nullptr)
 		{
-			for (int32 i = 0; i < SpecificSocketInfo.Num(); ++i)
+			for (int32 i = 0; i < FilteredSocketInfo.Num(); ++i)
 			{
-				NewSkelComp->GetSocketInfoByName(SpecificSockets[i], SpecificSocketInfo[i].Transform, SpecificSocketInfo[i].BoneIdx);
+				NewSkelComp->GetSocketInfoByName(FilteredSockets[i], FilteredSocketInfo[i].Transform, FilteredSocketInfo[i].BoneIdx);
 			}
 		}
 		else
 		{
-			for (int32 i = 0; i < SpecificSocketInfo.Num(); ++i)
+			for (int32 i = 0; i < FilteredSocketInfo.Num(); ++i)
 			{
-				SpecificSocketInfo[i].Transform = FTransform(Mesh->GetComposedRefPoseMatrix(SpecificSockets[i]));
-				SpecificSocketInfo[i].BoneIdx = INDEX_NONE;
+				FilteredSocketInfo[i].Transform = FTransform(Mesh->GetComposedRefPoseMatrix(FilteredSockets[i]));
+				FilteredSocketInfo[i].BoneIdx = INDEX_NONE;
 			}
 		}
 
-		SpecificSocketBoneOffset = Mesh->RefSkeleton.GetNum();
+		FilteredSocketBoneOffset = Mesh->RefSkeleton.GetNum();
 
-		SpecificSocketTransformsIndex = 0;
-		SpecificSocketTransforms[0].Reset(SpecificSockets.Num());
-		SpecificSocketTransforms[0].AddDefaulted(SpecificSockets.Num());
-		UpdateSpecificSocketTransforms();
-		for (int32 i=1; i < SpecificSocketTransforms.Num(); ++i)
+		FilteredSocketTransformsIndex = 0;
+		FilteredSocketTransforms[0].Reset(FilteredSockets.Num());
+		FilteredSocketTransforms[0].AddDefaulted(FilteredSockets.Num());
+		UpdateFilteredSocketTransforms();
+		for (int32 i=1; i < FilteredSocketTransforms.Num(); ++i)
 		{
-			SpecificSocketTransforms[i].Reset(SpecificSockets.Num());
-			SpecificSocketTransforms[i].Append(SpecificSocketTransforms[0]);
+			FilteredSocketTransforms[i].Reset(FilteredSockets.Num());
+			FilteredSocketTransforms[i].Append(FilteredSocketTransforms[0]);
 		}
 
 		TArray<FName, TInlineAllocator<16>> MissingSockets;
-		for (FName SocketName : SpecificSockets)
+		for (FName SocketName : FilteredSockets)
 		{
 			if (Mesh->FindSocket(SocketName) == nullptr)
 			{
@@ -1590,7 +1609,7 @@ bool FNDISkeletalMesh_InstanceData::Init(UNiagaraDataInterfaceSkeletalMesh* Inte
 		BeginInitResource(MeshGpuSpawnStaticBuffers);
 
 		MeshGpuSpawnDynamicBuffers = new FSkeletalMeshGpuDynamicBufferProxy();
-		MeshGpuSpawnDynamicBuffers->Initialise(RefSkel, LODData, SpecificSocketInfo.Num());
+		MeshGpuSpawnDynamicBuffers->Initialise(RefSkel, LODData, FilteredSocketInfo.Num());
 		BeginInitResource(MeshGpuSpawnDynamicBuffers);
 	}
 
@@ -1683,8 +1702,8 @@ bool FNDISkeletalMesh_InstanceData::Tick(UNiagaraDataInterfaceSkeletalMesh* Inte
 		}
 
 		// Cache socket transforms to avoid potentially calculating them multiple times during the VM
-		SpecificSocketTransformsIndex = (SpecificSocketTransformsIndex + 1) % SpecificSocketTransforms.Num();
-		UpdateSpecificSocketTransforms();
+		FilteredSocketTransformsIndex = (FilteredSocketTransformsIndex + 1) % FilteredSocketTransforms.Num();
+		UpdateFilteredSocketTransforms();
 
 		if (MeshGpuSpawnDynamicBuffers)
 		{
@@ -1702,14 +1721,14 @@ bool FNDISkeletalMesh_InstanceData::Tick(UNiagaraDataInterfaceSkeletalMesh* Inte
 	}
 }
 
-void FNDISkeletalMesh_InstanceData::UpdateSpecificSocketTransforms()
+void FNDISkeletalMesh_InstanceData::UpdateFilteredSocketTransforms()
 {
 	USkeletalMeshComponent* SkelComp = Cast<USkeletalMeshComponent>(Component.Get());
-	TArray<FTransform>& WriteBuffer = GetSpecificSocketsWriteBuffer();
+	TArray<FTransform>& WriteBuffer = GetFilteredSocketsWriteBuffer();
 
-	for (int32 i = 0; i < SpecificSocketInfo.Num(); ++i)
+	for (int32 i = 0; i < FilteredSocketInfo.Num(); ++i)
 	{
-		const FCachedSocketInfo& SocketInfo = SpecificSocketInfo[i];
+		const FCachedSocketInfo& SocketInfo = FilteredSocketInfo[i];
 		const FTransform& BoneTransform = SocketInfo.BoneIdx != INDEX_NONE ? SkelComp->GetBoneTransform(SocketInfo.BoneIdx, FTransform::Identity) : FTransform::Identity;
 		WriteBuffer[i] = SocketInfo.Transform * BoneTransform;
 	}
@@ -1809,6 +1828,13 @@ void UNiagaraDataInterfaceSkeletalMesh::GetFunctions(TArray<FNiagaraFunctionSign
 	GetTriangleSamplingFunctions(OutFunctions);
 	GetVertexSamplingFunctions(OutFunctions);
 	GetSkeletonSamplingFunctions(OutFunctions);
+
+#if WITH_EDITORONLY_DATA
+	for (FNiagaraFunctionSignature& Function : OutFunctions)
+	{
+		Function.FunctionVersion = FNiagaraSkelMeshDIFunctionVersion::LatestVersion;
+	}
+#endif
 }
 
 void UNiagaraDataInterfaceSkeletalMesh::GetVMExternalFunction(const FVMExternalFunctionBindingInfo& BindingInfo, void* InstanceData, FVMExternalFunction &OutFunc)
@@ -1878,8 +1904,8 @@ bool UNiagaraDataInterfaceSkeletalMesh::CopyToInternal(UNiagaraDataInterface* De
 	OtherTyped->SkinningMode = SkinningMode;
 	OtherTyped->SamplingRegions = SamplingRegions;
 	OtherTyped->WholeMeshLOD = WholeMeshLOD;
-	OtherTyped->SpecificBones = SpecificBones;
-	OtherTyped->SpecificSockets = SpecificSockets;
+	OtherTyped->FilteredBones = FilteredBones;
+	OtherTyped->FilteredSockets = FilteredSockets;
 #if WITH_EDITORONLY_DATA
 	OtherTyped->PreviewMesh = PreviewMesh;
 	OtherTyped->bRequiresCPUAccess = bRequiresCPUAccess;
@@ -1902,8 +1928,8 @@ bool UNiagaraDataInterfaceSkeletalMesh::Equals(const UNiagaraDataInterface* Othe
 		OtherTyped->SkinningMode == SkinningMode &&
 		OtherTyped->SamplingRegions == SamplingRegions &&
 		OtherTyped->WholeMeshLOD == WholeMeshLOD &&
-		OtherTyped->SpecificBones == SpecificBones &&
-		OtherTyped->SpecificSockets == SpecificSockets;
+		OtherTyped->FilteredBones == FilteredBones &&
+		OtherTyped->FilteredSockets == FilteredSockets;
 }
 
 bool UNiagaraDataInterfaceSkeletalMesh::InitPerInstanceData(void* PerInstanceData, FNiagaraSystemInstance* SystemInstance)
@@ -2023,7 +2049,6 @@ void UNiagaraDataInterfaceSkeletalMesh::ValidateFunction(const FNiagaraFunctionS
 	if (!DIFuncs.Contains(Function))
 	{
 		TArray<FNiagaraFunctionSignature> SkinnedDataDeprecatedFunctions;
-
 		{
 			FNiagaraFunctionSignature Sig;
 			Sig.Name = GetTriPositionName_DEPRECATED;
@@ -2139,10 +2164,11 @@ const FString UNiagaraDataInterfaceSkeletalMesh::MeshWeightStrideName(TEXT("Mesh
 const FString UNiagaraDataInterfaceSkeletalMesh::MeshSkinWeightIndexSizeName(TEXT("MeshSkinWeightIndexSize_"));
 const FString UNiagaraDataInterfaceSkeletalMesh::MeshNumTexCoordName(TEXT("MeshNumTexCoord_"));
 const FString UNiagaraDataInterfaceSkeletalMesh::MeshNumWeightsName(TEXT("MeshNumWeights_"));
-const FString UNiagaraDataInterfaceSkeletalMesh::NumSpecificBonesName(TEXT("NumSpecificBones_"));
-const FString UNiagaraDataInterfaceSkeletalMesh::SpecificBonesName(TEXT("SpecificBones_"));
-const FString UNiagaraDataInterfaceSkeletalMesh::NumSpecificSocketsName(TEXT("NumSpecificSockets_"));
-const FString UNiagaraDataInterfaceSkeletalMesh::SpecificSocketBoneOffsetName(TEXT("SpecificSocketBoneOffset_"));
+const FString UNiagaraDataInterfaceSkeletalMesh::NumBonesName(TEXT("NumBones_"));
+const FString UNiagaraDataInterfaceSkeletalMesh::NumFilteredBonesName(TEXT("NumFilteredBones_"));
+const FString UNiagaraDataInterfaceSkeletalMesh::FilteredBonesName(TEXT("FilteredBones_"));
+const FString UNiagaraDataInterfaceSkeletalMesh::NumFilteredSocketsName(TEXT("NumFilteredSockets_"));
+const FString UNiagaraDataInterfaceSkeletalMesh::FilteredSocketBoneOffsetName(TEXT("FilteredSocketBoneOffset_"));
 const FString UNiagaraDataInterfaceSkeletalMesh::InstanceTransformName(TEXT("InstanceTransform_"));
 const FString UNiagaraDataInterfaceSkeletalMesh::InstancePrevTransformName(TEXT("InstancePrevTransform_"));
 const FString UNiagaraDataInterfaceSkeletalMesh::InstanceInvDeltaTimeName(TEXT("InstanceInvDeltaTime_"));
@@ -2175,6 +2201,7 @@ bool UNiagaraDataInterfaceSkeletalMesh::GetFunctionHLSL(const FNiagaraDataInterf
 		{TEXT("GetDISkelMeshContextName"), TEXT("DISKELMESH_MAKE_CONTEXT(") + ParamInfo.DataInterfaceHLSLSymbol + TEXT(")")},
 	};
 
+	//////////////////////////////////////////////////////////////////////////////////////////////////////////////////////////////
 	// Triangle Sampling
 	if (FunctionInfo.DefinitionName == FSkeletalMeshInterfaceHelper::RandomTriCoordName)
 	{
@@ -2221,6 +2248,7 @@ bool UNiagaraDataInterfaceSkeletalMesh::GetFunctionHLSL(const FNiagaraDataInterf
 		static const TCHAR* FormatSample = TEXT("void {InstanceFunctionName} (in int TriangleIndex, out int OutVertexIndex0, out int OutVertexIndex1, out int OutVertexIndex2) { {GetDISkelMeshContextName} DISkelMesh_GetTriVertices(DIContext, TriangleIndex, OutVertexIndex0, OutVertexIndex1, OutVertexIndex2); }");
 		OutHLSL += FString::Format(FormatSample, ArgsSample);
 	}
+	//////////////////////////////////////////////////////////////////////////////////////////////////////////////////////////////
 	// Bone Sampling
 	else if (FunctionInfo.DefinitionName == FSkeletalMeshInterfaceHelper::GetSkinnedBoneDataName)
 	{
@@ -2242,10 +2270,11 @@ bool UNiagaraDataInterfaceSkeletalMesh::GetFunctionHLSL(const FNiagaraDataInterf
 		static const TCHAR* FormatSample = TEXT("void {InstanceFunctionName} (in int InBone, in float Interp, out float3 OutPosition, out float4 OutRotation, out float3 OutVelocity) { {GetDISkelMeshContextName} DISkelMesh_GetSkinnedBoneInterpolatedWS(DIContext, InBone, Interp, OutPosition, OutRotation, OutVelocity); }");
 		OutHLSL += FString::Format(FormatSample, ArgsSample);
 	}
+	//////////////////////////////////////////////////////////////////////////////////////////////////////////////////////////////
 	// Vertex Sampling
 	else if (FunctionInfo.DefinitionName == FSkeletalMeshInterfaceHelper::RandomVertexName)
 	{
-		static const TCHAR* FormatSample = TEXT("void {InstanceFunctionName}(out int OutVertex) { {GetDISkelMeshContextName} DISkelMesh_GetRandomVertex(DIContext, OutVertex); }");
+		static const TCHAR* FormatSample = TEXT("void {InstanceFunctionName}(NiagaraRandInfo InRandomInfo, out int OutVertex) { {GetDISkelMeshContextName} DISkelMesh_GetRandomVertex(DIContext, InRandomInfo.Seed1, InRandomInfo.Seed2, InRandomInfo.Seed3, OutVertex); }");
 		OutHLSL += FString::Format(FormatSample, ArgsSample);
 	}
 	else if (FunctionInfo.DefinitionName == FSkeletalMeshInterfaceHelper::IsValidVertexName)
@@ -2273,43 +2302,59 @@ bool UNiagaraDataInterfaceSkeletalMesh::GetFunctionHLSL(const FNiagaraDataInterf
 		static const TCHAR* FormatSample = TEXT("void {InstanceFunctionName} (in int Vertex, in int UVSet, out float2 OutUV) { {GetDISkelMeshContextName} DISkelMesh_GetVertexUV(DIContext, Vertex, UVSet, OutUV); }");
 		OutHLSL += FString::Format(FormatSample, ArgsSample);
 	}
-	// Specific Bone
-	else if (FunctionInfo.DefinitionName == FSkeletalMeshInterfaceHelper::GetSpecificBoneCountName)
+	//////////////////////////////////////////////////////////////////////////////////////////////////////////////////////////////
+	// Filtered Bone
+	else if (FunctionInfo.DefinitionName == FSkeletalMeshInterfaceHelper::IsValidBoneName)
+	{
+		static const TCHAR* FormatSample = TEXT("void {InstanceFunctionName} (in int BoneIndex, out bool IsValid) { {GetDISkelMeshContextName} DISkelMesh_IsValidBone(DIContext, BoneIndex, IsValid); }");
+		OutHLSL += FString::Format(FormatSample, ArgsSample);
+	}
+	else if (FunctionInfo.DefinitionName == FSkeletalMeshInterfaceHelper::RandomBoneName)
+	{
+		static const TCHAR* FormatSample = TEXT("void {InstanceFunctionName} (NiagaraRandInfo InRandomInfo, out int Bone) { {GetDISkelMeshContextName} DISkelMesh_RandomBone(DIContext, InRandomInfo.Seed1, InRandomInfo.Seed2, InRandomInfo.Seed3, Bone); }");
+		OutHLSL += FString::Format(FormatSample, ArgsSample);
+	}
+	else if (FunctionInfo.DefinitionName == FSkeletalMeshInterfaceHelper::GetBoneCountName)
+	{
+		static const TCHAR* FormatSample = TEXT("void {InstanceFunctionName} (out int Count) { {GetDISkelMeshContextName} DISkelMesh_GetBoneCount(DIContext, Count); }");
+		OutHLSL += FString::Format(FormatSample, ArgsSample);
+	}
+	else if (FunctionInfo.DefinitionName == FSkeletalMeshInterfaceHelper::GetFilteredBoneCountName)
 	{
-		static const TCHAR* FormatSample = TEXT("void {InstanceFunctionName} (out int Count) { {GetDISkelMeshContextName} DISkelMesh_GetSpecificBoneCount(DIContext, Count); }");
+		static const TCHAR* FormatSample = TEXT("void {InstanceFunctionName} (out int Count) { {GetDISkelMeshContextName} DISkelMesh_GetFilteredBoneCount(DIContext, Count); }");
 		OutHLSL += FString::Format(FormatSample, ArgsSample);
 	}
-	else if (FunctionInfo.DefinitionName == FSkeletalMeshInterfaceHelper::GetSpecificBoneAtName)
+	else if (FunctionInfo.DefinitionName == FSkeletalMeshInterfaceHelper::GetFilteredBoneAtName)
 	{
-		static const TCHAR* FormatSample = TEXT("void {InstanceFunctionName} (in int BoneIndex, out int Bone) { {GetDISkelMeshContextName} DISkelMesh_GetSpecificBoneAt(DIContext, BoneIndex, Bone); }");
+		static const TCHAR* FormatSample = TEXT("void {InstanceFunctionName} (in int BoneIndex, out int Bone) { {GetDISkelMeshContextName} DISkelMesh_GetFilteredBoneAt(DIContext, BoneIndex, Bone); }");
 		OutHLSL += FString::Format(FormatSample, ArgsSample);
 	}
-	else if (FunctionInfo.DefinitionName == FSkeletalMeshInterfaceHelper::RandomSpecificBoneName)
+	else if (FunctionInfo.DefinitionName == FSkeletalMeshInterfaceHelper::RandomFilteredBoneName)
 	{
-		static const TCHAR* FormatSample = TEXT("void {InstanceFunctionName} (out int Bone) { {GetDISkelMeshContextName} DISkelMesh_RandomSpecificBone(DIContext, Bone); }");
+		static const TCHAR* FormatSample = TEXT("void {InstanceFunctionName} (NiagaraRandInfo InRandomInfo, out int Bone) { {GetDISkelMeshContextName} DISkelMesh_RandomFilteredBone(DIContext, InRandomInfo.Seed1, InRandomInfo.Seed2, InRandomInfo.Seed3, Bone); }");
 		OutHLSL += FString::Format(FormatSample, ArgsSample);
 	}
-	// Specific Socket
-	else if (FunctionInfo.DefinitionName == FSkeletalMeshInterfaceHelper::GetSpecificSocketCountName)
+	// Filtered Socket
+	else if (FunctionInfo.DefinitionName == FSkeletalMeshInterfaceHelper::GetFilteredSocketCountName)
 	{
-		static const TCHAR* FormatSample = TEXT("void {InstanceFunctionName} (out int Count) { {GetDISkelMeshContextName} DISkelMesh_GetSpecificSocketCount(DIContext, Count); }");
+		static const TCHAR* FormatSample = TEXT("void {InstanceFunctionName} (out int Count) { {GetDISkelMeshContextName} DISkelMesh_GetFilteredSocketCount(DIContext, Count); }");
 		OutHLSL += FString::Format(FormatSample, ArgsSample);
 	}
-	else if (FunctionInfo.DefinitionName == FSkeletalMeshInterfaceHelper::GetSpecificSocketBoneAtName)
+	else if (FunctionInfo.DefinitionName == FSkeletalMeshInterfaceHelper::GetFilteredSocketBoneAtName)
 	{
-		static const TCHAR* FormatSample = TEXT("void {InstanceFunctionName} (in int SocketIndex, out int Bone) { {GetDISkelMeshContextName} DISkelMesh_GetSpecificSocketBoneAt(DIContext, SocketIndex, Bone); }");
+		static const TCHAR* FormatSample = TEXT("void {InstanceFunctionName} (in int SocketIndex, out int Bone) { {GetDISkelMeshContextName} DISkelMesh_GetFilteredSocketBoneAt(DIContext, SocketIndex, Bone); }");
 		OutHLSL += FString::Format(FormatSample, ArgsSample);
 	} 
-	else if (FunctionInfo.DefinitionName == FSkeletalMeshInterfaceHelper::GetSpecificSocketTransformName)
+	else if (FunctionInfo.DefinitionName == FSkeletalMeshInterfaceHelper::GetFilteredSocketTransformName)
 	{
 		// TODO: This just returns the Identity transform.
 		// TODO: Make this work on the GPU?
 		static const TCHAR* FormatSample = TEXT("void {InstanceFunctionName} (in int SocketIndex, in int bShouldApplyTransform, out float3 Translation, out float4 Rotation, out float4 Scale) { Translation = float3(0.0, 0.0, 0.0); Rotation = float4(0.0, 0.0, 0.0, 1.0); Scale = float3(1.0, 1.0, 1.0); }");
 		OutHLSL += FString::Format(FormatSample, ArgsSample);
 	} 
-	else if (FunctionInfo.DefinitionName == FSkeletalMeshInterfaceHelper::RandomSpecificSocketBoneName)
+	else if (FunctionInfo.DefinitionName == FSkeletalMeshInterfaceHelper::RandomFilteredSocketBoneName)
 	{
-		static const TCHAR* FormatSample = TEXT("void {InstanceFunctionName} (out int SocketBone) { {GetDISkelMeshContextName} DISkelMesh_RandomSpecificSocketBone(DIContext, SocketBone); }");
+		static const TCHAR* FormatSample = TEXT("void {InstanceFunctionName} (NiagaraRandInfo InRandomInfo, out int SocketBone) { {GetDISkelMeshContextName} DISkelMesh_RandomFilteredSocketBone(DIContext, InRandomInfo.Seed1, InRandomInfo.Seed2, InRandomInfo.Seed3, SocketBone); }");
 		OutHLSL += FString::Format(FormatSample, ArgsSample);
 	}
 	// Unsupported functionality
@@ -2317,7 +2362,6 @@ bool UNiagaraDataInterfaceSkeletalMesh::GetFunctionHLSL(const FNiagaraDataInterf
 	//else if (FunctionInfo.DefinitionName == FSkeletalMeshInterfaceHelper::GetVertexAtName)					// void GetFilteredVertex(in int FilterdIndex, out int VertexIndex)
 	//else if (FunctionInfo.DefinitionName == FSkeletalMeshInterfaceHelper::GetTriangleCountName)			// void GetFilteredTriangleCount(out int OutTriangleCount)
 	//else if (FunctionInfo.DefinitionName == FSkeletalMeshInterfaceHelper::GetTriangleAtName)				// void GetFilteredTriangle(in int Triangle, out {MeshTriCoordinateStructName} OutCoord)
-	//else if (FunctionInfo.DefinitionName == FSkeletalMeshInterfaceHelper::IsValidBoneName)					// void IsValidBoneName(in int BoneIndex, out bool IsValid)
 	else
 	{
 		// This function is not support
@@ -2328,6 +2372,66 @@ bool UNiagaraDataInterfaceSkeletalMesh::GetFunctionHLSL(const FNiagaraDataInterf
 	return true;
 }
 
+#if WITH_EDITORONLY_DATA
+bool UNiagaraDataInterfaceSkeletalMesh::UpgradeFunctionCall(FNiagaraFunctionSignature& FunctionSignature)
+{
+	// Perform renames
+	static const TPair<FName, FName> FunctionRenames[] =
+	{
+		MakeTuple(FName("IsValidBone"), FSkeletalMeshInterfaceHelper::IsValidBoneName),
+		MakeTuple(FName("RandomSpecificBone"), FSkeletalMeshInterfaceHelper::RandomFilteredBoneName),
+		MakeTuple(FName("GetSpecificBoneCount"), FSkeletalMeshInterfaceHelper::GetFilteredBoneCountName),
+		MakeTuple(FName("GetSpecificBone"), FSkeletalMeshInterfaceHelper::GetFilteredBoneAtName),
+		MakeTuple(FName("RandomSpecificSocketBone"), FSkeletalMeshInterfaceHelper::RandomFilteredSocketBoneName),
+		MakeTuple(FName("GetSpecificSocketCount"), FSkeletalMeshInterfaceHelper::GetFilteredSocketCountName),
+		MakeTuple(FName("GetSpecificSocketTransform"), FSkeletalMeshInterfaceHelper::GetFilteredSocketTransformName),
+		MakeTuple(FName("GetSpecificSocketBone"), FSkeletalMeshInterfaceHelper::GetFilteredSocketBoneAtName),
+	};
+	bool bWasChanged = false;
+
+	for (const auto& RenamePair : FunctionRenames)
+	{
+		if (FunctionSignature.Name == RenamePair.Key)
+		{
+			FunctionSignature.Name = RenamePair.Value;
+			bWasChanged = true;
+			break;
+		}
+	}
+
+	// Upgrade functions
+	if (FunctionSignature.FunctionVersion < FNiagaraSkelMeshDIFunctionVersion::AddedRandomInfo)
+	{
+		if (FunctionSignature.Name == FSkeletalMeshInterfaceHelper::RandomTriCoordName)
+		{
+			if (FunctionSignature.Inputs.Num() == 1)
+			{
+				FunctionSignature.Inputs.Add(FNiagaraVariable(FNiagaraTypeDefinition(FNiagaraRandInfo::StaticStruct()), TEXT("RandomInfo")));
+				bWasChanged = true;
+			}
+		}
+		else if (FunctionSignature.Name == FSkeletalMeshInterfaceHelper::RandomFilteredBoneName)
+		{
+			FunctionSignature.Inputs.Add(FNiagaraVariable(FNiagaraTypeDefinition(FNiagaraRandInfo::StaticStruct()), TEXT("RandomInfo")));
+			bWasChanged = true;
+		}
+		else if (FunctionSignature.Name == FSkeletalMeshInterfaceHelper::RandomFilteredSocketBoneName)
+		{
+			FunctionSignature.Inputs.Add(FNiagaraVariable(FNiagaraTypeDefinition(FNiagaraRandInfo::StaticStruct()), TEXT("RandomInfo")));
+			bWasChanged = true;
+		}
+		else if (FunctionSignature.Name == FSkeletalMeshInterfaceHelper::RandomVertexName)
+		{
+			FunctionSignature.Inputs.Add(FNiagaraVariable(FNiagaraTypeDefinition(FNiagaraRandInfo::StaticStruct()), TEXT("RandomInfo")));
+			bWasChanged = true;
+		}
+	}
+	FunctionSignature.FunctionVersion = FNiagaraSkelMeshDIFunctionVersion::LatestVersion;
+
+	return bWasChanged;
+}
+#endif
+
 void UNiagaraDataInterfaceSkeletalMesh::GetParameterDefinitionHLSL(const FNiagaraDataInterfaceGPUParamInfo& ParamInfo, FString& OutHLSL)
 {
 	OutHLSL += TEXT("DISKELMESH_DECLARE_CONSTANTS(") + ParamInfo.DataInterfaceHLSLSymbol + TEXT(")\n");
diff --git a/Engine/Plugins/FX/Niagara/Source/Niagara/Public/NiagaraCommon.h b/Engine/Plugins/FX/Niagara/Source/Niagara/Public/NiagaraCommon.h
--- a/Engine/Plugins/FX/Niagara/Source/Niagara/Public/NiagaraCommon.h
+++ b/Engine/Plugins/FX/Niagara/Source/Niagara/Public/NiagaraCommon.h
@@ -263,6 +263,10 @@ struct NIAGARA_API FNiagaraFunctionSignature
 	/** The message to display when a function is marked experimental. */
 	UPROPERTY(EditAnywhere, Category = Script, meta = (EditCondition = "bExperimental", MultiLine = true, SkipForCompileHash = true))
 	FText ExperimentalMessage;
+
+	/** Per function version, it is up to the discretion of the function as to what the version means. */
+	UPROPERTY()
+	uint32 FunctionVersion = 0;
 #endif
 
 	/** Support running on the CPU. */
@@ -347,13 +351,13 @@ struct NIAGARA_API FNiagaraFunctionSignature
 
 	bool EqualsIgnoringSpecifiers(const FNiagaraFunctionSignature& Other) const
 	{
-		bool bNamesEqual = Name.ToString().Equals(Other.Name.ToString());
-		bool bInputsEqual = Inputs == Other.Inputs;
-		bool bOutputsEqual = Outputs == Other.Outputs;
-		bool bContextsEqual = bRequiresContext == Other.bRequiresContext;
-		bool bMemberFunctionsEqual = bMemberFunction == Other.bMemberFunction;
-		bool bOwnerNamesEqual = OwnerName == Other.OwnerName;
-		return bNamesEqual && bInputsEqual && bOutputsEqual && bContextsEqual && bMemberFunctionsEqual && bOwnerNamesEqual;
+		bool bMatches = Name.ToString().Equals(Other.Name.ToString());
+		bMatches &= Inputs == Other.Inputs;
+		bMatches &= Outputs == Other.Outputs;
+		bMatches &= bRequiresContext == Other.bRequiresContext;
+		bMatches &= bMemberFunction == Other.bMemberFunction;
+		bMatches &= OwnerName == Other.OwnerName;
+		return bMatches;
 	}
 
 	FString GetName()const { return Name.ToString(); }
diff --git a/Engine/Plugins/FX/Niagara/Source/NiagaraEditor/Private/NiagaraNodeFunctionCall.cpp b/Engine/Plugins/FX/Niagara/Source/NiagaraEditor/Private/NiagaraNodeFunctionCall.cpp
--- a/Engine/Plugins/FX/Niagara/Source/NiagaraEditor/Private/NiagaraNodeFunctionCall.cpp
+++ b/Engine/Plugins/FX/Niagara/Source/NiagaraEditor/Private/NiagaraNodeFunctionCall.cpp
@@ -73,6 +73,20 @@ void UNiagaraNodeFunctionCall::PostLoad()
 		}
 	}
 
+	// Allow data interfaces an opportunity to intercept changes
+	if (Signature.IsValid() && Signature.bMemberFunction)
+	{
+		if ((Signature.Inputs.Num() > 0) && Signature.Inputs[0].GetType().IsDataInterface())
+		{
+			UNiagaraDataInterface* CDO = CastChecked<UNiagaraDataInterface>(Signature.Inputs[0].GetType().GetClass()->GetDefaultObject());
+			if (CDO->UpgradeFunctionCall(Signature))
+			{
+				FunctionDisplayName.Empty();
+				ReallocatePins();
+			}
+		}
+	}
+
 	// Clean up invalid old references to propagated parameters
 	CleanupPropagatedSwitchValues();
 	
@@ -84,6 +98,8 @@ void UNiagaraNodeFunctionCall::PostLoad()
 
 void UNiagaraNodeFunctionCall::UpgradeDIFunctionCalls()
 {
+	// We no longer use this upgrade path, but leaving here for convience in case we need to use this route again for other things
+#if 0
 	UClass* InterfaceClass = nullptr;
 	UNiagaraDataInterface* InterfaceCDO = nullptr;
 	if (Signature.IsValid() && FunctionScript == nullptr)
@@ -99,25 +115,6 @@ void UNiagaraNodeFunctionCall::UpgradeDIFunctionCalls()
 	}
 
 	FString UpgradeNote;
-
-	//TODO: Move this out into DI specific functions or helper classes?
-	if (InterfaceClass && InterfaceCDO)
-	{		
-		if (InterfaceClass == UNiagaraDataInterfaceSkeletalMesh::StaticClass())
-		{
-			if (Signature.Name == TEXT("RandomTriCoord"))
-			{
-				if (Signature.Inputs.Num() == 1)//If this is before we added the additional seed inputs. TODO: Add a per DI version number to make this simpler and clearer. In this case we can detect old data easy enough but a version number is better.
-				{
-					UpgradeNote = TEXT("Adding RandomInfo parameter to support deterministic random sampling.");
-
-					Signature.Inputs.Add(FNiagaraVariable(FNiagaraTypeDefinition(FNiagaraRandInfo::StaticStruct()), TEXT("RandomInfo")));
-					ReallocatePins();
-				}
-			}
-		}
-	}
-
 	if (!UpgradeNote.IsEmpty())
 	{
 		UE_LOG(LogNiagaraEditor, Log, TEXT("Upgradeing Niagara Data Interface fuction call node. This may cause unnessessary recompiles. Please resave these assets if this occurs. Or use fx.UpgradeAllNiagaraAssets."));
@@ -129,6 +126,7 @@ void UNiagaraNodeFunctionCall::UpgradeDIFunctionCalls()
 		UE_LOG(LogNiagaraEditor, Log, TEXT("Function: %s"), *Signature.GetName());
 		UE_LOG(LogNiagaraEditor, Log, TEXT("Upgrade Note: %s"),* UpgradeNote);
 	}
+#endif
 }
 
 TSharedPtr<SGraphNode> UNiagaraNodeFunctionCall::CreateVisualWidget()
diff --git a/Engine/Plugins/FX/Niagara/Source/NiagaraEditorWidgets/Private/DetailCustomizations/NiagaraDataInterfaceSkeletalMeshDetails.cpp b/Engine/Plugins/FX/Niagara/Source/NiagaraEditorWidgets/Private/DetailCustomizations/NiagaraDataInterfaceSkeletalMeshDetails.cpp
--- a/Engine/Plugins/FX/Niagara/Source/NiagaraEditorWidgets/Private/DetailCustomizations/NiagaraDataInterfaceSkeletalMeshDetails.cpp
+++ b/Engine/Plugins/FX/Niagara/Source/NiagaraEditorWidgets/Private/DetailCustomizations/NiagaraDataInterfaceSkeletalMeshDetails.cpp
@@ -68,8 +68,8 @@ void FNiagaraDataInterfaceSkeletalMeshDetails::CustomizeDetails(IDetailLayoutBui
 		 TArray<TSharedRef<IPropertyHandle>> SkelProperties;
 		 SkelCategory->GetDefaultProperties(SkelProperties, true, true);
 
-		 TSharedPtr<IPropertyHandle> BonesProperty = DetailBuilder.GetProperty(GET_MEMBER_NAME_CHECKED(UNiagaraDataInterfaceSkeletalMesh, SpecificBones));
-		 TSharedPtr<IPropertyHandle> SocketsProperty = DetailBuilder.GetProperty(GET_MEMBER_NAME_CHECKED(UNiagaraDataInterfaceSkeletalMesh, SpecificSockets));
+		 TSharedPtr<IPropertyHandle> BonesProperty = DetailBuilder.GetProperty(GET_MEMBER_NAME_CHECKED(UNiagaraDataInterfaceSkeletalMesh, FilteredBones));
+		 TSharedPtr<IPropertyHandle> SocketsProperty = DetailBuilder.GetProperty(GET_MEMBER_NAME_CHECKED(UNiagaraDataInterfaceSkeletalMesh, FilteredSockets));
 
 		 for (TSharedPtr<IPropertyHandle> Property : SkelProperties)
 		 {
