- Add RandomFilteredSocketOrBone functionality

diff --git a/Engine/Plugins/FX/Niagara/Shaders/Private/NiagaraDataInterfaceSkeletalMesh.ush b/Engine/Plugins/FX/Niagara/Shaders/Private/NiagaraDataInterfaceSkeletalMesh.ush
--- a/Engine/Plugins/FX/Niagara/Shaders/Private/NiagaraDataInterfaceSkeletalMesh.ush
+++ b/Engine/Plugins/FX/Niagara/Shaders/Private/NiagaraDataInterfaceSkeletalMesh.ush
@@ -501,6 +501,21 @@ void DISkelMesh_RandomFilteredSocketBone(in FDISkelMeshContext DIContext, uint S
     SocketBone = DIContext.FilteredSocketBoneOffset + min(int(RandT0 * float(DIContext.NumFilteredSockets)), DIContext.NumFilteredSockets - 1); // avoid % by using mul/min to Tri = MeshVertexCount
 }
 
+void DISkelMesh_RandomFilteredSocketOrBone(in FDISkelMeshContext DIContext, uint Seed1, uint Seed2, uint Seed3, out int Bone)
+{
+	float RandT0 = NiagaraInternalNoise(Seed1, Seed2, Seed3);
+	int NumFilteredSocketsAndBones = DIContext.NumFilteredBones + DIContext.NumFilteredSockets;
+    int FilteredIndex = min(int(RandT0 * float(NumFilteredSocketsAndBones)), NumFilteredSocketsAndBones - 1);
+	if ( FilteredIndex < DIContext.NumFilteredBones )
+	{
+		Bone = (int)DIContext.FilteredBones[FilteredIndex];
+	}
+	else
+	{
+		Bone = DIContext.FilteredSocketBoneOffset + FilteredIndex - DIContext.NumFilteredBones;
+	}
+}
+
 ////////////////////////////////////////////////////////////////////////////////////////////////////////////////////////////////////////////////////////////////
 
 #define DISKELMESH_DECLARE_CONSTANTS(NAME) \
diff --git a/Engine/Plugins/FX/Niagara/Source/Niagara/Classes/NiagaraDataInterfaceSkeletalMesh.h b/Engine/Plugins/FX/Niagara/Source/Niagara/Classes/NiagaraDataInterfaceSkeletalMesh.h
--- a/Engine/Plugins/FX/Niagara/Source/Niagara/Classes/NiagaraDataInterfaceSkeletalMesh.h
+++ b/Engine/Plugins/FX/Niagara/Source/Niagara/Classes/NiagaraDataInterfaceSkeletalMesh.h
@@ -745,7 +745,8 @@ public:
 	void GetFilteredSocketBoneAt(FVectorVMContext& Context);
 	void GetFilteredSocketTransform(FVectorVMContext& Context);
 	void RandomFilteredSocketBone(FVectorVMContext& Context);
-		
+
+	void RandomFilteredSocketOrBone(FVectorVMContext& Context);
 	// End of Direct Bone + Socket Sampling
 	//////////////////////////////////////////////////////////////////////////
 
@@ -784,6 +785,7 @@ public:
 	static const FName GetFilteredSocketCountName;
 	static const FName GetFilteredSocketBoneAtName;
 	static const FName GetFilteredSocketTransformName;
+	static const FName RandomFilteredSocketOrBoneName;
 
 	// Vertex Sampling
 	static const FName IsValidVertexName;
diff --git a/Engine/Plugins/FX/Niagara/Source/Niagara/Private/NDISkeletalMesh_BoneSampling.cpp b/Engine/Plugins/FX/Niagara/Source/Niagara/Private/NDISkeletalMesh_BoneSampling.cpp
--- a/Engine/Plugins/FX/Niagara/Source/Niagara/Private/NDISkeletalMesh_BoneSampling.cpp
+++ b/Engine/Plugins/FX/Niagara/Source/Niagara/Private/NDISkeletalMesh_BoneSampling.cpp
@@ -39,6 +39,8 @@ const FName FSkeletalMeshInterfaceHelper::GetFilteredSocketCountName("GetFiltere
 const FName FSkeletalMeshInterfaceHelper::GetFilteredSocketTransformName("GetFilteredSocketTransform");
 const FName FSkeletalMeshInterfaceHelper::GetFilteredSocketBoneAtName("GetFilteredSocketBone");
 
+const FName FSkeletalMeshInterfaceHelper::RandomFilteredSocketOrBoneName("RandomFilteredSocketOrBone");
+
 void UNiagaraDataInterfaceSkeletalMesh::GetSkeletonSamplingFunctions(TArray<FNiagaraFunctionSignature>& OutFunctions)
 {
 	//////////////////////////////////////////////////////////////////////////
@@ -187,7 +189,6 @@ void UNiagaraDataInterfaceSkeletalMesh::GetSkeletonSamplingFunctions(TArray<FNia
 
 	//////////////////////////////////////////////////////////////////////////
 	//Socket functions
-
 	{
 		FNiagaraFunctionSignature Sig;
 		Sig.Name = FSkeletalMeshInterfaceHelper::RandomFilteredSocketBoneName;
@@ -242,6 +243,22 @@ void UNiagaraDataInterfaceSkeletalMesh::GetSkeletonSamplingFunctions(TArray<FNia
 		Sig.bRequiresContext = false;
 #if WITH_EDITORONLY_DATA
 		Sig.Description = LOCTEXT("GetFilteredSocketTransformDesc", "Gets the transform for the socket at the passed index in the DI's filtered socket list. If the Source component is set it will respect the Relative Transform Space as well..");
+#endif
+		OutFunctions.Add(Sig);
+	}
+
+	//////////////////////////////////////////////////////////////////////////
+	// Misc Functions
+	{
+		FNiagaraFunctionSignature Sig;
+		Sig.Name = FSkeletalMeshInterfaceHelper::RandomFilteredSocketOrBoneName;
+		Sig.Inputs.Add(FNiagaraVariable(FNiagaraTypeDefinition(GetClass()), TEXT("SkeletalMesh")));
+		Sig.Inputs.Add(FNiagaraVariable(FNiagaraTypeDefinition(FNiagaraRandInfo::StaticStruct()), TEXT("RandomInfo")));
+		Sig.Outputs.Add(FNiagaraVariable(FNiagaraTypeDefinition::GetIntDef(), TEXT("Bone")));
+		Sig.bMemberFunction = true;
+		Sig.bRequiresContext = false;
+#if WITH_EDITORONLY_DATA
+		Sig.Description = LOCTEXT("RandomFilteredSocketOrBoneDesc", "Gets the bone for a random filtered socket or bone from the DI's list.");
 #endif
 		OutFunctions.Add(Sig);
 	}
@@ -321,9 +338,14 @@ void UNiagaraDataInterfaceSkeletalMesh::BindSkeletonSamplingFunction(const FVMEx
 		check(BindingInfo.GetNumInputs() == 3 && BindingInfo.GetNumOutputs() == 10);
 		NDI_FUNC_BINDER(UNiagaraDataInterfaceSkeletalMesh, GetFilteredSocketTransform)::Bind(this, OutFunc);
 	}
+	// Misc functions
+	else if (BindingInfo.Name == FSkeletalMeshInterfaceHelper::RandomFilteredSocketOrBoneName)
+	{
+		check(BindingInfo.GetNumInputs() == 4 && BindingInfo.GetNumOutputs() == 1);
+		OutFunc = FVMExternalFunction::CreateLambda([this](FVectorVMContext& Context) { this->RandomFilteredSocketOrBone(Context); });
+	}
 }
 
-
 //////////////////////////////////////////////////////////////////////////
 // Direct sampling from listed sockets and bones.
 
@@ -779,4 +801,36 @@ void UNiagaraDataInterfaceSkeletalMesh::RandomFilteredSocketBone(FVectorVMContex
 	}
 }
 
+void UNiagaraDataInterfaceSkeletalMesh::RandomFilteredSocketOrBone(FVectorVMContext& Context)
+{
+	FNDIRandomHelper RandHelper(Context);
+	VectorVM::FUserPtrHandler<FNDISkeletalMesh_InstanceData> InstData(Context);
+
+	VectorVM::FExternalFuncRegisterHandler<int32> OutBoneIndex(Context);
+
+	const int32 Max = FilteredSockets.Num() + FilteredBones.Num() - 1;
+	if (Max >= 0)
+	{
+		const int32 NumFilteredBones = FilteredBones.Num();
+		const int32 FilteredSocketBoneOffset = InstData->FilteredSocketBoneOffset;
+		for (int32 i = 0; i < Context.NumInstances; ++i)
+		{
+			RandHelper.GetAndAdvance();
+			const int32 FilteredIndex = RandHelper.RandRange(i, 0, Max);
+			if (FilteredIndex < NumFilteredBones)
+			{
+				*OutBoneIndex.GetDestAndAdvance() = InstData->FilteredBones[FilteredIndex];
+			}
+			else
+			{
+				*OutBoneIndex.GetDestAndAdvance() = FilteredSocketBoneOffset + FilteredIndex - NumFilteredBones;
+			}
+		}
+	}
+	else
+	{
+		FMemory::Memset(OutBoneIndex.GetDest(), 0xFF, Context.NumInstances * sizeof(int32));
+	}
+}
+
 #undef LOCTEXT_NAMESPACE
diff --git a/Engine/Plugins/FX/Niagara/Source/Niagara/Private/NiagaraDataInterfaceSkeletalMesh.cpp b/Engine/Plugins/FX/Niagara/Source/Niagara/Private/NiagaraDataInterfaceSkeletalMesh.cpp
--- a/Engine/Plugins/FX/Niagara/Source/Niagara/Private/NiagaraDataInterfaceSkeletalMesh.cpp
+++ b/Engine/Plugins/FX/Niagara/Source/Niagara/Private/NiagaraDataInterfaceSkeletalMesh.cpp
@@ -2334,6 +2334,7 @@ bool UNiagaraDataInterfaceSkeletalMesh::GetFunctionHLSL(const FNiagaraDataInterf
 		static const TCHAR* FormatSample = TEXT("void {InstanceFunctionName} (NiagaraRandInfo InRandomInfo, out int Bone) { {GetDISkelMeshContextName} DISkelMesh_RandomFilteredBone(DIContext, InRandomInfo.Seed1, InRandomInfo.Seed2, InRandomInfo.Seed3, Bone); }");
 		OutHLSL += FString::Format(FormatSample, ArgsSample);
 	}
+	//////////////////////////////////////////////////////////////////////////////////////////////////////////////////////////////
 	// Filtered Socket
 	else if (FunctionInfo.DefinitionName == FSkeletalMeshInterfaceHelper::GetFilteredSocketCountName)
 	{
@@ -2357,6 +2358,14 @@ bool UNiagaraDataInterfaceSkeletalMesh::GetFunctionHLSL(const FNiagaraDataInterf
 		static const TCHAR* FormatSample = TEXT("void {InstanceFunctionName} (NiagaraRandInfo InRandomInfo, out int SocketBone) { {GetDISkelMeshContextName} DISkelMesh_RandomFilteredSocketBone(DIContext, InRandomInfo.Seed1, InRandomInfo.Seed2, InRandomInfo.Seed3, SocketBone); }");
 		OutHLSL += FString::Format(FormatSample, ArgsSample);
 	}
+	//////////////////////////////////////////////////////////////////////////////////////////////////////////////////////////////
+	// Misc bone functions
+	else if (FunctionInfo.DefinitionName == FSkeletalMeshInterfaceHelper::RandomFilteredSocketOrBoneName)
+	{
+		static const TCHAR* FormatSample = TEXT("void {InstanceFunctionName} (NiagaraRandInfo InRandomInfo, out int Bone) { {GetDISkelMeshContextName} DISkelMesh_RandomFilteredSocketOrBone(DIContext, InRandomInfo.Seed1, InRandomInfo.Seed2, InRandomInfo.Seed3, Bone); }");
+		OutHLSL += FString::Format(FormatSample, ArgsSample);
+	}
+	//////////////////////////////////////////////////////////////////////////////////////////////////////////////////////////////
 	// Unsupported functionality
 	//else if (FunctionInfo.DefinitionName == FSkeletalMeshInterfaceHelper::GetVertexCountName)				// void GetFilteredVertexCount(out int VertexCount)
 	//else if (FunctionInfo.DefinitionName == FSkeletalMeshInterfaceHelper::GetVertexAtName)					// void GetFilteredVertex(in int FilterdIndex, out int VertexIndex)
