Added support for function specifiers in custom HLSL when calling a data interface function, using the following syntax:

DIName.FunctionName<Specifier1=Value1, Specifier2="Value 2">(Arguments);

Also fixed some bugs in the custom HLSL tokenizer:
 - multi-line comments didn't work (the final slash was not copied, resulting in an unterminated comment)
 - the equals sign was not treated as a separator, so omitting whitespace around it caused problems with renaming node outputs
 - tokens would sometimes be lost if there was no whitespace after a separator

#rnx

diff --git a/Engine/Plugins/FX/Niagara/Source/NiagaraEditor/Private/NiagaraHlslTranslator.cpp b/Engine/Plugins/FX/Niagara/Source/NiagaraEditor/Private/NiagaraHlslTranslator.cpp
--- a/Engine/Plugins/FX/Niagara/Source/NiagaraEditor/Private/NiagaraHlslTranslator.cpp
+++ b/Engine/Plugins/FX/Niagara/Source/NiagaraEditor/Private/NiagaraHlslTranslator.cpp
@@ -3046,6 +3046,32 @@ FString FHlslNiagaraTranslator::GetSanitizedSymbolName(FString SymbolName, bool
 	return Ret;
 }
 
+FString FHlslNiagaraTranslator::GetSanitizedDIFunctionName(const FString& FunctionName)
+{
+	bool bWordStart = true;
+	FString Sanitized;
+	for (int i = 0; i < FunctionName.Len(); ++i)
+	{
+		TCHAR c = FunctionName[i];
+
+		if (c == ' ')
+		{
+			bWordStart = true;
+		}
+		else
+		{
+			if (bWordStart)
+			{
+				c = FChar::ToUpper(c);
+				bWordStart = false;
+			}
+			Sanitized.AppendChar(c);
+		}
+	}
+
+	return Sanitized;
+}
+
 FString FHlslNiagaraTranslator::GetSanitizedFunctionNameSuffix(FString Name)
 {
 	if (Name.Len() == 0)
@@ -5489,6 +5515,117 @@ void FHlslNiagaraTranslator::FinalResolveNamespacedTokens(const FString& Paramet
 	}
 }
 
+static bool IsWhitespaceToken(const FString& Token)
+{
+	return
+		Token.Len() == 0 ||
+		Token[0] == TCHAR('\r') || Token[0] == TCHAR('\n') || Token[0] == TCHAR('\t') || Token[0] == TCHAR(' ') || 
+		(Token.Len() >= 2 && Token[0] == TCHAR('/') && (Token[1] == TCHAR('/') || Token[1] == TCHAR('*')))
+		;
+}
+
+bool FHlslNiagaraTranslator::ParseDIFunctionSpecifiers(UNiagaraNodeCustomHlsl* CustomHLSLNode, FNiagaraFunctionSignature& Sig, TArray<FString>& Tokens, int32& TokenIdx)
+{
+	const int32 NumTokens = Tokens.Num();
+
+	// Skip whitespace between the function name and the arguments or specifiers.
+	while (TokenIdx < NumTokens && IsWhitespaceToken(Tokens[TokenIdx]))
+	{
+		++TokenIdx;
+	}
+
+	// If we don't have a specifier list start token, we don't need to do anything.
+	if (TokenIdx == NumTokens || Tokens[TokenIdx] != TEXT("<"))
+	{
+		return true;
+	}
+
+	enum class EParserState
+	{
+		ExpectName,
+		ExpectEquals,
+		ExpectValue,
+		ExpectCommaOrEnd
+	};
+
+	EParserState ParserState = EParserState::ExpectName;
+	FString SpecifierName;
+
+	// All the tokens inside the specifier list, including the angle brackets, will be replaced with empty strings,
+	// because they're not valid HLSL. We just want to extract Key=Value pairs into the signature's specifier list.
+	while (TokenIdx < NumTokens)
+	{
+		FString Token = Tokens[TokenIdx];
+		Tokens[TokenIdx] = TEXT("");
+		++TokenIdx;
+
+		if (IsWhitespaceToken(Token))
+		{
+			continue;
+		}
+
+		if (Token[0] == TCHAR('<'))
+		{
+			// Nothing.
+		}
+		else if (Token[0] == TCHAR('>'))
+		{
+			if (ParserState != EParserState::ExpectCommaOrEnd)
+			{
+				Error(LOCTEXT("DataInterfaceFunctionCallUnexpectedEnd", "Unexpected end of specifier list."), CustomHLSLNode, nullptr);
+				return false;
+			}
+			break;
+		}
+		else if (Token[0] == TCHAR('='))
+		{
+			if (ParserState == EParserState::ExpectEquals)
+			{
+				ParserState = EParserState::ExpectValue;
+			}
+			else
+			{
+				Error(LOCTEXT("DataInterfaceFunctionCallExpectEquals", "Invalid token in specifier list, expecting '='."), CustomHLSLNode, nullptr);
+				return false;
+			}
+		}
+		else if (Token[0] == TCHAR(','))
+		{
+			if (ParserState == EParserState::ExpectCommaOrEnd)
+			{
+				ParserState = EParserState::ExpectName;
+			}
+			else
+			{
+				Error(LOCTEXT("DataInterfaceFunctionCallExpectComma", "Invalid token in specifier list, expecting ','."), CustomHLSLNode, nullptr);
+				return false;
+			}
+		}
+		else
+		{
+			if (ParserState == EParserState::ExpectName)
+			{
+				SpecifierName = Token;
+				ParserState = EParserState::ExpectEquals;
+			}
+			else if (ParserState == EParserState::ExpectValue)
+			{
+				int32 Start = 0, ValueLen = Token.Len();
+				// Remove the quotation marks if they are used.
+				if (Token.Len() >= 2 && Token[0] == TCHAR('"') && Token[ValueLen - 1] == TCHAR('"'))
+				{
+					Start = 1;
+					ValueLen -= 2;
+				}
+				Sig.FunctionSpecifiers.Add(FName(SpecifierName), FName(Token.Mid(Start, ValueLen)));
+				ParserState = EParserState::ExpectCommaOrEnd;
+			}
+		}
+	}
+
+	return true;
+}
+
 void FHlslNiagaraTranslator::HandleCustomHlslNode(UNiagaraNodeCustomHlsl* CustomFunctionHlsl, ENiagaraScriptUsage& OutScriptUsage, FString& OutName, FString& OutFullName, bool& bOutCustomHlsl, FString& OutCustomHlsl,
 	FNiagaraFunctionSignature& OutSignature, TArray<int32>& Inputs)
 {
@@ -5554,30 +5691,54 @@ void FHlslNiagaraTranslator::HandleCustomHlslNode(UNiagaraNodeCustomHlsl* Custom
 			// actual custom hlsl source. If they do, then add them to the function table that we need to map.
 			FNiagaraScriptDataInterfaceCompileInfo& Info = CompilationOutput.ScriptData.DataInterfaceInfo[OwnerIdx];
 			TArray<FNiagaraFunctionSignature> Funcs;
-			TArray<FNiagaraFunctionSignature> AddedFuncs;
 			CDO->GetFunctions(Funcs);
 			for (int32 FuncIdx = 0; FuncIdx < Funcs.Num(); FuncIdx++)
 			{
-				FNiagaraFunctionSignature Sig = Funcs[FuncIdx];
-				FString ReplaceSrc = Input.GetName().ToString() + TEXT(".") + Sig.GetName();
-				FString ReplaceDest = GetSanitizedSymbolName(Sig.GetName() + TEXT("_") + (Info.Name.ToString().Replace(TEXT("."), TEXT(""))));
-				uint32 NumFound = CustomFunctionHlsl->ReplaceExactMatchTokens(Tokens, ReplaceSrc, ReplaceDest, false);
-				if (NumFound != 0)
+				FString DIMethodInvocation = Input.GetName().ToString() + TEXT(".") + GetSanitizedDIFunctionName(Funcs[FuncIdx].GetName());
+
+				for (int32 TokenIndex = 0; TokenIndex < Tokens.Num();)
 				{
-					AddedFuncs.Add(Sig);
+					if (Tokens[TokenIndex].Compare(DIMethodInvocation, ESearchCase::CaseSensitive) == 0)
+					{
+						// We can't replace the method-style call with the actual function name yet, because function specifiers
+						// are part of the name, and we haven't determined them yet. Just store a pointer to the token for now.
+						FString& FunctionNameToken = Tokens[TokenIndex];
+						++TokenIndex;
+
+						FNiagaraFunctionSignature Sig = Funcs[FuncIdx];
+
+						// Override the owner id of the signature with the actual caller.
+						Sig.OwnerName = Info.Name;
+
+						// Function specifiers can be given inside angle brackets, using this syntax:
+						//
+						//		DI.Function<Specifier1=Value1, Specifier2="Value 2">(Arguments);
+						//
+						// We need to extract the specifiers and replace any tokens inside the angle brackets with empty strings,
+						// to arrive back at valid HLSL.
+						if (!ParseDIFunctionSpecifiers(CustomFunctionHlsl, Sig, Tokens, TokenIndex))
+						{
+							return;
+						}
+
+						// Now we can build the function name and replace the method call token with the final function name.
+						FunctionNameToken = GetFunctionSignatureSymbol(Sig);
+
+						if (Info.UserPtrIdx != INDEX_NONE && CompilationTarget != ENiagaraSimTarget::GPUComputeSim)
+						{
+							//This interface requires per instance data via a user ptr so place the index to it at the end of the inputs.
+							Sig.Inputs.Add(FNiagaraVariable(FNiagaraTypeDefinition::GetIntDef(), TEXT("InstanceData")));
+						}
+
+						Info.RegisteredFunctions.Add(Sig);
+						Functions.FindOrAdd(Sig);
 
-					if (Info.UserPtrIdx != INDEX_NONE && CompilationTarget != ENiagaraSimTarget::GPUComputeSim)
+						HandleDataInterfaceCall(Info, Sig);
+					}
+					else
 					{
-						//This interface requires per instance data via a user ptr so place the index to it at the end of the inputs.
-						//Skip this for now... Inputs.Add(AddSourceChunk(LexToString(Info.UserPtrIdx), FNiagaraTypeDefinition::GetIntDef(), false));
-						Sig.Inputs.Add(FNiagaraVariable(FNiagaraTypeDefinition::GetIntDef(), TEXT("InstanceData")));
+						++TokenIndex;
 					}
-					//Override the owner id of the signature with the actual caller.
-					Sig.OwnerName = Info.Name;
-					Info.RegisteredFunctions.Add(Sig);
-					Functions.FindOrAdd(Sig);
-
-					HandleDataInterfaceCall(Info, Sig);
 				}
 			}
 			SigInputs.Add(Input);
diff --git a/Engine/Plugins/FX/Niagara/Source/NiagaraEditor/Private/NiagaraHlslTranslator.h b/Engine/Plugins/FX/Niagara/Source/NiagaraEditor/Private/NiagaraHlslTranslator.h
--- a/Engine/Plugins/FX/Niagara/Source/NiagaraEditor/Private/NiagaraHlslTranslator.h
+++ b/Engine/Plugins/FX/Niagara/Source/NiagaraEditor/Private/NiagaraHlslTranslator.h
@@ -540,6 +540,7 @@ public:
 
 
 	static FString GetSanitizedSymbolName(FString SymbolName, bool bCollapseNamespaces=false);
+	static FString GetSanitizedDIFunctionName(const FString& FunctionName);
 	static FString GetSanitizedFunctionNameSuffix(FString Name);
 
 	/** Replaces all non-ascii characters with a "ASCXXX" string, where XXX is their int value */
@@ -578,6 +579,7 @@ private:
 	FString ComputeMatrixColumnAccess(const FString& Name);
 	FString ComputeMatrixRowAccess(const FString& Name);
 
+	bool ParseDIFunctionSpecifiers(UNiagaraNodeCustomHlsl* CustomHLSLNode, FNiagaraFunctionSignature& Sig, TArray<FString>& Tokens, int32& TokenIdx);
 	void HandleCustomHlslNode(UNiagaraNodeCustomHlsl* CustomFunctionHlsl, ENiagaraScriptUsage& OutScriptUsage, FString& OutName, FString& OutFullName, bool& bOutCustomHlsl, FString& OutCustomHlsl,
 		FNiagaraFunctionSignature& OutSignature, TArray<int32>& Inputs);
 	
diff --git a/Engine/Plugins/FX/Niagara/Source/NiagaraEditor/Private/NiagaraNodeCustomHlsl.cpp b/Engine/Plugins/FX/Niagara/Source/NiagaraEditor/Private/NiagaraNodeCustomHlsl.cpp
--- a/Engine/Plugins/FX/Niagara/Source/NiagaraEditor/Private/NiagaraNodeCustomHlsl.cpp
+++ b/Engine/Plugins/FX/Niagara/Source/NiagaraEditor/Private/NiagaraNodeCustomHlsl.cpp
@@ -79,77 +79,90 @@ bool UNiagaraNodeCustomHlsl::GetTokens(TArray<FString>& OutTokens) const
 		return false;
 	}
 	
-	FString InputVars = TEXT(";/*+-)(?:, []\t\n");
-	int32 LastValidIdx = INDEX_NONE;
-	bool bComment = false;
+	FString Separators = TEXT(";/*+-=)(?:, []<>\"\t\n");
+	int32 TokenStart = 0;
 	int32 TargetLength = HlslData.Len();
 	for (int32 i = 0; i < TargetLength; )
 	{
 		int32 Index = INDEX_NONE;
 
 		// Determine if we are a splitter character or a regular character.
-		if (InputVars.FindChar(HlslData[i], Index) && Index != INDEX_NONE)
+		if (Separators.FindChar(HlslData[i], Index) && Index != INDEX_NONE)
 		{
-			if (HlslData[i] == '/' && (i + 1 != TargetLength) && HlslData[i+1] == '/') // We are a comment, go to end of line.
+			// Commit the current token, if any.
+			if (i > TokenStart)
 			{
-				int32 FoundEndIdx = HlslData.Find("\n", ESearchCase::IgnoreCase, ESearchDir::FromStart, i + 2);
-				if (FoundEndIdx != INDEX_NONE)
-				{
-					// do nothing // FoundEndIdx = FoundEndIdx;
-				}
-				else
+				OutTokens.Add(HlslData.Mid(TokenStart, i - TokenStart));
+			}
+
+			if (HlslData[i] == '/' && (i + 1 != TargetLength) && HlslData[i+1] == '/')
+			{
+				// Single-line comment, everything up to the end of the line becomes a token (including the comment start and the newline).
+				int32 FoundEndIdx = HlslData.Find("\n", ESearchCase::CaseSensitive, ESearchDir::FromStart, i + 2);
+				if (FoundEndIdx == INDEX_NONE)
 				{
 					FoundEndIdx = TargetLength - 1;
 				}
 
-				OutTokens.Add(HlslData.Mid(i, FoundEndIdx - i));
-				OutTokens.Add("\n");
+				OutTokens.Add(HlslData.Mid(i, FoundEndIdx - i + 1));
 				i = FoundEndIdx + 1;
 			}
-			else if (HlslData[i] == '/' && (i + 1 != TargetLength) && HlslData[i+1] == '*') // We are a multiline comment, go to end comment. Nested comments are not supported currently.
+			else if (HlslData[i] == '/' && (i + 1 != TargetLength) && HlslData[i+1] == '*')
 			{
-				int32 FoundEndIdx = HlslData.Find("*/", ESearchCase::IgnoreCase, ESearchDir::FromStart, i + 2);
+				// Multi-line comment, all of it becomes a single token, including the start and end markers.
+				int32 FoundEndIdx = HlslData.Find("*/", ESearchCase::CaseSensitive, ESearchDir::FromStart, i + 2);
 				if (FoundEndIdx != INDEX_NONE)
 				{
-					FoundEndIdx = FoundEndIdx + 1;
+					// Include both characters of the terminator.
+					FoundEndIdx += 1;
 				}
 				else
 				{
+					// This is an unterminated multi-line comment, but there's nothing we can do at this point.
 					FoundEndIdx = TargetLength - 1;
 				}
 
-				OutTokens.Add(HlslData.Mid(i, FoundEndIdx - i));
+				OutTokens.Add(HlslData.Mid(i, FoundEndIdx - i + 1));
 				i = FoundEndIdx + 1;
 			}
-			else if (LastValidIdx != INDEX_NONE) // We encountered a non-splitter character in the past that hasn't been recorded yet but we are a splitter.
+			else if (HlslData[i] == '"')
 			{
-				OutTokens.Add(HlslData.Mid(LastValidIdx, i - LastValidIdx));
-				OutTokens.Add(FString(1, &HlslData[i]));
-				i++;
+				// Strings in HLSL, what?
+				// This is an extension used to support calling DI functions which have specifiers. The syntax is:
+				//		DIName.Function<Specifier1="Value 1", Specifier2="Value 2">();
+				// The string is considered a single token, including the quotation marks.
+				int32 FoundEndIdx = HlslData.Find("\"", ESearchCase::CaseSensitive, ESearchDir::FromStart, i + 1);
+				if (FoundEndIdx == INDEX_NONE)
+				{
+					// Unterminated string. A very weird compiler error will follow, but there's nothing we can do at this point.
+					FoundEndIdx = TargetLength - 1;
+				}
+
+				OutTokens.Add(HlslData.Mid(i, FoundEndIdx - i + 1));
+				i = FoundEndIdx + 1;
+
 			}
-			else //if (LastValidIdx == INDEX_NONE) // We are a splitter with no known unrecorded tokens prior.
+			else
 			{
+				// The separator will be its own token.
 				OutTokens.Add(FString(1, &HlslData[i]));
 				i++;
 			}
 
-			LastValidIdx = INDEX_NONE;
+			// Start a new token after the separator.
+			TokenStart = i;
 		}
-		else if (LastValidIdx == INDEX_NONE)
-		{
-			LastValidIdx = i; // Record that this is where we encountered the first non-splitter character that has yet to be recorded.
-			i++;
-		}		
 		else
 		{
+			// This character is part of a token, continue scanning.
 			i++;
 		}
 	}
 
 	// We may need to pull in the last chars from the end.
-	if (LastValidIdx != INDEX_NONE)
+	if (TokenStart < TargetLength)
 	{
-		OutTokens.Add(HlslData.Mid(LastValidIdx));
+		OutTokens.Add(HlslData.Mid(TokenStart));
 	}
 
 	return true;
