Add GetAttributeByIndex functions to the particle read DI.

#rnx

diff --git a/Engine/Plugins/FX/Niagara/Source/Niagara/Classes/NiagaraDataInterfaceParticleRead.h b/Engine/Plugins/FX/Niagara/Source/Niagara/Classes/NiagaraDataInterfaceParticleRead.h
--- a/Engine/Plugins/FX/Niagara/Source/Niagara/Classes/NiagaraDataInterfaceParticleRead.h
+++ b/Engine/Plugins/FX/Niagara/Source/Niagara/Classes/NiagaraDataInterfaceParticleRead.h
@@ -38,20 +38,30 @@ public:
 	void GetNumSpawnedParticles(FVectorVMContext& Context);
 	void GetSpawnedIDAtIndex(FVectorVMContext& Context);
 	void GetNumParticles(FVectorVMContext& Context);
+	void ReadInt(FVectorVMContext& Context, FName AttributeToRead);
+	void ReadBool(FVectorVMContext& Context, FName AttributeToRead);
 	void ReadFloat(FVectorVMContext& Context, FName AttributeToRead);
 	void ReadVector2(FVectorVMContext& Context, FName AttributeToRead);
 	void ReadVector3(FVectorVMContext& Context, FName AttributeToRead);
 	void ReadVector4(FVectorVMContext& Context, FName AttributeToRead);
-	void ReadInt(FVectorVMContext& Context, FName AttributeToRead);
-	void ReadBool(FVectorVMContext& Context, FName AttributeToRead);
 	void ReadColor(FVectorVMContext& Context, FName AttributeToRead);
 	void ReadQuat(FVectorVMContext& Context, FName AttributeToRead);
+	void ReadIntByIndex(FVectorVMContext& Context, FName AttributeToRead);
+	void ReadBoolByIndex(FVectorVMContext& Context, FName AttributeToRead);
+	void ReadFloatByIndex(FVectorVMContext& Context, FName AttributeToRead);
+	void ReadVector2ByIndex(FVectorVMContext& Context, FName AttributeToRead);
+	void ReadVector3ByIndex(FVectorVMContext& Context, FName AttributeToRead);
+	void ReadVector4ByIndex(FVectorVMContext& Context, FName AttributeToRead);
+	void ReadColorByIndex(FVectorVMContext& Context, FName AttributeToRead);
+	void ReadQuatByIndex(FVectorVMContext& Context, FName AttributeToRead);
 
 protected:
 	virtual bool CopyToInternal(UNiagaraDataInterface* Destination) const override;
 
 private:
-	template<typename T, typename AccessorType>
-	T RetrieveValueWithCheck(FNiagaraEmitterInstance* EmitterInstance, const FName& Attr, const FNiagaraID& ParticleID, bool& bValid);
+	template<typename Type, typename AccessorType>
+	Type RetrieveValueWithCheck(FNiagaraEmitterInstance* EmitterInstance, const FName& Attr, const FNiagaraID& ParticleID, bool& bValid);
 
+	template<typename Type, typename AccessorType>
+	Type RetrieveValueByIndexWithCheck(FNiagaraEmitterInstance* EmitterInstance, const FName& Attr, int32 ParticleIndex, bool& bValid);
 };
\ No newline at end of file
diff --git a/Engine/Plugins/FX/Niagara/Source/Niagara/Private/NiagaraDataInterfaceParticleRead.cpp b/Engine/Plugins/FX/Niagara/Source/Niagara/Private/NiagaraDataInterfaceParticleRead.cpp
--- a/Engine/Plugins/FX/Niagara/Source/Niagara/Private/NiagaraDataInterfaceParticleRead.cpp
+++ b/Engine/Plugins/FX/Niagara/Source/Niagara/Private/NiagaraDataInterfaceParticleRead.cpp
@@ -12,15 +12,25 @@
 static const FName GetNumSpawnedParticlesFunctionName("Get Num Spawned Particles");
 static const FName GetSpawnedIDAtIndexFunctionName("Get Spawned ID At Index");
 static const FName GetNumParticlesFunctionName("Get Num Particles");
+
 static const FName GetIntAttributeFunctionName("Get int Attribute");
+static const FName GetBoolAttributeFunctionName("Get bool Attribute");
 static const FName GetFloatAttributeFunctionName("Get float Attribute");
 static const FName GetVec2AttributeFunctionName("Get Vector2 Attribute");
 static const FName GetVec3AttributeFunctionName("Get Vector3 Attribute");
 static const FName GetVec4AttributeFunctionName("Get Vector4 Attribute");
-static const FName GetBoolAttributeFunctionName("Get bool Attribute");
 static const FName GetColorAttributeFunctionName("Get Color Attribute");
 static const FName GetQuatAttributeFunctionName("Get Quaternion Attribute");
 
+static const FName GetIntAttributeByIndexFunctionName("Get int Attribute By Index");
+static const FName GetBoolAttributeByIndexFunctionName("Get bool Attribute By Index");
+static const FName GetFloatAttributeByIndexFunctionName("Get float Attribute By Index");
+static const FName GetVec2AttributeByIndexFunctionName("Get Vector2 Attribute By Index");
+static const FName GetVec3AttributeByIndexFunctionName("Get Vector3 Attribute By Index");
+static const FName GetVec4AttributeByIndexFunctionName("Get Vector4 Attribute By Index");
+static const FName GetColorAttributeByIndexFunctionName("Get Color Attribute By Index");
+static const FName GetQuatAttributeByIndexFunctionName("Get Quaternion Attribute By Index");
+
 static const FString NumSpawnedParticlesBaseName(TEXT("NumSpawnedParticles_"));
 static const FString SpawnedParticlesAcquireTagBaseName(TEXT("SpawnedParticlesAcquireTag_"));
 static const FString InstanceCountOffsetBaseName(TEXT("InstanceCountOffset_"));
@@ -164,14 +174,14 @@
 
 	ENiagaraParticleDataValueType GetValueTypeFromFuncName(const FName& FuncName)
 	{
-		if (FuncName == GetIntAttributeFunctionName) return ENiagaraParticleDataValueType::Int;
-		if (FuncName == GetFloatAttributeFunctionName) return ENiagaraParticleDataValueType::Float;
-		if (FuncName == GetVec2AttributeFunctionName) return ENiagaraParticleDataValueType::Vec2;
-		if (FuncName == GetVec3AttributeFunctionName) return ENiagaraParticleDataValueType::Vec3;
-		if (FuncName == GetVec4AttributeFunctionName) return ENiagaraParticleDataValueType::Vec4;
-		if (FuncName == GetBoolAttributeFunctionName) return ENiagaraParticleDataValueType::Bool;
-		if (FuncName == GetColorAttributeFunctionName) return ENiagaraParticleDataValueType::Color;
-		if (FuncName == GetQuatAttributeFunctionName) return ENiagaraParticleDataValueType::Quat;
+		if (FuncName == GetIntAttributeFunctionName || FuncName == GetIntAttributeByIndexFunctionName) return ENiagaraParticleDataValueType::Int;
+		if (FuncName == GetBoolAttributeFunctionName || FuncName == GetBoolAttributeByIndexFunctionName) return ENiagaraParticleDataValueType::Bool;
+		if (FuncName == GetFloatAttributeFunctionName || FuncName == GetFloatAttributeByIndexFunctionName) return ENiagaraParticleDataValueType::Float;
+		if (FuncName == GetVec2AttributeFunctionName || FuncName == GetVec2AttributeByIndexFunctionName) return ENiagaraParticleDataValueType::Vec2;
+		if (FuncName == GetVec3AttributeFunctionName || FuncName == GetVec3AttributeByIndexFunctionName) return ENiagaraParticleDataValueType::Vec3;
+		if (FuncName == GetVec4AttributeFunctionName || FuncName == GetVec4AttributeByIndexFunctionName) return ENiagaraParticleDataValueType::Vec4;
+		if (FuncName == GetColorAttributeFunctionName || FuncName == GetColorAttributeByIndexFunctionName) return ENiagaraParticleDataValueType::Color;
+		if (FuncName == GetQuatAttributeFunctionName || FuncName == GetQuatAttributeByIndexFunctionName) return ENiagaraParticleDataValueType::Quat;
 		return ENiagaraParticleDataValueType::Invalid;
 	}
 
@@ -258,11 +268,11 @@
 		switch (AttributeType)
 		{
 			case ENiagaraParticleDataValueType::Int: return VarType == FNiagaraTypeDefinition::GetIntDef();
+			case ENiagaraParticleDataValueType::Bool: return VarType == FNiagaraTypeDefinition::GetBoolDef();
 			case ENiagaraParticleDataValueType::Float: return VarType == FNiagaraTypeDefinition::GetFloatDef();
 			case ENiagaraParticleDataValueType::Vec2: return VarType == FNiagaraTypeDefinition::GetVec2Def();
 			case ENiagaraParticleDataValueType::Vec3: return VarType == FNiagaraTypeDefinition::GetVec3Def();
 			case ENiagaraParticleDataValueType::Vec4: return VarType == FNiagaraTypeDefinition::GetVec4Def();
-			case ENiagaraParticleDataValueType::Bool: return VarType == FNiagaraTypeDefinition::GetBoolDef();
 			case ENiagaraParticleDataValueType::Color: return VarType == FNiagaraTypeDefinition::GetColorDef();
 			case ENiagaraParticleDataValueType::Quat: return VarType == FNiagaraTypeDefinition::GetQuatDef();
 			default: return false;
@@ -468,8 +478,15 @@
 
 		if (!bReadingOwnEmitter)
 		{
-			FRHIUnorderedAccessView* InputBuffers[] = { SourceData->GetGPUIDToIndexTable().UAV, SourceData->GetGPUBufferFloat().UAV, SourceData->GetGPUBufferInt().UAV };
-			RHICmdList.TransitionResources(EResourceTransitionAccess::EReadable, EResourceTransitionPipeline::EComputeToCompute, InputBuffers, UE_ARRAY_COUNT(InputBuffers));
+			FRHIUnorderedAccessView* InputBuffers[3];
+			int32 NumTransitions = 0;
+			InputBuffers[NumTransitions++] = SourceData->GetGPUBufferFloat().UAV;
+			InputBuffers[NumTransitions++] = SourceData->GetGPUBufferInt().UAV;
+			if (SourceData->GetGPUIDToIndexTable().UAV)
+			{
+				InputBuffers[NumTransitions++] = SourceData->GetGPUIDToIndexTable().UAV;
+			}
+			RHICmdList.TransitionResources(EResourceTransitionAccess::EReadable, EResourceTransitionPipeline::EComputeToCompute, InputBuffers, NumTransitions);
 
 			if (InstanceCountOffsetParam.IsBound())
 			{
@@ -588,6 +605,9 @@
 
 void UNiagaraDataInterfaceParticleRead::GetFunctions(TArray<FNiagaraFunctionSignature>& OutFunctions)
 {
+	//
+	// Spawn info and particle count
+	//
 	{
 		FNiagaraFunctionSignature Sig;
 		Sig.Name = GetNumSpawnedParticlesFunctionName;
@@ -597,6 +617,7 @@
 		Sig.bRequiresContext = false;
 		OutFunctions.Add(Sig);
 	}
+
 	{
 		FNiagaraFunctionSignature Sig;
 		Sig.Name = GetSpawnedIDAtIndexFunctionName;
@@ -611,6 +632,7 @@
 		Sig.bRequiresContext = false;
 		OutFunctions.Add(Sig);
 	}
+
 	{
 		FNiagaraFunctionSignature Sig;
 		Sig.Name = GetNumParticlesFunctionName;
@@ -620,6 +642,42 @@
 		Sig.bRequiresContext = false;
 		OutFunctions.Add(Sig);
 	}
+
+	//
+	// Get attribute by ID
+	//
 	{
 		FNiagaraFunctionSignature Sig;
+		Sig.Name = GetIntAttributeFunctionName;
+		Sig.Inputs.Add(FNiagaraVariable(FNiagaraTypeDefinition(GetClass()), TEXT("Particle Reader")));
+		Sig.Inputs.Add(FNiagaraVariable(FNiagaraTypeDefinition::GetIDDef(), TEXT("Particle ID")));
+
+		Sig.Outputs.Add(FNiagaraVariable(FNiagaraTypeDefinition::GetBoolDef(), TEXT("Valid")));
+		Sig.Outputs.Add(FNiagaraVariable(FNiagaraTypeDefinition::GetIntDef(), TEXT("Value")));
+
+		Sig.FunctionSpecifiers.Add(FName("Attribute"));
+
+		Sig.bMemberFunction = true;
+		Sig.bRequiresContext = false;
+		OutFunctions.Add(Sig);
+	}
+
+	{
+		FNiagaraFunctionSignature Sig;
+		Sig.Name = GetBoolAttributeFunctionName;
+		Sig.Inputs.Add(FNiagaraVariable(FNiagaraTypeDefinition(GetClass()), TEXT("Particle Reader")));
+		Sig.Inputs.Add(FNiagaraVariable(FNiagaraTypeDefinition::GetIDDef(), TEXT("Particle ID")));
+
+		Sig.Outputs.Add(FNiagaraVariable(FNiagaraTypeDefinition::GetBoolDef(), TEXT("Valid")));
+		Sig.Outputs.Add(FNiagaraVariable(FNiagaraTypeDefinition::GetBoolDef(), TEXT("Value")));
+
+		Sig.FunctionSpecifiers.Add(FName("Attribute"));
+
+		Sig.bMemberFunction = true;
+		Sig.bRequiresContext = false;
+		OutFunctions.Add(Sig);
+	}
+
+	{
+		FNiagaraFunctionSignature Sig;
 		Sig.Name = GetFloatAttributeFunctionName;
@@ -635,6 +693,7 @@
 		Sig.bRequiresContext = false;
 		OutFunctions.Add(Sig);
 	}
+
 	{
 		FNiagaraFunctionSignature Sig;
 		Sig.Name = GetVec2AttributeFunctionName;
@@ -650,6 +709,7 @@
 		Sig.bRequiresContext = false;
 		OutFunctions.Add(Sig);
 	}
+
 	{
 		FNiagaraFunctionSignature Sig;
 		Sig.Name = GetVec3AttributeFunctionName;
@@ -665,6 +725,7 @@
 		Sig.bRequiresContext = false;
 		OutFunctions.Add(Sig);
 	}
+
 	{
 		FNiagaraFunctionSignature Sig;
 		Sig.Name = GetVec4AttributeFunctionName;
@@ -680,14 +741,15 @@
 		Sig.bRequiresContext = false;
 		OutFunctions.Add(Sig);
 	}
+
 	{
 		FNiagaraFunctionSignature Sig;
-		Sig.Name = GetIntAttributeFunctionName;
+		Sig.Name = GetColorAttributeFunctionName;
 		Sig.Inputs.Add(FNiagaraVariable(FNiagaraTypeDefinition(GetClass()), TEXT("Particle Reader")));
 		Sig.Inputs.Add(FNiagaraVariable(FNiagaraTypeDefinition::GetIDDef(), TEXT("Particle ID")));
 
 		Sig.Outputs.Add(FNiagaraVariable(FNiagaraTypeDefinition::GetBoolDef(), TEXT("Valid")));
-		Sig.Outputs.Add(FNiagaraVariable(FNiagaraTypeDefinition::GetIntDef(), TEXT("Value")));
+		Sig.Outputs.Add(FNiagaraVariable(FNiagaraTypeDefinition::GetColorDef(), TEXT("Value")));
 
 		Sig.FunctionSpecifiers.Add(FName("Attribute"));
 
@@ -695,12 +757,48 @@
 		Sig.bRequiresContext = false;
 		OutFunctions.Add(Sig);
 	}
+
 	{
 		FNiagaraFunctionSignature Sig;
-		Sig.Name = GetBoolAttributeFunctionName;
+		Sig.Name = GetQuatAttributeFunctionName;
 		Sig.Inputs.Add(FNiagaraVariable(FNiagaraTypeDefinition(GetClass()), TEXT("Particle Reader")));
 		Sig.Inputs.Add(FNiagaraVariable(FNiagaraTypeDefinition::GetIDDef(), TEXT("Particle ID")));
 
 		Sig.Outputs.Add(FNiagaraVariable(FNiagaraTypeDefinition::GetBoolDef(), TEXT("Valid")));
+		Sig.Outputs.Add(FNiagaraVariable(FNiagaraTypeDefinition::GetQuatDef(), TEXT("Value")));
+
+		Sig.FunctionSpecifiers.Add(FName("Attribute"));
+
+		Sig.bMemberFunction = true;
+		Sig.bRequiresContext = false;
+		OutFunctions.Add(Sig);
+	}
+
+	//
+	// Get attribute by index
+	//
+	{
+		FNiagaraFunctionSignature Sig;
+		Sig.Name = GetIntAttributeByIndexFunctionName;
+		Sig.Inputs.Add(FNiagaraVariable(FNiagaraTypeDefinition(GetClass()), TEXT("Particle Reader")));
+		Sig.Inputs.Add(FNiagaraVariable(FNiagaraTypeDefinition::GetIntDef(), TEXT("Particle Index")));
+
+		Sig.Outputs.Add(FNiagaraVariable(FNiagaraTypeDefinition::GetBoolDef(), TEXT("Valid")));
+		Sig.Outputs.Add(FNiagaraVariable(FNiagaraTypeDefinition::GetIntDef(), TEXT("Value")));
+
+		Sig.FunctionSpecifiers.Add(FName("Attribute"));
+
+		Sig.bMemberFunction = true;
+		Sig.bRequiresContext = false;
+		OutFunctions.Add(Sig);
+	}
+
+	{
+		FNiagaraFunctionSignature Sig;
+		Sig.Name = GetBoolAttributeByIndexFunctionName;
+		Sig.Inputs.Add(FNiagaraVariable(FNiagaraTypeDefinition(GetClass()), TEXT("Particle Reader")));
+		Sig.Inputs.Add(FNiagaraVariable(FNiagaraTypeDefinition::GetIntDef(), TEXT("Particle Index")));
+
+		Sig.Outputs.Add(FNiagaraVariable(FNiagaraTypeDefinition::GetBoolDef(), TEXT("Valid")));
 		Sig.Outputs.Add(FNiagaraVariable(FNiagaraTypeDefinition::GetBoolDef(), TEXT("Value")));
 
@@ -710,11 +808,76 @@
 		Sig.bRequiresContext = false;
 		OutFunctions.Add(Sig);
 	}
+
 	{
 		FNiagaraFunctionSignature Sig;
-		Sig.Name = GetColorAttributeFunctionName;
+		Sig.Name = GetFloatAttributeByIndexFunctionName;
 		Sig.Inputs.Add(FNiagaraVariable(FNiagaraTypeDefinition(GetClass()), TEXT("Particle Reader")));
-		Sig.Inputs.Add(FNiagaraVariable(FNiagaraTypeDefinition::GetIDDef(), TEXT("Particle ID")));
+		Sig.Inputs.Add(FNiagaraVariable(FNiagaraTypeDefinition::GetIntDef(), TEXT("Particle Index")));
 
 		Sig.Outputs.Add(FNiagaraVariable(FNiagaraTypeDefinition::GetBoolDef(), TEXT("Valid")));
+		Sig.Outputs.Add(FNiagaraVariable(FNiagaraTypeDefinition::GetFloatDef(), TEXT("Value")));
+
+		Sig.FunctionSpecifiers.Add(FName("Attribute"));
+
+		Sig.bMemberFunction = true;
+		Sig.bRequiresContext = false;
+		OutFunctions.Add(Sig);
+	}
+
+	{
+		FNiagaraFunctionSignature Sig;
+		Sig.Name = GetVec2AttributeByIndexFunctionName;
+		Sig.Inputs.Add(FNiagaraVariable(FNiagaraTypeDefinition(GetClass()), TEXT("Particle Reader")));
+		Sig.Inputs.Add(FNiagaraVariable(FNiagaraTypeDefinition::GetIntDef(), TEXT("Particle Index")));
+
+		Sig.Outputs.Add(FNiagaraVariable(FNiagaraTypeDefinition::GetBoolDef(), TEXT("Valid")));
+		Sig.Outputs.Add(FNiagaraVariable(FNiagaraTypeDefinition::GetVec2Def(), TEXT("Value")));
+
+		Sig.FunctionSpecifiers.Add(FName("Attribute"));
+
+		Sig.bMemberFunction = true;
+		Sig.bRequiresContext = false;
+		OutFunctions.Add(Sig);
+	}
+
+	{
+		FNiagaraFunctionSignature Sig;
+		Sig.Name = GetVec3AttributeByIndexFunctionName;
+		Sig.Inputs.Add(FNiagaraVariable(FNiagaraTypeDefinition(GetClass()), TEXT("Particle Reader")));
+		Sig.Inputs.Add(FNiagaraVariable(FNiagaraTypeDefinition::GetIntDef(), TEXT("Particle Index")));
+
+		Sig.Outputs.Add(FNiagaraVariable(FNiagaraTypeDefinition::GetBoolDef(), TEXT("Valid")));
+		Sig.Outputs.Add(FNiagaraVariable(FNiagaraTypeDefinition::GetVec3Def(), TEXT("Value")));
+
+		Sig.FunctionSpecifiers.Add(FName("Attribute"));
+
+		Sig.bMemberFunction = true;
+		Sig.bRequiresContext = false;
+		OutFunctions.Add(Sig);
+	}
+
+	{
+		FNiagaraFunctionSignature Sig;
+		Sig.Name = GetVec4AttributeByIndexFunctionName;
+		Sig.Inputs.Add(FNiagaraVariable(FNiagaraTypeDefinition(GetClass()), TEXT("Particle Reader")));
+		Sig.Inputs.Add(FNiagaraVariable(FNiagaraTypeDefinition::GetIntDef(), TEXT("Particle Index")));
+
+		Sig.Outputs.Add(FNiagaraVariable(FNiagaraTypeDefinition::GetBoolDef(), TEXT("Valid")));
+		Sig.Outputs.Add(FNiagaraVariable(FNiagaraTypeDefinition::GetVec4Def(), TEXT("Value")));
+
+		Sig.FunctionSpecifiers.Add(FName("Attribute"));
+
+		Sig.bMemberFunction = true;
+		Sig.bRequiresContext = false;
+		OutFunctions.Add(Sig);
+	}
+
+	{
+		FNiagaraFunctionSignature Sig;
+		Sig.Name = GetColorAttributeByIndexFunctionName;
+		Sig.Inputs.Add(FNiagaraVariable(FNiagaraTypeDefinition(GetClass()), TEXT("Particle Reader")));
+		Sig.Inputs.Add(FNiagaraVariable(FNiagaraTypeDefinition::GetIntDef(), TEXT("Particle Index")));
+
+		Sig.Outputs.Add(FNiagaraVariable(FNiagaraTypeDefinition::GetBoolDef(), TEXT("Valid")));
 		Sig.Outputs.Add(FNiagaraVariable(FNiagaraTypeDefinition::GetColorDef(), TEXT("Value")));
@@ -725,11 +888,12 @@
 		Sig.bRequiresContext = false;
 		OutFunctions.Add(Sig);
 	}
+
 	{
 		FNiagaraFunctionSignature Sig;
-		Sig.Name = GetQuatAttributeFunctionName;
+		Sig.Name = GetQuatAttributeByIndexFunctionName;
 		Sig.Inputs.Add(FNiagaraVariable(FNiagaraTypeDefinition(GetClass()), TEXT("Particle Reader")));
-		Sig.Inputs.Add(FNiagaraVariable(FNiagaraTypeDefinition::GetIDDef(), TEXT("Particle ID")));
+		Sig.Inputs.Add(FNiagaraVariable(FNiagaraTypeDefinition::GetIntDef(), TEXT("Particle Index")));
 
 		Sig.Outputs.Add(FNiagaraVariable(FNiagaraTypeDefinition::GetBoolDef(), TEXT("Valid")));
 		Sig.Outputs.Add(FNiagaraVariable(FNiagaraTypeDefinition::GetQuatDef(), TEXT("Value")));
@@ -739,20 +903,28 @@
 		Sig.bMemberFunction = true;
 		Sig.bRequiresContext = false;
 		OutFunctions.Add(Sig);
-	}
+	} 
 }
 
 DEFINE_NDI_DIRECT_FUNC_BINDER(UNiagaraDataInterfaceParticleRead, GetNumSpawnedParticles);
 DEFINE_NDI_DIRECT_FUNC_BINDER(UNiagaraDataInterfaceParticleRead, GetSpawnedIDAtIndex);
 DEFINE_NDI_DIRECT_FUNC_BINDER(UNiagaraDataInterfaceParticleRead, GetNumParticles);
+DEFINE_NDI_DIRECT_FUNC_BINDER_WITH_PAYLOAD(UNiagaraDataInterfaceParticleRead, ReadInt);
+DEFINE_NDI_DIRECT_FUNC_BINDER_WITH_PAYLOAD(UNiagaraDataInterfaceParticleRead, ReadBool);
 DEFINE_NDI_DIRECT_FUNC_BINDER_WITH_PAYLOAD(UNiagaraDataInterfaceParticleRead, ReadFloat);
 DEFINE_NDI_DIRECT_FUNC_BINDER_WITH_PAYLOAD(UNiagaraDataInterfaceParticleRead, ReadVector2);
 DEFINE_NDI_DIRECT_FUNC_BINDER_WITH_PAYLOAD(UNiagaraDataInterfaceParticleRead, ReadVector3);
 DEFINE_NDI_DIRECT_FUNC_BINDER_WITH_PAYLOAD(UNiagaraDataInterfaceParticleRead, ReadVector4);
-DEFINE_NDI_DIRECT_FUNC_BINDER_WITH_PAYLOAD(UNiagaraDataInterfaceParticleRead, ReadInt);
-DEFINE_NDI_DIRECT_FUNC_BINDER_WITH_PAYLOAD(UNiagaraDataInterfaceParticleRead, ReadBool);
 DEFINE_NDI_DIRECT_FUNC_BINDER_WITH_PAYLOAD(UNiagaraDataInterfaceParticleRead, ReadColor);
 DEFINE_NDI_DIRECT_FUNC_BINDER_WITH_PAYLOAD(UNiagaraDataInterfaceParticleRead, ReadQuat);
+DEFINE_NDI_DIRECT_FUNC_BINDER_WITH_PAYLOAD(UNiagaraDataInterfaceParticleRead, ReadIntByIndex);
+DEFINE_NDI_DIRECT_FUNC_BINDER_WITH_PAYLOAD(UNiagaraDataInterfaceParticleRead, ReadBoolByIndex);
+DEFINE_NDI_DIRECT_FUNC_BINDER_WITH_PAYLOAD(UNiagaraDataInterfaceParticleRead, ReadFloatByIndex);
+DEFINE_NDI_DIRECT_FUNC_BINDER_WITH_PAYLOAD(UNiagaraDataInterfaceParticleRead, ReadVector2ByIndex);
+DEFINE_NDI_DIRECT_FUNC_BINDER_WITH_PAYLOAD(UNiagaraDataInterfaceParticleRead, ReadVector3ByIndex);
+DEFINE_NDI_DIRECT_FUNC_BINDER_WITH_PAYLOAD(UNiagaraDataInterfaceParticleRead, ReadVector4ByIndex);
+DEFINE_NDI_DIRECT_FUNC_BINDER_WITH_PAYLOAD(UNiagaraDataInterfaceParticleRead, ReadColorByIndex);
+DEFINE_NDI_DIRECT_FUNC_BINDER_WITH_PAYLOAD(UNiagaraDataInterfaceParticleRead, ReadQuatByIndex);
 
 static bool HasMatchingVariable(TArrayView<const FNiagaraVariable> Variables, FName AttributeName, const FNiagaraTypeDefinition& ValidType)
 {
@@ -761,6 +933,9 @@
 
 void UNiagaraDataInterfaceParticleRead::GetVMExternalFunction(const FVMExternalFunctionBindingInfo& BindingInfo, void* InstanceData, FVMExternalFunction& OutFunc)
 {
+	//
+	// Spawn info and particle count
+	//
 	if (BindingInfo.Name == GetNumSpawnedParticlesFunctionName)
 	{
 		NDI_FUNC_BINDER(UNiagaraDataInterfaceParticleRead, GetNumSpawnedParticles)::Bind(this, OutFunc);
@@ -793,7 +968,27 @@
 	TArrayView<const FNiagaraVariable> EmitterVariables = PIData->EmitterInstance->GetData().GetVariables();
 
 	const FName AttributeToRead = FunctionSpecifier->Value;
-	if (BindingInfo.Name == GetFloatAttributeFunctionName)
+
+	//
+	// Get attribute by ID
+	//
+	if (BindingInfo.Name == GetIntAttributeFunctionName)
 	{
+		if (HasMatchingVariable(EmitterVariables, AttributeToRead, FNiagaraTypeDefinition::GetIntDef()))
+		{
+			NDI_FUNC_BINDER(UNiagaraDataInterfaceParticleRead, ReadInt)::Bind(this, OutFunc, AttributeToRead);
+			bBindSuccessful = true;
+		}
+	}
+	else if (BindingInfo.Name == GetBoolAttributeFunctionName)
+	{
+		if (HasMatchingVariable(EmitterVariables, AttributeToRead, FNiagaraTypeDefinition::GetBoolDef()))
+		{
+			NDI_FUNC_BINDER(UNiagaraDataInterfaceParticleRead, ReadBool)::Bind(this, OutFunc, AttributeToRead);
+			bBindSuccessful = true;
+		}
+	}
+	else if (BindingInfo.Name == GetFloatAttributeFunctionName)
+	{
 		if (HasMatchingVariable(EmitterVariables, AttributeToRead, FNiagaraTypeDefinition::GetFloatDef())
 			|| HasMatchingVariable(EmitterVariables, AttributeToRead, FNiagaraTypeDefinition::GetHalfDef()))
@@ -829,22 +1024,6 @@
 			bBindSuccessful = true;
 		}
 	}
-	else if (BindingInfo.Name == GetIntAttributeFunctionName)
-	{
-		if (HasMatchingVariable(EmitterVariables, AttributeToRead, FNiagaraTypeDefinition::GetIntDef()))
-		{
-			NDI_FUNC_BINDER(UNiagaraDataInterfaceParticleRead, ReadInt)::Bind(this, OutFunc, AttributeToRead);
-			bBindSuccessful = true;
-		}
-	}
-	else if (BindingInfo.Name == GetBoolAttributeFunctionName)
-	{
-		if (HasMatchingVariable(EmitterVariables, AttributeToRead, FNiagaraTypeDefinition::GetBoolDef()))
-		{
-			NDI_FUNC_BINDER(UNiagaraDataInterfaceParticleRead, ReadBool)::Bind(this, OutFunc, AttributeToRead);
-			bBindSuccessful = true;
-		}
-	}
 	else if (BindingInfo.Name == GetColorAttributeFunctionName)
 	{
 		if (HasMatchingVariable(EmitterVariables, AttributeToRead, FNiagaraTypeDefinition::GetColorDef())
@@ -863,3 +1042,76 @@
 			bBindSuccessful = true;
 		}
 	}
+	//
+	// Get attribute by index
+	//
+	else if (BindingInfo.Name == GetIntAttributeByIndexFunctionName)
+	{
+		if (HasMatchingVariable(EmitterVariables, AttributeToRead, FNiagaraTypeDefinition::GetIntDef()))
+		{
+			NDI_FUNC_BINDER(UNiagaraDataInterfaceParticleRead, ReadIntByIndex)::Bind(this, OutFunc, AttributeToRead);
+			bBindSuccessful = true;
+		}
+	}
+	else if (BindingInfo.Name == GetBoolAttributeByIndexFunctionName)
+	{
+		if (HasMatchingVariable(EmitterVariables, AttributeToRead, FNiagaraTypeDefinition::GetBoolDef()))
+		{
+			NDI_FUNC_BINDER(UNiagaraDataInterfaceParticleRead, ReadBoolByIndex)::Bind(this, OutFunc, AttributeToRead);
+			bBindSuccessful = true;
+		}
+	}
+	else if (BindingInfo.Name == GetFloatAttributeByIndexFunctionName)
+	{
+		if (HasMatchingVariable(EmitterVariables, AttributeToRead, FNiagaraTypeDefinition::GetFloatDef())
+			|| HasMatchingVariable(EmitterVariables, AttributeToRead, FNiagaraTypeDefinition::GetHalfDef()))
+		{
+			NDI_FUNC_BINDER(UNiagaraDataInterfaceParticleRead, ReadFloatByIndex)::Bind(this, OutFunc, AttributeToRead);
+			bBindSuccessful = true;
+		}
+	}
+	else if (BindingInfo.Name == GetVec2AttributeByIndexFunctionName)
+	{
+		if (HasMatchingVariable(EmitterVariables, AttributeToRead, FNiagaraTypeDefinition::GetVec2Def())
+			|| HasMatchingVariable(EmitterVariables, AttributeToRead, FNiagaraTypeDefinition::GetHalfVec2Def()))
+		{
+			NDI_FUNC_BINDER(UNiagaraDataInterfaceParticleRead, ReadVector2ByIndex)::Bind(this, OutFunc, AttributeToRead);
+			bBindSuccessful = true;
+		}
+	}
+	else if (BindingInfo.Name == GetVec3AttributeByIndexFunctionName)
+	{
+		if (HasMatchingVariable(EmitterVariables, AttributeToRead, FNiagaraTypeDefinition::GetVec3Def())
+			|| HasMatchingVariable(EmitterVariables, AttributeToRead, FNiagaraTypeDefinition::GetHalfVec3Def()))
+		{
+			NDI_FUNC_BINDER(UNiagaraDataInterfaceParticleRead, ReadVector3ByIndex)::Bind(this, OutFunc, AttributeToRead);
+			bBindSuccessful = true;
+		}
+	}
+	else if (BindingInfo.Name == GetVec4AttributeByIndexFunctionName)
+	{
+		if (HasMatchingVariable(EmitterVariables, AttributeToRead, FNiagaraTypeDefinition::GetVec4Def())
+			|| HasMatchingVariable(EmitterVariables, AttributeToRead, FNiagaraTypeDefinition::GetHalfVec4Def()))
+		{
+			NDI_FUNC_BINDER(UNiagaraDataInterfaceParticleRead, ReadVector4ByIndex)::Bind(this, OutFunc, AttributeToRead);
+			bBindSuccessful = true;
+		}
+	}
+	else if (BindingInfo.Name == GetColorAttributeByIndexFunctionName)
+	{
+		if (HasMatchingVariable(EmitterVariables, AttributeToRead, FNiagaraTypeDefinition::GetColorDef())
+			|| HasMatchingVariable(EmitterVariables, AttributeToRead, FNiagaraTypeDefinition::GetHalfVec4Def()))
+		{
+			NDI_FUNC_BINDER(UNiagaraDataInterfaceParticleRead, ReadColorByIndex)::Bind(this, OutFunc, AttributeToRead);
+			bBindSuccessful = true;
+		}
+	}
+	else if (BindingInfo.Name == GetQuatAttributeByIndexFunctionName)
+	{
+		if (HasMatchingVariable(EmitterVariables, AttributeToRead, FNiagaraTypeDefinition::GetQuatDef())
+			|| HasMatchingVariable(EmitterVariables, AttributeToRead, FNiagaraTypeDefinition::GetHalfVec4Def()))
+		{
+			NDI_FUNC_BINDER(UNiagaraDataInterfaceParticleRead, ReadQuatByIndex)::Bind(this, OutFunc, AttributeToRead);
+			bBindSuccessful = true;
+		}
+	}
@@ -940,7 +1192,7 @@
 	}
 }
 
-void UNiagaraDataInterfaceParticleRead::ReadFloat(FVectorVMContext& Context, FName AttributeToRead)
+void UNiagaraDataInterfaceParticleRead::ReadInt(FVectorVMContext& Context, FName AttributeToRead)
 {
 	VectorVM::FExternalFuncInputHandler<int32> ParticleIDIndexParam(Context);
 	VectorVM::FExternalFuncInputHandler<int32> ParticleIDAcquireTagParam(Context);
@@ -948,13 +1200,13 @@
 	VectorVM::FUserPtrHandler<FNDIParticleRead_InstanceData> InstanceData(Context);
 
 	VectorVM::FExternalFuncRegisterHandler<FNiagaraBool> OutValid(Context);
-	VectorVM::FExternalFuncRegisterHandler<float> OutValue(Context);
-	
+	VectorVM::FExternalFuncRegisterHandler<int32> OutValue(Context);
+
 	for (int32 InstanceIdx = 0; InstanceIdx < Context.NumInstances; ++InstanceIdx)
 	{
 		FNiagaraID ParticleID = { ParticleIDIndexParam.GetAndAdvance(), ParticleIDAcquireTagParam.GetAndAdvance() };
 		bool bValid;
-		float Value = RetrieveValueWithCheck<float, FNiagaraDataConversions::FHalfOrFloat>(InstanceData->EmitterInstance, AttributeToRead, ParticleID, bValid);
+		int32 Value = RetrieveValueWithCheck<int, FNiagaraDataConversions::FNiagaraInt>(InstanceData->EmitterInstance, AttributeToRead, ParticleID, bValid);
 		FNiagaraBool ValidValue;
 		ValidValue.SetValue(bValid);
 		*OutValid.GetDestAndAdvance() = ValidValue;
@@ -962,7 +1214,7 @@
 	}
 }
 
-void UNiagaraDataInterfaceParticleRead::ReadVector2(FVectorVMContext& Context, FName AttributeToRead)
+void UNiagaraDataInterfaceParticleRead::ReadBool(FVectorVMContext& Context, FName AttributeToRead)
 {
 	VectorVM::FExternalFuncInputHandler<int32> ParticleIDIndexParam(Context);
 	VectorVM::FExternalFuncInputHandler<int32> ParticleIDAcquireTagParam(Context);
@@ -970,25 +1222,21 @@
 	VectorVM::FUserPtrHandler<FNDIParticleRead_InstanceData> InstanceData(Context);
 
 	VectorVM::FExternalFuncRegisterHandler<FNiagaraBool> OutValid(Context);
-	VectorVM::FExternalFuncRegisterHandler<float> OutX(Context);
-	VectorVM::FExternalFuncRegisterHandler<float> OutY(Context);
+	VectorVM::FExternalFuncRegisterHandler<FNiagaraBool> OutValue(Context);
 
 	for (int32 InstanceIdx = 0; InstanceIdx < Context.NumInstances; ++InstanceIdx)
 	{
 		FNiagaraID ParticleID = { ParticleIDIndexParam.GetAndAdvance(), ParticleIDAcquireTagParam.GetAndAdvance() };
 		bool bValid;
-
-		FVector2D Value = RetrieveValueWithCheck<FVector2D, FNiagaraDataConversions::FHalf2OrFloat2>(InstanceData->EmitterInstance, AttributeToRead, ParticleID, bValid);
-
+		FNiagaraBool Value = RetrieveValueWithCheck<FNiagaraBool, FNiagaraDataConversions::FNiagaraBool>(InstanceData->EmitterInstance, AttributeToRead, ParticleID, bValid);
 		FNiagaraBool ValidValue;
 		ValidValue.SetValue(bValid);
 		*OutValid.GetDestAndAdvance() = ValidValue;
-		*OutX.GetDestAndAdvance() = Value.X;
-		*OutY.GetDestAndAdvance() = Value.Y;
+		*OutValue.GetDestAndAdvance() = Value;
 	}
 }
 
-void UNiagaraDataInterfaceParticleRead::ReadVector3(FVectorVMContext& Context, FName AttributeToRead)
+void UNiagaraDataInterfaceParticleRead::ReadFloat(FVectorVMContext& Context, FName AttributeToRead)
 {
 	VectorVM::FExternalFuncInputHandler<int32> ParticleIDIndexParam(Context);
 	VectorVM::FExternalFuncInputHandler<int32> ParticleIDAcquireTagParam(Context);
@@ -996,27 +1244,21 @@
 	VectorVM::FUserPtrHandler<FNDIParticleRead_InstanceData> InstanceData(Context);
 
 	VectorVM::FExternalFuncRegisterHandler<FNiagaraBool> OutValid(Context);
-	VectorVM::FExternalFuncRegisterHandler<float> OutX(Context);
-	VectorVM::FExternalFuncRegisterHandler<float> OutY(Context);
-	VectorVM::FExternalFuncRegisterHandler<float> OutZ(Context);
-
+	VectorVM::FExternalFuncRegisterHandler<float> OutValue(Context);
+	
 	for (int32 InstanceIdx = 0; InstanceIdx < Context.NumInstances; ++InstanceIdx)
 	{
 		FNiagaraID ParticleID = { ParticleIDIndexParam.GetAndAdvance(), ParticleIDAcquireTagParam.GetAndAdvance() };
 		bool bValid;
-
-		FVector Value = RetrieveValueWithCheck<FVector, FNiagaraDataConversions::FHalf3OrFloat3>(InstanceData->EmitterInstance, AttributeToRead, ParticleID, bValid);
-
+		float Value = RetrieveValueWithCheck<float, FNiagaraDataConversions::FHalfOrFloat>(InstanceData->EmitterInstance, AttributeToRead, ParticleID, bValid);
 		FNiagaraBool ValidValue;
 		ValidValue.SetValue(bValid);
 		*OutValid.GetDestAndAdvance() = ValidValue;
-		*OutX.GetDestAndAdvance() = Value.X;
-		*OutY.GetDestAndAdvance() = Value.Y;
-		*OutZ.GetDestAndAdvance() = Value.Z;
+		*OutValue.GetDestAndAdvance() = Value;
 	}
 }
 
-void UNiagaraDataInterfaceParticleRead::ReadVector4(FVectorVMContext& Context, FName AttributeToRead)
+void UNiagaraDataInterfaceParticleRead::ReadVector2(FVectorVMContext& Context, FName AttributeToRead)
 {
 	VectorVM::FExternalFuncInputHandler<int32> ParticleIDIndexParam(Context);
 	VectorVM::FExternalFuncInputHandler<int32> ParticleIDAcquireTagParam(Context);
@@ -1026,27 +1268,21 @@
 	VectorVM::FExternalFuncRegisterHandler<FNiagaraBool> OutValid(Context);
 	VectorVM::FExternalFuncRegisterHandler<float> OutX(Context);
 	VectorVM::FExternalFuncRegisterHandler<float> OutY(Context);
-	VectorVM::FExternalFuncRegisterHandler<float> OutZ(Context);
-	VectorVM::FExternalFuncRegisterHandler<float> OutW(Context);
 
 	for (int32 InstanceIdx = 0; InstanceIdx < Context.NumInstances; ++InstanceIdx)
 	{
 		FNiagaraID ParticleID = { ParticleIDIndexParam.GetAndAdvance(), ParticleIDAcquireTagParam.GetAndAdvance() };
 		bool bValid;
-
-		FVector4 Value = RetrieveValueWithCheck<FVector4, FNiagaraDataConversions::FHalf4OrFloat4>(InstanceData->EmitterInstance, AttributeToRead, ParticleID, bValid);
-
+		FVector2D Value = RetrieveValueWithCheck<FVector2D, FNiagaraDataConversions::FHalf2OrFloat2>(InstanceData->EmitterInstance, AttributeToRead, ParticleID, bValid);
 		FNiagaraBool ValidValue;
 		ValidValue.SetValue(bValid);
 		*OutValid.GetDestAndAdvance() = ValidValue;
 		*OutX.GetDestAndAdvance() = Value.X;
 		*OutY.GetDestAndAdvance() = Value.Y;
-		*OutZ.GetDestAndAdvance() = Value.Z;
-		*OutW.GetDestAndAdvance() = Value.W;
 	}
 }
 
-void UNiagaraDataInterfaceParticleRead::ReadInt(FVectorVMContext& Context, FName AttributeToRead)
+void UNiagaraDataInterfaceParticleRead::ReadVector3(FVectorVMContext& Context, FName AttributeToRead)
 {
 	VectorVM::FExternalFuncInputHandler<int32> ParticleIDIndexParam(Context);
 	VectorVM::FExternalFuncInputHandler<int32> ParticleIDAcquireTagParam(Context);
@@ -1054,23 +1290,25 @@
 	VectorVM::FUserPtrHandler<FNDIParticleRead_InstanceData> InstanceData(Context);
 
 	VectorVM::FExternalFuncRegisterHandler<FNiagaraBool> OutValid(Context);
-	VectorVM::FExternalFuncRegisterHandler<int32> OutValue(Context);
+	VectorVM::FExternalFuncRegisterHandler<float> OutX(Context);
+	VectorVM::FExternalFuncRegisterHandler<float> OutY(Context);
+	VectorVM::FExternalFuncRegisterHandler<float> OutZ(Context);
 
 	for (int32 InstanceIdx = 0; InstanceIdx < Context.NumInstances; ++InstanceIdx)
 	{
 		FNiagaraID ParticleID = { ParticleIDIndexParam.GetAndAdvance(), ParticleIDAcquireTagParam.GetAndAdvance() };
 		bool bValid;
-
-		int32 Value = RetrieveValueWithCheck<int, FNiagaraDataConversions::FNiagaraInt>(InstanceData->EmitterInstance, AttributeToRead, ParticleID, bValid);
-
+		FVector Value = RetrieveValueWithCheck<FVector, FNiagaraDataConversions::FHalf3OrFloat3>(InstanceData->EmitterInstance, AttributeToRead, ParticleID, bValid);
 		FNiagaraBool ValidValue;
 		ValidValue.SetValue(bValid);
 		*OutValid.GetDestAndAdvance() = ValidValue;
-		*OutValue.GetDestAndAdvance() = Value;
+		*OutX.GetDestAndAdvance() = Value.X;
+		*OutY.GetDestAndAdvance() = Value.Y;
+		*OutZ.GetDestAndAdvance() = Value.Z;
 	}
 }
 
-void UNiagaraDataInterfaceParticleRead::ReadBool(FVectorVMContext& Context, FName AttributeToRead)
+void UNiagaraDataInterfaceParticleRead::ReadVector4(FVectorVMContext& Context, FName AttributeToRead)
 {
 	VectorVM::FExternalFuncInputHandler<int32> ParticleIDIndexParam(Context);
 	VectorVM::FExternalFuncInputHandler<int32> ParticleIDAcquireTagParam(Context);
@@ -1078,19 +1316,23 @@
 	VectorVM::FUserPtrHandler<FNDIParticleRead_InstanceData> InstanceData(Context);
 
 	VectorVM::FExternalFuncRegisterHandler<FNiagaraBool> OutValid(Context);
-	VectorVM::FExternalFuncRegisterHandler<FNiagaraBool> OutValue(Context);
+	VectorVM::FExternalFuncRegisterHandler<float> OutX(Context);
+	VectorVM::FExternalFuncRegisterHandler<float> OutY(Context);
+	VectorVM::FExternalFuncRegisterHandler<float> OutZ(Context);
+	VectorVM::FExternalFuncRegisterHandler<float> OutW(Context);
 
 	for (int32 InstanceIdx = 0; InstanceIdx < Context.NumInstances; ++InstanceIdx)
 	{
 		FNiagaraID ParticleID = { ParticleIDIndexParam.GetAndAdvance(), ParticleIDAcquireTagParam.GetAndAdvance() };
 		bool bValid;
-
-		FNiagaraBool Value = RetrieveValueWithCheck<FNiagaraBool, FNiagaraDataConversions::FNiagaraBool>(InstanceData->EmitterInstance, AttributeToRead, ParticleID, bValid);
-
+		FVector4 Value = RetrieveValueWithCheck<FVector4, FNiagaraDataConversions::FHalf4OrFloat4>(InstanceData->EmitterInstance, AttributeToRead, ParticleID, bValid);
 		FNiagaraBool ValidValue;
 		ValidValue.SetValue(bValid);
 		*OutValid.GetDestAndAdvance() = ValidValue;
-		*OutValue.GetDestAndAdvance() = Value;
+		*OutX.GetDestAndAdvance() = Value.X;
+		*OutY.GetDestAndAdvance() = Value.Y;
+		*OutZ.GetDestAndAdvance() = Value.Z;
+		*OutW.GetDestAndAdvance() = Value.W;
 	}
 }
 
@@ -1111,9 +1353,7 @@
 	{
 		FNiagaraID ParticleID = { ParticleIDIndexParam.GetAndAdvance(), ParticleIDAcquireTagParam.GetAndAdvance() };
 		bool bValid;
-
 		FLinearColor Value = RetrieveValueWithCheck<FLinearColor, FNiagaraDataConversions::FHalf4OrColor4>(InstanceData->EmitterInstance, AttributeToRead, ParticleID, bValid);
-
 		FNiagaraBool ValidValue;
 		ValidValue.SetValue(bValid);
 		*OutValid.GetDestAndAdvance() = ValidValue;
@@ -1141,9 +1381,7 @@
 	{
 		FNiagaraID ParticleID = { ParticleIDIndexParam.GetAndAdvance(), ParticleIDAcquireTagParam.GetAndAdvance() };
 		bool bValid;
-
 		FQuat Value = RetrieveValueWithCheck<FQuat, FNiagaraDataConversions::FHalf4OrQuat4>(InstanceData->EmitterInstance, AttributeToRead, ParticleID, bValid);
-
 		FNiagaraBool ValidValue;
 		ValidValue.SetValue(bValid);
 		*OutValid.GetDestAndAdvance() = ValidValue;
@@ -1191,6 +1429,221 @@
 	return Value;
 }
 
+void UNiagaraDataInterfaceParticleRead::ReadIntByIndex(FVectorVMContext& Context, FName AttributeToRead)
+{
+	VectorVM::FExternalFuncInputHandler<int32> ParticleIndexParam(Context);
+
+	VectorVM::FUserPtrHandler<FNDIParticleRead_InstanceData> InstanceData(Context);
+
+	VectorVM::FExternalFuncRegisterHandler<FNiagaraBool> OutValid(Context);
+	VectorVM::FExternalFuncRegisterHandler<int32> OutValue(Context);
+
+	for (int32 InstanceIdx = 0; InstanceIdx < Context.NumInstances; ++InstanceIdx)
+	{
+		int32 ParticleIndex = ParticleIndexParam.GetAndAdvance();
+		bool bValid;
+		int32 Value = RetrieveValueByIndexWithCheck<int32, FNiagaraDataConversions::FNiagaraInt>(InstanceData->EmitterInstance, AttributeToRead, ParticleIndex, bValid);
+		FNiagaraBool ValidValue;
+		ValidValue.SetValue(bValid);
+		*OutValid.GetDestAndAdvance() = ValidValue;
+		*OutValue.GetDestAndAdvance() = Value;
+	}
+}
+
+void UNiagaraDataInterfaceParticleRead::ReadBoolByIndex(FVectorVMContext& Context, FName AttributeToRead)
+{
+	VectorVM::FExternalFuncInputHandler<int32> ParticleIndexParam(Context);
+
+	VectorVM::FUserPtrHandler<FNDIParticleRead_InstanceData> InstanceData(Context);
+
+	VectorVM::FExternalFuncRegisterHandler<FNiagaraBool> OutValid(Context);
+	VectorVM::FExternalFuncRegisterHandler<FNiagaraBool> OutValue(Context);
+
+	for (int32 InstanceIdx = 0; InstanceIdx < Context.NumInstances; ++InstanceIdx)
+	{
+		int32 ParticleIndex = ParticleIndexParam.GetAndAdvance();
+		bool bValid;
+		FNiagaraBool Value = RetrieveValueByIndexWithCheck<FNiagaraBool, FNiagaraDataConversions::FNiagaraBool>(InstanceData->EmitterInstance, AttributeToRead, ParticleIndex, bValid);
+		FNiagaraBool ValidValue;
+		ValidValue.SetValue(bValid);
+		*OutValid.GetDestAndAdvance() = ValidValue;
+		*OutValue.GetDestAndAdvance() = Value;
+	}
+}
+
+void UNiagaraDataInterfaceParticleRead::ReadFloatByIndex(FVectorVMContext& Context, FName AttributeToRead)
+{
+	VectorVM::FExternalFuncInputHandler<int32> ParticleIndexParam(Context);
+
+	VectorVM::FUserPtrHandler<FNDIParticleRead_InstanceData> InstanceData(Context);
+
+	VectorVM::FExternalFuncRegisterHandler<FNiagaraBool> OutValid(Context);
+	VectorVM::FExternalFuncRegisterHandler<float> OutValue(Context);
+
+	for (int32 InstanceIdx = 0; InstanceIdx < Context.NumInstances; ++InstanceIdx)
+	{
+		int32 ParticleIndex = ParticleIndexParam.GetAndAdvance();
+		bool bValid;
+		float Value = RetrieveValueByIndexWithCheck<float, FNiagaraDataConversions::FHalfOrFloat>(InstanceData->EmitterInstance, AttributeToRead, ParticleIndex, bValid);
+		FNiagaraBool ValidValue;
+		ValidValue.SetValue(bValid);
+		*OutValid.GetDestAndAdvance() = ValidValue;
+		*OutValue.GetDestAndAdvance() = Value;
+	}
+}
+
+void UNiagaraDataInterfaceParticleRead::ReadVector2ByIndex(FVectorVMContext& Context, FName AttributeToRead)
+{
+	VectorVM::FExternalFuncInputHandler<int32> ParticleIndexParam(Context);
+
+	VectorVM::FUserPtrHandler<FNDIParticleRead_InstanceData> InstanceData(Context);
+
+	VectorVM::FExternalFuncRegisterHandler<FNiagaraBool> OutValid(Context);
+	VectorVM::FExternalFuncRegisterHandler<float> OutX(Context);
+	VectorVM::FExternalFuncRegisterHandler<float> OutY(Context);
+
+	for (int32 InstanceIdx = 0; InstanceIdx < Context.NumInstances; ++InstanceIdx)
+	{
+		int32 ParticleIndex = ParticleIndexParam.GetAndAdvance();
+		bool bValid;
+		FVector2D Value = RetrieveValueByIndexWithCheck<FVector2D, FNiagaraDataConversions::FHalf2OrFloat2>(InstanceData->EmitterInstance, AttributeToRead, ParticleIndex, bValid);
+		FNiagaraBool ValidValue;
+		ValidValue.SetValue(bValid);
+		*OutValid.GetDestAndAdvance() = ValidValue;
+		*OutX.GetDestAndAdvance() = Value.X;
+		*OutY.GetDestAndAdvance() = Value.Y;
+	}
+}
+
+void UNiagaraDataInterfaceParticleRead::ReadVector3ByIndex(FVectorVMContext& Context, FName AttributeToRead)
+{
+	VectorVM::FExternalFuncInputHandler<int32> ParticleIndexParam(Context);
+
+	VectorVM::FUserPtrHandler<FNDIParticleRead_InstanceData> InstanceData(Context);
+
+	VectorVM::FExternalFuncRegisterHandler<FNiagaraBool> OutValid(Context);
+	VectorVM::FExternalFuncRegisterHandler<float> OutX(Context);
+	VectorVM::FExternalFuncRegisterHandler<float> OutY(Context);
+	VectorVM::FExternalFuncRegisterHandler<float> OutZ(Context);
+
+	for (int32 InstanceIdx = 0; InstanceIdx < Context.NumInstances; ++InstanceIdx)
+	{
+		int32 ParticleIndex = ParticleIndexParam.GetAndAdvance();
+		bool bValid;
+		FVector Value = RetrieveValueByIndexWithCheck<FVector, FNiagaraDataConversions::FHalf3OrFloat3>(InstanceData->EmitterInstance, AttributeToRead, ParticleIndex, bValid);
+		FNiagaraBool ValidValue;
+		ValidValue.SetValue(bValid);
+		*OutValid.GetDestAndAdvance() = ValidValue;
+		*OutX.GetDestAndAdvance() = Value.X;
+		*OutY.GetDestAndAdvance() = Value.Y;
+		*OutZ.GetDestAndAdvance() = Value.Z;
+	}
+}
+
+void UNiagaraDataInterfaceParticleRead::ReadVector4ByIndex(FVectorVMContext& Context, FName AttributeToRead)
+{
+	VectorVM::FExternalFuncInputHandler<int32> ParticleIndexParam(Context);
+
+	VectorVM::FUserPtrHandler<FNDIParticleRead_InstanceData> InstanceData(Context);
+
+	VectorVM::FExternalFuncRegisterHandler<FNiagaraBool> OutValid(Context);
+	VectorVM::FExternalFuncRegisterHandler<float> OutX(Context);
+	VectorVM::FExternalFuncRegisterHandler<float> OutY(Context);
+	VectorVM::FExternalFuncRegisterHandler<float> OutZ(Context);
+	VectorVM::FExternalFuncRegisterHandler<float> OutW(Context);
+
+	for (int32 InstanceIdx = 0; InstanceIdx < Context.NumInstances; ++InstanceIdx)
+	{
+		int32 ParticleIndex = ParticleIndexParam.GetAndAdvance();
+		bool bValid;
+		FVector4 Value = RetrieveValueByIndexWithCheck<FVector4, FNiagaraDataConversions::FHalf4OrFloat4>(InstanceData->EmitterInstance, AttributeToRead, ParticleIndex, bValid);
+		FNiagaraBool ValidValue;
+		ValidValue.SetValue(bValid);
+		*OutValid.GetDestAndAdvance() = ValidValue;
+		*OutX.GetDestAndAdvance() = Value.X;
+		*OutY.GetDestAndAdvance() = Value.Y;
+		*OutZ.GetDestAndAdvance() = Value.Z;
+		*OutW.GetDestAndAdvance() = Value.W;
+	}
+}
+
+void UNiagaraDataInterfaceParticleRead::ReadColorByIndex(FVectorVMContext& Context, FName AttributeToRead)
+{
+	VectorVM::FExternalFuncInputHandler<int32> ParticleIndexParam(Context);
+
+	VectorVM::FUserPtrHandler<FNDIParticleRead_InstanceData> InstanceData(Context);
+
+	VectorVM::FExternalFuncRegisterHandler<FNiagaraBool> OutValid(Context);
+	VectorVM::FExternalFuncRegisterHandler<float> OutR(Context);
+	VectorVM::FExternalFuncRegisterHandler<float> OutG(Context);
+	VectorVM::FExternalFuncRegisterHandler<float> OutB(Context);
+	VectorVM::FExternalFuncRegisterHandler<float> OutA(Context);
+
+	for (int32 InstanceIdx = 0; InstanceIdx < Context.NumInstances; ++InstanceIdx)
+	{
+		int32 ParticleIndex = ParticleIndexParam.GetAndAdvance();
+		bool bValid;
+		FLinearColor Value = RetrieveValueByIndexWithCheck<FLinearColor, FNiagaraDataConversions::FHalf4OrColor4>(InstanceData->EmitterInstance, AttributeToRead, ParticleIndex, bValid);
+		FNiagaraBool ValidValue;
+		ValidValue.SetValue(bValid);
+		*OutValid.GetDestAndAdvance() = ValidValue;
+		*OutR.GetDestAndAdvance() = Value.R;
+		*OutG.GetDestAndAdvance() = Value.G;
+		*OutB.GetDestAndAdvance() = Value.B;
+		*OutA.GetDestAndAdvance() = Value.A;
+	}
+}
+
+void UNiagaraDataInterfaceParticleRead::ReadQuatByIndex(FVectorVMContext& Context, FName AttributeToRead)
+{
+	VectorVM::FExternalFuncInputHandler<int32> ParticleIndexParam(Context);
+
+	VectorVM::FUserPtrHandler<FNDIParticleRead_InstanceData> InstanceData(Context);
+
+	VectorVM::FExternalFuncRegisterHandler<FNiagaraBool> OutValid(Context);
+	VectorVM::FExternalFuncRegisterHandler<float> OutX(Context);
+	VectorVM::FExternalFuncRegisterHandler<float> OutY(Context);
+	VectorVM::FExternalFuncRegisterHandler<float> OutZ(Context);
+	VectorVM::FExternalFuncRegisterHandler<float> OutW(Context);
+
+	for (int32 InstanceIdx = 0; InstanceIdx < Context.NumInstances; ++InstanceIdx)
+	{
+		int32 ParticleIndex = ParticleIndexParam.GetAndAdvance();
+		bool bValid;
+		FQuat Value = RetrieveValueByIndexWithCheck<FQuat, FNiagaraDataConversions::FHalf4OrQuat4>(InstanceData->EmitterInstance, AttributeToRead, ParticleIndex, bValid);
+		FNiagaraBool ValidValue;
+		ValidValue.SetValue(bValid);
+		*OutValid.GetDestAndAdvance() = ValidValue;
+		*OutX.GetDestAndAdvance() = Value.X;
+		*OutY.GetDestAndAdvance() = Value.Y;
+		*OutZ.GetDestAndAdvance() = Value.Z;
+		*OutW.GetDestAndAdvance() = Value.W;
+	}
+}
+
+template <typename Type, typename AccessorType>
+Type UNiagaraDataInterfaceParticleRead::RetrieveValueByIndexWithCheck(FNiagaraEmitterInstance* EmitterInstance, const FName& Attr, int32 ParticleIndex, bool& bValid)
+{
+	Type Value = Type();
+	bValid = false;
+
+	const FNiagaraDataBuffer* CurrentData = EmitterInstance->GetData().GetCurrentData();
+	if (!CurrentData)
+	{
+		return Value;
+	}
+
+	FNiagaraDataSetAccessor<AccessorType> ValueData(EmitterInstance->GetData(), Attr);
+
+	if (ParticleIndex >= 0 && ParticleIndex < (int32)CurrentData->GetNumInstances())
+	{
+		Value = ValueData.GetSafe(ParticleIndex, Type());
+		bValid = true;
+	}
+
+	return Value;
+}
+
 bool UNiagaraDataInterfaceParticleRead::Equals(const UNiagaraDataInterface* Other) const
 {
 	if (!Super::Equals(Other))
@@ -1255,15 +1708,40 @@
 	OutHLSL += FString::Format(FormatDeclarations, ArgsDeclarations);
 }
 
-static bool GenerateGetFunctionHLSL(const FNiagaraDataInterfaceGPUParamInfo& ParamInfo, const FNiagaraDataInterfaceGeneratedFunction& FunctionInfo, int FunctionInstanceIndex, ENiagaraParticleDataComponentType ComponentType, int NumComponents, FString& OutHLSL)
+static bool GenerateGetFunctionHLSL(const FNiagaraDataInterfaceGPUParamInfo& ParamInfo, const FNiagaraDataInterfaceGeneratedFunction& FunctionInfo, int FunctionInstanceIndex, ENiagaraParticleDataComponentType ComponentType, int NumComponents, bool bByIndex, FString& OutHLSL)
 {
-	static const TCHAR* FuncTemplate = TEXT(
-		"void {FunctionName}(NiagaraID In_ParticleID, out bool Out_Valid, out {ValueType} Out_Value)\n"
+	FString FuncTemplate;
+
+	if (bByIndex)
+	{
+		FuncTemplate = TEXT("void {FunctionName}(int ParticleIndex, out bool Out_Valid, out {ValueType} Out_Value)\n");
+	}
+	else
+	{
+		FuncTemplate = TEXT("void {FunctionName}(NiagaraID In_ParticleID, out bool Out_Valid, out {ValueType} Out_Value)\n");
+	}
+
+	FuncTemplate += TEXT(
 		"{\n"
 		"    int RegisterIndex = {AttributeIndicesName}[{AttributeIndexGroup}]{AttributeIndexComponent};\n"
-		"    int ParticleIndex = (RegisterIndex != -1) && (In_ParticleID.Index >= 0) ? {IDToIndexTableName}[In_ParticleID.Index] : -1;\n"
-		"    int AcquireTag = (ParticleIndex != -1) ? {InputIntBufferName}[{AcquireTagRegisterIndexName}*{ParticleStrideIntName} + ParticleIndex] : 0;\n"
-		"    if(ParticleIndex != -1 && In_ParticleID.AcquireTag == AcquireTag)\n"
+	);
+
+	if (bByIndex)
+	{
+		FuncTemplate += TEXT(
+			"    if(RegisterIndex != -1)\n"
+		);
+	}
+	else
+	{
+		FuncTemplate += TEXT(
+			"    int ParticleIndex = (RegisterIndex != -1) && (In_ParticleID.Index >= 0) ? {IDToIndexTableName}[In_ParticleID.Index] : -1;\n"
+			"    int AcquireTag = (ParticleIndex != -1) ? {InputIntBufferName}[{AcquireTagRegisterIndexName}*{ParticleStrideIntName} + ParticleIndex] : 0;\n"
+			"    if(ParticleIndex != -1 && In_ParticleID.AcquireTag == AcquireTag)\n"
+		);
+	}
+
+	FuncTemplate += TEXT(
 		"    {\n"
 		"        Out_Valid = true;\n"
 		"\n"
@@ -1374,13 +1852,16 @@
 	FuncTemplateArgs.Add(TEXT("FetchCompressedValueCode"), FetchCompressedValueCode);
 	FuncTemplateArgs.Add(TEXT("ExtraDefaultValues"), ExtraDefaultValues);
 	
-	OutHLSL += FString::Format(FuncTemplate, FuncTemplateArgs);
+	OutHLSL += FString::Format(*FuncTemplate, FuncTemplateArgs);
 
 	return true;
 }
 
 bool UNiagaraDataInterfaceParticleRead::GetFunctionHLSL(const FNiagaraDataInterfaceGPUParamInfo& ParamInfo, const FNiagaraDataInterfaceGeneratedFunction& FunctionInfo, int FunctionInstanceIndex, FString& OutHLSL)
 {
+	//
+	// Spawn info and particle count
+	//
 	if (FunctionInfo.DefinitionName == GetNumSpawnedParticlesFunctionName)
 	{
 		static const TCHAR* FuncTemplate = TEXT(
@@ -1452,44 +1933,90 @@
 		return true;
 	}
 
+	//
+	// Get attribute by ID
+	//
 	if (FunctionInfo.DefinitionName == GetIntAttributeFunctionName)
 	{
-		return GenerateGetFunctionHLSL(ParamInfo, FunctionInfo, FunctionInstanceIndex, ENiagaraParticleDataComponentType::Int, 1, OutHLSL);
+		return GenerateGetFunctionHLSL(ParamInfo, FunctionInfo, FunctionInstanceIndex, ENiagaraParticleDataComponentType::Int, 1, false, OutHLSL);
 	}
 
+	if (FunctionInfo.DefinitionName == GetBoolAttributeFunctionName)
+	{
+		return GenerateGetFunctionHLSL(ParamInfo, FunctionInfo, FunctionInstanceIndex, ENiagaraParticleDataComponentType::Bool, 1, false, OutHLSL);
+	}
+
 	if (FunctionInfo.DefinitionName == GetFloatAttributeFunctionName)
 	{
-		return GenerateGetFunctionHLSL(ParamInfo, FunctionInfo, FunctionInstanceIndex, ENiagaraParticleDataComponentType::Float, 1, OutHLSL);
+		return GenerateGetFunctionHLSL(ParamInfo, FunctionInfo, FunctionInstanceIndex, ENiagaraParticleDataComponentType::Float, 1, false, OutHLSL);
 	}
 
 	if (FunctionInfo.DefinitionName == GetVec2AttributeFunctionName)
 	{
-		return GenerateGetFunctionHLSL(ParamInfo, FunctionInfo, FunctionInstanceIndex, ENiagaraParticleDataComponentType::Float, 2, OutHLSL);
+		return GenerateGetFunctionHLSL(ParamInfo, FunctionInfo, FunctionInstanceIndex, ENiagaraParticleDataComponentType::Float, 2, false, OutHLSL);
 	}
 
 	if (FunctionInfo.DefinitionName == GetVec3AttributeFunctionName)
 	{
-		return GenerateGetFunctionHLSL(ParamInfo, FunctionInfo, FunctionInstanceIndex, ENiagaraParticleDataComponentType::Float, 3, OutHLSL);
+		return GenerateGetFunctionHLSL(ParamInfo, FunctionInfo, FunctionInstanceIndex, ENiagaraParticleDataComponentType::Float, 3, false, OutHLSL);
 	}
 
 	if (FunctionInfo.DefinitionName == GetVec4AttributeFunctionName)
 	{
-		return GenerateGetFunctionHLSL(ParamInfo, FunctionInfo, FunctionInstanceIndex, ENiagaraParticleDataComponentType::Float, 4, OutHLSL);
+		return GenerateGetFunctionHLSL(ParamInfo, FunctionInfo, FunctionInstanceIndex, ENiagaraParticleDataComponentType::Float, 4, false, OutHLSL);
 	}
 
-	if (FunctionInfo.DefinitionName == GetBoolAttributeFunctionName)
-	{
-		return GenerateGetFunctionHLSL(ParamInfo, FunctionInfo, FunctionInstanceIndex, ENiagaraParticleDataComponentType::Bool, 1, OutHLSL);
-	}
-
 	if (FunctionInfo.DefinitionName == GetColorAttributeFunctionName)
 	{
-		return GenerateGetFunctionHLSL(ParamInfo, FunctionInfo, FunctionInstanceIndex, ENiagaraParticleDataComponentType::Float, 4, OutHLSL);
+		return GenerateGetFunctionHLSL(ParamInfo, FunctionInfo, FunctionInstanceIndex, ENiagaraParticleDataComponentType::Float, 4, false, OutHLSL);
 	}
 
 	if (FunctionInfo.DefinitionName == GetQuatAttributeFunctionName)
 	{
-		return GenerateGetFunctionHLSL(ParamInfo, FunctionInfo, FunctionInstanceIndex, ENiagaraParticleDataComponentType::Float, 4, OutHLSL);
+		return GenerateGetFunctionHLSL(ParamInfo, FunctionInfo, FunctionInstanceIndex, ENiagaraParticleDataComponentType::Float, 4, false, OutHLSL);
 	}
 
+	//
+	// Get attribute by index
+	//
+	if (FunctionInfo.DefinitionName == GetIntAttributeByIndexFunctionName)
+	{
+		return GenerateGetFunctionHLSL(ParamInfo, FunctionInfo, FunctionInstanceIndex, ENiagaraParticleDataComponentType::Int, 1, true, OutHLSL);
+	}
+
+	if (FunctionInfo.DefinitionName == GetBoolAttributeByIndexFunctionName)
+	{
+		return GenerateGetFunctionHLSL(ParamInfo, FunctionInfo, FunctionInstanceIndex, ENiagaraParticleDataComponentType::Bool, 1, true, OutHLSL);
+	}
+
+	if (FunctionInfo.DefinitionName == GetFloatAttributeByIndexFunctionName)
+	{
+		return GenerateGetFunctionHLSL(ParamInfo, FunctionInfo, FunctionInstanceIndex, ENiagaraParticleDataComponentType::Float, 1, true, OutHLSL);
+	}
+
+	if (FunctionInfo.DefinitionName == GetVec2AttributeByIndexFunctionName)
+	{
+		return GenerateGetFunctionHLSL(ParamInfo, FunctionInfo, FunctionInstanceIndex, ENiagaraParticleDataComponentType::Float, 2, true, OutHLSL);
+	}
+
+	if (FunctionInfo.DefinitionName == GetVec3AttributeByIndexFunctionName)
+	{
+		return GenerateGetFunctionHLSL(ParamInfo, FunctionInfo, FunctionInstanceIndex, ENiagaraParticleDataComponentType::Float, 3, true, OutHLSL);
+	}
+
+	if (FunctionInfo.DefinitionName == GetVec4AttributeByIndexFunctionName)
+	{
+		return GenerateGetFunctionHLSL(ParamInfo, FunctionInfo, FunctionInstanceIndex, ENiagaraParticleDataComponentType::Float, 4, true, OutHLSL);
+	}
+
+	if (FunctionInfo.DefinitionName == GetColorAttributeByIndexFunctionName)
+	{
+		return GenerateGetFunctionHLSL(ParamInfo, FunctionInfo, FunctionInstanceIndex, ENiagaraParticleDataComponentType::Float, 4, true, OutHLSL);
+	}
+
+	if (FunctionInfo.DefinitionName == GetQuatAttributeByIndexFunctionName)
+	{
+		return GenerateGetFunctionHLSL(ParamInfo, FunctionInfo, FunctionInstanceIndex, ENiagaraParticleDataComponentType::Float, 4, true, OutHLSL);
+	}
+
 	return false;
diff --git a/Engine/Plugins/FX/Niagara/Source/NiagaraCore/Private/NiagaraCustomVersion.cpp b/Engine/Plugins/FX/Niagara/Source/NiagaraCore/Private/NiagaraCustomVersion.cpp
--- a/Engine/Plugins/FX/Niagara/Source/NiagaraCore/Private/NiagaraCustomVersion.cpp
+++ b/Engine/Plugins/FX/Niagara/Source/NiagaraCore/Private/NiagaraCustomVersion.cpp
@@ -8,4 +8,4 @@ const FGuid FNiagaraCustomVersion::GUID(0xFCF57AFA, 0x50764283, 0xB9A9E658, 0xFF
 // Register the custom version with core
 FCustomVersionRegistration GRegisterNiagaraCustomVersion(FNiagaraCustomVersion::GUID, FNiagaraCustomVersion::LatestVersion, TEXT("NiagaraVer"));
 
-const FGuid FNiagaraCustomVersion::LatestScriptCompileVersion(0x76B3E9E4, 0x883C48CE, 0x9508A68C, 0xD4BA993D );
+const FGuid FNiagaraCustomVersion::LatestScriptCompileVersion(0x069BDBD3, 0xAEC9455E, 0x8F6DB20F, 0xCA920B88);
