Niagara - Don't show the "Simulating" overlay on niagara graphs.

diff --git a/Engine/Plugins/FX/Niagara/Source/NiagaraEditor/Private/Widgets/SNiagaraScriptGraph.cpp b/Engine/Plugins/FX/Niagara/Source/NiagaraEditor/Private/Widgets/SNiagaraScriptGraph.cpp
--- a/Engine/Plugins/FX/Niagara/Source/NiagaraEditor/Private/Widgets/SNiagaraScriptGraph.cpp
+++ b/Engine/Plugins/FX/Niagara/Source/NiagaraEditor/Private/Widgets/SNiagaraScriptGraph.cpp
@@ -203,7 +203,8 @@ TSharedRef<SGraphEditor> SNiagaraScriptGraph::ConstructGraphEditor()
 		.Appearance(AppearanceInfo)
 		.TitleBar(TitleBarWidget)
 		.GraphToEdit(ViewModel->GetGraph())
-		.GraphEvents(Events);
+		.GraphEvents(Events)
+		.ShowGraphStateOverlay(false);
 
 	// Set a niagara node factory.
 	CreatedGraphEditor->SetNodeFactory(MakeShareable(new FNiagaraNodeFactory()));
diff --git a/Engine/Plugins/FX/Niagara/Source/NiagaraEditorWidgets/Private/SNiagaraOverviewGraph.cpp b/Engine/Plugins/FX/Niagara/Source/NiagaraEditorWidgets/Private/SNiagaraOverviewGraph.cpp
--- a/Engine/Plugins/FX/Niagara/Source/NiagaraEditorWidgets/Private/SNiagaraOverviewGraph.cpp
+++ b/Engine/Plugins/FX/Niagara/Source/NiagaraEditorWidgets/Private/SNiagaraOverviewGraph.cpp
@@ -111,7 +111,8 @@ void SNiagaraOverviewGraph::Construct(const FArguments& InArgs, TSharedRef<FNiag
 		.Appearance(AppearanceInfo)
 		.TitleBar(TitleBarWidget)
 		.GraphToEdit(ViewModel->GetGraph())
-		.GraphEvents(Events);
+		.GraphEvents(Events)
+		.ShowGraphStateOverlay(false);
 
 	GraphEditor->SetNodeFactory(MakeShared<FNiagaraOverviewGraphNodeFactory>());
 
