Niagara - Provide better feedback when a disconnected numeric pin is encountered in the graph.

diff --git a/Engine/Plugins/FX/Niagara/Source/NiagaraEditor/Private/NiagaraHlslTranslator.cpp b/Engine/Plugins/FX/Niagara/Source/NiagaraEditor/Private/NiagaraHlslTranslator.cpp
--- a/Engine/Plugins/FX/Niagara/Source/NiagaraEditor/Private/NiagaraHlslTranslator.cpp
+++ b/Engine/Plugins/FX/Niagara/Source/NiagaraEditor/Private/NiagaraHlslTranslator.cpp
@@ -262,6 +262,11 @@ bool FHlslNiagaraTranslator::ValidateTypePins(UNiagaraNode* NodeToValidate)
 				Error(LOCTEXT("InvalidPinTypeError", "Node pin has an undefined type."), NodeToValidate, Pin);
 				bPinsAreValid = false;
 			}
+			else if (Type == FNiagaraTypeDefinition::GetGenericNumericDef())
+			{
+				Error(LOCTEXT("NumericPinError", "A numeric pin was not resolved to a known type.  Numeric pins must be connected or must be converted to an explicitly typed pin in order to compile."), NodeToValidate, Pin);
+				bPinsAreValid = false;
+			}
 		}
 
 		if (Pin->bOrphanedPin)
