Niagara - Add the collision channel enum as a parameter type from the collision DI.

diff --git a/Engine/Plugins/FX/Niagara/Source/Niagara/Private/NiagaraDataInterfaceCollisionQuery.cpp b/Engine/Plugins/FX/Niagara/Source/Niagara/Private/NiagaraDataInterfaceCollisionQuery.cpp
--- a/Engine/Plugins/FX/Niagara/Source/Niagara/Private/NiagaraDataInterfaceCollisionQuery.cpp
+++ b/Engine/Plugins/FX/Niagara/Source/Niagara/Private/NiagaraDataInterfaceCollisionQuery.cpp
@@ -47,6 +47,7 @@ void UNiagaraDataInterfaceCollisionQuery::PostInitProperties()
 	if (HasAnyFlags(RF_ClassDefaultObject))
 	{
 		FNiagaraTypeRegistry::Register(FNiagaraTypeDefinition(GetClass()), true, false, false);
+		FNiagaraTypeRegistry::Register(FNiagaraTypeDefinition(TraceChannelEnum), true, false, false);
 	}
 }
 
