Niagara - Fix a crash when enabling and disabling shader stages through the UI.

diff --git a/Engine/Plugins/FX/Niagara/Source/NiagaraEditor/Private/ViewModels/Stack/NiagaraStackRoot.cpp b/Engine/Plugins/FX/Niagara/Source/NiagaraEditor/Private/ViewModels/Stack/NiagaraStackRoot.cpp
--- a/Engine/Plugins/FX/Niagara/Source/NiagaraEditor/Private/ViewModels/Stack/NiagaraStackRoot.cpp
+++ b/Engine/Plugins/FX/Niagara/Source/NiagaraEditor/Private/ViewModels/Stack/NiagaraStackRoot.cpp
@@ -167,6 +167,12 @@ void UNiagaraStackRoot::RefreshChildrenInternal(const TArray<UNiagaraStackEntry*
 		AddSimulationStageGroup->SetOnItemAdded(UNiagaraStackSimulationStagesGroup::FOnItemAdded::CreateUObject(this, &UNiagaraStackRoot::EmitterArraysChanged));
 	}
 
+	if (bIncludeEmitterInformation && AddSimulationStageGroup != nullptr && GetEmitterViewModel()->GetEmitter()->bSimulationStagesEnabled == false)
+	{
+		// If this is no longer needed it needs to be nulled out since it will have been finalized, and reusing finalized entries is not supported.
+		AddSimulationStageGroup = nullptr;
+	}
+
 	if (bIncludeEmitterInformation && RenderGroup == nullptr)
 	{
 		RenderGroup = NewObject<UNiagaraStackRenderItemGroup>(this);
