Rearrange the code to get rid of bogus static analysis warning.

#rnx

diff --git a/Engine/Plugins/FX/Niagara/Source/Niagara/Private/NiagaraDataInterfaceParticleRead.cpp b/Engine/Plugins/FX/Niagara/Source/Niagara/Private/NiagaraDataInterfaceParticleRead.cpp
--- a/Engine/Plugins/FX/Niagara/Source/Niagara/Private/NiagaraDataInterfaceParticleRead.cpp
+++ b/Engine/Plugins/FX/Niagara/Source/Niagara/Private/NiagaraDataInterfaceParticleRead.cpp
@@ -480,12 +480,16 @@ struct FNiagaraDataInterfaceParametersCS_ParticleRead : public FNiagaraDataInter
 		{
 			FRHIUnorderedAccessView* InputBuffers[3];
 			int32 NumTransitions = 0;
-			InputBuffers[NumTransitions++] = SourceData->GetGPUBufferFloat().UAV;
-			InputBuffers[NumTransitions++] = SourceData->GetGPUBufferInt().UAV;
+			InputBuffers[NumTransitions] = SourceData->GetGPUBufferFloat().UAV;
+			++NumTransitions;
+			InputBuffers[NumTransitions] = SourceData->GetGPUBufferInt().UAV;
+			++NumTransitions;
 			if (SourceData->GetGPUIDToIndexTable().UAV)
 			{
-				InputBuffers[NumTransitions++] = SourceData->GetGPUIDToIndexTable().UAV;
+				InputBuffers[NumTransitions] = SourceData->GetGPUIDToIndexTable().UAV;
+				++NumTransitions;
 			}
+			checkSlow(NumTransitions <= UE_ARRAY_COUNT(InputBuffers));
 			RHICmdList.TransitionResources(EResourceTransitionAccess::EReadable, EResourceTransitionPipeline::EComputeToCompute, InputBuffers, NumTransitions);
 
 			if (InstanceCountOffsetParam.IsBound())
