Handle lingering Niagara shader compiles on engine exit inside OnEnginePreExit() callback to avoid asserting if shader compile would require querying the asset registry module after it has shutdown.

diff --git a/Engine/Plugins/FX/Niagara/Source/NiagaraEditor/Private/NiagaraEditorModule.cpp b/Engine/Plugins/FX/Niagara/Source/NiagaraEditor/Private/NiagaraEditorModule.cpp
--- a/Engine/Plugins/FX/Niagara/Source/NiagaraEditor/Private/NiagaraEditorModule.cpp
+++ b/Engine/Plugins/FX/Niagara/Source/NiagaraEditor/Private/NiagaraEditorModule.cpp
@@ -719,6 +719,16 @@ class FNiagaraSystemColorParameterTrackEditor : public FNiagaraSystemParameterTr
 void FNiagaraEditorModule::OnPreExit()
 {
 	UDeviceProfileManager::Get().OnManagerUpdated().Remove(DeviceProfileManagerUpdatedHandle);
+
+	// Ensure that we don't have any lingering compiles laying around that will explode after this module shuts down.
+	for (TObjectIterator<UNiagaraSystem> It; It; ++It)
+	{
+		UNiagaraSystem* Sys = *It;
+		if (Sys)
+		{
+			Sys->WaitForCompilationComplete();
+		}
+	}
 }
 
 void FNiagaraEditorModule::StartupModule()
@@ -747,7 +757,7 @@ void FNiagaraEditorModule::StartupModule()
 	FCoreDelegates::OnPostEngineInit.AddRaw(this, &FNiagaraEditorModule::OnPostEngineInit);
 
 	DeviceProfileManagerUpdatedHandle = UDeviceProfileManager::Get().OnManagerUpdated().AddRaw(this, &FNiagaraEditorModule::OnDeviceProfileManagerUpdated);
-	FCoreDelegates::OnPreExit.AddRaw(this, &FNiagaraEditorModule::OnPreExit);
+	FCoreDelegates::OnEnginePreExit.AddRaw(this, &FNiagaraEditorModule::OnPreExit);
 	
 	// register details customization
 	FPropertyEditorModule& PropertyModule = FModuleManager::LoadModuleChecked<FPropertyEditorModule>("PropertyEditor");
@@ -1003,16 +1013,6 @@ void FNiagaraEditorModule::StartupModule()
 
 void FNiagaraEditorModule::ShutdownModule()
 {
-	// Ensure that we don't have any lingering compiles laying around that will explode after this module shuts down.
-	for (TObjectIterator<UNiagaraSystem> It; It; ++It)
-	{
-		UNiagaraSystem* Sys = *It;
-		if (Sys)
-		{
-			Sys->WaitForCompilationComplete();
-		}
-	}
-
 	MenuExtensibilityManager.Reset();
 	ToolBarExtensibilityManager.Reset();
 	
@@ -1030,7 +1030,7 @@ void FNiagaraEditorModule::ShutdownModule()
 
 	FCoreUObjectDelegates::GetPreGarbageCollectDelegate().RemoveAll(this);
 	FCoreDelegates::OnPostEngineInit.RemoveAll(this);
-	FCoreDelegates::OnPreExit.RemoveAll(this);
+	FCoreDelegates::OnEnginePreExit.RemoveAll(this);
 	
 	if (GEditor)
 	{
