Fixing extraneous logging

diff --git a/Engine/Plugins/FX/Niagara/Source/Niagara/Private/NiagaraGPUInstanceCountManager.cpp b/Engine/Plugins/FX/Niagara/Source/Niagara/Private/NiagaraGPUInstanceCountManager.cpp
--- a/Engine/Plugins/FX/Niagara/Source/Niagara/Private/NiagaraGPUInstanceCountManager.cpp
+++ b/Engine/Plugins/FX/Niagara/Source/Niagara/Private/NiagaraGPUInstanceCountManager.cpp
@@ -75,7 +75,7 @@ uint32 FNiagaraGPUInstanceCountManager::AcquireEntry()
 	{
 		// @TODO : add realloc the buffer and copy the current content to it. Might require reallocating the readback in FNiagaraGPUInstanceCountManager::EnqueueGPUReadback()
 		ensure(UsedInstanceCounts < AllocatedInstanceCounts);
-		UE_LOG(LogNiagara, Error, TEXT("Niagara.MinGPUInstanceCount too small. UsedInstanceCounts: %d < AllocatedInstanceCounts: %d"), UsedInstanceCounts, AllocatedInstanceCounts);
+		//UE_LOG(LogNiagara, Error, TEXT("Niagara.MinGPUInstanceCount too small. UsedInstanceCounts: %d < AllocatedInstanceCounts: %d"), UsedInstanceCounts, AllocatedInstanceCounts);
 		return INDEX_NONE;
 	}
 }
@@ -112,7 +112,7 @@ void FNiagaraGPUInstanceCountManager::ResizeBuffers(FRHICommandListImmediate& RH
 			TResourceArray<uint32> InitData;
 			InitData.AddZeroed(AllocatedInstanceCounts);
 			CountBuffer.Initialize(sizeof(uint32), AllocatedInstanceCounts, EPixelFormat::PF_R32_UINT, BUF_Static | BUF_SourceCopy, TEXT("NiagaraGPUInstanceCounts"), &InitData);
-			UE_LOG(LogNiagara, Log, TEXT("FNiagaraGPUInstanceCountManager::ResizeBuffers Alloc AllocatedInstanceCounts: %d ReservedInstanceCounts: %d"), AllocatedInstanceCounts, ReservedInstanceCounts);
+			//UE_LOG(LogNiagara, Log, TEXT("FNiagaraGPUInstanceCountManager::ResizeBuffers Alloc AllocatedInstanceCounts: %d ReservedInstanceCounts: %d"), AllocatedInstanceCounts, ReservedInstanceCounts);
 		}
 		// If we need to increase the buffer size to RecommendedInstanceCounts because the buffer is too small.
 		else if (RequiredInstanceCounts > AllocatedInstanceCounts)
@@ -140,7 +140,7 @@ void FNiagaraGPUInstanceCountManager::ResizeBuffers(FRHICommandListImmediate& RH
 			// Swap the buffers
 			AllocatedInstanceCounts = RecommendedInstanceCounts;
 			FMemory::Memswap(&NextCountBuffer, &CountBuffer, sizeof(NextCountBuffer));
-			UE_LOG(LogNiagara, Log, TEXT("FNiagaraGPUInstanceCountManager::ResizeBuffers Resize AllocatedInstanceCounts: %d ReservedInstanceCounts: %d"), AllocatedInstanceCounts, ReservedInstanceCounts);
+			//UE_LOG(LogNiagara, Log, TEXT("FNiagaraGPUInstanceCountManager::ResizeBuffers Resize AllocatedInstanceCounts: %d ReservedInstanceCounts: %d"), AllocatedInstanceCounts, ReservedInstanceCounts);
 		}
 		// If we need to shrink the buffer size because use way to much buffer size.
 		else if ((int32)(RecommendedInstanceCounts * BufferSlack) < AllocatedInstanceCounts)
diff --git a/Engine/Plugins/FX/Niagara/Source/NiagaraEditor/Private/NiagaraHlslTranslator.cpp b/Engine/Plugins/FX/Niagara/Source/NiagaraEditor/Private/NiagaraHlslTranslator.cpp
--- a/Engine/Plugins/FX/Niagara/Source/NiagaraEditor/Private/NiagaraHlslTranslator.cpp
+++ b/Engine/Plugins/FX/Niagara/Source/NiagaraEditor/Private/NiagaraHlslTranslator.cpp
@@ -5197,7 +5197,7 @@ void FHlslNiagaraTranslator::FunctionCall(UNiagaraNodeFunctionCall* FunctionNode
 		return;
 	}
 
-	UE_LOG(LogNiagaraEditor, Log, TEXT("Function Call: %s %d"), *FunctionNode->GetFunctionName(), ActiveStageIdx);
+	//UE_LOG(LogNiagaraEditor, Log, TEXT("Function Call: %s %d"), *FunctionNode->GetFunctionName(), ActiveStageIdx);
 	// We need the generated string to generate the proper signature for now.
 	ActiveHistoryForFunctionCalls.EnterFunction(FunctionNode->GetFunctionName(), FunctionNode->FunctionScript, FunctionNode);
 
@@ -5577,7 +5577,7 @@ void FHlslNiagaraTranslator::HandleDataInterfaceCall(FNiagaraScriptDataInterface
 		Error(FText::Format(LOCTEXT("FunctionCallDataInterfaceGPUMissing", "Function call \"{0}\" does not work on GPU sims."), FText::FromName(InMatchingSignature.Name)), CurNode, nullptr);
 	}
 
-	UE_LOG(LogNiagaraEditor, Log, TEXT("HandleDataInterfaceCall %d %s %s %s"), ActiveStageIdx, *InMatchingSignature.Name.ToString(), InMatchingSignature.bWriteFunction ? TEXT("true") : TEXT("False"), *Info.Name.ToString());
+	//UE_LOG(LogNiagaraEditor, Log, TEXT("HandleDataInterfaceCall %d %s %s %s"), ActiveStageIdx, *InMatchingSignature.Name.ToString(), InMatchingSignature.bWriteFunction ? TEXT("true") : TEXT("False"), *Info.Name.ToString());
 
 	if (InMatchingSignature.bWriteFunction && CompilationOutput.ScriptData.SimulationStageMetaData.Num() != 0 && TranslationStages[ActiveStageIdx].SourceSimStage != -1)
 	{
@@ -5585,7 +5585,7 @@ void FHlslNiagaraTranslator::HandleDataInterfaceCall(FNiagaraScriptDataInterface
 		ensure(CompilationOutput.ScriptData.SimulationStageMetaData.Num() > SourceSimStage);
 		CompilationOutput.ScriptData.SimulationStageMetaData[SourceSimStage].OutputDestinations.AddUnique(Info.Name);
 		ActiveStageWriteTargets.Top().AddUnique(Info.Name);
-		UE_LOG(LogNiagaraEditor, Log, TEXT(" Add to %d %d"), ActiveStageIdx, SourceSimStage);
+		//UE_LOG(LogNiagaraEditor, Log, TEXT(" Add to %d %d"), ActiveStageIdx, SourceSimStage);
 	}
 }
 
@@ -5885,7 +5885,7 @@ void FHlslNiagaraTranslator::RegisterFunctionCall(ENiagaraScriptUsage ScriptUsag
 					ensure(CompilationOutput.ScriptData.SimulationStageMetaData.Num() > SourceSimStage);
 					CompilationOutput.ScriptData.SimulationStageMetaData[SourceSimStage].OutputDestinations.AddUnique(Entry);
 					ActiveStageWriteTargets.Top().AddUnique(Entry);
-					UE_LOG(LogNiagaraEditor, Log, TEXT(" Add to %d %d"), ActiveStageIdx, SourceSimStage);
+					//UE_LOG(LogNiagaraEditor, Log, TEXT(" Add to %d %d"), ActiveStageIdx, SourceSimStage);
 
 				}
 			}
