Removed UNiagaraDataInterfaceStaticMesh logs are required not anymore, when sampling static mesh that have missing color stream.

diff --git a/Engine/Plugins/FX/Niagara/Source/Niagara/Private/NiagaraDataInterfaceStaticMesh.cpp b/Engine/Plugins/FX/Niagara/Source/Niagara/Private/NiagaraDataInterfaceStaticMesh.cpp
--- a/Engine/Plugins/FX/Niagara/Source/Niagara/Private/NiagaraDataInterfaceStaticMesh.cpp
+++ b/Engine/Plugins/FX/Niagara/Source/Niagara/Private/NiagaraDataInterfaceStaticMesh.cpp
@@ -980,10 +980,6 @@ void UNiagaraDataInterfaceStaticMesh::GetVMExternalFunction(const FVMExternalFun
 {
 	FNDIStaticMesh_InstanceData* InstData = (FNDIStaticMesh_InstanceData*)InstanceData;
 	check(InstData && InstData->Mesh && InstData->Component.IsValid());
-
-	bool bNeedsVertexPositions = false;
-	bool bNeedsVertexColors = false;
-	bool bNeedsVertMain = true;//Assuming we always need this?
 	
 	if (BindingInfo.Name == StaticMeshHelpers::RandomSectionName)
 	{
@@ -1010,56 +1006,46 @@ void UNiagaraDataInterfaceStaticMesh::GetVMExternalFunction(const FVMExternalFun
 	else if (BindingInfo.Name == StaticMeshHelpers::GetTriPositionName)
 	{
 		check(BindingInfo.GetNumInputs() == 5 && BindingInfo.GetNumOutputs() == 3);
-		bNeedsVertexPositions = true;
 		TNDIExplicitBinder<FNDITransformHandlerNoop, NDI_FUNC_BINDER(UNiagaraDataInterfaceStaticMesh, GetTriCoordPosition)>::Bind(this, BindingInfo, InstanceData, OutFunc);
 	}
 	else if (BindingInfo.Name == StaticMeshHelpers::GetTriPositionWSName)
 	{
 		check(BindingInfo.GetNumInputs() == 5 && BindingInfo.GetNumOutputs() == 3);
-		bNeedsVertexPositions = true;
 		TNDIExplicitBinder<FNDITransformHandler, NDI_FUNC_BINDER(UNiagaraDataInterfaceStaticMesh, GetTriCoordPosition)>::Bind(this, BindingInfo, InstanceData, OutFunc);
 	}
 	else if (BindingInfo.Name == StaticMeshHelpers::GetTriNormalName)
 	{
 		check(BindingInfo.GetNumInputs() == 5 && BindingInfo.GetNumOutputs() == 3);
-		bNeedsVertMain = true;
 		TNDIExplicitBinder<FNDITransformHandlerNoop, NDI_FUNC_BINDER(UNiagaraDataInterfaceStaticMesh, GetTriCoordNormal)>::Bind(this, BindingInfo, InstanceData, OutFunc);
 	}
 	else if (BindingInfo.Name == StaticMeshHelpers::GetTriNormalWSName)
 	{
 		check(BindingInfo.GetNumInputs() == 5 && BindingInfo.GetNumOutputs() == 3);
-		bNeedsVertMain = true;
 		TNDIExplicitBinder<FNDITransformHandler, NDI_FUNC_BINDER(UNiagaraDataInterfaceStaticMesh, GetTriCoordNormal)>::Bind(this, BindingInfo, InstanceData, OutFunc);
 	}
 	else if (BindingInfo.Name == StaticMeshHelpers::GetTriTangentsName)
 	{
 		check(BindingInfo.GetNumInputs() == 5 && BindingInfo.GetNumOutputs() == 9);
-		bNeedsVertMain = true;
 		TTypedMeshAccessorBinder<TNDIExplicitBinder<FNDITransformHandlerNoop, NDI_FUNC_BINDER(UNiagaraDataInterfaceStaticMesh, GetTriCoordTangents)>>::Bind(this, BindingInfo, InstanceData, OutFunc);
 	}
 	else if (BindingInfo.Name == StaticMeshHelpers::GetTriTangentsWSName)
 	{
 		check(BindingInfo.GetNumInputs() == 5 && BindingInfo.GetNumOutputs() == 9);
-		bNeedsVertMain = true;
 		TTypedMeshAccessorBinder<TNDIExplicitBinder<FNDITransformHandler, NDI_FUNC_BINDER(UNiagaraDataInterfaceStaticMesh, GetTriCoordTangents)>>::Bind(this, BindingInfo, InstanceData, OutFunc);
 	}
 	else if (BindingInfo.Name == StaticMeshHelpers::GetTriColorName)
 	{
 		check(BindingInfo.GetNumInputs() == 5 && BindingInfo.GetNumOutputs() == 4);
-		bNeedsVertexColors = true;
 		NDI_FUNC_BINDER(UNiagaraDataInterfaceStaticMesh, GetTriCoordColor)::Bind(this, OutFunc);
 	}
 	else if (BindingInfo.Name == StaticMeshHelpers::GetTriUVName)
 	{
 		check(BindingInfo.GetNumInputs() == 6 && BindingInfo.GetNumOutputs() == 2);
-		bNeedsVertMain = true;
 		TTypedMeshAccessorBinder<NDI_FUNC_BINDER(UNiagaraDataInterfaceStaticMesh, GetTriCoordUV)>::Bind(this, BindingInfo, InstanceData, OutFunc);
 	}
 	else if (BindingInfo.Name == StaticMeshHelpers::GetTriPositionAndVelocityName)
 	{
 		check(BindingInfo.GetNumInputs() == 5 && BindingInfo.GetNumOutputs() == 6);
-		bNeedsVertMain = true;
-		bNeedsVertexPositions = true;  
 		NDI_FUNC_BINDER(UNiagaraDataInterfaceStaticMesh, GetTriCoordPositionAndVelocity)::Bind(this, OutFunc);
 	}
 	else if (BindingInfo.Name == StaticMeshHelpers::GetMeshLocalToWorldName)
@@ -1080,28 +1066,13 @@ void UNiagaraDataInterfaceStaticMesh::GetVMExternalFunction(const FVMExternalFun
 	else if (BindingInfo.Name == StaticMeshHelpers::GetVertexPositionName)
 	{
 		check(BindingInfo.GetNumInputs() == 2 && BindingInfo.GetNumOutputs() == 3);
-		bNeedsVertexPositions = true;
 		TNDIExplicitBinder<FNDITransformHandlerNoop, NDI_FUNC_BINDER(UNiagaraDataInterfaceStaticMesh, GetVertexPosition)>::Bind(this, BindingInfo, InstanceData, OutFunc);
 	}
 	else if (BindingInfo.Name == StaticMeshHelpers::GetVertexPositionWSName)
 	{
 		check(BindingInfo.GetNumInputs() == 2 && BindingInfo.GetNumOutputs() == 3);
-		bNeedsVertexPositions = true;
 		TNDIExplicitBinder<FNDITransformHandler, NDI_FUNC_BINDER(UNiagaraDataInterfaceStaticMesh, GetVertexPosition)>::Bind(this, BindingInfo, InstanceData, OutFunc);
 	}
-
-	if (bNeedsVertexPositions && !InstData->MeshHasPositions())
-	{
-		UE_LOG(LogNiagara, Log, TEXT("Static Mesh data interface is cannot run as it's reading position data on a mesh that does not provide it. - Mesh:%s  "), *InstData->Mesh->GetFullName());
-	}
-	if (bNeedsVertexColors && !InstData->MeshHasColors())
-	{
-		UE_LOG(LogNiagara, Log, TEXT("Static Mesh data interface is cannot run as it's reading color data on a mesh that does not provide it. - Mesh:%s  "), *InstData->Mesh->GetFullName());
-	}
-	if (bNeedsVertMain && !InstData->MeshHasVerts())
-	{
-		UE_LOG(LogNiagara, Log, TEXT("Static Mesh data interface is cannot run as it's reading vertex data on a mesh with no vertex data. - Mesh:%s  "), *InstData->Mesh->GetFullName());
-	}
 }
 
 bool UNiagaraDataInterfaceStaticMesh::CopyToInternal(UNiagaraDataInterface* Destination) const
