Merging my recent chagnes to Release-4.25.

diff --git a/Engine/Plugins/FX/Niagara/Source/Niagara/Classes/NiagaraEmitter.h b/Engine/Plugins/FX/Niagara/Source/Niagara/Classes/NiagaraEmitter.h
--- a/Engine/Plugins/FX/Niagara/Source/Niagara/Classes/NiagaraEmitter.h
+++ b/Engine/Plugins/FX/Niagara/Source/Niagara/Classes/NiagaraEmitter.h
@@ -241,13 +241,15 @@ public:
 	virtual void PostEditChangeProperty(struct FPropertyChangedEvent& PropertyChangedEvent) override;
 	NIAGARA_API FOnPropertiesChanged& OnPropertiesChanged();
 	NIAGARA_API FOnRenderersChanged& OnRenderersChanged();
-	bool IsEnabledOnPlatform(const FString& PlatformName);
 #endif
+	virtual bool NeedsLoadForTargetPlatform(const ITargetPlatform* TargetPlatform) const override;
 	void Serialize(FArchive& Ar)override;
 	virtual void PostInitProperties() override;
 	virtual void PostLoad() override;
 	//End UObject Interface
 
+	bool IsEnabledOnPlatform(const FString& PlatformName)const;
+
 	/** Toggles whether or not the particles within this emitter are relative to the emitter origin or in global space. */ 
 	UPROPERTY(EditAnywhere, Category = "Emitter")
 	bool bLocalSpace;
diff --git a/Engine/Plugins/FX/Niagara/Source/Niagara/Classes/NiagaraPlatformSet.h b/Engine/Plugins/FX/Niagara/Source/Niagara/Classes/NiagaraPlatformSet.h
--- a/Engine/Plugins/FX/Niagara/Source/Niagara/Classes/NiagaraPlatformSet.h
+++ b/Engine/Plugins/FX/Niagara/Source/Niagara/Classes/NiagaraPlatformSet.h
@@ -123,6 +123,9 @@ public:
 	/** Fill OutProfiles with all device proifles that have been overriden at the passed EffectsQuality. */
 	void GetOverridenDeviceProfiles(int32 EffectsQuality, TArray<UDeviceProfile*>& OutEnabledProfiles, TArray<UDeviceProfile*>& OutDisabledProfiles)const;
 
+	/** Returns true if this set is enabled for any profiles on the specified platform. */
+	bool IsEnabledForPlatform(const FString& PlatformName)const;
+
 	/** Returns true if the current platform can modify it's niagara scalability settings at runtime. */
 	static bool CanChangeScalabilityAtRuntime();
 
@@ -131,11 +134,11 @@ public:
 
 	static int32 GetEffectQualityMaskForDeviceProfile(const UDeviceProfile* Profile);
 
+	/** Returns true if the passed platform should prune emitters on cook. */
+	static bool ShouldPruneEmittersOnCook(const FString& PlatformName);
+
 	//Editor only public API
 #if WITH_EDITOR
-	/** Returns true if this set is enabled for any profiles on the specified platform. */
-	bool IsEnabledForPlatform(const FString& PlatformName);
-
 	/**
 	Does the enabled set of profiles in this set conflict with those in another set.
 	In some cases this is not allowed.
@@ -158,9 +161,6 @@ public:
 	/** Returns true if the passed platform can modify it's niagara scalability settings at runtime. */
 	static bool CanChangeScalabilityAtRuntime(const FString& PlatformName);
 
-	/** Returns true if the passed platform should prune emitters on cook. */
-	static bool ShouldPruneEmittersOnCook(const FString& PlatformName);
-
 	/** Inspects the passed sets and generates an array of all conflicts between these sets. Used to keep arrays of platform sets orthogonal. */
 	static bool GatherConflicts(const TArray<const FNiagaraPlatformSet*>& PlatformSets, TArray<FNiagaraPlatformSetConflictInfo>& OutConflicts);
 #endif
diff --git a/Engine/Plugins/FX/Niagara/Source/Niagara/Classes/NiagaraSystem.h b/Engine/Plugins/FX/Niagara/Source/Niagara/Classes/NiagaraSystem.h
--- a/Engine/Plugins/FX/Niagara/Source/Niagara/Classes/NiagaraSystem.h
+++ b/Engine/Plugins/FX/Niagara/Source/Niagara/Classes/NiagaraSystem.h
@@ -193,7 +193,8 @@ public:
 	void Serialize(FArchive& Ar)override;
 	virtual void PostLoad() override; 
 	virtual void BeginDestroy() override;
-	virtual void PreSave(const class ITargetPlatform * TargetPlatform) override;
+	virtual void PreSave(const class ITargetPlatform* TargetPlatform) override;
+	virtual bool NeedsLoadForTargetPlatform(const ITargetPlatform* TargetPlatform) const override;
 #if WITH_EDITOR
 	virtual void PreEditChange(FProperty* PropertyThatWillChange)override;
 	virtual void PostEditChangeProperty(struct FPropertyChangedEvent& PropertyChangedEvent) override; 
@@ -382,6 +383,7 @@ public:
 	uint32 bFixedBounds : 1;
 
 	TStatId GetStatID(bool bGameThread, bool bConcurrent)const;
+	void AddToInstanceCountStat(int32 NumInstances, bool bSolo)const;
 
 private:
 #if WITH_EDITORONLY_DATA
@@ -504,6 +506,9 @@ protected:
 	mutable TStatId StatID_GT_CNC;
 	mutable TStatId StatID_RT;
 	mutable TStatId StatID_RT_CNC;
+
+	mutable TStatId StatID_InstanceCount;
+	mutable TStatId StatID_InstanceCountSolo;
 #endif
 
 	FNiagaraSystemScalabilitySettings CurrentScalabilitySettings;
diff --git a/Engine/Plugins/FX/Niagara/Source/Niagara/Niagara.Build.cs b/Engine/Plugins/FX/Niagara/Source/Niagara/Niagara.Build.cs
--- a/Engine/Plugins/FX/Niagara/Source/Niagara/Niagara.Build.cs
+++ b/Engine/Plugins/FX/Niagara/Source/Niagara/Niagara.Build.cs
@@ -17,9 +17,10 @@ public class Niagara : ModuleRules
                 "TimeManagement",
                 "Renderer",
                 "JsonUtilities",
-				"Json",
-				"AudioPlatformConfiguration",
-				"SignalProcessing"
+		"Json",
+		"AudioPlatformConfiguration",
+		"SignalProcessing",
+		"ApplicationCore"
             }
         );
 
diff --git a/Engine/Plugins/FX/Niagara/Source/Niagara/Private/NiagaraCommon.cpp b/Engine/Plugins/FX/Niagara/Source/Niagara/Private/NiagaraCommon.cpp
--- a/Engine/Plugins/FX/Niagara/Source/Niagara/Private/NiagaraCommon.cpp
+++ b/Engine/Plugins/FX/Niagara/Source/Niagara/Private/NiagaraCommon.cpp
@@ -134,7 +134,11 @@ void FNiagaraSystemUpdateContext::AddAll(bool bReInit)
 	{
 		UNiagaraComponent* Comp = *It;
 		check(Comp);
-		AddInternal(Comp, bReInit);
+		bool bIsActive = Comp->IsActive() || Comp->IsRegisteredWithScalabilityManager();
+		if (bIsActive || !bOnlyActive)
+		{
+			AddInternal(Comp, bReInit);
+		}
 	}
 }
 
@@ -146,7 +150,11 @@ void FNiagaraSystemUpdateContext::Add(const UNiagaraSystem* System, bool bReInit
 		check(Comp);
 		if (Comp->GetAsset() == System)
 		{
-			AddInternal(Comp, bReInit);
+			bool bIsActive = Comp->IsActive() || Comp->IsRegisteredWithScalabilityManager();
+			if (bIsActive || !bOnlyActive)
+			{
+				AddInternal(Comp, bReInit);
+			}
 		}
 	}
 }
@@ -161,8 +169,12 @@ void FNiagaraSystemUpdateContext::Add(const UNiagaraEmitter* Emitter, bool bReIn
 		FNiagaraSystemInstance* SystemInst = Comp->GetSystemInstance();
 		if (SystemInst && SystemInst->UsesEmitter(Emitter))
 		{
-			AddInternal(Comp, bReInit);
-		}		
+			bool bIsActive = Comp->IsActive() || Comp->IsRegisteredWithScalabilityManager();
+			if (bIsActive || !bOnlyActive)
+			{
+				AddInternal(Comp, bReInit);
+			}
+		}
 	}
 }
 
@@ -175,7 +187,11 @@ void FNiagaraSystemUpdateContext::Add(const UNiagaraScript* Script, bool bReInit
 		UNiagaraSystem* System = Comp->GetAsset();
 		if (System && System->UsesScript(Script))
 		{
-			AddInternal(Comp, bReInit);
+			bool bIsActive = Comp->IsActive() || Comp->IsRegisteredWithScalabilityManager();
+			if (bIsActive || !bOnlyActive)
+			{
+				AddInternal(Comp, bReInit);
+			}
 		}
 	}
 }
@@ -205,7 +221,11 @@ void FNiagaraSystemUpdateContext::Add(const UNiagaraParameterCollection* Collect
 		FNiagaraSystemInstance* SystemInst = Comp->GetSystemInstance();
 		if (SystemInst && SystemInst->UsesCollection(Collection))
 		{
-			AddInternal(Comp, bReInit);
+			bool bIsActive = Comp->IsActive() || Comp->IsRegisteredWithScalabilityManager();
+			if (bIsActive || !bOnlyActive)
+			{
+				AddInternal(Comp, bReInit);
+			}
 		}
 	}
 }
diff --git a/Engine/Plugins/FX/Niagara/Source/Niagara/Private/NiagaraComponent.cpp b/Engine/Plugins/FX/Niagara/Source/Niagara/Private/NiagaraComponent.cpp
--- a/Engine/Plugins/FX/Niagara/Source/Niagara/Private/NiagaraComponent.cpp
+++ b/Engine/Plugins/FX/Niagara/Source/Niagara/Private/NiagaraComponent.cpp
@@ -449,6 +449,7 @@ UNiagaraComponent::UNiagaraComponent(const FObjectInitializer& ObjectInitializer
 	, bActivateShouldResetWhenReady(false)
 	, bDidAutoAttach(false)
 	, bAllowScalability(true)
+	, bIsCulledByScalability(false)
 	//, bIsChangingAutoAttachment(false)
 	, ScalabilityManagerHandle(INDEX_NONE)
 	, OwnerLOD(0)
@@ -573,6 +574,10 @@ void UNiagaraComponent::TickComponent(float DeltaSeconds, enum ELevelTick TickTy
 	check(SystemInstance->IsSolo());
 	if (IsActive() && SystemInstance.Get() && !SystemInstance->IsComplete())
 	{
+		Asset->AddToInstanceCountStat(1, true);
+		INC_DWORD_STAT_BY(STAT_TotalNiagaraSystemInstances, 1);
+		INC_DWORD_STAT_BY(STAT_TotalNiagaraSystemInstancesSolo, 1);
+
 		// If the interfaces have changed in a meaningful way, we need to potentially rebind and update the values.
 		if (OverrideParameters.GetInterfacesDirty())
 		{
@@ -821,8 +826,9 @@ void UNiagaraComponent::ActivateInternal(bool bReset /* = false */, bool bIsScal
 		return;
 	}
 	
+	UWorld* World = GetWorld();
 	// If the particle system can't ever render (ie on dedicated server or in a commandlet) than do not activate...
-	if (!FApp::CanEverRender())
+	if (!FApp::CanEverRender() || !World || World->IsNetMode(NM_DedicatedServer))
 	{
 		return;
 	}
@@ -863,16 +869,19 @@ void UNiagaraComponent::ActivateInternal(bool bReset /* = false */, bool bIsScal
 		UnregisterWithScalabilityManager();
 	}
 
-	if (!bIsScalabilityCull && ScalabilityManagerHandle != INDEX_NONE)
+	if (!bIsScalabilityCull && bIsCulledByScalability)
 	{
+		check(ScalabilityManagerHandle != INDEX_NONE);
 		//If this is a non scalability activate call and we're still registered with the manager.
 		//If we reach this point then we must have been previously culled by scalability so bail here.
 		return;
 	}
 
-	if (RegisterWithScalabilityManagerOrPreCull())
+	bIsCulledByScalability = false;
+	if (ShouldPreCull())
 	{
 		//We have decided to pre cull the system.
+		bIsCulledByScalability = true;
 		return;
 	}
 
@@ -931,6 +940,8 @@ void UNiagaraComponent::ActivateInternal(bool bReset /* = false */, bool bIsScal
 		return;
 	}
 
+	RegisterWithScalabilityManager();
+
 	SystemInstance->Activate(ResetMode);
 
 	/** We only need to tick the component if we require solo mode. */
@@ -944,10 +955,14 @@ void UNiagaraComponent::Deactivate()
 
 void UNiagaraComponent::DeactivateInternal(bool bIsScalabilityCull /* = false */)
 {
-	// Unregister with the scalability manager if this is a genuine deactivation from outside.
-	// The scalability manager itself can call this function when culling systems.
-	if (bIsScalabilityCull == false)
+	if (bIsScalabilityCull)
+	{
+		bIsCulledByScalability = true;
+	}
+	else
 	{
+		// Unregister with the scalability manager if this is a genuine deactivation from outside.
+		// The scalability manager itself can call this function when culling systems.
 		UnregisterWithScalabilityManager();
 	}
 
@@ -993,7 +1008,11 @@ void UNiagaraComponent::DeactivateImmediateInternal(bool bIsScalabilityCull)
 
 	//Unregister with the scalability manager if this is a genuine deactivation from outside.
 	//The scalability manager itself can call this function when culling systems.
-	if (bIsScalabilityCull == false)
+	if (bIsScalabilityCull)
+	{
+		bIsCulledByScalability = true;
+	}
+	else
 	{
 		UnregisterWithScalabilityManager();
 	}
@@ -1006,9 +1025,9 @@ void UNiagaraComponent::DeactivateImmediateInternal(bool bIsScalabilityCull)
 	}
 }
 
-bool UNiagaraComponent::RegisterWithScalabilityManagerOrPreCull()
+bool UNiagaraComponent::ShouldPreCull()
 {
-	if (ScalabilityManagerHandle == INDEX_NONE && bAllowScalability)
+	if (bAllowScalability)
 	{
 		if (UNiagaraSystem* System = GetAsset())
 		{
@@ -1021,10 +1040,6 @@ bool UNiagaraComponent::RegisterWithScalabilityManagerOrPreCull()
 						//If we're just set to check on spawn then check for precull here.
 						return WorldMan->ShouldPreCull(GetAsset(), this);
 					}
-					else
-					{
-						WorldMan->RegisterWithScalabilityManager(this);
-					}
 				}
 			}
 		}
@@ -1033,6 +1048,23 @@ bool UNiagaraComponent::RegisterWithScalabilityManagerOrPreCull()
 	return false;
 }
 
+void UNiagaraComponent::RegisterWithScalabilityManager()
+{
+	if (ScalabilityManagerHandle == INDEX_NONE && bAllowScalability)
+	{
+		if (UNiagaraSystem* System = GetAsset())
+		{
+			if (UNiagaraEffectType* EffectType = System->GetEffectType())
+			{
+				if (FNiagaraWorldManager* WorldMan = FNiagaraWorldManager::Get(GetWorld()))
+				{
+					WorldMan->RegisterWithScalabilityManager(this);
+				}
+			}
+		}
+	}
+}
+
 void UNiagaraComponent::UnregisterWithScalabilityManager()
 {
 	if (ScalabilityManagerHandle != INDEX_NONE)
@@ -1042,6 +1074,7 @@ void UNiagaraComponent::UnregisterWithScalabilityManager()
 			WorldMan->UnregisterWithScalabilityManager(this);
 		}
 	}
+	bIsCulledByScalability = false;
 	ScalabilityManagerHandle = INDEX_NONE;//Just to be sure our state is unregistered.
 }
 
@@ -1077,6 +1110,24 @@ void UNiagaraComponent::OnSystemComplete()
 	{
 		CancelAutoAttachment(/*bDetachFromParent=*/ true);
 	}
+
+	if (!bIsCulledByScalability && ScalabilityManagerHandle != INDEX_NONE)
+	{
+		//Can we be sure this isn't going to spam erroneously?
+		if (UNiagaraEffectType* EffectType = GetAsset()->GetEffectType())
+		{
+			//Only trigger warning if we're not being deactivated/completed from the outside and this is a natural completion by the system itself.
+			if (EffectType->CullReaction == ENiagaraCullReaction::DeactivateImmediateResume || EffectType->CullReaction == ENiagaraCullReaction::DeactivateResume)
+			{
+				//If we're completing naturally, i.e. we're a burst/non-looping system then we shouldn't be using a mode reactivates the effect.
+				UE_LOG(LogNiagara, Warning, TEXT("Niagara Effect has completed naturally but has an effect type with the \"Asleep\" cull reaction. If an effect like this is culled before it can complete then it could leak into the scalability manager and be reactivated incorrectly. Please verify this is using the correct EffctType.\nComponent:%s\nSystem:%s")
+					, *GetFullName(), *GetAsset()->GetFullName());
+			}
+		}
+
+		//We've completed naturally so unregister with the scalability manager.
+		UnregisterWithScalabilityManager();
+	}
 }
 
 void UNiagaraComponent::DestroyInstance()
@@ -1162,6 +1213,8 @@ void UNiagaraComponent::OnUnregister()
 
 	SetActiveFlag(false);
 
+	UnregisterWithScalabilityManager();
+
 	if (SystemInstance)
 	{
 		//UE_LOG(LogNiagara, Log, TEXT("UNiagaraComponent::OnUnregister: %p  %s\n"), SystemInstance.Get(), *GetAsset()->GetFullName());
@@ -1248,6 +1301,11 @@ void UNiagaraComponent::SendRenderDynamicData_Concurrent()
 			FNiagaraEmitterInstance* EmitterInst = &SystemInstance->GetEmitters()[i].Get();
 			UNiagaraEmitter* Emitter = EmitterInst->GetCachedEmitter();
 
+			if(Emitter == nullptr)
+			{
+				continue;
+			}
+
 #if STATS
 			TStatId EmitterStatID = Emitter->GetStatID(true, true);
 			FScopeCycleCounter EmitterStatCounter(EmitterStatID);
diff --git a/Engine/Plugins/FX/Niagara/Source/Niagara/Private/NiagaraCrashReporterHandler.cpp b/Engine/Plugins/FX/Niagara/Source/Niagara/Private/NiagaraCrashReporterHandler.cpp
new file mode 100644
--- /dev/null
+++ b/Engine/Plugins/FX/Niagara/Source/Niagara/Private/NiagaraCrashReporterHandler.cpp
@@ -0,0 +1,162 @@
+// Copyright Epic Games, Inc. All Rights Reserved.
+
+#include "NiagaraCrashReporterHandler.h"
+
+#include "CoreMinimal.h"
+#include "HAL\ThreadManager.h"
+#include "GenericPlatform\GenericPlatformCrashContext.h"
+#include "Misc\ScopeLock.h"
+#include "NiagaraSystemInstance.h"
+#include "NiagaraSystemSimulation.h"
+#include "NiagaraComponent.h"
+
+static int32 GbEnableNiagaraCRHandler = 0;
+static FAutoConsoleVariableRef CVarEnableNiagaraCRHandler(
+	TEXT("fx.EnableNiagaraCRHandler"),
+	GbEnableNiagaraCRHandler,
+	TEXT("If > 0 Niagara will push some state into the crash reporter. This is not free so should not be used unless actively tracking a crash in the wild. Even then it should only be enabled on the platforms needed etc. \n"),
+	FConsoleVariableDelegate::CreateStatic(&FNiagaraCrashReporterHandler::OnEnabledChanged),
+	ECVF_Default);
+
+FCriticalSection FNiagaraCrashReporterHandler::CritSec;
+TArray<FNiagaraCrashReporterHandler*> FNiagaraCrashReporterHandler::AllHandlers;
+
+FNiagaraCrashReporterHandler::FNiagaraCrashReporterHandler()
+{
+	FString ThreadName;
+	if (ThreadId == GGameThreadId)
+	{
+		ThreadName = TEXT("GameThread");
+	}
+	else if (ThreadId == GRenderThreadId)
+	{
+		ThreadName = TEXT("RenderThread");
+	}
+	else
+	{
+		ThreadName = FThreadManager::Get().GetThreadName(ThreadId);
+	}
+
+	Name = FString::Printf(TEXT("NiagaraCRInfo_%s_%u"), *ThreadName, ThreadId);
+	Name = Name.Replace(TEXT(" "), TEXT("_"));
+	NoneInfoString = TEXT("None");
+
+	FScopeLock Lock(&CritSec);
+	AllHandlers.AddUnique(this);
+}
+
+FNiagaraCrashReporterHandler::~FNiagaraCrashReporterHandler()
+{
+	FScopeLock Lock(&CritSec);
+	AllHandlers.Remove(this);
+}
+
+void FNiagaraCrashReporterHandler::Clear()
+{
+	FScopeLock ScopeLock(&CritSec);
+
+	ScopeInfoStack.Empty();
+	FString Empty;
+	//UE_LOG(LogNiagara, Log, TEXT("SetEngineData: %s"), *Empty);
+	FGenericCrashContext::SetEngineData(Name, Empty);
+}
+
+void FNiagaraCrashReporterHandler::OnEnabledChanged(IConsoleVariable* InVariable)
+{
+	if (InVariable->GetInt() == 0)
+	{
+		for (FNiagaraCrashReporterHandler* Handler : AllHandlers)
+		{
+			Handler->Clear();
+		}
+	}
+}
+
+void FNiagaraCrashReporterHandler::SetInfo(const FString& Info)
+{
+	//UE_LOG(LogNiagara, Log, TEXT("SetEngineData: %s"), *Info);
+	FGenericCrashContext::SetEngineData(Name, Info);
+}
+
+static const FString NullStr(TEXT("nullptr"));
+void FNiagaraCrashReporterHandler::PushInfo(FNiagaraSystemInstance* Inst)
+{
+	FScopeLock ScopeLock(&CritSec);
+	if (Inst)
+	{
+		PushInfo_Internal(Inst->GetCrashReporterTag());
+	}
+	else
+	{
+		PushInfo_Internal(NullStr);
+	}
+}
+
+void FNiagaraCrashReporterHandler::PushInfo(FNiagaraSystemSimulation* SystemSim)
+{
+	FScopeLock ScopeLock(&CritSec);
+	if (SystemSim)
+	{
+		PushInfo_Internal(SystemSim->GetCrashReporterTag());
+	}
+	else
+	{
+		PushInfo_Internal(NullStr);
+	}
+}
+
+void FNiagaraCrashReporterHandler::PushInfo_Internal(const FString& Info)
+{
+	ScopeInfoStack.Push(Info);
+	SetInfo(Info);
+}
+
+void FNiagaraCrashReporterHandler::PopInfo()
+{
+	FScopeLock ScopeLock(&CritSec);
+	ScopeInfoStack.Pop();
+	if (ScopeInfoStack.Num())
+	{
+		SetInfo(ScopeInfoStack.Top());
+	}
+	else
+	{
+		SetInfo(NoneInfoString);
+	}
+}
+
+//////////////////////////////////////////////////////////////////////////
+
+FNiagaraCrashReporterScope::FNiagaraCrashReporterScope(FNiagaraSystemInstance* Inst)
+{
+	bWasEnabled = GbEnableNiagaraCRHandler != 0;
+	if (GbEnableNiagaraCRHandler)
+	{
+		FNiagaraCrashReporterHandler::Get().PushInfo(Inst);
+	}
+}
+
+FNiagaraCrashReporterScope::FNiagaraCrashReporterScope(FNiagaraSystemSimulation* Sim)
+{
+	bWasEnabled = GbEnableNiagaraCRHandler != 0;
+	if (GbEnableNiagaraCRHandler)
+	{
+		FNiagaraCrashReporterHandler::Get().PushInfo(Sim);
+	}
+}
+
+FNiagaraCrashReporterScope::~FNiagaraCrashReporterScope()
+{
+	if(GbEnableNiagaraCRHandler)
+	{
+		FNiagaraCrashReporterHandler::Get().PopInfo();
+	}
+	else
+	{
+		//Cache this and clear in case GbEnableNiagaraCRHandler changes mid scope.
+		if(bWasEnabled)
+		{
+			FNiagaraCrashReporterHandler::Get().Clear();
+		}
+	}
+}
diff --git a/Engine/Plugins/FX/Niagara/Source/Niagara/Private/NiagaraEmitter.cpp b/Engine/Plugins/FX/Niagara/Source/Niagara/Private/NiagaraEmitter.cpp
--- a/Engine/Plugins/FX/Niagara/Source/Niagara/Private/NiagaraEmitter.cpp
+++ b/Engine/Plugins/FX/Niagara/Source/Niagara/Private/NiagaraEmitter.cpp
@@ -563,6 +563,21 @@ void UNiagaraEmitter::PostLoad()
 	ResolveScalabilitySettings();
 }
 
+bool UNiagaraEmitter::NeedsLoadForTargetPlatform(const ITargetPlatform* TargetPlatform)const
+{
+	if (!FNiagaraPlatformSet::ShouldPruneEmittersOnCook(TargetPlatform->IniPlatformName()))
+	{
+		return true;
+	}
+
+	bool bIsEnabled = IsEnabledOnPlatform(TargetPlatform->IniPlatformName());
+	if(!bIsEnabled)
+	{
+		UE_LOG(LogNiagara, Verbose, TEXT("Pruned emitter %s for platform %s"), *GetFullName(), *TargetPlatform->DisplayName().ToString())
+	}
+	return bIsEnabled;
+}
+
 #if WITH_EDITOR
 /** Creates a new emitter with the supplied emitter as a parent emitter and the supplied system as it's owner. */
 UNiagaraEmitter* UNiagaraEmitter::CreateWithParentAndOwner(UNiagaraEmitter& InParentEmitter, UObject* InOwner, FName InName, EObjectFlags FlagMask)
@@ -718,14 +733,12 @@ UNiagaraEmitter::FOnPropertiesChanged& UNiagaraEmitter::OnRenderersChanged()
 {
 	return OnRenderersChangedDelegate;
 }
+#endif
 
-bool UNiagaraEmitter::IsEnabledOnPlatform(const FString& PlatformName)
+bool UNiagaraEmitter::IsEnabledOnPlatform(const FString& PlatformName)const
 {
-	bool bCanPrune = FNiagaraPlatformSet::ShouldPruneEmittersOnCook(PlatformName);
-	return bCanPrune && Platforms.IsEnabledForPlatform(PlatformName);
+	return Platforms.IsEnabledForPlatform(PlatformName);
 }
-#endif
-
 
 bool UNiagaraEmitter::IsValid()const
 {
diff --git a/Engine/Plugins/FX/Niagara/Source/Niagara/Private/NiagaraEmitterInstance.cpp b/Engine/Plugins/FX/Niagara/Source/Niagara/Private/NiagaraEmitterInstance.cpp
--- a/Engine/Plugins/FX/Niagara/Source/Niagara/Private/NiagaraEmitterInstance.cpp
+++ b/Engine/Plugins/FX/Niagara/Source/Niagara/Private/NiagaraEmitterInstance.cpp
@@ -185,7 +185,7 @@ FBox FNiagaraEmitterInstance::GetBounds()
 
 bool FNiagaraEmitterInstance::IsReadyToRun() const
 {
-	if (!CachedEmitter->IsReadyToRun())
+	if (!IsDisabled() && !CachedEmitter->IsReadyToRun())
 	{
 		return false;
 	}
@@ -444,6 +444,20 @@ void FNiagaraEmitterInstance::Init(int32 InEmitterIdx, FNiagaraSystemInstanceID
 			EventContext.Parameters.UnbindFromSourceStores();
 		}
 	}	
+
+	const int32 NumEventHandlers = CachedEmitter->GetEventHandlers().Num();
+	EventHandlingInfo.Reset();
+	EventHandlingInfo.SetNum(NumEventHandlers);
+	for (int32 i = 0; i < NumEventHandlers; i++)
+	{
+		const FNiagaraEventScriptProperties& EventHandlerProps = CachedEmitter->GetEventHandlers()[i];
+		FNiagaraEventHandlingInfo& Info = EventHandlingInfo[i];
+		Info.SourceEmitterGuid = EventHandlerProps.SourceEmitterID;
+		Info.SourceEmitterName = Info.SourceEmitterGuid.IsValid() ? *Info.SourceEmitterGuid.ToString() : CachedIDName;
+		Info.SpawnCounts.Reset();
+		Info.TotalSpawnCount = 0;
+		Info.EventData = nullptr;
+	}
 }
 
 void FNiagaraEmitterInstance::ResetSimulation(bool bKillExisting /*= true*/)
@@ -577,6 +591,11 @@ void FNiagaraEmitterInstance::CheckForErrors()
 
 void FNiagaraEmitterInstance::DirtyDataInterfaces()
 {
+	if (IsDisabled())
+	{
+		return;
+	}
+
 	// Make sure that our function tables need to be regenerated...
 	SpawnExecContext.DirtyDataInterfaces();
 	UpdateExecContext.DirtyDataInterfaces();
@@ -847,7 +866,11 @@ void FNiagaraEmitterInstance::PostTick()
 
 	checkSlow(CachedEmitter);
 
-	EventHandlingInfo.Reset();
+	//Clear refs to event data buffers.
+	for (FNiagaraEventHandlingInfo& Info : EventHandlingInfo)
+	{
+		Info.SetEventData(nullptr);
+	}
 
 	CachedBounds.Init();
 	if (CachedSystemFixedBounds.IsSet())
@@ -992,27 +1015,22 @@ void FNiagaraEmitterInstance::PreTick()
 	//Gather events we're going to be reading from / handling this frame.
 	//We must do this in pre-tick so we can gather (and mark in use) all sets from other emitters.
 	const int32 NumEventHandlers = CachedEmitter->GetEventHandlers().Num();
-	if (NumEventHandlers > 0)
+	EventSpawnTotal = 0;
+	for (int32 i = 0; i < NumEventHandlers; i++)
 	{
-		EventHandlingInfo.Reset();
-		EventHandlingInfo.SetNum(NumEventHandlers);
-		EventSpawnTotal = 0;
-		for (int32 i = 0; i < NumEventHandlers; i++)
-		{
-			const FNiagaraEventScriptProperties &EventHandlerProps = CachedEmitter->GetEventHandlers()[i];
-			FNiagaraEventHandlingInfo& Info = EventHandlingInfo[i];
-			Info.SourceEmitterGuid = EventHandlerProps.SourceEmitterID;
-			Info.SourceEmitterName = Info.SourceEmitterGuid.IsValid() ? *Info.SourceEmitterGuid.ToString() : CachedIDName;
-			Info.SpawnCounts.Reset();
-			Info.TotalSpawnCount = 0;
-			Info.EventData = nullptr;
-			if (FNiagaraDataSet* EventSet = ParentSystemInstance->GetEventDataSet(Info.SourceEmitterName, EventHandlerProps.SourceEventName))
-			{
-				Info.SetEventData(&EventSet->GetCurrentDataChecked());
-				uint32 EventSpawnNum = CalculateEventSpawnCount(EventHandlerProps, Info.SpawnCounts, EventSet);
-				Info.TotalSpawnCount += EventSpawnNum;
-				EventSpawnTotal += EventSpawnNum;
-			}
+		const FNiagaraEventScriptProperties& EventHandlerProps = CachedEmitter->GetEventHandlers()[i];
+		FNiagaraEventHandlingInfo& Info = EventHandlingInfo[i];
+
+		Info.TotalSpawnCount = 0;//This was being done every frame but should be done in init?
+		Info.SpawnCounts.Reset();
+
+		//TODO: We can move this lookup into the init and just store a ptr to the other set?
+		if (FNiagaraDataSet* EventSet = ParentSystemInstance->GetEventDataSet(Info.SourceEmitterName, EventHandlerProps.SourceEventName))
+		{
+			Info.SetEventData(&EventSet->GetCurrentDataChecked());
+			uint32 EventSpawnNum = CalculateEventSpawnCount(EventHandlerProps, Info.SpawnCounts, EventSet);
+			Info.TotalSpawnCount += EventSpawnNum;
+			EventSpawnTotal += EventSpawnNum;
 		}
 	}
 
@@ -1039,6 +1057,19 @@ void FNiagaraEmitterInstance::SetSystemFixedBoundsOverride(FBox SystemFixedBound
 	CachedSystemFixedBounds = SystemFixedBounds;
 }
 
+static int32 GbTriggerCrash = 0;
+static FAutoConsoleVariableRef CVarTriggerCrash(
+	TEXT("fx.TriggerDebugCrash"),
+	GbTriggerCrash,
+	TEXT("If > 0 we deliberately crash to test Crash Reporter integration."),
+	ECVF_Default
+);
+
+FORCENOINLINE void NiagaraTestCrash()
+{
+	check(0);
+}
+
 void FNiagaraEmitterInstance::Tick(float DeltaSeconds)
 {
 	SCOPE_CYCLE_COUNTER(STAT_NiagaraTick);
@@ -1053,6 +1084,15 @@ void FNiagaraEmitterInstance::Tick(float DeltaSeconds)
 		return;
 	}
 
+	//Test crash allowing us to test CR functionality.
+#if !UE_BUILD_SHIPPING
+	if (GbTriggerCrash)
+	{
+		GbTriggerCrash = 0;
+		NiagaraTestCrash();
+	}
+#endif
+
 	checkSlow(ParticleDataSet);
 	FNiagaraDataSet& Data = *ParticleDataSet;
 	EmitterAge += DeltaSeconds;
@@ -1081,21 +1121,11 @@ void FNiagaraEmitterInstance::Tick(float DeltaSeconds)
 	}
 
 
-	check(Data.GetNumVariables() > 0);
-	check(CachedEmitter->SpawnScriptProps.Script);
-	check(CachedEmitter->UpdateScriptProps.Script);
+	checkSlow(Data.GetNumVariables() > 0);
+	checkSlow(CachedEmitter->SpawnScriptProps.Script);
+	checkSlow(CachedEmitter->UpdateScriptProps.Script);
 	
 	//TickEvents(DeltaSeconds);
-
-	// add system constants
-	{
-		SCOPE_CYCLE_COUNTER(STAT_NiagaraConstants);
-
-		auto& EmitterParameters = ParentSystemInstance->EditEmitterParameters(EmitterIdx);
-		EmitterParameters.EmitterTotalSpawnedParticles = TotalSpawnedParticles;
-		EmitterParameters.EmitterAge = EmitterAge;
-		EmitterParameters.EmitterRandomSeed = CachedEmitter->RandomSeed;
-	}
 	
 	// Calculate number of new particles from regular spawning 
 	uint32 SpawnTotal = 0;
@@ -1112,6 +1142,13 @@ void FNiagaraEmitterInstance::Tick(float DeltaSeconds)
 	int32 OrigNumParticles = GetNumParticles();
 	int32 AllocationEstimate = CachedEmitter->GetMaxParticleCountEstimate();
 	int32 RequiredSize = OrigNumParticles + SpawnTotal + EventSpawnTotal;
+
+	if (RequiredSize == 0)
+	{
+		//Early out if we have no particles to process.
+		//return;
+	}
+
 	int32 AllocationSize = FMath::Max<int32>(AllocationEstimate, RequiredSize);
 	if (AllocationSize > MaxAllocationCount)
 	{
@@ -1131,6 +1168,16 @@ void FNiagaraEmitterInstance::Tick(float DeltaSeconds)
 		MinOverallocation = Overallocation;
 	}
 
+	// add system constants
+	{
+		SCOPE_CYCLE_COUNTER(STAT_NiagaraConstants);
+
+		auto& EmitterParameters = ParentSystemInstance->EditEmitterParameters(EmitterIdx);
+		EmitterParameters.EmitterTotalSpawnedParticles = TotalSpawnedParticles;
+		EmitterParameters.EmitterAge = EmitterAge;
+		EmitterParameters.EmitterRandomSeed = CachedEmitter->RandomSeed;
+	}
+
 	/* GPU simulation -  we just create an FNiagaraComputeExecutionContext, queue it, and let the batcher take care of the rest
 	 */
 	if (CachedEmitter->SimTarget == ENiagaraSimTarget::GPUComputeSim && GPUExecContext != nullptr)
diff --git a/Engine/Plugins/FX/Niagara/Source/Niagara/Private/NiagaraEmitterInstanceBatcher.cpp b/Engine/Plugins/FX/Niagara/Source/Niagara/Private/NiagaraEmitterInstanceBatcher.cpp
--- a/Engine/Plugins/FX/Niagara/Source/Niagara/Private/NiagaraEmitterInstanceBatcher.cpp
+++ b/Engine/Plugins/FX/Niagara/Source/Niagara/Private/NiagaraEmitterInstanceBatcher.cpp
@@ -120,6 +120,12 @@ void NiagaraEmitterInstanceBatcher::GiveSystemTick_RenderThread(FNiagaraGPUSyste
 {
 	check(IsInRenderingThread());
 
+	if (!FNiagaraUtilities::AllowGPUParticles(GetShaderPlatform()))
+	{
+		return;
+	}
+
+
 	// @todo REMOVE THIS HACK
 	if (GFrameNumberRenderThread > LastFrameThatDrainedData + GNiagaraGpuMaxQueuedRenderFrames)
 	{
@@ -599,6 +605,11 @@ void NiagaraEmitterInstanceBatcher::DispatchAllOnCompute(FOverlappableTicks& Ove
 
 void NiagaraEmitterInstanceBatcher::PostRenderOpaque(FRHICommandListImmediate& RHICmdList, FRHIUniformBuffer* ViewUniformBuffer, const class FShaderParametersMetadata* SceneTexturesUniformBufferStruct, FRHIUniformBuffer* SceneTexturesUniformBuffer, bool bAllowGPUParticleUpdate)
 {
+	if (!FNiagaraUtilities::AllowGPUParticles(GetShaderPlatform()))
+	{
+		return;
+	}
+
 	LLM_SCOPE(ELLMTag::Niagara);
 
 	if (bAllowGPUParticleUpdate)
@@ -885,6 +896,11 @@ void NiagaraEmitterInstanceBatcher::ExecuteAll(FRHICommandList& RHICmdList, FRHI
 
 void NiagaraEmitterInstanceBatcher::PreInitViews(FRHICommandListImmediate& RHICmdList, bool bAllowGPUParticleUpdate)
 {
+	if (!FNiagaraUtilities::AllowGPUParticles(GetShaderPlatform()))
+	{
+		return;
+	}
+
 	LLM_SCOPE(ELLMTag::Niagara);
 
 	// Reset the list of GPUSort tasks and release any resources they hold on to.
@@ -972,6 +988,11 @@ void NiagaraEmitterInstanceBatcher::PreInitViews(FRHICommandListImmediate& RHICm
 
 void NiagaraEmitterInstanceBatcher::PostInitViews(FRHICommandListImmediate& RHICmdList, FRHIUniformBuffer* ViewUniformBuffer, bool bAllowGPUParticleUpdate)
 {
+	if (!FNiagaraUtilities::AllowGPUParticles(GetShaderPlatform()))
+	{
+		return;
+	}
+
 	LLM_SCOPE(ELLMTag::Niagara);
 
 	if (bAllowGPUParticleUpdate)
@@ -1021,6 +1042,11 @@ bool NiagaraEmitterInstanceBatcher::RequiresEarlyViewUniformBuffer() const
 
 void NiagaraEmitterInstanceBatcher::PreRender(FRHICommandListImmediate& RHICmdList, const class FGlobalDistanceFieldParameterData* GlobalDistanceFieldParameterData, bool bAllowGPUParticleUpdate)
 {
+	if (!FNiagaraUtilities::AllowGPUParticles(GetShaderPlatform()))
+	{
+		return;
+	}
+
 	LLM_SCOPE(ELLMTag::Niagara);
 
 	GlobalDistanceFieldParams = GlobalDistanceFieldParameterData ? *GlobalDistanceFieldParameterData : FGlobalDistanceFieldParameterData();
diff --git a/Engine/Plugins/FX/Niagara/Source/Niagara/Private/NiagaraModule.cpp b/Engine/Plugins/FX/Niagara/Source/Niagara/Private/NiagaraModule.cpp
--- a/Engine/Plugins/FX/Niagara/Source/Niagara/Private/NiagaraModule.cpp
+++ b/Engine/Plugins/FX/Niagara/Source/Niagara/Private/NiagaraModule.cpp
@@ -40,7 +40,7 @@ static FAutoConsoleVariableRef CVarLogCompileIdGeneration(
 
 float INiagaraModule::EngineGlobalSpawnCountScale = 1.0f;
 float INiagaraModule::EngineGlobalSystemCountScale = 1.0f;
-int32 INiagaraModule::EngineEffectsQuality = 4;
+int32 INiagaraModule::EngineEffectsQuality = INDEX_NONE;
 
 int32 GEnableVerboseNiagaraChangeIdLogging = 0;
 static FAutoConsoleVariableRef CVarEnableVerboseNiagaraChangeIdLogging(
@@ -150,6 +150,11 @@ void INiagaraModule::StartupModule()
 
 	FNiagaraWorldManager::OnStartup();
 
+	//Force update on module load.
+	EngineEffectsQuality = Scalability::DefaultQualityLevel;
+	static const auto CVarEQ = IConsoleManager::Get().FindTConsoleVariableDataInt(TEXT("sg.EffectsQuality"));
+	OnEffectsQualityChanged(CVarEQ->GetValueOnGameThread(), true);
+
 #if WITH_EDITOR	
 	// Loading uncooked data in a game environment, we still need to get some functionality from the NiagaraEditor module.
 	// This includes the ability to compile scripts and load WITH_EDITOR_ONLY data.
@@ -390,19 +395,22 @@ void INiagaraModule::UnregisterPrecompiler(FDelegateHandle DelegateHandle)
 
 #endif
 
-void INiagaraModule::OnEffectsQualityChanged(int32 NewEffectsQuality)
+void INiagaraModule::OnEffectsQualityChanged(int32 NewEffectsQuality, bool bForce)
 {
-	if (NewEffectsQuality != EngineEffectsQuality && FNiagaraPlatformSet::CanChangeScalabilityAtRuntime())
+	if (FNiagaraPlatformSet::CanChangeScalabilityAtRuntime() || bForce)
 	{
-		EngineEffectsQuality = NewEffectsQuality;
+		if (NewEffectsQuality != EngineEffectsQuality)
+		{
+			EngineEffectsQuality = NewEffectsQuality;
 
-		FNiagaraPlatformSet::InvalidateCachedData();
+			FNiagaraPlatformSet::InvalidateCachedData();
 
-		for (TObjectIterator<UNiagaraSystem> It; It; ++It)
-		{
-			UNiagaraSystem* System = *It;
-			check(System);
-			System->OnEffectsQualityChanged();
+			for (TObjectIterator<UNiagaraSystem> It; It; ++It)
+			{
+				UNiagaraSystem* System = *It;
+				check(System);
+				System->OnEffectsQualityChanged();
+			}
 		}
 	}
 }
@@ -947,10 +955,8 @@ void INiagaraModule::ProcessShaderCompilationQueue()
 
 void CVarSinkFunc()
 {
-	//This Cvar sink can happen before the one which primes the cached scalability cvars so we must grab this ourselves.
-	IConsoleManager& ConsoleMan = IConsoleManager::Get();
-	static const auto CVarEQ = ConsoleMan.FindTConsoleVariableDataInt(TEXT("sg.EffectsQuality"));
-	INiagaraModule::OnEffectsQualityChanged(CVarEQ->GetValueOnGameThread());
+	static const auto CVarEQ = IConsoleManager::Get().FindTConsoleVariableDataInt(TEXT("sg.EffectsQuality"));
+	INiagaraModule::OnEffectsQualityChanged(CVarEQ->GetValueOnGameThread(), false);
 }
 
 static FAutoConsoleVariableSink CVarSink(FConsoleCommandDelegate::CreateStatic(&CVarSinkFunc));
diff --git a/Engine/Plugins/FX/Niagara/Source/Niagara/Private/NiagaraPlatformSet.cpp b/Engine/Plugins/FX/Niagara/Source/Niagara/Private/NiagaraPlatformSet.cpp
--- a/Engine/Plugins/FX/Niagara/Source/Niagara/Private/NiagaraPlatformSet.cpp
+++ b/Engine/Plugins/FX/Niagara/Source/Niagara/Private/NiagaraPlatformSet.cpp
@@ -6,6 +6,11 @@
 #include "DeviceProfiles/DeviceProfileManager.h"
 #include "Scalability.h"
 #include "Misc/ConfigCacheIni.h"
+#include "Interfaces/ITargetPlatform.h"
+
+#if WITH_EDITOR
+#include "PlatformInfo.h"
+#endif
 
 #define LOCTEXT_NAMESPACE "NiagaraPlatformSet"
 
@@ -22,6 +27,15 @@ static FAutoConsoleVariableRef CVarNiagaraAllowRuntimeScalabilityChanges(
 	ECVF_ReadOnly
 );
 
+const TCHAR* PruneEmittersOnCookName = TEXT("fx.Niagara.PruneEmittersOnCook");
+int32 GbPruneEmittersOnCook = 1;
+static FAutoConsoleVariableRef CVarPruneEmittersOnCook(
+	PruneEmittersOnCookName,
+	GbPruneEmittersOnCook,
+	TEXT("If > 0 this platform will prune disabled emitters during cook. \n"),
+	ECVF_ReadOnly
+);
+
 uint32 FNiagaraPlatformSet::LastDirtiedFrame = 0;
 #if WITH_EDITOR
 TMap<const UDeviceProfile*, int32> FNiagaraPlatformSet::CachedEQMasksPerDeviceProfile;
@@ -227,9 +241,7 @@ void FNiagaraPlatformSet::InvalidateCachedData()
 	LastDirtiedFrame = GFrameNumber;
 }
 
-#if WITH_EDITOR
-
-bool FNiagaraPlatformSet::IsEnabledForPlatform(const FString& PlatformName)
+bool FNiagaraPlatformSet::IsEnabledForPlatform(const FString& PlatformName)const
 {
 	for (const UObject* ProfileObj : UDeviceProfileManager::Get().Profiles)
 	{
@@ -249,6 +261,18 @@ bool FNiagaraPlatformSet::IsEnabledForPlatform(const FString& PlatformName)
 	return false;
 }
 
+bool FNiagaraPlatformSet::ShouldPruneEmittersOnCook(const FString& PlatformName)
+{
+#if WITH_EDITOR
+	FPlatformIniSettings Settings = GetPlatformIniSettings(PlatformName);
+	return Settings.bPruneEmittersOnCook != 0;
+#else
+	return GbPruneEmittersOnCook != 0;
+#endif
+}
+
+#if WITH_EDITOR
+
 // bool FNiagaraPlatformSet::Conflicts(const FNiagaraPlatformSet& Other, TArray<const UDeviceProfile*>& OutConflictingProfiles, bool bIncludeProfilesWithVariableEffectsQuality)const
 // {
 // 	for (const UObject* ProfileObj : UDeviceProfileManager::Get().Profiles)
@@ -359,12 +383,6 @@ bool FNiagaraPlatformSet::CanChangeScalabilityAtRuntime(const UDeviceProfile* De
 	return true;//Assuming true if we fail to find the platform seems safest.
 }
 
-bool FNiagaraPlatformSet::ShouldPruneEmittersOnCook(const FString& PlatformName)
-{
-	FPlatformIniSettings Settings = GetPlatformIniSettings(PlatformName);
-	return Settings.bPruneEmittersOnCook != 0;
-}
-
 bool FNiagaraPlatformSet::GatherConflicts(const TArray<const FNiagaraPlatformSet*>& PlatformSets, TArray<FNiagaraPlatformSetConflictInfo>& OutConflicts)
 {
 	Scalability::FQualityLevels NumLevels = Scalability::GetQualityLevelCounts();
@@ -428,16 +446,21 @@ FNiagaraPlatformSet::FPlatformIniSettings FNiagaraPlatformSet::GetPlatformIniSet
 
 	auto FindCVarValue = [&](const TCHAR* Section, const TCHAR* CVarName, int32& OutVal)
 	{
+		bool bFound = true;
 		if (!ScalabilitySettings.GetInt(Section, CVarName, OutVal))
 		{
 			if (!GameSettings.GetInt(Section, CVarName, OutVal))
 			{
-				EngineSettings.GetInt(Section, CVarName, OutVal);
+				if (!EngineSettings.GetInt(Section, CVarName, OutVal))
+				{
+					bFound = false;
+				}
 			}
 		}
+		return bFound;
 	};
 
-	int32 EffectsQuality = INDEX_NONE;
+	int32 EffectsQuality = Scalability::DefaultQualityLevel;
 	int32 EQMask = INDEX_NONE;
 
 	//If this platform can change quality settings at runtime then we return the full mask.
@@ -451,12 +474,15 @@ FNiagaraPlatformSet::FPlatformIniSettings FNiagaraPlatformSet::GetPlatformIniSet
 	}
 	else
 	{
-		FindCVarValue(TEXT("ScalabilityGroups"), TEXT("sg.EffectsQuality"), EffectsQuality);
+		if (!FindCVarValue(TEXT("ScalabilityGroups"), TEXT("sg.EffectsQuality"), EffectsQuality))
+		{
+			FindCVarValue(TEXT("SystemSettings"), TEXT("sg.EffectsQuality"), EffectsQuality);
+		}
 		EQMask = CreateEQMask(EffectsQuality);
 	}
 
 	int32 PruneEmittersOnCook = 0;
-	FindCVarValue(TEXT("SystemSettings"), TEXT("fx.PruneEmittersOnCook"), PruneEmittersOnCook);
+	FindCVarValue(TEXT("SystemSettings"), PruneEmittersOnCookName, PruneEmittersOnCook);
 
 	FPlatformIniSettings& NewSetting = CachedPlatformIniSettings.Add(*PlatformName);
 	NewSetting = FPlatformIniSettings(CanChangeEffectQuality, PruneEmittersOnCook, EQMask);
@@ -470,4 +496,4 @@ int32 FNiagaraPlatformSet::GetEffectQualityMaskForPlatform(const FString& Platfo
 }
 #endif
 
-#undef LOCTEXT_NAMESPACE
\ No newline at end of file
+#undef LOCTEXT_NAMESPACE
diff --git a/Engine/Plugins/FX/Niagara/Source/Niagara/Private/NiagaraScalabilityManager.cpp b/Engine/Plugins/FX/Niagara/Source/Niagara/Private/NiagaraScalabilityManager.cpp
--- a/Engine/Plugins/FX/Niagara/Source/Niagara/Private/NiagaraScalabilityManager.cpp
+++ b/Engine/Plugins/FX/Niagara/Source/Niagara/Private/NiagaraScalabilityManager.cpp
@@ -42,17 +42,22 @@ void FNiagaraScalabilityManager::AddReferencedObjects(FReferenceCollector& Colle
 void FNiagaraScalabilityManager::PreGarbageCollectBeginDestroy()
 {
 	//After the GC has potentially nulled out references to the components we were tracking we clear them out here.
-	//This should only be in the case where MarkPendingKill() is called directly. Typical component destruction will unregister in OnComponentDestroyed().
+	//This should only be in the case where MarkPendingKill() is called directly. Typical component destruction will unregister in OnComponentDestroyed() or OnUnregister().
 	//Components then just clear their handle in BeginDestroy knowing they've already been removed from the manager.
 	//I would prefer some pre BeginDestroy() callback into the component in which I could cleanly unregister with the manager in all cases but I don't think that's possible.
 	int32 CompIdx = ManagedComponents.Num();
 	while (--CompIdx >= 0)
 	{
-		if (ManagedComponents[CompIdx] == nullptr)
+		UNiagaraComponent* Comp = ManagedComponents[CompIdx];
+		if (Comp == nullptr)
 		{
 			//UE_LOG(LogNiagara, Warning, TEXT("Unregister from PreGCBeginDestroy @%d/%d - %s"), CompIdx, ManagedComponents.Num(), *EffectType->GetName());
 			UnregisterAt(CompIdx);
 		}
+		else if (Comp->IsPendingKillOrUnreachable())
+		{
+			Unregister(Comp);
+		}
 	}
 }
 
diff --git a/Engine/Plugins/FX/Niagara/Source/Niagara/Private/NiagaraStats.h b/Engine/Plugins/FX/Niagara/Source/Niagara/Private/NiagaraStats.h
--- a/Engine/Plugins/FX/Niagara/Source/Niagara/Private/NiagaraStats.h
+++ b/Engine/Plugins/FX/Niagara/Source/Niagara/Private/NiagaraStats.h
@@ -27,4 +27,9 @@ DECLARE_CYCLE_STAT(TEXT("RT Concurrent Total"), STAT_NiagaraOverview_RT_CNC, STA
 
 
 DECLARE_STATS_GROUP(TEXT("Niagara Systems"), STATGROUP_NiagaraSystems, STATCAT_Advanced);
-DECLARE_STATS_GROUP(TEXT("Niagara Emitters"), STATGROUP_NiagaraEmitters, STATCAT_Advanced);
\ No newline at end of file
+DECLARE_STATS_GROUP(TEXT("Niagara Emitters"), STATGROUP_NiagaraEmitters, STATCAT_Advanced);
+
+DECLARE_STATS_GROUP(TEXT("Niagara System Counts"), STATGROUP_NiagaraSystemCounts, STATCAT_Advanced);
+
+DECLARE_DWORD_COUNTER_STAT(TEXT("Total System Instances"), STAT_TotalNiagaraSystemInstances, STATGROUP_NiagaraSystemCounts);
+DECLARE_DWORD_COUNTER_STAT(TEXT("Total Solo System Instances"), STAT_TotalNiagaraSystemInstancesSolo, STATGROUP_NiagaraSystemCounts);
\ No newline at end of file
diff --git a/Engine/Plugins/FX/Niagara/Source/Niagara/Private/NiagaraSystem.cpp b/Engine/Plugins/FX/Niagara/Source/Niagara/Private/NiagaraSystem.cpp
--- a/Engine/Plugins/FX/Niagara/Source/Niagara/Private/NiagaraSystem.cpp
+++ b/Engine/Plugins/FX/Niagara/Source/Niagara/Private/NiagaraSystem.cpp
@@ -43,14 +43,6 @@ static FAutoConsoleVariableRef CVarNiagaraForceSystemsToCookOutRapidIterationOnL
 	ECVF_Default
 );
 
-static int GNiagaraVerboseCookOutput = 0;
-static FAutoConsoleVariableRef CVarNiagaraVerboseCookOutput(
-	TEXT("fx.NiagaraVerboseCookOutput"),
-	GNiagaraVerboseCookOutput,
-	TEXT("When enabled UNiagaraSystems generate verbose output on cook."),
-	ECVF_Default
-);
-
 //////////////////////////////////////////////////////////////////////////
 
 UNiagaraSystem::UNiagaraSystem(const FObjectInitializer& ObjectInitializer)
@@ -98,6 +90,21 @@ void UNiagaraSystem::PreSave(const class ITargetPlatform * TargetPlatform)
 #endif
 }
 
+bool UNiagaraSystem::NeedsLoadForTargetPlatform(const ITargetPlatform* TargetPlatform)const
+{
+	bool bHasAnyEnabledEmitters = false;
+	for (const FNiagaraEmitterHandle& EmitterHandle : GetEmitterHandles())
+	{
+		if (EmitterHandle.GetIsEnabled() && EmitterHandle.GetInstance()->Platforms.IsEnabledForPlatform(TargetPlatform->IniPlatformName()))
+		{
+			bHasAnyEnabledEmitters = true;
+			break;
+		}
+	}
+
+	return bHasAnyEnabledEmitters;
+}
+
 #if WITH_EDITOR
 void UNiagaraSystem::BeginCacheForCookedPlatformData(const ITargetPlatform *TargetPlatform)
 {
@@ -209,67 +216,8 @@ void UNiagaraSystem::RequestCompileForEmitter(UNiagaraEmitter* InEmitter)
 
 void UNiagaraSystem::Serialize(FArchive& Ar)
 {
-	//TODO: Cook out emitters. This requires some wrangling on the compilation side for system scripts.
-// 	//When cooking we clear out any emitters that are disabled entirely or just for the target platform.
-// 	TArray<FNiagaraEmitterHandle> OldEmitters;
-// 	if (Ar.IsCooking())
-// 	{
-// 		OldEmitters = EmitterHandles;
-// 		const ITargetPlatform* Target = Ar.CookingTarget();
-// 		for (FNiagaraEmitterHandle& Handle : EmitterHandles)
-// 		{
-// 			if (Handle.GetIsEnabled() && Handle.GetInstance()->IsEnabledOnPlatform(Target->PlatformName()))
-// 			{
-// 				Handle.ClearEmitter();
-// 			}
-// 		}
-// 	}
-
-	//Some log spam to display emitters we should be cooking out.
-#if WITH_EDITOR && 0
-	if (Ar.IsCooking() && GNiagaraVerboseCookOutput)
-	{
-		TBitArray<TInlineAllocator<8>> EnabledEmitters;
-		EnabledEmitters.Reserve(EmitterHandles.Num());
-		int32 NumEnabled = 0;
-		const ITargetPlatform* Target = Ar.CookingTarget();
-		for (FNiagaraEmitterHandle& Handle : EmitterHandles)
-		{
-			bool bEnabled = Handle.GetInstance()->IsEnabledOnPlatform(Target->PlatformName());
-			EnabledEmitters.Add(bEnabled);
-			if (bEnabled)
-			{
-				NumEnabled++;
-			}
-		}
-
-		if (NumEnabled != EmitterHandles.Num())
-		{
-			UE_LOG(LogNiagara, Log, TEXT("==============================================="));
-			UE_LOG(LogNiagara, Log, TEXT("System has emitters that could be cooked out!"), *GetName());
-
-			for (int32 EmitterIdx = 0; EmitterIdx < EmitterHandles.Num(); ++EmitterIdx)
-			{
-				if (EnabledEmitters[EmitterIdx] == false)
-				{
-					UE_LOG(LogNiagara, Log, TEXT("%s"), *EmitterHandles[EmitterIdx].GetName().ToString());
-				}
-			}
-			UE_LOG(LogNiagara, Log, TEXT("==============================================="));
-		}
-	}
-
-#endif
-
 	Super::Serialize(Ar);
 
-	//TODO: Cook out emitters. This requires some wrangling on the compilation side for system scripts.
-// 	//Restore old emitters post cook pruning.
-// 	if (Ar.IsCooking())
-// 	{
-// 		EmitterHandles = OldEmitters;
-// 	}
-
 	Ar.UsingCustomVersion(FNiagaraCustomVersion::GUID);
 
 	if (Ar.CustomVer(FNiagaraCustomVersion::GUID) >= FNiagaraCustomVersion::ChangeEmitterCompiledDataToSharedRefs)
@@ -1446,13 +1394,41 @@ TStatId UNiagaraSystem::GetStatID(bool bGameThread, bool bConcurrent)const
 	return TStatId();
 }
 
+void UNiagaraSystem::AddToInstanceCountStat(int32 NumInstances, bool bSolo)const
+{
+#if STATS
+	if (!StatID_GT.IsValidStat())
+	{
+		GenerateStatID();
+	}
+
+	if (FThreadStats::IsCollectingData())
+	{
+		if (bSolo)
+		{
+			FThreadStats::AddMessage(StatID_InstanceCountSolo.GetName(), EStatOperation::Add, int64(NumInstances));
+			TRACE_STAT_ADD(StatID_InstanceCount.GetName(), int64(Value));
+		}
+		else
+		{
+			FThreadStats::AddMessage(StatID_InstanceCount.GetName(), EStatOperation::Add, int64(NumInstances));
+			TRACE_STAT_ADD(StatID_InstanceCount.GetName(), int64(Value));
+		}
+	}
+#endif
+}
+
 void UNiagaraSystem::GenerateStatID()const
 {
 #if STATS
-	StatID_GT = FDynamicStats::CreateStatId<FStatGroup_STATGROUP_NiagaraSystems>(GetPathName() + TEXT("[GT]"));
-	StatID_GT_CNC = FDynamicStats::CreateStatId<FStatGroup_STATGROUP_NiagaraSystems>(GetPathName() + TEXT("[GT_CNC]"));
-	StatID_RT = FDynamicStats::CreateStatId<FStatGroup_STATGROUP_NiagaraSystems>(GetPathName() + TEXT("[RT]"));
-	StatID_RT_CNC = FDynamicStats::CreateStatId<FStatGroup_STATGROUP_NiagaraSystems>(GetPathName() + TEXT("[RT_CNC]"));
+	StatID_GT = FDynamicStats::CreateStatId<FStatGroup_STATGROUP_NiagaraSystems>(GetPathName() + TEXT(" [GT]"));
+	StatID_GT_CNC = FDynamicStats::CreateStatId<FStatGroup_STATGROUP_NiagaraSystems>(GetPathName() + TEXT(" [GT_CNC]"));
+	StatID_RT = FDynamicStats::CreateStatId<FStatGroup_STATGROUP_NiagaraSystems>(GetPathName() + TEXT(" [RT]"));
+	StatID_RT_CNC = FDynamicStats::CreateStatId<FStatGroup_STATGROUP_NiagaraSystems>(GetPathName() + TEXT(" [RT_CNC]"));
+
+	StatID_InstanceCount = FDynamicStats::CreateStatIdInt64<FStatGroup_STATGROUP_NiagaraSystemCounts>(GetPathName());
+	StatID_InstanceCountSolo = FDynamicStats::CreateStatIdInt64<FStatGroup_STATGROUP_NiagaraSystemCounts>(GetPathName() + TEXT(" [SOLO]"));
+
 #endif
 }
 
@@ -1514,10 +1490,15 @@ void UNiagaraSystem::OnEffectsQualityChanged()
 
 	for (FNiagaraEmitterHandle& Handle : EmitterHandles)
 	{
-		Handle.GetInstance()->OnEffectsQualityChanged();
+		if (Handle.GetInstance())
+		{
+			Handle.GetInstance()->OnEffectsQualityChanged();
+		}
 	}
 
 	FNiagaraSystemUpdateContext UpdateCtx;
+	UpdateCtx.SetDestroyOnAdd(true);
+	UpdateCtx.SetOnlyActive(true);
 	UpdateCtx.Add(this, true);
 }
 
diff --git a/Engine/Plugins/FX/Niagara/Source/Niagara/Private/NiagaraSystemInstance.cpp b/Engine/Plugins/FX/Niagara/Source/Niagara/Private/NiagaraSystemInstance.cpp
--- a/Engine/Plugins/FX/Niagara/Source/Niagara/Private/NiagaraSystemInstance.cpp
+++ b/Engine/Plugins/FX/Niagara/Source/Niagara/Private/NiagaraSystemInstance.cpp
@@ -13,6 +13,7 @@
 #include "Templates/AlignmentTemplates.h"
 #include "NiagaraEmitterInstanceBatcher.h"
 #include "GameFramework/PlayerController.h"
+#include "NiagaraCrashReporterHandler.h"
 
 
 DECLARE_CYCLE_STAT(TEXT("System Activate [GT]"), STAT_NiagaraSystemActivate, STATGROUP_Niagara);
@@ -795,47 +796,66 @@ void FNiagaraSystemInstance::Reset(FNiagaraSystemInstance::EResetMode Mode)
 		bBindParams = !IsDisabled();
 	}
 	
-	if (bBindParams)
+	//If none of our emitters actually made it out of the init process we can just bail here before we ever tick.
+	bool bHasActiveEmitters = false;
+	for (auto& Inst : Emitters)
 	{
-		ResetParameters();
-		BindParameters();
+		if (!Inst->IsComplete())
+		{
+			bHasActiveEmitters = true;
+			break;
+		}
 	}
 
 	SetRequestedExecutionState(ENiagaraExecutionState::Active);
-	SetActualExecutionState(ENiagaraExecutionState::Active);
-
-	if (bBindParams)
+	if (bHasActiveEmitters)
 	{
-		InitDataInterfaces();
-	}
+		if (bBindParams)
+		{
+			ResetParameters();
+			BindParameters();
+		}
 
-	//Interface init can disable the system.
-	if (!IsComplete())
-	{
-		ComputeEmittersExecutionOrder();
+		SetActualExecutionState(ENiagaraExecutionState::Active);
 
-		bPendingSpawn = true;
-		SystemSimulation->AddInstance(this);
+		if (bBindParams)
+		{
+			InitDataInterfaces();
+		}
 
-		UNiagaraSystem* System = GetSystem();
-		if (System->NeedsWarmup())
+		//Interface init can disable the system.
+		if (!IsComplete())
 		{
-			int32 WarmupTicks = System->GetWarmupTickCount();
-			float WarmupDt = System->GetWarmupTickDelta();
-			
-			AdvanceSimulation(WarmupTicks, WarmupDt);
+			ComputeEmittersExecutionOrder();
+
+			bPendingSpawn = true;
+			SystemSimulation->AddInstance(this);
+
+			UNiagaraSystem* System = GetSystem();
+			if (System->NeedsWarmup())
+			{
+				int32 WarmupTicks = System->GetWarmupTickCount();
+				float WarmupDt = System->GetWarmupTickDelta();
+
+				AdvanceSimulation(WarmupTicks, WarmupDt);
 
-			//Reset age to zero.
-			Age = 0.0f;
-			TickCount = 0;
+				//Reset age to zero.
+				Age = 0.0f;
+				TickCount = 0;
+			}
 		}
-	}
 
-	if (Component)
+		if (Component)
+		{
+			// This system may not tick again immediately so we mark the render state dirty here so that
+			// the renderers will be reset this frame.
+			Component->MarkRenderDynamicDataDirty();
+		}
+	}
+	else
 	{
-		// This system may not tick again immediately so we mark the render state dirty here so that
-		// the renderers will be reset this frame.
-		Component->MarkRenderDynamicDataDirty();
+		SetActualExecutionState(ENiagaraExecutionState::Complete);
+		Complete();
 	}
 }
 
@@ -1332,6 +1352,11 @@ void FNiagaraSystemInstance::InitDataInterfaces()
 	for (TSharedRef<FNiagaraEmitterInstance, ESPMode::ThreadSafe> Simulation : Emitters)
 	{
 		FNiagaraEmitterInstance& Sim = Simulation.Get();
+		if (Sim.IsDisabled())
+		{
+			continue;
+		}
+
 		CalcInstDataSize(Sim.GetSpawnExecutionContext().GetDataInterfaces());
 		CalcInstDataSize(Sim.GetUpdateExecutionContext().GetDataInterfaces());
 		for (int32 i = 0; i < Sim.GetEventExecutionContexts().Num(); i++)
@@ -1911,6 +1936,8 @@ void FNiagaraSystemInstance::Tick_GameThread(float DeltaSeconds)
 	CSV_SCOPED_TIMING_STAT_EXCLUSIVE(Niagara);
 	LLM_SCOPE(ELLMTag::Niagara);
 
+	FNiagaraCrashReporterScope CRScope(this);
+
 	UNiagaraSystem* System = GetSystem();
 	FScopeCycleCounter SystemStat(System->GetStatID(true, false));
 
@@ -1940,6 +1967,8 @@ void FNiagaraSystemInstance::Tick_Concurrent()
 	LLM_SCOPE(ELLMTag::Niagara);
 	FScopeCycleCounterUObject AdditionalScope(GetSystem(), GET_STATID(STAT_NiagaraOverview_GT_CNC));
 
+	FNiagaraCrashReporterScope CRScope(this);
+
 	// Reset values that will be accumulated during emitter tick.
 	TotalGPUParamSize = 0;
 	ActiveGPUEmitterCount = 0;
@@ -2043,6 +2072,8 @@ void FNiagaraSystemInstance::FinalizeTick_GameThread()
 {
 	if (bNeedsFinalize)//We can come in here twice in one tick if the GT calls WaitForAsync() while there is a GT finalize task in the queue.
 	{
+		FNiagaraCrashReporterScope CRScope(this);
+
 		SCOPE_CYCLE_COUNTER(STAT_NiagaraOverview_GT);
 		SCOPE_CYCLE_COUNTER(STAT_NiagaraSystemInst_FinalizeGT);
 		CSV_SCOPED_TIMING_STAT_EXCLUSIVE(Niagara);
@@ -2209,3 +2240,19 @@ FNiagaraSystemInstance::FOnDestroyed& FNiagaraSystemInstance::OnDestroyed()
 	return OnDestroyedDelegate;
 }
 #endif
+
+const FString& FNiagaraSystemInstance::GetCrashReporterTag()const
+{
+	if(CrashReporterTag.IsEmpty())
+	{
+		UNiagaraSystem* Sys = Component ? Component->GetAsset() : nullptr;
+		USceneComponent* AttachParent = Component ? Component->GetAttachParent() : nullptr;
+
+		const FString& CompName = Component ? Component->GetFullName() : TEXT("nullptr");
+		const FString& AssetName = Sys ? Sys->GetFullName() : TEXT("nullptr");
+		const FString& AttachName = AttachParent ? AttachParent->GetFullName() : TEXT("nullptr");
+
+		CrashReporterTag = FString::Printf(TEXT("SystemInstance | System: %s | bSolo: %s | Component: %s | AttachedTo: %s |"), *AssetName, IsSolo() ? TEXT("true") : TEXT("false"), *CompName, *AttachName);
+	}
+	return CrashReporterTag;
+}
\ No newline at end of file
diff --git a/Engine/Plugins/FX/Niagara/Source/Niagara/Private/NiagaraSystemSimulation.cpp b/Engine/Plugins/FX/Niagara/Source/Niagara/Private/NiagaraSystemSimulation.cpp
--- a/Engine/Plugins/FX/Niagara/Source/Niagara/Private/NiagaraSystemSimulation.cpp
+++ b/Engine/Plugins/FX/Niagara/Source/Niagara/Private/NiagaraSystemSimulation.cpp
@@ -14,6 +14,7 @@
 #include "NiagaraComponent.h"
 #include "NiagaraWorldManager.h"
 #include "NiagaraEmitterInstanceBatcher.h"
+#include "NiagaraCrashReporterHandler.h"
 
 // Niagara simulations async will block the tick task from completion until all async work is finished
 // If simulations are allowed to tick async we will create a FNiagaraSystemSimulationTickTask task to run on any thread
@@ -42,6 +43,7 @@ DECLARE_CYCLE_STAT(TEXT("System Sim Init (DirectBindings) [GT]"), STAT_NiagaraSy
 DECLARE_CYCLE_STAT(TEXT("ForcedWaitForAsync"), STAT_NiagaraSystemSim_ForceWaitForAsync, STATGROUP_Niagara);
 DECLARE_CYCLE_STAT(TEXT("ForcedWait Fake Stall"), STAT_NiagaraSystemSim_ForceWaitFakeStall, STATGROUP_Niagara);
 
+
 static int32 GbDumpSystemData = 0;
 static FAutoConsoleVariableRef CVarNiagaraDumpSystemData(
 	TEXT("fx.DumpSystemData"),
@@ -419,11 +421,29 @@ void FNiagaraSystemSimulation::Destroy()
 
 	while (SystemInstances.Num())
 	{
-		SystemInstances.Last()->Deactivate(true);
+		FNiagaraSystemInstance* Inst = SystemInstances.Last();
+		check(Inst);
+		if (ensure(Inst->GetComponent()))//Currently we have no cases whre there shouldn't be a component but maybe in future.
+		{
+			Inst->GetComponent()->DeactivateImmediate();
+		}
+		else
+		{
+			Inst->Deactivate(true);
+		}
 	}
 	while (PendingSystemInstances.Num())
 	{
-		PendingSystemInstances.Last()->Deactivate(true);
+		FNiagaraSystemInstance* Inst = PendingSystemInstances.Last();
+		check(Inst);
+		if (ensure(Inst->GetComponent()))//Currently we have no cases whre there shouldn't be a component but maybe in future.
+		{
+			Inst->GetComponent()->DeactivateImmediate();
+		}
+		else
+		{
+			Inst->Deactivate(true);
+		}
 	}
 	SystemInstances.Empty();
 	PendingSystemInstances.Empty();
@@ -616,6 +636,8 @@ void FNiagaraSystemSimulation::Tick_GameThread(float DeltaSeconds, const FGraphE
 	check(IsInGameThread());
 	check(bInSpawnPhase == false);
 
+	FNiagaraCrashReporterScope CRScope(this);
+
 	WaitForSystemTickComplete(true);
 
 	SCOPE_CYCLE_COUNTER(STAT_NiagaraOverview_GT);
@@ -791,6 +813,13 @@ void FNiagaraSystemSimulation::Tick_GameThread(float DeltaSeconds, const FGraphE
 	static const auto EffectsQualityCVar = IConsoleManager::Get().FindConsoleVariable(TEXT("sg.EffectsQuality"));
 	FNiagaraSystemSimulationTickContext Context(this, SystemInstances, MainDataSet, DeltaSeconds, SpawnNum, EffectsQualityCVar->GetInt(), MyCompletionGraphEvent);
 
+	//Solo systems add their counts in their component tick.
+	if (GetIsSolo() == false)
+	{
+		System->AddToInstanceCountStat(SystemInstances.Num(), false);
+		INC_DWORD_STAT_BY(STAT_TotalNiagaraSystemInstances, SystemInstances.Num());
+	}
+
 	//Now kick of the concurrent tick.
 	if (Context.bTickAsync)
 	{
@@ -902,6 +931,8 @@ void FNiagaraSystemSimulation::Spawn_GameThread(float DeltaSeconds)
 	UNiagaraSystem* System = WeakSystem.Get();
 	FScopeCycleCounterUObject AdditionalScope(System, GET_STATID(STAT_NiagaraOverview_GT_CNC));
 
+	FNiagaraCrashReporterScope CRScope(this);
+
 	WaitForSystemTickComplete(true);
 
 	bInSpawnPhase = true;
@@ -1033,6 +1064,8 @@ void FNiagaraSystemSimulation::Tick_Concurrent(FNiagaraSystemSimulationTickConte
 	FNiagaraScopedRuntimeCycleCounter RuntimeScope(Context.System, true, true);
 	FNiagaraSystemInstance* SoloSystemInstance = bIsSolo && Context.Instances.Num() == 1 ? Context.Instances[0] : nullptr;
 
+	FNiagaraCrashReporterScope CRScope(this);
+
 	if (bCanExecute && Context.Instances.Num())
 	{
 		if (GbDumpSystemData || Context.System->bDumpDebugSystemInfo)
@@ -1688,34 +1721,48 @@ void FNiagaraSystemSimulation::InitParameterDataSetBindings(FNiagaraSystemInstan
 		for (int32 EmitterIdx = 0; EmitterIdx < EmitterCount; ++EmitterIdx)
 		{
 			FNiagaraEmitterInstance& EmitterInst = Emitters[EmitterIdx].Get();
-			const FString EmitterName = EmitterInst.GetCachedEmitter()->GetUniqueEmitterName();
-
-			FNiagaraScriptExecutionContext& SpawnContext = EmitterInst.GetSpawnExecutionContext();
-			DataSetToEmitterSpawnParameters[EmitterIdx].Init(MainDataSet, SpawnContext.Parameters);
+			if (!EmitterInst.IsDisabled())
+			{
+				const FString EmitterName = EmitterInst.GetCachedEmitter()->GetUniqueEmitterName();
 
-			FNiagaraScriptExecutionContext& UpdateContext = EmitterInst.GetUpdateExecutionContext();
-			DataSetToEmitterUpdateParameters[EmitterIdx].Init(MainDataSet, UpdateContext.Parameters);
+				FNiagaraScriptExecutionContext& SpawnContext = EmitterInst.GetSpawnExecutionContext();
+				DataSetToEmitterSpawnParameters[EmitterIdx].Init(MainDataSet, SpawnContext.Parameters);
 
+				FNiagaraScriptExecutionContext& UpdateContext = EmitterInst.GetUpdateExecutionContext();
+				DataSetToEmitterUpdateParameters[EmitterIdx].Init(MainDataSet, UpdateContext.Parameters);
 
-			FNiagaraComputeExecutionContext* GPUContext = EmitterInst.GetGPUContext();
-			if (GPUContext)
-			{
-				DataSetToEmitterGPUParameters[EmitterIdx].Init(MainDataSet, GPUContext->CombinedParamStore);
-			}
+				FNiagaraComputeExecutionContext* GPUContext = EmitterInst.GetGPUContext();
+				if (GPUContext)
+				{
+					DataSetToEmitterGPUParameters[EmitterIdx].Init(MainDataSet, GPUContext->CombinedParamStore);
+				}
 
-			TArray<FNiagaraScriptExecutionContext>& EventContexts = EmitterInst.GetEventExecutionContexts();
-			const int32 EventCount = EventContexts.Num();
-			DataSetToEmitterEventParameters[EmitterIdx].SetNum(EventCount);
+				TArray<FNiagaraScriptExecutionContext>& EventContexts = EmitterInst.GetEventExecutionContexts();
+				const int32 EventCount = EventContexts.Num();
+				DataSetToEmitterEventParameters[EmitterIdx].SetNum(EventCount);
 
-			for (int32 EventIdx = 0; EventIdx < EventCount; ++EventIdx)
-			{
-				FNiagaraScriptExecutionContext& EventContext = EventContexts[EventIdx];
-				DataSetToEmitterEventParameters[EmitterIdx][EventIdx].Init(MainDataSet, EventContext.Parameters);
+				for (int32 EventIdx = 0; EventIdx < EventCount; ++EventIdx)
+				{
+					FNiagaraScriptExecutionContext& EventContext = EventContexts[EventIdx];
+					DataSetToEmitterEventParameters[EmitterIdx][EventIdx].Init(MainDataSet, EventContext.Parameters);
+				}
 			}
 		}
 	}
 }
 
+const FString& FNiagaraSystemSimulation::GetCrashReporterTag()const
+{
+	if (CrashReporterTag.IsEmpty())
+	{
+		UNiagaraSystem* Sys = GetSystem();
+		const FString& AssetName = Sys ? Sys->GetFullName() : TEXT("nullptr");
+
+		CrashReporterTag = FString::Printf(TEXT("SystemSimulation | System: %s | bSolo: %s |"), *AssetName, bIsSolo ? TEXT("true") : TEXT("false"));
+	}
+	return CrashReporterTag;
+}
+
 void FNiagaraConstantBufferToDataSetBinding::Init(const FNiagaraSystemCompiledData& CompiledData)
 {
 	// for now we'll copy the data to our local structure so that we don't have to worry about the lifetime of the compiled data
diff --git a/Engine/Plugins/FX/Niagara/Source/Niagara/Private/NiagaraWorldManager.cpp b/Engine/Plugins/FX/Niagara/Source/Niagara/Private/NiagaraWorldManager.cpp
--- a/Engine/Plugins/FX/Niagara/Source/Niagara/Private/NiagaraWorldManager.cpp
+++ b/Engine/Plugins/FX/Niagara/Source/Niagara/Private/NiagaraWorldManager.cpp
@@ -19,10 +19,12 @@
 #include "NiagaraComponentPool.h"
 #include "NiagaraComponent.h"
 #include "NiagaraEffectType.h"
+#include "HAL/PlatformApplicationMisc.h"
 
 DECLARE_CYCLE_STAT(TEXT("Niagara Manager Update Scalability Managers [GT]"), STAT_UpdateScalabilityManagers, STATGROUP_Niagara);
 DECLARE_CYCLE_STAT(TEXT("Niagara Manager Tick [GT]"), STAT_NiagaraWorldManTick, STATGROUP_Niagara);
 DECLARE_CYCLE_STAT(TEXT("Niagara Manager Wait On Render [GT]"), STAT_NiagaraWorldManWaitOnRender, STATGROUP_Niagara);
+DECLARE_CYCLE_STAT(TEXT("Niagara Manager Wait Pre Garbage Collect [GT]"), STAT_NiagaraWorldManWaitPreGC, STATGROUP_Niagara);
 
 static int GNiagaraAllowAsyncWorkToEndOfFrame = 1;
 static FAutoConsoleVariableRef CVarNiagaraAllowAsyncWorkToEndOfFrame(
@@ -137,6 +139,7 @@ FName FNiagaraWorldManagerTickFunction::DiagnosticContext(bool bDetailed)
 FNiagaraWorldManager::FNiagaraWorldManager(UWorld* InWorld)
 	: World(InWorld)
 	, CachedEffectsQuality(INDEX_NONE)
+	, bAppHasFocus(true)
 {
 	for (int32 TickGroup=0; TickGroup < NiagaraNumTickGroups; ++TickGroup)
 	{
@@ -145,6 +148,7 @@ FNiagaraWorldManager::FNiagaraWorldManager(UWorld* InWorld)
 		TickFunc.EndTickGroup = GNiagaraAllowAsyncWorkToEndOfFrame ? TG_LastDemotable : (ETickingGroup)TickFunc.TickGroup;
 		TickFunc.bCanEverTick = true;
 		TickFunc.bStartWithTickEnabled = true;
+		TickFunc.bAllowTickOnDedicatedServer = false;
 		TickFunc.bHighPriority = true;
 		TickFunc.Owner = this;
 		TickFunc.RegisterTickFunction(InWorld->PersistentLevel);
@@ -371,6 +375,7 @@ void FNiagaraWorldManager::PreGarbageCollect()
 {
 	if (GNiagaraWaitOnPreGC)
 	{
+		SCOPE_CYCLE_COUNTER(STAT_NiagaraWorldManWaitPreGC);
 		// We must wait for system simulation & instance async ticks to complete before garbage collection can start
 		// The reason for this is that our async ticks could be referencing GC objects, i.e. skel meshes, etc, and we don't want them to become unreachable while we are potentially using them
 		for (int TG = 0; TG < NiagaraNumTickGroups; ++TG)
@@ -568,6 +573,12 @@ void FNiagaraWorldManager::Tick(ETickingGroup TickGroup, float DeltaSeconds, ELe
 	{
 		FNiagaraSharedObject::FlushDeletionList();
 
+#if PLATFORM_DESKTOP
+		bAppHasFocus = FPlatformApplicationMisc::IsThisApplicationForeground();
+#else
+		bAppHasFocus = true;
+#endif
+
 		// Cache player view locations for all system instances to access
 		//-TODO: Do we need to do this per tick group?
 		bCachedPlayerViewLocationsValid = true;
@@ -795,7 +806,8 @@ void FNiagaraWorldManager::CalculateScalabilityState(UNiagaraSystem* System, con
 	OwnerLODCull(EffectType, ScalabilitySettings, Component, OutState);
 	
 	//Can't cull dynamic bounds by visibility
-	if (System->bFixedBounds)
+
+	if (System->bFixedBounds && bAppHasFocus)
 	{
 		VisibilityCull(EffectType, ScalabilitySettings, Component, OutState);
 	}
diff --git a/Engine/Plugins/FX/Niagara/Source/Niagara/Public/NiagaraCommon.h b/Engine/Plugins/FX/Niagara/Source/Niagara/Public/NiagaraCommon.h
--- a/Engine/Plugins/FX/Niagara/Source/Niagara/Public/NiagaraCommon.h
+++ b/Engine/Plugins/FX/Niagara/Source/Niagara/Public/NiagaraCommon.h
@@ -546,18 +546,19 @@ struct NIAGARA_API FNiagaraSystemUpdateContext
 {
 	GENERATED_BODY()
 
-	FNiagaraSystemUpdateContext(const UNiagaraSystem* System, bool bReInit, bool bInDestroyOnAdd = false) :bDestroyOnAdd(bInDestroyOnAdd) { Add(System, bReInit); }
+	FNiagaraSystemUpdateContext(const UNiagaraSystem* System, bool bReInit, bool bInDestroyOnAdd = false, bool bInOnlyActive = false) :bDestroyOnAdd(bInDestroyOnAdd), bOnlyActive(bInOnlyActive) { Add(System, bReInit); }
 #if WITH_EDITORONLY_DATA
-	FNiagaraSystemUpdateContext(const UNiagaraEmitter* Emitter, bool bReInit, bool bInDestroyOnAdd = false) : bDestroyOnAdd(bInDestroyOnAdd) { Add(Emitter, bReInit); }
-	FNiagaraSystemUpdateContext(const UNiagaraScript* Script, bool bReInit, bool bInDestroyOnAdd = false) :bDestroyOnAdd(bInDestroyOnAdd) { Add(Script, bReInit); }
+	FNiagaraSystemUpdateContext(const UNiagaraEmitter* Emitter, bool bReInit, bool bInDestroyOnAdd = false, bool bInOnlyActive = false) :bDestroyOnAdd(bInDestroyOnAdd), bOnlyActive(bInOnlyActive)  { Add(Emitter, bReInit); }
+	FNiagaraSystemUpdateContext(const UNiagaraScript* Script, bool bReInit, bool bInDestroyOnAdd = false, bool bInOnlyActive = false) :bDestroyOnAdd(bInDestroyOnAdd), bOnlyActive(bInOnlyActive)  { Add(Script, bReInit); }
 	//FNiagaraSystemUpdateContext(UNiagaraDataInterface* Interface, bool bReinit) : Add(Interface, bReinit) {}
-	FNiagaraSystemUpdateContext(const UNiagaraParameterCollection* Collection, bool bReInit, bool bInDestroyOnAdd = false) :bDestroyOnAdd(bInDestroyOnAdd) { Add(Collection, bReInit); }
+	FNiagaraSystemUpdateContext(const UNiagaraParameterCollection* Collection, bool bReInit, bool bInDestroyOnAdd = false, bool bInOnlyActive = false) :bDestroyOnAdd(bInDestroyOnAdd), bOnlyActive(bInOnlyActive) { Add(Collection, bReInit); }
 #endif
-	FNiagaraSystemUpdateContext():bDestroyOnAdd(false){ }
+	FNiagaraSystemUpdateContext():bDestroyOnAdd(false), bOnlyActive(false){ }
 
 	~FNiagaraSystemUpdateContext();
 
 	void SetDestroyOnAdd(bool bInDestroyOnAdd) { bDestroyOnAdd = bInDestroyOnAdd; }
+	void SetOnlyActive(bool bInOnlyActive) { bOnlyActive = bInOnlyActive; }
 
 	void Add(const UNiagaraSystem* System, bool bReInit);
 #if WITH_EDITORONLY_DATA
@@ -585,6 +586,7 @@ private:
 	TArray<UNiagaraSystem*> SystemSimsToDestroy;
 
 	bool bDestroyOnAdd;
+	bool bOnlyActive;
 	//TODO: When we allow component less systems we'll also want to find and reset those.
 };
 
diff --git a/Engine/Plugins/FX/Niagara/Source/Niagara/Public/NiagaraComponent.h b/Engine/Plugins/FX/Niagara/Source/Niagara/Public/NiagaraComponent.h
--- a/Engine/Plugins/FX/Niagara/Source/Niagara/Public/NiagaraComponent.h
+++ b/Engine/Plugins/FX/Niagara/Source/Niagara/Public/NiagaraComponent.h
@@ -173,7 +173,8 @@ public:
 	void DeactivateInternal(bool bIsScalabilityCull);
 	void DeactivateImmediateInternal(bool bIsScalabilityCull);
 
-	bool RegisterWithScalabilityManagerOrPreCull();
+	bool ShouldPreCull();
+	void RegisterWithScalabilityManager();
 	void UnregisterWithScalabilityManager();
 
 	public:
@@ -519,6 +520,8 @@ public:
 	UFUNCTION(BlueprintCallable, Category = Scalability, meta = (Keywords = "LOD scalability"))
 	void SetAllowScalability(bool bAllow);
 
+	FORCEINLINE bool IsRegisteredWithScalabilityManager()const { return ScalabilityManagerHandle != INDEX_NONE; }
+	FORCEINLINE int32 GetScalabilityManagerHandle()const { return ScalabilityManagerHandle; }
 private:
 	/** Did we try and activate but fail due to the asset being not yet ready. Keep looping.*/
 	uint32 bAwaitingActivationDueToNotReady : 1;
@@ -531,6 +534,9 @@ private:
 	/** True if this component is allowed to perform scalability checks and potentially be culled etc. Occasionally it is useful to disable this for specific components. E.g. Effects on the local player. */
 	uint32 bAllowScalability : 1;
 
+	/** True if this component has been culled by the scalability manager. */
+	uint32 bIsCulledByScalability : 1;
+
 	/** Flag to mark us as currently changing auto attachment as part of Activate/Deactivate so we don't reset in the OnAttachmentChanged() callback. */
 	//uint32 bIsChangingAutoAttachment : 1;
 
diff --git a/Engine/Plugins/FX/Niagara/Source/Niagara/Public/NiagaraCrashReporterHandler.h b/Engine/Plugins/FX/Niagara/Source/Niagara/Public/NiagaraCrashReporterHandler.h
new file mode 100644
--- /dev/null
+++ b/Engine/Plugins/FX/Niagara/Source/Niagara/Public/NiagaraCrashReporterHandler.h
@@ -0,0 +1,62 @@
+// Copyright Epic Games, Inc. All Rights Reserved.
+
+#pragma once
+
+#include "CoreMinimal.h"
+#include "HAL/ThreadSingleton.h"
+#include "HAL/CriticalSection.h"
+
+class FNiagaraSystemInstance;
+class FNiagaraSystemSimulation;
+
+/** Manages Niagara's integration with the Crash Reporter. Allows Niagara code to add information to a crash report which can help track down rare bugs. */
+class FNiagaraCrashReporterHandler : public TThreadSingleton<FNiagaraCrashReporterHandler>
+{
+	//Each thread we execute on can add it's own named data to the CR.
+	FString Name;
+
+	/** String to clear to when we're not inside the scope of a system instance. */
+	FString NoneInfoString;
+
+	/** Stack of scope infos, allows multi level scope tracking. */
+	TArray<FString> ScopeInfoStack;
+
+	//Allows us to clear out all data if we disable reporting.
+	static TArray<FNiagaraCrashReporterHandler*> AllHandlers;
+
+	/** Critical section ensuring thread-safe access to shared Crash Reporter state. */
+	static FCriticalSection CritSec;
+
+	void SetInfo(const FString& Info);
+
+	/** Push an arbitrary string onto info stack and set as Niagara's current CR info for this thread. */
+	void PushInfo_Internal(const FString& Info);
+public:
+
+	FNiagaraCrashReporterHandler();
+	~FNiagaraCrashReporterHandler();
+
+	/** Clears out Niagara's CR entry for this thread. */
+	void Clear();
+
+	/** Push info on a system instance to the info stack and set as Niagara's current CR info for this thread. */
+	void PushInfo(FNiagaraSystemInstance* Inst);
+	/** Push info on a system simulation to the info stack and set as Niagara's current CR info for this thread. */
+	void PushInfo(FNiagaraSystemSimulation* SystemSim);
+	/** Pop the current info from the stack and set the previous info as Niagara's current CR info for this thread. */
+	void PopInfo();
+
+	/** Handle a change in the enabled CVar for Niagara's CR info. */
+	static void OnEnabledChanged(IConsoleVariable* InVariable);
+};
+
+/** Helper object allowing easy tracking of Niagara code in it's crash reporter integration.  */
+class FNiagaraCrashReporterScope
+{
+private:
+	bool bWasEnabled = false;
+public:
+	FNiagaraCrashReporterScope(FNiagaraSystemInstance* Inst);
+	FNiagaraCrashReporterScope(FNiagaraSystemSimulation* Sim);
+	~FNiagaraCrashReporterScope();
+};
\ No newline at end of file
diff --git a/Engine/Plugins/FX/Niagara/Source/Niagara/Public/NiagaraModule.h b/Engine/Plugins/FX/Niagara/Source/Niagara/Public/NiagaraModule.h
--- a/Engine/Plugins/FX/Niagara/Source/Niagara/Public/NiagaraModule.h
+++ b/Engine/Plugins/FX/Niagara/Source/Niagara/Public/NiagaraModule.h
@@ -165,7 +165,7 @@ public:
 	FORCEINLINE static const FNiagaraVariable&  GetVar_DataInstance_Alive() { return DataInstance_Alive; }
 	FORCEINLINE static const FNiagaraVariable&  GetVar_BeginDefaults() { return Translator_BeginDefaults; }
 
-	static void OnEffectsQualityChanged(int32 NewEffectsQuality);
+	static void OnEffectsQualityChanged(int32 NewEffectsQuality, bool bForce);
 
 	FOnProcessQueue OnProcessQueue;
 
diff --git a/Engine/Plugins/FX/Niagara/Source/Niagara/Public/NiagaraSystemInstance.h b/Engine/Plugins/FX/Niagara/Source/Niagara/Public/NiagaraSystemInstance.h
--- a/Engine/Plugins/FX/Niagara/Source/Niagara/Public/NiagaraSystemInstance.h
+++ b/Engine/Plugins/FX/Niagara/Source/Niagara/Public/NiagaraSystemInstance.h
@@ -280,6 +280,7 @@ public:
 
 	FORCEINLINE void SetLODDistance(float InLODDistance, float InMaxLODDistance);
 
+	const FString& GetCrashReporterTag()const;
 private:
 
 	void DestroyDataInterfaceInstanceData();
@@ -429,6 +430,9 @@ private:
 	/** Array of emitter indices sorted by execution priority. The emitters will be ticked in this order. */
 	TArray<int32> EmitterExecutionOrder;
 
+	/** Tag we feed into crash reporter for this instance. */
+	mutable FString CrashReporterTag;
+
 public:
 	// Transient data that is accumulated during tick.
 	uint32 TotalGPUParamSize = 0;
diff --git a/Engine/Plugins/FX/Niagara/Source/Niagara/Public/NiagaraSystemSimulation.h b/Engine/Plugins/FX/Niagara/Source/Niagara/Public/NiagaraSystemSimulation.h
--- a/Engine/Plugins/FX/Niagara/Source/Niagara/Public/NiagaraSystemSimulation.h
+++ b/Engine/Plugins/FX/Niagara/Source/Niagara/Public/NiagaraSystemSimulation.h
@@ -232,6 +232,7 @@ public:
 
 	void AddTickGroupPromotion(FNiagaraSystemInstance* Instance);
 
+	const FString& GetCrashReporterTag()const;
 protected:
 	/** Sets constant parameter values */
 	void SetupParameters_GameThread(float DeltaSeconds);
@@ -350,4 +351,6 @@ protected:
 
 	/** Current task that is executing */
 	FGraphEventRef SystemTickGraphEvent;
+
+	mutable FString CrashReporterTag;
 };
diff --git a/Engine/Plugins/FX/Niagara/Source/Niagara/Public/NiagaraWorldManager.h b/Engine/Plugins/FX/Niagara/Source/Niagara/Public/NiagaraWorldManager.h
--- a/Engine/Plugins/FX/Niagara/Source/Niagara/Public/NiagaraWorldManager.h
+++ b/Engine/Plugins/FX/Niagara/Source/Niagara/Public/NiagaraWorldManager.h
@@ -241,5 +241,8 @@ private:
 
 	UPROPERTY(transient)
 	TMap<UNiagaraEffectType*, FNiagaraScalabilityManager> ScalabilityManagers;
+
+	/** True if the app has focus. We prevent some culling if the app doesn't have focus as it can interefre. */
+	bool bAppHasFocus;
 };
 
diff --git a/Engine/Source/Editor/UnrealEd/Private/EditorViewportClient.cpp b/Engine/Source/Editor/UnrealEd/Private/EditorViewportClient.cpp
--- a/Engine/Source/Editor/UnrealEd/Private/EditorViewportClient.cpp
+++ b/Engine/Source/Editor/UnrealEd/Private/EditorViewportClient.cpp
@@ -3698,8 +3698,7 @@ void FEditorViewportClient::Draw(FViewport* InViewport, FCanvas* Canvas)
 	float DeltaTimeSeconds;
 
 	UWorld* World = GetWorld();
-	// During Simulation blueprints are directly using the World time, causing a mismatch with Material's Frame time
-	if (!World || ( GetScene() != World->Scene) || (IsRealtime() && !IsSimulateInEditorViewport()))
+	if (!World || (GetScene() != World->Scene) || UseAppTime()) 
 	{
 		// Use time relative to start time to avoid issues with float vs double
 		TimeSeconds = FApp::GetCurrentTime() - GStartTime;
diff --git a/Engine/Source/Editor/UnrealEd/Public/EditorViewportClient.h b/Engine/Source/Editor/UnrealEd/Public/EditorViewportClient.h
--- a/Engine/Source/Editor/UnrealEd/Public/EditorViewportClient.h
+++ b/Engine/Source/Editor/UnrealEd/Public/EditorViewportClient.h
@@ -794,6 +794,8 @@ public:
 	 */
 	virtual bool WantsDrawWhenAppIsHidden() const { return false; }
 
+	/** Should this viewport use app time instead of world time. */
+	virtual bool UseAppTime() const { return IsRealtime() && !IsSimulateInEditorViewport(); }
 public:
 	/** True if the window is maximized or floating */
 	bool IsVisible() const;
diff --git a/Engine/Source/Editor/UnrealEd/Public/LevelEditorViewport.h b/Engine/Source/Editor/UnrealEd/Public/LevelEditorViewport.h
--- a/Engine/Source/Editor/UnrealEd/Public/LevelEditorViewport.h
+++ b/Engine/Source/Editor/UnrealEd/Public/LevelEditorViewport.h
@@ -539,6 +539,8 @@ public:
 	 */
 	static UObject* GetOrCreateMaterialFromTexture( UTexture* UnrealTexture );
 
+	virtual bool UseAppTime() const override { return false; }
+
 protected:
 	/**
 	* Sets the state of creating a preview actor in the viewport.
diff --git a/Engine/Source/Runtime/Engine/Private/DeviceProfiles/DeviceProfileManager.cpp b/Engine/Source/Runtime/Engine/Private/DeviceProfiles/DeviceProfileManager.cpp
--- a/Engine/Source/Runtime/Engine/Private/DeviceProfiles/DeviceProfileManager.cpp
+++ b/Engine/Source/Runtime/Engine/Private/DeviceProfiles/DeviceProfileManager.cpp
@@ -559,7 +559,7 @@ void UDeviceProfileManager::LoadProfiles()
 				ITargetPlatform* Platform = TargetPlatforms[PlatformIndex];
 
 				// Set TextureLODSettings
-				const UTextureLODSettings* TextureLODSettingsObj = FindProfile(*Platform->GetPlatformInfo().VanillaPlatformName.ToString());
+				const UTextureLODSettings* TextureLODSettingsObj = FindProfile(*Platform->GetPlatformInfo().VanillaPlatformName.ToString(), false);
 				Platform->RegisterTextureLODSettings(TextureLODSettingsObj);
 			}
 		}
diff --git a/Engine/Source/Runtime/Engine/Private/Particles/ParticleSystemManager.cpp b/Engine/Source/Runtime/Engine/Private/Particles/ParticleSystemManager.cpp
--- a/Engine/Source/Runtime/Engine/Private/Particles/ParticleSystemManager.cpp
+++ b/Engine/Source/Runtime/Engine/Private/Particles/ParticleSystemManager.cpp
@@ -228,6 +228,7 @@ FParticleSystemWorldManager::FParticleSystemWorldManager(UWorld* InWorld)
 		TickFunc.EndTickGroup = TickFunc.TickGroup;
 		TickFunc.bCanEverTick = true;
 		TickFunc.bStartWithTickEnabled = true;
+		TickFunc.bAllowTickOnDedicatedServer = false;
 		TickFunc.bHighPriority = true;
 		TickFunc.Owner = this;
 		TickFunc.RegisterTickFunction(InWorld->PersistentLevel);
diff --git a/Engine/Source/Runtime/Engine/Private/Scalability.cpp b/Engine/Source/Runtime/Engine/Private/Scalability.cpp
--- a/Engine/Source/Runtime/Engine/Private/Scalability.cpp
+++ b/Engine/Source/Runtime/Engine/Private/Scalability.cpp
@@ -19,42 +19,42 @@ static TAutoConsoleVariable<float> CVarResolutionQuality(
 
 static TAutoConsoleVariable<int32> CVarViewDistanceQuality(
 	TEXT("sg.ViewDistanceQuality"),
-	3,
+	Scalability::DefaultQualityLevel,
 	TEXT("Scalability quality state (internally used by scalability system, ini load/save or using SCALABILITY console command)\n")
 	TEXT(" 0:low, 1:med, 2:high, 3:epic, 4:cinematic, default: 3"),
 	ECVF_ScalabilityGroup);
 
 static TAutoConsoleVariable<int32> CVarAntiAliasingQuality(
 	TEXT("sg.AntiAliasingQuality"),
-	3,
+	Scalability::DefaultQualityLevel,
 	TEXT("Scalability quality state (internally used by scalability system, ini load/save or using SCALABILITY console command)\n")
 	TEXT(" 0:low, 1:med, 2:high, 3:epic, 4:cinematic, default: 3"),
 	ECVF_ScalabilityGroup);
 
 static TAutoConsoleVariable<int32> CVarShadowQuality(
 	TEXT("sg.ShadowQuality"),
-	3,
+	Scalability::DefaultQualityLevel,
 	TEXT("Scalability quality state (internally used by scalability system, ini load/save or using SCALABILITY console command)\n")
 	TEXT(" 0:low, 1:med, 2:high, 3:epic, 4:cinematic, default: 3"),
 	ECVF_ScalabilityGroup);
 
 static TAutoConsoleVariable<int32> CVarPostProcessQuality(
 	TEXT("sg.PostProcessQuality"),
-	3,
+	Scalability::DefaultQualityLevel,
 	TEXT("Scalability quality state (internally used by scalability system, ini load/save or using SCALABILITY console command)\n")
 	TEXT(" 0:low, 1:med, 2:high, 3:epic, 4:cinematic, default: 3"),
 	ECVF_ScalabilityGroup);
 
 static TAutoConsoleVariable<int32> CVarTextureQuality(
 	TEXT("sg.TextureQuality"),
-	3,
+	Scalability::DefaultQualityLevel,
 	TEXT("Scalability quality state (internally used by scalability system, ini load/save or using SCALABILITY console command)\n")
 	TEXT(" 0:low, 1:med, 2:high, 3:epic, 4:cinematic, default: 3"),
 	ECVF_ScalabilityGroup);
 
 static TAutoConsoleVariable<int32> CVarEffectsQuality(
 	TEXT("sg.EffectsQuality"),
-	3,
+	Scalability::DefaultQualityLevel,
 	TEXT("Scalability quality state (internally used by scalability system, ini load/save or using SCALABILITY console command)\n")
 	TEXT(" 0:low, 1:med, 2:high, 3:epic, 4:cinematic, default: 3"),
 	ECVF_ScalabilityGroup);
@@ -68,7 +68,7 @@ static TAutoConsoleVariable<int32> CVarFoliageQuality(
 
 static TAutoConsoleVariable<int32> CVarShadingQuality(
 	TEXT("sg.ShadingQuality"),
-	3,
+	Scalability::DefaultQualityLevel,
 	TEXT("Scalability quality state (internally used by scalability system, ini load/save or using SCALABILITY console command)\n")
 	TEXT(" 0:low, 1:med, 2:high, 3:epic, 4:cinematic, default: 3"),
 	ECVF_ScalabilityGroup);
diff --git a/Engine/Source/Runtime/Engine/Public/Scalability.h b/Engine/Source/Runtime/Engine/Public/Scalability.h
--- a/Engine/Source/Runtime/Engine/Public/Scalability.h
+++ b/Engine/Source/Runtime/Engine/Public/Scalability.h
@@ -16,6 +16,8 @@ namespace Scalability
 		ERelativeToMax,
 	};
 
+	const int32 DefaultQualityLevel = 3;
+
 	/**
 	 * Structure for holding the state of the engine scalability groups
 	 * Actual engine state you can get though GetQualityLevels().
