- Replace direct hook to EffectsQuality with Niagara specific quality cvar. FX.Niagara.QualityLevel.
Getting around the issue of Epic meaning different things on PS4 vs Mobile.
- Moving Mobile Devices to be Low or Medium by default.
- Ensure destroy on add for update contexts is applied regardless of active state. Ensures that next activate of currently inactive systems will trigger a reinit that is needed.

diff --git a/Engine/Config/Android/AndroidScalability.ini b/Engine/Config/Android/AndroidScalability.ini
--- a/Engine/Config/Android/AndroidScalability.ini
+++ b/Engine/Config/Android/AndroidScalability.ini
@@ -119,6 +119,7 @@ r.SkyAtmosphere.FastSkyLUT.SampleCountMax=8.0
 r.SkyAtmosphere.FastSkyLUT.Width=96
 r.SkyAtmosphere.FastSkyLUT.Height=50
 r.SkyAtmosphere.TransmittanceLUT.UseSmallFormat=1
+fx.Niagara.QualityLevel=0
 
 [EffectsQuality@1]
 r.SkyAtmosphere.AerialPerspectiveLUT.FastApplyOnOpaque=1 ; Always have FastSkyLUT 1 in this case to avoid wrong sky
@@ -130,6 +131,7 @@ r.SkyAtmosphere.FastSkyLUT.SampleCountMax=8.0
 r.SkyAtmosphere.FastSkyLUT.Width=96
 r.SkyAtmosphere.FastSkyLUT.Height=50
 r.SkyAtmosphere.TransmittanceLUT.UseSmallFormat=1
+fx.Niagara.QualityLevel=0
 
 [EffectsQuality@2]
 r.SkyAtmosphere.AerialPerspectiveLUT.FastApplyOnOpaque=1 ; Always have FastSkyLUT 1 in this case to avoid wrong sky
@@ -141,6 +143,7 @@ r.SkyAtmosphere.FastSkyLUT.SampleCountMax=8.0
 r.SkyAtmosphere.FastSkyLUT.Width=96
 r.SkyAtmosphere.FastSkyLUT.Height=50
 r.SkyAtmosphere.TransmittanceLUT.UseSmallFormat=1
+fx.Niagara.QualityLevel=0
 
 [EffectsQuality@3]
 r.SkyAtmosphere.AerialPerspectiveLUT.FastApplyOnOpaque=1 ; Always have FastSkyLUT 1 in this case to avoid wrong sky
@@ -152,6 +155,7 @@ r.SkyAtmosphere.FastSkyLUT.SampleCountMax=8.0
 r.SkyAtmosphere.FastSkyLUT.Width=96
 r.SkyAtmosphere.FastSkyLUT.Height=50
 r.SkyAtmosphere.TransmittanceLUT.UseSmallFormat=1
+fx.Niagara.QualityLevel=1
 
 ; ------------------------------------------------------------------------------
 ; Texture settings
diff --git a/Engine/Config/BaseScalability.ini b/Engine/Config/BaseScalability.ini
--- a/Engine/Config/BaseScalability.ini
+++ b/Engine/Config/BaseScalability.ini
@@ -408,6 +408,7 @@ r.SkyAtmosphere.SampleCountMax=16.0
 r.SkyAtmosphere.TransmittanceLUT.UseSmallFormat=1
 r.SkyAtmosphere.TransmittanceLUT.SampleCount=10.0
 r.SkyAtmosphere.MultiScatteringLUT.SampleCount=15.0
+fx.Niagara.QualityLevel=0
 
 [EffectsQuality@1]
 r.TranslucencyLightingVolumeDim=32
@@ -436,6 +437,7 @@ r.SkyAtmosphere.SampleCountMax=32.0
 r.SkyAtmosphere.TransmittanceLUT.UseSmallFormat=0
 r.SkyAtmosphere.TransmittanceLUT.SampleCount=10.0
 r.SkyAtmosphere.MultiScatteringLUT.SampleCount=15.0
+fx.Niagara.QualityLevel=1
 
 [EffectsQuality@2]
 r.TranslucencyLightingVolumeDim=48
@@ -464,6 +466,7 @@ r.SkyAtmosphere.SampleCountMax=32.0
 r.SkyAtmosphere.TransmittanceLUT.UseSmallFormat=0
 r.SkyAtmosphere.TransmittanceLUT.SampleCount=10.0
 r.SkyAtmosphere.MultiScatteringLUT.SampleCount=15.0
+fx.Niagara.QualityLevel=2
 
 [EffectsQuality@3]
 r.TranslucencyLightingVolumeDim=64
@@ -492,6 +495,7 @@ r.SkyAtmosphere.SampleCountMax=32.0
 r.SkyAtmosphere.TransmittanceLUT.UseSmallFormat=0
 r.SkyAtmosphere.TransmittanceLUT.SampleCount=10.0
 r.SkyAtmosphere.MultiScatteringLUT.SampleCount=15.0
+fx.Niagara.QualityLevel=3
 
 [EffectsQuality@Cine]
 r.TranslucencyLightingVolumeDim=64
@@ -520,6 +524,7 @@ r.SkyAtmosphere.SampleCountMax=32.0
 r.SkyAtmosphere.TransmittanceLUT.UseSmallFormat=0
 r.SkyAtmosphere.TransmittanceLUT.SampleCount=30.0
 r.SkyAtmosphere.MultiScatteringLUT.SampleCount=20.0
+fx.Niagara.QualityLevel=4
 
 ;-----------------------------------------------------------------------------------------------------------------
 
diff --git a/Engine/Config/IOS/IOSScalability.ini b/Engine/Config/IOS/IOSScalability.ini
--- a/Engine/Config/IOS/IOSScalability.ini
+++ b/Engine/Config/IOS/IOSScalability.ini
@@ -119,6 +119,7 @@ r.SkyAtmosphere.FastSkyLUT.SampleCountMax=8.0
 r.SkyAtmosphere.FastSkyLUT.Width=96
 r.SkyAtmosphere.FastSkyLUT.Height=50
 r.SkyAtmosphere.TransmittanceLUT.UseSmallFormat=1
+fx.Niagara.QualityLevel=0
 
 [EffectsQuality@1]
 r.SkyAtmosphere.AerialPerspectiveLUT.FastApplyOnOpaque=1 ; Always have FastSkyLUT 1 in this case to avoid wrong sky
@@ -130,6 +131,7 @@ r.SkyAtmosphere.FastSkyLUT.SampleCountMax=8.0
 r.SkyAtmosphere.FastSkyLUT.Width=96
 r.SkyAtmosphere.FastSkyLUT.Height=50
 r.SkyAtmosphere.TransmittanceLUT.UseSmallFormat=1
+fx.Niagara.QualityLevel=0
 
 [EffectsQuality@2]
 r.SkyAtmosphere.AerialPerspectiveLUT.FastApplyOnOpaque=1 ; Always have FastSkyLUT 1 in this case to avoid wrong sky
@@ -141,6 +143,7 @@ r.SkyAtmosphere.FastSkyLUT.SampleCountMax=8.0
 r.SkyAtmosphere.FastSkyLUT.Width=96
 r.SkyAtmosphere.FastSkyLUT.Height=50
 r.SkyAtmosphere.TransmittanceLUT.UseSmallFormat=1
+fx.Niagara.QualityLevel=0
 
 [EffectsQuality@3]
 r.SkyAtmosphere.AerialPerspectiveLUT.FastApplyOnOpaque=1 ; Always have FastSkyLUT 1 in this case to avoid wrong sky
@@ -152,6 +155,7 @@ r.SkyAtmosphere.FastSkyLUT.SampleCountMax=8.0
 r.SkyAtmosphere.FastSkyLUT.Width=96
 r.SkyAtmosphere.FastSkyLUT.Height=50
 r.SkyAtmosphere.TransmittanceLUT.UseSmallFormat=1
+fx.Niagara.QualityLevel=1
 
 ; ------------------------------------------------------------------------------
 ; Texture settings
@@ -175,4 +179,4 @@ BlacklistCVars=r.MaxAnisotropy
 +BlacklistCVars=r.Streaming.Boost
 +BlacklistCVars=r.Streaming.AmortizeCPUToGPUCopy
 +BlacklistCVars=r.Streaming.MaxNumTexturesToStreamPerFrame
-+BlacklistCVars=r.Streaming.PoolSize
\ No newline at end of file
++BlacklistCVars=r.Streaming.PoolSize
diff --git a/Engine/Plugins/FX/Niagara/Source/Niagara/Classes/NiagaraEmitter.h b/Engine/Plugins/FX/Niagara/Source/Niagara/Classes/NiagaraEmitter.h
--- a/Engine/Plugins/FX/Niagara/Source/Niagara/Classes/NiagaraEmitter.h
+++ b/Engine/Plugins/FX/Niagara/Source/Niagara/Classes/NiagaraEmitter.h
@@ -441,7 +441,7 @@ public:
 
 	FORCEINLINE const FNiagaraEmitterScalabilitySettings& GetScalabilitySettings()const { return CurrentScalabilitySettings; }
 
-	/** Returns true if this emitter's platform filter allows it on this platform and EffectsQuality level. */
+	/** Returns true if this emitter's platform filter allows it on this platform and quality level. */
 	NIAGARA_API bool IsAllowedByScalability()const;
 
 	NIAGARA_API bool RequiresPersistentIDs() const;
@@ -503,7 +503,7 @@ public:
 	NIAGARA_API	void Reparent(UNiagaraEmitter& InParent);
 #endif
 
-	void OnEffectsQualityChanged();
+	void OnQualityLevelChanged();
 
 protected:
 	virtual void BeginDestroy() override;
diff --git a/Engine/Plugins/FX/Niagara/Source/Niagara/Classes/NiagaraPlatformSet.h b/Engine/Plugins/FX/Niagara/Source/Niagara/Classes/NiagaraPlatformSet.h
--- a/Engine/Plugins/FX/Niagara/Source/Niagara/Classes/NiagaraPlatformSet.h
+++ b/Engine/Plugins/FX/Niagara/Source/Niagara/Classes/NiagaraPlatformSet.h
@@ -21,26 +21,26 @@ struct FNiagaraDeviceProfileStateEntry
 	GENERATED_BODY();
 
 	FNiagaraDeviceProfileStateEntry()
-		: EffectsQualityMask(0)
-		, SetEffectsQualityMask(0)
+		: QualityLevelMask(0)
+		, SetQualityLevelMask(0)
 	{}
 
 	UPROPERTY(EditAnywhere, Category = Profile)
 	FName ProfileName;
 
-	/** The state of each set effects quality. */
+	/** The state of each set quality level. */
 	UPROPERTY(EditAnywhere, Category = Profile)
-	uint32 EffectsQualityMask;
+	uint32 QualityLevelMask;
 
 	/** Which Effects Qualities are set. */
 	UPROPERTY(EditAnywhere, Category = Profile)
-	uint32 SetEffectsQualityMask;
+	uint32 SetQualityLevelMask;
 
-	FORCEINLINE ENiagaraPlatformSelectionState GetState(int32 EffectsQuality)const;
+	FORCEINLINE ENiagaraPlatformSelectionState GetState(int32 QualityLevel)const;
 
-	FORCEINLINE void SetState(int32 EffectsQuality, ENiagaraPlatformSelectionState State);
+	FORCEINLINE void SetState(int32 QualityLevel, ENiagaraPlatformSelectionState State);
 
-	FORCEINLINE bool AllDefaulted()const { return SetEffectsQualityMask == 0; }
+	FORCEINLINE bool AllDefaulted()const { return SetQualityLevelMask == 0; }
 };
 
 UENUM()
@@ -61,7 +61,7 @@ struct FNiagaraPlatformSetConflictEntry
 
 	FNiagaraPlatformSetConflictEntry()
 		: ProfileName(NAME_None)
-		, EffectsQualityMask(0)
+		, QualityLevelMask(0)
 	{}
 
 	UPROPERTY()
@@ -69,7 +69,7 @@ struct FNiagaraPlatformSetConflictEntry
 
 	/** Mask of conflicting effects qualities for this profile. */
 	UPROPERTY()
-	int32 EffectsQualityMask;
+	int32 QualityLevelMask;
 };
 
 USTRUCT()
@@ -100,36 +100,39 @@ struct NIAGARA_API FNiagaraPlatformSet
 {
 	GENERATED_BODY();
 
-	static FORCEINLINE int32 EQFromMask(int32 EQMask);
-	static FORCEINLINE int32 CreateEQMask(int32 EQ);
-	static FText GetEffectsQualityText(int32 EffectsQuality);
-	static FText GetEffectsQualityMaskText(int32 EffectsQualityMask);
+	static FORCEINLINE int32 QualityLevelFromMask(int32 QLMask);
+	static FORCEINLINE int32 CreateQualityLevelMask(int32 QL);
+	static FText GetQualityLevelText(int32 QualityLevel);
+	static FText GetQualityLevelMaskText(int32 QualityLevelMask);
+
+	static void OnQualityLevelChanged(IConsoleVariable* Variable);
+	static int32 GetQualityLevel();
 public:
 
 	//Runtime public API
 
-	FNiagaraPlatformSet(int32 EQMask);
+	FNiagaraPlatformSet(int32 QLMask);
 	FNiagaraPlatformSet();
 
-	/** Is this set active right now. i.e. enabled for the current device profile and effects quality. */
+	/** Is this set active right now. i.e. enabled for the current device profile and quality level. */
 	bool IsActive()const;
 	
-	/** Is this platform set enabled on any effects quality for the passed device profile. Returns the EffectsQualityMask for all enabled effects qualities for this profile. */
+	/** Is this platform set enabled on any quality level for the passed device profile. Returns the QualityLevelMask for all enabled effects qualities for this profile. */
 	int32 IsEnabledForDeviceProfile(const UDeviceProfile* DeviceProfile)const;
 
-	/** Is this platform set enabled at this effects quality on any device profile. */
-	bool IsEnabledForEffectsQuality(int32 EffectsQuality)const;
+	/** Is this platform set enabled at this quality level on any device profile. */
+	bool IsEnabledForQualityLevel(int32 QualityLevel)const;
 
-	/** Fill OutProfiles with all device proifles that have been overriden at the passed EffectsQuality. */
-	void GetOverridenDeviceProfiles(int32 EffectsQuality, TArray<UDeviceProfile*>& OutEnabledProfiles, TArray<UDeviceProfile*>& OutDisabledProfiles)const;
+	/** Fill OutProfiles with all device profiles that have been overridden at the passed QualityLevel. */
+	void GetOverridenDeviceProfiles(int32 QualityLevel, TArray<UDeviceProfile*>& OutEnabledProfiles, TArray<UDeviceProfile*>& OutDisabledProfiles)const;
 
 	/** Returns true if this set is enabled for any profiles on the specified platform. */
 	bool IsEnabledForPlatform(const FString& PlatformName)const;
 
-	/** Returns true if the current platform can modify it's niagara scalability settings at runtime. */
+	/** Returns true if the current platform can modify it's Niagara scalability settings at runtime. */
 	static bool CanChangeScalabilityAtRuntime();
 
-	/**Will force all platform sets to regenerate their cacehd data next time they are used.*/
+	/**Will force all platform sets to regenerate their cached data next time they are used.*/
 	static void InvalidateCachedData();
 
 	static int32 GetEffectQualityMaskForDeviceProfile(const UDeviceProfile* Profile);
@@ -143,14 +146,14 @@ public:
 	Does the enabled set of profiles in this set conflict with those in another set.
 	In some cases this is not allowed.
 	*/
-	//bool Conflicts(const FNiagaraPlatformSet& Other, TArray<const UDeviceProfile*>& OutConflictingProfiles, bool bIncludeProfilesWithVariableEffectsQuality=false)const;
+	//bool Conflicts(const FNiagaraPlatformSet& Other, TArray<const UDeviceProfile*>& OutConflictingProfiles, bool bIncludeProfilesWithVariableQualityLevel=false)const;
 
 
 	/** Direct state accessors for the UI. */
 	bool IsEffectQualityEnabled(int32 EffectQuality)const;
 	void SetEnabledForEffectQuality(int32 EffectQuality, bool bEnabled);
-	void SetDeviceProfileState(UDeviceProfile* Profile, int32 EffectsQuality, ENiagaraPlatformSelectionState NewState);
-	ENiagaraPlatformSelectionState GetDeviceProfileState(UDeviceProfile* Profile, int32 EffectsQuality)const;
+	void SetDeviceProfileState(UDeviceProfile* Profile, int32 QualityLevel, ENiagaraPlatformSelectionState NewState);
+	ENiagaraPlatformSelectionState GetDeviceProfileState(UDeviceProfile* Profile, int32 QualityLevel)const;
 	
 	/** Invalidates any cached data on this platform set when something has changed. */
 	void OnChanged();
@@ -167,13 +170,15 @@ public:
 
 	/** Mask defining which effects qualities this set matches. */
 	UPROPERTY(EditAnywhere, Category = Platforms)
-	int32 EffectsQualityMask;
+	int32 QualityLevelMask;
 
 	/** States of specific device profiles we've set. */
 	UPROPERTY(EditAnywhere, Category = Platforms)
 	TArray<FNiagaraDeviceProfileStateEntry> DeviceProfileStates;
 private:
 
+	static int32 CachedQualityLevel;
+
 	//Set from outside when we need to force all cached values to be regenerated. For example on CVar changes.
 	static uint32 LastDirtiedFrame;
 
@@ -182,28 +187,30 @@ private:
 
 	mutable bool bEnabledForCurrentProfileAndEffectQuality;
 
-	bool IsEnabled(const UDeviceProfile* Profile, int32 EffectsQuality)const;
+	bool IsEnabled(const UDeviceProfile* Profile, int32 QualityLevel)const;
 
 #if WITH_EDITOR
 	//Data we pull from platform ini files.
 	struct FPlatformIniSettings
 	{
 		FPlatformIniSettings()
-			:bCanChangeEffectsQualityAtRuntime(0), EffectsQualityMask(INDEX_NONE)
+			:bCanChangeQualityLevelAtRuntime(0), bPruneEmittersOnCook(false), EffectsQuality(0)
 		{}
-		FPlatformIniSettings(int32 InbCanChangeEffectsQualityAtRuntime, int32 InbPruneEmittersOnCook, int32 InEffectsQualityMask)
-			:bCanChangeEffectsQualityAtRuntime(InbCanChangeEffectsQualityAtRuntime), bPruneEmittersOnCook(InbPruneEmittersOnCook), EffectsQualityMask(InEffectsQualityMask)
+		FPlatformIniSettings(int32 InbCanChangeQualityLevelAtRuntime, int32 InbPruneEmittersOnCook, int32 InEffectsQuality)
+			:bCanChangeQualityLevelAtRuntime(InbCanChangeQualityLevelAtRuntime), bPruneEmittersOnCook(InbPruneEmittersOnCook), EffectsQuality(InEffectsQuality)
 		{}
-		int32 bCanChangeEffectsQualityAtRuntime;
+		int32 bCanChangeQualityLevelAtRuntime;
 		int32 bPruneEmittersOnCook;
-		int32 EffectsQualityMask;
+		int32 EffectsQuality;
+
+		TArray<int32> QualityLevelsPerEffectsQuality;
 	};
 
 	//Cached data read from platform ini files.
 	static TMap<FName, FPlatformIniSettings> CachedPlatformIniSettings;
 
-	//Cached final EffectsQuality setting for each device profile.
-	static TMap<const UDeviceProfile*, int32> CachedEQMasksPerDeviceProfile;
+	//Cached final QualityLevel setting for each device profile.
+	static TMap<const UDeviceProfile*, int32> CachedQLMasksPerDeviceProfile;
 
 	static FPlatformIniSettings GetPlatformIniSettings(const FString& PlatformName);
 
@@ -212,51 +219,53 @@ private:
 };
 
 
-FORCEINLINE int32 FNiagaraPlatformSet::EQFromMask(int32 EQMask)
+FORCEINLINE int32 FNiagaraPlatformSet::QualityLevelFromMask(int32 QLMask)
 {
 	//Iterate through bits in mask to find the first set quality level.
-	int32 EffectsQuality = 0;
-	while (EQMask != 0)
+	int32 QualityLevel = 0;
+	while (QLMask != 0)
 	{
-		if (EQMask & 0x1)
+		if (QLMask & 0x1)
 		{
-			return EffectsQuality;
+			return QualityLevel;
 		}
-		++EffectsQuality;
-		EQMask >>= 1;
+		++QualityLevel;
+		QLMask >>= 1;
 	}
 
 	return INDEX_NONE;
 }
 
-FORCEINLINE int32 FNiagaraPlatformSet::CreateEQMask(int32 EQ)
+FORCEINLINE int32 FNiagaraPlatformSet::CreateQualityLevelMask(int32 QL)
 {
-	return EQ == INDEX_NONE ? INDEX_NONE : (1 << EQ);
+	return QL == INDEX_NONE ? INDEX_NONE : (1 << QL);
 }
 
 
-FORCEINLINE ENiagaraPlatformSelectionState FNiagaraDeviceProfileStateEntry::GetState(int32 EffectsQuality)const
+FORCEINLINE ENiagaraPlatformSelectionState FNiagaraDeviceProfileStateEntry::GetState(int32 QualityLevel)const
 {
-	int32 EQMask = FNiagaraPlatformSet::CreateEQMask(EffectsQuality);
-	return (EQMask & SetEffectsQualityMask) == 0 ? ENiagaraPlatformSelectionState::Default : ((EQMask & EffectsQualityMask) != 0 ? ENiagaraPlatformSelectionState::Enabled : ENiagaraPlatformSelectionState::Disabled);
+	int32 QLMask = FNiagaraPlatformSet::CreateQualityLevelMask(QualityLevel);
+	return (QLMask & SetQualityLevelMask) == 0 ? ENiagaraPlatformSelectionState::Default : ((QLMask & QualityLevelMask) != 0 ? ENiagaraPlatformSelectionState::Enabled : ENiagaraPlatformSelectionState::Disabled);
 }
 
-FORCEINLINE void FNiagaraDeviceProfileStateEntry::SetState(int32 EffectsQuality, ENiagaraPlatformSelectionState State)
+FORCEINLINE void FNiagaraDeviceProfileStateEntry::SetState(int32 QualityLevel, ENiagaraPlatformSelectionState State)
 {
-	int32 EQMask = FNiagaraPlatformSet::CreateEQMask(EffectsQuality);
+	int32 QLMask = FNiagaraPlatformSet::CreateQualityLevelMask(QualityLevel);
 	if (State == ENiagaraPlatformSelectionState::Default)
 	{
-		SetEffectsQualityMask &= ~EQMask;
-		EffectsQualityMask &= ~EQMask;
+		SetQualityLevelMask &= ~QLMask;
+		QualityLevelMask &= ~QLMask;
 	}
 	else if (State == ENiagaraPlatformSelectionState::Enabled)
 	{
-		SetEffectsQualityMask |= EQMask;
-		EffectsQualityMask |= EQMask;
+		SetQualityLevelMask |= QLMask;
+		QualityLevelMask |= QLMask;
 	}
 	else //(State == ENiagaraPlatformSelectionState::Disabled)
 	{
-		SetEffectsQualityMask |= EQMask;
-		EffectsQualityMask &= ~EQMask;
+		SetQualityLevelMask |= QLMask;
+		QualityLevelMask &= ~QLMask;
 	}
-}
\ No newline at end of file
+}
+
+extern int32 GNiagaraQualityLevel;
\ No newline at end of file
diff --git a/Engine/Plugins/FX/Niagara/Source/Niagara/Classes/NiagaraSystem.h b/Engine/Plugins/FX/Niagara/Source/Niagara/Classes/NiagaraSystem.h
--- a/Engine/Plugins/FX/Niagara/Source/Niagara/Classes/NiagaraSystem.h
+++ b/Engine/Plugins/FX/Niagara/Source/Niagara/Classes/NiagaraSystem.h
@@ -377,7 +377,7 @@ public:
 	UNiagaraEffectType* GetEffectType()const;
 	FORCEINLINE const FNiagaraSystemScalabilitySettings& GetScalabilitySettings() { return CurrentScalabilitySettings; }
 	
-	void OnEffectsQualityChanged();
+	void OnQualityLevelChanged();
 
 	/** Whether or not fixed bounds are enabled. */
 	UPROPERTY(EditAnywhere, Category = "System", meta = (InlineEditConditionToggle))
diff --git a/Engine/Plugins/FX/Niagara/Source/Niagara/Private/NiagaraCommon.cpp b/Engine/Plugins/FX/Niagara/Source/Niagara/Private/NiagaraCommon.cpp
--- a/Engine/Plugins/FX/Niagara/Source/Niagara/Private/NiagaraCommon.cpp
+++ b/Engine/Plugins/FX/Niagara/Source/Niagara/Private/NiagaraCommon.cpp
@@ -139,7 +139,14 @@ void FNiagaraSystemUpdateContext::AddAll(bool bReInit)
 	{
 		UNiagaraComponent* Comp = *It;
 		check(Comp);
+
 		bool bIsActive = Comp->IsActive() || Comp->IsRegisteredWithScalabilityManager();
+
+		if (bDestroyOnAdd)
+		{
+			Comp->DeactivateImmediate();
+		}
+
 		if (bIsActive || bOnlyActive == false)
 		{
 			AddInternal(Comp, bReInit);
@@ -156,6 +163,12 @@ void FNiagaraSystemUpdateContext::Add(const UNiagaraSystem* System, bool bReInit
 		if (Comp->GetAsset() == System)
 		{
 			bool bIsActive = Comp->IsActive() || Comp->IsRegisteredWithScalabilityManager();
+
+			if (bDestroyOnAdd)
+			{
+				Comp->DeactivateImmediate();
+			}
+
 			if (bIsActive || bOnlyActive == false)
 			{
 				AddInternal(Comp, bReInit);
@@ -175,6 +188,12 @@ void FNiagaraSystemUpdateContext::Add(const UNiagaraEmitter* Emitter, bool bReIn
 		if (SystemInst && SystemInst->UsesEmitter(Emitter))
 		{
 			bool bIsActive = Comp->IsActive() || Comp->IsRegisteredWithScalabilityManager();
+
+			if (bDestroyOnAdd)
+			{
+				Comp->DeactivateImmediate();
+			}
+
 			if (bIsActive || bOnlyActive == false)
 			{
 				AddInternal(Comp, bReInit);
@@ -193,6 +212,12 @@ void FNiagaraSystemUpdateContext::Add(const UNiagaraScript* Script, bool bReInit
 		if (System && System->UsesScript(Script))
 		{
 			bool bIsActive = Comp->IsActive() || Comp->IsRegisteredWithScalabilityManager();
+
+			if (bDestroyOnAdd)
+			{
+				Comp->DeactivateImmediate();
+			}
+
 			if (bIsActive || bOnlyActive == false)
 			{
 				AddInternal(Comp, bReInit);
@@ -227,6 +252,12 @@ void FNiagaraSystemUpdateContext::Add(const UNiagaraParameterCollection* Collect
 		if (SystemInst && SystemInst->UsesCollection(Collection))
 		{
 			bool bIsActive = Comp->IsActive() || Comp->IsRegisteredWithScalabilityManager();
+
+			if (bDestroyOnAdd)
+			{
+				Comp->DeactivateImmediate();
+			}
+
 			if (bIsActive || bOnlyActive == false)
 			{
 				AddInternal(Comp, bReInit);
@@ -238,11 +269,6 @@ void FNiagaraSystemUpdateContext::Add(const UNiagaraParameterCollection* Collect
 
 void FNiagaraSystemUpdateContext::AddInternal(UNiagaraComponent* Comp, bool bReInit)
 {
-	if (bDestroyOnAdd)
-	{
-		Comp->DeactivateImmediate();
-	}
-
 	if (bReInit)
 	{
 		ComponentsToReInit.AddUnique(Comp);
diff --git a/Engine/Plugins/FX/Niagara/Source/Niagara/Private/NiagaraEffectType.cpp b/Engine/Plugins/FX/Niagara/Source/Niagara/Private/NiagaraEffectType.cpp
--- a/Engine/Plugins/FX/Niagara/Source/Niagara/Private/NiagaraEffectType.cpp
+++ b/Engine/Plugins/FX/Niagara/Source/Niagara/Private/NiagaraEffectType.cpp
@@ -58,7 +58,7 @@ void UNiagaraEffectType::PostLoad()
 			FNiagaraSystemScalabilitySettings& DLSettings = DetailLevelScalabilitySettings_DEPRECATED[DL];
 			FNiagaraSystemScalabilitySettings& NewSettigns = SystemScalabilitySettings.Settings.AddDefaulted_GetRef();
 			NewSettigns = DLSettings;
-			NewSettigns.Platforms = FNiagaraPlatformSet(FNiagaraPlatformSet::CreateEQMask(DL));
+			NewSettigns.Platforms = FNiagaraPlatformSet(FNiagaraPlatformSet::CreateQualityLevelMask(DL));
 		}
 
 		//Ensure all types replicate the former global spawn count scales
@@ -70,7 +70,7 @@ void UNiagaraEffectType::PostLoad()
 				FNiagaraEmitterScalabilitySettings& EmitterSettings = EmitterScalabilitySettings.Settings.AddDefaulted_GetRef();
 				EmitterSettings.bScaleSpawnCount = true;
 				EmitterSettings.SpawnCountScale = LegacySpawnCountScales[DL];
-				EmitterSettings.Platforms = FNiagaraPlatformSet(FNiagaraPlatformSet::CreateEQMask(DL));
+				EmitterSettings.Platforms = FNiagaraPlatformSet(FNiagaraPlatformSet::CreateQualityLevelMask(DL));
 			}
 		}
 	}
@@ -120,7 +120,7 @@ void UNiagaraEffectType::PostEditChangeProperty(struct FPropertyChangedEvent& Pr
 		UNiagaraSystem* System = *It;
 		if (System->GetEffectType() == this)
 		{
-			System->OnEffectsQualityChanged();
+			System->OnQualityLevelChanged();
 			UpdateContext.Add(System, true);
 		}
 	}
diff --git a/Engine/Plugins/FX/Niagara/Source/Niagara/Private/NiagaraEmitter.cpp b/Engine/Plugins/FX/Niagara/Source/Niagara/Private/NiagaraEmitter.cpp
--- a/Engine/Plugins/FX/Niagara/Source/Niagara/Private/NiagaraEmitter.cpp
+++ b/Engine/Plugins/FX/Niagara/Source/Niagara/Private/NiagaraEmitter.cpp
@@ -328,44 +328,44 @@ void UNiagaraEmitter::PostLoad()
 	{
 		int32 MinDetailLevel = bUseMaxDetailLevel_DEPRECATED ? MinDetailLevel_DEPRECATED : 0;
 		int32 MaxDetailLevel = bUseMaxDetailLevel_DEPRECATED ? MaxDetailLevel_DEPRECATED : 4;
-		int32 NewEQMask = 0;
-		//Currently all detail levels were direct mappings to effects quality so just transfer them over to the new mask in PlatformSet.
-		for (int32 EQ = MinDetailLevel; EQ <= MaxDetailLevel; ++EQ)
+		int32 NewQLMask = 0;
+		//Currently all detail levels were direct mappings to quality level so just transfer them over to the new mask in PlatformSet.
+		for (int32 QL = MinDetailLevel; QL <= MaxDetailLevel; ++QL)
 		{
-			NewEQMask |= (1 << EQ);
+			NewQLMask |= (1 << QL);
 		}
 
-		Platforms = FNiagaraPlatformSet(NewEQMask);
+		Platforms = FNiagaraPlatformSet(NewQLMask);
 
 		//Transfer spawn rate scaling overrides
 		if (bOverrideGlobalSpawnCountScale_DEPRECATED)
 		{
 			FNiagaraEmitterScalabilityOverride& LowOverride = ScalabilityOverrides.Overrides.AddDefaulted_GetRef();
-			LowOverride.Platforms = FNiagaraPlatformSet(FNiagaraPlatformSet::CreateEQMask(0));
+			LowOverride.Platforms = FNiagaraPlatformSet(FNiagaraPlatformSet::CreateQualityLevelMask(0));
 			LowOverride.bOverrideSpawnCountScale = true;
 			LowOverride.bScaleSpawnCount = true;
 			LowOverride.SpawnCountScale = GlobalSpawnCountScaleOverrides_DEPRECATED.Low;
 
 			FNiagaraEmitterScalabilityOverride& MediumOverride = ScalabilityOverrides.Overrides.AddDefaulted_GetRef();
-			MediumOverride.Platforms = FNiagaraPlatformSet(FNiagaraPlatformSet::CreateEQMask(1));
+			MediumOverride.Platforms = FNiagaraPlatformSet(FNiagaraPlatformSet::CreateQualityLevelMask(1));
 			MediumOverride.bOverrideSpawnCountScale = true;
 			MediumOverride.bScaleSpawnCount = true;
 			MediumOverride.SpawnCountScale = GlobalSpawnCountScaleOverrides_DEPRECATED.Medium;
 
 			FNiagaraEmitterScalabilityOverride& HighOverride = ScalabilityOverrides.Overrides.AddDefaulted_GetRef();
-			HighOverride.Platforms = FNiagaraPlatformSet(FNiagaraPlatformSet::CreateEQMask(2));
+			HighOverride.Platforms = FNiagaraPlatformSet(FNiagaraPlatformSet::CreateQualityLevelMask(2));
 			HighOverride.bOverrideSpawnCountScale = true;
 			HighOverride.bScaleSpawnCount = true;
 			HighOverride.SpawnCountScale = GlobalSpawnCountScaleOverrides_DEPRECATED.High;
 
 			FNiagaraEmitterScalabilityOverride& EpicOverride = ScalabilityOverrides.Overrides.AddDefaulted_GetRef();
-			EpicOverride.Platforms = FNiagaraPlatformSet(FNiagaraPlatformSet::CreateEQMask(3));
+			EpicOverride.Platforms = FNiagaraPlatformSet(FNiagaraPlatformSet::CreateQualityLevelMask(3));
 			EpicOverride.bOverrideSpawnCountScale = true;
 			EpicOverride.bScaleSpawnCount = true;
 			EpicOverride.SpawnCountScale = GlobalSpawnCountScaleOverrides_DEPRECATED.Epic;
 
 			FNiagaraEmitterScalabilityOverride& CineOverride = ScalabilityOverrides.Overrides.AddDefaulted_GetRef();
-			CineOverride.Platforms = FNiagaraPlatformSet(FNiagaraPlatformSet::CreateEQMask(4));
+			CineOverride.Platforms = FNiagaraPlatformSet(FNiagaraPlatformSet::CreateQualityLevelMask(4));
 			CineOverride.bOverrideSpawnCountScale = true;
 			CineOverride.bScaleSpawnCount = true;
 			CineOverride.SpawnCountScale = GlobalSpawnCountScaleOverrides_DEPRECATED.Cine;
@@ -1651,7 +1651,7 @@ void UNiagaraEmitter::ResolveScalabilitySettings()
 	}
 }
 
-void UNiagaraEmitter::OnEffectsQualityChanged()
+void UNiagaraEmitter::OnQualityLevelChanged()
 {
 	ResolveScalabilitySettings();
 }
diff --git a/Engine/Plugins/FX/Niagara/Source/Niagara/Private/NiagaraModule.cpp b/Engine/Plugins/FX/Niagara/Source/Niagara/Private/NiagaraModule.cpp
--- a/Engine/Plugins/FX/Niagara/Source/Niagara/Private/NiagaraModule.cpp
+++ b/Engine/Plugins/FX/Niagara/Source/Niagara/Private/NiagaraModule.cpp
@@ -40,7 +40,6 @@ static FAutoConsoleVariableRef CVarLogCompileIdGeneration(
 
 float INiagaraModule::EngineGlobalSpawnCountScale = 1.0f;
 float INiagaraModule::EngineGlobalSystemCountScale = 1.0f;
-int32 INiagaraModule::EngineEffectsQuality = INDEX_NONE;
 
 int32 GEnableVerboseNiagaraChangeIdLogging = 0;
 static FAutoConsoleVariableRef CVarEnableVerboseNiagaraChangeIdLogging(
@@ -150,11 +149,6 @@ void INiagaraModule::StartupModule()
 
 	FNiagaraWorldManager::OnStartup();
 
-	//Force update on module load.
-	EngineEffectsQuality = Scalability::DefaultQualityLevel;
-	static const auto CVarEQ = IConsoleManager::Get().FindTConsoleVariableDataInt(TEXT("sg.EffectsQuality"));
-	OnEffectsQualityChanged(CVarEQ->GetValueOnGameThread(), true);
-
 #if WITH_EDITOR	
 	// Loading uncooked data in a game environment, we still need to get some functionality from the NiagaraEditor module.
 	// This includes the ability to compile scripts and load WITH_EDITOR_ONLY data.
@@ -395,26 +389,6 @@ void INiagaraModule::UnregisterPrecompiler(FDelegateHandle DelegateHandle)
 
 #endif
 
-void INiagaraModule::OnEffectsQualityChanged(int32 NewEffectsQuality, bool bForce)
-{
-	if (FNiagaraPlatformSet::CanChangeScalabilityAtRuntime() || bForce)
-	{
-		if (NewEffectsQuality != EngineEffectsQuality)
-		{
-			EngineEffectsQuality = NewEffectsQuality;
-
-			FNiagaraPlatformSet::InvalidateCachedData();
-
-			for (TObjectIterator<UNiagaraSystem> It; It; ++It)
-			{
-				UNiagaraSystem* System = *It;
-				check(System);
-				System->OnEffectsQualityChanged();
-			}
-		}
-	}
-}
-
 //////////////////////////////////////////////////////////////////////////
 
 UScriptStruct* FNiagaraTypeDefinition::ParameterMapStruct;
@@ -978,15 +952,6 @@ void INiagaraModule::ProcessShaderCompilationQueue()
 	return OnProcessQueue.Execute();
 }
 
-
-void CVarSinkFunc()
-{
-	static const auto CVarEQ = IConsoleManager::Get().FindTConsoleVariableDataInt(TEXT("sg.EffectsQuality"));
-	INiagaraModule::OnEffectsQualityChanged(CVarEQ->GetValueOnGameThread(), false);
-}
-
-static FAutoConsoleVariableSink CVarSink(FConsoleCommandDelegate::CreateStatic(&CVarSinkFunc));
-
 #if WITH_EDITOR
 const TArray<FNiagaraVariable>& FNiagaraGlobalParameters::GetVariables()
 {
diff --git a/Engine/Plugins/FX/Niagara/Source/Niagara/Private/NiagaraPlatformSet.cpp b/Engine/Plugins/FX/Niagara/Source/Niagara/Private/NiagaraPlatformSet.cpp
--- a/Engine/Plugins/FX/Niagara/Source/Niagara/Private/NiagaraPlatformSet.cpp
+++ b/Engine/Plugins/FX/Niagara/Source/Niagara/Private/NiagaraPlatformSet.cpp
@@ -1,12 +1,13 @@
 // Copyright Epic Games, Inc. All Rights Reserved.
 
 #include "NiagaraPlatformSet.h"
-#include "NiagaraModule.h"
 #include "DeviceProfiles/DeviceProfile.h"
 #include "DeviceProfiles/DeviceProfileManager.h"
 #include "Scalability.h"
 #include "Misc/ConfigCacheIni.h"
 #include "Interfaces/ITargetPlatform.h"
+#include "NiagaraSystem.h"
+#include "NiagaraSettings.h"
 
 #if WITH_EDITOR
 #include "PlatformInfo.h"
@@ -23,8 +24,8 @@ int32 GbNiagaraAllowRuntimeScalabilityChanges = 0;
 static FAutoConsoleVariableRef CVarNiagaraAllowRuntimeScalabilityChanges(
 	CanChangeEQCVarName,
 	GbNiagaraAllowRuntimeScalabilityChanges,
-	TEXT("If > 0 this platform allows scalability level changes at runtime. \n"),
-	ECVF_ReadOnly
+	TEXT("If > 0 this platform allows niagara scalability level changes at runtime. \n"),
+	ECVF_Scalability
 );
 
 const TCHAR* PruneEmittersOnCookName = TEXT("fx.Niagara.PruneEmittersOnCook");
@@ -33,40 +34,89 @@ static FAutoConsoleVariableRef CVarPruneEmittersOnCook(
 	PruneEmittersOnCookName,
 	GbPruneEmittersOnCook,
 	TEXT("If > 0 this platform will prune disabled emitters during cook. \n"),
-	ECVF_ReadOnly
+	ECVF_Scalability
 );
 
+static const int32 DefaultQualityLevel = 3;
+const TCHAR* NiagaraQualityLevelName = TEXT("fx.Niagara.QualityLevel");
+int32 GNiagaraQualityLevel = DefaultQualityLevel;
+static FAutoConsoleVariableRef CVarNiagaraQualityLevel(
+	NiagaraQualityLevelName,
+	GNiagaraQualityLevel,
+	TEXT("The quality level for Niagara Effects. \n"),
+	FConsoleVariableDelegate::CreateStatic(&FNiagaraPlatformSet::OnQualityLevelChanged),
+	ECVF_Scalability
+);
+
+int32 FNiagaraPlatformSet::CachedQualityLevel = INDEX_NONE;
+
+int32 FNiagaraPlatformSet::GetQualityLevel()
+{
+	if (CachedQualityLevel == INDEX_NONE)
+	{
+		CachedQualityLevel = GNiagaraQualityLevel;
+	}
+	return CachedQualityLevel;
+}
+
+void FNiagaraPlatformSet::OnQualityLevelChanged(IConsoleVariable* Variable)
+{
+	int32 NewQualityLevel = Variable->GetInt();
+	int32 CurrentLevel = GetQualityLevel();
+
+	if (CurrentLevel != NewQualityLevel)
+	{
+		CachedQualityLevel = NewQualityLevel;
+		InvalidateCachedData();
+
+		for (TObjectIterator<UNiagaraSystem> It; It; ++It)
+		{
+			UNiagaraSystem* System = *It;
+			check(System);
+			System->OnQualityLevelChanged();
+		}
+	}
+}
+
 uint32 FNiagaraPlatformSet::LastDirtiedFrame = 0;
 #if WITH_EDITOR
-TMap<const UDeviceProfile*, int32> FNiagaraPlatformSet::CachedEQMasksPerDeviceProfile;
+TMap<const UDeviceProfile*, int32> FNiagaraPlatformSet::CachedQLMasksPerDeviceProfile;
 TMap<FName, FNiagaraPlatformSet::FPlatformIniSettings> FNiagaraPlatformSet::CachedPlatformIniSettings;
 #endif
 
-FText FNiagaraPlatformSet::GetEffectsQualityText(int32 EffectsQuality)
+FText FNiagaraPlatformSet::GetQualityLevelText(int32 QualityLevel)
 {
-	Scalability::FQualityLevels Counts = Scalability::GetQualityLevelCounts();
-	return Scalability::GetQualityLevelText(EffectsQuality, Counts.EffectsQuality);
+	const UNiagaraSettings* Settings = GetDefault<UNiagaraSettings>();
+	check(Settings);
+
+	if (Settings->QualityLevels.IsValidIndex(QualityLevel))
+	{
+		return Settings->QualityLevels[QualityLevel];
+	}	
+	else
+	{
+		return FText::AsNumber(QualityLevel);
+	}
 }
 
-FText FNiagaraPlatformSet::GetEffectsQualityMaskText(int32 EffectsQualityMask)
+FText FNiagaraPlatformSet::GetQualityLevelMaskText(int32 QualityLevelMask)
 {
-	if (EffectsQualityMask == INDEX_NONE)
+	if (QualityLevelMask == INDEX_NONE)
 	{
-		return LOCTEXT("EffectsQualityAll", "All");
+		return LOCTEXT("QualityLevelAll", "All");
 	}
-	else if (EffectsQualityMask == 0)
+	else if (QualityLevelMask == 0)
 	{
-		return LOCTEXT("EffectsQualityNone", "None");
+		return LOCTEXT("QualityLevelNone", "None");
 	}
 	else
 	{
-		Scalability::FQualityLevels Counts = Scalability::GetQualityLevelCounts();
-		return Scalability::GetQualityLevelText(EQFromMask(EffectsQualityMask), Counts.EffectsQuality);
+		return GetQualityLevelText(QualityLevelFromMask(QualityLevelMask));
 	}
 }
 
-FNiagaraPlatformSet::FNiagaraPlatformSet(int32 EQMask)
-	: EffectsQualityMask(EQMask)
+FNiagaraPlatformSet::FNiagaraPlatformSet(int32 QLMask)
+	: QualityLevelMask(QLMask)
 	, LastBuiltFrame(0)
 	, bEnabledForCurrentProfileAndEffectQuality(false)
 {
@@ -74,7 +124,7 @@ FNiagaraPlatformSet::FNiagaraPlatformSet(int32 EQMask)
 }
 
 FNiagaraPlatformSet::FNiagaraPlatformSet()
-	: EffectsQualityMask(INDEX_NONE)
+	: QualityLevelMask(INDEX_NONE)
 	, LastBuiltFrame(0)
 	, bEnabledForCurrentProfileAndEffectQuality(false)
 {
@@ -84,7 +134,7 @@ bool FNiagaraPlatformSet::IsActive()const
 {
 	if (LastBuiltFrame <= LastDirtiedFrame || LastBuiltFrame == 0)
 	{
-		bEnabledForCurrentProfileAndEffectQuality = IsEnabled(UDeviceProfileManager::Get().GetActiveProfile(), INiagaraModule::GetEffectsQuality());
+		bEnabledForCurrentProfileAndEffectQuality = IsEnabled(UDeviceProfileManager::Get().GetActiveProfile(), GetQualityLevel());
 		LastBuiltFrame = GFrameNumber;
 	}
 	return bEnabledForCurrentProfileAndEffectQuality;
@@ -92,26 +142,28 @@ bool FNiagaraPlatformSet::IsActive()const
 
 int32 FNiagaraPlatformSet::IsEnabledForDeviceProfile(const UDeviceProfile* DeviceProfile)const
 {
-	Scalability::FQualityLevels NumLevels = Scalability::GetQualityLevelCounts();
-	int32 RetEQMask = 0;
-	for (int32 i = 0; i < NumLevels.EffectsQuality; ++i)
+	const UNiagaraSettings* Settings = GetDefault<UNiagaraSettings>();
+	check(Settings);
+
+	int32 RetQLMask = 0;
+	for (int32 i = 0; i < Settings->QualityLevels.Num(); ++i)
 	{
 		if (IsEnabled(DeviceProfile, i))
 		{
-			RetEQMask |= CreateEQMask(i);
+			RetQLMask |= CreateQualityLevelMask(i);
 		}
 	}
 
-	return RetEQMask;
+	return RetQLMask;
 }
 
-bool FNiagaraPlatformSet::IsEnabledForEffectsQuality(int32 EffectsQuality)const
+bool FNiagaraPlatformSet::IsEnabledForQualityLevel(int32 QualityLevel)const
 {
 	for (UObject* DeviceProfileObj : UDeviceProfileManager::Get().Profiles)
 	{
 		if (UDeviceProfile* Profile = Cast<UDeviceProfile>(DeviceProfileObj))
 		{
-			if (IsEnabled(Profile, EffectsQuality))
+			if (IsEnabled(Profile, QualityLevel))
 			{
 				return true;
 			}
@@ -121,20 +173,20 @@ bool FNiagaraPlatformSet::IsEnabledForEffectsQuality(int32 EffectsQuality)const
 	return false;
 }
 
-void FNiagaraPlatformSet::GetOverridenDeviceProfiles(int32 EffectsQuality, TArray<UDeviceProfile*>& OutEnabledProfiles, TArray<UDeviceProfile*>& OutDisabledProfiles)const
+void FNiagaraPlatformSet::GetOverridenDeviceProfiles(int32 QualityLevel, TArray<UDeviceProfile*>& OutEnabledProfiles, TArray<UDeviceProfile*>& OutDisabledProfiles)const
 {
-	int32 EQMask = CreateEQMask(EffectsQuality);
+	int32 QLMask = CreateQualityLevelMask(QualityLevel);
 	for (const FNiagaraDeviceProfileStateEntry& Entry : DeviceProfileStates)
 	{
 		if (UObject** DeviceProfile = UDeviceProfileManager::Get().Profiles.FindByPredicate([&](UObject* CheckProfile) {return CheckProfile->GetFName() == Entry.ProfileName;}))
 		{
 			UDeviceProfile* Profile = CastChecked<UDeviceProfile>(*DeviceProfile);
 			//If this platform cannot change at runtime then we store all EQs in the state so that the device is still overridden if someone changes it's EQ CVar.
-			//So here we must also check that this EffectsQuality is the right one for the platforms current setting.
-			int32 ProfileEQMask = GetEffectQualityMaskForDeviceProfile(Profile);
-			if (ProfileEQMask == INDEX_NONE || (EQMask & ProfileEQMask) != 0)
+			//So here we must also check that this QualityLevel is the right one for the platforms current setting.
+			int32 ProfileQLMask = GetEffectQualityMaskForDeviceProfile(Profile);
+			if (ProfileQLMask == INDEX_NONE || (QLMask & ProfileQLMask) != 0)
 			{
-				ENiagaraPlatformSelectionState State = Entry.GetState(EffectsQuality);
+				ENiagaraPlatformSelectionState State = Entry.GetState(QualityLevel);
 				if (State == ENiagaraPlatformSelectionState::Enabled)
 				{
 					OutEnabledProfiles.Add(Profile);
@@ -160,35 +212,45 @@ int32 FNiagaraPlatformSet::GetEffectQualityMaskForDeviceProfile(const UDevicePro
 
 	//When in the editor we may be asking for the EQ of a platform other than the current one.
 	//So we have to look in ini files and device profiles to find it.
-
-	if (int32* CachedEQMask = CachedEQMasksPerDeviceProfile.Find(Profile))
+	if (int32* CachedQLMask = CachedQLMasksPerDeviceProfile.Find(Profile))
 	{
 		//We've seen this Profile before so return the cached value;
-		return *CachedEQMask;
+		return *CachedQLMask;
 	}
 
-	int32 EQMask = INDEX_NONE;
+	int32 QLMask = INDEX_NONE;
 	FPlatformIniSettings PlatformSettings = GetPlatformIniSettings(Profile->DeviceType);
 	
-	if (PlatformSettings.bCanChangeEffectsQualityAtRuntime)
+	if (PlatformSettings.bCanChangeQualityLevelAtRuntime)
 	{
-		EQMask = INDEX_NONE;
+		QLMask = INDEX_NONE;
 	}
 	else
 	{
-		//Default to the platform EQ but allow device profiles to override.
-		EQMask = PlatformSettings.EffectsQualityMask;
-
-		//Check if the DPs set EQ directly.
-		int32 EffectsQuality = INDEX_NONE;
-		if (Profile->GetConsolidatedCVarValue(TEXT("sg.EffectsQuality"), EffectsQuality))
+		//Check if the DPs set Niagara quality directly.
+		int32 QualityLevel = INDEX_NONE;
+		if (!Profile->GetConsolidatedCVarValue(NiagaraQualityLevelName, QualityLevel))
 		{
-			EQMask = CreateEQMask(EffectsQuality);
+			QualityLevel = DefaultQualityLevel;
+			//If not, grab it from the effects quality setting.
+			int32 EffectsQuality = INDEX_NONE;
+			//See if this profile overrides effects quality.
+			if (!Profile->GetConsolidatedCVarValue(TEXT("sg.EffectsQuality"), EffectsQuality))
+			{
+				EffectsQuality = PlatformSettings.EffectsQuality;
+			}
+			check(EffectsQuality != INDEX_NONE);
+
+			QualityLevel = PlatformSettings.QualityLevelsPerEffectsQuality[EffectsQuality];
 		}
+		check(QualityLevel != INDEX_NONE);
+		QLMask = CreateQualityLevelMask(QualityLevel);
 	}
 
-	CachedEQMasksPerDeviceProfile.Add(Profile) = EQMask;
-	return EQMask;
+	//UE_LOG(LogNiagara, Warning, TEXT("%s - DeviceProfile(%d)"), *Profile->GetName(), QLMask);
+
+	CachedQLMasksPerDeviceProfile.Add(Profile) = QLMask;
+	return QLMask;
 
 #else
 
@@ -196,18 +258,18 @@ int32 FNiagaraPlatformSet::GetEffectQualityMaskForDeviceProfile(const UDevicePro
 	check(Profile == UDeviceProfileManager::Get().GetActiveProfile());
 	bool bCanChangeEQAtRuntime = CanChangeScalabilityAtRuntime();
 
-	return CreateEQMask(bCanChangeEQAtRuntime ? INDEX_NONE : INiagaraModule::GetEffectsQuality());
+	return CreateQualityLevelMask(bCanChangeEQAtRuntime ? INDEX_NONE : GetQualityLevel());
 
 #endif
 }
 
-bool FNiagaraPlatformSet::IsEnabled(const UDeviceProfile* Profile, int32 EffectsQuality)const
+bool FNiagaraPlatformSet::IsEnabled(const UDeviceProfile* Profile, int32 QualityLevel)const
 {
 	checkSlow(Profile);
 
-	int32 TestEQMask = CreateEQMask(EffectsQuality);
-	int32 ProfileEQMask = GetEffectQualityMaskForDeviceProfile(Profile);
-	bool bIsEnabledByEQ = (EffectsQualityMask & ProfileEQMask & TestEQMask) != 0;
+	int32 TestQLMask = CreateQualityLevelMask(QualityLevel);
+	int32 ProfileQLMask = GetEffectQualityMaskForDeviceProfile(Profile);
+	bool bIsEnabledByEQ = (QualityLevelMask & ProfileQLMask & TestQLMask) != 0;
 
 	if (DeviceProfileStates.Num() > 0)
 	{
@@ -217,7 +279,7 @@ bool FNiagaraPlatformSet::IsEnabled(const UDeviceProfile* Profile, int32 Effects
 		{
 			if (const FNiagaraDeviceProfileStateEntry* StateEntry = DeviceProfileStates.FindByPredicate([&](const FNiagaraDeviceProfileStateEntry& ProfileState) {return ProfileState.ProfileName == CurrProfile->GetFName(); }))
 			{
-				ENiagaraPlatformSelectionState SelectionState = StateEntry->GetState(EffectsQuality);
+				ENiagaraPlatformSelectionState SelectionState = StateEntry->GetState(QualityLevel);
 				if (SelectionState != ENiagaraPlatformSelectionState::Default)
 				{
 					return SelectionState == ENiagaraPlatformSelectionState::Enabled;
@@ -234,7 +296,7 @@ bool FNiagaraPlatformSet::IsEnabled(const UDeviceProfile* Profile, int32 Effects
 void FNiagaraPlatformSet::InvalidateCachedData()
 {
 #if WITH_EDITOR
-	CachedEQMasksPerDeviceProfile.Empty();
+	CachedQLMasksPerDeviceProfile.Empty();
 	CachedPlatformIniSettings.Empty();
 #endif
 
@@ -273,13 +335,13 @@ bool FNiagaraPlatformSet::ShouldPruneEmittersOnCook(const FString& PlatformName)
 
 #if WITH_EDITOR
 
-// bool FNiagaraPlatformSet::Conflicts(const FNiagaraPlatformSet& Other, TArray<const UDeviceProfile*>& OutConflictingProfiles, bool bIncludeProfilesWithVariableEffectsQuality)const
+// bool FNiagaraPlatformSet::Conflicts(const FNiagaraPlatformSet& Other, TArray<const UDeviceProfile*>& OutConflictingProfiles, bool bIncludeProfilesWithVariableQualityLevel)const
 // {
 // 	for (const UObject* ProfileObj : UDeviceProfileManager::Get().Profiles)
 // 	{
 // 		if (const UDeviceProfile* Profile = Cast<const UDeviceProfile>(ProfileObj))
 // 		{
-// 			if (bIncludeProfilesWithVariableEffectsQuality || CanChangeScalabilityAtRuntime(Profile) == false)
+// 			if (bIncludeProfilesWithVariableQualityLevel || CanChangeScalabilityAtRuntime(Profile) == false)
 // 			{
 // 				if (IsEnabledForDeviceProfile(Profile) && Other.IsEnabledForDeviceProfile(Profile))
 // 				{
@@ -294,7 +356,7 @@ bool FNiagaraPlatformSet::ShouldPruneEmittersOnCook(const FString& PlatformName)
 
 bool FNiagaraPlatformSet::IsEffectQualityEnabled(int32 EffectQuality)const
 {
-	return ((1 << EffectQuality) & EffectsQualityMask) != 0;
+	return ((1 << EffectQuality) & QualityLevelMask) != 0;
 }
 
 void FNiagaraPlatformSet::SetEnabledForEffectQuality(int32 EffectQuality, bool bEnabled)
@@ -302,25 +364,25 @@ void FNiagaraPlatformSet::SetEnabledForEffectQuality(int32 EffectQuality, bool b
 	int32 EQBit = (1 << EffectQuality);
 	if (bEnabled)
 	{
-		EffectsQualityMask |= EQBit;
+		QualityLevelMask |= EQBit;
 	}
 	else
 	{
-		EffectsQualityMask &= ~EQBit;
+		QualityLevelMask &= ~EQBit;
 	}
 	OnChanged();
 }
 
-void FNiagaraPlatformSet::SetDeviceProfileState(UDeviceProfile* Profile, int32 EffectsQuality, ENiagaraPlatformSelectionState NewState)
+void FNiagaraPlatformSet::SetDeviceProfileState(UDeviceProfile* Profile, int32 QualityLevel, ENiagaraPlatformSelectionState NewState)
 {
 	int32 Index = INDEX_NONE;
 	Index = DeviceProfileStates.IndexOfByPredicate([&](const FNiagaraDeviceProfileStateEntry& Entry) { return Entry.ProfileName == Profile->GetFName(); });
 
-	int32 ProfileEQMask = GetEffectQualityMaskForDeviceProfile(Profile);
-	if (ProfileEQMask != INDEX_NONE)
+	int32 ProfileQLMask = GetEffectQualityMaskForDeviceProfile(Profile);
+	if (ProfileQLMask != INDEX_NONE)
 	{
 		//For platforms that cannot change EQ at runtime we mark all state bits when setting state here so that if someone changes their EQ setting in the future, the state will be preserved.
-		EffectsQuality = INDEX_NONE;
+		QualityLevel = INDEX_NONE;
 	}
 
 	if(Index == INDEX_NONE)
@@ -329,12 +391,12 @@ void FNiagaraPlatformSet::SetDeviceProfileState(UDeviceProfile* Profile, int32 E
 		{
 			FNiagaraDeviceProfileStateEntry& NewEntry = DeviceProfileStates.AddDefaulted_GetRef();
 			NewEntry.ProfileName = Profile->GetFName();
-			NewEntry.SetState(EffectsQuality, NewState);
+			NewEntry.SetState(QualityLevel, NewState);
 		}
 	}
 	else
 	{
-		DeviceProfileStates[Index].SetState(EffectsQuality, NewState);
+		DeviceProfileStates[Index].SetState(QualityLevel, NewState);
 
 		if (DeviceProfileStates[Index].AllDefaulted())
 		{
@@ -344,17 +406,17 @@ void FNiagaraPlatformSet::SetDeviceProfileState(UDeviceProfile* Profile, int32 E
 	OnChanged();
 }
 
-ENiagaraPlatformSelectionState FNiagaraPlatformSet::GetDeviceProfileState(UDeviceProfile* Profile, int32 EffectsQuality)const
+ENiagaraPlatformSelectionState FNiagaraPlatformSet::GetDeviceProfileState(UDeviceProfile* Profile, int32 QualityLevel)const
 {
 	if (const FNiagaraDeviceProfileStateEntry* ExistingEntry = DeviceProfileStates.FindByPredicate([&](const FNiagaraDeviceProfileStateEntry& Entry) { return Entry.ProfileName == Profile->GetFName(); }))
 	{
-		int32 ProfileEQMask = GetEffectQualityMaskForDeviceProfile(Profile);
-		if (ProfileEQMask == INDEX_NONE || ProfileEQMask & CreateEQMask(EffectsQuality))
+		int32 ProfileQLMask = GetEffectQualityMaskForDeviceProfile(Profile);
+		if (ProfileQLMask == INDEX_NONE || ProfileQLMask & CreateQualityLevelMask(QualityLevel))
 		{
 			//For profiles that cannot change scalability at runtime we store all flags in their state so that if anyone ever changes their EQ Cvar, the state setting remains valid.
 			//This just means we also have to ensure this is the correct EQ here.
 
-			return ExistingEntry->GetState(EffectsQuality);
+			return ExistingEntry->GetState(QualityLevel);
 		}
 	}
 	return ENiagaraPlatformSelectionState::Default;
@@ -367,25 +429,27 @@ void FNiagaraPlatformSet::OnChanged()
 
 bool FNiagaraPlatformSet::CanChangeScalabilityAtRuntime(const FString& PlatformName)
 {
-	int32 EQMask = GetEffectQualityMaskForPlatform(PlatformName);
+	int32 QLMask = GetEffectQualityMaskForPlatform(PlatformName);
 	//If we're on a platform that can change EQ then it will have the full mask.
-	return EQMask == INDEX_NONE;
+	return QLMask == INDEX_NONE;
 }
 
 bool FNiagaraPlatformSet::CanChangeScalabilityAtRuntime(const UDeviceProfile* DeviceProfile)
 {
 	if (ensure(DeviceProfile))
 	{
-		int32 EQMask = GetEffectQualityMaskForDeviceProfile(DeviceProfile);
+		int32 QLMask = GetEffectQualityMaskForDeviceProfile(DeviceProfile);
 		//If we're on a platform that can change EQ then it will have the full mask.
-		return EQMask == INDEX_NONE;
+		return QLMask == INDEX_NONE;
 	}
 	return true;//Assuming true if we fail to find the platform seems safest.
 }
 
 bool FNiagaraPlatformSet::GatherConflicts(const TArray<const FNiagaraPlatformSet*>& PlatformSets, TArray<FNiagaraPlatformSetConflictInfo>& OutConflicts)
 {
-	Scalability::FQualityLevels NumLevels = Scalability::GetQualityLevelCounts();
+	const UNiagaraSettings* Settings = GetDefault<UNiagaraSettings>();
+	check(Settings);
+	int32 NumLevels = Settings->QualityLevels.Num();
 
 	FNiagaraPlatformSetConflictInfo* CurrentConflict = nullptr;
 	for (int32 A = 0; A < PlatformSets.Num(); ++A)
@@ -418,7 +482,7 @@ bool FNiagaraPlatformSet::GatherConflicts(const TArray<const FNiagaraPlatformSet
 					}
 					FNiagaraPlatformSetConflictEntry& ConflictEntry = CurrentConflict->Conflicts.AddDefaulted_GetRef();
 					ConflictEntry.ProfileName = Profile->GetFName();
-					ConflictEntry.EffectsQualityMask = ConflictMask;
+					ConflictEntry.QualityLevelMask = ConflictMask;
 				}
 			}
 		}
@@ -434,7 +498,7 @@ FNiagaraPlatformSet::FPlatformIniSettings FNiagaraPlatformSet::GetPlatformIniSet
 		return *CachedSettings;
 	}
 
-	//Load config files in which we can reasonable expect to find sg.EffectQuality and may be set.
+	//Load config files in which we can reasonable expect to find fx.Niagara.QualityLevel and may be set.
 	FConfigFile EngineSettings;
 	FConfigCacheIni::LoadLocalIniFile(EngineSettings, TEXT("Engine"), true, *PlatformName);//Should use BaseProfileName? Are either of these ensured to be correct? I worry this is brittle.
 
@@ -460,40 +524,49 @@ FNiagaraPlatformSet::FPlatformIniSettings FNiagaraPlatformSet::GetPlatformIniSet
 		return bFound;
 	};
 
+	//The effect quality for this platform.
 	int32 EffectsQuality = Scalability::DefaultQualityLevel;
-	int32 EQMask = INDEX_NONE;
 
 	//If this platform can change quality settings at runtime then we return the full mask.
 	int32 CanChangeEffectQuality = 0;
 
 	FindCVarValue(TEXT("SystemSettings"), CanChangeEQCVarName, CanChangeEffectQuality);
-	if (CanChangeEffectQuality != 0)
-	{
-		//If this platform can change it's EffectsQuality at runtime we return the ALL mask.
-		EQMask = INDEX_NONE;
-	}
-	else
+
+	if (!FindCVarValue(TEXT("ScalabilityGroups"), TEXT("sg.EffectsQuality"), EffectsQuality))
 	{
-		if (!FindCVarValue(TEXT("ScalabilityGroups"), TEXT("sg.EffectsQuality"), EffectsQuality))
-		{
-			FindCVarValue(TEXT("SystemSettings"), TEXT("sg.EffectsQuality"), EffectsQuality);
-		}
-		EQMask = CreateEQMask(EffectsQuality);
+		FindCVarValue(TEXT("SystemSettings"), TEXT("sg.EffectsQuality"), EffectsQuality);
 	}
 
-	int32 PruneEmittersOnCook = 0;
+	//Get the platforms default quality level setting.
+	//This can be overridden directly in a device profile or indirectly by overriding effects quality.
+
+	int32 PruneEmittersOnCook = GbPruneEmittersOnCook;
 	FindCVarValue(TEXT("SystemSettings"), PruneEmittersOnCookName, PruneEmittersOnCook);
 
 	FPlatformIniSettings& NewSetting = CachedPlatformIniSettings.Add(*PlatformName);
-	NewSetting = FPlatformIniSettings(CanChangeEffectQuality, PruneEmittersOnCook, EQMask);
+	NewSetting = FPlatformIniSettings(CanChangeEffectQuality, PruneEmittersOnCook, EffectsQuality);
+
+	//Find the Niagara Quality Levels set for each EffectsQuality Level for this platform.
+	int32 NumEffectsQualities = Scalability::GetQualityLevelCounts().EffectsQuality;
+	FString EQStr;
+	for (int32 EQ = 0; EQ < NumEffectsQualities; ++EQ)
+	{
+		FString SectionName = Scalability::GetScalabilitySectionString(TEXT("EffectsQuality"), EQ, NumEffectsQualities);
+		int32 NiagaraQualityLevelForEQ = DefaultQualityLevel;
+		ScalabilitySettings.GetInt(*SectionName, NiagaraQualityLevelName, NiagaraQualityLevelForEQ);
+		NewSetting.QualityLevelsPerEffectsQuality.Add(NiagaraQualityLevelForEQ);
+		EQStr += FString::Printf(TEXT("EQ:%d = NQL:%d\n"), EQ, NiagaraQualityLevelForEQ);
+	}
+
+	//UE_LOG(LogNiagara, Warning, TEXT("\n=====================================================\n%s - PlatformSettings(%d, %d, %d)\n%s\n========================================="), *PlatformName, CanChangeEffectQuality, PruneEmittersOnCook, EffectsQuality, *EQStr);
 	return NewSetting;
 }
 
 int32 FNiagaraPlatformSet::GetEffectQualityMaskForPlatform(const FString& PlatformName)
 {
 	FPlatformIniSettings PlatformSettings = GetPlatformIniSettings(PlatformName);
-	return PlatformSettings.EffectsQualityMask;
+	return PlatformSettings.bCanChangeQualityLevelAtRuntime;
 }
 #endif
 
-#undef LOCTEXT_NAMESPACE
+#undef LOCTEXT_NAMESPACE
\ No newline at end of file
diff --git a/Engine/Plugins/FX/Niagara/Source/Niagara/Private/NiagaraSystem.cpp b/Engine/Plugins/FX/Niagara/Source/Niagara/Private/NiagaraSystem.cpp
--- a/Engine/Plugins/FX/Niagara/Source/Niagara/Private/NiagaraSystem.cpp
+++ b/Engine/Plugins/FX/Niagara/Source/Niagara/Private/NiagaraSystem.cpp
@@ -336,7 +336,7 @@ void UNiagaraSystem::PostLoad()
 			FNiagaraSystemScalabilityOverride& LegacyOverride = ScalabilityOverrides_DEPRECATED[DL];
 			FNiagaraSystemScalabilityOverride& NewOverride = SystemScalabilityOverrides.Overrides.AddDefaulted_GetRef();
 			NewOverride = LegacyOverride;
-			NewOverride.Platforms = FNiagaraPlatformSet(FNiagaraPlatformSet::CreateEQMask(DL));
+			NewOverride.Platforms = FNiagaraPlatformSet(FNiagaraPlatformSet::CreateQualityLevelMask(DL));
 		}
 	}
 
@@ -1490,7 +1490,7 @@ void UNiagaraSystem::ResolveScalabilitySettings()
 	}
 }
 
-void UNiagaraSystem::OnEffectsQualityChanged()
+void UNiagaraSystem::OnQualityLevelChanged()
 {
 	ResolveScalabilitySettings();
 
@@ -1498,7 +1498,7 @@ void UNiagaraSystem::OnEffectsQualityChanged()
 	{
 		if (Handle.GetInstance())
 		{
-			Handle.GetInstance()->OnEffectsQualityChanged();
+			Handle.GetInstance()->OnQualityLevelChanged();
 		}
 	}
 
diff --git a/Engine/Plugins/FX/Niagara/Source/Niagara/Public/NiagaraModule.h b/Engine/Plugins/FX/Niagara/Source/Niagara/Public/NiagaraModule.h
--- a/Engine/Plugins/FX/Niagara/Source/Niagara/Public/NiagaraModule.h
+++ b/Engine/Plugins/FX/Niagara/Source/Niagara/Public/NiagaraModule.h
@@ -72,8 +72,6 @@ public:
 
 #endif
 
-	FORCEINLINE static int32 GetEffectsQuality() { return EngineEffectsQuality; }
-
 	FORCEINLINE static float GetGlobalSpawnCountScale() { return EngineGlobalSpawnCountScale; }
 	FORCEINLINE static float GetGlobalSystemCountScale() { return EngineGlobalSystemCountScale; }
 
@@ -165,8 +163,6 @@ public:
 	FORCEINLINE static const FNiagaraVariable&  GetVar_DataInstance_Alive() { return DataInstance_Alive; }
 	FORCEINLINE static const FNiagaraVariable&  GetVar_BeginDefaults() { return Translator_BeginDefaults; }
 
-	static void OnEffectsQualityChanged(int32 NewEffectsQuality, bool bForce);
-
 	FOnProcessQueue OnProcessQueue;
 
 #if WITH_EDITORONLY_DATA
diff --git a/Engine/Plugins/FX/Niagara/Source/Niagara/Public/NiagaraSettings.h b/Engine/Plugins/FX/Niagara/Source/Niagara/Public/NiagaraSettings.h
--- a/Engine/Plugins/FX/Niagara/Source/Niagara/Public/NiagaraSettings.h
+++ b/Engine/Plugins/FX/Niagara/Source/Niagara/Public/NiagaraSettings.h
@@ -28,6 +28,10 @@ class NIAGARA_API UNiagaraSettings : public UDeveloperSettings
 	UPROPERTY(config, EditAnywhere, Category = Niagara, meta = (AllowedClasses = "NiagaraEffectType"))
 	FSoftObjectPath DefaultEffectType;
 
+	/** The quality levels Niagara uses. */
+	UPROPERTY(config, EditAnywhere, Category = Scalability)
+	TArray<FText> QualityLevels;
+
 	// Begin UDeveloperSettings Interface
 	virtual FName GetCategoryName() const override;
 #if WITH_EDITOR
diff --git a/Engine/Plugins/FX/Niagara/Source/NiagaraEditor/Private/Customizations/NiagaraPlatformSetCustomization.cpp b/Engine/Plugins/FX/Niagara/Source/NiagaraEditor/Private/Customizations/NiagaraPlatformSetCustomization.cpp
--- a/Engine/Plugins/FX/Niagara/Source/NiagaraEditor/Private/Customizations/NiagaraPlatformSetCustomization.cpp
+++ b/Engine/Plugins/FX/Niagara/Source/NiagaraEditor/Private/Customizations/NiagaraPlatformSetCustomization.cpp
@@ -20,6 +20,7 @@
 #include "Widgets/Input/SMenuAnchor.h"
 #include "Widgets/Layout/SWrapBox.h"
 #include "Widgets/Text/STextBlock.h"
+#include "NiagaraSettings.h"
 
 #define LOCTEXT_NAMESPACE "FNiagaraPlatformSetCustomization"
 
@@ -96,11 +97,11 @@ void FNiagaraPlatformSetCustomization::CustomizeHeader(TSharedRef<IPropertyHandl
 	HeaderRow
 		.WholeRowContent()
 		[ 	
-			SAssignNew(EffectsQualityWidgetBox, SWrapBox)
+			SAssignNew(QualityLevelWidgetBox, SWrapBox)
 			.UseAllottedWidth(true)
 		];
 
-	GenerateEffectsQualitySelectionWidgets();
+	GenerateQualityLevelSelectionWidgets();
 }
 
 void FNiagaraPlatformSetCustomization::CustomizeChildren(TSharedRef<IPropertyHandle> StructPropertyHandle, class IDetailChildrenBuilder& ChildBuilder, IPropertyTypeCustomizationUtils& StructCustomizationUtils)
@@ -117,33 +118,37 @@ FText FNiagaraPlatformSetCustomization::GetTooltipText() const
 	return LOCTEXT("Platforms", "Platforms");
 }
 
-void FNiagaraPlatformSetCustomization::GenerateEffectsQualitySelectionWidgets()
+void FNiagaraPlatformSetCustomization::GenerateQualityLevelSelectionWidgets()
 {
-	EffectsQualityWidgetBox->ClearChildren();
+	QualityLevelWidgetBox->ClearChildren();
+
+	const UNiagaraSettings* Settings = GetDefault<UNiagaraSettings>();
+	check(Settings);
+
+	int32 NumQualityLevels = Settings->QualityLevels.Num();
 	
-	Scalability::FQualityLevels Counts = Scalability::GetQualityLevelCounts();
-	EffectsQualityMenuAnchors.SetNum(Counts.EffectsQuality);
-	EffectsQualityMenuContents.SetNum(Counts.EffectsQuality);
+	QualityLevelMenuAnchors.SetNum(NumQualityLevels);
+	QualityLevelMenuContents.SetNum(NumQualityLevels);
 
-	for (int32 EffectsQuality = 0; EffectsQuality < Counts.EffectsQuality; ++EffectsQuality)
+	for (int32 QualityLevel = 0; QualityLevel < NumQualityLevels; ++QualityLevel)
 	{
-		bool First = EffectsQuality == 0;
-		bool Last = EffectsQuality == (Counts.EffectsQuality - 1);
+		bool First = QualityLevel == 0;
+		bool Last = QualityLevel == (NumQualityLevels - 1);
 
-		if (!EffectsQualityMenuAnchors[EffectsQuality].IsValid())
+		if (!QualityLevelMenuAnchors[QualityLevel].IsValid())
 		{
-			EffectsQualityMenuAnchors[EffectsQuality] = SNew(SMenuAnchor)
+			QualityLevelMenuAnchors[QualityLevel] = SNew(SMenuAnchor)
 				.ToolTipText(LOCTEXT("AddPlatformOverride", "Add an override for a specific platform."))
-				.OnGetMenuContent(this, &FNiagaraPlatformSetCustomization::GenerateDeviceProfileTreeWidget, EffectsQuality)
+				.OnGetMenuContent(this, &FNiagaraPlatformSetCustomization::GenerateDeviceProfileTreeWidget, QualityLevel)
 				[
 					SNew(SButton)
-					.Visibility_Lambda([this, EffectsQuality]()
+					.Visibility_Lambda([this, QualityLevel]()
 					{
-						return EffectsQualityWidgetBox->GetChildren()->GetChildAt(EffectsQuality)->IsHovered() ? EVisibility::Visible : EVisibility::Hidden;
+						return QualityLevelWidgetBox->GetChildren()->GetChildAt(QualityLevel)->IsHovered() ? EVisibility::Visible : EVisibility::Hidden;
 					})
 					.ButtonStyle(FEditorStyle::Get(), "HoverHintOnly")
 					.ForegroundColor(FSlateColor::UseForeground())
-					.OnClicked(this, &FNiagaraPlatformSetCustomization::ToggleMenuOpenForEffectsQuality, EffectsQuality)
+					.OnClicked(this, &FNiagaraPlatformSetCustomization::ToggleMenuOpenForQualityLevel, QualityLevel)
 					[
 						SNew(SBox)
 						.WidthOverride(8)
@@ -157,7 +162,7 @@ void FNiagaraPlatformSetCustomization::GenerateEffectsQualitySelectionWidgets()
 				];
 		}
 
-		EffectsQualityWidgetBox->AddSlot()
+		QualityLevelWidgetBox->AddSlot()
 			.Padding(0, 0, 1, 0)
 			[
 				SNew(SBox)
@@ -171,8 +176,8 @@ void FNiagaraPlatformSetCustomization::GenerateEffectsQualitySelectionWidgets()
 						SNew(SCheckBox)
 						.Style(FNiagaraEditorStyle::Get(), First ? "NiagaraEditor.PlatformSet.StartButton" : 
 							(Last ? "NiagaraEditor.PlatformSet.EndButton" : "NiagaraEditor.PlatformSet.MiddleButton"))
-						.IsChecked(this, &FNiagaraPlatformSetCustomization::IsEQChecked, EffectsQuality)
-						.OnCheckStateChanged(this, &FNiagaraPlatformSetCustomization::EQCheckStateChanged, EffectsQuality)
+						.IsChecked(this, &FNiagaraPlatformSetCustomization::IsQLChecked, QualityLevel)
+						.OnCheckStateChanged(this, &FNiagaraPlatformSetCustomization::QLCheckStateChanged, QualityLevel)
 						[
 							SNew(SHorizontalBox)
 							+ SHorizontalBox::Slot()
@@ -181,8 +186,8 @@ void FNiagaraPlatformSetCustomization::GenerateEffectsQualitySelectionWidgets()
 							[
 								SNew(STextBlock)
 								.TextStyle(FNiagaraEditorStyle::Get(), "NiagaraEditor.PlatformSet.ButtonText")
-								.Text(FNiagaraPlatformSet::GetEffectsQualityText(EffectsQuality))
-								.ColorAndOpacity(this, &FNiagaraPlatformSetCustomization::GetEffectsQualityButtonTextColor, EffectsQuality)
+								.Text(FNiagaraPlatformSet::GetQualityLevelText(QualityLevel))
+								.ColorAndOpacity(this, &FNiagaraPlatformSetCustomization::GetQualityLevelButtonTextColor, QualityLevel)
 								.ShadowOffset(FVector2D(1, 1))
 							]
 							// error icon
@@ -197,8 +202,8 @@ void FNiagaraPlatformSetCustomization::GenerateEffectsQualitySelectionWidgets()
 								.HeightOverride(12)
 								[
 									SNew(SImage)
-									.ToolTipText(this, &FNiagaraPlatformSetCustomization::GetEffectsQualityErrorToolTip, EffectsQuality)
-									.Visibility(this, &FNiagaraPlatformSetCustomization::GetEffectsQualityErrorVisibility, EffectsQuality)
+									.ToolTipText(this, &FNiagaraPlatformSetCustomization::GetQualityLevelErrorToolTip, QualityLevel)
+									.Visibility(this, &FNiagaraPlatformSetCustomization::GetQualityLevelErrorVisibility, QualityLevel)
 									.Image(FEditorStyle::GetBrush("Icons.Error"))
 								]
 							]
@@ -209,14 +214,14 @@ void FNiagaraPlatformSetCustomization::GenerateEffectsQualitySelectionWidgets()
 							.VAlign(VAlign_Center)
 							.Padding(0,0,2,0)
 							[
-								EffectsQualityMenuAnchors[EffectsQuality].ToSharedRef()
+								QualityLevelMenuAnchors[QualityLevel].ToSharedRef()
 							]
 						]
 					]
 					+ SVerticalBox::Slot()
 					.AutoHeight()
 					[
-						GenerateAdditionalDevicesWidgetForEQ(EffectsQuality)
+						GenerateAdditionalDevicesWidgetForQL(QualityLevel)
 					]
 				]
 			];
@@ -313,7 +318,7 @@ void FNiagaraPlatformSetCustomization::InvalidateSiblingConflicts() const
 	}
 }
 
-EVisibility FNiagaraPlatformSetCustomization::GetEffectsQualityErrorVisibility(int32 EffectsQuality) const
+EVisibility FNiagaraPlatformSetCustomization::GetQualityLevelErrorVisibility(int32 QualityLevel) const
 {
 	// not part of an array
 	if (PlatformSetArrayIndex == INDEX_NONE)
@@ -326,12 +331,12 @@ EVisibility FNiagaraPlatformSetCustomization::GetEffectsQualityErrorVisibility(i
 		if (ConflictInfo.SetAIndex == PlatformSetArrayIndex || 
 			ConflictInfo.SetBIndex == PlatformSetArrayIndex)
 		{
-			// this conflict applies to this platform set, check if it applies to this effects quality button
-			const int32 EQMask = FNiagaraPlatformSet::CreateEQMask(EffectsQuality);
+			// this conflict applies to this platform set, check if it applies to this quality leve button
+			const int32 QLMask = FNiagaraPlatformSet::CreateQualityLevelMask(QualityLevel);
 
 			for (const FNiagaraPlatformSetConflictEntry& Conflict : ConflictInfo.Conflicts)
 			{
-				if ((EQMask & Conflict.EffectsQualityMask) != 0)
+				if ((QLMask & Conflict.QualityLevelMask) != 0)
 				{
 					return EVisibility::Visible;
 				}
@@ -342,7 +347,7 @@ EVisibility FNiagaraPlatformSetCustomization::GetEffectsQualityErrorVisibility(i
 	return EVisibility::Collapsed;
 }
 
-FText FNiagaraPlatformSetCustomization::GetEffectsQualityErrorToolTip(int32 EffectsQuality) const
+FText FNiagaraPlatformSetCustomization::GetQualityLevelErrorToolTip(int32 QualityLevel) const
 {
 	if (PlatformSetArrayIndex == INDEX_NONE)
 	{
@@ -365,13 +370,13 @@ FText FNiagaraPlatformSetCustomization::GetEffectsQualityErrorToolTip(int32 Effe
 
 		if (OtherIndex != INDEX_NONE)
 		{
-			const int32 EQMask = FNiagaraPlatformSet::CreateEQMask(EffectsQuality);
+			const int32 QLMask = FNiagaraPlatformSet::CreateQualityLevelMask(QualityLevel);
 
 			for (const FNiagaraPlatformSetConflictEntry& Conflict : ConflictInfo.Conflicts)
 			{
-				if ((EQMask & Conflict.EffectsQualityMask) != 0)
+				if ((QLMask & Conflict.QualityLevelMask) != 0)
 				{
-					FText FormatString = LOCTEXT("EffectsQualityConflictToolTip", "This effect quality conflicts with the set at index {Index} in this array."); 
+					FText FormatString = LOCTEXT("QualityLevelConflictToolTip", "This effect quality conflicts with the set at index {Index} in this array."); 
 
 					FFormatNamedArguments Args;
 					Args.Add(TEXT("Index"), OtherIndex);
@@ -385,14 +390,14 @@ FText FNiagaraPlatformSetCustomization::GetEffectsQualityErrorToolTip(int32 Effe
 	return FText::GetEmpty();
 }
 
-FSlateColor FNiagaraPlatformSetCustomization::GetEffectsQualityButtonTextColor(int32 EffectsQuality) const
+FSlateColor FNiagaraPlatformSetCustomization::GetQualityLevelButtonTextColor(int32 QualityLevel) const
 {
-	return TargetPlatformSet->IsEffectQualityEnabled(EffectsQuality) ? 
+	return TargetPlatformSet->IsEffectQualityEnabled(QualityLevel) ? 
 		FSlateColor(FLinearColor(0.95f, 0.95f, 0.95f)) :
 		FSlateColor::UseForeground();
 }
 
-TSharedRef<SWidget> FNiagaraPlatformSetCustomization::GenerateAdditionalDevicesWidgetForEQ(int32 EffectsQuality)
+TSharedRef<SWidget> FNiagaraPlatformSetCustomization::GenerateAdditionalDevicesWidgetForQL(int32 QualityLevel)
 {
 	TSharedRef<SVerticalBox> Container = SNew(SVerticalBox);
 
@@ -451,8 +456,8 @@ TSharedRef<SWidget> FNiagaraPlatformSetCustomization::GenerateAdditionalDevicesW
 				.HeightOverride(12)
 				[
 					SNew(SImage)
-					.ToolTipText(this, &FNiagaraPlatformSetCustomization::GetDeviceProfileErrorToolTip, Profile, EffectsQuality)
-					.Visibility(this, &FNiagaraPlatformSetCustomization::GetDeviceProfileErrorVisibility, Profile, EffectsQuality)
+					.ToolTipText(this, &FNiagaraPlatformSetCustomization::GetDeviceProfileErrorToolTip, Profile, QualityLevel)
+					.Visibility(this, &FNiagaraPlatformSetCustomization::GetDeviceProfileErrorVisibility, Profile, QualityLevel)
 					.Image(FEditorStyle::GetBrush("Icons.Error"))
 				]
 			];
@@ -466,7 +471,7 @@ TSharedRef<SWidget> FNiagaraPlatformSetCustomization::GenerateAdditionalDevicesW
 				SNew(SButton)
 				.ButtonStyle(FEditorStyle::Get(), "HoverHintOnly")
 				.ForegroundColor(FSlateColor::UseForeground())
-				.OnClicked(this, &FNiagaraPlatformSetCustomization::RemoveDeviceProfile, Profile, EffectsQuality)
+				.OnClicked(this, &FNiagaraPlatformSetCustomization::RemoveDeviceProfile, Profile, QualityLevel)
 				.ToolTipText(LOCTEXT("RemoveDevice", "Remove this device override."))
 				[
 					SNew(SBox)
@@ -483,7 +488,7 @@ TSharedRef<SWidget> FNiagaraPlatformSetCustomization::GenerateAdditionalDevicesW
 
 	TArray<UDeviceProfile*> EnabledProfiles;
 	TArray<UDeviceProfile*> DisabledProfiles;
-	TargetPlatformSet->GetOverridenDeviceProfiles(EffectsQuality, EnabledProfiles, DisabledProfiles);
+	TargetPlatformSet->GetOverridenDeviceProfiles(QualityLevel, EnabledProfiles, DisabledProfiles);
 
 	for (UDeviceProfile* Profile : EnabledProfiles)
 	{
@@ -498,10 +503,10 @@ TSharedRef<SWidget> FNiagaraPlatformSetCustomization::GenerateAdditionalDevicesW
 	return Container;
 }
 
-FReply FNiagaraPlatformSetCustomization::RemoveDeviceProfile(UDeviceProfile* Profile, int32 EffectsQuality)
+FReply FNiagaraPlatformSetCustomization::RemoveDeviceProfile(UDeviceProfile* Profile, int32 QualityLevel)
 {
 	PropertyHandle->NotifyPreChange();
-	TargetPlatformSet->SetDeviceProfileState(Profile, EffectsQuality, ENiagaraPlatformSelectionState::Default);
+	TargetPlatformSet->SetDeviceProfileState(Profile, QualityLevel, ENiagaraPlatformSelectionState::Default);
 	PropertyHandle->NotifyPostChange(EPropertyChangeType::ValueSet);
 
 	InvalidateSiblingConflicts();
@@ -509,7 +514,7 @@ FReply FNiagaraPlatformSetCustomization::RemoveDeviceProfile(UDeviceProfile* Pro
 	return FReply::Handled();
 }
 
-EVisibility FNiagaraPlatformSetCustomization::GetDeviceProfileErrorVisibility(UDeviceProfile* Profile, int32 EffectsQuality) const
+EVisibility FNiagaraPlatformSetCustomization::GetDeviceProfileErrorVisibility(UDeviceProfile* Profile, int32 QualityLevel) const
 {
 	// not part of an array
 	if (PlatformSetArrayIndex == INDEX_NONE)
@@ -522,11 +527,11 @@ EVisibility FNiagaraPlatformSetCustomization::GetDeviceProfileErrorVisibility(UD
 		if (ConflictInfo.SetAIndex == PlatformSetArrayIndex || 
 			ConflictInfo.SetBIndex == PlatformSetArrayIndex)
 		{
-			const int32 EQMask = FNiagaraPlatformSet::CreateEQMask(EffectsQuality);
+			const int32 QLMask = FNiagaraPlatformSet::CreateQualityLevelMask(QualityLevel);
 
 			for (const FNiagaraPlatformSetConflictEntry& Entry : ConflictInfo.Conflicts)
 			{
-				if ((Entry.EffectsQualityMask & EQMask) != 0 &&
+				if ((Entry.QualityLevelMask & QLMask) != 0 &&
 					Entry.ProfileName == Profile->GetFName())
 				{
 					return EVisibility::Visible;
@@ -538,7 +543,7 @@ EVisibility FNiagaraPlatformSetCustomization::GetDeviceProfileErrorVisibility(UD
 	return EVisibility::Collapsed;
 }
 
-FText FNiagaraPlatformSetCustomization::GetDeviceProfileErrorToolTip(UDeviceProfile* Profile, int32 EffectsQuality) const
+FText FNiagaraPlatformSetCustomization::GetDeviceProfileErrorToolTip(UDeviceProfile* Profile, int32 QualityLevel) const
 {
 	for (const FNiagaraPlatformSetConflictInfo& ConflictInfo : CachedConflicts)
 	{
@@ -556,11 +561,11 @@ FText FNiagaraPlatformSetCustomization::GetDeviceProfileErrorToolTip(UDeviceProf
 
 		if (OtherIndex != INDEX_NONE)
 		{
-			const int32 EQMask = FNiagaraPlatformSet::CreateEQMask(EffectsQuality);
+			const int32 QLMask = FNiagaraPlatformSet::CreateQualityLevelMask(QualityLevel);
 
 			for (const FNiagaraPlatformSetConflictEntry& Entry : ConflictInfo.Conflicts)
 			{
-				if ((Entry.EffectsQualityMask & EQMask) != 0 &&
+				if ((Entry.QualityLevelMask & QLMask) != 0 &&
 					Entry.ProfileName == Profile->GetFName())
 				{
 					FText FormatString = LOCTEXT("PlatformOverrideConflictToolTip", "This platform override conflicts with the set at index {Index} in this array."); 
@@ -577,28 +582,28 @@ FText FNiagaraPlatformSetCustomization::GetDeviceProfileErrorToolTip(UDeviceProf
 	return FText::GetEmpty();
 }
 
-FReply FNiagaraPlatformSetCustomization::ToggleMenuOpenForEffectsQuality(int32 EffectsQuality)
+FReply FNiagaraPlatformSetCustomization::ToggleMenuOpenForQualityLevel(int32 QualityLevel)
 {
-	check(EffectsQualityMenuAnchors.IsValidIndex(EffectsQuality));
+	check(QualityLevelMenuAnchors.IsValidIndex(QualityLevel));
 	
-	TSharedPtr<SMenuAnchor> MenuAnchor = EffectsQualityMenuAnchors[EffectsQuality];
+	TSharedPtr<SMenuAnchor> MenuAnchor = QualityLevelMenuAnchors[QualityLevel];
 	MenuAnchor->SetIsOpen(!MenuAnchor->IsOpen());
 
 	return FReply::Handled();
 }
 
 // Is or does a viewmodel contain any children active at the given quality?
-bool FNiagaraPlatformSetCustomization::IsTreeActiveForEQ(const TSharedPtr<FNiagaraDeviceProfileViewModel>& Tree, int32 EffectsQualityMask) const
+bool FNiagaraPlatformSetCustomization::IsTreeActiveForQL(const TSharedPtr<FNiagaraDeviceProfileViewModel>& Tree, int32 QualityLevelMask) const
 {
 	int32 Mask = TargetPlatformSet->GetEffectQualityMaskForDeviceProfile(Tree->Profile);
-	if ((Mask & EffectsQualityMask) != 0)
+	if ((Mask & QualityLevelMask) != 0)
 	{
 		return true;
 	}
 
 	for (const TSharedPtr<FNiagaraDeviceProfileViewModel>& Child : Tree->Children)
 	{
-		if (IsTreeActiveForEQ(Child, EffectsQualityMask))
+		if (IsTreeActiveForQL(Child, QualityLevelMask))
 		{
 			return true;
 		}
@@ -607,16 +612,16 @@ bool FNiagaraPlatformSetCustomization::IsTreeActiveForEQ(const TSharedPtr<FNiaga
 	return false;
 }
 
-void FNiagaraPlatformSetCustomization::FilterTreeForEQ(const TSharedPtr<FNiagaraDeviceProfileViewModel>& SourceTree, TSharedPtr<FNiagaraDeviceProfileViewModel>& FilteredTree, int32 EffectsQualityMask)
+void FNiagaraPlatformSetCustomization::FilterTreeForQL(const TSharedPtr<FNiagaraDeviceProfileViewModel>& SourceTree, TSharedPtr<FNiagaraDeviceProfileViewModel>& FilteredTree, int32 QualityLevelMask)
 {
 	for (const TSharedPtr<FNiagaraDeviceProfileViewModel>& SourceChild : SourceTree->Children)
 	{
-		if (IsTreeActiveForEQ(SourceChild, EffectsQualityMask))
+		if (IsTreeActiveForQL(SourceChild, QualityLevelMask))
 		{
 			TSharedPtr<FNiagaraDeviceProfileViewModel>& FilteredChild = FilteredTree->Children.Add_GetRef(MakeShared<FNiagaraDeviceProfileViewModel>());
 			FilteredChild->Profile = SourceChild->Profile;
 
-			FilterTreeForEQ(SourceChild, FilteredChild, EffectsQualityMask);
+			FilterTreeForQL(SourceChild, FilteredChild, QualityLevelMask);
 		}
 	}
 }
@@ -672,65 +677,68 @@ void FNiagaraPlatformSetCustomization::CreateDeviceProfileTree()
 		FullDeviceProfileTree.Add(CurrRoot);
 	}
 
-	Scalability::FQualityLevels Counts = Scalability::GetQualityLevelCounts();
-	FilteredDeviceProfileTrees.SetNum(Counts.EffectsQuality);
+	const UNiagaraSettings* Settings = GetDefault<UNiagaraSettings>();
+	check(Settings);
+
+	int32 NumQualityLevels = Settings->QualityLevels.Num();
+	FilteredDeviceProfileTrees.SetNum(NumQualityLevels);
 	
 	for (TSharedPtr<FNiagaraDeviceProfileViewModel>& FullDeviceRoot : FullDeviceProfileTree)
 	{
-		for (int32 EffectsQuality = 0; EffectsQuality < Counts.EffectsQuality; ++EffectsQuality)
+		for (int32 QualityLevel = 0; QualityLevel < NumQualityLevels; ++QualityLevel)
 		{
-			int32 EffectsQualityMask = FNiagaraPlatformSet::CreateEQMask(EffectsQuality);
+			int32 QualityLevelMask = FNiagaraPlatformSet::CreateQualityLevelMask(QualityLevel);
 
-			if (IsTreeActiveForEQ(FullDeviceRoot, EffectsQualityMask))
+			if (IsTreeActiveForQL(FullDeviceRoot, QualityLevelMask))
 			{
-				TArray<TSharedPtr<FNiagaraDeviceProfileViewModel>>& FilteredRoots = FilteredDeviceProfileTrees[EffectsQuality];
+				TArray<TSharedPtr<FNiagaraDeviceProfileViewModel>>& FilteredRoots = FilteredDeviceProfileTrees[QualityLevel];
 
 				TSharedPtr<FNiagaraDeviceProfileViewModel>& FilteredRoot = FilteredRoots.Add_GetRef(MakeShared<FNiagaraDeviceProfileViewModel>());
 				FilteredRoot->Profile = FullDeviceRoot->Profile;
 
-				FilterTreeForEQ(FullDeviceRoot, FilteredRoot, EffectsQualityMask);
+				FilterTreeForQL(FullDeviceRoot, FilteredRoot, QualityLevelMask);
 			}
 		}
 	}
 }
 
-TSharedRef<SWidget> FNiagaraPlatformSetCustomization::GenerateDeviceProfileTreeWidget(int32 EffectsQuality)
+TSharedRef<SWidget> FNiagaraPlatformSetCustomization::GenerateDeviceProfileTreeWidget(int32 QualityLevel)
 {
 	if (FullDeviceProfileTree.Num() == 0)
 	{
 		CreateDeviceProfileTree();	
 	}
 
-	if (EffectsQualityMenuContents[EffectsQuality].IsValid())
+	if (QualityLevelMenuContents[QualityLevel].IsValid())
 	{
-		return EffectsQualityMenuContents[EffectsQuality].ToSharedRef();
+		return QualityLevelMenuContents[QualityLevel].ToSharedRef();
 	}
 
 	TArray<TSharedPtr<FNiagaraDeviceProfileViewModel>>* TreeToUse = &FullDeviceProfileTree;
-	if (EffectsQuality != INDEX_NONE)
+	if (QualityLevel != INDEX_NONE)
 	{
-		check(EffectsQuality < FilteredDeviceProfileTrees.Num());
-		TreeToUse = &FilteredDeviceProfileTrees[EffectsQuality];
+		check(QualityLevel < FilteredDeviceProfileTrees.Num());
+		TreeToUse = &FilteredDeviceProfileTrees[QualityLevel];
 	}
 
-	return SAssignNew(EffectsQualityMenuContents[EffectsQuality], SBorder)
+	return SAssignNew(QualityLevelMenuContents[QualityLevel], SBorder)
 		.BorderImage(FEditorStyle::Get().GetBrush("Menu.Background"))
 		[
 			SAssignNew(DeviceProfileTreeWidget, STreeView<TSharedPtr<FNiagaraDeviceProfileViewModel>>)
 			.TreeItemsSource(TreeToUse)
-			.OnGenerateRow(this, &FNiagaraPlatformSetCustomization::OnGenerateDeviceProfileTreeRow, EffectsQuality)
-			.OnGetChildren(this, &FNiagaraPlatformSetCustomization::OnGetDeviceProfileTreeChildren, EffectsQuality)
+			.OnGenerateRow(this, &FNiagaraPlatformSetCustomization::OnGenerateDeviceProfileTreeRow, QualityLevel)
+			.OnGetChildren(this, &FNiagaraPlatformSetCustomization::OnGetDeviceProfileTreeChildren, QualityLevel)
 			.SelectionMode(ESelectionMode::None)
 		];
 }
 
-TSharedRef<ITableRow> FNiagaraPlatformSetCustomization::OnGenerateDeviceProfileTreeRow(TSharedPtr<FNiagaraDeviceProfileViewModel> InItem, const TSharedRef<STableViewBase>& OwnerTable, int32 EffectsQuality)
+TSharedRef<ITableRow> FNiagaraPlatformSetCustomization::OnGenerateDeviceProfileTreeRow(TSharedPtr<FNiagaraDeviceProfileViewModel> InItem, const TSharedRef<STableViewBase>& OwnerTable, int32 QualityLevel)
 {
 	TSharedPtr<SHorizontalBox> RowContainer;
 	SAssignNew(RowContainer, SHorizontalBox);
 
 	int32 ProfileMask = TargetPlatformSet->GetEffectQualityMaskForDeviceProfile(InItem->Profile);
-	FText NameTooltip = FText::Format(LOCTEXT("ProfileEQTooltipFmt", "Effects Quality: {0}"), FNiagaraPlatformSet::GetEffectsQualityMaskText(ProfileMask));
+	FText NameTooltip = FText::Format(LOCTEXT("ProfileQLTooltipFmt", "Effects Quality: {0}"), FNiagaraPlatformSet::GetQualityLevelMaskText(ProfileMask));
 	
 	//Top level profile. Look for a platform icon.
 	if (InItem->Profile->Parent == nullptr)
@@ -760,7 +768,7 @@ TSharedRef<ITableRow> FNiagaraPlatformSetCustomization::OnGenerateDeviceProfileT
 
 	FName TextStyleName("NormalText");
 	FSlateColor TextColor(FSlateColor::UseForeground());
-	ENiagaraPlatformSelectionState CurrentState = TargetPlatformSet->GetDeviceProfileState(InItem->Profile, EffectsQuality);
+	ENiagaraPlatformSelectionState CurrentState = TargetPlatformSet->GetDeviceProfileState(InItem->Profile, QualityLevel);
 	
 	if (CurrentState == ENiagaraPlatformSelectionState::Enabled)
 	{
@@ -774,8 +782,8 @@ TSharedRef<ITableRow> FNiagaraPlatformSetCustomization::OnGenerateDeviceProfileT
 	}
 
 
-	int32 EffectsQualityMask = FNiagaraPlatformSet::CreateEQMask(EffectsQuality);
-	if ((ProfileMask & EffectsQualityMask) == 0)
+	int32 QualityLevelMask = FNiagaraPlatformSet::CreateQualityLevelMask(QualityLevel);
+	if ((ProfileMask & QualityLevelMask) == 0)
 	{
 		TextColor = FSlateColor::UseSubduedForeground();
 	}
@@ -787,8 +795,8 @@ TSharedRef<ITableRow> FNiagaraPlatformSetCustomization::OnGenerateDeviceProfileT
 		[
 			SNew(SButton)
 			.ButtonStyle(FEditorStyle::Get(), "NoBorder")
-			.OnClicked(this, &FNiagaraPlatformSetCustomization::OnProfileMenuButtonClicked, InItem, EffectsQuality, false)
-			.IsEnabled(this, &FNiagaraPlatformSetCustomization::GetProfileMenuItemEnabled, InItem, EffectsQuality)
+			.OnClicked(this, &FNiagaraPlatformSetCustomization::OnProfileMenuButtonClicked, InItem, QualityLevel, false)
+			.IsEnabled(this, &FNiagaraPlatformSetCustomization::GetProfileMenuItemEnabled, InItem, QualityLevel)
 			.ForegroundColor(TextColor)
 			.ToolTipText(NameTooltip)
 			[
@@ -808,9 +816,9 @@ TSharedRef<ITableRow> FNiagaraPlatformSetCustomization::OnGenerateDeviceProfileT
 			.ButtonStyle(FEditorStyle::Get(), "HoverHintOnly")
 			.HAlign(HAlign_Center)
 			.VAlign(VAlign_Fill)
-			.Visibility(this, &FNiagaraPlatformSetCustomization::GetProfileMenuButtonVisibility, InItem, EffectsQuality)
-			.OnClicked(this, &FNiagaraPlatformSetCustomization::OnProfileMenuButtonClicked, InItem, EffectsQuality, true)
-			.ToolTipText(this, &FNiagaraPlatformSetCustomization::GetProfileMenuButtonToolTip, InItem, EffectsQuality)
+			.Visibility(this, &FNiagaraPlatformSetCustomization::GetProfileMenuButtonVisibility, InItem, QualityLevel)
+			.OnClicked(this, &FNiagaraPlatformSetCustomization::OnProfileMenuButtonClicked, InItem, QualityLevel, true)
+			.ToolTipText(this, &FNiagaraPlatformSetCustomization::GetProfileMenuButtonToolTip, InItem, QualityLevel)
 			[
 				SNew(SBox)
 				.WidthOverride(8)
@@ -819,7 +827,7 @@ TSharedRef<ITableRow> FNiagaraPlatformSetCustomization::OnGenerateDeviceProfileT
 				.HAlign(HAlign_Center)
 				[
 					SNew(SImage)
-					.Image(this, &FNiagaraPlatformSetCustomization::GetProfileMenuButtonImage, InItem, EffectsQuality)
+					.Image(this, &FNiagaraPlatformSetCustomization::GetProfileMenuButtonImage, InItem, QualityLevel)
 				]
 			]
 		];
@@ -839,19 +847,19 @@ enum class EProfileButtonMode
 	Remove
 };
 
-static EProfileButtonMode GetProfileMenuButtonMode(FNiagaraPlatformSet* PlatformSet, TSharedPtr<FNiagaraDeviceProfileViewModel> Item, int32 EffectsQuality)
+static EProfileButtonMode GetProfileMenuButtonMode(FNiagaraPlatformSet* PlatformSet, TSharedPtr<FNiagaraDeviceProfileViewModel> Item, int32 QualityLevel)
 {
 	int32 Mask = PlatformSet->GetEffectQualityMaskForDeviceProfile(Item->Profile);
-	int32 EQMask = FNiagaraPlatformSet::CreateEQMask(EffectsQuality);
+	int32 QLMask = FNiagaraPlatformSet::CreateQualityLevelMask(QualityLevel);
 
-	if ((Mask & EQMask) == 0)
+	if ((Mask & QLMask) == 0)
 	{
 		return EProfileButtonMode::None;
 	}
 
-	ENiagaraPlatformSelectionState CurrentState = PlatformSet->GetDeviceProfileState(Item->Profile, EffectsQuality);
+	ENiagaraPlatformSelectionState CurrentState = PlatformSet->GetDeviceProfileState(Item->Profile, QualityLevel);
 
-	bool bQualityEnabled = PlatformSet->IsEffectQualityEnabled(EffectsQuality);
+	bool bQualityEnabled = PlatformSet->IsEffectQualityEnabled(QualityLevel);
 	bool bIsDefault = CurrentState == ENiagaraPlatformSelectionState::Default;
 
 	if (bIsDefault && bQualityEnabled)
@@ -868,9 +876,9 @@ static EProfileButtonMode GetProfileMenuButtonMode(FNiagaraPlatformSet* Platform
 }
 
 
-FText FNiagaraPlatformSetCustomization::GetProfileMenuButtonToolTip(TSharedPtr<FNiagaraDeviceProfileViewModel> Item, int32 EffectsQuality) const
+FText FNiagaraPlatformSetCustomization::GetProfileMenuButtonToolTip(TSharedPtr<FNiagaraDeviceProfileViewModel> Item, int32 QualityLevel) const
 {
-	EProfileButtonMode Mode = GetProfileMenuButtonMode(TargetPlatformSet, Item, EffectsQuality);
+	EProfileButtonMode Mode = GetProfileMenuButtonMode(TargetPlatformSet, Item, QualityLevel);
 	switch(Mode)
 	{
 		case EProfileButtonMode::Include:
@@ -884,9 +892,9 @@ FText FNiagaraPlatformSetCustomization::GetProfileMenuButtonToolTip(TSharedPtr<F
 	return FText::GetEmpty();
 }
 
-EVisibility FNiagaraPlatformSetCustomization::GetProfileMenuButtonVisibility(TSharedPtr<FNiagaraDeviceProfileViewModel> Item, int32 EffectsQuality) const
+EVisibility FNiagaraPlatformSetCustomization::GetProfileMenuButtonVisibility(TSharedPtr<FNiagaraDeviceProfileViewModel> Item, int32 QualityLevel) const
 {
-	EProfileButtonMode Mode = GetProfileMenuButtonMode(TargetPlatformSet, Item, EffectsQuality);
+	EProfileButtonMode Mode = GetProfileMenuButtonMode(TargetPlatformSet, Item, QualityLevel);
 	if (Mode == EProfileButtonMode::None)
 	{
 		return EVisibility::Collapsed;
@@ -895,9 +903,9 @@ EVisibility FNiagaraPlatformSetCustomization::GetProfileMenuButtonVisibility(TSh
 	return EVisibility::Visible;
 }
 
-bool FNiagaraPlatformSetCustomization::GetProfileMenuItemEnabled(TSharedPtr<FNiagaraDeviceProfileViewModel> Item, int32 EffectsQuality) const
+bool FNiagaraPlatformSetCustomization::GetProfileMenuItemEnabled(TSharedPtr<FNiagaraDeviceProfileViewModel> Item, int32 QualityLevel) const
 {
-	EProfileButtonMode Mode = GetProfileMenuButtonMode(TargetPlatformSet, Item, EffectsQuality);
+	EProfileButtonMode Mode = GetProfileMenuButtonMode(TargetPlatformSet, Item, QualityLevel);
 	if (Mode == EProfileButtonMode::None)
 	{
 		return false;
@@ -906,9 +914,9 @@ bool FNiagaraPlatformSetCustomization::GetProfileMenuItemEnabled(TSharedPtr<FNia
 	return true;
 }
 
-const FSlateBrush* FNiagaraPlatformSetCustomization::GetProfileMenuButtonImage(TSharedPtr<FNiagaraDeviceProfileViewModel> Item, int32 EffectsQuality) const
+const FSlateBrush* FNiagaraPlatformSetCustomization::GetProfileMenuButtonImage(TSharedPtr<FNiagaraDeviceProfileViewModel> Item, int32 QualityLevel) const
 {
-	EProfileButtonMode Mode = GetProfileMenuButtonMode(TargetPlatformSet, Item, EffectsQuality);
+	EProfileButtonMode Mode = GetProfileMenuButtonMode(TargetPlatformSet, Item, QualityLevel);
 	switch(Mode)
 	{
 		case EProfileButtonMode::Include:
@@ -922,11 +930,11 @@ const FSlateBrush* FNiagaraPlatformSetCustomization::GetProfileMenuButtonImage(T
 	return FEditorStyle::GetBrush("NoBrush");
 }
 
-FReply FNiagaraPlatformSetCustomization::OnProfileMenuButtonClicked(TSharedPtr<FNiagaraDeviceProfileViewModel> Item, int32 EffectsQuality, bool bReopenMenu)
+FReply FNiagaraPlatformSetCustomization::OnProfileMenuButtonClicked(TSharedPtr<FNiagaraDeviceProfileViewModel> Item, int32 QualityLevel, bool bReopenMenu)
 {
 	ENiagaraPlatformSelectionState TargetState;
 
-	EProfileButtonMode Mode = GetProfileMenuButtonMode(TargetPlatformSet, Item, EffectsQuality);
+	EProfileButtonMode Mode = GetProfileMenuButtonMode(TargetPlatformSet, Item, QualityLevel);
 	switch(Mode)
 	{
 		case EProfileButtonMode::Include:
@@ -943,14 +951,14 @@ FReply FNiagaraPlatformSetCustomization::OnProfileMenuButtonClicked(TSharedPtr<F
 	}
 
 	PropertyHandle->NotifyPreChange();
-	TargetPlatformSet->SetDeviceProfileState(Item->Profile, EffectsQuality, TargetState);
+	TargetPlatformSet->SetDeviceProfileState(Item->Profile, QualityLevel, TargetState);
 	PropertyHandle->NotifyPostChange(EPropertyChangeType::ValueSet);
 
 	InvalidateSiblingConflicts();
 
 	if (!bReopenMenu)
 	{
-		EffectsQualityMenuAnchors[EffectsQuality]->SetIsOpen(false);
+		QualityLevelMenuAnchors[QualityLevel]->SetIsOpen(false);
 	}
 
 	DeviceProfileTreeWidget->RequestTreeRefresh();
@@ -958,27 +966,27 @@ FReply FNiagaraPlatformSetCustomization::OnProfileMenuButtonClicked(TSharedPtr<F
 	return FReply::Handled();
 }
 
-void FNiagaraPlatformSetCustomization::OnGetDeviceProfileTreeChildren(TSharedPtr<FNiagaraDeviceProfileViewModel> InItem, TArray< TSharedPtr<FNiagaraDeviceProfileViewModel> >& OutChildren, int32 EffectsQuality)
+void FNiagaraPlatformSetCustomization::OnGetDeviceProfileTreeChildren(TSharedPtr<FNiagaraDeviceProfileViewModel> InItem, TArray< TSharedPtr<FNiagaraDeviceProfileViewModel> >& OutChildren, int32 QualityLevel)
 {
-	if (TargetPlatformSet->GetDeviceProfileState(InItem->Profile, EffectsQuality) == ENiagaraPlatformSelectionState::Default)
+	if (TargetPlatformSet->GetDeviceProfileState(InItem->Profile, QualityLevel) == ENiagaraPlatformSelectionState::Default)
 	{
 		OutChildren = InItem->Children;
 	}
 }
 
-ECheckBoxState FNiagaraPlatformSetCustomization::IsEQChecked(int32 EffectsQuality)const
+ECheckBoxState FNiagaraPlatformSetCustomization::IsQLChecked(int32 QualityLevel)const
 {
 	if (TargetPlatformSet)
 	{
-		return TargetPlatformSet->IsEffectQualityEnabled(EffectsQuality) ? ECheckBoxState::Checked : ECheckBoxState::Unchecked;
+		return TargetPlatformSet->IsEffectQualityEnabled(QualityLevel) ? ECheckBoxState::Checked : ECheckBoxState::Unchecked;
 	}
 	return ECheckBoxState::Undetermined;
 }
 
-void FNiagaraPlatformSetCustomization::EQCheckStateChanged(ECheckBoxState CheckState, int32 EffectsQuality)
+void FNiagaraPlatformSetCustomization::QLCheckStateChanged(ECheckBoxState CheckState, int32 QualityLevel)
 {
 	PropertyHandle->NotifyPreChange();
-	TargetPlatformSet->SetEnabledForEffectQuality(EffectsQuality, CheckState == ECheckBoxState::Checked);
+	TargetPlatformSet->SetEnabledForEffectQuality(QualityLevel, CheckState == ECheckBoxState::Checked);
 	PropertyHandle->NotifyPostChange(EPropertyChangeType::ValueSet);
 
 	InvalidateSiblingConflicts();
@@ -986,7 +994,7 @@ void FNiagaraPlatformSetCustomization::EQCheckStateChanged(ECheckBoxState CheckS
 
 void FNiagaraPlatformSetCustomization::OnPropertyValueChanged()
 {
-	GenerateEffectsQualitySelectionWidgets();
+	GenerateQualityLevelSelectionWidgets();
 	UpdateCachedConflicts();
 }
 
diff --git a/Engine/Plugins/FX/Niagara/Source/NiagaraEditor/Private/Customizations/NiagaraPlatformSetCustomization.h b/Engine/Plugins/FX/Niagara/Source/NiagaraEditor/Private/Customizations/NiagaraPlatformSetCustomization.h
--- a/Engine/Plugins/FX/Niagara/Source/NiagaraEditor/Private/Customizations/NiagaraPlatformSetCustomization.h
+++ b/Engine/Plugins/FX/Niagara/Source/NiagaraEditor/Private/Customizations/NiagaraPlatformSetCustomization.h
@@ -40,24 +40,24 @@ private:
 
 	void CreateDeviceProfileTree();
 
-	TSharedRef<SWidget> GenerateDeviceProfileTreeWidget(int32 EffectsQuality);
-	TSharedRef<ITableRow> OnGenerateDeviceProfileTreeRow(TSharedPtr<FNiagaraDeviceProfileViewModel> InItem, const TSharedRef<STableViewBase>& OwnerTable, int32 EffectsQuality);
-	void OnGetDeviceProfileTreeChildren(TSharedPtr<FNiagaraDeviceProfileViewModel> InItem, TArray< TSharedPtr<FNiagaraDeviceProfileViewModel> >& OutChildren, int32 EffectsQuality);
-	FReply RemoveDeviceProfile(UDeviceProfile* Profile, int32 EffectsQuality);
-	TSharedRef<SWidget> GetEffectsQualityMenuContents(int32 EffectsQuality) const;
+	TSharedRef<SWidget> GenerateDeviceProfileTreeWidget(int32 QualityLevel);
+	TSharedRef<ITableRow> OnGenerateDeviceProfileTreeRow(TSharedPtr<FNiagaraDeviceProfileViewModel> InItem, const TSharedRef<STableViewBase>& OwnerTable, int32 QualityLevel);
+	void OnGetDeviceProfileTreeChildren(TSharedPtr<FNiagaraDeviceProfileViewModel> InItem, TArray< TSharedPtr<FNiagaraDeviceProfileViewModel> >& OutChildren, int32 QualityLevel);
+	FReply RemoveDeviceProfile(UDeviceProfile* Profile, int32 QualityLevel);
+	TSharedRef<SWidget> GetQualityLevelMenuContents(int32 QualityLevel) const;
 
-	EVisibility GetDeviceProfileErrorVisibility(UDeviceProfile* Profile, int32 EffectsQuality) const;
-	FText GetDeviceProfileErrorToolTip(UDeviceProfile* Profile, int32 EffectsQuality) const;
+	EVisibility GetDeviceProfileErrorVisibility(UDeviceProfile* Profile, int32 QualityLevel) const;
+	FText GetDeviceProfileErrorToolTip(UDeviceProfile* Profile, int32 QualityLevel) const;
 
-	FSlateColor GetEffectsQualityButtonTextColor(int32 EffectsQuality) const;
-	EVisibility GetEffectsQualityErrorVisibility(int32 EffectsQuality) const;
-	FText GetEffectsQualityErrorToolTip(int32 EffectsQuality) const;
+	FSlateColor GetQualityLevelButtonTextColor(int32 QualityLevel) const;
+	EVisibility GetQualityLevelErrorVisibility(int32 QualityLevel) const;
+	FText GetQualityLevelErrorToolTip(int32 QualityLevel) const;
 
-	const FSlateBrush* GetProfileMenuButtonImage(TSharedPtr<FNiagaraDeviceProfileViewModel> Item, int32 EffectsQuality) const;
-	EVisibility GetProfileMenuButtonVisibility(TSharedPtr<FNiagaraDeviceProfileViewModel> Item, int32 EffectsQuality) const;
-	bool GetProfileMenuItemEnabled(TSharedPtr<FNiagaraDeviceProfileViewModel> Item, int32 EffectsQuality) const;
-	FText GetProfileMenuButtonToolTip(TSharedPtr<FNiagaraDeviceProfileViewModel> Item, int32 EffectsQuality) const;
-	FReply OnProfileMenuButtonClicked(TSharedPtr<FNiagaraDeviceProfileViewModel> InItem, int32 EffectsQuality, bool bReopenMenu);
+	const FSlateBrush* GetProfileMenuButtonImage(TSharedPtr<FNiagaraDeviceProfileViewModel> Item, int32 QualityLevel) const;
+	EVisibility GetProfileMenuButtonVisibility(TSharedPtr<FNiagaraDeviceProfileViewModel> Item, int32 QualityLevel) const;
+	bool GetProfileMenuItemEnabled(TSharedPtr<FNiagaraDeviceProfileViewModel> Item, int32 QualityLevel) const;
+	FText GetProfileMenuButtonToolTip(TSharedPtr<FNiagaraDeviceProfileViewModel> Item, int32 QualityLevel) const;
+	FReply OnProfileMenuButtonClicked(TSharedPtr<FNiagaraDeviceProfileViewModel> InItem, int32 QualityLevel, bool bReopenMenu);
 
 	TSharedRef<SWidget> GetCurrentDeviceProfileSelectionWidget(TSharedPtr<FNiagaraDeviceProfileViewModel> ProfileView);
 	TSharedRef<SWidget> OnGenerateDeviceProfileSelectionWidget(TSharedPtr<ENiagaraPlatformSelectionState> InItem);
@@ -65,16 +65,16 @@ private:
 	FText GetCurrentText() const;
 	FText GetTooltipText() const;
 
-	void GenerateEffectsQualitySelectionWidgets();
-	TSharedRef<SWidget> GenerateAdditionalDevicesWidgetForEQ(int32 EffectsQuality);
+	void GenerateQualityLevelSelectionWidgets();
+	TSharedRef<SWidget> GenerateAdditionalDevicesWidgetForQL(int32 QualityLevel);
 
-	bool IsTreeActiveForEQ(const TSharedPtr<FNiagaraDeviceProfileViewModel>& Tree, int32 EffectsQualityMask) const;
-	void FilterTreeForEQ(const TSharedPtr<FNiagaraDeviceProfileViewModel>& SourceTree, TSharedPtr<FNiagaraDeviceProfileViewModel>& FilteredTree, int32 EffectsQualityMask);
+	bool IsTreeActiveForQL(const TSharedPtr<FNiagaraDeviceProfileViewModel>& Tree, int32 QualityLevelMask) const;
+	void FilterTreeForQL(const TSharedPtr<FNiagaraDeviceProfileViewModel>& SourceTree, TSharedPtr<FNiagaraDeviceProfileViewModel>& FilteredTree, int32 QualityLevelMask);
 
-	ECheckBoxState IsEQChecked(int32 EffectsQuality) const;
-	void EQCheckStateChanged(ECheckBoxState CheckState, int32 EffectsQuality);
+	ECheckBoxState IsQLChecked(int32 QualityLevel) const;
+	void QLCheckStateChanged(ECheckBoxState CheckState, int32 QualityLevel);
 
-	FReply ToggleMenuOpenForEffectsQuality(int32 EffectsQuality);
+	FReply ToggleMenuOpenForQualityLevel(int32 QualityLevel);
 
 	void UpdateCachedConflicts();
 	void InvalidateSiblingConflicts() const;
@@ -85,9 +85,9 @@ private:
 	TSharedPtr<IPropertyHandleArray> PlatformSetArray;
 	int32 PlatformSetArrayIndex;
 
-	TArray<TSharedPtr<SMenuAnchor>> EffectsQualityMenuAnchors;
-	TArray<TSharedPtr<SWidget>> EffectsQualityMenuContents;
-	TSharedPtr<SWrapBox> EffectsQualityWidgetBox;
+	TArray<TSharedPtr<SMenuAnchor>> QualityLevelMenuAnchors;
+	TArray<TSharedPtr<SWidget>> QualityLevelMenuContents;
+	TSharedPtr<SWrapBox> QualityLevelWidgetBox;
 
 	TSharedPtr<IPropertyHandle> PropertyHandle;
 	struct FNiagaraPlatformSet* TargetPlatformSet;
diff --git a/Engine/Plugins/FX/Niagara/Source/NiagaraEditor/Private/Customizations/NiagaraTypeCustomizations.cpp b/Engine/Plugins/FX/Niagara/Source/NiagaraEditor/Private/Customizations/NiagaraTypeCustomizations.cpp
--- a/Engine/Plugins/FX/Niagara/Source/NiagaraEditor/Private/Customizations/NiagaraTypeCustomizations.cpp
+++ b/Engine/Plugins/FX/Niagara/Source/NiagaraEditor/Private/Customizations/NiagaraTypeCustomizations.cpp
@@ -1,23 +1,30 @@
 // Copyright Epic Games, Inc. All Rights Reserved.
 
 #include "NiagaraTypeCustomizations.h"
+#include "CoreMinimal.h"
 #include "DetailLayoutBuilder.h"
 #include "DetailWidgetRow.h"
 #include "IDetailChildrenBuilder.h"
 #include "NiagaraConstants.h"
+#include "NiagaraConstants.h"
 #include "NiagaraEditorStyle.h"
 #include "NiagaraEmitter.h"
 #include "NiagaraNodeOutput.h"
 #include "NiagaraNodeParameterMapBase.h"
 #include "NiagaraParameterMapHistory.h"
+#include "NiagaraPlatformSet.h"
 #include "NiagaraRendererProperties.h"
 #include "NiagaraScriptSource.h"
 #include "NiagaraScriptVariable.h"
 #include "NiagaraSystem.h"
 #include "NiagaraTypes.h"
+#include "PlatformInfo.h"
 #include "PropertyHandle.h"
 #include "SGraphActionMenu.h"
+#include "Scalability.h"
 #include "ScopedTransaction.h"
+#include "DeviceProfiles/DeviceProfile.h"
+#include "DeviceProfiles/DeviceProfileManager.h"
 #include "Framework/Application/SlateApplication.h"
 #include "Widgets/Input/SButton.h"
 #include "Widgets/Input/SCheckBox.h"
@@ -27,6 +34,7 @@
 #include "NiagaraSimulationStageBase.h"
 #include "Widgets/Text/STextBlock.h"
 #include "NiagaraDataInterfaceRW.h"
+#include "NiagaraSettings.h"
 
 #define LOCTEXT_NAMESPACE "FNiagaraVariableAttributeBindingCustomization"
 
diff --git a/Engine/Plugins/FX/Niagara/Source/NiagaraEditor/Private/ViewModels/NiagaraEmitterViewModel.cpp b/Engine/Plugins/FX/Niagara/Source/NiagaraEditor/Private/ViewModels/NiagaraEmitterViewModel.cpp
--- a/Engine/Plugins/FX/Niagara/Source/NiagaraEditor/Private/ViewModels/NiagaraEmitterViewModel.cpp
+++ b/Engine/Plugins/FX/Niagara/Source/NiagaraEditor/Private/ViewModels/NiagaraEmitterViewModel.cpp
@@ -269,10 +269,10 @@ FText FNiagaraEmitterViewModel::GetStatsText() const
 					return LOCTEXT("InvalidInstance", "Invalid Emitter! May have compile errors.");
 				}
 
- 				int32 EffectsQuality = INiagaraModule::GetEffectsQuality();
+ 				int32 QualityLevel = FNiagaraPlatformSet::GetQualityLevel();
   				if (!HandleEmitter->IsAllowedByScalability())
   				{
-					return FText::Format(ParticleDisabledDueToScalability, FNiagaraPlatformSet::GetEffectsQualityText(EffectsQuality));
+					return FText::Format(ParticleDisabledDueToScalability, FNiagaraPlatformSet::GetQualityLevelText(QualityLevel));
   				}
 
 				if (NiagaraCommands::EmitterStatsFormat->GetInt() == 1)
diff --git a/Engine/Source/Runtime/Engine/Public/Scalability.h b/Engine/Source/Runtime/Engine/Public/Scalability.h
--- a/Engine/Source/Runtime/Engine/Public/Scalability.h
+++ b/Engine/Source/Runtime/Engine/Public/Scalability.h
@@ -181,4 +181,6 @@ namespace Scalability
 #endif
 
 	ENGINE_API FText GetQualityLevelText(int32 Value, int32 NumLevels);
+
+	ENGINE_API FString GetScalabilitySectionString(const TCHAR* InGroupName, int32 InQualityLevel, int32 InNumLevels);
 }
