Removing OwerLOD scalabiliy feature.

diff --git a/Engine/Plugins/FX/Niagara/Source/Niagara/Classes/NiagaraEffectType.h b/Engine/Plugins/FX/Niagara/Source/Niagara/Classes/NiagaraEffectType.h
--- a/Engine/Plugins/FX/Niagara/Source/Niagara/Classes/NiagaraEffectType.h
+++ b/Engine/Plugins/FX/Niagara/Source/Niagara/Classes/NiagaraEffectType.h
@@ -58,8 +58,6 @@ struct FNiagaraSystemScalabilitySettings
 	UPROPERTY(EditAnywhere, Category = "Scalability", meta = (InlineEditConditionToggle))
 	uint32 bCullMaxInstanceCount : 1;
 	UPROPERTY(EditAnywhere, Category = "Scalability", meta = (InlineEditConditionToggle))
-	uint32 bCullByMaxOwnerLOD : 1;
-	UPROPERTY(EditAnywhere, Category = "Scalability", meta = (InlineEditConditionToggle))
 	uint32 bCullByMaxTimeWithoutRender : 1;
 
 	/** Effects of this type are culled beyond this distance. */
@@ -69,11 +67,7 @@ struct FNiagaraSystemScalabilitySettings
 	/** Effects of this type will fail to spawn when total active instance count exceeds this number. */
 	UPROPERTY(EditAnywhere, Category = "Scalability", meta = (EditCondition = "bCullMaxInstanceCount"))
 	float MaxInstances;
-
-	/** The effect is culled beyond this LOD level. Useful for FX that are linked to a specific actor with an LOD level driven be some other system. */
-	UPROPERTY(EditAnywhere, Category = "Scalability", meta = (EditCondition = "bCullByMaxOwnerLOD"))
-	int32 MaxOwnerLOD;
-
+	
 	//TODO:
 	/** The effect is culled when it's bounds take up less that this fraction of the total screen area. Only usable with fixed bounds. */
 	//float ScreenFraction;
@@ -110,8 +104,6 @@ struct FNiagaraSystemScalabilityOverride : public FNiagaraSystemScalabilitySetti
 	UPROPERTY(EditAnywhere, Category = "Override")
 	uint32 bOverrideInstanceCountSettings : 1;
 	UPROPERTY(EditAnywhere, Category = "Override")
-	uint32 bOverrideOwnerLODSettings : 1;
-	UPROPERTY(EditAnywhere, Category = "Override")
 	uint32 bOverrideTimeSinceRendererSettings : 1;
 };
 
diff --git a/Engine/Plugins/FX/Niagara/Source/Niagara/Private/NiagaraComponent.cpp b/Engine/Plugins/FX/Niagara/Source/Niagara/Private/NiagaraComponent.cpp
--- a/Engine/Plugins/FX/Niagara/Source/Niagara/Private/NiagaraComponent.cpp
+++ b/Engine/Plugins/FX/Niagara/Source/Niagara/Private/NiagaraComponent.cpp
@@ -460,7 +460,6 @@ UNiagaraComponent::UNiagaraComponent(const FObjectInitializer& ObjectInitializer
 	, bIsCulledByScalability(false)
 	//, bIsChangingAutoAttachment(false)
 	, ScalabilityManagerHandle(INDEX_NONE)
-	, OwnerLOD(0)
 {
 	OverrideParameters.SetOwner(this);
 
@@ -1936,11 +1935,6 @@ void UNiagaraComponent::SetPreviewLODDistance(bool bInEnablePreviewLODDistance,
 void UNiagaraComponent::SetPreviewLODDistance(bool bInEnablePreviewLODDistance, float InPreviewLODDistance){}
 #endif
 
-void UNiagaraComponent::SetOwnerLOD(int32 InOwnerLOD)
-{
-	OwnerLOD = InOwnerLOD;  
-}
-
 void UNiagaraComponent::SetAllowScalability(bool bAllow)
 {
 	bAllowScalability = bAllow; 
diff --git a/Engine/Plugins/FX/Niagara/Source/Niagara/Private/NiagaraEffectType.cpp b/Engine/Plugins/FX/Niagara/Source/Niagara/Private/NiagaraEffectType.cpp
--- a/Engine/Plugins/FX/Niagara/Source/Niagara/Private/NiagaraEffectType.cpp
+++ b/Engine/Plugins/FX/Niagara/Source/Niagara/Private/NiagaraEffectType.cpp
@@ -171,7 +171,6 @@ void UNiagaraEffectType::ProcessLastFrameCycleCounts()
 FNiagaraSystemScalabilityOverride::FNiagaraSystemScalabilityOverride()
 	: bOverrideDistanceSettings(false)
 	, bOverrideInstanceCountSettings(false)
-	, bOverrideOwnerLODSettings(false)
 	, bOverrideTimeSinceRendererSettings(false)
 {
 }
@@ -185,12 +184,10 @@ void FNiagaraSystemScalabilitySettings::Clear()
 {
 	Platforms = FNiagaraPlatformSet();
 	bCullByDistance = false;
-	bCullByMaxOwnerLOD = false;
 	bCullByMaxTimeWithoutRender = false;
 	bCullMaxInstanceCount = false;
 	MaxDistance = 0.0f;
 	MaxInstances = 0;
-	MaxOwnerLOD = 0;
 	MaxTimeWithoutRender = 0.0f;
 }
 
diff --git a/Engine/Plugins/FX/Niagara/Source/Niagara/Private/NiagaraScalabilityManager.cpp b/Engine/Plugins/FX/Niagara/Source/Niagara/Private/NiagaraScalabilityManager.cpp
--- a/Engine/Plugins/FX/Niagara/Source/Niagara/Private/NiagaraScalabilityManager.cpp
+++ b/Engine/Plugins/FX/Niagara/Source/Niagara/Private/NiagaraScalabilityManager.cpp
@@ -180,7 +180,6 @@ void FNiagaraScalabilityManager::Update(FNiagaraWorldManager* WorldMan)
 		FNiagaraScalabilityState& CompState = State[i];
 #if DEBUG_SCALABILITY_STATE
 		CompState.bCulledByInstanceCount = false;
-		CompState.bCulledByMaxOwnerLOD = false;
 		CompState.bCulledBySignificance = false;
 		CompState.bCulledByVisibility = false;
 #endif
@@ -274,14 +273,12 @@ void FNiagaraScalabilityManager::Dump()
 			, NumCulledBySignificance(0)
 			, NumCulledByInstanceCount(0)
 			, NumCulledByVisibility(0)
-			, NumCulledByMaxOwnerLOD(0)
 		{}
 
 		int32 NumCulled;
 		int32 NumCulledBySignificance;
 		int32 NumCulledByInstanceCount;
 		int32 NumCulledByVisibility;
-		int32 NumCulledByMaxOwnerLOD;
 	}Summary;
 
 	for (int32 i = 0; i < ManagedComponents.Num(); ++i)
@@ -310,11 +307,6 @@ void FNiagaraScalabilityManager::Dump()
 			CulledStr += TEXT("-Visibility-");
 			++Summary.NumCulledByVisibility;
 		}
-		if (CompState.bCulledByMaxOwnerLOD)
-		{
-			CulledStr += TEXT("-Owner LOD-");
-			++Summary.NumCulledByMaxOwnerLOD;
-		}
 
 		DetailString += FString::Printf(TEXT("| %s | Sig: %2.4f | %0xP | %s | %s |\n"), *CulledStr, CompState.Significance, Comp, *Comp->GetAsset()->GetPathName(), *Comp->GetPathName());
 	}
@@ -328,7 +320,6 @@ void FNiagaraScalabilityManager::Dump()
 	UE_LOG(LogNiagara, Display, TEXT("| Num Culled: %d |"), Summary.NumCulled);
 	UE_LOG(LogNiagara, Display, TEXT("| Num Culled By Significance: %d |"), Summary.NumCulledBySignificance);
 	UE_LOG(LogNiagara, Display, TEXT("| Num Culled By Instance Count: %d |"), Summary.NumCulledByInstanceCount);
-	UE_LOG(LogNiagara, Display, TEXT("| Num Culled By Max Owner LOD: %d |"), Summary.NumCulledByMaxOwnerLOD);
 	UE_LOG(LogNiagara, Display, TEXT("| Num Culled By Visibility: %d |"), Summary.NumCulledByVisibility);
 	UE_LOG(LogNiagara, Display, TEXT("| Avg Frame GT: %d |"), EffectType->GetAverageFrameTime_GT());
 	UE_LOG(LogNiagara, Display, TEXT("| Avg Frame GT + CNC: %d |"), EffectType->GetAverageFrameTime_GT_CNC());
diff --git a/Engine/Plugins/FX/Niagara/Source/Niagara/Private/NiagaraSystem.cpp b/Engine/Plugins/FX/Niagara/Source/Niagara/Private/NiagaraSystem.cpp
--- a/Engine/Plugins/FX/Niagara/Source/Niagara/Private/NiagaraSystem.cpp
+++ b/Engine/Plugins/FX/Niagara/Source/Niagara/Private/NiagaraSystem.cpp
@@ -1474,12 +1474,6 @@ void UNiagaraSystem::ResolveScalabilitySettings()
 				CurrentScalabilitySettings.MaxInstances = Override.MaxInstances;
 			}
 
-			if (Override.bOverrideOwnerLODSettings)
-			{
-				CurrentScalabilitySettings.bCullByMaxOwnerLOD = Override.bCullByMaxOwnerLOD;
-				CurrentScalabilitySettings.MaxOwnerLOD = Override.MaxOwnerLOD;
-			}
-
 			if (Override.bOverrideTimeSinceRendererSettings)
 			{
 				CurrentScalabilitySettings.bCullByMaxTimeWithoutRender = Override.bCullByMaxTimeWithoutRender;
diff --git a/Engine/Plugins/FX/Niagara/Source/Niagara/Private/NiagaraWorldManager.cpp b/Engine/Plugins/FX/Niagara/Source/Niagara/Private/NiagaraWorldManager.cpp
--- a/Engine/Plugins/FX/Niagara/Source/Niagara/Private/NiagaraWorldManager.cpp
+++ b/Engine/Plugins/FX/Niagara/Source/Niagara/Private/NiagaraWorldManager.cpp
@@ -803,7 +803,6 @@ void FNiagaraWorldManager::CalculateScalabilityState(UNiagaraSystem* System, con
 	bool bOldCulled = OutState.bCulled;
 	OutState.bCulled = false;
 	SignificanceCull(EffectType, ScalabilitySettings, Significance, OutState);
-	OwnerLODCull(EffectType, ScalabilitySettings, Component, OutState);
 	
 	//Can't cull dynamic bounds by visibility
 
@@ -871,15 +870,6 @@ void FNiagaraWorldManager::VisibilityCull(UNiagaraEffectType* EffectType, const
 #endif
 }
 
-void FNiagaraWorldManager::OwnerLODCull(UNiagaraEffectType* EffectType, const FNiagaraSystemScalabilitySettings& ScalabilitySettings, UNiagaraComponent* Component, FNiagaraScalabilityState& OutState)
-{
-	bool bCull = ScalabilitySettings.bCullByMaxOwnerLOD && Component->GetOwnerLOD() > ScalabilitySettings.MaxOwnerLOD;
-	OutState.bCulled |= bCull;
-#if DEBUG_SCALABILITY_STATE
-	OutState.bCulledByMaxOwnerLOD = bCull;
-#endif
-}
-
 void FNiagaraWorldManager::InstanceCountCull(UNiagaraEffectType* EffectType, const FNiagaraSystemScalabilitySettings& ScalabilitySettings, FNiagaraScalabilityState& OutState)
 {
 	bool bCull = ScalabilitySettings.bCullMaxInstanceCount && EffectType->NumInstances > ScalabilitySettings.MaxInstances;
diff --git a/Engine/Plugins/FX/Niagara/Source/Niagara/Public/NiagaraComponent.h b/Engine/Plugins/FX/Niagara/Source/Niagara/Public/NiagaraComponent.h
--- a/Engine/Plugins/FX/Niagara/Source/Niagara/Public/NiagaraComponent.h
+++ b/Engine/Plugins/FX/Niagara/Source/Niagara/Public/NiagaraComponent.h
@@ -512,11 +512,6 @@ public:
 	uint32 bWaitForCompilationOnActivate : 1;
 #endif
 
-	virtual void SetOwnerLOD(int32 InOwnerLOD);
-
-	UFUNCTION(BlueprintCallable, Category = Scalability, meta = (Keywords = "LOD scalability"))
-	FORCEINLINE int32 GetOwnerLOD()const { return OwnerLOD; }
-
 	/** Set whether this component is allowed to perform scalability checks and potentially be culled etc. Occasionally it is useful to disable this for specific components. E.g. Effects on the local player. */
 	UFUNCTION(BlueprintCallable, Category = Scalability, meta = (Keywords = "LOD scalability"))
 	void SetAllowScalability(bool bAllow);
@@ -553,13 +548,6 @@ private:
 	FDelegateHandle AssetExposedParametersChangedHandle;
 
 	int32 ScalabilityManagerHandle;
-
-	/**
-	LOD level of our owning actor / component if it's been provided. 
-	Can be useful for scalability calculations for actor based FX on actors who have some other system determining their LOD level.
-	Would be nice if we just had a virtual on the actor Actor->GetLODLevel() for example that would return 0 by default and the correct level for any actors implementing a LOD level.
-	*/
-	int32 OwnerLOD;
 };
 
 #if WITH_NIAGARA_COMPONENT_PREVIEW_DATA
diff --git a/Engine/Plugins/FX/Niagara/Source/Niagara/Public/NiagaraScalabilityManager.h b/Engine/Plugins/FX/Niagara/Source/Niagara/Public/NiagaraScalabilityManager.h
--- a/Engine/Plugins/FX/Niagara/Source/Niagara/Public/NiagaraScalabilityManager.h
+++ b/Engine/Plugins/FX/Niagara/Source/Niagara/Public/NiagaraScalabilityManager.h
@@ -90,7 +90,6 @@ struct FNiagaraScalabilityState
 		, bCulledBySignificance(0)
 		, bCulledByInstanceCount(0)
 		, bCulledByVisibility(0)
-		, bCulledByMaxOwnerLOD(0)
 #endif
 	{
 	}
@@ -103,7 +102,6 @@ struct FNiagaraScalabilityState
 	uint8 bCulledBySignificance : 1;
 	uint8 bCulledByInstanceCount : 1;
 	uint8 bCulledByVisibility : 1;
-	uint8 bCulledByMaxOwnerLOD : 1;
 #endif
 };
 
diff --git a/Engine/Plugins/FX/Niagara/Source/Niagara/Public/NiagaraSystemInstance.h b/Engine/Plugins/FX/Niagara/Source/Niagara/Public/NiagaraSystemInstance.h
--- a/Engine/Plugins/FX/Niagara/Source/Niagara/Public/NiagaraSystemInstance.h
+++ b/Engine/Plugins/FX/Niagara/Source/Niagara/Public/NiagaraSystemInstance.h
@@ -146,8 +146,6 @@ public:
 
 //	float GetSystemTimeSinceRendered() const { return SystemTimeSinceRenderedParam.GetValue(); }
 
-	//float GetOwnerLODDistance() const { return OwnerLODDistanceParam.GetValue(); }
-
 	//int32 GetNumParticles(int32 EmitterIndex) const { return ParameterNumParticleBindings[EmitterIndex].GetValue(); }
 	//float GetSpawnCountScale(int32 EmitterIndex) const { return ParameterSpawnCountScaleBindings[EmitterIndex].GetValue(); }
 
diff --git a/Engine/Plugins/FX/Niagara/Source/Niagara/Public/NiagaraWorldManager.h b/Engine/Plugins/FX/Niagara/Source/Niagara/Public/NiagaraWorldManager.h
--- a/Engine/Plugins/FX/Niagara/Source/Niagara/Public/NiagaraWorldManager.h
+++ b/Engine/Plugins/FX/Niagara/Source/Niagara/Public/NiagaraWorldManager.h
@@ -199,7 +199,6 @@ private:
 
 	FORCEINLINE_DEBUGGABLE void SignificanceCull(UNiagaraEffectType* EffectType, const FNiagaraSystemScalabilitySettings& ScalabilitySettings, float Significance, FNiagaraScalabilityState& OutState);
 	FORCEINLINE_DEBUGGABLE void VisibilityCull(UNiagaraEffectType* EffectType, const FNiagaraSystemScalabilitySettings& ScalabilitySettings, UNiagaraComponent* Component, FNiagaraScalabilityState& OutState);
-	FORCEINLINE_DEBUGGABLE void OwnerLODCull(UNiagaraEffectType* EffectType, const FNiagaraSystemScalabilitySettings& ScalabilitySettings, UNiagaraComponent* Component, FNiagaraScalabilityState& OutState);
 	FORCEINLINE_DEBUGGABLE void InstanceCountCull(UNiagaraEffectType* EffectType, const FNiagaraSystemScalabilitySettings& ScalabilitySettings, FNiagaraScalabilityState& OutState);
 
 	/** Calculate significance contribution from the distance to nearest view. */
diff --git a/Engine/Source/Runtime/Engine/Classes/Particles/ParticleSystemComponent.h b/Engine/Source/Runtime/Engine/Classes/Particles/ParticleSystemComponent.h
--- a/Engine/Source/Runtime/Engine/Classes/Particles/ParticleSystemComponent.h
+++ b/Engine/Source/Runtime/Engine/Classes/Particles/ParticleSystemComponent.h
@@ -414,9 +414,6 @@ public:
 
 	/** Returns an approximate memory usage value for this component. */
 	virtual uint32 GetApproxMemoryUsage() const { return 0; }
-
-	UFUNCTION(BlueprintCallable, Category = Preview, meta = (Keywords = "LOD scalability"))
-	virtual void SetOwnerLOD(int32 LOD) {}
 };
 
 
