- Rename RandomFilteredSocketBone to RandomFilteredSocket
- Add unfiltered bone functionality

diff --git a/Engine/Plugins/FX/Niagara/Shaders/Private/NiagaraDataInterfaceSkeletalMesh.ush b/Engine/Plugins/FX/Niagara/Shaders/Private/NiagaraDataInterfaceSkeletalMesh.ush
--- a/Engine/Plugins/FX/Niagara/Shaders/Private/NiagaraDataInterfaceSkeletalMesh.ush
+++ b/Engine/Plugins/FX/Niagara/Shaders/Private/NiagaraDataInterfaceSkeletalMesh.ush
@@ -34,7 +34,8 @@ struct FDISkelMeshContext
 	uint			MeshNumWeights;
 	int				NumBones;
 	int				NumFilteredBones;
-	Buffer<uint>	FilteredBones;
+	int				NumUnfilteredBones;
+	Buffer<uint>	FilteredAndUnfilteredBones;
 	int				NumFilteredSockets;
 	int				FilteredSocketBoneOffset;
 	float4x4		InstanceTransform;
@@ -475,14 +476,32 @@ void DISkelMesh_GetFilteredBoneCount(in FDISkelMeshContext DIContext, out int Co
 void DISkelMesh_GetFilteredBoneAt(in FDISkelMeshContext DIContext, in int InBoneIndex, out int Bone)
 {
     int BoneIndex = min(InBoneIndex, DIContext.NumFilteredBones - 1);
-    Bone = (int)DIContext.FilteredBones[BoneIndex];
+    Bone = (int)DIContext.FilteredAndUnfilteredBones[BoneIndex];
 }
 
 void DISkelMesh_RandomFilteredBone(in FDISkelMeshContext DIContext, uint Seed1, uint Seed2, uint Seed3, out int Bone)
 {
 	float RandT0 = NiagaraInternalNoise(Seed1, Seed2, Seed3);
     int FilteredBone = min(int(RandT0 * float(DIContext.NumFilteredBones)), DIContext.NumFilteredBones - 1); // avoid % by using mul/min to Tri = MeshVertexCount
-    Bone = (int)DIContext.FilteredBones[FilteredBone];
+    Bone = (int)DIContext.FilteredAndUnfilteredBones[FilteredBone];
+}
+
+void DISkelMesh_GetUnfilteredBoneCount(in FDISkelMeshContext DIContext, out int Count)
+{
+    Count = DIContext.NumUnfilteredBones;
+}
+
+void DISkelMesh_GetUnfilteredBoneAt(in FDISkelMeshContext DIContext, in int InBoneIndex, out int Bone)
+{
+    int BoneIndex = min(InBoneIndex, DIContext.NumUnfilteredBones - 1);
+    Bone = (int)DIContext.FilteredAndUnfilteredBones[BoneIndex + DIContext.NumFilteredBones];
+}
+
+void DISkelMesh_RandomUnfilteredBone(in FDISkelMeshContext DIContext, uint Seed1, uint Seed2, uint Seed3, out int Bone)
+{
+	float RandT0 = NiagaraInternalNoise(Seed1, Seed2, Seed3);
+    int UnfilteredBone = min(int(RandT0 * float(DIContext.NumUnfilteredBones)), DIContext.NumUnfilteredBones - 1); // avoid % by using mul/min to Tri = MeshVertexCount
+    Bone = (int)DIContext.FilteredAndUnfilteredBones[UnfilteredBone + DIContext.NumFilteredBones];
 }
 
 void DISkelMesh_GetFilteredSocketCount(in FDISkelMeshContext DIContext, out int Count)
@@ -495,7 +514,7 @@ void DISkelMesh_GetFilteredSocketBoneAt(in FDISkelMeshContext DIContext, in int
     Bone = DIContext.FilteredSocketBoneOffset + clamp(InSocketIndex, 0, DIContext.NumFilteredSockets);
 }
 
-void DISkelMesh_RandomFilteredSocketBone(in FDISkelMeshContext DIContext, uint Seed1, uint Seed2, uint Seed3, out int SocketBone)
+void DISkelMesh_RandomFilteredSocket(in FDISkelMeshContext DIContext, uint Seed1, uint Seed2, uint Seed3, out int SocketBone)
 {
 	float RandT0 = NiagaraInternalNoise(Seed1, Seed2, Seed3);
     SocketBone = DIContext.FilteredSocketBoneOffset + min(int(RandT0 * float(DIContext.NumFilteredSockets)), DIContext.NumFilteredSockets - 1); // avoid % by using mul/min to Tri = MeshVertexCount
@@ -508,7 +527,7 @@ void DISkelMesh_RandomFilteredSocketOrBone(in FDISkelMeshContext DIContext, uint
     int FilteredIndex = min(int(RandT0 * float(NumFilteredSocketsAndBones)), NumFilteredSocketsAndBones - 1);
 	if ( FilteredIndex < DIContext.NumFilteredBones )
 	{
-		Bone = (int)DIContext.FilteredBones[FilteredIndex];
+		Bone = (int)DIContext.FilteredAndUnfilteredBones[FilteredIndex];
 	}
 	else
 	{
@@ -541,7 +560,8 @@ void DISkelMesh_RandomFilteredSocketOrBone(in FDISkelMeshContext DIContext, uint
 	uint			MeshNumWeights_ ## NAME; \
 	int				NumBones_ ## NAME; \
 	int				NumFilteredBones_ ## NAME; \
-	Buffer<uint>	FilteredBones_ ## NAME; \
+	int				NumUnfilteredBones_ ## NAME; \
+	Buffer<uint>	FilteredAndUnfilteredBones_ ## NAME; \
 	int				NumFilteredSockets_ ## NAME; \
 	int				FilteredSocketBoneOffset_ ## NAME; \
 	float4x4		InstanceTransform_ ## NAME; \
@@ -571,7 +591,8 @@ void DISkelMesh_RandomFilteredSocketOrBone(in FDISkelMeshContext DIContext, uint
 	DIContext.MeshNumWeights = MeshNumWeights_ ## NAME; \
 	DIContext.NumBones = NumBones_ ## NAME; \
 	DIContext.NumFilteredBones	= NumFilteredBones_ ## NAME; \
-	DIContext.FilteredBones	= FilteredBones_ ## NAME; \
+	DIContext.NumUnfilteredBones = NumUnfilteredBones_ ## NAME; \
+	DIContext.FilteredAndUnfilteredBones	= FilteredAndUnfilteredBones_ ## NAME; \
 	DIContext.NumFilteredSockets = NumFilteredSockets_ ## NAME; \
 	DIContext.FilteredSocketBoneOffset = FilteredSocketBoneOffset_ ## NAME; \
 	DIContext.MeshWeightStride = MeshWeightStride_ ## NAME; \
diff --git a/Engine/Plugins/FX/Niagara/Source/Niagara/Classes/NiagaraDataInterfaceSkeletalMesh.h b/Engine/Plugins/FX/Niagara/Source/Niagara/Classes/NiagaraDataInterfaceSkeletalMesh.h
--- a/Engine/Plugins/FX/Niagara/Source/Niagara/Classes/NiagaraDataInterfaceSkeletalMesh.h
+++ b/Engine/Plugins/FX/Niagara/Source/Niagara/Classes/NiagaraDataInterfaceSkeletalMesh.h
@@ -302,7 +302,8 @@ public:
 	uint32 GetNumWeights() const { return NumWeights; }
 
 	int32 GetNumFilteredBones() const { return NumFilteredBones; }
-	FRHIShaderResourceView* GetFilteredBonesSRV() const { return FilteredBonesSRV; }
+	int32 GetNumUnfilteredBones() const { return NumUnfilteredBones;  }
+	FRHIShaderResourceView* GetFilteredAndUnfilteredBonesSRV() const { return FilteredAndUnfilteredBonesSRV; }
 
 	int32 GetNumFilteredSockets() const { return NumFilteredSockets; }
 	int32 GetFilteredSocketBoneOffset() const { return FilteredSocketBoneOffset; }
@@ -317,9 +318,10 @@ protected:
 	FShaderResourceViewRHIRef BufferTriangleMatricesOffsetSRV = nullptr;
 
 	int32 NumFilteredBones = 0;
-	TResourceArray<uint16> FilteredBonesArray;
-	FVertexBufferRHIRef FilteredBonesBuffer;
-	FShaderResourceViewRHIRef FilteredBonesSRV;
+	int32 NumUnfilteredBones = 0;
+	TResourceArray<uint16> FilteredAndUnfilteredBonesArray;
+	FVertexBufferRHIRef FilteredAndUnfilteredBonesBuffer;
+	FShaderResourceViewRHIRef FilteredAndUnfilteredBonesSRV;
 
 	int32 NumFilteredSockets = 0;
 	int32 FilteredSocketBoneOffset = 0;
@@ -429,8 +431,12 @@ struct FNDISkeletalMesh_InstanceData
 	/** Time separating Transform and PrevTransform. */
 	float DeltaSeconds;
 
-	/** Indices of the bones filtered by the user. */
-	TArray<int32> FilteredBones;
+	/* Number of filtered bones in the array. */
+	int32 NumFilteredBones = 0;
+	/* Number of unfiltered bones in the array. */
+	int32 NumUnfilteredBones = 0;
+	/** Indices of the bones filtered by the user followed by the unfiltered bones, if this array is empty no filtering is in effect. */
+	TArray<uint16> FilteredAndUnfilteredBones;
 
 	/** Name of all the sockets we use. */
 	struct FCachedSocketInfo
@@ -623,7 +629,8 @@ public:
 	static const FString MeshNumWeightsName;
 	static const FString NumBonesName;
 	static const FString NumFilteredBonesName;
-	static const FString FilteredBonesName;
+	static const FString NumUnfilteredBonesName;
+	static const FString FilteredAndUnfilteredBonesName;
 	static const FString NumFilteredSocketsName;
 	static const FString FilteredSocketBoneOffsetName;
 	static const FString InstanceTransformName;
@@ -740,11 +747,14 @@ public:
 	void GetFilteredBoneAt(FVectorVMContext& Context);
 	void RandomFilteredBone(FVectorVMContext& Context);
 
+	void GetUnfilteredBoneCount(FVectorVMContext& Context);
+	void GetUnfilteredBoneAt(FVectorVMContext& Context);
+	void RandomUnfilteredBone(FVectorVMContext& Context);
 
 	void GetFilteredSocketCount(FVectorVMContext& Context);
 	void GetFilteredSocketBoneAt(FVectorVMContext& Context);
 	void GetFilteredSocketTransform(FVectorVMContext& Context);
-	void RandomFilteredSocketBone(FVectorVMContext& Context);
+	void RandomFilteredSocket(FVectorVMContext& Context);
 
 	void RandomFilteredSocketOrBone(FVectorVMContext& Context);
 	// End of Direct Bone + Socket Sampling
@@ -778,13 +788,20 @@ public:
 	static const FName IsValidBoneName;
 	static const FName RandomBoneName;
 	static const FName GetBoneCountName;
+
 	static const FName RandomFilteredBoneName;
 	static const FName GetFilteredBoneCountName;
 	static const FName GetFilteredBoneAtName;
-	static const FName RandomFilteredSocketBoneName;
+
+	static const FName RandomUnfilteredBoneName;
+	static const FName GetUnfilteredBoneCountName;
+	static const FName GetUnfilteredBoneAtName;
+
+	static const FName RandomFilteredSocketName;
 	static const FName GetFilteredSocketCountName;
 	static const FName GetFilteredSocketBoneAtName;
 	static const FName GetFilteredSocketTransformName;
+
 	static const FName RandomFilteredSocketOrBoneName;
 
 	// Vertex Sampling
diff --git a/Engine/Plugins/FX/Niagara/Source/Niagara/Private/NDISkeletalMesh_BoneSampling.cpp b/Engine/Plugins/FX/Niagara/Source/Niagara/Private/NDISkeletalMesh_BoneSampling.cpp
--- a/Engine/Plugins/FX/Niagara/Source/Niagara/Private/NDISkeletalMesh_BoneSampling.cpp
+++ b/Engine/Plugins/FX/Niagara/Source/Niagara/Private/NDISkeletalMesh_BoneSampling.cpp
@@ -16,7 +16,6 @@ DECLARE_CYCLE_STAT(TEXT("Skel Mesh Skeleton Sampling"), STAT_NiagaraSkel_Bone_Sa
 //Final binders for all static mesh interface functions.
 DEFINE_NDI_FUNC_BINDER(UNiagaraDataInterfaceSkeletalMesh, GetSkinnedBoneData)
 DEFINE_NDI_DIRECT_FUNC_BINDER(UNiagaraDataInterfaceSkeletalMesh, IsValidBone)
-DEFINE_NDI_DIRECT_FUNC_BINDER(UNiagaraDataInterfaceSkeletalMesh, GetFilteredBoneAt)
 
 DEFINE_NDI_DIRECT_FUNC_BINDER(UNiagaraDataInterfaceSkeletalMesh, GetFilteredSocketBoneAt)
 DEFINE_NDI_DIRECT_FUNC_BINDER(UNiagaraDataInterfaceSkeletalMesh, GetFilteredSocketTransform)
@@ -34,7 +33,11 @@ const FName FSkeletalMeshInterfaceHelper::RandomFilteredBoneName("RandomFiltered
 const FName FSkeletalMeshInterfaceHelper::GetFilteredBoneCountName("GetFilteredBoneCount");
 const FName FSkeletalMeshInterfaceHelper::GetFilteredBoneAtName("GetFilteredBone");
 
-const FName FSkeletalMeshInterfaceHelper::RandomFilteredSocketBoneName("RandomFilteredSocketBone");
+const FName FSkeletalMeshInterfaceHelper::RandomUnfilteredBoneName("RandomUnfilteredBone");
+const FName FSkeletalMeshInterfaceHelper::GetUnfilteredBoneCountName("GetUnfilteredBoneCount");
+const FName FSkeletalMeshInterfaceHelper::GetUnfilteredBoneAtName("GetUnfilteredBone");
+
+const FName FSkeletalMeshInterfaceHelper::RandomFilteredSocketName("RandomFilteredSocket");
 const FName FSkeletalMeshInterfaceHelper::GetFilteredSocketCountName("GetFilteredSocketCount");
 const FName FSkeletalMeshInterfaceHelper::GetFilteredSocketTransformName("GetFilteredSocketTransform");
 const FName FSkeletalMeshInterfaceHelper::GetFilteredSocketBoneAtName("GetFilteredSocketBone");
@@ -187,18 +190,56 @@ void UNiagaraDataInterfaceSkeletalMesh::GetSkeletonSamplingFunctions(TArray<FNia
 		OutFunctions.Add(Sig);
 	}
 
+	{
+		FNiagaraFunctionSignature Sig;
+		Sig.Name = FSkeletalMeshInterfaceHelper::RandomUnfilteredBoneName;
+		Sig.Inputs.Add(FNiagaraVariable(FNiagaraTypeDefinition(GetClass()), TEXT("SkeletalMesh")));
+		Sig.Inputs.Add(FNiagaraVariable(FNiagaraTypeDefinition(FNiagaraRandInfo::StaticStruct()), TEXT("RandomInfo")));
+		Sig.Outputs.Add(FNiagaraVariable(FNiagaraTypeDefinition::GetIntDef(), TEXT("Bone")));
+		Sig.bMemberFunction = true;
+		Sig.bRequiresContext = false;
+		OutFunctions.Add(Sig);
+	}
+
+	{
+		FNiagaraFunctionSignature Sig;
+		Sig.Name = FSkeletalMeshInterfaceHelper::GetUnfilteredBoneCountName;
+		Sig.Inputs.Add(FNiagaraVariable(FNiagaraTypeDefinition(GetClass()), TEXT("SkeletalMesh")));
+		Sig.Outputs.Add(FNiagaraVariable(FNiagaraTypeDefinition::GetIntDef(), TEXT("Count")));
+		Sig.bMemberFunction = true;
+		Sig.bRequiresContext = false;
+#if WITH_EDITORONLY_DATA
+		Sig.Description = LOCTEXT("GetUnfilteredBoneCountDesc", "Returns the number of unfiltered bones (i.e. the exclusion of filtered bones) in the DI list.");
+#endif
+		OutFunctions.Add(Sig);
+	}
+
+	{
+		FNiagaraFunctionSignature Sig;
+		Sig.Name = FSkeletalMeshInterfaceHelper::GetUnfilteredBoneAtName;
+		Sig.Inputs.Add(FNiagaraVariable(FNiagaraTypeDefinition(GetClass()), TEXT("SkeletalMesh")));
+		Sig.Inputs.Add(FNiagaraVariable(FNiagaraTypeDefinition::GetIntDef(), TEXT("Bone Index")));
+		Sig.Outputs.Add(FNiagaraVariable(FNiagaraTypeDefinition::GetIntDef(), TEXT("Bone")));
+		Sig.bMemberFunction = true;
+		Sig.bRequiresContext = false;
+#if WITH_EDITORONLY_DATA
+		Sig.Description = LOCTEXT("GetUnfilteredBoneAtDesc", "Gets the bone at the passed index from the exlusion of the DI's filter bones list.");
+#endif
+		OutFunctions.Add(Sig);
+	}
+
 	//////////////////////////////////////////////////////////////////////////
 	//Socket functions
 	{
 		FNiagaraFunctionSignature Sig;
-		Sig.Name = FSkeletalMeshInterfaceHelper::RandomFilteredSocketBoneName;
+		Sig.Name = FSkeletalMeshInterfaceHelper::RandomFilteredSocketName;
 		Sig.Inputs.Add(FNiagaraVariable(FNiagaraTypeDefinition(GetClass()), TEXT("SkeletalMesh")));
 		Sig.Inputs.Add(FNiagaraVariable(FNiagaraTypeDefinition(FNiagaraRandInfo::StaticStruct()), TEXT("RandomInfo")));
 		Sig.Outputs.Add(FNiagaraVariable(FNiagaraTypeDefinition::GetIntDef(), TEXT("Socket Bone")));
 		Sig.bMemberFunction = true;
 		Sig.bRequiresContext = false;
 #if WITH_EDITORONLY_DATA
-		Sig.Description = LOCTEXT("RandomFilteredSocketBoneDesc", "Gets the bone for a random socket in the DI's filtered socket list.");
+		Sig.Description = LOCTEXT("RandomFilteredSocketDesc", "Gets the bone for a random socket in the DI's filtered socket list.");
 #endif
 		OutFunctions.Add(Sig);
 	}
@@ -315,13 +356,28 @@ void UNiagaraDataInterfaceSkeletalMesh::BindSkeletonSamplingFunction(const FVMEx
 	else if (BindingInfo.Name == FSkeletalMeshInterfaceHelper::GetFilteredBoneAtName)
 	{
 		check(BindingInfo.GetNumInputs() == 2 && BindingInfo.GetNumOutputs() == 1);
-		NDI_FUNC_BINDER(UNiagaraDataInterfaceSkeletalMesh, GetFilteredBoneAt)::Bind(this, OutFunc);
+		OutFunc = FVMExternalFunction::CreateLambda([this](FVectorVMContext& Context) { this->GetFilteredBoneAt(Context); });
+	}
+	else if (BindingInfo.Name == FSkeletalMeshInterfaceHelper::RandomUnfilteredBoneName)
+	{
+		check(BindingInfo.GetNumInputs() == 4 && BindingInfo.GetNumOutputs() == 1);
+		OutFunc = FVMExternalFunction::CreateLambda([this](FVectorVMContext& Context) { this->RandomUnfilteredBone(Context); });
+	}
+	else if (BindingInfo.Name == FSkeletalMeshInterfaceHelper::GetUnfilteredBoneCountName)
+	{
+		check(BindingInfo.GetNumInputs() == 1 && BindingInfo.GetNumOutputs() == 1);
+		OutFunc = FVMExternalFunction::CreateLambda([this](FVectorVMContext& Context) { this->GetUnfilteredBoneCount(Context); });
+	}
+	else if (BindingInfo.Name == FSkeletalMeshInterfaceHelper::GetUnfilteredBoneAtName)
+	{
+		check(BindingInfo.GetNumInputs() == 2 && BindingInfo.GetNumOutputs() == 1);
+		OutFunc = FVMExternalFunction::CreateLambda([this](FVectorVMContext& Context) { this->GetUnfilteredBoneAt(Context); });
 	}
 	//Socket Functions
-	else if (BindingInfo.Name == FSkeletalMeshInterfaceHelper::RandomFilteredSocketBoneName)
+	else if (BindingInfo.Name == FSkeletalMeshInterfaceHelper::RandomFilteredSocketName)
 	{
 		check(BindingInfo.GetNumInputs() == 4 && BindingInfo.GetNumOutputs() == 1);
-		OutFunc = FVMExternalFunction::CreateLambda([this](FVectorVMContext& Context) { this->RandomFilteredSocketBone(Context); });
+		OutFunc = FVMExternalFunction::CreateLambda([this](FVectorVMContext& Context) { this->RandomFilteredSocket(Context); });
 	}
 	else if (BindingInfo.Name == FSkeletalMeshInterfaceHelper::GetFilteredSocketCountName)
 	{
@@ -356,7 +412,7 @@ void UNiagaraDataInterfaceSkeletalMesh::GetFilteredBoneCount(FVectorVMContext& C
 
 	VectorVM::FExternalFuncRegisterHandler<int32> OutCount(Context);
 
-	int32 Num = FilteredBones.Num();
+	const int32 Num = InstData->NumFilteredBones;
 	for (int32 i = 0; i < Context.NumInstances; ++i)
 	{
 		*OutCount.GetDestAndAdvance() = Num;
@@ -369,15 +425,14 @@ void UNiagaraDataInterfaceSkeletalMesh::GetFilteredBoneAt(FVectorVMContext& Cont
 	VectorVM::FUserPtrHandler<FNDISkeletalMesh_InstanceData> InstData(Context);
 
 	VectorVM::FExternalFuncRegisterHandler<int32> OutBone(Context);
-	const TArray<int32>& FilteredBonesArray = InstData->FilteredBones;
 
-	int32 Max = FilteredBones.Num() - 1;
+	const int32 Max = InstData->NumFilteredBones;
 	if (Max != INDEX_NONE)
 	{
 		for (int32 i = 0; i < Context.NumInstances; ++i)
 		{
-			int32 BoneIndex = FMath::Clamp(BoneParam.GetAndAdvance(), 0, Max);
-			*OutBone.GetDestAndAdvance() = FilteredBonesArray[BoneIndex];
+			const int32 BoneIndex = FMath::Clamp(BoneParam.GetAndAdvance(), 0, Max);
+			*OutBone.GetDestAndAdvance() = InstData->FilteredAndUnfilteredBones[BoneIndex];
 		}
 	}
 	else
@@ -392,16 +447,73 @@ void UNiagaraDataInterfaceSkeletalMesh::RandomFilteredBone(FVectorVMContext& Con
 	VectorVM::FUserPtrHandler<FNDISkeletalMesh_InstanceData> InstData(Context);
 
 	VectorVM::FExternalFuncRegisterHandler<int32> OutBone(Context);
-	const TArray<int32>& FilteredBonesArray = InstData->FilteredBones;
 
-	int32 Max = FilteredBones.Num() - 1;
+	const int32 Max = InstData->NumFilteredBones - 1;
+	if (Max != INDEX_NONE)
+	{
+		for (int32 i = 0; i < Context.NumInstances; ++i)
+		{
+			RandHelper.GetAndAdvance();
+			const int32 BoneIndex = RandHelper.RandRange(i, 0, Max);
+			*OutBone.GetDestAndAdvance() = InstData->FilteredAndUnfilteredBones[BoneIndex];
+		}
+	}
+	else
+	{
+		FMemory::Memset(OutBone.GetDest(), 0xFF, Context.NumInstances * sizeof(int32));
+	}
+}
+
+void UNiagaraDataInterfaceSkeletalMesh::GetUnfilteredBoneCount(FVectorVMContext& Context)
+{
+	SCOPE_CYCLE_COUNTER(STAT_NiagaraSkel_Bone_Sample);
+	VectorVM::FUserPtrHandler<FNDISkeletalMesh_InstanceData> InstData(Context);
+
+	VectorVM::FExternalFuncRegisterHandler<int32> OutCount(Context);
+
+	const int32 Num = InstData->NumUnfilteredBones;
+	for (int32 i = 0; i < Context.NumInstances; ++i)
+	{
+		*OutCount.GetDestAndAdvance() = Num;
+	}
+}
+
+void UNiagaraDataInterfaceSkeletalMesh::GetUnfilteredBoneAt(FVectorVMContext& Context)
+{
+	VectorVM::FExternalFuncInputHandler<int32> BoneParam(Context);
+	VectorVM::FUserPtrHandler<FNDISkeletalMesh_InstanceData> InstData(Context);
+
+	VectorVM::FExternalFuncRegisterHandler<int32> OutBone(Context);
+
+	const int32 Max = InstData->NumUnfilteredBones;
+	if (Max != INDEX_NONE)
+	{
+		for (int32 i = 0; i < Context.NumInstances; ++i)
+		{
+			const int32 BoneIndex = FMath::Clamp(BoneParam.GetAndAdvance(), 0, Max);
+			*OutBone.GetDestAndAdvance() = InstData->FilteredAndUnfilteredBones[BoneIndex + InstData->NumFilteredBones];
+		}
+	}
+	else
+	{
+		FMemory::Memset(OutBone.GetDest(), 0xFF, Context.NumInstances * sizeof(int32));
+	}
+}
+void UNiagaraDataInterfaceSkeletalMesh::RandomUnfilteredBone(FVectorVMContext& Context)
+{
+	FNDIRandomHelper RandHelper(Context);
+	VectorVM::FUserPtrHandler<FNDISkeletalMesh_InstanceData> InstData(Context);
+
+	VectorVM::FExternalFuncRegisterHandler<int32> OutBone(Context);
+
+	const int32 Max = InstData->NumUnfilteredBones - 1;
 	if (Max != INDEX_NONE)
 	{
 		for (int32 i = 0; i < Context.NumInstances; ++i)
 		{
 			RandHelper.GetAndAdvance();
 			const int32 BoneIndex = RandHelper.RandRange(i, 0, Max);
-			*OutBone.GetDestAndAdvance() = FilteredBonesArray[BoneIndex];
+			*OutBone.GetDestAndAdvance() = InstData->FilteredAndUnfilteredBones[BoneIndex + InstData->NumFilteredBones];
 		}
 	}
 	else
@@ -777,7 +889,7 @@ void UNiagaraDataInterfaceSkeletalMesh::GetFilteredSocketTransform(FVectorVMCont
 	}
 }
 
-void UNiagaraDataInterfaceSkeletalMesh::RandomFilteredSocketBone(FVectorVMContext& Context)
+void UNiagaraDataInterfaceSkeletalMesh::RandomFilteredSocket(FVectorVMContext& Context)
 {
 	FNDIRandomHelper RandHelper(Context);
 	VectorVM::FUserPtrHandler<FNDISkeletalMesh_InstanceData> InstData(Context);
@@ -808,10 +920,10 @@ void UNiagaraDataInterfaceSkeletalMesh::RandomFilteredSocketOrBone(FVectorVMCont
 
 	VectorVM::FExternalFuncRegisterHandler<int32> OutBoneIndex(Context);
 
-	const int32 Max = FilteredSockets.Num() + FilteredBones.Num() - 1;
+	const int32 Max = FilteredSockets.Num() + InstData->NumFilteredBones - 1;
 	if (Max >= 0)
 	{
-		const int32 NumFilteredBones = FilteredBones.Num();
+		const int32 NumFilteredBones = InstData->NumFilteredBones;
 		const int32 FilteredSocketBoneOffset = InstData->FilteredSocketBoneOffset;
 		for (int32 i = 0; i < Context.NumInstances; ++i)
 		{
@@ -819,7 +931,7 @@ void UNiagaraDataInterfaceSkeletalMesh::RandomFilteredSocketOrBone(FVectorVMCont
 			const int32 FilteredIndex = RandHelper.RandRange(i, 0, Max);
 			if (FilteredIndex < NumFilteredBones)
 			{
-				*OutBoneIndex.GetDestAndAdvance() = InstData->FilteredBones[FilteredIndex];
+				*OutBoneIndex.GetDestAndAdvance() = InstData->FilteredAndUnfilteredBones[FilteredIndex];
 			}
 			else
 			{
diff --git a/Engine/Plugins/FX/Niagara/Source/Niagara/Private/NiagaraDataInterfaceSkeletalMesh.cpp b/Engine/Plugins/FX/Niagara/Source/Niagara/Private/NiagaraDataInterfaceSkeletalMesh.cpp
--- a/Engine/Plugins/FX/Niagara/Source/Niagara/Private/NiagaraDataInterfaceSkeletalMesh.cpp
+++ b/Engine/Plugins/FX/Niagara/Source/Niagara/Private/NiagaraDataInterfaceSkeletalMesh.cpp
@@ -461,22 +461,20 @@ FSkeletalMeshGpuSpawnStaticBuffers::~FSkeletalMeshGpuSpawnStaticBuffers()
 
 void FSkeletalMeshGpuSpawnStaticBuffers::Initialise(FNDISkeletalMesh_InstanceData* InstData, const FSkeletalMeshLODRenderData& SkeletalMeshLODRenderData, const FSkeletalMeshSamplingLODBuiltData& MeshSamplingLODBuiltData)
 {
-	if (!InstData)
-    {
-		SkeletalMeshSamplingLODBuiltData = nullptr;
-		bUseGpuUniformlyDistributedSampling = false;
-
-		LODRenderData = nullptr;
-		TriangleCount = 0;
-		VertexCount = 0;
-
-		NumFilteredBones = 0;
-		FilteredBonesArray.Empty();
-		NumFilteredSockets = 0;
-		FilteredSocketBoneOffset = 0;
-		return;
-	}
-	else
+	SkeletalMeshSamplingLODBuiltData = nullptr;
+	bUseGpuUniformlyDistributedSampling = false;
+
+	LODRenderData = nullptr;
+	TriangleCount = 0;
+	VertexCount = 0;
+
+	NumFilteredBones = 0;
+	NumUnfilteredBones = 0;
+	FilteredAndUnfilteredBonesArray.Empty();
+	NumFilteredSockets = 0;
+	FilteredSocketBoneOffset = 0;
+
+	if (InstData)
 	{
 		SkeletalMeshSamplingLODBuiltData = &MeshSamplingLODBuiltData;
 		bUseGpuUniformlyDistributedSampling = InstData->bIsGpuUniformlyDistributedSampling;
@@ -496,15 +494,13 @@ void FSkeletalMeshGpuSpawnStaticBuffers::Initialise(FNDISkeletalMesh_InstanceDat
 
 		// Copy filtered Bones / Socket data into arrays that the renderer will use to create read buffers
 		//-TODO: Exclude setting up these arrays if we don't sample from them
-		NumFilteredBones = InstData->FilteredBones.Num();
-		if (NumFilteredBones > 0)
+		NumFilteredBones = InstData->NumFilteredBones;
+		NumUnfilteredBones = InstData->NumUnfilteredBones;
+
+		FilteredAndUnfilteredBonesArray.Reserve(InstData->FilteredAndUnfilteredBones.Num());
+		for (uint16 v : InstData->FilteredAndUnfilteredBones)
 		{
-			FilteredBonesArray.Reserve(NumFilteredBones);
-			for (int32 v : InstData->FilteredBones)
-			{
-				check(v <= 65535);
-				FilteredBonesArray.Add(v);
-			}
+			FilteredAndUnfilteredBonesArray.Add(v);
 		}
 
 		NumFilteredSockets = InstData->FilteredSocketInfo.Num();
@@ -585,20 +581,20 @@ void FSkeletalMeshGpuSpawnStaticBuffers::InitRHI()
 	}
 
 	// Create arrays for filtered bones / sockets
-	if ( NumFilteredBones > 0 )
+	if ( FilteredAndUnfilteredBonesArray.Num() > 0 )
 	{
 		FRHIResourceCreateInfo CreateInfo;
-		CreateInfo.ResourceArray = &FilteredBonesArray;
+		CreateInfo.ResourceArray = &FilteredAndUnfilteredBonesArray;
 
-		FilteredBonesBuffer = RHICreateVertexBuffer(NumFilteredBones * sizeof(uint16), BUF_Static | BUF_ShaderResource, CreateInfo);
-		FilteredBonesSRV = RHICreateShaderResourceView(FilteredBonesBuffer, sizeof(uint16), PF_R16_UINT);
+		FilteredAndUnfilteredBonesBuffer = RHICreateVertexBuffer(FilteredAndUnfilteredBonesArray.Num() * FilteredAndUnfilteredBonesArray.GetTypeSize(), BUF_Static | BUF_ShaderResource, CreateInfo);
+		FilteredAndUnfilteredBonesSRV = RHICreateShaderResourceView(FilteredAndUnfilteredBonesBuffer, sizeof(uint16), PF_R16_UINT);
 	}
 }
 
 void FSkeletalMeshGpuSpawnStaticBuffers::ReleaseRHI()
 {
-	FilteredBonesBuffer.SafeRelease();
-	FilteredBonesSRV.SafeRelease();
+	FilteredAndUnfilteredBonesBuffer.SafeRelease();
+	FilteredAndUnfilteredBonesSRV.SafeRelease();
 
 	BufferTriangleUniformSamplerProbaRHI.SafeRelease();
 	BufferTriangleUniformSamplerProbaSRV.SafeRelease();
@@ -866,7 +862,8 @@ struct FNDISkeletalMeshParametersName
 	FString MeshNumWeightsName;
 	FString NumBonesName;
 	FString NumFilteredBonesName;
-	FString FilteredBonesName;
+	FString NumUnfilteredBonesName;
+	FString FilteredAndUnfilteredBonesName;
 	FString NumFilteredSocketsName;
 	FString FilteredSocketBoneOffsetName;
 	FString InstanceTransformName;
@@ -899,7 +896,8 @@ static void GetNiagaraDataInterfaceParametersName(FNDISkeletalMeshParametersName
 	Names.MeshNumWeightsName = UNiagaraDataInterfaceSkeletalMesh::MeshNumWeightsName + Suffix;
 	Names.NumBonesName = UNiagaraDataInterfaceSkeletalMesh::NumBonesName + Suffix;
 	Names.NumFilteredBonesName = UNiagaraDataInterfaceSkeletalMesh::NumFilteredBonesName + Suffix;
-	Names.FilteredBonesName = UNiagaraDataInterfaceSkeletalMesh::FilteredBonesName + Suffix;
+	Names.NumUnfilteredBonesName = UNiagaraDataInterfaceSkeletalMesh::NumUnfilteredBonesName + Suffix;
+	Names.FilteredAndUnfilteredBonesName = UNiagaraDataInterfaceSkeletalMesh::FilteredAndUnfilteredBonesName + Suffix;
 	Names.NumFilteredSocketsName = UNiagaraDataInterfaceSkeletalMesh::NumFilteredSocketsName + Suffix;
 	Names.FilteredSocketBoneOffsetName = UNiagaraDataInterfaceSkeletalMesh::FilteredSocketBoneOffsetName + Suffix;
 	Names.InstanceTransformName = UNiagaraDataInterfaceSkeletalMesh::InstanceTransformName + Suffix;
@@ -939,7 +937,8 @@ public:
 		MeshNumWeights.Bind(ParameterMap, *ParamNames.MeshNumWeightsName);
 		NumBones.Bind(ParameterMap, *ParamNames.NumBonesName);
 		NumFilteredBones.Bind(ParameterMap, *ParamNames.NumFilteredBonesName);
-		FilteredBones.Bind(ParameterMap, *ParamNames.FilteredBonesName);
+		NumUnfilteredBones.Bind(ParameterMap, *ParamNames.NumUnfilteredBonesName);
+		FilteredAndUnfilteredBones.Bind(ParameterMap, *ParamNames.FilteredAndUnfilteredBonesName);
 		NumFilteredSockets.Bind(ParameterMap, *ParamNames.NumFilteredSocketsName);
 		FilteredSocketBoneOffset.Bind(ParameterMap, *ParamNames.FilteredSocketBoneOffsetName);
 		InstanceTransform.Bind(ParameterMap, *ParamNames.InstanceTransformName);
@@ -1026,9 +1025,10 @@ public:
 
 			SetShaderValue(RHICmdList, ComputeShaderRHI, NumBones, DynamicBuffers->GetNumBones());
 
-			FRHIShaderResourceView* FilteredBonesSRV = StaticBuffers->GetNumFilteredBones() > 0 ? StaticBuffers->GetFilteredBonesSRV() : FNiagaraRenderer::GetDummyUIntBuffer();
+			FRHIShaderResourceView* FilteredAndUnfilteredBonesSRV = StaticBuffers->GetNumFilteredBones() > 0 ? StaticBuffers->GetFilteredAndUnfilteredBonesSRV() : FNiagaraRenderer::GetDummyUIntBuffer();
 			SetShaderValue(RHICmdList, ComputeShaderRHI, NumFilteredBones, StaticBuffers->GetNumFilteredBones());
-			SetSRVParameter(RHICmdList, ComputeShaderRHI, FilteredBones, FilteredBonesSRV);
+			SetShaderValue(RHICmdList, ComputeShaderRHI, NumUnfilteredBones, StaticBuffers->GetNumUnfilteredBones());
+			SetSRVParameter(RHICmdList, ComputeShaderRHI, FilteredAndUnfilteredBones, FilteredAndUnfilteredBonesSRV);
 
 			SetShaderValue(RHICmdList, ComputeShaderRHI, NumFilteredSockets, StaticBuffers->GetNumFilteredSockets());
 			SetShaderValue(RHICmdList, ComputeShaderRHI, FilteredSocketBoneOffset, StaticBuffers->GetFilteredSocketBoneOffset());
@@ -1071,7 +1071,8 @@ public:
 
 			SetShaderValue(RHICmdList, ComputeShaderRHI, NumBones, 0);
 			SetShaderValue(RHICmdList, ComputeShaderRHI, NumFilteredBones, 0);
-			SetSRVParameter(RHICmdList, ComputeShaderRHI, FilteredBones, FNiagaraRenderer::GetDummyUIntBuffer());
+			SetShaderValue(RHICmdList, ComputeShaderRHI, NumUnfilteredBones, 0);
+			SetSRVParameter(RHICmdList, ComputeShaderRHI, FilteredAndUnfilteredBones, FNiagaraRenderer::GetDummyUIntBuffer());
 			SetShaderValue(RHICmdList, ComputeShaderRHI, NumFilteredSockets, 0);
 			SetShaderValue(RHICmdList, ComputeShaderRHI, FilteredSocketBoneOffset, 0);
 
@@ -1106,7 +1107,8 @@ private:
 	LAYOUT_FIELD(FShaderParameter, MeshNumWeights);
 	LAYOUT_FIELD(FShaderParameter, NumBones);
 	LAYOUT_FIELD(FShaderParameter, NumFilteredBones);
-	LAYOUT_FIELD(FShaderResourceParameter, FilteredBones);
+	LAYOUT_FIELD(FShaderParameter, NumUnfilteredBones);
+	LAYOUT_FIELD(FShaderResourceParameter, FilteredAndUnfilteredBones);
 	LAYOUT_FIELD(FShaderParameter, NumFilteredSockets);
 	LAYOUT_FIELD(FShaderParameter, FilteredSocketBoneOffset);
 	LAYOUT_FIELD(FShaderParameter, InstanceTransform);
@@ -1500,21 +1502,51 @@ bool FNDISkeletalMesh_InstanceData::Init(UNiagaraDataInterfaceSkeletalMesh* Inte
 
 	// Gather filtered bones information
 	FReferenceSkeleton& RefSkel = Mesh->RefSkeleton;
+	if (Interface->FilteredBones.Num() > 0)
 	{
-		FilteredBones.SetNumUninitialized(Interface->FilteredBones.Num());
+		if (RefSkel.GetNum() > TNumericLimits<uint16>::Max())
+		{
+			UE_LOG(LogNiagara, Warning, TEXT("Skeletal Mesh Data Interface '%s' requires more bones '%d' than we currently support '%d' Mesh '%s' Component '%s'"), *Interface->GetFullName(), RefSkel.GetNum(), TNumericLimits<uint16>::Max(), *Mesh->GetFullName(), *Component->GetFullName());
+			return false;
+		}
+
+		//-TODO: If the DI does not use unfiltered bones we can skip adding them here...
 		TArray<FName, TInlineAllocator<16>> MissingBones;
-		for (int32 BoneIdx = 0; BoneIdx < FilteredBones.Num(); ++BoneIdx)
+
+		FilteredAndUnfilteredBones.Reserve(RefSkel.GetNum());
+
+		// Append filtered bones to the array first
+		for (const FName& BoneName : Interface->FilteredBones)
 		{
-			FName BoneName = Interface->FilteredBones[BoneIdx];
-			int32 Bone = RefSkel.FindBoneIndex(BoneName);
+			const int32 Bone = RefSkel.FindBoneIndex(BoneName);
 			if (Bone == INDEX_NONE)
 			{
 				MissingBones.Add(BoneName);
-				FilteredBones[BoneIdx] = 0;
 			}
 			else
 			{
-				FilteredBones[BoneIdx] = Bone;
+				ensure(Bone <= TNumericLimits<uint16>::Max());
+				FilteredAndUnfilteredBones.Add(Bone);
+				++NumFilteredBones;
+			}
+		}
+
+		// Append unfiltered bones to the array
+		for (int32 i=0; i < RefSkel.GetNum(); ++i)
+		{
+			bool bExists = false;
+			for (int32 j=0; j < NumFilteredBones; ++j)
+			{
+				if (FilteredAndUnfilteredBones[j] == i)
+				{
+					bExists = true;
+					break;
+				}
+			}
+			if (!bExists)
+			{
+				FilteredAndUnfilteredBones.Add(i);
+				++NumUnfilteredBones;
 			}
 		}
 
@@ -2166,7 +2198,8 @@ const FString UNiagaraDataInterfaceSkeletalMesh::MeshNumTexCoordName(TEXT("MeshN
 const FString UNiagaraDataInterfaceSkeletalMesh::MeshNumWeightsName(TEXT("MeshNumWeights_"));
 const FString UNiagaraDataInterfaceSkeletalMesh::NumBonesName(TEXT("NumBones_"));
 const FString UNiagaraDataInterfaceSkeletalMesh::NumFilteredBonesName(TEXT("NumFilteredBones_"));
-const FString UNiagaraDataInterfaceSkeletalMesh::FilteredBonesName(TEXT("FilteredBones_"));
+const FString UNiagaraDataInterfaceSkeletalMesh::NumUnfilteredBonesName(TEXT("NumUnfilteredBones_"));
+const FString UNiagaraDataInterfaceSkeletalMesh::FilteredAndUnfilteredBonesName(TEXT("FilteredAndUnfilteredBones_"));
 const FString UNiagaraDataInterfaceSkeletalMesh::NumFilteredSocketsName(TEXT("NumFilteredSockets_"));
 const FString UNiagaraDataInterfaceSkeletalMesh::FilteredSocketBoneOffsetName(TEXT("FilteredSocketBoneOffset_"));
 const FString UNiagaraDataInterfaceSkeletalMesh::InstanceTransformName(TEXT("InstanceTransform_"));
@@ -2334,6 +2367,21 @@ bool UNiagaraDataInterfaceSkeletalMesh::GetFunctionHLSL(const FNiagaraDataInterf
 		static const TCHAR* FormatSample = TEXT("void {InstanceFunctionName} (NiagaraRandInfo InRandomInfo, out int Bone) { {GetDISkelMeshContextName} DISkelMesh_RandomFilteredBone(DIContext, InRandomInfo.Seed1, InRandomInfo.Seed2, InRandomInfo.Seed3, Bone); }");
 		OutHLSL += FString::Format(FormatSample, ArgsSample);
 	}
+	else if (FunctionInfo.DefinitionName == FSkeletalMeshInterfaceHelper::GetUnfilteredBoneAtName)
+	{
+		static const TCHAR* FormatSample = TEXT("void {InstanceFunctionName} (out int Count) { {GetDISkelMeshContextName} DISkelMesh_GetUnfilteredBoneCount(DIContext, Count); }");
+		OutHLSL += FString::Format(FormatSample, ArgsSample);
+	}
+	else if (FunctionInfo.DefinitionName == FSkeletalMeshInterfaceHelper::GetUnfilteredBoneAtName)
+	{
+		static const TCHAR* FormatSample = TEXT("void {InstanceFunctionName} (in int BoneIndex, out int Bone) { {GetDISkelMeshContextName} DISkelMesh_GetUnfilteredBoneAt(DIContext, BoneIndex, Bone); }");
+		OutHLSL += FString::Format(FormatSample, ArgsSample);
+	}
+	else if (FunctionInfo.DefinitionName == FSkeletalMeshInterfaceHelper::RandomUnfilteredBoneName)
+	{
+		static const TCHAR* FormatSample = TEXT("void {InstanceFunctionName} (NiagaraRandInfo InRandomInfo, out int Bone) { {GetDISkelMeshContextName} DISkelMesh_RandomUnfilteredBone(DIContext, InRandomInfo.Seed1, InRandomInfo.Seed2, InRandomInfo.Seed3, Bone); }");
+		OutHLSL += FString::Format(FormatSample, ArgsSample);
+	}
 	//////////////////////////////////////////////////////////////////////////////////////////////////////////////////////////////
 	// Filtered Socket
 	else if (FunctionInfo.DefinitionName == FSkeletalMeshInterfaceHelper::GetFilteredSocketCountName)
@@ -2353,9 +2401,9 @@ bool UNiagaraDataInterfaceSkeletalMesh::GetFunctionHLSL(const FNiagaraDataInterf
 		static const TCHAR* FormatSample = TEXT("void {InstanceFunctionName} (in int SocketIndex, in int bShouldApplyTransform, out float3 Translation, out float4 Rotation, out float4 Scale) { Translation = float3(0.0, 0.0, 0.0); Rotation = float4(0.0, 0.0, 0.0, 1.0); Scale = float3(1.0, 1.0, 1.0); }");
 		OutHLSL += FString::Format(FormatSample, ArgsSample);
 	} 
-	else if (FunctionInfo.DefinitionName == FSkeletalMeshInterfaceHelper::RandomFilteredSocketBoneName)
+	else if (FunctionInfo.DefinitionName == FSkeletalMeshInterfaceHelper::RandomFilteredSocketName)
 	{
-		static const TCHAR* FormatSample = TEXT("void {InstanceFunctionName} (NiagaraRandInfo InRandomInfo, out int SocketBone) { {GetDISkelMeshContextName} DISkelMesh_RandomFilteredSocketBone(DIContext, InRandomInfo.Seed1, InRandomInfo.Seed2, InRandomInfo.Seed3, SocketBone); }");
+		static const TCHAR* FormatSample = TEXT("void {InstanceFunctionName} (NiagaraRandInfo InRandomInfo, out int SocketBone) { {GetDISkelMeshContextName} DISkelMesh_RandomFilteredSocket(DIContext, InRandomInfo.Seed1, InRandomInfo.Seed2, InRandomInfo.Seed3, SocketBone); }");
 		OutHLSL += FString::Format(FormatSample, ArgsSample);
 	}
 	//////////////////////////////////////////////////////////////////////////////////////////////////////////////////////////////
@@ -2391,10 +2439,11 @@ bool UNiagaraDataInterfaceSkeletalMesh::UpgradeFunctionCall(FNiagaraFunctionSign
 		MakeTuple(FName("RandomSpecificBone"), FSkeletalMeshInterfaceHelper::RandomFilteredBoneName),
 		MakeTuple(FName("GetSpecificBoneCount"), FSkeletalMeshInterfaceHelper::GetFilteredBoneCountName),
 		MakeTuple(FName("GetSpecificBone"), FSkeletalMeshInterfaceHelper::GetFilteredBoneAtName),
-		MakeTuple(FName("RandomSpecificSocketBone"), FSkeletalMeshInterfaceHelper::RandomFilteredSocketBoneName),
+		MakeTuple(FName("RandomSpecificSocketBone"), FSkeletalMeshInterfaceHelper::RandomFilteredSocketName),
 		MakeTuple(FName("GetSpecificSocketCount"), FSkeletalMeshInterfaceHelper::GetFilteredSocketCountName),
 		MakeTuple(FName("GetSpecificSocketTransform"), FSkeletalMeshInterfaceHelper::GetFilteredSocketTransformName),
 		MakeTuple(FName("GetSpecificSocketBone"), FSkeletalMeshInterfaceHelper::GetFilteredSocketBoneAtName),
+		MakeTuple(FName("RandomFilteredSocketBone"), FSkeletalMeshInterfaceHelper::RandomFilteredSocketName),
 	};
 	bool bWasChanged = false;
 
@@ -2424,7 +2473,7 @@ bool UNiagaraDataInterfaceSkeletalMesh::UpgradeFunctionCall(FNiagaraFunctionSign
 			FunctionSignature.Inputs.Add(FNiagaraVariable(FNiagaraTypeDefinition(FNiagaraRandInfo::StaticStruct()), TEXT("RandomInfo")));
 			bWasChanged = true;
 		}
-		else if (FunctionSignature.Name == FSkeletalMeshInterfaceHelper::RandomFilteredSocketBoneName)
+		else if (FunctionSignature.Name == FSkeletalMeshInterfaceHelper::RandomFilteredSocketName)
 		{
 			FunctionSignature.Inputs.Add(FNiagaraVariable(FNiagaraTypeDefinition(FNiagaraRandInfo::StaticStruct()), TEXT("RandomInfo")));
 			bWasChanged = true;
@@ -2435,7 +2484,6 @@ bool UNiagaraDataInterfaceSkeletalMesh::UpgradeFunctionCall(FNiagaraFunctionSign
 			bWasChanged = true;
 		}
 	}
-	FunctionSignature.FunctionVersion = FNiagaraSkelMeshDIFunctionVersion::LatestVersion;
 
 	return bWasChanged;
 }
