Fix for ParameterStore crash when trying to rename a parameter to an existing name

diff --git a/Engine/Plugins/FX/Niagara/Source/Niagara/Private/NiagaraParameterStore.cpp b/Engine/Plugins/FX/Niagara/Source/Niagara/Private/NiagaraParameterStore.cpp
--- a/Engine/Plugins/FX/Niagara/Source/Niagara/Private/NiagaraParameterStore.cpp
+++ b/Engine/Plugins/FX/Niagara/Source/Niagara/Private/NiagaraParameterStore.cpp
@@ -637,28 +637,34 @@ void FNiagaraParameterStore::RenameParameter(const FNiagaraVariableBase& Param,
 
 		int32 NumBytesBefore = ParameterData.Num();
 		int32 NewIdx = INDEX_NONE;
-		AddParameter(NewParam, bInitInterfaces, bTriggerRebind, &NewIdx);
-		int32 NumBytesAfter = ParameterData.Num();
-		int32 DeltaBytes = NumBytesAfter - NumBytesBefore;
-		int32 SizeInBytes = Param.GetSizeInBytes();
+		if (AddParameter(NewParam, bInitInterfaces, bTriggerRebind, &NewIdx))
+		{
+			int32 NumBytesAfter = ParameterData.Num();
+			int32 DeltaBytes = NumBytesAfter - NumBytesBefore;
+			int32 SizeInBytes = Param.GetSizeInBytes();
 
-		check(DeltaBytes == SizeInBytes);
+			check(DeltaBytes == SizeInBytes);
 
-		if (Param.IsDataInterface())
-		{
-			SetDataInterface(GetDataInterface(Idx), NewIdx);
-		}
-		else if (Param.IsUObject())
-		{
-			SetUObject(GetUObject(Idx), NewIdx);
+			if (Param.IsDataInterface())
+			{
+				SetDataInterface(GetDataInterface(Idx), NewIdx);
+			}
+			else if (Param.IsUObject())
+			{
+				SetUObject(GetUObject(Idx), NewIdx);
+			}
+			else
+			{
+				SetParameterData(GetParameterData_Internal(Idx), NewIdx, Param.GetSizeInBytes());
+			}
+			RemoveParameter(Param);
+
+			OnLayoutChange();
 		}
 		else
 		{
-			SetParameterData(GetParameterData_Internal(Idx), NewIdx, Param.GetSizeInBytes());
+			UE_LOG(LogNiagara, Warning, TEXT("Ignored attempt to rename a parameter overtop of an existing parameter!  Old name: %s, New name: %s"), *Param.GetName().ToString(), *NewName.ToString());
 		}
-		RemoveParameter(Param);
-
-		OnLayoutChange();
 	}
 }
 
