Fix for NDC issue resulting from the non-deterministic ordering of the PostLoad of an unused GPUComputeScript resulting in the source graph getting recached

diff --git a/Engine/Plugins/FX/Niagara/Source/Niagara/Classes/NiagaraEmitter.h b/Engine/Plugins/FX/Niagara/Source/Niagara/Classes/NiagaraEmitter.h
--- a/Engine/Plugins/FX/Niagara/Source/Niagara/Classes/NiagaraEmitter.h
+++ b/Engine/Plugins/FX/Niagara/Source/Niagara/Classes/NiagaraEmitter.h
@@ -577,6 +577,8 @@ private:
 	FOnRenderersChanged OnRenderersChangedDelegate;
 #endif
 
+	void EnsureScriptsPostLoaded();
+
 	void GenerateStatID()const;
 #if STATS
 	mutable TStatId StatID_GT;
diff --git a/Engine/Plugins/FX/Niagara/Source/Niagara/Private/NiagaraEmitter.cpp b/Engine/Plugins/FX/Niagara/Source/Niagara/Private/NiagaraEmitter.cpp
--- a/Engine/Plugins/FX/Niagara/Source/Niagara/Private/NiagaraEmitter.cpp
+++ b/Engine/Plugins/FX/Niagara/Source/Niagara/Private/NiagaraEmitter.cpp
@@ -297,6 +297,27 @@ void UNiagaraEmitter::Serialize(FArchive& Ar)
 	Ar.UsingCustomVersion(FNiagaraCustomVersion::GUID);
 }
 
+void UNiagaraEmitter::EnsureScriptsPostLoaded()
+{
+	TArray<UNiagaraScript*> AllScripts;
+	GetScripts(AllScripts, false);
+
+	// Post load scripts for use below.
+	for (UNiagaraScript* Script : AllScripts)
+	{
+		Script->ConditionalPostLoad();
+	}
+
+	// Additionally we want to make sure that the GPUComputeScript, if it exists, is also post loaded immediately even if we're not using it.
+	// Currently an unused GPUComputeScript will cause the cached data to be invalidated and rebuilt because it will never get
+	// a valid CompilerVersionID assigned to it (since it's not being compiled because it's not being used).  The side effect of this is that
+	// the invalidation occurs in a non-deterministic location (based on PostLoad order) and can mess up with the cooking process
+	if (GPUComputeScript)
+	{
+		GPUComputeScript->ConditionalPostLoad();
+	}
+}
+
 void UNiagaraEmitter::PostLoad()
 {
 	Super::PostLoad();
@@ -436,14 +457,7 @@ void UNiagaraEmitter::PostLoad()
 	}
 #endif
 
-	TArray<UNiagaraScript*> AllScripts;
-	GetScripts(AllScripts, false);
-
-	// Post load scripts for use below.
-	for (UNiagaraScript* Script : AllScripts)
-	{
-		Script->ConditionalPostLoad();
-	}
+	EnsureScriptsPostLoaded();
 
 #if WITH_EDITORONLY_DATA
 	if (!GetOutermost()->bIsCookedForEditor)
@@ -493,6 +507,9 @@ void UNiagaraEmitter::PostLoad()
 		}
 		else
 		{
+			TArray<UNiagaraScript*> AllScripts;
+			GetScripts(AllScripts, false);
+
 			for (UNiagaraScript* Script : AllScripts)
 			{
 				if (Script->AreScriptAndSourceSynchronized() == false)
