- Fix an issue with GetFilteredBoneAt going OOB on the array
- Fix unfiltered bone access when we have no bones filtered
- Add support for triangle region sampling on the GPU
- Add new functionality to get triangle coordinates on CPU & GPU with respective support for getting counts
- Make vertex sampling functionality consistent with the rest of the DI
- Add support for vertex region samping on the GPU

#rnx

diff --git a/Engine/Plugins/FX/Niagara/Shaders/Private/NiagaraDataInterfaceSkeletalMesh.ush b/Engine/Plugins/FX/Niagara/Shaders/Private/NiagaraDataInterfaceSkeletalMesh.ush
--- a/Engine/Plugins/FX/Niagara/Shaders/Private/NiagaraDataInterfaceSkeletalMesh.ush
+++ b/Engine/Plugins/FX/Niagara/Shaders/Private/NiagaraDataInterfaceSkeletalMesh.ush
@@ -25,6 +25,12 @@ struct FDISkelMeshContext
     Buffer<float4>	MeshColorBuffer;
 	Buffer<float>	MeshTriangleSamplerProbaBuffer;
 	Buffer<uint>	MeshTriangleSamplerAliasBuffer;
+	uint			MeshNumSamplingRegionTriangles;
+	uint			MeshNumSamplingRegionVertices;
+	Buffer<uint>	MeshSamplingRegionsProbaBuffer;
+	Buffer<uint>	MeshSamplingRegionsAliasBuffer;
+	Buffer<uint>	MeshSampleRegionsTriangleIndices;
+	Buffer<uint>	MeshSampleRegionsVertices;
 	Buffer<uint>	MeshTriangleMatricesOffsetBuffer;
 	uint			MeshTriangleCount;
 	uint			MeshVertexCount;
@@ -42,6 +48,7 @@ struct FDISkelMeshContext
 	float4x4		InstancePrevTransform;
 	float			InstanceInvDeltaTime;
 	bool			UniformTriangleSamplingEnable;
+	bool			UniformSamplingRegionEnabled;
 	bool			UnlimitedBoneInfluences;
 };
 
@@ -57,6 +64,15 @@ struct FDISkelMeshSkinnedVertex
 	float3	TangentZ;
 };
 
+float3 DISKelMesh_RandomBarycentricCoord(uint Seed1, uint Seed2, uint Seed3)
+{
+	float r0 = NiagaraInternalNoise(Seed1, Seed2, Seed3);
+	float r1 = NiagaraInternalNoise(Seed1, Seed2, Seed3);
+	float sqrt0 = sqrt(r0);
+	float sqrt1 = sqrt(r1);
+	return float3(1.0f - sqrt0, sqrt0 * (1.0 - r1), r1 * sqrt0);
+}
+
 float3x4 DISkelMesh_GetPrevBoneSkinningMatrix(in FDISkelMeshContext DIContext, uint Bone)
 {
     return float3x4(DIContext.MeshPrevBonesBuffer[Bone * 3], DIContext.MeshPrevBonesBuffer[Bone * 3 + 1], DIContext.MeshPrevBonesBuffer[Bone * 3 + 2]);
@@ -248,11 +264,7 @@ void DISKelMesh_RandomTriCoord(in FDISkelMeshContext DIContext, uint Seed1, uint
 		OutTriangle = TriangleIndex;
 	}
 
-	float r0 = NiagaraInternalNoise(Seed1, Seed2, Seed3);
-	float r1 = NiagaraInternalNoise(Seed1, Seed2, Seed3);
-	float sqrt0 = sqrt(r0);
-	float sqrt1 = sqrt(r1);
-	OutBaryCoord = float3(1.0f - sqrt0, sqrt0 * (1.0 - r1), r1 * sqrt0);
+	OutBaryCoord = DISKelMesh_RandomBarycentricCoord(Seed1, Seed2, Seed3);
 }
 
 void DISKelMesh_GetSkinnedTriangleDataWS(in FDISkelMeshContext DIContext, in uint TriangleIndex, in float3 BaryCoord, out float3 OutPosition, out float3 OutVelocity, out float3 OutNormal, out float3 OutBinormal, out float3 OutTangent)
@@ -318,6 +330,56 @@ void DISKelMesh_GetTriUV(in FDISkelMeshContext DIContext, in uint TriangleIndex,
 	}
 }
 
+void DISKelMesh_RandomTriangle(in FDISkelMeshContext DIContext, uint Seed1, uint Seed2, uint Seed3, out int OutTriangle, out float3 OutBaryCoord)
+{
+	float RandT0 = NiagaraInternalNoise(Seed1, Seed2, Seed3);
+	OutTriangle = min(uint(RandT0 * float(DIContext.MeshTriangleCount)), DIContext.MeshTriangleCount - 1); // avoid % by using mul/min to Tri = MeshTriangleCountName
+	OutBaryCoord = DISKelMesh_RandomBarycentricCoord(Seed1, Seed2, Seed3);
+}
+
+void DISKelMesh_GetTriangleCount(in FDISkelMeshContext DIContext, out int OutCount)
+{
+	OutCount = DIContext.MeshTriangleCount;
+}
+
+void DISKelMesh_RandomFilteredTriangle(in FDISkelMeshContext DIContext, uint Seed1, uint Seed2, uint Seed3, out int OutTriangle, out float3 OutBaryCoord)
+{
+	float RandT0 = NiagaraInternalNoise(Seed1, Seed2, Seed3);
+	uint RegionTriangle = min(uint(RandT0 * float(DIContext.MeshNumSamplingRegionTriangles)), DIContext.MeshNumSamplingRegionTriangles - 1); // avoid % by using mul/min to Tri = MeshTriangleCountName
+
+	[branch]
+	if ( DIContext.UniformSamplingRegionEnabled )
+	{
+		float TriangleProbability = DIContext.MeshSamplingRegionsProbaBuffer[RegionTriangle];
+		float RandT1 = NiagaraInternalNoise(Seed1, Seed2, Seed3);
+		if ( RandT1 > TriangleProbability )
+		{
+			RegionTriangle = (int)DIContext.MeshSamplingRegionsAliasBuffer[RegionTriangle];
+		}
+	}
+	OutTriangle = (int)DIContext.MeshSampleRegionsTriangleIndices[RegionTriangle];
+	OutBaryCoord = DISKelMesh_RandomBarycentricCoord(Seed1, Seed2, Seed3);
+}
+
+void DISKelMesh_GetFilteredTriangleCount(in FDISkelMeshContext DIContext, out int OutCount)
+{
+	OutCount = DIContext.MeshNumSamplingRegionTriangles;
+}
+
+void DISKelMesh_GetFilteredTriangleAt(in FDISkelMeshContext DIContext, int FilteredIndex, out int OutTriangle, out float3 OutBaryCoord)
+{
+	if ( DIContext.MeshNumSamplingRegionTriangles > 0 )
+	{
+		FilteredIndex = clamp(FilteredIndex, 0, (int)DIContext.MeshNumSamplingRegionTriangles - 1);
+		OutTriangle = (int)DIContext.MeshSampleRegionsTriangleIndices[FilteredIndex];
+	}
+	else
+	{
+		OutTriangle = 0;
+	}
+	OutBaryCoord = 1.0f / 3.0f;
+}
+
 ////////////////////////////////////////////////////////////////////////////////////////////////////////////////////////////////////////////////////////////////
 // Bone Sampling
 
@@ -405,17 +467,6 @@ void DISkelMesh_GetTriColor(in FDISkelMeshContext DIContext, in uint TriangleInd
     OutColor = (DIContext.MeshColorBuffer[VertexIndex0] * BaryCoord.x) + (DIContext.MeshColorBuffer[VertexIndex1] * BaryCoord.y) + (DIContext.MeshColorBuffer[VertexIndex2] * BaryCoord.z);
 }
 
-void DISkelMesh_GetRandomVertex(in FDISkelMeshContext DIContext, uint Seed1, uint Seed2, uint Seed3, out int OutVertex)
-{
-	float RandT0 = NiagaraInternalNoise(Seed1, Seed2, Seed3);
-    OutVertex = min(uint(RandT0 * float(DIContext.MeshVertexCount)), DIContext.MeshVertexCount - 1); // avoid % by using mul/min to Tri = MeshVertexCount
-}
-
-void DISkelMesh_IsValidVertex(in FDISkelMeshContext DIContext, in int Vertex, out bool IsValid)
-{
-    IsValid = Vertex < (int)DIContext.MeshVertexCount;
-}
-
 void DISkelMesh_GetSkinnedVertex(in FDISkelMeshContext DIContext, in int VertexIndex, out float3 OutPosition, out float3 OutVelocity)
 {
     FDISkelMeshSkinnedVertex SkinnedVertex = DISkelMesh_GetSkinnedVertex(DIContext, VertexIndex);
@@ -450,6 +501,45 @@ void DISkelMesh_GetVertexUV(in FDISkelMeshContext DIContext, in int VertexIndex,
     }
 }
 
+void DISkelMesh_IsValidVertex(in FDISkelMeshContext DIContext, int Vertex, out bool IsValid)
+{
+    IsValid = Vertex < (int)DIContext.MeshVertexCount;
+}
+
+void DISkelMesh_GetRandomVertex(in FDISkelMeshContext DIContext, uint Seed1, uint Seed2, uint Seed3, out int OutVertex)
+{
+	float RandT0 = NiagaraInternalNoise(Seed1, Seed2, Seed3);
+	OutVertex = (int)min(uint(RandT0 * float(DIContext.MeshVertexCount)), DIContext.MeshVertexCount - 1); // avoid % by using mul/min to Tri = MeshVertexCount
+}
+
+void DISkelMesh_GetVertexCount(in FDISkelMeshContext DIContext, out int VertexCount)
+{
+	VertexCount = (int)DIContext.MeshVertexCount;
+}
+
+void DISkelMesh_IsValidFilteredVertex(in FDISkelMeshContext DIContext, int FilteredIndex, out bool IsValid)
+{
+	IsValid = FilteredIndex < (int)DIContext.MeshNumSamplingRegionVertices;
+}
+
+void DISkelMesh_GetRandomFilteredVertex(in FDISkelMeshContext DIContext, uint Seed1, uint Seed2, uint Seed3, out int OutVertex)
+{
+	float RandT0 = NiagaraInternalNoise(Seed1, Seed2, Seed3);
+	int FilteredIndex = (int)min(uint(RandT0 * float(DIContext.MeshNumSamplingRegionVertices)), DIContext.MeshNumSamplingRegionVertices - 1); // avoid % by using mul/min to Tri = MeshVertexCount
+	OutVertex = (int)DIContext.MeshSampleRegionsVertices[FilteredIndex];
+}
+
+void DISkelMesh_GetFilteredVertexCount(in FDISkelMeshContext DIContext, out int VertexCount)
+{
+	VertexCount = (int)DIContext.MeshNumSamplingRegionVertices;
+}
+
+void DISkelMesh_GetFilteredVertexAt(in FDISkelMeshContext DIContext, int FilteredIndex, out int VertexIndex)
+{
+	FilteredIndex = clamp(FilteredIndex, 0, (int)DIContext.MeshNumSamplingRegionVertices - 1);
+	VertexIndex = (int)DIContext.MeshSampleRegionsVertices[FilteredIndex];
+}
+
 ////////////////////////////////////////////////////////////////////////////////////////////////////////////////////////////////////////////////////////////////
 // Filtered Bones / Sockets
 void DISkelMesh_IsValidBone(in FDISkelMeshContext DIContext, in int BoneIndex, out bool IsValid)
@@ -494,14 +584,14 @@ void DISkelMesh_GetUnfilteredBoneCount(in FDISkelMeshContext DIContext, out int
 void DISkelMesh_GetUnfilteredBoneAt(in FDISkelMeshContext DIContext, in int InBoneIndex, out int Bone)
 {
     int BoneIndex = min(InBoneIndex, DIContext.NumUnfilteredBones - 1);
-    Bone = (int)DIContext.FilteredAndUnfilteredBones[BoneIndex + DIContext.NumFilteredBones];
+    Bone = DIContext.NumFilteredBones > 0 ? (int)DIContext.FilteredAndUnfilteredBones[BoneIndex + DIContext.NumFilteredBones] : BoneIndex;
 }
 
 void DISkelMesh_RandomUnfilteredBone(in FDISkelMeshContext DIContext, uint Seed1, uint Seed2, uint Seed3, out int Bone)
 {
 	float RandT0 = NiagaraInternalNoise(Seed1, Seed2, Seed3);
     int UnfilteredBone = min(int(RandT0 * float(DIContext.NumUnfilteredBones)), DIContext.NumUnfilteredBones - 1); // avoid % by using mul/min to Tri = MeshVertexCount
-    Bone = (int)DIContext.FilteredAndUnfilteredBones[UnfilteredBone + DIContext.NumFilteredBones];
+    Bone = DIContext.NumFilteredBones > 0 ? (int)DIContext.FilteredAndUnfilteredBones[UnfilteredBone + DIContext.NumFilteredBones] : UnfilteredBone;
 }
 
 void DISkelMesh_GetFilteredSocketCount(in FDISkelMeshContext DIContext, out int Count)
@@ -551,6 +641,12 @@ void DISkelMesh_RandomFilteredSocketOrBone(in FDISkelMeshContext DIContext, uint
 	Buffer<float4>	MeshColorBuffer_ ## NAME; \
 	Buffer<float>	MeshTriangleSamplerProbaBuffer_ ## NAME; \
 	Buffer<uint>	MeshTriangleSamplerAliasBuffer_ ## NAME; \
+	uint			MeshNumSamplingRegionTriangles_ ## NAME; \
+	uint			MeshNumSamplingRegionVertices_ ## NAME; \
+	Buffer<uint>	MeshSamplingRegionsProbaBuffer_ ## NAME; \
+	Buffer<uint>	MeshSamplingRegionsAliasBuffer_ ## NAME; \
+	Buffer<uint>	MeshSampleRegionsTriangleIndices_ ## NAME; \
+	Buffer<uint>	MeshSampleRegionsVertices_ ## NAME; \
 	Buffer<uint>	MeshTriangleMatricesOffsetBuffer_ ## NAME; \
 	uint			MeshTriangleCount_ ## NAME; \
 	uint			MeshVertexCount_ ## NAME; \
@@ -584,6 +680,12 @@ void DISkelMesh_RandomFilteredSocketOrBone(in FDISkelMeshContext DIContext, uint
 	DIContext.MeshColorBuffer = MeshColorBuffer_ ## NAME; \
 	DIContext.MeshTriangleSamplerProbaBuffer = MeshTriangleSamplerProbaBuffer_ ## NAME; \
 	DIContext.MeshTriangleSamplerAliasBuffer = MeshTriangleSamplerAliasBuffer_ ## NAME; \
+	DIContext.MeshNumSamplingRegionTriangles = MeshNumSamplingRegionTriangles_ ## NAME; \
+	DIContext.MeshNumSamplingRegionVertices = MeshNumSamplingRegionVertices_ ## NAME; \
+	DIContext.MeshSamplingRegionsProbaBuffer = MeshSamplingRegionsProbaBuffer_ ## NAME; \
+	DIContext.MeshSamplingRegionsAliasBuffer = MeshSamplingRegionsAliasBuffer_ ## NAME; \
+	DIContext.MeshSampleRegionsTriangleIndices = MeshSampleRegionsTriangleIndices_ ## NAME; \
+	DIContext.MeshSampleRegionsVertices = MeshSampleRegionsVertices_ ## NAME; \
 	DIContext.MeshTriangleMatricesOffsetBuffer = MeshTriangleMatricesOffsetBuffer_ ## NAME; \
 	DIContext.MeshTriangleCount = MeshTriangleCount_ ## NAME; \
 	DIContext.MeshVertexCount = MeshVertexCount_ ## NAME; \
@@ -601,4 +703,5 @@ void DISkelMesh_RandomFilteredSocketOrBone(in FDISkelMeshContext DIContext, uint
 	DIContext.InstancePrevTransform = InstancePrevTransform_ ## NAME; \
 	DIContext.InstanceInvDeltaTime = InstanceInvDeltaTime_ ## NAME; \
 	DIContext.UniformTriangleSamplingEnable = EnabledFeatures_ ## NAME & 0x1; \
-	DIContext.UnlimitedBoneInfluences = EnabledFeatures_ ## NAME & 0x2; \
+	DIContext.UniformSamplingRegionEnabled = EnabledFeatures_ ## NAME & 0x2; \
+	DIContext.UnlimitedBoneInfluences = EnabledFeatures_ ## NAME & 0x4; \
diff --git a/Engine/Plugins/FX/Niagara/Source/Niagara/Classes/NiagaraDataInterface.h b/Engine/Plugins/FX/Niagara/Source/Niagara/Classes/NiagaraDataInterface.h
--- a/Engine/Plugins/FX/Niagara/Source/Niagara/Classes/NiagaraDataInterface.h
+++ b/Engine/Plugins/FX/Niagara/Source/Niagara/Classes/NiagaraDataInterface.h
@@ -512,6 +512,16 @@ struct FNDIRandomHelper
 		return Min + (int(Rand(InstanceIndex) * (Range + 1)));
 	}
 
+	FORCEINLINE_DEBUGGABLE FVector RandomBarycentricCoord(int32 InstanceIndex)
+	{
+		//TODO: This is gonna be slooooow. Move to an LUT possibly or find faster method.
+		//Can probably handle lower quality randoms / uniformity for a decent speed win.
+		FVector2D r = Rand2(InstanceIndex);
+		float sqrt0 = FMath::Sqrt(r.X);
+		float sqrt1 = FMath::Sqrt(r.Y);
+		return FVector(1.0f - sqrt0, sqrt0 * (1.0 - r.Y), r.Y * sqrt0);
+	}
+
 	FVectorVMContext& Context;
 	FNDIParameter<FNiagaraRandInfo> RandParam;
 
diff --git a/Engine/Plugins/FX/Niagara/Source/Niagara/Classes/NiagaraDataInterfaceSkeletalMesh.h b/Engine/Plugins/FX/Niagara/Source/Niagara/Classes/NiagaraDataInterfaceSkeletalMesh.h
--- a/Engine/Plugins/FX/Niagara/Source/Niagara/Classes/NiagaraDataInterfaceSkeletalMesh.h
+++ b/Engine/Plugins/FX/Niagara/Source/Niagara/Classes/NiagaraDataInterfaceSkeletalMesh.h
@@ -292,6 +292,14 @@ public:
 	uint32 GetTriangleCount() const { return TriangleCount; }
 	uint32 GetVertexCount() const { return VertexCount; }
 
+	bool IsSamplingRegionsAllAreaWeighted() const { return bSamplingRegionsAllAreaWeighted; }
+	int32 GetNumSamplingRegionTriangles() const { return NumSamplingRegionTriangles; }
+	int32 GetNumSamplingRegionVertices() const { return NumSamplingRegionVertices; }
+	FShaderResourceViewRHIRef GetSampleRegionsProbSRV() const { return SampleRegionsProbSRV; }
+	FShaderResourceViewRHIRef GetSampleRegionsAliasSRV() const { return SampleRegionsAliasSRV; }
+	FShaderResourceViewRHIRef GetSampleRegionsTriangleIndicesSRV() const { return SampleRegionsTriangleIndicesSRV; }
+	FShaderResourceViewRHIRef GetSampleRegionsVerticesSRV() const { return SampleRegionsVerticesSRV; }
+
 	FRHIShaderResourceView* GetBufferPositionSRV() const { return MeshVertexBufferSrv; }
 	FRHIShaderResourceView* GetBufferIndexSRV() const { return MeshIndexBufferSrv; }
 	FRHIShaderResourceView* GetBufferTangentSRV() const { return MeshTangentBufferSRV; }
@@ -317,6 +325,23 @@ protected:
 	FVertexBufferRHIRef BufferTriangleMatricesOffsetRHI = nullptr;
 	FShaderResourceViewRHIRef BufferTriangleMatricesOffsetSRV = nullptr;
 
+	bool bSamplingRegionsAllAreaWeighted = false;
+	int32 NumSamplingRegionTriangles = 0;
+	int32 NumSamplingRegionVertices = 0;
+	TResourceArray<float> SampleRegionsProb;
+	TResourceArray<int32> SampleRegionsAlias;
+	TResourceArray<int32> SampleRegionsTriangleIndicies;
+	TResourceArray<int32> SampleRegionsVertices;
+
+	FVertexBufferRHIRef SampleRegionsProbBuffer;
+	FShaderResourceViewRHIRef SampleRegionsProbSRV;
+	FVertexBufferRHIRef SampleRegionsAliasBuffer;
+	FShaderResourceViewRHIRef SampleRegionsAliasSRV;
+	FVertexBufferRHIRef SampleRegionsTriangleIndicesBuffer;
+	FShaderResourceViewRHIRef SampleRegionsTriangleIndicesSRV;
+	FVertexBufferRHIRef SampleRegionsVerticesBuffer;
+	FShaderResourceViewRHIRef SampleRegionsVerticesSRV;
+
 	int32 NumFilteredBones = 0;
 	int32 NumUnfilteredBones = 0;
 	TResourceArray<uint16> FilteredAndUnfilteredBonesArray;
@@ -620,6 +645,12 @@ public:
 	static const FString MeshColorBufferName;
 	static const FString MeshTriangleSamplerProbaBufferName;
 	static const FString MeshTriangleSamplerAliasBufferName;
+	static const FString MeshNumSamplingRegionTrianglesName;
+	static const FString MeshNumSamplingRegionVerticesName;
+	static const FString MeshSamplingRegionsProbaBufferName;
+	static const FString MeshSamplingRegionsAliasBufferName;
+	static const FString MeshSampleRegionsTriangleIndicesName;
+	static const FString MeshSampleRegionsVerticesName;
 	static const FString MeshTriangleMatricesOffsetBufferName;
 	static const FString MeshTriangleCountName;
 	static const FString MeshVertexCountName;
@@ -678,6 +709,9 @@ private:
 	template<typename FilterMode, typename AreaWeightingMode>
 	FORCEINLINE int32 RandomTriIndex(FNDIRandomHelper& RandHelper, FSkeletalMeshAccessorHelper& Accessor, FNDISkeletalMesh_InstanceData* InstData, int32 InstanceIndex);
 
+	void RandomTriangle(FVectorVMContext& Context);
+	void GetTriangleCount(FVectorVMContext& Context);
+
 	template<typename FilterMode, typename AreaWeightingMode>
 	FORCEINLINE int32 GetFilteredTriangleCount(FSkeletalMeshAccessorHelper& Accessor, FNDISkeletalMesh_InstanceData* InstData);
 
@@ -694,17 +728,18 @@ public:
 	void GetVertexSamplingFunctions(TArray<FNiagaraFunctionSignature>& OutFunctions);
 	void BindVertexSamplingFunction(const FVMExternalFunctionBindingInfo& BindingInfo, FNDISkeletalMesh_InstanceData* InstData, FVMExternalFunction &OutFunc);
 
-	template<typename FilterMode>
-	void GetFilteredVertexCount(FVectorVMContext& Context);
+	void IsValidVertex(FVectorVMContext& Context);
+	void RandomVertex(FVectorVMContext& Context);
+	void GetVertexCount(FVectorVMContext& Context);
 
 	template<typename FilterMode>
-	void GetFilteredVertexAt(FVectorVMContext& Context);
-
+	void IsValidFilteredVertex(FVectorVMContext& Context);
 	template<typename FilterMode>
-	void RandomVertex(FVectorVMContext& Context);
-
+	void RandomFilteredVertex(FVectorVMContext& Context);
 	template<typename FilterMode>
-	void IsValidVertex(FVectorVMContext& Context);
+	void GetFilteredVertexCount(FVectorVMContext& Context);
+	template<typename FilterMode>
+	void GetFilteredVertexAt(FVectorVMContext& Context);
 
 	template<typename SkinningHandlerType, typename TransformHandlerType, typename VertexAccessorType>
 	void GetVertexSkinnedData(FVectorVMContext& Context);
@@ -718,7 +753,7 @@ public:
 
 private:
 	template<typename FilterMode>
-	FORCEINLINE int32 RandomVertIndex(FNDIRandomHelper& RandHelper, int32 Instance, FSkeletalMeshAccessorHelper& Accessor, FNDISkeletalMesh_InstanceData* InstData);
+	FORCEINLINE int32 RandomFilteredVertIndex(FNDIRandomHelper& RandHelper, int32 Instance, FSkeletalMeshAccessorHelper& Accessor, FNDISkeletalMesh_InstanceData* InstData);
 
 	template<typename FilterMode>
 	FORCEINLINE int32 GetFilteredVertexCount(FSkeletalMeshAccessorHelper& Accessor, FNDISkeletalMesh_InstanceData* InstData);
@@ -776,9 +811,12 @@ public:
 	static const FName GetSkinnedTriangleDataWSInterpName;
 	static const FName GetTriColorName;
 	static const FName GetTriUVName;
-	static const FName GetTriangleCountName;
-	static const FName GetTriangleAtName;
 	static const FName GetTriCoordVerticesName;
+	static const FName RandomTriangleName;
+	static const FName GetTriangleCountName;
+	static const FName RandomFilteredTriangleName;
+	static const FName GetFilteredTriangleCountName;
+	static const FName GetFilteredTriangleAtName;
 
 	// Bone Sampling
 	static const FName GetSkinnedBoneDataName;
@@ -805,14 +843,19 @@ public:
 	static const FName RandomFilteredSocketOrBoneName;
 
 	// Vertex Sampling
-	static const FName IsValidVertexName;
-	static const FName RandomVertexName;
 	static const FName GetSkinnedVertexDataName;
 	static const FName GetSkinnedVertexDataWSName;
 	static const FName GetVertexColorName;
 	static const FName GetVertexUVName;
+
+	static const FName IsValidVertexName;
+	static const FName RandomVertexName;
 	static const FName GetVertexCountName;
-	static const FName GetVertexAtName;
+
+	static const FName IsValidFilteredVertexName;
+	static const FName RandomFilteredVertexName;
+	static const FName GetFilteredVertexCountName;
+	static const FName GetFilteredVertexAtName;
 };
 
 struct FNiagaraDISkeletalMeshPassedDataToRT
diff --git a/Engine/Plugins/FX/Niagara/Source/Niagara/Private/NDISkeletalMesh_BoneSampling.cpp b/Engine/Plugins/FX/Niagara/Source/Niagara/Private/NDISkeletalMesh_BoneSampling.cpp
--- a/Engine/Plugins/FX/Niagara/Source/Niagara/Private/NDISkeletalMesh_BoneSampling.cpp
+++ b/Engine/Plugins/FX/Niagara/Source/Niagara/Private/NDISkeletalMesh_BoneSampling.cpp
@@ -426,8 +426,8 @@ void UNiagaraDataInterfaceSkeletalMesh::GetFilteredBoneAt(FVectorVMContext& Cont
 
 	VectorVM::FExternalFuncRegisterHandler<int32> OutBone(Context);
 
-	const int32 Max = InstData->NumFilteredBones;
-	if (Max != INDEX_NONE)
+	const int32 Max = InstData->NumFilteredBones - 1;
+	if (Max >= 0)
 	{
 		for (int32 i = 0; i < Context.NumInstances; ++i)
 		{
@@ -449,7 +449,7 @@ void UNiagaraDataInterfaceSkeletalMesh::RandomFilteredBone(FVectorVMContext& Con
 	VectorVM::FExternalFuncRegisterHandler<int32> OutBone(Context);
 
 	const int32 Max = InstData->NumFilteredBones - 1;
-	if (Max != INDEX_NONE)
+	if (Max >= 0)
 	{
 		for (int32 i = 0; i < Context.NumInstances; ++i)
 		{
@@ -485,13 +485,24 @@ void UNiagaraDataInterfaceSkeletalMesh::GetUnfilteredBoneAt(FVectorVMContext& Co
 
 	VectorVM::FExternalFuncRegisterHandler<int32> OutBone(Context);
 
-	const int32 Max = InstData->NumUnfilteredBones;
-	if (Max != INDEX_NONE)
+	const int32 Max = InstData->NumUnfilteredBones - 1;
+	if (Max >= 0)
 	{
-		for (int32 i = 0; i < Context.NumInstances; ++i)
+		if (InstData->NumFilteredBones == 0)
 		{
-			const int32 BoneIndex = FMath::Clamp(BoneParam.GetAndAdvance(), 0, Max);
-			*OutBone.GetDestAndAdvance() = InstData->FilteredAndUnfilteredBones[BoneIndex + InstData->NumFilteredBones];
+			for (int32 i = 0; i < Context.NumInstances; ++i)
+			{
+				const int32 BoneIndex = FMath::Clamp(BoneParam.GetAndAdvance(), 0, Max);
+				*OutBone.GetDestAndAdvance() = BoneIndex;
+			}
+		}
+		else
+		{
+			for (int32 i = 0; i < Context.NumInstances; ++i)
+			{
+				const int32 BoneIndex = FMath::Clamp(BoneParam.GetAndAdvance(), 0, Max);
+				*OutBone.GetDestAndAdvance() = InstData->FilteredAndUnfilteredBones[BoneIndex + InstData->NumFilteredBones];
+			}
 		}
 	}
 	else
@@ -507,13 +518,25 @@ void UNiagaraDataInterfaceSkeletalMesh::RandomUnfilteredBone(FVectorVMContext& C
 	VectorVM::FExternalFuncRegisterHandler<int32> OutBone(Context);
 
 	const int32 Max = InstData->NumUnfilteredBones - 1;
-	if (Max != INDEX_NONE)
+	if (Max >= 0)
 	{
-		for (int32 i = 0; i < Context.NumInstances; ++i)
+		if (InstData->NumFilteredBones == 0)
 		{
-			RandHelper.GetAndAdvance();
-			const int32 BoneIndex = RandHelper.RandRange(i, 0, Max);
-			*OutBone.GetDestAndAdvance() = InstData->FilteredAndUnfilteredBones[BoneIndex + InstData->NumFilteredBones];
+			for (int32 i = 0; i < Context.NumInstances; ++i)
+			{
+				RandHelper.GetAndAdvance();
+				const int32 BoneIndex = RandHelper.RandRange(i, 0, Max);
+				*OutBone.GetDestAndAdvance() = BoneIndex;
+			}
+		}
+		else
+		{
+			for (int32 i = 0; i < Context.NumInstances; ++i)
+			{
+				RandHelper.GetAndAdvance();
+				const int32 BoneIndex = RandHelper.RandRange(i, 0, Max);
+				*OutBone.GetDestAndAdvance() = InstData->FilteredAndUnfilteredBones[BoneIndex + InstData->NumFilteredBones];
+			}
 		}
 	}
 	else
diff --git a/Engine/Plugins/FX/Niagara/Source/Niagara/Private/NDISkeletalMesh_TriangleSampling.cpp b/Engine/Plugins/FX/Niagara/Source/Niagara/Private/NDISkeletalMesh_TriangleSampling.cpp
--- a/Engine/Plugins/FX/Niagara/Source/Niagara/Private/NDISkeletalMesh_TriangleSampling.cpp
+++ b/Engine/Plugins/FX/Niagara/Source/Niagara/Private/NDISkeletalMesh_TriangleSampling.cpp
@@ -32,23 +32,26 @@ const FName FSkeletalMeshInterfaceHelper::GetSkinnedTriangleDataInterpName("GetS
 const FName FSkeletalMeshInterfaceHelper::GetSkinnedTriangleDataWSInterpName("GetSkinnedTriangleDataWSInterpolated");
 const FName FSkeletalMeshInterfaceHelper::GetTriColorName("GetTriColor");
 const FName FSkeletalMeshInterfaceHelper::GetTriUVName("GetTriUV");
-const FName FSkeletalMeshInterfaceHelper::GetTriangleCountName("GetFilteredTriangleCount");
-const FName FSkeletalMeshInterfaceHelper::GetTriangleAtName("GetFilteredTriangle");
 const FName FSkeletalMeshInterfaceHelper::GetTriCoordVerticesName("GetTriCoordVertices");
-
+const FName FSkeletalMeshInterfaceHelper::RandomTriangleName("RandomTriangle");
+const FName FSkeletalMeshInterfaceHelper::GetTriangleCountName("GetTriangleCount");
+const FName FSkeletalMeshInterfaceHelper::RandomFilteredTriangleName("RandomFilteredTriangle");
+const FName FSkeletalMeshInterfaceHelper::GetFilteredTriangleCountName("GetFilteredTriangleCount");
+const FName FSkeletalMeshInterfaceHelper::GetFilteredTriangleAtName("GetFilteredTriangle");
 
 void UNiagaraDataInterfaceSkeletalMesh::GetTriangleSamplingFunctions(TArray<FNiagaraFunctionSignature>& OutFunctions)
 {
-	{
-		FNiagaraFunctionSignature Sig;
-		Sig.Name = FSkeletalMeshInterfaceHelper::RandomTriCoordName;
-		Sig.Inputs.Add(FNiagaraVariable(FNiagaraTypeDefinition(GetClass()), TEXT("SkeletalMesh")));
-		Sig.Inputs.Add(FNiagaraVariable(FNiagaraTypeDefinition(FNiagaraRandInfo::StaticStruct()), TEXT("RandomInfo")));
-		Sig.Outputs.Add(FNiagaraVariable(FNiagaraTypeDefinition(FMeshTriCoordinate::StaticStruct()), TEXT("Coord")));
-		Sig.bMemberFunction = true;
-		Sig.bRequiresContext = false;
-		OutFunctions.Add(Sig);
-	}
+// Hidden will remove later
+//	{
+//		FNiagaraFunctionSignature Sig;
+//		Sig.Name = FSkeletalMeshInterfaceHelper::RandomTriCoordName;
+//		Sig.Inputs.Add(FNiagaraVariable(FNiagaraTypeDefinition(GetClass()), TEXT("SkeletalMesh")));
+//		Sig.Inputs.Add(FNiagaraVariable(FNiagaraTypeDefinition(FNiagaraRandInfo::StaticStruct()), TEXT("RandomInfo")));
+//		Sig.Outputs.Add(FNiagaraVariable(FNiagaraTypeDefinition(FMeshTriCoordinate::StaticStruct()), TEXT("Coord")));
+//		Sig.bMemberFunction = true;
+//		Sig.bRequiresContext = false;
+//		OutFunctions.Add(Sig);
+//	}
 
 	{
 		FNiagaraFunctionSignature Sig;
@@ -162,6 +165,32 @@ void UNiagaraDataInterfaceSkeletalMesh::GetTriangleSamplingFunctions(TArray<FNia
 		OutFunctions.Add(Sig);
 	}
 
+	{
+		FNiagaraFunctionSignature Sig;
+		Sig.Name = FSkeletalMeshInterfaceHelper::GetTriCoordVerticesName;
+		Sig.Inputs.Add(FNiagaraVariable(FNiagaraTypeDefinition(GetClass()), TEXT("SkeletalMesh")));
+		Sig.Inputs.Add(FNiagaraVariable(FNiagaraTypeDefinition::GetIntDef(), TEXT("TriangleIndex")));
+		Sig.Outputs.Add(FNiagaraVariable(FNiagaraTypeDefinition::GetIntDef(), TEXT("Vertex 0")));
+		Sig.Outputs.Add(FNiagaraVariable(FNiagaraTypeDefinition::GetIntDef(), TEXT("Vertex 1")));
+		Sig.Outputs.Add(FNiagaraVariable(FNiagaraTypeDefinition::GetIntDef(), TEXT("Vertex 2")));
+		Sig.bMemberFunction = true;
+		Sig.bRequiresContext = false;
+#if WITH_EDITORONLY_DATA
+		Sig.Description = LOCTEXT("GetTriCoordVetsName", "Takes the TriangleIndex from a MeshTriCoord and returns the vertices for that triangle.");
+#endif
+		OutFunctions.Add(Sig);
+	}
+
+	{
+		FNiagaraFunctionSignature Sig;
+		Sig.Name = FSkeletalMeshInterfaceHelper::RandomTriangleName;
+		Sig.Inputs.Add(FNiagaraVariable(FNiagaraTypeDefinition(GetClass()), TEXT("SkeletalMesh")));
+		Sig.Inputs.Add(FNiagaraVariable(FNiagaraTypeDefinition(FNiagaraRandInfo::StaticStruct()), TEXT("RandomInfo")));
+		Sig.Outputs.Add(FNiagaraVariable(FNiagaraTypeDefinition(FMeshTriCoordinate::StaticStruct()), TEXT("Coord")));
+		Sig.bMemberFunction = true;
+		Sig.bRequiresContext = false;
+		OutFunctions.Add(Sig);
+	}
 	{
 		FNiagaraFunctionSignature Sig;
 		Sig.Name = FSkeletalMeshInterfaceHelper::GetTriangleCountName;
@@ -174,28 +203,31 @@ void UNiagaraDataInterfaceSkeletalMesh::GetTriangleSamplingFunctions(TArray<FNia
 
 	{
 		FNiagaraFunctionSignature Sig;
-		Sig.Name = FSkeletalMeshInterfaceHelper::GetTriangleAtName;
+		Sig.Name = FSkeletalMeshInterfaceHelper::RandomFilteredTriangleName;
 		Sig.Inputs.Add(FNiagaraVariable(FNiagaraTypeDefinition(GetClass()), TEXT("SkeletalMesh")));
-		Sig.Inputs.Add(FNiagaraVariable(FNiagaraTypeDefinition::GetIntDef(), TEXT("Index")));
+		Sig.Inputs.Add(FNiagaraVariable(FNiagaraTypeDefinition(FNiagaraRandInfo::StaticStruct()), TEXT("RandomInfo")));
 		Sig.Outputs.Add(FNiagaraVariable(FNiagaraTypeDefinition(FMeshTriCoordinate::StaticStruct()), TEXT("Coord")));
 		Sig.bMemberFunction = true;
 		Sig.bRequiresContext = false;
 		OutFunctions.Add(Sig);
 	}
-
 	{
 		FNiagaraFunctionSignature Sig;
-		Sig.Name = FSkeletalMeshInterfaceHelper::GetTriCoordVerticesName;
+		Sig.Name = FSkeletalMeshInterfaceHelper::GetFilteredTriangleCountName;
 		Sig.Inputs.Add(FNiagaraVariable(FNiagaraTypeDefinition(GetClass()), TEXT("SkeletalMesh")));
-		Sig.Inputs.Add(FNiagaraVariable(FNiagaraTypeDefinition::GetIntDef(), TEXT("TriangleIndex")));
-		Sig.Outputs.Add(FNiagaraVariable(FNiagaraTypeDefinition::GetIntDef(), TEXT("Vertex 0")));
-		Sig.Outputs.Add(FNiagaraVariable(FNiagaraTypeDefinition::GetIntDef(), TEXT("Vertex 1")));
-		Sig.Outputs.Add(FNiagaraVariable(FNiagaraTypeDefinition::GetIntDef(), TEXT("Vertex 2")));
+		Sig.Outputs.Add(FNiagaraVariable(FNiagaraTypeDefinition::GetIntDef(), TEXT("Count")));
+		Sig.bMemberFunction = true;
+		Sig.bRequiresContext = false;
+		OutFunctions.Add(Sig);
+	}
+	{
+		FNiagaraFunctionSignature Sig;
+		Sig.Name = FSkeletalMeshInterfaceHelper::GetFilteredTriangleAtName;
+		Sig.Inputs.Add(FNiagaraVariable(FNiagaraTypeDefinition(GetClass()), TEXT("SkeletalMesh")));
+		Sig.Inputs.Add(FNiagaraVariable(FNiagaraTypeDefinition::GetIntDef(), TEXT("Index")));
+		Sig.Outputs.Add(FNiagaraVariable(FNiagaraTypeDefinition(FMeshTriCoordinate::StaticStruct()), TEXT("Coord")));
 		Sig.bMemberFunction = true;
 		Sig.bRequiresContext = false;
-#if WITH_EDITORONLY_DATA
-		Sig.Description = LOCTEXT("GetTriCoordVetsName", "Takes the TriangleIndex from a MeshTriCoord and returns the vertices for that triangle.");
-#endif
 		OutFunctions.Add(Sig);
 	}
 }
@@ -249,22 +281,36 @@ void UNiagaraDataInterfaceSkeletalMesh::BindTriangleSamplingFunction(const FVMEx
 		check(BindingInfo.GetNumInputs() == 6 && BindingInfo.GetNumOutputs() == 2);
 		TVertexAccessorBinder<NDI_FUNC_BINDER(UNiagaraDataInterfaceSkeletalMesh, GetTriCoordUV)>::Bind(this, BindingInfo, InstanceData, OutFunc);
 	}
+	else if (BindingInfo.Name == FSkeletalMeshInterfaceHelper::GetTriCoordVerticesName)
+	{
+		check(BindingInfo.GetNumInputs() == 2 && BindingInfo.GetNumOutputs() == 3);
+		TSkinningModeBinder<NDI_FUNC_BINDER(UNiagaraDataInterfaceSkeletalMesh, GetTriCoordVertices)>::Bind(this, BindingInfo, InstanceData, OutFunc);
+	}
+	else if (BindingInfo.Name == FSkeletalMeshInterfaceHelper::RandomTriangleName)
+	{
+		check(BindingInfo.GetNumInputs() == 4 && BindingInfo.GetNumOutputs() == 4);
+		OutFunc = FVMExternalFunction::CreateLambda([this](FVectorVMContext& Context) { this->RandomTriangle(Context); });
+	}
 	else if (BindingInfo.Name == FSkeletalMeshInterfaceHelper::GetTriangleCountName)
+	{
+		check(BindingInfo.GetNumInputs() == 1 && BindingInfo.GetNumOutputs() == 1);
+		OutFunc = FVMExternalFunction::CreateLambda([this](FVectorVMContext& Context) { this->GetTriangleCount(Context); });
+	}
+	else if (BindingInfo.Name == FSkeletalMeshInterfaceHelper::RandomFilteredTriangleName)
+	{
+		check(BindingInfo.GetNumInputs() == 4 && BindingInfo.GetNumOutputs() == 4);
+		TFilterModeBinder<TAreaWeightingModeBinder<NDI_FUNC_BINDER(UNiagaraDataInterfaceSkeletalMesh, RandomTriCoord)>>::Bind(this, BindingInfo, InstanceData, OutFunc);
+	}
+	else if (BindingInfo.Name == FSkeletalMeshInterfaceHelper::GetFilteredTriangleCountName)
 	{
 		check(BindingInfo.GetNumInputs() == 1 && BindingInfo.GetNumOutputs() == 1);
 		TFilterModeBinder<TAreaWeightingModeBinder<NDI_FUNC_BINDER(UNiagaraDataInterfaceSkeletalMesh, GetFilteredTriangleCount)>>::Bind(this, BindingInfo, InstanceData, OutFunc);
 	}
-	else if (BindingInfo.Name == FSkeletalMeshInterfaceHelper::GetTriangleAtName)
+	else if (BindingInfo.Name == FSkeletalMeshInterfaceHelper::GetFilteredTriangleAtName)
 	{
 		check(BindingInfo.GetNumInputs() == 2 && BindingInfo.GetNumOutputs() == 4);
 		TFilterModeBinder<TAreaWeightingModeBinder<NDI_FUNC_BINDER(UNiagaraDataInterfaceSkeletalMesh, GetFilteredTriangleAt)>>::Bind(this, BindingInfo, InstanceData, OutFunc);
 	}
-	else if (BindingInfo.Name == FSkeletalMeshInterfaceHelper::GetTriCoordVerticesName)
-	{
-		check(BindingInfo.GetNumInputs() == 2 && BindingInfo.GetNumOutputs() == 3);
-		TSkinningModeBinder<NDI_FUNC_BINDER(UNiagaraDataInterfaceSkeletalMesh, GetTriCoordVertices)>::Bind(this, BindingInfo, InstanceData, OutFunc);
-	}
-
 }
 
 template<typename FilterMode, typename AreaWeightingMode>
@@ -411,6 +457,61 @@ void UNiagaraDataInterfaceSkeletalMesh::IsValidTriCoord(FVectorVMContext& Contex
 
 //////////////////////////////////////////////////////////////////////////
 
+void UNiagaraDataInterfaceSkeletalMesh::RandomTriangle(FVectorVMContext& Context)
+{
+	SCOPE_CYCLE_COUNTER(STAT_NiagaraSkel_Sample);
+
+	FNDIRandomHelper RandHelper(Context);
+	VectorVM::FUserPtrHandler<FNDISkeletalMesh_InstanceData> InstData(Context);
+	VectorVM::FExternalFuncRegisterHandler<int32> OutTri(Context);
+	VectorVM::FExternalFuncRegisterHandler<float> OutBaryX(Context);
+	VectorVM::FExternalFuncRegisterHandler<float> OutBaryY(Context);
+	VectorVM::FExternalFuncRegisterHandler<float> OutBaryZ(Context);
+
+	FSkeletalMeshAccessorHelper MeshAccessor;
+	MeshAccessor.Init<TIntegralConstant<int32, 0>, TIntegralConstant<int32, 0>>(InstData);
+
+	const int32 MaxTriangle = (MeshAccessor.IndexBuffer->Num() / 3) - 1;
+	if (MaxTriangle >= 0)
+	{
+		for (int32 i = 0; i < Context.NumInstances; ++i)
+		{
+			RandHelper.GetAndAdvance();
+
+			const int32 Triangle = RandHelper.RandRange(i, 0, MaxTriangle);
+			const FVector Bary = RandHelper.RandomBarycentricCoord(i);
+			*OutTri.GetDestAndAdvance() = Triangle;
+			*OutBaryX.GetDestAndAdvance() = Bary.X;
+			*OutBaryY.GetDestAndAdvance() = Bary.Y;
+			*OutBaryZ.GetDestAndAdvance() = Bary.Z;
+		}
+	}
+	else
+	{
+		FMemory::Memset(OutTri.GetDest(), 0xff, Context.NumInstances * sizeof(int32));
+		FMemory::Memset(OutBaryX.GetDest(), 0xff, Context.NumInstances * sizeof(int32));
+		FMemory::Memset(OutBaryY.GetDest(), 0xff, Context.NumInstances * sizeof(int32));
+		FMemory::Memset(OutBaryZ.GetDest(), 0xff, Context.NumInstances * sizeof(int32));
+	}
+}
+
+void UNiagaraDataInterfaceSkeletalMesh::GetTriangleCount(FVectorVMContext& Context)
+{
+	SCOPE_CYCLE_COUNTER(STAT_NiagaraSkel_Sample);
+
+	VectorVM::FUserPtrHandler<FNDISkeletalMesh_InstanceData> InstData(Context);
+	VectorVM::FExternalFuncRegisterHandler<int32> OutCount(Context);
+
+	FSkeletalMeshAccessorHelper MeshAccessor;
+	MeshAccessor.Init<TIntegralConstant<int32, 0>, TIntegralConstant<int32, 0>>(InstData);
+	
+	const int32 NumTriangles = MeshAccessor.IndexBuffer->Num() / 3;
+	for (int32 i = 0; i < Context.NumInstances; ++i)
+	{
+		*OutCount.GetDestAndAdvance() = NumTriangles;
+	}
+}
+
 template<typename FilterMode, typename AreaWeightingMode>
 FORCEINLINE int32 UNiagaraDataInterfaceSkeletalMesh::GetFilteredTriangleCount(FSkeletalMeshAccessorHelper& Accessor, FNDISkeletalMesh_InstanceData* InstData)
 {
@@ -518,7 +619,6 @@ void UNiagaraDataInterfaceSkeletalMesh::GetFilteredTriangleCount(FVectorVMContex
 	}
 }
 
-
 //////////////////////////////////////////////////////////////////////////
 
 template<typename FilterMode, typename AreaWeightingMode>
diff --git a/Engine/Plugins/FX/Niagara/Source/Niagara/Private/NDISkeletalMesh_VertexSampling.cpp b/Engine/Plugins/FX/Niagara/Source/Niagara/Private/NDISkeletalMesh_VertexSampling.cpp
--- a/Engine/Plugins/FX/Niagara/Source/Niagara/Private/NDISkeletalMesh_VertexSampling.cpp
+++ b/Engine/Plugins/FX/Niagara/Source/Niagara/Private/NDISkeletalMesh_VertexSampling.cpp
@@ -14,51 +14,32 @@
 DECLARE_CYCLE_STAT(TEXT("Skel Mesh Vertex Sampling"), STAT_NiagaraSkel_Vertex_Sample, STATGROUP_Niagara);
 
 //Final binders for all static mesh interface functions.
-DEFINE_NDI_FUNC_BINDER(UNiagaraDataInterfaceSkeletalMesh, RandomVertex);
 DEFINE_NDI_FUNC_BINDER(UNiagaraDataInterfaceSkeletalMesh, GetVertexSkinnedData);
 DEFINE_NDI_DIRECT_FUNC_BINDER(UNiagaraDataInterfaceSkeletalMesh, GetVertexColor);
 DEFINE_NDI_DIRECT_FUNC_BINDER(UNiagaraDataInterfaceSkeletalMesh, GetVertexColorFallback);
 DEFINE_NDI_FUNC_BINDER(UNiagaraDataInterfaceSkeletalMesh, GetVertexUV);
-DEFINE_NDI_FUNC_BINDER(UNiagaraDataInterfaceSkeletalMesh, IsValidVertex);
+
+DEFINE_NDI_FUNC_BINDER(UNiagaraDataInterfaceSkeletalMesh, IsValidFilteredVertex);
+DEFINE_NDI_FUNC_BINDER(UNiagaraDataInterfaceSkeletalMesh, RandomFilteredVertex);
 DEFINE_NDI_FUNC_BINDER(UNiagaraDataInterfaceSkeletalMesh, GetFilteredVertexCount);
 DEFINE_NDI_FUNC_BINDER(UNiagaraDataInterfaceSkeletalMesh, GetFilteredVertexAt);
 
-const FName FSkeletalMeshInterfaceHelper::IsValidVertexName("IsValidVertex");
-const FName FSkeletalMeshInterfaceHelper::RandomVertexName("RandomVertex");
 const FName FSkeletalMeshInterfaceHelper::GetSkinnedVertexDataName("GetSkinnedVertexData");
 const FName FSkeletalMeshInterfaceHelper::GetSkinnedVertexDataWSName("GetSkinnedVertexDataWS");
 const FName FSkeletalMeshInterfaceHelper::GetVertexColorName("GetVertexColor");
 const FName FSkeletalMeshInterfaceHelper::GetVertexUVName("GetVertexUV");
-const FName FSkeletalMeshInterfaceHelper::GetVertexCountName("GetFilteredVertexCount");
-const FName FSkeletalMeshInterfaceHelper::GetVertexAtName("GetFilteredVertex");
 
+const FName FSkeletalMeshInterfaceHelper::IsValidVertexName("IsValidVertex");
+const FName FSkeletalMeshInterfaceHelper::RandomVertexName("RandomVertex");
+const FName FSkeletalMeshInterfaceHelper::GetVertexCountName("GetVertexCount");
+
+const FName FSkeletalMeshInterfaceHelper::IsValidFilteredVertexName("IsValidFilteredVertex");
+const FName FSkeletalMeshInterfaceHelper::RandomFilteredVertexName("RandomFilteredVertex");
+const FName FSkeletalMeshInterfaceHelper::GetFilteredVertexCountName("GetFilteredVertexCount");
+const FName FSkeletalMeshInterfaceHelper::GetFilteredVertexAtName("GetFilteredVertex");
+
 void UNiagaraDataInterfaceSkeletalMesh::GetVertexSamplingFunctions(TArray<FNiagaraFunctionSignature>& OutFunctions)
 {
 	{
 		FNiagaraFunctionSignature Sig;
-		Sig.Name = FSkeletalMeshInterfaceHelper::RandomVertexName;
-		Sig.Inputs.Add(FNiagaraVariable(FNiagaraTypeDefinition(GetClass()), TEXT("SkeletalMesh")));
-		Sig.Inputs.Add(FNiagaraVariable(FNiagaraTypeDefinition(FNiagaraRandInfo::StaticStruct()), TEXT("RandomInfo")));
-		Sig.Outputs.Add(FNiagaraVariable(FNiagaraTypeDefinition::GetIntDef(), TEXT("Vertex")));
-		Sig.bMemberFunction = true;
-		Sig.bRequiresContext = false;
-		OutFunctions.Add(Sig);
-	}
-
-	{
-		FNiagaraFunctionSignature Sig;
-		Sig.Name = FSkeletalMeshInterfaceHelper::IsValidVertexName;
-		Sig.Inputs.Add(FNiagaraVariable(FNiagaraTypeDefinition(GetClass()), TEXT("SkeletalMesh")));
-		Sig.Inputs.Add(FNiagaraVariable(FNiagaraTypeDefinition::GetIntDef(), TEXT("Vertex")));
-		Sig.Outputs.Add(FNiagaraVariable(FNiagaraTypeDefinition::GetBoolDef(), TEXT("IsValid")));
-		Sig.bMemberFunction = true;
-		Sig.bRequiresContext = false;
-#if WITH_EDITORONLY_DATA
-		Sig.Description = LOCTEXT("IsValidDesc", "Determine if this tri coordinate's Vertex index is valid for this mesh. Note that this only checks the mesh index buffer size and does not include any filtering settings.");
-#endif
-		OutFunctions.Add(Sig);
-	}
-
-	{
-		FNiagaraFunctionSignature Sig;
 		Sig.Name = FSkeletalMeshInterfaceHelper::GetSkinnedVertexDataName;
@@ -115,4 +96,24 @@
 
 	{
 		FNiagaraFunctionSignature Sig;
+		Sig.Name = FSkeletalMeshInterfaceHelper::IsValidVertexName;
+		Sig.Inputs.Add(FNiagaraVariable(FNiagaraTypeDefinition(GetClass()), TEXT("SkeletalMesh")));
+		Sig.Inputs.Add(FNiagaraVariable(FNiagaraTypeDefinition::GetIntDef(), TEXT("Vertex")));
+		Sig.Outputs.Add(FNiagaraVariable(FNiagaraTypeDefinition::GetBoolDef(), TEXT("IsValid")));
+		Sig.bMemberFunction = true;
+		Sig.bRequiresContext = false;
+		OutFunctions.Add(Sig);
+	}
+	{
+		FNiagaraFunctionSignature Sig;
+		Sig.Name = FSkeletalMeshInterfaceHelper::RandomVertexName;
+		Sig.Inputs.Add(FNiagaraVariable(FNiagaraTypeDefinition(GetClass()), TEXT("SkeletalMesh")));
+		Sig.Inputs.Add(FNiagaraVariable(FNiagaraTypeDefinition(FNiagaraRandInfo::StaticStruct()), TEXT("RandomInfo")));
+		Sig.Outputs.Add(FNiagaraVariable(FNiagaraTypeDefinition::GetIntDef(), TEXT("Vertex")));
+		Sig.bMemberFunction = true;
+		Sig.bRequiresContext = false;
+		OutFunctions.Add(Sig);
+	}
+	{
+		FNiagaraFunctionSignature Sig;
 		Sig.Name = FSkeletalMeshInterfaceHelper::GetVertexCountName;
@@ -125,7 +126,36 @@
 
 	{
 		FNiagaraFunctionSignature Sig;
-		Sig.Name = FSkeletalMeshInterfaceHelper::GetVertexAtName;
+		Sig.Name = FSkeletalMeshInterfaceHelper::IsValidFilteredVertexName;
 		Sig.Inputs.Add(FNiagaraVariable(FNiagaraTypeDefinition(GetClass()), TEXT("SkeletalMesh")));
+		Sig.Inputs.Add(FNiagaraVariable(FNiagaraTypeDefinition::GetIntDef(), TEXT("Vertex")));
+		Sig.Outputs.Add(FNiagaraVariable(FNiagaraTypeDefinition::GetBoolDef(), TEXT("IsValid")));
+		Sig.bMemberFunction = true;
+		Sig.bRequiresContext = false;
+		OutFunctions.Add(Sig);
+	}
+	{
+		FNiagaraFunctionSignature Sig;
+		Sig.Name = FSkeletalMeshInterfaceHelper::RandomFilteredVertexName;
+		Sig.Inputs.Add(FNiagaraVariable(FNiagaraTypeDefinition(GetClass()), TEXT("SkeletalMesh")));
+		Sig.Inputs.Add(FNiagaraVariable(FNiagaraTypeDefinition(FNiagaraRandInfo::StaticStruct()), TEXT("RandomInfo")));
+		Sig.Outputs.Add(FNiagaraVariable(FNiagaraTypeDefinition::GetIntDef(), TEXT("Vertex")));
+		Sig.bMemberFunction = true;
+		Sig.bRequiresContext = false;
+		OutFunctions.Add(Sig);
+	}
+	{
+		FNiagaraFunctionSignature Sig;
+		Sig.Name = FSkeletalMeshInterfaceHelper::GetFilteredVertexCountName;
+		Sig.Inputs.Add(FNiagaraVariable(FNiagaraTypeDefinition(GetClass()), TEXT("SkeletalMesh")));
+		Sig.Outputs.Add(FNiagaraVariable(FNiagaraTypeDefinition::GetIntDef(), TEXT("Count")));
+		Sig.bMemberFunction = true;
+		Sig.bRequiresContext = false;
+		OutFunctions.Add(Sig);
+	}
+	{
+		FNiagaraFunctionSignature Sig;
+		Sig.Name = FSkeletalMeshInterfaceHelper::GetFilteredVertexAtName;
+		Sig.Inputs.Add(FNiagaraVariable(FNiagaraTypeDefinition(GetClass()), TEXT("SkeletalMesh")));
 		Sig.Inputs.Add(FNiagaraVariable(FNiagaraTypeDefinition::GetIntDef(), TEXT("Filtered Index")));
 		Sig.Outputs.Add(FNiagaraVariable(FNiagaraTypeDefinition::GetIntDef(), TEXT("Vertex")));
@@ -137,18 +167,8 @@
 
 void UNiagaraDataInterfaceSkeletalMesh::BindVertexSamplingFunction(const FVMExternalFunctionBindingInfo& BindingInfo, FNDISkeletalMesh_InstanceData* InstanceData, FVMExternalFunction &OutFunc)
 {
-
-	if (BindingInfo.Name == FSkeletalMeshInterfaceHelper::RandomVertexName)
+	//////////////////////////////////////////////////////////////////////////
+	if (BindingInfo.Name == FSkeletalMeshInterfaceHelper::GetSkinnedVertexDataName)
 	{
-		check(BindingInfo.GetNumInputs() == 1 && BindingInfo.GetNumOutputs() == 1);
-		TFilterModeBinder<NDI_FUNC_BINDER(UNiagaraDataInterfaceSkeletalMesh, RandomVertex)>::Bind(this, BindingInfo, InstanceData, OutFunc);
-	}
-	else if (BindingInfo.Name == FSkeletalMeshInterfaceHelper::IsValidVertexName)
-	{
-		check(BindingInfo.GetNumInputs() == 2 && BindingInfo.GetNumOutputs() == 1);
-		TFilterModeBinder<NDI_FUNC_BINDER(UNiagaraDataInterfaceSkeletalMesh, IsValidVertex)>::Bind(this, BindingInfo, InstanceData, OutFunc);
-	}
-	else if (BindingInfo.Name == FSkeletalMeshInterfaceHelper::GetSkinnedVertexDataName)
-	{
 		check(BindingInfo.GetNumInputs() == 2 && BindingInfo.GetNumOutputs() == 6);
 		TSkinningModeBinder<TNDIExplicitBinder<FNDITransformHandlerNoop, TVertexAccessorBinder<NDI_FUNC_BINDER(UNiagaraDataInterfaceSkeletalMesh, GetVertexSkinnedData)>>>::Bind(this, BindingInfo, InstanceData, OutFunc);
@@ -175,35 +195,129 @@
 		check(BindingInfo.GetNumInputs() == 3 && BindingInfo.GetNumOutputs() == 2);
 		TVertexAccessorBinder<NDI_FUNC_BINDER(UNiagaraDataInterfaceSkeletalMesh, GetVertexUV)>::Bind(this, BindingInfo, InstanceData, OutFunc);
 	}
+	//////////////////////////////////////////////////////////////////////////
+	else if (BindingInfo.Name == FSkeletalMeshInterfaceHelper::IsValidVertexName)
+	{
+		check(BindingInfo.GetNumInputs() == 2 && BindingInfo.GetNumOutputs() == 1);
+		OutFunc = FVMExternalFunction::CreateLambda([this](FVectorVMContext& Context) { this->IsValidVertex(Context); });
+	}
+	else if (BindingInfo.Name == FSkeletalMeshInterfaceHelper::RandomVertexName)
+	{
+		check(BindingInfo.GetNumInputs() == 4 && BindingInfo.GetNumOutputs() == 1);
+		OutFunc = FVMExternalFunction::CreateLambda([this](FVectorVMContext& Context) { this->RandomVertex(Context); });
+	}
 	else if (BindingInfo.Name == FSkeletalMeshInterfaceHelper::GetVertexCountName)
 	{
 		check(BindingInfo.GetNumInputs() == 1 && BindingInfo.GetNumOutputs() == 1);
+		OutFunc = FVMExternalFunction::CreateLambda([this](FVectorVMContext& Context) { this->GetVertexCount(Context); });
+	}
+	//////////////////////////////////////////////////////////////////////////
+	else if (BindingInfo.Name == FSkeletalMeshInterfaceHelper::IsValidFilteredVertexName)
+	{
+		check(BindingInfo.GetNumInputs() == 2 && BindingInfo.GetNumOutputs() == 1);
+		TFilterModeBinder<NDI_FUNC_BINDER(UNiagaraDataInterfaceSkeletalMesh, IsValidFilteredVertex)>::Bind(this, BindingInfo, InstanceData, OutFunc);
+	}
+	else if (BindingInfo.Name == FSkeletalMeshInterfaceHelper::RandomFilteredVertexName)
+	{
+		check(BindingInfo.GetNumInputs() == 4 && BindingInfo.GetNumOutputs() == 1);
+		TFilterModeBinder<NDI_FUNC_BINDER(UNiagaraDataInterfaceSkeletalMesh, RandomFilteredVertex)>::Bind(this, BindingInfo, InstanceData, OutFunc);
+	}
+	else if (BindingInfo.Name == FSkeletalMeshInterfaceHelper::GetFilteredVertexCountName)
+	{
+		check(BindingInfo.GetNumInputs() == 1 && BindingInfo.GetNumOutputs() == 1);
 		TFilterModeBinder<NDI_FUNC_BINDER(UNiagaraDataInterfaceSkeletalMesh, GetFilteredVertexCount)>::Bind(this, BindingInfo, InstanceData, OutFunc);
 	}
-	else if (BindingInfo.Name == FSkeletalMeshInterfaceHelper::GetVertexAtName)
+	else if (BindingInfo.Name == FSkeletalMeshInterfaceHelper::GetFilteredVertexAtName)
 	{
 		check(BindingInfo.GetNumInputs() == 2 && BindingInfo.GetNumOutputs() == 1);
 		TFilterModeBinder<NDI_FUNC_BINDER(UNiagaraDataInterfaceSkeletalMesh, GetFilteredVertexAt)>::Bind(this, BindingInfo, InstanceData, OutFunc);
 	}
-
 }
 
+//////////////////////////////////////////////////////////////////////////
+
+void UNiagaraDataInterfaceSkeletalMesh::IsValidVertex(FVectorVMContext& Context)
+{
+	SCOPE_CYCLE_COUNTER(STAT_NiagaraSkel_Vertex_Sample);
+
+	VectorVM::FExternalFuncInputHandler<int32> VertexParam(Context);
+	VectorVM::FUserPtrHandler<FNDISkeletalMesh_InstanceData> InstData(Context);
+	VectorVM::FExternalFuncRegisterHandler<FNiagaraBool> OutValid(Context);
+
+	FSkeletalMeshAccessorHelper MeshAccessor;
+	MeshAccessor.Init<TIntegralConstant<ENDISkeletalMesh_FilterMode, ENDISkeletalMesh_FilterMode::None>, TIntegralConstant<ENDISkelMesh_AreaWeightingMode, ENDISkelMesh_AreaWeightingMode::None>>(InstData);
+
+	const int32 MaxVertex = MeshAccessor.LODData->GetNumVertices();
+	for (int32 i = 0; i < Context.NumInstances; ++i)
+	{
+		const int32 VertexIndex = VertexParam.GetAndAdvance();
+
+		FNiagaraBool Value;
+		Value.SetValue(VertexIndex < MaxVertex);
+		*OutValid.GetDestAndAdvance() = Value;
+	}
+}
+
+void UNiagaraDataInterfaceSkeletalMesh::RandomVertex(FVectorVMContext& Context)
+{
+	SCOPE_CYCLE_COUNTER(STAT_NiagaraSkel_Vertex_Sample);
+
+	FNDIRandomHelper RandHelper(Context);
+	VectorVM::FUserPtrHandler<FNDISkeletalMesh_InstanceData> InstData(Context);
+	VectorVM::FExternalFuncRegisterHandler<int32> OutVertex(Context);
+
+	FSkeletalMeshAccessorHelper MeshAccessor;
+	MeshAccessor.Init<TIntegralConstant<ENDISkeletalMesh_FilterMode, ENDISkeletalMesh_FilterMode::None>, TIntegralConstant<ENDISkelMesh_AreaWeightingMode, ENDISkelMesh_AreaWeightingMode::None>>(InstData);
+
+	const int32 MaxVertex = MeshAccessor.LODData->GetNumVertices() - 1;
+	if (MaxVertex >= 0)
+	{
+		for (int32 i = 0; i < Context.NumInstances; ++i)
+		{
+			RandHelper.GetAndAdvance();
+			*OutVertex.GetDestAndAdvance() = RandHelper.RandRange(i, 0, MaxVertex);
+		}
+	}
+	else
+	{
+		FMemory::Memset(OutVertex.GetDest(), 0xff, Context.NumInstances * sizeof(int32));
+	}
+}
+
+void UNiagaraDataInterfaceSkeletalMesh::GetVertexCount(FVectorVMContext& Context)
+{
+	SCOPE_CYCLE_COUNTER(STAT_NiagaraSkel_Vertex_Sample);
+
+	VectorVM::FUserPtrHandler<FNDISkeletalMesh_InstanceData> InstData(Context);
+	VectorVM::FExternalFuncRegisterHandler<int32> OutVertexCount(Context);
+
+	FSkeletalMeshAccessorHelper MeshAccessor;
+	MeshAccessor.Init<TIntegralConstant<ENDISkeletalMesh_FilterMode, ENDISkeletalMesh_FilterMode::None>, TIntegralConstant<ENDISkelMesh_AreaWeightingMode, ENDISkelMesh_AreaWeightingMode::None>>(InstData);
+
+	const int32 MaxVertex = MeshAccessor.LODData->GetNumVertices();
+	for (int32 i = 0; i < Context.NumInstances; ++i)
+	{
+		*OutVertexCount.GetDestAndAdvance() = MaxVertex;
+	}
+}
+
+//////////////////////////////////////////////////////////////////////////
 template<typename FilterMode>
-FORCEINLINE int32 UNiagaraDataInterfaceSkeletalMesh::RandomVertIndex(FNDIRandomHelper& RandHelper, int32 Instance, FSkeletalMeshAccessorHelper& Accessor, FNDISkeletalMesh_InstanceData* InstData)
+FORCEINLINE int32 UNiagaraDataInterfaceSkeletalMesh::RandomFilteredVertIndex(FNDIRandomHelper& RandHelper, int32 Instance, FSkeletalMeshAccessorHelper& Accessor, FNDISkeletalMesh_InstanceData* InstData)
 {
 	checkf(false, TEXT("Invalid template call for RandomVertIndex. Bug in Filter binding or Area Weighting binding. Contact code team."));
 	return 0;
 }
 
 template<>
-FORCEINLINE int32 UNiagaraDataInterfaceSkeletalMesh::RandomVertIndex<TIntegralConstant<ENDISkeletalMesh_FilterMode, ENDISkeletalMesh_FilterMode::None>>
+FORCEINLINE int32 UNiagaraDataInterfaceSkeletalMesh::RandomFilteredVertIndex<TIntegralConstant<ENDISkeletalMesh_FilterMode, ENDISkeletalMesh_FilterMode::None>>
 	(FNDIRandomHelper& RandHelper, int32 Instance, FSkeletalMeshAccessorHelper& Accessor, FNDISkeletalMesh_InstanceData* InstData)
 {
 	return RandHelper.RandRange(Instance, 0, Accessor.LODData->GetNumVertices() - 1);
 }
 
 template<>
-FORCEINLINE int32 UNiagaraDataInterfaceSkeletalMesh::RandomVertIndex<TIntegralConstant<ENDISkeletalMesh_FilterMode, ENDISkeletalMesh_FilterMode::SingleRegion>>
+FORCEINLINE int32 UNiagaraDataInterfaceSkeletalMesh::RandomFilteredVertIndex<TIntegralConstant<ENDISkeletalMesh_FilterMode, ENDISkeletalMesh_FilterMode::SingleRegion>>
 	(FNDIRandomHelper& RandHelper, int32 Instance, FSkeletalMeshAccessorHelper& Accessor, FNDISkeletalMesh_InstanceData* InstData)
 {
 	int32 Idx = RandHelper.RandRange(Instance, 0, Accessor.SamplingRegionBuiltData->Vertices.Num() - 1);
@@ -211,7 +325,7 @@
 }
 
 template<>
-FORCEINLINE int32 UNiagaraDataInterfaceSkeletalMesh::RandomVertIndex<TIntegralConstant<ENDISkeletalMesh_FilterMode, ENDISkeletalMesh_FilterMode::MultiRegion>>
+FORCEINLINE int32 UNiagaraDataInterfaceSkeletalMesh::RandomFilteredVertIndex<TIntegralConstant<ENDISkeletalMesh_FilterMode, ENDISkeletalMesh_FilterMode::MultiRegion>>
 	(FNDIRandomHelper& RandHelper, int32 Instance, FSkeletalMeshAccessorHelper& Accessor, FNDISkeletalMesh_InstanceData* InstData)
 {
 	int32 RegionIdx = RandHelper.RandRange(Instance, 0, InstData->SamplingRegionIndices.Num() - 1);
@@ -223,7 +337,7 @@
 }
 
 template<typename FilterMode>
-void UNiagaraDataInterfaceSkeletalMesh::RandomVertex(FVectorVMContext& Context)
+void UNiagaraDataInterfaceSkeletalMesh::RandomFilteredVertex(FVectorVMContext& Context)
 {
 	SCOPE_CYCLE_COUNTER(STAT_NiagaraSkel_Vertex_Sample);
 
@@ -240,12 +354,12 @@
 	for (int32 i = 0; i < Context.NumInstances; ++i)
 	{
 		RandHelper.GetAndAdvance();
-		*OutVert.GetDestAndAdvance() = RandomVertIndex<FilterMode>(RandHelper, i, MeshAccessor, InstData);
+		*OutVert.GetDestAndAdvance() = RandomFilteredVertIndex<FilterMode>(RandHelper, i, MeshAccessor, InstData);
 	}
 }
 
 template<typename FilterMode>
-void UNiagaraDataInterfaceSkeletalMesh::IsValidVertex(FVectorVMContext& Context)
+void UNiagaraDataInterfaceSkeletalMesh::IsValidFilteredVertex(FVectorVMContext& Context)
 {
 	SCOPE_CYCLE_COUNTER(STAT_NiagaraSkel_Vertex_Sample);
 
@@ -331,7 +445,6 @@
 	}
 }
 
-
 //////////////////////////////////////////////////////////////////////////
 
 template<typename FilterMode>
@@ -402,6 +515,8 @@
 	}
 }
 
+//////////////////////////////////////////////////////////////////////////
+
 void UNiagaraDataInterfaceSkeletalMesh::GetVertexColor(FVectorVMContext& Context)
 {
 	SCOPE_CYCLE_COUNTER(STAT_NiagaraSkel_Vertex_Sample);
diff --git a/Engine/Plugins/FX/Niagara/Source/Niagara/Private/NiagaraDataInterfaceSkeletalMesh.cpp b/Engine/Plugins/FX/Niagara/Source/Niagara/Private/NiagaraDataInterfaceSkeletalMesh.cpp
--- a/Engine/Plugins/FX/Niagara/Source/Niagara/Private/NiagaraDataInterfaceSkeletalMesh.cpp
+++ b/Engine/Plugins/FX/Niagara/Source/Niagara/Private/NiagaraDataInterfaceSkeletalMesh.cpp
@@ -28,6 +28,7 @@ struct FNiagaraSkelMeshDIFunctionVersion
 	{
 		InitialVersion = 0,
 		AddedRandomInfo = 1,
+		CleanUpVertexSampling = 2,
 
 		VersionPlusOne,
 		LatestVersion = VersionPlusOne - 1
@@ -505,6 +506,57 @@ void FSkeletalMeshGpuSpawnStaticBuffers::Initialise(FNDISkeletalMesh_InstanceDat
 
 		NumFilteredSockets = InstData->FilteredSocketInfo.Num();
 		FilteredSocketBoneOffset = InstData->FilteredSocketBoneOffset;
+
+		// Create triangle / vertex region sampling data
+		if (InstData->SamplingRegionIndices.Num() > 0)
+		{
+			const FSkeletalMeshSamplingInfo& SamplingInfo = InstData->Mesh->GetSamplingInfo();
+
+			// Count required regions
+			bSamplingRegionsAllAreaWeighted = true;
+			NumSamplingRegionTriangles = 0;
+			NumSamplingRegionVertices = 0;
+
+			for (const int32 RegionIndex : InstData->SamplingRegionIndices)
+			{
+				const FSkeletalMeshSamplingRegionBuiltData& SamplingRegionBuildData = SamplingInfo.GetRegionBuiltData(RegionIndex);
+				NumSamplingRegionTriangles += SamplingRegionBuildData.TriangleIndices.Num();
+				NumSamplingRegionVertices += SamplingRegionBuildData.Vertices.Num();
+				bSamplingRegionsAllAreaWeighted &= SamplingRegionBuildData.AreaWeightedSampler.GetNumEntries() == SamplingRegionBuildData.TriangleIndices.Num();
+			}
+
+			// Build buffers
+			SampleRegionsProb.Reserve(NumSamplingRegionTriangles);
+			SampleRegionsAlias.Reserve(NumSamplingRegionTriangles);
+			SampleRegionsTriangleIndicies.Reserve(NumSamplingRegionTriangles);
+			SampleRegionsVertices.Reserve(NumSamplingRegionVertices);
+
+			int32 RegionOffset = 0;
+			for (const int32 RegionIndex : InstData->SamplingRegionIndices)
+			{
+				const FSkeletalMeshSamplingRegionBuiltData& SamplingRegionBuildData = SamplingInfo.GetRegionBuiltData(RegionIndex);
+				if (bSamplingRegionsAllAreaWeighted)
+				{
+					for (float v : SamplingRegionBuildData.AreaWeightedSampler.GetProb())
+					{
+						SampleRegionsProb.Add(v);
+					}
+					for (int v : SamplingRegionBuildData.AreaWeightedSampler.GetAlias())
+					{
+						SampleRegionsAlias.Add(v + RegionOffset);
+					}
+				}
+				for (int v : SamplingRegionBuildData.TriangleIndices)
+				{
+					SampleRegionsTriangleIndicies.Add(v / 3);
+				}
+				for (int v : SamplingRegionBuildData.Vertices)
+				{
+					SampleRegionsVertices.Add(v);
+				}
+				RegionOffset += SamplingRegionBuildData.TriangleIndices.Num();
+			}
+		}
 	}
 }
 
@@ -556,6 +608,29 @@ void FSkeletalMeshGpuSpawnStaticBuffers::InitRHI()
 		BufferTriangleUniformSamplerAliasSRV = RHICreateShaderResourceView(BufferTriangleUniformSamplerAliasRHI, sizeof(uint32), PF_R32_UINT);
 	}
 
+	// Prepare sampling regions (if we have any)
+	if (NumSamplingRegionTriangles > 0)
+	{
+		FRHIResourceCreateInfo CreateInfo;
+		if (bSamplingRegionsAllAreaWeighted)
+		{
+			CreateInfo.ResourceArray = &SampleRegionsProb;
+			SampleRegionsProbBuffer = RHICreateVertexBuffer(SampleRegionsProb.Num() * SampleRegionsProb.GetTypeSize(), BUF_Static | BUF_ShaderResource, CreateInfo);
+			SampleRegionsProbSRV = RHICreateShaderResourceView(SampleRegionsTriangleIndicesBuffer, sizeof(float), PF_R32_FLOAT);
+
+			CreateInfo.ResourceArray = &SampleRegionsAlias;
+			SampleRegionsAliasBuffer = RHICreateVertexBuffer(SampleRegionsAlias.Num() * SampleRegionsAlias.GetTypeSize(), BUF_Static | BUF_ShaderResource, CreateInfo);
+			SampleRegionsAliasSRV = RHICreateShaderResourceView(SampleRegionsAliasBuffer, sizeof(float), PF_R32_UINT);
+		}
+		CreateInfo.ResourceArray = &SampleRegionsTriangleIndicies;
+		SampleRegionsTriangleIndicesBuffer = RHICreateVertexBuffer(SampleRegionsTriangleIndicies.Num() * SampleRegionsTriangleIndicies.GetTypeSize(), BUF_Static | BUF_ShaderResource, CreateInfo);
+		SampleRegionsTriangleIndicesSRV = RHICreateShaderResourceView(SampleRegionsTriangleIndicesBuffer, sizeof(int32), PF_R32_UINT);
+
+		CreateInfo.ResourceArray = &SampleRegionsVertices;
+		SampleRegionsVerticesBuffer = RHICreateVertexBuffer(SampleRegionsVertices.Num() * SampleRegionsVertices.GetTypeSize(), BUF_Static | BUF_ShaderResource, CreateInfo);
+		SampleRegionsVerticesSRV = RHICreateShaderResourceView(SampleRegionsVerticesBuffer, sizeof(int32), PF_R32_UINT);
+	}
+
 	// Prepare the vertex matrix lookup offset for each of the sections. This is needed because per vertex BlendIndicies are stored relatively to each Section used matrices.
 	// And these offset per section need to point to the correct matrix according to each section BoneMap.
 	// There is not section selection/culling in the interface so technically we could compute that array in the pipeline.
@@ -601,6 +676,15 @@ void FSkeletalMeshGpuSpawnStaticBuffers::ReleaseRHI()
 	BufferTriangleUniformSamplerAliasRHI.SafeRelease();
 	BufferTriangleUniformSamplerAliasSRV.SafeRelease();
 
+	SampleRegionsProbBuffer.SafeRelease();
+	SampleRegionsProbSRV.SafeRelease();
+	SampleRegionsAliasBuffer.SafeRelease();
+	SampleRegionsAliasSRV.SafeRelease();
+	SampleRegionsTriangleIndicesBuffer.SafeRelease();
+	SampleRegionsTriangleIndicesSRV.SafeRelease();
+	SampleRegionsVerticesBuffer.SafeRelease();
+	SampleRegionsVerticesSRV.SafeRelease();
+
 	MeshVertexBufferSrv = nullptr;
 	MeshIndexBufferSrv = nullptr;
 	MeshTangentBufferSRV = nullptr;
@@ -853,6 +937,12 @@ struct FNDISkeletalMeshParametersName
 	FString MeshColorBufferName;
 	FString MeshTriangleSamplerProbaBufferName;
 	FString MeshTriangleSamplerAliasBufferName;
+	FString MeshNumSamplingRegionTrianglesName;
+	FString MeshNumSamplingRegionVerticesName;
+	FString MeshSamplingRegionsProbaBufferName;
+	FString MeshSamplingRegionsAliasBufferName;
+	FString MeshSampleRegionsTriangleIndicesName;
+	FString MeshSampleRegionsVerticesName;
 	FString MeshTriangleMatricesOffsetBufferName;
 	FString MeshTriangleCountName;
 	FString MeshVertexCountName;
@@ -887,6 +977,12 @@ static void GetNiagaraDataInterfaceParametersName(FNDISkeletalMeshParametersName
 	Names.MeshColorBufferName = UNiagaraDataInterfaceSkeletalMesh::MeshColorBufferName + Suffix;
 	Names.MeshTriangleSamplerProbaBufferName = UNiagaraDataInterfaceSkeletalMesh::MeshTriangleSamplerProbaBufferName + Suffix;
 	Names.MeshTriangleSamplerAliasBufferName = UNiagaraDataInterfaceSkeletalMesh::MeshTriangleSamplerAliasBufferName + Suffix;
+	Names.MeshNumSamplingRegionTrianglesName = UNiagaraDataInterfaceSkeletalMesh::MeshNumSamplingRegionTrianglesName + Suffix;
+	Names.MeshNumSamplingRegionVerticesName = UNiagaraDataInterfaceSkeletalMesh::MeshNumSamplingRegionVerticesName + Suffix;
+	Names.MeshSamplingRegionsProbaBufferName = UNiagaraDataInterfaceSkeletalMesh::MeshSamplingRegionsProbaBufferName + Suffix;
+	Names.MeshSamplingRegionsAliasBufferName = UNiagaraDataInterfaceSkeletalMesh::MeshSamplingRegionsAliasBufferName + Suffix;
+	Names.MeshSampleRegionsTriangleIndicesName = UNiagaraDataInterfaceSkeletalMesh::MeshSampleRegionsTriangleIndicesName + Suffix;
+	Names.MeshSampleRegionsVerticesName = UNiagaraDataInterfaceSkeletalMesh::MeshSampleRegionsVerticesName + Suffix;
 	Names.MeshTriangleMatricesOffsetBufferName = UNiagaraDataInterfaceSkeletalMesh::MeshTriangleMatricesOffsetBufferName + Suffix;
 	Names.MeshTriangleCountName = UNiagaraDataInterfaceSkeletalMesh::MeshTriangleCountName + Suffix;
 	Names.MeshVertexCountName = UNiagaraDataInterfaceSkeletalMesh::MeshVertexCountName + Suffix;
@@ -928,6 +1024,12 @@ public:
 		MeshColorBuffer.Bind(ParameterMap, *ParamNames.MeshColorBufferName);
 		MeshTriangleSamplerProbaBuffer.Bind(ParameterMap, *ParamNames.MeshTriangleSamplerProbaBufferName);
 		MeshTriangleSamplerAliasBuffer.Bind(ParameterMap, *ParamNames.MeshTriangleSamplerAliasBufferName);
+		MeshNumSamplingRegionTriangles.Bind(ParameterMap, *ParamNames.MeshNumSamplingRegionTrianglesName);
+		MeshNumSamplingRegionVertices.Bind(ParameterMap, *ParamNames.MeshNumSamplingRegionVerticesName);
+		MeshSamplingRegionsProbaBuffer.Bind(ParameterMap, *ParamNames.MeshSamplingRegionsProbaBufferName);
+		MeshSamplingRegionsAliasBuffer.Bind(ParameterMap, *ParamNames.MeshSamplingRegionsAliasBufferName);
+		MeshSampleRegionsTriangleIndices.Bind(ParameterMap, *ParamNames.MeshSampleRegionsTriangleIndicesName);
+		MeshSampleRegionsVertices.Bind(ParameterMap, *ParamNames.MeshSampleRegionsVerticesName);
 		MeshTriangleMatricesOffsetBuffer.Bind(ParameterMap, *ParamNames.MeshTriangleMatricesOffsetBufferName);
 		MeshTriangleCount.Bind(ParameterMap, *ParamNames.MeshTriangleCountName);
 		MeshVertexCount.Bind(ParameterMap, *ParamNames.MeshVertexCountName);
@@ -981,6 +1083,8 @@ public:
 			}
 			SetShaderValue(RHICmdList, ComputeShaderRHI, MeshTriangleCount, StaticBuffers->GetTriangleCount());
 			SetShaderValue(RHICmdList, ComputeShaderRHI, MeshVertexCount, StaticBuffers->GetVertexCount());
+
+			// Set triangle sampling buffer
 			if (InstanceData->bIsGpuUniformlyDistributedSampling)
 			{
 				SetSRVParameter(RHICmdList, ComputeShaderRHI, MeshTriangleSamplerProbaBuffer, StaticBuffers->GetBufferTriangleUniformSamplerProbaSRV().GetReference());
@@ -988,18 +1092,35 @@ public:
 			}
 			else
 			{
-				SetSRVParameter(RHICmdList, ComputeShaderRHI, MeshTriangleSamplerProbaBuffer, FNiagaraRenderer::GetDummyUIntBuffer());
+				SetSRVParameter(RHICmdList, ComputeShaderRHI, MeshTriangleSamplerProbaBuffer, FNiagaraRenderer::GetDummyFloatBuffer());
 				SetSRVParameter(RHICmdList, ComputeShaderRHI, MeshTriangleSamplerAliasBuffer, FNiagaraRenderer::GetDummyUIntBuffer());
 			}
 
+			// Set triangle sampling region buffer
+			SetShaderValue(RHICmdList, ComputeShaderRHI, MeshNumSamplingRegionTriangles, StaticBuffers->GetNumSamplingRegionTriangles());
+			SetShaderValue(RHICmdList, ComputeShaderRHI, MeshNumSamplingRegionVertices, StaticBuffers->GetNumSamplingRegionVertices());
+			if (StaticBuffers->IsSamplingRegionsAllAreaWeighted() && StaticBuffers->GetNumSamplingRegionTriangles() > 0)
+			{
+				SetSRVParameter(RHICmdList, ComputeShaderRHI, MeshSamplingRegionsProbaBuffer, StaticBuffers->GetSampleRegionsProbSRV().GetReference());
+				SetSRVParameter(RHICmdList, ComputeShaderRHI, MeshSamplingRegionsAliasBuffer, StaticBuffers->GetSampleRegionsAliasSRV().GetReference());
+			}
+			else
+			{
+				SetSRVParameter(RHICmdList, ComputeShaderRHI, MeshSamplingRegionsProbaBuffer, FNiagaraRenderer::GetDummyFloatBuffer());
+				SetSRVParameter(RHICmdList, ComputeShaderRHI, MeshSamplingRegionsAliasBuffer, FNiagaraRenderer::GetDummyUIntBuffer());
+			}
+			SetSRVParameter(RHICmdList, ComputeShaderRHI, MeshSampleRegionsTriangleIndices, StaticBuffers->GetNumSamplingRegionTriangles() > 0 ? StaticBuffers->GetSampleRegionsTriangleIndicesSRV().GetReference() : FNiagaraRenderer::GetDummyUIntBuffer());
+			SetSRVParameter(RHICmdList, ComputeShaderRHI, MeshSampleRegionsVertices, StaticBuffers->GetNumSamplingRegionVertices() > 0 ? StaticBuffers->GetSampleRegionsVerticesSRV().GetReference() : FNiagaraRenderer::GetDummyUIntBuffer());
+
 			SetSRVParameter(RHICmdList, ComputeShaderRHI, MeshSkinWeightBuffer, InstanceData->MeshSkinWeightBuffer->GetSRV());
 			SetSRVParameter(RHICmdList, ComputeShaderRHI, MeshSkinWeightLookupBuffer, InstanceData->MeshSkinWeightLookupBuffer->GetSRV());
 
 			SetShaderValue(RHICmdList, ComputeShaderRHI, MeshWeightStride, InstanceData->MeshWeightStrideByte/4);
-			SetShaderValue(RHICmdList, ComputeShaderRHI, MeshSkinWeightIndexSize, InstanceData->MeshSkinWeightIndexSizeByte);
 
-			uint32 EnabledFeaturesBits = InstanceData->bIsGpuUniformlyDistributedSampling ? 1 : 0;
-			EnabledFeaturesBits |= (InstanceData->bUnlimitedBoneInfluences ? (1 << 1) : 0);
+			uint32 EnabledFeaturesBits = 0;
+			EnabledFeaturesBits |= InstanceData->bIsGpuUniformlyDistributedSampling ? 1 : 0;
+			EnabledFeaturesBits |= StaticBuffers->IsSamplingRegionsAllAreaWeighted() ? 2 : 0;
+			EnabledFeaturesBits |= (InstanceData->bUnlimitedBoneInfluences ? 4 : 0);
 
 			FSkeletalMeshGpuDynamicBufferProxy* DynamicBuffers = InstanceData->DynamicBuffer;
 			check(DynamicBuffers);
@@ -1053,9 +1174,16 @@ public:
 			SetSRVParameter(RHICmdList, ComputeShaderRHI, MeshColorBuffer, FNiagaraRenderer::GetDummyFloatBuffer());
 			SetShaderValue(RHICmdList, ComputeShaderRHI, MeshTriangleCount, 0);
 			SetShaderValue(RHICmdList, ComputeShaderRHI, MeshVertexCount, 0);
-			SetSRVParameter(RHICmdList, ComputeShaderRHI, MeshTriangleSamplerProbaBuffer, FNiagaraRenderer::GetDummyUIntBuffer());
+			SetSRVParameter(RHICmdList, ComputeShaderRHI, MeshTriangleSamplerProbaBuffer, FNiagaraRenderer::GetDummyFloatBuffer());
 			SetSRVParameter(RHICmdList, ComputeShaderRHI, MeshTriangleSamplerAliasBuffer, FNiagaraRenderer::GetDummyUIntBuffer());
 
+			SetShaderValue(RHICmdList, ComputeShaderRHI, MeshNumSamplingRegionTriangles, 0);
+			SetShaderValue(RHICmdList, ComputeShaderRHI, MeshNumSamplingRegionVertices, 0);
+			SetSRVParameter(RHICmdList, ComputeShaderRHI, MeshSamplingRegionsProbaBuffer, FNiagaraRenderer::GetDummyFloatBuffer());
+			SetSRVParameter(RHICmdList, ComputeShaderRHI, MeshSamplingRegionsAliasBuffer, FNiagaraRenderer::GetDummyUIntBuffer());
+			SetSRVParameter(RHICmdList, ComputeShaderRHI, MeshSampleRegionsTriangleIndices, FNiagaraRenderer::GetDummyUIntBuffer());
+			SetSRVParameter(RHICmdList, ComputeShaderRHI, MeshSampleRegionsVertices, FNiagaraRenderer::GetDummyUIntBuffer());
+
 			SetSRVParameter(RHICmdList, ComputeShaderRHI, MeshSkinWeightBuffer, FNiagaraRenderer::GetDummyUIntBuffer());
 			SetSRVParameter(RHICmdList, ComputeShaderRHI, MeshSkinWeightLookupBuffer, FNiagaraRenderer::GetDummyUIntBuffer());
 
@@ -1098,6 +1226,12 @@ private:
 	LAYOUT_FIELD(FShaderResourceParameter, MeshColorBuffer);
 	LAYOUT_FIELD(FShaderResourceParameter, MeshTriangleSamplerProbaBuffer);
 	LAYOUT_FIELD(FShaderResourceParameter, MeshTriangleSamplerAliasBuffer);
+	LAYOUT_FIELD(FShaderParameter, MeshNumSamplingRegionTriangles);
+	LAYOUT_FIELD(FShaderParameter, MeshNumSamplingRegionVertices);
+	LAYOUT_FIELD(FShaderResourceParameter, MeshSamplingRegionsProbaBuffer);
+	LAYOUT_FIELD(FShaderResourceParameter, MeshSamplingRegionsAliasBuffer);
+	LAYOUT_FIELD(FShaderResourceParameter, MeshSampleRegionsTriangleIndices);
+	LAYOUT_FIELD(FShaderResourceParameter, MeshSampleRegionsVertices);
 	LAYOUT_FIELD(FShaderResourceParameter, MeshTriangleMatricesOffsetBuffer);
 	LAYOUT_FIELD(FShaderParameter, MeshTriangleCount);
 	LAYOUT_FIELD(FShaderParameter, MeshVertexCount);
@@ -1559,6 +1693,11 @@ bool FNDISkeletalMesh_InstanceData::Init(UNiagaraDataInterfaceSkeletalMesh* Inte
 			}
 		}
 	}
+	else
+	{
+		// Note: We do not allocate space in the array as that wastes memory, we handle this special case when reading from unfiltered
+		NumUnfilteredBones = RefSkel.GetNum();
+	}
 
 	// Gather filtered socket information
 	{
@@ -2189,6 +2328,12 @@ const FString UNiagaraDataInterfaceSkeletalMesh::MeshTexCoordBufferName(TEXT("Me
 const FString UNiagaraDataInterfaceSkeletalMesh::MeshColorBufferName(TEXT("MeshColorBuffer_"));
 const FString UNiagaraDataInterfaceSkeletalMesh::MeshTriangleSamplerProbaBufferName(TEXT("MeshTriangleSamplerProbaBuffer_"));
 const FString UNiagaraDataInterfaceSkeletalMesh::MeshTriangleSamplerAliasBufferName(TEXT("MeshTriangleSamplerAliasBuffer_"));
+const FString UNiagaraDataInterfaceSkeletalMesh::MeshNumSamplingRegionTrianglesName(TEXT("MeshNumSamplingRegionTriangles_"));
+const FString UNiagaraDataInterfaceSkeletalMesh::MeshNumSamplingRegionVerticesName(TEXT("MeshNumSamplingRegionVertices_"));
+const FString UNiagaraDataInterfaceSkeletalMesh::MeshSamplingRegionsProbaBufferName(TEXT("MeshSamplingRegionsProbaBuffer_"));
+const FString UNiagaraDataInterfaceSkeletalMesh::MeshSamplingRegionsAliasBufferName(TEXT("MeshSamplingRegionsAliasBuffer_"));
+const FString UNiagaraDataInterfaceSkeletalMesh::MeshSampleRegionsTriangleIndicesName(TEXT("MeshSampleRegionsTriangleIndices_"));
+const FString UNiagaraDataInterfaceSkeletalMesh::MeshSampleRegionsVerticesName(TEXT("MeshSampleRegionsVertices_"));
 const FString UNiagaraDataInterfaceSkeletalMesh::MeshTriangleMatricesOffsetBufferName(TEXT("MeshTriangleMatricesOffsetBuffer_"));
 const FString UNiagaraDataInterfaceSkeletalMesh::MeshTriangleCountName(TEXT("MeshTriangleCount_"));
 const FString UNiagaraDataInterfaceSkeletalMesh::MeshVertexCountName(TEXT("MeshVertexCount_"));
@@ -2241,6 +2386,11 @@ bool UNiagaraDataInterfaceSkeletalMesh::GetFunctionHLSL(const FNiagaraDataInterf
 		static const TCHAR* FormatSample = TEXT("void {InstanceFunctionName} (NiagaraRandInfo InRandomInfo, out {MeshTriCoordinateStructName} OutCoord) { {GetDISkelMeshContextName} DISKelMesh_RandomTriCoord(DIContext, InRandomInfo.Seed1, InRandomInfo.Seed2, InRandomInfo.Seed3, OutCoord.Tri, OutCoord.BaryCoord); }");
 		OutHLSL += FString::Format(FormatSample, ArgsSample);
 	}
+	else if (FunctionInfo.DefinitionName == FSkeletalMeshInterfaceHelper::IsValidTriCoordName)
+	{
+		static const TCHAR* FormatSample = TEXT("void {InstanceFunctionName} (in {MeshTriCoordinateStructName} InCoord, out bool IsValid) { {GetDISkelMeshContextName} IsValid = InCoord.Tri < DIContext.MeshTriangleCount; }");
+		OutHLSL += FString::Format(FormatSample, ArgsSample);
+	}
 	else if (FunctionInfo.DefinitionName == FSkeletalMeshInterfaceHelper::GetSkinnedTriangleDataWSName)
 	{
 		static const TCHAR* FormatSample = TEXT("void {InstanceFunctionName} (in {MeshTriCoordinateStructName} InCoord, out float3 OutPosition, out float3 OutVelocity, out float3 OutNormal, out float3 OutBinormal, out float3 OutTangent) { {GetDISkelMeshContextName} DISKelMesh_GetSkinnedTriangleDataWS(DIContext, InCoord.Tri, InCoord.BaryCoord, OutPosition, OutVelocity, OutNormal, OutBinormal, OutTangent); }");
@@ -2271,14 +2421,34 @@ bool UNiagaraDataInterfaceSkeletalMesh::GetFunctionHLSL(const FNiagaraDataInterf
 		static const TCHAR* FormatSample = TEXT("void {InstanceFunctionName} (in {MeshTriCoordinateStructName} InCoord, out float4 OutColor) { {GetDISkelMeshContextName} DISkelMesh_GetTriColor(DIContext, InCoord.Tri, InCoord.BaryCoord, OutColor); }");
 		OutHLSL += FString::Format(FormatSample, ArgsSample);
 	}
-	else if (FunctionInfo.DefinitionName == FSkeletalMeshInterfaceHelper::IsValidTriCoordName)
+	else if (FunctionInfo.DefinitionName == FSkeletalMeshInterfaceHelper::GetTriCoordVerticesName)
 	{
-		static const TCHAR* FormatSample = TEXT("void {InstanceFunctionName} (in {MeshTriCoordinateStructName} InCoord, out bool IsValid) { {GetDISkelMeshContextName} IsValid = InCoord.Tri < DIContext.MeshTriangleCount; }");
+		static const TCHAR* FormatSample = TEXT("void {InstanceFunctionName} (in int TriangleIndex, out int OutVertexIndex0, out int OutVertexIndex1, out int OutVertexIndex2) { {GetDISkelMeshContextName} DISkelMesh_GetTriVertices(DIContext, TriangleIndex, OutVertexIndex0, OutVertexIndex1, OutVertexIndex2); }");
 		OutHLSL += FString::Format(FormatSample, ArgsSample);
 	}
-	else if (FunctionInfo.DefinitionName == FSkeletalMeshInterfaceHelper::GetTriCoordVerticesName)
+	else if (FunctionInfo.DefinitionName == FSkeletalMeshInterfaceHelper::RandomTriangleName)
 	{
-		static const TCHAR* FormatSample = TEXT("void {InstanceFunctionName} (in int TriangleIndex, out int OutVertexIndex0, out int OutVertexIndex1, out int OutVertexIndex2) { {GetDISkelMeshContextName} DISkelMesh_GetTriVertices(DIContext, TriangleIndex, OutVertexIndex0, OutVertexIndex1, OutVertexIndex2); }");
+		static const TCHAR* FormatSample = TEXT("void {InstanceFunctionName} (NiagaraRandInfo InRandomInfo, out {MeshTriCoordinateStructName} OutCoord) { {GetDISkelMeshContextName} DISKelMesh_RandomTriangle(DIContext, InRandomInfo.Seed1, InRandomInfo.Seed2, InRandomInfo.Seed3, OutCoord.Tri, OutCoord.BaryCoord); }");
+		OutHLSL += FString::Format(FormatSample, ArgsSample);
+	}
+	else if (FunctionInfo.DefinitionName == FSkeletalMeshInterfaceHelper::GetTriangleCountName)
+	{
+		static const TCHAR* FormatSample = TEXT("void {InstanceFunctionName} (out int Count) { {GetDISkelMeshContextName} DISKelMesh_GetTriangleCount(DIContext, Count); }");
+		OutHLSL += FString::Format(FormatSample, ArgsSample);
+	}
+	else if (FunctionInfo.DefinitionName == FSkeletalMeshInterfaceHelper::RandomFilteredTriangleName)
+	{
+		static const TCHAR* FormatSample = TEXT("void {InstanceFunctionName} (NiagaraRandInfo InRandomInfo, out {MeshTriCoordinateStructName} OutCoord) { {GetDISkelMeshContextName} DISKelMesh_RandomFilteredTriangle(DIContext, InRandomInfo.Seed1, InRandomInfo.Seed2, InRandomInfo.Seed3, OutCoord.Tri, OutCoord.BaryCoord); }");
+		OutHLSL += FString::Format(FormatSample, ArgsSample);
+	}
+	else if (FunctionInfo.DefinitionName == FSkeletalMeshInterfaceHelper::GetFilteredTriangleCountName)
+	{
+		static const TCHAR* FormatSample = TEXT("void {InstanceFunctionName} (out int Count) { {GetDISkelMeshContextName} DISKelMesh_GetFilteredTriangleCount(DIContext, Count); }");
+		OutHLSL += FString::Format(FormatSample, ArgsSample);
+	}
+	else if (FunctionInfo.DefinitionName == FSkeletalMeshInterfaceHelper::GetFilteredTriangleAtName)
+	{
+		static const TCHAR* FormatSample = TEXT("void {InstanceFunctionName} (int FilteredIndex, out {MeshTriCoordinateStructName} OutCoord) { {GetDISkelMeshContextName} DISKelMesh_GetFilteredTriangleAt(DIContext, FilteredIndex, OutCoord.Tri, OutCoord.BaryCoord); }");
 		OutHLSL += FString::Format(FormatSample, ArgsSample);
 	}
 	//////////////////////////////////////////////////////////////////////////////////////////////////////////////////////////////
@@ -2305,16 +2475,6 @@ bool UNiagaraDataInterfaceSkeletalMesh::GetFunctionHLSL(const FNiagaraDataInterf
 	}
 	//////////////////////////////////////////////////////////////////////////////////////////////////////////////////////////////
 	// Vertex Sampling
-	else if (FunctionInfo.DefinitionName == FSkeletalMeshInterfaceHelper::RandomVertexName)
-	{
-		static const TCHAR* FormatSample = TEXT("void {InstanceFunctionName}(NiagaraRandInfo InRandomInfo, out int OutVertex) { {GetDISkelMeshContextName} DISkelMesh_GetRandomVertex(DIContext, InRandomInfo.Seed1, InRandomInfo.Seed2, InRandomInfo.Seed3, OutVertex); }");
-		OutHLSL += FString::Format(FormatSample, ArgsSample);
-	}
-	else if (FunctionInfo.DefinitionName == FSkeletalMeshInterfaceHelper::IsValidVertexName)
-	{
-		static const TCHAR* FormatSample = TEXT("void {InstanceFunctionName} (in int Vertex, out bool IsValid) { {GetDISkelMeshContextName} DISkelMesh_IsValidVertex(DIContext, Vertex, IsValid); }");
-		OutHLSL += FString::Format(FormatSample, ArgsSample);
-	}
 	else if (FunctionInfo.DefinitionName == FSkeletalMeshInterfaceHelper::GetSkinnedVertexDataName)
 	{
 		static const TCHAR* FormatSample = TEXT("void {InstanceFunctionName} (in int Vertex, out float3 OutPosition, out float3 OutVelocity) { {GetDISkelMeshContextName} DISkelMesh_GetSkinnedVertex(DIContext, Vertex, OutPosition, OutVelocity); }");
@@ -2335,6 +2495,42 @@ bool UNiagaraDataInterfaceSkeletalMesh::GetFunctionHLSL(const FNiagaraDataInterf
 		static const TCHAR* FormatSample = TEXT("void {InstanceFunctionName} (in int Vertex, in int UVSet, out float2 OutUV) { {GetDISkelMeshContextName} DISkelMesh_GetVertexUV(DIContext, Vertex, UVSet, OutUV); }");
 		OutHLSL += FString::Format(FormatSample, ArgsSample);
 	}
+	else if (FunctionInfo.DefinitionName == FSkeletalMeshInterfaceHelper::IsValidVertexName)
+	{
+		static const TCHAR* FormatSample = TEXT("void {InstanceFunctionName} (in int Vertex, out bool IsValid) { {GetDISkelMeshContextName} DISkelMesh_IsValidVertex(DIContext, Vertex, IsValid); }");
+		OutHLSL += FString::Format(FormatSample, ArgsSample);
+	}
+	else if (FunctionInfo.DefinitionName == FSkeletalMeshInterfaceHelper::RandomVertexName)
+	{
+		static const TCHAR* FormatSample = TEXT("void {InstanceFunctionName}(NiagaraRandInfo InRandomInfo, out int OutVertex) { {GetDISkelMeshContextName} DISkelMesh_GetRandomVertex(DIContext, InRandomInfo.Seed1, InRandomInfo.Seed2, InRandomInfo.Seed3, OutVertex); }");
+		OutHLSL += FString::Format(FormatSample, ArgsSample);
+	}
+	else if (FunctionInfo.DefinitionName == FSkeletalMeshInterfaceHelper::GetVertexCountName)
+	{
+		static const TCHAR* FormatSample = TEXT("void {InstanceFunctionName} (out int VertexCount) { {GetDISkelMeshContextName} DISkelMesh_GetVertexCount(DIContext, VertexCount); }");
+		OutHLSL += FString::Format(FormatSample, ArgsSample);
+	}
+	else if (FunctionInfo.DefinitionName == FSkeletalMeshInterfaceHelper::IsValidFilteredVertexName)
+	{
+		static const TCHAR* FormatSample = TEXT("void {InstanceFunctionName} (in int FilteredIndex, out bool IsValid) { {GetDISkelMeshContextName} DISkelMesh_IsValidFilteredVertex(DIContext, FilteredIndex, IsValid); }");
+		OutHLSL += FString::Format(FormatSample, ArgsSample);
+	}
+	else if (FunctionInfo.DefinitionName == FSkeletalMeshInterfaceHelper::RandomFilteredVertexName)
+	{
+		static const TCHAR* FormatSample = TEXT("void {InstanceFunctionName}(NiagaraRandInfo InRandomInfo, out int OutVertex) { {GetDISkelMeshContextName} DISkelMesh_GetRandomFilteredVertex(DIContext, InRandomInfo.Seed1, InRandomInfo.Seed2, InRandomInfo.Seed3, OutVertex); }");
+		OutHLSL += FString::Format(FormatSample, ArgsSample);
+	}
+	else if (FunctionInfo.DefinitionName == FSkeletalMeshInterfaceHelper::GetFilteredVertexCountName)
+	{
+		static const TCHAR* FormatSample = TEXT("void {InstanceFunctionName} (out int VertexCount) { {GetDISkelMeshContextName} DISkelMesh_GetFilteredVertexCount(DIContext, VertexCount); }");
+		OutHLSL += FString::Format(FormatSample, ArgsSample);
+	}
+	else if (FunctionInfo.DefinitionName == FSkeletalMeshInterfaceHelper::GetFilteredVertexAtName)
+	{
+		static const TCHAR* FormatSample = TEXT("void {InstanceFunctionName} (in int FilteredIndex, out int VertexIndex) { {GetDISkelMeshContextName} DISkelMesh_GetFilteredVertexAt(DIContext, FilteredIndex, VertexIndex); }");
+		OutHLSL += FString::Format(FormatSample, ArgsSample);
+	}
+
 	//////////////////////////////////////////////////////////////////////////////////////////////////////////////////////////////
 	// Filtered Bone
 	else if (FunctionInfo.DefinitionName == FSkeletalMeshInterfaceHelper::IsValidBoneName)
@@ -2415,10 +2611,6 @@ bool UNiagaraDataInterfaceSkeletalMesh::GetFunctionHLSL(const FNiagaraDataInterf
 	}
 	//////////////////////////////////////////////////////////////////////////////////////////////////////////////////////////////
 	// Unsupported functionality
-	//else if (FunctionInfo.DefinitionName == FSkeletalMeshInterfaceHelper::GetVertexCountName)				// void GetFilteredVertexCount(out int VertexCount)
-	//else if (FunctionInfo.DefinitionName == FSkeletalMeshInterfaceHelper::GetVertexAtName)					// void GetFilteredVertex(in int FilterdIndex, out int VertexIndex)
-	//else if (FunctionInfo.DefinitionName == FSkeletalMeshInterfaceHelper::GetTriangleCountName)			// void GetFilteredTriangleCount(out int OutTriangleCount)
-	//else if (FunctionInfo.DefinitionName == FSkeletalMeshInterfaceHelper::GetTriangleAtName)				// void GetFilteredTriangle(in int Triangle, out {MeshTriCoordinateStructName} OutCoord)
 	else
 	{
 		// This function is not support
@@ -2432,34 +2624,34 @@ bool UNiagaraDataInterfaceSkeletalMesh::GetFunctionHLSL(const FNiagaraDataInterf
 #if WITH_EDITORONLY_DATA
 bool UNiagaraDataInterfaceSkeletalMesh::UpgradeFunctionCall(FNiagaraFunctionSignature& FunctionSignature)
 {
-	// Perform renames
-	static const TPair<FName, FName> FunctionRenames[] =
-	{
-		MakeTuple(FName("IsValidBone"), FSkeletalMeshInterfaceHelper::IsValidBoneName),
-		MakeTuple(FName("RandomSpecificBone"), FSkeletalMeshInterfaceHelper::RandomFilteredBoneName),
-		MakeTuple(FName("GetSpecificBoneCount"), FSkeletalMeshInterfaceHelper::GetFilteredBoneCountName),
-		MakeTuple(FName("GetSpecificBone"), FSkeletalMeshInterfaceHelper::GetFilteredBoneAtName),
-		MakeTuple(FName("RandomSpecificSocketBone"), FSkeletalMeshInterfaceHelper::RandomFilteredSocketName),
-		MakeTuple(FName("GetSpecificSocketCount"), FSkeletalMeshInterfaceHelper::GetFilteredSocketCountName),
-		MakeTuple(FName("GetSpecificSocketTransform"), FSkeletalMeshInterfaceHelper::GetFilteredSocketTransformName),
-		MakeTuple(FName("GetSpecificSocketBone"), FSkeletalMeshInterfaceHelper::GetFilteredSocketBoneAtName),
-		MakeTuple(FName("RandomFilteredSocketBone"), FSkeletalMeshInterfaceHelper::RandomFilteredSocketName),
-	};
 	bool bWasChanged = false;
 
-	for (const auto& RenamePair : FunctionRenames)
+	// Renamed some functions and added Random Info to Various functions for consistency
+	if (FunctionSignature.FunctionVersion < FNiagaraSkelMeshDIFunctionVersion::AddedRandomInfo)
 	{
-		if (FunctionSignature.Name == RenamePair.Key)
+		static const TPair<FName, FName> FunctionRenames[] =
+		{
+			MakeTuple(FName("IsValidBone"), FSkeletalMeshInterfaceHelper::IsValidBoneName),
+			MakeTuple(FName("RandomSpecificBone"), FSkeletalMeshInterfaceHelper::RandomFilteredBoneName),
+			MakeTuple(FName("GetSpecificBoneCount"), FSkeletalMeshInterfaceHelper::GetFilteredBoneCountName),
+			MakeTuple(FName("GetSpecificBone"), FSkeletalMeshInterfaceHelper::GetFilteredBoneAtName),
+			MakeTuple(FName("RandomSpecificSocketBone"), FSkeletalMeshInterfaceHelper::RandomFilteredSocketName),
+			MakeTuple(FName("GetSpecificSocketCount"), FSkeletalMeshInterfaceHelper::GetFilteredSocketCountName),
+			MakeTuple(FName("GetSpecificSocketTransform"), FSkeletalMeshInterfaceHelper::GetFilteredSocketTransformName),
+			MakeTuple(FName("GetSpecificSocketBone"), FSkeletalMeshInterfaceHelper::GetFilteredSocketBoneAtName),
+			MakeTuple(FName("RandomFilteredSocketBone"), FSkeletalMeshInterfaceHelper::RandomFilteredSocketName),
+		};
+
+		for (const auto& RenamePair : FunctionRenames)
 		{
-			FunctionSignature.Name = RenamePair.Value;
-			bWasChanged = true;
-			break;
+			if (FunctionSignature.Name == RenamePair.Key)
+			{
+				FunctionSignature.Name = RenamePair.Value;
+				bWasChanged = true;
+				break;
+			}
 		}
-	}
 
-	// Upgrade functions
-	if (FunctionSignature.FunctionVersion < FNiagaraSkelMeshDIFunctionVersion::AddedRandomInfo)
-	{
 		if (FunctionSignature.Name == FSkeletalMeshInterfaceHelper::RandomTriCoordName)
 		{
 			if (FunctionSignature.Inputs.Num() == 1)
@@ -2485,6 +2677,27 @@ bool UNiagaraDataInterfaceSkeletalMesh::UpgradeFunctionCall(FNiagaraFunctionSign
 		}
 	}
 
+	// Vertex sampling clean up
+	if (FunctionSignature.FunctionVersion < FNiagaraSkelMeshDIFunctionVersion::CleanUpVertexSampling)
+	{
+		static const TPair<FName, FName> FunctionRenames[] =
+		{
+			MakeTuple(FName("IsValidVertex"), FSkeletalMeshInterfaceHelper::IsValidVertexName),
+			MakeTuple(FName("RandomVertex"), FSkeletalMeshInterfaceHelper::RandomFilteredVertexName),
+		};
+
+		for (const auto& RenamePair : FunctionRenames)
+		{
+			if (FunctionSignature.Name == RenamePair.Key)
+			{
+				FunctionSignature.Name = RenamePair.Value;
+				bWasChanged = true;
+				break;
+			}
+		}
+	}
+	FunctionSignature.FunctionVersion = FNiagaraSkelMeshDIFunctionVersion::LatestVersion;
+
 	return bWasChanged;
 }
 #endif
diff --git a/Engine/Plugins/FX/Niagara/Source/NiagaraCore/Private/NiagaraCustomVersion.cpp b/Engine/Plugins/FX/Niagara/Source/NiagaraCore/Private/NiagaraCustomVersion.cpp
--- a/Engine/Plugins/FX/Niagara/Source/NiagaraCore/Private/NiagaraCustomVersion.cpp
+++ b/Engine/Plugins/FX/Niagara/Source/NiagaraCore/Private/NiagaraCustomVersion.cpp
@@ -8,4 +8,4 @@ const FGuid FNiagaraCustomVersion::GUID(0xFCF57AFA, 0x50764283, 0xB9A9E658, 0xFF
 // Register the custom version with core
 FCustomVersionRegistration GRegisterNiagaraCustomVersion(FNiagaraCustomVersion::GUID, FNiagaraCustomVersion::LatestVersion, TEXT("NiagaraVer"));
 
-const FGuid FNiagaraCustomVersion::LatestScriptCompileVersion(0x069BDBD3, 0xAEC9455E, 0x8F6DB20F, 0xCA920B88);
+const FGuid FNiagaraCustomVersion::LatestScriptCompileVersion(0x6D5D0BD9, 0xF49C0D41, 0xA69D8E86, 0xE897BEA0);
\ No newline at end of file
