- Fix compile error

#rnx

diff --git a/Engine/Plugins/FX/Niagara/Source/Niagara/Private/NiagaraSystem.cpp b/Engine/Plugins/FX/Niagara/Source/Niagara/Private/NiagaraSystem.cpp
--- a/Engine/Plugins/FX/Niagara/Source/Niagara/Private/NiagaraSystem.cpp
+++ b/Engine/Plugins/FX/Niagara/Source/Niagara/Private/NiagaraSystem.cpp
@@ -1413,12 +1413,12 @@ void UNiagaraSystem::AddToInstanceCountStat(int32 NumInstances, bool bSolo)const
 		if (bSolo)
 		{
 			FThreadStats::AddMessage(StatID_InstanceCountSolo.GetName(), EStatOperation::Add, int64(NumInstances));
-			TRACE_STAT_ADD(StatID_InstanceCount.GetName(), int64(Value));
+			TRACE_STAT_ADD(StatID_InstanceCount.GetName(), int64(NumInstances));
 		}
 		else
 		{
 			FThreadStats::AddMessage(StatID_InstanceCount.GetName(), EStatOperation::Add, int64(NumInstances));
-			TRACE_STAT_ADD(StatID_InstanceCount.GetName(), int64(Value));
+			TRACE_STAT_ADD(StatID_InstanceCount.GetName(), int64(NumInstances));
 		}
 	}
 #endif
