Niagara - When undo/redo happens, defer refreshing the selection view model until the the other stack view models have been refreshed to prevent entries with invalid data from being refreshed before they can clean up.  Fixes undo crashes.

diff --git a/Engine/Plugins/FX/Niagara/Source/NiagaraEditor/Private/ViewModels/NiagaraSystemViewModel.cpp b/Engine/Plugins/FX/Niagara/Source/NiagaraEditor/Private/ViewModels/NiagaraSystemViewModel.cpp
--- a/Engine/Plugins/FX/Niagara/Source/NiagaraEditor/Private/ViewModels/NiagaraSystemViewModel.cpp
+++ b/Engine/Plugins/FX/Niagara/Source/NiagaraEditor/Private/ViewModels/NiagaraSystemViewModel.cpp
@@ -898,6 +898,8 @@ void FNiagaraSystemViewModel::RefreshAll()
 	RefreshSequencerTracks();
 	ResetCurveData();
 	ScriptScratchPadViewModel->RefreshScriptViewModels();
+	SystemStackViewModel->GetRootEntry()->RefreshChildren();
+	SelectionViewModel->Refresh();
 }
 
 void FNiagaraSystemViewModel::NotifyDataObjectChanged(UObject* ChangedObject)
@@ -1039,12 +1041,6 @@ void FNiagaraSystemViewModel::RefreshEmitterHandleViewModels()
 		}
 	}
 	GetSystem().SetIsolateEnabled(bAnyEmitterIsolated);
-
-	if (SelectionViewModel != nullptr)
-	{
-		SelectionViewModel->Refresh();
-	}
-
 	OnEmitterHandleViewModelsChangedDelegate.Broadcast();
 }
 
diff --git a/Engine/Plugins/FX/Niagara/Source/NiagaraEditor/Private/ViewModels/Stack/NiagaraStackViewModel.cpp b/Engine/Plugins/FX/Niagara/Source/NiagaraEditor/Private/ViewModels/Stack/NiagaraStackViewModel.cpp
--- a/Engine/Plugins/FX/Niagara/Source/NiagaraEditor/Private/ViewModels/Stack/NiagaraStackViewModel.cpp
+++ b/Engine/Plugins/FX/Niagara/Source/NiagaraEditor/Private/ViewModels/Stack/NiagaraStackViewModel.cpp
@@ -90,8 +90,6 @@ void UNiagaraStackViewModel::InitializeWithViewModels(TSharedPtr<FNiagaraSystemV
 
 	if (SystemViewModelPinned.IsValid())
 	{
-		GEditor->RegisterForUndo(this);
-
 		if (EmitterViewModel.IsValid())
 		{
 			EmitterViewModel->OnScriptCompiled().AddUObject(this, &UNiagaraStackViewModel::OnEmitterCompiled);
@@ -131,7 +129,6 @@ void UNiagaraStackViewModel::InitializeWithRootEntry(UNiagaraStackEntry* InRootE
 
 	bExternalRootEntry = true;
 
-	GEditor->RegisterForUndo(this);
 	StructureChangedDelegate.Broadcast();
 }
 
@@ -164,8 +161,6 @@ void UNiagaraStackViewModel::Reset()
 	}
 
 	TopLevelViewModels.Empty();
-	
-	GEditor->UnregisterForUndo(this);
 
 	CurrentIssueCycleIndex = -1;
 	CurrentFocusedSearchMatchIndex = -1;
@@ -651,12 +646,6 @@ void UNiagaraStackViewModel::NotifyStructureChanged()
 	EntryStructureChanged();
 }
 
-void UNiagaraStackViewModel::PostUndo(bool bSuccess)
-{
-	RootEntry->RefreshChildren();
-	OnSearchTextChanged(CurrentSearchText);
-}
-
 void UNiagaraStackViewModel::EntryStructureChanged()
 {
 	if (bUsesTopLevelViewModels)
diff --git a/Engine/Plugins/FX/Niagara/Source/NiagaraEditor/Public/ViewModels/Stack/NiagaraStackViewModel.h b/Engine/Plugins/FX/Niagara/Source/NiagaraEditor/Public/ViewModels/Stack/NiagaraStackViewModel.h
--- a/Engine/Plugins/FX/Niagara/Source/NiagaraEditor/Public/ViewModels/Stack/NiagaraStackViewModel.h
+++ b/Engine/Plugins/FX/Niagara/Source/NiagaraEditor/Public/ViewModels/Stack/NiagaraStackViewModel.h
@@ -3,7 +3,6 @@
 #pragma once
 
 #include "Templates/SharedPointer.h"
-#include "EditorUndoClient.h"
 #include "ViewModels/Stack/NiagaraStackEntry.h"
 #include "UObject/ObjectKey.h"
 #include "NiagaraStackViewModel.generated.h"
@@ -36,7 +35,7 @@ private:
 };
 
 UCLASS()
-class NIAGARAEDITOR_API UNiagaraStackViewModel : public UObject, public FEditorUndoClient
+class NIAGARAEDITOR_API UNiagaraStackViewModel : public UObject
 {
 	GENERATED_BODY()
 
@@ -109,10 +108,6 @@ public:
 
 	void NotifyStructureChanged();
 
-	//~ FEditorUndoClient Interface
-	virtual void PostUndo(bool bSuccess) override;
-	virtual void PostRedo(bool bSuccess) override { PostUndo(bSuccess); }
-
 	virtual void Tick();
 	//~ stack search stuff
 	void OnSearchTextChanged(const FText& SearchText);
