Niagara - Remove the old parameters panel from the system and script toolkits and replace it with the new parameters panel.

diff --git a/Engine/Plugins/FX/Niagara/Source/NiagaraEditor/Private/Toolkits/NiagaraScriptToolkit.cpp b/Engine/Plugins/FX/Niagara/Source/NiagaraEditor/Private/Toolkits/NiagaraScriptToolkit.cpp
--- a/Engine/Plugins/FX/Niagara/Source/NiagaraEditor/Private/Toolkits/NiagaraScriptToolkit.cpp
+++ b/Engine/Plugins/FX/Niagara/Source/NiagaraEditor/Private/Toolkits/NiagaraScriptToolkit.cpp
@@ -58,7 +58,6 @@ const FName FNiagaraScriptToolkit::NodeGraphTabId(TEXT("NiagaraEditor_NodeGraph"
 const FName FNiagaraScriptToolkit::ScriptDetailsTabId(TEXT("NiagaraEditor_ScriptDetails"));
 const FName FNiagaraScriptToolkit::SelectedDetailsTabId(TEXT("NiagaraEditor_SelectedDetails"));
 const FName FNiagaraScriptToolkit::ParametersTabId(TEXT("NiagaraEditor_Parameters"));
-const FName FNiagaraScriptToolkit::ParametersTabId2(TEXT("NiagaraEditor_Paramters2"));
 const FName FNiagaraScriptToolkit::StatsTabId(TEXT("NiagaraEditor_Stats"));
 const FName FNiagaraScriptToolkit::MessageLogTabID(TEXT("NiagaraEditor_MessageLog"));
 
@@ -100,10 +99,6 @@ void FNiagaraScriptToolkit::RegisterTabSpawners(const TSharedRef<class FTabManag
 	InTabManager->RegisterTabSpawner(ParametersTabId, FOnSpawnTab::CreateSP(this, &FNiagaraScriptToolkit::SpawnTabScriptParameters))
 		.SetDisplayName(LOCTEXT("ParametersTab", "Parameters"))
 		.SetGroup(WorkspaceMenuCategoryRef);
-
-	InTabManager->RegisterTabSpawner(ParametersTabId2, FOnSpawnTab::CreateSP(this, &FNiagaraScriptToolkit::SpawnTabScriptParameters2))
-		.SetDisplayName(LOCTEXT("ParametersTab2", "Parameters2"))
-		.SetGroup(WorkspaceMenuCategoryRef);
 		
 	InTabManager->RegisterTabSpawner(StatsTabId, FOnSpawnTab::CreateSP(this, &FNiagaraScriptToolkit::SpawnTabStats))
 		.SetDisplayName(LOCTEXT("StatsTab", "Stats"))
@@ -123,7 +118,6 @@ void FNiagaraScriptToolkit::UnregisterTabSpawners(const TSharedRef<class FTabMan
 	InTabManager->UnregisterTabSpawner(ScriptDetailsTabId);
 	InTabManager->UnregisterTabSpawner(SelectedDetailsTabId);
 	InTabManager->UnregisterTabSpawner(ParametersTabId);
-	InTabManager->UnregisterTabSpawner(ParametersTabId2);
 	InTabManager->UnregisterTabSpawner(StatsTabId);
 }
 
@@ -176,7 +170,7 @@ void FNiagaraScriptToolkit::Initialize( const EToolkitMode::Type Mode, const TSh
 	StatsListing = MessageLogModule.CreateLogListing("MaterialEditorStats", LogOptions);
 	Stats = MessageLogModule.CreateLogListingWidget(StatsListing.ToSharedRef());
 
-	TSharedRef<FTabManager::FLayout> StandaloneDefaultLayout = FTabManager::NewLayout("Standalone_Niagara_Layout_v9")
+	TSharedRef<FTabManager::FLayout> StandaloneDefaultLayout = FTabManager::NewLayout("Standalone_Niagara_Layout_v10")
 	->AddArea
 	(
 		FTabManager::NewPrimaryArea()->SetOrientation(Orient_Vertical)
@@ -382,23 +376,6 @@ TSharedRef<SDockTab> FNiagaraScriptToolkit::SpawnTabScriptParameters(const FSpaw
 {
 	checkf(Args.GetTabId().TabType == ParametersTabId, TEXT("Wrong tab ID in NiagaraScriptToolkit"));
 
-	TArray<TSharedRef<FNiagaraObjectSelection>> Array;
-	Array.Push(DetailsScriptSelection.ToSharedRef());
-	Array.Push(ScriptViewModel->GetVariableSelection());
-
-	TSharedRef<SDockTab> SpawnedTab =
-		SNew(SDockTab)
-		[
-			SNew(SNiagaraParameterMapView, Array, SNiagaraParameterMapView::EToolkitType::SCRIPT, GetToolkitCommands())
-		];
-
-	return SpawnedTab;
-}
-
-TSharedRef<SDockTab> FNiagaraScriptToolkit::SpawnTabScriptParameters2(const FSpawnTabArgs& Args)
-{
-	checkf(Args.GetTabId().TabType == ParametersTabId2, TEXT("Wrong tab ID in NiagaraScriptToolkit"));
-
 	TSharedRef<SDockTab> SpawnedTab =
 		SNew(SDockTab)
 		[
diff --git a/Engine/Plugins/FX/Niagara/Source/NiagaraEditor/Private/Toolkits/NiagaraScriptToolkit.h b/Engine/Plugins/FX/Niagara/Source/NiagaraEditor/Private/Toolkits/NiagaraScriptToolkit.h
--- a/Engine/Plugins/FX/Niagara/Source/NiagaraEditor/Private/Toolkits/NiagaraScriptToolkit.h
+++ b/Engine/Plugins/FX/Niagara/Source/NiagaraEditor/Private/Toolkits/NiagaraScriptToolkit.h
@@ -89,7 +89,6 @@ private:
 
 	/** Spawns the tab with the script details inside. */
 	TSharedRef<SDockTab> SpawnTabScriptParameters(const FSpawnTabArgs& Args);
-	TSharedRef<SDockTab> SpawnTabScriptParameters2(const FSpawnTabArgs& Args);
 	TSharedRef<SDockTab> SpawnTabStats(const FSpawnTabArgs& Args);
 
 	TSharedPtr< SNiagaraSelectedObjectsDetails> SelectedDetailsWidget;
@@ -150,7 +149,6 @@ private:
 	static const FName ScriptDetailsTabId;
 	static const FName SelectedDetailsTabId;
 	static const FName ParametersTabId;
-	static const FName ParametersTabId2;
 	static const FName StatsTabId;
 	static const FName MessageLogTabID;
 
diff --git a/Engine/Plugins/FX/Niagara/Source/NiagaraEditor/Private/Toolkits/NiagaraSystemToolkit.cpp b/Engine/Plugins/FX/Niagara/Source/NiagaraEditor/Private/Toolkits/NiagaraSystemToolkit.cpp
--- a/Engine/Plugins/FX/Niagara/Source/NiagaraEditor/Private/Toolkits/NiagaraSystemToolkit.cpp
+++ b/Engine/Plugins/FX/Niagara/Source/NiagaraEditor/Private/Toolkits/NiagaraSystemToolkit.cpp
@@ -68,7 +68,6 @@ const FName FNiagaraSystemToolkit::SequencerTabID(TEXT("NiagaraSystemEditor_Sequ
 const FName FNiagaraSystemToolkit::SystemScriptTabID(TEXT("NiagaraSystemEditor_SystemScript"));
 const FName FNiagaraSystemToolkit::SystemDetailsTabID(TEXT("NiagaraSystemEditor_SystemDetails"));
 const FName FNiagaraSystemToolkit::SystemParametersTabID(TEXT("NiagaraSystemEditor_SystemParameters"));
-const FName FNiagaraSystemToolkit::SystemParametersTabID2(TEXT("NiagaraSystemEditor_SystemParameters2"));
 const FName FNiagaraSystemToolkit::SelectedEmitterStackTabID(TEXT("NiagaraSystemEditor_SelectedEmitterStack"));
 const FName FNiagaraSystemToolkit::SelectedEmitterGraphTabID(TEXT("NiagaraSystemEditor_SelectedEmitterGraph"));
 const FName FNiagaraSystemToolkit::DebugSpreadsheetTabID(TEXT("NiagaraSystemEditor_DebugAttributeSpreadsheet"));
@@ -125,10 +124,6 @@ void FNiagaraSystemToolkit::RegisterTabSpawners(const TSharedRef<class FTabManag
 		.SetDisplayName(LOCTEXT("SystemParameters", "Parameters"))
 		.SetGroup(WorkspaceMenuCategory.ToSharedRef());
 
-	InTabManager->RegisterTabSpawner(SystemParametersTabID2, FOnSpawnTab::CreateSP(this, &FNiagaraSystemToolkit::SpawnTab_SystemParameters2))
-		.SetDisplayName(LOCTEXT("SystemParameters2", "Parameters2"))
-		.SetGroup(WorkspaceMenuCategory.ToSharedRef());
-
 	InTabManager->RegisterTabSpawner(SelectedEmitterStackTabID, FOnSpawnTab::CreateSP(this, &FNiagaraSystemToolkit::SpawnTab_SelectedEmitterStack))
 		.SetDisplayName(LOCTEXT("SelectedEmitterStacks", "Selected Emitters"))
 		.SetGroup(WorkspaceMenuCategory.ToSharedRef());
@@ -329,7 +324,7 @@ void FNiagaraSystemToolkit::InitializeInternal(const EToolkitMode::Type Mode, co
 	const float InTime = -0.02f;
 	const float OutTime = 3.2f;
 
-	TSharedRef<FTabManager::FLayout> StandaloneDefaultLayout = FTabManager::NewLayout("Standalone_Niagara_System_Layout_v22")
+	TSharedRef<FTabManager::FLayout> StandaloneDefaultLayout = FTabManager::NewLayout("Standalone_Niagara_System_Layout_v23")
 		->AddArea
 		(
 			FTabManager::NewPrimaryArea()->SetOrientation(Orient_Vertical)
@@ -526,24 +521,6 @@ TSharedRef<SDockTab> FNiagaraSystemToolkit::SpawnTab_SystemParameters(const FSpa
 	check(Args.GetTabId().TabType == SystemParametersTabID);
 
 
-	TArray<TSharedRef<FNiagaraObjectSelection>> ObjectSelections;
-	ObjectSelections.Add(ObjectSelectionForParameterMapView.ToSharedRef());
-
-	TSharedRef<SDockTab> SpawnedTab =
-		SNew(SDockTab)
-		[
-			SAssignNew(ParameterMapView, SNiagaraParameterMapView, ObjectSelections, SNiagaraParameterMapView::EToolkitType::SYSTEM, GetToolkitCommands()) //@todo(ng) cleanup
-		];
-	RefreshParameters();
-
-	return SpawnedTab;
-}
-
-TSharedRef<SDockTab> FNiagaraSystemToolkit::SpawnTab_SystemParameters2(const FSpawnTabArgs& Args)
-{
-	check(Args.GetTabId().TabType == SystemParametersTabID2);
-
-
 	TArray<TSharedRef<FNiagaraObjectSelection>> ObjectSelections;
 	ObjectSelections.Add(ObjectSelectionForParameterMapView.ToSharedRef());
 
diff --git a/Engine/Plugins/FX/Niagara/Source/NiagaraEditor/Private/Toolkits/NiagaraSystemToolkit.h b/Engine/Plugins/FX/Niagara/Source/NiagaraEditor/Private/Toolkits/NiagaraSystemToolkit.h
--- a/Engine/Plugins/FX/Niagara/Source/NiagaraEditor/Private/Toolkits/NiagaraSystemToolkit.h
+++ b/Engine/Plugins/FX/Niagara/Source/NiagaraEditor/Private/Toolkits/NiagaraSystemToolkit.h
@@ -101,7 +101,6 @@ private:
 	TSharedRef<SDockTab> SpawnTab_Sequencer(const FSpawnTabArgs& Args);
 	TSharedRef<SDockTab> SpawnTab_SystemScript(const FSpawnTabArgs& Args);
 	TSharedRef<SDockTab> SpawnTab_SystemParameters(const FSpawnTabArgs& Args);
-	TSharedRef<SDockTab> SpawnTab_SystemParameters2(const FSpawnTabArgs& Args);
 	TSharedRef<SDockTab> SpawnTab_SelectedEmitterStack(const FSpawnTabArgs& Args);
 	TSharedRef<SDockTab> SpawnTab_SelectedEmitterGraph(const FSpawnTabArgs& Args);
 	TSharedRef<SDockTab> SpawnTab_DebugSpreadsheet(const FSpawnTabArgs& Args);
@@ -173,8 +172,6 @@ private:
 	/** The command list for this editor */
 	TSharedPtr<FUICommandList> EditorCommands;
 
-	TSharedPtr<class SNiagaraParameterMapView> ParameterMapView; //@todo(ng) cleanup
-
 	TSharedPtr<FNiagaraSystemToolkitParameterPanelViewModel> ParameterPanelViewModel;
 	TSharedPtr<class SNiagaraParameterPanel> ParameterPanel;
 
@@ -191,7 +188,6 @@ public:
 	static const FName SystemScriptTabID;
 	static const FName SystemDetailsTabID;
 	static const FName SystemParametersTabID;
-	static const FName SystemParametersTabID2;
 	static const FName SelectedEmitterStackTabID;
 	static const FName SelectedEmitterGraphTabID;
 	static const FName DebugSpreadsheetTabID;
