Niagara - Prevent a crash when fixing up deprecated modules.

diff --git a/Engine/Plugins/FX/Niagara/Source/Niagara/Private/NiagaraParameterStore.cpp b/Engine/Plugins/FX/Niagara/Source/Niagara/Private/NiagaraParameterStore.cpp
--- a/Engine/Plugins/FX/Niagara/Source/Niagara/Private/NiagaraParameterStore.cpp
+++ b/Engine/Plugins/FX/Niagara/Source/Niagara/Private/NiagaraParameterStore.cpp
@@ -604,6 +604,13 @@ void FNiagaraParameterStore::RenameParameter(const FNiagaraVariable& Param, FNam
 	check(!ParameterOffsets.Num()); // Migration to SortedParameterOffsets
 #endif
 
+	if (Param.GetName() == NewName)
+	{
+		// Early out here to prevent crashes later on due to delta size mismatches when the newly named
+		// parameter isn't added.
+		return;
+	}
+
 	int32 Idx = IndexOf(Param);
 	if(Idx != INDEX_NONE)
 	{
diff --git a/Engine/Plugins/FX/Niagara/Source/NiagaraEditor/Private/ViewModels/Stack/NiagaraStackModuleItem.cpp b/Engine/Plugins/FX/Niagara/Source/NiagaraEditor/Private/ViewModels/Stack/NiagaraStackModuleItem.cpp
--- a/Engine/Plugins/FX/Niagara/Source/NiagaraEditor/Private/ViewModels/Stack/NiagaraStackModuleItem.cpp
+++ b/Engine/Plugins/FX/Niagara/Source/NiagaraEditor/Private/ViewModels/Stack/NiagaraStackModuleItem.cpp
@@ -1001,15 +1001,16 @@ void UNiagaraStackModuleItem::ReassignModuleScript(UNiagaraScript* ModuleScript)
 		RefreshChildren();
 
 		FunctionCallNode->SuggestName(FString());
-
 		const FString NewName = FunctionCallNode->GetFunctionName();
-		UNiagaraSystem& System = GetSystemViewModel()->GetSystem();
-		UNiagaraEmitter* Emitter = GetEmitterViewModel().IsValid() ? GetEmitterViewModel()->GetEmitter() : nullptr;
-		FNiagaraStackGraphUtilities::RenameReferencingParameters(System, Emitter, *FunctionCallNode, OldName, NewName);
-		FunctionCallNode->RefreshFromExternalChanges();
-		FunctionCallNode->MarkNodeRequiresSynchronization(TEXT("Module script reassigned."), true);
-		RefreshChildren();
-		
+		if (NewName != OldName)
+		{
+			UNiagaraSystem& System = GetSystemViewModel()->GetSystem();
+			UNiagaraEmitter* Emitter = GetEmitterViewModel().IsValid() ? GetEmitterViewModel()->GetEmitter() : nullptr;
+			FNiagaraStackGraphUtilities::RenameReferencingParameters(System, Emitter, *FunctionCallNode, OldName, NewName);
+			FunctionCallNode->RefreshFromExternalChanges();
+			FunctionCallNode->MarkNodeRequiresSynchronization(TEXT("Module script reassigned."), true);
+			RefreshChildren();
+		}
 		
 		if (ModuleScript->ConversionUtility != nullptr && OldClipboardContent != nullptr)
 		{
