Fixed merge error.

diff --git a/Engine/Plugins/FX/Niagara/Config/BaseNiagara.ini b/Engine/Plugins/FX/Niagara/Config/BaseNiagara.ini
--- a/Engine/Plugins/FX/Niagara/Config/BaseNiagara.ini
+++ b/Engine/Plugins/FX/Niagara/Config/BaseNiagara.ini
@@ -21,9 +21,12 @@
 +PropertyRedirects=(OldName="NiagaraDeviceProfileStateEntry.SetEffectsQualityMask",NewName="SetQualityLevelMask")
 +PropertyRedirects=(OldName="FNiagaraPlatformSetConflictEntry.EffectsQualityMask",NewName="QualityLevelMask")
 
++PropertyRedirects=(OldName="NiagaraDataInterfaceSkeletalMesh.DefaultMesh", NewName="NiagaraDataInterfaceSkeletalMesh.PreviewMesh")
+
 ;+EnumRedirects=(OldName="ENiagaraSortMode",ValueChanges=(("ENiagaraSortMode::SortNone","ENiagaraSortMode::None"), ("ENiagaraSortMode::SortViewDepth","ENiagaraSortMode::ViewDepth"), ("ENiagaraSortMode::SortViewDistance","ENiagaraSortMode::ViewDistance")) )
 
-+PropertyRedirects=(OldName="NiagaraDataInterfaceSkeletalMesh.DefaultMesh", NewName="NiagaraDataInterfaceSkeletalMesh.PreviewMesh")
++PropertyRedirects=(OldName="NiagaraDataInterfaceSkeletalMesh.SpecificBones", NewName="NiagaraDataInterfaceSkeletalMesh.FilteredBones")
++PropertyRedirects=(OldName="NiagaraDataInterfaceSkeletalMesh.SpecificSockets", NewName="NiagaraDataInterfaceSkeletalMesh.FilteredSockets")
 
 [/Script/Niagara.NiagaraSettings]
 +AdditionalParameterEnums=/Niagara/Enums/ENiagaraCoordinateSpace.ENiagaraCoordinateSpace
