Niagara StaticMesh interface changes:
- Interface no longer fails when no mesh or a non-CPU accessible mesh is selected (now falls back to sensible defaults)
- Add 'IsValid' function so scripts can detect invalid meshes
- Fix bug where the GPU-side WS normals were not calculated using the transform inverse transpose matrix.
- Replace usage of InverseFast() with Inverse() to correctly handle zero scale meshes
- Produce an error in the UI if no meshes (Source component, DefaultMesh, nor PreviewMesh) are provided.

diff --git a/Engine/Plugins/FX/Niagara/Source/Niagara/Classes/NiagaraDataInterfaceStaticMesh.h b/Engine/Plugins/FX/Niagara/Source/Niagara/Classes/NiagaraDataInterfaceStaticMesh.h
--- a/Engine/Plugins/FX/Niagara/Source/Niagara/Classes/NiagaraDataInterfaceStaticMesh.h
+++ b/Engine/Plugins/FX/Niagara/Source/Niagara/Classes/NiagaraDataInterfaceStaticMesh.h
@@ -145,7 +145,7 @@ struct FNDIStaticMesh_InstanceData
 	FORCEINLINE_DEBUGGABLE bool ResetRequired(UNiagaraDataInterfaceStaticMesh* Interface)const;
 
 	FORCEINLINE const TArray<int32>& GetValidSections()const { return ValidSections; }
-	FORCEINLINE const FStaticMeshAreaWeightedSectionSampler& GetAreaWeigtedSampler()const { return Sampler; }
+	FORCEINLINE const FStaticMeshAreaWeightedSectionSampler& GetAreaWeightedSampler() const { return Sampler; }
 
 	void InitVertexColorFiltering();
 
@@ -162,6 +162,13 @@ class NIAGARA_API UNiagaraDataInterfaceStaticMesh : public UNiagaraDataInterface
 
 public:
 
+	enum class ESampleMode : int32
+	{
+		Invalid = -1,
+		Default,
+		AreaWeighted
+	};
+
 	DECLARE_NIAGARA_DI_PARAMETER();
 	
 #if WITH_EDITORONLY_DATA
@@ -202,15 +209,15 @@ public:
 
 	//~ UNiagaraDataInterface interface
 
-	virtual bool InitPerInstanceData(void* PerInstanceData, FNiagaraSystemInstance* SystemInstance)override;
-	virtual void DestroyPerInstanceData(void* PerInstanceData, FNiagaraSystemInstance* SystemInstance)override;
+	virtual bool InitPerInstanceData(void* PerInstanceData, FNiagaraSystemInstance* SystemInstance) override;
+	virtual void DestroyPerInstanceData(void* PerInstanceData, FNiagaraSystemInstance* SystemInstance) override;
 	virtual bool PerInstanceTick(void* PerInstanceData, FNiagaraSystemInstance* SystemInstance, float DeltaSeconds) override;
-	virtual int32 PerInstanceDataSize()const override { return sizeof(FNDIStaticMesh_InstanceData); }
+	virtual int32 PerInstanceDataSize() const override { return sizeof(FNDIStaticMesh_InstanceData); }
 
-	virtual void GetFunctions(TArray<FNiagaraFunctionSignature>& OutFunctions)override;
+	virtual void GetFunctions(TArray<FNiagaraFunctionSignature>& OutFunctions) override;
 	virtual void GetVMExternalFunction(const FVMExternalFunctionBindingInfo& BindingInfo, void* InstanceData, FVMExternalFunction &OutFunc) override;
 	virtual bool Equals(const UNiagaraDataInterface* Other) const override;
-	virtual bool CanExecuteOnTarget(ENiagaraSimTarget Target)const override { return true; }
+	virtual bool CanExecuteOnTarget(ENiagaraSimTarget Target) const override { return true; }
 #if WITH_EDITOR
 	virtual TArray<FNiagaraDataInterfaceError> GetErrors() override;
 #endif
@@ -238,15 +245,15 @@ public:
 	static const FString UseColorBufferName;
 
 public:
-	void GetNumTriangles(FVectorVMContext& Context);
+	void IsValid(FVectorVMContext& Context);
 
-	template<typename TAreaWeighted>
+	template<typename TSampleMode>
 	void RandomSection(FVectorVMContext& Context);
 
-	template<typename TAreaWeighted>
+	template<typename TSampleMode>
 	void RandomTriCoord(FVectorVMContext& Context);
 
-	template<typename TAreaWeighted>
+	template<typename TSampleMode>
 	void RandomTriCoordOnSection(FVectorVMContext& Context);
 
  	void RandomTriCoordVertexColorFiltered(FVectorVMContext& Context);
@@ -287,13 +294,13 @@ protected:
 	
 private:
 	
-	template<typename TAreaWeighted, bool bFiltered>
+	template<typename TSampleMode, bool bFiltered>
 	FORCEINLINE_DEBUGGABLE int32 RandomSection(FRandomStream& RandStream, FStaticMeshLODResources& Res, FNDIStaticMesh_InstanceData* InstData);
 
-	template<typename TAreaWeighted, bool bFiltered>
+	template<typename TSampleMode, bool bFiltered>
 	FORCEINLINE_DEBUGGABLE int32 RandomTriIndex(FRandomStream& RandStream, FStaticMeshLODResources& Res, FNDIStaticMesh_InstanceData* InstData);
 
-	template<typename TAreaWeighted>
+	template<typename TSampleMode>
 	FORCEINLINE_DEBUGGABLE int32 RandomTriIndexOnSection(FRandomStream& RandStream, FStaticMeshLODResources& Res, int32 SectionIdx, FNDIStaticMesh_InstanceData* InstData);
 
 	void WriteTransform(const FMatrix& ToWrite, FVectorVMContext& Context);
diff --git a/Engine/Plugins/FX/Niagara/Source/Niagara/Private/NiagaraDataInterfaceStaticMesh.cpp b/Engine/Plugins/FX/Niagara/Source/Niagara/Private/NiagaraDataInterfaceStaticMesh.cpp
--- a/Engine/Plugins/FX/Niagara/Source/Niagara/Private/NiagaraDataInterfaceStaticMesh.cpp
+++ b/Engine/Plugins/FX/Niagara/Source/Niagara/Private/NiagaraDataInterfaceStaticMesh.cpp
@@ -13,6 +13,11 @@
 
 #define LOCTEXT_NAMESPACE "NiagaraDataInterfaceStaticMesh"
 
+// These are to help readability in template specializations
+using TSampleModeInvalid = TIntegralConstant<UNiagaraDataInterfaceStaticMesh::ESampleMode, UNiagaraDataInterfaceStaticMesh::ESampleMode::Invalid>;
+using TSampleModeDefault = TIntegralConstant<UNiagaraDataInterfaceStaticMesh::ESampleMode, UNiagaraDataInterfaceStaticMesh::ESampleMode::Default>;
+using TSampleModeAreaWeighted = TIntegralConstant<UNiagaraDataInterfaceStaticMesh::ESampleMode, UNiagaraDataInterfaceStaticMesh::ESampleMode::AreaWeighted>;
+
 const FString UNiagaraDataInterfaceStaticMesh::MeshIndexBufferName(TEXT("IndexBuffer_"));
 const FString UNiagaraDataInterfaceStaticMesh::MeshVertexBufferName(TEXT("VertexBuffer_"));
 const FString UNiagaraDataInterfaceStaticMesh::MeshTangentBufferName(TEXT("TangentBuffer_"));
@@ -233,66 +238,69 @@ bool FNDIStaticMesh_InstanceData::Init(UNiagaraDataInterfaceStaticMesh* Interfac
 	}
 #endif
 
-	if (Component.IsValid() && Mesh)
+	if (!Component.IsValid())
 	{
-		PrevTransform = Transform;
-		PrevTransformInverseTransposed = TransformInverseTransposed;
-		Transform = Component->GetComponentToWorld().ToMatrixWithScale();
-		TransformInverseTransposed = Transform.InverseFast().GetTransposed();
+		UE_LOG(LogNiagara, Log, TEXT("StaticMesh data interface has no valid component - %s"), *Interface->GetFullName());
+		return false;
 	}
 
+	PrevTransform = Transform;
+	PrevTransformInverseTransposed = TransformInverseTransposed;
+	Transform = Component->GetComponentToWorld().ToMatrixWithScale();
+	TransformInverseTransposed = Transform.Inverse().GetTransposed();
+
+	// Report missing or inaccessible meshes to the log
 	if (!Mesh)
 	{
-		UE_LOG(LogNiagara, Log, TEXT("StaticMesh data interface has no valid mesh. Failed InitPerInstanceData - %s"), *Interface->GetFullName());
-		return false;
+		UE_LOG(LogNiagara, Log, TEXT("StaticMesh data interface has no valid mesh - %s"), *Interface->GetFullName());
 	}
-
-	if (!Mesh->bAllowCPUAccess)
+	else if (!Mesh->bAllowCPUAccess)
 	{
-		UE_LOG(LogNiagara, Log, TEXT("StaticMesh data interface using a mesh that does not allow CPU access. Failed InitPerInstanceData - Mesh: %s"), *Mesh->GetFullName());
-		return false;
+		UE_LOG(LogNiagara, Log, TEXT("StaticMesh data interface using a mesh that does not allow CPU access. Interface: %s, Mesh: %s"),
+			*Interface->GetFullName(), *Mesh->GetFullName());
+		Mesh = nullptr; // Disallow usage of this mesh to prevent issues on cooked builds
 	}
 
-	if (!Component.IsValid())
+#if WITH_EDITOR
+	if (Mesh)
 	{
-		UE_LOG(LogNiagara, Log, TEXT("StaticMesh data interface has no valid component. Failed InitPerInstanceData - %s"), *Interface->GetFullName());
-		return false;
+		Mesh->GetOnMeshChanged().AddUObject(SystemInstance->GetComponent(), &UNiagaraComponent::ReinitializeSystem);
 	}
-
-#if WITH_EDITOR
-	Mesh->GetOnMeshChanged().AddUObject(SystemInstance->GetComponent(), &UNiagaraComponent::ReinitializeSystem);
 #endif
 
-	bMeshAllowsCpuAccess       = Mesh->bAllowCPUAccess;
-	bIsCpuUniformlyDistributedSampling = Mesh->bSupportUniformlyDistributedSampling;
-	bIsGpuUniformlyDistributedSampling = bIsCpuUniformlyDistributedSampling && Mesh->bSupportGpuUniformlyDistributedSampling;
-
-	//Init the instance filter
+	bMeshAllowsCpuAccess = false;
+	bIsCpuUniformlyDistributedSampling = false;
+	bIsGpuUniformlyDistributedSampling = false;
 	ValidSections.Empty();
-	FStaticMeshLODResources& Res = Mesh->RenderData->LODResources[0];
-	for (int32 i = 0; i < Res.Sections.Num(); ++i)
+	if (Mesh)
 	{
-		if (Interface->SectionFilter.AllowedMaterialSlots.Num() == 0 || Interface->SectionFilter.AllowedMaterialSlots.Contains(Res.Sections[i].MaterialIndex))
+		bMeshAllowsCpuAccess = Mesh->bAllowCPUAccess;
+		bIsCpuUniformlyDistributedSampling = Mesh->bSupportUniformlyDistributedSampling;
+		bIsGpuUniformlyDistributedSampling = bIsCpuUniformlyDistributedSampling && Mesh->bSupportGpuUniformlyDistributedSampling;
+
+		//Init the instance filter
+		FStaticMeshLODResources& Res = Mesh->RenderData->LODResources[0];
+		for (int32 i = 0; i < Res.Sections.Num(); ++i)
 		{
-			ValidSections.Add(i);
+			if (Interface->SectionFilter.AllowedMaterialSlots.Num() == 0 || Interface->SectionFilter.AllowedMaterialSlots.Contains(Res.Sections[i].MaterialIndex))
+			{
+				ValidSections.Add(i);
+			}
 		}
-	}
 
-	if (GetValidSections().Num() == 0)
-	{
-		UE_LOG(LogNiagara, Log, TEXT("StaticMesh data interface has a section filter preventing any spawning. Failed InitPerInstanceData - %s"), *Interface->GetFullName());
-		return false;
-	}
+		if (GetValidSections().Num() == 0)
+		{
+			UE_LOG(LogNiagara, Log, TEXT("StaticMesh data interface has a section filter preventing any spawning. Failed InitPerInstanceData - %s"), *Interface->GetFullName());
+		}
 
-	Sampler.Init(&Res, this);
+		Sampler.Init(&Res, this);
+	}
 
 	return true;
 }
 
-bool FNDIStaticMesh_InstanceData::ResetRequired(UNiagaraDataInterfaceStaticMesh* Interface)const
+bool FNDIStaticMesh_InstanceData::ResetRequired(UNiagaraDataInterfaceStaticMesh* Interface) const
 {
-	check(GetActualMesh());
-
 	if (!Component.IsValid())
 	{
 		//The component we were bound to is no longer valid so we have to trigger a reset.
@@ -305,7 +313,6 @@ bool FNDIStaticMesh_InstanceData::ResetRequired(UNiagaraDataInterfaceStaticMesh*
 	}
 
 	//bool bPrevVCSampling = bSupportingVertexColorSampling;//TODO: Vertex color filtering needs more work.
-	bool bReset = false;
 	if (Mesh)
 	{
 		const bool bNewMeshAllowsCpuAccess = Mesh->bAllowCPUAccess;
@@ -313,9 +320,15 @@ bool FNDIStaticMesh_InstanceData::ResetRequired(UNiagaraDataInterfaceStaticMesh*
 		const bool bNewIsGpuAreaWeightedSampling = bIsCpuUniformlyDistributedSampling && Mesh->bSupportGpuUniformlyDistributedSampling;
 
 		//bSupportingVertexColorSampling = bEnableVertexColorRangeSorting && MeshHasColors();
-		bReset = bNewMeshAllowsCpuAccess != bMeshAllowsCpuAccess || bNewIsCpuAreaWeightedSampling != bIsCpuUniformlyDistributedSampling || bNewIsGpuAreaWeightedSampling != bIsGpuUniformlyDistributedSampling /* || bSupportingVertexColorSampling != bPrevVCSampling*/;
+		return bNewMeshAllowsCpuAccess != bMeshAllowsCpuAccess || bNewIsCpuAreaWeightedSampling != bIsCpuUniformlyDistributedSampling || bNewIsGpuAreaWeightedSampling != bIsGpuUniformlyDistributedSampling /* || bSupportingVertexColorSampling != bPrevVCSampling*/;
 	}
-	return bReset;
+	else if (bMeshAllowsCpuAccess || bIsCpuUniformlyDistributedSampling || bIsGpuUniformlyDistributedSampling)
+	{
+		// We previously had a CPU accessible mesh, but now have none
+		return true;
+	}
+	
+	return false;
 }
 
 bool FNDIStaticMesh_InstanceData::Tick(UNiagaraDataInterfaceStaticMesh* Interface, FNiagaraSystemInstance* SystemInstance, float InDeltaSeconds)
@@ -327,12 +340,12 @@ bool FNDIStaticMesh_InstanceData::Tick(UNiagaraDataInterfaceStaticMesh* Interfac
 	else
 	{
 		DeltaSeconds = InDeltaSeconds;
-		if (Component.IsValid() && Mesh)
+		if (Component.IsValid())
 		{
 			PrevTransform = Transform;
 			PrevTransformInverseTransposed = TransformInverseTransposed;
 			Transform = Component->GetComponentToWorld().ToMatrixWithScale();
-			TransformInverseTransposed = Transform.InverseFast().GetTransposed();
+			TransformInverseTransposed = Transform.Inverse().GetTransposed();
 		}
 		else
 		{
@@ -489,10 +502,32 @@ public:
 			FNiagaraDataInterfaceProxyStaticMesh* InterfaceProxy = static_cast<FNiagaraDataInterfaceProxyStaticMesh*>(Context.DataInterface);
 			FNiagaraStaticMeshData* Data = InterfaceProxy->SystemInstancesToMeshData.Find(Context.SystemInstance);
 			ensureMsgf(Data, TEXT("Failed to find data for instance %s"), *FNiagaraUtilities::SystemInstanceIDToString(Context.SystemInstance));
-			if (Data != nullptr)
+
+			if (Data)
+			{
+				const float InvDeltaTime = Data->DeltaSeconds > 0.0f ? 1.0f / Data->DeltaSeconds : 0.0f;
+				const FVector DeltaPosition = Data->Transform.GetOrigin() - Data->PrevTransform.GetOrigin();
+
+				SetShaderValue(RHICmdList, ComputeShaderRHI, InstanceTransform, Data->Transform);
+				SetShaderValue(RHICmdList, ComputeShaderRHI, InstanceTransformInverseTransposed, Data->Transform.Inverse().GetTransposed());
+				SetShaderValue(RHICmdList, ComputeShaderRHI, InstancePrevTransform, Data->PrevTransform);
+				SetShaderValue(RHICmdList, ComputeShaderRHI, InstanceInvDeltaTime, InvDeltaTime);
+				SetShaderValue(RHICmdList, ComputeShaderRHI, InstanceWorldVelocity, DeltaPosition * InvDeltaTime);
+				SetShaderValue(RHICmdList, ComputeShaderRHI, AreaWeightedSampling, Data->bIsGpuUniformlyDistributedSampling ? 1 : 0);
+			}
+			else
 			{
-				FStaticMeshGpuSpawnBuffer* SpawnBuffer = Data->MeshGpuSpawnBuffer;
+				SetShaderValue(RHICmdList, ComputeShaderRHI, InstanceTransform, FMatrix::Identity);
+				SetShaderValue(RHICmdList, ComputeShaderRHI, InstanceTransformInverseTransposed, FMatrix::Identity);
+				SetShaderValue(RHICmdList, ComputeShaderRHI, InstancePrevTransform, FMatrix::Identity);
+				SetShaderValue(RHICmdList, ComputeShaderRHI, InstanceInvDeltaTime, 1.0f);
+				SetShaderValue(RHICmdList, ComputeShaderRHI, InstanceWorldVelocity, FVector::ZeroVector);
+				SetShaderValue(RHICmdList, ComputeShaderRHI, AreaWeightedSampling, 0);
+			}
 
+			FStaticMeshGpuSpawnBuffer* SpawnBuffer = Data ? Data->MeshGpuSpawnBuffer : nullptr;
+			if (SpawnBuffer)
+			{
 				SetSRVParameter(RHICmdList, ComputeShaderRHI, MeshVertexBuffer, SpawnBuffer->GetBufferPositionSRV());
 				SetSRVParameter(RHICmdList, ComputeShaderRHI, MeshTangentBuffer, SpawnBuffer->GetBufferTangentSRV());
 				SetSRVParameter(RHICmdList, ComputeShaderRHI, MeshIndexBuffer, SpawnBuffer->GetBufferIndexSRV());
@@ -518,7 +553,6 @@ public:
 					SetShaderValue(RHICmdList, ComputeShaderRHI, UseColorBuffer, 0);
 				}
 
-				check(SpawnBuffer); // should always be allocated, we always need the GPU buffer for a GpuSimulation.
 				SetShaderValue(RHICmdList, ComputeShaderRHI, SectionCount, SpawnBuffer->GetValidSectionCount());
 				SetSRVParameter(RHICmdList, ComputeShaderRHI, MeshSectionBuffer, SpawnBuffer->GetBufferSectionSRV());
 				if (Data->bIsGpuUniformlyDistributedSampling)
@@ -529,16 +563,6 @@ public:
 				{
 					SetSRVParameter(RHICmdList, ComputeShaderRHI, MeshTriangleBuffer, FNiagaraRenderer::GetDummyUIntBuffer());
 				}
-
-				const float InvDeltaTime = Data->DeltaSeconds > 0.0f ? 1.0f / Data->DeltaSeconds : 0.0f;
-				const FVector DeltaPosition = Data->Transform.GetOrigin() - Data->PrevTransform.GetOrigin();
-
-				SetShaderValue(RHICmdList, ComputeShaderRHI, InstanceTransform, Data->Transform);
-				SetShaderValue(RHICmdList, ComputeShaderRHI, InstanceTransformInverseTransposed, Data->Transform.InverseFast().GetTransposed());
-				SetShaderValue(RHICmdList, ComputeShaderRHI, InstancePrevTransform, Data->PrevTransform);
-				SetShaderValue(RHICmdList, ComputeShaderRHI, InstanceInvDeltaTime, InvDeltaTime);
-				SetShaderValue(RHICmdList, ComputeShaderRHI, InstanceWorldVelocity, DeltaPosition * InvDeltaTime);
-				SetShaderValue(RHICmdList, ComputeShaderRHI, AreaWeightedSampling, Data->bIsGpuUniformlyDistributedSampling ? 1 : 0);
 			}
 			else
 			{
@@ -553,14 +577,7 @@ public:
 
 				SetShaderValue(RHICmdList, ComputeShaderRHI, SectionCount, 0);
 				SetSRVParameter(RHICmdList, ComputeShaderRHI, MeshSectionBuffer, FNiagaraRenderer::GetDummyUInt4Buffer());
-				SetSRVParameter(RHICmdList, ComputeShaderRHI, MeshTriangleBuffer, FNiagaraRenderer::GetDummyUInt4Buffer());
-
-				SetShaderValue(RHICmdList, ComputeShaderRHI, InstanceTransform, FMatrix::Identity);
-				SetShaderValue(RHICmdList, ComputeShaderRHI, InstanceTransformInverseTransposed, FMatrix::Identity);
-				SetShaderValue(RHICmdList, ComputeShaderRHI, InstancePrevTransform, FMatrix::Identity);
-				SetShaderValue(RHICmdList, ComputeShaderRHI, InstanceInvDeltaTime, 1.0f);
-				SetShaderValue(RHICmdList, ComputeShaderRHI, InstanceWorldVelocity, FVector::ZeroVector);
-				SetShaderValue(RHICmdList, ComputeShaderRHI, AreaWeightedSampling, 0);
+				SetSRVParameter(RHICmdList, ComputeShaderRHI, MeshTriangleBuffer, FNiagaraRenderer::GetDummyUInt4Buffer());				
 			}
 		}
 	}
@@ -616,7 +633,6 @@ void FNiagaraDataInterfaceProxyStaticMesh::ConsumePerInstanceDataFromGameThread(
 	if (Data)
 	{
 		//UE_LOG(LogNiagara, Log, TEXT("ConsumePerInstanceDataFromGameThread() ... found %s"), *Instance.ToString());
-		check(Data->MeshGpuSpawnBuffer);
 
 		Data->bIsGpuUniformlyDistributedSampling = SourceData->bIsGpuUniformlyDistributedSampling;
 		Data->DeltaSeconds = SourceData->DeltaSeconds;
@@ -665,11 +681,11 @@ void UNiagaraDataInterfaceStaticMesh::PostEditChangeProperty(FPropertyChangedEve
 	Super::PostEditChangeProperty(PropertyChangedEvent);
 	ChangeId++;
 }
-
 #endif //WITH_EDITOR
 
 namespace StaticMeshHelpers
 {
+	static const FName IsValidName("IsValid");
 	static const FName RandomSectionName("RandomSection");
 	static const FName RandomTriCoordName("RandomTriCoord");
 	static const FName RandomTriCoordOnSectionName("RandomTriCoordOnSection");
@@ -699,6 +715,16 @@ namespace StaticMeshHelpers
 
 void UNiagaraDataInterfaceStaticMesh::GetFunctions(TArray<FNiagaraFunctionSignature>& OutFunctions)
 {
+	{
+		FNiagaraFunctionSignature Sig;
+		Sig.Name = StaticMeshHelpers::IsValidName;
+		Sig.Inputs.Add(FNiagaraVariable(FNiagaraTypeDefinition(GetClass()), TEXT("StaticMesh")));
+		Sig.Outputs.Add(FNiagaraVariable(FNiagaraTypeDefinition::GetBoolDef(), TEXT("Valid")));
+		Sig.bMemberFunction = true;
+		Sig.bRequiresContext = false;
+		OutFunctions.Add(Sig);
+	}
+
 	{
 		FNiagaraFunctionSignature Sig;
 		Sig.Name = StaticMeshHelpers::RandomSectionName;
@@ -911,24 +937,39 @@ void UNiagaraDataInterfaceStaticMesh::GetFunctions(TArray<FNiagaraFunctionSignat
 
 //External function binder choosing between template specializations based on UsesAreaWeighting
 template<typename NextBinder>
-struct TUsesAreaWeightingBinder
+struct TSampleModeBinder
 {
 	template<typename... ParamTypes>
 	static void Bind(UNiagaraDataInterface* Interface, const FVMExternalFunctionBindingInfo& BindingInfo, void* InstanceData, FVMExternalFunction &OutFunc)
 	{
 		FNDIStaticMesh_InstanceData* InstData = (FNDIStaticMesh_InstanceData*)InstanceData;
 		UNiagaraDataInterfaceStaticMesh* MeshInterface = CastChecked<UNiagaraDataInterfaceStaticMesh>(Interface);
-		if (InstData->UsesCpuUniformlyDistributedSampling())
+		if (InstData->Mesh == nullptr)
+		{
+			NextBinder::template Bind<ParamTypes..., TSampleModeInvalid>(Interface, BindingInfo, InstanceData, OutFunc);
+		}
+		else if (InstData->UsesCpuUniformlyDistributedSampling())
 		{
-			NextBinder::template Bind<ParamTypes..., TIntegralConstant<bool, true>>(Interface, BindingInfo, InstanceData, OutFunc);
+			NextBinder::template Bind<ParamTypes..., TSampleModeAreaWeighted>(Interface, BindingInfo, InstanceData, OutFunc);
 		}
 		else
 		{
-			NextBinder::template Bind<ParamTypes..., TIntegralConstant<bool, false>>(Interface, BindingInfo, InstanceData, OutFunc);
+			NextBinder::template Bind<ParamTypes..., TSampleModeDefault>(Interface, BindingInfo, InstanceData, OutFunc);
 		}
 	}
 };
 
+//Helper struct for stubbing access of vertex data.
+struct TNullMeshVertexAccessor
+{
+	TNullMeshVertexAccessor(const FStaticMeshVertexBuffer&) {}
+
+	FORCEINLINE FVector GetTangentX(int32 Idx)const { return FVector4(1.0f, 0.0f, 0.0f, 0.0f); }
+	FORCEINLINE FVector GetTangentY(int32 Idx)const { return FVector4(0.0f, 1.0f, 0.0f, 0.0f); }
+	FORCEINLINE FVector GetTangentZ(int32 Idx)const { return FVector4(0.0f, 0.0f, 1.0f, 0.0f); }
+	FORCEINLINE FVector2D GetUV(int32 Idx, int32 UVSet)const { return FVector2D(0.0f, 0.0f); }
+};
+
 //Helper struct for accessing typed vertex data.
 template<EStaticMeshVertexTangentBasisType TangentT, EStaticMeshVertexUVType UVTypeT>
 struct TTypedMeshVertexAccessor
@@ -952,8 +993,13 @@ struct TTypedMeshAccessorBinder
 	static void Bind(UNiagaraDataInterface* Interface, const FVMExternalFunctionBindingInfo& BindingInfo, void* InstanceData, FVMExternalFunction &OutFunc)
 	{
 		FNDIStaticMesh_InstanceData* InstData = (FNDIStaticMesh_InstanceData*)InstanceData;
+		if (!InstData->Mesh)
+		{
+			NextBinder::template Bind<ParamTypes..., TNullMeshVertexAccessor>(Interface, BindingInfo, InstanceData, OutFunc);
+			return;
+		}
+
 		UNiagaraDataInterfaceStaticMesh* MeshInterface = CastChecked<UNiagaraDataInterfaceStaticMesh>(Interface);
-		check(InstData->Mesh);
 		FStaticMeshLODResources& Res = InstData->Mesh->RenderData->LODResources[0];
 		if (Res.VertexBuffers.StaticMeshVertexBuffer.GetUseHighPrecisionTangentBasis())			
 		{
@@ -981,6 +1027,7 @@ struct TTypedMeshAccessorBinder
 };
 
 //Final binders for all static mesh interface functions.
+DEFINE_NDI_DIRECT_FUNC_BINDER(UNiagaraDataInterfaceStaticMesh, IsValid);
 DEFINE_NDI_FUNC_BINDER(UNiagaraDataInterfaceStaticMesh, RandomSection);
 DEFINE_NDI_FUNC_BINDER(UNiagaraDataInterfaceStaticMesh, RandomTriCoord);
 DEFINE_NDI_DIRECT_FUNC_BINDER(UNiagaraDataInterfaceStaticMesh, RandomTriCoordVertexColorFiltered);
@@ -997,17 +1044,22 @@ DEFINE_NDI_FUNC_BINDER(UNiagaraDataInterfaceStaticMesh, GetVertexPosition);
 void UNiagaraDataInterfaceStaticMesh::GetVMExternalFunction(const FVMExternalFunctionBindingInfo& BindingInfo, void* InstanceData, FVMExternalFunction &OutFunc)
 {
 	FNDIStaticMesh_InstanceData* InstData = (FNDIStaticMesh_InstanceData*)InstanceData;
-	check(InstData && InstData->Mesh && InstData->Component.IsValid());
+	check(InstData && InstData->Component.IsValid());
 	
-	if (BindingInfo.Name == StaticMeshHelpers::RandomSectionName)
+	if (BindingInfo.Name == StaticMeshHelpers::IsValidName)
 	{
 		check(BindingInfo.GetNumInputs() == 1 && BindingInfo.GetNumOutputs() == 1);
-		TUsesAreaWeightingBinder<NDI_FUNC_BINDER(UNiagaraDataInterfaceStaticMesh, RandomSection)>::Bind(this, BindingInfo, InstanceData, OutFunc);
+		NDI_FUNC_BINDER(UNiagaraDataInterfaceStaticMesh, IsValid)::Bind(this, OutFunc);
+	}
+	else if (BindingInfo.Name == StaticMeshHelpers::RandomSectionName)
+	{
+		check(BindingInfo.GetNumInputs() == 1 && BindingInfo.GetNumOutputs() == 1);
+		TSampleModeBinder<NDI_FUNC_BINDER(UNiagaraDataInterfaceStaticMesh, RandomSection)>::Bind(this, BindingInfo, InstanceData, OutFunc);
 	}
 	else if (BindingInfo.Name == StaticMeshHelpers::RandomTriCoordName)
 	{
 		check(BindingInfo.GetNumInputs() == 1 && BindingInfo.GetNumOutputs() == 4);
-		TUsesAreaWeightingBinder<NDI_FUNC_BINDER(UNiagaraDataInterfaceStaticMesh, RandomTriCoord)>::Bind(this, BindingInfo, InstanceData, OutFunc);
+		TSampleModeBinder<NDI_FUNC_BINDER(UNiagaraDataInterfaceStaticMesh, RandomTriCoord)>::Bind(this, BindingInfo, InstanceData, OutFunc);
 	}
 	//TODO: Vertex color filtering needs more work.
 	else if (BindingInfo.Name == StaticMeshHelpers::RandomTriCoordVCFilteredName)
@@ -1019,7 +1071,7 @@ void UNiagaraDataInterfaceStaticMesh::GetVMExternalFunction(const FVMExternalFun
 	else if (BindingInfo.Name == StaticMeshHelpers::RandomTriCoordOnSectionName)
 	{
 		check(BindingInfo.GetNumInputs() == 2 && BindingInfo.GetNumOutputs() == 4);
-		TUsesAreaWeightingBinder<NDI_FUNC_BINDER(UNiagaraDataInterfaceStaticMesh, RandomTriCoordOnSection)>::Bind(this, BindingInfo, InstanceData, OutFunc);
+		TSampleModeBinder<NDI_FUNC_BINDER(UNiagaraDataInterfaceStaticMesh, RandomTriCoordOnSection)>::Bind(this, BindingInfo, InstanceData, OutFunc);
 	}
 	else if (BindingInfo.Name == StaticMeshHelpers::GetTriPositionName)
 	{
@@ -1133,36 +1185,33 @@ bool UNiagaraDataInterfaceStaticMesh::InitPerInstanceData(void* PerInstanceData,
 
 	if (bSuccess)
 	{
-		// Always allocate when bAllowCPUAccess (index buffer can only have SRV created in this case as of today)
-		// We do not know if this interface is allocated for CPU or GPU so we allocate for both case... TODO: have some cached data created in case a GPU version is needed?
-		if (Inst->Mesh->bAllowCPUAccess)
+		FStaticMeshGpuSpawnBuffer* MeshGpuSpawnBuffer = nullptr;
+		if (Inst->Mesh)
 		{
-			FStaticMeshLODResources& Res = Inst->Mesh->RenderData->LODResources[0];
+			// Always allocate when bAllowCPUAccess (index buffer can only have SRV created in this case as of today)
+			// We do not know if this interface is allocated for CPU or GPU so we allocate for both case... TODO: have some cached data created in case a GPU version is needed?
+			ensure(Inst->Mesh->bAllowCPUAccess); // this should have been verified in Init()
 
-			FStaticMeshGpuSpawnBuffer* MeshGpuSpawnBuffer = new FStaticMeshGpuSpawnBuffer;
+			MeshGpuSpawnBuffer = new FStaticMeshGpuSpawnBuffer;
+			FStaticMeshLODResources& Res = Inst->Mesh->RenderData->LODResources[0];
 			MeshGpuSpawnBuffer->Initialise(Res, *this, Inst->bIsGpuUniformlyDistributedSampling, Inst->ValidSections, Inst->Sampler);
+		}
 
-			// Push instance data to RT
-			{
-				FNiagaraDataInterfaceProxyStaticMesh* ThisProxy = GetProxyAs<FNiagaraDataInterfaceProxyStaticMesh>();
-				ENQUEUE_RENDER_COMMAND(FNiagaraDIPushInitialInstanceDataToRT) (
-					[ThisProxy, InstanceID = SystemInstance->GetId(), MeshGpuSpawnBuffer](FRHICommandListImmediate& CmdList)
+		// Push instance data to RT
+		{
+			FNiagaraDataInterfaceProxyStaticMesh* ThisProxy = GetProxyAs<FNiagaraDataInterfaceProxyStaticMesh>();
+			ENQUEUE_RENDER_COMMAND(FNiagaraDIPushInitialInstanceDataToRT) (
+				[ThisProxy, InstanceID = SystemInstance->GetId(), MeshGpuSpawnBuffer](FRHICommandListImmediate& CmdList)
+				{
+					if (MeshGpuSpawnBuffer)
 					{
 						MeshGpuSpawnBuffer->InitResource();
-						ThisProxy->InitializePerInstanceData(InstanceID, MeshGpuSpawnBuffer);
 					}
-				);
-			}
-		}
-		else
-		{
-			ensure(Inst->Mesh->bAllowCPUAccess);
+					ThisProxy->InitializePerInstanceData(InstanceID, MeshGpuSpawnBuffer);
+				}
+			);
 		}
 	}
-	else
-	{
-		UE_LOG(LogNiagara, Warning, TEXT("UNiagaraDataInterfaceStaticMesh::InitPerInstanceData> Invalid per-instance init!!!"));
-	}
 
 	return bSuccess;
 }
@@ -1218,57 +1267,83 @@ TArray<FNiagaraDataInterfaceError> UNiagaraDataInterfaceStaticMesh::GetErrors()
 		Errors.Add(CPUAccessNotAllowedError);
 	}
 
+	bool bHasNoMeshAssignedError = (Source == nullptr && DefaultMesh == nullptr);
 #if WITH_EDITORONLY_DATA
-	if (Source == nullptr && PreviewMesh != nullptr && !PreviewMesh->bAllowCPUAccess)
+	if (bHasNoMeshAssignedError && PreviewMesh != nullptr)
 	{
-		FNiagaraDataInterfaceError CPUAccessNotAllowedError(FText::Format(LOCTEXT("CPUAccessNotAllowedError", "This mesh needs CPU access in order to be used properly.({0})"), FText::FromString(PreviewMesh->GetName())),
-			LOCTEXT("CPUAccessNotAllowedErrorSummary", "CPU access error"),
-			FNiagaraDataInterfaceFix::CreateLambda([=]()
+		bHasNoMeshAssignedError = false;
+
+		if (!PreviewMesh->bAllowCPUAccess)
 		{
-			PreviewMesh->Modify();
-			PreviewMesh->bAllowCPUAccess = true;
-			return true;
-		}));
+			FNiagaraDataInterfaceError CPUAccessNotAllowedError(FText::Format(LOCTEXT("CPUAccessNotAllowedError", "This mesh needs CPU access in order to be used properly.({0})"), FText::FromString(PreviewMesh->GetName())),
+				LOCTEXT("CPUAccessNotAllowedErrorSummary", "CPU access error"),
+				FNiagaraDataInterfaceFix::CreateLambda([=]()
+			{
+				PreviewMesh->Modify();
+				PreviewMesh->bAllowCPUAccess = true;
+				return true;
+			}));
 
-		Errors.Add(CPUAccessNotAllowedError);
+			Errors.Add(CPUAccessNotAllowedError);
+		}
 	}
 #endif
 
+	if (bHasNoMeshAssignedError)
+	{
+		FNiagaraDataInterfaceError NoMeshAssignedError(LOCTEXT("NoMeshAssignedError", "This Data Interface must be assigned a skeletal mesh to operate."),
+			LOCTEXT("NoMeshAssignedErrorSummary", "No mesh assigned error"),
+			FNiagaraDataInterfaceFix());
+
+		Errors.Add(NoMeshAssignedError);
+	}
+
 	return Errors;
 }
 #endif
 
+//Whether or not there is valid mesh data on this interface
+void UNiagaraDataInterfaceStaticMesh::IsValid(FVectorVMContext& Context)
+{
+	VectorVM::FUserPtrHandler<FNDIStaticMesh_InstanceData> InstData(Context);
+	VectorVM::FExternalFuncRegisterHandler<bool> OutValid(Context);
+	
+	for (int32 i = 0; i < Context.NumInstances; ++i)
+	{
+		*OutValid.GetDest() = InstData->Mesh != nullptr;
+		OutValid.Advance();
+	}
+}
+
 //RandomSection specializations.
-//Each combination for AreaWeighted and Section filtered.
+//Each combination for SampleMode and Section filtered. NOTE: Invalid sample mode left intentionally unimplemented
 template<>
-FORCEINLINE int32 UNiagaraDataInterfaceStaticMesh::RandomSection<TIntegralConstant<bool, true>, true>(FRandomStream& RandStream, FStaticMeshLODResources& Res, FNDIStaticMesh_InstanceData* InstData)
+FORCEINLINE int32 UNiagaraDataInterfaceStaticMesh::RandomSection<TSampleModeAreaWeighted, true>(FRandomStream& RandStream, FStaticMeshLODResources& Res, FNDIStaticMesh_InstanceData* InstData)
 {
-	checkSlow(InstData->GetValidSections().Num() > 0);
-	int32 Idx = InstData->GetAreaWeigtedSampler().GetEntryIndex(RandStream.GetFraction(), RandStream.GetFraction());
+	int32 Idx = InstData->GetAreaWeightedSampler().GetEntryIndex(RandStream.GetFraction(), RandStream.GetFraction());
 	return InstData->GetValidSections()[Idx];
 }
 
 template<>
-FORCEINLINE int32 UNiagaraDataInterfaceStaticMesh::RandomSection<TIntegralConstant<bool, true>, false>(FRandomStream& RandStream, FStaticMeshLODResources& Res, FNDIStaticMesh_InstanceData* InstData)
+FORCEINLINE int32 UNiagaraDataInterfaceStaticMesh::RandomSection<TSampleModeAreaWeighted, false>(FRandomStream& RandStream, FStaticMeshLODResources& Res, FNDIStaticMesh_InstanceData* InstData)
 {
 	return Res.AreaWeightedSampler.GetEntryIndex(RandStream.GetFraction(), RandStream.GetFraction());
 }
 
 template<>
-FORCEINLINE int32 UNiagaraDataInterfaceStaticMesh::RandomSection<TIntegralConstant<bool, false>, true>(FRandomStream& RandStream, FStaticMeshLODResources& Res, FNDIStaticMesh_InstanceData* InstData)
+FORCEINLINE int32 UNiagaraDataInterfaceStaticMesh::RandomSection<TSampleModeDefault, true>(FRandomStream& RandStream, FStaticMeshLODResources& Res, FNDIStaticMesh_InstanceData* InstData)
 {
-	checkSlow(InstData->GetValidSections().Num() > 0);
 	int32 Idx = RandStream.RandRange(0, InstData->GetValidSections().Num() - 1);
 	return InstData->GetValidSections()[Idx];
 }
 
 template<>
-FORCEINLINE int32 UNiagaraDataInterfaceStaticMesh::RandomSection<TIntegralConstant<bool, false>, false>(FRandomStream& RandStream, FStaticMeshLODResources& Res, FNDIStaticMesh_InstanceData* InstData)
+FORCEINLINE int32 UNiagaraDataInterfaceStaticMesh::RandomSection<TSampleModeDefault, false>(FRandomStream& RandStream, FStaticMeshLODResources& Res, FNDIStaticMesh_InstanceData* InstData)
 {
 	return RandStream.RandRange(0, Res.Sections.Num() - 1);
 }
 
-template<typename TAreaWeighted>
+template<typename TSampleMode>
 void UNiagaraDataInterfaceStaticMesh::RandomSection(FVectorVMContext& Context)
 {
 	VectorVM::FUserPtrHandler<FNDIStaticMesh_InstanceData> InstData(Context);
@@ -1277,50 +1352,63 @@ void UNiagaraDataInterfaceStaticMesh::RandomSection(FVectorVMContext& Context)
 	FStaticMeshLODResources& Res = InstData->Mesh->RenderData->LODResources[0];
 	for (int32 i = 0; i < Context.NumInstances; ++i)
 	{
-		*OutSection.GetDest() = RandomSection<TAreaWeighted, true>(Context.RandStream, Res, InstData);
-		OutSection.Advance();
+		*OutSection.GetDestAndAdvance() = RandomSection<TSampleMode, true>(Context.RandStream, Res, InstData);
+	}
+}
+
+// Invalid mesh specialization
+template<>
+void UNiagaraDataInterfaceStaticMesh::RandomSection<TSampleModeInvalid>(FVectorVMContext& Context)
+{
+	VectorVM::FUserPtrHandler<FNDIStaticMesh_InstanceData> InstData(Context);
+	VectorVM::FExternalFuncRegisterHandler<int32> OutSection(Context);
+
+	FStaticMeshLODResources& Res = InstData->Mesh->RenderData->LODResources[0];
+	for (int32 i = 0; i < Context.NumInstances; ++i)
+	{
+		*OutSection.GetDestAndAdvance() = -1;
 	}
 }
 
 //RandomTriIndex specializations.
-//Each combination for AreaWeighted and Section filtered.
+//Each combination for SampleMode and Section filtered. NOTE: Invalid sample mode left intentionally unimplemented
 template<>
-FORCEINLINE int32 UNiagaraDataInterfaceStaticMesh::RandomTriIndex<TIntegralConstant<bool, true>, true>(FRandomStream& RandStream, FStaticMeshLODResources& Res, FNDIStaticMesh_InstanceData* InstData)
+FORCEINLINE int32 UNiagaraDataInterfaceStaticMesh::RandomTriIndex<TSampleModeAreaWeighted, true>(FRandomStream& RandStream, FStaticMeshLODResources& Res, FNDIStaticMesh_InstanceData* InstData)
 {
-	int32 SecIdx = RandomSection<TIntegralConstant<bool, true>, true>(RandStream, Res, InstData);
+	int32 SecIdx = RandomSection<TSampleModeAreaWeighted, true>(RandStream, Res, InstData);
 	FStaticMeshSection&  Sec = Res.Sections[SecIdx];
 	int32 Tri = Res.AreaWeightedSectionSamplers[SecIdx].GetEntryIndex(RandStream.GetFraction(), RandStream.GetFraction());
 	return (Sec.FirstIndex / 3) + Tri;
 }
 
 template<>
-FORCEINLINE int32 UNiagaraDataInterfaceStaticMesh::RandomTriIndex<TIntegralConstant<bool, true>, false>(FRandomStream& RandStream, FStaticMeshLODResources& Res, FNDIStaticMesh_InstanceData* InstData)
+FORCEINLINE int32 UNiagaraDataInterfaceStaticMesh::RandomTriIndex<TSampleModeAreaWeighted, false>(FRandomStream& RandStream, FStaticMeshLODResources& Res, FNDIStaticMesh_InstanceData* InstData)
 {
-	int32 SecIdx = RandomSection<TIntegralConstant<bool, true>, false>(RandStream, Res, InstData);
+	int32 SecIdx = RandomSection<TSampleModeAreaWeighted, false>(RandStream, Res, InstData);
 	FStaticMeshSection&  Sec = Res.Sections[SecIdx];
 	int32 Tri = Res.AreaWeightedSectionSamplers[SecIdx].GetEntryIndex(RandStream.GetFraction(), RandStream.GetFraction());
 	return (Sec.FirstIndex / 3) + Tri;
 }
 
 template<>
-FORCEINLINE int32 UNiagaraDataInterfaceStaticMesh::RandomTriIndex<TIntegralConstant<bool, false>, true>(FRandomStream& RandStream, FStaticMeshLODResources& Res, FNDIStaticMesh_InstanceData* InstData)
+FORCEINLINE int32 UNiagaraDataInterfaceStaticMesh::RandomTriIndex<TSampleModeDefault, true>(FRandomStream& RandStream, FStaticMeshLODResources& Res, FNDIStaticMesh_InstanceData* InstData)
 {
-	int32 SecIdx = RandomSection<TIntegralConstant<bool, false>, true>(RandStream, Res, InstData);
+	int32 SecIdx = RandomSection<TSampleModeDefault, true>(RandStream, Res, InstData);
 	FStaticMeshSection&  Sec = Res.Sections[SecIdx];
 	int32 Tri = RandStream.RandRange(0, Sec.NumTriangles - 1);
 	return (Sec.FirstIndex / 3) + Tri;
 }
 
 template<>
-FORCEINLINE int32 UNiagaraDataInterfaceStaticMesh::RandomTriIndex<TIntegralConstant<bool, false>, false>(FRandomStream& RandStream, FStaticMeshLODResources& Res, FNDIStaticMesh_InstanceData* InstData)
+FORCEINLINE int32 UNiagaraDataInterfaceStaticMesh::RandomTriIndex<TSampleModeDefault, false>(FRandomStream& RandStream, FStaticMeshLODResources& Res, FNDIStaticMesh_InstanceData* InstData)
 {
-	int32 SecIdx = RandomSection<TIntegralConstant<bool, false>, false>(RandStream, Res, InstData);
+	int32 SecIdx = RandomSection<TSampleModeDefault, false>(RandStream, Res, InstData);
 	FStaticMeshSection&  Sec = Res.Sections[SecIdx];
 	int32 Tri = RandStream.RandRange(0, Sec.NumTriangles - 1);
 	return (Sec.FirstIndex / 3) + Tri;
 }
 
-template<typename TAreaWeighted>
+template<typename TSampleMode>
 void UNiagaraDataInterfaceStaticMesh::RandomTriCoord(FVectorVMContext& Context)
 {
 	VectorVM::FUserPtrHandler<FNDIStaticMesh_InstanceData> InstData(Context);
@@ -1330,11 +1418,12 @@ void UNiagaraDataInterfaceStaticMesh::RandomTriCoord(FVectorVMContext& Context)
 	VectorVM::FExternalFuncRegisterHandler<float> OutBaryY(Context);
 	VectorVM::FExternalFuncRegisterHandler<float> OutBaryZ(Context);
 
+	check(InstData->Mesh);
 	FStaticMeshLODResources& Res = InstData->Mesh->RenderData->LODResources[0];
 	FIndexArrayView Indices = Res.IndexBuffer.GetArrayView();
 	for (int32 i = 0; i < Context.NumInstances; ++i)
 	{
-		*OutTri.GetDest() = RandomTriIndex<TAreaWeighted, true>(Context.RandStream, Res, InstData);
+		*OutTri.GetDest() = RandomTriIndex<TSampleMode, true>(Context.RandStream, Res, InstData);
 		FVector Bary = RandomBarycentricCoord(Context.RandStream);
 		*OutBaryX.GetDest() = Bary.X;
 		*OutBaryY.GetDest() = Bary.Y;
@@ -1347,6 +1436,26 @@ void UNiagaraDataInterfaceStaticMesh::RandomTriCoord(FVectorVMContext& Context)
 	}
 }
 
+// Invalid mesh specialization
+template<>
+void UNiagaraDataInterfaceStaticMesh::RandomTriCoord<TSampleModeInvalid>(FVectorVMContext& Context)
+{
+	VectorVM::FUserPtrHandler<FNDIStaticMesh_InstanceData> InstData(Context);
+
+	VectorVM::FExternalFuncRegisterHandler<int32> OutTri(Context);
+	VectorVM::FExternalFuncRegisterHandler<float> OutBaryX(Context);
+	VectorVM::FExternalFuncRegisterHandler<float> OutBaryY(Context);
+	VectorVM::FExternalFuncRegisterHandler<float> OutBaryZ(Context);
+
+	for (int32 i = 0; i < Context.NumInstances; ++i)
+	{
+		*OutTri.GetDestAndAdvance() = -1;
+		*OutBaryX.GetDestAndAdvance() = 0.0f;
+		*OutBaryY.GetDestAndAdvance() = 0.0f;
+		*OutBaryZ.GetDestAndAdvance() = 0.0f;
+	}
+}
+
 void UNiagaraDataInterfaceStaticMesh::RandomTriCoordVertexColorFiltered(FVectorVMContext& Context)
 {
 	VectorVM::FExternalFuncRegisterHandler<int32> MinValue(Context);
@@ -1357,9 +1466,22 @@ void UNiagaraDataInterfaceStaticMesh::RandomTriCoordVertexColorFiltered(FVectorV
 	VectorVM::FExternalFuncRegisterHandler<float> OutBaryX(Context);
 	VectorVM::FExternalFuncRegisterHandler<float> OutBaryY(Context);
 	VectorVM::FExternalFuncRegisterHandler<float> OutBaryZ(Context);
+
+	// Handle no mesh case
+	//TODO: Maybe figure out a good way to stub this in bindings to prevent the branch
+	if (!InstData->Mesh)
+	{
+		for (int32 i = 0; i < Context.NumInstances; ++i)
+		{
+			*OutTri.GetDestAndAdvance() = -1;
+			*OutBaryX.GetDestAndAdvance() = 0.0f;
+			*OutBaryY.GetDestAndAdvance() = 0.0f;
+			*OutBaryZ.GetDestAndAdvance() = 0.0f;
+		}
+		return;
+	}
 	
 	FDynamicVertexColorFilterData* VCFData = InstData->DynamicVertexColorSampler.Get();
-
 	FStaticMeshLODResources& Res = InstData->Mesh->RenderData->LODResources[0];
 	FIndexArrayView Indices = Res.IndexBuffer.GetArrayView();
 
@@ -1412,7 +1534,7 @@ void UNiagaraDataInterfaceStaticMesh::RandomTriCoordVertexColorFiltered(FVectorV
 }
 
 template<>
-FORCEINLINE int32 UNiagaraDataInterfaceStaticMesh::RandomTriIndexOnSection<TIntegralConstant<bool, true>>(FRandomStream& RandStream, FStaticMeshLODResources& Res, int32 SecIdx, FNDIStaticMesh_InstanceData* InstData)
+FORCEINLINE int32 UNiagaraDataInterfaceStaticMesh::RandomTriIndexOnSection<TSampleModeAreaWeighted>(FRandomStream& RandStream, FStaticMeshLODResources& Res, int32 SecIdx, FNDIStaticMesh_InstanceData* InstData)
 {
 	FStaticMeshSection&  Sec = Res.Sections[SecIdx];
 	int32 Tri = Res.AreaWeightedSectionSamplers[SecIdx].GetEntryIndex(RandStream.GetFraction(), RandStream.GetFraction());
@@ -1420,14 +1542,14 @@ FORCEINLINE int32 UNiagaraDataInterfaceStaticMesh::RandomTriIndexOnSection<TInte
 }
 
 template<>
-FORCEINLINE int32 UNiagaraDataInterfaceStaticMesh::RandomTriIndexOnSection<TIntegralConstant<bool, false>>(FRandomStream& RandStream, FStaticMeshLODResources& Res, int32 SecIdx, FNDIStaticMesh_InstanceData* InstData)
+FORCEINLINE int32 UNiagaraDataInterfaceStaticMesh::RandomTriIndexOnSection<TSampleModeDefault>(FRandomStream& RandStream, FStaticMeshLODResources& Res, int32 SecIdx, FNDIStaticMesh_InstanceData* InstData)
 {
 	FStaticMeshSection&  Sec = Res.Sections[SecIdx];
 	int32 Tri = RandStream.RandRange(0, Sec.NumTriangles - 1);
 	return (Sec.FirstIndex / 3) + Tri;
 }
 
-template<typename TAreaWeighted>
+template<typename TSampleMode>
 void UNiagaraDataInterfaceStaticMesh::RandomTriCoordOnSection(FVectorVMContext& Context)
 {
 	VectorVM::FExternalFuncInputHandler<int32> SectionIdxParam(Context);
@@ -1438,6 +1560,7 @@ void UNiagaraDataInterfaceStaticMesh::RandomTriCoordOnSection(FVectorVMContext&
 	VectorVM::FExternalFuncRegisterHandler<float> OutBaryY(Context);
 	VectorVM::FExternalFuncRegisterHandler<float> OutBaryZ(Context);
 
+	check(InstData->Mesh);
 	FStaticMeshLODResources& Res = InstData->Mesh->RenderData->LODResources[0];
 	FIndexArrayView Indices = Res.IndexBuffer.GetArrayView();
 	const int32 MaxSection = Res.Sections.Num() - 1;
@@ -1445,7 +1568,7 @@ void UNiagaraDataInterfaceStaticMesh::RandomTriCoordOnSection(FVectorVMContext&
 	{
 
 		int32 SecIdx = FMath::Clamp(SectionIdxParam.Get(), 0, MaxSection);
-		*OutTri.GetDest() = RandomTriIndexOnSection<TAreaWeighted>(Context.RandStream, Res, SecIdx, InstData);
+		*OutTri.GetDest() = RandomTriIndexOnSection<TSampleMode>(Context.RandStream, Res, SecIdx, InstData);
 		FVector Bary = RandomBarycentricCoord(Context.RandStream);
 		*OutBaryX.GetDest() = Bary.X;
 		*OutBaryY.GetDest() = Bary.Y;
@@ -1459,6 +1582,27 @@ void UNiagaraDataInterfaceStaticMesh::RandomTriCoordOnSection(FVectorVMContext&
 	}
 }
 
+// Invalid mesh specialization
+template<>
+void UNiagaraDataInterfaceStaticMesh::RandomTriCoordOnSection<TSampleModeInvalid>(FVectorVMContext& Context)
+{
+	VectorVM::FExternalFuncInputHandler<int32> SectionIdxParam(Context);
+	VectorVM::FUserPtrHandler<FNDIStaticMesh_InstanceData> InstData(Context);
+
+	VectorVM::FExternalFuncRegisterHandler<int32> OutTri(Context);
+	VectorVM::FExternalFuncRegisterHandler<float> OutBaryX(Context);
+	VectorVM::FExternalFuncRegisterHandler<float> OutBaryY(Context);
+	VectorVM::FExternalFuncRegisterHandler<float> OutBaryZ(Context);
+
+	for (int32 i = 0; i < Context.NumInstances; ++i)
+	{
+		*OutTri.GetDestAndAdvance() = -1;
+		*OutBaryX.GetDestAndAdvance() = 0.0f;
+		*OutBaryY.GetDestAndAdvance() = 0.0f;
+		*OutBaryZ.GetDestAndAdvance() = 0.0f;
+	}
+}
+
 template<typename TransformHandlerType>
 void UNiagaraDataInterfaceStaticMesh::GetTriCoordPosition(FVectorVMContext& Context)
 {
@@ -1473,6 +1617,22 @@ void UNiagaraDataInterfaceStaticMesh::GetTriCoordPosition(FVectorVMContext& Cont
 	VectorVM::FExternalFuncRegisterHandler<float> OutPosY(Context);
 	VectorVM::FExternalFuncRegisterHandler<float> OutPosZ(Context);
 
+	// Handle no mesh case
+	//TODO: Maybe figure out a good way to stub this in bindings to prevent the branch
+	if (!InstData->Mesh)
+	{
+		FVector Pos(0.0f);
+		TransformHandler.TransformPosition(Pos, InstData->Transform);
+
+		for (int32 i = 0; i < Context.NumInstances; ++i)
+		{
+			*OutPosX.GetDestAndAdvance() = Pos.X;
+			*OutPosY.GetDestAndAdvance() = Pos.Y;
+			*OutPosZ.GetDestAndAdvance() = Pos.Z;
+		}
+		return;
+	}
+
 	FStaticMeshLODResources& Res = InstData->Mesh->RenderData->LODResources[0];
 	const FIndexArrayView& Indices = Res.IndexBuffer.GetArrayView();
 	const FPositionVertexBuffer& Positions = Res.VertexBuffers.PositionVertexBuffer;
@@ -1517,6 +1677,19 @@ void UNiagaraDataInterfaceStaticMesh::GetTriCoordNormal(FVectorVMContext& Contex
 	VectorVM::FExternalFuncRegisterHandler<float> OutNormY(Context);
 	VectorVM::FExternalFuncRegisterHandler<float> OutNormZ(Context);
 
+	// Handle no mesh case
+	//TODO: Maybe figure out a good way to stub this in bindings to prevent the branch
+	if (!InstData->Mesh)
+	{
+		for (int32 i = 0; i < Context.NumInstances; ++i)
+		{
+			*OutNormX.GetDestAndAdvance() = 0.0f;
+			*OutNormY.GetDestAndAdvance() = 0.0f;
+			*OutNormZ.GetDestAndAdvance() = 1.0f;
+		}
+		return;
+	}
+
 	FStaticMeshLODResources& Res = InstData->Mesh->RenderData->LODResources[0];
 	const FIndexArrayView& Indices = Res.IndexBuffer.GetArrayView();
 	const FStaticMeshVertexBuffer& Verts = Res.VertexBuffers.StaticMeshVertexBuffer;
@@ -1556,10 +1729,6 @@ void UNiagaraDataInterfaceStaticMesh::GetTriCoordTangents(FVectorVMContext& Cont
 	VectorVM::FExternalFuncInputHandler<float> BaryZParam(Context);
 	VectorVM::FUserPtrHandler<FNDIStaticMesh_InstanceData> InstData(Context);
 
-	FStaticMeshLODResources& Res = InstData->Mesh->RenderData->LODResources[0];
-	const FIndexArrayView& Indices = Res.IndexBuffer.GetArrayView();
-	const VertexAccessorType Verts(Res.VertexBuffers.StaticMeshVertexBuffer);
-
 	VectorVM::FExternalFuncRegisterHandler<float> OutTangentX(Context);
 	VectorVM::FExternalFuncRegisterHandler<float> OutTangentY(Context);
 	VectorVM::FExternalFuncRegisterHandler<float> OutTangentZ(Context);
@@ -1570,6 +1739,28 @@ void UNiagaraDataInterfaceStaticMesh::GetTriCoordTangents(FVectorVMContext& Cont
 	VectorVM::FExternalFuncRegisterHandler<float> OutNormY(Context);
 	VectorVM::FExternalFuncRegisterHandler<float> OutNormZ(Context);
 
+	// Handle no mesh case
+	//TODO: Maybe figure out a good way to stub this in bindings to prevent the branch
+	if (!InstData->Mesh)
+	{
+		for (int32 i = 0; i < Context.NumInstances; ++i)
+		{
+			*OutTangentX.GetDestAndAdvance() = 1.0f;
+			*OutTangentY.GetDestAndAdvance() = 0.0f;
+			*OutTangentZ.GetDestAndAdvance() = 0.0f;
+			*OutBinormX.GetDestAndAdvance() = 0.0f;
+			*OutBinormY.GetDestAndAdvance() = 1.0f;
+			*OutBinormZ.GetDestAndAdvance() = 0.0f;
+			*OutNormX.GetDestAndAdvance() = 0.0f;
+			*OutNormY.GetDestAndAdvance() = 0.0f;
+			*OutNormZ.GetDestAndAdvance() = 1.0f;
+		}
+		return;
+	}
+
+	FStaticMeshLODResources& Res = InstData->Mesh->RenderData->LODResources[0];
+	const FIndexArrayView& Indices = Res.IndexBuffer.GetArrayView();
+	const VertexAccessorType Verts(Res.VertexBuffers.StaticMeshVertexBuffer);
 	const int32 NumTriangles = Indices.Num() / 3;
 	for (int32 i = 0; i < Context.NumInstances; ++i)
 	{
@@ -1622,12 +1813,16 @@ void UNiagaraDataInterfaceStaticMesh::GetTriCoordColor(FVectorVMContext& Context
 	VectorVM::FExternalFuncRegisterHandler<float> OutColorB(Context);
 	VectorVM::FExternalFuncRegisterHandler<float> OutColorA(Context);
 
-	FStaticMeshLODResources& Res = InstData->Mesh->RenderData->LODResources[0];
-	const FIndexArrayView& Indices = Res.IndexBuffer.GetArrayView();
-	const FColorVertexBuffer& Colors = Res.VertexBuffers.ColorVertexBuffer;
+	FStaticMeshLODResources* Res = nullptr;
+	if (InstData->Mesh)
+	{
+		Res = &InstData->Mesh->RenderData->LODResources[0];
+	}
 
-	if (Colors.GetNumVertices() > 0)
+	if (Res && Res->VertexBuffers.ColorVertexBuffer.GetNumVertices() > 0)
 	{
+		const FIndexArrayView& Indices = Res->IndexBuffer.GetArrayView();
+		const FColorVertexBuffer& Colors = Res->VertexBuffers.ColorVertexBuffer;
 		const int32 NumTriangles = Indices.Num() / 3;
 		for (int32 i = 0; i < Context.NumInstances; ++i)
 		{
@@ -1655,22 +1850,14 @@ void UNiagaraDataInterfaceStaticMesh::GetTriCoordColor(FVectorVMContext& Context
 	}
 	else
 	{
-		// This mesh doesn't have color information so set the color to white.
+		// This mesh is invalid or doesn't have color information so set the color to white.
 		FLinearColor Color = FLinearColor::White;
 		for (int32 i = 0; i < Context.NumInstances; ++i)
 		{
-			*OutColorR.GetDest() = Color.R;
-			*OutColorG.GetDest() = Color.G;
-			*OutColorB.GetDest() = Color.B;
-			*OutColorA.GetDest() = Color.A;
-			TriParam.Advance();
-			BaryXParam.Advance();
-			BaryYParam.Advance();
-			BaryZParam.Advance();
-			OutColorR.Advance();
-			OutColorG.Advance();
-			OutColorB.Advance();
-			OutColorA.Advance();
+			*OutColorR.GetDestAndAdvance() = Color.R;
+			*OutColorG.GetDestAndAdvance() = Color.G;
+			*OutColorB.GetDestAndAdvance() = Color.B;
+			*OutColorA.GetDestAndAdvance() = Color.A;			
 		}
 	}
 }
@@ -1688,6 +1875,18 @@ void UNiagaraDataInterfaceStaticMesh::GetTriCoordUV(FVectorVMContext& Context)
 	VectorVM::FExternalFuncRegisterHandler<float> OutU(Context);
 	VectorVM::FExternalFuncRegisterHandler<float> OutV(Context);
 
+	// Handle no mesh case
+	//TODO: Maybe figure out a good way to stub this in bindings to prevent the branch
+	if (!InstData->Mesh)
+	{
+		for (int32 i = 0; i < Context.NumInstances; ++i)
+		{
+			*OutU.GetDestAndAdvance() = 0.0f;
+			*OutV.GetDestAndAdvance() = 0.0f;
+		}
+		return;
+	}
+
 	FStaticMeshLODResources& Res = InstData->Mesh->RenderData->LODResources[0];
 	const FIndexArrayView& Indices = Res.IndexBuffer.GetArrayView();
 	const VertexAccessorType Verts(Res.VertexBuffers.StaticMeshVertexBuffer);
@@ -1731,6 +1930,23 @@ void UNiagaraDataInterfaceStaticMesh::GetTriCoordPositionAndVelocity(FVectorVMCo
 	VectorVM::FExternalFuncRegisterHandler<float> OutVelY(Context);
 	VectorVM::FExternalFuncRegisterHandler<float> OutVelZ(Context);
 
+	// Handle no mesh case
+	//TODO: Maybe figure out a good way to stub this in bindings to prevent the branch
+	if (!InstData->Mesh)
+	{
+		FVector WSPos = InstData->Transform.TransformPosition(FVector(0.0f));
+		for (int32 i = 0; i < Context.NumInstances; ++i)
+		{
+			*OutPosX.GetDestAndAdvance() = WSPos.X;
+			*OutPosY.GetDestAndAdvance() = WSPos.Y;
+			*OutPosZ.GetDestAndAdvance() = WSPos.Z;
+			*OutVelX.GetDestAndAdvance() = 0.0f;
+			*OutVelY.GetDestAndAdvance() = 0.0f;
+			*OutVelZ.GetDestAndAdvance() = 0.0f;
+		}
+		return;
+	}
+
 	FStaticMeshLODResources& Res = InstData->Mesh->RenderData->LODResources[0];
 	const FIndexArrayView& Indices = Res.IndexBuffer.GetArrayView();
 	const FPositionVertexBuffer& Positions = Res.VertexBuffers.PositionVertexBuffer;
@@ -1859,6 +2075,20 @@ void UNiagaraDataInterfaceStaticMesh::GetVertexPosition(FVectorVMContext& Contex
 	VectorVM::FExternalFuncRegisterHandler<float> OutPosY(Context);
 	VectorVM::FExternalFuncRegisterHandler<float> OutPosZ(Context);
 
+	// Handle no mesh case
+	//TODO: Maybe figure out a good way to stub this in bindings to prevent the branch
+	if (!InstData->Mesh)
+	{
+		FVector WSPos = InstData->Transform.TransformPosition(FVector(0.0f));
+		for (int32 i = 0; i < Context.NumInstances; ++i)
+		{
+			*OutPosX.GetDestAndAdvance() = WSPos.X;
+			*OutPosY.GetDestAndAdvance() = WSPos.Y;
+			*OutPosZ.GetDestAndAdvance() = WSPos.Z;
+		}
+		return;
+	}
+
 	FStaticMeshLODResources& Res = InstData->Mesh->RenderData->LODResources[0];
 	const FPositionVertexBuffer& Positions = Res.VertexBuffers.PositionVertexBuffer;
 
@@ -1921,11 +2151,28 @@ bool UNiagaraDataInterfaceStaticMesh::GetFunctionHLSL(const FNiagaraDataInterfac
 		{TEXT("UseColorBufferName"), ParamNames.UseColorBufferName},
 	};
 
-	if (FunctionInfo.DefinitionName == StaticMeshHelpers::RandomSectionName)
+	if (FunctionInfo.DefinitionName == StaticMeshHelpers::IsValidName)
+	{
+		static const TCHAR* FormatSample = TEXT(R"(
+			void {InstanceFunctionName} (out bool Out_Valid)
+			{
+				Out_Valid = {SectionCountName} > 0;				
+			}
+			)");
+		OutHLSL += FString::Format(FormatSample, ArgsSample);
+	}
+	else if (FunctionInfo.DefinitionName == StaticMeshHelpers::RandomSectionName)
 	{
 		static const TCHAR *FormatSample = TEXT(R"(
 			void {InstanceFunctionName} (out int Out_Section)
 			{
+				[branch]
+				if ({SectionCountName} == 0)
+				{
+					Out_Section = -1;
+					return;
+				}
+
 				float RandS0 = NiagaraInternalNoise(1, 2, 3);
 				// Uniform sampling on mesh surface  (using alias method from Alias method from FWeightedRandomSampler)
 				uint SectionIndex = min(uint(RandS0 * float({SectionCountName})), {SectionCountName}-1);
@@ -1947,6 +2194,14 @@ bool UNiagaraDataInterfaceStaticMesh::GetFunctionHLSL(const FNiagaraDataInterfac
 		static const TCHAR *FormatSample = TEXT(R"(
 			void {InstanceFunctionName} (out {MeshTriCoordinateStructName} Out_Coord)
 			{
+				[branch]
+				if ({SectionCountName} == 0)
+				{
+					Out_Coord.Tri = -1;
+					Out_Coord.BaryCoord = (float3)0.0f;
+					return;
+				}
+
 				float RandS0 = NiagaraInternalNoise(1, 2, 3);
 
 				// Uniform sampling on mesh surface  (using alias method from Alias method from FWeightedRandomSampler)
@@ -2002,6 +2257,14 @@ bool UNiagaraDataInterfaceStaticMesh::GetFunctionHLSL(const FNiagaraDataInterfac
 		static const TCHAR *FormatSample = TEXT(R"(
 			void {InstanceFunctionName} (in int In_Section, out {MeshTriCoordinateStructName} Out_Coord)
 			{
+				[branch]
+				if ({SectionCountName} == 0)
+				{
+					Out_Coord.Tri = -1;
+					Out_Coord.BaryCoord = (float3)0.0f;
+					return;
+				}
+
 				int Section = clamp(In_Section, 0, (int)({SectionCountName} - 1));
 
 				uint4 SectionData = {MeshSectionBufferName}[Section];
@@ -2044,6 +2307,13 @@ bool UNiagaraDataInterfaceStaticMesh::GetFunctionHLSL(const FNiagaraDataInterfac
 		static const TCHAR *FormatSample = TEXT(R"(
 			void {InstanceFunctionName} (in {MeshTriCoordinateStructName} In_Coord, out float3 Out_Position)
 			{
+				[branch]
+				if ({SectionCountName} == 0)
+				{
+					Out_Position = (float3)0.0f;
+					return;
+				}
+
 				uint TriangleIndex = In_Coord.Tri * 3;
 				uint VertexIndex0 = {MeshIndexBufferName}[TriangleIndex  ] * 3;
 				uint VertexIndex1 = {MeshIndexBufferName}[TriangleIndex+1] * 3;
@@ -2063,6 +2333,13 @@ bool UNiagaraDataInterfaceStaticMesh::GetFunctionHLSL(const FNiagaraDataInterfac
 		static const TCHAR *FormatSample = TEXT(R"(
 			void {InstanceFunctionName} (in {MeshTriCoordinateStructName} In_Coord, out float3 Out_Position)
 			{
+				[branch]
+				if ({SectionCountName} == 0)
+				{
+					Out_Position = {InstanceTransformName}[3].xyz;
+					return;
+				}
+
 				uint TriangleIndex = In_Coord.Tri * 3;
 				uint VertexIndex0 = {MeshIndexBufferName}[TriangleIndex  ] * 3;
 				uint VertexIndex1 = {MeshIndexBufferName}[TriangleIndex+1] * 3;
@@ -2082,107 +2359,139 @@ bool UNiagaraDataInterfaceStaticMesh::GetFunctionHLSL(const FNiagaraDataInterfac
 	else if (FunctionInfo.DefinitionName == StaticMeshHelpers::GetTriNormalName)
 	{
 		static const TCHAR *FormatSample = TEXT(R"(
-		void {InstanceFunctionName} (in {MeshTriCoordinateStructName} In_Coord, out float3 Out_Normal)
-		{
-			uint TriangleIndex = In_Coord.Tri * 3;
-			uint VertexIndex0 = {MeshIndexBufferName}[TriangleIndex  ] * 2;
-			uint VertexIndex1 = {MeshIndexBufferName}[TriangleIndex+1] * 2;
-			uint VertexIndex2 = {MeshIndexBufferName}[TriangleIndex+2] * 2;
+			void {InstanceFunctionName} (in {MeshTriCoordinateStructName} In_Coord, out float3 Out_Normal)
+			{
+				[branch]
+				if ({SectionCountName} == 0)
+				{
+					Out_Normal = float3(0, 0, 1);
+					return;
+				}
 
-			float3 Normal0 = TangentBias({MeshTangentBufferName}[VertexIndex0+1].xyz);
-			float3 Normal1 = TangentBias({MeshTangentBufferName}[VertexIndex1+1].xyz);
-			float3 Normal2 = TangentBias({MeshTangentBufferName}[VertexIndex2+1].xyz);
+				uint TriangleIndex = In_Coord.Tri * 3;
+				uint VertexIndex0 = {MeshIndexBufferName}[TriangleIndex  ] * 2;
+				uint VertexIndex1 = {MeshIndexBufferName}[TriangleIndex+1] * 2;
+				uint VertexIndex2 = {MeshIndexBufferName}[TriangleIndex+2] * 2;
 
-			float3 Normal   = Normal0 * In_Coord.BaryCoord.x + Normal1 * In_Coord.BaryCoord.y + Normal2 * In_Coord.BaryCoord.z;
+				float3 Normal0 = TangentBias({MeshTangentBufferName}[VertexIndex0+1].xyz);
+				float3 Normal1 = TangentBias({MeshTangentBufferName}[VertexIndex1+1].xyz);
+				float3 Normal2 = TangentBias({MeshTangentBufferName}[VertexIndex2+1].xyz);
 
-			Out_Normal = normalize(Normal);
-		}
-		)");
+				float3 Normal   = Normal0 * In_Coord.BaryCoord.x + Normal1 * In_Coord.BaryCoord.y + Normal2 * In_Coord.BaryCoord.z;
+
+				Out_Normal = normalize(Normal);
+			}
+			)");
 		OutHLSL += FString::Format(FormatSample, ArgsSample);
 	}
 	else if (FunctionInfo.DefinitionName == StaticMeshHelpers::GetTriNormalWSName)
 	{
 		static const TCHAR *FormatSample = TEXT(R"(
-		void {InstanceFunctionName} (in {MeshTriCoordinateStructName} In_Coord, out float3 Out_Normal)
-		{
-			uint TriangleIndex = In_Coord.Tri * 3;
-			uint VertexIndex0 = {MeshIndexBufferName}[TriangleIndex  ] * 2;
-			uint VertexIndex1 = {MeshIndexBufferName}[TriangleIndex+1] * 2;
-			uint VertexIndex2 = {MeshIndexBufferName}[TriangleIndex+2] * 2;
+			void {InstanceFunctionName} (in {MeshTriCoordinateStructName} In_Coord, out float3 Out_Normal)
+			{
+				[branch]
+				if ({SectionCountName} == 0)
+				{
+					Out_Normal = float3(0, 0, 1);
+					return;
+				}
+
+				uint TriangleIndex = In_Coord.Tri * 3;
+				uint VertexIndex0 = {MeshIndexBufferName}[TriangleIndex  ] * 2;
+				uint VertexIndex1 = {MeshIndexBufferName}[TriangleIndex+1] * 2;
+				uint VertexIndex2 = {MeshIndexBufferName}[TriangleIndex+2] * 2;
 
-			float3 Normal0 = TangentBias({MeshTangentBufferName}[VertexIndex0+1].xyz);
-			float3 Normal1 = TangentBias({MeshTangentBufferName}[VertexIndex1+1].xyz);
-			float3 Normal2 = TangentBias({MeshTangentBufferName}[VertexIndex2+1].xyz);
+				float3 Normal0 = TangentBias({MeshTangentBufferName}[VertexIndex0+1].xyz);
+				float3 Normal1 = TangentBias({MeshTangentBufferName}[VertexIndex1+1].xyz);
+				float3 Normal2 = TangentBias({MeshTangentBufferName}[VertexIndex2+1].xyz);
 
-			float3 Normal   = Normal0 * In_Coord.BaryCoord.x + Normal1 * In_Coord.BaryCoord.y + Normal2 * In_Coord.BaryCoord.z;
+				float3 Normal   = Normal0 * In_Coord.BaryCoord.x + Normal1 * In_Coord.BaryCoord.y + Normal2 * In_Coord.BaryCoord.z;
 
-			Out_Normal = normalize(mul(float4(Normal, 0.0), {InstanceTransformName}).xyz);
-		}
-		)");
+				Out_Normal = normalize(mul(float4(Normal, 0.0), {InstanceTransformInverseTransposed}).xyz);
+			}
+			)");
 		OutHLSL += FString::Format(FormatSample, ArgsSample);
 	}
 	else if (FunctionInfo.DefinitionName == StaticMeshHelpers::GetTriTangentsName)
 	{
 		static const TCHAR *FormatSample = TEXT(R"(
-		void {InstanceFunctionName} (in {MeshTriCoordinateStructName} In_Coord, out float3 Out_Tangent, out float3 Out_Binormal, out float3 Out_Normal)
-		{
-			uint TriangleIndex = In_Coord.Tri * 3;
-			uint VertexIndex0 = {MeshIndexBufferName}[TriangleIndex  ] * 2;
-			uint VertexIndex1 = {MeshIndexBufferName}[TriangleIndex+1] * 2;
-			uint VertexIndex2 = {MeshIndexBufferName}[TriangleIndex+2] * 2;
-
-			float3 TangentX0 = TangentBias({MeshTangentBufferName}[VertexIndex0  ].xyz);
-			float4 TangentZ0 = TangentBias({MeshTangentBufferName}[VertexIndex0+1].xyzw);
-			float3 TangentX1 = TangentBias({MeshTangentBufferName}[VertexIndex1  ].xyz);
-			float4 TangentZ1 = TangentBias({MeshTangentBufferName}[VertexIndex1+1].xyzw);
-			float3 TangentX2 = TangentBias({MeshTangentBufferName}[VertexIndex2  ].xyz);
-			float4 TangentZ2 = TangentBias({MeshTangentBufferName}[VertexIndex2+1].xyzw);
-
-			float3 Binormal0   = cross(TangentZ0.xyz, TangentX0.xyz) * TangentZ0.w;
-			float3 Binormal1   = cross(TangentZ1.xyz, TangentX1.xyz) * TangentZ1.w;
-			float3 Binormal2   = cross(TangentZ2.xyz, TangentX2.xyz) * TangentZ2.w;
-
-			Out_Normal   = normalize(TangentZ0.xyz * In_Coord.BaryCoord.x + TangentZ1.xyz * In_Coord.BaryCoord.y + TangentZ2.xyz * In_Coord.BaryCoord.z);  // Normal is TangentZ
-			Out_Tangent  = normalize(TangentX0.xyz * In_Coord.BaryCoord.x + TangentX1.xyz * In_Coord.BaryCoord.y + TangentX2.xyz * In_Coord.BaryCoord.z);
-			Out_Binormal = normalize(Binormal0.xyz * In_Coord.BaryCoord.x + Binormal1.xyz * In_Coord.BaryCoord.y + Binormal2.xyz * In_Coord.BaryCoord.z);
-		}
-		)");
+			void {InstanceFunctionName} (in {MeshTriCoordinateStructName} In_Coord, out float3 Out_Tangent, out float3 Out_Binormal, out float3 Out_Normal)
+			{
+				[branch]
+				if ({SectionCountName} == 0)
+				{
+					Out_Tangent = float3(1, 0, 0);
+					Out_Binormal = float3(0, 1, 0);
+					Out_Normal = float3(0, 0, 1);
+					return;
+				}
+
+				uint TriangleIndex = In_Coord.Tri * 3;
+				uint VertexIndex0 = {MeshIndexBufferName}[TriangleIndex  ] * 2;
+				uint VertexIndex1 = {MeshIndexBufferName}[TriangleIndex+1] * 2;
+				uint VertexIndex2 = {MeshIndexBufferName}[TriangleIndex+2] * 2;
+
+				float3 TangentX0 = TangentBias({MeshTangentBufferName}[VertexIndex0  ].xyz);
+				float4 TangentZ0 = TangentBias({MeshTangentBufferName}[VertexIndex0+1].xyzw);
+				float3 TangentX1 = TangentBias({MeshTangentBufferName}[VertexIndex1  ].xyz);
+				float4 TangentZ1 = TangentBias({MeshTangentBufferName}[VertexIndex1+1].xyzw);
+				float3 TangentX2 = TangentBias({MeshTangentBufferName}[VertexIndex2  ].xyz);
+				float4 TangentZ2 = TangentBias({MeshTangentBufferName}[VertexIndex2+1].xyzw);
+
+				float3 Binormal0   = cross(TangentZ0.xyz, TangentX0.xyz) * TangentZ0.w;
+				float3 Binormal1   = cross(TangentZ1.xyz, TangentX1.xyz) * TangentZ1.w;
+				float3 Binormal2   = cross(TangentZ2.xyz, TangentX2.xyz) * TangentZ2.w;
+
+				Out_Normal   = normalize(TangentZ0.xyz * In_Coord.BaryCoord.x + TangentZ1.xyz * In_Coord.BaryCoord.y + TangentZ2.xyz * In_Coord.BaryCoord.z);  // Normal is TangentZ
+				Out_Tangent  = normalize(TangentX0.xyz * In_Coord.BaryCoord.x + TangentX1.xyz * In_Coord.BaryCoord.y + TangentX2.xyz * In_Coord.BaryCoord.z);
+				Out_Binormal = normalize(Binormal0.xyz * In_Coord.BaryCoord.x + Binormal1.xyz * In_Coord.BaryCoord.y + Binormal2.xyz * In_Coord.BaryCoord.z);
+			}
+			)");
 		OutHLSL += FString::Format(FormatSample, ArgsSample);
 	}
 	else if (FunctionInfo.DefinitionName == StaticMeshHelpers::GetTriTangentsWSName)
 	{
 		static const TCHAR *FormatSample = TEXT(R"(
-		void {InstanceFunctionName} (in {MeshTriCoordinateStructName} In_Coord, out float3 Out_Tangent, out float3 Out_Binormal, out float3 Out_Normal)
-		{
-			uint TriangleIndex = In_Coord.Tri * 3;
-			uint VertexIndex0 = {MeshIndexBufferName}[TriangleIndex  ] * 2;
-			uint VertexIndex1 = {MeshIndexBufferName}[TriangleIndex+1] * 2;
-			uint VertexIndex2 = {MeshIndexBufferName}[TriangleIndex+2] * 2;
-
-			float3 TangentX0 = TangentBias({MeshTangentBufferName}[VertexIndex0  ].xyz);
-			float4 TangentZ0 = TangentBias({MeshTangentBufferName}[VertexIndex0+1].xyzw);
-			float3 TangentX1 = TangentBias({MeshTangentBufferName}[VertexIndex1  ].xyz);
-			float4 TangentZ1 = TangentBias({MeshTangentBufferName}[VertexIndex1+1].xyzw);
-			float3 TangentX2 = TangentBias({MeshTangentBufferName}[VertexIndex2  ].xyz);
-			float4 TangentZ2 = TangentBias({MeshTangentBufferName}[VertexIndex2+1].xyzw);
-
-			float3 Binormal0   = cross(TangentZ0.xyz, TangentX0.xyz) * TangentZ0.w;
-			float3 Binormal1   = cross(TangentZ1.xyz, TangentX1.xyz) * TangentZ1.w;
-			float3 Binormal2   = cross(TangentZ2.xyz, TangentX2.xyz) * TangentZ2.w;
-
-			float3 Normal   = TangentZ0.xyz * In_Coord.BaryCoord.x + TangentZ1.xyz * In_Coord.BaryCoord.y + TangentZ2.xyz * In_Coord.BaryCoord.z;  // Normal is TangentZ
-			float3 Tangent  = TangentX0.xyz * In_Coord.BaryCoord.x + TangentX1.xyz * In_Coord.BaryCoord.y + TangentX2.xyz * In_Coord.BaryCoord.z;
-			float3 Binormal = Binormal0.xyz * In_Coord.BaryCoord.x + Binormal1.xyz * In_Coord.BaryCoord.y + Binormal2.xyz * In_Coord.BaryCoord.z;
-
-			float3 NormalWorld  = normalize(mul(float4(Normal  , 0.0), {InstanceTransformName}).xyz);
-			float3 TangentWorld = normalize(mul(float4(Tangent , 0.0), {InstanceTransformName}).xyz);
-			float3 BinormalWorld= normalize(mul(float4(Binormal, 0.0), {InstanceTransformName}).xyz);
-
-			Out_Normal = NormalWorld;
-			Out_Tangent = TangentWorld;
-			Out_Binormal = BinormalWorld;
-		}
-		)");
+			void {InstanceFunctionName} (in {MeshTriCoordinateStructName} In_Coord, out float3 Out_Tangent, out float3 Out_Binormal, out float3 Out_Normal)
+			{
+				[branch]
+				if ({SectionCountName} == 0)
+				{
+					Out_Tangent = float3(1, 0, 0);
+					Out_Binormal = float3(0, 1, 0);
+					Out_Normal = float3(0, 0, 1);
+					return;
+				}
+
+				uint TriangleIndex = In_Coord.Tri * 3;
+				uint VertexIndex0 = {MeshIndexBufferName}[TriangleIndex  ] * 2;
+				uint VertexIndex1 = {MeshIndexBufferName}[TriangleIndex+1] * 2;
+				uint VertexIndex2 = {MeshIndexBufferName}[TriangleIndex+2] * 2;
+
+				float3 TangentX0 = TangentBias({MeshTangentBufferName}[VertexIndex0  ].xyz);
+				float4 TangentZ0 = TangentBias({MeshTangentBufferName}[VertexIndex0+1].xyzw);
+				float3 TangentX1 = TangentBias({MeshTangentBufferName}[VertexIndex1  ].xyz);
+				float4 TangentZ1 = TangentBias({MeshTangentBufferName}[VertexIndex1+1].xyzw);
+				float3 TangentX2 = TangentBias({MeshTangentBufferName}[VertexIndex2  ].xyz);
+				float4 TangentZ2 = TangentBias({MeshTangentBufferName}[VertexIndex2+1].xyzw);
+
+				float3 Binormal0   = cross(TangentZ0.xyz, TangentX0.xyz) * TangentZ0.w;
+				float3 Binormal1   = cross(TangentZ1.xyz, TangentX1.xyz) * TangentZ1.w;
+				float3 Binormal2   = cross(TangentZ2.xyz, TangentX2.xyz) * TangentZ2.w;
+
+				float3 Normal   = TangentZ0.xyz * In_Coord.BaryCoord.x + TangentZ1.xyz * In_Coord.BaryCoord.y + TangentZ2.xyz * In_Coord.BaryCoord.z;  // Normal is TangentZ
+				float3 Tangent  = TangentX0.xyz * In_Coord.BaryCoord.x + TangentX1.xyz * In_Coord.BaryCoord.y + TangentX2.xyz * In_Coord.BaryCoord.z;
+				float3 Binormal = Binormal0.xyz * In_Coord.BaryCoord.x + Binormal1.xyz * In_Coord.BaryCoord.y + Binormal2.xyz * In_Coord.BaryCoord.z;
+
+				float3 NormalWorld  = normalize(mul(float4(Normal  , 0.0), {InstanceTransformInverseTransposed}).xyz);
+				float3 TangentWorld = normalize(mul(float4(Tangent , 0.0), {InstanceTransformInverseTransposed}).xyz);
+				float3 BinormalWorld= normalize(mul(float4(Binormal, 0.0), {InstanceTransformInverseTransposed}).xyz);
+
+				Out_Normal = NormalWorld;
+				Out_Tangent = TangentWorld;
+				Out_Binormal = BinormalWorld;
+			}
+			)");
 		OutHLSL += FString::Format(FormatSample, ArgsSample);
 	}
 	else if (FunctionInfo.DefinitionName == StaticMeshHelpers::GetTriColorName)
@@ -2191,6 +2500,7 @@ bool UNiagaraDataInterfaceStaticMesh::GetFunctionHLSL(const FNiagaraDataInterfac
 			void {InstanceFunctionName} (in {MeshTriCoordinateStructName} In_Coord, out float4 Out_Color)
 			{
 				Out_Color = float4(1, 1, 1, 1);
+				[branch]
 				if ({UseColorBufferName})
 				{
 					uint TriangleIndex = In_Coord.Tri * 3;
@@ -2213,6 +2523,7 @@ bool UNiagaraDataInterfaceStaticMesh::GetFunctionHLSL(const FNiagaraDataInterfac
 		static const TCHAR *FormatSample = TEXT(R"(
 			void {InstanceFunctionName} (in {MeshTriCoordinateStructName} In_Coord, in int In_UVSet, out float2 Out_UV)
 			{
+				[branch]
 				if({NumTexCoordName}>0)
 				{
 					uint TriangleIndex = In_Coord.Tri * 3;
@@ -2241,6 +2552,14 @@ bool UNiagaraDataInterfaceStaticMesh::GetFunctionHLSL(const FNiagaraDataInterfac
 		static const TCHAR *FormatSample = TEXT(R"(
 			void {InstanceFunctionName} (in {MeshTriCoordinateStructName} In_Coord, out float3 Out_Position, out float3 Out_Velocity)
 			{
+				[branch]
+				if ({SectionCountName} == 0)
+				{
+					Out_Position = {InstanceTransformName}[3].xyz;
+					Out_Velocity = (float3)0.0f;
+					return;
+				}
+
 				uint TriangleIndex = In_Coord.Tri * 3;
 				uint VertexIndex0 = {MeshIndexBufferName}[TriangleIndex  ] * 3;
 				uint VertexIndex1 = {MeshIndexBufferName}[TriangleIndex+1] * 3;
@@ -2295,24 +2614,38 @@ bool UNiagaraDataInterfaceStaticMesh::GetFunctionHLSL(const FNiagaraDataInterfac
 	else if (FunctionInfo.DefinitionName == StaticMeshHelpers::GetVertexPositionName)
 	{
 		static const TCHAR *FormatSample = TEXT(R"(
-				void {InstanceFunctionName} (in int VertexIndex, out float3 Out_Position)
+			void {InstanceFunctionName} (in int VertexIndex, out float3 Out_Position)
+			{
+				[branch]
+				if ({SectionCountName} == 0)
 				{
-					VertexIndex *= 3;
-					Out_Position = float3({MeshVertexBufferName}[VertexIndex], {MeshVertexBufferName}[VertexIndex+1], {MeshVertexBufferName}[VertexIndex+2]);
+					Out_Position = (float3)0.0f;
+					return;
 				}
-				)");
+
+				VertexIndex *= 3;
+				Out_Position = float3({MeshVertexBufferName}[VertexIndex], {MeshVertexBufferName}[VertexIndex+1], {MeshVertexBufferName}[VertexIndex+2]);
+			}
+			)");
 		OutHLSL += FString::Format(FormatSample, ArgsSample);
 	}
 	else if (FunctionInfo.DefinitionName == StaticMeshHelpers::GetVertexPositionWSName)
 	{
 		static const TCHAR *FormatSample = TEXT(R"(
-				void {InstanceFunctionName} (in int VertexIndex, out float3 Out_Position)
+			void {InstanceFunctionName} (in int VertexIndex, out float3 Out_Position)
+			{
+				[branch]
+				if ({SectionCountName} == 0)
 				{
-					VertexIndex *= 3;
-					Out_Position = float3({MeshVertexBufferName}[VertexIndex], {MeshVertexBufferName}[VertexIndex+1], {MeshVertexBufferName}[VertexIndex+2]);
-					Out_Position = mul(float4(Out_Position, 1.0), {InstanceTransformName}).xyz;
+					Out_Position = {InstanceTransformName}[3].xyz;
+					return;
 				}
-				)");
+
+				VertexIndex *= 3;
+				Out_Position = float3({MeshVertexBufferName}[VertexIndex], {MeshVertexBufferName}[VertexIndex+1], {MeshVertexBufferName}[VertexIndex+2]);
+				Out_Position = mul(float4(Out_Position, 1.0), {InstanceTransformName}).xyz;
+			}
+			)");
 		OutHLSL += FString::Format(FormatSample, ArgsSample);
 	}
 	else
diff --git a/Engine/Plugins/FX/Niagara/Source/NiagaraCore/Private/NiagaraCustomVersion.cpp b/Engine/Plugins/FX/Niagara/Source/NiagaraCore/Private/NiagaraCustomVersion.cpp
--- a/Engine/Plugins/FX/Niagara/Source/NiagaraCore/Private/NiagaraCustomVersion.cpp
+++ b/Engine/Plugins/FX/Niagara/Source/NiagaraCore/Private/NiagaraCustomVersion.cpp
@@ -8,4 +8,4 @@ const FGuid FNiagaraCustomVersion::GUID(0xFCF57AFA, 0x50764283, 0xB9A9E658, 0xFF
 // Register the custom version with core
 FCustomVersionRegistration GRegisterNiagaraCustomVersion(FNiagaraCustomVersion::GUID, FNiagaraCustomVersion::LatestVersion, TEXT("NiagaraVer"));
 
-const FGuid FNiagaraCustomVersion::LatestScriptCompileVersion(0x069BDBD3, 0xAEC9455E, 0x8F6DB20F, 0xCA920B88);
+const FGuid FNiagaraCustomVersion::LatestScriptCompileVersion(0x3F34305D, 0xD08F4574, 0xBC50C71C, 0xA27D34DD);
