Renaming a module to its source name and then undoing won't set the display name to the source name.

#rnx

diff --git a/Engine/Plugins/FX/Niagara/Source/NiagaraEditor/Private/ViewModels/Stack/NiagaraStackEntry.cpp b/Engine/Plugins/FX/Niagara/Source/NiagaraEditor/Private/ViewModels/Stack/NiagaraStackEntry.cpp
--- a/Engine/Plugins/FX/Niagara/Source/NiagaraEditor/Private/ViewModels/Stack/NiagaraStackEntry.cpp
+++ b/Engine/Plugins/FX/Niagara/Source/NiagaraEditor/Private/ViewModels/Stack/NiagaraStackEntry.cpp
@@ -801,11 +801,6 @@ void UNiagaraStackEntry::OnRenamed(FText NewName)
 	{
 		if (!NewName.EqualTo(AlternateDisplayName.Get(FText::GetEmpty())))
 		{
-			FScopedTransaction ScopedTransaction(NSLOCTEXT("NiagaraStackEntry", "RenameModule", "Rename Module"));
-
-			GetStackEditorData().Modify();
-			GetStackEditorData().SetStackEntryDisplayName(GetStackEditorDataKey(), NewName);
-
 			if (NewName.IsEmptyOrWhitespace() || NewName.EqualTo(GetDisplayName()))
 			{
 				AlternateDisplayName.Reset();
@@ -814,6 +809,12 @@ void UNiagaraStackEntry::OnRenamed(FText NewName)
 			{
 				AlternateDisplayName = NewName;
 			}
+
+			FScopedTransaction ScopedTransaction(NSLOCTEXT("NiagaraStackEntry", "RenameModule", "Rename Module"));
+			
+			GetStackEditorData().Modify();
+			GetStackEditorData().SetStackEntryDisplayName(GetStackEditorDataKey(), AlternateDisplayName.Get(FText::GetEmpty()));
+
 			AlternateDisplayNameChangedDelegate.Broadcast();
 		}
 	}
