Fixed a parameter name in the occlusion DI

diff --git a/Engine/Plugins/FX/Niagara/Source/Niagara/Private/NiagaraDataInterfaceOcclusion.cpp b/Engine/Plugins/FX/Niagara/Source/Niagara/Private/NiagaraDataInterfaceOcclusion.cpp
--- a/Engine/Plugins/FX/Niagara/Source/Niagara/Private/NiagaraDataInterfaceOcclusion.cpp
+++ b/Engine/Plugins/FX/Niagara/Source/Niagara/Private/NiagaraDataInterfaceOcclusion.cpp
@@ -36,7 +36,7 @@ void UNiagaraDataInterfaceOcclusion::GetFunctions(TArray<FNiagaraFunctionSignatu
 	Sig.bMemberFunction = true;
 	Sig.bRequiresContext = false;
 	Sig.bSupportsCPU = false;
-	Sig.Inputs.Add(FNiagaraVariable(FNiagaraTypeDefinition(GetClass()), TEXT("Camera interface")));
+	Sig.Inputs.Add(FNiagaraVariable(FNiagaraTypeDefinition(GetClass()), TEXT("Occlusion interface")));
 	Sig.Inputs.Add(FNiagaraVariable(FNiagaraTypeDefinition::GetVec3Def(), TEXT("Sample Center World Position")));
 	Sig.Inputs.Add(FNiagaraVariable(FNiagaraTypeDefinition::GetFloatDef(), TEXT("Sample Window Width World")));
 	Sig.Inputs.Add(FNiagaraVariable(FNiagaraTypeDefinition::GetFloatDef(), TEXT("Sample Window Height World")));
@@ -53,7 +53,7 @@ void UNiagaraDataInterfaceOcclusion::GetFunctions(TArray<FNiagaraFunctionSignatu
 	Sig.bMemberFunction = true;
 	Sig.bRequiresContext = false;
 	Sig.bSupportsCPU = false;
-	Sig.Inputs.Add(FNiagaraVariable(FNiagaraTypeDefinition(GetClass()), TEXT("Camera interface")));
+	Sig.Inputs.Add(FNiagaraVariable(FNiagaraTypeDefinition(GetClass()), TEXT("Occlusion interface")));
 	Sig.Inputs.Add(FNiagaraVariable(FNiagaraTypeDefinition::GetVec3Def(), TEXT("Sample Center World Position")));
 	Sig.Inputs.Add(FNiagaraVariable(FNiagaraTypeDefinition::GetFloatDef(), TEXT("Sample Window Diameter World")));
 	Sig.Inputs.Add(FNiagaraVariable(FNiagaraTypeDefinition::GetFloatDef(), TEXT("Samples per ring")));
