Reset system simulation when setting fixed bounds so the preview window shows the new bounds even when paused

diff --git a/Engine/Plugins/FX/Niagara/Source/NiagaraEditor/Private/ViewModels/NiagaraSystemViewModel.cpp b/Engine/Plugins/FX/Niagara/Source/NiagaraEditor/Private/ViewModels/NiagaraSystemViewModel.cpp
--- a/Engine/Plugins/FX/Niagara/Source/NiagaraEditor/Private/ViewModels/NiagaraSystemViewModel.cpp
+++ b/Engine/Plugins/FX/Niagara/Source/NiagaraEditor/Private/ViewModels/NiagaraSystemViewModel.cpp
@@ -1952,6 +1952,7 @@ void FNiagaraSystemViewModel::UpdateEmitterFixedBounds()
 		}
 	}
 	PreviewComponent->MarkRenderTransformDirty();
+	ResetSystem(ETimeResetMode::KeepCurrentTime, EMultiResetMode::ResetThisInstance, EReinitMode::ResetSystem);
 }
 
 void FNiagaraSystemViewModel::AddSystemEventHandlers()
