Improved error message for missing target usage bitmask for standalone script compilation

diff --git a/Engine/Plugins/FX/Niagara/Source/NiagaraEditor/Private/NiagaraHlslTranslator.cpp b/Engine/Plugins/FX/Niagara/Source/NiagaraEditor/Private/NiagaraHlslTranslator.cpp
--- a/Engine/Plugins/FX/Niagara/Source/NiagaraEditor/Private/NiagaraHlslTranslator.cpp
+++ b/Engine/Plugins/FX/Niagara/Source/NiagaraEditor/Private/NiagaraHlslTranslator.cpp
@@ -4618,7 +4618,14 @@ void FHlslNiagaraTranslator::HandleParameterRead(int32 ParamMapHistoryIdx, const
 	FString Namespace = FNiagaraParameterMapHistory::GetNamespace(Var);
 	if (!ParamMapHistories[ParamMapHistoryIdx].IsValidNamespaceForReading(CompileOptions.TargetUsage, CompileOptions.TargetUsageBitmask, Namespace))
 	{
-		Error(FText::Format(LOCTEXT("InvalidReadingNamespace", "Variable {0} is in a namespace that isn't valid for reading"), FText::FromName(Var.GetName())), ErrorNode, nullptr);
+		if (UNiagaraScript::IsStandaloneScript(CompileOptions.TargetUsage) && Namespace.StartsWith(PARAM_MAP_ATTRIBUTE_STR))
+		{
+			Error(FText::Format(LOCTEXT("InvalidReadingNamespaceStandalone", "Variable {0} is in a namespace that isn't valid for reading. Enable at least one of the 'particle' options in the target usage bitmask of your script to access the 'Particles.' namespace."), FText::FromName(Var.GetName())), ErrorNode, nullptr);
+		}
+		else
+		{
+			Error(FText::Format(LOCTEXT("InvalidReadingNamespace", "Variable {0} is in a namespace that isn't valid for reading"), FText::FromName(Var.GetName())), ErrorNode, nullptr);
+		}
 		return;
 	}
 
