Changed the default bitmask usage to include the particle simulation stage flag

diff --git a/Engine/Plugins/FX/Niagara/Source/Niagara/Private/NiagaraScript.cpp b/Engine/Plugins/FX/Niagara/Source/Niagara/Private/NiagaraScript.cpp
--- a/Engine/Plugins/FX/Niagara/Source/Niagara/Private/NiagaraScript.cpp
+++ b/Engine/Plugins/FX/Niagara/Source/Niagara/Private/NiagaraScript.cpp
@@ -255,7 +255,7 @@ UNiagaraScript::UNiagaraScript(const FObjectInitializer& ObjectInitializer)
 	, Usage(ENiagaraScriptUsage::Function)
 #if WITH_EDITORONLY_DATA
 	, UsageIndex_DEPRECATED(0)
-	, ModuleUsageBitmask( (1 << (int32)ENiagaraScriptUsage::ParticleSpawnScript) | (1 << (int32)ENiagaraScriptUsage::ParticleSpawnScriptInterpolated) | (1 << (int32)ENiagaraScriptUsage::ParticleUpdateScript) | (1 << (int32)ENiagaraScriptUsage::ParticleEventScript) )
+	, ModuleUsageBitmask( (1 << (int32)ENiagaraScriptUsage::ParticleSpawnScript) | (1 << (int32)ENiagaraScriptUsage::ParticleSpawnScriptInterpolated) | (1 << (int32)ENiagaraScriptUsage::ParticleUpdateScript) | (1 << (int32)ENiagaraScriptUsage::ParticleEventScript) | (1 << (int32)ENiagaraScriptUsage::ParticleSimulationStageScript))
 	, NumericOutputTypeSelectionMode(ENiagaraNumericOutputTypeSelectionMode::Largest)
 #endif
 {
@@ -907,10 +907,21 @@ void UNiagaraScript::PostLoad()
 				ModuleUsageBitmask &= ~(1 << CurrentIndex);
 			}
 		}
-		// Clear the shader stage bit.
+		// Clear the simulation stage bit.
 		ModuleUsageBitmask &= ~(1 << SimulationStageIndex);
 	}
 
+	if (NiagaraVer < FNiagaraCustomVersion::SimulationStageInUsageBitmask)
+	{
+		const TArray<ENiagaraScriptUsage> SupportedUsages = GetSupportedUsageContextsForBitmask(ModuleUsageBitmask);
+		if (SupportedUsages.Contains(ENiagaraScriptUsage::ParticleUpdateScript))
+		{
+			// Set the simulation stage bit by default to true for old assets if particle update is enabled as well
+			uint8 SimulationStageIndex = (uint8)ENiagaraScriptUsage::ParticleSimulationStageScript;
+			ModuleUsageBitmask |= (1 << SimulationStageIndex);
+		}
+	}
+
 	if (Source != nullptr)
 	{
 		Source->ConditionalPostLoad();
diff --git a/Engine/Plugins/FX/Niagara/Source/NiagaraCore/Public/NiagaraCustomVersion.h b/Engine/Plugins/FX/Niagara/Source/NiagaraCore/Public/NiagaraCustomVersion.h
--- a/Engine/Plugins/FX/Niagara/Source/NiagaraCore/Public/NiagaraCustomVersion.h
+++ b/Engine/Plugins/FX/Niagara/Source/NiagaraCore/Public/NiagaraCustomVersion.h
@@ -148,6 +148,8 @@ struct FNiagaraCustomVersion
 
 		PrecompileNamespaceFixup2, // Move FNiagaraVariableMetaData from storing scope enum to storing registered scope name.
 
+		SimulationStageInUsageBitmask, // Enable the simulation stage flag by default in the usage bitmask of modules and functions
+
 		// DO NOT ADD A NEW VERSION UNLESS YOU HAVE TALKED TO THE NIAGARA LEAD. Mismanagement of these versions can lead to data loss if it is adjusted in multiple streams simultaneously.
 		// -----<new versions can be added above this line>  -------------------------------------------------
 		VersionPlusOne,
diff --git a/Engine/Plugins/FX/Niagara/Source/NiagaraEditor/Private/NiagaraNodeAssignment.cpp b/Engine/Plugins/FX/Niagara/Source/NiagaraEditor/Private/NiagaraNodeAssignment.cpp
--- a/Engine/Plugins/FX/Niagara/Source/NiagaraEditor/Private/NiagaraNodeAssignment.cpp
+++ b/Engine/Plugins/FX/Niagara/Source/NiagaraEditor/Private/NiagaraNodeAssignment.cpp
@@ -586,12 +586,14 @@ int32 UNiagaraNodeAssignment::CalculateScriptUsageBitmask()
 		}
 		if (UNiagaraScript::IsEquivalentUsage(OutputNode->GetUsage(), ENiagaraScriptUsage::ParticleSpawnScript) ||
 			UNiagaraScript::IsEquivalentUsage(OutputNode->GetUsage(), ENiagaraScriptUsage::ParticleUpdateScript) ||
-			UNiagaraScript::IsEquivalentUsage(OutputNode->GetUsage(), ENiagaraScriptUsage::ParticleEventScript))
+			UNiagaraScript::IsEquivalentUsage(OutputNode->GetUsage(), ENiagaraScriptUsage::ParticleEventScript) ||
+			UNiagaraScript::IsEquivalentUsage(OutputNode->GetUsage(), ENiagaraScriptUsage::ParticleSimulationStageScript))
 		{
 			UsageBitmask =
 				UsageToBitmask(ENiagaraScriptUsage::ParticleSpawnScript) |
 				UsageToBitmask(ENiagaraScriptUsage::ParticleUpdateScript) |
-				UsageToBitmask(ENiagaraScriptUsage::ParticleEventScript);
+				UsageToBitmask(ENiagaraScriptUsage::ParticleEventScript) |
+				UsageToBitmask(ENiagaraScriptUsage::ParticleSimulationStageScript);
 		}
 	}
 	return UsageBitmask;
