- Integate CL 11294322: When destroying data interfaces also remove any pending ticks from the instance batcher to avoid queue'd ticks reading from data that has been removed

#rnx

diff --git a/Engine/Plugins/FX/Niagara/Source/Niagara/Private/NiagaraSystemInstance.cpp b/Engine/Plugins/FX/Niagara/Source/Niagara/Private/NiagaraSystemInstance.cpp
--- a/Engine/Plugins/FX/Niagara/Source/Niagara/Private/NiagaraSystemInstance.cpp
+++ b/Engine/Plugins/FX/Niagara/Source/Niagara/Private/NiagaraSystemInstance.cpp
@@ -2148,6 +2148,19 @@ bool FNiagaraSystemInstance::GetIsolateEnabled() const
 
 void FNiagaraSystemInstance::DestroyDataInterfaceInstanceData()
 {
+	NiagaraEmitterInstanceBatcher* InstanceBatcher = GetBatcher();
+	if (bHasGPUEmitters && FNiagaraUtilities::AllowGPUParticles(InstanceBatcher->GetShaderPlatform()))
+	{
+		ENQUEUE_RENDER_COMMAND(NiagaraRemoveGPUSystem)
+		(
+			[InstanceBatcher, InstanceID=GetId()](FRHICommandListImmediate& RHICmdList) mutable
+			{
+				InstanceBatcher->InstanceDeallocated_RenderThread(InstanceID);
+			}
+		);
+	}
+
+	//
 	for (TPair<TWeakObjectPtr<UNiagaraDataInterface>, int32>& Pair : DataInterfaceInstanceDataOffsets)
 	{
 		if (UNiagaraDataInterface* Interface = Pair.Key.Get())
