Niagara - Fix for incorrectly orphaned pins in scripts.

diff --git a/Engine/Plugins/FX/Niagara/Source/NiagaraEditor/Private/NiagaraNode.cpp b/Engine/Plugins/FX/Niagara/Source/NiagaraEditor/Private/NiagaraNode.cpp
--- a/Engine/Plugins/FX/Niagara/Source/NiagaraEditor/Private/NiagaraNode.cpp
+++ b/Engine/Plugins/FX/Niagara/Source/NiagaraEditor/Private/NiagaraNode.cpp
@@ -230,11 +230,15 @@ bool UNiagaraNode::ReallocatePins(bool bMarkNeedsResynchronizeOnChange)
 		else
 		{
 			bAllSame = false;
-			if (OldPin->DefaultValue != OldPin->AutogeneratedDefaultValue || OldPin->LinkedTo.Num() > 0)
+			if ((OldPin->bDefaultValueIsIgnored == false && OldPin->DefaultValue != OldPin->AutogeneratedDefaultValue) || OldPin->LinkedTo.Num() > 0)
 			{
 				// If an old pin is not reused mark it as orphaned so that it's value and connections aren't lost.
 				OldPin->bOrphanedPin = true;
 			}
+			else
+			{
+				OldPin->bOrphanedPin = false;
+			}
 		}
 	}
 
@@ -477,7 +481,9 @@ void UNiagaraNode::OnPinRemoved(UEdGraphPin* InRemovedPin)
 
 void UNiagaraNode::NodeConnectionListChanged()
 {
-	for (UEdGraphPin* Pin : Pins)
+	// Cache the pins so that we're not modifying the collection while iterating.
+	TArray<UEdGraphPin*> CurrentPins = Pins;
+	for (UEdGraphPin* Pin : CurrentPins)
 	{
 		if (Pin->bOrphanedPin && Pin->LinkedTo.Num() == 0)
 		{
