DI functions which need instance data can now be called from custom HLSL on CPU
* the translator was adding a hidden index parameter to these calls when used from the graph, but the custom HLSL handler wasn't doing the same thing
* the parameter was inserted after the last input
* changed it so that it's the first input, rather than the last, so that the custom HLSL handler doesn't need to parse function arguments, which is complicated. It can simply replace the invocation name and following open parenthesis with the actual function name, followed by the open parenthesis, followed by the index value.
* this required modifying all DIs which use instance data to fetch the pointer from the correct slot

Fixed: the GetNumParticles function from the direct reads DI didn't work on GPU when reading from the same emitter.
* the instance count offset for the dataset is reset when a readback is scheduled, so we can't use it when setting DI compute shader params. We need to use the information stored in the sim stage data.

Added GetIDAttribute(), GetIDAttributeByIndex() and GetParticleIndex() functions to the direct reads DI.The first two allow users to read attributes of type NiagaraID (not necessarily the particle's own ID), the latter allows retrieving the index for a given ID (returns -1 if the ID cannot be found).

Added a NiagaraAny() HLSL function which emulates the any() instruction on CPU, similar to how all() is handled (because the VM doesn't implement these opcodes).

diff --git a/Engine/Plugins/Experimental/ChaosNiagara/Source/ChaosNiagara/Private/NiagaraDataInterfaceChaosDestruction.cpp b/Engine/Plugins/Experimental/ChaosNiagara/Source/ChaosNiagara/Private/NiagaraDataInterfaceChaosDestruction.cpp
--- a/Engine/Plugins/Experimental/ChaosNiagara/Source/ChaosNiagara/Private/NiagaraDataInterfaceChaosDestruction.cpp
+++ b/Engine/Plugins/Experimental/ChaosNiagara/Source/ChaosNiagara/Private/NiagaraDataInterfaceChaosDestruction.cpp
@@ -2563,83 +2563,83 @@ void UNiagaraDataInterfaceChaosDestruction::GetVMExternalFunction(const FVMExter
 {
 	if (BindingInfo.Name == GetPositionName && BindingInfo.GetNumInputs() == 2 && BindingInfo.GetNumOutputs() == 3)
 	{
-		TNDIParamBinder<0, int32, NDI_FUNC_BINDER(UNiagaraDataInterfaceChaosDestruction, GetPosition)>::Bind(this, BindingInfo, InstanceData, OutFunc);
+		TNDIParamBinder<1, int32, NDI_FUNC_BINDER(UNiagaraDataInterfaceChaosDestruction, GetPosition)>::Bind(this, BindingInfo, InstanceData, OutFunc);
 	}
 	else if (BindingInfo.Name == GetNormalName && BindingInfo.GetNumInputs() == 2 && BindingInfo.GetNumOutputs() == 3)
 	{
-		TNDIParamBinder<0, int32, NDI_FUNC_BINDER(UNiagaraDataInterfaceChaosDestruction, GetNormal)>::Bind(this, BindingInfo, InstanceData, OutFunc);
+		TNDIParamBinder<1, int32, NDI_FUNC_BINDER(UNiagaraDataInterfaceChaosDestruction, GetNormal)>::Bind(this, BindingInfo, InstanceData, OutFunc);
 	}
 	else if (BindingInfo.Name == GetVelocityName && BindingInfo.GetNumInputs() == 2 && BindingInfo.GetNumOutputs() == 3)
 	{
-		TNDIParamBinder<0, int32, NDI_FUNC_BINDER(UNiagaraDataInterfaceChaosDestruction, GetVelocity)>::Bind(this, BindingInfo, InstanceData, OutFunc);
+		TNDIParamBinder<1, int32, NDI_FUNC_BINDER(UNiagaraDataInterfaceChaosDestruction, GetVelocity)>::Bind(this, BindingInfo, InstanceData, OutFunc);
 	}
 	else if (BindingInfo.Name == GetAngularVelocityName && BindingInfo.GetNumInputs() == 2 && BindingInfo.GetNumOutputs() == 3)
 	{
-		TNDIParamBinder<0, int32, NDI_FUNC_BINDER(UNiagaraDataInterfaceChaosDestruction, GetAngularVelocity)>::Bind(this, BindingInfo, InstanceData, OutFunc);
+		TNDIParamBinder<1, int32, NDI_FUNC_BINDER(UNiagaraDataInterfaceChaosDestruction, GetAngularVelocity)>::Bind(this, BindingInfo, InstanceData, OutFunc);
 	}
 	else if (BindingInfo.Name == GetExtentMinName && BindingInfo.GetNumInputs() == 2 && BindingInfo.GetNumOutputs() == 1)
 	{
-		TNDIParamBinder<0, int32, NDI_FUNC_BINDER(UNiagaraDataInterfaceChaosDestruction, GetExtentMin)>::Bind(this, BindingInfo, InstanceData, OutFunc);
+		TNDIParamBinder<1, int32, NDI_FUNC_BINDER(UNiagaraDataInterfaceChaosDestruction, GetExtentMin)>::Bind(this, BindingInfo, InstanceData, OutFunc);
 	}
 	else if (BindingInfo.Name == GetExtentMaxName && BindingInfo.GetNumInputs() == 2 && BindingInfo.GetNumOutputs() == 1)
 	{
-		TNDIParamBinder<0, int32, NDI_FUNC_BINDER(UNiagaraDataInterfaceChaosDestruction, GetExtentMax)>::Bind(this, BindingInfo, InstanceData, OutFunc);
+		TNDIParamBinder<1, int32, NDI_FUNC_BINDER(UNiagaraDataInterfaceChaosDestruction, GetExtentMax)>::Bind(this, BindingInfo, InstanceData, OutFunc);
 	}
 	else if (BindingInfo.Name == GetVolumeName && BindingInfo.GetNumInputs() == 2 && BindingInfo.GetNumOutputs() == 1)
 	{
-		TNDIParamBinder<0, int32, NDI_FUNC_BINDER(UNiagaraDataInterfaceChaosDestruction, GetVolume)>::Bind(this, BindingInfo, InstanceData, OutFunc);
+		TNDIParamBinder<1, int32, NDI_FUNC_BINDER(UNiagaraDataInterfaceChaosDestruction, GetVolume)>::Bind(this, BindingInfo, InstanceData, OutFunc);
 	}
 	else if (BindingInfo.Name == GetParticleIdsToSpawnAtTimeName && BindingInfo.GetNumInputs() == 2 && BindingInfo.GetNumOutputs() == 3)
 	{
-		TNDIParamBinder<0, float, NDI_FUNC_BINDER(UNiagaraDataInterfaceChaosDestruction, GetParticleIdsToSpawnAtTime)>::Bind(this, BindingInfo, InstanceData, OutFunc);
+		TNDIParamBinder<1, float, NDI_FUNC_BINDER(UNiagaraDataInterfaceChaosDestruction, GetParticleIdsToSpawnAtTime)>::Bind(this, BindingInfo, InstanceData, OutFunc);
 	}
 	else if (BindingInfo.Name == GetPointTypeName && BindingInfo.GetNumInputs() == 2 && BindingInfo.GetNumOutputs() == 1)
 	{
-		TNDIParamBinder<0, int32, NDI_FUNC_BINDER(UNiagaraDataInterfaceChaosDestruction, GetPointType)>::Bind(this, BindingInfo, InstanceData, OutFunc);
+		TNDIParamBinder<1, int32, NDI_FUNC_BINDER(UNiagaraDataInterfaceChaosDestruction, GetPointType)>::Bind(this, BindingInfo, InstanceData, OutFunc);
 	}
 	else if (BindingInfo.Name == GetColorName && BindingInfo.GetNumInputs() == 2 && BindingInfo.GetNumOutputs() == 4)
 	{
-		TNDIParamBinder<0, int32, NDI_FUNC_BINDER(UNiagaraDataInterfaceChaosDestruction, GetColor)>::Bind(this, BindingInfo, InstanceData, OutFunc);
+		TNDIParamBinder<1, int32, NDI_FUNC_BINDER(UNiagaraDataInterfaceChaosDestruction, GetColor)>::Bind(this, BindingInfo, InstanceData, OutFunc);
 	}
 	else if (BindingInfo.Name == GetSolverTimeName && BindingInfo.GetNumInputs() == 1 && BindingInfo.GetNumOutputs() == 1)
 	{
-		TNDIParamBinder<0, int32, NDI_FUNC_BINDER(UNiagaraDataInterfaceChaosDestruction, GetSolverTime)>::Bind(this, BindingInfo, InstanceData, OutFunc);
+		TNDIParamBinder<1, int32, NDI_FUNC_BINDER(UNiagaraDataInterfaceChaosDestruction, GetSolverTime)>::Bind(this, BindingInfo, InstanceData, OutFunc);
 	}
 	else if (BindingInfo.Name == GetDensityName && BindingInfo.GetNumInputs() == 2 && BindingInfo.GetNumOutputs() == 1)
 	{
-		TNDIParamBinder<0, int32, NDI_FUNC_BINDER(UNiagaraDataInterfaceChaosDestruction, GetDensity)>::Bind(this, BindingInfo, InstanceData, OutFunc);
+		TNDIParamBinder<1, int32, NDI_FUNC_BINDER(UNiagaraDataInterfaceChaosDestruction, GetDensity)>::Bind(this, BindingInfo, InstanceData, OutFunc);
 	}
 	else if (BindingInfo.Name == GetFrictionName && BindingInfo.GetNumInputs() == 2 && BindingInfo.GetNumOutputs() == 1)
 	{
-		TNDIParamBinder<0, int32, NDI_FUNC_BINDER(UNiagaraDataInterfaceChaosDestruction, GetFriction)>::Bind(this, BindingInfo, InstanceData, OutFunc);
+		TNDIParamBinder<1, int32, NDI_FUNC_BINDER(UNiagaraDataInterfaceChaosDestruction, GetFriction)>::Bind(this, BindingInfo, InstanceData, OutFunc);
 	}
 	else if (BindingInfo.Name == GetRestitutionName && BindingInfo.GetNumInputs() == 2 && BindingInfo.GetNumOutputs() == 1)
 	{
-		TNDIParamBinder<0, int32, NDI_FUNC_BINDER(UNiagaraDataInterfaceChaosDestruction, GetRestitution)>::Bind(this, BindingInfo, InstanceData, OutFunc);
+		TNDIParamBinder<1, int32, NDI_FUNC_BINDER(UNiagaraDataInterfaceChaosDestruction, GetRestitution)>::Bind(this, BindingInfo, InstanceData, OutFunc);
 	}
 	else if (BindingInfo.Name == GetTransformName && BindingInfo.GetNumInputs() == 2 && BindingInfo.GetNumOutputs() == 10)
 	{
-		TNDIParamBinder<0, int32, NDI_FUNC_BINDER(UNiagaraDataInterfaceChaosDestruction, GetTransform)>::Bind(this, BindingInfo, InstanceData, OutFunc);
+		TNDIParamBinder<1, int32, NDI_FUNC_BINDER(UNiagaraDataInterfaceChaosDestruction, GetTransform)>::Bind(this, BindingInfo, InstanceData, OutFunc);
 	}
 	else if (BindingInfo.Name == GetSizeName && BindingInfo.GetNumInputs() == 2 && BindingInfo.GetNumOutputs() == 3)
 	{
-		TNDIParamBinder<0, int32, NDI_FUNC_BINDER(UNiagaraDataInterfaceChaosDestruction, GetSize)>::Bind(this, BindingInfo, InstanceData, OutFunc);
+		TNDIParamBinder<1, int32, NDI_FUNC_BINDER(UNiagaraDataInterfaceChaosDestruction, GetSize)>::Bind(this, BindingInfo, InstanceData, OutFunc);
 	}
 	else if (BindingInfo.Name == GetSurfaceTypeName && BindingInfo.GetNumInputs() == 2 && BindingInfo.GetNumOutputs() == 1)
 	{
-		TNDIParamBinder<0, int32, NDI_FUNC_BINDER(UNiagaraDataInterfaceChaosDestruction, GetSurfaceType)>::Bind(this, BindingInfo, InstanceData, OutFunc);
+		TNDIParamBinder<1, int32, NDI_FUNC_BINDER(UNiagaraDataInterfaceChaosDestruction, GetSurfaceType)>::Bind(this, BindingInfo, InstanceData, OutFunc);
 	}
 	else if (BindingInfo.Name == GetCollisionDataName && BindingInfo.GetNumInputs() == 2 && BindingInfo.GetNumOutputs() == 24)
 	{
-		TNDIParamBinder<0, int32, NDI_FUNC_BINDER(UNiagaraDataInterfaceChaosDestruction, GetCollisionData)>::Bind(this, BindingInfo, InstanceData, OutFunc);
+		TNDIParamBinder<1, int32, NDI_FUNC_BINDER(UNiagaraDataInterfaceChaosDestruction, GetCollisionData)>::Bind(this, BindingInfo, InstanceData, OutFunc);
 	}
 	else if (BindingInfo.Name == GetBreakingDataName && BindingInfo.GetNumInputs() == 2 && BindingInfo.GetNumOutputs() == 11)
 	{
-		TNDIParamBinder<0, int32, NDI_FUNC_BINDER(UNiagaraDataInterfaceChaosDestruction, GetBreakingData)>::Bind(this, BindingInfo, InstanceData, OutFunc);
+		TNDIParamBinder<1, int32, NDI_FUNC_BINDER(UNiagaraDataInterfaceChaosDestruction, GetBreakingData)>::Bind(this, BindingInfo, InstanceData, OutFunc);
 	}
 	else if (BindingInfo.Name == GetTrailingDataName && BindingInfo.GetNumInputs() == 2 && BindingInfo.GetNumOutputs() == 11)
 	{
-		TNDIParamBinder<0, int32, NDI_FUNC_BINDER(UNiagaraDataInterfaceChaosDestruction, GetTrailingData)>::Bind(this, BindingInfo, InstanceData, OutFunc);
+		TNDIParamBinder<1, int32, NDI_FUNC_BINDER(UNiagaraDataInterfaceChaosDestruction, GetTrailingData)>::Bind(this, BindingInfo, InstanceData, OutFunc);
 	}
 	else
 	{
@@ -2653,8 +2653,8 @@ void UNiagaraDataInterfaceChaosDestruction::GetVMExternalFunction(const FVMExter
 template<typename ParticleIDParamType>
 void UNiagaraDataInterfaceChaosDestruction::GetPosition(FVectorVMContext& Context)
 {
-	ParticleIDParamType ParticleIDParam(Context);
 	VectorVM::FUserPtrHandler<FNDIChaosDestruction_InstanceData> InstData(Context);
+	ParticleIDParamType ParticleIDParam(Context);
 	VectorVM::FExternalFuncRegisterHandler<float> OutSampleX(Context);
 	VectorVM::FExternalFuncRegisterHandler<float> OutSampleY(Context);
 	VectorVM::FExternalFuncRegisterHandler<float> OutSampleZ(Context);
@@ -2684,8 +2684,8 @@ void UNiagaraDataInterfaceChaosDestruction::GetPosition(FVectorVMContext& Contex
 template<typename ParticleIDParamType>
 void UNiagaraDataInterfaceChaosDestruction::GetNormal(FVectorVMContext& Context)
 {
-	ParticleIDParamType ParticleIDParam(Context);
 	VectorVM::FUserPtrHandler<FNDIChaosDestruction_InstanceData> InstData(Context);
+	ParticleIDParamType ParticleIDParam(Context);
 	VectorVM::FExternalFuncRegisterHandler<float> OutSampleX(Context);
 	VectorVM::FExternalFuncRegisterHandler<float> OutSampleY(Context);
 	VectorVM::FExternalFuncRegisterHandler<float> OutSampleZ(Context);
@@ -2714,8 +2714,8 @@ void UNiagaraDataInterfaceChaosDestruction::GetNormal(FVectorVMContext& Context)
 template<typename ParticleIDParamType>
 void UNiagaraDataInterfaceChaosDestruction::GetVelocity(FVectorVMContext& Context)
 {
-	ParticleIDParamType ParticleIDParam(Context);
 	VectorVM::FUserPtrHandler<FNDIChaosDestruction_InstanceData> InstData(Context);
+	ParticleIDParamType ParticleIDParam(Context);
 	VectorVM::FExternalFuncRegisterHandler<float> OutSampleX(Context);
 	VectorVM::FExternalFuncRegisterHandler<float> OutSampleY(Context);
 	VectorVM::FExternalFuncRegisterHandler<float> OutSampleZ(Context);
@@ -2744,8 +2744,8 @@ void UNiagaraDataInterfaceChaosDestruction::GetVelocity(FVectorVMContext& Contex
 template<typename ParticleIDParamType>
 void UNiagaraDataInterfaceChaosDestruction::GetAngularVelocity(FVectorVMContext& Context)
 {
-	ParticleIDParamType ParticleIDParam(Context);
 	VectorVM::FUserPtrHandler<FNDIChaosDestruction_InstanceData> InstData(Context);
+	ParticleIDParamType ParticleIDParam(Context);
 	VectorVM::FExternalFuncRegisterHandler<float> OutSampleX(Context);
 	VectorVM::FExternalFuncRegisterHandler<float> OutSampleY(Context);
 	VectorVM::FExternalFuncRegisterHandler<float> OutSampleZ(Context);
@@ -2774,8 +2774,8 @@ void UNiagaraDataInterfaceChaosDestruction::GetAngularVelocity(FVectorVMContext&
 template<typename ParticleIDParamType>
 void UNiagaraDataInterfaceChaosDestruction::GetExtentMin(FVectorVMContext& Context)
 {
-	ParticleIDParamType ParticleIDParam(Context);
 	VectorVM::FUserPtrHandler<FNDIChaosDestruction_InstanceData> InstData(Context);
+	ParticleIDParamType ParticleIDParam(Context);
 	VectorVM::FExternalFuncRegisterHandler<float> OutValue(Context);
 
 	for (int32 i = 0; i < Context.NumInstances; ++i)
@@ -2798,8 +2798,8 @@ void UNiagaraDataInterfaceChaosDestruction::GetExtentMin(FVectorVMContext& Conte
 template<typename ParticleIDParamType>
 void UNiagaraDataInterfaceChaosDestruction::GetExtentMax(FVectorVMContext& Context)
 {
-	ParticleIDParamType ParticleIDParam(Context);
 	VectorVM::FUserPtrHandler<FNDIChaosDestruction_InstanceData> InstData(Context);
+	ParticleIDParamType ParticleIDParam(Context);
 	VectorVM::FExternalFuncRegisterHandler<float> OutValue(Context);
 
 	for (int32 i = 0; i < Context.NumInstances; ++i)
@@ -2822,8 +2822,8 @@ void UNiagaraDataInterfaceChaosDestruction::GetExtentMax(FVectorVMContext& Conte
 template<typename ParticleIDParamType>
 void UNiagaraDataInterfaceChaosDestruction::GetVolume(FVectorVMContext& Context)
 {
-	ParticleIDParamType ParticleIDParam(Context);
 	VectorVM::FUserPtrHandler<FNDIChaosDestruction_InstanceData> InstData(Context);
+	ParticleIDParamType ParticleIDParam(Context);
 	VectorVM::FExternalFuncRegisterHandler<float> OutValue(Context);
 
 	for (int32 i = 0; i < Context.NumInstances; ++i)
@@ -2846,8 +2846,8 @@ void UNiagaraDataInterfaceChaosDestruction::GetVolume(FVectorVMContext& Context)
 template<typename TimeParamType>
 void UNiagaraDataInterfaceChaosDestruction::GetParticleIdsToSpawnAtTime(FVectorVMContext& Context)
 {
-	TimeParamType TimeParam(Context);
 	VectorVM::FUserPtrHandler<FNDIChaosDestruction_InstanceData> InstData(Context);
+	TimeParamType TimeParam(Context);
 	VectorVM::FExternalFuncRegisterHandler<int32> OutMinValue(Context);
 	VectorVM::FExternalFuncRegisterHandler<int32> OutMaxValue(Context);
 	VectorVM::FExternalFuncRegisterHandler<int32> OutCountValue(Context);
@@ -2876,8 +2876,8 @@ void UNiagaraDataInterfaceChaosDestruction::GetParticleIdsToSpawnAtTime(FVectorV
 template<typename ParticleIDParamType>
 void UNiagaraDataInterfaceChaosDestruction::GetPointType(FVectorVMContext& Context)
 {
-	ParticleIDParamType ParticleIDParam(Context);
 	VectorVM::FUserPtrHandler<FNDIChaosDestruction_InstanceData> InstData(Context);
+	ParticleIDParamType ParticleIDParam(Context);
 	VectorVM::FExternalFuncRegisterHandler<int32> OutValue(Context);
 
 	for (int32 i = 0; i < Context.NumInstances; ++i)
@@ -2901,8 +2901,8 @@ void UNiagaraDataInterfaceChaosDestruction::GetPointType(FVectorVMContext& Conte
 template<typename ParticleIDParamType>
 void UNiagaraDataInterfaceChaosDestruction::GetColor(FVectorVMContext& Context)
 {
-	ParticleIDParamType ParticleIDParam(Context);
 	VectorVM::FUserPtrHandler<FNDIChaosDestruction_InstanceData> InstData(Context);
+	ParticleIDParamType ParticleIDParam(Context);
 	VectorVM::FExternalFuncRegisterHandler<float> OutSampleR(Context);
 	VectorVM::FExternalFuncRegisterHandler<float> OutSampleG(Context);
 	VectorVM::FExternalFuncRegisterHandler<float> OutSampleB(Context);
@@ -2945,8 +2945,8 @@ void UNiagaraDataInterfaceChaosDestruction::GetSolverTime(FVectorVMContext& Cont
 template<typename ParticleIDParamType>
 void UNiagaraDataInterfaceChaosDestruction::GetDensity(FVectorVMContext& Context)
 {
-	ParticleIDParamType ParticleIDParam(Context);
 	VectorVM::FUserPtrHandler<FNDIChaosDestruction_InstanceData> InstData(Context);
+	ParticleIDParamType ParticleIDParam(Context);
 	VectorVM::FExternalFuncRegisterHandler<float> OutValue(Context);
 
 	for (int32 i = 0; i < Context.NumInstances; ++i)
@@ -2969,8 +2969,8 @@ void UNiagaraDataInterfaceChaosDestruction::GetDensity(FVectorVMContext& Context
 template<typename ParticleIDParamType>
 void UNiagaraDataInterfaceChaosDestruction::GetFriction(FVectorVMContext& Context)
 {
-	ParticleIDParamType ParticleIDParam(Context);
 	VectorVM::FUserPtrHandler<FNDIChaosDestruction_InstanceData> InstData(Context);
+	ParticleIDParamType ParticleIDParam(Context);
 	VectorVM::FExternalFuncRegisterHandler<float> OutValue(Context);
 
 	for (int32 i = 0; i < Context.NumInstances; ++i)
@@ -2993,8 +2993,8 @@ void UNiagaraDataInterfaceChaosDestruction::GetFriction(FVectorVMContext& Contex
 template<typename ParticleIDParamType>
 void UNiagaraDataInterfaceChaosDestruction::GetRestitution(FVectorVMContext& Context)
 {
-	ParticleIDParamType ParticleIDParam(Context);
 	VectorVM::FUserPtrHandler<FNDIChaosDestruction_InstanceData> InstData(Context);
+	ParticleIDParamType ParticleIDParam(Context);
 	VectorVM::FExternalFuncRegisterHandler<float> OutValue(Context);
 
 	for (int32 i = 0; i < Context.NumInstances; ++i)
@@ -3017,8 +3017,8 @@ void UNiagaraDataInterfaceChaosDestruction::GetRestitution(FVectorVMContext& Con
 template<typename ParticleIDParamType>
 void UNiagaraDataInterfaceChaosDestruction::GetTransform(FVectorVMContext& Context)
 {
-	ParticleIDParamType ParticleIDParam(Context);
 	VectorVM::FUserPtrHandler<FNDIChaosDestruction_InstanceData> InstData(Context);
+	ParticleIDParamType ParticleIDParam(Context);
 	VectorVM::FExternalFuncRegisterHandler<float> OutTranslationX(Context);
 	VectorVM::FExternalFuncRegisterHandler<float> OutTranslationY(Context);
 	VectorVM::FExternalFuncRegisterHandler<float> OutTranslationZ(Context);
@@ -3074,8 +3074,8 @@ void UNiagaraDataInterfaceChaosDestruction::GetTransform(FVectorVMContext& Conte
 template<typename ParticleIDParamType>
 void UNiagaraDataInterfaceChaosDestruction::GetSize(FVectorVMContext& Context)
 {
-	ParticleIDParamType ParticleIDParam(Context);
 	VectorVM::FUserPtrHandler<FNDIChaosDestruction_InstanceData> InstData(Context);
+	ParticleIDParamType ParticleIDParam(Context);
 	VectorVM::FExternalFuncRegisterHandler<float> OutSizeX(Context);
 	VectorVM::FExternalFuncRegisterHandler<float> OutSizeY(Context);
 	VectorVM::FExternalFuncRegisterHandler<float> OutSizeZ(Context);
@@ -3104,8 +3104,8 @@ void UNiagaraDataInterfaceChaosDestruction::GetSize(FVectorVMContext& Context)
 template<typename ParticleIDParamType>
 void UNiagaraDataInterfaceChaosDestruction::GetSurfaceType(FVectorVMContext& Context)
 {
-	ParticleIDParamType ParticleIDParam(Context);
 	VectorVM::FUserPtrHandler<FNDIChaosDestruction_InstanceData> InstData(Context);
+	ParticleIDParamType ParticleIDParam(Context);
 	VectorVM::FExternalFuncRegisterHandler<int32> OutValue(Context);
 
 	for (int32 i = 0; i < Context.NumInstances; ++i)
@@ -3128,8 +3128,8 @@ void UNiagaraDataInterfaceChaosDestruction::GetSurfaceType(FVectorVMContext& Con
 template<typename ParticleIDParamType>
 void UNiagaraDataInterfaceChaosDestruction::GetCollisionData(FVectorVMContext& Context)
 {
-	ParticleIDParamType ParticleIDParam(Context);
 	VectorVM::FUserPtrHandler<FNDIChaosDestruction_InstanceData> InstData(Context);
+	ParticleIDParamType ParticleIDParam(Context);
 	VectorVM::FExternalFuncRegisterHandler<float> OutLocationX(Context);
 	VectorVM::FExternalFuncRegisterHandler<float> OutLocationY(Context);
 	VectorVM::FExternalFuncRegisterHandler<float> OutLocationZ(Context);
@@ -3242,8 +3242,8 @@ void UNiagaraDataInterfaceChaosDestruction::GetCollisionData(FVectorVMContext& C
 template<typename ParticleIDParamType>
 void UNiagaraDataInterfaceChaosDestruction::GetBreakingData(FVectorVMContext& Context)
 {
-	ParticleIDParamType ParticleIDParam(Context);
 	VectorVM::FUserPtrHandler<FNDIChaosDestruction_InstanceData> InstData(Context);
+	ParticleIDParamType ParticleIDParam(Context);
 	VectorVM::FExternalFuncRegisterHandler<float> OutLocationX(Context);
 	VectorVM::FExternalFuncRegisterHandler<float> OutLocationY(Context);
 	VectorVM::FExternalFuncRegisterHandler<float> OutLocationZ(Context);
@@ -3307,8 +3307,8 @@ void UNiagaraDataInterfaceChaosDestruction::GetBreakingData(FVectorVMContext& Co
 template<typename ParticleIDParamType>
 void UNiagaraDataInterfaceChaosDestruction::GetTrailingData(FVectorVMContext& Context)
 {
-	ParticleIDParamType ParticleIDParam(Context);
 	VectorVM::FUserPtrHandler<FNDIChaosDestruction_InstanceData> InstData(Context);
+	ParticleIDParamType ParticleIDParam(Context);
 	VectorVM::FExternalFuncRegisterHandler<float> OutLocationX(Context);
 	VectorVM::FExternalFuncRegisterHandler<float> OutLocationY(Context);
 	VectorVM::FExternalFuncRegisterHandler<float> OutLocationZ(Context);
diff --git a/Engine/Plugins/FX/Niagara/Shaders/Private/NiagaraEmitterInstanceShader.usf b/Engine/Plugins/FX/Niagara/Shaders/Private/NiagaraEmitterInstanceShader.usf
--- a/Engine/Plugins/FX/Niagara/Shaders/Private/NiagaraEmitterInstanceShader.usf
+++ b/Engine/Plugins/FX/Niagara/Shaders/Private/NiagaraEmitterInstanceShader.usf
@@ -32,18 +32,25 @@ float3 Reciprocal(float3 x){ return 1.0f/x; }
 float2 Reciprocal(float2 x){ return 1.0f/x;}
 float Reciprocal(float x){ return 1.0f/x; }
 
-// Annoyingly, all(bool) doesn't exist, so we'll make a NiagaraAll() function which
-// works with both scalars and vectors.
+// Annoyingly, all(bool) and any(bool) don't exist, so we'll make Niagara versions which
+// work with both scalars and vectors.
 bool NiagaraAll(bool b) { return b; }
+bool NiagaraAny(bool b) { return b; }
 #if GPU_SIMULATION
 bool NiagaraAll(bool2 b) { return all(b); }
 bool NiagaraAll(bool3 b) { return all(b); }
 bool NiagaraAll(bool4 b) { return all(b); }
+bool NiagaraAny(bool2 b) { return any(b); }
+bool NiagaraAny(bool3 b) { return any(b); }
+bool NiagaraAny(bool4 b) { return any(b); }
 #else
-// No all() opcode in our VM, emulate it.
+// No all() and any() opcodes in our VM, emulate them.
 bool NiagaraAll(bool2 b) { return b.x && b.y; }
 bool NiagaraAll(bool3 b) { return b.x && b.y && b.z; }
 bool NiagaraAll(bool4 b) { return b.x && b.y && b.z && b.w; }
+bool NiagaraAny(bool2 b) { return b.x || b.y; }
+bool NiagaraAny(bool3 b) { return b.x || b.y || b.z; }
+bool NiagaraAny(bool4 b) { return b.x || b.y || b.z || b.w; }
 #endif
 
 /* -----------------------------------------------------------------
diff --git a/Engine/Plugins/FX/Niagara/Source/Niagara/Classes/NiagaraDataInterfaceParticleRead.h b/Engine/Plugins/FX/Niagara/Source/Niagara/Classes/NiagaraDataInterfaceParticleRead.h
--- a/Engine/Plugins/FX/Niagara/Source/Niagara/Classes/NiagaraDataInterfaceParticleRead.h
+++ b/Engine/Plugins/FX/Niagara/Source/Niagara/Classes/NiagaraDataInterfaceParticleRead.h
@@ -38,6 +38,7 @@ public:
 	void GetNumSpawnedParticles(FVectorVMContext& Context);
 	void GetSpawnedIDAtIndex(FVectorVMContext& Context);
 	void GetNumParticles(FVectorVMContext& Context);
+	void GetParticleIndex(FVectorVMContext& Context);
 	void ReadInt(FVectorVMContext& Context, FName AttributeToRead);
 	void ReadBool(FVectorVMContext& Context, FName AttributeToRead);
 	void ReadFloat(FVectorVMContext& Context, FName AttributeToRead);
@@ -46,6 +47,7 @@ public:
 	void ReadVector4(FVectorVMContext& Context, FName AttributeToRead);
 	void ReadColor(FVectorVMContext& Context, FName AttributeToRead);
 	void ReadQuat(FVectorVMContext& Context, FName AttributeToRead);
+	void ReadID(FVectorVMContext& Context, FName AttributeToRead);
 	void ReadIntByIndex(FVectorVMContext& Context, FName AttributeToRead);
 	void ReadBoolByIndex(FVectorVMContext& Context, FName AttributeToRead);
 	void ReadFloatByIndex(FVectorVMContext& Context, FName AttributeToRead);
@@ -54,6 +56,7 @@ public:
 	void ReadVector4ByIndex(FVectorVMContext& Context, FName AttributeToRead);
 	void ReadColorByIndex(FVectorVMContext& Context, FName AttributeToRead);
 	void ReadQuatByIndex(FVectorVMContext& Context, FName AttributeToRead);
+	void ReadIDByIndex(FVectorVMContext& Context, FName AttributeToRead);
 
 protected:
 	virtual bool CopyToInternal(UNiagaraDataInterface* Destination) const override;
diff --git a/Engine/Plugins/FX/Niagara/Source/Niagara/Private/NiagaraDataInterfaceAudioOscilloscope.cpp b/Engine/Plugins/FX/Niagara/Source/Niagara/Private/NiagaraDataInterfaceAudioOscilloscope.cpp
--- a/Engine/Plugins/FX/Niagara/Source/Niagara/Private/NiagaraDataInterfaceAudioOscilloscope.cpp
+++ b/Engine/Plugins/FX/Niagara/Source/Niagara/Private/NiagaraDataInterfaceAudioOscilloscope.cpp
@@ -141,7 +141,7 @@ float FNiagaraDataInterfaceProxyOscilloscope::SampleAudio(float NormalizedPositi
 
 void UNiagaraDataInterfaceAudioOscilloscope::SampleAudio(FVectorVMContext& Context)
 {
-	//VectorVM::FUserPtrHandler<UNiagaraDataInterfaceAudioOscilloscope> InstData(Context);
+	VectorVM::FUserPtrHandler<UNiagaraDataInterfaceAudioOscilloscope> InstData(Context);
 	VectorVM::FExternalFuncRegisterHandler<float> InNormalizedPos(Context);
 	VectorVM::FExternalFuncRegisterHandler<int32> InChannel(Context);
 	VectorVM::FExternalFuncRegisterHandler<float> OutAmplitude(Context);
diff --git a/Engine/Plugins/FX/Niagara/Source/Niagara/Private/NiagaraDataInterfaceParticleRead.cpp b/Engine/Plugins/FX/Niagara/Source/Niagara/Private/NiagaraDataInterfaceParticleRead.cpp
--- a/Engine/Plugins/FX/Niagara/Source/Niagara/Private/NiagaraDataInterfaceParticleRead.cpp
+++ b/Engine/Plugins/FX/Niagara/Source/Niagara/Private/NiagaraDataInterfaceParticleRead.cpp
@@ -12,6 +12,7 @@
 static const FName GetNumSpawnedParticlesFunctionName("Get Num Spawned Particles");
 static const FName GetSpawnedIDAtIndexFunctionName("Get Spawned ID At Index");
 static const FName GetNumParticlesFunctionName("Get Num Particles");
+static const FName GetParticleIndexFunctionName("Get Particle Index");
 
 static const FName GetIntAttributeFunctionName("Get int Attribute");
 static const FName GetBoolAttributeFunctionName("Get bool Attribute");
@@ -21,6 +22,7 @@ static const FName GetVec3AttributeFunctionName("Get Vector3 Attribute");
 static const FName GetVec4AttributeFunctionName("Get Vector4 Attribute");
 static const FName GetColorAttributeFunctionName("Get Color Attribute");
 static const FName GetQuatAttributeFunctionName("Get Quaternion Attribute");
+static const FName GetIDAttributeFunctionName("Get ID Attribute");
 
 static const FName GetIntAttributeByIndexFunctionName("Get int Attribute By Index");
 static const FName GetBoolAttributeByIndexFunctionName("Get bool Attribute By Index");
@@ -30,6 +32,7 @@ static const FName GetVec3AttributeByIndexFunctionName("Get Vector3 Attribute By
 static const FName GetVec4AttributeByIndexFunctionName("Get Vector4 Attribute By Index");
 static const FName GetColorAttributeByIndexFunctionName("Get Color Attribute By Index");
 static const FName GetQuatAttributeByIndexFunctionName("Get Quaternion Attribute By Index");
+static const FName GetIDAttributeByIndexFunctionName("Get ID Attribute By Index");
 
 static const FString NumSpawnedParticlesBaseName(TEXT("NumSpawnedParticles_"));
 static const FString SpawnedParticlesAcquireTagBaseName(TEXT("SpawnedParticlesAcquireTag_"));
@@ -47,7 +50,8 @@ enum class ENiagaraParticleDataComponentType : uint8
 {
 	Float,
 	Int,
-	Bool
+	Bool,
+	ID
 };
 
 enum class ENiagaraParticleDataValueType : uint8
@@ -60,7 +64,8 @@ enum class ENiagaraParticleDataValueType : uint8
 	Vec4,
 	Bool,
 	Color,
-	Quat
+	Quat,
+	ID
 };
 
 DECLARE_INTRINSIC_TYPE_LAYOUT(ENiagaraParticleDataValueType);
@@ -78,6 +83,7 @@ static const TCHAR* NiagaraParticleDataValueTypeName(ENiagaraParticleDataValueTy
 		case ENiagaraParticleDataValueType::Bool:		return TEXT("bool");
 		case ENiagaraParticleDataValueType::Color:		return TEXT("color");
 		case ENiagaraParticleDataValueType::Quat:		return TEXT("quaternion");
+		case ENiagaraParticleDataValueType::ID:			return TEXT("ID");
 		default:										return TEXT("UNKNOWN");
 	}
 }
@@ -178,6 +184,7 @@ struct FNiagaraDataInterfaceParametersCS_ParticleRead : public FNiagaraDataInter
 		if (FuncName == GetVec4AttributeFunctionName || FuncName == GetVec4AttributeByIndexFunctionName) return ENiagaraParticleDataValueType::Vec4;
 		if (FuncName == GetColorAttributeFunctionName || FuncName == GetColorAttributeByIndexFunctionName) return ENiagaraParticleDataValueType::Color;
 		if (FuncName == GetQuatAttributeFunctionName || FuncName == GetQuatAttributeByIndexFunctionName) return ENiagaraParticleDataValueType::Quat;
+		if (FuncName == GetIDAttributeFunctionName || FuncName == GetIDAttributeByIndexFunctionName) return ENiagaraParticleDataValueType::ID;
 		return ENiagaraParticleDataValueType::Invalid;
 	}
 
@@ -260,6 +267,7 @@ struct FNiagaraDataInterfaceParametersCS_ParticleRead : public FNiagaraDataInter
 			case ENiagaraParticleDataValueType::Vec4: return VarType == FNiagaraTypeDefinition::GetVec4Def();
 			case ENiagaraParticleDataValueType::Color: return VarType == FNiagaraTypeDefinition::GetColorDef();
 			case ENiagaraParticleDataValueType::Quat: return VarType == FNiagaraTypeDefinition::GetQuatDef();
+			case ENiagaraParticleDataValueType::ID: return VarType == FNiagaraTypeDefinition::GetIDDef();
 			default: return false;
 		}
 	}
@@ -293,7 +301,8 @@ struct FNiagaraDataInterfaceParametersCS_ParticleRead : public FNiagaraDataInter
 					if (CheckVariableType(Var.GetType(), AttributeType))
 					{
 						const FNiagaraVariableLayoutInfo& Layout = SourceEmitterVariableLayouts[VarIdx];
-						InstanceData->AttributeIndices[AttrNameIdx] = (AttributeType == ENiagaraParticleDataValueType::Int || AttributeType == ENiagaraParticleDataValueType::Bool) ? Layout.Int32ComponentStart : Layout.FloatComponentStart;
+						InstanceData->AttributeIndices[AttrNameIdx] = 
+							(AttributeType == ENiagaraParticleDataValueType::Int || AttributeType == ENiagaraParticleDataValueType::Bool || AttributeType == ENiagaraParticleDataValueType::ID) ? Layout.Int32ComponentStart : Layout.FloatComponentStart;
 					}
 					else
 					{
@@ -390,7 +399,7 @@ struct FNiagaraDataInterfaceParametersCS_ParticleRead : public FNiagaraDataInter
 		}
 
 		FNiagaraDataBuffer* SourceData;
-		uint32 NumSpawnedInstances = 0, IDAcquireTag = 0;
+		uint32 NumSpawnedInstances = 0, IDAcquireTag = 0, InstanceCountOffset = 0xffffffff;
 		const bool bReadingOwnEmitter = Context.ComputeInstanceData->Context == InstanceData->SourceEmitterGPUContext;
 		if (bReadingOwnEmitter)
 		{
@@ -398,14 +407,16 @@ struct FNiagaraDataInterfaceParametersCS_ParticleRead : public FNiagaraDataInter
 			// ourselves. We can't use SourceDataSet->GetCurrentData() in that case, because EndSimulate() has already been
 			// called on the current emitter, and the current data has been set to the destination data. We need to use the
 			// current compute instance data to get to the input buffers.
-			SourceData = Context.ComputeInstanceData->SimStageData[0].Source;
+			const FNiagaraSimStageData& SimStageData = Context.ComputeInstanceData->SimStageData[Context.SimulationStageIndex];
+			SourceData = SimStageData.Source;
+			InstanceCountOffset = SimStageData.SourceCountOffset;
 
 			// We still want to get the spawn count and ID acquire tag from the destination data, because that's where
 			// NiagaraEmitterInstanceBatcher::Run() stores them.
-			if (Context.ComputeInstanceData->SimStageData[0].Destination != nullptr)
+			if (SimStageData.Destination != nullptr)
 			{
-				NumSpawnedInstances = Context.ComputeInstanceData->SimStageData[0].Destination->GetNumSpawnedInstances();
-				IDAcquireTag = Context.ComputeInstanceData->SimStageData[0].Destination->GetIDAcquireTag();
+				NumSpawnedInstances = SimStageData.Destination->GetNumSpawnedInstances();
+				IDAcquireTag = SimStageData.Destination->GetIDAcquireTag();
 			}
 		}
 		else
@@ -415,6 +426,7 @@ struct FNiagaraDataInterfaceParametersCS_ParticleRead : public FNiagaraDataInter
 			{
 				NumSpawnedInstances = SourceData->GetNumSpawnedInstances();
 				IDAcquireTag = SourceData->GetIDAcquireTag();
+				InstanceCountOffset = SourceData->GetGPUInstanceCountBufferOffset();
 			}
 		}
 
@@ -464,7 +476,7 @@ struct FNiagaraDataInterfaceParametersCS_ParticleRead : public FNiagaraDataInter
 		const uint32 ParticleStrideFloat = SourceData->GetFloatStride() / sizeof(float);
 		const uint32 ParticleStrideInt = SourceData->GetInt32Stride() / sizeof(int32);
 
-		SetShaderValue(RHICmdList, ComputeShader, InstanceCountOffsetParam, SourceData->GetGPUInstanceCountBufferOffset());
+		SetShaderValue(RHICmdList, ComputeShader, InstanceCountOffsetParam, InstanceCountOffset);
 		SetSRVParameter(RHICmdList, ComputeShader, IDToIndexTableParam, GetIntSRVWithFallback(SourceData->GetGPUIDToIndexTable()));
 		SetSRVParameter(RHICmdList, ComputeShader, InputFloatBufferParam, GetFloatSRVWithFallback(SourceData->GetGPUBufferFloat()));
 		SetSRVParameter(RHICmdList, ComputeShader, InputIntBufferParam, GetIntSRVWithFallback(SourceData->GetGPUBufferInt()));
@@ -598,6 +610,17 @@ void UNiagaraDataInterfaceParticleRead::GetFunctions(TArray<FNiagaraFunctionSign
 		OutFunctions.Add(Sig);
 	}
 
+	{
+		FNiagaraFunctionSignature Sig;
+		Sig.Name = GetParticleIndexFunctionName;
+		Sig.Inputs.Add(FNiagaraVariable(FNiagaraTypeDefinition(GetClass()), TEXT("Particle Reader")));
+		Sig.Inputs.Add(FNiagaraVariable(FNiagaraTypeDefinition::GetIDDef(), TEXT("Particle ID")));
+		Sig.Outputs.Add(FNiagaraVariable(FNiagaraTypeDefinition::GetIntDef(), TEXT("Index")));
+		Sig.bMemberFunction = true;
+		Sig.bRequiresContext = false;
+		OutFunctions.Add(Sig);
+	}
+
 	//
 	// Get attribute by ID
 	//
@@ -729,6 +752,22 @@ void UNiagaraDataInterfaceParticleRead::GetFunctions(TArray<FNiagaraFunctionSign
 		OutFunctions.Add(Sig);
 	}
 
+	{
+		FNiagaraFunctionSignature Sig;
+		Sig.Name = GetIDAttributeFunctionName;
+		Sig.Inputs.Add(FNiagaraVariable(FNiagaraTypeDefinition(GetClass()), TEXT("Particle Reader")));
+		Sig.Inputs.Add(FNiagaraVariable(FNiagaraTypeDefinition::GetIDDef(), TEXT("Particle ID")));
+
+		Sig.Outputs.Add(FNiagaraVariable(FNiagaraTypeDefinition::GetBoolDef(), TEXT("Valid")));
+		Sig.Outputs.Add(FNiagaraVariable(FNiagaraTypeDefinition::GetIDDef(), TEXT("Value")));
+
+		Sig.FunctionSpecifiers.Add(FName("Attribute"));
+
+		Sig.bMemberFunction = true;
+		Sig.bRequiresContext = false;
+		OutFunctions.Add(Sig);
+	}
+
 	//
 	// Get attribute by index
 	//
@@ -859,11 +898,28 @@ void UNiagaraDataInterfaceParticleRead::GetFunctions(TArray<FNiagaraFunctionSign
 		Sig.bRequiresContext = false;
 		OutFunctions.Add(Sig);
 	} 
+
+	{
+		FNiagaraFunctionSignature Sig;
+		Sig.Name = GetIDAttributeByIndexFunctionName;
+		Sig.Inputs.Add(FNiagaraVariable(FNiagaraTypeDefinition(GetClass()), TEXT("Particle Reader")));
+		Sig.Inputs.Add(FNiagaraVariable(FNiagaraTypeDefinition::GetIntDef(), TEXT("Particle Index")));
+
+		Sig.Outputs.Add(FNiagaraVariable(FNiagaraTypeDefinition::GetBoolDef(), TEXT("Valid")));
+		Sig.Outputs.Add(FNiagaraVariable(FNiagaraTypeDefinition::GetIDDef(), TEXT("Value")));
+
+		Sig.FunctionSpecifiers.Add(FName("Attribute"));
+
+		Sig.bMemberFunction = true;
+		Sig.bRequiresContext = false;
+		OutFunctions.Add(Sig);
+	}
 }
 
 DEFINE_NDI_DIRECT_FUNC_BINDER(UNiagaraDataInterfaceParticleRead, GetNumSpawnedParticles);
 DEFINE_NDI_DIRECT_FUNC_BINDER(UNiagaraDataInterfaceParticleRead, GetSpawnedIDAtIndex);
 DEFINE_NDI_DIRECT_FUNC_BINDER(UNiagaraDataInterfaceParticleRead, GetNumParticles);
+DEFINE_NDI_DIRECT_FUNC_BINDER(UNiagaraDataInterfaceParticleRead, GetParticleIndex);
 DEFINE_NDI_DIRECT_FUNC_BINDER_WITH_PAYLOAD(UNiagaraDataInterfaceParticleRead, ReadInt);
 DEFINE_NDI_DIRECT_FUNC_BINDER_WITH_PAYLOAD(UNiagaraDataInterfaceParticleRead, ReadBool);
 DEFINE_NDI_DIRECT_FUNC_BINDER_WITH_PAYLOAD(UNiagaraDataInterfaceParticleRead, ReadFloat);
@@ -872,6 +928,7 @@ DEFINE_NDI_DIRECT_FUNC_BINDER_WITH_PAYLOAD(UNiagaraDataInterfaceParticleRead, Re
 DEFINE_NDI_DIRECT_FUNC_BINDER_WITH_PAYLOAD(UNiagaraDataInterfaceParticleRead, ReadVector4);
 DEFINE_NDI_DIRECT_FUNC_BINDER_WITH_PAYLOAD(UNiagaraDataInterfaceParticleRead, ReadColor);
 DEFINE_NDI_DIRECT_FUNC_BINDER_WITH_PAYLOAD(UNiagaraDataInterfaceParticleRead, ReadQuat);
+DEFINE_NDI_DIRECT_FUNC_BINDER_WITH_PAYLOAD(UNiagaraDataInterfaceParticleRead, ReadID);
 DEFINE_NDI_DIRECT_FUNC_BINDER_WITH_PAYLOAD(UNiagaraDataInterfaceParticleRead, ReadIntByIndex);
 DEFINE_NDI_DIRECT_FUNC_BINDER_WITH_PAYLOAD(UNiagaraDataInterfaceParticleRead, ReadBoolByIndex);
 DEFINE_NDI_DIRECT_FUNC_BINDER_WITH_PAYLOAD(UNiagaraDataInterfaceParticleRead, ReadFloatByIndex);
@@ -880,6 +937,7 @@ DEFINE_NDI_DIRECT_FUNC_BINDER_WITH_PAYLOAD(UNiagaraDataInterfaceParticleRead, Re
 DEFINE_NDI_DIRECT_FUNC_BINDER_WITH_PAYLOAD(UNiagaraDataInterfaceParticleRead, ReadVector4ByIndex);
 DEFINE_NDI_DIRECT_FUNC_BINDER_WITH_PAYLOAD(UNiagaraDataInterfaceParticleRead, ReadColorByIndex);
 DEFINE_NDI_DIRECT_FUNC_BINDER_WITH_PAYLOAD(UNiagaraDataInterfaceParticleRead, ReadQuatByIndex);
+DEFINE_NDI_DIRECT_FUNC_BINDER_WITH_PAYLOAD(UNiagaraDataInterfaceParticleRead, ReadIDByIndex);
 
 void UNiagaraDataInterfaceParticleRead::GetVMExternalFunction(const FVMExternalFunctionBindingInfo& BindingInfo, void* InstanceData, FVMExternalFunction& OutFunc)
 {
@@ -904,6 +962,12 @@ void UNiagaraDataInterfaceParticleRead::GetVMExternalFunction(const FVMExternalF
 		return;
 	}
 
+	if (BindingInfo.Name == GetParticleIndexFunctionName)
+	{
+		NDI_FUNC_BINDER(UNiagaraDataInterfaceParticleRead, GetParticleIndex)::Bind(this, OutFunc);
+		return;
+	}
+
 	bool bBindSuccessful = false;
 	FNDIParticleRead_InstanceData* PIData = static_cast<FNDIParticleRead_InstanceData*>(InstanceData);
 	static const FName NAME_Attribute("Attribute");
@@ -992,6 +1056,15 @@ void UNiagaraDataInterfaceParticleRead::GetVMExternalFunction(const FVMExternalF
 			bBindSuccessful = true;
 		}
 	}
+	else if (BindingInfo.Name == GetIDAttributeFunctionName)
+	{
+		FNiagaraVariable VariableToRead(FNiagaraTypeDefinition::GetIDDef(), AttributeToRead);
+		if (PIData->EmitterInstance->GetData().GetVariables().Find(VariableToRead) != INDEX_NONE)
+		{
+			NDI_FUNC_BINDER(UNiagaraDataInterfaceParticleRead, ReadID)::Bind(this, OutFunc, AttributeToRead);
+			bBindSuccessful = true;
+		}
+	}
 	//
 	// Get attribute by index
 	//
@@ -1067,6 +1140,15 @@ void UNiagaraDataInterfaceParticleRead::GetVMExternalFunction(const FVMExternalF
 			bBindSuccessful = true;
 		}
 	}
+	else if (BindingInfo.Name == GetIDAttributeByIndexFunctionName)
+	{
+	FNiagaraVariable VariableToRead(FNiagaraTypeDefinition::GetIDDef(), AttributeToRead);
+	if (PIData->EmitterInstance->GetData().GetVariables().Find(VariableToRead) != INDEX_NONE)
+	{
+		NDI_FUNC_BINDER(UNiagaraDataInterfaceParticleRead, ReadIDByIndex)::Bind(this, OutFunc, AttributeToRead);
+		bBindSuccessful = true;
+	}
+	}
 
 	if (!bBindSuccessful)
 	{
@@ -1091,9 +1173,8 @@ void UNiagaraDataInterfaceParticleRead::GetNumSpawnedParticles(FVectorVMContext&
 
 void UNiagaraDataInterfaceParticleRead::GetSpawnedIDAtIndex(FVectorVMContext& Context)
 {
-	VectorVM::FExternalFuncInputHandler<int32> InIndex(Context);
-
 	VectorVM::FUserPtrHandler<FNDIParticleRead_InstanceData> InstData(Context);
+	VectorVM::FExternalFuncInputHandler<int32> InIndex(Context);
 
 	VectorVM::FExternalFuncRegisterHandler<FNiagaraBool> OutValid(Context);
 	VectorVM::FExternalFuncRegisterHandler<int32> OutIDIndex(Context);
@@ -1144,12 +1225,56 @@ void UNiagaraDataInterfaceParticleRead::GetNumParticles(FVectorVMContext& Contex
 	}
 }
 
-void UNiagaraDataInterfaceParticleRead::ReadInt(FVectorVMContext& Context, FName AttributeToRead)
+void UNiagaraDataInterfaceParticleRead::GetParticleIndex(FVectorVMContext& Context)
 {
+	VectorVM::FUserPtrHandler<FNDIParticleRead_InstanceData> InstData(Context);
 	VectorVM::FExternalFuncInputHandler<int32> ParticleIDIndexParam(Context);
 	VectorVM::FExternalFuncInputHandler<int32> ParticleIDAcquireTagParam(Context);
+	VectorVM::FExternalFuncRegisterHandler<int32> OutIndex(Context);
+
+	const FNiagaraEmitterInstance* EmitterInstance = InstData.Get()->EmitterInstance;
+	const FNiagaraDataBuffer* CurrentData = EmitterInstance->GetData().GetCurrentData();
+
+	if (!CurrentData)
+	{
+		for (int32 InstanceIdx = 0; InstanceIdx < Context.NumInstances; ++InstanceIdx)
+		{
+			*OutIndex.GetDestAndAdvance() = -1;
+		}
+		return;
+	}
+
+	FNiagaraVariable IDVar(FNiagaraTypeDefinition::GetIDDef(), "ID");
+	FNiagaraDataSetAccessor<FNiagaraID> IDData(EmitterInstance->GetData(), IDVar);
+	const TArray<int32>& IDTable = CurrentData->GetIDTable();
+
+	for (int32 InstanceIdx = 0; InstanceIdx < Context.NumInstances; ++InstanceIdx)
+	{
+		FNiagaraID ParticleID = { ParticleIDIndexParam.GetAndAdvance(), ParticleIDAcquireTagParam.GetAndAdvance() };
+
+		int32 ParticleIndex = -1;
+		if (ParticleID.Index >= 0 && ParticleID.Index < IDTable.Num())
+		{
+			ParticleIndex = IDTable[ParticleID.Index];
+			if (ParticleIndex >= 0)
+			{
+				FNiagaraID ActualID = IDData.GetSafe(ParticleIndex, NIAGARA_INVALID_ID);
+				if (ActualID != ParticleID)
+				{
+					ParticleIndex = -1;
+				}
+			}
+		}
 
+		*OutIndex.GetDestAndAdvance() = ParticleIndex;
+	}
+}
+
+void UNiagaraDataInterfaceParticleRead::ReadInt(FVectorVMContext& Context, FName AttributeToRead)
+{
 	VectorVM::FUserPtrHandler<FNDIParticleRead_InstanceData> InstanceData(Context);
+	VectorVM::FExternalFuncInputHandler<int32> ParticleIDIndexParam(Context);
+	VectorVM::FExternalFuncInputHandler<int32> ParticleIDAcquireTagParam(Context);
 
 	VectorVM::FExternalFuncRegisterHandler<FNiagaraBool> OutValid(Context);
 	VectorVM::FExternalFuncRegisterHandler<int32> OutValue(Context);
@@ -1170,11 +1295,10 @@ void UNiagaraDataInterfaceParticleRead::ReadInt(FVectorVMContext& Context, FName
 
 void UNiagaraDataInterfaceParticleRead::ReadBool(FVectorVMContext& Context, FName AttributeToRead)
 {
+	VectorVM::FUserPtrHandler<FNDIParticleRead_InstanceData> InstanceData(Context);
 	VectorVM::FExternalFuncInputHandler<int32> ParticleIDIndexParam(Context);
 	VectorVM::FExternalFuncInputHandler<int32> ParticleIDAcquireTagParam(Context);
 
-	VectorVM::FUserPtrHandler<FNDIParticleRead_InstanceData> InstanceData(Context);
-
 	VectorVM::FExternalFuncRegisterHandler<FNiagaraBool> OutValid(Context);
 	VectorVM::FExternalFuncRegisterHandler<FNiagaraBool> OutValue(Context);
 
@@ -1194,11 +1318,10 @@ void UNiagaraDataInterfaceParticleRead::ReadBool(FVectorVMContext& Context, FNam
 
 void UNiagaraDataInterfaceParticleRead::ReadFloat(FVectorVMContext& Context, FName AttributeToRead)
 {
+	VectorVM::FUserPtrHandler<FNDIParticleRead_InstanceData> InstanceData(Context);
 	VectorVM::FExternalFuncInputHandler<int32> ParticleIDIndexParam(Context);
 	VectorVM::FExternalFuncInputHandler<int32> ParticleIDAcquireTagParam(Context);
 
-	VectorVM::FUserPtrHandler<FNDIParticleRead_InstanceData> InstanceData(Context);
-
 	VectorVM::FExternalFuncRegisterHandler<FNiagaraBool> OutValid(Context);
 	VectorVM::FExternalFuncRegisterHandler<float> OutValue(Context);
 	
@@ -1216,11 +1339,10 @@ void UNiagaraDataInterfaceParticleRead::ReadFloat(FVectorVMContext& Context, FNa
 
 void UNiagaraDataInterfaceParticleRead::ReadVector2(FVectorVMContext& Context, FName AttributeToRead)
 {
+	VectorVM::FUserPtrHandler<FNDIParticleRead_InstanceData> InstanceData(Context);
 	VectorVM::FExternalFuncInputHandler<int32> ParticleIDIndexParam(Context);
 	VectorVM::FExternalFuncInputHandler<int32> ParticleIDAcquireTagParam(Context);
 
-	VectorVM::FUserPtrHandler<FNDIParticleRead_InstanceData> InstanceData(Context);
-
 	VectorVM::FExternalFuncRegisterHandler<FNiagaraBool> OutValid(Context);
 	VectorVM::FExternalFuncRegisterHandler<float> OutX(Context);
 	VectorVM::FExternalFuncRegisterHandler<float> OutY(Context);
@@ -1242,11 +1364,10 @@ void UNiagaraDataInterfaceParticleRead::ReadVector2(FVectorVMContext& Context, F
 
 void UNiagaraDataInterfaceParticleRead::ReadVector3(FVectorVMContext& Context, FName AttributeToRead)
 {
+	VectorVM::FUserPtrHandler<FNDIParticleRead_InstanceData> InstanceData(Context);
 	VectorVM::FExternalFuncInputHandler<int32> ParticleIDIndexParam(Context);
 	VectorVM::FExternalFuncInputHandler<int32> ParticleIDAcquireTagParam(Context);
 
-	VectorVM::FUserPtrHandler<FNDIParticleRead_InstanceData> InstanceData(Context);
-
 	VectorVM::FExternalFuncRegisterHandler<FNiagaraBool> OutValid(Context);
 	VectorVM::FExternalFuncRegisterHandler<float> OutX(Context);
 	VectorVM::FExternalFuncRegisterHandler<float> OutY(Context);
@@ -1270,11 +1391,10 @@ void UNiagaraDataInterfaceParticleRead::ReadVector3(FVectorVMContext& Context, F
 
 void UNiagaraDataInterfaceParticleRead::ReadVector4(FVectorVMContext& Context, FName AttributeToRead)
 {
+	VectorVM::FUserPtrHandler<FNDIParticleRead_InstanceData> InstanceData(Context);
 	VectorVM::FExternalFuncInputHandler<int32> ParticleIDIndexParam(Context);
 	VectorVM::FExternalFuncInputHandler<int32> ParticleIDAcquireTagParam(Context);
 
-	VectorVM::FUserPtrHandler<FNDIParticleRead_InstanceData> InstanceData(Context);
-
 	VectorVM::FExternalFuncRegisterHandler<FNiagaraBool> OutValid(Context);
 	VectorVM::FExternalFuncRegisterHandler<float> OutX(Context);
 	VectorVM::FExternalFuncRegisterHandler<float> OutY(Context);
@@ -1300,11 +1420,10 @@ void UNiagaraDataInterfaceParticleRead::ReadVector4(FVectorVMContext& Context, F
 
 void UNiagaraDataInterfaceParticleRead::ReadColor(FVectorVMContext& Context, FName AttributeToRead)
 {
+	VectorVM::FUserPtrHandler<FNDIParticleRead_InstanceData> InstanceData(Context);
 	VectorVM::FExternalFuncInputHandler<int32> ParticleIDIndexParam(Context);
 	VectorVM::FExternalFuncInputHandler<int32> ParticleIDAcquireTagParam(Context);
 
-	VectorVM::FUserPtrHandler<FNDIParticleRead_InstanceData> InstanceData(Context);
-
 	VectorVM::FExternalFuncRegisterHandler<FNiagaraBool> OutValid(Context);
 	VectorVM::FExternalFuncRegisterHandler<float> OutR(Context);
 	VectorVM::FExternalFuncRegisterHandler<float> OutG(Context);
@@ -1330,11 +1449,10 @@ void UNiagaraDataInterfaceParticleRead::ReadColor(FVectorVMContext& Context, FNa
 
 void UNiagaraDataInterfaceParticleRead::ReadQuat(FVectorVMContext& Context, FName AttributeToRead)
 {
+	VectorVM::FUserPtrHandler<FNDIParticleRead_InstanceData> InstanceData(Context);
 	VectorVM::FExternalFuncInputHandler<int32> ParticleIDIndexParam(Context);
 	VectorVM::FExternalFuncInputHandler<int32> ParticleIDAcquireTagParam(Context);
 
-	VectorVM::FUserPtrHandler<FNDIParticleRead_InstanceData> InstanceData(Context);
-
 	VectorVM::FExternalFuncRegisterHandler<FNiagaraBool> OutValid(Context);
 	VectorVM::FExternalFuncRegisterHandler<float> OutX(Context);
 	VectorVM::FExternalFuncRegisterHandler<float> OutY(Context);
@@ -1358,6 +1476,31 @@ void UNiagaraDataInterfaceParticleRead::ReadQuat(FVectorVMContext& Context, FNam
 	}
 }
 
+void UNiagaraDataInterfaceParticleRead::ReadID(FVectorVMContext& Context, FName AttributeToRead)
+{
+	VectorVM::FUserPtrHandler<FNDIParticleRead_InstanceData> InstanceData(Context);
+	VectorVM::FExternalFuncInputHandler<int32> ParticleIDIndexParam(Context);
+	VectorVM::FExternalFuncInputHandler<int32> ParticleIDAcquireTagParam(Context);
+
+	VectorVM::FExternalFuncRegisterHandler<FNiagaraBool> OutValid(Context);
+	VectorVM::FExternalFuncRegisterHandler<int> OutIDIndex(Context);
+	VectorVM::FExternalFuncRegisterHandler<int> OutIDAcquireTag(Context);
+
+	for (int32 InstanceIdx = 0; InstanceIdx < Context.NumInstances; ++InstanceIdx)
+	{
+		FNiagaraID ParticleID = { ParticleIDIndexParam.GetAndAdvance(), ParticleIDAcquireTagParam.GetAndAdvance() };
+		bool bValid;
+
+		FNiagaraID Value = RetrieveValueWithCheck<FNiagaraID>(InstanceData->EmitterInstance, FNiagaraTypeDefinition::GetIDDef(), AttributeToRead, ParticleID, bValid);
+
+		FNiagaraBool ValidValue;
+		ValidValue.SetValue(bValid);
+		*OutValid.GetDestAndAdvance() = ValidValue;
+		*OutIDIndex.GetDestAndAdvance() = Value.Index;
+		*OutIDAcquireTag.GetDestAndAdvance() = Value.AcquireTag;
+	}
+}
+
 template <typename T>
 T UNiagaraDataInterfaceParticleRead::RetrieveValueWithCheck(FNiagaraEmitterInstance* EmitterInstance, const FNiagaraTypeDefinition& Type, const FName& Attr, const FNiagaraID& ParticleID, bool& bValid)
 {
@@ -1398,9 +1541,8 @@ T UNiagaraDataInterfaceParticleRead::RetrieveValueWithCheck(FNiagaraEmitterInsta
 
 void UNiagaraDataInterfaceParticleRead::ReadIntByIndex(FVectorVMContext& Context, FName AttributeToRead)
 {
-	VectorVM::FExternalFuncInputHandler<int32> ParticleIndexParam(Context);
-
 	VectorVM::FUserPtrHandler<FNDIParticleRead_InstanceData> InstanceData(Context);
+	VectorVM::FExternalFuncInputHandler<int32> ParticleIndexParam(Context);
 
 	VectorVM::FExternalFuncRegisterHandler<FNiagaraBool> OutValid(Context);
 	VectorVM::FExternalFuncRegisterHandler<int32> OutValue(Context);
@@ -1419,9 +1561,8 @@ void UNiagaraDataInterfaceParticleRead::ReadIntByIndex(FVectorVMContext& Context
 
 void UNiagaraDataInterfaceParticleRead::ReadBoolByIndex(FVectorVMContext& Context, FName AttributeToRead)
 {
-	VectorVM::FExternalFuncInputHandler<int32> ParticleIndexParam(Context);
-
 	VectorVM::FUserPtrHandler<FNDIParticleRead_InstanceData> InstanceData(Context);
+	VectorVM::FExternalFuncInputHandler<int32> ParticleIndexParam(Context);
 
 	VectorVM::FExternalFuncRegisterHandler<FNiagaraBool> OutValid(Context);
 	VectorVM::FExternalFuncRegisterHandler<FNiagaraBool> OutValue(Context);
@@ -1440,9 +1581,8 @@ void UNiagaraDataInterfaceParticleRead::ReadBoolByIndex(FVectorVMContext& Contex
 
 void UNiagaraDataInterfaceParticleRead::ReadFloatByIndex(FVectorVMContext& Context, FName AttributeToRead)
 {
-	VectorVM::FExternalFuncInputHandler<int32> ParticleIndexParam(Context);
-
 	VectorVM::FUserPtrHandler<FNDIParticleRead_InstanceData> InstanceData(Context);
+	VectorVM::FExternalFuncInputHandler<int32> ParticleIndexParam(Context);
 
 	VectorVM::FExternalFuncRegisterHandler<FNiagaraBool> OutValid(Context);
 	VectorVM::FExternalFuncRegisterHandler<float> OutValue(Context);
@@ -1461,9 +1601,8 @@ void UNiagaraDataInterfaceParticleRead::ReadFloatByIndex(FVectorVMContext& Conte
 
 void UNiagaraDataInterfaceParticleRead::ReadVector2ByIndex(FVectorVMContext& Context, FName AttributeToRead)
 {
-	VectorVM::FExternalFuncInputHandler<int32> ParticleIndexParam(Context);
-
 	VectorVM::FUserPtrHandler<FNDIParticleRead_InstanceData> InstanceData(Context);
+	VectorVM::FExternalFuncInputHandler<int32> ParticleIndexParam(Context);
 
 	VectorVM::FExternalFuncRegisterHandler<FNiagaraBool> OutValid(Context);
 	VectorVM::FExternalFuncRegisterHandler<float> OutX(Context);
@@ -1484,9 +1623,8 @@ void UNiagaraDataInterfaceParticleRead::ReadVector2ByIndex(FVectorVMContext& Con
 
 void UNiagaraDataInterfaceParticleRead::ReadVector3ByIndex(FVectorVMContext& Context, FName AttributeToRead)
 {
-	VectorVM::FExternalFuncInputHandler<int32> ParticleIndexParam(Context);
-
 	VectorVM::FUserPtrHandler<FNDIParticleRead_InstanceData> InstanceData(Context);
+	VectorVM::FExternalFuncInputHandler<int32> ParticleIndexParam(Context);
 
 	VectorVM::FExternalFuncRegisterHandler<FNiagaraBool> OutValid(Context);
 	VectorVM::FExternalFuncRegisterHandler<float> OutX(Context);
@@ -1509,9 +1647,8 @@ void UNiagaraDataInterfaceParticleRead::ReadVector3ByIndex(FVectorVMContext& Con
 
 void UNiagaraDataInterfaceParticleRead::ReadVector4ByIndex(FVectorVMContext& Context, FName AttributeToRead)
 {
-	VectorVM::FExternalFuncInputHandler<int32> ParticleIndexParam(Context);
-
 	VectorVM::FUserPtrHandler<FNDIParticleRead_InstanceData> InstanceData(Context);
+	VectorVM::FExternalFuncInputHandler<int32> ParticleIndexParam(Context);
 
 	VectorVM::FExternalFuncRegisterHandler<FNiagaraBool> OutValid(Context);
 	VectorVM::FExternalFuncRegisterHandler<float> OutX(Context);
@@ -1536,9 +1673,8 @@ void UNiagaraDataInterfaceParticleRead::ReadVector4ByIndex(FVectorVMContext& Con
 
 void UNiagaraDataInterfaceParticleRead::ReadColorByIndex(FVectorVMContext& Context, FName AttributeToRead)
 {
-	VectorVM::FExternalFuncInputHandler<int32> ParticleIndexParam(Context);
-
 	VectorVM::FUserPtrHandler<FNDIParticleRead_InstanceData> InstanceData(Context);
+	VectorVM::FExternalFuncInputHandler<int32> ParticleIndexParam(Context);
 
 	VectorVM::FExternalFuncRegisterHandler<FNiagaraBool> OutValid(Context);
 	VectorVM::FExternalFuncRegisterHandler<float> OutR(Context);
@@ -1563,9 +1699,8 @@ void UNiagaraDataInterfaceParticleRead::ReadColorByIndex(FVectorVMContext& Conte
 
 void UNiagaraDataInterfaceParticleRead::ReadQuatByIndex(FVectorVMContext& Context, FName AttributeToRead)
 {
-	VectorVM::FExternalFuncInputHandler<int32> ParticleIndexParam(Context);
-
 	VectorVM::FUserPtrHandler<FNDIParticleRead_InstanceData> InstanceData(Context);
+	VectorVM::FExternalFuncInputHandler<int32> ParticleIndexParam(Context);
 
 	VectorVM::FExternalFuncRegisterHandler<FNiagaraBool> OutValid(Context);
 	VectorVM::FExternalFuncRegisterHandler<float> OutX(Context);
@@ -1588,6 +1723,28 @@ void UNiagaraDataInterfaceParticleRead::ReadQuatByIndex(FVectorVMContext& Contex
 	}
 }
 
+void UNiagaraDataInterfaceParticleRead::ReadIDByIndex(FVectorVMContext& Context, FName AttributeToRead)
+{
+	VectorVM::FUserPtrHandler<FNDIParticleRead_InstanceData> InstanceData(Context);
+	VectorVM::FExternalFuncInputHandler<int32> ParticleIndexParam(Context);
+
+	VectorVM::FExternalFuncRegisterHandler<FNiagaraBool> OutValid(Context);
+	VectorVM::FExternalFuncRegisterHandler<int32> OutIDIndex(Context);
+	VectorVM::FExternalFuncRegisterHandler<int32> OutIDAcquireTag(Context);
+
+	for (int32 InstanceIdx = 0; InstanceIdx < Context.NumInstances; ++InstanceIdx)
+	{
+		int32 ParticleIndex = ParticleIndexParam.GetAndAdvance();
+		bool bValid;
+		FNiagaraID Value = RetrieveValueByIndexWithCheck<FNiagaraID>(InstanceData->EmitterInstance, FNiagaraTypeDefinition::GetIDDef(), AttributeToRead, ParticleIndex, bValid);
+		FNiagaraBool ValidValue;
+		ValidValue.SetValue(bValid);
+		*OutValid.GetDestAndAdvance() = ValidValue;
+		*OutIDIndex.GetDestAndAdvance() = Value.Index;
+		*OutIDAcquireTag.GetDestAndAdvance() = Value.AcquireTag;
+	}
+}
+
 template <typename T>
 T UNiagaraDataInterfaceParticleRead::RetrieveValueByIndexWithCheck(FNiagaraEmitterInstance* EmitterInstance, const FNiagaraTypeDefinition& Type, const FName& Attr, int32 ParticleIndex, bool& bValid)
 {
@@ -1711,12 +1868,14 @@ static bool GenerateGetFunctionHLSL(const FNiagaraDataInterfaceGPUParamInfo& Par
 		"    else\n"
 		"    {\n"
 		"        Out_Valid = false;\n"
-		"        Out_Value = {ValueType}(0{ExtraDefaultValues});\n"
+		"        Out_Value = ({ValueType})0;\n"
 		"    }\n"
 		"}\n\n"
 	);
 
-	static const TCHAR* ComponentNames[] = { TEXT(".x"), TEXT(".y"), TEXT(".z"), TEXT(".w") };
+	static const TCHAR* VectorComponentNames[] = { TEXT(".x"), TEXT(".y"), TEXT(".z"), TEXT(".w") };
+	static const TCHAR* IDComponentNames[] = { TEXT(".Index"), TEXT(".AcquireTag") };
+	const TCHAR** ValueComponentNames = VectorComponentNames;
 
 	const FString ParticleStrideFloatName = ParticleStrideFloatBaseName + ParamInfo.DataInterfaceHLSLSymbol;
 	const FString ParticleStrideIntName = ParticleStrideIntBaseName + ParamInfo.DataInterfaceHLSLSymbol;
@@ -1743,39 +1902,52 @@ static bool GenerateGetFunctionHLSL(const FNiagaraDataInterfaceGPUParamInfo& Par
 			InputBufferName = *InputIntBufferName;
 			InputBufferStrideName = *ParticleStrideIntName;
 			break;
+		case ENiagaraParticleDataComponentType::ID:
+			ComponentTypeName = TEXT("int");
+			InputBufferName = *InputIntBufferName;
+			InputBufferStrideName = *ParticleStrideIntName;
+			// IDs are structs with 2 components.
+			NumComponents = 2;
+			ValueComponentNames = IDComponentNames;
+			break;
 		default:
 			UE_LOG(LogNiagara, Error, TEXT("Unknown component type %d while generating function %s"), ComponentType, *FunctionInfo.InstanceName);
 			return false;
 	}
 
-	FString ExtraDefaultValues;
-	for (int ComponentIndex = 1; ComponentIndex < NumComponents; ++ComponentIndex)
-	{
-		ExtraDefaultValues += TEXT(", 0");
-	}
-
 	FString FetchValueCode;
 	for (int ComponentIndex = 0; ComponentIndex < NumComponents; ++ComponentIndex)
 	{
-		const TCHAR* ComponentName = (NumComponents > 1) ? ComponentNames[ComponentIndex] : TEXT("");
+		const TCHAR* ComponentName = (NumComponents > 1) ? ValueComponentNames[ComponentIndex] : TEXT("");
 		const FString FetchComponentCode = FString::Printf(TEXT("        Out_Value%s = %s(%s[(RegisterIndex + %d)*%s + ParticleIndex]);\n"), ComponentName, ComponentTypeName, InputBufferName, ComponentIndex, InputBufferStrideName);
 		FetchValueCode += FetchComponentCode;
 	}
 
-	FString ValueTypeName = (NumComponents > 1) ? FString::Printf(TEXT("%s%d"), ComponentTypeName, NumComponents) : FString(ComponentTypeName);
+	FString ValueTypeName;
+	if (ComponentType == ENiagaraParticleDataComponentType::ID)
+	{
+		ValueTypeName = TEXT("NiagaraID");
+	}
+	else if (NumComponents == 1)
+	{
+		ValueTypeName = ComponentTypeName;
+	}
+	else
+	{
+		ValueTypeName = FString::Printf(TEXT("%s%d"), ComponentTypeName, NumComponents);
+	}
 
 	TMap<FString, FStringFormatArg> FuncTemplateArgs;
 	FuncTemplateArgs.Add(TEXT("FunctionName"), FunctionInfo.InstanceName);
 	FuncTemplateArgs.Add(TEXT("ValueType"), ValueTypeName);
 	FuncTemplateArgs.Add(TEXT("AttributeIndicesName"), AttributeIndicesBaseName + ParamInfo.DataInterfaceHLSLSymbol);
 	FuncTemplateArgs.Add(TEXT("AttributeIndexGroup"), FunctionInstanceIndex / 4);
-	FuncTemplateArgs.Add(TEXT("AttributeIndexComponent"), ComponentNames[FunctionInstanceIndex % 4]);
+	FuncTemplateArgs.Add(TEXT("AttributeIndexComponent"), VectorComponentNames[FunctionInstanceIndex % 4]);
 	FuncTemplateArgs.Add(TEXT("IDToIndexTableName"), IDToIndexTableBaseName + ParamInfo.DataInterfaceHLSLSymbol);
 	FuncTemplateArgs.Add(TEXT("InputIntBufferName"), InputIntBufferName);
 	FuncTemplateArgs.Add(TEXT("AcquireTagRegisterIndexName"), AcquireTagRegisterIndexBaseName + ParamInfo.DataInterfaceHLSLSymbol);
 	FuncTemplateArgs.Add(TEXT("ParticleStrideIntName"), ParticleStrideIntName);
 	FuncTemplateArgs.Add(TEXT("FetchValueCode"), FetchValueCode);
-	FuncTemplateArgs.Add(TEXT("ExtraDefaultValues"), ExtraDefaultValues);
 	
 	OutHLSL += FString::Format(*FuncTemplate, FuncTemplateArgs);
 
@@ -1858,6 +2030,31 @@ bool UNiagaraDataInterfaceParticleRead::GetFunctionHLSL(const FNiagaraDataInterf
 		return true;
 	}
 
+	if (FunctionInfo.DefinitionName == GetParticleIndexFunctionName)
+	{
+		static const TCHAR* FuncTemplate = TEXT(
+			"void {FunctionName}(NiagaraID In_ParticleID, out int Out_Index)\n"
+			"{\n"
+			"    Out_Index = (In_ParticleID.Index >= 0) ? {IDToIndexTableName}[In_ParticleID.Index] : -1;\n"
+			"    int AcquireTag = (Out_Index != -1) ? {InputIntBufferName}[{AcquireTagRegisterIndexName}*{ParticleStrideIntName} + Out_Index] : 0;\n"
+			"    if(In_ParticleID.AcquireTag != AcquireTag)\n"
+			"    {\n"
+			"        Out_Index = -1;\n"
+			"    }\n"
+			"}\n\n"
+		);
+
+		TMap<FString, FStringFormatArg> FuncTemplateArgs;
+		FuncTemplateArgs.Add(TEXT("FunctionName"), FunctionInfo.InstanceName);
+		FuncTemplateArgs.Add(TEXT("IDToIndexTableName"), IDToIndexTableBaseName + ParamInfo.DataInterfaceHLSLSymbol);
+		FuncTemplateArgs.Add(TEXT("InputIntBufferName"), InputIntBufferBaseName + ParamInfo.DataInterfaceHLSLSymbol);
+		FuncTemplateArgs.Add(TEXT("AcquireTagRegisterIndexName"), AcquireTagRegisterIndexBaseName + ParamInfo.DataInterfaceHLSLSymbol);
+		FuncTemplateArgs.Add(TEXT("ParticleStrideIntName"), ParticleStrideIntBaseName + ParamInfo.DataInterfaceHLSLSymbol);
+
+		OutHLSL += FString::Format(FuncTemplate, FuncTemplateArgs);
+		return true;
+	}
+
 	//
 	// Get attribute by ID
 	//
@@ -1901,6 +2098,11 @@ bool UNiagaraDataInterfaceParticleRead::GetFunctionHLSL(const FNiagaraDataInterf
 		return GenerateGetFunctionHLSL(ParamInfo, FunctionInfo, FunctionInstanceIndex, ENiagaraParticleDataComponentType::Float, 4, false, OutHLSL);
 	}
 
+	if (FunctionInfo.DefinitionName == GetIDAttributeFunctionName)
+	{
+		return GenerateGetFunctionHLSL(ParamInfo, FunctionInfo, FunctionInstanceIndex, ENiagaraParticleDataComponentType::ID, 1, false, OutHLSL);
+	}
+
 	//
 	// Get attribute by index
 	//
@@ -1944,6 +2146,11 @@ bool UNiagaraDataInterfaceParticleRead::GetFunctionHLSL(const FNiagaraDataInterf
 		return GenerateGetFunctionHLSL(ParamInfo, FunctionInfo, FunctionInstanceIndex, ENiagaraParticleDataComponentType::Float, 4, true, OutHLSL);
 	}
 
+	if (FunctionInfo.DefinitionName == GetIDAttributeByIndexFunctionName)
+	{
+		return GenerateGetFunctionHLSL(ParamInfo, FunctionInfo, FunctionInstanceIndex, ENiagaraParticleDataComponentType::ID, 1, true, OutHLSL);
+	}
+
 	return false;
 }
 
diff --git a/Engine/Plugins/FX/Niagara/Source/Niagara/Private/NiagaraDataInterfaceSpline.cpp b/Engine/Plugins/FX/Niagara/Source/Niagara/Private/NiagaraDataInterfaceSpline.cpp
--- a/Engine/Plugins/FX/Niagara/Source/Niagara/Private/NiagaraDataInterfaceSpline.cpp
+++ b/Engine/Plugins/FX/Niagara/Source/Niagara/Private/NiagaraDataInterfaceSpline.cpp
@@ -232,57 +232,57 @@ void UNiagaraDataInterfaceSpline::GetVMExternalFunction(const FVMExternalFunctio
 	if (BindingInfo.Name == SampleSplinePositionByUnitDistanceName)
 	{
 		check(BindingInfo.GetNumInputs() == 2 && BindingInfo.GetNumOutputs() == 3);
-		TNDIExplicitBinder<FNDITransformHandlerNoop, TNDIParamBinder<0, float, NDI_FUNC_BINDER(UNiagaraDataInterfaceSpline, SampleSplinePositionByUnitDistance)>>::Bind(this, BindingInfo, InstanceData, OutFunc);
+		TNDIExplicitBinder<FNDITransformHandlerNoop, TNDIParamBinder<1, float, NDI_FUNC_BINDER(UNiagaraDataInterfaceSpline, SampleSplinePositionByUnitDistance)>>::Bind(this, BindingInfo, InstanceData, OutFunc);
 	}
 	else if (BindingInfo.Name == SampleSplinePositionByUnitDistanceWSName)
 	{
 		check(BindingInfo.GetNumInputs() == 2 && BindingInfo.GetNumOutputs() == 3);
-		TNDIExplicitBinder<FNDITransformHandler, TNDIParamBinder<0, float, NDI_FUNC_BINDER(UNiagaraDataInterfaceSpline, SampleSplinePositionByUnitDistance)>>::Bind(this, BindingInfo, InstanceData, OutFunc);
+		TNDIExplicitBinder<FNDITransformHandler, TNDIParamBinder<1, float, NDI_FUNC_BINDER(UNiagaraDataInterfaceSpline, SampleSplinePositionByUnitDistance)>>::Bind(this, BindingInfo, InstanceData, OutFunc);
 	}
 	else if (BindingInfo.Name == SampleSplineUpVectorByUnitDistanceName)
 	{
 		check(BindingInfo.GetNumInputs() == 2 && BindingInfo.GetNumOutputs() == 3);
-		TNDIExplicitBinder<FNDITransformHandlerNoop, TNDIParamBinder<0, float, NDI_FUNC_BINDER(UNiagaraDataInterfaceSpline, SampleSplineUpVectorByUnitDistance)>>::Bind(this, BindingInfo, InstanceData, OutFunc);
+		TNDIExplicitBinder<FNDITransformHandlerNoop, TNDIParamBinder<1, float, NDI_FUNC_BINDER(UNiagaraDataInterfaceSpline, SampleSplineUpVectorByUnitDistance)>>::Bind(this, BindingInfo, InstanceData, OutFunc);
 	}
 	else if (BindingInfo.Name == SampleSplineUpVectorByUnitDistanceWSName)
 	{
 		check(BindingInfo.GetNumInputs() == 2 && BindingInfo.GetNumOutputs() == 3);
-		TNDIExplicitBinder<FNDITransformHandler, TNDIParamBinder<0, float, NDI_FUNC_BINDER(UNiagaraDataInterfaceSpline, SampleSplineUpVectorByUnitDistance)>>::Bind(this, BindingInfo, InstanceData, OutFunc);
+		TNDIExplicitBinder<FNDITransformHandler, TNDIParamBinder<1, float, NDI_FUNC_BINDER(UNiagaraDataInterfaceSpline, SampleSplineUpVectorByUnitDistance)>>::Bind(this, BindingInfo, InstanceData, OutFunc);
 	}
 	else if (BindingInfo.Name == SampleSplineDirectionByUnitDistanceName)
 	{
 		check(BindingInfo.GetNumInputs() == 2 && BindingInfo.GetNumOutputs() == 3);
-		TNDIExplicitBinder<FNDITransformHandlerNoop, TNDIParamBinder<0, float, NDI_FUNC_BINDER(UNiagaraDataInterfaceSpline, SampleSplineDirectionByUnitDistance)>>::Bind(this, BindingInfo, InstanceData, OutFunc);
+		TNDIExplicitBinder<FNDITransformHandlerNoop, TNDIParamBinder<1, float, NDI_FUNC_BINDER(UNiagaraDataInterfaceSpline, SampleSplineDirectionByUnitDistance)>>::Bind(this, BindingInfo, InstanceData, OutFunc);
 	}
 	else if (BindingInfo.Name == SampleSplineDirectionByUnitDistanceWSName)
 	{
 		check(BindingInfo.GetNumInputs() == 2 && BindingInfo.GetNumOutputs() == 3);
-		TNDIExplicitBinder<FNDITransformHandler, TNDIParamBinder<0, float, NDI_FUNC_BINDER(UNiagaraDataInterfaceSpline, SampleSplineDirectionByUnitDistance)>>::Bind(this, BindingInfo, InstanceData, OutFunc);
+		TNDIExplicitBinder<FNDITransformHandler, TNDIParamBinder<1, float, NDI_FUNC_BINDER(UNiagaraDataInterfaceSpline, SampleSplineDirectionByUnitDistance)>>::Bind(this, BindingInfo, InstanceData, OutFunc);
 	}
 	else if (BindingInfo.Name == SampleSplineRightVectorByUnitDistanceName)
 	{
 		check(BindingInfo.GetNumInputs() == 2 && BindingInfo.GetNumOutputs() == 3);
-		TNDIExplicitBinder<FNDITransformHandlerNoop, TNDIParamBinder<0, float, NDI_FUNC_BINDER(UNiagaraDataInterfaceSpline, SampleSplineRightVectorByUnitDistance)>>::Bind(this, BindingInfo, InstanceData, OutFunc);
+		TNDIExplicitBinder<FNDITransformHandlerNoop, TNDIParamBinder<1, float, NDI_FUNC_BINDER(UNiagaraDataInterfaceSpline, SampleSplineRightVectorByUnitDistance)>>::Bind(this, BindingInfo, InstanceData, OutFunc);
 	}
 	else if (BindingInfo.Name == SampleSplineRightVectorByUnitDistanceWSName)
 	{
 		check(BindingInfo.GetNumInputs() == 2 && BindingInfo.GetNumOutputs() == 3);
-		TNDIExplicitBinder<FNDITransformHandler, TNDIParamBinder<0, float, NDI_FUNC_BINDER(UNiagaraDataInterfaceSpline, SampleSplineRightVectorByUnitDistance)>>::Bind(this, BindingInfo, InstanceData, OutFunc);
+		TNDIExplicitBinder<FNDITransformHandler, TNDIParamBinder<1, float, NDI_FUNC_BINDER(UNiagaraDataInterfaceSpline, SampleSplineRightVectorByUnitDistance)>>::Bind(this, BindingInfo, InstanceData, OutFunc);
 	}
 	else if (BindingInfo.Name == SampleSplineTangentByUnitDistanceName)
 	{
 		check(BindingInfo.GetNumInputs() == 2 && BindingInfo.GetNumOutputs() == 3);
-		TNDIExplicitBinder<FNDITransformHandlerNoop, TNDIParamBinder<0, float, NDI_FUNC_BINDER(UNiagaraDataInterfaceSpline, SampleSplineTangentByUnitDistance)>>::Bind(this, BindingInfo, InstanceData, OutFunc);
+		TNDIExplicitBinder<FNDITransformHandlerNoop, TNDIParamBinder<1, float, NDI_FUNC_BINDER(UNiagaraDataInterfaceSpline, SampleSplineTangentByUnitDistance)>>::Bind(this, BindingInfo, InstanceData, OutFunc);
 	}
 	else if (BindingInfo.Name == SampleSplineTangentByUnitDistanceWSName)
 	{
 		check(BindingInfo.GetNumInputs() == 2 && BindingInfo.GetNumOutputs() == 3);
-		TNDIExplicitBinder<FNDITransformHandler, TNDIParamBinder<0, float, NDI_FUNC_BINDER(UNiagaraDataInterfaceSpline, SampleSplineTangentByUnitDistance)>>::Bind(this, BindingInfo, InstanceData, OutFunc);
+		TNDIExplicitBinder<FNDITransformHandler, TNDIParamBinder<1, float, NDI_FUNC_BINDER(UNiagaraDataInterfaceSpline, SampleSplineTangentByUnitDistance)>>::Bind(this, BindingInfo, InstanceData, OutFunc);
 	}
 	else if (BindingInfo.Name == FindClosestUnitDistanceFromPositionWSName)
 	{
 		check(BindingInfo.GetNumInputs() == 4 && BindingInfo.GetNumOutputs() == 1);
-		TNDIParamBinder<0, float, TNDIParamBinder<1, float, TNDIParamBinder<2, float, NDI_FUNC_BINDER(UNiagaraDataInterfaceSpline, FindClosestUnitDistanceFromPositionWS)>>>::Bind(this, BindingInfo, InstanceData, OutFunc);
+		TNDIParamBinder<1, float, TNDIParamBinder<2, float, TNDIParamBinder<3, float, NDI_FUNC_BINDER(UNiagaraDataInterfaceSpline, FindClosestUnitDistanceFromPositionWS)>>>::Bind(this, BindingInfo, InstanceData, OutFunc);
 	}
 	else if (BindingInfo.Name == GetSplineLocalToWorldName)
 	{
@@ -381,9 +381,9 @@ bool UNiagaraDataInterfaceSpline::PerInstanceTick(void* PerInstanceData, FNiagar
 template<typename TransformHandlerType, typename SplineSampleType>
 void UNiagaraDataInterfaceSpline::SampleSplinePositionByUnitDistance(FVectorVMContext& Context)
 {
+	VectorVM::FUserPtrHandler<FNDISpline_InstanceData> InstData(Context);
 	TransformHandlerType TransformHandler;
 	SplineSampleType SplineSampleParam(Context);
-	VectorVM::FUserPtrHandler<FNDISpline_InstanceData> InstData(Context);
 	VectorVM::FExternalFuncRegisterHandler<float> OutPosX(Context);
 	VectorVM::FExternalFuncRegisterHandler<float> OutPosY(Context);
 	VectorVM::FExternalFuncRegisterHandler<float> OutPosZ(Context);
@@ -430,9 +430,9 @@ void UNiagaraDataInterfaceSpline::SampleSplinePositionByUnitDistance(FVectorVMCo
 template<typename TransformHandlerType, typename SplineSampleType>
 void UNiagaraDataInterfaceSpline::SampleSplineUpVectorByUnitDistance(FVectorVMContext& Context)
 {
+	VectorVM::FUserPtrHandler<FNDISpline_InstanceData> InstData(Context);
 	TransformHandlerType TransformHandler;
 	SplineSampleType SplineSampleParam(Context);
-	VectorVM::FUserPtrHandler<FNDISpline_InstanceData> InstData(Context);
 	VectorVM::FExternalFuncRegisterHandler<float> OutPosX(Context);
 	VectorVM::FExternalFuncRegisterHandler<float> OutPosY(Context);
 	VectorVM::FExternalFuncRegisterHandler<float> OutPosZ(Context);
@@ -479,9 +479,9 @@ void UNiagaraDataInterfaceSpline::SampleSplineUpVectorByUnitDistance(FVectorVMCo
 template<typename TransformHandlerType, typename SplineSampleType>
 void UNiagaraDataInterfaceSpline::SampleSplineRightVectorByUnitDistance(FVectorVMContext& Context)
 {
+	VectorVM::FUserPtrHandler<FNDISpline_InstanceData> InstData(Context);
 	TransformHandlerType TransformHandler;
 	SplineSampleType SplineSampleParam(Context);
-	VectorVM::FUserPtrHandler<FNDISpline_InstanceData> InstData(Context);
 	VectorVM::FExternalFuncRegisterHandler<float> OutPosX(Context);
 	VectorVM::FExternalFuncRegisterHandler<float> OutPosY(Context);
 	VectorVM::FExternalFuncRegisterHandler<float> OutPosZ(Context);
@@ -526,9 +526,9 @@ void UNiagaraDataInterfaceSpline::SampleSplineRightVectorByUnitDistance(FVectorV
 template<typename TransformHandlerType, typename SplineSampleType>
 void UNiagaraDataInterfaceSpline::SampleSplineTangentByUnitDistance(FVectorVMContext& Context)
 {
+	VectorVM::FUserPtrHandler<FNDISpline_InstanceData> InstData(Context);
 	TransformHandlerType TransformHandler;
 	SplineSampleType SplineSampleParam(Context);
-	VectorVM::FUserPtrHandler<FNDISpline_InstanceData> InstData(Context);
 	VectorVM::FExternalFuncRegisterHandler<float> OutPosX(Context);
 	VectorVM::FExternalFuncRegisterHandler<float> OutPosY(Context);
 	VectorVM::FExternalFuncRegisterHandler<float> OutPosZ(Context);
@@ -574,9 +574,9 @@ void UNiagaraDataInterfaceSpline::SampleSplineTangentByUnitDistance(FVectorVMCon
 template<typename TransformHandlerType, typename SplineSampleType>
 void UNiagaraDataInterfaceSpline::SampleSplineDirectionByUnitDistance(FVectorVMContext& Context)
 {
+	VectorVM::FUserPtrHandler<FNDISpline_InstanceData> InstData(Context);
 	TransformHandlerType TransformHandler;
 	SplineSampleType SplineSampleParam(Context);
-	VectorVM::FUserPtrHandler<FNDISpline_InstanceData> InstData(Context);
 	VectorVM::FExternalFuncRegisterHandler<float> OutPosX(Context);
 	VectorVM::FExternalFuncRegisterHandler<float> OutPosY(Context);
 	VectorVM::FExternalFuncRegisterHandler<float> OutPosZ(Context);
@@ -662,10 +662,10 @@ void UNiagaraDataInterfaceSpline::WriteTransform(const FMatrix& ToWrite, FVector
 template<typename PosXType, typename PosYType, typename PosZType>
 void UNiagaraDataInterfaceSpline::FindClosestUnitDistanceFromPositionWS(FVectorVMContext& Context)
 {
+	VectorVM::FUserPtrHandler<FNDISpline_InstanceData> InstData(Context);
 	PosXType PosXParam(Context);
 	PosYType PosYParam(Context);
 	PosZType PosZParam(Context);
-	VectorVM::FUserPtrHandler<FNDISpline_InstanceData> InstData(Context);
 	VectorVM::FExternalFuncRegisterHandler<float> OutUnitDistance(Context);
 
 	if (USplineComponent* SplineComponent = InstData->Component.Get())
diff --git a/Engine/Plugins/FX/Niagara/Source/Niagara/Private/NiagaraDataInterfaceStaticMesh.cpp b/Engine/Plugins/FX/Niagara/Source/Niagara/Private/NiagaraDataInterfaceStaticMesh.cpp
--- a/Engine/Plugins/FX/Niagara/Source/Niagara/Private/NiagaraDataInterfaceStaticMesh.cpp
+++ b/Engine/Plugins/FX/Niagara/Source/Niagara/Private/NiagaraDataInterfaceStaticMesh.cpp
@@ -1458,9 +1458,9 @@ void UNiagaraDataInterfaceStaticMesh::RandomTriCoord<TSampleModeInvalid>(FVector
 
 void UNiagaraDataInterfaceStaticMesh::RandomTriCoordVertexColorFiltered(FVectorVMContext& Context)
 {
+	VectorVM::FUserPtrHandler<FNDIStaticMesh_InstanceData> InstData(Context);
 	VectorVM::FExternalFuncRegisterHandler<int32> MinValue(Context);
 	VectorVM::FExternalFuncRegisterHandler<int32> RangeValue(Context);
-	VectorVM::FUserPtrHandler<FNDIStaticMesh_InstanceData> InstData(Context);
 
 	VectorVM::FExternalFuncRegisterHandler<int32> OutTri(Context);
 	VectorVM::FExternalFuncRegisterHandler<float> OutBaryX(Context);
@@ -1552,8 +1552,8 @@ FORCEINLINE int32 UNiagaraDataInterfaceStaticMesh::RandomTriIndexOnSection<TSamp
 template<typename TSampleMode>
 void UNiagaraDataInterfaceStaticMesh::RandomTriCoordOnSection(FVectorVMContext& Context)
 {
-	VectorVM::FExternalFuncInputHandler<int32> SectionIdxParam(Context);
 	VectorVM::FUserPtrHandler<FNDIStaticMesh_InstanceData> InstData(Context);
+	VectorVM::FExternalFuncInputHandler<int32> SectionIdxParam(Context);
 
 	VectorVM::FExternalFuncRegisterHandler<int32> OutTri(Context);
 	VectorVM::FExternalFuncRegisterHandler<float> OutBaryX(Context);
@@ -1586,8 +1586,8 @@ void UNiagaraDataInterfaceStaticMesh::RandomTriCoordOnSection(FVectorVMContext&
 template<>
 void UNiagaraDataInterfaceStaticMesh::RandomTriCoordOnSection<TSampleModeInvalid>(FVectorVMContext& Context)
 {
-	VectorVM::FExternalFuncInputHandler<int32> SectionIdxParam(Context);
 	VectorVM::FUserPtrHandler<FNDIStaticMesh_InstanceData> InstData(Context);
+	VectorVM::FExternalFuncInputHandler<int32> SectionIdxParam(Context);
 
 	VectorVM::FExternalFuncRegisterHandler<int32> OutTri(Context);
 	VectorVM::FExternalFuncRegisterHandler<float> OutBaryX(Context);
@@ -1606,12 +1606,12 @@ void UNiagaraDataInterfaceStaticMesh::RandomTriCoordOnSection<TSampleModeInvalid
 template<typename TransformHandlerType>
 void UNiagaraDataInterfaceStaticMesh::GetTriCoordPosition(FVectorVMContext& Context)
 {
+	VectorVM::FUserPtrHandler<FNDIStaticMesh_InstanceData> InstData(Context);
 	TransformHandlerType TransformHandler;
 	VectorVM::FExternalFuncInputHandler<int32> TriParam(Context);
 	VectorVM::FExternalFuncInputHandler<float> BaryXParam(Context);
 	VectorVM::FExternalFuncInputHandler<float> BaryYParam(Context);
 	VectorVM::FExternalFuncInputHandler<float> BaryZParam(Context);
-	VectorVM::FUserPtrHandler<FNDIStaticMesh_InstanceData> InstData(Context);
 
 	VectorVM::FExternalFuncRegisterHandler<float> OutPosX(Context);
 	VectorVM::FExternalFuncRegisterHandler<float> OutPosY(Context);
@@ -1665,13 +1665,14 @@ void UNiagaraDataInterfaceStaticMesh::GetTriCoordPosition(FVectorVMContext& Cont
 template<typename TransformHandlerType>
 void UNiagaraDataInterfaceStaticMesh::GetTriCoordNormal(FVectorVMContext& Context)
 {
+	VectorVM::FUserPtrHandler<FNDIStaticMesh_InstanceData> InstData(Context);
+
 	TransformHandlerType TransformHandler;
 
 	VectorVM::FExternalFuncInputHandler<int32> TriParam(Context);
 	VectorVM::FExternalFuncInputHandler<float> BaryXParam(Context);
 	VectorVM::FExternalFuncInputHandler<float> BaryYParam(Context);
 	VectorVM::FExternalFuncInputHandler<float> BaryZParam(Context);
-	VectorVM::FUserPtrHandler<FNDIStaticMesh_InstanceData> InstData(Context);
 
 	VectorVM::FExternalFuncRegisterHandler<float> OutNormX(Context);
 	VectorVM::FExternalFuncRegisterHandler<float> OutNormY(Context);
@@ -1721,13 +1722,14 @@ void UNiagaraDataInterfaceStaticMesh::GetTriCoordNormal(FVectorVMContext& Contex
 template<typename VertexAccessorType, typename TransformHandlerType>
 void UNiagaraDataInterfaceStaticMesh::GetTriCoordTangents(FVectorVMContext& Context)
 {
-	TransformHandlerType TransformHandler;	
+	VectorVM::FUserPtrHandler<FNDIStaticMesh_InstanceData> InstData(Context);
+
+	TransformHandlerType TransformHandler;
 
 	VectorVM::FExternalFuncInputHandler<int32> TriParam(Context);
 	VectorVM::FExternalFuncInputHandler<float> BaryXParam(Context);
 	VectorVM::FExternalFuncInputHandler<float> BaryYParam(Context);
 	VectorVM::FExternalFuncInputHandler<float> BaryZParam(Context);
-	VectorVM::FUserPtrHandler<FNDIStaticMesh_InstanceData> InstData(Context);
 
 	VectorVM::FExternalFuncRegisterHandler<float> OutTangentX(Context);
 	VectorVM::FExternalFuncRegisterHandler<float> OutTangentY(Context);
@@ -1802,11 +1804,12 @@ void UNiagaraDataInterfaceStaticMesh::GetTriCoordTangents(FVectorVMContext& Cont
 
 void UNiagaraDataInterfaceStaticMesh::GetTriCoordColor(FVectorVMContext& Context)
 {
+	VectorVM::FUserPtrHandler<FNDIStaticMesh_InstanceData> InstData(Context);
+
 	VectorVM::FExternalFuncInputHandler<int32> TriParam(Context);
 	VectorVM::FExternalFuncInputHandler<float> BaryXParam(Context);
 	VectorVM::FExternalFuncInputHandler<float> BaryYParam(Context);
 	VectorVM::FExternalFuncInputHandler<float> BaryZParam(Context);
-	VectorVM::FUserPtrHandler<FNDIStaticMesh_InstanceData> InstData(Context);
 
 	VectorVM::FExternalFuncRegisterHandler<float> OutColorR(Context);
 	VectorVM::FExternalFuncRegisterHandler<float> OutColorG(Context);
@@ -1865,12 +1868,13 @@ void UNiagaraDataInterfaceStaticMesh::GetTriCoordColor(FVectorVMContext& Context
 template<typename VertexAccessorType>
 void UNiagaraDataInterfaceStaticMesh::GetTriCoordUV(FVectorVMContext& Context)
 {
+	VectorVM::FUserPtrHandler<FNDIStaticMesh_InstanceData> InstData(Context);
+
 	VectorVM::FExternalFuncInputHandler<int32> TriParam(Context);
 	VectorVM::FExternalFuncInputHandler<float> BaryXParam(Context);
 	VectorVM::FExternalFuncInputHandler<float> BaryYParam(Context);
 	VectorVM::FExternalFuncInputHandler<float> BaryZParam(Context);
 	VectorVM::FExternalFuncInputHandler<int32> UVSetParam(Context);
-	VectorVM::FUserPtrHandler<FNDIStaticMesh_InstanceData> InstData(Context);
 
 	VectorVM::FExternalFuncRegisterHandler<float> OutU(Context);
 	VectorVM::FExternalFuncRegisterHandler<float> OutV(Context);
@@ -1917,11 +1921,12 @@ void UNiagaraDataInterfaceStaticMesh::GetTriCoordUV(FVectorVMContext& Context)
 
 void UNiagaraDataInterfaceStaticMesh::GetTriCoordPositionAndVelocity(FVectorVMContext& Context)
 {
+	VectorVM::FUserPtrHandler<FNDIStaticMesh_InstanceData> InstData(Context);
+
 	VectorVM::FExternalFuncInputHandler<int32> TriParam(Context);
 	VectorVM::FExternalFuncInputHandler<float> BaryXParam(Context);
 	VectorVM::FExternalFuncInputHandler<float> BaryYParam(Context);
 	VectorVM::FExternalFuncInputHandler<float> BaryZParam(Context);
-	VectorVM::FUserPtrHandler<FNDIStaticMesh_InstanceData> InstData(Context);
 
 	VectorVM::FExternalFuncRegisterHandler<float> OutPosX(Context);
 	VectorVM::FExternalFuncRegisterHandler<float> OutPosY(Context);
@@ -2067,9 +2072,10 @@ void UNiagaraDataInterfaceStaticMesh::GetWorldVelocity(FVectorVMContext& Context
 template<typename TransformHandlerType>
 void UNiagaraDataInterfaceStaticMesh::GetVertexPosition(FVectorVMContext& Context)
 {
+	VectorVM::FUserPtrHandler<FNDIStaticMesh_InstanceData> InstData(Context);
+
 	TransformHandlerType TransformHandler;
 	VectorVM::FExternalFuncInputHandler<int32> VertexIndexParam(Context);
-	VectorVM::FUserPtrHandler<FNDIStaticMesh_InstanceData> InstData(Context);
 
 	VectorVM::FExternalFuncRegisterHandler<float> OutPosX(Context);
 	VectorVM::FExternalFuncRegisterHandler<float> OutPosY(Context);
diff --git a/Engine/Plugins/FX/Niagara/Source/NiagaraCore/Private/NiagaraCustomVersion.cpp b/Engine/Plugins/FX/Niagara/Source/NiagaraCore/Private/NiagaraCustomVersion.cpp
--- a/Engine/Plugins/FX/Niagara/Source/NiagaraCore/Private/NiagaraCustomVersion.cpp
+++ b/Engine/Plugins/FX/Niagara/Source/NiagaraCore/Private/NiagaraCustomVersion.cpp
@@ -8,4 +8,4 @@ const FGuid FNiagaraCustomVersion::GUID(0xFCF57AFA, 0x50764283, 0xB9A9E658, 0xFF
 // Register the custom version with core
 FCustomVersionRegistration GRegisterNiagaraCustomVersion(FNiagaraCustomVersion::GUID, FNiagaraCustomVersion::LatestVersion, TEXT("NiagaraVer"));
 
-const FGuid FNiagaraCustomVersion::LatestScriptCompileVersion(0x3F34305D, 0xD08F4574, 0xBC50C71C, 0xA27D34DD);
+const FGuid FNiagaraCustomVersion::LatestScriptCompileVersion(0x5FD3FF04, 0xC5BF4260, 0xB48F8499, 0xB449A6D5);
diff --git a/Engine/Plugins/FX/Niagara/Source/NiagaraEditor/Private/NiagaraHlslTranslator.cpp b/Engine/Plugins/FX/Niagara/Source/NiagaraEditor/Private/NiagaraHlslTranslator.cpp
--- a/Engine/Plugins/FX/Niagara/Source/NiagaraEditor/Private/NiagaraHlslTranslator.cpp
+++ b/Engine/Plugins/FX/Niagara/Source/NiagaraEditor/Private/NiagaraHlslTranslator.cpp
@@ -5676,8 +5676,31 @@ void FHlslNiagaraTranslator::HandleCustomHlslNode(UNiagaraNodeCustomHlsl* Custom
 
 						if (Info.UserPtrIdx != INDEX_NONE && CompilationTarget != ENiagaraSimTarget::GPUComputeSim)
 						{
-							//This interface requires per instance data via a user ptr so place the index to it at the end of the inputs.
-							Sig.Inputs.Add(FNiagaraVariable(FNiagaraTypeDefinition::GetIntDef(), TEXT("InstanceData")));
+							//This interface requires per instance data via a user ptr so place the index as the first input.
+							Sig.Inputs.Insert(FNiagaraVariable(FNiagaraTypeDefinition::GetIntDef(), TEXT("InstanceData")), 0);
+
+							// Look for the opening parenthesis.
+							while (TokenIndex < Tokens.Num() && Tokens[TokenIndex] != TEXT("("))
+							{
+								++TokenIndex;
+							}
+
+							if (TokenIndex < Tokens.Num())
+							{
+								// Skip the parenthesis.
+								++TokenIndex;
+
+								// Insert the instance index as the first argument. We don't need to do range checking because even if
+								// the tokens end after the parenthesis, we'll be inserting at the end of the array.
+								Tokens.Insert(LexToString(Info.UserPtrIdx), TokenIndex++);
+
+								if (Sig.Inputs.Num() > 1 || Sig.Outputs.Num() > 0)
+								{
+									// If there are other arguments, insert a comma and a space. These are separators, so they need to be different tokens.
+									Tokens.Insert(TEXT(","), TokenIndex++);
+									Tokens.Insert(TEXT(" "), TokenIndex++);
+								}
+							}
 						}
 
 						Info.RegisteredFunctions.Add(Sig);
@@ -6202,9 +6225,9 @@ void FHlslNiagaraTranslator::RegisterFunctionCall(ENiagaraScriptUsage ScriptUsag
 
 				if (Info.UserPtrIdx != INDEX_NONE && CompilationTarget != ENiagaraSimTarget::GPUComputeSim)
 				{
-					//This interface requires per instance data via a user ptr so place the index to it at the end of the inputs.
-					Inputs.Add(AddSourceChunk(LexToString(Info.UserPtrIdx), FNiagaraTypeDefinition::GetIntDef(), false));
-					OutSignature.Inputs.Add(FNiagaraVariable(FNiagaraTypeDefinition::GetIntDef(), TEXT("InstanceData")));
+					//This interface requires per instance data via a user ptr so place the index as the first input.
+					Inputs.Insert(AddSourceChunk(LexToString(Info.UserPtrIdx), FNiagaraTypeDefinition::GetIntDef(), false), 0);
+					OutSignature.Inputs.Insert(FNiagaraVariable(FNiagaraTypeDefinition::GetIntDef(), TEXT("InstanceData")), 0);
 				}
 			}
 
