Fixed GetID and GetIDByIndex functions to compile with the data conversion framework.
Cleaned up HLSL generation (no need to test the compressed flag for ints, bools and IDs).

#rnx

diff --git a/Engine/Plugins/FX/Niagara/Source/Niagara/Classes/NiagaraDataSet.h b/Engine/Plugins/FX/Niagara/Source/Niagara/Classes/NiagaraDataSet.h
--- a/Engine/Plugins/FX/Niagara/Source/Niagara/Classes/NiagaraDataSet.h
+++ b/Engine/Plugins/FX/Niagara/Source/Niagara/Classes/NiagaraDataSet.h
@@ -2005,6 +2005,8 @@ struct FNiagaraDataConversions
 	struct FNiagaraInt {};
 
 	struct FNiagaraBool {};
+
+	struct FNiagaraConvertID {};
 };
 
 template<>
@@ -2062,6 +2064,13 @@ struct FNiagaraDataSetAccessor<FNiagaraDataConversions::FNiagaraBool> : public F
 	}
 };
 
+template<>
+struct FNiagaraDataSetAccessor<FNiagaraDataConversions::FNiagaraConvertID> : public FNiagaraDataSetAccessor<FNiagaraID>
+{
+	FNiagaraDataSetAccessor<FNiagaraDataConversions::FNiagaraConvertID>(FNiagaraDataSet& InDataSet, const FName& InName) : FNiagaraDataSetAccessor<FNiagaraID>(InDataSet, FNiagaraVariable(FNiagaraTypeDefinition::GetIDDef(), InName))
+	{
+	}
+};
 
 /**
 Iterator that will pull or push data between a NiagaraDataBuffer and some FNiagaraVariables it contains.
diff --git a/Engine/Plugins/FX/Niagara/Source/Niagara/Private/NiagaraDataInterfaceParticleRead.cpp b/Engine/Plugins/FX/Niagara/Source/Niagara/Private/NiagaraDataInterfaceParticleRead.cpp
--- a/Engine/Plugins/FX/Niagara/Source/Niagara/Private/NiagaraDataInterfaceParticleRead.cpp
+++ b/Engine/Plugins/FX/Niagara/Source/Niagara/Private/NiagaraDataInterfaceParticleRead.cpp
@@ -1527,7 +1527,7 @@ void UNiagaraDataInterfaceParticleRead::ReadID(FVectorVMContext& Context, FName
 		FNiagaraID ParticleID = { ParticleIDIndexParam.GetAndAdvance(), ParticleIDAcquireTagParam.GetAndAdvance() };
 		bool bValid;
 
-		FNiagaraID Value = RetrieveValueWithCheck<FNiagaraID>(InstanceData->EmitterInstance, FNiagaraTypeDefinition::GetIDDef(), AttributeToRead, ParticleID, bValid);
+		FNiagaraID Value = RetrieveValueWithCheck<FNiagaraID, FNiagaraDataConversions::FNiagaraConvertID>(InstanceData->EmitterInstance, AttributeToRead, ParticleID, bValid);
 
 		FNiagaraBool ValidValue;
 		ValidValue.SetValue(bValid);
@@ -1771,7 +1771,7 @@ void UNiagaraDataInterfaceParticleRead::ReadIDByIndex(FVectorVMContext& Context,
 	{
 		int32 ParticleIndex = ParticleIndexParam.GetAndAdvance();
 		bool bValid;
-		FNiagaraID Value = RetrieveValueByIndexWithCheck<FNiagaraID>(InstanceData->EmitterInstance, FNiagaraTypeDefinition::GetIDDef(), AttributeToRead, ParticleIndex, bValid);
+		FNiagaraID Value = RetrieveValueByIndexWithCheck<FNiagaraID, FNiagaraDataConversions::FNiagaraConvertID>(InstanceData->EmitterInstance, AttributeToRead, ParticleIndex, bValid);
 		FNiagaraBool ValidValue;
 		ValidValue.SetValue(bValid);
 		*OutValid.GetDestAndAdvance() = ValidValue;
@@ -1867,6 +1867,24 @@ void UNiagaraDataInterfaceParticleRead::GetParameterDefinitionHLSL(const FNiagar
 	OutHLSL += FString::Format(FormatDeclarations, ArgsDeclarations);
 }
 
+static FString GenerateFetchValueHLSL(int NumComponents, const TCHAR* ComponentNames[], const TCHAR* ComponentTypeName, const FString& InputBufferName, const FString& ParticleStrideName, bool bExtraIndent)
+{
+	FString Code;
+
+	for (int ComponentIndex = 0; ComponentIndex < NumComponents; ++ComponentIndex)
+	{
+		const TCHAR* ComponentName = (NumComponents > 1) ? ComponentNames[ComponentIndex] : TEXT("");
+		const FString FetchComponentCode = FString::Printf(TEXT("            Out_Value%s = %s(%s[(RegisterIndex + %d)*%s + ParticleIndex]);\n"), ComponentName, ComponentTypeName, *InputBufferName, ComponentIndex, *ParticleStrideName);
+		if (bExtraIndent)
+		{
+			Code += TEXT("    ");
+		}
+		Code += FetchComponentCode;
+	}
+
+	return Code;
+}
+
 static bool GenerateGetFunctionHLSL(const FNiagaraDataInterfaceGPUParamInfo& ParamInfo, const FNiagaraDataInterfaceGeneratedFunction& FunctionInfo, int FunctionInstanceIndex, ENiagaraParticleDataComponentType ComponentType, int NumComponents, bool bByIndex, FString& OutHLSL)
 {
 	FString FuncTemplate;
@@ -1903,16 +1921,7 @@ static bool GenerateGetFunctionHLSL(const FNiagaraDataInterfaceGPUParamInfo& Par
 	FuncTemplate += TEXT(
 		"    {\n"
 		"        Out_Valid = true;\n"
-		"\n"
-		"        BRANCH\n"
-		"        if (!{AttributeCompressedName}[{AttributeIndexGroup}]{AttributeIndexComponent})\n"
-		"        {\n"
 		"{FetchValueCode}"
-		"        }\n"
-		"        else\n"
-		"        {\n"
-		"{FetchCompressedValueCode}"
-		"        }\n"
 		"    }\n"
 		"    else\n"
 		"    {\n"
@@ -1924,7 +1933,6 @@ static bool GenerateGetFunctionHLSL(const FNiagaraDataInterfaceGPUParamInfo& Par
 
 	static const TCHAR* VectorComponentNames[] = { TEXT(".x"), TEXT(".y"), TEXT(".z"), TEXT(".w") };
 	static const TCHAR* IDComponentNames[] = { TEXT(".Index"), TEXT(".AcquireTag") };
-	const TCHAR** ValueComponentNames = VectorComponentNames;
 
 	const FString ParticleStrideFloatName = ParticleStrideFloatBaseName + ParamInfo.DataInterfaceHLSLSymbol;
 	const FString ParticleStrideIntName = ParticleStrideIntBaseName + ParamInfo.DataInterfaceHLSLSymbol;
@@ -1934,67 +1942,51 @@ static bool GenerateGetFunctionHLSL(const FNiagaraDataInterfaceGPUParamInfo& Par
 	const FString InputHalfBufferName = InputHalfBufferBaseName + ParamInfo.DataInterfaceHLSLSymbol;
 
 	const TCHAR* ComponentTypeName;
-	const TCHAR* InputBufferName;
-	const TCHAR* InputBufferStrideName;
-	switch (ComponentType)
-	{
-		case ENiagaraParticleDataComponentType::Float:
-			ComponentTypeName = TEXT("float");
-			InputBufferName = nullptr; // unused
-			InputBufferStrideName = nullptr; // unused
-			break;
-		case ENiagaraParticleDataComponentType::Int:
-			ComponentTypeName = TEXT("int");
-			InputBufferName = *InputIntBufferName;
-			InputBufferStrideName = *ParticleStrideIntName;
-			break;
-		case ENiagaraParticleDataComponentType::Bool:
-			ComponentTypeName = TEXT("bool");
-			InputBufferName = *InputIntBufferName;
-			InputBufferStrideName = *ParticleStrideIntName;
-			break;
-		case ENiagaraParticleDataComponentType::ID:
-			ComponentTypeName = TEXT("int");
-			InputBufferName = *InputIntBufferName;
-			InputBufferStrideName = *ParticleStrideIntName;
-			// IDs are structs with 2 components.
-			NumComponents = 2;
-			ValueComponentNames = IDComponentNames;
-			break;
-		default:
-			UE_LOG(LogNiagara, Error, TEXT("Unknown component type %d while generating function %s"), ComponentType, *FunctionInfo.InstanceName);
-			return false;
-	}
-
 	FString FetchValueCode;
-	FString FetchCompressedValueCode;
-
-	// deal with floats differently because the data may be in the buffer corresponding to our half precision SRV
-	if (ComponentType == ENiagaraParticleDataComponentType::Float)
+	if (ComponentType != ENiagaraParticleDataComponentType::Float)
 	{
-		for (int ComponentIndex = 0; ComponentIndex < NumComponents; ++ComponentIndex)
+		const TCHAR** ValueComponentNames = VectorComponentNames;
+		switch (ComponentType)
 		{
-			const TCHAR* ComponentName = (NumComponents > 1) ? ValueComponentNames[ComponentIndex] : TEXT("");
-			const FString FetchComponentCode = FString::Printf(TEXT("            Out_Value%s = %s(%s[(RegisterIndex + %d)*%s + ParticleIndex]);\n"), ComponentName, ComponentTypeName, *InputFloatBufferName, ComponentIndex, *ParticleStrideFloatName);
-			FetchValueCode += FetchComponentCode;
+			case ENiagaraParticleDataComponentType::Int:
+				ComponentTypeName = TEXT("int");
+				break;
+			case ENiagaraParticleDataComponentType::Bool:
+				ComponentTypeName = TEXT("bool");
+				break;
+			case ENiagaraParticleDataComponentType::ID:
+				ComponentTypeName = TEXT("int");
+				// IDs are structs with 2 components and specific field names.
+				NumComponents = 2;
+				ValueComponentNames = IDComponentNames;
+				break;
+			default:
+				UE_LOG(LogNiagara, Error, TEXT("Unknown component type %d while generating function %s"), ComponentType, *FunctionInfo.InstanceName);
+				return false;
 		}
 
-		for (int ComponentIndex = 0; ComponentIndex < NumComponents; ++ComponentIndex)
-		{
-			const TCHAR* ComponentName = (NumComponents > 1) ? ValueComponentNames[ComponentIndex] : TEXT("");
-			const FString FetchComponentCode = FString::Printf(TEXT("            Out_Value%s = %s(%s[(RegisterIndex + %d)*%s + ParticleIndex]);\n"), ComponentName, ComponentTypeName, *InputHalfBufferName, ComponentIndex, *ParticleStrideHalfName);
-			FetchCompressedValueCode += FetchComponentCode;
-		}
+		FetchValueCode = GenerateFetchValueHLSL(NumComponents, ValueComponentNames, ComponentTypeName, InputIntBufferName, ParticleStrideIntName, false);
 	}
 	else
 	{
-		for (int ComponentIndex = 0; ComponentIndex < NumComponents; ++ComponentIndex)
-		{
-			const TCHAR* ComponentName = (NumComponents > 1) ? ValueComponentNames[ComponentIndex] : TEXT("");
-			const FString FetchComponentCode = FString::Printf(TEXT("        Out_Value%s = %s(%s[(RegisterIndex + %d)*%s + ParticleIndex]);\n"), ComponentName, ComponentTypeName, InputBufferName, ComponentIndex, InputBufferStrideName);
-			FetchValueCode += FetchComponentCode;
-		}
-		FetchCompressedValueCode = FetchValueCode;
+		ComponentTypeName = TEXT("float");
+		// Floats and vectors can be compressed, so we need to add extra code which checks.
+		FString FetchFloatCode = GenerateFetchValueHLSL(NumComponents, VectorComponentNames, ComponentTypeName, InputFloatBufferName, ParticleStrideFloatName, true);
+		FString FetchHalfCode = GenerateFetchValueHLSL(NumComponents, VectorComponentNames, ComponentTypeName, InputHalfBufferName, ParticleStrideHalfName, true);
+		FetchValueCode = FString(
+			TEXT(
+			"        BRANCH\n"
+			"        if (!{AttributeCompressedName}[{AttributeIndexGroup}]{AttributeIndexComponent})\n"
+			"        {\n"
+			)) +
+			FetchFloatCode + TEXT(
+			"        }\n"
+			"        else\n"
+			"        {\n"
+			) +
+			FetchHalfCode + TEXT(
+			"        }\n"
+			);
 	}
 
 	FString ValueTypeName;
@@ -2023,7 +2015,6 @@ static bool GenerateGetFunctionHLSL(const FNiagaraDataInterfaceGPUParamInfo& Par
 	FuncTemplateArgs.Add(TEXT("AcquireTagRegisterIndexName"), AcquireTagRegisterIndexBaseName + ParamInfo.DataInterfaceHLSLSymbol);
 	FuncTemplateArgs.Add(TEXT("ParticleStrideIntName"), ParticleStrideIntName);
 	FuncTemplateArgs.Add(TEXT("FetchValueCode"), FetchValueCode);
-	FuncTemplateArgs.Add(TEXT("FetchCompressedValueCode"), FetchCompressedValueCode);
 	
 	OutHLSL += FString::Format(*FuncTemplate, FuncTemplateArgs);
 
