- Added support for GetFilteredSocketOrBoneCount / GetFilteredSocketOrBoneBoneAt
- Fixed an issue with skinning mode Nonne and virtual bones

#rnx

diff --git a/Engine/Plugins/FX/Niagara/Shaders/Private/NiagaraDataInterfaceSkeletalMesh.ush b/Engine/Plugins/FX/Niagara/Shaders/Private/NiagaraDataInterfaceSkeletalMesh.ush
--- a/Engine/Plugins/FX/Niagara/Shaders/Private/NiagaraDataInterfaceSkeletalMesh.ush
+++ b/Engine/Plugins/FX/Niagara/Shaders/Private/NiagaraDataInterfaceSkeletalMesh.ush
@@ -601,7 +601,7 @@ void DISkelMesh_GetFilteredSocketCount(in FDISkelMeshContext DIContext, out int
 
 void DISkelMesh_GetFilteredSocketBoneAt(in FDISkelMeshContext DIContext, in int InSocketIndex, out int Bone)
 {
-    Bone = DIContext.FilteredSocketBoneOffset + clamp(InSocketIndex, 0, DIContext.NumFilteredSockets);
+    Bone = DIContext.FilteredSocketBoneOffset + clamp(InSocketIndex, 0, DIContext.NumFilteredSockets - 1);
 }
 
 void DISkelMesh_RandomFilteredSocket(in FDISkelMeshContext DIContext, uint Seed1, uint Seed2, uint Seed3, out int SocketBone)
@@ -625,6 +625,23 @@ void DISkelMesh_RandomFilteredSocketOrBone(in FDISkelMeshContext DIContext, uint
 	}
 }
 
+void DISkelMesh_GetFilteredSocketOrBoneCount(in FDISkelMeshContext DIContext, out int Count)
+{
+	Count = DIContext.NumFilteredBones + DIContext.NumFilteredSockets;
+}
+
+void DISkelMesh_GetFilteredSocketOrBoneAt(in FDISkelMeshContext DIContext, int FilteredIndex, out int Bone)
+{
+	if ( FilteredIndex < DIContext.NumFilteredBones )
+	{
+		Bone = (int)DIContext.FilteredAndUnfilteredBones[FilteredIndex];
+	}
+	else
+	{
+		Bone = DIContext.FilteredSocketBoneOffset + FilteredIndex - DIContext.NumFilteredBones;
+	}
+}
+
 ////////////////////////////////////////////////////////////////////////////////////////////////////////////////////////////////////////////////////////////////
 
 #define DISKELMESH_DECLARE_CONSTANTS(NAME) \
diff --git a/Engine/Plugins/FX/Niagara/Source/Niagara/Classes/NDISkeletalMeshCommon.h b/Engine/Plugins/FX/Niagara/Source/Niagara/Classes/NDISkeletalMeshCommon.h
--- a/Engine/Plugins/FX/Niagara/Source/Niagara/Classes/NDISkeletalMeshCommon.h
+++ b/Engine/Plugins/FX/Niagara/Source/Niagara/Classes/NDISkeletalMeshCommon.h
@@ -127,22 +127,36 @@ struct FSkinnedPositionAccessorHelper<TIntegralConstant<ENDISkeletalMesh_Skinnin
 
 	FORCEINLINE_DEBUGGABLE FVector GetSkinnedBonePosition(FSkeletalMeshAccessorHelper& Accessor, int32 BoneIndex)
 	{
-		return Accessor.Mesh->GetComposedRefPoseMatrix(BoneIndex).GetOrigin();
+		const int32 NumRealBones = Accessor.Mesh->RefSkeleton.GetRawBoneNum();
+		if (BoneIndex < NumRealBones)
+		{
+			return Accessor.Mesh->GetComposedRefPoseMatrix(BoneIndex).GetOrigin();
+		}
+
+		const FTransform& RefTransform = Accessor.Mesh->RefSkeleton.GetRefBonePose()[BoneIndex];
+		return RefTransform.GetLocation();
 	}
 
 	FORCEINLINE_DEBUGGABLE FVector GetSkinnedBonePreviousPosition(FSkeletalMeshAccessorHelper& Accessor, int32 BoneIndex)
 	{
-		return Accessor.Mesh->GetComposedRefPoseMatrix(BoneIndex).GetOrigin();
+		return GetSkinnedBonePosition(Accessor, BoneIndex);
 	}
 
 	FORCEINLINE_DEBUGGABLE FQuat GetSkinnedBoneRotation(FSkeletalMeshAccessorHelper& Accessor, int32 BoneIndex)
 	{
-		return Accessor.Mesh->GetComposedRefPoseMatrix(BoneIndex).ToQuat();
+		const int32 NumRealBones = Accessor.Mesh->RefSkeleton.GetRawBoneNum();
+		if (BoneIndex < NumRealBones)
+		{
+			return Accessor.Mesh->GetComposedRefPoseMatrix(BoneIndex).ToQuat();
+		}
+
+		const FTransform& RefTransform = Accessor.Mesh->RefSkeleton.GetRefBonePose()[BoneIndex];
+		return RefTransform.GetRotation();
 	}
 
 	FORCEINLINE_DEBUGGABLE FQuat GetSkinnedBonePreviousRotation(FSkeletalMeshAccessorHelper& Accessor, int32 BoneIndex)
 	{
-		return Accessor.Mesh->GetComposedRefPoseMatrix(BoneIndex).ToQuat();
+		return GetSkinnedBoneRotation(Accessor, BoneIndex);
 	}
 };
 
diff --git a/Engine/Plugins/FX/Niagara/Source/Niagara/Classes/NiagaraDataInterfaceSkeletalMesh.h b/Engine/Plugins/FX/Niagara/Source/Niagara/Classes/NiagaraDataInterfaceSkeletalMesh.h
--- a/Engine/Plugins/FX/Niagara/Source/Niagara/Classes/NiagaraDataInterfaceSkeletalMesh.h
+++ b/Engine/Plugins/FX/Niagara/Source/Niagara/Classes/NiagaraDataInterfaceSkeletalMesh.h
@@ -792,6 +792,8 @@ public:
 	void RandomFilteredSocket(FVectorVMContext& Context);
 
 	void RandomFilteredSocketOrBone(FVectorVMContext& Context);
+	void GetFilteredSocketOrBoneCount(FVectorVMContext& Context);
+	void GetFilteredSocketOrBoneBoneAt(FVectorVMContext& Context);
 	// End of Direct Bone + Socket Sampling
 	//////////////////////////////////////////////////////////////////////////
 
@@ -841,6 +843,8 @@ public:
 	static const FName GetFilteredSocketTransformName;
 
 	static const FName RandomFilteredSocketOrBoneName;
+	static const FName GetFilteredSocketOrBoneCountName;
+	static const FName GetFilteredSocketOrBoneAtName;
 
 	// Vertex Sampling
 	static const FName GetSkinnedVertexDataName;
diff --git a/Engine/Plugins/FX/Niagara/Source/Niagara/Private/NDISkeletalMesh_BoneSampling.cpp b/Engine/Plugins/FX/Niagara/Source/Niagara/Private/NDISkeletalMesh_BoneSampling.cpp
--- a/Engine/Plugins/FX/Niagara/Source/Niagara/Private/NDISkeletalMesh_BoneSampling.cpp
+++ b/Engine/Plugins/FX/Niagara/Source/Niagara/Private/NDISkeletalMesh_BoneSampling.cpp
@@ -40,9 +40,11 @@ const FName FSkeletalMeshInterfaceHelper::GetUnfilteredBoneAtName("GetUnfiltered
 const FName FSkeletalMeshInterfaceHelper::RandomFilteredSocketName("RandomFilteredSocket");
 const FName FSkeletalMeshInterfaceHelper::GetFilteredSocketCountName("GetFilteredSocketCount");
 const FName FSkeletalMeshInterfaceHelper::GetFilteredSocketTransformName("GetFilteredSocketTransform");
-const FName FSkeletalMeshInterfaceHelper::GetFilteredSocketBoneAtName("GetFilteredSocketBone");
+const FName FSkeletalMeshInterfaceHelper::GetFilteredSocketBoneAtName("GetFilteredSocket");
 
 const FName FSkeletalMeshInterfaceHelper::RandomFilteredSocketOrBoneName("RandomFilteredSocketOrBone");
+const FName FSkeletalMeshInterfaceHelper::GetFilteredSocketOrBoneCountName("GetFilteredSocketOrBoneCount");
+const FName FSkeletalMeshInterfaceHelper::GetFilteredSocketOrBoneAtName("GetFilteredSocketOrBone");
 
 void UNiagaraDataInterfaceSkeletalMesh::GetSkeletonSamplingFunctions(TArray<FNiagaraFunctionSignature>& OutFunctions)
 {
@@ -300,6 +302,31 @@ void UNiagaraDataInterfaceSkeletalMesh::GetSkeletonSamplingFunctions(TArray<FNia
 		Sig.bRequiresContext = false;
 #if WITH_EDITORONLY_DATA
 		Sig.Description = LOCTEXT("RandomFilteredSocketOrBoneDesc", "Gets the bone for a random filtered socket or bone from the DI's list.");
+#endif
+		OutFunctions.Add(Sig);
+	}
+	{
+		FNiagaraFunctionSignature Sig;
+		Sig.Name = FSkeletalMeshInterfaceHelper::GetFilteredSocketOrBoneCountName;
+		Sig.Inputs.Add(FNiagaraVariable(FNiagaraTypeDefinition(GetClass()), TEXT("SkeletalMesh")));
+		Sig.Outputs.Add(FNiagaraVariable(FNiagaraTypeDefinition::GetIntDef(), TEXT("Count")));
+		Sig.bMemberFunction = true;
+		Sig.bRequiresContext = false;
+#if WITH_EDITORONLY_DATA
+		Sig.Description = LOCTEXT("GetFilteredSocketOrBoneCountDesc", "Gets the total filtered socket and bone count from the DI's list.");
+#endif
+		OutFunctions.Add(Sig);
+	}
+	{
+		FNiagaraFunctionSignature Sig;
+		Sig.Name = FSkeletalMeshInterfaceHelper::GetFilteredSocketOrBoneAtName;
+		Sig.Inputs.Add(FNiagaraVariable(FNiagaraTypeDefinition(GetClass()), TEXT("SkeletalMesh")));
+		Sig.Inputs.Add(FNiagaraVariable(FNiagaraTypeDefinition::GetIntDef(), TEXT("Socket Or Bone Index")));
+		Sig.Outputs.Add(FNiagaraVariable(FNiagaraTypeDefinition::GetIntDef(), TEXT("Bone")));
+		Sig.bMemberFunction = true;
+		Sig.bRequiresContext = false;
+#if WITH_EDITORONLY_DATA
+		Sig.Description = LOCTEXT("GetFilteredSocketOrBoneAtDesc", "Gets a filtered socket or bone count from the DI's list.");
 #endif
 		OutFunctions.Add(Sig);
 	}
@@ -400,6 +427,16 @@ void UNiagaraDataInterfaceSkeletalMesh::BindSkeletonSamplingFunction(const FVMEx
 		check(BindingInfo.GetNumInputs() == 4 && BindingInfo.GetNumOutputs() == 1);
 		OutFunc = FVMExternalFunction::CreateLambda([this](FVectorVMContext& Context) { this->RandomFilteredSocketOrBone(Context); });
 	}
+	else if (BindingInfo.Name == FSkeletalMeshInterfaceHelper::GetFilteredSocketOrBoneCountName)
+	{
+		check(BindingInfo.GetNumInputs() == 1 && BindingInfo.GetNumOutputs() == 1);
+		OutFunc = FVMExternalFunction::CreateLambda([this](FVectorVMContext& Context) { this->GetFilteredSocketOrBoneCount(Context); });
+	}
+	else if (BindingInfo.Name == FSkeletalMeshInterfaceHelper::GetFilteredSocketOrBoneAtName)
+	{
+		check(BindingInfo.GetNumInputs() == 2 && BindingInfo.GetNumOutputs() == 1);
+		OutFunc = FVMExternalFunction::CreateLambda([this](FVectorVMContext& Context) { this->GetFilteredSocketOrBoneBoneAt(Context); });
+	}
 }
 
 //////////////////////////////////////////////////////////////////////////
@@ -968,4 +1005,46 @@ void UNiagaraDataInterfaceSkeletalMesh::RandomFilteredSocketOrBone(FVectorVMCont
 	}
 }
 
+void UNiagaraDataInterfaceSkeletalMesh::GetFilteredSocketOrBoneCount(FVectorVMContext& Context)
+{
+	VectorVM::FUserPtrHandler<FNDISkeletalMesh_InstanceData> InstData(Context);
+
+	VectorVM::FExternalFuncRegisterHandler<int32> OutCount(Context);
+	const int32 Count = FilteredSockets.Num() + InstData->NumFilteredBones;
+	for (int32 i=0; i < Context.NumInstances; ++i)
+	{
+		*OutCount.GetDestAndAdvance() = Count;
+	}
+}
+
+void UNiagaraDataInterfaceSkeletalMesh::GetFilteredSocketOrBoneBoneAt(FVectorVMContext& Context)
+{
+	VectorVM::FExternalFuncInputHandler<int32> IndexParam(Context);
+	VectorVM::FUserPtrHandler<FNDISkeletalMesh_InstanceData> InstData(Context);
+	VectorVM::FExternalFuncRegisterHandler<int32> OutBoneIndex(Context);
+
+	const int32 Max = FilteredSockets.Num() + InstData->NumFilteredBones - 1;
+	if (Max >= 0)
+	{
+		const int32 NumFilteredBones = InstData->NumFilteredBones;
+		const int32 FilteredSocketBoneOffset = InstData->FilteredSocketBoneOffset;
+		for (int32 i = 0; i < Context.NumInstances; ++i)
+		{
+			const int32 FilteredIndex = IndexParam.GetAndAdvance();
+			if (FilteredIndex < NumFilteredBones)
+			{
+				*OutBoneIndex.GetDestAndAdvance() = InstData->FilteredAndUnfilteredBones[FilteredIndex];
+			}
+			else
+			{
+				*OutBoneIndex.GetDestAndAdvance() = FilteredSocketBoneOffset + FilteredIndex - NumFilteredBones;
+			}
+		}
+	}
+	else
+	{
+		FMemory::Memset(OutBoneIndex.GetDest(), 0xFF, Context.NumInstances * sizeof(int32));
+	}
+}
+
 #undef LOCTEXT_NAMESPACE
diff --git a/Engine/Plugins/FX/Niagara/Source/Niagara/Private/NiagaraDataInterfaceSkeletalMesh.cpp b/Engine/Plugins/FX/Niagara/Source/Niagara/Private/NiagaraDataInterfaceSkeletalMesh.cpp
--- a/Engine/Plugins/FX/Niagara/Source/Niagara/Private/NiagaraDataInterfaceSkeletalMesh.cpp
+++ b/Engine/Plugins/FX/Niagara/Source/Niagara/Private/NiagaraDataInterfaceSkeletalMesh.cpp
@@ -29,6 +29,7 @@ struct FNiagaraSkelMeshDIFunctionVersion
 		InitialVersion = 0,
 		AddedRandomInfo = 1,
 		CleanUpVertexSampling = 2,
+		CleanupBoneSampling = 3,
 
 		VersionPlusOne,
 		LatestVersion = VersionPlusOne - 1
@@ -2563,7 +2564,7 @@ bool UNiagaraDataInterfaceSkeletalMesh::GetFunctionHLSL(const FNiagaraDataInterf
 		static const TCHAR* FormatSample = TEXT("void {InstanceFunctionName} (NiagaraRandInfo InRandomInfo, out int Bone) { {GetDISkelMeshContextName} DISkelMesh_RandomFilteredBone(DIContext, InRandomInfo.Seed1, InRandomInfo.Seed2, InRandomInfo.Seed3, Bone); }");
 		OutHLSL += FString::Format(FormatSample, ArgsSample);
 	}
-	else if (FunctionInfo.DefinitionName == FSkeletalMeshInterfaceHelper::GetUnfilteredBoneAtName)
+	else if (FunctionInfo.DefinitionName == FSkeletalMeshInterfaceHelper::GetUnfilteredBoneCountName)
 	{
 		static const TCHAR* FormatSample = TEXT("void {InstanceFunctionName} (out int Count) { {GetDISkelMeshContextName} DISkelMesh_GetUnfilteredBoneCount(DIContext, Count); }");
 		OutHLSL += FString::Format(FormatSample, ArgsSample);
@@ -2609,6 +2610,16 @@ bool UNiagaraDataInterfaceSkeletalMesh::GetFunctionHLSL(const FNiagaraDataInterf
 		static const TCHAR* FormatSample = TEXT("void {InstanceFunctionName} (NiagaraRandInfo InRandomInfo, out int Bone) { {GetDISkelMeshContextName} DISkelMesh_RandomFilteredSocketOrBone(DIContext, InRandomInfo.Seed1, InRandomInfo.Seed2, InRandomInfo.Seed3, Bone); }");
 		OutHLSL += FString::Format(FormatSample, ArgsSample);
 	}
+	else if (FunctionInfo.DefinitionName == FSkeletalMeshInterfaceHelper::GetFilteredSocketOrBoneCountName)
+	{
+		static const TCHAR* FormatSample = TEXT("void {InstanceFunctionName} (out int Count) { {GetDISkelMeshContextName} DISkelMesh_GetFilteredSocketOrBoneCount(DIContext, Count); }");
+		OutHLSL += FString::Format(FormatSample, ArgsSample);
+	}
+	else if (FunctionInfo.DefinitionName == FSkeletalMeshInterfaceHelper::GetFilteredSocketOrBoneAtName)
+	{
+		static const TCHAR* FormatSample = TEXT("void {InstanceFunctionName} (in int FilteredIndex, out int Bone) { {GetDISkelMeshContextName} DISkelMesh_GetFilteredSocketOrBoneAt(DIContext, FilteredIndex, Bone); }");
+		OutHLSL += FString::Format(FormatSample, ArgsSample);
+	}
 	//////////////////////////////////////////////////////////////////////////////////////////////////////////////////////////////
 	// Unsupported functionality
 	else
@@ -2626,6 +2637,12 @@ bool UNiagaraDataInterfaceSkeletalMesh::UpgradeFunctionCall(FNiagaraFunctionSign
 {
 	bool bWasChanged = false;
 
+	// Early out for version matching
+	if (FunctionSignature.FunctionVersion == FNiagaraSkelMeshDIFunctionVersion::LatestVersion)
+	{
+		return bWasChanged;
+	}
+
 	// Renamed some functions and added Random Info to Various functions for consistency
 	if (FunctionSignature.FunctionVersion < FNiagaraSkelMeshDIFunctionVersion::AddedRandomInfo)
 	{
@@ -2696,6 +2713,27 @@ bool UNiagaraDataInterfaceSkeletalMesh::UpgradeFunctionCall(FNiagaraFunctionSign
 			}
 		}
 	}
+
+	// Clean up CleanupBoneSampling
+	if (FunctionSignature.FunctionVersion < FNiagaraSkelMeshDIFunctionVersion::CleanupBoneSampling)
+	{
+		static const TPair<FName, FName> FunctionRenames[] =
+		{
+			MakeTuple(FName("GetFilteredSocketBone"), FSkeletalMeshInterfaceHelper::GetFilteredSocketBoneAtName),
+		};
+
+		for (const auto& RenamePair : FunctionRenames)
+		{
+			if (FunctionSignature.Name == RenamePair.Key)
+			{
+				FunctionSignature.Name = RenamePair.Value;
+				bWasChanged = true;
+				break;
+			}
+		}
+	}
+
+	// Set latest version
 	FunctionSignature.FunctionVersion = FNiagaraSkelMeshDIFunctionVersion::LatestVersion;
 
 	return bWasChanged;
