- Fix crash when we have no components to transfer to the GPU

#rnx

diff --git a/Engine/Plugins/FX/Niagara/Source/Niagara/Private/NiagaraRenderer.cpp b/Engine/Plugins/FX/Niagara/Source/Niagara/Private/NiagaraRenderer.cpp
--- a/Engine/Plugins/FX/Niagara/Source/Niagara/Private/NiagaraRenderer.cpp
+++ b/Engine/Plugins/FX/Niagara/Source/Niagara/Private/NiagaraRenderer.cpp
@@ -188,9 +188,11 @@ bool FNiagaraRenderer::SetVertexFactoryVariable(const FNiagaraDataSet& DataSet,
 
 FGlobalDynamicReadBuffer::FAllocation FNiagaraRenderer::TransferDataToGPU(FGlobalDynamicReadBuffer& DynamicReadBuffer, FNiagaraDataBuffer* SrcData)const
 {
-	int32 TotalFloatSize = TotalVFComponents * SrcData->GetNumInstances();
-	int32 ComponentStrideDest = SrcData->GetNumInstances() * sizeof(float);
-	FGlobalDynamicReadBuffer::FAllocation Allocation = DynamicReadBuffer.AllocateFloat(TotalFloatSize);
+	const int32 TotalFloatSize = TotalVFComponents * SrcData->GetNumInstances();
+	const int32 ComponentStrideDest = SrcData->GetNumInstances() * sizeof(float);
+
+	//-TODO: We are allocating 4 bytes where we need no memory, ideally we fix the underlying code to handle 0 allocations this is a temporary band-aid
+	FGlobalDynamicReadBuffer::FAllocation Allocation = DynamicReadBuffer.AllocateFloat(FMath::Max(4, TotalFloatSize));
 	for (const FNiagaraRendererVariableInfo& VarInfo : VFVariables)
 	{
 		int32 GpuOffset = VarInfo.GetGPUOffset();
