Don't allow visibility culling if there's an invalid LastRenderTime.
Don't allow any culling in the Niagara preview viewport.

#rnx

diff --git a/Engine/Plugins/FX/Niagara/Source/Niagara/Private/NiagaraWorldManager.cpp b/Engine/Plugins/FX/Niagara/Source/Niagara/Private/NiagaraWorldManager.cpp
--- a/Engine/Plugins/FX/Niagara/Source/Niagara/Private/NiagaraWorldManager.cpp
+++ b/Engine/Plugins/FX/Niagara/Source/Niagara/Private/NiagaraWorldManager.cpp
@@ -862,7 +862,7 @@ void FNiagaraWorldManager::SignificanceCull(UNiagaraEffectType* EffectType, cons
 void FNiagaraWorldManager::VisibilityCull(UNiagaraEffectType* EffectType, const FNiagaraSystemScalabilitySettings& ScalabilitySettings, UNiagaraComponent* Component, FNiagaraScalabilityState& OutState)
 {
 	float TimeSinceRendered = Component->GetSafeTimeSinceRendered(World->TimeSeconds);
-	bool bCull = ScalabilitySettings.bCullByMaxTimeWithoutRender && TimeSinceRendered > ScalabilitySettings.MaxTimeWithoutRender;
+	bool bCull = Component->GetLastRenderTime() >= 0.0f && ScalabilitySettings.bCullByMaxTimeWithoutRender && TimeSinceRendered > ScalabilitySettings.MaxTimeWithoutRender;
 
 	OutState.bCulled |= bCull;
 #if DEBUG_SCALABILITY_STATE
diff --git a/Engine/Plugins/FX/Niagara/Source/NiagaraEditor/Private/ViewModels/NiagaraSystemViewModel.cpp b/Engine/Plugins/FX/Niagara/Source/NiagaraEditor/Private/ViewModels/NiagaraSystemViewModel.cpp
--- a/Engine/Plugins/FX/Niagara/Source/NiagaraEditor/Private/ViewModels/NiagaraSystemViewModel.cpp
+++ b/Engine/Plugins/FX/Niagara/Source/NiagaraEditor/Private/ViewModels/NiagaraSystemViewModel.cpp
@@ -877,6 +877,7 @@ void FNiagaraSystemViewModel::SetupPreviewComponentAndInstance()
 		PreviewComponent = NewObject<UNiagaraComponent>(GetTransientPackage(), NAME_None, RF_Transient);
 		PreviewComponent->CastShadow = 1;
 		PreviewComponent->bCastDynamicShadow = 1;
+		PreviewComponent->SetAllowScalability(false);
 		PreviewComponent->SetAsset(System);
 		PreviewComponent->SetForceSolo(true);
 		PreviewComponent->SetAgeUpdateMode(ENiagaraAgeUpdateMode::DesiredAge);
