- Add support for uniform distrbution to RandomTriangle

#rnx

diff --git a/Engine/Plugins/FX/Niagara/Shaders/Private/NiagaraDataInterfaceSkeletalMesh.ush b/Engine/Plugins/FX/Niagara/Shaders/Private/NiagaraDataInterfaceSkeletalMesh.ush
--- a/Engine/Plugins/FX/Niagara/Shaders/Private/NiagaraDataInterfaceSkeletalMesh.ush
+++ b/Engine/Plugins/FX/Niagara/Shaders/Private/NiagaraDataInterfaceSkeletalMesh.ush
@@ -332,8 +332,29 @@ void DISKelMesh_GetTriUV(in FDISkelMeshContext DIContext, in uint TriangleIndex,
 
 void DISKelMesh_RandomTriangle(in FDISkelMeshContext DIContext, uint Seed1, uint Seed2, uint Seed3, out int OutTriangle, out float3 OutBaryCoord)
 {
+	//float RandT0 = NiagaraInternalNoise(Seed1, Seed2, Seed3);
+	//OutTriangle = min(uint(RandT0 * float(DIContext.MeshTriangleCount)), DIContext.MeshTriangleCount - 1); // avoid % by using mul/min to Tri = MeshTriangleCountName
+	//OutBaryCoord = DISKelMesh_RandomBarycentricCoord(Seed1, Seed2, Seed3);
+
+
 	float RandT0 = NiagaraInternalNoise(Seed1, Seed2, Seed3);
-	OutTriangle = min(uint(RandT0 * float(DIContext.MeshTriangleCount)), DIContext.MeshTriangleCount - 1); // avoid % by using mul/min to Tri = MeshTriangleCountName
+	int TriangleIndex = min(int(RandT0 * float(DIContext.MeshTriangleCount)), (int)DIContext.MeshTriangleCount - 1); // avoid % by using mul/min to Tri = MeshTriangleCountName
+
+	[branch]
+	if ( DIContext.UniformTriangleSamplingEnable )
+	{
+		// Uniform area weighted position selection (using alias method from Alias method from FWeightedRandomSampler)
+		float TriangleProbability = DIContext.MeshTriangleSamplerProbaBuffer[TriangleIndex];
+
+		// Alias check
+		float RandT1 = NiagaraInternalNoise(Seed1, Seed2, Seed3);
+		if( RandT1 > TriangleProbability )
+		{
+			TriangleIndex = (int)DIContext.MeshTriangleSamplerAliasBuffer[TriangleIndex];
+		}
+	}
+
+	OutTriangle = TriangleIndex;
 	OutBaryCoord = DISKelMesh_RandomBarycentricCoord(Seed1, Seed2, Seed3);
 }
 
diff --git a/Engine/Plugins/FX/Niagara/Source/Niagara/Private/NDISkeletalMesh_TriangleSampling.cpp b/Engine/Plugins/FX/Niagara/Source/Niagara/Private/NDISkeletalMesh_TriangleSampling.cpp
--- a/Engine/Plugins/FX/Niagara/Source/Niagara/Private/NDISkeletalMesh_TriangleSampling.cpp
+++ b/Engine/Plugins/FX/Niagara/Source/Niagara/Private/NDISkeletalMesh_TriangleSampling.cpp
@@ -472,6 +472,31 @@ void UNiagaraDataInterfaceSkeletalMesh::RandomTriangle(FVectorVMContext& Context
 	FSkeletalMeshAccessorHelper MeshAccessor;
 	MeshAccessor.Init<TIntegralConstant<int32, 0>, TIntegralConstant<int32, 0>>(InstData);
 
+	//-TODO: AREA WEIGHTED
+	const int32 LODIndex = InstData->GetLODIndex();
+	const bool bAreaWeighted = InstData->Mesh->GetLODInfo(LODIndex)->bSupportUniformlyDistributedSampling;
+
+	if (bAreaWeighted)
+	{
+		const FSkeletalMeshSamplingInfo& SamplingInfo = InstData->Mesh->GetSamplingInfo();
+		const FSkeletalMeshSamplingLODBuiltData& WholeMeshBuiltData = SamplingInfo.GetWholeMeshLODBuiltData(InstData->GetLODIndex());
+		if (WholeMeshBuiltData.AreaWeightedTriangleSampler.GetNumEntries() > 0)
+		{
+			for (int32 i = 0; i < Context.NumInstances; ++i)
+			{
+				RandHelper.GetAndAdvance();
+
+				const int32 Triangle = WholeMeshBuiltData.AreaWeightedTriangleSampler.GetEntryIndex(RandHelper.Rand(i), RandHelper.Rand(i));
+				const FVector Bary = RandHelper.RandomBarycentricCoord(i);
+				*OutTri.GetDestAndAdvance() = Triangle;
+				*OutBaryX.GetDestAndAdvance() = Bary.X;
+				*OutBaryY.GetDestAndAdvance() = Bary.Y;
+				*OutBaryZ.GetDestAndAdvance() = Bary.Z;
+			}
+			return;
+		}
+	}
+
 	const int32 MaxTriangle = (MeshAccessor.IndexBuffer->Num() / 3) - 1;
 	if (MaxTriangle >= 0)
 	{
diff --git a/Engine/Plugins/FX/Niagara/Source/NiagaraCore/Private/NiagaraCustomVersion.cpp b/Engine/Plugins/FX/Niagara/Source/NiagaraCore/Private/NiagaraCustomVersion.cpp
--- a/Engine/Plugins/FX/Niagara/Source/NiagaraCore/Private/NiagaraCustomVersion.cpp
+++ b/Engine/Plugins/FX/Niagara/Source/NiagaraCore/Private/NiagaraCustomVersion.cpp
@@ -8,4 +8,4 @@ const FGuid FNiagaraCustomVersion::GUID(0xFCF57AFA, 0x50764283, 0xB9A9E658, 0xFF
 // Register the custom version with core
 FCustomVersionRegistration GRegisterNiagaraCustomVersion(FNiagaraCustomVersion::GUID, FNiagaraCustomVersion::LatestVersion, TEXT("NiagaraVer"));
 
-const FGuid FNiagaraCustomVersion::LatestScriptCompileVersion(0x5FD3FF04, 0xC5BF4260, 0xB48F8499, 0xB449A6D5);
+const FGuid FNiagaraCustomVersion::LatestScriptCompileVersion(0x08611BB6, 0x382AD540, 0xAD008BCA, 0xD3CBDFEC);
