Merge from Dev-EngineMerge [at] 12452070
This represents approximately 4.25 preview 4 and originated from UE4/Release-Engine-Staging [at] CL# 12440029
#rnx

diff --git a/Engine/Plugins/FX/Niagara/Source/NiagaraCore/Private/NiagaraCustomVersion.cpp b/Engine/Plugins/FX/Niagara/Source/NiagaraCore/Private/NiagaraCustomVersion.cpp
--- a/Engine/Plugins/FX/Niagara/Source/NiagaraCore/Private/NiagaraCustomVersion.cpp
+++ b/Engine/Plugins/FX/Niagara/Source/NiagaraCore/Private/NiagaraCustomVersion.cpp
@@ -8,4 +8,4 @@ const FGuid FNiagaraCustomVersion::GUID(0xFCF57AFA, 0x50764283, 0xB9A9E658, 0xFF
 // Register the custom version with core
 FCustomVersionRegistration GRegisterNiagaraCustomVersion(FNiagaraCustomVersion::GUID, FNiagaraCustomVersion::LatestVersion, TEXT("NiagaraVer"));
 
-const FGuid FNiagaraCustomVersion::LatestScriptCompileVersion(0x08611BB6, 0x382AD540, 0xAD008BCA, 0xD3CBDFEC);
+const FGuid FNiagaraCustomVersion::LatestScriptCompileVersion(0x273E38A2, 0xB6044C95, 0xA6264C72, 0x90D8380F);
\ No newline at end of file
