The GPU versions of GetAttributeByIndex return false if the index is out of bounds, just like their CPU counterparts.
Updated engine test to check this.

#rnx

diff --git a/Engine/Plugins/FX/Niagara/Source/Niagara/Private/NiagaraDataInterfaceParticleRead.cpp b/Engine/Plugins/FX/Niagara/Source/Niagara/Private/NiagaraDataInterfaceParticleRead.cpp
--- a/Engine/Plugins/FX/Niagara/Source/Niagara/Private/NiagaraDataInterfaceParticleRead.cpp
+++ b/Engine/Plugins/FX/Niagara/Source/Niagara/Private/NiagaraDataInterfaceParticleRead.cpp
@@ -1906,7 +1906,8 @@ static bool GenerateGetFunctionHLSL(const FNiagaraDataInterfaceGPUParamInfo& Par
 	if (bByIndex)
 	{
 		FuncTemplate += TEXT(
-			"    if(RegisterIndex != -1)\n"
+			"    int NumParticles = {InstanceCountOffsetName} != 0xffffffff ? RWInstanceCounts[{InstanceCountOffsetName}] : 0;\n"
+			"    if(RegisterIndex != -1 && ParticleIndex >= 0 && ParticleIndex < NumParticles)\n"
 		);
 	}
 	else
@@ -1918,19 +1919,6 @@ static bool GenerateGetFunctionHLSL(const FNiagaraDataInterfaceGPUParamInfo& Par
 		);
 	}
 
-	FuncTemplate += TEXT(
-		"    {\n"
-		"        Out_Valid = true;\n"
-		"{FetchValueCode}"
-		"    }\n"
-		"    else\n"
-		"    {\n"
-		"        Out_Valid = false;\n"
-		"        Out_Value = ({ValueType})0;\n"
-		"    }\n"
-		"}\n\n"
-	);
-
 	static const TCHAR* VectorComponentNames[] = { TEXT(".x"), TEXT(".y"), TEXT(".z"), TEXT(".w") };
 	static const TCHAR* IDComponentNames[] = { TEXT(".Index"), TEXT(".AcquireTag") };
 
@@ -2003,6 +1991,23 @@ static bool GenerateGetFunctionHLSL(const FNiagaraDataInterfaceGPUParamInfo& Par
 		ValueTypeName = FString::Printf(TEXT("%s%d"), ComponentTypeName, NumComponents);
 	}
 
+	FuncTemplate += TEXT(
+		"    {\n"
+		"        Out_Valid = true;\n"
+	);
+
+	FuncTemplate += FetchValueCode;
+
+	FuncTemplate += TEXT(
+		"    }\n"
+		"    else\n"
+		"    {\n"
+		"        Out_Valid = false;\n"
+		"        Out_Value = ({ValueType})0;\n"
+		"    }\n"
+		"}\n\n"
+	);
+
 	TMap<FString, FStringFormatArg> FuncTemplateArgs;
 	FuncTemplateArgs.Add(TEXT("FunctionName"), FunctionInfo.InstanceName);
 	FuncTemplateArgs.Add(TEXT("ValueType"), ValueTypeName);
@@ -2015,7 +2020,8 @@ static bool GenerateGetFunctionHLSL(const FNiagaraDataInterfaceGPUParamInfo& Par
 	FuncTemplateArgs.Add(TEXT("AcquireTagRegisterIndexName"), AcquireTagRegisterIndexBaseName + ParamInfo.DataInterfaceHLSLSymbol);
 	FuncTemplateArgs.Add(TEXT("ParticleStrideIntName"), ParticleStrideIntName);
 	FuncTemplateArgs.Add(TEXT("FetchValueCode"), FetchValueCode);
-	
+	FuncTemplateArgs.Add(TEXT("InstanceCountOffsetName"), InstanceCountOffsetBaseName + ParamInfo.DataInterfaceHLSLSymbol);
+
 	OutHLSL += FString::Format(*FuncTemplate, FuncTemplateArgs);
 
 	return true;
diff --git a/Engine/Plugins/FX/Niagara/Source/NiagaraCore/Private/NiagaraCustomVersion.cpp b/Engine/Plugins/FX/Niagara/Source/NiagaraCore/Private/NiagaraCustomVersion.cpp
--- a/Engine/Plugins/FX/Niagara/Source/NiagaraCore/Private/NiagaraCustomVersion.cpp
+++ b/Engine/Plugins/FX/Niagara/Source/NiagaraCore/Private/NiagaraCustomVersion.cpp
@@ -8,4 +8,4 @@ const FGuid FNiagaraCustomVersion::GUID(0xFCF57AFA, 0x50764283, 0xB9A9E658, 0xFF
 // Register the custom version with core
 FCustomVersionRegistration GRegisterNiagaraCustomVersion(FNiagaraCustomVersion::GUID, FNiagaraCustomVersion::LatestVersion, TEXT("NiagaraVer"));
 
-const FGuid FNiagaraCustomVersion::LatestScriptCompileVersion(0x273E38A2, 0xB6044C95, 0xA6264C72, 0x90D8380F);
\ No newline at end of file
+const FGuid FNiagaraCustomVersion::LatestScriptCompileVersion(0x24C7603C, 0xC1454C9F, 0xB8688941, 0xB24F4D06);
