Reformatted Niagara error messages for deprecated and experimental modules. Added separate colors for info, warning and error messages in the stack.

diff --git a/Engine/Plugins/FX/Niagara/Source/NiagaraEditor/Private/NiagaraNodeFunctionCall.cpp b/Engine/Plugins/FX/Niagara/Source/NiagaraEditor/Private/NiagaraNodeFunctionCall.cpp
--- a/Engine/Plugins/FX/Niagara/Source/NiagaraEditor/Private/NiagaraNodeFunctionCall.cpp
+++ b/Engine/Plugins/FX/Niagara/Source/NiagaraEditor/Private/NiagaraNodeFunctionCall.cpp
@@ -434,7 +434,7 @@ static FText GetFormattedDeprecationMessage(const UNiagaraScript* FunctionScript
 
 	if (FunctionScript->DeprecationRecommendation != nullptr && FunctionScript->DeprecationMessage.IsEmptyOrWhitespace() == false)
 	{
-		FormatString = LOCTEXT("DeprecationErrorFmtMessageAndRecommendation", "Function call \"{NodeName}\" is deprecated. Reason: {Message}. Please use {Recommendation} instead.");
+		FormatString = LOCTEXT("DeprecationErrorFmtMessageAndRecommendation", "Function call \"{NodeName}\" is deprecated. Reason:\n{Message}.\nPlease use {Recommendation} instead.");
 	}
 	else if (FunctionScript->DeprecationRecommendation != nullptr)
 	{
@@ -442,7 +442,7 @@ static FText GetFormattedDeprecationMessage(const UNiagaraScript* FunctionScript
 	}
 	else if (FunctionScript->DeprecationMessage.IsEmptyOrWhitespace() == false)
 	{
-		FormatString = LOCTEXT("DeprecationErrorFmtMessage", "Function call \"{NodeName}\" is deprecated. Reason: {Message} ");
+		FormatString = LOCTEXT("DeprecationErrorFmtMessage", "Function call \"{NodeName}\" is deprecated. Reason:\n{Message} ");
 	}
 
 	return FText::Format(FormatString, Args);
@@ -650,7 +650,7 @@ void UNiagaraNodeFunctionCall::UpdateNodeErrorMessage()
 			{
 				FFormatNamedArguments Args;
 				Args.Add(TEXT("Message"), FunctionScript->ExperimentalMessage);
-				UEdGraphNode::NodeUpgradeMessage = FText::Format(LOCTEXT("FunctionExperimentalReason", "This function is marked as experimental, reason: {Message}."), Args);
+				UEdGraphNode::NodeUpgradeMessage = FText::Format(LOCTEXT("FunctionExperimentalReason", "This function is marked as experimental, reason:\n{Message}."), Args);
 			}
 		}
 		else
@@ -674,7 +674,7 @@ void UNiagaraNodeFunctionCall::UpdateNodeErrorMessage()
 			{
 				FFormatNamedArguments Args;
 				Args.Add(TEXT("Message"), Signature.ExperimentalMessage);
-				UEdGraphNode::NodeUpgradeMessage = FText::Format(LOCTEXT("FunctionExperimentalReason", "This function is marked as experimental, reason: {Message}."), Args);
+				UEdGraphNode::NodeUpgradeMessage = FText::Format(LOCTEXT("FunctionExperimentalReason", "This function is marked as experimental, reason:\n{Message}."), Args);
 			}
 		}
 	}
diff --git a/Engine/Plugins/FX/Niagara/Source/NiagaraEditor/Private/ViewModels/Stack/NiagaraStackErrorItem.cpp b/Engine/Plugins/FX/Niagara/Source/NiagaraEditor/Private/ViewModels/Stack/NiagaraStackErrorItem.cpp
--- a/Engine/Plugins/FX/Niagara/Source/NiagaraEditor/Private/ViewModels/Stack/NiagaraStackErrorItem.cpp
+++ b/Engine/Plugins/FX/Niagara/Source/NiagaraEditor/Private/ViewModels/Stack/NiagaraStackErrorItem.cpp
@@ -25,6 +25,11 @@ FText UNiagaraStackErrorItem::GetDisplayName() const
 	return StackIssue.GetShortDescription();
 }
 
+EStackIssueSeverity UNiagaraStackErrorItem::GetIssueSeverity() const
+{
+	return StackIssue.GetSeverity();
+}
+
 UNiagaraStackEntry::EStackRowStyle UNiagaraStackErrorItem::GetStackRowStyle() const
 {
 	return EStackRowStyle::StackIssue;
@@ -123,6 +128,11 @@ FText UNiagaraStackErrorItemLongDescription::GetDisplayName() const
 	return StackIssue.GetLongDescription();
 }
 
+EStackIssueSeverity UNiagaraStackErrorItemLongDescription::GetIssueSeverity() const
+{
+	return StackIssue.GetSeverity();
+}
+
 UNiagaraStackEntry::EStackRowStyle UNiagaraStackErrorItemLongDescription::GetStackRowStyle() const
 {
 	return EStackRowStyle::StackIssue;
@@ -150,6 +160,11 @@ FText UNiagaraStackErrorItemFix::GetDisplayName() const
 	return IssueFix.GetDescription();
 }
 
+EStackIssueSeverity UNiagaraStackErrorItemFix::GetIssueSeverity() const
+{
+	return StackIssue.GetSeverity();
+}
+
 UNiagaraStackEntry::EStackRowStyle UNiagaraStackErrorItemFix::GetStackRowStyle() const
 {
 	return EStackRowStyle::StackIssue;
diff --git a/Engine/Plugins/FX/Niagara/Source/NiagaraEditor/Private/ViewModels/Stack/NiagaraStackFunctionInput.cpp b/Engine/Plugins/FX/Niagara/Source/NiagaraEditor/Private/ViewModels/Stack/NiagaraStackFunctionInput.cpp
--- a/Engine/Plugins/FX/Niagara/Source/NiagaraEditor/Private/ViewModels/Stack/NiagaraStackFunctionInput.cpp
+++ b/Engine/Plugins/FX/Niagara/Source/NiagaraEditor/Private/ViewModels/Stack/NiagaraStackFunctionInput.cpp
@@ -461,7 +461,7 @@ void UNiagaraStackFunctionInput::RefreshChildrenInternal(const TArray<UNiagaraSt
 					if (InputValues.DynamicNode->FunctionScript->DeprecationRecommendation != nullptr &&
 						InputValues.DynamicNode->FunctionScript->DeprecationMessage.IsEmptyOrWhitespace() == false)
 					{
-						FormatString = LOCTEXT("DynamicInputScriptDeprecationMessageAndRecommendationLong", "The script asset for the assigned dynamic input {ScriptName} has been deprecated. Reason: {Message}. Suggested replacement: {Recommendation}");
+						FormatString = LOCTEXT("DynamicInputScriptDeprecationMessageAndRecommendationLong", "The script asset for the assigned dynamic input {ScriptName} has been deprecated. Reason:\n{Message}.\nSuggested replacement: {Recommendation}");
 					}
 					else if (InputValues.DynamicNode->FunctionScript->DeprecationRecommendation != nullptr)
 					{
@@ -469,7 +469,7 @@ void UNiagaraStackFunctionInput::RefreshChildrenInternal(const TArray<UNiagaraSt
 					}
 					else if (InputValues.DynamicNode->FunctionScript->DeprecationMessage.IsEmptyOrWhitespace() == false)
 					{
-						FormatString = LOCTEXT("DynamicInputScriptDeprecationMessageLong", "The script asset for the assigned dynamic input {ScriptName} has been deprecated. Reason: {Message}");
+						FormatString = LOCTEXT("DynamicInputScriptDeprecationMessageLong", "The script asset for the assigned dynamic input {ScriptName} has been deprecated. Reason:\n{Message}");
 					}
 
 					FText LongMessage = FText::Format(FormatString, Args);
diff --git a/Engine/Plugins/FX/Niagara/Source/NiagaraEditor/Private/ViewModels/Stack/NiagaraStackModuleItem.cpp b/Engine/Plugins/FX/Niagara/Source/NiagaraEditor/Private/ViewModels/Stack/NiagaraStackModuleItem.cpp
--- a/Engine/Plugins/FX/Niagara/Source/NiagaraEditor/Private/ViewModels/Stack/NiagaraStackModuleItem.cpp
+++ b/Engine/Plugins/FX/Niagara/Source/NiagaraEditor/Private/ViewModels/Stack/NiagaraStackModuleItem.cpp
@@ -326,7 +326,7 @@ void UNiagaraStackModuleItem::RefreshIssues(TArray<FStackIssue>& NewIssues)
 					if (FunctionCallNode->FunctionScript->DeprecationRecommendation != nullptr &&
 						FunctionCallNode->FunctionScript->DeprecationMessage.IsEmptyOrWhitespace() == false)
 					{
-						FormatString = LOCTEXT("ModuleScriptDeprecationMessageAndRecommendationLong", "The script asset for the assigned module {ScriptName} has been deprecated. Reason: {Message}. Suggested replacement: {Recommendation}");
+						FormatString = LOCTEXT("ModuleScriptDeprecationMessageAndRecommendationLong", "The script asset for the assigned module {ScriptName} has been deprecated. Reason:\n{Message}.\nSuggested replacement: {Recommendation}");
 					}
 					else if (FunctionCallNode->FunctionScript->DeprecationRecommendation != nullptr)
 					{
@@ -334,7 +334,7 @@ void UNiagaraStackModuleItem::RefreshIssues(TArray<FStackIssue>& NewIssues)
 					}
 					else if (FunctionCallNode->FunctionScript->DeprecationMessage.IsEmptyOrWhitespace() == false)
 					{
-						FormatString = LOCTEXT("ModuleScriptDeprecationMessageLong", "The script asset for the assigned module {ScriptName} has been deprecated. Reason: {Message}");
+						FormatString = LOCTEXT("ModuleScriptDeprecationMessageLong", "The script asset for the assigned module {ScriptName} has been deprecated. Reason:\n{Message}");
 					}
 
 					FText LongMessage = FText::Format(FormatString, Args);
@@ -392,7 +392,7 @@ void UNiagaraStackModuleItem::RefreshIssues(TArray<FStackIssue>& NewIssues)
 					FFormatNamedArguments Args;
 					Args.Add(TEXT("Module"), FText::FromString(FunctionCallNode->GetFunctionName()));
 					Args.Add(TEXT("Message"), FunctionCallNode->FunctionScript->ExperimentalMessage);
-					ErrorMessage = FText::Format(LOCTEXT("ModuleScriptExperimentalReason", "The script asset for this module is marked as experimental, reason: {Message}."), Args);
+					ErrorMessage = FText::Format(LOCTEXT("ModuleScriptExperimentalReason", "The script asset for this module is marked as experimental, reason:\n{Message}."), Args);
 				}
 
 				NewIssues.Add(FStackIssue(
diff --git a/Engine/Plugins/FX/Niagara/Source/NiagaraEditor/Public/ViewModels/Stack/NiagaraStackEntry.h b/Engine/Plugins/FX/Niagara/Source/NiagaraEditor/Public/ViewModels/Stack/NiagaraStackEntry.h
--- a/Engine/Plugins/FX/Niagara/Source/NiagaraEditor/Public/ViewModels/Stack/NiagaraStackEntry.h
+++ b/Engine/Plugins/FX/Niagara/Source/NiagaraEditor/Public/ViewModels/Stack/NiagaraStackEntry.h
@@ -18,7 +18,8 @@ enum class EStackIssueSeverity : uint8
 {
 	Error = 0,
 	Warning, 
-	Info
+	Info,
+	None
 };
 
 class FNiagaraStackEntryDragDropOp : public FDecoratedDragDropOp
@@ -357,6 +358,8 @@ public:
 
 	int32 GetTotalNumberOfErrorIssues() const;
 
+	virtual EStackIssueSeverity GetIssueSeverity() const { return EStackIssueSeverity::None; }
+
 	const TArray<FStackIssue>& GetIssues() const;
 
 	const TArray<UNiagaraStackEntry*>& GetAllChildrenWithIssues() const;
diff --git a/Engine/Plugins/FX/Niagara/Source/NiagaraEditor/Public/ViewModels/Stack/NiagaraStackErrorItem.h b/Engine/Plugins/FX/Niagara/Source/NiagaraEditor/Public/ViewModels/Stack/NiagaraStackErrorItem.h
--- a/Engine/Plugins/FX/Niagara/Source/NiagaraEditor/Public/ViewModels/Stack/NiagaraStackErrorItem.h
+++ b/Engine/Plugins/FX/Niagara/Source/NiagaraEditor/Public/ViewModels/Stack/NiagaraStackErrorItem.h
@@ -21,7 +21,8 @@ public:
 	virtual EStackRowStyle GetStackRowStyle() const override;
 	FOnIssueNotify& OnIssueModified();
 	virtual void GetSearchItems(TArray<FStackSearchItem>& SearchItems) const override;
-	
+	virtual EStackIssueSeverity GetIssueSeverity() const override;
+
 protected:
 	virtual void RefreshChildrenInternal(const TArray<UNiagaraStackEntry*>& CurrentChildren, TArray<UNiagaraStackEntry*>& NewChildren, TArray<FStackIssue>& NewIssues);
 
@@ -43,6 +44,7 @@ public:
 	void Initialize(FRequiredEntryData InRequiredEntryData, FStackIssue InStackIssue, FString InStackEditorDataKey);
 	virtual FText GetDisplayName() const override;
 	virtual EStackRowStyle GetStackRowStyle() const override;
+	virtual EStackIssueSeverity GetIssueSeverity() const override;
 
 protected:
 	FStackIssue StackIssue;
@@ -63,6 +65,7 @@ public:
 	virtual FText GetFixButtonText() const;
 	UNiagaraStackErrorItem::FOnIssueNotify& OnIssueFixed();
 	void SetFixDelegate(const FStackIssueFixDelegate& InFixDelegate);
+	virtual EStackIssueSeverity GetIssueSeverity() const override;
 
 protected:
 	FStackIssue StackIssue;
diff --git a/Engine/Plugins/FX/Niagara/Source/NiagaraEditorWidgets/Private/NiagaraEditorWidgetsStyle.cpp b/Engine/Plugins/FX/Niagara/Source/NiagaraEditorWidgets/Private/NiagaraEditorWidgetsStyle.cpp
--- a/Engine/Plugins/FX/Niagara/Source/NiagaraEditorWidgets/Private/NiagaraEditorWidgetsStyle.cpp
+++ b/Engine/Plugins/FX/Niagara/Source/NiagaraEditorWidgets/Private/NiagaraEditorWidgetsStyle.cpp
@@ -162,7 +162,10 @@ TSharedRef< FSlateStyleSet > FNiagaraEditorWidgetsStyle::Create()
 	Style->Set("NiagaraEditor.Stack.Item.ContentBackgroundColor", FLinearColor(FColor(62, 62, 62)));
 	Style->Set("NiagaraEditor.Stack.Item.ContentAdvancedBackgroundColor", FLinearColor(FColor(53, 53, 53)));
 	Style->Set("NiagaraEditor.Stack.Item.FooterBackgroundColor", FLinearColor(FColor(75, 75, 75)));
-	Style->Set("NiagaraEditor.Stack.Item.IssueBackgroundColor", FLinearColor(FColor(120, 120, 62)));
+	Style->Set("NiagaraEditor.Stack.Item.InfoBackgroundColor", FLinearColor(FColor(68, 100, 106)));
+	Style->Set("NiagaraEditor.Stack.Item.WarningBackgroundColor", FLinearColor(FColor(97, 97, 68)));
+	Style->Set("NiagaraEditor.Stack.Item.ErrorBackgroundColor", FLinearColor(FColor(126, 78, 68)));
+
 	Style->Set("NiagaraEditor.Stack.UnknownColor", FLinearColor(1, 0, 1));
 
 	Style->Set("NiagaraEditor.Stack.ItemHeaderFooter.BackgroundBrush", new FSlateColorBrush(FLinearColor(FColor(20, 20, 20))));
diff --git a/Engine/Plugins/FX/Niagara/Source/NiagaraEditorWidgets/Private/SNiagaraStack.cpp b/Engine/Plugins/FX/Niagara/Source/NiagaraEditorWidgets/Private/SNiagaraStack.cpp
--- a/Engine/Plugins/FX/Niagara/Source/NiagaraEditorWidgets/Private/SNiagaraStack.cpp
+++ b/Engine/Plugins/FX/Niagara/Source/NiagaraEditorWidgets/Private/SNiagaraStack.cpp
@@ -850,8 +850,23 @@ TSharedRef<SNiagaraStackTableRow> SNiagaraStack::ConstructContainerForItem(UNiag
 		bShowExecutionCategoryIcon = false;
 		break;
 	case UNiagaraStackEntry::EStackRowStyle::StackIssue:
+
+		switch (Item->GetIssueSeverity())
+		{
+			case EStackIssueSeverity::Error:
+				ItemBackgroundColor = FNiagaraEditorWidgetsStyle::Get().GetColor("NiagaraEditor.Stack.Item.ErrorBackgroundColor");
+				break;
+			case EStackIssueSeverity::Warning:
+				ItemBackgroundColor = FNiagaraEditorWidgetsStyle::Get().GetColor("NiagaraEditor.Stack.Item.WarningBackgroundColor");
+				break;
+			case EStackIssueSeverity::Info:
+				ItemBackgroundColor = FNiagaraEditorWidgetsStyle::Get().GetColor("NiagaraEditor.Stack.Item.InfoBackgroundColor");
+				break;
+			default:
+				checkf(false, TEXT("Issue severity not set for stack issue."));
+		}
+
 		ContentPadding = FMargin(LeftContentPadding, 3, RightContentPadding, 3);
-		ItemBackgroundColor = FNiagaraEditorWidgetsStyle::Get().GetColor("NiagaraEditor.Stack.Item.IssueBackgroundColor");
 		bIsCategoryIconHighlighted = false;
 		bShowExecutionCategoryIcon = false;
 		break;
