Renamed the direct read functions to be consistent with the type names used throughout Niagara, and to make it clear which versions take an ID and which take an index.

#rnx

diff --git a/Engine/Plugins/FX/Niagara/Source/Niagara/Private/NiagaraDataInterfaceParticleRead.cpp b/Engine/Plugins/FX/Niagara/Source/Niagara/Private/NiagaraDataInterfaceParticleRead.cpp
--- a/Engine/Plugins/FX/Niagara/Source/Niagara/Private/NiagaraDataInterfaceParticleRead.cpp
+++ b/Engine/Plugins/FX/Niagara/Source/Niagara/Private/NiagaraDataInterfaceParticleRead.cpp
@@ -14,25 +14,25 @@ static const FName GetSpawnedIDAtIndexFunctionName("Get Spawned ID At Index");
 static const FName GetNumParticlesFunctionName("Get Num Particles");
 static const FName GetParticleIndexFunctionName("Get Particle Index");
 
-static const FName GetIntAttributeFunctionName("Get int Attribute");
-static const FName GetBoolAttributeFunctionName("Get bool Attribute");
-static const FName GetFloatAttributeFunctionName("Get float Attribute");
-static const FName GetVec2AttributeFunctionName("Get Vector2 Attribute");
-static const FName GetVec3AttributeFunctionName("Get Vector3 Attribute");
-static const FName GetVec4AttributeFunctionName("Get Vector4 Attribute");
-static const FName GetColorAttributeFunctionName("Get Color Attribute");
-static const FName GetQuatAttributeFunctionName("Get Quaternion Attribute");
-static const FName GetIDAttributeFunctionName("Get ID Attribute");
-
-static const FName GetIntAttributeByIndexFunctionName("Get int Attribute By Index");
-static const FName GetBoolAttributeByIndexFunctionName("Get bool Attribute By Index");
-static const FName GetFloatAttributeByIndexFunctionName("Get float Attribute By Index");
-static const FName GetVec2AttributeByIndexFunctionName("Get Vector2 Attribute By Index");
-static const FName GetVec3AttributeByIndexFunctionName("Get Vector3 Attribute By Index");
-static const FName GetVec4AttributeByIndexFunctionName("Get Vector4 Attribute By Index");
-static const FName GetColorAttributeByIndexFunctionName("Get Color Attribute By Index");
-static const FName GetQuatAttributeByIndexFunctionName("Get Quaternion Attribute By Index");
-static const FName GetIDAttributeByIndexFunctionName("Get ID Attribute By Index");
+static const FName GetIntByIDFunctionName("Get Int By ID");
+static const FName GetBoolByIDFunctionName("Get Bool By ID");
+static const FName GetFloatByIDFunctionName("Get Float By ID");
+static const FName GetVec2ByIDFunctionName("Get Vector 2D By ID");
+static const FName GetVec3ByIDFunctionName("Get Vector By ID");
+static const FName GetVec4ByIDFunctionName("Get Vector 4 By ID");
+static const FName GetColorByIDFunctionName("Get Color By ID");
+static const FName GetQuatByIDFunctionName("Get Quaternion By ID");
+static const FName GetIDByIDFunctionName("Get ID By ID");
+
+static const FName GetIntByIndexFunctionName("Get Int By Index");
+static const FName GetBoolByIndexFunctionName("Get Bool By Index");
+static const FName GetFloatByIndexFunctionName("Get Float By Index");
+static const FName GetVec2ByIndexFunctionName("Get Vector 2D By Index");
+static const FName GetVec3ByIndexFunctionName("Get Vector By Index");
+static const FName GetVec4ByIndexFunctionName("Get Vector 4 By Index");
+static const FName GetColorByIndexFunctionName("Get Color By Index");
+static const FName GetQuatByIndexFunctionName("Get Quaternion By Index");
+static const FName GetIDByIndexFunctionName("Get ID By Index");
 
 static const FString NumSpawnedParticlesBaseName(TEXT("NumSpawnedParticles_"));
 static const FString SpawnedParticlesAcquireTagBaseName(TEXT("SpawnedParticlesAcquireTag_"));
@@ -180,15 +180,15 @@ struct FNiagaraDataInterfaceParametersCS_ParticleRead : public FNiagaraDataInter
 
 	ENiagaraParticleDataValueType GetValueTypeFromFuncName(const FName& FuncName)
 	{
-		if (FuncName == GetIntAttributeFunctionName || FuncName == GetIntAttributeByIndexFunctionName) return ENiagaraParticleDataValueType::Int;
-		if (FuncName == GetBoolAttributeFunctionName || FuncName == GetBoolAttributeByIndexFunctionName) return ENiagaraParticleDataValueType::Bool;
-		if (FuncName == GetFloatAttributeFunctionName || FuncName == GetFloatAttributeByIndexFunctionName) return ENiagaraParticleDataValueType::Float;
-		if (FuncName == GetVec2AttributeFunctionName || FuncName == GetVec2AttributeByIndexFunctionName) return ENiagaraParticleDataValueType::Vec2;
-		if (FuncName == GetVec3AttributeFunctionName || FuncName == GetVec3AttributeByIndexFunctionName) return ENiagaraParticleDataValueType::Vec3;
-		if (FuncName == GetVec4AttributeFunctionName || FuncName == GetVec4AttributeByIndexFunctionName) return ENiagaraParticleDataValueType::Vec4;
-		if (FuncName == GetColorAttributeFunctionName || FuncName == GetColorAttributeByIndexFunctionName) return ENiagaraParticleDataValueType::Color;
-		if (FuncName == GetQuatAttributeFunctionName || FuncName == GetQuatAttributeByIndexFunctionName) return ENiagaraParticleDataValueType::Quat;
-		if (FuncName == GetIDAttributeFunctionName || FuncName == GetIDAttributeByIndexFunctionName) return ENiagaraParticleDataValueType::ID;
+		if (FuncName == GetIntByIDFunctionName || FuncName == GetIntByIndexFunctionName) return ENiagaraParticleDataValueType::Int;
+		if (FuncName == GetBoolByIDFunctionName || FuncName == GetBoolByIndexFunctionName) return ENiagaraParticleDataValueType::Bool;
+		if (FuncName == GetFloatByIDFunctionName || FuncName == GetFloatByIndexFunctionName) return ENiagaraParticleDataValueType::Float;
+		if (FuncName == GetVec2ByIDFunctionName || FuncName == GetVec2ByIndexFunctionName) return ENiagaraParticleDataValueType::Vec2;
+		if (FuncName == GetVec3ByIDFunctionName || FuncName == GetVec3ByIndexFunctionName) return ENiagaraParticleDataValueType::Vec3;
+		if (FuncName == GetVec4ByIDFunctionName || FuncName == GetVec4ByIndexFunctionName) return ENiagaraParticleDataValueType::Vec4;
+		if (FuncName == GetColorByIDFunctionName || FuncName == GetColorByIndexFunctionName) return ENiagaraParticleDataValueType::Color;
+		if (FuncName == GetQuatByIDFunctionName || FuncName == GetQuatByIndexFunctionName) return ENiagaraParticleDataValueType::Quat;
+		if (FuncName == GetIDByIDFunctionName || FuncName == GetIDByIndexFunctionName) return ENiagaraParticleDataValueType::ID;
 		return ENiagaraParticleDataValueType::Invalid;
 	}
 
@@ -673,7 +673,7 @@ void UNiagaraDataInterfaceParticleRead::GetFunctions(TArray<FNiagaraFunctionSign
 	//
 	{
 		FNiagaraFunctionSignature Sig;
-		Sig.Name = GetIntAttributeFunctionName;
+		Sig.Name = GetIntByIDFunctionName;
 		Sig.Inputs.Add(FNiagaraVariable(FNiagaraTypeDefinition(GetClass()), TEXT("Particle Reader")));
 		Sig.Inputs.Add(FNiagaraVariable(FNiagaraTypeDefinition::GetIDDef(), TEXT("Particle ID")));
 
@@ -689,7 +689,7 @@ void UNiagaraDataInterfaceParticleRead::GetFunctions(TArray<FNiagaraFunctionSign
 
 	{
 		FNiagaraFunctionSignature Sig;
-		Sig.Name = GetBoolAttributeFunctionName;
+		Sig.Name = GetBoolByIDFunctionName;
 		Sig.Inputs.Add(FNiagaraVariable(FNiagaraTypeDefinition(GetClass()), TEXT("Particle Reader")));
 		Sig.Inputs.Add(FNiagaraVariable(FNiagaraTypeDefinition::GetIDDef(), TEXT("Particle ID")));
 
@@ -705,7 +705,7 @@ void UNiagaraDataInterfaceParticleRead::GetFunctions(TArray<FNiagaraFunctionSign
 
 	{
 		FNiagaraFunctionSignature Sig;
-		Sig.Name = GetFloatAttributeFunctionName;
+		Sig.Name = GetFloatByIDFunctionName;
 		Sig.Inputs.Add(FNiagaraVariable(FNiagaraTypeDefinition(GetClass()), TEXT("Particle Reader")));
 		Sig.Inputs.Add(FNiagaraVariable(FNiagaraTypeDefinition::GetIDDef(), TEXT("Particle ID")));
 
@@ -721,7 +721,7 @@ void UNiagaraDataInterfaceParticleRead::GetFunctions(TArray<FNiagaraFunctionSign
 
 	{
 		FNiagaraFunctionSignature Sig;
-		Sig.Name = GetVec2AttributeFunctionName;
+		Sig.Name = GetVec2ByIDFunctionName;
 		Sig.Inputs.Add(FNiagaraVariable(FNiagaraTypeDefinition(GetClass()), TEXT("Particle Reader")));
 		Sig.Inputs.Add(FNiagaraVariable(FNiagaraTypeDefinition::GetIDDef(), TEXT("Particle ID")));
 
@@ -737,7 +737,7 @@ void UNiagaraDataInterfaceParticleRead::GetFunctions(TArray<FNiagaraFunctionSign
 
 	{
 		FNiagaraFunctionSignature Sig;
-		Sig.Name = GetVec3AttributeFunctionName;
+		Sig.Name = GetVec3ByIDFunctionName;
 		Sig.Inputs.Add(FNiagaraVariable(FNiagaraTypeDefinition(GetClass()), TEXT("Particle Reader")));
 		Sig.Inputs.Add(FNiagaraVariable(FNiagaraTypeDefinition::GetIDDef(), TEXT("Particle ID")));
 
@@ -753,7 +753,7 @@ void UNiagaraDataInterfaceParticleRead::GetFunctions(TArray<FNiagaraFunctionSign
 
 	{
 		FNiagaraFunctionSignature Sig;
-		Sig.Name = GetVec4AttributeFunctionName;
+		Sig.Name = GetVec4ByIDFunctionName;
 		Sig.Inputs.Add(FNiagaraVariable(FNiagaraTypeDefinition(GetClass()), TEXT("Particle Reader")));
 		Sig.Inputs.Add(FNiagaraVariable(FNiagaraTypeDefinition::GetIDDef(), TEXT("Particle ID")));
 
@@ -769,7 +769,7 @@ void UNiagaraDataInterfaceParticleRead::GetFunctions(TArray<FNiagaraFunctionSign
 
 	{
 		FNiagaraFunctionSignature Sig;
-		Sig.Name = GetColorAttributeFunctionName;
+		Sig.Name = GetColorByIDFunctionName;
 		Sig.Inputs.Add(FNiagaraVariable(FNiagaraTypeDefinition(GetClass()), TEXT("Particle Reader")));
 		Sig.Inputs.Add(FNiagaraVariable(FNiagaraTypeDefinition::GetIDDef(), TEXT("Particle ID")));
 
@@ -785,7 +785,7 @@ void UNiagaraDataInterfaceParticleRead::GetFunctions(TArray<FNiagaraFunctionSign
 
 	{
 		FNiagaraFunctionSignature Sig;
-		Sig.Name = GetQuatAttributeFunctionName;
+		Sig.Name = GetQuatByIDFunctionName;
 		Sig.Inputs.Add(FNiagaraVariable(FNiagaraTypeDefinition(GetClass()), TEXT("Particle Reader")));
 		Sig.Inputs.Add(FNiagaraVariable(FNiagaraTypeDefinition::GetIDDef(), TEXT("Particle ID")));
 
@@ -801,7 +801,7 @@ void UNiagaraDataInterfaceParticleRead::GetFunctions(TArray<FNiagaraFunctionSign
 
 	{
 		FNiagaraFunctionSignature Sig;
-		Sig.Name = GetIDAttributeFunctionName;
+		Sig.Name = GetIDByIDFunctionName;
 		Sig.Inputs.Add(FNiagaraVariable(FNiagaraTypeDefinition(GetClass()), TEXT("Particle Reader")));
 		Sig.Inputs.Add(FNiagaraVariable(FNiagaraTypeDefinition::GetIDDef(), TEXT("Particle ID")));
 
@@ -820,7 +820,7 @@ void UNiagaraDataInterfaceParticleRead::GetFunctions(TArray<FNiagaraFunctionSign
 	//
 	{
 		FNiagaraFunctionSignature Sig;
-		Sig.Name = GetIntAttributeByIndexFunctionName;
+		Sig.Name = GetIntByIndexFunctionName;
 		Sig.Inputs.Add(FNiagaraVariable(FNiagaraTypeDefinition(GetClass()), TEXT("Particle Reader")));
 		Sig.Inputs.Add(FNiagaraVariable(FNiagaraTypeDefinition::GetIntDef(), TEXT("Particle Index")));
 
@@ -836,7 +836,7 @@ void UNiagaraDataInterfaceParticleRead::GetFunctions(TArray<FNiagaraFunctionSign
 
 	{
 		FNiagaraFunctionSignature Sig;
-		Sig.Name = GetBoolAttributeByIndexFunctionName;
+		Sig.Name = GetBoolByIndexFunctionName;
 		Sig.Inputs.Add(FNiagaraVariable(FNiagaraTypeDefinition(GetClass()), TEXT("Particle Reader")));
 		Sig.Inputs.Add(FNiagaraVariable(FNiagaraTypeDefinition::GetIntDef(), TEXT("Particle Index")));
 
@@ -852,7 +852,7 @@ void UNiagaraDataInterfaceParticleRead::GetFunctions(TArray<FNiagaraFunctionSign
 
 	{
 		FNiagaraFunctionSignature Sig;
-		Sig.Name = GetFloatAttributeByIndexFunctionName;
+		Sig.Name = GetFloatByIndexFunctionName;
 		Sig.Inputs.Add(FNiagaraVariable(FNiagaraTypeDefinition(GetClass()), TEXT("Particle Reader")));
 		Sig.Inputs.Add(FNiagaraVariable(FNiagaraTypeDefinition::GetIntDef(), TEXT("Particle Index")));
 
@@ -868,7 +868,7 @@ void UNiagaraDataInterfaceParticleRead::GetFunctions(TArray<FNiagaraFunctionSign
 
 	{
 		FNiagaraFunctionSignature Sig;
-		Sig.Name = GetVec2AttributeByIndexFunctionName;
+		Sig.Name = GetVec2ByIndexFunctionName;
 		Sig.Inputs.Add(FNiagaraVariable(FNiagaraTypeDefinition(GetClass()), TEXT("Particle Reader")));
 		Sig.Inputs.Add(FNiagaraVariable(FNiagaraTypeDefinition::GetIntDef(), TEXT("Particle Index")));
 
@@ -884,7 +884,7 @@ void UNiagaraDataInterfaceParticleRead::GetFunctions(TArray<FNiagaraFunctionSign
 
 	{
 		FNiagaraFunctionSignature Sig;
-		Sig.Name = GetVec3AttributeByIndexFunctionName;
+		Sig.Name = GetVec3ByIndexFunctionName;
 		Sig.Inputs.Add(FNiagaraVariable(FNiagaraTypeDefinition(GetClass()), TEXT("Particle Reader")));
 		Sig.Inputs.Add(FNiagaraVariable(FNiagaraTypeDefinition::GetIntDef(), TEXT("Particle Index")));
 
@@ -900,7 +900,7 @@ void UNiagaraDataInterfaceParticleRead::GetFunctions(TArray<FNiagaraFunctionSign
 
 	{
 		FNiagaraFunctionSignature Sig;
-		Sig.Name = GetVec4AttributeByIndexFunctionName;
+		Sig.Name = GetVec4ByIndexFunctionName;
 		Sig.Inputs.Add(FNiagaraVariable(FNiagaraTypeDefinition(GetClass()), TEXT("Particle Reader")));
 		Sig.Inputs.Add(FNiagaraVariable(FNiagaraTypeDefinition::GetIntDef(), TEXT("Particle Index")));
 
@@ -916,7 +916,7 @@ void UNiagaraDataInterfaceParticleRead::GetFunctions(TArray<FNiagaraFunctionSign
 
 	{
 		FNiagaraFunctionSignature Sig;
-		Sig.Name = GetColorAttributeByIndexFunctionName;
+		Sig.Name = GetColorByIndexFunctionName;
 		Sig.Inputs.Add(FNiagaraVariable(FNiagaraTypeDefinition(GetClass()), TEXT("Particle Reader")));
 		Sig.Inputs.Add(FNiagaraVariable(FNiagaraTypeDefinition::GetIntDef(), TEXT("Particle Index")));
 
@@ -932,7 +932,7 @@ void UNiagaraDataInterfaceParticleRead::GetFunctions(TArray<FNiagaraFunctionSign
 
 	{
 		FNiagaraFunctionSignature Sig;
-		Sig.Name = GetQuatAttributeByIndexFunctionName;
+		Sig.Name = GetQuatByIndexFunctionName;
 		Sig.Inputs.Add(FNiagaraVariable(FNiagaraTypeDefinition(GetClass()), TEXT("Particle Reader")));
 		Sig.Inputs.Add(FNiagaraVariable(FNiagaraTypeDefinition::GetIntDef(), TEXT("Particle Index")));
 
@@ -948,7 +948,7 @@ void UNiagaraDataInterfaceParticleRead::GetFunctions(TArray<FNiagaraFunctionSign
 
 	{
 		FNiagaraFunctionSignature Sig;
-		Sig.Name = GetIDAttributeByIndexFunctionName;
+		Sig.Name = GetIDByIndexFunctionName;
 		Sig.Inputs.Add(FNiagaraVariable(FNiagaraTypeDefinition(GetClass()), TEXT("Particle Reader")));
 		Sig.Inputs.Add(FNiagaraVariable(FNiagaraTypeDefinition::GetIntDef(), TEXT("Particle Index")));
 
@@ -1038,7 +1038,7 @@ void UNiagaraDataInterfaceParticleRead::GetVMExternalFunction(const FVMExternalF
 	//
 	// Get attribute by ID
 	//
-	if (BindingInfo.Name == GetIntAttributeFunctionName)
+	if (BindingInfo.Name == GetIntByIDFunctionName)
 	{
 		if (HasMatchingVariable(EmitterVariables, AttributeToRead, FNiagaraTypeDefinition::GetIntDef()))
 		{
@@ -1046,7 +1046,7 @@ void UNiagaraDataInterfaceParticleRead::GetVMExternalFunction(const FVMExternalF
 			bBindSuccessful = true;
 		}
 	}
-	else if (BindingInfo.Name == GetBoolAttributeFunctionName)
+	else if (BindingInfo.Name == GetBoolByIDFunctionName)
 	{
 		if (HasMatchingVariable(EmitterVariables, AttributeToRead, FNiagaraTypeDefinition::GetBoolDef()))
 		{
@@ -1054,7 +1054,7 @@ void UNiagaraDataInterfaceParticleRead::GetVMExternalFunction(const FVMExternalF
 			bBindSuccessful = true;
 		}
 	}
-	else if (BindingInfo.Name == GetFloatAttributeFunctionName)
+	else if (BindingInfo.Name == GetFloatByIDFunctionName)
 	{
 		if (HasMatchingVariable(EmitterVariables, AttributeToRead, FNiagaraTypeDefinition::GetFloatDef())
 			|| HasMatchingVariable(EmitterVariables, AttributeToRead, FNiagaraTypeDefinition::GetHalfDef()))
@@ -1063,7 +1063,7 @@ void UNiagaraDataInterfaceParticleRead::GetVMExternalFunction(const FVMExternalF
 			bBindSuccessful = true;
 		}
 	}
-	else if (BindingInfo.Name == GetVec2AttributeFunctionName)
+	else if (BindingInfo.Name == GetVec2ByIDFunctionName)
 	{
 		if (HasMatchingVariable(EmitterVariables, AttributeToRead, FNiagaraTypeDefinition::GetVec2Def())
 			|| HasMatchingVariable(EmitterVariables, AttributeToRead, FNiagaraTypeDefinition::GetHalfVec2Def()))
@@ -1072,7 +1072,7 @@ void UNiagaraDataInterfaceParticleRead::GetVMExternalFunction(const FVMExternalF
 			bBindSuccessful = true;
 		}
 	}
-	else if (BindingInfo.Name == GetVec3AttributeFunctionName)
+	else if (BindingInfo.Name == GetVec3ByIDFunctionName)
 	{
 		if (HasMatchingVariable(EmitterVariables, AttributeToRead, FNiagaraTypeDefinition::GetVec3Def())
 			|| HasMatchingVariable(EmitterVariables, AttributeToRead, FNiagaraTypeDefinition::GetHalfVec3Def()))
@@ -1081,7 +1081,7 @@ void UNiagaraDataInterfaceParticleRead::GetVMExternalFunction(const FVMExternalF
 			bBindSuccessful = true;
 		}
 	}
-	else if (BindingInfo.Name == GetVec4AttributeFunctionName)
+	else if (BindingInfo.Name == GetVec4ByIDFunctionName)
 	{
 		if (HasMatchingVariable(EmitterVariables, AttributeToRead, FNiagaraTypeDefinition::GetVec4Def())
 			|| HasMatchingVariable(EmitterVariables, AttributeToRead, FNiagaraTypeDefinition::GetHalfVec4Def()))
@@ -1090,7 +1090,7 @@ void UNiagaraDataInterfaceParticleRead::GetVMExternalFunction(const FVMExternalF
 			bBindSuccessful = true;
 		}
 	}
-	else if (BindingInfo.Name == GetColorAttributeFunctionName)
+	else if (BindingInfo.Name == GetColorByIDFunctionName)
 	{
 		if (HasMatchingVariable(EmitterVariables, AttributeToRead, FNiagaraTypeDefinition::GetColorDef())
 			|| HasMatchingVariable(EmitterVariables, AttributeToRead, FNiagaraTypeDefinition::GetHalfVec4Def()))
@@ -1099,7 +1099,7 @@ void UNiagaraDataInterfaceParticleRead::GetVMExternalFunction(const FVMExternalF
 			bBindSuccessful = true;
 		}
 	}
-	else if (BindingInfo.Name == GetQuatAttributeFunctionName)
+	else if (BindingInfo.Name == GetQuatByIDFunctionName)
 	{
 		if (HasMatchingVariable(EmitterVariables, AttributeToRead, FNiagaraTypeDefinition::GetQuatDef())
 			|| HasMatchingVariable(EmitterVariables, AttributeToRead, FNiagaraTypeDefinition::GetHalfVec4Def()))
@@ -1108,7 +1108,7 @@ void UNiagaraDataInterfaceParticleRead::GetVMExternalFunction(const FVMExternalF
 			bBindSuccessful = true;
 		}
 	}
-	else if (BindingInfo.Name == GetIDAttributeFunctionName)
+	else if (BindingInfo.Name == GetIDByIDFunctionName)
 	{
 		FNiagaraVariable VariableToRead(FNiagaraTypeDefinition::GetIDDef(), AttributeToRead);
 		if (PIData->EmitterInstance->GetData().GetVariables().Find(VariableToRead) != INDEX_NONE)
@@ -1120,7 +1120,7 @@ void UNiagaraDataInterfaceParticleRead::GetVMExternalFunction(const FVMExternalF
 	//
 	// Get attribute by index
 	//
-	else if (BindingInfo.Name == GetIntAttributeByIndexFunctionName)
+	else if (BindingInfo.Name == GetIntByIndexFunctionName)
 	{
 		if (HasMatchingVariable(EmitterVariables, AttributeToRead, FNiagaraTypeDefinition::GetIntDef()))
 		{
@@ -1128,7 +1128,7 @@ void UNiagaraDataInterfaceParticleRead::GetVMExternalFunction(const FVMExternalF
 			bBindSuccessful = true;
 		}
 	}
-	else if (BindingInfo.Name == GetBoolAttributeByIndexFunctionName)
+	else if (BindingInfo.Name == GetBoolByIndexFunctionName)
 	{
 		if (HasMatchingVariable(EmitterVariables, AttributeToRead, FNiagaraTypeDefinition::GetBoolDef()))
 		{
@@ -1136,7 +1136,7 @@ void UNiagaraDataInterfaceParticleRead::GetVMExternalFunction(const FVMExternalF
 			bBindSuccessful = true;
 		}
 	}
-	else if (BindingInfo.Name == GetFloatAttributeByIndexFunctionName)
+	else if (BindingInfo.Name == GetFloatByIndexFunctionName)
 	{
 		if (HasMatchingVariable(EmitterVariables, AttributeToRead, FNiagaraTypeDefinition::GetFloatDef())
 			|| HasMatchingVariable(EmitterVariables, AttributeToRead, FNiagaraTypeDefinition::GetHalfDef()))
@@ -1145,7 +1145,7 @@ void UNiagaraDataInterfaceParticleRead::GetVMExternalFunction(const FVMExternalF
 			bBindSuccessful = true;
 		}
 	}
-	else if (BindingInfo.Name == GetVec2AttributeByIndexFunctionName)
+	else if (BindingInfo.Name == GetVec2ByIndexFunctionName)
 	{
 		if (HasMatchingVariable(EmitterVariables, AttributeToRead, FNiagaraTypeDefinition::GetVec2Def())
 			|| HasMatchingVariable(EmitterVariables, AttributeToRead, FNiagaraTypeDefinition::GetHalfVec2Def()))
@@ -1154,7 +1154,7 @@ void UNiagaraDataInterfaceParticleRead::GetVMExternalFunction(const FVMExternalF
 			bBindSuccessful = true;
 		}
 	}
-	else if (BindingInfo.Name == GetVec3AttributeByIndexFunctionName)
+	else if (BindingInfo.Name == GetVec3ByIndexFunctionName)
 	{
 		if (HasMatchingVariable(EmitterVariables, AttributeToRead, FNiagaraTypeDefinition::GetVec3Def())
 			|| HasMatchingVariable(EmitterVariables, AttributeToRead, FNiagaraTypeDefinition::GetHalfVec3Def()))
@@ -1163,7 +1163,7 @@ void UNiagaraDataInterfaceParticleRead::GetVMExternalFunction(const FVMExternalF
 			bBindSuccessful = true;
 		}
 	}
-	else if (BindingInfo.Name == GetVec4AttributeByIndexFunctionName)
+	else if (BindingInfo.Name == GetVec4ByIndexFunctionName)
 	{
 		if (HasMatchingVariable(EmitterVariables, AttributeToRead, FNiagaraTypeDefinition::GetVec4Def())
 			|| HasMatchingVariable(EmitterVariables, AttributeToRead, FNiagaraTypeDefinition::GetHalfVec4Def()))
@@ -1172,7 +1172,7 @@ void UNiagaraDataInterfaceParticleRead::GetVMExternalFunction(const FVMExternalF
 			bBindSuccessful = true;
 		}
 	}
-	else if (BindingInfo.Name == GetColorAttributeByIndexFunctionName)
+	else if (BindingInfo.Name == GetColorByIndexFunctionName)
 	{
 		if (HasMatchingVariable(EmitterVariables, AttributeToRead, FNiagaraTypeDefinition::GetColorDef())
 			|| HasMatchingVariable(EmitterVariables, AttributeToRead, FNiagaraTypeDefinition::GetHalfVec4Def()))
@@ -1181,7 +1181,7 @@ void UNiagaraDataInterfaceParticleRead::GetVMExternalFunction(const FVMExternalF
 			bBindSuccessful = true;
 		}
 	}
-	else if (BindingInfo.Name == GetQuatAttributeByIndexFunctionName)
+	else if (BindingInfo.Name == GetQuatByIndexFunctionName)
 	{
 		if (HasMatchingVariable(EmitterVariables, AttributeToRead, FNiagaraTypeDefinition::GetQuatDef())
 			|| HasMatchingVariable(EmitterVariables, AttributeToRead, FNiagaraTypeDefinition::GetHalfVec4Def()))
@@ -1190,7 +1190,7 @@ void UNiagaraDataInterfaceParticleRead::GetVMExternalFunction(const FVMExternalF
 			bBindSuccessful = true;
 		}
 	}
-	else if (BindingInfo.Name == GetIDAttributeByIndexFunctionName)
+	else if (BindingInfo.Name == GetIDByIndexFunctionName)
 	{
 	FNiagaraVariable VariableToRead(FNiagaraTypeDefinition::GetIDDef(), AttributeToRead);
 	if (PIData->EmitterInstance->GetData().GetVariables().Find(VariableToRead) != INDEX_NONE)
@@ -2131,47 +2131,47 @@ bool UNiagaraDataInterfaceParticleRead::GetFunctionHLSL(const FNiagaraDataInterf
 	//
 	// Get attribute by ID
 	//
-	if (FunctionInfo.DefinitionName == GetIntAttributeFunctionName)
+	if (FunctionInfo.DefinitionName == GetIntByIDFunctionName)
 	{
 		return GenerateGetFunctionHLSL(ParamInfo, FunctionInfo, FunctionInstanceIndex, ENiagaraParticleDataComponentType::Int, 1, false, OutHLSL);
 	}
 
-	if (FunctionInfo.DefinitionName == GetBoolAttributeFunctionName)
+	if (FunctionInfo.DefinitionName == GetBoolByIDFunctionName)
 	{
 		return GenerateGetFunctionHLSL(ParamInfo, FunctionInfo, FunctionInstanceIndex, ENiagaraParticleDataComponentType::Bool, 1, false, OutHLSL);
 	}
 
-	if (FunctionInfo.DefinitionName == GetFloatAttributeFunctionName)
+	if (FunctionInfo.DefinitionName == GetFloatByIDFunctionName)
 	{
 		return GenerateGetFunctionHLSL(ParamInfo, FunctionInfo, FunctionInstanceIndex, ENiagaraParticleDataComponentType::Float, 1, false, OutHLSL);
 	}
 
-	if (FunctionInfo.DefinitionName == GetVec2AttributeFunctionName)
+	if (FunctionInfo.DefinitionName == GetVec2ByIDFunctionName)
 	{
 		return GenerateGetFunctionHLSL(ParamInfo, FunctionInfo, FunctionInstanceIndex, ENiagaraParticleDataComponentType::Float, 2, false, OutHLSL);
 	}
 
-	if (FunctionInfo.DefinitionName == GetVec3AttributeFunctionName)
+	if (FunctionInfo.DefinitionName == GetVec3ByIDFunctionName)
 	{
 		return GenerateGetFunctionHLSL(ParamInfo, FunctionInfo, FunctionInstanceIndex, ENiagaraParticleDataComponentType::Float, 3, false, OutHLSL);
 	}
 
-	if (FunctionInfo.DefinitionName == GetVec4AttributeFunctionName)
+	if (FunctionInfo.DefinitionName == GetVec4ByIDFunctionName)
 	{
 		return GenerateGetFunctionHLSL(ParamInfo, FunctionInfo, FunctionInstanceIndex, ENiagaraParticleDataComponentType::Float, 4, false, OutHLSL);
 	}
 
-	if (FunctionInfo.DefinitionName == GetColorAttributeFunctionName)
+	if (FunctionInfo.DefinitionName == GetColorByIDFunctionName)
 	{
 		return GenerateGetFunctionHLSL(ParamInfo, FunctionInfo, FunctionInstanceIndex, ENiagaraParticleDataComponentType::Float, 4, false, OutHLSL);
 	}
 
-	if (FunctionInfo.DefinitionName == GetQuatAttributeFunctionName)
+	if (FunctionInfo.DefinitionName == GetQuatByIDFunctionName)
 	{
 		return GenerateGetFunctionHLSL(ParamInfo, FunctionInfo, FunctionInstanceIndex, ENiagaraParticleDataComponentType::Float, 4, false, OutHLSL);
 	}
 
-	if (FunctionInfo.DefinitionName == GetIDAttributeFunctionName)
+	if (FunctionInfo.DefinitionName == GetIDByIDFunctionName)
 	{
 		return GenerateGetFunctionHLSL(ParamInfo, FunctionInfo, FunctionInstanceIndex, ENiagaraParticleDataComponentType::ID, 1, false, OutHLSL);
 	}
@@ -2179,47 +2179,47 @@ bool UNiagaraDataInterfaceParticleRead::GetFunctionHLSL(const FNiagaraDataInterf
 	//
 	// Get attribute by index
 	//
-	if (FunctionInfo.DefinitionName == GetIntAttributeByIndexFunctionName)
+	if (FunctionInfo.DefinitionName == GetIntByIndexFunctionName)
 	{
 		return GenerateGetFunctionHLSL(ParamInfo, FunctionInfo, FunctionInstanceIndex, ENiagaraParticleDataComponentType::Int, 1, true, OutHLSL);
 	}
 
-	if (FunctionInfo.DefinitionName == GetBoolAttributeByIndexFunctionName)
+	if (FunctionInfo.DefinitionName == GetBoolByIndexFunctionName)
 	{
 		return GenerateGetFunctionHLSL(ParamInfo, FunctionInfo, FunctionInstanceIndex, ENiagaraParticleDataComponentType::Bool, 1, true, OutHLSL);
 	}
 
-	if (FunctionInfo.DefinitionName == GetFloatAttributeByIndexFunctionName)
+	if (FunctionInfo.DefinitionName == GetFloatByIndexFunctionName)
 	{
 		return GenerateGetFunctionHLSL(ParamInfo, FunctionInfo, FunctionInstanceIndex, ENiagaraParticleDataComponentType::Float, 1, true, OutHLSL);
 	}
 
-	if (FunctionInfo.DefinitionName == GetVec2AttributeByIndexFunctionName)
+	if (FunctionInfo.DefinitionName == GetVec2ByIndexFunctionName)
 	{
 		return GenerateGetFunctionHLSL(ParamInfo, FunctionInfo, FunctionInstanceIndex, ENiagaraParticleDataComponentType::Float, 2, true, OutHLSL);
 	}
 
-	if (FunctionInfo.DefinitionName == GetVec3AttributeByIndexFunctionName)
+	if (FunctionInfo.DefinitionName == GetVec3ByIndexFunctionName)
 	{
 		return GenerateGetFunctionHLSL(ParamInfo, FunctionInfo, FunctionInstanceIndex, ENiagaraParticleDataComponentType::Float, 3, true, OutHLSL);
 	}
 
-	if (FunctionInfo.DefinitionName == GetVec4AttributeByIndexFunctionName)
+	if (FunctionInfo.DefinitionName == GetVec4ByIndexFunctionName)
 	{
 		return GenerateGetFunctionHLSL(ParamInfo, FunctionInfo, FunctionInstanceIndex, ENiagaraParticleDataComponentType::Float, 4, true, OutHLSL);
 	}
 
-	if (FunctionInfo.DefinitionName == GetColorAttributeByIndexFunctionName)
+	if (FunctionInfo.DefinitionName == GetColorByIndexFunctionName)
 	{
 		return GenerateGetFunctionHLSL(ParamInfo, FunctionInfo, FunctionInstanceIndex, ENiagaraParticleDataComponentType::Float, 4, true, OutHLSL);
 	}
 
-	if (FunctionInfo.DefinitionName == GetQuatAttributeByIndexFunctionName)
+	if (FunctionInfo.DefinitionName == GetQuatByIndexFunctionName)
 	{
 		return GenerateGetFunctionHLSL(ParamInfo, FunctionInfo, FunctionInstanceIndex, ENiagaraParticleDataComponentType::Float, 4, true, OutHLSL);
 	}
 
-	if (FunctionInfo.DefinitionName == GetIDAttributeByIndexFunctionName)
+	if (FunctionInfo.DefinitionName == GetIDByIndexFunctionName)
 	{
 		return GenerateGetFunctionHLSL(ParamInfo, FunctionInfo, FunctionInstanceIndex, ENiagaraParticleDataComponentType::ID, 1, true, OutHLSL);
 	}
diff --git a/Engine/Plugins/FX/Niagara/Source/NiagaraCore/Private/NiagaraCustomVersion.cpp b/Engine/Plugins/FX/Niagara/Source/NiagaraCore/Private/NiagaraCustomVersion.cpp
--- a/Engine/Plugins/FX/Niagara/Source/NiagaraCore/Private/NiagaraCustomVersion.cpp
+++ b/Engine/Plugins/FX/Niagara/Source/NiagaraCore/Private/NiagaraCustomVersion.cpp
@@ -8,4 +8,4 @@ const FGuid FNiagaraCustomVersion::GUID(0xFCF57AFA, 0x50764283, 0xB9A9E658, 0xFF
 // Register the custom version with core
 FCustomVersionRegistration GRegisterNiagaraCustomVersion(FNiagaraCustomVersion::GUID, FNiagaraCustomVersion::LatestVersion, TEXT("NiagaraVer"));
 
-const FGuid FNiagaraCustomVersion::LatestScriptCompileVersion(0x24C7603C, 0xC1454C9F, 0xB8688941, 0xB24F4D06);
+const FGuid FNiagaraCustomVersion::LatestScriptCompileVersion(0x4AF61ED9, 0x3723482F, 0xB6492E0E, 0x558E0F10);
