- Add a stat for locking the readback buffer so we can see if this is causing stalls in builds

#rnx

diff --git a/Engine/Plugins/FX/Niagara/Source/Niagara/Private/NiagaraGPUInstanceCountManager.cpp b/Engine/Plugins/FX/Niagara/Source/Niagara/Private/NiagaraGPUInstanceCountManager.cpp
--- a/Engine/Plugins/FX/Niagara/Source/Niagara/Private/NiagaraGPUInstanceCountManager.cpp
+++ b/Engine/Plugins/FX/Niagara/Source/Niagara/Private/NiagaraGPUInstanceCountManager.cpp
@@ -24,6 +24,7 @@ static FAutoConsoleVariableRef CVarNiagaraMinGPUDrawIndirectArgs(
 
 DECLARE_DWORD_COUNTER_STAT(TEXT("Used GPU Instance Counters"), STAT_NiagaraUsedGPUInstanceCounters, STATGROUP_Niagara);
 DECLARE_DWORD_COUNTER_STAT(TEXT("Max Num GPU Renderers"), STAT_NiagaraMaxNumGPURenderers, STATGROUP_Niagara);
+DECLARE_CYCLE_STAT(TEXT("GPU Readback Lock"), STAT_NiagaraGPUReadbackLock, STATGROUP_Niagara);
 
 
 //*****************************************************************************
@@ -256,6 +257,7 @@ const uint32* FNiagaraGPUInstanceCountManager::GetGPUReadback()
 {
 	if (CountReadback && CountReadbackSize && CountReadback->IsReady())
 	{
+		SCOPE_CYCLE_COUNTER(STAT_NiagaraGPUReadbackLock);
 		return (uint32*)(CountReadback->Lock(CountReadbackSize * sizeof(uint32)));
 	}
 	else
