Made button text white in the new Niagara emitter dialog.

diff --git a/Engine/Plugins/FX/Niagara/Source/NiagaraEditor/Private/SNiagaraNewAssetDialog.cpp b/Engine/Plugins/FX/Niagara/Source/NiagaraEditor/Private/SNiagaraNewAssetDialog.cpp
--- a/Engine/Plugins/FX/Niagara/Source/NiagaraEditor/Private/SNiagaraNewAssetDialog.cpp
+++ b/Engine/Plugins/FX/Niagara/Source/NiagaraEditor/Private/SNiagaraNewAssetDialog.cpp
@@ -118,6 +118,7 @@ void SNiagaraNewAssetDialog::Construct(const FArguments& InArgs, FName InSaveCon
 			.ButtonStyle(FEditorStyle::Get(), "FlatButton.Default")
 			.CancelButtonStyle(FEditorStyle::Get(), "FlatButton.Default")
 			.FinishButtonStyle(FEditorStyle::Get(), "FlatButton.Success")
+			.ButtonTextStyle(FEditorStyle::Get(), "FlatButton.DefaultTextStyle")
 			+SWizard::Page()
 			.CanShow(true)
 			.OnEnter(this, &SNiagaraNewAssetDialog::ResetStage)
