Fixed emitter dependency computation for emitters with simulation stages.

The code was assuming that the spawn and update execution context will contain all the data interfaces regardless of emitter type, but this is not true for emitters with stages. Stage DIs are only found in the GPU execution context, so we must use that information for GPU emitters.

Also fixed some typos in some error messages.

#rnx

diff --git a/Engine/Plugins/FX/Niagara/Source/Niagara/Private/NiagaraEmitterInstance.cpp b/Engine/Plugins/FX/Niagara/Source/Niagara/Private/NiagaraEmitterInstance.cpp
--- a/Engine/Plugins/FX/Niagara/Source/Niagara/Private/NiagaraEmitterInstance.cpp
+++ b/Engine/Plugins/FX/Niagara/Source/Niagara/Private/NiagaraEmitterInstance.cpp
@@ -412,6 +412,10 @@ void FNiagaraEmitterInstance::Init(int32 InEmitterIdx, FNiagaraSystemInstanceID
 		{
 			Scripts.Add(EventHandler.Script);
 		}
+		for (const UNiagaraSimulationStageBase* SimStage : CachedEmitter->GetSimulationStages())
+		{
+			Scripts.Add(SimStage->Script);
+		}
 		FNiagaraUtilities::CollectScriptDataInterfaceParameters(*CachedEmitter, MakeArrayView(Scripts), ScriptDefinedDataInterfaceParameters);
 
 		// Initialize bounds calculators - skip creating if we won't ever use it.  We leave the GPU sims in there with the editor so that we can
@@ -550,7 +554,7 @@ void FNiagaraEmitterInstance::CheckForErrors()
 	}
 	if (CachedEmitter->UpdateScriptProps.Script->GetVMExecutableData().Attributes.Num() == 0 || CachedEmitter->SpawnScriptProps.Script->GetVMExecutableData().Attributes.Num() == 0)
 	{
-		UE_LOG(LogNiagara, Error, TEXT("This emitter cannot be enabled because it's spawn or update script doesn't have any attriubtes.."));
+		UE_LOG(LogNiagara, Error, TEXT("This emitter cannot be enabled because its spawn or update script doesn't have any attributes.."));
 		SetExecutionState(ENiagaraExecutionState::Disabled);
 		return;
 	}
@@ -561,13 +565,13 @@ void FNiagaraEmitterInstance::CheckForErrors()
 		if (!CachedEmitter->SpawnScriptProps.Script->DidScriptCompilationSucceed(false))
 		{
 			bFailed = true;
-			UE_LOG(LogNiagara, Error, TEXT("This emitter cannot be enabled because it's CPU Spawn script failed to compile."));
+			UE_LOG(LogNiagara, Error, TEXT("This emitter cannot be enabled because its CPU Spawn script failed to compile."));
 		}
 
 		if (!CachedEmitter->UpdateScriptProps.Script->DidScriptCompilationSucceed(false))
 		{
 			bFailed = true;
-			UE_LOG(LogNiagara, Error, TEXT("This emitter cannot be enabled because it's CPU Update script failed to compile."));
+			UE_LOG(LogNiagara, Error, TEXT("This emitter cannot be enabled because its CPU Update script failed to compile."));
 		}
 
 		if (CachedEmitter->GetEventHandlers().Num() != 0)
@@ -577,7 +581,7 @@ void FNiagaraEmitterInstance::CheckForErrors()
 				if (!CachedEmitter->GetEventHandlers()[i].Script->DidScriptCompilationSucceed(false))
 				{
 					bFailed = true;
-					UE_LOG(LogNiagara, Error, TEXT("This emitter cannot be enabled because one of it's CPU Event scripts failed to compile."));
+					UE_LOG(LogNiagara, Error, TEXT("This emitter cannot be enabled because one of its CPU Event scripts failed to compile."));
 				}
 			}
 		}
@@ -593,13 +597,13 @@ void FNiagaraEmitterInstance::CheckForErrors()
 	{
 		if (CachedEmitter->GetGPUComputeScript()->IsScriptCompilationPending(true))
 		{
-			UE_LOG(LogNiagara, Error, TEXT("This emitter cannot be enabled because it's GPU script hasn't been compiled.."));
+			UE_LOG(LogNiagara, Error, TEXT("This emitter cannot be enabled because its GPU script hasn't been compiled.."));
 			SetExecutionState(ENiagaraExecutionState::Disabled);
 			return;
 		}
 		if (!CachedEmitter->GetGPUComputeScript()->DidScriptCompilationSucceed(true))
 		{
-			UE_LOG(LogNiagara, Error, TEXT("This emitter cannot be enabled because it's GPU script failed to compile."));
+			UE_LOG(LogNiagara, Error, TEXT("This emitter cannot be enabled because its GPU script failed to compile."));
 			SetExecutionState(ENiagaraExecutionState::Disabled);
 			return;
 		}
diff --git a/Engine/Plugins/FX/Niagara/Source/Niagara/Private/NiagaraScript.cpp b/Engine/Plugins/FX/Niagara/Source/Niagara/Private/NiagaraScript.cpp
--- a/Engine/Plugins/FX/Niagara/Source/Niagara/Private/NiagaraScript.cpp
+++ b/Engine/Plugins/FX/Niagara/Source/Niagara/Private/NiagaraScript.cpp
@@ -1477,7 +1477,7 @@ void UNiagaraScript::SetVMCompilationResults(const FNiagaraVMExecutableDataId& I
 	}
 
 	CachedDefaultDataInterfaces.Empty(CachedScriptVM.DataInterfaceInfo.Num());
-	for (FNiagaraScriptDataInterfaceCompileInfo Info : CachedScriptVM.DataInterfaceInfo)
+	for (const FNiagaraScriptDataInterfaceCompileInfo& Info : CachedScriptVM.DataInterfaceInfo)
 	{
 		int32 Idx = CachedDefaultDataInterfaces.AddDefaulted();
 		CachedDefaultDataInterfaces[Idx].UserPtrIdx = Info.UserPtrIdx;
diff --git a/Engine/Plugins/FX/Niagara/Source/Niagara/Private/NiagaraSystemInstance.cpp b/Engine/Plugins/FX/Niagara/Source/Niagara/Private/NiagaraSystemInstance.cpp
--- a/Engine/Plugins/FX/Niagara/Source/Niagara/Private/NiagaraSystemInstance.cpp
+++ b/Engine/Plugins/FX/Niagara/Source/Niagara/Private/NiagaraSystemInstance.cpp
@@ -691,8 +691,16 @@ void FNiagaraSystemInstance::ComputeEmittersExecutionOrder()
 		EmitterPriorities[EmitterIdx] = -1;
 
 		EmitterDependencies.SetNum(0, false);
-		FindDataInterfaceDependencies(Inst.GetSpawnExecutionContext().GetDataInterfaces(), EmitterDependencies);
-		FindDataInterfaceDependencies(Inst.GetUpdateExecutionContext().GetDataInterfaces(), EmitterDependencies);
+
+		if (Inst.GetCachedEmitter() && Inst.GetCachedEmitter()->SimTarget == ENiagaraSimTarget::GPUComputeSim && Inst.GetGPUContext())
+		{
+			FindDataInterfaceDependencies(Inst.GetGPUContext()->GetDataInterfaces(), EmitterDependencies);
+		}
+		else
+		{
+			FindDataInterfaceDependencies(Inst.GetSpawnExecutionContext().GetDataInterfaces(), EmitterDependencies);
+			FindDataInterfaceDependencies(Inst.GetUpdateExecutionContext().GetDataInterfaces(), EmitterDependencies);
+		}
 
 		// Map the pointers returned by the emitter to indices inside the Emitters array. This is O(N^2), but we expect
 		// to have few dependencies, so in practice it should be faster than a TMap. If it gets out of hand, we can also
