Fixed a bug that could prevent the SimulationTarget static switch to work for gpu emitters

diff --git a/Engine/Plugins/FX/Niagara/Source/NiagaraCore/Private/NiagaraCustomVersion.cpp b/Engine/Plugins/FX/Niagara/Source/NiagaraCore/Private/NiagaraCustomVersion.cpp
--- a/Engine/Plugins/FX/Niagara/Source/NiagaraCore/Private/NiagaraCustomVersion.cpp
+++ b/Engine/Plugins/FX/Niagara/Source/NiagaraCore/Private/NiagaraCustomVersion.cpp
@@ -8,4 +8,4 @@ const FGuid FNiagaraCustomVersion::GUID(0xFCF57AFA, 0x50764283, 0xB9A9E658, 0xFF
 // Register the custom version with core
 FCustomVersionRegistration GRegisterNiagaraCustomVersion(FNiagaraCustomVersion::GUID, FNiagaraCustomVersion::LatestVersion, TEXT("NiagaraVer"));
 
-const FGuid FNiagaraCustomVersion::LatestScriptCompileVersion(0x4AF61ED9, 0x3723482F, 0xB6492E0E, 0x558E0F10);
+const FGuid FNiagaraCustomVersion::LatestScriptCompileVersion(0x5FD3FF04, 0xC5BF4260, 0xB48F3599, 0xB449A6D5);
diff --git a/Engine/Plugins/FX/Niagara/Source/NiagaraEditor/Private/NiagaraCompiler.cpp b/Engine/Plugins/FX/Niagara/Source/NiagaraEditor/Private/NiagaraCompiler.cpp
--- a/Engine/Plugins/FX/Niagara/Source/NiagaraEditor/Private/NiagaraCompiler.cpp
+++ b/Engine/Plugins/FX/Niagara/Source/NiagaraEditor/Private/NiagaraCompiler.cpp
@@ -611,9 +611,9 @@ TSharedPtr<FNiagaraCompileRequestDataBase, ESPMode::ThreadSafe> FNiagaraEditorMo
 
 		// Now deep copy the system graphs, skipping traversal into any emitter references.
 		UNiagaraScriptSource* Source = Cast<UNiagaraScriptSource>(System->GetSystemSpawnScript()->GetSource());
-		BasePtr->DeepCopyGraphs(Source, ENiagaraScriptUsage::SystemSpawnScript, FCompileConstantResolver());
-		BasePtr->AddRapidIterationParameters(System->GetSystemSpawnScript()->RapidIterationParameters, FCompileConstantResolver());
-		BasePtr->AddRapidIterationParameters(System->GetSystemUpdateScript()->RapidIterationParameters, FCompileConstantResolver());
+		BasePtr->DeepCopyGraphs(Source, ENiagaraScriptUsage::SystemSpawnScript, EmptyResolver);
+		BasePtr->AddRapidIterationParameters(System->GetSystemSpawnScript()->RapidIterationParameters, EmptyResolver);
+		BasePtr->AddRapidIterationParameters(System->GetSystemUpdateScript()->RapidIterationParameters, EmptyResolver);
 		BasePtr->FinishPrecompile(Source, EncounterableVars, ENiagaraScriptUsage::SystemSpawnScript, EmptyResolver, nullptr);
 
 		// Add the User and System variables that we did encounter to the list that emitters might also encounter.
diff --git a/Engine/Plugins/FX/Niagara/Source/NiagaraEditor/Private/NiagaraHlslTranslator.cpp b/Engine/Plugins/FX/Niagara/Source/NiagaraEditor/Private/NiagaraHlslTranslator.cpp
--- a/Engine/Plugins/FX/Niagara/Source/NiagaraEditor/Private/NiagaraHlslTranslator.cpp
+++ b/Engine/Plugins/FX/Niagara/Source/NiagaraEditor/Private/NiagaraHlslTranslator.cpp
@@ -4215,7 +4215,7 @@ bool FHlslNiagaraTranslator::GetLiteralConstantVariable(FNiagaraVariable& OutVar
 		if (OutVar == FNiagaraVariable(FNiagaraTypeDefinition::GetSimulationTargetEnum(), TEXT("Emitter.SimulationTarget")))
 		{
 			FNiagaraInt32 EnumValue;
-			EnumValue.Value = (uint8) CompilationTarget;
+			EnumValue.Value = CompilationTarget == ENiagaraSimTarget::GPUComputeSim || CompileOptions.AdditionalDefines.Contains(TEXT("GPUComputeSim")) ? 1 : 0;
 			OutVar.SetValue(EnumValue);
 			return true;
 		}
