Improved support for static mesh streaming. CPU LOD data is now streamed out when LODs are streamed out.
FStaticMeshLODResources are now ref counted to prevent deleting the data while it is being accessed.
The stream out strategy now involves waiting several frame before streaming out, if there are external references to the LOD.
This number of checks is controlled by the new "r.Streaming.MaxReferenceChecksBeforeStreamOut" with a default value of 3.
Disabled StaticMesh streaming while in the editor to prevent issues with tools.

UNiagaraDataInterfaceStaticMesh now correctly support preview feature level and in game LOD streaming.
LODData that is cached initialization time is now invalidated if that LOD is streamed out, this
only affect GPU sampling and FStaticMeshFilteredAreaWeightedSectionSampler).

diff --git a/Engine/Plugins/FX/Niagara/Source/Niagara/Classes/NiagaraDataInterfaceStaticMesh.h b/Engine/Plugins/FX/Niagara/Source/Niagara/Classes/NiagaraDataInterfaceStaticMesh.h
--- a/Engine/Plugins/FX/Niagara/Source/Niagara/Classes/NiagaraDataInterfaceStaticMesh.h
+++ b/Engine/Plugins/FX/Niagara/Source/Niagara/Classes/NiagaraDataInterfaceStaticMesh.h
@@ -13,11 +13,11 @@ struct FNDIStaticMeshSectionFilter;
 struct FStaticMeshFilteredAreaWeightedSectionSampler : FStaticMeshAreaWeightedSectionSampler
 {
 	FStaticMeshFilteredAreaWeightedSectionSampler();
-	void Init(FStaticMeshLODResources* InRes, FNDIStaticMesh_InstanceData* InOwner);
+	void Init(const FStaticMeshLODResources* InRes, FNDIStaticMesh_InstanceData* InOwner);
 	virtual float GetWeights(TArray<float>& OutWeights)override;
 
 protected:
-	FStaticMeshLODResources* Res;
+	TRefCountPtr<const FStaticMeshLODResources> Res;
 	FNDIStaticMesh_InstanceData* Owner;
 };
 
@@ -46,7 +46,7 @@ public:
 
 	virtual ~FStaticMeshGpuSpawnBuffer();
 
-	void Initialise(const FStaticMeshLODResources& Res, const UNiagaraDataInterfaceStaticMesh& Interface,
+	void Initialise(const FStaticMeshLODResources* Res, const UNiagaraDataInterfaceStaticMesh& Interface,
 		bool bIsGpuUniformlyDistributedSampling, const TArray<int32>& ValidSection, const FStaticMeshFilteredAreaWeightedSectionSampler& SectionSampler);
 
 	virtual void InitRHI() override;
@@ -79,7 +79,7 @@ protected:
 	};
 
 	// Cached pointer to Section render data used for initialization only.
-	const FStaticMeshLODResources* SectionRenderData = nullptr;
+	TRefCountPtr<const FStaticMeshLODResources> SectionRenderData = nullptr;
 
 	TArray<SectionInfo> ValidSections;					// Only the section we want to spawn from
 
@@ -139,6 +139,11 @@ struct FNDIStaticMesh_InstanceData
 	/** Cached change id off of the data interface.*/
 	uint32 ChangeId;
 
+	/** The MinLOD, see UStaticMesh::MinLOD which is platform specific.*/
+	int32 MinLOD = 0;
+	/** The cached LODIdx used to initialize the FNDIStaticMesh_InstanceData.*/
+	int32 CachedLODIdx = 0;
+
 	FORCEINLINE UStaticMesh* GetActualMesh()const { return Mesh; }
 	FORCEINLINE bool UsesCpuUniformlyDistributedSampling() const { return bIsCpuUniformlyDistributedSampling; }
 	FORCEINLINE bool MeshHasPositions()const { return Mesh && Mesh->RenderData->LODResources[0].VertexBuffers.PositionVertexBuffer.GetNumVertices() > 0; }
@@ -155,6 +160,11 @@ struct FNDIStaticMesh_InstanceData
 	FORCEINLINE_DEBUGGABLE bool Init(UNiagaraDataInterfaceStaticMesh* Interface, FNiagaraSystemInstance* SystemInstance);
 	FORCEINLINE_DEBUGGABLE bool Tick(UNiagaraDataInterfaceStaticMesh* Interface, FNiagaraSystemInstance* SystemInstance, float InDeltaSeconds);
 	FORCEINLINE_DEBUGGABLE void Release();
+
+	FORCEINLINE const FStaticMeshLODResources* GetCurrentFirstLOD()
+	{
+		return Mesh->RenderData->GetCurrentFirstLOD(MinLOD);
+	}
 };
 
 /** Data Interface allowing sampling of static meshes. */
@@ -298,13 +308,13 @@ protected:
 private:
 	
 	template<typename TSampleMode, bool bFiltered>
-	FORCEINLINE_DEBUGGABLE int32 RandomSection(FRandomStream& RandStream, FStaticMeshLODResources& Res, FNDIStaticMesh_InstanceData* InstData);
+	FORCEINLINE_DEBUGGABLE int32 RandomSection(FRandomStream& RandStream, const FStaticMeshLODResources& Res, FNDIStaticMesh_InstanceData* InstData);
 
 	template<typename TSampleMode, bool bFiltered>
-	FORCEINLINE_DEBUGGABLE int32 RandomTriIndex(FRandomStream& RandStream, FStaticMeshLODResources& Res, FNDIStaticMesh_InstanceData* InstData);
+	FORCEINLINE_DEBUGGABLE int32 RandomTriIndex(FRandomStream& RandStream, const FStaticMeshLODResources& Res, FNDIStaticMesh_InstanceData* InstData);
 
 	template<typename TSampleMode>
-	FORCEINLINE_DEBUGGABLE int32 RandomTriIndexOnSection(FRandomStream& RandStream, FStaticMeshLODResources& Res, int32 SectionIdx, FNDIStaticMesh_InstanceData* InstData);
+	FORCEINLINE_DEBUGGABLE int32 RandomTriIndexOnSection(FRandomStream& RandStream, const FStaticMeshLODResources& Res, int32 SectionIdx, FNDIStaticMesh_InstanceData* InstData);
 
 	void WriteTransform(const FMatrix& ToWrite, FVectorVMContext& Context);
 };
diff --git a/Engine/Plugins/FX/Niagara/Source/Niagara/Private/NiagaraDataInterfaceStaticMesh.cpp b/Engine/Plugins/FX/Niagara/Source/Niagara/Private/NiagaraDataInterfaceStaticMesh.cpp
--- a/Engine/Plugins/FX/Niagara/Source/Niagara/Private/NiagaraDataInterfaceStaticMesh.cpp
+++ b/Engine/Plugins/FX/Niagara/Source/Niagara/Private/NiagaraDataInterfaceStaticMesh.cpp
@@ -49,7 +49,7 @@ FStaticMeshFilteredAreaWeightedSectionSampler::FStaticMeshFilteredAreaWeightedSe
 {
 }
 
-void FStaticMeshFilteredAreaWeightedSectionSampler::Init(FStaticMeshLODResources* InRes, FNDIStaticMesh_InstanceData* InOwner)
+void FStaticMeshFilteredAreaWeightedSectionSampler::Init(const FStaticMeshLODResources* InRes, FNDIStaticMesh_InstanceData* InOwner)
 {
 	Res = InRes;
 	Owner = InOwner;
@@ -62,17 +62,15 @@ float FStaticMeshFilteredAreaWeightedSectionSampler::GetWeights(TArray<float>& O
 	check(Owner && Owner->Mesh);
 	float Total = 0.0f;
 	OutWeights.Empty(Owner->GetValidSections().Num());
-	FStaticMeshLODResources& LODRes = Owner->Mesh->RenderData->LODResources[0];
-
-	if (Owner->Mesh->bSupportUniformlyDistributedSampling && LODRes.AreaWeightedSectionSamplers.Num() > 0)
-	{
-	for (int32 i = 0; i < Owner->GetValidSections().Num(); ++i)
+	if (Owner->Mesh->bSupportUniformlyDistributedSampling && Res->AreaWeightedSectionSamplers.Num() > 0)
 	{
-		int32 SecIdx = Owner->GetValidSections()[i];
-		float T = LODRes.AreaWeightedSectionSamplers[SecIdx].GetTotalWeight();
-		OutWeights.Add(T);
-		Total += T;
-	}
+		for (int32 i = 0; i < Owner->GetValidSections().Num(); ++i)
+		{
+			int32 SecIdx = Owner->GetValidSections()[i];
+			float T = Res->AreaWeightedSectionSamplers[SecIdx].GetTotalWeight();
+			OutWeights.Add(T);
+			Total += T;
+		}
 	}
 	else
 	{
@@ -95,11 +93,11 @@ FStaticMeshGpuSpawnBuffer::~FStaticMeshGpuSpawnBuffer()
 	//ValidSections.Empty();
 }
 
-void FStaticMeshGpuSpawnBuffer::Initialise(const FStaticMeshLODResources& Res, const UNiagaraDataInterfaceStaticMesh& Interface, bool bIsGpuUniformlyDistributedSampling, const TArray<int32>& ValidSection, const FStaticMeshFilteredAreaWeightedSectionSampler& SectionSamplerParam)
+void FStaticMeshGpuSpawnBuffer::Initialise(const FStaticMeshLODResources* Res, const UNiagaraDataInterfaceStaticMesh& Interface, bool bIsGpuUniformlyDistributedSampling, const TArray<int32>& ValidSection, const FStaticMeshFilteredAreaWeightedSectionSampler& SectionSamplerParam)
 {
 	// In this function we prepare some data to be uploaded on GPU from the available mesh data. This is a thread safe place to create this data.
 	// The section buffer needs to be specific to the current UI being built (section/material culling).
-	SectionRenderData = &Res;
+	SectionRenderData = Res;
 
 	const uint32 ValidSectionCount = ValidSection.Num();
 	const TArray<float, FMemoryImageAllocator>& Prob = SectionSamplerParam.GetProb();
@@ -111,15 +109,15 @@ void FStaticMeshGpuSpawnBuffer::Initialise(const FStaticMeshLODResources& Res, c
 	for (uint32 i = 0; i < ValidSectionCount; ++i)
 	{
 		uint32 ValidSectionId = ValidSection[i];
-		const FStaticMeshSection& Section = Res.Sections[ValidSectionId];
+		const FStaticMeshSection& Section = Res->Sections[ValidSectionId];
 		SectionInfo NewSectionInfo = { Section.FirstIndex / 3, Section.NumTriangles, Prob[i], (uint32)Alias[i] };
 		ValidSections.Add(NewSectionInfo);
 
-		check(!bIsGpuUniformlyDistributedSampling || bIsGpuUniformlyDistributedSampling && Res.AreaWeightedSectionSamplers[ValidSectionId].GetProb().Num() == Section.NumTriangles);
+		check(!bIsGpuUniformlyDistributedSampling || bIsGpuUniformlyDistributedSampling && Res->AreaWeightedSectionSamplers[ValidSectionId].GetProb().Num() == Section.NumTriangles);
 	}
 
 	if (bIsGpuUniformlyDistributedSampling)
-		BufferUniformTriangleSamplingSRV = Res.AreaWeightedSectionSamplersBuffer.GetBufferSRV(); // Cache that SRV for later
+		BufferUniformTriangleSamplingSRV = Res->AreaWeightedSectionSamplersBuffer.GetBufferSRV(); // Cache that SRV for later
 }
 
 void FStaticMeshGpuSpawnBuffer::InitRHI()
@@ -303,15 +301,18 @@ bool FNDIStaticMesh_InstanceData::Init(UNiagaraDataInterfaceStaticMesh* Interfac
 	ValidSections.Empty();
 	if (Mesh)
 	{
+	    MinLOD = Mesh->MinLOD.GetValueForFeatureLevel(SystemInstance->GetFeatureLevel());
+	    CachedLODIdx = Mesh->RenderData->GetCurrentFirstLODIdx(MinLOD);
+
 		bMeshAllowsCpuAccess = Mesh->bAllowCPUAccess;
 		bIsCpuUniformlyDistributedSampling = Mesh->bSupportUniformlyDistributedSampling;
 		bIsGpuUniformlyDistributedSampling = bIsCpuUniformlyDistributedSampling && Mesh->bSupportGpuUniformlyDistributedSampling;
 
 		//Init the instance filter
-		FStaticMeshLODResources& Res = Mesh->RenderData->LODResources[0];
-		for (int32 i = 0; i < Res.Sections.Num(); ++i)
+		TRefCountPtr<const FStaticMeshLODResources> Res = GetCurrentFirstLOD();
+		for (int32 i = 0; i < Res->Sections.Num(); ++i)
 		{
-			if (Interface->SectionFilter.AllowedMaterialSlots.Num() == 0 || Interface->SectionFilter.AllowedMaterialSlots.Contains(Res.Sections[i].MaterialIndex))
+			if (Interface->SectionFilter.AllowedMaterialSlots.Num() == 0 || Interface->SectionFilter.AllowedMaterialSlots.Contains(Res->Sections[i].MaterialIndex))
 			{
 				ValidSections.Add(i);
 			}
@@ -322,7 +323,7 @@ bool FNDIStaticMesh_InstanceData::Init(UNiagaraDataInterfaceStaticMesh* Interfac
 			UE_LOG(LogNiagara, Log, TEXT("StaticMesh data interface has a section filter preventing any spawning. Failed InitPerInstanceData - %s"), *Interface->GetFullName());
 		}
 
-		Sampler.Init(&Res, this);
+		Sampler.Init(Res, this);
 	}
 
 	SafeMesh_GT = Mesh;
@@ -349,6 +350,14 @@ bool FNDIStaticMesh_InstanceData::ResetRequired(UNiagaraDataInterfaceStaticMesh*
 		return true;
 	}
 
+	// Currently we only reset if the cached LOD was streamed out, to avoid performance hits. To revisit.
+	// We could probably just recache the data derived from the LOD instead of resetting everything.
+	if (Mesh && Mesh->RenderData->GetCurrentFirstLODIdx(MinLOD) > CachedLODIdx)
+	{
+		return true;
+	}
+
+	// The following conditions look like they could only be triggered in Editor...
 	//bool bPrevVCSampling = bSupportingVertexColorSampling;//TODO: Vertex color filtering needs more work.
 	if (Mesh)
 	{
@@ -1035,10 +1044,10 @@ struct TTypedMeshAccessorBinder
 		}
 
 		UNiagaraDataInterfaceStaticMesh* MeshInterface = CastChecked<UNiagaraDataInterfaceStaticMesh>(Interface);
-		FStaticMeshLODResources& Res = InstData->Mesh->RenderData->LODResources[0];
-		if (Res.VertexBuffers.StaticMeshVertexBuffer.GetUseHighPrecisionTangentBasis())			
+		TRefCountPtr<const FStaticMeshLODResources> Res = InstData->GetCurrentFirstLOD();
+		if (Res->VertexBuffers.StaticMeshVertexBuffer.GetUseHighPrecisionTangentBasis())			
 		{
-			if (Res.VertexBuffers.StaticMeshVertexBuffer.GetUseFullPrecisionUVs())
+			if (Res->VertexBuffers.StaticMeshVertexBuffer.GetUseFullPrecisionUVs())
 			{
 				NextBinder::template Bind<ParamTypes..., TTypedMeshVertexAccessor<EStaticMeshVertexTangentBasisType::HighPrecision, EStaticMeshVertexUVType::HighPrecision>>(Interface, BindingInfo, InstanceData, OutFunc);
 			}
@@ -1049,7 +1058,7 @@ struct TTypedMeshAccessorBinder
 		}
 		else
 		{
-			if (Res.VertexBuffers.StaticMeshVertexBuffer.GetUseFullPrecisionUVs())
+			if (Res->VertexBuffers.StaticMeshVertexBuffer.GetUseFullPrecisionUVs())
 			{
 				NextBinder::template Bind<ParamTypes..., TTypedMeshVertexAccessor<EStaticMeshVertexTangentBasisType::Default, EStaticMeshVertexUVType::HighPrecision>>(Interface, BindingInfo, InstanceData, OutFunc);
 			}
@@ -1228,7 +1237,7 @@ bool UNiagaraDataInterfaceStaticMesh::InitPerInstanceData(void* PerInstanceData,
 			ensure(Inst->Mesh->bAllowCPUAccess); // this should have been verified in Init()
 
 			MeshGpuSpawnBuffer = new FStaticMeshGpuSpawnBuffer;
-			FStaticMeshLODResources& Res = Inst->Mesh->RenderData->LODResources[0];
+			TRefCountPtr<const FStaticMeshLODResources> Res = Inst->GetCurrentFirstLOD();
 			MeshGpuSpawnBuffer->Initialise(Res, *this, Inst->bIsGpuUniformlyDistributedSampling, Inst->ValidSections, Inst->Sampler);
 		}
 
@@ -1353,7 +1362,7 @@ void UNiagaraDataInterfaceStaticMesh::IsValid(FVectorVMContext& Context)
 //RandomSection specializations.
 //Each combination for SampleMode and Section filtered. NOTE: Invalid sample mode left intentionally unimplemented
 template<>
-FORCEINLINE int32 UNiagaraDataInterfaceStaticMesh::RandomSection<TSampleModeAreaWeighted, true>(FRandomStream& RandStream, FStaticMeshLODResources& Res, FNDIStaticMesh_InstanceData* InstData)
+FORCEINLINE int32 UNiagaraDataInterfaceStaticMesh::RandomSection<TSampleModeAreaWeighted, true>(FRandomStream& RandStream, const FStaticMeshLODResources& Res, FNDIStaticMesh_InstanceData* InstData)
 {
 	if (InstData->GetAreaWeightedSampler().GetNumEntries() > 0)
 	{
@@ -1367,7 +1376,7 @@ FORCEINLINE int32 UNiagaraDataInterfaceStaticMesh::RandomSection<TSampleModeArea
 }
 
 template<>
-FORCEINLINE int32 UNiagaraDataInterfaceStaticMesh::RandomSection<TSampleModeAreaWeighted, false>(FRandomStream& RandStream, FStaticMeshLODResources& Res, FNDIStaticMesh_InstanceData* InstData)
+FORCEINLINE int32 UNiagaraDataInterfaceStaticMesh::RandomSection<TSampleModeAreaWeighted, false>(FRandomStream& RandStream, const FStaticMeshLODResources& Res, FNDIStaticMesh_InstanceData* InstData)
 {
 	if (Res.AreaWeightedSampler.GetNumEntries() > 0)
 	{
@@ -1380,14 +1389,14 @@ FORCEINLINE int32 UNiagaraDataInterfaceStaticMesh::RandomSection<TSampleModeArea
 }
 
 template<>
-FORCEINLINE int32 UNiagaraDataInterfaceStaticMesh::RandomSection<TSampleModeDefault, true>(FRandomStream& RandStream, FStaticMeshLODResources& Res, FNDIStaticMesh_InstanceData* InstData)
+FORCEINLINE int32 UNiagaraDataInterfaceStaticMesh::RandomSection<TSampleModeDefault, true>(FRandomStream& RandStream, const FStaticMeshLODResources& Res, FNDIStaticMesh_InstanceData* InstData)
 {
 	int32 Idx = RandStream.RandRange(0, InstData->GetValidSections().Num() - 1);
 	return InstData->GetValidSections()[Idx];
 }
 
 template<>
-FORCEINLINE int32 UNiagaraDataInterfaceStaticMesh::RandomSection<TSampleModeDefault, false>(FRandomStream& RandStream, FStaticMeshLODResources& Res, FNDIStaticMesh_InstanceData* InstData)
+FORCEINLINE int32 UNiagaraDataInterfaceStaticMesh::RandomSection<TSampleModeDefault, false>(FRandomStream& RandStream, const FStaticMeshLODResources& Res, FNDIStaticMesh_InstanceData* InstData)
 {
 	return RandStream.RandRange(0, Res.Sections.Num() - 1);
 }
@@ -1398,10 +1407,10 @@ void UNiagaraDataInterfaceStaticMesh::RandomSection(FVectorVMContext& Context)
 	VectorVM::FUserPtrHandler<FNDIStaticMesh_InstanceData> InstData(Context);
 	VectorVM::FExternalFuncRegisterHandler<int32> OutSection(Context);
 
-	FStaticMeshLODResources& Res = InstData->Mesh->RenderData->LODResources[0];
+	TRefCountPtr<const FStaticMeshLODResources> Res = InstData->GetCurrentFirstLOD();
 	for (int32 i = 0; i < Context.NumInstances; ++i)
 	{
-		*OutSection.GetDestAndAdvance() = RandomSection<TSampleMode, true>(Context.RandStream, Res, InstData);
+		*OutSection.GetDestAndAdvance() = RandomSection<TSampleMode, true>(Context.RandStream, *Res, InstData);
 	}
 }
 
@@ -1412,7 +1421,7 @@ void UNiagaraDataInterfaceStaticMesh::RandomSection<TSampleModeInvalid>(FVectorV
 	VectorVM::FUserPtrHandler<FNDIStaticMesh_InstanceData> InstData(Context);
 	VectorVM::FExternalFuncRegisterHandler<int32> OutSection(Context);
 
-	FStaticMeshLODResources& Res = InstData->Mesh->RenderData->LODResources[0];
+	TRefCountPtr<const FStaticMeshLODResources> Res = InstData->GetCurrentFirstLOD();
 	for (int32 i = 0; i < Context.NumInstances; ++i)
 	{
 		*OutSection.GetDestAndAdvance() = -1;
@@ -1422,12 +1431,12 @@ void UNiagaraDataInterfaceStaticMesh::RandomSection<TSampleModeInvalid>(FVectorV
 //RandomTriIndex specializations.
 //Each combination for SampleMode and Section filtered. NOTE: Invalid sample mode left intentionally unimplemented
 template<>
-FORCEINLINE int32 UNiagaraDataInterfaceStaticMesh::RandomTriIndex<TSampleModeAreaWeighted, true>(FRandomStream& RandStream, FStaticMeshLODResources& Res, FNDIStaticMesh_InstanceData* InstData)
+FORCEINLINE int32 UNiagaraDataInterfaceStaticMesh::RandomTriIndex<TSampleModeAreaWeighted, true>(FRandomStream& RandStream, const FStaticMeshLODResources& Res, FNDIStaticMesh_InstanceData* InstData)
 {
 	int32 SecIdx = RandomSection<TSampleModeAreaWeighted, true>(RandStream, Res, InstData);
 	if (SecIdx < Res.Sections.Num() && SecIdx < Res.AreaWeightedSectionSamplers.Num())
 	{
-		FStaticMeshSection&  Sec = Res.Sections[SecIdx];
+		const FStaticMeshSection&  Sec = Res.Sections[SecIdx];
 		if (Res.AreaWeightedSectionSamplers[SecIdx].GetNumEntries() > 0)
 		{
 			int32 Tri = Res.AreaWeightedSectionSamplers[SecIdx].GetEntryIndex(RandStream.GetFraction(), RandStream.GetFraction());
@@ -1439,12 +1448,12 @@ FORCEINLINE int32 UNiagaraDataInterfaceStaticMesh::RandomTriIndex<TSampleModeAre
 }
 
 template<>
-FORCEINLINE int32 UNiagaraDataInterfaceStaticMesh::RandomTriIndex<TSampleModeAreaWeighted, false>(FRandomStream& RandStream, FStaticMeshLODResources& Res, FNDIStaticMesh_InstanceData* InstData)
+FORCEINLINE int32 UNiagaraDataInterfaceStaticMesh::RandomTriIndex<TSampleModeAreaWeighted, false>(FRandomStream& RandStream, const FStaticMeshLODResources& Res, FNDIStaticMesh_InstanceData* InstData)
 {
 	int32 SecIdx = RandomSection<TSampleModeAreaWeighted, false>(RandStream, Res, InstData);
 	if (SecIdx < Res.Sections.Num() && SecIdx < Res.AreaWeightedSectionSamplers.Num())
 	{
-		FStaticMeshSection&  Sec = Res.Sections[SecIdx];
+		const FStaticMeshSection&  Sec = Res.Sections[SecIdx];
 		if (Res.AreaWeightedSectionSamplers[SecIdx].GetNumEntries() > 0)
 		{
 			int32 Tri = Res.AreaWeightedSectionSamplers[SecIdx].GetEntryIndex(RandStream.GetFraction(), RandStream.GetFraction());
@@ -1456,12 +1465,12 @@ FORCEINLINE int32 UNiagaraDataInterfaceStaticMesh::RandomTriIndex<TSampleModeAre
 }
 
 template<>
-FORCEINLINE int32 UNiagaraDataInterfaceStaticMesh::RandomTriIndex<TSampleModeDefault, true>(FRandomStream& RandStream, FStaticMeshLODResources& Res, FNDIStaticMesh_InstanceData* InstData)
+FORCEINLINE int32 UNiagaraDataInterfaceStaticMesh::RandomTriIndex<TSampleModeDefault, true>(FRandomStream& RandStream, const FStaticMeshLODResources& Res, FNDIStaticMesh_InstanceData* InstData)
 {
 	int32 SecIdx = RandomSection<TSampleModeDefault, true>(RandStream, Res, InstData);
 	if (SecIdx < Res.Sections.Num())
 	{
-		FStaticMeshSection&  Sec = Res.Sections[SecIdx];
+		const FStaticMeshSection&  Sec = Res.Sections[SecIdx];
 		int32 Tri = RandStream.RandRange(0, Sec.NumTriangles - 1);
 		return (Sec.FirstIndex / 3) + Tri;
 	}
@@ -1469,12 +1478,12 @@ FORCEINLINE int32 UNiagaraDataInterfaceStaticMesh::RandomTriIndex<TSampleModeDef
 }
 
 template<>
-FORCEINLINE int32 UNiagaraDataInterfaceStaticMesh::RandomTriIndex<TSampleModeDefault, false>(FRandomStream& RandStream, FStaticMeshLODResources& Res, FNDIStaticMesh_InstanceData* InstData)
+FORCEINLINE int32 UNiagaraDataInterfaceStaticMesh::RandomTriIndex<TSampleModeDefault, false>(FRandomStream& RandStream, const FStaticMeshLODResources& Res, FNDIStaticMesh_InstanceData* InstData)
 {
 	int32 SecIdx = RandomSection<TSampleModeDefault, false>(RandStream, Res, InstData);
 	if (SecIdx < Res.Sections.Num())
 	{
-		FStaticMeshSection&  Sec = Res.Sections[SecIdx];
+		const FStaticMeshSection&  Sec = Res.Sections[SecIdx];
 		int32 Tri = RandStream.RandRange(0, Sec.NumTriangles - 1);
 		return (Sec.FirstIndex / 3) + Tri;
 	}
@@ -1492,11 +1501,11 @@ void UNiagaraDataInterfaceStaticMesh::RandomTriCoord(FVectorVMContext& Context)
 	VectorVM::FExternalFuncRegisterHandler<float> OutBaryZ(Context);
 
 	check(InstData->Mesh);
-	FStaticMeshLODResources& Res = InstData->Mesh->RenderData->LODResources[0];
-	FIndexArrayView Indices = Res.IndexBuffer.GetArrayView();
+	TRefCountPtr<const FStaticMeshLODResources> Res = InstData->GetCurrentFirstLOD();
+	FIndexArrayView Indices = Res->IndexBuffer.GetArrayView();
 	for (int32 i = 0; i < Context.NumInstances; ++i)
 	{
-		*OutTri.GetDest() = RandomTriIndex<TSampleMode, true>(Context.RandStream, Res, InstData);
+		*OutTri.GetDest() = RandomTriIndex<TSampleMode, true>(Context.RandStream, *Res, InstData);
 		FVector Bary = RandomBarycentricCoord(Context.RandStream);
 		*OutBaryX.GetDest() = Bary.X;
 		*OutBaryY.GetDest() = Bary.Y;
@@ -1555,8 +1564,8 @@ void UNiagaraDataInterfaceStaticMesh::RandomTriCoordVertexColorFiltered(FVectorV
 	}
 	
 	FDynamicVertexColorFilterData* VCFData = InstData->DynamicVertexColorSampler.Get();
-	FStaticMeshLODResources& Res = InstData->Mesh->RenderData->LODResources[0];
-	FIndexArrayView Indices = Res.IndexBuffer.GetArrayView();
+	TRefCountPtr<const FStaticMeshLODResources> Res = InstData->GetCurrentFirstLOD();
+	FIndexArrayView Indices = Res->IndexBuffer.GetArrayView();
 
 	for (int32 i = 0; i < Context.NumInstances; ++i)
 	{
@@ -1607,17 +1616,17 @@ void UNiagaraDataInterfaceStaticMesh::RandomTriCoordVertexColorFiltered(FVectorV
 }
 
 template<>
-FORCEINLINE int32 UNiagaraDataInterfaceStaticMesh::RandomTriIndexOnSection<TSampleModeAreaWeighted>(FRandomStream& RandStream, FStaticMeshLODResources& Res, int32 SecIdx, FNDIStaticMesh_InstanceData* InstData)
+FORCEINLINE int32 UNiagaraDataInterfaceStaticMesh::RandomTriIndexOnSection<TSampleModeAreaWeighted>(FRandomStream& RandStream, const FStaticMeshLODResources& Res, int32 SecIdx, FNDIStaticMesh_InstanceData* InstData)
 {
-	FStaticMeshSection&  Sec = Res.Sections[SecIdx];
+	const FStaticMeshSection&  Sec = Res.Sections[SecIdx];
 	int32 Tri = Res.AreaWeightedSectionSamplers[SecIdx].GetEntryIndex(RandStream.GetFraction(), RandStream.GetFraction());
 	return (Sec.FirstIndex / 3) + Tri;
 }
 
 template<>
-FORCEINLINE int32 UNiagaraDataInterfaceStaticMesh::RandomTriIndexOnSection<TSampleModeDefault>(FRandomStream& RandStream, FStaticMeshLODResources& Res, int32 SecIdx, FNDIStaticMesh_InstanceData* InstData)
+FORCEINLINE int32 UNiagaraDataInterfaceStaticMesh::RandomTriIndexOnSection<TSampleModeDefault>(FRandomStream& RandStream, const FStaticMeshLODResources& Res, int32 SecIdx, FNDIStaticMesh_InstanceData* InstData)
 {
-	FStaticMeshSection&  Sec = Res.Sections[SecIdx];
+	const FStaticMeshSection&  Sec = Res.Sections[SecIdx];
 	int32 Tri = RandStream.RandRange(0, Sec.NumTriangles - 1);
 	return (Sec.FirstIndex / 3) + Tri;
 }
@@ -1634,14 +1643,14 @@ void UNiagaraDataInterfaceStaticMesh::RandomTriCoordOnSection(FVectorVMContext&
 	VectorVM::FExternalFuncRegisterHandler<float> OutBaryZ(Context);
 
 	check(InstData->Mesh);
-	FStaticMeshLODResources& Res = InstData->Mesh->RenderData->LODResources[0];
-	FIndexArrayView Indices = Res.IndexBuffer.GetArrayView();
-	const int32 MaxSection = Res.Sections.Num() - 1;
+	TRefCountPtr<const FStaticMeshLODResources> Res = InstData->GetCurrentFirstLOD();
+	FIndexArrayView Indices = Res->IndexBuffer.GetArrayView();
+	const int32 MaxSection = Res->Sections.Num() - 1;
 	for (int32 i = 0; i < Context.NumInstances; ++i)
 	{
 
 		int32 SecIdx = FMath::Clamp(SectionIdxParam.Get(), 0, MaxSection);
-		*OutTri.GetDest() = RandomTriIndexOnSection<TSampleMode>(Context.RandStream, Res, SecIdx, InstData);
+		*OutTri.GetDest() = RandomTriIndexOnSection<TSampleMode>(Context.RandStream, *Res, SecIdx, InstData);
 		FVector Bary = RandomBarycentricCoord(Context.RandStream);
 		*OutBaryX.GetDest() = Bary.X;
 		*OutBaryY.GetDest() = Bary.Y;
@@ -1706,9 +1715,9 @@ void UNiagaraDataInterfaceStaticMesh::GetTriCoordPosition(FVectorVMContext& Cont
 		return;
 	}
 
-	FStaticMeshLODResources& Res = InstData->Mesh->RenderData->LODResources[0];
-	const FIndexArrayView& Indices = Res.IndexBuffer.GetArrayView();
-	const FPositionVertexBuffer& Positions = Res.VertexBuffers.PositionVertexBuffer;
+	TRefCountPtr<const FStaticMeshLODResources> Res = InstData->GetCurrentFirstLOD();
+	const FIndexArrayView& Indices = Res->IndexBuffer.GetArrayView();
+	const FPositionVertexBuffer& Positions = Res->VertexBuffers.PositionVertexBuffer;
 
 	const int32 NumTriangles = Indices.Num() / 3;
 	for (int32 i = 0; i < Context.NumInstances; ++i)
@@ -1764,9 +1773,9 @@ void UNiagaraDataInterfaceStaticMesh::GetTriCoordNormal(FVectorVMContext& Contex
 		return;
 	}
 
-	FStaticMeshLODResources& Res = InstData->Mesh->RenderData->LODResources[0];
-	const FIndexArrayView& Indices = Res.IndexBuffer.GetArrayView();
-	const FStaticMeshVertexBuffer& Verts = Res.VertexBuffers.StaticMeshVertexBuffer;
+	TRefCountPtr<const FStaticMeshLODResources> Res = InstData->GetCurrentFirstLOD();
+	const FIndexArrayView& Indices = Res->IndexBuffer.GetArrayView();
+	const FStaticMeshVertexBuffer& Verts = Res->VertexBuffers.StaticMeshVertexBuffer;
 
 	const int32 NumTriangles = Indices.Num() / 3;
 	for (int32 i = 0; i < Context.NumInstances; ++i)
@@ -1833,9 +1842,9 @@ void UNiagaraDataInterfaceStaticMesh::GetTriCoordTangents(FVectorVMContext& Cont
 		return;
 	}
 
-	FStaticMeshLODResources& Res = InstData->Mesh->RenderData->LODResources[0];
-	const FIndexArrayView& Indices = Res.IndexBuffer.GetArrayView();
-	const VertexAccessorType Verts(Res.VertexBuffers.StaticMeshVertexBuffer);
+	TRefCountPtr<const FStaticMeshLODResources> Res = InstData->GetCurrentFirstLOD();
+	const FIndexArrayView& Indices = Res->IndexBuffer.GetArrayView();
+	const VertexAccessorType Verts(Res->VertexBuffers.StaticMeshVertexBuffer);
 	const int32 NumTriangles = Indices.Num() / 3;
 	for (int32 i = 0; i < Context.NumInstances; ++i)
 	{
@@ -1889,10 +1898,10 @@ void UNiagaraDataInterfaceStaticMesh::GetTriCoordColor(FVectorVMContext& Context
 	VectorVM::FExternalFuncRegisterHandler<float> OutColorB(Context);
 	VectorVM::FExternalFuncRegisterHandler<float> OutColorA(Context);
 
-	FStaticMeshLODResources* Res = nullptr;
+	TRefCountPtr<const FStaticMeshLODResources> Res;
 	if (InstData->Mesh)
 	{
-		Res = &InstData->Mesh->RenderData->LODResources[0];
+		Res = InstData->GetCurrentFirstLOD();
 	}
 
 	if (Res && Res->VertexBuffers.ColorVertexBuffer.GetNumVertices() > 0)
@@ -1964,9 +1973,9 @@ void UNiagaraDataInterfaceStaticMesh::GetTriCoordUV(FVectorVMContext& Context)
 		return;
 	}
 
-	FStaticMeshLODResources& Res = InstData->Mesh->RenderData->LODResources[0];
-	const FIndexArrayView& Indices = Res.IndexBuffer.GetArrayView();
-	const VertexAccessorType Verts(Res.VertexBuffers.StaticMeshVertexBuffer);
+	TRefCountPtr<const FStaticMeshLODResources> Res = InstData->GetCurrentFirstLOD();
+	const FIndexArrayView& Indices = Res->IndexBuffer.GetArrayView();
+	const VertexAccessorType Verts(Res->VertexBuffers.StaticMeshVertexBuffer);
 
 	const int32 NumTriangles = Indices.Num() / 3;
 	for (int32 i = 0; i < Context.NumInstances; ++i)
@@ -2025,9 +2034,9 @@ void UNiagaraDataInterfaceStaticMesh::GetTriCoordPositionAndVelocity(FVectorVMCo
 		return;
 	}
 
-	FStaticMeshLODResources& Res = InstData->Mesh->RenderData->LODResources[0];
-	const FIndexArrayView& Indices = Res.IndexBuffer.GetArrayView();
-	const FPositionVertexBuffer& Positions = Res.VertexBuffers.PositionVertexBuffer;
+	TRefCountPtr<const FStaticMeshLODResources> Res = InstData->GetCurrentFirstLOD();
+	const FIndexArrayView& Indices = Res->IndexBuffer.GetArrayView();
+	const FPositionVertexBuffer& Positions = Res->VertexBuffers.PositionVertexBuffer;
 
 	const int32 NumTriangles = Indices.Num() / 3;
 	float InvDt = InstData->DeltaSeconds > 0.0f ? (1.0f / InstData->DeltaSeconds) : 0.0f;
@@ -2168,8 +2177,8 @@ void UNiagaraDataInterfaceStaticMesh::GetVertexPosition(FVectorVMContext& Contex
 		return;
 	}
 
-	FStaticMeshLODResources& Res = InstData->Mesh->RenderData->LODResources[0];
-	const FPositionVertexBuffer& Positions = Res.VertexBuffers.PositionVertexBuffer;
+	TRefCountPtr<const FStaticMeshLODResources> Res = InstData->GetCurrentFirstLOD();
+	const FPositionVertexBuffer& Positions = Res->VertexBuffers.PositionVertexBuffer;
 
 	const int32 NumVerts = Positions.GetNumVertices();
 	FVector Pos;
@@ -2781,9 +2790,9 @@ bool FDynamicVertexColorFilterData::Init(FNDIStaticMesh_InstanceData* Owner)
 	VertexColorToTriangleStart.AddDefaulted(256);
 	check(Owner->Mesh);
 
-	FStaticMeshLODResources& Res = Owner->Mesh->RenderData->LODResources[0];
+	TRefCountPtr<const FStaticMeshLODResources> Res = Owner->GetCurrentFirstLOD();
 
-	if (Res.VertexBuffers.ColorVertexBuffer.GetNumVertices() == 0)
+	if (Res->VertexBuffers.ColorVertexBuffer.GetNumVertices() == 0)
 	{
 		UE_LOG(LogNiagara, Log, TEXT("Cannot initialize vertex color filter data for a mesh with no color data - %s"), *Owner->Mesh->GetFullName());
 		return false;
@@ -2796,20 +2805,20 @@ bool FDynamicVertexColorFilterData::Init(FNDIStaticMesh_InstanceData* Owner)
 		uint32 MaxVertexColorRed = i + 1;
 		VertexColorToTriangleStart[i] = TrianglesSortedByVertexColor.Num();
 
-		FIndexArrayView IndexView = Res.IndexBuffer.GetArrayView();
+		FIndexArrayView IndexView = Res->IndexBuffer.GetArrayView();
 		for (int32 j = 0; j < Owner->GetValidSections().Num(); j++)
 		{
 			int32 SectionIdx = Owner->GetValidSections()[j];
-			int32 TriStartIdx = Res.Sections[SectionIdx].FirstIndex;
-			for (uint32 TriIdx = 0; TriIdx < Res.Sections[SectionIdx].NumTriangles; TriIdx++)
+			int32 TriStartIdx = Res->Sections[SectionIdx].FirstIndex;
+			for (uint32 TriIdx = 0; TriIdx < Res->Sections[SectionIdx].NumTriangles; TriIdx++)
 			{
 				uint32 V0Idx = IndexView[TriStartIdx + TriIdx * 3 + 0];
 				uint32 V1Idx = IndexView[TriStartIdx + TriIdx * 3 + 1];
 				uint32 V2Idx = IndexView[TriStartIdx + TriIdx * 3 + 2];
 
-				uint8 MaxR = FMath::Max<uint8>(Res.VertexBuffers.ColorVertexBuffer.VertexColor(V0Idx).R,
-					FMath::Max<uint8>(Res.VertexBuffers.ColorVertexBuffer.VertexColor(V1Idx).R,
-						Res.VertexBuffers.ColorVertexBuffer.VertexColor(V2Idx).R));
+				uint8 MaxR = FMath::Max<uint8>(Res->VertexBuffers.ColorVertexBuffer.VertexColor(V0Idx).R,
+					FMath::Max<uint8>(Res->VertexBuffers.ColorVertexBuffer.VertexColor(V1Idx).R,
+						Res->VertexBuffers.ColorVertexBuffer.VertexColor(V2Idx).R));
 				if (MaxR >= MinVertexColorRed && MaxR < MaxVertexColorRed)
 				{
 					TrianglesSortedByVertexColor.Add(TriStartIdx + TriIdx * 3);
diff --git a/Engine/Plugins/FX/Niagara/Source/Niagara/Private/NiagaraSystemInstance.cpp b/Engine/Plugins/FX/Niagara/Source/Niagara/Private/NiagaraSystemInstance.cpp
--- a/Engine/Plugins/FX/Niagara/Source/Niagara/Private/NiagaraSystemInstance.cpp
+++ b/Engine/Plugins/FX/Niagara/Source/Niagara/Private/NiagaraSystemInstance.cpp
@@ -88,6 +88,7 @@ FNiagaraSystemInstance::FNiagaraSystemInstance(UNiagaraComponent* InComponent)
 	, CachedDeltaSeconds(0.0f)
 	, RequestedExecutionState(ENiagaraExecutionState::Complete)
 	, ActualExecutionState(ENiagaraExecutionState::Complete)
+	, FeatureLevel(GMaxRHIFeatureLevel)
 {
 	static TAtomic<uint64> IDCounter(1);
 	ID = IDCounter.IncrementExchange();
@@ -107,6 +108,7 @@ FNiagaraSystemInstance::FNiagaraSystemInstance(UNiagaraComponent* InComponent)
 			{
 				Batcher = static_cast<NiagaraEmitterInstanceBatcher*>(FXSystemInterface->GetInterface(NiagaraEmitterInstanceBatcher::Name));
 			}
+			FeatureLevel = World->FeatureLevel;
 		}
 	}
 }
diff --git a/Engine/Plugins/FX/Niagara/Source/Niagara/Public/NiagaraSystemInstance.h b/Engine/Plugins/FX/Niagara/Source/Niagara/Public/NiagaraSystemInstance.h
--- a/Engine/Plugins/FX/Niagara/Source/Niagara/Public/NiagaraSystemInstance.h
+++ b/Engine/Plugins/FX/Niagara/Source/Niagara/Public/NiagaraSystemInstance.h
@@ -426,7 +426,13 @@ private:
 	/** Tag we feed into crash reporter for this instance. */
 	mutable FString CrashReporterTag;
 
+	/** The feature level of for this component instance. */
+	ERHIFeatureLevel::Type FeatureLevel = ERHIFeatureLevel::Num;
+
 public:
+
+	ERHIFeatureLevel::Type GetFeatureLevel() const { return FeatureLevel; }
+
 	// Transient data that is accumulated during tick.
 	uint32 TotalGPUParamSize = 0;
 	uint32 ActiveGPUEmitterCount = 0;
diff --git a/Engine/Source/Runtime/CoreUObject/Private/Serialization/BulkData2.cpp b/Engine/Source/Runtime/CoreUObject/Private/Serialization/BulkData2.cpp
--- a/Engine/Source/Runtime/CoreUObject/Private/Serialization/BulkData2.cpp
+++ b/Engine/Source/Runtime/CoreUObject/Private/Serialization/BulkData2.cpp
@@ -645,10 +645,17 @@ FBulkDataBase& FBulkDataBase::operator=(const FBulkDataBase& Other)
 
 	if( !Other.IsDataMemoryMapped())
 	{
-		const int64 DataSize = Other.GetBulkDataSize();
+		if (Other.GetDataBufferReadOnly())
+		{
+			const int64 DataSize = Other.GetBulkDataSize();
 
-		void* Dst = AllocateData(DataSize);
-		FMemory::Memcpy(Dst, Other.GetDataBufferReadOnly(), DataSize);
+			void* Dst = AllocateData(DataSize);
+			FMemory::Memcpy(Dst, Other.GetDataBufferReadOnly(), DataSize);
+		}
+		else
+		{
+			Data.Fallback.BulkDataSize = Other.Data.Fallback.BulkDataSize;
+		}
 	}
 	else
 	{
diff --git a/Engine/Source/Runtime/Engine/Private/PositionVertexBuffer.cpp b/Engine/Source/Runtime/Engine/Private/PositionVertexBuffer.cpp
--- a/Engine/Source/Runtime/Engine/Private/PositionVertexBuffer.cpp
+++ b/Engine/Source/Runtime/Engine/Private/PositionVertexBuffer.cpp
@@ -219,6 +219,29 @@ FVertexBufferRHIRef FPositionVertexBuffer::CreateRHIBuffer_Async()
 	return CreateRHIBuffer_Internal<false>();
 }
 
+/** Copy everything, keeping reference to the same RHI resources. */
+void FPositionVertexBuffer::CopyRHIForStreaming(const FPositionVertexBuffer& Other, bool InAllowCPUAccess)
+{
+	// Copy serialized properties.
+	Stride = Other.Stride;
+	NumVertices = Other.NumVertices;
+
+	// Handle CPU access.
+	if (InAllowCPUAccess)
+	{
+		bNeedsCPUAccess = Other.bNeedsCPUAccess;
+		AllocateData(bNeedsCPUAccess);
+	}
+	else
+	{
+		bNeedsCPUAccess = false;
+	}
+
+	// Copy resource references.
+	VertexBufferRHI = Other.VertexBufferRHI;
+	PositionComponentSRV = Other.PositionComponentSRV;
+}
+
 void FPositionVertexBuffer::InitRHI()
 {
 	VertexBufferRHI = CreateRHIBuffer_RenderThread();
diff --git a/Engine/Source/Runtime/Engine/Private/RawIndexBuffer.cpp b/Engine/Source/Runtime/Engine/Private/RawIndexBuffer.cpp
--- a/Engine/Source/Runtime/Engine/Private/RawIndexBuffer.cpp
+++ b/Engine/Source/Runtime/Engine/Private/RawIndexBuffer.cpp
@@ -353,6 +353,22 @@ FIndexBufferRHIRef FRawStaticIndexBuffer::CreateRHIBuffer_Async()
 	return CreateRHIBuffer_Internal<false>();
 }
 
+void FRawStaticIndexBuffer::CopyRHIForStreaming(const FRawStaticIndexBuffer& Other, bool InAllowCPUAccess)
+{
+	// Copy serialized properties.
+	CachedNumIndices = Other.CachedNumIndices;
+	b32Bit = Other.b32Bit;
+
+	// Handle CPU access.
+	if (InAllowCPUAccess)
+	{
+		IndexStorage = Other.IndexStorage;
+	}
+
+	// Copy resource references.
+	IndexBufferRHI = Other.IndexBufferRHI;
+}
+
 void FRawStaticIndexBuffer::InitRHI()
 {
 	IndexBufferRHI = CreateRHIBuffer_RenderThread();
diff --git a/Engine/Source/Runtime/Engine/Private/Rendering/ColorVertexBuffer.cpp b/Engine/Source/Runtime/Engine/Private/Rendering/ColorVertexBuffer.cpp
--- a/Engine/Source/Runtime/Engine/Private/Rendering/ColorVertexBuffer.cpp
+++ b/Engine/Source/Runtime/Engine/Private/Rendering/ColorVertexBuffer.cpp
@@ -417,6 +417,28 @@ FVertexBufferRHIRef FColorVertexBuffer::CreateRHIBuffer_Async()
 	return CreateRHIBuffer_Internal<false>();
 }
 
+void FColorVertexBuffer::CopyRHIForStreaming(const FColorVertexBuffer& Other, bool InAllowCPUAccess)
+{
+	// Copy serialized properties.
+	Stride = Other.Stride;
+	NumVertices = Other.NumVertices;
+
+	// Handle CPU access.
+	if (InAllowCPUAccess)
+	{
+		NeedsCPUAccess = Other.NeedsCPUAccess;
+		AllocateData(NeedsCPUAccess);
+	}
+	else
+	{
+		NeedsCPUAccess = false;
+	}
+
+	// Copy resource references.
+	VertexBufferRHI = Other.VertexBufferRHI;
+	ColorComponentsSRV = Other.ColorComponentsSRV;
+}
+
 void FColorVertexBuffer::InitRHI()
 {
 	SCOPED_LOADTIMER(FColorVertexBuffer_InitRHI);
diff --git a/Engine/Source/Runtime/Engine/Private/Rendering/StaticMeshVertexBuffer.cpp b/Engine/Source/Runtime/Engine/Private/Rendering/StaticMeshVertexBuffer.cpp
--- a/Engine/Source/Runtime/Engine/Private/Rendering/StaticMeshVertexBuffer.cpp
+++ b/Engine/Source/Runtime/Engine/Private/Rendering/StaticMeshVertexBuffer.cpp
@@ -328,6 +328,34 @@ FVertexBufferRHIRef FStaticMeshVertexBuffer::CreateTexCoordRHIBuffer_Async()
 	return CreateTexCoordRHIBuffer_Internal<false>();
 }
 
+void FStaticMeshVertexBuffer::CopyRHIForStreaming(const FStaticMeshVertexBuffer& Other, bool InAllowCPUAccess)
+{
+	// Copy serialized properties.
+	TangentsStride = Other.TangentsStride;
+	TexcoordStride = Other.TexcoordStride;
+	NumTexCoords = Other.NumTexCoords;
+	NumVertices = Other.NumVertices;
+	bUseFullPrecisionUVs = Other.bUseFullPrecisionUVs;
+	bUseHighPrecisionTangentBasis = Other.bUseHighPrecisionTangentBasis;
+
+	// Handle CPU access.
+	if (InAllowCPUAccess)
+	{
+		NeedsCPUAccess = Other.NeedsCPUAccess;
+		AllocateData(NeedsCPUAccess);
+	}
+	else
+	{
+		NeedsCPUAccess = false;
+	}
+
+	// Copy resource references.
+	TangentsVertexBuffer.VertexBufferRHI = Other.TangentsVertexBuffer.VertexBufferRHI;
+	TexCoordVertexBuffer.VertexBufferRHI = Other.TexCoordVertexBuffer.VertexBufferRHI;
+	TangentsSRV = Other.TangentsSRV;
+	TextureCoordinatesSRV = Other.TextureCoordinatesSRV;
+}
+
 void FStaticMeshVertexBuffer::InitRHI()
 {
 	SCOPED_LOADTIMER(FStaticMeshVertexBuffer_InitRHI);
diff --git a/Engine/Source/Runtime/Engine/Private/StaticMesh.cpp b/Engine/Source/Runtime/Engine/Private/StaticMesh.cpp
--- a/Engine/Source/Runtime/Engine/Private/StaticMesh.cpp
+++ b/Engine/Source/Runtime/Engine/Private/StaticMesh.cpp
@@ -924,7 +924,7 @@ FStaticMeshAreaWeightedSectionSampler::FStaticMeshAreaWeightedSectionSampler()
 {
 }
 
-void FStaticMeshAreaWeightedSectionSampler::Init(FStaticMeshLODResources* InOwner)
+void FStaticMeshAreaWeightedSectionSampler::Init(const FStaticMeshLODResources* InOwner)
 {
 	Owner = InOwner;
 	Initialize();
@@ -1096,10 +1096,23 @@ void FStaticMeshVertexBuffers::InitFromDynamicVertex(FLocalVertexFactory* Vertex
 		});
 };
 
-FStaticMeshLODResources::FStaticMeshLODResources()
-	: AdditionalIndexBuffers(nullptr)
-	, DistanceFieldData(nullptr)
-	, MaxDeviation(0.0f)
+void FStaticMeshVertexBuffers::CopyRHIForStreaming(const FStaticMeshVertexBuffers& Other, bool InAllowCPUAccess)
+{
+	StaticMeshVertexBuffer.CopyRHIForStreaming(Other.StaticMeshVertexBuffer, InAllowCPUAccess);
+	PositionVertexBuffer.CopyRHIForStreaming(Other.PositionVertexBuffer, InAllowCPUAccess);
+	ColorVertexBuffer.CopyRHIForStreaming(Other.ColorVertexBuffer, InAllowCPUAccess);
+}
+
+void FAdditionalStaticMeshIndexBuffers::CopyRHIForStreaming(const FAdditionalStaticMeshIndexBuffers& Other, bool InAllowCPUAccess)
+{
+	ReversedIndexBuffer.CopyRHIForStreaming(Other.ReversedIndexBuffer, InAllowCPUAccess);
+	ReversedDepthOnlyIndexBuffer.CopyRHIForStreaming(Other.ReversedDepthOnlyIndexBuffer, InAllowCPUAccess);
+	WireframeIndexBuffer.CopyRHIForStreaming(Other.WireframeIndexBuffer, InAllowCPUAccess);
+	AdjacencyIndexBuffer.CopyRHIForStreaming(Other.AdjacencyIndexBuffer, InAllowCPUAccess);
+}
+
+FStaticMeshLODData::FStaticMeshLODData()
+	: MaxDeviation(0.0f)
 	, bHasAdjacencyInfo(false)
 	, bHasDepthOnlyIndices(false)
 	, bHasReversedIndices(false)
@@ -1119,8 +1132,29 @@ FStaticMeshLODResources::FStaticMeshLODResources()
 
 FStaticMeshLODResources::~FStaticMeshLODResources()
 {
-	// DistanceFieldData.SafeDelete();
+	check(GetRefCount() == 0);
 	delete DistanceFieldData;
+	delete AdditionalIndexBuffers;
+}
+
+void FStaticMeshLODResources::CopyRHIForStreaming(const FStaticMeshLODResources& Other, bool InAllowCPUAccess)
+{
+	VertexBuffers.CopyRHIForStreaming(Other.VertexBuffers, InAllowCPUAccess);
+	IndexBuffer.CopyRHIForStreaming(Other.IndexBuffer, InAllowCPUAccess);
+	DepthOnlyIndexBuffer.CopyRHIForStreaming(Other.DepthOnlyIndexBuffer, InAllowCPUAccess);
+	if (Other.AdditionalIndexBuffers)
+	{
+		if (!AdditionalIndexBuffers)
+		{
+			AdditionalIndexBuffers = new FAdditionalStaticMeshIndexBuffers;
+		}
+		AdditionalIndexBuffers->CopyRHIForStreaming(*Other.AdditionalIndexBuffers, InAllowCPUAccess);
+	}
+	else if (AdditionalIndexBuffers)
+	{
+		delete AdditionalIndexBuffers;
+		AdditionalIndexBuffers = nullptr;
+	}
 }
 
 template <bool bIncrement>
@@ -1129,7 +1163,7 @@ void FStaticMeshLODResources::UpdateIndexMemoryStats()
 #if STATS
 	if (bIncrement)
 	{
-		StaticMeshIndexMemory += IndexBuffer.GetAllocatedSize();
+		StaticMeshIndexMemory = IndexBuffer.GetAllocatedSize();
 		StaticMeshIndexMemory += DepthOnlyIndexBuffer.GetAllocatedSize();
 
 		if (AdditionalIndexBuffers)
@@ -1346,6 +1380,18 @@ FStaticMeshRenderData::FStaticMeshRenderData()
 	}
 }
 
+FStaticMeshRenderData::~FStaticMeshRenderData()
+{
+	FStaticMeshLODResources** LODResourcesArray = LODResources.GetData();
+	for (int32 LODIndex = 0; LODIndex < LODResources.Num(); ++LODIndex)
+	{
+		LODResourcesArray[LODIndex]->Release();
+		// Prevent the array from calling the destructor to handle correctly the refcount.
+		// For compatibility reason, LODResourcesArray is using ptr directly instead of TRefCountPtr.
+		LODResourcesArray[LODIndex] = nullptr;
+	}
+	LODResources.Empty();
+}
 
 void FStaticMeshRenderData::Serialize(FArchive& Ar, UStaticMesh* Owner, bool bCooked)
 {
@@ -1369,31 +1415,7 @@ void FStaticMeshRenderData::Serialize(FArchive& Ar, UStaticMesh* Owner, bool bCo
 
 #endif // #if WITH_EDITORONLY_DATA
 
-	if (Ar.IsLoading())
-	{
-		// Load array.
-		int32 NewNum;
-		Ar << NewNum;
-		LODResources.Empty(NewNum);
-		for (int32 Index = 0; Index < NewNum; Index++)
-		{
-			new (LODResources) FStaticMeshLODResources;
-		}
-		for (int32 Index = 0; Index < NewNum; Index++)
-		{
-			LODResources[Index].Serialize(Ar, Owner, Index);
-		}
-	}
-	else
-	{
-		// Save array.
-		int32 Num = LODResources.Num();
-		Ar << Num;
-		for (int32 Index = 0; Index < Num; Index++)
-		{
-			LODResources[Index].Serialize(Ar, Owner, Index);
-		}
-	}
+	LODResources.Serialize(Ar, Owner);
 
 #if WITH_EDITOR
 	if (Ar.IsSaving())
@@ -1608,11 +1630,22 @@ void FStaticMeshRenderData::AllocateLODResources(int32 NumLODs)
 	LODVertexFactories.Reserve(NumLODs);
 	while (LODResources.Num() < NumLODs)
 	{
-		new (LODResources) FStaticMeshLODResources;
+		LODResources.Add(new FStaticMeshLODResources);
 		new (LODVertexFactories) FStaticMeshVertexFactories(GMaxRHIFeatureLevel);
 	}
 }
 
+int32 FStaticMeshRenderData::GetFirstValidLODIdx(int32 MinIdx) const
+{
+	const int32 LODCount = LODResources.Num();
+	int32 LODIndex = FMath::Clamp<int32>(MinIdx, 0, LODCount - 1);
+	while (LODIndex < LODCount && !LODResources[LODIndex].GetNumVertices())
+	{
+		++LODIndex;
+	}
+	return LODIndex;
+}
+
 FStaticMeshOccluderData::FStaticMeshOccluderData()
 {
 	VerticesSP = MakeShared<FOccluderVertexArray, ESPMode::ThreadSafe>();
@@ -2661,7 +2694,11 @@ void UStaticMesh::InitResources()
 
 	// Determine whether or not this mesh can be streamed.
 	const int32 NumLODs = GetNumLODs();
-	bIsStreamable = !NeverStream
+
+	// Most tools assume that LOD0 is available and wouldn't work correctly if not. For example lighting build.
+	// Because of this, mesh streaming is in game mode (cooked or not).
+	bIsStreamable = !GIsEditor 
+		&& !NeverStream
 		&& NumLODs > 1
 		&& !RenderData->LODResources[0].bBuffersInlined;
 		//&& !bTemporarilyDisableStreaming;
@@ -2682,7 +2719,12 @@ void UStaticMesh::InitResources()
 	{
 		LinkStreaming();
 	}
-	
+	else if (GIsEditor)
+	{
+		RenderData->CurrentFirstLODIdx = RenderData->GetFirstValidLODIdx(0);
+		SetCachedNumResidentLODs(RenderData->LODResources.Num() - RenderData->CurrentFirstLODIdx);
+	}
+
 #if	STATS
 	UStaticMesh* This = this;
 	ENQUEUE_RENDER_COMMAND(UpdateMemoryStats)(
@@ -2781,9 +2823,7 @@ int32 UStaticMesh::GetNumLODs() const
 //
 bool UStaticMesh::HasValidRenderData(bool bCheckLODForVerts, int32 LODIndex) const
 {
-	if (RenderData != nullptr
-		&& RenderData->LODResources.Num() > 0
-		&& RenderData->LODResources.GetData() != nullptr)
+	if (RenderData && RenderData->LODResources.Num() && RenderData->LODResources.GetData())
 	{
 		if (bCheckLODForVerts)
 		{
@@ -2823,7 +2863,7 @@ int32 UStaticMesh::GetNumSections(int32 InLOD) const
 }
 
 #if WITH_EDITORONLY_DATA
-static float GetUVDensity(const TArrayView<FStaticMeshLODResources>& LODResources, int32 UVIndex)
+static float GetUVDensity(const FStaticMeshLODResourcesArray& LODResources, int32 UVIndex)
 {
 	float WeightedUVDensity = 0;
 	float WeightSum = 0;
@@ -5555,7 +5595,16 @@ bool UStaticMesh::StreamOut(int32 NewMipCount)
 	check(IsInGameThread());
 	if (bIsStreamable && !PendingUpdate && RenderData.IsValid() && RenderData->bReadyForStreaming && NewMipCount < GetNumResidentMips())
 	{
-		PendingUpdate = new FStaticMeshStreamOut(this, NewMipCount);
+#if WITH_EDITOR
+		if (FPlatformProperties::HasEditorOnlyData())
+		{
+			PendingUpdate = new FStaticMeshStreamOut_DDC(this, NewMipCount);
+		}
+		else
+#endif
+		{
+			PendingUpdate = new FStaticMeshStreamOut_IO(this, NewMipCount);
+		}
 		return !PendingUpdate->IsCancelled();
 	}
 	return false;
diff --git a/Engine/Source/Runtime/Engine/Private/Streaming/StaticMeshUpdate.cpp b/Engine/Source/Runtime/Engine/Private/Streaming/StaticMeshUpdate.cpp
--- a/Engine/Source/Runtime/Engine/Private/Streaming/StaticMeshUpdate.cpp
+++ b/Engine/Source/Runtime/Engine/Private/Streaming/StaticMeshUpdate.cpp
@@ -12,6 +12,14 @@ StaticMeshUpdate.cpp: Helpers to stream in and out static mesh LODs.
 #include "Serialization/MemoryReader.h"
 #include "Streaming/RenderAssetUpdate.inl"
 
+int32 GStreamingMaxReferenceChecks = 3;
+static FAutoConsoleVariableRef CVarGPUSortFrameCountBeforeBufferShrinking(
+	TEXT("r.Streaming.MaxReferenceChecksBeforeStreamOut"),
+	GStreamingMaxReferenceChecks,
+	TEXT("Number of times the engine wait for references to be released before forcing streamout. (default=3)"),
+	ECVF_Default
+);
+
 // Instantiate TRenderAssetUpdate for FStaticMeshUpdateContext
 template class TRenderAssetUpdate<FStaticMeshUpdateContext>;
 
@@ -265,13 +273,14 @@ void FStaticMeshStreamIn::DoCancel(const FContext& Context)
 	DoFinishUpdate(Context);
 }
 
-FStaticMeshStreamOut::FStaticMeshStreamOut(UStaticMesh* InMesh, int32 InRequestedMips)
+#if WITH_EDITOR
+FStaticMeshStreamOut_DDC::FStaticMeshStreamOut_DDC(UStaticMesh* InMesh, int32 InRequestedMips)
 	: FStaticMeshUpdate(InMesh, InRequestedMips)
 {
 	PushTask(FContext(InMesh, TT_None), TT_Render, SRA_UPDATE_CALLBACK(DoReleaseBuffers), TT_None, nullptr);
 }
 
-void FStaticMeshStreamOut::DoReleaseBuffers(const FContext& Context)
+void FStaticMeshStreamOut_DDC ::DoReleaseBuffers(const FContext& Context)
 {
 	check(Context.CurrentThread == TT_Render);
 	UStaticMesh* Mesh = Context.Mesh;
@@ -285,27 +294,11 @@ void FStaticMeshStreamOut::DoReleaseBuffers(const FContext& Context)
 		Mesh->SetCachedNumResidentLODs(static_cast<uint8>(RenderData->LODResources.Num() - PendingFirstMip));
 
 		TRHIResourceUpdateBatcher<GStaticMeshMaxNumResourceUpdatesPerBatch> Batcher;
-
 		for (int32 LODIdx = CurrentFirstLODIdx; LODIdx < PendingFirstMip; ++LODIdx)
 		{
 			FStaticMeshLODResources& LODResource = RenderData->LODResources[LODIdx];
-			FStaticMeshVertexBuffers& VBs = LODResource.VertexBuffers;
 			LODResource.DecrementMemoryStats();
-			VBs.StaticMeshVertexBuffer.ReleaseRHIForStreaming(Batcher);
-			VBs.PositionVertexBuffer.ReleaseRHIForStreaming(Batcher);
-			VBs.ColorVertexBuffer.ReleaseRHIForStreaming(Batcher);
-			// Index buffers don't need to update SRV so we can reuse ReleaseRHI
-			LODResource.IndexBuffer.ReleaseRHIForStreaming(Batcher);
-			LODResource.DepthOnlyIndexBuffer.ReleaseRHIForStreaming(Batcher);
-
-			if (LODResource.AdditionalIndexBuffers)
-			{
-				LODResource.AdditionalIndexBuffers->ReversedIndexBuffer.ReleaseRHIForStreaming(Batcher);
-				LODResource.AdditionalIndexBuffers->ReversedDepthOnlyIndexBuffer.ReleaseRHIForStreaming(Batcher);
-				LODResource.AdditionalIndexBuffers->WireframeIndexBuffer.ReleaseRHIForStreaming(Batcher);
-				LODResource.AdditionalIndexBuffers->AdjacencyIndexBuffer.ReleaseRHIForStreaming(Batcher);
-			}
-
+			LODResource.ReleaseRHIForStreaming(Batcher);
 #if RHI_RAYTRACING
 			if (IsRayTracingEnabled())
 			{
@@ -315,6 +308,89 @@ void FStaticMeshStreamOut::DoReleaseBuffers(const FContext& Context)
 		}
 	}
 }
+#endif // WITH_EDITOR
+
+
+FStaticMeshStreamOut_IO::FStaticMeshStreamOut_IO(UStaticMesh* InMesh, int32 InRequestedMips)
+	: FStaticMeshUpdate(InMesh, InRequestedMips)
+{
+	PushTask(FContext(InMesh, TT_None), TT_Render, SRA_UPDATE_CALLBACK(DoReleaseBuffers), TT_None, nullptr);
+}
+
+void FStaticMeshStreamOut_IO::DoReleaseBuffers(const FContext& Context)
+{
+	check(Context.CurrentThread == TT_Render);
+	UStaticMesh* Mesh = Context.Mesh;
+	FStaticMeshRenderData* RenderData = Context.RenderData;
+	if (!IsCancelled() && Mesh && RenderData)
+	{
+		RenderData->CurrentFirstLODIdx = PendingFirstMip;
+		Mesh->SetCachedNumResidentLODs(static_cast<uint8>(RenderData->LODResources.Num() - PendingFirstMip));
+
+		for (int32 LODIdx = CurrentFirstLODIdx; LODIdx < PendingFirstMip; ++LODIdx)
+		{
+			FStaticMeshLODResources* PreviousLODResources = RenderData->LODResources.GetData()[LODIdx];
+			FStaticMeshLODResources* NextLODResources = new FStaticMeshLODResources;
+
+			// Copy everything but the CPU data.
+			static_cast<FStaticMeshLODData&>(*NextLODResources) = static_cast<FStaticMeshLODData&>(*PreviousLODResources);
+			NextLODResources->CopyRHIForStreaming(*PreviousLODResources, false);
+
+			// Change which LODResource the static mesh render data points to.
+			// This will prevent new references from being made.
+			RenderData->LODResources.GetData()[LODIdx] = NextLODResources; 
+			LODResourcesToRelease.Add(PreviousLODResources);
+		}
+		DoReleaseLODResources(Context);
+	}
+}
+
+void FStaticMeshStreamOut_IO::DoReleaseLODResources(const FContext& Context)
+{
+	check(Context.CurrentThread == TT_Render);
+
+	bool bHasExternalReferences = false;
+	for (FStaticMeshLODResources& LODResource : LODResourcesToRelease)
+	{
+		if (LODResource.GetRefCount() != 1)
+		{
+			bHasExternalReferences = true;
+			++NumReferenceChecks;
+			break;
+		}
+	}
+
+	// It is expected that systems that keep a persistent reference to streamed out LODs will detect the streamout
+	// and recache their loddata. This is why we still streamout here after X attemps, assuming there is no active access on the loddata.
+	if (!bHasExternalReferences || NumReferenceChecks >= GStreamingMaxReferenceChecks)
+	{
+		ReleaseLODResources();
+	}
+	else
+	{
+		bDeferExecution = true;
+		PushTask(Context, TT_Render, SRA_UPDATE_CALLBACK(DoReleaseLODResources), TT_Render, SRA_UPDATE_CALLBACK(DoReleaseLODResources));
+	}
+}
+
+void FStaticMeshStreamOut_IO::ReleaseLODResources()
+{
+	if (LODResourcesToRelease.Num())
+	{
+	    TRHIResourceUpdateBatcher<GStaticMeshMaxNumResourceUpdatesPerBatch> Batcher;
+	    FStaticMeshLODResources** LODResourcesArray = LODResourcesToRelease.GetData();
+	    for (int32 LODIndex = 0; LODIndex < LODResourcesToRelease.Num(); ++LODIndex)
+	    {
+		    FStaticMeshLODResources* LODResources = LODResourcesArray[LODIndex];
+		    LODResourcesArray[LODIndex] = nullptr;
+    
+		    LODResources->ReleaseRHIForStreaming(Batcher);
+		    LODResources->ReleaseResources(); // Also decrement memory stats.
+		    LODResources->Release();
+	    }
+	    LODResourcesToRelease.Empty();
+	}
+}
 
 void FStaticMeshStreamIn_IO::FCancelIORequestsTask::DoWork()
 {
diff --git a/Engine/Source/Runtime/Engine/Private/Streaming/StaticMeshUpdate.h b/Engine/Source/Runtime/Engine/Private/Streaming/StaticMeshUpdate.h
--- a/Engine/Source/Runtime/Engine/Private/Streaming/StaticMeshUpdate.h
+++ b/Engine/Source/Runtime/Engine/Private/Streaming/StaticMeshUpdate.h
@@ -122,16 +122,41 @@ private:
 	void CreateBuffers_Internal(const FContext& Context);
 };
 
-class FStaticMeshStreamOut : public FStaticMeshUpdate
+#if WITH_EDITOR
+/** A streamout that doesn't actually touches the CPU data. Required because DDC stream in doesn't reset. */
+class FStaticMeshStreamOut_DDC : public FStaticMeshUpdate
 {
 public:
-	FStaticMeshStreamOut(UStaticMesh* InMesh, int32 InRequestedMips);
+	FStaticMeshStreamOut_DDC(UStaticMesh* InMesh, int32 InRequestedMips);
+
+private:
+	/** Release RHI buffers and update SRVs */
+	void DoReleaseBuffers(const FContext& Context);
+};
+#endif
 
-	virtual ~FStaticMeshStreamOut() {}
+class FStaticMeshStreamOut_IO : public FStaticMeshUpdate
+{
+public:
+	FStaticMeshStreamOut_IO(UStaticMesh* InMesh, int32 InRequestedMips);
+
+	virtual ~FStaticMeshStreamOut_IO() { ReleaseLODResources(); }
 
 private:
+
+	int32 NumReferenceChecks = 0;
+	TIndirectArray<FStaticMeshLODResources>  LODResourcesToRelease;
+
 	/** Release RHI buffers and update SRVs */
 	void DoReleaseBuffers(const FContext& Context);
+
+	/** Check reference and possibly release LODResourcesToRelease */
+	void DoReleaseLODResources(const FContext& Context);
+
+	/** Release all LODs in LODResourcesToRelease */
+	void ReleaseLODResources();
+
+
 };
 
 class FStaticMeshStreamIn_IO : public FStaticMeshStreamIn
diff --git a/Engine/Source/Runtime/Engine/Private/Streaming/StreamingManagerTexture.cpp b/Engine/Source/Runtime/Engine/Private/Streaming/StreamingManagerTexture.cpp
--- a/Engine/Source/Runtime/Engine/Private/Streaming/StreamingManagerTexture.cpp
+++ b/Engine/Source/Runtime/Engine/Private/Streaming/StreamingManagerTexture.cpp
@@ -2329,9 +2329,9 @@ bool FRenderAssetStreamingManager::HandleInvestigateRenderAssetCommand(const TCH
 						}
 
 						// LOD group MaxResolution clamp : see UTextureLODSettings::CalculateLODBias(), included in CachedCombinedLODBias
-						const int32 MipCountBeforeMaxRes = StreamingRenderAsset.MipCount - Texture2D->NumCinematicMipLevels -
+						const int32 MipCountBeforeMaxRes = StreamingRenderAsset.MipCount - RenderAsset->NumCinematicMipLevels -
 							(StreamingRenderAsset.LODGroup == TEXTUREGROUP_UI ? GUITextureLODBias : 0) - 
-							(FPlatformProperties::RequiresCookedData() ? 0 : (LODGroupInfo.LODBias + Texture2D->LODBias));
+							(FPlatformProperties::RequiresCookedData() ? 0 : (LODGroupInfo.LODBias + (Texture2D ? Texture2D->LODBias : 0)));
 						const int32 MaxResBias = MipCountBeforeMaxRes - (LODGroupInfo.MaxLODMipCount + 1);
 						if (MaxResBias > 0)
 						{
@@ -2340,7 +2340,7 @@ bool FRenderAssetStreamingManager::HandleInvestigateRenderAssetCommand(const TCH
 						}
 
 						// Asset LODBias : see UTextureLODSettings::CalculateLODBias(), included in CachedCombinedLODBias
-						if (Texture2D->LODBias)
+						if (Texture2D && Texture2D->LODBias)
 						{
 							if (FPlatformProperties::RequiresCookedData())
 							{
@@ -2354,10 +2354,10 @@ bool FRenderAssetStreamingManager::HandleInvestigateRenderAssetCommand(const TCH
 						}
 
 						// Asset cinematic mips : see UTextureLODSettings::CalculateLODBias() and FStreamingRenderAsset::UpdateDynamicData(), included in CachedCombinedLODBias
-						if (Texture2D->NumCinematicMipLevels && !(StreamingRenderAsset.bForceFullyLoad && Texture2D->bUseCinematicMipLevels))
+						if (RenderAsset->NumCinematicMipLevels && !(StreamingRenderAsset.bForceFullyLoad && RenderAsset->bUseCinematicMipLevels))
 						{
-							BiasDesc += FString::Printf(TEXT(" [Asset.Cine:%d]"), Texture2D->NumCinematicMipLevels);
-							CumuBias += Texture2D->NumCinematicMipLevels;
+							BiasDesc += FString::Printf(TEXT(" [Asset.Cine:%d]"), RenderAsset->NumCinematicMipLevels);
+							CumuBias += RenderAsset->NumCinematicMipLevels;
 						}
 					}
 
diff --git a/Engine/Source/Runtime/Engine/Public/RawIndexBuffer.h b/Engine/Source/Runtime/Engine/Public/RawIndexBuffer.h
--- a/Engine/Source/Runtime/Engine/Public/RawIndexBuffer.h
+++ b/Engine/Source/Runtime/Engine/Public/RawIndexBuffer.h
@@ -254,6 +254,9 @@ public:
 	FIndexBufferRHIRef CreateRHIBuffer_RenderThread();
 	FIndexBufferRHIRef CreateRHIBuffer_Async();
 
+	/** Copy everything, keeping reference to the same RHI resources. */
+	void CopyRHIForStreaming(const FRawStaticIndexBuffer& Other, bool InAllowCPUAccess);
+
 	/** Take over ownership of IntermediateBuffer */
 	template <uint32 MaxNumUpdates>
 	void InitRHIForStreaming(FRHIIndexBuffer* IntermediateBuffer, TRHIResourceUpdateBatcher<MaxNumUpdates>& Batcher)
diff --git a/Engine/Source/Runtime/Engine/Public/Rendering/ColorVertexBuffer.h b/Engine/Source/Runtime/Engine/Public/Rendering/ColorVertexBuffer.h
--- a/Engine/Source/Runtime/Engine/Public/Rendering/ColorVertexBuffer.h
+++ b/Engine/Source/Runtime/Engine/Public/Rendering/ColorVertexBuffer.h
@@ -141,6 +141,9 @@ public:
 	FVertexBufferRHIRef CreateRHIBuffer_RenderThread();
 	FVertexBufferRHIRef CreateRHIBuffer_Async();
 
+	/** Copy everything, keeping reference to the same RHI resources. */
+	void CopyRHIForStreaming(const FColorVertexBuffer& Other, bool InAllowCPUAccess);
+
 	/** Similar to Init/ReleaseRHI but only update existing SRV so references to the SRV stays valid */
 	template <uint32 MaxNumUpdates>
 	void InitRHIForStreaming(FRHIVertexBuffer* IntermediateBuffer, TRHIResourceUpdateBatcher<MaxNumUpdates>& Batcher)
diff --git a/Engine/Source/Runtime/Engine/Public/Rendering/PositionVertexBuffer.h b/Engine/Source/Runtime/Engine/Public/Rendering/PositionVertexBuffer.h
--- a/Engine/Source/Runtime/Engine/Public/Rendering/PositionVertexBuffer.h
+++ b/Engine/Source/Runtime/Engine/Public/Rendering/PositionVertexBuffer.h
@@ -98,6 +98,9 @@ public:
 	FVertexBufferRHIRef CreateRHIBuffer_RenderThread();
 	FVertexBufferRHIRef CreateRHIBuffer_Async();
 
+	/** Copy everything, keeping reference to the same RHI resources. */
+	void CopyRHIForStreaming(const FPositionVertexBuffer& Other, bool InAllowCPUAccess);
+
 	/** Similar to Init/ReleaseRHI but only update existing SRV so references to the SRV stays valid */
 	template <uint32 MaxNumUpdates>
 	void InitRHIForStreaming(FRHIVertexBuffer* IntermediateBuffer, TRHIResourceUpdateBatcher<MaxNumUpdates>& Batcher)
diff --git a/Engine/Source/Runtime/Engine/Public/Rendering/StaticMeshVertexBuffer.h b/Engine/Source/Runtime/Engine/Public/Rendering/StaticMeshVertexBuffer.h
--- a/Engine/Source/Runtime/Engine/Public/Rendering/StaticMeshVertexBuffer.h
+++ b/Engine/Source/Runtime/Engine/Public/Rendering/StaticMeshVertexBuffer.h
@@ -394,6 +394,9 @@ public:
 	FVertexBufferRHIRef CreateTexCoordRHIBuffer_RenderThread();
 	FVertexBufferRHIRef CreateTexCoordRHIBuffer_Async();
 
+	/** Copy everything, keeping reference to the same RHI resources. */
+	void CopyRHIForStreaming(const FStaticMeshVertexBuffer& Other, bool InAllowCPUAccess);
+
 	/** Similar to Init/ReleaseRHI but only update existing SRV so references to the SRV stays valid */
 	template <uint32 MaxNumUpdates>
 	void InitRHIForStreaming(
diff --git a/Engine/Source/Runtime/Engine/Public/StaticMeshResources.h b/Engine/Source/Runtime/Engine/Public/StaticMeshResources.h
--- a/Engine/Source/Runtime/Engine/Public/StaticMeshResources.h
+++ b/Engine/Source/Runtime/Engine/Public/StaticMeshResources.h
@@ -243,11 +243,11 @@ protected:
 struct ENGINE_API FStaticMeshAreaWeightedSectionSampler : FWeightedRandomSampler
 {
 	FStaticMeshAreaWeightedSectionSampler();
-	void Init(FStaticMeshLODResources* InOwner);
+	void Init(const FStaticMeshLODResources* InOwner);
 	virtual float GetWeights(TArray<float>& OutWeights)override;
 
 protected:
-	FStaticMeshLODResources* Owner;
+	TRefCountPtr<const FStaticMeshLODResources> Owner;
 };
 
 typedef TArray<FStaticMeshSectionAreaWeightedTriangleSampler, FMemoryImageAllocator> FStaticMeshSectionAreaWeightedTriangleSamplerArray;
@@ -308,6 +308,9 @@ struct FStaticMeshVertexBuffers
 
 	/* This is a temporary function to refactor and convert old code, do not copy this as is and try to build your data as SoA from the beginning.*/
 	void ENGINE_API InitModelVF(FLocalVertexFactory* VertexFactory);
+
+	/** Copy everything, keeping reference to the same RHI resources. */
+	void CopyRHIForStreaming(const FStaticMeshVertexBuffers& Other, bool InAllowCPUAccess);
 };
 
 struct FAdditionalStaticMeshIndexBuffers
@@ -320,30 +323,22 @@ struct FAdditionalStaticMeshIndexBuffers
 	FRawStaticIndexBuffer WireframeIndexBuffer;
 	/** Index buffer containing adjacency information required by tessellation. */
 	FRawStaticIndexBuffer AdjacencyIndexBuffer;
+
+	/** Copy everything, keeping reference to the same RHI resources. */
+	void CopyRHIForStreaming(const FAdditionalStaticMeshIndexBuffers& Other, bool InAllowCPUAccess);
 };
 
-/** Rendering resources needed to render an individual static mesh LOD. */
-struct FStaticMeshLODResources
+/** The part of FStaticMeshLODResources which is not streamable. */
+struct FStaticMeshLODData
 {
-	FStaticMeshVertexBuffers VertexBuffers;
-
-	/** Index buffer resource for rendering. */
-	FRawStaticIndexBuffer IndexBuffer;
-
-	/** Index buffer resource for rendering in depth only passes. */
-	FRawStaticIndexBuffer DepthOnlyIndexBuffer;
-
-	FAdditionalStaticMeshIndexBuffers* AdditionalIndexBuffers;
-
-	/** Resources used for ray tracing */
-	FRayTracingGeometry RayTracingGeometry;
+	FStaticMeshLODData();
 
 	/** Sections for this LOD. */
 	using FStaticMeshSectionArray = TArray<FStaticMeshSection, TInlineAllocator<1>>;
 	FStaticMeshSectionArray Sections;
 
 	/** Distance field data associated with this mesh, null if not present.  */
-	class FDistanceFieldVolumeData* DistanceFieldData; 
+	class FDistanceFieldVolumeData* DistanceFieldData = nullptr; 
 
 	/** The maximum distance by which this LOD deviates from the base from which it was generated. */
 	float MaxDeviation;
@@ -372,13 +367,6 @@ struct FStaticMeshLODResources
 
 	/** True if this LOD is optional. That is, vertex and index data may not be available */
 	uint32 bIsOptionalLOD : 1;
-	
-	/**	Allows uniform random selection of mesh sections based on their area. */
-	FStaticMeshAreaWeightedSectionSampler AreaWeightedSampler;
-	/**	Allows uniform random selection of triangles on each mesh section based on triangle area. */
-	FStaticMeshSectionAreaWeightedTriangleSamplerArray AreaWeightedSectionSamplers;
-	/** Allows uniform random selection of triangles on GPU. It is not cooked and serialised but created at runtime from AreaWeightedSectionSamplers when it is available and static mesh bSupportGpuUniformlyDistributedSampling=true*/
-	FStaticMeshSectionAreaWeightedTriangleSamplerBuffer AreaWeightedSectionSamplersBuffer;
 
 	uint32 DepthOnlyNumTriangles;
 
@@ -399,12 +387,76 @@ struct FStaticMeshLODResources
 	/** Map of wedge index to vertex index. Each LOD need one*/
 	TArray<int32> WedgeMap;
 #endif
+};
+
+/** 
+ * Rendering resources needed to render an individual static mesh LOD.
+ * They are ref counted to prevent streaming out LODs while they are being accessed for read.
+ * The stream out implies swapping the LODResoures entry by a dummy empty one. This is always executed within a render command.
+ * This means that any direct reference to a LODResources is safe within the scope of any render commands.
+ * If the reference is required on the gamethread, async threads or RHI threads (not synchronized at the end of the render command that created the tasks),
+ * then the ref count must be incremented / decremented through the use of a TRefCountPtr. 
+ * The stream out logic in FStaticMeshStreamOut will postpone the operation if the ref count is not 1.
+ * "1" accounting for the reference in the FStaticMeshRenderData::LODResources. In order to prevent postponing indefinitely the streamout, 
+ * it is required that anyone using a TRefCountPtr also check for CurrentFirstLODIdx for things that may run every frame.
+ */
+struct FStaticMeshLODResources : public FRefCountBase, public FStaticMeshLODData
+{
+public:
+
+	FStaticMeshVertexBuffers VertexBuffers;
+
+	/** Index buffer resource for rendering. */
+	FRawStaticIndexBuffer IndexBuffer;
+
+	/** Index buffer resource for rendering in depth only passes. */
+	FRawStaticIndexBuffer DepthOnlyIndexBuffer;
+
+	FAdditionalStaticMeshIndexBuffers* AdditionalIndexBuffers = nullptr;
+
+	/** Geometry for ray tracing. */
+	FRayTracingGeometry RayTracingGeometry;
+
+	/**	Allows uniform random selection of mesh sections based on their area. */
+	FStaticMeshAreaWeightedSectionSampler AreaWeightedSampler;
+	/**	Allows uniform random selection of triangles on each mesh section based on triangle area. */
+	FStaticMeshSectionAreaWeightedTriangleSamplerArray AreaWeightedSectionSamplers;
+	/** Allows uniform random selection of triangles on GPU. It is not cooked and serialised but created at runtime from AreaWeightedSectionSamplers when it is available and static mesh bSupportGpuUniformlyDistributedSampling=true*/
+	FStaticMeshSectionAreaWeightedTriangleSamplerBuffer AreaWeightedSectionSamplersBuffer;
 	
-	/** Default constructor. */
-	ENGINE_API FStaticMeshLODResources();
+	/** Default constructor. Add a reference if not stored with a TRefCountPtr */
+	FORCEINLINE FStaticMeshLODResources(bool bAddRef = true)
+	{
+		if (bAddRef)
+		{
+			AddRef();
+		}
+	}
 
 	ENGINE_API ~FStaticMeshLODResources();
 
+	/** Copy everything, keeping reference to the same RHI resources. */
+	void CopyRHIForStreaming(const FStaticMeshLODResources& Other, bool InAllowCPUAccess);
+
+	template <typename TBatcher>
+	void ReleaseRHIForStreaming(TBatcher& Batcher)
+	{
+		VertexBuffers.StaticMeshVertexBuffer.ReleaseRHIForStreaming(Batcher);
+		VertexBuffers.PositionVertexBuffer.ReleaseRHIForStreaming(Batcher);
+		VertexBuffers.ColorVertexBuffer.ReleaseRHIForStreaming(Batcher);
+
+		IndexBuffer.ReleaseRHIForStreaming(Batcher);
+		DepthOnlyIndexBuffer.ReleaseRHIForStreaming(Batcher);
+
+		if (AdditionalIndexBuffers)
+		{
+			AdditionalIndexBuffers->ReversedIndexBuffer.ReleaseRHIForStreaming(Batcher);
+			AdditionalIndexBuffers->ReversedDepthOnlyIndexBuffer.ReleaseRHIForStreaming(Batcher);
+			AdditionalIndexBuffers->WireframeIndexBuffer.ReleaseRHIForStreaming(Batcher);
+			AdditionalIndexBuffers->AdjacencyIndexBuffer.ReleaseRHIForStreaming(Batcher);
+		}
+	}
+
 	/** Initializes all rendering resources. */
 	void InitResources(UStaticMesh* Parent);
 
@@ -509,7 +561,7 @@ private:
 	friend class FStaticMeshRenderData;
 	friend class FStaticMeshStreamIn;
 	friend class FStaticMeshStreamIn_IO;
-	friend class FStaticMeshStreamOut;
+	friend class FStaticMeshStreamOut_DDC;
 };
 
 struct ENGINE_API FStaticMeshVertexFactories
@@ -552,7 +604,7 @@ struct ENGINE_API FStaticMeshVertexFactories
 	void ReleaseResources();
 };
 
-using FStaticMeshLODResourcesArray = TArray<FStaticMeshLODResources>;
+using FStaticMeshLODResourcesArray = TIndirectArray<FStaticMeshLODResources>;
 using FStaticMeshVertexFactoriesArray = TArray<FStaticMeshVertexFactories>;
 
 /**
@@ -561,10 +613,15 @@ using FStaticMeshVertexFactoriesArray = TArray<FStaticMeshVertexFactories>;
 class FStaticMeshRenderData
 {
 public:
+
 	/** Default constructor. */
 	ENGINE_API FStaticMeshRenderData();
+	ENGINE_API ~FStaticMeshRenderData();
 
-	/** Per-LOD resources. */
+	/**
+	 * Per-LOD resources. For compatibility reasons, the FStaticMeshLODResources array are not referenced through TRefCountPtr.
+	 * The LODResource still has a ref count of at least 1, see FStaticMeshLODResources() constructor.
+	 */
 	FStaticMeshLODResourcesArray LODResources;
 	FStaticMeshVertexFactoriesArray LODVertexFactories;
 
@@ -639,6 +696,25 @@ public:
 	ENGINE_API void ResolveSectionInfo(UStaticMesh* Owner);
 #endif // #if WITH_EDITORONLY_DATA
 
+	/** Return first valid LOD index starting at MinLODIdx. */
+	ENGINE_API int32 GetFirstValidLODIdx(int32 MinLODIdx) const;
+
+	/** Return the current first LODIdx that can be used. */
+	FORCEINLINE int32 GetCurrentFirstLODIdx(int32 MinLODIdx) const
+	{
+		return GetFirstValidLODIdx(FMath::Max<int32>(CurrentFirstLODIdx, MinLODIdx));
+	}
+
+	/** 
+	 * Return the current first LOD that can be used for rendering starting at MinLODIdx.
+	 * This takes into account the streaming status from CurrentFirstLODIdx, 
+	 * and MinLODIdx is expected to be UStaticMesh::MinLOD, which is platform specific.
+	 */
+	FORCEINLINE const FStaticMeshLODResources* GetCurrentFirstLOD(int32 MinLODIdx) const
+	{
+		return &LODResources[GetCurrentFirstLODIdx(MinLODIdx)];
+	}
+
 private:
 #if WITH_EDITORONLY_DATA
 	/** Allow the editor to explicitly update section information. */
