- Tracking down stall on PC by wrapping a scope on IsThisApplicationForeground which is where I suspect the stall is from

#rnx

diff --git a/Engine/Plugins/FX/Niagara/Source/Niagara/Private/NiagaraWorldManager.cpp b/Engine/Plugins/FX/Niagara/Source/Niagara/Private/NiagaraWorldManager.cpp
--- a/Engine/Plugins/FX/Niagara/Source/Niagara/Private/NiagaraWorldManager.cpp
+++ b/Engine/Plugins/FX/Niagara/Source/Niagara/Private/NiagaraWorldManager.cpp
@@ -574,7 +574,10 @@ void FNiagaraWorldManager::Tick(ETickingGroup TickGroup, float DeltaSeconds, ELe
 		FNiagaraSharedObject::FlushDeletionList();
 
 #if PLATFORM_DESKTOP
-		bAppHasFocus = FPlatformApplicationMisc::IsThisApplicationForeground();
+		{
+			QUICK_SCOPE_CYCLE_COUNTER(STAT_Niagara_IsThisApplicationForeground);
+			bAppHasFocus = FPlatformApplicationMisc::IsThisApplicationForeground();
+		}
 #else
 		bAppHasFocus = true;
 #endif
