- Adding support for gathering tangent basis from vertex sampling
- Fixing some more skel mesh functions post engine merge and instance data coming before other inputs

#rnx

diff --git a/Engine/Plugins/FX/Niagara/Shaders/Private/NiagaraDataInterfaceSkeletalMesh.ush b/Engine/Plugins/FX/Niagara/Shaders/Private/NiagaraDataInterfaceSkeletalMesh.ush
--- a/Engine/Plugins/FX/Niagara/Shaders/Private/NiagaraDataInterfaceSkeletalMesh.ush
+++ b/Engine/Plugins/FX/Niagara/Shaders/Private/NiagaraDataInterfaceSkeletalMesh.ush
@@ -332,11 +332,6 @@ void DISKelMesh_GetTriUV(in FDISkelMeshContext DIContext, in uint TriangleIndex,
 
 void DISKelMesh_RandomTriangle(in FDISkelMeshContext DIContext, uint Seed1, uint Seed2, uint Seed3, out int OutTriangle, out float3 OutBaryCoord)
 {
-	//float RandT0 = NiagaraInternalNoise(Seed1, Seed2, Seed3);
-	//OutTriangle = min(uint(RandT0 * float(DIContext.MeshTriangleCount)), DIContext.MeshTriangleCount - 1); // avoid % by using mul/min to Tri = MeshTriangleCountName
-	//OutBaryCoord = DISKelMesh_RandomBarycentricCoord(Seed1, Seed2, Seed3);
-
-
 	float RandT0 = NiagaraInternalNoise(Seed1, Seed2, Seed3);
 	int TriangleIndex = min(int(RandT0 * float(DIContext.MeshTriangleCount)), (int)DIContext.MeshTriangleCount - 1); // avoid % by using mul/min to Tri = MeshTriangleCountName
 
@@ -488,19 +483,25 @@ void DISkelMesh_GetTriColor(in FDISkelMeshContext DIContext, in uint TriangleInd
     OutColor = (DIContext.MeshColorBuffer[VertexIndex0] * BaryCoord.x) + (DIContext.MeshColorBuffer[VertexIndex1] * BaryCoord.y) + (DIContext.MeshColorBuffer[VertexIndex2] * BaryCoord.z);
 }
 
-void DISkelMesh_GetSkinnedVertex(in FDISkelMeshContext DIContext, in int VertexIndex, out float3 OutPosition, out float3 OutVelocity)
+void DISkelMesh_GetSkinnedVertex(in FDISkelMeshContext DIContext, in int VertexIndex, out float3 OutPosition, out float3 OutVelocity, out float3 OutNormal, out float3 OutBinormal, out float3 OutTangent)
 {
     FDISkelMeshSkinnedVertex SkinnedVertex = DISkelMesh_GetSkinnedVertex(DIContext, VertexIndex);
     OutPosition = SkinnedVertex.Position;
     OutVelocity = (SkinnedVertex.Position - SkinnedVertex.PrevPosition) * DIContext.InstanceInvDeltaTime;
+	OutNormal = SkinnedVertex.TangentZ;
+	OutBinormal = SkinnedVertex.TangentY;
+	OutTangent = SkinnedVertex.TangentX;
 }
 
-void DISkelMesh_GetSkinnedVertexWS(in FDISkelMeshContext DIContext, in int VertexIndex, out float3 OutPosition, out float3 OutVelocity)
+void DISkelMesh_GetSkinnedVertexWS(in FDISkelMeshContext DIContext, in int VertexIndex, out float3 OutPosition, out float3 OutVelocity, out float3 OutNormal, out float3 OutBinormal, out float3 OutTangent)
 {
     FDISkelMeshSkinnedVertex SkinnedVertex = DISkelMesh_GetSkinnedVertex(DIContext, VertexIndex);
     OutPosition = mul(float4(SkinnedVertex.Position, 1.0f), DIContext.InstanceTransform).xyz;
     float3 PrevPosition = mul(float4(SkinnedVertex.PrevPosition, 1.0f), DIContext.InstancePrevTransform).xyz;
-    OutVelocity = (OutPosition - PrevPosition) * DIContext.InstanceInvDeltaTime;
+    OutVelocity = (SkinnedVertex.Position - SkinnedVertex.PrevPosition) * DIContext.InstanceInvDeltaTime;
+	OutNormal = SkinnedVertex.TangentZ;
+	OutBinormal = SkinnedVertex.TangentY;
+	OutTangent = SkinnedVertex.TangentX;
 }
 
 void DISkelMesh_GetVertexColor(in FDISkelMeshContext DIContext, in int VertexIndex, out float4 OutColor)
@@ -653,6 +654,8 @@ void DISkelMesh_GetFilteredSocketOrBoneCount(in FDISkelMeshContext DIContext, ou
 
 void DISkelMesh_GetFilteredSocketOrBoneAt(in FDISkelMeshContext DIContext, int FilteredIndex, out int Bone)
 {
+	int NumFilteredSocketsAndBones = DIContext.NumFilteredBones + DIContext.NumFilteredSockets;
+	FilteredIndex = clamp(FilteredIndex, 0, NumFilteredSocketsAndBones - 1);
 	if ( FilteredIndex < DIContext.NumFilteredBones )
 	{
 		Bone = (int)DIContext.FilteredAndUnfilteredBones[FilteredIndex];
diff --git a/Engine/Plugins/FX/Niagara/Source/Niagara/Classes/NDISkeletalMeshCommon.h b/Engine/Plugins/FX/Niagara/Source/Niagara/Classes/NDISkeletalMeshCommon.h
--- a/Engine/Plugins/FX/Niagara/Source/Niagara/Classes/NDISkeletalMeshCommon.h
+++ b/Engine/Plugins/FX/Niagara/Source/Niagara/Classes/NDISkeletalMeshCommon.h
@@ -73,6 +73,7 @@ struct FSkinnedPositionAccessorHelper
 	FORCEINLINE void GetSkinnedTrianglePositions(FSkeletalMeshAccessorHelper& Accessor, int32 Tri, FVector& OutPos0, FVector& OutPos1, FVector& OutPos2, int32& Idx0, int32& Idx1, int32& Idx2) = delete;
 	FORCEINLINE FVector GetSkinnedVertexPosition(FSkeletalMeshAccessorHelper& Accessor, int32 VertexIndex) = delete;
 	FORCEINLINE FVector GetSkinnedVertexPreviousPosition(FSkeletalMeshAccessorHelper& Accessor, int32 VertexIndex) = delete;
+	FORCEINLINE void GetSkinnedTangentBasis(FSkeletalMeshAccessorHelper& Accessor, int32 VertexIndex, FVector& OutNormal, FVector& OutBinormal) = delete;
 	FORCEINLINE FVector GetSkinnedBonePosition(FSkeletalMeshAccessorHelper& Accessor, int32 BoneIndex) = delete;
 	FORCEINLINE FVector GetSkinnedBonePreviousPosition(FSkeletalMeshAccessorHelper& Accessor, int32 BoneIndex) = delete;
 	FORCEINLINE_DEBUGGABLE FQuat GetSkinnedBoneRotation(FSkeletalMeshAccessorHelper& Accessor, int32 BoneIndex) = delete;
@@ -125,6 +126,11 @@ struct FSkinnedPositionAccessorHelper<TIntegralConstant<ENDISkeletalMesh_Skinnin
 		return GetSkeletalMeshRefVertLocation(Accessor.Mesh, *Accessor.LODData, *Accessor.SkinWeightBuffer, VertexIndex);
 	}
 
+	FORCEINLINE void GetSkinnedTangentBasis(FSkeletalMeshAccessorHelper& Accessor, int32 VertexIndex, FVector& OutTangentX, FVector& OutTangentZ)
+	{
+		GetSkeletalMeshRefTangentBasis(Accessor.Mesh, *Accessor.LODData, *Accessor.SkinWeightBuffer, VertexIndex, OutTangentX, OutTangentZ);
+	}
+
 	FORCEINLINE_DEBUGGABLE FVector GetSkinnedBonePosition(FSkeletalMeshAccessorHelper& Accessor, int32 BoneIndex)
 	{
 		const int32 NumRealBones = Accessor.Mesh->RefSkeleton.GetRawBoneNum();
@@ -206,6 +212,11 @@ struct FSkinnedPositionAccessorHelper<TIntegralConstant<ENDISkeletalMesh_Skinnin
 		return USkeletalMeshComponent::GetSkinnedVertexPosition(Accessor.Comp, VertexIndex, *Accessor.LODData, *Accessor.SkinWeightBuffer, Accessor.SkinningData->PrevBoneRefToLocals());
 	}
 
+	FORCEINLINE void GetSkinnedTangentBasis(FSkeletalMeshAccessorHelper& Accessor, int32 VertexIndex, FVector& OutTangentX, FVector& OutTangentZ)
+	{
+		USkeletalMeshComponent::GetSkinnedTangentBasis(Accessor.Comp, VertexIndex, *Accessor.LODData, *Accessor.SkinWeightBuffer, Accessor.SkinningData->PrevBoneRefToLocals(), OutTangentX, OutTangentZ);
+	}
+
 	FORCEINLINE_DEBUGGABLE FVector GetSkinnedBonePosition(FSkeletalMeshAccessorHelper& Accessor, int32 BoneIndex)
 	{
 		return Accessor.SkinningData->CurrComponentTransforms()[BoneIndex].GetLocation();
@@ -273,6 +284,11 @@ struct FSkinnedPositionAccessorHelper<TIntegralConstant<ENDISkeletalMesh_Skinnin
 		return Accessor.SkinningData->GetPreviousPosition(Accessor.Usage.GetLODIndex(), VertexIndex);
 	}
 
+	FORCEINLINE void GetSkinnedTangentBasis(FSkeletalMeshAccessorHelper& Accessor, int32 VertexIndex, FVector& OutTangentX, FVector& OutTangentZ)
+	{
+		Accessor.SkinningData->GetTangentBasis(Accessor.Usage.GetLODIndex(), VertexIndex, OutTangentX, OutTangentZ);
+	}
+
 	FORCEINLINE_DEBUGGABLE FVector GetSkinnedBonePosition(FSkeletalMeshAccessorHelper& Accessor, int32 BoneIndex)
 	{
 		return Accessor.SkinningData->CurrComponentTransforms()[BoneIndex].GetLocation();
diff --git a/Engine/Plugins/FX/Niagara/Source/Niagara/Classes/NiagaraDataInterfaceSkeletalMesh.h b/Engine/Plugins/FX/Niagara/Source/Niagara/Classes/NiagaraDataInterfaceSkeletalMesh.h
--- a/Engine/Plugins/FX/Niagara/Source/Niagara/Classes/NiagaraDataInterfaceSkeletalMesh.h
+++ b/Engine/Plugins/FX/Niagara/Source/Niagara/Classes/NiagaraDataInterfaceSkeletalMesh.h
@@ -113,6 +113,12 @@ struct FSkeletalMeshSkinningData
 		return LODData[LODIndex].SkinnedCPUPositions[CurrIndex ^ 1][VertexIndex];
 	}
 
+	FORCEINLINE void GetTangentBasis(int32 LODIndex, int32 VertexIndex, FVector& OutTangentX, FVector& OutTangentZ)
+	{
+		OutTangentX = LODData[LODIndex].SkinnedTangentBasis[(VertexIndex * 2) + 0];
+		OutTangentZ = LODData[LODIndex].SkinnedTangentBasis[(VertexIndex * 2) + 1];
+	}
+
 	FORCEINLINE TArray<FVector>& CurrSkinnedPositions(int32 LODIndex)
 	{
 		return LODData[LODIndex].SkinnedCPUPositions[CurrIndex];
@@ -123,6 +129,11 @@ struct FSkeletalMeshSkinningData
 		return LODData[LODIndex].SkinnedCPUPositions[CurrIndex ^ 1];
 	}
 
+	FORCEINLINE TArray<FVector>& CurrSkinnedTangentBasis(int32 LODIndex)
+	{
+		return LODData[LODIndex].SkinnedTangentBasis;
+	}
+
 	FORCEINLINE TArray<FMatrix>& CurrBoneRefToLocals()
 	{
 		return BoneRefToLocals[CurrIndex];
@@ -195,6 +206,9 @@ private:
 
 		/** CPU Skinned vertex positions. Double buffered to allow accurate velocity calculation. */
 		TArray<FVector> SkinnedCPUPositions[2];
+
+		/** CPU Skinned tangent basis, where each vertex will map to TangentX + TangentZ */
+		TArray<FVector> SkinnedTangentBasis;
 	};
 	TArray<FLODData> LODData;
 
diff --git a/Engine/Plugins/FX/Niagara/Source/Niagara/Private/NDISkeletalMesh_VertexSampling.cpp b/Engine/Plugins/FX/Niagara/Source/Niagara/Private/NDISkeletalMesh_VertexSampling.cpp
--- a/Engine/Plugins/FX/Niagara/Source/Niagara/Private/NDISkeletalMesh_VertexSampling.cpp
+++ b/Engine/Plugins/FX/Niagara/Source/Niagara/Private/NDISkeletalMesh_VertexSampling.cpp
@@ -47,6 +47,9 @@ void UNiagaraDataInterfaceSkeletalMesh::GetVertexSamplingFunctions(TArray<FNiaga
 		Sig.Inputs.Add(FNiagaraVariable(FNiagaraTypeDefinition::GetIntDef(), TEXT("Vertex")));
 		Sig.Outputs.Add(FNiagaraVariable(FNiagaraTypeDefinition::GetVec3Def(), TEXT("Position")));
 		Sig.Outputs.Add(FNiagaraVariable(FNiagaraTypeDefinition::GetVec3Def(), TEXT("Velocity")));
+		Sig.Outputs.Add(FNiagaraVariable(FNiagaraTypeDefinition::GetVec3Def(), TEXT("Normal")));
+		Sig.Outputs.Add(FNiagaraVariable(FNiagaraTypeDefinition::GetVec3Def(), TEXT("Binormal")));
+		Sig.Outputs.Add(FNiagaraVariable(FNiagaraTypeDefinition::GetVec3Def(), TEXT("Tangent")));
 		Sig.bMemberFunction = true;
 		Sig.bRequiresContext = false;
 #if WITH_EDITORONLY_DATA
@@ -62,6 +65,9 @@ void UNiagaraDataInterfaceSkeletalMesh::GetVertexSamplingFunctions(TArray<FNiaga
 		Sig.Inputs.Add(FNiagaraVariable(FNiagaraTypeDefinition::GetIntDef(), TEXT("Vertex")));
 		Sig.Outputs.Add(FNiagaraVariable(FNiagaraTypeDefinition::GetVec3Def(), TEXT("Position")));
 		Sig.Outputs.Add(FNiagaraVariable(FNiagaraTypeDefinition::GetVec3Def(), TEXT("Velocity")));
+		Sig.Outputs.Add(FNiagaraVariable(FNiagaraTypeDefinition::GetVec3Def(), TEXT("Normal")));
+		Sig.Outputs.Add(FNiagaraVariable(FNiagaraTypeDefinition::GetVec3Def(), TEXT("Binormal")));
+		Sig.Outputs.Add(FNiagaraVariable(FNiagaraTypeDefinition::GetVec3Def(), TEXT("Tangent")));
 		Sig.bMemberFunction = true;
 		Sig.bRequiresContext = false;
 #if WITH_EDITORONLY_DATA
@@ -170,12 +176,12 @@ void UNiagaraDataInterfaceSkeletalMesh::BindVertexSamplingFunction(const FVMExte
 	//////////////////////////////////////////////////////////////////////////
 	if (BindingInfo.Name == FSkeletalMeshInterfaceHelper::GetSkinnedVertexDataName)
 	{
-		check(BindingInfo.GetNumInputs() == 2 && BindingInfo.GetNumOutputs() == 6);
+		check(BindingInfo.GetNumInputs() == 2 && BindingInfo.GetNumOutputs() == 15);
 		TSkinningModeBinder<TNDIExplicitBinder<FNDITransformHandlerNoop, TVertexAccessorBinder<NDI_FUNC_BINDER(UNiagaraDataInterfaceSkeletalMesh, GetVertexSkinnedData)>>>::Bind(this, BindingInfo, InstanceData, OutFunc);
 	}
 	else if (BindingInfo.Name == FSkeletalMeshInterfaceHelper::GetSkinnedVertexDataWSName)
 	{
-		check(BindingInfo.GetNumInputs() == 2 && BindingInfo.GetNumOutputs() == 6);
+		check(BindingInfo.GetNumInputs() == 2 && BindingInfo.GetNumOutputs() == 15);
 		TSkinningModeBinder<TNDIExplicitBinder<FNDITransformHandler, TVertexAccessorBinder<NDI_FUNC_BINDER(UNiagaraDataInterfaceSkeletalMesh, GetVertexSkinnedData)>>>::Bind(this, BindingInfo, InstanceData, OutFunc);
 	}
 	else if (BindingInfo.Name == FSkeletalMeshInterfaceHelper::GetVertexColorName)
@@ -240,8 +246,8 @@ void UNiagaraDataInterfaceSkeletalMesh::IsValidVertex(FVectorVMContext& Context)
 {
 	SCOPE_CYCLE_COUNTER(STAT_NiagaraSkel_Vertex_Sample);
 
-	VectorVM::FExternalFuncInputHandler<int32> VertexParam(Context);
 	VectorVM::FUserPtrHandler<FNDISkeletalMesh_InstanceData> InstData(Context);
+	VectorVM::FExternalFuncInputHandler<int32> VertexParam(Context);
 	VectorVM::FExternalFuncRegisterHandler<FNiagaraBool> OutValid(Context);
 
 	FSkeletalMeshAccessorHelper MeshAccessor;
@@ -262,8 +268,8 @@ void UNiagaraDataInterfaceSkeletalMesh::RandomVertex(FVectorVMContext& Context)
 {
 	SCOPE_CYCLE_COUNTER(STAT_NiagaraSkel_Vertex_Sample);
 
-	FNDIRandomHelper RandHelper(Context);
 	VectorVM::FUserPtrHandler<FNDISkeletalMesh_InstanceData> InstData(Context);
+	FNDIRandomHelper RandHelper(Context);
 	VectorVM::FExternalFuncRegisterHandler<int32> OutVertex(Context);
 
 	FSkeletalMeshAccessorHelper MeshAccessor;
@@ -341,8 +347,8 @@ void UNiagaraDataInterfaceSkeletalMesh::RandomFilteredVertex(FVectorVMContext& C
 {
 	SCOPE_CYCLE_COUNTER(STAT_NiagaraSkel_Vertex_Sample);
 
-	FNDIRandomHelper RandHelper(Context);
 	VectorVM::FUserPtrHandler<FNDISkeletalMesh_InstanceData> InstData(Context);
+	FNDIRandomHelper RandHelper(Context);
 	checkfSlow(InstData.Get(), TEXT("Skeletal Mesh Interface has invalid instance data. %s"), *GetPathName());
 	checkfSlow(InstData->Mesh, TEXT("Skeletal Mesh Interface has invalid mesh. %s"), *GetPathName());
 
@@ -609,16 +615,28 @@ struct FGetVertexSkinnedDataOutputHandler
 	FGetVertexSkinnedDataOutputHandler(FVectorVMContext& Context)
 		: PosX(Context), PosY(Context), PosZ(Context)
 		, VelX(Context), VelY(Context), VelZ(Context)
+		, TangentZX(Context), TangentZY(Context), TangentZZ(Context)
+		, TangentYX(Context), TangentYY(Context), TangentYZ(Context)
+		, TangentXX(Context), TangentXY(Context), TangentXZ(Context)
 		, bNeedsPosition(PosX.IsValid() || PosY.IsValid() || PosZ.IsValid())
 		, bNeedsVelocity(VelX.IsValid() || VelY.IsValid() || VelZ.IsValid())
+		, bNeedsTangentX(TangentXX.IsValid() || TangentXY.IsValid() || TangentXZ.IsValid())
+		, bNeedsTangentY(TangentYX.IsValid() || TangentYY.IsValid() || TangentYZ.IsValid())
+		, bNeedsTangentZ(TangentZX.IsValid() || TangentZY.IsValid() || TangentZZ.IsValid())
 	{
 	}
 
 	VectorVM::FExternalFuncRegisterHandler<float> PosX; VectorVM::FExternalFuncRegisterHandler<float> PosY; VectorVM::FExternalFuncRegisterHandler<float> PosZ;
 	VectorVM::FExternalFuncRegisterHandler<float> VelX; VectorVM::FExternalFuncRegisterHandler<float> VelY; VectorVM::FExternalFuncRegisterHandler<float> VelZ;
+	VectorVM::FExternalFuncRegisterHandler<float> TangentZX; VectorVM::FExternalFuncRegisterHandler<float> TangentZY; VectorVM::FExternalFuncRegisterHandler<float> TangentZZ;
+	VectorVM::FExternalFuncRegisterHandler<float> TangentYX; VectorVM::FExternalFuncRegisterHandler<float> TangentYY; VectorVM::FExternalFuncRegisterHandler<float> TangentYZ;
+	VectorVM::FExternalFuncRegisterHandler<float> TangentXX; VectorVM::FExternalFuncRegisterHandler<float> TangentXY; VectorVM::FExternalFuncRegisterHandler<float> TangentXZ;
 
 	const bool bNeedsPosition;
 	const bool bNeedsVelocity;
+	const bool bNeedsTangentX;
+	const bool bNeedsTangentY;
+	const bool bNeedsTangentZ;
 
 	FORCEINLINE void SetPosition(FVector Position)
 	{
@@ -633,6 +651,27 @@ struct FGetVertexSkinnedDataOutputHandler
 		*VelY.GetDestAndAdvance() = Velocity.Y;
 		*VelZ.GetDestAndAdvance() = Velocity.Z;
 	}
+
+	FORCEINLINE void SetTangentX(FVector TangentX)
+	{
+		*TangentXX.GetDestAndAdvance() = TangentX.X;
+		*TangentXY.GetDestAndAdvance() = TangentX.Y;
+		*TangentXZ.GetDestAndAdvance() = TangentX.Z;
+	}
+
+	FORCEINLINE void SetTangentY(FVector TangentY)
+	{
+		*TangentYX.GetDestAndAdvance() = TangentY.X;
+		*TangentYY.GetDestAndAdvance() = TangentY.Y;
+		*TangentYZ.GetDestAndAdvance() = TangentY.Z;
+	}
+
+	FORCEINLINE void SetTangentZ(FVector TangentZ)
+	{
+		*TangentZX.GetDestAndAdvance() = TangentZ.X;
+		*TangentZY.GetDestAndAdvance() = TangentZ.Y;
+		*TangentZZ.GetDestAndAdvance() = TangentZ.Z;
+	}
 };
 
 template<typename SkinningHandlerType, typename TransformHandlerType, typename VertexAccessorType>
@@ -658,18 +697,15 @@ void UNiagaraDataInterfaceSkeletalMesh::GetVertexSkinnedData(FVectorVMContext& C
 
 	FSkeletalMeshAccessorHelper Accessor;
 	Accessor.Init<TIntegralConstant<int32, 0>, TIntegralConstant<int32, 0>>(InstData);
-	int32 VertMax = Accessor.LODData->GetNumVertices();
-	float InvDt = 1.0f / InstData->DeltaSeconds;
-
-	FVector Pos;
-	FVector Prev;
-	FVector Velocity;
+	const int32 VertMax = Accessor.LODData->GetNumVertices() - 1;
+	const float InvDt = 1.0f / InstData->DeltaSeconds;
 
+	const bool bNeedsTangentBasis = Output.bNeedsTangentX || Output.bNeedsTangentY || Output.bNeedsTangentZ;
 	for (int32 i = 0; i < Context.NumInstances; ++i)
 	{
-		int32 Vertex = VertParam.GetAndAdvance();
-		Vertex = FMath::Min(VertMax, Vertex);
+		const int32 Vertex = FMath::Min(VertParam.GetAndAdvance(), VertMax);
 		
+		FVector Pos = FVector::ZeroVector;
 		if (Output.bNeedsPosition || Output.bNeedsVelocity)
 		{
 			Pos = SkinningHandler.GetSkinnedVertexPosition(Accessor, Vertex);
@@ -679,11 +715,35 @@ void UNiagaraDataInterfaceSkeletalMesh::GetVertexSkinnedData(FVectorVMContext& C
 
 		if (Output.bNeedsVelocity)
 		{
-			Prev = SkinningHandler.GetSkinnedVertexPreviousPosition(Accessor, Vertex);
+			FVector Prev = SkinningHandler.GetSkinnedVertexPreviousPosition(Accessor, Vertex);
 			TransformHandler.TransformPosition(Prev, PrevTransform);
-			Velocity = (Pos - Prev) * InvDt;
+			const FVector Velocity = (Pos - Prev) * InvDt;
 			Output.SetVelocity(Velocity);
 		}
+
+		if (bNeedsTangentBasis)
+		{
+			FVector TangentX = FVector::ZeroVector;
+			FVector TangentZ = FVector::ZeroVector;
+			SkinningHandler.GetSkinnedTangentBasis(Accessor, Vertex, TangentX, TangentZ);
+
+			if (Output.bNeedsTangentX)
+			{
+				TransformHandler.TransformVector(TangentX, Transform);
+				Output.SetTangentX(TangentX);
+			}
+
+			if (Output.bNeedsTangentY)
+			{
+				Output.SetTangentY((TangentX ^ TangentZ).GetSafeNormal());
+			}
+
+			if (Output.bNeedsTangentZ)
+			{
+				TransformHandler.TransformVector(TangentZ, Transform);
+				Output.SetTangentZ(TangentZ);
+			}
+		}
 	}
 }
 
diff --git a/Engine/Plugins/FX/Niagara/Source/Niagara/Private/NiagaraDataInterfaceSkeletalMesh.cpp b/Engine/Plugins/FX/Niagara/Source/Niagara/Private/NiagaraDataInterfaceSkeletalMesh.cpp
--- a/Engine/Plugins/FX/Niagara/Source/Niagara/Private/NiagaraDataInterfaceSkeletalMesh.cpp
+++ b/Engine/Plugins/FX/Niagara/Source/Niagara/Private/NiagaraDataInterfaceSkeletalMesh.cpp
@@ -30,6 +30,7 @@ struct FNiagaraSkelMeshDIFunctionVersion
 		AddedRandomInfo = 1,
 		CleanUpVertexSampling = 2,
 		CleanupBoneSampling = 3,
+		AddTangentBasisToGetSkinnedVertexData = 4,
 
 		VersionPlusOne,
 		LatestVersion = VersionPlusOne - 1
@@ -170,6 +171,7 @@ void FSkeletalMeshSkinningData::RegisterUser(FSkeletalMeshSkinningDataUsage Usag
 			FSkeletalMeshLODRenderData& SkelMeshLODData = SkelComp->SkeletalMesh->GetResourceForRendering()->LODRenderData[LODIndex];
 			FSkinWeightVertexBuffer* SkinWeightBuffer = SkelComp->GetSkinWeightBuffer(LODIndex);
 			USkeletalMeshComponent::ComputeSkinnedPositions(SkelComp, CurrSkinnedPositions(LODIndex), CurrBoneRefToLocals(), SkelMeshLODData, *SkinWeightBuffer);
+			USkeletalMeshComponent::ComputeSkinnedTangentBasis(SkelComp, CurrSkinnedTangentBasis(LODIndex), CurrBoneRefToLocals(), SkelMeshLODData, *SkinWeightBuffer);
 
 			//Prime the previous positions if they're missing
 			if (PrevSkinnedPositions(LODIndex).Num() != CurrSkinnedPositions(LODIndex).Num())
@@ -322,6 +324,7 @@ bool FSkeletalMeshSkinningData::Tick(float InDeltaSeconds, bool bRequirePreskin)
 				FSkeletalMeshLODRenderData& SkelMeshLODData = SkelComp->SkeletalMesh->GetResourceForRendering()->LODRenderData[LODIndex];
 				FSkinWeightVertexBuffer* SkinWeightBuffer = SkelComp->GetSkinWeightBuffer(LODIndex);
 				USkeletalMeshComponent::ComputeSkinnedPositions(SkelComp, CurrSkinnedPositions(LODIndex), CurrBoneRefToLocals(), SkelMeshLODData, *SkinWeightBuffer);
+				USkeletalMeshComponent::ComputeSkinnedTangentBasis(SkelComp, CurrSkinnedTangentBasis(LODIndex), CurrBoneRefToLocals(), SkelMeshLODData, *SkinWeightBuffer);
 				//check(CurrSkinnedPositions(LODIndex).Num() == SkelMeshLODData.NumVertices);
 				//Prime the previous positions if they're missing
 				if (PrevSkinnedPositions(LODIndex).Num() != CurrSkinnedPositions(LODIndex).Num())
@@ -617,7 +620,7 @@ void FSkeletalMeshGpuSpawnStaticBuffers::InitRHI()
 		{
 			CreateInfo.ResourceArray = &SampleRegionsProb;
 			SampleRegionsProbBuffer = RHICreateVertexBuffer(SampleRegionsProb.Num() * SampleRegionsProb.GetTypeSize(), BUF_Static | BUF_ShaderResource, CreateInfo);
-			SampleRegionsProbSRV = RHICreateShaderResourceView(SampleRegionsTriangleIndicesBuffer, sizeof(float), PF_R32_FLOAT);
+			SampleRegionsProbSRV = RHICreateShaderResourceView(SampleRegionsProbBuffer, sizeof(float), PF_R32_FLOAT);
 
 			CreateInfo.ResourceArray = &SampleRegionsAlias;
 			SampleRegionsAliasBuffer = RHICreateVertexBuffer(SampleRegionsAlias.Num() * SampleRegionsAlias.GetTypeSize(), BUF_Static | BUF_ShaderResource, CreateInfo);
@@ -2478,12 +2481,12 @@ bool UNiagaraDataInterfaceSkeletalMesh::GetFunctionHLSL(const FNiagaraDataInterf
 	// Vertex Sampling
 	else if (FunctionInfo.DefinitionName == FSkeletalMeshInterfaceHelper::GetSkinnedVertexDataName)
 	{
-		static const TCHAR* FormatSample = TEXT("void {InstanceFunctionName} (in int Vertex, out float3 OutPosition, out float3 OutVelocity) { {GetDISkelMeshContextName} DISkelMesh_GetSkinnedVertex(DIContext, Vertex, OutPosition, OutVelocity); }");
+		static const TCHAR* FormatSample = TEXT("void {InstanceFunctionName} (in int Vertex, out float3 OutPosition, out float3 OutVelocity, out float3 OutNormal, out float3 OutBinormal, out float3 OutTangent) { {GetDISkelMeshContextName} DISkelMesh_GetSkinnedVertex(DIContext, Vertex, OutPosition, OutVelocity, OutNormal, OutBinormal, OutTangent); }");
 		OutHLSL += FString::Format(FormatSample, ArgsSample);
 	}
 	else if (FunctionInfo.DefinitionName == FSkeletalMeshInterfaceHelper::GetSkinnedVertexDataWSName)
 	{
-		static const TCHAR* FormatSample = TEXT("void {InstanceFunctionName} (in int Vertex, out float3 OutPosition, out float3 OutVelocity) { {GetDISkelMeshContextName} DISkelMesh_GetSkinnedVertexWS(DIContext, Vertex, OutPosition, OutVelocity); }");
+		static const TCHAR* FormatSample = TEXT("void {InstanceFunctionName} (in int Vertex, out float3 OutPosition, out float3 OutVelocity, out float3 OutNormal, out float3 OutBinormal, out float3 OutTangent) { {GetDISkelMeshContextName} DISkelMesh_GetSkinnedVertexWS(DIContext, Vertex, OutPosition, OutVelocity, OutNormal, OutBinormal, OutTangent); }");
 		OutHLSL += FString::Format(FormatSample, ArgsSample);
 	}
 	else if (FunctionInfo.DefinitionName == FSkeletalMeshInterfaceHelper::GetVertexColorName)
@@ -2733,6 +2736,22 @@ bool UNiagaraDataInterfaceSkeletalMesh::UpgradeFunctionCall(FNiagaraFunctionSign
 		}
 	}
 
+	// Added tangent basis to GetSkinnedVertexData
+	if (FunctionSignature.FunctionVersion < FNiagaraSkelMeshDIFunctionVersion::AddTangentBasisToGetSkinnedVertexData)
+	{
+		if ((FunctionSignature.Name == FSkeletalMeshInterfaceHelper::GetSkinnedVertexDataName) ||
+			(FunctionSignature.Name == FSkeletalMeshInterfaceHelper::GetSkinnedVertexDataWSName))
+		{
+			if ( ensure(FunctionSignature.Outputs.Num() == 2) )
+			{
+				FunctionSignature.Outputs.Add(FNiagaraVariable(FNiagaraTypeDefinition::GetVec3Def(), TEXT("Normal")));
+				FunctionSignature.Outputs.Add(FNiagaraVariable(FNiagaraTypeDefinition::GetVec3Def(), TEXT("Binormal")));
+				FunctionSignature.Outputs.Add(FNiagaraVariable(FNiagaraTypeDefinition::GetVec3Def(), TEXT("Tangent")));
+				bWasChanged = true;
+			}
+		}
+	}
+
 	// Set latest version
 	FunctionSignature.FunctionVersion = FNiagaraSkelMeshDIFunctionVersion::LatestVersion;
 
diff --git a/Engine/Source/Runtime/Engine/Classes/Components/SkeletalMeshComponent.h b/Engine/Source/Runtime/Engine/Classes/Components/SkeletalMeshComponent.h
--- a/Engine/Source/Runtime/Engine/Classes/Components/SkeletalMeshComponent.h
+++ b/Engine/Source/Runtime/Engine/Classes/Components/SkeletalMeshComponent.h
@@ -1688,6 +1688,9 @@ public:
 	static FVector GetSkinnedVertexPosition(USkeletalMeshComponent* Component, int32 VertexIndex, const FSkeletalMeshLODRenderData& Model, FSkinWeightVertexBuffer& SkinWeightBuffer, TArray<FMatrix>& CachedRefToLocals);
 	static void ComputeSkinnedPositions(USkeletalMeshComponent* Component, TArray<FVector> & OutPositions, TArray<FMatrix>& CachedRefToLocals, const FSkeletalMeshLODRenderData& Model, FSkinWeightVertexBuffer& SkinWeightBuffer);
 
+	static void GetSkinnedTangentBasis(USkeletalMeshComponent* Component, int32 VertexIndex, const FSkeletalMeshLODRenderData& Model, FSkinWeightVertexBuffer& SkinWeightBuffer, TArray<FMatrix>& CachedRefToLocals, FVector& OutTangentX, FVector& OutTangentZ);
+	static void ComputeSkinnedTangentBasis(USkeletalMeshComponent* Component, TArray<FVector>& OutTangenXZ, TArray<FMatrix>& CachedRefToLocals, const FSkeletalMeshLODRenderData& Model, FSkinWeightVertexBuffer& SkinWeightBuffer);
+
 	void SetSkeletalMeshWithoutResettingAnimation(class USkeletalMesh* NewMesh);
 
 	virtual bool IsPlayingRootMotion() const override;
diff --git a/Engine/Source/Runtime/Engine/Classes/Components/SkinnedMeshComponent.h b/Engine/Source/Runtime/Engine/Classes/Components/SkinnedMeshComponent.h
--- a/Engine/Source/Runtime/Engine/Classes/Components/SkinnedMeshComponent.h
+++ b/Engine/Source/Runtime/Engine/Classes/Components/SkinnedMeshComponent.h
@@ -1565,6 +1565,17 @@ public:
 	}
 };
 
+/** Simple, CPU evaluation of a vertex's skinned tangent basis */
+void GetTypedSkinnedTangentBasis(
+	const USkinnedMeshComponent* SkinnedComp,
+	const FSkelMeshRenderSection& Section,
+	const FStaticMeshVertexBuffers& StaticVertexBuffers,
+	const FSkinWeightVertexBuffer& SkinWeightVertexBuffer,
+	const int32 VertIndex,
+	const TArray<FMatrix> & RefToLocals,
+	FVector& OutTangentX,
+	FVector& OutTangentZ
+);
 
 /** Simple, CPU evaluation of a vertex's skinned position helper function */
 template <bool bCachedMatrices>
diff --git a/Engine/Source/Runtime/Engine/Classes/Engine/SkeletalMesh.h b/Engine/Source/Runtime/Engine/Classes/Engine/SkeletalMesh.h
--- a/Engine/Source/Runtime/Engine/Classes/Engine/SkeletalMesh.h
+++ b/Engine/Source/Runtime/Engine/Classes/Engine/SkeletalMesh.h
@@ -1490,4 +1490,5 @@ protected:
  */
 ENGINE_API void RefreshSkelMeshOnPhysicsAssetChange(const USkeletalMesh* InSkeletalMesh);
 
-ENGINE_API FVector GetSkeletalMeshRefVertLocation(const USkeletalMesh* Mesh, const FSkeletalMeshLODRenderData& LODData, const FSkinWeightVertexBuffer& SkinWeightVertexBuffer, const int32 VertIndex);
\ No newline at end of file
+ENGINE_API FVector GetSkeletalMeshRefVertLocation(const USkeletalMesh* Mesh, const FSkeletalMeshLODRenderData& LODData, const FSkinWeightVertexBuffer& SkinWeightVertexBuffer, const int32 VertIndex);
+ENGINE_API void GetSkeletalMeshRefTangentBasis(const USkeletalMesh* Mesh, const FSkeletalMeshLODRenderData& LODData, const FSkinWeightVertexBuffer& SkinWeightVertexBuffer, const int32 VertIndex, FVector& OutTangentX, FVector& OutTangentZ);
diff --git a/Engine/Source/Runtime/Engine/Private/Components/SkinnedMeshComponent.cpp b/Engine/Source/Runtime/Engine/Private/Components/SkinnedMeshComponent.cpp
--- a/Engine/Source/Runtime/Engine/Private/Components/SkinnedMeshComponent.cpp
+++ b/Engine/Source/Runtime/Engine/Private/Components/SkinnedMeshComponent.cpp
@@ -3920,6 +3920,47 @@ float FAnimUpdateRateParameters::GetRootMotionInterp() const
 	}
 	return 1.f;
 }
+
+/** Simple, CPU evaluation of a vertex's skinned position helper function */
+void GetTypedSkinnedTangentBasis(
+	const USkinnedMeshComponent* SkinnedComp,
+	const FSkelMeshRenderSection& Section,
+	const FStaticMeshVertexBuffers& StaticVertexBuffers,
+	const FSkinWeightVertexBuffer& SkinWeightVertexBuffer,
+	const int32 VertIndex,
+	const TArray<FMatrix> & RefToLocals,
+	FVector& OutTangentX,
+	FVector& OutTangentZ
+)
+{
+	OutTangentX = FVector::ZeroVector;
+	OutTangentZ = FVector::ZeroVector;
+
+	const USkinnedMeshComponent* const MasterPoseComponentInst = SkinnedComp->MasterPoseComponent.Get();
+	const USkinnedMeshComponent* BaseComponent = MasterPoseComponentInst ? MasterPoseComponentInst : SkinnedComp;
+
+	// Do soft skinning for this vertex.
+	const int32 BufferVertIndex = Section.GetVertexBufferIndex() + VertIndex;
+	const int32 MaxBoneInfluences = SkinWeightVertexBuffer.GetMaxBoneInfluences();
+
+	const FVector VertexTangentX = StaticVertexBuffers.StaticMeshVertexBuffer.VertexTangentX(BufferVertIndex);
+	const FVector VertexTangentZ = StaticVertexBuffers.StaticMeshVertexBuffer.VertexTangentZ(BufferVertIndex);
+
+#if !PLATFORM_LITTLE_ENDIAN
+	// uint8[] elements in LOD.VertexBufferGPUSkin have been swapped for VET_UBYTE4 vertex stream use
+	for (int32 InfluenceIndex = MAX_INFLUENCES - 1; InfluenceIndex >= MAX_INFLUENCES - MaxBoneInfluences; InfluenceIndex--)
+#else
+	for (int32 InfluenceIndex = 0; InfluenceIndex < MaxBoneInfluences; InfluenceIndex++)
+#endif
+	{
+		const int32 MeshBoneIndex = Section.BoneMap[SkinWeightVertexBuffer.GetBoneIndex(BufferVertIndex, InfluenceIndex)];
+		const float	Weight = (float)SkinWeightVertexBuffer.GetBoneWeight(BufferVertIndex, InfluenceIndex) / 255.0f;
+		const FMatrix& RefToLocal = RefToLocals[MeshBoneIndex];
+		OutTangentX += RefToLocal.TransformVector(VertexTangentX) * Weight;
+		OutTangentZ += RefToLocal.TransformVector(VertexTangentZ) * Weight;
+	}
+}
+
 /** Simple, CPU evaluation of a vertex's skinned position helper function */
 template <bool bCachedMatrices>
 FVector GetTypedSkinnedVertexPosition(
diff --git a/Engine/Source/Runtime/Engine/Private/SkeletalMesh.cpp b/Engine/Source/Runtime/Engine/Private/SkeletalMesh.cpp
--- a/Engine/Source/Runtime/Engine/Private/SkeletalMesh.cpp
+++ b/Engine/Source/Runtime/Engine/Private/SkeletalMesh.cpp
@@ -5620,4 +5620,42 @@ FVector GetSkeletalMeshRefVertLocation(const USkeletalMesh* Mesh, const FSkeleta
 	return GetRefVertexLocationTyped(Mesh, Section, LODData.StaticVertexBuffers.PositionVertexBuffer, SkinWeightVertexBuffer, VertIndexInChunk);
 }
 
+//GetRefTangentBasisTyped
+void GetRefTangentBasisTyped(const USkeletalMesh* Mesh, const FSkelMeshRenderSection& Section, const FStaticMeshVertexBuffer& StaticVertexBuffer, const FSkinWeightVertexBuffer& SkinWeightVertexBuffer, const int32 VertIndex, FVector& OutTangentX, FVector& OutTangentZ)
+{
+	OutTangentX = FVector::ZeroVector;
+	OutTangentZ = FVector::ZeroVector;
+
+	// Do soft skinning for this vertex.
+	const int32 BufferVertIndex = Section.GetVertexBufferIndex() + VertIndex;
+	const int32 MaxBoneInfluences = SkinWeightVertexBuffer.GetMaxBoneInfluences();
+
+	const FVector VertexTangentX = StaticVertexBuffer.VertexTangentX(BufferVertIndex);
+	const FVector VertexTangentZ = StaticVertexBuffer.VertexTangentZ(BufferVertIndex);
+
+#if !PLATFORM_LITTLE_ENDIAN
+	// uint8[] elements in LOD.VertexBufferGPUSkin have been swapped for VET_UBYTE4 vertex stream use
+	for (int32 InfluenceIndex = MAX_INFLUENCES - 1; InfluenceIndex >= MAX_INFLUENCES - MaxBoneInfluences; InfluenceIndex--)
+#else
+	for (int32 InfluenceIndex = 0; InfluenceIndex < MaxBoneInfluences; InfluenceIndex++)
+#endif
+	{
+		const int32 MeshBoneIndex = Section.BoneMap[SkinWeightVertexBuffer.GetBoneIndex(BufferVertIndex, InfluenceIndex)];
+		const float	Weight = (float)SkinWeightVertexBuffer.GetBoneWeight(BufferVertIndex, InfluenceIndex) / 255.0f;
+		const FMatrix BoneTransformMatrix = FMatrix::Identity;//Mesh->GetComposedRefPoseMatrix(MeshBoneIndex);
+		//const FMatrix RefToLocal = Mesh->RefBasesInvMatrix[MeshBoneIndex] * BoneTransformMatrix;
+		OutTangentX += BoneTransformMatrix.TransformVector(VertexTangentX) * Weight;
+		OutTangentZ += BoneTransformMatrix.TransformVector(VertexTangentZ) * Weight;
+	}
+}
+
+void GetSkeletalMeshRefTangentBasis(const USkeletalMesh* Mesh, const FSkeletalMeshLODRenderData& LODData, const FSkinWeightVertexBuffer& SkinWeightVertexBuffer, const int32 VertIndex, FVector& OutTangentX, FVector& OutTangentZ)
+{
+	int32 SectionIndex;
+	int32 VertIndexInChunk;
+	LODData.GetSectionFromVertexIndex(VertIndex, SectionIndex, VertIndexInChunk);
+	const FSkelMeshRenderSection& Section = LODData.RenderSections[SectionIndex];
+	GetRefTangentBasisTyped(Mesh, Section, LODData.StaticVertexBuffers.StaticMeshVertexBuffer, SkinWeightVertexBuffer, VertIndexInChunk, OutTangentX, OutTangentZ);
+}
+
 #undef LOCTEXT_NAMESPACE
diff --git a/Engine/Source/Runtime/Engine/Private/SkeletalMeshComponentPhysics.cpp b/Engine/Source/Runtime/Engine/Private/SkeletalMeshComponentPhysics.cpp
--- a/Engine/Source/Runtime/Engine/Private/SkeletalMeshComponentPhysics.cpp
+++ b/Engine/Source/Runtime/Engine/Private/SkeletalMeshComponentPhysics.cpp
@@ -2053,6 +2053,43 @@ void USkeletalMeshComponent::ComputeSkinnedPositions(USkeletalMeshComponent* Com
 	}
 }
 
+void USkeletalMeshComponent::GetSkinnedTangentBasis(USkeletalMeshComponent* Component, int32 VertexIndex, const FSkeletalMeshLODRenderData& LODData, FSkinWeightVertexBuffer& SkinWeightBuffer, TArray<FMatrix>& CachedRefToLocals, FVector& OutTangentX, FVector& OutTangentZ)
+{
+	int32 SectionIndex;
+	int32 VertIndexInChunk;
+	LODData.GetSectionFromVertexIndex(VertexIndex, SectionIndex, VertIndexInChunk);
+
+	check(SectionIndex < LODData.RenderSections.Num());
+	const FSkelMeshRenderSection& Section = LODData.RenderSections[SectionIndex];
+
+	return GetTypedSkinnedTangentBasis(Component, Section, LODData.StaticVertexBuffers, SkinWeightBuffer, VertIndexInChunk, CachedRefToLocals, OutTangentX, OutTangentZ);
+}
+
+void USkeletalMeshComponent::ComputeSkinnedTangentBasis(USkeletalMeshComponent* Component, TArray<FVector>& OutTangenXZ, TArray<FMatrix>& CachedRefToLocals, const FSkeletalMeshLODRenderData& LODData, FSkinWeightVertexBuffer& SkinWeightBuffer)
+{
+	// Fail if no mesh
+	if (!Component->SkeletalMesh)
+	{
+		return;
+	}
+
+	OutTangenXZ.Empty();
+	OutTangenXZ.AddUninitialized(LODData.GetNumVertices() * 2);
+
+	for (int32 SectionIdx = 0; SectionIdx < LODData.RenderSections.Num(); ++SectionIdx)
+	{
+		const FSkelMeshRenderSection& Section = LODData.RenderSections[SectionIdx];
+
+		const uint32 SoftOffset = Section.GetVertexBufferIndex();
+		const uint32 NumSoftVerts = Section.GetNumVertices();
+		for (uint32 SoftIdx = 0; SoftIdx < NumSoftVerts; ++SoftIdx)
+		{
+			const uint32 TangentOffset = SoftOffset + SoftIdx * 2;
+			GetTypedSkinnedTangentBasis(Component, Section, LODData.StaticVertexBuffers, SkinWeightBuffer, SoftIdx, CachedRefToLocals, OutTangenXZ[TangentOffset + 0], OutTangenXZ[TangentOffset + 1]);
+		}
+	}
+}
+
 void USkeletalMeshComponent::SetEnableBodyGravity(bool bEnableGravity, FName BoneName)
 {
 	if (FBodyInstance* BI = GetBodyInstance(BoneName))
