REL - Speculative fix: redo work that had been done on the MaterialShaderCompiler after the NiagaraVersion was forked of. See CL: 4731964

diff --git a/Engine/Plugins/FX/Niagara/Source/NiagaraShader/Private/NiagaraShader.cpp b/Engine/Plugins/FX/Niagara/Source/NiagaraShader/Private/NiagaraShader.cpp
--- a/Engine/Plugins/FX/Niagara/Source/NiagaraShader/Private/NiagaraShader.cpp
+++ b/Engine/Plugins/FX/Niagara/Source/NiagaraShader/Private/NiagaraShader.cpp
@@ -60,6 +60,7 @@ namespace NiagaraShaderCookStats
 //
 // Globals
 //
+NIAGARASHADER_API FCriticalSection GIdToNiagaraShaderMapCS;
 TMap<FNiagaraShaderMapId, FNiagaraShaderMap*> FNiagaraShaderMap::GIdToNiagaraShaderMap[SP_NumPlatforms];
 TArray<FNiagaraShaderMap*> FNiagaraShaderMap::AllNiagaraShaderMaps;
 
@@ -517,7 +518,9 @@ FShader* FNiagaraShaderType::FinishCompileShader(
 */
 FNiagaraShaderMap* FNiagaraShaderMap::FindId(const FNiagaraShaderMapId& ShaderMapId, EShaderPlatform InPlatform)
 {
-	return GIdToNiagaraShaderMap[InPlatform].FindRef(ShaderMapId);
+	FNiagaraShaderMap* Result = GIdToNiagaraShaderMap[InPlatform].FindRef(ShaderMapId);
+	check(Result == nullptr || !Result->bDeletedThroughDeferredCleanup);
+	return Result;
 }
 
 /** Flushes the given shader types from any loaded FNiagaraShaderMap's. */
@@ -537,6 +540,8 @@ void FNiagaraShaderMap::FlushShaderTypes(TArray<const FShaderType*>& ShaderTypes
 #if 0
 void FNiagaraShaderMap::FixupShaderTypes(EShaderPlatform Platform, const TMap<FShaderType*, FString>& ShaderTypeNames)
 {
+	FScopeLock ScopeLock(&GIdToNiagaraShaderMapCS);
+
 	TArray<FNiagaraShaderMapId> Keys;
 	FNiagaraShaderMap::GIdToNiagaraShaderMap[Platform].GenerateKeyArray(Keys);
 
@@ -957,31 +962,42 @@ void FNiagaraShaderMap::Register(EShaderPlatform InShaderPlatform)
 		INC_DWORD_STAT(STAT_Shaders_NumShaderMaps);
 	}
 
-	GIdToNiagaraShaderMap[GetShaderPlatform()].Add(GetContent()->ShaderMapId,this);
-	bRegistered = true;
+	{
+		FScopeLock ScopeLock(&GIdToNiagaraShaderMapCS);
+		GIdToNiagaraShaderMap[GetShaderPlatform()].Add(GetContent()->ShaderMapId, this);
+		bRegistered = true;
+	}
 }
 
 void FNiagaraShaderMap::AddRef()
 {
+	FScopeLock ScopeLock(&GIdToNiagaraShaderMapCS);
 	check(!bDeletedThroughDeferredCleanup);
 	++NumRefs;
 }
 
 void FNiagaraShaderMap::Release()
 {
-	check(NumRefs > 0);
-	if(--NumRefs == 0)
 	{
-		if (bRegistered)
+		FScopeLock ScopeLock(&GIdToNiagaraShaderMapCS);
+
+		check(NumRefs > 0);
+		if (--NumRefs == 0)
 		{
-			DEC_DWORD_STAT(STAT_Shaders_NumShaderMaps);
+			if (bRegistered)
+			{
+				DEC_DWORD_STAT(STAT_Shaders_NumShaderMaps);
 
-			GIdToNiagaraShaderMap[GetShaderPlatform()].Remove(GetContent()->ShaderMapId);
-			bRegistered = false;
-		}
+				GIdToNiagaraShaderMap[GetShaderPlatform()].Remove(GetContent()->ShaderMapId);
+				bRegistered = false;
+			}
 
-		check(!bDeletedThroughDeferredCleanup);
-		bDeletedThroughDeferredCleanup = true;
+			check(!bDeletedThroughDeferredCleanup);
+			bDeletedThroughDeferredCleanup = true;
+		}
+	}
+	if (bDeletedThroughDeferredCleanup)
+	{
 		BeginCleanup(this);
 	}
 }
diff --git a/Engine/Plugins/FX/Niagara/Source/NiagaraShader/Private/NiagaraShared.cpp b/Engine/Plugins/FX/Niagara/Source/NiagaraShader/Private/NiagaraShared.cpp
--- a/Engine/Plugins/FX/Niagara/Source/NiagaraShader/Private/NiagaraShared.cpp
+++ b/Engine/Plugins/FX/Niagara/Source/NiagaraShader/Private/NiagaraShared.cpp
@@ -348,8 +348,13 @@ bool FNiagaraShaderScript::CacheShaders(const FNiagaraShaderMapId& ShaderMapId,
 	check(IsInGameThread());
 
 	{
-		// Find the script's cached shader map.
-		GameThreadShaderMap = FNiagaraShaderMap::FindId(ShaderMapId, Platform);
+		GameThreadShaderMap = nullptr;
+		{
+			NIAGARASHADER_API extern FCriticalSection GIdToNiagaraShaderMapCS;
+			FScopeLock ScopeLock(&GIdToNiagaraShaderMapCS);
+			// Find the script's cached shader map.
+			GameThreadShaderMap = FNiagaraShaderMap::FindId(ShaderMapId, Platform);
+		}
 
 		// Attempt to load from the derived data cache if we are uncooked
 		if (!bForceRecompile && !GameThreadShaderMap && !FPlatformProperties::RequiresCookedData())
