Niagara - Change orphaned pins in scripts from errors to warnings.

diff --git a/Engine/Plugins/FX/Niagara/Source/NiagaraEditor/Private/NiagaraHlslTranslator.cpp b/Engine/Plugins/FX/Niagara/Source/NiagaraEditor/Private/NiagaraHlslTranslator.cpp
--- a/Engine/Plugins/FX/Niagara/Source/NiagaraEditor/Private/NiagaraHlslTranslator.cpp
+++ b/Engine/Plugins/FX/Niagara/Source/NiagaraEditor/Private/NiagaraHlslTranslator.cpp
@@ -271,7 +271,7 @@ bool FHlslNiagaraTranslator::ValidateTypePins(UNiagaraNode* NodeToValidate)
 
 		if (Pin->bOrphanedPin)
 		{
-			Error(LOCTEXT("OrphanedPinError", "Node pin is no longer valid.  This pin must be disconnected or reset to default so it can be removed."), NodeToValidate, Pin);
+			Warning(LOCTEXT("OrphanedPinError", "Node pin is no longer valid.  This pin must be disconnected or reset to default so it can be removed."), NodeToValidate, Pin);
 		}
 	}
 	return bPinsAreValid;
@@ -6964,9 +6964,8 @@ int32 FHlslNiagaraTranslator::CompileOutputPin(const UEdGraphPin* InPin)
 	return Ret;
 }
 
-void FHlslNiagaraTranslator::Error(FText ErrorText, const UNiagaraNode* InNode, const UEdGraphPin* Pin)
+void FHlslNiagaraTranslator::Message(FNiagaraCompileEventSeverity Severity, FText MessageText, const UNiagaraNode* InNode, const UEdGraphPin* Pin)
 {
-
 	const UNiagaraNode* CurContextNode = ActiveHistoryForFunctionCalls.GetCallingContext();
 	const UNiagaraNode* TargetNode = InNode ? InNode : CurContextNode;
 
@@ -6997,34 +6996,27 @@ void FHlslNiagaraTranslator::Error(FText ErrorText, const UNiagaraNode* InNode,
 		NodePinSuffix = TEXT(" - ");
 	}
 
-	FString ErrorString = FString::Printf(TEXT("%s%s%s%s"), *ErrorText.ToString(), *NodePinPrefix, *NodePinStr, *NodePinSuffix);
-	TranslateResults.CompileEvents.Add(FNiagaraCompileEvent(FNiagaraCompileEventSeverity::Error, ErrorString, TargetNode ? TargetNode->NodeGuid : FGuid(), Pin ? Pin->PersistentGuid : FGuid(), GetCallstackGuids()));
-	TranslateResults.NumErrors++;
-}
-
-void FHlslNiagaraTranslator::Warning(FText WarningText, const UNiagaraNode* InNode, const UEdGraphPin* Pin)
-{
+	FString MessageString = FString::Printf(TEXT("%s%s%s%s"), *MessageText.ToString(), *NodePinPrefix, *NodePinStr, *NodePinSuffix);
+	TranslateResults.CompileEvents.Add(FNiagaraCompileEvent(Severity, MessageString, TargetNode ? TargetNode->NodeGuid : FGuid(), Pin ? Pin->PersistentGuid : FGuid(), GetCallstackGuids()));
 
-	const UNiagaraNode* CurContextNode = ActiveHistoryForFunctionCalls.GetCallingContext();
-	const UNiagaraNode* TargetNode = InNode ? InNode : CurContextNode;
-
-	FString NodePinStr = TEXT("");
-	FString NodePinPrefix = TEXT(" - ");
-	FString NodePinSuffix = TEXT("");
-	if (TargetNode && TargetNode->GetName().Len() > 0)
+	if (Severity == FNiagaraCompileEventSeverity::Error)
 	{
-		NodePinStr += TEXT("Node: ") + TargetNode->GetName();
-		NodePinSuffix = TEXT(" - ");
+		TranslateResults.NumErrors++;
 	}
-	if (Pin && Pin->PinFriendlyName.ToString().Len() > 0)
+	else if (Severity == FNiagaraCompileEventSeverity::Warning)
 	{
-		NodePinStr += TEXT(" Pin: ") + Pin->PinFriendlyName.ToString();
-		NodePinSuffix = TEXT(" - ");
+		TranslateResults.NumWarnings++;
 	}
+}
+
+void FHlslNiagaraTranslator::Error(FText ErrorText, const UNiagaraNode* InNode, const UEdGraphPin* Pin)
+{
+	Message(FNiagaraCompileEventSeverity::Error, ErrorText, InNode, Pin);
+}
 
-	FString WarnString = FString::Printf(TEXT("%s%s%s%s"), *WarningText.ToString(), *NodePinPrefix, *NodePinStr, *NodePinSuffix);
-	TranslateResults.CompileEvents.Add(FNiagaraCompileEvent(FNiagaraCompileEventSeverity::Warning, WarnString, TargetNode ? TargetNode->NodeGuid : FGuid(), Pin ? Pin->PersistentGuid : FGuid(), GetCallstackGuids()));
-	TranslateResults.NumWarnings++;
+void FHlslNiagaraTranslator::Warning(FText WarningText, const UNiagaraNode* InNode, const UEdGraphPin* Pin)
+{
+	Message(FNiagaraCompileEventSeverity::Warning, WarningText, InNode, Pin);
 }
 
 bool FHlslNiagaraTranslator::GetFunctionParameter(const FNiagaraVariable& Parameter, int32& OutParam)const
diff --git a/Engine/Plugins/FX/Niagara/Source/NiagaraEditor/Private/NiagaraHlslTranslator.h b/Engine/Plugins/FX/Niagara/Source/NiagaraEditor/Private/NiagaraHlslTranslator.h
--- a/Engine/Plugins/FX/Niagara/Source/NiagaraEditor/Private/NiagaraHlslTranslator.h
+++ b/Engine/Plugins/FX/Niagara/Source/NiagaraEditor/Private/NiagaraHlslTranslator.h
@@ -516,6 +516,7 @@ public:
 	virtual void Convert(class UNiagaraNodeConvert* Convert, TArray <int32>& Inputs, TArray<int32>& Outputs);
 	virtual void If(class UNiagaraNodeIf* IfNode, TArray<FNiagaraVariable>& Vars, int32 Condition, TArray<int32>& PathA, TArray<int32>& PathB, TArray<int32>& Outputs);
 
+	void Message(FNiagaraCompileEventSeverity Severity, FText MessageText, const UNiagaraNode* Node, const UEdGraphPin* Pin);
 	virtual void Error(FText ErrorText, const UNiagaraNode* Node, const UEdGraphPin* Pin);
 	virtual void Warning(FText WarningText, const UNiagaraNode* Node, const UEdGraphPin* Pin);
 
diff --git a/Engine/Plugins/FX/Niagara/Source/NiagaraEditor/Private/ViewModels/Stack/NiagaraStackFunctionInputCollection.cpp b/Engine/Plugins/FX/Niagara/Source/NiagaraEditor/Private/ViewModels/Stack/NiagaraStackFunctionInputCollection.cpp
--- a/Engine/Plugins/FX/Niagara/Source/NiagaraEditor/Private/ViewModels/Stack/NiagaraStackFunctionInputCollection.cpp
+++ b/Engine/Plugins/FX/Niagara/Source/NiagaraEditor/Private/ViewModels/Stack/NiagaraStackFunctionInputCollection.cpp
@@ -430,7 +430,7 @@ void UNiagaraStackFunctionInputCollection::RefreshIssues(TArray<FName> Duplicate
 		if (InputFunctionCallNodePin->Direction == EEdGraphPinDirection::EGPD_Input && InputFunctionCallNodePin->bOrphanedPin)
 		{
 			FStackIssue InvalidInputError(
-				EStackIssueSeverity::Error,
+				EStackIssueSeverity::Warning,
 				FText::Format(LOCTEXT("InvalidInputSummaryFormat", "Invalid Input: {0}"), FText::FromString(InputFunctionCallNodePin->PinName.ToString())),
 				FText::Format(LOCTEXT("InvalidInputFormat", "The input {0} was previously set but is no longer exposed by the function {1}.\nPress the fix button to remove this unused input data,\nor check the function definition to see why this input is no longer exposed."),
 					FText::FromString(InputFunctionCallNodePin->PinName.ToString()), FText::FromString(InputFunctionCallNode->GetFunctionName())),
