- Add back in RandomTriCoord until I deprecate it correctly

#rnx

diff --git a/Engine/Plugins/FX/Niagara/Source/Niagara/Private/NDISkeletalMesh_TriangleSampling.cpp b/Engine/Plugins/FX/Niagara/Source/Niagara/Private/NDISkeletalMesh_TriangleSampling.cpp
--- a/Engine/Plugins/FX/Niagara/Source/Niagara/Private/NDISkeletalMesh_TriangleSampling.cpp
+++ b/Engine/Plugins/FX/Niagara/Source/Niagara/Private/NDISkeletalMesh_TriangleSampling.cpp
@@ -41,17 +41,17 @@ const FName FSkeletalMeshInterfaceHelper::GetFilteredTriangleAtName("GetFiltered
 
 void UNiagaraDataInterfaceSkeletalMesh::GetTriangleSamplingFunctions(TArray<FNiagaraFunctionSignature>& OutFunctions)
 {
-// Hidden will remove later
-//	{
-//		FNiagaraFunctionSignature Sig;
-//		Sig.Name = FSkeletalMeshInterfaceHelper::RandomTriCoordName;
-//		Sig.Inputs.Add(FNiagaraVariable(FNiagaraTypeDefinition(GetClass()), TEXT("SkeletalMesh")));
-//		Sig.Inputs.Add(FNiagaraVariable(FNiagaraTypeDefinition(FNiagaraRandInfo::StaticStruct()), TEXT("RandomInfo")));
-//		Sig.Outputs.Add(FNiagaraVariable(FNiagaraTypeDefinition(FMeshTriCoordinate::StaticStruct()), TEXT("Coord")));
-//		Sig.bMemberFunction = true;
-//		Sig.bRequiresContext = false;
-//		OutFunctions.Add(Sig);
-//	}
+	//-TODO: Remove / deprecate this function!
+	{
+		FNiagaraFunctionSignature Sig;
+		Sig.Name = FSkeletalMeshInterfaceHelper::RandomTriCoordName;
+		Sig.Inputs.Add(FNiagaraVariable(FNiagaraTypeDefinition(GetClass()), TEXT("SkeletalMesh")));
+		Sig.Inputs.Add(FNiagaraVariable(FNiagaraTypeDefinition(FNiagaraRandInfo::StaticStruct()), TEXT("RandomInfo")));
+		Sig.Outputs.Add(FNiagaraVariable(FNiagaraTypeDefinition(FMeshTriCoordinate::StaticStruct()), TEXT("Coord")));
+		Sig.bMemberFunction = true;
+		Sig.bRequiresContext = false;
+		OutFunctions.Add(Sig);
+	}
 
 	{
 		FNiagaraFunctionSignature Sig;
