- Add support for excluding a single bone (generally the root bone) as this will be common case and avoid filtering / unfiltering in lots of cases which will be less efficient

#rnx

diff --git a/Engine/Plugins/FX/Niagara/Source/Niagara/Classes/NiagaraDataInterfaceSkeletalMesh.h b/Engine/Plugins/FX/Niagara/Source/Niagara/Classes/NiagaraDataInterfaceSkeletalMesh.h
--- a/Engine/Plugins/FX/Niagara/Source/Niagara/Classes/NiagaraDataInterfaceSkeletalMesh.h
+++ b/Engine/Plugins/FX/Niagara/Source/Niagara/Classes/NiagaraDataInterfaceSkeletalMesh.h
@@ -470,6 +470,9 @@ struct FNDISkeletalMesh_InstanceData
 	/** Time separating Transform and PrevTransform. */
 	float DeltaSeconds;
 
+	/* Excluded bone for some specific functions, generally the root bone which you don't want to include when picking a random bone. */
+	int32 ExcludedBoneIndex = INDEX_NONE;
+
 	/* Number of filtered bones in the array. */
 	int32 NumFilteredBones = 0;
 	/* Number of unfiltered bones in the array. */
@@ -599,6 +602,16 @@ public:
 	UPROPERTY(EditAnywhere, Category = "Skeleton")
 	TArray<FName> FilteredSockets;
 
+	/**
+	Optionally remove a single bone from Random / Random Unfiltered access.
+	You can still include this bone in filtered list and access using the direct index functionality.
+	*/
+	UPROPERTY(EditAnywhere, Category = "Skeleton", meta=(EditCondition="bExcludeBone"))
+	FName ExcludeBoneName;
+
+	UPROPERTY(EditAnywhere, Category = "Skeleton", meta = (InlineEditConditionToggle))
+	uint8 bExcludeBone : 1;
+
 #if WITH_EDITORONLY_DATA
 	/** Do we require CPU access to the data, this is set during GetVMExternalFunction. */
 	UPROPERTY(transient)
diff --git a/Engine/Plugins/FX/Niagara/Source/Niagara/Private/NDISkeletalMesh_BoneSampling.cpp b/Engine/Plugins/FX/Niagara/Source/Niagara/Private/NDISkeletalMesh_BoneSampling.cpp
--- a/Engine/Plugins/FX/Niagara/Source/Niagara/Private/NDISkeletalMesh_BoneSampling.cpp
+++ b/Engine/Plugins/FX/Niagara/Source/Niagara/Private/NDISkeletalMesh_BoneSampling.cpp
@@ -555,16 +555,25 @@ void UNiagaraDataInterfaceSkeletalMesh::RandomUnfilteredBone(FVectorVMContext& C
 
 	VectorVM::FExternalFuncRegisterHandler<int32> OutBone(Context);
 
-	const int32 Max = InstData->NumUnfilteredBones - 1;
-	if (Max >= 0)
+	const int32 UnfilteredMax = InstData->NumUnfilteredBones - 1;
+	if (UnfilteredMax >= 0)
 	{
 		if (InstData->NumFilteredBones == 0)
 		{
-			for (int32 i = 0; i < Context.NumInstances; ++i)
+			const int32 ExcludedBoneIndex = InstData->ExcludedBoneIndex;
+			const int32 NumBones = InstData->NumUnfilteredBones - (ExcludedBoneIndex >= 0 ? 2 : 1);
+			if (NumBones >= 0)
 			{
-				RandHelper.GetAndAdvance();
-				const int32 BoneIndex = RandHelper.RandRange(i, 0, Max);
-				*OutBone.GetDestAndAdvance() = BoneIndex;
+				for (int32 i = 0; i < Context.NumInstances; ++i)
+				{
+					RandHelper.GetAndAdvance();
+					const int32 BoneIndex = RandHelper.RandRange(i, 0, NumBones);
+					*OutBone.GetDestAndAdvance() = BoneIndex != ExcludedBoneIndex ? BoneIndex : BoneIndex + 1;
+				}
+			}
+			else
+			{
+				FMemory::Memset(OutBone.GetDest(), 0xFF, Context.NumInstances * sizeof(int32));
 			}
 		}
 		else
@@ -572,7 +581,7 @@ void UNiagaraDataInterfaceSkeletalMesh::RandomUnfilteredBone(FVectorVMContext& C
 			for (int32 i = 0; i < Context.NumInstances; ++i)
 			{
 				RandHelper.GetAndAdvance();
-				const int32 BoneIndex = RandHelper.RandRange(i, 0, Max);
+				const int32 BoneIndex = RandHelper.RandRange(i, 0, UnfilteredMax);
 				*OutBone.GetDestAndAdvance() = InstData->FilteredAndUnfilteredBones[BoneIndex + InstData->NumFilteredBones];
 			}
 		}
@@ -616,13 +625,15 @@ void UNiagaraDataInterfaceSkeletalMesh::RandomBone(FVectorVMContext& Context)
 	FSkeletalMeshAccessorHelper MeshAccessor;
 	MeshAccessor.Init<TIntegralConstant<int32, 0>, TIntegralConstant<int32, 0>>(InstData);
 	const FReferenceSkeleton& RefSkeleton = MeshAccessor.Mesh->RefSkeleton;
-	const int32 NumBones = RefSkeleton.GetNum() - 1;
+	const int32 ExcludedBoneIndex = InstData->ExcludedBoneIndex;
+	const int32 NumBones = RefSkeleton.GetNum() - (ExcludedBoneIndex >= 0 ? 2 : 1);
 	if (NumBones >= 0)
 	{
 		for (int32 i = 0; i < Context.NumInstances; ++i)
 		{
 			RandHelper.GetAndAdvance();
-			*OutBone.GetDestAndAdvance() = RandHelper.RandRange(i, 0, NumBones);
+			const int32 BoneIndex = RandHelper.RandRange(i, 0, NumBones);
+			*OutBone.GetDestAndAdvance() = BoneIndex != ExcludedBoneIndex ? BoneIndex : BoneIndex + 1;
 		}
 	}
 	else
diff --git a/Engine/Plugins/FX/Niagara/Source/Niagara/Private/NiagaraDataInterfaceSkeletalMesh.cpp b/Engine/Plugins/FX/Niagara/Source/Niagara/Private/NiagaraDataInterfaceSkeletalMesh.cpp
--- a/Engine/Plugins/FX/Niagara/Source/Niagara/Private/NiagaraDataInterfaceSkeletalMesh.cpp
+++ b/Engine/Plugins/FX/Niagara/Source/Niagara/Private/NiagaraDataInterfaceSkeletalMesh.cpp
@@ -1638,8 +1638,19 @@ bool FNDISkeletalMesh_InstanceData::Init(UNiagaraDataInterfaceSkeletalMesh* Inte
 		}
 	}
 
-	// Gather filtered bones information
+	// Generate excluded root bone index (if any)
 	FReferenceSkeleton& RefSkel = Mesh->RefSkeleton;
+	ExcludedBoneIndex = INDEX_NONE;
+	if (Interface->bExcludeBone && !Interface->ExcludeBoneName.IsNone())
+	{
+		ExcludedBoneIndex = RefSkel.FindBoneIndex(Interface->ExcludeBoneName);
+		if (ExcludedBoneIndex == INDEX_NONE)
+		{
+			UE_LOG(LogNiagara, Warning, TEXT("Skeletal Mesh Data Interface '%s' is missing bone '%s' this is ok but may not exclude what you want Mesh '%s' Component '%s'"), *Interface->GetFullName(), *Interface->ExcludeBoneName.ToString(), *Mesh->GetFullName(), *Component->GetFullName());
+		}
+	}
+
+	// Gather filtered bones information
 	if (Interface->FilteredBones.Num() > 0)
 	{
 		if (RefSkel.GetNum() > TNumericLimits<uint16>::Max())
@@ -1672,6 +1683,12 @@ bool FNDISkeletalMesh_InstanceData::Init(UNiagaraDataInterfaceSkeletalMesh* Inte
 		// Append unfiltered bones to the array
 		for (int32 i=0; i < RefSkel.GetNum(); ++i)
 		{
+			// Don't append excluded bone
+			if (i == ExcludedBoneIndex)
+			{
+				continue;
+			}
+
 			bool bExists = false;
 			for (int32 j=0; j < NumFilteredBones; ++j)
 			{
@@ -1962,6 +1979,10 @@ UNiagaraDataInterfaceSkeletalMesh::UNiagaraDataInterfaceSkeletalMesh(FObjectInit
 #endif
 	, ChangeId(0)
 {
+	static const FName NAME_Root("root");
+	ExcludeBoneName = NAME_Root;
+	bExcludeBone = false;
+
 	Proxy.Reset(new FNiagaraDataInterfaceProxySkeletalMesh());
 }
 
@@ -2081,6 +2102,8 @@ bool UNiagaraDataInterfaceSkeletalMesh::CopyToInternal(UNiagaraDataInterface* De
 	OtherTyped->WholeMeshLOD = WholeMeshLOD;
 	OtherTyped->FilteredBones = FilteredBones;
 	OtherTyped->FilteredSockets = FilteredSockets;
+	OtherTyped->bExcludeBone = bExcludeBone;
+	OtherTyped->ExcludeBoneName = ExcludeBoneName;
 #if WITH_EDITORONLY_DATA
 	OtherTyped->PreviewMesh = PreviewMesh;
 	OtherTyped->bRequiresCPUAccess = bRequiresCPUAccess;
@@ -2104,7 +2127,9 @@ bool UNiagaraDataInterfaceSkeletalMesh::Equals(const UNiagaraDataInterface* Othe
 		OtherTyped->SamplingRegions == SamplingRegions &&
 		OtherTyped->WholeMeshLOD == WholeMeshLOD &&
 		OtherTyped->FilteredBones == FilteredBones &&
-		OtherTyped->FilteredSockets == FilteredSockets;
+		OtherTyped->FilteredSockets == FilteredSockets &&
+		OtherTyped->bExcludeBone == bExcludeBone &&
+		OtherTyped->ExcludeBoneName == ExcludeBoneName;
 }
 
 bool UNiagaraDataInterfaceSkeletalMesh::InitPerInstanceData(void* PerInstanceData, FNiagaraSystemInstance* SystemInstance)
diff --git a/Engine/Plugins/FX/Niagara/Source/NiagaraEditorWidgets/Private/DetailCustomizations/NiagaraDataInterfaceSkeletalMeshDetails.cpp b/Engine/Plugins/FX/Niagara/Source/NiagaraEditorWidgets/Private/DetailCustomizations/NiagaraDataInterfaceSkeletalMeshDetails.cpp
--- a/Engine/Plugins/FX/Niagara/Source/NiagaraEditorWidgets/Private/DetailCustomizations/NiagaraDataInterfaceSkeletalMeshDetails.cpp
+++ b/Engine/Plugins/FX/Niagara/Source/NiagaraEditorWidgets/Private/DetailCustomizations/NiagaraDataInterfaceSkeletalMeshDetails.cpp
@@ -9,6 +9,7 @@
 #include "IDetailGroup.h"
 #include "NiagaraComponent.h"
 #include "Engine/SkeletalMeshSocket.h"
+#include "SNiagaraNamePropertySelector.h"
 
  #define LOCTEXT_NAMESPACE "FNiagaraDataInterfaceSkeletalMeshDetails"
 
@@ -70,13 +71,14 @@ void FNiagaraDataInterfaceSkeletalMeshDetails::CustomizeDetails(IDetailLayoutBui
 
 		 TSharedPtr<IPropertyHandle> BonesProperty = DetailBuilder.GetProperty(GET_MEMBER_NAME_CHECKED(UNiagaraDataInterfaceSkeletalMesh, FilteredBones));
 		 TSharedPtr<IPropertyHandle> SocketsProperty = DetailBuilder.GetProperty(GET_MEMBER_NAME_CHECKED(UNiagaraDataInterfaceSkeletalMesh, FilteredSockets));
+		 TSharedPtr<IPropertyHandle> ExcludeBoneProperty = DetailBuilder.GetProperty(GET_MEMBER_NAME_CHECKED(UNiagaraDataInterfaceSkeletalMesh, ExcludeBoneName));
 
 		 for (TSharedPtr<IPropertyHandle> Property : SkelProperties)
 		 {
 			 FProperty* PropertyPtr = Property->GetProperty();
 			 TArray<TSharedPtr<FName>> PossibleNames;
 
-			if (PropertyPtr == BonesProperty->GetProperty())
+			 if (PropertyPtr == BonesProperty->GetProperty())
 			 {
 				GenerateBonesArray(PossibleNames);
 				BonesBuilder = TSharedPtr<FNiagaraDetailSourcedArrayBuilder>(new FNiagaraDetailSourcedArrayBuilder(Property.ToSharedRef(), PossibleNames));
@@ -88,13 +90,30 @@ void FNiagaraDataInterfaceSkeletalMeshDetails::CustomizeDetails(IDetailLayoutBui
 				 SocketsBuilder = TSharedPtr<FNiagaraDetailSourcedArrayBuilder>(new FNiagaraDetailSourcedArrayBuilder(Property.ToSharedRef(), PossibleNames));
 				 SkelCategory->AddCustomBuilder(SocketsBuilder.ToSharedRef());
 			 }
+			 else if (PropertyPtr == ExcludeBoneProperty->GetProperty())
+			 {
+				 GenerateBonesArray(PossibleNames);
+				 ExcludeBoneWidget = SNew(SNiagaraNamePropertySelector, Property.ToSharedRef(), PossibleNames);
+
+				 IDetailPropertyRow& ExcludeBoneRow = SkelCategory->AddProperty(Property);
+				 ExcludeBoneRow.CustomWidget(false)
+					.NameContent()
+					[
+						Property->CreatePropertyNameWidget()
+					]
+					.ValueContent()
+					.MaxDesiredWidth(TOptional<float>())
+					[
+						ExcludeBoneWidget.ToSharedRef()
+					];
+			 }
 			 else
 			 {
 				 SkelCategory->AddProperty(Property);
 			 }			 
 		 }
 	 }
- }
+}
 
  TSharedRef<IDetailCustomization> FNiagaraDataInterfaceSkeletalMeshDetails::MakeInstance()
  {
@@ -141,6 +160,13 @@ void FNiagaraDataInterfaceSkeletalMeshDetails::CustomizeDetails(IDetailLayoutBui
 		GenerateSocketsArray(PossibleNames);
 		SocketsBuilder->SetSourceArray(PossibleNames);
 	}
+
+	if (ExcludeBoneWidget)
+	{
+		TArray<TSharedPtr<FName>> PossibleNames;
+		GenerateBonesArray(PossibleNames);
+		ExcludeBoneWidget->SetSourceArray(PossibleNames);
+	}
  }
 
 void FNiagaraDataInterfaceSkeletalMeshDetails::GenerateRegionsArray(TArray<TSharedPtr<FName>>& SourceArray)
diff --git a/Engine/Plugins/FX/Niagara/Source/NiagaraEditorWidgets/Private/DetailCustomizations/NiagaraDataInterfaceSkeletalMeshDetails.h b/Engine/Plugins/FX/Niagara/Source/NiagaraEditorWidgets/Private/DetailCustomizations/NiagaraDataInterfaceSkeletalMeshDetails.h
--- a/Engine/Plugins/FX/Niagara/Source/NiagaraEditorWidgets/Private/DetailCustomizations/NiagaraDataInterfaceSkeletalMeshDetails.h
+++ b/Engine/Plugins/FX/Niagara/Source/NiagaraEditorWidgets/Private/DetailCustomizations/NiagaraDataInterfaceSkeletalMeshDetails.h
@@ -14,6 +14,7 @@ class FName;
 class FNiagaraDetailSourcedArrayBuilder;
 class UNiagaraDataInterfaceSkeletalMesh;
 class USkeletalMesh;
+class SNiagaraNamePropertySelector;
 
 /** Details customization for Niagara skeletal mesh data interface. */
 class FNiagaraDataInterfaceSkeletalMeshDetails : public FNiagaraDataInterfaceDetailsBase
@@ -34,6 +35,7 @@ private:
 	TSharedPtr<FNiagaraDetailSourcedArrayBuilder> RegionsBuilder;
 	TSharedPtr<FNiagaraDetailSourcedArrayBuilder> BonesBuilder;
 	TSharedPtr<FNiagaraDetailSourcedArrayBuilder> SocketsBuilder;
+	TSharedPtr<SNiagaraNamePropertySelector> ExcludeBoneWidget;
 	IDetailLayoutBuilder* LayoutBuilder;
 	TWeakObjectPtr<UNiagaraDataInterfaceSkeletalMesh>  MeshInterface;
 	TWeakObjectPtr<USkeletalMesh> MeshObject;
