Fix support for audio Niagara Data Interfaces on the CPU.

diff --git a/Engine/Plugins/FX/Niagara/Source/Niagara/Classes/NiagaraDataInterfaceAudioOscilloscope.h b/Engine/Plugins/FX/Niagara/Source/Niagara/Classes/NiagaraDataInterfaceAudioOscilloscope.h
--- a/Engine/Plugins/FX/Niagara/Source/Niagara/Classes/NiagaraDataInterfaceAudioOscilloscope.h
+++ b/Engine/Plugins/FX/Niagara/Source/Niagara/Classes/NiagaraDataInterfaceAudioOscilloscope.h
@@ -27,7 +27,7 @@ struct FNiagaraDataInterfaceProxyOscilloscope : public FNiagaraDataInterfaceProx
 	 * @param Channel channel index.
 	 * @return Amplitude at this position.
 	 */
-	float SampleAudio(float NormalizedPositionInBuffer, int32 Channel);
+	float SampleAudio(float NormalizedPositionInBuffer, int32 Channel, int32 NumFramesInBuffer, int32 NumChannelsInBuffer);
 
 	/**
 	 * @return the number of channels in the buffer.
@@ -47,8 +47,8 @@ struct FNiagaraDataInterfaceProxyOscilloscope : public FNiagaraDataInterfaceProx
 	void PostAudioToGPU();
 	FReadBuffer& ComputeAndPostSRV();
 
-	// This function pops audio and downsamples it to our specified resolution.
-	void DownsampleAudioToBuffer();
+	// This function pops audio and downsamples it to our specified resolution. Returns the number of frames of audio in the downsampled buffer.
+	int32 DownsampleAudioToBuffer();
 
 	virtual int32 PerInstanceDataPassedToRenderThreadSize() const override
 	{
@@ -76,8 +76,11 @@ private:
 
 	// Handle for the SRV used by the generated HLSL.
 	FReadBuffer GPUDownsampledBuffer;
-	int32 NumChannelsInDownsampledBuffer;
+	FThreadSafeCounter NumChannelsInDownsampledBuffer;
 	
+	// Buffer read by VectorVM worker threads. This vector is guaranteed to not be mutated during the VectorVM tasks.
+	Audio::AlignedFloatBuffer VectorVMReadBuffer;
+
 	FDelegateHandle DeviceCreatedHandle;
 	FDelegateHandle DeviceDestroyedHandle;
 
@@ -111,16 +114,11 @@ public:
 	virtual void GetFunctions(TArray<FNiagaraFunctionSignature>& OutFunctions)override;
 	virtual void GetVMExternalFunction(const FVMExternalFunctionBindingInfo& BindingInfo, void* InstanceData, FVMExternalFunction &OutFunc) override;
 
-	virtual bool CanExecuteOnTarget(ENiagaraSimTarget Target) const override { 
-		if (Target == ENiagaraSimTarget::GPUComputeSim)
-		{
-			return true;
-		}
-		else
-		{
-			return false;
-		}
+	virtual bool CanExecuteOnTarget(ENiagaraSimTarget Target) const override
+	{
+		return true;
 	}
+
 	virtual bool RequiresDistanceFieldData() const override { return false; }
 
 	virtual bool GetFunctionHLSL(const FNiagaraDataInterfaceGPUParamInfo& ParamInfo, const FNiagaraDataInterfaceGeneratedFunction& FunctionInfo, int FunctionInstanceIndex, FString& OutHLSL) override;
diff --git a/Engine/Plugins/FX/Niagara/Source/Niagara/Classes/NiagaraDataInterfaceAudioSpectrum.h b/Engine/Plugins/FX/Niagara/Source/Niagara/Classes/NiagaraDataInterfaceAudioSpectrum.h
--- a/Engine/Plugins/FX/Niagara/Source/Niagara/Classes/NiagaraDataInterfaceAudioSpectrum.h
+++ b/Engine/Plugins/FX/Niagara/Source/Niagara/Classes/NiagaraDataInterfaceAudioSpectrum.h
@@ -62,15 +62,15 @@ struct FNiagaraDataInterfaceProxySpectrum : public FNiagaraDataInterfaceProxySub
 		return 0;
 	}
 
+	// Calculates the CQT
+	void UpdateSpectrum();
+
 private:
 	typedef Audio::TSlidingBuffer<float> FSlidingBuffer;
 	typedef Audio::TAutoSlidingWindow<float, Audio::FAudioBufferAlignedAllocator> FSlidingWindow;
 	typedef Audio::TAutoDeinterleaveView<float, Audio::FAudioBufferAlignedAllocator> FDeinterleaveView;
 	typedef FDeinterleaveView::TChannel<Audio::FAudioBufferAlignedAllocator> FChannel;
 
-	// Calculates the CQT
-	void UpdateSpectrum();
-
 	// Updates internal objects for num channels and samplerate
 	void SetAudioFormat(int32 InNumChannels, float InSampleRate);
 
@@ -159,21 +159,15 @@ public:
 	float NoiseFloorDb;
 
 	//VM function overrides:
-	void GetSpectrumValue(FVectorVMContext& Context) {}
-	void GetNumChannels(FVectorVMContext& Context) {}
+	void GetSpectrumValue(FVectorVMContext& Context);
+	void GetNumChannels(FVectorVMContext& Context);
 
 	virtual void GetFunctions(TArray<FNiagaraFunctionSignature>& OutFunctions) override;
 	virtual void GetVMExternalFunction(const FVMExternalFunctionBindingInfo& BindingInfo, void* InstanceData, FVMExternalFunction &OutFunc) override;
 
-	virtual bool CanExecuteOnTarget(ENiagaraSimTarget Target) const override { 
-		if (Target == ENiagaraSimTarget::GPUComputeSim)
-		{
-			return true;
-		}
-		else
-		{
-			return false;
-		}
+	virtual bool CanExecuteOnTarget(ENiagaraSimTarget Target) const override 
+	{
+		return true;
 	}
 
 	virtual bool GetFunctionHLSL(const FNiagaraDataInterfaceGPUParamInfo& ParamInfo, const FNiagaraDataInterfaceGeneratedFunction& FunctionInfo, int FunctionInstanceIndex, FString& OutHLSL) override;
diff --git a/Engine/Plugins/FX/Niagara/Source/Niagara/Private/NiagaraDataInterfaceAudioOscilloscope.cpp b/Engine/Plugins/FX/Niagara/Source/Niagara/Private/NiagaraDataInterfaceAudioOscilloscope.cpp
--- a/Engine/Plugins/FX/Niagara/Source/Niagara/Private/NiagaraDataInterfaceAudioOscilloscope.cpp
+++ b/Engine/Plugins/FX/Niagara/Source/Niagara/Private/NiagaraDataInterfaceAudioOscilloscope.cpp
@@ -36,6 +36,9 @@ FNiagaraDataInterfaceProxyOscilloscope::FNiagaraDataInterfaceProxyOscilloscope(i
 {
 	DeviceCreatedHandle = FAudioDeviceManagerDelegates::OnAudioDeviceCreated.AddRaw(this, &FNiagaraDataInterfaceProxyOscilloscope::OnNewDeviceCreated);
 	DeviceDestroyedHandle = FAudioDeviceManagerDelegates::OnAudioDeviceDestroyed.AddRaw(this, &FNiagaraDataInterfaceProxyOscilloscope::OnDeviceDestroyed);
+
+	VectorVMReadBuffer.Reset();
+	VectorVMReadBuffer.AddZeroed(Resolution * AUDIO_MIXER_MAX_OUTPUT_CHANNELS);
 }
 
 FNiagaraDataInterfaceProxyOscilloscope::~FNiagaraDataInterfaceProxyOscilloscope()
@@ -121,36 +124,41 @@ UNiagaraDataInterfaceAudioOscilloscope::UNiagaraDataInterfaceAudioOscilloscope(F
 {
 	Proxy = TUniquePtr<FNiagaraDataInterfaceProxyOscilloscope>(new FNiagaraDataInterfaceProxyOscilloscope(Resolution, ScopeInMilliseconds));
 }
-
-float FNiagaraDataInterfaceProxyOscilloscope::SampleAudio(float NormalizedPositionInBuffer, int32 ChannelIndex)
+PRAGMA_DISABLE_OPTIMIZATION
+float FNiagaraDataInterfaceProxyOscilloscope::SampleAudio(float NormalizedPositionInBuffer, int32 ChannelIndex, int32 NumFramesInBuffer, int32 NumChannelsInBuffer)
 {
-	if (DownsampledBuffer.Num() == 0)
+	if (NumFramesInBuffer == 0 || NumChannelsInBuffer == 0)
 	{
 		return 0.0f;
 	}
 
 	NormalizedPositionInBuffer = FMath::Clamp(NormalizedPositionInBuffer, 0.0f, 1.0f - SMALL_NUMBER);
-	float FrameIndex = NormalizedPositionInBuffer * DownsampledBuffer.Num();
+	//ChannelIndex = FMath::Clamp(ChannelIndex, 0, NumChannelsInDownsampledBuffer);
+
+	float FrameIndex = NormalizedPositionInBuffer * NumFramesInBuffer;
 	int32 LowerFrameIndex = FMath::FloorToInt(FrameIndex);
-	float LowerFrameAmplitude = DownsampledBuffer[LowerFrameIndex * NumChannelsInDownsampledBuffer + ChannelIndex];
-	int32 HigherFrameIndex = FMath::CeilToInt(FrameIndex);
-	float HigherFrameAmplitude = DownsampledBuffer[HigherFrameIndex * NumChannelsInDownsampledBuffer + ChannelIndex];
+	float LowerFrameAmplitude = VectorVMReadBuffer[LowerFrameIndex * NumChannelsInBuffer + ChannelIndex];
+	int32 HigherFrameIndex = FMath::Clamp(FMath::CeilToInt(FrameIndex), 0, NumFramesInBuffer - 1);
+	float HigherFrameAmplitude = VectorVMReadBuffer[HigherFrameIndex * NumChannelsInBuffer + ChannelIndex];
 	float Fraction = HigherFrameIndex - FrameIndex;
 	return FMath::Lerp<float>(LowerFrameAmplitude, HigherFrameAmplitude, Fraction);
 }
 
 void UNiagaraDataInterfaceAudioOscilloscope::SampleAudio(FVectorVMContext& Context)
 {
-	VectorVM::FUserPtrHandler<UNiagaraDataInterfaceAudioOscilloscope> InstData(Context);
-	VectorVM::FExternalFuncRegisterHandler<float> InNormalizedPos(Context);
-	VectorVM::FExternalFuncRegisterHandler<int32> InChannel(Context);
+	const int32 NumFramesInDownsampledBuffer = GetProxyAs<FNiagaraDataInterfaceProxyOscilloscope>()->DownsampleAudioToBuffer();
+
+	VectorVM::FExternalFuncInputHandler<float> InNormalizedPos(Context);
+	VectorVM::FExternalFuncInputHandler<int32> InChannel(Context);
 	VectorVM::FExternalFuncRegisterHandler<float> OutAmplitude(Context);
 
+	const int32 NumChannelsInDownsampledBuffer = GetProxyAs<FNiagaraDataInterfaceProxyOscilloscope>()->GetNumChannels();
+
 	for (int32 InstanceIdx = 0; InstanceIdx < Context.NumInstances; ++InstanceIdx)
 	{
 		float Position = InNormalizedPos.Get();
 		int32 Channel = InChannel.Get();
-		*OutAmplitude.GetDest() = GetProxyAs<FNiagaraDataInterfaceProxyOscilloscope>()->SampleAudio(Position, Channel);
+		*OutAmplitude.GetDest() = GetProxyAs<FNiagaraDataInterfaceProxyOscilloscope>()->SampleAudio(Position, Channel, NumFramesInDownsampledBuffer, NumChannelsInDownsampledBuffer);
 
 		InNormalizedPos.Advance();
 		InChannel.Advance();
@@ -160,12 +168,11 @@ void UNiagaraDataInterfaceAudioOscilloscope::SampleAudio(FVectorVMContext& Conte
 
 int32 FNiagaraDataInterfaceProxyOscilloscope::GetNumChannels()
 {
-	return NumChannelsInDownsampledBuffer;
+	return NumChannelsInDownsampledBuffer.GetValue();
 }
 
 void UNiagaraDataInterfaceAudioOscilloscope::GetNumChannels(FVectorVMContext& Context)
 {
-	VectorVM::FUserPtrHandler<UNiagaraDataInterfaceAudioOscilloscope> InstData(Context);
 	VectorVM::FExternalFuncRegisterHandler<int32> OutChannel(Context);
 
 	for (int32 InstanceIdx = 0; InstanceIdx < Context.NumInstances; ++InstanceIdx)
@@ -174,6 +181,8 @@ void UNiagaraDataInterfaceAudioOscilloscope::GetNumChannels(FVectorVMContext& Co
 	}
 }
 
+PRAGMA_ENABLE_OPTIMIZATION
+
 void UNiagaraDataInterfaceAudioOscilloscope::GetFunctions(TArray<FNiagaraFunctionSignature>& OutFunctions)
 {
 	Super::GetFunctions(OutFunctions);
@@ -336,7 +345,7 @@ void FNiagaraDataInterfaceProxyOscilloscope::OnUpdateResampling(int32 InResoluti
 	Resolution = InResolution;
 	ScopeInMilliseconds = InScopeInMilliseconds;
 
-	const int32 NumSamplesInBuffer = Resolution * NumChannelsInDownsampledBuffer;
+	const int32 NumSamplesInBuffer = Resolution * NumChannelsInDownsampledBuffer.GetValue();
 	if (NumSamplesInBuffer)
 	{
 		ENQUEUE_RENDER_COMMAND(FUpdateDIAudioBuffer)(
@@ -349,8 +358,9 @@ void FNiagaraDataInterfaceProxyOscilloscope::OnUpdateResampling(int32 InResoluti
 
 			GPUDownsampledBuffer.Initialize(sizeof(float), NumSamplesInBuffer, EPixelFormat::PF_R32_FLOAT, BUF_Static);
 		});
+
+		VectorVMReadBuffer.SetNumZeroed(Resolution * AUDIO_MIXER_MAX_OUTPUT_CHANNELS);
 	}
-	
 }
 
 FReadBuffer& FNiagaraDataInterfaceProxyOscilloscope::ComputeAndPostSRV()
@@ -457,7 +467,7 @@ void FNiagaraDataInterfaceProxyOscilloscope::PostAudioToGPU()
 		size_t BufferSize = DownsampledBuffer.Num() * sizeof(float);
 		if (BufferSize != 0 && !GPUDownsampledBuffer.NumBytes)
 		{
-			GPUDownsampledBuffer.Initialize(sizeof(float), Resolution * NumChannelsInDownsampledBuffer, EPixelFormat::PF_R32_FLOAT, BUF_Static);
+			GPUDownsampledBuffer.Initialize(sizeof(float), Resolution * NumChannelsInDownsampledBuffer.GetValue(), EPixelFormat::PF_R32_FLOAT, BUF_Static);
 		}
 
 		if (GPUDownsampledBuffer.NumBytes > 0)
@@ -470,39 +480,40 @@ void FNiagaraDataInterfaceProxyOscilloscope::PostAudioToGPU()
 	});
 }
 
-void FNiagaraDataInterfaceProxyOscilloscope::DownsampleAudioToBuffer()
+int32 FNiagaraDataInterfaceProxyOscilloscope::DownsampleAudioToBuffer()
 {
 	FScopeLock ScopeLock(&DownsampleBufferLock);
 
 	// Get the number of channels from the first listener we find.
 	// If the num channels hasn't been set, either SubmixListeners is empty
 	// or we haven't pushed audio to any of them yet.
-	NumChannelsInDownsampledBuffer = 0;
 	float SourceSampleRate = 0.0f;
+	int32 NumChannels = 0;
 
 	for (auto& Listener : SubmixListeners)
 	{
-		NumChannelsInDownsampledBuffer = Listener.Value->GetNumChannels();
+		NumChannels = Listener.Value->GetNumChannels();
 		SourceSampleRate = Listener.Value->GetSampleRate();
 
-		if (NumChannelsInDownsampledBuffer != 0)
+		if (NumChannels != 0)
 		{
 			break;
 		}
 	}
 
-	if (NumChannelsInDownsampledBuffer == 0 || FMath::IsNearlyZero(SourceSampleRate))
+	if (NumChannels == 0 || FMath::IsNearlyZero(SourceSampleRate))
 	{
-		return;
+		return 0;
 	}
 
-	check(NumChannelsInDownsampledBuffer > 0);
+	check(NumChannels > 0);
+	NumChannelsInDownsampledBuffer.Set(NumChannels);
 
 	// Get the number of frames of audio in the original sample rate to show.
 	int32 NumFramesToPop = Align(FMath::FloorToInt(SourceSampleRate * (ScopeInMilliseconds / 1000.0f)), 4);
 
 	// If we have enough frames buffered to pop, update the PopBuffer. Otherwise, use the previous frames buffer.
-	int32 NumSamplesToPop = NumFramesToPop * NumChannelsInDownsampledBuffer;
+	int32 NumSamplesToPop = NumFramesToPop * NumChannels;
 
 	if (PopBuffer.Num() != NumSamplesToPop)
 	{
@@ -521,7 +532,7 @@ void FNiagaraDataInterfaceProxyOscilloscope::DownsampleAudioToBuffer()
 
 	Audio::FResamplingParameters ResampleParameters = {
 		Audio::EResamplingMethod::Linear,
-		NumChannelsInDownsampledBuffer,
+		NumChannels,
 		SourceSampleRate,
 		DestinationSampleRate,
 		PopBuffer
@@ -539,8 +550,12 @@ void FNiagaraDataInterfaceProxyOscilloscope::DownsampleAudioToBuffer()
 	check(bResampleResult);
 
 	// here we may have an extra sample or two due to roundoff.
-	DownsampledBuffer.SetNumZeroed(Resolution * NumChannelsInDownsampledBuffer);
-}
+	DownsampledBuffer.SetNumZeroed(Resolution * NumChannels);
 
+	check(DownsampledBuffer.Num() <= VectorVMReadBuffer.Num());
+	FMemory::Memcpy(VectorVMReadBuffer.GetData(), DownsampledBuffer.GetData(), VectorVMReadBuffer.Num());
+
+	return Resolution;
+}
 
 #undef LOCTEXT_NAMESPACE
\ No newline at end of file
diff --git a/Engine/Plugins/FX/Niagara/Source/Niagara/Private/NiagaraDataInterfaceAudioSpectrum.cpp b/Engine/Plugins/FX/Niagara/Source/Niagara/Private/NiagaraDataInterfaceAudioSpectrum.cpp
--- a/Engine/Plugins/FX/Niagara/Source/Niagara/Private/NiagaraDataInterfaceAudioSpectrum.cpp
+++ b/Engine/Plugins/FX/Niagara/Source/Niagara/Private/NiagaraDataInterfaceAudioSpectrum.cpp
@@ -48,7 +48,7 @@ float FNiagaraDataInterfaceProxySpectrum::GetSpectrumValue(float InNormalizedPos
 	check(InChannelIndex >= 0);
 
 	// Check if channel index is within channel spectrum buffers
-	if (ChannelSpectrumBuffers.Num() >= InChannelIndex)
+	if (InChannelIndex >= ChannelSpectrumBuffers.Num())
 	{
 		return 0.f;
 	}
@@ -197,8 +197,6 @@ void FNiagaraDataInterfaceProxySpectrum::PostDataToGPU()
 
 void FNiagaraDataInterfaceProxySpectrum::UpdateSpectrum()
 {
-	check(IsInRenderingThread());
-
 	FScopeLock ScopeLock(&BufferLock);
 
 	int32 SourceNumChannels = GetNumChannels();
@@ -566,6 +564,36 @@ UNiagaraDataInterfaceAudioSpectrum::UNiagaraDataInterfaceAudioSpectrum(FObjectIn
 }
 
 
+void UNiagaraDataInterfaceAudioSpectrum::GetSpectrumValue(FVectorVMContext& Context)
+{
+	GetProxyAs<FNiagaraDataInterfaceProxySpectrum>()->UpdateSpectrum();
+
+	VectorVM::FExternalFuncInputHandler<float> InNormalizedPos(Context);
+	VectorVM::FExternalFuncInputHandler<int32> InChannel(Context);
+	VectorVM::FExternalFuncRegisterHandler<float> OutValue(Context);
+
+	for (int32 InstanceIdx = 0; InstanceIdx < Context.NumInstances; ++InstanceIdx)
+	{
+		float Position = InNormalizedPos.Get();
+		int32 Channel = InChannel.Get();
+		*OutValue.GetDest() = GetProxyAs<FNiagaraDataInterfaceProxySpectrum>()->GetSpectrumValue(Position, Channel);
+
+		InNormalizedPos.Advance();
+		InChannel.Advance();
+		OutValue.Advance();
+	}
+}
+
+void UNiagaraDataInterfaceAudioSpectrum::GetNumChannels(FVectorVMContext& Context)
+{
+	VectorVM::FExternalFuncRegisterHandler<int32> OutChannel(Context);
+
+	for (int32 InstanceIdx = 0; InstanceIdx < Context.NumInstances; ++InstanceIdx)
+	{
+		*OutChannel.GetDestAndAdvance() = GetProxyAs<FNiagaraDataInterfaceProxySpectrum>()->GetNumChannels();
+	}
+}
+
 void UNiagaraDataInterfaceAudioSpectrum::GetFunctions(TArray<FNiagaraFunctionSignature>& OutFunctions)
 {
 	Super::GetFunctions(OutFunctions);
