Reimplement SNiagaraParameterPanel to work with SItemSelector to allow for subcategory handling and general code cleanliness.

diff --git a/Engine/Plugins/FX/Niagara/Source/Niagara/Private/NiagaraModule.cpp b/Engine/Plugins/FX/Niagara/Source/Niagara/Private/NiagaraModule.cpp
--- a/Engine/Plugins/FX/Niagara/Source/Niagara/Private/NiagaraModule.cpp
+++ b/Engine/Plugins/FX/Niagara/Source/Niagara/Private/NiagaraModule.cpp
@@ -418,6 +418,7 @@ UEnum* FNiagaraTypeDefinition::ExecutionStateSourceEnum;
 UEnum* FNiagaraTypeDefinition::ScriptUsageEnum;
 
 UEnum* FNiagaraTypeDefinition::ParameterScopeEnum;
+UEnum* FNiagaraTypeDefinition::ParameterPanelCategoryEnum;
 
 FNiagaraTypeDefinition FNiagaraTypeDefinition::ParameterMapDef;
 FNiagaraTypeDefinition FNiagaraTypeDefinition::IDDef;
@@ -553,6 +554,7 @@ void FNiagaraTypeDefinition::Init()
 	ScriptUsageEnum = StaticEnum<ENiagaraScriptUsage>();
 
 	ParameterScopeEnum = StaticEnum<ENiagaraParameterScope>();
+	ParameterPanelCategoryEnum = StaticEnum<ENiagaraParameterPanelCategory>();
 	
 #if WITH_EDITOR
 	RecreateUserDefinedTypeRegistry();
diff --git a/Engine/Plugins/FX/Niagara/Source/Niagara/Public/NiagaraTypes.h b/Engine/Plugins/FX/Niagara/Source/Niagara/Public/NiagaraTypes.h
--- a/Engine/Plugins/FX/Niagara/Source/Niagara/Public/NiagaraTypes.h
+++ b/Engine/Plugins/FX/Niagara/Source/Niagara/Public/NiagaraTypes.h
@@ -618,6 +618,51 @@ private:
 		FString NamespaceString;
 };
 
+#if WITH_EDITOR
+UENUM()
+enum class ENiagaraParameterPanelCategory : uint32
+{
+	/** Parameter that is an input argument into this graph.*/
+	Input UMETA(DisplayName = "Module Inputs"),
+
+	Attributes UMETA(DisplayName = "Input Attributes"),
+
+	/** Parameter is output to the owning stack stage from this script, but is only meaningful if bound elsewhere in the stage.*/
+	Output UMETA(DisplayName = "Output Attributes"),
+
+	/** Parameter is initialized at the start of this script and is only used within the context of this script. It is invisible to the parent stage stack.*/
+	Local UMETA(DisplayName = "Local"),
+
+	/** Parameter that is exposed to the owning component for editing and are read-only when used in the graph*/
+	User UMETA(DisplayName = "User"),
+
+	/** Parameter provided by the engine. These are explicitly defined by the engine codebase and read-only. */
+	Engine UMETA(DisplayName = "Engine (Generic)", Hidden),
+
+	/** Parameter provided by the engine focused on the owning component. These are explicitly defined by the engine codebase and read-only.*/
+	Owner UMETA(DisplayName = "Engine (Owner)", Hidden),
+
+	/** Parameter is an attribute of the owning system payload. It is persistent across frames and initialized in the System Spawn stage of the stack.*/
+	System  UMETA(DisplayName = "System"),
+
+	/** Parameter is an attribute of the owning emitter payload. It is persistent across frames and initialized in the Emitter Spawn stage of the stack.*/
+	Emitter   UMETA(DisplayName = "Emitter"),
+
+	/** Parameter is an attribute of the owning particle payload. It is persistent across frames and initialized in the Particle Spawn stage of the stack.*/
+	Particles  UMETA(DisplayName = "Particles"),
+
+	/** Parameter is initialized at the start of this stage and can be shared amongst other modules within this stack stage, but is not persistent across runs or from stack stage to stack stage.*/
+	ScriptTransient UMETA(DisplayName = "Stage (Transient)"),
+
+	StaticSwitch UMETA(DisplayName = "Static Switch"),
+
+	// insert new categories before
+	None UMETA(Hidden),
+
+	Num UMETA(Hidden)
+};
+#endif
+
 USTRUCT()
 struct NIAGARA_API FNiagaraVariableMetaData
 {
@@ -996,6 +1041,7 @@ public:
 	static UEnum* GetExecutionStateSouceEnum() { return ExecutionStateSourceEnum; }
 	static UEnum* GetSimulationTargetEnum() { return SimulationTargetEnum; }
 	static UEnum* GetScriptUsageEnum() { return ScriptUsageEnum; }
+	static UEnum* GetParameterPanelCategoryEnum() { return ParameterPanelCategoryEnum; }
 
 	static UEnum* GetParameterScopeEnum() { return ParameterScopeEnum; }
 
@@ -1066,6 +1112,7 @@ private:
 	static UEnum* ExecutionStateSourceEnum;
 
 	static UEnum* ParameterScopeEnum;
+	static UEnum* ParameterPanelCategoryEnum;
 
 	static UScriptStruct* ParameterMapStruct;
 	static UScriptStruct* IDStruct;
diff --git a/Engine/Plugins/FX/Niagara/Source/NiagaraEditor/Private/NiagaraEditorCommon.cpp b/Engine/Plugins/FX/Niagara/Source/NiagaraEditor/Private/NiagaraEditorCommon.cpp
--- a/Engine/Plugins/FX/Niagara/Source/NiagaraEditor/Private/NiagaraEditorCommon.cpp
+++ b/Engine/Plugins/FX/Niagara/Source/NiagaraEditor/Private/NiagaraEditorCommon.cpp
@@ -1362,3 +1362,8 @@ INiagaraScriptGraphFocusInfo::~INiagaraScriptGraphFocusInfo()
 {
 	//Stand-in definition for abstract INiagaraScriptGraphFocusInfo's pure virtual dtor
 }
+
+bool FNiagaraScriptVariableAndViewInfo::operator==(const FNiagaraScriptVariableAndViewInfo& Other) const
+{
+	return ScriptVariable == Other.ScriptVariable && MetaData.GetUsage() == Other.MetaData.GetUsage();
+}
diff --git a/Engine/Plugins/FX/Niagara/Source/NiagaraEditor/Private/ViewModels/NiagaraParameterNameViewModel.cpp b/Engine/Plugins/FX/Niagara/Source/NiagaraEditor/Private/ViewModels/NiagaraParameterNameViewModel.cpp
--- a/Engine/Plugins/FX/Niagara/Source/NiagaraEditor/Private/ViewModels/NiagaraParameterNameViewModel.cpp
+++ b/Engine/Plugins/FX/Niagara/Source/NiagaraEditor/Private/ViewModels/NiagaraParameterNameViewModel.cpp
@@ -20,9 +20,9 @@
 /// Parameter Panel Entry Parameter Name ViewModel							///
 ///////////////////////////////////////////////////////////////////////////////
 
-FNiagaraParameterPanelEntryParameterNameViewModel::FNiagaraParameterPanelEntryParameterNameViewModel(FCreateWidgetForActionData* const InCreateData, const FNiagaraScriptVariableAndViewInfo& InScriptVarAndViewInfo)
-	: CreateData(InCreateData)
-	, CachedScriptVarAndViewInfo(InScriptVarAndViewInfo)
+FNiagaraParameterPanelEntryParameterNameViewModel::FNiagaraParameterPanelEntryParameterNameViewModel(/*FCreateWidgetForActionData* const InCreateData,*/ const FNiagaraScriptVariableAndViewInfo& InScriptVarAndViewInfo)
+	/*: CreateData(InCreateData)*/
+	: CachedScriptVarAndViewInfo(InScriptVarAndViewInfo)
 {};
 
 TSharedRef<SWidget> FNiagaraParameterPanelEntryParameterNameViewModel::CreateScopeSlotWidget() const
@@ -71,13 +71,13 @@ TSharedRef<SInlineEditableTextBlock> FNiagaraParameterPanelEntryParameterNameVie
 	SAssignNew(DisplayWidget, SInlineEditableTextBlock)
 		.Text(this, &FNiagaraParameterPanelEntryParameterNameViewModel::GetParameterNameText)
 		.Font(NameFont)
-		.HighlightText(CreateData->HighlightText)
+		/*.HighlightText(CreateData->HighlightText)*/
 		.OnTextCommitted(this, &FNiagaraParameterPanelEntryParameterNameViewModel::OnParameterRenamed)
 		.OnVerifyTextChanged(this, &FNiagaraParameterPanelEntryParameterNameViewModel::VerifyParameterNameChanged)
-		.IsSelected(CreateData->IsRowSelectedDelegate)
+		/*.IsSelected(CreateData->IsRowSelectedDelegate)*/
 		.IsReadOnly(bIsReadOnly);
 
-	CreateData->OnRenameRequest->BindSP(static_cast<SInlineEditableTextBlock*>(DisplayWidget.Get()), &SInlineEditableTextBlock::EnterEditingMode);
+	/*CreateData->OnRenameRequest->BindSP(static_cast<SInlineEditableTextBlock*>(DisplayWidget.Get()), &SInlineEditableTextBlock::EnterEditingMode);*/
 
 	return DisplayWidget.ToSharedRef();
 }
diff --git a/Engine/Plugins/FX/Niagara/Source/NiagaraEditor/Private/ViewModels/NiagaraParameterNameViewModel.h b/Engine/Plugins/FX/Niagara/Source/NiagaraEditor/Private/ViewModels/NiagaraParameterNameViewModel.h
--- a/Engine/Plugins/FX/Niagara/Source/NiagaraEditor/Private/ViewModels/NiagaraParameterNameViewModel.h
+++ b/Engine/Plugins/FX/Niagara/Source/NiagaraEditor/Private/ViewModels/NiagaraParameterNameViewModel.h
@@ -40,7 +40,7 @@ public:
 	DECLARE_DELEGATE_ThreeParams(FOnParameterRenamed, const FNiagaraVariable&, const FNiagaraVariableMetaData&, const FText&)
 	DECLARE_DELEGATE_RetVal_FourParams(bool, FOnVerifyParameterRenamed, const FNiagaraVariable&, const FNiagaraVariableMetaData&, TOptional<const FText> , FText&)
 
-	FNiagaraParameterPanelEntryParameterNameViewModel(FCreateWidgetForActionData* const InCreateData, const FNiagaraScriptVariableAndViewInfo& InScriptVarAndViewInfo);
+	FNiagaraParameterPanelEntryParameterNameViewModel(const FNiagaraScriptVariableAndViewInfo& InScriptVarAndViewInfo);
 
 	/** Begin INiagaraParameterNameViewModel Interface */
 	virtual TSharedRef<SWidget> CreateScopeSlotWidget() const override;
@@ -65,7 +65,7 @@ private:
 	FOnParameterRenamed OnParameterRenamedDelegate;
 	FOnVerifyParameterRenamed OnVerifyParameterRenamedDelegate;
 
-	FCreateWidgetForActionData* const CreateData;
+	//FCreateWidgetForActionData* const CreateData;
 	const FNiagaraScriptVariableAndViewInfo CachedScriptVarAndViewInfo;
 
 	TAttribute<bool> bIsReadOnly; //@todo(ng) debug
diff --git a/Engine/Plugins/FX/Niagara/Source/NiagaraEditor/Private/ViewModels/NiagaraParameterPanelViewModel.cpp b/Engine/Plugins/FX/Niagara/Source/NiagaraEditor/Private/ViewModels/NiagaraParameterPanelViewModel.cpp
--- a/Engine/Plugins/FX/Niagara/Source/NiagaraEditor/Private/ViewModels/NiagaraParameterPanelViewModel.cpp
+++ b/Engine/Plugins/FX/Niagara/Source/NiagaraEditor/Private/ViewModels/NiagaraParameterPanelViewModel.cpp
@@ -27,207 +27,38 @@
 #include "ScopedTransaction.h"
 #include "NiagaraSystemEditorData.h"
 #include "Widgets/Text/SInlineEditableTextBlock.h"
+#include "NiagaraEditorModule.h"
+#include "NiagaraTypes.h"
 
 #define LOCTEXT_NAMESPACE "NiagaraParameterPanelViewModel"
 
-FText NiagaraParameterPanelSectionID::OnGetSectionTitle(const NiagaraParameterPanelSectionID::Type InSection)
-{
-	/* Setup an appropriate name for the section for this node */
-	FText SectionTitle;
-	switch (InSection)
-	{
-	case NiagaraParameterPanelSectionID::ENGINE:
-		SectionTitle = LOCTEXT("EngineParamSection", "Engine");
-		break;
-	case NiagaraParameterPanelSectionID::USER:
-		SectionTitle = LOCTEXT("UserParamSection", "User");
-		break;
-	case NiagaraParameterPanelSectionID::SYSTEM:
-		SectionTitle = LOCTEXT("SystemParamSection", "System");
-		break;
-	case NiagaraParameterPanelSectionID::EMITTER:
-		SectionTitle = LOCTEXT("EmitterParamSection", "Emitter");
-		break;
-	case NiagaraParameterPanelSectionID::OWNER:
-		SectionTitle = LOCTEXT("OwnerParamSection", "Owner");
-		break;
-	case NiagaraParameterPanelSectionID::PARTICLES:
-		SectionTitle = LOCTEXT("ParticlesParamSection", "Particles");
-		break;
-	case NiagaraParameterPanelSectionID::INPUTS:
-		SectionTitle = LOCTEXT("InputsParamSection", "Inputs");
-		break;
-	case NiagaraParameterPanelSectionID::REFERENCES:
-		SectionTitle = LOCTEXT("ReferencesParamSection", "References");
-		break;
-	case NiagaraParameterPanelSectionID::OUTPUTS:
-		SectionTitle = LOCTEXT("OutputsParamSection", "Outputs");
-		break;
-	case NiagaraParameterPanelSectionID::LOCALS:
-		SectionTitle = LOCTEXT("LocalsParamSection", "Locals");
-		break;
-	case NiagaraParameterPanelSectionID::INITIALVALUES:
-		SectionTitle = LOCTEXT("InitialValuesParamSection", "Initial Values");
-		break;
-	case NiagaraParameterPanelSectionID::CUSTOM:
-		SectionTitle = LOCTEXT("CustomParamSection", "Custom");
-		break;
-// 	case NiagaraParameterPanelSectionID::PARAMETERCOLLECTION: //@TODO implement parameter collection handling
-// 		SectionTitle = NSLOCTEXT("GraphActionNode", "ParameterCollection", "Parameter Collection");
-// 		break;
-	case NiagaraParameterPanelSectionID::NONE:
-	default:
-		SectionTitle = FText::GetEmpty();
-		break;
-	}
-
-	return SectionTitle;
-}
-
-const NiagaraParameterPanelSectionID::Type NiagaraParameterPanelSectionID::GetSectionForScope(ENiagaraParameterScope InScope)
-{
-	switch (InScope) {
-	case ENiagaraParameterScope::Engine:
-		return NiagaraParameterPanelSectionID::ENGINE;
-	case ENiagaraParameterScope::Owner:
-		return NiagaraParameterPanelSectionID::OWNER;
-	case ENiagaraParameterScope::User:
-		return NiagaraParameterPanelSectionID::USER;
-	case ENiagaraParameterScope::System:
-		return NiagaraParameterPanelSectionID::SYSTEM;
-	case ENiagaraParameterScope::Emitter:
-		return NiagaraParameterPanelSectionID::EMITTER;
-	case ENiagaraParameterScope::Particles:
-		return NiagaraParameterPanelSectionID::PARTICLES;
-	case ENiagaraParameterScope::Custom:
-		return NiagaraParameterPanelSectionID::CUSTOM;
-	case ENiagaraParameterScope::ScriptPersistent:
-	case ENiagaraParameterScope::ScriptTransient:
-		// This is a potential situation if a script alias param has not had its scope cached by compiling.
-		ensureMsgf(false, TEXT("Tried to get section ID for script transient or persistent scope parameter!"));
-		return NiagaraParameterPanelSectionID::NONE;
-	default:
-		ensureMsgf(false, TEXT("Failed to find matching section ID for script parameter scope!"));
-	};
-	return NiagaraParameterPanelSectionID::NONE;
-}
-
-const NiagaraParameterPanelSectionID::Type NiagaraParameterPanelSectionID::GetSectionForParameterMetaData(const FNiagaraVariableMetaData& MetaData)
-{
-	if (MetaData.GetIsStaticSwitch())
-	{
-		return NiagaraParameterPanelSectionID::INPUTS;
-	}
-
-	ENiagaraParameterScope MetaDataScope;
-	FNiagaraEditorUtilities::GetVariableMetaDataScope(MetaData, MetaDataScope);
-	ENiagaraScriptParameterUsage MetaDataUsage = MetaData.GetUsage();
-
-	if (MetaDataUsage == ENiagaraScriptParameterUsage::Local)
-	{
-		return NiagaraParameterPanelSectionID::LOCALS;
-	}
-	else if (MetaDataUsage == ENiagaraScriptParameterUsage::InitialValueInput)
-	{
-		return NiagaraParameterPanelSectionID::INITIALVALUES;
-	}
-	else if (MetaDataUsage == ENiagaraScriptParameterUsage::Input)
-	{
-		if (MetaDataScope == ENiagaraParameterScope::Input)
-		{
-			return NiagaraParameterPanelSectionID::INPUTS;
-		}
-		else
-		{
-			return NiagaraParameterPanelSectionID::REFERENCES;
-		}
-	}
-	else if (MetaDataUsage == ENiagaraScriptParameterUsage::Output)
-	{
-		return NiagaraParameterPanelSectionID::OUTPUTS;
-	}
-	else if (MetaDataUsage == ENiagaraScriptParameterUsage::InputOutput)
-	{
-		ensureMsgf(false, TEXT("Encountered an InputOutput parameter usage when getting section ID for parameter panel!"));
-		return NiagaraParameterPanelSectionID::REFERENCES;
-	}
-	else
-	{
-		ensureMsgf(false, TEXT("Failed to find matching section ID for script parameter usage!"));
-	}
-	return NiagaraParameterPanelSectionID::LOCALS;
-}
-
-ENiagaraParameterScope NiagaraParameterPanelSectionID::GetScopeForNewParametersInSection(NiagaraParameterPanelSectionID::Type InSection)
-{
-	switch (InSection) {
-	case NiagaraParameterPanelSectionID::ENGINE:
-		return ENiagaraParameterScope::Engine;
-	case NiagaraParameterPanelSectionID::OWNER:
-		return ENiagaraParameterScope::Owner;
-	case NiagaraParameterPanelSectionID::USER:
-		return ENiagaraParameterScope::User;
-	case NiagaraParameterPanelSectionID::SYSTEM:
-		return ENiagaraParameterScope::System;
-	case NiagaraParameterPanelSectionID::EMITTER:
-		return ENiagaraParameterScope::Emitter;
-	case NiagaraParameterPanelSectionID::PARTICLES:
-		return ENiagaraParameterScope::Particles;
-	case NiagaraParameterPanelSectionID::LOCALS:
-		return ENiagaraParameterScope::Local;
-	case NiagaraParameterPanelSectionID::INITIALVALUES:
-		return ENiagaraParameterScope::Particles;
-	case NiagaraParameterPanelSectionID::INPUTS:
-		return ENiagaraParameterScope::Input;
-	case NiagaraParameterPanelSectionID::OUTPUTS:
-		return ENiagaraParameterScope::Output;
-
-	// Default to Particles scope if coming from section IDs that are not directly associated with scope.
-	case NiagaraParameterPanelSectionID::REFERENCES:
-		return ENiagaraParameterScope::Particles;
-
-	case NiagaraParameterPanelSectionID::NONE:
-		ensureMsgf(false, TEXT("Encountered invalid parameter panel section ID NONE when getting scope from section!"));
-		break;
-	};
-	ensureMsgf(false, TEXT("Did not encounter a known section ID when getting scope from section!"));
-	return ENiagaraParameterScope::Particles;
-}
-
-ENiagaraScriptParameterUsage NiagaraParameterPanelSectionID::GetUsageForNewParametersInSection(NiagaraParameterPanelSectionID::Type InSection)
-{
-	switch (InSection) {
-	case NiagaraParameterPanelSectionID::INPUTS:
-	case NiagaraParameterPanelSectionID::REFERENCES:
-		return ENiagaraScriptParameterUsage::Input;
-	case NiagaraParameterPanelSectionID::OUTPUTS:
-		return ENiagaraScriptParameterUsage::Output;
-	case NiagaraParameterPanelSectionID::LOCALS:
-		return ENiagaraScriptParameterUsage::Local;
-	case NiagaraParameterPanelSectionID::INITIALVALUES:
-		return ENiagaraScriptParameterUsage::InitialValueInput;
-
-	// By convention, default new parameters created in system editor to output usage.
-	case NiagaraParameterPanelSectionID::ENGINE:
-	case NiagaraParameterPanelSectionID::USER:
-	case NiagaraParameterPanelSectionID::SYSTEM:
-	case NiagaraParameterPanelSectionID::EMITTER:
-	case NiagaraParameterPanelSectionID::PARTICLES:
-		return ENiagaraScriptParameterUsage::Output;
-
-	case NiagaraParameterPanelSectionID::NONE:
-		ensureMsgf(false, TEXT("Encountered invalid parameter panel section ID NONE when getting usage from section!"));
-		break;
-	};
-	ensureMsgf(false, TEXT("Did not encounter a known section ID when getting usage from section!"));
-	return ENiagaraScriptParameterUsage::Output;
-}
-
-TSharedRef<SWidget> INiagaraParameterPanelViewModel::GetScriptParameterVisualWidget(FCreateWidgetForActionData* const InCreateData) const
-{
-	TSharedPtr<FNiagaraScriptVarAndViewInfoAction> ScriptVarAndViewInfoAction = StaticCastSharedPtr<FNiagaraScriptVarAndViewInfoAction>(InCreateData->Action);
-	const FNiagaraScriptVariableAndViewInfo& ScriptVarAndViewInfo = ScriptVarAndViewInfoAction.Get()->ScriptVariableAndViewInfo;
-	TSharedPtr<FNiagaraParameterPanelEntryParameterNameViewModel> ParameterNameViewModel = MakeShared<FNiagaraParameterPanelEntryParameterNameViewModel>(InCreateData, ScriptVarAndViewInfo);
+const TArray<TArray<ENiagaraParameterPanelCategory>> FNiagaraScriptToolkitParameterPanelViewModel::DefaultCategoryPaths = {
+	{ENiagaraParameterPanelCategory::Input}
+	, {ENiagaraParameterPanelCategory::Attributes, ENiagaraParameterPanelCategory::User}
+	, {ENiagaraParameterPanelCategory::Attributes, ENiagaraParameterPanelCategory::Engine}
+	, {ENiagaraParameterPanelCategory::Attributes, ENiagaraParameterPanelCategory::Owner}
+	, {ENiagaraParameterPanelCategory::Attributes, ENiagaraParameterPanelCategory::System}
+	, {ENiagaraParameterPanelCategory::Attributes, ENiagaraParameterPanelCategory::Emitter}
+	, {ENiagaraParameterPanelCategory::Attributes, ENiagaraParameterPanelCategory::Particles}
+	, {ENiagaraParameterPanelCategory::Output, ENiagaraParameterPanelCategory::System}
+	, {ENiagaraParameterPanelCategory::Output, ENiagaraParameterPanelCategory::Emitter}
+	, {ENiagaraParameterPanelCategory::Output, ENiagaraParameterPanelCategory::Particles}
+	, {ENiagaraParameterPanelCategory::Output, ENiagaraParameterPanelCategory::ScriptTransient}
+	, {ENiagaraParameterPanelCategory::Local}
+};
+
+const TArray<TArray<ENiagaraParameterPanelCategory>> FNiagaraSystemToolkitParameterPanelViewModel::DefaultCategoryPaths = {
+	{ENiagaraParameterPanelCategory::User}
+	, {ENiagaraParameterPanelCategory::Engine}
+	, {ENiagaraParameterPanelCategory::Owner}
+	, {ENiagaraParameterPanelCategory::System}
+	, {ENiagaraParameterPanelCategory::Emitter}
+	, {ENiagaraParameterPanelCategory::Particles}
+};
+
+TSharedRef<class SWidget> INiagaraParameterPanelViewModel::GetScriptParameterVisualWidget(const FNiagaraScriptVariableAndViewInfo& ScriptVarAndViewInfo) const
+{
+	TSharedPtr<FNiagaraParameterPanelEntryParameterNameViewModel> ParameterNameViewModel = MakeShared<FNiagaraParameterPanelEntryParameterNameViewModel>(ScriptVarAndViewInfo);
 	ParameterNameViewModel->GetOnParameterRenamedDelegate().BindSP(this, &INiagaraParameterPanelViewModel::RenameParameter);
 	ParameterNameViewModel->GetOnScopeSelectionChangedDelegate().BindSP(this, &INiagaraParameterPanelViewModel::ChangeParameterScope);
 	ParameterNameViewModel->GetOnVerifyParameterRenamedDelegate().BindSP(this, &INiagaraParameterPanelViewModel::GetCanRenameParameterAndToolTip);
@@ -275,38 +106,6 @@ void FNiagaraSystemToolkitParameterPanelViewModel::Refresh() const
 	OnParameterPanelViewModelRefreshed.ExecuteIfBound();
 }
 
-void FNiagaraSystemToolkitParameterPanelViewModel::CollectStaticSections(TArray<int32>& StaticSectionIDs) const
-{
-	// Generic Emitter/System view, categorize by scope
-	StaticSectionIDs.Add(NiagaraParameterPanelSectionID::USER);
-	StaticSectionIDs.Add(NiagaraParameterPanelSectionID::ENGINE);
-	StaticSectionIDs.Add(NiagaraParameterPanelSectionID::OWNER);
-	StaticSectionIDs.Add(NiagaraParameterPanelSectionID::SYSTEM);
-	StaticSectionIDs.Add(NiagaraParameterPanelSectionID::EMITTER);
-	StaticSectionIDs.Add(NiagaraParameterPanelSectionID::PARTICLES);
-}
-
-NiagaraParameterPanelSectionID::Type FNiagaraSystemToolkitParameterPanelViewModel::GetSectionForVarAndViewInfo(const FNiagaraScriptVariableAndViewInfo& VarAndViewInfo) const
-{
-	// Generic Emitter/System view, categorize by scope
-
-	ENiagaraParameterScope ParameterScope = ENiagaraParameterScope::None;
-	FNiagaraEditorUtilities::GetVariableMetaDataScope(VarAndViewInfo.MetaData, ParameterScope);
-	if (VarAndViewInfo.MetaData.GetUsage() == ENiagaraScriptParameterUsage::Input && ParameterScope != ENiagaraParameterScope::Input)
-	{
-		return NiagaraParameterPanelSectionID::GetSectionForScope(ParameterScope);
-	}
-	else
-	{
-		//return NiagaraParameterPanelSectionID::GetSectionForScope(VarAndViewInfo.ScriptVariableMetaData.LastPrecompileScope); //@todo(ng) cache the known scope in the niagarasystem
-		FNiagaraVariableMetaData OutMetaData;
-		FNiagaraEditorUtilities::GetParameterMetaDataFromName(VarAndViewInfo.ScriptVariable.GetName(), OutMetaData);
-		ENiagaraParameterScope OutMetaDataScope;
-		FNiagaraEditorUtilities::GetVariableMetaDataScope(OutMetaData, OutMetaDataScope);
-		return NiagaraParameterPanelSectionID::GetSectionForScope(OutMetaDataScope);
-	}
-}
-
 const UNiagaraScriptVariable* FNiagaraSystemToolkitParameterPanelViewModel::AddParameter(FNiagaraVariable& VariableToAdd, const FNiagaraVariableMetaData& InVariableMetaDataToAssign)
 {
 	FScopedTransaction AddParameter(LOCTEXT("AddParameter", "Add Parameter"));
@@ -327,12 +126,13 @@ const UNiagaraScriptVariable* FNiagaraSystemToolkitParameterPanelViewModel::AddP
 		UNiagaraSystemEditorData* SystemEditorData = CastChecked<UNiagaraSystemEditorData>(System->GetEditorData(), ECastCheckedType::NullChecked);
 		SystemEditorData->Modify();
 		bool bSuccess = FNiagaraEditorUtilities::AddParameter(VariableToAdd, System->GetExposedParameters(), *System, SystemEditorData->GetStackEditorData());
+		Refresh();
 	}
 	else if (NewScope == ENiagaraParameterScope::System)
 	{
 		UNiagaraGraph* Graph = SystemViewModel->GetSystemScriptViewModel()->GetGraphViewModel()->GetGraph();
 		Graph->Modify();
-		const UNiagaraScriptVariable* NewScriptVar = Graph->AddParameter(VariableToAdd, AddParameterOptions); //@todo(ng) verify
+		const UNiagaraScriptVariable* NewScriptVar = Graph->AddParameter(VariableToAdd, AddParameterOptions);
 		Refresh();
 		return NewScriptVar;
 	}
@@ -365,7 +165,7 @@ void FNiagaraSystemToolkitParameterPanelViewModel::DeleteParameter(const FNiagar
 	ENiagaraParameterScope TargetVariableScope = ENiagaraParameterScope::None;
 	FNiagaraEditorUtilities::GetVariableMetaDataScope(TargetVariableMetaData, TargetVariableScope);
 
-	if (TargetVariableMetaData.IsInputUsage() && TargetVariableScope == ENiagaraParameterScope::User) //@todo(ng) verify
+	if (TargetVariableMetaData.IsInputUsage() && TargetVariableScope == ENiagaraParameterScope::User)
 	{
 		SystemViewModel->GetSystem().Modify();
 		SystemViewModel->GetSystem().GetExposedParameters().RemoveParameter(TargetVariableToRemove);
@@ -452,17 +252,17 @@ bool FNiagaraSystemToolkitParameterPanelViewModel::GetCanRenameParameterAndToolT
 
 FReply FNiagaraSystemToolkitParameterPanelViewModel::HandleActionDragged(const TSharedPtr<FEdGraphSchemaAction>& InAction, const FPointerEvent& MouseEvent) const
 {
-	const FText TooltipFormat = LOCTEXT("Parameters", "Name: {0} \nType: {1}\nScope: {2}\nUsage: {3}"); //@todo(ng) move to static define
+	const FText TooltipFormat = LOCTEXT("Parameters", "Name: {0} \nType: {1}\nScope: {2}\nUsage: {3}");
 	const FNiagaraScriptVarAndViewInfoAction* ScriptVarAction = static_cast<FNiagaraScriptVarAndViewInfoAction*>(InAction.Get());
 	const FNiagaraScriptVariableAndViewInfo& ScriptVarAndViewInfo = ScriptVarAction->ScriptVariableAndViewInfo;
-	NiagaraParameterPanelSectionID::Type Section = NiagaraParameterPanelSectionID::GetSectionForParameterMetaData(ScriptVarAndViewInfo.MetaData);
+	//NiagaraParameterPanelSectionID::Type Section = NiagaraParameterPanelSectionID::GetSectionForParameterMetaData(ScriptVarAndViewInfo.MetaData);
 	const FNiagaraVariable& Var = ScriptVarAndViewInfo.ScriptVariable;
 	const FText Name = FText::FromName(Var.GetName());
 	const FText ScopeText = FText::FromName(ScriptVarAndViewInfo.MetaData.GetScopeName());
 	const FText UsageText = StaticEnum<ENiagaraScriptParameterUsage>()->GetDisplayNameTextByValue((int64)ScriptVarAndViewInfo.MetaData.GetUsage());
 	const FText Tooltip = FText::Format(TooltipFormat, FText::FromName(Var.GetName()), Var.GetType().GetNameText(), ScopeText, UsageText);
 
-	TSharedPtr<FNiagaraParameterAction> ParameterAction = MakeShared<FNiagaraParameterAction>(Var, FText::GetEmpty(), Name, Tooltip, 0, FText(), Section);
+	TSharedPtr<FNiagaraParameterAction> ParameterAction = MakeShared<FNiagaraParameterAction>(Var, FText::GetEmpty(), Name, Tooltip, 0, FText(), 0);
 	TSharedRef<FNiagaraParameterDragOperation> DragOperation = MakeShared<FNiagaraParameterDragOperation>(ParameterAction);
 	DragOperation->CurrentHoverText = InAction->GetMenuDescription();
 	DragOperation->SetupDefaults();
@@ -470,6 +270,56 @@ FReply FNiagaraSystemToolkitParameterPanelViewModel::HandleActionDragged(const T
 	return FReply::Handled().BeginDragDrop(DragOperation);
 }
 
+const TArray<TArray<ENiagaraParameterPanelCategory>>& FNiagaraSystemToolkitParameterPanelViewModel::GetDefaultCategoryPaths() const
+{
+	return FNiagaraSystemToolkitParameterPanelViewModel::DefaultCategoryPaths;
+}
+
+TArray<ENiagaraParameterPanelCategory> FNiagaraSystemToolkitParameterPanelViewModel::GetCategoriesForParameter(const FNiagaraScriptVariableAndViewInfo& ScriptVar) const
+{
+	TArray<ENiagaraParameterPanelCategory> Categories;
+	
+	if (ScriptVar.MetaData.GetIsStaticSwitch())
+	{
+		Categories.Add(ENiagaraParameterPanelCategory::StaticSwitch);
+		return Categories;
+	}
+
+	const FNiagaraParameterScopeInfo* ScopeInfo = FNiagaraEditorModule::FindParameterScopeInfo(ScriptVar.MetaData.GetScopeName());
+	if (ScopeInfo != nullptr)
+	{
+		const ENiagaraParameterScope ParameterScope = ScopeInfo->GetScope();
+		switch (ParameterScope) {
+		case ENiagaraParameterScope::User:
+			Categories.Add(ENiagaraParameterPanelCategory::User);
+			return Categories;
+		case ENiagaraParameterScope::Engine:
+			Categories.Add(ENiagaraParameterPanelCategory::Engine);
+			return Categories;
+		case ENiagaraParameterScope::Owner:
+			Categories.Add(ENiagaraParameterPanelCategory::Owner);
+			return Categories;
+		case ENiagaraParameterScope::System:
+			Categories.Add(ENiagaraParameterPanelCategory::System);
+			return Categories;
+		case ENiagaraParameterScope::Emitter:
+			Categories.Add(ENiagaraParameterPanelCategory::Emitter);
+			return Categories;
+		case ENiagaraParameterScope::Particles:
+			Categories.Add(ENiagaraParameterPanelCategory::Particles);
+			return Categories;
+		case ENiagaraParameterScope::ScriptTransient:
+			Categories.Add(ENiagaraParameterPanelCategory::ScriptTransient);
+			return Categories;
+		case ENiagaraParameterScope::Custom:
+			return Categories;
+		default:
+			ensureMsgf(false, TEXT("Unexpected parameter scope encountered when getting category for parameter panel entry!"));
+		}
+	}
+	return Categories;
+}
+
 TArray<TWeakObjectPtr<UNiagaraGraph>> FNiagaraSystemToolkitParameterPanelViewModel::GetEditableGraphs() const
 {
 	TArray<TWeakObjectPtr<UNiagaraGraph>> EditableGraphs;
@@ -616,20 +466,6 @@ void FNiagaraScriptToolkitParameterPanelViewModel::HandleGraphSubObjectSelection
 	OnParameterPanelViewModelExternalSelectionChanged.Broadcast(Obj);
 }
 
-void FNiagaraScriptToolkitParameterPanelViewModel::CollectStaticSections(TArray<int32>& StaticSectionIDs) const
-{
-	StaticSectionIDs.Add(NiagaraParameterPanelSectionID::INPUTS);
-	StaticSectionIDs.Add(NiagaraParameterPanelSectionID::REFERENCES);
-	StaticSectionIDs.Add(NiagaraParameterPanelSectionID::OUTPUTS);
-	StaticSectionIDs.Add(NiagaraParameterPanelSectionID::LOCALS);
-	StaticSectionIDs.Add(NiagaraParameterPanelSectionID::INITIALVALUES);
-}
-
-NiagaraParameterPanelSectionID::Type FNiagaraScriptToolkitParameterPanelViewModel::GetSectionForVarAndViewInfo(const FNiagaraScriptVariableAndViewInfo& VarAndViewInfo) const
-{
-	return NiagaraParameterPanelSectionID::GetSectionForParameterMetaData(VarAndViewInfo.MetaData);
-}
-
 const UNiagaraScriptVariable* FNiagaraScriptToolkitParameterPanelViewModel::AddParameter(FNiagaraVariable& VariableToAdd, const FNiagaraVariableMetaData& InVariableMetaDataToAssign)
 {
 	ENiagaraScriptParameterUsage InVariableUsage = InVariableMetaDataToAssign.GetUsage();
@@ -648,7 +484,9 @@ const UNiagaraScriptVariable* FNiagaraScriptToolkitParameterPanelViewModel::AddP
 
 		UNiagaraGraph* Graph = ScriptViewModel->GetGraphViewModel()->GetGraph();
 		Graph->Modify();
-		return Graph->AddParameter(VariableToAdd, AddParameterOptions);
+		const UNiagaraScriptVariable* NewScriptVariable = Graph->AddParameter(VariableToAdd, AddParameterOptions);
+		Refresh();
+		return NewScriptVariable;
 	}
 	return nullptr;
 }
@@ -835,6 +673,98 @@ FReply FNiagaraScriptToolkitParameterPanelViewModel::HandleActionDragged(const T
 	return FReply::Handled().BeginDragDrop(DragOperation);
 }
 
+const TArray<TArray<ENiagaraParameterPanelCategory>>& FNiagaraScriptToolkitParameterPanelViewModel::GetDefaultCategoryPaths() const
+{
+	return FNiagaraScriptToolkitParameterPanelViewModel::DefaultCategoryPaths;
+}
+
+TArray<ENiagaraParameterPanelCategory> FNiagaraScriptToolkitParameterPanelViewModel::GetCategoriesForParameter(const FNiagaraScriptVariableAndViewInfo& ScriptVar) const
+{
+	TArray<ENiagaraParameterPanelCategory> Categories;
+
+	if (ScriptVar.MetaData.GetIsStaticSwitch())
+	{
+		Categories.Add(ENiagaraParameterPanelCategory::Input);
+		Categories.Add(ENiagaraParameterPanelCategory::StaticSwitch);
+		return Categories;
+	}
+
+	const ENiagaraScriptParameterUsage& Usage = ScriptVar.MetaData.GetUsage();
+	if (Usage == ENiagaraScriptParameterUsage::Input || Usage == ENiagaraScriptParameterUsage::Output)
+	{
+		const FNiagaraParameterScopeInfo* ScopeInfo = FNiagaraEditorModule::FindParameterScopeInfo(ScriptVar.MetaData.GetScopeName());
+		if (ScopeInfo != nullptr)
+		{
+			if (Usage == ENiagaraScriptParameterUsage::Input)
+			{
+				if (ScopeInfo->GetNamespaceString() == PARAM_MAP_MODULE_STR)
+				{
+					Categories.Add(ENiagaraParameterPanelCategory::Input);
+					return Categories;
+				}
+				Categories.Add(ENiagaraParameterPanelCategory::Attributes);
+			}
+			else
+			{
+				Categories.Add(ENiagaraParameterPanelCategory::Output);
+			}
+			
+			const ENiagaraParameterScope& ParameterScope = ScopeInfo->GetScope();
+			switch (ParameterScope) {
+			case ENiagaraParameterScope::User:
+				Categories.Add(ENiagaraParameterPanelCategory::User);
+				return Categories;
+			case ENiagaraParameterScope::Engine:
+				Categories.Add(ENiagaraParameterPanelCategory::Engine);
+				return Categories;
+			case ENiagaraParameterScope::Owner:
+				Categories.Add(ENiagaraParameterPanelCategory::Owner);
+				return Categories;
+			case ENiagaraParameterScope::System:
+				Categories.Add(ENiagaraParameterPanelCategory::System);
+				return Categories;
+			case ENiagaraParameterScope::Emitter:
+				Categories.Add(ENiagaraParameterPanelCategory::Emitter);
+				return Categories;
+			case ENiagaraParameterScope::Particles:
+				Categories.Add(ENiagaraParameterPanelCategory::Particles);
+				return Categories;
+			case ENiagaraParameterScope::ScriptTransient:
+				Categories.Add(ENiagaraParameterPanelCategory::ScriptTransient);
+				return Categories;
+			case ENiagaraParameterScope::Output:
+				Categories.Add(ENiagaraParameterPanelCategory::Output);
+				return Categories;
+			default:
+				ensureMsgf(false, TEXT("Unexpected ENiagaraParameterScope encountered when getting categories for parameter panel entry!"));
+				return Categories;
+			}
+		}
+		else
+		{
+			ensureMsgf(false, TEXT("Unregistered ScopeName encountered when getting categories for parameter panel entry! Name: %s"), *ScriptVar.MetaData.GetScopeName().ToString());
+			return Categories;
+		}
+	}
+	else if (Usage == ENiagaraScriptParameterUsage::Local)
+	{
+		Categories.Add(ENiagaraParameterPanelCategory::Local);
+		return Categories;
+	}
+	else if (Usage == ENiagaraScriptParameterUsage::InputOutput)
+	{
+		ensureMsgf(false, TEXT("Encountered InputOutput usage entry when getting categories for parameter panel entry! InputOuput usage should be split to separate Input and Output usage entries!"));
+		return Categories;
+	}
+	else
+	{
+		ensureMsgf(false, TEXT("Illegal ENiagaraScriptParameterUsage encountered when getting categories for parameter panel entry!"));
+		return Categories;
+	}
+
+	return Categories;
+}
+
 TSharedRef<SWidget> FNiagaraScriptToolkitParameterPanelViewModel::GetScriptParameterVisualWidget(const UEdGraphPin* Pin) const
 {
 	const UEdGraphSchema_Niagara* Schema = GetDefault<UEdGraphSchema_Niagara>();
diff --git a/Engine/Plugins/FX/Niagara/Source/NiagaraEditor/Private/Widgets/SNiagaraGraphPinAdd.cpp b/Engine/Plugins/FX/Niagara/Source/NiagaraEditor/Private/Widgets/SNiagaraGraphPinAdd.cpp
--- a/Engine/Plugins/FX/Niagara/Source/NiagaraEditor/Private/Widgets/SNiagaraGraphPinAdd.cpp
+++ b/Engine/Plugins/FX/Niagara/Source/NiagaraEditor/Private/Widgets/SNiagaraGraphPinAdd.cpp
@@ -79,7 +79,7 @@ TSharedRef<SWidget> SNiagaraGraphPinAdd::OnGetAddButtonMenuContent()
 			TSharedRef<SNiagaraAddParameterMenu2> MenuWidget = SNew(SNiagaraAddParameterMenu2, Graphs)
 				.NewParameterScope(ENiagaraParameterScope::Local)
 				.NewParameterNamespace(Cast<UNiagaraNodeParameterMapBase>(OwningNode)->GetNewPinDefaultNamespace().ToString())
-				.ShowKnownConstantParametersFilter(NiagaraParameterPanelSectionID::Type::REFERENCES)
+				.ShowKnownConstantParametersFilter(ENiagaraParameterPanelCategory::Attributes)
 				.OnAddParameter_UObject(OwningNode, &UNiagaraNodeWithDynamicPins::AddParameter, (const UEdGraphPin*)GetPinObj())
 				.OnCollectCustomActions_UObject(OwningNode, &UNiagaraNodeWithDynamicPins::CollectAddPinActions, GetPinObj())
 				.OnAllowMakeType_UObject(OwningNode, &UNiagaraNodeWithDynamicPins::AllowNiagaraTypeForAddPin)
diff --git a/Engine/Plugins/FX/Niagara/Source/NiagaraEditor/Private/Widgets/SNiagaraParameterPanel.cpp b/Engine/Plugins/FX/Niagara/Source/NiagaraEditor/Private/Widgets/SNiagaraParameterPanel.cpp
--- a/Engine/Plugins/FX/Niagara/Source/NiagaraEditor/Private/Widgets/SNiagaraParameterPanel.cpp
+++ b/Engine/Plugins/FX/Niagara/Source/NiagaraEditor/Private/Widgets/SNiagaraParameterPanel.cpp
@@ -17,6 +17,9 @@
 #include "Editor/GraphEditor/Private/GraphActionNode.h"
 #include "NiagaraConstants.h"
 #include "Framework/MultiBox/MultiBoxBuilder.h"
+#include "Widgets/SItemSelector.h"
+#include "NiagaraEditorModule.h"
+#include "NiagaraTypes.h"
 
 #define LOCTEXT_NAMESPACE "NiagaraParameterPanel"
 
@@ -25,8 +28,8 @@ SNiagaraParameterPanel::~SNiagaraParameterPanel()
 	ParameterPanelViewModel->GetOnRefreshed().Unbind();
 
 	// Unregister all commands for right click on action node
-	ToolkitCommands->UnmapAction(FNiagaraParameterPanelCommands::Get().DeleteEntry);
-	ToolkitCommands->UnmapAction(FGenericCommands::Get().Rename);
+// 	ToolkitCommands->UnmapAction(FNiagaraParameterPanelCommands::Get().DeleteEntry);
+// 	ToolkitCommands->UnmapAction(FGenericCommands::Get().Rename);
 }
 
 void FNiagaraParameterPanelCommands::RegisterCommands()
@@ -36,262 +39,225 @@ void FNiagaraParameterPanelCommands::RegisterCommands()
 
 void SNiagaraParameterPanel::Construct(const FArguments& InArgs, const TSharedPtr<INiagaraParameterPanelViewModel>& InParameterPanelViewModel, const TSharedPtr<FUICommandList>& InToolkitCommands)
 {
-	bNeedsRefresh = false;
-	bGraphActionPendingRename = false;
-
 	ParameterPanelViewModel = InParameterPanelViewModel;
-	ToolkitCommands = InToolkitCommands;
+	ToolkitCommands = InToolkitCommands; //@todo(ng) verify
 
 	ParameterPanelViewModel->GetOnRefreshed().BindRaw(this, &SNiagaraParameterPanel::Refresh);
 	ParameterPanelViewModel->GetExternalSelectionChanged().AddRaw(this, &SNiagaraParameterPanel::HandleExternalSelectionChanged);
 
-	AddParameterButtons.SetNum(NiagaraParameterPanelSectionID::CUSTOM + 1); //@Todo(ng) verify
-
-	// Register all commands for right click on action node
-	FNiagaraParameterPanelCommands::Register();
-	TSharedPtr<FUICommandList> ToolKitCommandList = ToolkitCommands;
-	ToolKitCommandList->MapAction(FNiagaraParameterPanelCommands::Get().DeleteEntry,
-		FExecuteAction::CreateSP(this, &SNiagaraParameterPanel::TryDeleteEntries),
-		FCanExecuteAction::CreateSP(this, &SNiagaraParameterPanel::CanTryDeleteEntries));
-	ToolKitCommandList->MapAction(FGenericCommands::Get().Rename,
-		FExecuteAction::CreateSP(this, &SNiagaraParameterPanel::OnRequestRenameOnActionNode),
-		FCanExecuteAction::CreateSP(this, &SNiagaraParameterPanel::CanRequestRenameOnActionNode));
-
-	// create the main action list piece of this widget
-	SAssignNew(GraphActionMenu, SGraphActionMenu, false)
-		.OnCreateWidgetForAction(this, &SNiagaraParameterPanel::OnCreateWidgetForAction)
-		.OnCollectAllActions(this, &SNiagaraParameterPanel::CollectAllActions)
- 		.OnCollectStaticSections(this, &SNiagaraParameterPanel::CollectStaticSections)
- 		.OnActionDragged(this, &SNiagaraParameterPanel::OnActionDragged)
-		.OnActionSelected(this, &SNiagaraParameterPanel::OnActionSelected)
-// 		.OnActionDoubleClicked(this, &SNiagaraParameterMapView::OnActionDoubleClicked) //@todo(ng) impl
-		.OnContextMenuOpening(this, &SNiagaraParameterPanel::OnContextMenuOpening)
-		.OnCanRenameSelectedAction(this, &SNiagaraParameterPanel::CanRequestRenameOnActionNode)
- 		.OnGetSectionTitle(this, &SNiagaraParameterPanel::OnGetSectionTitle)
- 		.OnGetSectionWidget(this, &SNiagaraParameterPanel::OnGetSectionWidget)
- 		.OnCreateCustomRowExpander_Static(&SNiagaraParameterPanel::CreateCustomActionExpander)
-		.OnActionMatchesName(this, &SNiagaraParameterPanel::HandleActionMatchesName)
-		.AutoExpandActionMenu(false)
-		.AlphaSortItems(false)
-		.UseSectionStyling(true)
-		.ShowFilterTextBox(true);
+	bNeedsRefresh = false;
 
 	ChildSlot
 	[
 		SNew(SBox)
 		.MinDesiredWidth(300)
 		[
-			SNew(SVerticalBox)	
-			+ SVerticalBox::Slot()
-			.FillHeight(1.0f)
-			[
-				GraphActionMenu.ToSharedRef()
-			]
+			SAssignNew(ItemSelector, SNiagaraParameterPanelSelector)
+			.Items(ParameterPanelViewModel->GetViewedParameters())
+			.OnGetCategoriesForItem(this, &SNiagaraParameterPanel::OnGetCategoriesForItem)
+			.OnCompareCategoriesForEquality(this, &SNiagaraParameterPanel::OnCompareCategoriesForEquality)
+			.OnCompareCategoriesForSorting(this, &SNiagaraParameterPanel::OnCompareCategoriesForSorting)
+			.OnCompareItemsForEquality(this, &SNiagaraParameterPanel::OnCompareItemsForEquality)
+			.OnCompareItemsForSorting(this, &SNiagaraParameterPanel::OnCompareItemsForSorting)
+			.OnDoesItemMatchFilterText(this, &SNiagaraParameterPanel::OnDoesItemMatchFilterText)
+			.OnGenerateWidgetForCategory(this, &SNiagaraParameterPanel::OnGenerateWidgetForCategory)
+			.OnGenerateWidgetForItem(this, &SNiagaraParameterPanel::OnGenerateWidgetForItem)
+			.AllowMultiselect(false)
+			.DefaultCategoryPaths(GetDefaultCategoryPaths())
+			.ClearSelectionOnClick(true)
 		]
 	];
+
+// 	// Register all commands for right click on action node
+// 	FNiagaraParameterPanelCommands::Register();
+// 	TSharedPtr<FUICommandList> ToolKitCommandList = ToolkitCommands;
+// 	ToolKitCommandList->MapAction(FNiagaraParameterPanelCommands::Get().DeleteEntry,
+// 		FExecuteAction::CreateSP(this, &SNiagaraParameterPanel::TryDeleteEntries),
+// 		FCanExecuteAction::CreateSP(this, &SNiagaraParameterPanel::CanTryDeleteEntries));
+// 	ToolKitCommandList->MapAction(FGenericCommands::Get().Rename,
+// 		FExecuteAction::CreateSP(this, &SNiagaraParameterPanel::OnRequestRenameOnActionNode),
+// 		FCanExecuteAction::CreateSP(this, &SNiagaraParameterPanel::CanRequestRenameOnActionNode));
+
+
+
+// 	// create the main action list piece of this widget
+// 	SAssignNew(GraphActionMenu, SGraphActionMenu, false)
+// 		.OnCreateWidgetForAction(this, &SNiagaraParameterPanel::OnCreateWidgetForAction)
+// 		.OnCollectAllActions(this, &SNiagaraParameterPanel::CollectAllActions)
+//  		.OnCollectStaticSections(this, &SNiagaraParameterPanel::CollectStaticSections)
+//  		.OnActionDragged(this, &SNiagaraParameterPanel::OnActionDragged)
+// 		.OnActionSelected(this, &SNiagaraParameterPanel::OnActionSelected)
+// // 		.OnActionDoubleClicked(this, &SNiagaraParameterMapView::OnActionDoubleClicked)
+// 		.OnContextMenuOpening(this, &SNiagaraParameterPanel::OnContextMenuOpening)
+// 		.OnCanRenameSelectedAction(this, &SNiagaraParameterPanel::CanRequestRenameOnActionNode)
+//  		.OnGetSectionTitle(this, &SNiagaraParameterPanel::OnGetSectionTitle)
+//  		.OnGetSectionWidget(this, &SNiagaraParameterPanel::OnGetSectionWidget)
+//  		.OnCreateCustomRowExpander_Static(&SNiagaraParameterPanel::CreateCustomActionExpander)
+// 		.OnActionMatchesName(this, &SNiagaraParameterPanel::HandleActionMatchesName)
+// 		.AutoExpandActionMenu(false)
+// 		.AlphaSortItems(false)
+// 		.UseSectionStyling(true)
+// 		.ShowFilterTextBox(true);
+// 
+// 	ChildSlot
+// 	[
+// 		SNew(SBox)
+// 		.MinDesiredWidth(300)
+// 		[
+// 			SNew(SVerticalBox)	
+// 			+ SVerticalBox::Slot()
+// 			.FillHeight(1.0f)
+// 			[
+// 				GraphActionMenu.ToSharedRef()
+// 			]
+// 		]
+// 	];
 }
 
-void SNiagaraParameterPanel::Tick(const FGeometry& AllottedGeometry, const double InCurrentTime, const float InDeltaTime)
-{
-	if (bNeedsRefresh)
-	{
-		GraphActionMenu->RefreshAllActions(true);
-		bNeedsRefresh = false;
-	}
 
-	if (bGraphActionPendingRename)
-	{
-		if (GraphActionMenu->CanRequestRenameOnActionNode())
-		{
-			GraphActionMenu->OnRequestRenameOnActionNode();
-			bGraphActionPendingRename = false;
-		}
-	}
+TArray<ENiagaraParameterPanelCategory> SNiagaraParameterPanel::OnGetCategoriesForItem(const FNiagaraScriptVariableAndViewInfo& Item)
+{
+	return ParameterPanelViewModel->GetCategoriesForParameter(Item);
 }
 
-void SNiagaraParameterPanel::HandleExternalSelectionChanged(const UObject* Obj)
+bool SNiagaraParameterPanel::OnCompareCategoriesForEquality(const ENiagaraParameterPanelCategory& CategoryA, const ENiagaraParameterPanelCategory& CategoryB) const
 {
-	if (Obj && Obj->IsA< UNiagaraScriptVariable>())
-	{
-		const UNiagaraScriptVariable* Var = Cast< UNiagaraScriptVariable>(Obj);
-		if (Var)
-		{
-			GraphActionMenu->SelectItemByName(Var->Variable.GetName());
-		}
-	}
+	return CategoryA == CategoryB;
 }
 
-void SNiagaraParameterPanel::AddParameter(FNiagaraVariable NewVariable, const NiagaraParameterPanelSectionID::Type SectionID)
+bool SNiagaraParameterPanel::OnCompareCategoriesForSorting(const ENiagaraParameterPanelCategory& CategoryA, const ENiagaraParameterPanelCategory& CategoryB) const
 {
+	return CategoryA < CategoryB;
+}
 
-	FNiagaraVariableMetaData GuessedMetaData;
-	FNiagaraEditorUtilities::GetParameterMetaDataFromName(NewVariable.GetName(), GuessedMetaData);
-	const UNiagaraScriptVariable* NewScriptVar = ParameterPanelViewModel->AddParameter(NewVariable, GuessedMetaData);
-	
-	if (NewScriptVar != nullptr)
-	{
-		GraphActionMenu->RefreshAllActions(true);
-		GraphActionMenu->SelectItemByName(NewScriptVar->Variable.GetName());
-
-		// Delay calling the rename delegate one tick as the widget to which it is bound is not yet constructed.
-		if (FNiagaraConstants::IsNiagaraConstant(NewVariable) == false)
-			bGraphActionPendingRename = true;
-	}
+bool SNiagaraParameterPanel::OnCompareItemsForEquality(const FNiagaraScriptVariableAndViewInfo& ItemA, const FNiagaraScriptVariableAndViewInfo& ItemB) const
+{
+	return ItemA == ItemB;
 }
 
-bool SNiagaraParameterPanel::AllowMakeType(const FNiagaraTypeDefinition& InType) const
+bool SNiagaraParameterPanel::OnCompareItemsForSorting(const FNiagaraScriptVariableAndViewInfo& ItemA, const FNiagaraScriptVariableAndViewInfo& ItemB) const
 {
-	return InType != FNiagaraTypeDefinition::GetParameterMapDef() && InType != FNiagaraTypeDefinition::GetGenericNumericDef();
+	if (ItemA.MetaData.GetIsUsingLegacyNameString())
+	{
+		if (ItemB.MetaData.GetIsUsingLegacyNameString())
+		{
+			return ItemA.ScriptVariable.GetName().LexicalLess(ItemB.ScriptVariable.GetName());
+		}
+		return false;
+	}
+	else if (ItemB.MetaData.GetIsUsingLegacyNameString())
+	{
+		return true;
+	}
+	FName AName;
+	FName BName;
+	ItemA.MetaData.GetParameterName(AName);
+	ItemB.MetaData.GetParameterName(BName);
+	return AName.LexicalLess(BName);
 }
 
-TSharedRef<SWidget> SNiagaraParameterPanel::OnCreateWidgetForAction(struct FCreateWidgetForActionData* const InCreateData)
+bool SNiagaraParameterPanel::OnDoesItemMatchFilterText(const FText& FilterText, const FNiagaraScriptVariableAndViewInfo& Item)
 {
-	TSharedPtr<SWidget> ParameterVisualWidget = ParameterPanelViewModel->GetScriptParameterVisualWidget(InCreateData);
-	return SNew(SNiagaraParameterPanelPaletteItem, InCreateData, ParameterVisualWidget->AsShared());
+	return Item.ScriptVariable.GetName().ToString().Contains(FilterText.ToString());
 }
 
-void SNiagaraParameterPanel::CollectAllActions(FGraphActionListBuilderBase& OutAllActions)
+TSharedRef<SWidget> SNiagaraParameterPanel::OnGenerateWidgetForCategory(const ENiagaraParameterPanelCategory& Category)
 {
-	const FText TooltipFormat = LOCTEXT("Parameters", "Name: {0} \nType: {1}"); //@todo(ng) move to const static
+	const UEnum* ParameterPanelCategoryEnum = FNiagaraTypeDefinition::GetParameterPanelCategoryEnum();
 
-	auto SortParameterPanelEntries = [](const FNiagaraScriptVariableAndViewInfo& A, const FNiagaraScriptVariableAndViewInfo& B) {
-		if (A.MetaData.GetIsUsingLegacyNameString())
-		{
-			if (B.MetaData.GetIsUsingLegacyNameString())
-			{
-				return A.ScriptVariable.GetName().LexicalLess(B.ScriptVariable.GetName());
-			}
-			return false;
-		}
-		else if (B.MetaData.GetIsUsingLegacyNameString())
-		{
-			return true;
-		}
-		FName AName;
-		FName BName;
-		A.MetaData.GetParameterName(AName);
-		B.MetaData.GetParameterName(BName);
-		return AName.LexicalLess(BName);
-	};
+	TSharedRef<SHorizontalBox> ParameterPanelCategoryHorizontalBox = SNew(SHorizontalBox);
+	ParameterPanelCategoryHorizontalBox->AddSlot()
+		.VAlign(VAlign_Center)
+		.Padding(3, 0, 0, 0)
+		[
+			SNew(STextBlock)
+			.TextStyle(FEditorStyle::Get(), "DetailsView.CategoryTextStyle")
+		.Text(ParameterPanelCategoryEnum->GetDisplayNameTextByValue((const int64)Category))
+		];
 
-	TArray<FNiagaraScriptVariableAndViewInfo> VarAndViewInfos = ParameterPanelViewModel->GetViewedParameters();
-	VarAndViewInfos.Sort(SortParameterPanelEntries);
-	//UE_LOG(LogNiagaraEditor, Log, TEXT("Adding Vars--------------------------------------------------"));
-	TArray<FNiagaraVariable> AddedVars;
-	for (const FNiagaraScriptVariableAndViewInfo& VarAndViewInfo : VarAndViewInfos)
+	if (GetCanAddParametersToCategory(Category))
 	{
-		int32 FoundVarIndex = AddedVars.Find(VarAndViewInfo.ScriptVariable);
-		if (FoundVarIndex < 0)
-		{
-			const FText Name = FText::FromName(VarAndViewInfo.ScriptVariable.GetName());
-			const FText Tooltip = FText::Format(TooltipFormat, FText::FromName(VarAndViewInfo.ScriptVariable.GetName()), VarAndViewInfo.ScriptVariable.GetType().GetNameText());
-			const NiagaraParameterPanelSectionID::Type Section = ParameterPanelViewModel->GetSectionForVarAndViewInfo(VarAndViewInfo);
-			TSharedPtr<FNiagaraScriptVarAndViewInfoAction> ScriptVarAndViewInfoAction(new FNiagaraScriptVarAndViewInfoAction(VarAndViewInfo, FText::GetEmpty(), Name, Tooltip, 0, FText(), Section /*= 0*/));
-			OutAllActions.AddAction(ScriptVarAndViewInfoAction);
-			AddedVars.Add(VarAndViewInfo.ScriptVariable);
-
-			//UE_LOG(LogNiagaraEditor, Log, TEXT("%s"), *Name.ToString());
-		}
+		FText AddNewText = LOCTEXT("AddNewParameter", "Add Parameter");
+		FName MetaDataTag = TEXT("AddNewParameter");
 
+		ParameterPanelCategoryHorizontalBox->AddSlot()
+			.AutoWidth()
+			.Padding(0.0f, 4.0f, 3.0f, 4.0f)
+			[
+				CreateAddToCategoryButton(Category, AddNewText, MetaDataTag)
+			];
 	}
+
+	return ParameterPanelCategoryHorizontalBox;
 }
- 
-void SNiagaraParameterPanel::CollectStaticSections(TArray<int32>& StaticSectionIDs) const
+
+TSharedRef<SWidget> SNiagaraParameterPanel::OnGenerateWidgetForItem(const FNiagaraScriptVariableAndViewInfo& Item)
 {
-	return ParameterPanelViewModel->CollectStaticSections(StaticSectionIDs);
+	TSharedRef<SWidget> ParameterVisualWidget = ParameterPanelViewModel->GetScriptParameterVisualWidget(Item);
+	return SNew(SNiagaraParameterPanelPaletteItem, ParameterVisualWidget);
 }
 
-FReply SNiagaraParameterPanel::OnActionDragged(const TArray<TSharedPtr<FEdGraphSchemaAction>>& InActions, const FPointerEvent& MouseEvent)
+const TArray<TArray<ENiagaraParameterPanelCategory>>& SNiagaraParameterPanel::GetDefaultCategoryPaths() const
 {
-	if (InActions.Num() == 1 && InActions[0].IsValid())
-	{
-		return ParameterPanelViewModel->HandleActionDragged(InActions[0], MouseEvent);
-	}
-	return FReply::Handled();
+	return ParameterPanelViewModel->GetDefaultCategoryPaths();
 }
 
-void SNiagaraParameterPanel::OnActionSelected(const TArray< TSharedPtr<FEdGraphSchemaAction> >& InActions, ESelectInfo::Type InSelectionType)
+void SNiagaraParameterPanel::Tick(const FGeometry& AllottedGeometry, const double InCurrentTime, const float InDeltaTime)
 {
-	if (InActions.Num() == 1 && InActions[0].IsValid())
+	if (bNeedsRefresh)
 	{
-		ParameterPanelViewModel->HandleActionSelected(InActions[0], InSelectionType);
+		ItemSelector->RefreshItemsAndDefaultCategories(ParameterPanelViewModel->GetViewedParameters(), ParameterPanelViewModel->GetDefaultCategoryPaths());
+		bNeedsRefresh = false;
 	}
+
+// 	if (bGraphActionPendingRename) //@todo(ng) impl
+// 	{
+// 		if (GraphActionMenu->CanRequestRenameOnActionNode())
+// 		{
+// 			GraphActionMenu->OnRequestRenameOnActionNode();
+// 			bGraphActionPendingRename = false;
+// 		}
+// 	}
 }
 
-TSharedPtr<SWidget> SNiagaraParameterPanel::OnContextMenuOpening()
+void SNiagaraParameterPanel::HandleExternalSelectionChanged(const UObject* Obj)
 {
-	// Get the current selected action
-	TArray<TSharedPtr<FEdGraphSchemaAction> > SelectedActions;
-	GraphActionMenu->GetSelectedActions(SelectedActions);
-	ensureMsgf(SelectedActions.Num() == 1 || SelectedActions.Num() == 0, TEXT("Unexpected number of selected actions encountered when getting current selected action for parameter panel!"));
-	if (SelectedActions.Num() == 1 && SelectedActions[0].IsValid())
+	if (Obj && Obj->IsA< UNiagaraScriptVariable>())
 	{
-		TSharedPtr<FNiagaraScriptVarAndViewInfoAction> CurrentAction = StaticCastSharedPtr<FNiagaraScriptVarAndViewInfoAction>(SelectedActions[0]);
-		if (ensureMsgf(CurrentAction.Get(), TEXT("Action pointer was null when getting selected action for parameter panel!")))
+		const UNiagaraScriptVariable* Var = Cast< UNiagaraScriptVariable>(Obj);
+		if (Var)
 		{
-			const bool bShouldCloseWindowAfterMenuSelection = true;
-			FMenuBuilder MenuBuilder(bShouldCloseWindowAfterMenuSelection, ToolkitCommands);
-			MenuBuilder.BeginSection("BasicOperations");
-
-			const FNiagaraVariable& CurrentParameter = CurrentAction->ScriptVariableAndViewInfo.ScriptVariable;
-			const FNiagaraVariableMetaData& CurrentParameterMetaData = CurrentAction->ScriptVariableAndViewInfo.MetaData;
-
-			FText CanRenameParameterToolTip;
-			bool bCanRequestRenameParameter = ParameterPanelViewModel->GetCanRenameParameterAndToolTip(CurrentParameter, CurrentParameterMetaData, TOptional<const FText>(), CanRenameParameterToolTip);
-			MenuBuilder.AddMenuEntry(
-				LOCTEXT("ParameterPanelRenameParameterContextMenu", "Rename Parameter"),
-				CanRenameParameterToolTip,
-				FSlateIcon(),
-				FUIAction(
-					FExecuteAction::CreateLambda([this]()
-						{
-							OnRequestRenameOnActionNode();
-						}),
-					FCanExecuteAction::CreateLambda([bCanRequestRenameParameter]() {return bCanRequestRenameParameter; }))
-			);
-
-			FText CanDeleteParameterToolTip;
-			bool bCanDeleteParameter = ParameterPanelViewModel->GetCanDeleteParameterAndToolTip(CurrentParameter, CurrentParameterMetaData, CanDeleteParameterToolTip);
-			MenuBuilder.AddMenuEntry(
-				LOCTEXT("ParameterPanelDeleteParameterContextMenu", "Delete Parameter"),
-				CanDeleteParameterToolTip,
-				FSlateIcon(),
-				FUIAction(
-					FExecuteAction::CreateLambda([this, CurrentParameter, CurrentParameterMetaData]()
-						{
-							ParameterPanelViewModel->DeleteParameter(CurrentParameter, CurrentParameterMetaData);
-						}),
-					FCanExecuteAction::CreateLambda([bCanDeleteParameter]() {return bCanDeleteParameter; }))
-			);
-
-			MenuBuilder.EndSection();
-			return MenuBuilder.MakeWidget();
+			SelectParameterEntryByName(Var->Variable.GetName());
 		}
 	}
-	return SNullWidget::NullWidget;
 }
 
-FText SNiagaraParameterPanel::OnGetSectionTitle(int32 InSectionID)
+void SNiagaraParameterPanel::AddParameter(FNiagaraVariable NewVariable, const ENiagaraParameterPanelCategory Category)
 {
-	return NiagaraParameterPanelSectionID::OnGetSectionTitle((NiagaraParameterPanelSectionID::Type)InSectionID);
+	FNiagaraVariableMetaData GuessedMetaData;
+	FNiagaraEditorUtilities::GetParameterMetaDataFromName(NewVariable.GetName(), GuessedMetaData);
+	const UNiagaraScriptVariable* NewScriptVar = ParameterPanelViewModel->AddParameter(NewVariable, GuessedMetaData);
+	
+	if (NewScriptVar != nullptr)
+	{
+		const FName ScriptVarName = NewScriptVar->Variable.GetName();
+		ensureMsgf(SelectParameterEntryByName(ScriptVarName), TEXT("Failed to find newly created UNiagaraScriptVariable to select!"));
+	}
 }
 
-TSharedRef<SWidget> SNiagaraParameterPanel::OnGetSectionWidget(TSharedRef<SWidget> RowWidget, int32 InSectionID)
+bool SNiagaraParameterPanel::AllowMakeType(const FNiagaraTypeDefinition& InType) const
 {
-	TWeakPtr<SWidget> WeakRowWidget = RowWidget;
-	FText AddNewText = LOCTEXT("AddNewParameter", "Add Parameter");
-	FName MetaDataTag = TEXT("AddNewParameter");
-	return CreateAddToSectionButton((NiagaraParameterPanelSectionID::Type) InSectionID, WeakRowWidget, AddNewText, MetaDataTag); //@todo(ng) refactor to get rid of int32 insectionid
+	return InType != FNiagaraTypeDefinition::GetParameterMapDef() && InType != FNiagaraTypeDefinition::GetGenericNumericDef();
 }
 
-TSharedRef<SWidget> SNiagaraParameterPanel::CreateAddToSectionButton(const NiagaraParameterPanelSectionID::Type InSection, TWeakPtr<SWidget> WeakRowWidget, FText AddNewText, FName MetaDataTag)
+TSharedRef<SWidget> SNiagaraParameterPanel::CreateAddToCategoryButton(const ENiagaraParameterPanelCategory Category, FText AddNewText, FName MetaDataTag)
 {
 	TSharedPtr<SComboButton> Button;
 	SAssignNew(Button, SComboButton)
 	.ButtonStyle(FEditorStyle::Get(), "RoundButton")
 	.ForegroundColor(FEditorStyle::GetSlateColor("DefaultForeground"))
 	.ContentPadding(FMargin(2, 0))
-	.OnGetMenuContent(this, &SNiagaraParameterPanel::OnGetParameterMenu, InSection)
+	.OnGetMenuContent(this, &SNiagaraParameterPanel::OnGetParameterMenu, Category)
 	//.IsEnabled(this, &SNiagaraParameterMapView::ParameterAddEnabled) //@Todo(ng) impl
 	.HAlign(HAlign_Center)
 	.VAlign(VAlign_Center)
@@ -310,120 +276,182 @@ TSharedRef<SWidget> SNiagaraParameterPanel::CreateAddToSectionButton(const Niaga
 
 	];
 
-	AddParameterButtons[(int32)InSection] = Button;
+	//AddParameterButtons[(int32)InSection] = Button;
 
 	return Button.ToSharedRef();
 }
 
-TSharedRef<SWidget> SNiagaraParameterPanel::OnGetParameterMenu(const NiagaraParameterPanelSectionID::Type InSection)
+TSharedRef<SWidget> SNiagaraParameterPanel::OnGetParameterMenu(const ENiagaraParameterPanelCategory Category)
 {
-	ENiagaraParameterScope NewParameterScopeForSection = NiagaraParameterPanelSectionID::GetScopeForNewParametersInSection(InSection);
+	ENiagaraParameterScope NewParameterScopeForSection = GetScopeForNewParametersInCategory(Category);
+
+	const bool bCanCreateNew = Category != ENiagaraParameterPanelCategory::Engine && Category != ENiagaraParameterPanelCategory::Owner;
+	const bool bAutoExpand = Category == ENiagaraParameterPanelCategory::Local || Category == ENiagaraParameterPanelCategory::Input ||
+		Category == ENiagaraParameterPanelCategory::Output || Category == ENiagaraParameterPanelCategory::User || 
+		Category == ENiagaraParameterPanelCategory::Engine || Category == ENiagaraParameterPanelCategory::Owner;
 
-	const bool bCanCreateNew = InSection != NiagaraParameterPanelSectionID::Type::ENGINE && InSection != NiagaraParameterPanelSectionID::Type::OWNER;
-	const bool bAutoExpand = InSection == NiagaraParameterPanelSectionID::Type::LOCALS || InSection == NiagaraParameterPanelSectionID::Type::INPUTS ||
-		InSection == NiagaraParameterPanelSectionID::Type::OUTPUTS || InSection == NiagaraParameterPanelSectionID::Type::USER || 
-		InSection == NiagaraParameterPanelSectionID::Type::ENGINE || InSection == NiagaraParameterPanelSectionID::Type::OWNER;
 	TSharedRef<SNiagaraAddParameterMenu2> MenuWidget = SNew(SNiagaraAddParameterMenu2, ParameterPanelViewModel->GetEditableGraphs())
-		.OnAddParameter(this, &SNiagaraParameterPanel::AddParameter, InSection)
+		.OnAddParameter(this, &SNiagaraParameterPanel::AddParameter, Category)
 		.OnAllowMakeType(this, &SNiagaraParameterPanel::AllowMakeType)
 		.ShowGraphParameters(false)
-		.ShowKnownConstantParametersFilter(InSection)
+		.ShowKnownConstantParametersFilter(Category)
 		.AllowCreatingNew(bCanCreateNew)
 		.AutoExpandMenu(bAutoExpand)
 		.NewParameterScope(NewParameterScopeForSection);
 
-	AddParameterButtons[(int32)InSection]->SetMenuContentWidgetToFocus(MenuWidget->GetSearchBox()->AsShared());
+	//AddParameterButtons[(int32)InSection]->SetMenuContentWidgetToFocus(MenuWidget->GetSearchBox()->AsShared());
 	return MenuWidget;
 }
 
-void SNiagaraParameterPanel::TryDeleteEntries()
-{
-	TArray<TSharedPtr<FEdGraphSchemaAction>> SelectedActions;
-	GraphActionMenu->GetSelectedActions(SelectedActions);
-
-	for (auto& Action : SelectedActions)
-	{
-		TSharedPtr<FNiagaraScriptVarAndViewInfoAction> ScriptVarAndViewInfoAction = StaticCastSharedPtr<FNiagaraScriptVarAndViewInfoAction>(Action);
-		if (ScriptVarAndViewInfoAction.Get())
-		{
-			const FNiagaraVariable& Parameter = ScriptVarAndViewInfoAction->ScriptVariableAndViewInfo.ScriptVariable;
-			const FNiagaraVariableMetaData& ParameterMetaData = ScriptVarAndViewInfoAction->ScriptVariableAndViewInfo.MetaData;
-			FText ToolTipText;
-			if (ParameterPanelViewModel->GetCanDeleteParameterAndToolTip(Parameter, ParameterMetaData, ToolTipText))
-			{
-				ParameterPanelViewModel->DeleteParameter(Parameter, ParameterMetaData);
-			}
-		}
-	}
-}
-
-bool SNiagaraParameterPanel::CanTryDeleteEntries() const
+ENiagaraParameterScope SNiagaraParameterPanel::GetScopeForNewParametersInCategory(const ENiagaraParameterPanelCategory Category)
 {
-	TArray<TSharedPtr<FEdGraphSchemaAction>> SelectedActions;
-	GraphActionMenu->GetSelectedActions(SelectedActions);
-
-	for (auto& Action : SelectedActions)
+	if (ensureMsgf(GetCanAddParametersToCategory(Category), TEXT("Cannot add parameters to category which we are trying to get scope for!")))
 	{
-		TSharedPtr<FNiagaraScriptVarAndViewInfoAction> ScriptVarAndViewInfoAction = StaticCastSharedPtr<FNiagaraScriptVarAndViewInfoAction>(Action);
-		if (ScriptVarAndViewInfoAction.Get())
-		{
-			const FNiagaraVariable& Parameter = ScriptVarAndViewInfoAction->ScriptVariableAndViewInfo.ScriptVariable;
-			const FNiagaraVariableMetaData& ParameterMetaData = ScriptVarAndViewInfoAction->ScriptVariableAndViewInfo.MetaData;
-			FText ToolTipText;
-			if (ParameterPanelViewModel->GetCanDeleteParameterAndToolTip(Parameter, ParameterMetaData, ToolTipText))
-			{
-				return true;
-			}
-		}
+		switch (Category) {
+		case ENiagaraParameterPanelCategory::User:
+			return ENiagaraParameterScope::User;
+		case ENiagaraParameterPanelCategory::Engine:
+			return ENiagaraParameterScope::Engine;
+		case ENiagaraParameterPanelCategory::Owner:
+			return ENiagaraParameterScope::Owner;
+		case ENiagaraParameterPanelCategory::System:
+			return ENiagaraParameterScope::System;
+		case ENiagaraParameterPanelCategory::Emitter:
+			return ENiagaraParameterScope::Emitter;
+		case ENiagaraParameterPanelCategory::Particles:
+			return ENiagaraParameterScope::Particles;
+		case ENiagaraParameterPanelCategory::Input:
+			return ENiagaraParameterScope::Input;
+		case ENiagaraParameterPanelCategory::Local:
+			return ENiagaraParameterScope::Local;
+		case ENiagaraParameterPanelCategory::ScriptTransient:
+			return ENiagaraParameterScope::ScriptTransient;
+		default:
+			ensureMsgf(false, TEXT("Unexpected category encountered when getting scope for new parameters in category!"));
+		};
 	}
-	return false;
+	return ENiagaraParameterScope::None;
 }
 
-void SNiagaraParameterPanel::OnRequestRenameOnActionNode()
+bool SNiagaraParameterPanel::GetCanAddParametersToCategory(const ENiagaraParameterPanelCategory Category)
 {
-	GraphActionMenu->OnRequestRenameOnActionNode();
+	switch (Category) {
+	case ENiagaraParameterPanelCategory::User:
+	case ENiagaraParameterPanelCategory::Engine:
+	case ENiagaraParameterPanelCategory::Owner:
+	case ENiagaraParameterPanelCategory::System:
+	case ENiagaraParameterPanelCategory::Emitter:
+	case ENiagaraParameterPanelCategory::Particles:
+	case ENiagaraParameterPanelCategory::Input:
+	case ENiagaraParameterPanelCategory::Local:
+	case ENiagaraParameterPanelCategory::ScriptTransient:
+		return true;
+	default:
+		return false;
+	};
 }
 
-bool SNiagaraParameterPanel::CanRequestRenameOnActionNode(TWeakPtr<FGraphActionNode> InSelectedNode) const
+// void SNiagaraParameterPanel::TryDeleteEntries()
+// {
+// 	TArray<TSharedPtr<FEdGraphSchemaAction>> SelectedActions;
+// 	GraphActionMenu->GetSelectedActions(SelectedActions);
+// 
+// 	for (auto& Action : SelectedActions)
+// 	{
+// 		TSharedPtr<FNiagaraScriptVarAndViewInfoAction> ScriptVarAndViewInfoAction = StaticCastSharedPtr<FNiagaraScriptVarAndViewInfoAction>(Action);
+// 		if (ScriptVarAndViewInfoAction.Get())
+// 		{
+// 			const FNiagaraVariable& Parameter = ScriptVarAndViewInfoAction->ScriptVariableAndViewInfo.ScriptVariable;
+// 			const FNiagaraVariableMetaData& ParameterMetaData = ScriptVarAndViewInfoAction->ScriptVariableAndViewInfo.MetaData;
+// 			FText ToolTipText;
+// 			if (ParameterPanelViewModel->GetCanDeleteParameterAndToolTip(Parameter, ParameterMetaData, ToolTipText))
+// 			{
+// 				ParameterPanelViewModel->DeleteParameter(Parameter, ParameterMetaData);
+// 			}
+// 		}
+// 	}
+// }
+// 
+// bool SNiagaraParameterPanel::CanTryDeleteEntries() const
+// {
+// 	TArray<TSharedPtr<FEdGraphSchemaAction>> SelectedActions;
+// 	GraphActionMenu->GetSelectedActions(SelectedActions);
+// 
+// 	for (auto& Action : SelectedActions)
+// 	{
+// 		TSharedPtr<FNiagaraScriptVarAndViewInfoAction> ScriptVarAndViewInfoAction = StaticCastSharedPtr<FNiagaraScriptVarAndViewInfoAction>(Action);
+// 		if (ScriptVarAndViewInfoAction.Get())
+// 		{
+// 			const FNiagaraVariable& Parameter = ScriptVarAndViewInfoAction->ScriptVariableAndViewInfo.ScriptVariable;
+// 			const FNiagaraVariableMetaData& ParameterMetaData = ScriptVarAndViewInfoAction->ScriptVariableAndViewInfo.MetaData;
+// 			FText ToolTipText;
+// 			if (ParameterPanelViewModel->GetCanDeleteParameterAndToolTip(Parameter, ParameterMetaData, ToolTipText))
+// 			{
+// 				return true;
+// 			}
+// 		}
+// 	}
+// 	return false;
+// }
+// 
+// void SNiagaraParameterPanel::OnRequestRenameOnActionNode()
+// {
+// 	GraphActionMenu->OnRequestRenameOnActionNode();
+// }
+// 
+// bool SNiagaraParameterPanel::CanRequestRenameOnActionNode(TWeakPtr<FGraphActionNode> InSelectedNode) const
+// {
+// 	if (InSelectedNode.IsValid() && InSelectedNode.Pin()->Actions.Num() == 1 && InSelectedNode.Pin()->Actions[0].IsValid())
+// 	{
+// 		const FNiagaraScriptVarAndViewInfoAction* SelectedAction = static_cast<FNiagaraScriptVarAndViewInfoAction*>(InSelectedNode.Pin()->Actions[0].Get());
+// 		FText BlankText = FText();
+// 		return ParameterPanelViewModel->GetCanRenameParameterAndToolTip(
+// 			SelectedAction->ScriptVariableAndViewInfo.ScriptVariable,
+// 			SelectedAction->ScriptVariableAndViewInfo.MetaData,
+// 			TOptional<const FText>(),
+// 			BlankText
+// 		);
+// 	}
+// 	
+// 	return false;
+// }
+// 
+// bool SNiagaraParameterPanel::CanRequestRenameOnActionNode() const
+// {
+// 	TArray<TSharedPtr<FEdGraphSchemaAction> > SelectedActions;
+// 	GraphActionMenu->GetSelectedActions(SelectedActions);
+// 
+// 	// If there is anything selected in the GraphActionMenu, check the item for if it can be renamed.
+// 	if (SelectedActions.Num() > 0)
+// 	{
+// 		return GraphActionMenu->CanRequestRenameOnActionNode();
+// 	}
+// 	return false;
+// }
+// 
+// bool SNiagaraParameterPanel::HandleActionMatchesName(FEdGraphSchemaAction* InAction, const FName& InName) const
+// {
+// 	return FName(*InAction->GetMenuDescription().ToString()) == InName;
+// }
+// 
+void SNiagaraParameterPanel::Refresh()
 {
-	if (InSelectedNode.IsValid() && InSelectedNode.Pin()->Actions.Num() == 1 && InSelectedNode.Pin()->Actions[0].IsValid())
-	{
-		const FNiagaraScriptVarAndViewInfoAction* SelectedAction = static_cast<FNiagaraScriptVarAndViewInfoAction*>(InSelectedNode.Pin()->Actions[0].Get());
-		FText BlankText = FText();
-		return ParameterPanelViewModel->GetCanRenameParameterAndToolTip(
-			SelectedAction->ScriptVariableAndViewInfo.ScriptVariable,
-			SelectedAction->ScriptVariableAndViewInfo.MetaData,
-			TOptional<const FText>(),
-			BlankText
-		);
-	}
-	
-	return false;
+	bNeedsRefresh = true;
 }
 
-bool SNiagaraParameterPanel::CanRequestRenameOnActionNode() const
+bool SNiagaraParameterPanel::SelectParameterEntryByName(const FName& ParameterName) const
 {
-	TArray<TSharedPtr<FEdGraphSchemaAction> > SelectedActions;
-	GraphActionMenu->GetSelectedActions(SelectedActions);
-
-	// If there is anything selected in the GraphActionMenu, check the item for if it can be renamed.
-	if (SelectedActions.Num() > 0)
+	const TArray<FNiagaraScriptVariableAndViewInfo>& ViewedParameters = ParameterPanelViewModel->GetViewedParameters();
+	const FNiagaraScriptVariableAndViewInfo* ScriptVarToSelect = ViewedParameters.FindByPredicate([ParameterName](const FNiagaraScriptVariableAndViewInfo& ScriptVar) {return ScriptVar.ScriptVariable.GetName() == ParameterName; });
+	if (ScriptVarToSelect != nullptr)
 	{
-		return GraphActionMenu->CanRequestRenameOnActionNode();
+		const TArray<FNiagaraScriptVariableAndViewInfo> ItemsToSelect = { *ScriptVarToSelect };
+		ItemSelector->SetSelectedItems(ItemsToSelect);
+		return true;
 	}
 	return false;
 }
 
-bool SNiagaraParameterPanel::HandleActionMatchesName(FEdGraphSchemaAction* InAction, const FName& InName) const
-{
-	return FName(*InAction->GetMenuDescription().ToString()) == InName;
-}
-
-void SNiagaraParameterPanel::Refresh()
-{
-	bNeedsRefresh = true;
-}
-
 	/************************************************************************/
 	/* SNiagaraAddParameterMenu2                                             */
 	/************************************************************************/
@@ -620,42 +648,42 @@ void SNiagaraAddParameterMenu2::CollectAllActions(FGraphActionListBuilderBase& O
 
 	if (ShowKnownConstantParametersFilter.IsSet())
 	{
-		NiagaraParameterPanelSectionID::Type Type = ShowKnownConstantParametersFilter.Get();
-		if (Type != NiagaraParameterPanelSectionID::Type::NONE)
+		const ENiagaraParameterPanelCategory Category = ShowKnownConstantParametersFilter.Get();
+		if (Category != ENiagaraParameterPanelCategory::None)
 		{
 			// Handle engine parameters, but only in read-only contexts.
 			if ((!IsParameterRead.IsSet() || (IsParameterRead.IsSet() && IsParameterRead.Get())) &&
-				(Type == NiagaraParameterPanelSectionID::Type::ENGINE ||
-				Type == NiagaraParameterPanelSectionID::Type::OWNER ||
-				Type == NiagaraParameterPanelSectionID::Type::REFERENCES ||
-				Type == NiagaraParameterPanelSectionID::Type::SYSTEM ||
-				Type == NiagaraParameterPanelSectionID::Type::EMITTER))
+				(Category == ENiagaraParameterPanelCategory::Engine ||
+				Category == ENiagaraParameterPanelCategory::Owner ||
+				Category == ENiagaraParameterPanelCategory::Attributes ||
+				Category == ENiagaraParameterPanelCategory::System ||
+				Category == ENiagaraParameterPanelCategory::Emitter ))
 			{
 				const TArray<FNiagaraVariable>& EngineConstants = FNiagaraConstants::GetEngineConstants();
 				for (const FNiagaraVariable& Var : EngineConstants)
 				{
-					const FText Category = LOCTEXT("NiagaraAddEngineParameterMenu", "Engine Constant");
+					const FText CategoryText = LOCTEXT("NiagaraAddEngineParameterMenu", "Engine Constant");
 					const FText DisplayName = FText::FromName(Var.GetName());
 					const FText Tooltip = FNiagaraConstants::GetEngineConstantDescription(Var);
 
 
 					FNiagaraVariableMetaData GuessedMetaData;
 					FNiagaraEditorUtilities::GetParameterMetaDataFromName(Var.GetName(), GuessedMetaData);
-					bool bAdd = Type == NiagaraParameterPanelSectionID::Type::REFERENCES; // References has all of these!
-					if (Type == NiagaraParameterPanelSectionID::Type::ENGINE && GuessedMetaData.GetScopeName() == FNiagaraConstants::EngineNamespace)
+					bool bAdd = Category == ENiagaraParameterPanelCategory::Attributes; // Attributes has all of these!
+					if (Category == ENiagaraParameterPanelCategory::Engine && GuessedMetaData.GetScopeName() == FNiagaraConstants::EngineNamespace)
 						bAdd = true;
-					else if (Type == NiagaraParameterPanelSectionID::Type::OWNER && GuessedMetaData.GetScopeName() == FNiagaraConstants::EngineOwnerScopeName)
+					else if (Category == ENiagaraParameterPanelCategory::Owner && GuessedMetaData.GetScopeName() == FNiagaraConstants::EngineOwnerScopeName)
 						bAdd = true;
-					else if (Type == NiagaraParameterPanelSectionID::Type::SYSTEM && GuessedMetaData.GetScopeName() == FNiagaraConstants::EngineSystemScopeName)
+					else if (Category == ENiagaraParameterPanelCategory::System && GuessedMetaData.GetScopeName() == FNiagaraConstants::EngineSystemScopeName)
 						bAdd = true;
-					else if (Type == NiagaraParameterPanelSectionID::Type::EMITTER && GuessedMetaData.GetScopeName() == FNiagaraConstants::EngineEmitterScopeName)
+					else if (Category == ENiagaraParameterPanelCategory::Emitter && GuessedMetaData.GetScopeName() == FNiagaraConstants::EngineEmitterScopeName)
 						bAdd = true;
 
 					if (!bAdd)
 						continue;
 
 					TSharedPtr<FNiagaraMenuAction> Action(new FNiagaraMenuAction(
-						Category, DisplayName, Tooltip, 0, FText::GetEmpty(),
+						CategoryText, DisplayName, Tooltip, 0, FText::GetEmpty(),
 						FNiagaraMenuAction::FOnExecuteStackAction::CreateSP(this, &SNiagaraAddParameterMenu2::AddParameterSelected, Var)));
 
 					OutAllActions.AddAction(Action);
@@ -663,19 +691,19 @@ void SNiagaraAddParameterMenu2::CollectAllActions(FGraphActionListBuilderBase& O
 			}
 
 			// Handle particles 
-			if (Type == NiagaraParameterPanelSectionID::Type::PARTICLES ||
-				Type == NiagaraParameterPanelSectionID::Type::REFERENCES )
+			if (Category == ENiagaraParameterPanelCategory::Particles ||
+				Category == ENiagaraParameterPanelCategory::Attributes )
 			{
 				const TArray<FNiagaraVariable>& Vars = FNiagaraConstants::GetCommonParticleAttributes();
 
 				for (const FNiagaraVariable& Var : Vars)
 				{
-					const FText Category = LOCTEXT("NiagaraAddCommParticleParameterMenu", "Common Particle Attributes");
+					const FText CategoryText = LOCTEXT("NiagaraAddCommParticleParameterMenu", "Common Particle Attributes");
 					const FText DisplayName = FText::FromName(Var.GetName());
 					const FText Tooltip = FNiagaraConstants::GetAttributeDescription(Var);
 					
 					TSharedPtr<FNiagaraMenuAction> Action(new FNiagaraMenuAction(
-						Category, DisplayName, Tooltip, 0, FText::GetEmpty(),
+						CategoryText, DisplayName, Tooltip, 0, FText::GetEmpty(),
 						FNiagaraMenuAction::FOnExecuteStackAction::CreateSP(this, &SNiagaraAddParameterMenu2::AddParameterSelected, Var)));
 
 					OutAllActions.AddAction(Action);
diff --git a/Engine/Plugins/FX/Niagara/Source/NiagaraEditor/Private/Widgets/SNiagaraParameterPanelPaletteItem.cpp b/Engine/Plugins/FX/Niagara/Source/NiagaraEditor/Private/Widgets/SNiagaraParameterPanelPaletteItem.cpp
--- a/Engine/Plugins/FX/Niagara/Source/NiagaraEditor/Private/Widgets/SNiagaraParameterPanelPaletteItem.cpp
+++ b/Engine/Plugins/FX/Niagara/Source/NiagaraEditor/Private/Widgets/SNiagaraParameterPanelPaletteItem.cpp
@@ -7,15 +7,15 @@
 
 #define LOCTEXT_NAMESPACE "NiagaraParameterMapPalleteItem"
 
-void SNiagaraParameterPanelPaletteItem::Construct(const FArguments& InArgs, FCreateWidgetForActionData* const InCreateData, TSharedRef<SWidget> ParameterNameViewWidget)
+void SNiagaraParameterPanelPaletteItem::Construct(const FArguments& InArgs/*, FCreateWidgetForActionData* const InCreateData*/, TSharedRef<SWidget> ParameterNameViewWidget)
 {
-	check(InCreateData->Action.IsValid());
-	TSharedPtr<FNiagaraScriptVarAndViewInfoAction> Action = StaticCastSharedPtr<FNiagaraScriptVarAndViewInfoAction>(InCreateData->Action);
-	ActionPtr = InCreateData->Action;
+// 	check(InCreateData->Action.IsValid());
+// 	TSharedPtr<FNiagaraScriptVarAndViewInfoAction> Action = StaticCastSharedPtr<FNiagaraScriptVarAndViewInfoAction>(InCreateData->Action);
+// 	ActionPtr = InCreateData->Action;
 
 	FTutorialMetaData TagMeta("PaletteItem");
 
-	const FLinearColor TypeColor = UEdGraphSchema_Niagara::GetTypeColor(Action->GetScriptVarType());
+	const FLinearColor TypeColor = FLinearColor(1, 1, 1, 0);/* UEdGraphSchema_Niagara::GetTypeColor(Action->GetScriptVarType());*/
 	FSlateBrush const* IconBrush = FEditorStyle::GetBrush(TEXT("Kismet.AllClasses.VariableIcon"));
 	FSlateBrush const* SecondaryBrush = FEditorStyle::GetBrush(TEXT("NoBrush"));
 	FSlateColor        IconColor = FSlateColor(TypeColor);
@@ -35,7 +35,7 @@ void SNiagaraParameterPanelPaletteItem::Construct(const FArguments& InArgs, FCre
 	[
 		SNew(SHorizontalBox)
 		.AddMetaData<FTutorialMetaData>(TagMeta)
-		.ToolTipText(Action->GetTooltipDescription())
+		//.ToolTipText(Action->GetTooltipDescription())
 		// icon slot
 		+SHorizontalBox::Slot()
 		.AutoWidth()
diff --git a/Engine/Plugins/FX/Niagara/Source/NiagaraEditor/Private/Widgets/SNiagaraParameterPanelPaletteItem.h b/Engine/Plugins/FX/Niagara/Source/NiagaraEditor/Private/Widgets/SNiagaraParameterPanelPaletteItem.h
--- a/Engine/Plugins/FX/Niagara/Source/NiagaraEditor/Private/Widgets/SNiagaraParameterPanelPaletteItem.h
+++ b/Engine/Plugins/FX/Niagara/Source/NiagaraEditor/Private/Widgets/SNiagaraParameterPanelPaletteItem.h
@@ -14,5 +14,5 @@ public:
 	{}
 	SLATE_END_ARGS();
 
-	void Construct(const FArguments& InArgs, FCreateWidgetForActionData* const InCreateData, TSharedRef<SWidget> ParameterNameViewWidget);
+	void Construct(const FArguments& InArgs, TSharedRef<SWidget> ParameterNameViewWidget);
 };
\ No newline at end of file
diff --git a/Engine/Plugins/FX/Niagara/Source/NiagaraEditor/Public/NiagaraEditorCommon.h b/Engine/Plugins/FX/Niagara/Source/NiagaraEditor/Public/NiagaraEditorCommon.h
--- a/Engine/Plugins/FX/Niagara/Source/NiagaraEditor/Public/NiagaraEditorCommon.h
+++ b/Engine/Plugins/FX/Niagara/Source/NiagaraEditor/Public/NiagaraEditorCommon.h
@@ -175,6 +175,8 @@ struct FNiagaraScriptVariableAndViewInfo
 		, bIsSelectionRelevant(Other.bIsSelectionRelevant)
 	{};
 
+	bool operator== (const FNiagaraScriptVariableAndViewInfo& Other) const;
+
 	FNiagaraVariable ScriptVariable;
 	FNiagaraVariableMetaData MetaData;
 
@@ -184,35 +186,3 @@ struct FNiagaraScriptVariableAndViewInfo
 	// Whether this entry is related to the current selection state, e.g. if a module is selected in the Stack, mark this entry if ScriptVariable is a member of that module.
 	bool bIsSelectionRelevant;
 };
-
-/* Legacy style Namespace::Type Enum for Parameter Panel section management. */
-namespace NiagaraParameterPanelSectionID
-{
-	enum Type
-	{
-		NONE = 0,
-
-		// Group A, for contextual information when editing scripts
-		INPUTS,
-		REFERENCES,
-		OUTPUTS,
-		LOCALS,
-		INITIALVALUES,
-
-		// Group B, for contextual information when editing systems
-		ENGINE,
-		OWNER,
-		USER,
-		SYSTEM,
-		EMITTER,
-		PARTICLES,
-		CUSTOM,
-	};
-
-	FText OnGetSectionTitle(const NiagaraParameterPanelSectionID::Type InSection);
-	static const NiagaraParameterPanelSectionID::Type GetSectionForScope(ENiagaraParameterScope InScope);
-	static const NiagaraParameterPanelSectionID::Type GetSectionForParameterMetaData(const FNiagaraVariableMetaData& MetaData);
-	ENiagaraParameterScope GetScopeForNewParametersInSection(NiagaraParameterPanelSectionID::Type InSection);
-	ENiagaraScriptParameterUsage GetUsageForNewParametersInSection(NiagaraParameterPanelSectionID::Type InSection);
-
-}
diff --git a/Engine/Plugins/FX/Niagara/Source/NiagaraEditor/Public/ViewModels/NiagaraParameterPanelViewModel.h b/Engine/Plugins/FX/Niagara/Source/NiagaraEditor/Public/ViewModels/NiagaraParameterPanelViewModel.h
--- a/Engine/Plugins/FX/Niagara/Source/NiagaraEditor/Public/ViewModels/NiagaraParameterPanelViewModel.h
+++ b/Engine/Plugins/FX/Niagara/Source/NiagaraEditor/Public/ViewModels/NiagaraParameterPanelViewModel.h
@@ -36,10 +36,6 @@ public:
 
 	virtual void Refresh() const = 0;
 
-	virtual void CollectStaticSections(TArray<int32>& StaticSectionIDs) const = 0;
-
-	virtual NiagaraParameterPanelSectionID::Type GetSectionForVarAndViewInfo(const FNiagaraScriptVariableAndViewInfo& VarAndViewInfo) const = 0;
-	
 	virtual const UNiagaraScriptVariable* AddParameter(FNiagaraVariable& InVariableToAdd, const FNiagaraVariableMetaData& InVariableMetaDataToAssign) = 0;
 
 	virtual void DeleteParameter(const FNiagaraVariable& TargetVariableToRemove, const FNiagaraVariableMetaData& TargetVariableMetaData) = 0;
@@ -58,7 +54,9 @@ public:
 
 	virtual FReply HandleActionDragged(const TSharedPtr<FEdGraphSchemaAction>& InAction, const FPointerEvent& MouseEvent) const = 0;
 
-	//virtual void HandleParameterDropped() {}; //@todo(ng) impl
+	virtual const TArray<TArray<ENiagaraParameterPanelCategory>>& GetDefaultCategoryPaths() const = 0;
+
+	virtual TArray<ENiagaraParameterPanelCategory> GetCategoriesForParameter(const FNiagaraScriptVariableAndViewInfo& ScriptVar) const = 0;
 
 	virtual bool CanDropParameter(const FNiagaraVariable& InTargetScriptVariableToDrop) const { return false; };
 
@@ -75,7 +73,7 @@ public:
 
 	virtual const TArray<FNiagaraScriptVariableAndViewInfo> GetViewedParameters() = 0;
 
-	TSharedRef<class SWidget> GetScriptParameterVisualWidget(FCreateWidgetForActionData* const InCreateData) const;
+	TSharedRef<class SWidget> GetScriptParameterVisualWidget(const FNiagaraScriptVariableAndViewInfo& ScriptVarAndViewInfo) const;
 
 	FOnParameterPanelViewModelRefreshed& GetOnRefreshed() { return OnParameterPanelViewModelRefreshed; };
 
@@ -85,7 +83,6 @@ protected:
 	FOnParameterPanelViewModelExternalSelectionChanged OnParameterPanelViewModelExternalSelectionChanged;
 	/** Cached list of parameters sent to SNiagarParameterPanel, updated whenever GetViewedParameters is called. */
 	TArray<FNiagaraScriptVariableAndViewInfo> CachedViewedParameters;
-
 };
 
 class FNiagaraSystemToolkitParameterPanelViewModel : public INiagaraParameterPanelViewModel, public FEditorUndoClient
@@ -101,10 +98,6 @@ public:
 
 	virtual void Refresh() const override;
 
-	virtual void CollectStaticSections(TArray<int32>& StaticSectionIDs) const override;
-
-	virtual NiagaraParameterPanelSectionID::Type GetSectionForVarAndViewInfo(const FNiagaraScriptVariableAndViewInfo& VarAndViewInfo) const override;
-
 	virtual const UNiagaraScriptVariable* AddParameter(FNiagaraVariable& VariableToAdd, const FNiagaraVariableMetaData& VariableMetaDataToAssign) override;
 
 	virtual void DeleteParameter(const FNiagaraVariable& TargetVariableToRemove, const FNiagaraVariableMetaData& TargetVariableMetaData) override;
@@ -121,6 +114,10 @@ public:
 
 	virtual FReply HandleActionDragged(const TSharedPtr<FEdGraphSchemaAction>& InAction, const FPointerEvent& MouseEvent) const override;
 
+	virtual const TArray<TArray<ENiagaraParameterPanelCategory>>& GetDefaultCategoryPaths() const override;
+
+	virtual TArray<ENiagaraParameterPanelCategory> GetCategoriesForParameter(const FNiagaraScriptVariableAndViewInfo& ScriptVar) const override;
+
 	virtual TArray<TWeakObjectPtr<UNiagaraGraph>> GetEditableGraphs() const override;
 
 	virtual const TArray<FNiagaraScriptVariableAndViewInfo> GetViewedParameters() override;
@@ -143,6 +140,8 @@ private:
 
 	TSharedPtr<FNiagaraSystemViewModel> SystemViewModel;
 	UNiagaraSystemSelectionViewModel* OverviewSelectionViewModel;
+
+	static const TArray<TArray<ENiagaraParameterPanelCategory>> DefaultCategoryPaths;
 };
 
 class FNiagaraScriptToolkitParameterPanelViewModel : public INiagaraParameterPanelViewModel
@@ -158,10 +157,6 @@ public:
 
 	virtual void Refresh() const override;
 
-	virtual void CollectStaticSections(TArray<int32>& StaticSectionIDs) const override;
-
-	virtual NiagaraParameterPanelSectionID::Type GetSectionForVarAndViewInfo(const FNiagaraScriptVariableAndViewInfo& VarAndViewInfo) const override;
-
 	virtual const UNiagaraScriptVariable* AddParameter(FNiagaraVariable& VariableToAdd, const FNiagaraVariableMetaData& VariableMetaDataToAssign) override;
 
 	virtual void DeleteParameter(const FNiagaraVariable& TargetVariableToRemove, const FNiagaraVariableMetaData& TargetVariableMetaData) override;
@@ -182,6 +177,10 @@ public:
 
 	virtual FReply HandleActionDragged(const TSharedPtr<FEdGraphSchemaAction>& InAction, const FPointerEvent& MouseEvent) const override;
 
+	virtual const TArray<TArray<ENiagaraParameterPanelCategory>>& GetDefaultCategoryPaths() const override;
+
+	virtual TArray<ENiagaraParameterPanelCategory> GetCategoriesForParameter(const FNiagaraScriptVariableAndViewInfo& ScriptVar) const override;
+
 	virtual TArray<TWeakObjectPtr<UNiagaraGraph>> GetEditableGraphs() const override;
 
 	virtual const TArray<FNiagaraScriptVariableAndViewInfo> GetViewedParameters() override;
@@ -205,4 +204,6 @@ private:
 	TSharedRef<class SWidget> GetScriptParameterVisualWidget(const UEdGraphPin* Pin) const;
 
 	TSharedPtr<FNiagaraObjectSelection> VariableObjectSelection;
+
+	static const TArray<TArray<ENiagaraParameterPanelCategory>> DefaultCategoryPaths;
 };
diff --git a/Engine/Plugins/FX/Niagara/Source/NiagaraEditor/Public/Widgets/SNiagaraParameterPanel.h b/Engine/Plugins/FX/Niagara/Source/NiagaraEditor/Public/Widgets/SNiagaraParameterPanel.h
--- a/Engine/Plugins/FX/Niagara/Source/NiagaraEditor/Public/Widgets/SNiagaraParameterPanel.h
+++ b/Engine/Plugins/FX/Niagara/Source/NiagaraEditor/Public/Widgets/SNiagaraParameterPanel.h
@@ -7,6 +7,9 @@
 #include "NiagaraTypes.h"
 #include "NiagaraActions.h"
 #include "Framework/Commands/Commands.h"
+#include "Widgets/SItemSelector.h"
+
+typedef SItemSelector<ENiagaraParameterPanelCategory, FNiagaraScriptVariableAndViewInfo> SNiagaraParameterPanelSelector;
 
 class SGraphActionMenu;
 class SEditableTextBox;
@@ -41,12 +44,6 @@ public:
 class SNiagaraParameterPanel : public SCompoundWidget
 {
 public:
-	enum EToolkitType
-	{
-		SCRIPT,
-		SYSTEM,
-	};
-
 	SLATE_BEGIN_ARGS(SNiagaraParameterPanel)
 	{}
 	SLATE_END_ARGS();
@@ -55,75 +52,81 @@ public:
 
 	NIAGARAEDITOR_API void Construct(const FArguments& InArgs, const TSharedPtr<INiagaraParameterPanelViewModel>& InParameterPanelViewModel, const TSharedPtr<FUICommandList>& InToolkitCommands);
 
-	NIAGARAEDITOR_API virtual void Tick(const FGeometry& AllottedGeometry, const double InCurrentTime, const float InDeltaTime) override;
+	TArray<ENiagaraParameterPanelCategory> OnGetCategoriesForItem(const FNiagaraScriptVariableAndViewInfo& Item);
+	bool OnCompareCategoriesForEquality(const ENiagaraParameterPanelCategory& CategoryA, const ENiagaraParameterPanelCategory& CategoryB) const;
+	bool OnCompareCategoriesForSorting(const ENiagaraParameterPanelCategory& CategoryA, const ENiagaraParameterPanelCategory& CategoryB) const;
+	bool OnCompareItemsForEquality(const FNiagaraScriptVariableAndViewInfo& ItemA, const FNiagaraScriptVariableAndViewInfo& ItemB) const;
+	bool OnCompareItemsForSorting(const FNiagaraScriptVariableAndViewInfo& ItemA, const FNiagaraScriptVariableAndViewInfo& ItemB) const;
+	bool OnDoesItemMatchFilterText(const FText& FilterText, const FNiagaraScriptVariableAndViewInfo& Item);
+	TSharedRef<SWidget> OnGenerateWidgetForCategory(const ENiagaraParameterPanelCategory& Category);
+	TSharedRef<SWidget> OnGenerateWidgetForItem(const FNiagaraScriptVariableAndViewInfo& Item);
+	const TArray<TArray<ENiagaraParameterPanelCategory>>& GetDefaultCategoryPaths() const;
 
-	/** Whether the add parameter button should be enabled. */
-// 	bool ParameterAddEnabled() const; //@todo(ng) impl
-
-	void AddParameter(FNiagaraVariable NewVariable, const NiagaraParameterPanelSectionID::Type SectionID);
+	NIAGARAEDITOR_API virtual void Tick(const FGeometry& AllottedGeometry, const double InCurrentTime, const float InDeltaTime) override;
+// 
+// 	/** Whether the add parameter button should be enabled. */
+// // 	bool ParameterAddEnabled() const; //@todo(ng) impl
+// 
+	void AddParameter(FNiagaraVariable NewVariable, const ENiagaraParameterPanelCategory Category);
 
- 	/** Refreshes the palette items for the graph action menu on the next tick. */
- 	void Refresh();
+	/** Refreshes the items for the item selector on the next tick. */
+	void Refresh();
 
 	static TSharedRef<SExpanderArrow> CreateCustomActionExpander(const struct FCustomExpanderData& ActionMenuData);
 
 private:
+	bool SelectParameterEntryByName(const FName& ParameterName) const;
+
 	/** Function to bind to SNiagaraAddParameterMenus to filter types we allow creating */
 	bool AllowMakeType(const FNiagaraTypeDefinition& InType) const;
 
-	// SGraphActionMenu delegates
-	TSharedRef<SWidget> OnCreateWidgetForAction(struct FCreateWidgetForActionData* const InCreateData);
-	void CollectAllActions(FGraphActionListBuilderBase& OutAllActions);
-	void CollectStaticSections(TArray<int32>& StaticSectionIDs) const;
-	FReply OnActionDragged(const TArray<TSharedPtr<FEdGraphSchemaAction>>& InActions, const FPointerEvent& MouseEvent);
-	void OnActionSelected(const TArray<TSharedPtr<FEdGraphSchemaAction>>& InActions, ESelectInfo::Type InSelectionType);
-// 	void OnActionDoubleClicked(const TArray<TSharedPtr<FEdGraphSchemaAction>>& InActions); //@todo(ng) impl
-	TSharedPtr<SWidget> OnContextMenuOpening();
- 	FText OnGetSectionTitle(int32 InSectionID);
-	TSharedRef<SWidget> OnGetSectionWidget(TSharedRef<SWidget> RowWidget, int32 InSectionID);
-	TSharedRef<SWidget> CreateAddToSectionButton(const NiagaraParameterPanelSectionID::Type InSection, TWeakPtr<SWidget> WeakRowWidget, FText AddNewText, FName MetaDataTag);
+	TSharedRef<SWidget> CreateAddToCategoryButton(const ENiagaraParameterPanelCategory Category, FText AddNewText, FName MetaDataTag);
+// // 
+// // 	/** Checks if the selected action has context menu */
+// // 	bool SelectionHasContextMenu() const; //@todo(ng) impl
+// // 
+	TSharedRef<SWidget> OnGetParameterMenu(const ENiagaraParameterPanelCategory Category);
+
+	static ENiagaraParameterScope GetScopeForNewParametersInCategory(const ENiagaraParameterPanelCategory Category);
+	static bool GetCanAddParametersToCategory(const ENiagaraParameterPanelCategory Category);
+// 	
+// 	//EVisibility OnAddButtonTextVisibility(TWeakPtr<SWidget> RowWidget, const NiagaraParameterMapSectionID::Type InSection) const; //@todo(ng) impl
 // 
-// 	/** Checks if the selected action has context menu */
-// 	bool SelectionHasContextMenu() const; //@todo(ng) impl
 // 
-	TSharedRef<SWidget> OnGetParameterMenu(const NiagaraParameterPanelSectionID::Type InSection);
-	
-	//EVisibility OnAddButtonTextVisibility(TWeakPtr<SWidget> RowWidget, const NiagaraParameterMapSectionID::Type InSection) const; //@todo(ng) impl
-
-
-// 	//Callbacks
-	/** Try to delete all entries that are selected. If only some selected entries pass CanDeleteEntry(), then only those selected entries will be deleted. */
-	void TryDeleteEntries();
-	/** Test if ALL selected entries cannot be deleted. If only one selected entry passes INiagaraParameterPanelViewModel::CanRemoveParameter(), return true. Per entry deletion handling is done in TryDeleteEntries(). */
-	bool CanTryDeleteEntries() const;
-	void OnRequestRenameOnActionNode();
-	bool CanRequestRenameOnActionNode(TWeakPtr<struct FGraphActionNode> InSelectedNode) const;
-	bool CanRequestRenameOnActionNode() const;
-	void HandlePaletteItemParameterRenamed(const FText& InText, const FNiagaraScriptVarAndViewInfoAction& InAction);
-	void HandlePaletteItemScopeComboSelectionChanged(ENiagaraParameterScope InScope, const FNiagaraScriptVarAndViewInfoAction& InAction);
+// // 	//Callbacks
+// 	/** Try to delete all entries that are selected. If only some selected entries pass CanDeleteEntry(), then only those selected entries will be deleted. */
+// 	void TryDeleteEntries();
+// 	/** Test if ALL selected entries cannot be deleted. If only one selected entry passes INiagaraParameterPanelViewModel::CanRemoveParameter(), return true. Per entry deletion handling is done in TryDeleteEntries(). */
+// 	bool CanTryDeleteEntries() const;
+// 	void OnRequestRenameOnActionNode();
+// 	bool CanRequestRenameOnActionNode(TWeakPtr<struct FGraphActionNode> InSelectedNode) const;
+// 	bool CanRequestRenameOnActionNode() const;
+// 	void HandlePaletteItemParameterRenamed(const FText& InText, const FNiagaraScriptVarAndViewInfoAction& InAction);
+// 	void HandlePaletteItemScopeComboSelectionChanged(ENiagaraParameterScope InScope, const FNiagaraScriptVarAndViewInfoAction& InAction);
 	void HandleExternalSelectionChanged(const UObject* Obj);
-
-	/** Delegate handler used to match an FName to an action in the list, used for renaming keys */
-	bool HandleActionMatchesName(struct FEdGraphSchemaAction* InAction, const FName& InName) const;
-
-private:
-	/** Graph Action Menu for displaying all our variables and functions */
-	TSharedPtr<SGraphActionMenu> GraphActionMenu;
-
-	/** The filter box that handles filtering for both graph action menus. */
-	TSharedPtr<SSearchBox> FilterBox;
- 
-	/** Add parameter buttons for all sections. */
-	TArray<TSharedPtr<SComboButton>> AddParameterButtons;
-
-	TSharedPtr<FNiagaraObjectSelection> SelectedVariableObjects;
-
+// 
+// 	/** Delegate handler used to match an FName to an action in the list, used for renaming keys */
+// 	bool HandleActionMatchesName(struct FEdGraphSchemaAction* InAction, const FName& InName) const;
+// 
+// private:
+// 	/** Graph Action Menu for displaying all our variables and functions */
+// 	TSharedPtr<SGraphActionMenu> GraphActionMenu;
+// 
+// 	/** The filter box that handles filtering for both graph action menus. */
+// 	TSharedPtr<SSearchBox> FilterBox;
+//  
+// 	/** Add parameter buttons for all sections. */
+	TMap<ENiagaraParameterPanelCategory, TSharedPtr<SComboButton>> AddParameterButtons;
+// 
+// 	TSharedPtr<FNiagaraObjectSelection> SelectedVariableObjects;
+// 
 	TSharedPtr<FUICommandList> ToolkitCommands; //@todo(ng) add Find And Rename Parameter command
-
+// 
 	bool bNeedsRefresh;
-	bool bGraphActionPendingRename;
-
+// 	bool bGraphActionPendingRename;
+// 
 	TSharedPtr<INiagaraParameterPanelViewModel> ParameterPanelViewModel;
+	TSharedPtr<SNiagaraParameterPanelSelector> ItemSelector;
 };
 
 class SNiagaraAddParameterMenu2 : public SCompoundWidget
@@ -137,7 +140,7 @@ public:
 	SLATE_BEGIN_ARGS(SNiagaraAddParameterMenu2)
 		: _AllowCreatingNew(true)
 		, _ShowGraphParameters(true)
-		, _ShowKnownConstantParametersFilter(NiagaraParameterPanelSectionID::Type::NONE)
+		, _ShowKnownConstantParametersFilter(ENiagaraParameterPanelCategory::None)
 		, _AutoExpandMenu(false)
 		, _IsParameterRead(true)
 		, _NewParameterScope(ENiagaraParameterScope::Particles) {}
@@ -146,7 +149,7 @@ public:
 		SLATE_EVENT(FOnAllowMakeType, OnAllowMakeType)
 		SLATE_ATTRIBUTE(bool, AllowCreatingNew)
 		SLATE_ATTRIBUTE(bool, ShowGraphParameters)
-		SLATE_ATTRIBUTE(NiagaraParameterPanelSectionID::Type, ShowKnownConstantParametersFilter)
+		SLATE_ATTRIBUTE(ENiagaraParameterPanelCategory, ShowKnownConstantParametersFilter)
 		SLATE_ATTRIBUTE(bool, AutoExpandMenu)
 		SLATE_ATTRIBUTE(bool, IsParameterRead)
 		SLATE_ATTRIBUTE(FString, NewParameterNamespace)
@@ -174,7 +177,7 @@ private:
 
 	TAttribute<bool> AllowCreatingNew;
 	TAttribute<bool> ShowGraphParameters;
-	TAttribute<NiagaraParameterPanelSectionID::Type> ShowKnownConstantParametersFilter;
+	TAttribute<ENiagaraParameterPanelCategory> ShowKnownConstantParametersFilter;
 	TAttribute<bool> AutoExpandMenu;
 	TAttribute<bool> IsParameterRead;
 	TAttribute<FString> NewParameterNamespace;
