Remove with editor define to fix build.

diff --git a/Engine/Plugins/FX/Niagara/Source/Niagara/Public/NiagaraTypes.h b/Engine/Plugins/FX/Niagara/Source/Niagara/Public/NiagaraTypes.h
--- a/Engine/Plugins/FX/Niagara/Source/Niagara/Public/NiagaraTypes.h
+++ b/Engine/Plugins/FX/Niagara/Source/Niagara/Public/NiagaraTypes.h
@@ -618,7 +618,6 @@ private:
 		FString NamespaceString;
 };
 
-#if WITH_EDITOR
 UENUM()
 enum class ENiagaraParameterPanelCategory : uint32
 {
@@ -661,7 +660,6 @@ enum class ENiagaraParameterPanelCategory : uint32
 
 	Num UMETA(Hidden)
 };
-#endif
 
 USTRUCT()
 struct NIAGARA_API FNiagaraVariableMetaData
