Niagara - Fix issues where compile messages weren't showing in the UI correctly.

diff --git a/Engine/Plugins/FX/Niagara/Source/NiagaraEditor/Private/NiagaraCompiler.cpp b/Engine/Plugins/FX/Niagara/Source/NiagaraEditor/Private/NiagaraCompiler.cpp
--- a/Engine/Plugins/FX/Niagara/Source/NiagaraEditor/Private/NiagaraCompiler.cpp
+++ b/Engine/Plugins/FX/Niagara/Source/NiagaraEditor/Private/NiagaraCompiler.cpp
@@ -830,6 +830,11 @@ TSharedPtr<FNiagaraVMExecutableData> FNiagaraEditorModule::GetCompilationResult(
 	
 	Results.Data->ErrorMsg = OutGraphLevelErrorMessages;
 	Results.Data->LastCompileStatus = (FNiagaraCompileResults::CompileResultsToSummary(&Results));
+	if (Results.Data->LastCompileStatus != ENiagaraScriptCompileStatus::NCS_Error)
+	{
+		// When there are no errors the compile events get emptied, so add them back here.
+		Results.Data->LastCompileEvents.Append(Results.CompileEvents);
+	}
 	return CompileResult->Data;
 }
 
diff --git a/Engine/Plugins/FX/Niagara/Source/NiagaraEditor/Private/NiagaraMessageManager.cpp b/Engine/Plugins/FX/Niagara/Source/NiagaraEditor/Private/NiagaraMessageManager.cpp
--- a/Engine/Plugins/FX/Niagara/Source/NiagaraEditor/Private/NiagaraMessageManager.cpp
+++ b/Engine/Plugins/FX/Niagara/Source/NiagaraEditor/Private/NiagaraMessageManager.cpp
@@ -174,7 +174,7 @@ void FNiagaraMessageManager::DoMessageJobsInQueueTick()
 			}
 			OnRequestRefresh.Broadcast(CurrentMessageJobBatch.MessageJobsAssetKey, GeneratedMessagesForCurrentMessageJobBatch);
 			GeneratedMessagesForCurrentMessageJobBatch.Reset();
-			MessageJobBatchArr.Pop();
+			MessageJobBatchArr.RemoveAt(0);
 		} 
 	}
 }
@@ -370,7 +370,7 @@ TSharedRef<const INiagaraMessage> FNiagaraMessageJobCompileEvent::GenerateNiagar
 		TOptional<const FText> ScriptNameAndPathsGetterFailureReason;
 		TOptional<const FString> OwningScriptNameStringCopy = OwningScriptNameString;
 		TArray<FObjectKey> ContextObjectKeys;
-		bool bSuccessfullyFoundScripts = RecursiveGetScriptNamesAndAssetPathsFromContextStack(ContextStackGuids, ScriptGraph, ContextScriptNamesAndPaths, OwningScriptNameStringCopy, ScriptNameAndPathsGetterFailureReason, ContextObjectKeys);
+		bool bSuccessfullyFoundScripts = RecursiveGetScriptNamesAndAssetPathsFromContextStack(ContextStackGuids, CompileEvent.NodeGuid, ScriptGraph, ContextScriptNamesAndPaths, OwningScriptNameStringCopy, ScriptNameAndPathsGetterFailureReason, ContextObjectKeys);
 
 		if (bSuccessfullyFoundScripts == false)
 		{
@@ -407,6 +407,7 @@ TSharedRef<const INiagaraMessage> FNiagaraMessageJobCompileEvent::GenerateNiagar
 
 bool FNiagaraMessageJobCompileEvent::RecursiveGetScriptNamesAndAssetPathsFromContextStack(
 	  TArray<FGuid>& InContextStackNodeGuids
+	, FGuid NodeGuid
 	, const UNiagaraGraph* InGraphToSearch
 	, TArray<FNiagaraScriptNameAndAssetPath>& OutContextScriptNamesAndAssetPaths
 	, TOptional<const FString>& OutEmitterName
@@ -416,7 +417,16 @@ bool FNiagaraMessageJobCompileEvent::RecursiveGetScriptNamesAndAssetPathsFromCon
 {
 	checkf(InGraphToSearch, TEXT("Failed to get a node graph to search!"));
 
-	if (InGraphToSearch && InContextStackNodeGuids.Num() == 0)
+	if (NodeGuid.IsValid())
+	{
+		UEdGraphNode*const* EventNodePtr = InGraphToSearch->Nodes.FindByPredicate([NodeGuid](UEdGraphNode* Node) { return Node->NodeGuid == NodeGuid; });
+		if (EventNodePtr != nullptr)
+		{
+			OutContextNodeObjectKeys.AddUnique(FObjectKey(*EventNodePtr));
+		}
+	}
+
+	if (InContextStackNodeGuids.Num() == 0)
 	{
 		//StackGuids arr has been cleared out which means we have walked the entire context stack.
 		return true;
@@ -473,7 +483,7 @@ bool FNiagaraMessageJobCompileEvent::RecursiveGetScriptNamesAndAssetPathsFromCon
 				return false;
 			}
 			OutContextScriptNamesAndAssetPaths.Add(FNiagaraScriptNameAndAssetPath(FunctionCallNodeAssignedScript->GetName(), FunctionCallNodeAssignedScript->GetPathName()));
-			return RecursiveGetScriptNamesAndAssetPathsFromContextStack(InContextStackNodeGuids, FunctionScriptGraph, OutContextScriptNamesAndAssetPaths, OutEmitterName, OutFailureReason, OutContextNodeObjectKeys);
+			return RecursiveGetScriptNamesAndAssetPathsFromContextStack(InContextStackNodeGuids, NodeGuid, FunctionScriptGraph, OutContextScriptNamesAndAssetPaths, OutEmitterName, OutFailureReason, OutContextNodeObjectKeys);
 		}
 
 		UNiagaraNodeEmitter* EmitterNode = Cast<UNiagaraNodeEmitter>(ContextNode.GetValue());
@@ -486,7 +496,7 @@ bool FNiagaraMessageJobCompileEvent::RecursiveGetScriptNamesAndAssetPathsFromCon
 			checkf(EmitterScriptGraph, TEXT("Emitter Script Source does not have a UNiagaraGraph!"));
 
 			OutEmitterName = EmitterNode->GetEmitterUniqueName();
-			return RecursiveGetScriptNamesAndAssetPathsFromContextStack(InContextStackNodeGuids, EmitterScriptGraph, OutContextScriptNamesAndAssetPaths, OutEmitterName, OutFailureReason, OutContextNodeObjectKeys);
+			return RecursiveGetScriptNamesAndAssetPathsFromContextStack(InContextStackNodeGuids, NodeGuid, EmitterScriptGraph, OutContextScriptNamesAndAssetPaths, OutEmitterName, OutFailureReason, OutContextNodeObjectKeys);
 		}
 		checkf(false, TEXT("Matching node is not a function call or emitter node!"));
 	}
diff --git a/Engine/Plugins/FX/Niagara/Source/NiagaraEditor/Private/ViewModels/NiagaraSystemViewModel.cpp b/Engine/Plugins/FX/Niagara/Source/NiagaraEditor/Private/ViewModels/NiagaraSystemViewModel.cpp
--- a/Engine/Plugins/FX/Niagara/Source/NiagaraEditor/Private/ViewModels/NiagaraSystemViewModel.cpp
+++ b/Engine/Plugins/FX/Niagara/Source/NiagaraEditor/Private/ViewModels/NiagaraSystemViewModel.cpp
@@ -113,7 +113,6 @@ void FNiagaraSystemViewModel::Initialize(UNiagaraSystem& InSystem, FNiagaraSyste
 	SetupSequencer();
 	RefreshAll();
 	AddSystemEventHandlers();
-	SendLastCompileMessageJobs();
 }
 
 void FNiagaraSystemViewModel::DumpToText(FString& ExportText)
@@ -898,6 +897,8 @@ void FNiagaraSystemViewModel::RefreshAll()
 	RefreshEmitterHandleViewModels();
 	RefreshSequencerTracks();
 	ResetCurveData();
+	InvalidateCachedCompileStatus();
+	SendLastCompileMessageJobs();
 	ScriptScratchPadViewModel->RefreshScriptViewModels();
 	SystemStackViewModel->GetRootEntry()->RefreshChildren();
 	SelectionViewModel->Refresh();
diff --git a/Engine/Plugins/FX/Niagara/Source/NiagaraEditor/Private/ViewModels/Stack/NiagaraStackScriptItemGroup.cpp b/Engine/Plugins/FX/Niagara/Source/NiagaraEditor/Private/ViewModels/Stack/NiagaraStackScriptItemGroup.cpp
--- a/Engine/Plugins/FX/Niagara/Source/NiagaraEditor/Private/ViewModels/Stack/NiagaraStackScriptItemGroup.cpp
+++ b/Engine/Plugins/FX/Niagara/Source/NiagaraEditor/Private/ViewModels/Stack/NiagaraStackScriptItemGroup.cpp
@@ -569,17 +569,42 @@ void UNiagaraStackScriptItemGroup::RefreshIssues(TArray<FStackIssue>& NewIssues)
 		ENiagaraScriptCompileStatus Status = ScriptViewModelPinned->GetScriptCompileStatus(GetScriptUsage(), GetScriptUsageId());
 		if (!bForcedError)
 		{
+
 			if (Status == ENiagaraScriptCompileStatus::NCS_Error)
 			{
 				FStackIssue CompileError(
 					EStackIssueSeverity::Error,
-					LOCTEXT("ConpileErrorSummary", "The stack has compile errors."),
-					ScriptViewModelPinned->GetScriptErrors(GetScriptUsage(), GetScriptUsageId()),
+					LOCTEXT("ConpileErrorSummary", "This stage has compile errors."),
+					LOCTEXT("CompileErrorLong", "Errors were generated while compiling this stage. Check for modules showing errors or check the niagara log for details."),
 					GetStackEditorDataKey(),
-					false);
+					false,
+					FStackIssueFix(
+						LOCTEXT("ShowLogForErrorsFix", "Show the errors in the Niagara Log"), 
+						FStackIssueFixDelegate::CreateLambda([this]()
+						{
+							GetSystemViewModel()->FocusTab(FNiagaraSystemToolkit::MessageLogTabID);
+						}),
+						EStackIssueFixStyle::Link));
 
 				NewIssues.Add(CompileError);
 			}
+			else if (Status == ENiagaraScriptCompileStatus::NCS_UpToDateWithWarnings || Status == ENiagaraScriptCompileStatus::NCS_ComputeUpToDateWithWarnings)
+			{
+				FStackIssue CompileWarning(
+					EStackIssueSeverity::Warning,
+					LOCTEXT("CompileWarningSummary", "This stage has compile warnings."),
+					LOCTEXT("CompileWarningLong", "Warnings were generated while compiling this stage. Check for modules showing warnings or check the niagara log for details."),
+					GetStackEditorDataKey(),
+					false,
+					FStackIssueFix(
+						LOCTEXT("ShowLogForWarningsFix", "Show the warnings in the Niagara Log"),
+						FStackIssueFixDelegate::CreateLambda([this]()
+							{
+								GetSystemViewModel()->FocusTab(FNiagaraSystemToolkit::MessageLogTabID);
+							}),
+						EStackIssueFixStyle::Link));
+				NewIssues.Add(CompileWarning);
+			}
 		}
 	}
 }
diff --git a/Engine/Plugins/FX/Niagara/Source/NiagaraEditor/Public/NiagaraMessageManager.h b/Engine/Plugins/FX/Niagara/Source/NiagaraEditor/Public/NiagaraMessageManager.h
--- a/Engine/Plugins/FX/Niagara/Source/NiagaraEditor/Public/NiagaraMessageManager.h
+++ b/Engine/Plugins/FX/Niagara/Source/NiagaraEditor/Public/NiagaraMessageManager.h
@@ -87,6 +87,7 @@ private:
 
 	bool RecursiveGetScriptNamesAndAssetPathsFromContextStack(
 		  TArray<FGuid>& InContextStackNodeGuids
+		, FGuid NodeGuid
 		, const UNiagaraGraph* InGraphToSearch
 		, TArray<FNiagaraScriptNameAndAssetPath>& OutContextScriptNamesAndAssetPaths
 		, TOptional<const FString>& OutEmitterName
