Niagara parameters that do not have metadata set for them in the graph no longer cause a crash when creating the tooltip.

diff --git a/Engine/Plugins/FX/Niagara/Source/NiagaraEditor/Private/NiagaraNodeParameterMapBase.cpp b/Engine/Plugins/FX/Niagara/Source/NiagaraEditor/Private/NiagaraNodeParameterMapBase.cpp
--- a/Engine/Plugins/FX/Niagara/Source/NiagaraEditor/Private/NiagaraNodeParameterMapBase.cpp
+++ b/Engine/Plugins/FX/Niagara/Source/NiagaraEditor/Private/NiagaraNodeParameterMapBase.cpp
@@ -176,19 +176,22 @@ void UNiagaraNodeParameterMapBase::GetPinHoverText(const UEdGraphPin& Pin, FStri
 			}
 
 			FNiagaraVariable Var = FNiagaraVariable(TypeDef, Pin.PinName);
-			TOptional<FNiagaraVariableMetaData> Metadata = NiagaraGraph->GetMetaData(Var);
-
-			FText Description = Metadata.IsSet() ? Metadata->Description : FText::GetEmpty();
-			const FText TooltipFormat = LOCTEXT("Parameters", "Name: {0} \nType: {1}\nDescription: {2}\nScope: {3}\nUser Editable: {4}\nUsage: {5}");
 			const FText Name = FText::FromName(Var.GetName());
-			FName CachedParamName = NAME_None;
+
+			FText Description, ScopeText, UsageText, UserEditableText;
+			TOptional<FNiagaraVariableMetaData> Metadata = NiagaraGraph->GetMetaData(Var);
 			if (Metadata.IsSet())
 			{
+				FName CachedParamName;
 				Metadata->GetParameterName(CachedParamName);
+				UserEditableText = FText::FromName(CachedParamName);
+
+				Description = Metadata->Description;
+				ScopeText = FText::FromName(Metadata->GetScopeName());
+				UsageText = StaticEnum<ENiagaraScriptParameterUsage>()->GetDisplayNameTextByValue((int64) Metadata->GetUsage());
 			}
-			const FText ScopeText = FText::FromName(Metadata->GetScopeName());
-			const FText UserEditableText = FText::FromName(CachedParamName);
-			const FText UsageText = StaticEnum<ENiagaraScriptParameterUsage>()->GetDisplayNameTextByValue((int64)Metadata->GetUsage());
+
+			const FText TooltipFormat = LOCTEXT("Parameters", "Name: {0} \nType: {1}\nDescription: {2}\nScope: {3}\nUser Editable: {4}\nUsage: {5}");
 			const FText ToolTipText = FText::Format(TooltipFormat, FText::FromName(Var.GetName()), Var.GetType().GetNameText(), Description, ScopeText, UserEditableText, UsageText);
 			HoverTextOut = ToolTipText.ToString();
 		}
