Niagara System parameters can now be renamed from the Parameters panel.

diff --git a/Engine/Plugins/FX/Niagara/Source/NiagaraEditor/Private/ViewModels/NiagaraParameterPanelViewModel.cpp b/Engine/Plugins/FX/Niagara/Source/NiagaraEditor/Private/ViewModels/NiagaraParameterPanelViewModel.cpp
--- a/Engine/Plugins/FX/Niagara/Source/NiagaraEditor/Private/ViewModels/NiagaraParameterPanelViewModel.cpp
+++ b/Engine/Plugins/FX/Niagara/Source/NiagaraEditor/Private/ViewModels/NiagaraParameterPanelViewModel.cpp
@@ -121,11 +121,11 @@ const UNiagaraScriptVariable* FNiagaraSystemToolkitParameterPanelViewModel::AddP
 	FNiagaraEditorUtilities::GetVariableMetaDataScope(InVariableMetaDataToAssign, NewScope);
 	if (NewScope == ENiagaraParameterScope::User)
 	{
-		SystemViewModel->GetSystem().Modify();
-		UNiagaraSystem* System = &SystemViewModel->GetSystem();
-		UNiagaraSystemEditorData* SystemEditorData = CastChecked<UNiagaraSystemEditorData>(System->GetEditorData(), ECastCheckedType::NullChecked);
+		UNiagaraSystem& System = SystemViewModel->GetSystem();
+		System.Modify();
+		UNiagaraSystemEditorData* SystemEditorData = CastChecked<UNiagaraSystemEditorData>(System.GetEditorData(), ECastCheckedType::NullChecked);
 		SystemEditorData->Modify();
-		bool bSuccess = FNiagaraEditorUtilities::AddParameter(VariableToAdd, System->GetExposedParameters(), *System, SystemEditorData->GetStackEditorData());
+		bool bSuccess = FNiagaraEditorUtilities::AddParameter(VariableToAdd, System.GetExposedParameters(), System, SystemEditorData->GetStackEditorData());
 		Refresh();
 	}
 	else if (NewScope == ENiagaraParameterScope::System)
@@ -211,14 +211,32 @@ void FNiagaraSystemToolkitParameterPanelViewModel::RenameParameter(const FNiagar
 		NewVariableName = FName(*NewVariableNameText.ToString());
 	}
 
-	for (const TWeakObjectPtr<UNiagaraGraph>& Graph : GetEditableEmitterScriptGraphs())
+	ENiagaraParameterScope NewScope;
+	FNiagaraEditorUtilities::GetVariableMetaDataScope(TargetVariableMetaData, NewScope);
+	if (NewScope == ENiagaraParameterScope::User)
+	{
+		UNiagaraSystem& System = SystemViewModel->GetSystem();
+		System.Modify();
+		System.GetExposedParameters().RenameParameter(TargetVariableToRename, NewVariableName);
+	}
+	else if (NewScope == ENiagaraParameterScope::System)
+	{
+		UNiagaraGraph* Graph = SystemViewModel->GetSystemScriptViewModel()->GetGraphViewModel()->GetGraph();
+		Graph->Modify();
+		Graph->RenameParameter(TargetVariableToRename, NewVariableName);
+	}
+	else
 	{
-		if (ensureMsgf(Graph.IsValid(), TEXT("Editable Emitter Script Graph was stale when renaming parameter!")))
+		for (TWeakObjectPtr<UNiagaraGraph> Graph : GetEditableEmitterScriptGraphs())
 		{
-			Graph->Modify();
-			Graph->RenameParameter(TargetVariableToRename, NewVariableName, false, TargetVariableMetaData.GetScopeName()); //@todo(ng) handle renaming system params
+			if (ensureMsgf(Graph.IsValid(), TEXT("Editable Emitter Script Graph was stale when adding parameter!")))
+			{
+				Graph->Modify();
+				Graph->RenameParameter(TargetVariableToRename, NewVariableName);
+			}
 		}
 	}
+
 	Refresh();
 }
 
