Make sure NiagaraID is always defined in the generated HLSL, since it can be used from custom HLSL nodes without declaring any pins of that type.

Also cleaned up a few places where FNiagaraTypeDefinition was being passed by value.

diff --git a/Engine/Plugins/FX/Niagara/Source/NiagaraEditor/Private/NiagaraHlslTranslator.cpp b/Engine/Plugins/FX/Niagara/Source/NiagaraEditor/Private/NiagaraHlslTranslator.cpp
--- a/Engine/Plugins/FX/Niagara/Source/NiagaraEditor/Private/NiagaraHlslTranslator.cpp
+++ b/Engine/Plugins/FX/Niagara/Source/NiagaraEditor/Private/NiagaraHlslTranslator.cpp
@@ -1326,12 +1326,16 @@ const FNiagaraTranslateResults &FHlslNiagaraTranslator::Translate(const FNiagara
 			}
 		}
 
+		// NiagaraID is a fundamental type which can be used in custom HLSL without declaring any pins of that type in any nodes.
+		// Make sure it's always defined in the generated HLSL.
+		StructsToDefine.AddUnique(FNiagaraTypeDefinition::GetIDDef());
+
 		// Generate the Parameter Map HLSL definitions. We don't add to the final HLSL output here. We just build up the strings and tables
 		// that are needed later.
 		TArray<FNiagaraVariable> PrimaryDataSetOutputEntries;
 		FString ParameterMapDefinitionStr = BuildParameterMapHlslDefinitions(PrimaryDataSetOutputEntries);
 
-		for (FNiagaraTypeDefinition Type : StructsToDefine)
+		for (const FNiagaraTypeDefinition& Type : StructsToDefine)
 		{
 			FText ErrorMessage;
 			HlslOutput += BuildHLSLStructDecl(Type, ErrorMessage);
@@ -7093,7 +7097,7 @@ FGuid FHlslNiagaraTranslator::GetTargetUsageId() const
 }
 //////////////////////////////////////////////////////////////////////////
 
-FString FHlslNiagaraTranslator::GetHlslDefaultForType(FNiagaraTypeDefinition Type)
+FString FHlslNiagaraTranslator::GetHlslDefaultForType(const FNiagaraTypeDefinition& Type)
 {
 	if (Type == FNiagaraTypeDefinition::GetFloatDef())
 	{
@@ -7133,7 +7137,7 @@ FString FHlslNiagaraTranslator::GetHlslDefaultForType(FNiagaraTypeDefinition Typ
 	}
 }
 
-bool FHlslNiagaraTranslator::IsBuiltInHlslType(FNiagaraTypeDefinition Type)
+bool FHlslNiagaraTranslator::IsBuiltInHlslType(const FNiagaraTypeDefinition& Type)
 {
 	return
 		Type == FNiagaraTypeDefinition::GetFloatDef() ||
@@ -7148,7 +7152,7 @@ bool FHlslNiagaraTranslator::IsBuiltInHlslType(FNiagaraTypeDefinition Type)
 		Type == FNiagaraTypeDefinition::GetBoolDef();
 }
 
-FString FHlslNiagaraTranslator::GetStructHlslTypeName(FNiagaraTypeDefinition Type)
+FString FHlslNiagaraTranslator::GetStructHlslTypeName(const FNiagaraTypeDefinition& Type)
 {
 	if (Type.IsValid() == false)
 	{
@@ -7229,7 +7233,7 @@ FString FHlslNiagaraTranslator::GetPropertyHlslTypeName(const FProperty* Propert
 	}
 }
 
-FString FHlslNiagaraTranslator::BuildHLSLStructDecl(FNiagaraTypeDefinition Type, FText& OutErrorMessage)
+FString FHlslNiagaraTranslator::BuildHLSLStructDecl(const FNiagaraTypeDefinition& Type, FText& OutErrorMessage)
 {
 	if (!IsBuiltInHlslType(Type))
 	{
@@ -7255,7 +7259,7 @@ FString FHlslNiagaraTranslator::BuildHLSLStructDecl(FNiagaraTypeDefinition Type,
 	return TEXT("");
 }
 
-bool FHlslNiagaraTranslator::IsHlslBuiltinVector(FNiagaraTypeDefinition Type)
+bool FHlslNiagaraTranslator::IsHlslBuiltinVector(const FNiagaraTypeDefinition& Type)
 {
 	if ((Type == FNiagaraTypeDefinition::GetVec2Def()) ||
 		(Type == FNiagaraTypeDefinition::GetVec3Def()) ||
diff --git a/Engine/Plugins/FX/Niagara/Source/NiagaraEditor/Private/NiagaraHlslTranslator.h b/Engine/Plugins/FX/Niagara/Source/NiagaraEditor/Private/NiagaraHlslTranslator.h
--- a/Engine/Plugins/FX/Niagara/Source/NiagaraEditor/Private/NiagaraHlslTranslator.h
+++ b/Engine/Plugins/FX/Niagara/Source/NiagaraEditor/Private/NiagaraHlslTranslator.h
@@ -531,12 +531,12 @@ public:
 		return CompilationTarget;
 	}
 
-	static bool IsBuiltInHlslType(FNiagaraTypeDefinition Type);
-	static FString GetStructHlslTypeName(FNiagaraTypeDefinition Type);
+	static bool IsBuiltInHlslType(const FNiagaraTypeDefinition& Type);
+	static FString GetStructHlslTypeName(const FNiagaraTypeDefinition& Type);
 	static FString GetPropertyHlslTypeName(const FProperty* Property);
-	static FString BuildHLSLStructDecl(FNiagaraTypeDefinition Type, FText& OutErrorMessage);
-	static FString GetHlslDefaultForType(FNiagaraTypeDefinition Type);
-	static bool IsHlslBuiltinVector(FNiagaraTypeDefinition Type);
+	static FString BuildHLSLStructDecl(const FNiagaraTypeDefinition& Type, FText& OutErrorMessage);
+	static FString GetHlslDefaultForType(const FNiagaraTypeDefinition& Type);
+	static bool IsHlslBuiltinVector(const FNiagaraTypeDefinition& Type);
 	static TArray<FName> ConditionPropertyPath(const FNiagaraTypeDefinition& Type, const TArray<FName>& InPath);
 
 
