Niagara - Reenable the old parameters panel and put the new one behind a cvar for the 4.25 release as it's not ready yet.

diff --git a/Engine/Plugins/FX/Niagara/Source/NiagaraEditor/Private/NiagaraEditorModule.cpp b/Engine/Plugins/FX/Niagara/Source/NiagaraEditor/Private/NiagaraEditorModule.cpp
--- a/Engine/Plugins/FX/Niagara/Source/NiagaraEditor/Private/NiagaraEditorModule.cpp
+++ b/Engine/Plugins/FX/Niagara/Source/NiagaraEditor/Private/NiagaraEditorModule.cpp
@@ -148,6 +148,14 @@ TArray<TPair<FName, FNiagaraParameterScopeInfo>> FNiagaraEditorModule::Registere
 
 EAssetTypeCategories::Type FNiagaraEditorModule::NiagaraAssetCategory;
 
+int32 GbShowNiagaraDeveloperWindows = 0;
+static FAutoConsoleVariableRef CVarShowNiagaraDeveloperWindows(
+	TEXT("fx.ShowNiagaraDeveloperWindows"),
+	GbShowNiagaraDeveloperWindows,
+	TEXT("If > 0 the niagara system, emitter, and script editors will show additional developer windows.\nThese windows are for niagara tool development and debugging and editing the data\n directly in these windows can cause instability.\n"),
+	ECVF_Default
+	);
+
 const FNiagaraParameterScopeInfo* FNiagaraEditorModule::FindParameterScopeInfo(const FName& ParameterScopeInfoName)
 {
 	auto FindPredicate = [ParameterScopeInfoName](const TPair<FName, FNiagaraParameterScopeInfo>& ScopeInfoPair) {return ScopeInfoPair.Key == ParameterScopeInfoName; };
diff --git a/Engine/Plugins/FX/Niagara/Source/NiagaraEditor/Private/Toolkits/NiagaraScriptToolkit.cpp b/Engine/Plugins/FX/Niagara/Source/NiagaraEditor/Private/Toolkits/NiagaraScriptToolkit.cpp
--- a/Engine/Plugins/FX/Niagara/Source/NiagaraEditor/Private/Toolkits/NiagaraScriptToolkit.cpp
+++ b/Engine/Plugins/FX/Niagara/Source/NiagaraEditor/Private/Toolkits/NiagaraScriptToolkit.cpp
@@ -58,6 +58,7 @@ const FName FNiagaraScriptToolkit::NodeGraphTabId(TEXT("NiagaraEditor_NodeGraph"
 const FName FNiagaraScriptToolkit::ScriptDetailsTabId(TEXT("NiagaraEditor_ScriptDetails"));
 const FName FNiagaraScriptToolkit::SelectedDetailsTabId(TEXT("NiagaraEditor_SelectedDetails"));
 const FName FNiagaraScriptToolkit::ParametersTabId(TEXT("NiagaraEditor_Parameters"));
+const FName FNiagaraScriptToolkit::ParametersTabId2(TEXT("NiagaraEditor_Paramters2"));
 const FName FNiagaraScriptToolkit::StatsTabId(TEXT("NiagaraEditor_Stats"));
 const FName FNiagaraScriptToolkit::MessageLogTabID(TEXT("NiagaraEditor_MessageLog"));
 
@@ -99,6 +100,11 @@ void FNiagaraScriptToolkit::RegisterTabSpawners(const TSharedRef<class FTabManag
 	InTabManager->RegisterTabSpawner(ParametersTabId, FOnSpawnTab::CreateSP(this, &FNiagaraScriptToolkit::SpawnTabScriptParameters))
 		.SetDisplayName(LOCTEXT("ParametersTab", "Parameters"))
 		.SetGroup(WorkspaceMenuCategoryRef);
+
+	InTabManager->RegisterTabSpawner(ParametersTabId2, FOnSpawnTab::CreateSP(this, &FNiagaraScriptToolkit::SpawnTabScriptParameters2))
+		.SetDisplayName(LOCTEXT("ParametersTab2", "Parameters2"))
+		.SetGroup(WorkspaceMenuCategoryRef)
+		.SetAutoGenerateMenuEntry(GbShowNiagaraDeveloperWindows != 0);
 		
 	InTabManager->RegisterTabSpawner(StatsTabId, FOnSpawnTab::CreateSP(this, &FNiagaraScriptToolkit::SpawnTabStats))
 		.SetDisplayName(LOCTEXT("StatsTab", "Stats"))
@@ -118,6 +124,7 @@ void FNiagaraScriptToolkit::UnregisterTabSpawners(const TSharedRef<class FTabMan
 	InTabManager->UnregisterTabSpawner(ScriptDetailsTabId);
 	InTabManager->UnregisterTabSpawner(SelectedDetailsTabId);
 	InTabManager->UnregisterTabSpawner(ParametersTabId);
+	InTabManager->UnregisterTabSpawner(ParametersTabId2);
 	InTabManager->UnregisterTabSpawner(StatsTabId);
 }
 
@@ -170,7 +177,7 @@ void FNiagaraScriptToolkit::Initialize( const EToolkitMode::Type Mode, const TSh
 	StatsListing = MessageLogModule.CreateLogListing("MaterialEditorStats", LogOptions);
 	Stats = MessageLogModule.CreateLogListingWidget(StatsListing.ToSharedRef());
 
-	TSharedRef<FTabManager::FLayout> StandaloneDefaultLayout = FTabManager::NewLayout("Standalone_Niagara_Layout_v10")
+	TSharedRef<FTabManager::FLayout> StandaloneDefaultLayout = FTabManager::NewLayout("Standalone_Niagara_Layout_v11")
 	->AddArea
 	(
 		FTabManager::NewPrimaryArea()->SetOrientation(Orient_Vertical)
@@ -376,6 +383,23 @@ TSharedRef<SDockTab> FNiagaraScriptToolkit::SpawnTabScriptParameters(const FSpaw
 {
 	checkf(Args.GetTabId().TabType == ParametersTabId, TEXT("Wrong tab ID in NiagaraScriptToolkit"));
 
+	TArray<TSharedRef<FNiagaraObjectSelection>> Array;
+	Array.Push(DetailsScriptSelection.ToSharedRef());
+	Array.Push(ScriptViewModel->GetVariableSelection());
+
+	TSharedRef<SDockTab> SpawnedTab =
+		SNew(SDockTab)
+		[
+			SNew(SNiagaraParameterMapView, Array, SNiagaraParameterMapView::EToolkitType::SCRIPT, GetToolkitCommands())
+		];
+
+	return SpawnedTab;
+}
+
+TSharedRef<SDockTab> FNiagaraScriptToolkit::SpawnTabScriptParameters2(const FSpawnTabArgs& Args)
+{
+	checkf(Args.GetTabId().TabType == ParametersTabId2, TEXT("Wrong tab ID in NiagaraScriptToolkit"));
+
 	TSharedRef<SDockTab> SpawnedTab =
 		SNew(SDockTab)
 		[
diff --git a/Engine/Plugins/FX/Niagara/Source/NiagaraEditor/Private/Toolkits/NiagaraScriptToolkit.h b/Engine/Plugins/FX/Niagara/Source/NiagaraEditor/Private/Toolkits/NiagaraScriptToolkit.h
--- a/Engine/Plugins/FX/Niagara/Source/NiagaraEditor/Private/Toolkits/NiagaraScriptToolkit.h
+++ b/Engine/Plugins/FX/Niagara/Source/NiagaraEditor/Private/Toolkits/NiagaraScriptToolkit.h
@@ -89,6 +89,7 @@ private:
 
 	/** Spawns the tab with the script details inside. */
 	TSharedRef<SDockTab> SpawnTabScriptParameters(const FSpawnTabArgs& Args);
+	TSharedRef<SDockTab> SpawnTabScriptParameters2(const FSpawnTabArgs& Args);
 	TSharedRef<SDockTab> SpawnTabStats(const FSpawnTabArgs& Args);
 
 	TSharedPtr< SNiagaraSelectedObjectsDetails> SelectedDetailsWidget;
@@ -149,6 +150,7 @@ private:
 	static const FName ScriptDetailsTabId;
 	static const FName SelectedDetailsTabId;
 	static const FName ParametersTabId;
+	static const FName ParametersTabId2;
 	static const FName StatsTabId;
 	static const FName MessageLogTabID;
 
diff --git a/Engine/Plugins/FX/Niagara/Source/NiagaraEditor/Private/Toolkits/NiagaraSystemToolkit.cpp b/Engine/Plugins/FX/Niagara/Source/NiagaraEditor/Private/Toolkits/NiagaraSystemToolkit.cpp
--- a/Engine/Plugins/FX/Niagara/Source/NiagaraEditor/Private/Toolkits/NiagaraSystemToolkit.cpp
+++ b/Engine/Plugins/FX/Niagara/Source/NiagaraEditor/Private/Toolkits/NiagaraSystemToolkit.cpp
@@ -68,6 +68,7 @@ const FName FNiagaraSystemToolkit::SequencerTabID(TEXT("NiagaraSystemEditor_Sequ
 const FName FNiagaraSystemToolkit::SystemScriptTabID(TEXT("NiagaraSystemEditor_SystemScript"));
 const FName FNiagaraSystemToolkit::SystemDetailsTabID(TEXT("NiagaraSystemEditor_SystemDetails"));
 const FName FNiagaraSystemToolkit::SystemParametersTabID(TEXT("NiagaraSystemEditor_SystemParameters"));
+const FName FNiagaraSystemToolkit::SystemParametersTabID2(TEXT("NiagaraSystemEditor_SystemParameters2"));
 const FName FNiagaraSystemToolkit::SelectedEmitterStackTabID(TEXT("NiagaraSystemEditor_SelectedEmitterStack"));
 const FName FNiagaraSystemToolkit::SelectedEmitterGraphTabID(TEXT("NiagaraSystemEditor_SelectedEmitterGraph"));
 const FName FNiagaraSystemToolkit::DebugSpreadsheetTabID(TEXT("NiagaraSystemEditor_DebugAttributeSpreadsheet"));
@@ -85,14 +86,6 @@ static FAutoConsoleVariableRef CVarSuppressNiagaraSystems(
 	ECVF_Default
 );
 
-static int32 GbShowNiagaraDeveloperWindows = 0;
-static FAutoConsoleVariableRef CVarShowNiagaraDeveloperWindows(
-	TEXT("fx.ShowNiagaraDeveloperWindows"),
-	GbShowNiagaraDeveloperWindows,
-	TEXT("If > 0 the niagara system and emitter editors will show additional developer windows.\nThese windows are for niagara tool development and debugging and editing the data\n directly in these windows can cause instability.\n"),
-	ECVF_Default
-);
-
 bool FNiagaraSystemToolkit::bShowLibraryOnly = false;
 bool FNiagaraSystemToolkit::bShowTemplateOnly = false;
 
@@ -124,6 +117,11 @@ void FNiagaraSystemToolkit::RegisterTabSpawners(const TSharedRef<class FTabManag
 		.SetDisplayName(LOCTEXT("SystemParameters", "Parameters"))
 		.SetGroup(WorkspaceMenuCategory.ToSharedRef());
 
+	InTabManager->RegisterTabSpawner(SystemParametersTabID2, FOnSpawnTab::CreateSP(this, &FNiagaraSystemToolkit::SpawnTab_SystemParameters2))
+		.SetDisplayName(LOCTEXT("SystemParameters2", "Parameters2"))
+		.SetGroup(WorkspaceMenuCategory.ToSharedRef())
+		.SetAutoGenerateMenuEntry(GbShowNiagaraDeveloperWindows != 0);
+
 	InTabManager->RegisterTabSpawner(SelectedEmitterStackTabID, FOnSpawnTab::CreateSP(this, &FNiagaraSystemToolkit::SpawnTab_SelectedEmitterStack))
 		.SetDisplayName(LOCTEXT("SelectedEmitterStacks", "Selected Emitters"))
 		.SetGroup(WorkspaceMenuCategory.ToSharedRef());
@@ -324,7 +322,7 @@ void FNiagaraSystemToolkit::InitializeInternal(const EToolkitMode::Type Mode, co
 	const float InTime = -0.02f;
 	const float OutTime = 3.2f;
 
-	TSharedRef<FTabManager::FLayout> StandaloneDefaultLayout = FTabManager::NewLayout("Standalone_Niagara_System_Layout_v23")
+	TSharedRef<FTabManager::FLayout> StandaloneDefaultLayout = FTabManager::NewLayout("Standalone_Niagara_System_Layout_v24")
 		->AddArea
 		(
 			FTabManager::NewPrimaryArea()->SetOrientation(Orient_Vertical)
@@ -521,6 +519,24 @@ TSharedRef<SDockTab> FNiagaraSystemToolkit::SpawnTab_SystemParameters(const FSpa
 	check(Args.GetTabId().TabType == SystemParametersTabID);
 
 
+	TArray<TSharedRef<FNiagaraObjectSelection>> ObjectSelections;
+	ObjectSelections.Add(ObjectSelectionForParameterMapView.ToSharedRef());
+
+	TSharedRef<SDockTab> SpawnedTab =
+		SNew(SDockTab)
+		[
+			SAssignNew(ParameterMapView, SNiagaraParameterMapView, ObjectSelections, SNiagaraParameterMapView::EToolkitType::SYSTEM, GetToolkitCommands()) //@todo(ng) cleanup
+		];
+	RefreshParameters();
+
+	return SpawnedTab;
+}
+
+TSharedRef<SDockTab> FNiagaraSystemToolkit::SpawnTab_SystemParameters2(const FSpawnTabArgs& Args)
+{
+	check(Args.GetTabId().TabType == SystemParametersTabID2);
+
+
 	TArray<TSharedRef<FNiagaraObjectSelection>> ObjectSelections;
 	ObjectSelections.Add(ObjectSelectionForParameterMapView.ToSharedRef());
 
diff --git a/Engine/Plugins/FX/Niagara/Source/NiagaraEditor/Private/Toolkits/NiagaraSystemToolkit.h b/Engine/Plugins/FX/Niagara/Source/NiagaraEditor/Private/Toolkits/NiagaraSystemToolkit.h
--- a/Engine/Plugins/FX/Niagara/Source/NiagaraEditor/Private/Toolkits/NiagaraSystemToolkit.h
+++ b/Engine/Plugins/FX/Niagara/Source/NiagaraEditor/Private/Toolkits/NiagaraSystemToolkit.h
@@ -101,6 +101,7 @@ private:
 	TSharedRef<SDockTab> SpawnTab_Sequencer(const FSpawnTabArgs& Args);
 	TSharedRef<SDockTab> SpawnTab_SystemScript(const FSpawnTabArgs& Args);
 	TSharedRef<SDockTab> SpawnTab_SystemParameters(const FSpawnTabArgs& Args);
+	TSharedRef<SDockTab> SpawnTab_SystemParameters2(const FSpawnTabArgs& Args);
 	TSharedRef<SDockTab> SpawnTab_SelectedEmitterStack(const FSpawnTabArgs& Args);
 	TSharedRef<SDockTab> SpawnTab_SelectedEmitterGraph(const FSpawnTabArgs& Args);
 	TSharedRef<SDockTab> SpawnTab_DebugSpreadsheet(const FSpawnTabArgs& Args);
@@ -172,6 +173,8 @@ private:
 	/** The command list for this editor */
 	TSharedPtr<FUICommandList> EditorCommands;
 
+	TSharedPtr<class SNiagaraParameterMapView> ParameterMapView; //@todo(ng) cleanup
+
 	TSharedPtr<FNiagaraSystemToolkitParameterPanelViewModel> ParameterPanelViewModel;
 	TSharedPtr<class SNiagaraParameterPanel> ParameterPanel;
 
@@ -188,6 +191,7 @@ public:
 	static const FName SystemScriptTabID;
 	static const FName SystemDetailsTabID;
 	static const FName SystemParametersTabID;
+	static const FName SystemParametersTabID2;
 	static const FName SelectedEmitterStackTabID;
 	static const FName SelectedEmitterGraphTabID;
 	static const FName DebugSpreadsheetTabID;
diff --git a/Engine/Plugins/FX/Niagara/Source/NiagaraEditor/Public/NiagaraEditorModule.h b/Engine/Plugins/FX/Niagara/Source/NiagaraEditor/Public/NiagaraEditorModule.h
--- a/Engine/Plugins/FX/Niagara/Source/NiagaraEditor/Public/NiagaraEditorModule.h
+++ b/Engine/Plugins/FX/Niagara/Source/NiagaraEditor/Public/NiagaraEditorModule.h
@@ -32,6 +32,8 @@ class FHlslNiagaraCompiler;
 
 DECLARE_STATS_GROUP(TEXT("Niagara Editor"), STATGROUP_NiagaraEditor, STATCAT_Advanced);
 
+extern int32 GbShowNiagaraDeveloperWindows;
+
 /* Defines methods for allowing external modules to supply widgets to the core editor module. */
 class NIAGARAEDITOR_API INiagaraEditorWidgetProvider
 {
