Removed some unnecessary log spam

diff --git a/Engine/Plugins/FX/Niagara/Source/NiagaraShader/Private/NiagaraShader.cpp b/Engine/Plugins/FX/Niagara/Source/NiagaraShader/Private/NiagaraShader.cpp
--- a/Engine/Plugins/FX/Niagara/Source/NiagaraShader/Private/NiagaraShader.cpp
+++ b/Engine/Plugins/FX/Niagara/Source/NiagaraShader/Private/NiagaraShader.cpp
@@ -1045,8 +1045,9 @@ void FNiagaraShaderMap::Serialize(FArchive& Ar, bool bInlineShaderResources, boo
 	Super::Serialize(Ar, bInlineShaderResources, bLoadedByCookedMaterial);
 }
 
-void FNiagaraShaderMap::RemovePendingScript(FNiagaraShaderScript* Script)
+bool FNiagaraShaderMap::RemovePendingScript(FNiagaraShaderScript* Script)
 {
+	bool bRemoved = false;
 	//All access to NiagaraShaderMapsBeingCompiled must be done on the game thread!
 	check(IsInGameThread());
 	for (TMap<TRefCountPtr<FNiagaraShaderMap>, TArray<FNiagaraShaderScript*> >::TIterator It(NiagaraShaderMapsBeingCompiled); It; ++It)
@@ -1057,6 +1058,7 @@ void FNiagaraShaderMap::RemovePendingScript(FNiagaraShaderScript* Script)
 		{
 			Script->RemoveOutstandingCompileId(It.Key()->CompilingId);
 			Script->NotifyCompilationFinished();
+			bRemoved = true;
 		}
 #if DEBUG_INFINITESHADERCOMPILE
 		if ( Result )
@@ -1065,6 +1067,8 @@ void FNiagaraShaderMap::RemovePendingScript(FNiagaraShaderScript* Script)
 		}
 #endif
 	}
+
+	return bRemoved;
 }
 
 
diff --git a/Engine/Plugins/FX/Niagara/Source/NiagaraShader/Private/NiagaraShared.cpp b/Engine/Plugins/FX/Niagara/Source/NiagaraShader/Private/NiagaraShared.cpp
--- a/Engine/Plugins/FX/Niagara/Source/NiagaraShader/Private/NiagaraShared.cpp
+++ b/Engine/Plugins/FX/Niagara/Source/NiagaraShader/Private/NiagaraShared.cpp
@@ -79,10 +79,14 @@ NIAGARASHADER_API void FNiagaraShaderScript::CancelCompilation()
 {
 #if WITH_EDITOR
 	check(IsInGameThread());
-	FNiagaraShaderMap::RemovePendingScript(this);
+	bool bWasPending = FNiagaraShaderMap::RemovePendingScript(this);
 	FNiagaraCompilationQueue::Get()->RemovePending(this);
 
-	UE_LOG(LogShaders, Log, TEXT("CancelCompilation %p."), this);
+	// don't spam the log if no cancelling actually happened : 
+	if (bWasPending)
+	{
+		UE_LOG(LogShaders, Log, TEXT("CancelCompilation %p."), this);
+	}
 	OutstandingCompileShaderMapIds.Empty();
 #endif
 }
diff --git a/Engine/Plugins/FX/Niagara/Source/NiagaraShader/Public/NiagaraShared.h b/Engine/Plugins/FX/Niagara/Source/NiagaraShader/Public/NiagaraShared.h
--- a/Engine/Plugins/FX/Niagara/Source/NiagaraShader/Public/NiagaraShared.h
+++ b/Engine/Plugins/FX/Niagara/Source/NiagaraShader/Public/NiagaraShared.h
@@ -454,8 +454,10 @@ public:
 	*/
 	void FlushShadersByShaderType(const FShaderType* ShaderType);
 
-	/** Removes a Script from NiagaraShaderMapsBeingCompiled. */
-	NIAGARASHADER_API static void RemovePendingScript(FNiagaraShaderScript* Script);
+	/** Removes a Script from NiagaraShaderMapsBeingCompiled. 
+	* @return true if something was actually removed.
+	*/
+	NIAGARASHADER_API static bool RemovePendingScript(FNiagaraShaderScript* Script);
 	NIAGARASHADER_API static void RemovePendingMap(FNiagaraShaderMap* Map);
 
 	/** Finds a shader map currently being compiled that was enqueued for the given script. */
