- Fix GPU deterministic random inside skeletal mesh and replace screenshot so test passes

#rnx
#tests EngineTest on multiple platforms

diff --git a/Engine/Plugins/FX/Niagara/Shaders/Private/NiagaraDataInterfaceSkeletalMesh.ush b/Engine/Plugins/FX/Niagara/Shaders/Private/NiagaraDataInterfaceSkeletalMesh.ush
--- a/Engine/Plugins/FX/Niagara/Shaders/Private/NiagaraDataInterfaceSkeletalMesh.ush
+++ b/Engine/Plugins/FX/Niagara/Shaders/Private/NiagaraDataInterfaceSkeletalMesh.ush
@@ -55,6 +55,12 @@ struct FDISkelMeshContext
 ////////////////////////////////////////////////////////////////////////////////////////////////////////////////////////////////////////////////////////////////
 // Triangle Sampling
 
+//-TODO: We need to consolodate the random interface inside NiagaraEmitterInstanceShader.usf it's a bit messy at the moment but should follow what FNDIRandomHelper does
+float DISKelMesh_Random(uint S1, uint S2, uint S3)
+{
+	return S3 == 0xffffffff ? NiagaraInternalNoise(S1, S2, S3) : rand(1.0f, S1, S2, S3);
+}
+
 struct FDISkelMeshSkinnedVertex
 {
 	float3	Position;
@@ -66,8 +72,8 @@ struct FDISkelMeshSkinnedVertex
 
 float3 DISKelMesh_RandomBarycentricCoord(uint Seed1, uint Seed2, uint Seed3)
 {
-	float r0 = NiagaraInternalNoise(Seed1, Seed2, Seed3);
-	float r1 = NiagaraInternalNoise(Seed1, Seed2, Seed3);
+	float r0 = DISKelMesh_Random(Seed1, Seed2, Seed3);
+	float r1 = DISKelMesh_Random(Seed1, Seed2, Seed3);
 	float sqrt0 = sqrt(r0);
 	float sqrt1 = sqrt(r1);
 	return float3(1.0f - sqrt0, sqrt0 * (1.0 - r1), r1 * sqrt0);
@@ -241,7 +247,7 @@ FDISkelMeshSkinnedVertex DISkelMesh_GetSkinnedPointOnTriangleWS(in FDISkelMeshCo
 
 void DISKelMesh_RandomTriCoord(in FDISkelMeshContext DIContext, uint Seed1, uint Seed2, uint Seed3, out int OutTriangle, out float3 OutBaryCoord)
 {
-	float RandT0 = NiagaraInternalNoise(Seed1, Seed2, Seed3);
+	float RandT0 = DISKelMesh_Random(Seed1, Seed2, Seed3);
 
 	[branch]
 	if ( !DIContext.UniformTriangleSamplingEnable )
@@ -256,7 +262,7 @@ void DISKelMesh_RandomTriCoord(in FDISkelMeshContext DIContext, uint Seed1, uint
 		float TriangleProbability = DIContext.MeshTriangleSamplerProbaBuffer[TriangleIndex];
 
 		// Alias check
-		float RandT1 = NiagaraInternalNoise(Seed1, Seed2, Seed3);
+		float RandT1 = DISKelMesh_Random(Seed1, Seed2, Seed3);
 		if( RandT1 > TriangleProbability )
 		{
 			TriangleIndex = DIContext.MeshTriangleSamplerAliasBuffer[TriangleIndex];
@@ -332,7 +338,7 @@ void DISKelMesh_GetTriUV(in FDISkelMeshContext DIContext, in uint TriangleIndex,
 
 void DISKelMesh_RandomTriangle(in FDISkelMeshContext DIContext, uint Seed1, uint Seed2, uint Seed3, out int OutTriangle, out float3 OutBaryCoord)
 {
-	float RandT0 = NiagaraInternalNoise(Seed1, Seed2, Seed3);
+	float RandT0 = DISKelMesh_Random(Seed1, Seed2, Seed3);
 	int TriangleIndex = min(int(RandT0 * float(DIContext.MeshTriangleCount)), (int)DIContext.MeshTriangleCount - 1); // avoid % by using mul/min to Tri = MeshTriangleCountName
 
 	[branch]
@@ -342,7 +348,7 @@ void DISKelMesh_RandomTriangle(in FDISkelMeshContext DIContext, uint Seed1, uint
 		float TriangleProbability = DIContext.MeshTriangleSamplerProbaBuffer[TriangleIndex];
 
 		// Alias check
-		float RandT1 = NiagaraInternalNoise(Seed1, Seed2, Seed3);
+		float RandT1 = DISKelMesh_Random(Seed1, Seed2, Seed3);
 		if( RandT1 > TriangleProbability )
 		{
 			TriangleIndex = (int)DIContext.MeshTriangleSamplerAliasBuffer[TriangleIndex];
@@ -360,14 +366,14 @@ void DISKelMesh_GetTriangleCount(in FDISkelMeshContext DIContext, out int OutCou
 
 void DISKelMesh_RandomFilteredTriangle(in FDISkelMeshContext DIContext, uint Seed1, uint Seed2, uint Seed3, out int OutTriangle, out float3 OutBaryCoord)
 {
-	float RandT0 = NiagaraInternalNoise(Seed1, Seed2, Seed3);
+	float RandT0 = DISKelMesh_Random(Seed1, Seed2, Seed3);
 	uint RegionTriangle = min(uint(RandT0 * float(DIContext.MeshNumSamplingRegionTriangles)), DIContext.MeshNumSamplingRegionTriangles - 1); // avoid % by using mul/min to Tri = MeshTriangleCountName
 
 	[branch]
 	if ( DIContext.UniformSamplingRegionEnabled )
 	{
 		float TriangleProbability = DIContext.MeshSamplingRegionsProbaBuffer[RegionTriangle];
-		float RandT1 = NiagaraInternalNoise(Seed1, Seed2, Seed3);
+		float RandT1 = DISKelMesh_Random(Seed1, Seed2, Seed3);
 		if ( RandT1 > TriangleProbability )
 		{
 			RegionTriangle = (int)DIContext.MeshSamplingRegionsAliasBuffer[RegionTriangle];
@@ -530,7 +536,7 @@ void DISkelMesh_IsValidVertex(in FDISkelMeshContext DIContext, int Vertex, out b
 
 void DISkelMesh_GetRandomVertex(in FDISkelMeshContext DIContext, uint Seed1, uint Seed2, uint Seed3, out int OutVertex)
 {
-	float RandT0 = NiagaraInternalNoise(Seed1, Seed2, Seed3);
+	float RandT0 = DISKelMesh_Random(Seed1, Seed2, Seed3);
 	OutVertex = (int)min(uint(RandT0 * float(DIContext.MeshVertexCount)), DIContext.MeshVertexCount - 1); // avoid % by using mul/min to Tri = MeshVertexCount
 }
 
@@ -546,7 +552,7 @@ void DISkelMesh_IsValidFilteredVertex(in FDISkelMeshContext DIContext, int Filte
 
 void DISkelMesh_GetRandomFilteredVertex(in FDISkelMeshContext DIContext, uint Seed1, uint Seed2, uint Seed3, out int OutVertex)
 {
-	float RandT0 = NiagaraInternalNoise(Seed1, Seed2, Seed3);
+	float RandT0 = DISKelMesh_Random(Seed1, Seed2, Seed3);
 	int FilteredIndex = (int)min(uint(RandT0 * float(DIContext.MeshNumSamplingRegionVertices)), DIContext.MeshNumSamplingRegionVertices - 1); // avoid % by using mul/min to Tri = MeshVertexCount
 	OutVertex = (int)DIContext.MeshSampleRegionsVertices[FilteredIndex];
 }
@@ -571,7 +577,7 @@ void DISkelMesh_IsValidBone(in FDISkelMeshContext DIContext, in int BoneIndex, o
 
 void DISkelMesh_RandomBone(in FDISkelMeshContext DIContext, uint Seed1, uint Seed2, uint Seed3, out int Bone)
 {
-	float RandT0 = NiagaraInternalNoise(Seed1, Seed2, Seed3);
+	float RandT0 = DISKelMesh_Random(Seed1, Seed2, Seed3);
     Bone = min(int(RandT0 * float(DIContext.NumBones)), DIContext.NumBones - 1); // avoid % by using mul/min to Tri = MeshVertexCount
 }
 
@@ -593,7 +599,7 @@ void DISkelMesh_GetFilteredBoneAt(in FDISkelMeshContext DIContext, in int InBone
 
 void DISkelMesh_RandomFilteredBone(in FDISkelMeshContext DIContext, uint Seed1, uint Seed2, uint Seed3, out int Bone)
 {
-	float RandT0 = NiagaraInternalNoise(Seed1, Seed2, Seed3);
+	float RandT0 = DISKelMesh_Random(Seed1, Seed2, Seed3);
     int FilteredBone = min(int(RandT0 * float(DIContext.NumFilteredBones)), DIContext.NumFilteredBones - 1); // avoid % by using mul/min to Tri = MeshVertexCount
     Bone = (int)DIContext.FilteredAndUnfilteredBones[FilteredBone];
 }
@@ -611,7 +617,7 @@ void DISkelMesh_GetUnfilteredBoneAt(in FDISkelMeshContext DIContext, in int InBo
 
 void DISkelMesh_RandomUnfilteredBone(in FDISkelMeshContext DIContext, uint Seed1, uint Seed2, uint Seed3, out int Bone)
 {
-	float RandT0 = NiagaraInternalNoise(Seed1, Seed2, Seed3);
+	float RandT0 = DISKelMesh_Random(Seed1, Seed2, Seed3);
     int UnfilteredBone = min(int(RandT0 * float(DIContext.NumUnfilteredBones)), DIContext.NumUnfilteredBones - 1); // avoid % by using mul/min to Tri = MeshVertexCount
     Bone = DIContext.NumFilteredBones > 0 ? (int)DIContext.FilteredAndUnfilteredBones[UnfilteredBone + DIContext.NumFilteredBones] : UnfilteredBone;
 }
@@ -628,13 +634,13 @@ void DISkelMesh_GetFilteredSocketBoneAt(in FDISkelMeshContext DIContext, in int
 
 void DISkelMesh_RandomFilteredSocket(in FDISkelMeshContext DIContext, uint Seed1, uint Seed2, uint Seed3, out int SocketBone)
 {
-	float RandT0 = NiagaraInternalNoise(Seed1, Seed2, Seed3);
+	float RandT0 = DISKelMesh_Random(Seed1, Seed2, Seed3);
     SocketBone = DIContext.FilteredSocketBoneOffset + min(int(RandT0 * float(DIContext.NumFilteredSockets)), DIContext.NumFilteredSockets - 1); // avoid % by using mul/min to Tri = MeshVertexCount
 }
 
 void DISkelMesh_RandomFilteredSocketOrBone(in FDISkelMeshContext DIContext, uint Seed1, uint Seed2, uint Seed3, out int Bone)
 {
-	float RandT0 = NiagaraInternalNoise(Seed1, Seed2, Seed3);
+	float RandT0 = DISKelMesh_Random(Seed1, Seed2, Seed3);
 	int NumFilteredSocketsAndBones = DIContext.NumFilteredBones + DIContext.NumFilteredSockets;
     int FilteredIndex = min(int(RandT0 * float(NumFilteredSocketsAndBones)), NumFilteredSocketsAndBones - 1);
 	if ( FilteredIndex < DIContext.NumFilteredBones )
diff --git a/Engine/Plugins/FX/Niagara/Source/NiagaraCore/Private/NiagaraCustomVersion.cpp b/Engine/Plugins/FX/Niagara/Source/NiagaraCore/Private/NiagaraCustomVersion.cpp
--- a/Engine/Plugins/FX/Niagara/Source/NiagaraCore/Private/NiagaraCustomVersion.cpp
+++ b/Engine/Plugins/FX/Niagara/Source/NiagaraCore/Private/NiagaraCustomVersion.cpp
@@ -8,4 +8,4 @@ const FGuid FNiagaraCustomVersion::GUID(0xFCF57AFA, 0x50764283, 0xB9A9E658, 0xFF
 // Register the custom version with core
 FCustomVersionRegistration GRegisterNiagaraCustomVersion(FNiagaraCustomVersion::GUID, FNiagaraCustomVersion::LatestVersion, TEXT("NiagaraVer"));
 
-const FGuid FNiagaraCustomVersion::LatestScriptCompileVersion(0x08611BB6, 0x382AD540, 0xAD008BCA, 0xD3CBDFEC);
+const FGuid FNiagaraCustomVersion::LatestScriptCompileVersion(0x62AE44D9, 0x32500143, 0x99B84C1B, 0x354729F9);
