Get rid of PRAGMA_DISABLE_OPTIMIZATION block left in the Niagara Oscilloscope.

diff --git a/Engine/Plugins/FX/Niagara/Source/Niagara/Private/NiagaraDataInterfaceAudioOscilloscope.cpp b/Engine/Plugins/FX/Niagara/Source/Niagara/Private/NiagaraDataInterfaceAudioOscilloscope.cpp
--- a/Engine/Plugins/FX/Niagara/Source/Niagara/Private/NiagaraDataInterfaceAudioOscilloscope.cpp
+++ b/Engine/Plugins/FX/Niagara/Source/Niagara/Private/NiagaraDataInterfaceAudioOscilloscope.cpp
@@ -124,7 +124,7 @@ UNiagaraDataInterfaceAudioOscilloscope::UNiagaraDataInterfaceAudioOscilloscope(F
 {
 	Proxy = TUniquePtr<FNiagaraDataInterfaceProxyOscilloscope>(new FNiagaraDataInterfaceProxyOscilloscope(Resolution, ScopeInMilliseconds));
 }
-PRAGMA_DISABLE_OPTIMIZATION
+
 float FNiagaraDataInterfaceProxyOscilloscope::SampleAudio(float NormalizedPositionInBuffer, int32 ChannelIndex, int32 NumFramesInBuffer, int32 NumChannelsInBuffer)
 {
 	if (NumFramesInBuffer == 0 || NumChannelsInBuffer == 0)
@@ -181,8 +181,6 @@ void UNiagaraDataInterfaceAudioOscilloscope::GetNumChannels(FVectorVMContext& Co
 	}
 }
 
-PRAGMA_ENABLE_OPTIMIZATION
-
 void UNiagaraDataInterfaceAudioOscilloscope::GetFunctions(TArray<FNiagaraFunctionSignature>& OutFunctions)
 {
 	Super::GetFunctions(OutFunctions);
