- Add missing functionality for excluded bone on the GPU

#rnx

diff --git a/Engine/Plugins/FX/Niagara/Shaders/Private/NiagaraDataInterfaceSkeletalMesh.ush b/Engine/Plugins/FX/Niagara/Shaders/Private/NiagaraDataInterfaceSkeletalMesh.ush
--- a/Engine/Plugins/FX/Niagara/Shaders/Private/NiagaraDataInterfaceSkeletalMesh.ush
+++ b/Engine/Plugins/FX/Niagara/Shaders/Private/NiagaraDataInterfaceSkeletalMesh.ush
@@ -41,6 +41,8 @@ struct FDISkelMeshContext
 	int				NumBones;
 	int				NumFilteredBones;
 	int				NumUnfilteredBones;
+	int				RandomMaxBone;
+	int				ExcludeBoneIndex;
 	Buffer<uint>	FilteredAndUnfilteredBones;
 	int				NumFilteredSockets;
 	int				FilteredSocketBoneOffset;
@@ -578,7 +580,8 @@ void DISkelMesh_IsValidBone(in FDISkelMeshContext DIContext, in int BoneIndex, o
 void DISkelMesh_RandomBone(in FDISkelMeshContext DIContext, uint Seed1, uint Seed2, uint Seed3, out int Bone)
 {
 	float RandT0 = DISKelMesh_Random(Seed1, Seed2, Seed3);
-    Bone = min(int(RandT0 * float(DIContext.NumBones)), DIContext.NumBones - 1); // avoid % by using mul/min to Tri = MeshVertexCount
+	Bone = min(int(RandT0 * float(DIContext.RandomMaxBone)), DIContext.RandomMaxBone);
+	Bone = Bone != DIContext.ExcludeBoneIndex ? Bone : Bone + 1;
 }
 
 void DISkelMesh_GetBoneCount(in FDISkelMeshContext DIContext, out int Count)
@@ -618,8 +621,16 @@ void DISkelMesh_GetUnfilteredBoneAt(in FDISkelMeshContext DIContext, in int InBo
 void DISkelMesh_RandomUnfilteredBone(in FDISkelMeshContext DIContext, uint Seed1, uint Seed2, uint Seed3, out int Bone)
 {
 	float RandT0 = DISKelMesh_Random(Seed1, Seed2, Seed3);
-    int UnfilteredBone = min(int(RandT0 * float(DIContext.NumUnfilteredBones)), DIContext.NumUnfilteredBones - 1); // avoid % by using mul/min to Tri = MeshVertexCount
-    Bone = DIContext.NumFilteredBones > 0 ? (int)DIContext.FilteredAndUnfilteredBones[UnfilteredBone + DIContext.NumFilteredBones] : UnfilteredBone;
+	if ( DIContext.NumFilteredBones == 0 )
+	{
+		Bone = min(int(RandT0 * float(DIContext.RandomMaxBone)), DIContext.RandomMaxBone);
+		Bone = Bone != DIContext.ExcludeBoneIndex ? Bone : Bone + 1;
+	}
+	else
+	{
+		int UnfilteredBone = min(int(RandT0 * float(DIContext.NumUnfilteredBones)), DIContext.NumUnfilteredBones - 1); // avoid % by using mul/min to Tri = MeshVertexCount
+		Bone = DIContext.NumFilteredBones > 0 ? (int)DIContext.FilteredAndUnfilteredBones[UnfilteredBone + DIContext.NumFilteredBones] : UnfilteredBone;
+	}
 }
 
 void DISkelMesh_GetFilteredSocketCount(in FDISkelMeshContext DIContext, out int Count)
@@ -704,6 +715,8 @@ void DISkelMesh_GetFilteredSocketOrBoneAt(in FDISkelMeshContext DIContext, int F
 	int				NumBones_ ## NAME; \
 	int				NumFilteredBones_ ## NAME; \
 	int				NumUnfilteredBones_ ## NAME; \
+	int				RandomMaxBone_ ## NAME; \
+	int				ExcludeBoneIndex_ ## NAME; \
 	Buffer<uint>	FilteredAndUnfilteredBones_ ## NAME; \
 	int				NumFilteredSockets_ ## NAME; \
 	int				FilteredSocketBoneOffset_ ## NAME; \
@@ -741,6 +754,8 @@ void DISkelMesh_GetFilteredSocketOrBoneAt(in FDISkelMeshContext DIContext, int F
 	DIContext.NumBones = NumBones_ ## NAME; \
 	DIContext.NumFilteredBones	= NumFilteredBones_ ## NAME; \
 	DIContext.NumUnfilteredBones = NumUnfilteredBones_ ## NAME; \
+	DIContext.RandomMaxBone = RandomMaxBone_ ## NAME; \
+	DIContext.ExcludeBoneIndex = ExcludeBoneIndex_ ## NAME; \
 	DIContext.FilteredAndUnfilteredBones	= FilteredAndUnfilteredBones_ ## NAME; \
 	DIContext.NumFilteredSockets = NumFilteredSockets_ ## NAME; \
 	DIContext.FilteredSocketBoneOffset = FilteredSocketBoneOffset_ ## NAME; \
diff --git a/Engine/Plugins/FX/Niagara/Source/Niagara/Classes/NiagaraDataInterfaceSkeletalMesh.h b/Engine/Plugins/FX/Niagara/Source/Niagara/Classes/NiagaraDataInterfaceSkeletalMesh.h
--- a/Engine/Plugins/FX/Niagara/Source/Niagara/Classes/NiagaraDataInterfaceSkeletalMesh.h
+++ b/Engine/Plugins/FX/Niagara/Source/Niagara/Classes/NiagaraDataInterfaceSkeletalMesh.h
@@ -325,6 +325,7 @@ public:
 
 	int32 GetNumFilteredBones() const { return NumFilteredBones; }
 	int32 GetNumUnfilteredBones() const { return NumUnfilteredBones;  }
+	int32 GetExcludedBoneIndex() const { return ExcludedBoneIndex; }
 	FRHIShaderResourceView* GetFilteredAndUnfilteredBonesSRV() const { return FilteredAndUnfilteredBonesSRV; }
 
 	int32 GetNumFilteredSockets() const { return NumFilteredSockets; }
@@ -358,6 +359,7 @@ protected:
 
 	int32 NumFilteredBones = 0;
 	int32 NumUnfilteredBones = 0;
+	int32 ExcludedBoneIndex = INDEX_NONE;
 	TResourceArray<uint16> FilteredAndUnfilteredBonesArray;
 	FVertexBufferRHIRef FilteredAndUnfilteredBonesBuffer;
 	FShaderResourceViewRHIRef FilteredAndUnfilteredBonesSRV;
@@ -688,6 +690,8 @@ public:
 	static const FString NumBonesName;
 	static const FString NumFilteredBonesName;
 	static const FString NumUnfilteredBonesName;
+	static const FString RandomMaxBoneName;
+	static const FString ExcludeBoneIndexName;
 	static const FString FilteredAndUnfilteredBonesName;
 	static const FString NumFilteredSocketsName;
 	static const FString FilteredSocketBoneOffsetName;
diff --git a/Engine/Plugins/FX/Niagara/Source/Niagara/Private/NiagaraDataInterfaceSkeletalMesh.cpp b/Engine/Plugins/FX/Niagara/Source/Niagara/Private/NiagaraDataInterfaceSkeletalMesh.cpp
--- a/Engine/Plugins/FX/Niagara/Source/Niagara/Private/NiagaraDataInterfaceSkeletalMesh.cpp
+++ b/Engine/Plugins/FX/Niagara/Source/Niagara/Private/NiagaraDataInterfaceSkeletalMesh.cpp
@@ -501,6 +501,7 @@ void FSkeletalMeshGpuSpawnStaticBuffers::Initialise(FNDISkeletalMesh_InstanceDat
 		//-TODO: Exclude setting up these arrays if we don't sample from them
 		NumFilteredBones = InstData->NumFilteredBones;
 		NumUnfilteredBones = InstData->NumUnfilteredBones;
+		ExcludedBoneIndex = InstData->ExcludedBoneIndex;
 
 		FilteredAndUnfilteredBonesArray.Reserve(InstData->FilteredAndUnfilteredBones.Num());
 		for (uint16 v : InstData->FilteredAndUnfilteredBones)
@@ -957,6 +958,8 @@ struct FNDISkeletalMeshParametersName
 	FString NumBonesName;
 	FString NumFilteredBonesName;
 	FString NumUnfilteredBonesName;
+	FString RandomMaxBoneName;
+	FString ExcludeBoneIndexName;
 	FString FilteredAndUnfilteredBonesName;
 	FString NumFilteredSocketsName;
 	FString FilteredSocketBoneOffsetName;
@@ -997,6 +1000,8 @@ static void GetNiagaraDataInterfaceParametersName(FNDISkeletalMeshParametersName
 	Names.NumBonesName = UNiagaraDataInterfaceSkeletalMesh::NumBonesName + Suffix;
 	Names.NumFilteredBonesName = UNiagaraDataInterfaceSkeletalMesh::NumFilteredBonesName + Suffix;
 	Names.NumUnfilteredBonesName = UNiagaraDataInterfaceSkeletalMesh::NumUnfilteredBonesName + Suffix;
+	Names.RandomMaxBoneName = UNiagaraDataInterfaceSkeletalMesh::RandomMaxBoneName + Suffix;
+	Names.ExcludeBoneIndexName = UNiagaraDataInterfaceSkeletalMesh::ExcludeBoneIndexName + Suffix;
 	Names.FilteredAndUnfilteredBonesName = UNiagaraDataInterfaceSkeletalMesh::FilteredAndUnfilteredBonesName + Suffix;
 	Names.NumFilteredSocketsName = UNiagaraDataInterfaceSkeletalMesh::NumFilteredSocketsName + Suffix;
 	Names.FilteredSocketBoneOffsetName = UNiagaraDataInterfaceSkeletalMesh::FilteredSocketBoneOffsetName + Suffix;
@@ -1044,6 +1049,8 @@ public:
 		NumBones.Bind(ParameterMap, *ParamNames.NumBonesName);
 		NumFilteredBones.Bind(ParameterMap, *ParamNames.NumFilteredBonesName);
 		NumUnfilteredBones.Bind(ParameterMap, *ParamNames.NumUnfilteredBonesName);
+		RandomMaxBone.Bind(ParameterMap, *ParamNames.RandomMaxBoneName);
+		ExcludeBoneIndex.Bind(ParameterMap, *ParamNames.ExcludeBoneIndexName);
 		FilteredAndUnfilteredBones.Bind(ParameterMap, *ParamNames.FilteredAndUnfilteredBonesName);
 		NumFilteredSockets.Bind(ParameterMap, *ParamNames.NumFilteredSocketsName);
 		FilteredSocketBoneOffset.Bind(ParameterMap, *ParamNames.FilteredSocketBoneOffsetName);
@@ -1153,6 +1160,8 @@ public:
 			FRHIShaderResourceView* FilteredAndUnfilteredBonesSRV = StaticBuffers->GetNumFilteredBones() > 0 ? StaticBuffers->GetFilteredAndUnfilteredBonesSRV() : FNiagaraRenderer::GetDummyUIntBuffer();
 			SetShaderValue(RHICmdList, ComputeShaderRHI, NumFilteredBones, StaticBuffers->GetNumFilteredBones());
 			SetShaderValue(RHICmdList, ComputeShaderRHI, NumUnfilteredBones, StaticBuffers->GetNumUnfilteredBones());
+			SetShaderValue(RHICmdList, ComputeShaderRHI, RandomMaxBone, StaticBuffers->GetExcludedBoneIndex() >= 0 ? DynamicBuffers->GetNumBones() - 2 : DynamicBuffers->GetNumBones() - 1);
+			SetShaderValue(RHICmdList, ComputeShaderRHI, ExcludeBoneIndex, StaticBuffers->GetExcludedBoneIndex());
 			SetSRVParameter(RHICmdList, ComputeShaderRHI, FilteredAndUnfilteredBones, FilteredAndUnfilteredBonesSRV);
 
 			SetShaderValue(RHICmdList, ComputeShaderRHI, NumFilteredSockets, StaticBuffers->GetNumFilteredSockets());
@@ -1204,6 +1213,8 @@ public:
 			SetShaderValue(RHICmdList, ComputeShaderRHI, NumBones, 0);
 			SetShaderValue(RHICmdList, ComputeShaderRHI, NumFilteredBones, 0);
 			SetShaderValue(RHICmdList, ComputeShaderRHI, NumUnfilteredBones, 0);
+			SetShaderValue(RHICmdList, ComputeShaderRHI, RandomMaxBone, 0);
+			SetShaderValue(RHICmdList, ComputeShaderRHI, ExcludeBoneIndex, 0);
 			SetSRVParameter(RHICmdList, ComputeShaderRHI, FilteredAndUnfilteredBones, FNiagaraRenderer::GetDummyUIntBuffer());
 			SetShaderValue(RHICmdList, ComputeShaderRHI, NumFilteredSockets, 0);
 			SetShaderValue(RHICmdList, ComputeShaderRHI, FilteredSocketBoneOffset, 0);
@@ -1246,6 +1257,8 @@ private:
 	LAYOUT_FIELD(FShaderParameter, NumBones);
 	LAYOUT_FIELD(FShaderParameter, NumFilteredBones);
 	LAYOUT_FIELD(FShaderParameter, NumUnfilteredBones);
+	LAYOUT_FIELD(FShaderParameter, RandomMaxBone);
+	LAYOUT_FIELD(FShaderParameter, ExcludeBoneIndex);
 	LAYOUT_FIELD(FShaderResourceParameter, FilteredAndUnfilteredBones);
 	LAYOUT_FIELD(FShaderParameter, NumFilteredSockets);
 	LAYOUT_FIELD(FShaderParameter, FilteredSocketBoneOffset);
@@ -2373,6 +2386,8 @@ const FString UNiagaraDataInterfaceSkeletalMesh::MeshNumWeightsName(TEXT("MeshNu
 const FString UNiagaraDataInterfaceSkeletalMesh::NumBonesName(TEXT("NumBones_"));
 const FString UNiagaraDataInterfaceSkeletalMesh::NumFilteredBonesName(TEXT("NumFilteredBones_"));
 const FString UNiagaraDataInterfaceSkeletalMesh::NumUnfilteredBonesName(TEXT("NumUnfilteredBones_"));
+const FString UNiagaraDataInterfaceSkeletalMesh::RandomMaxBoneName(TEXT("RandomMaxBone_"));
+const FString UNiagaraDataInterfaceSkeletalMesh::ExcludeBoneIndexName(TEXT("ExcludeBoneIndex_"));
 const FString UNiagaraDataInterfaceSkeletalMesh::FilteredAndUnfilteredBonesName(TEXT("FilteredAndUnfilteredBones_"));
 const FString UNiagaraDataInterfaceSkeletalMesh::NumFilteredSocketsName(TEXT("NumFilteredSockets_"));
 const FString UNiagaraDataInterfaceSkeletalMesh::FilteredSocketBoneOffsetName(TEXT("FilteredSocketBoneOffset_"));
