- Minor tweaks to allow cleaner inheritance from NiagaraComponent if necessary:
   - Exported FNiagaraSceneProxy under NIAGARA_API
   - Make GetViewRelevance public to match the primitive scene proxy API
   - Removed the final keyword from FNiagaraSceneProxy's inheritance

diff --git a/Engine/Plugins/FX/Niagara/Source/Niagara/Public/NiagaraComponent.h b/Engine/Plugins/FX/Niagara/Source/Niagara/Public/NiagaraComponent.h
--- a/Engine/Plugins/FX/Niagara/Source/Niagara/Public/NiagaraComponent.h
+++ b/Engine/Plugins/FX/Niagara/Source/Niagara/Public/NiagaraComponent.h
@@ -566,7 +566,7 @@ FORCEINLINE int32 UNiagaraComponent::GetPreviewLODDistance()const { return 0.0f;
 /**
 * Scene proxy for drawing niagara particle simulations.
 */
-class FNiagaraSceneProxy final : public FPrimitiveSceneProxy
+class NIAGARA_API FNiagaraSceneProxy : public FPrimitiveSceneProxy
 {
 public:
 	SIZE_T GetTypeHash() const override;
@@ -597,6 +597,8 @@ public:
 
 	FRHIUniformBuffer* GetUniformBufferNoVelocity() const;
 
+	virtual FPrimitiveViewRelevance GetViewRelevance(const FSceneView* View) const override;
+
 private:
 	void ReleaseRenderThreadResources();
 
@@ -608,8 +610,6 @@ private:
 
 	virtual void GetDynamicMeshElements(const TArray<const FSceneView*>& Views, const FSceneViewFamily& ViewFamily, uint32 VisibilityMap, FMeshElementCollector& Collector) const override;
 
-	virtual FPrimitiveViewRelevance GetViewRelevance(const FSceneView* View) const override;
-
 	/*
 	virtual bool CanBeOccluded() const override
 	{
