Niagara - Fix issues from parameter panel switch.
+ Prevent the view parameter panel view models from being created if they're not being used.
+ Add the old style parameter panel to the scratch pad.
+ Fix UI issues in scratch pad.

diff --git a/Engine/Plugins/FX/Niagara/Source/NiagaraEditor/Private/NiagaraNodeWithDynamicPins.cpp b/Engine/Plugins/FX/Niagara/Source/NiagaraEditor/Private/NiagaraNodeWithDynamicPins.cpp
--- a/Engine/Plugins/FX/Niagara/Source/NiagaraEditor/Private/NiagaraNodeWithDynamicPins.cpp
+++ b/Engine/Plugins/FX/Niagara/Source/NiagaraEditor/Private/NiagaraNodeWithDynamicPins.cpp
@@ -11,7 +11,7 @@
 #include "Widgets/Layout/SBox.h"
 #include "INiagaraEditorTypeUtilities.h"
 #include "NiagaraEditorUtilities.h"
-#include "SNiagaraParameterMapView.h"
+#include "Widgets/SNiagaraParameterMapView.h"
 #include "Framework/Application/SlateApplication.h"
 #include "NiagaraNodeParameterMapBase.h"
 #include "NiagaraScriptVariable.h"
diff --git a/Engine/Plugins/FX/Niagara/Source/NiagaraEditor/Private/Toolkits/NiagaraScriptToolkit.cpp b/Engine/Plugins/FX/Niagara/Source/NiagaraEditor/Private/Toolkits/NiagaraScriptToolkit.cpp
--- a/Engine/Plugins/FX/Niagara/Source/NiagaraEditor/Private/Toolkits/NiagaraScriptToolkit.cpp
+++ b/Engine/Plugins/FX/Niagara/Source/NiagaraEditor/Private/Toolkits/NiagaraScriptToolkit.cpp
@@ -14,7 +14,7 @@
 #include "Widgets/SNiagaraSelectedObjectsDetails.h"
 #include "SNiagaraParameterCollection.h"
 #include "UObject/Package.h"
-#include "SNiagaraParameterMapView.h"
+#include "Widgets/SNiagaraParameterMapView.h"
 #include "Widgets/SNiagaraParameterPanel.h"
 #include "NiagaraEditorStyle.h"
 #include "NiagaraSystem.h"
@@ -158,8 +158,12 @@ void FNiagaraScriptToolkit::Initialize( const EToolkitMode::Type Mode, const TSh
 
 	ScriptViewModel = MakeShareable(new FNiagaraStandaloneScriptViewModel(DisplayName, ENiagaraParameterEditMode::EditAll, NiagaraMessageLogViewModel, MessageLogGuidKey));
 	ScriptViewModel->Initialize(EditedNiagaraScript, OriginalNiagaraScript);
-	ParameterPanelViewModel = MakeShareable(new FNiagaraScriptToolkitParameterPanelViewModel(ScriptViewModel));
-	ParameterPanelViewModel->InitBindings();
+
+	if (GbShowNiagaraDeveloperWindows)
+	{
+		ParameterPanelViewModel = MakeShareable(new FNiagaraScriptToolkitParameterPanelViewModel(ScriptViewModel));
+		ParameterPanelViewModel->InitBindings();
+	}
 
 	OnEditedScriptGraphChangedHandle = ScriptViewModel->GetGraphViewModel()->GetGraph()->AddOnGraphNeedsRecompileHandler(
 		FOnGraphChanged::FDelegate::CreateRaw(this, &FNiagaraScriptToolkit::OnEditedScriptGraphChanged));
diff --git a/Engine/Plugins/FX/Niagara/Source/NiagaraEditor/Private/ViewModels/NiagaraScratchPadScriptViewModel.cpp b/Engine/Plugins/FX/Niagara/Source/NiagaraEditor/Private/ViewModels/NiagaraScratchPadScriptViewModel.cpp
--- a/Engine/Plugins/FX/Niagara/Source/NiagaraEditor/Private/ViewModels/NiagaraScratchPadScriptViewModel.cpp
+++ b/Engine/Plugins/FX/Niagara/Source/NiagaraEditor/Private/ViewModels/NiagaraScratchPadScriptViewModel.cpp
@@ -8,6 +8,7 @@
 #include "NiagaraSystem.h"
 #include "NiagaraEmitter.h"
 #include "NiagaraEditorUtilities.h"
+#include "NiagaraEditorModule.h"
 
 #include "ObjectTools.h"
 #include "ScopedTransaction.h"
@@ -47,8 +48,11 @@ void FNiagaraScratchPadScriptViewModel::Initialize(UNiagaraScript* Script)
 	OnGraphNeedsRecompileHandle = EditScriptSource->NodeGraph->AddOnGraphNeedsRecompileHandler(FOnGraphChanged::FDelegate::CreateSP(this, &FNiagaraScratchPadScriptViewModel::OnScriptGraphChanged));
 	EditScript->OnPropertyChanged().AddSP(this, &FNiagaraScratchPadScriptViewModel::OnScriptPropertyChanged);
 	ParameterPanelCommands = MakeShared<FUICommandList>();
-	ParameterPaneViewModel = MakeShared<FNiagaraScriptToolkitParameterPanelViewModel>(this->AsShared());
-	ParameterPaneViewModel->InitBindings();
+	if (GbShowNiagaraDeveloperWindows)
+	{
+		ParameterPaneViewModel = MakeShared<FNiagaraScriptToolkitParameterPanelViewModel>(this->AsShared());
+		ParameterPaneViewModel->InitBindings();
+	}
 }
 
 void FNiagaraScratchPadScriptViewModel::Finalize()
diff --git a/Engine/Plugins/FX/Niagara/Source/NiagaraEditor/Private/ViewModels/NiagaraScratchPadViewModel.cpp b/Engine/Plugins/FX/Niagara/Source/NiagaraEditor/Private/ViewModels/NiagaraScratchPadViewModel.cpp
--- a/Engine/Plugins/FX/Niagara/Source/NiagaraEditor/Private/ViewModels/NiagaraScratchPadViewModel.cpp
+++ b/Engine/Plugins/FX/Niagara/Source/NiagaraEditor/Private/ViewModels/NiagaraScratchPadViewModel.cpp
@@ -121,7 +121,8 @@ void UNiagaraScratchPadViewModel::RefreshScriptViewModels()
 	bool bEditViewModelsChanged = false;
 	if (ActiveScriptViewModel.IsValid() && ScriptViewModels.Contains(ActiveScriptViewModel.ToSharedRef()) == false)
 	{
-		ActiveScriptViewModel.Reset();
+		bool bRefreshEditScriptViewModels = false;
+		ResetActiveScriptViewModelInternal(bRefreshEditScriptViewModels);
 		bEditViewModelsChanged = true;
 	}
 
@@ -171,6 +172,16 @@ TSharedPtr<FNiagaraScratchPadScriptViewModel> UNiagaraScratchPadViewModel::GetVi
 	return TSharedPtr<FNiagaraScratchPadScriptViewModel>();
 }
 
+TSharedPtr<FNiagaraScratchPadScriptViewModel> UNiagaraScratchPadViewModel::GetViewModelForEditScript(UNiagaraScript* InEditScript)
+{
+	TSharedRef<FNiagaraScratchPadScriptViewModel>* ViewModelForEditScript = ScriptViewModels.FindByPredicate([InEditScript](TSharedRef<FNiagaraScratchPadScriptViewModel>& ScriptViewModel) { return ScriptViewModel->GetEditScript() == InEditScript; });
+	if (ViewModelForEditScript != nullptr)
+	{
+		return *ViewModelForEditScript;
+	}
+	return TSharedPtr<FNiagaraScratchPadScriptViewModel>();
+}
+
 const TArray<ENiagaraScriptUsage>& UNiagaraScratchPadViewModel::GetAvailableUsages() const
 {
 	return AvailableUsages;
@@ -222,12 +233,23 @@ void UNiagaraScratchPadViewModel::FocusScratchPadScriptViewModel(TSharedRef<FNia
 }
 
 void UNiagaraScratchPadViewModel::ResetActiveScriptViewModel()
+{
+	bool bRefreshEditScriptViewModels = true;
+	ResetActiveScriptViewModelInternal(bRefreshEditScriptViewModels);
+}
+
+void UNiagaraScratchPadViewModel::ResetActiveScriptViewModelInternal(bool bRefreshEditScriptViewModels)
 {
 	if (ActiveScriptViewModel.IsValid())
 	{
 		ActiveScriptViewModel.Reset();
 		ObjectSelection->ClearSelectedObjects();
-		RefreshEditScriptViewModels();
+
+		if (bRefreshEditScriptViewModels)
+		{
+			RefreshEditScriptViewModels();
+		}
+
 		OnActiveScriptChangedDelegate.Broadcast();
 	}
 }
diff --git a/Engine/Plugins/FX/Niagara/Source/NiagaraEditor/Private/Widgets/SNiagaraParameterMapView.cpp b/Engine/Plugins/FX/Niagara/Source/NiagaraEditor/Private/Widgets/SNiagaraParameterMapView.cpp
--- a/Engine/Plugins/FX/Niagara/Source/NiagaraEditor/Private/Widgets/SNiagaraParameterMapView.cpp
+++ b/Engine/Plugins/FX/Niagara/Source/NiagaraEditor/Private/Widgets/SNiagaraParameterMapView.cpp
@@ -1,6 +1,6 @@
 // Copyright Epic Games, Inc. All Rights Reserved.
 
-#include "SNiagaraParameterMapView.h"
+#include "Widgets/SNiagaraParameterMapView.h"
 #include "SNiagaraParameterMapPaletteItem.h"
 #include "NiagaraObjectSelection.h"
 #include "Widgets/Images/SImage.h"
@@ -472,7 +472,10 @@ void SNiagaraParameterMapView::OnActionSelected(const TArray< TSharedPtr<FEdGrap
 	
 	// If a variable wasn't selected just clear the current selection
 	// TODO: Get proper clearing to work. Current there's no way to clear while clicking on an empty location in the graph area
-	SelectedVariableObjects->ClearSelectedObjects();
+	if (SelectedVariableObjects.IsValid())
+	{
+		SelectedVariableObjects->ClearSelectedObjects();
+	}
 }
 
 void SNiagaraParameterMapView::OnActionDoubleClicked(const TArray< TSharedPtr<FEdGraphSchemaAction> >& InActions)
diff --git a/Engine/Plugins/FX/Niagara/Source/NiagaraEditor/Private/Widgets/SNiagaraParameterMapView.h b/Engine/Plugins/FX/Niagara/Source/NiagaraEditor/Private/Widgets/SNiagaraParameterMapView.h
deleted file mode 100644
diff --git a/Engine/Plugins/FX/Niagara/Source/NiagaraEditor/Public/ViewModels/NiagaraScratchPadViewModel.h b/Engine/Plugins/FX/Niagara/Source/NiagaraEditor/Public/ViewModels/NiagaraScratchPadViewModel.h
--- a/Engine/Plugins/FX/Niagara/Source/NiagaraEditor/Public/ViewModels/NiagaraScratchPadViewModel.h
+++ b/Engine/Plugins/FX/Niagara/Source/NiagaraEditor/Public/ViewModels/NiagaraScratchPadViewModel.h
@@ -35,6 +35,8 @@ public:
 
 	TSharedPtr<FNiagaraScratchPadScriptViewModel> GetViewModelForScript(UNiagaraScript* InScript);
 
+	TSharedPtr<FNiagaraScratchPadScriptViewModel> GetViewModelForEditScript(UNiagaraScript* InEditScript);
+
 	const TArray<ENiagaraScriptUsage>& GetAvailableUsages() const;
 
 	FText GetDisplayNameForUsage(ENiagaraScriptUsage InUsage) const;
@@ -84,6 +86,8 @@ private:
 
 	void TearDownScriptViewModel(TSharedRef<FNiagaraScratchPadScriptViewModel> InScriptViewModel);
 
+	void ResetActiveScriptViewModelInternal(bool bRefreshEditScriptViewModels);
+
 	void RefreshEditScriptViewModels();
 
 	void ScriptGraphNodeSelectionChanged(TWeakPtr<FNiagaraScratchPadScriptViewModel> InScriptViewModelWeak);
diff --git a/Engine/Plugins/FX/Niagara/Source/NiagaraEditor/Public/ViewModels/NiagaraScriptViewModel.h b/Engine/Plugins/FX/Niagara/Source/NiagaraEditor/Public/ViewModels/NiagaraScriptViewModel.h
--- a/Engine/Plugins/FX/Niagara/Source/NiagaraEditor/Public/ViewModels/NiagaraScriptViewModel.h
+++ b/Engine/Plugins/FX/Niagara/Source/NiagaraEditor/Public/ViewModels/NiagaraScriptViewModel.h
@@ -49,7 +49,7 @@ public:
 	NIAGARAEDITOR_API TSharedRef<FNiagaraScriptGraphViewModel> GetGraphViewModel();
 
 	/** Gets the currently selected script variables. */
-	TSharedRef<FNiagaraObjectSelection> GetVariableSelection();
+	NIAGARAEDITOR_API TSharedRef<FNiagaraObjectSelection> GetVariableSelection();
 
 	/** Updates the script with the latest compile status. */
 	void UpdateCompileStatus(ENiagaraScriptCompileStatus InAggregateCompileStatus, const FString& InAggregateCompileErrors,
diff --git a/Engine/Plugins/FX/Niagara/Source/NiagaraEditor/Public/Widgets/SNiagaraParameterMapView.h b/Engine/Plugins/FX/Niagara/Source/NiagaraEditor/Public/Widgets/SNiagaraParameterMapView.h
new file mode 100644
--- /dev/null
+++ b/Engine/Plugins/FX/Niagara/Source/NiagaraEditor/Public/Widgets/SNiagaraParameterMapView.h
@@ -0,0 +1,229 @@
+// Copyright Epic Games, Inc. All Rights Reserved.
+
+#pragma once
+
+#include "NiagaraTypes.h"
+#include "Widgets/SCompoundWidget.h"
+#include "Delegates/Delegate.h"
+#include "Widgets/DeclarativeSyntaxSupport.h"
+#include "PropertyEditorDelegates.h"
+#include "EdGraph/EdGraphSchema.h"
+#include "ViewModels/Stack/NiagaraParameterHandle.h"
+#include "Framework/Commands/Commands.h"
+#include "Framework/Commands/UICommandList.h"
+#include "NiagaraActions.h"
+#include "EditorStyleSet.h"
+
+class SGraphActionMenu;
+class SEditableTextBox;
+class SExpanderArrow;
+class SSearchBox;
+class SComboButton;
+class FNiagaraObjectSelection;
+class UNiagaraGraph;
+struct FEdGraphSchemaAction;
+
+/* Enums to use when grouping the blueprint members in the list panel. The order here will determine the order in the list */
+namespace NiagaraParameterMapSectionID
+{
+	enum Type
+	{
+		NONE = 0,
+		MODULE,
+		STATIC_SWITCH,
+		ENGINE,
+		PARAMETERCOLLECTION,
+		USER,
+		SYSTEM,
+		EMITTER,
+		PARTICLE,
+		OTHER,
+	};
+
+	static FText OnGetSectionTitle(const NiagaraParameterMapSectionID::Type InSection);
+	static NiagaraParameterMapSectionID::Type OnGetSectionFromVariable(const FNiagaraVariable& InVar, bool IsStaticSwitchVariable, FNiagaraParameterHandle& OutParameterHandle, const NiagaraParameterMapSectionID::Type DefaultType = NiagaraParameterMapSectionID::Type::NONE);
+};
+
+class FNiagaraParameterMapViewCommands : public TCommands<FNiagaraParameterMapViewCommands>
+{
+public:
+	/** Constructor */
+	FNiagaraParameterMapViewCommands()
+		: TCommands<FNiagaraParameterMapViewCommands>(TEXT("NiagaraParameterMapViewCommands"), NSLOCTEXT("Contexts", "NiagaraParameterMap", "NiagaraParameterMap"), NAME_None, FEditorStyle::GetStyleSetName())
+	{
+	}
+
+	// Basic operations
+	TSharedPtr<FUICommandInfo> DeleteEntry;
+
+	/** Initialize commands */
+	virtual void RegisterCommands() override;
+};
+
+/** A widget for viewing and editing a set of selected objects with a details panel. */
+class SNiagaraParameterMapView : public SCompoundWidget
+{
+public:
+	enum EToolkitType
+	{
+		SCRIPT,
+		SYSTEM,
+	};
+
+	SLATE_BEGIN_ARGS(SNiagaraParameterMapView)
+	{}
+	SLATE_END_ARGS();
+
+	NIAGARAEDITOR_API virtual ~SNiagaraParameterMapView();
+
+	NIAGARAEDITOR_API void Construct(const FArguments& InArgs, const TArray<TSharedRef<FNiagaraObjectSelection>>& InSelectedObjects, const EToolkitType InToolkitType, const TSharedPtr<FUICommandList>& InToolkitCommands);
+
+	NIAGARAEDITOR_API virtual void Tick(const FGeometry& AllottedGeometry, const double InCurrentTime, const float InDeltaTime) override;
+
+	/** Whether the add parameter button should be enabled. */
+	bool ParameterAddEnabled() const;
+
+	/** Adds parameter to the graph parameter store and refreshes the menu. */
+	void AddParameter(FNiagaraVariable NewVariable);
+
+	/** Refreshes the graphs used for this menu. */
+	void Refresh(bool bRefreshMenu = true);
+
+	static TSharedRef<SExpanderArrow> CreateCustomActionExpander(const struct FCustomExpanderData& ActionMenuData);
+
+	static bool IsStaticSwitchParameter(const FNiagaraVariable& Variable, const TArray<TWeakObjectPtr<UNiagaraGraph>>& Graphs);
+
+private:
+	/** Function to bind to SNiagaraAddParameterMenus to filter types we allow creating in generic parameters*/
+	bool AllowMakeTypeGeneric(const FNiagaraTypeDefinition& InType) const;
+
+	/** Function to bind to SNiagaraAddParameterMenus to filter types we allow creating for attributes */
+	bool AllowMakeTypeAttribute(const FNiagaraTypeDefinition& InType) const;
+
+	/** Callback when the filter is changed, forces the action tree(s) to filter */
+	void OnFilterTextChanged(const FText& InFilterText);
+
+	// SGraphActionMenu delegates
+	FText GetFilterText() const;
+	TSharedRef<SWidget> OnCreateWidgetForAction(struct FCreateWidgetForActionData* const InCreateData);
+	void CollectAllActions(FGraphActionListBuilderBase& OutAllActions);
+	void CollectStaticSections(TArray<int32>& StaticSectionIDs);
+	FReply OnActionDragged(const TArray<TSharedPtr<FEdGraphSchemaAction>>& InActions, const FPointerEvent& MouseEvent);
+	void OnActionSelected(const TArray<TSharedPtr<FEdGraphSchemaAction>>& InActions, ESelectInfo::Type InSelectionType);
+	void OnActionDoubleClicked(const TArray<TSharedPtr<FEdGraphSchemaAction>>& InActions);
+	TSharedPtr<SWidget> OnContextMenuOpening();
+	FText OnGetSectionTitle(int32 InSectionID);
+	TSharedRef<SWidget> OnGetSectionWidget(TSharedRef<SWidget> RowWidget, int32 InSectionID);
+	TSharedRef<SWidget> CreateAddToSectionButton(const NiagaraParameterMapSectionID::Type InSection, TWeakPtr<SWidget> WeakRowWidget, FText AddNewText, FName MetaDataTag);
+	
+	/** Checks if the selected action has context menu */
+	bool SelectionHasContextMenu() const;
+	
+	TSharedRef<SWidget> OnGetParameterMenu(const NiagaraParameterMapSectionID::Type InSection = NiagaraParameterMapSectionID::NONE);
+	EVisibility OnAddButtonTextVisibility(TWeakPtr<SWidget> RowWidget, const NiagaraParameterMapSectionID::Type InSection) const;
+
+	void SelectedObjectsChanged();
+	void EmptyGraphs();
+	void AddGraph(UNiagaraGraph* Graph);
+	void AddGraph(class UNiagaraScriptSourceBase* SourceBase);
+	void OnGraphChanged(const struct FEdGraphEditAction& InAction);
+
+	//Callbacks
+	void OnDeleteEntry();
+	bool CanDeleteEntry() const;
+	void OnRequestRenameOnActionNode();
+	bool CanRequestRenameOnActionNode(TWeakPtr<struct FGraphActionNode> InSelectedNode) const;
+	bool CanRequestRenameOnActionNode() const;
+	void OnPostRenameActionNode(const FText& InText, FNiagaraParameterAction& InAction);
+
+	bool IsSystemToolkit();
+	bool IsScriptToolkit();
+
+	/** Delegate handler used to match an FName to an action in the list, used for renaming keys */
+	bool HandleActionMatchesName(struct FEdGraphSchemaAction* InAction, const FName& InName) const;
+
+private:
+
+	/** Sets bNeedsRefresh to true. Causing the list to be refreshed next tick. */
+	void RefreshActions();
+
+	/** Graph Action Menu for displaying all our variables and functions */
+	TSharedPtr<SGraphActionMenu> GraphActionMenu;
+
+	/** The filter box that handles filtering for both graph action menus. */
+	TSharedPtr<SSearchBox> FilterBox;
+
+	/** Add parameter buttons for all sections. */
+	TArray<TSharedPtr<SComboButton>> AddParameterButtons;
+
+	/** The selected objects being viewed and edited by this widget. */
+	TSharedPtr<FNiagaraObjectSelection> SelectedScriptObjects;
+	TSharedPtr<FNiagaraObjectSelection> SelectedVariableObjects;
+
+	TArray<TWeakObjectPtr<UNiagaraGraph>> Graphs;
+
+	/** The handle to the graph changed delegate. */
+	TArray<FDelegateHandle> OnGraphChangedHandles;
+	TArray<FDelegateHandle> OnRecompileHandles;
+
+	EToolkitType ToolkitType;
+	TSharedPtr<FUICommandList> ToolkitCommands;
+
+	bool bNeedsRefresh;
+};
+
+class SNiagaraAddParameterMenu : public SCompoundWidget
+{
+public:
+	/** Delegate that can be used to create a widget for a particular action */
+	DECLARE_DELEGATE_OneParam(FOnAddParameter, FNiagaraVariable);
+	DECLARE_DELEGATE_TwoParams(FOnCollectCustomActions, FGraphActionListBuilderBase&, bool&);
+	DECLARE_DELEGATE_RetVal_OneParam(bool, FOnAllowMakeType, const FNiagaraTypeDefinition&);
+
+	SLATE_BEGIN_ARGS(SNiagaraAddParameterMenu)
+		: _Section(NiagaraParameterMapSectionID::NONE)
+		, _AllowCreatingNew(true)
+		, _ShowNamespaceCategory(true)
+		, _ShowGraphParameters(true)
+		, _AutoExpandMenu(false)
+		, _IsParameterRead(true) {}
+		SLATE_EVENT(FOnAddParameter, OnAddParameter)
+		SLATE_EVENT(FOnCollectCustomActions, OnCollectCustomActions)
+		SLATE_EVENT(FOnAllowMakeType, OnAllowMakeType)
+		SLATE_ATTRIBUTE(NiagaraParameterMapSectionID::Type, Section)
+		SLATE_ATTRIBUTE(bool, AllowCreatingNew)
+		SLATE_ATTRIBUTE(bool, ShowNamespaceCategory)
+		SLATE_ATTRIBUTE(bool, ShowGraphParameters)
+		SLATE_ATTRIBUTE(bool, AutoExpandMenu)
+		SLATE_ATTRIBUTE(bool, IsParameterRead)
+	SLATE_END_ARGS();
+
+	void Construct(const FArguments& InArgs, TArray<TWeakObjectPtr<UNiagaraGraph>> InGraphs);
+
+	TSharedRef<SEditableTextBox> GetSearchBox();
+
+	void AddParameterGroup(FGraphActionListBuilderBase& OutActions, TArray<FNiagaraVariable>& Variables, const NiagaraParameterMapSectionID::Type InSection = NiagaraParameterMapSectionID::NONE, const FText& Category = FText::GetEmpty(), const FString& RootCategory = FString(), const bool bSort = true, const bool bCustomName = true);
+	void CollectParameterCollectionsActions(FGraphActionListBuilderBase& OutActions);
+	void CollectMakeNew(FGraphActionListBuilderBase& OutActions, const NiagaraParameterMapSectionID::Type InSection = NiagaraParameterMapSectionID::NONE);
+
+private:
+	void OnActionSelected(const TArray<TSharedPtr<FEdGraphSchemaAction>>& SelectedActions, ESelectInfo::Type InSelectionType);
+	void CollectAllActions(FGraphActionListBuilderBase& OutAllActions);
+	void AddParameterSelected(FNiagaraVariable NewVariable, const bool bCreateCustomName = true, const NiagaraParameterMapSectionID::Type InSection = NiagaraParameterMapSectionID::NONE);
+	
+	TSharedPtr<SGraphActionMenu> GraphMenu;
+
+	/** Delegate that gets fired when a parameter was added. */
+	FOnAddParameter OnAddParameter;
+	FOnCollectCustomActions OnCollectCustomActions;
+	FOnAllowMakeType OnAllowMakeType;
+
+	TArray<TWeakObjectPtr<UNiagaraGraph>> Graphs;
+
+	TAttribute<NiagaraParameterMapSectionID::Type> Section;
+	TAttribute<bool> AllowCreatingNew;
+	TAttribute<bool> ShowNamespaceCategory;
+	TAttribute<bool> ShowGraphParameters;
+	TAttribute<bool> AutoExpandMenu;
+	TAttribute<bool> IsParameterRead;
+};
diff --git a/Engine/Plugins/FX/Niagara/Source/NiagaraEditorWidgets/Private/NiagaraEditorWidgetsStyle.cpp b/Engine/Plugins/FX/Niagara/Source/NiagaraEditorWidgets/Private/NiagaraEditorWidgetsStyle.cpp
--- a/Engine/Plugins/FX/Niagara/Source/NiagaraEditorWidgets/Private/NiagaraEditorWidgetsStyle.cpp
+++ b/Engine/Plugins/FX/Niagara/Source/NiagaraEditorWidgets/Private/NiagaraEditorWidgetsStyle.cpp
@@ -254,6 +254,11 @@ TSharedRef< FSlateStyleSet > FNiagaraEditorWidgetsStyle::Create()
 		.SetFont(ScratchPadEditorHeaderFont);
 	Style->Set("NiagaraEditor.ScratchPad.EditorHeaderText", ScratchPadEditorHeaderText);
 
+	FSlateFontInfo ScratchPadSubSectionHeaderFont = DEFAULT_FONT("Bold", 9);
+	FTextBlockStyle ScratchPadSubSectionHeaderText = FTextBlockStyle(NormalText)
+		.SetFont(ScratchPadSubSectionHeaderFont);
+	Style->Set("NiagaraEditor.ScratchPad.SubSectionHeaderText", ScratchPadSubSectionHeaderText);
+
 	FSlateBrush ScratchPadCategoryBrush = BOX_PLUGIN_BRUSH("Icons/CategoryRow", FMargin(2.0f / 8.0f), FLinearColor(FColor(48, 48, 48)));
 	FSlateBrush ScratchPadHoveredCategoryBrush = BOX_PLUGIN_BRUSH("Icons/CategoryRow", FMargin(2.0f / 8.0f), FLinearColor(FColor(38, 38, 38)));
 	Style->Set("NiagaraEditor.ScratchPad.CategoryRow", FTableRowStyle(NormalTableRowStyle)
diff --git a/Engine/Plugins/FX/Niagara/Source/NiagaraEditorWidgets/Private/SNiagaraScratchPad.cpp b/Engine/Plugins/FX/Niagara/Source/NiagaraEditorWidgets/Private/SNiagaraScratchPad.cpp
--- a/Engine/Plugins/FX/Niagara/Source/NiagaraEditorWidgets/Private/SNiagaraScratchPad.cpp
+++ b/Engine/Plugins/FX/Niagara/Source/NiagaraEditorWidgets/Private/SNiagaraScratchPad.cpp
@@ -14,6 +14,7 @@
 #include "NiagaraEditorWidgetsStyle.h"
 #include "NiagaraScratchPadCommandContext.h"
 #include "Widgets/SNiagaraParameterPanel.h"
+#include "Widgets/SNiagaraParameterMapView.h"
 
 #include "Widgets/Text/SInlineEditableTextBlock.h"
 #include "Widgets/Input/SButton.h"
@@ -468,28 +469,8 @@ class SNiagaraScratchPadScriptEditor : public SCompoundWidget
 			[
 				SNew(SHorizontalBox)
 				+ SHorizontalBox::Slot()
-				.VAlign(VAlign_Center)
-				.Padding(2, 0, 2, 0)
-				.AutoWidth()
-				[
-					SNew(STextBlock)
-					.TextStyle(FNiagaraEditorWidgetsStyle::Get(), "NiagaraEditor.ScratchPad.EditorHeaderText")
-					.Text(this, &SNiagaraScratchPadScriptEditor::GetNameText)
-					.ToolTipText(ScriptViewModel.ToSharedRef(), &FNiagaraScratchPadScriptViewModel::GetToolTip)
-				]
-				+ SHorizontalBox::Slot()
 				.AutoWidth()
-				.Padding(1)
-				[
-					SNew(SNiagaraPinButton)
-					.IsPinned(ScriptViewModel.ToSharedRef(), &FNiagaraScratchPadScriptViewModel::GetIsPinned)
-					.OnPinnedChanged(ScriptViewModel.ToSharedRef(), &FNiagaraScratchPadScriptViewModel::SetIsPinned)
-					.PinItemDisplayName(LOCTEXT("PinItem", "script"))
-					.PinTargetDisplayName(LOCTEXT("PinTarget", "edit list"))
-				]
-				+ SHorizontalBox::Slot()
-				.HAlign(HAlign_Right)
-				.Padding(0.0f, 2.0f, 1.0f, 4.0f)
+				.Padding(2.0f, 3.0f, 5.0f, 3.0f)
 				[
 					SNew(SButton)
 					.OnClicked(this, &SNiagaraScratchPadScriptEditor::OnApplyButtonClicked)
@@ -517,6 +498,26 @@ class SNiagaraScratchPadScriptEditor : public SCompoundWidget
 						]
 					]
 				]
+				+ SHorizontalBox::Slot()
+				.VAlign(VAlign_Center)
+				.Padding(5, 0, 2, 0)
+				.AutoWidth()
+				[
+					SNew(STextBlock)
+					.TextStyle(FNiagaraEditorWidgetsStyle::Get(), "NiagaraEditor.ScratchPad.EditorHeaderText")
+					.Text(this, &SNiagaraScratchPadScriptEditor::GetNameText)
+					.ToolTipText(ScriptViewModel.ToSharedRef(), &FNiagaraScratchPadScriptViewModel::GetToolTip)
+				]
+				+ SHorizontalBox::Slot()
+				.AutoWidth()
+				.Padding(1)
+				[
+					SNew(SNiagaraPinButton)
+					.IsPinned(ScriptViewModel.ToSharedRef(), &FNiagaraScratchPadScriptViewModel::GetIsPinned)
+					.OnPinnedChanged(ScriptViewModel.ToSharedRef(), &FNiagaraScratchPadScriptViewModel::SetIsPinned)
+					.PinItemDisplayName(LOCTEXT("PinItem", "script"))
+					.PinTargetDisplayName(LOCTEXT("PinTarget", "edit list"))
+				]
 			]
 			+ SVerticalBox::Slot()
 			[
@@ -562,6 +563,11 @@ class SNiagaraScratchPadScriptEditorList : public SCompoundWidget
 		ViewModel->OnActiveScriptChanged().AddSP(this, &SNiagaraScratchPadScriptEditorList::ActiveScriptChanged);
 		bIsUpdatingSelection = false;
 
+		ChildSlot
+		[
+			SAssignNew(ContentBorder, SBorder)
+			.BorderImage(FEditorStyle::Get().GetBrush("ToolPanel.DarkGroupBorder"))
+		];
 		UpdateContentFromEditScriptViewModels();
 	}
 
@@ -603,7 +609,13 @@ private:
 		TSharedPtr<SWidget> NewContent;
 		if (ViewModel->GetEditScriptViewModels().Num() == 0)
 		{
-			NewContent = SNullWidget::NullWidget;
+			NewContent = SNew(SBox)
+				.HAlign(HAlign_Center)
+				.Padding(FMargin(0.0f, 10.0f, 0.0f, 0.0f))
+				[
+					SNew(STextBlock)
+					.Text(LOCTEXT("NoScriptToEdit", "No script selected"))
+				];
 			ScriptEditorList.Reset();
 		}
 		else if(ViewModel->GetEditScriptViewModels().Num() == 1)
@@ -627,10 +639,7 @@ private:
 			NewContent = ScriptEditorList;
 		}
 
-		ChildSlot
-		[
-			NewContent.ToSharedRef()
-		];
+		ContentBorder->SetContent(NewContent.ToSharedRef());
 	}
 
 	void ActiveScriptChanged()
@@ -689,6 +698,7 @@ private:
 	};
 
 	UNiagaraScratchPadViewModel* ViewModel;
+	TSharedPtr<SBorder> ContentBorder;
 	TSharedPtr<SListView<TSharedRef<FNiagaraScratchPadScriptViewModel>>> ScriptEditorList;
 	TArray<FScriptViewModelWidgetPair> ScriptViewModelWidgetPairs;
 	bool bIsUpdatingSelection;
@@ -703,7 +713,8 @@ class SNiagaraScratchPadParameterPanel : public SCompoundWidget
 	{
 		ViewModel = InViewModel;
 		ViewModel->OnActiveScriptChanged().AddSP(this, &SNiagaraScratchPadParameterPanel::ActiveScriptChanged);
-		UpdateContent();
+		bool bForce = true;
+		UpdateContent(bForce);
 	}
 
 	virtual FReply OnKeyDown(const FGeometry& MyGeometry, const FKeyEvent& InKeyEvent) override
@@ -717,30 +728,72 @@ class SNiagaraScratchPadParameterPanel : public SCompoundWidget
 	}
 
 private:
-	void UpdateContent()
+	void UpdateContent(bool bForce)
 	{
-		TSharedPtr<FNiagaraScratchPadScriptViewModel> ScriptViewModel = ScriptViewModelWeak.Pin();
+		TSharedPtr<FNiagaraScratchPadScriptViewModel> OldScriptViewModel = ScriptViewModelWeak.Pin();
 		TSharedPtr<FNiagaraScratchPadScriptViewModel> NewScriptViewModel = ViewModel->GetActiveScriptViewModel();
-		if (NewScriptViewModel != ScriptViewModel)
+		if (NewScriptViewModel != OldScriptViewModel || bForce)
 		{
+			ScriptViewModelWeak = NewScriptViewModel;
 			if (NewScriptViewModel.IsValid())
 			{
+				TSharedPtr<SWidget> ParameterPanelWidget;
+				if (NewScriptViewModel->GetParameterPanelViewModel().IsValid())
+				{
+					ParameterPanelWidget = SNew(SNiagaraParameterPanel, NewScriptViewModel->GetParameterPanelViewModel(), NewScriptViewModel->GetParameterPanelCommands());
+				}
+				else
+				{
+					TSharedRef<FNiagaraObjectSelection> ScriptSelection = MakeShared<FNiagaraObjectSelection>();
+					ScriptSelection->SetSelectedObject(NewScriptViewModel->GetEditScript());
+					TArray<TSharedRef<FNiagaraObjectSelection>> ParameterSelections;
+					ParameterSelections.Add(ScriptSelection);
+					ParameterSelections.Add(NewScriptViewModel->GetVariableSelection());
+					ParameterPanelWidget = SNew(SNiagaraParameterMapView, ParameterSelections, SNiagaraParameterMapView::EToolkitType::SCRIPT, NewScriptViewModel->GetParameterPanelCommands());
+				}
 				ChildSlot
 				[
-					SNew(SNiagaraParameterPanel, NewScriptViewModel->GetParameterPanelViewModel(), NewScriptViewModel->GetParameterPanelCommands())
+					SNew(SVerticalBox)
+					+ SVerticalBox::Slot()
+					.AutoHeight()
+					.Padding(2.0f, 2.0f, 2.0f, 5.0f)
+					[
+						SNew(STextBlock)
+						.TextStyle(FNiagaraEditorWidgetsStyle::Get(), "NiagaraEditor.ScratchPad.SubSectionHeaderText")
+						.Text(this, &SNiagaraScratchPadParameterPanel::GetSubHeaderText)
+					]
+					+ SVerticalBox::Slot()
+					[
+						ParameterPanelWidget.ToSharedRef()
+					]
 				];
-				ScriptViewModelWeak = NewScriptViewModel;
 			}
 			else
 			{
-				ChildSlot[SNullWidget::NullWidget];
+				ChildSlot
+				[
+					SNew(SBox)
+					.HAlign(HAlign_Center)
+					.Padding(FMargin(0.0f, 10.0f, 0.0f, 0.0f))
+					[
+						SNew(STextBlock)
+						.Text(LOCTEXT("NoScriptForParameters", "No script selected"))
+					]
+				];
 			}
 		}
 	}
 
 	void ActiveScriptChanged()
 	{
-		UpdateContent();
+		bool bForce = false;
+		UpdateContent(bForce);
+	}
+
+	FText GetSubHeaderText() const
+	{
+		TSharedPtr<FNiagaraScratchPadScriptViewModel> ScriptViewModel = ScriptViewModelWeak.Pin();
+		return ScriptViewModel.IsValid() ? ScriptViewModel->GetDisplayName() : FText();
 	}
 
 private:
@@ -780,6 +833,8 @@ class SNiagaraScratchPadSectionBox : public SCompoundWidget
 void SNiagaraScratchPad::Construct(const FArguments& InArgs, UNiagaraScratchPadViewModel* InViewModel)
 {
 	ViewModel = InViewModel;
+	ViewModel->GetObjectSelection()->OnSelectedObjectsChanged().AddSP(this, &SNiagaraScratchPad::ObjectSelectionChanged);
+
 	CommandContext = MakeShared<FNiagaraScratchPadCommandContext>(InViewModel);
 
 	ChildSlot
@@ -839,11 +894,7 @@ void SNiagaraScratchPad::Construct(const FArguments& InArgs, UNiagaraScratchPadV
 				+ SSplitter::Slot()
 				.Value(0.6f)
 				[
-					SNew(SBox)
-					.Padding(FMargin(3.0f, 0.0f))
-					[
-						InNamedWidgetProvider.GetNamedWidget(ScriptEditorName)
-					]
+					InNamedWidgetProvider.GetNamedWidget(ScriptEditorName)
 				]
 				+ SSplitter::Slot()
 				.Value(0.25f)
@@ -884,11 +935,7 @@ void SNiagaraScratchPad::Construct(const FArguments& InArgs, UNiagaraScratchPadV
 				+ SSplitter::Slot()
 				.Value(0.7f)
 				[
-					SNew(SBox)
-					.Padding(FMargin(3.0f, 0.0f, 0.0f, 0.0f))
-					[
-						InNamedWidgetProvider.GetNamedWidget(ScriptEditorName)
-					]
+					InNamedWidgetProvider.GetNamedWidget(ScriptEditorName)
 				];
 			}
 			else
@@ -911,6 +958,46 @@ void SNiagaraScratchPad::Construct(const FArguments& InArgs, UNiagaraScratchPadV
 	];
 }
 
+void SNiagaraScratchPad::ObjectSelectionChanged()
+{
+	int32 SelectionCount = ViewModel->GetObjectSelection()->GetSelectedObjects().Num();
+	if (SelectionCount == 0)
+	{
+		ObjectSelectionSubHeaderText = FText();
+	}
+	else if (SelectionCount == 1)
+	{
+		UObject* SelectedObject = ViewModel->GetObjectSelection()->GetSelectedObjects().Array()[0];
+		if (SelectedObject->IsA<UEdGraphNode>())
+		{
+			UEdGraphNode* SelectedGraphNode = CastChecked<UEdGraphNode>(SelectedObject);
+			ObjectSelectionSubHeaderText = SelectedGraphNode->GetNodeTitle(ENodeTitleType::ListView);
+		}
+		else if (SelectedObject->IsA<UNiagaraScript>())
+		{
+			UNiagaraScript* SelectedScript = CastChecked<UNiagaraScript>(SelectedObject);
+			TSharedPtr<FNiagaraScratchPadScriptViewModel> ScratchPadScriptViewModel = ViewModel->GetViewModelForEditScript(SelectedScript);
+			if (ScratchPadScriptViewModel.IsValid())
+			{
+				ObjectSelectionSubHeaderText = ScratchPadScriptViewModel->GetDisplayName();
+			}
+			else
+			{
+				ObjectSelectionSubHeaderText = FText::FromString(SelectedScript->GetName());
+			}
+		}
+		else
+		{
+			ObjectSelectionSubHeaderText = FText::FromString(SelectedObject->GetName());
+		}
+	}
+	else
+	{
+		ObjectSelectionSubHeaderText = FText::Format(LOCTEXT("MultipleSelectionFormat", "{0} Objects Selected..."),
+			FText::AsNumber(SelectionCount));
+	}
+}
+
 TSharedRef<SWidget> SNiagaraScratchPad::ConstructScriptSelector()
 {
 	return SNew(SNiagaraScratchPadSectionBox)
@@ -939,8 +1026,45 @@ TSharedRef<SWidget> SNiagaraScratchPad::ConstructSelectionEditor()
 	return SNew(SNiagaraScratchPadSectionBox)
 	.HeaderText(LOCTEXT("ScratchPadSelection", "Scratch Pad Selection"))
 	[
-		SNew(SNiagaraSelectedObjectsDetails, ViewModel->GetObjectSelection())
+		SNew(SVerticalBox)
+		+ SVerticalBox::Slot()
+		.AutoHeight()
+		.Padding(2.0f, 2.0f, 2.0f, 5.0f)
+		[
+			SNew(STextBlock)
+			.TextStyle(FNiagaraEditorWidgetsStyle::Get(), "NiagaraEditor.ScratchPad.SubSectionHeaderText")
+			.Visibility(this, &SNiagaraScratchPad::GetObjectSelectionSubHeaderTextVisibility)
+			.Text(this, &SNiagaraScratchPad::GetObjectSelectionSubHeaderText)
+		]
+		+ SVerticalBox::Slot()
+		.HAlign(HAlign_Center)
+		.AutoHeight()
+		.Padding(FMargin(0.0f, 10.0f, 0.0f, 0.0f))
+		[
+			SNew(STextBlock)
+			.Visibility(this, &SNiagaraScratchPad::GetObjectSelectionNoSelectionTextVisibility)
+			.Text(LOCTEXT("NoObjectSelection", "No object selected"))
+		]
+		+ SVerticalBox::Slot()
+		[
+			SNew(SNiagaraSelectedObjectsDetails, ViewModel->GetObjectSelection())
+		]
 	];
 }
 
+EVisibility SNiagaraScratchPad::GetObjectSelectionSubHeaderTextVisibility() const
+{
+	return ViewModel->GetObjectSelection()->GetSelectedObjects().Num() != 0 ? EVisibility::Visible : EVisibility::Collapsed;
+}
+
+FText SNiagaraScratchPad::GetObjectSelectionSubHeaderText() const
+{
+	return ObjectSelectionSubHeaderText;
+}
+
+EVisibility SNiagaraScratchPad::GetObjectSelectionNoSelectionTextVisibility() const
+{
+	return ViewModel->GetObjectSelection()->GetSelectedObjects().Num() == 0 ? EVisibility::Visible : EVisibility::Collapsed;
+}
+
 #undef LOCTEXT_NAMESPACE
\ No newline at end of file
diff --git a/Engine/Plugins/FX/Niagara/Source/NiagaraEditorWidgets/Private/SNiagaraScratchPad.h b/Engine/Plugins/FX/Niagara/Source/NiagaraEditorWidgets/Private/SNiagaraScratchPad.h
--- a/Engine/Plugins/FX/Niagara/Source/NiagaraEditorWidgets/Private/SNiagaraScratchPad.h
+++ b/Engine/Plugins/FX/Niagara/Source/NiagaraEditorWidgets/Private/SNiagaraScratchPad.h
@@ -20,6 +20,8 @@ class SNiagaraScratchPad : public SCompoundWidget
 	void Construct(const FArguments& InArgs, UNiagaraScratchPadViewModel* InViewModel);
 
 private:
+	void ObjectSelectionChanged();
+
 	TSharedRef<SWidget> ConstructScriptSelector();
 
 	TSharedRef<SWidget> ConstructParameterPanel();
@@ -28,7 +30,15 @@ private:
 
 	TSharedRef<SWidget> ConstructSelectionEditor();
 
+	EVisibility GetObjectSelectionSubHeaderTextVisibility() const;
+
+	FText GetObjectSelectionSubHeaderText() const;
+
+	EVisibility GetObjectSelectionNoSelectionTextVisibility() const;
+
 private:
+	FText ObjectSelectionSubHeaderText;
+
 	TWeakObjectPtr<UNiagaraScratchPadViewModel> ViewModel;
 
 	TSharedPtr<FNiagaraScratchPadCommandContext> CommandContext;
