Niagara - Fix crash when applying changes to modules.

diff --git a/Engine/Plugins/FX/Niagara/Source/NiagaraEditor/Private/Toolkits/NiagaraScriptToolkit.cpp b/Engine/Plugins/FX/Niagara/Source/NiagaraEditor/Private/Toolkits/NiagaraScriptToolkit.cpp
--- a/Engine/Plugins/FX/Niagara/Source/NiagaraEditor/Private/Toolkits/NiagaraScriptToolkit.cpp
+++ b/Engine/Plugins/FX/Niagara/Source/NiagaraEditor/Private/Toolkits/NiagaraScriptToolkit.cpp
@@ -343,9 +343,14 @@ void FNiagaraScriptToolkit::OnEditedScriptPropertyFinishedChanging(const FProper
 void FNiagaraScriptToolkit::OnVMScriptCompiled(UNiagaraScript* InScript)
 {
 	UpdateModuleStats();
-	ParameterPanelViewModel->Refresh();
+	if (ParameterPanelViewModel.IsValid())
+	{
+		ParameterPanelViewModel->Refresh();
+	}
 	if (SelectedDetailsWidget.IsValid())
+	{
 		SelectedDetailsWidget->SelectedObjectsChanged();
+	}
 }
 
 TSharedRef<SDockTab> FNiagaraScriptToolkit::SpawnTabScriptDetails(const FSpawnTabArgs& Args)
