The computation of the emitter execution order takes into account events, so that emitters which produce events run before the emitters which consume those events.

This makes it possible to gather event sources during emitter tick, instead of pre-tick, so events are no longer delayed by one frame.

diff --git a/Engine/Plugins/FX/Niagara/Source/Niagara/Private/NiagaraEmitterInstance.cpp b/Engine/Plugins/FX/Niagara/Source/Niagara/Private/NiagaraEmitterInstance.cpp
--- a/Engine/Plugins/FX/Niagara/Source/Niagara/Private/NiagaraEmitterInstance.cpp
+++ b/Engine/Plugins/FX/Niagara/Source/Niagara/Private/NiagaraEmitterInstance.cpp
@@ -1045,31 +1045,6 @@ void FNiagaraEmitterInstance::PreTick()
 		bResetPending = false;
 	}
 
-	if (EventInstanceData.IsValid())
-	{
-		//Gather events we're going to be reading from / handling this frame.
-		//We must do this in pre-tick so we can gather (and mark in use) all sets from other emitters.
-		const int32 NumEventHandlers = CachedEmitter->GetEventHandlers().Num();
-		EventInstanceData->EventSpawnTotal = 0;
-		for (int32 i = 0; i < NumEventHandlers; i++)
-		{
-			const FNiagaraEventScriptProperties& EventHandlerProps = CachedEmitter->GetEventHandlers()[i];
-			FNiagaraEventHandlingInfo& Info = EventInstanceData->EventHandlingInfo[i];
-
-			Info.TotalSpawnCount = 0;//This was being done every frame but should be done in init?
-			Info.SpawnCounts.Reset();
-
-			//TODO: We can move this lookup into the init and just store a ptr to the other set?
-			if (FNiagaraDataSet* EventSet = ParentSystemInstance->GetEventDataSet(Info.SourceEmitterName, EventHandlerProps.SourceEventName))
-			{
-				Info.SetEventData(&EventSet->GetCurrentDataChecked());
-				uint32 EventSpawnNum = CalculateEventSpawnCount(EventHandlerProps, Info.SpawnCounts, EventSet);
-				Info.TotalSpawnCount += EventSpawnNum;
-				EventInstanceData->EventSpawnTotal += EventSpawnNum;
-			}
-		}
-	}
-
 	++TickCount;
 	ParticleDataSet->SetIDAcquireTag(TickCount);
 }
@@ -1161,8 +1136,31 @@ void FNiagaraEmitterInstance::Tick(float DeltaSeconds)
 	checkSlow(CachedEmitter->SpawnScriptProps.Script);
 	checkSlow(CachedEmitter->UpdateScriptProps.Script);
 	
-	//TickEvents(DeltaSeconds);
-	
+	if (EventInstanceData.IsValid())
+	{
+		// Set up the spawn counts and source datasets for the events. The system ensures that we will run after any emitters
+		// we're receiving from, so we can use the data buffers that our sources have computed this tick.
+		const int32 NumEventHandlers = CachedEmitter->GetEventHandlers().Num();
+		EventInstanceData->EventSpawnTotal = 0;
+		for (int32 i = 0; i < NumEventHandlers; i++)
+		{
+			const FNiagaraEventScriptProperties& EventHandlerProps = CachedEmitter->GetEventHandlers()[i];
+			FNiagaraEventHandlingInfo& Info = EventInstanceData->EventHandlingInfo[i];
+
+			Info.TotalSpawnCount = 0;//This was being done every frame but should be done in init?
+			Info.SpawnCounts.Reset();
+
+			//TODO: We can move this lookup into the init and just store a ptr to the other set?
+			if (FNiagaraDataSet* EventSet = ParentSystemInstance->GetEventDataSet(Info.SourceEmitterName, EventHandlerProps.SourceEventName))
+			{
+				Info.SetEventData(&EventSet->GetCurrentDataChecked());
+				uint32 EventSpawnNum = CalculateEventSpawnCount(EventHandlerProps, Info.SpawnCounts, EventSet);
+				Info.TotalSpawnCount += EventSpawnNum;
+				EventInstanceData->EventSpawnTotal += EventSpawnNum;
+			}
+		}
+	}
+
 	// Calculate number of new particles from regular spawning 
 	uint32 SpawnTotal = 0;
 	if (ExecutionState == ENiagaraExecutionState::Active)
diff --git a/Engine/Plugins/FX/Niagara/Source/Niagara/Private/NiagaraSystemInstance.cpp b/Engine/Plugins/FX/Niagara/Source/Niagara/Private/NiagaraSystemInstance.cpp
--- a/Engine/Plugins/FX/Niagara/Source/Niagara/Private/NiagaraSystemInstance.cpp
+++ b/Engine/Plugins/FX/Niagara/Source/Niagara/Private/NiagaraSystemInstance.cpp
@@ -671,6 +671,41 @@ void FNiagaraSystemInstance::FindDataInterfaceDependencies(const TArray<UNiagara
 	}
 }
 
+void FNiagaraSystemInstance::FindEventDependencies(FNiagaraEmitterInstance& EmitterInst, TArray<FNiagaraEmitterInstance*>& Dependencies)
+{
+	UNiagaraEmitter* Emitter = EmitterInst.GetCachedEmitter();
+	if (!Emitter)
+	{
+		return;
+	}
+
+	const TArray<FNiagaraEmitterHandle>& EmitterHandles = GetSystem()->GetEmitterHandles();
+
+	const TArray<FNiagaraEventScriptProperties>& EventHandlers = Emitter->GetEventHandlers();
+	for (const FNiagaraEventScriptProperties& Handler : EventHandlers)
+	{
+		// An empty ID means the event reads from the same emitter, so we don't need to record a dependency.
+		if (!Handler.SourceEmitterID.IsValid())
+		{
+			continue;
+		}
+
+		// Look for the ID in the list of emitter handles from the system object.
+		FString SourceEmitterIDName = Handler.SourceEmitterID.ToString();
+		for (int EmitterIdx = 0; EmitterIdx < EmitterHandles.Num(); ++EmitterIdx)
+		{
+			FName EmitterIDName = EmitterHandles[EmitterIdx].GetIdName();
+			if (EmitterIDName.ToString() == SourceEmitterIDName)
+			{
+				// The Emitters array is in the same order as the EmitterHandles array.
+				FNiagaraEmitterInstance* Sender = &Emitters[EmitterIdx].Get();
+				Dependencies.Add(Sender);
+				break;
+			}
+		}
+	}
+}
+
 void FNiagaraSystemInstance::ComputeEmittersExecutionOrder()
 {
 	const int32 NumEmitters = Emitters.Num();
@@ -696,12 +731,16 @@ void FNiagaraSystemInstance::ComputeEmittersExecutionOrder()
 
 		if (Inst.GetCachedEmitter() && Inst.GetCachedEmitter()->SimTarget == ENiagaraSimTarget::GPUComputeSim && Inst.GetGPUContext())
 		{
+			// GPU emitters have a combined execution context for spawn and update.
 			FindDataInterfaceDependencies(Inst.GetGPUContext()->GetDataInterfaces(), EmitterDependencies);
 		}
 		else
 		{
+			// CPU emitters have separate contexts for spawn and update, so we need to gather DIs from both. They also support events,
+			// so we need to look at the event sources for extra dependencies.
 			FindDataInterfaceDependencies(Inst.GetSpawnExecutionContext().GetDataInterfaces(), EmitterDependencies);
 			FindDataInterfaceDependencies(Inst.GetUpdateExecutionContext().GetDataInterfaces(), EmitterDependencies);
+			FindEventDependencies(Inst, EmitterDependencies);
 		}
 
 		// Map the pointers returned by the emitter to indices inside the Emitters array. This is O(N^2), but we expect
diff --git a/Engine/Plugins/FX/Niagara/Source/Niagara/Public/NiagaraSystemInstance.h b/Engine/Plugins/FX/Niagara/Source/Niagara/Public/NiagaraSystemInstance.h
--- a/Engine/Plugins/FX/Niagara/Source/Niagara/Public/NiagaraSystemInstance.h
+++ b/Engine/Plugins/FX/Niagara/Source/Niagara/Public/NiagaraSystemInstance.h
@@ -306,6 +306,9 @@ private:
 	/** Queries all the data interfaces in the array for emitter dependencies. */
 	void FindDataInterfaceDependencies(const TArray<UNiagaraDataInterface*>& DataInterfaces, TArray<FNiagaraEmitterInstance*>& Dependencies);
 
+	/** Looks at all the event handlers in the emitter to determine which other emitters it depends on. */
+	void FindEventDependencies(FNiagaraEmitterInstance& EmitterInst, TArray<FNiagaraEmitterInstance*>& Dependencies);
+
 	/** Computes the order in which the emitters in the Emitters array will be ticked and stores the results in EmitterExecutionOrder. */
 	void ComputeEmittersExecutionOrder();
 
