Altering how default effect types.

Previous behavior would swap out a null effect type for the default under the hood. People found this confusing.
New behavior will set the default effect type for new assets but allow None to mean None.

diff --git a/Engine/Plugins/FX/Niagara/Source/Niagara/Classes/NiagaraSystem.h b/Engine/Plugins/FX/Niagara/Source/Niagara/Classes/NiagaraSystem.h
--- a/Engine/Plugins/FX/Niagara/Source/Niagara/Classes/NiagaraSystem.h
+++ b/Engine/Plugins/FX/Niagara/Source/Niagara/Classes/NiagaraSystem.h
@@ -374,6 +374,9 @@ public:
 
 	FORCEINLINE int32* GetCycleCounter(bool bGameThread, bool bConcurrent);
 
+#if WITH_EDITOR
+	void SetEffectType(UNiagaraEffectType* EffectType);
+#endif
 	UNiagaraEffectType* GetEffectType()const;
 	FORCEINLINE const FNiagaraSystemScalabilitySettings& GetScalabilitySettings() { return CurrentScalabilitySettings; }
 	
diff --git a/Engine/Plugins/FX/Niagara/Source/Niagara/Private/NiagaraSystem.cpp b/Engine/Plugins/FX/Niagara/Source/Niagara/Private/NiagaraSystem.cpp
--- a/Engine/Plugins/FX/Niagara/Source/Niagara/Private/NiagaraSystem.cpp
+++ b/Engine/Plugins/FX/Niagara/Source/Niagara/Private/NiagaraSystem.cpp
@@ -60,6 +60,9 @@ UNiagaraSystem::UNiagaraSystem(const FObjectInitializer& ObjectInitializer)
 {
 	ExposedParameters.SetOwner(this);
 	MaxPoolSize = 32;
+
+	EffectType = nullptr;
+	bOverrideScalabilitySettings = false;
 }
 
 UNiagaraSystem::UNiagaraSystem(FVTableHelper& Helper)
@@ -1440,15 +1443,22 @@ void UNiagaraSystem::GenerateStatID()const
 
 UNiagaraEffectType* UNiagaraSystem::GetEffectType()const
 {
-	if (EffectType)
-	{
-		return EffectType;
-	}
-	
-	const UNiagaraSettings* Settings = GetDefault<UNiagaraSettings>();
-	check(Settings);
-	return Settings->GetDefaultEffectType();
+	return EffectType;
+}
+
+#if WITH_EDITOR
+void UNiagaraSystem::SetEffectType(UNiagaraEffectType* InEffectType)
+{
+	if (InEffectType != EffectType)
+	{
+		Modify();
+		EffectType = InEffectType;
+		ResolveScalabilitySettings();
+		FNiagaraSystemUpdateContext UpdateCtx;
+		UpdateCtx.Add(this, true);
+	}	
 }
+#endif
 
 void UNiagaraSystem::ResolveScalabilitySettings()
 {
@@ -1571,4 +1581,4 @@ void FNiagaraParameterDataSetBindingCollection::BuildInternal(const TArray<FNiag
 	FloatOffsets.Shrink();
 	Int32Offsets.Shrink();
 }
-#endif
\ No newline at end of file
+#endif
diff --git a/Engine/Plugins/FX/Niagara/Source/NiagaraEditor/Private/NiagaraSystemFactoryNew.cpp b/Engine/Plugins/FX/Niagara/Source/NiagaraEditor/Private/NiagaraSystemFactoryNew.cpp
--- a/Engine/Plugins/FX/Niagara/Source/NiagaraEditor/Private/NiagaraSystemFactoryNew.cpp
+++ b/Engine/Plugins/FX/Niagara/Source/NiagaraEditor/Private/NiagaraSystemFactoryNew.cpp
@@ -16,6 +16,7 @@
 #include "Interfaces/IMainFrameModule.h"
 #include "Framework/Application/SlateApplication.h"
 #include "NiagaraEditorUtilities.h"
+#include "NiagaraSettings.h"
 
 #define LOCTEXT_NAMESPACE "NiagaraSystemFactory"
 
@@ -179,6 +180,10 @@ void UNiagaraSystemFactoryNew::InitializeSystem(UNiagaraSystem* System, bool bCr
 			FNiagaraStackGraphUtilities::RelayoutGraph(*SystemScriptSource->NodeGraph);
 		}
 	}
+
+	const UNiagaraSettings* Settings = GetDefault<UNiagaraSettings>();
+	check(Settings);
+	System->SetEffectType(Settings->GetDefaultEffectType());
 }
 
 #undef LOCTEXT_NAMESPACE
