Adding warning in stack for potentially confusing issue with EffectTypes and Scalability overrides until it can be fixed properly.
If you have no EffectType then Scalability overrides don't function.
Certain parts of the system scalability code need an effect type atm and a change to fix this is too risky at this stage imo.

diff --git a/Engine/Plugins/FX/Niagara/Source/Niagara/Classes/NiagaraSystem.h b/Engine/Plugins/FX/Niagara/Source/Niagara/Classes/NiagaraSystem.h
--- a/Engine/Plugins/FX/Niagara/Source/Niagara/Classes/NiagaraSystem.h
+++ b/Engine/Plugins/FX/Niagara/Source/Niagara/Classes/NiagaraSystem.h
@@ -183,6 +183,7 @@ class NIAGARA_API UNiagaraSystem : public UFXSystemAsset
 public:
 #if WITH_EDITOR
 	DECLARE_MULTICAST_DELEGATE_OneParam(FOnSystemCompiled, UNiagaraSystem*);
+	DECLARE_MULTICAST_DELEGATE_OneParam(FOnSystemPostEditChange, UNiagaraSystem*);
 #endif
 	//TestChange
 
@@ -289,6 +290,9 @@ public:
 	/** Delegate called when the system's dependencies have all been compiled.*/
 	FOnSystemCompiled& OnSystemCompiled();
 
+	/** Delegate called on PostEditChange.*/
+	FOnSystemPostEditChange& OnSystemPostEditChange();
+
 	/** Gets editor specific data stored with this system. */
 	UNiagaraEditorDataBase* GetEditorData();
 
@@ -376,6 +380,9 @@ public:
 
 #if WITH_EDITOR
 	void SetEffectType(UNiagaraEffectType* EffectType);
+
+	FORCEINLINE bool GetOverrideScalabilitySettings()const { return bOverrideScalabilitySettings; }
+	FORCEINLINE void SetOverrideScalabilitySettings(bool bOverride) { bOverrideScalabilitySettings = bOverride; }
 #endif
 	UNiagaraEffectType* GetEffectType()const;
 	FORCEINLINE const FNiagaraSystemScalabilitySettings& GetScalabilitySettings() { return CurrentScalabilitySettings; }
@@ -477,6 +484,9 @@ protected:
 
 	/** A multicast delegate which is called whenever the script has been compiled (successfully or not). */
 	FOnSystemCompiled OnSystemCompiledDelegate;
+
+	/** A multicast delegate which is called whenever this system's properties are changed. */
+	FOnSystemPostEditChange OnSystemPostEditChangeDelegate;
 #endif
 
 	/** The fixed bounding box value. bFixedBounds is the condition whether the fixed bounds can be edited. */
diff --git a/Engine/Plugins/FX/Niagara/Source/Niagara/Private/NiagaraSystem.cpp b/Engine/Plugins/FX/Niagara/Source/Niagara/Private/NiagaraSystem.cpp
--- a/Engine/Plugins/FX/Niagara/Source/Niagara/Private/NiagaraSystem.cpp
+++ b/Engine/Plugins/FX/Niagara/Source/Niagara/Private/NiagaraSystem.cpp
@@ -295,6 +295,8 @@ void UNiagaraSystem::PostEditChangeProperty(struct FPropertyChangedEvent& Proper
 	ResolveScalabilitySettings();
 
 	UpdateContext.CommitUpdate();
+	
+	OnSystemPostEditChangeDelegate.Broadcast(this);
 }
 #endif 
 
@@ -794,6 +796,11 @@ UNiagaraSystem::FOnSystemCompiled& UNiagaraSystem::OnSystemCompiled()
 	return OnSystemCompiledDelegate;
 }
 
+UNiagaraSystem::FOnSystemPostEditChange& UNiagaraSystem::OnSystemPostEditChange()
+{
+	return OnSystemPostEditChangeDelegate;
+}
+
 void UNiagaraSystem::ForceGraphToRecompileOnNextCheck()
 {
 	check(SystemSpawnScript->GetSource() == SystemUpdateScript->GetSource());
diff --git a/Engine/Plugins/FX/Niagara/Source/NiagaraEditor/Private/ViewModels/NiagaraSystemViewModel.cpp b/Engine/Plugins/FX/Niagara/Source/NiagaraEditor/Private/ViewModels/NiagaraSystemViewModel.cpp
--- a/Engine/Plugins/FX/Niagara/Source/NiagaraEditor/Private/ViewModels/NiagaraSystemViewModel.cpp
+++ b/Engine/Plugins/FX/Niagara/Source/NiagaraEditor/Private/ViewModels/NiagaraSystemViewModel.cpp
@@ -90,6 +90,8 @@ void FNiagaraSystemViewModel::Initialize(UNiagaraSystem& InSystem, FNiagaraSyste
 	OnGetSequencerAddMenuContent = InOptions.OnGetSequencerAddMenuContent;
 	SystemMessageLogGuidKey = InOptions.MessageLogGuid;	
 	
+	SystemChangedDelegateHandle = System->OnSystemPostEditChange().AddSP(this, &FNiagaraSystemViewModel::SystemChanged);
+
 	SelectionViewModel = NewObject<UNiagaraSystemSelectionViewModel>(GetTransientPackage());
 	SelectionViewModel->Initialize(this->AsShared());
 	SelectionViewModel->OnEmitterHandleIdSelectionChanged().AddSP(this, &FNiagaraSystemViewModel::SystemSelectionChanged);
@@ -144,6 +146,11 @@ void FNiagaraSystemViewModel::Cleanup()
 
 	GEditor->UnregisterForUndo(this);
 
+	if (System)
+	{
+		System->OnSystemPostEditChange().Remove(SystemChangedDelegateHandle);
+	}
+
 	// Make sure that we clear out all of our event handlers
 	UnregisterViewModelWithMap(RegisteredHandle);
 
@@ -2055,6 +2062,12 @@ void FNiagaraSystemViewModel::BuildStackModuleData(UNiagaraScript* Script, FGuid
 
 }
 
+void FNiagaraSystemViewModel::SystemChanged(UNiagaraSystem* ChangedSystem)
+{
+	check(System == ChangedSystem);
+	RefreshAll();
+}
+
 void FNiagaraSystemViewModel::StackViewModelStructureChanged()
 {
 	if (SelectionViewModel != nullptr)
diff --git a/Engine/Plugins/FX/Niagara/Source/NiagaraEditor/Private/ViewModels/Stack/NiagaraStackSystemSettingsGroup.cpp b/Engine/Plugins/FX/Niagara/Source/NiagaraEditor/Private/ViewModels/Stack/NiagaraStackSystemSettingsGroup.cpp
--- a/Engine/Plugins/FX/Niagara/Source/NiagaraEditor/Private/ViewModels/Stack/NiagaraStackSystemSettingsGroup.cpp
+++ b/Engine/Plugins/FX/Niagara/Source/NiagaraEditor/Private/ViewModels/Stack/NiagaraStackSystemSettingsGroup.cpp
@@ -16,6 +16,8 @@
 #include "ViewModels/Stack/NiagaraStackGraphUtilities.h"
 #include "ViewModels/Stack/NiagaraStackSystemPropertiesItem.h"
 #include "NiagaraConstants.h"
+#include "ViewModels/NiagaraSystemViewModel.h"
+#include "NiagaraSystem.h"
 
 
 #define LOCTEXT_NAMESPACE "UNiagaraStackParameterStoreGroup"
@@ -140,6 +142,36 @@ void UNiagaraStackSystemSettingsGroup::RefreshChildrenInternal(const TArray<UNia
 	}
 	NewChildren.Add(SystemPropertiesItem);
 
+	//Check if we're trying to override scalability settings without an EffectType. Ideally we can allow this but it's somewhat awkward so for now we just post a warning and ignore this.
+	UNiagaraSystem* System = &GetSystemViewModel()->GetSystem();
+	TWeakPtr<FNiagaraSystemViewModel> WeakSysViewModel = GetSystemViewModel();
+	if(System && System->GetOverrideScalabilitySettings() && System->GetEffectType())
+	{
+		FText FixDescription = LOCTEXT("FixOverridesWithNoEffectType", "Disable Overrides");
+		FStackIssueFix FixIssue(
+			FixDescription,
+			FStackIssueFixDelegate::CreateLambda([=]()
+				{
+					if (auto Pinned = WeakSysViewModel.Pin())
+					{
+						FScopedTransaction ScopedTransaction(FixDescription);
+						Pinned->GetSystem().Modify();
+						Pinned->GetSystem().SetOverrideScalabilitySettings(false);
+						Pinned->RefreshAll();
+					}
+				}));
+
+		FStackIssue OverridesWithNoEffectTypeWarning(
+			EStackIssueSeverity::Warning,
+			LOCTEXT("FixOverridesWithNoEffectTypeSummaryText", "Scalability overrides with no Effect Type."),
+			LOCTEXT("FixOverridesWithNoEffectTypeErrorText", "Scalability settings cannot be overriden if the System has no Effect Type."),
+			GetStackEditorDataKey(),
+			false,
+			FixIssue);
+
+		NewIssues.Add(OverridesWithNoEffectTypeWarning);
+	}
+
 	Super::RefreshChildrenInternal(CurrentChildren, NewChildren, NewIssues);
 }
 
diff --git a/Engine/Plugins/FX/Niagara/Source/NiagaraEditor/Public/ViewModels/NiagaraSystemViewModel.h b/Engine/Plugins/FX/Niagara/Source/NiagaraEditor/Public/ViewModels/NiagaraSystemViewModel.h
--- a/Engine/Plugins/FX/Niagara/Source/NiagaraEditor/Public/ViewModels/NiagaraSystemViewModel.h
+++ b/Engine/Plugins/FX/Niagara/Source/NiagaraEditor/Public/ViewModels/NiagaraSystemViewModel.h
@@ -436,6 +436,9 @@ private:
 	/** builds stack module data for use in module dependencies */
 	void BuildStackModuleData(UNiagaraScript* Script, FGuid InEmitterHandleId, TArray<FNiagaraStackModuleData>& OutStackModuleData);
 
+	/** Called when a property on the System is changed. */
+	void SystemChanged(UNiagaraSystem* ChangedSystem);
+
 	/** Called whenever one of the owned stack viewmodels structure changes. */
 	void StackViewModelStructureChanged();
 
@@ -583,4 +586,6 @@ private:
 	TArray<ENiagaraScriptCompileStatus> ScriptCompileStatuses;
 
 	TOptional<ENiagaraScriptCompileStatus> LatestCompileStatusCache;
+
+	FDelegateHandle SystemChangedDelegateHandle;
 };
\ No newline at end of file
