- Implement IsReadyForOwnerToAutoDestroy

#rnx

diff --git a/Engine/Plugins/FX/Niagara/Source/Niagara/Private/NiagaraComponent.cpp b/Engine/Plugins/FX/Niagara/Source/Niagara/Private/NiagaraComponent.cpp
--- a/Engine/Plugins/FX/Niagara/Source/Niagara/Private/NiagaraComponent.cpp
+++ b/Engine/Plugins/FX/Niagara/Source/Niagara/Private/NiagaraComponent.cpp
@@ -1213,6 +1213,11 @@ void UNiagaraComponent::OnRegister()
 	Super::OnRegister();
 }
 
+bool UNiagaraComponent::IsReadyForOwnerToAutoDestroy() const
+{
+	return !IsActive();
+}
+
 void UNiagaraComponent::OnComponentDestroyed(bool bDestroyingHierarchy)
 {
 	//UE_LOG(LogNiagara, Log, TEXT("OnComponentDestroyed %p %p"), this, SystemInstance.Get());
diff --git a/Engine/Plugins/FX/Niagara/Source/Niagara/Public/NiagaraComponent.h b/Engine/Plugins/FX/Niagara/Source/Niagara/Public/NiagaraComponent.h
--- a/Engine/Plugins/FX/Niagara/Source/Niagara/Public/NiagaraComponent.h
+++ b/Engine/Plugins/FX/Niagara/Source/Niagara/Public/NiagaraComponent.h
@@ -184,6 +184,7 @@ public:
 
 	virtual void TickComponent(float DeltaTime, enum ELevelTick TickType, FActorComponentTickFunction* ThisTickFunction) override;
 	virtual const UObject* AdditionalStatObject() const override;
+	virtual bool IsReadyForOwnerToAutoDestroy() const override;
 	virtual void OnComponentDestroyed(bool bDestroyingHierarchy) override;
 	//~ End UActorComponent Interface.
 
