- Fix broken socket transforms
- Virtual bones are not part of ref to locals so do not limit the count for bone sampling
- Triangle / vertex sampling will use ref to locals but verticies can not be bound to virtual bones so this is ok

#rnx

diff --git a/Engine/Plugins/FX/Niagara/Source/Niagara/Classes/NiagaraDataInterfaceSkeletalMesh.h b/Engine/Plugins/FX/Niagara/Source/Niagara/Classes/NiagaraDataInterfaceSkeletalMesh.h
--- a/Engine/Plugins/FX/Niagara/Source/Niagara/Classes/NiagaraDataInterfaceSkeletalMesh.h
+++ b/Engine/Plugins/FX/Niagara/Source/Niagara/Classes/NiagaraDataInterfaceSkeletalMesh.h
@@ -93,11 +93,10 @@ struct FSkeletalMeshSkinningData
 
 	FORCEINLINE int32 GetBoneCount(bool RequiresPrevious) const
 	{
-		int32 BoneCount = FMath::Min(CurrBoneRefToLocals().Num(), CurrComponentTransforms().Num());
-
+		int32 BoneCount = CurrComponentTransforms().Num();
 		if (RequiresPrevious)
 		{
-			BoneCount = FMath::Min3(BoneCount, PrevBoneRefToLocals().Num(), PrevComponentTransforms().Num());
+			BoneCount = FMath::Min(BoneCount, PrevComponentTransforms().Num());
 		}
 
 		return BoneCount;
