- Manually merge code to improve Niagara draw event information, allows you to trace back to the source emitter and adds instruction count

#rnx

diff --git a/Engine/Plugins/FX/Niagara/Source/Niagara/Classes/NiagaraScriptExecutionContext.h b/Engine/Plugins/FX/Niagara/Source/Niagara/Classes/NiagaraScriptExecutionContext.h
--- a/Engine/Plugins/FX/Niagara/Source/Niagara/Classes/NiagaraScriptExecutionContext.h
+++ b/Engine/Plugins/FX/Niagara/Source/Niagara/Classes/NiagaraScriptExecutionContext.h
@@ -216,7 +216,7 @@ struct FNiagaraComputeExecutionContext
 
 	void Reset(NiagaraEmitterInstanceBatcher* Batcher);
 
-	void InitParams(UNiagaraScript* InGPUComputeScript, ENiagaraSimTarget InSimTarget, const FString& InDebugSimName, const uint32 InDefaultSimulationStageIndex, int32 InMaxUpdateIterations, const TSet<uint32> InSpawnStages);
+	void InitParams(UNiagaraScript* InGPUComputeScript, ENiagaraSimTarget InSimTarget, const uint32 InDefaultSimulationStageIndex, int32 InMaxUpdateIterations, const TSet<uint32> InSpawnStages);
 	void DirtyDataInterfaces();
 	bool Tick(FNiagaraSystemInstance* ParentSystemInstance);
 
@@ -233,6 +233,14 @@ struct FNiagaraComputeExecutionContext
 		uint32 CPUCount = 0;
 	}  EmitterInstanceReadback;
 	
+#if !UE_BUILD_SHIPPING
+	const TCHAR* GetDebugSimName() const { return *DebugSimName; }
+	void SetDebugName(FString InDebugName) { DebugSimName = InDebugName; }
+#else
+	const TCHAR* GetDebugSimName() const { return TEXT(""); }
+	void SetDebugName(FString InDebugName) { }
+#endif
+
 private:
 	void ResetInternal(NiagaraEmitterInstanceBatcher* Batcher);
 
@@ -240,11 +248,7 @@ public:
 	static uint32 TickCounter;
 
 #if !UE_BUILD_SHIPPING
-	//Persistent state 
 	FString DebugSimName;
-	FORCEINLINE const TCHAR* GetDebugSimName() const { return *DebugSimName; }
-#else
-	FORCEINLINE const TCHAR* GetDebugSimName() const { return TEXT(""); }
 #endif
 
 	const TArray<UNiagaraDataInterface*>& GetDataInterfaces()const { return CombinedParamStore.GetDataInterfaces(); }
diff --git a/Engine/Plugins/FX/Niagara/Source/Niagara/Private/NiagaraEmitterInstance.cpp b/Engine/Plugins/FX/Niagara/Source/Niagara/Private/NiagaraEmitterInstance.cpp
--- a/Engine/Plugins/FX/Niagara/Source/Niagara/Private/NiagaraEmitterInstance.cpp
+++ b/Engine/Plugins/FX/Niagara/Source/Niagara/Private/NiagaraEmitterInstance.cpp
@@ -372,7 +372,10 @@ void FNiagaraEmitterInstance::Init(int32 InEmitterIdx, FNiagaraSystemInstanceID
 		{
 			GPUExecContext = new FNiagaraComputeExecutionContext();
 			const uint32 MaxUpdateIterations = CachedEmitter->bDeprecatedShaderStagesEnabled ? CachedEmitter->MaxUpdateIterations : 1;
-			GPUExecContext->InitParams(CachedEmitter->GetGPUComputeScript(), CachedEmitter->SimTarget, CachedEmitter->GetUniqueEmitterName(), CachedEmitter->DefaultShaderStageIndex, MaxUpdateIterations, CachedEmitter->SpawnStages);
+			GPUExecContext->InitParams(CachedEmitter->GetGPUComputeScript(), CachedEmitter->SimTarget, CachedEmitter->DefaultShaderStageIndex, MaxUpdateIterations, CachedEmitter->SpawnStages);
+#if !UE_BUILD_SHIPPING
+			GPUExecContext->SetDebugName(CachedEmitter->GetFullName());
+#endif
 			GPUExecContext->MainDataSet = &Data;
 			GPUExecContext->GPUScript_RT = CachedEmitter->GetGPUComputeScript()->GetRenderThreadScript();
 
diff --git a/Engine/Plugins/FX/Niagara/Source/Niagara/Private/NiagaraEmitterInstanceBatcher.cpp b/Engine/Plugins/FX/Niagara/Source/Niagara/Private/NiagaraEmitterInstanceBatcher.cpp
--- a/Engine/Plugins/FX/Niagara/Source/Niagara/Private/NiagaraEmitterInstanceBatcher.cpp
+++ b/Engine/Plugins/FX/Niagara/Source/Niagara/Private/NiagaraEmitterInstanceBatcher.cpp
@@ -1304,10 +1304,12 @@ void NiagaraEmitterInstanceBatcher::Run(const FNiagaraGPUSystemTick& Tick, const
 		SimulationStageIndex);
 		*/
 
-	SCOPED_DRAW_EVENTF(RHICmdList, NiagaraGPUSimulationCS, TEXT("Niagara Gpu Sim - %s - NumInstances: %u - StageNumber: %u"),
+	SCOPED_DRAW_EVENTF(RHICmdList, NiagaraGPUSimulationCS, TEXT("Niagara Gpu Sim - %s - NumInstances: %u - StageNumber: %u - NumInstructions %u"),
 		Context->GetDebugSimName(),
 		TotalNumInstances,
-		SimulationStageIndex);
+		SimulationStageIndex,
+		Shader->GetNumInstructions()
+	);
 
 	//UE_LOG(LogNiagara, Warning, TEXT("Run"));
 	
diff --git a/Engine/Plugins/FX/Niagara/Source/Niagara/Private/NiagaraScriptExecutionContext.cpp b/Engine/Plugins/FX/Niagara/Source/Niagara/Private/NiagaraScriptExecutionContext.cpp
--- a/Engine/Plugins/FX/Niagara/Source/Niagara/Private/NiagaraScriptExecutionContext.cpp
+++ b/Engine/Plugins/FX/Niagara/Source/Niagara/Private/NiagaraScriptExecutionContext.cpp
@@ -553,11 +553,8 @@ void FNiagaraComputeExecutionContext::Reset(NiagaraEmitterInstanceBatcher* Batch
 	);
 }
 
-void FNiagaraComputeExecutionContext::InitParams(UNiagaraScript* InGPUComputeScript, ENiagaraSimTarget InSimTarget, const FString& InDebugSimName, const uint32 InDefaultSimulationStageIndex, const int32 InMaxUpdateIterations, const TSet<uint32> InSpawnStages)
+void FNiagaraComputeExecutionContext::InitParams(UNiagaraScript* InGPUComputeScript, ENiagaraSimTarget InSimTarget, const uint32 InDefaultSimulationStageIndex, const int32 InMaxUpdateIterations, const TSet<uint32> InSpawnStages)
 {
-#if !UE_BUILD_SHIPPING
-	DebugSimName = InDebugSimName;
-#endif
 	GPUScript = InGPUComputeScript;
 	CombinedParamStore.InitFromOwningContext(InGPUComputeScript, InSimTarget, true);
 	DefaultSimulationStageIndex = InDefaultSimulationStageIndex;
