- Fix message about being exposed to library, test was reversed

#rnx

diff --git a/Engine/Plugins/FX/Niagara/Source/NiagaraEditor/Private/NiagaraEditorUtilities.cpp b/Engine/Plugins/FX/Niagara/Source/NiagaraEditor/Private/NiagaraEditorUtilities.cpp
--- a/Engine/Plugins/FX/Niagara/Source/NiagaraEditor/Private/NiagaraEditorUtilities.cpp
+++ b/Engine/Plugins/FX/Niagara/Source/NiagaraEditor/Private/NiagaraEditorUtilities.cpp
@@ -1462,7 +1462,7 @@ NIAGARAEDITOR_API FText FNiagaraEditorUtilities::FormatScriptName(FName Name, bo
 
 FText FNiagaraEditorUtilities::FormatScriptDescription(FText Description, FName Path, bool bIsInLibrary)
 {
-	FText LibrarySuffix = bIsInLibrary
+	FText LibrarySuffix = !bIsInLibrary
 		? LOCTEXT("LibrarySuffix", "\n* Script is not exposed to the library.")
 		: FText();
 
