Fixed a bug where the read/write event nodes would not update their pins if the referenced struct changed

diff --git a/Engine/Plugins/FX/Niagara/Source/NiagaraEditor/Private/NiagaraNodeDataSetBase.cpp b/Engine/Plugins/FX/Niagara/Source/NiagaraEditor/Private/NiagaraNodeDataSetBase.cpp
--- a/Engine/Plugins/FX/Niagara/Source/NiagaraEditor/Private/NiagaraNodeDataSetBase.cpp
+++ b/Engine/Plugins/FX/Niagara/Source/NiagaraEditor/Private/NiagaraNodeDataSetBase.cpp
@@ -85,7 +85,11 @@ void UNiagaraNodeDataSetBase::PostLoad()
 		UE_LOG(LogNiagaraEditor, Display, TEXT("Niagara script '%s' references struct asset '%s'"), *GetFullName(), *ExternalStructAsset->GetFullName());
 	}*/
 	
-	IsSynchronizedWithStruct(true, nullptr, true);
+	if (!IsSynchronizedWithStruct(true, nullptr, true))
+	{
+		SynchronizeWithStruct();
+		ReallocatePins();
+	}
 
 	const int32 NiagaraVer = GetLinkerCustomVersion(FNiagaraCustomVersion::GUID);
 
