Niagara Override Parameters now inherit from blueprints to instances of those blueprints.

Known limitation: child blueprints will not inherit changes made to parent blueprints.

This is accomplished by having 2 separate TMaps in UNiagaraComponent - one for overrides in the instance (InstanceOverrideParameters), and another for changes in the template (TemplateOverrideParameters). Because any change to a TMap is counted as overriding the entire property rather than as a delta to a specific key, changes made in the instance block inheritance from the template, which leads to poor UX. This is a compromise, in that we still do not support arbitrary inheritance hierarchies.

diff --git a/Engine/Plugins/FX/Niagara/Source/Niagara/Private/NiagaraComponent.cpp b/Engine/Plugins/FX/Niagara/Source/Niagara/Private/NiagaraComponent.cpp
--- a/Engine/Plugins/FX/Niagara/Source/Niagara/Private/NiagaraComponent.cpp
+++ b/Engine/Plugins/FX/Niagara/Source/Niagara/Private/NiagaraComponent.cpp
@@ -941,6 +941,10 @@ void UNiagaraComponent::ActivateInternal(bool bReset /* = false */, bool bIsScal
 			CancelAutoAttachment(true);
 		}
 	}
+	
+#if WITH_EDITOR
+	ApplyOverridesToParameterStore();
+#endif
 
 	FNiagaraSystemInstance::EResetMode ResetMode = FNiagaraSystemInstance::EResetMode::ResetSystem;
 	if (InitializeSystem())
@@ -1516,7 +1520,7 @@ void UNiagaraComponent::SetVariableLinearColor(FName InVariableName, const FLine
 	const FNiagaraVariable VariableDesc(FNiagaraTypeDefinition::GetColorDef(), InVariableName);
 	OverrideParameters.SetParameterValue(InValue, VariableDesc, true);
 #if WITH_EDITOR
-	SetParameterValueOverriddenLocally(VariableDesc, true, false);
+	SetParameterOverride(VariableDesc, FNiagaraVariant(&InValue, sizeof(FLinearColor)));
 #endif
 }
 
@@ -1530,7 +1534,7 @@ void UNiagaraComponent::SetVariableQuat(FName InVariableName, const FQuat& InVal
 	const FNiagaraVariable VariableDesc(FNiagaraTypeDefinition::GetQuatDef(), InVariableName);
 	OverrideParameters.SetParameterValue(InValue, VariableDesc, true);
 #if WITH_EDITOR
-	SetParameterValueOverriddenLocally(VariableDesc, true, false);
+	SetParameterOverride(VariableDesc, FNiagaraVariant(&InValue, sizeof(FQuat)));
 #endif
 }
 
@@ -1544,7 +1548,7 @@ void UNiagaraComponent::SetVariableVec4(FName InVariableName, const FVector4& In
 	const FNiagaraVariable VariableDesc(FNiagaraTypeDefinition::GetVec4Def(), InVariableName);
 	OverrideParameters.SetParameterValue(InValue, VariableDesc, true);
 #if WITH_EDITOR
-	SetParameterValueOverriddenLocally(VariableDesc, true, false);
+	SetParameterOverride(VariableDesc, FNiagaraVariant(&InValue, sizeof(FVector4)));
 #endif
 }
 
@@ -1558,7 +1562,7 @@ void UNiagaraComponent::SetVariableVec3(FName InVariableName, FVector InValue)
 	const FNiagaraVariable VariableDesc(FNiagaraTypeDefinition::GetVec3Def(), InVariableName);
 	OverrideParameters.SetParameterValue(InValue, VariableDesc, true);
 #if WITH_EDITOR
-	SetParameterValueOverriddenLocally(VariableDesc, true, false);
+	SetParameterOverride(VariableDesc, FNiagaraVariant(&InValue, sizeof(FVector)));
 #endif
 }
 
@@ -1572,7 +1576,7 @@ void UNiagaraComponent::SetVariableVec2(FName InVariableName, FVector2D InValue)
 	const FNiagaraVariable VariableDesc(FNiagaraTypeDefinition::GetVec2Def(), InVariableName);
 	OverrideParameters.SetParameterValue(InValue, VariableDesc, true);
 #if WITH_EDITOR
-	SetParameterValueOverriddenLocally(VariableDesc, true, false);
+	SetParameterOverride(VariableDesc, FNiagaraVariant(&InValue, sizeof(FVector2D)));
 #endif
 }
 
@@ -1586,7 +1590,7 @@ void UNiagaraComponent::SetVariableFloat(FName InVariableName, float InValue)
 	const FNiagaraVariable VariableDesc(FNiagaraTypeDefinition::GetFloatDef(), InVariableName);
 	OverrideParameters.SetParameterValue(InValue, VariableDesc, true);
 #if WITH_EDITOR
-	SetParameterValueOverriddenLocally(VariableDesc, true, false);
+	SetParameterOverride(VariableDesc, FNiagaraVariant(&InValue, sizeof(float)));
 #endif
 }
 
@@ -1600,7 +1604,7 @@ void UNiagaraComponent::SetVariableInt(FName InVariableName, int32 InValue)
 	const FNiagaraVariable VariableDesc(FNiagaraTypeDefinition::GetIntDef(), InVariableName);
 	OverrideParameters.SetParameterValue(InValue, VariableDesc, true);
 #if WITH_EDITOR
-	SetParameterValueOverriddenLocally(VariableDesc, true, false);
+	SetParameterOverride(VariableDesc, FNiagaraVariant(&InValue, sizeof(int32)));
 #endif
 }
 
@@ -1614,7 +1618,7 @@ void UNiagaraComponent::SetVariableBool(FName InVariableName, bool InValue)
 	const FNiagaraVariable VariableDesc(FNiagaraTypeDefinition::GetBoolDef(), InVariableName);
 	OverrideParameters.SetParameterValue(InValue ? FNiagaraBool::True : FNiagaraBool::False, VariableDesc, true);
 #if WITH_EDITOR
-	SetParameterValueOverriddenLocally(VariableDesc, true, false);
+	SetParameterOverride(VariableDesc, FNiagaraVariant(&InValue, sizeof(bool)));
 #endif
 }
 
@@ -1638,7 +1642,7 @@ void UNiagaraComponent::SetVariableObject(FName InVariableName, UObject* InValue
 	const FNiagaraVariable VariableDesc(FNiagaraTypeDefinition::GetUObjectDef(), InVariableName);
 	OverrideParameters.SetUObject(InValue, VariableDesc);
 #if WITH_EDITOR
-	SetParameterValueOverriddenLocally(VariableDesc, true, false);
+	SetParameterOverride(VariableDesc, FNiagaraVariant(InValue));
 #endif
 }
 
@@ -1652,7 +1656,7 @@ void UNiagaraComponent::SetVariableMaterial(FName InVariableName, UMaterialInter
 	const FNiagaraVariable VariableDesc(FNiagaraTypeDefinition::GetUMaterialDef(), InVariableName);
 	OverrideParameters.SetUObject(InValue, VariableDesc);
 #if WITH_EDITOR
-	SetParameterValueOverriddenLocally(VariableDesc, true, false);
+	SetParameterOverride(VariableDesc, FNiagaraVariant(InValue));
 #endif
 	MarkRenderStateDirty(); // Materials might be using this on the system, so invalidate the render state to re-gather them.
 }
@@ -1688,7 +1692,7 @@ TArray<FVector> UNiagaraComponent::GetNiagaraParticleValueVec3_DebugOnly(const F
 				}
 				else
 				{
-					UE_LOG(LogNiagara, Warning, TEXT("Unabled to find variable %s on %s per-particle data. Returning zeroes."), *InValueName, *GetPathName());
+					UE_LOG(LogNiagara, Warning, TEXT("Unable to find variable %s on %s per-particle data. Returning zeroes."), *InValueName, *GetPathName());
 					for (int32 i = 0; i < NumParticles; ++i)
 					{
 						Positions[i] = FVector::ZeroVector;
@@ -1733,49 +1737,55 @@ void UNiagaraComponent::PostLoad()
 
 	OverrideParameters.PostLoad();
 
-	if (Asset)
+#if WITH_EDITOR
+	if (Asset != nullptr)
 	{
 		Asset->ConditionalPostLoad();
-#if WITH_EDITOR
-		OverrideParameters.SanityCheckData();
-		PostLoadNormalizeOverrideNames();
+		
+		UpgradeDeprecatedParameterOverrides();
 		SynchronizeWithSourceSystem();
+
 		AssetExposedParametersChangedHandle = Asset->GetExposedParameters().AddOnChangedHandler(
 			FNiagaraParameterStore::FOnChanged::FDelegate::CreateUObject(this, &UNiagaraComponent::AssetExposedParametersChanged));
-#endif
 	}
+#endif
 }
 
-void UNiagaraComponent::SetUserParametersToDefaultValues()
+static FNiagaraVariant GetParameterValueFromStore(const FNiagaraVariableBase& Var, const FNiagaraParameterStore& Store)
 {
-	OverrideParameters.Empty();
-	if (Asset == nullptr)
+	if (Var.IsDataInterface())
 	{
-		return;
+		const int32 Index = Store.IndexOf(Var);
+		if (Index != INDEX_NONE)
+		{
+			return FNiagaraVariant(Store.GetDataInterfaces()[Index]);
+		}
 	}
-
-	TArray<FNiagaraVariable> SourceVars;
-	Asset->GetExposedParameters().GetParameters(SourceVars);
-	for (FNiagaraVariable& Param : SourceVars)
+	else if (Var.IsUObject())
 	{
-		OverrideParameters.AddParameter(Param, true);
+		const int32 Index = Store.IndexOf(Var);
+		if (Index != INDEX_NONE)
+		{
+			return FNiagaraVariant(Store.GetUObjects()[Index]);
+		}
 	}
 
-	TArray<FNiagaraVariable> ExistingVars;
-	OverrideParameters.GetUserParameters(ExistingVars);
-	for (FNiagaraVariable ExistingVar : ExistingVars)
+	const uint8* ParameterData = Store.GetParameterData(Var);
+	if (ParameterData == nullptr)
 	{
-		Asset->GetExposedParameters().CopyParameterData(OverrideParameters, ExistingVar);
+		return FNiagaraVariant();
 	}
 
-	OverrideParameters.Rebind();
+	return FNiagaraVariant(ParameterData, Var.GetSizeInBytes());
 }
 
 #if WITH_EDITOR
 
 void UNiagaraComponent::PreEditChange(FProperty* PropertyAboutToChange)
 {
-	if (PropertyAboutToChange != nullptr && PropertyAboutToChange->GetFName() == GET_MEMBER_NAME_CHECKED(UNiagaraComponent, Asset) && Asset != nullptr)
+	if (PropertyAboutToChange != nullptr && 
+		PropertyAboutToChange->GetFName() == GET_MEMBER_NAME_CHECKED(UNiagaraComponent, Asset) && 
+		Asset != nullptr)
 	{
 		Asset->GetExposedParameters().RemoveOnChangedHandler(AssetExposedParametersChangedHandle);
 		DestroyInstance();
@@ -1803,37 +1813,150 @@ void UNiagaraComponent::PostEditChangeProperty(FPropertyChangedEvent& PropertyCh
 	{
 		SynchronizeWithSourceSystem();
 	}
+	else if (PropertyName == GET_MEMBER_NAME_CHECKED(UNiagaraComponent, TemplateParameterOverrides) ||
+		PropertyName == GET_MEMBER_NAME_CHECKED(UNiagaraComponent, InstanceParameterOverrides))
+	{
+		ApplyOverridesToParameterStore();
+	}
 
 	ReinitializeSystem();
 
 	Super::PostEditChangeProperty(PropertyChangedEvent);
 }
 
-void UNiagaraComponent::OverrideUObjectParameter(const FNiagaraVariable& InVar, UObject* InObj)
+#endif
+
+void UNiagaraComponent::SetUserParametersToDefaultValues()
 {
-	FNiagaraUserRedirectionParameterStore& ParamStore = GetOverrideParameters();
-	ParamStore.SetUObject(InObj, InVar);
-	SetParameterValueOverriddenLocally(InVar, true, false);
+#if WITH_EDITORONLY_DATA
+	EditorOverridesValue_DEPRECATED.Empty();
+	TemplateParameterOverrides.Empty();
+	InstanceParameterOverrides.Empty();
+#endif
+
+	OverrideParameters.Empty();
+
+	if (Asset == nullptr)
+	{
+		return;
+	}
+
+	CopyParametersFromAsset();
+	OverrideParameters.Rebind();
 }
 
-void UNiagaraComponent::SynchronizeWithSourceSystem()
+#if WITH_EDITOR
+
+void UNiagaraComponent::UpgradeDeprecatedParameterOverrides()
 {
-	// Synchronizing parameters will create new data interface objects and if the old data
-	// interface objects are currently being used by a simulation they may be destroyed due to garbage
-	// collection, so preemptively kill the instance here.
-	DestroyInstance();
+	OverrideParameters.SanityCheckData();
+	PostLoadNormalizeOverrideNames();
+	
+	TArray<FNiagaraVariable> UserParameters;
+	OverrideParameters.GetUserParameters(UserParameters);
 
-	//TODO: Look through params in system in "Owner" namespace and add to our parameters.
+	for (const TPair<FName, bool>& Pair : EditorOverridesValue_DEPRECATED)
+	{
+		FNiagaraVariable* Found = UserParameters.FindByPredicate([Pair](const FNiagaraVariable& Var)
+		{
+			return Var.GetName() == Pair.Key;
+		});
+
+		if (Found != nullptr)
+		{
+			FNiagaraVariant StoredValue = GetParameterValueFromStore(*Found, OverrideParameters);
+
+			if (StoredValue.IsValid())
+			{
+				SetParameterOverride(*Found, StoredValue);
+			}
+		}
+	}
+
+	EditorOverridesValue_DEPRECATED.Empty();
+}
+
+void UNiagaraComponent::EnsureOverrideParametersConsistent() const
+{
 	if (Asset == nullptr)
 	{
-		OverrideParameters.Empty();
-		EditorOverridesValue.Empty();
-#if WITH_EDITORONLY_DATA
-		OnSynchronizedWithAssetParametersDelegate.Broadcast();
-#endif
 		return;
 	}
+	
+	TArray<FNiagaraVariable> UserParameters;
+	Asset->GetExposedParameters().GetUserParameters(UserParameters);
+	
+	for (const FNiagaraVariable& Key : UserParameters)
+	{
+		FNiagaraVariant OverrideValue = FindParameterOverride(Key);
+		if (OverrideValue.IsValid())
+		{
+			if (Key.IsDataInterface())
+			{
+				const UNiagaraDataInterface* ActualDI = OverrideParameters.GetDataInterface(Key);
+				if (ActualDI != nullptr)
+				{
+					ensureAlways(OverrideValue.GetDataInterface()->Equals(ActualDI));
+				}
+			}
+			else if (Key.IsUObject())
+			{
+				const UObject* ActualObj = OverrideParameters.GetUObject(Key);
+				if (ActualObj != nullptr)
+				{
+					ensureAlways(OverrideValue.GetUObject() == ActualObj);
+				}
+			}
+			else
+			{
+				const uint8* ActualData = OverrideParameters.GetParameterData(Key);
+				if (ActualData != nullptr)
+				{
+					ensureAlways(FMemory::Memcmp(ActualData, OverrideValue.GetBytes(), Key.GetSizeInBytes()) == 0);
+				}
+			}
+		}
+	}
+}
 
+void UNiagaraComponent::ApplyOverridesToParameterStore()
+{
+	if (!IsTemplate())
+	{
+		if (UNiagaraComponent* Archetype = Cast<UNiagaraComponent>(GetArchetype()))
+		{
+			TemplateParameterOverrides = Archetype->TemplateParameterOverrides;
+		}
+	}
+
+	for (const auto& Pair : TemplateParameterOverrides)
+	{ 
+		const int32* ExistingParam = OverrideParameters.FindParameterOffset(Pair.Key);
+		if (ExistingParam != nullptr)
+		{
+			SetOverrideParameterStoreValue(Pair.Key, Pair.Value);
+		}
+	}
+
+	if (!IsTemplate())
+	{
+		for (const auto& Pair : InstanceParameterOverrides)
+		{ 
+			const int32* ExistingParam = OverrideParameters.FindParameterOffset(Pair.Key);
+			if (ExistingParam != nullptr)
+			{
+				SetOverrideParameterStoreValue(Pair.Key, Pair.Value);
+			}
+		}
+	}
+
+	EnsureRenderDataConsistent();
+}
+
+#endif
+
+void UNiagaraComponent::CopyParametersFromAsset()
+{
 	TArray<FNiagaraVariable> SourceVars;
 	Asset->GetExposedParameters().GetParameters(SourceVars);
 	for (FNiagaraVariable& Param : SourceVars)
@@ -1842,28 +1965,45 @@ void UNiagaraComponent::SynchronizeWithSourceSystem()
 	}
 
 	TArray<FNiagaraVariable> ExistingVars;
-	OverrideParameters.GetUserParameters(ExistingVars);
-	Asset->GetExposedParameters().GetUserParameters(SourceVars);
+	OverrideParameters.GetParameters(ExistingVars);
 
 	for (FNiagaraVariable ExistingVar : ExistingVars)
 	{
-		if (!SourceVars.Contains(ExistingVar))
+		if (SourceVars.Contains(ExistingVar))
+		{
+			Asset->GetExposedParameters().CopyParameterData(OverrideParameters, ExistingVar);
+		}
+		else
 		{
 			OverrideParameters.RemoveParameter(ExistingVar);
-			EditorOverridesValue.Remove(ExistingVar.GetName());
 		}
 	}
+}
 
-	for (FNiagaraVariable ExistingVar : ExistingVars)
-	{
-		bool* FoundVar = EditorOverridesValue.Find(ExistingVar.GetName());
+void UNiagaraComponent::SynchronizeWithSourceSystem()
+{
+	// Synchronizing parameters will create new data interface objects and if the old data
+	// interface objects are currently being used by a simulation they may be destroyed due to garbage
+	// collection, so preemptively kill the instance here.
+	DestroyInstance();
 
-		if (!IsParameterValueOverriddenLocally(ExistingVar.GetName()))
-		{
-			Asset->GetExposedParameters().CopyParameterData(OverrideParameters, ExistingVar);
-		}
+	//TODO: Look through params in system in "Owner" namespace and add to our parameters.
+	if (Asset == nullptr)
+	{
+#if WITH_EDITORONLY_DATA
+		OverrideParameters.Empty();
+		EditorOverridesValue_DEPRECATED.Empty();
+		
+		OnSynchronizedWithAssetParametersDelegate.Broadcast();
+#endif
+		return;
 	}
 
+#if WITH_EDITOR
+	CopyParametersFromAsset();
+	ApplyOverridesToParameterStore();
+#endif
+
 	OverrideParameters.Rebind();
 
 #if WITH_EDITORONLY_DATA
@@ -1876,6 +2016,143 @@ void UNiagaraComponent::AssetExposedParametersChanged()
 	SynchronizeWithSourceSystem();
 	ReinitializeSystem();
 }
+
+#if WITH_EDITOR
+
+bool UNiagaraComponent::HasParameterOverride(const FNiagaraVariableBase& InKey) const 
+{
+	if (IsTemplate())
+	{
+		const FNiagaraVariant* ThisValue = TemplateParameterOverrides.Find(InKey);
+
+		const FNiagaraVariant* ArchetypeValue = nullptr;
+		if (const UNiagaraComponent* Archetype = Cast<UNiagaraComponent>(GetArchetype()))
+		{
+			ArchetypeValue = Archetype->TemplateParameterOverrides.Find(InKey);
+		}
+
+		if (ThisValue != nullptr && ArchetypeValue != nullptr)
+		{
+			// exists in both, check values
+			return *ThisValue != *ArchetypeValue;
+		}
+		else if (ThisValue != nullptr || ArchetypeValue != nullptr)
+		{
+			// either added or removed in this
+			return true;
+		}
+	}
+	else
+	{
+		if (InstanceParameterOverrides.Contains(InKey))
+		{
+			return true;
+		}
+	}
+
+	return false;
+}
+
+FNiagaraVariant UNiagaraComponent::FindParameterOverride(const FNiagaraVariableBase& InKey) const
+{
+	if (Asset == nullptr)
+	{
+		return FNiagaraVariant();
+	}
+
+	if (Asset->GetExposedParameters().FindParameterOffset(InKey) == nullptr)
+	{
+		return FNiagaraVariant();
+	}
+
+	if (!IsTemplate())
+	{
+		const FNiagaraVariant* Value = InstanceParameterOverrides.Find(InKey);
+		if (Value != nullptr)
+		{
+			return *Value;
+		}
+	}
+
+	const FNiagaraVariant* Value = TemplateParameterOverrides.Find(InKey);
+	if (Value != nullptr)
+	{
+		return *Value;
+	}
+
+	return FNiagaraVariant();
+}
+
+void UNiagaraComponent::SetOverrideParameterStoreValue(const FNiagaraVariableBase& InKey, const FNiagaraVariant& InValue)
+{
+	if (InKey.IsDataInterface())
+	{
+		UNiagaraDataInterface* DuplicatedDI = DuplicateObject(InValue.GetDataInterface(), this);
+		OverrideParameters.SetDataInterface(DuplicatedDI, InKey);
+	}
+	else if (InKey.IsUObject())
+	{
+		OverrideParameters.SetUObject(InValue.GetUObject(), InKey);
+	}
+	else
+	{
+		OverrideParameters.SetParameterData(InValue.GetBytes(), InKey, true);
+	}
+}
+
+void UNiagaraComponent::SetParameterOverride(const FNiagaraVariableBase& InKey, const FNiagaraVariant& InValue)
+{
+	if (!ensure(InValue.IsValid()))
+	{
+		return;
+	}
+
+	if (IsTemplate())
+	{
+		TemplateParameterOverrides.Add(InKey, InValue);
+	}
+	else
+	{
+		InstanceParameterOverrides.Add(InKey, InValue);
+	}
+
+	SetOverrideParameterStoreValue(InKey, InValue);
+}
+
+void UNiagaraComponent::RemoveParameterOverride(const FNiagaraVariableBase& InKey)
+{
+	if (!IsTemplate())
+	{
+		InstanceParameterOverrides.Remove(InKey);
+	}
+	else
+	{
+		TemplateParameterOverrides.Remove(InKey);
+
+		// we are an archetype, but check if we have an archetype and inherit the value from there
+		const UNiagaraComponent* Archetype = Cast<UNiagaraComponent>(GetArchetype());
+		if (Archetype != nullptr)
+		{
+			FNiagaraVariant ArchetypeValue = Archetype->FindParameterOverride(InKey);
+			if (ArchetypeValue.IsValid())
+			{
+				// defined in archetype, reset value to that
+				if (InKey.IsDataInterface())
+				{
+					UNiagaraDataInterface* DataInterface = DuplicateObject(ArchetypeValue.GetDataInterface(), this);
+					TemplateParameterOverrides.Add(InKey, FNiagaraVariant(DataInterface));
+				}
+				else
+				{
+					TemplateParameterOverrides.Add(InKey, ArchetypeValue);
+				}
+			}
+		}
+	}
+
+	SynchronizeWithSourceSystem();
+}
+
 #endif
 
 ENiagaraAgeUpdateMode UNiagaraComponent::GetAgeUpdateMode() const
@@ -1954,125 +2231,54 @@ void UNiagaraComponent::SetAllowScalability(bool bAllow)
 void UNiagaraComponent::PostLoadNormalizeOverrideNames()
 {
 	TMap<FName, bool> ValueMap;
-	for (TPair<FName, bool> Pair : EditorOverridesValue)
-	{
-		bool IsOldUserParam = Pair.Key.ToString().StartsWith(TEXT("User."));
-		FName ValueName = IsOldUserParam ? (*Pair.Key.ToString().RightChop(5)) : Pair.Key;
-		ValueMap.Add(ValueName, Pair.Value);
-	}
-	EditorOverridesValue = ValueMap;
-}
-
-bool UNiagaraComponent::IsParameterValueOverriddenLocally(const FName& InParamName)
-{
-	bool* FoundVar = EditorOverridesValue.Find(InParamName);
-
-	if (FoundVar != nullptr && *(FoundVar))
+	for (TPair<FName, bool> Pair : EditorOverridesValue_DEPRECATED)
 	{
-		return true;
-	}
-	return false;
-}
-
-void UNiagaraComponent::SetParameterValueOverriddenLocally(const FNiagaraVariable& InParam, bool bInOverriden, bool bRequiresSystemInstanceReset)
-{
-	FName ParameterName = InParam.GetName();
-	{
-		FString ParameterString = ParameterName.ToString();
-		if (ParameterString.RemoveFromStart(TEXT("User.")))
+		FString ValueNameString = Pair.Key.ToString();
+		if (ValueNameString.StartsWith(TEXT("User.")))
 		{
-			ParameterName = FName(*ParameterString);
+			 ValueNameString = ValueNameString.RightChop(5);
 		}
-	}
 
-	bool* FoundVar = EditorOverridesValue.Find(InParam.GetName());
-	if (FoundVar != nullptr && bInOverriden) 
-	{
-		*(FoundVar) = bInOverriden;
-	}
-	else if (FoundVar == nullptr && bInOverriden)			
-	{
-		EditorOverridesValue.Add(InParam.GetName(), true);
-	}
-	else
-	{
-		EditorOverridesValue.Remove(InParam.GetName());
-		Asset->GetExposedParameters().CopyParameterData(OverrideParameters, InParam);
+		ValueMap.Add(FName(*ValueNameString), Pair.Value);
 	}
-	
-	if (InParam.IsUObject() && InParam.GetType().GetClass()->IsChildOf(UMaterialInterface::StaticClass()))
-	{
-		MarkRenderStateDirty();
-	}
-
-	if (bRequiresSystemInstanceReset && SystemInstance)
-	{
-		SystemInstance->Reset(FNiagaraSystemInstance::EResetMode::ResetAll);
-	}
-	
+	EditorOverridesValue_DEPRECATED = ValueMap;
 }
 
-
 #endif // WITH_EDITOR
 
-
-
 void UNiagaraComponent::SetAsset(UNiagaraSystem* InAsset)
 {
-	if (Asset != InAsset)
+	if (Asset == InAsset)
 	{
+		return;
+	}
+
 #if WITH_EDITOR
-		if (Asset != nullptr)
-		{
-			Asset->GetExposedParameters().RemoveOnChangedHandler(AssetExposedParametersChangedHandle);
-		}
+	if (Asset != nullptr)
+	{
+		Asset->GetExposedParameters().RemoveOnChangedHandler(AssetExposedParametersChangedHandle);
+	}
 #endif
-		Asset = InAsset;
+	Asset = InAsset;
 
 #if WITH_EDITOR
-		SynchronizeWithSourceSystem();
-		if (Asset != nullptr)
-		{
-			AssetExposedParametersChangedHandle = Asset->GetExposedParameters().AddOnChangedHandler(
-				FNiagaraParameterStore::FOnChanged::FDelegate::CreateUObject(this, &UNiagaraComponent::AssetExposedParametersChanged));
-		}
-		else
-		{
-			AssetExposedParametersChangedHandle.Reset();
-		}
+	SynchronizeWithSourceSystem();
+	if (Asset != nullptr)
+	{
+		AssetExposedParametersChangedHandle = Asset->GetExposedParameters().AddOnChangedHandler(
+			FNiagaraParameterStore::FOnChanged::FDelegate::CreateUObject(this, &UNiagaraComponent::AssetExposedParametersChanged));
+	}
+	else
+	{
+		AssetExposedParametersChangedHandle.Reset();
+	}
 #else
-		// We need to populate the override parameters here
-		{
-			TArray<FNiagaraVariable> SourceVars;
-			Asset->GetExposedParameters().GetParameters(SourceVars);
-			for (FNiagaraVariable& Param : SourceVars)
-			{
-				OverrideParameters.AddParameter(Param, true);
-			}
-
-			TArray<FNiagaraVariable> ExistingVars;
-			OverrideParameters.GetUserParameters(ExistingVars);
-			Asset->GetExposedParameters().GetUserParameters(SourceVars);
-
-			for (FNiagaraVariable ExistingVar : ExistingVars)
-			{
-				if (SourceVars.Contains(ExistingVar))
-				{
-					Asset->GetExposedParameters().CopyParameterData(OverrideParameters, ExistingVar);
-				}
-				else
-				{
-					OverrideParameters.RemoveParameter(ExistingVar);
-				}
-			}
-
-			OverrideParameters.Rebind();
-		}
+	CopyParametersFromAsset();
+	OverrideParameters.Rebind();
 #endif
 
-		//Force a reinit.
-		DestroyInstance();
-	}
+	//Force a reinit.
+	DestroyInstance();
 }
 
 void UNiagaraComponent::SetForceSolo(bool bInForceSolo) 
diff --git a/Engine/Plugins/FX/Niagara/Source/Niagara/Private/NiagaraFunctionLibrary.cpp b/Engine/Plugins/FX/Niagara/Source/Niagara/Private/NiagaraFunctionLibrary.cpp
--- a/Engine/Plugins/FX/Niagara/Source/Niagara/Private/NiagaraFunctionLibrary.cpp
+++ b/Engine/Plugins/FX/Niagara/Source/Niagara/Private/NiagaraFunctionLibrary.cpp
@@ -415,17 +415,17 @@ void UNiagaraFunctionLibrary::SetTextureObject(UNiagaraComponent* NiagaraSystem,
 		return;
 	}
 
-	UNiagaraDataInterfaceTexture* TexureDI = Cast<UNiagaraDataInterfaceTexture>(OverrideParameters.GetDataInterface(Index));
-	if (!TexureDI)
+	UNiagaraDataInterfaceTexture* TextureDI = Cast<UNiagaraDataInterfaceTexture>(OverrideParameters.GetDataInterface(Index));
+	if (!TextureDI)
 	{
 		UE_LOG(LogNiagara, Warning, TEXT("Did not find a matching Texture Data Interface variable named \"%s\" in the User variables of NiagaraSystem \"%s\" ."), *OverrideName, *NiagaraSystem->GetOwner()->GetName());
 		return;
 	}
 
 #if WITH_EDITOR
-	NiagaraSystem->SetParameterValueOverriddenLocally(Variable, true, false);
+	NiagaraSystem->SetParameterOverride(Variable, FNiagaraVariant(TextureDI));
 #endif
-	TexureDI->SetTexture(Texture);
+	TextureDI->SetTexture(Texture);
 }
 
 void UNiagaraFunctionLibrary::SetVolumeTextureObject(UNiagaraComponent* NiagaraSystem, const FString& OverrideName, UVolumeTexture* Texture)
@@ -453,17 +453,17 @@ void UNiagaraFunctionLibrary::SetVolumeTextureObject(UNiagaraComponent* NiagaraS
 		return;
 	}
 
-	UNiagaraDataInterfaceVolumeTexture* TexureDI = Cast<UNiagaraDataInterfaceVolumeTexture>(OverrideParameters.GetDataInterface(Index));
-	if (!TexureDI)
+	UNiagaraDataInterfaceVolumeTexture* TextureDI = Cast<UNiagaraDataInterfaceVolumeTexture>(OverrideParameters.GetDataInterface(Index));
+	if (!TextureDI)
 	{
 		UE_LOG(LogNiagara, Warning, TEXT("Did not find a matching Volume Texture Data Interface variable named \"%s\" in the User variables of NiagaraSystem \"%s\" ."), *OverrideName, *NiagaraSystem->GetOwner()->GetName());
 		return;
 	}
 
 #if WITH_EDITOR
-	NiagaraSystem->SetParameterValueOverriddenLocally(Variable, true, false);
+	NiagaraSystem->SetParameterOverride(Variable, FNiagaraVariant(TextureDI));
 #endif
-	TexureDI->SetTexture(Texture);
+	TextureDI->SetTexture(Texture);
 }
 
 UNiagaraParameterCollectionInstance* UNiagaraFunctionLibrary::GetNiagaraParameterCollection(UObject* WorldContextObject, UNiagaraParameterCollection* Collection)
diff --git a/Engine/Plugins/FX/Niagara/Source/Niagara/Private/NiagaraVariant.cpp b/Engine/Plugins/FX/Niagara/Source/Niagara/Private/NiagaraVariant.cpp
new file mode 100644
--- /dev/null
+++ b/Engine/Plugins/FX/Niagara/Source/Niagara/Private/NiagaraVariant.cpp
@@ -0,0 +1,121 @@
+// Copyright Epic Games, Inc. All Rights Reserved.
+
+#include "NiagaraVariant.h"
+#include "NiagaraDataInterface.h"
+
+FNiagaraVariant::FNiagaraVariant()
+{
+	CurrentMode = ENiagaraVariantMode::None;
+	Object = nullptr;
+	DataInterface = nullptr;
+}
+
+FNiagaraVariant::FNiagaraVariant(const FNiagaraVariant& Other)
+{
+	CurrentMode = Other.CurrentMode;
+	Object = Other.Object;
+	DataInterface = Other.DataInterface;
+	Bytes = Other.Bytes;
+}
+
+FNiagaraVariant::FNiagaraVariant(UObject* InObject)
+{
+	CurrentMode = ENiagaraVariantMode::Object;
+	Object = InObject;
+	DataInterface = nullptr;
+}
+
+FNiagaraVariant::FNiagaraVariant(UNiagaraDataInterface* InDataInterface)
+{
+	CurrentMode = ENiagaraVariantMode::DataInterface;
+	DataInterface = InDataInterface;
+	Object = nullptr;
+}
+
+FNiagaraVariant::FNiagaraVariant(const TArray<uint8>& InBytes)
+{
+	CurrentMode = ENiagaraVariantMode::Bytes;
+	Bytes = InBytes;
+	DataInterface = nullptr;
+	Object = nullptr;
+}
+
+FNiagaraVariant::FNiagaraVariant(const void* InBytes, int32 Size)
+{
+	CurrentMode = ENiagaraVariantMode::Bytes;
+	Bytes.Append((const uint8*) InBytes, Size);
+	DataInterface = nullptr;
+	Object = nullptr;
+}
+
+UObject* FNiagaraVariant::GetUObject() const 
+{
+	ensure(CurrentMode == ENiagaraVariantMode::Object);
+	return Object;
+}
+
+void FNiagaraVariant::SetUObject(UObject* InObject)
+{
+	ensure(CurrentMode == ENiagaraVariantMode::None || CurrentMode == ENiagaraVariantMode::Object);
+
+	CurrentMode = ENiagaraVariantMode::Object;
+	Object = InObject;
+}
+
+UNiagaraDataInterface* FNiagaraVariant::GetDataInterface() const 
+{
+	ensure(CurrentMode == ENiagaraVariantMode::DataInterface);
+	return DataInterface;
+}
+
+void FNiagaraVariant::SetDataInterface(UNiagaraDataInterface* InDataInterface)
+{
+	ensure(CurrentMode == ENiagaraVariantMode::None || CurrentMode == ENiagaraVariantMode::DataInterface);
+		
+	CurrentMode = ENiagaraVariantMode::DataInterface;
+	DataInterface = InDataInterface;
+}
+
+void FNiagaraVariant::SetBytes(uint8* InBytes, int32 InCount)
+{
+	check(InCount > 0);
+	ensure(CurrentMode == ENiagaraVariantMode::None || CurrentMode == ENiagaraVariantMode::Bytes);
+
+	CurrentMode = ENiagaraVariantMode::Bytes;
+	Bytes.Reset(InCount);
+	Bytes.Append(InBytes, InCount);
+}
+
+uint8* FNiagaraVariant::GetBytes() const
+{
+	ensure(CurrentMode == ENiagaraVariantMode::Bytes);
+	return (uint8*) Bytes.GetData();
+}
+
+bool FNiagaraVariant::operator==(const FNiagaraVariant& Other) const
+{
+	if (CurrentMode == Other.CurrentMode)
+	{
+		switch (CurrentMode)
+		{
+			case ENiagaraVariantMode::Bytes:
+				return Bytes == Other.Bytes;
+
+			case ENiagaraVariantMode::Object:
+				return Object == Other.Object;
+
+			case ENiagaraVariantMode::DataInterface:
+				return GetDataInterface()->Equals(Other.GetDataInterface());
+
+			case ENiagaraVariantMode::None:
+				return true;
+		}
+	}
+
+	return false;
+}
+
+bool FNiagaraVariant::operator!=(const FNiagaraVariant& Other) const
+{
+	return !(operator==(Other));
+}
diff --git a/Engine/Plugins/FX/Niagara/Source/Niagara/Public/NiagaraComponent.h b/Engine/Plugins/FX/Niagara/Source/Niagara/Public/NiagaraComponent.h
--- a/Engine/Plugins/FX/Niagara/Source/Niagara/Public/NiagaraComponent.h
+++ b/Engine/Plugins/FX/Niagara/Source/Niagara/Public/NiagaraComponent.h
@@ -5,13 +5,14 @@
 #include "CoreMinimal.h"
 #include "UObject/ObjectMacros.h"
 #include "NiagaraCommon.h"
-#include "PrimitiveViewRelevance.h"
-#include "PrimitiveSceneProxy.h"
-#include "Particles/ParticleSystemComponent.h"
-#include "NiagaraUserRedirectionParameterStore.h"
-#include "NiagaraSystemInstance.h"
 #include "NiagaraComponentPool.h"
+#include "NiagaraSystemInstance.h"
+#include "NiagaraUserRedirectionParameterStore.h"
+#include "NiagaraVariant.h"
+#include "PrimitiveSceneProxy.h"
+#include "PrimitiveViewRelevance.h"
 #include "Particles/ParticlePerfStats.h"
+#include "Particles/ParticleSystemComponent.h"
 
 #include "NiagaraComponent.generated.h"
 
@@ -68,16 +69,18 @@ private:
 	UPROPERTY(EditAnywhere, Category = "Niagara", meta = (DisplayName = "Niagara Tick Behavior"))
 	ENiagaraTickBehavior TickBehavior = ENiagaraTickBehavior::UsePrereqs;
 
-	/** Initial values for parameter overrides. 
-	TODO: This should be a minimal set of explicitly override parameters similar to how parameter collection instances override their collections store. 
-	Should expose anything in the "User" namespace.
-	*/
-	UPROPERTY(EditAnywhere, Category = Parameters)
+	UPROPERTY()
 	FNiagaraUserRedirectionParameterStore OverrideParameters;
 
 #if WITH_EDITORONLY_DATA
 	UPROPERTY()
-	TMap<FName, bool> EditorOverridesValue;
+	TMap<FName, bool> EditorOverridesValue_DEPRECATED;
+
+	UPROPERTY(EditAnywhere, Category="Niagara")
+	TMap<FNiagaraVariableBase, FNiagaraVariant> TemplateParameterOverrides;
+
+	UPROPERTY(EditAnywhere, Category="Niagara")
+	TMap<FNiagaraVariableBase, FNiagaraVariant> InstanceParameterOverrides;
 
 	FOnSystemInstanceChanged OnSystemInstanceChangedDelegate;
 
@@ -394,18 +397,32 @@ public:
 	void SetPaused(bool bInPaused);
 
 	UFUNCTION(BlueprintCallable, Category = Niagara)
-	bool IsPaused()const;
+	bool IsPaused() const;
 
 	UFUNCTION(BlueprintCallable, Category = Niagara)
-	UNiagaraDataInterface * GetDataInterface(const FString &Name);
+	UNiagaraDataInterface* GetDataInterface(const FString &Name);
 
 	//~ Begin UObject Interface.
-	virtual void PostLoad();
+	virtual void PostLoad() override;
+
 #if WITH_EDITOR
 	virtual void PreEditChange(FProperty* PropertyAboutToChange) override;
 	virtual void PostEditChangeProperty(FPropertyChangedEvent& PropertyChangedEvent) override;
-	void OverrideUObjectParameter(const FNiagaraVariable& InVar, UObject* InObj);
 
+	/** 
+	  * Find the value of an overridden parameter. 
+	  * Returns null if the parameter isn't overridden by this component.
+	  */
+	FNiagaraVariant FindParameterOverride(const FNiagaraVariableBase& InKey) const;
+
+	bool HasParameterOverride(const FNiagaraVariableBase& InKey) const;
+	void SetParameterOverride(const FNiagaraVariableBase& InKey, const FNiagaraVariant& InValue);
+
+	/** Remove an override for a given parameter if one exists. */
+	void RemoveParameterOverride(const FNiagaraVariableBase& InKey);
+
+	void UpgradeDeprecatedParameterOverrides();
+	void EnsureOverrideParametersConsistent() const;
 #endif
 	//~ End UObject Interface
 
@@ -423,8 +440,11 @@ public:
 
 #if WITH_EDITOR
 	void PostLoadNormalizeOverrideNames();
-	bool IsParameterValueOverriddenLocally(const FName& InParamName);
-	void SetParameterValueOverriddenLocally(const FNiagaraVariable& InParam, bool bInOverridden, bool bRequiresSystemInstanceReset);
+	UE_DEPRECATED(4.25, "This function is replaced by HasParameterOverride().")
+	bool IsParameterValueOverriddenLocally(const FName& InParamName) { return false; }
+
+	UE_DEPRECATED(4.25, "This function is replaced by SetParameterOverride().")
+	void SetParameterValueOverriddenLocally(const FNiagaraVariable& InParam, bool bInOverridden, bool bRequiresSystemInstanceReset) {}
 	
 	FOnSystemInstanceChanged& OnSystemInstanceChanged() { return OnSystemInstanceChangedDelegate; }
 
@@ -452,6 +472,13 @@ private:
 
 	void AssetExposedParametersChanged();
 
+	void CopyParametersFromAsset();
+	
+#if WITH_EDITOR
+	void SetOverrideParameterStoreValue(const FNiagaraVariableBase& InKey, const FNiagaraVariant& InValue);
+	void ApplyOverridesToParameterStore();
+#endif
+
 public:
 	/**
 	 * Component we automatically attach to when activated, if bAutoManageAttachment is true.
@@ -560,10 +587,6 @@ FORCEINLINE int32 UNiagaraComponent::GetPreviewLODDistance()const { return 0.0f;
 #endif
 
 
-
-
-
-
 /**
 * Scene proxy for drawing niagara particle simulations.
 */
@@ -656,4 +679,4 @@ extern float GLastRenderTimeSafetyBias;
 FORCEINLINE float UNiagaraComponent::GetSafeTimeSinceRendered(float WorldTime)const
 {
 	return FMath::Max(0.0f, WorldTime - GetLastRenderTime() - GLastRenderTimeSafetyBias);
-}
\ No newline at end of file
+}
diff --git a/Engine/Plugins/FX/Niagara/Source/Niagara/Public/NiagaraTypes.h b/Engine/Plugins/FX/Niagara/Source/Niagara/Public/NiagaraTypes.h
--- a/Engine/Plugins/FX/Niagara/Source/Niagara/Public/NiagaraTypes.h
+++ b/Engine/Plugins/FX/Niagara/Source/Niagara/Public/NiagaraTypes.h
@@ -970,11 +970,11 @@ public:
 	This can be a UClass, UStruct or UEnum.  Pointing to something like the struct for an FVector, etc.
 	In occasional situations this may be a UClass when we're dealing with DataInterface etc.
 	*/
-	UPROPERTY()
+	UPROPERTY(EditAnywhere, Category=Type)
 	UObject* ClassStructOrEnum;
 
 	// See enumeration FUnderlyingType for possible values
-	UPROPERTY()
+	UPROPERTY(EditAnywhere, Category=Type)
 	uint16 UnderlyingType;
 
 	bool Serialize(FArchive& Ar);
@@ -1311,6 +1311,11 @@ protected:
 	FNiagaraTypeDefinition TypeDef;
 };
 
+FORCEINLINE uint32 GetTypeHash(const FNiagaraVariableBase& Var)
+{
+	return HashCombine(GetTypeHash(Var.GetType()), GetTypeHash(Var.GetName()));
+}
+
 USTRUCT()
 struct FNiagaraVariable : public FNiagaraVariableBase
 {
@@ -1493,11 +1498,6 @@ private:
 	TArray<uint8> VarData;
 };
 
-FORCEINLINE uint32 GetTypeHash(const FNiagaraVariableBase& Var)
-{
-	return HashCombine(GetTypeHash(Var.GetType()), GetTypeHash(Var.GetName()));
-}
-
 template<>
 inline bool FNiagaraVariable::GetValue<bool>() const
 {
diff --git a/Engine/Plugins/FX/Niagara/Source/Niagara/Public/NiagaraVariant.h b/Engine/Plugins/FX/Niagara/Source/Niagara/Public/NiagaraVariant.h
new file mode 100644
--- /dev/null
+++ b/Engine/Plugins/FX/Niagara/Source/Niagara/Public/NiagaraVariant.h
@@ -0,0 +1,59 @@
+// Copyright Epic Games, Inc. All Rights Reserved.
+
+#pragma once
+
+#include "UObject/ObjectMacros.h"
+
+#include "NiagaraVariant.generated.h"
+
+class UNiagaraDataInterface;
+
+UENUM()
+enum class ENiagaraVariantMode
+{
+	None,
+	Object,
+	DataInterface,
+	Bytes
+};
+
+USTRUCT()
+struct NIAGARA_API FNiagaraVariant
+{
+	GENERATED_BODY()
+
+	FNiagaraVariant();
+	FNiagaraVariant(const FNiagaraVariant& Other);
+	explicit FNiagaraVariant(UNiagaraDataInterface* InDataInterface);
+	explicit FNiagaraVariant(UObject* InObject);
+	explicit FNiagaraVariant(const TArray<uint8>& InBytes);
+	FNiagaraVariant(const void* InBytes, int32 Size);
+
+	UObject* GetUObject() const;
+	void SetUObject(UObject* InObject);
+
+	UNiagaraDataInterface* GetDataInterface() const;
+	void SetDataInterface(UNiagaraDataInterface* InDataInterface);
+
+	void SetBytes(uint8* InBytes, int32 InCount);
+	uint8* GetBytes() const;
+
+	bool IsValid() const { return CurrentMode != ENiagaraVariantMode::None; }
+	ENiagaraVariantMode GetMode() const { return CurrentMode; }
+
+	bool operator==(const FNiagaraVariant& Other) const;
+	bool operator!=(const FNiagaraVariant& Other) const;
+
+private:
+	UPROPERTY(EditAnywhere, Category=Variant)
+	UObject* Object;
+
+	UPROPERTY(EditAnywhere, Category=Variant)
+	UNiagaraDataInterface* DataInterface;
+
+	UPROPERTY(EditAnywhere, Category=Variant)
+	TArray<uint8> Bytes;
+
+	UPROPERTY(EditAnywhere, Category=Variant)
+	ENiagaraVariantMode CurrentMode;
+};
\ No newline at end of file
diff --git a/Engine/Plugins/FX/Niagara/Source/NiagaraEditor/Private/Customizations/NiagaraComponentDetails.cpp b/Engine/Plugins/FX/Niagara/Source/NiagaraEditor/Private/Customizations/NiagaraComponentDetails.cpp
--- a/Engine/Plugins/FX/Niagara/Source/NiagaraEditor/Private/Customizations/NiagaraComponentDetails.cpp
+++ b/Engine/Plugins/FX/Niagara/Source/NiagaraEditor/Private/Customizations/NiagaraComponentDetails.cpp
@@ -48,30 +48,221 @@
 #include "Widgets/Layout/SBox.h"
 #include "UObject/WeakObjectPtr.h"
 #include "NiagaraUserRedirectionParameterStore.h"
+
 #define LOCTEXT_NAMESPACE "NiagaraComponentDetails"
 
-class FNiagaraComponentNodeBuilder : public IDetailCustomNodeBuilder
+static FNiagaraVariant GetParameterValueFromAsset(const FNiagaraVariableBase& Parameter, const UNiagaraComponent* Component)
+{
+	FNiagaraUserRedirectionParameterStore& UserParameterStore = Component->GetAsset()->GetExposedParameters();
+	if (Parameter.IsDataInterface())
+	{
+		int32 Index = UserParameterStore.IndexOf(Parameter);
+		if (Index != INDEX_NONE)
+		{
+			return FNiagaraVariant(UserParameterStore.GetDataInterfaces()[Index]);
+		}
+	}
+	
+	if (Parameter.IsUObject())
+	{
+		int32 Index = UserParameterStore.IndexOf(Parameter);
+		if (Index != INDEX_NONE)
+		{
+			return FNiagaraVariant(UserParameterStore.GetUObjects()[Index]);
+		}
+	}
+
+	const uint8* ParameterData = UserParameterStore.GetParameterData(Parameter);
+	if (ParameterData == nullptr)
+	{
+		return FNiagaraVariant();
+	}
+	
+	return FNiagaraVariant(ParameterData, Parameter.GetSizeInBytes());
+}
+
+static FNiagaraVariant GetCurrentParameterValue(const FNiagaraVariableBase& Parameter, const UNiagaraComponent* Component)
+{
+	FNiagaraVariant OverriddenValue = Component->FindParameterOverride(Parameter);
+	if (OverriddenValue.IsValid())
+	{
+		return OverriddenValue;
+	}
+	
+	return GetParameterValueFromAsset(Parameter, Component);
+}
+
+// Proxy class to allow us to override values on the component that are not yet overridden.
+class FNiagaraParameterProxy
 {
 public:
-	FNiagaraComponentNodeBuilder(UNiagaraComponent* InComponent)						   
+	FNiagaraParameterProxy(TWeakObjectPtr<UNiagaraComponent> InComponent, const FNiagaraVariableBase& InKey, const FNiagaraVariant& InValue, const FSimpleDelegate& InOnRebuild, TArray<TSharedPtr<IPropertyHandle>> InPropertyHandles)
 	{
+		bResettingToDefault = false;
 		Component = InComponent;
-		Component->OnSynchronizedWithAssetParameters().AddRaw(this, &FNiagaraComponentNodeBuilder::ComponentSynchronizedWithAssetParameters);
+		ParameterKey = InKey;
+		ParameterValue = InValue;
+		OnRebuild = InOnRebuild;
+		PropertyHandles = InPropertyHandles;
+	}
+
+	FReply OnResetToDefaultClicked()
+	{
+		OnResetToDefault();
+		return FReply::Handled();
+	}
+
+	void OnResetToDefault()
+	{
+		UNiagaraComponent* RawComponent = Component.Get();
+		if (RawComponent != nullptr)
+		{
+			FScopedTransaction ScopedTransaction(NSLOCTEXT("UnrealEd", "PropertyWindowResetToDefault", "Reset to Default"));
+			RawComponent->Modify();
+
+			bResettingToDefault = true;
+
+			for (TSharedPtr<IPropertyHandle> PropertyHandle : PropertyHandles)
+			{
+				PropertyHandle->NotifyPreChange();
+			}
+
+			RawComponent->RemoveParameterOverride(ParameterKey);
+
+			for (TSharedPtr<IPropertyHandle> PropertyHandle : PropertyHandles)
+			{
+				PropertyHandle->NotifyPostChange(EPropertyChangeType::ValueSet);
+			}
+
+			OnRebuild.ExecuteIfBound();
+			bResettingToDefault = false;
+		}
+	}
+
+	EVisibility GetResetToDefaultVisibility() const 
+	{
+		UNiagaraComponent* RawComponent = Component.Get();
+		if (RawComponent != nullptr)
+		{
+			return RawComponent->HasParameterOverride(ParameterKey) ? EVisibility::Visible : EVisibility::Hidden;
+		}
+		return EVisibility::Hidden;
+	}
 
-		if (InComponent->GetAsset() != nullptr)
+	FNiagaraVariant FindExistingOverride() const
+	{
+		UNiagaraComponent* RawComponent = Component.Get();
+		if (RawComponent != nullptr)
 		{
-			UNiagaraScript* ScriptSpawn = InComponent->GetAsset()->GetSystemSpawnScript();
-			if (ScriptSpawn != nullptr)
+			return RawComponent->FindParameterOverride(ParameterKey);
+		}
+		return FNiagaraVariant();
+	}
+
+	void OnParameterPreChange()
+	{
+		UNiagaraComponent* RawComponent = Component.Get();
+		if (RawComponent != nullptr)
+		{
+			RawComponent->Modify();
+
+			for (TSharedPtr<IPropertyHandle> PropertyHandle : PropertyHandles)
 			{
-				OriginalScripts.Add(ScriptSpawn);
+				PropertyHandle->NotifyPreChange();
 			}
+		}
+	}
+
+	void OnParameterChanged()
+	{
+		if (bResettingToDefault)
+		{
+			return;
+		}
 
-			UNiagaraScript* ScriptUpdate = InComponent->GetAsset()->GetSystemUpdateScript();
-			if (ScriptUpdate != nullptr)
+		UNiagaraComponent* RawComponent = Component.Get();
+		if (RawComponent != nullptr)
+		{
+			RawComponent->SetParameterOverride(ParameterKey, ParameterValue);
+
+			for (TSharedPtr<IPropertyHandle> PropertyHandle : PropertyHandles)
+			{
+				PropertyHandle->NotifyPostChange(EPropertyChangeType::ValueSet);
+			}
+		}
+	}
+
+	void OnAssetSelectedFromPicker(const FAssetData& InAssetData)
+	{
+		UNiagaraComponent* RawComponent = Component.Get();
+		if (RawComponent != nullptr)
+		{
+			UObject* Asset = InAssetData.GetAsset();
+			if (Asset == nullptr || Asset->GetClass()->IsChildOf(ParameterKey.GetType().GetClass()))
 			{
-				OriginalScripts.Add(ScriptUpdate);
+				FScopedTransaction ScopedTransaction(LOCTEXT("ChangeAsset", "Change asset"));
+				RawComponent->Modify();
+				RawComponent->SetParameterOverride(ParameterKey, FNiagaraVariant(Asset));
+
+				for (TSharedPtr<IPropertyHandle> PropertyHandle : PropertyHandles)
+				{
+					PropertyHandle->NotifyPostChange(EPropertyChangeType::ValueSet);
+				}
 			}
 		}
+	}
+
+	FString GetCurrentAssetPath() const
+	{
+		UObject* CurrentObject = nullptr;
+
+		UNiagaraComponent* RawComponent = Component.Get();
+		if (RawComponent != nullptr)
+		{
+			FNiagaraVariant CurrentValue = FindExistingOverride();
+			if (CurrentValue.IsValid())
+			{
+				CurrentObject = CurrentValue.GetUObject();
+			}
+			else
+			{
+				// fetch from asset
+				UNiagaraSystem* System = RawComponent->GetAsset();
+				if (System != nullptr)
+				{
+					FNiagaraUserRedirectionParameterStore& AssetParamStore = System->GetExposedParameters();
+					CurrentObject = AssetParamStore.GetUObject(ParameterKey);
+				}
+			}
+		}
+
+		return CurrentObject != nullptr ? CurrentObject->GetPathName() : FString();
+	}
+
+
+	const FNiagaraVariableBase& Key() const { return ParameterKey; }
+	FNiagaraVariant& Value() { return ParameterValue; }
+
+
+private:
+	TWeakObjectPtr<UNiagaraComponent> Component;
+	TArray<TSharedPtr<IPropertyHandle>> PropertyHandles;
+	FNiagaraVariableBase ParameterKey;
+	FNiagaraVariant ParameterValue;
+	FSimpleDelegate OnRebuild;
+	bool bResettingToDefault;
+};
+
+class FNiagaraComponentNodeBuilder : public IDetailCustomNodeBuilder
+{
+public:
+	FNiagaraComponentNodeBuilder(UNiagaraComponent* InComponent, TArray<TSharedPtr<IPropertyHandle>> InOverridePropertyHandles) 
+	{
+		OverridePropertyHandles = InOverridePropertyHandles;
+
+		Component = InComponent;
+		Component->OnSynchronizedWithAssetParameters().AddRaw(this, &FNiagaraComponentNodeBuilder::ComponentSynchronizedWithAssetParameters);
+
 		//UE_LOG(LogNiagaraEditor, Log, TEXT("FNiagaraComponentNodeBuilder %p Component %p"), this, Component.Get());
 	}
 
@@ -81,6 +272,7 @@ public:
 		{
 			Component->OnSynchronizedWithAssetParameters().RemoveAll(this);
 		}
+
 		//UE_LOG(LogNiagaraEditor, Log, TEXT("~FNiagaraComponentNodeBuilder %p Component %p"), this, Component.Get());
 	}
 
@@ -99,43 +291,26 @@ public:
 		return NiagaraComponentNodeBuilder;
 	}
 
-
-	void OnAssetSelectedFromPicker(const FAssetData& InAssetData, FNiagaraVariable InVar)
+	virtual void GenerateChildContent(IDetailChildrenBuilder& ChildrenBuilder) override
 	{
-		FScopedTransaction ScopedTransaction(LOCTEXT("ChangeAsset", "Change asset"));
+		check(Component.IsValid());
 
-		UObject* Asset = InAssetData.GetAsset();
-		if (Asset == nullptr || Asset->GetClass()->IsChildOf(InVar.GetType().GetClass()))
-		{
-			check(Component.IsValid());
-			Component->Modify();
-			Component->OverrideUObjectParameter(InVar, Asset);
-		}
-	}
+		ParameterProxies.Reset();
 
-	FString GetCurrentAssetPath(FNiagaraVariable InVar) const
-	{
-		check(Component.IsValid());
-		TArray<FNiagaraVariable> Parameters;
-		FNiagaraUserRedirectionParameterStore& ParamStore = Component->GetOverrideParameters();
-		UObject* Obj = ParamStore.GetUObject(InVar);
-		if (Obj)
+		FNiagaraEditorModule& NiagaraEditorModule = FModuleManager::GetModuleChecked<FNiagaraEditorModule>("NiagaraEditor");
+
+		UNiagaraSystem* SystemAsset = Component->GetAsset();
+		if (SystemAsset == nullptr)
 		{
-			return Obj->GetPathName();
+			return;
 		}
-		return FString();
-	}
 
-	virtual void GenerateChildContent(IDetailChildrenBuilder& ChildrenBuilder) override
-	{
-		check(Component.IsValid());
-		TArray<FNiagaraVariable> Parameters;
-		FNiagaraUserRedirectionParameterStore& ParamStore = Component->GetOverrideParameters();
-		ParamStore.GetUserParameters(Parameters);
+		TArray<FNiagaraVariable> UserParameters;
+		SystemAsset->GetExposedParameters().GetUserParameters(UserParameters);
 
-		FNiagaraEditorModule& NiagaraEditorModule = FModuleManager::GetModuleChecked<FNiagaraEditorModule>("NiagaraEditor");
-		
-		for (const FNiagaraVariable& Parameter : Parameters)
+		ParameterProxies.Reserve(UserParameters.Num());
+
+		for (const FNiagaraVariable& Parameter : UserParameters)
 		{
 			TSharedPtr<SWidget> NameWidget;
 
@@ -143,18 +318,29 @@ public:
 				SNew(STextBlock)
 				.TextStyle(FNiagaraEditorStyle::Get(), "NiagaraEditor.ParameterText")
 				.Text(FText::FromName(Parameter.GetName()));
-			
+
 			IDetailPropertyRow* Row = nullptr;
 
 			TSharedPtr<SWidget> CustomValueWidget;
+			
+			FNiagaraVariant ParameterValue = GetCurrentParameterValue(Parameter, Component.Get());
+			if (!ParameterValue.IsValid())
+			{
+				continue;
+			}
 
 			if (Parameter.IsDataInterface())
 			{
-				int32 DataInterfaceOffset = ParamStore.IndexOf(Parameter);
-				UObject* DefaultValueObject = ParamStore.GetDataInterfaces()[DataInterfaceOffset];
+				ParameterValue = FNiagaraVariant(DuplicateObject(ParameterValue.GetDataInterface(), Component.Get()));
+			}
 
-				TArray<UObject*> Objects;
-				Objects.Add(DefaultValueObject);
+			TSharedPtr<FNiagaraParameterProxy> ParameterProxy = ParameterProxies.Add_GetRef(MakeShareable(new FNiagaraParameterProxy(Component, Parameter, ParameterValue, OnRebuildChildren, OverridePropertyHandles)));
+				
+			if (Parameter.IsDataInterface())
+			{
+				// duplicate the DI here so that if it is changed, the component's SetParameterOverride will override the value
+				// if no changes are made, then it'll just be the same as the asset
+				TArray<UObject*> Objects { ParameterProxy->Value().GetDataInterface() };
 
 				FAddPropertyParams Params = FAddPropertyParams()
 					.UniqueId(Parameter.GetName())
@@ -166,226 +352,107 @@ public:
 				CustomValueWidget =
 					SNew(STextBlock)
 					.TextStyle(FNiagaraEditorStyle::Get(), "NiagaraEditor.ParameterText")
-					.Text(FText::FromString(FName::NameToDisplayString(DefaultValueObject->GetClass()->GetName(), false)));
+					.Text(FText::FromString(FName::NameToDisplayString(Parameter.GetType().GetClass()->GetName(), false)));
 			}
 			else if (Parameter.IsUObject())
 			{
-				int32 ObjectOffst = ParamStore.IndexOf(Parameter);
-				UObject* DefaultValueObject = ParamStore.GetUObjects()[ObjectOffst];
-
-				TArray<UObject*> Objects;
-				Objects.Add(DefaultValueObject);
+				TArray<UObject*> Objects { ParameterProxy->Value().GetUObject() };
 
-				//How do I set this up so I can have this pick actors from the level?
+				// How do I set this up so I can have this pick actors from the level?
 
 				FAddPropertyParams Params = FAddPropertyParams()
 					.UniqueId(Parameter.GetName())
-					.AllowChildren(false) //Don't show the material's properties
+					.AllowChildren(false) // Don't show the material's properties
 					.CreateCategoryNodes(false);
 
 				Row = ChildrenBuilder.AddExternalObjectProperty(Objects, NAME_None, Params);
 
-				//How do I make this an object picker from the level editor?
-				//Neither of these seem to work well.
-// 				CustomValueWidget = PropertyCustomizationHelpers::MakeActorPickerWithMenu(nullptr,
-// 					true,
-// 					FOnShouldFilterActor::CreateRaw(this, &FNiagaraComponentNodeBuilder::IsFilteredActor),
-// 					FOnActorSelected::CreateRaw(this, &FNiagaraComponentNodeBuilder::OnActorSelected),
-// 					FSimpleDelegate::CreateRaw(this, &FNiagaraComponentNodeBuilder::CloseComboButton),
-// 					FSimpleDelegate::CreateRaw(this, &FNiagaraComponentNodeBuilder::OnUse));
-				
-// 				CustomValueWidget = PropertyCustomizationHelpers::MakeActorPickerAnchorButton(
-// 					FOnGetActorFilters::CreateRaw(this, &FNiagaraComponentNodeBuilder::OnGetActorFiltersForSceneOutliner),
-// 					FOnActorSelected::CreateRaw(this, &FNiagaraComponentNodeBuilder::OnActorSelected));
-
-				if (Parameter.GetType().GetClass() == UMaterialInterface::StaticClass())
-				{
-									
-					CustomValueWidget = SNew(SObjectPropertyEntryBox)
-						.ObjectPath_Raw(this, &FNiagaraComponentNodeBuilder::GetCurrentAssetPath, Parameter )
-						.AllowedClass(UMaterialInterface::StaticClass())
-						.OnObjectChanged_Raw(this, &FNiagaraComponentNodeBuilder::OnAssetSelectedFromPicker, Parameter)
-						.AllowClear(false)
-						.DisplayUseSelected(true)
-						.DisplayBrowse(true)
-						.DisplayThumbnail(true)
-						.NewAssetFactories(TArray<UFactory*>());
-				}
-				else
-				{
-					FString ValueName = DefaultValueObject ? FName::NameToDisplayString(DefaultValueObject->GetClass()->GetName(), false) : TEXT("null");
-					CustomValueWidget =
-						SNew(STextBlock)
-						.TextStyle(FNiagaraEditorStyle::Get(), "NiagaraEditor.ParameterText")
-						.Text(FText::FromString(ValueName));
-				}
+				CustomValueWidget = SNew(SObjectPropertyEntryBox)
+					.ObjectPath(ParameterProxy.ToSharedRef(), &FNiagaraParameterProxy::GetCurrentAssetPath)
+					.AllowedClass(Parameter.GetType().GetClass())
+					.OnObjectChanged(ParameterProxy.ToSharedRef(), &FNiagaraParameterProxy::OnAssetSelectedFromPicker)
+					.AllowClear(false)
+					.DisplayUseSelected(true)
+					.DisplayBrowse(true)
+					.DisplayThumbnail(true)
+					.NewAssetFactories(TArray<UFactory*>());
 			}
 			else
 			{
-				TSharedPtr<FStructOnScope> StructOnScope = MakeShareable(new FStructOnScope(Parameter.GetType().GetStruct(), (uint8*)ParamStore.GetParameterData(Parameter)));
+				TSharedPtr<FStructOnScope> StructOnScope = MakeShareable(new FStructOnScope(Parameter.GetType().GetStruct(), ParameterProxy->Value().GetBytes()));
 
-				Row = ChildrenBuilder.AddExternalStructureProperty(StructOnScope.ToSharedRef(), NAME_None, FAddPropertyParams().UniqueId(Parameter.GetName()));
+				FAddPropertyParams Params = FAddPropertyParams()
+					.UniqueId(Parameter.GetName());
 
+				Row = ChildrenBuilder.AddExternalStructureProperty(StructOnScope.ToSharedRef(), NAME_None, Params);
 			}
 
-			if (Row)
-			{
-				TSharedPtr<SWidget> DefaultNameWidget;
-				TSharedPtr<SWidget> DefaultValueWidget;
-
-				Row->DisplayName(FText::FromName(Parameter.GetName()));
+			check(Row && ParameterProxy.IsValid() && ParameterProxy->Value().IsValid());
 
-				TSharedPtr<IPropertyHandle> PropertyHandle = Row->GetPropertyHandle();
+			TSharedPtr<SWidget> DefaultNameWidget;
+			TSharedPtr<SWidget> DefaultValueWidget;
 
-				FDetailWidgetRow& CustomWidget = Row->CustomWidget(true);
+			Row->DisplayName(FText::FromName(Parameter.GetName()));
 
-				TArray<UObject*> Objects;
-				Row->GetPropertyHandle()->GetOuterObjects(Objects);
+			FDetailWidgetRow& CustomWidget = Row->CustomWidget(true);
 
-				Row->GetDefaultWidgets(DefaultNameWidget, DefaultValueWidget, CustomWidget);
+			Row->GetDefaultWidgets(DefaultNameWidget, DefaultValueWidget, CustomWidget);
 
-				if (Parameter.IsDataInterface())
-				{
-					PropertyHandle->SetOnPropertyValuePreChange(
-						FSimpleDelegate::CreateRaw(this, &FNiagaraComponentNodeBuilder::OnDataInterfacePreChange, Parameter));
-					PropertyHandle->SetOnChildPropertyValuePreChange(
-						FSimpleDelegate::CreateRaw(this, &FNiagaraComponentNodeBuilder::OnDataInterfacePreChange, Parameter));
-					PropertyHandle->SetOnPropertyValueChanged(
-						FSimpleDelegate::CreateRaw(this, &FNiagaraComponentNodeBuilder::OnDataInterfaceChanged, Parameter));
-					PropertyHandle->SetOnChildPropertyValueChanged(
-						FSimpleDelegate::CreateRaw(this, &FNiagaraComponentNodeBuilder::OnDataInterfaceChanged, Parameter));
-				}
-				else if (Parameter.IsUObject())
-				{
-					PropertyHandle->SetOnPropertyValuePreChange(
-						FSimpleDelegate::CreateRaw(this, &FNiagaraComponentNodeBuilder::OnUObjectPreChange, Parameter));
-					PropertyHandle->SetOnChildPropertyValuePreChange(
-						FSimpleDelegate::CreateRaw(this, &FNiagaraComponentNodeBuilder::OnUObjectPreChange, Parameter));
-					PropertyHandle->SetOnPropertyValueChanged(
-						FSimpleDelegate::CreateRaw(this, &FNiagaraComponentNodeBuilder::OnUObjectChanged, Parameter));
-					PropertyHandle->SetOnChildPropertyValueChanged(
-						FSimpleDelegate::CreateRaw(this, &FNiagaraComponentNodeBuilder::OnUObjectChanged, Parameter));
-				}
-				else
-				{
-					PropertyHandle->SetOnPropertyValuePreChange(
-						FSimpleDelegate::CreateRaw(this, &FNiagaraComponentNodeBuilder::OnParameterPreChange, Parameter));
-					PropertyHandle->SetOnChildPropertyValuePreChange(
-						FSimpleDelegate::CreateRaw(this, &FNiagaraComponentNodeBuilder::OnParameterPreChange, Parameter));
-					PropertyHandle->SetOnPropertyValueChanged(
-						FSimpleDelegate::CreateRaw(this, &FNiagaraComponentNodeBuilder::OnParameterChanged, Parameter));
-					PropertyHandle->SetOnChildPropertyValueChanged(
-						FSimpleDelegate::CreateRaw(this, &FNiagaraComponentNodeBuilder::OnParameterChanged, Parameter));
-				}
+			Row->GetPropertyHandle()->SetOnPropertyValuePreChange(FSimpleDelegate::CreateSP(ParameterProxy.ToSharedRef(), &FNiagaraParameterProxy::OnParameterPreChange));
+			Row->GetPropertyHandle()->SetOnChildPropertyValuePreChange(FSimpleDelegate::CreateSP(ParameterProxy.ToSharedRef(), &FNiagaraParameterProxy::OnParameterPreChange));
+			Row->GetPropertyHandle()->SetOnPropertyValueChanged(FSimpleDelegate::CreateSP(ParameterProxy.ToSharedRef(), &FNiagaraParameterProxy::OnParameterChanged));
+			Row->GetPropertyHandle()->SetOnChildPropertyValueChanged(FSimpleDelegate::CreateSP(ParameterProxy.ToSharedRef(), &FNiagaraParameterProxy::OnParameterChanged));
+			Row->GetPropertyHandle()->SetOnPropertyResetToDefault(FSimpleDelegate::CreateSP(ParameterProxy.ToSharedRef(), &FNiagaraParameterProxy::OnResetToDefault));
 
-				CustomWidget
-					.NameContent()
+			CustomWidget
+				.NameContent()
+				[
+					SNew(SBox)
+					.Padding(FMargin(0.0f, 2.0f))
 					[
-						SNew(SBox)
-						.Padding(FMargin(0.0f, 2.0f))
-						[
-							NameWidget.ToSharedRef()
-						]
-					];
+						NameWidget.ToSharedRef()
+					]
+				];
 
-				TSharedPtr<SWidget> ValueWidget = DefaultValueWidget;
-				if (CustomValueWidget.IsValid())
-				{
-					ValueWidget = CustomValueWidget;
-				}
+			TSharedPtr<SWidget> ValueWidget = DefaultValueWidget;
+			if (CustomValueWidget.IsValid())
+			{
+				ValueWidget = CustomValueWidget;
+			}
 
-				CustomWidget
-					.ValueContent()
+			CustomWidget
+				.ValueContent()
+				[
+					SNew(SHorizontalBox)
+					+ SHorizontalBox::Slot()
+					.HAlign(HAlign_Fill)
+					.Padding(4.0f)
 					[
-						SNew(SHorizontalBox)
-						+ SHorizontalBox::Slot()
-						.HAlign(HAlign_Fill)
-						.Padding(4.0f)
-						[
-							// Add in the parameter editor factoried above.
-							ValueWidget.ToSharedRef()
-						]
-						+ SHorizontalBox::Slot()
-						.VAlign(VAlign_Center)
-						.AutoWidth()
+						// Add in the parameter editor factoried above.
+						ValueWidget.ToSharedRef()
+					]
+					+ SHorizontalBox::Slot()
+					.VAlign(VAlign_Center)
+					.AutoWidth()
+					[
+						// Add in the "reset to default" buttons
+						SNew(SButton)
+						.OnClicked(ParameterProxy.ToSharedRef(), &FNiagaraParameterProxy::OnResetToDefaultClicked)
+						.Visibility(ParameterProxy.ToSharedRef(), &FNiagaraParameterProxy::GetResetToDefaultVisibility)
+						.ContentPadding(FMargin(5.f, 0.f))
+						.ToolTipText(LOCTEXT("ResetToDefaultToolTip", "Reset to Default"))
+						.ButtonStyle(FEditorStyle::Get(), "NoBorder")
+						.Content()
 						[
-							// Add in the "reset to default" buttons
-							SNew(SButton)
-							.OnClicked_Raw(this, &FNiagaraComponentNodeBuilder::OnLocationResetClicked, Parameter)
-							.Visibility_Raw(this, &FNiagaraComponentNodeBuilder::GetLocationResetVisibility, Parameter)
-							.ContentPadding(FMargin(5.f, 0.f))
-							.ToolTipText(LOCTEXT("ResetToDefaultToolTip", "Reset to Default"))
-							.ButtonStyle(FEditorStyle::Get(), "NoBorder")
-							.Content()
-							[
-								SNew(SImage)
-								.Image(FEditorStyle::GetBrush("PropertyWindow.DiffersFromDefault"))
-							]
+							SNew(SImage)
+							.Image(FEditorStyle::GetBrush("PropertyWindow.DiffersFromDefault"))
 						]
-					];
-			}
+					]
+				];
 		}
 	}
 
 private:
-	void OnParameterPreChange(FNiagaraVariable Var)
-	{
-		check(Component.IsValid());
-		Component->Modify();
-	}
-
-	void OnDataInterfacePreChange(FNiagaraVariable Var)
-	{
-		check(Component.IsValid());
-		Component->Modify();
-	}
-	
-	void OnUObjectPreChange(FNiagaraVariable Var)
-	{
-		check(Component.IsValid());
-		Component->Modify();
-	}
-
-	void OnParameterChanged(FNiagaraVariable Var)
-	{
-		check(Component.IsValid());
-		Component->GetOverrideParameters().OnParameterChange();
-		Component->SetParameterValueOverriddenLocally(Var, true, true);
-	}
-
-	void OnDataInterfaceChanged(FNiagaraVariable Var)
-	{
-		Component->GetOverrideParameters().OnInterfaceChange();
-		check(Component.IsValid());
-		Component->SetParameterValueOverriddenLocally(Var, true, true);
-	}
-
-	void OnUObjectChanged(FNiagaraVariable Var)
-	{
-		Component->GetOverrideParameters().OnUObjectChange();
-		check(Component.IsValid());
-		Component->SetParameterValueOverriddenLocally(Var, true, true);
-	}
-
-	bool DoesParameterDifferFromDefault(const FNiagaraVariable& Var)
-	{
-		check(Component.IsValid());
-		return Component->IsParameterValueOverriddenLocally(Var.GetName());
-	}
-
-	FReply OnLocationResetClicked(FNiagaraVariable Parameter)
-	{
-		check(Component.IsValid());
-		FScopedTransaction ScopedTransaction(LOCTEXT("ResetParameterValue", "Reset parameter value to system defaults."));
-		Component->Modify();
-		Component->SetParameterValueOverriddenLocally(Parameter, false, true);
-		return FReply::Handled();
-	}
-
-	EVisibility GetLocationResetVisibility(FNiagaraVariable Parameter) const
-	{
-		return Component.IsValid() && Component->IsParameterValueOverriddenLocally(Parameter.GetName()) ? EVisibility::Visible : EVisibility::Collapsed;
-	}
 
 	void ComponentSynchronizedWithAssetParameters()
 	{
@@ -394,9 +461,9 @@ private:
 
 private:
 	TWeakObjectPtr<UNiagaraComponent> Component;
+	TArray<TSharedPtr<IPropertyHandle>> OverridePropertyHandles;
 	FSimpleDelegate OnRebuildChildren;
-	TArray<TSharedRef<FStructOnScope>> CreatedStructOnScopes;
-	TArray<UNiagaraScript*> OriginalScripts;
+	TArray<TSharedPtr<FNiagaraParameterProxy>> ParameterProxies;
 };
 
 TSharedRef<IDetailCustomization> FNiagaraComponentDetails::MakeInstance()
@@ -456,12 +523,20 @@ void FNiagaraComponentDetails::CustomizeDetails(IDetailLayoutBuilder& DetailBuil
 	static const FName ParamCategoryName = TEXT("NiagaraComponent_Parameters");
 	static const FName ScriptCategoryName = TEXT("Parameters");
 
-	TSharedPtr<IPropertyHandle> LocalOverridesPropertyHandle = DetailBuilder.GetProperty(TEXT("OverrideParameters"));
+	TSharedPtr<IPropertyHandle> LocalOverridesPropertyHandle = DetailBuilder.GetProperty("OverrideParameters");
 	if (LocalOverridesPropertyHandle.IsValid())
 	{
 		LocalOverridesPropertyHandle->MarkHiddenByCustomization();
 	}
 
+	TSharedPtr<IPropertyHandle> TemplateParameterOverridesPropertyHandle = DetailBuilder.GetProperty("TemplateParameterOverrides");
+	TemplateParameterOverridesPropertyHandle->MarkHiddenByCustomization();
+
+	TSharedPtr<IPropertyHandle> InstanceParameterOverridesPropertyHandle = DetailBuilder.GetProperty("InstanceParameterOverrides");
+	InstanceParameterOverridesPropertyHandle->MarkHiddenByCustomization();
+
+	TArray<TSharedPtr<IPropertyHandle>> PropertyHandles { TemplateParameterOverridesPropertyHandle, InstanceParameterOverridesPropertyHandle };
+
 	TArray<TWeakObjectPtr<UObject>> ObjectsCustomized;
 	DetailBuilder.GetObjectsBeingCustomized(ObjectsCustomized);
 
@@ -477,7 +552,7 @@ void FNiagaraComponentDetails::CustomizeDetails(IDetailLayoutBuilder& DetailBuil
 		FGameDelegates::Get().GetEndPlayMapDelegate().AddRaw(this, &FNiagaraComponentDetails::OnPiEEnd);
 			
 		IDetailCategoryBuilder& InputParamCategory = DetailBuilder.EditCategory(ParamCategoryName, LOCTEXT("ParamCategoryName", "Override Parameters"));
-		InputParamCategory.AddCustomBuilder(MakeShared<FNiagaraComponentNodeBuilder>(Component.Get()));
+		InputParamCategory.AddCustomBuilder(MakeShared<FNiagaraComponentNodeBuilder>(Component.Get(), PropertyHandles));
 	}
 	else if (ObjectsCustomized.Num() > 1)
 	{
