Fix for Niagara particle command buffer hang on iOS
- Added flag to test for RW Texture Buffer support (Enabled for all platforms by default and disabled on iOS < A10)
- Split DrawIndirect particle update into two passes for devices that dont support RW texture buffers

diff --git a/Engine/Config/BaseDeviceProfiles.ini b/Engine/Config/BaseDeviceProfiles.ini
--- a/Engine/Config/BaseDeviceProfiles.ini
+++ b/Engine/Config/BaseDeviceProfiles.ini
@@ -170,28 +170,24 @@ BaseProfileName=
 DeviceType=IOS
 BaseProfileName=IOS
 +CVars=r.RenderTargetSwitchWorkaround=1
-+CVars=fx.NiagaraAllowComputeShaders=0
 +CVars=r.EyeAdaptationQuality=0
 
 [iPad3 DeviceProfile]
 DeviceType=IOS
 BaseProfileName=IOS
 +CVars=r.RenderTargetSwitchWorkaround=1
-+CVars=fx.NiagaraAllowComputeShaders=0
 +CVars=r.EyeAdaptationQuality=0
 
 [iPad4 DeviceProfile]
 DeviceType=IOS
 BaseProfileName=IOS
 +CVars=r.RenderTargetSwitchWorkaround=1
-+CVars=fx.NiagaraAllowComputeShaders=0
 +CVars=r.EyeAdaptationQuality=0
 
 [iPadAir DeviceProfile]
 DeviceType=IOS
 BaseProfileName=IOS
 +CVars=sg.PostProcessQuality=1
-+CVars=fx.NiagaraAllowComputeShaders=0
 +CVars=r.EyeAdaptationQuality=0
 
 [iPadAir2 DeviceProfile]
@@ -199,27 +195,23 @@ DeviceType=IOS
 BaseProfileName=IOS
 +CVars=sg.PostProcessQuality=3
 +CVars=sg.ShadowQuality=2
-+CVars=fx.NiagaraAllowComputeShaders=0
 +CVars=r.EyeAdaptationQuality=0
 
 [iPadMini DeviceProfile]
 DeviceType=IOS
 BaseProfileName=IOS
 +CVars=r.RenderTargetSwitchWorkaround=1
-+CVars=fx.NiagaraAllowComputeShaders=0
 +CVars=r.EyeAdaptationQuality=0
 
 [iPadMini2 DeviceProfile]
 DeviceType=IOS
 BaseProfileName=iPadAir
-+CVars=fx.NiagaraAllowComputeShaders=0
 +CVars=r.EyeAdaptationQuality=0
 
 [iPadMini4 DeviceProfile]
 DeviceType=IOS
 BaseProfileName=iPadMini2
 +CVars=sg.PostProcessQuality=2
-+CVars=fx.NiagaraAllowComputeShaders=0
 +CVars=r.EyeAdaptationQuality=0
 
 [iPhone5S DeviceProfile]
@@ -227,21 +219,18 @@ DeviceType=IOS
 BaseProfileName=IOS
 +CVars=r.MobileContentScaleFactor=2
 +CVars=sg.PostProcessQuality=3
-+CVars=fx.NiagaraAllowComputeShaders=0
 +CVars=r.EyeAdaptationQuality=0
 
 [iPodTouch6 DeviceProfile]
 DeviceType=IOS
 BaseProfileName=iPhone6
 +CVars=r.MobileContentScaleFactor=2
-+CVars=fx.NiagaraAllowComputeShaders=0
 +CVars=r.EyeAdaptationQuality=0
 
 [iPodTouch7 DeviceProfile]
 DeviceType=IOS
 BaseProfileName=iPhone7
 +CVars=r.MobileContentScaleFactor=2
-+CVars=fx.NiagaraAllowComputeShaders=0
 +CVars=r.EyeAdaptationQuality=0
 
 [iPhone6 DeviceProfile]
@@ -249,7 +238,6 @@ DeviceType=IOS
 BaseProfileName=IOS
 +CVars=r.MobileContentScaleFactor=2
 +CVars=sg.PostProcessQuality=3
-+CVars=fx.NiagaraAllowComputeShaders=0
 +CVars=r.EyeAdaptationQuality=0
 
 [iPhone6Plus DeviceProfile]
@@ -257,52 +245,44 @@ DeviceType=IOS
 BaseProfileName=IOS
 +CVars=r.MobileContentScaleFactor=2
 +CVars=sg.PostProcessQuality=3
-+CVars=fx.NiagaraAllowComputeShaders=0
 
 [iPhone6S DeviceProfile]
 DeviceType=IOS
 BaseProfileName=iPhone6
 +CVars=r.MobileContentScaleFactor=2
 +CVars=sg.ShadowQuality=2
-+CVars=fx.NiagaraAllowComputeShaders=0
 
 [iPhone6SPlus DeviceProfile]
 DeviceType=IOS
 BaseProfileName=iPhone6Plus
 +CVars=r.MobileContentScaleFactor=2
 +CVars=sg.ShadowQuality=2
-+CVars=fx.NiagaraAllowComputeShaders=0
 
 [iPhoneSE DeviceProfile]
 DeviceType=IOS
 BaseProfileName=IOS
 +CVars=r.MobileContentScaleFactor=2
 +CVars=sg.PostProcessQuality=3
-+CVars=fx.NiagaraAllowComputeShaders=0
 
 [iPhone7 DeviceProfile]
 DeviceType=IOS
 BaseProfileName=iPhone6
 +CVars=r.MobileContentScaleFactor=2
-+CVars=fx.NiagaraAllowComputeShaders=0
 
 [iPhone7Plus DeviceProfile]
 DeviceType=IOS
 BaseProfileName=iPhone6Plus
 +CVars=r.MobileContentScaleFactor=2
-+CVars=fx.NiagaraAllowComputeShaders=0
 
 [iPhone8 DeviceProfile]
 DeviceType=IOS
 BaseProfileName=iPhone7
 +CVars=sg.ShadowQuality=3
-+CVars=fx.NiagaraAllowComputeShaders=1
 
 [iPhone8Plus DeviceProfile]
 DeviceType=IOS
 BaseProfileName=iPhone7Plus
 +CVars=sg.ShadowQuality=3
-+CVars=fx.NiagaraAllowComputeShaders=1
 
 [iPhoneX DeviceProfile]
 DeviceType=IOS
@@ -365,7 +345,6 @@ BaseProfileName=AppleTV
 [iPadPro97 DeviceProfile]
 DeviceType=IOS
 BaseProfileName=IPadPro
-+CVars=fx.NiagaraAllowComputeShaders=0
 
 [iPadPro129 DeviceProfile]
 DeviceType=IOS
diff --git a/Engine/Plugins/FX/Niagara/Shaders/Private/NiagaraDrawIndirectArgsGen.usf b/Engine/Plugins/FX/Niagara/Shaders/Private/NiagaraDrawIndirectArgsGen.usf
--- a/Engine/Plugins/FX/Niagara/Shaders/Private/NiagaraDrawIndirectArgsGen.usf
+++ b/Engine/Plugins/FX/Niagara/Shaders/Private/NiagaraDrawIndirectArgsGen.usf
@@ -29,32 +29,45 @@ void MainCS(uint TaskIndex : SV_DispatchThreadID)
 		RWDrawIndirectArgs[ArgOffset + 3] = 0; // BaseVertexLocation
 		RWDrawIndirectArgs[ArgOffset + 4] = 0; // StartInstanceLocation
 	}
-/*
-Explanation of why we are using two ifs here instead of if/else: 
-Mali requires texture indexes to be warp invariant, which means that the if / if code translates to
-	if id.x < uniform :
-	[to 5 stores]
-	[branch to end]
-	if <the other condition>
-	do 1 image store
-	<end>
-and the if / else translates to
-	if uniform > ID.x: goto 4 image stores
-	if <other condition> :
-	setup registers
-	goto <last store>
-	<4 images stores branch target>
-	do 4 stores
-	setup registers
-	<last store target>
-	do 1 store
-Notice how it completely ignores the fact that there are two different texture buffers, which leads the compiler to do an "optimization" and separate the last store.This means that the else condition does indeed evaluate correctly, but we are writing in the wrong texture!
-Note: This is imagestore specific and not compute specific and is a known issue for Arm. This is fixed in driver version R21
-*/
+	
+	#if SUPPORTS_TEXTURE_RW == 1
+		/*
+		Explanation of why we are using two ifs here instead of if/else: 
+		Mali requires texture indexes to be warp invariant, which means that the if / if code translates to
+			if id.x < uniform :
+			[to 5 stores]
+			[branch to end]
+			if <the other condition>
+			do 1 image store
+			<end>
+		and the if / else translates to
+			if uniform > ID.x: goto 4 image stores
+			if <other condition> :
+			setup registers
+			goto <last store>
+			<4 images stores branch target>
+			do 4 stores
+			setup registers
+			<last store target>
+			do 1 store
+		Notice how it completely ignores the fact that there are two different texture buffers, which leads the compiler to do an "optimization" and separate the last store.This means that the else condition does indeed evaluate correctly, but we are writing in the wrong texture!
+		Note: This is imagestore specific and not compute specific and is a known issue for Arm. This is fixed in driver version R21
+		*/
+	
+		if (TaskIndex < TaskCount.z && TaskIndex >= TaskCount.x)
+		{
+			const uint InfoOffset = TaskCount.x * NIAGARA_DRAW_INDIRECT_TASK_INFO_SIZE + (TaskIndex - TaskCount.x);
+			RWInstanceCounts[TaskInfos[InfoOffset + 0]] = 0;
+		}
+	#endif
+}
 
-	if (TaskIndex < TaskCount.z && TaskIndex >= TaskCount.x)
+[numthreads(THREAD_COUNT,1,1)]
+void ResetCountsCS(uint TaskIndex : SV_DispatchThreadID)
+{
+	if (TaskIndex < TaskCount.y)
 	{
-		const uint InfoOffset = TaskCount.x * NIAGARA_DRAW_INDIRECT_TASK_INFO_SIZE + (TaskIndex - TaskCount.x);
+		const uint InfoOffset = TaskCount.x * NIAGARA_DRAW_INDIRECT_TASK_INFO_SIZE + (TaskIndex);
 		RWInstanceCounts[TaskInfos[InfoOffset + 0]] = 0;
 	}
-}
\ No newline at end of file
+}
diff --git a/Engine/Plugins/FX/Niagara/Source/Niagara/Classes/NiagaraGPUInstanceCountManager.h b/Engine/Plugins/FX/Niagara/Source/Niagara/Classes/NiagaraGPUInstanceCountManager.h
--- a/Engine/Plugins/FX/Niagara/Source/Niagara/Classes/NiagaraGPUInstanceCountManager.h
+++ b/Engine/Plugins/FX/Niagara/Source/Niagara/Classes/NiagaraGPUInstanceCountManager.h
@@ -21,7 +21,7 @@ public:
 	int32 Value = 0;
 };
 
-FORCEINLINE uint32 GetTypeHash(const FNiagaraDrawIndirectArgsGenCS::FArgGenTaskInfo& Info)
+FORCEINLINE uint32 GetTypeHash(const FNiagaraDrawIndirectArgGenTaskInfo& Info)
 {
 	return HashCombine(Info.InstanceCountBufferOffset, HashCombine(Info.NumIndicesPerInstance, Info.StartIndexLocation));
 }
@@ -96,7 +96,7 @@ protected:
 	int32 AllocatedDrawIndirectArgs = 0;
 
 	/** The list of all draw indirected tasks that are to be run in UpdateDrawIndirectBuffer() */
-	typedef FNiagaraDrawIndirectArgsGenCS::FArgGenTaskInfo FArgGenTaskInfo;
+	typedef FNiagaraDrawIndirectArgGenTaskInfo FArgGenTaskInfo;
 
 	TArray<FArgGenTaskInfo> DrawIndirectArgGenTasks;
 	/** The map between each task FArgGenTaskInfo and entry offset from DrawIndirectArgGenTasks. Used to reuse entries. */
diff --git a/Engine/Plugins/FX/Niagara/Source/Niagara/Private/NiagaraGPUInstanceCountManager.cpp b/Engine/Plugins/FX/Niagara/Source/Niagara/Private/NiagaraGPUInstanceCountManager.cpp
--- a/Engine/Plugins/FX/Niagara/Source/Niagara/Private/NiagaraGPUInstanceCountManager.cpp
+++ b/Engine/Plugins/FX/Niagara/Source/Niagara/Private/NiagaraGPUInstanceCountManager.cpp
@@ -204,8 +204,6 @@ void FNiagaraGPUInstanceCountManager::UpdateDrawIndirectBuffer(FRHICommandList&
 	{
 		if (FNiagaraUtilities::AllowGPUParticles(GShaderPlatformForFeatureLevel[FeatureLevel]))
 		{
-			RHICmdList.TransitionResource(EResourceTransitionAccess::ERWBarrier, EResourceTransitionPipeline::EComputeToCompute, CountBuffer.UAV);
-
 			FReadBuffer TaskInfosBuffer;
 
 			//for (int32 i = 0; i < InstanceCountClearTasks.Num(); i++)
@@ -224,18 +222,46 @@ void FNiagaraGPUInstanceCountManager::UpdateDrawIndirectBuffer(FRHICommandList&
 				RHIUnlockVertexBuffer(TaskInfosBuffer.Buffer);
 			}
 
-			RHICmdList.TransitionResource(EResourceTransitionAccess::EWritable, EResourceTransitionPipeline::EGfxToCompute, DrawIndirectBuffer.UAV);
-
 			FNiagaraDrawIndirectArgsGenCS::FPermutationDomain PermutationVector;
+			PermutationVector.Set<FNiagaraDrawIndirectArgsGenCS::FSupportsTextureRW>(GRHISupportsRWTextureBuffers ? 1 : 0);
 			TShaderMapRef<FNiagaraDrawIndirectArgsGenCS> DrawIndirectArgsGenCS(GetGlobalShaderMap(FeatureLevel), PermutationVector);
+
+			RHICmdList.TransitionResource(GRHISupportsRWTextureBuffers ? EResourceTransitionAccess::ERWBarrier : EResourceTransitionAccess::EReadable,
+											EResourceTransitionPipeline::EComputeToCompute, CountBuffer.UAV);
+			RHICmdList.TransitionResource(EResourceTransitionAccess::EWritable, EResourceTransitionPipeline::EGfxToCompute, DrawIndirectBuffer.UAV);
+
 			RHICmdList.SetComputeShader(DrawIndirectArgsGenCS.GetComputeShader());
 			DrawIndirectArgsGenCS->SetOutput(RHICmdList, DrawIndirectBuffer.UAV, CountBuffer.UAV);
 			DrawIndirectArgsGenCS->SetParameters(RHICmdList, TaskInfosBuffer.SRV, DrawIndirectArgGenTasks.Num(), InstanceCountClearTasks.Num());
-			DispatchComputeShader(RHICmdList, DrawIndirectArgsGenCS.GetShader(), FMath::DivideAndRoundUp(DrawIndirectArgGenTasks.Num() + InstanceCountClearTasks.Num(), NIAGARA_DRAW_INDIRECT_ARGS_GEN_THREAD_COUNT), 1, 1);
-			DrawIndirectArgsGenCS->UnbindBuffers(RHICmdList);
 
-			// Sync after clear.
-			RHICmdList.TransitionResource(EResourceTransitionAccess::ERWBarrier, EResourceTransitionPipeline::EComputeToCompute, CountBuffer.UAV);
+			// If the device supports RW Texture buffers then we can use a single compute pass, otherwise we need to split into two passes
+			if (GRHISupportsRWTextureBuffers)
+			{
+				DispatchComputeShader(RHICmdList, DrawIndirectArgsGenCS.GetShader(), FMath::DivideAndRoundUp(DrawIndirectArgGenTasks.Num() + InstanceCountClearTasks.Num(), NIAGARA_DRAW_INDIRECT_ARGS_GEN_THREAD_COUNT), 1, 1);
+				DrawIndirectArgsGenCS->UnbindBuffers(RHICmdList);
+
+				// Sync after clear.
+				RHICmdList.TransitionResource(EResourceTransitionAccess::ERWBarrier, EResourceTransitionPipeline::EComputeToCompute, CountBuffer.UAV);
+			}
+			else
+			{
+				DispatchComputeShader(RHICmdList, DrawIndirectArgsGenCS.GetShader(), FMath::DivideAndRoundUp(DrawIndirectArgGenTasks.Num(), NIAGARA_DRAW_INDIRECT_ARGS_GEN_THREAD_COUNT), 1, 1);
+				DrawIndirectArgsGenCS->UnbindBuffers(RHICmdList);
+
+				RHICmdList.TransitionResource(EResourceTransitionAccess::EWritable, EResourceTransitionPipeline::EComputeToCompute, CountBuffer.UAV);
+
+				FNiagaraDrawIndirectResetCountsCS::FPermutationDomain PermutationVectorResetCounts;
+				TShaderMapRef<FNiagaraDrawIndirectResetCountsCS> DrawIndirectResetCountsArgsGenCS(GetGlobalShaderMap(FeatureLevel), PermutationVectorResetCounts);
+				RHICmdList.SetComputeShader(DrawIndirectResetCountsArgsGenCS.GetComputeShader());
+				DrawIndirectResetCountsArgsGenCS->SetOutput(RHICmdList, CountBuffer.UAV);
+				DrawIndirectResetCountsArgsGenCS->SetParameters(RHICmdList, TaskInfosBuffer.SRV, DrawIndirectArgGenTasks.Num(), InstanceCountClearTasks.Num());
+				DispatchComputeShader(RHICmdList, DrawIndirectResetCountsArgsGenCS.GetShader(), FMath::DivideAndRoundUp(InstanceCountClearTasks.Num(), NIAGARA_DRAW_INDIRECT_ARGS_GEN_THREAD_COUNT), 1, 1);
+				DrawIndirectResetCountsArgsGenCS->UnbindBuffers(RHICmdList);
+
+				// Sync after clear.
+				RHICmdList.TransitionResource(EResourceTransitionAccess::EReadable, EResourceTransitionPipeline::EComputeToCompute, CountBuffer.UAV);
+			}
+
 			// Transition draw indirect to readable for gfx draw indirect.
 			RHICmdList.TransitionResource(EResourceTransitionAccess::EReadable, EResourceTransitionPipeline::EComputeToGfx, DrawIndirectBuffer.UAV);
 		}
diff --git a/Engine/Plugins/FX/Niagara/Source/Niagara/Private/NiagaraRendererMeshes.cpp b/Engine/Plugins/FX/Niagara/Source/Niagara/Private/NiagaraRendererMeshes.cpp
--- a/Engine/Plugins/FX/Niagara/Source/Niagara/Private/NiagaraRendererMeshes.cpp
+++ b/Engine/Plugins/FX/Niagara/Source/Niagara/Private/NiagaraRendererMeshes.cpp
@@ -752,7 +752,8 @@ void FNiagaraRendererMeshes::GetDynamicRayTracingInstances(FRayTracingMaterialGa
 		InstanceGPUTransformsBuffer.Initialize(3*4*sizeof(float), CPUInstancesCount, BUF_Static);
 		RayTracingInstance.InstanceGPUTransformsSRV = InstanceGPUTransformsBuffer.SRV;
 
-		FNiagaraDrawIndirectArgsGenCS::FPermutationDomain PermutationVector;
+		FNiagaraGPURayTracingTransformsCS::FPermutationDomain PermutationVector;
+
 		TShaderMapRef<FNiagaraGPURayTracingTransformsCS> GPURayTracingTransformsCS(GetGlobalShaderMap(FeatureLevel), PermutationVector);
 		RHICmdList.SetComputeShader(GPURayTracingTransformsCS.GetComputeShader());
 
diff --git a/Engine/Plugins/FX/Niagara/Source/NiagaraVertexFactories/Private/NiagaraDrawIndirect.cpp b/Engine/Plugins/FX/Niagara/Source/NiagaraVertexFactories/Private/NiagaraDrawIndirect.cpp
--- a/Engine/Plugins/FX/Niagara/Source/NiagaraVertexFactories/Private/NiagaraDrawIndirect.cpp
+++ b/Engine/Plugins/FX/Niagara/Source/NiagaraVertexFactories/Private/NiagaraDrawIndirect.cpp
@@ -9,6 +9,7 @@ NiagaraDrawIndirect.cpp : Niagara shader to generate the draw indirect args for
 #include "ShaderParameterUtils.h"
 
 IMPLEMENT_GLOBAL_SHADER(FNiagaraDrawIndirectArgsGenCS, "/Plugin/FX/Niagara/Private/NiagaraDrawIndirectArgsGen.usf", "MainCS", SF_Compute);
+IMPLEMENT_GLOBAL_SHADER(FNiagaraDrawIndirectResetCountsCS, "/Plugin/FX/Niagara/Private/NiagaraDrawIndirectArgsGen.usf", "ResetCountsCS", SF_Compute);
 
 void FNiagaraDrawIndirectArgsGenCS::ModifyCompilationEnvironment(const FGlobalShaderPermutationParameters& Parameters, FShaderCompilerEnvironment& OutEnvironment)
 {
@@ -69,6 +70,54 @@ void FNiagaraDrawIndirectArgsGenCS::UnbindBuffers(FRHICommandList& RHICmdList)
 	{
 		RHICmdList.SetUAVParameter(ComputeShaderRHI, DrawIndirectArgsParam.GetUAVIndex(), nullptr);
 	}
+	if (InstanceCountsParam.IsBound())
+	{
+		RHICmdList.SetUAVParameter(ComputeShaderRHI, InstanceCountsParam.GetUAVIndex(), nullptr);
+	}
+}
+
+void FNiagaraDrawIndirectResetCountsCS::ModifyCompilationEnvironment(const FGlobalShaderPermutationParameters& Parameters, FShaderCompilerEnvironment& OutEnvironment)
+{
+	FGlobalShader::ModifyCompilationEnvironment(Parameters, OutEnvironment);
+	OutEnvironment.SetDefine(TEXT("THREAD_COUNT"), NIAGARA_DRAW_INDIRECT_ARGS_GEN_THREAD_COUNT);
+	OutEnvironment.SetDefine(TEXT("NIAGARA_DRAW_INDIRECT_ARGS_SIZE"), NIAGARA_DRAW_INDIRECT_ARGS_SIZE);
+	OutEnvironment.SetDefine(TEXT("NIAGARA_DRAW_INDIRECT_TASK_INFO_SIZE"), NIAGARA_DRAW_INDIRECT_TASK_INFO_SIZE);
+}
+
+FNiagaraDrawIndirectResetCountsCS::FNiagaraDrawIndirectResetCountsCS(const ShaderMetaType::CompiledShaderInitializerType& Initializer)
+	: FGlobalShader(Initializer)
+{
+	TaskInfosParam.Bind(Initializer.ParameterMap, TEXT("TaskInfos"));
+	InstanceCountsParam.Bind(Initializer.ParameterMap, TEXT("InstanceCounts"));
+	TaskCountParam.Bind(Initializer.ParameterMap, TEXT("TaskCount"));
+}
+
+void FNiagaraDrawIndirectResetCountsCS::SetOutput(FRHICommandList& RHICmdList, FRHIUnorderedAccessView* InstanceCountsUAV)
+{
+	FRHIComputeShader* ComputeShaderRHI = RHICmdList.GetBoundComputeShader();
+
+	if (InstanceCountsParam.IsBound())
+	{
+		RHICmdList.SetUAVParameter(ComputeShaderRHI, InstanceCountsParam.GetUAVIndex(), InstanceCountsUAV);
+	}
+}
+
+void FNiagaraDrawIndirectResetCountsCS::SetParameters(FRHICommandList& RHICmdList, FRHIShaderResourceView* TaskInfosBuffer, int32 NumArgGenTasks, int32 NumInstanceCountClearTasks)
+{
+	FRHIComputeShader* ComputeShaderRHI = RHICmdList.GetBoundComputeShader();
+
+	SetSRVParameter(RHICmdList, ComputeShaderRHI, TaskInfosParam, TaskInfosBuffer);
+
+	const FUintVector4 TaskCountValue((int32)NumArgGenTasks, (int32)NumInstanceCountClearTasks, (int32)(NumArgGenTasks + NumInstanceCountClearTasks), 0);
+	SetShaderValue(RHICmdList, ComputeShaderRHI, TaskCountParam, TaskCountValue);
+}
+
+void FNiagaraDrawIndirectResetCountsCS::UnbindBuffers(FRHICommandList& RHICmdList)
+{
+	FRHIComputeShader* ComputeShaderRHI = RHICmdList.GetBoundComputeShader();
+
+	SetSRVParameter(RHICmdList, ComputeShaderRHI, TaskInfosParam, nullptr);
+
 	if (InstanceCountsParam.IsBound())
 	{
 		RHICmdList.SetUAVParameter(ComputeShaderRHI, InstanceCountsParam.GetUAVIndex(), nullptr);
diff --git a/Engine/Plugins/FX/Niagara/Source/NiagaraVertexFactories/Public/NiagaraDrawIndirect.h b/Engine/Plugins/FX/Niagara/Source/NiagaraVertexFactories/Public/NiagaraDrawIndirect.h
--- a/Engine/Plugins/FX/Niagara/Source/NiagaraVertexFactories/Public/NiagaraDrawIndirect.h
+++ b/Engine/Plugins/FX/Niagara/Source/NiagaraVertexFactories/Public/NiagaraDrawIndirect.h
@@ -9,6 +9,7 @@ NiagaraDrawIndirect.h: Niagara shader to generate the draw indirect args for Nia
 #include "CoreMinimal.h"
 #include "GlobalShader.h"
 #include "ShaderPermutation.h"
+#include "ShaderParameterUtils.h"
 
 #define NIAGARA_DRAW_INDIRECT_ARGS_GEN_THREAD_COUNT 64
 #define NIAGARA_DRAW_INDIRECT_ARGS_SIZE 5
@@ -17,6 +18,31 @@ NiagaraDrawIndirect.h: Niagara shader to generate the draw indirect args for Nia
 // #define NIAGARA_COPY_BUFFER_THREAD_COUNT 64
 // #define NIAGARA_COPY_BUFFER_BUFFER_COUNT 3
 
+/**
+* Task info when generating draw indirect frame buffer.
+* Task is either about generate Niagara renderers drawindirect buffer,
+* or about resetting released instance counters.
+*/
+struct FNiagaraDrawIndirectArgGenTaskInfo
+{
+	FNiagaraDrawIndirectArgGenTaskInfo(uint32 InInstanceCountBufferOffset, uint32 InNumIndicesPerInstance, uint32 InStartIndexLocation)
+		: InstanceCountBufferOffset(InInstanceCountBufferOffset)
+		, NumIndicesPerInstance(InNumIndicesPerInstance)
+		, StartIndexLocation(InStartIndexLocation)
+	{}
+
+	uint32 InstanceCountBufferOffset;
+	uint32 NumIndicesPerInstance; // When -1 the counter needs to be reset to 0.
+	uint32 StartIndexLocation;
+
+	bool operator==(const FNiagaraDrawIndirectArgGenTaskInfo& Rhs) const
+	{
+		return InstanceCountBufferOffset == Rhs.InstanceCountBufferOffset
+			&& NumIndicesPerInstance == Rhs.NumIndicesPerInstance
+			&& StartIndexLocation == Rhs.StartIndexLocation;
+	}
+};
+
 /**
  * Compute shader used to generate GPU emitter draw indirect args.
  * It also resets unused instance count entries.
@@ -26,33 +52,8 @@ class NIAGARAVERTEXFACTORIES_API FNiagaraDrawIndirectArgsGenCS : public FGlobalS
 	DECLARE_GLOBAL_SHADER(FNiagaraDrawIndirectArgsGenCS);
 
 public:
-
-	/** 
-	 * Task info when generating draw indirect frame buffer.
-	 * Task is either about generate Niagara renderers drawindirect buffer,
-	 * or about resetting released instance counters.
-	 */
-	struct FArgGenTaskInfo
-	{
-		FArgGenTaskInfo(uint32 InInstanceCountBufferOffset, uint32 InNumIndicesPerInstance, uint32 InStartIndexLocation) 
-			: InstanceCountBufferOffset(InInstanceCountBufferOffset)
-			, NumIndicesPerInstance(InNumIndicesPerInstance)
-			, StartIndexLocation(InStartIndexLocation)
-		{}
-
-		uint32 InstanceCountBufferOffset;
-		uint32 NumIndicesPerInstance; // When -1 the counter needs to be reset to 0.
-		uint32 StartIndexLocation;
-
-		bool operator==(const FArgGenTaskInfo& Rhs) const
-		{
-			return InstanceCountBufferOffset == Rhs.InstanceCountBufferOffset
-				&& NumIndicesPerInstance == Rhs.NumIndicesPerInstance
-				&& StartIndexLocation == Rhs.StartIndexLocation;
-		}
-	};
-
-	using FPermutationDomain = TShaderPermutationDomain<>;
+	class FSupportsTextureRW : SHADER_PERMUTATION_INT("SUPPORTS_TEXTURE_RW", 2);
+	using FPermutationDomain = TShaderPermutationDomain<FSupportsTextureRW>;
 
 	static bool ShouldCompilePermutation(const FGlobalShaderPermutationParameters& Parameters)
 	{
@@ -76,3 +77,35 @@ protected:
 	LAYOUT_FIELD(FRWShaderParameter, DrawIndirectArgsParam)
 	LAYOUT_FIELD(FShaderParameter, TaskCountParam)
 };
+
+/**
+ * Compute shader used to reset unused instance count entries. Used if the platform doesn't support RW texture buffers
+ */
+class NIAGARAVERTEXFACTORIES_API FNiagaraDrawIndirectResetCountsCS : public FGlobalShader
+{
+	DECLARE_GLOBAL_SHADER(FNiagaraDrawIndirectResetCountsCS);
+
+public:
+	using FPermutationDomain = TShaderPermutationDomain<>;
+
+	static bool ShouldCompilePermutation(const FGlobalShaderPermutationParameters& Parameters)
+	{
+		return RHISupportsComputeShaders(Parameters.Platform);
+	}
+
+	static void ModifyCompilationEnvironment(const FGlobalShaderPermutationParameters& Parameters, FShaderCompilerEnvironment& OutEnvironment);
+
+	FNiagaraDrawIndirectResetCountsCS() {}
+	FNiagaraDrawIndirectResetCountsCS(const ShaderMetaType::CompiledShaderInitializerType& Initializer);
+
+	//bool Serialize(FArchive& Ar);
+	void SetOutput(FRHICommandList& RHICmdList, FRHIUnorderedAccessView* InstanceCountsUAV);
+	void SetParameters(FRHICommandList& RHICmdList, FRHIShaderResourceView* TaskInfosBuffer, int32 NumArgGenTasks, int32 NumInstanceCountClearTasks);
+	void UnbindBuffers(FRHICommandList& RHICmdList);
+
+protected:
+
+	LAYOUT_FIELD(FShaderResourceParameter, TaskInfosParam)
+	LAYOUT_FIELD(FRWShaderParameter, InstanceCountsParam)
+	LAYOUT_FIELD(FShaderParameter, TaskCountParam)
+};
\ No newline at end of file
diff --git a/Engine/Source/Runtime/Apple/MetalRHI/Private/MetalRHI.cpp b/Engine/Source/Runtime/Apple/MetalRHI/Private/MetalRHI.cpp
--- a/Engine/Source/Runtime/Apple/MetalRHI/Private/MetalRHI.cpp
+++ b/Engine/Source/Runtime/Apple/MetalRHI/Private/MetalRHI.cpp
@@ -137,6 +137,8 @@ FMetalDynamicRHI::FMetalDynamicRHI(ERHIFeatureLevel::Type RequestedFeatureLevel)
 	bool bCanUseWideMRTs = Device.SupportsFeatureSet(mtlpp::FeatureSet::iOS_GPUFamily2_v1);
 	bool bCanUseASTC = Device.SupportsFeatureSet(mtlpp::FeatureSet::iOS_GPUFamily2_v1) && !FParse::Param(FCommandLine::Get(),TEXT("noastc"));
 	
+	GRHISupportsRWTextureBuffers = Device.SupportsFeatureSet(mtlpp::FeatureSet::iOS_GPUFamily4_v1);
+
 	const mtlpp::FeatureSet FeatureSets[] = {
 		mtlpp::FeatureSet::iOS_GPUFamily1_v1,
 		mtlpp::FeatureSet::iOS_GPUFamily2_v1,
diff --git a/Engine/Source/Runtime/RHI/Private/RHI.cpp b/Engine/Source/Runtime/RHI/Private/RHI.cpp
--- a/Engine/Source/Runtime/RHI/Private/RHI.cpp
+++ b/Engine/Source/Runtime/RHI/Private/RHI.cpp
@@ -679,6 +679,7 @@ TRHIGlobal<int32> GMaxTextureDimensions(2048);
 TRHIGlobal<int64> GMaxBufferDimensions(2<<27);
 TRHIGlobal<int32> GMaxVolumeTextureDimensions(2048);
 TRHIGlobal<int32> GMaxCubeTextureDimensions(2048);
+bool GRHISupportsRWTextureBuffers = true;
 int32 GMaxTextureArrayLayers = 256;
 int32 GMaxTextureSamplers = 16;
 bool GUsingNullRHI = false;
diff --git a/Engine/Source/Runtime/RHI/Public/RHI.h b/Engine/Source/Runtime/RHI/Public/RHI.h
--- a/Engine/Source/Runtime/RHI/Public/RHI.h
+++ b/Engine/Source/Runtime/RHI/Public/RHI.h
@@ -426,6 +426,9 @@ extern RHI_API TRHIGlobal<int64> GMaxBufferDimensions;
 /** The maximum size allowed for 3D textures in all three dimensions. */
 extern RHI_API TRHIGlobal<int32> GMaxVolumeTextureDimensions;
 
+/** Whether RW texture buffers are supported */
+extern RHI_API bool GRHISupportsRWTextureBuffers;
+
 FORCEINLINE uint64 GetMaxBufferDimension()
 {
 	return GMaxBufferDimensions;
