Fixing build failure caused by CL 12505182.

diff --git a/Engine/Plugins/FX/Niagara/Source/Niagara/Private/NiagaraComponent.cpp b/Engine/Plugins/FX/Niagara/Source/Niagara/Private/NiagaraComponent.cpp
--- a/Engine/Plugins/FX/Niagara/Source/Niagara/Private/NiagaraComponent.cpp
+++ b/Engine/Plugins/FX/Niagara/Source/Niagara/Private/NiagaraComponent.cpp
@@ -1950,7 +1950,7 @@ void UNiagaraComponent::ApplyOverridesToParameterStore()
 		}
 	}
 
-	EnsureRenderDataConsistent();
+	EnsureOverrideParametersConsistent();
 }
 
 #endif
