Fix: incorrect module parameter values at runtime when a module was renamed.

The name of the module is part of the variable names for its input parameters in the script parameter store. However, that name was not included in the compile hash, so renaming a module did not update the DDC key; the old compiled script was still retrieved from the DDC, with the wrong variable names, which led to all the inputs of the renamed module having default values.

#rnx

diff --git a/Engine/Plugins/FX/Niagara/Source/NiagaraEditor/Public/NiagaraNodeFunctionCall.h b/Engine/Plugins/FX/Niagara/Source/NiagaraEditor/Public/NiagaraNodeFunctionCall.h
--- a/Engine/Plugins/FX/Niagara/Source/NiagaraEditor/Public/NiagaraNodeFunctionCall.h
+++ b/Engine/Plugins/FX/Niagara/Source/NiagaraEditor/Public/NiagaraNodeFunctionCall.h
@@ -156,7 +156,7 @@ protected:
 	UPROPERTY(meta = (SkipForCompileHash="true"))
 	FGuid CachedChangeId;
 
-	UPROPERTY(meta = (SkipForCompileHash = "true"))
+	UPROPERTY()
 	FString FunctionDisplayName;
 
 	FOnInputsChanged OnInputsChangedDelegate;
