Fixed backwards logic from 12496791.

diff --git a/Engine/Plugins/FX/Niagara/Source/NiagaraEditor/Private/ViewModels/Stack/NiagaraStackSystemSettingsGroup.cpp b/Engine/Plugins/FX/Niagara/Source/NiagaraEditor/Private/ViewModels/Stack/NiagaraStackSystemSettingsGroup.cpp
--- a/Engine/Plugins/FX/Niagara/Source/NiagaraEditor/Private/ViewModels/Stack/NiagaraStackSystemSettingsGroup.cpp
+++ b/Engine/Plugins/FX/Niagara/Source/NiagaraEditor/Private/ViewModels/Stack/NiagaraStackSystemSettingsGroup.cpp
@@ -145,7 +145,7 @@ void UNiagaraStackSystemSettingsGroup::RefreshChildrenInternal(const TArray<UNia
 	//Check if we're trying to override scalability settings without an EffectType. Ideally we can allow this but it's somewhat awkward so for now we just post a warning and ignore this.
 	UNiagaraSystem* System = &GetSystemViewModel()->GetSystem();
 	TWeakPtr<FNiagaraSystemViewModel> WeakSysViewModel = GetSystemViewModel();
-	if(System && System->GetOverrideScalabilitySettings() && System->GetEffectType())
+	if(System && System->GetOverrideScalabilitySettings() && System->GetEffectType() == nullptr)
 	{
 		FText FixDescription = LOCTEXT("FixOverridesWithNoEffectType", "Disable Overrides");
 		FStackIssueFix FixIssue(
