Moved new stack warning to the correct place in the stack.

diff --git a/Engine/Plugins/FX/Niagara/Source/NiagaraEditor/Private/ViewModels/Stack/NiagaraStackSystemPropertiesItem.cpp b/Engine/Plugins/FX/Niagara/Source/NiagaraEditor/Private/ViewModels/Stack/NiagaraStackSystemPropertiesItem.cpp
--- a/Engine/Plugins/FX/Niagara/Source/NiagaraEditor/Private/ViewModels/Stack/NiagaraStackSystemPropertiesItem.cpp
+++ b/Engine/Plugins/FX/Niagara/Source/NiagaraEditor/Private/ViewModels/Stack/NiagaraStackSystemPropertiesItem.cpp
@@ -36,9 +36,39 @@ void UNiagaraStackSystemPropertiesItem::RefreshChildrenInternal(const TArray<UNi
 		FRequiredEntryData RequiredEntryData(GetSystemViewModel(), GetEmitterViewModel(), FExecutionCategoryNames::System, NAME_None, GetStackEditorData());
 		SystemObject->Initialize(RequiredEntryData, System.Get(), GetStackEditorDataKey());
 	}
-
 	NewChildren.Add(SystemObject);
 	bCanResetToBase.Reset();
+
+	//Check if we're trying to override scalability settings without an EffectType. Ideally we can allow this but it's somewhat awkward so for now we just post a warning and ignore this.
+	UNiagaraSystem* SystemPtr = System.Get();
+	TWeakPtr<FNiagaraSystemViewModel> WeakSysViewModel = GetSystemViewModel();
+	if (SystemPtr && SystemPtr->GetOverrideScalabilitySettings() && SystemPtr->GetEffectType() == nullptr)
+	{
+		FText FixDescription = LOCTEXT("FixOverridesWithNoEffectType", "Disable Overrides");
+		FStackIssueFix FixIssue(
+			FixDescription,
+			FStackIssueFixDelegate::CreateLambda([=]()
+				{
+					if (auto Pinned = WeakSysViewModel.Pin())
+					{
+						FScopedTransaction ScopedTransaction(FixDescription);
+						Pinned->GetSystem().Modify();
+						Pinned->GetSystem().SetOverrideScalabilitySettings(false);
+						Pinned->RefreshAll();
+					}
+				}));
+
+		FStackIssue OverridesWithNoEffectTypeWarning(
+			EStackIssueSeverity::Warning,
+			LOCTEXT("FixOverridesWithNoEffectTypeSummaryText", "Scalability overrides with no Effect Type."),
+			LOCTEXT("FixOverridesWithNoEffectTypeErrorText", "Scalability settings cannot be overriden if the System has no Effect Type."),
+			GetStackEditorDataKey(),
+			false,
+			FixIssue);
+
+		NewIssues.Add(OverridesWithNoEffectTypeWarning);
+	}
+
 	Super::RefreshChildrenInternal(CurrentChildren, NewChildren, NewIssues);
 }
 
diff --git a/Engine/Plugins/FX/Niagara/Source/NiagaraEditor/Private/ViewModels/Stack/NiagaraStackSystemSettingsGroup.cpp b/Engine/Plugins/FX/Niagara/Source/NiagaraEditor/Private/ViewModels/Stack/NiagaraStackSystemSettingsGroup.cpp
--- a/Engine/Plugins/FX/Niagara/Source/NiagaraEditor/Private/ViewModels/Stack/NiagaraStackSystemSettingsGroup.cpp
+++ b/Engine/Plugins/FX/Niagara/Source/NiagaraEditor/Private/ViewModels/Stack/NiagaraStackSystemSettingsGroup.cpp
@@ -142,36 +142,6 @@ void UNiagaraStackSystemSettingsGroup::RefreshChildrenInternal(const TArray<UNia
 	}
 	NewChildren.Add(SystemPropertiesItem);
 
-	//Check if we're trying to override scalability settings without an EffectType. Ideally we can allow this but it's somewhat awkward so for now we just post a warning and ignore this.
-	UNiagaraSystem* System = &GetSystemViewModel()->GetSystem();
-	TWeakPtr<FNiagaraSystemViewModel> WeakSysViewModel = GetSystemViewModel();
-	if(System && System->GetOverrideScalabilitySettings() && System->GetEffectType() == nullptr)
-	{
-		FText FixDescription = LOCTEXT("FixOverridesWithNoEffectType", "Disable Overrides");
-		FStackIssueFix FixIssue(
-			FixDescription,
-			FStackIssueFixDelegate::CreateLambda([=]()
-				{
-					if (auto Pinned = WeakSysViewModel.Pin())
-					{
-						FScopedTransaction ScopedTransaction(FixDescription);
-						Pinned->GetSystem().Modify();
-						Pinned->GetSystem().SetOverrideScalabilitySettings(false);
-						Pinned->RefreshAll();
-					}
-				}));
-
-		FStackIssue OverridesWithNoEffectTypeWarning(
-			EStackIssueSeverity::Warning,
-			LOCTEXT("FixOverridesWithNoEffectTypeSummaryText", "Scalability overrides with no Effect Type."),
-			LOCTEXT("FixOverridesWithNoEffectTypeErrorText", "Scalability settings cannot be overriden if the System has no Effect Type."),
-			GetStackEditorDataKey(),
-			false,
-			FixIssue);
-
-		NewIssues.Add(OverridesWithNoEffectTypeWarning);
-	}
-
 	Super::RefreshChildrenInternal(CurrentChildren, NewChildren, NewIssues);
 }
 
