CIS Fix

diff --git a/Engine/Plugins/FX/Niagara/Source/NiagaraEditor/Private/ViewModels/Stack/NiagaraStackSystemPropertiesItem.cpp b/Engine/Plugins/FX/Niagara/Source/NiagaraEditor/Private/ViewModels/Stack/NiagaraStackSystemPropertiesItem.cpp
--- a/Engine/Plugins/FX/Niagara/Source/NiagaraEditor/Private/ViewModels/Stack/NiagaraStackSystemPropertiesItem.cpp
+++ b/Engine/Plugins/FX/Niagara/Source/NiagaraEditor/Private/ViewModels/Stack/NiagaraStackSystemPropertiesItem.cpp
@@ -4,6 +4,7 @@
 #include "ViewModels/Stack/NiagaraStackObject.h"
 #include "NiagaraSystem.h"
 #include "ViewModels/NiagaraSystemViewModel.h"
+#include "ScopedTransaction.h"
 
 #define LOCTEXT_NAMESPACE "UNiagaraStackSystemItemGroup"
 
