Niagara Skeletal Mesh Data Interface changes:
- Neither null meshes nor non-CPU accessible meshes are a fail case now, and will sample the component's origin where data is inaccessible
- Improved error reporting in the UI when a mesh needs CPU access for the system/emitter (as opposed to only erroring in the output log)

diff --git a/Engine/Plugins/FX/Niagara/Source/Niagara/Classes/NDISkeletalMeshCommon.h b/Engine/Plugins/FX/Niagara/Source/Niagara/Classes/NDISkeletalMeshCommon.h
--- a/Engine/Plugins/FX/Niagara/Source/Niagara/Classes/NDISkeletalMeshCommon.h
+++ b/Engine/Plugins/FX/Niagara/Source/Niagara/Classes/NDISkeletalMeshCommon.h
@@ -1,6 +1,19 @@
 // Copyright Epic Games, Inc. All Rights Reserved.
 #pragma once
 
+// These are to help readability of template specializations
+using TNDISkelMesh_FilterModeNone = TIntegralConstant<ENDISkeletalMesh_FilterMode, ENDISkeletalMesh_FilterMode::None>;
+using TNDISkelMesh_FilterModeSingle = TIntegralConstant<ENDISkeletalMesh_FilterMode, ENDISkeletalMesh_FilterMode::SingleRegion>;
+using TNDISkelMesh_FilterModeMulti = TIntegralConstant<ENDISkeletalMesh_FilterMode, ENDISkeletalMesh_FilterMode::MultiRegion>;
+
+using TNDISkelMesh_AreaWeightingOff = TIntegralConstant<ENDISkelMesh_AreaWeightingMode, ENDISkelMesh_AreaWeightingMode::None>;
+using TNDISkelMesh_AreaWeightingOn = TIntegralConstant<ENDISkelMesh_AreaWeightingMode, ENDISkelMesh_AreaWeightingMode::AreaWeighted>;
+
+using TNDISkelMesh_SkinningModeInvalid = TIntegralConstant<ENDISkeletalMesh_SkinningMode, ENDISkeletalMesh_SkinningMode::Invalid>;
+using TNDISkelMesh_SkinningModeNone = TIntegralConstant<ENDISkeletalMesh_SkinningMode, ENDISkeletalMesh_SkinningMode::None>;
+using TNDISkelMesh_SkinningModeOnTheFly = TIntegralConstant<ENDISkeletalMesh_SkinningMode, ENDISkeletalMesh_SkinningMode::SkinOnTheFly>;
+using TNDISkelMesh_SkinningModePreSkin = TIntegralConstant<ENDISkeletalMesh_SkinningMode, ENDISkeletalMesh_SkinningMode::PreSkin>;
+
 //////////////////////////////////////////////////////////////////////////
 // Helper classes for reducing duplicate code when accessing vertex positions. 
 
@@ -27,8 +40,8 @@ struct FSkeletalMeshAccessorHelper
 	{
 		Comp = Cast<USkeletalMeshComponent>(InstData->Component.Get());
 		Mesh = InstData->Mesh;
-		LODData = &InstData->GetLODRenderDataAndSkinWeights(SkinWeightBuffer);
-		IndexBuffer = LODData->MultiSizeIndexContainer.GetIndexBuffer();
+		LODData = InstData->GetLODRenderDataAndSkinWeights(SkinWeightBuffer);
+		IndexBuffer = LODData ? LODData->MultiSizeIndexContainer.GetIndexBuffer() : nullptr;
 		SkinningData = InstData->SkinningData.SkinningData.Get();
 		Usage = InstData->SkinningData.Usage;
 
@@ -38,6 +51,16 @@ struct FSkeletalMeshAccessorHelper
 		}
 	}
 
+	FORCEINLINE bool AreBonesAccessible() const
+	{
+		return Mesh != nullptr;
+	}
+
+	FORCEINLINE bool IsSkinAccessible() const
+	{
+		return SkinWeightBuffer != nullptr;
+	}
+
 	USkeletalMeshComponent* Comp = nullptr;
 	USkeletalMesh* Mesh = nullptr;
 	TWeakObjectPtr<USkeletalMesh> MeshSafe;
@@ -51,16 +74,10 @@ struct FSkeletalMeshAccessorHelper
 };
 
 template<>
-void FSkeletalMeshAccessorHelper::Init<
-	TIntegralConstant<ENDISkeletalMesh_FilterMode, ENDISkeletalMesh_FilterMode::SingleRegion>,
-	TIntegralConstant<ENDISkelMesh_AreaWeightingMode, ENDISkelMesh_AreaWeightingMode::None>>
-	(FNDISkeletalMesh_InstanceData* InstData);
+void FSkeletalMeshAccessorHelper::Init<TNDISkelMesh_FilterModeSingle, TNDISkelMesh_AreaWeightingOff>(FNDISkeletalMesh_InstanceData* InstData);
 
 template<>
-void FSkeletalMeshAccessorHelper::Init<
-	TIntegralConstant<ENDISkeletalMesh_FilterMode, ENDISkeletalMesh_FilterMode::SingleRegion>,
-	TIntegralConstant<ENDISkelMesh_AreaWeightingMode, ENDISkelMesh_AreaWeightingMode::AreaWeighted>>
-	(FNDISkeletalMesh_InstanceData* InstData);
+void FSkeletalMeshAccessorHelper::Init<TNDISkelMesh_FilterModeSingle, TNDISkelMesh_AreaWeightingOn>(FNDISkeletalMesh_InstanceData* InstData);
 
 //////////////////////////////////////////////////////////////////////////
 
@@ -81,7 +98,73 @@ struct FSkinnedPositionAccessorHelper
 };
 
 template<>
-struct FSkinnedPositionAccessorHelper<TIntegralConstant<ENDISkeletalMesh_SkinningMode, ENDISkeletalMesh_SkinningMode::None>>
+struct FSkinnedPositionAccessorHelper<TNDISkelMesh_SkinningModeInvalid>
+{
+	FORCEINLINE int32 GetBoneCount(FSkeletalMeshAccessorHelper& Accessor, bool RequiresPrevious) const
+	{
+		if (const USkeletalMesh* Mesh = Accessor.Mesh)
+		{
+			return Mesh->RefSkeleton.GetNum();
+		}
+
+		return 0;
+	}
+
+	FORCEINLINE void GetTriangleIndices(FSkeletalMeshAccessorHelper& Accessor, int32 Tri, int32& Idx0, int32& Idx1, int32& Idx2)
+	{
+		Idx0 = Idx1 = Idx2 = -1;
+	}
+
+	FORCEINLINE void GetSkinnedTrianglePositions(FSkeletalMeshAccessorHelper& Accessor, int32 Idx0, int32 Idx1, int32 Idx2, FVector& OutPos0, FVector& OutPos1, FVector& OutPos2)
+	{
+		OutPos0 = OutPos1 = OutPos2 = FVector::ZeroVector;
+	}
+
+	FORCEINLINE void GetSkinnedTrianglePreviousPositions(FSkeletalMeshAccessorHelper& Accessor, int32 Idx0, int32 Idx1, int32 Idx2, FVector& OutPos0, FVector& OutPos1, FVector& OutPos2)
+	{
+		OutPos0 = OutPos1 = OutPos2 = FVector::ZeroVector;
+	}
+
+	FORCEINLINE FVector GetSkinnedVertexPosition(FSkeletalMeshAccessorHelper& Accessor, int32 VertexIndex)
+	{
+		return FVector::ZeroVector;
+	}
+
+	FORCEINLINE FVector GetSkinnedVertexPreviousPosition(FSkeletalMeshAccessorHelper& Accessor, int32 VertexIndex)
+	{
+		return FVector::ZeroVector;
+	}
+
+	FORCEINLINE void GetSkinnedTangentBasis(FSkeletalMeshAccessorHelper& Accessor, int32 VertexIndex, FVector& OutTangentX, FVector& OutTangentZ)
+	{
+		OutTangentX = FVector(1.0f, 0.0f, 0.0f);
+		OutTangentZ = FVector(0.0f, 0.0f, 1.0f);
+	}
+
+	// The bone accessor functions are valid if a mesh is present, so don't stub them entirely
+	FORCEINLINE_DEBUGGABLE FVector GetSkinnedBonePosition(FSkeletalMeshAccessorHelper& Accessor, int32 BoneIndex)
+	{
+		return Accessor.Mesh ? Accessor.Mesh->GetComposedRefPoseMatrix(BoneIndex).GetOrigin() : FVector::ZeroVector;
+	}
+
+	FORCEINLINE_DEBUGGABLE FVector GetSkinnedBonePreviousPosition(FSkeletalMeshAccessorHelper& Accessor, int32 BoneIndex)
+	{
+		return Accessor.Mesh ? Accessor.Mesh->GetComposedRefPoseMatrix(BoneIndex).GetOrigin() : FVector::ZeroVector;
+	}
+
+	FORCEINLINE_DEBUGGABLE FQuat GetSkinnedBoneRotation(FSkeletalMeshAccessorHelper& Accessor, int32 BoneIndex)
+	{
+		return Accessor.Mesh ? Accessor.Mesh->GetComposedRefPoseMatrix(BoneIndex).ToQuat() : FQuat::Identity;
+	}
+
+	FORCEINLINE_DEBUGGABLE FQuat GetSkinnedBonePreviousRotation(FSkeletalMeshAccessorHelper& Accessor, int32 BoneIndex)
+	{
+		return Accessor.Mesh ? Accessor.Mesh->GetComposedRefPoseMatrix(BoneIndex).ToQuat() : FQuat::Identity;
+	}
+};
+
+template<>
+struct FSkinnedPositionAccessorHelper<TNDISkelMesh_SkinningModeNone>
 {
 	FORCEINLINE int32 GetBoneCount(FSkeletalMeshAccessorHelper& Accessor, bool RequiresPrevious) const
 	{
@@ -96,7 +179,7 @@ struct FSkinnedPositionAccessorHelper<TIntegralConstant<ENDISkeletalMesh_Skinnin
 	FORCEINLINE void GetTriangleIndices(FSkeletalMeshAccessorHelper& Accessor, int32 Tri, int32& Idx0, int32& Idx1, int32& Idx2)
 	{
 		const int32 BaseIndex = Tri * 3;
-		checkSlow(BaseIndex + 2 < Accessor.IndexBuffer->Num());
+		check(BaseIndex + 2 < Accessor.IndexBuffer->Num());
 		Idx0 = Accessor.IndexBuffer->Get(BaseIndex);
 		Idx1 = Accessor.IndexBuffer->Get(BaseIndex + 1);
 		Idx2 = Accessor.IndexBuffer->Get(BaseIndex + 2);
@@ -167,7 +250,7 @@ struct FSkinnedPositionAccessorHelper<TIntegralConstant<ENDISkeletalMesh_Skinnin
 };
 
 template<>
-struct FSkinnedPositionAccessorHelper<TIntegralConstant<ENDISkeletalMesh_SkinningMode, ENDISkeletalMesh_SkinningMode::SkinOnTheFly>>
+struct FSkinnedPositionAccessorHelper<TNDISkelMesh_SkinningModeOnTheFly>
 {
 	FORCEINLINE int32 GetBoneCount(FSkeletalMeshAccessorHelper& Accessor, bool RequiresPrevious) const
 	{
@@ -182,7 +265,7 @@ struct FSkinnedPositionAccessorHelper<TIntegralConstant<ENDISkeletalMesh_Skinnin
 	FORCEINLINE void GetTriangleIndices(FSkeletalMeshAccessorHelper& Accessor, int32 Tri, int32& Idx0, int32& Idx1, int32& Idx2)
 	{
 		const int32 BaseIndex = Tri * 3;
-		checkSlow(BaseIndex + 2 < Accessor.IndexBuffer->Num());
+		check(BaseIndex + 2 < Accessor.IndexBuffer->Num());
 		Idx0 = Accessor.IndexBuffer->Get(BaseIndex);
 		Idx1 = Accessor.IndexBuffer->Get(BaseIndex + 1);
 		Idx2 = Accessor.IndexBuffer->Get(BaseIndex + 2);
@@ -239,7 +322,7 @@ struct FSkinnedPositionAccessorHelper<TIntegralConstant<ENDISkeletalMesh_Skinnin
 };
 
 template<>
-struct FSkinnedPositionAccessorHelper<TIntegralConstant<ENDISkeletalMesh_SkinningMode, ENDISkeletalMesh_SkinningMode::PreSkin>>
+struct FSkinnedPositionAccessorHelper<TNDISkelMesh_SkinningModePreSkin>
 {
 	FORCEINLINE int32 GetBoneCount(FSkeletalMeshAccessorHelper& Accessor, bool RequiresPrevious) const
 	{
@@ -254,7 +337,7 @@ struct FSkinnedPositionAccessorHelper<TIntegralConstant<ENDISkeletalMesh_Skinnin
 	FORCEINLINE void GetTriangleIndices(FSkeletalMeshAccessorHelper& Accessor, int32 Tri, int32& Idx0, int32& Idx1, int32& Idx2)
 	{
 		const int32 BaseIndex = Tri * 3;
-		checkSlow(BaseIndex + 2 < Accessor.IndexBuffer->Num());
+		check(BaseIndex + 2 < Accessor.IndexBuffer->Num());
 		Idx0 = Accessor.IndexBuffer->Get(BaseIndex);
 		Idx1 = Accessor.IndexBuffer->Get(BaseIndex + 1);
 		Idx2 = Accessor.IndexBuffer->Get(BaseIndex + 2);
@@ -311,28 +394,43 @@ struct FSkinnedPositionAccessorHelper<TIntegralConstant<ENDISkeletalMesh_Skinnin
 };
 
 //////////////////////////////////////////////////////////////////////////
-// Helper for accessing misc vertex data
+// Helpers for accessing misc vertex data
 template<bool bUseFullPrecisionUVs>
 struct FSkelMeshVertexAccessor
 {
-	FORCEINLINE FVector2D GetVertexUV(FSkeletalMeshLODRenderData& LODData, int32 VertexIdx, int32 UVChannel)const
+	FORCEINLINE FVector2D GetVertexUV(FSkeletalMeshLODRenderData* LODData, int32 VertexIdx, int32 UVChannel)const
 	{
+		check(LODData);
 		if (bUseFullPrecisionUVs)
 		{
-			return LODData.StaticVertexBuffers.StaticMeshVertexBuffer.GetVertexUV_Typed<EStaticMeshVertexUVType::HighPrecision>(VertexIdx, UVChannel);
+			return LODData->StaticVertexBuffers.StaticMeshVertexBuffer.GetVertexUV_Typed<EStaticMeshVertexUVType::HighPrecision>(VertexIdx, UVChannel);
 		}
 		else
 		{
-			return LODData.StaticVertexBuffers.StaticMeshVertexBuffer.GetVertexUV_Typed<EStaticMeshVertexUVType::Default>(VertexIdx, UVChannel);
+			return LODData->StaticVertexBuffers.StaticMeshVertexBuffer.GetVertexUV_Typed<EStaticMeshVertexUVType::Default>(VertexIdx, UVChannel);
 		}
 	}
 
-	FORCEINLINE FLinearColor GetVertexColor(FSkeletalMeshLODRenderData& LODData, int32 VertexIdx)const
+	FORCEINLINE FLinearColor GetVertexColor(FSkeletalMeshLODRenderData* LODData, int32 VertexIdx)const
 	{
+		check(LODData);
 		return LODData.StaticVertexBuffers.ColorVertexBuffer.VertexColor(VertexIdx);
 	}
 };
 
+struct FSkelMeshVertexAccessorNoop
+{
+	FORCEINLINE FVector2D GetVertexUV(FSkeletalMeshLODRenderData* LODData, int32 VertexIdx, int32 UVChannel)const
+	{
+		return FVector2D(0.0f, 0.0f);
+	}
+
+	FORCEINLINE FLinearColor GetVertexColor(FSkeletalMeshLODRenderData* LODData, int32 VertexIdx)const
+	{
+		return FLinearColor::White;
+	}
+};
+
 //////////////////////////////////////////////////////////////////////////
 //Function Binders.
 
@@ -345,32 +443,36 @@ struct TAreaWeightingModeBinder
 	{
 		FNDISkeletalMesh_InstanceData* InstData = (FNDISkeletalMesh_InstanceData*)InstanceData;
 		check(InstData);
+
 		UNiagaraDataInterfaceSkeletalMesh* MeshInterface = CastChecked<UNiagaraDataInterfaceSkeletalMesh>(Interface);
-		const FSkeletalMeshSamplingInfo& SamplingInfo = InstData->Mesh->GetSamplingInfo();
 
 		bool bAreaWeighting = false;
-		if (InstData->SamplingRegionIndices.Num() > 1)
-		{
-			bAreaWeighting = InstData->SamplingRegionAreaWeightedSampler.IsValid();
-		}
-		else if (InstData->SamplingRegionIndices.Num() == 1)
+		if (InstData->bAllowCPUMeshDataAccess)
 		{
-			const FSkeletalMeshSamplingRegion& Region = SamplingInfo.GetRegion(InstData->SamplingRegionIndices[0]);
-			bAreaWeighting = Region.bSupportUniformlyDistributedSampling;
-		}
-		else
-		{
-			int32 LODIndex = InstData->GetLODIndex();
-			bAreaWeighting = InstData->Mesh->GetLODInfo(LODIndex)->bSupportUniformlyDistributedSampling;
+			if (InstData->SamplingRegionIndices.Num() > 1)
+			{
+				bAreaWeighting = InstData->SamplingRegionAreaWeightedSampler.IsValid();
+			}
+			else if (InstData->SamplingRegionIndices.Num() == 1)
+			{
+				const FSkeletalMeshSamplingInfo& SamplingInfo = InstData->Mesh->GetSamplingInfo();
+				const FSkeletalMeshSamplingRegion& Region = SamplingInfo.GetRegion(InstData->SamplingRegionIndices[0]);
+				bAreaWeighting = Region.bSupportUniformlyDistributedSampling;
+			}
+			else
+			{
+				int32 LODIndex = InstData->GetLODIndex();
+				bAreaWeighting = InstData->Mesh->GetLODInfo(LODIndex)->bSupportUniformlyDistributedSampling;
+			}
 		}
 
 		if (bAreaWeighting)
 		{
-			NextBinder::template Bind<ParamTypes..., TIntegralConstant<ENDISkelMesh_AreaWeightingMode, ENDISkelMesh_AreaWeightingMode::AreaWeighted>>(Interface, BindingInfo, InstanceData, OutFunc);
+			NextBinder::template Bind<ParamTypes..., TNDISkelMesh_AreaWeightingOn>(Interface, BindingInfo, InstanceData, OutFunc);
 		}
 		else
 		{
-			NextBinder::template Bind<ParamTypes..., TIntegralConstant<ENDISkelMesh_AreaWeightingMode, ENDISkelMesh_AreaWeightingMode::None>>(Interface, BindingInfo, InstanceData, OutFunc);
+			NextBinder::template Bind<ParamTypes..., TNDISkelMesh_AreaWeightingOff>(Interface, BindingInfo, InstanceData, OutFunc);
 		}
 	}
 };
@@ -384,24 +486,23 @@ struct TFilterModeBinder
 	{
 		FNDISkeletalMesh_InstanceData* InstData = (FNDISkeletalMesh_InstanceData*)InstanceData;
 		check(InstData);
-		UNiagaraDataInterfaceSkeletalMesh* MeshInterface = CastChecked<UNiagaraDataInterfaceSkeletalMesh>(Interface);
 
 		if (InstData->SamplingRegionIndices.Num() == 1)
 		{
-			NextBinder::template Bind<ParamTypes..., TIntegralConstant<ENDISkeletalMesh_FilterMode, ENDISkeletalMesh_FilterMode::SingleRegion>>(Interface, BindingInfo, InstanceData, OutFunc);
+			NextBinder::template Bind<ParamTypes..., TNDISkelMesh_FilterModeSingle>(Interface, BindingInfo, InstanceData, OutFunc);
 		}
 		else if (InstData->SamplingRegionIndices.Num() > 1)
 		{
-			NextBinder::template Bind<ParamTypes..., TIntegralConstant<ENDISkeletalMesh_FilterMode, ENDISkeletalMesh_FilterMode::MultiRegion>>(Interface, BindingInfo, InstanceData, OutFunc);
+			NextBinder::template Bind<ParamTypes..., TNDISkelMesh_FilterModeMulti>(Interface, BindingInfo, InstanceData, OutFunc);
 		}
 		else
 		{
-			NextBinder::template Bind<ParamTypes..., TIntegralConstant<ENDISkeletalMesh_FilterMode, ENDISkeletalMesh_FilterMode::None>>(Interface, BindingInfo, InstanceData, OutFunc);
+			NextBinder::template Bind<ParamTypes..., TNDISkelMesh_FilterModeNone>(Interface, BindingInfo, InstanceData, OutFunc);
 		}
 	}
 };
 
-//External function binder choosing between template specializations based vetrex data format
+//External function binder choosing between template specializations based vertex data format
 template<typename NextBinder>
 struct TVertexAccessorBinder
 {
@@ -409,18 +510,24 @@ struct TVertexAccessorBinder
 	static void Bind(UNiagaraDataInterface* Interface, const FVMExternalFunctionBindingInfo& BindingInfo, void* InstanceData, FVMExternalFunction &OutFunc)
 	{
 		FNDISkeletalMesh_InstanceData* InstData = (FNDISkeletalMesh_InstanceData*)InstanceData;
-		UNiagaraDataInterfaceSkeletalMesh* MeshInterface = CastChecked<UNiagaraDataInterfaceSkeletalMesh>(Interface);
-		USkeletalMeshComponent* Component = Cast<USkeletalMeshComponent>(InstData->Component.Get());
-		FSkinWeightVertexBuffer* SkinWeightBuffer = nullptr;
-		FSkeletalMeshLODRenderData& LODData = InstData->GetLODRenderDataAndSkinWeights(SkinWeightBuffer);
+		check(InstData);
 
-		if (LODData.StaticVertexBuffers.StaticMeshVertexBuffer.GetUseFullPrecisionUVs())
+		if (InstData->bAllowCPUMeshDataAccess)
 		{
-			NextBinder::template Bind<ParamTypes..., FSkelMeshVertexAccessor<true>>(Interface, BindingInfo, InstanceData, OutFunc);
+			FSkinWeightVertexBuffer* SkinWeightBuffer = nullptr;
+			FSkeletalMeshLODRenderData* LODData = InstData->GetLODRenderDataAndSkinWeights(SkinWeightBuffer);
+			if (LODData->StaticVertexBuffers.StaticMeshVertexBuffer.GetUseFullPrecisionUVs())
+			{
+				NextBinder::template Bind<ParamTypes..., FSkelMeshVertexAccessor<true>>(Interface, BindingInfo, InstanceData, OutFunc);
+			}
+			else
+			{
+				NextBinder::template Bind<ParamTypes..., FSkelMeshVertexAccessor<false>>(Interface, BindingInfo, InstanceData, OutFunc);
+			}
 		}
 		else
 		{
-			NextBinder::template Bind<ParamTypes..., FSkelMeshVertexAccessor<false>>(Interface, BindingInfo, InstanceData, OutFunc);
+			NextBinder::template Bind<ParamTypes..., FSkelMeshVertexAccessorNoop>(Interface, BindingInfo, InstanceData, OutFunc);
 		}
 	}
 };
@@ -435,19 +542,23 @@ struct TSkinningModeBinder
 		FNDISkeletalMesh_InstanceData* InstData = (FNDISkeletalMesh_InstanceData*)InstanceData;
 		UNiagaraDataInterfaceSkeletalMesh* MeshInterface = CastChecked<UNiagaraDataInterfaceSkeletalMesh>(Interface);
 		USkeletalMeshComponent* Component = Cast<USkeletalMeshComponent>(InstData->Component.Get());
-		if (MeshInterface->SkinningMode == ENDISkeletalMesh_SkinningMode::None || !Component)//Can't skin if we have no component.
+		if (!InstData->bAllowCPUMeshDataAccess) // No-op when we can't access the mesh on CPU
+		{
+			NextBinder::template Bind<ParamTypes..., FSkinnedPositionAccessorHelper<TNDISkelMesh_SkinningModeInvalid>>(Interface, BindingInfo, InstanceData, OutFunc);
+		}
+		else if (MeshInterface->SkinningMode == ENDISkeletalMesh_SkinningMode::None || !Component) // Can't skin if we have no component.
 		{
-			NextBinder::template Bind<ParamTypes..., FSkinnedPositionAccessorHelper<TIntegralConstant<ENDISkeletalMesh_SkinningMode, ENDISkeletalMesh_SkinningMode::None>>>(Interface, BindingInfo, InstanceData, OutFunc);
+			NextBinder::template Bind<ParamTypes..., FSkinnedPositionAccessorHelper<TNDISkelMesh_SkinningModeNone>>(Interface, BindingInfo, InstanceData, OutFunc);
 		}
 		else if (MeshInterface->SkinningMode == ENDISkeletalMesh_SkinningMode::SkinOnTheFly)
 		{
 			check(Component);
-			NextBinder::template Bind<ParamTypes..., FSkinnedPositionAccessorHelper<TIntegralConstant<ENDISkeletalMesh_SkinningMode, ENDISkeletalMesh_SkinningMode::SkinOnTheFly>>>(Interface, BindingInfo, InstanceData, OutFunc);
+			NextBinder::template Bind<ParamTypes..., FSkinnedPositionAccessorHelper<TNDISkelMesh_SkinningModeOnTheFly>>(Interface, BindingInfo, InstanceData, OutFunc);
 		}
 		else if (MeshInterface->SkinningMode == ENDISkeletalMesh_SkinningMode::PreSkin)
 		{
 			check(Component);
-			NextBinder::template Bind<ParamTypes..., FSkinnedPositionAccessorHelper<TIntegralConstant<ENDISkeletalMesh_SkinningMode, ENDISkeletalMesh_SkinningMode::PreSkin>>>(Interface, BindingInfo, InstanceData, OutFunc);
+			NextBinder::template Bind<ParamTypes..., FSkinnedPositionAccessorHelper<TNDISkelMesh_SkinningModePreSkin>>(Interface, BindingInfo, InstanceData, OutFunc);
 		}
 		else
 		{
diff --git a/Engine/Plugins/FX/Niagara/Source/Niagara/Classes/NiagaraDataInterface.h b/Engine/Plugins/FX/Niagara/Source/Niagara/Classes/NiagaraDataInterface.h
--- a/Engine/Plugins/FX/Niagara/Source/Niagara/Classes/NiagaraDataInterface.h
+++ b/Engine/Plugins/FX/Niagara/Source/Niagara/Classes/NiagaraDataInterface.h
@@ -122,38 +122,87 @@ public:
 		, ErrorSummaryText(InErrorSummaryText)
 		, Fix(InFix)
 
-	{};
+	{}
+
 	FNiagaraDataInterfaceError()
-	{};
+	{}
+
 	/** Returns true if the error can be fixed automatically. */
 	bool GetErrorFixable() const
 	{
 		return Fix.IsBound();
-	};
+	}
 
 	/** Applies the fix if a delegate is bound for it.*/
 	bool TryFixError()
 	{
 		return Fix.IsBound() ? Fix.Execute() : false;
-	};
+	}
 
 	/** Full error description text */
 	FText GetErrorText() const
 	{
 		return ErrorText;
-	};
+	}
 
 	/** Shortened error description text*/
 	FText GetErrorSummaryText() const
 	{
 		return ErrorSummaryText;
-	};
+	}
 
 private:
 	FText ErrorText;
 	FText ErrorSummaryText;
 	FNiagaraDataInterfaceFix Fix;
 };
+
+// Helper class for GUI feedback handling
+DECLARE_DELEGATE_RetVal(bool, FNiagaraDataInterfaceFix);
+class FNiagaraDataInterfaceFeedback
+{
+public:
+	FNiagaraDataInterfaceFeedback(FText InFeedbackText,
+		FText InFeedbackSummaryText,
+		FNiagaraDataInterfaceFix InFix)
+		: FeedbackText(InFeedbackText)
+		, FeedbackSummaryText(InFeedbackSummaryText)
+		, Fix(InFix)
+
+	{}
+
+	FNiagaraDataInterfaceFeedback()
+	{}
+
+	/** Returns true if the feedback can be fixed automatically. */
+	bool GetFeedbackFixable() const
+	{
+		return Fix.IsBound();
+	}
+
+	/** Applies the fix if a delegate is bound for it.*/
+	bool TryFixFeedback()
+	{
+		return Fix.IsBound() ? Fix.Execute() : false;
+	}
+
+	/** Full feedback description text */
+	FText GetFeedbackText() const
+	{
+		return FeedbackText;
+	}
+
+	/** Shortened feedback description text*/
+	FText GetFeedbackSummaryText() const
+	{
+		return FeedbackSummaryText;
+	}
+
+private:
+	FText FeedbackText;
+	FText FeedbackSummaryText;
+	FNiagaraDataInterfaceFix Fix;
+};
 #endif
 
 //////////////////////////////////////////////////////////////////////////
@@ -195,13 +244,15 @@ UCLASS(abstract, EditInlineNew)
 class NIAGARA_API UNiagaraDataInterface : public UNiagaraDataInterfaceBase
 {
 	GENERATED_UCLASS_BODY()
-		 
-public: 
 
+public:
 	virtual ~UNiagaraDataInterface();
 
 	// UObject Interface
 	virtual void PostLoad() override;
+#if WITH_EDITOR
+	virtual void PostEditChangeProperty(struct FPropertyChangedEvent& PropertyChangedEvent) override;
+#endif
 	// UObject Interface END
 
 	/** Initializes the per instance data for this interface. Returns false if there was some error and the simulation should be disabled. */
@@ -307,9 +358,24 @@ public:
 #if WITH_EDITOR	
 	/** Refreshes and returns the errors detected with the corresponding data, if any.*/
 	virtual TArray<FNiagaraDataInterfaceError> GetErrors() { return TArray<FNiagaraDataInterfaceError>(); }
+	
+	/**
+		Query the data interface to give feedback to the end user. 
+		Note that the default implementation, just calls GetErrors on the DataInterface, but derived classes can do much more.
+		Also, InAsset or InComponent may be null values, as the UI for DataInterfaces is displayed in a variety of locations. 
+		In these cases, only provide information that is relevant to that context.
+	*/
+	virtual void GetFeedback(UNiagaraSystem* InAsset, UNiagaraComponent* InComponent, TArray<FNiagaraDataInterfaceError>& OutErrors, TArray<FNiagaraDataInterfaceFeedback>& OutWarnings, TArray<FNiagaraDataInterfaceFeedback>& OutInfo);
+
+	static void GetFeedback(UNiagaraDataInterface* DataInterface, TArray<FNiagaraDataInterfaceError>& Errors, TArray<FNiagaraDataInterfaceFeedback>& Warnings,
+		TArray<FNiagaraDataInterfaceFeedback>& Info);
 
 	/** Validates a function being compiled and allows interface classes to post custom compile errors when their API changes. */
 	virtual void ValidateFunction(const FNiagaraFunctionSignature& Function, TArray<FText>& OutValidationErrors);
+
+	void RefreshErrors();
+
+	FSimpleMulticastDelegate& OnErrorsRefreshed();
 #endif
 
 	FNiagaraDataInterfaceProxy* GetProxy()
@@ -337,6 +403,11 @@ protected:
 	virtual bool CopyToInternal(UNiagaraDataInterface* Destination) const;
 
 	TUniquePtr<FNiagaraDataInterfaceProxy> Proxy;
+
+private:
+#if WITH_EDITOR
+	FSimpleMulticastDelegate OnErrorsRefreshedDelegate;
+#endif
 };
 
 /** Helper class for decoding NDI parameters into a usable struct type. */
diff --git a/Engine/Plugins/FX/Niagara/Source/Niagara/Classes/NiagaraDataInterfaceSkeletalMesh.h b/Engine/Plugins/FX/Niagara/Source/Niagara/Classes/NiagaraDataInterfaceSkeletalMesh.h
--- a/Engine/Plugins/FX/Niagara/Source/Niagara/Classes/NiagaraDataInterfaceSkeletalMesh.h
+++ b/Engine/Plugins/FX/Niagara/Source/Niagara/Classes/NiagaraDataInterfaceSkeletalMesh.h
@@ -238,8 +238,10 @@ public:
 UENUM()
 enum class ENDISkeletalMesh_SkinningMode : uint8
 {
+	Invalid = (uint8)-1 UMETA(Hidden),
+
 	/** No skinning. */
-	None,
+	None = 0,
 	/** Skin vertex locations as you need them. Use if you have a high poly mesh or you are sampling the interface a small number of times. */
 	SkinOnTheFly,
 	/**
@@ -514,27 +516,34 @@ struct FNDISkeletalMesh_InstanceData
 	FSkeletalMeshGpuSpawnStaticBuffers* MeshGpuSpawnStaticBuffers;
 	FSkeletalMeshGpuDynamicBufferProxy* MeshGpuSpawnDynamicBuffers;
 
-	/** Temporary flag to deny binding VM functions that rely on mesh data being accessible on the CPU */
+	/** Flag to stub VM functions that rely on mesh data being accessible on the CPU */
 	bool bAllowCPUMeshDataAccess;
 
 	FORCEINLINE_DEBUGGABLE bool ResetRequired(UNiagaraDataInterfaceSkeletalMesh* Interface)const;
 
 	bool Init(UNiagaraDataInterfaceSkeletalMesh* Interface, FNiagaraSystemInstance* SystemInstance);
-	FORCEINLINE_DEBUGGABLE bool Cleanup(UNiagaraDataInterfaceSkeletalMesh* Interface, FNiagaraSystemInstance* SystemInstance);
 	FORCEINLINE_DEBUGGABLE bool Tick(UNiagaraDataInterfaceSkeletalMesh* Interface, FNiagaraSystemInstance* SystemInstance, float InDeltaSeconds);
 	FORCEINLINE_DEBUGGABLE void Release();
 	FORCEINLINE int32 GetLODIndex()const { return SkinningData.Usage.GetLODIndex(); }
 
-	FORCEINLINE_DEBUGGABLE FSkeletalMeshLODRenderData& GetLODRenderDataAndSkinWeights(FSkinWeightVertexBuffer*& OutSkinWeightBuffer)
+	FORCEINLINE_DEBUGGABLE FSkeletalMeshLODRenderData* GetLODRenderDataAndSkinWeights(FSkinWeightVertexBuffer*& OutSkinWeightBuffer)
 	{
-		FSkeletalMeshLODRenderData& Ret = Mesh->GetResourceForRendering()->LODRenderData[GetLODIndex()];
-		if (USkeletalMeshComponent* SkelComp = Cast<USkeletalMeshComponent>(Component.Get()))
+		FSkeletalMeshLODRenderData* Ret = nullptr;
+		OutSkinWeightBuffer = nullptr;
+		if (Mesh)
 		{
-			OutSkinWeightBuffer = SkelComp->GetSkinWeightBuffer(GetLODIndex());
-		}
-		else
-		{
-			OutSkinWeightBuffer = &Ret.SkinWeightVertexBuffer; // Todo @sckime fix this when main comes in for real
+			Ret = &Mesh->GetResourceForRendering()->LODRenderData[GetLODIndex()];
+			if (bAllowCPUMeshDataAccess)
+			{
+				if (USkeletalMeshComponent* SkelComp = Cast<USkeletalMeshComponent>(Component.Get()))
+				{
+					OutSkinWeightBuffer = SkelComp->GetSkinWeightBuffer(GetLODIndex());
+				}
+				else
+				{
+					OutSkinWeightBuffer = &Ret->SkinWeightVertexBuffer; // Todo @sckime fix this when main comes in for real
+				}
+			}
 		}
 		return Ret;
 	}
@@ -613,12 +622,6 @@ public:
 	UPROPERTY(EditAnywhere, Category = "Skeleton", meta = (InlineEditConditionToggle))
 	uint8 bExcludeBone : 1;
 
-#if WITH_EDITORONLY_DATA
-	/** Do we require CPU access to the data, this is set during GetVMExternalFunction. */
-	UPROPERTY(transient)
-	bool bRequiresCPUAccess;
-#endif
-
 	/** Cached change id off of the data interface.*/
 	uint32 ChangeId;
 
@@ -641,7 +644,8 @@ public:
 	virtual bool Equals(const UNiagaraDataInterface* Other) const override;
 	virtual bool CanExecuteOnTarget(ENiagaraSimTarget Target)const override { return true; }
 #if WITH_EDITOR	
-	virtual TArray<FNiagaraDataInterfaceError> GetErrors() override;
+	virtual void GetFeedback(UNiagaraSystem* Asset, UNiagaraComponent* Component, TArray<FNiagaraDataInterfaceError>& OutErrors,
+		TArray<FNiagaraDataInterfaceFeedback>& Warnings, TArray<FNiagaraDataInterfaceFeedback>& Info) override;
 	virtual void ValidateFunction(const FNiagaraFunctionSignature& Function, TArray<FText>& OutValidationErrors) override;
 #endif
 	virtual bool HasTickGroupPrereqs() const override { return true; }
@@ -725,6 +729,9 @@ public:
 	template<typename SkinningHandlerType, typename TransformHandlerType, typename VertexAccessorType, typename bInterpolated>
 	void GetTriCoordSkinnedData(FVectorVMContext& Context);
 
+	template<typename TransformHandlerType, typename bInterpolated>
+	void GetTriCoordSkinnedDataFallback(FVectorVMContext& Context);
+
 	void GetTriCoordColor(FVectorVMContext& Context);
 
 	void GetTriCoordColorFallback(FVectorVMContext& Context);
@@ -802,8 +809,10 @@ public:
 	void BindSkeletonSamplingFunction(const FVMExternalFunctionBindingInfo& BindingInfo, FNDISkeletalMesh_InstanceData* InstData, FVMExternalFunction &OutFunc);
 
 	template<typename SkinningHandlerType, typename TransformHandlerType, typename bInterpolated>
-	void GetSkinnedBoneData(FVectorVMContext& Context);
-
+	void GetSkinnedBoneData(FVectorVMContext& Context);	
+	template<typename TransformHandlerType, typename bInterpolated>
+	void GetSkinnedBoneDataFallback(FVectorVMContext& Context);
+	
 	void IsValidBone(FVectorVMContext& Context);
 	void RandomBone(FVectorVMContext& Context);
 	void GetBoneCount(FVectorVMContext& Context);
@@ -820,7 +829,7 @@ public:
 	void GetFilteredSocketBoneAt(FVectorVMContext& Context);
 	void GetFilteredSocketTransform(FVectorVMContext& Context);
 	void RandomFilteredSocket(FVectorVMContext& Context);
-
+		
 	void RandomFilteredSocketOrBone(FVectorVMContext& Context);
 	void GetFilteredSocketOrBoneCount(FVectorVMContext& Context);
 	void GetFilteredSocketOrBoneBoneAt(FVectorVMContext& Context);
diff --git a/Engine/Plugins/FX/Niagara/Source/Niagara/Private/NDISkeletalMesh_BoneSampling.cpp b/Engine/Plugins/FX/Niagara/Source/Niagara/Private/NDISkeletalMesh_BoneSampling.cpp
--- a/Engine/Plugins/FX/Niagara/Source/Niagara/Private/NDISkeletalMesh_BoneSampling.cpp
+++ b/Engine/Plugins/FX/Niagara/Source/Niagara/Private/NDISkeletalMesh_BoneSampling.cpp
@@ -15,6 +15,7 @@ DECLARE_CYCLE_STAT(TEXT("Skel Mesh Skeleton Sampling"), STAT_NiagaraSkel_Bone_Sa
 
 //Final binders for all static mesh interface functions.
 DEFINE_NDI_FUNC_BINDER(UNiagaraDataInterfaceSkeletalMesh, GetSkinnedBoneData)
+DEFINE_NDI_FUNC_BINDER(UNiagaraDataInterfaceSkeletalMesh, GetSkinnedBoneDataFallback)
 DEFINE_NDI_DIRECT_FUNC_BINDER(UNiagaraDataInterfaceSkeletalMesh, IsValidBone)
 
 DEFINE_NDI_DIRECT_FUNC_BINDER(UNiagaraDataInterfaceSkeletalMesh, GetFilteredSocketBoneAt)
@@ -334,26 +335,57 @@ void UNiagaraDataInterfaceSkeletalMesh::GetSkeletonSamplingFunctions(TArray<FNia
 
 void UNiagaraDataInterfaceSkeletalMesh::BindSkeletonSamplingFunction(const FVMExternalFunctionBindingInfo& BindingInfo, FNDISkeletalMesh_InstanceData* InstanceData, FVMExternalFunction &OutFunc)
 {
+	using TInterpOff = TIntegralConstant<bool, false>;
+	using TInterpOn = TIntegralConstant<bool, true>;
+
 	//Bone Functions
 	if (BindingInfo.Name == FSkeletalMeshInterfaceHelper::GetSkinnedBoneDataName)
 	{
 		ensure(BindingInfo.GetNumInputs() == 2 && BindingInfo.GetNumOutputs() == 10);
-		TSkinningModeBinder<TNDIExplicitBinder<FNDITransformHandlerNoop, TNDIExplicitBinder<TIntegralConstant<bool, false>, NDI_FUNC_BINDER(UNiagaraDataInterfaceSkeletalMesh, GetSkinnedBoneData)>>>::Bind(this, BindingInfo, InstanceData, OutFunc);
+		if (InstanceData->Mesh)
+		{
+			TSkinningModeBinder<TNDIExplicitBinder<FNDITransformHandlerNoop, TNDIExplicitBinder<TInterpOff, NDI_FUNC_BINDER(UNiagaraDataInterfaceSkeletalMesh, GetSkinnedBoneData)>>>::Bind(this, BindingInfo, InstanceData, OutFunc);
+		}
+		else
+		{
+			NDI_FUNC_BINDER(UNiagaraDataInterfaceSkeletalMesh, GetSkinnedBoneDataFallback)::Bind<FNDITransformHandlerNoop, TInterpOff>(this, BindingInfo, InstanceData, OutFunc);
+		}
 	}
 	else if (BindingInfo.Name == FSkeletalMeshInterfaceHelper::GetSkinnedBoneDataWSName)
 	{
 		ensure(BindingInfo.GetNumInputs() == 2 && BindingInfo.GetNumOutputs() == 10);
-		TSkinningModeBinder<TNDIExplicitBinder<FNDITransformHandler, TNDIExplicitBinder<TIntegralConstant<bool, false>, NDI_FUNC_BINDER(UNiagaraDataInterfaceSkeletalMesh, GetSkinnedBoneData)>>>::Bind(this, BindingInfo, InstanceData, OutFunc);
+		if (InstanceData->Mesh)
+		{
+			TSkinningModeBinder<TNDIExplicitBinder<FNDITransformHandler, TNDIExplicitBinder<TInterpOff, NDI_FUNC_BINDER(UNiagaraDataInterfaceSkeletalMesh, GetSkinnedBoneData)>>>::Bind(this, BindingInfo, InstanceData, OutFunc);
+		}
+		else
+		{
+			NDI_FUNC_BINDER(UNiagaraDataInterfaceSkeletalMesh, GetSkinnedBoneDataFallback)::Bind<FNDITransformHandler, TInterpOff>(this, BindingInfo, InstanceData, OutFunc);
+		}
 	}
 	else if (BindingInfo.Name == FSkeletalMeshInterfaceHelper::GetSkinnedBoneDataInterpolatedName)
 	{
 		ensure(BindingInfo.GetNumInputs() == 3 && BindingInfo.GetNumOutputs() == 10);
-		TSkinningModeBinder<TNDIExplicitBinder<FNDITransformHandlerNoop, TNDIExplicitBinder<TIntegralConstant<bool, true>, NDI_FUNC_BINDER(UNiagaraDataInterfaceSkeletalMesh, GetSkinnedBoneData)>>>::Bind(this, BindingInfo, InstanceData, OutFunc);
+		if (InstanceData->Mesh)
+		{
+			TSkinningModeBinder<TNDIExplicitBinder<FNDITransformHandlerNoop, TNDIExplicitBinder<TInterpOn, NDI_FUNC_BINDER(UNiagaraDataInterfaceSkeletalMesh, GetSkinnedBoneData)>>>::Bind(this, BindingInfo, InstanceData, OutFunc);
+		}
+		else
+		{
+			NDI_FUNC_BINDER(UNiagaraDataInterfaceSkeletalMesh, GetSkinnedBoneDataFallback)::Bind<FNDITransformHandlerNoop, TInterpOn>(this, BindingInfo, InstanceData, OutFunc);
+		}
 	}
 	else if (BindingInfo.Name == FSkeletalMeshInterfaceHelper::GetSkinnedBoneDataWSInterpolatedName)
 	{
 		ensure(BindingInfo.GetNumInputs() == 3 && BindingInfo.GetNumOutputs() == 10);
-		TSkinningModeBinder<TNDIExplicitBinder<FNDITransformHandler, TNDIExplicitBinder<TIntegralConstant<bool, true>, NDI_FUNC_BINDER(UNiagaraDataInterfaceSkeletalMesh, GetSkinnedBoneData)>>>::Bind(this, BindingInfo, InstanceData, OutFunc);
+		if (InstanceData->Mesh)
+		{
+			TSkinningModeBinder<TNDIExplicitBinder<FNDITransformHandler, TNDIExplicitBinder<TInterpOn, NDI_FUNC_BINDER(UNiagaraDataInterfaceSkeletalMesh, GetSkinnedBoneData)>>>::Bind(this, BindingInfo, InstanceData, OutFunc);
+		}
+		else
+		{
+			NDI_FUNC_BINDER(UNiagaraDataInterfaceSkeletalMesh, GetSkinnedBoneDataFallback)::Bind<FNDITransformHandler, TInterpOn>(this, BindingInfo, InstanceData, OutFunc);
+		}
 	}
 	else if (BindingInfo.Name == FSkeletalMeshInterfaceHelper::IsValidBoneName)
 	{
@@ -598,21 +630,33 @@ void UNiagaraDataInterfaceSkeletalMesh::IsValidBone(FVectorVMContext& Context)
 	VectorVM::FExternalFuncInputHandler<int32> BoneParam(Context);
 
 	checkfSlow(InstData.Get(), TEXT("Skeletal Mesh Interface has invalid instance data. %s"), *GetPathName());
-	checkfSlow(InstData->Mesh, TEXT("Skeletal Mesh Interface has invalid mesh. %s"), *GetPathName());
 
 	VectorVM::FExternalFuncRegisterHandler<FNiagaraBool> OutValid(Context);
 
 	FSkeletalMeshAccessorHelper MeshAccessor;
-	MeshAccessor.Init<TIntegralConstant<int32, 0>, TIntegralConstant<int32, 0>>(InstData);
-	const FReferenceSkeleton& RefSkeleton = MeshAccessor.Mesh->RefSkeleton;
-	int32 NumBones = RefSkeleton.GetNum();
-	for (int32 i = 0; i < Context.NumInstances; ++i)
+	MeshAccessor.Init<TNDISkelMesh_FilterModeNone, TNDISkelMesh_AreaWeightingOff>(InstData);	
+
+	if (MeshAccessor.AreBonesAccessible())
 	{
-		int32 RequestedIndex = BoneParam.GetAndAdvance();
+		const FReferenceSkeleton& RefSkeleton = MeshAccessor.Mesh->RefSkeleton;
+		int32 NumBones = RefSkeleton.GetNum();
+		for (int32 i = 0; i < Context.NumInstances; ++i)
+		{
+			int32 RequestedIndex = BoneParam.GetAndAdvance();
 
-		FNiagaraBool Value;
-		Value.SetValue(RequestedIndex >= 0 && RequestedIndex < NumBones);
-		*OutValid.GetDestAndAdvance() = Value;
+			FNiagaraBool Value;
+			Value.SetValue(RequestedIndex >= 0 && RequestedIndex < NumBones);
+			*OutValid.GetDestAndAdvance() = Value;
+		}
+	}
+	else
+	{
+		for (int32 i = 0; i < Context.NumInstances; ++i)
+		{
+			FNiagaraBool Value;
+			Value.SetValue(false);
+			*OutValid.GetDestAndAdvance() = Value;
+		}
 	}
 }
 
@@ -624,9 +668,15 @@ void UNiagaraDataInterfaceSkeletalMesh::RandomBone(FVectorVMContext& Context)
 
 	FSkeletalMeshAccessorHelper MeshAccessor;
 	MeshAccessor.Init<TIntegralConstant<int32, 0>, TIntegralConstant<int32, 0>>(InstData);
-	const FReferenceSkeleton& RefSkeleton = MeshAccessor.Mesh->RefSkeleton;
+
 	const int32 ExcludedBoneIndex = InstData->ExcludedBoneIndex;
-	const int32 NumBones = RefSkeleton.GetNum() - (ExcludedBoneIndex >= 0 ? 2 : 1);
+	int32 NumBones = 0;
+	if (MeshAccessor.AreBonesAccessible())
+	{
+		const FReferenceSkeleton& RefSkeleton = MeshAccessor.Mesh->RefSkeleton;
+		NumBones = RefSkeleton.GetNum() - (ExcludedBoneIndex >= 0 ? 2 : 1);
+	}
+
 	if (NumBones >= 0)
 	{
 		for (int32 i = 0; i < Context.NumInstances; ++i)
@@ -651,8 +701,14 @@ void UNiagaraDataInterfaceSkeletalMesh::GetBoneCount(FVectorVMContext& Context)
 
 	FSkeletalMeshAccessorHelper MeshAccessor;
 	MeshAccessor.Init<TIntegralConstant<int32, 0>, TIntegralConstant<int32, 0>>(InstData);
-	const FReferenceSkeleton& RefSkeleton = MeshAccessor.Mesh->RefSkeleton;
-	const int32 NumBones = RefSkeleton.GetNum();
+	
+	int32 NumBones = 0;
+	if (MeshAccessor.AreBonesAccessible())
+	{
+		const FReferenceSkeleton& RefSkeleton = MeshAccessor.Mesh->RefSkeleton;
+		NumBones = RefSkeleton.GetNum();
+	}
+
 	for (int32 i = 0; i < Context.NumInstances; ++i)
 	{
 		*OutCount.GetDestAndAdvance() = NumBones;
@@ -734,7 +790,8 @@ void UNiagaraDataInterfaceSkeletalMesh::GetSkinnedBoneData(FVectorVMContext& Con
 	FSkinWeightVertexBuffer* SkinWeightBuffer = InstData->GetSkinWeights();
 
 	FSkeletalMeshAccessorHelper Accessor;
-	Accessor.Init<TIntegralConstant<int32, 0>, TIntegralConstant<int32, 0>>(InstData);
+	Accessor.Init<TNDISkelMesh_FilterModeNone, TNDISkelMesh_AreaWeightingOff>(InstData);
+	check(Accessor.AreBonesAccessible());
 
 	const int32 BoneCount = SkinningHandler.GetBoneCount(Accessor, bInterpolated::Value);
 	const int32 BoneAndSocketCount = BoneCount + InstData->FilteredSocketInfo.Num();
@@ -748,18 +805,15 @@ void UNiagaraDataInterfaceSkeletalMesh::GetSkinnedBoneData(FVectorVMContext& Con
 		const float Interp = bInterpolated::Value ? InterpParam.GetAndAdvance() : 1.0f;
 
 		// Determine bone or socket
-		const int32 RawBone = BoneParam.GetAndAdvance();
-		const int32 Bone = FMath::Clamp(RawBone, 0, BoneAndSocketCount - 1);
+		const int32 Bone = BoneParam.GetAndAdvance();
 		const bool bIsSocket = Bone >= BoneCount;
 		const int32 Socket = Bone - BoneCount;
 
 		FVector Pos;
 		FVector Prev;
-		FVector Velocity;
 
-		// Handle edge cases first...
-		if ((!bIsSocket && Bone >= BoneCount) ||
-			(bIsSocket && !FilteredSocketCurrTransforms.IsValidIndex(Socket)))
+		// Handle invalid bone indices first
+		if (Bone < 0 || Bone >= BoneAndSocketCount)
 		{
 			Pos = FVector::ZeroVector;
 			TransformHandler.TransformPosition(Pos, InstanceTransform);
@@ -771,17 +825,10 @@ void UNiagaraDataInterfaceSkeletalMesh::GetSkinnedBoneData(FVectorVMContext& Con
 			}
 			if (Output.bNeedsRotation)
 			{
-				FQuat Rotation = FQuat::Identity;
-				if (bInterpolated::Value)
-				{
-					FQuat PrevRotation = FQuat::Identity;
-					Rotation = FQuat::Identity;
-				}
-
-				Output.SetRotation(Rotation);
+				Output.SetRotation(FQuat::Identity);
 			}
 		}
-		else if ( bIsSocket )
+		else if (bIsSocket)
 		{
 			FTransform CurrSocketTransform = FilteredSocketCurrTransforms[Socket];
 			FTransform PrevSocketTransform = FilteredSocketPrevTransforms[Socket];
@@ -849,7 +896,70 @@ void UNiagaraDataInterfaceSkeletalMesh::GetSkinnedBoneData(FVectorVMContext& Con
 		if(Output.bNeedsVelocity)
 		{
 			//Don't have enough information to get a better interpolated velocity.
-			Velocity = (Pos - Prev) * InvDt;
+			FVector Velocity = (Pos - Prev) * InvDt;
+			Output.SetVelocity(Velocity);
+		}
+	}
+}
+
+template<typename TransformHandlerType, typename bInterpolated>
+void UNiagaraDataInterfaceSkeletalMesh::GetSkinnedBoneDataFallback(FVectorVMContext& Context)
+{
+	SCOPE_CYCLE_COUNTER(STAT_NiagaraSkel_Bone_Sample);
+	TransformHandlerType TransformHandler;
+	VectorVM::FUserPtrHandler<FNDISkeletalMesh_InstanceData> InstData(Context);
+
+	VectorVM::FExternalFuncInputHandler<int32> BoneParam(Context);
+	VectorVM::FExternalFuncInputHandler<float> InterpParam;
+
+	if (bInterpolated::Value)
+	{
+		InterpParam.Init(Context);
+	}
+
+	checkfSlow(InstData.Get(), TEXT("Skeletal Mesh Interface has invalid instance data. %s"), *GetPathName());	
+
+	FBoneSocketSkinnedDataOutputHandler Output(Context);
+
+	//TODO: Replace this by storing off FTransforms and doing a proper lerp to get a final transform.
+	//Also need to pull in a per particle interpolation factor.
+	const FMatrix& InstanceTransform = InstData->Transform;
+	const FMatrix& PrevInstanceTransform = InstData->PrevTransform;
+	const FQuat InstanceRotation = Output.bNeedsRotation ? InstData->Transform.GetMatrixWithoutScale().ToQuat() : FQuat::Identity;
+	const FQuat PrevInstanceRotation = Output.bNeedsRotation ? InstData->Transform.GetMatrixWithoutScale().ToQuat() : FQuat::Identity;
+	const float InvDt = 1.0f / InstData->DeltaSeconds;
+
+	for (int32 i = 0; i < Context.NumInstances; ++i)
+	{
+		const float Interp = bInterpolated::Value ? InterpParam.GetAndAdvance() : 1.0f;
+
+		FVector Prev = FVector::ZeroVector;
+		FVector Pos = FVector::ZeroVector;		
+		TransformHandler.TransformPosition(Pos, InstanceTransform);
+
+		if (Output.bNeedsVelocity || bInterpolated::Value)
+		{
+			TransformHandler.TransformPosition(Prev, PrevInstanceTransform);
+		}
+
+		if (Output.bNeedsRotation)
+		{
+			Output.SetRotation(FQuat::Identity);
+		}
+		
+		if (Output.bNeedsVelocity || bInterpolated::Value)
+		{
+			Pos = FMath::Lerp(Prev, Pos, Interp);
+		}
+
+		if (Output.bNeedsPosition)
+		{
+			Output.SetPosition(Pos);
+		}
+
+		if (Output.bNeedsVelocity)
+		{
+			FVector Velocity = (Pos - Prev) * InvDt;
 			Output.SetVelocity(Velocity);
 		}
 	}
diff --git a/Engine/Plugins/FX/Niagara/Source/Niagara/Private/NDISkeletalMesh_TriangleSampling.cpp b/Engine/Plugins/FX/Niagara/Source/Niagara/Private/NDISkeletalMesh_TriangleSampling.cpp
--- a/Engine/Plugins/FX/Niagara/Source/Niagara/Private/NDISkeletalMesh_TriangleSampling.cpp
+++ b/Engine/Plugins/FX/Niagara/Source/Niagara/Private/NDISkeletalMesh_TriangleSampling.cpp
@@ -16,6 +16,7 @@ DECLARE_CYCLE_STAT(TEXT("Skel Mesh Sampling"), STAT_NiagaraSkel_Sample, STATGROU
 //Final binders for all static mesh interface functions.
 DEFINE_NDI_FUNC_BINDER(UNiagaraDataInterfaceSkeletalMesh, RandomTriCoord);
 DEFINE_NDI_FUNC_BINDER(UNiagaraDataInterfaceSkeletalMesh, GetTriCoordSkinnedData);
+DEFINE_NDI_FUNC_BINDER(UNiagaraDataInterfaceSkeletalMesh, GetTriCoordSkinnedDataFallback);
 DEFINE_NDI_DIRECT_FUNC_BINDER(UNiagaraDataInterfaceSkeletalMesh, GetTriCoordColor);
 DEFINE_NDI_DIRECT_FUNC_BINDER(UNiagaraDataInterfaceSkeletalMesh, GetTriCoordColorFallback);
 DEFINE_NDI_FUNC_BINDER(UNiagaraDataInterfaceSkeletalMesh, GetTriCoordUV);
@@ -234,6 +235,9 @@ void UNiagaraDataInterfaceSkeletalMesh::GetTriangleSamplingFunctions(TArray<FNia
 
 void UNiagaraDataInterfaceSkeletalMesh::BindTriangleSamplingFunction(const FVMExternalFunctionBindingInfo& BindingInfo, FNDISkeletalMesh_InstanceData* InstanceData, FVMExternalFunction &OutFunc)
 {
+	using TInterpOff = TIntegralConstant<bool, false>;
+	using TInterpOn = TIntegralConstant<bool, true>;
+
 	if (BindingInfo.Name == FSkeletalMeshInterfaceHelper::RandomTriCoordName)
 	{
 		check(BindingInfo.GetNumInputs() == 4 && BindingInfo.GetNumOutputs() == 4);
@@ -247,22 +251,50 @@ void UNiagaraDataInterfaceSkeletalMesh::BindTriangleSamplingFunction(const FVMEx
 	else if (BindingInfo.Name == FSkeletalMeshInterfaceHelper::GetSkinnedTriangleDataName)
 	{
 		check(BindingInfo.GetNumInputs() == 5 && BindingInfo.GetNumOutputs() == 15);
-		TSkinningModeBinder<TNDIExplicitBinder<FNDITransformHandlerNoop, TVertexAccessorBinder<TNDIExplicitBinder<TIntegralConstant<bool, false>, NDI_FUNC_BINDER(UNiagaraDataInterfaceSkeletalMesh, GetTriCoordSkinnedData)>>>>::Bind(this, BindingInfo, InstanceData, OutFunc);
+		if (InstanceData->bAllowCPUMeshDataAccess)
+		{
+			TSkinningModeBinder<TNDIExplicitBinder<FNDITransformHandlerNoop, TVertexAccessorBinder<TNDIExplicitBinder<TInterpOff, NDI_FUNC_BINDER(UNiagaraDataInterfaceSkeletalMesh, GetTriCoordSkinnedData)>>>>::Bind(this, BindingInfo, InstanceData, OutFunc);
+		}
+		else
+		{
+			NDI_FUNC_BINDER(UNiagaraDataInterfaceSkeletalMesh, GetTriCoordSkinnedDataFallback)::Bind<FNDITransformHandlerNoop, TInterpOff>(this, BindingInfo, InstanceData, OutFunc);
+		}
 	}
 	else if (BindingInfo.Name == FSkeletalMeshInterfaceHelper::GetSkinnedTriangleDataWSName)
 	{
 		check(BindingInfo.GetNumInputs() == 5 && BindingInfo.GetNumOutputs() == 15);
-		TSkinningModeBinder<TNDIExplicitBinder<FNDITransformHandler, TVertexAccessorBinder<TNDIExplicitBinder<TIntegralConstant<bool, false>, NDI_FUNC_BINDER(UNiagaraDataInterfaceSkeletalMesh, GetTriCoordSkinnedData)>>>>::Bind(this, BindingInfo, InstanceData, OutFunc);
+		if (InstanceData->bAllowCPUMeshDataAccess)
+		{
+			TSkinningModeBinder<TNDIExplicitBinder<FNDITransformHandler, TVertexAccessorBinder<TNDIExplicitBinder<TIntegralConstant<bool, false>, NDI_FUNC_BINDER(UNiagaraDataInterfaceSkeletalMesh, GetTriCoordSkinnedData)>>>>::Bind(this, BindingInfo, InstanceData, OutFunc);
+		}
+		else
+		{
+			NDI_FUNC_BINDER(UNiagaraDataInterfaceSkeletalMesh, GetTriCoordSkinnedDataFallback)::Bind<FNDITransformHandler, TInterpOff>(this, BindingInfo, InstanceData, OutFunc);
+		}
 	}
 	else if (BindingInfo.Name == FSkeletalMeshInterfaceHelper::GetSkinnedTriangleDataInterpName)
 	{
 		check(BindingInfo.GetNumInputs() == 6 && BindingInfo.GetNumOutputs() == 15);
-		TSkinningModeBinder<TNDIExplicitBinder<FNDITransformHandlerNoop, TVertexAccessorBinder<TNDIExplicitBinder<TIntegralConstant<bool, true>, NDI_FUNC_BINDER(UNiagaraDataInterfaceSkeletalMesh, GetTriCoordSkinnedData)>>>>::Bind(this, BindingInfo, InstanceData, OutFunc);
+		if (InstanceData->bAllowCPUMeshDataAccess)
+		{
+			TSkinningModeBinder<TNDIExplicitBinder<FNDITransformHandlerNoop, TVertexAccessorBinder<TNDIExplicitBinder<TInterpOn, NDI_FUNC_BINDER(UNiagaraDataInterfaceSkeletalMesh, GetTriCoordSkinnedData)>>>>::Bind(this, BindingInfo, InstanceData, OutFunc);
+		}
+		else
+		{
+			NDI_FUNC_BINDER(UNiagaraDataInterfaceSkeletalMesh, GetTriCoordSkinnedDataFallback)::Bind<FNDITransformHandlerNoop, TInterpOn>(this, BindingInfo, InstanceData, OutFunc);
+		}
 	}
 	else if (BindingInfo.Name == FSkeletalMeshInterfaceHelper::GetSkinnedTriangleDataWSInterpName)
 	{
 		check(BindingInfo.GetNumInputs() == 6 && BindingInfo.GetNumOutputs() == 15);
-		TSkinningModeBinder<TNDIExplicitBinder< FNDITransformHandler, TVertexAccessorBinder<TNDIExplicitBinder<TIntegralConstant<bool, true>,NDI_FUNC_BINDER(UNiagaraDataInterfaceSkeletalMesh, GetTriCoordSkinnedData)>>>>::Bind(this, BindingInfo, InstanceData, OutFunc);
+		if (InstanceData->bAllowCPUMeshDataAccess)
+		{
+			TSkinningModeBinder<TNDIExplicitBinder< FNDITransformHandler, TVertexAccessorBinder<TNDIExplicitBinder<TIntegralConstant<bool, true>, NDI_FUNC_BINDER(UNiagaraDataInterfaceSkeletalMesh, GetTriCoordSkinnedData)>>>>::Bind(this, BindingInfo, InstanceData, OutFunc);
+		}
+		else
+		{
+			NDI_FUNC_BINDER(UNiagaraDataInterfaceSkeletalMesh, GetTriCoordSkinnedDataFallback)::Bind<FNDITransformHandler, TInterpOn>(this, BindingInfo, InstanceData, OutFunc);
+		}
 	}
 	else if (BindingInfo.Name == FSkeletalMeshInterfaceHelper::GetTriColorName)
 	{
@@ -284,7 +316,7 @@ void UNiagaraDataInterfaceSkeletalMesh::BindTriangleSamplingFunction(const FVMEx
 	else if (BindingInfo.Name == FSkeletalMeshInterfaceHelper::GetTriCoordVerticesName)
 	{
 		check(BindingInfo.GetNumInputs() == 2 && BindingInfo.GetNumOutputs() == 3);
-		TSkinningModeBinder<NDI_FUNC_BINDER(UNiagaraDataInterfaceSkeletalMesh, GetTriCoordVertices)>::Bind(this, BindingInfo, InstanceData, OutFunc);
+		TSkinningModeBinder<NDI_FUNC_BINDER(UNiagaraDataInterfaceSkeletalMesh, GetTriCoordVertices)>::Bind(this, BindingInfo, InstanceData, OutFunc);		
 	}
 	else if (BindingInfo.Name == FSkeletalMeshInterfaceHelper::RandomTriangleName)
 	{
@@ -321,9 +353,7 @@ FORCEINLINE int32 UNiagaraDataInterfaceSkeletalMesh::RandomTriIndex(FNDIRandomHe
 }
 
 template<>
-FORCEINLINE int32 UNiagaraDataInterfaceSkeletalMesh::RandomTriIndex<
-	TIntegralConstant<ENDISkeletalMesh_FilterMode, ENDISkeletalMesh_FilterMode::None>,
-	TIntegralConstant<ENDISkelMesh_AreaWeightingMode, ENDISkelMesh_AreaWeightingMode::None>>
+FORCEINLINE int32 UNiagaraDataInterfaceSkeletalMesh::RandomTriIndex<TNDISkelMesh_FilterModeNone, TNDISkelMesh_AreaWeightingOff>
 	(FNDIRandomHelper& RandHelper, FSkeletalMeshAccessorHelper& Accessor, FNDISkeletalMesh_InstanceData* InstData, int32 InstanceIndex)
 {
 	int32 SecIdx = RandHelper.RandRange(InstanceIndex, 0, Accessor.LODData->RenderSections.Num() - 1);
@@ -334,9 +364,7 @@ FORCEINLINE int32 UNiagaraDataInterfaceSkeletalMesh::RandomTriIndex<
 }
 
 template<>
-FORCEINLINE int32 UNiagaraDataInterfaceSkeletalMesh::RandomTriIndex<
-	TIntegralConstant<ENDISkeletalMesh_FilterMode, ENDISkeletalMesh_FilterMode::None>,
-	TIntegralConstant<ENDISkelMesh_AreaWeightingMode, ENDISkelMesh_AreaWeightingMode::AreaWeighted>>
+FORCEINLINE int32 UNiagaraDataInterfaceSkeletalMesh::RandomTriIndex<TNDISkelMesh_FilterModeNone, TNDISkelMesh_AreaWeightingOn>
 	(FNDIRandomHelper& RandHelper, FSkeletalMeshAccessorHelper& Accessor, FNDISkeletalMesh_InstanceData* InstData, int32 InstanceIndex)
 {
 	const FSkeletalMeshSamplingInfo& SamplingInfo = InstData->Mesh->GetSamplingInfo();
@@ -346,9 +374,7 @@ FORCEINLINE int32 UNiagaraDataInterfaceSkeletalMesh::RandomTriIndex<
 }
 
 template<>
-FORCEINLINE int32 UNiagaraDataInterfaceSkeletalMesh::RandomTriIndex<
-	TIntegralConstant<ENDISkeletalMesh_FilterMode, ENDISkeletalMesh_FilterMode::SingleRegion>,
-	TIntegralConstant<ENDISkelMesh_AreaWeightingMode, ENDISkelMesh_AreaWeightingMode::None>>
+FORCEINLINE int32 UNiagaraDataInterfaceSkeletalMesh::RandomTriIndex<TNDISkelMesh_FilterModeSingle, TNDISkelMesh_AreaWeightingOff>
 	(FNDIRandomHelper& RandHelper, FSkeletalMeshAccessorHelper& Accessor, FNDISkeletalMesh_InstanceData* InstData, int32 InstanceIndex)
 {
 	int32 Idx = RandHelper.RandRange(InstanceIndex, 0, Accessor.SamplingRegionBuiltData->TriangleIndices.Num() - 1);
@@ -356,9 +382,7 @@ FORCEINLINE int32 UNiagaraDataInterfaceSkeletalMesh::RandomTriIndex<
 }
 
 template<>
-FORCEINLINE int32 UNiagaraDataInterfaceSkeletalMesh::RandomTriIndex<
-	TIntegralConstant<ENDISkeletalMesh_FilterMode, ENDISkeletalMesh_FilterMode::SingleRegion>,
-	TIntegralConstant<ENDISkelMesh_AreaWeightingMode, ENDISkelMesh_AreaWeightingMode::AreaWeighted>>
+FORCEINLINE int32 UNiagaraDataInterfaceSkeletalMesh::RandomTriIndex<TNDISkelMesh_FilterModeSingle, TNDISkelMesh_AreaWeightingOn>
 	(FNDIRandomHelper& RandHelper, FSkeletalMeshAccessorHelper& Accessor, FNDISkeletalMesh_InstanceData* InstData, int32 InstanceIndex)
 {
 	int32 Idx = Accessor.SamplingRegionBuiltData->AreaWeightedSampler.GetEntryIndex(RandHelper.Rand(InstanceIndex), RandHelper.Rand(InstanceIndex));
@@ -366,9 +390,7 @@ FORCEINLINE int32 UNiagaraDataInterfaceSkeletalMesh::RandomTriIndex<
 }
 
 template<>
-FORCEINLINE int32 UNiagaraDataInterfaceSkeletalMesh::RandomTriIndex<
-	TIntegralConstant<ENDISkeletalMesh_FilterMode, ENDISkeletalMesh_FilterMode::MultiRegion>,
-	TIntegralConstant<ENDISkelMesh_AreaWeightingMode, ENDISkelMesh_AreaWeightingMode::None>>
+FORCEINLINE int32 UNiagaraDataInterfaceSkeletalMesh::RandomTriIndex<TNDISkelMesh_FilterModeMulti, TNDISkelMesh_AreaWeightingOff>
 	(FNDIRandomHelper& RandHelper, FSkeletalMeshAccessorHelper& Accessor, FNDISkeletalMesh_InstanceData* InstData, int32 InstanceIndex)
 {
 	int32 RegionIdx = RandHelper.RandRange(InstanceIndex, 0, InstData->SamplingRegionIndices.Num() - 1);
@@ -380,9 +402,7 @@ FORCEINLINE int32 UNiagaraDataInterfaceSkeletalMesh::RandomTriIndex<
 }
 
 template<>
-FORCEINLINE int32 UNiagaraDataInterfaceSkeletalMesh::RandomTriIndex<
-	TIntegralConstant<ENDISkeletalMesh_FilterMode, ENDISkeletalMesh_FilterMode::MultiRegion>,
-	TIntegralConstant<ENDISkelMesh_AreaWeightingMode, ENDISkelMesh_AreaWeightingMode::AreaWeighted>>
+FORCEINLINE int32 UNiagaraDataInterfaceSkeletalMesh::RandomTriIndex<TNDISkelMesh_FilterModeMulti, TNDISkelMesh_AreaWeightingOn>
 	(FNDIRandomHelper& RandHelper, FSkeletalMeshAccessorHelper& Accessor, FNDISkeletalMesh_InstanceData* InstData, int32 InstanceIndex)
 {
 	int32 RegionIdx = InstData->SamplingRegionAreaWeightedSampler.GetEntryIndex(RandHelper.Rand(InstanceIndex), RandHelper.Rand(InstanceIndex));
@@ -403,23 +423,37 @@ void UNiagaraDataInterfaceSkeletalMesh::RandomTriCoord(FVectorVMContext& Context
 	FNDIRandomHelper RandHelper(Context);
 
 	checkfSlow(InstData.Get(), TEXT("Skeletal Mesh Interface has invalid instance data. %s"), *GetPathName());
-	checkfSlow(InstData->Mesh, TEXT("Skeletal Mesh Interface has invalid mesh. %s"), *GetPathName());
 
 	VectorVM::FExternalFuncRegisterHandler<int32> OutTri(Context);
-	VectorVM::FExternalFuncRegisterHandler<float> OutBaryX(Context);	VectorVM::FExternalFuncRegisterHandler<float> OutBaryY(Context);	VectorVM::FExternalFuncRegisterHandler<float> OutBaryZ(Context);
+	VectorVM::FExternalFuncRegisterHandler<float> OutBaryX(Context);
+	VectorVM::FExternalFuncRegisterHandler<float> OutBaryY(Context);
+	VectorVM::FExternalFuncRegisterHandler<float> OutBaryZ(Context);
 
 	FSkeletalMeshAccessorHelper MeshAccessor;
 	MeshAccessor.Init<FilterMode, AreaWeightingMode>(InstData);
-	for (int32 i = 0; i < Context.NumInstances; ++i)
-	{
-		RandHelper.GetAndAdvance();//We grab the rand info to a local value first so it can be used for multiple rand calls from the helper.
 
-		*OutTri.GetDest() = RandomTriIndex<FilterMode, AreaWeightingMode>(RandHelper, MeshAccessor, InstData, i);
-		FVector Bary = RandomBarycentricCoord(Context.RandStream);
-		*OutBaryX.GetDest() = Bary.X;		*OutBaryY.GetDest() = Bary.Y;		*OutBaryZ.GetDest() = Bary.Z;
+	if (MeshAccessor.IsSkinAccessible())
+	{
+		for (int32 i = 0; i < Context.NumInstances; ++i)
+		{
+			RandHelper.GetAndAdvance();//We grab the rand info to a local value first so it can be used for multiple rand calls from the helper.
 
-		OutTri.Advance();
-		OutBaryX.Advance();		OutBaryY.Advance();		OutBaryZ.Advance();
+			*OutTri.GetDestAndAdvance() = RandomTriIndex<FilterMode, AreaWeightingMode>(RandHelper, MeshAccessor, InstData, i);
+			FVector Bary = RandomBarycentricCoord(Context.RandStream);
+			*OutBaryX.GetDestAndAdvance() = Bary.X;
+			*OutBaryY.GetDestAndAdvance() = Bary.Y;
+			*OutBaryZ.GetDestAndAdvance() = Bary.Z;
+		}
+	}
+	else
+	{
+		for (int32 i = 0; i < Context.NumInstances; ++i)
+		{
+			*OutTri.GetDestAndAdvance() = -1;
+			*OutBaryX.GetDestAndAdvance() = 0.0f;
+			*OutBaryY.GetDestAndAdvance() = 0.0f;
+			*OutBaryZ.GetDestAndAdvance() = 0.0f;
+		}
 	}
 }
 
@@ -436,23 +470,31 @@ void UNiagaraDataInterfaceSkeletalMesh::IsValidTriCoord(FVectorVMContext& Contex
 	VectorVM::FExternalFuncInputHandler<float> BaryZParam(Context);
 
 	checkfSlow(InstData.Get(), TEXT("Skeletal Mesh Interface has invalid instance data. %s"), *GetPathName());
-	checkfSlow(InstData->Mesh, TEXT("Skeletal Mesh Interface has invalid mesh. %s"), *GetPathName());
 
 	VectorVM::FExternalFuncRegisterHandler<FNiagaraBool> OutValid(Context);
 
 	FSkeletalMeshAccessorHelper MeshAccessor;
 	MeshAccessor.Init<FilterMode, AreaWeightingMode>(InstData);
 
-	for (int32 i = 0; i < Context.NumInstances; ++i)
+	if (MeshAccessor.IsSkinAccessible())
 	{
-		int32 RequestedIndex = (TriParam.Get() * 3) + 2; // Get the last triangle index of the set
-
-		FNiagaraBool Value;
-		Value.SetValue(MeshAccessor.IndexBuffer != nullptr && MeshAccessor.IndexBuffer->Num() > RequestedIndex);
-		*OutValid.GetDest() = Value;
+		for (int32 i = 0; i < Context.NumInstances; ++i)
+		{
+			int32 RequestedIndex = (TriParam.GetAndAdvance() * 3) + 2; // Get the last triangle index of the set
 
-		OutValid.Advance();
-		BaryXParam.Advance();		BaryYParam.Advance();		BaryZParam.Advance(); TriParam.Advance();
+			FNiagaraBool Value;
+			Value.SetValue(MeshAccessor.IndexBuffer != nullptr && MeshAccessor.IndexBuffer->Num() > RequestedIndex);
+			*OutValid.GetDestAndAdvance() = Value;
+		}
+	}
+	else
+	{
+		for (int32 i = 0; i < Context.NumInstances; ++i)
+		{
+			FNiagaraBool Value;
+			Value.SetValue(false);
+			*OutValid.GetDestAndAdvance() = Value;
+		}
 	}
 }
 
@@ -472,6 +514,18 @@ void UNiagaraDataInterfaceSkeletalMesh::RandomTriangle(FVectorVMContext& Context
 	FSkeletalMeshAccessorHelper MeshAccessor;
 	MeshAccessor.Init<TIntegralConstant<int32, 0>, TIntegralConstant<int32, 0>>(InstData);
 
+	if (!MeshAccessor.IsSkinAccessible())
+	{
+		for (int32 i = 0; i < Context.NumInstances; ++i)
+		{
+			*OutTri.GetDestAndAdvance() = -1;
+			*OutBaryX.GetDestAndAdvance() = 0.0f;
+			*OutBaryY.GetDestAndAdvance() = 0.0f;
+			*OutBaryZ.GetDestAndAdvance() = 0.0f;
+		}
+		return;
+	}
+
 	//-TODO: AREA WEIGHTED
 	const int32 LODIndex = InstData->GetLODIndex();
 	const bool bAreaWeighted = InstData->Mesh->GetLODInfo(LODIndex)->bSupportUniformlyDistributedSampling;
@@ -531,7 +585,7 @@ void UNiagaraDataInterfaceSkeletalMesh::GetTriangleCount(FVectorVMContext& Conte
 	FSkeletalMeshAccessorHelper MeshAccessor;
 	MeshAccessor.Init<TIntegralConstant<int32, 0>, TIntegralConstant<int32, 0>>(InstData);
 	
-	const int32 NumTriangles = MeshAccessor.IndexBuffer->Num() / 3;
+	const int32 NumTriangles = MeshAccessor.IsSkinAccessible() ? MeshAccessor.IndexBuffer->Num() / 3 : 0;
 	for (int32 i = 0; i < Context.NumInstances; ++i)
 	{
 		*OutCount.GetDestAndAdvance() = NumTriangles;
@@ -546,9 +600,7 @@ FORCEINLINE int32 UNiagaraDataInterfaceSkeletalMesh::GetFilteredTriangleCount(FS
 }
 
 template<>
-FORCEINLINE int32 UNiagaraDataInterfaceSkeletalMesh::GetFilteredTriangleCount<
-	TIntegralConstant<ENDISkeletalMesh_FilterMode, ENDISkeletalMesh_FilterMode::None>,
-	TIntegralConstant<ENDISkelMesh_AreaWeightingMode, ENDISkelMesh_AreaWeightingMode::None>>
+FORCEINLINE int32 UNiagaraDataInterfaceSkeletalMesh::GetFilteredTriangleCount<TNDISkelMesh_FilterModeNone, TNDISkelMesh_AreaWeightingOff>
 	(FSkeletalMeshAccessorHelper& Accessor, FNDISkeletalMesh_InstanceData* InstData)
 {
 	int32 NumTris = 0;
@@ -560,9 +612,7 @@ FORCEINLINE int32 UNiagaraDataInterfaceSkeletalMesh::GetFilteredTriangleCount<
 }
 
 template<>
-FORCEINLINE int32 UNiagaraDataInterfaceSkeletalMesh::GetFilteredTriangleCount<
-	TIntegralConstant<ENDISkeletalMesh_FilterMode, ENDISkeletalMesh_FilterMode::None>,
-	TIntegralConstant<ENDISkelMesh_AreaWeightingMode, ENDISkelMesh_AreaWeightingMode::AreaWeighted>>
+FORCEINLINE int32 UNiagaraDataInterfaceSkeletalMesh::GetFilteredTriangleCount<TNDISkelMesh_FilterModeNone, TNDISkelMesh_AreaWeightingOn>
 	(FSkeletalMeshAccessorHelper& Accessor, FNDISkeletalMesh_InstanceData* InstData)
 {
 	const FSkeletalMeshSamplingInfo& SamplingInfo = InstData->Mesh->GetSamplingInfo();
@@ -571,27 +621,21 @@ FORCEINLINE int32 UNiagaraDataInterfaceSkeletalMesh::GetFilteredTriangleCount<
 }
 
 template<>
-FORCEINLINE int32 UNiagaraDataInterfaceSkeletalMesh::GetFilteredTriangleCount<
-	TIntegralConstant<ENDISkeletalMesh_FilterMode, ENDISkeletalMesh_FilterMode::SingleRegion>,
-	TIntegralConstant<ENDISkelMesh_AreaWeightingMode, ENDISkelMesh_AreaWeightingMode::None>>
+FORCEINLINE int32 UNiagaraDataInterfaceSkeletalMesh::GetFilteredTriangleCount<TNDISkelMesh_FilterModeSingle, TNDISkelMesh_AreaWeightingOff>
 	(FSkeletalMeshAccessorHelper& Accessor, FNDISkeletalMesh_InstanceData* InstData)
 {
 	return Accessor.SamplingRegionBuiltData->TriangleIndices.Num();
 }
 
 template<>
-FORCEINLINE int32 UNiagaraDataInterfaceSkeletalMesh::GetFilteredTriangleCount<
-	TIntegralConstant<ENDISkeletalMesh_FilterMode, ENDISkeletalMesh_FilterMode::SingleRegion>,
-	TIntegralConstant<ENDISkelMesh_AreaWeightingMode, ENDISkelMesh_AreaWeightingMode::AreaWeighted>>
+FORCEINLINE int32 UNiagaraDataInterfaceSkeletalMesh::GetFilteredTriangleCount<TNDISkelMesh_FilterModeSingle, TNDISkelMesh_AreaWeightingOn>
 	(FSkeletalMeshAccessorHelper& Accessor, FNDISkeletalMesh_InstanceData* InstData)
 {
 	return Accessor.SamplingRegionBuiltData->AreaWeightedSampler.GetNumEntries();
 }
 
 template<>
-FORCEINLINE int32 UNiagaraDataInterfaceSkeletalMesh::GetFilteredTriangleCount<
-	TIntegralConstant<ENDISkeletalMesh_FilterMode, ENDISkeletalMesh_FilterMode::MultiRegion>,
-	TIntegralConstant<ENDISkelMesh_AreaWeightingMode, ENDISkelMesh_AreaWeightingMode::None>>
+FORCEINLINE int32 UNiagaraDataInterfaceSkeletalMesh::GetFilteredTriangleCount<TNDISkelMesh_FilterModeMulti, TNDISkelMesh_AreaWeightingOff>
 	(FSkeletalMeshAccessorHelper& Accessor, FNDISkeletalMesh_InstanceData* InstData)
 {
 	int32 NumTris = 0;
@@ -607,9 +651,7 @@ FORCEINLINE int32 UNiagaraDataInterfaceSkeletalMesh::GetFilteredTriangleCount<
 }
 
 template<>
-FORCEINLINE int32 UNiagaraDataInterfaceSkeletalMesh::GetFilteredTriangleCount<
-	TIntegralConstant<ENDISkeletalMesh_FilterMode, ENDISkeletalMesh_FilterMode::MultiRegion>,
-	TIntegralConstant<ENDISkelMesh_AreaWeightingMode, ENDISkelMesh_AreaWeightingMode::AreaWeighted>>
+FORCEINLINE int32 UNiagaraDataInterfaceSkeletalMesh::GetFilteredTriangleCount<TNDISkelMesh_FilterModeMulti, TNDISkelMesh_AreaWeightingOn>
 	(FSkeletalMeshAccessorHelper& Accessor, FNDISkeletalMesh_InstanceData* InstData)
 {
 	int32 NumTris = 0;
@@ -630,18 +672,16 @@ void UNiagaraDataInterfaceSkeletalMesh::GetFilteredTriangleCount(FVectorVMContex
 	SCOPE_CYCLE_COUNTER(STAT_NiagaraSkel_Sample);
 	VectorVM::FUserPtrHandler<FNDISkeletalMesh_InstanceData> InstData(Context);
 	checkfSlow(InstData.Get(), TEXT("Skeletal Mesh Interface has invalid instance data. %s"), *GetPathName());
-	checkfSlow(InstData->Mesh, TEXT("Skeletal Mesh Interface has invalid mesh. %s"), *GetPathName());
 
 	VectorVM::FExternalFuncRegisterHandler<int32> OutTri(Context);
 
 	FSkeletalMeshAccessorHelper MeshAccessor;
 	MeshAccessor.Init<FilterMode, AreaWeightingMode>(InstData);
 
-	int32 Count = GetFilteredTriangleCount<FilterMode, AreaWeightingMode>(MeshAccessor, InstData);
+	int32 Count = MeshAccessor.IsSkinAccessible() ? GetFilteredTriangleCount<FilterMode, AreaWeightingMode>(MeshAccessor, InstData) : 0;
 	for (int32 i = 0; i < Context.NumInstances; ++i)
 	{
-		*OutTri.GetDest() = Count;
-		OutTri.Advance();
+		*OutTri.GetDestAndAdvance() = Count;
 	}
 }
 
@@ -655,9 +695,7 @@ FORCEINLINE int32 UNiagaraDataInterfaceSkeletalMesh::GetFilteredTriangleAt(FSkel
 }
 
 template<>
-FORCEINLINE int32 UNiagaraDataInterfaceSkeletalMesh::GetFilteredTriangleAt<
-	TIntegralConstant<ENDISkeletalMesh_FilterMode, ENDISkeletalMesh_FilterMode::None>,
-	TIntegralConstant<ENDISkelMesh_AreaWeightingMode, ENDISkelMesh_AreaWeightingMode::None>>
+FORCEINLINE int32 UNiagaraDataInterfaceSkeletalMesh::GetFilteredTriangleAt<TNDISkelMesh_FilterModeNone, TNDISkelMesh_AreaWeightingOff>
 	(FSkeletalMeshAccessorHelper& Accessor, FNDISkeletalMesh_InstanceData* InstData, int32 FilteredIndex)
 {
 	for (int32 i = 0; i < Accessor.LODData->RenderSections.Num(); i++)
@@ -673,9 +711,7 @@ FORCEINLINE int32 UNiagaraDataInterfaceSkeletalMesh::GetFilteredTriangleAt<
 }
 
 template<>
-FORCEINLINE int32 UNiagaraDataInterfaceSkeletalMesh::GetFilteredTriangleAt<
-	TIntegralConstant<ENDISkeletalMesh_FilterMode, ENDISkeletalMesh_FilterMode::None>,
-	TIntegralConstant<ENDISkelMesh_AreaWeightingMode, ENDISkelMesh_AreaWeightingMode::AreaWeighted>>
+FORCEINLINE int32 UNiagaraDataInterfaceSkeletalMesh::GetFilteredTriangleAt<TNDISkelMesh_FilterModeNone, TNDISkelMesh_AreaWeightingOn>
 	(FSkeletalMeshAccessorHelper& Accessor, FNDISkeletalMesh_InstanceData* InstData, int32 FilteredIndex)
 {
 	int32 TriIdx = FilteredIndex;
@@ -683,9 +719,7 @@ FORCEINLINE int32 UNiagaraDataInterfaceSkeletalMesh::GetFilteredTriangleAt<
 }
 
 template<>
-FORCEINLINE int32 UNiagaraDataInterfaceSkeletalMesh::GetFilteredTriangleAt<
-	TIntegralConstant<ENDISkeletalMesh_FilterMode, ENDISkeletalMesh_FilterMode::SingleRegion>,
-	TIntegralConstant<ENDISkelMesh_AreaWeightingMode, ENDISkelMesh_AreaWeightingMode::None>>
+FORCEINLINE int32 UNiagaraDataInterfaceSkeletalMesh::GetFilteredTriangleAt<TNDISkelMesh_FilterModeSingle, TNDISkelMesh_AreaWeightingOff>
 	(FSkeletalMeshAccessorHelper& Accessor, FNDISkeletalMesh_InstanceData* InstData, int32 FilteredIndex)
 {
 	int32 MaxIdx = Accessor.SamplingRegionBuiltData->TriangleIndices.Num() - 1;
@@ -694,9 +728,7 @@ FORCEINLINE int32 UNiagaraDataInterfaceSkeletalMesh::GetFilteredTriangleAt<
 }
 
 template<>
-FORCEINLINE int32 UNiagaraDataInterfaceSkeletalMesh::GetFilteredTriangleAt<
-	TIntegralConstant<ENDISkeletalMesh_FilterMode, ENDISkeletalMesh_FilterMode::SingleRegion>,
-	TIntegralConstant<ENDISkelMesh_AreaWeightingMode, ENDISkelMesh_AreaWeightingMode::AreaWeighted>>
+FORCEINLINE int32 UNiagaraDataInterfaceSkeletalMesh::GetFilteredTriangleAt<TNDISkelMesh_FilterModeSingle, TNDISkelMesh_AreaWeightingOn>
 	(FSkeletalMeshAccessorHelper& Accessor, FNDISkeletalMesh_InstanceData* InstData, int32 FilteredIndex)
 {
 	int32 Idx = FilteredIndex;
@@ -707,9 +739,7 @@ FORCEINLINE int32 UNiagaraDataInterfaceSkeletalMesh::GetFilteredTriangleAt<
 }
 
 template<>
-FORCEINLINE int32 UNiagaraDataInterfaceSkeletalMesh::GetFilteredTriangleAt<
-	TIntegralConstant<ENDISkeletalMesh_FilterMode, ENDISkeletalMesh_FilterMode::MultiRegion>,
-	TIntegralConstant<ENDISkelMesh_AreaWeightingMode, ENDISkelMesh_AreaWeightingMode::None>>
+FORCEINLINE int32 UNiagaraDataInterfaceSkeletalMesh::GetFilteredTriangleAt<TNDISkelMesh_FilterModeMulti, TNDISkelMesh_AreaWeightingOff>
 	(FSkeletalMeshAccessorHelper& Accessor, FNDISkeletalMesh_InstanceData* InstData, int32 FilteredIndex)
 {
 	for (int32 RegionIdx = 0; RegionIdx < InstData->SamplingRegionIndices.Num(); RegionIdx++)
@@ -728,9 +758,7 @@ FORCEINLINE int32 UNiagaraDataInterfaceSkeletalMesh::GetFilteredTriangleAt<
 }
 
 template<>
-FORCEINLINE int32 UNiagaraDataInterfaceSkeletalMesh::GetFilteredTriangleAt<
-	TIntegralConstant<ENDISkeletalMesh_FilterMode, ENDISkeletalMesh_FilterMode::MultiRegion>,
-	TIntegralConstant<ENDISkelMesh_AreaWeightingMode, ENDISkelMesh_AreaWeightingMode::AreaWeighted>>
+FORCEINLINE int32 UNiagaraDataInterfaceSkeletalMesh::GetFilteredTriangleAt<TNDISkelMesh_FilterModeMulti, TNDISkelMesh_AreaWeightingOn>
 	(FSkeletalMeshAccessorHelper& Accessor, FNDISkeletalMesh_InstanceData* InstData, int32 FilteredIndex)
 {
 	for (int32 RegionIdx = 0; RegionIdx < InstData->SamplingRegionIndices.Num(); RegionIdx++)
@@ -755,29 +783,42 @@ void UNiagaraDataInterfaceSkeletalMesh::GetFilteredTriangleAt(FVectorVMContext&
 	VectorVM::FUserPtrHandler<FNDISkeletalMesh_InstanceData> InstData(Context);
 	VectorVM::FExternalFuncInputHandler<int32> TriParam(Context);
 	checkfSlow(InstData.Get(), TEXT("Skeletal Mesh Interface has invalid instance data. %s"), *GetPathName());
-	checkfSlow(InstData->Mesh, TEXT("Skeletal Mesh Interface has invalid mesh. %s"), *GetPathName());
+
 	VectorVM::FExternalFuncRegisterHandler<int32> OutTri(Context);
-	VectorVM::FExternalFuncRegisterHandler<float> OutBaryX(Context);	VectorVM::FExternalFuncRegisterHandler<float> OutBaryY(Context);	VectorVM::FExternalFuncRegisterHandler<float> OutBaryZ(Context);
+	VectorVM::FExternalFuncRegisterHandler<float> OutBaryX(Context);
+	VectorVM::FExternalFuncRegisterHandler<float> OutBaryY(Context);
+	VectorVM::FExternalFuncRegisterHandler<float> OutBaryZ(Context);
 
 	FSkeletalMeshAccessorHelper Accessor;
 	Accessor.Init<FilterMode, AreaWeightingMode>(InstData);
 
-	for (int32 i = 0; i < Context.NumInstances; ++i)
+	if (Accessor.IsSkinAccessible())
 	{
-		int32 Tri = TriParam.Get();
-		int32 RealIdx = 0;
-		RealIdx = GetFilteredTriangleAt<FilterMode, AreaWeightingMode>(Accessor, InstData, Tri);
-
-		const int32 TriMax = (Accessor.IndexBuffer->Num() / 3) - 1;
-		RealIdx = FMath::Clamp(RealIdx, 0, TriMax);
-
-		*OutTri.GetDest() = RealIdx;
-		float Coord = 1.0f / 3.0f;
-		*OutBaryX.GetDest() = Coord;		*OutBaryY.GetDest() = Coord;		*OutBaryZ.GetDest() = Coord;
-
-		TriParam.Advance();
-		OutTri.Advance();
-		OutBaryX.Advance();		OutBaryY.Advance();		OutBaryZ.Advance();
+		for (int32 i = 0; i < Context.NumInstances; ++i)
+		{
+			int32 Tri = TriParam.GetAndAdvance();
+			int32 RealIdx = 0;
+			RealIdx = GetFilteredTriangleAt<FilterMode, AreaWeightingMode>(Accessor, InstData, Tri);
+
+			const int32 TriMax = (Accessor.IndexBuffer->Num() / 3) - 1;
+			RealIdx = FMath::Clamp(RealIdx, 0, TriMax);
+
+			*OutTri.GetDestAndAdvance() = RealIdx;
+			float Coord = 1.0f / 3.0f;
+			*OutBaryX.GetDestAndAdvance() = Coord;
+			*OutBaryY.GetDestAndAdvance() = Coord;
+			*OutBaryZ.GetDestAndAdvance() = Coord;
+		}
+	}
+	else
+	{
+		for (int32 i = 0; i < Context.NumInstances; ++i)
+		{
+			*OutTri.GetDestAndAdvance() = -1;
+			*OutBaryX.GetDestAndAdvance() = 0.0f;
+			*OutBaryY.GetDestAndAdvance() = 0.0f;
+			*OutBaryZ.GetDestAndAdvance() = 0.0f;
+		}
 	}
 }
 
@@ -797,11 +838,13 @@ void UNiagaraDataInterfaceSkeletalMesh::GetTriCoordColor(FVectorVMContext& Conte
 
 	USkeletalMeshComponent* Comp = Cast<USkeletalMeshComponent>(InstData->Component.Get());
 	FSkinWeightVertexBuffer* SkinWeightBuffer;
-	FSkeletalMeshLODRenderData& LODData = InstData->GetLODRenderDataAndSkinWeights(SkinWeightBuffer);
-	const FColorVertexBuffer& Colors = LODData.StaticVertexBuffers.ColorVertexBuffer;
+	FSkeletalMeshLODRenderData* LODData = InstData->GetLODRenderDataAndSkinWeights(SkinWeightBuffer);
+	check(LODData);
+
+	const FColorVertexBuffer& Colors = LODData->StaticVertexBuffers.ColorVertexBuffer;
 	checkfSlow(Colors.GetNumVertices() != 0, TEXT("Trying to access vertex colors from mesh without any."));
 
-	FMultiSizeIndexContainer& Indices = LODData.MultiSizeIndexContainer;
+	FMultiSizeIndexContainer& Indices = LODData->MultiSizeIndexContainer;
 	const FRawStaticIndexBuffer16or32Interface* IndexBuffer = Indices.GetIndexBuffer();
 	const int32 TriMax = (IndexBuffer->Num() / 3) - 1;
 	for (int32 i = 0; i < Context.NumInstances; ++i)
@@ -871,12 +914,13 @@ void UNiagaraDataInterfaceSkeletalMesh::GetTriCoordUV(FVectorVMContext& Context)
 
 	USkeletalMeshComponent* Comp = Cast<USkeletalMeshComponent>(InstData->Component.Get());
 	FSkinWeightVertexBuffer* SkinWeightBuffer;
-	FSkeletalMeshLODRenderData& LODData = InstData->GetLODRenderDataAndSkinWeights(SkinWeightBuffer);
+	FSkeletalMeshLODRenderData* LODData = InstData->GetLODRenderDataAndSkinWeights(SkinWeightBuffer);
+	check(LODData);
 
-	FMultiSizeIndexContainer& Indices = LODData.MultiSizeIndexContainer;
+	FMultiSizeIndexContainer& Indices = LODData->MultiSizeIndexContainer;
 	FRawStaticIndexBuffer16or32Interface* IndexBuffer = Indices.GetIndexBuffer();
 	const int32 TriMax = (IndexBuffer->Num() / 3) - 1;
-	const int32 UVSetMax = LODData.StaticVertexBuffers.StaticMeshVertexBuffer.GetNumTexCoords() - 1;
+	const int32 UVSetMax = LODData->StaticVertexBuffers.StaticMeshVertexBuffer.GetNumTexCoords() - 1;
 	const float InvDt = 1.0f / InstData->DeltaSeconds;
 	for (int32 i = 0; i < Context.NumInstances; ++i)
 	{
@@ -905,9 +949,30 @@ void UNiagaraDataInterfaceSkeletalMesh::GetTriCoordUV(FVectorVMContext& Context)
 	}
 }
 
-struct FGetTriCoodSkinnedDataOutputHandler
+// Stub specialization for no valid mesh data on the data interface
+template<>
+void UNiagaraDataInterfaceSkeletalMesh::GetTriCoordUV<FSkelMeshVertexAccessorNoop>(FVectorVMContext& Context)
+{
+	VectorVM::FExternalFuncInputHandler<int32> TriParam(Context);
+	VectorVM::FExternalFuncInputHandler<float> BaryXParam(Context);
+	VectorVM::FExternalFuncInputHandler<float> BaryYParam(Context);
+	VectorVM::FExternalFuncInputHandler<float> BaryZParam(Context);
+	VectorVM::FExternalFuncInputHandler<int32> UVSetParam(Context);
+	VectorVM::FUserPtrHandler<FNDISkeletalMesh_InstanceData> InstData(Context);
+
+	VectorVM::FExternalFuncRegisterHandler<float> OutUVX(Context);
+	VectorVM::FExternalFuncRegisterHandler<float> OutUVY(Context);
+
+	for (int32 i = 0; i < Context.NumInstances; ++i)
+	{
+		*OutUVX.GetDestAndAdvance() = 0.0f;
+		*OutUVY.GetDestAndAdvance() = 0.0f;
+	}
+}
+
+struct FGetTriCoordSkinnedDataOutputHandler
 {
-	FGetTriCoodSkinnedDataOutputHandler(FVectorVMContext& Context)
+	FGetTriCoordSkinnedDataOutputHandler(FVectorVMContext& Context)
 		: PosX(Context), PosY(Context), PosZ(Context)
 		, VelX(Context), VelY(Context), VelZ(Context)
 		, NormX(Context), NormY(Context), NormZ(Context)
@@ -998,13 +1063,15 @@ void UNiagaraDataInterfaceSkeletalMesh::GetTriCoordSkinnedData(FVectorVMContext&
 	const FMatrix& Transform = InstData->Transform;
 	const FMatrix& PrevTransform = InstData->PrevTransform;
 
-	FGetTriCoodSkinnedDataOutputHandler Output(Context);
-
-	FSkinWeightVertexBuffer* SkinWeightBuffer;
-	FSkeletalMeshLODRenderData& LODData = InstData->GetLODRenderDataAndSkinWeights(SkinWeightBuffer);
+	FGetTriCoordSkinnedDataOutputHandler Output(Context);
 
 	FSkeletalMeshAccessorHelper Accessor;
-	Accessor.Init<TIntegralConstant<ENDISkeletalMesh_FilterMode, ENDISkeletalMesh_FilterMode::None>, TIntegralConstant<ENDISkelMesh_AreaWeightingMode, ENDISkelMesh_AreaWeightingMode::None>>(InstData);
+	Accessor.Init<TNDISkelMesh_FilterModeNone, TNDISkelMesh_AreaWeightingOff>(InstData);
+
+	check(Accessor.IsSkinAccessible()); // supposed to use the fallback for invalid mesh
+
+	FSkeletalMeshLODRenderData* LODData = Accessor.LODData;
+
 	const int32 TriMax = (Accessor.IndexBuffer->Num() / 3) - 1;
 	const float InvDt = 1.0f / InstData->DeltaSeconds;
 
@@ -1176,6 +1243,87 @@ void UNiagaraDataInterfaceSkeletalMesh::GetTriCoordSkinnedData(FVectorVMContext&
 	}
 }
 
+// Fallback sampling function for no valid mesh on the interface
+template<typename TransformHandlerType, typename bInterpolated>
+void UNiagaraDataInterfaceSkeletalMesh::GetTriCoordSkinnedDataFallback(FVectorVMContext& Context)
+{
+	SCOPE_CYCLE_COUNTER(STAT_NiagaraSkel_Sample);
+	TransformHandlerType TransformHandler;
+
+	VectorVM::FUserPtrHandler<FNDISkeletalMesh_InstanceData> InstData(Context);	
+	VectorVM::FExternalFuncInputHandler<int32> TriParam(Context);
+	VectorVM::FExternalFuncInputHandler<float> BaryXParam(Context);
+	VectorVM::FExternalFuncInputHandler<float> BaryYParam(Context);
+	VectorVM::FExternalFuncInputHandler<float> BaryZParam(Context);
+	VectorVM::FExternalFuncInputHandler<float> InterpParam;
+
+	if (bInterpolated::Value)
+	{
+		InterpParam.Init(Context);
+	}
+
+	checkfSlow(InstData.Get(), TEXT("Skeletal Mesh Interface has invalid instance data. %s"), *GetPathName());
+
+	//TODO: Replace this by storing off FTransforms and doing a proper lerp to get a final transform.
+	//Also need to pull in a per particle interpolation factor.
+	const FMatrix& Transform = InstData->Transform;
+	const FMatrix& PrevTransform = InstData->PrevTransform;
+
+	FGetTriCoordSkinnedDataOutputHandler Output(Context);
+	bool bNeedsPrev = bInterpolated::Value || Output.bNeedsVelocity;
+
+	const float InvDt = 1.0f / InstData->DeltaSeconds;
+
+	for (int32 i = 0; i < Context.NumInstances; ++i)
+	{
+		float Interp = 1.0f;
+		if (bInterpolated::Value)
+		{
+			Interp = InterpParam.GetAndAdvance();
+		}
+
+		FVector Prev(0.0f);
+		FVector Pos(0.0f);
+		if (bNeedsPrev)
+		{
+			TransformHandler.TransformPosition(Prev, PrevTransform);
+		}
+
+		if (Output.bNeedsPosition || Output.bNeedsVelocity)
+		{
+			TransformHandler.TransformPosition(Pos, Transform);
+
+			if (bInterpolated::Value)
+			{
+				Pos = FMath::Lerp(Prev, Pos, Interp);
+			}
+
+			Output.SetPosition(Pos);
+		}
+
+		if (Output.bNeedsVelocity)
+		{			
+			FVector Velocity = (Pos - Prev) * InvDt;
+			Output.SetVelocity(Velocity);
+		}
+
+		if (Output.bNeedsNorm)
+		{
+			Output.SetNormal(FVector(0.0f, 0.0f, 1.0f));
+		}
+
+		if (Output.bNeedsBinorm)
+		{
+			Output.SetBinormal(FVector(0.0f, 1.0f, 0.0f));
+		}
+		
+		if (Output.bNeedsTangent)
+		{
+			Output.SetTangent(FVector(1.0f, 0.0f, 0.0f));
+		}		
+	}
+}
+
 template<typename SkinningHandlerType>
 void UNiagaraDataInterfaceSkeletalMesh::GetTriCoordVertices(FVectorVMContext& Context)
 {
@@ -1194,10 +1342,9 @@ void UNiagaraDataInterfaceSkeletalMesh::GetTriCoordVertices(FVectorVMContext& Co
 
 	int32 Idx0; int32 Idx1; int32 Idx2;
 	FSkeletalMeshAccessorHelper Accessor;
-	Accessor.Init<TIntegralConstant<int32, 0>, TIntegralConstant<int32, 0>>(InstData);
-
-	const int32 TriMax = (Accessor.IndexBuffer->Num() / 3) - 1;
+	Accessor.Init<TNDISkelMesh_FilterModeNone, TNDISkelMesh_AreaWeightingOff>(InstData);
 
+	const int32 TriMax = Accessor.IsSkinAccessible() ? (Accessor.IndexBuffer->Num() / 3) - 1 : 0;
 	for (int32 i = 0; i < Context.NumInstances; ++i)
 	{
 		const int32 Tri = FMath::Clamp(TriParam.GetAndAdvance(), 0, TriMax);
diff --git a/Engine/Plugins/FX/Niagara/Source/Niagara/Private/NDISkeletalMesh_VertexSampling.cpp b/Engine/Plugins/FX/Niagara/Source/Niagara/Private/NDISkeletalMesh_VertexSampling.cpp
--- a/Engine/Plugins/FX/Niagara/Source/Niagara/Private/NDISkeletalMesh_VertexSampling.cpp
+++ b/Engine/Plugins/FX/Niagara/Source/Niagara/Private/NDISkeletalMesh_VertexSampling.cpp
@@ -251,9 +251,9 @@ void UNiagaraDataInterfaceSkeletalMesh::IsValidVertex(FVectorVMContext& Context)
 	VectorVM::FExternalFuncRegisterHandler<FNiagaraBool> OutValid(Context);
 
 	FSkeletalMeshAccessorHelper MeshAccessor;
-	MeshAccessor.Init<TIntegralConstant<ENDISkeletalMesh_FilterMode, ENDISkeletalMesh_FilterMode::None>, TIntegralConstant<ENDISkelMesh_AreaWeightingMode, ENDISkelMesh_AreaWeightingMode::None>>(InstData);
+	MeshAccessor.Init<TNDISkelMesh_FilterModeNone, TNDISkelMesh_AreaWeightingOff>(InstData);
 
-	const int32 MaxVertex = MeshAccessor.LODData->GetNumVertices();
+	const int32 MaxVertex = MeshAccessor.IsSkinAccessible() ? MeshAccessor.LODData->GetNumVertices() : 0;
 	for (int32 i = 0; i < Context.NumInstances; ++i)
 	{
 		const int32 VertexIndex = VertexParam.GetAndAdvance();
@@ -273,9 +273,9 @@ void UNiagaraDataInterfaceSkeletalMesh::RandomVertex(FVectorVMContext& Context)
 	VectorVM::FExternalFuncRegisterHandler<int32> OutVertex(Context);
 
 	FSkeletalMeshAccessorHelper MeshAccessor;
-	MeshAccessor.Init<TIntegralConstant<ENDISkeletalMesh_FilterMode, ENDISkeletalMesh_FilterMode::None>, TIntegralConstant<ENDISkelMesh_AreaWeightingMode, ENDISkelMesh_AreaWeightingMode::None>>(InstData);
+	MeshAccessor.Init<TNDISkelMesh_FilterModeNone, TNDISkelMesh_AreaWeightingOff>(InstData);
 
-	const int32 MaxVertex = MeshAccessor.LODData->GetNumVertices() - 1;
+	const int32 MaxVertex = MeshAccessor.IsSkinAccessible() ? MeshAccessor.LODData->GetNumVertices() - 1 : -1;
 	if (MaxVertex >= 0)
 	{
 		for (int32 i = 0; i < Context.NumInstances; ++i)
@@ -298,9 +298,9 @@ void UNiagaraDataInterfaceSkeletalMesh::GetVertexCount(FVectorVMContext& Context
 	VectorVM::FExternalFuncRegisterHandler<int32> OutVertexCount(Context);
 
 	FSkeletalMeshAccessorHelper MeshAccessor;
-	MeshAccessor.Init<TIntegralConstant<ENDISkeletalMesh_FilterMode, ENDISkeletalMesh_FilterMode::None>, TIntegralConstant<ENDISkelMesh_AreaWeightingMode, ENDISkelMesh_AreaWeightingMode::None>>(InstData);
+	MeshAccessor.Init<TNDISkelMesh_FilterModeNone, TNDISkelMesh_AreaWeightingOff>(InstData);
 
-	const int32 MaxVertex = MeshAccessor.LODData->GetNumVertices();
+	const int32 MaxVertex = MeshAccessor.IsSkinAccessible() ? MeshAccessor.LODData->GetNumVertices() : 0;
 	for (int32 i = 0; i < Context.NumInstances; ++i)
 	{
 		*OutVertexCount.GetDestAndAdvance() = MaxVertex;
@@ -349,18 +349,27 @@ void UNiagaraDataInterfaceSkeletalMesh::RandomFilteredVertex(FVectorVMContext& C
 
 	VectorVM::FUserPtrHandler<FNDISkeletalMesh_InstanceData> InstData(Context);
 	FNDIRandomHelper RandHelper(Context);
-	checkfSlow(InstData.Get(), TEXT("Skeletal Mesh Interface has invalid instance data. %s"), *GetPathName());
-	checkfSlow(InstData->Mesh, TEXT("Skeletal Mesh Interface has invalid mesh. %s"), *GetPathName());
+	checkfSlow(InstData.Get(), TEXT("Skeletal Mesh Interface has invalid instance data. %s"), *GetPathName());	
 
 	VectorVM::FExternalFuncRegisterHandler<int32> OutVert(Context);
 
 	FSkeletalMeshAccessorHelper MeshAccessor;
 	MeshAccessor.Init<FilterMode, TIntegralConstant<ENDISkelMesh_AreaWeightingMode, ENDISkelMesh_AreaWeightingMode::None>>(InstData);
 
-	for (int32 i = 0; i < Context.NumInstances; ++i)
+	if (MeshAccessor.IsSkinAccessible())
 	{
-		RandHelper.GetAndAdvance();
-		*OutVert.GetDestAndAdvance() = RandomFilteredVertIndex<FilterMode>(RandHelper, i, MeshAccessor, InstData);
+		for (int32 i = 0; i < Context.NumInstances; ++i)
+		{
+			RandHelper.GetAndAdvance();
+			*OutVert.GetDestAndAdvance() = RandomFilteredVertIndex<FilterMode>(RandHelper, i, MeshAccessor, InstData);
+		}
+	}
+	else
+	{
+		for (int32 i = 0; i < Context.NumInstances; ++i)
+		{
+			*OutVert.GetDestAndAdvance() = -1;
+		}
 	}
 }
 
@@ -373,19 +382,18 @@ void UNiagaraDataInterfaceSkeletalMesh::IsValidFilteredVertex(FVectorVMContext&
 	VectorVM::FExternalFuncInputHandler<int32> VertexParam(Context);
 
 	checkfSlow(InstData.Get(), TEXT("Skeletal Mesh Interface has invalid instance data. %s"), *GetPathName());
-	checkfSlow(InstData->Mesh, TEXT("Skeletal Mesh Interface has invalid mesh. %s"), *GetPathName());
 
 	VectorVM::FExternalFuncRegisterHandler<FNiagaraBool> OutValid(Context);
 
 	FSkeletalMeshAccessorHelper MeshAccessor;
-	MeshAccessor.Init<FilterMode, TIntegralConstant<ENDISkelMesh_AreaWeightingMode, ENDISkelMesh_AreaWeightingMode::None>>(InstData);
+	MeshAccessor.Init<FilterMode, TNDISkelMesh_AreaWeightingOff>(InstData);
 
 	for (int32 i = 0; i < Context.NumInstances; ++i)
 	{
 		int32 RequestedIndex = VertexParam.Get();
 
 		FNiagaraBool Value;
-		Value.SetValue((int32)MeshAccessor.LODData->GetNumVertices() > RequestedIndex);
+		Value.SetValue(MeshAccessor.IsSkinAccessible() && (int32)MeshAccessor.LODData->GetNumVertices() > RequestedIndex);
 		*OutValid.GetDest() = Value;
 
 		OutValid.Advance();
@@ -403,21 +411,21 @@ FORCEINLINE int32 UNiagaraDataInterfaceSkeletalMesh::GetFilteredVertexCount(FSke
 }
 
 template<>
-FORCEINLINE_DEBUGGABLE int32 UNiagaraDataInterfaceSkeletalMesh::GetFilteredVertexCount<TIntegralConstant<ENDISkeletalMesh_FilterMode, ENDISkeletalMesh_FilterMode::None>>
+FORCEINLINE_DEBUGGABLE int32 UNiagaraDataInterfaceSkeletalMesh::GetFilteredVertexCount<TNDISkelMesh_FilterModeNone>
 	(FSkeletalMeshAccessorHelper& Accessor, FNDISkeletalMesh_InstanceData* InstData)
 {
 	return  Accessor.LODData->GetNumVertices();;
 }
 
 template<>
-FORCEINLINE_DEBUGGABLE int32 UNiagaraDataInterfaceSkeletalMesh::GetFilteredVertexCount<TIntegralConstant<ENDISkeletalMesh_FilterMode, ENDISkeletalMesh_FilterMode::SingleRegion>>
+FORCEINLINE_DEBUGGABLE int32 UNiagaraDataInterfaceSkeletalMesh::GetFilteredVertexCount<TNDISkelMesh_FilterModeSingle>
 	(FSkeletalMeshAccessorHelper& Accessor, FNDISkeletalMesh_InstanceData* InstData)
 {
 	return Accessor.SamplingRegionBuiltData->Vertices.Num();
 }
 
 template<>
-FORCEINLINE_DEBUGGABLE int32 UNiagaraDataInterfaceSkeletalMesh::GetFilteredVertexCount<TIntegralConstant<ENDISkeletalMesh_FilterMode, ENDISkeletalMesh_FilterMode::MultiRegion>>
+FORCEINLINE_DEBUGGABLE int32 UNiagaraDataInterfaceSkeletalMesh::GetFilteredVertexCount<TNDISkelMesh_FilterModeMulti>
 	(FSkeletalMeshAccessorHelper& Accessor, FNDISkeletalMesh_InstanceData* InstData)
 {
 	int32 NumVerts = 0;
@@ -437,17 +445,26 @@ void UNiagaraDataInterfaceSkeletalMesh::GetFilteredVertexCount(FVectorVMContext&
 	SCOPE_CYCLE_COUNTER(STAT_NiagaraSkel_Vertex_Sample);
 	VectorVM::FUserPtrHandler<FNDISkeletalMesh_InstanceData> InstData(Context);
 	checkfSlow(InstData.Get(), TEXT("Skeletal Mesh Interface has invalid instance data. %s"), *GetPathName());
-	checkfSlow(InstData->Mesh, TEXT("Skeletal Mesh Interface has invalid mesh. %s"), *GetPathName());
 
 	VectorVM::FExternalFuncRegisterHandler<int32> OutVert(Context);
 
 	FSkeletalMeshAccessorHelper MeshAccessor;
-	MeshAccessor.Init<FilterMode, TIntegralConstant<ENDISkelMesh_AreaWeightingMode, ENDISkelMesh_AreaWeightingMode::None>>(InstData);
+	MeshAccessor.Init<FilterMode, TNDISkelMesh_AreaWeightingOff>(InstData);
 
-	int32 Count = GetFilteredVertexCount<FilterMode>(MeshAccessor, InstData);
-	for (int32 i = 0; i < Context.NumInstances; ++i)
+	if (MeshAccessor.IsSkinAccessible())
+	{
+		int32 Count = GetFilteredVertexCount<FilterMode>(MeshAccessor, InstData);
+		for (int32 i = 0; i < Context.NumInstances; ++i)
+		{
+			*OutVert.GetDestAndAdvance() = Count;
+		}
+	}
+	else
 	{
-		*OutVert.GetDestAndAdvance() = Count;
+		for (int32 i = 0; i < Context.NumInstances; ++i)
+		{
+			*OutVert.GetDestAndAdvance() = 0;
+		}
 	}
 }
 
@@ -461,14 +478,14 @@ FORCEINLINE int32 UNiagaraDataInterfaceSkeletalMesh::GetFilteredVertexAt(FSkelet
 }
 
 template<>
-FORCEINLINE_DEBUGGABLE int32 UNiagaraDataInterfaceSkeletalMesh::GetFilteredVertexAt<TIntegralConstant<ENDISkeletalMesh_FilterMode, ENDISkeletalMesh_FilterMode::None>>
+FORCEINLINE_DEBUGGABLE int32 UNiagaraDataInterfaceSkeletalMesh::GetFilteredVertexAt<TNDISkelMesh_FilterModeNone>
 	(FSkeletalMeshAccessorHelper& Accessor, FNDISkeletalMesh_InstanceData* InstData, int32 FilteredIndex)
 {
 	return FilteredIndex;
 }
 
 template<>
-FORCEINLINE_DEBUGGABLE int32 UNiagaraDataInterfaceSkeletalMesh::GetFilteredVertexAt<TIntegralConstant<ENDISkeletalMesh_FilterMode, ENDISkeletalMesh_FilterMode::SingleRegion>>
+FORCEINLINE_DEBUGGABLE int32 UNiagaraDataInterfaceSkeletalMesh::GetFilteredVertexAt<TNDISkelMesh_FilterModeSingle>
 	(FSkeletalMeshAccessorHelper& Accessor, FNDISkeletalMesh_InstanceData* InstData, int32 FilteredIndex)
 {
 	int32 MaxIdx = Accessor.SamplingRegionBuiltData->Vertices.Num() - 1;
@@ -477,7 +494,7 @@ FORCEINLINE_DEBUGGABLE int32 UNiagaraDataInterfaceSkeletalMesh::GetFilteredVerte
 }
 
 template<>
-FORCEINLINE_DEBUGGABLE int32 UNiagaraDataInterfaceSkeletalMesh::GetFilteredVertexAt<TIntegralConstant<ENDISkeletalMesh_FilterMode, ENDISkeletalMesh_FilterMode::MultiRegion>>
+FORCEINLINE_DEBUGGABLE int32 UNiagaraDataInterfaceSkeletalMesh::GetFilteredVertexAt<TNDISkelMesh_FilterModeMulti>
 	(FSkeletalMeshAccessorHelper& Accessor, FNDISkeletalMesh_InstanceData* InstData, int32 FilteredIndex)
 {
 	for (int32 RegionIdx = 0; RegionIdx < InstData->SamplingRegionIndices.Num(); RegionIdx++)
@@ -503,21 +520,28 @@ void UNiagaraDataInterfaceSkeletalMesh::GetFilteredVertexAt(FVectorVMContext& Co
 	VectorVM::FUserPtrHandler<FNDISkeletalMesh_InstanceData> InstData(Context);
 	VectorVM::FExternalFuncInputHandler<int32> FilteredVertexParam(Context);
 	checkfSlow(InstData.Get(), TEXT("Skeletal Mesh Interface has invalid instance data. %s"), *GetPathName());
-	checkfSlow(InstData->Mesh, TEXT("Skeletal Mesh Interface has invalid mesh. %s"), *GetPathName());
+
 	VectorVM::FExternalFuncRegisterHandler<int32> OutVert(Context);
 
 	FSkeletalMeshAccessorHelper Accessor;
-	Accessor.Init<FilterMode, TIntegralConstant<ENDISkelMesh_AreaWeightingMode, ENDISkelMesh_AreaWeightingMode::None>>(InstData);
+	Accessor.Init<FilterMode, TNDISkelMesh_AreaWeightingOff>(InstData);
 
-	int32 VertMax = Accessor.LODData->GetNumVertices();
-	for (int32 i = 0; i < Context.NumInstances; ++i)
+	if (Accessor.IsSkinAccessible())
 	{
-		int32 FilteredVert = FilteredVertexParam.GetAndAdvance();
-		int32 RealIdx = 0;
-		RealIdx = GetFilteredVertexAt<FilterMode>(Accessor, InstData, FilteredVert);
-		RealIdx = FMath::Min(RealIdx, VertMax);
-
-		*OutVert.GetDestAndAdvance() = RealIdx;
+		for (int32 i = 0; i < Context.NumInstances; ++i)
+		{
+			int32 FilteredVert = FilteredVertexParam.GetAndAdvance();
+			int32 RealIdx = 0;
+			RealIdx = GetFilteredVertexAt<FilterMode>(Accessor, InstData, FilteredVert);
+			*OutVert.GetDestAndAdvance() = RealIdx;
+		}
+	}
+	else
+	{
+		for (int32 i = 0; i < Context.NumInstances; ++i)
+		{
+			*OutVert.GetDestAndAdvance() = -1;
+		}
 	}
 }
 
@@ -536,13 +560,14 @@ void UNiagaraDataInterfaceSkeletalMesh::GetVertexColor(FVectorVMContext& Context
 
 	USkeletalMeshComponent* Comp = Cast<USkeletalMeshComponent>(InstData->Component.Get());
 	FSkinWeightVertexBuffer* SkinWeightBuffer;
-	FSkeletalMeshLODRenderData& LODData = InstData->GetLODRenderDataAndSkinWeights(SkinWeightBuffer);
-	const FColorVertexBuffer& Colors = LODData.StaticVertexBuffers.ColorVertexBuffer;
+	FSkeletalMeshLODRenderData* LODData = InstData->GetLODRenderDataAndSkinWeights(SkinWeightBuffer);
+	check(LODData);
+	const FColorVertexBuffer& Colors = LODData->StaticVertexBuffers.ColorVertexBuffer;
 	checkfSlow(Colors.GetNumVertices() != 0, TEXT("Trying to access vertex colors from mesh without any."));
 
-	FMultiSizeIndexContainer& Indices = LODData.MultiSizeIndexContainer;
+	FMultiSizeIndexContainer& Indices = LODData->MultiSizeIndexContainer;
 	const FRawStaticIndexBuffer16or32Interface* IndexBuffer = Indices.GetIndexBuffer();
-	int32 VertMax = LODData.GetNumVertices();
+	int32 VertMax = LODData->GetNumVertices();
 	for (int32 i = 0; i < Context.NumInstances; ++i)
 	{
 		int32 Vertex = VertParam.GetAndAdvance();
@@ -592,11 +617,12 @@ void UNiagaraDataInterfaceSkeletalMesh::GetVertexUV(FVectorVMContext& Context)
 
 	USkeletalMeshComponent* Comp = Cast<USkeletalMeshComponent>(InstData->Component.Get());
 	FSkinWeightVertexBuffer* SkinWeightBuffer;
-	FSkeletalMeshLODRenderData& LODData = InstData->GetLODRenderDataAndSkinWeights(SkinWeightBuffer);
+	FSkeletalMeshLODRenderData* LODData = InstData->GetLODRenderDataAndSkinWeights(SkinWeightBuffer);
+	check(LODData);
 
-	FMultiSizeIndexContainer& Indices = LODData.MultiSizeIndexContainer;
+	FMultiSizeIndexContainer& Indices = LODData->MultiSizeIndexContainer;
 	FRawStaticIndexBuffer16or32Interface* IndexBuffer = Indices.GetIndexBuffer();
-	const int32 VertMax = LODData.GetNumVertices() - 1;
+	const int32 VertMax = LODData->GetNumVertices() - 1;
 	for (int32 i = 0; i < Context.NumInstances; ++i)
 	{
 		int32 Vert = VertParam.GetAndAdvance();
@@ -610,6 +636,24 @@ void UNiagaraDataInterfaceSkeletalMesh::GetVertexUV(FVectorVMContext& Context)
 	}
 }
 
+// Stub specialization for no valid mesh data on the data interface
+template<>
+void UNiagaraDataInterfaceSkeletalMesh::GetVertexUV<FSkelMeshVertexAccessorNoop>(FVectorVMContext& Context)
+{
+	VectorVM::FExternalFuncInputHandler<int32> VertParam(Context);
+	VectorVM::FExternalFuncInputHandler<int32> UVSetParam(Context);
+	VectorVM::FUserPtrHandler<FNDISkeletalMesh_InstanceData> InstData(Context);
+
+	VectorVM::FExternalFuncRegisterHandler<float> OutUVX(Context);
+	VectorVM::FExternalFuncRegisterHandler<float> OutUVY(Context);
+
+	for (int32 i = 0; i < Context.NumInstances; ++i)
+	{
+		*OutUVX.GetDestAndAdvance() = 0.0f;
+		*OutUVY.GetDestAndAdvance() = 0.0f;
+	}
+}
+
 struct FGetVertexSkinnedDataOutputHandler
 {
 	FGetVertexSkinnedDataOutputHandler(FVectorVMContext& Context)
@@ -685,22 +729,18 @@ void UNiagaraDataInterfaceSkeletalMesh::GetVertexSkinnedData(FVectorVMContext& C
 	VectorVM::FExternalFuncInputHandler<int32> VertParam(Context);
 
 	checkfSlow(InstData.Get(), TEXT("Skeletal Mesh Interface has invalid instance data. %s"), *GetPathName());
-	checkfSlow(InstData->Mesh, TEXT("Skeletal Mesh Interface has invalid mesh. %s"), *GetPathName());
 
 	const FMatrix& Transform = InstData->Transform;
 	const FMatrix& PrevTransform = InstData->PrevTransform;
 
 	FGetVertexSkinnedDataOutputHandler Output(Context);
 
-	FSkinWeightVertexBuffer* SkinWeightBuffer;
-	FSkeletalMeshLODRenderData& LODData = InstData->GetLODRenderDataAndSkinWeights(SkinWeightBuffer);
-
 	FSkeletalMeshAccessorHelper Accessor;
-	Accessor.Init<TIntegralConstant<int32, 0>, TIntegralConstant<int32, 0>>(InstData);
-	const int32 VertMax = Accessor.LODData->GetNumVertices() - 1;
+	Accessor.Init<TNDISkelMesh_FilterModeNone, TNDISkelMesh_AreaWeightingOff>(InstData);
+	const int32 VertMax = Accessor.IsSkinAccessible() ? Accessor.LODData->GetNumVertices() : 0;
 	const float InvDt = 1.0f / InstData->DeltaSeconds;
-
 	const bool bNeedsTangentBasis = Output.bNeedsTangentX || Output.bNeedsTangentY || Output.bNeedsTangentZ;
+
 	for (int32 i = 0; i < Context.NumInstances; ++i)
 	{
 		const int32 Vertex = FMath::Min(VertParam.GetAndAdvance(), VertMax);
diff --git a/Engine/Plugins/FX/Niagara/Source/Niagara/Private/NiagaraDataInterface.cpp b/Engine/Plugins/FX/Niagara/Source/Niagara/Private/NiagaraDataInterface.cpp
--- a/Engine/Plugins/FX/Niagara/Source/Niagara/Private/NiagaraDataInterface.cpp
+++ b/Engine/Plugins/FX/Niagara/Source/Niagara/Private/NiagaraDataInterface.cpp
@@ -7,6 +7,7 @@
 #include "NiagaraTypes.h"
 #include "ShaderParameterUtils.h"
 #include "NiagaraShader.h"
+#include "NiagaraComponent.h"
 
 #define LOCTEXT_NAMESPACE "NiagaraDataInterface"
 
@@ -41,6 +42,14 @@ void UNiagaraDataInterface::PostLoad()
 	SetFlags(RF_Public);
 }
 
+#if WITH_EDITOR
+void UNiagaraDataInterface::PostEditChangeProperty(struct FPropertyChangedEvent& PropertyChangedEvent)
+{
+	RefreshErrors();
+	Super::PostEditChangeProperty(PropertyChangedEvent);
+}
+#endif
+
 bool UNiagaraDataInterface::CopyTo(UNiagaraDataInterface* Destination) const 
 {
 	bool result = CopyToInternal(Destination);
@@ -79,6 +88,44 @@ bool UNiagaraDataInterface::CopyToInternal(UNiagaraDataInterface* Destination) c
 }
 
 #if WITH_EDITOR
+void UNiagaraDataInterface::GetFeedback(UNiagaraSystem* InAsset, UNiagaraComponent* InComponent, TArray<FNiagaraDataInterfaceError>& OutErrors, TArray<FNiagaraDataInterfaceFeedback>& OutWarnings, TArray<FNiagaraDataInterfaceFeedback>& OutInfo)
+{
+	OutErrors = GetErrors();
+	OutWarnings.Empty();
+	OutInfo.Empty();
+}
+
+void UNiagaraDataInterface::GetFeedback(UNiagaraDataInterface* DataInterface, TArray<FNiagaraDataInterfaceError>& Errors, TArray<FNiagaraDataInterfaceFeedback>& Warnings,
+	TArray<FNiagaraDataInterfaceFeedback>& Info)
+{
+	if (!DataInterface)
+		return;
+
+	UNiagaraSystem* Asset = nullptr;
+	UNiagaraComponent* Component = nullptr;
+
+	// Walk the hierarchy to attempt to get the system and/or component
+	UObject* Curr = DataInterface->GetOuter();
+	while (Curr)
+	{
+		Asset = Cast<UNiagaraSystem>(Curr);
+		if (Asset)
+		{
+			break;
+		}
+
+		Component = Cast<UNiagaraComponent>(Curr);
+		if (Component)
+		{
+			Asset = Component->GetAsset();
+			break;
+		}
+
+		Curr = Curr->GetOuter();
+	}
+
+	DataInterface->GetFeedback(Asset, Component, Errors, Warnings, Info);
+}
 
 void UNiagaraDataInterface::ValidateFunction(const FNiagaraFunctionSignature& Function, TArray<FText>& OutValidationErrors)
 {
@@ -101,6 +148,16 @@ void UNiagaraDataInterface::ValidateFunction(const FNiagaraFunctionSignature& Fu
 	}
 }
 
+void UNiagaraDataInterface::RefreshErrors()
+{
+	OnErrorsRefreshedDelegate.Broadcast();
+}
+
+FSimpleMulticastDelegate& UNiagaraDataInterface::OnErrorsRefreshed()
+{
+	return OnErrorsRefreshedDelegate;
+}
+
 #endif
 
 #undef LOCTEXT_NAMESPACE
diff --git a/Engine/Plugins/FX/Niagara/Source/Niagara/Private/NiagaraDataInterfaceSkeletalMesh.cpp b/Engine/Plugins/FX/Niagara/Source/Niagara/Private/NiagaraDataInterfaceSkeletalMesh.cpp
--- a/Engine/Plugins/FX/Niagara/Source/Niagara/Private/NiagaraDataInterfaceSkeletalMesh.cpp
+++ b/Engine/Plugins/FX/Niagara/Source/Niagara/Private/NiagaraDataInterfaceSkeletalMesh.cpp
@@ -1176,8 +1176,6 @@ public:
 		else
 		{
 			// Bind dummy buffers
-			ensure(!InstanceData || InstanceData->StaticBuffers);
-
 			SetSRVParameter(RHICmdList, ComputeShaderRHI, MeshVertexBuffer, FNiagaraRenderer::GetDummyFloatBuffer());
 			SetSRVParameter(RHICmdList, ComputeShaderRHI, MeshIndexBuffer, FNiagaraRenderer::GetDummyUIntBuffer());
 			SetSRVParameter(RHICmdList, ComputeShaderRHI, MeshTangentBuffer, FNiagaraRenderer::GetDummyFloatBuffer());
@@ -1451,6 +1449,7 @@ USkeletalMesh* UNiagaraDataInterfaceSkeletalMesh::GetSkeletalMesh(UNiagaraCompon
 
 bool FNDISkeletalMesh_InstanceData::Init(UNiagaraDataInterfaceSkeletalMesh* Interface, FNiagaraSystemInstance* SystemInstance)
 {
+	check(Interface);
 	check(SystemInstance);
 
 	// Initialize members
@@ -1469,24 +1468,13 @@ bool FNDISkeletalMesh_InstanceData::Init(UNiagaraDataInterfaceSkeletalMesh* Inte
 	MeshSkinWeightIndexSizeByte = 0;
 	MeshGpuSpawnStaticBuffers = nullptr;
 	MeshGpuSpawnDynamicBuffers = nullptr;
-
-	bAllowCPUMeshDataAccess = true;
+	bAllowCPUMeshDataAccess = false;
 
 	// Get skel mesh and confirm have valid data
 	USkeletalMeshComponent* NewSkelComp = nullptr;
 	Mesh = Interface->GetSkeletalMesh(SystemInstance->GetComponent(), Component, NewSkelComp, this);
 	MeshSafe = Mesh;
 
-	if (!Mesh)
-	{
-		/*USceneComponent* Comp = Component.Get();
-		UE_LOG(LogNiagara, Log, TEXT("SkeletalMesh data interface has no valid mesh. Failed InitPerInstanceData!\nInterface: %s\nComponent: %s\nActor: %s\n")
-			, *Interface->GetFullName()
-			, Comp ? *Component->GetFullName() : TEXT("Null Component!")
-			, Comp ? *Comp->GetOwner()->GetFullName() : TEXT("NA"));*/
-		return false;
-	}
-
 	if (!Component.IsValid())
 	{
 		UE_LOG(LogNiagara, Log, TEXT("SkeletalMesh data interface has no valid component. Failed InitPerInstanceData - %s"), *Interface->GetFullName());
@@ -1500,15 +1488,12 @@ bool FNDISkeletalMesh_InstanceData::Init(UNiagaraDataInterfaceSkeletalMesh* Inte
 	CachedAttachParent = Component->GetAttachParent();
 
 #if WITH_EDITOR
-	MeshSafe->GetOnMeshChanged().AddUObject(SystemInstance->GetComponent(), &UNiagaraComponent::ReinitializeSystem);
+	if (MeshSafe.IsValid())
+	{
+		MeshSafe->GetOnMeshChanged().AddUObject(SystemInstance->GetComponent(), &UNiagaraComponent::ReinitializeSystem);
+	}
 #endif
 
-// 	if (!Mesh->bAllowCPUAccess)
-// 	{
-// 		UE_LOG(LogNiagara, Log, TEXT("SkeletalMesh data interface using a mesh that does not allow CPU access. Failed InitPerInstanceData - Mesh: %s"), *Mesh->GetFullName());
-// 		return false;
-// 	}
-
 	//Setup where to spawn from
 	SamplingRegionIndices.Empty();
 	bool bAllRegionsAreAreaWeighting = true;
@@ -1520,11 +1505,11 @@ bool FNDISkeletalMesh_InstanceData::Init(UNiagaraDataInterfaceSkeletalMesh* Inte
 		//If we have no regions, sample the whole mesh at the specified LOD.
 		if (LODIndex == INDEX_NONE)
 		{
-			LODIndex = Mesh->GetLODNum() - 1;
+			LODIndex = Mesh ? Mesh->GetLODNum() - 1 : 0;
 		}
 		else
 		{
-			LODIndex = FMath::Clamp(Interface->WholeMeshLOD, 0, Mesh->GetLODNum() - 1);
+			LODIndex = FMath::Clamp(LODIndex, 0, Mesh->GetLODNum() - 1);
 		}
 	}
 	else
@@ -1622,24 +1607,27 @@ bool FNDISkeletalMesh_InstanceData::Init(UNiagaraDataInterfaceSkeletalMesh* Inte
 		SamplingRegionAreaWeightedSampler.Init(this);
 	}
 
-	FSkinWeightVertexBuffer* SkinWeightBuffer = nullptr;
-	FSkeletalMeshLODRenderData& LODData = GetLODRenderDataAndSkinWeights(SkinWeightBuffer);
 
-	// Check for the validity of the Mesh's cpu data.
+	if (Mesh)
 	{
-		if ( Mesh->GetLODInfo(LODIndex)->bAllowCPUAccess )
+		bAllowCPUMeshDataAccess = true; // Assume accessibility until proven otherwise below
+		FSkinWeightVertexBuffer* SkinWeightBuffer = nullptr;
+		FSkeletalMeshLODRenderData* LODData = GetLODRenderDataAndSkinWeights(SkinWeightBuffer);
+		check(LODData);
+		check(SkinWeightBuffer);
+
+		// Check for the validity of the Mesh's cpu data.
+		if (Mesh->GetLODInfo(LODIndex)->bAllowCPUAccess)
 		{
-			const bool LODDataNumVerticesCorrect = LODData.GetNumVertices() > 0;
-			const bool LODDataPositonNumVerticesCorrect = LODData.StaticVertexBuffers.PositionVertexBuffer.GetNumVertices() > 0;
-			const bool bSkinWeightBuffer = SkinWeightBuffer != nullptr;
-			const bool SkinWeightBufferNumVerticesCorrect = bSkinWeightBuffer && (SkinWeightBuffer->GetNumVertices() > 0);
-			const bool bIndexBufferValid = LODData.MultiSizeIndexContainer.IsIndexBufferValid();
-			const bool bIndexBufferFound = bIndexBufferValid && (LODData.MultiSizeIndexContainer.GetIndexBuffer() != nullptr);
-			const bool bIndexBufferNumCorrect = bIndexBufferFound && (LODData.MultiSizeIndexContainer.GetIndexBuffer()->Num() > 0);
+			const bool LODDataNumVerticesCorrect = LODData->GetNumVertices() > 0;
+			const bool LODDataPositonNumVerticesCorrect = LODData->StaticVertexBuffers.PositionVertexBuffer.GetNumVertices() > 0;
+			const bool SkinWeightBufferNumVerticesCorrect = SkinWeightBuffer->GetNumVertices() > 0;
+			const bool bIndexBufferValid = LODData->MultiSizeIndexContainer.IsIndexBufferValid();
+			const bool bIndexBufferFound = bIndexBufferValid && (LODData->MultiSizeIndexContainer.GetIndexBuffer() != nullptr);
+			const bool bIndexBufferNumCorrect = bIndexBufferFound && (LODData->MultiSizeIndexContainer.GetIndexBuffer()->Num() > 0);
 
 			bAllowCPUMeshDataAccess = LODDataNumVerticesCorrect &&
 				LODDataPositonNumVerticesCorrect &&
-				bSkinWeightBuffer &&
 				SkinWeightBufferNumVerticesCorrect &&
 				bIndexBufferValid &&
 				bIndexBufferFound &&
@@ -1649,173 +1637,173 @@ bool FNDISkeletalMesh_InstanceData::Init(UNiagaraDataInterfaceSkeletalMesh* Inte
 		{
 			bAllowCPUMeshDataAccess = false;
 		}
-	}
 
-	// Generate excluded root bone index (if any)
-	FReferenceSkeleton& RefSkel = Mesh->RefSkeleton;
-	ExcludedBoneIndex = INDEX_NONE;
-	if (Interface->bExcludeBone && !Interface->ExcludeBoneName.IsNone())
-	{
-		ExcludedBoneIndex = RefSkel.FindBoneIndex(Interface->ExcludeBoneName);
-		if (ExcludedBoneIndex == INDEX_NONE)
+		// Generate excluded root bone index (if any)
+		FReferenceSkeleton& RefSkel = Mesh->RefSkeleton;
+		ExcludedBoneIndex = INDEX_NONE;
+		if (Interface->bExcludeBone && !Interface->ExcludeBoneName.IsNone())
 		{
-			UE_LOG(LogNiagara, Warning, TEXT("Skeletal Mesh Data Interface '%s' is missing bone '%s' this is ok but may not exclude what you want Mesh '%s' Component '%s'"), *Interface->GetFullName(), *Interface->ExcludeBoneName.ToString(), *Mesh->GetFullName(), *Component->GetFullName());
+			ExcludedBoneIndex = RefSkel.FindBoneIndex(Interface->ExcludeBoneName);
+			if (ExcludedBoneIndex == INDEX_NONE)
+			{
+				UE_LOG(LogNiagara, Warning, TEXT("Skeletal Mesh Data Interface '%s' is missing bone '%s' this is ok but may not exclude what you want Mesh '%s' Component '%s'"), *Interface->GetFullName(), *Interface->ExcludeBoneName.ToString(), *Mesh->GetFullName(), *Component->GetFullName());
+			}
 		}
-	}
 
-	// Gather filtered bones information
-	if (Interface->FilteredBones.Num() > 0)
-	{
-		if (RefSkel.GetNum() > TNumericLimits<uint16>::Max())
+		// Gather filtered bones information
+		if (Interface->FilteredBones.Num() > 0)
 		{
-			UE_LOG(LogNiagara, Warning, TEXT("Skeletal Mesh Data Interface '%s' requires more bones '%d' than we currently support '%d' Mesh '%s' Component '%s'"), *Interface->GetFullName(), RefSkel.GetNum(), TNumericLimits<uint16>::Max(), *Mesh->GetFullName(), *Component->GetFullName());
-			return false;
-		}
+			if (RefSkel.GetNum() > TNumericLimits<uint16>::Max())
+			{
+				UE_LOG(LogNiagara, Warning, TEXT("Skeletal Mesh Data Interface '%s' requires more bones '%d' than we currently support '%d' Mesh '%s' Component '%s'"), *Interface->GetFullName(), RefSkel.GetNum(), TNumericLimits<uint16>::Max(), *Mesh->GetFullName(), *Component->GetFullName());
+				return false;
+			}
 
-		//-TODO: If the DI does not use unfiltered bones we can skip adding them here...
-		TArray<FName, TInlineAllocator<16>> MissingBones;
+			//-TODO: If the DI does not use unfiltered bones we can skip adding them here...
+			TArray<FName, TInlineAllocator<16>> MissingBones;
 
-		FilteredAndUnfilteredBones.Reserve(RefSkel.GetNum());
+			FilteredAndUnfilteredBones.Reserve(RefSkel.GetNum());
 
-		// Append filtered bones to the array first
-		for (const FName& BoneName : Interface->FilteredBones)
-		{
-			const int32 Bone = RefSkel.FindBoneIndex(BoneName);
-			if (Bone == INDEX_NONE)
-			{
-				MissingBones.Add(BoneName);
-			}
-			else
+			// Append filtered bones to the array first
+			for (const FName& BoneName : Interface->FilteredBones)
 			{
-				ensure(Bone <= TNumericLimits<uint16>::Max());
-				FilteredAndUnfilteredBones.Add(Bone);
-				++NumFilteredBones;
+				const int32 Bone = RefSkel.FindBoneIndex(BoneName);
+				if (Bone == INDEX_NONE)
+				{
+					MissingBones.Add(BoneName);
+				}
+				else
+				{
+					ensure(Bone <= TNumericLimits<uint16>::Max());
+					FilteredAndUnfilteredBones.Add(Bone);
+					++NumFilteredBones;
+				}
 			}
-		}
 
-		// Append unfiltered bones to the array
-		for (int32 i=0; i < RefSkel.GetNum(); ++i)
-		{
-			// Don't append excluded bone
-			if (i == ExcludedBoneIndex)
+			// Append unfiltered bones to the array
+			for (int32 i = 0; i < RefSkel.GetNum(); ++i)
 			{
-				continue;
-			}
+				// Don't append excluded bone
+				if (i == ExcludedBoneIndex)
+				{
+					continue;
+				}
 
-			bool bExists = false;
-			for (int32 j=0; j < NumFilteredBones; ++j)
-			{
-				if (FilteredAndUnfilteredBones[j] == i)
+				bool bExists = false;
+				for (int32 j = 0; j < NumFilteredBones; ++j)
 				{
-					bExists = true;
-					break;
+					if (FilteredAndUnfilteredBones[j] == i)
+					{
+						bExists = true;
+						break;
+					}
+				}
+				if (!bExists)
+				{
+					FilteredAndUnfilteredBones.Add(i);
+					++NumUnfilteredBones;
 				}
 			}
-			if (!bExists)
+
+			if (MissingBones.Num() > 0)
 			{
-				FilteredAndUnfilteredBones.Add(i);
-				++NumUnfilteredBones;
+				UE_LOG(LogNiagara, Warning, TEXT("Skeletal Mesh Data Interface is trying to sample from bones that don't exist in it's skeleton.\nMesh: %s\nBones: "), *Mesh->GetName());
+				for (FName BoneName : MissingBones)
+				{
+					UE_LOG(LogNiagara, Warning, TEXT("%s\n"), *BoneName.ToString());
+				}
 			}
 		}
-
-		if (MissingBones.Num() > 0)
+		else
 		{
-			UE_LOG(LogNiagara, Warning, TEXT("Skeletal Mesh Data Interface is trying to sample from bones that don't exist in it's skeleton.\nMesh: %s\nBones: "), *Mesh->GetName());
-			for (FName BoneName : MissingBones)
-			{
-				UE_LOG(LogNiagara, Warning, TEXT("%s\n"), *BoneName.ToString());
-			}
+			// Note: We do not allocate space in the array as that wastes memory, we handle this special case when reading from unfiltered
+			NumUnfilteredBones = RefSkel.GetNum();
 		}
-	}
-	else
-	{
-		// Note: We do not allocate space in the array as that wastes memory, we handle this special case when reading from unfiltered
-		NumUnfilteredBones = RefSkel.GetNum();
-	}
 
-	// Gather filtered socket information
-	{
-		TArray<FName>& FilteredSockets = Interface->FilteredSockets;
-		FilteredSocketInfo.SetNum(FilteredSockets.Num());
-
-		//-TODO: We may need to handle skinning mode changes here
-		if (NewSkelComp != nullptr)
+		// Gather filtered socket information
 		{
-			for (int32 i = 0; i < FilteredSocketInfo.Num(); ++i)
+			TArray<FName>& FilteredSockets = Interface->FilteredSockets;
+			FilteredSocketInfo.SetNum(FilteredSockets.Num());
+
+			//-TODO: We may need to handle skinning mode changes here
+			if (NewSkelComp != nullptr)
 			{
-				NewSkelComp->GetSocketInfoByName(FilteredSockets[i], FilteredSocketInfo[i].Transform, FilteredSocketInfo[i].BoneIdx);
+				for (int32 i = 0; i < FilteredSocketInfo.Num(); ++i)
+				{
+					NewSkelComp->GetSocketInfoByName(FilteredSockets[i], FilteredSocketInfo[i].Transform, FilteredSocketInfo[i].BoneIdx);
+				}
 			}
-		}
-		else
-		{
-			for (int32 i = 0; i < FilteredSocketInfo.Num(); ++i)
+			else
 			{
-				FilteredSocketInfo[i].Transform = FTransform(Mesh->GetComposedRefPoseMatrix(FilteredSockets[i]));
-				FilteredSocketInfo[i].BoneIdx = INDEX_NONE;
+				for (int32 i = 0; i < FilteredSocketInfo.Num(); ++i)
+				{
+					FilteredSocketInfo[i].Transform = FTransform(Mesh->GetComposedRefPoseMatrix(FilteredSockets[i]));
+					FilteredSocketInfo[i].BoneIdx = INDEX_NONE;
+				}
 			}
-		}
 
-		FilteredSocketBoneOffset = Mesh->RefSkeleton.GetNum();
+			FilteredSocketBoneOffset = Mesh->RefSkeleton.GetNum();
 
-		FilteredSocketTransformsIndex = 0;
-		FilteredSocketTransforms[0].Reset(FilteredSockets.Num());
-		FilteredSocketTransforms[0].AddDefaulted(FilteredSockets.Num());
-		UpdateFilteredSocketTransforms();
-		for (int32 i=1; i < FilteredSocketTransforms.Num(); ++i)
-		{
-			FilteredSocketTransforms[i].Reset(FilteredSockets.Num());
-			FilteredSocketTransforms[i].Append(FilteredSocketTransforms[0]);
-		}
+			FilteredSocketTransformsIndex = 0;
+			FilteredSocketTransforms[0].Reset(FilteredSockets.Num());
+			FilteredSocketTransforms[0].AddDefaulted(FilteredSockets.Num());
+			UpdateFilteredSocketTransforms();
+			for (int32 i = 1; i < FilteredSocketTransforms.Num(); ++i)
+			{
+				FilteredSocketTransforms[i].Reset(FilteredSockets.Num());
+				FilteredSocketTransforms[i].Append(FilteredSocketTransforms[0]);
+			}
 
-		TArray<FName, TInlineAllocator<16>> MissingSockets;
-		for (FName SocketName : FilteredSockets)
-		{
-			if (Mesh->FindSocket(SocketName) == nullptr)
+			TArray<FName, TInlineAllocator<16>> MissingSockets;
+			for (FName SocketName : FilteredSockets)
 			{
-				MissingSockets.Add(SocketName);
+				if (Mesh->FindSocket(SocketName) == nullptr)
+				{
+					MissingSockets.Add(SocketName);
+				}
 			}
-		}
 
-		if (MissingSockets.Num() > 0)
-		{
-			UE_LOG(LogNiagara, Warning, TEXT("Skeletal Mesh Data Interface is trying to sample from sockets that don't exist in it's skeleton.\nMesh: %s\nSockets: "), *Mesh->GetName());
-			for (FName SocketName : MissingSockets)
+			if (MissingSockets.Num() > 0)
 			{
-				UE_LOG(LogNiagara, Warning, TEXT("%s\n"), *SocketName.ToString());
+				UE_LOG(LogNiagara, Warning, TEXT("Skeletal Mesh Data Interface is trying to sample from sockets that don't exist in it's skeleton.\nMesh: %s\nSockets: "), *Mesh->GetName());
+				for (FName SocketName : MissingSockets)
+				{
+					UE_LOG(LogNiagara, Warning, TEXT("%s\n"), *SocketName.ToString());
+				}
 			}
 		}
-	}
 
-	//-TODO: We should find out if this DI is connected to a GPU emitter or not rather than a blanket accross the system
-	if ( SystemInstance->HasGPUEmitters() )
-	{
-		GPUSkinBoneInfluenceType BoneInfluenceType = SkinWeightBuffer->GetBoneInfluenceType();
-		bUnlimitedBoneInfluences = (BoneInfluenceType == GPUSkinBoneInfluenceType::UnlimitedBoneInfluence);
-		MeshWeightStrideByte = SkinWeightBuffer->GetConstantInfluencesVertexStride();
-		MeshSkinWeightIndexSizeByte = SkinWeightBuffer->GetBoneIndexByteSize();
-		MeshSkinWeightBuffer = SkinWeightBuffer->GetDataVertexBuffer();
-		//check(MeshSkinWeightBufferSrv->IsValid()); // not available in this stream
-		MeshSkinWeightLookupBuffer = SkinWeightBuffer->GetLookupVertexBuffer();
+		//-TODO: We should find out if this DI is connected to a GPU emitter or not rather than a blanket accross the system
+		if (SystemInstance->HasGPUEmitters())
+		{
+			GPUSkinBoneInfluenceType BoneInfluenceType = SkinWeightBuffer->GetBoneInfluenceType();
+			bUnlimitedBoneInfluences = (BoneInfluenceType == GPUSkinBoneInfluenceType::UnlimitedBoneInfluence);
+			MeshWeightStrideByte = SkinWeightBuffer->GetConstantInfluencesVertexStride();
+			MeshSkinWeightIndexSizeByte = SkinWeightBuffer->GetBoneIndexByteSize();
+			MeshSkinWeightBuffer = SkinWeightBuffer->GetDataVertexBuffer();
+			//check(MeshSkinWeightBufferSrv->IsValid()); // not available in this stream
+			MeshSkinWeightLookupBuffer = SkinWeightBuffer->GetLookupVertexBuffer();
 
-		FSkeletalMeshLODInfo* LODInfo = Mesh->GetLODInfo(LODIndex);
-		bIsGpuUniformlyDistributedSampling = LODInfo->bSupportUniformlyDistributedSampling && bAllRegionsAreAreaWeighting;
+			FSkeletalMeshLODInfo* LODInfo = Mesh->GetLODInfo(LODIndex);
+			bIsGpuUniformlyDistributedSampling = LODInfo->bSupportUniformlyDistributedSampling && bAllRegionsAreAreaWeighting;
 
-		if (Mesh->HasActiveClothingAssets())
-		{
-			UE_LOG(LogNiagara, Warning, TEXT("Skeletal Mesh %s has cloth asset on it: spawning from it might not work properly."), *Mesh->GetName());
-		}
-		if (LODData.GetVertexBufferMaxBoneInfluences() > MAX_INFLUENCES_PER_STREAM)
-		{
-			UE_LOG(LogNiagara, Warning, TEXT("Skeletal Mesh %s has bones extra influence: spawning from it might not work properly."), *Mesh->GetName());
-		}
+			if (Mesh->HasActiveClothingAssets())
+			{
+				UE_LOG(LogNiagara, Warning, TEXT("Skeletal Mesh %s has cloth asset on it: spawning from it might not work properly."), *Mesh->GetName());
+			}
+			if (LODData->GetVertexBufferMaxBoneInfluences() > MAX_INFLUENCES_PER_STREAM)
+			{
+				UE_LOG(LogNiagara, Warning, TEXT("Skeletal Mesh %s has bones extra influence: spawning from it might not work properly."), *Mesh->GetName());
+			}
 
-		MeshGpuSpawnStaticBuffers = new FSkeletalMeshGpuSpawnStaticBuffers();
-		MeshGpuSpawnStaticBuffers->Initialise(this, LODData, SamplingInfo.GetBuiltData().WholeMeshBuiltData[LODIndex]);
-		BeginInitResource(MeshGpuSpawnStaticBuffers);
+			MeshGpuSpawnStaticBuffers = new FSkeletalMeshGpuSpawnStaticBuffers();
+			MeshGpuSpawnStaticBuffers->Initialise(this, *LODData, SamplingInfo.GetBuiltData().WholeMeshBuiltData[LODIndex]);
+			BeginInitResource(MeshGpuSpawnStaticBuffers);
 
-		MeshGpuSpawnDynamicBuffers = new FSkeletalMeshGpuDynamicBufferProxy();
-		MeshGpuSpawnDynamicBuffers->Initialise(RefSkel, LODData, FilteredSocketInfo.Num());
-		BeginInitResource(MeshGpuSpawnDynamicBuffers);
+			MeshGpuSpawnDynamicBuffers = new FSkeletalMeshGpuDynamicBufferProxy();
+			MeshGpuSpawnDynamicBuffers->Initialise(RefSkel, *LODData, FilteredSocketInfo.Num());
+			BeginInitResource(MeshGpuSpawnDynamicBuffers);
+		}
 	}
 
 	return true;
@@ -1864,15 +1852,6 @@ bool FNDISkeletalMesh_InstanceData::ResetRequired(UNiagaraDataInterfaceSkeletalM
 			return true;
 		}
 	}
-	else
-	{
-#if WITH_EDITORONLY_DATA
-		if (!Interface->PreviewMesh)
-		{
-			return true;
-		}
-#endif
-	}
 
 	if (Interface->ChangeId != ChangeId)
 	{
@@ -1893,7 +1872,7 @@ bool FNDISkeletalMesh_InstanceData::Tick(UNiagaraDataInterfaceSkeletalMesh* Inte
 	{
 		DeltaSeconds = InDeltaSeconds;
 
-		if (Component.IsValid() && Mesh)
+		if (Component.IsValid())
 		{
 			PrevTransform = Transform;
 			Transform = Component->GetComponentToWorld().ToMatrixWithScale();
@@ -1941,11 +1920,10 @@ void FNDISkeletalMesh_InstanceData::UpdateFilteredSocketTransforms()
 
 bool FNDISkeletalMesh_InstanceData::HasColorData()
 {
-	check(Mesh);
 	FSkinWeightVertexBuffer* SkinWeightBuffer;
-	FSkeletalMeshLODRenderData& LODData = GetLODRenderDataAndSkinWeights(SkinWeightBuffer);
+	FSkeletalMeshLODRenderData* LODData = GetLODRenderDataAndSkinWeights(SkinWeightBuffer);
 
-	return LODData.StaticVertexBuffers.ColorVertexBuffer.GetNumVertices() != 0;
+	return LODData && LODData->StaticVertexBuffers.ColorVertexBuffer.GetNumVertices() != 0;
 }
 
 void FNDISkeletalMesh_InstanceData::Release()
@@ -1987,9 +1965,6 @@ UNiagaraDataInterfaceSkeletalMesh::UNiagaraDataInterfaceSkeletalMesh(FObjectInit
 	, Source(nullptr)
 	, SkinningMode(ENDISkeletalMesh_SkinningMode::SkinOnTheFly)
 	, WholeMeshLOD(INDEX_NONE)
-#if WITH_EDITORONLY_DATA
-	, bRequiresCPUAccess(false)
-#endif
 	, ChangeId(0)
 {
 	static const FName NAME_Root("root");
@@ -2018,15 +1993,6 @@ void UNiagaraDataInterfaceSkeletalMesh::PostInitProperties()
 void UNiagaraDataInterfaceSkeletalMesh::PostEditChangeProperty(FPropertyChangedEvent& PropertyChangedEvent)
 {
 	Super::PostEditChangeProperty(PropertyChangedEvent);
-
-#if WITH_EDITORONLY_DATA
-	// If the change comes from an interaction (and not just a generic change) reset the usage flags.
-	// todo : this and the usage binding need to be done in the a precompilation parsing (or whever the script is compiled).
-	if (PropertyChangedEvent.Property)
-	{
-		bRequiresCPUAccess = false;
-	}
-#endif
 	ChangeId++;
 }
 
@@ -2051,7 +2017,7 @@ void UNiagaraDataInterfaceSkeletalMesh::GetVMExternalFunction(const FVMExternalF
 	FNDISkeletalMesh_InstanceData* InstData = (FNDISkeletalMesh_InstanceData*)InstanceData;
 	USkeletalMeshComponent* SkelComp = InstData != nullptr ? Cast<USkeletalMeshComponent>(InstData->Component.Get()) : nullptr;
 	
-	if (!InstData || !InstData->Mesh)
+	if (!InstData)
 	{
 		OutFunc = FVMExternalFunction();
 		return;
@@ -2071,13 +2037,7 @@ void UNiagaraDataInterfaceSkeletalMesh::GetVMExternalFunction(const FVMExternalF
 		if (!InstData->bAllowCPUMeshDataAccess)
 		{
 			UE_LOG(LogNiagara, Warning, TEXT("Skeletal Mesh Data Interface is trying to use triangle sampling but CPU access or the data is invalid. Interface: %s"), *GetFullName());
-			OutFunc.Unbind();
-			return;
 		}
-
-#if WITH_EDITOR
-		bRequiresCPUAccess = true;
-#endif // WITH_EDITOR
 		return;
 	}
 
@@ -2088,13 +2048,7 @@ void UNiagaraDataInterfaceSkeletalMesh::GetVMExternalFunction(const FVMExternalF
 		if (!InstData->bAllowCPUMeshDataAccess)
 		{
 			UE_LOG(LogNiagara, Warning, TEXT("Skeletal Mesh Data Interface is trying to use vertex sampling but CPU access or the data is invalid. Interface: %s"), *GetFullName());
-			OutFunc.Unbind();
-			return;
 		}
-
-#if WITH_EDITOR
-		bRequiresCPUAccess = true;
-#endif // WITH_EDITOR
 		return;
 	}
 }
@@ -2119,7 +2073,6 @@ bool UNiagaraDataInterfaceSkeletalMesh::CopyToInternal(UNiagaraDataInterface* De
 	OtherTyped->ExcludeBoneName = ExcludeBoneName;
 #if WITH_EDITORONLY_DATA
 	OtherTyped->PreviewMesh = PreviewMesh;
-	OtherTyped->bRequiresCPUAccess = bRequiresCPUAccess;
 #endif
 	return true;
 }
@@ -2186,26 +2139,90 @@ bool UNiagaraDataInterfaceSkeletalMesh::PerInstanceTick(void* PerInstanceData, F
 }
 
 #if WITH_EDITOR	
-TArray<FNiagaraDataInterfaceError> UNiagaraDataInterfaceSkeletalMesh::GetErrors()
+void UNiagaraDataInterfaceSkeletalMesh::GetFeedback(UNiagaraSystem* Asset, UNiagaraComponent* Component, TArray<FNiagaraDataInterfaceError>& OutErrors,
+	TArray<FNiagaraDataInterfaceFeedback>& OutWarnings, TArray<FNiagaraDataInterfaceFeedback>& OutInfo)
 {
-	TArray<FNiagaraDataInterfaceError> Errors;
-	bool bHasCPUAccessError = false;
+	bool bHasCPUAccessWarning = false;
 	bool bHasNoMeshAssignedError = false;
 	
 	// Collect Errors
 #if WITH_EDITORONLY_DATA
-	if (PreviewMesh != nullptr)
+	TWeakObjectPtr<USceneComponent> SceneComponent;
+	USkeletalMeshComponent* SkelMeshComponent = nullptr;
+	USkeletalMesh* SkelMesh = GetSkeletalMesh(Component, SceneComponent, SkelMeshComponent);
+	if (SkelMesh != nullptr)
 	{
-		if (bRequiresCPUAccess)
+		bool bHasCPUAccess = true;
+		for (const FSkeletalMeshLODInfo& LODInfo : SkelMesh->GetLODInfoArray())
 		{
-			for (const FSkeletalMeshLODInfo& LODInfo : PreviewMesh->GetLODInfoArray())
+			if (!LODInfo.bAllowCPUAccess)
+			{
+				bHasCPUAccess = false;
+				break;
+			}
+		}
+
+		// Check for the possibility that this mesh won't behave properly because of no CPU access
+		if (!bHasCPUAccess)
+		{			
+			// Filter through all the relevant CPU scripts
+			TArray<UNiagaraScript*> Scripts;
+			Scripts.Add(Asset->GetSystemSpawnScript());
+			Scripts.Add(Asset->GetSystemUpdateScript());
+			for (auto&& EmitterHandle : Asset->GetEmitterHandles())
+			{
+				TArray<UNiagaraScript*> OutScripts;
+				EmitterHandle.GetInstance()->GetScripts(OutScripts, false);
+				Scripts.Append(OutScripts.FilterByPredicate(
+					[&EmitterHandle](const UNiagaraScript* Script)
+					{
+						if (EmitterHandle.GetInstance()->SimTarget == ENiagaraSimTarget::GPUComputeSim)
+						{
+							// Ignore the spawn and update scripts
+							if (Script->Usage == ENiagaraScriptUsage::ParticleSpawnScript ||
+								Script->Usage == ENiagaraScriptUsage::ParticleSpawnScriptInterpolated ||
+								Script->Usage == ENiagaraScriptUsage::ParticleUpdateScript)
+							{
+								return false;
+							}
+						}
+						return Script->Usage != ENiagaraScriptUsage::ParticleGPUComputeScript;						
+					}
+				));
+			}
+
+			// Now check if any CPU script uses functions that require CPU access
+			//TODO: This isn't complete enough. It doesn't guarantee that the DI used by these functions are THIS DI.
+			// Finding that information out is currently non-trivial so just pop a warning with the possibility of false
+			// positives
+			TArray<FNiagaraFunctionSignature> CPUFunctions;
+			GetTriangleSamplingFunctions(CPUFunctions);
+			GetVertexSamplingFunctions(CPUFunctions);
+
+			bHasCPUAccessWarning = [this, &Scripts, &CPUFunctions]()
 			{
-				if (!LODInfo.bAllowCPUAccess)
+				for (const auto Script : Scripts)
 				{
-					bHasCPUAccessError = true;
-					break;
+					for (const auto& DIInfo : Script->GetVMExecutableData().DataInterfaceInfo)
+					{
+						if (DIInfo.GetDefaultDataInterface()->GetClass() == GetClass())
+						{
+							for (const auto& Func : DIInfo.RegisteredFunctions)
+							{
+								auto Filter = [&Func](const FNiagaraFunctionSignature& CPUSig)
+								{
+									return CPUSig.Name == Func.Name;
+								};
+								if (CPUFunctions.FindByPredicate(Filter))
+								{
+									return true;
+								}
+							}
+						}
+					}
 				}
-			}
+				return false;
+			}();			
 		}
 	}
 	else
@@ -2213,22 +2230,22 @@ TArray<FNiagaraDataInterfaceError> UNiagaraDataInterfaceSkeletalMesh::GetErrors(
 		bHasNoMeshAssignedError = true;
 	}
 
-	// Report Errors
-	if (Source == nullptr && bHasCPUAccessError)
+	// Report Errors/Warnings
+	if (bHasCPUAccessWarning)
 	{
-		FNiagaraDataInterfaceError CPUAccessNotAllowedError(FText::Format(LOCTEXT("CPUAccessNotAllowedError", "This mesh needs CPU access in order to be used properly.({0})"), FText::FromString(PreviewMesh->GetName())),
+		FNiagaraDataInterfaceFeedback CPUAccessNotAllowedWarning(FText::Format(LOCTEXT("CPUAccessNotAllowedError", "This mesh may need CPU access in order to be used properly. ({0})"), FText::FromString(SkelMesh->GetName())),
 			LOCTEXT("CPUAccessNotAllowedErrorSummary", "CPU access error"),
 			FNiagaraDataInterfaceFix::CreateLambda([=]()
-		{
-			PreviewMesh->Modify();
-			for (FSkeletalMeshLODInfo& LODInfo : PreviewMesh->GetLODInfoArray())
-			{
-				LODInfo.bAllowCPUAccess = true;
-			}
-			return true;
-		}));
+				{
+					SkelMesh->Modify();
+					for (FSkeletalMeshLODInfo& LODInfo : SkelMesh->GetLODInfoArray())
+					{
+						LODInfo.bAllowCPUAccess = true;
+					}
+					return true;
+				}));
 
-		Errors.Add(CPUAccessNotAllowedError);
+		OutWarnings.Add(CPUAccessNotAllowedWarning);
 	}
 #endif
 
@@ -2238,10 +2255,8 @@ TArray<FNiagaraDataInterfaceError> UNiagaraDataInterfaceSkeletalMesh::GetErrors(
 			LOCTEXT("NoMeshAssignedErrorSummary", "No mesh assigned error"),
 			FNiagaraDataInterfaceFix());
 
-		Errors.Add(NoMeshAssignedError);
+		OutErrors.Add(NoMeshAssignedError);
 	}
-
-	return Errors;
 }
 
 //Deprecated functions we check for and advise on updates in ValidateFunction
@@ -2827,22 +2842,27 @@ ETickingGroup UNiagaraDataInterfaceSkeletalMesh::CalculateTickGroup(void* PerIns
 //////////////////////////////////////////////////////////////////////////
 
 template<>
-void FSkeletalMeshAccessorHelper::Init<
-	TIntegralConstant<ENDISkeletalMesh_FilterMode, ENDISkeletalMesh_FilterMode::SingleRegion>,
-	TIntegralConstant<ENDISkelMesh_AreaWeightingMode, ENDISkelMesh_AreaWeightingMode::None>>
-	(FNDISkeletalMesh_InstanceData* InstData)
+void FSkeletalMeshAccessorHelper::Init<TNDISkelMesh_FilterModeSingle, TNDISkelMesh_AreaWeightingOff>(FNDISkeletalMesh_InstanceData* InstData)
 {
 	Comp = Cast<USkeletalMeshComponent>(InstData->Component.Get());
 	Mesh = InstData->Mesh;
-	LODData = &InstData->GetLODRenderDataAndSkinWeights(SkinWeightBuffer);
-	IndexBuffer = LODData->MultiSizeIndexContainer.GetIndexBuffer();
+	LODData = InstData->GetLODRenderDataAndSkinWeights(SkinWeightBuffer);
+	IndexBuffer = LODData ? LODData->MultiSizeIndexContainer.GetIndexBuffer() : nullptr;
 	SkinningData = InstData->SkinningData.SkinningData.Get();
 	Usage = InstData->SkinningData.Usage;
 
-	const FSkeletalMeshSamplingInfo& SamplingInfo = InstData->Mesh->GetSamplingInfo();
-	SamplingRegion = &SamplingInfo.GetRegion(InstData->SamplingRegionIndices[0]);
-	SamplingRegionBuiltData = &SamplingInfo.GetRegionBuiltData(InstData->SamplingRegionIndices[0]);
-
+	if (InstData->Mesh)
+	{
+		const FSkeletalMeshSamplingInfo& SamplingInfo = InstData->Mesh->GetSamplingInfo();
+		SamplingRegion = &SamplingInfo.GetRegion(InstData->SamplingRegionIndices[0]);
+		SamplingRegionBuiltData = &SamplingInfo.GetRegionBuiltData(InstData->SamplingRegionIndices[0]);
+	}
+	else
+	{
+		SamplingRegion = nullptr;
+		SamplingRegionBuiltData = nullptr;
+	}
+	
 	if (SkinningData != nullptr)
 	{
 		SkinningData->EnterRead();
@@ -2850,21 +2870,26 @@ void FSkeletalMeshAccessorHelper::Init<
 }
 
 template<>
-void FSkeletalMeshAccessorHelper::Init<
-	TIntegralConstant<ENDISkeletalMesh_FilterMode, ENDISkeletalMesh_FilterMode::SingleRegion>,
-	TIntegralConstant<ENDISkelMesh_AreaWeightingMode, ENDISkelMesh_AreaWeightingMode::AreaWeighted>>
-	(FNDISkeletalMesh_InstanceData* InstData)
+void FSkeletalMeshAccessorHelper::Init<TNDISkelMesh_FilterModeSingle, TNDISkelMesh_AreaWeightingOn>(FNDISkeletalMesh_InstanceData* InstData)
 {
 	Comp = Cast<USkeletalMeshComponent>(InstData->Component.Get());
 	Mesh = InstData->Mesh;
-	LODData = &InstData->GetLODRenderDataAndSkinWeights(SkinWeightBuffer);
-	IndexBuffer = LODData->MultiSizeIndexContainer.GetIndexBuffer();
+	LODData = InstData->GetLODRenderDataAndSkinWeights(SkinWeightBuffer);
+	IndexBuffer = LODData ? LODData->MultiSizeIndexContainer.GetIndexBuffer() : nullptr;
 	SkinningData = InstData->SkinningData.SkinningData.Get();
 	Usage = InstData->SkinningData.Usage;
 
-	const FSkeletalMeshSamplingInfo& SamplingInfo = InstData->Mesh->GetSamplingInfo();
-	SamplingRegion = &SamplingInfo.GetRegion(InstData->SamplingRegionIndices[0]);
-	SamplingRegionBuiltData = &SamplingInfo.GetRegionBuiltData(InstData->SamplingRegionIndices[0]);
+	if (InstData->Mesh)
+	{
+		const FSkeletalMeshSamplingInfo& SamplingInfo = InstData->Mesh->GetSamplingInfo();
+		SamplingRegion = &SamplingInfo.GetRegion(InstData->SamplingRegionIndices[0]);
+		SamplingRegionBuiltData = &SamplingInfo.GetRegionBuiltData(InstData->SamplingRegionIndices[0]);
+	}
+	else
+	{
+		SamplingRegion = nullptr;
+		SamplingRegionBuiltData = nullptr;
+	}
 
 	if (SkinningData != nullptr)
 	{
diff --git a/Engine/Plugins/FX/Niagara/Source/Niagara/Private/NiagaraParameterStore.cpp b/Engine/Plugins/FX/Niagara/Source/Niagara/Private/NiagaraParameterStore.cpp
--- a/Engine/Plugins/FX/Niagara/Source/Niagara/Private/NiagaraParameterStore.cpp
+++ b/Engine/Plugins/FX/Niagara/Source/Niagara/Private/NiagaraParameterStore.cpp
@@ -959,7 +959,7 @@ void FNiagaraParameterStore::OnLayoutChange()
 #endif
 }
 
-const FNiagaraVariableBase* FNiagaraParameterStore::FindVariable(UNiagaraDataInterface* Interface)const
+const FNiagaraVariableBase* FNiagaraParameterStore::FindVariable(const UNiagaraDataInterface* Interface) const
 {
 	SCOPE_CYCLE_COUNTER(STAT_NiagaraParameterStoreFindVar);
 	int32 Idx = DataInterfaces.IndexOfByKey(Interface);
diff --git a/Engine/Plugins/FX/Niagara/Source/Niagara/Private/NiagaraScriptExecutionContext.cpp b/Engine/Plugins/FX/Niagara/Source/Niagara/Private/NiagaraScriptExecutionContext.cpp
--- a/Engine/Plugins/FX/Niagara/Source/Niagara/Private/NiagaraScriptExecutionContext.cpp
+++ b/Engine/Plugins/FX/Niagara/Source/Niagara/Private/NiagaraScriptExecutionContext.cpp
@@ -160,6 +160,14 @@ bool FNiagaraScriptExecutionContext::Tick(FNiagaraSystemInstance* ParentSystemIn
 				FunctionTable[LocalFunctionTableIndices[LocalFunctionIt]] = &LocalFunctionTable[LocalFunctionIt];
 			}
 
+#if WITH_EDITOR	
+			// We may now have new errors that we need to broadcast about, so flush the asset parameters delegate..
+			if (ParentSystemInstance)
+			{
+				ParentSystemInstance->RaiseNeedsUIResync();
+			}
+#endif
+
 			if (!bSuccessfullyMapped)
 			{
 				UE_LOG(LogNiagara, Warning, TEXT("Error building data interface function table!"));
diff --git a/Engine/Plugins/FX/Niagara/Source/Niagara/Private/NiagaraSystemInstance.cpp b/Engine/Plugins/FX/Niagara/Source/Niagara/Private/NiagaraSystemInstance.cpp
--- a/Engine/Plugins/FX/Niagara/Source/Niagara/Private/NiagaraSystemInstance.cpp
+++ b/Engine/Plugins/FX/Niagara/Source/Niagara/Private/NiagaraSystemInstance.cpp
@@ -14,6 +14,7 @@
 #include "NiagaraEmitterInstanceBatcher.h"
 #include "GameFramework/PlayerController.h"
 #include "NiagaraCrashReporterHandler.h"
+#include "Async/Async.h"
 
 
 DECLARE_CYCLE_STAT(TEXT("System Activate [GT]"), STAT_NiagaraSystemActivate, STATGROUP_Niagara);
@@ -1915,6 +1916,7 @@ void FNiagaraSystemInstance::ComponentTick(float DeltaSeconds, const FGraphEvent
 	check(Component);
 
 	SystemSim->Tick_GameThread(DeltaSeconds, MyCompletionGraphEvent);
+
 }
 
 void FNiagaraSystemInstance::WaitForAsyncTickDoNotFinalize(bool bEnsureComplete)
@@ -2187,6 +2189,23 @@ void FNiagaraSystemInstance::FinalizeTick_GameThread()
 	}
 }
 
+#if WITH_EDITOR
+void FNiagaraSystemInstance::RaiseNeedsUIResync()
+{
+	AsyncTask(
+		ENamedThreads::GameThread,
+		[WeakComponent = TWeakObjectPtr<UNiagaraComponent>(Component)]() mutable
+	{
+		UNiagaraComponent* NiagaraComponent = WeakComponent.Get();
+		if (NiagaraComponent != nullptr )
+		{
+			NiagaraComponent->OnSynchronizedWithAssetParameters().Broadcast();
+		}
+	}
+	);
+}
+#endif
+
 #if WITH_EDITORONLY_DATA
 bool FNiagaraSystemInstance::GetIsolateEnabled() const
 {
diff --git a/Engine/Plugins/FX/Niagara/Source/Niagara/Public/NiagaraParameterStore.h b/Engine/Plugins/FX/Niagara/Source/Niagara/Public/NiagaraParameterStore.h
--- a/Engine/Plugins/FX/Niagara/Source/Niagara/Public/NiagaraParameterStore.h
+++ b/Engine/Plugins/FX/Niagara/Source/Niagara/Public/NiagaraParameterStore.h
@@ -356,7 +356,7 @@ public:
 	}
 
 	/** Returns the associated FNiagaraVariable for the passed data interface if it exists in the store. Null if not.*/
-	const FNiagaraVariableBase* FindVariable(UNiagaraDataInterface* Interface)const;
+	const FNiagaraVariableBase* FindVariable(const UNiagaraDataInterface* Interface) const;
 
 	virtual const int32* FindParameterOffset(const FNiagaraVariableBase& Parameter, bool IgnoreType = false) const;
 
diff --git a/Engine/Plugins/FX/Niagara/Source/Niagara/Public/NiagaraSystemInstance.h b/Engine/Plugins/FX/Niagara/Source/Niagara/Public/NiagaraSystemInstance.h
--- a/Engine/Plugins/FX/Niagara/Source/Niagara/Public/NiagaraSystemInstance.h
+++ b/Engine/Plugins/FX/Niagara/Source/Niagara/Public/NiagaraSystemInstance.h
@@ -276,6 +276,11 @@ public:
 	FORCEINLINE void SetLODDistance(float InLODDistance, float InMaxLODDistance);
 
 	const FString& GetCrashReporterTag()const;
+
+#if WITH_EDITOR
+	void RaiseNeedsUIResync();
+#endif
+
 private:
 
 	void DestroyDataInterfaceInstanceData();
diff --git a/Engine/Plugins/FX/Niagara/Source/NiagaraEditor/Private/ViewModels/Stack/NiagaraStackObject.cpp b/Engine/Plugins/FX/Niagara/Source/NiagaraEditor/Private/ViewModels/Stack/NiagaraStackObject.cpp
--- a/Engine/Plugins/FX/Niagara/Source/NiagaraEditor/Private/ViewModels/Stack/NiagaraStackObject.cpp
+++ b/Engine/Plugins/FX/Niagara/Source/NiagaraEditor/Private/ViewModels/Stack/NiagaraStackObject.cpp
@@ -133,6 +133,45 @@ void UNiagaraStackObject::RefreshChildrenInternal(const TArray<UNiagaraStackEntr
 
 		NewChildren.Add(ChildRow);
 	}
+
+	// TODO: Handle this in a more generic way.  Maybe add error apis to UNiagaraMergable, or use a UObject interface, or create a
+	// data interface specific implementation of UNiagaraStackObject.
+	UNiagaraDataInterface* DataInterfaceObject = Cast<UNiagaraDataInterface>(Object);
+	if (DataInterfaceObject != nullptr)
+	{
+		TArray<FNiagaraDataInterfaceError> Errors;
+		TArray<FNiagaraDataInterfaceFeedback> Warnings, Info;
+		UNiagaraDataInterface::GetFeedback(DataInterfaceObject, Errors, Warnings, Info);
+		if (Errors.Num() > 0)
+		{
+			NewIssues.Add(FStackIssue(
+				EStackIssueSeverity::Error,
+				NSLOCTEXT("StackObject", "ObjectErrorsShort", "Object has errors"),
+				NSLOCTEXT("StackObject", "ObjectErrorsLong", "The displayed object has errors.  Check the object properties or the message log for details."),
+				GetStackEditorDataKey(),
+				false));
+		}
+		if (Warnings.Num() > 0)
+		{
+			NewIssues.Add(FStackIssue(
+				EStackIssueSeverity::Warning,
+				NSLOCTEXT("StackObject", "ObjectWarningsShort", "Object has warnings"),
+				NSLOCTEXT("StackObject", "ObjectWarningsLong", "The displayed object has warnings.  Check the object properties or the message log for details."),
+				GetStackEditorDataKey(),
+				false));
+		}
+	}
+}
+
+void UNiagaraStackObject::PostRefreshChildrenInternal()
+{
+	Super::PostRefreshChildrenInternal();
+
+	UNiagaraDataInterface* DataInterfaceObject = Cast<UNiagaraDataInterface>(Object);
+	if (DataInterfaceObject != nullptr)
+	{
+		DataInterfaceObject->RefreshErrors();		
+	}
 }
 
 void UNiagaraStackObject::PropertyRowsRefreshed()
diff --git a/Engine/Plugins/FX/Niagara/Source/NiagaraEditor/Public/ViewModels/Stack/NiagaraStackObject.h b/Engine/Plugins/FX/Niagara/Source/NiagaraEditor/Public/ViewModels/Stack/NiagaraStackObject.h
--- a/Engine/Plugins/FX/Niagara/Source/NiagaraEditor/Public/ViewModels/Stack/NiagaraStackObject.h
+++ b/Engine/Plugins/FX/Niagara/Source/NiagaraEditor/Public/ViewModels/Stack/NiagaraStackObject.h
@@ -41,7 +41,8 @@ public:
 protected:
 	virtual void FinalizeInternal() override;
 
-	virtual void RefreshChildrenInternal(const TArray<UNiagaraStackEntry*>& CurrentChildren, TArray<UNiagaraStackEntry*>& NewChildren, TArray<FStackIssue>& NewIssues);
+	virtual void RefreshChildrenInternal(const TArray<UNiagaraStackEntry*>& CurrentChildren, TArray<UNiagaraStackEntry*>& NewChildren, TArray<FStackIssue>& NewIssues) override;
+	virtual void PostRefreshChildrenInternal() override;
 
 private:
 	void PropertyRowsRefreshed();
diff --git a/Engine/Plugins/FX/Niagara/Source/NiagaraEditorWidgets/Private/DetailCustomizations/NiagaraDataInterfaceDetails.cpp b/Engine/Plugins/FX/Niagara/Source/NiagaraEditorWidgets/Private/DetailCustomizations/NiagaraDataInterfaceDetails.cpp
--- a/Engine/Plugins/FX/Niagara/Source/NiagaraEditorWidgets/Private/DetailCustomizations/NiagaraDataInterfaceDetails.cpp
+++ b/Engine/Plugins/FX/Niagara/Source/NiagaraEditorWidgets/Private/DetailCustomizations/NiagaraDataInterfaceDetails.cpp
@@ -17,6 +17,11 @@
 #include "IPropertyUtilities.h"
 #include "Misc/NotifyHook.h"
 #include "IDetailChildrenBuilder.h"
+#include "NiagaraComponent.h"
+#include "NiagaraScript.h"
+#include "NiagaraEmitter.h"
+#include "NiagaraSystem.h"
+#include "NiagaraNodeInput.h"
 
 #define LOCTEXT_NAMESPACE "FNiagaraDataInterfaceDetailsBase"
 #define ErrorsCategoryName  TEXT("Errors")
@@ -100,7 +105,7 @@ private:
 
 	FReply OnFixNowClicked()
 	{
-		FScopedTransaction ScopedTransaction(NSLOCTEXT("NiagaraDataInterfaceDetails", "FixDataIntefraceTransaction", "Fix asset for data interface"));
+		FScopedTransaction ScopedTransaction(NSLOCTEXT("NiagaraDataInterfaceDetails", "FixDataInterfaceTransaction", "Fix asset for data interface"));
 		OnFixTriggered.ExecuteIfBound();
 		Error.TryFixError();
 		DataInterface->PostEditChange();
@@ -114,6 +119,101 @@ private:
 	FSimpleDelegate OnFixTriggered;
 };
 
+class SNiagaraDataInterfaceFeedback : public SCompoundWidget
+{
+public:
+	DECLARE_DELEGATE(FOnFixTriggered);
+
+public:
+	SLATE_BEGIN_ARGS(SNiagaraDataInterfaceFeedback) {}
+		SLATE_EVENT(FOnFixTriggered, OnFixTriggered)
+	SLATE_END_ARGS();
+
+	void Construct(const FArguments& InArgs, UNiagaraDataInterface* InDataInterface, FNiagaraDataInterfaceFeedback InFeedback, bool bInWarning)
+	{
+		OnFixTriggered = InArgs._OnFixTriggered;
+
+		bWarning = bInWarning;
+		Feedback = InFeedback;
+		DataInterface = InDataInterface;
+
+		TSharedRef<SHorizontalBox> FeedbackBox = SNew(SHorizontalBox);
+
+		FeedbackBox->AddSlot()
+			.AutoWidth()
+			[
+				SNew(SHorizontalBox)
+				.ToolTipText(GetFeedbackTextTooltip())
+			+ SHorizontalBox::Slot()
+			.AutoWidth()
+			.HAlign(HAlign_Left)
+			.VAlign(VAlign_Center)
+			[
+				SNew(SImage)
+				.Image(FEditorStyle::GetBrush(bWarning ? "Icons.Warning" : "Icons.Info"))
+			]
+		+ SHorizontalBox::Slot()
+			.AutoWidth()
+			.HAlign(HAlign_Left)
+			.VAlign(VAlign_Center)
+			[
+				SNew(STextBlock)
+				.TextStyle(FNiagaraEditorStyle::Get(), "NiagaraEditor.ParameterText")
+			.Clipping(EWidgetClipping::ClipToBounds)
+			.Text(GetFeedbackSummary())
+			]
+			];
+		if (Feedback.GetFeedbackFixable())
+		{
+			FeedbackBox->AddSlot()
+				.VAlign(VAlign_Top)
+				.Padding(5, 0, 0, 0)
+				.AutoWidth()
+				[
+					SNew(SButton)
+					.OnClicked(this, &SNiagaraDataInterfaceFeedback::OnFixNowClicked)
+				.ToolTipText(NSLOCTEXT("NiagaraDataInterfaceFeedback", "FixButtonLabelToolTip", "Fix the data linked to this interface."))
+				.Content()
+				[
+					SNew(STextBlock)
+					.Text(NSLOCTEXT("NiagaraDataInterfaceFeedback", "FixButtonLabel", "Fix Now"))
+				]
+				];
+		}
+		ChildSlot
+			[
+				FeedbackBox
+			];
+	}
+
+private:
+	FText GetFeedbackSummary() const
+	{
+		return Feedback.GetFeedbackSummaryText();
+	}
+
+	FText GetFeedbackTextTooltip() const
+	{
+		return Feedback.GetFeedbackText();
+	}
+
+	FReply OnFixNowClicked()
+	{
+		FScopedTransaction ScopedTransaction(NSLOCTEXT("NiagaraDataInterfaceDetails", "FixDataInterfaceTransaction", "Fix asset for data interface"));
+		OnFixTriggered.ExecuteIfBound();
+		Feedback.TryFixFeedback();
+		DataInterface->PostEditChange();
+		return FReply::Handled();
+	}
+
+private:
+	FNiagaraDataInterfaceFeedback Feedback;
+	IDetailLayoutBuilder* DetailBuilder;
+	UNiagaraDataInterface* DataInterface;
+	FSimpleDelegate OnFixTriggered;
+	bool bWarning;
+};
+
 class FNiagaraDataInterfaceCustomNodeBuilder : public IDetailCustomNodeBuilder
 											 , public TSharedFromThis<FNiagaraDataInterfaceCustomNodeBuilder>
 {
@@ -155,20 +255,52 @@ public:
 
 	virtual void GenerateChildContent(IDetailChildrenBuilder& ChildrenBuilder) override
 	{
-		for (FNiagaraDataInterfaceError Error : DataInterface->GetErrors())
+		TArray<FNiagaraDataInterfaceError> Errors;
+		TArray<FNiagaraDataInterfaceFeedback> Warnings;
+		TArray<FNiagaraDataInterfaceFeedback> Infos;
+		UNiagaraDataInterface::GetFeedback(DataInterface.Get(), Errors, Warnings, Infos);
+
+		for (FNiagaraDataInterfaceError Error : Errors)
 		{
 			FDetailWidgetRow& Row = ChildrenBuilder.AddCustomRow(NSLOCTEXT("NiagaraDataInterfaceDetails", "DataError", "Data Error"));
 			Row.WholeRowContent()
-			[
-				SNew(SHorizontalBox)
-				+ SHorizontalBox::Slot()
+				[
+					SNew(SHorizontalBox)
+					+ SHorizontalBox::Slot()
 				.AutoWidth()
 				[
 					SNew(SNiagaraDataInterfaceError, DataInterface.Get(), Error)
 					.OnFixTriggered(this, &FNiagaraDataInterfaceCustomNodeBuilder::OnErrorFixTriggered)
-			
 				]
-			];
+				];
+		}
+		for (FNiagaraDataInterfaceFeedback Warning : Warnings)
+		{
+			FDetailWidgetRow& Row = ChildrenBuilder.AddCustomRow(NSLOCTEXT("NiagaraDataInterfaceDetails", "DataWarning", "Data Warning"));
+			Row.WholeRowContent()
+				[
+					SNew(SHorizontalBox)
+					+ SHorizontalBox::Slot()
+				.AutoWidth()
+				[
+					SNew(SNiagaraDataInterfaceFeedback, DataInterface.Get(), Warning, true)
+					.OnFixTriggered(this, &FNiagaraDataInterfaceCustomNodeBuilder::OnErrorFixTriggered)
+				]
+				];
+		}
+		for (FNiagaraDataInterfaceFeedback Info : Infos)
+		{
+			FDetailWidgetRow& Row = ChildrenBuilder.AddCustomRow(NSLOCTEXT("NiagaraDataInterfaceDetails", "DataInfo", "Data Info"));
+			Row.WholeRowContent()
+				[
+					SNew(SHorizontalBox)
+					+ SHorizontalBox::Slot()
+				.AutoWidth()
+				[
+					SNew(SNiagaraDataInterfaceFeedback, DataInterface.Get(), Info, false)
+					.OnFixTriggered(this, &FNiagaraDataInterfaceCustomNodeBuilder::OnErrorFixTriggered)
+				]
+				];
 		}
 	}
 private:
@@ -187,7 +319,6 @@ private:
 		}
 	}
 
-
 private:
 	TWeakObjectPtr<UNiagaraDataInterface> DataInterface;
 	IDetailLayoutBuilder* DetailBuilder;
@@ -203,16 +334,16 @@ void FNiagaraDataInterfaceDetailsBase::CustomizeDetails(IDetailLayoutBuilder& De
 	check(SelectedObjects.Num() == 1);
 	DataInterface = Cast<UNiagaraDataInterface>(SelectedObjects[0].Get());
 	check(DataInterface.IsValid());
-	DataInterface->OnChanged().AddSP(this, &FNiagaraDataInterfaceDetailsBase::OnDataChanged);
-	IDetailCategoryBuilder& ErrorsBuilderRef = DetailBuilder.EditCategory(ErrorsCategoryName, LOCTEXT("Errors", "Errors"), ECategoryPriority::Important);
+	DataInterface->OnErrorsRefreshed().AddSP(this, &FNiagaraDataInterfaceDetailsBase::OnErrorsRefreshed);
+	IDetailCategoryBuilder& ErrorsBuilderRef = DetailBuilder.EditCategory(ErrorsCategoryName, LOCTEXT("ErrorsAndWarnings", "Errors And Warnings"), ECategoryPriority::Important);
 	ErrorsCategoryBuilder = &ErrorsBuilderRef;
 	CustomBuilder = MakeShared<FNiagaraDataInterfaceCustomNodeBuilder>(&DetailBuilder);
 	CustomBuilder->Initialize(*DataInterface);
 	ErrorsCategoryBuilder->AddCustomBuilder(CustomBuilder.ToSharedRef());
-	OnDataChanged();
+	OnErrorsRefreshed();
 }
 
-void FNiagaraDataInterfaceDetailsBase::OnDataChanged() // need to only refresh errors, and all will be good
+void FNiagaraDataInterfaceDetailsBase::OnErrorsRefreshed() // need to only refresh errors, and all will be good
 {
 	TSharedPtr<IPropertyUtilities> PropertyUtilities = PropertyUtilitiesWeak.Pin();
 	bool bStillValid =
@@ -224,15 +355,13 @@ void FNiagaraDataInterfaceDetailsBase::OnDataChanged() // need to only refresh e
 
 	if (bStillValid)
 	{
-		int CurrentErrorCount = DataInterface->GetErrors().Num();
-		if (CurrentErrorCount == 0)
-		{
-			ErrorsCategoryBuilder->SetCategoryVisibility(false);
-		}
-		else
-		{
-			ErrorsCategoryBuilder->SetCategoryVisibility(true);
-		}
+		TArray<FNiagaraDataInterfaceError> Errors;
+		TArray<FNiagaraDataInterfaceFeedback> Warnings;
+		TArray<FNiagaraDataInterfaceFeedback> Info;
+		UNiagaraDataInterface::GetFeedback(DataInterface.Get(), Errors, Warnings, Info);
+
+		int CurrentErrorCount = Errors.Num() + Warnings.Num() + Info.Num();
+		ErrorsCategoryBuilder->SetCategoryVisibility(CurrentErrorCount > 0);
 	}
 }
 
diff --git a/Engine/Plugins/FX/Niagara/Source/NiagaraEditorWidgets/Private/DetailCustomizations/NiagaraDataInterfaceDetails.h b/Engine/Plugins/FX/Niagara/Source/NiagaraEditorWidgets/Private/DetailCustomizations/NiagaraDataInterfaceDetails.h
--- a/Engine/Plugins/FX/Niagara/Source/NiagaraEditorWidgets/Private/DetailCustomizations/NiagaraDataInterfaceDetails.h
+++ b/Engine/Plugins/FX/Niagara/Source/NiagaraEditorWidgets/Private/DetailCustomizations/NiagaraDataInterfaceDetails.h
@@ -21,7 +21,7 @@ public:
 	static TSharedRef<IDetailCustomization> MakeInstance();
 
 private:
-	void OnDataChanged();
+	void OnErrorsRefreshed();
 
 private:
 	TWeakObjectPtr<UNiagaraDataInterface> DataInterface;
