Fix crash when fx.LogCompileIdGeneration is non-zero and an emitter has simulation stages.

#rnx

diff --git a/Engine/Plugins/FX/Niagara/Source/Niagara/Private/NiagaraSimulationStageBase.cpp b/Engine/Plugins/FX/Niagara/Source/Niagara/Private/NiagaraSimulationStageBase.cpp
--- a/Engine/Plugins/FX/Niagara/Source/Niagara/Private/NiagaraSimulationStageBase.cpp
+++ b/Engine/Plugins/FX/Niagara/Source/Niagara/Private/NiagaraSimulationStageBase.cpp
@@ -8,6 +8,11 @@
 
 bool UNiagaraSimulationStageGeneric::AppendCompileHash(FNiagaraCompileHashVisitor* InVisitor) const 
 {
+#if WITH_EDITORONLY_DATA
+	int32 Index = InVisitor->Values.AddDefaulted();
+	InVisitor->Values[Index].Object = FString::Printf(TEXT("Class: \"%s\"  Name: \"%s\""), *GetClass()->GetName(), *GetName());
+#endif
+
 	InVisitor->UpdatePOD(TEXT("Iterations"), Iterations);
 	InVisitor->UpdatePOD(TEXT("IterationSource"), (int32)IterationSource);
 	InVisitor->UpdatePOD(TEXT("bSpawnOnly"), bSpawnOnly ? 1 : 0);
