- Compile fix for some platforms

#rnx

diff --git a/Engine/Plugins/FX/Niagara/Source/Niagara/Classes/NDISkeletalMeshCommon.h b/Engine/Plugins/FX/Niagara/Source/Niagara/Classes/NDISkeletalMeshCommon.h
--- a/Engine/Plugins/FX/Niagara/Source/Niagara/Classes/NDISkeletalMeshCommon.h
+++ b/Engine/Plugins/FX/Niagara/Source/Niagara/Classes/NDISkeletalMeshCommon.h
@@ -414,7 +414,7 @@ struct FSkelMeshVertexAccessor
 	FORCEINLINE FLinearColor GetVertexColor(FSkeletalMeshLODRenderData* LODData, int32 VertexIdx)const
 	{
 		check(LODData);
-		return LODData.StaticVertexBuffers.ColorVertexBuffer.VertexColor(VertexIdx);
+		return LODData->StaticVertexBuffers.ColorVertexBuffer.VertexColor(VertexIdx);
 	}
 };
 
