Fix a small bug with the new IsValid function on the Niagara StaticMesh Data Interface

diff --git a/Engine/Plugins/FX/Niagara/Source/Niagara/Private/NiagaraDataInterfaceStaticMesh.cpp b/Engine/Plugins/FX/Niagara/Source/Niagara/Private/NiagaraDataInterfaceStaticMesh.cpp
--- a/Engine/Plugins/FX/Niagara/Source/Niagara/Private/NiagaraDataInterfaceStaticMesh.cpp
+++ b/Engine/Plugins/FX/Niagara/Source/Niagara/Private/NiagaraDataInterfaceStaticMesh.cpp
@@ -1350,11 +1350,13 @@ TArray<FNiagaraDataInterfaceError> UNiagaraDataInterfaceStaticMesh::GetErrors()
 void UNiagaraDataInterfaceStaticMesh::IsValid(FVectorVMContext& Context)
 {
 	VectorVM::FUserPtrHandler<FNDIStaticMesh_InstanceData> InstData(Context);
-	VectorVM::FExternalFuncRegisterHandler<bool> OutValid(Context);
+	VectorVM::FExternalFuncRegisterHandler<FNiagaraBool> OutValid(Context);
 	
+	FNiagaraBool Valid;
+	Valid.SetValue(InstData->Mesh != nullptr);
 	for (int32 i = 0; i < Context.NumInstances; ++i)
 	{
-		*OutValid.GetDest() = InstData->Mesh != nullptr;
+		*OutValid.GetDest() = Valid;
 		OutValid.Advance();
 	}
 }
