- Fix crash when reading back data from the GPU for the attribute spreadsheet

diff --git a/Engine/Plugins/FX/Niagara/Source/Niagara/Classes/NiagaraScriptExecutionContext.h b/Engine/Plugins/FX/Niagara/Source/Niagara/Classes/NiagaraScriptExecutionContext.h
--- a/Engine/Plugins/FX/Niagara/Source/Niagara/Classes/NiagaraScriptExecutionContext.h
+++ b/Engine/Plugins/FX/Niagara/Source/Niagara/Classes/NiagaraScriptExecutionContext.h
@@ -298,17 +298,17 @@ public:
 	const FSimulationStageMetaData* GetSimStageMetaData(uint32 SimulationStageIndex) const;
 
 #if WITH_EDITORONLY_DATA
-	mutable FRHIGPUMemoryReadback *GPUDebugDataReadbackFloat;
-	mutable FRHIGPUMemoryReadback *GPUDebugDataReadbackInt;
-	mutable FRHIGPUMemoryReadback *GPUDebugDataReadbackHalf;
-	mutable FRHIGPUMemoryReadback *GPUDebugDataReadbackCounts;
-	mutable uint32 GPUDebugDataFloatSize;
-	mutable uint32 GPUDebugDataIntSize;
-	mutable uint32 GPUDebugDataHalfSize;
-	mutable uint32 GPUDebugDataFloatStride;
-	mutable uint32 GPUDebugDataIntStride;
-	mutable uint32 GPUDebugDataHalfStride;
-	mutable uint32 GPUDebugDataCountOffset;
+	mutable TUniquePtr<FRHIGPUMemoryReadback> GPUDebugDataReadbackFloat;
+	mutable TUniquePtr<FRHIGPUMemoryReadback> GPUDebugDataReadbackInt;
+	mutable TUniquePtr<FRHIGPUMemoryReadback> GPUDebugDataReadbackHalf;
+	mutable TUniquePtr<FRHIGPUMemoryReadback> GPUDebugDataReadbackCounts;
+	mutable uint32 GPUDebugDataFloatSize = 0;
+	mutable uint32 GPUDebugDataIntSize = 0;
+	mutable uint32 GPUDebugDataHalfSize = 0;
+	mutable uint32 GPUDebugDataFloatStride = 0;
+	mutable uint32 GPUDebugDataIntStride = 0;
+	mutable uint32 GPUDebugDataHalfStride = 0;
+	mutable uint32 GPUDebugDataCountOffset = INDEX_NONE;
 	mutable TSharedPtr<struct FNiagaraScriptDebuggerInfo, ESPMode::ThreadSafe> DebugInfo;
 #endif
 
diff --git a/Engine/Plugins/FX/Niagara/Source/Niagara/Private/NiagaraEmitterInstanceBatcher.cpp b/Engine/Plugins/FX/Niagara/Source/Niagara/Private/NiagaraEmitterInstanceBatcher.cpp
--- a/Engine/Plugins/FX/Niagara/Source/Niagara/Private/NiagaraEmitterInstanceBatcher.cpp
+++ b/Engine/Plugins/FX/Niagara/Source/Niagara/Private/NiagaraEmitterInstanceBatcher.cpp
@@ -1125,9 +1125,9 @@ void NiagaraEmitterInstanceBatcher::ProcessDebugInfo(FRHICommandList &RHICmdList
 			// Do nothing.., handled in Run
 		}
 		// We may not have floats or ints, but we should have at least one of the two
-		else if ((Context->GPUDebugDataReadbackFloat == nullptr || Context->GPUDebugDataReadbackFloat->IsReady()) 
-				&& (Context->GPUDebugDataReadbackInt == nullptr || Context->GPUDebugDataReadbackInt->IsReady())
-				&& (Context->GPUDebugDataReadbackHalf == nullptr || Context->GPUDebugDataReadbackHalf->IsReady())
+		else if (  (!Context->GPUDebugDataReadbackFloat || Context->GPUDebugDataReadbackFloat->IsReady())
+				&& (!Context->GPUDebugDataReadbackInt || Context->GPUDebugDataReadbackInt->IsReady())
+				&& (!Context->GPUDebugDataReadbackHalf || Context->GPUDebugDataReadbackHalf->IsReady())
 				&& Context->GPUDebugDataReadbackCounts->IsReady()
 			)
 		{
@@ -1172,25 +1172,14 @@ void NiagaraEmitterInstanceBatcher::ProcessDebugInfo(FRHICommandList &RHICmdList
 			{
 				// The following code seems to take significant time on d3d12
 				// Clear out the readback buffers...
-				if (Context->GPUDebugDataReadbackFloat)
-				{
-					delete Context->GPUDebugDataReadbackFloat;
-					Context->GPUDebugDataReadbackFloat = nullptr;
-				}
-				if (Context->GPUDebugDataReadbackInt)
-				{
-					delete Context->GPUDebugDataReadbackInt;
-					Context->GPUDebugDataReadbackInt = nullptr;
-				}
-				if (Context->GPUDebugDataReadbackInt)
-				{
-					delete Context->GPUDebugDataReadbackHalf;
-					Context->GPUDebugDataReadbackHalf = nullptr;
-				}
-				delete Context->GPUDebugDataReadbackCounts;
-				Context->GPUDebugDataReadbackCounts = nullptr;	
+				Context->GPUDebugDataReadbackFloat.Reset();
+				Context->GPUDebugDataReadbackInt.Reset();
+				Context->GPUDebugDataReadbackHalf.Reset();
+				Context->GPUDebugDataReadbackCounts.Reset();
+
 				Context->GPUDebugDataFloatSize = 0;
 				Context->GPUDebugDataIntSize = 0;
+				Context->GPUDebugDataHalfSize = 0;
 				Context->GPUDebugDataFloatStride = 0;
 				Context->GPUDebugDataIntStride = 0;
 				Context->GPUDebugDataHalfStride = 0;
@@ -1455,7 +1444,7 @@ void NiagaraEmitterInstanceBatcher::Run(const FNiagaraGPUSystemTick& Tick, const
 			if (DestinationData.GetGPUBufferFloat().NumBytes > 0)
 			{
 				static const FName ReadbackFloatName(TEXT("Niagara GPU Debug Info Float Emitter Readback"));
-				Context->GPUDebugDataReadbackFloat = new FRHIGPUBufferReadback(ReadbackFloatName);
+				Context->GPUDebugDataReadbackFloat.Reset(new FRHIGPUBufferReadback(ReadbackFloatName));
 				Context->GPUDebugDataReadbackFloat->EnqueueCopy(RHICmdList, DestinationData.GetGPUBufferFloat().Buffer);
 				Context->GPUDebugDataFloatSize = DestinationData.GetGPUBufferFloat().NumBytes;
 				Context->GPUDebugDataFloatStride = DestinationData.GetFloatStride();
@@ -1464,7 +1453,7 @@ void NiagaraEmitterInstanceBatcher::Run(const FNiagaraGPUSystemTick& Tick, const
 			if (DestinationData.GetGPUBufferInt().NumBytes > 0)
 			{
 				static const FName ReadbackIntName(TEXT("Niagara GPU Debug Info Int Emitter Readback"));
-				Context->GPUDebugDataReadbackInt = new FRHIGPUBufferReadback(ReadbackIntName);
+				Context->GPUDebugDataReadbackInt.Reset(new FRHIGPUBufferReadback(ReadbackIntName));
 				Context->GPUDebugDataReadbackInt->EnqueueCopy(RHICmdList, DestinationData.GetGPUBufferInt().Buffer);
 				Context->GPUDebugDataIntSize = DestinationData.GetGPUBufferInt().NumBytes;
 				Context->GPUDebugDataIntStride = DestinationData.GetInt32Stride();
@@ -1473,14 +1462,14 @@ void NiagaraEmitterInstanceBatcher::Run(const FNiagaraGPUSystemTick& Tick, const
 			if (DestinationData.GetGPUBufferHalf().NumBytes > 0)
 			{
 				static const FName ReadbackHalfName(TEXT("Niagara GPU Debug Info Half Emitter Readback"));
-				Context->GPUDebugDataReadbackHalf = new FRHIGPUBufferReadback(ReadbackHalfName);
+				Context->GPUDebugDataReadbackHalf.Reset(new FRHIGPUBufferReadback(ReadbackHalfName));
 				Context->GPUDebugDataReadbackHalf->EnqueueCopy(RHICmdList, DestinationData.GetGPUBufferHalf().Buffer);
 				Context->GPUDebugDataHalfSize = DestinationData.GetGPUBufferHalf().NumBytes;
 				Context->GPUDebugDataHalfStride = DestinationData.GetHalfStride();
 			}
 
 			static const FName ReadbackCountsName(TEXT("Niagara GPU Emitter Readback"));
-			Context->GPUDebugDataReadbackCounts = new FRHIGPUBufferReadback(ReadbackCountsName);
+			Context->GPUDebugDataReadbackCounts.Reset(new FRHIGPUBufferReadback(ReadbackCountsName));
 			Context->GPUDebugDataReadbackCounts->EnqueueCopy(RHICmdList, GPUInstanceCounterManager.GetInstanceCountBuffer().Buffer);
 			Context->GPUDebugDataCountOffset = DestinationData.GetGPUInstanceCountBufferOffset();
 		}
diff --git a/Engine/Plugins/FX/Niagara/Source/Niagara/Private/NiagaraScriptExecutionContext.cpp b/Engine/Plugins/FX/Niagara/Source/Niagara/Private/NiagaraScriptExecutionContext.cpp
--- a/Engine/Plugins/FX/Niagara/Source/Niagara/Private/NiagaraScriptExecutionContext.cpp
+++ b/Engine/Plugins/FX/Niagara/Source/Niagara/Private/NiagaraScriptExecutionContext.cpp
@@ -498,16 +498,7 @@ FNiagaraComputeExecutionContext::FNiagaraComputeExecutionContext()
 	, EmitterCBufferLayout(TEXT("Niagara GPU Emitter CBuffer"))
 	, ExternalCBufferLayout(TEXT("Niagara GPU External CBuffer"))
 	, DataToRender(nullptr)
-#if WITH_EDITORONLY_DATA
-	, GPUDebugDataReadbackFloat(nullptr)
-	, GPUDebugDataReadbackInt(nullptr)
-	, GPUDebugDataReadbackCounts(nullptr)
-	, GPUDebugDataFloatSize(0)
-	, GPUDebugDataIntSize(0)
-	, GPUDebugDataFloatStride(0)
-	, GPUDebugDataIntStride(0)
-	, GPUDebugDataCountOffset(INDEX_NONE)
-#endif	  
+
 {
 	GlobalCBufferLayout.ConstantBufferSize = sizeof(FNiagaraGlobalParameters);
 	GlobalCBufferLayout.ComputeHash();
@@ -529,21 +520,10 @@ FNiagaraComputeExecutionContext::~FNiagaraComputeExecutionContext()
 	// check(EmitterInstanceReadback.GPUCountOffset == INDEX_NONE);
 
 #if WITH_EDITORONLY_DATA
-	if (GPUDebugDataReadbackFloat)
-	{
-		delete GPUDebugDataReadbackFloat;
-		GPUDebugDataReadbackFloat = nullptr;
-	}
-	if (GPUDebugDataReadbackInt)
-	{
-		delete GPUDebugDataReadbackInt;
-		GPUDebugDataReadbackInt = nullptr;
-	}
-	if (GPUDebugDataReadbackCounts)
-	{
-		delete GPUDebugDataReadbackCounts;
-		GPUDebugDataReadbackCounts = nullptr;
-	}
+	GPUDebugDataReadbackFloat.Reset();
+	GPUDebugDataReadbackInt.Reset();
+	GPUDebugDataReadbackHalf.Reset();
+	GPUDebugDataReadbackCounts.Reset();
 #endif
 
 	SetDataToRender(nullptr);
@@ -841,21 +821,10 @@ void FNiagaraComputeExecutionContext::ResetInternal(NiagaraEmitterInstanceBatche
 	}
 
 #if WITH_EDITORONLY_DATA
-	if (GPUDebugDataReadbackFloat)
-	{
-		delete GPUDebugDataReadbackFloat;
-		GPUDebugDataReadbackFloat = nullptr;
-	}
-	if (GPUDebugDataReadbackInt)
-	{
-		delete GPUDebugDataReadbackInt;
-		GPUDebugDataReadbackInt = nullptr;
-	}
-	if (GPUDebugDataReadbackCounts)
-	{
-		delete GPUDebugDataReadbackCounts;
-		GPUDebugDataReadbackCounts = nullptr;
-	}
+	GPUDebugDataReadbackFloat.Reset();
+	GPUDebugDataReadbackInt.Reset();
+	GPUDebugDataReadbackHalf.Reset();
+	GPUDebugDataReadbackCounts.Reset();
 #endif
 
 	SetDataToRender(nullptr);
