- Fix D3DDebug validation error, the buffer is a float not a uint

#rnx

diff --git a/Engine/Plugins/FX/Niagara/Shaders/Private/NiagaraDataInterfaceSkeletalMesh.ush b/Engine/Plugins/FX/Niagara/Shaders/Private/NiagaraDataInterfaceSkeletalMesh.ush
--- a/Engine/Plugins/FX/Niagara/Shaders/Private/NiagaraDataInterfaceSkeletalMesh.ush
+++ b/Engine/Plugins/FX/Niagara/Shaders/Private/NiagaraDataInterfaceSkeletalMesh.ush
@@ -27,7 +27,7 @@ struct FDISkelMeshContext
 	Buffer<uint>	MeshTriangleSamplerAliasBuffer;
 	uint			MeshNumSamplingRegionTriangles;
 	uint			MeshNumSamplingRegionVertices;
-	Buffer<uint>	MeshSamplingRegionsProbaBuffer;
+	Buffer<float>	MeshSamplingRegionsProbaBuffer;
 	Buffer<uint>	MeshSamplingRegionsAliasBuffer;
 	Buffer<uint>	MeshSampleRegionsTriangleIndices;
 	Buffer<uint>	MeshSampleRegionsVertices;
@@ -701,7 +701,7 @@ void DISkelMesh_GetFilteredSocketOrBoneAt(in FDISkelMeshContext DIContext, int F
 	Buffer<uint>	MeshTriangleSamplerAliasBuffer_ ## NAME; \
 	uint			MeshNumSamplingRegionTriangles_ ## NAME; \
 	uint			MeshNumSamplingRegionVertices_ ## NAME; \
-	Buffer<uint>	MeshSamplingRegionsProbaBuffer_ ## NAME; \
+	Buffer<float>	MeshSamplingRegionsProbaBuffer_ ## NAME; \
 	Buffer<uint>	MeshSamplingRegionsAliasBuffer_ ## NAME; \
 	Buffer<uint>	MeshSampleRegionsTriangleIndices_ ## NAME; \
 	Buffer<uint>	MeshSampleRegionsVertices_ ## NAME; \
