- Move call to IsThisApplicationForeground to editor only as we have seen it causing stalls on cooked builds when on remote desktop (it's unclear currently if the stall exists outside of remote desktop)

diff --git a/Engine/Plugins/FX/Niagara/Source/Niagara/Private/NiagaraWorldManager.cpp b/Engine/Plugins/FX/Niagara/Source/Niagara/Private/NiagaraWorldManager.cpp
--- a/Engine/Plugins/FX/Niagara/Source/Niagara/Private/NiagaraWorldManager.cpp
+++ b/Engine/Plugins/FX/Niagara/Source/Niagara/Private/NiagaraWorldManager.cpp
@@ -573,7 +573,7 @@ void FNiagaraWorldManager::Tick(ETickingGroup TickGroup, float DeltaSeconds, ELe
 	{
 		FNiagaraSharedObject::FlushDeletionList();
 
-#if PLATFORM_DESKTOP
+#if WITH_EDITOR //PLATFORM_DESKTOP
 		{
 			QUICK_SCOPE_CYCLE_COUNTER(STAT_Niagara_IsThisApplicationForeground);
 			bAppHasFocus = FPlatformApplicationMisc::IsThisApplicationForeground();
