Fixed shader compiler errors when using the collision DI on platforms which don't support distance fields.

#rnx

diff --git a/Engine/Plugins/FX/Niagara/Source/Niagara/Private/NiagaraDataInterfaceCollisionQuery.cpp b/Engine/Plugins/FX/Niagara/Source/Niagara/Private/NiagaraDataInterfaceCollisionQuery.cpp
--- a/Engine/Plugins/FX/Niagara/Source/Niagara/Private/NiagaraDataInterfaceCollisionQuery.cpp
+++ b/Engine/Plugins/FX/Niagara/Source/Niagara/Private/NiagaraDataInterfaceCollisionQuery.cpp
@@ -287,6 +287,7 @@ bool UNiagaraDataInterfaceCollisionQuery::GetFunctionHLSL(const FNiagaraDataInte
 		}\
 		\n}\n\n");
 		OutHLSL += TEXT("void ") + DistanceFieldFunction + TEXT("(in float3 InPosition, in float3 In_TraceEndPos, out bool OutCollisionValid, out float3 Out_CollisionPos, out float3 Out_CollisionNormal)\n{\n\
+		#if PLATFORM_SUPPORTS_DISTANCE_FIELDS\n\
 		float DistanceToNearestSurface = GetDistanceToNearestSurfaceGlobal(InPosition);\n\
 		if (DistanceToNearestSurface < length(In_TraceEndPos - InPosition))\n\
 		{\n\
@@ -295,6 +296,7 @@ bool UNiagaraDataInterfaceCollisionQuery::GetFunctionHLSL(const FNiagaraDataInte
 			Out_CollisionPos = InPosition - Out_CollisionNormal * DistanceToNearestSurface;\n\
 		}\n\
 		else\n\
+		#endif\n\
 		{\n\
 			OutCollisionValid = false;\n\
 			Out_CollisionNormal = float3(0.0, 0.0, 1.0);\n\
@@ -347,9 +349,14 @@ bool UNiagaraDataInterfaceCollisionQuery::GetFunctionHLSL(const FNiagaraDataInte
 	{
 		OutHLSL += TEXT("void ") + FunctionInfo.InstanceName + TEXT("(in float3 In_SamplePos, out float Out_DistanceToNearestSurface, out float3 Out_FieldGradient) \n{\n");
 		OutHLSL += TEXT("\
+			#if PLATFORM_SUPPORTS_DISTANCE_FIELDS\n\
 			Out_DistanceToNearestSurface = GetDistanceToNearestSurfaceGlobal(In_SamplePos);\n\
-			Out_FieldGradient = GetDistanceFieldGradientGlobal(In_SamplePos);\
-			\n}\n\n");
+			Out_FieldGradient = GetDistanceFieldGradientGlobal(In_SamplePos);\n\
+			#else\n\
+			Out_DistanceToNearestSurface = 0;\n\
+			Out_FieldGradient = (float3)0;\n\
+			#endif\n\
+			}\n\n");
 		return true;
 	}
 
diff --git a/Engine/Plugins/FX/Niagara/Source/NiagaraCore/Private/NiagaraCustomVersion.cpp b/Engine/Plugins/FX/Niagara/Source/NiagaraCore/Private/NiagaraCustomVersion.cpp
--- a/Engine/Plugins/FX/Niagara/Source/NiagaraCore/Private/NiagaraCustomVersion.cpp
+++ b/Engine/Plugins/FX/Niagara/Source/NiagaraCore/Private/NiagaraCustomVersion.cpp
@@ -8,4 +8,4 @@ const FGuid FNiagaraCustomVersion::GUID(0xFCF57AFA, 0x50764283, 0xB9A9E658, 0xFF
 // Register the custom version with core
 FCustomVersionRegistration GRegisterNiagaraCustomVersion(FNiagaraCustomVersion::GUID, FNiagaraCustomVersion::LatestVersion, TEXT("NiagaraVer"));
 
-const FGuid FNiagaraCustomVersion::LatestScriptCompileVersion(0x9E046792, 0x5F59194A, 0x9CDF3576, 0xC5AAAF52);
+const FGuid FNiagaraCustomVersion::LatestScriptCompileVersion(0xBC664CF8, 0x97074D74, 0xB6D452BA, 0xB5417E32);
