Fixed crash when events spawn particles beyond the CPU max limit.
Fixed issue where allocation estimate would leave emitters in a state unable to spawn if they ever blew their max allocation count.

diff --git a/Engine/Plugins/FX/Niagara/Source/Niagara/Classes/NiagaraEmitter.h b/Engine/Plugins/FX/Niagara/Source/Niagara/Classes/NiagaraEmitter.h
--- a/Engine/Plugins/FX/Niagara/Source/Niagara/Classes/NiagaraEmitter.h
+++ b/Engine/Plugins/FX/Niagara/Source/Niagara/Classes/NiagaraEmitter.h
@@ -487,6 +487,7 @@ public:
 
 	TStatId GetStatID(bool bGameThread, bool bConcurrent) const;
 
+	void ClearRuntimeAllocationEstimate(uint64 ReportHandle = INDEX_NONE);
 	/* This is used by the emitter instances to report runtime allocations to reduce reallocation in future simulation runs. */
 	int32 AddRuntimeAllocation(uint64 ReporterHandle, int32 AllocationCount);
 
diff --git a/Engine/Plugins/FX/Niagara/Source/Niagara/Private/NiagaraEmitter.cpp b/Engine/Plugins/FX/Niagara/Source/Niagara/Private/NiagaraEmitter.cpp
--- a/Engine/Plugins/FX/Niagara/Source/Niagara/Private/NiagaraEmitter.cpp
+++ b/Engine/Plugins/FX/Niagara/Source/Niagara/Private/NiagaraEmitter.cpp
@@ -1552,9 +1552,25 @@ TStatId UNiagaraEmitter::GetStatID(bool bGameThread, bool bConcurrent)const
 	return TStatId();
 }
 
+void UNiagaraEmitter::ClearRuntimeAllocationEstimate(uint64 ReportHandle)
+{
+	FScopeLock Lock(&EstimationCriticalSection);
+	if (ReportHandle == INDEX_NONE)
+	{
+		RuntimeEstimation.AllocationEstimate = 0;
+		RuntimeEstimation.RuntimeAllocations.Empty();
+		RuntimeEstimation.IsEstimationDirty = true;
+	}
+	else
+	{
+		RuntimeEstimation.RuntimeAllocations.Remove(ReportHandle);
+		RuntimeEstimation.IsEstimationDirty = true;
+	}
+}
+
 int32 UNiagaraEmitter::AddRuntimeAllocation(uint64 ReporterHandle, int32 AllocationCount)
 {
-	FScopeLock lock(&EstimationCriticalSection);
+	FScopeLock Lock(&EstimationCriticalSection);
 	int32* Estimate = RuntimeEstimation.RuntimeAllocations.Find(ReporterHandle);
 	if (!Estimate || *Estimate < AllocationCount)
 	{
@@ -1583,6 +1599,7 @@ int32 UNiagaraEmitter::GetMaxParticleCountEstimate()
 	{
 		FScopeLock lock(&EstimationCriticalSection);
 		int32 EstimationCount = RuntimeEstimation.RuntimeAllocations.Num();
+		RuntimeEstimation.AllocationEstimate = 0;
 		if (EstimationCount > 0)
 		{
 			RuntimeEstimation.RuntimeAllocations.ValueSort(TGreater<int32>());
diff --git a/Engine/Plugins/FX/Niagara/Source/Niagara/Private/NiagaraEmitterInstance.cpp b/Engine/Plugins/FX/Niagara/Source/Niagara/Private/NiagaraEmitterInstance.cpp
--- a/Engine/Plugins/FX/Niagara/Source/Niagara/Private/NiagaraEmitterInstance.cpp
+++ b/Engine/Plugins/FX/Niagara/Source/Niagara/Private/NiagaraEmitterInstance.cpp
@@ -1347,6 +1347,9 @@ void FNiagaraEmitterInstance::Tick(float DeltaSeconds)
 	{
 		UE_LOG(LogNiagara, Warning, TEXT("%s has attempted to exceed the max CPU particle count! | Max: %d | Requested: %u"), *CachedEmitter->GetFullName(), GMaxNiagaraCPUParticlesPerEmitter, AllocationSize);
 
+		//We clear the emitters estimate otherwise we get stuck in this state forever.
+		CachedEmitter->ClearRuntimeAllocationEstimate();
+
 		//For now we completely bail out of spawning new particles. Possibly should improve this in future.
 		AllocationSize = OrigNumParticles;
 		SpawnTotal = 0;
@@ -1355,6 +1358,12 @@ void FNiagaraEmitterInstance::Tick(float DeltaSeconds)
 		{
 			EventSpawnTotal = 0;
 			EventInstanceData->EventSpawnTotal = 0;
+
+			for (FNiagaraEventHandlingInfo& Info : EventInstanceData->EventHandlingInfo)
+			{
+				Info.SpawnCounts.Empty();
+				Info.TotalSpawnCount = 0;
+			}
 		}
 	}
 
@@ -1835,6 +1844,7 @@ void FNiagaraEmitterInstance::Tick(float DeltaSeconds)
 uint32 FNiagaraEmitterInstance::CalculateEventSpawnCount(const FNiagaraEventScriptProperties &EventHandlerProps, TArray<int32, TInlineAllocator<16>>& EventSpawnCounts, FNiagaraDataSet *EventSet)
 {
 	uint32 SpawnTotal = 0;
+	uint32 MaxSpawnCount = INT_MAX; //We could probably do to have a CVar for limiting the max event spawn directly but for now just keep the count from overflowing so it's caught by the overall partcle count checks later.
 	int32 NumEventsToProcess = 0;
 
 	if (EventSet)
@@ -1849,10 +1859,11 @@ uint32 FNiagaraEmitterInstance::CalculateEventSpawnCount(const FNiagaraEventScri
 		for (int32 i = 0; i < NumEventsToProcess; i++)
 		{
 			const uint32 SpawnNumber = bUseRandom ? FMath::RandRange((int32)EventHandlerProps.MinSpawnNumber, (int32)EventHandlerProps.SpawnNumber) : EventHandlerProps.SpawnNumber;
-			if (ExecutionState == ENiagaraExecutionState::Active && SpawnNumber > 0)
+			uint32 NewSpawnTotal = SpawnTotal + SpawnNumber;
+			if (ExecutionState == ENiagaraExecutionState::Active && SpawnNumber > 0 && NewSpawnTotal < MaxSpawnCount)
 			{
 				EventSpawnCounts.Add(SpawnNumber);
-				SpawnTotal += SpawnNumber;
+				SpawnTotal = NewSpawnTotal;
 			}
 		}
 	}
