Fixed a bug with event write nodes that caused a crash when the event name was changed

diff --git a/Engine/Plugins/FX/Niagara/Source/NiagaraEditor/Private/NiagaraNodeWriteDataSet.cpp b/Engine/Plugins/FX/Niagara/Source/NiagaraEditor/Private/NiagaraNodeWriteDataSet.cpp
--- a/Engine/Plugins/FX/Niagara/Source/NiagaraEditor/Private/NiagaraNodeWriteDataSet.cpp
+++ b/Engine/Plugins/FX/Niagara/Source/NiagaraEditor/Private/NiagaraNodeWriteDataSet.cpp
@@ -57,6 +57,17 @@ TSharedPtr<SGraphNode> UNiagaraNodeWriteDataSet::CreateVisualWidget()
 	return SNew(SNiagaraGraphNodeWriteDataSet, this);
 }
 
+bool UNiagaraNodeWriteDataSet::SynchronizeWithStruct()
+{
+	bool bSynchronized = Super::SynchronizeWithStruct();
+	if (EventName.IsNone())
+	{
+		EventName = DataSet.Name;
+		VisualsChangedDelegate.Broadcast(this);
+	}
+	return bSynchronized;
+}
+
 FText UNiagaraNodeWriteDataSet::GetNodeTitle(ENodeTitleType::Type TitleType) const
 {
 	return FText::Format(LOCTEXT("NiagaraDataSetWriteFormat", "{0} Write"), FText::FromName(DataSet.Name));
@@ -88,6 +99,7 @@ void UNiagaraNodeWriteDataSet::Compile(class FHlslNiagaraTranslator* Translator,
 	if (EventName.IsNone())
 	{
 		EventName = DataSet.Name;
+		VisualsChangedDelegate.Broadcast(this);
 	}
 
 	FNiagaraDataSetID AlteredDataSet = DataSet;
diff --git a/Engine/Plugins/FX/Niagara/Source/NiagaraEditor/Private/Widgets/SNiagaraGraphNodeWriteDataSet.cpp b/Engine/Plugins/FX/Niagara/Source/NiagaraEditor/Private/Widgets/SNiagaraGraphNodeWriteDataSet.cpp
--- a/Engine/Plugins/FX/Niagara/Source/NiagaraEditor/Private/Widgets/SNiagaraGraphNodeWriteDataSet.cpp
+++ b/Engine/Plugins/FX/Niagara/Source/NiagaraEditor/Private/Widgets/SNiagaraGraphNodeWriteDataSet.cpp
@@ -28,11 +28,6 @@ void SNiagaraGraphNodeWriteDataSet::Construct(const FArguments& InArgs, UEdGraph
 
 TSharedRef<SWidget> SNiagaraGraphNodeWriteDataSet::CreateNodeContentArea()
 {
-	FSinglePropertyParams InitParams;
-	InitParams.NamePlacement = EPropertyNamePlacement::Hidden;
-	FPropertyEditorModule& PropertyEditorModule = FModuleManager::GetModuleChecked<FPropertyEditorModule>("PropertyEditor");
-	TSharedPtr<ISinglePropertyView> SinglePropView = PropertyEditorModule.CreateSingleProperty(GraphNode, GET_MEMBER_NAME_CHECKED(UNiagaraNodeWriteDataSet, EventName), InitParams);
-
 	TSharedRef<SWidget> ContentAreaWidget = SGraphNode::CreateNodeContentArea();
 	TSharedPtr<SVerticalBox> VertContainer = SNew(SVerticalBox)
 		+ SVerticalBox::Slot()
@@ -52,7 +47,22 @@ TSharedRef<SWidget> SNiagaraGraphNodeWriteDataSet::CreateNodeContentArea()
 			.HAlign(HAlign_Center)
 			.AutoWidth()
 			[
-				SinglePropView.ToSharedRef()
+				SNew(SEditableTextBox)
+				.OnTextCommitted_Lambda([this](const FText& NewText, ETextCommit::Type CommitType)
+				{
+					if (UNiagaraNodeWriteDataSet* WriteDataSetNode = Cast<UNiagaraNodeWriteDataSet>(GraphNode))
+					{
+						WriteDataSetNode->EventName = FName(*NewText.ToString());
+					}
+				})
+				.Text_Lambda([this]()
+				{
+					if (UNiagaraNodeWriteDataSet* WriteDataSetNode = Cast<UNiagaraNodeWriteDataSet>(GraphNode))
+					{
+						return FText::FromName(WriteDataSetNode->EventName);
+					}
+					return FText();
+				})
 			]
 		]
 		+ SVerticalBox::Slot()
diff --git a/Engine/Plugins/FX/Niagara/Source/NiagaraEditor/Public/NiagaraNodeDataSetBase.h b/Engine/Plugins/FX/Niagara/Source/NiagaraEditor/Public/NiagaraNodeDataSetBase.h
--- a/Engine/Plugins/FX/Niagara/Source/NiagaraEditor/Public/NiagaraNodeDataSetBase.h
+++ b/Engine/Plugins/FX/Niagara/Source/NiagaraEditor/Public/NiagaraNodeDataSetBase.h
@@ -48,7 +48,7 @@ protected:
 	bool InitializeFromStructInternal(const UStruct* PayloadStruct);
 
 	bool IsSynchronizedWithStruct(bool bIgnoreConditionVar, FString* Issues, bool bLogIssues = true);
-	bool SynchronizeWithStruct();
+	virtual bool SynchronizeWithStruct();
 	bool GetSupportedNiagaraTypeDef(const FProperty* Property, FNiagaraTypeDefinition& TypeDef);
 
 	static const FName ConditionVarName;
diff --git a/Engine/Plugins/FX/Niagara/Source/NiagaraEditor/Public/NiagaraNodeWriteDataSet.h b/Engine/Plugins/FX/Niagara/Source/NiagaraEditor/Public/NiagaraNodeWriteDataSet.h
--- a/Engine/Plugins/FX/Niagara/Source/NiagaraEditor/Public/NiagaraNodeWriteDataSet.h
+++ b/Engine/Plugins/FX/Niagara/Source/NiagaraEditor/Public/NiagaraNodeWriteDataSet.h
@@ -36,6 +36,7 @@ public:
 
 	virtual void BuildParameterMapHistory(FNiagaraParameterMapHistoryBuilder& OutHistory, bool bRecursive = true, bool bFilterForCompilation = true) const override;
 	virtual TSharedPtr<SGraphNode> CreateVisualWidget() override;
+	virtual bool SynchronizeWithStruct() override;
 protected:
 	void AddConditionPin(int32 PinIndex = -1);
 };
