Fixed tooltip.

#rnx

diff --git a/Engine/Plugins/FX/Niagara/Source/NiagaraEditorWidgets/Private/Stack/SNiagaraStackModuleItem.cpp b/Engine/Plugins/FX/Niagara/Source/NiagaraEditorWidgets/Private/Stack/SNiagaraStackModuleItem.cpp
--- a/Engine/Plugins/FX/Niagara/Source/NiagaraEditorWidgets/Private/Stack/SNiagaraStackModuleItem.cpp
+++ b/Engine/Plugins/FX/Niagara/Source/NiagaraEditorWidgets/Private/Stack/SNiagaraStackModuleItem.cpp
@@ -93,7 +93,7 @@ void SNiagaraStackModuleItem::AddCustomRowWidgets(TSharedRef<SHorizontalBox> Hor
 			SNew(SButton)
 			.ButtonStyle(FEditorStyle::Get(), "RoundButton")
 			.OnClicked(this, &SNiagaraStackModuleItem::ScratchButtonPressed)
-			.ToolTipText(LOCTEXT("OpenInScratchToolTip", "Open this dynamic input in the scratch pad."))
+			.ToolTipText(LOCTEXT("OpenInScratchToolTip", "Open this module in the scratch pad."))
 			.ContentPadding(FMargin(1.0f, 0.0f))
 			.Content()
 			[
