Fix a bug caught by static analysis

diff --git a/Engine/Plugins/FX/Niagara/Source/Niagara/Private/NiagaraDataInterfaceSkeletalMesh.cpp b/Engine/Plugins/FX/Niagara/Source/Niagara/Private/NiagaraDataInterfaceSkeletalMesh.cpp
--- a/Engine/Plugins/FX/Niagara/Source/Niagara/Private/NiagaraDataInterfaceSkeletalMesh.cpp
+++ b/Engine/Plugins/FX/Niagara/Source/Niagara/Private/NiagaraDataInterfaceSkeletalMesh.cpp
@@ -1497,15 +1497,19 @@ bool FNDISkeletalMesh_InstanceData::Init(UNiagaraDataInterfaceSkeletalMesh* Inte
 	//Setup where to spawn from
 	SamplingRegionIndices.Empty();
 	bool bAllRegionsAreAreaWeighting = true;
-	const FSkeletalMeshSamplingInfo& SamplingInfo = Mesh->GetSamplingInfo();
 	int32 LODIndex = INDEX_NONE;
-	if (Interface->SamplingRegions.Num() == 0)
+	if (Mesh == nullptr)
+	{
+		// Just say we're sampling LOD 0, even though there are no LODs
+		LODIndex = 0;
+	}
+	else if (Interface->SamplingRegions.Num() == 0)
 	{
-		LODIndex = Interface->WholeMeshLOD;
 		//If we have no regions, sample the whole mesh at the specified LOD.
+		LODIndex = Interface->WholeMeshLOD;
 		if (LODIndex == INDEX_NONE)
 		{
-			LODIndex = Mesh ? Mesh->GetLODNum() - 1 : 0;
+			LODIndex = Mesh->GetLODNum() - 1;
 		}
 		else
 		{
@@ -1515,6 +1519,7 @@ bool FNDISkeletalMesh_InstanceData::Init(UNiagaraDataInterfaceSkeletalMesh* Inte
 	else
 	{
 		//Sampling from regions. Gather the indices of the regions we'll sample from.
+		const FSkeletalMeshSamplingInfo& SamplingInfo = Mesh->GetSamplingInfo();
 		for (FName RegionName : Interface->SamplingRegions)
 		{
 			int32 RegionIdx = SamplingInfo.IndexOfRegion(RegionName);
@@ -1796,6 +1801,7 @@ bool FNDISkeletalMesh_InstanceData::Init(UNiagaraDataInterfaceSkeletalMesh* Inte
 				UE_LOG(LogNiagara, Warning, TEXT("Skeletal Mesh %s has bones extra influence: spawning from it might not work properly."), *Mesh->GetName());
 			}
 
+			const FSkeletalMeshSamplingInfo& SamplingInfo = Mesh->GetSamplingInfo();
 			MeshGpuSpawnStaticBuffers = new FSkeletalMeshGpuSpawnStaticBuffers();
 			MeshGpuSpawnStaticBuffers->Initialise(this, *LODData, SamplingInfo.GetBuiltData().WholeMeshBuiltData[LODIndex]);
 			BeginInitResource(MeshGpuSpawnStaticBuffers);
