Fixed header file includes for non-unity build

diff --git a/Engine/Plugins/FX/Niagara/Source/NiagaraEditor/Private/Widgets/SNiagaraGraphNodeWriteDataSet.cpp b/Engine/Plugins/FX/Niagara/Source/NiagaraEditor/Private/Widgets/SNiagaraGraphNodeWriteDataSet.cpp
--- a/Engine/Plugins/FX/Niagara/Source/NiagaraEditor/Private/Widgets/SNiagaraGraphNodeWriteDataSet.cpp
+++ b/Engine/Plugins/FX/Niagara/Source/NiagaraEditor/Private/Widgets/SNiagaraGraphNodeWriteDataSet.cpp
@@ -7,12 +7,9 @@
 #include "Framework/MultiBox/MultiBoxBuilder.h"
 #include "Widgets/SBoxPanel.h"
 #include "Widgets/Input/SButton.h"
+#include "Widgets/Input/SEditableTextBox.h"
 #include "Widgets/Images/SImage.h"
-#include "ScopedTransaction.h"
-#include "Editor/PropertyEditor/Public/PropertyEditorModule.h"
-#include "ISinglePropertyView.h"
 #include "NiagaraNodeWriteDataSet.h"
-#include "Modules/ModuleManager.h"
 #include "GraphEditorSettings.h"
 
 #define LOCTEXT_NAMESPACE "SNiagaraGraphNodeWriteDataSet"
