Don't set the shader map on Niagara scripts if Serialize() returns false.

#rnx

diff --git a/Engine/Plugins/FX/Niagara/Source/NiagaraShader/Private/NiagaraShader.cpp b/Engine/Plugins/FX/Niagara/Source/NiagaraShader/Private/NiagaraShader.cpp
--- a/Engine/Plugins/FX/Niagara/Source/NiagaraShader/Private/NiagaraShader.cpp
+++ b/Engine/Plugins/FX/Niagara/Source/NiagaraShader/Private/NiagaraShader.cpp
@@ -1037,12 +1037,12 @@ void FNiagaraShaderMap::FlushShadersByShaderType(const FShaderType* ShaderType)
 
 
 
-void FNiagaraShaderMap::Serialize(FArchive& Ar, bool bInlineShaderResources, bool bLoadedByCookedMaterial)
+bool FNiagaraShaderMap::Serialize(FArchive& Ar, bool bInlineShaderResources, bool bLoadedByCookedMaterial)
 {
 	// Note: This is saved to the DDC, not into packages (except when cooked)
 	// Backwards compatibility therefore will not work based on the version of Ar
 	// Instead, just bump NIAGARASHADERMAP_DERIVEDDATA_VER
-	Super::Serialize(Ar, bInlineShaderResources, bLoadedByCookedMaterial);
+	return Super::Serialize(Ar, bInlineShaderResources, bLoadedByCookedMaterial);
 }
 
 bool FNiagaraShaderMap::RemovePendingScript(FNiagaraShaderScript* Script)
diff --git a/Engine/Plugins/FX/Niagara/Source/NiagaraShader/Private/NiagaraShared.cpp b/Engine/Plugins/FX/Niagara/Source/NiagaraShader/Private/NiagaraShared.cpp
--- a/Engine/Plugins/FX/Niagara/Source/NiagaraShader/Private/NiagaraShared.cpp
+++ b/Engine/Plugins/FX/Niagara/Source/NiagaraShader/Private/NiagaraShared.cpp
@@ -268,11 +268,11 @@ void FNiagaraShaderScript::SerializeShaderMap(FArchive& Ar)
 			if (bValid)
 			{
 				TRefCountPtr<FNiagaraShaderMap> LoadedShaderMap = new FNiagaraShaderMap();
-				LoadedShaderMap->Serialize(Ar, true, true);
+				bool bLoaded = LoadedShaderMap->Serialize(Ar, true, true);
 
 				// Toss the loaded shader data if this is a server only instance
 				//@todo - don't cook it in the first place
-				if (FApp::CanEverRender())
+				if (FApp::CanEverRender() && bLoaded)
 				{
 					GameThreadShaderMap = RenderingThreadShaderMap = LoadedShaderMap;
 				}
diff --git a/Engine/Plugins/FX/Niagara/Source/NiagaraShader/Public/NiagaraShared.h b/Engine/Plugins/FX/Niagara/Source/NiagaraShader/Public/NiagaraShared.h
--- a/Engine/Plugins/FX/Niagara/Source/NiagaraShader/Public/NiagaraShared.h
+++ b/Engine/Plugins/FX/Niagara/Source/NiagaraShader/Public/NiagaraShared.h
@@ -464,7 +464,7 @@ public:
 	static const FNiagaraShaderMap* GetShaderMapBeingCompiled(const FNiagaraShaderScript* Script);
 
 	/** Serializes the shader map. */
-	void Serialize(FArchive& Ar, bool bInlineShaderResources = true, bool bLoadedByCookedMaterial = false);
+	bool Serialize(FArchive& Ar, bool bInlineShaderResources = true, bool bLoadedByCookedMaterial = false);
 
 	/** Saves this shader map to the derived data cache. */
 	void SaveToDerivedDataCache();
