Niagara - Parameter map parameter name fixup on load and update the UI to show better categories and add a custom widget to show parameters everywhere in the UI.

diff --git a/Engine/Plugins/FX/Niagara/Source/Niagara/Classes/NiagaraConstants.h b/Engine/Plugins/FX/Niagara/Source/Niagara/Classes/NiagaraConstants.h
--- a/Engine/Plugins/FX/Niagara/Source/Niagara/Classes/NiagaraConstants.h
+++ b/Engine/Plugins/FX/Niagara/Source/Niagara/Classes/NiagaraConstants.h
@@ -154,6 +154,11 @@ struct NIAGARA_API FNiagaraConstants
 	static const FName EmitterNamespace;
 	static const FName ParticleAttributeNamespace;
 	static const FName ModuleNamespace;
+	static const FName OutputNamespace;
+	static const FName TransientNamespace;
+	static const FName DataInstanceNamespace;
+	static const FName StaticSwitchNamespace;
+	static const FName ArrayNamespace;
 	static const FName ParameterCollectionNamespace;
 	static const FString InitialPrefix;
 	static const FName LocalNamespace;
diff --git a/Engine/Plugins/FX/Niagara/Source/Niagara/Private/NiagaraConstants.cpp b/Engine/Plugins/FX/Niagara/Source/Niagara/Private/NiagaraConstants.cpp
--- a/Engine/Plugins/FX/Niagara/Source/Niagara/Private/NiagaraConstants.cpp
+++ b/Engine/Plugins/FX/Niagara/Source/Niagara/Private/NiagaraConstants.cpp
@@ -29,6 +29,11 @@ const FName FNiagaraConstants::SystemNamespace(TEXT("System"));
 const FName FNiagaraConstants::EmitterNamespace(TEXT("Emitter"));
 const FName FNiagaraConstants::ParticleAttributeNamespace(TEXT("Particles"));
 const FName FNiagaraConstants::ModuleNamespace(TEXT("Module"));
+const FName FNiagaraConstants::OutputNamespace(TEXT("Output"));
+const FName FNiagaraConstants::TransientNamespace(TEXT("Transient"));
+const FName FNiagaraConstants::DataInstanceNamespace(TEXT("DataInstance"));
+const FName FNiagaraConstants::StaticSwitchNamespace(TEXT("StaticSwitch"));
+const FName FNiagaraConstants::ArrayNamespace(TEXT("Array"));
 const FName FNiagaraConstants::ParameterCollectionNamespace(TEXT("NPC"));
 const FString FNiagaraConstants::InitialPrefix(TEXT("Initial"));
 const FName FNiagaraConstants::LocalNamespace(TEXT("Local"));
diff --git a/Engine/Plugins/FX/Niagara/Source/NiagaraCore/Private/NiagaraCustomVersion.cpp b/Engine/Plugins/FX/Niagara/Source/NiagaraCore/Private/NiagaraCustomVersion.cpp
--- a/Engine/Plugins/FX/Niagara/Source/NiagaraCore/Private/NiagaraCustomVersion.cpp
+++ b/Engine/Plugins/FX/Niagara/Source/NiagaraCore/Private/NiagaraCustomVersion.cpp
@@ -8,4 +8,4 @@ const FGuid FNiagaraCustomVersion::GUID(0xFCF57AFA, 0x50764283, 0xB9A9E658, 0xFF
 // Register the custom version with core
 FCustomVersionRegistration GRegisterNiagaraCustomVersion(FNiagaraCustomVersion::GUID, FNiagaraCustomVersion::LatestVersion, TEXT("NiagaraVer"));
 
-const FGuid FNiagaraCustomVersion::LatestScriptCompileVersion(0xBC664CF8, 0x97074D74, 0xB6D452BA, 0xB5417E32);
+const FGuid FNiagaraCustomVersion::LatestScriptCompileVersion(0x269AF844, 0x3EB24E94, 0xB1DECF9C, 0x21937120);
diff --git a/Engine/Plugins/FX/Niagara/Source/NiagaraCore/Public/NiagaraCustomVersion.h b/Engine/Plugins/FX/Niagara/Source/NiagaraCore/Public/NiagaraCustomVersion.h
--- a/Engine/Plugins/FX/Niagara/Source/NiagaraCore/Public/NiagaraCustomVersion.h
+++ b/Engine/Plugins/FX/Niagara/Source/NiagaraCore/Public/NiagaraCustomVersion.h
@@ -150,6 +150,8 @@ struct FNiagaraCustomVersion
 
 		SimulationStageInUsageBitmask, // Enable the simulation stage flag by default in the usage bitmask of modules and functions
 
+		StandardizeParameterNames, // Fix graph parameter map parameters on post load so that they all have a consisten parsable format and update the UI to show and filter based on these formats.
+
 		// DO NOT ADD A NEW VERSION UNLESS YOU HAVE TALKED TO THE NIAGARA LEAD. Mismanagement of these versions can lead to data loss if it is adjusted in multiple streams simultaneously.
 		// -----<new versions can be added above this line>  -------------------------------------------------
 		VersionPlusOne,
diff --git a/Engine/Plugins/FX/Niagara/Source/NiagaraEditor/Private/NiagaraActions.cpp b/Engine/Plugins/FX/Niagara/Source/NiagaraEditor/Private/NiagaraActions.cpp
--- a/Engine/Plugins/FX/Niagara/Source/NiagaraEditor/Private/NiagaraActions.cpp
+++ b/Engine/Plugins/FX/Niagara/Source/NiagaraEditor/Private/NiagaraActions.cpp
@@ -37,6 +37,7 @@ FNiagaraParameterAction::FNiagaraParameterAction(const FNiagaraVariable& InParam
 	: FEdGraphSchemaAction(MoveTemp(InNodeCategory), MoveTemp(InMenuDesc), MoveTemp(InToolTip), InGrouping, MoveTemp(InKeywords), InSectionID)
 	, Parameter(InParameter)
 	, ReferenceCollection(InReferenceCollection)
+	, bSubnamespaceRenamePending(false)
 {
 }
 
@@ -44,6 +45,7 @@ FNiagaraParameterAction::FNiagaraParameterAction(const FNiagaraVariable& InParam
 	FText InNodeCategory, FText InMenuDesc, FText InToolTip, const int32 InGrouping, FText InKeywords, int32 InSectionID)
 	: FEdGraphSchemaAction(MoveTemp(InNodeCategory), MoveTemp(InMenuDesc), MoveTemp(InToolTip), InGrouping, MoveTemp(InKeywords), InSectionID)
 	, Parameter(InParameter)
+	, bSubnamespaceRenamePending(false)
 {
 }
 
diff --git a/Engine/Plugins/FX/Niagara/Source/NiagaraEditor/Private/NiagaraEditorSettings.cpp b/Engine/Plugins/FX/Niagara/Source/NiagaraEditor/Private/NiagaraEditorSettings.cpp
--- a/Engine/Plugins/FX/Niagara/Source/NiagaraEditor/Private/NiagaraEditorSettings.cpp
+++ b/Engine/Plugins/FX/Niagara/Source/NiagaraEditor/Private/NiagaraEditorSettings.cpp
@@ -1,6 +1,20 @@
 // Copyright Epic Games, Inc. All Rights Reserved.
 
 #include "NiagaraEditorSettings.h"
+#include "NiagaraConstants.h"
+
+FNiagaraNamespaceMetadata::FNiagaraNamespaceMetadata()
+{
+}
+
+FNiagaraNamespaceMetadata::FNiagaraNamespaceMetadata(TArray<FName> InNamespaces, FText InDisplayName, FText InDescription, FLinearColor InBackgroundColor, TArray<ENiagaraNamespaceMetadataOptions> InOptions)
+	: Namespaces(InNamespaces)
+	, DisplayName(InDisplayName)
+	, Description(InDescription)
+	, BackgroundColor(InBackgroundColor)
+	, Options(InOptions)
+{
+}
 
 UNiagaraEditorSettings::UNiagaraEditorSettings(const FObjectInitializer& ObjectInitlaizer)
 	: Super(ObjectInitlaizer)
@@ -10,6 +24,76 @@ UNiagaraEditorSettings::UNiagaraEditorSettings(const FObjectInitializer& ObjectI
 	bResetSimulationOnChange = true;
 	bResimulateOnChangeWhilePaused = true;
 	bResetDependentSystemsWhenEditingEmitters = false;
+
+	NamespaceMetadata = 
+	{
+		FNiagaraNamespaceMetadata(
+			{NAME_None},
+			NSLOCTEXT("NamespaceMetadata", "DefaultDisplayName", "None"),
+			NSLOCTEXT("NamespaceMetadata", "DefaultDescription", "Arbitrary sub-namespace for specifying module specific dataset attributes, or calling nested modules."),
+			FLinearColor(FColor(102, 102, 102)),
+			{ }),
+		FNiagaraNamespaceMetadata(
+			{FNiagaraConstants::SystemNamespace},
+			NSLOCTEXT("NamespaceMetadata", "SystemDisplayName", "System"),
+			NSLOCTEXT("NamespaceMetadata", "SystemDescription", "Persistent attribute in the system which is written in a system\n stage and can be read anywhere."),
+			FLinearColor(FColor(67, 105, 124)),
+			{ENiagaraNamespaceMetadataOptions::SupportsSubnamespace}),
+		FNiagaraNamespaceMetadata(
+			{FNiagaraConstants::EmitterNamespace},
+			NSLOCTEXT("NamespaceMetadata", "EmitterDisplayName", "Emitter"),
+			NSLOCTEXT("NamespaceMetadata", "EmitterDescription", "Persistent attribute which is written in a emitter\n stage and can be read in emitter and particle stages."),
+			FLinearColor(FColor(126, 87, 67)),
+			{ENiagaraNamespaceMetadataOptions::SupportsSubnamespace}),
+		FNiagaraNamespaceMetadata(
+			{FNiagaraConstants::ParticleAttributeNamespace},
+			NSLOCTEXT("NamespaceMetadata", "ParticleDisplayName", "Particles"),
+			NSLOCTEXT("NamespaceMetadata", "ParticleDescription", "Persistent attribute which is written in a particle\n stage and can be read in particle stages."),
+			FLinearColor(FColor(87, 107, 61)),
+			{ENiagaraNamespaceMetadataOptions::SupportsSubnamespace}),
+		FNiagaraNamespaceMetadata(
+			{FNiagaraConstants::ModuleNamespace},
+			NSLOCTEXT("NamespaceMetadata", "ModuleDisplayName", "Input"),
+			NSLOCTEXT("NamespaceMetadata", "ModuleDescription", "A transient value which exposes a module input to the system and emitter editor."),
+			FLinearColor(FColor(118, 85, 48)),
+			{ENiagaraNamespaceMetadataOptions::SupportsSubnamespace}),
+		FNiagaraNamespaceMetadata(
+			{FNiagaraConstants::OutputNamespace, FNiagaraConstants::ModuleNamespace},
+			NSLOCTEXT("NamespaceMetadata", "ModuleOutputDisplayName", "Output"),
+			NSLOCTEXT("NamespaceMetadata", "ModuleOutputDescription", "A transient value which is unique to the module that wrote it and can be read from any other module.\nTransient values do not persist from frame to frame, or between stages, e.g. emitter to particle, or spawn to update."),
+			FLinearColor(FColor(118, 85, 48)),
+			{ENiagaraNamespaceMetadataOptions::Advanced}),
+		FNiagaraNamespaceMetadata(
+			{FNiagaraConstants::LocalNamespace, FNiagaraConstants::ModuleNamespace},
+			NSLOCTEXT("NamespaceMetadata", "ModuleLocalDisplayName", "Local"),
+			NSLOCTEXT("NamespaceMetadata", "ModuleLocalDescription", "A transient value which can be written to and read from within a single module.\nTransient values do not persist from frame to frame, or between stages, e.g. emitter to particle, or spawn to update."),
+			FLinearColor(FColor(118, 85, 48)),
+			{ENiagaraNamespaceMetadataOptions::Advanced}),
+		FNiagaraNamespaceMetadata(
+			{FNiagaraConstants::TransientNamespace},
+			NSLOCTEXT("NamespaceMetadata", "TransientDisplayName", "Transient"),
+			NSLOCTEXT("NamespaceMetadata", "TransientDescription", "A transient value which can be written to and read from from any module.\nTransient values do not persist from frame to frame, or between stages, e.g. emitter to particle, or spawn to update."),
+			FLinearColor(FColor(179, 122, 45)),
+			{ENiagaraNamespaceMetadataOptions::Advanced}),
+		FNiagaraNamespaceMetadata(
+			{FNiagaraConstants::EngineNamespace},
+			NSLOCTEXT("NamespaceMetadata", "EngineDisplayName", "Engine"),
+			NSLOCTEXT("NamespaceMetadata", "EngineDescription", "A read only value which is provided by the engine.\nThis value's source can be the simulation itsef\ne.g. ExecutionCount, or the owner of the simulation e.g. (Owner) Scale."),
+			FLinearColor(FColor(40, 40, 40)),
+			{ENiagaraNamespaceMetadataOptions::PreventRenaming}),
+		FNiagaraNamespaceMetadata(
+			{FNiagaraConstants::UserNamespace},
+			NSLOCTEXT("NamespaceMetadata", "UserDisplayName", "User"),
+			NSLOCTEXT("NamespaceMetadata", "UserDescription", "A read only value which can be initialized per system and\nmodified externally in the level, by blueprint, or by c++."),
+			FLinearColor(FColor(67, 105, 124)),
+			{ }),
+		FNiagaraNamespaceMetadata(
+			{FNiagaraConstants::DataInstanceNamespace},
+			NSLOCTEXT("NamespaceMetadata", "DataInstanceDisplayName", "Data Instance"),
+			NSLOCTEXT("NamespaceMetadata", "DataInstanceDescription", "A special transient value which has a single bool IsAlive value, which determines if a particle is alive or not."),
+			FLinearColor(FColor(179, 122, 45)),
+			{ENiagaraNamespaceMetadataOptions::Advanced}),
+	};
 }
 
 bool UNiagaraEditorSettings::GetAutoCompile() const
@@ -98,6 +182,49 @@ void UNiagaraEditorSettings::SetNewAssetDialogConfig(FName InDialogConfigKey, co
 	SaveConfig();
 }
 
+FNiagaraNamespaceMetadata UNiagaraEditorSettings::GetMetaDataForNamespaces(TArray<FName> InNamespaces) const
+{
+	TArray<FNiagaraNamespaceMetadata> MatchingMetadata;
+	for (const FNiagaraNamespaceMetadata& NamespaceMetadataItem : NamespaceMetadata)
+	{
+		if (NamespaceMetadataItem.Namespaces.Num() <= InNamespaces.Num())
+		{	
+			bool bNamespacesMatch = true;
+			for (int32 i = 0; i < NamespaceMetadataItem.Namespaces.Num() && bNamespacesMatch; i++)
+			{
+				if(NamespaceMetadataItem.Namespaces[i] != InNamespaces[i])
+				{
+					bNamespacesMatch = false;
+				}
+			}
+			if (bNamespacesMatch)
+			{
+				MatchingMetadata.Add(NamespaceMetadataItem);
+			}
+		}
+	}
+	if (MatchingMetadata.Num() == 0)
+	{
+		return FNiagaraNamespaceMetadata();
+	}
+	else if (MatchingMetadata.Num() == 1)
+	{
+		return MatchingMetadata[0];
+	}
+	else
+	{
+		int32 IndexOfLargestMatch = 0;
+		for (int32 i = 1; i < MatchingMetadata.Num(); i++)
+		{
+			if (MatchingMetadata[i].Namespaces.Num() > MatchingMetadata[IndexOfLargestMatch].Namespaces.Num())
+			{
+				IndexOfLargestMatch = i;
+			}
+		}
+		return MatchingMetadata[IndexOfLargestMatch];
+	}
+}
+
 FName UNiagaraEditorSettings::GetCategoryName() const
 {
 	return TEXT("Plugins");
diff --git a/Engine/Plugins/FX/Niagara/Source/NiagaraEditor/Private/NiagaraEditorStyle.cpp b/Engine/Plugins/FX/Niagara/Source/NiagaraEditor/Private/NiagaraEditorStyle.cpp
--- a/Engine/Plugins/FX/Niagara/Source/NiagaraEditor/Private/NiagaraEditorStyle.cpp
+++ b/Engine/Plugins/FX/Niagara/Source/NiagaraEditor/Private/NiagaraEditorStyle.cpp
@@ -45,6 +45,7 @@ NIAGARAEDITOR_API FString RelativePathToPluginPath(const FString& RelativePath,
 #define IMAGE_PLUGIN_BRUSH( RelativePath, ... ) FSlateImageBrush( RelativePathToPluginPath( RelativePath, ".png" ), __VA_ARGS__ )
 #define BOX_BRUSH( RelativePath, ... ) FSlateBoxBrush( Style->RootToContentDir( RelativePath, TEXT(".png") ), __VA_ARGS__ )
 #define BOX_CORE_BRUSH( RelativePath, ... ) FSlateBoxBrush( FPaths::EngineContentDir() / "Editor/Slate" / RelativePath + TEXT(".png"), __VA_ARGS__ )
+#define BOX_PLUGIN_BRUSH( RelativePath, ... ) FSlateBoxBrush( RelativePathToPluginPath( RelativePath, ".png"), __VA_ARGS__ )
 #define BORDER_BRUSH( RelativePath, ... ) FSlateBorderBrush( Style->RootToContentDir( RelativePath, TEXT(".png") ), __VA_ARGS__ )
 #define DEFAULT_FONT(...) FCoreStyle::GetDefaultFontStyle(__VA_ARGS__)
 
@@ -183,6 +184,13 @@ TSharedRef< FSlateStyleSet > FNiagaraEditorStyle::Create()
 
 	Style->Set("NiagaraEditor.ParameterSpinbox", ParameterSpinBox);
 
+	Style->Set("NiagaraEditor.ParameterName.NamespaceBorder", new BOX_PLUGIN_BRUSH("Icons/NamespaceBorder", FMargin(4.0f / 16.0f)));
+	Style->Set("NiagaraEditor.ParameterName.NamespaceText", FTextBlockStyle(NormalText)
+		.SetFont(DEFAULT_FONT("Bold", 8))
+		.SetColorAndOpacity(FLinearColor(1.0f, 1.0f, 1.0f, 0.9f))
+		.SetShadowOffset(FVector2D(1, 1))
+		.SetShadowColorAndOpacity(FLinearColor(0, 0, 0, 0.9f)));
+
 	// Code View
 	{
 		Style->Set("NiagaraEditor.CodeView.Checkbox.Text", FTextBlockStyle(NormalText)
diff --git a/Engine/Plugins/FX/Niagara/Source/NiagaraEditor/Private/NiagaraGraph.cpp b/Engine/Plugins/FX/Niagara/Source/NiagaraEditor/Private/NiagaraGraph.cpp
--- a/Engine/Plugins/FX/Niagara/Source/NiagaraEditor/Private/NiagaraGraph.cpp
+++ b/Engine/Plugins/FX/Niagara/Source/NiagaraEditor/Private/NiagaraGraph.cpp
@@ -23,6 +23,7 @@
 #include "NiagaraNodeParameterMapBase.h"
 #include "NiagaraNodeParameterMapGet.h"
 #include "NiagaraNodeParameterMapSet.h"
+#include "NiagaraNodeReroute.h"
 #include "INiagaraEditorTypeUtilities.h"
 #include "NiagaraConstants.h"
 #include "NiagaraEditorModule.h"
@@ -384,6 +385,11 @@ void UNiagaraGraph::PostLoad()
 		NotifyGraphNeedsRecompile();
 	}
 
+	if (NiagaraVer < FNiagaraCustomVersion::StandardizeParameterNames)
+	{
+		StandardizeParameterNames();
+	}
+
 	InvalidateCachedParameterData();
 }
 
@@ -491,6 +497,284 @@ void UNiagaraGraph::ValidateDefaultPins()
 	}
 }
 
+FName StandardizeName(FName Name, ENiagaraScriptUsage Usage, bool bIsGet, bool bIsSet)
+{
+	bool bIsScriptUsage = 
+		Usage == ENiagaraScriptUsage::Module || 
+		Usage == ENiagaraScriptUsage::DynamicInput || 
+		Usage == ENiagaraScriptUsage::Function;
+
+	TArray<FString> NamePartStrings;
+	TArray<FName> NameParts;
+	Name.ToString().ParseIntoArray(NamePartStrings, TEXT("."));
+
+	for (const FString& NamePartString : NamePartStrings)
+	{
+		NameParts.Add(*NamePartString);
+	}
+	if (NameParts.Num() == 0)
+	{
+		NameParts.Add(NAME_None);
+	}
+
+	FName Namespace;
+	if (NameParts[0] == FNiagaraConstants::EngineNamespace ||
+		NameParts[0] == FNiagaraConstants::ParameterCollectionNamespace ||
+		NameParts[0] == FNiagaraConstants::UserNamespace ||
+		NameParts[0] == FNiagaraConstants::ArrayNamespace ||
+		NameParts[0] == FNiagaraConstants::DataInstanceNamespace)
+	{
+		// Don't modify engine, parameter collections, or user parameters since they're defined externally and won't have fixed up
+		// names.  Also ignore the array and DataInstance namespaces because they're special.
+		return Name;
+	}
+	else if (NameParts[0] == FNiagaraConstants::LocalNamespace)
+	{
+		// Local can only be used in the context of a single module so force it to have a secondary module namespace.
+		static const FName LocalNamespace = *(FNiagaraConstants::LocalNamespace.ToString() + TEXT(".") + FNiagaraConstants::ModuleNamespace.ToString());
+		Namespace = LocalNamespace;
+		NameParts.RemoveAt(0);
+	}
+	else if (NameParts[0] == FNiagaraConstants::OutputNamespace)
+	{
+		static const FName OutputScriptNamespace = *(FNiagaraConstants::OutputNamespace.ToString() + TEXT(".") + FNiagaraConstants::ModuleNamespace.ToString());
+		static const FName OutputUnknownNamespace = *(FNiagaraConstants::OutputNamespace.ToString() + TEXT(".") + FName(NAME_None).ToString());
+		if (bIsScriptUsage)
+		{
+			if (bIsSet || NameParts.Num() <= 2)
+			{
+				// In a script outputs must always be written to the module namespace, and if the previous output didn't specify module
+				// or a sub-namespace assume it was reading from the module namespace.
+				Namespace = OutputScriptNamespace;
+				NameParts.RemoveAt(0);
+			}
+			else
+			{
+				// When reading they can be from the module namespace, or a more specific namespace from a different module.
+				Namespace = *(NameParts[0].ToString() + TEXT(".") + NameParts[1].ToString());
+				NameParts.RemoveAt(0, 2);
+			}
+		}
+		else
+		{
+			// The only valid usage for output parameters in system and emitter graphs is reading them from an aliased parameter of the form
+			// Output.ModuleName.ValueName.  If there are not enough name parts put in 'none' for the module name.
+			if (NameParts.Num() > 2)
+			{
+				Namespace = *(NameParts[0].ToString() + TEXT(".") + NameParts[1].ToString());
+				NameParts.RemoveAt(0, 2);
+			}
+			else
+			{
+				Namespace = OutputUnknownNamespace;
+				NameParts.RemoveAt(0);
+			}
+		}
+	}
+	else if (
+		NameParts[0] == FNiagaraConstants::ModuleNamespace ||
+		NameParts[0] == FNiagaraConstants::SystemNamespace ||
+		NameParts[0] == FNiagaraConstants::EmitterNamespace ||
+		NameParts[0] == FNiagaraConstants::ParticleAttributeNamespace)
+	{
+		if (NameParts.Num() == 2)
+		{
+			// Standard module input or dataset attribute.
+			Namespace = NameParts[0];
+			NameParts.RemoveAt(0);
+		}
+		else
+		{
+			// If there are more than 2 name parts, allow the first 2 for a namespace and sub-namespace.
+			// Sub-namespaces are used for module specific dataset attributes, or for configuring
+			// inputs for nested modules.
+			Namespace = *(NameParts[0].ToString() + TEXT(".") + NameParts[1].ToString());
+			NameParts.RemoveAt(0, 2);
+		}
+	}
+	else if (NameParts[0] == FNiagaraConstants::TransientNamespace)
+	{
+		// Transient names have a single namespace.
+		Namespace = FNiagaraConstants::TransientNamespace;
+		NameParts.RemoveAt(0);
+	}
+	else
+	{
+		// Namespace was unknown.
+		if (bIsScriptUsage || bIsGet)
+		{
+			// If it's in a get node or it's in a script, force it into the transient namespace.
+			Namespace = FNiagaraConstants::TransientNamespace;
+		}
+		else
+		{
+			// Otherwise it's a set node in a system or emitter script so it must be used to configure a module input
+			// so allow 1 arbitrary namespace, unless it's for a set variables node, then allow for 2 arbitrary namespaces.
+			if (NameParts[0].ToString().StartsWith(TRANSLATOR_SET_VARIABLES_STR))
+			{
+				Namespace = *(NameParts[0].ToString() + TEXT(".") + NameParts[1].ToString());
+				NameParts.RemoveAt(0, 2);
+			}
+			else
+			{
+				Namespace = NameParts[0];
+				NameParts.RemoveAt(0);
+			}
+		}
+	}
+
+	checkf(Namespace != NAME_None, TEXT("No namespace picked."));
+
+	// Remove any module or emitter namespaces which weren't handled above since they'll no longer be aliased anyway.
+	NameParts.Remove(FNiagaraConstants::ModuleNamespace);
+	NameParts.Remove(FNiagaraConstants::EmitterNamespace);
+	if (NameParts.Num() == 0)
+	{
+		NameParts.Add(NAME_None);
+	}
+
+	// Form the name by concatenating the remaining parts of the name.
+	FString ParameterName;
+	if (NameParts.Num() == 1)
+	{
+		ParameterName = NameParts[0].ToString();
+	}
+	else
+	{
+		TArray<FString> RemainingNamePartStrings;
+		for (FName NamePart : NameParts)
+		{
+			RemainingNamePartStrings.Add(NamePart.ToString());
+		}
+		ParameterName = FString::Join(RemainingNamePartStrings, TEXT("_"));
+	}
+
+	// Last, combine it with the namespace(s) chosen above.
+	return *FString::Printf(TEXT("%s.%s"), *Namespace.ToString(), *ParameterName);
+}
+
+void UNiagaraGraph::StandardizeParameterNames()
+{
+	TMap<FName, FName> OldNameToStandardizedNameMap;
+	bool bAnyNamesModified = false;
+	const UEdGraphSchema_Niagara* NiagaraSchema = GetDefault<UEdGraphSchema_Niagara>();
+	
+	TArray<UNiagaraNodeOutput*> OutputNodes;
+	GetNodesOfClass(OutputNodes);
+
+	auto HandleParameterMapNode = [&bAnyNamesModified](const UEdGraphSchema_Niagara* NiagaraSchema, UNiagaraNodeParameterMapBase* Node, EEdGraphPinDirection PinDirection, ENiagaraScriptUsage Usage, bool bIsGet, bool bIsSet, TMap<FName, FName>& OldNameToStandardizedNameMap)
+	{
+		for (UEdGraphPin* Pin : Node->Pins)
+		{
+			if (Pin->Direction == PinDirection)
+			{
+				FNiagaraTypeDefinition PinType = NiagaraSchema->PinToTypeDefinition(Pin);
+				if (PinType.IsValid() && PinType != FNiagaraTypeDefinition::GetParameterMapDef())
+				{
+					FName* NewNamePtr = OldNameToStandardizedNameMap.Find(Pin->PinName);
+					FName NewName;
+					if (NewNamePtr != nullptr)
+					{
+						NewName = *NewNamePtr;
+					}
+					else
+					{
+						NewName = StandardizeName(Pin->PinName, Usage, bIsGet, bIsSet);
+						OldNameToStandardizedNameMap.Add(Pin->PinName, NewName);
+					}
+
+					if (Pin->PinName != NewName)
+					{
+						bAnyNamesModified = true;
+					}
+					Pin->PinName = NewName;
+					Pin->PinFriendlyName = FText::FromName(NewName);
+				}
+			}
+		}
+	};
+
+	TSet<UNiagaraNode*> AllTraversedNodes;
+	for (UNiagaraNodeOutput* OutputNode : OutputNodes)
+	{
+		TArray<UNiagaraNode*> TraversedNodes;
+		BuildTraversal(TraversedNodes, OutputNode, false);
+		AllTraversedNodes.Append(TraversedNodes);
+
+		for (UNiagaraNode* TraversedNode : TraversedNodes)
+		{
+			TraversedNode->ConditionalPostLoad();
+			if (TraversedNode->IsA<UNiagaraNodeParameterMapGet>())
+			{
+				HandleParameterMapNode(NiagaraSchema, CastChecked<UNiagaraNodeParameterMapBase>(TraversedNode), EGPD_Output, OutputNode->GetUsage(), true, false, OldNameToStandardizedNameMap);
+			}
+			else if (TraversedNode->IsA<UNiagaraNodeParameterMapSet>())
+			{
+				HandleParameterMapNode(NiagaraSchema, CastChecked<UNiagaraNodeParameterMapBase>(TraversedNode), EGPD_Input, OutputNode->GetUsage(), false, true, OldNameToStandardizedNameMap);
+			}
+		}
+	}
+
+	for (UEdGraphNode* Node : Nodes)
+	{
+		UNiagaraNodeParameterMapBase* ParameterMapNode = Cast<UNiagaraNodeParameterMapBase>(Node);
+		if (ParameterMapNode != nullptr && AllTraversedNodes.Contains(ParameterMapNode) == false)
+		{
+			if (ParameterMapNode->IsA<UNiagaraNodeParameterMapGet>())
+			{
+				HandleParameterMapNode(NiagaraSchema, ParameterMapNode, EGPD_Output, ENiagaraScriptUsage::Module, true, false, OldNameToStandardizedNameMap);
+			}
+			else if (ParameterMapNode->IsA<UNiagaraNodeParameterMapSet>())
+			{
+				HandleParameterMapNode(NiagaraSchema, ParameterMapNode, EGPD_Input, ENiagaraScriptUsage::Module, false, true, OldNameToStandardizedNameMap);
+			}
+		}
+	}
+
+	// Since we'll be modifying the keys, make a copy of the map and then clear the original so it cal be 
+	// repopulated.
+	TMap<FNiagaraVariable, UNiagaraScriptVariable*> OldVariableToScriptVariable = VariableToScriptVariable;
+	VariableToScriptVariable.Empty();
+	for (TPair<FNiagaraVariable, UNiagaraScriptVariable*> VariableScriptVariablePair : OldVariableToScriptVariable)
+	{
+		FNiagaraVariable Variable = VariableScriptVariablePair.Key;
+		UNiagaraScriptVariable* ScriptVariable = VariableScriptVariablePair.Value;
+		ScriptVariable->PostLoad();
+		
+		if (ScriptVariable->Metadata.GetIsStaticSwitch() == false)
+		{
+			// We ignore static switches here because they're not in the parameter and so they don't need
+			// their parameter names to be fixed up.
+			FName* NewNamePtr = OldNameToStandardizedNameMap.Find(Variable.GetName());
+			FName NewName;
+			if (NewNamePtr != nullptr)
+			{
+				NewName = *NewNamePtr;
+			}
+			else
+			{
+				FName OldName = Variable.GetName();
+				NewName = StandardizeName(OldName, ENiagaraScriptUsage::Module, false, false);
+				OldNameToStandardizedNameMap.Add(OldName, NewName);
+			}
+
+			if (Variable.GetName() != NewName)
+			{
+				bAnyNamesModified = true;
+			}
+			Variable.SetName(NewName);
+			ScriptVariable->Variable.SetName(NewName);
+		}
+
+		VariableToScriptVariable.Add(Variable, ScriptVariable);
+	}
+
+	if (bAnyNamesModified)
+	{
+		MarkGraphRequiresSynchronization(TEXT("Standardized parameter names"));
+	}
+}
+
 void UNiagaraGraph::PostEditChangeProperty(FPropertyChangedEvent& PropertyChangedEvent)
 {
 	NotifyGraphChanged();
@@ -2287,6 +2571,11 @@ void UNiagaraGraph::RefreshParameterReferences() const
 	const UEdGraphSchema_Niagara* NiagaraSchema = GetNiagaraSchema();
 	for (UEdGraphNode* Node : Nodes)
 	{
+		if (Node->IsA<UNiagaraNodeReroute>())
+		{
+			continue;
+		}
+
 		UNiagaraNodeStaticSwitch* SwitchNode = Cast<UNiagaraNodeStaticSwitch>(Node);
 		if (SwitchNode && !SwitchNode->IsSetByCompiler())
 		{
diff --git a/Engine/Plugins/FX/Niagara/Source/NiagaraEditor/Private/SNiagaraParameterName.cpp b/Engine/Plugins/FX/Niagara/Source/NiagaraEditor/Private/SNiagaraParameterName.cpp
new file mode 100644
--- /dev/null
+++ b/Engine/Plugins/FX/Niagara/Source/NiagaraEditor/Private/SNiagaraParameterName.cpp
@@ -0,0 +1,341 @@
+// Copyright Epic Games, Inc. All Rights Reserved.
+
+#include "Widgets/SNiagaraParameterName.h"
+#include "NiagaraEditorSettings.h"
+#include "NiagaraEditorStyle.h"
+
+#include "Widgets/SBoxPanel.h"
+#include "Widgets/Layout/SBorder.h"
+#include "Widgets/Text/SInlineEditableTextBlock.h"
+#include "Widgets/Text/STextBlock.h"
+
+void SNiagaraParameterName::Construct(const FArguments& InArgs)
+{
+	EditableTextStyle = InArgs._EditableTextStyle;
+	ReadOnlyTextStyle = InArgs._ReadOnlyTextStyle;
+	ParameterName = InArgs._ParameterName;
+	bIsReadOnly = InArgs._IsReadOnly;
+	OnVerifyNameChangeDelegate = InArgs._OnVerifyNameChange;
+	OnNameChangedDelegate = InArgs._OnNameChanged;
+	OnDoubleClickedDelegate = InArgs._OnDoubleClicked;
+	IsSelected = InArgs._IsSelected;
+
+	UpdateContent(ParameterName.Get());
+}
+
+void SNiagaraParameterName::Tick(const FGeometry& AllottedGeometry, const double InCurrentTime, const float InDeltaTime)
+{
+	FName CurrentParameterName = ParameterName.Get();
+	if (DisplayedParameterName != CurrentParameterName)
+	{
+		UpdateContent(CurrentParameterName);
+	}
+}
+
+FReply SNiagaraParameterName::OnMouseButtonDoubleClick(const FGeometry& InMyGeometry, const FPointerEvent& InMouseEvent)
+{
+	if (OnDoubleClickedDelegate.IsBound())
+	{
+		return OnDoubleClickedDelegate.Execute(InMyGeometry, InMouseEvent);
+	}
+	else
+	{
+		return SCompoundWidget::OnMouseButtonDoubleClick(InMyGeometry, InMouseEvent);
+	}
+}
+
+TSharedRef<SBorder> SNiagaraParameterName::CreateNamespaceWidget(FText NamespaceDisplayName, FText NamespaceDescription, FLinearColor NamespaceBorderColor)
+{
+	return SNew(SBorder)
+	.BorderImage(FNiagaraEditorStyle::Get().GetBrush("NiagaraEditor.ParameterName.NamespaceBorder"))
+	.BorderBackgroundColor(NamespaceBorderColor)
+	.ToolTipText(NamespaceDescription)
+	.VAlign(VAlign_Center)
+	.Padding(FMargin(4.0f, 1.0f, 4.0f, 1.0f))
+	[
+		SNew(STextBlock)
+		.TextStyle(FNiagaraEditorStyle::Get(), "NiagaraEditor.ParameterName.NamespaceText")
+		.Text(NamespaceDisplayName)
+	];
+}
+
+void SNiagaraParameterName::UpdateContent(FName InDisplayedParameterName)
+{
+	DisplayedParameterName = InDisplayedParameterName;
+	
+	FString DisplayedParameterNameString = DisplayedParameterName.ToString();
+	TArray<FString> NameParts;
+	DisplayedParameterNameString.ParseIntoArray(NameParts, TEXT("."));
+
+	TArray<FName> Namespaces;
+	for (int32 i = 0; i < NameParts.Num() - 1; i++)
+	{
+		Namespaces.Add(*NameParts[i]);
+	}
+
+	TSharedRef<SHorizontalBox> ContentBox = SNew(SHorizontalBox);
+	auto AddNamespaceSlot = [](TSharedRef<SHorizontalBox> ContentBox, FText NamespaceDisplayName, FText NamespaceDescription, FLinearColor NamespaceBorderColor)
+	{
+		ContentBox->AddSlot()
+		.VAlign(VAlign_Center)
+		.AutoWidth()
+		.Padding(0.0f, 0.0f, 5.0f, 0.0f)
+		[
+
+				SNew(SBorder)
+				.BorderImage(FNiagaraEditorStyle::Get().GetBrush("NiagaraEditor.ParameterName.NamespaceBorder"))
+				.BorderBackgroundColor(NamespaceBorderColor)
+				.ToolTipText(NamespaceDescription)
+				.VAlign(VAlign_Center)
+				.Padding(FMargin(4.0f, 1.0f, 4.0f, 1.0f))
+				[
+					SNew(STextBlock)
+					.TextStyle(FNiagaraEditorStyle::Get(), "NiagaraEditor.ParameterName.NamespaceText")
+					.Text(NamespaceDisplayName)
+				]
+		];
+	};
+
+	FNiagaraNamespaceMetadata NamespaceMetadata = GetDefault<UNiagaraEditorSettings>()->GetMetaDataForNamespaces(Namespaces);
+	if (NamespaceMetadata.IsValid())
+	{
+		Namespaces.RemoveAt(0, NamespaceMetadata.Namespaces.Num());
+		ContentBox->AddSlot()
+		.VAlign(VAlign_Center)
+		.AutoWidth()
+		.Padding(0.0f, 0.0f, 5.0f, 0.0f)
+		[
+			CreateNamespaceWidget(NamespaceMetadata.DisplayName.ToUpper(), NamespaceMetadata.Description, NamespaceMetadata.BackgroundColor)
+		];
+	}
+
+	FNiagaraNamespaceMetadata DefaultNamespaceMetadata = GetDefault<UNiagaraEditorSettings>()->GetMetaDataForNamespaces({NAME_None});
+	if (Namespaces.Num() > 0)
+	{
+		DisplayedSubnamespace = Namespaces[0];
+		Namespaces.RemoveAt(0);
+		SubnamespaceBorder = CreateNamespaceWidget(
+			FText::FromString(FName::NameToDisplayString(DisplayedSubnamespace.ToString(), false).ToUpper()),
+			DefaultNamespaceMetadata.Description,
+			DefaultNamespaceMetadata.BackgroundColor);
+
+		ContentBox->AddSlot()
+		.VAlign(VAlign_Center)
+		.AutoWidth()
+		.Padding(0.0f, 0.0f, 5.0f, 0.0f)
+		[
+			SubnamespaceBorder.ToSharedRef()
+		];
+	}
+	else
+	{
+		DisplayedSubnamespace = NAME_None;
+		SubnamespaceBorder.Reset();
+	}
+
+	if (ensureMsgf(Namespaces.Num() == 0, TEXT("Extra namespaces in parameter.")) == false)
+	{
+		// If there are extra namespaces found, add them to the UI.
+		for (FName Namespace : Namespaces)
+		{
+			ContentBox->AddSlot()
+			.VAlign(VAlign_Center)
+			.AutoWidth()
+			.Padding(0.0f, 0.0f, 5.0f, 0.0f)
+			[
+				CreateNamespaceWidget(
+					FText::FromString(FName::NameToDisplayString(Namespace.ToString(), false).ToUpper()),
+					DefaultNamespaceMetadata.Description,
+					DefaultNamespaceMetadata.BackgroundColor)
+			];
+		}
+	}
+
+	TSharedPtr<SWidget> NameWidget;
+	if (bIsReadOnly)
+	{
+		NameWidget = SNew(STextBlock)
+			.TextStyle(ReadOnlyTextStyle)
+			.Text(FText::FromString(NameParts.Last()));
+	}
+	else
+	{
+		NameWidget = SAssignNew(EditableTextBlock, SInlineEditableTextBlock)
+		.Style(EditableTextStyle)
+		.Text(FText::FromString(NameParts.Last()))
+		.IsSelected(IsSelected)
+		.OnVerifyTextChanged(this, &SNiagaraParameterName::VerifyNameTextChange)
+		.OnTextCommitted(this, &SNiagaraParameterName::NameTextCommitted);
+	}
+
+	ContentBox->AddSlot()
+	.VAlign(VAlign_Center)
+	.AutoWidth()
+	[
+		NameWidget.ToSharedRef()
+	];
+
+	ChildSlot
+	[
+		ContentBox
+	];
+}
+
+FName SNiagaraParameterName::ReconstructNameFromEditText(const FText& InEditText)
+{
+	FString CurrentParameterNameString = ParameterName.Get().ToString();
+	TArray<FString> NameParts;
+	CurrentParameterNameString.ParseIntoArray(NameParts, TEXT("."));
+
+	NameParts[NameParts.Num() - 1] = InEditText.ToString().Replace(TEXT("."), TEXT("_"));
+	FString NewParameterNameString = FString::Join(NameParts, TEXT("."));
+	return *NewParameterNameString;
+}
+
+FName SNiagaraParameterName::ReconstructNameFromSubnamespaceEditText(const FText& InEditText)
+{
+	FString CurrentParameterNameString = ParameterName.Get().ToString();
+	TArray<FString> NameParts;
+	CurrentParameterNameString.ParseIntoArray(NameParts, TEXT("."));
+
+	if (NameParts.Num() >= 3)
+	{
+		NameParts[1] = InEditText.ToString().Replace(TEXT("."), TEXT("_"));
+	}
+	FString NewParameterNameString = FString::Join(NameParts, TEXT("."));
+	return *NewParameterNameString;
+}
+
+FText SNiagaraParameterName::GetSubnamespaceText()
+{
+	FString CurrentParameterNameString = ParameterName.Get().ToString();
+	TArray<FString> NameParts;
+	CurrentParameterNameString.ParseIntoArray(NameParts, TEXT("."));
+	return NameParts.Num() >= 3 
+		? FText::FromString(NameParts[1]) 
+		: FText::FromName(NAME_None);
+}
+
+bool SNiagaraParameterName::VerifyNameTextChange(const FText& InNewNameText, FText& OutErrorMessage)
+{
+	FName NewParameterName = ReconstructNameFromEditText(InNewNameText);
+	if (OnVerifyNameChangeDelegate.IsBound())
+	{
+		return OnVerifyNameChangeDelegate.Execute(NewParameterName, OutErrorMessage);
+	}
+	return true;
+}
+
+void SNiagaraParameterName::NameTextCommitted(const FText& InNewNameText, ETextCommit::Type CommitType)
+{
+	if (CommitType == ETextCommit::OnEnter)
+	{
+		FName NewParameterName = ReconstructNameFromEditText(InNewNameText);
+		OnNameChangedDelegate.ExecuteIfBound(NewParameterName);
+	}
+}
+
+bool SNiagaraParameterName::VerifySubnamespaceTextChange(const FText& InNewNameText, FText& OutErrorMessage)
+{
+	FName NewParameterName = ReconstructNameFromSubnamespaceEditText(InNewNameText);
+	if (OnVerifyNameChangeDelegate.IsBound())
+	{
+		return OnVerifyNameChangeDelegate.Execute(NewParameterName, OutErrorMessage);
+	}
+	return true;
+}
+
+void SNiagaraParameterName::SubnamespaceTextCommitted(const FText& InNewNameText, ETextCommit::Type CommitType)
+{
+	if (CommitType == ETextCommit::OnEnter)
+	{
+		FName NewParameterName = ReconstructNameFromSubnamespaceEditText(InNewNameText);
+		OnNameChangedDelegate.ExecuteIfBound(NewParameterName);
+	}
+
+	if (SubnamespaceBorder.IsValid())
+	{
+		SubnamespaceBorder->SetContent(SNew(STextBlock)
+			.TextStyle(FNiagaraEditorStyle::Get(), "NiagaraEditor.ParameterName.NamespaceText")
+			.Text(GetSubnamespaceText()));
+	}
+}
+
+void SNiagaraParameterName::EnterEditingMode()
+{
+	if (EditableTextBlock.IsValid())
+	{
+		EditableTextBlock->EnterEditingMode();
+	}
+}
+
+void SNiagaraParameterName::EnterSubnamespaceEditingMode()
+{
+	if (SubnamespaceBorder.IsValid())
+	{
+		TSharedRef<SInlineEditableTextBlock> SubnamespaceEditableTextBlock = SNew(SInlineEditableTextBlock)
+			.Style(EditableTextStyle)
+			.Text(GetSubnamespaceText())
+			.IsSelected(IsSelected)
+			.OnVerifyTextChanged(this, &SNiagaraParameterName::VerifySubnamespaceTextChange)
+			.OnTextCommitted(this, &SNiagaraParameterName::SubnamespaceTextCommitted);
+		SubnamespaceEditableTextBlock->EnterEditingMode();
+		SubnamespaceBorder->SetContent(SubnamespaceEditableTextBlock);
+	}
+}
+
+void SNiagaraParameterNameTextBlock::Construct(const FArguments& InArgs)
+{
+	ParameterText = InArgs._ParameterText;
+	OnVerifyNameTextChangedDelegate = InArgs._OnVerifyTextChanged;
+	OnNameTextCommittedDelegate = InArgs._OnTextCommitted;
+
+	ChildSlot
+	[
+		SAssignNew(ParameterName, SNiagaraParameterName)
+		.EditableTextStyle(InArgs._EditableTextStyle)
+		.ParameterName(this, &SNiagaraParameterNameTextBlock::GetParameterName)
+		.IsReadOnly(InArgs._IsReadOnly)
+		.IsSelected(InArgs._IsSelected)
+		.OnVerifyNameChange(this, &SNiagaraParameterNameTextBlock::VerifyNameChange)
+		.OnNameChanged(this, &SNiagaraParameterNameTextBlock::NameChanged)
+	];
+}
+
+FName SNiagaraParameterNameTextBlock::GetParameterName() const
+{
+	FText CurrentPinText = ParameterText.Get();
+	if (CurrentPinText.IdenticalTo(DisplayedParameterTextCache) == false)
+	{
+		DisplayedParameterTextCache = CurrentPinText;
+		ParameterNameCache = *DisplayedParameterTextCache.ToString();
+	}
+	return ParameterNameCache;
+}
+
+bool SNiagaraParameterNameTextBlock::VerifyNameChange(FName InNewName, FText& OutErrorMessage)
+{
+	if (OnVerifyNameTextChangedDelegate.IsBound())
+	{
+		return OnVerifyNameTextChangedDelegate.Execute(FText::FromName(InNewName), OutErrorMessage);
+	}
+	else
+	{
+		return true;
+	}
+}
+
+void SNiagaraParameterNameTextBlock::NameChanged(FName InNewName)
+{
+	OnNameTextCommittedDelegate.ExecuteIfBound(FText::FromName(InNewName), ETextCommit::OnEnter);
+}
+
+void SNiagaraParameterNameTextBlock::EnterEditingMode()
+{
+	ParameterName->EnterEditingMode();
+}
+
+void SNiagaraParameterNameTextBlock::EnterSubnamespaceEditingMode()
+{
+	ParameterName->EnterSubnamespaceEditingMode();
+}
\ No newline at end of file
diff --git a/Engine/Plugins/FX/Niagara/Source/NiagaraEditor/Private/TNiagaraGraphPinEditableName.h b/Engine/Plugins/FX/Niagara/Source/NiagaraEditor/Private/TNiagaraGraphPinEditableName.h
--- a/Engine/Plugins/FX/Niagara/Source/NiagaraEditor/Private/TNiagaraGraphPinEditableName.h
+++ b/Engine/Plugins/FX/Niagara/Source/NiagaraEditor/Private/TNiagaraGraphPinEditableName.h
@@ -6,6 +6,7 @@
 #include "NiagaraNode.h"
 #include "NiagaraNodeParameterMapBase.h"
 #include "Widgets/Text/SInlineEditableTextBlock.h"
+#include "Widgets/SNiagaraParameterName.h"
 #include "NiagaraNodeCustomHlsl.h"
 
 /** A graph pin widget for allowing a pin to have an editable name for a pin. */
@@ -70,9 +71,16 @@ protected:
 
 	virtual void Tick(const FGeometry& AllottedGeometry, const double InCurrentTime, const float InDeltaTime) override
 	{
-		if (bPendingRename && CreatedTextBlock.IsValid())
+		if (bPendingRename)
 		{
-			CreatedTextBlock->EnterEditingMode();
+			if (CreatedTextBlock.IsValid())
+			{
+				CreatedTextBlock->EnterEditingMode();
+			}
+			else if (CreatedParameterNamePinLabel.IsValid())
+			{
+				CreatedParameterNamePinLabel->EnterEditingMode();
+			}
 			bPendingRename = false;
 		}
 		BaseClass::Tick(AllottedGeometry, InCurrentTime, InDeltaTime);
@@ -117,7 +125,16 @@ protected:
 				{
 					return NiagaraGraph->GetPinVisualWidget(this->GraphPinObj);
 				}
-				return CreateLabelTextBlock();
+				else
+				{
+					CreatedParameterNamePinLabel = SNew(SNiagaraParameterNameTextBlock)
+						.EditableTextStyle(&FEditorStyle::Get().GetWidgetStyle<FInlineEditableTextBlockStyle>("Graph.Node.InlineEditablePinName"))
+						.ParameterText(this, &TNiagaraGraphPinEditableName<BaseClass>::GetParentPinLabel)
+						.Visibility(this, &TNiagaraGraphPinEditableName<BaseClass>::GetParentPinVisibility)
+						.OnVerifyTextChanged(this, &TNiagaraGraphPinEditableName<BaseClass>::OnVerifyTextChanged)
+						.OnTextCommitted(this, &TNiagaraGraphPinEditableName<BaseClass>::OnTextCommitted);
+					return CreatedParameterNamePinLabel.ToSharedRef();
+				}
 			}
 			else
 			{
@@ -132,4 +149,5 @@ protected:
 
 	bool bPendingRename;
 	TSharedPtr<SInlineEditableTextBlock> CreatedTextBlock;
+	TSharedPtr<SNiagaraParameterNameTextBlock> CreatedParameterNamePinLabel;
 };
diff --git a/Engine/Plugins/FX/Niagara/Source/NiagaraEditor/Private/ViewModels/Stack/NiagaraParameterHandle.cpp b/Engine/Plugins/FX/Niagara/Source/NiagaraEditor/Private/ViewModels/Stack/NiagaraParameterHandle.cpp
--- a/Engine/Plugins/FX/Niagara/Source/NiagaraEditor/Private/ViewModels/Stack/NiagaraParameterHandle.cpp
+++ b/Engine/Plugins/FX/Niagara/Source/NiagaraEditor/Private/ViewModels/Stack/NiagaraParameterHandle.cpp
@@ -122,6 +122,16 @@ bool FNiagaraParameterHandle::IsModuleHandle() const
 	return Namespace == FNiagaraConstants::ModuleNamespace;
 }
 
+bool FNiagaraParameterHandle::IsOutputHandle() const
+{
+	return Namespace == FNiagaraConstants::OutputNamespace;
+}
+
+bool FNiagaraParameterHandle::IsLocalHandle() const
+{
+	return Namespace == FNiagaraConstants::LocalNamespace;
+}
+
 bool FNiagaraParameterHandle::IsParameterCollectionHandle() const
 {
 	return Namespace == FNiagaraConstants::ParameterCollectionNamespace;
@@ -138,12 +148,11 @@ bool FNiagaraParameterHandle::IsReadOnlyHandle() const
 
 bool FNiagaraParameterHandle::IsTransientHandle() const
 {
-	return
-		Namespace != FNiagaraConstants::UserNamespace &&
-		Namespace != FNiagaraConstants::EngineNamespace &&
-		Namespace != FNiagaraConstants::SystemNamespace &&
-		Namespace != FNiagaraConstants::EmitterNamespace &&
-		Namespace != FNiagaraConstants::ParticleAttributeNamespace &&
-		Namespace != FNiagaraConstants::ModuleNamespace &&
-		Namespace != FNiagaraConstants::ParameterCollectionNamespace;
-}
\ No newline at end of file
+	return Namespace == FNiagaraConstants::TransientNamespace;
+}
+
+bool FNiagaraParameterHandle::IsDataInstanceHandle() const
+{
+	return Namespace == FNiagaraConstants::DataInstanceNamespace;
+}
+
diff --git a/Engine/Plugins/FX/Niagara/Source/NiagaraEditor/Private/Widgets/SNiagaraParameterMapPaletteItem.cpp b/Engine/Plugins/FX/Niagara/Source/NiagaraEditor/Private/Widgets/SNiagaraParameterMapPaletteItem.cpp
--- a/Engine/Plugins/FX/Niagara/Source/NiagaraEditor/Private/Widgets/SNiagaraParameterMapPaletteItem.cpp
+++ b/Engine/Plugins/FX/Niagara/Source/NiagaraEditor/Private/Widgets/SNiagaraParameterMapPaletteItem.cpp
@@ -9,6 +9,7 @@
 #include "Widgets/Input/SComboButton.h"
 #include "Widgets/Images/SImage.h"
 #include "ScopedTransaction.h"
+#include "Widgets/SNiagaraParameterName.h"
 
 #define LOCTEXT_NAMESPACE "NiagaraParameterMapPalleteItem"
 
@@ -39,7 +40,22 @@ void SNiagaraParameterMapPalleteItem::Construct(const FArguments& InArgs, FCreat
 	static const FName ItalicFontName = FName("Italic");
 	const FName FontType = ItalicFontName;
 	FSlateFontInfo Font = FCoreStyle::GetDefaultFontStyle(FontType, 10);
-	TSharedRef<SWidget> NameSlotWidget = CreateTextSlotWidget(Font, InCreateData, bIsReadOnly);
+
+	if (InCreateData->bHandleMouseButtonDown)
+	{
+		MouseButtonDownDelegate = InCreateData->MouseButtonDownDelegate;
+	}
+
+	ParameterNameTextBlock = SNew(SNiagaraParameterNameTextBlock)
+		.ParameterText(this, &SNiagaraParameterMapPalleteItem::GetDisplayText)
+		//.HighlightText(InCreateData->HighlightText)
+		.ToolTipText(this, &SNiagaraParameterMapPalleteItem::GetItemTooltip)
+		.OnTextCommitted(this, &SNiagaraParameterMapPalleteItem::OnNameTextCommitted)
+		.OnVerifyTextChanged(this, &SNiagaraParameterMapPalleteItem::OnNameTextVerifyChanged)
+		.IsSelected(InCreateData->IsRowSelectedDelegate)
+		.IsReadOnly(InCreateData->bIsReadOnly);
+
+	InCreateData->OnRenameRequest->BindSP(ParameterNameTextBlock.ToSharedRef(), &SNiagaraParameterNameTextBlock::EnterEditingMode);
 
 	// now, create the actual widget
 	ChildSlot
@@ -57,11 +73,11 @@ void SNiagaraParameterMapPalleteItem::Construct(const FArguments& InArgs, FCreat
 		+SHorizontalBox::Slot()
 		.FillWidth(1.f)
 		.VAlign(VAlign_Center)
-		.Padding(3,0)
+		.Padding(5,0)
 		[
-			NameSlotWidget
+			ParameterNameTextBlock.ToSharedRef()
 		]
-		// name slot
+		// reference count
 		+ SHorizontalBox::Slot()
 		.AutoWidth()
 		.VAlign(VAlign_Center)
@@ -85,6 +101,20 @@ void SNiagaraParameterMapPalleteItem::Construct(const FArguments& InArgs, FCreat
 	];
 }
 
+void SNiagaraParameterMapPalleteItem::Tick(const FGeometry& AllottedGeometry, const double InCurrentTime, const float InDeltaTime)
+{
+	SGraphPaletteItem::Tick(AllottedGeometry, InCurrentTime, InDeltaTime);
+	if (ParameterNameTextBlock.IsValid())
+	{
+		TSharedPtr<FNiagaraParameterAction> ParameterAction = StaticCastSharedPtr<FNiagaraParameterAction>(ActionPtr.Pin());
+		if (ParameterAction.IsValid() && ParameterAction->bSubnamespaceRenamePending)
+		{
+			ParameterAction->bSubnamespaceRenamePending = false;
+			ParameterNameTextBlock->EnterSubnamespaceEditingMode();
+		}
+	}
+}
+
 void SNiagaraParameterMapPalleteItem::OnNameTextCommitted(const FText& NewText, ETextCommit::Type InTextCommit)
 {
 	TSharedPtr<FNiagaraParameterAction> ParameterAction = StaticCastSharedPtr<FNiagaraParameterAction>(ActionPtr.Pin());
diff --git a/Engine/Plugins/FX/Niagara/Source/NiagaraEditor/Private/Widgets/SNiagaraParameterMapPaletteItem.h b/Engine/Plugins/FX/Niagara/Source/NiagaraEditor/Private/Widgets/SNiagaraParameterMapPaletteItem.h
--- a/Engine/Plugins/FX/Niagara/Source/NiagaraEditor/Private/Widgets/SNiagaraParameterMapPaletteItem.h
+++ b/Engine/Plugins/FX/Niagara/Source/NiagaraEditor/Private/Widgets/SNiagaraParameterMapPaletteItem.h
@@ -10,6 +10,8 @@
 #include "NiagaraActions.h"
 #include "EdGraph/EdGraphNodeUtils.h"
 
+class SNiagaraParameterNameTextBlock;
+
 class SNiagaraParameterMapPalleteItem : public SGraphPaletteItem
 {
 public:
@@ -23,11 +25,15 @@ public:
 
 	void Construct(const FArguments& InArgs, FCreateWidgetForActionData* const InCreateData);
 
+	virtual void Tick(const FGeometry& AllottedGeometry, const double InCurrentTime, const float InDeltaTime) override;
+
 protected:
 	/** Callback when rename text is committed */
 	virtual void OnNameTextCommitted(const FText& NewText, ETextCommit::Type InTextCommit) override;
 
 	FText GetReferenceCount() const;
+
 private:
 	FOnItemRenamed OnItemRenamed;
+	TSharedPtr<SNiagaraParameterNameTextBlock> ParameterNameTextBlock;
 };
\ No newline at end of file
diff --git a/Engine/Plugins/FX/Niagara/Source/NiagaraEditor/Private/Widgets/SNiagaraParameterMapView.cpp b/Engine/Plugins/FX/Niagara/Source/NiagaraEditor/Private/Widgets/SNiagaraParameterMapView.cpp
--- a/Engine/Plugins/FX/Niagara/Source/NiagaraEditor/Private/Widgets/SNiagaraParameterMapView.cpp
+++ b/Engine/Plugins/FX/Niagara/Source/NiagaraEditor/Private/Widgets/SNiagaraParameterMapView.cpp
@@ -39,47 +39,72 @@
 #include "NiagaraSystemEditorData.h"
 #include "ViewModels/Stack/NiagaraStackSystemSettingsGroup.h"
 #include "SNiagaraGraphActionWidget.h"
+#include "NiagaraEditorSettings.h"
 
 #define LOCTEXT_NAMESPACE "NiagaraParameterMapView"
 
 FText NiagaraParameterMapSectionID::OnGetSectionTitle(const NiagaraParameterMapSectionID::Type InSection)
 {
-	/* Setup an appropriate name for the section for this node */
-	FText SectionTitle;
+	TArray<FName> SectionNamespaces;
+	OnGetSectionNamespaces(InSection, SectionNamespaces);
+	FNiagaraNamespaceMetadata NamespaceMetadata = GetDefault<UNiagaraEditorSettings>()->GetMetaDataForNamespaces(SectionNamespaces);
+	if (NamespaceMetadata.IsValid())
+	{
+		return NamespaceMetadata.DisplayName;
+	}
+	else if (SectionNamespaces.Num() == 1)
+	{
+		return FText::FromName(SectionNamespaces[0]);
+	}
+	else
+	{
+		return NSLOCTEXT("GraphActionNode", "Unknown", "Unknown");
+	}
+}
+
+void NiagaraParameterMapSectionID::OnGetSectionNamespaces(const NiagaraParameterMapSectionID::Type InSection, TArray<FName>& OutSectionNamespaces)
+{
 	switch (InSection)
 	{
 	case NiagaraParameterMapSectionID::ENGINE:
-		SectionTitle = NSLOCTEXT("GraphActionNode", "Engine", "Engine");
+		OutSectionNamespaces.Add(FNiagaraConstants::EngineNamespace);
 		break;
 	case NiagaraParameterMapSectionID::EMITTER:
-		SectionTitle = NSLOCTEXT("GraphActionNode", "Emitter", "Emitter");
+		OutSectionNamespaces.Add(FNiagaraConstants::EmitterNamespace);
+		break;
+	case NiagaraParameterMapSectionID::MODULE_INPUT:
+		OutSectionNamespaces.Add(FNiagaraConstants::ModuleNamespace);
+		break;
+	case NiagaraParameterMapSectionID::MODULE_OUTPUT:
+		OutSectionNamespaces.Add(FNiagaraConstants::OutputNamespace);
+		OutSectionNamespaces.Add(FNiagaraConstants::ModuleNamespace);
+		break;
+	case NiagaraParameterMapSectionID::MODULE_LOCAL:
+		OutSectionNamespaces.Add(FNiagaraConstants::LocalNamespace);
+		OutSectionNamespaces.Add(FNiagaraConstants::ModuleNamespace);
 		break;
-	case NiagaraParameterMapSectionID::MODULE:
-		SectionTitle = NSLOCTEXT("GraphActionNode", "Module", "Module");
+	case NiagaraParameterMapSectionID::TRANSIENT:
+		OutSectionNamespaces.Add(FNiagaraConstants::TransientNamespace);
+		break;
+	case NiagaraParameterMapSectionID::DATA_INSTANCE:
+		OutSectionNamespaces.Add(FNiagaraConstants::DataInstanceNamespace);
 		break;
 	case NiagaraParameterMapSectionID::STATIC_SWITCH:
-		SectionTitle = NSLOCTEXT("GraphActionNode", "StaticSwitch", "Static Switch");
+		OutSectionNamespaces.Add(FNiagaraConstants::StaticSwitchNamespace);
 		break;
 	case NiagaraParameterMapSectionID::SYSTEM:
-		SectionTitle = NSLOCTEXT("GraphActionNode", "System", "System");
+		OutSectionNamespaces.Add(FNiagaraConstants::SystemNamespace);
 		break;
 	case NiagaraParameterMapSectionID::PARTICLE:
-		SectionTitle = NSLOCTEXT("GraphActionNode", "Particles", "Particles");
+		OutSectionNamespaces.Add(FNiagaraConstants::ParticleAttributeNamespace);
 		break;
 	case NiagaraParameterMapSectionID::USER:
-		SectionTitle = NSLOCTEXT("GraphActionNode", "User", "User");
+		OutSectionNamespaces.Add(FNiagaraConstants::UserNamespace);
 		break;
 	case NiagaraParameterMapSectionID::PARAMETERCOLLECTION:
-		SectionTitle = NSLOCTEXT("GraphActionNode", "ParameterCollection", "Parameter Collection");
-		break;
-	case NiagaraParameterMapSectionID::OTHER:
-		SectionTitle = NSLOCTEXT("GraphActionNode", "Other", "Other");
-		break;
-	case NiagaraParameterMapSectionID::NONE:
-		SectionTitle = FText::GetEmpty();
+		OutSectionNamespaces.Add(FNiagaraConstants::ParameterCollectionNamespace);
 		break;
 	}
-	return SectionTitle;
 }
 
 NiagaraParameterMapSectionID::Type NiagaraParameterMapSectionID::OnGetSectionFromVariable(const FNiagaraVariable& InVar, bool IsStaticSwitchVariable, FNiagaraParameterHandle& OutParameterHandle, const NiagaraParameterMapSectionID::Type DefaultType)
@@ -96,7 +121,15 @@ NiagaraParameterMapSectionID::Type NiagaraParameterMapSectionID::OnGetSectionFro
 	}
 	else if (OutParameterHandle.IsModuleHandle())
 	{
-		SectionID = NiagaraParameterMapSectionID::MODULE;
+		SectionID = NiagaraParameterMapSectionID::MODULE_INPUT;
+	}
+	else if (OutParameterHandle.IsOutputHandle())
+	{
+		SectionID = NiagaraParameterMapSectionID::MODULE_OUTPUT;
+	}
+	else if (OutParameterHandle.IsLocalHandle())
+	{
+		SectionID = NiagaraParameterMapSectionID::MODULE_LOCAL;
 	}
 	else if (OutParameterHandle.IsUserHandle())
 	{
@@ -118,6 +151,14 @@ NiagaraParameterMapSectionID::Type NiagaraParameterMapSectionID::OnGetSectionFro
 	{
 		SectionID = NiagaraParameterMapSectionID::PARAMETERCOLLECTION;
 	}
+	else if (OutParameterHandle.IsTransientHandle())
+	{
+		SectionID = NiagaraParameterMapSectionID::TRANSIENT;
+	}
+	else if (OutParameterHandle.IsDataInstanceHandle())
+	{
+		SectionID = NiagaraParameterMapSectionID::DATA_INSTANCE;
+	}
 
 	return SectionID;
 }
@@ -157,7 +198,7 @@ void SNiagaraParameterMapView::Construct(const FArguments& InArgs, const TArray<
 	bNeedsRefresh = false;
 	ToolkitType = InToolkitType;
 	ToolkitCommands = InToolkitCommands;
-	AddParameterButtons.SetNum(NiagaraParameterMapSectionID::OTHER + 1);
+	AddParameterButtons.SetNum(NiagaraParameterMapSectionID::Num);
 
 	SelectedScriptObjects = InSelectedObjects[0];
 	SelectedScriptObjects->OnSelectedObjectsChanged().AddSP(this, &SNiagaraParameterMapView::SelectedObjectsChanged);
@@ -391,8 +432,12 @@ void SNiagaraParameterMapView::CollectAllActions(FGraphActionListBuilderBase& Ou
 	{
 		const FNiagaraVariable& Parameter = ParameterEntry.Key;
 		FNiagaraParameterHandle Handle;
-		const NiagaraParameterMapSectionID::Type Section = NiagaraParameterMapSectionID::OnGetSectionFromVariable(Parameter, IsStaticSwitchParameter(Parameter, Graphs), Handle, /*Default*/ NiagaraParameterMapSectionID::OTHER);
-		if (!IsSystemToolkit() || (IsSystemToolkit() && Section != NiagaraParameterMapSectionID::MODULE))
+		const NiagaraParameterMapSectionID::Type Section = NiagaraParameterMapSectionID::OnGetSectionFromVariable(Parameter, IsStaticSwitchParameter(Parameter, Graphs), Handle);
+		if (Section == NiagaraParameterMapSectionID::NONE)
+		{
+			continue;
+		}
+		if (!IsSystemToolkit() || (IsSystemToolkit() && Section != NiagaraParameterMapSectionID::MODULE_INPUT && Section != NiagaraParameterMapSectionID::MODULE_LOCAL))
 		{
 			const FText Name = FText::FromName(Parameter.GetName());
 			const FText Tooltip = FText::Format(TooltipFormat, FText::FromName(Parameter.GetName()), Parameter.GetType().GetNameText());
@@ -408,16 +453,22 @@ void SNiagaraParameterMapView::CollectStaticSections(TArray<int32>& StaticSectio
 {
 	if (!IsSystemToolkit())
 	{
-		StaticSectionIDs.Add(NiagaraParameterMapSectionID::MODULE);
+		StaticSectionIDs.Add(NiagaraParameterMapSectionID::MODULE_INPUT);
+		StaticSectionIDs.Add(NiagaraParameterMapSectionID::MODULE_OUTPUT);
+		StaticSectionIDs.Add(NiagaraParameterMapSectionID::MODULE_LOCAL);
+		StaticSectionIDs.Add(NiagaraParameterMapSectionID::DATA_INSTANCE);
+	}
+	else
+	{
+		StaticSectionIDs.Add(NiagaraParameterMapSectionID::USER);
 	}
 	StaticSectionIDs.Add(NiagaraParameterMapSectionID::STATIC_SWITCH);
 	StaticSectionIDs.Add(NiagaraParameterMapSectionID::ENGINE);
 	StaticSectionIDs.Add(NiagaraParameterMapSectionID::PARAMETERCOLLECTION);
-	StaticSectionIDs.Add(NiagaraParameterMapSectionID::USER);
 	StaticSectionIDs.Add(NiagaraParameterMapSectionID::SYSTEM);
 	StaticSectionIDs.Add(NiagaraParameterMapSectionID::EMITTER);
 	StaticSectionIDs.Add(NiagaraParameterMapSectionID::PARTICLE);
-	StaticSectionIDs.Add(NiagaraParameterMapSectionID::OTHER);
+	StaticSectionIDs.Add(NiagaraParameterMapSectionID::TRANSIENT);
 }
 
 FReply SNiagaraParameterMapView::OnActionDragged(const TArray<TSharedPtr<FEdGraphSchemaAction>>& InActions, const FPointerEvent& MouseEvent)
@@ -490,12 +541,48 @@ TSharedPtr<SWidget> SNiagaraParameterMapView::OnContextMenuOpening()
 	{
 		const bool bShouldCloseWindowAfterMenuSelection = true;
 		FMenuBuilder MenuBuilder(bShouldCloseWindowAfterMenuSelection, ToolkitCommands);
-		MenuBuilder.BeginSection("BasicOperations");
+		MenuBuilder.BeginSection("Edit", LOCTEXT("EditMenuHeader", "Edit"));
 		{
 			MenuBuilder.AddMenuEntry(FGenericCommands::Get().Rename, NAME_None, LOCTEXT("Rename", "Rename"), LOCTEXT("Rename_Tooltip", "Renames this parameter"));
 			MenuBuilder.AddMenuEntry(FNiagaraParameterMapViewCommands::Get().DeleteEntry);
 		}
 		MenuBuilder.EndSection();
+
+		TArray<TSharedPtr<FEdGraphSchemaAction>> SelectedActions;
+		GraphActionMenu->GetSelectedActions(SelectedActions);
+
+		if (SelectedActions.Num() == 1)
+		{
+			TSharedPtr<FNiagaraParameterAction> SelectedNiagaraAction = StaticCastSharedPtr<FNiagaraParameterAction>(SelectedActions[0]);
+			if (SelectedNiagaraAction.IsValid())
+			{
+				MenuBuilder.BeginSection("Subnamespace", LOCTEXT("SubnamespaceMenuSection", "Parameter Modifiers"));
+				{
+					MenuBuilder.AddMenuEntry(
+						LOCTEXT("AddSubnamespace", "Add a modifier to this parameter."),
+						FText(),
+						FSlateIcon(),
+						FUIAction(
+							FExecuteAction::CreateSP(this, &SNiagaraParameterMapView::OnAddSubnamespace),
+							FCanExecuteAction()));
+					MenuBuilder.AddMenuEntry(
+						LOCTEXT("RemoveSubnamespace", "Remove the modifier from this parameter."),
+						FText(),
+						FSlateIcon(),
+						FUIAction(
+							FExecuteAction::CreateSP(this, &SNiagaraParameterMapView::OnRemoveSubnamespace),
+							FCanExecuteAction()));
+					MenuBuilder.AddMenuEntry(
+						LOCTEXT("EditSubnamespace", "Edit the modifier for this parameter."),
+						FText(),
+						FSlateIcon(),
+						FUIAction(
+							FExecuteAction::CreateSP(this, &SNiagaraParameterMapView::OnEditSubnamespace),
+							FCanExecuteAction()));
+				}
+				MenuBuilder.EndSection();
+			}
+		}
 		return MenuBuilder.MakeWidget();
 	}
 
@@ -779,6 +866,82 @@ void SNiagaraParameterMapView::OnPostRenameActionNode(const FText& InText, FNiag
 	}
 }
 
+void SNiagaraParameterMapView::OnAddSubnamespace()
+{
+	TArray<TSharedPtr<FEdGraphSchemaAction>> SelectedActions;
+	GraphActionMenu->GetSelectedActions(SelectedActions);
+	if (SelectedActions.Num() == 1)
+	{
+		TSharedPtr<FNiagaraParameterAction> ParameterAction = StaticCastSharedPtr<FNiagaraParameterAction>(SelectedActions[0]);
+		if (ParameterAction.IsValid())
+		{
+			FString ParameterNameString = ParameterAction->Parameter.GetName().ToString();
+			TArray<FString> NameParts;
+			ParameterNameString.ParseIntoArray(NameParts, TEXT("."));
+			if (NameParts.Num() == 2)
+			{
+				NameParts.Insert(FNiagaraConstants::ModuleNamespace.ToString(), 1);
+				FName NewName = *FString::Join(NameParts, TEXT("."));
+				if (Graphs.Num() > 0)
+				{
+					for (const TWeakObjectPtr<UNiagaraGraph>& Graph : Graphs)
+					{
+						if (Graph.IsValid())
+						{
+							Graph.Get()->RenameParameter(ParameterAction->Parameter, NewName);
+						}
+					}
+				}
+			}
+		}
+	}
+}
+
+void SNiagaraParameterMapView::OnRemoveSubnamespace()
+{
+	TArray<TSharedPtr<FEdGraphSchemaAction>> SelectedActions;
+	GraphActionMenu->GetSelectedActions(SelectedActions);
+	if (SelectedActions.Num() == 1)
+	{
+		TSharedPtr<FNiagaraParameterAction> ParameterAction = StaticCastSharedPtr<FNiagaraParameterAction>(SelectedActions[0]);
+		if (ParameterAction.IsValid())
+		{
+			FString ParameterNameString = ParameterAction->Parameter.GetName().ToString();
+			TArray<FString> NameParts;
+			ParameterNameString.ParseIntoArray(NameParts, TEXT("."));
+			if (NameParts.Num() == 3)
+			{
+				NameParts.RemoveAt(1);
+				FName NewName = *FString::Join(NameParts, TEXT("."));
+				if (Graphs.Num() > 0)
+				{
+					for (const TWeakObjectPtr<UNiagaraGraph>& Graph : Graphs)
+					{
+						if (Graph.IsValid())
+						{
+							Graph.Get()->RenameParameter(ParameterAction->Parameter, NewName);
+						}
+					}
+				}
+			}
+		}
+	}
+}
+
+void SNiagaraParameterMapView::OnEditSubnamespace()
+{
+	TArray<TSharedPtr<FEdGraphSchemaAction>> SelectedActions;
+	GraphActionMenu->GetSelectedActions(SelectedActions);
+	if (SelectedActions.Num() == 1)
+	{
+		TSharedPtr<FNiagaraParameterAction> ParameterAction = StaticCastSharedPtr<FNiagaraParameterAction>(SelectedActions[0]);
+		if (ParameterAction.IsValid())
+		{
+			ParameterAction->bSubnamespaceRenamePending = true;
+		}
+	}
+}
+
 bool SNiagaraParameterMapView::IsSystemToolkit()
 {
 	return ToolkitType == EToolkitType::SYSTEM;
@@ -961,9 +1124,27 @@ void SNiagaraAddParameterMenu::CollectAllActions(FGraphActionListBuilderBase& Ou
 	}
 
 	// Module
-	if (CanCollectSection(NiagaraParameterMapSectionID::MODULE) && !IDsExcluded.Contains(NiagaraParameterMapSectionID::MODULE))
+	if (CanCollectSection(NiagaraParameterMapSectionID::MODULE_INPUT) && !IDsExcluded.Contains(NiagaraParameterMapSectionID::MODULE_INPUT))
+	{
+		CollectMakeNew(OutAllActions, NiagaraParameterMapSectionID::MODULE_INPUT);
+	}
+
+	// Module Output
+	if (CanCollectSection(NiagaraParameterMapSectionID::MODULE_OUTPUT) && !IDsExcluded.Contains(NiagaraParameterMapSectionID::MODULE_OUTPUT))
 	{
-		CollectMakeNew(OutAllActions, NiagaraParameterMapSectionID::MODULE);
+		CollectMakeNew(OutAllActions, NiagaraParameterMapSectionID::MODULE_OUTPUT);
+	}
+
+	// Module Local
+	if (CanCollectSection(NiagaraParameterMapSectionID::MODULE_LOCAL) && !IDsExcluded.Contains(NiagaraParameterMapSectionID::MODULE_LOCAL))
+	{
+		CollectMakeNew(OutAllActions, NiagaraParameterMapSectionID::MODULE_LOCAL);
+	}
+
+	// Transient
+	if (CanCollectSection(NiagaraParameterMapSectionID::TRANSIENT) && !IDsExcluded.Contains(NiagaraParameterMapSectionID::TRANSIENT))
+	{
+		CollectMakeNew(OutAllActions, NiagaraParameterMapSectionID::TRANSIENT);
 	}
 
 	// System
@@ -984,11 +1165,6 @@ void SNiagaraAddParameterMenu::CollectAllActions(FGraphActionListBuilderBase& Ou
 		CollectParameterCollectionsActions(OutAllActions);
 	}
 
-	if (CanCollectSection(NiagaraParameterMapSectionID::OTHER))
-	{
-		CollectMakeNew(OutAllActions, NiagaraParameterMapSectionID::OTHER);
-	}
-
 	// Engine
 	if (CanCollectSection(NiagaraParameterMapSectionID::ENGINE) && !IDsExcluded.Contains(NiagaraParameterMapSectionID::ENGINE))
 	{
@@ -1014,7 +1190,7 @@ void SNiagaraAddParameterMenu::CollectAllActions(FGraphActionListBuilderBase& Ou
 				const FNiagaraVariable& Parameter = ParameterEntry.Key;
 				FNiagaraParameterHandle Handle;
 				bool IsStaticSwitch = SNiagaraParameterMapView::IsStaticSwitchParameter(Parameter, Graphs);
-				const NiagaraParameterMapSectionID::Type ParameterSectionID = NiagaraParameterMapSectionID::OnGetSectionFromVariable(Parameter, IsStaticSwitch, Handle, /*Default*/ NiagaraParameterMapSectionID::OTHER);
+				const NiagaraParameterMapSectionID::Type ParameterSectionID = NiagaraParameterMapSectionID::OnGetSectionFromVariable(Parameter, IsStaticSwitch, Handle);
 				if (CanCollectSection(ParameterSectionID))
 				{
 					if (IDsExcluded.Contains(ParameterSectionID))
@@ -1146,8 +1322,15 @@ void SNiagaraAddParameterMenu::AddParameterSelected(FNiagaraVariable NewVariable
 	if (bCreateCustomName)
 	{
 		const static FString NewVariableDefaultName = FString("NewVariable");
-		const FString ResultName = (InSection != NiagaraParameterMapSectionID::NONE ? NiagaraParameterMapSectionID::OnGetSectionTitle(InSection).ToString() + TEXT(".") : FString())
-			+ NewVariableDefaultName;
+		TArray<FName> SectionNamespaces;
+		NiagaraParameterMapSectionID::OnGetSectionNamespaces(InSection, SectionNamespaces);
+		TArray<FString> NameParts;
+		for (FName SectionNamespace : SectionNamespaces)
+		{
+			NameParts.Add(SectionNamespace.ToString());
+		}
+		NameParts.Add(NewVariableDefaultName);
+		const FString ResultName = FString::Join(NameParts, TEXT("."));
 		NewVariable.SetName(FName(*ResultName));
 	}
 
diff --git a/Engine/Plugins/FX/Niagara/Source/NiagaraEditor/Public/NiagaraActions.h b/Engine/Plugins/FX/Niagara/Source/NiagaraEditor/Public/NiagaraActions.h
--- a/Engine/Plugins/FX/Niagara/Source/NiagaraEditor/Public/NiagaraActions.h
+++ b/Engine/Plugins/FX/Niagara/Source/NiagaraEditor/Public/NiagaraActions.h
@@ -55,7 +55,11 @@ struct NIAGARAEDITOR_API FNiagaraScriptVarAndViewInfoAction : public FEdGraphSch
 
 struct NIAGARAEDITOR_API FNiagaraParameterAction : public FEdGraphSchemaAction
 {
-	FNiagaraParameterAction() {}
+	FNiagaraParameterAction()
+		: bSubnamespaceRenamePending(false)
+	{
+	}
+
 	FNiagaraParameterAction(const FNiagaraVariable& InParameter,
 		const TArray<FNiagaraGraphParameterReferenceCollection>& InReferenceCollection,
 		FText InNodeCategory, FText InMenuDesc, FText InToolTip, const int32 InGrouping, FText InKeywords, int32 InSectionID = 0);
@@ -67,6 +71,8 @@ struct NIAGARAEDITOR_API FNiagaraParameterAction : public FEdGraphSchemaAction
 	FNiagaraVariable Parameter;
 
 	TArray<FNiagaraGraphParameterReferenceCollection> ReferenceCollection;
+
+	bool bSubnamespaceRenamePending;
 };
 
 struct NIAGARAEDITOR_API FNiagaraScriptParameterAction : public FEdGraphSchemaAction
diff --git a/Engine/Plugins/FX/Niagara/Source/NiagaraEditor/Public/NiagaraEditorSettings.h b/Engine/Plugins/FX/Niagara/Source/NiagaraEditor/Public/NiagaraEditorSettings.h
--- a/Engine/Plugins/FX/Niagara/Source/NiagaraEditor/Public/NiagaraEditorSettings.h
+++ b/Engine/Plugins/FX/Niagara/Source/NiagaraEditor/Public/NiagaraEditorSettings.h
@@ -27,6 +27,41 @@ struct FNiagaraNewAssetDialogConfig
 	}
 };
 
+UENUM()
+enum class ENiagaraNamespaceMetadataOptions
+{
+	Advanced,
+	PreventRenaming,
+	SupportsSubnamespace
+};
+
+USTRUCT()
+struct FNiagaraNamespaceMetadata
+{
+	GENERATED_BODY()
+
+	FNiagaraNamespaceMetadata();
+
+	FNiagaraNamespaceMetadata(TArray<FName> InNamespaces, FText InDisplayName, FText InDescription, FLinearColor InBackgroundColor, TArray<ENiagaraNamespaceMetadataOptions> InOptions);
+
+	UPROPERTY()
+	TArray<FName> Namespaces;
+
+	UPROPERTY()
+	FText DisplayName;
+
+	UPROPERTY()
+	FText Description;
+
+	UPROPERTY()
+	FLinearColor BackgroundColor;
+
+	UPROPERTY()
+	TArray<ENiagaraNamespaceMetadataOptions> Options;
+
+	bool IsValid() const { return Namespaces.Num() > 0; }
+};
+
 UCLASS(config = Niagara, defaultconfig, meta=(DisplayName="Niagara"))
 class UNiagaraEditorSettings : public UDeveloperSettings
 {
@@ -90,6 +125,8 @@ public:
 	FNiagaraNewAssetDialogConfig GetNewAssetDailogConfig(FName InDialogConfigKey) const;
 
 	void SetNewAssetDialogConfig(FName InDialogConfigKey, const FNiagaraNewAssetDialogConfig& InNewAssetDialogConfig);
+
+	FNiagaraNamespaceMetadata GetMetaDataForNamespaces(TArray<FName> Namespaces) const;
 	
 	// Begin UDeveloperSettings Interface
 	virtual FName GetCategoryName() const override;
@@ -104,6 +141,7 @@ public:
 	static FOnNiagaraEditorSettingsChanged& OnSettingsChanged();
 
 	const TMap<FString, FString>& GetHLSLKeywordReplacementsMap()const { return HLSLKeywordReplacements; }
+
 protected:
 	static FOnNiagaraEditorSettingsChanged SettingsChangedDelegate;
 
@@ -133,4 +171,7 @@ private:
 
 	UPROPERTY(config)
 	TMap<FString, FString> HLSLKeywordReplacements;
+
+	UPROPERTY()
+	TArray<FNiagaraNamespaceMetadata> NamespaceMetadata;
 };
diff --git a/Engine/Plugins/FX/Niagara/Source/NiagaraEditor/Public/NiagaraGraph.h b/Engine/Plugins/FX/Niagara/Source/NiagaraEditor/Public/NiagaraGraph.h
--- a/Engine/Plugins/FX/Niagara/Source/NiagaraEditor/Public/NiagaraGraph.h
+++ b/Engine/Plugins/FX/Niagara/Source/NiagaraEditor/Public/NiagaraGraph.h
@@ -426,7 +426,10 @@ private:
 
 	/** Compares the values on the default pins with the metadata and syncs the two if necessary */
 	void ValidateDefaultPins();
-	
+
+	void StandardizeParameterNames();
+
+private:
 	/** The current change identifier for this graph overall. Used to sync status with UNiagaraScripts.*/
 	UPROPERTY()
 	FGuid ChangeId;
diff --git a/Engine/Plugins/FX/Niagara/Source/NiagaraEditor/Public/ViewModels/Stack/NiagaraParameterHandle.h b/Engine/Plugins/FX/Niagara/Source/NiagaraEditor/Public/ViewModels/Stack/NiagaraParameterHandle.h
--- a/Engine/Plugins/FX/Niagara/Source/NiagaraEditor/Public/ViewModels/Stack/NiagaraParameterHandle.h
+++ b/Engine/Plugins/FX/Niagara/Source/NiagaraEditor/Public/ViewModels/Stack/NiagaraParameterHandle.h
@@ -53,12 +53,18 @@ public:
 
 	bool IsModuleHandle() const;
 
+	bool IsOutputHandle() const;
+
+	bool IsLocalHandle() const;
+
 	bool IsParameterCollectionHandle() const;
 
 	bool IsReadOnlyHandle() const;
 
 	bool IsTransientHandle() const;
 
+	bool IsDataInstanceHandle() const;
+
 private:
 	FName ParameterHandleName;
 	FName Name;
diff --git a/Engine/Plugins/FX/Niagara/Source/NiagaraEditor/Public/Widgets/SNiagaraParameterMapView.h b/Engine/Plugins/FX/Niagara/Source/NiagaraEditor/Public/Widgets/SNiagaraParameterMapView.h
--- a/Engine/Plugins/FX/Niagara/Source/NiagaraEditor/Public/Widgets/SNiagaraParameterMapView.h
+++ b/Engine/Plugins/FX/Niagara/Source/NiagaraEditor/Public/Widgets/SNiagaraParameterMapView.h
@@ -28,19 +28,28 @@ namespace NiagaraParameterMapSectionID
 {
 	enum Type
 	{
-		NONE = 0,
-		MODULE,
+		NONE,
+		MODULE_INPUT,
+		MODULE_OUTPUT,
+		MODULE_LOCAL,
+		TRANSIENT,
+		DATA_INSTANCE,
 		STATIC_SWITCH,
 		ENGINE,
 		PARAMETERCOLLECTION,
 		USER,
 		SYSTEM,
+		SYSTEM_MODULE,
 		EMITTER,
+		EMITTER_MODULE,
 		PARTICLE,
-		OTHER,
+		PARTICLE_MODULE,
+
+		Num
 	};
 
 	static FText OnGetSectionTitle(const NiagaraParameterMapSectionID::Type InSection);
+	static void OnGetSectionNamespaces(const NiagaraParameterMapSectionID::Type InSection, TArray<FName>& OutSectionNamespaces);
 	static NiagaraParameterMapSectionID::Type OnGetSectionFromVariable(const FNiagaraVariable& InVar, bool IsStaticSwitchVariable, FNiagaraParameterHandle& OutParameterHandle, const NiagaraParameterMapSectionID::Type DefaultType = NiagaraParameterMapSectionID::Type::NONE);
 };
 
@@ -135,6 +144,9 @@ private:
 	bool CanRequestRenameOnActionNode(TWeakPtr<struct FGraphActionNode> InSelectedNode) const;
 	bool CanRequestRenameOnActionNode() const;
 	void OnPostRenameActionNode(const FText& InText, FNiagaraParameterAction& InAction);
+	void OnAddSubnamespace();
+	void OnRemoveSubnamespace();
+	void OnEditSubnamespace();
 
 	bool IsSystemToolkit();
 	bool IsScriptToolkit();
diff --git a/Engine/Plugins/FX/Niagara/Source/NiagaraEditor/Public/Widgets/SNiagaraParameterName.h b/Engine/Plugins/FX/Niagara/Source/NiagaraEditor/Public/Widgets/SNiagaraParameterName.h
new file mode 100644
--- /dev/null
+++ b/Engine/Plugins/FX/Niagara/Source/NiagaraEditor/Public/Widgets/SNiagaraParameterName.h
@@ -0,0 +1,111 @@
+// Copyright Epic Games, Inc. All Rights Reserved.
+
+#pragma once
+
+#include "Widgets/SCompoundWidget.h"
+#include "EditorStyleSet.h"
+
+class SInlineEditableTextBlock;
+
+class NIAGARAEDITOR_API SNiagaraParameterName : public SCompoundWidget
+{
+public:
+	DECLARE_DELEGATE_OneParam(FOnNameChanged, FName /* InNewName */);
+	DECLARE_DELEGATE_RetVal_TwoParams(bool, FOnVerifyNameChange, FName /* InNewName */, FText& /* OutErrorMessage */)
+
+public:
+	SLATE_BEGIN_ARGS(SNiagaraParameterName) 
+		: _EditableTextStyle(&FEditorStyle::Get().GetWidgetStyle<FInlineEditableTextBlockStyle>("InlineEditableTextBlockStyle"))
+		, _ReadOnlyTextStyle(&FEditorStyle::Get().GetWidgetStyle<FTextBlockStyle>("NormalText"))
+		, _IsReadOnly(false)
+	{}
+	SLATE_STYLE_ARGUMENT(FInlineEditableTextBlockStyle, EditableTextStyle)
+		SLATE_STYLE_ARGUMENT(FTextBlockStyle, ReadOnlyTextStyle)
+		SLATE_ATTRIBUTE(FName, ParameterName)
+		SLATE_ARGUMENT(bool, IsReadOnly)
+		SLATE_EVENT(FOnVerifyNameChange, OnVerifyNameChange)
+		SLATE_EVENT(FOnNameChanged, OnNameChanged)
+		SLATE_EVENT(FIsSelected, IsSelected)
+		SLATE_EVENT(FPointerEventHandler, OnDoubleClicked)
+	SLATE_END_ARGS()
+
+	void Construct(const FArguments& InArgs);
+
+	virtual void Tick(const FGeometry& AllottedGeometry, const double InCurrentTime, const float InDeltaTime) override;
+
+	virtual FReply OnMouseButtonDoubleClick(const FGeometry& InMyGeometry, const FPointerEvent& InMouseEvent) override;
+
+	void EnterEditingMode();
+
+	void EnterSubnamespaceEditingMode();
+
+private:
+	TSharedRef<SBorder> CreateNamespaceWidget(FText NamespaceDisplayName, FText NamespaceDescription, FLinearColor NamespaceBorderColor);
+
+	void UpdateContent(FName InDisplayedParameterName);
+
+	FName ReconstructNameFromEditText(const FText& InEditText);
+
+	FName ReconstructNameFromSubnamespaceEditText(const FText& InEditText);
+
+	FText GetSubnamespaceText();
+
+	bool VerifyNameTextChange(const FText& InNewNameText, FText& OutErrorMessage);
+
+	void NameTextCommitted(const FText& InNewNameText, ETextCommit::Type CommitType);
+
+	bool VerifySubnamespaceTextChange(const FText& InNewNameText, FText& OutErrorMessage);
+
+	void SubnamespaceTextCommitted(const FText& InNewNameText, ETextCommit::Type CommitType);
+
+private:
+	const FInlineEditableTextBlockStyle* EditableTextStyle;
+	const FTextBlockStyle* ReadOnlyTextStyle;
+	TAttribute<FName> ParameterName;
+	bool bIsReadOnly;
+	FOnVerifyNameChange OnVerifyNameChangeDelegate;
+	FOnNameChanged OnNameChangedDelegate;
+	FPointerEventHandler OnDoubleClickedDelegate;
+	FIsSelected IsSelected;
+	FName DisplayedParameterName;
+	FName DisplayedSubnamespace;
+	TSharedPtr<SInlineEditableTextBlock> EditableTextBlock;
+	TSharedPtr<SBorder> SubnamespaceBorder;
+};
+
+class NIAGARAEDITOR_API SNiagaraParameterNameTextBlock : public SCompoundWidget
+{
+public:
+	SLATE_BEGIN_ARGS(SNiagaraParameterNameTextBlock)
+		: _EditableTextStyle(&FEditorStyle::Get().GetWidgetStyle<FInlineEditableTextBlockStyle>("InlineEditableTextBlockStyle"))
+	{}
+		SLATE_ATTRIBUTE(FText, ParameterText)
+		SLATE_STYLE_ARGUMENT(FInlineEditableTextBlockStyle, EditableTextStyle)
+		SLATE_ARGUMENT(bool, IsReadOnly)
+		SLATE_EVENT(FOnVerifyTextChanged, OnVerifyTextChanged)
+		SLATE_EVENT(FOnTextCommitted, OnTextCommitted)
+		SLATE_EVENT(FIsSelected, IsSelected)
+	SLATE_END_ARGS()
+
+	void Construct(const FArguments& InArgs);
+
+	void EnterEditingMode();
+
+	void EnterSubnamespaceEditingMode();
+
+private:
+	FName GetParameterName() const;
+
+	bool VerifyNameChange(FName InNewName, FText& OutErrorMessage);
+
+	void NameChanged(FName InNewName);
+
+private:
+	TAttribute<FText> ParameterText;
+	FOnVerifyTextChanged OnVerifyNameTextChangedDelegate;
+	FOnTextCommitted OnNameTextCommittedDelegate;
+
+	mutable FText DisplayedParameterTextCache;
+	mutable FName ParameterNameCache;
+	TSharedPtr<SNiagaraParameterName> ParameterName;
+};
\ No newline at end of file
diff --git a/Engine/Plugins/FX/Niagara/Source/NiagaraEditorWidgets/Private/SNiagaraStack.cpp b/Engine/Plugins/FX/Niagara/Source/NiagaraEditorWidgets/Private/SNiagaraStack.cpp
--- a/Engine/Plugins/FX/Niagara/Source/NiagaraEditorWidgets/Private/SNiagaraStack.cpp
+++ b/Engine/Plugins/FX/Niagara/Source/NiagaraEditorWidgets/Private/SNiagaraStack.cpp
@@ -964,7 +964,8 @@ SNiagaraStack::FRowWidgets SNiagaraStack::ConstructNameAndValueWidgetsForItem(UN
 	{
 		UNiagaraStackParameterStoreEntry* StackEntry = CastChecked<UNiagaraStackParameterStoreEntry>(Item);
 		return FRowWidgets(
-			SNew(SNiagaraStackParameterStoreEntryName, StackEntry, StackViewModel),
+			SNew(SNiagaraStackParameterStoreEntryName, StackEntry, StackViewModel)
+			.IsSelected(Container, &SNiagaraStackTableRow::IsSelected),
 			SNew(SNiagaraStackParameterStoreEntryValue, StackEntry));
 	}
 	else if (Item->IsA<UNiagaraStackInputCategory>())
diff --git a/Engine/Plugins/FX/Niagara/Source/NiagaraEditorWidgets/Private/Stack/SNiagaraStackFunctionInputName.cpp b/Engine/Plugins/FX/Niagara/Source/NiagaraEditorWidgets/Private/Stack/SNiagaraStackFunctionInputName.cpp
--- a/Engine/Plugins/FX/Niagara/Source/NiagaraEditorWidgets/Private/Stack/SNiagaraStackFunctionInputName.cpp
+++ b/Engine/Plugins/FX/Niagara/Source/NiagaraEditorWidgets/Private/Stack/SNiagaraStackFunctionInputName.cpp
@@ -8,6 +8,9 @@
 #include "Widgets/Text/SInlineEditableTextBlock.h"
 #include "Widgets/SBoxPanel.h"
 #include "Widgets/Input/SCheckBox.h"
+#include "Widgets/SNiagaraParameterName.h"
+#include "NiagaraNodeAssignment.h"
+#include "NiagaraNodeFunctionCall.h"
 
 void SNiagaraStackFunctionInputName::Construct(const FArguments& InArgs, UNiagaraStackFunctionInput* InFunctionInput, UNiagaraStackViewModel* InStackViewModel)
 {
@@ -16,24 +19,44 @@ void SNiagaraStackFunctionInputName::Construct(const FArguments& InArgs, UNiagar
 	StackEntryItem = InFunctionInput;
 	IsSelected = InArgs._IsSelected;
 
+	TSharedPtr<SHorizontalBox> NameBox;
 	ChildSlot
 	[
-		SNew(SHorizontalBox)
+		SAssignNew(NameBox, SHorizontalBox)
 		.IsEnabled_UObject(FunctionInput, &UNiagaraStackEntry::GetOwnerIsEnabled)
+	];
 
-		// Edit condition checkbox
-		+ SHorizontalBox::Slot()
-		.AutoWidth()
-		.Padding(0, 0, 3, 0)
+	// Edit condition checkbox
+	NameBox->AddSlot()
+	.AutoWidth()
+	.Padding(0, 0, 3, 0)
+	[
+		SNew(SCheckBox)
+		.Visibility(this, &SNiagaraStackFunctionInputName::GetEditConditionCheckBoxVisibility)
+		.IsChecked(this, &SNiagaraStackFunctionInputName::GetEditConditionCheckState)
+		.OnCheckStateChanged(this, &SNiagaraStackFunctionInputName::OnEditConditionCheckStateChanged)
+	];
+	
+	// Name Label
+	if(FunctionInput->GetInputFunctionCallNode().IsA<UNiagaraNodeAssignment>())
+	{
+		NameBox->AddSlot()
+		.VAlign(VAlign_Center)
 		[
-			SNew(SCheckBox)
-			.Visibility(this, &SNiagaraStackFunctionInputName::GetEditConditionCheckBoxVisibility)
-			.IsChecked(this, &SNiagaraStackFunctionInputName::GetEditConditionCheckState)
-			.OnCheckStateChanged(this, &SNiagaraStackFunctionInputName::OnEditConditionCheckStateChanged)
-		]
-
-		// Name Label
-		+ SHorizontalBox::Slot()
+			SAssignNew(ParameterTextBlock, SNiagaraParameterNameTextBlock)
+			.EditableTextStyle(FNiagaraEditorStyle::Get(), "NiagaraEditor.ParameterInlineEditableText")
+			.ParameterText_UObject(FunctionInput, &UNiagaraStackEntry::GetDisplayName)
+			.IsReadOnly(FunctionInput->SupportsRename() == false)
+			.IsEnabled(this, &SNiagaraStackFunctionInputName::GetIsEnabled)
+			.IsSelected(this, &SNiagaraStackFunctionInputName::GetIsNameWidgetSelected)
+			.OnTextCommitted(this, &SNiagaraStackFunctionInputName::OnNameTextCommitted)
+			//.HighlightText_UObject(InStackViewModel, &UNiagaraStackViewModel::GetCurrentSearchText)
+			.ToolTipText(this, &SNiagaraStackFunctionInputName::GetToolTipText)
+		];
+	}
+	else
+	{
+		NameBox->AddSlot()
 		.VAlign(VAlign_Center)
 		[
 			SAssignNew(NameTextBlock, SInlineEditableTextBlock)
@@ -45,15 +68,22 @@ void SNiagaraStackFunctionInputName::Construct(const FArguments& InArgs, UNiagar
 			.OnTextCommitted(this, &SNiagaraStackFunctionInputName::OnNameTextCommitted)
 			.HighlightText_UObject(InStackViewModel, &UNiagaraStackViewModel::GetCurrentSearchText)
 			.ToolTipText(this, &SNiagaraStackFunctionInputName::GetToolTipText)
-		]
-	];
+		];
+	}
 }
 
 void SNiagaraStackFunctionInputName::Tick(const FGeometry& AllottedGeometry, const double InCurrentTime, const float InDeltaTime)
 {
-	if (FunctionInput->GetIsRenamePending() && NameTextBlock.IsValid())
+	if (FunctionInput->GetIsRenamePending())
 	{
-		NameTextBlock->EnterEditingMode();
+		if (NameTextBlock.IsValid())
+		{
+			NameTextBlock->EnterEditingMode();
+		}
+		else if (ParameterTextBlock.IsValid())
+		{
+			ParameterTextBlock->EnterEditingMode();
+		}
 		FunctionInput->SetIsRenamePending(false);
 	}
 	SCompoundWidget::Tick(AllottedGeometry, InCurrentTime, InDeltaTime);
diff --git a/Engine/Plugins/FX/Niagara/Source/NiagaraEditorWidgets/Private/Stack/SNiagaraStackFunctionInputName.h b/Engine/Plugins/FX/Niagara/Source/NiagaraEditorWidgets/Private/Stack/SNiagaraStackFunctionInputName.h
--- a/Engine/Plugins/FX/Niagara/Source/NiagaraEditorWidgets/Private/Stack/SNiagaraStackFunctionInputName.h
+++ b/Engine/Plugins/FX/Niagara/Source/NiagaraEditorWidgets/Private/Stack/SNiagaraStackFunctionInputName.h
@@ -8,6 +8,7 @@
 
 class UNiagaraStackFunctionInput;
 class SInlineEditableTextBlock;
+class SNiagaraParameterNameTextBlock;
 
 class SNiagaraStackFunctionInputName: public SNiagaraStackEntryWidget
 {
@@ -45,5 +46,7 @@ private:
 
 	TSharedPtr<SInlineEditableTextBlock> NameTextBlock;
 
+	TSharedPtr<SNiagaraParameterNameTextBlock> ParameterTextBlock;
+
 	TAttribute<bool> IsSelected;
 };
\ No newline at end of file
diff --git a/Engine/Plugins/FX/Niagara/Source/NiagaraEditorWidgets/Private/Stack/SNiagaraStackFunctionInputValue.cpp b/Engine/Plugins/FX/Niagara/Source/NiagaraEditorWidgets/Private/Stack/SNiagaraStackFunctionInputValue.cpp
--- a/Engine/Plugins/FX/Niagara/Source/NiagaraEditorWidgets/Private/Stack/SNiagaraStackFunctionInputValue.cpp
+++ b/Engine/Plugins/FX/Niagara/Source/NiagaraEditorWidgets/Private/Stack/SNiagaraStackFunctionInputValue.cpp
@@ -34,6 +34,7 @@
 #include "Subsystems/AssetEditorSubsystem.h"
 #include "EditorFontGlyphs.h"
 #include "Widgets/SNiagaraLibraryOnlyToggleHeader.h"
+#include "Widgets/SNiagaraParameterName.h"
 
 #define LOCTEXT_NAMESPACE "NiagaraStackFunctionInputValue"
 
@@ -177,9 +178,9 @@ TSharedRef<SWidget> SNiagaraStackFunctionInputValue::ConstructValueWidgets()
 	}
 	case UNiagaraStackFunctionInput::EValueMode::Linked:
 	{
-		return SNew(STextBlock)
-			.TextStyle(FNiagaraEditorStyle::Get(), "NiagaraEditor.ParameterText")
-			.Text(this, &SNiagaraStackFunctionInputValue::GetLinkedValueHandleText)
+		return SNew(SNiagaraParameterName)
+			.ReadOnlyTextStyle(FNiagaraEditorStyle::Get(), "NiagaraEditor.ParameterText")
+			.ParameterName(this, &SNiagaraStackFunctionInputValue::GetLinkedValueHandleName)
 			.OnDoubleClicked(this, &SNiagaraStackFunctionInputValue::OnLinkedInputDoubleClicked);
 	}
 	case UNiagaraStackFunctionInput::EValueMode::Data:
@@ -368,9 +369,9 @@ void SNiagaraStackFunctionInputValue::ParameterPropertyValueChanged(const FPrope
 	FunctionInput->SetLocalValue(DisplayedLocalValueStruct.ToSharedRef());
 }
 
-FText SNiagaraStackFunctionInputValue::GetLinkedValueHandleText() const
+FName SNiagaraStackFunctionInputValue::GetLinkedValueHandleName() const
 {
-	return FText::FromName(FunctionInput->GetLinkedValueHandle().GetParameterHandleString());
+	return FunctionInput->GetLinkedValueHandle().GetParameterHandleString();
 }
 
 FText SNiagaraStackFunctionInputValue::GetDataValueText() const
diff --git a/Engine/Plugins/FX/Niagara/Source/NiagaraEditorWidgets/Private/Stack/SNiagaraStackFunctionInputValue.h b/Engine/Plugins/FX/Niagara/Source/NiagaraEditorWidgets/Private/Stack/SNiagaraStackFunctionInputValue.h
--- a/Engine/Plugins/FX/Niagara/Source/NiagaraEditorWidgets/Private/Stack/SNiagaraStackFunctionInputValue.h
+++ b/Engine/Plugins/FX/Niagara/Source/NiagaraEditorWidgets/Private/Stack/SNiagaraStackFunctionInputValue.h
@@ -45,7 +45,7 @@ private:
 
 	void ParameterPropertyValueChanged(const FPropertyChangedEvent& PropertyChangedEvent);
 
-	FText GetLinkedValueHandleText() const;
+	FName GetLinkedValueHandleName() const;
 
 	FText GetDataValueText() const;
 
diff --git a/Engine/Plugins/FX/Niagara/Source/NiagaraEditorWidgets/Private/Stack/SNiagaraStackParameterStoreEntryName.cpp b/Engine/Plugins/FX/Niagara/Source/NiagaraEditorWidgets/Private/Stack/SNiagaraStackParameterStoreEntryName.cpp
--- a/Engine/Plugins/FX/Niagara/Source/NiagaraEditorWidgets/Private/Stack/SNiagaraStackParameterStoreEntryName.cpp
+++ b/Engine/Plugins/FX/Niagara/Source/NiagaraEditorWidgets/Private/Stack/SNiagaraStackParameterStoreEntryName.cpp
@@ -6,7 +6,7 @@
 #include "ViewModels/Stack/NiagaraStackParameterStoreEntry.h"
 #include "Framework/Notifications/NotificationManager.h"
 #include "Widgets/Notifications/SNotificationList.h"
-#include "Widgets/Text/SInlineEditableTextBlock.h"
+#include "Widgets/SNiagaraParameterName.h"
 
 #define LOCTEXT_NAMESPACE "NiagaraStackParameterStoreEntryName"
 
@@ -15,6 +15,7 @@ void SNiagaraStackParameterStoreEntryName::Construct(const FArguments& InArgs, U
 	StackEntry = InStackEntry;
 	StackEntryItem = InStackEntry;
 	StackViewModel = InStackViewModel;
+	IsSelected = InArgs._IsSelected;
 
 	ChildSlot
 	[
@@ -23,14 +24,14 @@ void SNiagaraStackParameterStoreEntryName::Construct(const FArguments& InArgs, U
 		+ SHorizontalBox::Slot()
 		.VAlign(VAlign_Center)
 		[
-			SAssignNew(NameTextBlock, SInlineEditableTextBlock)
-			.Style(FNiagaraEditorStyle::Get(), "NiagaraEditor.ParameterInlineEditableText")
-			.Text_UObject(StackEntry, &UNiagaraStackEntry::GetDisplayName)
-			.IsReadOnly(this, &SNiagaraStackParameterStoreEntryName::GetIsNameReadOnly)
+			SAssignNew(NameTextBlock, SNiagaraParameterNameTextBlock)
+			.EditableTextStyle(FNiagaraEditorStyle::Get(), "NiagaraEditor.ParameterInlineEditableText")
+			.ParameterText_UObject(StackEntry, &UNiagaraStackEntry::GetDisplayName)
+			.IsReadOnly(StackEntry->SupportsRename() == false)
 			.IsSelected(this, &SNiagaraStackParameterStoreEntryName::GetIsNameWidgetSelected)
 			.OnTextCommitted(this, &SNiagaraStackParameterStoreEntryName::OnNameTextCommitted)
 			.OnVerifyTextChanged(this, &SNiagaraStackParameterStoreEntryName::VerifyNameTextChanged)
-			.HighlightText_UObject(InStackViewModel, &UNiagaraStackViewModel::GetCurrentSearchText)
+			//.HighlightText_UObject(InStackViewModel, &UNiagaraStackViewModel::GetCurrentSearchText)
 		]
 	];
 }
@@ -46,14 +47,9 @@ void SNiagaraStackParameterStoreEntryName::Tick(const FGeometry& AllottedGeometr
 	SCompoundWidget::Tick(AllottedGeometry, InCurrentTime, InDeltaTime);
 }
 
-bool SNiagaraStackParameterStoreEntryName::GetIsNameReadOnly() const
-{
-	return StackEntry->SupportsRename() == false;
-}
-
 bool SNiagaraStackParameterStoreEntryName::GetIsNameWidgetSelected() const
 {
-	return true;
+	return IsSelected.Get();
 }
 
 bool SNiagaraStackParameterStoreEntryName::VerifyNameTextChanged(const FText& NewText, FText& OutErrorMessage)
diff --git a/Engine/Plugins/FX/Niagara/Source/NiagaraEditorWidgets/Private/Stack/SNiagaraStackParameterStoreEntryName.h b/Engine/Plugins/FX/Niagara/Source/NiagaraEditorWidgets/Private/Stack/SNiagaraStackParameterStoreEntryName.h
--- a/Engine/Plugins/FX/Niagara/Source/NiagaraEditorWidgets/Private/Stack/SNiagaraStackParameterStoreEntryName.h
+++ b/Engine/Plugins/FX/Niagara/Source/NiagaraEditorWidgets/Private/Stack/SNiagaraStackParameterStoreEntryName.h
@@ -6,7 +6,7 @@
 
 class UNiagaraStackParameterStoreEntry;
 class UNiagaraStackViewModel;
-class SInlineEditableTextBlock;
+class SNiagaraParameterNameTextBlock;
 
 class SNiagaraStackParameterStoreEntryName: public SNiagaraStackEntryWidget
 {
@@ -15,6 +15,7 @@ public:
 
 public:
 	SLATE_BEGIN_ARGS(SNiagaraStackParameterStoreEntryName) { }
+		SLATE_ATTRIBUTE(bool, IsSelected);
 	SLATE_END_ARGS();
 
 	void Construct(const FArguments& InArgs, UNiagaraStackParameterStoreEntry* InStackEntry, UNiagaraStackViewModel* InStackViewModel);
@@ -22,8 +23,6 @@ public:
 	virtual void Tick(const FGeometry& AllottedGeometry, const double InCurrentTime, const float InDeltaTime) override;
 
 private:
-	bool GetIsNameReadOnly() const;
-
 	bool GetIsNameWidgetSelected() const;
 
 	bool VerifyNameTextChanged(const FText& NewText, FText& OutErrorMessage);
@@ -33,5 +32,7 @@ private:
 private:
 	UNiagaraStackParameterStoreEntry* StackEntry;
 
-	TSharedPtr<SInlineEditableTextBlock> NameTextBlock;
+	TSharedPtr<SNiagaraParameterNameTextBlock> NameTextBlock;
+
+	TAttribute<bool> IsSelected;
 };
\ No newline at end of file
