Replacing references to a deleted module now also renames the referenced modules.

diff --git a/Engine/Plugins/FX/Niagara/Source/NiagaraEditor/Private/NiagaraNodeFunctionCall.cpp b/Engine/Plugins/FX/Niagara/Source/NiagaraEditor/Private/NiagaraNodeFunctionCall.cpp
--- a/Engine/Plugins/FX/Niagara/Source/NiagaraEditor/Private/NiagaraNodeFunctionCall.cpp
+++ b/Engine/Plugins/FX/Niagara/Source/NiagaraEditor/Private/NiagaraNodeFunctionCall.cpp
@@ -94,6 +94,14 @@ void UNiagaraNodeFunctionCall::PostLoad()
 	{
 		ComputeNodeName();
 	}
+
+	// if we have a valid name, but the display name we're using does not match it,
+	// we've probably had our function script replaced and should rename ourselves
+	FString DetectedName = FunctionScript ? FunctionScript->GetName() : Signature.GetName();
+	if (!DetectedName.IsEmpty() && !FunctionDisplayName.StartsWith(DetectedName))
+	{
+		ComputeNodeName(DetectedName);
+	}
 }
 
 void UNiagaraNodeFunctionCall::UpgradeDIFunctionCalls()
