Niagara - Fix non-unity issues.

diff --git a/Engine/Plugins/FX/Niagara/Source/NiagaraEditor/Public/Widgets/SNiagaraParameterName.h b/Engine/Plugins/FX/Niagara/Source/NiagaraEditor/Public/Widgets/SNiagaraParameterName.h
--- a/Engine/Plugins/FX/Niagara/Source/NiagaraEditor/Public/Widgets/SNiagaraParameterName.h
+++ b/Engine/Plugins/FX/Niagara/Source/NiagaraEditor/Public/Widgets/SNiagaraParameterName.h
@@ -4,8 +4,12 @@
 
 #include "Widgets/SCompoundWidget.h"
 #include "EditorStyleSet.h"
+#include "Widgets/DeclarativeSyntaxSupport.h"
+#include "Styling/SlateTypes.h"
+#include "Framework/SlateDelegates.h"
 
 class SInlineEditableTextBlock;
+class SBorder;
 
 class NIAGARAEDITOR_API SNiagaraParameterName : public SCompoundWidget
 {
@@ -19,7 +23,7 @@ public:
 		, _ReadOnlyTextStyle(&FEditorStyle::Get().GetWidgetStyle<FTextBlockStyle>("NormalText"))
 		, _IsReadOnly(false)
 	{}
-	SLATE_STYLE_ARGUMENT(FInlineEditableTextBlockStyle, EditableTextStyle)
+		SLATE_STYLE_ARGUMENT(FInlineEditableTextBlockStyle, EditableTextStyle)
 		SLATE_STYLE_ARGUMENT(FTextBlockStyle, ReadOnlyTextStyle)
 		SLATE_ATTRIBUTE(FName, ParameterName)
 		SLATE_ARGUMENT(bool, IsReadOnly)
