Niagara - Fix additional include issues.

diff --git a/Engine/Plugins/FX/Niagara/Source/NiagaraEditor/Public/Widgets/SNiagaraParameterName.h b/Engine/Plugins/FX/Niagara/Source/NiagaraEditor/Public/Widgets/SNiagaraParameterName.h
--- a/Engine/Plugins/FX/Niagara/Source/NiagaraEditor/Public/Widgets/SNiagaraParameterName.h
+++ b/Engine/Plugins/FX/Niagara/Source/NiagaraEditor/Public/Widgets/SNiagaraParameterName.h
@@ -7,6 +7,7 @@
 #include "Widgets/DeclarativeSyntaxSupport.h"
 #include "Styling/SlateTypes.h"
 #include "Framework/SlateDelegates.h"
+#include "Styling/SlateWidgetStyleAsset.h"
 
 class SInlineEditableTextBlock;
 class SBorder;
