Pass at new colors for namespace refactor

diff --git a/Engine/Plugins/FX/Niagara/Source/NiagaraEditor/Private/NiagaraEditorSettings.cpp b/Engine/Plugins/FX/Niagara/Source/NiagaraEditor/Private/NiagaraEditorSettings.cpp
--- a/Engine/Plugins/FX/Niagara/Source/NiagaraEditor/Private/NiagaraEditorSettings.cpp
+++ b/Engine/Plugins/FX/Niagara/Source/NiagaraEditor/Private/NiagaraEditorSettings.cpp
@@ -37,61 +37,61 @@ UNiagaraEditorSettings::UNiagaraEditorSettings(const FObjectInitializer& ObjectI
 			{FNiagaraConstants::SystemNamespace},
 			NSLOCTEXT("NamespaceMetadata", "SystemDisplayName", "System"),
 			NSLOCTEXT("NamespaceMetadata", "SystemDescription", "Persistent attribute in the system which is written in a system\n stage and can be read anywhere."),
-			FLinearColor(FColor(67, 105, 124)),
+			FLinearColor(FColor(49, 113, 142)),
 			{ENiagaraNamespaceMetadataOptions::SupportsSubnamespace}),
 		FNiagaraNamespaceMetadata(
 			{FNiagaraConstants::EmitterNamespace},
 			NSLOCTEXT("NamespaceMetadata", "EmitterDisplayName", "Emitter"),
 			NSLOCTEXT("NamespaceMetadata", "EmitterDescription", "Persistent attribute which is written in a emitter\n stage and can be read in emitter and particle stages."),
-			FLinearColor(FColor(126, 87, 67)),
+			FLinearColor(FColor(145, 99, 56)),
 			{ENiagaraNamespaceMetadataOptions::SupportsSubnamespace}),
 		FNiagaraNamespaceMetadata(
 			{FNiagaraConstants::ParticleAttributeNamespace},
 			NSLOCTEXT("NamespaceMetadata", "ParticleDisplayName", "Particles"),
 			NSLOCTEXT("NamespaceMetadata", "ParticleDescription", "Persistent attribute which is written in a particle\n stage and can be read in particle stages."),
-			FLinearColor(FColor(87, 107, 61)),
+			FLinearColor(FColor(72, 130, 71)),
 			{ENiagaraNamespaceMetadataOptions::SupportsSubnamespace}),
 		FNiagaraNamespaceMetadata(
 			{FNiagaraConstants::ModuleNamespace},
 			NSLOCTEXT("NamespaceMetadata", "ModuleDisplayName", "Input"),
 			NSLOCTEXT("NamespaceMetadata", "ModuleDescription", "A transient value which exposes a module input to the system and emitter editor."),
-			FLinearColor(FColor(118, 85, 48)),
+			FLinearColor(FColor(136, 66, 65)),
 			{ENiagaraNamespaceMetadataOptions::SupportsSubnamespace}),
 		FNiagaraNamespaceMetadata(
 			{FNiagaraConstants::OutputNamespace, FNiagaraConstants::ModuleNamespace},
 			NSLOCTEXT("NamespaceMetadata", "ModuleOutputDisplayName", "Output"),
 			NSLOCTEXT("NamespaceMetadata", "ModuleOutputDescription", "A transient value which is unique to the module that wrote it and can be read from any other module.\nTransient values do not persist from frame to frame, or between stages, e.g. emitter to particle, or spawn to update."),
-			FLinearColor(FColor(118, 85, 48)),
+			FLinearColor(FColor(109, 95, 124)),
 			{ENiagaraNamespaceMetadataOptions::Advanced}),
 		FNiagaraNamespaceMetadata(
 			{FNiagaraConstants::LocalNamespace, FNiagaraConstants::ModuleNamespace},
 			NSLOCTEXT("NamespaceMetadata", "ModuleLocalDisplayName", "Local"),
 			NSLOCTEXT("NamespaceMetadata", "ModuleLocalDescription", "A transient value which can be written to and read from within a single module.\nTransient values do not persist from frame to frame, or between stages, e.g. emitter to particle, or spawn to update."),
-			FLinearColor(FColor(118, 85, 48)),
+			FLinearColor(FColor(174, 157, 40)),
 			{ENiagaraNamespaceMetadataOptions::Advanced}),
 		FNiagaraNamespaceMetadata(
 			{FNiagaraConstants::TransientNamespace},
 			NSLOCTEXT("NamespaceMetadata", "TransientDisplayName", "Transient"),
 			NSLOCTEXT("NamespaceMetadata", "TransientDescription", "A transient value which can be written to and read from from any module.\nTransient values do not persist from frame to frame, or between stages, e.g. emitter to particle, or spawn to update."),
-			FLinearColor(FColor(179, 122, 45)),
+			FLinearColor(FColor(109, 95, 124)),
 			{ENiagaraNamespaceMetadataOptions::Advanced}),
 		FNiagaraNamespaceMetadata(
 			{FNiagaraConstants::EngineNamespace},
 			NSLOCTEXT("NamespaceMetadata", "EngineDisplayName", "Engine"),
 			NSLOCTEXT("NamespaceMetadata", "EngineDescription", "A read only value which is provided by the engine.\nThis value's source can be the simulation itsef\ne.g. ExecutionCount, or the owner of the simulation e.g. (Owner) Scale."),
-			FLinearColor(FColor(40, 40, 40)),
+			FLinearColor(FColor(45, 45, 45)),
 			{ENiagaraNamespaceMetadataOptions::PreventRenaming}),
 		FNiagaraNamespaceMetadata(
 			{FNiagaraConstants::UserNamespace},
 			NSLOCTEXT("NamespaceMetadata", "UserDisplayName", "User"),
 			NSLOCTEXT("NamespaceMetadata", "UserDescription", "A read only value which can be initialized per system and\nmodified externally in the level, by blueprint, or by c++."),
-			FLinearColor(FColor(67, 105, 124)),
+			FLinearColor(FColor(49, 113, 142)),
 			{ }),
 		FNiagaraNamespaceMetadata(
 			{FNiagaraConstants::DataInstanceNamespace},
 			NSLOCTEXT("NamespaceMetadata", "DataInstanceDisplayName", "Data Instance"),
 			NSLOCTEXT("NamespaceMetadata", "DataInstanceDescription", "A special transient value which has a single bool IsAlive value, which determines if a particle is alive or not."),
-			FLinearColor(FColor(179, 122, 45)),
+			FLinearColor(FColor(45, 45, 45)),
 			{ENiagaraNamespaceMetadataOptions::Advanced}),
 	};
 }
