Fixed: missing Niagara compute shaders when cooking for a platform with multiple RHIs enabled.

UNiagaraScript::CacheResourceShadersForCooking was checking for duplicate requests, but it was only taking into account the feature level. When cooking for multiple shader platforms which provide the same feature level (e.g. D3D and Vulkan on Windows), only one set of shaders was compiled.

#rnx

diff --git a/Engine/Plugins/FX/Niagara/Source/Niagara/Classes/NiagaraScript.h b/Engine/Plugins/FX/Niagara/Source/Niagara/Classes/NiagaraScript.h
--- a/Engine/Plugins/FX/Niagara/Source/Niagara/Classes/NiagaraScript.h
+++ b/Engine/Plugins/FX/Niagara/Source/Niagara/Classes/NiagaraScript.h
@@ -585,7 +585,7 @@ public:
 	NIAGARA_API void CacheResourceShadersForRendering(bool bRegenerateId, bool bForceRecompile=false);
 	void BeginCacheForCookedPlatformData(const ITargetPlatform *TargetPlatform);
 	virtual bool IsCachedCookedPlatformDataLoaded(const ITargetPlatform* TargetPlatform) override;
-	void CacheShadersForResources(EShaderPlatform ShaderPlatform, FNiagaraShaderScript *ResourceToCache, bool bApplyCompletedShaderMapForRendering, bool bForceRecompile = false, bool bCooking=false);
+	void CacheShadersForResources(FNiagaraShaderScript* ResourceToCache, bool bApplyCompletedShaderMapForRendering, bool bForceRecompile = false, bool bCooking=false);
 #endif // WITH_EDITOR
 	FNiagaraShaderScript* AllocateResource();
 	FNiagaraShaderScript* GetRenderThreadScript()
diff --git a/Engine/Plugins/FX/Niagara/Source/Niagara/Private/NiagaraScript.cpp b/Engine/Plugins/FX/Niagara/Source/Niagara/Private/NiagaraScript.cpp
--- a/Engine/Plugins/FX/Niagara/Source/Niagara/Private/NiagaraScript.cpp
+++ b/Engine/Plugins/FX/Niagara/Source/Niagara/Private/NiagaraScript.cpp
@@ -1822,7 +1822,7 @@ void UNiagaraScript::CacheResourceShadersForCooking(EShaderPlatform ShaderPlatfo
 			ERHIFeatureLevel::Type TargetFeatureLevel = GetMaxSupportedFeatureLevel(ShaderPlatform);
 			const auto FindExistingScriptPredicate = [&](const FNiagaraShaderScript* ExistingScript)
 			{
-				return ExistingScript->MatchesScript(TargetFeatureLevel, CachedScriptVMId);
+				return ExistingScript->MatchesScript(TargetFeatureLevel, ShaderPlatform, CachedScriptVMId);
 			};
 
 			// see if the script has already been added before adding a new version
@@ -1836,14 +1836,14 @@ void UNiagaraScript::CacheResourceShadersForCooking(EShaderPlatform ShaderPlatfo
 			check(CachedScriptVMId.CompilerVersionID.IsValid());
 			check(CachedScriptVMId.BaseScriptCompileHash.IsValid());
 
-			NewResource->SetScript(this, (ERHIFeatureLevel::Type)TargetFeatureLevel, CachedScriptVMId.CompilerVersionID, CachedScriptVMId.AdditionalDefines,
+			NewResource->SetScript(this, TargetFeatureLevel, ShaderPlatform, CachedScriptVMId.CompilerVersionID, CachedScriptVMId.AdditionalDefines,
 				CachedScriptVMId.BaseScriptCompileHash,	CachedScriptVMId.ReferencedCompileHashes, 
 				CachedScriptVMId.bUsesRapidIterationParams, GetFriendlyName());
 			ResourceToCache = NewResource;
 
 			check(ResourceToCache);
 
-			CacheShadersForResources(ShaderPlatform, ResourceToCache, false, false, true);
+			CacheShadersForResources(ResourceToCache, false, false, true);
 
 			INiagaraModule NiagaraModule = FModuleManager::GetModuleChecked<INiagaraModule>(TEXT("Niagara"));
 			NiagaraModule.ProcessShaderCompilationQueue();
@@ -1855,19 +1855,19 @@ void UNiagaraScript::CacheResourceShadersForCooking(EShaderPlatform ShaderPlatfo
 
 
 
-void UNiagaraScript::CacheShadersForResources(EShaderPlatform ShaderPlatform, FNiagaraShaderScript *ResourceToCache, bool bApplyCompletedShaderMapForRendering, bool bForceRecompile, bool bCooking)
+void UNiagaraScript::CacheShadersForResources(FNiagaraShaderScript* ResourceToCache, bool bApplyCompletedShaderMapForRendering, bool bForceRecompile, bool bCooking)
 {
 	if (CanBeRunOnGpu())
 	{
 		// When not running in the editor, the shaders are created in-sync (in the postload) to avoid update issues.
-		const bool bSuccess = ResourceToCache->CacheShaders(ShaderPlatform, bApplyCompletedShaderMapForRendering, bForceRecompile, bCooking || !GIsEditor || GIsAutomationTesting);
+		const bool bSuccess = ResourceToCache->CacheShaders(bApplyCompletedShaderMapForRendering, bForceRecompile, bCooking || !GIsEditor || GIsAutomationTesting);
 
 #if defined(NIAGARA_SCRIPT_COMPILE_LOGGING_MEDIUM)
 		if (!bSuccess)
 		{
 			UE_LOG(LogNiagara, Warning, TEXT("Failed to compile Niagara shader %s for platform %s."),
 				*GetPathName(),
-				*LegacyShaderPlatformToShaderFormat(ShaderPlatform).ToString());
+				*LegacyShaderPlatformToShaderFormat(ResourceToCache->GetShaderPlatform()).ToString());
 
 			const TArray<FString>& CompileErrors = ResourceToCache->GetCompileErrors();
 			for (int32 ErrorIndex = 0; ErrorIndex < CompileErrors.Num(); ErrorIndex++)
@@ -1901,21 +1901,17 @@ void UNiagaraScript::CacheResourceShadersForRendering(bool bRegenerateId, bool b
 		// Need to make sure the owner supports GPU scripts, otherwise this is a wasted compile.
 		if (Source && OwnerCanBeRunOnGpu())
 		{
-			FNiagaraShaderScript* ResourceToCache;
 			ERHIFeatureLevel::Type CacheFeatureLevel = GMaxRHIFeatureLevel;
-			ScriptResource->SetScript(this, CacheFeatureLevel, CachedScriptVMId.CompilerVersionID, CachedScriptVMId.AdditionalDefines,
+			const EShaderPlatform ShaderPlatform = GShaderPlatformForFeatureLevel[CacheFeatureLevel];
+
+			ScriptResource->SetScript(this, FeatureLevel, ShaderPlatform, CachedScriptVMId.CompilerVersionID, CachedScriptVMId.AdditionalDefines,
 				CachedScriptVMId.BaseScriptCompileHash, CachedScriptVMId.ReferencedCompileHashes, 
 				CachedScriptVMId.bUsesRapidIterationParams, GetFriendlyName());
 
-			//if (ScriptResourcesByFeatureLevel[CacheFeatureLevel])
+			if (FNiagaraUtilities::SupportsGPUParticles(ShaderPlatform))
 			{
-				const EShaderPlatform ShaderPlatform = GShaderPlatformForFeatureLevel[CacheFeatureLevel];
-				if (FNiagaraUtilities::SupportsGPUParticles(ShaderPlatform))
-				{
-					ResourceToCache = ScriptResourcesByFeatureLevel[CacheFeatureLevel];
-					CacheShadersForResources(ShaderPlatform, ScriptResource.Get(), true);
-					ScriptResourcesByFeatureLevel[CacheFeatureLevel] = ScriptResource.Get();
-				}
+				CacheShadersForResources(ScriptResource.Get(), true);
+				ScriptResourcesByFeatureLevel[CacheFeatureLevel] = ScriptResource.Get();
 			}
 		}
 		else
diff --git a/Engine/Plugins/FX/Niagara/Source/NiagaraShader/Private/NiagaraShared.cpp b/Engine/Plugins/FX/Niagara/Source/NiagaraShader/Private/NiagaraShared.cpp
--- a/Engine/Plugins/FX/Niagara/Source/NiagaraShader/Private/NiagaraShared.cpp
+++ b/Engine/Plugins/FX/Niagara/Source/NiagaraShader/Private/NiagaraShared.cpp
@@ -49,7 +49,7 @@ void FNiagaraShaderScript::SetupShaderCompilationEnvironment(
 
 NIAGARASHADER_API bool FNiagaraShaderScript::ShouldCache(EShaderPlatform Platform, const FShaderType* ShaderType) const
 {
-	check(ShaderType->GetNiagaraShaderType() )
+	check(ShaderType->GetNiagaraShaderType() != nullptr);
 	return true;
 }
 
@@ -285,7 +285,7 @@ void FNiagaraShaderScript::SerializeShaderMap(FArchive& Ar)
 	}
 }
 
-void FNiagaraShaderScript::SetScript(UNiagaraScript *InScript, ERHIFeatureLevel::Type InFeatureLevel, const FGuid& InCompilerVersionID,  const TArray<FString>& InAdditionalDefines,
+void FNiagaraShaderScript::SetScript(UNiagaraScript* InScript, ERHIFeatureLevel::Type InFeatureLevel, EShaderPlatform InShaderPlatform, const FGuid& InCompilerVersionID,  const TArray<FString>& InAdditionalDefines,
 		const FNiagaraCompileHash& InBaseCompileHash, const TArray<FNiagaraCompileHash>& InReferencedCompileHashes, 
 		bool bInUsesRapidIterationParams, FString InFriendlyName)
 {
@@ -299,17 +299,19 @@ void FNiagaraShaderScript::SetScript(UNiagaraScript *InScript, ERHIFeatureLevel:
 	ReferencedCompileHashes = InReferencedCompileHashes;
 	FriendlyName = InFriendlyName;
 	SetFeatureLevel(InFeatureLevel);
+	ShaderPlatform = InShaderPlatform;
 }
 
 #if WITH_EDITOR
-bool FNiagaraShaderScript::MatchesScript(ERHIFeatureLevel::Type InFeatureLevel, const FNiagaraVMExecutableDataId& ScriptId) const
+bool FNiagaraShaderScript::MatchesScript(ERHIFeatureLevel::Type InFeatureLevel, EShaderPlatform InShaderPlatform, const FNiagaraVMExecutableDataId& ScriptId) const
 {
 	return CompilerVersionId == ScriptId.CompilerVersionID
 		&& AdditionalDefines == ScriptId.AdditionalDefines
 		&& bUsesRapidIterationParams == ScriptId.bUsesRapidIterationParams
 		&& BaseCompileHash == ScriptId.BaseScriptCompileHash
 		&& ReferencedCompileHashes == ScriptId.ReferencedCompileHashes
-		&& FeatureLevel == InFeatureLevel;
+		&& FeatureLevel == InFeatureLevel
+		&& ShaderPlatform == InShaderPlatform;
 }
 #endif
 
@@ -335,17 +337,17 @@ NIAGARASHADER_API  bool FNiagaraShaderScript::IsCompilationFinished() const
 */
 #if WITH_EDITOR
 
-bool FNiagaraShaderScript::CacheShaders(EShaderPlatform Platform, bool bApplyCompletedShaderMapForRendering, bool bForceRecompile, bool bSynchronous)
+bool FNiagaraShaderScript::CacheShaders(bool bApplyCompletedShaderMapForRendering, bool bForceRecompile, bool bSynchronous)
 {
 	FNiagaraShaderMapId NoStaticParametersId;
-	GetShaderMapId(Platform, NoStaticParametersId);
-	return CacheShaders(NoStaticParametersId, Platform, bApplyCompletedShaderMapForRendering, bForceRecompile, bSynchronous);
+	GetShaderMapId(ShaderPlatform, NoStaticParametersId);
+	return CacheShaders(NoStaticParametersId, bApplyCompletedShaderMapForRendering, bForceRecompile, bSynchronous);
 }
 
 /**
 * Caches the shaders for this script
 */
-bool FNiagaraShaderScript::CacheShaders(const FNiagaraShaderMapId& ShaderMapId, EShaderPlatform Platform, bool bApplyCompletedShaderMapForRendering, bool bForceRecompile, bool bSynchronous)
+bool FNiagaraShaderScript::CacheShaders(const FNiagaraShaderMapId& ShaderMapId, bool bApplyCompletedShaderMapForRendering, bool bForceRecompile, bool bSynchronous)
 {
 	bool bSucceeded = false;
 
@@ -357,13 +359,13 @@ bool FNiagaraShaderScript::CacheShaders(const FNiagaraShaderMapId& ShaderMapId,
 			NIAGARASHADER_API extern FCriticalSection GIdToNiagaraShaderMapCS;
 			FScopeLock ScopeLock(&GIdToNiagaraShaderMapCS);
 			// Find the script's cached shader map.
-			GameThreadShaderMap = FNiagaraShaderMap::FindId(ShaderMapId, Platform);
+			GameThreadShaderMap = FNiagaraShaderMap::FindId(ShaderMapId, ShaderPlatform);
 		}
 
 		// Attempt to load from the derived data cache if we are uncooked
 		if (!bForceRecompile && !GameThreadShaderMap && !FPlatformProperties::RequiresCookedData())
 		{
-			FNiagaraShaderMap::LoadFromDerivedDataCache(this, ShaderMapId, Platform, GameThreadShaderMap);
+			FNiagaraShaderMap::LoadFromDerivedDataCache(this, ShaderMapId, ShaderPlatform, GameThreadShaderMap);
 			if (GameThreadShaderMap && GameThreadShaderMap->IsValid())
 			{
 				UE_LOG(LogTemp, Verbose, TEXT("Loaded shader %s for Niagara script %s from DDC"), *GameThreadShaderMap->GetFriendlyName(), *GetFriendlyName());
@@ -407,7 +409,7 @@ bool FNiagaraShaderScript::CacheShaders(const FNiagaraShaderMapId& ShaderMapId,
 
 			// If there's no cached shader map for this script compile a new one.
 			// This is just kicking off the compile, GameThreadShaderMap will not be complete yet
-			bSucceeded = BeginCompileShaderMap(ShaderMapId, Platform, GameThreadShaderMap, bApplyCompletedShaderMapForRendering, bSynchronous);
+			bSucceeded = BeginCompileShaderMap(ShaderMapId, GameThreadShaderMap, bApplyCompletedShaderMapForRendering, bSynchronous);
 
 			if (!bSucceeded)
 			{
@@ -510,7 +512,6 @@ void FNiagaraShaderScript::GetShaderMapIDsWithUnfinishedCompilation(TArray<int32
 */
 bool FNiagaraShaderScript::BeginCompileShaderMap(
 	const FNiagaraShaderMapId& ShaderMapId,
-	EShaderPlatform Platform,
 	TRefCountPtr<FNiagaraShaderMap>& OutShaderMap,
 	bool bApplyCompletedShaderMapForRendering,
 	bool bSynchronous)
@@ -529,7 +530,7 @@ bool FNiagaraShaderScript::BeginCompileShaderMap(
 	OutstandingCompileShaderMapIds.AddUnique(NewShaderMap->GetCompilingId());		
 	UE_LOG(LogShaders, Log, TEXT("BeginCompileShaderMap AddUnique %p %d"), this, NewShaderMap->GetCompilingId());
 
-	FNiagaraCompilationQueue::Get()->Queue(this, NewShaderMap, ShaderMapId, Platform, bApplyCompletedShaderMapForRendering);
+	FNiagaraCompilationQueue::Get()->Queue(this, NewShaderMap, ShaderMapId, ShaderPlatform, bApplyCompletedShaderMapForRendering);
 	if (bSynchronous)
 	{
 		INiagaraShaderModule NiagaraShaderModule = FModuleManager::GetModuleChecked<INiagaraShaderModule>(TEXT("NiagaraShader"));
diff --git a/Engine/Plugins/FX/Niagara/Source/NiagaraShader/Public/NiagaraShared.h b/Engine/Plugins/FX/Niagara/Source/NiagaraShader/Public/NiagaraShared.h
--- a/Engine/Plugins/FX/Niagara/Source/NiagaraShader/Public/NiagaraShared.h
+++ b/Engine/Plugins/FX/Niagara/Source/NiagaraShader/Public/NiagaraShared.h
@@ -583,8 +583,8 @@ public:
 	 * Caches the shaders for this script with no static parameters on the given platform.
 	 * This is used by FNiagaraShaderScript
 	 */
-	NIAGARASHADER_API  bool CacheShaders(EShaderPlatform Platform, bool bApplyCompletedShaderMapForRendering, bool bForceRecompile, bool bSynchronous = false);
-	bool CacheShaders(const FNiagaraShaderMapId& ShaderMapId, EShaderPlatform Platform, bool bApplyCompletedShaderMapForRendering, bool bForceRecompile, bool bSynchronous = false);
+	NIAGARASHADER_API  bool CacheShaders(bool bApplyCompletedShaderMapForRendering, bool bForceRecompile, bool bSynchronous = false);
+	bool CacheShaders(const FNiagaraShaderMapId& ShaderMapId, bool bApplyCompletedShaderMapForRendering, bool bForceRecompile, bool bSynchronous = false);
 
 
 	NIAGARASHADER_API uint32 GetUseSimStagesDefine() const;
@@ -652,6 +652,7 @@ public:
 	void SetCompileErrors(const TArray<FString>& InCompileErrors) { CompileErrors = InCompileErrors; }
 
 	ERHIFeatureLevel::Type GetFeatureLevel() const { return FeatureLevel; }
+	EShaderPlatform GetShaderPlatform() const { return ShaderPlatform; }
 
 	class FNiagaraShaderMap* GetGameThreadShaderMap() const
 	{
@@ -713,11 +714,11 @@ public:
 	const FString& GetFriendlyName()	const { return FriendlyName; }
 
 
-	NIAGARASHADER_API void SetScript(UNiagaraScript *InScript, ERHIFeatureLevel::Type InFeatureLevel, const FGuid& InCompilerVersion, const TArray<FString>& InAdditionalDefines,
+	NIAGARASHADER_API void SetScript(UNiagaraScript* InScript, ERHIFeatureLevel::Type InFeatureLevel, EShaderPlatform InShaderPlatform, const FGuid& InCompilerVersion, const TArray<FString>& InAdditionalDefines,
 		const FNiagaraCompileHash& InBaseCompileHash, const TArray<FNiagaraCompileHash>& InReferencedCompileHashes, 
 		bool bInUsesRapidIterationParams, FString InFriendlyName);
 #if WITH_EDITOR
-	NIAGARASHADER_API bool MatchesScript(ERHIFeatureLevel::Type InFeatureLevel, const FNiagaraVMExecutableDataId& ScriptId) const;
+	NIAGARASHADER_API bool MatchesScript(ERHIFeatureLevel::Type InFeatureLevel, EShaderPlatform InShaderPlatform, const FNiagaraVMExecutableDataId& ScriptId) const;
 #endif
 
 	UNiagaraScript *GetBaseVMScript()
@@ -813,8 +814,9 @@ private:
 	 */
 	TArray<int32, TInlineAllocator<1> > OutstandingCompileShaderMapIds;
 
-	/** Feature level that this script is representing. */
+	/** Feature level and shader platform that this script is representing. */
 	ERHIFeatureLevel::Type FeatureLevel;
+	EShaderPlatform ShaderPlatform;
 
 	uint32 bLoadedCookedShaderMapId : 1;
 	uint32 bLoadedFromCookedMaterial : 1;
@@ -828,7 +830,6 @@ private:
 	*/
 	bool BeginCompileShaderMap(
 		const FNiagaraShaderMapId& ShaderMapId,
-		EShaderPlatform Platform, 
 		TRefCountPtr<class FNiagaraShaderMap>& OutShaderMap, 
 		bool bApplyCompletedShaderMapForRendering,
 		bool bSynchronous = false);
