Added missing initializer for member variable.

#rnx

diff --git a/Engine/Plugins/FX/Niagara/Source/NiagaraShader/Public/NiagaraShared.h b/Engine/Plugins/FX/Niagara/Source/NiagaraShader/Public/NiagaraShared.h
--- a/Engine/Plugins/FX/Niagara/Source/NiagaraShader/Public/NiagaraShared.h
+++ b/Engine/Plugins/FX/Niagara/Source/NiagaraShader/Public/NiagaraShared.h
@@ -570,6 +570,7 @@ public:
 		, GameThreadShaderMap(nullptr)
 		, RenderingThreadShaderMap(nullptr)
 		, FeatureLevel(GMaxRHIFeatureLevel)
+		, ShaderPlatform(SP_NumPlatforms)
 		, bLoadedCookedShaderMapId(false)
 		, bLoadedFromCookedMaterial(false)
 	{}
