Fixes to MeshParticleSpace material transfrom node
-niagara particles will properly handle to/from world space
-cascade particles will (continue to) not handle non-uniform scales if used with world -> particle space
-local vertex factory (static meshes) will define particle transforms now if required

diff --git a/Engine/Plugins/FX/Niagara/Shaders/Private/NiagaraMeshVertexFactory.ush b/Engine/Plugins/FX/Niagara/Shaders/Private/NiagaraMeshVertexFactory.ush
--- a/Engine/Plugins/FX/Niagara/Shaders/Private/NiagaraMeshVertexFactory.ush
+++ b/Engine/Plugins/FX/Niagara/Shaders/Private/NiagaraMeshVertexFactory.ush
@@ -12,6 +12,8 @@
 #define USE_PARTICLE_TIME (NEEDS_PARTICLE_TIME)
 #define USE_PARTICLE_RANDOM (NEEDS_PARTICLE_RANDOM)
 #define USE_PARTICLE_TRANSFORM (NEEDS_PARTICLE_TRANSFORM)
+#define USE_PARTICLE_LOCAL_TO_WORLD (NEEDS_PARTICLE_LOCAL_TO_WORLD)
+#define USE_PARTICLE_WORLD_TO_LOCAL (NEEDS_PARTICLE_WORLD_TO_LOCAL)
 
 float3 TransformPosition(float3 InPosition)
 {
@@ -134,8 +136,12 @@ struct FVertexFactoryInterpolantsVSToPS
 	nointerpolation float ParticleRandom : PARTICLE_RANDOM;
 #endif
 
-#if USE_PARTICLE_TRANSFORM
-	nointerpolation float4 LocalToWorld[4] : PARTICLE_TRANSFORM;
+#if USE_PARTICLE_LOCAL_TO_WORLD
+	nointerpolation float3x4 ParticleToWorld : PARTICLE_LOCAL_TO_WORLD;
+#endif
+
+#if USE_PARTICLE_WORLD_TO_LOCAL
+	nointerpolation float3x4 WorldToParticle : PARTICLE_WORLD_TO_LOCAL;
 #endif
 
 #if INSTANCED_STEREO
@@ -175,7 +181,7 @@ struct FVertexFactoryIntermediates
 #endif
 
 	/** The position of the particle, XYZ: direction. */
-	float4 ParticlePosition;
+	float3 ParticlePosition;
 
 	/** The velocity of the particle, XYZ: direction, W: speed. */
 	float4 ParticleVelocity;
@@ -193,13 +199,20 @@ struct FVertexFactoryIntermediates
 	float ParticleRandom;
 #endif
 
-	float4x4 LocalToWorld;
-	float4x4 PrevLocalToWorld;
+	float4x4 ParticleToWorld;
+	float4x4 PrevParticleToWorld;
+
+	float4x4 WorldToParticle;
 };
 
 float4x4 GetParticleTransform(FVertexFactoryIntermediates Inters)
 {
-	return Inters.LocalToWorld;
+	return Inters.ParticleToWorld;
+}
+
+float4x4 GetParticleInvTransform(FVertexFactoryIntermediates Inters)
+{
+	return Inters.WorldToParticle;
 }
 
 float4x4 GetParticlePrevTransform(FVertexFactoryIntermediates Inters)
@@ -217,7 +230,7 @@ float4x4 GetParticlePrevTransform(FVertexFactoryIntermediates Inters)
 	//Transform is transposed on CPU and position is packed into W. Transpose again to get the correct transform back.
 	return transpose(PrevTransform4x4);
 #else
-	return Inters.PrevLocalToWorld;
+	return Inters.PrevParticleToWorld;
 #endif
 }
 
@@ -308,11 +321,12 @@ FMaterialPixelParameters GetMaterialPixelParameters(FVertexFactoryInterpolantsVS
 	Result.Particle.Random = 0.0f;
 #endif
 
-#if USE_PARTICLE_TRANSFORM
-	Result.Particle.LocalToWorld[0] = float4(Interpolants.LocalToWorld[0].x, Interpolants.LocalToWorld[1].x, Interpolants.LocalToWorld[2].x, Interpolants.LocalToWorld[3].x);
-	Result.Particle.LocalToWorld[1] = float4(Interpolants.LocalToWorld[0].y, Interpolants.LocalToWorld[1].y, Interpolants.LocalToWorld[2].y, Interpolants.LocalToWorld[3].y);
-	Result.Particle.LocalToWorld[2] = float4(Interpolants.LocalToWorld[0].z, Interpolants.LocalToWorld[1].z, Interpolants.LocalToWorld[2].z, Interpolants.LocalToWorld[3].z);
-	Result.Particle.LocalToWorld[3] = float4(Interpolants.LocalToWorld[0].w, Interpolants.LocalToWorld[1].w, Interpolants.LocalToWorld[2].w, Interpolants.LocalToWorld[3].w);
+#if USE_PARTICLE_LOCAL_TO_WORLD
+	Result.Particle.ParticleToWorld = transpose(float4x4(Interpolants.ParticleToWorld, float4(0.0f, 0.0f, 0.0f, 1.0f)));
+#endif
+
+#if USE_PARTICLE_WORLD_TO_LOCAL
+	Result.Particle.WorldToParticle = transpose(float4x4(Interpolants.WorldToParticle, float4(0.0f, 0.0f, 0.0f, 1.0f)));
 #endif
 
 	Result.Particle.MotionBlurFade = 1.0f;
@@ -375,12 +389,8 @@ FMaterialVertexParameters GetMaterialVertexParameters(FVertexFactoryInput Input,
 	float4x4 Transform = GetParticleTransform(Intermediates);
 	Result.InstanceLocalToWorld = Transform;
 
-#if NEEDS_PARTICLE_TRANSFORM
-	Result.Particle.LocalToWorld[0] = float4(Transform[0].x, Transform[1].x, Transform[2].x, Transform[3].x);
-	Result.Particle.LocalToWorld[1] = float4(Transform[0].y, Transform[1].y, Transform[2].y, Transform[3].y);
-	Result.Particle.LocalToWorld[2] = float4(Transform[0].z, Transform[1].z, Transform[2].z, Transform[3].z);
-	Result.Particle.LocalToWorld[3] = float4(Transform[0].w, Transform[1].w, Transform[2].w, Transform[3].w);
-#endif
+	Result.Particle.ParticleToWorld = Transform;
+	Result.Particle.WorldToParticle = GetParticleInvTransform(Intermediates);
 
 	Result.TangentToWorld = mul(TangentToLocal, (float3x3)Transform );
 
@@ -426,7 +436,48 @@ float3x3 QuatTo3x3(float4 Q)
 	return Mat;
 }
 
-float3x3 GetFacingMatrix(in FVertexFactoryIntermediates Intermediates)
+float4x4 ComposeTransformMatrix(float3 S, float3x3 R, float3 T)
+{
+	float3x3 ScaleMat = float3x3(
+		S.x, 0.0f, 0.0f,
+		0.0f, S.y, 0.0f,
+		0.0f, 0.0f, S.z);
+
+	float3x3 ScaleRot = mul(ScaleMat, R);
+
+	return float4x4(
+		ScaleRot[0], 0.0f,
+		ScaleRot[1], 0.0f,
+		ScaleRot[2], 0.0f,
+		T, 1.0f);
+}
+
+float4x4 ComposeInvTransformMatrix(float3 S, float3x3 R, float3 T)
+{
+	float3x3 InvScaleMat = float3x3(
+		1.0f / S.x, 0.0f, 0.0f,
+		0.0f, 1.0f / S.y, 0.0f,
+		0.0f, 0.0f, 1.0f / S.z);
+
+	float3x3 InvRotScale = mul(transpose(R), InvScaleMat);
+
+	float4x4 InvRotScale44 = float4x4(
+		InvRotScale[0], 0.0f,
+		InvRotScale[1], 0.0f,
+		InvRotScale[2], 0.0f,
+		0.0f, 0.0f, 0.0f, 1.0f
+	);
+
+	float4x4 InvTransMat = float4x4(
+		1.0f, 0.0f, 0.0f, 0.0f,
+		0.0f, 1.0f, 0.0f, 0.0f,
+		0.0f, 0.0f, 1.0f, 0.0f,
+		-T, 1.0f);
+
+	return mul(InvTransMat, InvRotScale44);
+}
+
+float3x3 GetFacingMatrix(in float3 ParticlePosition, in float4 ParticleVelocity)
 {
 	float3 WorldX = float3(1, 0, 0);
 	float3 WorldZ = float3(0, 0, 1);
@@ -438,13 +489,13 @@ float3x3 GetFacingMatrix(in FVertexFactoryIntermediates Intermediates)
 	{
 		case MESH_FACING_VELOCITY:
 		{
-			FacingDir = Intermediates.ParticleVelocity.w > 0.0f ? Intermediates.ParticleVelocity.xyz : WorldZ;
+			FacingDir = ParticleVelocity.w > 0.0f ? ParticleVelocity.xyz : WorldZ;
 			break;
 		}
 
 		case MESH_FACING_CAMERA_POSITION:
 		{
-			float3 WorldPosition = TransformPosition(Intermediates.ParticlePosition.xyz);
+			float3 WorldPosition = TransformPosition(ParticlePosition);
 			float3 CameraVector = ResolvedView.WorldCameraOrigin - WorldPosition;
 			FacingDir = normalize(CameraVector);
 			break;
@@ -519,7 +570,7 @@ float3x3 GetFacingMatrix(in FVertexFactoryIntermediates Intermediates)
 	return float3x3(XAxis, YAxis, ZAxis);
 }
 
-float3 CalculateCameraOffset(FVertexFactoryIntermediates Intermediates, uint InstanceId)
+float3 CalculateCameraOffset(float3 WorldPosition, uint InstanceId)
 {
 	if (NiagaraMeshVF.CameraOffsetDataOffset == -1)
 	{
@@ -527,7 +578,6 @@ float3 CalculateCameraOffset(FVertexFactoryIntermediates Intermediates, uint Ins
 	}
 
 	float Offset = GetFloat(NiagaraMeshVF.CameraOffsetDataOffset, InstanceId);
-	float3 WorldPosition = TransformPosition(Intermediates.ParticlePosition.xyz);
 	float3 CameraVector = ResolvedView.WorldCameraOrigin - WorldPosition;
 	float3 CameraDir = normalize(CameraVector);
 	if (NiagaraMeshVF.bLocalSpace)
@@ -538,6 +588,58 @@ float3 CalculateCameraOffset(FVertexFactoryIntermediates Intermediates, uint Ins
 	return CameraDir * Offset;
 }
 
+void CalculateTransforms(int InstanceId, in FVertexFactoryIntermediates Intermediates, out float4x4 LocalToWorld, out float4x4 WorldToLocal)
+{
+	////////////////////////////////////////////////////////////////////////////////
+	// Translation
+	float3 WorldParticlePosition = float3(0.0f, 0.0f, 0.0f);
+	if (NiagaraMeshVF.PositionDataOffset == -1)
+	{
+		WorldParticlePosition = NiagaraMeshVF.DefaultPos.xyz;
+	}
+	else
+	{
+		WorldParticlePosition = GetVec3(NiagaraMeshVF.PositionDataOffset, InstanceId);
+	}
+
+	// Apply CameraOffset
+	WorldParticlePosition += CalculateCameraOffset(WorldParticlePosition, InstanceId);
+
+	////////////////////////////////////////////////////////////////////////////////
+	// Rotation
+	float3x3 RotationMat = {
+		1, 0, 0,
+		0, 1, 0,
+		0, 0, 1
+	};
+
+	if(NiagaraMeshVF.TransformDataOffset != -1)
+	{
+		float4 Quat = GetVec4(NiagaraMeshVF.TransformDataOffset, InstanceId);
+		RotationMat = QuatTo3x3(normalize(Quat));
+	}
+
+	if (NiagaraMeshVF.FacingMode != MESH_FACING_DEFAULT)
+	{
+		// Factor in facing rotation
+		float3x3 FacingMat = GetFacingMatrix(WorldParticlePosition, Intermediates.ParticleVelocity);
+		RotationMat = mul(RotationMat, FacingMat);
+	}
+
+	////////////////////////////////////////////////////////////////////////////////
+	// Scale
+	float3 Scale = (float3)1.0f;
+	if(NiagaraMeshVF.ScaleDataOffset != -1)
+	{
+		Scale = GetVec3(NiagaraMeshVF.ScaleDataOffset, InstanceId);		
+	}
+	
+	////////////////////////////////////////////////////////////////////////////////
+	// Compose Transforms
+	LocalToWorld = ComposeTransformMatrix(Scale, RotationMat, WorldParticlePosition);
+	WorldToLocal = ComposeInvTransformMatrix(Scale, RotationMat, WorldParticlePosition);
+}
+
 FVertexFactoryIntermediates GetVertexFactoryIntermediates(FVertexFactoryInput Input)
 {
 	FVertexFactoryIntermediates Intermediates = (FVertexFactoryIntermediates)0;
@@ -550,15 +652,6 @@ FVertexFactoryIntermediates GetVertexFactoryIntermediates(FVertexFactoryInput In
 	}
 #endif
 
-	if (NiagaraMeshVF.PositionDataOffset == -1)
-	{
-		Intermediates.ParticlePosition = NiagaraMeshVF.DefaultPos;
-	}
-	else
-	{
-		Intermediates.ParticlePosition = float4(GetVec3(NiagaraMeshVF.PositionDataOffset, Input.InstanceId), 1);
-	}
-
 	if (NiagaraMeshVF.ColorDataOffset == -1)
 	{
 		Intermediates.ParticleColor = float4(1.0f, 1.0f, 1.0f, 1.0f);
@@ -645,71 +738,39 @@ FVertexFactoryIntermediates GetVertexFactoryIntermediates(FVertexFactoryInput In
 	// Swizzle vertex color.
 	Intermediates.VertexColor = Input.VertexColor FCOLOR_COMPONENT_SWIZZLE;
 
-	// Retrieve rotation
-	float3x3 RotationMat = {
-		1, 0, 0,
-		0, 1, 0,
-		0, 0, 1
-	};
-	if(NiagaraMeshVF.TransformDataOffset != -1)
-	{
-		float4 Quat = GetVec4(NiagaraMeshVF.TransformDataOffset, Input.InstanceId);
-		RotationMat = QuatTo3x3(normalize(Quat));
-	}
-	
-	// Retrieve Scale
-	float3 Scale = (float3)1.0f;
-	if(NiagaraMeshVF.ScaleDataOffset != -1)
-	{
-		Scale = GetVec3(NiagaraMeshVF.ScaleDataOffset, Input.InstanceId);		
-	}
-	
-	float3x3 ScaleMat = {
-		Scale.x,	0,			0,
-		0,			Scale.y,	0,
-		0,			0,			Scale.z
-	};
-	
-	float3x3 ScaleRotMat = mul(ScaleMat, RotationMat);
-	if (NiagaraMeshVF.FacingMode != MESH_FACING_DEFAULT)
-	{
-		// Factor in facing rotation
-		float3x3 FacingMat = GetFacingMatrix(Intermediates);
-		ScaleRotMat = mul(ScaleRotMat, FacingMat);
-	}
+	float4x4 ParticleToWorld;
+	float4x4 WorldToParticle;
 
-	// Apply CameraOffset
-	Intermediates.ParticlePosition.xyz += CalculateCameraOffset(Intermediates, Input.InstanceId);
-
-	// Determine final "simulation space" matrix
-	Intermediates.LocalToWorld = float4x4(
-		float4(ScaleRotMat[0], 0),
-		float4(ScaleRotMat[1], 0),
-		float4(ScaleRotMat[2], 0),
-		Intermediates.ParticlePosition
-	);
+	CalculateTransforms(Input.InstanceId, Intermediates, ParticleToWorld, WorldToParticle);
 
 	//-TODO: If / when we store previous transform data use them here (providing they are in local space, if not handle below)
 	//if (MeshParticleVF.PrevTransformAvailable > 0)
 	//{
-    //   //-TODO:
+	//   //-TODO:
 	//}
-    //else
-    {
+	//else
+	{
 		// NOTE: Velocity is in simulation space, so calculate prev transform here, before handling local space
-        Intermediates.PrevLocalToWorld = Intermediates.LocalToWorld;
-        Intermediates.PrevLocalToWorld[3].xyz -= ParticleVelocity * NiagaraMeshVF.DeltaSeconds;
-    }
+		Intermediates.PrevParticleToWorld = ParticleToWorld;
+		Intermediates.PrevParticleToWorld[3].xyz -= ParticleVelocity * NiagaraMeshVF.DeltaSeconds;
+	}
 
 	//Handle local space. TODO: Optimize behind define.
 	//Apply local space to local to world, also pretranspose here to avoid it in the vs&ps
 
 	if (NiagaraMeshVF.bLocalSpace)
 	{
-		Intermediates.PrevLocalToWorld = mul(Intermediates.PrevLocalToWorld, Primitive.PreviousLocalToWorld);
-		Intermediates.LocalToWorld = mul(Intermediates.LocalToWorld, Primitive.LocalToWorld);
-		Intermediates.ParticlePosition = Intermediates.LocalToWorld[3];
+		Intermediates.PrevParticleToWorld = mul(Intermediates.PrevParticleToWorld, Primitive.PreviousLocalToWorld);
+		Intermediates.ParticleToWorld = mul(ParticleToWorld, Primitive.LocalToWorld);
+		Intermediates.WorldToParticle = mul(Primitive.WorldToLocal, WorldToParticle);
 	}
+	else
+	{
+		Intermediates.ParticleToWorld = ParticleToWorld;
+		Intermediates.WorldToParticle = WorldToParticle;
+	}
+
+	Intermediates.ParticlePosition = Intermediates.ParticleToWorld[3];
 
 	// Translated world position.	
 	Intermediates.ParticleTranslatedWorldPosition = Intermediates.ParticlePosition.xyz + ResolvedView.PreViewTranslation.xyz;
@@ -790,7 +851,7 @@ float4 VertexFactoryGetRasterizedWorldPosition(FVertexFactoryInput Input, FVerte
 
 float3 VertexFactoryGetWorldNormal(FVertexFactoryInput Input, FVertexFactoryIntermediates Intermediates)
 {
-	return Intermediates.LocalToWorld[2].xyz;
+	return Intermediates.ParticleToWorld[2].xyz;
 }
 
 float3 VertexFactoryGetPositionForVertexLighting(FVertexFactoryInput Input, FVertexFactoryIntermediates Intermediates, float3 TranslatedWorldPosition)
@@ -894,11 +955,12 @@ FVertexFactoryInterpolantsVSToPS VertexFactoryGetInterpolantsVSToPS(FVertexFacto
 	Interpolants.ParticleRandom = Intermediates.ParticleRandom;
 #endif
 
-#if USE_PARTICLE_TRANSFORM
-	Interpolants.LocalToWorld[0] = Intermediates.LocalToWorld[0];
-	Interpolants.LocalToWorld[1] = Intermediates.LocalToWorld[1];
-	Interpolants.LocalToWorld[2] = Intermediates.LocalToWorld[2];
-	Interpolants.LocalToWorld[3] = Intermediates.LocalToWorld[3];
+#if USE_PARTICLE_LOCAL_TO_WORLD
+	Interpolants.ParticleToWorld = (float3x4) transpose(Intermediates.ParticleToWorld);
+#endif
+
+#if USE_PARTICLE_WORLD_TO_LOCAL
+	Interpolants.WorldToParticle = (float3x4) transpose(Intermediates.WorldToParticle);
 #endif
 
 #if INSTANCED_STEREO
diff --git a/Engine/Shaders/Private/LocalVertexFactory.ush b/Engine/Shaders/Private/LocalVertexFactory.ush
--- a/Engine/Shaders/Private/LocalVertexFactory.ush
+++ b/Engine/Shaders/Private/LocalVertexFactory.ush
@@ -465,6 +465,14 @@ FMaterialPixelParameters GetMaterialPixelParameters(FVertexFactoryInterpolantsVS
 	Result.TwoSidedSign = 1;
 	Result.PrimitiveId = GetPrimitiveId(Interpolants);
 
+#if NEEDS_PARTICLE_LOCAL_TO_WORLD
+	Result.Particle.ParticleToWorld = GetPrimitiveData(Result.PrimitiveId).LocalToWorld;
+#endif
+
+#if NEEDS_PARTICLE_WORLD_TO_LOCAL
+	Result.Particle.WorldToParticle = GetPrimitiveData(Result.PrimitiveId).WorldToLocal;
+#endif
+
 	return Result;
 }
 
diff --git a/Engine/Shaders/Private/MaterialTemplate.ush b/Engine/Shaders/Private/MaterialTemplate.ush
--- a/Engine/Shaders/Private/MaterialTemplate.ush
+++ b/Engine/Shaders/Private/MaterialTemplate.ush
@@ -157,8 +157,11 @@ struct FMaterialParticleParameters
 	half4 DynamicParameter3;
 #endif
 
-	/** mesh particle orientation */
-	float4x4 LocalToWorld;
+	/** mesh particle transform */
+	float4x4 ParticleToWorld;
+
+	/** Inverse mesh particle transform */
+	float4x4 WorldToParticle;
 
 #if USE_PARTICLE_SUBUVS
 	/** SubUV texture coordinates*/
diff --git a/Engine/Shaders/Private/MeshParticleVertexFactory.ush b/Engine/Shaders/Private/MeshParticleVertexFactory.ush
--- a/Engine/Shaders/Private/MeshParticleVertexFactory.ush
+++ b/Engine/Shaders/Private/MeshParticleVertexFactory.ush
@@ -120,8 +120,8 @@ struct FVertexFactoryInterpolantsVSToPS
 	float ParticleTime : PARTICLE_TIME;
 #endif
 
-#if NEEDS_PARTICLE_TRANSFORM
-	float4 LocalToWorld[4] : PARTICLE_TRANSFORM;
+#if NEEDS_PARTICLE_LOCAL_TO_WORLD || NEEDS_PARTICLE_WORLD_TO_LOCAL
+	float3x4 LocalToWorld : PARTICLE_TRANSFORM;
 #endif
 
 #if INSTANCED_STEREO
@@ -159,8 +159,8 @@ struct FVertexFactoryIntermediates
 	/** Relative time. */
 	float RelativeTime;
 
-#if NEEDS_PARTICLE_TRANSFORM
-	float4 LocalToWorld[4];
+#if NEEDS_PARTICLE_LOCAL_TO_WORLD || NEEDS_PARTICLE_WORLD_TO_LOCAL
+	float4x4 LocalToWorld;
 #endif
 };
 
@@ -199,6 +199,28 @@ float4x4 GetParticlePrevTransform(FVertexFactoryInput Input)
 }
 #endif
 
+float4x4 PseudoInverseTransform(float4x4 InTransform)
+{
+	// we'll assume that the incoming transform is composed of Scale * Rotation * Translation
+	// we'll additionally assume that the scale is uniform
+	float3x3 InvRotScale = transpose((float3x3) InTransform);
+
+	float4x4 InvRotScale44 = float4x4(
+		InvRotScale[0], 0.0f,
+		InvRotScale[1], 0.0f,
+		InvRotScale[2], 0.0f,
+		0.0f, 0.0f, 0.0f, 1.0f);
+
+	float4x4 InvTransMat = float4x4(
+		1.0f, 0.0f, 0.0f, 0.0f,
+		0.0f, 1.0f, 0.0f, 0.0f,
+		0.0f, 0.0f, 1.0f, 0.0f,
+		-InTransform[3].xyz, 1.0f
+	);
+
+	return mul(InvTransMat, InvRotScale44);
+}
+
 /** Converts from vertex factory specific interpolants to a FMaterialPixelParameters, which is used by material inputs. */
 FMaterialPixelParameters GetMaterialPixelParameters(FVertexFactoryInterpolantsVSToPS Interpolants, float4 SvPosition)
 {
@@ -268,11 +290,17 @@ FMaterialPixelParameters GetMaterialPixelParameters(FVertexFactoryInterpolantsVS
 #if USE_PARTICLE_TIME
 	Result.Particle.RelativeTime = Interpolants.ParticleTime;
 #endif
-#if NEEDS_PARTICLE_TRANSFORM
-	Result.Particle.LocalToWorld[0] = Interpolants.LocalToWorld[0];
-	Result.Particle.LocalToWorld[1] = Interpolants.LocalToWorld[1];
-	Result.Particle.LocalToWorld[2] = Interpolants.LocalToWorld[2];
-	Result.Particle.LocalToWorld[3] = Interpolants.LocalToWorld[3];
+
+#if NEEDS_PARTICLE_LOCAL_TO_WORLD || NEEDS_PARTICLE_WORLD_TO_LOCAL
+	float4x4 ParticleTransform = transpose(float4x4(Interpolants.LocalToWorld, float4(0.0f, 0.0f, 0.0f, 1.0f)));
+
+	#if NEEDS_PARTICLE_LOCAL_TO_WORLD
+		Result.Particle.ParticleToWorld = ParticleTransform;
+	#endif
+
+	#if NEEDS_PARTICLE_WORLD_TO_LOCAL
+		Result.Particle.WorldToParticle = PseudoInverseTransform(ParticleTransform);
+	#endif
 #endif
 
 	Result.Particle.MotionBlurFade = 1.0f;
@@ -304,21 +332,23 @@ FMaterialVertexParameters GetMaterialVertexParameters(FVertexFactoryInput Input,
 	Result.Particle.DynamicParameter = Intermediates.DynamicParameter;
 #endif	//USE_DYNAMIC_PARAMETERS
 
-	float4x4 Transform = GetParticleTransform(Input);
-	Result.InstanceLocalToWorld = Transform;
-#if NEEDS_PARTICLE_TRANSFORM
-	Result.Particle.LocalToWorld[0] = Transform[0];
-	Result.Particle.LocalToWorld[1] = Transform[1];
-	Result.Particle.LocalToWorld[2] = Transform[2];
-	Result.Particle.LocalToWorld[3] = Transform[3];
+	float4x4 ParticleTransform = GetParticleTransform(Input);
+	Result.InstanceLocalToWorld = Primitive_LocalToWorld;
+
+#if NEEDS_PARTICLE_LOCAL_TO_WORLD
+	Result.Particle.ParticleToWorld = ParticleTransform;
 #endif
 
-	Result.TangentToWorld = mul(TangentToLocal, (float3x3)Transform ); 
+#if NEEDS_PARTICLE_WORLD_TO_LOCAL
+	Result.Particle.WorldToParticle = PseudoInverseTransform(ParticleTransform);
+#endif
+
+	Result.TangentToWorld = mul(TangentToLocal, (float3x3)ParticleTransform ); 
 
 #if USE_PREVIOUS_TRANSFORM
 	Result.PrevFrameLocalToWorld = GetParticlePrevTransform(Input);
 #else
-	Result.PrevFrameLocalToWorld = GetParticleTransform(Input);
+	Result.PrevFrameLocalToWorld = ParticleTransform;
 #endif
 
 #if NUM_MATERIAL_TEXCOORDS_VERTEX
@@ -361,19 +391,14 @@ FVertexFactoryIntermediates GetVertexFactoryIntermediates(FVertexFactoryInput In
 	Intermediates.VertexColor = Input.VertexColor FCOLOR_COMPONENT_SWIZZLE;
 
 	// World position and localtoworld transform.
-#if PARTICLE_MESH_INSTANCED
-	float3 ParticleWorldPosition = float3( Input.Transform1.w, Input.Transform2.w, Input.Transform3.w);
-	float4x4 LocalToWorldMat = float4x4(Input.Transform1, Input.Transform2, Input.Transform3, float4(0.0f, 0.0f, 0.0f, 1.0f));
-#if NEEDS_PARTICLE_TRANSFORM
-	Intermediates.LocalToWorld[0] = LocalToWorldMat[0];
-	Intermediates.LocalToWorld[1] = LocalToWorldMat[1];
-	Intermediates.LocalToWorld[2] = LocalToWorldMat[2];
-	Intermediates.LocalToWorld[3] = LocalToWorldMat[3];
-#endif
-#else
-	float3 ParticleWorldPosition = float3( Transform1.w, Transform2.w, Transform3.w);
+	float4x4 LocalToWorldMat = GetParticleTransform(Input);
+
+#if NEEDS_PARTICLE_LOCAL_TO_WORLD || NEEDS_PARTICLE_WORLD_TO_LOCAL
+	Intermediates.LocalToWorld = LocalToWorldMat;
 #endif
 
+	float3 ParticleWorldPosition = LocalToWorldMat[3].xyz;
+
 	Intermediates.ParticleTranslatedWorldPosition = ParticleWorldPosition + ResolvedView.PreViewTranslation.xyz;
 
 #if PARTICLE_MESH_INSTANCED
diff --git a/Engine/Source/Runtime/Engine/Private/Materials/HLSLMaterialTranslator.cpp b/Engine/Source/Runtime/Engine/Private/Materials/HLSLMaterialTranslator.cpp
--- a/Engine/Source/Runtime/Engine/Private/Materials/HLSLMaterialTranslator.cpp
+++ b/Engine/Source/Runtime/Engine/Private/Materials/HLSLMaterialTranslator.cpp
@@ -93,7 +93,8 @@ FHLSLMaterialTranslator::FHLSLMaterialTranslator(FMaterial* InMaterial,
 ,	bUsesSkyAtmosphere(false)
 ,	bUsesVertexColor(false)
 ,	bUsesParticleColor(false)
-,	bUsesParticleTransform(false)
+,	bUsesParticleLocalToWorld(false)
+,	bUsesParticleWorldToLocal(false)
 ,	bUsesVertexPosition(false)
 ,	bUsesTransformVector(false)
 ,	bCompilingPreviousFrame(false)
@@ -1215,7 +1216,8 @@ void FHLSLMaterialTranslator::GetMaterialEnvironment(EShaderPlatform InPlatform,
 	OutEnvironment.SetDefine(TEXT("MATERIAL_SKY_ATMOSPHERE"), bUsesSkyAtmosphere);
 	OutEnvironment.SetDefine(TEXT("INTERPOLATE_VERTEX_COLOR"), bUsesVertexColor);
 	OutEnvironment.SetDefine(TEXT("NEEDS_PARTICLE_COLOR"), bUsesParticleColor); 
-	OutEnvironment.SetDefine(TEXT("NEEDS_PARTICLE_TRANSFORM"), bUsesParticleTransform);
+	OutEnvironment.SetDefine(TEXT("NEEDS_PARTICLE_LOCAL_TO_WORLD"), bUsesParticleLocalToWorld);
+	OutEnvironment.SetDefine(TEXT("NEEDS_PARTICLE_WORLD_TO_LOCAL"), bUsesParticleWorldToLocal);
 	OutEnvironment.SetDefine(TEXT("USES_TRANSFORM_VECTOR"), bUsesTransformVector);
 	OutEnvironment.SetDefine(TEXT("WANT_PIXEL_DEPTH_OFFSET"), bUsesPixelDepthOffset);
 	if (IsMetalPlatform(InPlatform))
@@ -5921,8 +5923,11 @@ int32 FHLSLMaterialTranslator::TransformBase(EMaterialCommonBasis SourceCoordBas
 			}
 			else if (DestCoordBasis == MCB_MeshParticle)
 			{
-				CodeStr = TEXT("mul(<A>, <MATRIX>(Parameters.Particle.LocalToWorld))");
-				bUsesParticleTransform = true;
+				CodeStr = TEXT("mul(<A>, <MATRIX>(Parameters.Particle.WorldToParticle))");
+				if (ShaderFrequency == SF_Pixel)
+				{
+					bUsesParticleWorldToLocal = true;
+				}
 			}
 
 			// else use MCB_TranslatedWorld as intermediary basis
@@ -5953,13 +5958,13 @@ int32 FHLSLMaterialTranslator::TransformBase(EMaterialCommonBasis SourceCoordBas
 		{
 			if (DestCoordBasis == MCB_World)
 			{
-				CodeStr = TEXT("mul(<MATRIX>(Parameters.Particle.LocalToWorld), <A>)");
-				bUsesParticleTransform = true;
-			}
-			else
-			{
-				return Errorf(TEXT("Can transform only to world space from particle space"));
+				CodeStr = TEXT("mul(<A>, <MATRIX>(Parameters.Particle.ParticleToWorld))");
+				if (ShaderFrequency == SF_Pixel)
+				{
+					bUsesParticleLocalToWorld = true;
+				}
 			}
+			// use World as an intermediary base
 			break;
 		}
 		default:
diff --git a/Engine/Source/Runtime/Engine/Private/Materials/HLSLMaterialTranslator.h b/Engine/Source/Runtime/Engine/Private/Materials/HLSLMaterialTranslator.h
--- a/Engine/Source/Runtime/Engine/Private/Materials/HLSLMaterialTranslator.h
+++ b/Engine/Source/Runtime/Engine/Private/Materials/HLSLMaterialTranslator.h
@@ -252,8 +252,10 @@ protected:
 	uint32 bUsesVertexColor : 1;
 	/** true if the material reads particle color in the pixel shader. */
 	uint32 bUsesParticleColor : 1;
-	/** true if the material reads mesh particle transform in the pixel shader. */
-	uint32 bUsesParticleTransform : 1;
+	/** true if the material reads mesh particle local to world in the pixel shader. */
+	uint32 bUsesParticleLocalToWorld : 1;
+	/** true if the material reads mesh particle world to local in the pixel shader. */
+	uint32 bUsesParticleWorldToLocal : 1;
 
 	/** true if the material uses any type of vertex position */
 	uint32 bUsesVertexPosition : 1;
