Fixed compile after merge.

#rnx

diff --git a/Engine/Plugins/FX/Niagara/Source/Niagara/Private/NiagaraScript.cpp b/Engine/Plugins/FX/Niagara/Source/Niagara/Private/NiagaraScript.cpp
--- a/Engine/Plugins/FX/Niagara/Source/Niagara/Private/NiagaraScript.cpp
+++ b/Engine/Plugins/FX/Niagara/Source/Niagara/Private/NiagaraScript.cpp
@@ -1904,7 +1904,7 @@ void UNiagaraScript::CacheResourceShadersForRendering(bool bRegenerateId, bool b
 			ERHIFeatureLevel::Type CacheFeatureLevel = GMaxRHIFeatureLevel;
 			const EShaderPlatform ShaderPlatform = GShaderPlatformForFeatureLevel[CacheFeatureLevel];
 
-			ScriptResource->SetScript(this, FeatureLevel, ShaderPlatform, CachedScriptVMId.CompilerVersionID, CachedScriptVMId.AdditionalDefines,
+			ScriptResource->SetScript(this, CacheFeatureLevel, ShaderPlatform, CachedScriptVMId.CompilerVersionID, CachedScriptVMId.AdditionalDefines,
 				CachedScriptVMId.BaseScriptCompileHash, CachedScriptVMId.ReferencedCompileHashes, 
 				CachedScriptVMId.bUsesRapidIterationParams, GetFriendlyName());
 
