Additional tweaks to namepsace colors and drop shadow intensity

diff --git a/Engine/Plugins/FX/Niagara/Source/NiagaraEditor/Private/NiagaraEditorSettings.cpp b/Engine/Plugins/FX/Niagara/Source/NiagaraEditor/Private/NiagaraEditorSettings.cpp
--- a/Engine/Plugins/FX/Niagara/Source/NiagaraEditor/Private/NiagaraEditorSettings.cpp
+++ b/Engine/Plugins/FX/Niagara/Source/NiagaraEditor/Private/NiagaraEditorSettings.cpp
@@ -85,7 +85,7 @@ UNiagaraEditorSettings::UNiagaraEditorSettings(const FObjectInitializer& ObjectI
 			{FNiagaraConstants::UserNamespace},
 			NSLOCTEXT("NamespaceMetadata", "UserDisplayName", "User"),
 			NSLOCTEXT("NamespaceMetadata", "UserDescription", "A read only value which can be initialized per system and\nmodified externally in the level, by blueprint, or by c++."),
-			FLinearColor(FColor(49, 113, 142)),
+			FLinearColor(FColor(91, 161, 194)),
 			{ }),
 		FNiagaraNamespaceMetadata(
 			{FNiagaraConstants::DataInstanceNamespace},
diff --git a/Engine/Plugins/FX/Niagara/Source/NiagaraEditor/Private/NiagaraEditorStyle.cpp b/Engine/Plugins/FX/Niagara/Source/NiagaraEditor/Private/NiagaraEditorStyle.cpp
--- a/Engine/Plugins/FX/Niagara/Source/NiagaraEditor/Private/NiagaraEditorStyle.cpp
+++ b/Engine/Plugins/FX/Niagara/Source/NiagaraEditor/Private/NiagaraEditorStyle.cpp
@@ -189,7 +189,7 @@ TSharedRef< FSlateStyleSet > FNiagaraEditorStyle::Create()
 		.SetFont(DEFAULT_FONT("Bold", 8))
 		.SetColorAndOpacity(FLinearColor(1.0f, 1.0f, 1.0f, 0.9f))
 		.SetShadowOffset(FVector2D(1, 1))
-		.SetShadowColorAndOpacity(FLinearColor(0, 0, 0, 0.9f)));
+		.SetShadowColorAndOpacity(FLinearColor(0, 0, 0, 0.7f)));
 
 	// Code View
 	{
