Fixed script compile error when creating a new custom expression.

diff --git a/Engine/Plugins/FX/Niagara/Source/NiagaraEditorWidgets/Private/Stack/SNiagaraStackFunctionInputValue.cpp b/Engine/Plugins/FX/Niagara/Source/NiagaraEditorWidgets/Private/Stack/SNiagaraStackFunctionInputValue.cpp
--- a/Engine/Plugins/FX/Niagara/Source/NiagaraEditorWidgets/Private/Stack/SNiagaraStackFunctionInputValue.cpp
+++ b/Engine/Plugins/FX/Niagara/Source/NiagaraEditorWidgets/Private/Stack/SNiagaraStackFunctionInputValue.cpp
@@ -724,7 +724,8 @@ void SNiagaraStackFunctionInputValue::DynamicInputScriptSelected(UNiagaraScript*
 
 void SNiagaraStackFunctionInputValue::CustomExpressionSelected()
 {
-	FunctionInput->SetCustomExpression(TEXT("// Insert expression here"));
+	FText CustomHLSLComment = LOCTEXT("NewCustomExpressionComment", "Custom HLSL!");
+	FunctionInput->SetCustomExpression(FHlslNiagaraTranslator::GetHlslDefaultForType(FunctionInput->GetInputType()) + TEXT(" /* ") + CustomHLSLComment.ToString() + TEXT(" */"));
 }
 
 void SNiagaraStackFunctionInputValue::CreateScratchSelected()
