Fix DDC access for NiagaraScript when iterating over standalone emitters

diff --git a/Engine/Plugins/FX/Niagara/Source/Niagara/Private/NiagaraScript.cpp b/Engine/Plugins/FX/Niagara/Source/Niagara/Private/NiagaraScript.cpp
--- a/Engine/Plugins/FX/Niagara/Source/Niagara/Private/NiagaraScript.cpp
+++ b/Engine/Plugins/FX/Niagara/Source/Niagara/Private/NiagaraScript.cpp
@@ -1577,7 +1577,7 @@ void UNiagaraScript::RequestCompile(bool bForceCompile)
 		}
 
 		// check the ddc first
-		if (!GetDerivedDataCacheRef().GetSynchronous(*GetNiagaraDDCKeyString(), OutData, GetPathName()))
+		if (GetDerivedDataCacheRef().GetSynchronous(*GetNiagaraDDCKeyString(), OutData, GetPathName()))
 		{
 			FNiagaraVMExecutableData ExeData;
 			if (BinaryToExecData(OutData, ExeData))
