Added an optional update context to help reduce latency when recompiling the scripts for a system

diff --git a/Engine/Plugins/FX/Niagara/Source/Niagara/Classes/NiagaraSystem.h b/Engine/Plugins/FX/Niagara/Source/Niagara/Classes/NiagaraSystem.h
--- a/Engine/Plugins/FX/Niagara/Source/Niagara/Classes/NiagaraSystem.h
+++ b/Engine/Plugins/FX/Niagara/Source/Niagara/Classes/NiagaraSystem.h
@@ -276,7 +276,7 @@ public:
 	void RemoveSystemParametersForEmitter(const FNiagaraEmitterHandle& EmitterHandle);
 
 	/** Request that any dirty scripts referenced by this system be compiled.*/
-	bool RequestCompile(bool bForce);
+	bool RequestCompile(bool bForce, FNiagaraSystemUpdateContext* OptionalUpdateContext = nullptr);
 
 	/** If we have a pending compile request, is it done with yet? */
 	bool PollForCompilationComplete();
diff --git a/Engine/Plugins/FX/Niagara/Source/Niagara/Private/NiagaraSystem.cpp b/Engine/Plugins/FX/Niagara/Source/Niagara/Private/NiagaraSystem.cpp
--- a/Engine/Plugins/FX/Niagara/Source/Niagara/Private/NiagaraSystem.cpp
+++ b/Engine/Plugins/FX/Niagara/Source/Niagara/Private/NiagaraSystem.cpp
@@ -1109,7 +1109,7 @@ void UNiagaraSystem::InitEmitterDataSetCompiledData(FNiagaraDataSetCompiledData&
 }
 #endif
 
-bool UNiagaraSystem::RequestCompile(bool bForce)
+bool UNiagaraSystem::RequestCompile(bool bForce, FNiagaraSystemUpdateContext* OptionalUpdateContext)
 {
 	if (bForce)
 	{
@@ -1234,7 +1234,14 @@ bool UNiagaraSystem::RequestCompile(bool bForce)
 		}
 	}
 
-	FNiagaraSystemUpdateContext UpdateCtx(this, true);
+	if (OptionalUpdateContext)
+	{
+		OptionalUpdateContext->Add(this, true);
+	}
+	else
+	{
+		FNiagaraSystemUpdateContext UpdateCtx(this, true);
+	}
 
 	return bAnyCompiled;
 }
diff --git a/Engine/Plugins/FX/Niagara/Source/NiagaraEditor/Private/NiagaraEditorUtilities.cpp b/Engine/Plugins/FX/Niagara/Source/NiagaraEditor/Private/NiagaraEditorUtilities.cpp
--- a/Engine/Plugins/FX/Niagara/Source/NiagaraEditor/Private/NiagaraEditorUtilities.cpp
+++ b/Engine/Plugins/FX/Niagara/Source/NiagaraEditor/Private/NiagaraEditorUtilities.cpp
@@ -861,37 +861,43 @@ TSharedPtr<SWidget> FNiagaraEditorUtilities::CreateInlineErrorText(TAttribute<FT
 
 void FNiagaraEditorUtilities::CompileExistingEmitters(const TArray<UNiagaraEmitter*>& AffectedEmitters)
 {
-	TSet<UNiagaraEmitter*> CompiledEmitters;
-	for (UNiagaraEmitter* Emitter : AffectedEmitters)
+	TArray<TSharedPtr<FNiagaraSystemViewModel>> ExistingSystemViewModels;
+
 	{
-		// If we've already compiled this emitter, or it's invalid skip it.
-		if (Emitter == nullptr || CompiledEmitters.Contains(Emitter) || Emitter->IsPendingKillOrUnreachable())
-		{
-			continue;
-		}
+		FNiagaraSystemUpdateContext UpdateCtx;
 
-		// We only need to compile emitters referenced directly as instances by systems since emitters can now only be used in the context 
-		// of a system.
-		for (TObjectIterator<UNiagaraSystem> SystemIterator; SystemIterator; ++SystemIterator)
+		TSet<UNiagaraEmitter*> CompiledEmitters;
+		for (UNiagaraEmitter* Emitter : AffectedEmitters)
 		{
-			if (SystemIterator->ReferencesInstanceEmitter(*Emitter))
+			// If we've already compiled this emitter, or it's invalid skip it.
+			if (Emitter == nullptr || CompiledEmitters.Contains(Emitter) || Emitter->IsPendingKillOrUnreachable())
 			{
-				SystemIterator->RequestCompile(false);
+				continue;
+			}
 
-				TArray<TSharedPtr<FNiagaraSystemViewModel>> ExistingSystemViewModels;
-				FNiagaraSystemViewModel::GetAllViewModelsForObject(*SystemIterator, ExistingSystemViewModels);
-				for (TSharedPtr<FNiagaraSystemViewModel> SystemViewModel : ExistingSystemViewModels)
+			// We only need to compile emitters referenced directly as instances by systems since emitters can now only be used in the context 
+			// of a system.
+			for (TObjectIterator<UNiagaraSystem> SystemIterator; SystemIterator; ++SystemIterator)
+			{
+				if (SystemIterator->ReferencesInstanceEmitter(*Emitter))
 				{
-					SystemViewModel->RefreshAll();
-				}
+					SystemIterator->RequestCompile(false, &UpdateCtx);
 
-				for (const FNiagaraEmitterHandle& EmitterHandle : SystemIterator->GetEmitterHandles())
-				{
-					CompiledEmitters.Add(EmitterHandle.GetInstance());
+					FNiagaraSystemViewModel::GetAllViewModelsForObject(*SystemIterator, ExistingSystemViewModels);
+
+					for (const FNiagaraEmitterHandle& EmitterHandle : SystemIterator->GetEmitterHandles())
+					{
+						CompiledEmitters.Add(EmitterHandle.GetInstance());
+					}
 				}
 			}
 		}
 	}
+
+	for (TSharedPtr<FNiagaraSystemViewModel> SystemViewModel : ExistingSystemViewModels)
+	{
+		SystemViewModel->RefreshAll();
+	}
 }
 
 bool FNiagaraEditorUtilities::TryGetEventDisplayName(UNiagaraEmitter* Emitter, FGuid EventUsageId, FText& OutEventDisplayName)
