Minimal change to fix constant recompile spam.
This code is a bit messy now and in need of a tidy.
So close to release this change is absolute minimal to fix the recompile spam.
As the results does not contain the compile errors a later step does not enter it's error state in the UI. Causing an infinite loop of recompile requests.
Adding the errors to the results avoids this.

#rnx

diff --git a/Engine/Plugins/FX/Niagara/Source/NiagaraEditor/Private/NiagaraCompiler.cpp b/Engine/Plugins/FX/Niagara/Source/NiagaraEditor/Private/NiagaraCompiler.cpp
--- a/Engine/Plugins/FX/Niagara/Source/NiagaraEditor/Private/NiagaraCompiler.cpp
+++ b/Engine/Plugins/FX/Niagara/Source/NiagaraEditor/Private/NiagaraCompiler.cpp
@@ -1233,6 +1233,12 @@ TOptional<FNiagaraCompileResults> FHlslNiagaraCompiler::GetCompileResult(int32 J
 	}
 	DumpDebugInfo(CompileResults, false);
 
+	//Seems like Results is a bit of a cobbled together mess at this point.
+	//Ideally we can tidy this up in future.
+	//Doing this as a minimal risk free fix for not having errors passed through into the compile results.
+	Results.NumErrors = CompileResults.NumErrors;
+	Results.CompileEvents = CompileResults.CompileEvents;
+
 	CompilationJob.Reset();
 	return Results;
 }
