implement view options for param map view and filter categories and entries that don't pass the view options filter.

diff --git a/Engine/Plugins/FX/Niagara/Source/NiagaraEditor/Private/NiagaraEditorSettings.cpp b/Engine/Plugins/FX/Niagara/Source/NiagaraEditor/Private/NiagaraEditorSettings.cpp
--- a/Engine/Plugins/FX/Niagara/Source/NiagaraEditor/Private/NiagaraEditorSettings.cpp
+++ b/Engine/Plugins/FX/Niagara/Source/NiagaraEditor/Private/NiagaraEditorSettings.cpp
@@ -166,6 +166,20 @@ void UNiagaraEditorSettings::SetResetDependentSystemsWhenEditingEmitters(bool bI
 	}
 }
 
+bool UNiagaraEditorSettings::GetDisplayAdvancedParameterPanelCategories() const
+{
+	return bDisplayAdvancedParameterPanelCategories;
+}
+
+void UNiagaraEditorSettings::SetDisplayAdvancedParameterPanelCategories(bool bInDisplayAdvancedParameterPanelCategories)
+{
+	if (bDisplayAdvancedParameterPanelCategories != bInDisplayAdvancedParameterPanelCategories)
+	{
+		bDisplayAdvancedParameterPanelCategories = bInDisplayAdvancedParameterPanelCategories;
+		SaveConfig();
+	}
+}
+
 FNiagaraNewAssetDialogConfig UNiagaraEditorSettings::GetNewAssetDailogConfig(FName InDialogConfigKey) const
 {
 	const FNiagaraNewAssetDialogConfig* Config = NewAssetDialogConfigMap.Find(InDialogConfigKey);
diff --git a/Engine/Plugins/FX/Niagara/Source/NiagaraEditor/Private/NiagaraEditorStyle.cpp b/Engine/Plugins/FX/Niagara/Source/NiagaraEditor/Private/NiagaraEditorStyle.cpp
--- a/Engine/Plugins/FX/Niagara/Source/NiagaraEditor/Private/NiagaraEditorStyle.cpp
+++ b/Engine/Plugins/FX/Niagara/Source/NiagaraEditor/Private/NiagaraEditorStyle.cpp
@@ -9,6 +9,7 @@
 #include "Styling/SlateTypes.h"
 #include "Styling/CoreStyle.h"
 #include "Interfaces/IPluginManager.h"
+#include "Classes/EditorStyleSettings.h"
 
 TSharedPtr< FSlateStyleSet > FNiagaraEditorStyle::NiagaraEditorStyleInstance = NULL;
 
@@ -191,6 +192,16 @@ TSharedRef< FSlateStyleSet > FNiagaraEditorStyle::Create()
 		.SetShadowOffset(FVector2D(1, 1))
 		.SetShadowColorAndOpacity(FLinearColor(0, 0, 0, 0.7f)));
 
+	Style->Set("NiagaraEditor.Stack.HighlightedButtonBrush", new BOX_CORE_BRUSH("Common/ButtonHoverHint", FMargin(4 / 16.0f), GetDefault<UEditorStyleSettings>()->SelectionColor));
+
+	// Parameter Map View
+	Style->Set("NiagaraEditor.Stack.DepressedHighlightedButtonBrush", new BOX_CORE_BRUSH("Common/ButtonHoverHint", FMargin(4 / 16.0f), GetDefault<UEditorStyleSettings>()->PressedSelectionColor));
+	Style->Set("NiagaraEditor.Stack.ViewOptionsShadowColor", FLinearColor::Black);
+	Style->Set("NiagaraEditor.Stack.FlatButtonColor", FLinearColor(FColor(205, 205, 205)));
+
+	const FVector2D ViewOptionsShadowOffset = FVector2D(0, 1);
+	Style->Set("NiagaraEditor.Stack.ViewOptionsShadowOffset", ViewOptionsShadowOffset);
+
 	// Code View
 	{
 		Style->Set("NiagaraEditor.CodeView.Checkbox.Text", FTextBlockStyle(NormalText)
diff --git a/Engine/Plugins/FX/Niagara/Source/NiagaraEditor/Private/Widgets/SNiagaraParameterMapView.cpp b/Engine/Plugins/FX/Niagara/Source/NiagaraEditor/Private/Widgets/SNiagaraParameterMapView.cpp
--- a/Engine/Plugins/FX/Niagara/Source/NiagaraEditor/Private/Widgets/SNiagaraParameterMapView.cpp
+++ b/Engine/Plugins/FX/Niagara/Source/NiagaraEditor/Private/Widgets/SNiagaraParameterMapView.cpp
@@ -40,6 +40,8 @@
 #include "ViewModels/Stack/NiagaraStackSystemSettingsGroup.h"
 #include "SNiagaraGraphActionWidget.h"
 #include "NiagaraEditorSettings.h"
+#include "Classes/EditorStyleSettings.h"
+#include "NiagaraEditorStyle.h"
 
 #define LOCTEXT_NAMESPACE "NiagaraParameterMapView"
 
@@ -163,6 +165,29 @@ NiagaraParameterMapSectionID::Type NiagaraParameterMapSectionID::OnGetSectionFro
 	return SectionID;
 }
 
+bool NiagaraParameterMapSectionID::GetSectionIsAdvanced(const Type InSection)
+{
+	TArray<FName> SectionNamespaces;
+	OnGetSectionNamespaces(InSection, SectionNamespaces);
+	FNiagaraNamespaceMetadata NamespaceMetadata = GetDefault<UNiagaraEditorSettings>()->GetMetaDataForNamespaces(SectionNamespaces);
+	if (NamespaceMetadata.IsValid())
+	{
+		bool bIsAdvanced = false;
+		for (const ENiagaraNamespaceMetadataOptions& Option : NamespaceMetadata.Options)
+		{
+			if (Option == ENiagaraNamespaceMetadataOptions::Advanced)
+			{
+				return true;
+			}
+		}
+		return false;
+	}
+	else
+	{
+		return false;
+	}
+}
+
 void FNiagaraParameterMapViewCommands::RegisterCommands()
 {
 	UI_COMMAND(DeleteEntry, "Delete", "Delete this parameter", EUserInterfaceActionType::Button, FInputChord(EKeys::Platform_Delete));
@@ -193,12 +218,21 @@ SNiagaraParameterMapView::~SNiagaraParameterMapView()
 	}
 }
 
+const FSlateBrush* SNiagaraParameterMapView::GetViewOptionsBorderBrush()
+{
+	UNiagaraEditorSettings* Settings = GetMutableDefault<UNiagaraEditorSettings>();
+	return Settings->GetDisplayAdvancedParameterPanelCategories()
+		? FNiagaraEditorStyle::Get().GetBrush("NiagaraEditor.Stack.DepressedHighlightedButtonBrush")
+		: FEditorStyle::GetBrush("NoBrush");
+}
+
 void SNiagaraParameterMapView::Construct(const FArguments& InArgs, const TArray<TSharedRef<FNiagaraObjectSelection>>& InSelectedObjects, const EToolkitType InToolkitType, const TSharedPtr<FUICommandList>& InToolkitCommands)
 {
 	bNeedsRefresh = false;
 	ToolkitType = InToolkitType;
 	ToolkitCommands = InToolkitCommands;
 	AddParameterButtons.SetNum(NiagaraParameterMapSectionID::Num);
+	const FVector2D ViewOptionsShadowOffset = FNiagaraEditorStyle::Get().GetVector("NiagaraEditor.Stack.ViewOptionsShadowOffset");
 
 	SelectedScriptObjects = InSelectedObjects[0];
 	SelectedScriptObjects->OnSelectedObjectsChanged().AddSP(this, &SNiagaraParameterMapView::SelectedObjectsChanged);
@@ -225,6 +259,41 @@ void SNiagaraParameterMapView::Construct(const FArguments& InArgs, const TArray<
 	SAssignNew(FilterBox, SSearchBox)
 		.OnTextChanged(this, &SNiagaraParameterMapView::OnFilterTextChanged);
 
+	// View options
+	TSharedRef<SWidget> ViewOptionsWidget = SNew(SBorder)
+		.Padding(0)
+		.BorderImage_Static(&SNiagaraParameterMapView::GetViewOptionsBorderBrush)
+		[
+			SNew(SComboButton)
+			.ContentPadding(0)
+			.ForegroundColor(FSlateColor::UseForeground())
+			.ButtonStyle(FEditorStyle::Get(), "ToggleButton")
+			.AddMetaData<FTagMetaData>(FTagMetaData(TEXT("ViewOptions")))
+			.ToolTipText(LOCTEXT("ViewOptionsToolTip", "View Options"))
+			.OnGetMenuContent(this, &SNiagaraParameterMapView::GetViewOptionsMenu)
+			.ButtonContent()
+			[
+				SNew(SOverlay)
+				// drop shadow
+				+ SOverlay::Slot()
+				.VAlign(VAlign_Top)
+				.Padding(ViewOptionsShadowOffset.X, ViewOptionsShadowOffset.Y, 0, 0)
+				[
+					SNew(SImage)
+					.Image(FEditorStyle::GetBrush("GenericViewButton"))
+					.ColorAndOpacity(FNiagaraEditorStyle::Get().GetColor("NiagaraEditor.Stack.ViewOptionsShadowColor"))
+				]
+				+ SOverlay::Slot()
+				.VAlign(VAlign_Top)
+				[
+					SNew(SImage)
+					.Image(FEditorStyle::GetBrush("GenericViewButton"))
+					.ColorAndOpacity(FNiagaraEditorStyle::Get().GetColor("NiagaraEditor.Stack.FlatButtonColor"))
+				]
+			]
+		];
+
+
 	// create the main action list piece of this widget
 	SAssignNew(GraphActionMenu, SGraphActionMenu, false)
 		.OnGetFilterText(this, &SNiagaraParameterMapView::GetFilterText)
@@ -265,13 +334,20 @@ void SNiagaraParameterMapView::Construct(const FArguments& InArgs, const TArray<
 					.AutoHeight()
 					[
 						SNew(SHorizontalBox)
-
+						// Filter Box
 						+ SHorizontalBox::Slot()
 						.FillWidth(1.0f)
 						.VAlign(VAlign_Center)
 						[
 							FilterBox.ToSharedRef()
 						]
+						// Filter Box View Options
+						+SHorizontalBox::Slot()
+						.AutoWidth()
+						.Padding(4, 0, 0, 0)
+						[
+							ViewOptionsWidget
+						]
 					]
 				]
 			]
@@ -349,6 +425,37 @@ void SNiagaraParameterMapView::AddParameter(FNiagaraVariable NewVariable)
 		GraphActionMenu->OnRequestRenameOnActionNode();
 	}
 }
+
+TSharedRef<SWidget> SNiagaraParameterMapView::GetViewOptionsMenu()
+{
+	FMenuBuilder MenuBuilder(false, nullptr);
+
+	auto ToggleShowAdvancedCategoriesActionLambda = [this]() {
+		UNiagaraEditorSettings* Settings = GetMutableDefault<UNiagaraEditorSettings>();
+		Settings->SetDisplayAdvancedParameterPanelCategories(!Settings->GetDisplayAdvancedParameterPanelCategories());
+		Refresh();
+	};
+
+	auto GetShowAdvancedCategoriesCheckStateActionLambda = []() {
+		UNiagaraEditorSettings* Settings = GetMutableDefault<UNiagaraEditorSettings>();
+		return Settings->GetDisplayAdvancedParameterPanelCategories() ? ECheckBoxState::Checked : ECheckBoxState::Unchecked;;
+	};
+
+	MenuBuilder.AddMenuEntry(
+		LOCTEXT("ShowAdvancedCategoriesLabel", "Show Advanced Categories"),
+		LOCTEXT("ShowAdvancedCategoriesToolTip", "Display advanced categories for the parameter panel."),
+		FSlateIcon(),
+		FUIAction(
+			FExecuteAction::CreateLambda(ToggleShowAdvancedCategoriesActionLambda),
+			FCanExecuteAction(),
+			FGetActionCheckState::CreateLambda(GetShowAdvancedCategoriesCheckStateActionLambda)),
+		NAME_None,
+		EUserInterfaceActionType::Check
+	);
+		
+	return MenuBuilder.MakeWidget();
+}
+
 bool SNiagaraParameterMapView::AllowMakeTypeGeneric(const FNiagaraTypeDefinition& InType) const
 {
 	return InType != FNiagaraTypeDefinition::GetParameterMapDef();
@@ -377,6 +484,9 @@ TSharedRef<SWidget> SNiagaraParameterMapView::OnCreateWidgetForAction(struct FCr
 
 void SNiagaraParameterMapView::CollectAllActions(FGraphActionListBuilderBase& OutAllActions)
 {
+	UNiagaraEditorSettings* EditorSettings = GetMutableDefault<UNiagaraEditorSettings>();
+	const bool bDisplayAdvancedCategories = EditorSettings->GetDisplayAdvancedParameterPanelCategories();
+
 	if (Graphs.Num() == 0)
 	{
 		return;
@@ -439,6 +549,10 @@ void SNiagaraParameterMapView::CollectAllActions(FGraphActionListBuilderBase& Ou
 		}
 		if (!IsSystemToolkit() || (IsSystemToolkit() && Section != NiagaraParameterMapSectionID::MODULE_INPUT && Section != NiagaraParameterMapSectionID::MODULE_LOCAL))
 		{
+			if (bDisplayAdvancedCategories == false && NiagaraParameterMapSectionID::GetSectionIsAdvanced(Section))
+			{
+				continue;
+			}
 			const FText Name = FText::FromName(Parameter.GetName());
 			const FText Tooltip = FText::Format(TooltipFormat, FText::FromName(Parameter.GetName()), Parameter.GetType().GetNameText());
 			TSharedPtr<FNiagaraParameterAction> ParameterAction(new FNiagaraParameterAction(Parameter, ParameterEntry.Value, FText::GetEmpty(), Name, Tooltip, 0, FText(), Section));
@@ -469,6 +583,18 @@ void SNiagaraParameterMapView::CollectStaticSections(TArray<int32>& StaticSectio
 	StaticSectionIDs.Add(NiagaraParameterMapSectionID::EMITTER);
 	StaticSectionIDs.Add(NiagaraParameterMapSectionID::PARTICLE);
 	StaticSectionIDs.Add(NiagaraParameterMapSectionID::TRANSIENT);
+
+	UNiagaraEditorSettings* EditorSettings = GetMutableDefault<UNiagaraEditorSettings>();
+	if (EditorSettings->GetDisplayAdvancedParameterPanelCategories() == false)
+	{
+		for (int i = StaticSectionIDs.Num() - 1; i > 0; --i)
+		{
+			if (NiagaraParameterMapSectionID::GetSectionIsAdvanced(NiagaraParameterMapSectionID::Type(StaticSectionIDs[i])))
+			{
+				StaticSectionIDs.RemoveAt(i);
+			}
+		}
+	}
 }
 
 FReply SNiagaraParameterMapView::OnActionDragged(const TArray<TSharedPtr<FEdGraphSchemaAction>>& InActions, const FPointerEvent& MouseEvent)
diff --git a/Engine/Plugins/FX/Niagara/Source/NiagaraEditor/Public/NiagaraEditorSettings.h b/Engine/Plugins/FX/Niagara/Source/NiagaraEditor/Public/NiagaraEditorSettings.h
--- a/Engine/Plugins/FX/Niagara/Source/NiagaraEditor/Public/NiagaraEditorSettings.h
+++ b/Engine/Plugins/FX/Niagara/Source/NiagaraEditor/Public/NiagaraEditorSettings.h
@@ -122,6 +122,12 @@ public:
 	/** Sets whether or not to reset all components that include the system that is currently being reset */
 	void SetResetDependentSystemsWhenEditingEmitters(bool bInResetDependentSystemsWhenEditingEmitters);
 
+	/** Gets whether or not to display advanced categories for the parameter panel. */
+	bool GetDisplayAdvancedParameterPanelCategories() const;
+
+	/** Sets whether or not to display advanced categories for the parameter panel. */
+	void SetDisplayAdvancedParameterPanelCategories(bool bInDisplayAdvancedParameterPanelCategories);
+
 	FNiagaraNewAssetDialogConfig GetNewAssetDailogConfig(FName InDialogConfigKey) const;
 
 	void SetNewAssetDialogConfig(FName InDialogConfigKey, const FNiagaraNewAssetDialogConfig& InNewAssetDialogConfig);
@@ -164,7 +170,11 @@ private:
 
 	/** Whether or not to reset all components that include the system currently being reset. */
 	UPROPERTY(config, EditAnywhere, Category = SimulationOptions)
-		bool bResetDependentSystemsWhenEditingEmitters;
+	bool bResetDependentSystemsWhenEditingEmitters;
+
+	/** Whether or not to display advanced categories for the parameter panel. */
+	UPROPERTY(config, EditAnywhere, Category = Niagara)
+	bool bDisplayAdvancedParameterPanelCategories;
 
 	UPROPERTY(config)
 	TMap<FName, FNiagaraNewAssetDialogConfig> NewAssetDialogConfigMap;
diff --git a/Engine/Plugins/FX/Niagara/Source/NiagaraEditor/Public/Widgets/SNiagaraParameterMapView.h b/Engine/Plugins/FX/Niagara/Source/NiagaraEditor/Public/Widgets/SNiagaraParameterMapView.h
--- a/Engine/Plugins/FX/Niagara/Source/NiagaraEditor/Public/Widgets/SNiagaraParameterMapView.h
+++ b/Engine/Plugins/FX/Niagara/Source/NiagaraEditor/Public/Widgets/SNiagaraParameterMapView.h
@@ -22,6 +22,7 @@ class SComboButton;
 class FNiagaraObjectSelection;
 class UNiagaraGraph;
 struct FEdGraphSchemaAction;
+struct FSlateBrush;
 
 /* Enums to use when grouping the blueprint members in the list panel. The order here will determine the order in the list */
 namespace NiagaraParameterMapSectionID
@@ -51,6 +52,7 @@ namespace NiagaraParameterMapSectionID
 	static FText OnGetSectionTitle(const NiagaraParameterMapSectionID::Type InSection);
 	static void OnGetSectionNamespaces(const NiagaraParameterMapSectionID::Type InSection, TArray<FName>& OutSectionNamespaces);
 	static NiagaraParameterMapSectionID::Type OnGetSectionFromVariable(const FNiagaraVariable& InVar, bool IsStaticSwitchVariable, FNiagaraParameterHandle& OutParameterHandle, const NiagaraParameterMapSectionID::Type DefaultType = NiagaraParameterMapSectionID::Type::NONE);
+	static bool GetSectionIsAdvanced(const NiagaraParameterMapSectionID::Type InSection);
 };
 
 class FNiagaraParameterMapViewCommands : public TCommands<FNiagaraParameterMapViewCommands>
@@ -103,6 +105,10 @@ public:
 	static bool IsStaticSwitchParameter(const FNiagaraVariable& Variable, const TArray<TWeakObjectPtr<UNiagaraGraph>>& Graphs);
 
 private:
+	
+	TSharedRef<SWidget> GetViewOptionsMenu();
+	static const FSlateBrush* GetViewOptionsBorderBrush();
+
 	/** Function to bind to SNiagaraAddParameterMenus to filter types we allow creating in generic parameters*/
 	bool AllowMakeTypeGeneric(const FNiagaraTypeDefinition& InType) const;
 
