Prevent using a recommended deprecation asset if it is not a module script.

diff --git a/Engine/Plugins/FX/Niagara/Source/NiagaraEditor/Private/ViewModels/Stack/NiagaraStackFunctionInput.cpp b/Engine/Plugins/FX/Niagara/Source/NiagaraEditor/Private/ViewModels/Stack/NiagaraStackFunctionInput.cpp
--- a/Engine/Plugins/FX/Niagara/Source/NiagaraEditor/Private/ViewModels/Stack/NiagaraStackFunctionInput.cpp
+++ b/Engine/Plugins/FX/Niagara/Source/NiagaraEditor/Private/ViewModels/Stack/NiagaraStackFunctionInput.cpp
@@ -494,7 +494,15 @@ void UNiagaraStackFunctionInput::RefreshChildrenInternal(const TArray<UNiagaraSt
 						NewIssues[AddIdx].InsertFix(0,
 							FStackIssueFix(
 							LOCTEXT("SelectNewDynamicInputScriptFixUseRecommended", "Use recommended replacement"),
-							FStackIssueFixDelegate::CreateLambda([this]() { ReassignDynamicInputScript(InputValues.DynamicNode->FunctionScript->DeprecationRecommendation); })));
+							FStackIssueFixDelegate::CreateLambda([this]() 
+								{ 
+									if (InputValues.DynamicNode->FunctionScript->DeprecationRecommendation->GetUsage() != ENiagaraScriptUsage::DynamicInput)
+									{
+										FNiagaraEditorUtilities::WarnWithToastAndLog(LOCTEXT("FailedDynamicInputDeprecationReplacement", "Failed to replace dynamic input as recommended replacement script is not a dynamic input!"));
+										return;
+									}
+									ReassignDynamicInputScript(InputValues.DynamicNode->FunctionScript->DeprecationRecommendation); 
+								})));
 					}
 				}
 
diff --git a/Engine/Plugins/FX/Niagara/Source/NiagaraEditor/Private/ViewModels/Stack/NiagaraStackScriptItemGroup.cpp b/Engine/Plugins/FX/Niagara/Source/NiagaraEditor/Private/ViewModels/Stack/NiagaraStackScriptItemGroup.cpp
--- a/Engine/Plugins/FX/Niagara/Source/NiagaraEditor/Private/ViewModels/Stack/NiagaraStackScriptItemGroup.cpp
+++ b/Engine/Plugins/FX/Niagara/Source/NiagaraEditor/Private/ViewModels/Stack/NiagaraStackScriptItemGroup.cpp
@@ -1071,6 +1071,12 @@ void UNiagaraStackScriptItemGroup::ChildRequestDeprecatedRecommendation(UNiagara
 
 	// Step 1: Make a copy of the existing node so that we can perform surgery on it to update.
 	UNiagaraScript* TargetScript = TargetChild->GetModuleNode().FunctionScript->DeprecationRecommendation;
+	if (TargetScript->GetUsage() != ENiagaraScriptUsage::Module)
+	{
+		// Early out if the deprecation recommendation asset is not module script usage.
+		FNiagaraEditorUtilities::WarnWithToastAndLog(LOCTEXT("DeprecationAssetReplacementFail","Failed to insert recommended deprecation asset as it is not a module script!"));
+		return;
+	}
 	TargetChild->Copy(ClipboardContent);
 
 	// Step 2: Disable the old master so that end users can use it as reference and inheritance isn't wiped out.
