Niagara - New parameter workflow fixes.
+ Remove static switches and data instance from the system/emitter parameters panel.
+ In the system/emitter parameters panel, only show the parameters which are actually encountered when the graph would compile.  This will now show properly alised attributes and outputs, and will skip disabled modules.
+ Focus the namespace modifier textblock when it goes into edit mode.
+ Add the option for namespace metadata for namespace modifiers.

diff --git a/Engine/Plugins/FX/Niagara/Source/NiagaraEditor/Private/NiagaraActions.cpp b/Engine/Plugins/FX/Niagara/Source/NiagaraEditor/Private/NiagaraActions.cpp
--- a/Engine/Plugins/FX/Niagara/Source/NiagaraEditor/Private/NiagaraActions.cpp
+++ b/Engine/Plugins/FX/Niagara/Source/NiagaraEditor/Private/NiagaraActions.cpp
@@ -37,7 +37,7 @@ FNiagaraParameterAction::FNiagaraParameterAction(const FNiagaraVariable& InParam
 	: FEdGraphSchemaAction(MoveTemp(InNodeCategory), MoveTemp(InMenuDesc), MoveTemp(InToolTip), InGrouping, MoveTemp(InKeywords), InSectionID)
 	, Parameter(InParameter)
 	, ReferenceCollection(InReferenceCollection)
-	, bSubnamespaceRenamePending(false)
+	, bNamespaceModifierRenamePending(false)
 {
 }
 
@@ -45,7 +45,7 @@ FNiagaraParameterAction::FNiagaraParameterAction(const FNiagaraVariable& InParam
 	FText InNodeCategory, FText InMenuDesc, FText InToolTip, const int32 InGrouping, FText InKeywords, int32 InSectionID)
 	: FEdGraphSchemaAction(MoveTemp(InNodeCategory), MoveTemp(InMenuDesc), MoveTemp(InToolTip), InGrouping, MoveTemp(InKeywords), InSectionID)
 	, Parameter(InParameter)
-	, bSubnamespaceRenamePending(false)
+	, bNamespaceModifierRenamePending(false)
 {
 }
 
diff --git a/Engine/Plugins/FX/Niagara/Source/NiagaraEditor/Private/NiagaraEditorSettings.cpp b/Engine/Plugins/FX/Niagara/Source/NiagaraEditor/Private/NiagaraEditorSettings.cpp
--- a/Engine/Plugins/FX/Niagara/Source/NiagaraEditor/Private/NiagaraEditorSettings.cpp
+++ b/Engine/Plugins/FX/Niagara/Source/NiagaraEditor/Private/NiagaraEditorSettings.cpp
@@ -24,7 +24,11 @@ UNiagaraEditorSettings::UNiagaraEditorSettings(const FObjectInitializer& ObjectI
 	bResetSimulationOnChange = true;
 	bResimulateOnChangeWhilePaused = true;
 	bResetDependentSystemsWhenEditingEmitters = false;
+	SetupNamespaceMetadata();
+}
 
+void UNiagaraEditorSettings::SetupNamespaceMetadata()
+{
 	NamespaceMetadata = 
 	{
 		FNiagaraNamespaceMetadata(
@@ -38,31 +42,37 @@ UNiagaraEditorSettings::UNiagaraEditorSettings(const FObjectInitializer& ObjectI
 			NSLOCTEXT("NamespaceMetadata", "SystemDisplayName", "System"),
 			NSLOCTEXT("NamespaceMetadata", "SystemDescription", "Persistent attribute in the system which is written in a system\n stage and can be read anywhere."),
 			FLinearColor(FColor(49, 113, 142)),
-			{ENiagaraNamespaceMetadataOptions::SupportsSubnamespace}),
+			{ENiagaraNamespaceMetadataOptions::SupportsNamespaceModifier}),
 		FNiagaraNamespaceMetadata(
 			{FNiagaraConstants::EmitterNamespace},
 			NSLOCTEXT("NamespaceMetadata", "EmitterDisplayName", "Emitter"),
 			NSLOCTEXT("NamespaceMetadata", "EmitterDescription", "Persistent attribute which is written in a emitter\n stage and can be read in emitter and particle stages."),
 			FLinearColor(FColor(145, 99, 56)),
-			{ENiagaraNamespaceMetadataOptions::SupportsSubnamespace}),
+			{ENiagaraNamespaceMetadataOptions::SupportsNamespaceModifier}),
 		FNiagaraNamespaceMetadata(
 			{FNiagaraConstants::ParticleAttributeNamespace},
 			NSLOCTEXT("NamespaceMetadata", "ParticleDisplayName", "Particles"),
 			NSLOCTEXT("NamespaceMetadata", "ParticleDescription", "Persistent attribute which is written in a particle\n stage and can be read in particle stages."),
 			FLinearColor(FColor(72, 130, 71)),
-			{ENiagaraNamespaceMetadataOptions::SupportsSubnamespace}),
+			{ENiagaraNamespaceMetadataOptions::SupportsNamespaceModifier}),
 		FNiagaraNamespaceMetadata(
 			{FNiagaraConstants::ModuleNamespace},
 			NSLOCTEXT("NamespaceMetadata", "ModuleDisplayName", "Input"),
 			NSLOCTEXT("NamespaceMetadata", "ModuleDescription", "A transient value which exposes a module input to the system and emitter editor."),
 			FLinearColor(FColor(136, 66, 65)),
-			{ENiagaraNamespaceMetadataOptions::SupportsSubnamespace}),
+			{ENiagaraNamespaceMetadataOptions::SupportsNamespaceModifier}),
 		FNiagaraNamespaceMetadata(
 			{FNiagaraConstants::OutputNamespace, FNiagaraConstants::ModuleNamespace},
 			NSLOCTEXT("NamespaceMetadata", "ModuleOutputDisplayName", "Output"),
 			NSLOCTEXT("NamespaceMetadata", "ModuleOutputDescription", "A transient value which is unique to the module that wrote it and can be read from any other module.\nTransient values do not persist from frame to frame, or between stages, e.g. emitter to particle, or spawn to update."),
 			FLinearColor(FColor(109, 95, 124)),
 			{ENiagaraNamespaceMetadataOptions::Advanced}),
+		FNiagaraNamespaceMetadata(
+			{FNiagaraConstants::OutputNamespace},
+			NSLOCTEXT("NamespaceMetadata", "OutputDisplayName", "Output"),
+			NSLOCTEXT("NamespaceMetadata", "OutputDescription", "A transient value which is unique to the module that wrote it and can be read from any other module.\nTransient values do not persist from frame to frame, or between stages, e.g. emitter to particle, or spawn to update."),
+			FLinearColor(FColor(109, 95, 124)),
+			{ENiagaraNamespaceMetadataOptions::Advanced}),
 		FNiagaraNamespaceMetadata(
 			{FNiagaraConstants::LocalNamespace, FNiagaraConstants::ModuleNamespace},
 			NSLOCTEXT("NamespaceMetadata", "ModuleLocalDisplayName", "Local"),
@@ -87,6 +97,12 @@ UNiagaraEditorSettings::UNiagaraEditorSettings(const FObjectInitializer& ObjectI
 			NSLOCTEXT("NamespaceMetadata", "UserDescription", "A read only value which can be initialized per system and\nmodified externally in the level, by blueprint, or by c++."),
 			FLinearColor(FColor(91, 161, 194)),
 			{ }),
+		FNiagaraNamespaceMetadata(
+			{FNiagaraConstants::ParameterCollectionNamespace},
+			NSLOCTEXT("NamespaceMetadata", "NiagaraParameterCollectionDisplayName", "Niagara Parameter Collection"),
+			NSLOCTEXT("NamespaceMetadata", "NiagaraParameterCollectionDescription", "Values read from a niagara parameter collection asset.\nRead only in a niagara system."),
+			FLinearColor(FColor(45, 45, 45)),
+			{ENiagaraNamespaceMetadataOptions::Advanced}),
 		FNiagaraNamespaceMetadata(
 			{FNiagaraConstants::DataInstanceNamespace},
 			NSLOCTEXT("NamespaceMetadata", "DataInstanceDisplayName", "Data Instance"),
@@ -94,6 +110,22 @@ UNiagaraEditorSettings::UNiagaraEditorSettings(const FObjectInitializer& ObjectI
 			FLinearColor(FColor(45, 45, 45)),
 			{ENiagaraNamespaceMetadataOptions::Advanced}),
 	};
+
+	NamespaceModifierMetadata =
+	{
+		FNiagaraNamespaceMetadata(
+			{FNiagaraConstants::InitialNamespace},
+			NSLOCTEXT("NamespaceMetadata", "InitialModifierDisplayName", "Initial"),
+			NSLOCTEXT("NamespaceMetadata", "InitialModifierDescription", "A namespace modifier for dataset attributes which when used in\na linked input in an update script will get the initial value from the spawn script."),
+			FLinearColor(FColor(102, 102, 102)),
+			{ }),
+		FNiagaraNamespaceMetadata(
+			{FNiagaraConstants::ModuleNamespace},
+			NSLOCTEXT("NamespaceMetadata", "ModuleModifierDisplayName", "Module"),
+			NSLOCTEXT("NamespaceMetadata", "ModuleModifierDescription", "A namespace modifier which is used when writing dataset attributes\nin module which will make that attribute unique to the module instance when used."),
+			FLinearColor(FColor(102, 102, 102)),
+			{ })
+	};
 }
 
 bool UNiagaraEditorSettings::GetAutoCompile() const
@@ -239,6 +271,18 @@ FNiagaraNamespaceMetadata UNiagaraEditorSettings::GetMetaDataForNamespaces(TArra
 	}
 }
 
+FNiagaraNamespaceMetadata UNiagaraEditorSettings::GetMetaDataForNamespaceModifier(FName NamespaceModifier) const
+{
+	for (const FNiagaraNamespaceMetadata& NamespaceModifierMetadataItem : NamespaceModifierMetadata)
+	{
+		if (NamespaceModifierMetadataItem.Namespaces.Num() == 1 && NamespaceModifierMetadataItem.Namespaces[0] == NamespaceModifier)
+		{
+			return NamespaceModifierMetadataItem;
+		}
+	}
+	return FNiagaraNamespaceMetadata();
+}
+
 FName UNiagaraEditorSettings::GetCategoryName() const
 {
 	return TEXT("Plugins");
diff --git a/Engine/Plugins/FX/Niagara/Source/NiagaraEditor/Private/SNiagaraParameterName.cpp b/Engine/Plugins/FX/Niagara/Source/NiagaraEditor/Private/SNiagaraParameterName.cpp
--- a/Engine/Plugins/FX/Niagara/Source/NiagaraEditor/Private/SNiagaraParameterName.cpp
+++ b/Engine/Plugins/FX/Niagara/Source/NiagaraEditor/Private/SNiagaraParameterName.cpp
@@ -74,81 +74,78 @@ void SNiagaraParameterName::UpdateContent(FName InDisplayedParameterName)
 	}
 
 	TSharedRef<SHorizontalBox> ContentBox = SNew(SHorizontalBox);
-	auto AddNamespaceSlot = [](TSharedRef<SHorizontalBox> ContentBox, FText NamespaceDisplayName, FText NamespaceDescription, FLinearColor NamespaceBorderColor)
-	{
-		ContentBox->AddSlot()
-		.VAlign(VAlign_Center)
-		.AutoWidth()
-		.Padding(0.0f, 0.0f, 5.0f, 0.0f)
-		[
 
-				SNew(SBorder)
-				.BorderImage(FNiagaraEditorStyle::Get().GetBrush("NiagaraEditor.ParameterName.NamespaceBorder"))
-				.BorderBackgroundColor(NamespaceBorderColor)
-				.ToolTipText(NamespaceDescription)
-				.VAlign(VAlign_Center)
-				.Padding(FMargin(4.0f, 1.0f, 4.0f, 1.0f))
-				[
-					SNew(STextBlock)
-					.TextStyle(FNiagaraEditorStyle::Get(), "NiagaraEditor.ParameterName.NamespaceText")
-					.Text(NamespaceDisplayName)
-				]
-		];
-	};
+	FNiagaraNamespaceMetadata DefaultNamespaceMetadata = GetDefault<UNiagaraEditorSettings>()->GetMetaDataForNamespaces({ NAME_None });
 
+	// Add the namespace widget.
 	FNiagaraNamespaceMetadata NamespaceMetadata = GetDefault<UNiagaraEditorSettings>()->GetMetaDataForNamespaces(Namespaces);
+	TSharedPtr<SWidget> NamespaceWidget;
 	if (NamespaceMetadata.IsValid())
 	{
 		Namespaces.RemoveAt(0, NamespaceMetadata.Namespaces.Num());
-		ContentBox->AddSlot()
-		.VAlign(VAlign_Center)
-		.AutoWidth()
-		.Padding(0.0f, 0.0f, 5.0f, 0.0f)
-		[
-			CreateNamespaceWidget(NamespaceMetadata.DisplayName.ToUpper(), NamespaceMetadata.Description, NamespaceMetadata.BackgroundColor)
-		];
+		NamespaceWidget = CreateNamespaceWidget(NamespaceMetadata.DisplayName.ToUpper(), NamespaceMetadata.Description, NamespaceMetadata.BackgroundColor);
 	}
-
-	FNiagaraNamespaceMetadata DefaultNamespaceMetadata = GetDefault<UNiagaraEditorSettings>()->GetMetaDataForNamespaces({NAME_None});
-	if (Namespaces.Num() > 0)
+	else
 	{
-		DisplayedSubnamespace = Namespaces[0];
+		FText NamespaceDisplayName = FText::FromString(FName::NameToDisplayString(Namespaces[0].ToString(), false).ToUpper());
 		Namespaces.RemoveAt(0);
-		SubnamespaceBorder = CreateNamespaceWidget(
-			FText::FromString(FName::NameToDisplayString(DisplayedSubnamespace.ToString(), false).ToUpper()),
+		CreateNamespaceWidget(
+			NamespaceDisplayName,
 			DefaultNamespaceMetadata.Description,
 			DefaultNamespaceMetadata.BackgroundColor);
+	}
+
+	ContentBox->AddSlot()
+	.VAlign(VAlign_Center)
+	.AutoWidth()
+	.Padding(0.0f, 0.0f, 5.0f, 0.0f)
+	[
+		NamespaceWidget.ToSharedRef()
+	];
+
+	// Next the namespace modifier widget is there is a namespace modifier.
+	if (Namespaces.Num() > 0)
+	{
+		DisplayedNamespaceModifier = Namespaces[0];
+		FNiagaraNamespaceMetadata DisplayedNamespaceModifierMetadata = GetDefault<UNiagaraEditorSettings>()->GetMetaDataForNamespaceModifier(DisplayedNamespaceModifier);
+		if (DisplayedNamespaceModifierMetadata.IsValid() == false)
+		{
+			DisplayedNamespaceModifierMetadata = DefaultNamespaceMetadata;
+		}
+
+		Namespaces.RemoveAt(0);
+		NamespaceModifierBorder = CreateNamespaceWidget(
+			FText::FromString(FName::NameToDisplayString(DisplayedNamespaceModifier.ToString(), false).ToUpper()),
+			DisplayedNamespaceModifierMetadata.Description,
+			DisplayedNamespaceModifierMetadata.BackgroundColor);
 
 		ContentBox->AddSlot()
 		.VAlign(VAlign_Center)
 		.AutoWidth()
 		.Padding(0.0f, 0.0f, 5.0f, 0.0f)
 		[
-			SubnamespaceBorder.ToSharedRef()
+			NamespaceModifierBorder.ToSharedRef()
 		];
 	}
 	else
 	{
-		DisplayedSubnamespace = NAME_None;
-		SubnamespaceBorder.Reset();
+		DisplayedNamespaceModifier = NAME_None;
+		NamespaceModifierBorder.Reset();
 	}
 
-	if (ensureMsgf(Namespaces.Num() == 0, TEXT("Extra namespaces in parameter.")) == false)
+	// If there are extra namespaces found, add them to the UI without metadata.
+	for (FName Namespace : Namespaces)
 	{
-		// If there are extra namespaces found, add them to the UI.
-		for (FName Namespace : Namespaces)
-		{
-			ContentBox->AddSlot()
-			.VAlign(VAlign_Center)
-			.AutoWidth()
-			.Padding(0.0f, 0.0f, 5.0f, 0.0f)
-			[
-				CreateNamespaceWidget(
-					FText::FromString(FName::NameToDisplayString(Namespace.ToString(), false).ToUpper()),
-					DefaultNamespaceMetadata.Description,
-					DefaultNamespaceMetadata.BackgroundColor)
-			];
-		}
+		ContentBox->AddSlot()
+		.VAlign(VAlign_Center)
+		.AutoWidth()
+		.Padding(0.0f, 0.0f, 5.0f, 0.0f)
+		[
+			CreateNamespaceWidget(
+				FText::FromString(FName::NameToDisplayString(Namespace.ToString(), false).ToUpper()),
+				DefaultNamespaceMetadata.Description,
+				DefaultNamespaceMetadata.BackgroundColor)
+		];
 	}
 
 	TSharedPtr<SWidget> NameWidget;
@@ -192,7 +189,7 @@ FName SNiagaraParameterName::ReconstructNameFromEditText(const FText& InEditText
 	return *NewParameterNameString;
 }
 
-FName SNiagaraParameterName::ReconstructNameFromSubnamespaceEditText(const FText& InEditText)
+FName SNiagaraParameterName::ReconstructNameFromNamespaceModifierEditText(const FText& InEditText)
 {
 	FString CurrentParameterNameString = ParameterName.Get().ToString();
 	TArray<FString> NameParts;
@@ -206,7 +203,7 @@ FName SNiagaraParameterName::ReconstructNameFromSubnamespaceEditText(const FText
 	return *NewParameterNameString;
 }
 
-FText SNiagaraParameterName::GetSubnamespaceText()
+FText SNiagaraParameterName::GetNamespaceModifierText()
 {
 	FString CurrentParameterNameString = ParameterName.Get().ToString();
 	TArray<FString> NameParts;
@@ -235,9 +232,9 @@ void SNiagaraParameterName::NameTextCommitted(const FText& InNewNameText, ETextC
 	}
 }
 
-bool SNiagaraParameterName::VerifySubnamespaceTextChange(const FText& InNewNameText, FText& OutErrorMessage)
+bool SNiagaraParameterName::VerifyNamespaceModifierTextChange(const FText& InNewNameText, FText& OutErrorMessage)
 {
-	FName NewParameterName = ReconstructNameFromSubnamespaceEditText(InNewNameText);
+	FName NewParameterName = ReconstructNameFromNamespaceModifierEditText(InNewNameText);
 	if (OnVerifyNameChangeDelegate.IsBound())
 	{
 		return OnVerifyNameChangeDelegate.Execute(NewParameterName, OutErrorMessage);
@@ -245,19 +242,19 @@ bool SNiagaraParameterName::VerifySubnamespaceTextChange(const FText& InNewNameT
 	return true;
 }
 
-void SNiagaraParameterName::SubnamespaceTextCommitted(const FText& InNewNameText, ETextCommit::Type CommitType)
+void SNiagaraParameterName::NamespaceModifierTextCommitted(const FText& InNewNameText, ETextCommit::Type CommitType)
 {
 	if (CommitType == ETextCommit::OnEnter)
 	{
-		FName NewParameterName = ReconstructNameFromSubnamespaceEditText(InNewNameText);
+		FName NewParameterName = ReconstructNameFromNamespaceModifierEditText(InNewNameText);
 		OnNameChangedDelegate.ExecuteIfBound(NewParameterName);
 	}
 
-	if (SubnamespaceBorder.IsValid())
+	if (NamespaceModifierBorder.IsValid())
 	{
-		SubnamespaceBorder->SetContent(SNew(STextBlock)
+		NamespaceModifierBorder->SetContent(SNew(STextBlock)
 			.TextStyle(FNiagaraEditorStyle::Get(), "NiagaraEditor.ParameterName.NamespaceText")
-			.Text(GetSubnamespaceText()));
+			.Text(GetNamespaceModifierText()));
 	}
 }
 
@@ -269,18 +266,18 @@ void SNiagaraParameterName::EnterEditingMode()
 	}
 }
 
-void SNiagaraParameterName::EnterSubnamespaceEditingMode()
+void SNiagaraParameterName::EnterNamespaceModifierEditingMode()
 {
-	if (SubnamespaceBorder.IsValid())
+	if (NamespaceModifierBorder.IsValid())
 	{
-		TSharedRef<SInlineEditableTextBlock> SubnamespaceEditableTextBlock = SNew(SInlineEditableTextBlock)
+		TSharedRef<SInlineEditableTextBlock> NamespaceModifierEditableTextBlock = SNew(SInlineEditableTextBlock)
 			.Style(EditableTextStyle)
-			.Text(GetSubnamespaceText())
+			.Text(GetNamespaceModifierText())
 			.IsSelected(IsSelected)
-			.OnVerifyTextChanged(this, &SNiagaraParameterName::VerifySubnamespaceTextChange)
-			.OnTextCommitted(this, &SNiagaraParameterName::SubnamespaceTextCommitted);
-		SubnamespaceEditableTextBlock->EnterEditingMode();
-		SubnamespaceBorder->SetContent(SubnamespaceEditableTextBlock);
+			.OnVerifyTextChanged(this, &SNiagaraParameterName::VerifyNamespaceModifierTextChange)
+			.OnTextCommitted(this, &SNiagaraParameterName::NamespaceModifierTextCommitted);
+		NamespaceModifierBorder->SetContent(NamespaceModifierEditableTextBlock);
+		NamespaceModifierEditableTextBlock->EnterEditingMode();
 	}
 }
 
@@ -335,7 +332,7 @@ void SNiagaraParameterNameTextBlock::EnterEditingMode()
 	ParameterName->EnterEditingMode();
 }
 
-void SNiagaraParameterNameTextBlock::EnterSubnamespaceEditingMode()
+void SNiagaraParameterNameTextBlock::EnterNamespaceModifierEditingMode()
 {
-	ParameterName->EnterSubnamespaceEditingMode();
+	ParameterName->EnterNamespaceModifierEditingMode();
 }
\ No newline at end of file
diff --git a/Engine/Plugins/FX/Niagara/Source/NiagaraEditor/Private/Widgets/SNiagaraParameterMapPaletteItem.cpp b/Engine/Plugins/FX/Niagara/Source/NiagaraEditor/Private/Widgets/SNiagaraParameterMapPaletteItem.cpp
--- a/Engine/Plugins/FX/Niagara/Source/NiagaraEditor/Private/Widgets/SNiagaraParameterMapPaletteItem.cpp
+++ b/Engine/Plugins/FX/Niagara/Source/NiagaraEditor/Private/Widgets/SNiagaraParameterMapPaletteItem.cpp
@@ -107,10 +107,10 @@ void SNiagaraParameterMapPalleteItem::Tick(const FGeometry& AllottedGeometry, co
 	if (ParameterNameTextBlock.IsValid())
 	{
 		TSharedPtr<FNiagaraParameterAction> ParameterAction = StaticCastSharedPtr<FNiagaraParameterAction>(ActionPtr.Pin());
-		if (ParameterAction.IsValid() && ParameterAction->bSubnamespaceRenamePending)
+		if (ParameterAction.IsValid() && ParameterAction->bNamespaceModifierRenamePending)
 		{
-			ParameterAction->bSubnamespaceRenamePending = false;
-			ParameterNameTextBlock->EnterSubnamespaceEditingMode();
+			ParameterAction->bNamespaceModifierRenamePending = false;
+			ParameterNameTextBlock->EnterNamespaceModifierEditingMode();
 		}
 	}
 }
diff --git a/Engine/Plugins/FX/Niagara/Source/NiagaraEditor/Private/Widgets/SNiagaraParameterMapView.cpp b/Engine/Plugins/FX/Niagara/Source/NiagaraEditor/Private/Widgets/SNiagaraParameterMapView.cpp
--- a/Engine/Plugins/FX/Niagara/Source/NiagaraEditor/Private/Widgets/SNiagaraParameterMapView.cpp
+++ b/Engine/Plugins/FX/Niagara/Source/NiagaraEditor/Private/Widgets/SNiagaraParameterMapView.cpp
@@ -38,10 +38,14 @@
 #include "Widgets/Layout/SBox.h"
 #include "NiagaraSystemEditorData.h"
 #include "ViewModels/Stack/NiagaraStackSystemSettingsGroup.h"
+#include "ViewModels/Stack/NiagaraStackGraphUtilities.h"
 #include "SNiagaraGraphActionWidget.h"
 #include "NiagaraEditorSettings.h"
 #include "Classes/EditorStyleSettings.h"
 #include "NiagaraEditorStyle.h"
+#include "NiagaraParameterMapHistory.h"
+#include "NiagaraNodeOutput.h"
+#include "NiagaraNodeEmitter.h"
 
 #define LOCTEXT_NAMESPACE "NiagaraParameterMapView"
 
@@ -210,6 +214,11 @@ SNiagaraParameterMapView::~SNiagaraParameterMapView()
 	}
 
 	EmptyGraphs();
+	if (CachedSystem.IsValid())
+	{
+		CachedSystem->GetExposedParameters().RemoveOnChangedHandler(UserParameterStoreChangedHandle);
+		CachedSystem.Reset();
+	}
 
 	SelectedScriptObjects->OnSelectedObjectsChanged().RemoveAll(this);
 	if (SelectedVariableObjects)
@@ -493,6 +502,45 @@ void SNiagaraParameterMapView::CollectAllActions(FGraphActionListBuilderBase& Ou
 	}
 
 	TMap<FNiagaraVariable, TArray<FNiagaraGraphParameterReferenceCollection>> ParameterEntries;
+	if (ToolkitType == SCRIPT)
+	{
+		CollectAllActionsForScriptToolkit(ParameterEntries);
+	}
+	else if (ToolkitType == SYSTEM)
+	{
+		CollectAllActionsForSystemToolkit(ParameterEntries);
+	}
+
+	ParameterEntries.KeySort([](const FNiagaraVariable& A, const FNiagaraVariable& B) { return A.GetName().LexicalLess(B.GetName()); });
+	const FText TooltipFormat = LOCTEXT("Parameters", "Name: {0} \nType: {1}");
+	for (const auto& ParameterEntry : ParameterEntries)
+	{
+		const FNiagaraVariable& Parameter = ParameterEntry.Key;
+		FNiagaraParameterHandle Handle;
+		const NiagaraParameterMapSectionID::Type Section = NiagaraParameterMapSectionID::OnGetSectionFromVariable(Parameter, IsStaticSwitchParameter(Parameter, Graphs), Handle);
+		if (Section == NiagaraParameterMapSectionID::NONE)
+		{
+			continue;
+		}
+		TArray<NiagaraParameterMapSectionID::Type> ExcludedSystemIds = {
+			NiagaraParameterMapSectionID::MODULE_INPUT,
+			NiagaraParameterMapSectionID::MODULE_LOCAL,
+			NiagaraParameterMapSectionID::DATA_INSTANCE };
+
+		if ((IsSystemToolkit() == false || ExcludedSystemIds.Contains(Section) == false) &&
+			(bDisplayAdvancedCategories || NiagaraParameterMapSectionID::GetSectionIsAdvanced(Section) == false))
+		{
+			const FText Name = FText::FromName(Parameter.GetName());
+			const FText Tooltip = FText::Format(TooltipFormat, FText::FromName(Parameter.GetName()), Parameter.GetType().GetNameText());
+			TSharedPtr<FNiagaraParameterAction> ParameterAction(new FNiagaraParameterAction(Parameter, ParameterEntry.Value, FText::GetEmpty(), Name, Tooltip, 0, FText(), Section));
+			OutAllActions.AddAction(ParameterAction);
+		}
+	}
+}
+
+void SNiagaraParameterMapView::CollectAllActionsForScriptToolkit(TMap<FNiagaraVariable, TArray<FNiagaraGraphParameterReferenceCollection>>& OutParameterEntries)
+{
+	// For scripts we use the reference maps cached in the graph to collect parameters.
 	for (auto& GraphWeakPtr : Graphs)
 	{
 		if (!GraphWeakPtr.IsValid())
@@ -500,9 +548,9 @@ void SNiagaraParameterMapView::CollectAllActions(FGraphActionListBuilderBase& Ou
 			continue;
 		}
 		UNiagaraGraph* Graph = GraphWeakPtr.Get();
-		for (const auto& ParameterElement : Graph->GetParameterReferenceMap())
+		for (const TPair<FNiagaraVariable, FNiagaraGraphParameterReferenceCollection>& ParameterElement : Graph->GetParameterReferenceMap())
 		{
-			TArray<FNiagaraGraphParameterReferenceCollection>* Found = ParameterEntries.Find(ParameterElement.Key);
+			TArray<FNiagaraGraphParameterReferenceCollection>* Found = OutParameterEntries.Find(ParameterElement.Key);
 			if (Found)
 			{
 				Found->Add(ParameterElement.Value);
@@ -511,72 +559,130 @@ void SNiagaraParameterMapView::CollectAllActions(FGraphActionListBuilderBase& Ou
 			{
 				TArray<FNiagaraGraphParameterReferenceCollection> Collection;
 				Collection.Add(ParameterElement.Value);
-				ParameterEntries.Add(ParameterElement.Key, Collection);
+				OutParameterEntries.Add(ParameterElement.Key, Collection);
 			}
 		}
 	}
+}
 
-	TSet<UObject*> Objects = SelectedScriptObjects->GetSelectedObjects();
-	for (UObject* Object : Objects)
-	{
-		if (UNiagaraSystem* System = Cast<UNiagaraSystem>(Object))
+void SNiagaraParameterMapView::CollectAllActionsForSystemToolkit(TMap<FNiagaraVariable, TArray<FNiagaraGraphParameterReferenceCollection>>& OutParameterEntries)
+{
+	// For systems we need to collect the user parameters if a system is selected, and then we use parameter map traversal
+	// to find the compile time parameters.
+	UNiagaraSystem* System = CachedSystem.Get();
+	if (System != nullptr)
+	{
+		// Collect user parameters.
+		TArray<FNiagaraVariable> ExposedVars;
+		System->GetExposedParameters().GetParameters(ExposedVars);
+		for (const FNiagaraVariable& ExposedVar : ExposedVars)
 		{
-			TArray<FNiagaraVariable> ExposedVars;
-			System->GetExposedParameters().GetParameters(ExposedVars);
-			for (const FNiagaraVariable& ExposedVar : ExposedVars)
+			TArray<FNiagaraGraphParameterReferenceCollection>* Found = OutParameterEntries.Find(ExposedVar);
+			if (!Found)
 			{
-				TArray<FNiagaraGraphParameterReferenceCollection>* Found = ParameterEntries.Find(ExposedVar);
-				if (!Found)
-				{
-					TArray<FNiagaraGraphParameterReferenceCollection> Collection;
-					ParameterEntries.Add(ExposedVar, Collection);
-				}
+				TArray<FNiagaraGraphParameterReferenceCollection> Collection;
+				OutParameterEntries.Add(ExposedVar, Collection);
 			}
 		}
 	}
 
-	ParameterEntries.KeySort([](const FNiagaraVariable& A, const FNiagaraVariable& B) { return A.GetName().LexicalLess(B.GetName()); });
-
-	const FText TooltipFormat = LOCTEXT("Parameters", "Name: {0} \nType: {1}");
-	for (const auto& ParameterEntry : ParameterEntries)
+	for (TWeakObjectPtr<UNiagaraGraph> GraphWeak : Graphs)
 	{
-		const FNiagaraVariable& Parameter = ParameterEntry.Key;
-		FNiagaraParameterHandle Handle;
-		const NiagaraParameterMapSectionID::Type Section = NiagaraParameterMapSectionID::OnGetSectionFromVariable(Parameter, IsStaticSwitchParameter(Parameter, Graphs), Handle);
-		if (Section == NiagaraParameterMapSectionID::NONE)
+		UNiagaraGraph* Graph = GraphWeak.Get();
+		if (Graph == nullptr)
 		{
 			continue;
 		}
-		if (!IsSystemToolkit() || (IsSystemToolkit() && Section != NiagaraParameterMapSectionID::MODULE_INPUT && Section != NiagaraParameterMapSectionID::MODULE_LOCAL))
+
+		TArray<UNiagaraNodeOutput*> OutputNodes;
+		Graph->GetNodesOfClass<UNiagaraNodeOutput>(OutputNodes);
+		for (UNiagaraNodeOutput* OutputNode : OutputNodes)
 		{
-			if (bDisplayAdvancedCategories == false && NiagaraParameterMapSectionID::GetSectionIsAdvanced(Section))
+			UNiagaraNode* NodeToTraverse = OutputNode;
+			if (OutputNode->GetUsage() == ENiagaraScriptUsage::SystemSpawnScript || OutputNode->GetUsage() == ENiagaraScriptUsage::SystemUpdateScript)
+			{
+				// Traverse past the emitter nodes, otherwise the system scripts will pick up all of the emitter and particle script parameters.
+				UEdGraphPin* InputPin = FNiagaraStackGraphUtilities::GetParameterMapInputPin(*NodeToTraverse);
+				while (NodeToTraverse != nullptr && InputPin != nullptr && InputPin->LinkedTo.Num() == 1 &&
+					(NodeToTraverse->IsA<UNiagaraNodeOutput>() || NodeToTraverse->IsA<UNiagaraNodeEmitter>()))
+				{
+					NodeToTraverse = Cast<UNiagaraNode>(InputPin->LinkedTo[0]->GetOwningNode());
+					InputPin = NodeToTraverse != nullptr ? FNiagaraStackGraphUtilities::GetParameterMapInputPin(*NodeToTraverse) : nullptr;
+				}
+			}
+
+			if (NodeToTraverse == nullptr)
 			{
 				continue;
 			}
-			const FText Name = FText::FromName(Parameter.GetName());
-			const FText Tooltip = FText::Format(TooltipFormat, FText::FromName(Parameter.GetName()), Parameter.GetType().GetNameText());
-			TSharedPtr<FNiagaraParameterAction> ParameterAction(new FNiagaraParameterAction(Parameter, ParameterEntry.Value, FText::GetEmpty(), Name, Tooltip, 0, FText(), Section));
-			OutAllActions.AddAction(ParameterAction);
-		}
-	}
 
+			bool bIgnoreDisabled = true;
+			FNiagaraParameterMapHistoryBuilder Builder;
+			UNiagaraEmitter* GraphOwningEmitter = Graph->GetTypedOuter<UNiagaraEmitter>();
+			FCompileConstantResolver ConstantResolver = GraphOwningEmitter != nullptr
+				? FCompileConstantResolver(GraphOwningEmitter)
+				: FCompileConstantResolver();
+
+			Builder.SetIgnoreDisabled(bIgnoreDisabled);
+			Builder.ConstantResolver = ConstantResolver;
+			NodeToTraverse->BuildParameterMapHistory(Builder, true, false);
+			
+			TMap<FNiagaraVariable, FNiagaraGraphParameterReferenceCollection> ReferenceCollectionsForTraversedNode;
+			if (Builder.Histories.Num() == 1)
+			{
+				for (int32 VariableIndex = 0; VariableIndex < Builder.Histories[0].Variables.Num(); VariableIndex++)
+				{
+					FNiagaraVariable& HistoryVariable = Builder.Histories[0].Variables[VariableIndex];
+					FNiagaraGraphParameterReferenceCollection* ReferenceCollection = ReferenceCollectionsForTraversedNode.Find(HistoryVariable);
+					if (ReferenceCollection == nullptr)
+					{
+						FNiagaraGraphParameterReferenceCollection NewReferenceCollection(false);
+						NewReferenceCollection.Graph = Graph;
+						ReferenceCollection = &ReferenceCollectionsForTraversedNode.Add(HistoryVariable, NewReferenceCollection);
+					}
+
+					TArray<TTuple<const UEdGraphPin*, const UEdGraphPin*>>& ReadHistory = Builder.Histories[0].PerVariableReadHistory[VariableIndex];
+					for (const TTuple<const UEdGraphPin*, const UEdGraphPin*>& Read : ReadHistory)
+					{
+						if (Read.Key->GetOwningNode() != nullptr)
+						{
+							ReferenceCollection->ParameterReferences.Add(FNiagaraGraphParameterReference(Read.Key->PersistentGuid, Read.Key->GetOwningNode()));
+						}
+					}
+
+					TArray<const UEdGraphPin*>& WriteHistory = Builder.Histories[0].PerVariableWriteHistory[VariableIndex];
+					for (const UEdGraphPin* Write : WriteHistory)
+					{
+						if (Write->GetOwningNode() != nullptr)
+						{
+							ReferenceCollection->ParameterReferences.Add(FNiagaraGraphParameterReference(Write->PersistentGuid, Write->GetOwningNode()));
+						}
+					}
+				}
+			}
 
+			for (const TPair<FNiagaraVariable, FNiagaraGraphParameterReferenceCollection>& ReferenceCollectionForOutputNode : ReferenceCollectionsForTraversedNode)
+			{
+				OutParameterEntries.FindOrAdd(ReferenceCollectionForOutputNode.Key).Add(ReferenceCollectionForOutputNode.Value);
+			}
+		}
+	}
 }
 
 void SNiagaraParameterMapView::CollectStaticSections(TArray<int32>& StaticSectionIDs)
 {
-	if (!IsSystemToolkit())
+	if (IsSystemToolkit())
 	{
-		StaticSectionIDs.Add(NiagaraParameterMapSectionID::MODULE_INPUT);
-		StaticSectionIDs.Add(NiagaraParameterMapSectionID::MODULE_OUTPUT);
-		StaticSectionIDs.Add(NiagaraParameterMapSectionID::MODULE_LOCAL);
-		StaticSectionIDs.Add(NiagaraParameterMapSectionID::DATA_INSTANCE);
+		StaticSectionIDs.Add(NiagaraParameterMapSectionID::USER);
 	}
 	else
 	{
-		StaticSectionIDs.Add(NiagaraParameterMapSectionID::USER);
+		StaticSectionIDs.Add(NiagaraParameterMapSectionID::MODULE_INPUT);
+		StaticSectionIDs.Add(NiagaraParameterMapSectionID::MODULE_LOCAL);
+		StaticSectionIDs.Add(NiagaraParameterMapSectionID::STATIC_SWITCH);
+		StaticSectionIDs.Add(NiagaraParameterMapSectionID::DATA_INSTANCE);
 	}
-	StaticSectionIDs.Add(NiagaraParameterMapSectionID::STATIC_SWITCH);
+	StaticSectionIDs.Add(NiagaraParameterMapSectionID::MODULE_OUTPUT);
 	StaticSectionIDs.Add(NiagaraParameterMapSectionID::ENGINE);
 	StaticSectionIDs.Add(NiagaraParameterMapSectionID::PARAMETERCOLLECTION);
 	StaticSectionIDs.Add(NiagaraParameterMapSectionID::SYSTEM);
@@ -682,28 +788,28 @@ TSharedPtr<SWidget> SNiagaraParameterMapView::OnContextMenuOpening()
 			TSharedPtr<FNiagaraParameterAction> SelectedNiagaraAction = StaticCastSharedPtr<FNiagaraParameterAction>(SelectedActions[0]);
 			if (SelectedNiagaraAction.IsValid())
 			{
-				MenuBuilder.BeginSection("Subnamespace", LOCTEXT("SubnamespaceMenuSection", "Parameter Modifiers"));
+				MenuBuilder.BeginSection("NamespaceModifier", LOCTEXT("NamespaceModifierMenuSection", "Parameter Modifiers"));
 				{
 					MenuBuilder.AddMenuEntry(
-						LOCTEXT("AddSubnamespace", "Add a modifier to this parameter."),
+						LOCTEXT("AddNamespaceModifier", "Add a modifier to this parameter."),
 						FText(),
 						FSlateIcon(),
 						FUIAction(
-							FExecuteAction::CreateSP(this, &SNiagaraParameterMapView::OnAddSubnamespace),
+							FExecuteAction::CreateSP(this, &SNiagaraParameterMapView::OnAddNamespaceModifier),
 							FCanExecuteAction()));
 					MenuBuilder.AddMenuEntry(
-						LOCTEXT("RemoveSubnamespace", "Remove the modifier from this parameter."),
+						LOCTEXT("RemoveNamespaceModifier", "Remove the modifier from this parameter."),
 						FText(),
 						FSlateIcon(),
 						FUIAction(
-							FExecuteAction::CreateSP(this, &SNiagaraParameterMapView::OnRemoveSubnamespace),
+							FExecuteAction::CreateSP(this, &SNiagaraParameterMapView::OnRemoveNamespaceModifier),
 							FCanExecuteAction()));
 					MenuBuilder.AddMenuEntry(
-						LOCTEXT("EditSubnamespace", "Edit the modifier for this parameter."),
+						LOCTEXT("EditNamespaceModifier", "Edit the modifier for this parameter."),
 						FText(),
 						FSlateIcon(),
 						FUIAction(
-							FExecuteAction::CreateSP(this, &SNiagaraParameterMapView::OnEditSubnamespace),
+							FExecuteAction::CreateSP(this, &SNiagaraParameterMapView::OnEditNamespaceModifier),
 							FCanExecuteAction()));
 				}
 				MenuBuilder.EndSection();
@@ -816,6 +922,11 @@ EVisibility SNiagaraParameterMapView::OnAddButtonTextVisibility(TWeakPtr<SWidget
 void SNiagaraParameterMapView::Refresh(bool bRefreshMenu/* = true*/)
 {
 	EmptyGraphs();
+	if (CachedSystem.IsValid())
+	{
+		CachedSystem->GetExposedParameters().RemoveOnChangedHandler(UserParameterStoreChangedHandle);
+		CachedSystem.Reset();
+	}
 
 	TSet<UObject*> Objects = SelectedScriptObjects->GetSelectedObjects();
 	for (UObject* Object : Objects)
@@ -832,6 +943,8 @@ void SNiagaraParameterMapView::Refresh(bool bRefreshMenu/* = true*/)
 		else if (UNiagaraSystem* System = Cast<UNiagaraSystem>(Object))
 		{
 			AddGraph(System->GetSystemSpawnScript()->GetSource());
+			UserParameterStoreChangedHandle = System->GetExposedParameters().AddOnChangedHandler(
+				FNiagaraParameterStore::FOnChanged::FDelegate::CreateSP(this, &SNiagaraParameterMapView::OnUserParameterStoreChanged));
 		}
 	}
 
@@ -891,6 +1004,14 @@ void SNiagaraParameterMapView::OnGraphChanged(const FEdGraphEditAction& InAction
 	RefreshActions();
 }
 
+void SNiagaraParameterMapView::OnUserParameterStoreChanged()
+{
+	if (CachedSystem.IsValid())
+	{
+		RefreshActions();
+	}
+}
+
 void SNiagaraParameterMapView::OnDeleteEntry()
 {
 	TArray<TSharedPtr<FEdGraphSchemaAction>> SelectedActions;
@@ -992,7 +1113,7 @@ void SNiagaraParameterMapView::OnPostRenameActionNode(const FText& InText, FNiag
 	}
 }
 
-void SNiagaraParameterMapView::OnAddSubnamespace()
+void SNiagaraParameterMapView::OnAddNamespaceModifier()
 {
 	TArray<TSharedPtr<FEdGraphSchemaAction>> SelectedActions;
 	GraphActionMenu->GetSelectedActions(SelectedActions);
@@ -1023,7 +1144,7 @@ void SNiagaraParameterMapView::OnAddSubnamespace()
 	}
 }
 
-void SNiagaraParameterMapView::OnRemoveSubnamespace()
+void SNiagaraParameterMapView::OnRemoveNamespaceModifier()
 {
 	TArray<TSharedPtr<FEdGraphSchemaAction>> SelectedActions;
 	GraphActionMenu->GetSelectedActions(SelectedActions);
@@ -1054,7 +1175,7 @@ void SNiagaraParameterMapView::OnRemoveSubnamespace()
 	}
 }
 
-void SNiagaraParameterMapView::OnEditSubnamespace()
+void SNiagaraParameterMapView::OnEditNamespaceModifier()
 {
 	TArray<TSharedPtr<FEdGraphSchemaAction>> SelectedActions;
 	GraphActionMenu->GetSelectedActions(SelectedActions);
@@ -1063,7 +1184,7 @@ void SNiagaraParameterMapView::OnEditSubnamespace()
 		TSharedPtr<FNiagaraParameterAction> ParameterAction = StaticCastSharedPtr<FNiagaraParameterAction>(SelectedActions[0]);
 		if (ParameterAction.IsValid())
 		{
-			ParameterAction->bSubnamespaceRenamePending = true;
+			ParameterAction->bNamespaceModifierRenamePending = true;
 		}
 	}
 }
diff --git a/Engine/Plugins/FX/Niagara/Source/NiagaraEditor/Public/NiagaraActions.h b/Engine/Plugins/FX/Niagara/Source/NiagaraEditor/Public/NiagaraActions.h
--- a/Engine/Plugins/FX/Niagara/Source/NiagaraEditor/Public/NiagaraActions.h
+++ b/Engine/Plugins/FX/Niagara/Source/NiagaraEditor/Public/NiagaraActions.h
@@ -56,7 +56,7 @@ struct NIAGARAEDITOR_API FNiagaraScriptVarAndViewInfoAction : public FEdGraphSch
 struct NIAGARAEDITOR_API FNiagaraParameterAction : public FEdGraphSchemaAction
 {
 	FNiagaraParameterAction()
-		: bSubnamespaceRenamePending(false)
+		: bNamespaceModifierRenamePending(false)
 	{
 	}
 
@@ -72,7 +72,7 @@ struct NIAGARAEDITOR_API FNiagaraParameterAction : public FEdGraphSchemaAction
 
 	TArray<FNiagaraGraphParameterReferenceCollection> ReferenceCollection;
 
-	bool bSubnamespaceRenamePending;
+	bool bNamespaceModifierRenamePending;
 };
 
 struct NIAGARAEDITOR_API FNiagaraScriptParameterAction : public FEdGraphSchemaAction
diff --git a/Engine/Plugins/FX/Niagara/Source/NiagaraEditor/Public/NiagaraEditorSettings.h b/Engine/Plugins/FX/Niagara/Source/NiagaraEditor/Public/NiagaraEditorSettings.h
--- a/Engine/Plugins/FX/Niagara/Source/NiagaraEditor/Public/NiagaraEditorSettings.h
+++ b/Engine/Plugins/FX/Niagara/Source/NiagaraEditor/Public/NiagaraEditorSettings.h
@@ -32,7 +32,7 @@ enum class ENiagaraNamespaceMetadataOptions
 {
 	Advanced,
 	PreventRenaming,
-	SupportsSubnamespace
+	SupportsNamespaceModifier,
 };
 
 USTRUCT()
@@ -133,6 +133,7 @@ public:
 	void SetNewAssetDialogConfig(FName InDialogConfigKey, const FNiagaraNewAssetDialogConfig& InNewAssetDialogConfig);
 
 	FNiagaraNamespaceMetadata GetMetaDataForNamespaces(TArray<FName> Namespaces) const;
+	FNiagaraNamespaceMetadata GetMetaDataForNamespaceModifier(FName NamespaceModifier) const;
 	
 	// Begin UDeveloperSettings Interface
 	virtual FName GetCategoryName() const override;
@@ -148,6 +149,9 @@ public:
 
 	const TMap<FString, FString>& GetHLSLKeywordReplacementsMap()const { return HLSLKeywordReplacements; }
 
+private:
+	void SetupNamespaceMetadata();
+
 protected:
 	static FOnNiagaraEditorSettingsChanged SettingsChangedDelegate;
 
@@ -184,4 +188,7 @@ private:
 
 	UPROPERTY()
 	TArray<FNiagaraNamespaceMetadata> NamespaceMetadata;
+
+	UPROPERTY()
+	TArray<FNiagaraNamespaceMetadata> NamespaceModifierMetadata;
 };
diff --git a/Engine/Plugins/FX/Niagara/Source/NiagaraEditor/Public/Widgets/SNiagaraParameterMapView.h b/Engine/Plugins/FX/Niagara/Source/NiagaraEditor/Public/Widgets/SNiagaraParameterMapView.h
--- a/Engine/Plugins/FX/Niagara/Source/NiagaraEditor/Public/Widgets/SNiagaraParameterMapView.h
+++ b/Engine/Plugins/FX/Niagara/Source/NiagaraEditor/Public/Widgets/SNiagaraParameterMapView.h
@@ -11,6 +11,7 @@
 #include "ViewModels/Stack/NiagaraParameterHandle.h"
 #include "Framework/Commands/Commands.h"
 #include "Framework/Commands/UICommandList.h"
+#include "NiagaraGraph.h"
 #include "NiagaraActions.h"
 #include "EditorStyleSet.h"
 
@@ -23,6 +24,7 @@ class FNiagaraObjectSelection;
 class UNiagaraGraph;
 struct FEdGraphSchemaAction;
 struct FSlateBrush;
+class UNiagaraSystem;
 
 /* Enums to use when grouping the blueprint members in the list panel. The order here will determine the order in the list */
 namespace NiagaraParameterMapSectionID
@@ -30,21 +32,22 @@ namespace NiagaraParameterMapSectionID
 	enum Type
 	{
 		NONE,
-		MODULE_INPUT,
-		MODULE_OUTPUT,
-		MODULE_LOCAL,
-		TRANSIENT,
-		DATA_INSTANCE,
-		STATIC_SWITCH,
-		ENGINE,
-		PARAMETERCOLLECTION,
+
 		USER,
+
 		SYSTEM,
-		SYSTEM_MODULE,
 		EMITTER,
-		EMITTER_MODULE,
 		PARTICLE,
-		PARTICLE_MODULE,
+
+		MODULE_INPUT,
+		STATIC_SWITCH,
+		MODULE_LOCAL,
+		MODULE_OUTPUT,
+
+		ENGINE,
+		TRANSIENT,
+		PARAMETERCOLLECTION,
+		DATA_INSTANCE,
 
 		Num
 	};
@@ -131,6 +134,9 @@ private:
 	TSharedRef<SWidget> OnGetSectionWidget(TSharedRef<SWidget> RowWidget, int32 InSectionID);
 	TSharedRef<SWidget> CreateAddToSectionButton(const NiagaraParameterMapSectionID::Type InSection, TWeakPtr<SWidget> WeakRowWidget, FText AddNewText, FName MetaDataTag);
 	
+	void CollectAllActionsForScriptToolkit(TMap<FNiagaraVariable, TArray<FNiagaraGraphParameterReferenceCollection>>& ParameterEntries);
+	void CollectAllActionsForSystemToolkit(TMap<FNiagaraVariable, TArray<FNiagaraGraphParameterReferenceCollection>>& ParameterEntries);
+
 	/** Checks if the selected action has context menu */
 	bool SelectionHasContextMenu() const;
 	
@@ -142,6 +148,7 @@ private:
 	void AddGraph(UNiagaraGraph* Graph);
 	void AddGraph(class UNiagaraScriptSourceBase* SourceBase);
 	void OnGraphChanged(const struct FEdGraphEditAction& InAction);
+	void OnUserParameterStoreChanged();
 
 	//Callbacks
 	void OnDeleteEntry();
@@ -150,9 +157,9 @@ private:
 	bool CanRequestRenameOnActionNode(TWeakPtr<struct FGraphActionNode> InSelectedNode) const;
 	bool CanRequestRenameOnActionNode() const;
 	void OnPostRenameActionNode(const FText& InText, FNiagaraParameterAction& InAction);
-	void OnAddSubnamespace();
-	void OnRemoveSubnamespace();
-	void OnEditSubnamespace();
+	void OnAddNamespaceModifier();
+	void OnRemoveNamespaceModifier();
+	void OnEditNamespaceModifier();
 
 	bool IsSystemToolkit();
 	bool IsScriptToolkit();
@@ -179,6 +186,8 @@ private:
 	TSharedPtr<FNiagaraObjectSelection> SelectedVariableObjects;
 
 	TArray<TWeakObjectPtr<UNiagaraGraph>> Graphs;
+	TWeakObjectPtr<UNiagaraSystem> CachedSystem;
+	FDelegateHandle UserParameterStoreChangedHandle;
 
 	/** The handle to the graph changed delegate. */
 	TArray<FDelegateHandle> OnGraphChangedHandles;
diff --git a/Engine/Plugins/FX/Niagara/Source/NiagaraEditor/Public/Widgets/SNiagaraParameterName.h b/Engine/Plugins/FX/Niagara/Source/NiagaraEditor/Public/Widgets/SNiagaraParameterName.h
--- a/Engine/Plugins/FX/Niagara/Source/NiagaraEditor/Public/Widgets/SNiagaraParameterName.h
+++ b/Engine/Plugins/FX/Niagara/Source/NiagaraEditor/Public/Widgets/SNiagaraParameterName.h
@@ -42,7 +42,7 @@ public:
 
 	void EnterEditingMode();
 
-	void EnterSubnamespaceEditingMode();
+	void EnterNamespaceModifierEditingMode();
 
 private:
 	TSharedRef<SBorder> CreateNamespaceWidget(FText NamespaceDisplayName, FText NamespaceDescription, FLinearColor NamespaceBorderColor);
@@ -51,17 +51,17 @@ private:
 
 	FName ReconstructNameFromEditText(const FText& InEditText);
 
-	FName ReconstructNameFromSubnamespaceEditText(const FText& InEditText);
+	FName ReconstructNameFromNamespaceModifierEditText(const FText& InEditText);
 
-	FText GetSubnamespaceText();
+	FText GetNamespaceModifierText();
 
 	bool VerifyNameTextChange(const FText& InNewNameText, FText& OutErrorMessage);
 
 	void NameTextCommitted(const FText& InNewNameText, ETextCommit::Type CommitType);
 
-	bool VerifySubnamespaceTextChange(const FText& InNewNameText, FText& OutErrorMessage);
+	bool VerifyNamespaceModifierTextChange(const FText& InNewNameText, FText& OutErrorMessage);
 
-	void SubnamespaceTextCommitted(const FText& InNewNameText, ETextCommit::Type CommitType);
+	void NamespaceModifierTextCommitted(const FText& InNewNameText, ETextCommit::Type CommitType);
 
 private:
 	const FInlineEditableTextBlockStyle* EditableTextStyle;
@@ -73,9 +73,9 @@ private:
 	FPointerEventHandler OnDoubleClickedDelegate;
 	FIsSelected IsSelected;
 	FName DisplayedParameterName;
-	FName DisplayedSubnamespace;
+	FName DisplayedNamespaceModifier;
 	TSharedPtr<SInlineEditableTextBlock> EditableTextBlock;
-	TSharedPtr<SBorder> SubnamespaceBorder;
+	TSharedPtr<SBorder> NamespaceModifierBorder;
 };
 
 class NIAGARAEDITOR_API SNiagaraParameterNameTextBlock : public SCompoundWidget
@@ -83,7 +83,9 @@ class NIAGARAEDITOR_API SNiagaraParameterNameTextBlock : public SCompoundWidget
 public:
 	SLATE_BEGIN_ARGS(SNiagaraParameterNameTextBlock)
 		: _EditableTextStyle(&FEditorStyle::Get().GetWidgetStyle<FInlineEditableTextBlockStyle>("InlineEditableTextBlockStyle"))
-	{}
+	{
+		_Clipping = EWidgetClipping::OnDemand;
+	}
 		SLATE_ATTRIBUTE(FText, ParameterText)
 		SLATE_STYLE_ARGUMENT(FInlineEditableTextBlockStyle, EditableTextStyle)
 		SLATE_ARGUMENT(bool, IsReadOnly)
@@ -96,7 +98,7 @@ public:
 
 	void EnterEditingMode();
 
-	void EnterSubnamespaceEditingMode();
+	void EnterNamespaceModifierEditingMode();
 
 private:
 	FName GetParameterName() const;
