Niagara - Update the stack input dropdown menu to use the new categories and styling for parameters.  Also fix text highlighting on the paramter name widget.

diff --git a/Engine/Plugins/FX/Niagara/Source/NiagaraEditor/Private/NiagaraActions.cpp b/Engine/Plugins/FX/Niagara/Source/NiagaraEditor/Private/NiagaraActions.cpp
--- a/Engine/Plugins/FX/Niagara/Source/NiagaraEditor/Private/NiagaraActions.cpp
+++ b/Engine/Plugins/FX/Niagara/Source/NiagaraEditor/Private/NiagaraActions.cpp
@@ -28,6 +28,16 @@ FNiagaraMenuAction::FNiagaraMenuAction(FText InNodeCategory, FText InMenuDesc, F
 	, CanPerformAction(InCanPerformAction)
 {}
 
+TOptional<FNiagaraParameterHandle> FNiagaraMenuAction::GetParameterHandle() const
+{
+	return ParameterHandle;
+}
+
+void FNiagaraMenuAction::SetParamterHandle(const FNiagaraParameterHandle& InParameterHandle)
+{
+	ParameterHandle = InParameterHandle;
+}
+
 /************************************************************************/
 /* FNiagaraParameterAction												*/
 /************************************************************************/
diff --git a/Engine/Plugins/FX/Niagara/Source/NiagaraEditor/Private/SNiagaraParameterName.cpp b/Engine/Plugins/FX/Niagara/Source/NiagaraEditor/Private/SNiagaraParameterName.cpp
--- a/Engine/Plugins/FX/Niagara/Source/NiagaraEditor/Private/SNiagaraParameterName.cpp
+++ b/Engine/Plugins/FX/Niagara/Source/NiagaraEditor/Private/SNiagaraParameterName.cpp
@@ -15,6 +15,7 @@ void SNiagaraParameterName::Construct(const FArguments& InArgs)
 	ReadOnlyTextStyle = InArgs._ReadOnlyTextStyle;
 	ParameterName = InArgs._ParameterName;
 	bIsReadOnly = InArgs._IsReadOnly;
+	HighlightText = InArgs._HighlightText;
 	OnVerifyNameChangeDelegate = InArgs._OnVerifyNameChange;
 	OnNameChangedDelegate = InArgs._OnNameChanged;
 	OnDoubleClickedDelegate = InArgs._OnDoubleClicked;
@@ -51,11 +52,12 @@ TSharedRef<SBorder> SNiagaraParameterName::CreateNamespaceWidget(FText Namespace
 	.BorderBackgroundColor(NamespaceBorderColor)
 	.ToolTipText(NamespaceDescription)
 	.VAlign(VAlign_Center)
-	.Padding(FMargin(4.0f, 1.0f, 4.0f, 1.0f))
+	.Padding(FMargin(5.0f, 1.0f, 5.0f, 1.0f))
 	[
 		SNew(STextBlock)
 		.TextStyle(FNiagaraEditorStyle::Get(), "NiagaraEditor.ParameterName.NamespaceText")
 		.Text(NamespaceDisplayName)
+		.HighlightText(HighlightText)
 	];
 }
 
@@ -153,7 +155,8 @@ void SNiagaraParameterName::UpdateContent(FName InDisplayedParameterName)
 	{
 		NameWidget = SNew(STextBlock)
 			.TextStyle(ReadOnlyTextStyle)
-			.Text(FText::FromString(NameParts.Last()));
+			.Text(FText::FromString(NameParts.Last()))
+			.HighlightText(HighlightText);
 	}
 	else
 	{
@@ -162,7 +165,8 @@ void SNiagaraParameterName::UpdateContent(FName InDisplayedParameterName)
 		.Text(FText::FromString(NameParts.Last()))
 		.IsSelected(IsSelected)
 		.OnVerifyTextChanged(this, &SNiagaraParameterName::VerifyNameTextChange)
-		.OnTextCommitted(this, &SNiagaraParameterName::NameTextCommitted);
+		.OnTextCommitted(this, &SNiagaraParameterName::NameTextCommitted)
+		.HighlightText(HighlightText);
 	}
 
 	ContentBox->AddSlot()
@@ -293,6 +297,7 @@ void SNiagaraParameterNameTextBlock::Construct(const FArguments& InArgs)
 		.EditableTextStyle(InArgs._EditableTextStyle)
 		.ParameterName(this, &SNiagaraParameterNameTextBlock::GetParameterName)
 		.IsReadOnly(InArgs._IsReadOnly)
+		.HighlightText(InArgs._HighlightText)
 		.IsSelected(InArgs._IsSelected)
 		.OnVerifyNameChange(this, &SNiagaraParameterNameTextBlock::VerifyNameChange)
 		.OnNameChanged(this, &SNiagaraParameterNameTextBlock::NameChanged)
diff --git a/Engine/Plugins/FX/Niagara/Source/NiagaraEditor/Private/ViewModels/Stack/NiagaraParameterHandle.cpp b/Engine/Plugins/FX/Niagara/Source/NiagaraEditor/Private/ViewModels/Stack/NiagaraParameterHandle.cpp
--- a/Engine/Plugins/FX/Niagara/Source/NiagaraEditor/Private/ViewModels/Stack/NiagaraParameterHandle.cpp
+++ b/Engine/Plugins/FX/Niagara/Source/NiagaraEditor/Private/ViewModels/Stack/NiagaraParameterHandle.cpp
@@ -92,6 +92,21 @@ const FName FNiagaraParameterHandle::GetNamespace() const
 	return Namespace;
 }
 
+const TArray<FName> FNiagaraParameterHandle::GetHandleParts() const
+{
+	if (HandlePartsCache.Num() == 0)
+	{
+		FString HandleString = ParameterHandleName.ToString();
+		TArray<FString> HandlePartStrings;
+		HandleString.ParseIntoArray(HandlePartStrings, TEXT("."));
+		for (const FString& HandlePartString : HandlePartStrings)
+		{
+			HandlePartsCache.Add(*HandlePartString);
+		}
+	}
+	return HandlePartsCache;
+}
+
 bool FNiagaraParameterHandle::IsUserHandle() const
 {
 	return Namespace == FNiagaraConstants::UserNamespace;
diff --git a/Engine/Plugins/FX/Niagara/Source/NiagaraEditor/Private/Widgets/SNiagaraGraphActionWidget.cpp b/Engine/Plugins/FX/Niagara/Source/NiagaraEditor/Private/Widgets/SNiagaraGraphActionWidget.cpp
--- a/Engine/Plugins/FX/Niagara/Source/NiagaraEditor/Private/Widgets/SNiagaraGraphActionWidget.cpp
+++ b/Engine/Plugins/FX/Niagara/Source/NiagaraEditor/Private/Widgets/SNiagaraGraphActionWidget.cpp
@@ -7,6 +7,7 @@
 #include "Widgets/SBoxPanel.h"
 #include "Widgets/Images/SImage.h"
 #include "Widgets/Text/STextBlock.h"
+#include "Widgets/SNiagaraParameterName.h"
 
 #define LOCTEXT_NAMESPACE "NiagaraGraphActionWidget"
 
@@ -17,6 +18,22 @@ void SNiagaraGraphActionWidget::Construct(const FArguments& InArgs, const FCreat
 
 	TSharedPtr<FNiagaraMenuAction> NiagaraAction = StaticCastSharedPtr<FNiagaraMenuAction>(InCreateData->Action);
 
+	TSharedPtr<SWidget> NameWidget;
+	if (NiagaraAction->GetParameterHandle().IsSet())
+	{
+		NameWidget = SNew(SNiagaraParameterName)
+			.ParameterName(NiagaraAction->GetParameterHandle()->GetParameterHandleString())
+			.IsReadOnly(true)
+			.HighlightText(InArgs._HighlightText);
+	}
+	else
+	{
+		NameWidget = SNew(STextBlock)
+			.Font(FCoreStyle::GetDefaultFontStyle("Regular", 9))
+			.Text(InCreateData->Action->GetMenuDescription())
+			.HighlightText(InArgs._HighlightText);
+	}
+
 	this->ChildSlot
 	[
 		SNew(SHorizontalBox)
@@ -25,10 +42,7 @@ void SNiagaraGraphActionWidget::Construct(const FArguments& InArgs, const FCreat
 		.FillWidth(1)
 		.VAlign(VAlign_Center)
 		[
-			SNew(STextBlock)
-			.Font(FCoreStyle::GetDefaultFontStyle("Regular", 9))
-			.Text(InCreateData->Action->GetMenuDescription())
-			.HighlightText(InArgs._HighlightText)
+			NameWidget.ToSharedRef()
 		]
 		+ SHorizontalBox::Slot()
 		.AutoWidth()
diff --git a/Engine/Plugins/FX/Niagara/Source/NiagaraEditor/Private/Widgets/SNiagaraParameterMapPaletteItem.cpp b/Engine/Plugins/FX/Niagara/Source/NiagaraEditor/Private/Widgets/SNiagaraParameterMapPaletteItem.cpp
--- a/Engine/Plugins/FX/Niagara/Source/NiagaraEditor/Private/Widgets/SNiagaraParameterMapPaletteItem.cpp
+++ b/Engine/Plugins/FX/Niagara/Source/NiagaraEditor/Private/Widgets/SNiagaraParameterMapPaletteItem.cpp
@@ -48,7 +48,7 @@ void SNiagaraParameterMapPalleteItem::Construct(const FArguments& InArgs, FCreat
 
 	ParameterNameTextBlock = SNew(SNiagaraParameterNameTextBlock)
 		.ParameterText(this, &SNiagaraParameterMapPalleteItem::GetDisplayText)
-		//.HighlightText(InCreateData->HighlightText)
+		.HighlightText(InCreateData->HighlightText)
 		.ToolTipText(this, &SNiagaraParameterMapPalleteItem::GetItemTooltip)
 		.OnTextCommitted(this, &SNiagaraParameterMapPalleteItem::OnNameTextCommitted)
 		.OnVerifyTextChanged(this, &SNiagaraParameterMapPalleteItem::OnNameTextVerifyChanged)
diff --git a/Engine/Plugins/FX/Niagara/Source/NiagaraEditor/Public/NiagaraActions.h b/Engine/Plugins/FX/Niagara/Source/NiagaraEditor/Public/NiagaraActions.h
--- a/Engine/Plugins/FX/Niagara/Source/NiagaraEditor/Public/NiagaraActions.h
+++ b/Engine/Plugins/FX/Niagara/Source/NiagaraEditor/Public/NiagaraActions.h
@@ -8,6 +8,7 @@
 #include "GraphEditorDragDropAction.h"
 #include "DragAndDrop/DecoratedDragDropOp.h"
 #include "NiagaraEditorCommon.h"
+#include "ViewModels/Stack/NiagaraParameterHandle.h"
 #include "NiagaraActions.generated.h"
 
 USTRUCT()
@@ -38,7 +39,11 @@ struct NIAGARAEDITOR_API FNiagaraMenuAction : public FEdGraphSchemaAction
 
 	bool IsExperimental = false;
 
+	TOptional<FNiagaraParameterHandle> GetParameterHandle() const;
+	void SetParamterHandle(const FNiagaraParameterHandle& InParameterHandle);
+
 private:
+	TOptional<FNiagaraParameterHandle> ParameterHandle;
 	FOnExecuteStackAction Action;
 	FCanExecuteStackAction CanPerformAction;
 };
diff --git a/Engine/Plugins/FX/Niagara/Source/NiagaraEditor/Public/NiagaraEditorSettings.h b/Engine/Plugins/FX/Niagara/Source/NiagaraEditor/Public/NiagaraEditorSettings.h
--- a/Engine/Plugins/FX/Niagara/Source/NiagaraEditor/Public/NiagaraEditorSettings.h
+++ b/Engine/Plugins/FX/Niagara/Source/NiagaraEditor/Public/NiagaraEditorSettings.h
@@ -63,7 +63,7 @@ struct FNiagaraNamespaceMetadata
 };
 
 UCLASS(config = Niagara, defaultconfig, meta=(DisplayName="Niagara"))
-class UNiagaraEditorSettings : public UDeveloperSettings
+class NIAGARAEDITOR_API UNiagaraEditorSettings : public UDeveloperSettings
 {
 public:
 	GENERATED_UCLASS_BODY()
diff --git a/Engine/Plugins/FX/Niagara/Source/NiagaraEditor/Public/ViewModels/Stack/NiagaraParameterHandle.h b/Engine/Plugins/FX/Niagara/Source/NiagaraEditor/Public/ViewModels/Stack/NiagaraParameterHandle.h
--- a/Engine/Plugins/FX/Niagara/Source/NiagaraEditor/Public/ViewModels/Stack/NiagaraParameterHandle.h
+++ b/Engine/Plugins/FX/Niagara/Source/NiagaraEditor/Public/ViewModels/Stack/NiagaraParameterHandle.h
@@ -41,6 +41,8 @@ public:
 
 	const FName GetNamespace() const;
 
+	const TArray<FName> GetHandleParts() const;
+
 	bool IsUserHandle() const;
 
 	bool IsEngineHandle() const;
@@ -69,4 +71,5 @@ private:
 	FName ParameterHandleName;
 	FName Name;
 	FName Namespace;
+	mutable TArray<FName> HandlePartsCache;
 };
\ No newline at end of file
diff --git a/Engine/Plugins/FX/Niagara/Source/NiagaraEditor/Public/Widgets/SNiagaraParameterName.h b/Engine/Plugins/FX/Niagara/Source/NiagaraEditor/Public/Widgets/SNiagaraParameterName.h
--- a/Engine/Plugins/FX/Niagara/Source/NiagaraEditor/Public/Widgets/SNiagaraParameterName.h
+++ b/Engine/Plugins/FX/Niagara/Source/NiagaraEditor/Public/Widgets/SNiagaraParameterName.h
@@ -28,6 +28,7 @@ public:
 		SLATE_STYLE_ARGUMENT(FTextBlockStyle, ReadOnlyTextStyle)
 		SLATE_ATTRIBUTE(FName, ParameterName)
 		SLATE_ARGUMENT(bool, IsReadOnly)
+		SLATE_ATTRIBUTE(FText, HighlightText)
 		SLATE_EVENT(FOnVerifyNameChange, OnVerifyNameChange)
 		SLATE_EVENT(FOnNameChanged, OnNameChanged)
 		SLATE_EVENT(FIsSelected, IsSelected)
@@ -68,6 +69,7 @@ private:
 	const FTextBlockStyle* ReadOnlyTextStyle;
 	TAttribute<FName> ParameterName;
 	bool bIsReadOnly;
+	TAttribute<FText> HighlightText;
 	FOnVerifyNameChange OnVerifyNameChangeDelegate;
 	FOnNameChanged OnNameChangedDelegate;
 	FPointerEventHandler OnDoubleClickedDelegate;
@@ -89,6 +91,7 @@ public:
 		SLATE_ATTRIBUTE(FText, ParameterText)
 		SLATE_STYLE_ARGUMENT(FInlineEditableTextBlockStyle, EditableTextStyle)
 		SLATE_ARGUMENT(bool, IsReadOnly)
+		SLATE_ATTRIBUTE(FText, HighlightText)
 		SLATE_EVENT(FOnVerifyTextChanged, OnVerifyTextChanged)
 		SLATE_EVENT(FOnTextCommitted, OnTextCommitted)
 		SLATE_EVENT(FIsSelected, IsSelected)
diff --git a/Engine/Plugins/FX/Niagara/Source/NiagaraEditorWidgets/Private/Stack/SNiagaraStackFunctionInputValue.cpp b/Engine/Plugins/FX/Niagara/Source/NiagaraEditorWidgets/Private/Stack/SNiagaraStackFunctionInputValue.cpp
--- a/Engine/Plugins/FX/Niagara/Source/NiagaraEditorWidgets/Private/Stack/SNiagaraStackFunctionInputValue.cpp
+++ b/Engine/Plugins/FX/Niagara/Source/NiagaraEditorWidgets/Private/Stack/SNiagaraStackFunctionInputValue.cpp
@@ -35,6 +35,7 @@
 #include "EditorFontGlyphs.h"
 #include "Widgets/SNiagaraLibraryOnlyToggleHeader.h"
 #include "Widgets/SNiagaraParameterName.h"
+#include "NiagaraEditorSettings.h"
 
 #define LOCTEXT_NAMESPACE "NiagaraStackFunctionInputValue"
 
@@ -544,11 +545,10 @@ void SNiagaraStackFunctionInputValue::CollectAllActions(FGraphActionListBuilderB
 			FunctionInput->GetValueMode() != UNiagaraStackFunctionInput::EValueMode::Local &&
 			FunctionInput->GetInputType().IsDataInterface() == false;
 
-		const FText NameText = LOCTEXT("LocalValue", "Set a local value");
+		const FText NameText = LOCTEXT("LocalValue", "New Local Value");
 		const FText Tooltip = FText::Format(LOCTEXT("LocalValueToolTip", "Set a local editable value for this input."), NameText);
-		const FText CategoryName = LOCTEXT("LocalValueCategory", "Local");
 		TSharedPtr<FNiagaraMenuAction> SetLocalValueAction(
-			new FNiagaraMenuAction(CategoryName, NameText, Tooltip, 0, FText(),
+			new FNiagaraMenuAction(FText(), NameText, Tooltip, 0, FText(),
 				FNiagaraMenuAction::FOnExecuteStackAction::CreateSP(this, &SNiagaraStackFunctionInputValue::SetToLocalValue),
 				FNiagaraMenuAction::FCanExecuteStackAction::CreateLambda([=]() { return bCanSetLocalValue; })));
 		OutAllActions.AddAction(SetLocalValueAction);
@@ -576,67 +576,24 @@ void SNiagaraStackFunctionInputValue::CollectAllActions(FGraphActionListBuilderB
 	TArray<FNiagaraParameterHandle> AvailableHandles;
 	FunctionInput->GetAvailableParameterHandles(AvailableHandles);
 
-	TArray<FNiagaraParameterHandle> ParameterCollectionHandles;
-	TArray<FNiagaraParameterHandle> UserHandles;
-	TArray<FNiagaraParameterHandle> EngineHandles;
-	TArray<FNiagaraParameterHandle> SystemHandles;
-	TArray<FNiagaraParameterHandle> EmitterHandles;
-	TArray<FNiagaraParameterHandle> ParticleAttributeHandles;
-	TArray<FNiagaraParameterHandle> OtherHandles;
-	for (const FNiagaraParameterHandle AvailableHandle : AvailableHandles)
+	const FString RootCategoryName = FString("Link Inputs");
+	const FText MapInputFormat = LOCTEXT("LinkInputFormat", "Link this input to {0}");
+	for (const FNiagaraParameterHandle& AvailableHandle : AvailableHandles)
 	{
-		 if (AvailableHandle.IsParameterCollectionHandle())
+		TArray<FName> HandleParts = AvailableHandle.GetHandleParts();
+		FNiagaraNamespaceMetadata NamespaceMetadata = GetDefault<UNiagaraEditorSettings>()->GetMetaDataForNamespaces(HandleParts);
+		if (NamespaceMetadata.IsValid())
 		{
-			ParameterCollectionHandles.Add(AvailableHandle);
+			// Only add handles which are in known namespaces to prevent collecting parameter handles
+			// which are being used to configure modules and dynamic inputs in the stack graphs.
+			const FText Category = NamespaceMetadata.DisplayName;
+			const FText DisplayName = FText::FromName(AvailableHandle.GetParameterHandleString());
+			const FText Tooltip = FText::Format(MapInputFormat, FText::FromName(AvailableHandle.GetParameterHandleString()));
+			TSharedPtr<FNiagaraMenuAction> LinkAction(new FNiagaraMenuAction(Category, DisplayName, Tooltip, 0, FText(),
+				FNiagaraMenuAction::FOnExecuteStackAction::CreateSP(this, &SNiagaraStackFunctionInputValue::ParameterHandleSelected, AvailableHandle)));
+			LinkAction->SetParamterHandle(AvailableHandle);
+			OutAllActions.AddAction(LinkAction, RootCategoryName);
 		}
-		else if (AvailableHandle.IsUserHandle())
-		{
-			UserHandles.Add(AvailableHandle);
-		}
-		else if (AvailableHandle.IsEngineHandle())
-		{
-			EngineHandles.Add(AvailableHandle);
-		}
-		else if (AvailableHandle.IsSystemHandle())
-		{
-			SystemHandles.Add(AvailableHandle);
-		}
-		else if (AvailableHandle.IsEmitterHandle())
-		{
-			EmitterHandles.Add(AvailableHandle);
-		}
-		else if (AvailableHandle.IsParticleAttributeHandle())
-		{
-			ParticleAttributeHandles.Add(AvailableHandle);
-		}
-		else
-		{
-			OtherHandles.Add(AvailableHandle);
-		}
-	}
-
-	{
-		const FString RootCategoryName = FString("Link Inputs");
-		auto AddMenuItemsForHandleList = [&](const TArray<FNiagaraParameterHandle>& Handles, const FText& SectionDisplay)
-		{
-			const FText MapInputFormat = LOCTEXT("LinkInputFormat", "Link this input to {0}");
-			for (const FNiagaraParameterHandle& Handle : Handles)
-			{
-				const FText DisplayName = FText::FromName(Handle.GetParameterHandleString());
-				const FText Tooltip = FText::Format(MapInputFormat, FText::FromName(Handle.GetParameterHandleString()));
-				TSharedPtr<FNiagaraMenuAction> LinkAction(new FNiagaraMenuAction(SectionDisplay, DisplayName, Tooltip, 0, FText(),
-					FNiagaraMenuAction::FOnExecuteStackAction::CreateSP(this, &SNiagaraStackFunctionInputValue::ParameterHandleSelected, Handle)));
-				OutAllActions.AddAction(LinkAction, RootCategoryName);
-			}
-		};
-
-		AddMenuItemsForHandleList(ParameterCollectionHandles, LOCTEXT("NPC", "Parameter Collections"));
-		AddMenuItemsForHandleList(UserHandles, LOCTEXT("UserSection", "User Exposed"));
-		AddMenuItemsForHandleList(EngineHandles, LOCTEXT("EngineSection", "Engine"));
-		AddMenuItemsForHandleList(SystemHandles, LOCTEXT("SystemSection", "System"));
-		AddMenuItemsForHandleList(EmitterHandles, LOCTEXT("EmitterSection", "Emitter"));
-		AddMenuItemsForHandleList(ParticleAttributeHandles, LOCTEXT("ParticleAttributeSection", "Particle Attribute"));
-		AddMenuItemsForHandleList(OtherHandles, LOCTEXT("OtherSection", "Other"));
 	}
 
 	// Read from new attribute
@@ -672,26 +629,24 @@ void SNiagaraStackFunctionInputValue::CollectAllActions(FGraphActionListBuilderB
 	}
 
 	{
-		const FText CategoryName = LOCTEXT("ExpressionCategory", "Expression");
-		const FText DisplayName = LOCTEXT("ExpressionLabel", "Make new expression");
+		const FText DisplayName = LOCTEXT("ExpressionLabel", "New Expression");
 		const FText Tooltip = LOCTEXT("ExpressionToolTipl", "Resolve this variable with a custom expression.");
-		TSharedPtr<FNiagaraMenuAction> ExpressionAction(new FNiagaraMenuAction(CategoryName, DisplayName, Tooltip, 0, FText(),
+		TSharedPtr<FNiagaraMenuAction> ExpressionAction(new FNiagaraMenuAction(FText(), DisplayName, Tooltip, 0, FText(),
 			FNiagaraMenuAction::FOnExecuteStackAction::CreateSP(this, &SNiagaraStackFunctionInputValue::CustomExpressionSelected)));
 		OutAllActions.AddAction(ExpressionAction);
 	}
 
 	{
-		const FText CategoryName = LOCTEXT("ScratchCategory", "Scratch");
-		const FText CreateDisplayName = LOCTEXT("ScratchLabel", "Make New Scratch Dynamic Input");
+		const FText CreateDisplayName = LOCTEXT("ScratchLabel", "New Scratch Dynamic Input");
 		const FText CreateTooltip = LOCTEXT("ScratchToolTipl", "Create a new dynamic input in the scratch pad.");
-		TSharedPtr<FNiagaraMenuAction> CreateScratchAction(new FNiagaraMenuAction(CategoryName, CreateDisplayName, CreateTooltip, 0, FText(),
+		TSharedPtr<FNiagaraMenuAction> CreateScratchAction(new FNiagaraMenuAction(FText(), CreateDisplayName, CreateTooltip, 0, FText(),
 			FNiagaraMenuAction::FOnExecuteStackAction::CreateSP(this, &SNiagaraStackFunctionInputValue::CreateScratchSelected)));
 		OutAllActions.AddAction(CreateScratchAction);
 	}
 
 	if (FunctionInput->CanDeleteInput())
 	{
-		const FText NameText = LOCTEXT("DeleteInput", "Remove");
+		const FText NameText = LOCTEXT("DeleteInput", "Remove this input");
 		const FText Tooltip = FText::Format(LOCTEXT("DeleteInputTooltip", "Remove input from module."), NameText);
 		TSharedPtr<FNiagaraMenuAction> SetLocalValueAction(
 			new FNiagaraMenuAction(FText::GetEmpty(), NameText, Tooltip, 0, FText::GetEmpty(),
