Niagara - Add additional metadata for namespace modifiers so that they have descriptive tooltips and the correct colors.

diff --git a/Engine/Plugins/FX/Niagara/Source/NiagaraEditor/Private/NiagaraEditorSettings.cpp b/Engine/Plugins/FX/Niagara/Source/NiagaraEditor/Private/NiagaraEditorSettings.cpp
--- a/Engine/Plugins/FX/Niagara/Source/NiagaraEditor/Private/NiagaraEditorSettings.cpp
+++ b/Engine/Plugins/FX/Niagara/Source/NiagaraEditor/Private/NiagaraEditorSettings.cpp
@@ -124,7 +124,25 @@ void UNiagaraEditorSettings::SetupNamespaceMetadata()
 			NSLOCTEXT("NamespaceMetadata", "ModuleModifierDisplayName", "Module"),
 			NSLOCTEXT("NamespaceMetadata", "ModuleModifierDescription", "A namespace modifier which is used when writing dataset attributes\nin module which will make that attribute unique to the module instance when used."),
 			FLinearColor(FColor(102, 102, 102)),
-			{ })
+			{ }),
+		FNiagaraNamespaceMetadata(
+			{FNiagaraConstants::SystemNamespace},
+			NSLOCTEXT("NamespaceMetadata", "SystemModifierDisplayName", "System"),
+			NSLOCTEXT("NamespaceMetadata", "SystemModifierDescription", "A namespace modifier which specifies that an engine provided parameter comes from the system."),
+			FLinearColor(FColor(49, 113, 142)),
+			{ }),
+		FNiagaraNamespaceMetadata(
+			{FNiagaraConstants::EmitterNamespace},
+			NSLOCTEXT("NamespaceMetadata", "EmitterModifierDisplayName", "Emitter"),
+			NSLOCTEXT("NamespaceMetadata", "EmitterModifierDescription", "A namespace modifier which specifies that an engine provided parameter comes from the emitter."),
+			FLinearColor(FColor(145, 99, 56)),
+			{ }),
+		FNiagaraNamespaceMetadata(
+			{FNiagaraConstants::OwnerNamespace},
+			NSLOCTEXT("NamespaceMetadata", "OwnerDisplayName", "Owner"),
+			NSLOCTEXT("NamespaceMetadata", "OwnerDescription", "A namespace modifier which specifies that an engine provided parameter comes from the owner, or component."),
+			FLinearColor(FColor(45, 45, 45)),
+			{ }),
 	};
 }
 
