Niagara - More parameter workflow improvements.
+ Fix adding parameters in the system editor.
+ Fix the styling of the add menu for the set variables node in the stack so that it uses the new parameter name widget.
+ Prevent crashes in the parameter name widget when parameters are empty or have no namespace.
+ Remove the options to create new parameters or set existing parameters from the add module menu as it causes confusion and replace it with a single option to add a set variables node.
+ Update the set variables node UI to show a message when it has no varibles to let the user know how to add them.
+ Show a * on parameters in the parameter panel which were added there and not from graph discovery.
+ Fix undo/redo for adding parameters.
+ Fix the styling of the add menu in the parameters panel to use the new parameter widget.

diff --git a/Engine/Plugins/FX/Niagara/Source/Niagara/Classes/NiagaraSystem.h b/Engine/Plugins/FX/Niagara/Source/Niagara/Classes/NiagaraSystem.h
--- a/Engine/Plugins/FX/Niagara/Source/Niagara/Classes/NiagaraSystem.h
+++ b/Engine/Plugins/FX/Niagara/Source/Niagara/Classes/NiagaraSystem.h
@@ -320,6 +320,9 @@ public:
 	UPROPERTY()
 	TArray<UNiagaraScript*> ScratchPadScripts;
 
+	UPROPERTY(transient)
+	FNiagaraParameterStore EditorOnlyAddedParameters;
+
 	bool GetIsolateEnabled() const;
 	void SetIsolateEnabled(bool bIsolate);
 	
diff --git a/Engine/Plugins/FX/Niagara/Source/NiagaraEditor/Private/NiagaraNodeAssignment.cpp b/Engine/Plugins/FX/Niagara/Source/NiagaraEditor/Private/NiagaraNodeAssignment.cpp
--- a/Engine/Plugins/FX/Niagara/Source/NiagaraEditor/Private/NiagaraNodeAssignment.cpp
+++ b/Engine/Plugins/FX/Niagara/Source/NiagaraEditor/Private/NiagaraNodeAssignment.cpp
@@ -21,6 +21,8 @@
 #include "ScopedTransaction.h"
 #include "ViewModels/Stack/INiagaraStackItemGroupAddUtilities.h"
 #include "ViewModels/Stack/NiagaraStackGraphUtilities.h"
+#include "Widgets/Layout/SBox.h"
+#include "Widgets/SNiagaraParameterName.h"
 
 #define LOCTEXT_NAMESPACE "NiagaraNodeAssigment"
 
@@ -255,12 +257,16 @@ void UNiagaraNodeAssignment::BuildAddParameterMenu(FMenuBuilder& MenuBuilder, EN
 		bool bCanExecute = AssignmentTargets.Contains(AvailableParameter) == false; 
 
 		MenuBuilder.AddMenuEntry(
-			NameText,
-			TooltipDesc,
-			FSlateIcon(),
 			FUIAction(
 				FExecuteAction::CreateUObject(this, &UNiagaraNodeAssignment::AddParameter, AvailableParameter, VarDefaultValue),
-				FCanExecuteAction::CreateLambda([bCanExecute] { return bCanExecute; })));
+				FCanExecuteAction::CreateLambda([bCanExecute] { return bCanExecute; })),
+			SNew(SBox)
+			.Padding(FMargin(5, 1, 5, 1))
+			[
+				SNew(SNiagaraParameterName)
+				.ParameterName(AvailableParameter.GetName())
+				.IsReadOnly(true)
+			]);
 	}
 
 }
diff --git a/Engine/Plugins/FX/Niagara/Source/NiagaraEditor/Private/SNiagaraParameterName.cpp b/Engine/Plugins/FX/Niagara/Source/NiagaraEditor/Private/SNiagaraParameterName.cpp
--- a/Engine/Plugins/FX/Niagara/Source/NiagaraEditor/Private/SNiagaraParameterName.cpp
+++ b/Engine/Plugins/FX/Niagara/Source/NiagaraEditor/Private/SNiagaraParameterName.cpp
@@ -69,6 +69,11 @@ void SNiagaraParameterName::UpdateContent(FName InDisplayedParameterName)
 	TArray<FString> NameParts;
 	DisplayedParameterNameString.ParseIntoArray(NameParts, TEXT("."));
 
+	if (NameParts.Num() == 0)
+	{
+		return;
+	}
+
 	TArray<FName> Namespaces;
 	for (int32 i = 0; i < NameParts.Num() - 1; i++)
 	{
@@ -79,31 +84,34 @@ void SNiagaraParameterName::UpdateContent(FName InDisplayedParameterName)
 
 	FNiagaraNamespaceMetadata DefaultNamespaceMetadata = GetDefault<UNiagaraEditorSettings>()->GetMetaDataForNamespaces({ NAME_None });
 
-	// Add the namespace widget.
-	FNiagaraNamespaceMetadata NamespaceMetadata = GetDefault<UNiagaraEditorSettings>()->GetMetaDataForNamespaces(Namespaces);
-	TSharedPtr<SWidget> NamespaceWidget;
-	if (NamespaceMetadata.IsValid())
-	{
-		Namespaces.RemoveAt(0, NamespaceMetadata.Namespaces.Num());
-		NamespaceWidget = CreateNamespaceWidget(NamespaceMetadata.DisplayName.ToUpper(), NamespaceMetadata.Description, NamespaceMetadata.BackgroundColor);
-	}
-	else
+	if (Namespaces.Num() > 0)
 	{
-		FText NamespaceDisplayName = FText::FromString(FName::NameToDisplayString(Namespaces[0].ToString(), false).ToUpper());
-		Namespaces.RemoveAt(0);
-		CreateNamespaceWidget(
-			NamespaceDisplayName,
-			DefaultNamespaceMetadata.Description,
-			DefaultNamespaceMetadata.BackgroundColor);
-	}
+		// Add the namespace widget.
+		FNiagaraNamespaceMetadata NamespaceMetadata = GetDefault<UNiagaraEditorSettings>()->GetMetaDataForNamespaces(Namespaces);
+		TSharedPtr<SWidget> NamespaceWidget;
+		if (NamespaceMetadata.IsValid())
+		{
+			Namespaces.RemoveAt(0, NamespaceMetadata.Namespaces.Num());
+			NamespaceWidget = CreateNamespaceWidget(NamespaceMetadata.DisplayName.ToUpper(), NamespaceMetadata.Description, NamespaceMetadata.BackgroundColor);
+		}
+		else
+		{
+			FText NamespaceDisplayName = FText::FromString(FName::NameToDisplayString(Namespaces[0].ToString(), false).ToUpper());
+			Namespaces.RemoveAt(0);
+			CreateNamespaceWidget(
+				NamespaceDisplayName,
+				DefaultNamespaceMetadata.Description,
+				DefaultNamespaceMetadata.BackgroundColor);
+		}
 
-	ContentBox->AddSlot()
-	.VAlign(VAlign_Center)
-	.AutoWidth()
-	.Padding(0.0f, 0.0f, 5.0f, 0.0f)
-	[
-		NamespaceWidget.ToSharedRef()
-	];
+		ContentBox->AddSlot()
+		.VAlign(VAlign_Center)
+		.AutoWidth()
+		.Padding(0.0f, 0.0f, 5.0f, 0.0f)
+		[
+			NamespaceWidget.ToSharedRef()
+		];
+	}
 
 	// Next the namespace modifier widget is there is a namespace modifier.
 	if (Namespaces.Num() > 0)
diff --git a/Engine/Plugins/FX/Niagara/Source/NiagaraEditor/Private/ViewModels/Stack/NiagaraStackItem.cpp b/Engine/Plugins/FX/Niagara/Source/NiagaraEditor/Private/ViewModels/Stack/NiagaraStackItem.cpp
--- a/Engine/Plugins/FX/Niagara/Source/NiagaraEditor/Private/ViewModels/Stack/NiagaraStackItem.cpp
+++ b/Engine/Plugins/FX/Niagara/Source/NiagaraEditor/Private/ViewModels/Stack/NiagaraStackItem.cpp
@@ -172,3 +172,15 @@ bool UNiagaraStackItemContent::FilterAdvancedChildren(const UNiagaraStackEntry&
 		return GetStackEditorData().GetShowAllAdvanced() || GetStackEditorData().GetStackItemShowAdvanced(OwningStackItemEditorDataKey, false);
 	}
 }
+
+void UNiagaraStackItemTextContent::Initialize(FRequiredEntryData InRequiredEntryData, FText InDisplayText, bool bInIsAdvanced, FString InOwningStackItemEditorDataKey)
+{
+	Super::Initialize(InRequiredEntryData, bInIsAdvanced, InOwningStackItemEditorDataKey, InOwningStackItemEditorDataKey + InDisplayText.ToString());
+	DisplayText = InDisplayText;
+}
+
+FText UNiagaraStackItemTextContent::GetDisplayName() const
+{
+	return DisplayText;
+}
+
diff --git a/Engine/Plugins/FX/Niagara/Source/NiagaraEditor/Private/ViewModels/Stack/NiagaraStackModuleItem.cpp b/Engine/Plugins/FX/Niagara/Source/NiagaraEditor/Private/ViewModels/Stack/NiagaraStackModuleItem.cpp
--- a/Engine/Plugins/FX/Niagara/Source/NiagaraEditor/Private/ViewModels/Stack/NiagaraStackModuleItem.cpp
+++ b/Engine/Plugins/FX/Niagara/Source/NiagaraEditor/Private/ViewModels/Stack/NiagaraStackModuleItem.cpp
@@ -251,6 +251,21 @@ void UNiagaraStackModuleItem::RefreshChildrenInternal(const TArray<UNiagaraStack
 			InputCollection->SetShouldShowInStack(false);
 
 			NewChildren.Add(InputCollection);
+
+			UNiagaraNodeAssignment* AssignmentNode = CastChecked<UNiagaraNodeAssignment>(FunctionCallNode);
+			if (AssignmentNode->GetAssignmentTargets().Num() == 0)
+			{
+				FText EmptyAssignmentNodeMessageText = LOCTEXT("EmptyAssignmentNodeMessage", "No variables\n\nTo add a variable use the add button in the header, or drag a parameter from the parameters tab to the header.");
+				UNiagaraStackItemTextContent* EmptyAssignmentNodeMessage = FindCurrentChildOfTypeByPredicate<UNiagaraStackItemTextContent>(NewChildren,
+					[&](UNiagaraStackItemTextContent* CurrentStackItemTextContent) { return CurrentStackItemTextContent->GetDisplayName().IdenticalTo(EmptyAssignmentNodeMessageText); });
+
+				if (EmptyAssignmentNodeMessage == nullptr)
+				{
+					EmptyAssignmentNodeMessage = NewObject<UNiagaraStackItemTextContent>(this);
+					EmptyAssignmentNodeMessage->Initialize(CreateDefaultChildRequiredData(), EmptyAssignmentNodeMessageText, false, GetStackEditorDataKey());
+				}
+				NewChildren.Add(EmptyAssignmentNodeMessage);
+			}
 		}
 	}
 
diff --git a/Engine/Plugins/FX/Niagara/Source/NiagaraEditor/Private/ViewModels/Stack/NiagaraStackScriptItemGroup.cpp b/Engine/Plugins/FX/Niagara/Source/NiagaraEditor/Private/ViewModels/Stack/NiagaraStackScriptItemGroup.cpp
--- a/Engine/Plugins/FX/Niagara/Source/NiagaraEditor/Private/ViewModels/Stack/NiagaraStackScriptItemGroup.cpp
+++ b/Engine/Plugins/FX/Niagara/Source/NiagaraEditor/Private/ViewModels/Stack/NiagaraStackScriptItemGroup.cpp
@@ -58,7 +58,7 @@ public:
 		FText Keywords;
 		AssetData.GetTagValue(GET_MEMBER_NAME_CHECKED(UNiagaraScript, Keywords), Keywords);
 
-		return MakeShareable(new FScriptGroupAddAction(Category, DisplayName, Description, Keywords, FNiagaraVariable(), false, AssetData, nullptr, false));
+		return MakeShareable(new FScriptGroupAddAction(Category, DisplayName, Description, Keywords, FNiagaraVariable(), false, AssetData, nullptr, false, false));
 	}
 
 	static TSharedRef<FScriptGroupAddAction> CreateScriptModuleAction(UNiagaraScript* ModuleScript)
@@ -73,7 +73,7 @@ public:
 		FText Description = FNiagaraEditorUtilities::FormatScriptDescription(ModuleScript->Description, *ModuleScript->GetPathName(), ModuleScript->bExposeToLibrary);
 		FText Keywords = ModuleScript->Keywords;
 
-		return MakeShareable(new FScriptGroupAddAction(Category, DisplayName, Description, Keywords, FNiagaraVariable(), false, FAssetData(), ModuleScript, false));
+		return MakeShareable(new FScriptGroupAddAction(Category, DisplayName, Description, Keywords, FNiagaraVariable(), false, FAssetData(), ModuleScript, false, false));
 	}
 
 	static TSharedRef<FScriptGroupAddAction> CreateExistingParameterModuleAction(FNiagaraVariable ParameterVariable)
@@ -84,7 +84,7 @@ public:
 		FText AttributeDescription = FNiagaraConstants::GetAttributeDescription(ParameterVariable);
 		FText Description = FText::Format(LOCTEXT("ExistingParameterModuleDescriptoinFormat", "Description: Set the parameter {0}. {1}"), FText::FromName(ParameterVariable.GetName()), AttributeDescription);
 
-		return MakeShareable(new FScriptGroupAddAction(Category, DisplayName, Description, FText(), ParameterVariable, false, FAssetData(), nullptr, false));
+		return MakeShareable(new FScriptGroupAddAction(Category, DisplayName, Description, FText(), ParameterVariable, false, FAssetData(), nullptr, false, false));
 	}
 
 	static TSharedRef<FScriptGroupAddAction> CreateNewParameterModuleAction(FName NewParameterNamespace, FNiagaraTypeDefinition NewParameterType)
@@ -96,16 +96,21 @@ public:
 		FNiagaraParameterHandle NewParameterHandle(NewParameterNamespace, *(TEXT("New") + NewParameterType.GetName()));
 		FNiagaraVariable NewParameter(NewParameterType, NewParameterHandle.GetParameterHandleString());
 
-		return MakeShareable(new FScriptGroupAddAction(Category, DisplayName, Description, FText(), NewParameter, true, FAssetData(), nullptr, false));
+		return MakeShareable(new FScriptGroupAddAction(Category, DisplayName, Description, FText(), NewParameter, true, FAssetData(), nullptr, false, false));
 	}
 
 	static TSharedRef<FScriptGroupAddAction> CreateNewScratchModuleAction()
 	{
-		FText Category = LOCTEXT("NewScratchModuleCategory", "Scratch");
 		FText DisplayName = LOCTEXT("NewScratchModuleName", "New Scratch Pad Module");
 		FText Description = LOCTEXT("NewScratchModuleDescription", "Description: Create a new scratch pad module.");
+		return MakeShareable(new FScriptGroupAddAction(FText(), DisplayName, Description, FText(), FNiagaraVariable(), false, FAssetData(), nullptr, true, false));
+	}
 
-		return MakeShareable(new FScriptGroupAddAction(Category, DisplayName, Description, FText(), FNiagaraVariable(), false, FAssetData(), nullptr, true));
+	static TSharedRef<FScriptGroupAddAction> CreateNewSetSpecificModuleAction()
+	{
+		FText DisplayName = LOCTEXT("NewSetSpecificModuleName", "Set new or existing value directly");
+		FText Description = LOCTEXT("NewSetSpecificModuleDescription", "Description: Create a new module that can set new orexisting values directly.");
+		return MakeShareable(new FScriptGroupAddAction(FText(), DisplayName, Description, FText(), FNiagaraVariable(), false, FAssetData(), nullptr, false, true));
 	}
 
 	virtual FText GetCategory() const override
@@ -153,10 +158,16 @@ public:
 		return bIsNewScratchModuleAction;
 	}
 
+	bool GetIsNewSetSpecificModuleAction() const
+	{
+		return bIsNewSetSpecificModuleAction;
+	}
+
 private:
 	FScriptGroupAddAction(FText InCategory, FText InDisplayName, FText InDescription, FText InKeywords,
 		FNiagaraVariable InModuleParameterVariable, bool bInRenameParameterOnAdd,
-		FAssetData InModuleAssetData, UNiagaraScript* InModuleScript, bool bInIsNewScratchModuleAction)
+		FAssetData InModuleAssetData, UNiagaraScript* InModuleScript,
+		bool bInIsNewScratchModuleAction, bool bInIsNewSetSpecificModuleAction)
 
 		: Category(InCategory)
 		, DisplayName(InDisplayName)
@@ -167,6 +178,7 @@ private:
 		, ModuleAssetData(InModuleAssetData)
 		, ModuleScript(InModuleScript)
 		, bIsNewScratchModuleAction(bInIsNewScratchModuleAction)
+		, bIsNewSetSpecificModuleAction(bInIsNewSetSpecificModuleAction)
 	{
 	}
 
@@ -181,6 +193,7 @@ private:
 	UNiagaraScript* ModuleScript;
 	bool bIsMaterialParameterModuleAction;
 	bool bIsNewScratchModuleAction;
+	bool bIsNewSetSpecificModuleAction;
 };
 
 class FScriptItemGroupAddUtilities : public TNiagaraStackItemGroupAddUtilities<UNiagaraNodeFunctionCall*>
@@ -222,26 +235,6 @@ public:
 			OutAddActions.Add(FScriptGroupAddAction::CreateAssetModuleAction(ModuleAsset));
 		}
 
-		// Generate actions for the available parameters to set.
-		TArray<FNiagaraVariable> AvailableParameters;
-		FNiagaraStackGraphUtilities::GetAvailableParametersForScript(*OutputNode, AvailableParameters);
-		for (const FNiagaraVariable& AvailableParameter : AvailableParameters)
-		{
-			OutAddActions.Add(FScriptGroupAddAction::CreateExistingParameterModuleAction(AvailableParameter));
-		}
-
-		// Generate actions for setting new typed parameters.
-		TOptional<FName> NewParameterNamespace = FNiagaraStackGraphUtilities::GetNamespaceForScriptUsage(OutputNode->GetUsage());
-		if (NewParameterNamespace.IsSet())
-		{
-			TArray<FNiagaraTypeDefinition> AvailableTypes;
-			FNiagaraStackGraphUtilities::GetNewParameterAvailableTypes(AvailableTypes, NewParameterNamespace.Get(TEXT("Unknown")));
-			for (const FNiagaraTypeDefinition& AvailableType : AvailableTypes)
-			{
-				OutAddActions.Add(FScriptGroupAddAction::CreateNewParameterModuleAction(NewParameterNamespace.GetValue(), AvailableType));
-			}
-		}
-
 		// Generate actions for scratch pad modules
 		for (TSharedRef<FNiagaraScratchPadScriptViewModel> ScratchPadScriptViewModel : SystemViewModel.Pin()->GetScriptScratchPadViewModel()->GetScriptViewModels())
 		{
@@ -255,7 +248,9 @@ public:
 				}
 			}
 		}
+
 		OutAddActions.Add(FScriptGroupAddAction::CreateNewScratchModuleAction());
+		OutAddActions.Add(FScriptGroupAddAction::CreateNewSetSpecificModuleAction());
 	}
 
 	virtual void ExecuteAddAction(TSharedRef<INiagaraStackItemGroupAddAction> AddAction, int32 TargetIndex) override
@@ -285,6 +280,10 @@ public:
 				ScratchPadScriptViewModel->SetIsPendingRename(true);
 			}
 		}
+		else if (ScriptGroupAddAction->GetIsNewSetSpecificModuleAction())
+		{
+			NewModuleNode = AddParameterModule(FNiagaraVariable(), false, TargetIndex);
+		}
 
 		if (NewModuleNode != nullptr)
 		{
@@ -302,9 +301,14 @@ private:
 	UNiagaraNodeFunctionCall* AddParameterModule(const FNiagaraVariable& ParameterVariable, bool bRenameParameterOnAdd, int32 TargetIndex)
 	{
 		TArray<FNiagaraVariable> Vars;
-		Vars.Add(ParameterVariable);
 		TArray<FString> DefaultVals;
-		DefaultVals.Add(FNiagaraConstants::GetAttributeDefaultValue(ParameterVariable));
+
+		if (ParameterVariable.IsValid())
+		{
+			Vars.Add(ParameterVariable);
+			DefaultVals.Add(FNiagaraConstants::GetAttributeDefaultValue(ParameterVariable));
+		}
+
 		UNiagaraNodeAssignment* NewAssignmentModule = FNiagaraStackGraphUtilities::AddParameterModuleToStack(Vars, *OutputNode, TargetIndex,DefaultVals );
 		
 		TArray<const UEdGraphPin*> InputPins;
diff --git a/Engine/Plugins/FX/Niagara/Source/NiagaraEditor/Private/Widgets/SNiagaraParameterMapPaletteItem.cpp b/Engine/Plugins/FX/Niagara/Source/NiagaraEditor/Private/Widgets/SNiagaraParameterMapPaletteItem.cpp
--- a/Engine/Plugins/FX/Niagara/Source/NiagaraEditor/Private/Widgets/SNiagaraParameterMapPaletteItem.cpp
+++ b/Engine/Plugins/FX/Niagara/Source/NiagaraEditor/Private/Widgets/SNiagaraParameterMapPaletteItem.cpp
@@ -46,10 +46,13 @@ void SNiagaraParameterMapPalleteItem::Construct(const FArguments& InArgs, FCreat
 		MouseButtonDownDelegate = InCreateData->MouseButtonDownDelegate;
 	}
 
+	TAttribute<FText> ParameterToolTipText;
+	ParameterToolTipText.Bind(this, &SNiagaraParameterMapPalleteItem::GetItemTooltip);
+	SetToolTipText(ParameterToolTipText);
+
 	ParameterNameTextBlock = SNew(SNiagaraParameterNameTextBlock)
 		.ParameterText(this, &SNiagaraParameterMapPalleteItem::GetDisplayText)
 		.HighlightText(InCreateData->HighlightText)
-		.ToolTipText(this, &SNiagaraParameterMapPalleteItem::GetItemTooltip)
 		.OnTextCommitted(this, &SNiagaraParameterMapPalleteItem::OnNameTextCommitted)
 		.OnVerifyTextChanged(this, &SNiagaraParameterMapPalleteItem::OnNameTextVerifyChanged)
 		.IsSelected(InCreateData->IsRowSelectedDelegate)
@@ -57,12 +60,13 @@ void SNiagaraParameterMapPalleteItem::Construct(const FArguments& InArgs, FCreat
 
 	InCreateData->OnRenameRequest->BindSP(ParameterNameTextBlock.ToSharedRef(), &SNiagaraParameterNameTextBlock::EnterEditingMode);
 
+	bWasCreated = ParameterAction->ReferenceCollection.ContainsByPredicate([](const FNiagaraGraphParameterReferenceCollection& ReferenceCollection) { return ReferenceCollection.WasCreated(); });
+
 	// now, create the actual widget
 	ChildSlot
 	[
 		SNew(SHorizontalBox)
 		.AddMetaData<FTutorialMetaData>(TagMeta)
-		.ToolTipText(ParameterAction->GetTooltipDescription())
 		// icon slot
 		+SHorizontalBox::Slot()
 		.AutoWidth()
@@ -71,16 +75,24 @@ void SNiagaraParameterMapPalleteItem::Construct(const FArguments& InArgs, FCreat
 		]
 		// name slot
 		+SHorizontalBox::Slot()
-		.FillWidth(1.f)
+		.AutoWidth()
 		.VAlign(VAlign_Center)
 		.Padding(5,0)
 		[
 			ParameterNameTextBlock.ToSharedRef()
 		]
-		// reference count
 		+ SHorizontalBox::Slot()
 		.AutoWidth()
 		.VAlign(VAlign_Center)
+		[
+			SNew(STextBlock)
+			.Visibility(bWasCreated ? EVisibility::Visible : EVisibility::Collapsed)
+			.Text(LOCTEXT("AddedText", "*"))
+		]
+		// reference count
+		+ SHorizontalBox::Slot()
+		.VAlign(VAlign_Center)
+		.HAlign(HAlign_Right)
 		.Padding(3, 0)
 		[
 			SNew(SComboButton)
@@ -115,6 +127,24 @@ void SNiagaraParameterMapPalleteItem::Tick(const FGeometry& AllottedGeometry, co
 	}
 }
 
+FText SNiagaraParameterMapPalleteItem::GetItemTooltip() const
+{
+	if (bWasCreated)
+	{
+		FText CurrentToolTip = SGraphPaletteItem::GetItemTooltip();
+		if (CurrentToolTip.IdenticalTo(ToolTipCache) == false)
+		{
+			ToolTipCache = CurrentToolTip;
+			CreatedToolTipCache = FText::Format(LOCTEXT("CreatedToolTipFormat", "{0}\n* Created through this panel."), ToolTipCache);
+		}
+		return CreatedToolTipCache;
+	}
+	else
+	{
+		return SGraphPaletteItem::GetItemTooltip();
+	}
+}
+
 void SNiagaraParameterMapPalleteItem::OnNameTextCommitted(const FText& NewText, ETextCommit::Type InTextCommit)
 {
 	TSharedPtr<FNiagaraParameterAction> ParameterAction = StaticCastSharedPtr<FNiagaraParameterAction>(ActionPtr.Pin());
diff --git a/Engine/Plugins/FX/Niagara/Source/NiagaraEditor/Private/Widgets/SNiagaraParameterMapPaletteItem.h b/Engine/Plugins/FX/Niagara/Source/NiagaraEditor/Private/Widgets/SNiagaraParameterMapPaletteItem.h
--- a/Engine/Plugins/FX/Niagara/Source/NiagaraEditor/Private/Widgets/SNiagaraParameterMapPaletteItem.h
+++ b/Engine/Plugins/FX/Niagara/Source/NiagaraEditor/Private/Widgets/SNiagaraParameterMapPaletteItem.h
@@ -27,6 +27,8 @@ public:
 
 	virtual void Tick(const FGeometry& AllottedGeometry, const double InCurrentTime, const float InDeltaTime) override;
 
+	virtual FText GetItemTooltip() const override;
+
 protected:
 	/** Callback when rename text is committed */
 	virtual void OnNameTextCommitted(const FText& NewText, ETextCommit::Type InTextCommit) override;
@@ -36,4 +38,7 @@ protected:
 private:
 	FOnItemRenamed OnItemRenamed;
 	TSharedPtr<SNiagaraParameterNameTextBlock> ParameterNameTextBlock;
+	bool bWasCreated;
+	mutable FText ToolTipCache;
+	mutable FText CreatedToolTipCache;
 };
\ No newline at end of file
diff --git a/Engine/Plugins/FX/Niagara/Source/NiagaraEditor/Private/Widgets/SNiagaraParameterMapView.cpp b/Engine/Plugins/FX/Niagara/Source/NiagaraEditor/Private/Widgets/SNiagaraParameterMapView.cpp
--- a/Engine/Plugins/FX/Niagara/Source/NiagaraEditor/Private/Widgets/SNiagaraParameterMapView.cpp
+++ b/Engine/Plugins/FX/Niagara/Source/NiagaraEditor/Private/Widgets/SNiagaraParameterMapView.cpp
@@ -388,41 +388,67 @@ void SNiagaraParameterMapView::AddParameter(FNiagaraVariable NewVariable)
 {
 	FNiagaraParameterHandle ParameterHandle;
 	bool bSuccess = false;
-	if (NiagaraParameterMapSectionID::OnGetSectionFromVariable(NewVariable, IsStaticSwitchParameter(NewVariable, Graphs), ParameterHandle) == NiagaraParameterMapSectionID::USER)
+
+	if (ToolkitType == SCRIPT)
 	{
-		for (UObject* Object : SelectedScriptObjects->GetSelectedObjects())
+		if (Graphs.Num() > 0)
 		{
-			if (UNiagaraSystem* System = Cast<UNiagaraSystem>(Object))
+			TSet<FName> Names;
+			for (const TWeakObjectPtr<UNiagaraGraph>& GraphWeakPtr : Graphs)
 			{
-				UNiagaraSystemEditorData* SystemEditorData = CastChecked<UNiagaraSystemEditorData>(System->GetEditorData(), ECastCheckedType::NullChecked);
-				bSuccess = FNiagaraEditorUtilities::AddParameter(NewVariable, System->GetExposedParameters(), *System, SystemEditorData->GetStackEditorData());
+				if (GraphWeakPtr.IsValid())
+				{
+					UNiagaraGraph* Graph = GraphWeakPtr.Get();
+					for (const auto& ParameterElement : Graph->GetParameterReferenceMap())
+					{
+						Names.Add(ParameterElement.Key.GetName());
+					}
+				}
 			}
-		}
-	}
-	else if (Graphs.Num() > 0)
-	{
-		TSet<FName> Names;
-		for (const TWeakObjectPtr<UNiagaraGraph>& GraphWeakPtr : Graphs)
-		{
-			if (GraphWeakPtr.IsValid())
+			const FName NewUniqueName = FNiagaraUtilities::GetUniqueName(NewVariable.GetName(), Names);
+			NewVariable.SetName(NewUniqueName);
+
+			FScopedTransaction AddTransaction(LOCTEXT("AddScriptParameterTransaction", "Add parameter to script."));
+			for (const TWeakObjectPtr<UNiagaraGraph>& GraphWeakPtr : Graphs)
 			{
-				UNiagaraGraph* Graph = GraphWeakPtr.Get();
-				for (const auto& ParameterElement : Graph->GetParameterReferenceMap())
+				if (GraphWeakPtr.IsValid())
 				{
-					Names.Add(ParameterElement.Key.GetName());
+					UNiagaraGraph* Graph = GraphWeakPtr.Get();
+					Graph->Modify();
+					Graph->AddParameter(NewVariable);
+					bSuccess = true;
 				}
 			}
 		}
-		const FName NewUniqueName = FNiagaraUtilities::GetUniqueName(NewVariable.GetName(), Names);
-		NewVariable.SetName(NewUniqueName);
-
-		for (const TWeakObjectPtr<UNiagaraGraph>& GraphWeakPtr : Graphs)
+	}
+	else if (ToolkitType == SYSTEM)
+	{
+		UNiagaraSystem* System = CachedSystem.Get();
+		if (System != nullptr)
 		{
-			if (GraphWeakPtr.IsValid())
+			FScopedTransaction AddTransaction(LOCTEXT("AddSystemParameterTransaction", "Add parameter to system."));
+			System->Modify();
+			if (NiagaraParameterMapSectionID::OnGetSectionFromVariable(NewVariable, IsStaticSwitchParameter(NewVariable, Graphs), ParameterHandle) == NiagaraParameterMapSectionID::USER)
 			{
-				UNiagaraGraph* Graph = GraphWeakPtr.Get();
-				Graph->AddParameter(NewVariable);
-				bSuccess = true;
+				UNiagaraSystemEditorData* SystemEditorData = CastChecked<UNiagaraSystemEditorData>(System->GetEditorData(), ECastCheckedType::NullChecked);
+				bSuccess = FNiagaraEditorUtilities::AddParameter(NewVariable, System->GetExposedParameters(), *System, SystemEditorData->GetStackEditorData());
+			}
+			else
+			{
+				FNiagaraVariable VariableToAdd = NewVariable;
+				if (System->EditorOnlyAddedParameters.IndexOf(VariableToAdd) != INDEX_NONE)
+				{
+					TSet<FName> ParameterNames;
+					TArray<FNiagaraVariable> AddedParameterVariables;
+					System->EditorOnlyAddedParameters.GetParameters(AddedParameterVariables);
+					for (const FNiagaraVariable& AddedParameterVariable : AddedParameterVariables)
+					{
+						ParameterNames.Add(AddedParameterVariable.GetName());
+					}
+					FName UniqueName = FNiagaraUtilities::GetUniqueName(VariableToAdd.GetName(), ParameterNames);
+					VariableToAdd.SetName(UniqueName);
+				}
+				bSuccess = System->EditorOnlyAddedParameters.AddParameter(VariableToAdd, false, false);
 			}
 		}
 	}
@@ -580,10 +606,23 @@ void SNiagaraParameterMapView::CollectAllActionsForSystemToolkit(TMap<FNiagaraVa
 			TArray<FNiagaraGraphParameterReferenceCollection>* Found = OutParameterEntries.Find(ExposedVar);
 			if (!Found)
 			{
-				TArray<FNiagaraGraphParameterReferenceCollection> Collection;
+				TArray<FNiagaraGraphParameterReferenceCollection> Collection = { FNiagaraGraphParameterReferenceCollection(true) };
 				OutParameterEntries.Add(ExposedVar, Collection);
 			}
 		}
+
+		// Collect manually added parameters.
+		TArray<FNiagaraVariable> AddedVars;
+		System->EditorOnlyAddedParameters.GetParameters(AddedVars);
+		for (const FNiagaraVariable& AddedVar : AddedVars)
+		{
+			TArray<FNiagaraGraphParameterReferenceCollection>* Found = OutParameterEntries.Find(AddedVar);
+			if (!Found)
+			{
+				TArray<FNiagaraGraphParameterReferenceCollection> Collection = { FNiagaraGraphParameterReferenceCollection(true) };
+				OutParameterEntries.Add(AddedVar, Collection);
+			}
+		}
 	}
 
 	for (TWeakObjectPtr<UNiagaraGraph> GraphWeak : Graphs)
@@ -942,6 +981,7 @@ void SNiagaraParameterMapView::Refresh(bool bRefreshMenu/* = true*/)
 		}
 		else if (UNiagaraSystem* System = Cast<UNiagaraSystem>(Object))
 		{
+			CachedSystem = System;
 			AddGraph(System->GetSystemSpawnScript()->GetSource());
 			UserParameterStoreChangedHandle = System->GetExposedParameters().AddOnChangedHandler(
 				FNiagaraParameterStore::FOnChanged::FDelegate::CreateSP(this, &SNiagaraParameterMapView::OnUserParameterStoreChanged));
@@ -1357,7 +1397,7 @@ void SNiagaraAddParameterMenu::CollectAllActions(FGraphActionListBuilderBase& Ou
 	// Particle
 	if (CanCollectSection(NiagaraParameterMapSectionID::PARTICLE) && !IDsExcluded.Contains(NiagaraParameterMapSectionID::PARTICLE))
 	{
-		const FText Category = ShowNamespaceCategory.Get() ? NiagaraParameterMapSectionID::OnGetSectionTitle(NiagaraParameterMapSectionID::PARTICLE) : FText::GetEmpty();
+		const FText Category = ShowNamespaceCategory.Get() ? NiagaraParameterMapSectionID::OnGetSectionTitle(NiagaraParameterMapSectionID::PARTICLE) : LOCTEXT("UseExistingParticleAttribute", "Use Existing");
 		TArray<FNiagaraVariable> Variables;
 		Variables = FNiagaraConstants::GetCommonParticleAttributes();
 		AddParameterGroup(OutAllActions, Variables, NiagaraParameterMapSectionID::PARTICLE, Category, FString(), true, false);
@@ -1483,7 +1523,15 @@ void SNiagaraAddParameterMenu::CollectAllActions(FGraphActionListBuilderBase& Ou
 	}
 }
 
-void SNiagaraAddParameterMenu::AddParameterGroup(FGraphActionListBuilderBase& OutActions, TArray<FNiagaraVariable>& Variables, const NiagaraParameterMapSectionID::Type InSection, const FText& Category, const FString& RootCategory, const bool bSort, const bool bCustomName)
+void SNiagaraAddParameterMenu::AddParameterGroup(
+	FGraphActionListBuilderBase& OutActions,
+	TArray<FNiagaraVariable>& Variables,
+	const NiagaraParameterMapSectionID::Type InSection,
+	const FText& Category,
+	const FString& RootCategory,
+	const bool bSort,
+	const bool bCustomName,
+	bool bForMakeNew)
 {
 	if (bSort)
 	{
@@ -1502,6 +1550,10 @@ void SNiagaraAddParameterMenu::AddParameterGroup(FGraphActionListBuilderBase& Ou
 
 		TSharedPtr<FNiagaraMenuAction> Action(new FNiagaraMenuAction(Category, DisplayName, Tooltip, 0, FText(),
 			FNiagaraMenuAction::FOnExecuteStackAction::CreateSP(this, &SNiagaraAddParameterMenu::AddParameterSelected, Variable, bCustomName, InSection)));
+		if (bForMakeNew == false)
+		{
+			Action->SetParamterHandle(FNiagaraParameterHandle(Variable.GetName()));
+		}
 
 		if (Variable.IsDataInterface())
 		{
@@ -1561,7 +1613,7 @@ void SNiagaraAddParameterMenu::CollectMakeNew(FGraphActionListBuilderBase& OutAc
 	AddParameterGroup(OutActions, Variables, InSection,
 		LOCTEXT("MakeNewCat", "Make New"), 
 		ShowNamespaceCategory.Get() ? NiagaraParameterMapSectionID::OnGetSectionTitle(InSection).ToString() : FString(), 
-		true, true);
+		true, true, true);
 }
 
 void SNiagaraAddParameterMenu::AddParameterSelected(FNiagaraVariable NewVariable, const bool bCreateCustomName, const NiagaraParameterMapSectionID::Type InSection)
diff --git a/Engine/Plugins/FX/Niagara/Source/NiagaraEditor/Public/ViewModels/Stack/NiagaraStackItem.h b/Engine/Plugins/FX/Niagara/Source/NiagaraEditor/Public/ViewModels/Stack/NiagaraStackItem.h
--- a/Engine/Plugins/FX/Niagara/Source/NiagaraEditor/Public/ViewModels/Stack/NiagaraStackItem.h
+++ b/Engine/Plugins/FX/Niagara/Source/NiagaraEditor/Public/ViewModels/Stack/NiagaraStackItem.h
@@ -92,3 +92,17 @@ private:
 	FString OwningStackItemEditorDataKey;
 	bool bIsAdvanced;
 };
+
+UCLASS()
+class NIAGARAEDITOR_API UNiagaraStackItemTextContent : public UNiagaraStackItemContent
+{
+	GENERATED_BODY()
+
+public:
+	void Initialize(FRequiredEntryData InRequiredEntryData, FText InDisplayText, bool bInIsAdvanced, FString InOwningStackItemEditorDataKey);
+
+	virtual FText GetDisplayName() const override;
+
+private:
+	FText DisplayText;
+};
diff --git a/Engine/Plugins/FX/Niagara/Source/NiagaraEditor/Public/Widgets/SNiagaraParameterMapView.h b/Engine/Plugins/FX/Niagara/Source/NiagaraEditor/Public/Widgets/SNiagaraParameterMapView.h
--- a/Engine/Plugins/FX/Niagara/Source/NiagaraEditor/Public/Widgets/SNiagaraParameterMapView.h
+++ b/Engine/Plugins/FX/Niagara/Source/NiagaraEditor/Public/Widgets/SNiagaraParameterMapView.h
@@ -229,7 +229,16 @@ public:
 
 	TSharedRef<SEditableTextBox> GetSearchBox();
 
-	void AddParameterGroup(FGraphActionListBuilderBase& OutActions, TArray<FNiagaraVariable>& Variables, const NiagaraParameterMapSectionID::Type InSection = NiagaraParameterMapSectionID::NONE, const FText& Category = FText::GetEmpty(), const FString& RootCategory = FString(), const bool bSort = true, const bool bCustomName = true);
+	void AddParameterGroup(
+		FGraphActionListBuilderBase& OutActions,
+		TArray<FNiagaraVariable>& Variables,
+		const NiagaraParameterMapSectionID::Type InSection = NiagaraParameterMapSectionID::NONE,
+		const FText& Category = FText::GetEmpty(),
+		const FString& RootCategory = FString(),
+		const bool bSort = true,
+		const bool bCustomName = true,
+		bool bForMakeNew = false);
+	
 	void CollectParameterCollectionsActions(FGraphActionListBuilderBase& OutActions);
 	void CollectMakeNew(FGraphActionListBuilderBase& OutActions, const NiagaraParameterMapSectionID::Type InSection = NiagaraParameterMapSectionID::NONE);
 
diff --git a/Engine/Plugins/FX/Niagara/Source/NiagaraEditorWidgets/Private/NiagaraEditorWidgetsStyle.cpp b/Engine/Plugins/FX/Niagara/Source/NiagaraEditorWidgets/Private/NiagaraEditorWidgetsStyle.cpp
--- a/Engine/Plugins/FX/Niagara/Source/NiagaraEditorWidgets/Private/NiagaraEditorWidgetsStyle.cpp
+++ b/Engine/Plugins/FX/Niagara/Source/NiagaraEditorWidgets/Private/NiagaraEditorWidgetsStyle.cpp
@@ -99,10 +99,10 @@ TSharedRef< FSlateStyleSet > FNiagaraEditorWidgetsStyle::Create()
 		.SetFont(ParameterFont);
 	Style->Set("NiagaraEditor.Stack.ParameterText", ParameterText);
 
-	FSlateFontInfo ParameterCollectionFont = DEFAULT_FONT("Regular", 9);
-	FTextBlockStyle ParameterCollectionText = FTextBlockStyle(NormalText)
-		.SetFont(ParameterCollectionFont);
-	Style->Set("NiagaraEditor.Stack.ParameterCollectionText", ParameterCollectionText);
+	FSlateFontInfo TextContentFont = DEFAULT_FONT("Regular", 9);
+	FTextBlockStyle TextContentText = FTextBlockStyle(NormalText)
+		.SetFont(TextContentFont);
+	Style->Set("NiagaraEditor.Stack.TextContentText", TextContentText);
 
 	FSlateFontInfo StackItemFont = DEFAULT_FONT("Regular", 11);
 	FTextBlockStyle StackItemText = FTextBlockStyle(NormalText)
diff --git a/Engine/Plugins/FX/Niagara/Source/NiagaraEditorWidgets/Private/SNiagaraStack.cpp b/Engine/Plugins/FX/Niagara/Source/NiagaraEditorWidgets/Private/SNiagaraStack.cpp
--- a/Engine/Plugins/FX/Niagara/Source/NiagaraEditorWidgets/Private/SNiagaraStack.cpp
+++ b/Engine/Plugins/FX/Niagara/Source/NiagaraEditorWidgets/Private/SNiagaraStack.cpp
@@ -1040,6 +1040,16 @@ SNiagaraStack::FRowWidgets SNiagaraStack::ConstructNameAndValueWidgetsForItem(UN
 		UNiagaraStackItem* StackItem = CastChecked<UNiagaraStackItem>(Item);
 		return FRowWidgets(SNew(SNiagaraStackItem, *StackItem, StackViewModel));
 	}
+	else if (Item->IsA<UNiagaraStackItemTextContent>())
+	{
+		Container->SetContentPadding(FMargin(5));
+		UNiagaraStackItemTextContent* ItemTextContent = CastChecked<UNiagaraStackItemTextContent>(Item);
+		return FRowWidgets(SNew(STextBlock)
+			.TextStyle(FNiagaraEditorWidgetsStyle::Get(), "NiagaraEditor.Stack.TextContentText")
+			.Text(ItemTextContent->GetDisplayName())
+			.AutoWrapText(true)
+			.Justification(ETextJustify::Center));
+	}
 	else
 	{
 		return FRowWidgets(SNullWidget::NullWidget);
