Niagara - Fix additional issues in the parameters panel.
+ Undo/redo for name and namespace modifier changes now work.
+ Adding and removing parameters in the system editor now works consistently.
+ The context menu will now disable entries and will show a tooltip when operations are not valid with a reason.
+ Added a "Copy Parameter" menu option which will copy the old style parameter name for use with hlsl nodes and expressions and copy/paste of parameters in the future.
+ Fix a crash when displaying unknown namespaces.
+ Prevent emitter attributes from showing up in the system add menu.
+ Prevent adding new outputs in the system parameter editor.
+ Add the namespace metadata tooltips to the parameter panel categories.

diff --git a/Engine/Plugins/FX/Niagara/Source/NiagaraEditor/Private/NiagaraEditorSettings.cpp b/Engine/Plugins/FX/Niagara/Source/NiagaraEditor/Private/NiagaraEditorSettings.cpp
--- a/Engine/Plugins/FX/Niagara/Source/NiagaraEditor/Private/NiagaraEditorSettings.cpp
+++ b/Engine/Plugins/FX/Niagara/Source/NiagaraEditor/Private/NiagaraEditorSettings.cpp
@@ -4,6 +4,7 @@
 #include "NiagaraConstants.h"
 
 FNiagaraNamespaceMetadata::FNiagaraNamespaceMetadata()
+	: BackgroundColor(FLinearColor::Black)
 {
 }
 
@@ -42,31 +43,31 @@ void UNiagaraEditorSettings::SetupNamespaceMetadata()
 			NSLOCTEXT("NamespaceMetadata", "SystemDisplayName", "System"),
 			NSLOCTEXT("NamespaceMetadata", "SystemDescription", "Persistent attribute in the system which is written in a system\n stage and can be read anywhere."),
 			FLinearColor(FColor(49, 113, 142)),
-			{ENiagaraNamespaceMetadataOptions::SupportsNamespaceModifier}),
+			{ENiagaraNamespaceMetadataOptions::CanChangeNamespaceModifier}),
 		FNiagaraNamespaceMetadata(
 			{FNiagaraConstants::EmitterNamespace},
 			NSLOCTEXT("NamespaceMetadata", "EmitterDisplayName", "Emitter"),
 			NSLOCTEXT("NamespaceMetadata", "EmitterDescription", "Persistent attribute which is written in a emitter\n stage and can be read in emitter and particle stages."),
 			FLinearColor(FColor(145, 99, 56)),
-			{ENiagaraNamespaceMetadataOptions::SupportsNamespaceModifier}),
+			{ENiagaraNamespaceMetadataOptions::CanChangeNamespaceModifier}),
 		FNiagaraNamespaceMetadata(
 			{FNiagaraConstants::ParticleAttributeNamespace},
 			NSLOCTEXT("NamespaceMetadata", "ParticleDisplayName", "Particles"),
 			NSLOCTEXT("NamespaceMetadata", "ParticleDescription", "Persistent attribute which is written in a particle\n stage and can be read in particle stages."),
 			FLinearColor(FColor(72, 130, 71)),
-			{ENiagaraNamespaceMetadataOptions::SupportsNamespaceModifier}),
+			{ENiagaraNamespaceMetadataOptions::CanChangeNamespaceModifier}),
 		FNiagaraNamespaceMetadata(
 			{FNiagaraConstants::ModuleNamespace},
 			NSLOCTEXT("NamespaceMetadata", "ModuleDisplayName", "Input"),
 			NSLOCTEXT("NamespaceMetadata", "ModuleDescription", "A transient value which exposes a module input to the system and emitter editor."),
 			FLinearColor(FColor(136, 66, 65)),
-			{ENiagaraNamespaceMetadataOptions::SupportsNamespaceModifier}),
+			{ENiagaraNamespaceMetadataOptions::CanChangeNamespaceModifier}),
 		FNiagaraNamespaceMetadata(
 			{FNiagaraConstants::OutputNamespace, FNiagaraConstants::ModuleNamespace},
 			NSLOCTEXT("NamespaceMetadata", "ModuleOutputDisplayName", "Output"),
 			NSLOCTEXT("NamespaceMetadata", "ModuleOutputDescription", "A transient value which is unique to the module that wrote it and can be read from any other module.\nTransient values do not persist from frame to frame, or between stages, e.g. emitter to particle, or spawn to update."),
 			FLinearColor(FColor(109, 95, 124)),
-			{ENiagaraNamespaceMetadataOptions::Advanced}),
+			{ENiagaraNamespaceMetadataOptions::Advanced, ENiagaraNamespaceMetadataOptions::PreventCreatingInSystemEditor}),
 		FNiagaraNamespaceMetadata(
 			{FNiagaraConstants::OutputNamespace},
 			NSLOCTEXT("NamespaceMetadata", "OutputDisplayName", "Output"),
@@ -102,13 +103,13 @@ void UNiagaraEditorSettings::SetupNamespaceMetadata()
 			NSLOCTEXT("NamespaceMetadata", "NiagaraParameterCollectionDisplayName", "Niagara Parameter Collection"),
 			NSLOCTEXT("NamespaceMetadata", "NiagaraParameterCollectionDescription", "Values read from a niagara parameter collection asset.\nRead only in a niagara system."),
 			FLinearColor(FColor(45, 45, 45)),
-			{ENiagaraNamespaceMetadataOptions::Advanced}),
+			{ENiagaraNamespaceMetadataOptions::Advanced, ENiagaraNamespaceMetadataOptions::PreventRenaming}),
 		FNiagaraNamespaceMetadata(
 			{FNiagaraConstants::DataInstanceNamespace},
 			NSLOCTEXT("NamespaceMetadata", "DataInstanceDisplayName", "Data Instance"),
 			NSLOCTEXT("NamespaceMetadata", "DataInstanceDescription", "A special transient value which has a single bool IsAlive value, which determines if a particle is alive or not."),
 			FLinearColor(FColor(45, 45, 45)),
-			{ENiagaraNamespaceMetadataOptions::Advanced}),
+			{ENiagaraNamespaceMetadataOptions::Advanced, ENiagaraNamespaceMetadataOptions::PreventRenaming}),
 	};
 
 	NamespaceModifierMetadata =
diff --git a/Engine/Plugins/FX/Niagara/Source/NiagaraEditor/Private/NiagaraGraph.cpp b/Engine/Plugins/FX/Niagara/Source/NiagaraEditor/Private/NiagaraGraph.cpp
--- a/Engine/Plugins/FX/Niagara/Source/NiagaraEditor/Private/NiagaraGraph.cpp
+++ b/Engine/Plugins/FX/Niagara/Source/NiagaraEditor/Private/NiagaraGraph.cpp
@@ -111,11 +111,6 @@ void UNiagaraGraph::RemoveOnGraphNeedsRecompileHandler(FDelegateHandle Handle)
 	OnGraphNeedsRecompile.Remove(Handle);
 }
 
-void UNiagaraGraph::MutableAddParameter(FNiagaraVariable& Parameter, const FAddParameterOptions Options) const
-{
-	const_cast<UNiagaraGraph*>(this)->AddParameter(Parameter, Options);
-}
-
 void UNiagaraGraph::NotifyGraphChanged(const FEdGraphEditAction& InAction)
 {
 	InvalidateCachedParameterData();
@@ -2657,31 +2652,6 @@ void UNiagaraGraph::RefreshParameterReferences() const
 	{
 		VariableToScriptVariable.Remove(Variable);
 	}
-	
-	{
-		// Add the array indices to the parameters
-		// When a particle attribute is created we need access the corresponding RegisterIdx if we want to 
-		// query this attribute at a different location inside the InputData buffer. This index must be 
-		// available as well inside the UI if we want to pass it to nodes. It is why we are adding them automatically 
-		// to the ParameterToReferencesMap.
-		TArray<FString> RegisterNames;
-		for (auto& ParameterEntry : ParameterToReferencesMap)
-		{
-			const FNiagaraVariable& NiagaraVariable = ParameterEntry.Key;
-			if (FNiagaraParameterMapHistory::IsAttribute(NiagaraVariable))
-			{
-				const FString VariableName = FHlslNiagaraTranslator::GetSanitizedSymbolName(NiagaraVariable.GetName().ToString());
-				RegisterNames.Add(VariableName.Replace(PARAM_MAP_ATTRIBUTE_STR, PARAM_MAP_INDICES_STR));
-			}
-		}
-		for (const FString& RegisterName : RegisterNames)
-		{
-			FAddParameterOptions AddParameterOptions = FAddParameterOptions();
-			AddParameterOptions.bRefreshMetaDataScopeAndUsage = true;
-			FNiagaraVariable Parameter = FNiagaraVariable(FNiagaraTypeDefinition::GetIntDef(), *RegisterName);
-			MutableAddParameter(Parameter, AddParameterOptions);
-		}
-	}
 
 	bParameterReferenceRefreshPending = false;
 }
diff --git a/Engine/Plugins/FX/Niagara/Source/NiagaraEditor/Private/SNiagaraParameterName.cpp b/Engine/Plugins/FX/Niagara/Source/NiagaraEditor/Private/SNiagaraParameterName.cpp
--- a/Engine/Plugins/FX/Niagara/Source/NiagaraEditor/Private/SNiagaraParameterName.cpp
+++ b/Engine/Plugins/FX/Niagara/Source/NiagaraEditor/Private/SNiagaraParameterName.cpp
@@ -98,7 +98,7 @@ void SNiagaraParameterName::UpdateContent(FName InDisplayedParameterName)
 		{
 			FText NamespaceDisplayName = FText::FromString(FName::NameToDisplayString(Namespaces[0].ToString(), false).ToUpper());
 			Namespaces.RemoveAt(0);
-			CreateNamespaceWidget(
+			NamespaceWidget = CreateNamespaceWidget(
 				NamespaceDisplayName,
 				DefaultNamespaceMetadata.Description,
 				DefaultNamespaceMetadata.BackgroundColor);
diff --git a/Engine/Plugins/FX/Niagara/Source/NiagaraEditor/Private/ViewModels/Stack/NiagaraStackSystemSettingsGroup.cpp b/Engine/Plugins/FX/Niagara/Source/NiagaraEditor/Private/ViewModels/Stack/NiagaraStackSystemSettingsGroup.cpp
--- a/Engine/Plugins/FX/Niagara/Source/NiagaraEditor/Private/ViewModels/Stack/NiagaraStackSystemSettingsGroup.cpp
+++ b/Engine/Plugins/FX/Niagara/Source/NiagaraEditor/Private/ViewModels/Stack/NiagaraStackSystemSettingsGroup.cpp
@@ -159,6 +159,8 @@ void UNiagaraStackParameterStoreItem::Initialize(
 
 	Owner = InOwner;
 	ParameterStore = InParameterStore;
+	ParameterStoreChangedHandle = ParameterStore->AddOnChangedHandler(
+		FNiagaraParameterStore::FOnChanged::FDelegate::CreateUObject(this, &UNiagaraStackParameterStoreItem::ParameterStoreChanged));
 }
 
 FText UNiagaraStackParameterStoreItem::GetDisplayName() const
@@ -200,6 +202,25 @@ void UNiagaraStackParameterStoreItem::RefreshChildrenInternal(const TArray<UNiag
 	Super::RefreshChildrenInternal(CurrentChildren, NewChildren, NewIssues);
 }
 
+void UNiagaraStackParameterStoreItem::FinalizeInternal()
+{
+	if (Owner.IsValid() && ParameterStore != nullptr)
+	{
+		ParameterStore->RemoveOnChangedHandler(ParameterStoreChangedHandle);
+		Owner.Reset();
+		ParameterStore = nullptr;
+	}
+	Super::FinalizeInternal();
+}
+
+void UNiagaraStackParameterStoreItem::ParameterStoreChanged()
+{
+	if (IsFinalized() == false && Owner.IsValid())
+	{
+		RefreshChildren();
+	}
+}
+
 void UNiagaraStackParameterStoreItem::ParameterDeleted()
 {
 	RefreshChildren();
diff --git a/Engine/Plugins/FX/Niagara/Source/NiagaraEditor/Private/Widgets/SNiagaraParameterMapPaletteItem.cpp b/Engine/Plugins/FX/Niagara/Source/NiagaraEditor/Private/Widgets/SNiagaraParameterMapPaletteItem.cpp
--- a/Engine/Plugins/FX/Niagara/Source/NiagaraEditor/Private/Widgets/SNiagaraParameterMapPaletteItem.cpp
+++ b/Engine/Plugins/FX/Niagara/Source/NiagaraEditor/Private/Widgets/SNiagaraParameterMapPaletteItem.cpp
@@ -148,7 +148,6 @@ FText SNiagaraParameterMapPalleteItem::GetItemTooltip() const
 void SNiagaraParameterMapPalleteItem::OnNameTextCommitted(const FText& NewText, ETextCommit::Type InTextCommit)
 {
 	TSharedPtr<FNiagaraParameterAction> ParameterAction = StaticCastSharedPtr<FNiagaraParameterAction>(ActionPtr.Pin());
-	FScopedTransaction RenameParametersWithPins(LOCTEXT("RenameParameter", "Rename parameter, referenced pins and metadata"));
 	OnItemRenamed.ExecuteIfBound(NewText, *ParameterAction.Get());
 }
 
diff --git a/Engine/Plugins/FX/Niagara/Source/NiagaraEditor/Private/Widgets/SNiagaraParameterMapView.cpp b/Engine/Plugins/FX/Niagara/Source/NiagaraEditor/Private/Widgets/SNiagaraParameterMapView.cpp
--- a/Engine/Plugins/FX/Niagara/Source/NiagaraEditor/Private/Widgets/SNiagaraParameterMapView.cpp
+++ b/Engine/Plugins/FX/Niagara/Source/NiagaraEditor/Private/Widgets/SNiagaraParameterMapView.cpp
@@ -36,6 +36,8 @@
 #include "NiagaraScriptVariable.h"
 #include "Widgets/SPanel.h"
 #include "Widgets/Layout/SBox.h"
+#include "Widgets/SNullWidget.h"
+#include "Widgets/SToolTip.h"
 #include "NiagaraSystemEditorData.h"
 #include "ViewModels/Stack/NiagaraStackSystemSettingsGroup.h"
 #include "ViewModels/Stack/NiagaraStackGraphUtilities.h"
@@ -46,6 +48,7 @@
 #include "NiagaraParameterMapHistory.h"
 #include "NiagaraNodeOutput.h"
 #include "NiagaraNodeEmitter.h"
+#include "HAL/PlatformApplicationMisc.h"
 
 #define LOCTEXT_NAMESPACE "NiagaraParameterMapView"
 
@@ -192,16 +195,12 @@ bool NiagaraParameterMapSectionID::GetSectionIsAdvanced(const Type InSection)
 	}
 }
 
-void FNiagaraParameterMapViewCommands::RegisterCommands()
-{
-	UI_COMMAND(DeleteEntry, "Delete", "Delete this parameter", EUserInterfaceActionType::Button, FInputChord(EKeys::Platform_Delete));
-}
-
 SNiagaraParameterMapView::~SNiagaraParameterMapView()
 {
 	// Unregister all commands for right click on action node
-	ToolkitCommands->UnmapAction(FNiagaraParameterMapViewCommands::Get().DeleteEntry);
+	ToolkitCommands->UnmapAction(FGenericCommands::Get().Delete);
 	ToolkitCommands->UnmapAction(FGenericCommands::Get().Rename);
+	ToolkitCommands->UnmapAction(FGenericCommands::Get().Copy);
 
 	TSet<UObject*> Objects = SelectedScriptObjects->GetSelectedObjects();
 	for (UObject* Object : Objects)
@@ -217,6 +216,7 @@ SNiagaraParameterMapView::~SNiagaraParameterMapView()
 	if (CachedSystem.IsValid())
 	{
 		CachedSystem->GetExposedParameters().RemoveOnChangedHandler(UserParameterStoreChangedHandle);
+		CachedSystem->EditorOnlyAddedParameters.RemoveOnChangedHandler(AddedParameterStoreChangedHandle);
 		CachedSystem.Reset();
 	}
 
@@ -253,14 +253,16 @@ void SNiagaraParameterMapView::Construct(const FArguments& InArgs, const TArray<
 	
 	// Register all commands for right click on action node
 	{
-		FNiagaraParameterMapViewCommands::Register();
 		TSharedPtr<FUICommandList> ToolKitCommandList = ToolkitCommands;
-		ToolKitCommandList->MapAction(FNiagaraParameterMapViewCommands::Get().DeleteEntry,
+		ToolKitCommandList->MapAction(FGenericCommands::Get().Delete,
 			FExecuteAction::CreateSP(this, &SNiagaraParameterMapView::OnDeleteEntry),
 			FCanExecuteAction::CreateSP(this, &SNiagaraParameterMapView::CanDeleteEntry));
 		ToolKitCommandList->MapAction(FGenericCommands::Get().Rename,
 			FExecuteAction::CreateSP(this, &SNiagaraParameterMapView::OnRequestRenameOnActionNode),
 			FCanExecuteAction::CreateSP(this, &SNiagaraParameterMapView::CanRequestRenameOnActionNode));
+		ToolKitCommandList->MapAction(FGenericCommands::Get().Copy,
+			FExecuteAction::CreateSP(this, &SNiagaraParameterMapView::OnCopyParameterReference),
+			FCanExecuteAction::CreateSP(this, &SNiagaraParameterMapView::CanCopyParameterReference));
 	}
 
 	Refresh(false);
@@ -313,8 +315,8 @@ void SNiagaraParameterMapView::Construct(const FArguments& InArgs, const TArray<
 		.OnActionSelected(this, &SNiagaraParameterMapView::OnActionSelected)
 		.OnActionDoubleClicked(this, &SNiagaraParameterMapView::OnActionDoubleClicked)
 		.OnContextMenuOpening(this, &SNiagaraParameterMapView::OnContextMenuOpening)
-		.OnCanRenameSelectedAction(this, &SNiagaraParameterMapView::CanRequestRenameOnActionNode)
 		.OnGetSectionTitle(this, &SNiagaraParameterMapView::OnGetSectionTitle)
+		.OnGetSectionToolTip(this, &SNiagaraParameterMapView::OnGetSectionToolTip)
 		.OnGetSectionWidget(this, &SNiagaraParameterMapView::OnGetSectionWidget)
 		.OnCreateCustomRowExpander_Static(&SNiagaraParameterMapView::CreateCustomActionExpander)
 		.OnActionMatchesName(this, &SNiagaraParameterMapView::HandleActionMatchesName)
@@ -814,8 +816,17 @@ TSharedPtr<SWidget> SNiagaraParameterMapView::OnContextMenuOpening()
 		FMenuBuilder MenuBuilder(bShouldCloseWindowAfterMenuSelection, ToolkitCommands);
 		MenuBuilder.BeginSection("Edit", LOCTEXT("EditMenuHeader", "Edit"));
 		{
-			MenuBuilder.AddMenuEntry(FGenericCommands::Get().Rename, NAME_None, LOCTEXT("Rename", "Rename"), LOCTEXT("Rename_Tooltip", "Renames this parameter"));
-			MenuBuilder.AddMenuEntry(FNiagaraParameterMapViewCommands::Get().DeleteEntry);
+			TAttribute<FText> CopyReferenceToolTip;
+			CopyReferenceToolTip.Bind(this, &SNiagaraParameterMapView::GetCopyParameterReferenceToolTip);
+			MenuBuilder.AddMenuEntry(FGenericCommands::Get().Copy, NAME_None, LOCTEXT("CopyReference", "Copy Reference"), CopyReferenceToolTip);
+
+			TAttribute<FText> DeleteToolTip;
+			DeleteToolTip.Bind(this, &SNiagaraParameterMapView::GetDeleteEntryToolTip);
+			MenuBuilder.AddMenuEntry(FGenericCommands::Get().Delete, NAME_None, TAttribute<FText>(), DeleteToolTip);
+
+			TAttribute<FText> RenameToolTip;
+			RenameToolTip.Bind(this, &SNiagaraParameterMapView::GetRenameOnActionNodeToolTip);
+			MenuBuilder.AddMenuEntry(FGenericCommands::Get().Rename, NAME_None, LOCTEXT("Rename", "Rename"), RenameToolTip);
 		}
 		MenuBuilder.EndSection();
 
@@ -829,27 +840,35 @@ TSharedPtr<SWidget> SNiagaraParameterMapView::OnContextMenuOpening()
 			{
 				MenuBuilder.BeginSection("NamespaceModifier", LOCTEXT("NamespaceModifierMenuSection", "Parameter Modifiers"));
 				{
+					TAttribute<FText> AddToolTip;
+					AddToolTip.Bind(this, &SNiagaraParameterMapView::GetAddNamespaceModifierToolTip);
 					MenuBuilder.AddMenuEntry(
-						LOCTEXT("AddNamespaceModifier", "Add a modifier to this parameter."),
-						FText(),
+						LOCTEXT("AddNamespaceModifier", "Add"),
+						AddToolTip,
 						FSlateIcon(),
 						FUIAction(
 							FExecuteAction::CreateSP(this, &SNiagaraParameterMapView::OnAddNamespaceModifier),
-							FCanExecuteAction()));
+							FCanExecuteAction::CreateSP(this, &SNiagaraParameterMapView::CanAddNamespaceModifier)));
+
+					TAttribute<FText> RemoveToolTip;
+					RemoveToolTip.Bind(this, &SNiagaraParameterMapView::GetRemoveNamespaceModifierToolTip);
 					MenuBuilder.AddMenuEntry(
-						LOCTEXT("RemoveNamespaceModifier", "Remove the modifier from this parameter."),
-						FText(),
+						LOCTEXT("RemoveNamespaceModifier", "Remove"),
+						RemoveToolTip,
 						FSlateIcon(),
 						FUIAction(
 							FExecuteAction::CreateSP(this, &SNiagaraParameterMapView::OnRemoveNamespaceModifier),
-							FCanExecuteAction()));
+							FCanExecuteAction::CreateSP(this, &SNiagaraParameterMapView::CanRemoveNamespaceModifier)));
+
+					TAttribute<FText> EditToolTip;
+					EditToolTip.Bind(this, &SNiagaraParameterMapView::GetEditNamespaceModifierToolTip);
 					MenuBuilder.AddMenuEntry(
-						LOCTEXT("EditNamespaceModifier", "Edit the modifier for this parameter."),
-						FText(),
+						LOCTEXT("EditNamespaceModifier", "Edit"),
+						EditToolTip,
 						FSlateIcon(),
 						FUIAction(
 							FExecuteAction::CreateSP(this, &SNiagaraParameterMapView::OnEditNamespaceModifier),
-							FCanExecuteAction()));
+							FCanExecuteAction::CreateSP(this, &SNiagaraParameterMapView::CanEditNamespaceModifier)));
 				}
 				MenuBuilder.EndSection();
 			}
@@ -865,13 +884,35 @@ FText SNiagaraParameterMapView::OnGetSectionTitle(int32 InSectionID)
 	return NiagaraParameterMapSectionID::OnGetSectionTitle((NiagaraParameterMapSectionID::Type)InSectionID);
 }
 
+TSharedPtr<IToolTip> SNiagaraParameterMapView::OnGetSectionToolTip(int32 InSectionID)
+{
+	TArray<FName> SectionNamespaces;
+	NiagaraParameterMapSectionID::OnGetSectionNamespaces((NiagaraParameterMapSectionID::Type)InSectionID, SectionNamespaces);
+	FNiagaraNamespaceMetadata NamespaceMetadata = GetDefault<UNiagaraEditorSettings>()->GetMetaDataForNamespaces(SectionNamespaces);
+	if (NamespaceMetadata.IsValid() && NamespaceMetadata.Description.IsEmptyOrWhitespace() == false)
+	{
+		return SNew(SToolTip)
+			.Text(NamespaceMetadata.Description);
+	}
+	return TSharedPtr<IToolTip>();
+}
+
 TSharedRef<SWidget> SNiagaraParameterMapView::OnGetSectionWidget(TSharedRef<SWidget> RowWidget, int32 InSectionID)
 {
 	if (InSectionID == NiagaraParameterMapSectionID::STATIC_SWITCH)
 	{
-		TSharedPtr<SBox> Empty;
-		SAssignNew(Empty, SBox);
-		return Empty.ToSharedRef();
+		return SNullWidget::NullWidget;
+	}
+
+	if (IsSystemToolkit())
+	{
+		TArray<FName> SectionNamespaces;
+		NiagaraParameterMapSectionID::OnGetSectionNamespaces((NiagaraParameterMapSectionID::Type)InSectionID, SectionNamespaces);
+		FNiagaraNamespaceMetadata NamespaceMetadata = GetDefault<UNiagaraEditorSettings>()->GetMetaDataForNamespaces(SectionNamespaces);
+		if (NamespaceMetadata.IsValid() && NamespaceMetadata.Options.Contains(ENiagaraNamespaceMetadataOptions::PreventCreatingInSystemEditor))
+		{
+			return SNullWidget::NullWidget;
+		}
 	}
 
 	TWeakPtr<SWidget> WeakRowWidget = RowWidget;
@@ -964,6 +1005,7 @@ void SNiagaraParameterMapView::Refresh(bool bRefreshMenu/* = true*/)
 	if (CachedSystem.IsValid())
 	{
 		CachedSystem->GetExposedParameters().RemoveOnChangedHandler(UserParameterStoreChangedHandle);
+		CachedSystem->EditorOnlyAddedParameters.RemoveOnChangedHandler(AddedParameterStoreChangedHandle);
 		CachedSystem.Reset();
 	}
 
@@ -984,7 +1026,9 @@ void SNiagaraParameterMapView::Refresh(bool bRefreshMenu/* = true*/)
 			CachedSystem = System;
 			AddGraph(System->GetSystemSpawnScript()->GetSource());
 			UserParameterStoreChangedHandle = System->GetExposedParameters().AddOnChangedHandler(
-				FNiagaraParameterStore::FOnChanged::FDelegate::CreateSP(this, &SNiagaraParameterMapView::OnUserParameterStoreChanged));
+				FNiagaraParameterStore::FOnChanged::FDelegate::CreateSP(this, &SNiagaraParameterMapView::OnSystemParameterStoreChanged));
+			AddedParameterStoreChangedHandle = System->EditorOnlyAddedParameters.AddOnChangedHandler(
+				FNiagaraParameterStore::FOnChanged::FDelegate::CreateSP(this, &SNiagaraParameterMapView::OnSystemParameterStoreChanged));
 		}
 	}
 
@@ -1044,7 +1088,7 @@ void SNiagaraParameterMapView::OnGraphChanged(const FEdGraphEditAction& InAction
 	RefreshActions();
 }
 
-void SNiagaraParameterMapView::OnUserParameterStoreChanged()
+void SNiagaraParameterMapView::OnSystemParameterStoreChanged()
 {
 	if (CachedSystem.IsValid())
 	{
@@ -1052,45 +1096,118 @@ void SNiagaraParameterMapView::OnUserParameterStoreChanged()
 	}
 }
 
+FText SNiagaraParameterMapView::GetDeleteEntryToolTip() const
+{
+	if (ToolkitType == SYSTEM)
+	{
+		TArray<TSharedPtr<FEdGraphSchemaAction>> SelectedActions;
+		GraphActionMenu->GetSelectedActions(SelectedActions);
+		for (TSharedPtr<FEdGraphSchemaAction> SelectedAction : SelectedActions)
+		{
+			TSharedPtr<FNiagaraParameterAction> ParameterAction = StaticCastSharedPtr<FNiagaraParameterAction>(SelectedAction);
+			bool bWasCreated = ParameterAction->ReferenceCollection.ContainsByPredicate(
+				[](FNiagaraGraphParameterReferenceCollection& ParameterReferenceCollection) { return ParameterReferenceCollection.WasCreated(); });
+			if (bWasCreated == false)
+			{
+				return LOCTEXT("CantDeleteNonAddedSystem", "Can only delete system parameteters which were added directly.");
+			}
+		}
+	}
+	return LOCTEXT("DeleteSelected", "Delete the selected parameter.");
+}
+
 void SNiagaraParameterMapView::OnDeleteEntry()
 {
-	TArray<TSharedPtr<FEdGraphSchemaAction>> SelectedActions;
-	GraphActionMenu->GetSelectedActions(SelectedActions);
-	
-	for (auto& Action : SelectedActions)
+	if (ToolkitType == SCRIPT)
 	{
-		TSharedPtr<FNiagaraParameterAction> ParameterAction = StaticCastSharedPtr<FNiagaraParameterAction>(Action);
-		if (ParameterAction.Get())
+		TArray<TSharedPtr<FEdGraphSchemaAction>> SelectedActions;
+		GraphActionMenu->GetSelectedActions(SelectedActions);
+		for (auto& Action : SelectedActions)
 		{
-			FNiagaraParameterHandle ParameterHandle;
-			FNiagaraVariable Parameter = ParameterAction->GetParameter();
-			if (NiagaraParameterMapSectionID::OnGetSectionFromVariable(Parameter, IsStaticSwitchParameter(Parameter, Graphs), ParameterHandle) == NiagaraParameterMapSectionID::USER)
+			TSharedPtr<FNiagaraParameterAction> ParameterAction = StaticCastSharedPtr<FNiagaraParameterAction>(Action);
+			if (ParameterAction.Get())
 			{
-				for (UObject* Object : SelectedScriptObjects->GetSelectedObjects())
+				FScopedTransaction RemoveParametersWithPins(LOCTEXT("RemoveParametersWithPins", "Remove parameter and referenced pins"));
+				for (const TWeakObjectPtr<UNiagaraGraph>& GraphWeakPtr : Graphs)
 				{
-					if (UNiagaraSystem* System = Cast<UNiagaraSystem>(Object))
+					if (GraphWeakPtr.IsValid())
 					{
-						System->GetExposedParameters().RemoveParameter(Parameter);
+						UNiagaraGraph* Graph = GraphWeakPtr.Get();
+						Graph->RemoveParameter(ParameterAction->Parameter);
 					}
 				}
 			}
+		}
+	}
+	else if (ToolkitType == SYSTEM)
+	{
+		UNiagaraSystem* System = CachedSystem.Get();
+		if (System != nullptr)
+		{
+			TSharedPtr<FNiagaraParameterAction> ParameterAction;
+			FNiagaraParameterHandle ParameterHandle;
+			FNiagaraNamespaceMetadata NamespaceMetadata;
+			FText ErrorMessage;
+			if (GetNamespaceEditDataForSelection(ParameterAction, ParameterHandle, NamespaceMetadata, ErrorMessage))
+			{
+				FScopedTransaction RemoveParametersWithPins(LOCTEXT("RemoveParametersFromSystem", "Remove parameter"));
+				System->Modify();
+				System->GetExposedParameters().RemoveParameter(ParameterAction->Parameter);
+				System->EditorOnlyAddedParameters.RemoveParameter(ParameterAction->Parameter);
+			}
+		}
+	}
+}
 
-			FScopedTransaction RemoveParametersWithPins(LOCTEXT("RemoveParametersWithPins", "Remove parameter and referenced pins"));
-			for (const TWeakObjectPtr<UNiagaraGraph>& GraphWeakPtr : Graphs)
+bool SNiagaraParameterMapView::CanDeleteEntry() const
+{
+	if (ToolkitType == SCRIPT)
+	{
+		return true;
+	}
+	if (ToolkitType == SYSTEM)
+	{
+		TArray<TSharedPtr<FEdGraphSchemaAction>> SelectedActions;
+		GraphActionMenu->GetSelectedActions(SelectedActions);
+		if (SelectedActions.Num() == 1)
+		{
+			TSharedPtr<FNiagaraParameterAction> ParameterAction = StaticCastSharedPtr<FNiagaraParameterAction>(SelectedActions[0]);
+			bool bWasCreated = false;
+			for (const FNiagaraGraphParameterReferenceCollection& ParameterReferenceCollection : ParameterAction->ReferenceCollection)
 			{
-				if (GraphWeakPtr.IsValid())
+				if (ParameterReferenceCollection.WasCreated())
 				{
-					UNiagaraGraph* Graph = GraphWeakPtr.Get();
-					Graph->RemoveParameter(Parameter);
+					bWasCreated = true;
+					break;
 				}
 			}
+			return bWasCreated;
 		}
 	}
+	return false;
 }
 
-bool SNiagaraParameterMapView::CanDeleteEntry() const
+FText SNiagaraParameterMapView::GetRenameOnActionNodeToolTip() const
 {
-	return true;
+	TSharedPtr<FNiagaraParameterAction> ParameterAction;
+	FNiagaraParameterHandle ParameterHandle;
+	FNiagaraNamespaceMetadata NamespaceMetadata;
+	FText ErrorMessage;
+	if (GetNamespaceEditDataForSelection(ParameterAction, ParameterHandle, NamespaceMetadata, ErrorMessage))
+	{
+		if (NamespaceMetadata.Options.Contains(ENiagaraNamespaceMetadataOptions::PreventRenaming))
+		{	
+			return LOCTEXT("RenamingNotSupported", "The namespace for this parameter doesn't support renaming.");
+		}
+		else
+		{
+			return LOCTEXT("RenameParameter", "Rename this parameter.");
+		}
+	}
+	else
+	{
+		return ErrorMessage;
+	}
 }
 
 void SNiagaraParameterMapView::OnRequestRenameOnActionNode()
@@ -1100,60 +1217,240 @@ void SNiagaraParameterMapView::OnRequestRenameOnActionNode()
 }
 
 
-bool SNiagaraParameterMapView::CanRequestRenameOnActionNode(TWeakPtr<struct FGraphActionNode> InSelectedNode) const
+bool SNiagaraParameterMapView::CanRequestRenameOnActionNode() const
 {
-	return true;
+	TSharedPtr<FNiagaraParameterAction> ParameterAction;
+	FNiagaraParameterHandle ParameterHandle;
+	FNiagaraNamespaceMetadata NamespaceMetadata;
+	FText Unused;
+	if (GetNamespaceEditDataForSelection(ParameterAction, ParameterHandle, NamespaceMetadata, Unused))
+	{
+		bool bPreventRenameing = NamespaceMetadata.Options.Contains(ENiagaraNamespaceMetadataOptions::PreventRenaming);
+		return bPreventRenameing == false;
+	}
+	else 
+	{
+		return false;
+	}
 }
 
-bool SNiagaraParameterMapView::CanRequestRenameOnActionNode() const
+void SNiagaraParameterMapView::OnPostRenameActionNode(const FText& InText, FNiagaraParameterAction& InAction)
+{
+	FText TransactionName;
+	if (IsScriptToolkit())
+	{
+		TransactionName = LOCTEXT("RenameParameterScriptTransaction", "Rename parameter and pins.");
+	}
+	else if (IsSystemToolkit())
+	{
+		TransactionName = LOCTEXT("RenameParameterSystemTransaction", "Rename parameter.");
+	}
+	FScopedTransaction RenameTransaction(TransactionName);
+	RenameParameter(InAction.Parameter, *InText.ToString());
+}
+
+bool SNiagaraParameterMapView::GetNamespaceEditDataForSelection(
+	TSharedPtr<FNiagaraParameterAction>& OutParameterAction,
+	FNiagaraParameterHandle& OutParameterHandle,
+	FNiagaraNamespaceMetadata& OutNamespaceMetadata,
+	FText& OutErrorMessage) const
 {
-	TArray<TSharedPtr<FEdGraphSchemaAction> > SelectedActions;
+	TArray<TSharedPtr<FEdGraphSchemaAction>> SelectedActions;
 	GraphActionMenu->GetSelectedActions(SelectedActions);
 
-	// If there is anything selected in the GraphActionMenu, check the item for if it can be renamed.
-	if (SelectedActions.Num())
+	if (SelectedActions.Num() != 1)
+	{
+		// Can only operate on single items.
+		OutParameterAction.Reset();
+		OutErrorMessage = LOCTEXT("CanOnlyEditSingle", "Can only edit single selections.");
+		return false;
+	}
+
+	OutParameterAction = StaticCastSharedPtr<FNiagaraParameterAction>(SelectedActions[0]);
+	if (OutParameterAction.IsValid() == false)
 	{
-		return GraphActionMenu->CanRequestRenameOnActionNode();
+		// Invalid action.
+		OutErrorMessage = LOCTEXT("InvalidParameterAction", "Parameter action is invalid.");
+		return false;
+	}
+
+	if (ToolkitType == SYSTEM)
+	{
+		bool bWasCreated = OutParameterAction->ReferenceCollection.ContainsByPredicate([](FNiagaraGraphParameterReferenceCollection& ReferenceCollecion) { return ReferenceCollecion.WasCreated(); });
+		if (bWasCreated == false)
+		{
+			// Can only modify parameters in the system toolkit if they were added there.
+			OutParameterAction.Reset();
+			OutErrorMessage = LOCTEXT("NonAddedSystem", "Can only edit system parameteters which were added directly.");
+			return false;
+		}
+	}
+
+	OutParameterHandle = FNiagaraParameterHandle(OutParameterAction->Parameter.GetName());
+	TArray<FName> NameParts = OutParameterHandle.GetHandleParts();
+	OutNamespaceMetadata = GetDefault<UNiagaraEditorSettings>()->GetMetaDataForNamespaces(NameParts);
+	if (OutNamespaceMetadata.IsValid() == false)
+	{
+		OutErrorMessage = LOCTEXT("NoMetadataForNamespace", "This parameter doesn't support editing.");
+		return false;
+	}
+
+	return true;
+}
+
+bool SNiagaraParameterMapView::GetNamespaceModifierEditDataForSelection(
+	TSharedPtr<FNiagaraParameterAction>& OutParameterAction,
+	FNiagaraParameterHandle& OutParameterHandle,
+	FNiagaraNamespaceMetadata& OutNamespaceMetadata,
+	FText& OutErrorMessage) const
+{
+	if (GetNamespaceEditDataForSelection(OutParameterAction, OutParameterHandle, OutNamespaceMetadata, OutErrorMessage))
+	{
+		if (OutNamespaceMetadata.IsValid() == false || OutNamespaceMetadata.Options.Contains(ENiagaraNamespaceMetadataOptions::CanChangeNamespaceModifier) == false)
+		{
+			OutErrorMessage = LOCTEXT("NotSupportedForThisNamespace", "This parameter doesn't support namespace modifiers.");
+			return false;
+		}
+		return true;
+	}
+	else
+	{
+		return false;
+	}
+}
+
+FName NamePartsToName(const TArray<FName>& NameParts)
+{
+	TArray<FString> NamePartStrings;
+	for (FName NamePart : NameParts)
+	{
+		NamePartStrings.Add(NamePart.ToString());
+	}
+	return *FString::Join(NamePartStrings, TEXT("."));
+}
+
+FText SNiagaraParameterMapView::GetAddNamespaceModifierToolTip() const
+{
+	TSharedPtr<FNiagaraParameterAction> ParameterAction;
+	FNiagaraParameterHandle ParameterHandle;
+	FNiagaraNamespaceMetadata NamespaceMetadata;
+	FText ErrorMessage;
+	if (GetNamespaceModifierEditDataForSelection(ParameterAction, ParameterHandle, NamespaceMetadata, ErrorMessage))
+	{
+		return LOCTEXT("AddNamespaceModifierToolTip", "Add a default namespace modifier to this parameter.");
+	}
+	else
+	{
+		return ErrorMessage;
+	}
+}
+
+bool SNiagaraParameterMapView::CanAddNamespaceModifier() const
+{
+	TSharedPtr<FNiagaraParameterAction> ParameterAction;
+	FNiagaraParameterHandle ParameterHandle;
+	FNiagaraNamespaceMetadata NamespaceMetadata;
+	FText Unused;
+	if(GetNamespaceModifierEditDataForSelection(ParameterAction, ParameterHandle, NamespaceMetadata, Unused))
+	{
+		int32 NumberOfNamePartsAfterNamespace = ParameterHandle.GetHandleParts().Num() - NamespaceMetadata.Namespaces.Num();
+		return NumberOfNamePartsAfterNamespace == 1;
 	}
 	return false;
 }
 
-void SNiagaraParameterMapView::OnPostRenameActionNode(const FText& InText, FNiagaraParameterAction& InAction)
+void SNiagaraParameterMapView::OnAddNamespaceModifier()
 {
-	const FName NewName = FName(*InText.ToString());
-	if (!InAction.Parameter.GetName().IsEqual(NewName, ENameCase::CaseSensitive))
+	TSharedPtr<FNiagaraParameterAction> ParameterAction;
+	FNiagaraParameterHandle ParameterHandle;
+	FNiagaraNamespaceMetadata NamespaceMetadata;
+	FText Unused;
+	if (GetNamespaceModifierEditDataForSelection(ParameterAction, ParameterHandle, NamespaceMetadata, Unused))
 	{
-		FNiagaraParameterHandle ParameterHandle;
-		if (NiagaraParameterMapSectionID::OnGetSectionFromVariable(InAction.Parameter, IsStaticSwitchParameter(InAction.Parameter, Graphs), ParameterHandle) == NiagaraParameterMapSectionID::USER)
+		TArray<FName> NameParts = ParameterHandle.GetHandleParts();
+		int32 NumberOfNamePartsAfterNamespace = NameParts.Num() - NamespaceMetadata.Namespaces.Num();
+		if (NumberOfNamePartsAfterNamespace == 1)
 		{
-			// Check if the new name is also an user variable.
-			const FNiagaraVariable NewParameterValidTest = FNiagaraVariable(InAction.Parameter.GetType(), NewName);
-			if (NiagaraParameterMapSectionID::OnGetSectionFromVariable(NewParameterValidTest, IsStaticSwitchParameter(NewParameterValidTest, Graphs), ParameterHandle) == NiagaraParameterMapSectionID::USER)
-			{
-				for (UObject* Object : SelectedScriptObjects->GetSelectedObjects())
-				{
-					if (UNiagaraSystem* System = Cast<UNiagaraSystem>(Object))
-					{
-						System->GetExposedParameters().RenameParameter(InAction.Parameter, NewName);
-					}
-				}
-			}
+			NameParts.Insert(FNiagaraConstants::ModuleNamespace, NamespaceMetadata.Namespaces.Num());
+			FName NewName = NamePartsToName(NameParts);
+			FScopedTransaction Transaction(LOCTEXT("AddNamespaceModifierTransaction", "Add namespace modifier."));
+			RenameParameter(ParameterAction->Parameter, NewName);
 		}
+	}
+}
 
-		if (Graphs.Num() > 0)
+FText SNiagaraParameterMapView::GetRemoveNamespaceModifierToolTip() const
+{
+	TSharedPtr<FNiagaraParameterAction> ParameterAction;
+	FNiagaraParameterHandle ParameterHandle;
+	FNiagaraNamespaceMetadata NamespaceMetadata;
+	FText ErrorMessage;
+	if (GetNamespaceModifierEditDataForSelection(ParameterAction, ParameterHandle, NamespaceMetadata, ErrorMessage))
+	{
+		return LOCTEXT("RemoveNamespaceModifierToolTip", "Remove this parameter's namespace modifier.");
+	}
+	else
+	{
+		return ErrorMessage;
+	}
+}
+
+bool SNiagaraParameterMapView::CanRemoveNamespaceModifier() const
+{
+	return CanEditNamespaceModifier();
+}
+
+void SNiagaraParameterMapView::OnRemoveNamespaceModifier()
+{
+	TSharedPtr<FNiagaraParameterAction> ParameterAction;
+	FNiagaraParameterHandle ParameterHandle;
+	FNiagaraNamespaceMetadata NamespaceMetadata;
+	FText Unused;
+	if (GetNamespaceModifierEditDataForSelection(ParameterAction, ParameterHandle, NamespaceMetadata, Unused))
+	{
+		TArray<FName> NameParts = ParameterHandle.GetHandleParts();
+		int32 NumberOfNamePartsAfterNamespace = NameParts.Num() - NamespaceMetadata.Namespaces.Num();
+		if (NumberOfNamePartsAfterNamespace == 2)
 		{
-			for (const TWeakObjectPtr<UNiagaraGraph>& Graph : Graphs)
-			{
-				if (Graph.IsValid())
-				{
-					Graph.Get()->RenameParameter(InAction.Parameter, NewName);
-				}
-			}
+			NameParts.RemoveAt(NamespaceMetadata.Namespaces.Num());
+			FName NewName = NamePartsToName(NameParts);
+			FScopedTransaction Transaction(LOCTEXT("RemoveNamespaceModifierTransaction", "Remove namespace modifier."));
+			RenameParameter(ParameterAction->Parameter, NewName);
 		}
 	}
 }
 
-void SNiagaraParameterMapView::OnAddNamespaceModifier()
+FText SNiagaraParameterMapView::GetEditNamespaceModifierToolTip() const
+{
+	TSharedPtr<FNiagaraParameterAction> ParameterAction;
+	FNiagaraParameterHandle ParameterHandle;
+	FNiagaraNamespaceMetadata NamespaceMetadata;
+	FText ErrorMessage;
+	if (GetNamespaceModifierEditDataForSelection(ParameterAction, ParameterHandle, NamespaceMetadata, ErrorMessage))
+	{
+		return LOCTEXT("EditNamespaceModifierToolTip", "Edit this parameter's namespace modifier.");
+	}
+	else
+	{
+		return ErrorMessage;
+	}
+}
+
+bool SNiagaraParameterMapView::CanEditNamespaceModifier() const
+{
+	TSharedPtr<FNiagaraParameterAction> ParameterAction;
+	FNiagaraParameterHandle ParameterHandle;
+	FNiagaraNamespaceMetadata NamespaceMetadata;
+	FText Unused;
+	if (GetNamespaceModifierEditDataForSelection(ParameterAction, ParameterHandle, NamespaceMetadata, Unused))
+	{
+		int32 NumberOfNamePartsAfterNamespace = ParameterHandle.GetHandleParts().Num() - NamespaceMetadata.Namespaces.Num();
+		return NumberOfNamePartsAfterNamespace == 2;
+	}
+	return false;
+}
+
+void SNiagaraParameterMapView::OnEditNamespaceModifier()
 {
 	TArray<TSharedPtr<FEdGraphSchemaAction>> SelectedActions;
 	GraphActionMenu->GetSelectedActions(SelectedActions);
@@ -1162,69 +1459,78 @@ void SNiagaraParameterMapView::OnAddNamespaceModifier()
 		TSharedPtr<FNiagaraParameterAction> ParameterAction = StaticCastSharedPtr<FNiagaraParameterAction>(SelectedActions[0]);
 		if (ParameterAction.IsValid())
 		{
-			FString ParameterNameString = ParameterAction->Parameter.GetName().ToString();
-			TArray<FString> NameParts;
-			ParameterNameString.ParseIntoArray(NameParts, TEXT("."));
-			if (NameParts.Num() == 2)
-			{
-				NameParts.Insert(FNiagaraConstants::ModuleNamespace.ToString(), 1);
-				FName NewName = *FString::Join(NameParts, TEXT("."));
-				if (Graphs.Num() > 0)
-				{
-					for (const TWeakObjectPtr<UNiagaraGraph>& Graph : Graphs)
-					{
-						if (Graph.IsValid())
-						{
-							Graph.Get()->RenameParameter(ParameterAction->Parameter, NewName);
-						}
-					}
-				}
-			}
+			ParameterAction->bNamespaceModifierRenamePending = true;
 		}
 	}
 }
 
-void SNiagaraParameterMapView::OnRemoveNamespaceModifier()
+FText SNiagaraParameterMapView::GetCopyParameterReferenceToolTip() const
+{
+	TArray<TSharedPtr<FEdGraphSchemaAction>> SelectedActions;
+	GraphActionMenu->GetSelectedActions(SelectedActions);
+	if (SelectedActions.Num() != 1)
+	{
+		return LOCTEXT("CantCopyMultipleSelection", "Can only copy single parameters.");
+	}
+
+	TSharedPtr<FNiagaraParameterAction> ParameterAction = StaticCastSharedPtr<FNiagaraParameterAction>(SelectedActions[0]);
+	if (ParameterAction.IsValid() == false)
+	{
+		return LOCTEXT("CantCopyInvalidToolTip", "Can only copy valid parameters.");
+	}
+
+	return LOCTEXT("CopyReferenceToolTip", "Copy a string reference for this parameter to the clipboard.\nThis reference can be used in expressions and custom HLSL nodes.");
+}
+
+bool SNiagaraParameterMapView::CanCopyParameterReference() const
 {
 	TArray<TSharedPtr<FEdGraphSchemaAction>> SelectedActions;
 	GraphActionMenu->GetSelectedActions(SelectedActions);
 	if (SelectedActions.Num() == 1)
 	{
 		TSharedPtr<FNiagaraParameterAction> ParameterAction = StaticCastSharedPtr<FNiagaraParameterAction>(SelectedActions[0]);
-		if (ParameterAction.IsValid())
+		if (ParameterAction.IsValid() && ParameterAction->Parameter.IsValid())
 		{
-			FString ParameterNameString = ParameterAction->Parameter.GetName().ToString();
-			TArray<FString> NameParts;
-			ParameterNameString.ParseIntoArray(NameParts, TEXT("."));
-			if (NameParts.Num() == 3)
-			{
-				NameParts.RemoveAt(1);
-				FName NewName = *FString::Join(NameParts, TEXT("."));
-				if (Graphs.Num() > 0)
-				{
-					for (const TWeakObjectPtr<UNiagaraGraph>& Graph : Graphs)
-					{
-						if (Graph.IsValid())
-						{
-							Graph.Get()->RenameParameter(ParameterAction->Parameter, NewName);
-						}
-					}
-				}
-			}
+			return true;
 		}
 	}
+	return false;
 }
 
-void SNiagaraParameterMapView::OnEditNamespaceModifier()
+void SNiagaraParameterMapView::OnCopyParameterReference()
 {
 	TArray<TSharedPtr<FEdGraphSchemaAction>> SelectedActions;
 	GraphActionMenu->GetSelectedActions(SelectedActions);
 	if (SelectedActions.Num() == 1)
 	{
 		TSharedPtr<FNiagaraParameterAction> ParameterAction = StaticCastSharedPtr<FNiagaraParameterAction>(SelectedActions[0]);
-		if (ParameterAction.IsValid())
+		FPlatformApplicationMisc::ClipboardCopy(*ParameterAction->Parameter.GetName().ToString());
+	}
+}
+
+void SNiagaraParameterMapView::RenameParameter(FNiagaraVariable& Parameter, FName NewName)
+{
+	if (ToolkitType == SCRIPT)
+	{
+		if (Graphs.Num() > 0)
 		{
-			ParameterAction->bNamespaceModifierRenamePending = true;
+			for (const TWeakObjectPtr<UNiagaraGraph>& Graph : Graphs)
+			{
+				if (Graph.IsValid())
+				{
+					Graph.Get()->RenameParameter(Parameter, NewName);
+				}
+			}
+		}
+	}
+	else if (ToolkitType == SYSTEM)
+	{
+		UNiagaraSystem* System = CachedSystem.Get();
+		if (System != nullptr)
+		{
+			System->Modify();
+			System->GetExposedParameters().RenameParameter(Parameter, NewName);
+			System->EditorOnlyAddedParameters.RenameParameter(Parameter, NewName);
 		}
 	}
 }
@@ -1456,7 +1762,30 @@ void SNiagaraAddParameterMenu::CollectAllActions(FGraphActionListBuilderBase& Ou
 	if (CanCollectSection(NiagaraParameterMapSectionID::ENGINE) && !IDsExcluded.Contains(NiagaraParameterMapSectionID::ENGINE))
 	{
 		const FText Category = NiagaraParameterMapSectionID::OnGetSectionTitle(NiagaraParameterMapSectionID::ENGINE);
+		
 		TArray<FNiagaraVariable> Variables = FNiagaraConstants::GetEngineConstants();
+		TArray<FName> EngineNamespaces;
+		NiagaraParameterMapSectionID::OnGetSectionNamespaces(NiagaraParameterMapSectionID::ENGINE, EngineNamespaces);
+		Variables.RemoveAll([EngineNamespaces](const FNiagaraVariable& Variable)
+		{
+			FNiagaraParameterHandle VariableHandle(Variable.GetName());
+			TArray<FName> VariableNameParts = VariableHandle.GetHandleParts();
+			if (VariableNameParts.Num() <= EngineNamespaces.Num())
+			{
+				return true;
+			}
+
+			for (int32 NamespaceIndex = 0; NamespaceIndex < EngineNamespaces.Num(); NamespaceIndex++)
+			{
+				if (VariableNameParts[NamespaceIndex] != EngineNamespaces[NamespaceIndex])
+				{
+					return true;
+				}
+			}
+
+			return false;
+		});
+
 		AddParameterGroup(OutAllActions,
 			Variables,
 			NiagaraParameterMapSectionID::ENGINE,
diff --git a/Engine/Plugins/FX/Niagara/Source/NiagaraEditor/Public/NiagaraEditorSettings.h b/Engine/Plugins/FX/Niagara/Source/NiagaraEditor/Public/NiagaraEditorSettings.h
--- a/Engine/Plugins/FX/Niagara/Source/NiagaraEditor/Public/NiagaraEditorSettings.h
+++ b/Engine/Plugins/FX/Niagara/Source/NiagaraEditor/Public/NiagaraEditorSettings.h
@@ -32,7 +32,8 @@ enum class ENiagaraNamespaceMetadataOptions
 {
 	Advanced,
 	PreventRenaming,
-	SupportsNamespaceModifier,
+	CanChangeNamespaceModifier,
+	PreventCreatingInSystemEditor
 };
 
 USTRUCT()
diff --git a/Engine/Plugins/FX/Niagara/Source/NiagaraEditor/Public/NiagaraGraph.h b/Engine/Plugins/FX/Niagara/Source/NiagaraEditor/Public/NiagaraGraph.h
--- a/Engine/Plugins/FX/Niagara/Source/NiagaraEditor/Public/NiagaraGraph.h
+++ b/Engine/Plugins/FX/Niagara/Source/NiagaraEditor/Public/NiagaraGraph.h
@@ -392,10 +392,6 @@ protected:
 	bool AppendCompileHash(FNiagaraCompileHashVisitor* InVisitor, const TArray<UNiagaraNode*>& InTraversal) const;
 
 private:
-
-	/** The const is a lie! Public implementation AddParameter interacts with mutable members only. Private implementation is provided to handle adding parameters in RefreshParameterReferences(). */
-	void MutableAddParameter(FNiagaraVariable& Parameter, const FAddParameterOptions Options) const;
-
 	virtual void NotifyGraphChanged(const FEdGraphEditAction& InAction) override;
 
 	/** Find parameters in the graph. */
diff --git a/Engine/Plugins/FX/Niagara/Source/NiagaraEditor/Public/ViewModels/Stack/NiagaraStackSystemSettingsGroup.h b/Engine/Plugins/FX/Niagara/Source/NiagaraEditor/Public/ViewModels/Stack/NiagaraStackSystemSettingsGroup.h
--- a/Engine/Plugins/FX/Niagara/Source/NiagaraEditor/Public/ViewModels/Stack/NiagaraStackSystemSettingsGroup.h
+++ b/Engine/Plugins/FX/Niagara/Source/NiagaraEditor/Public/ViewModels/Stack/NiagaraStackSystemSettingsGroup.h
@@ -6,6 +6,7 @@
 #include "NiagaraCommon.h"
 #include "NiagaraParameterStore.h"
 #include "ViewModels/Stack/NiagaraStackParameterStoreEntry.h"
+#include "Delegates/IDelegateInstance.h"
 #include "NiagaraStackSystemSettingsGroup.generated.h"
 
 class FNiagaraScriptViewModel;
@@ -45,11 +46,14 @@ public:
 
 protected:
 	virtual void RefreshChildrenInternal(const TArray<UNiagaraStackEntry*>& CurrentChildren, TArray<UNiagaraStackEntry*>& NewChildren, TArray<FStackIssue>& NewIssues) override;
+	virtual void FinalizeInternal() override;
 
 private:
+	void ParameterStoreChanged();
 	void ParameterDeleted();
 
 private:
 	TWeakObjectPtr<UObject> Owner;
 	FNiagaraParameterStore* ParameterStore;
+	FDelegateHandle ParameterStoreChangedHandle;
 };
\ No newline at end of file
diff --git a/Engine/Plugins/FX/Niagara/Source/NiagaraEditor/Public/Widgets/SNiagaraParameterMapView.h b/Engine/Plugins/FX/Niagara/Source/NiagaraEditor/Public/Widgets/SNiagaraParameterMapView.h
--- a/Engine/Plugins/FX/Niagara/Source/NiagaraEditor/Public/Widgets/SNiagaraParameterMapView.h
+++ b/Engine/Plugins/FX/Niagara/Source/NiagaraEditor/Public/Widgets/SNiagaraParameterMapView.h
@@ -25,6 +25,8 @@ class UNiagaraGraph;
 struct FEdGraphSchemaAction;
 struct FSlateBrush;
 class UNiagaraSystem;
+struct FNiagaraNamespaceMetadata;
+class IToolTip;
 
 /* Enums to use when grouping the blueprint members in the list panel. The order here will determine the order in the list */
 namespace NiagaraParameterMapSectionID
@@ -58,22 +60,6 @@ namespace NiagaraParameterMapSectionID
 	static bool GetSectionIsAdvanced(const NiagaraParameterMapSectionID::Type InSection);
 };
 
-class FNiagaraParameterMapViewCommands : public TCommands<FNiagaraParameterMapViewCommands>
-{
-public:
-	/** Constructor */
-	FNiagaraParameterMapViewCommands()
-		: TCommands<FNiagaraParameterMapViewCommands>(TEXT("NiagaraParameterMapViewCommands"), NSLOCTEXT("Contexts", "NiagaraParameterMap", "NiagaraParameterMap"), NAME_None, FEditorStyle::GetStyleSetName())
-	{
-	}
-
-	// Basic operations
-	TSharedPtr<FUICommandInfo> DeleteEntry;
-
-	/** Initialize commands */
-	virtual void RegisterCommands() override;
-};
-
 /** A widget for viewing and editing a set of selected objects with a details panel. */
 class SNiagaraParameterMapView : public SCompoundWidget
 {
@@ -131,6 +117,7 @@ private:
 	void OnActionDoubleClicked(const TArray<TSharedPtr<FEdGraphSchemaAction>>& InActions);
 	TSharedPtr<SWidget> OnContextMenuOpening();
 	FText OnGetSectionTitle(int32 InSectionID);
+	TSharedPtr<IToolTip> OnGetSectionToolTip(int32 InSectionID);
 	TSharedRef<SWidget> OnGetSectionWidget(TSharedRef<SWidget> RowWidget, int32 InSectionID);
 	TSharedRef<SWidget> CreateAddToSectionButton(const NiagaraParameterMapSectionID::Type InSection, TWeakPtr<SWidget> WeakRowWidget, FText AddNewText, FName MetaDataTag);
 	
@@ -148,19 +135,47 @@ private:
 	void AddGraph(UNiagaraGraph* Graph);
 	void AddGraph(class UNiagaraScriptSourceBase* SourceBase);
 	void OnGraphChanged(const struct FEdGraphEditAction& InAction);
-	void OnUserParameterStoreChanged();
+	void OnSystemParameterStoreChanged();
 
 	//Callbacks
+	FText GetDeleteEntryToolTip() const;
 	void OnDeleteEntry();
 	bool CanDeleteEntry() const;
+	FText GetRenameOnActionNodeToolTip() const;
 	void OnRequestRenameOnActionNode();
-	bool CanRequestRenameOnActionNode(TWeakPtr<struct FGraphActionNode> InSelectedNode) const;
 	bool CanRequestRenameOnActionNode() const;
 	void OnPostRenameActionNode(const FText& InText, FNiagaraParameterAction& InAction);
+
+	bool GetNamespaceEditDataForSelection(
+		TSharedPtr<FNiagaraParameterAction>& OutParameterAction,
+		FNiagaraParameterHandle& OutParameterHandle,
+		FNiagaraNamespaceMetadata& OutNamespaceMetadata,
+		FText& OutErrorMessage) const;
+
+	bool GetNamespaceModifierEditDataForSelection(
+		TSharedPtr<FNiagaraParameterAction>& OutParameterAction,
+		FNiagaraParameterHandle& OutParameterHandle,
+		FNiagaraNamespaceMetadata& OutNamespaceMetadata,
+		FText& OutErrorMessage) const;
+
+	FText GetAddNamespaceModifierToolTip() const;
+	bool CanAddNamespaceModifier() const;
 	void OnAddNamespaceModifier();
+
+	FText GetRemoveNamespaceModifierToolTip() const;
+	bool CanRemoveNamespaceModifier() const;
 	void OnRemoveNamespaceModifier();
+
+	FText GetEditNamespaceModifierToolTip() const;
+	bool CanEditNamespaceModifier() const;
 	void OnEditNamespaceModifier();
 
+	FText GetCopyParameterReferenceToolTip() const;
+	bool CanCopyParameterReference() const;
+	void OnCopyParameterReference();
+
+	void RenameParameter(FNiagaraVariable& Parameter, FName NewName);
+
 	bool IsSystemToolkit();
 	bool IsScriptToolkit();
 
@@ -188,6 +203,7 @@ private:
 	TArray<TWeakObjectPtr<UNiagaraGraph>> Graphs;
 	TWeakObjectPtr<UNiagaraSystem> CachedSystem;
 	FDelegateHandle UserParameterStoreChangedHandle;
+	FDelegateHandle AddedParameterStoreChangedHandle;
 
 	/** The handle to the graph changed delegate. */
 	TArray<FDelegateHandle> OnGraphChangedHandles;
