Niagara - Fix an issue where renaming scratch pad scripts weren't triggering a recompile.

diff --git a/Engine/Plugins/FX/Niagara/Source/NiagaraEditor/Private/ViewModels/NiagaraScratchPadScriptViewModel.cpp b/Engine/Plugins/FX/Niagara/Source/NiagaraEditor/Private/ViewModels/NiagaraScratchPadScriptViewModel.cpp
--- a/Engine/Plugins/FX/Niagara/Source/NiagaraEditor/Private/ViewModels/NiagaraScratchPadScriptViewModel.cpp
+++ b/Engine/Plugins/FX/Niagara/Source/NiagaraEditor/Private/ViewModels/NiagaraScratchPadScriptViewModel.cpp
@@ -130,6 +130,7 @@ void FNiagaraScratchPadScriptViewModel::SetScriptName(FText InScriptName)
 			if (System != nullptr)
 			{
 				FNiagaraStackGraphUtilities::RenameReferencingParameters(*System, Emitter, *ReferencingFunctionCallNode, OldFunctionName, NewFunctionName);
+				ReferencingFunctionCallNode->MarkNodeRequiresSynchronization(TEXT("ScratchPad script renamed"), true);
 			}
 		}
 
