Niagara - Fix an issue with the parameters panel where a parameter added by the user would be removed after the last graph reference was removed.

diff --git a/Engine/Plugins/FX/Niagara/Source/NiagaraEditor/Private/NiagaraNodeWithDynamicPins.cpp b/Engine/Plugins/FX/Niagara/Source/NiagaraEditor/Private/NiagaraNodeWithDynamicPins.cpp
--- a/Engine/Plugins/FX/Niagara/Source/NiagaraEditor/Private/NiagaraNodeWithDynamicPins.cpp
+++ b/Engine/Plugins/FX/Niagara/Source/NiagaraEditor/Private/NiagaraNodeWithDynamicPins.cpp
@@ -358,13 +358,23 @@ void UNiagaraNodeWithDynamicPins::RemoveDynamicPin(UEdGraphPin* Pin)
 		{
 			const UEdGraphSchema_Niagara* Schema = GetDefault<UEdGraphSchema_Niagara>();
 			FNiagaraVariable PinVariable = Schema->PinToNiagaraVariable(Pin, false);
-			UNiagaraScriptVariable** PinAssociatedScriptVariable = Graph->GetAllMetaData().Find(PinVariable);
-			if (PinAssociatedScriptVariable != nullptr)
+			if (PinVariable.IsValid())
 			{
-				bool bShouldRemoveScriptVariable = !Graph->UpdateUsageForScriptVariable(*PinAssociatedScriptVariable);
-				if (bShouldRemoveScriptVariable)
+				const FNiagaraGraphParameterReferenceCollection* ReferenceCollection = Graph->GetParameterReferenceMap().Find(PinVariable);
+				if (ReferenceCollection != nullptr && ReferenceCollection->WasCreated())
 				{
-					Graph->RemoveParameter((*PinAssociatedScriptVariable)->Variable);
+					// Don't remove parameters from the graph which were created by the user.
+					return;
+				}
+
+				UNiagaraScriptVariable** PinAssociatedScriptVariable = Graph->GetAllMetaData().Find(PinVariable);
+				if (PinAssociatedScriptVariable != nullptr)
+				{
+					bool bShouldRemoveScriptVariable = !Graph->UpdateUsageForScriptVariable(*PinAssociatedScriptVariable);
+					if (bShouldRemoveScriptVariable)
+					{
+						Graph->RemoveParameter((*PinAssociatedScriptVariable)->Variable);
+					}
 				}
 			}
 		}
