UE4.25 - Deprecate Begin/EndComputePass

diff --git a/Engine/Plugins/FX/Niagara/Source/Niagara/Private/NiagaraEmitterInstanceBatcher.cpp b/Engine/Plugins/FX/Niagara/Source/Niagara/Private/NiagaraEmitterInstanceBatcher.cpp
--- a/Engine/Plugins/FX/Niagara/Source/Niagara/Private/NiagaraEmitterInstanceBatcher.cpp
+++ b/Engine/Plugins/FX/Niagara/Source/Niagara/Private/NiagaraEmitterInstanceBatcher.cpp
@@ -846,7 +846,9 @@ void NiagaraEmitterInstanceBatcher::ExecuteAll(FRHICommandList& RHICmdList, FRHI
 	// Note: We can not encapsulate the whole Niagara pass as some DI's may not be compatable (i.e. use CopyTexture function), we need to fix this with future RDG conversion
 	if (SimPasses.Num() > 0)
 	{
+		PRAGMA_DISABLE_DEPRECATION_WARNINGS
 		RHICmdList.BeginComputePass(TEXT("NiagaraCompute"));
+		PRAGMA_ENABLE_DEPRECATION_WARNINGS
 		RHICmdList.BeginUAVOverlap();
 
 		FEmitterInstanceList InstancesWithPersistentIDs;
@@ -899,7 +901,9 @@ void NiagaraEmitterInstanceBatcher::ExecuteAll(FRHICommandList& RHICmdList, FRHI
 		RHICmdList.TransitionResources(EResourceTransitionAccess::EReadable, EResourceTransitionPipeline::EComputeToGfx, OutputGraphicsBuffers.GetData(), OutputGraphicsBuffers.Num());
 
 		RHICmdList.EndUAVOverlap();
+		PRAGMA_DISABLE_DEPRECATION_WARNINGS
 		RHICmdList.EndComputePass();
+		PRAGMA_ENABLE_DEPRECATION_WARNINGS
 	}
 }
 
diff --git a/Engine/Source/Runtime/RHI/Public/RHICommandList.h b/Engine/Source/Runtime/RHI/Public/RHICommandList.h
--- a/Engine/Source/Runtime/RHI/Public/RHICommandList.h
+++ b/Engine/Source/Runtime/RHI/Public/RHICommandList.h
@@ -3523,6 +3523,7 @@ public:
 		IncrementSubpass();
 	}
 
+	UE_DEPRECATED(4.25, "BeginComputePass API is deprecated. Use SetComputeShader() instead.")
 	FORCEINLINE_DEBUGGABLE void BeginComputePass(const TCHAR* Name)
 	{
 		check(!IsInsideRenderPass());
@@ -3540,6 +3541,7 @@ public:
 		Data.bInsideComputePass = true;
 	}
 
+	UE_DEPRECATED(4.25, "EndComputePass API is deprecated. BeginRenderPass() or SetComputeShader() instead.")
 	void EndComputePass()
 	{
 		check(IsInsideComputePass());
