Do not reply to mouse button down events as handled when receiving a hit on niagara pins to allow the SNodePanel to handle the event and zero the mouse delta.
This fixes the pin context menu failing to open on right click due to the pending mouse delta causing the event to be handled as a drag/drop.

diff --git a/Engine/Plugins/FX/Niagara/Source/NiagaraEditor/Private/Widgets/SNiagaraGraphParameterMapGetNode.cpp b/Engine/Plugins/FX/Niagara/Source/NiagaraEditor/Private/Widgets/SNiagaraGraphParameterMapGetNode.cpp
--- a/Engine/Plugins/FX/Niagara/Source/NiagaraEditor/Private/Widgets/SNiagaraGraphParameterMapGetNode.cpp
+++ b/Engine/Plugins/FX/Niagara/Source/NiagaraEditor/Private/Widgets/SNiagaraGraphParameterMapGetNode.cpp
@@ -195,7 +195,7 @@ FReply SNiagaraGraphParameterMapGetNode::OnBorderMouseButtonDown(const FGeometry
 		}
 
 	}
-	return FReply::Handled();
+	return FReply::Unhandled();
 }
 
 FReply SNiagaraGraphParameterMapGetNode::OnDroppedOnTarget(TSharedPtr<FDragDropOperation> DropOperation)
