Niagara - Fix additional issues with the parameter workflow.
+ Add short and long options for namespace display names.  This allows us to show NPC on the parameter widgets, Niagara Parameter Collection, in the parameters panel header.
+ Fix an issue where user parameters weren't set to rename on add.
+ Fix an issue where where new names for parameters were "New Variable" instead of using the type like "New Float".
+ Fix an issue where you could not make new parameters from the input menu in event and simulation stage scripts, also add Transient as an option for making new values.
+ Fix an issue where you could try to create expressions and scratch dynamic inputs for data interface and object types which could crash.
+ Fix a crash when working with data interfaces in the parameters panel.

diff --git a/Engine/Plugins/FX/Niagara/Source/Niagara/Private/NiagaraSystem.cpp b/Engine/Plugins/FX/Niagara/Source/Niagara/Private/NiagaraSystem.cpp
--- a/Engine/Plugins/FX/Niagara/Source/Niagara/Private/NiagaraSystem.cpp
+++ b/Engine/Plugins/FX/Niagara/Source/Niagara/Private/NiagaraSystem.cpp
@@ -59,6 +59,9 @@ UNiagaraSystem::UNiagaraSystem(const FObjectInitializer& ObjectInitializer)
 , bHasSystemScriptDIsWithPerInstanceData(false)
 {
 	ExposedParameters.SetOwner(this);
+#if WITH_EDITORONLY_DATA
+	EditorOnlyAddedParameters.SetOwner(this);
+#endif
 	MaxPoolSize = 32;
 
 	EffectType = nullptr;
diff --git a/Engine/Plugins/FX/Niagara/Source/NiagaraEditor/Private/NiagaraEditorSettings.cpp b/Engine/Plugins/FX/Niagara/Source/NiagaraEditor/Private/NiagaraEditorSettings.cpp
--- a/Engine/Plugins/FX/Niagara/Source/NiagaraEditor/Private/NiagaraEditorSettings.cpp
+++ b/Engine/Plugins/FX/Niagara/Source/NiagaraEditor/Private/NiagaraEditorSettings.cpp
@@ -8,12 +8,8 @@ FNiagaraNamespaceMetadata::FNiagaraNamespaceMetadata()
 {
 }
 
-FNiagaraNamespaceMetadata::FNiagaraNamespaceMetadata(TArray<FName> InNamespaces, FText InDisplayName, FText InDescription, FLinearColor InBackgroundColor, TArray<ENiagaraNamespaceMetadataOptions> InOptions)
+FNiagaraNamespaceMetadata::FNiagaraNamespaceMetadata(TArray<FName> InNamespaces)
 	: Namespaces(InNamespaces)
-	, DisplayName(InDisplayName)
-	, Description(InDescription)
-	, BackgroundColor(InBackgroundColor)
-	, Options(InOptions)
 {
 }
 
@@ -28,125 +24,108 @@ UNiagaraEditorSettings::UNiagaraEditorSettings(const FObjectInitializer& ObjectI
 	SetupNamespaceMetadata();
 }
 
+#define LOCTEXT_NAMESPACE "NamespaceMetadata"
+
 void UNiagaraEditorSettings::SetupNamespaceMetadata()
 {
 	NamespaceMetadata = 
 	{
-		FNiagaraNamespaceMetadata(
-			{NAME_None},
-			NSLOCTEXT("NamespaceMetadata", "DefaultDisplayName", "None"),
-			NSLOCTEXT("NamespaceMetadata", "DefaultDescription", "Arbitrary sub-namespace for specifying module specific dataset attributes, or calling nested modules."),
-			FLinearColor(FColor(102, 102, 102)),
-			{ }),
-		FNiagaraNamespaceMetadata(
-			{FNiagaraConstants::SystemNamespace},
-			NSLOCTEXT("NamespaceMetadata", "SystemDisplayName", "System"),
-			NSLOCTEXT("NamespaceMetadata", "SystemDescription", "Persistent attribute in the system which is written in a system\n stage and can be read anywhere."),
-			FLinearColor(FColor(49, 113, 142)),
-			{ENiagaraNamespaceMetadataOptions::CanChangeNamespaceModifier}),
-		FNiagaraNamespaceMetadata(
-			{FNiagaraConstants::EmitterNamespace},
-			NSLOCTEXT("NamespaceMetadata", "EmitterDisplayName", "Emitter"),
-			NSLOCTEXT("NamespaceMetadata", "EmitterDescription", "Persistent attribute which is written in a emitter\n stage and can be read in emitter and particle stages."),
-			FLinearColor(FColor(145, 99, 56)),
-			{ENiagaraNamespaceMetadataOptions::CanChangeNamespaceModifier}),
-		FNiagaraNamespaceMetadata(
-			{FNiagaraConstants::ParticleAttributeNamespace},
-			NSLOCTEXT("NamespaceMetadata", "ParticleDisplayName", "Particles"),
-			NSLOCTEXT("NamespaceMetadata", "ParticleDescription", "Persistent attribute which is written in a particle\n stage and can be read in particle stages."),
-			FLinearColor(FColor(72, 130, 71)),
-			{ENiagaraNamespaceMetadataOptions::CanChangeNamespaceModifier}),
-		FNiagaraNamespaceMetadata(
-			{FNiagaraConstants::ModuleNamespace},
-			NSLOCTEXT("NamespaceMetadata", "ModuleDisplayName", "Input"),
-			NSLOCTEXT("NamespaceMetadata", "ModuleDescription", "A transient value which exposes a module input to the system and emitter editor."),
-			FLinearColor(FColor(136, 66, 65)),
-			{ENiagaraNamespaceMetadataOptions::CanChangeNamespaceModifier}),
-		FNiagaraNamespaceMetadata(
-			{FNiagaraConstants::OutputNamespace, FNiagaraConstants::ModuleNamespace},
-			NSLOCTEXT("NamespaceMetadata", "ModuleOutputDisplayName", "Output"),
-			NSLOCTEXT("NamespaceMetadata", "ModuleOutputDescription", "A transient value which is unique to the module that wrote it and can be read from any other module.\nTransient values do not persist from frame to frame, or between stages, e.g. emitter to particle, or spawn to update."),
-			FLinearColor(FColor(109, 95, 124)),
-			{ENiagaraNamespaceMetadataOptions::Advanced, ENiagaraNamespaceMetadataOptions::PreventCreatingInSystemEditor}),
-		FNiagaraNamespaceMetadata(
-			{FNiagaraConstants::OutputNamespace},
-			NSLOCTEXT("NamespaceMetadata", "OutputDisplayName", "Output"),
-			NSLOCTEXT("NamespaceMetadata", "OutputDescription", "A transient value which is unique to the module that wrote it and can be read from any other module.\nTransient values do not persist from frame to frame, or between stages, e.g. emitter to particle, or spawn to update."),
-			FLinearColor(FColor(109, 95, 124)),
-			{ENiagaraNamespaceMetadataOptions::Advanced}),
-		FNiagaraNamespaceMetadata(
-			{FNiagaraConstants::LocalNamespace, FNiagaraConstants::ModuleNamespace},
-			NSLOCTEXT("NamespaceMetadata", "ModuleLocalDisplayName", "Local"),
-			NSLOCTEXT("NamespaceMetadata", "ModuleLocalDescription", "A transient value which can be written to and read from within a single module.\nTransient values do not persist from frame to frame, or between stages, e.g. emitter to particle, or spawn to update."),
-			FLinearColor(FColor(174, 157, 40)),
-			{ENiagaraNamespaceMetadataOptions::Advanced}),
-		FNiagaraNamespaceMetadata(
-			{FNiagaraConstants::TransientNamespace},
-			NSLOCTEXT("NamespaceMetadata", "TransientDisplayName", "Transient"),
-			NSLOCTEXT("NamespaceMetadata", "TransientDescription", "A transient value which can be written to and read from from any module.\nTransient values do not persist from frame to frame, or between stages, e.g. emitter to particle, or spawn to update."),
-			FLinearColor(FColor(109, 95, 124)),
-			{ENiagaraNamespaceMetadataOptions::Advanced}),
-		FNiagaraNamespaceMetadata(
-			{FNiagaraConstants::EngineNamespace},
-			NSLOCTEXT("NamespaceMetadata", "EngineDisplayName", "Engine"),
-			NSLOCTEXT("NamespaceMetadata", "EngineDescription", "A read only value which is provided by the engine.\nThis value's source can be the simulation itsef\ne.g. ExecutionCount, or the owner of the simulation e.g. (Owner) Scale."),
-			FLinearColor(FColor(45, 45, 45)),
-			{ENiagaraNamespaceMetadataOptions::PreventRenaming}),
-		FNiagaraNamespaceMetadata(
-			{FNiagaraConstants::UserNamespace},
-			NSLOCTEXT("NamespaceMetadata", "UserDisplayName", "User"),
-			NSLOCTEXT("NamespaceMetadata", "UserDescription", "A read only value which can be initialized per system and\nmodified externally in the level, by blueprint, or by c++."),
-			FLinearColor(FColor(91, 161, 194)),
-			{ }),
-		FNiagaraNamespaceMetadata(
-			{FNiagaraConstants::ParameterCollectionNamespace},
-			NSLOCTEXT("NamespaceMetadata", "NiagaraParameterCollectionDisplayName", "Niagara Parameter Collection"),
-			NSLOCTEXT("NamespaceMetadata", "NiagaraParameterCollectionDescription", "Values read from a niagara parameter collection asset.\nRead only in a niagara system."),
-			FLinearColor(FColor(45, 45, 45)),
-			{ENiagaraNamespaceMetadataOptions::Advanced, ENiagaraNamespaceMetadataOptions::PreventRenaming}),
-		FNiagaraNamespaceMetadata(
-			{FNiagaraConstants::DataInstanceNamespace},
-			NSLOCTEXT("NamespaceMetadata", "DataInstanceDisplayName", "Data Instance"),
-			NSLOCTEXT("NamespaceMetadata", "DataInstanceDescription", "A special transient value which has a single bool IsAlive value, which determines if a particle is alive or not."),
-			FLinearColor(FColor(45, 45, 45)),
-			{ENiagaraNamespaceMetadataOptions::Advanced, ENiagaraNamespaceMetadataOptions::PreventRenaming}),
+		FNiagaraNamespaceMetadata({NAME_None})
+			.SetDisplayName(LOCTEXT("DefaultDisplayName", "None"))
+			.SetDescription(LOCTEXT("DefaultDescription", "Arbitrary sub-namespace for specifying module specific dataset attributes, or calling nested modules."))
+			.SetBackgroundColor(FLinearColor(FColor(102, 102, 102))),
+		FNiagaraNamespaceMetadata({FNiagaraConstants::SystemNamespace})
+			.SetDisplayName(LOCTEXT("SystemDisplayName", "System"))
+			.SetDescription(LOCTEXT("SystemDescription", "Persistent attribute in the system which is written in a system\n stage and can be read anywhere."))
+			.SetBackgroundColor(FLinearColor(FColor(49, 113, 142)))
+			.AddOption(ENiagaraNamespaceMetadataOptions::CanChangeNamespaceModifier),
+		FNiagaraNamespaceMetadata({FNiagaraConstants::EmitterNamespace})
+			.SetDisplayName(LOCTEXT("EmitterDisplayName", "Emitter"))
+			.SetDescription(LOCTEXT("EmitterDescription", "Persistent attribute which is written in a emitter\n stage and can be read in emitter and particle stages."))
+			.SetBackgroundColor(FLinearColor(FColor(145, 99, 56)))
+			.AddOption(ENiagaraNamespaceMetadataOptions::CanChangeNamespaceModifier),
+		FNiagaraNamespaceMetadata({FNiagaraConstants::ParticleAttributeNamespace})
+			.SetDisplayName(LOCTEXT("ParticleDisplayName", "Particles"))
+			.SetDescription(LOCTEXT("ParticleDescription", "Persistent attribute which is written in a particle\n stage and can be read in particle stages."))
+			.SetBackgroundColor(FLinearColor(FColor(72, 130, 71)))
+			.AddOption(ENiagaraNamespaceMetadataOptions::CanChangeNamespaceModifier),
+		FNiagaraNamespaceMetadata({FNiagaraConstants::ModuleNamespace})
+			.SetDisplayName(LOCTEXT("ModuleDisplayName", "Input"))
+			.SetDescription(LOCTEXT("ModuleDescription", "A transient value which exposes a module input to the system and emitter editor."))
+			.SetBackgroundColor(FLinearColor(FColor(136, 66, 65)))
+			.AddOption(ENiagaraNamespaceMetadataOptions::CanChangeNamespaceModifier),
+		FNiagaraNamespaceMetadata({FNiagaraConstants::OutputNamespace, FNiagaraConstants::ModuleNamespace})
+			.SetDisplayName(LOCTEXT("ModuleOutputDisplayName", "Output"))
+			.SetDescription(LOCTEXT("ModuleOutputDescription", "A transient value which is unique to the module that wrote it and can be read from any other module.\nTransient values do not persist from frame to frame, or between stages, e.g. emitter to particle, or spawn to update."))
+			.SetBackgroundColor(FLinearColor(FColor(109, 95, 124)))
+			.AddOption(ENiagaraNamespaceMetadataOptions::Advanced)
+			.AddOption(ENiagaraNamespaceMetadataOptions::PreventCreatingInSystemEditor),
+		FNiagaraNamespaceMetadata({FNiagaraConstants::OutputNamespace})
+			.SetDisplayName(LOCTEXT("OutputDisplayName", "Output"))
+			.SetDescription(LOCTEXT("OutputDescription", "A transient value which is unique to the module that wrote it and can be read from any other module.\nTransient values do not persist from frame to frame, or between stages, e.g. emitter to particle, or spawn to update."))
+			.SetBackgroundColor(FLinearColor(FColor(109, 95, 124)))
+			.AddOption({ENiagaraNamespaceMetadataOptions::Advanced}),
+		FNiagaraNamespaceMetadata({FNiagaraConstants::LocalNamespace, FNiagaraConstants::ModuleNamespace})
+			.SetDisplayName(LOCTEXT("ModuleLocalDisplayName", "Local"))
+			.SetDescription(LOCTEXT("ModuleLocalDescription", "A transient value which can be written to and read from within a single module.\nTransient values do not persist from frame to frame, or between stages, e.g. emitter to particle, or spawn to update."))
+			.SetBackgroundColor(FLinearColor(FColor(174, 157, 40)))
+			.AddOption(ENiagaraNamespaceMetadataOptions::Advanced),
+		FNiagaraNamespaceMetadata({FNiagaraConstants::TransientNamespace})
+			.SetDisplayName(LOCTEXT("TransientDisplayName", "Transient"))
+			.SetDescription(LOCTEXT("TransientDescription", "A transient value which can be written to and read from from any module.\nTransient values do not persist from frame to frame, or between stages, e.g. emitter to particle, or spawn to update."))
+			.SetBackgroundColor(FLinearColor(FColor(109, 95, 124)))
+			.AddOption(ENiagaraNamespaceMetadataOptions::Advanced),
+		FNiagaraNamespaceMetadata({FNiagaraConstants::EngineNamespace})
+			.SetDisplayName(LOCTEXT("EngineDisplayName", "Engine"))
+			.SetDescription(LOCTEXT("EngineDescription", "A read only value which is provided by the engine.\nThis value's source can be the simulation itsef\ne.g. ExecutionCount, or the owner of the simulation e.g. (Owner) Scale."))
+			.SetBackgroundColor(FLinearColor(FColor(45, 45, 45)))
+			.AddOption(ENiagaraNamespaceMetadataOptions::PreventRenaming),
+		FNiagaraNamespaceMetadata({FNiagaraConstants::UserNamespace})
+			.SetDisplayName(LOCTEXT("UserDisplayName", "User"))
+			.SetDescription(LOCTEXT("UserDescription", "A read only value which can be initialized per system and\nmodified externally in the level, by blueprint, or by c++."))
+			.SetBackgroundColor(FLinearColor(FColor(91, 161, 194))),
+		FNiagaraNamespaceMetadata({FNiagaraConstants::ParameterCollectionNamespace})
+			.SetDisplayName(LOCTEXT("NiagaraParameterCollectionDisplayName", "NPC"))
+			.SetDisplayNameLong(LOCTEXT("NiagaraParameterCollectionDisplayNameLong", "Niagara Parameter Collection"))
+			.SetDescription(LOCTEXT("NiagaraParameterCollectionDescription", "Values read from a niagara parameter collection asset.\nRead only in a niagara system."))
+			.SetBackgroundColor(FLinearColor(FColor(45, 45, 45)))
+			.AddOption(ENiagaraNamespaceMetadataOptions::Advanced)
+			.AddOption(ENiagaraNamespaceMetadataOptions::PreventRenaming),
+		FNiagaraNamespaceMetadata({FNiagaraConstants::DataInstanceNamespace})
+			.SetDisplayName(LOCTEXT("DataInstanceDisplayName", "Data Instance"))
+			.SetDescription(LOCTEXT("DataInstanceDescription", "A special transient value which has a single bool IsAlive value, which determines if a particle is alive or not."))
+			.SetBackgroundColor(FLinearColor(FColor(45, 45, 45)))
+			.AddOption(ENiagaraNamespaceMetadataOptions::Advanced)
+			.AddOption(ENiagaraNamespaceMetadataOptions::PreventRenaming),
 	};
 
 	NamespaceModifierMetadata =
 	{
-		FNiagaraNamespaceMetadata(
-			{FNiagaraConstants::InitialNamespace},
-			NSLOCTEXT("NamespaceMetadata", "InitialModifierDisplayName", "Initial"),
-			NSLOCTEXT("NamespaceMetadata", "InitialModifierDescription", "A namespace modifier for dataset attributes which when used in\na linked input in an update script will get the initial value from the spawn script."),
-			FLinearColor(FColor(102, 102, 102)),
-			{ }),
-		FNiagaraNamespaceMetadata(
-			{FNiagaraConstants::ModuleNamespace},
-			NSLOCTEXT("NamespaceMetadata", "ModuleModifierDisplayName", "Module"),
-			NSLOCTEXT("NamespaceMetadata", "ModuleModifierDescription", "A namespace modifier which is used when writing dataset attributes\nin module which will make that attribute unique to the module instance when used."),
-			FLinearColor(FColor(102, 102, 102)),
-			{ }),
-		FNiagaraNamespaceMetadata(
-			{FNiagaraConstants::SystemNamespace},
-			NSLOCTEXT("NamespaceMetadata", "SystemModifierDisplayName", "System"),
-			NSLOCTEXT("NamespaceMetadata", "SystemModifierDescription", "A namespace modifier which specifies that an engine provided parameter comes from the system."),
-			FLinearColor(FColor(49, 113, 142)),
-			{ }),
-		FNiagaraNamespaceMetadata(
-			{FNiagaraConstants::EmitterNamespace},
-			NSLOCTEXT("NamespaceMetadata", "EmitterModifierDisplayName", "Emitter"),
-			NSLOCTEXT("NamespaceMetadata", "EmitterModifierDescription", "A namespace modifier which specifies that an engine provided parameter comes from the emitter."),
-			FLinearColor(FColor(145, 99, 56)),
-			{ }),
-		FNiagaraNamespaceMetadata(
-			{FNiagaraConstants::OwnerNamespace},
-			NSLOCTEXT("NamespaceMetadata", "OwnerDisplayName", "Owner"),
-			NSLOCTEXT("NamespaceMetadata", "OwnerDescription", "A namespace modifier which specifies that an engine provided parameter comes from the owner, or component."),
-			FLinearColor(FColor(45, 45, 45)),
-			{ }),
+		FNiagaraNamespaceMetadata({FNiagaraConstants::InitialNamespace})
+			.SetDisplayName(LOCTEXT("InitialModifierDisplayName", "Initial"))
+			.SetDescription(LOCTEXT("InitialModifierDescription", "A namespace modifier for dataset attributes which when used in\na linked input in an update script will get the initial value from the spawn script."))
+			.SetBackgroundColor(FLinearColor(FColor(102, 102, 102))),
+		FNiagaraNamespaceMetadata({FNiagaraConstants::ModuleNamespace})
+			.SetDisplayName(LOCTEXT("ModuleModifierDisplayName", "Module"))
+			.SetDescription(LOCTEXT("ModuleModifierDescription", "A namespace modifier which is used when writing dataset attributes\nin module which will make that attribute unique to the module instance when used."))
+			.SetBackgroundColor(FLinearColor(FColor(102, 102, 102))),
+		FNiagaraNamespaceMetadata({FNiagaraConstants::SystemNamespace})
+			.SetDisplayName(LOCTEXT("SystemModifierDisplayName", "System"))
+			.SetDescription(LOCTEXT("SystemModifierDescription", "A namespace modifier which specifies that an engine provided parameter comes from the system."))
+			.SetBackgroundColor(FLinearColor(FColor(49, 113, 142))),
+		FNiagaraNamespaceMetadata({FNiagaraConstants::EmitterNamespace})
+			.SetDisplayName(LOCTEXT("EmitterModifierDisplayName", "Emitter"))
+			.SetDescription(LOCTEXT("EmitterModifierDescription", "A namespace modifier which specifies that an engine provided parameter comes from the emitter."))
+			.SetBackgroundColor(FLinearColor(FColor(145, 99, 56))),
+		FNiagaraNamespaceMetadata({FNiagaraConstants::OwnerNamespace})
+			.SetDisplayName(LOCTEXT("OwnerDisplayName", "Owner"))
+			.SetDescription(LOCTEXT("OwnerDescription", "A namespace modifier which specifies that an engine provided parameter comes from the owner, or component."))
+			.SetBackgroundColor(FLinearColor(FColor(45, 45, 45))),
 	};
 }
 
+#undef LOCTEXT_NAMESPACE
+
 bool UNiagaraEditorSettings::GetAutoCompile() const
 {
 	return bAutoCompile;
diff --git a/Engine/Plugins/FX/Niagara/Source/NiagaraEditor/Private/NiagaraEditorUtilities.cpp b/Engine/Plugins/FX/Niagara/Source/NiagaraEditor/Private/NiagaraEditorUtilities.cpp
--- a/Engine/Plugins/FX/Niagara/Source/NiagaraEditor/Private/NiagaraEditorUtilities.cpp
+++ b/Engine/Plugins/FX/Niagara/Source/NiagaraEditor/Private/NiagaraEditorUtilities.cpp
@@ -1656,7 +1656,7 @@ bool FNiagaraEditorUtilities::VerifyNameChangeForInputOrOutputNode(const UNiagar
 	return true;
 }
 
-bool FNiagaraEditorUtilities::AddParameter(FNiagaraVariable& NewParameterVariable, FNiagaraParameterStore& TargetParameterStore, UObject& ParameterStoreOwner, UNiagaraStackEditorData& StackEditorData)
+bool FNiagaraEditorUtilities::AddParameter(FNiagaraVariable& NewParameterVariable, FNiagaraParameterStore& TargetParameterStore, UObject& ParameterStoreOwner, UNiagaraStackEditorData* StackEditorData)
 {
 	FScopedTransaction AddTransaction(LOCTEXT("AddParameter", "Add Parameter"));
 	ParameterStoreOwner.Modify();
@@ -1673,9 +1673,9 @@ bool FNiagaraEditorUtilities::AddParameter(FNiagaraVariable& NewParameterVariabl
 	NewParameterVariable.SetName(FNiagaraUtilities::GetUniqueName(NewParameterVariable.GetName(), ExistingParameterStoreNames));
 
 	bool bSuccess = TargetParameterStore.AddParameter(NewParameterVariable);
-	if (bSuccess)
+	if (bSuccess && StackEditorData != nullptr)
 	{
-		StackEditorData.SetStackEntryIsRenamePending(NewParameterVariable.GetName().ToString(), true);
+		StackEditorData->SetStackEntryIsRenamePending(NewParameterVariable.GetName().ToString(), true);
 	}
 	return bSuccess;
 }
diff --git a/Engine/Plugins/FX/Niagara/Source/NiagaraEditor/Private/ViewModels/NiagaraParameterPanelViewModel.cpp b/Engine/Plugins/FX/Niagara/Source/NiagaraEditor/Private/ViewModels/NiagaraParameterPanelViewModel.cpp
--- a/Engine/Plugins/FX/Niagara/Source/NiagaraEditor/Private/ViewModels/NiagaraParameterPanelViewModel.cpp
+++ b/Engine/Plugins/FX/Niagara/Source/NiagaraEditor/Private/ViewModels/NiagaraParameterPanelViewModel.cpp
@@ -125,7 +125,7 @@ const UNiagaraScriptVariable* FNiagaraSystemToolkitParameterPanelViewModel::AddP
 		System.Modify();
 		UNiagaraSystemEditorData* SystemEditorData = CastChecked<UNiagaraSystemEditorData>(System.GetEditorData(), ECastCheckedType::NullChecked);
 		SystemEditorData->Modify();
-		bool bSuccess = FNiagaraEditorUtilities::AddParameter(VariableToAdd, System.GetExposedParameters(), System, SystemEditorData->GetStackEditorData());
+		bool bSuccess = FNiagaraEditorUtilities::AddParameter(VariableToAdd, System.GetExposedParameters(), System, &SystemEditorData->GetStackEditorData());
 		Refresh();
 	}
 	else if (NewScope == ENiagaraParameterScope::System)
diff --git a/Engine/Plugins/FX/Niagara/Source/NiagaraEditor/Private/ViewModels/Stack/NiagaraStackFunctionInput.cpp b/Engine/Plugins/FX/Niagara/Source/NiagaraEditor/Private/ViewModels/Stack/NiagaraStackFunctionInput.cpp
--- a/Engine/Plugins/FX/Niagara/Source/NiagaraEditor/Private/ViewModels/Stack/NiagaraStackFunctionInput.cpp
+++ b/Engine/Plugins/FX/Niagara/Source/NiagaraEditor/Private/ViewModels/Stack/NiagaraStackFunctionInput.cpp
@@ -1611,30 +1611,32 @@ void UNiagaraStackFunctionInput::GetNamespacesForNewParameters(TArray<FName>& Ou
 	UNiagaraNodeOutput* OutputNode = FNiagaraStackGraphUtilities::GetEmitterOutputNodeForStackNode(*OwningFunctionCallNode);
 	bool bIsEditingSystem = GetSystemViewModel()->GetEditMode() == ENiagaraSystemViewModelEditMode::SystemAsset;
 
-	if (OutputNode->GetUsage() == ENiagaraScriptUsage::ParticleSpawnScript || OutputNode->GetUsage() == ENiagaraScriptUsage::ParticleUpdateScript)
+	switch (OutputNode->GetUsage())
+	{
+	case ENiagaraScriptUsage::ParticleSpawnScript:
+	case ENiagaraScriptUsage::ParticleSpawnScriptInterpolated:
+	case ENiagaraScriptUsage::ParticleUpdateScript:
+	case ENiagaraScriptUsage::ParticleEventScript:
+	case ENiagaraScriptUsage::ParticleSimulationStageScript:
 	{
 		OutNamespacesForNewParameters.Add(FNiagaraConstants::ParticleAttributeNamespace);
 		OutNamespacesForNewParameters.Add(FNiagaraConstants::EmitterNamespace);
-		if (bIsEditingSystem)
-		{
-			OutNamespacesForNewParameters.Add(FNiagaraConstants::SystemNamespace);
-			OutNamespacesForNewParameters.Add(FNiagaraConstants::UserNamespace);
-		}
+		break;
 	}
-	else if (OutputNode->GetUsage() == ENiagaraScriptUsage::EmitterSpawnScript || OutputNode->GetUsage() == ENiagaraScriptUsage::EmitterUpdateScript)
+	case ENiagaraScriptUsage::EmitterSpawnScript:
+	case ENiagaraScriptUsage::EmitterUpdateScript:
 	{
 		OutNamespacesForNewParameters.Add(FNiagaraConstants::EmitterNamespace);
-		if (bIsEditingSystem)
-		{
-			OutNamespacesForNewParameters.Add(FNiagaraConstants::SystemNamespace);
-			OutNamespacesForNewParameters.Add(FNiagaraConstants::UserNamespace);
-		}
+		break;
 	}
-	else if ((OutputNode->GetUsage() == ENiagaraScriptUsage::SystemSpawnScript || OutputNode->GetUsage() == ENiagaraScriptUsage::SystemUpdateScript) && bIsEditingSystem)
+	}
+
+	if (bIsEditingSystem)
 	{
-		OutNamespacesForNewParameters.Add(FNiagaraConstants::SystemNamespace);
 		OutNamespacesForNewParameters.Add(FNiagaraConstants::UserNamespace);
+		OutNamespacesForNewParameters.Add(FNiagaraConstants::SystemNamespace);
 	}
+	OutNamespacesForNewParameters.Add(FNiagaraConstants::TransientNamespace);
 }
 
 UNiagaraStackFunctionInput::FOnValueChanged& UNiagaraStackFunctionInput::OnValueChanged()
diff --git a/Engine/Plugins/FX/Niagara/Source/NiagaraEditor/Private/ViewModels/Stack/NiagaraStackSystemSettingsGroup.cpp b/Engine/Plugins/FX/Niagara/Source/NiagaraEditor/Private/ViewModels/Stack/NiagaraStackSystemSettingsGroup.cpp
--- a/Engine/Plugins/FX/Niagara/Source/NiagaraEditor/Private/ViewModels/Stack/NiagaraStackSystemSettingsGroup.cpp
+++ b/Engine/Plugins/FX/Niagara/Source/NiagaraEditor/Private/ViewModels/Stack/NiagaraStackSystemSettingsGroup.cpp
@@ -88,7 +88,7 @@ public:
 	{
 		TSharedRef<FParameterStoreGroupAddAction> ParameterAddAction = StaticCastSharedRef<FParameterStoreGroupAddAction>(AddAction);
 		FNiagaraVariable NewParameterVariable = ParameterAddAction->GetNewParameterVariable();
-		bool bSuccess = FNiagaraEditorUtilities::AddParameter(NewParameterVariable, ParameterStore, ParameterStoreOwner, StackEditorData);
+		bool bSuccess = FNiagaraEditorUtilities::AddParameter(NewParameterVariable, ParameterStore, ParameterStoreOwner, &StackEditorData);
 		if (bSuccess)
 		{
 			OnItemAdded.ExecuteIfBound(NewParameterVariable);
diff --git a/Engine/Plugins/FX/Niagara/Source/NiagaraEditor/Private/Widgets/SNiagaraParameterMapView.cpp b/Engine/Plugins/FX/Niagara/Source/NiagaraEditor/Private/Widgets/SNiagaraParameterMapView.cpp
--- a/Engine/Plugins/FX/Niagara/Source/NiagaraEditor/Private/Widgets/SNiagaraParameterMapView.cpp
+++ b/Engine/Plugins/FX/Niagara/Source/NiagaraEditor/Private/Widgets/SNiagaraParameterMapView.cpp
@@ -59,7 +59,9 @@ FText NiagaraParameterMapSectionID::OnGetSectionTitle(const NiagaraParameterMapS
 	FNiagaraNamespaceMetadata NamespaceMetadata = GetDefault<UNiagaraEditorSettings>()->GetMetaDataForNamespaces(SectionNamespaces);
 	if (NamespaceMetadata.IsValid())
 	{
-		return NamespaceMetadata.DisplayName;
+		return NamespaceMetadata.DisplayNameLong.IsEmptyOrWhitespace() == false
+			? NamespaceMetadata.DisplayNameLong
+			: NamespaceMetadata.DisplayName;
 	}
 	else if (SectionNamespaces.Num() == 1)
 	{
@@ -238,6 +240,7 @@ const FSlateBrush* SNiagaraParameterMapView::GetViewOptionsBorderBrush()
 void SNiagaraParameterMapView::Construct(const FArguments& InArgs, const TArray<TSharedRef<FNiagaraObjectSelection>>& InSelectedObjects, const EToolkitType InToolkitType, const TSharedPtr<FUICommandList>& InToolkitCommands)
 {
 	bNeedsRefresh = false;
+	bIsAddingParameter = false;
 	ToolkitType = InToolkitType;
 	ToolkitCommands = InToolkitCommands;
 	AddParameterButtons.SetNum(NiagaraParameterMapSectionID::Num);
@@ -388,6 +391,7 @@ bool SNiagaraParameterMapView::ParameterAddEnabled() const
 
 void SNiagaraParameterMapView::AddParameter(FNiagaraVariable NewVariable)
 {
+	TGuardValue<bool> AddParameterRefreshGuard(bIsAddingParameter, true);
 	FNiagaraParameterHandle ParameterHandle;
 	bool bSuccess = false;
 
@@ -432,25 +436,11 @@ void SNiagaraParameterMapView::AddParameter(FNiagaraVariable NewVariable)
 			System->Modify();
 			if (NiagaraParameterMapSectionID::OnGetSectionFromVariable(NewVariable, IsStaticSwitchParameter(NewVariable, Graphs), ParameterHandle) == NiagaraParameterMapSectionID::USER)
 			{
-				UNiagaraSystemEditorData* SystemEditorData = CastChecked<UNiagaraSystemEditorData>(System->GetEditorData(), ECastCheckedType::NullChecked);
-				bSuccess = FNiagaraEditorUtilities::AddParameter(NewVariable, System->GetExposedParameters(), *System, SystemEditorData->GetStackEditorData());
+				bSuccess = FNiagaraEditorUtilities::AddParameter(NewVariable, System->GetExposedParameters(), *System, nullptr);
 			}
 			else
 			{
-				FNiagaraVariable VariableToAdd = NewVariable;
-				if (System->EditorOnlyAddedParameters.IndexOf(VariableToAdd) != INDEX_NONE)
-				{
-					TSet<FName> ParameterNames;
-					TArray<FNiagaraVariable> AddedParameterVariables;
-					System->EditorOnlyAddedParameters.GetParameters(AddedParameterVariables);
-					for (const FNiagaraVariable& AddedParameterVariable : AddedParameterVariables)
-					{
-						ParameterNames.Add(AddedParameterVariable.GetName());
-					}
-					FName UniqueName = FNiagaraUtilities::GetUniqueName(VariableToAdd.GetName(), ParameterNames);
-					VariableToAdd.SetName(UniqueName);
-				}
-				bSuccess = System->EditorOnlyAddedParameters.AddParameter(VariableToAdd, false, false);
+				bSuccess = FNiagaraEditorUtilities::AddParameter(NewVariable, System->EditorOnlyAddedParameters, *System, nullptr);
 			}
 		}
 	}
@@ -1090,7 +1080,7 @@ void SNiagaraParameterMapView::OnGraphChanged(const FEdGraphEditAction& InAction
 
 void SNiagaraParameterMapView::OnSystemParameterStoreChanged()
 {
-	if (CachedSystem.IsValid())
+	if (bIsAddingParameter == false && CachedSystem.IsValid())
 	{
 		RefreshActions();
 	}
@@ -1949,7 +1939,23 @@ void SNiagaraAddParameterMenu::AddParameterSelected(FNiagaraVariable NewVariable
 {
 	if (bCreateCustomName)
 	{
-		const static FString NewVariableDefaultName = FString("NewVariable");
+		FString TypeDisplayName;
+		if (NewVariable.GetType().GetEnum() != nullptr)
+		{
+			TypeDisplayName = ((UField*)NewVariable.GetType().GetEnum())->GetDisplayNameText().ToString();
+		}
+		else if (NewVariable.GetType().GetStruct() != nullptr)
+		{
+			TypeDisplayName = NewVariable.GetType().GetStruct()->GetDisplayNameText().ToString();
+		}
+		else if (NewVariable.GetType().GetClass() != nullptr)
+		{
+			TypeDisplayName = NewVariable.GetType().GetClass()->GetDisplayNameText().ToString();
+		}
+		FString NewVariableDefaultName = TypeDisplayName.IsEmpty() 
+			? TEXT("New Variable")
+			: TEXT("New ") + TypeDisplayName;
+
 		TArray<FName> SectionNamespaces;
 		NiagaraParameterMapSectionID::OnGetSectionNamespaces(InSection, SectionNamespaces);
 		TArray<FString> NameParts;
diff --git a/Engine/Plugins/FX/Niagara/Source/NiagaraEditor/Public/NiagaraEditorSettings.h b/Engine/Plugins/FX/Niagara/Source/NiagaraEditor/Public/NiagaraEditorSettings.h
--- a/Engine/Plugins/FX/Niagara/Source/NiagaraEditor/Public/NiagaraEditorSettings.h
+++ b/Engine/Plugins/FX/Niagara/Source/NiagaraEditor/Public/NiagaraEditorSettings.h
@@ -43,7 +43,7 @@ struct FNiagaraNamespaceMetadata
 
 	FNiagaraNamespaceMetadata();
 
-	FNiagaraNamespaceMetadata(TArray<FName> InNamespaces, FText InDisplayName, FText InDescription, FLinearColor InBackgroundColor, TArray<ENiagaraNamespaceMetadataOptions> InOptions);
+	FNiagaraNamespaceMetadata(TArray<FName> InNamespaces);
 
 	UPROPERTY()
 	TArray<FName> Namespaces;
@@ -51,6 +51,9 @@ struct FNiagaraNamespaceMetadata
 	UPROPERTY()
 	FText DisplayName;
 
+	UPROPERTY()
+	FText DisplayNameLong;
+
 	UPROPERTY()
 	FText Description;
 
@@ -60,6 +63,36 @@ struct FNiagaraNamespaceMetadata
 	UPROPERTY()
 	TArray<ENiagaraNamespaceMetadataOptions> Options;
 
+	FNiagaraNamespaceMetadata& SetDisplayName(FText InDisplayName)
+	{
+		DisplayName = InDisplayName;
+		return *this;
+	}
+
+	FNiagaraNamespaceMetadata& SetDisplayNameLong(FText InDisplayNameLong)
+	{
+		DisplayNameLong = InDisplayNameLong;
+		return *this;
+	}
+
+	FNiagaraNamespaceMetadata& SetDescription(FText InDescription)
+	{
+		Description = InDescription;
+		return *this;
+	}
+
+	FNiagaraNamespaceMetadata& SetBackgroundColor(FLinearColor InBackgroundColor)
+	{
+		BackgroundColor = InBackgroundColor;
+		return *this;
+	}
+
+	FNiagaraNamespaceMetadata& AddOption(ENiagaraNamespaceMetadataOptions Option)
+	{
+		Options.Add(Option);
+		return *this;
+	}
+
 	bool IsValid() const { return Namespaces.Num() > 0; }
 };
 
diff --git a/Engine/Plugins/FX/Niagara/Source/NiagaraEditor/Public/NiagaraEditorUtilities.h b/Engine/Plugins/FX/Niagara/Source/NiagaraEditor/Public/NiagaraEditorUtilities.h
--- a/Engine/Plugins/FX/Niagara/Source/NiagaraEditor/Public/NiagaraEditorUtilities.h
+++ b/Engine/Plugins/FX/Niagara/Source/NiagaraEditor/Public/NiagaraEditorUtilities.h
@@ -227,7 +227,7 @@ namespace FNiagaraEditorUtilities
 	 * @param StackEditorData The editor data used to mark the newly added FNiagaraVariable in the Stack for renaming.
 	 * @returns Bool for whether adding the parameter succeeded.
 	 */
-	bool AddParameter(FNiagaraVariable& NewParameterVariable, FNiagaraParameterStore& TargetParameterStore, UObject& ParameterStoreOwner, UNiagaraStackEditorData& StackEditorData);
+	bool AddParameter(FNiagaraVariable& NewParameterVariable, FNiagaraParameterStore& TargetParameterStore, UObject& ParameterStoreOwner, UNiagaraStackEditorData* StackEditorData);
 
 	NIAGARAEDITOR_API bool AddEmitterContextMenuActions(FMenuBuilder& MenuBuilder, const TSharedPtr<FNiagaraEmitterHandleViewModel>& EmitterHandleViewModel);
 
diff --git a/Engine/Plugins/FX/Niagara/Source/NiagaraEditor/Public/Widgets/SNiagaraParameterMapView.h b/Engine/Plugins/FX/Niagara/Source/NiagaraEditor/Public/Widgets/SNiagaraParameterMapView.h
--- a/Engine/Plugins/FX/Niagara/Source/NiagaraEditor/Public/Widgets/SNiagaraParameterMapView.h
+++ b/Engine/Plugins/FX/Niagara/Source/NiagaraEditor/Public/Widgets/SNiagaraParameterMapView.h
@@ -213,6 +213,7 @@ private:
 	TSharedPtr<FUICommandList> ToolkitCommands;
 
 	bool bNeedsRefresh;
+	bool bIsAddingParameter;
 };
 
 class SNiagaraAddParameterMenu : public SCompoundWidget
diff --git a/Engine/Plugins/FX/Niagara/Source/NiagaraEditorWidgets/Private/Stack/SNiagaraStackFunctionInputValue.cpp b/Engine/Plugins/FX/Niagara/Source/NiagaraEditorWidgets/Private/Stack/SNiagaraStackFunctionInputValue.cpp
--- a/Engine/Plugins/FX/Niagara/Source/NiagaraEditorWidgets/Private/Stack/SNiagaraStackFunctionInputValue.cpp
+++ b/Engine/Plugins/FX/Niagara/Source/NiagaraEditorWidgets/Private/Stack/SNiagaraStackFunctionInputValue.cpp
@@ -539,11 +539,12 @@ void SNiagaraStackFunctionInputValue::OnActionSelected(const TArray<TSharedPtr<F
 
 void SNiagaraStackFunctionInputValue::CollectAllActions(FGraphActionListBuilderBase& OutAllActions)
 {
+	bool bIsDataInterfaceOrObject = FunctionInput->GetInputType().IsDataInterface() || FunctionInput->GetInputType().IsUObject();
+
 	// Set a local value
+	if(bIsDataInterfaceOrObject == false)
 	{
-		bool bCanSetLocalValue =
-			FunctionInput->GetValueMode() != UNiagaraStackFunctionInput::EValueMode::Local &&
-			FunctionInput->GetInputType().IsDataInterface() == false;
+		bool bCanSetLocalValue = FunctionInput->GetValueMode() != UNiagaraStackFunctionInput::EValueMode::Local;
 
 		const FText NameText = LOCTEXT("LocalValue", "New Local Value");
 		const FText Tooltip = FText::Format(LOCTEXT("LocalValueToolTip", "Set a local editable value for this input."), NameText);
@@ -608,7 +609,7 @@ void SNiagaraStackFunctionInputValue::CollectAllActions(FGraphActionListBuilderB
 		{
 			InputNames.Add(FunctionInput->GetInputParameterHandlePath()[i].GetName().ToString());
 		}
-		FName InputName = *FString::Join(InputNames, TEXT("."));
+		FName InputName = *FString::Join(InputNames, TEXT("_"));
 
 		for (const FName AvailableNamespace : AvailableNamespaces)
 		{
@@ -628,19 +629,29 @@ void SNiagaraStackFunctionInputValue::CollectAllActions(FGraphActionListBuilderB
 		}
 	}
 
+	if (bIsDataInterfaceOrObject == false)
 	{
+		// Leaving the internal usage of bIsDataInterfaceObject that the tooltip and disabling will work properly when they're moved out of a graph action menu.
 		const FText DisplayName = LOCTEXT("ExpressionLabel", "New Expression");
-		const FText Tooltip = LOCTEXT("ExpressionToolTipl", "Resolve this variable with a custom expression.");
+		const FText Tooltip = bIsDataInterfaceOrObject
+			? LOCTEXT("NoExpresionsForObjects", "Expressions can not be used to set object or data interface parameters.")
+			: LOCTEXT("ExpressionToolTipl", "Resolve this variable with a custom expression.");
 		TSharedPtr<FNiagaraMenuAction> ExpressionAction(new FNiagaraMenuAction(FText(), DisplayName, Tooltip, 0, FText(),
-			FNiagaraMenuAction::FOnExecuteStackAction::CreateSP(this, &SNiagaraStackFunctionInputValue::CustomExpressionSelected)));
+			FNiagaraMenuAction::FOnExecuteStackAction::CreateSP(this, &SNiagaraStackFunctionInputValue::CustomExpressionSelected),
+			FNiagaraMenuAction::FCanExecuteStackAction::CreateLambda([bIsDataInterfaceOrObject]() { return bIsDataInterfaceOrObject == false; })));
 		OutAllActions.AddAction(ExpressionAction);
 	}
 
+	if (bIsDataInterfaceOrObject == false)
 	{
+		// Leaving the internal usage of bIsDataInterfaceObject that the tooltip and disabling will work properly when they're moved out of a graph action menu.
 		const FText CreateDisplayName = LOCTEXT("ScratchLabel", "New Scratch Dynamic Input");
-		const FText CreateTooltip = LOCTEXT("ScratchToolTipl", "Create a new dynamic input in the scratch pad.");
+		const FText CreateTooltip = bIsDataInterfaceOrObject
+			? LOCTEXT("NoScratchForObjects", "Dynamic inputs can not be used to set object or data interface parameters.")
+			: LOCTEXT("ScratchToolTipl", "Create a new dynamic input in the scratch pad.");
 		TSharedPtr<FNiagaraMenuAction> CreateScratchAction(new FNiagaraMenuAction(FText(), CreateDisplayName, CreateTooltip, 0, FText(),
-			FNiagaraMenuAction::FOnExecuteStackAction::CreateSP(this, &SNiagaraStackFunctionInputValue::CreateScratchSelected)));
+			FNiagaraMenuAction::FOnExecuteStackAction::CreateSP(this, &SNiagaraStackFunctionInputValue::CreateScratchSelected),
+			FNiagaraMenuAction::FCanExecuteStackAction::CreateLambda([bIsDataInterfaceOrObject]() { return bIsDataInterfaceOrObject == false; })));
 		OutAllActions.AddAction(CreateScratchAction);
 	}
 
