Improved info message on how to set a default value

diff --git a/Engine/Plugins/FX/Niagara/Source/NiagaraEditor/Private/Customizations/NiagaraScriptVariableCustomization.cpp b/Engine/Plugins/FX/Niagara/Source/NiagaraEditor/Private/Customizations/NiagaraScriptVariableCustomization.cpp
--- a/Engine/Plugins/FX/Niagara/Source/NiagaraEditor/Private/Customizations/NiagaraScriptVariableCustomization.cpp
+++ b/Engine/Plugins/FX/Niagara/Source/NiagaraEditor/Private/Customizations/NiagaraScriptVariableCustomization.cpp
@@ -248,7 +248,7 @@ void FNiagaraScriptVariableDetails::CustomizeDetails(IDetailLayoutBuilder& Detai
 					[
 						SNew(STextBlock)
 						.Font(FNiagaraEditorStyle::Get().GetFontStyle("NiagaraEditor.ParameterFont"))
-					.Text(NSLOCTEXT("ScriptVariableCustomization", "MissingDefaults", "To set default, add to a Map Get node."))
+					.Text(NSLOCTEXT("ScriptVariableCustomization", "MissingDefaults", "To set default, add to a Map Get node that is wired to the graph."))
 					];
 			}
 		}
