Fixed bound default parameters.

Replaced special case code with standard HandleParameterRead() call.
Previous special case code was not covering all paths and would only work in the spawn script.

diff --git a/Engine/Plugins/FX/Niagara/Source/NiagaraCore/Private/NiagaraCustomVersion.cpp b/Engine/Plugins/FX/Niagara/Source/NiagaraCore/Private/NiagaraCustomVersion.cpp
--- a/Engine/Plugins/FX/Niagara/Source/NiagaraCore/Private/NiagaraCustomVersion.cpp
+++ b/Engine/Plugins/FX/Niagara/Source/NiagaraCore/Private/NiagaraCustomVersion.cpp
@@ -8,4 +8,4 @@ const FGuid FNiagaraCustomVersion::GUID(0xFCF57AFA, 0x50764283, 0xB9A9E658, 0xFF
 // Register the custom version with core
 FCustomVersionRegistration GRegisterNiagaraCustomVersion(FNiagaraCustomVersion::GUID, FNiagaraCustomVersion::LatestVersion, TEXT("NiagaraVer"));
 
-const FGuid FNiagaraCustomVersion::LatestScriptCompileVersion(0x269AF844, 0x3EB24E94, 0xB1DECF9C, 0x21937120);
+const FGuid FNiagaraCustomVersion::LatestScriptCompileVersion(0xF4357B0E, 0x349C4526, 0x97A94D8E, 0x85A7380D);
\ No newline at end of file
diff --git a/Engine/Plugins/FX/Niagara/Source/NiagaraEditor/Private/NiagaraHlslTranslator.cpp b/Engine/Plugins/FX/Niagara/Source/NiagaraEditor/Private/NiagaraHlslTranslator.cpp
--- a/Engine/Plugins/FX/Niagara/Source/NiagaraEditor/Private/NiagaraHlslTranslator.cpp
+++ b/Engine/Plugins/FX/Niagara/Source/NiagaraEditor/Private/NiagaraHlslTranslator.cpp
@@ -4604,62 +4604,7 @@ void FHlslNiagaraTranslator::ParameterMapGet(UNiagaraNodeParameterMapGet* GetNod
 			FNiagaraVariable Var = Schema->PinToNiagaraVariable(OutputPins[i], bNeedsValue);
 
 			UNiagaraScriptVariable* Variable = GetNode->GetNiagaraGraph()->GetScriptVariable(Var);
-
-			// Handle parameter map overrides for bindings as a special case, 
-			// effectively copying the bound value to the variable in question prior to invoking a function call.
-			// Especially relevant for module parameters that are also bindings
-			bool bFoundBinding = false;
-			if (Variable && Variable->DefaultMode == ENiagaraDefaultMode::Binding)
-			{
-				FNiagaraScriptVariableBinding Bind = Variable->DefaultBinding;
-				if (Bind.IsValid())
-				{
-					int LastSetChunkIdx = INDEX_NONE;
-
-					// Check whether we've encountered this variable before, i.e. in BuildMissingDefaults in a spawn script. In this case we'll skip
-					for (auto& UniqueVar : UniqueVarToChunk)
-					{
-						if (UniqueVar.Key.IsEquivalent(FNiagaraVariable(Var.GetType(), Var.GetName())))
-						{
-							LastSetChunkIdx = UniqueVar.Value;
-							break;
-						}
-					}
-					
-					if (LastSetChunkIdx == INDEX_NONE)
-					{
-						// If we haven't encountered this variable before, 
-						// search through the param history for a variable matching the bindings and add a source chunk to it
-						Var = ActiveHistoryForFunctionCalls.ResolveAliases(Var);
-						for (auto HistoryVariable : ParamMapHistories[ParamMapHistoryIdx].Variables)
-						{
-							if (HistoryVariable.IsEquivalent(FNiagaraVariable(Var.GetType(), Bind.GetName())))
-							{
-								FString SanitizedName = GetParameterMapInstanceName(ActiveStageIdx) + TEXT(".") + GetSanitizedSymbolName(Var.GetName().ToString());
-								LastSetChunkIdx = AddSourceChunk(SanitizedName, Var.GetType());
-								break;
-							}
-						}
-						if (LastSetChunkIdx != INDEX_NONE && Var.GetType().GetClass() == nullptr)
-						{
-							int32 VarIdx = ParamMapHistories[ParamMapHistoryIdx].FindVariableByName(Var.GetName());
-							if (VarIdx != INDEX_NONE && VarIdx < ParamMapSetVariablesToChunks[ParamMapHistoryIdx].Num())
-							{
-								// Record that we wrote to it
-								ParamMapSetVariablesToChunks[ParamMapHistoryIdx][VarIdx] = LastSetChunkIdx;
-								ParamMapDefinedAttributesToNamespaceVars.FindOrAdd(Var.GetName()) = Var;
-								Outputs[i] = LastSetChunkIdx;
-								bFoundBinding = true;
-							}
-						}
-					}
-				}
-			}
-
-			if (!bFoundBinding)
-			{
-				HandleParameterRead(ParamMapHistoryIdx, Var, GetNode->GetDefaultPin(OutputPins[i]), GetNode, Outputs[i], nullptr);
-			}
+			HandleParameterRead(ParamMapHistoryIdx, Var, GetNode->GetDefaultPin(OutputPins[i]), GetNode, Outputs[i], Variable);
 		}
 	}
 }
@@ -4890,42 +4835,19 @@ void FHlslNiagaraTranslator::HandleParameterRead(int32 ParamMapHistoryIdx, const
 			{
 				FNiagaraScriptVariableBinding Bind = Variable->DefaultBinding;
 
-				// Check whether we've encountered the binding before, if so return its chunk
-				bool bFoundVariable = false;
-				for (auto& UniqueVar : UniqueVars)
-				{
-					if (Bind.GetName() == UniqueVar.GetName())
-					{
-						bFoundVariable = true;
-						break;
-					}
-				}
+				int32 Out = INDEX_NONE;
+				FNiagaraVariable BindVar = FNiagaraVariable(InVar.GetType(), Bind.GetName());
+				HandleParameterRead(ActiveStageIdx, BindVar, nullptr, ErrorNode, Out, nullptr);
 
-				if (bFoundVariable)
+				if (Out != INDEX_NONE)
 				{
-					for (auto& UniqueVar : UniqueVarToChunk)
-					{
-						if (UniqueVar.Key.IsEquivalent(FNiagaraVariable(Variable->Variable.GetType(), Bind.GetName())))
-						{
-							LastSetChunkIdx = UniqueVar.Value;
-							break;
-						}
-					}
+					LastSetChunkIdx = Out;
 				}
 				else
 				{
-					// We haven't encountered the binding before, try to check whether it's a known variable or not.
-					int Out = GetParameter(FNiagaraVariable(InVar.GetType(), Bind.GetName()));
-					if (Out != INDEX_NONE)
-					{
-						LastSetChunkIdx = Out;
-					} 
-					else 
-					{
-						Error(FText::Format(LOCTEXT("CannotFindBinding", "The module input {0} is bound to {1}, but {1} is not yet defined. Make sure {1} is defined prior to this module call."),
-							FText::FromName(Var.GetName()),
-							FText::FromName(Bind.GetName())), ErrorNode, nullptr);
-					}
+					Error(FText::Format(LOCTEXT("CannotFindBinding", "The module input {0} is bound to {1}, but {1} is not yet defined. Make sure {1} is defined prior to this module call."),
+						FText::FromName(Var.GetName()),
+						FText::FromName(Bind.GetName())), ErrorNode, nullptr);
 				}
 			}
 			else if (InputPin != nullptr) // Default was found, trace back its inputs.
