Niagara - Fix an issue where nested namespace outputs weren't getting fixed on load correctly, also fix linux build error.

diff --git a/Engine/Plugins/FX/Niagara/Source/NiagaraEditor/Private/NiagaraEditorSettings.cpp b/Engine/Plugins/FX/Niagara/Source/NiagaraEditor/Private/NiagaraEditorSettings.cpp
--- a/Engine/Plugins/FX/Niagara/Source/NiagaraEditor/Private/NiagaraEditorSettings.cpp
+++ b/Engine/Plugins/FX/Niagara/Source/NiagaraEditor/Private/NiagaraEditorSettings.cpp
@@ -59,12 +59,13 @@ void UNiagaraEditorSettings::SetupNamespaceMetadata()
 			.SetDescription(LOCTEXT("ModuleOutputDescription", "A transient value which is unique to the module that wrote it and can be read from any other module.\nTransient values do not persist from frame to frame, or between stages, e.g. emitter to particle, or spawn to update."))
 			.SetBackgroundColor(FLinearColor(FColor(109, 95, 124)))
 			.AddOption(ENiagaraNamespaceMetadataOptions::Advanced)
-			.AddOption(ENiagaraNamespaceMetadataOptions::PreventCreatingInSystemEditor),
+			.AddOption(ENiagaraNamespaceMetadataOptions::PreventCreatingInSystemEditor)
+			.AddOption(ENiagaraNamespaceMetadataOptions::CanChangeNamespaceModifier),
 		FNiagaraNamespaceMetadata({FNiagaraConstants::OutputNamespace})
 			.SetDisplayName(LOCTEXT("OutputDisplayName", "Output"))
 			.SetDescription(LOCTEXT("OutputDescription", "A transient value which is unique to the module that wrote it and can be read from any other module.\nTransient values do not persist from frame to frame, or between stages, e.g. emitter to particle, or spawn to update."))
 			.SetBackgroundColor(FLinearColor(FColor(109, 95, 124)))
-			.AddOption({ENiagaraNamespaceMetadataOptions::Advanced}),
+			.AddOption(ENiagaraNamespaceMetadataOptions::Advanced),
 		FNiagaraNamespaceMetadata({FNiagaraConstants::LocalNamespace, FNiagaraConstants::ModuleNamespace})
 			.SetDisplayName(LOCTEXT("ModuleLocalDisplayName", "Local"))
 			.SetDescription(LOCTEXT("ModuleLocalDescription", "A transient value which can be written to and read from within a single module.\nTransient values do not persist from frame to frame, or between stages, e.g. emitter to particle, or spawn to update."))
diff --git a/Engine/Plugins/FX/Niagara/Source/NiagaraEditor/Private/NiagaraGraph.cpp b/Engine/Plugins/FX/Niagara/Source/NiagaraEditor/Private/NiagaraGraph.cpp
--- a/Engine/Plugins/FX/Niagara/Source/NiagaraEditor/Private/NiagaraGraph.cpp
+++ b/Engine/Plugins/FX/Niagara/Source/NiagaraEditor/Private/NiagaraGraph.cpp
@@ -545,9 +545,18 @@ FName StandardizeName(FName Name, ENiagaraScriptUsage Usage, bool bIsGet, bool b
 			}
 			else
 			{
-				// When reading they can be from the module namespace, or a more specific namespace from a different module.
-				Namespace = *(NameParts[0].ToString() + TEXT(".") + NameParts[1].ToString());
-				NameParts.RemoveAt(0, 2);
+				// When reading they can be from the module namespace, or a more specific namespace from a different module, we also need
+				// to handle the case where they are reading from the output of a nested module, so allow an additional namespace too.
+				if (NameParts.Num() > 3)
+				{
+					Namespace = *(NameParts[0].ToString() + TEXT(".") + NameParts[1].ToString() + TEXT(".") + NameParts[2].ToString());
+					NameParts.RemoveAt(0, 3);
+				}
+				else
+				{
+					Namespace = *(NameParts[0].ToString() + TEXT(".") + NameParts[1].ToString());
+					NameParts.RemoveAt(0, 2);
+				}
 			}
 		}
 		else
