Niagara - Fixes for collision modules and more UX fixups.
+ Add a dark text style option for namespace metadata which allows us to use lighter colors as the backgrounds.  Used this to update the engine metadata which was blending into the background in certain UIs.
+ Made it possible to rename and modify namespace modifiers for discovered parameters.  This will create a new parameter and will show a toast explaining to the user what's happened.
+ Fixed a layout issue where parameter names were pushing the reference count button out of view.
+ Fixed an issue where editing the namespace modifier on a parameter that had a 2 part namespace like Output.Module would not generate the correct parameter name.

diff --git a/Engine/Plugins/FX/Niagara/Source/NiagaraEditor/Private/NiagaraEditorSettings.cpp b/Engine/Plugins/FX/Niagara/Source/NiagaraEditor/Private/NiagaraEditorSettings.cpp
--- a/Engine/Plugins/FX/Niagara/Source/NiagaraEditor/Private/NiagaraEditorSettings.cpp
+++ b/Engine/Plugins/FX/Niagara/Source/NiagaraEditor/Private/NiagaraEditorSettings.cpp
@@ -5,11 +5,14 @@
 
 FNiagaraNamespaceMetadata::FNiagaraNamespaceMetadata()
 	: BackgroundColor(FLinearColor::Black)
+	, ForegroundStyle("NiagaraEditor.ParameterName.NamespaceText")
 {
 }
 
 FNiagaraNamespaceMetadata::FNiagaraNamespaceMetadata(TArray<FName> InNamespaces)
 	: Namespaces(InNamespaces)
+	, BackgroundColor(FLinearColor::Black)
+	, ForegroundStyle("NiagaraEditor.ParameterName.NamespaceText")
 {
 }
 
@@ -79,7 +82,8 @@ void UNiagaraEditorSettings::SetupNamespaceMetadata()
 		FNiagaraNamespaceMetadata({FNiagaraConstants::EngineNamespace})
 			.SetDisplayName(LOCTEXT("EngineDisplayName", "Engine"))
 			.SetDescription(LOCTEXT("EngineDescription", "A read only value which is provided by the engine.\nThis value's source can be the simulation itsef\ne.g. ExecutionCount, or the owner of the simulation e.g. (Owner) Scale."))
-			.SetBackgroundColor(FLinearColor(FColor(45, 45, 45)))
+			.SetBackgroundColor(FLinearColor(FColor(170, 170, 170)))
+			.SetForegroundStyle("NiagaraEditor.ParameterName.NamespaceTextDark")
 			.AddOption(ENiagaraNamespaceMetadataOptions::PreventRenaming),
 		FNiagaraNamespaceMetadata({FNiagaraConstants::UserNamespace})
 			.SetDisplayName(LOCTEXT("UserDisplayName", "User"))
@@ -89,13 +93,15 @@ void UNiagaraEditorSettings::SetupNamespaceMetadata()
 			.SetDisplayName(LOCTEXT("NiagaraParameterCollectionDisplayName", "NPC"))
 			.SetDisplayNameLong(LOCTEXT("NiagaraParameterCollectionDisplayNameLong", "Niagara Parameter Collection"))
 			.SetDescription(LOCTEXT("NiagaraParameterCollectionDescription", "Values read from a niagara parameter collection asset.\nRead only in a niagara system."))
-			.SetBackgroundColor(FLinearColor(FColor(45, 45, 45)))
+			.SetBackgroundColor(FLinearColor(FColor(170, 170, 170)))
+			.SetForegroundStyle("NiagaraEditor.ParameterName.NamespaceTextDark")
 			.AddOption(ENiagaraNamespaceMetadataOptions::Advanced)
 			.AddOption(ENiagaraNamespaceMetadataOptions::PreventRenaming),
 		FNiagaraNamespaceMetadata({FNiagaraConstants::DataInstanceNamespace})
 			.SetDisplayName(LOCTEXT("DataInstanceDisplayName", "Data Instance"))
 			.SetDescription(LOCTEXT("DataInstanceDescription", "A special transient value which has a single bool IsAlive value, which determines if a particle is alive or not."))
-			.SetBackgroundColor(FLinearColor(FColor(45, 45, 45)))
+			.SetBackgroundColor(FLinearColor(FColor(170, 170, 170)))
+			.SetForegroundStyle("NiagaraEditor.ParameterName.NamespaceTextDark")
 			.AddOption(ENiagaraNamespaceMetadataOptions::Advanced)
 			.AddOption(ENiagaraNamespaceMetadataOptions::PreventRenaming),
 	};
@@ -121,7 +127,8 @@ void UNiagaraEditorSettings::SetupNamespaceMetadata()
 		FNiagaraNamespaceMetadata({FNiagaraConstants::OwnerNamespace})
 			.SetDisplayName(LOCTEXT("OwnerDisplayName", "Owner"))
 			.SetDescription(LOCTEXT("OwnerDescription", "A namespace modifier which specifies that an engine provided parameter comes from the owner, or component."))
-			.SetBackgroundColor(FLinearColor(FColor(45, 45, 45))),
+			.SetBackgroundColor(FLinearColor(FColor(170, 170, 170)))
+			.SetForegroundStyle("NiagaraEditor.ParameterName.NamespaceTextDark"),
 	};
 }
 
diff --git a/Engine/Plugins/FX/Niagara/Source/NiagaraEditor/Private/NiagaraEditorStyle.cpp b/Engine/Plugins/FX/Niagara/Source/NiagaraEditor/Private/NiagaraEditorStyle.cpp
--- a/Engine/Plugins/FX/Niagara/Source/NiagaraEditor/Private/NiagaraEditorStyle.cpp
+++ b/Engine/Plugins/FX/Niagara/Source/NiagaraEditor/Private/NiagaraEditorStyle.cpp
@@ -186,12 +186,19 @@ TSharedRef< FSlateStyleSet > FNiagaraEditorStyle::Create()
 	Style->Set("NiagaraEditor.ParameterSpinbox", ParameterSpinBox);
 
 	Style->Set("NiagaraEditor.ParameterName.NamespaceBorder", new BOX_PLUGIN_BRUSH("Icons/NamespaceBorder", FMargin(4.0f / 16.0f)));
+
 	Style->Set("NiagaraEditor.ParameterName.NamespaceText", FTextBlockStyle(NormalText)
 		.SetFont(DEFAULT_FONT("Bold", 8))
 		.SetColorAndOpacity(FLinearColor(1.0f, 1.0f, 1.0f, 0.9f))
 		.SetShadowOffset(FVector2D(1, 1))
 		.SetShadowColorAndOpacity(FLinearColor(0, 0, 0, 0.7f)));
 
+	Style->Set("NiagaraEditor.ParameterName.NamespaceTextDark", FTextBlockStyle(NormalText)
+		.SetFont(DEFAULT_FONT("Bold", 8))
+		.SetColorAndOpacity(FLinearColor(0.0f, 0.0f, 0.0f, 0.9f))
+		.SetShadowOffset(FVector2D(1, 1))
+		.SetShadowColorAndOpacity(FLinearColor(1.0, 1.0, 1.0, 0.5f)));
+
 	Style->Set("NiagaraEditor.Stack.HighlightedButtonBrush", new BOX_CORE_BRUSH("Common/ButtonHoverHint", FMargin(4 / 16.0f), GetDefault<UEditorStyleSettings>()->SelectionColor));
 
 	// Parameter Map View
diff --git a/Engine/Plugins/FX/Niagara/Source/NiagaraEditor/Private/NiagaraEditorUtilities.cpp b/Engine/Plugins/FX/Niagara/Source/NiagaraEditor/Private/NiagaraEditorUtilities.cpp
--- a/Engine/Plugins/FX/Niagara/Source/NiagaraEditor/Private/NiagaraEditorUtilities.cpp
+++ b/Engine/Plugins/FX/Niagara/Source/NiagaraEditor/Private/NiagaraEditorUtilities.cpp
@@ -2058,15 +2058,15 @@ void FNiagaraEditorUtilities::WarnWithToastAndLog(FText WarningMessage)
 	UE_LOG(LogNiagaraEditor, Warning, TEXT("%s"), *WarningMessage.ToString());
 }
 
-void FNiagaraEditorUtilities::InfoWithToastAndLog(FText InfoMessage)
+void FNiagaraEditorUtilities::InfoWithToastAndLog(FText InfoMessage, float ToastDuration)
 {
 	FNotificationInfo WarningNotification(InfoMessage);
-	WarningNotification.ExpireDuration = 5.0f;
+	WarningNotification.ExpireDuration = ToastDuration;
 	WarningNotification.bFireAndForget = true;
 	WarningNotification.bUseLargeFont = false;
 	WarningNotification.Image = FCoreStyle::Get().GetBrush(TEXT("MessageLog.Note"));
 	FSlateNotificationManager::Get().AddNotification(WarningNotification);
-	UE_LOG(LogNiagaraEditor, Warning, TEXT("%s"), *InfoMessage.ToString());
+	UE_LOG(LogNiagaraEditor, Log, TEXT("%s"), *InfoMessage.ToString());
 }
 
 FName FNiagaraEditorUtilities::GetUniqueObjectName(UObject* Outer, UClass* ObjectClass, const FString& CandidateName)
@@ -2432,4 +2432,17 @@ void FNiagaraEditorUtilities::GetReferencingFunctionCallNodes(UNiagaraScript* Sc
 	}
 }
 
+FText FNiagaraEditorUtilities::FormatParameterNameForTextDisplay(FName ParameterName)
+{
+	FNiagaraParameterHandle ParameterHandle(ParameterName);
+	TArray<FName> HandleParts = ParameterHandle.GetHandleParts();
+	FString DisplayString;
+	for (int32 HandlePartIndex = 0; HandlePartIndex < HandleParts.Num() - 1; HandlePartIndex++)
+	{
+		DisplayString += TEXT("(") + HandleParts[HandlePartIndex].ToString().ToUpper() + TEXT(") ");
+	}
+	DisplayString += HandleParts[HandleParts.Num() - 1].ToString();
+	return FText::FromString(DisplayString);
+}
+
 #undef LOCTEXT_NAMESPACE
diff --git a/Engine/Plugins/FX/Niagara/Source/NiagaraEditor/Private/SNiagaraParameterName.cpp b/Engine/Plugins/FX/Niagara/Source/NiagaraEditor/Private/SNiagaraParameterName.cpp
--- a/Engine/Plugins/FX/Niagara/Source/NiagaraEditor/Private/SNiagaraParameterName.cpp
+++ b/Engine/Plugins/FX/Niagara/Source/NiagaraEditor/Private/SNiagaraParameterName.cpp
@@ -20,6 +20,7 @@ void SNiagaraParameterName::Construct(const FArguments& InArgs)
 	OnNameChangedDelegate = InArgs._OnNameChanged;
 	OnDoubleClickedDelegate = InArgs._OnDoubleClicked;
 	IsSelected = InArgs._IsSelected;
+	Decorator = InArgs._Decorator.Widget;
 
 	UpdateContent(ParameterName.Get());
 }
@@ -45,7 +46,7 @@ FReply SNiagaraParameterName::OnMouseButtonDoubleClick(const FGeometry& InMyGeom
 	}
 }
 
-TSharedRef<SBorder> SNiagaraParameterName::CreateNamespaceWidget(FText NamespaceDisplayName, FText NamespaceDescription, FLinearColor NamespaceBorderColor)
+TSharedRef<SBorder> SNiagaraParameterName::CreateNamespaceWidget(FText NamespaceDisplayName, FText NamespaceDescription, FLinearColor NamespaceBorderColor, FName NamespaceForegroundStyle)
 {
 	return SNew(SBorder)
 	.BorderImage(FNiagaraEditorStyle::Get().GetBrush("NiagaraEditor.ParameterName.NamespaceBorder"))
@@ -55,7 +56,7 @@ TSharedRef<SBorder> SNiagaraParameterName::CreateNamespaceWidget(FText Namespace
 	.Padding(FMargin(5.0f, 1.0f, 5.0f, 1.0f))
 	[
 		SNew(STextBlock)
-		.TextStyle(FNiagaraEditorStyle::Get(), "NiagaraEditor.ParameterName.NamespaceText")
+		.TextStyle(FNiagaraEditorStyle::Get(), NamespaceForegroundStyle)
 		.Text(NamespaceDisplayName)
 		.HighlightText(HighlightText)
 	];
@@ -92,16 +93,17 @@ void SNiagaraParameterName::UpdateContent(FName InDisplayedParameterName)
 		if (NamespaceMetadata.IsValid())
 		{
 			Namespaces.RemoveAt(0, NamespaceMetadata.Namespaces.Num());
-			NamespaceWidget = CreateNamespaceWidget(NamespaceMetadata.DisplayName.ToUpper(), NamespaceMetadata.Description, NamespaceMetadata.BackgroundColor);
+			NamespaceWidget = CreateNamespaceWidget(
+				NamespaceMetadata.DisplayName.ToUpper(), NamespaceMetadata.Description,
+				NamespaceMetadata.BackgroundColor, NamespaceMetadata.ForegroundStyle);
 		}
 		else
 		{
 			FText NamespaceDisplayName = FText::FromString(FName::NameToDisplayString(Namespaces[0].ToString(), false).ToUpper());
 			Namespaces.RemoveAt(0);
 			NamespaceWidget = CreateNamespaceWidget(
-				NamespaceDisplayName,
-				DefaultNamespaceMetadata.Description,
-				DefaultNamespaceMetadata.BackgroundColor);
+				NamespaceDisplayName, DefaultNamespaceMetadata.Description,
+				DefaultNamespaceMetadata.BackgroundColor, DefaultNamespaceMetadata.ForegroundStyle);
 		}
 
 		ContentBox->AddSlot()
@@ -125,9 +127,8 @@ void SNiagaraParameterName::UpdateContent(FName InDisplayedParameterName)
 
 		Namespaces.RemoveAt(0);
 		NamespaceModifierBorder = CreateNamespaceWidget(
-			FText::FromString(FName::NameToDisplayString(DisplayedNamespaceModifier.ToString(), false).ToUpper()),
-			DisplayedNamespaceModifierMetadata.Description,
-			DisplayedNamespaceModifierMetadata.BackgroundColor);
+			FText::FromString(FName::NameToDisplayString(DisplayedNamespaceModifier.ToString(), false).ToUpper()), DisplayedNamespaceModifierMetadata.Description,
+			DisplayedNamespaceModifierMetadata.BackgroundColor, DisplayedNamespaceModifierMetadata.ForegroundStyle);
 
 		ContentBox->AddSlot()
 		.VAlign(VAlign_Center)
@@ -152,9 +153,8 @@ void SNiagaraParameterName::UpdateContent(FName InDisplayedParameterName)
 		.Padding(0.0f, 0.0f, 5.0f, 0.0f)
 		[
 			CreateNamespaceWidget(
-				FText::FromString(FName::NameToDisplayString(Namespace.ToString(), false).ToUpper()),
-				DefaultNamespaceMetadata.Description,
-				DefaultNamespaceMetadata.BackgroundColor)
+				FText::FromString(FName::NameToDisplayString(Namespace.ToString(), false).ToUpper()), DefaultNamespaceMetadata.Description,
+				DefaultNamespaceMetadata.BackgroundColor, DefaultNamespaceMetadata.ForegroundStyle)
 		];
 	}
 
@@ -184,6 +184,17 @@ void SNiagaraParameterName::UpdateContent(FName InDisplayedParameterName)
 		NameWidget.ToSharedRef()
 	];
 
+	if (Decorator.IsValid())
+	{
+		ContentBox->AddSlot()
+		.VAlign(VAlign_Center)
+		.AutoWidth()
+		.Padding(3.0f, 0.0f, 0.0f, 0.0f)
+		[
+			Decorator.ToSharedRef()
+		];
+	}
+
 	ChildSlot
 	[
 		ContentBox
@@ -207,12 +218,29 @@ FName SNiagaraParameterName::ReconstructNameFromNamespaceModifierEditText(const
 	TArray<FString> NameParts;
 	CurrentParameterNameString.ParseIntoArray(NameParts, TEXT("."));
 
-	if (NameParts.Num() >= 3)
+	TArray<FName> Namespaces;
+	for (int32 NamePartIndex = 0; NamePartIndex < NameParts.Num() - 1; NamePartIndex++)
 	{
-		NameParts[1] = InEditText.ToString().Replace(TEXT("."), TEXT("_"));
+		Namespaces.Add(*NameParts[NamePartIndex]);
+	}
+
+	FNiagaraNamespaceMetadata NamespaceMetadata = GetDefault<UNiagaraEditorSettings>()->GetMetaDataForNamespaces(Namespaces);
+
+	if (NamespaceMetadata.IsValid())
+	{
+		if (NameParts.Num() > NamespaceMetadata.Namespaces.Num() + 1)
+		{
+			NameParts[NamespaceMetadata.Namespaces.Num()] = InEditText.ToString().Replace(TEXT("."), TEXT("_"));
+		}
+		FString NewParameterNameString = FString::Join(NameParts, TEXT("."));
+		return *NewParameterNameString;
+	}
+	else
+	{
+		// If there isn't valid namespace metadata for this parameter we can't safely edit the modifier
+		// since there is not way to know how many parts the namespace actually has.
+		return ParameterName.Get();
 	}
-	FString NewParameterNameString = FString::Join(NameParts, TEXT("."));
-	return *NewParameterNameString;
 }
 
 FText SNiagaraParameterName::GetNamespaceModifierText()
@@ -309,6 +337,10 @@ void SNiagaraParameterNameTextBlock::Construct(const FArguments& InArgs)
 		.IsSelected(InArgs._IsSelected)
 		.OnVerifyNameChange(this, &SNiagaraParameterNameTextBlock::VerifyNameChange)
 		.OnNameChanged(this, &SNiagaraParameterNameTextBlock::NameChanged)
+		.Decorator()
+		[
+			InArgs._Decorator.Widget
+		]
 	];
 }
 
diff --git a/Engine/Plugins/FX/Niagara/Source/NiagaraEditor/Private/Widgets/SNiagaraParameterMapPaletteItem.cpp b/Engine/Plugins/FX/Niagara/Source/NiagaraEditor/Private/Widgets/SNiagaraParameterMapPaletteItem.cpp
--- a/Engine/Plugins/FX/Niagara/Source/NiagaraEditor/Private/Widgets/SNiagaraParameterMapPaletteItem.cpp
+++ b/Engine/Plugins/FX/Niagara/Source/NiagaraEditor/Private/Widgets/SNiagaraParameterMapPaletteItem.cpp
@@ -50,18 +50,25 @@ void SNiagaraParameterMapPalleteItem::Construct(const FArguments& InArgs, FCreat
 	ParameterToolTipText.Bind(this, &SNiagaraParameterMapPalleteItem::GetItemTooltip);
 	SetToolTipText(ParameterToolTipText);
 
+	bWasCreated = ParameterAction->ReferenceCollection.ContainsByPredicate(
+		[](const FNiagaraGraphParameterReferenceCollection& ReferenceCollection) { return ReferenceCollection.WasCreated(); });
+
 	ParameterNameTextBlock = SNew(SNiagaraParameterNameTextBlock)
 		.ParameterText(this, &SNiagaraParameterMapPalleteItem::GetDisplayText)
 		.HighlightText(InCreateData->HighlightText)
 		.OnTextCommitted(this, &SNiagaraParameterMapPalleteItem::OnNameTextCommitted)
 		.OnVerifyTextChanged(this, &SNiagaraParameterMapPalleteItem::OnNameTextVerifyChanged)
 		.IsSelected(InCreateData->IsRowSelectedDelegate)
-		.IsReadOnly(InCreateData->bIsReadOnly);
+		.IsReadOnly(InCreateData->bIsReadOnly)
+		.Decorator()
+		[
+			SNew(STextBlock)
+			.Visibility(bWasCreated ? EVisibility::Visible : EVisibility::Collapsed)
+			.Text(LOCTEXT("AddedText", "*"))
+		];
 
 	InCreateData->OnRenameRequest->BindSP(ParameterNameTextBlock.ToSharedRef(), &SNiagaraParameterNameTextBlock::EnterEditingMode);
 
-	bWasCreated = ParameterAction->ReferenceCollection.ContainsByPredicate([](const FNiagaraGraphParameterReferenceCollection& ReferenceCollection) { return ReferenceCollection.WasCreated(); });
-
 	// now, create the actual widget
 	ChildSlot
 	[
@@ -75,24 +82,15 @@ void SNiagaraParameterMapPalleteItem::Construct(const FArguments& InArgs, FCreat
 		]
 		// name slot
 		+SHorizontalBox::Slot()
-		.AutoWidth()
 		.VAlign(VAlign_Center)
 		.Padding(5,0)
 		[
 			ParameterNameTextBlock.ToSharedRef()
 		]
-		+ SHorizontalBox::Slot()
-		.AutoWidth()
-		.VAlign(VAlign_Center)
-		[
-			SNew(STextBlock)
-			.Visibility(bWasCreated ? EVisibility::Visible : EVisibility::Collapsed)
-			.Text(LOCTEXT("AddedText", "*"))
-		]
 		// reference count
 		+ SHorizontalBox::Slot()
+		.AutoWidth()
 		.VAlign(VAlign_Center)
-		.HAlign(HAlign_Right)
 		.Padding(3, 0)
 		[
 			SNew(SComboButton)
diff --git a/Engine/Plugins/FX/Niagara/Source/NiagaraEditor/Private/Widgets/SNiagaraParameterMapView.cpp b/Engine/Plugins/FX/Niagara/Source/NiagaraEditor/Private/Widgets/SNiagaraParameterMapView.cpp
--- a/Engine/Plugins/FX/Niagara/Source/NiagaraEditor/Private/Widgets/SNiagaraParameterMapView.cpp
+++ b/Engine/Plugins/FX/Niagara/Source/NiagaraEditor/Private/Widgets/SNiagaraParameterMapView.cpp
@@ -390,6 +390,12 @@ bool SNiagaraParameterMapView::ParameterAddEnabled() const
 }
 
 void SNiagaraParameterMapView::AddParameter(FNiagaraVariable NewVariable)
+{
+	bool bEnterRenameModeOnAdd = true;
+	AddParameter(NewVariable, bEnterRenameModeOnAdd);
+}
+
+void SNiagaraParameterMapView::AddParameter(FNiagaraVariable NewVariable, bool bEnterRenameModeOnAdd)
 {
 	TGuardValue<bool> AddParameterRefreshGuard(bIsAddingParameter, true);
 	FNiagaraParameterHandle ParameterHandle;
@@ -449,7 +455,10 @@ void SNiagaraParameterMapView::AddParameter(FNiagaraVariable NewVariable)
 	{
 		GraphActionMenu->RefreshAllActions(true);
 		GraphActionMenu->SelectItemByName(NewVariable.GetName());
-		GraphActionMenu->OnRequestRenameOnActionNode();
+		if (bEnterRenameModeOnAdd)
+		{
+			GraphActionMenu->OnRequestRenameOnActionNode();
+		}
 	}
 }
 
@@ -1264,18 +1273,6 @@ bool SNiagaraParameterMapView::GetNamespaceEditDataForSelection(
 		return false;
 	}
 
-	if (ToolkitType == SYSTEM)
-	{
-		bool bWasCreated = OutParameterAction->ReferenceCollection.ContainsByPredicate([](FNiagaraGraphParameterReferenceCollection& ReferenceCollecion) { return ReferenceCollecion.WasCreated(); });
-		if (bWasCreated == false)
-		{
-			// Can only modify parameters in the system toolkit if they were added there.
-			OutParameterAction.Reset();
-			OutErrorMessage = LOCTEXT("NonAddedSystem", "Can only edit system parameteters which were added directly.");
-			return false;
-		}
-	}
-
 	OutParameterHandle = FNiagaraParameterHandle(OutParameterAction->Parameter.GetName());
 	TArray<FName> NameParts = OutParameterHandle.GetHandleParts();
 	OutNamespaceMetadata = GetDefault<UNiagaraEditorSettings>()->GetMetaDataForNamespaces(NameParts);
@@ -1504,12 +1501,58 @@ void SNiagaraParameterMapView::RenameParameter(FNiagaraVariable& Parameter, FNam
 	{
 		if (Graphs.Num() > 0)
 		{
-			for (const TWeakObjectPtr<UNiagaraGraph>& Graph : Graphs)
+			for (const TWeakObjectPtr<UNiagaraGraph>& GraphWeak : Graphs)
 			{
-				if (Graph.IsValid())
+				UNiagaraGraph* Graph = GraphWeak.Get();
+				if (Graph == nullptr)
+				{
+					// Ignore invalid graphs.
+					continue;
+				}
+
+				const FNiagaraGraphParameterReferenceCollection* ReferenceCollection = Graph->GetParameterReferenceMap().Find(Parameter);
+				if (ensureMsgf(ReferenceCollection != nullptr, TEXT("Parameter in view which wasn't in the reference collection.")) == false)
+				{
+					// Can't handle parameters with no reference collections.
+					continue;
+				}
+
+				if (ReferenceCollection->WasCreated())
 				{
-					Graph.Get()->RenameParameter(Parameter, NewName);
+					// If the parameter was created directly by the user just rename it.
+					Graph->RenameParameter(Parameter, NewName);
 				}
+				else
+				{
+					// Otherwise see if it has local graph references.  If so it can be renamed, if not than we have to
+					// create a new parameter with the new name since we can't rename parameters from nested graphs.
+					bool bHasLocalGraphReference = false;
+					for (const FNiagaraGraphParameterReference& Reference : ReferenceCollection->ParameterReferences)
+					{
+						UNiagaraNode* ReferencingNode = Cast<UNiagaraNode>(Reference.Value.Get());
+						if (ReferencingNode != nullptr && ReferencingNode->GetGraph() == Graph)
+						{
+							bHasLocalGraphReference = true;
+							break;
+						}
+					}
+
+					if (bHasLocalGraphReference)
+					{
+						Graph->RenameParameter(Parameter, NewName);
+					}
+					else
+					{
+						FNiagaraEditorUtilities::InfoWithToastAndLog(FText::Format(
+							LOCTEXT("AddInsteadOfRenameInScriptFormat", "The parameter {0} was not added directly\n or was discovered in a nested graph and can't be renamed.\nA new parameter {1} has been added instead."),
+							FNiagaraEditorUtilities::FormatParameterNameForTextDisplay(Parameter.GetName()),
+							FNiagaraEditorUtilities::FormatParameterNameForTextDisplay(NewName)), 7.0f);
+
+						FNiagaraVariable NewVariable(Parameter.GetType(), NewName);
+						bool bEnterRenameModeOnAdd = false;
+						AddParameter(NewVariable, bEnterRenameModeOnAdd);
+					}
+				}		
 			}
 		}
 	}
@@ -1519,8 +1562,23 @@ void SNiagaraParameterMapView::RenameParameter(FNiagaraVariable& Parameter, FNam
 		if (System != nullptr)
 		{
 			System->Modify();
-			System->GetExposedParameters().RenameParameter(Parameter, NewName);
-			System->EditorOnlyAddedParameters.RenameParameter(Parameter, NewName);
+			if (System->GetExposedParameters().IndexOf(Parameter) == INDEX_NONE && System->EditorOnlyAddedParameters.IndexOf(Parameter) == INDEX_NONE)
+			{
+				// If the parameter isn't in one of the existing collections than it's a discovered parameter and instead of renaming it needs to be added.
+				FNiagaraEditorUtilities::InfoWithToastAndLog(FText::Format(
+					LOCTEXT("AddInsteadOfRenameInSystemFormat", "The parameter {0}\n was not added directly and can't be renamed.\nA new parameter {1} has been added instead."),
+					FNiagaraEditorUtilities::FormatParameterNameForTextDisplay(Parameter.GetName()),
+					FNiagaraEditorUtilities::FormatParameterNameForTextDisplay(NewName)), 7.0f);
+
+				FNiagaraVariable NewVariable(Parameter.GetType(), NewName);
+				bool bEnterRenameModeOnAdd = false;
+				AddParameter(NewVariable, bEnterRenameModeOnAdd);
+			}
+			else
+			{
+				System->GetExposedParameters().RenameParameter(Parameter, NewName);
+				System->EditorOnlyAddedParameters.RenameParameter(Parameter, NewName);
+			}
 		}
 	}
 }
diff --git a/Engine/Plugins/FX/Niagara/Source/NiagaraEditor/Public/NiagaraEditorSettings.h b/Engine/Plugins/FX/Niagara/Source/NiagaraEditor/Public/NiagaraEditorSettings.h
--- a/Engine/Plugins/FX/Niagara/Source/NiagaraEditor/Public/NiagaraEditorSettings.h
+++ b/Engine/Plugins/FX/Niagara/Source/NiagaraEditor/Public/NiagaraEditorSettings.h
@@ -60,6 +60,9 @@ struct FNiagaraNamespaceMetadata
 	UPROPERTY()
 	FLinearColor BackgroundColor;
 
+	UPROPERTY()
+	FName ForegroundStyle;
+
 	UPROPERTY()
 	TArray<ENiagaraNamespaceMetadataOptions> Options;
 
@@ -87,6 +90,12 @@ struct FNiagaraNamespaceMetadata
 		return *this;
 	}
 
+	FNiagaraNamespaceMetadata& SetForegroundStyle(FName InForegroundStyle)
+	{
+		ForegroundStyle = InForegroundStyle;
+		return *this;
+	}
+
 	FNiagaraNamespaceMetadata& AddOption(ENiagaraNamespaceMetadataOptions Option)
 	{
 		Options.Add(Option);
diff --git a/Engine/Plugins/FX/Niagara/Source/NiagaraEditor/Public/NiagaraEditorUtilities.h b/Engine/Plugins/FX/Niagara/Source/NiagaraEditor/Public/NiagaraEditorUtilities.h
--- a/Engine/Plugins/FX/Niagara/Source/NiagaraEditor/Public/NiagaraEditorUtilities.h
+++ b/Engine/Plugins/FX/Niagara/Source/NiagaraEditor/Public/NiagaraEditorUtilities.h
@@ -243,7 +243,7 @@ namespace FNiagaraEditorUtilities
 	void CreateAssetFromEmitter(TSharedRef<FNiagaraEmitterHandleViewModel> EmitterHandleViewModel);
 
 	NIAGARAEDITOR_API void WarnWithToastAndLog(FText WarningMessage);
-	NIAGARAEDITOR_API void InfoWithToastAndLog(FText WarningMessage);
+	NIAGARAEDITOR_API void InfoWithToastAndLog(FText WarningMessage, float ToastDuration = 5.0f);
 
 	void GetScriptRunAndExecutionIndexFromUsage(const ENiagaraScriptUsage& InUsage, int32& OutRunIndex, int32&OutExecutionIndex);
 
@@ -291,4 +291,6 @@ namespace FNiagaraEditorUtilities
 	FString GetNamespacelessVariableNameString(const FName& InVarName);
 
 	void GetReferencingFunctionCallNodes(UNiagaraScript* Script, TArray<UNiagaraNodeFunctionCall*>& OutReferencingFunctionCallNodes);
+
+	FText FormatParameterNameForTextDisplay(FName ParameterName);
 };
diff --git a/Engine/Plugins/FX/Niagara/Source/NiagaraEditor/Public/Widgets/SNiagaraParameterMapView.h b/Engine/Plugins/FX/Niagara/Source/NiagaraEditor/Public/Widgets/SNiagaraParameterMapView.h
--- a/Engine/Plugins/FX/Niagara/Source/NiagaraEditor/Public/Widgets/SNiagaraParameterMapView.h
+++ b/Engine/Plugins/FX/Niagara/Source/NiagaraEditor/Public/Widgets/SNiagaraParameterMapView.h
@@ -85,6 +85,7 @@ public:
 
 	/** Adds parameter to the graph parameter store and refreshes the menu. */
 	void AddParameter(FNiagaraVariable NewVariable);
+	void AddParameter(FNiagaraVariable NewVariable, bool bEnterRenameModeOnAdd);
 
 	/** Refreshes the graphs used for this menu. */
 	void Refresh(bool bRefreshMenu = true);
diff --git a/Engine/Plugins/FX/Niagara/Source/NiagaraEditor/Public/Widgets/SNiagaraParameterName.h b/Engine/Plugins/FX/Niagara/Source/NiagaraEditor/Public/Widgets/SNiagaraParameterName.h
--- a/Engine/Plugins/FX/Niagara/Source/NiagaraEditor/Public/Widgets/SNiagaraParameterName.h
+++ b/Engine/Plugins/FX/Niagara/Source/NiagaraEditor/Public/Widgets/SNiagaraParameterName.h
@@ -33,6 +33,7 @@ public:
 		SLATE_EVENT(FOnNameChanged, OnNameChanged)
 		SLATE_EVENT(FIsSelected, IsSelected)
 		SLATE_EVENT(FPointerEventHandler, OnDoubleClicked)
+		SLATE_NAMED_SLOT(FArguments, Decorator)
 	SLATE_END_ARGS()
 
 	void Construct(const FArguments& InArgs);
@@ -46,7 +47,7 @@ public:
 	void EnterNamespaceModifierEditingMode();
 
 private:
-	TSharedRef<SBorder> CreateNamespaceWidget(FText NamespaceDisplayName, FText NamespaceDescription, FLinearColor NamespaceBorderColor);
+	TSharedRef<SBorder> CreateNamespaceWidget(FText NamespaceDisplayName, FText NamespaceDescription, FLinearColor NamespaceBorderColor, FName NamespaceForegroundStyle);
 
 	void UpdateContent(FName InDisplayedParameterName);
 
@@ -78,6 +79,7 @@ private:
 	FName DisplayedNamespaceModifier;
 	TSharedPtr<SInlineEditableTextBlock> EditableTextBlock;
 	TSharedPtr<SBorder> NamespaceModifierBorder;
+	TSharedPtr<SWidget> Decorator;
 };
 
 class NIAGARAEDITOR_API SNiagaraParameterNameTextBlock : public SCompoundWidget
@@ -95,6 +97,8 @@ public:
 		SLATE_EVENT(FOnVerifyTextChanged, OnVerifyTextChanged)
 		SLATE_EVENT(FOnTextCommitted, OnTextCommitted)
 		SLATE_EVENT(FIsSelected, IsSelected)
+		SLATE_NAMED_SLOT(FArguments, Decorator)
+		
 	SLATE_END_ARGS()
 
 	void Construct(const FArguments& InArgs);
