Fix for overlapping constant names which could result in reading past the end of niagara script literals

diff --git a/Engine/Plugins/FX/Niagara/Source/NiagaraCore/Private/NiagaraCustomVersion.cpp b/Engine/Plugins/FX/Niagara/Source/NiagaraCore/Private/NiagaraCustomVersion.cpp
--- a/Engine/Plugins/FX/Niagara/Source/NiagaraCore/Private/NiagaraCustomVersion.cpp
+++ b/Engine/Plugins/FX/Niagara/Source/NiagaraCore/Private/NiagaraCustomVersion.cpp
@@ -8,4 +8,4 @@ const FGuid FNiagaraCustomVersion::GUID(0xFCF57AFA, 0x50764283, 0xB9A9E658, 0xFF
 // Register the custom version with core
 FCustomVersionRegistration GRegisterNiagaraCustomVersion(FNiagaraCustomVersion::GUID, FNiagaraCustomVersion::LatestVersion, TEXT("NiagaraVer"));
 
-const FGuid FNiagaraCustomVersion::LatestScriptCompileVersion(0xF4357B0E, 0x349C4526, 0x97A94D8E, 0x85A7380D);
\ No newline at end of file
+const FGuid FNiagaraCustomVersion::LatestScriptCompileVersion(0x562A1A9C, 0xD29944FC, 0x82B4CE57, 0xC8FC4599);
\ No newline at end of file
diff --git a/Engine/Plugins/FX/Niagara/Source/NiagaraEditor/Private/NiagaraCompiler.cpp b/Engine/Plugins/FX/Niagara/Source/NiagaraEditor/Private/NiagaraCompiler.cpp
--- a/Engine/Plugins/FX/Niagara/Source/NiagaraEditor/Private/NiagaraCompiler.cpp
+++ b/Engine/Plugins/FX/Niagara/Source/NiagaraEditor/Private/NiagaraCompiler.cpp
@@ -1153,6 +1153,7 @@ TOptional<FNiagaraCompileResults> FHlslNiagaraCompiler::GetCompileResult(int32 J
 		Results.Data->InternalParameters.Empty();
 		for (int32 i = 0; i < CompilationOutput.InternalConstantOffsets.Num(); ++i)
 		{
+			const FName ConstantName(TEXT("InternalConstant"), i);
 			EVectorVMBaseTypes Type = CompilationOutput.InternalConstantTypes[i];
 			int32 Offset = CompilationOutput.InternalConstantOffsets[i];
 			switch (Type)
@@ -1160,19 +1161,19 @@ TOptional<FNiagaraCompileResults> FHlslNiagaraCompiler::GetCompileResult(int32 J
 			case EVectorVMBaseTypes::Float:
 			{
 				float Val = *(float*)(CompilationOutput.InternalConstantData.GetData() + Offset);
-				Results.Data->InternalParameters.SetOrAdd(FNiagaraVariable(FNiagaraTypeDefinition::GetFloatDef(), *LexToString(Val)))->SetValue(Val);
+				Results.Data->InternalParameters.SetOrAdd(FNiagaraVariable(FNiagaraTypeDefinition::GetFloatDef(), ConstantName))->SetValue(Val);
 			}
 			break;
 			case EVectorVMBaseTypes::Int:
 			{
 				int32 Val = *(int32*)(CompilationOutput.InternalConstantData.GetData() + Offset);
-				Results.Data->InternalParameters.SetOrAdd(FNiagaraVariable(FNiagaraTypeDefinition::GetIntDef(), *LexToString(Val)))->SetValue(Val);
+				Results.Data->InternalParameters.SetOrAdd(FNiagaraVariable(FNiagaraTypeDefinition::GetIntDef(), ConstantName))->SetValue(Val);
 			}
 			break;
 			case EVectorVMBaseTypes::Bool:
 			{
 				int32 Val = *(int32*)(CompilationOutput.InternalConstantData.GetData() + Offset);
-				Results.Data->InternalParameters.SetOrAdd(FNiagaraVariable(FNiagaraTypeDefinition::GetIntDef(), Val == 0 ? TEXT("FALSE") : TEXT("TRUE")))->SetValue(Val);
+				Results.Data->InternalParameters.SetOrAdd(FNiagaraVariable(FNiagaraTypeDefinition::GetIntDef(), ConstantName))->SetValue(Val);
 			}
 			break;
 			}
