Remove logic for disabled SNiagaraParameterPanel implementation which would cause an assert on creating new DataInstance namespace parameters.

diff --git a/Engine/Plugins/FX/Niagara/Source/Niagara/Public/NiagaraTypes.h b/Engine/Plugins/FX/Niagara/Source/Niagara/Public/NiagaraTypes.h
--- a/Engine/Plugins/FX/Niagara/Source/Niagara/Public/NiagaraTypes.h
+++ b/Engine/Plugins/FX/Niagara/Source/Niagara/Public/NiagaraTypes.h
@@ -770,14 +770,14 @@ private:
 	bool bOutputIsPersistent;
 
 	/** Namespace-less name for associated FNiagaraVariable. Edited directly by user and then used to generate full Name of associated FNiagaraVariable. */
-	UPROPERTY(VisibleAnywhere, Category = "Variable", DisplayName = "Property Name", meta = (SkipForCompileHash = "true"))
+	UPROPERTY(DisplayName = "Property Name", meta = (SkipForCompileHash = "true"))
 	FName CachedNamespacelessVariableName;
 
 	/** Track if the associated parameter was created in the Emitter/System editor. Used to determine whether the associated parameter can be deleted from the Emitter/System editor. */
 	UPROPERTY(meta = (SkipForCompileHash = "true"))
 	bool bCreatedInSystemEditor;
 
-	UPROPERTY(EditAnywhere, Category = "Custom Name", meta = (ToolTip = "Enable using a legacy custom name string.", SkipForCompileHash = "true"))
+	UPROPERTY(meta = (ToolTip = "Enable using a legacy custom name string.", SkipForCompileHash = "true"))
 	bool bUseLegacyNameString;
 
 public:
diff --git a/Engine/Plugins/FX/Niagara/Source/NiagaraEditor/Private/NiagaraNodeWithDynamicPins.cpp b/Engine/Plugins/FX/Niagara/Source/NiagaraEditor/Private/NiagaraNodeWithDynamicPins.cpp
--- a/Engine/Plugins/FX/Niagara/Source/NiagaraEditor/Private/NiagaraNodeWithDynamicPins.cpp
+++ b/Engine/Plugins/FX/Niagara/Source/NiagaraEditor/Private/NiagaraNodeWithDynamicPins.cpp
@@ -310,10 +310,6 @@ void UNiagaraNodeWithDynamicPins::AddParameter(FNiagaraVariable Parameter, const
 
 		Graph->Modify();
 		Graph->AddParameter(Parameter, AddParameterOptions);
-
-		// Add the Parameter Reference 
-		FNiagaraGraphParameterReference NewParameterReference = FNiagaraGraphParameterReference(Pin->PersistentGuid, this);
-		Graph->AddParameterReference(Parameter, NewParameterReference);
 	}
 	else
 	{
@@ -336,9 +332,7 @@ void UNiagaraNodeWithDynamicPins::AddParameter(FNiagaraVariable Parameter, const
 		Graph->AddParameter(Parameter, AddParameterOptions);
 
 		Modify();
-		UEdGraphPin* Pin = RequestNewTypedPin(NewPinDirection, Parameter.GetType(), Parameter.GetName());
-		FNiagaraGraphParameterReference NewParameterReference = FNiagaraGraphParameterReference(Pin->PersistentGuid, this);
-		Graph->AddParameterReference(Parameter, NewParameterReference);
+		RequestNewTypedPin(NewPinDirection, Parameter.GetType(), Parameter.GetName());
 	}
 	else
 	{
