Fix for android shader compiler not liking float4x4(float3x4, float4) constructor

diff --git a/Engine/Plugins/FX/Niagara/Shaders/Private/NiagaraMeshVertexFactory.ush b/Engine/Plugins/FX/Niagara/Shaders/Private/NiagaraMeshVertexFactory.ush
--- a/Engine/Plugins/FX/Niagara/Shaders/Private/NiagaraMeshVertexFactory.ush
+++ b/Engine/Plugins/FX/Niagara/Shaders/Private/NiagaraMeshVertexFactory.ush
@@ -322,11 +322,19 @@ FMaterialPixelParameters GetMaterialPixelParameters(FVertexFactoryInterpolantsVS
 #endif
 
 #if USE_PARTICLE_LOCAL_TO_WORLD
-	Result.Particle.ParticleToWorld = transpose(float4x4(Interpolants.ParticleToWorld, float4(0.0f, 0.0f, 0.0f, 1.0f)));
+	Result.Particle.ParticleToWorld = transpose(float4x4(
+		Interpolants.ParticleToWorld[0],
+		Interpolants.ParticleToWorld[1],
+		Interpolants.ParticleToWorld[2],
+		float4(0.0f, 0.0f, 0.0f, 1.0f)));
 #endif
 
 #if USE_PARTICLE_WORLD_TO_LOCAL
-	Result.Particle.WorldToParticle = transpose(float4x4(Interpolants.WorldToParticle, float4(0.0f, 0.0f, 0.0f, 1.0f)));
+	Result.Particle.WorldToParticle = transpose(float4x4(
+		Interpolants.WorldToParticle[0],
+		Interpolants.WorldToParticle[1],
+		Interpolants.WorldToParticle[2],
+		float4(0.0f, 0.0f, 0.0f, 1.0f)));
 #endif
 
 	Result.Particle.MotionBlurFade = 1.0f;
diff --git a/Engine/Shaders/Private/MeshParticleVertexFactory.ush b/Engine/Shaders/Private/MeshParticleVertexFactory.ush
--- a/Engine/Shaders/Private/MeshParticleVertexFactory.ush
+++ b/Engine/Shaders/Private/MeshParticleVertexFactory.ush
@@ -292,7 +292,11 @@ FMaterialPixelParameters GetMaterialPixelParameters(FVertexFactoryInterpolantsVS
 #endif
 
 #if NEEDS_PARTICLE_LOCAL_TO_WORLD || NEEDS_PARTICLE_WORLD_TO_LOCAL
-	float4x4 ParticleTransform = transpose(float4x4(Interpolants.LocalToWorld, float4(0.0f, 0.0f, 0.0f, 1.0f)));
+	float4x4 ParticleTransform = transpose(float4x4(
+		Interpolants.LocalToWorld[0],
+		Interpolants.LocalToWorld[1],
+		Interpolants.LocalToWorld[2],
+		float4(0.0f, 0.0f, 0.0f, 1.0f)));
 
 	#if NEEDS_PARTICLE_LOCAL_TO_WORLD
 		Result.Particle.ParticleToWorld = ParticleTransform;
