Niagara - Separate advanced options for parameter sections into system and script.

diff --git a/Engine/Plugins/FX/Niagara/Source/NiagaraEditor/Private/NiagaraEditorSettings.cpp b/Engine/Plugins/FX/Niagara/Source/NiagaraEditor/Private/NiagaraEditorSettings.cpp
--- a/Engine/Plugins/FX/Niagara/Source/NiagaraEditor/Private/NiagaraEditorSettings.cpp
+++ b/Engine/Plugins/FX/Niagara/Source/NiagaraEditor/Private/NiagaraEditorSettings.cpp
@@ -61,24 +61,28 @@ void UNiagaraEditorSettings::SetupNamespaceMetadata()
 			.SetDisplayName(LOCTEXT("ModuleOutputDisplayName", "Output"))
 			.SetDescription(LOCTEXT("ModuleOutputDescription", "A transient value which is unique to the module that wrote it and can be read from any other module.\nTransient values do not persist from frame to frame, or between stages, e.g. emitter to particle, or spawn to update."))
 			.SetBackgroundColor(FLinearColor(FColor(109, 95, 124)))
-			.AddOption(ENiagaraNamespaceMetadataOptions::Advanced)
+			.AddOption(ENiagaraNamespaceMetadataOptions::AdvancedInScript)
+			.AddOption(ENiagaraNamespaceMetadataOptions::AdvancedInSystem)
 			.AddOption(ENiagaraNamespaceMetadataOptions::PreventCreatingInSystemEditor)
 			.AddOption(ENiagaraNamespaceMetadataOptions::CanChangeNamespaceModifier),
 		FNiagaraNamespaceMetadata({FNiagaraConstants::OutputNamespace})
 			.SetDisplayName(LOCTEXT("OutputDisplayName", "Output"))
 			.SetDescription(LOCTEXT("OutputDescription", "A transient value which is unique to the module that wrote it and can be read from any other module.\nTransient values do not persist from frame to frame, or between stages, e.g. emitter to particle, or spawn to update."))
 			.SetBackgroundColor(FLinearColor(FColor(109, 95, 124)))
-			.AddOption(ENiagaraNamespaceMetadataOptions::Advanced),
+			.AddOption(ENiagaraNamespaceMetadataOptions::AdvancedInScript)
+			.AddOption(ENiagaraNamespaceMetadataOptions::AdvancedInSystem),
 		FNiagaraNamespaceMetadata({FNiagaraConstants::LocalNamespace, FNiagaraConstants::ModuleNamespace})
 			.SetDisplayName(LOCTEXT("ModuleLocalDisplayName", "Local"))
 			.SetDescription(LOCTEXT("ModuleLocalDescription", "A transient value which can be written to and read from within a single module.\nTransient values do not persist from frame to frame, or between stages, e.g. emitter to particle, or spawn to update."))
 			.SetBackgroundColor(FLinearColor(FColor(174, 157, 40)))
-			.AddOption(ENiagaraNamespaceMetadataOptions::Advanced),
+			.AddOption(ENiagaraNamespaceMetadataOptions::AdvancedInScript)
+			.AddOption(ENiagaraNamespaceMetadataOptions::AdvancedInSystem),
 		FNiagaraNamespaceMetadata({FNiagaraConstants::TransientNamespace})
 			.SetDisplayName(LOCTEXT("TransientDisplayName", "Transient"))
 			.SetDescription(LOCTEXT("TransientDescription", "A transient value which can be written to and read from from any module.\nTransient values do not persist from frame to frame, or between stages, e.g. emitter to particle, or spawn to update."))
 			.SetBackgroundColor(FLinearColor(FColor(109, 95, 124)))
-			.AddOption(ENiagaraNamespaceMetadataOptions::Advanced),
+			.AddOption(ENiagaraNamespaceMetadataOptions::AdvancedInScript)
+			.AddOption(ENiagaraNamespaceMetadataOptions::AdvancedInSystem),
 		FNiagaraNamespaceMetadata({FNiagaraConstants::EngineNamespace})
 			.SetDisplayName(LOCTEXT("EngineDisplayName", "Engine"))
 			.SetDescription(LOCTEXT("EngineDescription", "A read only value which is provided by the engine.\nThis value's source can be the simulation itsef\ne.g. ExecutionCount, or the owner of the simulation e.g. (Owner) Scale."))
@@ -95,14 +99,16 @@ void UNiagaraEditorSettings::SetupNamespaceMetadata()
 			.SetDescription(LOCTEXT("NiagaraParameterCollectionDescription", "Values read from a niagara parameter collection asset.\nRead only in a niagara system."))
 			.SetBackgroundColor(FLinearColor(FColor(170, 170, 170)))
 			.SetForegroundStyle("NiagaraEditor.ParameterName.NamespaceTextDark")
-			.AddOption(ENiagaraNamespaceMetadataOptions::Advanced)
+			.AddOption(ENiagaraNamespaceMetadataOptions::AdvancedInScript)
+			.AddOption(ENiagaraNamespaceMetadataOptions::AdvancedInSystem)
 			.AddOption(ENiagaraNamespaceMetadataOptions::PreventRenaming),
 		FNiagaraNamespaceMetadata({FNiagaraConstants::DataInstanceNamespace})
 			.SetDisplayName(LOCTEXT("DataInstanceDisplayName", "Data Instance"))
 			.SetDescription(LOCTEXT("DataInstanceDescription", "A special transient value which has a single bool IsAlive value, which determines if a particle is alive or not."))
 			.SetBackgroundColor(FLinearColor(FColor(170, 170, 170)))
 			.SetForegroundStyle("NiagaraEditor.ParameterName.NamespaceTextDark")
-			.AddOption(ENiagaraNamespaceMetadataOptions::Advanced)
+			.AddOption(ENiagaraNamespaceMetadataOptions::AdvancedInScript)
+			.AddOption(ENiagaraNamespaceMetadataOptions::AdvancedInSystem)
 			.AddOption(ENiagaraNamespaceMetadataOptions::PreventRenaming),
 	};
 
diff --git a/Engine/Plugins/FX/Niagara/Source/NiagaraEditor/Private/Widgets/SNiagaraParameterMapView.cpp b/Engine/Plugins/FX/Niagara/Source/NiagaraEditor/Private/Widgets/SNiagaraParameterMapView.cpp
--- a/Engine/Plugins/FX/Niagara/Source/NiagaraEditor/Private/Widgets/SNiagaraParameterMapView.cpp
+++ b/Engine/Plugins/FX/Niagara/Source/NiagaraEditor/Private/Widgets/SNiagaraParameterMapView.cpp
@@ -174,27 +174,20 @@ NiagaraParameterMapSectionID::Type NiagaraParameterMapSectionID::OnGetSectionFro
 	return SectionID;
 }
 
-bool NiagaraParameterMapSectionID::GetSectionIsAdvanced(const Type InSection)
+bool NiagaraParameterMapSectionID::GetSectionIsAdvancedForScript(const Type InSection)
 {
 	TArray<FName> SectionNamespaces;
 	OnGetSectionNamespaces(InSection, SectionNamespaces);
 	FNiagaraNamespaceMetadata NamespaceMetadata = GetDefault<UNiagaraEditorSettings>()->GetMetaDataForNamespaces(SectionNamespaces);
-	if (NamespaceMetadata.IsValid())
-	{
-		bool bIsAdvanced = false;
-		for (const ENiagaraNamespaceMetadataOptions& Option : NamespaceMetadata.Options)
-		{
-			if (Option == ENiagaraNamespaceMetadataOptions::Advanced)
-			{
-				return true;
-			}
-		}
-		return false;
-	}
-	else
-	{
-		return false;
-	}
+	return NamespaceMetadata.IsValid() && NamespaceMetadata.Options.Contains(ENiagaraNamespaceMetadataOptions::AdvancedInScript);
+}
+
+bool NiagaraParameterMapSectionID::GetSectionIsAdvancedForSystem(const Type InSection)
+{
+	TArray<FName> SectionNamespaces;
+	OnGetSectionNamespaces(InSection, SectionNamespaces);
+	FNiagaraNamespaceMetadata NamespaceMetadata = GetDefault<UNiagaraEditorSettings>()->GetMetaDataForNamespaces(SectionNamespaces);
+	return NamespaceMetadata.IsValid() && NamespaceMetadata.Options.Contains(ENiagaraNamespaceMetadataOptions::AdvancedInSystem);
 }
 
 SNiagaraParameterMapView::~SNiagaraParameterMapView()
@@ -549,19 +542,27 @@ void SNiagaraParameterMapView::CollectAllActions(FGraphActionListBuilderBase& Ou
 		{
 			continue;
 		}
-		TArray<NiagaraParameterMapSectionID::Type> ExcludedSystemIds = {
+
+		static const TArray<NiagaraParameterMapSectionID::Type> ExcludedSystemIds = {
 			NiagaraParameterMapSectionID::MODULE_INPUT,
 			NiagaraParameterMapSectionID::MODULE_LOCAL,
 			NiagaraParameterMapSectionID::DATA_INSTANCE };
+		if (IsSystemToolkit() && ExcludedSystemIds.Contains(Section))
+		{
+			continue;
+		}
 
-		if ((IsSystemToolkit() == false || ExcludedSystemIds.Contains(Section) == false) &&
-			(bDisplayAdvancedCategories || NiagaraParameterMapSectionID::GetSectionIsAdvanced(Section) == false))
+		if (bDisplayAdvancedCategories == false &&
+			((IsScriptToolkit() && NiagaraParameterMapSectionID::GetSectionIsAdvancedForScript(Section)) ||
+			(IsSystemToolkit() && NiagaraParameterMapSectionID::GetSectionIsAdvancedForSystem(Section))))
 		{
-			const FText Name = FText::FromName(Parameter.GetName());
-			const FText Tooltip = FText::Format(TooltipFormat, FText::FromName(Parameter.GetName()), Parameter.GetType().GetNameText());
-			TSharedPtr<FNiagaraParameterAction> ParameterAction(new FNiagaraParameterAction(Parameter, ParameterEntry.Value, FText::GetEmpty(), Name, Tooltip, 0, FText(), Section));
-			OutAllActions.AddAction(ParameterAction);
+			continue;
 		}
+
+		const FText Name = FText::FromName(Parameter.GetName());
+		const FText Tooltip = FText::Format(TooltipFormat, FText::FromName(Parameter.GetName()), Parameter.GetType().GetNameText());
+		TSharedPtr<FNiagaraParameterAction> ParameterAction(new FNiagaraParameterAction(Parameter, ParameterEntry.Value, FText::GetEmpty(), Name, Tooltip, 0, FText(), Section));
+		OutAllActions.AddAction(ParameterAction);
 	}
 }
 
@@ -735,7 +736,18 @@ void SNiagaraParameterMapView::CollectStaticSections(TArray<int32>& StaticSectio
 	{
 		for (int i = StaticSectionIDs.Num() - 1; i > 0; --i)
 		{
-			if (NiagaraParameterMapSectionID::GetSectionIsAdvanced(NiagaraParameterMapSectionID::Type(StaticSectionIDs[i])))
+			bool bIsAdvanced = false;
+			NiagaraParameterMapSectionID::Type SectionID = NiagaraParameterMapSectionID::Type(StaticSectionIDs[i]);
+			if (IsScriptToolkit())
+			{
+				bIsAdvanced = NiagaraParameterMapSectionID::GetSectionIsAdvancedForScript(SectionID);
+			}
+			else if (IsSystemToolkit())
+			{
+				bIsAdvanced = NiagaraParameterMapSectionID::GetSectionIsAdvancedForSystem(SectionID);
+			}
+
+			if (bIsAdvanced)
 			{
 				StaticSectionIDs.RemoveAt(i);
 			}
diff --git a/Engine/Plugins/FX/Niagara/Source/NiagaraEditor/Public/NiagaraEditorSettings.h b/Engine/Plugins/FX/Niagara/Source/NiagaraEditor/Public/NiagaraEditorSettings.h
--- a/Engine/Plugins/FX/Niagara/Source/NiagaraEditor/Public/NiagaraEditorSettings.h
+++ b/Engine/Plugins/FX/Niagara/Source/NiagaraEditor/Public/NiagaraEditorSettings.h
@@ -30,7 +30,8 @@ struct FNiagaraNewAssetDialogConfig
 UENUM()
 enum class ENiagaraNamespaceMetadataOptions
 {
-	Advanced,
+	AdvancedInScript,
+	AdvancedInSystem,
 	PreventRenaming,
 	CanChangeNamespaceModifier,
 	PreventCreatingInSystemEditor
diff --git a/Engine/Plugins/FX/Niagara/Source/NiagaraEditor/Public/Widgets/SNiagaraParameterMapView.h b/Engine/Plugins/FX/Niagara/Source/NiagaraEditor/Public/Widgets/SNiagaraParameterMapView.h
--- a/Engine/Plugins/FX/Niagara/Source/NiagaraEditor/Public/Widgets/SNiagaraParameterMapView.h
+++ b/Engine/Plugins/FX/Niagara/Source/NiagaraEditor/Public/Widgets/SNiagaraParameterMapView.h
@@ -57,7 +57,8 @@ namespace NiagaraParameterMapSectionID
 	static FText OnGetSectionTitle(const NiagaraParameterMapSectionID::Type InSection);
 	static void OnGetSectionNamespaces(const NiagaraParameterMapSectionID::Type InSection, TArray<FName>& OutSectionNamespaces);
 	static NiagaraParameterMapSectionID::Type OnGetSectionFromVariable(const FNiagaraVariable& InVar, bool IsStaticSwitchVariable, FNiagaraParameterHandle& OutParameterHandle, const NiagaraParameterMapSectionID::Type DefaultType = NiagaraParameterMapSectionID::Type::NONE);
-	static bool GetSectionIsAdvanced(const NiagaraParameterMapSectionID::Type InSection);
+	static bool GetSectionIsAdvancedForScript(const NiagaraParameterMapSectionID::Type InSection);
+	static bool GetSectionIsAdvancedForSystem(const NiagaraParameterMapSectionID::Type InSection);
 };
 
 /** A widget for viewing and editing a set of selected objects with a details panel. */
