Niagara - One last update the graph fixup on load, and the collision assets refixed.

diff --git a/Engine/Plugins/FX/Niagara/Source/NiagaraCore/Private/NiagaraCustomVersion.cpp b/Engine/Plugins/FX/Niagara/Source/NiagaraCore/Private/NiagaraCustomVersion.cpp
--- a/Engine/Plugins/FX/Niagara/Source/NiagaraCore/Private/NiagaraCustomVersion.cpp
+++ b/Engine/Plugins/FX/Niagara/Source/NiagaraCore/Private/NiagaraCustomVersion.cpp
@@ -8,4 +8,4 @@ const FGuid FNiagaraCustomVersion::GUID(0xFCF57AFA, 0x50764283, 0xB9A9E658, 0xFF
 // Register the custom version with core
 FCustomVersionRegistration GRegisterNiagaraCustomVersion(FNiagaraCustomVersion::GUID, FNiagaraCustomVersion::LatestVersion, TEXT("NiagaraVer"));
 
-const FGuid FNiagaraCustomVersion::LatestScriptCompileVersion(0x562A1A9C, 0xD29944FC, 0x82B4CE57, 0xC8FC4599);
\ No newline at end of file
+const FGuid FNiagaraCustomVersion::LatestScriptCompileVersion(0xFB5CE6AE, 0xE8FF4C4D, 0xBB10A479, 0x3BE0D367);
diff --git a/Engine/Plugins/FX/Niagara/Source/NiagaraEditor/Private/NiagaraGraph.cpp b/Engine/Plugins/FX/Niagara/Source/NiagaraEditor/Private/NiagaraGraph.cpp
--- a/Engine/Plugins/FX/Niagara/Source/NiagaraEditor/Private/NiagaraGraph.cpp
+++ b/Engine/Plugins/FX/Niagara/Source/NiagaraEditor/Private/NiagaraGraph.cpp
@@ -605,7 +605,14 @@ FName StandardizeName(FName Name, ENiagaraScriptUsage Usage, bool bIsGet, bool b
 	else
 	{
 		// Namespace was unknown.
-		if (bIsScriptUsage || bIsGet)
+		if (bIsScriptUsage && NameParts.Contains(FNiagaraConstants::ModuleNamespace))
+		{
+			// If we're in a script check for a misplaced module namespace and if it has one, force it to be a 
+			// module output to help with fixing up usages.
+			static const FName OutputScriptNamespace = *(FNiagaraConstants::OutputNamespace.ToString() + TEXT(".") + FNiagaraConstants::ModuleNamespace.ToString());
+			Namespace = OutputScriptNamespace;
+		}
+		else if(bIsScriptUsage || bIsGet)
 		{
 			// If it's in a get node or it's in a script, force it into the transient namespace.
 			Namespace = FNiagaraConstants::TransientNamespace;
@@ -629,9 +636,7 @@ FName StandardizeName(FName Name, ENiagaraScriptUsage Usage, bool bIsGet, bool b
 
 	checkf(Namespace != NAME_None, TEXT("No namespace picked."));
 
-	// Remove any module or emitter namespaces which weren't handled above since they'll no longer be aliased anyway.
 	NameParts.Remove(FNiagaraConstants::ModuleNamespace);
-	NameParts.Remove(FNiagaraConstants::EmitterNamespace);
 	if (NameParts.Num() == 0)
 	{
 		NameParts.Add(NAME_None);
@@ -650,7 +655,7 @@ FName StandardizeName(FName Name, ENiagaraScriptUsage Usage, bool bIsGet, bool b
 		{
 			RemainingNamePartStrings.Add(NamePart.ToString());
 		}
-		ParameterName = FString::Join(RemainingNamePartStrings, TEXT("_"));
+		ParameterName = FString::Join(RemainingNamePartStrings, TEXT(""));
 	}
 
 	// Last, combine it with the namespace(s) chosen above.
