Fixed race condition crash when deleting an event handler from a running emitter.

diff --git a/Engine/Plugins/FX/Niagara/Source/NiagaraEditor/Private/ViewModels/Stack/NiagaraStackEventScriptItemGroup.cpp b/Engine/Plugins/FX/Niagara/Source/NiagaraEditor/Private/ViewModels/Stack/NiagaraStackEventScriptItemGroup.cpp
--- a/Engine/Plugins/FX/Niagara/Source/NiagaraEditor/Private/ViewModels/Stack/NiagaraStackEventScriptItemGroup.cpp
+++ b/Engine/Plugins/FX/Niagara/Source/NiagaraEditor/Private/ViewModels/Stack/NiagaraStackEventScriptItemGroup.cpp
@@ -262,6 +262,11 @@ void UNiagaraStackEventScriptItemGroup::Delete()
 		return;
 	}
 
+	//Need to tear down existing systems now.
+	FNiagaraSystemUpdateContext UpdateCtx;
+	UpdateCtx.SetDestroyOnAdd(true);
+	UpdateCtx.Add(Emitter, true);
+
 	FScopedTransaction Transaction(FText::Format(LOCTEXT("DeleteEventHandler", "Deleted {0}"), GetDisplayName()));
 	Emitter->Modify();
 	Source->NodeGraph->Modify();
