Fixed some bad tooltips in scalability settings.

diff --git a/Engine/Plugins/FX/Niagara/Source/Niagara/Classes/NiagaraEffectType.h b/Engine/Plugins/FX/Niagara/Source/Niagara/Classes/NiagaraEffectType.h
--- a/Engine/Plugins/FX/Niagara/Source/Niagara/Classes/NiagaraEffectType.h
+++ b/Engine/Plugins/FX/Niagara/Source/Niagara/Classes/NiagaraEffectType.h
@@ -53,10 +53,13 @@ struct FNiagaraSystemScalabilitySettings
 	UPROPERTY(EditAnywhere, Category = "Scalability")
 	FNiagaraPlatformSet Platforms;
 
+	/** Controls whether distance culling is enabled. */
 	UPROPERTY(EditAnywhere, Category = "Scalability", meta = (InlineEditConditionToggle))
 	uint32 bCullByDistance : 1;
+	/** Controls whether instance count culling is enabled. */
 	UPROPERTY(EditAnywhere, Category = "Scalability", meta = (InlineEditConditionToggle))
 	uint32 bCullMaxInstanceCount : 1;
+	/** Controls whether visibility culling is enabled. */
 	UPROPERTY(EditAnywhere, Category = "Scalability", meta = (InlineEditConditionToggle))
 	uint32 bCullByMaxTimeWithoutRender : 1;
 
@@ -98,11 +101,13 @@ struct FNiagaraSystemScalabilityOverride : public FNiagaraSystemScalabilitySetti
 
 	FNiagaraSystemScalabilityOverride();
 
-	//TODO: Detail customization that effectively allows these values to be edit conditions for their respective properties inside FNiagaraScalabilitySettings.
+	/** Controls whether we override the distance culling settings. */
 	UPROPERTY(EditAnywhere, Category = "Override")
 	uint32 bOverrideDistanceSettings : 1;
+	/** Controls whether we override the instance count culling settings. */
 	UPROPERTY(EditAnywhere, Category = "Override")
 	uint32 bOverrideInstanceCountSettings : 1;
+	/** Controls whether we override the visibility culling settings. */
 	UPROPERTY(EditAnywhere, Category = "Override")
 	uint32 bOverrideTimeSinceRendererSettings : 1;
 };
@@ -127,10 +132,11 @@ struct FNiagaraEmitterScalabilitySettings
 	UPROPERTY(EditAnywhere, Category = "Scalability")
 	FNiagaraPlatformSet Platforms;
 
+	/** Enable spawn count scaling */
 	UPROPERTY(EditAnywhere, Category = "Scalability", meta = (InlineEditConditionToggle))
 	uint32 bScaleSpawnCount : 1;
 
-	/** Effects of this type are culled beyond this distance. */
+	/** Scale factor applied to spawn counts for this emitter. */
 	UPROPERTY(EditAnywhere, Category = "Scalability", meta = (EditCondition = "bScaleSpawnCount"))
 	float SpawnCountScale;
 
@@ -155,7 +161,7 @@ struct FNiagaraEmitterScalabilityOverride : public FNiagaraEmitterScalabilitySet
 
 	FNiagaraEmitterScalabilityOverride();
 
-	//TODO: Detail customization that effectively allows these values to be edit conditions for their respective properties inside FNiagaraScalabilitySettings.
+	//Controls whether spawn count scale should be overridden.
 	UPROPERTY(EditAnywhere, Category = "Override")
 	uint32 bOverrideSpawnCountScale : 1;
 };
