Change linked script input and output categories in SNiagaraStack to module reads and module writes categories, respectively.
Now display entries in these categories as SNiagaraParameterNames. Category entries are now sorted by namespace and alphabetized.

diff --git a/Engine/Plugins/FX/Niagara/Source/NiagaraEditor/Private/NiagaraEditorUtilities.cpp b/Engine/Plugins/FX/Niagara/Source/NiagaraEditor/Private/NiagaraEditorUtilities.cpp
--- a/Engine/Plugins/FX/Niagara/Source/NiagaraEditor/Private/NiagaraEditorUtilities.cpp
+++ b/Engine/Plugins/FX/Niagara/Source/NiagaraEditor/Private/NiagaraEditorUtilities.cpp
@@ -57,6 +57,7 @@
 #include "Styling/CoreStyle.h"
 #include "Framework/Notifications/NotificationManager.h"
 #include "NiagaraSimulationStageBase.h"
+#include "NiagaraEditorSettings.h"
 
 #define LOCTEXT_NAMESPACE "FNiagaraEditorUtilities"
 
@@ -2445,4 +2446,82 @@ FText FNiagaraEditorUtilities::FormatParameterNameForTextDisplay(FName Parameter
 	return FText::FromString(DisplayString);
 }
 
+bool FNiagaraEditorUtilities::GetVariableSortPriority(const FName& VarNameA, const FName& VarNameB)
+{
+	const FNiagaraNamespaceMetadata& NamespaceMetaDataA = GetNamespaceMetaDataForVariableName(VarNameA);
+	if (NamespaceMetaDataA.IsValid() == false)
+	{
+		return false;
+	}
+
+	const FNiagaraNamespaceMetadata& NamespaceMetaDataB = GetNamespaceMetaDataForVariableName(VarNameB);
+	const int32 NamespaceAPriority = GetNamespaceMetaDataSortPriority(NamespaceMetaDataA, NamespaceMetaDataB);
+	if (NamespaceAPriority == 0)
+	{
+		return VarNameA.LexicalLess(VarNameB);
+	}
+	return NamespaceAPriority > 0;
+}
+
+int32 FNiagaraEditorUtilities::GetNamespaceMetaDataSortPriority(const FNiagaraNamespaceMetadata& A, const FNiagaraNamespaceMetadata& B)
+{
+	if (A.IsValid() == false)
+	{
+		return false;
+	}
+	else if (B.IsValid() == false)
+	{
+		return true;
+	}
+
+	const int32 ANum = A.Namespaces.Num();
+	const int32 BNum = B.Namespaces.Num();
+	for (int32 i = 0; i < FMath::Min(ANum, BNum); ++i)
+	{
+		const int32 ANamespacePriority = GetNamespaceSortPriority(A.Namespaces[i]);
+		const int32 BNamespacePriority = GetNamespaceSortPriority(B.Namespaces[i]);
+		if (ANamespacePriority != BNamespacePriority)
+			return ANamespacePriority < BNamespacePriority ? 1 : -1;
+	}
+	if (ANum == BNum)
+		return 0;
+
+	return ANum < BNum ? 1 : -1;
+}
+
+int32 FNiagaraEditorUtilities::GetNamespaceSortPriority(const FName& Namespace)
+{
+	if (Namespace == FNiagaraConstants::UserNamespace)
+		return 0;
+	else if (Namespace == FNiagaraConstants::ModuleNamespace)
+		return 1;
+	else if (Namespace == FNiagaraConstants::StaticSwitchNamespace)
+		return 2;
+	else if (Namespace == FNiagaraConstants::DataInstanceNamespace)
+		return 3;
+	else if (Namespace == FNiagaraConstants::OutputNamespace)
+		return 4;
+	else if (Namespace == FNiagaraConstants::EngineNamespace)
+		return 5;
+	else if (Namespace == FNiagaraConstants::ParameterCollectionNamespace)
+		return 6;
+	else if (Namespace == FNiagaraConstants::SystemNamespace)
+		return 7;
+	else if (Namespace == FNiagaraConstants::EmitterNamespace)
+		return 8;
+	else if (Namespace == FNiagaraConstants::ParticleAttributeNamespace)
+		return 9;
+	else if (Namespace == FNiagaraConstants::TransientNamespace)
+		return 10;
+
+	return 11;
+}
+
+const FNiagaraNamespaceMetadata FNiagaraEditorUtilities::GetNamespaceMetaDataForVariableName(const FName& VarName)
+{
+	const FNiagaraParameterHandle VarHandle = FNiagaraParameterHandle(VarName);
+	const TArray<FName> VarHandleNameParts = VarHandle.GetHandleParts();
+	return GetDefault<UNiagaraEditorSettings>()->GetMetaDataForNamespaces(VarHandleNameParts);
+}
+
 #undef LOCTEXT_NAMESPACE
diff --git a/Engine/Plugins/FX/Niagara/Source/NiagaraEditor/Private/ViewModels/Stack/NiagaraStackModuleItemLinkedInputCollection.cpp b/Engine/Plugins/FX/Niagara/Source/NiagaraEditor/Private/ViewModels/Stack/NiagaraStackModuleItemLinkedInputCollection.cpp
--- a/Engine/Plugins/FX/Niagara/Source/NiagaraEditor/Private/ViewModels/Stack/NiagaraStackModuleItemLinkedInputCollection.cpp
+++ b/Engine/Plugins/FX/Niagara/Source/NiagaraEditor/Private/ViewModels/Stack/NiagaraStackModuleItemLinkedInputCollection.cpp
@@ -9,6 +9,8 @@
 #include "ViewModels/Stack/NiagaraStackGraphUtilities.h"
 
 #include "EdGraph/EdGraphPin.h"
+#include "NiagaraEditorSettings.h"
+#include "NiagaraEditorUtilities.h"
 
 UNiagaraStackModuleItemLinkedInputCollection::UNiagaraStackModuleItemLinkedInputCollection()
 	: FunctionCallNode(nullptr)
@@ -25,7 +27,7 @@ void UNiagaraStackModuleItemLinkedInputCollection::Initialize(FRequiredEntryData
  
 FText UNiagaraStackModuleItemLinkedInputCollection::GetDisplayName() const
 {
-	return NSLOCTEXT("StackModuleItemLinkedInputCollection", "LinkedInputsLabel", "Linked Script Inputs");
+	return NSLOCTEXT("StackModuleItemLinkedInputCollection", "ParameterReadsLabel", "Parameter Reads");
 }
 
 bool UNiagaraStackModuleItemLinkedInputCollection::IsExpandedByDefault() const
@@ -71,26 +73,41 @@ void UNiagaraStackModuleItemLinkedInputCollection::RefreshChildrenInternal(const
 			for (int32 i = 0; i < Builder.Histories[0].Variables.Num(); i++)
 			{
 				FNiagaraVariable& Variable = Builder.Histories[0].Variables[i];
-				TArray<const UEdGraphPin*>& WriteHistory = Builder.Histories[0].PerVariableWriteHistory[i];
-
-				for (const UEdGraphPin* WritePin : WriteHistory)
+				const TArray<TTuple<const UEdGraphPin*, const UEdGraphPin*>>& ReadHistory = Builder.Histories[0].PerVariableReadHistory[i];
+				if (ReadHistory.Num() > 0)
 				{
-					if (Cast<UNiagaraNodeParameterMapGet>(WritePin->GetOwningNode()) != nullptr)
+					const FNiagaraNamespaceMetadata NamespaceMetadata = FNiagaraEditorUtilities::GetNamespaceMetaDataForVariableName(Variable.GetName());
+					if ((NamespaceMetadata.IsValid()) && (NamespaceMetadata.Namespaces.Contains(FNiagaraConstants::LocalNamespace) == false))
 					{
-						UNiagaraStackModuleItemOutput* Output = FindCurrentChildOfTypeByPredicate<UNiagaraStackModuleItemOutput>(CurrentChildren,
-							[&](UNiagaraStackModuleItemOutput* CurrentOutput) { return CurrentOutput->GetOutputParameterHandle().GetParameterHandleString() == Variable.GetName(); });
-
-						if (Output == nullptr)
+						for (const TTuple<const UEdGraphPin*, const UEdGraphPin*>& ReadPair : ReadHistory)
 						{
-							Output = NewObject<UNiagaraStackModuleItemOutput>(this);
-							Output->Initialize(CreateDefaultChildRequiredData(), *FunctionCallNode, Variable.GetName(), Variable.GetType());
-						}
+							const UEdGraphPin* ReadPin = ReadPair.Key;
+							if (Cast<UNiagaraNodeParameterMapGet>(ReadPin->GetOwningNode()) != nullptr)
+							{
+								UNiagaraStackModuleItemOutput* Output = FindCurrentChildOfTypeByPredicate<UNiagaraStackModuleItemOutput>(CurrentChildren,
+									[&](UNiagaraStackModuleItemOutput* CurrentOutput) { return CurrentOutput->GetOutputParameterHandle().GetParameterHandleString() == Variable.GetName(); });
+
+								if (Output == nullptr)
+								{
+									Output = NewObject<UNiagaraStackModuleItemOutput>(this);
+									Output->Initialize(CreateDefaultChildRequiredData(), *FunctionCallNode, Variable.GetName(), Variable.GetType());
+								}
 
-						NewChildren.Add(Output);
-						break;
+								NewChildren.Add(Output);
+								break;
+							}
+						}
 					}
 				}
 			}
 		}
+
+		auto SortNewChildrenPred = [](const UNiagaraStackEntry& EntryA, const UNiagaraStackEntry& EntryB) {
+			const UNiagaraStackModuleItemOutput* ModuleItemA = static_cast<const UNiagaraStackModuleItemOutput*>(&EntryA);
+			const UNiagaraStackModuleItemOutput* ModuleItemB = static_cast<const UNiagaraStackModuleItemOutput*>(&EntryB);
+			return FNiagaraEditorUtilities::GetVariableSortPriority(ModuleItemA->GetOutputParameterHandle().GetParameterHandleString(), ModuleItemB->GetOutputParameterHandle().GetParameterHandleString());
+		};
+
+		NewChildren.Sort(SortNewChildrenPred);
 	}
 }
diff --git a/Engine/Plugins/FX/Niagara/Source/NiagaraEditor/Private/ViewModels/Stack/NiagaraStackModuleItemOutputCollection.cpp b/Engine/Plugins/FX/Niagara/Source/NiagaraEditor/Private/ViewModels/Stack/NiagaraStackModuleItemOutputCollection.cpp
--- a/Engine/Plugins/FX/Niagara/Source/NiagaraEditor/Private/ViewModels/Stack/NiagaraStackModuleItemOutputCollection.cpp
+++ b/Engine/Plugins/FX/Niagara/Source/NiagaraEditor/Private/ViewModels/Stack/NiagaraStackModuleItemOutputCollection.cpp
@@ -9,6 +9,8 @@
 #include "ViewModels/Stack/NiagaraStackGraphUtilities.h"
 
 #include "EdGraph/EdGraphPin.h"
+#include "NiagaraEditorSettings.h"
+#include "NiagaraEditorUtilities.h"
 
 UNiagaraStackModuleItemOutputCollection::UNiagaraStackModuleItemOutputCollection()
 	: FunctionCallNode(nullptr)
@@ -25,7 +27,7 @@ void UNiagaraStackModuleItemOutputCollection::Initialize(FRequiredEntryData InRe
 
 FText UNiagaraStackModuleItemOutputCollection::GetDisplayName() const
 {
-	return NSLOCTEXT("StackModuleItemOutputCollection", "OutputsLabel", "Outputs");
+	return NSLOCTEXT("StackModuleItemOutputCollection", "ParameterWritesLabel", "Parameter Writes");
 }
 
 bool UNiagaraStackModuleItemOutputCollection::IsExpandedByDefault() const
@@ -53,23 +55,32 @@ void UNiagaraStackModuleItemOutputCollection::RefreshChildrenInternal(const TArr
 		FCompileConstantResolver ConstantResolver = GetEmitterViewModel().IsValid()
 			? FCompileConstantResolver(GetEmitterViewModel()->GetEmitter())
 			: FCompileConstantResolver();
-		FNiagaraStackGraphUtilities::GetStackFunctionOutputVariables(*FunctionCallNode, ConstantResolver, OutputVariables, Unused);
+		FNiagaraStackGraphUtilities::GetStackFunctionOutputVariables(*FunctionCallNode, ConstantResolver, Unused, OutputVariables);
 		
 		for(int32 OutputIndex = 0; OutputIndex < OutputVariables.Num(); OutputIndex++)
 		{
 			FNiagaraVariable& Variable = OutputVariables[OutputIndex];
+			const FNiagaraNamespaceMetadata NamespaceMetadata = FNiagaraEditorUtilities::GetNamespaceMetaDataForVariableName(Variable.GetName());
+			if ((NamespaceMetadata.IsValid()) && (NamespaceMetadata.Namespaces.Contains(FNiagaraConstants::LocalNamespace) == false))
+			{
+				UNiagaraStackModuleItemOutput* Output = FindCurrentChildOfTypeByPredicate<UNiagaraStackModuleItemOutput>(CurrentChildren,
+					[&](UNiagaraStackModuleItemOutput* CurrentOutput) { return CurrentOutput->GetOutputParameterHandle().GetParameterHandleString() == Variable.GetName(); });
 
-			UNiagaraStackModuleItemOutput* Output = FindCurrentChildOfTypeByPredicate<UNiagaraStackModuleItemOutput>(CurrentChildren,
-				[&](UNiagaraStackModuleItemOutput* CurrentOutput) { return CurrentOutput->GetOutputParameterHandle().GetParameterHandleString() == Variable.GetName(); });
+				if (Output == nullptr)
+				{
+					Output = NewObject<UNiagaraStackModuleItemOutput>(this);
+					Output->Initialize(CreateDefaultChildRequiredData(), *FunctionCallNode, Variable.GetName(), Variable.GetType());
+				}
 
-			if (Output == nullptr)
-			{
-				Output = NewObject<UNiagaraStackModuleItemOutput>(this);
-				Output->Initialize(CreateDefaultChildRequiredData(), *FunctionCallNode, Variable.GetName(), Variable.GetType());
+				NewChildren.Add(Output);
 			}
-
-			NewChildren.Add(Output);
-			break;
 		}
+		auto SortNewChildrenPred = [](const UNiagaraStackEntry& EntryA, const UNiagaraStackEntry& EntryB) {
+			const UNiagaraStackModuleItemOutput* ModuleItemA = static_cast<const UNiagaraStackModuleItemOutput*>(&EntryA);
+			const UNiagaraStackModuleItemOutput* ModuleItemB = static_cast<const UNiagaraStackModuleItemOutput*>(&EntryB);
+			return FNiagaraEditorUtilities::GetVariableSortPriority(ModuleItemA->GetOutputParameterHandle().GetParameterHandleString(), ModuleItemB->GetOutputParameterHandle().GetParameterHandleString());
+		};
+
+		NewChildren.Sort(SortNewChildrenPred);
 	}
 }
diff --git a/Engine/Plugins/FX/Niagara/Source/NiagaraEditor/Public/NiagaraEditorUtilities.h b/Engine/Plugins/FX/Niagara/Source/NiagaraEditor/Public/NiagaraEditorUtilities.h
--- a/Engine/Plugins/FX/Niagara/Source/NiagaraEditor/Public/NiagaraEditorUtilities.h
+++ b/Engine/Plugins/FX/Niagara/Source/NiagaraEditor/Public/NiagaraEditorUtilities.h
@@ -33,6 +33,7 @@ class FMenuBuilder;
 class FNiagaraEmitterViewModel;
 class FNiagaraEmitterHandleViewModel;
 enum class ECheckBoxState : uint8;
+struct FNiagaraNamespaceMetadata;
 
 namespace FNiagaraEditorUtilities
 {
@@ -293,4 +294,15 @@ namespace FNiagaraEditorUtilities
 	void GetReferencingFunctionCallNodes(UNiagaraScript* Script, TArray<UNiagaraNodeFunctionCall*>& OutReferencingFunctionCallNodes);
 
 	FText FormatParameterNameForTextDisplay(FName ParameterName);
+
+	// Compare two FNiagaraVariable names for the sort priority relative to the first argument VarNameA. Sorting is ordered by namespace and then alphabetized. 
+	bool GetVariableSortPriority(const FName& VarNameA, const FName& VarNameB);
+
+	// Compare two FNiagaraNamespaceMetadata for the sort priority relative to the first argument A, where a lower number represents a higher priority.
+	int32 GetNamespaceMetaDataSortPriority(const FNiagaraNamespaceMetadata& A, const FNiagaraNamespaceMetadata& B);
+
+	// Get the sort priority of a registered namespace FName, where a lower number represents a higher priority.
+	int32 GetNamespaceSortPriority(const FName& Namespace);
+
+	const FNiagaraNamespaceMetadata GetNamespaceMetaDataForVariableName(const FName& VarName);
 };
diff --git a/Engine/Plugins/FX/Niagara/Source/NiagaraEditorWidgets/Private/SNiagaraStack.cpp b/Engine/Plugins/FX/Niagara/Source/NiagaraEditorWidgets/Private/SNiagaraStack.cpp
--- a/Engine/Plugins/FX/Niagara/Source/NiagaraEditorWidgets/Private/SNiagaraStack.cpp
+++ b/Engine/Plugins/FX/Niagara/Source/NiagaraEditorWidgets/Private/SNiagaraStack.cpp
@@ -57,6 +57,7 @@
 #include "EditorFontGlyphs.h"
 #include "Framework/Commands/UICommandList.h"
 #include "ViewModels/NiagaraOverviewGraphViewModel.h"
+#include "Widgets/SNiagaraParameterName.h"
 
 
 #define LOCTEXT_NAMESPACE "NiagaraStack"
@@ -567,8 +568,8 @@ TSharedRef<SWidget> SNiagaraStack::GetViewOptionsMenu() const
 		NAME_None, EUserInterfaceActionType::Check);
 
 	MenuBuilder.AddMenuEntry(
-		LOCTEXT("ShowLinkedInputsLabel", "Show Linked Script Inputs"),
-		LOCTEXT("ShowLinkedInputsToolTip", "Whether or not to show internal module linked inputs in the stack."),
+		LOCTEXT("ShowParameterReadsLabel", "Show Parameter Reads"),
+		LOCTEXT("ShowParameterReadsToolTip", "Whether or not to show the parameters that a module reads from."),
 		FSlateIcon(),
 		FUIAction(
 			FExecuteAction::CreateLambda([=]() { StackViewModel->SetShowLinkedInputs(!StackViewModel->GetShowLinkedInputs()); }),
@@ -577,8 +578,8 @@ TSharedRef<SWidget> SNiagaraStack::GetViewOptionsMenu() const
 		NAME_None, EUserInterfaceActionType::Check);
 
 	MenuBuilder.AddMenuEntry(
-		LOCTEXT("ShowOutputsLabel", "Show Outputs"),
-		LOCTEXT("ShowOutputsToolTip", "Whether or not to show module outputs in the stack."),
+		LOCTEXT("ShowParameterWritesLabel", "Show Parameter Writes"),
+		LOCTEXT("ShowParameterWritesToolTip", "Whether or not to show parameters that a module writes to."),
 		FSlateIcon(),
 		FUIAction(
 			FExecuteAction::CreateLambda([=]() { StackViewModel->SetShowOutputs(!StackViewModel->GetShowOutputs()); }),
@@ -983,19 +984,14 @@ SNiagaraStack::FRowWidgets SNiagaraStack::ConstructNameAndValueWidgetsForItem(UN
 	{
 		UNiagaraStackModuleItemOutput* ModuleItemOutput = CastChecked<UNiagaraStackModuleItemOutput>(Item);
 		return FRowWidgets(
-			SNew(STextBlock)
-			.TextStyle(FNiagaraEditorWidgetsStyle::Get(), "NiagaraEditor.Stack.DefaultText")
-			.ToolTipText_UObject(Item, &UNiagaraStackEntry::GetTooltipText)
-			.Text_UObject(Item, &UNiagaraStackEntry::GetDisplayName)
-			.ColorAndOpacity(this, &SNiagaraStack::GetTextColorForItem, Item)
-			.HighlightText_UObject(StackViewModel, &UNiagaraStackViewModel::GetCurrentSearchText)
-			.IsEnabled_UObject(Item, &UNiagaraStackEntry::GetOwnerIsEnabled),
-			SNew(STextBlock)
-			.TextStyle(FNiagaraEditorWidgetsStyle::Get(), "NiagaraEditor.Stack.ParameterText")
-			.Text_UObject(ModuleItemOutput, &UNiagaraStackModuleItemOutput::GetOutputParameterHandleText)
-			.ColorAndOpacity(this, &SNiagaraStack::GetTextColorForItem, Item)
-			.HighlightText_UObject(StackViewModel, &UNiagaraStackViewModel::GetCurrentSearchText)
-			.IsEnabled_UObject(Item, &UNiagaraStackEntry::GetOwnerIsEnabled));
+			SNew(SNiagaraParameterName)
+				.ReadOnlyTextStyle(FNiagaraEditorWidgetsStyle::Get(), "NiagaraEditor.Stack.DefaultText")
+				.ToolTipText_UObject(Item, &UNiagaraStackEntry::GetTooltipText)
+				.ParameterName(ModuleItemOutput->GetOutputParameterHandle().GetParameterHandleString())
+				.HighlightText_UObject(StackViewModel, &UNiagaraStackViewModel::GetCurrentSearchText)
+				.IsEnabled_UObject(Item, &UNiagaraStackEntry::GetOwnerIsEnabled)
+				.IsReadOnly(true)
+			);
 	}
 	else if (Item->IsA<UNiagaraStackFunctionInputCollection>() ||
 		Item->IsA<UNiagaraStackModuleItemOutputCollection>() ||
