Colors and metadata pass on namespace refactor presentation. tooltip descriptions, color tweaks, etc.

diff --git a/Engine/Plugins/FX/Niagara/Source/NiagaraEditor/Private/NiagaraEditorSettings.cpp b/Engine/Plugins/FX/Niagara/Source/NiagaraEditor/Private/NiagaraEditorSettings.cpp
--- a/Engine/Plugins/FX/Niagara/Source/NiagaraEditor/Private/NiagaraEditorSettings.cpp
+++ b/Engine/Plugins/FX/Niagara/Source/NiagaraEditor/Private/NiagaraEditorSettings.cpp
@@ -44,48 +44,49 @@ void UNiagaraEditorSettings::SetupNamespaceMetadata()
 			.AddOption(ENiagaraNamespaceMetadataOptions::CanChangeNamespaceModifier),
 		FNiagaraNamespaceMetadata({FNiagaraConstants::EmitterNamespace})
 			.SetDisplayName(LOCTEXT("EmitterDisplayName", "Emitter"))
-			.SetDescription(LOCTEXT("EmitterDescription", "Persistent attribute which is written in a emitter\n stage and can be read in emitter and particle stages."))
+			.SetDescription(LOCTEXT("EmitterDescription", "Persistent attribute which is written in a emitter\nstage and can be read in emitter and particle stages."))
 			.SetBackgroundColor(FLinearColor(FColor(145, 99, 56)))
 			.AddOption(ENiagaraNamespaceMetadataOptions::CanChangeNamespaceModifier),
 		FNiagaraNamespaceMetadata({FNiagaraConstants::ParticleAttributeNamespace})
 			.SetDisplayName(LOCTEXT("ParticleDisplayName", "Particles"))
-			.SetDescription(LOCTEXT("ParticleDescription", "Persistent attribute which is written in a particle\n stage and can be read in particle stages."))
+			.SetDescription(LOCTEXT("ParticleDescription", "Persistent attribute which is written in a particle\nstage and can be read in particle stages."))
 			.SetBackgroundColor(FLinearColor(FColor(72, 130, 71)))
 			.AddOption(ENiagaraNamespaceMetadataOptions::CanChangeNamespaceModifier),
 		FNiagaraNamespaceMetadata({FNiagaraConstants::ModuleNamespace})
 			.SetDisplayName(LOCTEXT("ModuleDisplayName", "Input"))
-			.SetDescription(LOCTEXT("ModuleDescription", "A transient value which exposes a module input to the system and emitter editor."))
+			.SetDescription(LOCTEXT("ModuleDescription", "A value which exposes a module input to the system and emitter editor."))
 			.SetBackgroundColor(FLinearColor(FColor(136, 66, 65)))
 			.AddOption(ENiagaraNamespaceMetadataOptions::CanChangeNamespaceModifier),
 		FNiagaraNamespaceMetadata({FNiagaraConstants::OutputNamespace, FNiagaraConstants::ModuleNamespace})
 			.SetDisplayName(LOCTEXT("ModuleOutputDisplayName", "Output"))
-			.SetDescription(LOCTEXT("ModuleOutputDescription", "A transient value which is unique to the module that wrote it and can be read from any other module.\nTransient values do not persist from frame to frame, or between stages, e.g. emitter to particle, or spawn to update."))
-			.SetBackgroundColor(FLinearColor(FColor(109, 95, 124)))
+			.SetDescription(LOCTEXT("ModuleOutputDescription", "A transient value which the module author has decided might be useful to other modules further down in the stage.\nTransient values do not persist from frame to frame, or between stages, e.g. emitter to particle, or spawn to update."))
+			.SetBackgroundColor(FLinearColor(FColor(108, 87, 131)))
 			.AddOption(ENiagaraNamespaceMetadataOptions::AdvancedInScript)
 			.AddOption(ENiagaraNamespaceMetadataOptions::AdvancedInSystem)
 			.AddOption(ENiagaraNamespaceMetadataOptions::PreventCreatingInSystemEditor)
 			.AddOption(ENiagaraNamespaceMetadataOptions::CanChangeNamespaceModifier),
 		FNiagaraNamespaceMetadata({FNiagaraConstants::OutputNamespace})
 			.SetDisplayName(LOCTEXT("OutputDisplayName", "Output"))
-			.SetDescription(LOCTEXT("OutputDescription", "A transient value which is unique to the module that wrote it and can be read from any other module.\nTransient values do not persist from frame to frame, or between stages, e.g. emitter to particle, or spawn to update."))
-			.SetBackgroundColor(FLinearColor(FColor(109, 95, 124)))
+			.SetDescription(LOCTEXT("OutputDescription", "A transient value which the module author has decided might be useful to other modules further down in the stage.\nTransient values do not persist from frame to frame, or between stages, e.g. emitter to particle, or spawn to update."))
+			.SetBackgroundColor(FLinearColor(FColor(108, 87, 131)))
 			.AddOption(ENiagaraNamespaceMetadataOptions::AdvancedInScript)
 			.AddOption(ENiagaraNamespaceMetadataOptions::AdvancedInSystem),
 		FNiagaraNamespaceMetadata({FNiagaraConstants::LocalNamespace, FNiagaraConstants::ModuleNamespace})
 			.SetDisplayName(LOCTEXT("ModuleLocalDisplayName", "Local"))
 			.SetDescription(LOCTEXT("ModuleLocalDescription", "A transient value which can be written to and read from within a single module.\nTransient values do not persist from frame to frame, or between stages, e.g. emitter to particle, or spawn to update."))
-			.SetBackgroundColor(FLinearColor(FColor(174, 157, 40)))
-			.AddOption(ENiagaraNamespaceMetadataOptions::AdvancedInScript)
+			.SetBackgroundColor(FLinearColor(FColor(191, 176, 84)))
+			.SetForegroundStyle("NiagaraEditor.ParameterName.NamespaceTextDark")
+			//.AddOption(ENiagaraNamespaceMetadataOptions::AdvancedInScript)
 			.AddOption(ENiagaraNamespaceMetadataOptions::AdvancedInSystem),
 		FNiagaraNamespaceMetadata({FNiagaraConstants::TransientNamespace})
 			.SetDisplayName(LOCTEXT("TransientDisplayName", "Transient"))
 			.SetDescription(LOCTEXT("TransientDescription", "A transient value which can be written to and read from from any module.\nTransient values do not persist from frame to frame, or between stages, e.g. emitter to particle, or spawn to update."))
-			.SetBackgroundColor(FLinearColor(FColor(109, 95, 124)))
+			.SetBackgroundColor(FLinearColor(FColor(108, 87, 131)))
 			.AddOption(ENiagaraNamespaceMetadataOptions::AdvancedInScript)
 			.AddOption(ENiagaraNamespaceMetadataOptions::AdvancedInSystem),
 		FNiagaraNamespaceMetadata({FNiagaraConstants::EngineNamespace})
 			.SetDisplayName(LOCTEXT("EngineDisplayName", "Engine"))
-			.SetDescription(LOCTEXT("EngineDescription", "A read only value which is provided by the engine.\nThis value's source can be the simulation itsef\ne.g. ExecutionCount, or the owner of the simulation e.g. (Owner) Scale."))
+			.SetDescription(LOCTEXT("EngineDescription", "A read only value which is provided by the engine.\nThis value's source can be the simulation itsef\ne.g. ExecutionCount, or the owner of the simulation (The component), e.g. (Owner) Scale."))
 			.SetBackgroundColor(FLinearColor(FColor(170, 170, 170)))
 			.SetForegroundStyle("NiagaraEditor.ParameterName.NamespaceTextDark")
 			.AddOption(ENiagaraNamespaceMetadataOptions::PreventRenaming),
@@ -104,7 +105,7 @@ void UNiagaraEditorSettings::SetupNamespaceMetadata()
 			.AddOption(ENiagaraNamespaceMetadataOptions::PreventRenaming),
 		FNiagaraNamespaceMetadata({FNiagaraConstants::DataInstanceNamespace})
 			.SetDisplayName(LOCTEXT("DataInstanceDisplayName", "Data Instance"))
-			.SetDescription(LOCTEXT("DataInstanceDescription", "A special transient value which has a single bool IsAlive value, which determines if a particle is alive or not."))
+			.SetDescription(LOCTEXT("DataInstanceDescription", "A special value which has a single bool IsAlive value, which determines if a particle is alive or not."))
 			.SetBackgroundColor(FLinearColor(FColor(170, 170, 170)))
 			.SetForegroundStyle("NiagaraEditor.ParameterName.NamespaceTextDark")
 			.AddOption(ENiagaraNamespaceMetadataOptions::AdvancedInScript)
@@ -117,11 +118,12 @@ void UNiagaraEditorSettings::SetupNamespaceMetadata()
 		FNiagaraNamespaceMetadata({FNiagaraConstants::InitialNamespace})
 			.SetDisplayName(LOCTEXT("InitialModifierDisplayName", "Initial"))
 			.SetDescription(LOCTEXT("InitialModifierDescription", "A namespace modifier for dataset attributes which when used in\na linked input in an update script will get the initial value from the spawn script."))
-			.SetBackgroundColor(FLinearColor(FColor(102, 102, 102))),
+			.SetBackgroundColor(FLinearColor(FColor(170, 170, 170)))
+			.SetForegroundStyle("NiagaraEditor.ParameterName.NamespaceTextDark"),
 		FNiagaraNamespaceMetadata({FNiagaraConstants::ModuleNamespace})
 			.SetDisplayName(LOCTEXT("ModuleModifierDisplayName", "Module"))
-			.SetDescription(LOCTEXT("ModuleModifierDescription", "A namespace modifier which is used when writing dataset attributes\nin module which will make that attribute unique to the module instance when used."))
-			.SetBackgroundColor(FLinearColor(FColor(102, 102, 102))),
+			.SetDescription(LOCTEXT("ModuleModifierDescription", "A namespace modifier which makes that attribute unique to the module\ninstance by appending the unique module name."))
+			.SetBackgroundColor(FLinearColor(FColor(102, 102, 152))),
 		FNiagaraNamespaceMetadata({FNiagaraConstants::SystemNamespace})
 			.SetDisplayName(LOCTEXT("SystemModifierDisplayName", "System"))
 			.SetDescription(LOCTEXT("SystemModifierDescription", "A namespace modifier which specifies that an engine provided parameter comes from the system."))
diff --git a/Engine/Plugins/FX/Niagara/Source/NiagaraEditor/Private/NiagaraEditorStyle.cpp b/Engine/Plugins/FX/Niagara/Source/NiagaraEditor/Private/NiagaraEditorStyle.cpp
--- a/Engine/Plugins/FX/Niagara/Source/NiagaraEditor/Private/NiagaraEditorStyle.cpp
+++ b/Engine/Plugins/FX/Niagara/Source/NiagaraEditor/Private/NiagaraEditorStyle.cpp
@@ -197,7 +197,7 @@ TSharedRef< FSlateStyleSet > FNiagaraEditorStyle::Create()
 		.SetFont(DEFAULT_FONT("Bold", 8))
 		.SetColorAndOpacity(FLinearColor(0.0f, 0.0f, 0.0f, 0.9f))
 		.SetShadowOffset(FVector2D(1, 1))
-		.SetShadowColorAndOpacity(FLinearColor(1.0, 1.0, 1.0, 0.5f)));
+		.SetShadowColorAndOpacity(FLinearColor(1.0, 1.0, 1.0, 0.25f)));
 
 	Style->Set("NiagaraEditor.Stack.HighlightedButtonBrush", new BOX_CORE_BRUSH("Common/ButtonHoverHint", FMargin(4 / 16.0f), GetDefault<UEditorStyleSettings>()->SelectionColor));
 
