Added a method for computing the maximum number of per-particle components needed by a renderer. This is necessary so that the number can be clamped to not exceed the maximum allowed GPU buffer size.

Also cleaned up GetBoundAttributes(), since all the subclasses were doing the same thing.

#rnx

diff --git a/Engine/Plugins/FX/Niagara/Source/Niagara/Private/NiagaraLightRendererProperties.cpp b/Engine/Plugins/FX/Niagara/Source/Niagara/Private/NiagaraLightRendererProperties.cpp
--- a/Engine/Plugins/FX/Niagara/Source/Niagara/Private/NiagaraLightRendererProperties.cpp
+++ b/Engine/Plugins/FX/Niagara/Source/Niagara/Private/NiagaraLightRendererProperties.cpp
@@ -22,6 +22,13 @@ TArray<TWeakObjectPtr<UNiagaraLightRendererProperties>> UNiagaraLightRendererPro
 UNiagaraLightRendererProperties::UNiagaraLightRendererProperties()
 	: bUseInverseSquaredFalloff(1), bAffectsTranslucency(0), bOverrideRenderingEnabled(0), RadiusScale(1.0f), ColorAdd(FVector(0.0f, 0.0f, 0.0f))
 {
+	AttributeBindings.Reserve(6);
+	AttributeBindings.Add(&LightRenderingEnabledBinding);
+	AttributeBindings.Add(&LightExponentBinding);
+	AttributeBindings.Add(&PositionBinding);
+	AttributeBindings.Add(&ColorBinding);
+	AttributeBindings.Add(&RadiusBinding);
+	AttributeBindings.Add(&VolumetricScatteringBinding);
 }
 
 void UNiagaraLightRendererProperties::PostInitProperties()
@@ -91,12 +98,6 @@ void UNiagaraLightRendererProperties::GetUsedMaterials(const FNiagaraEmitterInst
 
 #if WITH_EDITORONLY_DATA
 
-const TArray<FNiagaraVariable>& UNiagaraLightRendererProperties::GetRequiredAttributes()
-{
-	static TArray<FNiagaraVariable> Attrs;
-	return Attrs;
-}
-
 const TArray<FNiagaraVariable>& UNiagaraLightRendererProperties::GetOptionalAttributes()
 {
 	static TArray<FNiagaraVariable> Attrs;
@@ -140,37 +141,6 @@ void UNiagaraLightRendererProperties::FixMaterial(UMaterial* Material)
 {
 }
 
-const TArray<FNiagaraVariable>& UNiagaraLightRendererProperties::GetBoundAttributes()
-{
-	CurrentAttributeBindings.Reset();
-
-	TArray<const FNiagaraVariableAttributeBinding*> AttributeBindings;
-	AttributeBindings.Add(&LightRenderingEnabledBinding);
-	AttributeBindings.Add(&LightExponentBinding);
-	AttributeBindings.Add(&PositionBinding);
-	AttributeBindings.Add(&ColorBinding);
-	AttributeBindings.Add(&RadiusBinding);
-	AttributeBindings.Add(&VolumetricScatteringBinding);
-
-	for (const FNiagaraVariableAttributeBinding* AttributeBinding : AttributeBindings)
-	{
-		if (AttributeBinding->BoundVariable.IsValid())
-		{
-			CurrentAttributeBindings.Add(AttributeBinding->BoundVariable);
-		}
-		else if (AttributeBinding->DataSetVariable.IsValid())
-		{
-			CurrentAttributeBindings.Add(AttributeBinding->DataSetVariable);
-		}
-		else
-		{
-			CurrentAttributeBindings.Add(AttributeBinding->DefaultValueIfNonExistent);
-		}
-	}
-
-	return CurrentAttributeBindings;
-}
-
 #endif // WITH_EDITORONLY_DATA
 
 #undef LOCTEXT_NAMESPACE
\ No newline at end of file
diff --git a/Engine/Plugins/FX/Niagara/Source/Niagara/Private/NiagaraMeshRendererProperties.cpp b/Engine/Plugins/FX/Niagara/Source/Niagara/Private/NiagaraMeshRendererProperties.cpp
--- a/Engine/Plugins/FX/Niagara/Source/Niagara/Private/NiagaraMeshRendererProperties.cpp
+++ b/Engine/Plugins/FX/Niagara/Source/Niagara/Private/NiagaraMeshRendererProperties.cpp
@@ -53,6 +53,19 @@ UNiagaraMeshRendererProperties::UNiagaraMeshRendererProperties()
 	, LockedAxis(0.0f, 0.0f, 1.0f)
 	, LockedAxisSpace(ENiagaraMeshLockedAxisSpace::Simulation)
 {
+	AttributeBindings.Reserve(12);
+	AttributeBindings.Add(&PositionBinding);
+	AttributeBindings.Add(&ColorBinding);
+	AttributeBindings.Add(&VelocityBinding);
+	AttributeBindings.Add(&MeshOrientationBinding);
+	AttributeBindings.Add(&ScaleBinding);
+	AttributeBindings.Add(&DynamicMaterialBinding);
+	AttributeBindings.Add(&DynamicMaterial1Binding);
+	AttributeBindings.Add(&DynamicMaterial2Binding);
+	AttributeBindings.Add(&DynamicMaterial3Binding);
+	AttributeBindings.Add(&MaterialRandomBinding);
+	AttributeBindings.Add(&CustomSortingBinding);
+	AttributeBindings.Add(&NormalizedAgeBinding);
 }
 
 FNiagaraRenderer* UNiagaraMeshRendererProperties::CreateEmitterRenderer(ERHIFeatureLevel::Type FeatureLevel, const FNiagaraEmitterInstance* Emitter)
@@ -255,55 +268,6 @@ void UNiagaraMeshRendererProperties::FixMaterial(UMaterial* Material)
 	Material->ForceRecompileForRendering();
 }
 
-const TArray<FNiagaraVariable>& UNiagaraMeshRendererProperties::GetBoundAttributes()
-{
-	CurrentAttributeBindings.Reset();
-
-	TArray<const FNiagaraVariableAttributeBinding*> AttributeBindings;
-	AttributeBindings.Add(&PositionBinding);
-	AttributeBindings.Add(&ColorBinding);
-	AttributeBindings.Add(&VelocityBinding);
-	AttributeBindings.Add(&MeshOrientationBinding);
-	AttributeBindings.Add(&ScaleBinding);
-	AttributeBindings.Add(&DynamicMaterialBinding);
-	AttributeBindings.Add(&DynamicMaterial1Binding);
-	AttributeBindings.Add(&DynamicMaterial2Binding);
-	AttributeBindings.Add(&DynamicMaterial3Binding);
-	AttributeBindings.Add(&MaterialRandomBinding);
-	AttributeBindings.Add(&CustomSortingBinding);
-	AttributeBindings.Add(&NormalizedAgeBinding);
-
-	for (const FNiagaraVariableAttributeBinding* AttributeBinding : AttributeBindings)
-	{
-		if (AttributeBinding->BoundVariable.IsValid())
-		{
-			CurrentAttributeBindings.Add(AttributeBinding->BoundVariable);
-		}
-		else if (AttributeBinding->DataSetVariable.IsValid())
-		{
-			CurrentAttributeBindings.Add(AttributeBinding->DataSetVariable);
-		}
-		else
-		{
-			CurrentAttributeBindings.Add(AttributeBinding->DefaultValueIfNonExistent);
-		}
-	}
-
-	return CurrentAttributeBindings;
-}
-
-const TArray<FNiagaraVariable>& UNiagaraMeshRendererProperties::GetRequiredAttributes()
-{
-	static TArray<FNiagaraVariable> Attrs;
-
-	if (Attrs.Num() == 0)
-	{
-	}
-
-	return Attrs;
-}
-
-
 const TArray<FNiagaraVariable>& UNiagaraMeshRendererProperties::GetOptionalAttributes()
 {
 	static TArray<FNiagaraVariable> Attrs;
diff --git a/Engine/Plugins/FX/Niagara/Source/Niagara/Private/NiagaraRendererProperties.cpp b/Engine/Plugins/FX/Niagara/Source/Niagara/Private/NiagaraRendererProperties.cpp
new file mode 100644
--- /dev/null
+++ b/Engine/Plugins/FX/Niagara/Source/Niagara/Private/NiagaraRendererProperties.cpp
@@ -0,0 +1,80 @@
+// Copyright Epic Games, Inc. All Rights Reserved.
+#include "NiagaraRendererProperties.h"
+#include "NiagaraTypes.h"
+#include "NiagaraCommon.h"
+#include "NiagaraDataSet.h"
+
+#if WITH_EDITORONLY_DATA
+const TArray<FNiagaraVariable>& UNiagaraRendererProperties::GetBoundAttributes()
+{
+	CurrentBoundAttributes.Reset();
+
+	for (const FNiagaraVariableAttributeBinding* AttributeBinding : AttributeBindings)
+	{
+		if (AttributeBinding->BoundVariable.IsValid())
+		{
+			CurrentBoundAttributes.Add(AttributeBinding->BoundVariable);
+		}
+		else if (AttributeBinding->DataSetVariable.IsValid())
+		{
+			CurrentBoundAttributes.Add(AttributeBinding->DataSetVariable);
+		}
+		else
+		{
+			CurrentBoundAttributes.Add(AttributeBinding->DefaultValueIfNonExistent);
+		}
+	}
+
+	return CurrentBoundAttributes;
+}
+#endif
+
+uint32 UNiagaraRendererProperties::ComputeMaxUsedComponents(const FNiagaraDataSet& DataSet) const
+{
+	enum BaseType
+	{
+		BaseType_Int,
+		BaseType_Float,
+		BaseType_NUM
+	};
+
+	uint32 MaxNumComponents = 0;
+
+	TArray<int32, TInlineAllocator<32>> SeenOffsets[BaseType_NUM];
+	uint32 NumComponents[BaseType_NUM] = { 0, 0 };
+	for (const FNiagaraVariableAttributeBinding* Binding : AttributeBindings)
+	{
+		const FNiagaraVariable& Var = Binding->DataSetVariable;
+
+		int32 FloatOffset, IntOffset;
+		DataSet.GetVariableComponentOffsets(Var, FloatOffset, IntOffset);
+
+		int32 VarOffset, VarBaseSize;
+		BaseType VarBaseType;
+		if (FloatOffset != INDEX_NONE)
+		{
+			VarBaseType = BaseType_Float;
+			VarOffset = FloatOffset;
+			VarBaseSize = sizeof(float);
+		}
+		else if (IntOffset != INDEX_NONE)
+		{
+			VarBaseType = BaseType_Int;
+			VarOffset = IntOffset;
+			VarBaseSize = sizeof(int32);
+		}
+		else
+		{
+			continue;
+		}
+
+		if (!SeenOffsets[VarBaseType].Contains(VarOffset))
+		{
+			SeenOffsets[VarBaseType].Add(VarOffset);
+			NumComponents[VarBaseType] += Var.GetSizeInBytes() / VarBaseSize;
+			MaxNumComponents = FMath::Max(MaxNumComponents, NumComponents[VarBaseType]);
+		}
+	}
+
+	return MaxNumComponents;
+}
diff --git a/Engine/Plugins/FX/Niagara/Source/Niagara/Private/NiagaraRibbonRendererProperties.cpp b/Engine/Plugins/FX/Niagara/Source/Niagara/Private/NiagaraRibbonRendererProperties.cpp
--- a/Engine/Plugins/FX/Niagara/Source/Niagara/Private/NiagaraRibbonRendererProperties.cpp
+++ b/Engine/Plugins/FX/Niagara/Source/Niagara/Private/NiagaraRibbonRendererProperties.cpp
@@ -38,6 +38,22 @@ UNiagaraRibbonRendererProperties::UNiagaraRibbonRendererProperties()
 {
 	FNiagaraTypeDefinition MaterialDef(UMaterialInterface::StaticClass());
 	MaterialUserParamBinding.Parameter.SetType(MaterialDef);
+
+	AttributeBindings.Reserve(14);
+	AttributeBindings.Add(&PositionBinding);
+	AttributeBindings.Add(&ColorBinding);
+	AttributeBindings.Add(&VelocityBinding);
+	AttributeBindings.Add(&NormalizedAgeBinding);
+	AttributeBindings.Add(&RibbonTwistBinding);
+	AttributeBindings.Add(&RibbonWidthBinding);
+	AttributeBindings.Add(&RibbonFacingBinding);
+	AttributeBindings.Add(&RibbonIdBinding);
+	AttributeBindings.Add(&RibbonLinkOrderBinding);
+	AttributeBindings.Add(&MaterialRandomBinding);
+	AttributeBindings.Add(&DynamicMaterialBinding);
+	AttributeBindings.Add(&DynamicMaterial1Binding);
+	AttributeBindings.Add(&DynamicMaterial2Binding);
+	AttributeBindings.Add(&DynamicMaterial3Binding);
 }
 
 FNiagaraRenderer* UNiagaraRibbonRendererProperties::CreateEmitterRenderer(ERHIFeatureLevel::Type FeatureLevel, const FNiagaraEmitterInstance* Emitter)
@@ -148,18 +164,6 @@ void UNiagaraRibbonRendererProperties::PostEditChangeProperty(struct FPropertyCh
 	}
 }
 
-const TArray<FNiagaraVariable>& UNiagaraRibbonRendererProperties::GetRequiredAttributes()
-{
-	static TArray<FNiagaraVariable> Attrs;
-
-	if (Attrs.Num() == 0)
-	{
-	}
-
-	return Attrs;
-}
-
-
 const TArray<FNiagaraVariable>& UNiagaraRibbonRendererProperties::GetOptionalAttributes()
 {
 	static TArray<FNiagaraVariable> Attrs;
@@ -244,45 +248,6 @@ void UNiagaraRibbonRendererProperties::FixMaterial(UMaterial* InMaterial)
 	InMaterial->ForceRecompileForRendering();
 }
 
-const TArray<FNiagaraVariable>& UNiagaraRibbonRendererProperties::GetBoundAttributes()
-{
-	CurrentAttributeBindings.Reset();
-
-	TArray<const FNiagaraVariableAttributeBinding*> AttributeBindings;
-	AttributeBindings.Add(&PositionBinding);
-	AttributeBindings.Add(&ColorBinding);
-	AttributeBindings.Add(&VelocityBinding);
-	AttributeBindings.Add(&NormalizedAgeBinding);
-	AttributeBindings.Add(&RibbonTwistBinding);
-	AttributeBindings.Add(&RibbonWidthBinding);
-	AttributeBindings.Add(&RibbonFacingBinding);
-	AttributeBindings.Add(&RibbonIdBinding);
-	AttributeBindings.Add(&RibbonLinkOrderBinding);
-	AttributeBindings.Add(&MaterialRandomBinding);
-	AttributeBindings.Add(&DynamicMaterialBinding);
-	AttributeBindings.Add(&DynamicMaterial1Binding);
-	AttributeBindings.Add(&DynamicMaterial2Binding);
-	AttributeBindings.Add(&DynamicMaterial3Binding);
-
-	for (const FNiagaraVariableAttributeBinding* AttributeBinding : AttributeBindings)
-	{
-		if (AttributeBinding->BoundVariable.IsValid())
-		{
-			CurrentAttributeBindings.Add(AttributeBinding->BoundVariable);
-		}
-		else if (AttributeBinding->DataSetVariable.IsValid())
-		{
-			CurrentAttributeBindings.Add(AttributeBinding->DataSetVariable);
-		}
-		else
-		{
-			CurrentAttributeBindings.Add(AttributeBinding->DefaultValueIfNonExistent);
-		}
-	}
-
-	return CurrentAttributeBindings;
-}
-
 bool UNiagaraRibbonRendererProperties::CanEditChange(const FProperty* InProperty) const
 {
 
diff --git a/Engine/Plugins/FX/Niagara/Source/Niagara/Private/NiagaraSpriteRendererProperties.cpp b/Engine/Plugins/FX/Niagara/Source/Niagara/Private/NiagaraSpriteRendererProperties.cpp
--- a/Engine/Plugins/FX/Niagara/Source/Niagara/Private/NiagaraSpriteRendererProperties.cpp
+++ b/Engine/Plugins/FX/Niagara/Source/Niagara/Private/NiagaraSpriteRendererProperties.cpp
@@ -54,6 +54,25 @@ UNiagaraSpriteRendererProperties::UNiagaraSpriteRendererProperties()
 {
 	FNiagaraTypeDefinition MaterialDef(UMaterialInterface::StaticClass());
 	MaterialUserParamBinding.Parameter.SetType(MaterialDef);
+
+	AttributeBindings.Reserve(17);
+	AttributeBindings.Add(&PositionBinding);
+	AttributeBindings.Add(&ColorBinding);
+	AttributeBindings.Add(&VelocityBinding);
+	AttributeBindings.Add(&SpriteRotationBinding);
+	AttributeBindings.Add(&SpriteSizeBinding);
+	AttributeBindings.Add(&SpriteFacingBinding);
+	AttributeBindings.Add(&SpriteAlignmentBinding);
+	AttributeBindings.Add(&SubImageIndexBinding);
+	AttributeBindings.Add(&DynamicMaterialBinding);
+	AttributeBindings.Add(&DynamicMaterial1Binding);
+	AttributeBindings.Add(&DynamicMaterial2Binding);
+	AttributeBindings.Add(&DynamicMaterial3Binding);
+	AttributeBindings.Add(&CameraOffsetBinding);
+	AttributeBindings.Add(&UVScaleBinding);
+	AttributeBindings.Add(&MaterialRandomBinding);
+	AttributeBindings.Add(&CustomSortingBinding);
+	AttributeBindings.Add(&NormalizedAgeBinding);
 }
 
 FNiagaraRenderer* UNiagaraSpriteRendererProperties::CreateEmitterRenderer(ERHIFeatureLevel::Type FeatureLevel, const FNiagaraEmitterInstance* Emitter)
@@ -216,59 +235,6 @@ void UNiagaraSpriteRendererProperties::PostEditChangeProperty(struct FPropertyCh
 
 }
 
-const TArray<FNiagaraVariable>& UNiagaraSpriteRendererProperties::GetBoundAttributes()
-{
-	CurrentAttributeBindings.Reset();
-
-	TArray<const FNiagaraVariableAttributeBinding*> AttributeBindings;
-	AttributeBindings.Add(&PositionBinding);
-	AttributeBindings.Add(&ColorBinding);
-	AttributeBindings.Add(&VelocityBinding);
-	AttributeBindings.Add(&SpriteRotationBinding);
-	AttributeBindings.Add(&SpriteSizeBinding);
-	AttributeBindings.Add(&SpriteFacingBinding);
-	AttributeBindings.Add(&SpriteAlignmentBinding);
-	AttributeBindings.Add(&SubImageIndexBinding);
-	AttributeBindings.Add(&DynamicMaterialBinding);
-	AttributeBindings.Add(&DynamicMaterial1Binding);
-	AttributeBindings.Add(&DynamicMaterial2Binding);
-	AttributeBindings.Add(&DynamicMaterial3Binding);
-	AttributeBindings.Add(&CameraOffsetBinding);
-	AttributeBindings.Add(&UVScaleBinding);
-	AttributeBindings.Add(&MaterialRandomBinding);
-	AttributeBindings.Add(&CustomSortingBinding);
-	AttributeBindings.Add(&NormalizedAgeBinding);
-
-	for (const FNiagaraVariableAttributeBinding* AttributeBinding : AttributeBindings)
-	{
-		if (AttributeBinding->BoundVariable.IsValid())
-		{
-			CurrentAttributeBindings.Add(AttributeBinding->BoundVariable);
-		}
-		else if (AttributeBinding->DataSetVariable.IsValid())
-		{
-			CurrentAttributeBindings.Add(AttributeBinding->DataSetVariable);
-		}
-		else
-		{
-			CurrentAttributeBindings.Add(AttributeBinding->DefaultValueIfNonExistent);
-		}
-	}
-
-	return CurrentAttributeBindings;
-}
-
-const TArray<FNiagaraVariable>& UNiagaraSpriteRendererProperties::GetRequiredAttributes()
-{
-	static TArray<FNiagaraVariable> Attrs;
-
-	if (Attrs.Num() == 0)
-	{
-	}
-
-	return Attrs;
-}
-
 const TArray<FNiagaraVariable>& UNiagaraSpriteRendererProperties::GetOptionalAttributes()
 {
 	static TArray<FNiagaraVariable> Attrs;
diff --git a/Engine/Plugins/FX/Niagara/Source/Niagara/Public/NiagaraLightRendererProperties.h b/Engine/Plugins/FX/Niagara/Source/Niagara/Public/NiagaraLightRendererProperties.h
--- a/Engine/Plugins/FX/Niagara/Source/Niagara/Public/NiagaraLightRendererProperties.h
+++ b/Engine/Plugins/FX/Niagara/Source/Niagara/Public/NiagaraLightRendererProperties.h
@@ -33,8 +33,6 @@ public:
 #if WITH_EDITORONLY_DATA
 	virtual bool IsMaterialValidForRenderer(UMaterial* Material, FText& InvalidMessage) override;
 	virtual void FixMaterial(UMaterial* Material) override;
-	virtual const TArray<FNiagaraVariable>& GetBoundAttributes() override;
-	virtual const TArray<FNiagaraVariable>& GetRequiredAttributes() override;
 	virtual const TArray<FNiagaraVariable>& GetOptionalAttributes() override;
 	virtual void GetRendererWidgets(const FNiagaraEmitterInstance* InEmitter, TArray<TSharedPtr<SWidget>>& OutWidgets, TSharedPtr<FAssetThumbnailPool> InThumbnailPool) const override;
 	virtual void GetRendererTooltipWidgets(const FNiagaraEmitterInstance* InEmitter, TArray<TSharedPtr<SWidget>>& OutWidgets, TSharedPtr<FAssetThumbnailPool> InThumbnailPool) const override;
diff --git a/Engine/Plugins/FX/Niagara/Source/Niagara/Public/NiagaraMeshRendererProperties.h b/Engine/Plugins/FX/Niagara/Source/Niagara/Public/NiagaraMeshRendererProperties.h
--- a/Engine/Plugins/FX/Niagara/Source/Niagara/Public/NiagaraMeshRendererProperties.h
+++ b/Engine/Plugins/FX/Niagara/Source/Niagara/Public/NiagaraMeshRendererProperties.h
@@ -95,8 +95,6 @@ public:
 #if WITH_EDITORONLY_DATA
 	virtual bool IsMaterialValidForRenderer(UMaterial* Material, FText& InvalidMessage) override;
 	virtual void FixMaterial(UMaterial* Material) override;
-	virtual const TArray<FNiagaraVariable>& GetBoundAttributes() override;
-	virtual const TArray<FNiagaraVariable>& GetRequiredAttributes() override;
 	virtual const TArray<FNiagaraVariable>& GetOptionalAttributes() override;
 	virtual	void GetRendererWidgets(const FNiagaraEmitterInstance* InEmitter, TArray<TSharedPtr<SWidget>>& OutWidgets, TSharedPtr<FAssetThumbnailPool> InThumbnailPool) const override;
 	virtual	void GetRendererTooltipWidgets(const FNiagaraEmitterInstance* InEmitter, TArray<TSharedPtr<SWidget>>& OutWidgets, TSharedPtr<FAssetThumbnailPool> InThumbnailPool) const override;
diff --git a/Engine/Plugins/FX/Niagara/Source/Niagara/Public/NiagaraRendererProperties.h b/Engine/Plugins/FX/Niagara/Source/Niagara/Public/NiagaraRendererProperties.h
--- a/Engine/Plugins/FX/Niagara/Source/Niagara/Public/NiagaraRendererProperties.h
+++ b/Engine/Plugins/FX/Niagara/Source/Niagara/Public/NiagaraRendererProperties.h
@@ -25,6 +25,7 @@ class UMaterialInterface;
 class FNiagaraEmitterInstance;
 class SWidget;
 class FAssetThumbnailPool;
+class FNiagaraDataSet;
 
 UCLASS(ABSTRACT)
 class NIAGARA_API UNiagaraRendererProperties : public UNiagaraMergeable
@@ -45,13 +46,16 @@ public:
 
 	virtual bool IsSimTargetSupported(ENiagaraSimTarget InSimTarget) const { return false; };
 
+	const TArray<const FNiagaraVariableAttributeBinding*>& GetAttributeBindings() const { return AttributeBindings; }
+	uint32 ComputeMaxUsedComponents(const FNiagaraDataSet& DataSet) const;
+
 #if WITH_EDITORONLY_DATA
 
 	virtual bool IsMaterialValidForRenderer(UMaterial* Material, FText& InvalidMessage) { return true; }
 
 	virtual void FixMaterial(UMaterial* Material) { }
 
-	virtual const TArray<FNiagaraVariable>& GetBoundAttributes() { static TArray<FNiagaraVariable> Vars; return Vars; };
+	virtual const TArray<FNiagaraVariable>& GetBoundAttributes();
 	virtual const TArray<FNiagaraVariable>& GetRequiredAttributes() { static TArray<FNiagaraVariable> Vars; return Vars; };
 	virtual const TArray<FNiagaraVariable>& GetOptionalAttributes() { static TArray<FNiagaraVariable> Vars; return Vars; };
 
@@ -86,7 +90,9 @@ public:
 	bool bMotionBlurEnabled;
 
 protected:
-	// Copy of variables in the attribute binding, updated when GetRequiredAttributes() is called.
-	TArray<FNiagaraVariable> CurrentAttributeBindings;
+	TArray<const FNiagaraVariableAttributeBinding*> AttributeBindings;
+
+	// Copy of variables in the attribute binding, updated when GetBoundAttributes() is called.
+	TArray<FNiagaraVariable> CurrentBoundAttributes;
 
 };
diff --git a/Engine/Plugins/FX/Niagara/Source/Niagara/Public/NiagaraRibbonRendererProperties.h b/Engine/Plugins/FX/Niagara/Source/Niagara/Public/NiagaraRibbonRendererProperties.h
--- a/Engine/Plugins/FX/Niagara/Source/Niagara/Public/NiagaraRibbonRendererProperties.h
+++ b/Engine/Plugins/FX/Niagara/Source/Niagara/Public/NiagaraRibbonRendererProperties.h
@@ -83,8 +83,6 @@ public:
 #if WITH_EDITOR
 	virtual bool IsMaterialValidForRenderer(UMaterial* Material, FText& InvalidMessage) override;
 	virtual void FixMaterial(UMaterial* Material);
-	virtual const TArray<FNiagaraVariable>& GetBoundAttributes() override;
-	virtual const TArray<FNiagaraVariable>& GetRequiredAttributes() override;
 	virtual const TArray<FNiagaraVariable>& GetOptionalAttributes() override;
 	virtual void GetRendererWidgets(const FNiagaraEmitterInstance* InEmitter, TArray<TSharedPtr<SWidget>>& OutWidgets, TSharedPtr<FAssetThumbnailPool> InThumbnailPool) const override;
 	virtual void GetRendererTooltipWidgets(const FNiagaraEmitterInstance* InEmitter, TArray<TSharedPtr<SWidget>>& OutWidgets, TSharedPtr<FAssetThumbnailPool> InThumbnailPool) const override;
diff --git a/Engine/Plugins/FX/Niagara/Source/Niagara/Public/NiagaraSpriteRendererProperties.h b/Engine/Plugins/FX/Niagara/Source/Niagara/Public/NiagaraSpriteRendererProperties.h
--- a/Engine/Plugins/FX/Niagara/Source/Niagara/Public/NiagaraSpriteRendererProperties.h
+++ b/Engine/Plugins/FX/Niagara/Source/Niagara/Public/NiagaraSpriteRendererProperties.h
@@ -68,8 +68,6 @@ public:
 #if WITH_EDITOR
 	virtual bool IsMaterialValidForRenderer(UMaterial* Material, FText& InvalidMessage) override;
 	virtual void FixMaterial(UMaterial* Material) override;
-	virtual const TArray<FNiagaraVariable>& GetBoundAttributes() override;
-	virtual const TArray<FNiagaraVariable>& GetRequiredAttributes() override;
 	virtual const TArray<FNiagaraVariable>& GetOptionalAttributes() override;
 	virtual void GetRendererWidgets(const FNiagaraEmitterInstance* InEmitter, TArray<TSharedPtr<SWidget>>& OutWidgets, TSharedPtr<FAssetThumbnailPool> InThumbnailPool) const override;
 	virtual void GetRendererTooltipWidgets(const FNiagaraEmitterInstance* InEmitter, TArray<TSharedPtr<SWidget>>& OutWidgets, TSharedPtr<FAssetThumbnailPool> InThumbnailPool) const override;
