Niagara has float16 components in this branch. Fixed change coming from 4.25 to account for those.

#rnx

diff --git a/Engine/Plugins/FX/Niagara/Source/Niagara/Private/NiagaraRendererProperties.cpp b/Engine/Plugins/FX/Niagara/Source/Niagara/Private/NiagaraRendererProperties.cpp
--- a/Engine/Plugins/FX/Niagara/Source/Niagara/Private/NiagaraRendererProperties.cpp
+++ b/Engine/Plugins/FX/Niagara/Source/Niagara/Private/NiagaraRendererProperties.cpp
@@ -35,6 +35,7 @@ uint32 UNiagaraRendererProperties::ComputeMaxUsedComponents(const FNiagaraDataSe
 	{
 		BaseType_Int,
 		BaseType_Float,
+		BaseType_Half,
 		BaseType_NUM
 	};
 
@@ -46,8 +47,8 @@ uint32 UNiagaraRendererProperties::ComputeMaxUsedComponents(const FNiagaraDataSe
 	{
 		const FNiagaraVariable& Var = Binding->DataSetVariable;
 
-		int32 FloatOffset, IntOffset;
-		DataSet.GetVariableComponentOffsets(Var, FloatOffset, IntOffset);
+		int32 FloatOffset, IntOffset, HalfOffset;
+		DataSet.GetVariableComponentOffsets(Var, FloatOffset, IntOffset, HalfOffset);
 
 		int32 VarOffset, VarBaseSize;
 		BaseType VarBaseType;
@@ -63,6 +64,12 @@ uint32 UNiagaraRendererProperties::ComputeMaxUsedComponents(const FNiagaraDataSe
 			VarOffset = IntOffset;
 			VarBaseSize = sizeof(int32);
 		}
+		else if (HalfOffset != INDEX_NONE)
+		{
+			VarBaseType = BaseType_Half;
+			VarOffset = HalfOffset;
+			VarBaseSize = sizeof(float) / 2;
+		}
 		else
 		{
 			continue;
