- Fix for raytraced Niagara Meshes on the GPU with multiple sections

diff --git a/Engine/Plugins/FX/Niagara/Source/Niagara/Private/NiagaraRendererMeshes.cpp b/Engine/Plugins/FX/Niagara/Source/Niagara/Private/NiagaraRendererMeshes.cpp
--- a/Engine/Plugins/FX/Niagara/Source/Niagara/Private/NiagaraRendererMeshes.cpp
+++ b/Engine/Plugins/FX/Niagara/Source/Niagara/Private/NiagaraRendererMeshes.cpp
@@ -588,8 +588,7 @@ void FNiagaraRendererMeshes::GetDynamicRayTracingInstances(FRayTracingMaterialGa
 	FRayTracingInstance RayTracingInstance;
 	RayTracingInstance.Geometry = &Geometry;
 
-	int32 SectionCount = SimTarget == ENiagaraSimTarget::GPUComputeSim ? 1 : LODModel.Sections.Num();
-	for (int32 SectionIndex = 0; SectionIndex < SectionCount; SectionIndex++)
+	for (int32 SectionIndex = 0; SectionIndex < LODModel.Sections.Num(); SectionIndex++)
 	{
 		const FStaticMeshSection& Section = LODModel.Sections[SectionIndex];
 		FMaterialRenderProxy* MaterialProxy = DynamicDataMesh->Materials[SectionIndex];
