Niagara - Fix an assert due to trying to add duplicate compile messages as stack issues.

diff --git a/Engine/Plugins/FX/Niagara/Source/NiagaraEditor/Private/NiagaraMessageManager.cpp b/Engine/Plugins/FX/Niagara/Source/NiagaraEditor/Private/NiagaraMessageManager.cpp
--- a/Engine/Plugins/FX/Niagara/Source/NiagaraEditor/Private/NiagaraMessageManager.cpp
+++ b/Engine/Plugins/FX/Niagara/Source/NiagaraEditor/Private/NiagaraMessageManager.cpp
@@ -11,6 +11,7 @@
 #include "NiagaraScriptToolkit.h"
 #include "../Public/NiagaraEditorModule.h"
 #include "Subsystems/AssetEditorSubsystem.h"
+#include "NiagaraEditorUtilities.h"
 
 #define LOCTEXT_NAMESPACE "NiagaraMessageManager"
 
@@ -591,7 +592,8 @@ void FNiagaraCompileEventToken::OpenScriptAssetByPathAndFocusNodeOrPinIfSet(
 	FAssetData AssetData = AssetRegistry.GetAssetByObjectPath(*InScriptAssetPath);
 	if (AssetData.IsValid())
 	{
-		if (UNiagaraScript* ScriptAsset = Cast<UNiagaraScript>(AssetData.GetAsset()))
+		UNiagaraScript* ScriptAsset = Cast<UNiagaraScript>(AssetData.GetAsset());
+		if (ScriptAsset != nullptr && ScriptAsset->IsAsset())
 		{
 			GEditor->GetEditorSubsystem<UAssetEditorSubsystem>()->OpenEditorForAsset(ScriptAsset, EToolkitMode::Standalone);
 			FNiagaraEditorModule& NiagaraEditorModule = FModuleManager::LoadModuleChecked<FNiagaraEditorModule>("NiagaraEditor");
@@ -611,7 +613,9 @@ void FNiagaraCompileEventToken::OpenScriptAssetByPathAndFocusNodeOrPinIfSet(
 		}
 		else
 		{
-			checkf(false, TEXT("CompileEventJob referenced asset was not a UNiagaraScript!"));
+			FNiagaraEditorUtilities::WarnWithToastAndLog(FText::Format(
+				LOCTEXT("CantNavigateWarning", "Could not navigate to script {0}\nIt was either was not a valid script, or isn't an asset which can be opened directly."),
+				FText::FromString(InScriptAssetPath)));
 		}
 	}
 }
diff --git a/Engine/Plugins/FX/Niagara/Source/NiagaraEditor/Private/ViewModels/Stack/NiagaraStackModuleItem.cpp b/Engine/Plugins/FX/Niagara/Source/NiagaraEditor/Private/ViewModels/Stack/NiagaraStackModuleItem.cpp
--- a/Engine/Plugins/FX/Niagara/Source/NiagaraEditor/Private/ViewModels/Stack/NiagaraStackModuleItem.cpp
+++ b/Engine/Plugins/FX/Niagara/Source/NiagaraEditor/Private/ViewModels/Stack/NiagaraStackModuleItem.cpp
@@ -423,7 +423,13 @@ void UNiagaraStackModuleItem::RefreshIssues(TArray<FStackIssue>& NewIssues)
 			GetSystemViewModel()->GetMessageLogGuid(), FObjectKey(FunctionCallNode));
 		for (TSharedRef<const INiagaraMessage> Message : Messages)
 		{
-			NewIssues.Add(FNiagaraStackGraphUtilities::MessageManagerMessageToStackIssue(Message, GetStackEditorDataKey()));
+			// Sometimes compile errors with the same info are generated, so guard against duplicates here.
+			FStackIssue Issue = FNiagaraStackGraphUtilities::MessageManagerMessageToStackIssue(Message, GetStackEditorDataKey());
+			if (NewIssues.ContainsByPredicate([&Issue](const FStackIssue& NewIssue)
+				{ return NewIssue.GetUniqueIdentifier() == Issue.GetUniqueIdentifier(); }) == false)
+			{
+				NewIssues.Add(Issue);
+			}
 		}
 
 		if (FunctionCallNode->FunctionScript == nullptr && FunctionCallNode->GetClass() == UNiagaraNodeFunctionCall::StaticClass())
