Niagara - Fix a crash when deleting user parameters that are data interfaces.

diff --git a/Engine/Plugins/FX/Niagara/Source/NiagaraEditor/Private/ViewModels/Stack/NiagaraStackParameterStoreEntry.cpp b/Engine/Plugins/FX/Niagara/Source/NiagaraEditor/Private/ViewModels/Stack/NiagaraStackParameterStoreEntry.cpp
--- a/Engine/Plugins/FX/Niagara/Source/NiagaraEditor/Private/ViewModels/Stack/NiagaraStackParameterStoreEntry.cpp
+++ b/Engine/Plugins/FX/Niagara/Source/NiagaraEditor/Private/ViewModels/Stack/NiagaraStackParameterStoreEntry.cpp
@@ -286,19 +286,19 @@ void UNiagaraStackParameterStoreEntry::Delete()
 	TArray<UEdGraphPin*> OwningPins = GetOwningPins();
 	RemovePins(OwningPins);
 
+	// Cache these here since the entry will be finalized when the owning parameter store changes.
+	TSharedPtr<FNiagaraSystemViewModel> CachedSystemViewModel = GetSystemViewModel();
+	UNiagaraDataInterface* DataInterface = Cast<UNiagaraDataInterface>(ValueObject.Get());
+
 	//remove from store
 	Owner->Modify();
 	ParameterStore->RemoveParameter(FNiagaraVariable(InputType, ParameterName));
-	if (InputType.IsDataInterface())
+
+	// Notify the system view model that the DI has been modifier.
+	if (CachedSystemViewModel.IsValid() && DataInterface != nullptr)
 	{
-		UNiagaraDataInterface* DataInterface = NewObject<UNiagaraDataInterface>(this, const_cast<UClass*>(InputType.GetClass()));
-		if (DataInterface != nullptr)
-		{
-			GetSystemViewModel()->NotifyDataObjectChanged(DataInterface);
-		}
+		CachedSystemViewModel->NotifyDataObjectChanged(DataInterface);
 	}
-
-	ParameterDeletedDelegate.Broadcast();
 }
 
 void UNiagaraStackParameterStoreEntry::RemovePins(TArray<UEdGraphPin*> OwningPins /*, bool bSetPreviousValue*/)
@@ -350,11 +350,6 @@ UNiagaraStackParameterStoreEntry::FOnValueChanged& UNiagaraStackParameterStoreEn
 	return ValueChangedDelegate;
 }
 
-UNiagaraStackParameterStoreEntry::FOnParameterDeleted& UNiagaraStackParameterStoreEntry::OnParameterDeleted()
-{
-	return ParameterDeletedDelegate;
-}
-
 TSharedPtr<FNiagaraVariable> UNiagaraStackParameterStoreEntry::GetCurrentValueVariable()
 {
 	if (InputType.GetClass() == nullptr)
diff --git a/Engine/Plugins/FX/Niagara/Source/NiagaraEditor/Private/ViewModels/Stack/NiagaraStackSystemSettingsGroup.cpp b/Engine/Plugins/FX/Niagara/Source/NiagaraEditor/Private/ViewModels/Stack/NiagaraStackSystemSettingsGroup.cpp
--- a/Engine/Plugins/FX/Niagara/Source/NiagaraEditor/Private/ViewModels/Stack/NiagaraStackSystemSettingsGroup.cpp
+++ b/Engine/Plugins/FX/Niagara/Source/NiagaraEditor/Private/ViewModels/Stack/NiagaraStackSystemSettingsGroup.cpp
@@ -193,7 +193,6 @@ void UNiagaraStackParameterStoreItem::RefreshChildrenInternal(const TArray<UNiag
 			{
 				ValueObjectEntry = NewObject<UNiagaraStackParameterStoreEntry>(this);
 				ValueObjectEntry->Initialize(CreateDefaultChildRequiredData(), Owner.Get(), ParameterStore, Var.GetName().ToString(), Var.GetType(), GetStackEditorDataKey());
-				ValueObjectEntry->OnParameterDeleted().AddUObject(this, &UNiagaraStackParameterStoreItem::ParameterDeleted);
 			}
 
 			NewChildren.Add(ValueObjectEntry);
@@ -221,9 +220,4 @@ void UNiagaraStackParameterStoreItem::ParameterStoreChanged()
 	}
 }
 
-void UNiagaraStackParameterStoreItem::ParameterDeleted()
-{
-	RefreshChildren();
-}
-
 #undef LOCTEXT_NAMESPACE
diff --git a/Engine/Plugins/FX/Niagara/Source/NiagaraEditor/Public/ViewModels/Stack/NiagaraStackParameterStoreEntry.h b/Engine/Plugins/FX/Niagara/Source/NiagaraEditor/Public/ViewModels/Stack/NiagaraStackParameterStoreEntry.h
--- a/Engine/Plugins/FX/Niagara/Source/NiagaraEditor/Public/ViewModels/Stack/NiagaraStackParameterStoreEntry.h
+++ b/Engine/Plugins/FX/Niagara/Source/NiagaraEditor/Public/ViewModels/Stack/NiagaraStackParameterStoreEntry.h
@@ -23,7 +23,6 @@ class NIAGARAEDITOR_API UNiagaraStackParameterStoreEntry : public UNiagaraStackI
 
 public:
 	DECLARE_MULTICAST_DELEGATE(FOnValueChanged);
-	DECLARE_MULTICAST_DELEGATE(FOnParameterDeleted);
 
 public:
 	UNiagaraStackParameterStoreEntry();
@@ -84,9 +83,6 @@ public:
 	/** Gets a multicast delegate which is called whenever the value on this input changes. */
 	FOnValueChanged& OnValueChanged();
 
-	/** Gets a multicast delegate which is called when this parameter is deleted. */
-	FOnParameterDeleted& OnParameterDeleted();
-
 	/** Delete the parameter from the ParameterStore and notify that the store changed. */
 	void Delete();
 
@@ -121,9 +117,6 @@ private:
 	/** A multicast delegate which is called when the value of this input is changed. */
 	FOnValueChanged ValueChangedDelegate;
 
-	/** A multicast delegate which is called when this parameter is deleted. */
-	FOnParameterDeleted ParameterDeletedDelegate;
-
 	/** A pointer to the data interface object for this input if one is available. */
 	TWeakObjectPtr<UObject> ValueObject;
 
diff --git a/Engine/Plugins/FX/Niagara/Source/NiagaraEditor/Public/ViewModels/Stack/NiagaraStackSystemSettingsGroup.h b/Engine/Plugins/FX/Niagara/Source/NiagaraEditor/Public/ViewModels/Stack/NiagaraStackSystemSettingsGroup.h
--- a/Engine/Plugins/FX/Niagara/Source/NiagaraEditor/Public/ViewModels/Stack/NiagaraStackSystemSettingsGroup.h
+++ b/Engine/Plugins/FX/Niagara/Source/NiagaraEditor/Public/ViewModels/Stack/NiagaraStackSystemSettingsGroup.h
@@ -50,7 +50,6 @@ protected:
 
 private:
 	void ParameterStoreChanged();
-	void ParameterDeleted();
 
 private:
 	TWeakObjectPtr<UObject> Owner;
