Apply material instance mask for Niagara meshes. Allows support for partially opaque niagara meshes.

#[fyi] juan.canada

diff --git a/Engine/Plugins/FX/Niagara/Source/Niagara/Private/NiagaraRendererMeshes.cpp b/Engine/Plugins/FX/Niagara/Source/Niagara/Private/NiagaraRendererMeshes.cpp
--- a/Engine/Plugins/FX/Niagara/Source/Niagara/Private/NiagaraRendererMeshes.cpp
+++ b/Engine/Plugins/FX/Niagara/Source/Niagara/Private/NiagaraRendererMeshes.cpp
@@ -784,6 +784,7 @@ void FNiagaraRendererMeshes::GetDynamicRayTracingInstances(FRayTracingMaterialGa
 		RHICmdList.TransitionResource(EResourceTransitionAccess::EReadable, EResourceTransitionPipeline::EComputeToCompute, InstanceGPUTransformsBuffer.UAV);
 	}
 
+	RayTracingInstance.BuildInstanceMaskAndFlags();
 	OutRayTracingInstances.Add(RayTracingInstance);
 }
 #endif
