- Fix for attempting to read from GPU buffers on the CPU (nullptr access)

diff --git a/Engine/Plugins/FX/Niagara/Source/Niagara/Private/NiagaraDataInterfaceParticleRead.cpp b/Engine/Plugins/FX/Niagara/Source/Niagara/Private/NiagaraDataInterfaceParticleRead.cpp
--- a/Engine/Plugins/FX/Niagara/Source/Niagara/Private/NiagaraDataInterfaceParticleRead.cpp
+++ b/Engine/Plugins/FX/Niagara/Source/Niagara/Private/NiagaraDataInterfaceParticleRead.cpp
@@ -1792,6 +1792,12 @@ Type UNiagaraDataInterfaceParticleRead::RetrieveValueByIndexWithCheck(FNiagaraEm
 		return Value;
 	}
 
+	//-TODO: Investigate removing this in favor of making GetSafe really a safe call and handle GPU buffers with no CPU side data
+	if (EmitterInstance->GetGPUContext() != nullptr)
+	{
+		return Value;
+	}
+
 	FNiagaraDataSetAccessor<AccessorType> ValueData(EmitterInstance->GetData(), Attr);
 
 	if (ParticleIndex >= 0 && ParticleIndex < (int32)CurrentData->GetNumInstances())
@@ -2249,4 +2255,4 @@ void UNiagaraDataInterfaceParticleRead::GetEmitterDependencies(void* PerInstance
 
 IMPLEMENT_NIAGARA_DI_PARAMETER(UNiagaraDataInterfaceParticleRead, FNiagaraDataInterfaceParametersCS_ParticleRead);
 
-#undef LOCTEXT_NAMESPACE
\ No newline at end of file
+#undef LOCTEXT_NAMESPACE
