- Fix for simulation stages that end in the name Particles not compiling.
- This would cause our Array substitution to fail, instead look for the occurance of Particles as a whole namespace.

diff --git a/Engine/Plugins/FX/Niagara/Source/NiagaraCore/Private/NiagaraCustomVersion.cpp b/Engine/Plugins/FX/Niagara/Source/NiagaraCore/Private/NiagaraCustomVersion.cpp
--- a/Engine/Plugins/FX/Niagara/Source/NiagaraCore/Private/NiagaraCustomVersion.cpp
+++ b/Engine/Plugins/FX/Niagara/Source/NiagaraCore/Private/NiagaraCustomVersion.cpp
@@ -8,4 +8,4 @@ const FGuid FNiagaraCustomVersion::GUID(0xFCF57AFA, 0x50764283, 0xB9A9E658, 0xFF
 // Register the custom version with core
 FCustomVersionRegistration GRegisterNiagaraCustomVersion(FNiagaraCustomVersion::GUID, FNiagaraCustomVersion::LatestVersion, TEXT("NiagaraVer"));
 
-const FGuid FNiagaraCustomVersion::LatestScriptCompileVersion(0xFB5CE6AE, 0xE8FF4C4D, 0xBB10A479, 0x3BE0D367);
+const FGuid FNiagaraCustomVersion::LatestScriptCompileVersion(0x5B0333B2, 0x52513F42, 0x99533E4B, 0x2401DBEE);
\ No newline at end of file
diff --git a/Engine/Plugins/FX/Niagara/Source/NiagaraEditor/Private/NiagaraHlslTranslator.cpp b/Engine/Plugins/FX/Niagara/Source/NiagaraEditor/Private/NiagaraHlslTranslator.cpp
--- a/Engine/Plugins/FX/Niagara/Source/NiagaraEditor/Private/NiagaraHlslTranslator.cpp
+++ b/Engine/Plugins/FX/Niagara/Source/NiagaraEditor/Private/NiagaraHlslTranslator.cpp
@@ -2234,12 +2234,21 @@ void FHlslNiagaraTranslator::DefineMainGPUFunctions(
 							{
 								FString RegisterName = VarName;
 								const FString AttribStr = PARAM_MAP_ATTRIBUTE_STR;
+
+								//-TODO: If a simulation stage ends in the name "Particles." (i.e. ResolveParticles) a single search can fail so we must loop to find the correct namespace
+								//       Ideally we replace this with some more robust namespace functionality.
 								int32 ParticlesIdx = RegisterName.Find(AttribStr);
-								if (ParticlesIdx != -1 && (ParticlesIdx == 0 || (ParticlesIdx > 0 && RegisterName[ParticlesIdx - 1] == '.')))
+								while (ParticlesIdx != INDEX_NONE)
 								{
-									RegisterName.RemoveAt(ParticlesIdx, AttribStr.Len());
-									RegisterName.InsertAt(ParticlesIdx, PARAM_MAP_INDICES_STR);
+									if ( (ParticlesIdx == 0 || (ParticlesIdx > 0 && RegisterName[ParticlesIdx - 1] == '.')) )
+									{
+										RegisterName.RemoveAt(ParticlesIdx, AttribStr.Len());
+										RegisterName.InsertAt(ParticlesIdx, PARAM_MAP_INDICES_STR);
+										break;
+									}
+									ParticlesIdx = RegisterName.Find(AttribStr, ESearchCase::IgnoreCase, ESearchDir::FromStart, ParticlesIdx + AttribStr.Len());
 								}
+
 								const int32 RegisterValue = Var.GetType().IsFloatPrimitive() ? FloatCounter : IntCounter;
 								HlslOutput += RegisterName + FString::Printf(TEXT(" = %d;\n"), RegisterValue);
 							}
