Niagara - Handle assignment/set variables nodes more correctly in the post load parameter fix up.

diff --git a/Engine/Plugins/FX/Niagara/Source/NiagaraEditor/Private/NiagaraGraph.cpp b/Engine/Plugins/FX/Niagara/Source/NiagaraEditor/Private/NiagaraGraph.cpp
--- a/Engine/Plugins/FX/Niagara/Source/NiagaraEditor/Private/NiagaraGraph.cpp
+++ b/Engine/Plugins/FX/Niagara/Source/NiagaraEditor/Private/NiagaraGraph.cpp
@@ -492,7 +492,7 @@ void UNiagaraGraph::ValidateDefaultPins()
 	}
 }
 
-FName StandardizeName(FName Name, ENiagaraScriptUsage Usage, bool bIsGet, bool bIsSet)
+FName UNiagaraGraph::StandardizeName(FName Name, ENiagaraScriptUsage Usage, bool bIsGet, bool bIsSet)
 {
 	bool bIsScriptUsage = 
 		Usage == ENiagaraScriptUsage::Module || 
@@ -619,17 +619,45 @@ FName StandardizeName(FName Name, ENiagaraScriptUsage Usage, bool bIsGet, bool b
 		}
 		else
 		{
-			// Otherwise it's a set node in a system or emitter script so it must be used to configure a module input
-			// so allow 1 arbitrary namespace, unless it's for a set variables node, then allow for 2 arbitrary namespaces.
+			// Otherwise it's a set node in a system or emitter script so it must be used to configure a module input.  For this situation
+			// we have 2 cases, regular modules and set variables nodes.
 			if (NameParts[0].ToString().StartsWith(TRANSLATOR_SET_VARIABLES_STR))
 			{
-				Namespace = *(NameParts[0].ToString() + TEXT(".") + NameParts[1].ToString());
-				NameParts.RemoveAt(0, 2);
+				// For a set variables node we need to strip the module name and then fully standardize the remainder of the
+				// name using the settings for a map set in a module.  We do this because the format of the input parameter will be
+				// Module.NamespaceName.ValueName and the NamespaceName.ValueName portion will have been standardized as part of the
+				// UNiagaraNodeAssignment post load.
+				Namespace = NameParts[0];
+				NameParts.RemoveAt(0);
+
+				FString AssignmentTargetString = NameParts[0].ToString();
+				for (int32 i = 1; i < NameParts.Num(); i++)
+				{
+					AssignmentTargetString += TEXT(".") + NameParts[i].ToString();
+				}
+				FName StandardizedAssignmentTargetName = StandardizeName(*AssignmentTargetString, ENiagaraScriptUsage::Module, false, true);
+
+				NameParts.Empty();
+				NameParts.Add(StandardizedAssignmentTargetName);
 			}
 			else
 			{
+				// For standard inputs we need to replace the module name with the module namespace and then standardize it, and then
+				// remove the module namespace and use that as the name, and the module name as the namespace.
 				Namespace = NameParts[0];
 				NameParts.RemoveAt(0);
+
+				FString ModuleInputString = FNiagaraConstants::ModuleNamespace.ToString();
+				for (int32 i = 0; i < NameParts.Num(); i++)
+				{
+					ModuleInputString += TEXT(".") + NameParts[i].ToString();
+				}
+				FName StandardizedModuleInput = StandardizeName(*ModuleInputString, ENiagaraScriptUsage::Module, true, false);
+				FString StandardizedModuleInputString = StandardizedModuleInput.ToString();
+				StandardizedModuleInputString.RemoveFromStart(FNiagaraConstants::ModuleNamespace.ToString() + TEXT("."));
+
+				NameParts.Empty();
+				NameParts.Add(*StandardizedModuleInputString);
 			}
 		}
 	}
diff --git a/Engine/Plugins/FX/Niagara/Source/NiagaraEditor/Private/NiagaraNodeAssignment.cpp b/Engine/Plugins/FX/Niagara/Source/NiagaraEditor/Private/NiagaraNodeAssignment.cpp
--- a/Engine/Plugins/FX/Niagara/Source/NiagaraEditor/Private/NiagaraNodeAssignment.cpp
+++ b/Engine/Plugins/FX/Niagara/Source/NiagaraEditor/Private/NiagaraNodeAssignment.cpp
@@ -23,6 +23,7 @@
 #include "ViewModels/Stack/NiagaraStackGraphUtilities.h"
 #include "Widgets/Layout/SBox.h"
 #include "Widgets/SNiagaraParameterName.h"
+#include "NiagaraCustomVersion.h"
 
 #define LOCTEXT_NAMESPACE "NiagaraNodeAssigment"
 
@@ -210,6 +211,18 @@ void UNiagaraNodeAssignment::PostLoad()
 		}
 	}
 
+	const int32 NiagaraVer = GetLinkerCustomVersion(FNiagaraCustomVersion::GUID);
+	if (NiagaraVer < FNiagaraCustomVersion::StandardizeParameterNames)
+	{
+		for (FNiagaraVariable& AssignmentTarget : AssignmentTargets)
+		{
+			FName CurrentName = AssignmentTarget.GetName();
+			FName NewName = UNiagaraGraph::StandardizeName(CurrentName, ENiagaraScriptUsage::Module, false, true);
+			AssignmentTarget.SetName(NewName);
+		}
+		RefreshFromExternalChanges();
+	}
+
 	RefreshTitle();
 }
 
diff --git a/Engine/Plugins/FX/Niagara/Source/NiagaraEditor/Public/NiagaraEditorModule.h b/Engine/Plugins/FX/Niagara/Source/NiagaraEditor/Public/NiagaraEditorModule.h
--- a/Engine/Plugins/FX/Niagara/Source/NiagaraEditor/Public/NiagaraEditorModule.h
+++ b/Engine/Plugins/FX/Niagara/Source/NiagaraEditor/Public/NiagaraEditorModule.h
@@ -32,7 +32,7 @@ class FHlslNiagaraCompiler;
 
 DECLARE_STATS_GROUP(TEXT("Niagara Editor"), STATGROUP_NiagaraEditor, STATCAT_Advanced);
 
-extern int32 GbShowNiagaraDeveloperWindows;
+extern NIAGARAEDITOR_API int32 GbShowNiagaraDeveloperWindows;
 
 /* Defines methods for allowing external modules to supply widgets to the core editor module. */
 class NIAGARAEDITOR_API INiagaraEditorWidgetProvider
diff --git a/Engine/Plugins/FX/Niagara/Source/NiagaraEditor/Public/NiagaraGraph.h b/Engine/Plugins/FX/Niagara/Source/NiagaraEditor/Public/NiagaraGraph.h
--- a/Engine/Plugins/FX/Niagara/Source/NiagaraEditor/Public/NiagaraGraph.h
+++ b/Engine/Plugins/FX/Niagara/Source/NiagaraEditor/Public/NiagaraGraph.h
@@ -384,6 +384,8 @@ class UNiagaraGraph : public UEdGraph
 
 	static FName MakeUniqueParameterNameAcrossGraphs(const FName& InName, TArray<TWeakObjectPtr<UNiagaraGraph>>& InGraphs);
 
+	static FName StandardizeName(FName Name, ENiagaraScriptUsage Usage, bool bIsGet, bool bIsSet);
+
 protected:
 	void RebuildNumericCache();
 	bool bNeedNumericCacheRebuilt;
diff --git a/Engine/Plugins/FX/Niagara/Source/NiagaraEditorWidgets/Private/Stack/SNiagaraStackModuleItem.cpp b/Engine/Plugins/FX/Niagara/Source/NiagaraEditorWidgets/Private/Stack/SNiagaraStackModuleItem.cpp
--- a/Engine/Plugins/FX/Niagara/Source/NiagaraEditorWidgets/Private/Stack/SNiagaraStackModuleItem.cpp
+++ b/Engine/Plugins/FX/Niagara/Source/NiagaraEditorWidgets/Private/Stack/SNiagaraStackModuleItem.cpp
@@ -32,6 +32,7 @@
 #include "EditorFontGlyphs.h"
 #include "ViewModels/Stack/NiagaraStackClipboardUtilities.h"
 #include "Widgets/SNiagaraLibraryOnlyToggleHeader.h"
+#include "NiagaraEditorModule.h"
 
 #define LOCTEXT_NAMESPACE "NiagaraStackModuleItem"
 
@@ -54,7 +55,7 @@ FReply SNiagaraStackModuleItem::OnMouseButtonDoubleClick(const FGeometry& InMyGe
 	const UNiagaraNodeFunctionCall& ModuleFunctionCall = ModuleItem->GetModuleNode();
 	if (ModuleFunctionCall.FunctionScript != nullptr)
 	{
-		if (ModuleFunctionCall.FunctionScript->IsAsset())
+		if (ModuleFunctionCall.FunctionScript->IsAsset() || GbShowNiagaraDeveloperWindows > 0)
 		{
 			GEditor->GetEditorSubsystem<UAssetEditorSubsystem>()->OpenEditorForAsset(const_cast<UNiagaraScript*>(ModuleFunctionCall.FunctionScript));
 			return FReply::Handled();
