Fix Niagara GPU Particle crash on older iOS and tvOS devices (pre A9/iPhone 6s)

diff --git a/Engine/Plugins/FX/Niagara/Source/Niagara/Private/NiagaraCommon.cpp b/Engine/Plugins/FX/Niagara/Source/Niagara/Private/NiagaraCommon.cpp
--- a/Engine/Plugins/FX/Niagara/Source/Niagara/Private/NiagaraCommon.cpp
+++ b/Engine/Plugins/FX/Niagara/Source/Niagara/Private/NiagaraCommon.cpp
@@ -285,12 +285,12 @@ void FNiagaraSystemUpdateContext::AddInternal(UNiagaraComponent* Comp, bool bReI
 
 bool FNiagaraUtilities::AllowGPUParticles(EShaderPlatform ShaderPlatform)
 {
-	return SupportsGPUParticles(ShaderPlatform) && GNiagaraAllowGPUParticles && GNiagaraAllowComputeShaders;
+	return SupportsGPUParticles(ShaderPlatform) && GNiagaraAllowGPUParticles && GNiagaraAllowComputeShaders && GRHISupportsDrawIndirect;
 }
 
 bool FNiagaraUtilities::AllowComputeShaders(EShaderPlatform ShaderPlatform)
 {
-	return RHISupportsComputeShaders(ShaderPlatform) && GNiagaraAllowComputeShaders;
+	return RHISupportsComputeShaders(ShaderPlatform) && GNiagaraAllowComputeShaders && GRHISupportsDrawIndirect;
 }
 
 FName NIAGARA_API FNiagaraUtilities::GetUniqueName(FName CandidateName, const TSet<FName>& ExistingNames)
diff --git a/Engine/Source/Runtime/Apple/MetalRHI/Private/MetalRHI.cpp b/Engine/Source/Runtime/Apple/MetalRHI/Private/MetalRHI.cpp
--- a/Engine/Source/Runtime/Apple/MetalRHI/Private/MetalRHI.cpp
+++ b/Engine/Source/Runtime/Apple/MetalRHI/Private/MetalRHI.cpp
@@ -133,11 +133,13 @@ FMetalDynamicRHI::FMetalDynamicRHI(ERHIFeatureLevel::Type RequestedFeatureLevel)
 #if PLATFORM_TVOS
 	bool bCanUseWideMRTs = true;
 	bool bCanUseASTC = true;
+	GRHISupportsDrawIndirect = Device.SupportsFeatureSet(mtlpp::FeatureSet::tvOS_GPUFamily2_v1);
 #else
 	bool bCanUseWideMRTs = Device.SupportsFeatureSet(mtlpp::FeatureSet::iOS_GPUFamily2_v1);
 	bool bCanUseASTC = Device.SupportsFeatureSet(mtlpp::FeatureSet::iOS_GPUFamily2_v1) && !FParse::Param(FCommandLine::Get(),TEXT("noastc"));
 	
 	GRHISupportsRWTextureBuffers = Device.SupportsFeatureSet(mtlpp::FeatureSet::iOS_GPUFamily4_v1);
+	GRHISupportsDrawIndirect = Device.SupportsFeatureSet(mtlpp::FeatureSet::iOS_GPUFamily3_v1);
 
 	const mtlpp::FeatureSet FeatureSets[] = {
 		mtlpp::FeatureSet::iOS_GPUFamily1_v1,
diff --git a/Engine/Source/Runtime/Apple/MetalRHI/Private/MetalRHIPrivate.h b/Engine/Source/Runtime/Apple/MetalRHI/Private/MetalRHIPrivate.h
--- a/Engine/Source/Runtime/Apple/MetalRHI/Private/MetalRHIPrivate.h
+++ b/Engine/Source/Runtime/Apple/MetalRHI/Private/MetalRHIPrivate.h
@@ -183,7 +183,7 @@ void SafeReleaseMetalRenderPassDescriptor(mtlpp::RenderPassDescriptor& Desc);
 // Access the underlying surface object from any kind of texture
 FMetalSurface* GetMetalSurfaceFromRHITexture(FRHITexture* Texture);
 
-#define NOT_SUPPORTED(Func) UE_LOG(LogMetal, Fatal, TEXT("'%s' is not supported"), L##Func);
+#define NOT_SUPPORTED(Func) UE_LOG(LogMetal, Fatal, TEXT("'%s' is not supported"), TEXT(Func));
 
 // Verifies we are on the correct thread to mutate internal MetalRHI resources.
 FORCEINLINE void CheckMetalThread()
diff --git a/Engine/Source/Runtime/RHI/Private/RHI.cpp b/Engine/Source/Runtime/RHI/Private/RHI.cpp
--- a/Engine/Source/Runtime/RHI/Private/RHI.cpp
+++ b/Engine/Source/Runtime/RHI/Private/RHI.cpp
@@ -697,6 +697,7 @@ bool GSupportsTexture3D = true;
 bool GSupportsMobileMultiView = false;
 bool GSupportsImageExternal = false;
 bool GSupportsResourceView = true;
+bool GRHISupportsDrawIndirect = true;
 bool GRHISupportsMultithreading = false;
 TRHIGlobal<bool> GSupportsMultipleRenderTargets(true);
 bool GSupportsWideMRT = true;
diff --git a/Engine/Source/Runtime/RHI/Public/RHI.h b/Engine/Source/Runtime/RHI/Public/RHI.h
--- a/Engine/Source/Runtime/RHI/Public/RHI.h
+++ b/Engine/Source/Runtime/RHI/Public/RHI.h
@@ -82,6 +82,9 @@ extern RHI_API EShaderPlatform GMaxRHIShaderPlatform;
 /** true if the RHI supports SRVs */
 extern RHI_API bool GSupportsResourceView;
 
+/** true if the RHI supports Draw Indirect */
+extern RHI_API bool GRHISupportsDrawIndirect;
+
 /** Whether the RHI can send commands to the device context from multiple threads. Used in the GPU readback to avoid stalling the RHI threads. */
 extern RHI_API bool GRHISupportsMultithreading;
 
