- Add method to get spline rotation to avoid having to evalulate the spline multiple times

#rnx

diff --git a/Engine/Plugins/FX/Niagara/Source/Niagara/Classes/NiagaraDataInterfaceSpline.h b/Engine/Plugins/FX/Niagara/Source/Niagara/Classes/NiagaraDataInterfaceSpline.h
--- a/Engine/Plugins/FX/Niagara/Source/Niagara/Classes/NiagaraDataInterfaceSpline.h
+++ b/Engine/Plugins/FX/Niagara/Source/Niagara/Classes/NiagaraDataInterfaceSpline.h
@@ -49,6 +49,8 @@ public:
 	template<typename TransformHandlerType, typename SplineSampleType>
 	void SampleSplinePositionByUnitDistance(FVectorVMContext& Context);
 	template<typename TransformHandlerType, typename SplineSampleType>
+	void SampleSplineRotationByUnitDistance(FVectorVMContext& Context);
+	template<typename TransformHandlerType, typename SplineSampleType>
 	void SampleSplineUpVectorByUnitDistance(FVectorVMContext& Context);
 	template<typename TransformHandlerType, typename SplineSampleType>
 	void SampleSplineRightVectorByUnitDistance(FVectorVMContext& Context);
diff --git a/Engine/Plugins/FX/Niagara/Source/Niagara/Private/NiagaraDataInterfaceSpline.cpp b/Engine/Plugins/FX/Niagara/Source/Niagara/Private/NiagaraDataInterfaceSpline.cpp
--- a/Engine/Plugins/FX/Niagara/Source/Niagara/Private/NiagaraDataInterfaceSpline.cpp
+++ b/Engine/Plugins/FX/Niagara/Source/Niagara/Private/NiagaraDataInterfaceSpline.cpp
@@ -29,6 +29,9 @@ void UNiagaraDataInterfaceSpline::PostInitProperties()
 static const FName SampleSplinePositionByUnitDistanceName("SampleSplinePositionByUnitDistance");
 static const FName SampleSplinePositionByUnitDistanceWSName("SampleSplinePositionByUnitDistanceWS");
 
+static const FName SampleSplineRotationByUnitDistanceName("SampleSplineRotationByUnitDistance");
+static const FName SampleSplineRotationByUnitDistanceWSName("SampleSplineRotationByUnitDistanceWS");
+
 static const FName SampleSplineUpVectorByUnitDistanceName("SampleSplineUpVectorByUnitDistance");
 static const FName SampleSplineUpVectorByUnitDistanceWSName("SampleSplineUpVectorByUnitDistanceWS");
 
@@ -77,6 +80,32 @@ void UNiagaraDataInterfaceSpline::GetFunctions(TArray<FNiagaraFunctionSignature>
 		OutFunctions.Add(Sig);
 	}
 
+	{
+		FNiagaraFunctionSignature Sig;
+		Sig.Name = SampleSplineRotationByUnitDistanceName;
+		Sig.Inputs.Add(FNiagaraVariable(FNiagaraTypeDefinition(GetClass()), TEXT("Spline")));
+		Sig.Inputs.Add(FNiagaraVariable(FNiagaraTypeDefinition::GetFloatDef(), TEXT("U")));
+		Sig.Outputs.Add(FNiagaraVariable(FNiagaraTypeDefinition::GetQuatDef(), TEXT("Rotation")));
+		Sig.bMemberFunction = true;
+		Sig.bRequiresContext = false;
+		//Sig.Owner = *GetName();
+		Sig.SetDescription(LOCTEXT("DataInterfaceSpline_SampleSplineRotationByUnitDistance", "Sample the spline Rotation where U is a 0 to 1 value representing the start and normalized length of the spline.\nThis is in the local space of the referenced USplineComponent."));
+		OutFunctions.Add(Sig);
+	}
+
+	{
+		FNiagaraFunctionSignature Sig;
+		Sig.Name = SampleSplineRotationByUnitDistanceWSName;
+		Sig.Inputs.Add(FNiagaraVariable(FNiagaraTypeDefinition(GetClass()), TEXT("Spline")));
+		Sig.Inputs.Add(FNiagaraVariable(FNiagaraTypeDefinition::GetFloatDef(), TEXT("U")));
+		Sig.Outputs.Add(FNiagaraVariable(FNiagaraTypeDefinition::GetQuatDef(), TEXT("Rotation")));
+		Sig.bMemberFunction = true;
+		Sig.bRequiresContext = false;
+		//Sig.Owner = *GetName();
+		Sig.SetDescription(LOCTEXT("DataInterfaceSpline_SampleSplineRotationByUnitDistanceWS", "Sample the spline Rotation where U is a 0 to 1 value representing the start and normalized length of the spline.\nThis is in the world space of the level."));
+		OutFunctions.Add(Sig);
+	}
+
 	{
 		FNiagaraFunctionSignature Sig;
 		Sig.Name = SampleSplineDirectionByUnitDistanceName;
@@ -221,6 +250,7 @@ void UNiagaraDataInterfaceSpline::GetFunctions(TArray<FNiagaraFunctionSignature>
 }
 
 DEFINE_NDI_FUNC_BINDER(UNiagaraDataInterfaceSpline, SampleSplinePositionByUnitDistance);
+DEFINE_NDI_FUNC_BINDER(UNiagaraDataInterfaceSpline, SampleSplineRotationByUnitDistance);
 DEFINE_NDI_FUNC_BINDER(UNiagaraDataInterfaceSpline, SampleSplineUpVectorByUnitDistance);
 DEFINE_NDI_FUNC_BINDER(UNiagaraDataInterfaceSpline, SampleSplineRightVectorByUnitDistance);
 DEFINE_NDI_FUNC_BINDER(UNiagaraDataInterfaceSpline, SampleSplineDirectionByUnitDistance);
@@ -239,6 +269,16 @@ void UNiagaraDataInterfaceSpline::GetVMExternalFunction(const FVMExternalFunctio
 		check(BindingInfo.GetNumInputs() == 2 && BindingInfo.GetNumOutputs() == 3);
 		TNDIExplicitBinder<FNDITransformHandler, TNDIParamBinder<1, float, NDI_FUNC_BINDER(UNiagaraDataInterfaceSpline, SampleSplinePositionByUnitDistance)>>::Bind(this, BindingInfo, InstanceData, OutFunc);
 	}
+	else if (BindingInfo.Name == SampleSplineRotationByUnitDistanceName)
+	{
+		check(BindingInfo.GetNumInputs() == 2 && BindingInfo.GetNumOutputs() == 4);
+		TNDIExplicitBinder<FNDITransformHandlerNoop, TNDIParamBinder<0, float, NDI_FUNC_BINDER(UNiagaraDataInterfaceSpline, SampleSplineRotationByUnitDistance)>>::Bind(this, BindingInfo, InstanceData, OutFunc);
+	}
+	else if (BindingInfo.Name == SampleSplineRotationByUnitDistanceWSName)
+	{
+		check(BindingInfo.GetNumInputs() == 2 && BindingInfo.GetNumOutputs() == 4);
+		TNDIExplicitBinder<FNDITransformHandler, TNDIParamBinder<0, float, NDI_FUNC_BINDER(UNiagaraDataInterfaceSpline, SampleSplineRotationByUnitDistance)>>::Bind(this, BindingInfo, InstanceData, OutFunc);
+	}
 	else if (BindingInfo.Name == SampleSplineUpVectorByUnitDistanceName)
 	{
 		check(BindingInfo.GetNumInputs() == 2 && BindingInfo.GetNumOutputs() == 3);
@@ -390,12 +430,12 @@ void UNiagaraDataInterfaceSpline::SampleSplinePositionByUnitDistance(FVectorVMCo
 
 	if (USplineComponent* SplineComponent = InstData->Component.Get())
 	{
-
+		const float SplineLength = SplineComponent->GetSplineLength();
 		for (int32 i = 0; i < Context.NumInstances; ++i)
 		{
 			float DistanceUnitDistance = SplineSampleParam.Get();
 
-			FVector Pos = SplineComponent->GetLocationAtDistanceAlongSpline(DistanceUnitDistance * SplineComponent->GetSplineLength(), ESplineCoordinateSpace::Local);
+			FVector Pos = SplineComponent->GetLocationAtDistanceAlongSpline(DistanceUnitDistance * SplineLength, ESplineCoordinateSpace::Local);
 			TransformHandler.TransformPosition(Pos, InstData->Transform);
 
 			*OutPosX.GetDest() = Pos.X;
@@ -427,6 +467,46 @@ void UNiagaraDataInterfaceSpline::SampleSplinePositionByUnitDistance(FVectorVMCo
 	}
 }
 
+template<typename TransformHandlerType, typename SplineSampleType>
+void UNiagaraDataInterfaceSpline::SampleSplineRotationByUnitDistance(FVectorVMContext& Context)
+{
+	TransformHandlerType TransformHandler;
+	SplineSampleType SplineSampleParam(Context);
+	VectorVM::FUserPtrHandler<FNDISpline_InstanceData> InstData(Context);
+	VectorVM::FExternalFuncRegisterHandler<float> OutQuatX(Context);
+	VectorVM::FExternalFuncRegisterHandler<float> OutQuatY(Context);
+	VectorVM::FExternalFuncRegisterHandler<float> OutQuatZ(Context);
+	VectorVM::FExternalFuncRegisterHandler<float> OutQuatW(Context);
+
+	if (USplineComponent* SplineComponent = InstData->Component.Get())
+	{
+		const FQuat TransformQuat = InstData->Transform.ToQuat();
+		const float SplineLength = SplineComponent->GetSplineLength();
+		for (int32 i = 0; i < Context.NumInstances; ++i)
+		{
+			const float DistanceUnitDistance = SplineSampleParam.GetAndAdvance();
+
+			FQuat Quat = SplineComponent->GetQuaternionAtDistanceAlongSpline(DistanceUnitDistance * SplineLength, ESplineCoordinateSpace::Local);
+			TransformHandler.TransformRotation(Quat, TransformQuat);
+
+			*OutQuatX.GetDestAndAdvance() = Quat.X;
+			*OutQuatY.GetDestAndAdvance() = Quat.Y;
+			*OutQuatZ.GetDestAndAdvance() = Quat.Z;
+			*OutQuatW.GetDestAndAdvance() = Quat.W;
+		}
+	}
+	else
+	{
+		for (int32 i = 0; i < Context.NumInstances; ++i)
+		{
+			*OutQuatX.GetDestAndAdvance() = FQuat::Identity.X;
+			*OutQuatY.GetDestAndAdvance() = FQuat::Identity.Y;
+			*OutQuatZ.GetDestAndAdvance() = FQuat::Identity.Z;
+			*OutQuatW.GetDestAndAdvance() = FQuat::Identity.W;
+		}
+	}
+}
+
 template<typename TransformHandlerType, typename SplineSampleType>
 void UNiagaraDataInterfaceSpline::SampleSplineUpVectorByUnitDistance(FVectorVMContext& Context)
 {
@@ -439,12 +519,12 @@ void UNiagaraDataInterfaceSpline::SampleSplineUpVectorByUnitDistance(FVectorVMCo
 
 	if (USplineComponent* SplineComponent = InstData->Component.Get())
 	{
-
+		const float SplineLength = SplineComponent->GetSplineLength();
 		for (int32 i = 0; i < Context.NumInstances; ++i)
 		{
 			float DistanceUnitDistance = SplineSampleParam.Get();
 
-			FVector Pos = SplineComponent->GetUpVectorAtDistanceAlongSpline(DistanceUnitDistance * SplineComponent->GetSplineLength(), ESplineCoordinateSpace::Local);
+			FVector Pos = SplineComponent->GetUpVectorAtDistanceAlongSpline(DistanceUnitDistance * SplineLength, ESplineCoordinateSpace::Local);
 			TransformHandler.TransformVector(Pos, InstData->Transform);
 
 			*OutPosX.GetDest() = Pos.X;
@@ -489,11 +569,12 @@ void UNiagaraDataInterfaceSpline::SampleSplineRightVectorByUnitDistance(FVectorV
 	
 	if (USplineComponent* SplineComponent = InstData->Component.Get())
 	{
+		const float SplineLength = SplineComponent->GetSplineLength();
 		for (int32 i = 0; i < Context.NumInstances; ++i)
 		{
 			float DistanceUnitDistance = SplineSampleParam.Get();
 
-			FVector Pos = SplineComponent->GetRightVectorAtDistanceAlongSpline(DistanceUnitDistance * SplineComponent->GetSplineLength(), ESplineCoordinateSpace::Local);
+			FVector Pos = SplineComponent->GetRightVectorAtDistanceAlongSpline(DistanceUnitDistance * SplineLength, ESplineCoordinateSpace::Local);
 			TransformHandler.TransformVector(Pos, InstData->Transform);
 
 			*OutPosX.GetDest() = Pos.X;
@@ -535,11 +616,12 @@ void UNiagaraDataInterfaceSpline::SampleSplineTangentByUnitDistance(FVectorVMCon
 
 	if (USplineComponent* SplineComponent = InstData->Component.Get())
 	{
+		const float SplineLength = SplineComponent->GetSplineLength();
 		for (int32 i = 0; i < Context.NumInstances; ++i)
 		{
 			float DistanceUnitDistance = SplineSampleParam.Get();
 
-			FVector Pos = SplineComponent->GetTangentAtDistanceAlongSpline(DistanceUnitDistance * SplineComponent->GetSplineLength(), ESplineCoordinateSpace::Local);
+			FVector Pos = SplineComponent->GetTangentAtDistanceAlongSpline(DistanceUnitDistance * SplineLength, ESplineCoordinateSpace::Local);
 			TransformHandler.TransformVector(Pos, InstData->Transform);
 
 			*OutPosX.GetDest() = Pos.X;
@@ -583,11 +665,12 @@ void UNiagaraDataInterfaceSpline::SampleSplineDirectionByUnitDistance(FVectorVMC
 
 	if (USplineComponent* SplineComponent = InstData->Component.Get())
 	{
+		const float SplineLength = SplineComponent->GetSplineLength();
 		for (int32 i = 0; i < Context.NumInstances; ++i)
 		{
 			float DistanceUnitDistance = SplineSampleParam.Get();
 
-			FVector Pos = SplineComponent->GetDirectionAtDistanceAlongSpline(DistanceUnitDistance * SplineComponent->GetSplineLength(), ESplineCoordinateSpace::Local);
+			FVector Pos = SplineComponent->GetDirectionAtDistanceAlongSpline(DistanceUnitDistance * SplineLength, ESplineCoordinateSpace::Local);
 			TransformHandler.TransformVector(Pos, InstData->Transform);
 
 			*OutPosX.GetDest() = Pos.X;
