Added DDC stats to UNiagaraScript

#rnx

diff --git a/Engine/Plugins/FX/Niagara/Source/Niagara/Private/NiagaraScript.cpp b/Engine/Plugins/FX/Niagara/Source/Niagara/Private/NiagaraScript.cpp
--- a/Engine/Plugins/FX/Niagara/Source/Niagara/Private/NiagaraScript.cpp
+++ b/Engine/Plugins/FX/Niagara/Source/Niagara/Private/NiagaraScript.cpp
@@ -13,6 +13,7 @@
 #include "Serialization/MemoryReader.h"
 #include "Misc/SecureHash.h"
 
+#include "ProfilingDebugging/CookStats.h"
 #include "Stats/Stats.h"
 #include "UObject/Linker.h"
 #include "HAL/PlatformFilemanager.h"
@@ -43,6 +44,17 @@
 
 DECLARE_STATS_GROUP(TEXT("Niagara Detailed"), STATGROUP_NiagaraDetailed, STATCAT_Advanced);
 
+#if ENABLE_COOK_STATS
+namespace NiagaraScriptCookStats
+{
+	FCookStats::FDDCResourceUsageStats UsageStats;
+	static FCookStatsManager::FAutoRegisterCallback RegisterCookStats([](FCookStatsManager::AddStatFuncRef AddStat)
+	{
+		UsageStats.LogStats(AddStat, TEXT("NiagaraScript.Usage"), TEXT(""));
+	});
+}
+#endif
+
 int32 GNiagaraDumpKeyGen = 0;
 static FAutoConsoleVariableRef CVarNiagaraDumpKeyGen(
 	TEXT("fx.DumpGraphKeyGen"),
@@ -1562,6 +1574,7 @@ void UNiagaraScript::RequestCompile(bool bForceCompile)
 			return;
 		}
 
+		COOK_STAT(auto Timer = NiagaraScriptCookStats::UsageStats.TimeSyncWork());
 
 		CachedScriptVM.LastCompileStatus = ENiagaraScriptCompileStatus::NCS_BeingCreated;
 
@@ -1572,6 +1585,7 @@ void UNiagaraScript::RequestCompile(bool bForceCompile)
 
 		if (RequestData.IsValid() == false)
 		{
+			COOK_STAT(Timer.TrackCyclesOnly());
 			UE_LOG(LogNiagara, Error, TEXT("Failed to precompile %s.  This is due to unexpected invalid or broken data.  Additional details should be in the log."), *GetPathName());
 			return;
 		}
@@ -1582,6 +1596,7 @@ void UNiagaraScript::RequestCompile(bool bForceCompile)
 			FNiagaraVMExecutableData ExeData;
 			if (BinaryToExecData(OutData, ExeData))
 			{
+				COOK_STAT(Timer.AddHit(OutData.Num()));
 				SetVMCompilationResults(LastGeneratedVMId, ExeData, RequestData.Get());
 				return;
 			}
@@ -1599,6 +1614,7 @@ void UNiagaraScript::RequestCompile(bool bForceCompile)
 			// save result to the ddc
 			if (ExecToBinaryData(OutData, *ExeData))
 			{
+				COOK_STAT(Timer.AddMiss(OutData.Num()));
 				GetDerivedDataCacheRef().Put(*GetNiagaraDDCKeyString(), OutData, GetPathName());
 			}
 		}
@@ -1612,6 +1628,9 @@ void UNiagaraScript::RequestCompile(bool bForceCompile)
 
 bool UNiagaraScript::RequestExternallyManagedAsyncCompile(const TSharedPtr<FNiagaraCompileRequestDataBase, ESPMode::ThreadSafe>& RequestData, FNiagaraVMExecutableDataId& OutCompileId, uint32& OutAsyncHandle)
 {
+	COOK_STAT(auto Timer = NiagaraScriptCookStats::UsageStats.TimeSyncWork());
+	COOK_STAT(Timer.TrackCyclesOnly());
+
 	OutCompileId = LastGeneratedVMId;
 
 	if (!AreScriptAndSourceSynchronized())
diff --git a/Engine/Plugins/FX/Niagara/Source/Niagara/Private/NiagaraSystem.cpp b/Engine/Plugins/FX/Niagara/Source/Niagara/Private/NiagaraSystem.cpp
--- a/Engine/Plugins/FX/Niagara/Source/Niagara/Private/NiagaraSystem.cpp
+++ b/Engine/Plugins/FX/Niagara/Source/Niagara/Private/NiagaraSystem.cpp
@@ -21,6 +21,7 @@
 #include "NiagaraEmitterInstanceBatcher.h"
 #include "NiagaraSettings.h"
 #include "Interfaces/ITargetPlatform.h"
+#include "ProfilingDebugging/CookStats.h"
 
 
 #if WITH_EDITOR
@@ -31,6 +32,13 @@ DECLARE_CYCLE_STAT(TEXT("Niagara - System - Precompile"), STAT_Niagara_System_Pr
 DECLARE_CYCLE_STAT(TEXT("Niagara - System - CompileScript"), STAT_Niagara_System_CompileScript, STATGROUP_Niagara);
 DECLARE_CYCLE_STAT(TEXT("Niagara - System - CompileScript_ResetAfter"), STAT_Niagara_System_CompileScriptResetAfter, STATGROUP_Niagara);
 
+#if ENABLE_COOK_STATS
+namespace NiagaraScriptCookStats
+{
+	extern FCookStats::FDDCResourceUsageStats UsageStats;
+}
+#endif
+
 //Disable for now until we can spend more time on a good method of applying the data gathered.
 int32 GEnableNiagaraRuntimeCycleCounts = 0;
 static FAutoConsoleVariableRef CVarEnableNiagaraRuntimeCycleCounts(TEXT("fx.EnableNiagaraRuntimeCycleCounts"), GEnableNiagaraRuntimeCycleCounts, TEXT("Toggle for runtime cylce counts tracking Niagara's frame time. \n"), ECVF_ReadOnly);
@@ -1033,11 +1041,14 @@ bool UNiagaraSystem::QueryCompileComplete(bool bWait, bool bDoPost, bool bDoNotA
 
 bool UNiagaraSystem::ProcessCompilationResult(FEmitterCompiledScriptPair& ScriptPair, bool bWait, bool bDoNotApply)
 {
+	COOK_STAT(auto Timer = NiagaraScriptCookStats::UsageStats.TimeAsyncWait());
+
 	INiagaraModule& NiagaraModule = FModuleManager::Get().LoadModuleChecked<INiagaraModule>(TEXT("Niagara"));
 	TSharedPtr<FNiagaraVMExecutableData> ExeData = NiagaraModule.GetCompileJobResult(ScriptPair.PendingJobID, bWait);
 
 	if (!bWait && !ExeData.IsValid())
 	{
+		COOK_STAT(Timer.TrackCyclesOnly());
 		return false;
 	}
 	check(ExeData.IsValid());
@@ -1045,20 +1056,24 @@ bool UNiagaraSystem::ProcessCompilationResult(FEmitterCompiledScriptPair& Script
 	{
 		ScriptPair.CompileResults = ExeData;
 	}
-	
+
 	// save result to the ddc
 	TArray<uint8> OutData;
 	if (UNiagaraScript::ExecToBinaryData(OutData, *ExeData))
 	{
+		COOK_STAT(Timer.AddMiss(OutData.Num()));
 		GetDerivedDataCacheRef().Put(*ScriptPair.CompiledScript->GetNiagaraDDCKeyString(), OutData, GetPathName());
 		return true;
 	}
-	
+
+	COOK_STAT(Timer.TrackCyclesOnly());
 	return false;
 }
 
 bool UNiagaraSystem::GetFromDDC(FEmitterCompiledScriptPair& ScriptPair)
 {
+	COOK_STAT(auto Timer = NiagaraScriptCookStats::UsageStats.TimeSyncWork());
+
 	FNiagaraVMExecutableDataId NewID;
 	ScriptPair.CompiledScript->ComputeVMCompilationId(NewID);
 	ScriptPair.CompileId = NewID;
@@ -1069,12 +1084,15 @@ bool UNiagaraSystem::GetFromDDC(FEmitterCompiledScriptPair& ScriptPair)
 		TSharedPtr<FNiagaraVMExecutableData> ExeData = MakeShared<FNiagaraVMExecutableData>();
 		if (ScriptPair.CompiledScript->BinaryToExecData(Data, *ExeData))
 		{
+			COOK_STAT(Timer.AddHit(Data.Num()));
 			ExeData->CompileTime = 0; // since we didn't actually compile anything
 			ScriptPair.CompileResults = ExeData;
 			ScriptPair.bResultsReady = true;
 			return true;
 		}
 	}
+
+	COOK_STAT(Timer.TrackCyclesOnly());
 	return false;
 }
 
@@ -1134,30 +1152,37 @@ bool UNiagaraSystem::RequestCompile(bool bForce, FNiagaraSystemUpdateContext* Op
 	TArray<FNiagaraVariable> OriginalExposedParams;
 	GetExposedParameters().GetParameters(OriginalExposedParams);
 
-	INiagaraModule& NiagaraModule = FModuleManager::Get().LoadModuleChecked<INiagaraModule>(TEXT("Niagara"));
-	TSharedPtr<FNiagaraCompileRequestDataBase, ESPMode::ThreadSafe> SystemPrecompiledData = NiagaraModule.Precompile(this);
+	TSharedPtr<FNiagaraCompileRequestDataBase, ESPMode::ThreadSafe> SystemPrecompiledData;
 
-	if (SystemPrecompiledData.IsValid() == false)
+	// Time this block as sync work as it is a prerequisite to generating the key for DDC.
 	{
-		UE_LOG(LogNiagara, Error, TEXT("Failed to precompile %s.  This is due to unexpected invalid or broken data.  Additional details should be in the log."), *GetPathName());
-		return false;
-	}
+		COOK_STAT(auto Timer = NiagaraScriptCookStats::UsageStats.TimeSyncWork());
+		COOK_STAT(Timer.TrackCyclesOnly());
+
+		INiagaraModule& NiagaraModule = FModuleManager::Get().LoadModuleChecked<INiagaraModule>(TEXT("Niagara"));
+		SystemPrecompiledData = NiagaraModule.Precompile(this);
 
-	SystemPrecompiledData->GetReferencedObjects(ActiveCompilations[ActiveCompileIdx].RootObjects);
+		if (SystemPrecompiledData.IsValid() == false)
+		{
+			UE_LOG(LogNiagara, Error, TEXT("Failed to precompile %s.  This is due to unexpected invalid or broken data.  Additional details should be in the log."), *GetPathName());
+			return false;
+		}
 
-	//Compile all emitters
-	bool bTrulyAsync = true;
-	bool bAnyUnsynchronized = false;	
+		SystemPrecompiledData->GetReferencedObjects(ActiveCompilations[ActiveCompileIdx].RootObjects);
 
-	ActiveCompilations[ActiveCompileIdx].MappedData.Add(SystemSpawnScript, SystemPrecompiledData);
-	ActiveCompilations[ActiveCompileIdx].MappedData.Add(SystemUpdateScript, SystemPrecompiledData);
+		//Compile all emitters
+		ActiveCompilations[ActiveCompileIdx].MappedData.Add(SystemSpawnScript, SystemPrecompiledData);
+		ActiveCompilations[ActiveCompileIdx].MappedData.Add(SystemUpdateScript, SystemPrecompiledData);
 
-	check(EmitterHandles.Num() == SystemPrecompiledData->GetDependentRequestCount());
+		check(EmitterHandles.Num() == SystemPrecompiledData->GetDependentRequestCount());
+	}
 
 	// Grab the list of user variables that were actually encountered so that we can add to them later.
 	TArray<FNiagaraVariable> EncounteredExposedVars;
 	SystemPrecompiledData->GatherPreCompiledVariables(TEXT("User"), EncounteredExposedVars);
 
+	bool bAnyUnsynchronized = false;
+
 	for (int32 i = 0; i < EmitterHandles.Num(); i++)
 	{
 		FNiagaraEmitterHandle Handle = EmitterHandles[i];
