Properly display aliased parameters for parameter writes entries in the stack.
Fix parameter reads and writes collecting entries for parameters of the opposing type (reads should not display Outputs, vice versa.)

diff --git a/Engine/Plugins/FX/Niagara/Source/NiagaraEditor/Private/ViewModels/Stack/NiagaraStackModuleItemLinkedInputCollection.cpp b/Engine/Plugins/FX/Niagara/Source/NiagaraEditor/Private/ViewModels/Stack/NiagaraStackModuleItemLinkedInputCollection.cpp
--- a/Engine/Plugins/FX/Niagara/Source/NiagaraEditor/Private/ViewModels/Stack/NiagaraStackModuleItemLinkedInputCollection.cpp
+++ b/Engine/Plugins/FX/Niagara/Source/NiagaraEditor/Private/ViewModels/Stack/NiagaraStackModuleItemLinkedInputCollection.cpp
@@ -77,7 +77,7 @@ void UNiagaraStackModuleItemLinkedInputCollection::RefreshChildrenInternal(const
 				if (ReadHistory.Num() > 0)
 				{
 					const FNiagaraNamespaceMetadata NamespaceMetadata = FNiagaraEditorUtilities::GetNamespaceMetaDataForVariableName(Variable.GetName());
-					if ((NamespaceMetadata.IsValid()) && (NamespaceMetadata.Namespaces.Contains(FNiagaraConstants::LocalNamespace) == false))
+					if ( (NamespaceMetadata.IsValid()) && (NamespaceMetadata.Namespaces.Contains(FNiagaraConstants::LocalNamespace) == false) && (NamespaceMetadata.Namespaces.Contains(FNiagaraConstants::OutputNamespace) == false) )
 					{
 						for (const TTuple<const UEdGraphPin*, const UEdGraphPin*>& ReadPair : ReadHistory)
 						{
diff --git a/Engine/Plugins/FX/Niagara/Source/NiagaraEditor/Private/ViewModels/Stack/NiagaraStackModuleItemOutputCollection.cpp b/Engine/Plugins/FX/Niagara/Source/NiagaraEditor/Private/ViewModels/Stack/NiagaraStackModuleItemOutputCollection.cpp
--- a/Engine/Plugins/FX/Niagara/Source/NiagaraEditor/Private/ViewModels/Stack/NiagaraStackModuleItemOutputCollection.cpp
+++ b/Engine/Plugins/FX/Niagara/Source/NiagaraEditor/Private/ViewModels/Stack/NiagaraStackModuleItemOutputCollection.cpp
@@ -55,13 +55,13 @@ void UNiagaraStackModuleItemOutputCollection::RefreshChildrenInternal(const TArr
 		FCompileConstantResolver ConstantResolver = GetEmitterViewModel().IsValid()
 			? FCompileConstantResolver(GetEmitterViewModel()->GetEmitter())
 			: FCompileConstantResolver();
-		FNiagaraStackGraphUtilities::GetStackFunctionOutputVariables(*FunctionCallNode, ConstantResolver, Unused, OutputVariables);
+		FNiagaraStackGraphUtilities::GetStackFunctionOutputVariables(*FunctionCallNode, ConstantResolver, OutputVariables, Unused);
 		
 		for(int32 OutputIndex = 0; OutputIndex < OutputVariables.Num(); OutputIndex++)
 		{
 			FNiagaraVariable& Variable = OutputVariables[OutputIndex];
 			const FNiagaraNamespaceMetadata NamespaceMetadata = FNiagaraEditorUtilities::GetNamespaceMetaDataForVariableName(Variable.GetName());
-			if ((NamespaceMetadata.IsValid()) && (NamespaceMetadata.Namespaces.Contains(FNiagaraConstants::LocalNamespace) == false))
+			if ( (NamespaceMetadata.IsValid()) && (NamespaceMetadata.Namespaces.Contains(FNiagaraConstants::LocalNamespace) == false) && (NamespaceMetadata.Namespaces.Contains(FNiagaraConstants::ModuleNamespace) == false) )
 			{
 				UNiagaraStackModuleItemOutput* Output = FindCurrentChildOfTypeByPredicate<UNiagaraStackModuleItemOutput>(CurrentChildren,
 					[&](UNiagaraStackModuleItemOutput* CurrentOutput) { return CurrentOutput->GetOutputParameterHandle().GetParameterHandleString() == Variable.GetName(); });
