Clamp the number of particles so that we don't exceed the maximum GPU buffer size.

#rnx

diff --git a/Engine/Plugins/FX/Niagara/Source/Niagara/Classes/NiagaraDataSet.h b/Engine/Plugins/FX/Niagara/Source/Niagara/Classes/NiagaraDataSet.h
--- a/Engine/Plugins/FX/Niagara/Source/Niagara/Classes/NiagaraDataSet.h
+++ b/Engine/Plugins/FX/Niagara/Source/Niagara/Classes/NiagaraDataSet.h
@@ -216,8 +216,8 @@ private:
 	// GPU Data
 	/** The buffer offset where the instance count is accumulated. */
 	uint32 GPUInstanceCountBufferOffset;
-	/** The num of allocated chunks, each being of size ALLOC_CHUNKSIZE */
-	uint32 NumChunksAllocatedForGPU;
+	/** The num of allocated instances, which is aligned to ALLOC_CHUNKSIZE and usually larger than the needed number to amortize allocation cost. */
+	uint32 NumInstancesAllocatedForGPU;
 	/** GPU Buffer containing floating point values for GPU simulations. */
 	FRWBuffer GPUBufferFloat;
 	/** GPU Buffer containing integer values for GPU simulations. */
@@ -384,6 +384,9 @@ public:
 
 	void AllocateGPUFreeIDs(uint32 InNumInstances, FRHICommandList& RHICmdList, ERHIFeatureLevel::Type FeatureLevel, const TCHAR* DebugSimName);
 
+	void SetMaxInstanceCount(uint32 InInstanceCount) { MaxInstanceCount = InInstanceCount; }
+	uint32 GetMaxInstanceCount() const { return MaxInstanceCount; }
+
 private:
 
 	void Reset();
@@ -443,6 +446,8 @@ private:
 	*/
 	TArray<FNiagaraDataBuffer*, TInlineAllocator<2>> Data;
 
+	uint32 MaxInstanceCount;
+
 	bool bInitialized;
 };
 
diff --git a/Engine/Plugins/FX/Niagara/Source/Niagara/Classes/NiagaraEmitterInstance.h b/Engine/Plugins/FX/Niagara/Source/Niagara/Classes/NiagaraEmitterInstance.h
--- a/Engine/Plugins/FX/Niagara/Source/Niagara/Classes/NiagaraEmitterInstance.h
+++ b/Engine/Plugins/FX/Niagara/Source/Niagara/Classes/NiagaraEmitterInstance.h
@@ -211,6 +211,8 @@ private:
 	int32 MinOverallocation = -1;
 	int32 ReallocationCount = 0;
 
+	uint32 MaxInstanceCount = 0;
+
 	/* Emitter tick state */
 	ENiagaraExecutionState ExecutionState = ENiagaraExecutionState::Inactive;
 
diff --git a/Engine/Plugins/FX/Niagara/Source/Niagara/Private/NiagaraDataSet.cpp b/Engine/Plugins/FX/Niagara/Source/Niagara/Private/NiagaraDataSet.cpp
--- a/Engine/Plugins/FX/Niagara/Source/Niagara/Private/NiagaraDataSet.cpp
+++ b/Engine/Plugins/FX/Niagara/Source/Niagara/Private/NiagaraDataSet.cpp
@@ -54,7 +54,7 @@ static int32 GNiagaraDataBufferMinSize = 512;
 static FAutoConsoleVariableRef CVarRenderDataBlockSize(
 	TEXT("fx.NiagaraDataBufferMinSize"),
 	GNiagaraDataBufferMinSize,
-	TEXT("Niagara data buffer minimum allocation size in bytes (Default=512). \n"),
+	TEXT("Niagara data buffer minimum allocation size in bytes (Default=512)."),
 	ECVF_Default
 );
 
@@ -63,7 +63,7 @@ static FAutoConsoleVariableRef CVarNiagaraRenderBufferShrinkFactor(
 	TEXT("fx.NiagaraDataBufferShrinkFactor"),
 	GNiagaraDataBufferShrinkFactor,
 	TEXT("Niagara data buffer size threshold for shrinking. (Default=3) \n")
-	TEXT("The buffer will be reallocated when the used size becomes 1/F of the allocated size. \n"),
+	TEXT("The buffer will be reallocated when the used size becomes 1/F of the allocated size."),
 	ECVF_Default
 );
 
@@ -71,7 +71,7 @@ static float GNiagaraGPUDataBufferBufferSlack = 1.1;
 static FAutoConsoleVariableRef CVarNiagaraGPUDataBufferBufferSlack(
 	TEXT("fx.NiagaraGPUDataBufferBufferSlack"),
 	GNiagaraGPUDataBufferBufferSlack,
-	TEXT("Niagara GPU data buffer size threshold for resizing. <= 1 to disable shrinking. (Default=1.1) \n"),
+	TEXT("Niagara GPU data buffer size threshold for resizing. <= 1 to disable shrinking. (Default=1.1)"),
 	ECVF_Default
 );
 
@@ -82,6 +82,7 @@ FNiagaraDataSet::FNiagaraDataSet()
 	, GPUNumAllocatedIDs(0)
 	, CurrentData(nullptr)
 	, DestinationData(nullptr)
+	, MaxInstanceCount(UINT_MAX)
 	, bInitialized(false)
 {
 	CompiledData.Init(&FNiagaraDataSetCompiledData::DummyCompiledData);
@@ -488,7 +489,7 @@ void FNiagaraDataSet::CopyFromGPUReadback(const float* GPUReadBackFloat, const i
 FNiagaraDataBuffer::FNiagaraDataBuffer(FNiagaraDataSet* InOwner)
 	: Owner(InOwner)
 	, GPUInstanceCountBufferOffset(INDEX_NONE)
-	, NumChunksAllocatedForGPU(0)
+	, NumInstancesAllocatedForGPU(0)
 	, NumInstances(0)
 	, NumInstancesAllocated(0)
 	, FloatStride(0)
@@ -703,8 +704,14 @@ void FNiagaraDataBuffer::AllocateGPU(uint32 InNumInstances, FNiagaraGPUInstanceC
 
 	NumInstancesAllocated = InNumInstances;
 
-	// Round the count up to the nearest threadgroup size
+	// GetMaxInstanceCount() returns the maximum number of usable instances, but it's computed in such a way as to allow fitting an extra
+	// scratch instance in the buffer. Our allocation maximum is therefore one more than what this function returns.
+	const uint32 MaxAllocatedInstances = Owner->GetMaxInstanceCount() + 1;
+
+	// Round the count up to the nearest threadgroup size. GetMaxNumInstances() ensures that the returned value is aligned to NIAGARA_COMPUTE_THREADGROUP_SIZE, so if the calling
+	// code clamps the instance count correctly, this operation should never exceed the max instance count.
 	const uint32 PaddedNumInstances = FMath::DivideAndRoundUp(NumInstancesAllocated, NIAGARA_COMPUTE_THREADGROUP_SIZE) * NIAGARA_COMPUTE_THREADGROUP_SIZE;
+	check(PaddedNumInstances <= MaxAllocatedInstances);
 
 	// Pack the data so that the space between elements is the padded thread group size
 	FloatStride = PaddedNumInstances * sizeof(float);
@@ -713,7 +720,6 @@ void FNiagaraDataBuffer::AllocateGPU(uint32 InNumInstances, FNiagaraGPUInstanceC
 
 	DEC_MEMORY_STAT_BY(STAT_NiagaraGPUParticleMemory, GPUBufferFloat.NumBytes + GPUBufferInt.NumBytes + GPUIDToIndexTable.NumBytes);
 
-	// This never seems to happen because of the +1 in NiagaraEmitterInstanceBatcher::ResizeBuffersAndGatherResources()
 	if (PaddedNumInstances == 0)
 	{
 		if (GPUBufferFloat.Buffer)
@@ -734,15 +740,18 @@ void FNiagaraDataBuffer::AllocateGPU(uint32 InNumInstances, FNiagaraGPUInstanceC
 			GPUIDToIndexTable.Release();
 		}
 
-		NumChunksAllocatedForGPU = 0;
+		NumInstancesAllocatedForGPU = 0;
 	}
 	else // Otherwise check for growing and possibly shrinking (if GNiagaraGPUDataBufferBufferSlack > 1) .
 	{
-		const uint32 RecommendedNumChunks = FMath::DivideAndRoundUp<uint32>((uint32)(PaddedNumInstances * FMath::Max<float>(GNiagaraGPUDataBufferBufferSlack, 1.f)), ALLOC_CHUNKSIZE); 
-		if (PaddedNumInstances > NumChunksAllocatedForGPU * ALLOC_CHUNKSIZE || (GNiagaraGPUDataBufferBufferSlack > 1.f && (uint32)(RecommendedNumChunks * GNiagaraGPUDataBufferBufferSlack) < NumChunksAllocatedForGPU))
+		const uint32 NumInstancesWithSlack = (uint32)(PaddedNumInstances * FMath::Max<float>(GNiagaraGPUDataBufferBufferSlack, 1.0f));
+		uint32 NumInstancesChunkAligned = FMath::DivideAndRoundUp<uint32>(NumInstancesWithSlack, ALLOC_CHUNKSIZE) * ALLOC_CHUNKSIZE;
+		// Make sure we don't exceed the instance limit by aligning to the chunk size.
+		NumInstancesChunkAligned = FMath::Min(NumInstancesChunkAligned, MaxAllocatedInstances);
+
+		if (PaddedNumInstances > NumInstancesAllocatedForGPU || (GNiagaraGPUDataBufferBufferSlack > 1.f && (uint32)(NumInstancesChunkAligned * GNiagaraGPUDataBufferBufferSlack) < NumInstancesAllocatedForGPU))
 		{
-			NumChunksAllocatedForGPU = RecommendedNumChunks; 
-			const uint32 NumElementsToAlloc = NumChunksAllocatedForGPU * ALLOC_CHUNKSIZE;
+			NumInstancesAllocatedForGPU = NumInstancesChunkAligned;
 
 			uint32 DataBufferFlags = BUF_Static;
 #if WITH_EDITORONLY_DATA
@@ -756,7 +765,7 @@ void FNiagaraDataBuffer::AllocateGPU(uint32 InNumInstances, FNiagaraGPUInstanceC
 				{
 					GPUBufferFloat.Release();
 				}
-				GPUBufferFloat.Initialize(sizeof(float), NumElementsToAlloc * Owner->GetNumFloatComponents(), EPixelFormat::PF_R32_FLOAT, DataBufferFlags, TEXT("NiagaraFloatDataBuffer"));
+				GPUBufferFloat.Initialize(sizeof(float), NumInstancesAllocatedForGPU * Owner->GetNumFloatComponents(), EPixelFormat::PF_R32_FLOAT, DataBufferFlags, TEXT("NiagaraFloatDataBuffer"));
 			}
 			if (Owner->GetNumInt32Components())
 			{
@@ -764,7 +773,7 @@ void FNiagaraDataBuffer::AllocateGPU(uint32 InNumInstances, FNiagaraGPUInstanceC
 				{
 					GPUBufferInt.Release();
 				}
-				GPUBufferInt.Initialize(sizeof(int32), NumElementsToAlloc * Owner->GetNumInt32Components(), EPixelFormat::PF_R32_SINT, DataBufferFlags, TEXT("NiagaraIntDataBuffer"));
+				GPUBufferInt.Initialize(sizeof(int32), NumInstancesAllocatedForGPU * Owner->GetNumInt32Components(), EPixelFormat::PF_R32_SINT, DataBufferFlags, TEXT("NiagaraIntDataBuffer"));
 			}
 			if (Owner->GetNumHalfComponents())
 			{
diff --git a/Engine/Plugins/FX/Niagara/Source/Niagara/Private/NiagaraEmitterInstance.cpp b/Engine/Plugins/FX/Niagara/Source/Niagara/Private/NiagaraEmitterInstance.cpp
--- a/Engine/Plugins/FX/Niagara/Source/Niagara/Private/NiagaraEmitterInstance.cpp
+++ b/Engine/Plugins/FX/Niagara/Source/Niagara/Private/NiagaraEmitterInstance.cpp
@@ -84,6 +84,15 @@ static FAutoConsoleVariableRef CVarMaxNiagaraGPUParticlesSpawnPerFrame(
 	TEXT("The max number of GPU particles we expect to spawn in a single frame.\n"),
 	ECVF_Default
 );
+
+static int32 GDebugForcedMaxGPUBufferElements = 0;
+static FAutoConsoleVariableRef CVarNiagaraDebugForcedMaxGPUBufferElements(
+	TEXT("fx.NiagaraDebugForcedMaxGPUBufferElements"),
+	GDebugForcedMaxGPUBufferElements,
+	TEXT("Force the maximum buffer size supported by the GPU to this value, for debugging purposes."),
+	ECVF_Default
+);
+
 //////////////////////////////////////////////////////////////////////////
 
 template<bool bAccumulate>
@@ -478,6 +487,39 @@ void FNiagaraEmitterInstance::Init(int32 InEmitterIdx, FNiagaraSystemInstanceID
 			}
 		}
 	}	
+
+	// Prevent division by 0 in case there are no renderers.
+	uint32 MaxGPUBufferComponents = 1;
+	if (CachedEmitter->SimTarget == ENiagaraSimTarget::CPUSim)
+	{
+		// CPU emitters only upload the data needed by the renderers to the GPU. Compute the maximum number of components per particle
+		// among all the enabled renderers, since this will decide how many particles we can upload.
+		for (UNiagaraRendererProperties* RendererProperty : CachedEmitter->GetEnabledRenderers())
+		{
+			uint32 RendererMaxNumComponents = RendererProperty->ComputeMaxUsedComponents(*ParticleDataSet);
+			MaxGPUBufferComponents = FMath::Max(MaxGPUBufferComponents, RendererMaxNumComponents);
+		}
+	}
+	else
+	{
+		// GPU emitters must store the entire particle payload on GPU buffers, so get the maximum component count from the dataset.
+		MaxGPUBufferComponents = FMath::Max(ParticleDataSet->GetNumFloatComponents(), ParticleDataSet->GetNumInt32Components());
+	}
+
+	// See how many particles we can fit in a GPU buffer. This number can be quite small on some platforms.
+	uint64 MaxBufferElements = (GDebugForcedMaxGPUBufferElements > 0) ? (uint64)GDebugForcedMaxGPUBufferElements : GetMaxBufferDimension();
+	// Don't just cast the result of the division to 32-bit, since that will produce garbage if MaxNumInstances is larger than UINT_MAX. Saturate instead.
+	MaxInstanceCount = (uint32)FMath::Min(MaxBufferElements / MaxGPUBufferComponents, (uint64)UINT_MAX);
+
+	if (CachedEmitter->SimTarget == ENiagaraSimTarget::GPUComputeSim)
+	{
+		// On GPU, the size of the allocated buffers must be a multiple of NIAGARA_COMPUTE_THREADGROUP_SIZE, so round down.
+		MaxInstanceCount = (MaxInstanceCount / NIAGARA_COMPUTE_THREADGROUP_SIZE) * NIAGARA_COMPUTE_THREADGROUP_SIZE;
+		// We will need an extra scratch instance, so the maximum number of usable instances is one less than the value we computed.
+		MaxInstanceCount -= 1;
+	}
+
+	ParticleDataSet->SetMaxInstanceCount(MaxInstanceCount);
 }
 
 void FNiagaraEmitterInstance::ResetSimulation(bool bKillExisting /*= true*/)
@@ -1189,6 +1231,8 @@ void FNiagaraEmitterInstance::Tick(float DeltaSeconds)
 	}
 
 	int32 AllocationSize = FMath::Max<int32>(AllocationEstimate, RequiredSize);
+	AllocationSize = (int32)FMath::Min((uint32)AllocationSize, MaxInstanceCount);
+
 	if (AllocationSize > MaxAllocationCount)
 	{
 		ReallocationCount++;
@@ -1257,9 +1301,13 @@ void FNiagaraEmitterInstance::Tick(float DeltaSeconds)
 			{
 				for (int32 SpawnInfoIdx = 0; SpawnInfoIdx < SpawnInfos.Num(); SpawnInfoIdx++)
 				{
-					FNiagaraSpawnInfo& Info = SpawnInfos[SpawnInfoIdx];
+					const FNiagaraSpawnInfo& Info = SpawnInfos[SpawnInfoIdx];
 					if (Info.Count > 0 && (NumSpawnInfos < NIAGARA_MAX_GPU_SPAWN_INFOS))
 					{
+						// Ideally, we should clamp the spawn count here, to make sure that we don't exceed the maximum number of particles. However, the value returned by
+						// GetNumParticles() can lag behind the real number, so we can't actually determine on the game thread how many particles we're still allowed to
+						// spawn. Therefore, we'll send the spawn requests to the render thread as if there was no limit, and we'll clamp the values there, when we prepare
+						// the destination dataset for simulation.
 						NumSpawnedOnGPUThisFrame += Info.Count;
 						if (NumSpawnedOnGPUThisFrame > GMaxNiagaraGPUParticlesSpawnPerFrame)
 						{
@@ -1283,7 +1331,7 @@ void FNiagaraEmitterInstance::Tick(float DeltaSeconds)
 						break;
 					}
 
-					// NOTE(mv): Separate particle count path for GPU emitters, as they early out..
+					// Warning: this will be be inaccurate if the render thread clamps the spawn count to keep the total particle count below the limit.
 					TotalSpawnedParticles += Info.Count;
 				}
 			}
@@ -1455,6 +1503,8 @@ void FNiagaraEmitterInstance::Tick(float DeltaSeconds)
 
 	Data.GetSpawnedIDsTable().SetNum(0, false);
 
+	int32 SpawnCountRemaining = AllocationSize - OrigNumParticles;
+
 	//Init new particles with the spawn script.
 	if (SpawnTotal + EventSpawnTotal > 0)
 	{
@@ -1520,16 +1570,18 @@ void FNiagaraEmitterInstance::Tick(float DeltaSeconds)
 		};
 
 		//Perform all our regular spawning that's driven by our emitter script.
-		for (FNiagaraSpawnInfo& Info : SpawnInfos)
+		for (const FNiagaraSpawnInfo& Info : SpawnInfos)
 		{
-			if ( Info.Count > 0 )
+			const int32 AdjustedSpawnCount = FMath::Min(Info.Count, SpawnCountRemaining);
+			if ( AdjustedSpawnCount > 0 )
 			{
 				auto& EmitterParameters = ParentSystemInstance->EditEmitterParameters(EmitterIdx);
 				SpawnIntervalBinding.SetValue(Info.IntervalDt);
 				InterpSpawnStartBinding.SetValue(Info.InterpStartDt);
 				SpawnGroupBinding.SetValue(Info.SpawnGroup);
-				SpawnParticles(Info.Count, TEXT("Regular Spawn"));
+				SpawnParticles(AdjustedSpawnCount, TEXT("Regular Spawn"));
 			}
+			SpawnCountRemaining -= AdjustedSpawnCount;
 		}
 
 		EventSpawnStart = Data.GetDestinationDataChecked().GetNumInstances();
@@ -1542,7 +1594,7 @@ void FNiagaraEmitterInstance::Tick(float DeltaSeconds)
 				//Spawn particles coming from events.
 				for (int32 i = 0; i < Info.SpawnCounts.Num(); i++)
 				{
-					const int32 EventNumToSpawn = Info.SpawnCounts[i];
+					const int32 EventNumToSpawn = FMath::Min(Info.SpawnCounts[i], SpawnCountRemaining);
 					if (EventNumToSpawn > 0)
 					{
 						const int32 CurrNumParticles = Data.GetDestinationDataChecked().GetNumInstances();
@@ -1555,9 +1607,10 @@ void FNiagaraEmitterInstance::Tick(float DeltaSeconds)
 						SpawnParticles(EventNumToSpawn, TEXT("Event Spawn"));
 
 						//Update EventSpawnCounts to the number actually spawned.
-						int32 NumActuallySpawned = Data.GetDestinationDataChecked().GetNumInstances() - CurrNumParticles;
+						const int32 NumActuallySpawned = Data.GetDestinationDataChecked().GetNumInstances() - CurrNumParticles;
 						TotalActualEventSpawns += NumActuallySpawned;
 						Info.SpawnCounts[i] = NumActuallySpawned;
+						SpawnCountRemaining -= NumActuallySpawned;
 					}
 				}
 			}
diff --git a/Engine/Plugins/FX/Niagara/Source/Niagara/Private/NiagaraEmitterInstanceBatcher.cpp b/Engine/Plugins/FX/Niagara/Source/Niagara/Private/NiagaraEmitterInstanceBatcher.cpp
--- a/Engine/Plugins/FX/Niagara/Source/Niagara/Private/NiagaraEmitterInstanceBatcher.cpp
+++ b/Engine/Plugins/FX/Niagara/Source/Niagara/Private/NiagaraEmitterInstanceBatcher.cpp
@@ -427,22 +427,33 @@ void NiagaraEmitterInstanceBatcher::ResizeBuffersAndGatherResources(FOverlappabl
 			FNiagaraDataBuffer* CurrentData = Instance.SimStageData[0].Source;
 			FNiagaraDataBuffer* DestinationData = Instance.SimStageData[0].Destination;
 
+			const uint32 MaxInstanceCount = Context->MainDataSet->GetMaxInstanceCount();
 			const uint32 PrevNumInstances = bNeedsReset ? 0 : CurrentData->GetNumInstances();
-			const uint32 NewNumInstances = Instance.SpawnInfo.SpawnRateInstances + Instance.SpawnInfo.EventSpawnTotal + PrevNumInstances;
+			check(PrevNumInstances <= MaxInstanceCount && Context->ScratchMaxInstances <= MaxInstanceCount);
 
-			//We must assume all particles survive when allocating here. 
-			//If this is not true, the read back in ResolveDatasetWrites will shrink the buffers.
+			// At this point we can finally clamp the number of instances we're trying to allocate to the max instance count. We don't need to adjust the
+			// information inside Instance.SpawnInfo, because that just stores the start indices for each spawner; they're still correct if we spawn
+			// fewer particles than expected, it just means that some spawners will not be used.
+			const uint32 NewNumInstances = FMath::Min(Instance.SpawnInfo.SpawnRateInstances + Instance.SpawnInfo.EventSpawnTotal + PrevNumInstances, MaxInstanceCount);
+			const uint32 AdjustedSpawnCount = (uint32)FMath::Max<int32>((int32)NewNumInstances - PrevNumInstances, 0);
+
+			// We must assume all particles survive when allocating here. If this is not true, the read back in ResolveDatasetWrites will shrink the buffers.
 			const uint32 RequiredInstances = FMath::Max(PrevNumInstances, NewNumInstances);
-			const uint32 AllocatedInstances = FMath::Max(RequiredInstances, Context->ScratchMaxInstances);
+
+			// We need an extra scratch instance. The code which computes the instance count cap knows about this, so it ensures that we stay within
+			// the buffer limit including this instance.
+			const uint32 AllocatedInstances = FMath::Max(RequiredInstances, Context->ScratchMaxInstances) + 1;
 
 			if (bRequiresPersistentIDs)
 			{
-				Context->MainDataSet->AllocateGPUFreeIDs(AllocatedInstances + 1, RHICmdList, FeatureLevel, Context->GetDebugSimName());
+				Context->MainDataSet->AllocateGPUFreeIDs(AllocatedInstances, RHICmdList, FeatureLevel, Context->GetDebugSimName());
 				InstancesWithPersistentIDs.Add(&Instance);
 			}
 
-			DestinationData->AllocateGPU(AllocatedInstances + 1, GPUInstanceCounterManager, RHICmdList, FeatureLevel, Context->GetDebugSimName());
+			DestinationData->AllocateGPU(AllocatedInstances, GPUInstanceCounterManager, RHICmdList, FeatureLevel, Context->GetDebugSimName());
 			DestinationData->SetNumInstances(RequiredInstances);
+			DestinationData->SetNumSpawnedInstances(AdjustedSpawnCount);
+
 			Instance.SimStageData[0].SourceCountOffset = Instance.SimStageData[0].Source->GetGPUInstanceCountBufferOffset();
 			if (Instance.SimStageData[0].SourceCountOffset == INDEX_NONE) // It is possible that this has been queued for readback, taking ownership of the data. Use that instead.
 			{
@@ -509,7 +520,7 @@ void NiagaraEmitterInstanceBatcher::ResizeBuffersAndGatherResources(FOverlappabl
 							Instance.SimStageData[SimulationStageIndex].Source = Context->MainDataSet->GetCurrentData();
 							DestinationData = &Context->MainDataSet->BeginSimulate(false);
 							Instance.SimStageData[SimulationStageIndex].Destination = DestinationData;
-							DestinationData->AllocateGPU(AllocatedInstances + 1, GPUInstanceCounterManager, RHICmdList, FeatureLevel, Context->GetDebugSimName());
+							DestinationData->AllocateGPU(AllocatedInstances, GPUInstanceCounterManager, RHICmdList, FeatureLevel, Context->GetDebugSimName());
 							DestinationData->SetNumInstances(RequiredInstances);
 							Instance.SimStageData[SimulationStageIndex].SourceCountOffset = Instance.SimStageData[SimulationStageIndex].Source->GetGPUInstanceCountBufferOffset();
 							Instance.SimStageData[SimulationStageIndex].DestinationCountOffset = Instance.SimStageData[SimulationStageIndex].Destination->GetGPUInstanceCountBufferOffset();
@@ -787,6 +798,10 @@ void NiagaraEmitterInstanceBatcher::ExecuteAll(FRHICommandList& RHICmdList, FRHI
 					PrevNumInstances = Tick.bNeedsReset ? 0 : ExecContext->ScratchNumInstances;
 				}
 				ExecContext->ScratchNumInstances = InstanceData.SpawnInfo.SpawnRateInstances + InstanceData.SpawnInfo.EventSpawnTotal + PrevNumInstances;
+
+				const uint32 MaxInstanceCount = ExecContext->MainDataSet->GetMaxInstanceCount();
+				ExecContext->ScratchNumInstances = FMath::Min(ExecContext->ScratchNumInstances, MaxInstanceCount);
+
 				ExecContext->ScratchMaxInstances = FMath::Max(ExecContext->ScratchMaxInstances, ExecContext->ScratchNumInstances);
 			}
 		}
@@ -1313,7 +1328,8 @@ void NiagaraEmitterInstanceBatcher::Run(const FNiagaraGPUSystemTick& Tick, const
 
 	//UE_LOG(LogScript, Warning, TEXT("Run [%d] TotalInstances %d  src:%p dest:%p"), SimulationStageIndex, TotalNumInstances, Instance->SimStageData[SimulationStageIndex].Source, Instance->SimStageData[SimulationStageIndex].Destination);
 
-	int32 InstancesToSpawnThisFrame = Instance->SpawnInfo.SpawnRateInstances + Instance->SpawnInfo.EventSpawnTotal;
+	int32 InstancesToSpawnThisFrame = DestinationData.GetNumSpawnedInstances();
+	DestinationData.SetIDAcquireTag(FNiagaraComputeExecutionContext::TickCounter);
 
 	// Only spawn particles on the first stage
 	if (HasRunParticleStage)
@@ -1322,9 +1338,6 @@ void NiagaraEmitterInstanceBatcher::Run(const FNiagaraGPUSystemTick& Tick, const
 	}
 
 	FRHIComputeShader* ComputeShader = Shader.GetComputeShader();
-	DestinationData.SetNumSpawnedInstances(InstancesToSpawnThisFrame);
-	DestinationData.SetIDAcquireTag(FNiagaraComputeExecutionContext::TickCounter);
-
 	RHICmdList.SetComputeShader(ComputeShader);
 
 	// #todo(dmp): clean up this logic for shader stages on first frame
diff --git a/Engine/Plugins/FX/Niagara/Source/NiagaraCore/Private/NiagaraCustomVersion.cpp b/Engine/Plugins/FX/Niagara/Source/NiagaraCore/Private/NiagaraCustomVersion.cpp
--- a/Engine/Plugins/FX/Niagara/Source/NiagaraCore/Private/NiagaraCustomVersion.cpp
+++ b/Engine/Plugins/FX/Niagara/Source/NiagaraCore/Private/NiagaraCustomVersion.cpp
@@ -8,4 +8,4 @@ const FGuid FNiagaraCustomVersion::GUID(0xFCF57AFA, 0x50764283, 0xB9A9E658, 0xFF
 // Register the custom version with core
 FCustomVersionRegistration GRegisterNiagaraCustomVersion(FNiagaraCustomVersion::GUID, FNiagaraCustomVersion::LatestVersion, TEXT("NiagaraVer"));
 
-const FGuid FNiagaraCustomVersion::LatestScriptCompileVersion(0x5B0333B2, 0x52513F42, 0x99533E4B, 0x2401DBEE);
\ No newline at end of file
+const FGuid FNiagaraCustomVersion::LatestScriptCompileVersion(0xF301F0D6, 0xA2B64D87, 0xAFFB0B25, 0xD1AAA66E);
