Take into account float16 components when computing max particle counts.

#rnx

diff --git a/Engine/Plugins/FX/Niagara/Source/Niagara/Private/NiagaraDataSet.cpp b/Engine/Plugins/FX/Niagara/Source/Niagara/Private/NiagaraDataSet.cpp
--- a/Engine/Plugins/FX/Niagara/Source/Niagara/Private/NiagaraDataSet.cpp
+++ b/Engine/Plugins/FX/Niagara/Source/Niagara/Private/NiagaraDataSet.cpp
@@ -781,7 +781,7 @@ void FNiagaraDataBuffer::AllocateGPU(uint32 InNumInstances, FNiagaraGPUInstanceC
 				{
 					GPUBufferHalf.Release();
 				}
-				GPUBufferHalf.Initialize(sizeof(FFloat16), NumElementsToAlloc * Owner->GetNumHalfComponents(), EPixelFormat::PF_R16F, BUF_Static);
+				GPUBufferHalf.Initialize(sizeof(FFloat16), NumInstancesAllocatedForGPU * Owner->GetNumHalfComponents(), EPixelFormat::PF_R16F, BUF_Static);
 			}
 		}
 
diff --git a/Engine/Plugins/FX/Niagara/Source/Niagara/Private/NiagaraEmitterInstance.cpp b/Engine/Plugins/FX/Niagara/Source/Niagara/Private/NiagaraEmitterInstance.cpp
--- a/Engine/Plugins/FX/Niagara/Source/Niagara/Private/NiagaraEmitterInstance.cpp
+++ b/Engine/Plugins/FX/Niagara/Source/Niagara/Private/NiagaraEmitterInstance.cpp
@@ -503,7 +503,7 @@ void FNiagaraEmitterInstance::Init(int32 InEmitterIdx, FNiagaraSystemInstanceID
 	else
 	{
 		// GPU emitters must store the entire particle payload on GPU buffers, so get the maximum component count from the dataset.
-		MaxGPUBufferComponents = FMath::Max(ParticleDataSet->GetNumFloatComponents(), ParticleDataSet->GetNumInt32Components());
+		MaxGPUBufferComponents = FMath::Max3(ParticleDataSet->GetNumFloatComponents(), ParticleDataSet->GetNumInt32Components(), ParticleDataSet->GetNumHalfComponents());
 	}
 
 	// See how many particles we can fit in a GPU buffer. This number can be quite small on some platforms.
