- Add UI clarity around SpawnOnly simulation stage, modify display name, add a tool tip and only allow when the iteration source is a data interface

diff --git a/Engine/Plugins/FX/Niagara/Source/Niagara/Public/NiagaraSimulationStageBase.h b/Engine/Plugins/FX/Niagara/Source/Niagara/Public/NiagaraSimulationStageBase.h
--- a/Engine/Plugins/FX/Niagara/Source/Niagara/Public/NiagaraSimulationStageBase.h
+++ b/Engine/Plugins/FX/Niagara/Source/Niagara/Public/NiagaraSimulationStageBase.h
@@ -46,7 +46,7 @@ public:
 	UPROPERTY(EditAnywhere, Category = "Simulation Stage")
 	int32 Iterations;
 
-	UPROPERTY(EditAnywhere, Category = "Simulation Stage")
+	UPROPERTY(EditAnywhere, Category = "Simulation Stage", meta = (DisplayName = "Emitter Reset Only", Tooltip = "When enabled the stage will only run on the first tick after the emitter is reset, only valid for data interface iteration stages", EditCondition = "IterationSource == ENiagaraIterationSource::DataInterface"))
 	uint32 bSpawnOnly : 1;
 
 	UPROPERTY(EditAnywhere, Category = "Simulation Stage", meta = (editcondition = "IterationSource == ENiagaraIterationSource::DataInterface"))
