Make it so attempting to sample a volume texture in a Niagara CPU script results in a compiler error.

diff --git a/Engine/Plugins/FX/Niagara/Source/Niagara/Private/NiagaraDataInterfaceVolumeTexture.cpp b/Engine/Plugins/FX/Niagara/Source/Niagara/Private/NiagaraDataInterfaceVolumeTexture.cpp
--- a/Engine/Plugins/FX/Niagara/Source/Niagara/Private/NiagaraDataInterfaceVolumeTexture.cpp
+++ b/Engine/Plugins/FX/Niagara/Source/Niagara/Private/NiagaraDataInterfaceVolumeTexture.cpp
@@ -84,6 +84,8 @@ void UNiagaraDataInterfaceVolumeTexture::GetFunctions(TArray<FNiagaraFunctionSig
 		Sig.Name = SampleVolumeTextureName;
 		Sig.bMemberFunction = true;
 		Sig.bRequiresContext = false;
+		Sig.bSupportsCPU = false;
+		Sig.bSupportsGPU = true;
 		Sig.Inputs.Add(FNiagaraVariable(FNiagaraTypeDefinition(GetClass()), TEXT("Texture")));
 		Sig.Inputs.Add(FNiagaraVariable(FNiagaraTypeDefinition::GetVec3Def(), TEXT("UVW")));
 		Sig.Inputs.Add(FNiagaraVariable(FNiagaraTypeDefinition::GetFloatDef(), TEXT("MipLevel")));
