SNiagaraAddParameterMenu2 now properly sets parameter handles to allow building SNiagaraParameterName widgets for entries.

diff --git a/Engine/Plugins/FX/Niagara/Source/NiagaraEditor/Private/Widgets/SNiagaraParameterPanel.cpp b/Engine/Plugins/FX/Niagara/Source/NiagaraEditor/Private/Widgets/SNiagaraParameterPanel.cpp
--- a/Engine/Plugins/FX/Niagara/Source/NiagaraEditor/Private/Widgets/SNiagaraParameterPanel.cpp
+++ b/Engine/Plugins/FX/Niagara/Source/NiagaraEditor/Private/Widgets/SNiagaraParameterPanel.cpp
@@ -527,6 +527,8 @@ void SNiagaraAddParameterMenu2::OnActionSelected(const TArray<TSharedPtr<FEdGrap
 
 void SNiagaraAddParameterMenu2::CollectAllActions(FGraphActionListBuilderBase& OutAllActions)
 {
+	TArray<TSharedPtr<FNiagaraMenuAction>> TempOutAllActions;
+
 	if (OnCollectCustomActions.IsBound())
 	{
 		bool bCreateRemainingActions = true;
@@ -604,11 +606,6 @@ void SNiagaraAddParameterMenu2::CollectAllActions(FGraphActionListBuilderBase& O
 			}
 		}
 
-		// Sort new action infos to add to the GraphActionMenu in alphabetical order of their display name.
-		NewParameterMenuActionInfos.Sort([](const FNiagaraMenuActionInfo& A, const FNiagaraMenuActionInfo& B)
-			{ return A.DisplayNameText.CompareToCaseIgnored(B.DisplayNameText) < 0; }
-		);
-
 		// Create the actual actions and mark them experimental if so.
 		for (const FNiagaraMenuActionInfo& ActionInfo : NewParameterMenuActionInfos)
 		{
@@ -616,8 +613,9 @@ void SNiagaraAddParameterMenu2::CollectAllActions(FGraphActionListBuilderBase& O
 				FNiagaraMenuAction::FOnExecuteStackAction::CreateSP(this, &SNiagaraAddParameterMenu2::AddParameterSelected, ActionInfo.NewVariable)));
 
 			Action->IsExperimental = ActionInfo.bIsExperimental;
+			Action->SetParamterHandle(FNiagaraParameterHandle(ActionInfo.NewVariable.GetName()));
 
-			OutAllActions.AddAction(Action);
+			TempOutAllActions.Add(Action);
 		}
 	}
 
@@ -640,8 +638,9 @@ void SNiagaraAddParameterMenu2::CollectAllActions(FGraphActionListBuilderBase& O
 				TSharedPtr<FNiagaraMenuAction> Action(new FNiagaraMenuAction(
 					Category, DisplayName, Tooltip, 0, FText::GetEmpty(),
 					FNiagaraMenuAction::FOnExecuteStackAction::CreateSP(this, &SNiagaraAddParameterMenu2::AddParameterSelected, Variable)));
+				Action->SetParamterHandle(FNiagaraParameterHandle(Variable.GetName()));
 
-				OutAllActions.AddAction(Action);
+				TempOutAllActions.Add(Action);
 			}
 		}
 	}
@@ -685,8 +684,9 @@ void SNiagaraAddParameterMenu2::CollectAllActions(FGraphActionListBuilderBase& O
 					TSharedPtr<FNiagaraMenuAction> Action(new FNiagaraMenuAction(
 						CategoryText, DisplayName, Tooltip, 0, FText::GetEmpty(),
 						FNiagaraMenuAction::FOnExecuteStackAction::CreateSP(this, &SNiagaraAddParameterMenu2::AddParameterSelected, Var)));
+					Action->SetParamterHandle(FNiagaraParameterHandle(Var.GetName()));
 
-					OutAllActions.AddAction(Action);
+					TempOutAllActions.Add(Action);
 				}
 			}
 
@@ -705,14 +705,31 @@ void SNiagaraAddParameterMenu2::CollectAllActions(FGraphActionListBuilderBase& O
 					TSharedPtr<FNiagaraMenuAction> Action(new FNiagaraMenuAction(
 						CategoryText, DisplayName, Tooltip, 0, FText::GetEmpty(),
 						FNiagaraMenuAction::FOnExecuteStackAction::CreateSP(this, &SNiagaraAddParameterMenu2::AddParameterSelected, Var)));
+					Action->SetParamterHandle(FNiagaraParameterHandle(Var.GetName()));
 
-					OutAllActions.AddAction(Action);
+					TempOutAllActions.Add(Action);
 
 				}
 			}
 
 		}
 	}
+	auto SortOutAllActionsPred = [](const TSharedPtr<FNiagaraMenuAction>& A, const TSharedPtr<FNiagaraMenuAction>& B)->bool {
+		if (A->GetParameterHandle().IsSet() == false)
+		{
+			return false;
+		}
+		else if (B->GetParameterHandle().IsSet() == false)
+		{
+			return true;
+		}
+		return FNiagaraEditorUtilities::GetVariableSortPriority(A->GetParameterHandle()->GetParameterHandleString(), B->GetParameterHandle()->GetParameterHandleString());
+	};
+	TempOutAllActions.Sort(SortOutAllActionsPred);
+	for (const TSharedPtr<FNiagaraMenuAction>& Action : TempOutAllActions)
+	{
+		OutAllActions.AddAction(Action);
+	}
 }
 
 void SNiagaraAddParameterMenu2::AddParameterSelected(FNiagaraVariable NewVariable)
