PR #6831: Fix Niagara Scirpt (sic) Typos (Contributed by [REDACTED_PERSON])

diff --git a/Engine/Plugins/FX/Niagara/Source/NiagaraEditor/Private/ViewModels/NiagaraSystemViewModel.cpp b/Engine/Plugins/FX/Niagara/Source/NiagaraEditor/Private/ViewModels/NiagaraSystemViewModel.cpp
--- a/Engine/Plugins/FX/Niagara/Source/NiagaraEditor/Private/ViewModels/NiagaraSystemViewModel.cpp
+++ b/Engine/Plugins/FX/Niagara/Source/NiagaraEditor/Private/ViewModels/NiagaraSystemViewModel.cpp
@@ -816,10 +816,10 @@ void FNiagaraSystemViewModel::TickCompileStatus()
 				}
 				else
 				{
-					ENiagaraScriptCompileStatus LastScirptCompileStatus = ScriptToCheck->GetLastCompileStatus();
-					if (LastScirptCompileStatus == ENiagaraScriptCompileStatus::NCS_Unknown ||
-						LastScirptCompileStatus == ENiagaraScriptCompileStatus::NCS_Dirty ||
-						LastScirptCompileStatus == ENiagaraScriptCompileStatus::NCS_BeingCreated)
+					ENiagaraScriptCompileStatus LastScriptCompileStatus = ScriptToCheck->GetLastCompileStatus();
+					if (LastScriptCompileStatus == ENiagaraScriptCompileStatus::NCS_Unknown ||
+						LastScriptCompileStatus == ENiagaraScriptCompileStatus::NCS_Dirty ||
+						LastScriptCompileStatus == ENiagaraScriptCompileStatus::NCS_BeingCreated)
 					{
 						// If the script doesn't have a vaild last compile status assume that it's dirty and by extension
 						// so is the system and we can ignore the other statuses.
@@ -830,7 +830,7 @@ void FNiagaraSystemViewModel::TickCompileStatus()
 					else
 					{
 						// Otherwise save it for processing once all scripts are done.
-						ScriptCompileStatuses.Add(LastScirptCompileStatus);
+						ScriptCompileStatuses.Add(LastScriptCompileStatus);
 					}
 				}
 			}
diff --git a/Engine/Plugins/FX/Niagara/Source/NiagaraEditorWidgets/Private/Stack/SNiagaraStackFunctionInputValue.cpp b/Engine/Plugins/FX/Niagara/Source/NiagaraEditorWidgets/Private/Stack/SNiagaraStackFunctionInputValue.cpp
--- a/Engine/Plugins/FX/Niagara/Source/NiagaraEditorWidgets/Private/Stack/SNiagaraStackFunctionInputValue.cpp
+++ b/Engine/Plugins/FX/Niagara/Source/NiagaraEditorWidgets/Private/Stack/SNiagaraStackFunctionInputValue.cpp
@@ -241,7 +241,7 @@ TSharedRef<SWidget> SNiagaraStackFunctionInputValue::ConstructValueWidgets()
 	{
 		return SNew(STextBlock)
 			.TextStyle(FNiagaraEditorStyle::Get(), "NiagaraEditor.ParameterText")
-			.Text(LOCTEXT("InvalidOverrideText", "Invalid Scirpt Value"));
+			.Text(LOCTEXT("InvalidOverrideText", "Invalid Script Value"));
 	}
 	case UNiagaraStackFunctionInput::EValueMode::UnsupportedDefault:
 	{
