Niagara - Fix issues with script search.
+ Style pass on UI.
+ When navigating to results use the graph traversal to search in orderish.
+ Fix an issue where it would navigate twice with each button press.
+ Update the string matching so that matches what's displayed in the pin widgets with the new parameter UX.
+ Only show one result per node since we don't actually show which pin is selected for pin matches.

diff --git a/Engine/Plugins/FX/Niagara/Source/NiagaraEditor/Private/NiagaraEditorStyle.cpp b/Engine/Plugins/FX/Niagara/Source/NiagaraEditor/Private/NiagaraEditorStyle.cpp
--- a/Engine/Plugins/FX/Niagara/Source/NiagaraEditor/Private/NiagaraEditorStyle.cpp
+++ b/Engine/Plugins/FX/Niagara/Source/NiagaraEditor/Private/NiagaraEditorStyle.cpp
@@ -364,6 +364,8 @@ TSharedRef< FSlateStyleSet > FNiagaraEditorStyle::Create()
 	Style->Set("NiagaraEditor.DropTarget.BorderVertical", new IMAGE_PLUGIN_BRUSH("Icons/StackDropTargetBorder_Vertical", FVector2D(2, 8), FLinearColor::White, ESlateBrushTileType::Vertical));
 	Style->Set("NiagaraEditor.DropTarget.BorderHorizontal", new IMAGE_PLUGIN_BRUSH("Icons/StackDropTargetBorder_Horizontal", FVector2D(8, 2), FLinearColor::White, ESlateBrushTileType::Horizontal));
 
+	Style->Set("NiagaraEditor.ScriptGraph.SearchBorderColor", FLinearColor(.1f, .1f, .1f, 1.f));
+
 	return Style;
 }
 
diff --git a/Engine/Plugins/FX/Niagara/Source/NiagaraEditor/Private/NiagaraEditorUtilities.cpp b/Engine/Plugins/FX/Niagara/Source/NiagaraEditor/Private/NiagaraEditorUtilities.cpp
--- a/Engine/Plugins/FX/Niagara/Source/NiagaraEditor/Private/NiagaraEditorUtilities.cpp
+++ b/Engine/Plugins/FX/Niagara/Source/NiagaraEditor/Private/NiagaraEditorUtilities.cpp
@@ -2524,4 +2524,53 @@ const FNiagaraNamespaceMetadata FNiagaraEditorUtilities::GetNamespaceMetaDataFor
 	return GetDefault<UNiagaraEditorSettings>()->GetMetaDataForNamespaces(VarHandleNameParts);
 }
 
+bool FNiagaraParameterUtilities::DoesParameterNameMatchSearchText(FName ParameterName, const FString& SearchTextString)
+{
+	FNiagaraParameterHandle ParameterHandle(ParameterName);
+	TArray<FName> HandleParts = ParameterHandle.GetHandleParts();
+	FNiagaraNamespaceMetadata NamespaceMetadata = GetDefault<UNiagaraEditorSettings>()->GetMetaDataForNamespaces(HandleParts);
+	if (NamespaceMetadata.IsValid())
+	{
+		// If it's a registered namespace, check the display name of the namespace.
+		if (NamespaceMetadata.DisplayName.ToString().Contains(SearchTextString))
+		{
+			return true;
+		}
+
+		// Check the namespace modifier if it has one
+		if (HandleParts.Num() - NamespaceMetadata.Namespaces.Num() > 1)
+		{
+			FNiagaraNamespaceMetadata NamespaceModifierMetadata = GetDefault<UNiagaraEditorSettings>()->GetMetaDataForNamespaceModifier(HandleParts[NamespaceMetadata.Namespaces.Num()]);
+			if (NamespaceModifierMetadata.IsValid())
+			{
+				// Check first by modifier metadata display name.
+				if (NamespaceModifierMetadata.DisplayName.ToString().Contains(SearchTextString))
+				{
+					return true;
+				}
+			}
+			else
+			{
+				// Otherwise just check the string.
+				if (HandleParts[NamespaceMetadata.Namespaces.Num()].ToString().Contains(SearchTextString))
+				{
+					return true;
+				}
+			}
+		}
+
+		// Last check the variable name.
+		if (HandleParts.Last().ToString().Contains(SearchTextString))
+		{
+			return true;
+		}
+	}
+	else if (HandleParts.ContainsByPredicate([&SearchTextString](FName NamePart) { return NamePart.ToString().Contains(SearchTextString); }))
+	{
+		// Otherwise if it's not in a valid namespace, just check all name parts.
+		return true;
+	}
+	return false;
+}
+
 #undef LOCTEXT_NAMESPACE
diff --git a/Engine/Plugins/FX/Niagara/Source/NiagaraEditor/Private/Widgets/SNiagaraScriptGraph.cpp b/Engine/Plugins/FX/Niagara/Source/NiagaraEditor/Private/Widgets/SNiagaraScriptGraph.cpp
--- a/Engine/Plugins/FX/Niagara/Source/NiagaraEditor/Private/Widgets/SNiagaraScriptGraph.cpp
+++ b/Engine/Plugins/FX/Niagara/Source/NiagaraEditor/Private/Widgets/SNiagaraScriptGraph.cpp
@@ -6,6 +6,9 @@
 #include "NiagaraGraph.h"
 #include "NiagaraNode.h"
 #include "NiagaraNodeInput.h"
+#include "NiagaraEditorStyle.h"
+#include "NiagaraNodeParameterMapBase.h"
+#include "NiagaraNodeReroute.h"
 
 #include "GraphEditor.h"
 #include "EditorStyleSet.h"
@@ -28,6 +31,7 @@
 #include "GraphEditorActions.h"
 #include "Subsystems/AssetEditorSubsystem.h"
 #include "Editor.h"
+#include "EditorFontGlyphs.h"
 
 #define LOCTEXT_NAMESPACE "NiagaraScriptGraph"
 
@@ -68,6 +72,7 @@ TSharedRef<SGraphEditor> SNiagaraScriptGraph::ConstructGraphEditor()
 		.BorderImage(FEditorStyle::GetBrush(TEXT("Graph.TitleBackground")))
 		.HAlign(HAlign_Fill)
 		[
+			// Error Indicator and Title
 			SNew(SOverlay)
 			+SOverlay::Slot()
 			[
@@ -91,54 +96,54 @@ TSharedRef<SGraphEditor> SNiagaraScriptGraph::ConstructGraphEditor()
 					.Justification(ETextJustify::Center)
 				]
 			]
+			// Search Box
 			+SOverlay::Slot()
 			.HAlign(HAlign_Right)
 			.VAlign(VAlign_Fill)
 			[
-				SNew(SHorizontalBox)
-				+ SHorizontalBox::Slot()
-				.HAlign(HAlign_Fill)
-				.VAlign(VAlign_Fill)
-				.AutoWidth()
-				.MaxWidth(400.0f) // Limit max search box width to avoid extending over titlebar
-				[
-					SAssignNew(SearchBox, SSearchBox)
-					.HintText(LOCTEXT("GraphSearchBoxHint", "Search Nodes and Pins in Graph"))
-					.SearchResultData(this, &SNiagaraScriptGraph::GetSearchResultData)
-					.OnTextChanged(this, &SNiagaraScriptGraph::OnSearchTextChanged)
-					.OnTextCommitted(this, &SNiagaraScriptGraph::OnSearchBoxTextCommitted)
-					.DelayChangeNotificationsWhileTyping(true)
-					.OnSearch(this, &SNiagaraScriptGraph::OnSearchBoxSearch)
-					.Visibility(this, &SNiagaraScriptGraph::GetGraphSearchBoxVisibility)
-					.OnKeyDownHandler(this, &SNiagaraScriptGraph::HandleGraphSearchBoxKeyDown)
-				]
-				+SHorizontalBox::Slot()
-				.HAlign(HAlign_Fill)
-				.VAlign(VAlign_Fill)
-				.AutoWidth()
+				SNew(SBorder)
+				.BorderImage(FEditorStyle::GetBrush("WhiteBrush"))
+				.BorderBackgroundColor(FNiagaraEditorStyle::Get().GetColor("NiagaraEditor.ScriptGraph.SearchBorderColor"))
+				.Visibility(this, &SNiagaraScriptGraph::GetGraphSearchBoxVisibility)
+				.Padding(5)
 				[
-					SNew(SBorder)
-					.HAlign(HAlign_Center)
+					SNew(SHorizontalBox)
+					+ SHorizontalBox::Slot()
+					.HAlign(HAlign_Fill)
 					.VAlign(VAlign_Fill)
-					.BorderImage(&Style.TextBoxStyle.BackgroundImageHovered)
-					.BorderBackgroundColor(Style.TextBoxStyle.BackgroundColor)
-					.ForegroundColor(Style.TextBoxStyle.ForegroundColor)
-					.Padding(0)
+					.AutoWidth()
+					[
+						SNew(SBox)
+						.WidthOverride(200)
+						[
+							SAssignNew(SearchBox, SSearchBox)
+							.HintText(LOCTEXT("GraphSearchBoxHint", "Search Nodes and Pins in Graph"))
+							.SearchResultData(this, &SNiagaraScriptGraph::GetSearchResultData)
+							.OnTextChanged(this, &SNiagaraScriptGraph::OnSearchTextChanged)
+							.OnTextCommitted(this, &SNiagaraScriptGraph::OnSearchBoxTextCommitted)
+							.DelayChangeNotificationsWhileTyping(true)
+							.OnSearch(this, &SNiagaraScriptGraph::OnSearchBoxSearch)
+							.OnKeyDownHandler(this, &SNiagaraScriptGraph::HandleGraphSearchBoxKeyDown)
+						]
+					]
+					+SHorizontalBox::Slot()
+					.HAlign(HAlign_Center)
+					.VAlign(VAlign_Center)
+					.AutoWidth()
+					.Padding(2.0f, 0.0f, 0.0f, 0.0f)
 					[
 						SNew(SButton)
-						.ButtonStyle(FCoreStyle::Get(), "NoBorder")
-						.ContentPadding(0)
-						.HAlign(HAlign_Center)
-						.VAlign(VAlign_Center)
+						.ButtonStyle(FEditorStyle::Get(), "HoverHintOnly")
 						.IsFocusable(false)
+						.ForegroundColor(FNiagaraEditorStyle::Get().GetColor("NiagaraEditor.Stack.FlatButtonColor"))
 						.ToolTipText(LOCTEXT("CloseGraphSearchBox", "Close Graph search box"))
-						.Visibility(this, &SNiagaraScriptGraph::GetGraphSearchBoxVisibility)
 						.OnClicked(this, &SNiagaraScriptGraph::CloseGraphSearchBoxPressed)
+						.ContentPadding(3)
 						.Content()
 						[
- 							SNew(SImage)
- 							.Image(FEditorStyle::GetBrush("Symbols.X"))
-							.ColorAndOpacity(FSlateColor::UseForeground())
+							SNew(STextBlock)
+							.Font(FEditorStyle::Get().GetFontStyle("FontAwesome.10"))
+							.Text(FEditorFontGlyphs::Times)
 						]
 					]
 				]
@@ -477,30 +482,83 @@ void SNiagaraScriptGraph::FocusGraphElement(const INiagaraScriptGraphFocusInfo*
 	checkf(false, TEXT("Requested focus for a graph element without specifying a Node or Pin to focus!"));
 }
 
+bool NodeMatchesSearch(UNiagaraNode* Node, const FString& SearchTextString)
+{
+	if (Node->IsA<UNiagaraNodeReroute>())
+	{
+		// Ignore reroute nodes.
+		return false;
+	}
+
+	if (Node->GetNodeTitle(ENodeTitleType::FullTitle).ToString().Contains(SearchTextString))
+	{
+		return true;
+	}
+	
+	if (Node->IsA<UNiagaraNodeParameterMapBase>())
+	{
+		UNiagaraNodeParameterMapBase* ParameterMapNode = CastChecked<UNiagaraNodeParameterMapBase>(Node);
+		for (UEdGraphPin* Pin : ParameterMapNode->Pins)
+		{
+			if (ParameterMapNode->IsAddPin(Pin))
+			{
+				if (Pin->GetDisplayName().ToString().Contains(SearchTextString))
+				{
+					return true;
+				}
+			}
+			else
+			{
+				if (FNiagaraParameterUtilities::DoesParameterNameMatchSearchText(Pin->PinName, SearchTextString))
+				{
+					return true;
+				}
+			}
+		}
+	}
+	else
+	{
+		if (Node->Pins.ContainsByPredicate([&SearchTextString](UEdGraphPin* Pin) { return Pin->GetDisplayName().ToString().Contains(SearchTextString); }))
+		{
+			return true;
+		}
+	}
+	return false;
+}
+
 void SNiagaraScriptGraph::OnSearchTextChanged(const FText& SearchText)
 {
 	if (!CurrentSearchText.EqualTo(SearchText))
 	{
 		CurrentSearchResults.Empty();
 		CurrentSearchText = SearchText;
-		TArray<UNiagaraNode*> Nodes;
-		ViewModel->GetGraph()->GetNodesOfClass<UNiagaraNode>(Nodes);
-		for (UNiagaraNode* Node : Nodes)
+
+		TArray<UNiagaraNode*> AllNodes;
+		ViewModel->GetGraph()->GetNodesOfClass<UNiagaraNode>(AllNodes);
+
+		TArray<UNiagaraNodeOutput*> OutputNodes;
+		ViewModel->GetGraph()->GetNodesOfClass<UNiagaraNodeOutput>(OutputNodes);
+
+		TArray<UNiagaraNode*> TraversedNodes;
+		for (UNiagaraNodeOutput* OutputNode : OutputNodes)
 		{
-			if (Node->GetNodeTitle(ENodeTitleType::FullTitle).ToString().Contains(SearchText.ToString()))
-			{	
-				CurrentSearchResults.Add(MakeShared<FNiagaraScriptGraphNodeToFocusInfo>(Node->NodeGuid));
-			}
-			
-			if (Node->IsA<UNiagaraNodeOutput>() == false) 
+			TArray<UNiagaraNode*> OutputNodeTraversal;
+			UNiagaraGraph::BuildTraversal(OutputNodeTraversal, OutputNode, false);
+			TraversedNodes.Append(OutputNodeTraversal);
+		}
+
+		AllNodes.RemoveAll([&TraversedNodes](UNiagaraNode* Node) { return TraversedNodes.Contains(Node); });
+
+		TArray<UNiagaraNode*> OrderedSearchNodes;
+		OrderedSearchNodes.Append(TraversedNodes);
+		OrderedSearchNodes.Append(AllNodes);
+
+		FString SearchTextString = SearchText.ToString();
+		for (UNiagaraNode* SearchNode : OrderedSearchNodes)
+		{
+			if(NodeMatchesSearch(SearchNode, SearchTextString))
 			{
-				for (UEdGraphPin* Pin : Node->GetAllPins())
-				{
-					if (Pin->GetDisplayName().ToString().Contains(SearchText.ToString()))
-					{
-						CurrentSearchResults.Add(MakeShared<FNiagaraScriptGraphPinToFocusInfo>(Pin->PersistentGuid));
-					}
-				}
+				CurrentSearchResults.Add(MakeShared<FNiagaraScriptGraphNodeToFocusInfo>(SearchNode->NodeGuid));
 			}
 		}
 
@@ -514,9 +572,16 @@ void SNiagaraScriptGraph::OnSearchTextChanged(const FText& SearchText)
 
 void SNiagaraScriptGraph::OnSearchBoxTextCommitted(const FText& NewText, ETextCommit::Type CommitInfo)
 {
-	if (SearchBox->HasKeyboardFocus())
+	if (CommitInfo == ETextCommit::OnEnter)
 	{
-		OnSearchBoxSearch(SSearchBox::Next);
+		if (CurrentSearchText.CompareTo(NewText) == 0)
+		{
+			OnSearchBoxSearch(SSearchBox::Next);
+		}
+		else
+		{
+			OnSearchTextChanged(NewText);
+		}
 	}
 }
 
diff --git a/Engine/Plugins/FX/Niagara/Source/NiagaraEditor/Public/NiagaraEditorUtilities.h b/Engine/Plugins/FX/Niagara/Source/NiagaraEditor/Public/NiagaraEditorUtilities.h
--- a/Engine/Plugins/FX/Niagara/Source/NiagaraEditor/Public/NiagaraEditorUtilities.h
+++ b/Engine/Plugins/FX/Niagara/Source/NiagaraEditor/Public/NiagaraEditorUtilities.h
@@ -306,3 +306,8 @@ namespace FNiagaraEditorUtilities
 
 	const FNiagaraNamespaceMetadata GetNamespaceMetaDataForVariableName(const FName& VarName);
 };
+
+namespace FNiagaraParameterUtilities
+{
+	bool DoesParameterNameMatchSearchText(FName ParameterName, const FString& SearchTextString);
+};
diff --git a/Engine/Source/Runtime/Slate/Private/Widgets/Input/SSearchBox.cpp b/Engine/Source/Runtime/Slate/Private/Widgets/Input/SSearchBox.cpp
--- a/Engine/Source/Runtime/Slate/Private/Widgets/Input/SSearchBox.cpp
+++ b/Engine/Source/Runtime/Slate/Private/Widgets/Input/SSearchBox.cpp
@@ -75,7 +75,7 @@ void SSearchBox::Construct( const FArguments& InArgs )
 		.AutoWidth()
 		.Padding(InArgs._Style->ImagePadding)
 		.HAlign(HAlign_Center)
-		.VAlign(VAlign_Center)
+		.VAlign(VAlign_Fill)
 		[
 			SNew(SButton)
 			.ButtonStyle( FCoreStyle::Get(), "NoBorder" )
@@ -97,7 +97,7 @@ void SSearchBox::Construct( const FArguments& InArgs )
 		.AutoWidth()
 		.Padding(InArgs._Style->ImagePadding)
 		.HAlign(HAlign_Center)
-		.VAlign(VAlign_Center)
+		.VAlign(VAlign_Fill)
 		[
 			SNew(SButton)
 			.ButtonStyle( FCoreStyle::Get(), "NoBorder" )
