- Fix crash in GetFilteredSocketTransform when you have no elements in the array
- Fix incorrect transform application (and reading from component) in GetFilteredSocketTransform
- Remove stub GPU implementation of GetFilteredSocketTransform so an stack error occurs, the new functions shoud be used instead
- Fix for not applying the rotation on the GPU for gathering bone transforms

diff --git a/Engine/Plugins/FX/Niagara/Shaders/Private/NiagaraDataInterfaceSkeletalMesh.ush b/Engine/Plugins/FX/Niagara/Shaders/Private/NiagaraDataInterfaceSkeletalMesh.ush
--- a/Engine/Plugins/FX/Niagara/Shaders/Private/NiagaraDataInterfaceSkeletalMesh.ush
+++ b/Engine/Plugins/FX/Niagara/Shaders/Private/NiagaraDataInterfaceSkeletalMesh.ush
@@ -48,6 +48,8 @@ struct FDISkelMeshContext
 	int				FilteredSocketBoneOffset;
 	float4x4		InstanceTransform;
 	float4x4		InstancePrevTransform;
+	float4			InstanceRotation;
+	float4			InstancePrevRotation;
 	float			InstanceInvDeltaTime;
 	bool			UniformTriangleSamplingEnable;
 	bool			UniformSamplingRegionEnabled;
@@ -407,7 +409,7 @@ void DISKelMesh_GetFilteredTriangleAt(in FDISkelMeshContext DIContext, int Filte
 ////////////////////////////////////////////////////////////////////////////////////////////////////////////////////////////////////////////////////////////////
 // Bone Sampling
 
-float4 QuatSlerp(float4 Quat1, float4 Quat2, float Slerp)
+float4 DISKelMesh_QuatSlerp(float4 Quat1, float4 Quat2, float Slerp)
 {
     const float RawCosom = dot(Quat1, Quat2);
     const float Cosom = abs(RawCosom);
@@ -432,6 +434,16 @@ float4 QuatSlerp(float4 Quat1, float4 Quat2, float Slerp)
     return (Scale0 * Quat1) + (Scale1 * Quat2);
 }
 
+float4 DISKelMesh_QuatMul(float4 Q1, float4 Q2)
+{
+	float4 QOut;
+	QOut.x = Q1.w*Q2.x + Q1.x*Q2.w + Q1.y*Q2.z - Q1.z*Q2.y;
+	QOut.y = Q1.w*Q2.y - Q1.x*Q2.z + Q1.y*Q2.w + Q1.z*Q2.x;
+	QOut.z = Q1.w*Q2.z + Q1.x*Q2.y - Q1.y*Q2.x + Q1.z*Q2.w;
+	QOut.w = Q1.w*Q2.w - Q1.x*Q2.x - Q1.y*Q2.y - Q1.z*Q2.z;
+	return QOut;
+}
+
 void DISkelMesh_GetSkinnedBoneCommon(in FDISkelMeshContext DIContext, in int Bone, in float Interp, in bool bWorldSpace, out float3 OutPosition, out float4 OutRotation, out float3 OutVelocity)
 {
     float3 PrevPosition = DIContext.MeshPrevSamplingBonesBuffer[Bone * 2].xyz;
@@ -443,11 +455,13 @@ void DISkelMesh_GetSkinnedBoneCommon(in FDISkelMeshContext DIContext, in int Bon
     {
         PrevPosition = mul(float4(PrevPosition, 1), DIContext.InstancePrevTransform).xyz;
         CurrPosition = mul(float4(CurrPosition, 1), DIContext.InstanceTransform).xyz;
+		PrevRotation = DISKelMesh_QuatMul(DIContext.InstancePrevRotation, PrevRotation);
+		CurrRotation = DISKelMesh_QuatMul(DIContext.InstanceRotation, CurrRotation);
     }
     CurrPosition = lerp(PrevPosition, CurrPosition, Interp);
 
     OutPosition = CurrPosition;
-    OutRotation = QuatSlerp(PrevRotation, CurrRotation, Interp);
+    OutRotation = DISKelMesh_QuatSlerp(PrevRotation, CurrRotation, Interp);
     OutVelocity = (CurrPosition - PrevPosition) * DIContext.InstanceInvDeltaTime;
 }
 
@@ -722,6 +736,8 @@ void DISkelMesh_GetFilteredSocketOrBoneAt(in FDISkelMeshContext DIContext, int F
 	int				FilteredSocketBoneOffset_ ## NAME; \
 	float4x4		InstanceTransform_ ## NAME; \
 	float4x4		InstancePrevTransform_ ## NAME; \
+	float4			InstanceRotation_ ## NAME; \
+	float4			InstancePrevRotation_ ## NAME; \
 	float			InstanceInvDeltaTime_ ## NAME; \
 	uint			EnabledFeatures_ ## NAME; \
 
@@ -763,6 +779,8 @@ void DISkelMesh_GetFilteredSocketOrBoneAt(in FDISkelMeshContext DIContext, int F
 	DIContext.MeshSkinWeightIndexSize = MeshSkinWeightIndexSize_ ## NAME; \
 	DIContext.InstanceTransform = InstanceTransform_ ## NAME; \
 	DIContext.InstancePrevTransform = InstancePrevTransform_ ## NAME; \
+	DIContext.InstanceRotation = InstanceRotation_ ## NAME; \
+	DIContext.InstancePrevRotation = InstancePrevRotation_ ## NAME; \
 	DIContext.InstanceInvDeltaTime = InstanceInvDeltaTime_ ## NAME; \
 	DIContext.UniformTriangleSamplingEnable = EnabledFeatures_ ## NAME & 0x1; \
 	DIContext.UniformSamplingRegionEnabled = EnabledFeatures_ ## NAME & 0x2; \
diff --git a/Engine/Plugins/FX/Niagara/Source/Niagara/Classes/NiagaraDataInterfaceSkeletalMesh.h b/Engine/Plugins/FX/Niagara/Source/Niagara/Classes/NiagaraDataInterfaceSkeletalMesh.h
--- a/Engine/Plugins/FX/Niagara/Source/Niagara/Classes/NiagaraDataInterfaceSkeletalMesh.h
+++ b/Engine/Plugins/FX/Niagara/Source/Niagara/Classes/NiagaraDataInterfaceSkeletalMesh.h
@@ -700,6 +700,8 @@ public:
 	static const FString FilteredSocketBoneOffsetName;
 	static const FString InstanceTransformName;
 	static const FString InstancePrevTransformName;
+	static const FString InstanceRotationName;
+	static const FString InstancePrevRotationName;
 	static const FString InstanceInvDeltaTimeName;
 	static const FString EnabledFeaturesName;
 
diff --git a/Engine/Plugins/FX/Niagara/Source/Niagara/Private/NDISkeletalMesh_BoneSampling.cpp b/Engine/Plugins/FX/Niagara/Source/Niagara/Private/NDISkeletalMesh_BoneSampling.cpp
--- a/Engine/Plugins/FX/Niagara/Source/Niagara/Private/NDISkeletalMesh_BoneSampling.cpp
+++ b/Engine/Plugins/FX/Niagara/Source/Niagara/Private/NDISkeletalMesh_BoneSampling.cpp
@@ -1008,7 +1008,7 @@ void UNiagaraDataInterfaceSkeletalMesh::GetFilteredSocketTransform(FVectorVMCont
 {
 	VectorVM::FUserPtrHandler<FNDISkeletalMesh_InstanceData> InstData(Context);
 	VectorVM::FExternalFuncInputHandler<int32> SocketParam(Context);
-	VectorVM::FExternalFuncInputHandler<int32> ApplyWorldTransform(Context);
+	VectorVM::FExternalFuncInputHandler<FNiagaraBool> ApplyWorldTransform(Context);
 
 	VectorVM::FExternalFuncRegisterHandler<float> OutSocketTranslateX(Context);
 	VectorVM::FExternalFuncRegisterHandler<float> OutSocketTranslateY(Context);
@@ -1022,51 +1022,61 @@ void UNiagaraDataInterfaceSkeletalMesh::GetFilteredSocketTransform(FVectorVMCont
 	VectorVM::FExternalFuncRegisterHandler<float> OutSocketScaleX(Context);
 	VectorVM::FExternalFuncRegisterHandler<float> OutSocketScaleY(Context);
 	VectorVM::FExternalFuncRegisterHandler<float> OutSocketScaleZ(Context);
- 
- 	//USkeletalMesh* Mesh = InstData->Mesh;
-	USkeletalMeshComponent* SkelComp = Cast<USkeletalMeshComponent>(InstData->Component.Get());
-	// UE_LOG(LogNiagara, Warning, TEXT("SkelComp: %p, Mesh: %p, Num: %d"), SkelComp, InstData->Mesh, Context.NumInstances);
-
-	FTransform ComponentToWorld;
-	if (SkelComp != nullptr) {
-		ComponentToWorld = SkelComp->GetComponentToWorld();
-	}
 
 	const TArray<FTransform>& CurrentFilteredSockets = InstData->GetFilteredSocketsCurrBuffer();
+	const int32 SocketMax = CurrentFilteredSockets.Num() - 1;
+	if (SocketMax >= 0)
+	{
+		const bool bNeedsRotation = OutSocketRotationX.IsValid() || OutSocketRotationY.IsValid() || OutSocketRotationZ.IsValid() || OutSocketRotationW.IsValid();
+		const FMatrix InstanceTransform = InstData->Transform;
+		const FQuat InstanceRotation = bNeedsRotation ? InstanceTransform.ToQuat() : FQuat::Identity;
 
-	for (int32 i = 0; i < Context.NumInstances; ++i)
-	{//
-		const int32 SocketIndex = FMath::Clamp(SocketParam.GetAndAdvance(), 0, CurrentFilteredSockets.Num()-1);
-		const int32 ShouldApplyWorldTransform = ApplyWorldTransform.GetAndAdvance();
-		
-		FTransform CurrentTransform = CurrentFilteredSockets[SocketIndex];
-		if (ShouldApplyWorldTransform == -1)
+		for (int32 i = 0; i < Context.NumInstances; ++i)
 		{
-			CurrentTransform = CurrentTransform * ComponentToWorld;
-		}
+			const int32 SocketIndex = FMath::Clamp(SocketParam.GetAndAdvance(), 0, SocketMax);
+			FVector SocketTranslation = CurrentFilteredSockets[SocketIndex].GetTranslation();
+			FQuat SocketRotation = CurrentFilteredSockets[SocketIndex].GetRotation();
+			FVector SocketScale = CurrentFilteredSockets[SocketIndex].GetScale3D();
+
+			const bool bApplyTransform = ApplyWorldTransform.GetAndAdvance().GetValue();
+			if (bApplyTransform)
+			{
+				SocketTranslation = InstanceTransform.TransformPosition(SocketTranslation);
+				SocketRotation = InstanceRotation * SocketRotation;
+				SocketScale = InstanceTransform.TransformVector(SocketScale);
+			}
+
+			*OutSocketTranslateX.GetDestAndAdvance() = SocketTranslation.X;
+			*OutSocketTranslateY.GetDestAndAdvance() = SocketTranslation.Y;
+			*OutSocketTranslateZ.GetDestAndAdvance() = SocketTranslation.Z;
+
+			*OutSocketRotationX.GetDestAndAdvance() = SocketRotation.X;
+			*OutSocketRotationY.GetDestAndAdvance() = SocketRotation.Y;
+			*OutSocketRotationZ.GetDestAndAdvance() = SocketRotation.Z;
+			*OutSocketRotationW.GetDestAndAdvance() = SocketRotation.W;
 
-		FQuat Rotation = CurrentTransform.GetRotation();
-		FVector Translation = CurrentTransform.GetTranslation();
-		FVector Scale = CurrentTransform.GetScale3D();
-
-		*OutSocketTranslateX.GetDestAndAdvance() = Translation.X;
-		*OutSocketTranslateY.GetDestAndAdvance() = Translation.Y;
-		*OutSocketTranslateZ.GetDestAndAdvance() = Translation.Z;
-
-		*OutSocketRotationX.GetDestAndAdvance() = Rotation.X;
-		*OutSocketRotationY.GetDestAndAdvance() = Rotation.Y;
-		*OutSocketRotationZ.GetDestAndAdvance() = Rotation.Z;
-		*OutSocketRotationW.GetDestAndAdvance() = Rotation.W;
-
-		*OutSocketScaleX.GetDestAndAdvance() = Scale.X;
-		*OutSocketScaleY.GetDestAndAdvance() = Scale.Y;
-		*OutSocketScaleZ.GetDestAndAdvance() = Scale.Z;
-			
-		// UE_LOG(LogNiagara, Warning, TEXT("Enum: %d, translation: %f %f %f, Rotation: %f %f %f %f, scale: %f %f %f"), 
-		// 		  ShouldApplyWorldTransform,
-		// 		  Translation.X, Translation.Y, Translation.Z,
-		// 		  Rotation.X, Rotation.Y, Rotation.Z, Rotation.W,
-		// 		  Scale.X, Scale.Y, Scale.Z);
+			*OutSocketScaleX.GetDestAndAdvance() = SocketScale.X;
+			*OutSocketScaleY.GetDestAndAdvance() = SocketScale.Y;
+			*OutSocketScaleZ.GetDestAndAdvance() = SocketScale.Z;
+		}
+	}
+	else
+	{
+		for (int32 i=0; i < Context.NumInstances; ++i)
+		{
+			*OutSocketTranslateX.GetDestAndAdvance() = 0.0f;
+			*OutSocketTranslateY.GetDestAndAdvance() = 0.0f;
+			*OutSocketTranslateZ.GetDestAndAdvance() = 0.0f;
+
+			*OutSocketRotationX.GetDestAndAdvance() = FQuat::Identity.X;
+			*OutSocketRotationY.GetDestAndAdvance() = FQuat::Identity.Y;
+			*OutSocketRotationZ.GetDestAndAdvance() = FQuat::Identity.Z;
+			*OutSocketRotationW.GetDestAndAdvance() = FQuat::Identity.W;
+
+			*OutSocketScaleX.GetDestAndAdvance() = 1.0f;
+			*OutSocketScaleY.GetDestAndAdvance() = 1.0f;
+			*OutSocketScaleZ.GetDestAndAdvance() = 1.0f;
+		}
 	}
 }
 
diff --git a/Engine/Plugins/FX/Niagara/Source/Niagara/Private/NiagaraDataInterfaceSkeletalMesh.cpp b/Engine/Plugins/FX/Niagara/Source/Niagara/Private/NiagaraDataInterfaceSkeletalMesh.cpp
--- a/Engine/Plugins/FX/Niagara/Source/Niagara/Private/NiagaraDataInterfaceSkeletalMesh.cpp
+++ b/Engine/Plugins/FX/Niagara/Source/Niagara/Private/NiagaraDataInterfaceSkeletalMesh.cpp
@@ -965,6 +965,8 @@ struct FNDISkeletalMeshParametersName
 	FString FilteredSocketBoneOffsetName;
 	FString InstanceTransformName;
 	FString InstancePrevTransformName;
+	FString InstanceRotationName;
+	FString InstancePrevRotationName;
 	FString InstanceInvDeltaTimeName;
 	FString EnabledFeaturesName;
 };
@@ -1007,6 +1009,8 @@ static void GetNiagaraDataInterfaceParametersName(FNDISkeletalMeshParametersName
 	Names.FilteredSocketBoneOffsetName = UNiagaraDataInterfaceSkeletalMesh::FilteredSocketBoneOffsetName + Suffix;
 	Names.InstanceTransformName = UNiagaraDataInterfaceSkeletalMesh::InstanceTransformName + Suffix;
 	Names.InstancePrevTransformName = UNiagaraDataInterfaceSkeletalMesh::InstancePrevTransformName + Suffix;
+	Names.InstanceRotationName = UNiagaraDataInterfaceSkeletalMesh::InstanceRotationName + Suffix;
+	Names.InstancePrevRotationName = UNiagaraDataInterfaceSkeletalMesh::InstancePrevRotationName + Suffix;
 	Names.InstanceInvDeltaTimeName = UNiagaraDataInterfaceSkeletalMesh::InstanceInvDeltaTimeName + Suffix;
 	Names.EnabledFeaturesName = UNiagaraDataInterfaceSkeletalMesh::EnabledFeaturesName + Suffix;
 }
@@ -1056,6 +1060,8 @@ public:
 		FilteredSocketBoneOffset.Bind(ParameterMap, *ParamNames.FilteredSocketBoneOffsetName);
 		InstanceTransform.Bind(ParameterMap, *ParamNames.InstanceTransformName);
 		InstancePrevTransform.Bind(ParameterMap, *ParamNames.InstancePrevTransformName);
+		InstanceRotation.Bind(ParameterMap, *ParamNames.InstanceRotationName);
+		InstancePrevRotation.Bind(ParameterMap, *ParamNames.InstancePrevRotationName);
 		InstanceInvDeltaTime.Bind(ParameterMap, *ParamNames.InstanceInvDeltaTimeName);
 		EnabledFeatures.Bind(ParameterMap, *ParamNames.EnabledFeaturesName);
 	}
@@ -1169,6 +1175,8 @@ public:
 
 			SetShaderValue(RHICmdList, ComputeShaderRHI, InstanceTransform, InstanceData->Transform);
 			SetShaderValue(RHICmdList, ComputeShaderRHI, InstancePrevTransform, InstanceData->PrevTransform);
+			SetShaderValue(RHICmdList, ComputeShaderRHI, InstanceRotation, InstanceData->Transform.ToQuat());
+			SetShaderValue(RHICmdList, ComputeShaderRHI, InstancePrevRotation, InstanceData->PrevTransform.ToQuat());
 			SetShaderValue(RHICmdList, ComputeShaderRHI, InstanceInvDeltaTime, 1.0f / InstanceData->DeltaSeconds);
 
 			SetShaderValue(RHICmdList, ComputeShaderRHI, EnabledFeatures, EnabledFeaturesBits);
@@ -1219,6 +1227,8 @@ public:
 
 			SetShaderValue(RHICmdList, ComputeShaderRHI, InstanceTransform, FMatrix::Identity);
 			SetShaderValue(RHICmdList, ComputeShaderRHI, InstancePrevTransform, FMatrix::Identity);
+			SetShaderValue(RHICmdList, ComputeShaderRHI, InstanceRotation, FQuat::Identity);
+			SetShaderValue(RHICmdList, ComputeShaderRHI, InstancePrevRotation, FQuat::Identity);
 			SetShaderValue(RHICmdList, ComputeShaderRHI, InstanceInvDeltaTime, 0.0f);
 
 			SetShaderValue(RHICmdList, ComputeShaderRHI, EnabledFeatures, 0);
@@ -1262,6 +1272,8 @@ private:
 	LAYOUT_FIELD(FShaderParameter, FilteredSocketBoneOffset);
 	LAYOUT_FIELD(FShaderParameter, InstanceTransform);
 	LAYOUT_FIELD(FShaderParameter, InstancePrevTransform);
+	LAYOUT_FIELD(FShaderParameter, InstanceRotation);
+	LAYOUT_FIELD(FShaderParameter, InstancePrevRotation);
 	LAYOUT_FIELD(FShaderParameter, InstanceInvDeltaTime);
 	LAYOUT_FIELD(FShaderParameter, EnabledFeatures);
 };
@@ -2414,6 +2426,8 @@ const FString UNiagaraDataInterfaceSkeletalMesh::NumFilteredSocketsName(TEXT("Nu
 const FString UNiagaraDataInterfaceSkeletalMesh::FilteredSocketBoneOffsetName(TEXT("FilteredSocketBoneOffset_"));
 const FString UNiagaraDataInterfaceSkeletalMesh::InstanceTransformName(TEXT("InstanceTransform_"));
 const FString UNiagaraDataInterfaceSkeletalMesh::InstancePrevTransformName(TEXT("InstancePrevTransform_"));
+const FString UNiagaraDataInterfaceSkeletalMesh::InstanceRotationName(TEXT("InstanceRotation_"));
+const FString UNiagaraDataInterfaceSkeletalMesh::InstancePrevRotationName(TEXT("InstancePrevRotation_"));
 const FString UNiagaraDataInterfaceSkeletalMesh::InstanceInvDeltaTimeName(TEXT("InstanceInvDeltaTime_"));
 const FString UNiagaraDataInterfaceSkeletalMesh::EnabledFeaturesName(TEXT("EnabledFeatures_"));
 
@@ -2655,13 +2669,13 @@ bool UNiagaraDataInterfaceSkeletalMesh::GetFunctionHLSL(const FNiagaraDataInterf
 		static const TCHAR* FormatSample = TEXT("void {InstanceFunctionName} (in int SocketIndex, out int Bone) { {GetDISkelMeshContextName} DISkelMesh_GetFilteredSocketBoneAt(DIContext, SocketIndex, Bone); }");
 		OutHLSL += FString::Format(FormatSample, ArgsSample);
 	} 
-	else if (FunctionInfo.DefinitionName == FSkeletalMeshInterfaceHelper::GetFilteredSocketTransformName)
-	{
-		// TODO: This just returns the Identity transform.
-		// TODO: Make this work on the GPU?
-		static const TCHAR* FormatSample = TEXT("void {InstanceFunctionName} (in int SocketIndex, in int bShouldApplyTransform, out float3 Translation, out float4 Rotation, out float4 Scale) { Translation = float3(0.0, 0.0, 0.0); Rotation = float4(0.0, 0.0, 0.0, 1.0); Scale = float3(1.0, 1.0, 1.0); }");
-		OutHLSL += FString::Format(FormatSample, ArgsSample);
-	} 
+	//else if (FunctionInfo.DefinitionName == FSkeletalMeshInterfaceHelper::GetFilteredSocketTransformName)
+	//{
+	//	// TODO: This just returns the Identity transform.
+	//	// TODO: Make this work on the GPU?
+	//	static const TCHAR* FormatSample = TEXT("void {InstanceFunctionName} (in int SocketIndex, in int bShouldApplyTransform, out float3 Translation, out float4 Rotation, out float3 Scale) { Translation = float3(0.0, 0.0, 0.0); Rotation = float4(0.0, 0.0, 0.0, 1.0); Scale = float3(1.0, 1.0, 1.0); }");
+	//	OutHLSL += FString::Format(FormatSample, ArgsSample);
+	//} 
 	else if (FunctionInfo.DefinitionName == FSkeletalMeshInterfaceHelper::RandomFilteredSocketName)
 	{
 		static const TCHAR* FormatSample = TEXT("void {InstanceFunctionName} (NiagaraRandInfo InRandomInfo, out int SocketBone) { {GetDISkelMeshContextName} DISkelMesh_RandomFilteredSocket(DIContext, InRandomInfo.Seed1, InRandomInfo.Seed2, InRandomInfo.Seed3, SocketBone); }");
diff --git a/Engine/Plugins/FX/Niagara/Source/NiagaraCore/Private/NiagaraCustomVersion.cpp b/Engine/Plugins/FX/Niagara/Source/NiagaraCore/Private/NiagaraCustomVersion.cpp
--- a/Engine/Plugins/FX/Niagara/Source/NiagaraCore/Private/NiagaraCustomVersion.cpp
+++ b/Engine/Plugins/FX/Niagara/Source/NiagaraCore/Private/NiagaraCustomVersion.cpp
@@ -8,4 +8,4 @@ const FGuid FNiagaraCustomVersion::GUID(0xFCF57AFA, 0x50764283, 0xB9A9E658, 0xFF
 // Register the custom version with core
 FCustomVersionRegistration GRegisterNiagaraCustomVersion(FNiagaraCustomVersion::GUID, FNiagaraCustomVersion::LatestVersion, TEXT("NiagaraVer"));
 
-const FGuid FNiagaraCustomVersion::LatestScriptCompileVersion(0xF301F0D6, 0xA2B64D87, 0xAFFB0B25, 0xD1AAA66E);
+const FGuid FNiagaraCustomVersion::LatestScriptCompileVersion(0x1F365103, 0x4B5DC149, 0x8308F39D, 0x698AC03E);
\ No newline at end of file
