Minor fix for audio thread hang that could occur with a Niagara Emitter using the Osicilloscope or Spectrum. Issue was a redundant call to OnUpdateSubmix on init, which causes us to flush the audio thread too early in startup.

diff --git a/Engine/Plugins/FX/Niagara/Source/Niagara/Private/NiagaraDataInterfaceAudio.cpp b/Engine/Plugins/FX/Niagara/Source/Niagara/Private/NiagaraDataInterfaceAudio.cpp
--- a/Engine/Plugins/FX/Niagara/Source/Niagara/Private/NiagaraDataInterfaceAudio.cpp
+++ b/Engine/Plugins/FX/Niagara/Source/Niagara/Private/NiagaraDataInterfaceAudio.cpp
@@ -321,7 +321,6 @@ void UNiagaraDataInterfaceAudioSubmix::BeginDestroy()
 void UNiagaraDataInterfaceAudioSubmix::PostLoad()
 {
 	Super::PostLoad();
-	GetProxyAs<FNiagaraDataInterfaceProxySubmix>()->OnUpdateSubmix(Submix);
 }
 
 bool UNiagaraDataInterfaceAudioSubmix::CopyToInternal(UNiagaraDataInterface* Destination) const
