Direct Read DI Optimizations and marking as experimental.

- Rework CPU calls to avoid work in the inner loop.
- Improve readability and maintainabiltiy using templates to remove ~250 lines of code.
- Add several helper functions and classes for templated DI code.
- Marked all functions as experimental

diff --git a/Engine/Plugins/FX/Niagara/Source/Niagara/Classes/NiagaraDataInterface.h b/Engine/Plugins/FX/Niagara/Source/Niagara/Classes/NiagaraDataInterface.h
--- a/Engine/Plugins/FX/Niagara/Source/Niagara/Classes/NiagaraDataInterface.h
+++ b/Engine/Plugins/FX/Niagara/Source/Niagara/Classes/NiagaraDataInterface.h
@@ -597,4 +597,105 @@ struct FNDIRandomHelper
 	FNDIParameter<FNiagaraRandInfo> RandParam;
 
 	FNiagaraRandInfo RandInfo;
+};
+
+//Helper to deal with types with potentially several output registers.
+template<typename T>
+struct FNDIOutputParam
+{
+	VectorVM::FExternalFuncRegisterHandler<T> Data;
+	FORCEINLINE FNDIOutputParam(FVectorVMContext& Context) : Data(Context) {}
+	FORCEINLINE void SetAndAdvance(T Val) { *Data.GetDestAndAdvance() = Val; }
+};
+
+template<>
+struct FNDIOutputParam<FVector2D>
+{
+	VectorVM::FExternalFuncRegisterHandler<float> X;
+	VectorVM::FExternalFuncRegisterHandler<float> Y;
+	FORCEINLINE FNDIOutputParam(FVectorVMContext& Context) : X(Context), Y(Context) {}
+	FORCEINLINE void SetAndAdvance(FVector2D Val)
+	{
+		*X.GetDestAndAdvance() = Val.X;
+		*Y.GetDestAndAdvance() = Val.Y;
+	}
+};
+
+template<>
+struct FNDIOutputParam<FVector>
+{
+	VectorVM::FExternalFuncRegisterHandler<float> X;
+	VectorVM::FExternalFuncRegisterHandler<float> Y;
+	VectorVM::FExternalFuncRegisterHandler<float> Z;
+	FNDIOutputParam(FVectorVMContext& Context) : X(Context), Y(Context), Z(Context) {}
+	FORCEINLINE void SetAndAdvance(FVector Val)
+	{
+		*X.GetDestAndAdvance() = Val.X;
+		*Y.GetDestAndAdvance() = Val.Y;
+		*Z.GetDestAndAdvance() = Val.Z;
+	}
+};
+
+template<>
+struct FNDIOutputParam<FVector4>
+{
+	VectorVM::FExternalFuncRegisterHandler<float> X;
+	VectorVM::FExternalFuncRegisterHandler<float> Y;
+	VectorVM::FExternalFuncRegisterHandler<float> Z;
+	VectorVM::FExternalFuncRegisterHandler<float> W;
+	FORCEINLINE FNDIOutputParam(FVectorVMContext& Context) : X(Context), Y(Context), Z(Context), W(Context) {}
+	FORCEINLINE void SetAndAdvance(FVector4 Val)
+	{
+		*X.GetDestAndAdvance() = Val.X;
+		*Y.GetDestAndAdvance() = Val.Y;
+		*Z.GetDestAndAdvance() = Val.Z;
+		*W.GetDestAndAdvance() = Val.W;
+	}
+};
+
+template<>
+struct FNDIOutputParam<FQuat>
+{
+	VectorVM::FExternalFuncRegisterHandler<float> X;
+	VectorVM::FExternalFuncRegisterHandler<float> Y;
+	VectorVM::FExternalFuncRegisterHandler<float> Z;
+	VectorVM::FExternalFuncRegisterHandler<float> W;
+	FORCEINLINE FNDIOutputParam(FVectorVMContext& Context) : X(Context), Y(Context), Z(Context), W(Context) {}
+	FORCEINLINE void SetAndAdvance(FQuat Val)
+	{
+		*X.GetDestAndAdvance() = Val.X;
+		*Y.GetDestAndAdvance() = Val.Y;
+		*Z.GetDestAndAdvance() = Val.Z;
+		*W.GetDestAndAdvance() = Val.W;
+	}
+};
+
+template<>
+struct FNDIOutputParam<FLinearColor>
+{
+	VectorVM::FExternalFuncRegisterHandler<float> R;
+	VectorVM::FExternalFuncRegisterHandler<float> G;
+	VectorVM::FExternalFuncRegisterHandler<float> B;
+	VectorVM::FExternalFuncRegisterHandler<float> A;
+	FORCEINLINE FNDIOutputParam(FVectorVMContext& Context) : R(Context), G(Context), B(Context), A(Context) {}
+	FORCEINLINE void SetAndAdvance(FLinearColor Val)
+	{
+		*R.GetDestAndAdvance() = Val.R;
+		*G.GetDestAndAdvance() = Val.G;
+		*B.GetDestAndAdvance() = Val.B;
+		*A.GetDestAndAdvance() = Val.A;
+	}
+};
+
+template<>
+struct FNDIOutputParam<FNiagaraID>
+{
+	VectorVM::FExternalFuncRegisterHandler<int32> Index;
+	VectorVM::FExternalFuncRegisterHandler<int32> AcquireTag;
+	FORCEINLINE FNDIOutputParam(FVectorVMContext& Context) : Index(Context), AcquireTag(Context) {}
+	FORCEINLINE void SetAndAdvance(FNiagaraID Val)
+	{
+		*Index.GetDestAndAdvance() = Val.Index;
+		*AcquireTag.GetDestAndAdvance() = Val.AcquireTag;
+	}
 };
\ No newline at end of file
diff --git a/Engine/Plugins/FX/Niagara/Source/Niagara/Classes/NiagaraDataInterfaceParticleRead.h b/Engine/Plugins/FX/Niagara/Source/Niagara/Classes/NiagaraDataInterfaceParticleRead.h
--- a/Engine/Plugins/FX/Niagara/Source/Niagara/Classes/NiagaraDataInterfaceParticleRead.h
+++ b/Engine/Plugins/FX/Niagara/Source/Niagara/Classes/NiagaraDataInterfaceParticleRead.h
@@ -60,11 +60,4 @@ public:
 
 protected:
 	virtual bool CopyToInternal(UNiagaraDataInterface* Destination) const override;
-
-private:
-	template<typename Type, typename AccessorType>
-	Type RetrieveValueWithCheck(FNiagaraEmitterInstance* EmitterInstance, const FName& Attr, const FNiagaraID& ParticleID, bool& bValid);
-
-	template<typename Type, typename AccessorType>
-	Type RetrieveValueByIndexWithCheck(FNiagaraEmitterInstance* EmitterInstance, const FName& Attr, int32 ParticleIndex, bool& bValid);
 };
\ No newline at end of file
diff --git a/Engine/Plugins/FX/Niagara/Source/Niagara/Classes/NiagaraDataSet.h b/Engine/Plugins/FX/Niagara/Source/Niagara/Classes/NiagaraDataSet.h
--- a/Engine/Plugins/FX/Niagara/Source/Niagara/Classes/NiagaraDataSet.h
+++ b/Engine/Plugins/FX/Niagara/Source/Niagara/Classes/NiagaraDataSet.h
@@ -1694,6 +1694,18 @@ struct FNiagaraEitherDataSetAccessor
 	typedef typename TagType::EitherB EitherB;
 
 	FNiagaraEitherDataSetAccessor() {}
+	FNiagaraEitherDataSetAccessor(FNiagaraDataSet& InDataSet, const FNiagaraVariable& InVar)
+	{
+		if (InDataSet.HasVariable(InVar))
+		{
+			FloatDataSet = FNiagaraDataSetAccessor<EitherA>(InDataSet, FNiagaraVariable(TagType::GetTypeA(), InVar.GetName()));
+		}
+		else
+		{
+			HalfDataSet = FNiagaraDataSetAccessor<EitherB>(InDataSet, FNiagaraVariable(TagType::GetTypeB(), InVar.GetName()));
+		}
+		InitForAccess();
+	}
 	FNiagaraEitherDataSetAccessor(FNiagaraDataSet& InDataSet, const FName& InName)
 	{
 		if (InDataSet.HasVariable(FNiagaraVariable(TagType::GetTypeA(), InName)))
@@ -1825,256 +1837,240 @@ private:
 	FNiagaraDataSetAccessor<EitherB> HalfDataSet;
 };
 
+//////////////////////////////////////////////////////////////////////////
+//Conversion helpers for compressable types.
+
+template<typename T>
 struct FNiagaraDataConversions
 {
-	struct FHalfOrFloat
-	{
-		using EitherA = float;
-		using EitherB = TStaticArray<FFloat16, 1>;
-
-		static TStaticArray<FFloat16, 1> ConvertA(float In)
-		{
-			TStaticArray<FFloat16, 1> Ret;
-			FPlatformMath::StoreHalf((uint16*)&Ret[0], In);
-			return Ret;
-		}
-		static float ConvertB(TStaticArray<FFloat16, 1> In)
-		{
-			return FPlatformMath::LoadHalf((uint16*)&In[0]);
-		}
-
-		static FNiagaraTypeDefinition GetTypeA()
-		{
-			return FNiagaraTypeDefinition::GetFloatDef();
-		}
+};
 
-		static FNiagaraTypeDefinition GetTypeB()
-		{
-			return FNiagaraTypeDefinition::GetHalfDef();
-		}
-	};
+template<>
+struct FNiagaraDataConversions<float>
+{
+	using EitherA = float;
+	using EitherB = TStaticArray<FFloat16, 1>;
 
-	struct FHalf2OrFloat2
+	static TStaticArray<FFloat16, 1> ConvertA(float In)
 	{
-		using EitherA = FVector2D;
-		using EitherB = TStaticArray<FFloat16, 2>;
-
-		static TStaticArray<FFloat16, 2> ConvertA(FVector2D In)
-		{
-			TStaticArray<FFloat16, 2> Ret;
-			FPlatformMath::StoreHalf((uint16*)&Ret[0], In.X);
-			FPlatformMath::StoreHalf((uint16*)&Ret[1], In.Y);
-			return Ret;
-		}
-		static FVector2D ConvertB(TStaticArray<FFloat16, 2> In)
-		{
-			return FVector2D
-			(
-				FPlatformMath::LoadHalf((uint16*)&In[0]),
-				FPlatformMath::LoadHalf((uint16*)&In[1])
-			);
-		}
-
-		static FNiagaraTypeDefinition GetTypeA()
-		{
-			return FNiagaraTypeDefinition::GetVec2Def();
-		}
-
-		static FNiagaraTypeDefinition GetTypeB()
-		{
-			return FNiagaraTypeDefinition::GetHalfVec2Def();
-		}
-	};
-
-	struct FHalf3OrFloat3
+		TStaticArray<FFloat16, 1> Ret;
+		FPlatformMath::StoreHalf((uint16*)&Ret[0], In);
+		return Ret;
+	}
+	static float ConvertB(TStaticArray<FFloat16, 1> In)
 	{
-		using EitherA = FVector;
-		using EitherB = TStaticArray<FFloat16, 3>;
+		return FPlatformMath::LoadHalf((uint16*)&In[0]);
+	}
 
-		static TStaticArray<FFloat16, 3> ConvertA(FVector In)
-		{
-			TStaticArray<FFloat16, 3> Ret;
-			FPlatformMath::StoreHalf((uint16*)&Ret[0], In.X);
-			FPlatformMath::StoreHalf((uint16*)&Ret[1], In.Y);
-			FPlatformMath::StoreHalf((uint16*)&Ret[2], In.Z);
-			return Ret;
-		}
-		static FVector ConvertB(TStaticArray<FFloat16, 3> In)
-		{
-			return FVector
-			(
-				FPlatformMath::LoadHalf((uint16*)&In[0]), 
-				FPlatformMath::LoadHalf((uint16*)&In[1]), 
-				FPlatformMath::LoadHalf((uint16*)&In[2])
-			);
-		}
+	static FNiagaraTypeDefinition GetTypeA()
+	{
+		return FNiagaraTypeDefinition::GetFloatDef();
+	}
 
-		static FNiagaraTypeDefinition GetTypeA()
-		{
-			return FNiagaraTypeDefinition::GetVec3Def();
-		}
+	static FNiagaraTypeDefinition GetTypeB()
+	{
+		return FNiagaraTypeDefinition::GetHalfDef();
+	}
+};
 
-		static FNiagaraTypeDefinition GetTypeB()
-		{
-			return FNiagaraTypeDefinition::GetHalfVec3Def();
-		}
-	};
+template<>
+struct FNiagaraDataConversions<FVector2D>
+{
+	using EitherA = FVector2D;
+	using EitherB = TStaticArray<FFloat16, 2>;
 
-	struct FHalf4OrFloat4
+	static TStaticArray<FFloat16, 2> ConvertA(FVector2D In)
 	{
-		using EitherA = FVector4;
-		using EitherB = TStaticArray<FFloat16, 4>;
+		TStaticArray<FFloat16, 2> Ret;
+		FPlatformMath::StoreHalf((uint16*)&Ret[0], In.X);
+		FPlatformMath::StoreHalf((uint16*)&Ret[1], In.Y);
+		return Ret;
+	}
+	static FVector2D ConvertB(TStaticArray<FFloat16, 2> In)
+	{
+		return FVector2D
+		(
+			FPlatformMath::LoadHalf((uint16*)&In[0]),
+			FPlatformMath::LoadHalf((uint16*)&In[1])
+		);
+	}
 
-		static TStaticArray<FFloat16, 4> ConvertA(FVector4 In)
-		{
-			TStaticArray<FFloat16, 4> Ret;
-			FPlatformMath::VectorStoreHalf((uint16*)&Ret[0], (float*)&In);
-			return Ret;
-		}
-		static FVector4 ConvertB(TStaticArray<FFloat16, 4> In)
-		{
-			FVector4 Ret;
-			FPlatformMath::VectorLoadHalf((float*)&Ret, (uint16*)&In[0]);
-			return Ret;
-		}
+	static FNiagaraTypeDefinition GetTypeA()
+	{
+		return FNiagaraTypeDefinition::GetVec2Def();
+	}
 
-		static FNiagaraTypeDefinition GetTypeA()
-		{
-			return FNiagaraTypeDefinition::GetVec4Def();
-		}
+	static FNiagaraTypeDefinition GetTypeB()
+	{
+		return FNiagaraTypeDefinition::GetHalfVec2Def();
+	}
+};
 
-		static FNiagaraTypeDefinition GetTypeB()
-		{
-			return FNiagaraTypeDefinition::GetHalfVec4Def();
-		}
-	};
+template<>
+struct FNiagaraDataConversions<FVector>
+{
+	using EitherA = FVector;
+	using EitherB = TStaticArray<FFloat16, 3>;
 
-	struct FHalf4OrColor4
+	static TStaticArray<FFloat16, 3> ConvertA(FVector In)
+	{
+		TStaticArray<FFloat16, 3> Ret;
+		FPlatformMath::StoreHalf((uint16*)&Ret[0], In.X);
+		FPlatformMath::StoreHalf((uint16*)&Ret[1], In.Y);
+		FPlatformMath::StoreHalf((uint16*)&Ret[2], In.Z);
+		return Ret;
+	}
+	static FVector ConvertB(TStaticArray<FFloat16, 3> In)
 	{
-		using EitherA = FLinearColor;
-		using EitherB = TStaticArray<FFloat16, 4>;
+		return FVector
+		(
+			FPlatformMath::LoadHalf((uint16*)&In[0]),
+			FPlatformMath::LoadHalf((uint16*)&In[1]),
+			FPlatformMath::LoadHalf((uint16*)&In[2])
+		);
+	}
 
-		static TStaticArray<FFloat16, 4> ConvertA(FLinearColor In)
-		{
-			TStaticArray<FFloat16, 4> Ret;
-			FPlatformMath::VectorStoreHalf((uint16*)&Ret[0], (float*)&In);
-			return Ret;
-		}
-		static FLinearColor ConvertB(TStaticArray<FFloat16, 4> In)
-		{
-			FLinearColor Ret;
-			FPlatformMath::VectorLoadHalf((float*)&Ret, (uint16*)&In[0]);
-			return Ret;
-		}
+	static FNiagaraTypeDefinition GetTypeA()
+	{
+		return FNiagaraTypeDefinition::GetVec3Def();
+	}
 
-		static FNiagaraTypeDefinition GetTypeA()
-		{
-			return FNiagaraTypeDefinition::GetColorDef();
-		}
+	static FNiagaraTypeDefinition GetTypeB()
+	{
+		return FNiagaraTypeDefinition::GetHalfVec3Def();
+	}
+};
 
-		static FNiagaraTypeDefinition GetTypeB()
-		{
-			return FNiagaraTypeDefinition::GetHalfVec4Def();
-		}
-	};
+template<>
+struct FNiagaraDataConversions<FVector4>
+{
+	using EitherA = FVector4;
+	using EitherB = TStaticArray<FFloat16, 4>;
 
-	struct FHalf4OrQuat4
+	static TStaticArray<FFloat16, 4> ConvertA(FVector4 In)
+	{
+		TStaticArray<FFloat16, 4> Ret;
+		FPlatformMath::VectorStoreHalf((uint16*)&Ret[0], (float*)&In);
+		return Ret;
+	}
+	static FVector4 ConvertB(TStaticArray<FFloat16, 4> In)
 	{
-		using EitherA = FQuat;
-		using EitherB = TStaticArray<FFloat16, 4>;
+		FVector4 Ret;
+		FPlatformMath::VectorLoadHalf((float*)&Ret, (uint16*)&In[0]);
+		return Ret;
+	}
 
-		static TStaticArray<FFloat16, 4> ConvertA(FQuat In)
-		{
-			TStaticArray<FFloat16, 4> Ret;
-			FPlatformMath::VectorStoreHalf((uint16*)&Ret[0], (float*)&In);
-			return Ret;
-		}
-		static FQuat ConvertB(TStaticArray<FFloat16, 4> In)
-		{
-			FQuat Ret;
-			FPlatformMath::VectorLoadHalf((float*)&Ret, (uint16*)&In[0]);
-			return Ret;
-		}
+	static FNiagaraTypeDefinition GetTypeA()
+	{
+		return FNiagaraTypeDefinition::GetVec4Def();
+	}
 
-		static FNiagaraTypeDefinition GetTypeA()
-		{
-			return FNiagaraTypeDefinition::GetQuatDef();
-		}
+	static FNiagaraTypeDefinition GetTypeB()
+	{
+		return FNiagaraTypeDefinition::GetHalfVec4Def();
+	}
+};
 
-		static FNiagaraTypeDefinition GetTypeB()
-		{
-			return FNiagaraTypeDefinition::GetHalfVec4Def();
-		}
-	};
+template<>
+struct FNiagaraDataConversions<FLinearColor>
+{
+	using EitherA = FLinearColor;
+	using EitherB = TStaticArray<FFloat16, 4>;
 
-	struct FNiagaraInt {};
+	static TStaticArray<FFloat16, 4> ConvertA(FLinearColor In)
+	{
+		TStaticArray<FFloat16, 4> Ret;
+		FPlatformMath::VectorStoreHalf((uint16*)&Ret[0], (float*)&In);
+		return Ret;
+	}
+	static FLinearColor ConvertB(TStaticArray<FFloat16, 4> In)
+	{
+		FLinearColor Ret;
+		FPlatformMath::VectorLoadHalf((float*)&Ret, (uint16*)&In[0]);
+		return Ret;
+	}
 
-	struct FNiagaraBool {};
+	static FNiagaraTypeDefinition GetTypeA()
+	{
+		return FNiagaraTypeDefinition::GetColorDef();
+	}
 
-	struct FNiagaraConvertID {};
+	static FNiagaraTypeDefinition GetTypeB()
+	{
+		return FNiagaraTypeDefinition::GetHalfVec4Def();
+	}
 };
 
 template<>
-struct FNiagaraDataSetAccessor<FNiagaraDataConversions::FHalfOrFloat> : public FNiagaraEitherDataSetAccessor<FNiagaraDataConversions::FHalfOrFloat>
+struct FNiagaraDataConversions<FQuat>
 {
-	using FNiagaraEitherDataSetAccessor<FNiagaraDataConversions::FHalfOrFloat>::FNiagaraEitherDataSetAccessor;
-};
+	using EitherA = FQuat;
+	using EitherB = TStaticArray<FFloat16, 4>;
 
+	static TStaticArray<FFloat16, 4> ConvertA(FQuat In)
+	{
+		TStaticArray<FFloat16, 4> Ret;
+		FPlatformMath::VectorStoreHalf((uint16*)&Ret[0], (float*)&In);
+		return Ret;
+	}
+	static FQuat ConvertB(TStaticArray<FFloat16, 4> In)
+	{
+		FQuat Ret;
+		FPlatformMath::VectorLoadHalf((float*)&Ret, (uint16*)&In[0]);
+		return Ret;
+	}
 
-template<>
-struct FNiagaraDataSetAccessor<FNiagaraDataConversions::FHalf2OrFloat2> : public FNiagaraEitherDataSetAccessor<FNiagaraDataConversions::FHalf2OrFloat2>
-{
-	using FNiagaraEitherDataSetAccessor<FNiagaraDataConversions::FHalf2OrFloat2>::FNiagaraEitherDataSetAccessor;
+	static FNiagaraTypeDefinition GetTypeA()
+	{
+		return FNiagaraTypeDefinition::GetQuatDef();
+	}
+
+	static FNiagaraTypeDefinition GetTypeB()
+	{
+		return FNiagaraTypeDefinition::GetHalfVec4Def();
+	}
 };
 
 template<>
-struct FNiagaraDataSetAccessor<FNiagaraDataConversions::FHalf3OrFloat3> : public FNiagaraEitherDataSetAccessor<FNiagaraDataConversions::FHalf3OrFloat3>
+struct FNiagaraDataSetAccessor<FNiagaraDataConversions<float>> : public FNiagaraEitherDataSetAccessor<FNiagaraDataConversions<float>>
 {
-	using FNiagaraEitherDataSetAccessor<FNiagaraDataConversions::FHalf3OrFloat3>::FNiagaraEitherDataSetAccessor;
+	using FNiagaraEitherDataSetAccessor<FNiagaraDataConversions<float>>::FNiagaraEitherDataSetAccessor;
 };
 
-
 template<>
-struct FNiagaraDataSetAccessor<FNiagaraDataConversions::FHalf4OrFloat4> : public FNiagaraEitherDataSetAccessor<FNiagaraDataConversions::FHalf4OrFloat4>
+struct FNiagaraDataSetAccessor<FNiagaraDataConversions<FVector2D>> : public FNiagaraEitherDataSetAccessor<FNiagaraDataConversions<FVector2D>>
 {
-	using FNiagaraEitherDataSetAccessor<FNiagaraDataConversions::FHalf4OrFloat4>::FNiagaraEitherDataSetAccessor;
+	using FNiagaraEitherDataSetAccessor<FNiagaraDataConversions<FVector2D>>::FNiagaraEitherDataSetAccessor;
 };
 
-
 template<>
-struct FNiagaraDataSetAccessor<FNiagaraDataConversions::FHalf4OrColor4> : public FNiagaraEitherDataSetAccessor<FNiagaraDataConversions::FHalf4OrColor4>
+struct FNiagaraDataSetAccessor<FNiagaraDataConversions<FVector>> : public FNiagaraEitherDataSetAccessor<FNiagaraDataConversions<FVector>>
 {
-	using FNiagaraEitherDataSetAccessor<FNiagaraDataConversions::FHalf4OrColor4>::FNiagaraEitherDataSetAccessor;
+	using FNiagaraEitherDataSetAccessor<FNiagaraDataConversions<FVector>>::FNiagaraEitherDataSetAccessor;
 };
 
 template<>
-struct FNiagaraDataSetAccessor<FNiagaraDataConversions::FHalf4OrQuat4> : public FNiagaraEitherDataSetAccessor<FNiagaraDataConversions::FHalf4OrQuat4>
+struct FNiagaraDataSetAccessor<FNiagaraDataConversions<FVector4>> : public FNiagaraEitherDataSetAccessor<FNiagaraDataConversions<FVector4>>
 {
-	using FNiagaraEitherDataSetAccessor<FNiagaraDataConversions::FHalf4OrQuat4>::FNiagaraEitherDataSetAccessor;
+	using FNiagaraEitherDataSetAccessor<FNiagaraDataConversions<FVector4>>::FNiagaraEitherDataSetAccessor;
 };
 
 template<>
-struct FNiagaraDataSetAccessor<FNiagaraDataConversions::FNiagaraInt> : public FNiagaraDataSetAccessor<int>
+struct FNiagaraDataSetAccessor<FNiagaraDataConversions<FLinearColor>> : public FNiagaraEitherDataSetAccessor<FNiagaraDataConversions<FLinearColor>>
 {
-	FNiagaraDataSetAccessor<FNiagaraDataConversions::FNiagaraInt>(FNiagaraDataSet& InDataSet, const FName& InName) : FNiagaraDataSetAccessor<int>(InDataSet, FNiagaraVariable(FNiagaraTypeDefinition::GetIntDef(), InName))
-	{
-	}
+	using FNiagaraEitherDataSetAccessor<FNiagaraDataConversions<FLinearColor>>::FNiagaraEitherDataSetAccessor;
 };
 
 template<>
-struct FNiagaraDataSetAccessor<FNiagaraDataConversions::FNiagaraBool> : public FNiagaraDataSetAccessor<FNiagaraBool>
+struct FNiagaraDataSetAccessor<FNiagaraDataConversions<FQuat>> : public FNiagaraEitherDataSetAccessor<FNiagaraDataConversions<FQuat>>
 {
-	FNiagaraDataSetAccessor<FNiagaraDataConversions::FNiagaraBool>(FNiagaraDataSet& InDataSet, const FName& InName) : FNiagaraDataSetAccessor<FNiagaraBool>(InDataSet, FNiagaraVariable(FNiagaraTypeDefinition::GetBoolDef(), InName))
-	{
-	}
+	using FNiagaraEitherDataSetAccessor<FNiagaraDataConversions<FQuat>>::FNiagaraEitherDataSetAccessor;
 };
 
-template<>
-struct FNiagaraDataSetAccessor<FNiagaraDataConversions::FNiagaraConvertID> : public FNiagaraDataSetAccessor<FNiagaraID>
+template<typename T>
+struct FNiagaraDataSetAccessor<FNiagaraDataConversions<T>> : public FNiagaraDataSetAccessor<T>
 {
-	FNiagaraDataSetAccessor<FNiagaraDataConversions::FNiagaraConvertID>(FNiagaraDataSet& InDataSet, const FName& InName) : FNiagaraDataSetAccessor<FNiagaraID>(InDataSet, FNiagaraVariable(FNiagaraTypeDefinition::GetIDDef(), InName))
-	{
-	}
+	using FNiagaraDataSetAccessor<T>::FNiagaraDataSetAccessor;
 };
 
 /**
diff --git a/Engine/Plugins/FX/Niagara/Source/Niagara/Private/NiagaraBoundsCalculatorHelper.h b/Engine/Plugins/FX/Niagara/Source/Niagara/Private/NiagaraBoundsCalculatorHelper.h
--- a/Engine/Plugins/FX/Niagara/Source/Niagara/Private/NiagaraBoundsCalculatorHelper.h
+++ b/Engine/Plugins/FX/Niagara/Source/Niagara/Private/NiagaraBoundsCalculatorHelper.h
@@ -24,38 +24,38 @@ public:
 
 		if (DataSet.HasVariable(PositionName))
 		{
-			PositionAccessor = FNiagaraDataSetAccessor<FNiagaraDataConversions::FHalf3OrFloat3>(DataSet, PositionName);
+			PositionAccessor = FNiagaraDataSetAccessor<FNiagaraDataConversions<FVector>>(DataSet, PositionName);
 		}
 		else
 		{
-			PositionAccessor = FNiagaraDataSetAccessor<FNiagaraDataConversions::FHalf3OrFloat3>();
+			PositionAccessor = FNiagaraDataSetAccessor<FNiagaraDataConversions<FVector>>();
 		}
 
 		if (bUsedWithSprites && DataSet.HasVariable(SpriteSizeName))
 		{
-			SpriteSizeAccessor = FNiagaraDataSetAccessor<FNiagaraDataConversions::FHalf2OrFloat2>(DataSet, SpriteSizeName);
+			SpriteSizeAccessor = FNiagaraDataSetAccessor<FNiagaraDataConversions<FVector2D>>(DataSet, SpriteSizeName);
 		}
 		else
 		{
-			SpriteSizeAccessor = FNiagaraDataSetAccessor<FNiagaraDataConversions::FHalf2OrFloat2>();
+			SpriteSizeAccessor = FNiagaraDataSetAccessor<FNiagaraDataConversions<FVector2D>>();
 		}
 
 		if (bUsedWithMeshes && DataSet.HasVariable(ScaleName))
 		{
-			ScaleAccessor = FNiagaraDataSetAccessor<FNiagaraDataConversions::FHalf3OrFloat3>(DataSet, ScaleName);
+			ScaleAccessor = FNiagaraDataSetAccessor<FNiagaraDataConversions<FVector>>(DataSet, ScaleName);
 		}
 		else
 		{
-			ScaleAccessor = FNiagaraDataSetAccessor<FNiagaraDataConversions::FHalf3OrFloat3>();
+			ScaleAccessor = FNiagaraDataSetAccessor<FNiagaraDataConversions<FVector>>();
 		}
 
 		if (bUsedWithRibbons && DataSet.HasVariable(RibbonWidthName))
 		{
-			RibbonWidthAccessor = FNiagaraDataSetAccessor<FNiagaraDataConversions::FHalfOrFloat>(DataSet, RibbonWidthName);
+			RibbonWidthAccessor = FNiagaraDataSetAccessor<FNiagaraDataConversions<float>>(DataSet, RibbonWidthName);
 		}
 		else
 		{
-			RibbonWidthAccessor = FNiagaraDataSetAccessor<FNiagaraDataConversions::FHalfOrFloat>();
+			RibbonWidthAccessor = FNiagaraDataSetAccessor<FNiagaraDataConversions<float>>();
 		}
 	}
 
@@ -137,9 +137,9 @@ public:
 		return FBox(PositionMin, PositionMax).ExpandBy(MaxSize);
 	}
 
-	FNiagaraDataSetAccessor<FNiagaraDataConversions::FHalf3OrFloat3> PositionAccessor;
-	FNiagaraDataSetAccessor<FNiagaraDataConversions::FHalf2OrFloat2> SpriteSizeAccessor;
-	FNiagaraDataSetAccessor<FNiagaraDataConversions::FHalf3OrFloat3> ScaleAccessor;
-	FNiagaraDataSetAccessor<FNiagaraDataConversions::FHalfOrFloat> RibbonWidthAccessor;
+	FNiagaraDataSetAccessor<FNiagaraDataConversions<FVector>> PositionAccessor;
+	FNiagaraDataSetAccessor<FNiagaraDataConversions<FVector2D>> SpriteSizeAccessor;
+	FNiagaraDataSetAccessor<FNiagaraDataConversions<FVector>> ScaleAccessor;
+	FNiagaraDataSetAccessor<FNiagaraDataConversions<float>> RibbonWidthAccessor;
 	const FVector MeshExtents = FVector::OneVector;
 };
diff --git a/Engine/Plugins/FX/Niagara/Source/Niagara/Private/NiagaraComponent.cpp b/Engine/Plugins/FX/Niagara/Source/Niagara/Private/NiagaraComponent.cpp
--- a/Engine/Plugins/FX/Niagara/Source/Niagara/Private/NiagaraComponent.cpp
+++ b/Engine/Plugins/FX/Niagara/Source/Niagara/Private/NiagaraComponent.cpp
@@ -1718,7 +1718,7 @@ TArray<float> UNiagaraComponent::GetNiagaraParticleValues_DebugOnly(const FStrin
 				FNiagaraDataBuffer& ParticleData = Sim->GetData().GetCurrentDataChecked();
 				int32 NumParticles = ParticleData.GetNumInstances();
 				Values.SetNum(NumParticles);
-				FNiagaraDataSetAccessor<FNiagaraDataConversions::FHalfOrFloat> ValueData(Sim->GetData(), *InValueName);
+				FNiagaraDataSetAccessor<FNiagaraDataConversions<float>> ValueData(Sim->GetData(), *InValueName);
 				for (int32 i = 0; i < NumParticles; ++i)
 				{
 					float Value;
diff --git a/Engine/Plugins/FX/Niagara/Source/Niagara/Private/NiagaraDataInterfaceParticleRead.cpp b/Engine/Plugins/FX/Niagara/Source/Niagara/Private/NiagaraDataInterfaceParticleRead.cpp
--- a/Engine/Plugins/FX/Niagara/Source/Niagara/Private/NiagaraDataInterfaceParticleRead.cpp
+++ b/Engine/Plugins/FX/Niagara/Source/Niagara/Private/NiagaraDataInterfaceParticleRead.cpp
@@ -629,6 +629,7 @@ void UNiagaraDataInterfaceParticleRead::GetFunctions(TArray<FNiagaraFunctionSign
 		Sig.Outputs.Add(FNiagaraVariable(FNiagaraTypeDefinition::GetIntDef(), TEXT("Num Spawned")));
 		Sig.bMemberFunction = true;
 		Sig.bRequiresContext = false;
+		Sig.bExperimental = true;
 		OutFunctions.Add(Sig);
 	}
 
@@ -644,6 +645,7 @@ void UNiagaraDataInterfaceParticleRead::GetFunctions(TArray<FNiagaraFunctionSign
 
 		Sig.bMemberFunction = true;
 		Sig.bRequiresContext = false;
+		Sig.bExperimental = true;
 		OutFunctions.Add(Sig);
 	}
 
@@ -665,6 +667,7 @@ void UNiagaraDataInterfaceParticleRead::GetFunctions(TArray<FNiagaraFunctionSign
 		Sig.Outputs.Add(FNiagaraVariable(FNiagaraTypeDefinition::GetIntDef(), TEXT("Index")));
 		Sig.bMemberFunction = true;
 		Sig.bRequiresContext = false;
+		Sig.bExperimental = true;
 		OutFunctions.Add(Sig);
 	}
 
@@ -684,6 +687,7 @@ void UNiagaraDataInterfaceParticleRead::GetFunctions(TArray<FNiagaraFunctionSign
 
 		Sig.bMemberFunction = true;
 		Sig.bRequiresContext = false;
+		Sig.bExperimental = true;
 		OutFunctions.Add(Sig);
 	}
 
@@ -700,6 +704,7 @@ void UNiagaraDataInterfaceParticleRead::GetFunctions(TArray<FNiagaraFunctionSign
 
 		Sig.bMemberFunction = true;
 		Sig.bRequiresContext = false;
+		Sig.bExperimental = true;
 		OutFunctions.Add(Sig);
 	}
 
@@ -716,6 +721,7 @@ void UNiagaraDataInterfaceParticleRead::GetFunctions(TArray<FNiagaraFunctionSign
 
 		Sig.bMemberFunction = true;
 		Sig.bRequiresContext = false;
+		Sig.bExperimental = true;
 		OutFunctions.Add(Sig);
 	}
 
@@ -732,6 +738,7 @@ void UNiagaraDataInterfaceParticleRead::GetFunctions(TArray<FNiagaraFunctionSign
 
 		Sig.bMemberFunction = true;
 		Sig.bRequiresContext = false;
+		Sig.bExperimental = true;
 		OutFunctions.Add(Sig);
 	}
 
@@ -748,6 +755,7 @@ void UNiagaraDataInterfaceParticleRead::GetFunctions(TArray<FNiagaraFunctionSign
 
 		Sig.bMemberFunction = true;
 		Sig.bRequiresContext = false;
+		Sig.bExperimental = true;
 		OutFunctions.Add(Sig);
 	}
 
@@ -764,6 +772,7 @@ void UNiagaraDataInterfaceParticleRead::GetFunctions(TArray<FNiagaraFunctionSign
 
 		Sig.bMemberFunction = true;
 		Sig.bRequiresContext = false;
+		Sig.bExperimental = true;
 		OutFunctions.Add(Sig);
 	}
 
@@ -780,6 +789,7 @@ void UNiagaraDataInterfaceParticleRead::GetFunctions(TArray<FNiagaraFunctionSign
 
 		Sig.bMemberFunction = true;
 		Sig.bRequiresContext = false;
+		Sig.bExperimental = true;
 		OutFunctions.Add(Sig);
 	}
 
@@ -796,6 +806,7 @@ void UNiagaraDataInterfaceParticleRead::GetFunctions(TArray<FNiagaraFunctionSign
 
 		Sig.bMemberFunction = true;
 		Sig.bRequiresContext = false;
+		Sig.bExperimental = true;
 		OutFunctions.Add(Sig);
 	}
 
@@ -812,6 +823,7 @@ void UNiagaraDataInterfaceParticleRead::GetFunctions(TArray<FNiagaraFunctionSign
 
 		Sig.bMemberFunction = true;
 		Sig.bRequiresContext = false;
+		Sig.bExperimental = true;
 		OutFunctions.Add(Sig);
 	}
 
@@ -831,6 +843,7 @@ void UNiagaraDataInterfaceParticleRead::GetFunctions(TArray<FNiagaraFunctionSign
 
 		Sig.bMemberFunction = true;
 		Sig.bRequiresContext = false;
+		Sig.bExperimental = true;
 		OutFunctions.Add(Sig);
 	}
 
@@ -847,6 +860,7 @@ void UNiagaraDataInterfaceParticleRead::GetFunctions(TArray<FNiagaraFunctionSign
 
 		Sig.bMemberFunction = true;
 		Sig.bRequiresContext = false;
+		Sig.bExperimental = true;
 		OutFunctions.Add(Sig);
 	}
 
@@ -863,6 +877,7 @@ void UNiagaraDataInterfaceParticleRead::GetFunctions(TArray<FNiagaraFunctionSign
 
 		Sig.bMemberFunction = true;
 		Sig.bRequiresContext = false;
+		Sig.bExperimental = true;
 		OutFunctions.Add(Sig);
 	}
 
@@ -879,6 +894,7 @@ void UNiagaraDataInterfaceParticleRead::GetFunctions(TArray<FNiagaraFunctionSign
 
 		Sig.bMemberFunction = true;
 		Sig.bRequiresContext = false;
+		Sig.bExperimental = true;
 		OutFunctions.Add(Sig);
 	}
 
@@ -895,6 +911,7 @@ void UNiagaraDataInterfaceParticleRead::GetFunctions(TArray<FNiagaraFunctionSign
 
 		Sig.bMemberFunction = true;
 		Sig.bRequiresContext = false;
+		Sig.bExperimental = true;
 		OutFunctions.Add(Sig);
 	}
 
@@ -911,6 +928,7 @@ void UNiagaraDataInterfaceParticleRead::GetFunctions(TArray<FNiagaraFunctionSign
 
 		Sig.bMemberFunction = true;
 		Sig.bRequiresContext = false;
+		Sig.bExperimental = true;
 		OutFunctions.Add(Sig);
 	}
 
@@ -927,6 +945,7 @@ void UNiagaraDataInterfaceParticleRead::GetFunctions(TArray<FNiagaraFunctionSign
 
 		Sig.bMemberFunction = true;
 		Sig.bRequiresContext = false;
+		Sig.bExperimental = true;
 		OutFunctions.Add(Sig);
 	}
 
@@ -943,6 +962,7 @@ void UNiagaraDataInterfaceParticleRead::GetFunctions(TArray<FNiagaraFunctionSign
 
 		Sig.bMemberFunction = true;
 		Sig.bRequiresContext = false;
+		Sig.bExperimental = true;
 		OutFunctions.Add(Sig);
 	} 
 
@@ -959,6 +979,7 @@ void UNiagaraDataInterfaceParticleRead::GetFunctions(TArray<FNiagaraFunctionSign
 
 		Sig.bMemberFunction = true;
 		Sig.bRequiresContext = false;
+		Sig.bExperimental = true;
 		OutFunctions.Add(Sig);
 	}
 }
@@ -1320,493 +1341,256 @@ void UNiagaraDataInterfaceParticleRead::GetParticleIndex(FVectorVMContext& Conte
 	}
 }
 
-void UNiagaraDataInterfaceParticleRead::ReadInt(FVectorVMContext& Context, FName AttributeToRead)
-{
-	VectorVM::FUserPtrHandler<FNDIParticleRead_InstanceData> InstanceData(Context);
-	VectorVM::FExternalFuncInputHandler<int32> ParticleIDIndexParam(Context);
-	VectorVM::FExternalFuncInputHandler<int32> ParticleIDAcquireTagParam(Context);
+//////////////////////////////////////////////////////////////////////////
 
-	VectorVM::FExternalFuncRegisterHandler<FNiagaraBool> OutValid(Context);
-	VectorVM::FExternalFuncRegisterHandler<int32> OutValue(Context);
-
-	for (int32 InstanceIdx = 0; InstanceIdx < Context.NumInstances; ++InstanceIdx)
+template<typename T>
+struct FDirectReadParamHandler
+{
+	FDirectReadParamHandler(FVectorVMContext& Context)
+		: InstanceData(Context)
+		, ParticleIDIndexParam(Context)
+		, ParticleIDAcquireTagParam(Context)
+		, OutValid(Context)
+		, OutValue(Context)
 	{
-		FNiagaraID ParticleID = { ParticleIDIndexParam.GetAndAdvance(), ParticleIDAcquireTagParam.GetAndAdvance() };
-		bool bValid;
-		int32 Value = RetrieveValueWithCheck<int, FNiagaraDataConversions::FNiagaraInt>(InstanceData->EmitterInstance, AttributeToRead, ParticleID, bValid);
-		FNiagaraBool ValidValue;
-		ValidValue.SetValue(bValid);
-		*OutValid.GetDestAndAdvance() = ValidValue;
-		*OutValue.GetDestAndAdvance() = Value;
 	}
-}
 
-void UNiagaraDataInterfaceParticleRead::ReadBool(FVectorVMContext& Context, FName AttributeToRead)
+	FORCEINLINE FNiagaraEmitterInstance* GetEmitterInstance() { return InstanceData->EmitterInstance; }
+	FORCEINLINE FNiagaraSystemInstance* GetSysteInstance() { return InstanceData->SystemInstance; }
+
+	FORCEINLINE FNiagaraID GetID() { return FNiagaraID(ParticleIDIndexParam.GetAndAdvance(), ParticleIDAcquireTagParam.GetAndAdvance()); }
+	FORCEINLINE void SetValid(bool bValid) { OutValid.SetAndAdvance(FNiagaraBool(bValid)); }
+	FORCEINLINE void SetValue(T Val) { OutValue.SetAndAdvance(Val); }
+
+	VectorVM::FUserPtrHandler<FNDIParticleRead_InstanceData> InstanceData;
+	VectorVM::FExternalFuncInputHandler<int32> ParticleIDIndexParam;
+	VectorVM::FExternalFuncInputHandler<int32> ParticleIDAcquireTagParam;
+	FNDIOutputParam<FNiagaraBool> OutValid;
+	FNDIOutputParam<T> OutValue;
+};
+
+template <typename T>
+FORCEINLINE void ReadWithCheck(FVectorVMContext& Context, FName AttributeToRead, T Default)
 {
-	VectorVM::FUserPtrHandler<FNDIParticleRead_InstanceData> InstanceData(Context);
-	VectorVM::FExternalFuncInputHandler<int32> ParticleIDIndexParam(Context);
-	VectorVM::FExternalFuncInputHandler<int32> ParticleIDAcquireTagParam(Context);
+	using AccessorType = FNiagaraDataConversions<T>;//TODO: Route template calls in here so we can use a regular accessor rather than a converatable accessor.
+	FDirectReadParamHandler<T> Params(Context);
+	if (FNiagaraEmitterInstance* EmitterInstance = Params.GetEmitterInstance())
+	{
+		const FNiagaraDataBuffer* CurrentData = EmitterInstance->GetData().GetCurrentData();//TODO: We should really be grabbing this in the instance data tick and adding a read ref to it.
+		if (CurrentData && EmitterInstance->GetGPUContext() == nullptr)
+		{
+			static const FNiagaraVariable IDVar(FNiagaraTypeDefinition::GetIDDef(), "ID");
+			const TArray<int32>& IDTable = CurrentData->GetIDTable();
+			int32 NumSourceInstances = CurrentData->GetNumInstances();
 
-	VectorVM::FExternalFuncRegisterHandler<FNiagaraBool> OutValid(Context);
-	VectorVM::FExternalFuncRegisterHandler<FNiagaraBool> OutValue(Context);
+			if (IDTable.Num() > 0)
+			{
+				FNiagaraVariable ReadVar(FNiagaraTypeDefinition::Get<T>(), AttributeToRead);
+				FNiagaraDataSetAccessor<AccessorType> ValueData(EmitterInstance->GetData(), ReadVar);
+
+				FNiagaraDataSetAccessor<FNiagaraID> IDData(EmitterInstance->GetData(), IDVar);
+
+				for (int32 InstanceIdx = 0; InstanceIdx < Context.NumInstances; ++InstanceIdx)
+				{
+					FNiagaraID ParticleID = Params.GetID();
+					bool bValid = false;
+					T Value = Default;
+
+					if (ParticleID.Index >= 0 && ParticleID.Index < IDTable.Num())
+					{
+						int32 ParticleIndex = IDTable[ParticleID.Index];
+						if (ParticleIndex >= 0 && ParticleIndex < NumSourceInstances)
+						{
+							FNiagaraID ActualID = IDData.GetSafe(ParticleIndex, NIAGARA_INVALID_ID);
+							if (ActualID == ParticleID)
+							{
+								Value = ValueData.GetSafe(ParticleIndex, T());
+								bValid = true;
+							}
+						}
+					}
+
+					Params.SetValid(bValid);
+					Params.SetValue(Value);
+				}
+			}
+
+			return;
+		}
+	}
 
+	//If we get there there's been a problem so just fill the output with default;
 	for (int32 InstanceIdx = 0; InstanceIdx < Context.NumInstances; ++InstanceIdx)
 	{
-		FNiagaraID ParticleID = { ParticleIDIndexParam.GetAndAdvance(), ParticleIDAcquireTagParam.GetAndAdvance() };
-		bool bValid;
-		FNiagaraBool Value = RetrieveValueWithCheck<FNiagaraBool, FNiagaraDataConversions::FNiagaraBool>(InstanceData->EmitterInstance, AttributeToRead, ParticleID, bValid);
-		FNiagaraBool ValidValue;
-		ValidValue.SetValue(bValid);
-		*OutValid.GetDestAndAdvance() = ValidValue;
-		*OutValue.GetDestAndAdvance() = Value;
+		Params.SetValid(false);
+		Params.SetValue(Default);
 	}
 }
 
-void UNiagaraDataInterfaceParticleRead::ReadFloat(FVectorVMContext& Context, FName AttributeToRead)
+//TODO: We can remove all this boiler plate and bind directly to ReadWithCheck.
+void UNiagaraDataInterfaceParticleRead::ReadInt(FVectorVMContext& Context, FName AttributeToRead)
 {
-	VectorVM::FUserPtrHandler<FNDIParticleRead_InstanceData> InstanceData(Context);
-	VectorVM::FExternalFuncInputHandler<int32> ParticleIDIndexParam(Context);
-	VectorVM::FExternalFuncInputHandler<int32> ParticleIDAcquireTagParam(Context);
-
-	VectorVM::FExternalFuncRegisterHandler<FNiagaraBool> OutValid(Context);
-	VectorVM::FExternalFuncRegisterHandler<float> OutValue(Context);
-	
-	for (int32 InstanceIdx = 0; InstanceIdx < Context.NumInstances; ++InstanceIdx)
-	{
-		FNiagaraID ParticleID = { ParticleIDIndexParam.GetAndAdvance(), ParticleIDAcquireTagParam.GetAndAdvance() };
-		bool bValid;
-		float Value = RetrieveValueWithCheck<float, FNiagaraDataConversions::FHalfOrFloat>(InstanceData->EmitterInstance, AttributeToRead, ParticleID, bValid);
-		FNiagaraBool ValidValue;
-		ValidValue.SetValue(bValid);
-		*OutValid.GetDestAndAdvance() = ValidValue;
-		*OutValue.GetDestAndAdvance() = Value;
-	}
+	ReadWithCheck<int32>(Context, AttributeToRead, 0);
 }
 
-void UNiagaraDataInterfaceParticleRead::ReadVector2(FVectorVMContext& Context, FName AttributeToRead)
+void UNiagaraDataInterfaceParticleRead::ReadBool(FVectorVMContext& Context, FName AttributeToRead)
 {
-	VectorVM::FUserPtrHandler<FNDIParticleRead_InstanceData> InstanceData(Context);
-	VectorVM::FExternalFuncInputHandler<int32> ParticleIDIndexParam(Context);
-	VectorVM::FExternalFuncInputHandler<int32> ParticleIDAcquireTagParam(Context);
+	ReadWithCheck<FNiagaraBool>(Context, AttributeToRead, FNiagaraBool(false));
+}
 
-	VectorVM::FExternalFuncRegisterHandler<FNiagaraBool> OutValid(Context);
-	VectorVM::FExternalFuncRegisterHandler<float> OutX(Context);
-	VectorVM::FExternalFuncRegisterHandler<float> OutY(Context);
+void UNiagaraDataInterfaceParticleRead::ReadFloat(FVectorVMContext& Context, FName AttributeToRead)
+{
+	ReadWithCheck<float>(Context, AttributeToRead, 0.0f);
+}
 
-	for (int32 InstanceIdx = 0; InstanceIdx < Context.NumInstances; ++InstanceIdx)
-	{
-		FNiagaraID ParticleID = { ParticleIDIndexParam.GetAndAdvance(), ParticleIDAcquireTagParam.GetAndAdvance() };
-		bool bValid;
-		FVector2D Value = RetrieveValueWithCheck<FVector2D, FNiagaraDataConversions::FHalf2OrFloat2>(InstanceData->EmitterInstance, AttributeToRead, ParticleID, bValid);
-		FNiagaraBool ValidValue;
-		ValidValue.SetValue(bValid);
-		*OutValid.GetDestAndAdvance() = ValidValue;
-		*OutX.GetDestAndAdvance() = Value.X;
-		*OutY.GetDestAndAdvance() = Value.Y;
-	}
+void UNiagaraDataInterfaceParticleRead::ReadVector2(FVectorVMContext& Context, FName AttributeToRead)
+{
+	ReadWithCheck<FVector2D>(Context, AttributeToRead, FVector2D::ZeroVector);
 }
 
 void UNiagaraDataInterfaceParticleRead::ReadVector3(FVectorVMContext& Context, FName AttributeToRead)
 {
-	VectorVM::FUserPtrHandler<FNDIParticleRead_InstanceData> InstanceData(Context);
-	VectorVM::FExternalFuncInputHandler<int32> ParticleIDIndexParam(Context);
-	VectorVM::FExternalFuncInputHandler<int32> ParticleIDAcquireTagParam(Context);
-
-	VectorVM::FExternalFuncRegisterHandler<FNiagaraBool> OutValid(Context);
-	VectorVM::FExternalFuncRegisterHandler<float> OutX(Context);
-	VectorVM::FExternalFuncRegisterHandler<float> OutY(Context);
-	VectorVM::FExternalFuncRegisterHandler<float> OutZ(Context);
-
-	for (int32 InstanceIdx = 0; InstanceIdx < Context.NumInstances; ++InstanceIdx)
-	{
-		FNiagaraID ParticleID = { ParticleIDIndexParam.GetAndAdvance(), ParticleIDAcquireTagParam.GetAndAdvance() };
-		bool bValid;
-		FVector Value = RetrieveValueWithCheck<FVector, FNiagaraDataConversions::FHalf3OrFloat3>(InstanceData->EmitterInstance, AttributeToRead, ParticleID, bValid);
-		FNiagaraBool ValidValue;
-		ValidValue.SetValue(bValid);
-		*OutValid.GetDestAndAdvance() = ValidValue;
-		*OutX.GetDestAndAdvance() = Value.X;
-		*OutY.GetDestAndAdvance() = Value.Y;
-		*OutZ.GetDestAndAdvance() = Value.Z;
-	}
+	ReadWithCheck<FVector>(Context, AttributeToRead, FVector::ZeroVector);
 }
 
 void UNiagaraDataInterfaceParticleRead::ReadVector4(FVectorVMContext& Context, FName AttributeToRead)
 {
-	VectorVM::FUserPtrHandler<FNDIParticleRead_InstanceData> InstanceData(Context);
-	VectorVM::FExternalFuncInputHandler<int32> ParticleIDIndexParam(Context);
-	VectorVM::FExternalFuncInputHandler<int32> ParticleIDAcquireTagParam(Context);
-
-	VectorVM::FExternalFuncRegisterHandler<FNiagaraBool> OutValid(Context);
-	VectorVM::FExternalFuncRegisterHandler<float> OutX(Context);
-	VectorVM::FExternalFuncRegisterHandler<float> OutY(Context);
-	VectorVM::FExternalFuncRegisterHandler<float> OutZ(Context);
-	VectorVM::FExternalFuncRegisterHandler<float> OutW(Context);
-
-	for (int32 InstanceIdx = 0; InstanceIdx < Context.NumInstances; ++InstanceIdx)
-	{
-		FNiagaraID ParticleID = { ParticleIDIndexParam.GetAndAdvance(), ParticleIDAcquireTagParam.GetAndAdvance() };
-		bool bValid;
-		FVector4 Value = RetrieveValueWithCheck<FVector4, FNiagaraDataConversions::FHalf4OrFloat4>(InstanceData->EmitterInstance, AttributeToRead, ParticleID, bValid);
-		FNiagaraBool ValidValue;
-		ValidValue.SetValue(bValid);
-		*OutValid.GetDestAndAdvance() = ValidValue;
-		*OutX.GetDestAndAdvance() = Value.X;
-		*OutY.GetDestAndAdvance() = Value.Y;
-		*OutZ.GetDestAndAdvance() = Value.Z;
-		*OutW.GetDestAndAdvance() = Value.W;
-	}
+	ReadWithCheck<FVector4>(Context, AttributeToRead, FVector4(ForceInit));
 }
 
 void UNiagaraDataInterfaceParticleRead::ReadColor(FVectorVMContext& Context, FName AttributeToRead)
 {
-	VectorVM::FUserPtrHandler<FNDIParticleRead_InstanceData> InstanceData(Context);
-	VectorVM::FExternalFuncInputHandler<int32> ParticleIDIndexParam(Context);
-	VectorVM::FExternalFuncInputHandler<int32> ParticleIDAcquireTagParam(Context);
-
-	VectorVM::FExternalFuncRegisterHandler<FNiagaraBool> OutValid(Context);
-	VectorVM::FExternalFuncRegisterHandler<float> OutR(Context);
-	VectorVM::FExternalFuncRegisterHandler<float> OutG(Context);
-	VectorVM::FExternalFuncRegisterHandler<float> OutB(Context);
-	VectorVM::FExternalFuncRegisterHandler<float> OutA(Context);
-
-	for (int32 InstanceIdx = 0; InstanceIdx < Context.NumInstances; ++InstanceIdx)
-	{
-		FNiagaraID ParticleID = { ParticleIDIndexParam.GetAndAdvance(), ParticleIDAcquireTagParam.GetAndAdvance() };
-		bool bValid;
-		FLinearColor Value = RetrieveValueWithCheck<FLinearColor, FNiagaraDataConversions::FHalf4OrColor4>(InstanceData->EmitterInstance, AttributeToRead, ParticleID, bValid);
-		FNiagaraBool ValidValue;
-		ValidValue.SetValue(bValid);
-		*OutValid.GetDestAndAdvance() = ValidValue;
-		*OutR.GetDestAndAdvance() = Value.R;
-		*OutG.GetDestAndAdvance() = Value.G;
-		*OutB.GetDestAndAdvance() = Value.B;
-		*OutA.GetDestAndAdvance() = Value.A;
-	}
+	ReadWithCheck<FLinearColor>(Context, AttributeToRead, FLinearColor::White);
 }
 
 void UNiagaraDataInterfaceParticleRead::ReadQuat(FVectorVMContext& Context, FName AttributeToRead)
 {
-	VectorVM::FUserPtrHandler<FNDIParticleRead_InstanceData> InstanceData(Context);
-	VectorVM::FExternalFuncInputHandler<int32> ParticleIDIndexParam(Context);
-	VectorVM::FExternalFuncInputHandler<int32> ParticleIDAcquireTagParam(Context);
-
-	VectorVM::FExternalFuncRegisterHandler<FNiagaraBool> OutValid(Context);
-	VectorVM::FExternalFuncRegisterHandler<float> OutX(Context);
-	VectorVM::FExternalFuncRegisterHandler<float> OutY(Context);
-	VectorVM::FExternalFuncRegisterHandler<float> OutZ(Context);
-	VectorVM::FExternalFuncRegisterHandler<float> OutW(Context);
-
-	for (int32 InstanceIdx = 0; InstanceIdx < Context.NumInstances; ++InstanceIdx)
-	{
-		FNiagaraID ParticleID = { ParticleIDIndexParam.GetAndAdvance(), ParticleIDAcquireTagParam.GetAndAdvance() };
-		bool bValid;
-		FQuat Value = RetrieveValueWithCheck<FQuat, FNiagaraDataConversions::FHalf4OrQuat4>(InstanceData->EmitterInstance, AttributeToRead, ParticleID, bValid);
-		FNiagaraBool ValidValue;
-		ValidValue.SetValue(bValid);
-		*OutValid.GetDestAndAdvance() = ValidValue;
-		*OutX.GetDestAndAdvance() = Value.X;
-		*OutY.GetDestAndAdvance() = Value.Y;
-		*OutZ.GetDestAndAdvance() = Value.Z;
-		*OutW.GetDestAndAdvance() = Value.W;
-	}
+	ReadWithCheck<FQuat>(Context, AttributeToRead, FQuat::Identity);
 }
 
 void UNiagaraDataInterfaceParticleRead::ReadID(FVectorVMContext& Context, FName AttributeToRead)
 {
-	VectorVM::FUserPtrHandler<FNDIParticleRead_InstanceData> InstanceData(Context);
-	VectorVM::FExternalFuncInputHandler<int32> ParticleIDIndexParam(Context);
-	VectorVM::FExternalFuncInputHandler<int32> ParticleIDAcquireTagParam(Context);
+	ReadWithCheck<FNiagaraID>(Context, AttributeToRead, FNiagaraID());
+}
 
-	VectorVM::FExternalFuncRegisterHandler<FNiagaraBool> OutValid(Context);
-	VectorVM::FExternalFuncRegisterHandler<int> OutIDIndex(Context);
-	VectorVM::FExternalFuncRegisterHandler<int> OutIDAcquireTag(Context);
+//////////////////////////////////////////////////////////////////////////
 
-	for (int32 InstanceIdx = 0; InstanceIdx < Context.NumInstances; ++InstanceIdx)
+template<typename T>
+struct FDirectReadByIndexParamHandler
+{
+	FDirectReadByIndexParamHandler(FVectorVMContext& Context)
+		: InstanceData(Context)
+		, ParticleIndexParam(Context)
+		, OutValid(Context)
+		, OutValue(Context)
 	{
-		FNiagaraID ParticleID = { ParticleIDIndexParam.GetAndAdvance(), ParticleIDAcquireTagParam.GetAndAdvance() };
-		bool bValid;
+	}
 
-		FNiagaraID Value = RetrieveValueWithCheck<FNiagaraID, FNiagaraDataConversions::FNiagaraConvertID>(InstanceData->EmitterInstance, AttributeToRead, ParticleID, bValid);
+	FORCEINLINE FNiagaraEmitterInstance* GetEmitterInstance() { return InstanceData->EmitterInstance; }
+	FORCEINLINE FNiagaraSystemInstance* GetSysteInstance() { return InstanceData->SystemInstance; }
 
-		FNiagaraBool ValidValue;
-		ValidValue.SetValue(bValid);
-		*OutValid.GetDestAndAdvance() = ValidValue;
-		*OutIDIndex.GetDestAndAdvance() = Value.Index;
-		*OutIDAcquireTag.GetDestAndAdvance() = Value.AcquireTag;
-	}
-}
+	FORCEINLINE int32 GetIndex() { return ParticleIndexParam.GetAndAdvance(); }
+	FORCEINLINE void SetValid(bool bValid) { OutValid.SetAndAdvance(FNiagaraBool(bValid)); }
+	FORCEINLINE void SetValue(T Val) { OutValue.SetAndAdvance(Val); }
+
+	VectorVM::FUserPtrHandler<FNDIParticleRead_InstanceData> InstanceData;
+	VectorVM::FExternalFuncInputHandler<int32> ParticleIndexParam;
+	FNDIOutputParam<FNiagaraBool> OutValid;
+	FNDIOutputParam<T> OutValue;
+};
 
-template <typename Type, typename AccessorType>
-Type UNiagaraDataInterfaceParticleRead::RetrieveValueWithCheck(FNiagaraEmitterInstance* EmitterInstance, const FName& Attr, const FNiagaraID& ParticleID, bool& bValid)
+template<typename T>
+FORCEINLINE void ReadByIndexWithCheck(FVectorVMContext& Context, FName AttributeToRead, T Default)
 {
-	Type Value = Type();
-	bValid = false;
+	using AccessorType = FNiagaraDataConversions<T>;//TODO: Route template calls in here so we can use a regular accessor rather than a converatable accessor.
+	FDirectReadByIndexParamHandler<T> Params(Context);
 
-	const FNiagaraDataBuffer* CurrentData = EmitterInstance->GetData().GetCurrentData();
-	if (!CurrentData)
+	if (FNiagaraEmitterInstance* EmitterInstance = Params.GetEmitterInstance())
 	{
-		return Value;
-	}
+		const FNiagaraDataBuffer* CurrentData = EmitterInstance->GetData().GetCurrentData();//TODO: We should really be grabbing these during instance data tick and adding a read ref. Releasing that on PostTick.
+		if (CurrentData && CurrentData->GetNumInstances() > 0 && EmitterInstance->GetGPUContext() == nullptr)
+		{
+			static const FNiagaraVariable IDVar(FNiagaraTypeDefinition::GetIDDef(), "ID");
+			const TArray<int32>& IDTable = CurrentData->GetIDTable();
+			int32 NumSourceInstances = (int32)CurrentData->GetNumInstances();
 
-	const TArray<int32>& IDTable = CurrentData->GetIDTable();
-	if (ParticleID.Index < 0 || ParticleID.Index >= IDTable.Num())
-	{
-		return Value;
-	}
+			FNiagaraVariable ReadVar(FNiagaraTypeDefinition::Get<T>(), AttributeToRead);
+			FNiagaraDataSetAccessor<AccessorType> ValueData(EmitterInstance->GetData(), ReadVar);
 
-	FNiagaraDataSetAccessor<AccessorType> ValueData(EmitterInstance->GetData(), Attr);
+			FNiagaraDataSetAccessor<FNiagaraID> IDData(EmitterInstance->GetData(), IDVar);
 
-	FNiagaraVariable IDVar(FNiagaraTypeDefinition::GetIDDef(), "ID");
-	FNiagaraDataSetAccessor<FNiagaraID> IDData(EmitterInstance->GetData(), IDVar);
+			for (int32 InstanceIdx = 0; InstanceIdx < Context.NumInstances; ++InstanceIdx)
+			{
+				int32 ParticleIndex = Params.GetIndex();
 
-	int32 ParticleIndex = IDTable[ParticleID.Index];
-	if (ParticleIndex >= 0)
-	{
-		FNiagaraID ActualID = IDData.GetSafe(ParticleIndex, NIAGARA_INVALID_ID);
-		if (ActualID == ParticleID)
-		{
-			Value = ValueData.GetSafe(ParticleIndex, Type());
-			bValid = true;
+				T Value = Default;
+				bool bValid = false;
+				if (ParticleIndex >= 0 && ParticleIndex < NumSourceInstances)
+				{
+					Value = ValueData.GetSafe(ParticleIndex, Default);
+					bValid = true;					
+				}
+
+				Params.SetValid(bValid);
+				Params.SetValue(Value);
+			}
+
+			return;
 		}
 	}
 
-	return Value;
+	//If we get there there's been a problem so just fill the output with default;
+	for (int32 InstanceIdx = 0; InstanceIdx < Context.NumInstances; ++InstanceIdx)
+	{
+		Params.SetValid(false);
+		Params.SetValue(Default);
+	}
 }
 
+//TODO: We can remove all this boiler plate and bind directly to ReadByIndexWithCheck.
 void UNiagaraDataInterfaceParticleRead::ReadIntByIndex(FVectorVMContext& Context, FName AttributeToRead)
 {
-	VectorVM::FUserPtrHandler<FNDIParticleRead_InstanceData> InstanceData(Context);
-	VectorVM::FExternalFuncInputHandler<int32> ParticleIndexParam(Context);
-
-	VectorVM::FExternalFuncRegisterHandler<FNiagaraBool> OutValid(Context);
-	VectorVM::FExternalFuncRegisterHandler<int32> OutValue(Context);
-
-	for (int32 InstanceIdx = 0; InstanceIdx < Context.NumInstances; ++InstanceIdx)
-	{
-		int32 ParticleIndex = ParticleIndexParam.GetAndAdvance();
-		bool bValid;
-		int32 Value = RetrieveValueByIndexWithCheck<int32, FNiagaraDataConversions::FNiagaraInt>(InstanceData->EmitterInstance, AttributeToRead, ParticleIndex, bValid);
-		FNiagaraBool ValidValue;
-		ValidValue.SetValue(bValid);
-		*OutValid.GetDestAndAdvance() = ValidValue;
-		*OutValue.GetDestAndAdvance() = Value;
-	}
+	ReadByIndexWithCheck<int32>(Context, AttributeToRead, 0);
 }
 
 void UNiagaraDataInterfaceParticleRead::ReadBoolByIndex(FVectorVMContext& Context, FName AttributeToRead)
 {
-	VectorVM::FUserPtrHandler<FNDIParticleRead_InstanceData> InstanceData(Context);
-	VectorVM::FExternalFuncInputHandler<int32> ParticleIndexParam(Context);
-
-	VectorVM::FExternalFuncRegisterHandler<FNiagaraBool> OutValid(Context);
-	VectorVM::FExternalFuncRegisterHandler<FNiagaraBool> OutValue(Context);
-
-	for (int32 InstanceIdx = 0; InstanceIdx < Context.NumInstances; ++InstanceIdx)
-	{
-		int32 ParticleIndex = ParticleIndexParam.GetAndAdvance();
-		bool bValid;
-		FNiagaraBool Value = RetrieveValueByIndexWithCheck<FNiagaraBool, FNiagaraDataConversions::FNiagaraBool>(InstanceData->EmitterInstance, AttributeToRead, ParticleIndex, bValid);
-		FNiagaraBool ValidValue;
-		ValidValue.SetValue(bValid);
-		*OutValid.GetDestAndAdvance() = ValidValue;
-		*OutValue.GetDestAndAdvance() = Value;
-	}
+	ReadByIndexWithCheck<FNiagaraBool>(Context, AttributeToRead, FNiagaraBool(false));
 }
 
 void UNiagaraDataInterfaceParticleRead::ReadFloatByIndex(FVectorVMContext& Context, FName AttributeToRead)
 {
-	VectorVM::FUserPtrHandler<FNDIParticleRead_InstanceData> InstanceData(Context);
-	VectorVM::FExternalFuncInputHandler<int32> ParticleIndexParam(Context);
-
-	VectorVM::FExternalFuncRegisterHandler<FNiagaraBool> OutValid(Context);
-	VectorVM::FExternalFuncRegisterHandler<float> OutValue(Context);
-
-	for (int32 InstanceIdx = 0; InstanceIdx < Context.NumInstances; ++InstanceIdx)
-	{
-		int32 ParticleIndex = ParticleIndexParam.GetAndAdvance();
-		bool bValid;
-		float Value = RetrieveValueByIndexWithCheck<float, FNiagaraDataConversions::FHalfOrFloat>(InstanceData->EmitterInstance, AttributeToRead, ParticleIndex, bValid);
-		FNiagaraBool ValidValue;
-		ValidValue.SetValue(bValid);
-		*OutValid.GetDestAndAdvance() = ValidValue;
-		*OutValue.GetDestAndAdvance() = Value;
-	}
+	ReadByIndexWithCheck<float>(Context, AttributeToRead, 0.0f);
 }
 
 void UNiagaraDataInterfaceParticleRead::ReadVector2ByIndex(FVectorVMContext& Context, FName AttributeToRead)
 {
-	VectorVM::FUserPtrHandler<FNDIParticleRead_InstanceData> InstanceData(Context);
-	VectorVM::FExternalFuncInputHandler<int32> ParticleIndexParam(Context);
-
-	VectorVM::FExternalFuncRegisterHandler<FNiagaraBool> OutValid(Context);
-	VectorVM::FExternalFuncRegisterHandler<float> OutX(Context);
-	VectorVM::FExternalFuncRegisterHandler<float> OutY(Context);
-
-	for (int32 InstanceIdx = 0; InstanceIdx < Context.NumInstances; ++InstanceIdx)
-	{
-		int32 ParticleIndex = ParticleIndexParam.GetAndAdvance();
-		bool bValid;
-		FVector2D Value = RetrieveValueByIndexWithCheck<FVector2D, FNiagaraDataConversions::FHalf2OrFloat2>(InstanceData->EmitterInstance, AttributeToRead, ParticleIndex, bValid);
-		FNiagaraBool ValidValue;
-		ValidValue.SetValue(bValid);
-		*OutValid.GetDestAndAdvance() = ValidValue;
-		*OutX.GetDestAndAdvance() = Value.X;
-		*OutY.GetDestAndAdvance() = Value.Y;
-	}
+	ReadByIndexWithCheck<FVector2D>(Context, AttributeToRead, FVector2D::ZeroVector);
 }
 
 void UNiagaraDataInterfaceParticleRead::ReadVector3ByIndex(FVectorVMContext& Context, FName AttributeToRead)
 {
-	VectorVM::FUserPtrHandler<FNDIParticleRead_InstanceData> InstanceData(Context);
-	VectorVM::FExternalFuncInputHandler<int32> ParticleIndexParam(Context);
-
-	VectorVM::FExternalFuncRegisterHandler<FNiagaraBool> OutValid(Context);
-	VectorVM::FExternalFuncRegisterHandler<float> OutX(Context);
-	VectorVM::FExternalFuncRegisterHandler<float> OutY(Context);
-	VectorVM::FExternalFuncRegisterHandler<float> OutZ(Context);
-
-	for (int32 InstanceIdx = 0; InstanceIdx < Context.NumInstances; ++InstanceIdx)
-	{
-		int32 ParticleIndex = ParticleIndexParam.GetAndAdvance();
-		bool bValid;
-		FVector Value = RetrieveValueByIndexWithCheck<FVector, FNiagaraDataConversions::FHalf3OrFloat3>(InstanceData->EmitterInstance, AttributeToRead, ParticleIndex, bValid);
-		FNiagaraBool ValidValue;
-		ValidValue.SetValue(bValid);
-		*OutValid.GetDestAndAdvance() = ValidValue;
-		*OutX.GetDestAndAdvance() = Value.X;
-		*OutY.GetDestAndAdvance() = Value.Y;
-		*OutZ.GetDestAndAdvance() = Value.Z;
-	}
+	ReadByIndexWithCheck<FVector>(Context, AttributeToRead, FVector::ZeroVector);
 }
 
 void UNiagaraDataInterfaceParticleRead::ReadVector4ByIndex(FVectorVMContext& Context, FName AttributeToRead)
 {
-	VectorVM::FUserPtrHandler<FNDIParticleRead_InstanceData> InstanceData(Context);
-	VectorVM::FExternalFuncInputHandler<int32> ParticleIndexParam(Context);
-
-	VectorVM::FExternalFuncRegisterHandler<FNiagaraBool> OutValid(Context);
-	VectorVM::FExternalFuncRegisterHandler<float> OutX(Context);
-	VectorVM::FExternalFuncRegisterHandler<float> OutY(Context);
-	VectorVM::FExternalFuncRegisterHandler<float> OutZ(Context);
-	VectorVM::FExternalFuncRegisterHandler<float> OutW(Context);
-
-	for (int32 InstanceIdx = 0; InstanceIdx < Context.NumInstances; ++InstanceIdx)
-	{
-		int32 ParticleIndex = ParticleIndexParam.GetAndAdvance();
-		bool bValid;
-		FVector4 Value = RetrieveValueByIndexWithCheck<FVector4, FNiagaraDataConversions::FHalf4OrFloat4>(InstanceData->EmitterInstance, AttributeToRead, ParticleIndex, bValid);
-		FNiagaraBool ValidValue;
-		ValidValue.SetValue(bValid);
-		*OutValid.GetDestAndAdvance() = ValidValue;
-		*OutX.GetDestAndAdvance() = Value.X;
-		*OutY.GetDestAndAdvance() = Value.Y;
-		*OutZ.GetDestAndAdvance() = Value.Z;
-		*OutW.GetDestAndAdvance() = Value.W;
-	}
+	ReadByIndexWithCheck<FVector4>(Context, AttributeToRead, FVector4(ForceInit));
 }
 
 void UNiagaraDataInterfaceParticleRead::ReadColorByIndex(FVectorVMContext& Context, FName AttributeToRead)
 {
-	VectorVM::FUserPtrHandler<FNDIParticleRead_InstanceData> InstanceData(Context);
-	VectorVM::FExternalFuncInputHandler<int32> ParticleIndexParam(Context);
-
-	VectorVM::FExternalFuncRegisterHandler<FNiagaraBool> OutValid(Context);
-	VectorVM::FExternalFuncRegisterHandler<float> OutR(Context);
-	VectorVM::FExternalFuncRegisterHandler<float> OutG(Context);
-	VectorVM::FExternalFuncRegisterHandler<float> OutB(Context);
-	VectorVM::FExternalFuncRegisterHandler<float> OutA(Context);
-
-	for (int32 InstanceIdx = 0; InstanceIdx < Context.NumInstances; ++InstanceIdx)
-	{
-		int32 ParticleIndex = ParticleIndexParam.GetAndAdvance();
-		bool bValid;
-		FLinearColor Value = RetrieveValueByIndexWithCheck<FLinearColor, FNiagaraDataConversions::FHalf4OrColor4>(InstanceData->EmitterInstance, AttributeToRead, ParticleIndex, bValid);
-		FNiagaraBool ValidValue;
-		ValidValue.SetValue(bValid);
-		*OutValid.GetDestAndAdvance() = ValidValue;
-		*OutR.GetDestAndAdvance() = Value.R;
-		*OutG.GetDestAndAdvance() = Value.G;
-		*OutB.GetDestAndAdvance() = Value.B;
-		*OutA.GetDestAndAdvance() = Value.A;
-	}
+	ReadByIndexWithCheck<FLinearColor>(Context, AttributeToRead, FLinearColor::White);
 }
 
 void UNiagaraDataInterfaceParticleRead::ReadQuatByIndex(FVectorVMContext& Context, FName AttributeToRead)
 {
-	VectorVM::FUserPtrHandler<FNDIParticleRead_InstanceData> InstanceData(Context);
-	VectorVM::FExternalFuncInputHandler<int32> ParticleIndexParam(Context);
-
-	VectorVM::FExternalFuncRegisterHandler<FNiagaraBool> OutValid(Context);
-	VectorVM::FExternalFuncRegisterHandler<float> OutX(Context);
-	VectorVM::FExternalFuncRegisterHandler<float> OutY(Context);
-	VectorVM::FExternalFuncRegisterHandler<float> OutZ(Context);
-	VectorVM::FExternalFuncRegisterHandler<float> OutW(Context);
-
-	for (int32 InstanceIdx = 0; InstanceIdx < Context.NumInstances; ++InstanceIdx)
-	{
-		int32 ParticleIndex = ParticleIndexParam.GetAndAdvance();
-		bool bValid;
-		FQuat Value = RetrieveValueByIndexWithCheck<FQuat, FNiagaraDataConversions::FHalf4OrQuat4>(InstanceData->EmitterInstance, AttributeToRead, ParticleIndex, bValid);
-		FNiagaraBool ValidValue;
-		ValidValue.SetValue(bValid);
-		*OutValid.GetDestAndAdvance() = ValidValue;
-		*OutX.GetDestAndAdvance() = Value.X;
-		*OutY.GetDestAndAdvance() = Value.Y;
-		*OutZ.GetDestAndAdvance() = Value.Z;
-		*OutW.GetDestAndAdvance() = Value.W;
-	}
+	ReadByIndexWithCheck<FQuat>(Context, AttributeToRead, FQuat::Identity);
 }
 
 void UNiagaraDataInterfaceParticleRead::ReadIDByIndex(FVectorVMContext& Context, FName AttributeToRead)
 {
-	VectorVM::FUserPtrHandler<FNDIParticleRead_InstanceData> InstanceData(Context);
-	VectorVM::FExternalFuncInputHandler<int32> ParticleIndexParam(Context);
-
-	VectorVM::FExternalFuncRegisterHandler<FNiagaraBool> OutValid(Context);
-	VectorVM::FExternalFuncRegisterHandler<int32> OutIDIndex(Context);
-	VectorVM::FExternalFuncRegisterHandler<int32> OutIDAcquireTag(Context);
-
-	for (int32 InstanceIdx = 0; InstanceIdx < Context.NumInstances; ++InstanceIdx)
-	{
-		int32 ParticleIndex = ParticleIndexParam.GetAndAdvance();
-		bool bValid;
-		FNiagaraID Value = RetrieveValueByIndexWithCheck<FNiagaraID, FNiagaraDataConversions::FNiagaraConvertID>(InstanceData->EmitterInstance, AttributeToRead, ParticleIndex, bValid);
-		FNiagaraBool ValidValue;
-		ValidValue.SetValue(bValid);
-		*OutValid.GetDestAndAdvance() = ValidValue;
-		*OutIDIndex.GetDestAndAdvance() = Value.Index;
-		*OutIDAcquireTag.GetDestAndAdvance() = Value.AcquireTag;
-	}
-}
-
-template <typename Type, typename AccessorType>
-Type UNiagaraDataInterfaceParticleRead::RetrieveValueByIndexWithCheck(FNiagaraEmitterInstance* EmitterInstance, const FName& Attr, int32 ParticleIndex, bool& bValid)
-{
-	Type Value = Type();
-	bValid = false;
-
-	const FNiagaraDataBuffer* CurrentData = EmitterInstance->GetData().GetCurrentData();
-	if (!CurrentData)
-	{
-		return Value;
-	}
-
-	//-TODO: Investigate removing this in favor of making GetSafe really a safe call and handle GPU buffers with no CPU side data
-	if (EmitterInstance->GetGPUContext() != nullptr)
-	{
-		return Value;
-	}
-
-	FNiagaraDataSetAccessor<AccessorType> ValueData(EmitterInstance->GetData(), Attr);
-
-	if (ParticleIndex >= 0 && ParticleIndex < (int32)CurrentData->GetNumInstances())
-	{
-		Value = ValueData.GetSafe(ParticleIndex, Type());
-		bValid = true;
-	}
-
-	return Value;
+	ReadByIndexWithCheck<FNiagaraID>(Context, AttributeToRead, FNiagaraID());
 }
 
 bool UNiagaraDataInterfaceParticleRead::Equals(const UNiagaraDataInterface* Other) const
diff --git a/Engine/Plugins/FX/Niagara/Source/Niagara/Private/NiagaraRendererLights.cpp b/Engine/Plugins/FX/Niagara/Source/Niagara/Private/NiagaraRendererLights.cpp
--- a/Engine/Plugins/FX/Niagara/Source/Niagara/Private/NiagaraRendererLights.cpp
+++ b/Engine/Plugins/FX/Niagara/Source/Niagara/Private/NiagaraRendererLights.cpp
@@ -61,25 +61,25 @@ FNiagaraDynamicDataBase* FNiagaraRendererLights::GenerateDynamicData(const FNiag
 	//I'm not a great fan of pulling scalar components out to a structured vert buffer like this.
 	//TODO: Experiment with a new VF that reads the data directly from the scalar layout.
 	FNiagaraDataSetAccessor<FVector> PosAccessor(Data, Properties->PositionBinding.DataSetVariable);
-	FNiagaraDataSetAccessor<FNiagaraDataConversions::FHalf4OrColor4> ColAccessor;
+	FNiagaraDataSetAccessor<FNiagaraDataConversions<FLinearColor>> ColAccessor;
 	if (Data.HasVariable(Properties->ColorBinding.DataSetVariable.GetName()))
 	{
-		ColAccessor = FNiagaraDataSetAccessor<FNiagaraDataConversions::FHalf4OrColor4>(Data, Properties->ColorBinding.DataSetVariable.GetName());
+		ColAccessor = FNiagaraDataSetAccessor<FNiagaraDataConversions<FLinearColor>>(Data, Properties->ColorBinding.DataSetVariable.GetName());
 	}
-	FNiagaraDataSetAccessor<FNiagaraDataConversions::FHalfOrFloat> RadiusAccessor;
+	FNiagaraDataSetAccessor<FNiagaraDataConversions<float>> RadiusAccessor;
 	if (Data.HasVariable(Properties->RadiusBinding.DataSetVariable.GetName()))
 	{
-		RadiusAccessor = FNiagaraDataSetAccessor<FNiagaraDataConversions::FHalfOrFloat>(Data, Properties->RadiusBinding.DataSetVariable.GetName());
+		RadiusAccessor = FNiagaraDataSetAccessor<FNiagaraDataConversions<float>>(Data, Properties->RadiusBinding.DataSetVariable.GetName());
 	}
-	FNiagaraDataSetAccessor<FNiagaraDataConversions::FHalfOrFloat> ExponentAccessor;
+	FNiagaraDataSetAccessor<FNiagaraDataConversions<float>> ExponentAccessor;
 	if (Data.HasVariable(Properties->LightExponentBinding.DataSetVariable.GetName()))
 	{
-		ExponentAccessor = FNiagaraDataSetAccessor<FNiagaraDataConversions::FHalfOrFloat>(Data, Properties->LightExponentBinding.DataSetVariable.GetName());
+		ExponentAccessor = FNiagaraDataSetAccessor<FNiagaraDataConversions<float>>(Data, Properties->LightExponentBinding.DataSetVariable.GetName());
 	}
-	FNiagaraDataSetAccessor<FNiagaraDataConversions::FHalfOrFloat> ScatteringAccessor;
+	FNiagaraDataSetAccessor<FNiagaraDataConversions<float>> ScatteringAccessor;
 	if (Data.HasVariable(Properties->VolumetricScatteringBinding.DataSetVariable.GetName()))
 	{
-		ScatteringAccessor = FNiagaraDataSetAccessor<FNiagaraDataConversions::FHalfOrFloat>(Data, Properties->VolumetricScatteringBinding.DataSetVariable.GetName());
+		ScatteringAccessor = FNiagaraDataSetAccessor<FNiagaraDataConversions<float>>(Data, Properties->VolumetricScatteringBinding.DataSetVariable.GetName());
 	}
 	FNiagaraDataSetAccessor<int32> EnabledAccessor;
 	if (Data.HasVariable(Properties->LightRenderingEnabledBinding.DataSetVariable.GetName()))
diff --git a/Engine/Plugins/FX/Niagara/Source/Niagara/Private/NiagaraRendererRibbons.cpp b/Engine/Plugins/FX/Niagara/Source/Niagara/Private/NiagaraRendererRibbons.cpp
--- a/Engine/Plugins/FX/Niagara/Source/Niagara/Private/NiagaraRendererRibbons.cpp
+++ b/Engine/Plugins/FX/Niagara/Source/Niagara/Private/NiagaraRendererRibbons.cpp
@@ -401,7 +401,7 @@ int FNiagaraRendererRibbons::GetDynamicDataSize()const
 }
 
 void CalculateUVScaleAndOffsets(
-	const FNiagaraDataSetAccessor<FNiagaraDataConversions::FHalfOrFloat>& SortKeyData, const TArray<int32>& RibbonIndices,
+	const FNiagaraDataSetAccessor<FNiagaraDataConversions<float>>& SortKeyData, const TArray<int32>& RibbonIndices,
 	bool bSortKeyIsAge, int32 StartIndex, int32 EndIndex, int32 NumSegments,
 	float InUTilingDistance, float InUScale, float InUOffset, ENiagaraRibbonAgeOffsetMode InAgeOffsetMode,
 	float& OutUScale, float& OutUOffset)
@@ -467,48 +467,48 @@ FNiagaraDynamicDataBase* FNiagaraRendererRibbons::GenerateDynamicData(const FNia
 	const UNiagaraRibbonRendererProperties* Properties = CastChecked<const UNiagaraRibbonRendererProperties>(InProperties);
 
 	bool bSortKeyIsAge = false;
-	FNiagaraDataSetAccessor<FNiagaraDataConversions::FHalfOrFloat> SortKeyData;
+	FNiagaraDataSetAccessor<FNiagaraDataConversions<float>> SortKeyData;
 	if (Data.HasVariable(Properties->RibbonLinkOrderBinding.DataSetVariable.GetName()))
 	{
-		SortKeyData = FNiagaraDataSetAccessor<FNiagaraDataConversions::FHalfOrFloat>(Data, Properties->RibbonLinkOrderBinding.DataSetVariable.GetName());
+		SortKeyData = FNiagaraDataSetAccessor<FNiagaraDataConversions<float>>(Data, Properties->RibbonLinkOrderBinding.DataSetVariable.GetName());
 	}
 	else
 	{
-		SortKeyData = FNiagaraDataSetAccessor<FNiagaraDataConversions::FHalfOrFloat>(Data, Properties->NormalizedAgeBinding.DataSetVariable.GetName());
+		SortKeyData = FNiagaraDataSetAccessor<FNiagaraDataConversions<float>>(Data, Properties->NormalizedAgeBinding.DataSetVariable.GetName());
 		bSortKeyIsAge = true;
 	}
 
-	FNiagaraDataSetAccessor<FNiagaraDataConversions::FHalf3OrFloat3> PosData(Data, Properties->PositionBinding.DataSetVariable.GetName());
-	FNiagaraDataSetAccessor<FNiagaraDataConversions::FHalfOrFloat> SizeData(Data, Properties->RibbonWidthBinding.DataSetVariable.GetName());
-	FNiagaraDataSetAccessor<FNiagaraDataConversions::FHalfOrFloat> TwistData;
+	FNiagaraDataSetAccessor<FNiagaraDataConversions<FVector>> PosData(Data, Properties->PositionBinding.DataSetVariable.GetName());
+	FNiagaraDataSetAccessor<FNiagaraDataConversions<float>> SizeData(Data, Properties->RibbonWidthBinding.DataSetVariable.GetName());
+	FNiagaraDataSetAccessor<FNiagaraDataConversions<float>> TwistData;
 	if (Data.HasVariable(Properties->RibbonTwistBinding.DataSetVariable.GetName()))
 	{
-		TwistData = FNiagaraDataSetAccessor<FNiagaraDataConversions::FHalfOrFloat>(Data, Properties->RibbonTwistBinding.DataSetVariable.GetName());
+		TwistData = FNiagaraDataSetAccessor<FNiagaraDataConversions<float>>(Data, Properties->RibbonTwistBinding.DataSetVariable.GetName());
 	}
-	FNiagaraDataSetAccessor<FNiagaraDataConversions::FHalf3OrFloat3> FacingData;
+	FNiagaraDataSetAccessor<FNiagaraDataConversions<FVector>> FacingData;
 	if (Data.HasVariable(Properties->RibbonFacingBinding.DataSetVariable.GetName()))
 	{
-		FacingData = FNiagaraDataSetAccessor<FNiagaraDataConversions::FHalf3OrFloat3>(Data, Properties->RibbonFacingBinding.DataSetVariable.GetName());
+		FacingData = FNiagaraDataSetAccessor<FNiagaraDataConversions<FVector>>(Data, Properties->RibbonFacingBinding.DataSetVariable.GetName());
 	}
-	FNiagaraDataSetAccessor<FNiagaraDataConversions::FHalf4OrFloat4> MaterialParamData;
+	FNiagaraDataSetAccessor<FNiagaraDataConversions<FVector4>> MaterialParamData;
 	if (Data.HasVariable(Properties->DynamicMaterialBinding.DataSetVariable.GetName()))
 	{
-		MaterialParamData = FNiagaraDataSetAccessor<FNiagaraDataConversions::FHalf4OrFloat4>(Data, Properties->DynamicMaterialBinding.DataSetVariable.GetName());
+		MaterialParamData = FNiagaraDataSetAccessor<FNiagaraDataConversions<FVector4>>(Data, Properties->DynamicMaterialBinding.DataSetVariable.GetName());
 	}
-	FNiagaraDataSetAccessor<FNiagaraDataConversions::FHalf4OrFloat4> MaterialParam1Data;
+	FNiagaraDataSetAccessor<FNiagaraDataConversions<FVector4>> MaterialParam1Data;
 	if (Data.HasVariable(Properties->DynamicMaterial1Binding.DataSetVariable.GetName()))
 	{
-		MaterialParam1Data = FNiagaraDataSetAccessor<FNiagaraDataConversions::FHalf4OrFloat4>(Data, Properties->DynamicMaterial1Binding.DataSetVariable.GetName());
+		MaterialParam1Data = FNiagaraDataSetAccessor<FNiagaraDataConversions<FVector4>>(Data, Properties->DynamicMaterial1Binding.DataSetVariable.GetName());
 	}
-	FNiagaraDataSetAccessor<FNiagaraDataConversions::FHalf4OrFloat4> MaterialParam2Data;
+	FNiagaraDataSetAccessor<FNiagaraDataConversions<FVector4>> MaterialParam2Data;
 	if (Data.HasVariable(Properties->DynamicMaterial2Binding.DataSetVariable.GetName()))
 	{
-		MaterialParam2Data = FNiagaraDataSetAccessor<FNiagaraDataConversions::FHalf4OrFloat4>(Data, Properties->DynamicMaterial2Binding.DataSetVariable.GetName());
+		MaterialParam2Data = FNiagaraDataSetAccessor<FNiagaraDataConversions<FVector4>>(Data, Properties->DynamicMaterial2Binding.DataSetVariable.GetName());
 	}
-	FNiagaraDataSetAccessor<FNiagaraDataConversions::FHalf4OrFloat4> MaterialParam3Data;
+	FNiagaraDataSetAccessor<FNiagaraDataConversions<FVector4>> MaterialParam3Data;
 	if (Data.HasVariable(Properties->DynamicMaterial3Binding.DataSetVariable.GetName()))
 	{
-		MaterialParam3Data = FNiagaraDataSetAccessor<FNiagaraDataConversions::FHalf4OrFloat4>(Data, Properties->DynamicMaterial3Binding.DataSetVariable.GetName());
+		MaterialParam3Data = FNiagaraDataSetAccessor<FNiagaraDataConversions<FVector4>>(Data, Properties->DynamicMaterial3Binding.DataSetVariable.GetName());
 	}
 
 	FNiagaraDataSetAccessor<int32> RibbonIdData;
diff --git a/Engine/Plugins/FX/Niagara/Source/Niagara/Public/NiagaraTypes.h b/Engine/Plugins/FX/Niagara/Source/Niagara/Public/NiagaraTypes.h
--- a/Engine/Plugins/FX/Niagara/Source/Niagara/Public/NiagaraTypes.h
+++ b/Engine/Plugins/FX/Niagara/Source/Niagara/Public/NiagaraTypes.h
@@ -228,6 +228,9 @@ struct FNiagaraID
 	UPROPERTY(EditAnywhere, Category = ID)
 	int32 AcquireTag = 0;
 
+	FNiagaraID() : Index(INDEX_NONE), AcquireTag(INDEX_NONE) {}
+	FNiagaraID(int32 InIndex, int32 InAcquireTag): Index(InIndex), AcquireTag(InAcquireTag){}
+
 	bool operator==(const FNiagaraID& Other)const { return Index == Other.Index && AcquireTag == Other.AcquireTag; }
 	bool operator!=(const FNiagaraID& Other)const { return !(*this == Other); }
 	bool operator<(const FNiagaraID& Other)const { return Index < Other.Index && AcquireTag < Other.AcquireTag; }
@@ -1017,6 +1020,9 @@ public:
 	static const FNiagaraTypeDefinition& GetHalfVec3Def() { return HalfVec3Def; }
 	static const FNiagaraTypeDefinition& GetHalfVec4Def() { return HalfVec4Def; }
 
+	template<typename T>
+	static const FNiagaraTypeDefinition& Get();
+
 	static UScriptStruct* GetFloatStruct() { return FloatStruct; }
 	static UScriptStruct* GetBoolStruct() { return BoolStruct; }
 	static UScriptStruct* GetIntStruct() { return IntStruct; }
@@ -1137,6 +1143,27 @@ struct TStructOpsTypeTraits<FNiagaraTypeDefinition> : public TStructOpsTypeTrait
 	};
 };
 
+
+//Helper to get the correct typedef for templated code.
+template<typename T>
+const FNiagaraTypeDefinition& FNiagaraTypeDefinition::Get()
+{
+	if (TIsSame<T, float>::Value) { return FNiagaraTypeDefinition::GetFloatDef(); }
+	if (TIsSame<T, FVector2D>::Value) { return FNiagaraTypeDefinition::GetVec2Def(); }
+	if (TIsSame<T, FVector>::Value) { return FNiagaraTypeDefinition::GetVec3Def(); }
+	if (TIsSame<T, FVector4>::Value) { return FNiagaraTypeDefinition::GetVec4Def(); }
+	if (TIsSame<T, float>::Value) { return FNiagaraTypeDefinition::GetFloatDef(); }
+	if (TIsSame<T, int32>::Value) { return FNiagaraTypeDefinition::GetIntDef(); }
+	if (TIsSame<T, FNiagaraBool>::Value) { return FNiagaraTypeDefinition::GetBoolDef(); }
+	if (TIsSame<T, FQuat>::Value) { return FNiagaraTypeDefinition::GetQuatDef(); }
+	if (TIsSame<T, FMatrix>::Value) { return FNiagaraTypeDefinition::GetMatrix4Def(); }
+	if (TIsSame<T, FLinearColor>::Value) { return FNiagaraTypeDefinition::GetColorDef(); }
+	if (TIsSame<T, FNiagaraID>::Value) { return FNiagaraTypeDefinition::GetIDDef(); }
+}
+
+//////////////////////////////////////////////////////////////////////////
+
+
 /* Contains all types currently available for use in Niagara
 * Used by UI to provide selection; new uniforms and variables
 * may be instanced using the types provided here
