Niagara - More Parameter UX improvements.
+ Show the type of parameters in menus.
+ Added long display names to more namespaces.
+ Added a sort id to namespace metadata to sort in menus.
+ Added the ability to hide namespaces in script and system contexts using metadata.
+ Fixed an issue where changing the advanced display for the parameter panel in one window would change the menu state but not refresh the panel in other windows.
+ Fix inconsistency with parameter renaming on pins and in the parameter panel.
+ Added Alive to the DataInstance add menu in the parameter panel.
+ Can not open module asset editors or scratch scripts by double clicking in the system overview.

diff --git a/Engine/Plugins/FX/Niagara/Source/NiagaraEditor/Private/NiagaraActions.cpp b/Engine/Plugins/FX/Niagara/Source/NiagaraEditor/Private/NiagaraActions.cpp
--- a/Engine/Plugins/FX/Niagara/Source/NiagaraEditor/Private/NiagaraActions.cpp
+++ b/Engine/Plugins/FX/Niagara/Source/NiagaraEditor/Private/NiagaraActions.cpp
@@ -28,14 +28,14 @@ FNiagaraMenuAction::FNiagaraMenuAction(FText InNodeCategory, FText InMenuDesc, F
 	, CanPerformAction(InCanPerformAction)
 {}
 
-TOptional<FNiagaraParameterHandle> FNiagaraMenuAction::GetParameterHandle() const
+TOptional<FNiagaraVariable> FNiagaraMenuAction::GetParameterVariable() const
 {
-	return ParameterHandle;
+	return ParameterVariable;
 }
 
-void FNiagaraMenuAction::SetParamterHandle(const FNiagaraParameterHandle& InParameterHandle)
+void FNiagaraMenuAction::SetParamterVariable(const FNiagaraVariable& InParameterVariable)
 {
-	ParameterHandle = InParameterHandle;
+	ParameterVariable = InParameterVariable;
 }
 
 /************************************************************************/
diff --git a/Engine/Plugins/FX/Niagara/Source/NiagaraEditor/Private/NiagaraEditorSettings.cpp b/Engine/Plugins/FX/Niagara/Source/NiagaraEditor/Private/NiagaraEditorSettings.cpp
--- a/Engine/Plugins/FX/Niagara/Source/NiagaraEditor/Private/NiagaraEditorSettings.cpp
+++ b/Engine/Plugins/FX/Niagara/Source/NiagaraEditor/Private/NiagaraEditorSettings.cpp
@@ -6,6 +6,7 @@
 FNiagaraNamespaceMetadata::FNiagaraNamespaceMetadata()
 	: BackgroundColor(FLinearColor::Black)
 	, ForegroundStyle("NiagaraEditor.ParameterName.NamespaceText")
+	, SortId(TNumericLimits<int32>::Max())
 {
 }
 
@@ -13,6 +14,7 @@ FNiagaraNamespaceMetadata::FNiagaraNamespaceMetadata(TArray<FName> InNamespaces)
 	: Namespaces(InNamespaces)
 	, BackgroundColor(FLinearColor::Black)
 	, ForegroundStyle("NiagaraEditor.ParameterName.NamespaceText")
+	, SortId(TNumericLimits<int32>::Max())
 {
 }
 
@@ -31,36 +33,48 @@ UNiagaraEditorSettings::UNiagaraEditorSettings(const FObjectInitializer& ObjectI
 
 void UNiagaraEditorSettings::SetupNamespaceMetadata()
 {
-	NamespaceMetadata = 
+	NamespaceMetadata =
 	{
 		FNiagaraNamespaceMetadata({NAME_None})
 			.SetDisplayName(LOCTEXT("DefaultDisplayName", "None"))
 			.SetDescription(LOCTEXT("DefaultDescription", "Arbitrary sub-namespace for specifying module specific dataset attributes, or calling nested modules."))
-			.SetBackgroundColor(FLinearColor(FColor(102, 102, 102))),
+			.SetBackgroundColor(FLinearColor(FColor(102, 102, 102)))
+			.AddOption(ENiagaraNamespaceMetadataOptions::PreventEditing),
 		FNiagaraNamespaceMetadata({FNiagaraConstants::SystemNamespace})
 			.SetDisplayName(LOCTEXT("SystemDisplayName", "System"))
+			.SetDisplayNameLong(LOCTEXT("SystemDisplayNameLong", "System Attributes"))
 			.SetDescription(LOCTEXT("SystemDescription", "Persistent attribute in the system which is written in a system\n stage and can be read anywhere."))
 			.SetBackgroundColor(FLinearColor(FColor(49, 113, 142)))
+			.SetSortId(10)
 			.AddOption(ENiagaraNamespaceMetadataOptions::CanChangeNamespaceModifier),
 		FNiagaraNamespaceMetadata({FNiagaraConstants::EmitterNamespace})
 			.SetDisplayName(LOCTEXT("EmitterDisplayName", "Emitter"))
+		.SetDisplayNameLong(LOCTEXT("EmitterDisplayNameLong", "Emitter Attributes"))
 			.SetDescription(LOCTEXT("EmitterDescription", "Persistent attribute which is written in a emitter\nstage and can be read in emitter and particle stages."))
 			.SetBackgroundColor(FLinearColor(FColor(145, 99, 56)))
+			.SetSortId(20)
 			.AddOption(ENiagaraNamespaceMetadataOptions::CanChangeNamespaceModifier),
 		FNiagaraNamespaceMetadata({FNiagaraConstants::ParticleAttributeNamespace})
 			.SetDisplayName(LOCTEXT("ParticleDisplayName", "Particles"))
+			.SetDisplayNameLong(LOCTEXT("ParticleDisplayNameLong", "Particle Attributes"))
 			.SetDescription(LOCTEXT("ParticleDescription", "Persistent attribute which is written in a particle\nstage and can be read in particle stages."))
 			.SetBackgroundColor(FLinearColor(FColor(72, 130, 71)))
+			.SetSortId(30)
 			.AddOption(ENiagaraNamespaceMetadataOptions::CanChangeNamespaceModifier),
 		FNiagaraNamespaceMetadata({FNiagaraConstants::ModuleNamespace})
 			.SetDisplayName(LOCTEXT("ModuleDisplayName", "Input"))
+			.SetDisplayNameLong(LOCTEXT("ModuleDisplayNameLong", "Module Inputs"))
 			.SetDescription(LOCTEXT("ModuleDescription", "A value which exposes a module input to the system and emitter editor."))
 			.SetBackgroundColor(FLinearColor(FColor(136, 66, 65)))
+			.SetSortId(40)
+			.AddOption(ENiagaraNamespaceMetadataOptions::HideInSystem)
 			.AddOption(ENiagaraNamespaceMetadataOptions::CanChangeNamespaceModifier),
 		FNiagaraNamespaceMetadata({FNiagaraConstants::OutputNamespace, FNiagaraConstants::ModuleNamespace})
 			.SetDisplayName(LOCTEXT("ModuleOutputDisplayName", "Output"))
+			.SetDisplayNameLong(LOCTEXT("ModuleOutputDisplayNameLong", "Module Outputs"))
 			.SetDescription(LOCTEXT("ModuleOutputDescription", "A transient value which the module author has decided might be useful to other modules further down in the stage.\nTransient values do not persist from frame to frame, or between stages, e.g. emitter to particle, or spawn to update."))
 			.SetBackgroundColor(FLinearColor(FColor(108, 87, 131)))
+			.SetSortId(60)
 			.AddOption(ENiagaraNamespaceMetadataOptions::AdvancedInScript)
 			.AddOption(ENiagaraNamespaceMetadataOptions::AdvancedInSystem)
 			.AddOption(ENiagaraNamespaceMetadataOptions::PreventCreatingInSystemEditor)
@@ -70,47 +84,60 @@ void UNiagaraEditorSettings::SetupNamespaceMetadata()
 			.SetDescription(LOCTEXT("OutputDescription", "A transient value which the module author has decided might be useful to other modules further down in the stage.\nTransient values do not persist from frame to frame, or between stages, e.g. emitter to particle, or spawn to update."))
 			.SetBackgroundColor(FLinearColor(FColor(108, 87, 131)))
 			.AddOption(ENiagaraNamespaceMetadataOptions::AdvancedInScript)
-			.AddOption(ENiagaraNamespaceMetadataOptions::AdvancedInSystem),
+			.AddOption(ENiagaraNamespaceMetadataOptions::AdvancedInSystem)
+			.AddOption(ENiagaraNamespaceMetadataOptions::PreventEditing),
 		FNiagaraNamespaceMetadata({FNiagaraConstants::LocalNamespace, FNiagaraConstants::ModuleNamespace})
 			.SetDisplayName(LOCTEXT("ModuleLocalDisplayName", "Local"))
+			.SetDisplayNameLong(LOCTEXT("ModuleLocalDisplayNameLong", "Module Locals"))
 			.SetDescription(LOCTEXT("ModuleLocalDescription", "A transient value which can be written to and read from within a single module.\nTransient values do not persist from frame to frame, or between stages, e.g. emitter to particle, or spawn to update."))
 			.SetBackgroundColor(FLinearColor(FColor(191, 176, 84)))
 			.SetForegroundStyle("NiagaraEditor.ParameterName.NamespaceTextDark")
+			.SetSortId(50)
 			//.AddOption(ENiagaraNamespaceMetadataOptions::AdvancedInScript)
 			.AddOption(ENiagaraNamespaceMetadataOptions::AdvancedInSystem),
 		FNiagaraNamespaceMetadata({FNiagaraConstants::TransientNamespace})
 			.SetDisplayName(LOCTEXT("TransientDisplayName", "Transient"))
+			.SetDisplayNameLong(LOCTEXT("TransientDisplayNameLong", "Stage Transients"))
 			.SetDescription(LOCTEXT("TransientDescription", "A transient value which can be written to and read from from any module.\nTransient values do not persist from frame to frame, or between stages, e.g. emitter to particle, or spawn to update."))
 			.SetBackgroundColor(FLinearColor(FColor(108, 87, 131)))
+			.SetSortId(80)
 			.AddOption(ENiagaraNamespaceMetadataOptions::AdvancedInScript)
 			.AddOption(ENiagaraNamespaceMetadataOptions::AdvancedInSystem),
 		FNiagaraNamespaceMetadata({FNiagaraConstants::EngineNamespace})
 			.SetDisplayName(LOCTEXT("EngineDisplayName", "Engine"))
+			.SetDisplayNameLong(LOCTEXT("EngineDisplayNameLong", "Engine Provided"))
 			.SetDescription(LOCTEXT("EngineDescription", "A read only value which is provided by the engine.\nThis value's source can be the simulation itsef\ne.g. ExecutionCount, or the owner of the simulation (The component), e.g. (Owner) Scale."))
 			.SetBackgroundColor(FLinearColor(FColor(170, 170, 170)))
 			.SetForegroundStyle("NiagaraEditor.ParameterName.NamespaceTextDark")
-			.AddOption(ENiagaraNamespaceMetadataOptions::PreventRenaming),
+			.SetSortId(70)
+			.AddOption(ENiagaraNamespaceMetadataOptions::PreventEditing),
 		FNiagaraNamespaceMetadata({FNiagaraConstants::UserNamespace})
 			.SetDisplayName(LOCTEXT("UserDisplayName", "User"))
+			.SetDisplayNameLong(LOCTEXT("UserDisplayNameLong", "User Exposed"))
 			.SetDescription(LOCTEXT("UserDescription", "A read only value which can be initialized per system and\nmodified externally in the level, by blueprint, or by c++."))
-			.SetBackgroundColor(FLinearColor(FColor(91, 161, 194))),
+			.SetBackgroundColor(FLinearColor(FColor(91, 161, 194)))
+			.SetSortId(0)
+			.AddOption(ENiagaraNamespaceMetadataOptions::HideInScript),
 		FNiagaraNamespaceMetadata({FNiagaraConstants::ParameterCollectionNamespace})
 			.SetDisplayName(LOCTEXT("NiagaraParameterCollectionDisplayName", "NPC"))
 			.SetDisplayNameLong(LOCTEXT("NiagaraParameterCollectionDisplayNameLong", "Niagara Parameter Collection"))
 			.SetDescription(LOCTEXT("NiagaraParameterCollectionDescription", "Values read from a niagara parameter collection asset.\nRead only in a niagara system."))
 			.SetBackgroundColor(FLinearColor(FColor(170, 170, 170)))
 			.SetForegroundStyle("NiagaraEditor.ParameterName.NamespaceTextDark")
+			.SetSortId(90)
 			.AddOption(ENiagaraNamespaceMetadataOptions::AdvancedInScript)
 			.AddOption(ENiagaraNamespaceMetadataOptions::AdvancedInSystem)
-			.AddOption(ENiagaraNamespaceMetadataOptions::PreventRenaming),
+			.AddOption(ENiagaraNamespaceMetadataOptions::PreventEditing),
 		FNiagaraNamespaceMetadata({FNiagaraConstants::DataInstanceNamespace})
 			.SetDisplayName(LOCTEXT("DataInstanceDisplayName", "Data Instance"))
 			.SetDescription(LOCTEXT("DataInstanceDescription", "A special value which has a single bool IsAlive value, which determines if a particle is alive or not."))
 			.SetBackgroundColor(FLinearColor(FColor(170, 170, 170)))
 			.SetForegroundStyle("NiagaraEditor.ParameterName.NamespaceTextDark")
+			.SetSortId(100)
+			.AddOption(ENiagaraNamespaceMetadataOptions::HideInSystem)
 			.AddOption(ENiagaraNamespaceMetadataOptions::AdvancedInScript)
 			.AddOption(ENiagaraNamespaceMetadataOptions::AdvancedInSystem)
-			.AddOption(ENiagaraNamespaceMetadataOptions::PreventRenaming),
+			.AddOption(ENiagaraNamespaceMetadataOptions::PreventEditing),
 	};
 
 	NamespaceModifierMetadata =
@@ -223,6 +250,7 @@ void UNiagaraEditorSettings::SetDisplayAdvancedParameterPanelCategories(bool bIn
 	{
 		bDisplayAdvancedParameterPanelCategories = bInDisplayAdvancedParameterPanelCategories;
 		SaveConfig();
+		SettingsChangedDelegate.Broadcast(GET_MEMBER_NAME_CHECKED(UNiagaraEditorSettings, bDisplayAdvancedParameterPanelCategories).ToString(), this);
 	}
 }
 
@@ -285,6 +313,11 @@ FNiagaraNamespaceMetadata UNiagaraEditorSettings::GetMetaDataForNamespaces(TArra
 	}
 }
 
+const TArray<FNiagaraNamespaceMetadata>& UNiagaraEditorSettings::GetAllNamespaceMetadata() const
+{
+	return NamespaceMetadata;
+}
+
 FNiagaraNamespaceMetadata UNiagaraEditorSettings::GetMetaDataForNamespaceModifier(FName NamespaceModifier) const
 {
 	for (const FNiagaraNamespaceMetadata& NamespaceModifierMetadataItem : NamespaceModifierMetadata)
@@ -317,7 +350,5 @@ void UNiagaraEditorSettings::PostEditChangeProperty(FPropertyChangedEvent& Prope
 
 UNiagaraEditorSettings::FOnNiagaraEditorSettingsChanged& UNiagaraEditorSettings::OnSettingsChanged()
 {
-	return SettingsChangedDelegate;
-}
-
-UNiagaraEditorSettings::FOnNiagaraEditorSettingsChanged UNiagaraEditorSettings::SettingsChangedDelegate;
\ No newline at end of file
+	return GetMutableDefault<UNiagaraEditorSettings>()->SettingsChangedDelegate;
+}
\ No newline at end of file
diff --git a/Engine/Plugins/FX/Niagara/Source/NiagaraEditor/Private/NiagaraEditorStyle.cpp b/Engine/Plugins/FX/Niagara/Source/NiagaraEditor/Private/NiagaraEditorStyle.cpp
--- a/Engine/Plugins/FX/Niagara/Source/NiagaraEditor/Private/NiagaraEditorStyle.cpp
+++ b/Engine/Plugins/FX/Niagara/Source/NiagaraEditor/Private/NiagaraEditorStyle.cpp
@@ -199,6 +199,10 @@ TSharedRef< FSlateStyleSet > FNiagaraEditorStyle::Create()
 		.SetShadowOffset(FVector2D(1, 1))
 		.SetShadowColorAndOpacity(FLinearColor(1.0, 1.0, 1.0, 0.25f)));
 
+	Style->Set("NiagaraEditor.ParameterName.TypeText", FTextBlockStyle(NormalText)
+		.SetFont(DEFAULT_FONT("Regular", 8))
+		.SetColorAndOpacity(FLinearColor(1.0f, 1.0f, 1.0f, 0.5f)));
+
 	Style->Set("NiagaraEditor.Stack.HighlightedButtonBrush", new BOX_CORE_BRUSH("Common/ButtonHoverHint", FMargin(4 / 16.0f), GetDefault<UEditorStyleSettings>()->SelectionColor));
 
 	// Parameter Map View
diff --git a/Engine/Plugins/FX/Niagara/Source/NiagaraEditor/Private/NiagaraEditorUtilities.cpp b/Engine/Plugins/FX/Niagara/Source/NiagaraEditor/Private/NiagaraEditorUtilities.cpp
--- a/Engine/Plugins/FX/Niagara/Source/NiagaraEditor/Private/NiagaraEditorUtilities.cpp
+++ b/Engine/Plugins/FX/Niagara/Source/NiagaraEditor/Private/NiagaraEditorUtilities.cpp
@@ -2433,19 +2433,6 @@ void FNiagaraEditorUtilities::GetReferencingFunctionCallNodes(UNiagaraScript* Sc
 	}
 }
 
-FText FNiagaraEditorUtilities::FormatParameterNameForTextDisplay(FName ParameterName)
-{
-	FNiagaraParameterHandle ParameterHandle(ParameterName);
-	TArray<FName> HandleParts = ParameterHandle.GetHandleParts();
-	FString DisplayString;
-	for (int32 HandlePartIndex = 0; HandlePartIndex < HandleParts.Num() - 1; HandlePartIndex++)
-	{
-		DisplayString += TEXT("(") + HandleParts[HandlePartIndex].ToString().ToUpper() + TEXT(") ");
-	}
-	DisplayString += HandleParts[HandleParts.Num() - 1].ToString();
-	return FText::FromString(DisplayString);
-}
-
 bool FNiagaraEditorUtilities::GetVariableSortPriority(const FName& VarNameA, const FName& VarNameB)
 {
 	const FNiagaraNamespaceMetadata& NamespaceMetaDataA = GetNamespaceMetaDataForVariableName(VarNameA);
@@ -2573,4 +2560,209 @@ bool FNiagaraParameterUtilities::DoesParameterNameMatchSearchText(FName Paramete
 	return false;
 }
 
+FText FNiagaraParameterUtilities::FormatParameterNameForTextDisplay(FName ParameterName)
+{
+	FNiagaraParameterHandle ParameterHandle(ParameterName);
+	TArray<FName> HandleParts = ParameterHandle.GetHandleParts();
+	FString DisplayString;
+	for (int32 HandlePartIndex = 0; HandlePartIndex < HandleParts.Num() - 1; HandlePartIndex++)
+	{
+		DisplayString += TEXT("(") + HandleParts[HandlePartIndex].ToString().ToUpper() + TEXT(") ");
+	}
+	DisplayString += HandleParts[HandleParts.Num() - 1].ToString();
+	return FText::FromString(DisplayString);
+}
+
+FName NamePartsToName(const TArray<FName>& NameParts)
+{
+	TArray<FString> NamePartStrings;
+	for (FName NamePart : NameParts)
+	{
+		NamePartStrings.Add(NamePart.ToString());
+	}
+	return *FString::Join(NamePartStrings, TEXT("."));
+}
+
+bool FNiagaraParameterUtilities::GetNamespaceEditData(
+	FName InParameterName, 
+	FNiagaraParameterHandle& OutParameterHandle,
+	FNiagaraNamespaceMetadata& OutNamespaceMetadata,
+	FText& OutErrorMessage)
+{
+	OutParameterHandle = FNiagaraParameterHandle(InParameterName);
+	TArray<FName> NameParts = OutParameterHandle.GetHandleParts();
+	OutNamespaceMetadata = GetDefault<UNiagaraEditorSettings>()->GetMetaDataForNamespaces(NameParts);
+	if (OutNamespaceMetadata.IsValid() == false || OutNamespaceMetadata.Options.Contains(ENiagaraNamespaceMetadataOptions::PreventEditing))
+	{
+		OutErrorMessage = LOCTEXT("NoMetadataForNamespace", "This parameter doesn't support editing.");
+		return false;
+	}
+
+	return true;
+}
+
+bool FNiagaraParameterUtilities::GetNamespaceModifierEditData(
+	FName InParameterName,
+	FNiagaraParameterHandle& OutParameterHandle,
+	FNiagaraNamespaceMetadata& OutNamespaceMetadata,
+	FText& OutErrorMessage)
+{
+	if (GetNamespaceEditData(InParameterName, OutParameterHandle, OutNamespaceMetadata, OutErrorMessage))
+	{
+		if (OutNamespaceMetadata.IsValid() == false || OutNamespaceMetadata.Options.Contains(ENiagaraNamespaceMetadataOptions::CanChangeNamespaceModifier) == false)
+		{
+			OutErrorMessage = LOCTEXT("NotSupportedForThisNamespace", "This parameter doesn't support namespace modifiers.");
+			return false;
+		}
+		return true;
+	}
+	else
+	{
+		return false;
+	}
+}
+
+bool FNiagaraParameterUtilities::TestCanChangeNamespaceWithMessage(FName ParameterName, const FNiagaraNamespaceMetadata& NewNamespaceMetadata, FText& OutMessage)
+{
+	FNiagaraParameterHandle ParameterHandle;
+	FNiagaraNamespaceMetadata NamespaceMetadata;
+	if (GetNamespaceEditData(ParameterName, ParameterHandle, NamespaceMetadata, OutMessage))
+	{
+		if (NewNamespaceMetadata.Options.Contains(ENiagaraNamespaceMetadataOptions::PreventEditing))
+		{
+			OutMessage = LOCTEXT("NewNamespaceIsntValid", "The new namespace does not support editing so it can not be assigned.");
+			return false;
+		}
+		else
+		{
+			OutMessage = FText::Format(LOCTEXT("ChagneNamespaceFormat", "Change this parameters namespace to {0}"), NewNamespaceMetadata.DisplayName);
+			return true;
+		}
+	}
+	return false;
+}
+
+FName FNiagaraParameterUtilities::ChangeNamespace(FName ParameterName, const FNiagaraNamespaceMetadata& NewNamespaceMetadata)
+{
+	FNiagaraParameterHandle ParameterHandle;
+	FNiagaraNamespaceMetadata NamespaceMetadata;
+	FText Unused;
+	if (NewNamespaceMetadata.IsValid() &&
+		NewNamespaceMetadata.Options.Contains(ENiagaraNamespaceMetadataOptions::PreventEditing) == false &&
+		GetNamespaceEditData(ParameterName, ParameterHandle, NamespaceMetadata, Unused))
+	{
+		TArray<FName> NameParts = ParameterHandle.GetHandleParts();
+		NameParts.RemoveAt(0, NamespaceMetadata.Namespaces.Num());
+		NameParts.Insert(NewNamespaceMetadata.Namespaces, 0);
+		return NamePartsToName(NameParts);
+	}
+	return NAME_None;
+}
+
+bool FNiagaraParameterUtilities::TestCanAddNamespaceModifierWithMessage(FName ParameterName, FText& OutMessage)
+{
+	FNiagaraParameterHandle ParameterHandle;
+	FNiagaraNamespaceMetadata NamespaceMetadata;
+	if (GetNamespaceModifierEditData(ParameterName, ParameterHandle, NamespaceMetadata, OutMessage))
+	{
+		int32 NumberOfNamePartsAfterNamespace = ParameterHandle.GetHandleParts().Num() - NamespaceMetadata.Namespaces.Num();
+		if (NumberOfNamePartsAfterNamespace == 1)
+		{
+			OutMessage = LOCTEXT("AddNamespaceModifier", "Add a namespace modifier to this parameter.");
+			return true;
+		}
+		else
+		{
+			OutMessage = LOCTEXT("CantAddAnotherNamespaceModfier", "Only one namespace modifier is supported.");
+			return false;
+		}
+	}
+	return false;
+}
+
+FName FNiagaraParameterUtilities::AddNamespaceModifier(FName InParameterName)
+{
+	FNiagaraParameterHandle ParameterHandle;
+	FNiagaraNamespaceMetadata NamespaceMetadata;
+	FText Unused;
+	if (GetNamespaceModifierEditData(InParameterName, ParameterHandle, NamespaceMetadata, Unused))
+	{
+		TArray<FName> NameParts = ParameterHandle.GetHandleParts();
+		int32 NumberOfNamePartsAfterNamespace = NameParts.Num() - NamespaceMetadata.Namespaces.Num();
+		if (NumberOfNamePartsAfterNamespace == 1)
+		{
+			NameParts.Insert(FNiagaraConstants::ModuleNamespace, NamespaceMetadata.Namespaces.Num());
+			return NamePartsToName(NameParts);
+		}
+	}
+	return NAME_None;
+}
+
+bool FNiagaraParameterUtilities::TestCanRemoveNamespaceModifierWithMessage(FName ParameterName, FText& OutMessage)
+{
+	FNiagaraParameterHandle ParameterHandle;
+	FNiagaraNamespaceMetadata NamespaceMetadata;
+	if (GetNamespaceModifierEditData(ParameterName, ParameterHandle, NamespaceMetadata, OutMessage))
+	{
+		int32 NumberOfNamePartsAfterNamespace = ParameterHandle.GetHandleParts().Num() - NamespaceMetadata.Namespaces.Num();
+		if (NumberOfNamePartsAfterNamespace == 2)
+		{
+			OutMessage = LOCTEXT("RemoveNamespaceModifier", "Remove the namespace modifier from this parameter.");
+			return true;
+		}
+		else
+		{
+			OutMessage = LOCTEXT("NoNamespaceModifierToRemove", "No namespace modifier to remove.");
+			return false;
+		}
+	}
+	return false;
+}
+
+FName FNiagaraParameterUtilities::RemoveNamespaceModifier(FName InParameterName)
+{
+	FNiagaraParameterHandle ParameterHandle;
+	FNiagaraNamespaceMetadata NamespaceMetadata;
+	FText Unused;
+	if (GetNamespaceModifierEditData(InParameterName, ParameterHandle, NamespaceMetadata, Unused))
+	{
+		TArray<FName> NameParts = ParameterHandle.GetHandleParts();
+		int32 NumberOfNamePartsAfterNamespace = NameParts.Num() - NamespaceMetadata.Namespaces.Num();
+		if (NumberOfNamePartsAfterNamespace == 2)
+		{
+			NameParts.RemoveAt(NamespaceMetadata.Namespaces.Num());
+			return NamePartsToName(NameParts);
+		}
+	}
+	return NAME_None;
+}
+
+bool FNiagaraParameterUtilities::TestCanEditNamespaceModifierWithMessage(FName ParameterName, FText& OutMessage)
+{
+	FNiagaraParameterHandle ParameterHandle;
+	FNiagaraNamespaceMetadata NamespaceMetadata;
+	if (GetNamespaceModifierEditData(ParameterName, ParameterHandle, NamespaceMetadata, OutMessage))
+	{
+		int32 NumberOfNamePartsAfterNamespace = ParameterHandle.GetHandleParts().Num() - NamespaceMetadata.Namespaces.Num();
+		if (NumberOfNamePartsAfterNamespace == 2)
+		{
+			OutMessage = LOCTEXT("EditNamespaceModifier", "Edit the namespace modifier for this parameter.");
+			return true;
+		}
+		else
+		{
+			OutMessage = LOCTEXT("NoNamespaceModifierToEdit", "No namespace modifier to edit.");
+			return false;
+		}
+	}
+	return false;
+}
+
+bool FNiagaraParameterUtilities::TestCanRenameWithMessage(FName ParameterName, FText& OutMessage)
+{
+	FNiagaraParameterHandle ParameterHandle;
+	FNiagaraNamespaceMetadata NamespaceMetadata;
+	return GetNamespaceEditData(ParameterName, ParameterHandle, NamespaceMetadata, OutMessage);
+}
+
 #undef LOCTEXT_NAMESPACE
diff --git a/Engine/Plugins/FX/Niagara/Source/NiagaraEditor/Private/NiagaraNode.cpp b/Engine/Plugins/FX/Niagara/Source/NiagaraEditor/Private/NiagaraNode.cpp
--- a/Engine/Plugins/FX/Niagara/Source/NiagaraEditor/Private/NiagaraNode.cpp
+++ b/Engine/Plugins/FX/Niagara/Source/NiagaraEditor/Private/NiagaraNode.cpp
@@ -372,6 +372,23 @@ bool UNiagaraNode::ConvertNumericPinToType(UEdGraphPin* InGraphPin, FNiagaraType
 	return true;
 }
 
+bool UNiagaraNode::GetIsPinRenamePending(const UEdGraphPin* Pin)
+{
+	return PinsGuidsWithRenamePending.Contains(Pin->PersistentGuid);
+}
+
+void UNiagaraNode::SetIsPinRenamePending(const UEdGraphPin* Pin, bool bInIsRenamePending)
+{
+	if (bInIsRenamePending)
+	{
+		PinsGuidsWithRenamePending.AddUnique(Pin->PersistentGuid);
+	}
+	else
+	{
+		PinsGuidsWithRenamePending.Remove(Pin->PersistentGuid);
+	}
+}
+
 TSharedPtr<SGraphNode> UNiagaraNode::CreateVisualWidget() 
 {
 	return SNew(SNiagaraGraphNode, this);
diff --git a/Engine/Plugins/FX/Niagara/Source/NiagaraEditor/Private/NiagaraNodeParameterMapBase.cpp b/Engine/Plugins/FX/Niagara/Source/NiagaraEditor/Private/NiagaraNodeParameterMapBase.cpp
--- a/Engine/Plugins/FX/Niagara/Source/NiagaraEditor/Private/NiagaraNodeParameterMapBase.cpp
+++ b/Engine/Plugins/FX/Niagara/Source/NiagaraEditor/Private/NiagaraNodeParameterMapBase.cpp
@@ -19,6 +19,7 @@
 #include "Widgets/Input/SEditableTextBox.h"
 #include "ToolMenus.h"
 #include "NiagaraScriptVariable.h"
+#include "Widgets/SNiagaraParameterName.h"
 
 #include "IAssetTools.h"
 #include "AssetRegistryModule.h"
@@ -210,6 +211,34 @@ bool UNiagaraNodeParameterMapBase::OnAllowDrop(TSharedPtr<FDragDropOperation> Dr
 	return true;
 }
 
+bool UNiagaraNodeParameterMapBase::CanRenamePin(const UEdGraphPin* Pin) const
+{
+	if (IsAddPin(Pin))
+	{
+		return false;
+	}
+
+	FText Unused;
+	return FNiagaraParameterUtilities::TestCanRenameWithMessage(Pin->PinName, Unused);
+}
+
+bool UNiagaraNodeParameterMapBase::GetIsPinEditNamespaceModifierPending(const UEdGraphPin* Pin)
+{
+	return PinsGuidsWithEditNamespaceModifierPending.Contains(Pin->PersistentGuid);
+}
+
+void UNiagaraNodeParameterMapBase::SetIsPinEditNamespaceModifierPending(const UEdGraphPin* Pin, bool bInIsEditNamespaceModifierPending)
+{
+	if (bInIsEditNamespaceModifierPending)
+	{
+		PinsGuidsWithEditNamespaceModifierPending.AddUnique(Pin->PersistentGuid);
+	}
+	else
+	{
+		PinsGuidsWithEditNamespaceModifierPending.Remove(Pin->PersistentGuid);
+	}
+}
+
 void UNiagaraNodeParameterMapBase::OnPinRenamed(UEdGraphPin* RenamedPin, const FString& OldName)
 {
 	RenamedPin->PinFriendlyName = FText::FromName(RenamedPin->PinName);
@@ -258,7 +287,24 @@ void UNiagaraNodeParameterMapBase::GetNodeContextMenuActions(UToolMenu* Menu, UG
 		FNiagaraVariable Var = CastChecked<UEdGraphSchema_Niagara>(GetSchema())->PinToNiagaraVariable(Pin);
 		const UNiagaraGraph* Graph = GetNiagaraGraph();
 
+		UNiagaraNodeParameterMapBase* NonConstThis = const_cast<UNiagaraNodeParameterMapBase*>(this);
+		UEdGraphPin* NonConstPin = const_cast<UEdGraphPin*>(Context->Pin);
 		//if (!FNiagaraConstants::IsNiagaraConstant(Var))
+		FToolMenuSection& EditSection = Menu->FindOrAddSection("EditPin");
+		{
+			EditSection.AddSubMenu(
+				"ChangeNamespace",
+				LOCTEXT("ChangeNamespace", "Change Namespace"),
+				LOCTEXT("ChangeNamespaceToolTip", "Change the namespace for this parameter pin."),
+				FNewToolMenuDelegate::CreateUObject(NonConstThis, &UNiagaraNodeParameterMapBase::GetChangeNamespaceSubMenuForPin, NonConstPin));
+
+			EditSection.AddSubMenu(
+				"ChangeNamespaceModifier",
+				LOCTEXT("ChangeNamespaceModifier", "Change Namespace Modifier"),
+				LOCTEXT("ChangeNamespaceModifierToolTip", "Change the namespace modifier for this parameter pin."),
+				FNewToolMenuDelegate::CreateUObject(NonConstThis, &UNiagaraNodeParameterMapBase::GetChangeNamespaceModifierSubMenuForPin, NonConstPin));
+		}
+
 		{
 			FToolMenuSection& Section = Menu->AddSection("EdGraphSchema_NiagaraParamAction", LOCTEXT("EditPinMenuHeader", "Parameters"));
 			Section.AddMenuEntry(
@@ -266,12 +312,74 @@ void UNiagaraNodeParameterMapBase::GetNodeContextMenuActions(UToolMenu* Menu, UG
 				LOCTEXT("SelectParameterPin", "Select parameter"),
 				LOCTEXT("SelectParameterPinToolTip", "Select this parameter in the paramter panel"),
 				FSlateIcon(),
-				FUIAction(FExecuteAction::CreateUObject(const_cast<UNiagaraNodeParameterMapBase*>(this), &UNiagaraNodeParameterMapBase::SelectParameterFromPin, const_cast<UEdGraphPin*>(Context->Pin))));
+				FUIAction(FExecuteAction::CreateUObject(NonConstThis, &UNiagaraNodeParameterMapBase::SelectParameterFromPin, Context->Pin)));
 		}
 	}
 }
 
-void UNiagaraNodeParameterMapBase::SelectParameterFromPin(UEdGraphPin* InPin) 
+void UNiagaraNodeParameterMapBase::GetChangeNamespaceSubMenuForPin(UToolMenu* Menu, UEdGraphPin* InPin)
+{
+	FToolMenuSection& Section = Menu->AddSection("Section");
+
+	TArray<FNiagaraNamespaceMetadata> NamespaceMetadata = GetDefault<UNiagaraEditorSettings>()->GetAllNamespaceMetadata();
+	NamespaceMetadata.Sort([](const FNiagaraNamespaceMetadata& MetadataA, const FNiagaraNamespaceMetadata& MetadataB)
+		{ return MetadataA.SortId < MetadataB.SortId; });
+
+	for (const FNiagaraNamespaceMetadata& Metadata : NamespaceMetadata)
+	{
+		if (Metadata.IsValid() == false || Metadata.Options.Contains(ENiagaraNamespaceMetadataOptions::PreventEditing) || Metadata.Options.Contains(ENiagaraNamespaceMetadataOptions::HideInScript))
+		{
+			continue;
+		}
+
+		FText CanChangeMessage;
+		bool bCanChange = FNiagaraParameterUtilities::TestCanChangeNamespaceWithMessage(InPin->PinName, Metadata, CanChangeMessage);
+
+		FText ToolTip;
+		if (bCanChange)
+		{
+			ToolTip = FText::Format(LOCTEXT("ChangeNamespaceToolTipFormat", "{0}\n\nDescription:\n{1}"), CanChangeMessage, Metadata.Description);
+		}
+		else
+		{
+			ToolTip = CanChangeMessage;
+		}
+
+		FString NamespaceNameString = Metadata.Namespaces[0].ToString();
+		for (int32 i = 1; i < Metadata.Namespaces.Num(); i++)
+		{
+			NamespaceNameString += TEXT(".") + Metadata.Namespaces[i].ToString();
+		}
+
+		FUIAction Action = FUIAction(
+			FExecuteAction::CreateUObject(this, &UNiagaraNodeParameterMapBase::ChangeNamespaceForPin, InPin, Metadata),
+			FCanExecuteAction::CreateLambda([bCanChange]() { return bCanChange; }));
+
+		TSharedRef<SWidget> NamespaceWidget =
+			SNew(SBox)
+			.Padding(FMargin(7, 2, 7, 2))
+			[
+				SNew(SNiagaraParameterName)
+				.ParameterName(*NamespaceNameString)
+				.IsReadOnly(true)
+				.ToolTipText(ToolTip)
+			];
+
+		Section.AddEntry(FToolMenuEntry::InitMenuEntry(NAME_None, Action, NamespaceWidget));
+	}
+}
+
+void UNiagaraNodeParameterMapBase::ChangeNamespaceForPin(UEdGraphPin* InPin, FNiagaraNamespaceMetadata NewNamespaceMetadata)
+{
+	FName NewName = FNiagaraParameterUtilities::ChangeNamespace(InPin->PinName, NewNamespaceMetadata);
+	if (NewName != NAME_None)
+	{
+		FScopedTransaction Transaction(LOCTEXT("ChangeNamespaceTransaction", "Change parameter namespace"));
+		CommitEditablePinName(FText::FromName(NewName), InPin);
+	}
+}
+
+void UNiagaraNodeParameterMapBase::SelectParameterFromPin(const UEdGraphPin* InPin) 
 {
 	UNiagaraGraph* NiagaraGraph = GetNiagaraGraph();
 	if (NiagaraGraph && InPin)
@@ -295,5 +403,107 @@ void UNiagaraNodeParameterMapBase::SelectParameterFromPin(UEdGraphPin* InPin)
 	}
 }
 
+void UNiagaraNodeParameterMapBase::GetChangeNamespaceModifierSubMenuForPin(UToolMenu* Menu, UEdGraphPin* InPin)
+{
+	FToolMenuSection& Section = Menu->AddSection("Section");
+
+	TAttribute<FText> AddToolTip = TAttribute<FText>::Create(TAttribute<FText>::FGetter::CreateUObject(
+		this, &UNiagaraNodeParameterMapBase::GetAddNamespaceModifierForPinToolTip, (const UEdGraphPin*)InPin));
+	Section.AddMenuEntry(
+		"AddNamespaceModifier",
+		LOCTEXT("AddNamespaceModifierForPin", "Add"),
+		AddToolTip,
+		FSlateIcon(),
+		FUIAction(
+			FExecuteAction::CreateUObject(this, &UNiagaraNodeParameterMapBase::AddNamespaceModifierForPin, InPin),
+			FCanExecuteAction::CreateUObject(this, &UNiagaraNodeParameterMapBase::CanAddNamespaceModifierForPin, (const UEdGraphPin*)InPin)));
+
+	TAttribute<FText> RemoveToolTip = TAttribute<FText>::Create(TAttribute<FText>::FGetter::CreateUObject(
+		this, &UNiagaraNodeParameterMapBase::GetRemoveNamespaceModifierForPinToolTip, (const UEdGraphPin*)InPin));
+	Section.AddMenuEntry(
+		"RemoveNamespceModifier",
+		LOCTEXT("RemoveNamespaceModifierForPin", "Remove"),
+		RemoveToolTip,
+		FSlateIcon(),
+		FUIAction(
+			FExecuteAction::CreateUObject(this, &UNiagaraNodeParameterMapBase::RemoveNamespaceModifierForPin, InPin),
+			FCanExecuteAction::CreateUObject(this, &UNiagaraNodeParameterMapBase::CanRemoveNamespaceModifierForPin, (const UEdGraphPin*)InPin)));
+
+	TAttribute<FText> EditToolTip = TAttribute<FText>::Create(TAttribute<FText>::FGetter::CreateUObject(
+		this, &UNiagaraNodeParameterMapBase::GetEditNamespaceModifierForPinToolTip, (const UEdGraphPin*)InPin));
+	Section.AddMenuEntry(
+		"EditNamespaceModifier",
+		LOCTEXT("EditNamespaceModifierForPin", "Edit"),
+		EditToolTip,
+		FSlateIcon(),
+		FUIAction(
+			FExecuteAction::CreateUObject(this, &UNiagaraNodeParameterMapBase::EditNamespaceModifierForPin, (const UEdGraphPin*)InPin),
+			FCanExecuteAction::CreateUObject(this, &UNiagaraNodeParameterMapBase::CanEditNamespaceModifierForPin, (const UEdGraphPin*)InPin)));
+}
+
+FText UNiagaraNodeParameterMapBase::GetAddNamespaceModifierForPinToolTip(const UEdGraphPin* InPin) const
+{
+	FText AddMessage;
+	FNiagaraParameterUtilities::TestCanAddNamespaceModifierWithMessage(InPin->PinName, AddMessage);
+	return AddMessage;
+}
+
+bool UNiagaraNodeParameterMapBase::CanAddNamespaceModifierForPin(const UEdGraphPin* InPin) const
+{
+	FText Unused;
+	return FNiagaraParameterUtilities::TestCanAddNamespaceModifierWithMessage(InPin->PinName, Unused);
+}
+
+void UNiagaraNodeParameterMapBase::AddNamespaceModifierForPin(UEdGraphPin* InPin)
+{
+	FName NewName = FNiagaraParameterUtilities::AddNamespaceModifier(InPin->PinName);
+	if (NewName != NAME_None)
+	{
+		FScopedTransaction Transaction(LOCTEXT("AddNamespaceModifierTransaction", "Add namespace modifier"));
+		CommitEditablePinName(FText::FromName(NewName), InPin);
+	}
+}
+
+FText UNiagaraNodeParameterMapBase::GetRemoveNamespaceModifierForPinToolTip(const UEdGraphPin* InPin) const
+{
+	FText RemoveMessage;
+	FNiagaraParameterUtilities::TestCanRemoveNamespaceModifierWithMessage(InPin->PinName, RemoveMessage);
+	return RemoveMessage;
+}
+
+bool UNiagaraNodeParameterMapBase::CanRemoveNamespaceModifierForPin(const UEdGraphPin* InPin) const
+{
+	FText Unused;
+	return FNiagaraParameterUtilities::TestCanRemoveNamespaceModifierWithMessage(InPin->PinName, Unused);
+}
+
+void UNiagaraNodeParameterMapBase::RemoveNamespaceModifierForPin(UEdGraphPin* InPin)
+{
+	FName NewName = FNiagaraParameterUtilities::RemoveNamespaceModifier(InPin->PinName);
+	if (NewName != NAME_None)
+	{
+		FScopedTransaction Transaction(LOCTEXT("RemoveNamespaceModifierTransaction", "Remove namespace modifier"));
+		CommitEditablePinName(FText::FromName(NewName), InPin);
+	}
+}
+
+FText UNiagaraNodeParameterMapBase::GetEditNamespaceModifierForPinToolTip(const UEdGraphPin* InPin) const
+{
+	FText EditMessage;
+	FNiagaraParameterUtilities::TestCanEditNamespaceModifierWithMessage(InPin->PinName, EditMessage);
+	return EditMessage;
+}
+
+bool UNiagaraNodeParameterMapBase::CanEditNamespaceModifierForPin(const UEdGraphPin* InPin) const
+{
+	FText Unused;
+	return FNiagaraParameterUtilities::TestCanEditNamespaceModifierWithMessage(InPin->PinName, Unused);
+}
+
+void UNiagaraNodeParameterMapBase::EditNamespaceModifierForPin(const UEdGraphPin* InPin)
+{
+	SetIsPinEditNamespaceModifierPending(InPin, true);
+}
+
 
 #undef LOCTEXT_NAMESPACE
\ No newline at end of file
diff --git a/Engine/Plugins/FX/Niagara/Source/NiagaraEditor/Private/NiagaraNodeParameterMapBase.h b/Engine/Plugins/FX/Niagara/Source/NiagaraEditor/Private/NiagaraNodeParameterMapBase.h
--- a/Engine/Plugins/FX/Niagara/Source/NiagaraEditor/Private/NiagaraNodeParameterMapBase.h
+++ b/Engine/Plugins/FX/Niagara/Source/NiagaraEditor/Private/NiagaraNodeParameterMapBase.h
@@ -5,11 +5,11 @@
 #include "NiagaraNodeWithDynamicPins.h"
 #include "SGraphPin.h"
 #include "NiagaraConstants.h"
+#include "NiagaraEditorSettings.h"
 #include "NiagaraNodeParameterMapBase.generated.h"
 
 class UEdGraphPin;
 
-
 /** A node which allows the user to build a set of arbitrary output types from an arbitrary set of input types by connecting their inner components. */
 UCLASS()
 class UNiagaraNodeParameterMapBase : public UNiagaraNodeWithDynamicPins
@@ -43,17 +43,42 @@ public:
 
 	bool OnAllowDrop(TSharedPtr<FDragDropOperation> DragDropOperation);
 
-	virtual bool CanRenamePinFromContextMenu(const UEdGraphPin* Pin) const { return false; }
+	virtual bool CanRenamePinFromContextMenu(const UEdGraphPin* Pin) const override { return false; }
+
+	virtual bool CanRenamePin(const UEdGraphPin* Pin) const override;
 
 	virtual FName GetNewPinDefaultNamespace() const { return PARAM_MAP_LOCAL_MODULE_STR; }
 
+	bool GetIsPinEditNamespaceModifierPending(const UEdGraphPin* Pin);
+
+	void SetIsPinEditNamespaceModifierPending(const UEdGraphPin* Pin, bool bInIsEditNamespaceModifierPending);
+
 protected:
-	virtual void SelectParameterFromPin(UEdGraphPin* InPin);
-	virtual void OnPinRenamed(UEdGraphPin* RenamedPin, const FString& OldName) override;
+	void GetChangeNamespaceSubMenuForPin(UToolMenu* Menu, UEdGraphPin* InPin);
+	void ChangeNamespaceForPin(UEdGraphPin* InPin, FNiagaraNamespaceMetadata NewNamespaceMetadata);
+
+	virtual void SelectParameterFromPin(const UEdGraphPin* InPin);
 
+	void GetChangeNamespaceModifierSubMenuForPin(UToolMenu* Menu, UEdGraphPin* InPin);
+
+	FText GetAddNamespaceModifierForPinToolTip(const UEdGraphPin* InPin) const;
+	bool CanAddNamespaceModifierForPin(const UEdGraphPin* InPin) const;
+	void AddNamespaceModifierForPin(UEdGraphPin* InPin);
+
+	FText GetRemoveNamespaceModifierForPinToolTip(const UEdGraphPin* InPin) const;
+	bool CanRemoveNamespaceModifierForPin(const UEdGraphPin* InPin) const;
+	void RemoveNamespaceModifierForPin(UEdGraphPin* InPin);
+
+	FText GetEditNamespaceModifierForPinToolTip(const UEdGraphPin* InPin) const;
+	bool CanEditNamespaceModifierForPin(const UEdGraphPin* InPin) const;
+	void EditNamespaceModifierForPin(const UEdGraphPin* InPin);
+
+	virtual void OnPinRenamed(UEdGraphPin* RenamedPin, const FString& OldName) override;
 
 	UEdGraphPin* PinPendingRename;
 
+	TArray<FGuid> PinsGuidsWithEditNamespaceModifierPending;
+
 public:
 	/** The sub category for parameter pins. */
 	static const FName ParameterPinSubCategory;
diff --git a/Engine/Plugins/FX/Niagara/Source/NiagaraEditor/Private/NiagaraNodeParameterMapSet.h b/Engine/Plugins/FX/Niagara/Source/NiagaraEditor/Private/NiagaraNodeParameterMapSet.h
--- a/Engine/Plugins/FX/Niagara/Source/NiagaraEditor/Private/NiagaraNodeParameterMapSet.h
+++ b/Engine/Plugins/FX/Niagara/Source/NiagaraEditor/Private/NiagaraNodeParameterMapSet.h
@@ -37,7 +37,7 @@ public:
 	virtual void GetNodeContextMenuActions(class UToolMenu* Menu, class UGraphNodeContextMenuContext* Context) const override;
 	virtual bool IncludeParentNodeContextMenu() const { return true; }
 	virtual void PostLoad() override;
-	virtual FName GetNewPinDefaultNamespace() const { return PARAM_MAP_OUTPUT_MODULE_STR; }
+	virtual FName GetNewPinDefaultNamespace() const { return PARAM_MAP_LOCAL_MODULE_STR; }
 
 protected:
 	virtual void OnNewTypedPinAdded(UEdGraphPin* NewPin); 
diff --git a/Engine/Plugins/FX/Niagara/Source/NiagaraEditor/Private/NiagaraNodeWithDynamicPins.cpp b/Engine/Plugins/FX/Niagara/Source/NiagaraEditor/Private/NiagaraNodeWithDynamicPins.cpp
--- a/Engine/Plugins/FX/Niagara/Source/NiagaraEditor/Private/NiagaraNodeWithDynamicPins.cpp
+++ b/Engine/Plugins/FX/Niagara/Source/NiagaraEditor/Private/NiagaraNodeWithDynamicPins.cpp
@@ -189,7 +189,7 @@ void UNiagaraNodeWithDynamicPins::GetNodeContextMenuActions(UToolMenu* Menu, UGr
 	Super::GetNodeContextMenuActions(Menu, Context);
 	if (Context->Pin != nullptr)
 	{
-		FToolMenuSection& Section = Menu->AddSection("EdGraphSchema_NiagaraPinActions", LOCTEXT("EditPinMenuHeader", "Edit Pin"));
+		FToolMenuSection& Section = Menu->AddSection("EditPin", LOCTEXT("EditPinMenuHeader", "Edit Pin"));
 		if (CanRenamePinFromContextMenu(Context->Pin))
 		{
 			UEdGraphPin* Pin = const_cast<UEdGraphPin*>(Context->Pin);
@@ -205,6 +205,16 @@ void UNiagaraNodeWithDynamicPins::GetNodeContextMenuActions(UToolMenu* Menu, UGr
 				];
 			Section.AddEntry(FToolMenuEntry::InitWidget("RenameWidget", RenameWidget, LOCTEXT("NameMenuItem", "Name")));
 		}
+		else if (CanRenamePin(Context->Pin))
+		{
+			Section.AddMenuEntry(
+				NAME_None,
+				LOCTEXT("RenameDynamicPin", "Rename pin"),
+				LOCTEXT("RenameDynamicPinToolTip", "Rename this pin."),
+				FSlateIcon(),
+				FUIAction(FExecuteAction::CreateUObject(const_cast<UNiagaraNodeWithDynamicPins*>(this), &UNiagaraNodeWithDynamicPins::RenameDynamicPinFromMenu, const_cast<UEdGraphPin*>(Context->Pin))));
+		}
+
 		if (CanRemovePin(Context->Pin))
 		{
 			Section.AddMenuEntry(
@@ -394,7 +404,10 @@ void UNiagaraNodeWithDynamicPins::PinNameTextCommitted(const FText& Text, ETextC
 	}
 }
 
-
+void UNiagaraNodeWithDynamicPins::RenameDynamicPinFromMenu(UEdGraphPin* Pin)
+{
+	SetIsPinRenamePending(Pin, true);
+}
 
 void UNiagaraNodeWithDynamicPins::RemoveDynamicPinFromMenu(UEdGraphPin* Pin)
 {
diff --git a/Engine/Plugins/FX/Niagara/Source/NiagaraEditor/Private/SNiagaraParameterName.cpp b/Engine/Plugins/FX/Niagara/Source/NiagaraEditor/Private/SNiagaraParameterName.cpp
--- a/Engine/Plugins/FX/Niagara/Source/NiagaraEditor/Private/SNiagaraParameterName.cpp
+++ b/Engine/Plugins/FX/Niagara/Source/NiagaraEditor/Private/SNiagaraParameterName.cpp
@@ -3,6 +3,7 @@
 #include "Widgets/SNiagaraParameterName.h"
 #include "NiagaraEditorSettings.h"
 #include "NiagaraEditorStyle.h"
+#include "NiagaraNodeParameterMapBase.h"
 
 #include "Widgets/SBoxPanel.h"
 #include "Widgets/Layout/SBorder.h"
@@ -20,6 +21,7 @@ void SNiagaraParameterName::Construct(const FArguments& InArgs)
 	OnNameChangedDelegate = InArgs._OnNameChanged;
 	OnDoubleClickedDelegate = InArgs._OnDoubleClicked;
 	IsSelected = InArgs._IsSelected;
+	DecoratorHAlign = InArgs._DecoratorHAlign;
 	Decorator = InArgs._Decorator.Widget;
 
 	UpdateContent(ParameterName.Get());
@@ -65,78 +67,74 @@ TSharedRef<SBorder> SNiagaraParameterName::CreateNamespaceWidget(FText Namespace
 void SNiagaraParameterName::UpdateContent(FName InDisplayedParameterName)
 {
 	DisplayedParameterName = InDisplayedParameterName;
-	
+
 	FString DisplayedParameterNameString = DisplayedParameterName.ToString();
-	TArray<FString> NameParts;
-	DisplayedParameterNameString.ParseIntoArray(NameParts, TEXT("."));
+	TArray<FString> NamePartStrings;
+	DisplayedParameterNameString.ParseIntoArray(NamePartStrings, TEXT("."));
 
-	if (NameParts.Num() == 0)
+	if (NamePartStrings.Num() == 0)
 	{
 		return;
 	}
 
-	TArray<FName> Namespaces;
-	for (int32 i = 0; i < NameParts.Num() - 1; i++)
+	TArray<FName> NameParts;
+	for (int32 i = 0; i < NamePartStrings.Num(); i++)
 	{
-		Namespaces.Add(*NameParts[i]);
+		NameParts.Add(*NamePartStrings[i]);
 	}
 
 	TSharedRef<SHorizontalBox> ContentBox = SNew(SHorizontalBox);
 
+	// Add the namespace widget.
 	FNiagaraNamespaceMetadata DefaultNamespaceMetadata = GetDefault<UNiagaraEditorSettings>()->GetMetaDataForNamespaces({ NAME_None });
-
-	if (Namespaces.Num() > 0)
+	FNiagaraNamespaceMetadata NamespaceMetadata = GetDefault<UNiagaraEditorSettings>()->GetMetaDataForNamespaces(NameParts);
+	TSharedPtr<SWidget> NamespaceWidget;
+	if (NamespaceMetadata.IsValid())
 	{
-		// Add the namespace widget.
-		FNiagaraNamespaceMetadata NamespaceMetadata = GetDefault<UNiagaraEditorSettings>()->GetMetaDataForNamespaces(Namespaces);
-		TSharedPtr<SWidget> NamespaceWidget;
-		if (NamespaceMetadata.IsValid())
-		{
-			Namespaces.RemoveAt(0, NamespaceMetadata.Namespaces.Num());
-			NamespaceWidget = CreateNamespaceWidget(
-				NamespaceMetadata.DisplayName.ToUpper(), NamespaceMetadata.Description,
-				NamespaceMetadata.BackgroundColor, NamespaceMetadata.ForegroundStyle);
-		}
-		else
-		{
-			FText NamespaceDisplayName = FText::FromString(FName::NameToDisplayString(Namespaces[0].ToString(), false).ToUpper());
-			Namespaces.RemoveAt(0);
-			NamespaceWidget = CreateNamespaceWidget(
-				NamespaceDisplayName, DefaultNamespaceMetadata.Description,
-				DefaultNamespaceMetadata.BackgroundColor, DefaultNamespaceMetadata.ForegroundStyle);
-		}
+		NameParts.RemoveAt(0, NamespaceMetadata.Namespaces.Num());
+		NamespaceWidget = CreateNamespaceWidget(
+			NamespaceMetadata.DisplayName.ToUpper(), NamespaceMetadata.Description,
+			NamespaceMetadata.BackgroundColor, NamespaceMetadata.ForegroundStyle);
+	}
+	else
+	{
+		FText NamespaceDisplayName = FText::FromString(FName::NameToDisplayString(NameParts[0].ToString(), false).ToUpper());
+		NameParts.RemoveAt(0);
+		NamespaceWidget = CreateNamespaceWidget(
+			NamespaceDisplayName, DefaultNamespaceMetadata.Description,
+			DefaultNamespaceMetadata.BackgroundColor, DefaultNamespaceMetadata.ForegroundStyle);
+	}
 
-		ContentBox->AddSlot()
+	ContentBox->AddSlot()
 		.VAlign(VAlign_Center)
 		.AutoWidth()
 		.Padding(0.0f, 0.0f, 5.0f, 0.0f)
 		[
 			NamespaceWidget.ToSharedRef()
 		];
-	}
 
-	// Next the namespace modifier widget is there is a namespace modifier.
-	if (Namespaces.Num() > 0)
+	// Next the namespace modifier widget if there is a namespace modifier.
+	if (NameParts.Num() > 1)
 	{
-		DisplayedNamespaceModifier = Namespaces[0];
+		DisplayedNamespaceModifier = NameParts[0];
 		FNiagaraNamespaceMetadata DisplayedNamespaceModifierMetadata = GetDefault<UNiagaraEditorSettings>()->GetMetaDataForNamespaceModifier(DisplayedNamespaceModifier);
 		if (DisplayedNamespaceModifierMetadata.IsValid() == false)
 		{
 			DisplayedNamespaceModifierMetadata = DefaultNamespaceMetadata;
 		}
 
-		Namespaces.RemoveAt(0);
+		NameParts.RemoveAt(0);
 		NamespaceModifierBorder = CreateNamespaceWidget(
 			FText::FromString(FName::NameToDisplayString(DisplayedNamespaceModifier.ToString(), false).ToUpper()), DisplayedNamespaceModifierMetadata.Description,
 			DisplayedNamespaceModifierMetadata.BackgroundColor, DisplayedNamespaceModifierMetadata.ForegroundStyle);
 
 		ContentBox->AddSlot()
-		.VAlign(VAlign_Center)
-		.AutoWidth()
-		.Padding(0.0f, 0.0f, 5.0f, 0.0f)
-		[
-			NamespaceModifierBorder.ToSharedRef()
-		];
+			.VAlign(VAlign_Center)
+			.AutoWidth()
+			.Padding(0.0f, 0.0f, 5.0f, 0.0f)
+			[
+				NamespaceModifierBorder.ToSharedRef()
+			];
 	}
 	else
 	{
@@ -145,50 +143,59 @@ void SNiagaraParameterName::UpdateContent(FName InDisplayedParameterName)
 	}
 
 	// If there are extra namespaces found, add them to the UI without metadata.
-	for (FName Namespace : Namespaces)
+	while (NameParts.Num() > 1)
 	{
+		FName ExtraNamespace = NameParts[0];
+		NameParts.RemoveAt(0);
+
 		ContentBox->AddSlot()
-		.VAlign(VAlign_Center)
-		.AutoWidth()
-		.Padding(0.0f, 0.0f, 5.0f, 0.0f)
-		[
-			CreateNamespaceWidget(
-				FText::FromString(FName::NameToDisplayString(Namespace.ToString(), false).ToUpper()), DefaultNamespaceMetadata.Description,
-				DefaultNamespaceMetadata.BackgroundColor, DefaultNamespaceMetadata.ForegroundStyle)
-		];
+			.VAlign(VAlign_Center)
+			.AutoWidth()
+			.Padding(0.0f, 0.0f, 5.0f, 0.0f)
+			[
+				CreateNamespaceWidget(
+					FText::FromString(FName::NameToDisplayString(ExtraNamespace.ToString(), false).ToUpper()), DefaultNamespaceMetadata.Description,
+					DefaultNamespaceMetadata.BackgroundColor, DefaultNamespaceMetadata.ForegroundStyle)
+			];
 	}
 
 	TSharedPtr<SWidget> NameWidget;
-	if (bIsReadOnly)
+	if (NameParts.Num() > 0)
 	{
-		NameWidget = SNew(STextBlock)
-			.TextStyle(ReadOnlyTextStyle)
-			.Text(FText::FromString(NameParts.Last()))
-			.HighlightText(HighlightText);
+		if (bIsReadOnly)
+		{
+			NameWidget = SNew(STextBlock)
+				.TextStyle(ReadOnlyTextStyle)
+				.Text(FText::FromName(NameParts[0]))
+				.HighlightText(HighlightText);
+		}
+		else
+		{
+			NameWidget = SAssignNew(EditableTextBlock, SInlineEditableTextBlock)
+				.Style(EditableTextStyle)
+				.Text(FText::FromName(NameParts[0]))
+				.IsSelected(IsSelected)
+				.OnVerifyTextChanged(this, &SNiagaraParameterName::VerifyNameTextChange)
+				.OnTextCommitted(this, &SNiagaraParameterName::NameTextCommitted)
+				.HighlightText(HighlightText);
+		}
 	}
-	else
+
+	if (NameWidget.IsValid())
 	{
-		NameWidget = SAssignNew(EditableTextBlock, SInlineEditableTextBlock)
-		.Style(EditableTextStyle)
-		.Text(FText::FromString(NameParts.Last()))
-		.IsSelected(IsSelected)
-		.OnVerifyTextChanged(this, &SNiagaraParameterName::VerifyNameTextChange)
-		.OnTextCommitted(this, &SNiagaraParameterName::NameTextCommitted)
-		.HighlightText(HighlightText);
+		ContentBox->AddSlot()
+		.VAlign(VAlign_Center)
+		.AutoWidth()
+		[
+			NameWidget.ToSharedRef()
+		];
 	}
 
-	ContentBox->AddSlot()
-	.VAlign(VAlign_Center)
-	.AutoWidth()
-	[
-		NameWidget.ToSharedRef()
-	];
-
 	if (Decorator.IsValid())
 	{
 		ContentBox->AddSlot()
 		.VAlign(VAlign_Center)
-		.AutoWidth()
+		.HAlign(DecoratorHAlign)
 		.Padding(3.0f, 0.0f, 0.0f, 0.0f)
 		[
 			Decorator.ToSharedRef()
@@ -380,4 +387,37 @@ void SNiagaraParameterNameTextBlock::EnterEditingMode()
 void SNiagaraParameterNameTextBlock::EnterNamespaceModifierEditingMode()
 {
 	ParameterName->EnterNamespaceModifierEditingMode();
+}
+
+void SNiagaraParameterNamePinLabel::Construct(const FArguments& InArgs, UEdGraphPin* InTargetPin)
+{
+	TargetPin = InTargetPin;
+
+	SNiagaraParameterNameTextBlock::Construct(SNiagaraParameterNameTextBlock::FArguments()
+		.EditableTextStyle(InArgs._EditableTextStyle)
+		.ParameterText(InArgs._ParameterText)
+		.IsReadOnly(InArgs._IsReadOnly)
+		.HighlightText(InArgs._HighlightText)
+		.OnVerifyTextChanged(InArgs._OnVerifyTextChanged)
+		.OnTextCommitted(InArgs._OnTextCommitted)
+		.IsSelected(InArgs._IsSelected)
+		.Decorator()
+		[
+			InArgs._Decorator.Widget
+		]);
+}
+
+void SNiagaraParameterNamePinLabel::Tick(const FGeometry& AllottedGeometry, const double InCurrentTime, const float InDeltaTime)
+{
+	SNiagaraParameterNameTextBlock::Tick(AllottedGeometry, InCurrentTime, InDeltaTime);
+
+	UNiagaraNodeParameterMapBase* ParameterMapNode = Cast<UNiagaraNodeParameterMapBase>(TargetPin->GetOwningNode());
+	if (ParameterMapNode != nullptr)
+	{
+		if (ParameterMapNode->GetIsPinEditNamespaceModifierPending(TargetPin))
+		{
+			ParameterMapNode->SetIsPinEditNamespaceModifierPending(TargetPin, false);
+			EnterNamespaceModifierEditingMode();
+		}
+	}
 }
\ No newline at end of file
diff --git a/Engine/Plugins/FX/Niagara/Source/NiagaraEditor/Private/TNiagaraGraphPinEditableName.h b/Engine/Plugins/FX/Niagara/Source/NiagaraEditor/Private/TNiagaraGraphPinEditableName.h
--- a/Engine/Plugins/FX/Niagara/Source/NiagaraEditor/Private/TNiagaraGraphPinEditableName.h
+++ b/Engine/Plugins/FX/Niagara/Source/NiagaraEditor/Private/TNiagaraGraphPinEditableName.h
@@ -71,6 +71,15 @@ protected:
 
 	virtual void Tick(const FGeometry& AllottedGeometry, const double InCurrentTime, const float InDeltaTime) override
 	{
+		if (bPendingRename == false)
+		{
+			UNiagaraNode* ParentNode = Cast<UNiagaraNode>(this->GraphPinObj->GetOwningNode());
+			if (ParentNode != nullptr && ParentNode->GetIsPinRenamePending(this->GraphPinObj))
+			{
+				bPendingRename = true;
+				ParentNode->SetIsPinRenamePending(this->GraphPinObj, false);
+			}
+		}
 		if (bPendingRename)
 		{
 			if (CreatedTextBlock.IsValid())
@@ -116,31 +125,30 @@ protected:
 			return CreatedTextBlock.ToSharedRef();
 		};
 
-		if (ParentNode && ParentNode->IsPinNameEditable(this->GraphPinObj))
+		bool bIsPinEditable = ParentNode && ParentNode->IsPinNameEditable(this->GraphPinObj);
+		if (ParentNode && ParentNode->IsA<UNiagaraNodeParameterMapBase>())
 		{
 			UNiagaraGraph* NiagaraGraph = ParentNode->GetNiagaraGraph();
-			if (ParentNode->IsA<UNiagaraNodeParameterMapBase>())
+			if (NiagaraGraph->IsPinVisualWidgetProviderRegistered())
 			{
-				if (NiagaraGraph->IsPinVisualWidgetProviderRegistered())
-				{
-					return NiagaraGraph->GetPinVisualWidget(this->GraphPinObj);
-				}
-				else
-				{
-					CreatedParameterNamePinLabel = SNew(SNiagaraParameterNameTextBlock)
-						.EditableTextStyle(&FEditorStyle::Get().GetWidgetStyle<FInlineEditableTextBlockStyle>("Graph.Node.InlineEditablePinName"))
-						.ParameterText(this, &TNiagaraGraphPinEditableName<BaseClass>::GetParentPinLabel)
-						.Visibility(this, &TNiagaraGraphPinEditableName<BaseClass>::GetParentPinVisibility)
-						.OnVerifyTextChanged(this, &TNiagaraGraphPinEditableName<BaseClass>::OnVerifyTextChanged)
-						.OnTextCommitted(this, &TNiagaraGraphPinEditableName<BaseClass>::OnTextCommitted);
-					return CreatedParameterNamePinLabel.ToSharedRef();
-				}
+				return NiagaraGraph->GetPinVisualWidget(this->GraphPinObj);
 			}
 			else
 			{
-				return CreateRenamableLabelTextBlock();
+				CreatedParameterNamePinLabel = SNew(SNiagaraParameterNamePinLabel, this->GraphPinObj)
+					.EditableTextStyle(&FEditorStyle::Get().GetWidgetStyle<FInlineEditableTextBlockStyle>("Graph.Node.InlineEditablePinName"))
+					.ParameterText(this, &TNiagaraGraphPinEditableName<BaseClass>::GetParentPinLabel)
+					.IsReadOnly(bIsPinEditable == false)
+					.Visibility(this, &TNiagaraGraphPinEditableName<BaseClass>::GetParentPinVisibility)
+					.OnVerifyTextChanged(this, &TNiagaraGraphPinEditableName<BaseClass>::OnVerifyTextChanged)
+					.OnTextCommitted(this, &TNiagaraGraphPinEditableName<BaseClass>::OnTextCommitted);
+				return CreatedParameterNamePinLabel.ToSharedRef();
 			}
 		}
+		else if (bIsPinEditable)
+		{
+			return CreateRenamableLabelTextBlock();
+		}
 		else
 		{	
 			return BaseClass::GetLabelWidget(InLabelStyle);
@@ -148,6 +156,7 @@ protected:
 	}
 
 	bool bPendingRename;
+	bool bPendingEditNamespaceModifier;
 	TSharedPtr<SInlineEditableTextBlock> CreatedTextBlock;
-	TSharedPtr<SNiagaraParameterNameTextBlock> CreatedParameterNamePinLabel;
+	TSharedPtr<SNiagaraParameterNamePinLabel> CreatedParameterNamePinLabel;
 };
diff --git a/Engine/Plugins/FX/Niagara/Source/NiagaraEditor/Private/Widgets/SNiagaraGraphActionWidget.cpp b/Engine/Plugins/FX/Niagara/Source/NiagaraEditor/Private/Widgets/SNiagaraGraphActionWidget.cpp
--- a/Engine/Plugins/FX/Niagara/Source/NiagaraEditor/Private/Widgets/SNiagaraGraphActionWidget.cpp
+++ b/Engine/Plugins/FX/Niagara/Source/NiagaraEditor/Private/Widgets/SNiagaraGraphActionWidget.cpp
@@ -8,6 +8,7 @@
 #include "Widgets/Images/SImage.h"
 #include "Widgets/Text/STextBlock.h"
 #include "Widgets/SNiagaraParameterName.h"
+#include "NiagaraEditorStyle.h"
 
 #define LOCTEXT_NAMESPACE "NiagaraGraphActionWidget"
 
@@ -19,12 +20,23 @@ void SNiagaraGraphActionWidget::Construct(const FArguments& InArgs, const FCreat
 	TSharedPtr<FNiagaraMenuAction> NiagaraAction = StaticCastSharedPtr<FNiagaraMenuAction>(InCreateData->Action);
 
 	TSharedPtr<SWidget> NameWidget;
-	if (NiagaraAction->GetParameterHandle().IsSet())
+	if (NiagaraAction->GetParameterVariable().IsSet())
 	{
 		NameWidget = SNew(SNiagaraParameterName)
-			.ParameterName(NiagaraAction->GetParameterHandle()->GetParameterHandleString())
+			.ParameterName(NiagaraAction->GetParameterVariable()->GetName())
 			.IsReadOnly(true)
-			.HighlightText(InArgs._HighlightText);
+			.HighlightText(InArgs._HighlightText)
+			.DecoratorHAlign(HAlign_Right)
+			.Decorator()
+			[
+				SNew(SBox)
+				.Padding(FMargin(7.0f, 0.0f, 0.0f, 0.0f))
+				[
+					SNew(STextBlock)
+					.TextStyle(FNiagaraEditorStyle::Get(), "NiagaraEditor.ParameterName.TypeText")
+					.Text(NiagaraAction->GetParameterVariable()->GetType().GetNameText())
+				]
+			];
 	}
 	else
 	{
diff --git a/Engine/Plugins/FX/Niagara/Source/NiagaraEditor/Private/Widgets/SNiagaraParameterMapPaletteItem.cpp b/Engine/Plugins/FX/Niagara/Source/NiagaraEditor/Private/Widgets/SNiagaraParameterMapPaletteItem.cpp
--- a/Engine/Plugins/FX/Niagara/Source/NiagaraEditor/Private/Widgets/SNiagaraParameterMapPaletteItem.cpp
+++ b/Engine/Plugins/FX/Niagara/Source/NiagaraEditor/Private/Widgets/SNiagaraParameterMapPaletteItem.cpp
@@ -10,6 +10,8 @@
 #include "Widgets/Images/SImage.h"
 #include "ScopedTransaction.h"
 #include "Widgets/SNiagaraParameterName.h"
+#include "Widgets/SNiagaraParameterMapView.h"
+#include "NiagaraEditorSettings.h"
 
 #define LOCTEXT_NAMESPACE "NiagaraParameterMapPalleteItem"
 
@@ -53,13 +55,18 @@ void SNiagaraParameterMapPalleteItem::Construct(const FArguments& InArgs, FCreat
 	bWasCreated = ParameterAction->ReferenceCollection.ContainsByPredicate(
 		[](const FNiagaraGraphParameterReferenceCollection& ReferenceCollection) { return ReferenceCollection.WasCreated(); });
 
+	TArray<FName> Namespaces;
+	NiagaraParameterMapSectionID::OnGetSectionNamespaces((NiagaraParameterMapSectionID::Type)InCreateData->Action->SectionID, Namespaces);
+	FNiagaraNamespaceMetadata NamespaceMetadata = GetDefault<UNiagaraEditorSettings>()->GetMetaDataForNamespaces(Namespaces);
+	bool bForceReadOnly = NamespaceMetadata.IsValid() == false || NamespaceMetadata.Options.Contains(ENiagaraNamespaceMetadataOptions::PreventEditing);
+
 	ParameterNameTextBlock = SNew(SNiagaraParameterNameTextBlock)
 		.ParameterText(this, &SNiagaraParameterMapPalleteItem::GetDisplayText)
 		.HighlightText(InCreateData->HighlightText)
 		.OnTextCommitted(this, &SNiagaraParameterMapPalleteItem::OnNameTextCommitted)
 		.OnVerifyTextChanged(this, &SNiagaraParameterMapPalleteItem::OnNameTextVerifyChanged)
 		.IsSelected(InCreateData->IsRowSelectedDelegate)
-		.IsReadOnly(InCreateData->bIsReadOnly)
+		.IsReadOnly(InCreateData->bIsReadOnly || bForceReadOnly)
 		.Decorator()
 		[
 			SNew(STextBlock)
diff --git a/Engine/Plugins/FX/Niagara/Source/NiagaraEditor/Private/Widgets/SNiagaraParameterMapView.cpp b/Engine/Plugins/FX/Niagara/Source/NiagaraEditor/Private/Widgets/SNiagaraParameterMapView.cpp
--- a/Engine/Plugins/FX/Niagara/Source/NiagaraEditor/Private/Widgets/SNiagaraParameterMapView.cpp
+++ b/Engine/Plugins/FX/Niagara/Source/NiagaraEditor/Private/Widgets/SNiagaraParameterMapView.cpp
@@ -174,22 +174,6 @@ NiagaraParameterMapSectionID::Type NiagaraParameterMapSectionID::OnGetSectionFro
 	return SectionID;
 }
 
-bool NiagaraParameterMapSectionID::GetSectionIsAdvancedForScript(const Type InSection)
-{
-	TArray<FName> SectionNamespaces;
-	OnGetSectionNamespaces(InSection, SectionNamespaces);
-	FNiagaraNamespaceMetadata NamespaceMetadata = GetDefault<UNiagaraEditorSettings>()->GetMetaDataForNamespaces(SectionNamespaces);
-	return NamespaceMetadata.IsValid() && NamespaceMetadata.Options.Contains(ENiagaraNamespaceMetadataOptions::AdvancedInScript);
-}
-
-bool NiagaraParameterMapSectionID::GetSectionIsAdvancedForSystem(const Type InSection)
-{
-	TArray<FName> SectionNamespaces;
-	OnGetSectionNamespaces(InSection, SectionNamespaces);
-	FNiagaraNamespaceMetadata NamespaceMetadata = GetDefault<UNiagaraEditorSettings>()->GetMetaDataForNamespaces(SectionNamespaces);
-	return NamespaceMetadata.IsValid() && NamespaceMetadata.Options.Contains(ENiagaraNamespaceMetadataOptions::AdvancedInSystem);
-}
-
 SNiagaraParameterMapView::~SNiagaraParameterMapView()
 {
 	// Unregister all commands for right click on action node
@@ -220,6 +204,8 @@ SNiagaraParameterMapView::~SNiagaraParameterMapView()
 	{
 		SelectedVariableObjects->OnSelectedObjectsChanged().RemoveAll(this);
 	}
+
+	UNiagaraEditorSettings::OnSettingsChanged().RemoveAll(this);
 }
 
 const FSlateBrush* SNiagaraParameterMapView::GetViewOptionsBorderBrush()
@@ -261,6 +247,8 @@ void SNiagaraParameterMapView::Construct(const FArguments& InArgs, const TArray<
 			FCanExecuteAction::CreateSP(this, &SNiagaraParameterMapView::CanCopyParameterReference));
 	}
 
+	UNiagaraEditorSettings::OnSettingsChanged().AddSP(this, &SNiagaraParameterMapView::NiagaraEditorSettingsChanged);
+
 	Refresh(false);
 
 	SAssignNew(FilterBox, SSearchBox)
@@ -394,6 +382,7 @@ void SNiagaraParameterMapView::AddParameter(FNiagaraVariable NewVariable, bool b
 	FNiagaraParameterHandle ParameterHandle;
 	bool bSuccess = false;
 
+	NiagaraParameterMapSectionID::Type SectionID = NiagaraParameterMapSectionID::OnGetSectionFromVariable(NewVariable, IsStaticSwitchParameter(NewVariable, Graphs), ParameterHandle);
 	if (ToolkitType == SCRIPT)
 	{
 		if (Graphs.Num() > 0)
@@ -433,7 +422,7 @@ void SNiagaraParameterMapView::AddParameter(FNiagaraVariable NewVariable, bool b
 		{
 			FScopedTransaction AddTransaction(LOCTEXT("AddSystemParameterTransaction", "Add parameter to system."));
 			System->Modify();
-			if (NiagaraParameterMapSectionID::OnGetSectionFromVariable(NewVariable, IsStaticSwitchParameter(NewVariable, Graphs), ParameterHandle) == NiagaraParameterMapSectionID::USER)
+			if (SectionID == NiagaraParameterMapSectionID::USER)
 			{
 				bSuccess = FNiagaraEditorUtilities::AddParameter(NewVariable, System->GetExposedParameters(), *System, nullptr);
 			}
@@ -450,7 +439,13 @@ void SNiagaraParameterMapView::AddParameter(FNiagaraVariable NewVariable, bool b
 		GraphActionMenu->SelectItemByName(NewVariable.GetName());
 		if (bEnterRenameModeOnAdd)
 		{
-			GraphActionMenu->OnRequestRenameOnActionNode();
+			TArray<FName> Namespaces;
+			NiagaraParameterMapSectionID::OnGetSectionNamespaces(SectionID, Namespaces);
+			FNiagaraNamespaceMetadata NamespaceMetadata = GetDefault<UNiagaraEditorSettings>()->GetMetaDataForNamespaces(Namespaces);
+			if (NamespaceMetadata.IsValid() && NamespaceMetadata.Options.Contains(ENiagaraNamespaceMetadataOptions::PreventEditing) == false)
+			{
+				GraphActionMenu->OnRequestRenameOnActionNode();
+			}
 		}
 	}
 }
@@ -462,7 +457,6 @@ TSharedRef<SWidget> SNiagaraParameterMapView::GetViewOptionsMenu()
 	auto ToggleShowAdvancedCategoriesActionLambda = [this]() {
 		UNiagaraEditorSettings* Settings = GetMutableDefault<UNiagaraEditorSettings>();
 		Settings->SetDisplayAdvancedParameterPanelCategories(!Settings->GetDisplayAdvancedParameterPanelCategories());
-		Refresh();
 	};
 
 	auto GetShowAdvancedCategoriesCheckStateActionLambda = []() {
@@ -513,9 +507,6 @@ TSharedRef<SWidget> SNiagaraParameterMapView::OnCreateWidgetForAction(struct FCr
 
 void SNiagaraParameterMapView::CollectAllActions(FGraphActionListBuilderBase& OutAllActions)
 {
-	UNiagaraEditorSettings* EditorSettings = GetMutableDefault<UNiagaraEditorSettings>();
-	const bool bDisplayAdvancedCategories = EditorSettings->GetDisplayAdvancedParameterPanelCategories();
-
 	if (Graphs.Num() == 0)
 	{
 		return;
@@ -543,18 +534,7 @@ void SNiagaraParameterMapView::CollectAllActions(FGraphActionListBuilderBase& Ou
 			continue;
 		}
 
-		static const TArray<NiagaraParameterMapSectionID::Type> ExcludedSystemIds = {
-			NiagaraParameterMapSectionID::MODULE_INPUT,
-			NiagaraParameterMapSectionID::MODULE_LOCAL,
-			NiagaraParameterMapSectionID::DATA_INSTANCE };
-		if (IsSystemToolkit() && ExcludedSystemIds.Contains(Section))
-		{
-			continue;
-		}
-
-		if (bDisplayAdvancedCategories == false &&
-			((IsScriptToolkit() && NiagaraParameterMapSectionID::GetSectionIsAdvancedForScript(Section)) ||
-			(IsSystemToolkit() && NiagaraParameterMapSectionID::GetSectionIsAdvancedForSystem(Section))))
+		if (HiddenSectionIDs.Contains(Section))
 		{
 			continue;
 		}
@@ -712,45 +692,11 @@ void SNiagaraParameterMapView::CollectAllActionsForSystemToolkit(TMap<FNiagaraVa
 
 void SNiagaraParameterMapView::CollectStaticSections(TArray<int32>& StaticSectionIDs)
 {
-	if (IsSystemToolkit())
-	{
-		StaticSectionIDs.Add(NiagaraParameterMapSectionID::USER);
-	}
-	else
-	{
-		StaticSectionIDs.Add(NiagaraParameterMapSectionID::MODULE_INPUT);
-		StaticSectionIDs.Add(NiagaraParameterMapSectionID::MODULE_LOCAL);
-		StaticSectionIDs.Add(NiagaraParameterMapSectionID::STATIC_SWITCH);
-		StaticSectionIDs.Add(NiagaraParameterMapSectionID::DATA_INSTANCE);
-	}
-	StaticSectionIDs.Add(NiagaraParameterMapSectionID::MODULE_OUTPUT);
-	StaticSectionIDs.Add(NiagaraParameterMapSectionID::ENGINE);
-	StaticSectionIDs.Add(NiagaraParameterMapSectionID::PARAMETERCOLLECTION);
-	StaticSectionIDs.Add(NiagaraParameterMapSectionID::SYSTEM);
-	StaticSectionIDs.Add(NiagaraParameterMapSectionID::EMITTER);
-	StaticSectionIDs.Add(NiagaraParameterMapSectionID::PARTICLE);
-	StaticSectionIDs.Add(NiagaraParameterMapSectionID::TRANSIENT);
-
-	UNiagaraEditorSettings* EditorSettings = GetMutableDefault<UNiagaraEditorSettings>();
-	if (EditorSettings->GetDisplayAdvancedParameterPanelCategories() == false)
+	for (int32 SectionID = 0; SectionID < NiagaraParameterMapSectionID::Num; SectionID++)
 	{
-		for (int i = StaticSectionIDs.Num() - 1; i > 0; --i)
+		if (HiddenSectionIDs.Contains(SectionID) == false)
 		{
-			bool bIsAdvanced = false;
-			NiagaraParameterMapSectionID::Type SectionID = NiagaraParameterMapSectionID::Type(StaticSectionIDs[i]);
-			if (IsScriptToolkit())
-			{
-				bIsAdvanced = NiagaraParameterMapSectionID::GetSectionIsAdvancedForScript(SectionID);
-			}
-			else if (IsSystemToolkit())
-			{
-				bIsAdvanced = NiagaraParameterMapSectionID::GetSectionIsAdvancedForSystem(SectionID);
-			}
-
-			if (bIsAdvanced)
-			{
-				StaticSectionIDs.RemoveAt(i);
-			}
+			StaticSectionIDs.Add(SectionID);
 		}
 	}
 }
@@ -849,7 +795,7 @@ TSharedPtr<SWidget> SNiagaraParameterMapView::OnContextMenuOpening()
 			TSharedPtr<FNiagaraParameterAction> SelectedNiagaraAction = StaticCastSharedPtr<FNiagaraParameterAction>(SelectedActions[0]);
 			if (SelectedNiagaraAction.IsValid())
 			{
-				MenuBuilder.BeginSection("NamespaceModifier", LOCTEXT("NamespaceModifierMenuSection", "Parameter Modifiers"));
+				MenuBuilder.BeginSection("NamespaceModifier", LOCTEXT("NamespaceModifierMenuSection", "Namespace Modifiers"));
 				{
 					TAttribute<FText> AddToolTip;
 					AddToolTip.Bind(this, &SNiagaraParameterMapView::GetAddNamespaceModifierToolTip);
@@ -1043,6 +989,24 @@ void SNiagaraParameterMapView::Refresh(bool bRefreshMenu/* = true*/)
 		}
 	}
 
+	HiddenSectionIDs.Empty();
+	const UNiagaraEditorSettings* NiagaraEditorSettings = GetDefault<UNiagaraEditorSettings>();
+	bool bShowAdvanced = NiagaraEditorSettings->GetDisplayAdvancedParameterPanelCategories();
+	for (int32 SectionID = 0; SectionID < NiagaraParameterMapSectionID::Num; SectionID++)
+	{
+		TArray<FName> Namespaces;
+		NiagaraParameterMapSectionID::OnGetSectionNamespaces((NiagaraParameterMapSectionID::Type)SectionID, Namespaces);
+		FNiagaraNamespaceMetadata NamespaceMetadata = NiagaraEditorSettings->GetMetaDataForNamespaces(Namespaces);
+		if (NamespaceMetadata.IsValid() == false ||
+			(IsScriptToolkit() && NamespaceMetadata.Options.Contains(ENiagaraNamespaceMetadataOptions::HideInScript)) ||
+			(IsSystemToolkit() && NamespaceMetadata.Options.Contains(ENiagaraNamespaceMetadataOptions::HideInSystem)) ||
+			(IsScriptToolkit() && bShowAdvanced == false && NamespaceMetadata.Options.Contains(ENiagaraNamespaceMetadataOptions::AdvancedInScript)) ||
+			(IsSystemToolkit() && bShowAdvanced == false && NamespaceMetadata.Options.Contains(ENiagaraNamespaceMetadataOptions::AdvancedInSystem)))
+		{
+			HiddenSectionIDs.Add(SectionID);
+		}
+	}
+
 	if (bRefreshMenu)
 	{
 		GraphActionMenu->RefreshAllActions(true);
@@ -1158,8 +1122,9 @@ void SNiagaraParameterMapView::OnDeleteEntry()
 			TSharedPtr<FNiagaraParameterAction> ParameterAction;
 			FNiagaraParameterHandle ParameterHandle;
 			FNiagaraNamespaceMetadata NamespaceMetadata;
-			FText ErrorMessage;
-			if (GetNamespaceEditDataForSelection(ParameterAction, ParameterHandle, NamespaceMetadata, ErrorMessage))
+			FText Unused;
+			if (GetSingleParameterActionForSelection(ParameterAction, Unused) &&
+				FNiagaraParameterUtilities::GetNamespaceEditData(ParameterAction->GetParameter().GetName(), ParameterHandle, NamespaceMetadata, Unused))
 			{
 				FScopedTransaction RemoveParametersWithPins(LOCTEXT("RemoveParametersFromSystem", "Remove parameter"));
 				System->Modify();
@@ -1204,9 +1169,10 @@ FText SNiagaraParameterMapView::GetRenameOnActionNodeToolTip() const
 	FNiagaraParameterHandle ParameterHandle;
 	FNiagaraNamespaceMetadata NamespaceMetadata;
 	FText ErrorMessage;
-	if (GetNamespaceEditDataForSelection(ParameterAction, ParameterHandle, NamespaceMetadata, ErrorMessage))
+	if (GetSingleParameterActionForSelection(ParameterAction, ErrorMessage) &&
+		FNiagaraParameterUtilities::GetNamespaceEditData(ParameterAction->GetParameter().GetName(), ParameterHandle, NamespaceMetadata, ErrorMessage))
 	{
-		if (NamespaceMetadata.Options.Contains(ENiagaraNamespaceMetadataOptions::PreventRenaming))
+		if (NamespaceMetadata.Options.Contains(ENiagaraNamespaceMetadataOptions::PreventEditing))
 		{	
 			return LOCTEXT("RenamingNotSupported", "The namespace for this parameter doesn't support renaming.");
 		}
@@ -1234,9 +1200,10 @@ bool SNiagaraParameterMapView::CanRequestRenameOnActionNode() const
 	FNiagaraParameterHandle ParameterHandle;
 	FNiagaraNamespaceMetadata NamespaceMetadata;
 	FText Unused;
-	if (GetNamespaceEditDataForSelection(ParameterAction, ParameterHandle, NamespaceMetadata, Unused))
+	if (GetSingleParameterActionForSelection(ParameterAction, Unused) &&
+		FNiagaraParameterUtilities::GetNamespaceEditData(ParameterAction->GetParameter().GetName(), ParameterHandle, NamespaceMetadata, Unused))
 	{
-		bool bPreventRenameing = NamespaceMetadata.Options.Contains(ENiagaraNamespaceMetadataOptions::PreventRenaming);
+		bool bPreventRenameing = NamespaceMetadata.Options.Contains(ENiagaraNamespaceMetadataOptions::PreventEditing);
 		return bPreventRenameing == false;
 	}
 	else 
@@ -1260,10 +1227,8 @@ void SNiagaraParameterMapView::OnPostRenameActionNode(const FText& InText, FNiag
 	RenameParameter(InAction.Parameter, *InText.ToString());
 }
 
-bool SNiagaraParameterMapView::GetNamespaceEditDataForSelection(
+bool SNiagaraParameterMapView::GetSingleParameterActionForSelection(
 	TSharedPtr<FNiagaraParameterAction>& OutParameterAction,
-	FNiagaraParameterHandle& OutParameterHandle,
-	FNiagaraNamespaceMetadata& OutNamespaceMetadata,
 	FText& OutErrorMessage) const
 {
 	TArray<TSharedPtr<FEdGraphSchemaAction>> SelectedActions;
@@ -1284,94 +1249,37 @@ bool SNiagaraParameterMapView::GetNamespaceEditDataForSelection(
 		OutErrorMessage = LOCTEXT("InvalidParameterAction", "Parameter action is invalid.");
 		return false;
 	}
-
-	OutParameterHandle = FNiagaraParameterHandle(OutParameterAction->Parameter.GetName());
-	TArray<FName> NameParts = OutParameterHandle.GetHandleParts();
-	OutNamespaceMetadata = GetDefault<UNiagaraEditorSettings>()->GetMetaDataForNamespaces(NameParts);
-	if (OutNamespaceMetadata.IsValid() == false)
-	{
-		OutErrorMessage = LOCTEXT("NoMetadataForNamespace", "This parameter doesn't support editing.");
-		return false;
-	}
-
 	return true;
 }
 
-bool SNiagaraParameterMapView::GetNamespaceModifierEditDataForSelection(
-	TSharedPtr<FNiagaraParameterAction>& OutParameterAction,
-	FNiagaraParameterHandle& OutParameterHandle,
-	FNiagaraNamespaceMetadata& OutNamespaceMetadata,
-	FText& OutErrorMessage) const
-{
-	if (GetNamespaceEditDataForSelection(OutParameterAction, OutParameterHandle, OutNamespaceMetadata, OutErrorMessage))
-	{
-		if (OutNamespaceMetadata.IsValid() == false || OutNamespaceMetadata.Options.Contains(ENiagaraNamespaceMetadataOptions::CanChangeNamespaceModifier) == false)
-		{
-			OutErrorMessage = LOCTEXT("NotSupportedForThisNamespace", "This parameter doesn't support namespace modifiers.");
-			return false;
-		}
-		return true;
-	}
-	else
-	{
-		return false;
-	}
-}
-
-FName NamePartsToName(const TArray<FName>& NameParts)
-{
-	TArray<FString> NamePartStrings;
-	for (FName NamePart : NameParts)
-	{
-		NamePartStrings.Add(NamePart.ToString());
-	}
-	return *FString::Join(NamePartStrings, TEXT("."));
-}
-
 FText SNiagaraParameterMapView::GetAddNamespaceModifierToolTip() const
 {
 	TSharedPtr<FNiagaraParameterAction> ParameterAction;
-	FNiagaraParameterHandle ParameterHandle;
-	FNiagaraNamespaceMetadata NamespaceMetadata;
-	FText ErrorMessage;
-	if (GetNamespaceModifierEditDataForSelection(ParameterAction, ParameterHandle, NamespaceMetadata, ErrorMessage))
+	FText AddMessage;
+	if (GetSingleParameterActionForSelection(ParameterAction, AddMessage))
 	{
-		return LOCTEXT("AddNamespaceModifierToolTip", "Add a default namespace modifier to this parameter.");
-	}
-	else
-	{
-		return ErrorMessage;
+		FNiagaraParameterUtilities::TestCanAddNamespaceModifierWithMessage(ParameterAction->Parameter.GetName(), AddMessage);
 	}
+	return AddMessage;
 }
 
 bool SNiagaraParameterMapView::CanAddNamespaceModifier() const
 {
 	TSharedPtr<FNiagaraParameterAction> ParameterAction;
-	FNiagaraParameterHandle ParameterHandle;
-	FNiagaraNamespaceMetadata NamespaceMetadata;
 	FText Unused;
-	if(GetNamespaceModifierEditDataForSelection(ParameterAction, ParameterHandle, NamespaceMetadata, Unused))
-	{
-		int32 NumberOfNamePartsAfterNamespace = ParameterHandle.GetHandleParts().Num() - NamespaceMetadata.Namespaces.Num();
-		return NumberOfNamePartsAfterNamespace == 1;
-	}
-	return false;
+	return GetSingleParameterActionForSelection(ParameterAction, Unused) && 
+		FNiagaraParameterUtilities::TestCanAddNamespaceModifierWithMessage(ParameterAction->Parameter.GetName(), Unused);
 }
 
 void SNiagaraParameterMapView::OnAddNamespaceModifier()
 {
 	TSharedPtr<FNiagaraParameterAction> ParameterAction;
-	FNiagaraParameterHandle ParameterHandle;
-	FNiagaraNamespaceMetadata NamespaceMetadata;
 	FText Unused;
-	if (GetNamespaceModifierEditDataForSelection(ParameterAction, ParameterHandle, NamespaceMetadata, Unused))
+	if (GetSingleParameterActionForSelection(ParameterAction, Unused))
 	{
-		TArray<FName> NameParts = ParameterHandle.GetHandleParts();
-		int32 NumberOfNamePartsAfterNamespace = NameParts.Num() - NamespaceMetadata.Namespaces.Num();
-		if (NumberOfNamePartsAfterNamespace == 1)
+		FName NewName = FNiagaraParameterUtilities::AddNamespaceModifier(ParameterAction->GetParameter().GetName());
+		if(NewName != NAME_None)
 		{
-			NameParts.Insert(FNiagaraConstants::ModuleNamespace, NamespaceMetadata.Namespaces.Num());
-			FName NewName = NamePartsToName(NameParts);
 			FScopedTransaction Transaction(LOCTEXT("AddNamespaceModifierTransaction", "Add namespace modifier."));
 			RenameParameter(ParameterAction->Parameter, NewName);
 		}
@@ -1381,38 +1289,31 @@ void SNiagaraParameterMapView::OnAddNamespaceModifier()
 FText SNiagaraParameterMapView::GetRemoveNamespaceModifierToolTip() const
 {
 	TSharedPtr<FNiagaraParameterAction> ParameterAction;
-	FNiagaraParameterHandle ParameterHandle;
-	FNiagaraNamespaceMetadata NamespaceMetadata;
-	FText ErrorMessage;
-	if (GetNamespaceModifierEditDataForSelection(ParameterAction, ParameterHandle, NamespaceMetadata, ErrorMessage))
+	FText RemoveMessage;
+	if (GetSingleParameterActionForSelection(ParameterAction, RemoveMessage))
 	{
-		return LOCTEXT("RemoveNamespaceModifierToolTip", "Remove this parameter's namespace modifier.");
-	}
-	else
-	{
-		return ErrorMessage;
+		FNiagaraParameterUtilities::TestCanRemoveNamespaceModifierWithMessage(ParameterAction->Parameter.GetName(), RemoveMessage);
 	}
+	return RemoveMessage;
 }
 
 bool SNiagaraParameterMapView::CanRemoveNamespaceModifier() const
 {
-	return CanEditNamespaceModifier();
+	TSharedPtr<FNiagaraParameterAction> ParameterAction;
+	FText Unused;
+	return GetSingleParameterActionForSelection(ParameterAction, Unused) &&
+		FNiagaraParameterUtilities::TestCanRemoveNamespaceModifierWithMessage(ParameterAction->Parameter.GetName(), Unused);
 }
 
 void SNiagaraParameterMapView::OnRemoveNamespaceModifier()
 {
 	TSharedPtr<FNiagaraParameterAction> ParameterAction;
-	FNiagaraParameterHandle ParameterHandle;
-	FNiagaraNamespaceMetadata NamespaceMetadata;
 	FText Unused;
-	if (GetNamespaceModifierEditDataForSelection(ParameterAction, ParameterHandle, NamespaceMetadata, Unused))
+	if (GetSingleParameterActionForSelection(ParameterAction, Unused))
 	{
-		TArray<FName> NameParts = ParameterHandle.GetHandleParts();
-		int32 NumberOfNamePartsAfterNamespace = NameParts.Num() - NamespaceMetadata.Namespaces.Num();
-		if (NumberOfNamePartsAfterNamespace == 2)
+		FName NewName = FNiagaraParameterUtilities::RemoveNamespaceModifier(ParameterAction->GetParameter().GetName());
+		if (NewName != NAME_None)
 		{
-			NameParts.RemoveAt(NamespaceMetadata.Namespaces.Num());
-			FName NewName = NamePartsToName(NameParts);
 			FScopedTransaction Transaction(LOCTEXT("RemoveNamespaceModifierTransaction", "Remove namespace modifier."));
 			RenameParameter(ParameterAction->Parameter, NewName);
 		}
@@ -1422,31 +1323,20 @@ void SNiagaraParameterMapView::OnRemoveNamespaceModifier()
 FText SNiagaraParameterMapView::GetEditNamespaceModifierToolTip() const
 {
 	TSharedPtr<FNiagaraParameterAction> ParameterAction;
-	FNiagaraParameterHandle ParameterHandle;
-	FNiagaraNamespaceMetadata NamespaceMetadata;
-	FText ErrorMessage;
-	if (GetNamespaceModifierEditDataForSelection(ParameterAction, ParameterHandle, NamespaceMetadata, ErrorMessage))
+	FText EditMessage;
+	if (GetSingleParameterActionForSelection(ParameterAction, EditMessage))
 	{
-		return LOCTEXT("EditNamespaceModifierToolTip", "Edit this parameter's namespace modifier.");
-	}
-	else
-	{
-		return ErrorMessage;
+		FNiagaraParameterUtilities::TestCanEditNamespaceModifierWithMessage(ParameterAction->Parameter.GetName(), EditMessage);
 	}
+	return EditMessage;
 }
 
 bool SNiagaraParameterMapView::CanEditNamespaceModifier() const
 {
 	TSharedPtr<FNiagaraParameterAction> ParameterAction;
-	FNiagaraParameterHandle ParameterHandle;
-	FNiagaraNamespaceMetadata NamespaceMetadata;
 	FText Unused;
-	if (GetNamespaceModifierEditDataForSelection(ParameterAction, ParameterHandle, NamespaceMetadata, Unused))
-	{
-		int32 NumberOfNamePartsAfterNamespace = ParameterHandle.GetHandleParts().Num() - NamespaceMetadata.Namespaces.Num();
-		return NumberOfNamePartsAfterNamespace == 2;
-	}
-	return false;
+	return GetSingleParameterActionForSelection(ParameterAction, Unused) &&
+		FNiagaraParameterUtilities::TestCanEditNamespaceModifierWithMessage(ParameterAction->Parameter.GetName(), Unused);
 }
 
 void SNiagaraParameterMapView::OnEditNamespaceModifier()
@@ -1557,8 +1447,8 @@ void SNiagaraParameterMapView::RenameParameter(FNiagaraVariable& Parameter, FNam
 					{
 						FNiagaraEditorUtilities::InfoWithToastAndLog(FText::Format(
 							LOCTEXT("AddInsteadOfRenameInScriptFormat", "The parameter {0} was not added directly\n or was discovered in a nested graph and can't be renamed.\nA new parameter {1} has been added instead."),
-							FNiagaraEditorUtilities::FormatParameterNameForTextDisplay(Parameter.GetName()),
-							FNiagaraEditorUtilities::FormatParameterNameForTextDisplay(NewName)), 7.0f);
+							FNiagaraParameterUtilities::FormatParameterNameForTextDisplay(Parameter.GetName()),
+							FNiagaraParameterUtilities::FormatParameterNameForTextDisplay(NewName)), 7.0f);
 
 						FNiagaraVariable NewVariable(Parameter.GetType(), NewName);
 						bool bEnterRenameModeOnAdd = false;
@@ -1579,8 +1469,8 @@ void SNiagaraParameterMapView::RenameParameter(FNiagaraVariable& Parameter, FNam
 				// If the parameter isn't in one of the existing collections than it's a discovered parameter and instead of renaming it needs to be added.
 				FNiagaraEditorUtilities::InfoWithToastAndLog(FText::Format(
 					LOCTEXT("AddInsteadOfRenameInSystemFormat", "The parameter {0}\n was not added directly and can't be renamed.\nA new parameter {1} has been added instead."),
-					FNiagaraEditorUtilities::FormatParameterNameForTextDisplay(Parameter.GetName()),
-					FNiagaraEditorUtilities::FormatParameterNameForTextDisplay(NewName)), 7.0f);
+					FNiagaraParameterUtilities::FormatParameterNameForTextDisplay(Parameter.GetName()),
+					FNiagaraParameterUtilities::FormatParameterNameForTextDisplay(NewName)), 7.0f);
 
 				FNiagaraVariable NewVariable(Parameter.GetType(), NewName);
 				bool bEnterRenameModeOnAdd = false;
@@ -1631,6 +1521,11 @@ bool SNiagaraParameterMapView::IsStaticSwitchParameter(const FNiagaraVariable& V
 	return false;
 }
 
+void SNiagaraParameterMapView::NiagaraEditorSettingsChanged(const FString& PropertyName, const UNiagaraEditorSettings* NiagaraEditorSettings)
+{
+	Refresh();
+}
+
 /************************************************************************/
 /* SNiagaraAddParameterMenu                                             */
 /************************************************************************/
@@ -1770,6 +1665,14 @@ void SNiagaraAddParameterMenu::CollectAllActions(FGraphActionListBuilderBase& Ou
 		CollectMakeNew(OutAllActions, NiagaraParameterMapSectionID::PARTICLE);
 	}
 
+	// DataInstance
+	if (CanCollectSection(NiagaraParameterMapSectionID::DATA_INSTANCE) && IDsExcluded.Contains(NiagaraParameterMapSectionID::DATA_INSTANCE) == false && IDsExcluded.Contains(NiagaraParameterMapSectionID::PARTICLE) == false)
+	{
+		TArray<FNiagaraVariable> Variables;
+		Variables.Add(SYS_PARAM_INSTANCE_ALIVE);
+		AddParameterGroup(OutAllActions, Variables, NiagaraParameterMapSectionID::DATA_INSTANCE, FText(), FString(), true, false);
+	}
+
 	// Emitter
 	if (CanCollectSection(NiagaraParameterMapSectionID::EMITTER) && !IDsExcluded.Contains(NiagaraParameterMapSectionID::EMITTER))
 	{
@@ -1941,7 +1844,7 @@ void SNiagaraAddParameterMenu::AddParameterGroup(
 			FNiagaraMenuAction::FOnExecuteStackAction::CreateSP(this, &SNiagaraAddParameterMenu::AddParameterSelected, Variable, bCustomName, InSection)));
 		if (bForMakeNew == false)
 		{
-			Action->SetParamterHandle(FNiagaraParameterHandle(Variable.GetName()));
+			Action->SetParamterVariable(Variable);
 		}
 
 		if (Variable.IsDataInterface())
diff --git a/Engine/Plugins/FX/Niagara/Source/NiagaraEditor/Private/Widgets/SNiagaraParameterPanel.cpp b/Engine/Plugins/FX/Niagara/Source/NiagaraEditor/Private/Widgets/SNiagaraParameterPanel.cpp
--- a/Engine/Plugins/FX/Niagara/Source/NiagaraEditor/Private/Widgets/SNiagaraParameterPanel.cpp
+++ b/Engine/Plugins/FX/Niagara/Source/NiagaraEditor/Private/Widgets/SNiagaraParameterPanel.cpp
@@ -613,7 +613,7 @@ void SNiagaraAddParameterMenu2::CollectAllActions(FGraphActionListBuilderBase& O
 				FNiagaraMenuAction::FOnExecuteStackAction::CreateSP(this, &SNiagaraAddParameterMenu2::AddParameterSelected, ActionInfo.NewVariable)));
 
 			Action->IsExperimental = ActionInfo.bIsExperimental;
-			Action->SetParamterHandle(FNiagaraParameterHandle(ActionInfo.NewVariable.GetName()));
+			Action->SetParamterVariable(ActionInfo.NewVariable);
 
 			TempOutAllActions.Add(Action);
 		}
@@ -638,7 +638,7 @@ void SNiagaraAddParameterMenu2::CollectAllActions(FGraphActionListBuilderBase& O
 				TSharedPtr<FNiagaraMenuAction> Action(new FNiagaraMenuAction(
 					Category, DisplayName, Tooltip, 0, FText::GetEmpty(),
 					FNiagaraMenuAction::FOnExecuteStackAction::CreateSP(this, &SNiagaraAddParameterMenu2::AddParameterSelected, Variable)));
-				Action->SetParamterHandle(FNiagaraParameterHandle(Variable.GetName()));
+				Action->SetParamterVariable(Variable);
 
 				TempOutAllActions.Add(Action);
 			}
@@ -684,7 +684,7 @@ void SNiagaraAddParameterMenu2::CollectAllActions(FGraphActionListBuilderBase& O
 					TSharedPtr<FNiagaraMenuAction> Action(new FNiagaraMenuAction(
 						CategoryText, DisplayName, Tooltip, 0, FText::GetEmpty(),
 						FNiagaraMenuAction::FOnExecuteStackAction::CreateSP(this, &SNiagaraAddParameterMenu2::AddParameterSelected, Var)));
-					Action->SetParamterHandle(FNiagaraParameterHandle(Var.GetName()));
+					Action->SetParamterVariable(Var);
 
 					TempOutAllActions.Add(Action);
 				}
@@ -705,7 +705,7 @@ void SNiagaraAddParameterMenu2::CollectAllActions(FGraphActionListBuilderBase& O
 					TSharedPtr<FNiagaraMenuAction> Action(new FNiagaraMenuAction(
 						CategoryText, DisplayName, Tooltip, 0, FText::GetEmpty(),
 						FNiagaraMenuAction::FOnExecuteStackAction::CreateSP(this, &SNiagaraAddParameterMenu2::AddParameterSelected, Var)));
-					Action->SetParamterHandle(FNiagaraParameterHandle(Var.GetName()));
+					Action->SetParamterVariable(Var);
 
 					TempOutAllActions.Add(Action);
 
@@ -715,15 +715,15 @@ void SNiagaraAddParameterMenu2::CollectAllActions(FGraphActionListBuilderBase& O
 		}
 	}
 	auto SortOutAllActionsPred = [](const TSharedPtr<FNiagaraMenuAction>& A, const TSharedPtr<FNiagaraMenuAction>& B)->bool {
-		if (A->GetParameterHandle().IsSet() == false)
+		if (A->GetParameterVariable().IsSet() == false)
 		{
 			return false;
 		}
-		else if (B->GetParameterHandle().IsSet() == false)
+		else if (B->GetParameterVariable().IsSet() == false)
 		{
 			return true;
 		}
-		return FNiagaraEditorUtilities::GetVariableSortPriority(A->GetParameterHandle()->GetParameterHandleString(), B->GetParameterHandle()->GetParameterHandleString());
+		return FNiagaraEditorUtilities::GetVariableSortPriority(A->GetParameterVariable()->GetName(), B->GetParameterVariable()->GetName());
 	};
 	TempOutAllActions.Sort(SortOutAllActionsPred);
 	for (const TSharedPtr<FNiagaraMenuAction>& Action : TempOutAllActions)
diff --git a/Engine/Plugins/FX/Niagara/Source/NiagaraEditor/Public/NiagaraActions.h b/Engine/Plugins/FX/Niagara/Source/NiagaraEditor/Public/NiagaraActions.h
--- a/Engine/Plugins/FX/Niagara/Source/NiagaraEditor/Public/NiagaraActions.h
+++ b/Engine/Plugins/FX/Niagara/Source/NiagaraEditor/Public/NiagaraActions.h
@@ -39,11 +39,11 @@ struct NIAGARAEDITOR_API FNiagaraMenuAction : public FEdGraphSchemaAction
 
 	bool IsExperimental = false;
 
-	TOptional<FNiagaraParameterHandle> GetParameterHandle() const;
-	void SetParamterHandle(const FNiagaraParameterHandle& InParameterHandle);
+	TOptional<FNiagaraVariable> GetParameterVariable() const;
+	void SetParamterVariable(const FNiagaraVariable& InParameterVariable);
 
 private:
-	TOptional<FNiagaraParameterHandle> ParameterHandle;
+	TOptional<FNiagaraVariable> ParameterVariable;
 	FOnExecuteStackAction Action;
 	FCanExecuteStackAction CanPerformAction;
 };
diff --git a/Engine/Plugins/FX/Niagara/Source/NiagaraEditor/Public/NiagaraEditorSettings.h b/Engine/Plugins/FX/Niagara/Source/NiagaraEditor/Public/NiagaraEditorSettings.h
--- a/Engine/Plugins/FX/Niagara/Source/NiagaraEditor/Public/NiagaraEditorSettings.h
+++ b/Engine/Plugins/FX/Niagara/Source/NiagaraEditor/Public/NiagaraEditorSettings.h
@@ -30,9 +30,11 @@ struct FNiagaraNewAssetDialogConfig
 UENUM()
 enum class ENiagaraNamespaceMetadataOptions
 {
+	HideInScript,
+	HideInSystem,
 	AdvancedInScript,
 	AdvancedInSystem,
-	PreventRenaming,
+	PreventEditing,
 	CanChangeNamespaceModifier,
 	PreventCreatingInSystemEditor
 };
@@ -64,6 +66,9 @@ struct FNiagaraNamespaceMetadata
 	UPROPERTY()
 	FName ForegroundStyle;
 
+	UPROPERTY()
+	int32 SortId;
+
 	UPROPERTY()
 	TArray<ENiagaraNamespaceMetadataOptions> Options;
 
@@ -97,6 +102,12 @@ struct FNiagaraNamespaceMetadata
 		return *this;
 	}
 
+	FNiagaraNamespaceMetadata& SetSortId(int32 InSortId)
+	{
+		SortId = InSortId;
+		return *this;
+	}
+
 	FNiagaraNamespaceMetadata& AddOption(ENiagaraNamespaceMetadataOptions Option)
 	{
 		Options.Add(Option);
@@ -177,6 +188,7 @@ public:
 	void SetNewAssetDialogConfig(FName InDialogConfigKey, const FNiagaraNewAssetDialogConfig& InNewAssetDialogConfig);
 
 	FNiagaraNamespaceMetadata GetMetaDataForNamespaces(TArray<FName> Namespaces) const;
+	const TArray<FNiagaraNamespaceMetadata>& GetAllNamespaceMetadata() const;
 	FNiagaraNamespaceMetadata GetMetaDataForNamespaceModifier(FName NamespaceModifier) const;
 	
 	// Begin UDeveloperSettings Interface
@@ -197,7 +209,7 @@ private:
 	void SetupNamespaceMetadata();
 
 protected:
-	static FOnNiagaraEditorSettingsChanged SettingsChangedDelegate;
+	FOnNiagaraEditorSettingsChanged SettingsChangedDelegate;
 
 private:
 	/** Whether or not auto-compile is enabled in the editors. */
diff --git a/Engine/Plugins/FX/Niagara/Source/NiagaraEditor/Public/NiagaraEditorUtilities.h b/Engine/Plugins/FX/Niagara/Source/NiagaraEditor/Public/NiagaraEditorUtilities.h
--- a/Engine/Plugins/FX/Niagara/Source/NiagaraEditor/Public/NiagaraEditorUtilities.h
+++ b/Engine/Plugins/FX/Niagara/Source/NiagaraEditor/Public/NiagaraEditorUtilities.h
@@ -34,6 +34,7 @@ class FNiagaraEmitterViewModel;
 class FNiagaraEmitterHandleViewModel;
 enum class ECheckBoxState : uint8;
 struct FNiagaraNamespaceMetadata;
+class FNiagaraParameterHandle;
 
 namespace FNiagaraEditorUtilities
 {
@@ -293,8 +294,6 @@ namespace FNiagaraEditorUtilities
 
 	void GetReferencingFunctionCallNodes(UNiagaraScript* Script, TArray<UNiagaraNodeFunctionCall*>& OutReferencingFunctionCallNodes);
 
-	FText FormatParameterNameForTextDisplay(FName ParameterName);
-
 	// Compare two FNiagaraVariable names for the sort priority relative to the first argument VarNameA. Sorting is ordered by namespace and then alphabetized. 
 	bool GetVariableSortPriority(const FName& VarNameA, const FName& VarNameB);
 
@@ -310,4 +309,34 @@ namespace FNiagaraEditorUtilities
 namespace FNiagaraParameterUtilities
 {
 	bool DoesParameterNameMatchSearchText(FName ParameterName, const FString& SearchTextString);
+
+	FText FormatParameterNameForTextDisplay(FName ParameterName);
+
+	bool GetNamespaceEditData(
+		FName InParameterName,
+		FNiagaraParameterHandle& OutParameterHandle,
+		FNiagaraNamespaceMetadata& OutNamespaceMetadata,
+		FText& OutErrorMessage);
+
+	bool GetNamespaceModifierEditData(
+		FName InParameterName,
+		FNiagaraParameterHandle& OutParameterHandle,
+		FNiagaraNamespaceMetadata& OutNamespaceMetadata,
+		FText& OutErrorMessage);
+
+	bool TestCanChangeNamespaceWithMessage(FName ParameterName, const FNiagaraNamespaceMetadata& NewNamespaceMetadata, FText& OutMessage);
+
+	FName ChangeNamespace(FName ParameterName, const FNiagaraNamespaceMetadata& NewNamespaceMetadata);
+
+	bool TestCanAddNamespaceModifierWithMessage(FName ParameterName, FText& OutMessage);
+
+	FName AddNamespaceModifier(FName InParameterName);
+
+	bool TestCanRemoveNamespaceModifierWithMessage(FName ParameterName, FText& OutMessage);
+
+	FName RemoveNamespaceModifier(FName InParameterName);
+
+	bool TestCanEditNamespaceModifierWithMessage(FName ParameterName, FText& OutMessage);
+
+	bool TestCanRenameWithMessage(FName ParameterName, FText& OutMessage);
 };
diff --git a/Engine/Plugins/FX/Niagara/Source/NiagaraEditor/Public/NiagaraNode.h b/Engine/Plugins/FX/Niagara/Source/NiagaraEditor/Public/NiagaraNode.h
--- a/Engine/Plugins/FX/Niagara/Source/NiagaraEditor/Public/NiagaraNode.h
+++ b/Engine/Plugins/FX/Niagara/Source/NiagaraEditor/Public/NiagaraNode.h
@@ -100,6 +100,12 @@ public:
 	/** Notify the rename was cancelled.*/
 	virtual bool CancelEditablePinName(const FText& InName, UEdGraphPin* InGraphPinObj) { return false; }
 
+	/** Returns whether or not the supplied pin has a rename pending. */
+	bool GetIsPinRenamePending(const UEdGraphPin* Pin);
+
+	/** Sets whether or not the supplied pin has a rename pending. */
+	void SetIsPinRenamePending(const UEdGraphPin* Pin, bool bInIsRenamePending);
+
 	virtual void BuildParameterMapHistory(FNiagaraParameterMapHistoryBuilder& OutHistory, bool bRecursive = true, bool bFilterForCompilation = true) const;
 	
 	/** Go through all the external dependencies of this node in isolation and add them to the reference id list.*/
@@ -169,4 +175,6 @@ protected:
 	FGuid ChangeId;
 
 	FOnNodeVisualsChanged VisualsChangedDelegate;
+
+	TArray<FGuid> PinsGuidsWithRenamePending;
 };
diff --git a/Engine/Plugins/FX/Niagara/Source/NiagaraEditor/Public/NiagaraNodeWithDynamicPins.h b/Engine/Plugins/FX/Niagara/Source/NiagaraEditor/Public/NiagaraNodeWithDynamicPins.h
--- a/Engine/Plugins/FX/Niagara/Source/NiagaraEditor/Public/NiagaraNodeWithDynamicPins.h
+++ b/Engine/Plugins/FX/Niagara/Source/NiagaraEditor/Public/NiagaraNodeWithDynamicPins.h
@@ -95,6 +95,8 @@ private:
 	/** Called when a pin's name text is committed. */
 	void PinNameTextCommitted(const FText& Text, ETextCommit::Type CommitType, UEdGraphPin* Pin);
 	
+	void RenameDynamicPinFromMenu(UEdGraphPin* Pin);
+
 	void RemoveDynamicPinFromMenu(UEdGraphPin* Pin);
 
 	void MoveDynamicPinFromMenu(UEdGraphPin* Pin, int32 DirectionToMove);
diff --git a/Engine/Plugins/FX/Niagara/Source/NiagaraEditor/Public/Widgets/SNiagaraParameterMapView.h b/Engine/Plugins/FX/Niagara/Source/NiagaraEditor/Public/Widgets/SNiagaraParameterMapView.h
--- a/Engine/Plugins/FX/Niagara/Source/NiagaraEditor/Public/Widgets/SNiagaraParameterMapView.h
+++ b/Engine/Plugins/FX/Niagara/Source/NiagaraEditor/Public/Widgets/SNiagaraParameterMapView.h
@@ -27,6 +27,7 @@ struct FSlateBrush;
 class UNiagaraSystem;
 struct FNiagaraNamespaceMetadata;
 class IToolTip;
+class UNiagaraEditorSettings;
 
 /* Enums to use when grouping the blueprint members in the list panel. The order here will determine the order in the list */
 namespace NiagaraParameterMapSectionID
@@ -55,7 +56,7 @@ namespace NiagaraParameterMapSectionID
 	};
 
 	static FText OnGetSectionTitle(const NiagaraParameterMapSectionID::Type InSection);
-	static void OnGetSectionNamespaces(const NiagaraParameterMapSectionID::Type InSection, TArray<FName>& OutSectionNamespaces);
+	void OnGetSectionNamespaces(const NiagaraParameterMapSectionID::Type InSection, TArray<FName>& OutSectionNamespaces);
 	static NiagaraParameterMapSectionID::Type OnGetSectionFromVariable(const FNiagaraVariable& InVar, bool IsStaticSwitchVariable, FNiagaraParameterHandle& OutParameterHandle, const NiagaraParameterMapSectionID::Type DefaultType = NiagaraParameterMapSectionID::Type::NONE);
 	static bool GetSectionIsAdvancedForScript(const NiagaraParameterMapSectionID::Type InSection);
 	static bool GetSectionIsAdvancedForSystem(const NiagaraParameterMapSectionID::Type InSection);
@@ -96,6 +97,7 @@ public:
 	static bool IsStaticSwitchParameter(const FNiagaraVariable& Variable, const TArray<TWeakObjectPtr<UNiagaraGraph>>& Graphs);
 
 private:
+	void NiagaraEditorSettingsChanged(const FString& PropertyName, const UNiagaraEditorSettings* NiagaraEditorSettings);
 	
 	TSharedRef<SWidget> GetViewOptionsMenu();
 	static const FSlateBrush* GetViewOptionsBorderBrush();
@@ -148,16 +150,8 @@ private:
 	bool CanRequestRenameOnActionNode() const;
 	void OnPostRenameActionNode(const FText& InText, FNiagaraParameterAction& InAction);
 
-	bool GetNamespaceEditDataForSelection(
+	bool GetSingleParameterActionForSelection(
 		TSharedPtr<FNiagaraParameterAction>& OutParameterAction,
-		FNiagaraParameterHandle& OutParameterHandle,
-		FNiagaraNamespaceMetadata& OutNamespaceMetadata,
-		FText& OutErrorMessage) const;
-
-	bool GetNamespaceModifierEditDataForSelection(
-		TSharedPtr<FNiagaraParameterAction>& OutParameterAction,
-		FNiagaraParameterHandle& OutParameterHandle,
-		FNiagaraNamespaceMetadata& OutNamespaceMetadata,
 		FText& OutErrorMessage) const;
 
 	FText GetAddNamespaceModifierToolTip() const;
@@ -214,6 +208,8 @@ private:
 	EToolkitType ToolkitType;
 	TSharedPtr<FUICommandList> ToolkitCommands;
 
+	TArray<int32> HiddenSectionIDs;
+
 	bool bNeedsRefresh;
 	bool bIsAddingParameter;
 };
diff --git a/Engine/Plugins/FX/Niagara/Source/NiagaraEditor/Public/Widgets/SNiagaraParameterName.h b/Engine/Plugins/FX/Niagara/Source/NiagaraEditor/Public/Widgets/SNiagaraParameterName.h
--- a/Engine/Plugins/FX/Niagara/Source/NiagaraEditor/Public/Widgets/SNiagaraParameterName.h
+++ b/Engine/Plugins/FX/Niagara/Source/NiagaraEditor/Public/Widgets/SNiagaraParameterName.h
@@ -11,6 +11,7 @@
 
 class SInlineEditableTextBlock;
 class SBorder;
+class UEdGraphPin;
 
 class NIAGARAEDITOR_API SNiagaraParameterName : public SCompoundWidget
 {
@@ -23,6 +24,7 @@ public:
 		: _EditableTextStyle(&FEditorStyle::Get().GetWidgetStyle<FInlineEditableTextBlockStyle>("InlineEditableTextBlockStyle"))
 		, _ReadOnlyTextStyle(&FEditorStyle::Get().GetWidgetStyle<FTextBlockStyle>("NormalText"))
 		, _IsReadOnly(false)
+		, _DecoratorHAlign(HAlign_Left)
 	{}
 		SLATE_STYLE_ARGUMENT(FInlineEditableTextBlockStyle, EditableTextStyle)
 		SLATE_STYLE_ARGUMENT(FTextBlockStyle, ReadOnlyTextStyle)
@@ -33,6 +35,7 @@ public:
 		SLATE_EVENT(FOnNameChanged, OnNameChanged)
 		SLATE_EVENT(FIsSelected, IsSelected)
 		SLATE_EVENT(FPointerEventHandler, OnDoubleClicked)
+		SLATE_ARGUMENT(EHorizontalAlignment, DecoratorHAlign)
 		SLATE_NAMED_SLOT(FArguments, Decorator)
 	SLATE_END_ARGS()
 
@@ -79,6 +82,7 @@ private:
 	FName DisplayedNamespaceModifier;
 	TSharedPtr<SInlineEditableTextBlock> EditableTextBlock;
 	TSharedPtr<SBorder> NamespaceModifierBorder;
+	EHorizontalAlignment DecoratorHAlign;
 	TSharedPtr<SWidget> Decorator;
 };
 
@@ -87,11 +91,12 @@ class NIAGARAEDITOR_API SNiagaraParameterNameTextBlock : public SCompoundWidget
 public:
 	SLATE_BEGIN_ARGS(SNiagaraParameterNameTextBlock)
 		: _EditableTextStyle(&FEditorStyle::Get().GetWidgetStyle<FInlineEditableTextBlockStyle>("InlineEditableTextBlockStyle"))
+		, _IsReadOnly(false)
 	{
 		_Clipping = EWidgetClipping::OnDemand;
 	}
-		SLATE_ATTRIBUTE(FText, ParameterText)
 		SLATE_STYLE_ARGUMENT(FInlineEditableTextBlockStyle, EditableTextStyle)
+		SLATE_ATTRIBUTE(FText, ParameterText)
 		SLATE_ARGUMENT(bool, IsReadOnly)
 		SLATE_ATTRIBUTE(FText, HighlightText)
 		SLATE_EVENT(FOnVerifyTextChanged, OnVerifyTextChanged)
@@ -122,4 +127,31 @@ private:
 	mutable FText DisplayedParameterTextCache;
 	mutable FName ParameterNameCache;
 	TSharedPtr<SNiagaraParameterName> ParameterName;
+};
+
+class NIAGARAEDITOR_API SNiagaraParameterNamePinLabel : public SNiagaraParameterNameTextBlock
+{
+public:
+	SLATE_BEGIN_ARGS(SNiagaraParameterNamePinLabel)
+		: _EditableTextStyle(&FEditorStyle::Get().GetWidgetStyle<FInlineEditableTextBlockStyle>("InlineEditableTextBlockStyle"))
+		, _IsReadOnly(false)
+	{
+		_Clipping = EWidgetClipping::OnDemand;
+	}
+		SLATE_STYLE_ARGUMENT(FInlineEditableTextBlockStyle, EditableTextStyle)
+		SLATE_ATTRIBUTE(FText, ParameterText)
+		SLATE_ARGUMENT(bool, IsReadOnly)
+		SLATE_ATTRIBUTE(FText, HighlightText)
+		SLATE_EVENT(FOnVerifyTextChanged, OnVerifyTextChanged)
+		SLATE_EVENT(FOnTextCommitted, OnTextCommitted)
+		SLATE_EVENT(FIsSelected, IsSelected)
+		SLATE_NAMED_SLOT(FArguments, Decorator)
+		SLATE_END_ARGS()
+
+		void Construct(const FArguments& InArgs, UEdGraphPin* InTargetPin);
+
+	virtual void Tick(const FGeometry& AllottedGeometry, const double InCurrentTime, const float InDeltaTime) override;
+
+private:
+	UEdGraphPin* TargetPin;
 };
\ No newline at end of file
diff --git a/Engine/Plugins/FX/Niagara/Source/NiagaraEditorWidgets/Private/SNiagaraOverviewStack.cpp b/Engine/Plugins/FX/Niagara/Source/NiagaraEditorWidgets/Private/SNiagaraOverviewStack.cpp
--- a/Engine/Plugins/FX/Niagara/Source/NiagaraEditorWidgets/Private/SNiagaraOverviewStack.cpp
+++ b/Engine/Plugins/FX/Niagara/Source/NiagaraEditorWidgets/Private/SNiagaraOverviewStack.cpp
@@ -7,6 +7,8 @@
 #include "ViewModels/NiagaraSystemViewModel.h"
 #include "ViewModels/NiagaraEmitterHandleViewModel.h"
 #include "ViewModels/NiagaraSystemSelectionViewModel.h"
+#include "ViewModels/NiagaraScratchPadViewModel.h"
+#include "ViewModels/NiagaraScratchPadScriptViewModel.h"
 #include "ViewModels/Stack/NiagaraStackEntry.h"
 #include "ViewModels/Stack/NiagaraStackItemGroup.h"
 #include "ViewModels/Stack/NiagaraStackItem.h"
@@ -19,6 +21,8 @@
 #include "Stack/SNiagaraStackItemGroupAddButton.h"
 #include "Stack/SNiagaraStackIssueIcon.h"
 #include "NiagaraStackCommandContext.h"
+#include "NiagaraNodeFunctionCall.h"
+#include "NiagaraEditorModule.h"
 
 #include "Widgets/Layout/SScrollBox.h"
 #include "Widgets/Layout/SGridPanel.h"
@@ -30,6 +34,7 @@
 #include "Framework/Application/SlateApplication.h"
 #include "Framework/MultiBox/MultiBoxBuilder.h"
 #include "Framework/Commands/GenericCommands.h"
+#include "Subsystems/AssetEditorSubsystem.h"
 
 #define LOCTEXT_NAMESPACE "NiagaraOverviewStack"
 
@@ -168,6 +173,33 @@ class SNiagaraSystemOverviewEntryListRow : public STableRow<UNiagaraStackEntry*>
 		return STableRow<UNiagaraStackEntry*>::OnMouseButtonUp(MyGeometry, MouseEvent);
 	}
 
+	virtual FReply OnMouseButtonDoubleClick(const FGeometry& InMyGeometry, const FPointerEvent& InMouseEvent) override
+	{
+		UNiagaraStackModuleItem* ModuleItem = Cast<UNiagaraStackModuleItem>(StackEntry);
+		if (ModuleItem != nullptr)
+		{
+			const UNiagaraNodeFunctionCall& ModuleFunctionCall = ModuleItem->GetModuleNode();
+			if (ModuleFunctionCall.FunctionScript != nullptr)
+			{
+				if (ModuleFunctionCall.FunctionScript->IsAsset() || GbShowNiagaraDeveloperWindows > 0)
+				{
+					GEditor->GetEditorSubsystem<UAssetEditorSubsystem>()->OpenEditorForAsset(const_cast<UNiagaraScript*>(ModuleFunctionCall.FunctionScript));
+					return FReply::Handled();
+				}
+				else if (ModuleItem->IsScratchModule())
+				{
+					TSharedPtr<FNiagaraScratchPadScriptViewModel> ScratchPadScriptViewModel = ModuleItem->GetSystemViewModel()->GetScriptScratchPadViewModel()->GetViewModelForScript(ModuleFunctionCall.FunctionScript);
+					if (ScratchPadScriptViewModel.IsValid())
+					{
+						ModuleItem->GetSystemViewModel()->GetScriptScratchPadViewModel()->FocusScratchPadScriptViewModel(ScratchPadScriptViewModel.ToSharedRef());
+						return FReply::Handled();
+					}
+				}
+			}
+		}
+		return FReply::Unhandled();
+	}
+
 private:
 	FSlateColor GetBackgroundColor() const
 	{
diff --git a/Engine/Plugins/FX/Niagara/Source/NiagaraEditorWidgets/Private/Stack/SNiagaraStackFunctionInputValue.cpp b/Engine/Plugins/FX/Niagara/Source/NiagaraEditorWidgets/Private/Stack/SNiagaraStackFunctionInputValue.cpp
--- a/Engine/Plugins/FX/Niagara/Source/NiagaraEditorWidgets/Private/Stack/SNiagaraStackFunctionInputValue.cpp
+++ b/Engine/Plugins/FX/Niagara/Source/NiagaraEditorWidgets/Private/Stack/SNiagaraStackFunctionInputValue.cpp
@@ -592,7 +592,7 @@ void SNiagaraStackFunctionInputValue::CollectAllActions(FGraphActionListBuilderB
 			const FText Tooltip = FText::Format(MapInputFormat, FText::FromName(AvailableHandle.GetParameterHandleString()));
 			TSharedPtr<FNiagaraMenuAction> LinkAction(new FNiagaraMenuAction(Category, DisplayName, Tooltip, 0, FText(),
 				FNiagaraMenuAction::FOnExecuteStackAction::CreateSP(this, &SNiagaraStackFunctionInputValue::ParameterHandleSelected, AvailableHandle)));
-			LinkAction->SetParamterHandle(AvailableHandle);
+			LinkAction->SetParamterVariable(FNiagaraVariable(FunctionInput->GetInputType(), AvailableHandle.GetParameterHandleString()));
 			OutAllActions.AddAction(LinkAction, RootCategoryName);
 		}
 	}
