DXC shader validation error/warning fixes

- Random.ush
  - VoronoiNoise3D_ALU()
    - DXC was unable to determine the interation range for the first set of loops and hence could not unroll them
     - explicit loop iteration counts have been added to the unroll to resolve
    - while in here I noticed the unroll counts on the other loops are incorrect so I've updated them accordingly

- MeshParticleVertexFactory.ush & NiagaraMeshVertexFactory.ush
  - the texture coordiate pack/pack pattern utilized here confuses DXC
    - in particluar, it appears to struggle with the conditional in the loop to handle the tail interation on the odd texture coordiate count case
      - it reports this as a potential out of bounds memory access
      - unsure if the DXC generated code does so or just warns
      - regardless, refactored this pattern where found to move the final texture coordinate in the odd count case out of the loop and drive it off a preprocessor check

#rnx

diff --git a/Engine/Plugins/FX/Niagara/Shaders/Private/NiagaraMeshVertexFactory.ush b/Engine/Plugins/FX/Niagara/Shaders/Private/NiagaraMeshVertexFactory.ush
--- a/Engine/Plugins/FX/Niagara/Shaders/Private/NiagaraMeshVertexFactory.ush
+++ b/Engine/Plugins/FX/Niagara/Shaders/Private/NiagaraMeshVertexFactory.ush
@@ -260,15 +260,14 @@ FMaterialPixelParameters GetMaterialPixelParameters(FVertexFactoryInterpolantsVS
 		}
 	#else
 		UNROLL
-		for(int CoordinateIndex = 0;CoordinateIndex <  NUM_TEX_COORD_INTERPOLATORS;CoordinateIndex += 2)
+		for (int CoordinateIndex = 0; CoordinateIndex * 2 + 1 < NUM_TEX_COORD_INTERPOLATORS; ++CoordinateIndex)
 		{
-			Result.TexCoords[CoordinateIndex] = Interpolants.TexCoords[CoordinateIndex/2].xy;
-
-			if(CoordinateIndex + 1 <  NUM_TEX_COORD_INTERPOLATORS)
-			{
-				Result.TexCoords[CoordinateIndex + 1] = Interpolants.TexCoords[CoordinateIndex/2].wz;
-			}
+			Result.TexCoords[CoordinateIndex * 2] = Interpolants.TexCoords[CoordinateIndex].xy;
+			Result.TexCoords[CoordinateIndex * 2 + 1] = Interpolants.TexCoords[CoordinateIndex].wz;
 		}
+		#if NUM_TEX_COORD_INTERPOLATORS & 1
+			Result.TexCoords[NUM_TEX_COORD_INTERPOLATORS - 1] = Interpolants.TexCoords[NUM_TEX_COORD_INTERPOLATORS / 2].xy;
+		#endif // #if NUM_TEX_COORD_INTERPOLATORS & 1
 	#endif
 #endif
 
@@ -907,14 +906,14 @@ FVertexFactoryInterpolantsVSToPS VertexFactoryGetInterpolantsVSToPS(FVertexFacto
 		GetCustomInterpolators(VertexParameters, CustomizedUVs);
 
 		UNROLL
-		for(int CoordinateIndex = 0;CoordinateIndex <  NUM_TEX_COORD_INTERPOLATORS;CoordinateIndex += 2)
+		for (int CoordinateIndex = 0; CoordinateIndex * 2 + 1 < NUM_TEX_COORD_INTERPOLATORS; ++CoordinateIndex)
 		{
-				Interpolants.TexCoords[CoordinateIndex / 2].xy = CustomizedUVs[CoordinateIndex];
-			if(CoordinateIndex + 1 <  NUM_TEX_COORD_INTERPOLATORS)
-			{
-					Interpolants.TexCoords[CoordinateIndex / 2].wz = CustomizedUVs[CoordinateIndex + 1];
-			}
+			Interpolants.TexCoords[CoordinateIndex].xy = CustomizedUVs[CoordinateIndex * 2];
+			Interpolants.TexCoords[CoordinateIndex].wz = CustomizedUVs[CoordinateIndex * 2 + 1];
 		}
+		#if NUM_TEX_COORD_INTERPOLATORS & 1
+			Interpolants.TexCoords[NUM_TEX_COORD_INTERPOLATORS / 2].xy = CustomizedUVs[NUM_TEX_COORD_INTERPOLATORS - 1];
+		#endif // #if NUM_TEX_COORD_INTERPOLATORS & 1
 	#endif
 #endif
 
@@ -1013,17 +1012,17 @@ uint VertexFactoryGetPrimitiveId(FVertexFactoryInterpolantsVSToPS Interpolants)
 	FMaterialTessellationParameters GetMaterialTessellationParameters(FVertexFactoryInterpolantsVSToDS Interpolants, float3 CameraLocalWorldPosition)
 	{
 		FMaterialTessellationParameters	Result;
-	#if  NUM_TEX_COORD_INTERPOLATORS
-		UNROLL
-		for(int CoordinateIndex = 0;CoordinateIndex <  NUM_TEX_COORD_INTERPOLATORS;CoordinateIndex += 2)
-		{
-			Result.TexCoords[CoordinateIndex] = Interpolants.InterpolantsVSToPS.TexCoords[CoordinateIndex/2].xy;
-			if(CoordinateIndex + 1 <  NUM_TEX_COORD_INTERPOLATORS)
+		#if NUM_TEX_COORD_INTERPOLATORS
+			UNROLL
+			for (int CoordinateIndex = 0; CoordinateIndex * 2 + 1 < NUM_TEX_COORD_INTERPOLATORS; ++CoordinateIndex)
 			{
-				Result.TexCoords[CoordinateIndex + 1] = Interpolants.InterpolantsVSToPS.TexCoords[CoordinateIndex/2].wz;
+				Result.TexCoords[CoordinateIndex * 2] = Interpolants.InterpolantsVSToPS.TexCoords[CoordinateIndex].xy;
+				Result.TexCoords[CoordinateIndex * 2 + 1] = Interpolants.InterpolantsVSToPS.TexCoords[CoordinateIndex].wz;
 			}
-		}
-	#endif
+			#if NUM_TEX_COORD_INTERPOLATORS & 1
+				Result.TexCoords[NUM_TEX_COORD_INTERPOLATORS - 1] = Interpolants.InterpolantsVSToPS.TexCoords[NUM_TEX_COORD_INTERPOLATORS / 2].xy;
+			#endif // #if NUM_TEX_COORD_INTERPOLATORS & 1
+		#endif
 
 		Result.VertexColor = float4(1,1,1,1);
 
diff --git a/Engine/Shaders/Private/MeshParticleVertexFactory.ush b/Engine/Shaders/Private/MeshParticleVertexFactory.ush
--- a/Engine/Shaders/Private/MeshParticleVertexFactory.ush
+++ b/Engine/Shaders/Private/MeshParticleVertexFactory.ush
@@ -247,15 +247,14 @@ FMaterialPixelParameters GetMaterialPixelParameters(FVertexFactoryInterpolantsVS
 		}
 	#else
 		UNROLL
-		for(int CoordinateIndex = 0;CoordinateIndex < NUM_TEX_COORD_INTERPOLATORS;CoordinateIndex += 2)
+		for (int CoordinateIndex = 0; CoordinateIndex * 2 + 1 < NUM_TEX_COORD_INTERPOLATORS; ++CoordinateIndex)
 		{
-			Result.TexCoords[CoordinateIndex] = Interpolants.TexCoords[CoordinateIndex/2].xy;
-
-			if(CoordinateIndex + 1 < NUM_TEX_COORD_INTERPOLATORS)
-			{
-				Result.TexCoords[CoordinateIndex + 1] = Interpolants.TexCoords[CoordinateIndex/2].wz;
-			}
+			Result.TexCoords[CoordinateIndex * 2] = Interpolants.TexCoords[CoordinateIndex].xy;
+			Result.TexCoords[CoordinateIndex * 2 + 1] = Interpolants.TexCoords[CoordinateIndex].wz;
 		}
+		#if NUM_TEX_COORD_INTERPOLATORS & 1
+			Result.TexCoords[NUM_TEX_COORD_INTERPOLATORS - 1] = Interpolants.TexCoords[NUM_TEX_COORD_INTERPOLATORS / 2].xy;
+		#endif // #if NUM_TEX_COORD_INTERPOLATORS & 1
 	#endif
 #endif
 
@@ -549,14 +548,14 @@ FVertexFactoryInterpolantsVSToPS VertexFactoryGetInterpolantsVSToPS(FVertexFacto
 		GetCustomInterpolators(VertexParameters, CustomizedUVs);
 
 		UNROLL
-		for(int CoordinateIndex = 0;CoordinateIndex < NUM_TEX_COORD_INTERPOLATORS;CoordinateIndex += 2)
+		for (int CoordinateIndex = 0; CoordinateIndex * 2 + 1 < NUM_TEX_COORD_INTERPOLATORS; ++CoordinateIndex)
 		{
-			Interpolants.TexCoords[CoordinateIndex / 2].xy = CustomizedUVs[CoordinateIndex];
-			if(CoordinateIndex + 1 < NUM_TEX_COORD_INTERPOLATORS)
-			{
-				Interpolants.TexCoords[CoordinateIndex / 2].wz = CustomizedUVs[CoordinateIndex + 1];
-			}
+			Interpolants.TexCoords[CoordinateIndex].xy = CustomizedUVs[CoordinateIndex * 2];
+			Interpolants.TexCoords[CoordinateIndex].wz = CustomizedUVs[CoordinateIndex * 2 + 1];
 		}
+		#if NUM_TEX_COORD_INTERPOLATORS & 1
+			Interpolants.TexCoords[NUM_TEX_COORD_INTERPOLATORS / 2].xy = CustomizedUVs[NUM_TEX_COORD_INTERPOLATORS - 1];
+		#endif // #if NUM_TEX_COORD_INTERPOLATORS & 1
 	#endif 
 #endif
 
@@ -630,17 +629,17 @@ float4 VertexFactoryGetPreviousWorldPosition(FVertexFactoryInput Input, FVertexF
 	FMaterialTessellationParameters GetMaterialTessellationParameters(FVertexFactoryInterpolantsVSToDS Interpolants, float3 CameraLocalWorldPosition)
 	{
 		FMaterialTessellationParameters	Result;
-	#if NUM_TEX_COORD_INTERPOLATORS
-		UNROLL
-		for(int CoordinateIndex = 0;CoordinateIndex < NUM_TEX_COORD_INTERPOLATORS;CoordinateIndex += 2)
-		{
-			Result.TexCoords[CoordinateIndex] = Interpolants.InterpolantsVSToPS.TexCoords[CoordinateIndex/2].xy;
-			if(CoordinateIndex + 1 < NUM_TEX_COORD_INTERPOLATORS)
+		#if NUM_TEX_COORD_INTERPOLATORS
+			UNROLL
+			for (int CoordinateIndex = 0; CoordinateIndex * 2 + 1 < NUM_TEX_COORD_INTERPOLATORS; ++CoordinateIndex)
 			{
-				Result.TexCoords[CoordinateIndex + 1] = Interpolants.InterpolantsVSToPS.TexCoords[CoordinateIndex/2].wz;
+				Result.TexCoords[CoordinateIndex * 2] = Interpolants.InterpolantsVSToPS.TexCoords[CoordinateIndex].xy;
+				Result.TexCoords[CoordinateIndex * 2 + 1] = Interpolants.InterpolantsVSToPS.TexCoords[CoordinateIndex].wz;
 			}
-		}
-	#endif
+			#if NUM_TEX_COORD_INTERPOLATORS & 1
+				Result.TexCoords[NUM_TEX_COORD_INTERPOLATORS - 1] = Interpolants.InterpolantsVSToPS.TexCoords[NUM_TEX_COORD_INTERPOLATORS / 2].xy;
+			#endif // #if NUM_TEX_COORD_INTERPOLATORS & 1
+		#endif
 
 		Result.VertexColor = float4(1,1,1,1);
 
diff --git a/Engine/Shaders/Private/Random.ush b/Engine/Shaders/Private/Random.ush
--- a/Engine/Shaders/Private/Random.ush
+++ b/Engine/Shaders/Private/Random.ush
@@ -643,11 +643,11 @@ float4 VoronoiNoise3D_ALU(float3 v, int Quality, bool bTiling, float RepeatSize,
 	// quality level 3: do a 3x3x3 search
 	if (Quality == 3)
 	{
-		UNROLL for (offset.x = -1; offset.x <= 1; ++offset.x)
+		UNROLL_N(3) for (offset.x = -1; offset.x <= 1; ++offset.x)
 		{
-			UNROLL for (offset.y = -1; offset.y <= 1; ++offset.y)
+			UNROLL_N(3) for (offset.y = -1; offset.y <= 1; ++offset.y)
 			{
-				UNROLL for (offset.z = -1; offset.z <= 1; ++offset.z)
+				UNROLL_N(3) for (offset.z = -1; offset.z <= 1; ++offset.z)
 				{
 					p = offset + VoronoiCornerSample(NoiseTileWrap(iv2 + offset, bTiling, RepeatSize), Quality);
 					mindist = VoronoiCompare(mindist, iv2 + p, fv2 - p, bDistanceOnly);
@@ -659,11 +659,11 @@ float4 VoronoiNoise3D_ALU(float3 v, int Quality, bool bTiling, float RepeatSize,
 	// everybody else searches a base 2x2x2 neighborhood
 	else
 	{
-		UNROLL_N(3) for (offset.x = 0; offset.x <= 1; ++offset.x)
+		UNROLL_N(2) for (offset.x = 0; offset.x <= 1; ++offset.x)
 		{
-			UNROLL_N(3) for (offset.y = 0; offset.y <= 1; ++offset.y)
+			UNROLL_N(2) for (offset.y = 0; offset.y <= 1; ++offset.y)
 			{
-				UNROLL_N(3) for (offset.z = 0; offset.z <= 1; ++offset.z)
+				UNROLL_N(2) for (offset.z = 0; offset.z <= 1; ++offset.z)
 				{
 					p = offset + VoronoiCornerSample(NoiseTileWrap(iv + offset, bTiling, RepeatSize), Quality);
 					mindist = VoronoiCompare(mindist, iv + p, fv - p, bDistanceOnly);
@@ -683,11 +683,11 @@ float4 VoronoiNoise3D_ALU(float3 v, int Quality, bool bTiling, float RepeatSize,
 	// quality level 4: add extra sets of four cells in each direction
 	if (Quality >= 4)
 	{
-		UNROLL_N(3) for (offset.x = -1; offset.x <= 2; offset.x += 3)
+		UNROLL_N(2) for (offset.x = -1; offset.x <= 2; offset.x += 3)
 		{
-			UNROLL_N(3) for (offset.y = 0; offset.y <= 1; ++offset.y)
+			UNROLL_N(2) for (offset.y = 0; offset.y <= 1; ++offset.y)
 			{
-				UNROLL_N(3) for (offset.z = 0; offset.z <= 1; ++offset.z)
+				UNROLL_N(2) for (offset.z = 0; offset.z <= 1; ++offset.z)
 				{
 					// along x axis
 					p = offset.xyz + VoronoiCornerSample(NoiseTileWrap(iv + offset.xyz, bTiling, RepeatSize), Quality);
