Niagara - Fix a crash when undoing deleting of data interface user parameters.

diff --git a/Engine/Plugins/FX/Niagara/Source/NiagaraEditor/Private/ViewModels/NiagaraSystemViewModel.cpp b/Engine/Plugins/FX/Niagara/Source/NiagaraEditor/Private/ViewModels/NiagaraSystemViewModel.cpp
--- a/Engine/Plugins/FX/Niagara/Source/NiagaraEditor/Private/ViewModels/NiagaraSystemViewModel.cpp
+++ b/Engine/Plugins/FX/Niagara/Source/NiagaraEditor/Private/ViewModels/NiagaraSystemViewModel.cpp
@@ -2064,8 +2064,11 @@ void FNiagaraSystemViewModel::BuildStackModuleData(UNiagaraScript* Script, FGuid
 
 void FNiagaraSystemViewModel::SystemChanged(UNiagaraSystem* ChangedSystem)
 {
-	check(System == ChangedSystem);
-	RefreshAll();
+	if (GIsTransacting == false)
+	{
+		check(System == ChangedSystem);
+		RefreshAll();
+	}
 }
 
 void FNiagaraSystemViewModel::StackViewModelStructureChanged()
