Niagara - Fix static switches not showing up in the parameters panel.

diff --git a/Engine/Plugins/FX/Niagara/Source/NiagaraEditor/Private/NiagaraEditorSettings.cpp b/Engine/Plugins/FX/Niagara/Source/NiagaraEditor/Private/NiagaraEditorSettings.cpp
--- a/Engine/Plugins/FX/Niagara/Source/NiagaraEditor/Private/NiagaraEditorSettings.cpp
+++ b/Engine/Plugins/FX/Niagara/Source/NiagaraEditor/Private/NiagaraEditorSettings.cpp
@@ -39,36 +39,34 @@ void UNiagaraEditorSettings::SetupNamespaceMetadata()
 			.SetDisplayName(LOCTEXT("DefaultDisplayName", "None"))
 			.SetDescription(LOCTEXT("DefaultDescription", "Arbitrary sub-namespace for specifying module specific dataset attributes, or calling nested modules."))
 			.SetBackgroundColor(FLinearColor(FColor(102, 102, 102)))
-			.AddOption(ENiagaraNamespaceMetadataOptions::PreventEditing),
+			.AddOption(ENiagaraNamespaceMetadataOptions::PreventEditingNamespace)
+			.AddOption(ENiagaraNamespaceMetadataOptions::PreventEditingNamespaceModifier)
+			.AddOption(ENiagaraNamespaceMetadataOptions::PreventEditingName),
 		FNiagaraNamespaceMetadata({FNiagaraConstants::SystemNamespace})
 			.SetDisplayName(LOCTEXT("SystemDisplayName", "System"))
 			.SetDisplayNameLong(LOCTEXT("SystemDisplayNameLong", "System Attributes"))
 			.SetDescription(LOCTEXT("SystemDescription", "Persistent attribute in the system which is written in a system\n stage and can be read anywhere."))
 			.SetBackgroundColor(FLinearColor(FColor(49, 113, 142)))
-			.SetSortId(10)
-			.AddOption(ENiagaraNamespaceMetadataOptions::CanChangeNamespaceModifier),
+			.SetSortId(10),
 		FNiagaraNamespaceMetadata({FNiagaraConstants::EmitterNamespace})
 			.SetDisplayName(LOCTEXT("EmitterDisplayName", "Emitter"))
-		.SetDisplayNameLong(LOCTEXT("EmitterDisplayNameLong", "Emitter Attributes"))
+			.SetDisplayNameLong(LOCTEXT("EmitterDisplayNameLong", "Emitter Attributes"))
 			.SetDescription(LOCTEXT("EmitterDescription", "Persistent attribute which is written in a emitter\nstage and can be read in emitter and particle stages."))
 			.SetBackgroundColor(FLinearColor(FColor(145, 99, 56)))
-			.SetSortId(20)
-			.AddOption(ENiagaraNamespaceMetadataOptions::CanChangeNamespaceModifier),
+			.SetSortId(20),
 		FNiagaraNamespaceMetadata({FNiagaraConstants::ParticleAttributeNamespace})
 			.SetDisplayName(LOCTEXT("ParticleDisplayName", "Particles"))
 			.SetDisplayNameLong(LOCTEXT("ParticleDisplayNameLong", "Particle Attributes"))
 			.SetDescription(LOCTEXT("ParticleDescription", "Persistent attribute which is written in a particle\nstage and can be read in particle stages."))
 			.SetBackgroundColor(FLinearColor(FColor(72, 130, 71)))
-			.SetSortId(30)
-			.AddOption(ENiagaraNamespaceMetadataOptions::CanChangeNamespaceModifier),
+			.SetSortId(30),
 		FNiagaraNamespaceMetadata({FNiagaraConstants::ModuleNamespace})
 			.SetDisplayName(LOCTEXT("ModuleDisplayName", "Input"))
 			.SetDisplayNameLong(LOCTEXT("ModuleDisplayNameLong", "Module Inputs"))
 			.SetDescription(LOCTEXT("ModuleDescription", "A value which exposes a module input to the system and emitter editor."))
 			.SetBackgroundColor(FLinearColor(FColor(136, 66, 65)))
 			.SetSortId(40)
-			.AddOption(ENiagaraNamespaceMetadataOptions::HideInSystem)
-			.AddOption(ENiagaraNamespaceMetadataOptions::CanChangeNamespaceModifier),
+			.AddOption(ENiagaraNamespaceMetadataOptions::HideInSystem),
 		FNiagaraNamespaceMetadata({FNiagaraConstants::OutputNamespace, FNiagaraConstants::ModuleNamespace})
 			.SetDisplayName(LOCTEXT("ModuleOutputDisplayName", "Output"))
 			.SetDisplayNameLong(LOCTEXT("ModuleOutputDisplayNameLong", "Module Outputs"))
@@ -77,15 +75,16 @@ void UNiagaraEditorSettings::SetupNamespaceMetadata()
 			.SetSortId(60)
 			.AddOption(ENiagaraNamespaceMetadataOptions::AdvancedInScript)
 			.AddOption(ENiagaraNamespaceMetadataOptions::AdvancedInSystem)
-			.AddOption(ENiagaraNamespaceMetadataOptions::PreventCreatingInSystemEditor)
-			.AddOption(ENiagaraNamespaceMetadataOptions::CanChangeNamespaceModifier),
+			.AddOption(ENiagaraNamespaceMetadataOptions::PreventCreatingInSystemEditor),
 		FNiagaraNamespaceMetadata({FNiagaraConstants::OutputNamespace})
 			.SetDisplayName(LOCTEXT("OutputDisplayName", "Output"))
 			.SetDescription(LOCTEXT("OutputDescription", "A transient value which the module author has decided might be useful to other modules further down in the stage.\nTransient values do not persist from frame to frame, or between stages, e.g. emitter to particle, or spawn to update."))
 			.SetBackgroundColor(FLinearColor(FColor(108, 87, 131)))
 			.AddOption(ENiagaraNamespaceMetadataOptions::AdvancedInScript)
 			.AddOption(ENiagaraNamespaceMetadataOptions::AdvancedInSystem)
-			.AddOption(ENiagaraNamespaceMetadataOptions::PreventEditing),
+			.AddOption(ENiagaraNamespaceMetadataOptions::PreventEditingNamespace)
+			.AddOption(ENiagaraNamespaceMetadataOptions::PreventEditingNamespaceModifier)
+			.AddOption(ENiagaraNamespaceMetadataOptions::PreventEditingName),
 		FNiagaraNamespaceMetadata({FNiagaraConstants::LocalNamespace, FNiagaraConstants::ModuleNamespace})
 			.SetDisplayName(LOCTEXT("ModuleLocalDisplayName", "Local"))
 			.SetDisplayNameLong(LOCTEXT("ModuleLocalDisplayNameLong", "Module Locals"))
@@ -93,8 +92,8 @@ void UNiagaraEditorSettings::SetupNamespaceMetadata()
 			.SetBackgroundColor(FLinearColor(FColor(191, 176, 84)))
 			.SetForegroundStyle("NiagaraEditor.ParameterName.NamespaceTextDark")
 			.SetSortId(50)
-			//.AddOption(ENiagaraNamespaceMetadataOptions::AdvancedInScript)
-			.AddOption(ENiagaraNamespaceMetadataOptions::AdvancedInSystem),
+			.AddOption(ENiagaraNamespaceMetadataOptions::HideInSystem)
+			.AddOption(ENiagaraNamespaceMetadataOptions::PreventEditingNamespaceModifier),
 		FNiagaraNamespaceMetadata({FNiagaraConstants::TransientNamespace})
 			.SetDisplayName(LOCTEXT("TransientDisplayName", "Transient"))
 			.SetDisplayNameLong(LOCTEXT("TransientDisplayNameLong", "Stage Transients"))
@@ -102,7 +101,8 @@ void UNiagaraEditorSettings::SetupNamespaceMetadata()
 			.SetBackgroundColor(FLinearColor(FColor(108, 87, 131)))
 			.SetSortId(80)
 			.AddOption(ENiagaraNamespaceMetadataOptions::AdvancedInScript)
-			.AddOption(ENiagaraNamespaceMetadataOptions::AdvancedInSystem),
+			.AddOption(ENiagaraNamespaceMetadataOptions::AdvancedInSystem)
+			.AddOption(ENiagaraNamespaceMetadataOptions::PreventEditingNamespaceModifier),
 		FNiagaraNamespaceMetadata({FNiagaraConstants::EngineNamespace})
 			.SetDisplayName(LOCTEXT("EngineDisplayName", "Engine"))
 			.SetDisplayNameLong(LOCTEXT("EngineDisplayNameLong", "Engine Provided"))
@@ -110,14 +110,17 @@ void UNiagaraEditorSettings::SetupNamespaceMetadata()
 			.SetBackgroundColor(FLinearColor(FColor(170, 170, 170)))
 			.SetForegroundStyle("NiagaraEditor.ParameterName.NamespaceTextDark")
 			.SetSortId(70)
-			.AddOption(ENiagaraNamespaceMetadataOptions::PreventEditing),
+			.AddOption(ENiagaraNamespaceMetadataOptions::PreventEditingNamespace)
+			.AddOption(ENiagaraNamespaceMetadataOptions::PreventEditingNamespaceModifier)
+			.AddOption(ENiagaraNamespaceMetadataOptions::PreventEditingName),
 		FNiagaraNamespaceMetadata({FNiagaraConstants::UserNamespace})
 			.SetDisplayName(LOCTEXT("UserDisplayName", "User"))
 			.SetDisplayNameLong(LOCTEXT("UserDisplayNameLong", "User Exposed"))
 			.SetDescription(LOCTEXT("UserDescription", "A read only value which can be initialized per system and\nmodified externally in the level, by blueprint, or by c++."))
 			.SetBackgroundColor(FLinearColor(FColor(91, 161, 194)))
 			.SetSortId(0)
-			.AddOption(ENiagaraNamespaceMetadataOptions::HideInScript),
+			.AddOption(ENiagaraNamespaceMetadataOptions::HideInScript)
+			.AddOption(ENiagaraNamespaceMetadataOptions::PreventEditingNamespaceModifier),
 		FNiagaraNamespaceMetadata({FNiagaraConstants::ParameterCollectionNamespace})
 			.SetDisplayName(LOCTEXT("NiagaraParameterCollectionDisplayName", "NPC"))
 			.SetDisplayNameLong(LOCTEXT("NiagaraParameterCollectionDisplayNameLong", "Niagara Parameter Collection"))
@@ -127,7 +130,9 @@ void UNiagaraEditorSettings::SetupNamespaceMetadata()
 			.SetSortId(90)
 			.AddOption(ENiagaraNamespaceMetadataOptions::AdvancedInScript)
 			.AddOption(ENiagaraNamespaceMetadataOptions::AdvancedInSystem)
-			.AddOption(ENiagaraNamespaceMetadataOptions::PreventEditing),
+			.AddOption(ENiagaraNamespaceMetadataOptions::PreventEditingNamespace)
+			.AddOption(ENiagaraNamespaceMetadataOptions::PreventEditingNamespaceModifier)
+			.AddOption(ENiagaraNamespaceMetadataOptions::PreventEditingName),
 		FNiagaraNamespaceMetadata({FNiagaraConstants::DataInstanceNamespace})
 			.SetDisplayName(LOCTEXT("DataInstanceDisplayName", "Data Instance"))
 			.SetDescription(LOCTEXT("DataInstanceDescription", "A special value which has a single bool IsAlive value, which determines if a particle is alive or not."))
@@ -137,7 +142,17 @@ void UNiagaraEditorSettings::SetupNamespaceMetadata()
 			.AddOption(ENiagaraNamespaceMetadataOptions::HideInSystem)
 			.AddOption(ENiagaraNamespaceMetadataOptions::AdvancedInScript)
 			.AddOption(ENiagaraNamespaceMetadataOptions::AdvancedInSystem)
-			.AddOption(ENiagaraNamespaceMetadataOptions::PreventEditing),
+			.AddOption(ENiagaraNamespaceMetadataOptions::PreventEditingNamespace)
+			.AddOption(ENiagaraNamespaceMetadataOptions::PreventEditingNamespaceModifier)
+			.AddOption(ENiagaraNamespaceMetadataOptions::PreventEditingName),
+		FNiagaraNamespaceMetadata({FNiagaraConstants::StaticSwitchNamespace})
+			.SetDisplayName(LOCTEXT("StatisSwitchDisplayName", "Static Switch"))
+			.SetDescription(LOCTEXT("StaticSwitchDescription", "Values which can only be set at edit time."))
+			.SetSortId(45)
+			.AddOption(ENiagaraNamespaceMetadataOptions::HideInSystem)
+			.AddOption(ENiagaraNamespaceMetadataOptions::PreventEditingNamespace)
+			.AddOption(ENiagaraNamespaceMetadataOptions::PreventEditingNamespaceModifier)
+			.AddOption(ENiagaraNamespaceMetadataOptions::PreventEditingName),
 	};
 
 	NamespaceModifierMetadata =
diff --git a/Engine/Plugins/FX/Niagara/Source/NiagaraEditor/Private/NiagaraEditorUtilities.cpp b/Engine/Plugins/FX/Niagara/Source/NiagaraEditor/Private/NiagaraEditorUtilities.cpp
--- a/Engine/Plugins/FX/Niagara/Source/NiagaraEditor/Private/NiagaraEditorUtilities.cpp
+++ b/Engine/Plugins/FX/Niagara/Source/NiagaraEditor/Private/NiagaraEditorUtilities.cpp
@@ -2592,7 +2592,7 @@ bool FNiagaraParameterUtilities::GetNamespaceEditData(
 	OutParameterHandle = FNiagaraParameterHandle(InParameterName);
 	TArray<FName> NameParts = OutParameterHandle.GetHandleParts();
 	OutNamespaceMetadata = GetDefault<UNiagaraEditorSettings>()->GetMetaDataForNamespaces(NameParts);
-	if (OutNamespaceMetadata.IsValid() == false || OutNamespaceMetadata.Options.Contains(ENiagaraNamespaceMetadataOptions::PreventEditing))
+	if (OutNamespaceMetadata.IsValid() == false)
 	{
 		OutErrorMessage = LOCTEXT("NoMetadataForNamespace", "This parameter doesn't support editing.");
 		return false;
@@ -2609,7 +2609,7 @@ bool FNiagaraParameterUtilities::GetNamespaceModifierEditData(
 {
 	if (GetNamespaceEditData(InParameterName, OutParameterHandle, OutNamespaceMetadata, OutErrorMessage))
 	{
-		if (OutNamespaceMetadata.IsValid() == false || OutNamespaceMetadata.Options.Contains(ENiagaraNamespaceMetadataOptions::CanChangeNamespaceModifier) == false)
+		if (OutNamespaceMetadata.Options.Contains(ENiagaraNamespaceMetadataOptions::PreventEditingNamespaceModifier))
 		{
 			OutErrorMessage = LOCTEXT("NotSupportedForThisNamespace", "This parameter doesn't support namespace modifiers.");
 			return false;
@@ -2624,13 +2624,19 @@ bool FNiagaraParameterUtilities::GetNamespaceModifierEditData(
 
 bool FNiagaraParameterUtilities::TestCanChangeNamespaceWithMessage(FName ParameterName, const FNiagaraNamespaceMetadata& NewNamespaceMetadata, FText& OutMessage)
 {
+	if (NewNamespaceMetadata.Options.Contains(ENiagaraNamespaceMetadataOptions::PreventEditingNamespace))
+	{
+		OutMessage = LOCTEXT("NewNamespaceNotSupported", "The new namespace does not support editing so it can not be assigned.");
+		return false;
+	}
+
 	FNiagaraParameterHandle ParameterHandle;
 	FNiagaraNamespaceMetadata NamespaceMetadata;
 	if (GetNamespaceEditData(ParameterName, ParameterHandle, NamespaceMetadata, OutMessage))
 	{
-		if (NewNamespaceMetadata.Options.Contains(ENiagaraNamespaceMetadataOptions::PreventEditing))
+		if (NamespaceMetadata.Options.Contains(ENiagaraNamespaceMetadataOptions::PreventEditingNamespace))
 		{
-			OutMessage = LOCTEXT("NewNamespaceIsntValid", "The new namespace does not support editing so it can not be assigned.");
+			OutMessage = LOCTEXT("ParameterDoesntSupportChangingNamespace", "This parameter doesn't support changing its namespace.");
 			return false;
 		}
 		else
@@ -2648,7 +2654,7 @@ FName FNiagaraParameterUtilities::ChangeNamespace(FName ParameterName, const FNi
 	FNiagaraNamespaceMetadata NamespaceMetadata;
 	FText Unused;
 	if (NewNamespaceMetadata.IsValid() &&
-		NewNamespaceMetadata.Options.Contains(ENiagaraNamespaceMetadataOptions::PreventEditing) == false &&
+		NewNamespaceMetadata.Options.Contains(ENiagaraNamespaceMetadataOptions::PreventEditingNamespace) == false &&
 		GetNamespaceEditData(ParameterName, ParameterHandle, NamespaceMetadata, Unused))
 	{
 		TArray<FName> NameParts = ParameterHandle.GetHandleParts();
@@ -2762,7 +2768,20 @@ bool FNiagaraParameterUtilities::TestCanRenameWithMessage(FName ParameterName, F
 {
 	FNiagaraParameterHandle ParameterHandle;
 	FNiagaraNamespaceMetadata NamespaceMetadata;
-	return GetNamespaceEditData(ParameterName, ParameterHandle, NamespaceMetadata, OutMessage);
+	if (GetNamespaceEditData(ParameterName, ParameterHandle, NamespaceMetadata, OutMessage))
+	{
+		if (NamespaceMetadata.Options.Contains(ENiagaraNamespaceMetadataOptions::PreventEditingName))
+		{
+			OutMessage = LOCTEXT("ParameterDoesntSupportChangingName", "This parameter doesn't support changing its name.");
+			return false;
+		}
+		else
+		{
+			OutMessage = LOCTEXT("ChangeParameterName", "Edit this parameter's name.");
+			return true;
+		}
+	}
+	return false;
 }
 
 #undef LOCTEXT_NAMESPACE
diff --git a/Engine/Plugins/FX/Niagara/Source/NiagaraEditor/Private/NiagaraNodeParameterMapBase.cpp b/Engine/Plugins/FX/Niagara/Source/NiagaraEditor/Private/NiagaraNodeParameterMapBase.cpp
--- a/Engine/Plugins/FX/Niagara/Source/NiagaraEditor/Private/NiagaraNodeParameterMapBase.cpp
+++ b/Engine/Plugins/FX/Niagara/Source/NiagaraEditor/Private/NiagaraNodeParameterMapBase.cpp
@@ -327,7 +327,7 @@ void UNiagaraNodeParameterMapBase::GetChangeNamespaceSubMenuForPin(UToolMenu* Me
 
 	for (const FNiagaraNamespaceMetadata& Metadata : NamespaceMetadata)
 	{
-		if (Metadata.IsValid() == false || Metadata.Options.Contains(ENiagaraNamespaceMetadataOptions::PreventEditing) || Metadata.Options.Contains(ENiagaraNamespaceMetadataOptions::HideInScript))
+		if (Metadata.IsValid() == false || Metadata.Options.Contains(ENiagaraNamespaceMetadataOptions::PreventEditingNamespace) || Metadata.Options.Contains(ENiagaraNamespaceMetadataOptions::HideInScript))
 		{
 			continue;
 		}
@@ -362,6 +362,7 @@ void UNiagaraNodeParameterMapBase::GetChangeNamespaceSubMenuForPin(UToolMenu* Me
 				SNew(SNiagaraParameterName)
 				.ParameterName(*NamespaceNameString)
 				.IsReadOnly(true)
+				.SingleNameDisplayMode(SNiagaraParameterName::ESingleNameDisplayMode::Namespace)
 				.ToolTipText(ToolTip)
 			];
 
diff --git a/Engine/Plugins/FX/Niagara/Source/NiagaraEditor/Private/SNiagaraParameterName.cpp b/Engine/Plugins/FX/Niagara/Source/NiagaraEditor/Private/SNiagaraParameterName.cpp
--- a/Engine/Plugins/FX/Niagara/Source/NiagaraEditor/Private/SNiagaraParameterName.cpp
+++ b/Engine/Plugins/FX/Niagara/Source/NiagaraEditor/Private/SNiagaraParameterName.cpp
@@ -16,6 +16,7 @@ void SNiagaraParameterName::Construct(const FArguments& InArgs)
 	ReadOnlyTextStyle = InArgs._ReadOnlyTextStyle;
 	ParameterName = InArgs._ParameterName;
 	bIsReadOnly = InArgs._IsReadOnly;
+	SingleNameDisplayMode = InArgs._SingleNameDisplayMode;
 	HighlightText = InArgs._HighlightText;
 	OnVerifyNameChangeDelegate = InArgs._OnVerifyNameChange;
 	OnNameChangedDelegate = InArgs._OnNameChanged;
@@ -85,78 +86,81 @@ void SNiagaraParameterName::UpdateContent(FName InDisplayedParameterName)
 
 	TSharedRef<SHorizontalBox> ContentBox = SNew(SHorizontalBox);
 
-	// Add the namespace widget.
-	FNiagaraNamespaceMetadata DefaultNamespaceMetadata = GetDefault<UNiagaraEditorSettings>()->GetMetaDataForNamespaces({ NAME_None });
-	FNiagaraNamespaceMetadata NamespaceMetadata = GetDefault<UNiagaraEditorSettings>()->GetMetaDataForNamespaces(NameParts);
-	TSharedPtr<SWidget> NamespaceWidget;
-	if (NamespaceMetadata.IsValid())
+	if (NameParts.Num() > 1 || SingleNameDisplayMode == ESingleNameDisplayMode::Namespace)
 	{
-		NameParts.RemoveAt(0, NamespaceMetadata.Namespaces.Num());
-		NamespaceWidget = CreateNamespaceWidget(
-			NamespaceMetadata.DisplayName.ToUpper(), NamespaceMetadata.Description,
-			NamespaceMetadata.BackgroundColor, NamespaceMetadata.ForegroundStyle);
-	}
-	else
-	{
-		FText NamespaceDisplayName = FText::FromString(FName::NameToDisplayString(NameParts[0].ToString(), false).ToUpper());
-		NameParts.RemoveAt(0);
-		NamespaceWidget = CreateNamespaceWidget(
-			NamespaceDisplayName, DefaultNamespaceMetadata.Description,
-			DefaultNamespaceMetadata.BackgroundColor, DefaultNamespaceMetadata.ForegroundStyle);
-	}
-
-	ContentBox->AddSlot()
-		.VAlign(VAlign_Center)
-		.AutoWidth()
-		.Padding(0.0f, 0.0f, 5.0f, 0.0f)
-		[
-			NamespaceWidget.ToSharedRef()
-		];
-
-	// Next the namespace modifier widget if there is a namespace modifier.
-	if (NameParts.Num() > 1)
-	{
-		DisplayedNamespaceModifier = NameParts[0];
-		FNiagaraNamespaceMetadata DisplayedNamespaceModifierMetadata = GetDefault<UNiagaraEditorSettings>()->GetMetaDataForNamespaceModifier(DisplayedNamespaceModifier);
-		if (DisplayedNamespaceModifierMetadata.IsValid() == false)
+		// Add the namespace widget.
+		FNiagaraNamespaceMetadata DefaultNamespaceMetadata = GetDefault<UNiagaraEditorSettings>()->GetMetaDataForNamespaces({ NAME_None });
+		FNiagaraNamespaceMetadata NamespaceMetadata = GetDefault<UNiagaraEditorSettings>()->GetMetaDataForNamespaces(NameParts);
+		TSharedPtr<SWidget> NamespaceWidget;
+		if (NamespaceMetadata.IsValid())
 		{
-			DisplayedNamespaceModifierMetadata = DefaultNamespaceMetadata;
+			NameParts.RemoveAt(0, NamespaceMetadata.Namespaces.Num());
+			NamespaceWidget = CreateNamespaceWidget(
+				NamespaceMetadata.DisplayName.ToUpper(), NamespaceMetadata.Description,
+				NamespaceMetadata.BackgroundColor, NamespaceMetadata.ForegroundStyle);
+		}
+		else
+		{
+			FText NamespaceDisplayName = FText::FromString(FName::NameToDisplayString(NameParts[0].ToString(), false).ToUpper());
+			NameParts.RemoveAt(0);
+			NamespaceWidget = CreateNamespaceWidget(
+				NamespaceDisplayName, DefaultNamespaceMetadata.Description,
+				DefaultNamespaceMetadata.BackgroundColor, DefaultNamespaceMetadata.ForegroundStyle);
 		}
-
-		NameParts.RemoveAt(0);
-		NamespaceModifierBorder = CreateNamespaceWidget(
-			FText::FromString(FName::NameToDisplayString(DisplayedNamespaceModifier.ToString(), false).ToUpper()), DisplayedNamespaceModifierMetadata.Description,
-			DisplayedNamespaceModifierMetadata.BackgroundColor, DisplayedNamespaceModifierMetadata.ForegroundStyle);
 
 		ContentBox->AddSlot()
 			.VAlign(VAlign_Center)
 			.AutoWidth()
 			.Padding(0.0f, 0.0f, 5.0f, 0.0f)
 			[
-				NamespaceModifierBorder.ToSharedRef()
+				NamespaceWidget.ToSharedRef()
 			];
-	}
-	else
-	{
-		DisplayedNamespaceModifier = NAME_None;
-		NamespaceModifierBorder.Reset();
-	}
 
-	// If there are extra namespaces found, add them to the UI without metadata.
-	while (NameParts.Num() > 1)
-	{
-		FName ExtraNamespace = NameParts[0];
-		NameParts.RemoveAt(0);
+		// Next the namespace modifier widget if there is a namespace modifier.
+		if (NameParts.Num() > 1)
+		{
+			DisplayedNamespaceModifier = NameParts[0];
+			FNiagaraNamespaceMetadata DisplayedNamespaceModifierMetadata = GetDefault<UNiagaraEditorSettings>()->GetMetaDataForNamespaceModifier(DisplayedNamespaceModifier);
+			if (DisplayedNamespaceModifierMetadata.IsValid() == false)
+			{
+				DisplayedNamespaceModifierMetadata = DefaultNamespaceMetadata;
+			}
+
+			NameParts.RemoveAt(0);
+			NamespaceModifierBorder = CreateNamespaceWidget(
+				FText::FromString(FName::NameToDisplayString(DisplayedNamespaceModifier.ToString(), false).ToUpper()), DisplayedNamespaceModifierMetadata.Description,
+				DisplayedNamespaceModifierMetadata.BackgroundColor, DisplayedNamespaceModifierMetadata.ForegroundStyle);
+
+			ContentBox->AddSlot()
+				.VAlign(VAlign_Center)
+				.AutoWidth()
+				.Padding(0.0f, 0.0f, 5.0f, 0.0f)
+				[
+					NamespaceModifierBorder.ToSharedRef()
+				];
+		}
+		else
+		{
+			DisplayedNamespaceModifier = NAME_None;
+			NamespaceModifierBorder.Reset();
+		}
 
-		ContentBox->AddSlot()
-			.VAlign(VAlign_Center)
-			.AutoWidth()
-			.Padding(0.0f, 0.0f, 5.0f, 0.0f)
-			[
-				CreateNamespaceWidget(
-					FText::FromString(FName::NameToDisplayString(ExtraNamespace.ToString(), false).ToUpper()), DefaultNamespaceMetadata.Description,
-					DefaultNamespaceMetadata.BackgroundColor, DefaultNamespaceMetadata.ForegroundStyle)
-			];
+		// If there are extra namespaces found, add them to the UI without metadata.
+		while (NameParts.Num() > 1)
+		{
+			FName ExtraNamespace = NameParts[0];
+			NameParts.RemoveAt(0);
+
+			ContentBox->AddSlot()
+				.VAlign(VAlign_Center)
+				.AutoWidth()
+				.Padding(0.0f, 0.0f, 5.0f, 0.0f)
+				[
+					CreateNamespaceWidget(
+						FText::FromString(FName::NameToDisplayString(ExtraNamespace.ToString(), false).ToUpper()), DefaultNamespaceMetadata.Description,
+						DefaultNamespaceMetadata.BackgroundColor, DefaultNamespaceMetadata.ForegroundStyle)
+				];
+		}
 	}
 
 	TSharedPtr<SWidget> NameWidget;
diff --git a/Engine/Plugins/FX/Niagara/Source/NiagaraEditor/Private/Widgets/SNiagaraParameterMapPaletteItem.cpp b/Engine/Plugins/FX/Niagara/Source/NiagaraEditor/Private/Widgets/SNiagaraParameterMapPaletteItem.cpp
--- a/Engine/Plugins/FX/Niagara/Source/NiagaraEditor/Private/Widgets/SNiagaraParameterMapPaletteItem.cpp
+++ b/Engine/Plugins/FX/Niagara/Source/NiagaraEditor/Private/Widgets/SNiagaraParameterMapPaletteItem.cpp
@@ -58,7 +58,7 @@ void SNiagaraParameterMapPalleteItem::Construct(const FArguments& InArgs, FCreat
 	TArray<FName> Namespaces;
 	NiagaraParameterMapSectionID::OnGetSectionNamespaces((NiagaraParameterMapSectionID::Type)InCreateData->Action->SectionID, Namespaces);
 	FNiagaraNamespaceMetadata NamespaceMetadata = GetDefault<UNiagaraEditorSettings>()->GetMetaDataForNamespaces(Namespaces);
-	bool bForceReadOnly = NamespaceMetadata.IsValid() == false || NamespaceMetadata.Options.Contains(ENiagaraNamespaceMetadataOptions::PreventEditing);
+	bool bForceReadOnly = NamespaceMetadata.IsValid() == false || NamespaceMetadata.Options.Contains(ENiagaraNamespaceMetadataOptions::PreventEditingName);
 
 	ParameterNameTextBlock = SNew(SNiagaraParameterNameTextBlock)
 		.ParameterText(this, &SNiagaraParameterMapPalleteItem::GetDisplayText)
diff --git a/Engine/Plugins/FX/Niagara/Source/NiagaraEditor/Private/Widgets/SNiagaraParameterMapView.cpp b/Engine/Plugins/FX/Niagara/Source/NiagaraEditor/Private/Widgets/SNiagaraParameterMapView.cpp
--- a/Engine/Plugins/FX/Niagara/Source/NiagaraEditor/Private/Widgets/SNiagaraParameterMapView.cpp
+++ b/Engine/Plugins/FX/Niagara/Source/NiagaraEditor/Private/Widgets/SNiagaraParameterMapView.cpp
@@ -442,7 +442,7 @@ void SNiagaraParameterMapView::AddParameter(FNiagaraVariable NewVariable, bool b
 			TArray<FName> Namespaces;
 			NiagaraParameterMapSectionID::OnGetSectionNamespaces(SectionID, Namespaces);
 			FNiagaraNamespaceMetadata NamespaceMetadata = GetDefault<UNiagaraEditorSettings>()->GetMetaDataForNamespaces(Namespaces);
-			if (NamespaceMetadata.IsValid() && NamespaceMetadata.Options.Contains(ENiagaraNamespaceMetadataOptions::PreventEditing) == false)
+			if (NamespaceMetadata.IsValid() && NamespaceMetadata.Options.Contains(ENiagaraNamespaceMetadataOptions::PreventEditingName) == false)
 			{
 				GraphActionMenu->OnRequestRenameOnActionNode();
 			}
@@ -1172,7 +1172,7 @@ FText SNiagaraParameterMapView::GetRenameOnActionNodeToolTip() const
 	if (GetSingleParameterActionForSelection(ParameterAction, ErrorMessage) &&
 		FNiagaraParameterUtilities::GetNamespaceEditData(ParameterAction->GetParameter().GetName(), ParameterHandle, NamespaceMetadata, ErrorMessage))
 	{
-		if (NamespaceMetadata.Options.Contains(ENiagaraNamespaceMetadataOptions::PreventEditing))
+		if (NamespaceMetadata.Options.Contains(ENiagaraNamespaceMetadataOptions::PreventEditingName))
 		{	
 			return LOCTEXT("RenamingNotSupported", "The namespace for this parameter doesn't support renaming.");
 		}
@@ -1203,7 +1203,7 @@ bool SNiagaraParameterMapView::CanRequestRenameOnActionNode() const
 	if (GetSingleParameterActionForSelection(ParameterAction, Unused) &&
 		FNiagaraParameterUtilities::GetNamespaceEditData(ParameterAction->GetParameter().GetName(), ParameterHandle, NamespaceMetadata, Unused))
 	{
-		bool bPreventRenameing = NamespaceMetadata.Options.Contains(ENiagaraNamespaceMetadataOptions::PreventEditing);
+		bool bPreventRenameing = NamespaceMetadata.Options.Contains(ENiagaraNamespaceMetadataOptions::PreventEditingName);
 		return bPreventRenameing == false;
 	}
 	else 
diff --git a/Engine/Plugins/FX/Niagara/Source/NiagaraEditor/Public/NiagaraEditorSettings.h b/Engine/Plugins/FX/Niagara/Source/NiagaraEditor/Public/NiagaraEditorSettings.h
--- a/Engine/Plugins/FX/Niagara/Source/NiagaraEditor/Public/NiagaraEditorSettings.h
+++ b/Engine/Plugins/FX/Niagara/Source/NiagaraEditor/Public/NiagaraEditorSettings.h
@@ -34,8 +34,9 @@ enum class ENiagaraNamespaceMetadataOptions
 	HideInSystem,
 	AdvancedInScript,
 	AdvancedInSystem,
-	PreventEditing,
-	CanChangeNamespaceModifier,
+	PreventEditingNamespace,
+	PreventEditingNamespaceModifier,
+	PreventEditingName,
 	PreventCreatingInSystemEditor
 };
 
diff --git a/Engine/Plugins/FX/Niagara/Source/NiagaraEditor/Public/Widgets/SNiagaraParameterName.h b/Engine/Plugins/FX/Niagara/Source/NiagaraEditor/Public/Widgets/SNiagaraParameterName.h
--- a/Engine/Plugins/FX/Niagara/Source/NiagaraEditor/Public/Widgets/SNiagaraParameterName.h
+++ b/Engine/Plugins/FX/Niagara/Source/NiagaraEditor/Public/Widgets/SNiagaraParameterName.h
@@ -19,17 +19,27 @@ public:
 	DECLARE_DELEGATE_OneParam(FOnNameChanged, FName /* InNewName */);
 	DECLARE_DELEGATE_RetVal_TwoParams(bool, FOnVerifyNameChange, FName /* InNewName */, FText& /* OutErrorMessage */)
 
+public: 
+	enum class ESingleNameDisplayMode
+	{
+		Namespace,
+		Name
+	};
 public:
 	SLATE_BEGIN_ARGS(SNiagaraParameterName) 
 		: _EditableTextStyle(&FEditorStyle::Get().GetWidgetStyle<FInlineEditableTextBlockStyle>("InlineEditableTextBlockStyle"))
 		, _ReadOnlyTextStyle(&FEditorStyle::Get().GetWidgetStyle<FTextBlockStyle>("NormalText"))
 		, _IsReadOnly(false)
+		, _SingleNameDisplayMode(ESingleNameDisplayMode::Name)
 		, _DecoratorHAlign(HAlign_Left)
-	{}
+	{
+		_Clipping = EWidgetClipping::OnDemand;
+	}
 		SLATE_STYLE_ARGUMENT(FInlineEditableTextBlockStyle, EditableTextStyle)
 		SLATE_STYLE_ARGUMENT(FTextBlockStyle, ReadOnlyTextStyle)
 		SLATE_ATTRIBUTE(FName, ParameterName)
 		SLATE_ARGUMENT(bool, IsReadOnly)
+		SLATE_ARGUMENT(ESingleNameDisplayMode, SingleNameDisplayMode)
 		SLATE_ATTRIBUTE(FText, HighlightText)
 		SLATE_EVENT(FOnVerifyNameChange, OnVerifyNameChange)
 		SLATE_EVENT(FOnNameChanged, OnNameChanged)
@@ -73,6 +83,7 @@ private:
 	const FTextBlockStyle* ReadOnlyTextStyle;
 	TAttribute<FName> ParameterName;
 	bool bIsReadOnly;
+	ESingleNameDisplayMode SingleNameDisplayMode;
 	TAttribute<FText> HighlightText;
 	FOnVerifyNameChange OnVerifyNameChangeDelegate;
 	FOnNameChanged OnNameChangedDelegate;
diff --git a/Engine/Plugins/FX/Niagara/Source/NiagaraEditorWidgets/Private/Stack/SNiagaraStackFunctionInputValue.cpp b/Engine/Plugins/FX/Niagara/Source/NiagaraEditorWidgets/Private/Stack/SNiagaraStackFunctionInputValue.cpp
--- a/Engine/Plugins/FX/Niagara/Source/NiagaraEditorWidgets/Private/Stack/SNiagaraStackFunctionInputValue.cpp
+++ b/Engine/Plugins/FX/Niagara/Source/NiagaraEditorWidgets/Private/Stack/SNiagaraStackFunctionInputValue.cpp
@@ -486,7 +486,6 @@ TSharedRef<SWidget> SNiagaraStackFunctionInputValue::OnGetAvailableHandleMenu()
 	.Padding(5)
 	[
 		SNew(SBox)
-		.WidthOverride(300)
 		.HeightOverride(400)
 		[
 			SNew(SVerticalBox)
