Niagara - Last major parameter UX fixes.
+ Add context menu options to the input rows for set variables nodes for editing a parameters namespace and namespace modifiers.
+ Add context menu options to the parameters panel for editing a parameters namespace.
+ Show a stack of parameter widgets in the overview on set variables nodes.
+ Fix the stack display name and tooltip for set variables modules so that it doesn't show the Namespace.Name syntax.

diff --git a/Engine/Plugins/FX/Niagara/Source/NiagaraEditor/Private/NiagaraEditorUtilities.cpp b/Engine/Plugins/FX/Niagara/Source/NiagaraEditor/Private/NiagaraEditorUtilities.cpp
--- a/Engine/Plugins/FX/Niagara/Source/NiagaraEditor/Private/NiagaraEditorUtilities.cpp
+++ b/Engine/Plugins/FX/Niagara/Source/NiagaraEditor/Private/NiagaraEditorUtilities.cpp
@@ -58,6 +58,7 @@
 #include "Framework/Notifications/NotificationManager.h"
 #include "NiagaraSimulationStageBase.h"
 #include "NiagaraEditorSettings.h"
+#include "Widgets/SNiagaraParameterName.h"
 
 #define LOCTEXT_NAMESPACE "FNiagaraEditorUtilities"
 
@@ -2622,6 +2623,59 @@ bool FNiagaraParameterUtilities::GetNamespaceModifierEditData(
 	}
 }
 
+void FNiagaraParameterUtilities::GetChangeNamespaceMenuData(FName InParameterName, EParameterContext InParameterContext, TArray<FChangeNamespaceMenuData>& OutChangeNamespaceMenuData)
+{
+	TArray<FNiagaraNamespaceMetadata> NamespaceMetadata = GetDefault<UNiagaraEditorSettings>()->GetAllNamespaceMetadata();
+	NamespaceMetadata.Sort([](const FNiagaraNamespaceMetadata& MetadataA, const FNiagaraNamespaceMetadata& MetadataB)
+		{ return MetadataA.SortId < MetadataB.SortId; });
+
+	for (const FNiagaraNamespaceMetadata& Metadata : NamespaceMetadata)
+	{
+		if (Metadata.IsValid() == false || Metadata.Options.Contains(ENiagaraNamespaceMetadataOptions::PreventEditingNamespace) ||
+			(InParameterContext == EParameterContext::Script && Metadata.Options.Contains(ENiagaraNamespaceMetadataOptions::HideInScript)) ||
+			(InParameterContext == EParameterContext::System && Metadata.Options.Contains(ENiagaraNamespaceMetadataOptions::HideInSystem)))
+		{
+			continue;
+		}
+
+		FChangeNamespaceMenuData& MenuData = OutChangeNamespaceMenuData.AddDefaulted_GetRef();
+		MenuData.Metadata = Metadata;
+
+		FText CanChangeMessage;
+		MenuData.bCanChange = FNiagaraParameterUtilities::TestCanChangeNamespaceWithMessage(InParameterName, Metadata, CanChangeMessage);
+
+		if (MenuData.bCanChange)
+		{
+			MenuData.CanChangeToolTip = FText::Format(LOCTEXT("ChangeNamespaceToolTipFormat", "{0}\n\nDescription:\n{1}"), CanChangeMessage, Metadata.Description);
+		}
+		else
+		{
+			MenuData.CanChangeToolTip = CanChangeMessage;
+		}
+
+		FString NamespaceNameString = Metadata.Namespaces[0].ToString();
+		for (int32 i = 1; i < Metadata.Namespaces.Num(); i++)
+		{
+			NamespaceNameString += TEXT(".") + Metadata.Namespaces[i].ToString();
+		}
+
+		MenuData.NamespaceParameterName = *NamespaceNameString;
+	}
+}
+
+TSharedRef<SWidget> FNiagaraParameterUtilities::CreateNamespaceMenuItemWidget(FName Namespace, FText ToolTipText)
+{
+	return SNew(SBox)
+		.Padding(FMargin(7, 2, 7, 2))
+		[
+			SNew(SNiagaraParameterName)
+			.ParameterName(Namespace)
+			.IsReadOnly(true)
+			.SingleNameDisplayMode(SNiagaraParameterName::ESingleNameDisplayMode::Namespace)
+			.ToolTipText(ToolTipText)
+		];
+}
+
 bool FNiagaraParameterUtilities::TestCanChangeNamespaceWithMessage(FName ParameterName, const FNiagaraNamespaceMetadata& NewNamespaceMetadata, FText& OutMessage)
 {
 	if (NewNamespaceMetadata.Options.Contains(ENiagaraNamespaceMetadataOptions::PreventEditingNamespace))
@@ -2659,6 +2713,13 @@ FName FNiagaraParameterUtilities::ChangeNamespace(FName ParameterName, const FNi
 	{
 		TArray<FName> NameParts = ParameterHandle.GetHandleParts();
 		NameParts.RemoveAt(0, NamespaceMetadata.Namespaces.Num());
+		bool bHasNamespaceModifier = NameParts.Num() > 1;
+		bool bNewNamespaceCanHaveNamespaceModifier =
+			NewNamespaceMetadata.Options.Contains(ENiagaraNamespaceMetadataOptions::PreventEditingNamespaceModifier) == false;
+		if (bHasNamespaceModifier && bNewNamespaceCanHaveNamespaceModifier == false)
+		{
+			NameParts.RemoveAt(0);
+		}
 		NameParts.Insert(NewNamespaceMetadata.Namespaces, 0);
 		return NamePartsToName(NameParts);
 	}
diff --git a/Engine/Plugins/FX/Niagara/Source/NiagaraEditor/Private/NiagaraNodeAssignment.cpp b/Engine/Plugins/FX/Niagara/Source/NiagaraEditor/Private/NiagaraNodeAssignment.cpp
--- a/Engine/Plugins/FX/Niagara/Source/NiagaraEditor/Private/NiagaraNodeAssignment.cpp
+++ b/Engine/Plugins/FX/Niagara/Source/NiagaraEditor/Private/NiagaraNodeAssignment.cpp
@@ -42,11 +42,11 @@ void UNiagaraNodeAssignment::RefreshTitle()
 {
 	if (AssignmentTargets.Num() == 1)
 	{
-		Title = FText::Format(LOCTEXT("NodeTitleSingle", "Set {0}"), FText::FromName(AssignmentTargets[0].GetName()));
+		Title = FText::Format(LOCTEXT("NodeTitleSingle", "Set {0}"), FNiagaraParameterUtilities::FormatParameterNameForTextDisplay(AssignmentTargets[0].GetName()));
 	}
 	else if (AssignmentTargets.Num() > 1)
 	{
-		Title = FText::Format(LOCTEXT("NodeTitleMultiple", "Set {0} (+{1})"), FText::FromName(AssignmentTargets[0].GetName()), AssignmentTargets.Num() - 1);
+		Title = FText::Format(LOCTEXT("NodeTitleMultiple", "Set {0} (+{1})"), FNiagaraParameterUtilities::FormatParameterNameForTextDisplay(AssignmentTargets[0].GetName()), AssignmentTargets.Num() - 1);
 	}
 	else
 	{
@@ -62,7 +62,7 @@ FText UNiagaraNodeAssignment::GetTooltipText() const
 	TargetNames.Add(BaseText);
 	for (const FNiagaraVariable& Var : AssignmentTargets)
 	{
-		TargetNames.Add(FText::Format(LOCTEXT("Indent", "\t{0}"), FText::FromName(Var.GetName())));
+		TargetNames.Add(FText::Format(LOCTEXT("Indent", "\t{0}"), FNiagaraParameterUtilities::FormatParameterNameForTextDisplay(Var.GetName())));
 	}
 
 	return FText::Join(FText::FromString("\n"), TargetNames);
@@ -355,6 +355,7 @@ void UNiagaraNodeAssignment::AddParameter(FNiagaraVariable InVar, FString InDefa
 	MarkNodeRequiresSynchronization(__FUNCTION__, true);
 
 	RefreshTitle();
+	AssignmentTargetsChangedDelegate.Broadcast();
 }
 
 void UNiagaraNodeAssignment::RemoveParameter(const FNiagaraVariable& InVar)
@@ -384,6 +385,7 @@ void UNiagaraNodeAssignment::RemoveParameter(const FNiagaraVariable& InVar)
 	MarkNodeRequiresSynchronization(__FUNCTION__, true);
 
 	RefreshTitle();
+	AssignmentTargetsChangedDelegate.Broadcast();
 }
 
 void UNiagaraNodeAssignment::UpdateUsageBitmaskFromOwningScript()
@@ -656,6 +658,7 @@ int32 UNiagaraNodeAssignment::AddAssignmentTarget(const FNiagaraVariable& InVar,
 	int32 IdxB = AssignmentDefaultValues.AddDefaulted();
 	check(IdxA == IdxB);
 	SetAssignmentTarget(IdxA, InVar, InDefaultValue);
+	AssignmentTargetsChangedDelegate.Broadcast();
 	return IdxA;
 }
 
@@ -677,6 +680,8 @@ bool UNiagaraNodeAssignment::SetAssignmentTarget(int32 Idx, const FNiagaraVariab
 		MarkNodeRequiresSynchronization(__FUNCTION__, true);
 		bRetValue = true;
 	}
+
+	AssignmentTargetsChangedDelegate.Broadcast();
 	return bRetValue;
 }
 
@@ -696,6 +701,7 @@ bool UNiagaraNodeAssignment::RenameAssignmentTarget(FName OldName, FName NewName
 
 			AssignmentTarget.SetName(NewName);
 			RefreshTitle();
+			AssignmentTargetsChangedDelegate.Broadcast();
 			return true;
 		}
 	}
diff --git a/Engine/Plugins/FX/Niagara/Source/NiagaraEditor/Private/NiagaraNodeParameterMapBase.cpp b/Engine/Plugins/FX/Niagara/Source/NiagaraEditor/Private/NiagaraNodeParameterMapBase.cpp
--- a/Engine/Plugins/FX/Niagara/Source/NiagaraEditor/Private/NiagaraNodeParameterMapBase.cpp
+++ b/Engine/Plugins/FX/Niagara/Source/NiagaraEditor/Private/NiagaraNodeParameterMapBase.cpp
@@ -321,52 +321,17 @@ void UNiagaraNodeParameterMapBase::GetChangeNamespaceSubMenuForPin(UToolMenu* Me
 {
 	FToolMenuSection& Section = Menu->AddSection("Section");
 
-	TArray<FNiagaraNamespaceMetadata> NamespaceMetadata = GetDefault<UNiagaraEditorSettings>()->GetAllNamespaceMetadata();
-	NamespaceMetadata.Sort([](const FNiagaraNamespaceMetadata& MetadataA, const FNiagaraNamespaceMetadata& MetadataB)
-		{ return MetadataA.SortId < MetadataB.SortId; });
-
-	for (const FNiagaraNamespaceMetadata& Metadata : NamespaceMetadata)
+	TArray<FNiagaraParameterUtilities::FChangeNamespaceMenuData> MenuData;
+	FNiagaraParameterUtilities::GetChangeNamespaceMenuData(InPin->PinName, FNiagaraParameterUtilities::EParameterContext::Script, MenuData);
+	for(const FNiagaraParameterUtilities::FChangeNamespaceMenuData& MenuDataItem : MenuData)
 	{
-		if (Metadata.IsValid() == false || Metadata.Options.Contains(ENiagaraNamespaceMetadataOptions::PreventEditingNamespace) || Metadata.Options.Contains(ENiagaraNamespaceMetadataOptions::HideInScript))
-		{
-			continue;
-		}
-
-		FText CanChangeMessage;
-		bool bCanChange = FNiagaraParameterUtilities::TestCanChangeNamespaceWithMessage(InPin->PinName, Metadata, CanChangeMessage);
-
-		FText ToolTip;
-		if (bCanChange)
-		{
-			ToolTip = FText::Format(LOCTEXT("ChangeNamespaceToolTipFormat", "{0}\n\nDescription:\n{1}"), CanChangeMessage, Metadata.Description);
-		}
-		else
-		{
-			ToolTip = CanChangeMessage;
-		}
-
-		FString NamespaceNameString = Metadata.Namespaces[0].ToString();
-		for (int32 i = 1; i < Metadata.Namespaces.Num(); i++)
-		{
-			NamespaceNameString += TEXT(".") + Metadata.Namespaces[i].ToString();
-		}
-
+		bool bCanChange = MenuDataItem.bCanChange;
 		FUIAction Action = FUIAction(
-			FExecuteAction::CreateUObject(this, &UNiagaraNodeParameterMapBase::ChangeNamespaceForPin, InPin, Metadata),
+			FExecuteAction::CreateUObject(this, &UNiagaraNodeParameterMapBase::ChangeNamespaceForPin, InPin, MenuDataItem.Metadata),
 			FCanExecuteAction::CreateLambda([bCanChange]() { return bCanChange; }));
 
-		TSharedRef<SWidget> NamespaceWidget =
-			SNew(SBox)
-			.Padding(FMargin(7, 2, 7, 2))
-			[
-				SNew(SNiagaraParameterName)
-				.ParameterName(*NamespaceNameString)
-				.IsReadOnly(true)
-				.SingleNameDisplayMode(SNiagaraParameterName::ESingleNameDisplayMode::Namespace)
-				.ToolTipText(ToolTip)
-			];
-
-		Section.AddEntry(FToolMenuEntry::InitMenuEntry(NAME_None, Action, NamespaceWidget));
+		TSharedRef<SWidget> MenuItemWidget = FNiagaraParameterUtilities::CreateNamespaceMenuItemWidget(MenuDataItem.NamespaceParameterName, MenuDataItem.CanChangeToolTip);
+		Section.AddEntry(FToolMenuEntry::InitMenuEntry(NAME_None, Action, MenuItemWidget));
 	}
 }
 
diff --git a/Engine/Plugins/FX/Niagara/Source/NiagaraEditor/Private/TNiagaraGraphPinEditableName.h b/Engine/Plugins/FX/Niagara/Source/NiagaraEditor/Private/TNiagaraGraphPinEditableName.h
--- a/Engine/Plugins/FX/Niagara/Source/NiagaraEditor/Private/TNiagaraGraphPinEditableName.h
+++ b/Engine/Plugins/FX/Niagara/Source/NiagaraEditor/Private/TNiagaraGraphPinEditableName.h
@@ -135,13 +135,17 @@ protected:
 			}
 			else
 			{
-				CreatedParameterNamePinLabel = SNew(SNiagaraParameterNamePinLabel, this->GraphPinObj)
-					.EditableTextStyle(&FEditorStyle::Get().GetWidgetStyle<FInlineEditableTextBlockStyle>("Graph.Node.InlineEditablePinName"))
-					.ParameterText(this, &TNiagaraGraphPinEditableName<BaseClass>::GetParentPinLabel)
-					.IsReadOnly(bIsPinEditable == false)
-					.Visibility(this, &TNiagaraGraphPinEditableName<BaseClass>::GetParentPinVisibility)
-					.OnVerifyTextChanged(this, &TNiagaraGraphPinEditableName<BaseClass>::OnVerifyTextChanged)
-					.OnTextCommitted(this, &TNiagaraGraphPinEditableName<BaseClass>::OnTextCommitted);
+				return SNew(SBox)
+					.Padding(FMargin(0.0f, 1.0f, 0.0f, 1.0f))
+					[
+						SAssignNew(CreatedParameterNamePinLabel, SNiagaraParameterNamePinLabel, this->GraphPinObj)
+						.EditableTextStyle(&FEditorStyle::Get().GetWidgetStyle<FInlineEditableTextBlockStyle>("Graph.Node.InlineEditablePinName"))
+						.ParameterText(this, &TNiagaraGraphPinEditableName<BaseClass>::GetParentPinLabel)
+						.IsReadOnly(bIsPinEditable == false)
+						.Visibility(this, &TNiagaraGraphPinEditableName<BaseClass>::GetParentPinVisibility)
+						.OnVerifyTextChanged(this, &TNiagaraGraphPinEditableName<BaseClass>::OnVerifyTextChanged)
+						.OnTextCommitted(this, &TNiagaraGraphPinEditableName<BaseClass>::OnTextCommitted)
+					];
 				return CreatedParameterNamePinLabel.ToSharedRef();
 			}
 		}
diff --git a/Engine/Plugins/FX/Niagara/Source/NiagaraEditor/Private/ViewModels/Stack/NiagaraStackFunctionInput.cpp b/Engine/Plugins/FX/Niagara/Source/NiagaraEditor/Private/ViewModels/Stack/NiagaraStackFunctionInput.cpp
--- a/Engine/Plugins/FX/Niagara/Source/NiagaraEditor/Private/ViewModels/Stack/NiagaraStackFunctionInput.cpp
+++ b/Engine/Plugins/FX/Niagara/Source/NiagaraEditor/Private/ViewModels/Stack/NiagaraStackFunctionInput.cpp
@@ -1606,7 +1606,7 @@ void UNiagaraStackFunctionInput::DeleteInput()
 	}
 }
 
-void UNiagaraStackFunctionInput::GetNamespacesForNewParameters(TArray<FName>& OutNamespacesForNewParameters) const
+void UNiagaraStackFunctionInput::GetNamespacesForNewReadParameters(TArray<FName>& OutNamespacesForNewParameters) const
 {
 	UNiagaraNodeOutput* OutputNode = FNiagaraStackGraphUtilities::GetEmitterOutputNodeForStackNode(*OwningFunctionCallNode);
 	bool bIsEditingSystem = GetSystemViewModel()->GetEditMode() == ENiagaraSystemViewModelEditMode::SystemAsset;
@@ -1639,6 +1639,40 @@ void UNiagaraStackFunctionInput::GetNamespacesForNewParameters(TArray<FName>& Ou
 	OutNamespacesForNewParameters.Add(FNiagaraConstants::TransientNamespace);
 }
 
+void UNiagaraStackFunctionInput::GetNamespacesForNewWriteParameters(TArray<FName>& OutNamespacesForNewParameters) const
+{
+	UNiagaraNodeOutput* OutputNode = FNiagaraStackGraphUtilities::GetEmitterOutputNodeForStackNode(*OwningFunctionCallNode);
+	bool bIsEditingSystem = GetSystemViewModel()->GetEditMode() == ENiagaraSystemViewModelEditMode::SystemAsset;
+
+	switch (OutputNode->GetUsage())
+	{
+	case ENiagaraScriptUsage::ParticleSpawnScript:
+	case ENiagaraScriptUsage::ParticleSpawnScriptInterpolated:
+	case ENiagaraScriptUsage::ParticleUpdateScript:
+	case ENiagaraScriptUsage::ParticleEventScript:
+	case ENiagaraScriptUsage::ParticleSimulationStageScript:
+	{
+		OutNamespacesForNewParameters.Add(FNiagaraConstants::ParticleAttributeNamespace);
+		break;
+	}
+	case ENiagaraScriptUsage::EmitterSpawnScript:
+	case ENiagaraScriptUsage::EmitterUpdateScript:
+	{
+		OutNamespacesForNewParameters.Add(FNiagaraConstants::EmitterNamespace);
+		break;
+	}
+	case ENiagaraScriptUsage::SystemSpawnScript:
+	case ENiagaraScriptUsage::SystemUpdateScript:
+		if (bIsEditingSystem)
+		{
+			OutNamespacesForNewParameters.Add(FNiagaraConstants::SystemNamespace);
+		}
+		break;
+	}
+
+	OutNamespacesForNewParameters.Add(FNiagaraConstants::TransientNamespace);
+}
+
 UNiagaraStackFunctionInput::FOnValueChanged& UNiagaraStackFunctionInput::OnValueChanged()
 {
 	return ValueChangedDelegate;
diff --git a/Engine/Plugins/FX/Niagara/Source/NiagaraEditor/Private/Widgets/SNiagaraParameterMapView.cpp b/Engine/Plugins/FX/Niagara/Source/NiagaraEditor/Private/Widgets/SNiagaraParameterMapView.cpp
--- a/Engine/Plugins/FX/Niagara/Source/NiagaraEditor/Private/Widgets/SNiagaraParameterMapView.cpp
+++ b/Engine/Plugins/FX/Niagara/Source/NiagaraEditor/Private/Widgets/SNiagaraParameterMapView.cpp
@@ -784,52 +784,19 @@ TSharedPtr<SWidget> SNiagaraParameterMapView::OnContextMenuOpening()
 			TAttribute<FText> RenameToolTip;
 			RenameToolTip.Bind(this, &SNiagaraParameterMapView::GetRenameOnActionNodeToolTip);
 			MenuBuilder.AddMenuEntry(FGenericCommands::Get().Rename, NAME_None, LOCTEXT("Rename", "Rename"), RenameToolTip);
-		}
-		MenuBuilder.EndSection();
 
-		TArray<TSharedPtr<FEdGraphSchemaAction>> SelectedActions;
-		GraphActionMenu->GetSelectedActions(SelectedActions);
+			MenuBuilder.AddSubMenu(
+				LOCTEXT("ChangeNamespace", "Namespace"),
+				LOCTEXT("ChangeNamespaceToolTip", "Select a new namespace for the selected parameter."),
+				FNewMenuDelegate::CreateSP(this, &SNiagaraParameterMapView::GetChangeNamespaceSubMenu));
 
-		if (SelectedActions.Num() == 1)
-		{
-			TSharedPtr<FNiagaraParameterAction> SelectedNiagaraAction = StaticCastSharedPtr<FNiagaraParameterAction>(SelectedActions[0]);
-			if (SelectedNiagaraAction.IsValid())
-			{
-				MenuBuilder.BeginSection("NamespaceModifier", LOCTEXT("NamespaceModifierMenuSection", "Namespace Modifiers"));
-				{
-					TAttribute<FText> AddToolTip;
-					AddToolTip.Bind(this, &SNiagaraParameterMapView::GetAddNamespaceModifierToolTip);
-					MenuBuilder.AddMenuEntry(
-						LOCTEXT("AddNamespaceModifier", "Add"),
-						AddToolTip,
-						FSlateIcon(),
-						FUIAction(
-							FExecuteAction::CreateSP(this, &SNiagaraParameterMapView::OnAddNamespaceModifier),
-							FCanExecuteAction::CreateSP(this, &SNiagaraParameterMapView::CanAddNamespaceModifier)));
-
-					TAttribute<FText> RemoveToolTip;
-					RemoveToolTip.Bind(this, &SNiagaraParameterMapView::GetRemoveNamespaceModifierToolTip);
-					MenuBuilder.AddMenuEntry(
-						LOCTEXT("RemoveNamespaceModifier", "Remove"),
-						RemoveToolTip,
-						FSlateIcon(),
-						FUIAction(
-							FExecuteAction::CreateSP(this, &SNiagaraParameterMapView::OnRemoveNamespaceModifier),
-							FCanExecuteAction::CreateSP(this, &SNiagaraParameterMapView::CanRemoveNamespaceModifier)));
-
-					TAttribute<FText> EditToolTip;
-					EditToolTip.Bind(this, &SNiagaraParameterMapView::GetEditNamespaceModifierToolTip);
-					MenuBuilder.AddMenuEntry(
-						LOCTEXT("EditNamespaceModifier", "Edit"),
-						EditToolTip,
-						FSlateIcon(),
-						FUIAction(
-							FExecuteAction::CreateSP(this, &SNiagaraParameterMapView::OnEditNamespaceModifier),
-							FCanExecuteAction::CreateSP(this, &SNiagaraParameterMapView::CanEditNamespaceModifier)));
-				}
-				MenuBuilder.EndSection();
-			}
+			MenuBuilder.AddSubMenu(
+				LOCTEXT("ChangeNamespaceModifier", "Namespace Modifier"),
+				LOCTEXT("ChangeNamespaceModifierToolTip", "Edit the namespace modifier for the selected parameter."),
+				FNewMenuDelegate::CreateSP(this, &SNiagaraParameterMapView::GetChangeNamespaceModifierSubMenu));
 		}
+		MenuBuilder.EndSection();
+
 		return MenuBuilder.MakeWidget();
 	}
 
@@ -1252,6 +1219,76 @@ bool SNiagaraParameterMapView::GetSingleParameterActionForSelection(
 	return true;
 }
 
+void SNiagaraParameterMapView::GetChangeNamespaceSubMenu(FMenuBuilder& MenuBuilder)
+{
+	TSharedPtr<FNiagaraParameterAction> ParameterAction;
+	FText Unused;
+	if (GetSingleParameterActionForSelection(ParameterAction, Unused))
+	{
+		TArray<FNiagaraParameterUtilities::FChangeNamespaceMenuData> MenuData;
+		FNiagaraParameterUtilities::GetChangeNamespaceMenuData(ParameterAction->Parameter.GetName(),
+			IsScriptToolkit() ? FNiagaraParameterUtilities::EParameterContext::Script : FNiagaraParameterUtilities::EParameterContext::System, MenuData);
+		for (const FNiagaraParameterUtilities::FChangeNamespaceMenuData& MenuDataItem : MenuData)
+		{
+			bool bCanChange = MenuDataItem.bCanChange;
+			FUIAction Action = FUIAction(
+				FExecuteAction::CreateSP(this, &SNiagaraParameterMapView::OnChangeNamespace, MenuDataItem.Metadata),
+				FCanExecuteAction::CreateLambda([bCanChange]() { return bCanChange; }));
+
+			TSharedRef<SWidget> MenuItemWidget = FNiagaraParameterUtilities::CreateNamespaceMenuItemWidget(MenuDataItem.NamespaceParameterName, MenuDataItem.CanChangeToolTip);
+			MenuBuilder.AddMenuEntry(Action, MenuItemWidget, NAME_None, MenuDataItem.CanChangeToolTip);
+		}
+	}
+}
+
+void SNiagaraParameterMapView::OnChangeNamespace(FNiagaraNamespaceMetadata Metadata)
+{
+	TSharedPtr<FNiagaraParameterAction> ParameterAction;
+	FText Unused;
+	if (GetSingleParameterActionForSelection(ParameterAction, Unused))
+	{
+		FName NewName = FNiagaraParameterUtilities::ChangeNamespace(ParameterAction->Parameter.GetName(), Metadata);
+		if (NewName != NAME_None)
+		{
+			FScopedTransaction Transaction(LOCTEXT("ChangeNamespaceTransaction", "Change namespace"));
+			RenameParameter(ParameterAction->Parameter, NewName);
+		}
+	}
+}
+
+void SNiagaraParameterMapView::GetChangeNamespaceModifierSubMenu(FMenuBuilder& MenuBuilder)
+{
+	TAttribute<FText> AddToolTip;
+	AddToolTip.Bind(this, &SNiagaraParameterMapView::GetAddNamespaceModifierToolTip);
+	MenuBuilder.AddMenuEntry(
+		LOCTEXT("AddNamespaceModifier", "Add"),
+		AddToolTip,
+		FSlateIcon(),
+		FUIAction(
+			FExecuteAction::CreateSP(this, &SNiagaraParameterMapView::OnAddNamespaceModifier),
+			FCanExecuteAction::CreateSP(this, &SNiagaraParameterMapView::CanAddNamespaceModifier)));
+
+	TAttribute<FText> RemoveToolTip;
+	RemoveToolTip.Bind(this, &SNiagaraParameterMapView::GetRemoveNamespaceModifierToolTip);
+	MenuBuilder.AddMenuEntry(
+		LOCTEXT("RemoveNamespaceModifier", "Remove"),
+		RemoveToolTip,
+		FSlateIcon(),
+		FUIAction(
+			FExecuteAction::CreateSP(this, &SNiagaraParameterMapView::OnRemoveNamespaceModifier),
+			FCanExecuteAction::CreateSP(this, &SNiagaraParameterMapView::CanRemoveNamespaceModifier)));
+
+	TAttribute<FText> EditToolTip;
+	EditToolTip.Bind(this, &SNiagaraParameterMapView::GetEditNamespaceModifierToolTip);
+	MenuBuilder.AddMenuEntry(
+		LOCTEXT("EditNamespaceModifier", "Edit"),
+		EditToolTip,
+		FSlateIcon(),
+		FUIAction(
+			FExecuteAction::CreateSP(this, &SNiagaraParameterMapView::OnEditNamespaceModifier),
+			FCanExecuteAction::CreateSP(this, &SNiagaraParameterMapView::CanEditNamespaceModifier)));
+}
+
 FText SNiagaraParameterMapView::GetAddNamespaceModifierToolTip() const
 {
 	TSharedPtr<FNiagaraParameterAction> ParameterAction;
diff --git a/Engine/Plugins/FX/Niagara/Source/NiagaraEditor/Public/NiagaraEditorUtilities.h b/Engine/Plugins/FX/Niagara/Source/NiagaraEditor/Public/NiagaraEditorUtilities.h
--- a/Engine/Plugins/FX/Niagara/Source/NiagaraEditor/Public/NiagaraEditorUtilities.h
+++ b/Engine/Plugins/FX/Niagara/Source/NiagaraEditor/Public/NiagaraEditorUtilities.h
@@ -8,6 +8,7 @@
 #include "Misc/Attribute.h"
 #include "AssetData.h"
 #include "NiagaraGraph.h"
+#include "NiagaraEditorSettings.h"
 
 class UNiagaraNodeInput;
 class UNiagaraNodeOutput;
@@ -324,19 +325,37 @@ namespace FNiagaraParameterUtilities
 		FNiagaraNamespaceMetadata& OutNamespaceMetadata,
 		FText& OutErrorMessage);
 
-	bool TestCanChangeNamespaceWithMessage(FName ParameterName, const FNiagaraNamespaceMetadata& NewNamespaceMetadata, FText& OutMessage);
+	enum class EParameterContext
+	{
+		Script,
+		System
+	};
+
+	struct FChangeNamespaceMenuData
+	{
+		bool bCanChange;
+		FText CanChangeToolTip;
+		FName NamespaceParameterName;
+		FNiagaraNamespaceMetadata Metadata;
+	};
+
+	NIAGARAEDITOR_API void GetChangeNamespaceMenuData(FName InParameterName, EParameterContext InParameterContext, TArray<FChangeNamespaceMenuData>& OutChangeNamespaceMenuData);
+
+	NIAGARAEDITOR_API TSharedRef<SWidget> CreateNamespaceMenuItemWidget(FName Namespace, FText ToolTip);
+
+	NIAGARAEDITOR_API bool TestCanChangeNamespaceWithMessage(FName ParameterName, const FNiagaraNamespaceMetadata& NewNamespaceMetadata, FText& OutMessage);
 
-	FName ChangeNamespace(FName ParameterName, const FNiagaraNamespaceMetadata& NewNamespaceMetadata);
+	NIAGARAEDITOR_API FName ChangeNamespace(FName ParameterName, const FNiagaraNamespaceMetadata& NewNamespaceMetadata);
 
-	bool TestCanAddNamespaceModifierWithMessage(FName ParameterName, FText& OutMessage);
+	NIAGARAEDITOR_API bool TestCanAddNamespaceModifierWithMessage(FName ParameterName, FText& OutMessage);
 
-	FName AddNamespaceModifier(FName InParameterName);
+	NIAGARAEDITOR_API FName AddNamespaceModifier(FName InParameterName);
 
-	bool TestCanRemoveNamespaceModifierWithMessage(FName ParameterName, FText& OutMessage);
+	NIAGARAEDITOR_API bool TestCanRemoveNamespaceModifierWithMessage(FName ParameterName, FText& OutMessage);
 
-	FName RemoveNamespaceModifier(FName InParameterName);
+	NIAGARAEDITOR_API FName RemoveNamespaceModifier(FName InParameterName);
 
-	bool TestCanEditNamespaceModifierWithMessage(FName ParameterName, FText& OutMessage);
+	NIAGARAEDITOR_API bool TestCanEditNamespaceModifierWithMessage(FName ParameterName, FText& OutMessage);
 
-	bool TestCanRenameWithMessage(FName ParameterName, FText& OutMessage);
+	NIAGARAEDITOR_API bool TestCanRenameWithMessage(FName ParameterName, FText& OutMessage);
 };
diff --git a/Engine/Plugins/FX/Niagara/Source/NiagaraEditor/Public/NiagaraNodeAssignment.h b/Engine/Plugins/FX/Niagara/Source/NiagaraEditor/Public/NiagaraNodeAssignment.h
--- a/Engine/Plugins/FX/Niagara/Source/NiagaraEditor/Public/NiagaraNodeAssignment.h
+++ b/Engine/Plugins/FX/Niagara/Source/NiagaraEditor/Public/NiagaraNodeAssignment.h
@@ -55,6 +55,8 @@ public:
 
 	virtual FText GetTooltipText() const override;
 
+	FSimpleMulticastDelegate& OnAssignmentTargetsChanged() { return AssignmentTargetsChangedDelegate; }
+
 protected:
 
 	UPROPERTY()
@@ -74,6 +76,8 @@ protected:
 
 	FText Title;
 
+	FSimpleMulticastDelegate AssignmentTargetsChangedDelegate;
+
 private:
 	void GenerateScript();
 
diff --git a/Engine/Plugins/FX/Niagara/Source/NiagaraEditor/Public/ViewModels/Stack/NiagaraStackFunctionInput.h b/Engine/Plugins/FX/Niagara/Source/NiagaraEditor/Public/ViewModels/Stack/NiagaraStackFunctionInput.h
--- a/Engine/Plugins/FX/Niagara/Source/NiagaraEditor/Public/ViewModels/Stack/NiagaraStackFunctionInput.h
+++ b/Engine/Plugins/FX/Niagara/Source/NiagaraEditor/Public/ViewModels/Stack/NiagaraStackFunctionInput.h
@@ -199,7 +199,10 @@ public:
 	void DeleteInput();
 
 	/** Gets the namespaces which new parameters for this input can be read from. */
-	void GetNamespacesForNewParameters(TArray<FName>& OutNamespacesForNewParameters) const;
+	void GetNamespacesForNewReadParameters(TArray<FName>& OutNamespacesForNewParameters) const;
+
+	/** Gets the namespaces which new parameters for this input can write to. */
+	void GetNamespacesForNewWriteParameters(TArray<FName>& OutNamespacesForNewParameters) const;
 
 	/** Gets a multicast delegate which is called whenever the value on this input changes. */
 	FOnValueChanged& OnValueChanged();
diff --git a/Engine/Plugins/FX/Niagara/Source/NiagaraEditor/Public/Widgets/SNiagaraParameterMapView.h b/Engine/Plugins/FX/Niagara/Source/NiagaraEditor/Public/Widgets/SNiagaraParameterMapView.h
--- a/Engine/Plugins/FX/Niagara/Source/NiagaraEditor/Public/Widgets/SNiagaraParameterMapView.h
+++ b/Engine/Plugins/FX/Niagara/Source/NiagaraEditor/Public/Widgets/SNiagaraParameterMapView.h
@@ -14,6 +14,7 @@
 #include "NiagaraGraph.h"
 #include "NiagaraActions.h"
 #include "EditorStyleSet.h"
+#include "NiagaraEditorSettings.h"
 
 class SGraphActionMenu;
 class SEditableTextBox;
@@ -154,6 +155,11 @@ private:
 		TSharedPtr<FNiagaraParameterAction>& OutParameterAction,
 		FText& OutErrorMessage) const;
 
+	void GetChangeNamespaceSubMenu(FMenuBuilder& MenuBuilder);
+	void OnChangeNamespace(FNiagaraNamespaceMetadata Metadata);
+
+	void GetChangeNamespaceModifierSubMenu(FMenuBuilder& MenuBuilder);
+
 	FText GetAddNamespaceModifierToolTip() const;
 	bool CanAddNamespaceModifier() const;
 	void OnAddNamespaceModifier();
diff --git a/Engine/Plugins/FX/Niagara/Source/NiagaraEditorWidgets/NiagaraEditorWidgets.Build.cs b/Engine/Plugins/FX/Niagara/Source/NiagaraEditorWidgets/NiagaraEditorWidgets.Build.cs
--- a/Engine/Plugins/FX/Niagara/Source/NiagaraEditorWidgets/NiagaraEditorWidgets.Build.cs
+++ b/Engine/Plugins/FX/Niagara/Source/NiagaraEditorWidgets/NiagaraEditorWidgets.Build.cs
@@ -28,7 +28,8 @@ public class NiagaraEditorWidgets : ModuleRules
 			"Projects",
 			"Sequencer",
             "EditorWidgets",
-        });
+			"ApplicationCore",
+		});
 
 		PrivateIncludePathModuleNames.AddRange(new string[] {
 		});
diff --git a/Engine/Plugins/FX/Niagara/Source/NiagaraEditorWidgets/Private/SNiagaraOverviewStack.cpp b/Engine/Plugins/FX/Niagara/Source/NiagaraEditorWidgets/Private/SNiagaraOverviewStack.cpp
--- a/Engine/Plugins/FX/Niagara/Source/NiagaraEditorWidgets/Private/SNiagaraOverviewStack.cpp
+++ b/Engine/Plugins/FX/Niagara/Source/NiagaraEditorWidgets/Private/SNiagaraOverviewStack.cpp
@@ -23,9 +23,12 @@
 #include "NiagaraStackCommandContext.h"
 #include "NiagaraNodeFunctionCall.h"
 #include "NiagaraEditorModule.h"
+#include "NiagaraNodeAssignment.h"
+#include "Widgets/SNiagaraParameterName.h"
 
 #include "Widgets/Layout/SScrollBox.h"
 #include "Widgets/Layout/SGridPanel.h"
+#include "Widgets/Layout/SScaleBox.h"
 #include "Widgets/Images/SImage.h"
 #include "Widgets/Input/SButton.h"
 #include "Widgets/SToolTip.h"
@@ -299,25 +302,121 @@ class SNiagaraSystemOverviewItemName : public SCompoundWidget
 		StackItem = InStackItem;
 		StackItem->OnAlternateDisplayNameChanged().AddSP(this, &SNiagaraSystemOverviewItemName::UpdateFromAlternateDisplayName);
 
-		ChildSlot
-		[
-			SAssignNew(NameTextBlock, STextBlock)
-			.Text(this, &SNiagaraSystemOverviewItemName::GetItemDisplayName)
-			.ToolTipText(this, &SNiagaraSystemOverviewItemName::GetItemToolTip)
-			.IsEnabled_UObject(StackItem.Get(), &UNiagaraStackEntry::GetIsEnabledAndOwnerIsEnabled)
-		];
-		UpdateFromAlternateDisplayName();
+		if (StackItem->IsA<UNiagaraStackModuleItem>() &&
+			CastChecked<UNiagaraStackModuleItem>(StackItem)->GetModuleNode().IsA<UNiagaraNodeAssignment>())
+		{
+			UNiagaraNodeAssignment* AssignmentNode = CastChecked<UNiagaraNodeAssignment>(&CastChecked<UNiagaraStackModuleItem>(StackItem)->GetModuleNode());
+			AssignmentNode->OnAssignmentTargetsChanged().AddSP(this, &SNiagaraSystemOverviewItemName::UpdateFromAssignmentTargets);
+		}
+
+		UpdateContent();
 	}
 
 	~SNiagaraSystemOverviewItemName()
 	{
-		if (StackItem.IsValid())
+		if (StackItem.IsValid() && StackItem->IsFinalized() == false)
 		{
 			StackItem->OnAlternateDisplayNameChanged().RemoveAll(this);
+
+			if (StackItem->IsA<UNiagaraStackModuleItem>() &&
+				CastChecked<UNiagaraStackModuleItem>(StackItem)->GetModuleNode().IsA<UNiagaraNodeAssignment>())
+			{
+				UNiagaraNodeAssignment* AssignmentNode = CastChecked<UNiagaraNodeAssignment>(&CastChecked<UNiagaraStackModuleItem>(StackItem)->GetModuleNode());
+				AssignmentNode->OnAssignmentTargetsChanged().RemoveAll(this);
+			}
 		}
 	}
 
 private:
+	void UpdateContent()
+	{
+		UNiagaraStackModuleItem* ModuleItem = Cast<UNiagaraStackModuleItem>(StackItem);
+		UNiagaraNodeAssignment* AssignmentNode = ModuleItem != nullptr 
+			? Cast<UNiagaraNodeAssignment>(&ModuleItem->GetModuleNode())
+			: nullptr;
+
+		if (StackItem->GetAlternateDisplayName().IsSet() == false &&
+			AssignmentNode != nullptr &&
+			AssignmentNode->GetAssignmentTargets().Num() > 0)
+		{
+			NameTextBlock.Reset();
+
+			TSharedPtr<SWidget> ParameterWidget;
+			if (AssignmentNode->GetAssignmentTargets().Num() == 1)
+			{
+				ParameterWidget = SNew(SNiagaraParameterName)
+					.ParameterName(AssignmentNode->GetAssignmentTargets()[0].GetName())
+					.IsReadOnly(true);
+			}
+			else
+			{
+				TSharedRef<SVerticalBox> ParameterBox = SNew(SVerticalBox);
+				for (const FNiagaraVariable& AssignmentTarget : AssignmentNode->GetAssignmentTargets())
+				{
+					ParameterBox->AddSlot()
+					.AutoHeight()
+					.Padding(0.0f, 1.0f, 0.0f, 1.0f)
+					[
+						SNew(SScaleBox)
+						.UserSpecifiedScale(0.85f)
+						.Stretch(EStretch::UserSpecified)
+						.HAlign(HAlign_Left)
+						[
+							SNew(SNiagaraParameterName)
+							.ParameterName(AssignmentTarget.GetName())
+							.IsReadOnly(true)
+						]
+					];
+				}
+				ParameterWidget = ParameterBox;
+			}
+
+			ChildSlot
+			[
+				SNew(SHorizontalBox)
+				.IsEnabled_UObject(StackItem.Get(), &UNiagaraStackEntry::GetIsEnabledAndOwnerIsEnabled)
+				+ SHorizontalBox::Slot()
+				.AutoWidth()
+				.Padding(0.0f, 0.0f, 6.0f, 0.0f)
+				[
+					SNew(STextBlock)
+					.TextStyle(FNiagaraEditorWidgetsStyle::Get(), "NiagaraEditor.SystemOverview.ItemText")
+					.Text(LOCTEXT("SetVariablesPrefix", "Set:"))
+				]
+				+ SHorizontalBox::Slot()
+				[
+					ParameterWidget.ToSharedRef()
+				]
+			];
+		}
+		else
+		{
+			ChildSlot
+			[
+				SAssignNew(NameTextBlock, STextBlock)
+				.Text(this, &SNiagaraSystemOverviewItemName::GetItemDisplayName)
+				.ToolTipText(this, &SNiagaraSystemOverviewItemName::GetItemToolTip)
+				.IsEnabled_UObject(StackItem.Get(), &UNiagaraStackEntry::GetIsEnabledAndOwnerIsEnabled)
+			];
+			UpdateTextStyle();
+		}
+	}
+
+	void UpdateTextStyle()
+	{
+		if (NameTextBlock.IsValid())
+		{
+			if (StackItem->GetAlternateDisplayName().IsSet())
+			{
+				NameTextBlock->SetTextStyle(&FNiagaraEditorWidgetsStyle::Get().GetWidgetStyle<FTextBlockStyle>("NiagaraEditor.SystemOverview.AlternateItemText"));
+			}
+			else
+			{
+				NameTextBlock->SetTextStyle(&FNiagaraEditorWidgetsStyle::Get().GetWidgetStyle<FTextBlockStyle>("NiagaraEditor.SystemOverview.ItemText"));
+			}
+		}
+	}
+
 	FText GetItemDisplayName() const
 	{
 		if (StackItem.IsValid() && StackItem->IsFinalized() == false)
@@ -359,13 +458,35 @@ private:
 
 	void UpdateFromAlternateDisplayName()
 	{
-		if (StackItem.IsValid() && StackItem->IsFinalized() == false && StackItem->GetAlternateDisplayName().IsSet())
+		if (StackItem.IsValid() && StackItem->IsFinalized() == false)
 		{
-			NameTextBlock->SetTextStyle(&FNiagaraEditorWidgetsStyle::Get().GetWidgetStyle<FTextBlockStyle>("NiagaraEditor.SystemOverview.AlternateItemText"));
+			if (StackItem->GetAlternateDisplayName().IsSet())
+			{
+				if (NameTextBlock.IsValid() == false)
+				{
+					UpdateContent();
+				}
+				else
+				{
+					UpdateTextStyle();
+				}
+			}
+			else
+			{
+				if (NameTextBlock.IsValid())
+				{
+					UpdateContent();
+				}
+			}
 		}
-		else
+		ToolTipCache.Reset();
+	}
+
+	void UpdateFromAssignmentTargets()
+	{
+		if (StackItem.IsValid() && StackItem->IsFinalized() == false)
 		{
-			NameTextBlock->SetTextStyle(&FNiagaraEditorWidgetsStyle::Get().GetWidgetStyle<FTextBlockStyle>("NiagaraEditor.SystemOverview.ItemText"));
+			UpdateContent();
 		}
 		ToolTipCache.Reset();
 	}
diff --git a/Engine/Plugins/FX/Niagara/Source/NiagaraEditorWidgets/Private/SNiagaraStack.cpp b/Engine/Plugins/FX/Niagara/Source/NiagaraEditorWidgets/Private/SNiagaraStack.cpp
--- a/Engine/Plugins/FX/Niagara/Source/NiagaraEditorWidgets/Private/SNiagaraStack.cpp
+++ b/Engine/Plugins/FX/Niagara/Source/NiagaraEditorWidgets/Private/SNiagaraStack.cpp
@@ -932,9 +932,13 @@ SNiagaraStack::FRowWidgets SNiagaraStack::ConstructNameAndValueWidgetsForItem(UN
 	else if (Item->IsA<UNiagaraStackFunctionInput>())
 	{
 		UNiagaraStackFunctionInput* FunctionInput = CastChecked<UNiagaraStackFunctionInput>(Item);
-		return FRowWidgets(
+
+		TSharedRef<SNiagaraStackFunctionInputName> FunctionInputNameWidget =
 			SNew(SNiagaraStackFunctionInputName, FunctionInput, StackViewModel)
-			.IsSelected(Container, &SNiagaraStackTableRow::IsSelected),
+			.IsSelected(Container, &SNiagaraStackTableRow::IsSelected);
+		Container->AddFillRowContextMenuHandler(SNiagaraStackTableRow::FOnFillRowContextMenu::CreateSP(FunctionInputNameWidget, &SNiagaraStackFunctionInputName::FillRowContextMenu));
+
+		return FRowWidgets(FunctionInputNameWidget,
 			SNew(SNiagaraStackFunctionInputValue, FunctionInput));
 	}
 	else if (Item->IsA<UNiagaraStackErrorItem>())
diff --git a/Engine/Plugins/FX/Niagara/Source/NiagaraEditorWidgets/Private/Stack/SNiagaraStackFunctionInputName.cpp b/Engine/Plugins/FX/Niagara/Source/NiagaraEditorWidgets/Private/Stack/SNiagaraStackFunctionInputName.cpp
--- a/Engine/Plugins/FX/Niagara/Source/NiagaraEditorWidgets/Private/Stack/SNiagaraStackFunctionInputName.cpp
+++ b/Engine/Plugins/FX/Niagara/Source/NiagaraEditorWidgets/Private/Stack/SNiagaraStackFunctionInputName.cpp
@@ -11,6 +11,10 @@
 #include "Widgets/SNiagaraParameterName.h"
 #include "NiagaraNodeAssignment.h"
 #include "NiagaraNodeFunctionCall.h"
+#include "NiagaraEditorUtilities.h"
+#include "HAL/PlatformApplicationMisc.h"
+
+#define LOCTEXT_NAMESPACE "NiagaraStackFunctionInputName"
 
 void SNiagaraStackFunctionInputName::Construct(const FArguments& InArgs, UNiagaraStackFunctionInput* InFunctionInput, UNiagaraStackViewModel* InStackViewModel)
 {
@@ -89,6 +93,36 @@ void SNiagaraStackFunctionInputName::Tick(const FGeometry& AllottedGeometry, con
 	SCompoundWidget::Tick(AllottedGeometry, InCurrentTime, InDeltaTime);
 }
 
+void SNiagaraStackFunctionInputName::FillRowContextMenu(class FMenuBuilder& MenuBuilder)
+{
+	if (FunctionInput->GetInputFunctionCallNode().IsA<UNiagaraNodeAssignment>())
+	{
+		MenuBuilder.BeginSection("Parameter", LOCTEXT("ParameterHeader", "Parameter"));
+		{
+			TAttribute<FText> CopyReferenceToolTip;
+			CopyReferenceToolTip.Bind(this, &SNiagaraStackFunctionInputName::GetCopyParameterReferenceToolTip);
+			MenuBuilder.AddMenuEntry(
+				LOCTEXT("CopyParameterReference", "Copy Reference"),
+				CopyReferenceToolTip,
+				FSlateIcon(),
+				FUIAction(
+					FExecuteAction::CreateSP(this, &SNiagaraStackFunctionInputName::OnCopyParameterReference),
+					FCanExecuteAction::CreateSP(this, &SNiagaraStackFunctionInputName::CanCopyParameterReference)));
+
+			MenuBuilder.AddSubMenu(
+				LOCTEXT("ChangeNamespace", "Namespace"),
+				LOCTEXT("ChangeNamespaceToolTip", "Select a new namespace for the selected parameter."),
+				FNewMenuDelegate::CreateSP(this, &SNiagaraStackFunctionInputName::GetChangeNamespaceSubMenu));
+
+			MenuBuilder.AddSubMenu(
+				LOCTEXT("ChangeNamespaceModifier", "Namespace Modifier"),
+				LOCTEXT("ChangeNamespaceModifierToolTip", "Edit the namespace modifier for the selected parameter."),
+				FNewMenuDelegate::CreateSP(this, &SNiagaraStackFunctionInputName::GetChangeNamespaceModifierSubMenu));
+		}
+		MenuBuilder.EndSection();
+	}
+}
+
 EVisibility SNiagaraStackFunctionInputName::GetEditConditionCheckBoxVisibility() const
 {
 	return FunctionInput->GetHasEditCondition() && FunctionInput->GetShowEditConditionInline() ? EVisibility::Visible : EVisibility::Collapsed;
@@ -134,3 +168,164 @@ void SNiagaraStackFunctionInputName::OnNameTextCommitted(const FText& InText, ET
 {
 	FunctionInput->OnRenamed(InText);
 }
+
+void SNiagaraStackFunctionInputName::GetChangeNamespaceSubMenu(FMenuBuilder& MenuBuilder)
+{
+	FName InputParameterName = *FunctionInput->GetDisplayName().ToString();
+
+	TArray<FName> NamespacesForWriteParameters;
+	FunctionInput->GetNamespacesForNewWriteParameters(NamespacesForWriteParameters);
+
+	TArray<FNiagaraParameterUtilities::FChangeNamespaceMenuData> MenuData;
+	FNiagaraParameterUtilities::GetChangeNamespaceMenuData(*FunctionInput->GetDisplayName().ToString(), FNiagaraParameterUtilities::EParameterContext::System, MenuData);
+	for (const FNiagaraParameterUtilities::FChangeNamespaceMenuData& MenuDataItem : MenuData)
+	{
+		if (MenuDataItem.Metadata.Namespaces.Num() == 1 && NamespacesForWriteParameters.Contains(MenuDataItem.Metadata.Namespaces[0]))
+		{
+			bool bCanChange = MenuDataItem.bCanChange;
+			FUIAction Action = FUIAction(
+				FExecuteAction::CreateSP(this, &SNiagaraStackFunctionInputName::OnChangeNamespace, MenuDataItem.Metadata),
+				FCanExecuteAction::CreateLambda([bCanChange]() { return bCanChange; }));
+
+			TSharedRef<SWidget> MenuItemWidget = FNiagaraParameterUtilities::CreateNamespaceMenuItemWidget(MenuDataItem.NamespaceParameterName, MenuDataItem.CanChangeToolTip);
+			MenuBuilder.AddMenuEntry(Action, MenuItemWidget, NAME_None, MenuDataItem.CanChangeToolTip);
+		}
+	}
+}
+
+void SNiagaraStackFunctionInputName::OnChangeNamespace(FNiagaraNamespaceMetadata Metadata)
+{
+	FName InputParameterName = *FunctionInput->GetDisplayName().ToString();
+	FName NewName = FNiagaraParameterUtilities::ChangeNamespace(InputParameterName, Metadata);
+	if (NewName != NAME_None)
+	{
+		FScopedTransaction Transaction(LOCTEXT("ChangeNamespaceTransaction", "Change namespace"));
+		FunctionInput->OnRenamed(FText::FromName(NewName));
+	}
+}
+
+void SNiagaraStackFunctionInputName::GetChangeNamespaceModifierSubMenu(FMenuBuilder& MenuBuilder)
+{
+	TAttribute<FText> AddToolTip;
+	AddToolTip.Bind(this, &SNiagaraStackFunctionInputName::GetAddNamespaceModifierToolTip);
+	MenuBuilder.AddMenuEntry(
+		LOCTEXT("AddNamespaceModifier", "Add"),
+		AddToolTip,
+		FSlateIcon(),
+		FUIAction(
+			FExecuteAction::CreateSP(this, &SNiagaraStackFunctionInputName::OnAddNamespaceModifier),
+			FCanExecuteAction::CreateSP(this, &SNiagaraStackFunctionInputName::CanAddNamespaceModifier)));
+
+	TAttribute<FText> RemoveToolTip;
+	RemoveToolTip.Bind(this, &SNiagaraStackFunctionInputName::GetRemoveNamespaceModifierToolTip);
+	MenuBuilder.AddMenuEntry(
+		LOCTEXT("RemoveNamespaceModifier", "Remove"),
+		RemoveToolTip,
+		FSlateIcon(),
+		FUIAction(
+			FExecuteAction::CreateSP(this, &SNiagaraStackFunctionInputName::OnRemoveNamespaceModifier),
+			FCanExecuteAction::CreateSP(this, &SNiagaraStackFunctionInputName::CanRemoveNamespaceModifier)));
+
+	TAttribute<FText> EditToolTip;
+	EditToolTip.Bind(this, &SNiagaraStackFunctionInputName::GetEditNamespaceModifierToolTip);
+	MenuBuilder.AddMenuEntry(
+		LOCTEXT("EditNamespaceModifier", "Edit"),
+		EditToolTip,
+		FSlateIcon(),
+		FUIAction(
+			FExecuteAction::CreateSP(this, &SNiagaraStackFunctionInputName::OnEditNamespaceModifier),
+			FCanExecuteAction::CreateSP(this, &SNiagaraStackFunctionInputName::CanEditNamespaceModifier)));
+}
+
+FText SNiagaraStackFunctionInputName::GetAddNamespaceModifierToolTip() const
+{
+	FName InputParameterName = *FunctionInput->GetDisplayName().ToString();
+	FText AddMessage;
+	FNiagaraParameterUtilities::TestCanAddNamespaceModifierWithMessage(InputParameterName, AddMessage);
+	return AddMessage;
+}
+
+bool SNiagaraStackFunctionInputName::CanAddNamespaceModifier() const
+{
+	FName InputParameterName = *FunctionInput->GetDisplayName().ToString();
+	FText Unused;
+	return FNiagaraParameterUtilities::TestCanAddNamespaceModifierWithMessage(InputParameterName, Unused);
+}
+
+void SNiagaraStackFunctionInputName::OnAddNamespaceModifier()
+{
+	FName InputParameterName = *FunctionInput->GetDisplayName().ToString();
+	FName NewName = FNiagaraParameterUtilities::AddNamespaceModifier(InputParameterName);
+	if (NewName != NAME_None)
+	{
+		FScopedTransaction Transaction(LOCTEXT("AddNamespaceModifierTransaction", "Add namespace modifier."));
+		FunctionInput->OnRenamed(FText::FromName(NewName));
+	}
+}
+
+FText SNiagaraStackFunctionInputName::GetRemoveNamespaceModifierToolTip() const
+{
+	FName InputParameterName = *FunctionInput->GetDisplayName().ToString();
+	FText RemoveMessage;
+	FNiagaraParameterUtilities::TestCanRemoveNamespaceModifierWithMessage(InputParameterName, RemoveMessage);
+	return RemoveMessage;
+}
+
+bool SNiagaraStackFunctionInputName::CanRemoveNamespaceModifier() const
+{
+	FName InputParameterName = *FunctionInput->GetDisplayName().ToString();
+	FText Unused;
+	return FNiagaraParameterUtilities::TestCanRemoveNamespaceModifierWithMessage(InputParameterName, Unused);
+}
+
+void SNiagaraStackFunctionInputName::OnRemoveNamespaceModifier()
+{
+	FName InputParameterName = *FunctionInput->GetDisplayName().ToString();
+	FText Unused;
+	FName NewName = FNiagaraParameterUtilities::RemoveNamespaceModifier(InputParameterName);
+	if (NewName != NAME_None)
+	{
+		FScopedTransaction Transaction(LOCTEXT("RemoveNamespaceModifierTransaction", "Remove namespace modifier."));
+		FunctionInput->OnRenamed(FText::FromName(NewName));
+	}
+}
+
+FText SNiagaraStackFunctionInputName::GetEditNamespaceModifierToolTip() const
+{
+	FName InputParameterName = *FunctionInput->GetDisplayName().ToString();
+	FText EditMessage;
+	FNiagaraParameterUtilities::TestCanEditNamespaceModifierWithMessage(InputParameterName, EditMessage);
+	return EditMessage;
+}
+
+bool SNiagaraStackFunctionInputName::CanEditNamespaceModifier() const
+{
+	FName InputParameterName = *FunctionInput->GetDisplayName().ToString();
+	FText Unused;
+	return FNiagaraParameterUtilities::TestCanEditNamespaceModifierWithMessage(InputParameterName, Unused);
+}
+
+void SNiagaraStackFunctionInputName::OnEditNamespaceModifier()
+{
+	if (ParameterTextBlock.IsValid())
+	{
+		ParameterTextBlock->EnterNamespaceModifierEditingMode();
+	}
+}
+
+FText SNiagaraStackFunctionInputName::GetCopyParameterReferenceToolTip() const
+{
+	return LOCTEXT("CopyReferenceToolTip", "Copy a string reference for this parameter to the clipboard.\nThis reference can be used in expressions and custom HLSL nodes.");
+}
+
+bool SNiagaraStackFunctionInputName::CanCopyParameterReference() const
+{
+	return true;
+}
+
+void SNiagaraStackFunctionInputName::OnCopyParameterReference()
+{
+	FPlatformApplicationMisc::ClipboardCopy(*FunctionInput->GetDisplayName().ToString());
+}
+
+#undef LOCTEXT_NAMESPACE
diff --git a/Engine/Plugins/FX/Niagara/Source/NiagaraEditorWidgets/Private/Stack/SNiagaraStackFunctionInputName.h b/Engine/Plugins/FX/Niagara/Source/NiagaraEditorWidgets/Private/Stack/SNiagaraStackFunctionInputName.h
--- a/Engine/Plugins/FX/Niagara/Source/NiagaraEditorWidgets/Private/Stack/SNiagaraStackFunctionInputName.h
+++ b/Engine/Plugins/FX/Niagara/Source/NiagaraEditorWidgets/Private/Stack/SNiagaraStackFunctionInputName.h
@@ -5,6 +5,7 @@
 #include "SNiagaraStackEntryWidget.h"
 #include "Layout/Visibility.h"
 #include "Styling/SlateTypes.h"
+#include "NiagaraEditorSettings.h"
 
 class UNiagaraStackFunctionInput;
 class SInlineEditableTextBlock;
@@ -24,6 +25,8 @@ public:
 
 	virtual void Tick(const FGeometry& AllottedGeometry, const double InCurrentTime, const float InDeltaTime) override;
 
+	void FillRowContextMenu(class FMenuBuilder& MenuBuilder);
+
 private:
 	EVisibility GetEditConditionCheckBoxVisibility() const;
 
@@ -41,6 +44,27 @@ private:
 
 	void OnNameTextCommitted(const FText& InText, ETextCommit::Type InCommitType);
 
+	void GetChangeNamespaceSubMenu(FMenuBuilder& MenuBuilder);
+	void OnChangeNamespace(FNiagaraNamespaceMetadata Metadata);
+
+	void GetChangeNamespaceModifierSubMenu(FMenuBuilder& MenuBuilder);
+
+	FText GetAddNamespaceModifierToolTip() const;
+	bool CanAddNamespaceModifier() const;
+	void OnAddNamespaceModifier();
+
+	FText GetRemoveNamespaceModifierToolTip() const;
+	bool CanRemoveNamespaceModifier() const;
+	void OnRemoveNamespaceModifier();
+
+	FText GetEditNamespaceModifierToolTip() const;
+	bool CanEditNamespaceModifier() const;
+	void OnEditNamespaceModifier();
+
+	FText GetCopyParameterReferenceToolTip() const;
+	bool CanCopyParameterReference() const;
+	void OnCopyParameterReference();
+
 private:
 	UNiagaraStackFunctionInput* FunctionInput;
 
diff --git a/Engine/Plugins/FX/Niagara/Source/NiagaraEditorWidgets/Private/Stack/SNiagaraStackFunctionInputValue.cpp b/Engine/Plugins/FX/Niagara/Source/NiagaraEditorWidgets/Private/Stack/SNiagaraStackFunctionInputValue.cpp
--- a/Engine/Plugins/FX/Niagara/Source/NiagaraEditorWidgets/Private/Stack/SNiagaraStackFunctionInputValue.cpp
+++ b/Engine/Plugins/FX/Niagara/Source/NiagaraEditorWidgets/Private/Stack/SNiagaraStackFunctionInputValue.cpp
@@ -601,7 +601,7 @@ void SNiagaraStackFunctionInputValue::CollectAllActions(FGraphActionListBuilderB
 		const FText CategoryName = LOCTEXT("MakeCategory", "Make");
 
 		TArray<FName> AvailableNamespaces;
-		FunctionInput->GetNamespacesForNewParameters(AvailableNamespaces);
+		FunctionInput->GetNamespacesForNewReadParameters(AvailableNamespaces);
 
 		TArray<FString> InputNames;
 		for (int32 i = FunctionInput->GetInputParameterHandlePath().Num() - 1; i >= 0; i--)
