Add a stack error if the source emitter name is not set for a direct read DI.

#rnx

diff --git a/Engine/Plugins/FX/Niagara/Source/Niagara/Classes/NiagaraDataInterfaceParticleRead.h b/Engine/Plugins/FX/Niagara/Source/Niagara/Classes/NiagaraDataInterfaceParticleRead.h
--- a/Engine/Plugins/FX/Niagara/Source/Niagara/Classes/NiagaraDataInterfaceParticleRead.h
+++ b/Engine/Plugins/FX/Niagara/Source/Niagara/Classes/NiagaraDataInterfaceParticleRead.h
@@ -32,6 +32,9 @@ public:
 	virtual void GetParameterDefinitionHLSL(const FNiagaraDataInterfaceGPUParamInfo& ParamInfo, FString& OutHLSL) override;
 	virtual bool GetFunctionHLSL(const FNiagaraDataInterfaceGPUParamInfo& ParamInfo, const FNiagaraDataInterfaceGeneratedFunction& FunctionInfo, int FunctionInstanceIndex, FString& OutHLSL) override;
 	virtual void ProvidePerInstanceDataForRenderThread(void* DataForRenderThread, void* PerInstanceData, const FNiagaraSystemInstanceID& SystemInstance) override;
+#if WITH_EDITOR	
+	virtual void GetFeedback(UNiagaraSystem* Asset, UNiagaraComponent* Component, TArray<FNiagaraDataInterfaceError>& OutErrors, TArray<FNiagaraDataInterfaceFeedback>& Warnings, TArray<FNiagaraDataInterfaceFeedback>& Info) override;
+#endif
 	virtual void GetEmitterDependencies(void* PerInstanceData, FNiagaraSystemInstance* SystemInstance, TArray<FNiagaraEmitterInstance*>& Dependencies) const override;
 	//UNiagaraDataInterface Interface End
 
diff --git a/Engine/Plugins/FX/Niagara/Source/Niagara/Private/NiagaraDataInterfaceParticleRead.cpp b/Engine/Plugins/FX/Niagara/Source/Niagara/Private/NiagaraDataInterfaceParticleRead.cpp
--- a/Engine/Plugins/FX/Niagara/Source/Niagara/Private/NiagaraDataInterfaceParticleRead.cpp
+++ b/Engine/Plugins/FX/Niagara/Source/Niagara/Private/NiagaraDataInterfaceParticleRead.cpp
@@ -2028,6 +2028,30 @@ void UNiagaraDataInterfaceParticleRead::ProvidePerInstanceDataForRenderThread(vo
 	}
 }
 
+#if WITH_EDITOR	
+void UNiagaraDataInterfaceParticleRead::GetFeedback(UNiagaraSystem* Asset, UNiagaraComponent* Component, TArray<FNiagaraDataInterfaceError>& OutErrors, TArray<FNiagaraDataInterfaceFeedback>& Warnings, TArray<FNiagaraDataInterfaceFeedback>& Info)
+{
+	UNiagaraEmitter* FoundSourceEmitter = nullptr;
+	for (const FNiagaraEmitterHandle& EmitterHandle : Asset->GetEmitterHandles())
+	{
+		UNiagaraEmitter* EmitterInstance = EmitterHandle.GetInstance();
+		if (EmitterInstance && EmitterInstance->GetUniqueEmitterName() == EmitterName)
+		{
+			FoundSourceEmitter = EmitterInstance;
+			break;
+		}
+	}
+
+	if (!FoundSourceEmitter)
+	{
+		FNiagaraDataInterfaceError SourceEmitterNotFoundError(LOCTEXT("SourceEmitterNotFoundError", "Source emitter not found."),
+			LOCTEXT("SourceEmitterNotFoundErrorSummary", "Source emitter not found"),
+			FNiagaraDataInterfaceFix());
+		OutErrors.Add(SourceEmitterNotFoundError);
+	}
+}
+#endif
+
 void UNiagaraDataInterfaceParticleRead::GetEmitterDependencies(void* PerInstanceData, FNiagaraSystemInstance* SystemInstance, TArray<FNiagaraEmitterInstance*>& Dependencies) const
 {
 	FNDIParticleRead_InstanceData* PIData = static_cast<FNDIParticleRead_InstanceData*>(PerInstanceData);
