Niagara - Parameter menu fixes
+ In the set variables node, replace add button's standard menu with a searchable menu like for linked inputs and modules.
+ In the set variables add menu show parameters which were added directly in the parameters panel.
+ Set a min width on the link inputs menu so it's not constantly changing size.

diff --git a/Engine/Plugins/FX/Niagara/Source/NiagaraEditor/Private/NiagaraNodeAssignment.cpp b/Engine/Plugins/FX/Niagara/Source/NiagaraEditor/Private/NiagaraNodeAssignment.cpp
--- a/Engine/Plugins/FX/Niagara/Source/NiagaraEditor/Private/NiagaraNodeAssignment.cpp
+++ b/Engine/Plugins/FX/Niagara/Source/NiagaraEditor/Private/NiagaraNodeAssignment.cpp
@@ -255,7 +255,7 @@ void UNiagaraNodeAssignment::MergeUp()
 	//NiagaraStackUtilities::
 }
 
-void UNiagaraNodeAssignment::BuildAddParameterMenu(FMenuBuilder& MenuBuilder, ENiagaraScriptUsage InUsage, UNiagaraNodeOutput* InGraphOutputNode)
+void UNiagaraNodeAssignment::CollectAddExistingActions(ENiagaraScriptUsage InUsage, UNiagaraNodeOutput* InGraphOutputNode, TArray<TSharedPtr<FNiagaraMenuAction>>& OutAddExistingActions)
 {
 	TArray<FNiagaraVariable> AvailableParameters;
 	FNiagaraStackGraphUtilities::GetAvailableParametersForScript(*InGraphOutputNode, AvailableParameters);
@@ -266,25 +266,20 @@ void UNiagaraNodeAssignment::BuildAddParameterMenu(FMenuBuilder& MenuBuilder, EN
 		FText VarDesc = FNiagaraConstants::GetAttributeDescription(AvailableParameter);
 		FString VarDefaultValue = FNiagaraConstants::GetAttributeDefaultValue(AvailableParameter);
 		const FText TooltipDesc = FText::Format(LOCTEXT("SetFunctionPopupTooltip", "Description: Set the parameter {0}. {1}"), NameText, VarDesc);
-		FText CategoryName = LOCTEXT("ModuleSetCategory", "Set Specific Parameters");
+		FText Category = LOCTEXT("ModuleSetCategory", "Set Specific Parameters");
 		bool bCanExecute = AssignmentTargets.Contains(AvailableParameter) == false; 
 
-		MenuBuilder.AddMenuEntry(
-			FUIAction(
-				FExecuteAction::CreateUObject(this, &UNiagaraNodeAssignment::AddParameter, AvailableParameter, VarDefaultValue),
-				FCanExecuteAction::CreateLambda([bCanExecute] { return bCanExecute; })),
-			SNew(SBox)
-			.Padding(FMargin(5, 1, 5, 1))
-			[
-				SNew(SNiagaraParameterName)
-				.ParameterName(AvailableParameter.GetName())
-				.IsReadOnly(true)
-			]);
+		TSharedRef<FNiagaraMenuAction> AddExistingAction = MakeShareable<FNiagaraMenuAction>(new FNiagaraMenuAction(
+			Category, NameText, TooltipDesc,
+			0, FText(),
+			FNiagaraMenuAction::FOnExecuteStackAction::CreateUObject(this, &UNiagaraNodeAssignment::AddParameter, AvailableParameter, VarDefaultValue),
+			FNiagaraMenuAction::FCanExecuteStackAction::CreateLambda([bCanExecute] { return bCanExecute; })));
+		AddExistingAction->SetParamterVariable(AvailableParameter);
+		OutAddExistingActions.Add(AddExistingAction);
 	}
-
 }
 
-void UNiagaraNodeAssignment::BuildCreateParameterMenu(FMenuBuilder& MenuBuilder, ENiagaraScriptUsage InUsage, UNiagaraNodeOutput* InGraphOutputNode)
+void UNiagaraNodeAssignment::CollectCreateNewActions(ENiagaraScriptUsage InUsage, UNiagaraNodeOutput* InGraphOutputNode, TArray<TSharedPtr<FNiagaraMenuAction>>& OutCreateNewActions)
 {
 	// Generate actions for creating new typed parameters.
 	TOptional<FName> NewParameterNamespace = FNiagaraStackGraphUtilities::GetNamespaceForScriptUsage(InGraphOutputNode->GetUsage());
@@ -323,12 +318,11 @@ void UNiagaraNodeAssignment::BuildCreateParameterMenu(FMenuBuilder& MenuBuilder,
 			const FText TooltipDesc = FText::Format(LOCTEXT("NewParameterModuleDescriptionFormat", "Description: Create a new {0} parameter. {1}"), TypeText, VarDesc);
 			FText Category = LOCTEXT("NewParameterModuleCategory", "Create New Parameter");
 
-			MenuBuilder.AddMenuEntry(
-				TypeText,
-				TooltipDesc,
-				FSlateIcon(),
-				FUIAction(
-					FExecuteAction::CreateUObject(this, &UNiagaraNodeAssignment::AddParameter, NewParameter, VarDefaultValue)));
+			TSharedRef<FNiagaraMenuAction> CreateNewAction = MakeShareable<FNiagaraMenuAction>(new FNiagaraMenuAction(
+				Category, TypeText, TooltipDesc,
+				0, FText(),
+				FNiagaraMenuAction::FOnExecuteStackAction::CreateUObject(this, &UNiagaraNodeAssignment::AddParameter, NewParameter, VarDefaultValue)));
+			OutCreateNewActions.Add(CreateNewAction);
 		}
 	}
 }
diff --git a/Engine/Plugins/FX/Niagara/Source/NiagaraEditor/Private/ViewModels/Stack/NiagaraStackGraphUtilities.cpp b/Engine/Plugins/FX/Niagara/Source/NiagaraEditor/Private/ViewModels/Stack/NiagaraStackGraphUtilities.cpp
--- a/Engine/Plugins/FX/Niagara/Source/NiagaraEditor/Private/ViewModels/Stack/NiagaraStackGraphUtilities.cpp
+++ b/Engine/Plugins/FX/Niagara/Source/NiagaraEditor/Private/ViewModels/Stack/NiagaraStackGraphUtilities.cpp
@@ -1597,6 +1597,20 @@ void FNiagaraStackGraphUtilities::GetAvailableParametersForScript(UNiagaraNodeOu
 				OutAvailableParameters.AddUnique(Entry.Key);
 			}
 		}
+
+		UNiagaraSystem* System = ScriptOutputNode.GetTypedOuter<UNiagaraSystem>();
+		if (System != nullptr)
+		{
+			TArray<FNiagaraVariable> EditorOnlyParameters;
+			System->EditorOnlyAddedParameters.GetParameters(EditorOnlyParameters);
+			for (FNiagaraVariable& EditorOnlyParameter : EditorOnlyParameters)
+			{
+				if (EditorOnlyParameter.IsInNameSpace(UsageNamespace.GetValue().ToString()))
+				{
+					OutAvailableParameters.AddUnique(EditorOnlyParameter);
+				}
+			}
+		}
 	}
 }
 
diff --git a/Engine/Plugins/FX/Niagara/Source/NiagaraEditor/Public/NiagaraNodeAssignment.h b/Engine/Plugins/FX/Niagara/Source/NiagaraEditor/Public/NiagaraNodeAssignment.h
--- a/Engine/Plugins/FX/Niagara/Source/NiagaraEditor/Public/NiagaraNodeAssignment.h
+++ b/Engine/Plugins/FX/Niagara/Source/NiagaraEditor/Public/NiagaraNodeAssignment.h
@@ -35,8 +35,8 @@ public:
 	const TArray<FString>& GetAssignmentDefaults() const { return AssignmentDefaultValues; }
 
 	void MergeUp();
-	void BuildAddParameterMenu(class FMenuBuilder& MenuBuilder, ENiagaraScriptUsage InUsage, UNiagaraNodeOutput* InGraphOutputNode);
-	void BuildCreateParameterMenu(class FMenuBuilder& MenuBuilder, ENiagaraScriptUsage InUsage, UNiagaraNodeOutput* InGraphOutputNode);
+	void CollectAddExistingActions(ENiagaraScriptUsage InUsage, UNiagaraNodeOutput* InGraphOutputNode, TArray<TSharedPtr<FNiagaraMenuAction>>& OutAddExistingActions);
+	void CollectCreateNewActions(ENiagaraScriptUsage InUsage, UNiagaraNodeOutput* InGraphOutputNode, TArray<TSharedPtr<FNiagaraMenuAction>>& OutCreateNewActions);
 
 	//~ Begin EdGraphNode Interface
 	virtual void PostLoad() override;
diff --git a/Engine/Plugins/FX/Niagara/Source/NiagaraEditorWidgets/Private/Stack/SNiagaraStackFunctionInputValue.cpp b/Engine/Plugins/FX/Niagara/Source/NiagaraEditorWidgets/Private/Stack/SNiagaraStackFunctionInputValue.cpp
--- a/Engine/Plugins/FX/Niagara/Source/NiagaraEditorWidgets/Private/Stack/SNiagaraStackFunctionInputValue.cpp
+++ b/Engine/Plugins/FX/Niagara/Source/NiagaraEditorWidgets/Private/Stack/SNiagaraStackFunctionInputValue.cpp
@@ -486,6 +486,7 @@ TSharedRef<SWidget> SNiagaraStackFunctionInputValue::OnGetAvailableHandleMenu()
 	.Padding(5)
 	[
 		SNew(SBox)
+		.MinDesiredWidth(300)
 		.HeightOverride(400)
 		[
 			SNew(SVerticalBox)
diff --git a/Engine/Plugins/FX/Niagara/Source/NiagaraEditorWidgets/Private/Stack/SNiagaraStackModuleItem.cpp b/Engine/Plugins/FX/Niagara/Source/NiagaraEditorWidgets/Private/Stack/SNiagaraStackModuleItem.cpp
--- a/Engine/Plugins/FX/Niagara/Source/NiagaraEditorWidgets/Private/Stack/SNiagaraStackModuleItem.cpp
+++ b/Engine/Plugins/FX/Niagara/Source/NiagaraEditorWidgets/Private/Stack/SNiagaraStackModuleItem.cpp
@@ -33,6 +33,7 @@
 #include "ViewModels/Stack/NiagaraStackClipboardUtilities.h"
 #include "Widgets/SNiagaraLibraryOnlyToggleHeader.h"
 #include "NiagaraEditorModule.h"
+#include "SNiagaraGraphActionWidget.h"
 
 #define LOCTEXT_NAMESPACE "NiagaraStackModuleItem"
 
@@ -108,7 +109,7 @@ void SNiagaraStackModuleItem::AddCustomRowWidgets(TSharedRef<SHorizontalBox> Hor
 	.VAlign(VAlign_Center)
 	.AutoWidth()
 	[
-		SNew(SComboButton)
+		SAssignNew(AddButton, SComboButton)
 		.HasDownArrow(false)
 		.ButtonStyle(FEditorStyle::Get(), "HoverHintOnly")
 		.ForegroundColor(FSlateColor::UseForeground())
@@ -189,35 +190,48 @@ FReply SNiagaraStackModuleItem::ScratchButtonPressed() const
 	return FReply::Unhandled();
 }
 
+void OnActionSelected(const TArray<TSharedPtr<FEdGraphSchemaAction>>& SelectedActions, ESelectInfo::Type InSelectionType)
+{
+	if (SelectedActions.Num() == 1 && (InSelectionType == ESelectInfo::OnKeyPress || InSelectionType == ESelectInfo::OnMouseClick))
+	{
+		TSharedPtr<FNiagaraMenuAction> Action = StaticCastSharedPtr<FNiagaraMenuAction>(SelectedActions[0]);
+		if (Action.IsValid())
+		{
+			FSlateApplication::Get().DismissAllMenus();
+			Action->ExecuteAction();
+		}
+	}
+}
+
 TSharedRef<SWidget> SNiagaraStackModuleItem::RaiseActionMenuClicked()
 {
 	if (CanRaiseActionMenu())
 	{
-		UNiagaraNodeAssignment* AssignmentNode = Cast<UNiagaraNodeAssignment>(&ModuleItem->GetModuleNode());
-		if (AssignmentNode != nullptr)
-		{
-			UNiagaraNodeOutput* OutputNode = ModuleItem->GetOutputNode();
-			if (OutputNode)
-			{
-				FMenuBuilder MenuBuilder(true, nullptr);
-				MenuBuilder.AddSubMenu(LOCTEXT("AddVariables", "Add Variable"), 
-					LOCTEXT("AddVariablesTooltip", "Add another variable to the end of the list"),
-					FNewMenuDelegate::CreateLambda([OutputNode, AssignmentNode](FMenuBuilder& SubMenuBuilder)
-				{
-					AssignmentNode->BuildAddParameterMenu(SubMenuBuilder, OutputNode->GetUsage(), OutputNode);
-				}));
-				MenuBuilder.AddSubMenu(LOCTEXT("CreateVariables", "Create New Variable"),
-					LOCTEXT("CreateVariablesTooltip", "Create a new variable and set its value"),
-					FNewMenuDelegate::CreateLambda([OutputNode, AssignmentNode](FMenuBuilder& SubMenuBuilder)
+		TSharedPtr<SGraphActionMenu> GraphActionMenu;
+
+		TSharedRef<SBorder> MenuWidget = SNew(SBorder)
+		.BorderImage(FEditorStyle::GetBrush("Menu.Background"))
+		.Padding(5)
+		[
+			SNew(SBox)
+			.MinDesiredWidth(300)
+			.HeightOverride(400)
+			[
+				SAssignNew(GraphActionMenu, SGraphActionMenu)
+				.OnActionSelected_Static(OnActionSelected)
+				.OnCollectAllActions(this, &SNiagaraStackModuleItem::CollectParameterActions)
+				.AutoExpandActionMenu(false)
+				.ShowFilterTextBox(true)
+				.OnCreateWidgetForAction_Lambda([](const FCreateWidgetForActionData* InData)
 				{
-					AssignmentNode->BuildCreateParameterMenu(SubMenuBuilder, OutputNode->GetUsage(), OutputNode);
-				}));
+					return SNew(SNiagaraGraphActionWidget, InData);
+				})
+			]
+		];
 
-				return MenuBuilder.MakeWidget();
-			}
-		}
+		AddButton->SetMenuContentWidgetToFocus(GraphActionMenu->GetFilterTextBox()->AsShared());
+		return MenuWidget;
 	}
-
 	return SNullWidget::NullWidget;
 }
 
@@ -269,15 +283,22 @@ void ReassignModuleScript(UNiagaraStackModuleItem* ModuleItem, FAssetData NewMod
 	}
 }
 
-void OnActionSelected(const TArray<TSharedPtr<FEdGraphSchemaAction>>& SelectedActions, ESelectInfo::Type InSelectionType)
+void SNiagaraStackModuleItem::CollectParameterActions(FGraphActionListBuilderBase& ModuleActions)
 {
-	if (SelectedActions.Num() == 1 && (InSelectionType == ESelectInfo::OnKeyPress || InSelectionType == ESelectInfo::OnMouseClick))
+	UNiagaraNodeAssignment* AssignmentNode = Cast<UNiagaraNodeAssignment>(&ModuleItem->GetModuleNode());
+	if (AssignmentNode != nullptr)
 	{
-		TSharedPtr<FNiagaraMenuAction> Action = StaticCastSharedPtr<FNiagaraMenuAction>(SelectedActions[0]);
-		if (Action.IsValid())
+		UNiagaraNodeOutput* OutputNode = ModuleItem->GetOutputNode();
+		if (OutputNode != nullptr)
 		{
-			FSlateApplication::Get().DismissAllMenus();
-			Action->ExecuteAction();
+			TArray<TSharedPtr<FNiagaraMenuAction>> AllActions;
+			AssignmentNode->CollectAddExistingActions(OutputNode->GetUsage(), OutputNode, AllActions);
+			AssignmentNode->CollectCreateNewActions(OutputNode->GetUsage(), OutputNode, AllActions);
+			
+			for (TSharedPtr<FNiagaraMenuAction> Action : AllActions)
+			{
+				ModuleActions.AddAction(Action);
+			}
 		}
 	}
 }
diff --git a/Engine/Plugins/FX/Niagara/Source/NiagaraEditorWidgets/Private/Stack/SNiagaraStackModuleItem.h b/Engine/Plugins/FX/Niagara/Source/NiagaraEditorWidgets/Private/Stack/SNiagaraStackModuleItem.h
--- a/Engine/Plugins/FX/Niagara/Source/NiagaraEditorWidgets/Private/Stack/SNiagaraStackModuleItem.h
+++ b/Engine/Plugins/FX/Niagara/Source/NiagaraEditorWidgets/Private/Stack/SNiagaraStackModuleItem.h
@@ -11,6 +11,7 @@ class UNiagaraStackModuleItem;
 class UNiagaraStackViewModel;
 class SNiagaraStackDisplayName;
 struct FGraphActionListBuilderBase;
+class SComboButton;
 
 class SNiagaraStackModuleItem : public SNiagaraStackItem
 {
@@ -50,6 +51,8 @@ private:
 
 	bool OnModuleItemAllowDrop(TSharedPtr<class FDragDropOperation> DragDropOperation);
 
+	void CollectParameterActions(FGraphActionListBuilderBase& ModuleActions);
+
 	void CollectModuleActions(FGraphActionListBuilderBase& ModuleActions);
 
 	void ShowReassignModuleScriptMenu();
@@ -61,5 +64,7 @@ private:
 private:
 	UNiagaraStackModuleItem* ModuleItem;
 
+	TSharedPtr<SComboButton> AddButton;
+
 	static bool bLibraryOnly;
 };
\ No newline at end of file
