Add on selection change from graph delegate to SNiagaraParameterMapView.

diff --git a/Engine/Plugins/FX/Niagara/Source/NiagaraEditor/Private/Widgets/SNiagaraGraphParameterMapGetNode.cpp b/Engine/Plugins/FX/Niagara/Source/NiagaraEditor/Private/Widgets/SNiagaraGraphParameterMapGetNode.cpp
--- a/Engine/Plugins/FX/Niagara/Source/NiagaraEditor/Private/Widgets/SNiagaraGraphParameterMapGetNode.cpp
+++ b/Engine/Plugins/FX/Niagara/Source/NiagaraEditor/Private/Widgets/SNiagaraGraphParameterMapGetNode.cpp
@@ -154,7 +154,7 @@ void SNiagaraGraphParameterMapGetNode::CreatePinWidgets()
 			.HAlign(HAlign_Fill)
 			.VAlign(VAlign_Fill)
 			.Padding(FMargin(0, 3))
-			.OnMouseButtonDown(this, &SNiagaraGraphParameterMapGetNode::OnBorderMouseButtonDown, i)
+			//.OnMouseButtonDown(this, &SNiagaraGraphParameterMapGetNode::OnBorderMouseButtonDown, i)
 			[
 				Widget.ToSharedRef()
 			];
diff --git a/Engine/Plugins/FX/Niagara/Source/NiagaraEditor/Private/Widgets/SNiagaraParameterMapView.cpp b/Engine/Plugins/FX/Niagara/Source/NiagaraEditor/Private/Widgets/SNiagaraParameterMapView.cpp
--- a/Engine/Plugins/FX/Niagara/Source/NiagaraEditor/Private/Widgets/SNiagaraParameterMapView.cpp
+++ b/Engine/Plugins/FX/Niagara/Source/NiagaraEditor/Private/Widgets/SNiagaraParameterMapView.cpp
@@ -1010,6 +1010,10 @@ void SNiagaraParameterMapView::AddGraph(UNiagaraGraph* Graph)
 			FOnGraphChanged::FDelegate::CreateRaw(this, &SNiagaraParameterMapView::OnGraphChanged));
 		FDelegateHandle OnRecompileHandle = Graph->AddOnGraphNeedsRecompileHandler(
 			FOnGraphChanged::FDelegate::CreateRaw(this, &SNiagaraParameterMapView::OnGraphChanged));
+		if (ToolkitType == EToolkitType::SCRIPT)
+		{
+			OnSubObjectSelectionChangedHandle = Graph->OnSubObjectSelectionChanged().AddSP(this, &SNiagaraParameterMapView::HandleGraphSubObjectSelectionChanged);
+		}
 
 		OnGraphChangedHandles.Add(OnGraphChangedHandle);
 		OnRecompileHandles.Add(OnRecompileHandle);
@@ -1542,6 +1546,15 @@ void SNiagaraParameterMapView::RefreshActions()
 	bNeedsRefresh = true;
 }
 
+void SNiagaraParameterMapView::HandleGraphSubObjectSelectionChanged(const UObject* NewSelection)
+{
+	if (NewSelection->IsA<UNiagaraScriptVariable>())
+	{
+		GraphActionMenu->SelectItemByName(static_cast<const UNiagaraScriptVariable*>(NewSelection)->Variable.GetName());
+	}
+	SelectedVariableObjects->SetSelectedObject(const_cast<UObject*>(NewSelection));
+}
+
 bool SNiagaraParameterMapView::IsStaticSwitchParameter(const FNiagaraVariable& Variable, const TArray<TWeakObjectPtr<UNiagaraGraph>>& Graphs)
 {
 	for (auto& GraphWeakPtr : Graphs)
diff --git a/Engine/Plugins/FX/Niagara/Source/NiagaraEditor/Public/Widgets/SNiagaraParameterMapView.h b/Engine/Plugins/FX/Niagara/Source/NiagaraEditor/Public/Widgets/SNiagaraParameterMapView.h
--- a/Engine/Plugins/FX/Niagara/Source/NiagaraEditor/Public/Widgets/SNiagaraParameterMapView.h
+++ b/Engine/Plugins/FX/Niagara/Source/NiagaraEditor/Public/Widgets/SNiagaraParameterMapView.h
@@ -189,6 +189,8 @@ private:
 	/** Sets bNeedsRefresh to true. Causing the list to be refreshed next tick. */
 	void RefreshActions();
 
+	void HandleGraphSubObjectSelectionChanged(const UObject* NewSelection);
+
 	/** Graph Action Menu for displaying all our variables and functions */
 	TSharedPtr<SGraphActionMenu> GraphActionMenu;
 
@@ -210,6 +212,7 @@ private:
 	/** The handle to the graph changed delegate. */
 	TArray<FDelegateHandle> OnGraphChangedHandles;
 	TArray<FDelegateHandle> OnRecompileHandles;
+	FDelegateHandle OnSubObjectSelectionChangedHandle;
 
 	EToolkitType ToolkitType;
 	TSharedPtr<FUICommandList> ToolkitCommands;
