Prevent copy pasting to emitter overview graphs.

diff --git a/Engine/Plugins/FX/Niagara/Source/NiagaraEditor/Private/ViewModels/NiagaraOverviewGraphViewModel.cpp b/Engine/Plugins/FX/Niagara/Source/NiagaraEditor/Private/ViewModels/NiagaraOverviewGraphViewModel.cpp
--- a/Engine/Plugins/FX/Niagara/Source/NiagaraEditor/Private/ViewModels/NiagaraOverviewGraphViewModel.cpp
+++ b/Engine/Plugins/FX/Niagara/Source/NiagaraEditor/Private/ViewModels/NiagaraOverviewGraphViewModel.cpp
@@ -284,6 +284,14 @@ void FNiagaraOverviewGraphViewModel::PasteNodes()
 	UEdGraph* Graph = GetGraph();
 	if (Graph != nullptr)
 	{
+		// Check if we need to early out due to trying to paste nodes to an emitter asset edit mode overview graph.
+		TSharedPtr<FNiagaraSystemViewModel> WeakSystemViewModel = SystemViewModel.Pin();
+		ENiagaraSystemViewModelEditMode SystemEditMode = WeakSystemViewModel->GetEditMode();
+		if (SystemEditMode == ENiagaraSystemViewModelEditMode::EmitterAsset)
+		{
+			return;
+		}
+
 		const FScopedTransaction Transaction(FGenericCommands::Get().Paste->GetDescription());;
 		Graph->Modify();
 
