Niagara - Fix the rename toasts in the graph to show the new parameter style text.

diff --git a/Engine/Plugins/FX/Niagara/Source/NiagaraEditor/Private/NiagaraEditorUtilities.cpp b/Engine/Plugins/FX/Niagara/Source/NiagaraEditor/Private/NiagaraEditorUtilities.cpp
--- a/Engine/Plugins/FX/Niagara/Source/NiagaraEditor/Private/NiagaraEditorUtilities.cpp
+++ b/Engine/Plugins/FX/Niagara/Source/NiagaraEditor/Private/NiagaraEditorUtilities.cpp
@@ -2566,12 +2566,73 @@ FText FNiagaraParameterUtilities::FormatParameterNameForTextDisplay(FName Parame
 	FNiagaraParameterHandle ParameterHandle(ParameterName);
 	TArray<FName> HandleParts = ParameterHandle.GetHandleParts();
 	FString DisplayString;
-	for (int32 HandlePartIndex = 0; HandlePartIndex < HandleParts.Num() - 1; HandlePartIndex++)
+
+	if (HandleParts.Num() == 0)
+	{
+		return FText();
+	}
+	else if (HandleParts.Num() == 1)
+	{
+		return FText::FromName(HandleParts[0]);
+	}
+
+	FText NamespaceText;
+	FNiagaraNamespaceMetadata NamespaceMetadata = GetDefault<UNiagaraEditorSettings>()->GetMetaDataForNamespaces(HandleParts);
+	if (NamespaceMetadata.IsValid() && NamespaceMetadata.DisplayName.IsEmptyOrWhitespace() == false)
+	{
+		NamespaceText = NamespaceMetadata.DisplayName.ToUpper();
+		HandleParts.RemoveAt(0, NamespaceMetadata.Namespaces.Num());
+	}
+	else
+	{
+		NamespaceText = FText::FromName(HandleParts[0]).ToUpper();
+		HandleParts.RemoveAt(0);
+	}
+
+	FText NamespaceModfierText;
+	if(HandleParts.Num() > 1)
+	{
+		FNiagaraNamespaceMetadata NamespaceModifierMetadata = GetDefault<UNiagaraEditorSettings>()->GetMetaDataForNamespaceModifier(HandleParts[0]);
+		if (NamespaceModifierMetadata.IsValid() && NamespaceModifierMetadata.DisplayName.IsEmptyOrWhitespace() == false)
+		{
+			NamespaceModfierText = NamespaceModifierMetadata.DisplayName.ToUpper();
+		}
+		else
+		{
+			NamespaceModfierText = FText::FromName(HandleParts[0]).ToUpper();
+		}
+		HandleParts.RemoveAt(0);
+	}
+
+	FText NameText;
+	if (HandleParts.Num() == 0)
+	{
+		NameText = FText::FromName(NAME_None);
+	}
+	else if (HandleParts.Num() == 1)
+	{
+		NameText = FText::FromName(HandleParts[0]);
+	}
+	else
+	{
+		// This shouldn't happen for new data, and old data should have been fixed up, but handle it
+		// anyway just in case.
+		TArray<FString> HandlePartStrings;
+		for (FName HandlePart : HandleParts)
+		{
+			HandlePartStrings.Add(HandlePart.ToString());
+		}
+		NameText = FText::FromString(FString::Join(HandlePartStrings, TEXT(".")));
+	}
+
+	if (NamespaceModfierText.IsEmptyOrWhitespace())
+	{
+		return FText::Format(LOCTEXT("ParameterNameFormat", "({0}) {1}"), NamespaceText, NameText);
+	}
+	else
 	{
-		DisplayString += TEXT("(") + HandleParts[HandlePartIndex].ToString().ToUpper() + TEXT(") ");
+		return FText::Format(LOCTEXT("ParameterNameWithModifierFormat", "({0}) ({1}) {2}"), NamespaceText, NamespaceModfierText, NameText);
 	}
-	DisplayString += HandleParts[HandleParts.Num() - 1].ToString();
-	return FText::FromString(DisplayString);
 }
 
 FName NamePartsToName(const TArray<FName>& NameParts)
diff --git a/Engine/Plugins/FX/Niagara/Source/NiagaraEditor/Private/NiagaraGraph.cpp b/Engine/Plugins/FX/Niagara/Source/NiagaraEditor/Private/NiagaraGraph.cpp
--- a/Engine/Plugins/FX/Niagara/Source/NiagaraEditor/Private/NiagaraGraph.cpp
+++ b/Engine/Plugins/FX/Niagara/Source/NiagaraEditor/Private/NiagaraGraph.cpp
@@ -1558,9 +1558,9 @@ bool UNiagaraGraph::RenameParameterFromPin(const FNiagaraVariable& Parameter, FN
 			{
 				if (!bMerged)
 				{
-					FName ParamName = NewName;
-					FNiagaraEditorUtilities::DecomposeVariableNamespace(NewName, ParamName);
-					FNiagaraEditorUtilities::InfoWithToastAndLog(FText::Format(LOCTEXT("RenamedVarInGraphForAll", "\"{0}\" has been fully renamed as it was only used on this node."), FText::FromName(ParamName)));
+					FNiagaraEditorUtilities::InfoWithToastAndLog(FText::Format(
+						LOCTEXT("RenamedVarInGraphForAll", "\"{0}\" has been fully renamed as it was only used on this node."), 
+						FNiagaraParameterUtilities::FormatParameterNameForTextDisplay(NewName)));
 				}
 				return true;
 			}
@@ -1630,21 +1630,19 @@ bool UNiagaraGraph::RenameParameterFromPin(const FNiagaraVariable& Parameter, FN
 	else
 	{
 		bMerged = true;
-		FName NewParamName = NewName;
-		FNiagaraEditorUtilities::DecomposeVariableNamespace(NewName, NewParamName);
-		FName OldParamName = Parameter.GetName();
-		FNiagaraEditorUtilities::DecomposeVariableNamespace(Parameter.GetName(), OldParamName);
-		FNiagaraEditorUtilities::InfoWithToastAndLog(FText::Format(LOCTEXT("MergedVar", "\"{0}\" has been merged with parameter \"{1}\".\nAll of \"{1}\"'s meta-data will be used, overwriting \"{0}\"'s meta-data."), FText::FromName(OldParamName), FText::FromName(NewParamName)));
+		FNiagaraEditorUtilities::InfoWithToastAndLog(FText::Format(
+			LOCTEXT("MergedVar", "\"{0}\" has been merged with parameter \"{1}\".\nAll of \"{1}\"'s meta-data will be used, overwriting \"{0}\"'s meta-data."),
+			FNiagaraParameterUtilities::FormatParameterNameForTextDisplay(Parameter.GetName()),
+			FNiagaraParameterUtilities::FormatParameterNameForTextDisplay(NewName)));
 	}
 	RefreshParameterReferences();
 
 	if (!bMerged)
 	{
-		FName NewParamName = NewName;
-		FNiagaraEditorUtilities::DecomposeVariableNamespace(NewName, NewParamName);
-		FName OldParamName = Parameter.GetName();
-		FNiagaraEditorUtilities::DecomposeVariableNamespace(Parameter.GetName(), OldParamName);
-		FNiagaraEditorUtilities::InfoWithToastAndLog(FText::Format(LOCTEXT("RenamedVarInGraphForNode", "\"{0}\" has been duplicated as \"{1}\" as it is used in multiple locations.\nPlease edit the Parameters Panel version to change in all locations."), FText::FromName(OldParamName), FText::FromName(NewParamName)));
+		FNiagaraEditorUtilities::InfoWithToastAndLog(FText::Format(
+			LOCTEXT("RenamedVarInGraphForNode", "\"{0}\" has been duplicated as \"{1}\" as it is used in multiple locations.\nPlease edit the Parameters Panel version to change in all locations."),
+			FNiagaraParameterUtilities::FormatParameterNameForTextDisplay(Parameter.GetName()), 
+			FNiagaraParameterUtilities::FormatParameterNameForTextDisplay(NewName)));
 	}
 
 	return false;
@@ -1740,7 +1738,10 @@ bool UNiagaraGraph::RenameParameter(const FNiagaraVariable& Parameter, FName New
 		FNiagaraEditorUtilities::DecomposeVariableNamespace(NewName, NewParamName);
 		FName OldParamName = Parameter.GetName();
 		FNiagaraEditorUtilities::DecomposeVariableNamespace(Parameter.GetName(), OldParamName);
-		FNiagaraEditorUtilities::InfoWithToastAndLog(FText::Format(LOCTEXT("MergedVar", "\"{0}\" has been merged with parameter \"{1}\".\nAll of \"{1}\"'s meta-data will be used, overwriting \"{0}\"'s meta-data."), FText::FromName(OldParamName), FText::FromName(NewParamName)));
+		FNiagaraEditorUtilities::InfoWithToastAndLog(FText::Format(
+			LOCTEXT("MergedVar", "\"{0}\" has been merged with parameter \"{1}\".\nAll of \"{1}\"'s meta-data will be used, overwriting \"{0}\"'s meta-data."),
+			FNiagaraParameterUtilities::FormatParameterNameForTextDisplay(Parameter.GetName()),
+			FNiagaraParameterUtilities::FormatParameterNameForTextDisplay(NewName)));
 	}
 
 	// Fixup reference collection and pin names
