Niagara - Fix stack search to be more consistent.

diff --git a/Engine/Plugins/FX/Niagara/Source/NiagaraEditor/Private/ViewModels/Stack/NiagaraStackModuleItemOutput.cpp b/Engine/Plugins/FX/Niagara/Source/NiagaraEditor/Private/ViewModels/Stack/NiagaraStackModuleItemOutput.cpp
--- a/Engine/Plugins/FX/Niagara/Source/NiagaraEditor/Private/ViewModels/Stack/NiagaraStackModuleItemOutput.cpp
+++ b/Engine/Plugins/FX/Niagara/Source/NiagaraEditor/Private/ViewModels/Stack/NiagaraStackModuleItemOutput.cpp
@@ -69,7 +69,6 @@ UNiagaraStackEntry::EStackRowStyle UNiagaraStackModuleItemOutput::GetStackRowSty
 
 void UNiagaraStackModuleItemOutput::GetSearchItems(TArray<FStackSearchItem>& SearchItems) const
 {
-	SearchItems.Add({ FName("DisplayName"), GetDisplayName() });
 	SearchItems.Add({ FName("OutputParamHandleText"), GetOutputParameterHandleText() });
 }
 
diff --git a/Engine/Plugins/FX/Niagara/Source/NiagaraEditor/Private/ViewModels/Stack/NiagaraStackViewModel.cpp b/Engine/Plugins/FX/Niagara/Source/NiagaraEditor/Private/ViewModels/Stack/NiagaraStackViewModel.cpp
--- a/Engine/Plugins/FX/Niagara/Source/NiagaraEditor/Private/ViewModels/Stack/NiagaraStackViewModel.cpp
+++ b/Engine/Plugins/FX/Niagara/Source/NiagaraEditor/Private/ViewModels/Stack/NiagaraStackViewModel.cpp
@@ -192,8 +192,8 @@ void UNiagaraStackViewModel::Tick()
 		if (bRefreshPending)
 		{
 			RootEntry->RefreshChildren();
-			OnSearchTextChanged(CurrentSearchText);
 			bRefreshPending = false;
+			InvalidateSearchResults();
 		}
 
 		SearchTick();
@@ -205,7 +205,6 @@ void UNiagaraStackViewModel::OnSearchTextChanged(const FText& SearchText)
 	if (RootEntry && !CurrentSearchText.EqualTo(SearchText))
 	{
 		CurrentSearchText = SearchText;
-		RestoreStackEntryExpansionPreSearch();
 		// postpone searching until next tick; protects against crashes from the GC
 		// also this can be triggered by multiple events, so better wait
 		bRestartSearch = true;
@@ -469,26 +468,9 @@ void UNiagaraStackViewModel::GeneratePathForEntry(UNiagaraStackEntry* Root, UNia
 	}
 }
 
-void UNiagaraStackViewModel::RestoreStackEntryExpansionPreSearch()
+void UNiagaraStackViewModel::InvalidateSearchResults()
 {
-	GenerateTraversalEntries(RootEntry, TArray<UNiagaraStackEntry*>(), ItemsToRestoreExpansionState);
-
-	while (ItemsToRestoreExpansionState.Num() != 0)
-	{
-		UNiagaraStackEntry* EntryToProcess = ItemsToRestoreExpansionState[0].GetEntry();
-		if (EntryToProcess->GetIsSearchResult() && !EntryToProcess->IsA<UNiagaraStackRoot>())
-		{
-			for (UNiagaraStackEntry* EntryToExpand : ItemsToRestoreExpansionState[0].EntryPath)
-			{
-				if (!EntryToExpand->IsA<UNiagaraStackRoot>())
-				{
-					UNiagaraStackEditorData& StackEditorData = EntryToExpand->GetStackEditorData();
-					EntryToExpand->SetIsExpanded(StackEditorData.GetStackEntryWasExpandedPreSearch(EntryToExpand->GetStackEditorDataKey(), false));
-				}
-			}
-		}
-		ItemsToRestoreExpansionState.RemoveAtSwap(0);
-	}
+	bRestartSearch = true;
 }
 
 UNiagaraStackEntry* UNiagaraStackViewModel::GetRootEntry()
@@ -538,7 +520,7 @@ void UNiagaraStackViewModel::SetShowAllAdvanced(bool bInShowAllAdvanced)
 		}
 	}
 
-	OnSearchTextChanged(CurrentSearchText);
+	InvalidateSearchResults();
 	StructureChangedDelegate.Broadcast();
 }
 
@@ -565,7 +547,7 @@ void UNiagaraStackViewModel::SetShowOutputs(bool bInShowOutputs)
 	}
 	
 	// Showing outputs changes indenting so a full refresh is needed.
-	OnSearchTextChanged(CurrentSearchText);
+	InvalidateSearchResults();
 	RootEntry->RefreshChildren();
 }
 
@@ -592,7 +574,7 @@ void UNiagaraStackViewModel::SetShowLinkedInputs(bool bInShowLinkedInputs)
 	}
 
 	// Showing linked inputs changes indenting so a full refresh is needed.
-	OnSearchTextChanged(CurrentSearchText);
+	InvalidateSearchResults();
 	RootEntry->RefreshChildren();
 }
 
@@ -618,7 +600,7 @@ void UNiagaraStackViewModel::SetShowOnlyIssues(bool bInShowOnlyIssues)
 		}
 	}
 
-	OnSearchTextChanged(CurrentSearchText);
+	InvalidateSearchResults();
 	RootEntry->RefreshChildren();
 }
 
@@ -654,7 +636,7 @@ void UNiagaraStackViewModel::EntryStructureChanged()
 	}
 	RefreshHasIssues();
 	StructureChangedDelegate.Broadcast();
-	OnSearchTextChanged(CurrentSearchText);
+	InvalidateSearchResults();
 }
 
 void UNiagaraStackViewModel::EntryDataObjectModified(UObject* ChangedObject)
@@ -663,7 +645,7 @@ void UNiagaraStackViewModel::EntryDataObjectModified(UObject* ChangedObject)
 	{
 		SystemViewModel.Pin()->NotifyDataObjectChanged(ChangedObject);
 	}
-	OnSearchTextChanged(CurrentSearchText);
+	InvalidateSearchResults();
 	DataObjectChangedDelegate.Broadcast(ChangedObject);
 }
 
diff --git a/Engine/Plugins/FX/Niagara/Source/NiagaraEditor/Public/ViewModels/Stack/NiagaraStackViewModel.h b/Engine/Plugins/FX/Niagara/Source/NiagaraEditor/Public/ViewModels/Stack/NiagaraStackViewModel.h
--- a/Engine/Plugins/FX/Niagara/Source/NiagaraEditor/Public/ViewModels/Stack/NiagaraStackViewModel.h
+++ b/Engine/Plugins/FX/Niagara/Source/NiagaraEditor/Public/ViewModels/Stack/NiagaraStackViewModel.h
@@ -168,7 +168,8 @@ private:
 		TArray<FSearchWorkItem>& TraversedArray);
 	bool ItemMatchesSearchCriteria(UNiagaraStackEntry::FStackSearchItem SearchItem);
 	void GeneratePathForEntry(UNiagaraStackEntry* Root, UNiagaraStackEntry* Entry, TArray<UNiagaraStackEntry*> CurrentPath, TArray<UNiagaraStackEntry*>& EntryPath) const;
-	void RestoreStackEntryExpansionPreSearch();
+
+	void InvalidateSearchResults();
 
 private:
 	TWeakPtr<FNiagaraEmitterHandleViewModel> EmitterHandleViewModel;
@@ -189,7 +190,6 @@ private:
 	int CurrentFocusedSearchMatchIndex;
 	FOnSearchCompleted SearchCompletedDelegate;
 	TArray<FSearchWorkItem> ItemsToSearch;
-	TArray<FSearchWorkItem> ItemsToRestoreExpansionState;
 	TArray<FSearchResult> CurrentSearchResults;
 	static const double MaxSearchTime;
 	bool bRestartSearch;
diff --git a/Engine/Plugins/FX/Niagara/Source/NiagaraEditorWidgets/Private/SNiagaraStack.cpp b/Engine/Plugins/FX/Niagara/Source/NiagaraEditorWidgets/Private/SNiagaraStack.cpp
--- a/Engine/Plugins/FX/Niagara/Source/NiagaraEditorWidgets/Private/SNiagaraStack.cpp
+++ b/Engine/Plugins/FX/Niagara/Source/NiagaraEditorWidgets/Private/SNiagaraStack.cpp
@@ -263,7 +263,6 @@ void SNiagaraStack::Construct(const FArguments& InArgs, UNiagaraStackViewModel*
 	StackViewModel->OnSearchCompleted().AddSP(this, &SNiagaraStack::OnStackSearchComplete); 
 	NameColumnWidth = .3f;
 	ContentColumnWidth = .7f;
-	bNeedsJumpToNextOccurence = false;
 	StackCommandContext = MakeShared<FNiagaraStackCommandContext>();
 
 	ChildSlot
@@ -414,7 +413,6 @@ const FSlateBrush* SNiagaraStack::GetViewOptionsIconBrush() const
 
 void SNiagaraStack::OnSearchTextChanged(const FText& SearchText)
 {
-	bNeedsJumpToNextOccurence = true;
 	StackViewModel->OnSearchTextChanged(SearchText);
 }
 
@@ -630,11 +628,7 @@ FReply SNiagaraStack::OnRowAcceptDrop(const FDragDropEvent& InDragDropEvent, EIt
 void SNiagaraStack::OnStackSearchComplete()
 {
 	ExpandSearchResults();
-	if (bNeedsJumpToNextOccurence)
-	{
-		ScrollToNextMatch();
-		bNeedsJumpToNextOccurence = false;
-	}
+	ScrollToNextMatch();
 }
 
 void SNiagaraStack::ExpandSearchResults()
@@ -654,10 +648,9 @@ void SNiagaraStack::ExpandSearchResults()
 
 void SNiagaraStack::OnSearchBoxTextCommitted(const FText& NewText, ETextCommit::Type CommitInfo)
 {
-	if (bNeedsJumpToNextOccurence || CommitInfo == ETextCommit::OnEnter) // hasn't been autojumped yet or we hit enter
+	if (CommitInfo == ETextCommit::OnEnter) // hasn't been autojumped yet or we hit enter
 	{
 		AddSearchScrollOffset(+1);
-		bNeedsJumpToNextOccurence = false;
 	}
 }
 
diff --git a/Engine/Plugins/FX/Niagara/Source/NiagaraEditorWidgets/Private/SNiagaraStack.h b/Engine/Plugins/FX/Niagara/Source/NiagaraEditorWidgets/Private/SNiagaraStack.h
--- a/Engine/Plugins/FX/Niagara/Source/NiagaraEditorWidgets/Private/SNiagaraStack.h
+++ b/Engine/Plugins/FX/Niagara/Source/NiagaraEditorWidgets/Private/SNiagaraStack.h
@@ -130,7 +130,6 @@ private:
 
 	// ~ search stuff
 	TSharedPtr<SSearchBox> SearchBox;
-	bool bNeedsJumpToNextOccurence;
 
 	TSharedPtr<FNiagaraStackCommandContext> StackCommandContext;
 };
diff --git a/Engine/Plugins/FX/Niagara/Source/NiagaraEditorWidgets/Private/Stack/SNiagaraStackTableRow.cpp b/Engine/Plugins/FX/Niagara/Source/NiagaraEditorWidgets/Private/Stack/SNiagaraStackTableRow.cpp
--- a/Engine/Plugins/FX/Niagara/Source/NiagaraEditorWidgets/Private/Stack/SNiagaraStackTableRow.cpp
+++ b/Engine/Plugins/FX/Niagara/Source/NiagaraEditorWidgets/Private/Stack/SNiagaraStackTableRow.cpp
@@ -359,17 +359,26 @@ const FSlateBrush* SNiagaraStackTableRow::GetBorder() const
 	return FEditorStyle::GetBrush("NoBrush");
 }
 
-void SNiagaraStackTableRow::CollapseChildren()
+void SetExpansionStateRecursive(UNiagaraStackEntry* StackEntry, bool bIsExpanded)
 {
+	if (StackEntry->GetCanExpand())
+	{
+		StackEntry->SetIsExpanded(bIsExpanded);
+	}
+
 	TArray<UNiagaraStackEntry*> Children;
 	StackEntry->GetUnfilteredChildren(Children);
 	for (UNiagaraStackEntry* Child : Children)
 	{
-		if (Child->GetCanExpand())
-		{
-			Child->SetIsExpanded(false);
-		}
+		SetExpansionStateRecursive(Child, bIsExpanded);
 	}
+}
+
+void SNiagaraStackTableRow::CollapseChildren()
+{
+	bool bIsExpanded = false;
+	SetExpansionStateRecursive(StackEntry, bIsExpanded);
+
 	// Calling SetIsExpanded doesn't broadcast structure change automatically due to the expense of synchronizing
 	// expanded state with the tree to prevent items being expanded on tick, so we call this manually here.
 	StackEntry->OnStructureChanged().Broadcast();
@@ -377,15 +386,9 @@ void SNiagaraStackTableRow::CollapseChildren()
 
 void SNiagaraStackTableRow::ExpandChildren()
 {
-	TArray<UNiagaraStackEntry*> Children;
-	StackEntry->GetUnfilteredChildren(Children);
-	for (UNiagaraStackEntry* Child : Children)
-	{
-		if (Child->GetCanExpand())
-		{
-			Child->SetIsExpanded(true);
-		}
-	}
+	bool bIsExpanded = true;
+	SetExpansionStateRecursive(StackEntry, bIsExpanded);
+
 	// Calling SetIsExpanded doesn't broadcast structure change automatically due to the expense of synchronizing
 	// expanded state with the tree to prevent items being expanded on tick, so we call this manually here.
 	StackEntry->OnStructureChanged().Broadcast();
