Niagara - Fix an issue where the add button on assignment nodes was showing parameters which could not be set.

diff --git a/Engine/Plugins/FX/Niagara/Source/NiagaraEditor/Private/NiagaraNodeAssignment.cpp b/Engine/Plugins/FX/Niagara/Source/NiagaraEditor/Private/NiagaraNodeAssignment.cpp
--- a/Engine/Plugins/FX/Niagara/Source/NiagaraEditor/Private/NiagaraNodeAssignment.cpp
+++ b/Engine/Plugins/FX/Niagara/Source/NiagaraEditor/Private/NiagaraNodeAssignment.cpp
@@ -5,6 +5,8 @@
 #include "NiagaraGraph.h"
 #include "NiagaraScriptSource.h"
 #include "NiagaraScript.h"
+#include "NiagaraSystem.h"
+#include "NiagaraSystemEditorData.h"
 #include "NiagaraNodeInput.h"
 #include "NiagaraNodeOutput.h"
 #include "EdGraphSchema_Niagara.h"
@@ -257,11 +259,36 @@ void UNiagaraNodeAssignment::MergeUp()
 
 void UNiagaraNodeAssignment::CollectAddExistingActions(ENiagaraScriptUsage InUsage, UNiagaraNodeOutput* InGraphOutputNode, TArray<TSharedPtr<FNiagaraMenuAction>>& OutAddExistingActions)
 {
+	UNiagaraSystem* OwningSystem = GetTypedOuter<UNiagaraSystem>();
+	if (OwningSystem == nullptr)
+	{
+		return;
+	}
+
+	UNiagaraSystemEditorData* OwningSystemEditorData = Cast<UNiagaraSystemEditorData>(OwningSystem->GetEditorData());
+	if (OwningSystemEditorData == nullptr)
+	{
+		return;
+	}
+
+	bool bOwningSystemIsPlaceholder = OwningSystemEditorData->GetOwningSystemIsPlaceholder();
+
 	TArray<FNiagaraVariable> AvailableParameters;
 	FNiagaraStackGraphUtilities::GetAvailableParametersForScript(*InGraphOutputNode, AvailableParameters);
 
+	TArray<FName> AvailableWriteNamespaces;
+	FNiagaraStackGraphUtilities::GetNamespacesForNewWriteParameters(
+		bOwningSystemIsPlaceholder ? FNiagaraStackGraphUtilities::EStackEditContext::Emitter : FNiagaraStackGraphUtilities::EStackEditContext::System,
+		InUsage, AvailableWriteNamespaces);
+
 	for (const FNiagaraVariable& AvailableParameter : AvailableParameters)
 	{
+		FNiagaraParameterHandle AvailableParameterHandle(AvailableParameter.GetName());
+		if (AvailableWriteNamespaces.Contains(AvailableParameterHandle.GetHandleParts()[0]) == false)
+		{
+			continue;
+		}
+
 		const FText NameText = FText::FromName(AvailableParameter.GetName());
 		FText VarDesc = FNiagaraConstants::GetAttributeDescription(AvailableParameter);
 		FString VarDefaultValue = FNiagaraConstants::GetAttributeDefaultValue(AvailableParameter);
diff --git a/Engine/Plugins/FX/Niagara/Source/NiagaraEditor/Private/ViewModels/Stack/NiagaraStackFunctionInput.cpp b/Engine/Plugins/FX/Niagara/Source/NiagaraEditor/Private/ViewModels/Stack/NiagaraStackFunctionInput.cpp
--- a/Engine/Plugins/FX/Niagara/Source/NiagaraEditor/Private/ViewModels/Stack/NiagaraStackFunctionInput.cpp
+++ b/Engine/Plugins/FX/Niagara/Source/NiagaraEditor/Private/ViewModels/Stack/NiagaraStackFunctionInput.cpp
@@ -1611,32 +1611,9 @@ void UNiagaraStackFunctionInput::GetNamespacesForNewReadParameters(TArray<FName>
 	UNiagaraNodeOutput* OutputNode = FNiagaraStackGraphUtilities::GetEmitterOutputNodeForStackNode(*OwningFunctionCallNode);
 	bool bIsEditingSystem = GetSystemViewModel()->GetEditMode() == ENiagaraSystemViewModelEditMode::SystemAsset;
 
-	switch (OutputNode->GetUsage())
-	{
-	case ENiagaraScriptUsage::ParticleSpawnScript:
-	case ENiagaraScriptUsage::ParticleSpawnScriptInterpolated:
-	case ENiagaraScriptUsage::ParticleUpdateScript:
-	case ENiagaraScriptUsage::ParticleEventScript:
-	case ENiagaraScriptUsage::ParticleSimulationStageScript:
-	{
-		OutNamespacesForNewParameters.Add(FNiagaraConstants::ParticleAttributeNamespace);
-		OutNamespacesForNewParameters.Add(FNiagaraConstants::EmitterNamespace);
-		break;
-	}
-	case ENiagaraScriptUsage::EmitterSpawnScript:
-	case ENiagaraScriptUsage::EmitterUpdateScript:
-	{
-		OutNamespacesForNewParameters.Add(FNiagaraConstants::EmitterNamespace);
-		break;
-	}
-	}
-
-	if (bIsEditingSystem)
-	{
-		OutNamespacesForNewParameters.Add(FNiagaraConstants::UserNamespace);
-		OutNamespacesForNewParameters.Add(FNiagaraConstants::SystemNamespace);
-	}
-	OutNamespacesForNewParameters.Add(FNiagaraConstants::TransientNamespace);
+	FNiagaraStackGraphUtilities::GetNamespacesForNewReadParameters(
+		bIsEditingSystem ? FNiagaraStackGraphUtilities::EStackEditContext::System : FNiagaraStackGraphUtilities::EStackEditContext::Emitter,
+		OutputNode->GetUsage(), OutNamespacesForNewParameters);
 }
 
 void UNiagaraStackFunctionInput::GetNamespacesForNewWriteParameters(TArray<FName>& OutNamespacesForNewParameters) const
@@ -1644,33 +1621,9 @@ void UNiagaraStackFunctionInput::GetNamespacesForNewWriteParameters(TArray<FName
 	UNiagaraNodeOutput* OutputNode = FNiagaraStackGraphUtilities::GetEmitterOutputNodeForStackNode(*OwningFunctionCallNode);
 	bool bIsEditingSystem = GetSystemViewModel()->GetEditMode() == ENiagaraSystemViewModelEditMode::SystemAsset;
 
-	switch (OutputNode->GetUsage())
-	{
-	case ENiagaraScriptUsage::ParticleSpawnScript:
-	case ENiagaraScriptUsage::ParticleSpawnScriptInterpolated:
-	case ENiagaraScriptUsage::ParticleUpdateScript:
-	case ENiagaraScriptUsage::ParticleEventScript:
-	case ENiagaraScriptUsage::ParticleSimulationStageScript:
-	{
-		OutNamespacesForNewParameters.Add(FNiagaraConstants::ParticleAttributeNamespace);
-		break;
-	}
-	case ENiagaraScriptUsage::EmitterSpawnScript:
-	case ENiagaraScriptUsage::EmitterUpdateScript:
-	{
-		OutNamespacesForNewParameters.Add(FNiagaraConstants::EmitterNamespace);
-		break;
-	}
-	case ENiagaraScriptUsage::SystemSpawnScript:
-	case ENiagaraScriptUsage::SystemUpdateScript:
-		if (bIsEditingSystem)
-		{
-			OutNamespacesForNewParameters.Add(FNiagaraConstants::SystemNamespace);
-		}
-		break;
-	}
-
-	OutNamespacesForNewParameters.Add(FNiagaraConstants::TransientNamespace);
+	FNiagaraStackGraphUtilities::GetNamespacesForNewWriteParameters(
+		bIsEditingSystem ? FNiagaraStackGraphUtilities::EStackEditContext::System : FNiagaraStackGraphUtilities::EStackEditContext::Emitter,
+		OutputNode->GetUsage(), OutNamespacesForNewParameters);
 }
 
 UNiagaraStackFunctionInput::FOnValueChanged& UNiagaraStackFunctionInput::OnValueChanged()
diff --git a/Engine/Plugins/FX/Niagara/Source/NiagaraEditor/Private/ViewModels/Stack/NiagaraStackGraphUtilities.cpp b/Engine/Plugins/FX/Niagara/Source/NiagaraEditor/Private/ViewModels/Stack/NiagaraStackGraphUtilities.cpp
--- a/Engine/Plugins/FX/Niagara/Source/NiagaraEditor/Private/ViewModels/Stack/NiagaraStackGraphUtilities.cpp
+++ b/Engine/Plugins/FX/Niagara/Source/NiagaraEditor/Private/ViewModels/Stack/NiagaraStackGraphUtilities.cpp
@@ -2614,4 +2614,65 @@ UNiagaraStackEntry::FStackIssue FNiagaraStackGraphUtilities::MessageManagerMessa
 		FixLinks);
 }
 
+void FNiagaraStackGraphUtilities::GetNamespacesForNewReadParameters(EStackEditContext EditContext, ENiagaraScriptUsage Usage, TArray<FName>& OutNamespacesForNewParameters)
+{
+	switch (Usage)
+	{
+	case ENiagaraScriptUsage::ParticleSpawnScript:
+	case ENiagaraScriptUsage::ParticleSpawnScriptInterpolated:
+	case ENiagaraScriptUsage::ParticleUpdateScript:
+	case ENiagaraScriptUsage::ParticleEventScript:
+	case ENiagaraScriptUsage::ParticleSimulationStageScript:
+	{
+		OutNamespacesForNewParameters.Add(FNiagaraConstants::ParticleAttributeNamespace);
+		OutNamespacesForNewParameters.Add(FNiagaraConstants::EmitterNamespace);
+		break;
+	}
+	case ENiagaraScriptUsage::EmitterSpawnScript:
+	case ENiagaraScriptUsage::EmitterUpdateScript:
+	{
+		OutNamespacesForNewParameters.Add(FNiagaraConstants::EmitterNamespace);
+		break;
+	}
+	}
+
+	if (EditContext == EStackEditContext::System)
+	{
+		OutNamespacesForNewParameters.Add(FNiagaraConstants::UserNamespace);
+		OutNamespacesForNewParameters.Add(FNiagaraConstants::SystemNamespace);
+	}
+	OutNamespacesForNewParameters.Add(FNiagaraConstants::TransientNamespace);
+}
+
+void FNiagaraStackGraphUtilities::GetNamespacesForNewWriteParameters(EStackEditContext EditContext, ENiagaraScriptUsage Usage, TArray<FName>& OutNamespacesForNewParameters)
+{
+	switch (Usage)
+	{
+	case ENiagaraScriptUsage::ParticleSpawnScript:
+	case ENiagaraScriptUsage::ParticleSpawnScriptInterpolated:
+	case ENiagaraScriptUsage::ParticleUpdateScript:
+	case ENiagaraScriptUsage::ParticleEventScript:
+	case ENiagaraScriptUsage::ParticleSimulationStageScript:
+	{
+		OutNamespacesForNewParameters.Add(FNiagaraConstants::ParticleAttributeNamespace);
+		break;
+	}
+	case ENiagaraScriptUsage::EmitterSpawnScript:
+	case ENiagaraScriptUsage::EmitterUpdateScript:
+	{
+		OutNamespacesForNewParameters.Add(FNiagaraConstants::EmitterNamespace);
+		break;
+	}
+	case ENiagaraScriptUsage::SystemSpawnScript:
+	case ENiagaraScriptUsage::SystemUpdateScript:
+		if (EditContext == EStackEditContext::System)
+		{
+			OutNamespacesForNewParameters.Add(FNiagaraConstants::SystemNamespace);
+		}
+		break;
+	}
+
+	OutNamespacesForNewParameters.Add(FNiagaraConstants::TransientNamespace);
+}
+
 #undef LOCTEXT_NAMESPACE
diff --git a/Engine/Plugins/FX/Niagara/Source/NiagaraEditor/Public/ViewModels/Stack/NiagaraStackGraphUtilities.h b/Engine/Plugins/FX/Niagara/Source/NiagaraEditor/Public/ViewModels/Stack/NiagaraStackGraphUtilities.h
--- a/Engine/Plugins/FX/Niagara/Source/NiagaraEditor/Public/ViewModels/Stack/NiagaraStackGraphUtilities.h
+++ b/Engine/Plugins/FX/Niagara/Source/NiagaraEditor/Public/ViewModels/Stack/NiagaraStackGraphUtilities.h
@@ -195,4 +195,16 @@ namespace FNiagaraStackGraphUtilities
 	void GatherRenamedStackFunctionInputAndOutputVariableNames(UNiagaraEmitter* Emitter, UNiagaraNodeFunctionCall& FunctionCallNode, const FString& OldFunctionName, const FString& NewFunctionName, TMap<FName, FName>& OutOldToNewNameMap);
 
 	UNiagaraStackEntry::FStackIssue MessageManagerMessageToStackIssue(TSharedRef<const INiagaraMessage> InMessage, FString InStackEditorDataKey);
+
+	enum class EStackEditContext
+	{
+		System,
+		Emitter
+	};
+
+	/** Gets the valid namespaces which new parameters for this usage can be read from. */
+	void GetNamespacesForNewReadParameters(EStackEditContext EditContext, ENiagaraScriptUsage Usage, TArray<FName>& OutNamespacesForNewParameters);
+
+	/** Gets the valid namespaces which new parameters for this usage can write to. */
+	void GetNamespacesForNewWriteParameters(EStackEditContext EditContext, ENiagaraScriptUsage Usage, TArray<FName>& OutNamespacesForNewParameters);
 }
