Niagara - Fix a crash when trying to move a parameter to a namespace that already has a parameter with that name.

diff --git a/Engine/Plugins/FX/Niagara/Source/NiagaraEditor/Private/NiagaraEditorUtilities.cpp b/Engine/Plugins/FX/Niagara/Source/NiagaraEditor/Private/NiagaraEditorUtilities.cpp
--- a/Engine/Plugins/FX/Niagara/Source/NiagaraEditor/Private/NiagaraEditorUtilities.cpp
+++ b/Engine/Plugins/FX/Niagara/Source/NiagaraEditor/Private/NiagaraEditorUtilities.cpp
@@ -2690,9 +2690,13 @@ void FNiagaraParameterUtilities::GetChangeNamespaceMenuData(FName InParameterNam
 	NamespaceMetadata.Sort([](const FNiagaraNamespaceMetadata& MetadataA, const FNiagaraNamespaceMetadata& MetadataB)
 		{ return MetadataA.SortId < MetadataB.SortId; });
 
+	FNiagaraParameterHandle ParameterHandle(InParameterName);
+	FNiagaraNamespaceMetadata CurrentMetadata = GetDefault<UNiagaraEditorSettings>()->GetMetaDataForNamespaces(ParameterHandle.GetHandleParts());
 	for (const FNiagaraNamespaceMetadata& Metadata : NamespaceMetadata)
 	{
-		if (Metadata.IsValid() == false || Metadata.Options.Contains(ENiagaraNamespaceMetadataOptions::PreventEditingNamespace) ||
+		if (Metadata.IsValid() == false ||
+			Metadata == CurrentMetadata ||
+			Metadata.Options.Contains(ENiagaraNamespaceMetadataOptions::PreventEditingNamespace) ||
 			(InParameterContext == EParameterContext::Script && Metadata.Options.Contains(ENiagaraNamespaceMetadataOptions::HideInScript)) ||
 			(InParameterContext == EParameterContext::System && Metadata.Options.Contains(ENiagaraNamespaceMetadataOptions::HideInSystem)))
 		{
diff --git a/Engine/Plugins/FX/Niagara/Source/NiagaraEditor/Private/Widgets/SNiagaraParameterMapView.cpp b/Engine/Plugins/FX/Niagara/Source/NiagaraEditor/Private/Widgets/SNiagaraParameterMapView.cpp
--- a/Engine/Plugins/FX/Niagara/Source/NiagaraEditor/Private/Widgets/SNiagaraParameterMapView.cpp
+++ b/Engine/Plugins/FX/Niagara/Source/NiagaraEditor/Private/Widgets/SNiagaraParameterMapView.cpp
@@ -507,6 +507,8 @@ TSharedRef<SWidget> SNiagaraParameterMapView::OnCreateWidgetForAction(struct FCr
 
 void SNiagaraParameterMapView::CollectAllActions(FGraphActionListBuilderBase& OutAllActions)
 {
+	LastCollectedParameters.Empty();
+
 	if (Graphs.Num() == 0)
 	{
 		return;
@@ -543,6 +545,7 @@ void SNiagaraParameterMapView::CollectAllActions(FGraphActionListBuilderBase& Ou
 		const FText Tooltip = FText::Format(TooltipFormat, FText::FromName(Parameter.GetName()), Parameter.GetType().GetNameText());
 		TSharedPtr<FNiagaraParameterAction> ParameterAction(new FNiagaraParameterAction(Parameter, ParameterEntry.Value, FText::GetEmpty(), Name, Tooltip, 0, FText(), Section));
 		OutAllActions.AddAction(ParameterAction);
+		LastCollectedParameters.Add(Parameter);
 	}
 }
 
@@ -1223,6 +1226,12 @@ bool SNiagaraParameterMapView::GetSingleParameterActionForSelection(
 	return true;
 }
 
+bool SNiagaraParameterMapView::ParameterExistsByName(FName ParameterName) const
+{
+	auto MatchVariableByName = [ParameterName](const FNiagaraVariable& Variable) { return Variable.GetName() == ParameterName; };
+	return LastCollectedParameters.ContainsByPredicate(MatchVariableByName);
+}
+
 void SNiagaraParameterMapView::GetChangeNamespaceSubMenu(FMenuBuilder& MenuBuilder)
 {
 	TSharedPtr<FNiagaraParameterAction> ParameterAction;
@@ -1235,12 +1244,24 @@ void SNiagaraParameterMapView::GetChangeNamespaceSubMenu(FMenuBuilder& MenuBuild
 		for (const FNiagaraParameterUtilities::FChangeNamespaceMenuData& MenuDataItem : MenuData)
 		{
 			bool bCanChange = MenuDataItem.bCanChange;
+			FText CanChangeToolTip = MenuDataItem.CanChangeToolTip;
+			if (bCanChange)
+			{
+				// Check for an existing duplicate by name.
+				FName NewName = FNiagaraParameterUtilities::ChangeNamespace(ParameterAction->Parameter.GetName(), MenuDataItem.Metadata);
+				if (ParameterExistsByName(NewName))
+				{
+					bCanChange = false;
+					CanChangeToolTip = LOCTEXT("CantMoveAlreadyExits", "Can not move to this namespace because a parameter with this name already exists.");
+				}
+			}
+
 			FUIAction Action = FUIAction(
 				FExecuteAction::CreateSP(this, &SNiagaraParameterMapView::OnChangeNamespace, MenuDataItem.Metadata),
 				FCanExecuteAction::CreateLambda([bCanChange]() { return bCanChange; }));
 
-			TSharedRef<SWidget> MenuItemWidget = FNiagaraParameterUtilities::CreateNamespaceMenuItemWidget(MenuDataItem.NamespaceParameterName, MenuDataItem.CanChangeToolTip);
-			MenuBuilder.AddMenuEntry(Action, MenuItemWidget, NAME_None, MenuDataItem.CanChangeToolTip);
+			TSharedRef<SWidget> MenuItemWidget = FNiagaraParameterUtilities::CreateNamespaceMenuItemWidget(MenuDataItem.NamespaceParameterName, CanChangeToolTip);
+			MenuBuilder.AddMenuEntry(Action, MenuItemWidget, NAME_None, CanChangeToolTip);
 		}
 	}
 }
diff --git a/Engine/Plugins/FX/Niagara/Source/NiagaraEditor/Public/NiagaraEditorSettings.h b/Engine/Plugins/FX/Niagara/Source/NiagaraEditor/Public/NiagaraEditorSettings.h
--- a/Engine/Plugins/FX/Niagara/Source/NiagaraEditor/Public/NiagaraEditorSettings.h
+++ b/Engine/Plugins/FX/Niagara/Source/NiagaraEditor/Public/NiagaraEditorSettings.h
@@ -49,6 +49,19 @@ struct FNiagaraNamespaceMetadata
 
 	FNiagaraNamespaceMetadata(TArray<FName> InNamespaces);
 
+	bool operator==(const FNiagaraNamespaceMetadata& Other) const
+	{
+		return
+			Namespaces == Other.Namespaces &&
+			DisplayName.IdenticalTo(Other.DisplayName) &&
+			DisplayNameLong.IdenticalTo(Other.DisplayNameLong) &&
+			Description.IdenticalTo(Other.Description) &&
+			BackgroundColor == Other.BackgroundColor &&
+			ForegroundStyle == Other.ForegroundStyle &&
+			SortId == Other.SortId &&
+			Options == Other.Options;
+	}
+
 	UPROPERTY()
 	TArray<FName> Namespaces;
 
diff --git a/Engine/Plugins/FX/Niagara/Source/NiagaraEditor/Public/Widgets/SNiagaraParameterMapView.h b/Engine/Plugins/FX/Niagara/Source/NiagaraEditor/Public/Widgets/SNiagaraParameterMapView.h
--- a/Engine/Plugins/FX/Niagara/Source/NiagaraEditor/Public/Widgets/SNiagaraParameterMapView.h
+++ b/Engine/Plugins/FX/Niagara/Source/NiagaraEditor/Public/Widgets/SNiagaraParameterMapView.h
@@ -155,6 +155,8 @@ private:
 		TSharedPtr<FNiagaraParameterAction>& OutParameterAction,
 		FText& OutErrorMessage) const;
 
+	bool ParameterExistsByName(FName ParameterName) const;
+
 	void GetChangeNamespaceSubMenu(FMenuBuilder& MenuBuilder);
 	void OnChangeNamespace(FNiagaraNamespaceMetadata Metadata);
 
@@ -219,6 +221,8 @@ private:
 
 	TArray<int32> HiddenSectionIDs;
 
+	TArray<FNiagaraVariable> LastCollectedParameters;
+
 	bool bNeedsRefresh;
 	bool bIsAddingParameter;
 };
