Fixed: wrong execution index in reset-only stages which use an iteration interface

The generated HLSL called SetupExecIndexAndSpawnInfoForGPU for DI-based stages too, which doesn't make sense since their execution index has nothing to do with the spawn counts.

#rnx

diff --git a/Engine/Plugins/FX/Niagara/Source/NiagaraCore/Private/NiagaraCustomVersion.cpp b/Engine/Plugins/FX/Niagara/Source/NiagaraCore/Private/NiagaraCustomVersion.cpp
--- a/Engine/Plugins/FX/Niagara/Source/NiagaraCore/Private/NiagaraCustomVersion.cpp
+++ b/Engine/Plugins/FX/Niagara/Source/NiagaraCore/Private/NiagaraCustomVersion.cpp
@@ -8,4 +8,4 @@ const FGuid FNiagaraCustomVersion::GUID(0xFCF57AFA, 0x50764283, 0xB9A9E658, 0xFF
 // Register the custom version with core
 FCustomVersionRegistration GRegisterNiagaraCustomVersion(FNiagaraCustomVersion::GUID, FNiagaraCustomVersion::LatestVersion, TEXT("NiagaraVer"));
 
-const FGuid FNiagaraCustomVersion::LatestScriptCompileVersion(0x1F365103, 0x4B5DC149, 0x8308F39D, 0x698AC03E);
\ No newline at end of file
+const FGuid FNiagaraCustomVersion::LatestScriptCompileVersion(0x9315AB96, 0xF1704D25, 0xAAAA6038, 0xF25C36DD);
diff --git a/Engine/Plugins/FX/Niagara/Source/NiagaraEditor/Private/NiagaraHlslTranslator.cpp b/Engine/Plugins/FX/Niagara/Source/NiagaraEditor/Private/NiagaraHlslTranslator.cpp
--- a/Engine/Plugins/FX/Niagara/Source/NiagaraEditor/Private/NiagaraHlslTranslator.cpp
+++ b/Engine/Plugins/FX/Niagara/Source/NiagaraEditor/Private/NiagaraHlslTranslator.cpp
@@ -2501,6 +2501,7 @@ void FHlslNiagaraTranslator::DefineMainGPUFunctions(
 			"	}\n"
 			"	*/\n"
 			"	\n"
+			"   // The CPU code will set UpdateStartInstance to 0 and ReadInstanceCountOffset to -1 for stages.\n"
 			"	const uint InstanceID = UpdateStartInstance + DispatchThreadId.x; \n"
 			"	if (ReadInstanceCountOffset == 0xFFFFFFFF)\n"
 			"	{\n"
@@ -2510,19 +2511,21 @@ void FHlslNiagaraTranslator::DefineMainGPUFunctions(
 			"	{\n"
 			"		GSpawnStartInstance = RWInstanceCounts[ReadInstanceCountOffset];				// needed by ExecIndex()\n"
 			"	}\n"
-			"	const int MaxInstances = GSpawnStartInstance + SpawnedInstances; \n"
-			"		\n"
-			"	bool bRunUpdateLogic = InstanceID < GSpawnStartInstance && InstanceID < UpdateStartInstance + MaxInstances; \n"
-			"	bool bRunSpawnLogic = InstanceID >= GSpawnStartInstance && InstanceID < UpdateStartInstance + MaxInstances; \n"
-			"		\n"
+			"	bool bRunUpdateLogic, bRunSpawnLogic;\n"
 			"	#if USE_SIMULATION_STAGES\n"
 			"	if (IterationInterfaceCount > 0)\n"
 			"	{\n"
-			"		bRunUpdateLogic = InstanceID < IterationInterfaceCount && GSimStart != 1; \n"
-			"		bRunSpawnLogic = InstanceID < IterationInterfaceCount && GSimStart == 1; \n"
+			"		bRunUpdateLogic = InstanceID < IterationInterfaceCount && GSimStart != 1;\n"
+			"		bRunSpawnLogic = InstanceID < IterationInterfaceCount && GSimStart == 1;\n"
 			"	}\n"
+			"	else\n"
 			"	#endif // USE_SIMULATION_STAGES\n"
-			"		\n"
+			"	{\n"
+			"	    const int MaxInstances = GSpawnStartInstance + SpawnedInstances;\n"
+			"		bRunUpdateLogic = InstanceID < GSpawnStartInstance && InstanceID < UpdateStartInstance + MaxInstances;\n"
+			"		bRunSpawnLogic = InstanceID >= GSpawnStartInstance && InstanceID < UpdateStartInstance + MaxInstances;\n"
+			"	}\n"
+			"	\n"
 			"	const float RandomSeedInitialisation = NiagaraInternalNoise(InstanceID * 16384, 0 * 8196, (bRunUpdateLogic ? 4096 : 0) + EmitterTickCounter);	// initialise the random state seed\n"
 			"	\n"
 			"	FSimulationContext Context = (FSimulationContext)0; \n"
@@ -2544,7 +2547,17 @@ void FHlslNiagaraTranslator::DefineMainGPUFunctions(
 			"	else if (bRunSpawnLogic)\n"
 			"	{\n"
 			"		GCurrentPhase = GSpawnPhase; \n"
-			"		SetupExecIndexAndSpawnInfoForGPU(); \n"
+			"	#if USE_SIMULATION_STAGES\n"
+			"		// Only process the spawn info for particle-based stages. Stages with an iteration interface expect the exec index to simply be the thread index.\n"
+			"		if (IterationInterfaceCount > 0)\n"
+			"		{\n"
+			"			SetupExecIndexForGPU();\n"
+			"		}\n"
+			"		else\n"
+			"	#endif\n"
+			"		{\n"
+			"			SetupExecIndexAndSpawnInfoForGPU();\n"
+			"		}\n"
 			"		InitConstants(Context); \n"
 			"		InitSpawnVariables(Context); \n"
 			"		ReadDataSets(Context); \n"
