Integrated mesh particle instance culling and custom visibility and added a test for it

diff --git a/Engine/Plugins/FX/Niagara/Shaders/Private/NiagaraDrawIndirectArgsGen.usf b/Engine/Plugins/FX/Niagara/Shaders/Private/NiagaraDrawIndirectArgsGen.usf
--- a/Engine/Plugins/FX/Niagara/Shaders/Private/NiagaraDrawIndirectArgsGen.usf
+++ b/Engine/Plugins/FX/Niagara/Shaders/Private/NiagaraDrawIndirectArgsGen.usf
@@ -10,6 +10,7 @@
 // ArgGen = (InstanceCountBufferOffset, NumIndicesPerInstance, StartIndexLocation)
 // InstanceCountClear = (InstanceCountBufferOffset)
 Buffer<uint> TaskInfos;
+Buffer<uint> CulledInstanceCounts;
 
 RWBuffer<uint> RWInstanceCounts;
 RWBuffer<uint> RWDrawIndirectArgs;
@@ -23,8 +24,21 @@ void MainCS(uint TaskIndex : SV_DispatchThreadID)
 	{
 		const uint ArgOffset = NIAGARA_DRAW_INDIRECT_ARGS_SIZE * TaskIndex;
 		const uint InfoOffset = TaskIndex * NIAGARA_DRAW_INDIRECT_TASK_INFO_SIZE;
+		const bool bUseCulledCounts = TaskInfos[InfoOffset + 3] != 0;
+
+		uint InstanceCount = 0;
+		BRANCH
+		if (bUseCulledCounts)
+		{
+			InstanceCount = CulledInstanceCounts[TaskInfos[InfoOffset + 0]];
+		}
+		else
+		{
+			InstanceCount = RWInstanceCounts[TaskInfos[InfoOffset + 0]];
+		}
+
 		RWDrawIndirectArgs[ArgOffset + 0] = TaskInfos[InfoOffset + 1]; // NumIndicesPerInstance
-		RWDrawIndirectArgs[ArgOffset + 1] = RWInstanceCounts[TaskInfos[InfoOffset + 0]]; // InstanceCount
+		RWDrawIndirectArgs[ArgOffset + 1] = InstanceCount;
 		RWDrawIndirectArgs[ArgOffset + 2] = TaskInfos[InfoOffset + 2]; // StartIndexLocation
 		RWDrawIndirectArgs[ArgOffset + 3] = 0; // BaseVertexLocation
 		RWDrawIndirectArgs[ArgOffset + 4] = 0; // StartInstanceLocation
@@ -70,4 +84,4 @@ void ResetCountsCS(uint TaskIndex : SV_DispatchThreadID)
 		const uint InfoOffset = TaskCount.x * NIAGARA_DRAW_INDIRECT_TASK_INFO_SIZE + (TaskIndex);
 		RWInstanceCounts[TaskInfos[InfoOffset + 0]] = 0;
 	}
-}
+}
\ No newline at end of file
diff --git a/Engine/Plugins/FX/Niagara/Shaders/Private/NiagaraSortKeyGen.usf b/Engine/Plugins/FX/Niagara/Shaders/Private/NiagaraSortKeyGen.usf
--- a/Engine/Plugins/FX/Niagara/Shaders/Private/NiagaraSortKeyGen.usf
+++ b/Engine/Plugins/FX/Niagara/Shaders/Private/NiagaraSortKeyGen.usf
@@ -7,6 +7,7 @@
 /*------------------------------------------------------------------------------
 	Compile time parameters:
 		THREAD_COUNT - The number of threads to launch per workgroup.
+		SORT_NONE
 		SORT_VIEW_DEPTH
 		SORT_VIEW_DISTANCE
 		SORT_CUSTOM_ASCENDING
@@ -15,76 +16,222 @@
 
 #include "/Engine/Private/Common.ush"
 
+#define NIAGARA_KEY_GEN_MAX_CULL_PLANES 6 // must match the code
 Buffer<float> NiagaraParticleDataFloat;
-uint NiagaraFloatDataStride;
+Buffer<float> NiagaraParticleDataHalf;
+Buffer<uint> NiagaraParticleDataInt;
+Buffer<uint> GPUParticleCountBuffer;
+
+uint FloatDataStride;
+uint HalfDataStride;
+uint IntDataStride;
+uint ParticleCount;
+uint GPUParticleCountOffset;
+uint CulledGPUParticleCountOffset;
+uint EmitterKey;
+uint OutputOffset;
+float3 CameraPosition;
+float3 CameraDirection;
+uint SortMode;
+uint SortAttributeOffset;
+#if SORT_MAX_PRECISION
+	uint SortKeyMask;
+	uint SortKeyShift;
+	uint SortKeySignBit;
+#endif
+uint CullPositionAttributeOffset;
+uint CullOrientationAttributeOffset;
+uint CullScaleAttributeOffset;
+int NumCullPlanes;
+float4 CullPlanes[NIAGARA_KEY_GEN_MAX_CULL_PLANES];
+int RendererVisibility;
+uint RendererVisTagAttributeOffset;
+float2 CullDistanceRangeSquared;
+float4 LocalBoundingSphere;
 
-float GetFloat(int RegisterIdx, uint InstanceID)
+RWBuffer<uint> OutKeys;
+RWBuffer<uint> OutParticleIndices;
+RWBuffer<uint> OutCulledParticleCounts;
+
+#if ENABLE_CULLING && !USE_WAVE_INTRINSICS
+groupshared uint CullGroupCount;
+#endif
+
+float GetFloat(Buffer<float> Buff, uint DataStride, uint RegisterIdx, uint InstanceID)
 {
-	RegisterIdx &= (~(1 << 31));
-	return NiagaraParticleDataFloat[(RegisterIdx * NiagaraFloatDataStride + InstanceID)];
+	return Buff[RegisterIdx * DataStride + InstanceID];
 }
 
-float3 GetVec3(int RegisterIndex, uint InstanceID)
+float GetFloat(uint RegisterIdx, uint InstanceID)
 {
-	return float3(GetFloat(RegisterIndex, InstanceID), GetFloat(RegisterIndex+1, InstanceID), GetFloat(RegisterIndex+2, InstanceID));
+	BRANCH
+	if (RegisterIdx & (1 << 31))
+	{
+		RegisterIdx &= ~(1 << 31);
+		return GetFloat(NiagaraParticleDataHalf, HalfDataStride, RegisterIdx, InstanceID);
+	}
+	else
+	{
+		return GetFloat(NiagaraParticleDataFloat, FloatDataStride, RegisterIdx, InstanceID);
+	}
 }
 
-Buffer<uint> GPUParticleCountBuffer;
+int GetInt(uint RegisterIdx, uint InstanceID)
+{
+	return NiagaraParticleDataInt[RegisterIdx * IntDataStride + InstanceID];
+}
 
-// (ParticleCount, GPUParticleCountOffset)
-uint2 ParticleCountParams;
+float2 GetVec2(uint RegisterIndex, uint InstanceID)
+{
+	return float2(GetFloat(RegisterIndex, InstanceID), GetFloat(RegisterIndex + 1, InstanceID));
+}
 
-// (EmitterKey, OutputOffset, SortMode, SortAttributeOffset)
-uint4 SortParams;
+float3 GetVec3(uint RegisterIndex, uint InstanceID)
+{
+	return float3(GetFloat(RegisterIndex, InstanceID), GetFloat(RegisterIndex + 1, InstanceID), GetFloat(RegisterIndex + 2, InstanceID));
+}
 
-float3 CameraPosition;
-float3 CameraDirection;
+float4 GetVec4(uint RegisterIndex, uint InstanceID)
+{
+	return float4(GetFloat(RegisterIndex, InstanceID), GetFloat(RegisterIndex + 1, InstanceID), GetFloat(RegisterIndex + 2, InstanceID), GetFloat(RegisterIndex + 3, InstanceID));
+}
 
-/** Output key buffer. */
-RWBuffer<uint> OutKeys;
-/** Output indices buffer. */
-RWBuffer<int> OutParticleIndices;
+float3 RotateVectorByQuat(in float3 Vector, in float4 Quat)
+{
+	float3 T = 2.0 * cross(Quat.xyz, Vector);
+	return Vector + Quat.w * T + cross(Quat.xyz, T);
+}
 
-#if SORT_MAX_PRECISION
-// (SortKeyMask, SortKeyShift, SortKeySignBit)
-uint4 SortKeyParams;
+void CullingPrologue(in uint GroupTID)
+{
+#if ENABLE_CULLING && !USE_WAVE_INTRINSICS
+	if (GroupTID == 0)
+	{
+		CullGroupCount = 0;
+	}
+
+	GroupMemoryBarrierWithGroupSync();
 #endif
+}
 
-[numthreads(THREAD_COUNT,1,1)]
-void GenerateParticleSortKeys(uint ParticleIndex : SV_DispatchThreadID)
+void CullingEpilogue(in bool bCulled, in uint GroupTID)
 {
-	const uint ParticleCount = ParticleCountParams.x;
-	const uint GPUParticleCountOffset = ParticleCountParams.y;
+#if ENABLE_CULLING
+	uint CountOffset = CulledGPUParticleCountOffset;
 
-	const uint GPUParticleCount = GPUParticleCountOffset != -1 ? GPUParticleCountBuffer[GPUParticleCountOffset] : ParticleCount;
+#if USE_WAVE_INTRINSICS
+	uint CullWaveCount = WaveActiveCountBits(!bCulled);
+	if (WaveOnce())
+	{
+		InterlockedAdd(OutCulledParticleCounts[CountOffset], CullWaveCount);
+	}
+#else
+	InterlockedAdd(CullGroupCount, bCulled ? 0 : 1);
+	GroupMemoryBarrierWithGroupSync();
 
-	const uint EmitterKey = SortParams.x;
-	const uint OutputOffset = SortParams.y;
-	const uint SortMode = SortParams.z;
-	const uint SortAttributeOffset = SortParams.w;
+	if (GroupTID == 0)
+	{
+		InterlockedAdd(OutCulledParticleCounts[CountOffset], CullGroupCount);
+	}
+#endif // USE_WAVE_INTRINSICS
+#endif // ENABLE_CULLING
+}
 
-#if SORT_MAX_PRECISION
-	const uint SortKeyMask = SortKeyParams.x; 
-	const uint SortKeyShift = SortKeyParams.y;
-	const uint SortKeySignBit = SortKeyParams.z;
-#else
-	const uint SortKeyMask = 0xffff; 
+bool CullParticle(in uint ParticleIndex)
+{
+#if ENABLE_CULLING
+	float3 Pos = 0;
+	if (CullPositionAttributeOffset != -1)
+	{
+		Pos = GetVec3(CullPositionAttributeOffset, ParticleIndex);
+	}
+
+	float4 Quat = { 0, 0, 0, 1 };
+	if (CullOrientationAttributeOffset != -1)
+	{
+		Quat = GetVec4(CullOrientationAttributeOffset, ParticleIndex);
+	}
+
+	float3 Scale = 1;
+	if (CullScaleAttributeOffset != -1)
+	{
+		Scale = GetVec3(CullScaleAttributeOffset, ParticleIndex);
+	}
+
+	// Transform the bounding sphere
+	float4 BSphere = LocalBoundingSphere;
+	BSphere.xyz = RotateVectorByQuat(BSphere.xyz * Scale, Quat) + Pos;
+	BSphere.w *= max(Scale.x, max(Scale.y, Scale.z));
+
+	// Cull particles using distance ranges
+	float3 CamVec = BSphere.xyz - CameraPosition;
+	float CamDistSq = dot(CamVec, CamVec);
+	if (CamDistSq < CullDistanceRangeSquared.x  || CamDistSq > CullDistanceRangeSquared.y)
+	{
+		// We are outside the distance ranges required
+		return true;
+	}
+
+	// Cull based on renderer visibility
+	if (RendererVisTagAttributeOffset != -1)
+	{
+		int Category = GetInt(RendererVisTagAttributeOffset, ParticleIndex);
+		if (Category != RendererVisibility)
+		{
+			// This particle is not visible in this renderer.
+			return true;
+		}
+	}
+
+	// Cull particles using cull planes
+	for (int i = 0; i < NIAGARA_KEY_GEN_MAX_CULL_PLANES; ++i)
+	{
+		if (i >= NumCullPlanes)
+			break;
+
+		float4 Plane = CullPlanes[i];
+		if (dot(Plane.xyz, BSphere.xyz) - Plane.w > BSphere.w)
+		{
+			// We are completely outside one of the culling planes
+			return true;
+		}
+	}
+#endif
+
+	return false;
+}
+
+[numthreads(THREAD_COUNT,1,1)]
+void GenerateParticleSortKeys(uint ParticleIndex : SV_DispatchThreadID, uint GroupTID : SV_GroupThreadID)
+{
+	CullingPrologue(GroupTID);
+
+	const uint GPUParticleCount = GPUParticleCountOffset != -1 ? GPUParticleCountBuffer[GPUParticleCountOffset] : ParticleCount;
+	
+#if !SORT_MAX_PRECISION
+	const uint SortKeyMask = 0xffff;
 	const uint SortKeySignBit = 0x8000;
 #endif
 
+	bool bCulled = true;
 	if (ParticleIndex < ParticleCount)
 	{
 		float SortValue = 0;
-		uint SortKey = SortKeyMask; 
+		uint SortKey = SortKeyMask;
 
 		if (ParticleIndex < GPUParticleCount)
 		{
-			if (SortMode == SORT_VIEW_DEPTH || SortMode == SORT_VIEW_DISTANCE)
-			{ 
+			if (SortMode == SORT_NONE)
+			{
+				// This happens if we're only culling particles, set live particles to a constant valid key
+				SortKey = 0;
+			}
+			else if (SortMode == SORT_VIEW_DEPTH || SortMode == SORT_VIEW_DISTANCE)
+			{
 				float3 ParticlePosition = GetVec3(SortAttributeOffset, ParticleIndex); 
 				float3 CameraVector = ParticlePosition - CameraPosition;
 				SortValue = SortMode == SORT_VIEW_DEPTH ? dot(CameraVector, CameraDirection) : length(CameraVector);
-				
+
 #if SORT_MAX_PRECISION
 				SortKey = (asuint(SortValue)) >> SortKeyShift;
 #else
@@ -112,10 +259,19 @@ void GenerateParticleSortKeys(uint ParticleIndex : SV_DispatchThreadID)
 			SortKey = min(SortKey , SortKeyMask - 1);
 		}
 
+		bCulled = CullParticle(ParticleIndex);
+		if (bCulled)
+		{
+			SortKey = SortKeyMask;
+		}
+
 		SortKey |= EmitterKey;
 
+		
 		OutKeys[OutputOffset + ParticleIndex] = SortKey;
 		OutParticleIndices[OutputOffset + ParticleIndex] = ParticleIndex;
 	}
+
+	CullingEpilogue(bCulled, GroupTID);
 }
- 
+
diff --git a/Engine/Plugins/FX/Niagara/Source/Niagara/Classes/NiagaraConstants.h b/Engine/Plugins/FX/Niagara/Source/Niagara/Classes/NiagaraConstants.h
--- a/Engine/Plugins/FX/Niagara/Source/Niagara/Classes/NiagaraConstants.h
+++ b/Engine/Plugins/FX/Niagara/Source/Niagara/Classes/NiagaraConstants.h
@@ -107,6 +107,7 @@
 #define SYS_PARAM_PARTICLES_LIGHT_EXPONENT               INiagaraModule::GetVar_Particles_LightExponent()
 #define SYS_PARAM_PARTICLES_LIGHT_ENABLED                INiagaraModule::GetVar_Particles_LightEnabled()
 #define SYS_PARAM_PARTICLES_LIGHT_VOLUMETRIC_SCATTERING  INiagaraModule::GetVar_Particles_LightVolumetricScattering()
+#define SYS_PARAM_PARTICLES_VISIBILITY_TAG				 INiagaraModule::GetVar_Particles_VisibilityTag()
 
 #define SYS_PARAM_PARTICLES_RIBBONID                     INiagaraModule::GetVar_Particles_RibbonID()
 #define SYS_PARAM_PARTICLES_RIBBONWIDTH                  INiagaraModule::GetVar_Particles_RibbonWidth()
diff --git a/Engine/Plugins/FX/Niagara/Source/Niagara/Classes/NiagaraGPUInstanceCountManager.h b/Engine/Plugins/FX/Niagara/Source/Niagara/Classes/NiagaraGPUInstanceCountManager.h
--- a/Engine/Plugins/FX/Niagara/Source/Niagara/Classes/NiagaraGPUInstanceCountManager.h
+++ b/Engine/Plugins/FX/Niagara/Source/Niagara/Classes/NiagaraGPUInstanceCountManager.h
@@ -52,13 +52,21 @@ public:
 	void FreeEntry(uint32& BufferOffset);
 	uint32 AcquireEntry();
 
+	uint32 AcquireCulledEntry()
+	{
+		check(!bAcquiredCulledCounts);
+		return RequiredCulledCounts++;
+	}
+
+	FRWBuffer* AcquireCulledCountsBuffer(FRHICommandListImmediate& RHICmdList, ERHIFeatureLevel::Type FeatureLevel);
+
 	const uint32* GetGPUReadback();
 	void ReleaseGPUReadback();
 	void EnqueueGPUReadback(FRHICommandListImmediate& RHICmdList);
 	bool HasPendingGPUReadback() const;
 
 	/** Add a draw indirect task to generate the draw indirect args. Returns the draw indirect arg buffer offset. */
-	uint32 AddDrawIndirect(uint32 InstanceCountBufferOffset, uint32 NumIndicesPerInstance, uint32 StartIndexLocation);
+	uint32 AddDrawIndirect(uint32 InstanceCountBufferOffset, uint32 NumIndicesPerInstance, uint32 StartIndexLocation, bool bCulled = false);
 	FRWBuffer& GetDrawIndirectBuffer() { return DrawIndirectBuffer; }
 
 	/** 
@@ -76,16 +84,21 @@ public:
 
 protected:
 
-	/** Buffer slack size hysteresis used to managed the size of CountBuffer and DrawIndirectBuffer. */
-	float BufferSlack = 1.5f;
-
 	/** The current used instance counts allocated from FNiagaraDataBuffer::AllocateGPU() */
 	int32 UsedInstanceCounts = 0;
 	/** The allocated instance counts in CountBuffer */
 	int32 AllocatedInstanceCounts = 0;
 
+	/** The number of culled instance counts needed from view culling */
+	int32 RequiredCulledCounts = 0;
+	/** The allocated instance counts in the culled count buffer*/
+	int32 AllocatedCulledCounts = 0;
+	/** Whether or not the culled counts were acquired this frame */
+	bool bAcquiredCulledCounts = false;
+
 	/** A buffer holding the each emitter particle count after a simulation tick. */
 	FRWBuffer CountBuffer;
+	FRWBuffer CulledCountBuffer;
 	TArray<uint32> FreeEntries;
 	FRHIGPUMemoryReadback* CountReadback = nullptr;
 	int32 CountReadbackSize = 0;
diff --git a/Engine/Plugins/FX/Niagara/Source/Niagara/Classes/NiagaraGPUSortInfo.h b/Engine/Plugins/FX/Niagara/Source/Niagara/Classes/NiagaraGPUSortInfo.h
--- a/Engine/Plugins/FX/Niagara/Source/Niagara/Classes/NiagaraGPUSortInfo.h
+++ b/Engine/Plugins/FX/Niagara/Source/Niagara/Classes/NiagaraGPUSortInfo.h
@@ -28,21 +28,38 @@ enum class ENiagaraSortMode : uint8
 
 struct FNiagaraGPUSortInfo
 {
+	static constexpr uint32 MaxCullPlanes = 6;
+
 	// The number of particles in the system.
 	int32 ParticleCount = 0;
 	// How the particles should be sorted.
 	ENiagaraSortMode SortMode = ENiagaraSortMode::None;
 	// On which attribute to base the sorting
-	int32 SortAttributeOffset = INDEX_NONE;
-	// The data buffer that holds the particle attributes.
+	int32 SortAttributeOffset = INDEX_NONE;	
+	// The data buffers that hold the particle attributes and their strides
 	FShaderResourceViewRHIRef ParticleDataFloatSRV;
+	FShaderResourceViewRHIRef ParticleDataHalfSRV;
+	FShaderResourceViewRHIRef ParticleDataIntSRV;
 	uint32 FloatDataStride = 0;
+	uint32 HalfDataStride = 0;
+	uint32 IntDataStride = 0;
 	// The actual GPU sim particle count. Needed to get an exact match on the index list.
 	FShaderResourceViewRHIRef GPUParticleCountSRV;
 	uint32 GPUParticleCountOffset = INDEX_NONE;
+	uint32 CulledGPUParticleCountOffset = INDEX_NONE;
 	// View data.
 	FVector ViewOrigin = FVector(0, 0, 0);
 	FVector ViewDirection = FVector(0, 0, 1);
+	// Culling/Visibility data
+	bool bEnableCulling = false;
+	int32 CullPositionAttributeOffset = INDEX_NONE;
+	int32 CullOrientationAttributeOffset = INDEX_NONE;
+	int32 CullScaleAttributeOffset = INDEX_NONE;
+	int32 RendererVisTagAttributeOffset = INDEX_NONE;
+	int32 RendererVisibility = 0;
+	FSphere LocalBSphere;
+	FVector2D DistanceCullRange { 0.0f, FLT_MAX };
+	TArray<FPlane, TFixedAllocator<MaxCullPlanes>> CullPlanes;
 
 	// The GPUSortManager bindings for this sort task.
 	FGPUSortManager::FAllocationInfo AllocationInfo;
diff --git a/Engine/Plugins/FX/Niagara/Source/Niagara/Private/NiagaraConstants.cpp b/Engine/Plugins/FX/Niagara/Source/Niagara/Private/NiagaraConstants.cpp
--- a/Engine/Plugins/FX/Niagara/Source/Niagara/Private/NiagaraConstants.cpp
+++ b/Engine/Plugins/FX/Niagara/Source/Niagara/Private/NiagaraConstants.cpp
@@ -230,6 +230,7 @@ void FNiagaraConstants::Init()
 		Attributes.Add(SYS_PARAM_PARTICLES_UV_SCALE);
 		Attributes.Add(SYS_PARAM_PARTICLES_MATERIAL_RANDOM);
 		Attributes.Add(SYS_PARAM_PARTICLES_LIGHT_RADIUS);
+		Attributes.Add(SYS_PARAM_PARTICLES_VISIBILITY_TAG);
 		Attributes.Add(SYS_PARAM_PARTICLES_RIBBONID);
 		Attributes.Add(SYS_PARAM_PARTICLES_RIBBONWIDTH);
 		Attributes.Add(SYS_PARAM_PARTICLES_RIBBONTWIST);
@@ -264,6 +265,7 @@ void FNiagaraConstants::Init()
 		AttrDataSetKeyMap.Add(SYS_PARAM_PARTICLES_LIGHT_RADIUS, GetAttributeAsDataSetKey(SYS_PARAM_PARTICLES_LIGHT_RADIUS));
 		AttrDataSetKeyMap.Add(SYS_PARAM_PARTICLES_LIGHT_EXPONENT, GetAttributeAsDataSetKey(SYS_PARAM_PARTICLES_LIGHT_EXPONENT));
 		AttrDataSetKeyMap.Add(SYS_PARAM_PARTICLES_LIGHT_VOLUMETRIC_SCATTERING, GetAttributeAsDataSetKey(SYS_PARAM_PARTICLES_LIGHT_VOLUMETRIC_SCATTERING));
+		AttrDataSetKeyMap.Add(SYS_PARAM_PARTICLES_VISIBILITY_TAG, GetAttributeAsDataSetKey(SYS_PARAM_PARTICLES_VISIBILITY_TAG));
 		AttrDataSetKeyMap.Add(SYS_PARAM_PARTICLES_RIBBONID, GetAttributeAsDataSetKey(SYS_PARAM_PARTICLES_RIBBONID));
 		AttrDataSetKeyMap.Add(SYS_PARAM_PARTICLES_RIBBONWIDTH, GetAttributeAsDataSetKey(SYS_PARAM_PARTICLES_RIBBONWIDTH));
 		AttrDataSetKeyMap.Add(SYS_PARAM_PARTICLES_RIBBONTWIST, GetAttributeAsDataSetKey(SYS_PARAM_PARTICLES_RIBBONTWIST));
@@ -382,6 +384,11 @@ void FNiagaraConstants::Init()
 		Var.SetValue<float>(0.0f);
 		AttrDefaultsValueMap.Add(SYS_PARAM_PARTICLES_LIGHT_VOLUMETRIC_SCATTERING, Var);
 
+		AttrDefaultsStrMap.Add(SYS_PARAM_PARTICLES_VISIBILITY_TAG, TEXT("0"));
+		Var = SYS_PARAM_PARTICLES_VISIBILITY_TAG;
+		Var.SetValue<float>(0.0f);
+		AttrDefaultsValueMap.Add(SYS_PARAM_PARTICLES_VISIBILITY_TAG, Var);
+
 		AttrDefaultsStrMap.Add(SYS_PARAM_PARTICLES_RIBBONID, TEXT("0"));
 		Var = SYS_PARAM_PARTICLES_RIBBONID;
 		Var.SetValue<FNiagaraID>(FNiagaraID());
@@ -434,6 +441,7 @@ void FNiagaraConstants::Init()
 		AttrDescStrMap.Add(SYS_PARAM_PARTICLES_LIGHT_EXPONENT, LOCTEXT("LightExponentParamDesc", "Used to drive the attenuation of the light when using a Light renderer without inverse squared falloff enabled."));
 		AttrDescStrMap.Add(SYS_PARAM_PARTICLES_LIGHT_VOLUMETRIC_SCATTERING, LOCTEXT("LightVolumetricScatteringParamDesc", "Used to drive the volumetric scattering intensity of the light when using a Light renderer."));
 		AttrDescStrMap.Add(SYS_PARAM_INSTANCE_ALIVE, LOCTEXT("AliveParamDesc", "Used to determine whether or not this particle instance is still valid or if it can be deleted."));
+		AttrDescStrMap.Add(SYS_PARAM_PARTICLES_VISIBILITY_TAG, LOCTEXT("VisibilityTag", "Used for selecting renderers to use when rendering this particle. Without this, the particle will render in all renderers"));
 		AttrDescStrMap.Add(SYS_PARAM_PARTICLES_RIBBONID, LOCTEXT("RibbonIDDesc", "Sets the ribbon id for a particle. Particles with the same ribbon id will be connected into a ribbon."));
 		AttrDescStrMap.Add(SYS_PARAM_PARTICLES_RIBBONWIDTH, LOCTEXT("RibbonWidthDesc", "Sets the ribbon width for a particle, in UE4 units."));
 		AttrDescStrMap.Add(SYS_PARAM_PARTICLES_RIBBONTWIST, LOCTEXT("RibbonTwistDesc", "Sets the ribbon twist for a particle, in degrees."));
diff --git a/Engine/Plugins/FX/Niagara/Source/Niagara/Private/NiagaraEmitterInstanceBatcher.cpp b/Engine/Plugins/FX/Niagara/Source/Niagara/Private/NiagaraEmitterInstanceBatcher.cpp
--- a/Engine/Plugins/FX/Niagara/Source/Niagara/Private/NiagaraEmitterInstanceBatcher.cpp
+++ b/Engine/Plugins/FX/Niagara/Source/Niagara/Private/NiagaraEmitterInstanceBatcher.cpp
@@ -89,6 +89,17 @@ NiagaraEmitterInstanceBatcher::NiagaraEmitterInstanceBatcher(ERHIFeatureLevel::T
 		}), 
 		EGPUSortFlags::AnyKeyPrecision | EGPUSortFlags::KeyGenAfterPreRender | EGPUSortFlags::AnySortLocation | EGPUSortFlags::ValuesAsInt32,
 		Name);
+
+		if (FNiagaraUtilities::AllowGPUParticles(GetShaderPlatform()))
+		{
+			// Because of culled indirect draw args, we have to update the draw indirect buffer after the sort key generation
+			GPUSortManager->PostPreRenderEvent.AddLambda(
+				[this](FRHICommandListImmediate& RHICmdList)
+				{
+					GPUInstanceCounterManager.UpdateDrawIndirectBuffer(RHICmdList, FeatureLevel);
+				}
+			);
+		}
 	}
 }
 
@@ -1078,9 +1089,6 @@ void NiagaraEmitterInstanceBatcher::PreRender(FRHICommandListImmediate& RHICmdLi
 	LLM_SCOPE(ELLMTag::Niagara);
 
 	GlobalDistanceFieldParams = GlobalDistanceFieldParameterData ? *GlobalDistanceFieldParameterData : FGlobalDistanceFieldParameterData();
-
-	// Update draw indirect args from the simulation results.
-	GPUInstanceCounterManager.UpdateDrawIndirectBuffer(RHICmdList, FeatureLevel);
 }
 
 void NiagaraEmitterInstanceBatcher::OnDestroy()
@@ -1109,26 +1117,77 @@ void NiagaraEmitterInstanceBatcher::GenerateSortKeys(FRHICommandListImmediate& R
 	// Currently all Niagara KeyGen must execute after PreRender() - in between PreInitViews() and PostRenderOpaque(), when the GPU simulation are possibly ticked.
 	check(EnumHasAnyFlags(Flags, EGPUSortFlags::KeyGenAfterPreRender));
 
+	FRWBuffer* CulledCountsBuffer = GPUInstanceCounterManager.AcquireCulledCountsBuffer(RHICmdList, FeatureLevel);
+
 	const FGPUSortManager::FKeyGenInfo KeyGenInfo((uint32)NumElementsInBatch, EnumHasAnyFlags(Flags, EGPUSortFlags::HighPrecisionKeys));
 
-	FNiagaraSortKeyGenCS::FPermutationDomain PermutationVector;
-	PermutationVector.Set<FNiagaraSortKeyGenCS::FSortUsingMaxPrecision>(EnumHasAnyFlags(Flags, EGPUSortFlags::HighPrecisionKeys));
-	TShaderMapRef<FNiagaraSortKeyGenCS> KeyGenCS(GetGlobalShaderMap(FeatureLevel), PermutationVector);
-	RHICmdList.SetComputeShader(KeyGenCS.GetComputeShader());
-	KeyGenCS->SetOutput(RHICmdList, KeysUAV, ValuesUAV);
+	FNiagaraSortKeyGenCS::FPermutationDomain SortPermutationVector;
+	SortPermutationVector.Set<FNiagaraSortKeyGenCS::FSortUsingMaxPrecision>(GNiagaraGPUSortingUseMaxPrecision != 0);
+	SortPermutationVector.Set<FNiagaraSortKeyGenCS::FEnableCulling>(false);
+
+	FNiagaraSortKeyGenCS::FPermutationDomain SortAndCullPermutationVector;
+	SortAndCullPermutationVector.Set<FNiagaraSortKeyGenCS::FSortUsingMaxPrecision>(GNiagaraGPUSortingUseMaxPrecision != 0);
+	SortAndCullPermutationVector.Set<FNiagaraSortKeyGenCS::FEnableCulling>(true);
+
+	TShaderMapRef<FNiagaraSortKeyGenCS> SortKeyGenCS(GetGlobalShaderMap(FeatureLevel), SortPermutationVector);
+	TShaderMapRef<FNiagaraSortKeyGenCS> SortAndCullKeyGenCS(GetGlobalShaderMap(FeatureLevel), SortAndCullPermutationVector);
+	
+	FNiagaraSortKeyGenCS::FParameters Params;
+	Params.SortKeyMask = KeyGenInfo.SortKeyParams.X;
+	Params.SortKeyShift = KeyGenInfo.SortKeyParams.Y;
+	Params.SortKeySignBit = KeyGenInfo.SortKeyParams.Z;
+	Params.OutKeys = KeysUAV;
+	Params.OutParticleIndices = ValuesUAV;
+	Params.OutCulledParticleCounts = CulledCountsBuffer ? (FRHIUnorderedAccessView*)CulledCountsBuffer->UAV : GetEmptyUAVFromPool(RHICmdList, PF_R32_UINT, false);
 
 	FRHIUnorderedAccessView* OutputUAVs[] = { KeysUAV, ValuesUAV };
 	for (const FNiagaraGPUSortInfo& SortInfo : SimulationsToSort)
 	{
 		if (SortInfo.AllocationInfo.SortBatchId == BatchId)
 		{
-			KeyGenCS->SetParameters(RHICmdList, SortInfo, (uint32)SortInfo.AllocationInfo.ElementIndex << KeyGenInfo.ElementKeyShift, SortInfo.AllocationInfo.BufferOffset, KeyGenInfo.SortKeyParams);
-			DispatchComputeShader(RHICmdList, KeyGenCS, FMath::DivideAndRoundUp(SortInfo.ParticleCount, NIAGARA_KEY_GEN_THREAD_COUNT), 1, 1);
+			Params.NiagaraParticleDataFloat = SortInfo.ParticleDataFloatSRV;
+			Params.NiagaraParticleDataHalf = SortInfo.ParticleDataHalfSRV;
+			Params.NiagaraParticleDataInt = SortInfo.ParticleDataIntSRV;
+			Params.GPUParticleCountBuffer = SortInfo.GPUParticleCountSRV;
+			Params.FloatDataStride = SortInfo.FloatDataStride;
+			Params.HalfDataStride = SortInfo.HalfDataStride;
+			Params.IntDataStride = SortInfo.IntDataStride;
+			Params.ParticleCount = SortInfo.ParticleCount;
+			Params.GPUParticleCountBuffer = SortInfo.GPUParticleCountSRV;
+			Params.GPUParticleCountOffset = SortInfo.GPUParticleCountOffset;
+			Params.CulledGPUParticleCountOffset = SortInfo.CulledGPUParticleCountOffset;
+			Params.EmitterKey = (uint32)SortInfo.AllocationInfo.ElementIndex << KeyGenInfo.ElementKeyShift;
+			Params.OutputOffset = SortInfo.AllocationInfo.BufferOffset;
+			Params.CameraPosition = SortInfo.ViewOrigin;
+			Params.CameraDirection = SortInfo.ViewDirection;
+			Params.SortMode = (uint32)SortInfo.SortMode;
+			Params.SortAttributeOffset = SortInfo.SortAttributeOffset;
+			Params.CullPositionAttributeOffset = SortInfo.CullPositionAttributeOffset;
+			Params.CullOrientationAttributeOffset = SortInfo.CullOrientationAttributeOffset;
+			Params.CullScaleAttributeOffset = SortInfo.CullScaleAttributeOffset;
+			Params.RendererVisibility = SortInfo.RendererVisibility;
+			Params.RendererVisTagAttributeOffset = SortInfo.RendererVisTagAttributeOffset;
+			Params.CullDistanceRangeSquared = SortInfo.DistanceCullRange * SortInfo.DistanceCullRange;
+			Params.LocalBoundingSphere = FVector4(SortInfo.LocalBSphere.Center, SortInfo.LocalBSphere.W);
+
+			Params.NumCullPlanes = 0;
+			for (const FPlane& Plane : SortInfo.CullPlanes)
+			{
+				Params.CullPlanes[Params.NumCullPlanes++] = FVector4(Plane.X, Plane.Y, Plane.Z, Plane.W);
+			}
+
+			// Choose the shader to bind
+			TShaderMapRef<FNiagaraSortKeyGenCS> KeyGenCS = SortInfo.bEnableCulling ? SortAndCullKeyGenCS : SortKeyGenCS;
+			RHICmdList.SetComputeShader(KeyGenCS.GetComputeShader());
+
+			SetShaderParameters(RHICmdList, KeyGenCS, KeyGenCS.GetComputeShader(), Params);
+			DispatchComputeShader(RHICmdList, KeyGenCS, FMath::DivideAndRoundUp(SortInfo.ParticleCount, NIAGARA_KEY_GEN_THREAD_COUNT), 1, 1);			
+			UnsetShaderUAVs(RHICmdList, KeyGenCS, KeyGenCS.GetComputeShader());
+			
 			// TR-KeyGen : No sync needed between tasks since they update different parts of the data (assuming it's ok if cache lines overlap).
 			RHICmdList.TransitionResources(EResourceTransitionAccess::ERWNoBarrier, EResourceTransitionPipeline::EComputeToCompute, OutputUAVs, UE_ARRAY_COUNT(OutputUAVs));
 		}
 	}
-	KeyGenCS->UnbindBuffers(RHICmdList);
 }
 
 void NiagaraEmitterInstanceBatcher::ProcessDebugInfo(FRHICommandList &RHICmdList, const FNiagaraComputeExecutionContext* Context) const
diff --git a/Engine/Plugins/FX/Niagara/Source/Niagara/Private/NiagaraGPUInstanceCountManager.cpp b/Engine/Plugins/FX/Niagara/Source/Niagara/Private/NiagaraGPUInstanceCountManager.cpp
--- a/Engine/Plugins/FX/Niagara/Source/Niagara/Private/NiagaraGPUInstanceCountManager.cpp
+++ b/Engine/Plugins/FX/Niagara/Source/Niagara/Private/NiagaraGPUInstanceCountManager.cpp
@@ -5,6 +5,8 @@
 #include "Containers/DynamicRHIResourceArray.h"
 #include "GPUSortManager.h" // CopyUIntBufferToTargets
 #include "ProfilingDebugging/RealtimeGPUProfiler.h"
+#include "NiagaraRenderer.h"
+#include "ClearQuad.h"
 
 int32 GNiagaraMinGPUInstanceCount = 2048;
 static FAutoConsoleVariableRef CVarNiagaraMinGPUInstanceCount(
@@ -14,6 +16,14 @@ static FAutoConsoleVariableRef CVarNiagaraMinGPUInstanceCount(
 	ECVF_Default
 );
 
+int32 GNiagaraMinCulledGPUInstanceCount = 2048;
+static FAutoConsoleVariableRef CVarNiagaraMinCulledGPUInstanceCount(
+	TEXT("Niagara.MinCulledGPUInstanceCount"),
+	GNiagaraMinCulledGPUInstanceCount,
+	TEXT("Minimum number of culled (per-view) instance count entries allocated in the global buffer. (default=2048)"),
+	ECVF_Default
+);
+
 int32 GNiagaraMinGPUDrawIndirectArgs = 256;
 static FAutoConsoleVariableRef CVarNiagaraMinGPUDrawIndirectArgs(
 	TEXT("Niagara.MinGPUDrawIndirectArgs"),
@@ -22,6 +32,22 @@ static FAutoConsoleVariableRef CVarNiagaraMinGPUDrawIndirectArgs(
 	ECVF_Default
 );
 
+float GNiagaraGPUCountBufferSlack = 1.5f;
+static FAutoConsoleVariableRef CVarNiagaraGPUCountBufferSlack(
+	TEXT("Niagara.GPUCountBufferSlack"),
+	GNiagaraGPUCountBufferSlack,
+	TEXT("Multiplier of the GPU count buffer size to prevent frequent re-allocation."),
+	ECVF_Default
+);
+
+float GNiagaraGPUDrawIndirectArgsBufferSlack = 2.0f;
+static FAutoConsoleVariableRef CVarNiagaraGPUDrawIndirectArgsBufferSlack(
+	TEXT("Niagara.GPUDrawIndirectArgsBufferSlack"),
+	GNiagaraGPUDrawIndirectArgsBufferSlack,
+	TEXT("Multiplier on the required indirect args buffer size to account for culling and multiple mesh sections on mesh renderers."),
+	ECVF_Default
+);
+
 DECLARE_DWORD_COUNTER_STAT(TEXT("Used GPU Instance Counters"), STAT_NiagaraUsedGPUInstanceCounters, STATGROUP_Niagara);
 DECLARE_DWORD_COUNTER_STAT(TEXT("Max Num GPU Renderers"), STAT_NiagaraMaxNumGPURenderers, STATGROUP_Niagara);
 DECLARE_CYCLE_STAT(TEXT("GPU Readback Lock"), STAT_NiagaraGPUReadbackLock, STATGROUP_Niagara);
@@ -46,9 +72,11 @@ void FNiagaraGPUInstanceCountManager::InitRHI()
 void FNiagaraGPUInstanceCountManager::ReleaseRHI()
 {
 	CountBuffer.Release();
+	CulledCountBuffer.Release();
 	DrawIndirectBuffer.Release();
 
 	AllocatedInstanceCounts = 0;
+	AllocatedCulledCounts = 0;
 	AllocatedDrawIndirectArgs = 0;
 
 	if (CountReadback)
@@ -99,13 +127,43 @@ void FNiagaraGPUInstanceCountManager::FreeEntry(uint32& BufferOffset)
 	}
 }
 
+FRWBuffer* FNiagaraGPUInstanceCountManager::AcquireCulledCountsBuffer(FRHICommandListImmediate& RHICmdList, ERHIFeatureLevel::Type FeatureLevel)
+{
+	if (RequiredCulledCounts > 0)
+	{
+		if (!bAcquiredCulledCounts)
+		{
+			const int32 RecommendedCulledCounts = FMath::Max(GNiagaraMinCulledGPUInstanceCount, (int32)(RequiredCulledCounts * GNiagaraGPUCountBufferSlack));
+			if (RecommendedCulledCounts > AllocatedCulledCounts)
+			{
+				// We need a bigger buffer
+				CulledCountBuffer.Release();
+
+				AllocatedCulledCounts = RecommendedCulledCounts;
+				CulledCountBuffer.Initialize(sizeof(uint32), AllocatedCulledCounts, EPixelFormat::PF_R32_UINT, BUF_Transient, TEXT("NiagaraCulledGPUInstanceCounts"));
+			}
+
+			CulledCountBuffer.AcquireTransientResource();
+
+			// Initialize the buffer by clearing it to zero then transition it to be ready to write to
+			RHICmdList.ClearUAVUint(CulledCountBuffer.UAV, FUintVector4(EForceInit::ForceInitToZero));
+			RHICmdList.TransitionResource(EResourceTransitionAccess::ERWBarrier, EResourceTransitionPipeline::EComputeToCompute, CulledCountBuffer.UAV);
+		}
+
+		return &CulledCountBuffer;
+	}	
+
+	bAcquiredCulledCounts = true;
+	return nullptr;
+}
+
 void FNiagaraGPUInstanceCountManager::ResizeBuffers(FRHICommandListImmediate& RHICmdList, ERHIFeatureLevel::Type FeatureLevel, int32 ReservedInstanceCounts)
 {
 	const int32 RequiredInstanceCounts = UsedInstanceCounts + FMath::Max<int32>(ReservedInstanceCounts - FreeEntries.Num(), 0);
-	const int32 MaxDrawIndirectArgs = NumRegisteredGPURenderers->Value;
-	if (RequiredInstanceCounts > 0 || MaxDrawIndirectArgs > 0)
+	const int32 ExpectedDrawIndirectArgs = NumRegisteredGPURenderers->Value; // can be more with multiple mesh sections
+	if (RequiredInstanceCounts > 0 || ExpectedDrawIndirectArgs > 0)
 	{
-		const int32 RecommendedInstanceCounts = FMath::Max(GNiagaraMinGPUInstanceCount, (int32)(RequiredInstanceCounts * BufferSlack));
+		const int32 RecommendedInstanceCounts = FMath::Max(GNiagaraMinGPUInstanceCount, (int32)(RequiredInstanceCounts * GNiagaraGPUCountBufferSlack));
 		// If the buffer is not allocated, allocate it to the recommended size.
 		if (!AllocatedInstanceCounts)
 		{
@@ -144,14 +202,14 @@ void FNiagaraGPUInstanceCountManager::ResizeBuffers(FRHICommandListImmediate& RH
 			//UE_LOG(LogNiagara, Log, TEXT("FNiagaraGPUInstanceCountManager::ResizeBuffers Resize AllocatedInstanceCounts: %d ReservedInstanceCounts: %d"), AllocatedInstanceCounts, ReservedInstanceCounts);
 		}
 		// If we need to shrink the buffer size because use way to much buffer size.
-		else if ((int32)(RecommendedInstanceCounts * BufferSlack) < AllocatedInstanceCounts)
+		else if ((int32)(RecommendedInstanceCounts * GNiagaraGPUCountBufferSlack) < AllocatedInstanceCounts)
 		{
 			// possibly shrink but hard to do because of sparse array allocation.
 		}
 
-		// We expect 2 niagara buffer per instance for simplicity
-		int32 RecommendedDrawIndirectArgsCount = FMath::Max(GNiagaraMinGPUDrawIndirectArgs, (int32)(MaxDrawIndirectArgs * BufferSlack));
-		if (MaxDrawIndirectArgs > AllocatedDrawIndirectArgs || (int32)(RecommendedDrawIndirectArgsCount * BufferSlack) < AllocatedDrawIndirectArgs)
+		// We expect more slots to be needed due to multiple mesh draws per renderer, and for culled draws per view
+		int32 RecommendedDrawIndirectArgsCount = FMath::Max(GNiagaraMinGPUDrawIndirectArgs, (int32)(ExpectedDrawIndirectArgs * GNiagaraGPUDrawIndirectArgsBufferSlack));
+		if (ExpectedDrawIndirectArgs > AllocatedDrawIndirectArgs || (int32)(RecommendedDrawIndirectArgsCount * GNiagaraGPUDrawIndirectArgsBufferSlack) < AllocatedDrawIndirectArgs)
 		{
 			DrawIndirectBuffer.Release();
 			AllocatedDrawIndirectArgs = RecommendedDrawIndirectArgsCount;
@@ -164,14 +222,14 @@ void FNiagaraGPUInstanceCountManager::ResizeBuffers(FRHICommandListImmediate& RH
 	}
 
 	INC_DWORD_STAT_BY(STAT_NiagaraUsedGPUInstanceCounters, RequiredInstanceCounts);
-	INC_DWORD_STAT_BY(STAT_NiagaraMaxNumGPURenderers, MaxDrawIndirectArgs);
+	INC_DWORD_STAT_BY(STAT_NiagaraMaxNumGPURenderers, ExpectedDrawIndirectArgs);
 }
 
-uint32 FNiagaraGPUInstanceCountManager::AddDrawIndirect(uint32 InstanceCountBufferOffset, uint32 NumIndicesPerInstance, uint32 StartIndexLocation)
+uint32 FNiagaraGPUInstanceCountManager::AddDrawIndirect(uint32 InstanceCountBufferOffset, uint32 NumIndicesPerInstance, uint32 StartIndexLocation, bool bCulled)
 {
 	checkSlow(IsInRenderingThread());
 
-	const FArgGenTaskInfo Info(InstanceCountBufferOffset, NumIndicesPerInstance, StartIndexLocation);
+	const FArgGenTaskInfo Info(InstanceCountBufferOffset, NumIndicesPerInstance, StartIndexLocation, bCulled);
 
 	uint32& CachedOffset = DrawIndirectArgMap.FindOrAdd(Info, INDEX_NONE);
 	if (CachedOffset != INDEX_NONE)
@@ -180,20 +238,13 @@ uint32 FNiagaraGPUInstanceCountManager::AddDrawIndirect(uint32 InstanceCountBuff
 	}
 	else if (DrawIndirectArgGenTasks.Num() < AllocatedDrawIndirectArgs)
 	{
-#if !UE_BUILD_SHIPPING
-		const int32 MaxDrawIndirectArgs = NumRegisteredGPURenderers->Value;
-		if (DrawIndirectArgGenTasks.Num() >= MaxDrawIndirectArgs)
-		{
-			UE_LOG(LogNiagara, Warning, TEXT("More draw indirect args then expected (%d / %d)"), DrawIndirectArgGenTasks.Num() + 1, MaxDrawIndirectArgs);
-		}
-#endif
 		DrawIndirectArgGenTasks.Add(Info);
 		CachedOffset = (DrawIndirectArgGenTasks.Num() - 1) * NIAGARA_DRAW_INDIRECT_ARGS_SIZE * sizeof(uint32);
 		return CachedOffset;
 	}
 	else
 	{
-		UE_LOG(LogNiagara, Error, TEXT("Too many draw indirect args. Some draw call will be corrupted"));
+		UE_LOG(LogNiagara, Error, TEXT("Ran out of space in the GPU Draw Indirect Args buffer. Increase Niagara.GPUDrawIndirectArgsBufferSlack"));
 		return INDEX_NONE;
 	}
 }
@@ -204,6 +255,12 @@ void FNiagaraGPUInstanceCountManager::UpdateDrawIndirectBuffer(FRHICommandList&
 	{
 		if (FNiagaraUtilities::AllowGPUParticles(GShaderPlatformForFeatureLevel[FeatureLevel]))
 		{
+			if (RequiredCulledCounts > 0)
+			{
+				RHICmdList.TransitionResource(EResourceTransitionAccess::EReadable, EResourceTransitionPipeline::EComputeToCompute, CulledCountBuffer.UAV);
+			}
+			RHICmdList.TransitionResource(EResourceTransitionAccess::ERWBarrier, EResourceTransitionPipeline::EComputeToCompute, CountBuffer.UAV);
+
 			FReadBuffer TaskInfosBuffer;
 
 			//for (int32 i = 0; i < InstanceCountClearTasks.Num(); i++)
@@ -222,46 +279,19 @@ void FNiagaraGPUInstanceCountManager::UpdateDrawIndirectBuffer(FRHICommandList&
 				RHIUnlockVertexBuffer(TaskInfosBuffer.Buffer);
 			}
 
-			FNiagaraDrawIndirectArgsGenCS::FPermutationDomain PermutationVector;
-			PermutationVector.Set<FNiagaraDrawIndirectArgsGenCS::FSupportsTextureRW>(GRHISupportsRWTextureBuffers ? 1 : 0);
-			TShaderMapRef<FNiagaraDrawIndirectArgsGenCS> DrawIndirectArgsGenCS(GetGlobalShaderMap(FeatureLevel), PermutationVector);
-
-			RHICmdList.TransitionResource(GRHISupportsRWTextureBuffers ? EResourceTransitionAccess::ERWBarrier : EResourceTransitionAccess::EReadable,
-											EResourceTransitionPipeline::EComputeToCompute, CountBuffer.UAV);
 			RHICmdList.TransitionResource(EResourceTransitionAccess::EWritable, EResourceTransitionPipeline::EGfxToCompute, DrawIndirectBuffer.UAV);
 
+			FRHIShaderResourceView* CulledCountsSRV = CulledCountBuffer.SRV.IsValid() ? CulledCountBuffer.SRV.GetReference() : FNiagaraRenderer::GetDummyUIntBuffer();
+			FNiagaraDrawIndirectArgsGenCS::FPermutationDomain PermutationVector;
+			TShaderMapRef<FNiagaraDrawIndirectArgsGenCS> DrawIndirectArgsGenCS(GetGlobalShaderMap(FeatureLevel), PermutationVector);
 			RHICmdList.SetComputeShader(DrawIndirectArgsGenCS.GetComputeShader());
 			DrawIndirectArgsGenCS->SetOutput(RHICmdList, DrawIndirectBuffer.UAV, CountBuffer.UAV);
-			DrawIndirectArgsGenCS->SetParameters(RHICmdList, TaskInfosBuffer.SRV, DrawIndirectArgGenTasks.Num(), InstanceCountClearTasks.Num());
-
-			// If the device supports RW Texture buffers then we can use a single compute pass, otherwise we need to split into two passes
-			if (GRHISupportsRWTextureBuffers)
-			{
-				DispatchComputeShader(RHICmdList, DrawIndirectArgsGenCS.GetShader(), FMath::DivideAndRoundUp(DrawIndirectArgGenTasks.Num() + InstanceCountClearTasks.Num(), NIAGARA_DRAW_INDIRECT_ARGS_GEN_THREAD_COUNT), 1, 1);
-				DrawIndirectArgsGenCS->UnbindBuffers(RHICmdList);
-
-				// Sync after clear.
-				RHICmdList.TransitionResource(EResourceTransitionAccess::ERWBarrier, EResourceTransitionPipeline::EComputeToCompute, CountBuffer.UAV);
-			}
-			else
-			{
-				DispatchComputeShader(RHICmdList, DrawIndirectArgsGenCS.GetShader(), FMath::DivideAndRoundUp(DrawIndirectArgGenTasks.Num(), NIAGARA_DRAW_INDIRECT_ARGS_GEN_THREAD_COUNT), 1, 1);
-				DrawIndirectArgsGenCS->UnbindBuffers(RHICmdList);
-
-				RHICmdList.TransitionResource(EResourceTransitionAccess::EWritable, EResourceTransitionPipeline::EComputeToCompute, CountBuffer.UAV);
-
-				FNiagaraDrawIndirectResetCountsCS::FPermutationDomain PermutationVectorResetCounts;
-				TShaderMapRef<FNiagaraDrawIndirectResetCountsCS> DrawIndirectResetCountsArgsGenCS(GetGlobalShaderMap(FeatureLevel), PermutationVectorResetCounts);
-				RHICmdList.SetComputeShader(DrawIndirectResetCountsArgsGenCS.GetComputeShader());
-				DrawIndirectResetCountsArgsGenCS->SetOutput(RHICmdList, CountBuffer.UAV);
-				DrawIndirectResetCountsArgsGenCS->SetParameters(RHICmdList, TaskInfosBuffer.SRV, DrawIndirectArgGenTasks.Num(), InstanceCountClearTasks.Num());
-				DispatchComputeShader(RHICmdList, DrawIndirectResetCountsArgsGenCS.GetShader(), FMath::DivideAndRoundUp(InstanceCountClearTasks.Num(), NIAGARA_DRAW_INDIRECT_ARGS_GEN_THREAD_COUNT), 1, 1);
-				DrawIndirectResetCountsArgsGenCS->UnbindBuffers(RHICmdList);
-
-				// Sync after clear.
-				RHICmdList.TransitionResource(EResourceTransitionAccess::EReadable, EResourceTransitionPipeline::EComputeToCompute, CountBuffer.UAV);
-			}
+			DrawIndirectArgsGenCS->SetParameters(RHICmdList, TaskInfosBuffer.SRV, CulledCountsSRV, DrawIndirectArgGenTasks.Num(), InstanceCountClearTasks.Num());
+			DispatchComputeShader(RHICmdList, DrawIndirectArgsGenCS.GetShader(), FMath::DivideAndRoundUp(DrawIndirectArgGenTasks.Num() + InstanceCountClearTasks.Num(), NIAGARA_DRAW_INDIRECT_ARGS_GEN_THREAD_COUNT), 1, 1);
+			DrawIndirectArgsGenCS->UnbindBuffers(RHICmdList);
 
+			// Sync after clear.
+			RHICmdList.TransitionResource(EResourceTransitionAccess::ERWBarrier, EResourceTransitionPipeline::EComputeToCompute, CountBuffer.UAV);
 			// Transition draw indirect to readable for gfx draw indirect.
 			RHICmdList.TransitionResource(EResourceTransitionAccess::EReadable, EResourceTransitionPipeline::EComputeToGfx, DrawIndirectBuffer.UAV);
 		}
@@ -276,6 +306,13 @@ void FNiagaraGPUInstanceCountManager::UpdateDrawIndirectBuffer(FRHICommandList&
 		DrawIndirectArgGenTasks.Empty();
 		DrawIndirectArgMap.Empty();
 		InstanceCountClearTasks.Empty();
+
+		if (bAcquiredCulledCounts && RequiredCulledCounts > 0)
+		{
+			CulledCountBuffer.DiscardTransientResource();
+		}
+		bAcquiredCulledCounts = false;
+		RequiredCulledCounts = 0; // reset this counter now that we're done with them
 	}
 }
 
diff --git a/Engine/Plugins/FX/Niagara/Source/Niagara/Private/NiagaraMeshRendererProperties.cpp b/Engine/Plugins/FX/Niagara/Source/Niagara/Private/NiagaraMeshRendererProperties.cpp
--- a/Engine/Plugins/FX/Niagara/Source/Niagara/Private/NiagaraMeshRendererProperties.cpp
+++ b/Engine/Plugins/FX/Niagara/Source/Niagara/Private/NiagaraMeshRendererProperties.cpp
@@ -66,6 +66,8 @@ UNiagaraMeshRendererProperties::UNiagaraMeshRendererProperties()
 	AttributeBindings.Add(&MaterialRandomBinding);
 	AttributeBindings.Add(&CustomSortingBinding);
 	AttributeBindings.Add(&NormalizedAgeBinding);
+	AttributeBindings.Add(&CameraOffsetBinding);
+	AttributeBindings.Add(&RendererVisibilityTagBinding);
 }
 
 FNiagaraRenderer* UNiagaraMeshRendererProperties::CreateEmitterRenderer(ERHIFeatureLevel::Type FeatureLevel, const FNiagaraEmitterInstance* Emitter)
@@ -153,6 +155,7 @@ void UNiagaraMeshRendererProperties::InitBindings()
 		MaterialRandomBinding = FNiagaraConstants::GetAttributeDefaultBinding(SYS_PARAM_PARTICLES_MATERIAL_RANDOM);
 		NormalizedAgeBinding = FNiagaraConstants::GetAttributeDefaultBinding(SYS_PARAM_PARTICLES_NORMALIZED_AGE);
 		CameraOffsetBinding = FNiagaraConstants::GetAttributeDefaultBinding(SYS_PARAM_PARTICLES_CAMERA_OFFSET);
+		RendererVisibilityTagBinding = FNiagaraConstants::GetAttributeDefaultBinding(SYS_PARAM_PARTICLES_VISIBILITY_TAG);
 
 		//Default custom sorting to age
 		CustomSortingBinding = FNiagaraConstants::GetAttributeDefaultBinding(SYS_PARAM_PARTICLES_NORMALIZED_AGE);
diff --git a/Engine/Plugins/FX/Niagara/Source/Niagara/Private/NiagaraModule.cpp b/Engine/Plugins/FX/Niagara/Source/Niagara/Private/NiagaraModule.cpp
--- a/Engine/Plugins/FX/Niagara/Source/Niagara/Private/NiagaraModule.cpp
+++ b/Engine/Plugins/FX/Niagara/Source/Niagara/Private/NiagaraModule.cpp
@@ -137,6 +137,7 @@ FNiagaraVariable INiagaraModule::Particles_RibbonWidth;
 FNiagaraVariable INiagaraModule::Particles_RibbonTwist;
 FNiagaraVariable INiagaraModule::Particles_RibbonFacing;
 FNiagaraVariable INiagaraModule::Particles_RibbonLinkOrder;
+FNiagaraVariable INiagaraModule::Particles_VisibilityTag;
 FNiagaraVariable INiagaraModule::ScriptUsage;
 FNiagaraVariable INiagaraModule::DataInstance_Alive;
 FNiagaraVariable INiagaraModule::Translator_BeginDefaults;
@@ -239,6 +240,7 @@ void INiagaraModule::StartupModule()
 	Particles_RibbonTwist = FNiagaraVariable(FNiagaraTypeDefinition::GetFloatDef(), TEXT("Particles.RibbonTwist"));
 	Particles_RibbonFacing = FNiagaraVariable(FNiagaraTypeDefinition::GetVec3Def(), TEXT("Particles.RibbonFacing"));
 	Particles_RibbonLinkOrder = FNiagaraVariable(FNiagaraTypeDefinition::GetFloatDef(), TEXT("Particles.RibbonLinkOrder"));
+	Particles_VisibilityTag = FNiagaraVariable(FNiagaraTypeDefinition::GetIntDef(), TEXT("Particles.VisibilityTag"));
 
 	ScriptUsage = FNiagaraVariable(FNiagaraTypeDefinition::GetScriptUsageEnum(), TEXT("Script.Usage"));
 	DataInstance_Alive = FNiagaraVariable(FNiagaraTypeDefinition::GetBoolDef(), TEXT("DataInstance.Alive"));
diff --git a/Engine/Plugins/FX/Niagara/Source/Niagara/Private/NiagaraRenderer.cpp b/Engine/Plugins/FX/Niagara/Source/Niagara/Private/NiagaraRenderer.cpp
--- a/Engine/Plugins/FX/Niagara/Source/Niagara/Private/NiagaraRenderer.cpp
+++ b/Engine/Plugins/FX/Niagara/Source/Niagara/Private/NiagaraRenderer.cpp
@@ -162,7 +162,6 @@ FRHIShaderResourceView* FNiagaraRenderer::GetDummyHalfBuffer()
 	return DummyHalfBuffer.SRV;
 }
 
-
 bool FNiagaraRenderer::SetVertexFactoryVariable(const FNiagaraDataSet& DataSet, const FNiagaraVariable& Var, int32 VFVarOffset)
 {
 	// use the DataSetVariable to figure out the information about the data that we'll be sending to the renderer
@@ -349,7 +348,7 @@ FNiagaraRenderer::~FNiagaraRenderer()
 
 void FNiagaraRenderer::CreateRenderThreadResources(NiagaraEmitterInstanceBatcher* Batcher) 
 {
-	if (Batcher && SimTarget == ENiagaraSimTarget::GPUComputeSim)
+	if (Batcher)
 	{
 		NumRegisteredGPURenderers = Batcher->GetGPUInstanceCounterManager().GetGPURendererCount();
 		if (NumRegisteredGPURenderers)
diff --git a/Engine/Plugins/FX/Niagara/Source/Niagara/Private/NiagaraRendererMeshes.cpp b/Engine/Plugins/FX/Niagara/Source/Niagara/Private/NiagaraRendererMeshes.cpp
--- a/Engine/Plugins/FX/Niagara/Source/Niagara/Private/NiagaraRendererMeshes.cpp
+++ b/Engine/Plugins/FX/Niagara/Source/Niagara/Private/NiagaraRendererMeshes.cpp
@@ -13,6 +13,7 @@
 #include "RayTracingDefinitions.h"
 #include "RayTracingDynamicGeometryCollection.h"
 #include "RayTracingInstance.h"
+#include "Renderer/Private/ScenePrivate.h"
 
 DECLARE_CYCLE_STAT(TEXT("Generate Mesh Vertex Data [GT]"), STAT_NiagaraGenMeshVertexData, STATGROUP_Niagara);
 DECLARE_CYCLE_STAT(TEXT("Render Meshes [RT]"), STAT_NiagaraRenderMeshes, STATGROUP_Niagara);
@@ -91,7 +92,7 @@ namespace ENiagaraMeshVFLayout
 	};
 };
 
-FNiagaraRendererMeshes::FNiagaraRendererMeshes(ERHIFeatureLevel::Type FeatureLevel, const UNiagaraRendererProperties *Props, const FNiagaraEmitterInstance* Emitter)
+FNiagaraRendererMeshes::FNiagaraRendererMeshes(ERHIFeatureLevel::Type FeatureLevel, const UNiagaraRendererProperties* Props, const FNiagaraEmitterInstance* Emitter)
 	: FNiagaraRenderer(FeatureLevel, Props, Emitter)
 	, MaterialParamValidMask(0)
 {
@@ -99,10 +100,10 @@ FNiagaraRendererMeshes::FNiagaraRendererMeshes(ERHIFeatureLevel::Type FeatureLev
 	check(Props);
 
 	const UNiagaraMeshRendererProperties* Properties = CastChecked<const UNiagaraMeshRendererProperties>(Props);
-	check(Properties->ParticleMesh);
-
-	MeshRenderData = Properties->ParticleMesh->RenderData.Get();
+	UStaticMesh* Mesh = Properties->ParticleMesh;
+	check(Mesh);
 
+	MeshRenderData = Mesh->RenderData.Get();
 	FacingMode = Properties->FacingMode;
 	bLockedAxisEnable = Properties->bLockedAxisEnable;
 	LockedAxis = Properties->LockedAxis;
@@ -112,6 +113,17 @@ FNiagaraRendererMeshes::FNiagaraRendererMeshes(ERHIFeatureLevel::Type FeatureLev
 	bOverrideMaterials = Properties->bOverrideMaterials;
 	SubImageSize = Properties->SubImageSize;
 	bSubImageBlend = Properties->bSubImageBlend;
+	bEnableFrustumCulling = Properties->bEnableFrustumCulling;
+	bEnableCulling = bEnableFrustumCulling;
+	DistanceCullRange = FVector2D(0, FLT_MAX);
+	RendererVisibility = Properties->RendererVisibility;	
+	LocalCullingSphere = Mesh->ExtendedBounds.GetSphere();
+
+	if (Properties->bEnableCameraDistanceCulling)
+	{
+		DistanceCullRange = FVector2D(Properties->MinCameraDistance, Properties->MaxCameraDistance);
+		bEnableCulling = true;
+	}
 
 	// Ensure valid value for the locked axis
 	if (!LockedAxis.Normalize())
@@ -121,6 +133,15 @@ FNiagaraRendererMeshes::FNiagaraRendererMeshes(ERHIFeatureLevel::Type FeatureLev
 
 	const FNiagaraDataSet& Data = Emitter->GetData();
 
+	RendererVisTagOffset = INDEX_NONE;
+	int32 FloatOffset;
+	int32 HalfOffset;
+	if (Data.GetVariableComponentOffsets(Properties->RendererVisibilityTagBinding.DataSetVariable, FloatOffset, RendererVisTagOffset, HalfOffset))
+	{
+		// If the renderer visibility tag is bound, we have to do it in the culling pass		
+		bEnableCulling = true;
+	}
+
 	MaterialParamValidMask = 0;
 	TotalVFHalfComponents = 0;
 	TotalVFFloatComponents = 0;
@@ -173,6 +194,12 @@ void FNiagaraRendererMeshes::ReleaseRenderThreadResources()
 
 int32 FNiagaraRendererMeshes::GetMaxIndirectArgs() const
 {
+	// If we're CPU, we only need indirect args if we're culling
+	if (SimTarget == ENiagaraSimTarget::CPUSim && !bEnableCulling)
+	{
+		return 0;
+	}
+
 	// currently the most indirect args we can add would be for a single lod, so search for the LOD with the highest number of sections
 	// this value should be constant for the life of the renderer as it is being used for NumRegisteredGPURenderers
 	int32 MaxSectionCount = 0;
@@ -181,10 +208,12 @@ int32 FNiagaraRendererMeshes::GetMaxIndirectArgs() const
 	{
 		MaxSectionCount = FMath::Max(MaxSectionCount, IndexInfo.Num());
 	}
+
+	//REVERB_WIP: This needs to be multiplied by the number of active viewsv
 	return MaxSectionCount;
 }
 
-void FNiagaraRendererMeshes::SetupVertexFactory(FNiagaraMeshVertexFactory *InVertexFactory, const FStaticMeshLODResources& LODResources) const
+void FNiagaraRendererMeshes::SetupVertexFactory(FNiagaraMeshVertexFactory* InVertexFactory, const FStaticMeshLODResources& LODResources) const
 {
 	FStaticMeshDataType Data;
 
@@ -213,7 +242,7 @@ int32 FNiagaraRendererMeshes::GetLODIndex() const
 	return LODIndex;
 }
 
-void FNiagaraRendererMeshes::GetDynamicMeshElements(const TArray<const FSceneView*>& Views, const FSceneViewFamily& ViewFamily, uint32 VisibilityMap, FMeshElementCollector& Collector, const FNiagaraSceneProxy *SceneProxy) const
+void FNiagaraRendererMeshes::GetDynamicMeshElements(const TArray<const FSceneView*>& Views, const FSceneViewFamily& ViewFamily, uint32 VisibilityMap, FMeshElementCollector& Collector, const FNiagaraSceneProxy* SceneProxy) const
 {
 	check(SceneProxy);
 
@@ -222,7 +251,7 @@ void FNiagaraRendererMeshes::GetDynamicMeshElements(const TArray<const FSceneVie
 	PARTICLE_PERF_STAT_CYCLES(SceneProxy->PerfAsset, GetDynamicMeshElements);
 
 	NiagaraEmitterInstanceBatcher* Batcher = SceneProxy->GetBatcher();
-	FNiagaraDynamicDataMesh *DynamicDataMesh = (static_cast<FNiagaraDynamicDataMesh*>(DynamicDataRender));
+	FNiagaraDynamicDataMesh* DynamicDataMesh = (static_cast<FNiagaraDynamicDataMesh*>(DynamicDataRender));
 	if (!DynamicDataMesh || !Batcher)
 	{
 		return;
@@ -244,14 +273,13 @@ void FNiagaraRendererMeshes::GetDynamicMeshElements(const TArray<const FSceneVie
 	FScopeCycleCounter EmitterStatsCounter(EmitterStatID);
 #endif
 
-	// @TODO : support multiple LOD
-	TArray<uint32, TInlineAllocator<8>> IndirectArgsOffsets;
 	int32 NumInstances = SourceParticleData->GetNumInstances();
 
 	FGlobalDynamicReadBuffer& DynamicReadBuffer = Collector.GetDynamicReadBuffer();
-	FParticleRenderData ParticleData;
+	FParticleRenderData ParticleFloatData;
+	FGlobalDynamicReadBuffer::FAllocation ParticleIntData;
 
-	//Grab the material proxies we'll be using for each section and check them for translucency.
+	// Grab the material proxies we'll be using for each section and check them for translucency.
 	bool bHasTranslucentMaterials = false;
 	for (FMaterialRenderProxy* MaterialProxy : DynamicDataMesh->Materials)
 	{
@@ -260,6 +288,10 @@ void FNiagaraRendererMeshes::GetDynamicMeshElements(const TArray<const FSceneVie
 		bHasTranslucentMaterials |= IsTranslucentBlendMode(BlendMode);
 	}
 
+	// NOTE: have to run the GPU sort when culling is enabled if supported on this platform
+	// TODO: implement culling and renderer visibility on the CPU for other platforms
+	const bool bGPUSortEnabled = FNiagaraUtilities::AllowComputeShaders(Batcher->GetShaderPlatform());
+	const bool bDoGPUCulling = bEnableCulling && GNiagaraGPUCulling && FNiagaraUtilities::AllowComputeShaders(Batcher->GetShaderPlatform());
 	const bool bShouldSort = SortMode != ENiagaraSortMode::None && (bHasTranslucentMaterials || !bSortOnlyWhenTranslucent);
 	const bool bCustomSorting = SortMode == ENiagaraSortMode::CustomAscending || SortMode == ENiagaraSortMode::CustomDecending;
 	//Disable the upload of sorting data if we're using a material that doesn't need it.
@@ -271,47 +303,68 @@ void FNiagaraRendererMeshes::GetDynamicMeshElements(const TArray<const FSceneVie
 	{
 		if (GbEnableMinimalGPUBuffers)
 		{
-			ParticleData = TransferDataToGPU(DynamicReadBuffer, SourceParticleData);
+			ParticleFloatData = TransferDataToGPU(DynamicReadBuffer, SourceParticleData);
 		}
 		else
 		{
 			SCOPE_CYCLE_COUNTER(STAT_NiagaraRenderMeshes_AllocateGPUData);
 			int32 TotalFloatSize = SourceParticleData->GetFloatBuffer().Num() / sizeof(float);
-			ParticleData.FloatData = DynamicReadBuffer.AllocateFloat(TotalFloatSize);
-			FMemory::Memcpy(ParticleData.FloatData.Buffer, SourceParticleData->GetFloatBuffer().GetData(), SourceParticleData->GetFloatBuffer().Num());
+			ParticleFloatData.FloatData = DynamicReadBuffer.AllocateFloat(TotalFloatSize);
+			FMemory::Memcpy(ParticleFloatData.FloatData.Buffer, SourceParticleData->GetFloatBuffer().GetData(), SourceParticleData->GetFloatBuffer().Num());
 			int32 TotalHalfSize = SourceParticleData->GetHalfBuffer().Num() / sizeof(FFloat16);
-			ParticleData.HalfData = DynamicReadBuffer.AllocateHalf(TotalFloatSize);
-			FMemory::Memcpy(ParticleData.HalfData.Buffer, SourceParticleData->GetHalfBuffer().GetData(), SourceParticleData->GetHalfBuffer().Num());
+			ParticleFloatData.HalfData = DynamicReadBuffer.AllocateHalf(TotalFloatSize);
+			FMemory::Memcpy(ParticleFloatData.HalfData.Buffer, SourceParticleData->GetHalfBuffer().GetData(), SourceParticleData->GetHalfBuffer().Num());
 		}
-	}
-	
-	const int32 LODIndex = GetLODIndex();
-	const FStaticMeshLODResources& LODModel = MeshRenderData->LODResources[LODIndex];
-
-	if (SimTarget == ENiagaraSimTarget::GPUComputeSim)
-	{
-		const int32 SectionCount = LODModel.Sections.Num();
 
-		IndirectArgsOffsets.SetNum(SectionCount);
-		for (int32 SectionIdx = 0; SectionIdx < SectionCount; ++SectionIdx)
+		if (RendererVisTagOffset != INDEX_NONE)
 		{
-			IndirectArgsOffsets[SectionIdx] = Batcher->GetGPUInstanceCounterManager().AddDrawIndirect(
-				SourceParticleData->GetGPUInstanceCountBufferOffset(),
-				IndexInfoPerSection[LODIndex][SectionIdx].Key,
-				IndexInfoPerSection[LODIndex][SectionIdx].Value);
+			SCOPE_CYCLE_COUNTER(STAT_NiagaraRenderMeshes_AllocateGPUData);
+
+			// For CPU sims, we need to also copy off the renderer visibility tags for the sort shader
+			ParticleIntData = DynamicReadBuffer.AllocateInt32(NumInstances);
+			int32* Dest = (int32*)ParticleIntData.Buffer;
+			const int32* Src = (const int32*)SourceParticleData->GetInt32Buffer().GetData();
+			const uint32 IntStride = SourceParticleData->GetInt32Stride() / sizeof(uint32);
+			for (int32 InstIdx = 0; InstIdx < NumInstances; ++InstIdx)
+			{
+				Dest[InstIdx] = Src[RendererVisTagOffset * IntStride + InstIdx];
+			}
 		}
 	}
+	
+	// @TODO : support multiple LOD	
+	const int32 LODIndex = GetLODIndex();
+	const FStaticMeshLODResources& LODModel = MeshRenderData->LODResources[LODIndex];
+	const int32 SectionCount = LODModel.Sections.Num();
 
 	{
 		int32 VertexFactoryIndex = 0;
 
 		// Compute the per-view uniform buffers.
-		for (int32 ViewIndex = 0; ViewIndex < Views.Num(); ViewIndex++)
+		const int32 NumViews = Views.Num();
+		for (int32 ViewIndex = 0; ViewIndex < NumViews; ViewIndex++)
 		{
 			if (VisibilityMap & (1 << ViewIndex))
 			{
 				const FSceneView* View = Views[ViewIndex];
 
+				int32 CulledGPUCountOffset = bDoGPUCulling ? Batcher->GetGPUInstanceCounterManager().AcquireCulledEntry() : INDEX_NONE;
+
+				// Alloc indirect draw counts for every section this view
+				TArray<uint32, TInlineAllocator<8>> IndirectArgsOffsets;
+				if (SimTarget == ENiagaraSimTarget::GPUComputeSim || bDoGPUCulling)
+				{
+					IndirectArgsOffsets.SetNum(SectionCount);
+					for (int32 SectionIdx = 0; SectionIdx < SectionCount; ++SectionIdx)
+					{
+						IndirectArgsOffsets[SectionIdx] = Batcher->GetGPUInstanceCounterManager().AddDrawIndirect(
+							bDoGPUCulling ? CulledGPUCountOffset : SourceParticleData->GetGPUInstanceCountBufferOffset(),
+							IndexInfoPerSection[LODIndex][SectionIdx].Key,
+							IndexInfoPerSection[LODIndex][SectionIdx].Value,
+							bDoGPUCulling);
+					}
+				}
+
 				// Get the next vertex factory to use
 				FNiagaraMeshVertexFactory* VertexFactory = nullptr;
 				if (VertexFactories.IsValidIndex(VertexFactoryIndex))
@@ -375,49 +428,110 @@ void FNiagaraRendererMeshes::GetDynamicMeshElements(const TArray<const FSceneVie
 				CollectorResources.VertexFactory->SetSortedIndices(nullptr, 0xFFFFFFFF);
 
 				FNiagaraGPUSortInfo SortInfo;
-				const int32 SortVarIdx = bCustomSorting ? ENiagaraMeshVFLayout::CustomSorting : ENiagaraMeshVFLayout::Position;
-				const bool bSortVarIsHalf = bShouldSort && VFVariables[SortVarIdx].bHalfType;
-				if (bShouldSort && VFVariables[SortVarIdx].GetGPUOffset() != INDEX_NONE)
+				int32 SortVarIdx = INDEX_NONE;
+				if (bShouldSort || bDoGPUCulling)
 				{
 					SortInfo.ParticleCount = NumInstances;
 					SortInfo.SortMode = SortMode;
-					SortInfo.SortAttributeOffset = (VFVariables[SortVarIdx].GetGPUOffset() & ~(1 << 31));
 					SortInfo.SetSortFlags(GNiagaraGPUSortingUseMaxPrecision != 0, bHasTranslucentMaterials); 
-					if (!bCustomSorting)
+					SortInfo.bEnableCulling = bDoGPUCulling;
+					SortInfo.LocalBSphere = LocalCullingSphere;
+					SortInfo.RendererVisTagAttributeOffset = RendererVisTagOffset;
+					SortInfo.RendererVisibility = RendererVisibility;
+					SortInfo.DistanceCullRange = DistanceCullRange;
+
+					SortVarIdx = bCustomSorting ? ENiagaraMeshVFLayout::CustomSorting : ENiagaraMeshVFLayout::Position;
+					SortInfo.SortAttributeOffset = VFVariables[SortVarIdx].GetGPUOffset();
+
+#if !(UE_BUILD_SHIPPING || UE_BUILD_TEST)
+					const FSceneViewState* ViewState = View->State != nullptr ? View->State->GetConcreteViewState() : nullptr;
+					if (ViewState && ViewState->bIsFrozen && ViewState->bIsFrozenViewMatricesCached)
+					{
+						// Use the frozen view for culling so we can test that it's working
+						const FViewMatrices& Matrices = ViewState->CachedViewMatrices;
+						SortInfo.ViewOrigin = Matrices.GetViewOrigin();
+						SortInfo.ViewDirection = Matrices.GetViewMatrix().GetColumn(2);
+
+						if (bEnableFrustumCulling)
+						{
+							const FMatrix& ViewProj = Matrices.GetViewProjectionMatrix();
+							SortInfo.CullPlanes.SetNumZeroed(4);
+							ViewProj.GetFrustumLeftPlane(SortInfo.CullPlanes[0]);
+							ViewProj.GetFrustumTopPlane(SortInfo.CullPlanes[1]);
+							ViewProj.GetFrustumRightPlane(SortInfo.CullPlanes[2]);
+							ViewProj.GetFrustumBottomPlane(SortInfo.CullPlanes[3]);							
+						}
+					}
+					else
+#endif
 					{
 						SortInfo.ViewOrigin = View->ViewMatrices.GetViewOrigin();
 						SortInfo.ViewDirection = View->GetViewDirection();
-						if (bLocalSpace)
+						if (bEnableFrustumCulling)
 						{
-							SortInfo.ViewOrigin = SceneProxy->GetLocalToWorldInverse().TransformPosition(SortInfo.ViewOrigin);
-							SortInfo.ViewDirection = SceneProxy->GetLocalToWorld().GetTransposed().TransformVector(SortInfo.ViewDirection);
+							SortInfo.CullPlanes = View->ViewFrustum.Planes;
 						}
 					}
-				};
 
-				if (SimTarget == ENiagaraSimTarget::CPUSim)//TODO: Compute shader for sorting gpu sims and larger cpu sims.
+					if (bLocalSpace)
+					{
+						SortInfo.ViewOrigin = SceneProxy->GetLocalToWorldInverse().TransformPosition(SortInfo.ViewOrigin);
+						SortInfo.ViewDirection = SceneProxy->GetLocalToWorld().GetTransposed().TransformVector(SortInfo.ViewDirection);
+						if (bEnableFrustumCulling)
+						{
+							for (FPlane& Plane : SortInfo.CullPlanes)
+							{
+								Plane = Plane.TransformBy(SceneProxy->GetLocalToWorldInverse());
+							}
+						}
+					}
+
+					if (bDoGPUCulling)
+					{
+						SortInfo.CullPositionAttributeOffset = VFVariables[ENiagaraMeshVFLayout::Position].GetGPUOffset();
+						SortInfo.CullOrientationAttributeOffset = VFVariables[ENiagaraMeshVFLayout::Transform].GetGPUOffset();
+						SortInfo.CullScaleAttributeOffset = VFVariables[ENiagaraMeshVFLayout::Scale].GetGPUOffset();
+					}
+				}
+
+				if (SimTarget == ENiagaraSimTarget::CPUSim)
 				{
-					FRHIShaderResourceView* FloatSRV = ParticleData.FloatData.IsValid() ? ParticleData.FloatData.SRV : (FRHIShaderResourceView*)FNiagaraRenderer::GetDummyFloatBuffer();
-					FRHIShaderResourceView* HalfSRV = ParticleData.HalfData.IsValid() ? ParticleData.HalfData.SRV : (FRHIShaderResourceView*)FNiagaraRenderer::GetDummyHalfBuffer();
-					const uint32 FloatParticleDataStride = GbEnableMinimalGPUBuffers ? SourceParticleData->GetNumInstances() : (SourceParticleData->GetFloatStride() / sizeof(float));
-					check(GbEnableMinimalGPUBuffers || FloatParticleDataStride == SourceParticleData->GetHalfStride() / sizeof(FFloat16));
+					check(RendererVisTagOffset == INDEX_NONE || ParticleIntData.IsValid());
 
-					if (SortInfo.SortMode != ENiagaraSortMode::None && SortInfo.SortAttributeOffset != INDEX_NONE)
+					FRHIShaderResourceView* FloatSRV = ParticleFloatData.FloatData.IsValid() ? (FRHIShaderResourceView*)ParticleFloatData.FloatData.SRV : FNiagaraRenderer::GetDummyFloatBuffer();
+					FRHIShaderResourceView* HalfSRV = ParticleFloatData.HalfData.IsValid() ? (FRHIShaderResourceView*)ParticleFloatData.HalfData.SRV : FNiagaraRenderer::GetDummyHalfBuffer();
+					FRHIShaderResourceView* IntSRV = ParticleIntData.IsValid() ? (FRHIShaderResourceView*)ParticleIntData.SRV : FNiagaraRenderer::GetDummyIntBuffer();
+					const uint32 FloatParticleDataStride = GbEnableMinimalGPUBuffers ? NumInstances : (SourceParticleData->GetFloatStride() / sizeof(float));
+					const uint32 HalfParticleDataStride = GbEnableMinimalGPUBuffers ? NumInstances : (SourceParticleData->GetHalfStride() / sizeof(FFloat16));
+					const uint32 IntParticleDataStride = ParticleIntData.IsValid() ? NumInstances : 0; // because we copied it off
+
+					if (bShouldSort || bDoGPUCulling)
 					{
-						if (GNiagaraGPUSortingCPUToGPUThreshold >= 0 &&
-							SortInfo.ParticleCount >= GNiagaraGPUSortingCPUToGPUThreshold &&
-							FNiagaraUtilities::AllowComputeShaders(Batcher->GetShaderPlatform()))
+						const int32 Threshold = GNiagaraGPUSortingCPUToGPUThreshold;
+						if (bDoGPUCulling || (bGPUSortEnabled && Threshold >= 0 && SortInfo.ParticleCount > Threshold))
 						{
+							// We want to sort or cull on GPU
 							SortInfo.ParticleCount = NumInstances;
-							SortInfo.ParticleDataFloatSRV = bSortVarIsHalf ? HalfSRV : FloatSRV;
-							SortInfo.FloatDataStride = FloatParticleDataStride; // Same if Halfs.
-							if (Batcher->AddSortedGPUSimulation(SortInfo))
+							SortInfo.ParticleDataFloatSRV = FloatSRV;
+							SortInfo.ParticleDataHalfSRV = HalfSRV;
+							SortInfo.ParticleDataIntSRV = IntSRV;
+							SortInfo.FloatDataStride = FloatParticleDataStride;
+							SortInfo.HalfDataStride = HalfParticleDataStride;
+							SortInfo.IntDataStride = IntParticleDataStride;
+							SortInfo.GPUParticleCountSRV = Batcher->GetGPUInstanceCounterManager().GetInstanceCountBuffer().SRV;
+							SortInfo.GPUParticleCountOffset = SourceParticleData->GetGPUInstanceCountBufferOffset();
+							SortInfo.CulledGPUParticleCountOffset = CulledGPUCountOffset;
+							SortInfo.RendererVisTagAttributeOffset = RendererVisTagOffset == INDEX_NONE ? INDEX_NONE : 0; // because it's copied off							
+
+							const int32 IndexBufferOffset = Batcher->AddSortedGPUSimulation(SortInfo);
+							if (IndexBufferOffset != INDEX_NONE)
 							{
 								CollectorResources.VertexFactory->SetSortedIndices(SortInfo.AllocationInfo.BufferSRV, SortInfo.AllocationInfo.BufferOffset);
 							}
 						}
 						else
 						{
+							// We want to sort on CPU
 							FGlobalDynamicReadBuffer::FAllocation SortedIndices;
 							SortedIndices = DynamicReadBuffer.AllocateInt32(NumInstances);
 							SortIndices(SortInfo, SortVarIdx, *SourceParticleData, SortedIndices);
@@ -431,20 +545,30 @@ void FNiagaraRendererMeshes::GetDynamicMeshElements(const TArray<const FSceneVie
 				}
 				else
 				{
-					FRHIShaderResourceView* FloatSRV = SourceParticleData->GetGPUBufferFloat().SRV.IsValid() ? (FRHIShaderResourceView*)SourceParticleData->GetGPUBufferFloat().SRV : (FRHIShaderResourceView*)FNiagaraRenderer::GetDummyFloatBuffer();
-					FRHIShaderResourceView* HalfSRV = SourceParticleData->GetGPUBufferHalf().SRV.IsValid() ? (FRHIShaderResourceView*)SourceParticleData->GetGPUBufferHalf().SRV : (FRHIShaderResourceView*)FNiagaraRenderer::GetDummyHalfBuffer();
+					FRHIShaderResourceView* FloatSRV = SourceParticleData->GetGPUBufferFloat().SRV.IsValid() ? (FRHIShaderResourceView*)SourceParticleData->GetGPUBufferFloat().SRV : FNiagaraRenderer::GetDummyFloatBuffer();
+					FRHIShaderResourceView* HalfSRV = SourceParticleData->GetGPUBufferHalf().SRV.IsValid() ? (FRHIShaderResourceView*)SourceParticleData->GetGPUBufferHalf().SRV : FNiagaraRenderer::GetDummyHalfBuffer();
+					FRHIShaderResourceView* IntSRV = SourceParticleData->GetGPUBufferInt().SRV.IsValid() ? (FRHIShaderResourceView*)SourceParticleData->GetGPUBufferInt().SRV : FNiagaraRenderer::GetDummyIntBuffer();
 					const uint32 FloatParticleDataStride = SourceParticleData->GetFloatStride() / sizeof(float);
-					check(FloatParticleDataStride == SourceParticleData->GetHalfStride() / sizeof(FFloat16));
+					const uint32 HalfParticleDataStride = SourceParticleData->GetHalfStride() / sizeof(FFloat16);
+					const uint32 IntParticleDataStride = SourceParticleData->GetInt32Stride() / sizeof(int32);
 
-					if (SortInfo.SortMode != ENiagaraSortMode::None && SortInfo.SortAttributeOffset != INDEX_NONE)
+					if (bShouldSort || bDoGPUCulling)
 					{
 						// Here we need to be conservative about the InstanceCount, since the final value is only known on the GPU after the simulation.
 						SortInfo.ParticleCount = SourceParticleData->GetNumInstances();
-						SortInfo.ParticleDataFloatSRV = bSortVarIsHalf ? HalfSRV : FloatSRV;
-						SortInfo.FloatDataStride = FloatParticleDataStride; // Same if Halfs.
+						SortInfo.ParticleDataFloatSRV = FloatSRV;
+						SortInfo.ParticleDataHalfSRV = HalfSRV;
+						SortInfo.ParticleDataIntSRV = IntSRV;
+						SortInfo.FloatDataStride = FloatParticleDataStride;
+						SortInfo.HalfDataStride = HalfParticleDataStride;
+						SortInfo.IntDataStride = IntParticleDataStride;
 						SortInfo.GPUParticleCountSRV = Batcher->GetGPUInstanceCounterManager().GetInstanceCountBuffer().SRV;
 						SortInfo.GPUParticleCountOffset = SourceParticleData->GetGPUInstanceCountBufferOffset();
-						if (Batcher->AddSortedGPUSimulation(SortInfo))
+						SortInfo.CulledGPUParticleCountOffset = CulledGPUCountOffset;
+						SortInfo.RendererVisTagAttributeOffset = RendererVisTagOffset;
+
+						const int32 IndexBufferOffset = Batcher->AddSortedGPUSimulation(SortInfo);
+						if (IndexBufferOffset != INDEX_NONE && SortInfo.GPUParticleCountOffset != INDEX_NONE)
 						{
 							CollectorResources.VertexFactory->SetSortedIndices(SortInfo.AllocationInfo.BufferSRV, SortInfo.AllocationInfo.BufferOffset);
 						}
@@ -463,9 +587,6 @@ void FNiagaraRendererMeshes::GetDynamicMeshElements(const TArray<const FSceneVie
 				INC_DWORD_STAT_BY(STAT_NiagaraNumMeshVerts, NumInstances * LODModel.GetNumVertices());
 				INC_DWORD_STAT_BY(STAT_NiagaraNumMeshes, NumInstances);
 
-				// GPU mesh rendering currently only supports one mesh section.
-				// TODO: Add proper support for multiple mesh sections for GPU mesh particles.
-				const int32 SectionCount = LODModel.Sections.Num();
 				const bool bIsWireframe = AllowDebugViewmodes() && View && View->Family->EngineShowFlags.Wireframe;
 				for (int32 SectionIndex = 0; SectionIndex < SectionCount; SectionIndex++)
 				{
@@ -486,6 +607,7 @@ void FNiagaraRendererMeshes::GetDynamicMeshElements(const TArray<const FSceneVie
 					Mesh.CastRayTracedShadow = SceneProxy->CastsDynamicShadow();
 #endif
 					Mesh.DepthPriorityGroup = (ESceneDepthPriorityGroup)SceneProxy->GetDepthPriorityGroup(View);
+					Mesh.bUseSelectionOutline = !bDoGPUCulling; //TODO: selection outline can't use the culled instance index buffer, so it causes bad flickering
 
 					FMeshBatchElement& BatchElement = Mesh.Elements[0];
 					BatchElement.PrimitiveUniformBuffer = IsMotionBlurEnabled() ? SceneProxy->GetUniformBuffer() : SceneProxy->GetUniformBufferNoVelocity();
@@ -503,7 +625,6 @@ void FNiagaraRendererMeshes::GetDynamicMeshElements(const TArray<const FSceneVie
 							BatchElement.FirstIndex = 0;
 							BatchElement.IndexBuffer = &LODModel.AdditionalIndexBuffers->WireframeIndexBuffer;
 							BatchElement.NumPrimitives = LODModel.AdditionalIndexBuffers->WireframeIndexBuffer.GetNumIndices() / 2;
-
 						}
 						else
 						{
@@ -559,7 +680,7 @@ void FNiagaraRendererMeshes::GetDynamicRayTracingInstances(FRayTracingMaterialGa
 	check(SceneProxy);
 
 	NiagaraEmitterInstanceBatcher* Batcher = SceneProxy->GetBatcher();
-	FNiagaraDynamicDataMesh *DynamicDataMesh = (static_cast<FNiagaraDynamicDataMesh*>(DynamicDataRender));
+	FNiagaraDynamicDataMesh* DynamicDataMesh = (static_cast<FNiagaraDynamicDataMesh*>(DynamicDataRender));
 	if (!DynamicDataMesh || !Batcher)
 	{
 		return;
@@ -627,13 +748,13 @@ void FNiagaraRendererMeshes::GetDynamicRayTracingInstances(FRayTracingMaterialGa
 	int32 ComponentStrideDest = SourceParticleData->GetNumInstances() * sizeof(float);
 
 	//ENiagaraMeshVFLayout::Transform just contains a Quat, not the whole transform
-	FNiagaraRendererVariableInfo& VarPositionInfo	= VFVariables[ENiagaraMeshVFLayout::Position];
-	FNiagaraRendererVariableInfo& VarScaleInfo		= VFVariables[ENiagaraMeshVFLayout::Scale];
-	FNiagaraRendererVariableInfo& VarTransformInfo	= VFVariables[ENiagaraMeshVFLayout::Transform];
+	FNiagaraRendererVariableInfo& VarPositionInfo = VFVariables[ENiagaraMeshVFLayout::Position];
+	FNiagaraRendererVariableInfo& VarScaleInfo = VFVariables[ENiagaraMeshVFLayout::Scale];
+	FNiagaraRendererVariableInfo& VarTransformInfo = VFVariables[ENiagaraMeshVFLayout::Transform];
 	
-	int32 PositionBaseCompOffset	= VarPositionInfo.DatasetOffset;
-	int32 ScaleBaseCompOffset		= VarScaleInfo.DatasetOffset;
-	int32 TransformBaseCompOffset	= VarTransformInfo.DatasetOffset;
+	int32 PositionBaseCompOffset = VarPositionInfo.DatasetOffset;
+	int32 ScaleBaseCompOffset = VarScaleInfo.DatasetOffset;
+	int32 TransformBaseCompOffset = VarTransformInfo.DatasetOffset;
 	
 	float* RESTRICT PositionX = (float*)SourceParticleData->GetComponentPtrFloat(PositionBaseCompOffset);
 	float* RESTRICT PositionY = (float*)SourceParticleData->GetComponentPtrFloat(PositionBaseCompOffset + 1);
@@ -748,7 +869,7 @@ void FNiagaraRendererMeshes::GetDynamicRayTracingInstances(FRayTracingMaterialGa
 
 		FRWBufferStructured InstanceGPUTransformsBuffer;
 		//InstanceGPUTransformsBuffer.Initialize(sizeof(FMatrix), CPUInstancesCount, BUF_Static);
-		InstanceGPUTransformsBuffer.Initialize(3*4*sizeof(float), CPUInstancesCount, BUF_Static);
+		InstanceGPUTransformsBuffer.Initialize(3 * 4 * sizeof(float), CPUInstancesCount, BUF_Static);
 		RayTracingInstance.InstanceGPUTransformsSRV = InstanceGPUTransformsBuffer.SRV;
 
 		FNiagaraGPURayTracingTransformsCS::FPermutationDomain PermutationVector;
@@ -760,7 +881,7 @@ void FNiagaraRendererMeshes::GetDynamicRayTracingInstances(FRayTracingMaterialGa
 			VFVariables[ENiagaraMeshVFLayout::Position].GetGPUOffset(), 
 			VFVariables[ENiagaraMeshVFLayout::Transform].GetGPUOffset(),
 			VFVariables[ENiagaraMeshVFLayout::Scale].GetGPUOffset(),
-			bLocalSpace? 1 : 0);
+			bLocalSpace ? 1 : 0);
 
 		uint32 FloatDataOffset = 0;
 		uint32 FloatDataStride = SourceParticleData->GetFloatStride() / sizeof(float);
@@ -790,7 +911,7 @@ void FNiagaraRendererMeshes::GetDynamicRayTracingInstances(FRayTracingMaterialGa
 #endif
 
 
-FNiagaraDynamicDataBase *FNiagaraRendererMeshes::GenerateDynamicData(const FNiagaraSceneProxy* Proxy, const UNiagaraRendererProperties* InProperties, const FNiagaraEmitterInstance* Emitter) const
+FNiagaraDynamicDataBase* FNiagaraRendererMeshes::GenerateDynamicData(const FNiagaraSceneProxy* Proxy, const UNiagaraRendererProperties* InProperties, const FNiagaraEmitterInstance* Emitter) const
 {
 	SCOPE_CYCLE_COUNTER(STAT_NiagaraRenderGT);
 	SCOPE_CYCLE_COUNTER(STAT_NiagaraGenMeshVertexData);
@@ -802,7 +923,7 @@ FNiagaraDynamicDataBase *FNiagaraRendererMeshes::GenerateDynamicData(const FNiag
 		return nullptr;
 	}
 
-	FNiagaraDynamicDataMesh *DynamicData = nullptr;
+	FNiagaraDynamicDataMesh* DynamicData = nullptr;
 
 	if (Properties->ParticleMesh)
 	{
diff --git a/Engine/Plugins/FX/Niagara/Source/Niagara/Public/NiagaraMeshRendererProperties.h b/Engine/Plugins/FX/Niagara/Source/Niagara/Public/NiagaraMeshRendererProperties.h
--- a/Engine/Plugins/FX/Niagara/Source/Niagara/Public/NiagaraMeshRendererProperties.h
+++ b/Engine/Plugins/FX/Niagara/Source/Niagara/Public/NiagaraMeshRendererProperties.h
@@ -135,7 +135,7 @@ public:
 	UPROPERTY(EditAnywhere, Category = "SubUV", meta = (DisplayName = "Sub UV Blending Enabled"))
 	uint32 bSubImageBlend : 1;
 
-	/** Determines how the mesh orients itself relative to the camera.*/
+	/** Determines how the mesh orients itself relative to the camera. */
 	UPROPERTY(EditAnywhere, Category = "Mesh Rendering")
 	ENiagaraMeshFacingMode FacingMode;
 
@@ -151,6 +151,24 @@ public:
 	UPROPERTY(EditAnywhere, Category = "Mesh Rendering", meta = (EditCondition = "bLockedAxisEnable"))
 	ENiagaraMeshLockedAxisSpace LockedAxisSpace;
 	
+	/** Enables frustum culling of individual mesh particles */
+	UPROPERTY(EditAnywhere, Category = "Visibility")
+	uint32 bEnableFrustumCulling : 1;
+
+	/** Enables frustum culling of individual mesh particles */
+	UPROPERTY(EditAnywhere, Category = "Visibility")
+	uint32 bEnableCameraDistanceCulling : 1;
+
+	UPROPERTY(EditAnywhere, Category = "Visibility", meta = (EditCondition = "bEnableCameraDistanceCulling", UIMin = 0.0f))
+	float MinCameraDistance;
+	
+	UPROPERTY(EditAnywhere, Category = "Visibility", meta = (EditCondition = "bEnableCameraDistanceCulling", UIMin = 0.0f))
+	float MaxCameraDistance = 1000.0f;
+
+	/** If a render visibility tag is present, particles whose tag matches this value will be visible in this renderer. */
+	UPROPERTY(EditAnywhere, Category = "Visibility")
+	uint32 RendererVisibility = 0;
+	
 	/** Which attribute should we use for position when generating instanced meshes?*/
 	UPROPERTY(EditAnywhere, Category = "Bindings")
 	FNiagaraVariableAttributeBinding PositionBinding;
@@ -207,6 +225,9 @@ public:
 	UPROPERTY(EditAnywhere, Category = "Bindings")
 	FNiagaraVariableAttributeBinding CameraOffsetBinding;
 
+	/** Which attribute should we use for the renderer visibility tag? */
+	UPROPERTY(EditAnywhere, Category = "Bindings")
+	FNiagaraVariableAttributeBinding RendererVisibilityTagBinding;
 
 protected:
 	bool FindBinding(const FNiagaraUserParameterBinding& InBinding, const FNiagaraEmitterInstance* InEmitter, TArray<UMaterialInterface*>& OutMaterials);
diff --git a/Engine/Plugins/FX/Niagara/Source/Niagara/Public/NiagaraModule.h b/Engine/Plugins/FX/Niagara/Source/Niagara/Public/NiagaraModule.h
--- a/Engine/Plugins/FX/Niagara/Source/Niagara/Public/NiagaraModule.h
+++ b/Engine/Plugins/FX/Niagara/Source/Niagara/Public/NiagaraModule.h
@@ -159,6 +159,7 @@ public:
 	FORCEINLINE static const FNiagaraVariable&  GetVar_Particles_RibbonTwist() { return Particles_RibbonTwist; }
 	FORCEINLINE static const FNiagaraVariable&  GetVar_Particles_RibbonFacing() { return Particles_RibbonFacing; }
 	FORCEINLINE static const FNiagaraVariable&  GetVar_Particles_RibbonLinkOrder() { return Particles_RibbonLinkOrder; }
+	FORCEINLINE static const FNiagaraVariable&  GetVar_Particles_VisibilityTag() { return Particles_VisibilityTag; }
 	
 	FORCEINLINE static const FNiagaraVariable&  GetVar_DataInstance_Alive() { return DataInstance_Alive; }
 	FORCEINLINE static const FNiagaraVariable&  GetVar_BeginDefaults() { return Translator_BeginDefaults; }
@@ -246,6 +247,7 @@ private:
 	static FNiagaraVariable Particles_Scale;
 	static FNiagaraVariable Particles_Lifetime;
 	static FNiagaraVariable Particles_MeshOrientation;
+	static FNiagaraVariable Particles_VisibilityTag;
 	static FNiagaraVariable Particles_UVScale;
 	static FNiagaraVariable Particles_CameraOffset;
 	static FNiagaraVariable Particles_MaterialRandom;
diff --git a/Engine/Plugins/FX/Niagara/Source/Niagara/Public/NiagaraRenderer.h b/Engine/Plugins/FX/Niagara/Source/Niagara/Public/NiagaraRenderer.h
--- a/Engine/Plugins/FX/Niagara/Source/Niagara/Public/NiagaraRenderer.h
+++ b/Engine/Plugins/FX/Niagara/Source/Niagara/Public/NiagaraRenderer.h
@@ -163,7 +163,7 @@ protected:
 	TStatId EmitterStatID;
 #endif
 
-	virtual int32 GetMaxIndirectArgs() const { return 1; }
+	virtual int32 GetMaxIndirectArgs() const { return SimTarget == ENiagaraSimTarget::GPUComputeSim ? 1 : 0; }
 
 	bool SetVertexFactoryVariable(const FNiagaraDataSet& DataSet, const FNiagaraVariable& Var, int32 VFVarOffset);
 	FParticleRenderData TransferDataToGPU(FGlobalDynamicReadBuffer& DynamicReadBuffer, FNiagaraDataBuffer* SrcData)const;
diff --git a/Engine/Plugins/FX/Niagara/Source/Niagara/Public/NiagaraRendererMeshes.h b/Engine/Plugins/FX/Niagara/Source/Niagara/Public/NiagaraRendererMeshes.h
--- a/Engine/Plugins/FX/Niagara/Source/Niagara/Public/NiagaraRendererMeshes.h
+++ b/Engine/Plugins/FX/Niagara/Source/Niagara/Public/NiagaraRendererMeshes.h
@@ -48,6 +48,8 @@ private:
 	uint32 bOverrideMaterials : 1;
 	uint32 bSortOnlyWhenTranslucent : 1;
 	uint32 bLockedAxisEnable : 1;
+	uint32 bEnableCulling : 1;
+	uint32 bEnableFrustumCulling : 1;
 
 	uint32 bSubImageBlend : 1;
 	FVector2D SubImageSize;
@@ -55,6 +57,10 @@ private:
 	FVector LockedAxis;
 	ENiagaraMeshLockedAxisSpace LockedAxisSpace;
 
+	FSphere LocalCullingSphere;
+	FVector2D DistanceCullRange;
+	int32 RendererVisTagOffset;
+	int32 RendererVisibility;
 	uint32 MaterialParamValidMask;
 
 	int32 MeshMinimumLOD = 0;
diff --git a/Engine/Plugins/FX/Niagara/Source/NiagaraVertexFactories/Private/NiagaraDrawIndirect.cpp b/Engine/Plugins/FX/Niagara/Source/NiagaraVertexFactories/Private/NiagaraDrawIndirect.cpp
--- a/Engine/Plugins/FX/Niagara/Source/NiagaraVertexFactories/Private/NiagaraDrawIndirect.cpp
+++ b/Engine/Plugins/FX/Niagara/Source/NiagaraVertexFactories/Private/NiagaraDrawIndirect.cpp
@@ -23,6 +23,7 @@ FNiagaraDrawIndirectArgsGenCS::FNiagaraDrawIndirectArgsGenCS(const ShaderMetaTyp
 	: FGlobalShader(Initializer)
 {
 	TaskInfosParam.Bind(Initializer.ParameterMap, TEXT("TaskInfos"));
+	CulledInstanceCountsParam.Bind(Initializer.ParameterMap, TEXT("CulledInstanceCounts"));
 	InstanceCountsParam.Bind(Initializer.ParameterMap, TEXT("InstanceCounts"));
 	DrawIndirectArgsParam.Bind(Initializer.ParameterMap, TEXT("DrawIndirectArgs"));
 	TaskCountParam.Bind(Initializer.ParameterMap, TEXT("TaskCount"));
@@ -32,6 +33,7 @@ FNiagaraDrawIndirectArgsGenCS::FNiagaraDrawIndirectArgsGenCS(const ShaderMetaTyp
 {
 	bool bShaderHasOutdatedParameters = FGlobalShader::Serialize(Ar);
 	Ar << TaskInfosParam;
+	Ar << CulledInstanceCountsParam;
 	Ar << InstanceCountsParam;
 	Ar << DrawIndirectArgsParam;
 	Ar << TaskCountParam;
@@ -51,11 +53,12 @@ void FNiagaraDrawIndirectArgsGenCS::SetOutput(FRHICommandList& RHICmdList, FRHIU
 	}
 }
 
-void FNiagaraDrawIndirectArgsGenCS::SetParameters(FRHICommandList& RHICmdList, FRHIShaderResourceView* TaskInfosBuffer, int32 NumArgGenTasks, int32 NumInstanceCountClearTasks)
+void FNiagaraDrawIndirectArgsGenCS::SetParameters(FRHICommandList& RHICmdList, FRHIShaderResourceView* TaskInfosBuffer, FRHIShaderResourceView* CulledInstanceCountsBuffer, int32 NumArgGenTasks, int32 NumInstanceCountClearTasks)
 {
 	FRHIComputeShader* ComputeShaderRHI = RHICmdList.GetBoundComputeShader();
 
 	SetSRVParameter(RHICmdList, ComputeShaderRHI, TaskInfosParam, TaskInfosBuffer);
+	SetSRVParameter(RHICmdList, ComputeShaderRHI, CulledInstanceCountsParam, CulledInstanceCountsBuffer);
 
 	const FUintVector4 TaskCountValue((int32)NumArgGenTasks, (int32)NumInstanceCountClearTasks, (int32)(NumArgGenTasks + NumInstanceCountClearTasks), 0);
 	SetShaderValue(RHICmdList, ComputeShaderRHI, TaskCountParam, TaskCountValue);
@@ -75,7 +78,6 @@ void FNiagaraDrawIndirectArgsGenCS::UnbindBuffers(FRHICommandList& RHICmdList)
 		RHICmdList.SetUAVParameter(ComputeShaderRHI, InstanceCountsParam.GetUAVIndex(), nullptr);
 	}
 }
-
 void FNiagaraDrawIndirectResetCountsCS::ModifyCompilationEnvironment(const FGlobalShaderPermutationParameters& Parameters, FShaderCompilerEnvironment& OutEnvironment)
 {
 	FGlobalShader::ModifyCompilationEnvironment(Parameters, OutEnvironment);
diff --git a/Engine/Plugins/FX/Niagara/Source/NiagaraVertexFactories/Private/NiagaraSortingGPU.cpp b/Engine/Plugins/FX/Niagara/Source/NiagaraVertexFactories/Private/NiagaraSortingGPU.cpp
--- a/Engine/Plugins/FX/Niagara/Source/NiagaraVertexFactories/Private/NiagaraSortingGPU.cpp
+++ b/Engine/Plugins/FX/Niagara/Source/NiagaraVertexFactories/Private/NiagaraSortingGPU.cpp
@@ -8,6 +8,14 @@ NiagaraSortingGPU.cpp: Niagara sorting shaders
 #include "NiagaraGPUSortInfo.h"
 #include "ShaderParameterUtils.h"
 
+int32 GNiagaraGPUCulling = 1;
+static FAutoConsoleVariableRef CVarNiagaraGPUCulling(
+	TEXT("Niagara.GPUCulling"),
+	GNiagaraGPUCulling,
+	TEXT("Whether to frustum and camera distance cull particles on the GPU"),
+	ECVF_Default
+);
+
 int32 GNiagaraGPUSortingUseMaxPrecision = 0;
 static FAutoConsoleVariableRef CVarNiagaraGPUSortinUseMaxPrecision(
 	TEXT("Niagara.GPUSorting.UseMaxPrecision"),
@@ -30,80 +38,13 @@ void FNiagaraSortKeyGenCS::ModifyCompilationEnvironment(const FGlobalShaderPermu
 {
 	FGlobalShader::ModifyCompilationEnvironment(Parameters, OutEnvironment);
 	OutEnvironment.SetDefine(TEXT("THREAD_COUNT"), NIAGARA_KEY_GEN_THREAD_COUNT);
+	OutEnvironment.SetDefine(TEXT("SORT_NONE"), (uint8)ENiagaraSortMode::None);
 	OutEnvironment.SetDefine(TEXT("SORT_VIEW_DEPTH"), (uint8)ENiagaraSortMode::ViewDepth);
 	OutEnvironment.SetDefine(TEXT("SORT_VIEW_DISTANCE"), (uint8)ENiagaraSortMode::ViewDistance);
 	OutEnvironment.SetDefine(TEXT("SORT_CUSTOM_ASCENDING"), (uint8)ENiagaraSortMode::CustomAscending);
 	OutEnvironment.SetDefine(TEXT("SORT_CUSTOM_DESCENDING"), (uint8)ENiagaraSortMode::CustomDecending);
-}
-
-FNiagaraSortKeyGenCS::FNiagaraSortKeyGenCS(const ShaderMetaType::CompiledShaderInitializerType& Initializer)
-	: FGlobalShader(Initializer)
-{
-	NiagaraParticleDataFloat.Bind(Initializer.ParameterMap, TEXT("NiagaraParticleDataFloat"));
-	FloatDataStride.Bind(Initializer.ParameterMap, TEXT("NiagaraFloatDataStride"));
-	GPUParticleCountBuffer.Bind(Initializer.ParameterMap, TEXT("GPUParticleCountBuffer"));
-	ParticleCountParams.Bind(Initializer.ParameterMap, TEXT("ParticleCountParams"));
-	SortParams.Bind(Initializer.ParameterMap, TEXT("SortParams"));
-	SortKeyParams.Bind(Initializer.ParameterMap, TEXT("SortKeyParams"));
-	CameraPosition.Bind(Initializer.ParameterMap, TEXT("CameraPosition"));
-	CameraDirection.Bind(Initializer.ParameterMap, TEXT("CameraDirection"));
 
-	OutKeys.Bind(Initializer.ParameterMap, TEXT("OutKeys"));
-	OutParticleIndices.Bind(Initializer.ParameterMap, TEXT("OutParticleIndices"));
+	bool bUseWaveIntrinsics = FDataDrivenShaderPlatformInfo::GetSupportsWaveOperations(Parameters.Platform);
+	OutEnvironment.SetDefine(TEXT("USE_WAVE_INTRINSICS"), bUseWaveIntrinsics);
 }
 
-/*bool FNiagaraSortKeyGenCS::Serialize(FArchive& Ar)
-{
-	bool bShaderHasOutdatedParameters = FGlobalShader::Serialize(Ar);
-	Ar << NiagaraParticleDataFloat;
-	Ar << FloatDataStride;
-	Ar << GPUParticleCountBuffer;
-	Ar << ParticleCountParams;
-	Ar << SortParams;
-	Ar << SortKeyParams;
-	Ar << CameraPosition;
-	Ar << CameraDirection;
-	Ar << OutKeys;
-	Ar << OutParticleIndices;
-	return bShaderHasOutdatedParameters;
-}*/
-
-void FNiagaraSortKeyGenCS::SetOutput(FRHICommandList& RHICmdList, FRHIUnorderedAccessView* OutKeysUAV, FRHIUnorderedAccessView* OutIndicesUAV)
-{
-	FRHIComputeShader* ComputeShaderRHI = RHICmdList.GetBoundComputeShader();
-	SetUAVParameter(RHICmdList, ComputeShaderRHI, OutKeys, OutKeysUAV);
-	SetUAVParameter(RHICmdList, ComputeShaderRHI, OutParticleIndices, OutIndicesUAV);
-}
-
-void FNiagaraSortKeyGenCS::SetParameters(FRHICommandList& RHICmdList, const FNiagaraGPUSortInfo& SortInfo, uint32 EmitterKey, int32 OutputOffset, const FUintVector4& SortKeyParamsValue)
-{
-	FRHIComputeShader* ComputeShaderRHI = RHICmdList.GetBoundComputeShader();
-
-	SetSRVParameter(RHICmdList, ComputeShaderRHI, NiagaraParticleDataFloat, SortInfo.ParticleDataFloatSRV);
-	SetShaderValue(RHICmdList, ComputeShaderRHI, FloatDataStride, SortInfo.FloatDataStride);
-
-	SetSRVParameter(RHICmdList, ComputeShaderRHI, GPUParticleCountBuffer, SortInfo.GPUParticleCountSRV);
-
-	const FUintVector4 ParticleCountParamsValue(SortInfo.ParticleCount, SortInfo.GPUParticleCountOffset, 0, 0);
-	SetShaderValue(RHICmdList, ComputeShaderRHI, ParticleCountParams, ParticleCountParamsValue);
-
-	// (EmitterKey, OutputOffset, SortMode, SortAttributeOffset)
-	const FUintVector4 SortParamsValue(EmitterKey, OutputOffset, (uint8)SortInfo.SortMode, SortInfo.SortAttributeOffset);
-	SetShaderValue(RHICmdList, ComputeShaderRHI, SortParams, SortParamsValue);
-
-	// SortKeyParams only exists in the permutation with SORT_MAX_PRECISION set.
-	SetShaderValue(RHICmdList, ComputeShaderRHI, SortKeyParams, SortKeyParamsValue);
-	SetShaderValue(RHICmdList, ComputeShaderRHI, CameraPosition, SortInfo.ViewOrigin);
-	SetShaderValue(RHICmdList, ComputeShaderRHI, CameraDirection, SortInfo.ViewDirection);
-}
-
-void FNiagaraSortKeyGenCS::UnbindBuffers(FRHICommandList& RHICmdList)
-{
-	FRHIComputeShader* ComputeShaderRHI = RHICmdList.GetBoundComputeShader();
-	SetUAVParameter(RHICmdList, ComputeShaderRHI, NiagaraParticleDataFloat, nullptr);
-	SetUAVParameter(RHICmdList, ComputeShaderRHI, GPUParticleCountBuffer, nullptr);
-	SetUAVParameter(RHICmdList, ComputeShaderRHI, OutKeys, nullptr);
-	SetUAVParameter(RHICmdList, ComputeShaderRHI, OutParticleIndices, nullptr);
-}
-
-
diff --git a/Engine/Plugins/FX/Niagara/Source/NiagaraVertexFactories/Public/NiagaraDrawIndirect.h b/Engine/Plugins/FX/Niagara/Source/NiagaraVertexFactories/Public/NiagaraDrawIndirect.h
--- a/Engine/Plugins/FX/Niagara/Source/NiagaraVertexFactories/Public/NiagaraDrawIndirect.h
+++ b/Engine/Plugins/FX/Niagara/Source/NiagaraVertexFactories/Public/NiagaraDrawIndirect.h
@@ -13,7 +13,7 @@ NiagaraDrawIndirect.h: Niagara shader to generate the draw indirect args for Nia
 
 #define NIAGARA_DRAW_INDIRECT_ARGS_GEN_THREAD_COUNT 64
 #define NIAGARA_DRAW_INDIRECT_ARGS_SIZE 5
-#define NIAGARA_DRAW_INDIRECT_TASK_INFO_SIZE 3
+#define NIAGARA_DRAW_INDIRECT_TASK_INFO_SIZE 4
 
 // #define NIAGARA_COPY_BUFFER_THREAD_COUNT 64
 // #define NIAGARA_COPY_BUFFER_BUFFER_COUNT 3
@@ -25,21 +25,24 @@ NiagaraDrawIndirect.h: Niagara shader to generate the draw indirect args for Nia
 */
 struct FNiagaraDrawIndirectArgGenTaskInfo
 {
-	FNiagaraDrawIndirectArgGenTaskInfo(uint32 InInstanceCountBufferOffset, uint32 InNumIndicesPerInstance, uint32 InStartIndexLocation)
+	FNiagaraDrawIndirectArgGenTaskInfo(uint32 InInstanceCountBufferOffset, uint32 InNumIndicesPerInstance, uint32 InStartIndexLocation, bool bInUseCulledCounts)
 		: InstanceCountBufferOffset(InInstanceCountBufferOffset)
 		, NumIndicesPerInstance(InNumIndicesPerInstance)
 		, StartIndexLocation(InStartIndexLocation)
+		, bUseCulledCounts(bInUseCulledCounts ? 1 : 0)
 	{}
 
 	uint32 InstanceCountBufferOffset;
 	uint32 NumIndicesPerInstance; // When -1 the counter needs to be reset to 0.
 	uint32 StartIndexLocation;
+	uint32 bUseCulledCounts;
 
 	bool operator==(const FNiagaraDrawIndirectArgGenTaskInfo& Rhs) const
 	{
 		return InstanceCountBufferOffset == Rhs.InstanceCountBufferOffset
 			&& NumIndicesPerInstance == Rhs.NumIndicesPerInstance
-			&& StartIndexLocation == Rhs.StartIndexLocation;
+			&& StartIndexLocation == Rhs.StartIndexLocation
+			&& bUseCulledCounts == Rhs.bUseCulledCounts;
 	}
 };
 
@@ -57,7 +60,7 @@ public:
 
 	static bool ShouldCompilePermutation(const FGlobalShaderPermutationParameters& Parameters)
 	{
-		return RHISupportsComputeShaders(Parameters.Platform);
+		return (RHISupportsComputeShaders(Parameters.Platform) && !(Parameters.Platform == EShaderPlatform::SP_METAL || Parameters.Platform == EShaderPlatform::SP_METAL_TVOS));
 	}
 
 	static void ModifyCompilationEnvironment(const FGlobalShaderPermutationParameters& Parameters, FShaderCompilerEnvironment& OutEnvironment);
@@ -67,12 +70,13 @@ public:
 
 	//bool Serialize(FArchive& Ar);
 	void SetOutput(FRHICommandList& RHICmdList, FRHIUnorderedAccessView* DrawIndirectArgsUAV, FRHIUnorderedAccessView* InstanceCountsUAV);
-	void SetParameters(FRHICommandList& RHICmdList, FRHIShaderResourceView* TaskInfosBuffer, int32 NumArgGenTasks, int32 NumInstanceCountClearTasks);
+	void SetParameters(FRHICommandList& RHICmdList, FRHIShaderResourceView* TaskInfosBuffer, FRHIShaderResourceView* CulledInstanceCountsBuffer, int32 NumArgGenTasks, int32 NumInstanceCountClearTasks);
 	void UnbindBuffers(FRHICommandList& RHICmdList);
 
 protected:
 	
 	LAYOUT_FIELD(FShaderResourceParameter, TaskInfosParam)
+	LAYOUT_FIELD(FShaderResourceParameter, CulledInstanceCountsParam);
 	LAYOUT_FIELD(FRWShaderParameter, InstanceCountsParam)
 	LAYOUT_FIELD(FRWShaderParameter, DrawIndirectArgsParam)
 	LAYOUT_FIELD(FShaderParameter, TaskCountParam)
diff --git a/Engine/Plugins/FX/Niagara/Source/NiagaraVertexFactories/Public/NiagaraSortingGPU.h b/Engine/Plugins/FX/Niagara/Source/NiagaraVertexFactories/Public/NiagaraSortingGPU.h
--- a/Engine/Plugins/FX/Niagara/Source/NiagaraVertexFactories/Public/NiagaraSortingGPU.h
+++ b/Engine/Plugins/FX/Niagara/Source/NiagaraVertexFactories/Public/NiagaraSortingGPU.h
@@ -9,15 +9,18 @@ NiagaraSortingGPU.h: Niagara sorting shaders
 #include "CoreMinimal.h"
 #include "GlobalShader.h"
 #include "ShaderPermutation.h"
+#include "ShaderParameterStruct.h"
 
 struct FNiagaraGPUSortInfo;
 
+extern NIAGARAVERTEXFACTORIES_API int32 GNiagaraGPUCulling;
 extern NIAGARAVERTEXFACTORIES_API int32 GNiagaraGPUSortingUseMaxPrecision;
 extern NIAGARAVERTEXFACTORIES_API int32 GNiagaraGPUSortingCPUToGPUThreshold;
 
 #define NIAGARA_KEY_GEN_THREAD_COUNT 64
 #define NIAGARA_COPY_BUFFER_THREAD_COUNT 64
 #define NIAGARA_COPY_BUFFER_BUFFER_COUNT 3
+#define NIAGARA_KEY_GEN_MAX_CULL_PLANES 6
 
 /**
  * Compute shader used to generate particle sort keys.
@@ -25,11 +28,48 @@ extern NIAGARAVERTEXFACTORIES_API int32 GNiagaraGPUSortingCPUToGPUThreshold;
 class NIAGARAVERTEXFACTORIES_API FNiagaraSortKeyGenCS : public FGlobalShader
 {
 	DECLARE_GLOBAL_SHADER(FNiagaraSortKeyGenCS);
+	SHADER_USE_PARAMETER_STRUCT(FNiagaraSortKeyGenCS, FGlobalShader);
 
 public:
-
+	class FEnableCulling : SHADER_PERMUTATION_BOOL("ENABLE_CULLING");
 	class FSortUsingMaxPrecision : SHADER_PERMUTATION_BOOL("SORT_MAX_PRECISION");
-	using FPermutationDomain = TShaderPermutationDomain<FSortUsingMaxPrecision>;
+	using FPermutationDomain = TShaderPermutationDomain<FEnableCulling, FSortUsingMaxPrecision>;
+
+	BEGIN_SHADER_PARAMETER_STRUCT(FParameters, NIAGARAVERTEXFACTORIES_API)
+		SHADER_PARAMETER_SRV(Buffer, NiagaraParticleDataFloat)
+		SHADER_PARAMETER_SRV(Buffer, NiagaraParticleDataHalf)
+		SHADER_PARAMETER_SRV(Buffer, NiagaraParticleDataInt)
+		SHADER_PARAMETER_SRV(Buffer, GPUParticleCountBuffer)
+
+		SHADER_PARAMETER(uint32, FloatDataStride)
+		SHADER_PARAMETER(uint32, HalfDataStride)
+		SHADER_PARAMETER(uint32, IntDataStride)
+		SHADER_PARAMETER(uint32, ParticleCount)
+		SHADER_PARAMETER(uint32, GPUParticleCountOffset)
+		SHADER_PARAMETER(uint32, CulledGPUParticleCountOffset)
+		SHADER_PARAMETER(uint32, EmitterKey)
+		SHADER_PARAMETER(uint32, OutputOffset)
+		SHADER_PARAMETER(FVector, CameraPosition)
+		SHADER_PARAMETER(FVector, CameraDirection)
+		SHADER_PARAMETER(uint32, SortMode)
+		SHADER_PARAMETER(uint32, SortAttributeOffset)
+		SHADER_PARAMETER(uint32, SortKeyMask)
+		SHADER_PARAMETER(uint32, SortKeyShift)
+		SHADER_PARAMETER(uint32, SortKeySignBit)
+		SHADER_PARAMETER(uint32, CullPositionAttributeOffset)
+		SHADER_PARAMETER(uint32, CullOrientationAttributeOffset)
+		SHADER_PARAMETER(uint32, CullScaleAttributeOffset)
+		SHADER_PARAMETER(int32, NumCullPlanes)
+		SHADER_PARAMETER_ARRAY(FVector4, CullPlanes, [NIAGARA_KEY_GEN_MAX_CULL_PLANES])
+		SHADER_PARAMETER(int32, RendererVisibility)
+		SHADER_PARAMETER(uint32, RendererVisTagAttributeOffset)
+		SHADER_PARAMETER(FVector2D, CullDistanceRangeSquared)
+		SHADER_PARAMETER(FVector4, LocalBoundingSphere)
+
+		SHADER_PARAMETER_UAV(Buffer, OutKeys)
+		SHADER_PARAMETER_UAV(Buffer, OutParticleIndices)
+		SHADER_PARAMETER_UAV(Buffer, OutCulledParticleCounts)
+	END_SHADER_PARAMETER_STRUCT()
 
 	static bool ShouldCompilePermutation(const FGlobalShaderPermutationParameters& Parameters)
 	{
@@ -37,43 +77,4 @@ public:
 	}
 
 	static void ModifyCompilationEnvironment(const FGlobalShaderPermutationParameters& Parameters, FShaderCompilerEnvironment& OutEnvironment);
-
-	/** Default constructor. */
-	FNiagaraSortKeyGenCS() {}
-
-	/** Initialization constructor. */
-	explicit FNiagaraSortKeyGenCS(const ShaderMetaType::CompiledShaderInitializerType& Initializer);
-
-
-	/**
-	 * Set output buffers for this shader.
-	 */
-	void SetOutput(FRHICommandList& RHICmdList, FRHIUnorderedAccessView* OutKeysUAV, FRHIUnorderedAccessView* OutIndicesUAV);
-
-	/**
-	 * Set input parameters.
-	 */
-	void SetParameters(FRHICommandList& RHICmdList,	const FNiagaraGPUSortInfo& SortInfo, uint32 EmitterKey, int32 OutputOffset, const FUintVector4& SortKeyParamsValue);
-
-	/**
-	 * Unbinds any buffers that have been bound.
-	 */
-	void UnbindBuffers(FRHICommandList& RHICmdList);
-
-private:
-
-	LAYOUT_FIELD(FShaderResourceParameter, NiagaraParticleDataFloat)
-	LAYOUT_FIELD(FShaderParameter, FloatDataStride)
-	LAYOUT_FIELD(FShaderResourceParameter, NiagaraParticleDataHalf)
-	LAYOUT_FIELD(FShaderResourceParameter, GPUParticleCountBuffer)
-	LAYOUT_FIELD(FShaderParameter, ParticleCountParams)
-	LAYOUT_FIELD(FShaderParameter, SortParams)
-	LAYOUT_FIELD(FShaderParameter, SortKeyParams)
-	LAYOUT_FIELD(FShaderParameter, CameraPosition)
-	LAYOUT_FIELD(FShaderParameter, CameraDirection)
-
-	/** Output key buffer. */
-	LAYOUT_FIELD(FShaderResourceParameter, OutKeys)
-	/** Output indices buffer. */
-	LAYOUT_FIELD(FShaderResourceParameter, OutParticleIndices)
 };
diff --git a/Engine/Source/Runtime/Engine/Private/GPUSortManager.cpp b/Engine/Source/Runtime/Engine/Private/GPUSortManager.cpp
--- a/Engine/Source/Runtime/Engine/Private/GPUSortManager.cpp
+++ b/Engine/Source/Runtime/Engine/Private/GPUSortManager.cpp
@@ -737,6 +737,8 @@ void FGPUSortManager::OnPreRender(FRHICommandListImmediate& RHICmdList)
 			SortBatch.GenerateKeys(RHICmdList, Callbacks, EGPUSortFlags::KeyGenAfterPreRender);
 		}
 	}
+
+	PostPreRenderEvent.Broadcast(RHICmdList);
 }
 
 void FGPUSortManager::OnPostRenderOpaque(FRHICommandListImmediate& RHICmdList)
diff --git a/Engine/Source/Runtime/Engine/Public/GPUSortManager.h b/Engine/Source/Runtime/Engine/Public/GPUSortManager.h
--- a/Engine/Source/Runtime/Engine/Public/GPUSortManager.h
+++ b/Engine/Source/Runtime/Engine/Public/GPUSortManager.h
@@ -156,7 +156,7 @@ ENUM_CLASS_FLAGS(EGPUSortFlags)
  * @param KeysUAV - The UAV that holds all the initial keys used to sort the values.
  * @param ValuesUAV - The UAV that holds the initial values to be sorted accordingly to the keys.
  */
- DECLARE_DELEGATE_SixParams(FGPUSortKeyGenDelegate, FRHICommandListImmediate&, int32, int32, EGPUSortFlags, FRHIUnorderedAccessView*, FRHIUnorderedAccessView*);
+DECLARE_DELEGATE_SixParams(FGPUSortKeyGenDelegate, FRHICommandListImmediate&, int32, int32, EGPUSortFlags, FRHIUnorderedAccessView*, FRHIUnorderedAccessView*);
 
 /**
  * A manager that handles different GPU sort tasks. 
@@ -412,7 +412,7 @@ public:
 	 * @param UsedFlags - The possibly used flags for each of the client system tasks, used to perform some optimizations.
 	 * @param InName - A name defining this client system, used for GPU markers.
 	 */
-	void Register(const FGPUSortKeyGenDelegate& CallbackDelegate, EGPUSortFlags UsedFlags, const FName& InName);
+	void Register(const FGPUSortKeyGenDelegate& CallbackDelegate, EGPUSortFlags UsedFlags, const FName& InName);	
 
 	/**
 	 * Add a GPU sort task to process this frame. Tasks are expected to be created before OnPreRender() gets called.
@@ -446,6 +446,12 @@ public:
 	 */
 	void OnPostRenderOpaque(FRHICommandListImmediate& RHICmdList);
 
+	/**
+	 * Event to register and receive post-prerender notification.
+	 */	
+	DECLARE_EVENT_OneParam(FGPUSortManager, FPostPreRenderEvent, FRHICommandListImmediate&);
+	FPostPreRenderEvent PostPreRenderEvent;
+
 private:
 
 	/** A helper to test that a flag has one and exactly one of A or B. */
