Removing legacy show grid behavior and using the engine show flags so toggling works correctly in the material editor and Niagara system editor.

diff --git a/Engine/Plugins/FX/Niagara/Source/NiagaraEditor/Private/Widgets/SNiagaraSystemViewport.cpp b/Engine/Plugins/FX/Niagara/Source/NiagaraEditor/Private/Widgets/SNiagaraSystemViewport.cpp
--- a/Engine/Plugins/FX/Niagara/Source/NiagaraEditor/Private/Widgets/SNiagaraSystemViewport.cpp
+++ b/Engine/Plugins/FX/Niagara/Source/NiagaraEditor/Private/Widgets/SNiagaraSystemViewport.cpp
@@ -42,7 +42,6 @@ public:
 	virtual bool CanSetWidgetMode(FWidget::EWidgetMode NewMode) const override { return false; }
 	virtual bool CanCycleWidgetMode() const override { return false; }
 
-	void SetShowGrid(bool bShowGrid);
 
 	virtual void SetIsSimulateInEditorViewport(bool bInIsSimulateInEditorViewport)override;
 
@@ -256,11 +255,6 @@ FSceneView* FNiagaraSystemViewportClient::CalcSceneView(FSceneViewFamily* ViewFa
 	return SceneView;
 }
 
-void FNiagaraSystemViewportClient::SetShowGrid(bool bShowGrid)
-{
-	DrawHelper.bDrawGrid = bShowGrid;
-}
-
 void FNiagaraSystemViewportClient::SetIsSimulateInEditorViewport(bool bInIsSimulateInEditorViewport)
 {
 	bIsSimulateInEditorViewport = bInIsSimulateInEditorViewport;
@@ -281,7 +275,6 @@ void FNiagaraSystemViewportClient::SetIsSimulateInEditorViewport(bool bInIsSimul
 void SNiagaraSystemViewport::Construct(const FArguments& InArgs)
 {
 	DrawFlags = 0;
-	bShowGrid = false;
 	bShowBackground = false;
 	PreviewComponent = nullptr;
 	AdvancedPreviewScene = MakeShareable(new FAdvancedPreviewScene(FPreviewScene::ConstructionValues()));
@@ -456,14 +449,13 @@ void SNiagaraSystemViewport::OnFocusViewportToSelection()
 
 void SNiagaraSystemViewport::TogglePreviewGrid()
 {
-	bShowGrid = !bShowGrid;
-	SystemViewportClient->SetShowGrid(bShowGrid);
+	SystemViewportClient->SetShowGrid();
 	RefreshViewport();
 }
 
 bool SNiagaraSystemViewport::IsTogglePreviewGridChecked() const
 {
-	return bShowGrid;
+	return SystemViewportClient->IsSetShowGridChecked();
 }
 
 void SNiagaraSystemViewport::TogglePreviewBackground()
diff --git a/Engine/Source/Editor/MaterialEditor/Private/SMaterialEditorViewport.cpp b/Engine/Source/Editor/MaterialEditor/Private/SMaterialEditorViewport.cpp
--- a/Engine/Source/Editor/MaterialEditor/Private/SMaterialEditorViewport.cpp
+++ b/Engine/Source/Editor/MaterialEditor/Private/SMaterialEditorViewport.cpp
@@ -46,8 +46,6 @@ public:
 	virtual void Tick(float DeltaSeconds) override;
 	virtual void Draw(FViewport* Viewport,FCanvas* Canvas) override;
 	virtual bool ShouldOrbitCamera() const override;
-	
-	void SetShowGrid(bool bShowGrid);
 
 	/**
 	* Focuses the viewport to the center of the bounding box/sphere ensuring that the entire bounds are in view
@@ -183,11 +181,6 @@ FLinearColor FMaterialEditorViewportClient::GetBackgroundColor() const
 	}
 }
 
-void FMaterialEditorViewportClient::SetShowGrid(bool bShowGrid)
-{
-	DrawHelper.bDrawGrid = bShowGrid;
-}
-
 void FMaterialEditorViewportClient::FocusViewportOnBounds(const FBoxSphereBounds Bounds, bool bInstant /*= false*/)
 {
 	const FVector Position = Bounds.Origin;
@@ -249,7 +242,6 @@ void SMaterialEditor3DPreviewViewport::Construct(const FArguments& InArgs)
 			AdvancedPreviewScene->GetWorld()->ChangeFeatureLevel(NewFeatureLevel);
 		});
 
-	bShowGrid = false;
 
 	PreviewPrimType = TPT_None;
 
@@ -645,14 +637,13 @@ bool SMaterialEditor3DPreviewViewport::IsPreviewMeshFromSelectionChecked() const
 
 void SMaterialEditor3DPreviewViewport::TogglePreviewGrid()
 {
-	bShowGrid = !bShowGrid;
-	EditorViewportClient->SetShowGrid(bShowGrid);
+	EditorViewportClient->SetShowGrid();
 	RefreshViewport();
 }
 
 bool SMaterialEditor3DPreviewViewport::IsTogglePreviewGridChecked() const
 {
-	return bShowGrid;
+	return EditorViewportClient->IsSetShowGridChecked();
 }
 
 void SMaterialEditor3DPreviewViewport::TogglePreviewBackground()
