Removing unneeded variable after CL 12807334
#rnx

diff --git a/Engine/Plugins/FX/Niagara/Source/NiagaraEditor/Private/Widgets/SNiagaraSystemViewport.h b/Engine/Plugins/FX/Niagara/Source/NiagaraEditor/Private/Widgets/SNiagaraSystemViewport.h
--- a/Engine/Plugins/FX/Niagara/Source/NiagaraEditor/Private/Widgets/SNiagaraSystemViewport.h
+++ b/Engine/Plugins/FX/Niagara/Source/NiagaraEditor/Private/Widgets/SNiagaraSystemViewport.h
@@ -48,10 +48,7 @@ public:
 	
 	/** If true, render background object in the preview scene. */
 	bool bShowBackground;
-	
-	/** If true, render grid the preview scene. */
-	bool bShowGrid;
-	
+
 	TSharedRef<class FAdvancedPreviewScene> GetPreviewScene() { return AdvancedPreviewScene.ToSharedRef(); }
 
 	/** The material editor has been added to a tab */
diff --git a/Engine/Source/Editor/MaterialEditor/Private/SMaterialEditorViewport.h b/Engine/Source/Editor/MaterialEditor/Private/SMaterialEditorViewport.h
--- a/Engine/Source/Editor/MaterialEditor/Private/SMaterialEditorViewport.h
+++ b/Engine/Source/Editor/MaterialEditor/Private/SMaterialEditorViewport.h
@@ -64,9 +64,6 @@ public:
 	
 	/** The preview primitive we are using. */
 	EThumbnailPrimType PreviewPrimType;
-
-	/** If true, render grid the preview scene. */
-	bool bShowGrid;
 	
 	/** The material editor has been added to a tab */
 	void OnAddedToTab( const TSharedRef<SDockTab>& OwnerTab );
