Integrate fix to sprite particle GPU sorting

diff --git a/Engine/Plugins/FX/Niagara/Source/Niagara/Private/NiagaraRendererSprites.cpp b/Engine/Plugins/FX/Niagara/Source/Niagara/Private/NiagaraRendererSprites.cpp
--- a/Engine/Plugins/FX/Niagara/Source/Niagara/Private/NiagaraRendererSprites.cpp
+++ b/Engine/Plugins/FX/Niagara/Source/Niagara/Private/NiagaraRendererSprites.cpp
@@ -45,11 +45,10 @@ struct FNiagaraDynamicDataSprites : public FNiagaraDynamicDataBase
 {
 	FNiagaraDynamicDataSprites(const FNiagaraEmitterInstance* InEmitter)
 		: FNiagaraDynamicDataBase(InEmitter)
-		, Material(nullptr)
 	{
 	}
 	
-	FMaterialRenderProxy* Material;
+	FMaterialRenderProxy* Material = nullptr;
 };
 
 /* Mesh collector classes */
@@ -329,7 +328,6 @@ void FNiagaraRendererSprites::SetVertexFactoryParticleData(
 	{
 		SCOPE_CYCLE_COUNTER(STAT_NiagaraRenderSpritesSorting)
 
-
 		FMaterialRenderProxy* MaterialRenderProxy = DynamicDataSprites->Material;
 		check(MaterialRenderProxy);
 		EBlendMode BlendMode = MaterialRenderProxy->GetMaterial(FeatureLevel)->GetBlendMode();
@@ -341,12 +339,11 @@ void FNiagaraRendererSprites::SetVertexFactoryParticleData(
 		const bool bShouldSort = SortMode != ENiagaraSortMode::None && (bHasTranslucentMaterials || !bSortOnlyWhenTranslucent);
 		const bool bCustomSorting = SortMode == ENiagaraSortMode::CustomAscending || SortMode == ENiagaraSortMode::CustomDecending;
 		const int32 SortVarIdx = bCustomSorting ? ENiagaraSpriteVFLayout::CustomSorting : ENiagaraSpriteVFLayout::Position;
-		const bool bSortVarIsHalf = bShouldSort && VFVariables[SortVarIdx].bHalfType;
 		if (bShouldSort && VFVariables[SortVarIdx].GetGPUOffset() != INDEX_NONE)
 		{
 			SortInfo.ParticleCount = NumInstances;
 			SortInfo.SortMode = SortMode;
-			SortInfo.SortAttributeOffset = (VFVariables[SortVarIdx].GetGPUOffset() & ~(1 << 31));
+			SortInfo.SortAttributeOffset = VFVariables[SortVarIdx].GetGPUOffset();
 			SortInfo.SetSortFlags(GNiagaraGPUSortingUseMaxPrecision != 0, bHasTranslucentMaterials); 
 			if (!bCustomSorting)
 			{
@@ -364,8 +361,8 @@ void FNiagaraRendererSprites::SetVertexFactoryParticleData(
 		{
 			FRHIShaderResourceView* FloatSRV = CPUSimParticleDataAllocation.ParticleData.FloatData.IsValid() ? CPUSimParticleDataAllocation.ParticleData.FloatData.SRV : (FRHIShaderResourceView*)FNiagaraRenderer::GetDummyFloatBuffer();
 			FRHIShaderResourceView* HalfSRV = CPUSimParticleDataAllocation.ParticleData.HalfData.IsValid() ? CPUSimParticleDataAllocation.ParticleData.HalfData.SRV : (FRHIShaderResourceView*)FNiagaraRenderer::GetDummyHalfBuffer();
-			const uint32 ParticleStrideInFloats = GbEnableMinimalGPUBuffers ? SourceParticleData->GetNumInstances() : (SourceParticleData->GetFloatStride() / sizeof(float));
-			check(GbEnableMinimalGPUBuffers || ParticleStrideInFloats == SourceParticleData->GetHalfStride() / sizeof(FFloat16));
+			const uint32 ParticleFloatDataStride = GbEnableMinimalGPUBuffers ? SourceParticleData->GetNumInstances() : (SourceParticleData->GetFloatStride() / sizeof(float));
+			const uint32 ParticleHalfDataStride = GbEnableMinimalGPUBuffers ? SourceParticleData->GetNumInstances() : (SourceParticleData->GetHalfStride() / sizeof(FFloat16));			
 
 			if (SortInfo.SortMode != ENiagaraSortMode::None && SortInfo.SortAttributeOffset != INDEX_NONE)
 			{
@@ -374,8 +371,10 @@ void FNiagaraRendererSprites::SetVertexFactoryParticleData(
 					FNiagaraUtilities::AllowComputeShaders(Batcher->GetShaderPlatform()))
 				{
 					SortInfo.ParticleCount = NumInstances;
-					SortInfo.ParticleDataFloatSRV = bSortVarIsHalf ? HalfSRV : FloatSRV;
-					SortInfo.FloatDataStride = ParticleStrideInFloats; // Same if Halfs.
+					SortInfo.ParticleDataFloatSRV = FloatSRV;
+					SortInfo.ParticleDataHalfSRV = HalfSRV;
+					SortInfo.FloatDataStride = ParticleFloatDataStride;
+					SortInfo.HalfDataStride = ParticleHalfDataStride;
 					if (Batcher->AddSortedGPUSimulation(SortInfo))
 					{
 						OutVertexFactory.SetSortedIndices(SortInfo.AllocationInfo.BufferSRV, SortInfo.AllocationInfo.BufferOffset);
@@ -393,7 +392,9 @@ void FNiagaraRendererSprites::SetVertexFactoryParticleData(
 			}
 			auto& ParticleData = CPUSimParticleDataAllocation.ParticleData;
 
-			VFLooseParams.NiagaraFloatDataStride = ParticleStrideInFloats;
+			check(ParticleFloatDataStride == ParticleHalfDataStride); // sanity check for the loose params
+
+			VFLooseParams.NiagaraFloatDataStride = ParticleFloatDataStride;
 			VFLooseParams.NiagaraParticleDataFloat = FloatSRV;
 			VFLooseParams.NiagaraParticleDataHalf = HalfSRV;
 		}
@@ -401,16 +402,18 @@ void FNiagaraRendererSprites::SetVertexFactoryParticleData(
 		{
 			FRHIShaderResourceView* FloatSRV = SourceParticleData->GetGPUBufferFloat().SRV.IsValid() ? (FRHIShaderResourceView*)SourceParticleData->GetGPUBufferFloat().SRV : (FRHIShaderResourceView*)FNiagaraRenderer::GetDummyFloatBuffer();
 			FRHIShaderResourceView* HalfSRV = SourceParticleData->GetGPUBufferHalf().SRV.IsValid() ? (FRHIShaderResourceView*)SourceParticleData->GetGPUBufferHalf().SRV : (FRHIShaderResourceView*)FNiagaraRenderer::GetDummyHalfBuffer();
-			const uint32 ParticleStrideInFloats = SourceParticleData->GetFloatStride() / sizeof(float);
-			check(ParticleStrideInFloats == SourceParticleData->GetHalfStride() / sizeof(FFloat16));
+			const uint32 ParticleFloatDataStride = SourceParticleData->GetFloatStride() / sizeof(float);
+			const uint32 ParticleHalfDataStride = SourceParticleData->GetHalfStride() / sizeof(FFloat16);
 
 			if (SortInfo.SortMode != ENiagaraSortMode::None && SortInfo.SortAttributeOffset != INDEX_NONE)
 			{
 				// Here we need to be conservative about the InstanceCount, since the final value is only known on the GPU after the simulation.
 				SortInfo.ParticleCount = SourceParticleData->GetNumInstances();
 
-				SortInfo.ParticleDataFloatSRV = bSortVarIsHalf ? HalfSRV : FloatSRV;
-				SortInfo.FloatDataStride = ParticleStrideInFloats;
+				SortInfo.ParticleDataFloatSRV = FloatSRV;
+				SortInfo.ParticleDataHalfSRV = HalfSRV;
+				SortInfo.FloatDataStride = ParticleFloatDataStride;
+				SortInfo.HalfDataStride = ParticleHalfDataStride;
 				SortInfo.GPUParticleCountSRV = Batcher->GetGPUInstanceCounterManager().GetInstanceCountBuffer().SRV;
 				SortInfo.GPUParticleCountOffset = SourceParticleData->GetGPUInstanceCountBufferOffset();
 				if (Batcher->AddSortedGPUSimulation(SortInfo))
@@ -419,7 +422,9 @@ void FNiagaraRendererSprites::SetVertexFactoryParticleData(
 				}
 			}
 
-			VFLooseParams.NiagaraFloatDataStride = ParticleStrideInFloats;
+			check(ParticleFloatDataStride == ParticleHalfDataStride); // sanity check for the loose params
+
+			VFLooseParams.NiagaraFloatDataStride = ParticleFloatDataStride;
 			VFLooseParams.NiagaraParticleDataFloat = FloatSRV;
 			VFLooseParams.NiagaraParticleDataHalf = HalfSRV;
 		}
@@ -681,7 +686,6 @@ FNiagaraDynamicDataBase *FNiagaraRendererSprites::GenerateDynamicData(const FNia
 		{
 			DynamicData = new FNiagaraDynamicDataSprites(Emitter);
 
-
 			//In preparation for a material override feature, we pass our material(s) and relevance in via dynamic data.
 			//The renderer ensures we have the correct usage and relevance for materials in BaseMaterials_GT.
 			//Any override feature must also do the same for materials that are set.
