Fixed crash when adding the direct read DI as a function input.

#rnx

diff --git a/Engine/Plugins/FX/Niagara/Source/Niagara/Private/NiagaraDataInterfaceParticleRead.cpp b/Engine/Plugins/FX/Niagara/Source/Niagara/Private/NiagaraDataInterfaceParticleRead.cpp
--- a/Engine/Plugins/FX/Niagara/Source/Niagara/Private/NiagaraDataInterfaceParticleRead.cpp
+++ b/Engine/Plugins/FX/Niagara/Source/Niagara/Private/NiagaraDataInterfaceParticleRead.cpp
@@ -2031,6 +2031,11 @@ void UNiagaraDataInterfaceParticleRead::ProvidePerInstanceDataForRenderThread(vo
 #if WITH_EDITOR	
 void UNiagaraDataInterfaceParticleRead::GetFeedback(UNiagaraSystem* Asset, UNiagaraComponent* Component, TArray<FNiagaraDataInterfaceError>& OutErrors, TArray<FNiagaraDataInterfaceFeedback>& Warnings, TArray<FNiagaraDataInterfaceFeedback>& Info)
 {
+	if (!Asset)
+	{
+		return;
+	}
+
 	UNiagaraEmitter* FoundSourceEmitter = nullptr;
 	for (const FNiagaraEmitterHandle& EmitterHandle : Asset->GetEmitterHandles())
 	{
