Fix speling mistake in user facing error message

diff --git a/Engine/Plugins/FX/Niagara/Source/NiagaraEditor/Private/NiagaraHlslTranslator.cpp b/Engine/Plugins/FX/Niagara/Source/NiagaraEditor/Private/NiagaraHlslTranslator.cpp
--- a/Engine/Plugins/FX/Niagara/Source/NiagaraEditor/Private/NiagaraHlslTranslator.cpp
+++ b/Engine/Plugins/FX/Niagara/Source/NiagaraEditor/Private/NiagaraHlslTranslator.cpp
@@ -6904,7 +6904,7 @@ int32 FHlslNiagaraTranslator::CompileOutputPin(const UEdGraphPin* InPin)
 			}
 			else
 			{
-				Error(LOCTEXT("IncorrectNumOutputsError", "Incorect number of outputs. Can possibly be fixed with a graph refresh."), Node, nullptr);
+				Error(LOCTEXT("IncorrectNumOutputsError", "Incorrect number of outputs. Can possibly be fixed with a graph refresh."), Node, nullptr);
 			}
 		}
 	}
