Ensure that GPUComputeScript compilation is considered when evaluating the auto compilation in the editor as changes are made to systems/emitters
-should reduce the risk of long running GPUCompute compilations (like for shader stages) accumulating which can crush a computer.

diff --git a/Engine/Plugins/FX/Niagara/Source/Niagara/Classes/NiagaraEmitter.h b/Engine/Plugins/FX/Niagara/Source/Niagara/Classes/NiagaraEmitter.h
--- a/Engine/Plugins/FX/Niagara/Source/Niagara/Classes/NiagaraEmitter.h
+++ b/Engine/Plugins/FX/Niagara/Source/Niagara/Classes/NiagaraEmitter.h
@@ -364,6 +364,7 @@ public:
 	NIAGARA_API UNiagaraScript* GetScript(ENiagaraScriptUsage Usage, FGuid UsageId);
 
 	NIAGARA_API UNiagaraScript* GetGPUComputeScript() { return GPUComputeScript; }
+	NIAGARA_API const UNiagaraScript* GetGPUComputeScript() const { return GPUComputeScript; }
 
 #if WITH_EDITORONLY_DATA
 	/** 'Source' data/graphs for the scripts used by this emitter. */
diff --git a/Engine/Plugins/FX/Niagara/Source/NiagaraEditor/Private/ViewModels/NiagaraSystemViewModel.cpp b/Engine/Plugins/FX/Niagara/Source/NiagaraEditor/Private/ViewModels/NiagaraSystemViewModel.cpp
--- a/Engine/Plugins/FX/Niagara/Source/NiagaraEditor/Private/ViewModels/NiagaraSystemViewModel.cpp
+++ b/Engine/Plugins/FX/Niagara/Source/NiagaraEditor/Private/ViewModels/NiagaraSystemViewModel.cpp
@@ -525,6 +525,37 @@ void FNiagaraSystemViewModel::PostUndo(bool bSuccess)
 	System->RequestCompile(false);
 }
 
+bool FNiagaraSystemViewModel::WaitingOnCompilation() const
+{
+	UNiagaraSystem& ViewedSystem = GetSystem();
+
+	if (ViewedSystem.HasOutstandingCompilationRequests())
+	{
+		return true;
+	}
+
+	// the above check only handles the VM script generation, and so GPU compute script compilation can still
+	// be underway, so we'll check for that explicitly so that we don't burden the user with excessive compiles
+	for (const FNiagaraEmitterHandle& EmitterHandle : ViewedSystem.GetEmitterHandles())
+	{
+		if (const UNiagaraEmitter* Emitter = EmitterHandle.GetInstance())
+		{
+			if (const UNiagaraScript* GPUComputeScript = Emitter->GetGPUComputeScript())
+			{
+				if (const FNiagaraShaderScript* ShaderScript = GPUComputeScript->GetRenderThreadScript())
+				{
+					if (!ShaderScript->IsCompilationFinished())
+					{
+						return true;
+					}
+				}
+			}
+		}
+	}
+
+	return false;
+}
+
 void FNiagaraSystemViewModel::Tick(float DeltaTime)
 {
 	if (System == nullptr)
@@ -549,8 +580,7 @@ void FNiagaraSystemViewModel::Tick(float DeltaTime)
 	TickCompileStatus();
 
 	bool bAutoCompileThisFrame = GetDefault<UNiagaraEditorSettings>()->GetAutoCompile() && bCanAutoCompile && GetLatestCompileStatus() == ENiagaraScriptCompileStatus::NCS_Dirty;
-	if ((bForceAutoCompileOnce || bAutoCompileThisFrame) && GetSystem().HasOutstandingCompilationRequests() == false)
-
+	if ((bForceAutoCompileOnce || bAutoCompileThisFrame) && !WaitingOnCompilation())
 	{
 		CompileSystem(false);
 		bForceAutoCompileOnce = false;
diff --git a/Engine/Plugins/FX/Niagara/Source/NiagaraEditor/Public/ViewModels/NiagaraSystemViewModel.h b/Engine/Plugins/FX/Niagara/Source/NiagaraEditor/Public/ViewModels/NiagaraSystemViewModel.h
--- a/Engine/Plugins/FX/Niagara/Source/NiagaraEditor/Public/ViewModels/NiagaraSystemViewModel.h
+++ b/Engine/Plugins/FX/Niagara/Source/NiagaraEditor/Public/ViewModels/NiagaraSystemViewModel.h
@@ -436,6 +436,9 @@ private:
 	/** builds stack module data for use in module dependencies */
 	void BuildStackModuleData(UNiagaraScript* Script, FGuid InEmitterHandleId, TArray<FNiagaraStackModuleData>& OutStackModuleData);
 
+	/** Returns true if our system has scripts pending compilation (including GPU compute scripts) */
+	bool WaitingOnCompilation() const;
+
 	/** Called when a property on the System is changed. */
 	void SystemChanged(UNiagaraSystem* ChangedSystem);
 
