Fixed a bug where user parameters could not be assigned to data interfaces in the stack ui

diff --git a/Engine/Plugins/FX/Niagara/Source/Niagara/Private/NiagaraDataInterfaceExport.cpp b/Engine/Plugins/FX/Niagara/Source/Niagara/Private/NiagaraDataInterfaceExport.cpp
--- a/Engine/Plugins/FX/Niagara/Source/Niagara/Private/NiagaraDataInterfaceExport.cpp
+++ b/Engine/Plugins/FX/Niagara/Source/Niagara/Private/NiagaraDataInterfaceExport.cpp
@@ -52,6 +52,8 @@ public:
 UNiagaraDataInterfaceExport::UNiagaraDataInterfaceExport(FObjectInitializer const& ObjectInitializer)
 	: Super(ObjectInitializer)
 {
+	FNiagaraTypeDefinition Def(UObject::StaticClass());
+	CallbackHandlerParameter.Parameter.SetType(Def);
 }
 
 void UNiagaraDataInterfaceExport::PostInitProperties()
diff --git a/Engine/Plugins/FX/Niagara/Source/Niagara/Private/NiagaraDataInterfaceSkeletalMesh.cpp b/Engine/Plugins/FX/Niagara/Source/Niagara/Private/NiagaraDataInterfaceSkeletalMesh.cpp
--- a/Engine/Plugins/FX/Niagara/Source/Niagara/Private/NiagaraDataInterfaceSkeletalMesh.cpp
+++ b/Engine/Plugins/FX/Niagara/Source/Niagara/Private/NiagaraDataInterfaceSkeletalMesh.cpp
@@ -1985,6 +1985,10 @@ UNiagaraDataInterfaceSkeletalMesh::UNiagaraDataInterfaceSkeletalMesh(FObjectInit
 	, WholeMeshLOD(INDEX_NONE)
 	, ChangeId(0)
 {
+
+	FNiagaraTypeDefinition Def(UObject::StaticClass());
+	MeshUserParameter.Parameter.SetType(Def);
+
 	static const FName NAME_Root("root");
 	ExcludeBoneName = NAME_Root;
 	bExcludeBone = false;
