Fix a bug where if two modules use the same function with different data interface input parameters, they would both use the same data interface.
Re-enabled the ReuseFunction test that catches this bug

diff --git a/Engine/Plugins/FX/Niagara/Source/NiagaraCore/Private/NiagaraCustomVersion.cpp b/Engine/Plugins/FX/Niagara/Source/NiagaraCore/Private/NiagaraCustomVersion.cpp
--- a/Engine/Plugins/FX/Niagara/Source/NiagaraCore/Private/NiagaraCustomVersion.cpp
+++ b/Engine/Plugins/FX/Niagara/Source/NiagaraCore/Private/NiagaraCustomVersion.cpp
@@ -8,4 +8,4 @@ const FGuid FNiagaraCustomVersion::GUID(0xFCF57AFA, 0x50764283, 0xB9A9E658, 0xFF
 // Register the custom version with core
 FCustomVersionRegistration GRegisterNiagaraCustomVersion(FNiagaraCustomVersion::GUID, FNiagaraCustomVersion::LatestVersion, TEXT("NiagaraVer"));
 
-const FGuid FNiagaraCustomVersion::LatestScriptCompileVersion(0x9315AB96, 0xF1704D25, 0xAAAA6038, 0xF25C36DD);
+const FGuid FNiagaraCustomVersion::LatestScriptCompileVersion(0xD52A505F, 0x0CAD485D, 0x8DE36D97, 0x3857F2F6);
\ No newline at end of file
diff --git a/Engine/Plugins/FX/Niagara/Source/NiagaraEditor/Private/NiagaraHlslTranslator.cpp b/Engine/Plugins/FX/Niagara/Source/NiagaraEditor/Private/NiagaraHlslTranslator.cpp
--- a/Engine/Plugins/FX/Niagara/Source/NiagaraEditor/Private/NiagaraHlslTranslator.cpp
+++ b/Engine/Plugins/FX/Niagara/Source/NiagaraEditor/Private/NiagaraHlslTranslator.cpp
@@ -279,12 +279,13 @@ bool FHlslNiagaraTranslator::ValidateTypePins(UNiagaraNode* NodeToValidate)
 
 
 void FHlslNiagaraTranslator::GenerateFunctionSignature(ENiagaraScriptUsage ScriptUsage, FString InName, const FString& InFullName, UNiagaraGraph* FuncGraph, TArray<int32>& Inputs,
-	bool bHadNumericInputs, bool bHasParameterMapParameters, TArray<UEdGraphPin*> StaticSwitchValues, FNiagaraFunctionSignature& OutSig)const
+	bool bHasNumericInputs, bool bHasParameterMapParameters, TArray<UEdGraphPin*> StaticSwitchValues, FNiagaraFunctionSignature& OutSig)const
 {
 	SCOPE_CYCLE_COUNTER(STAT_NiagaraEditor_Module_NiagaraHLSLTranslator_GenerateFunctionSignature);
 
 	TArray<FNiagaraVariable> InputVars;
 	TArray<UNiagaraNodeInput*> InputsNodes;
+	bool bHasDIParameters = false;
 
 	{
 		SCOPE_CYCLE_COUNTER(STAT_NiagaraEditor_Module_NiagaraHLSLTranslator_GenerateFunctionSignature_FindInputNodes);
@@ -319,21 +320,25 @@ void FHlslNiagaraTranslator::GenerateFunctionSignature(ENiagaraScriptUsage Scrip
 			//Only add to the signature if the caller has provided it, otherwise we use a local default.
 			if (Inputs[i] != INDEX_NONE)
 			{
-				FNiagaraVariable LiteralConstant = InputsNodes[i]->Input;
-				if (GetLiteralConstantVariable(LiteralConstant))
+				FNiagaraVariable InputVar = InputsNodes[i]->Input;
+				if (GetLiteralConstantVariable(InputVar))
 				{
-					checkf(LiteralConstant.GetType() == FNiagaraTypeDefinition::GetBoolDef(), TEXT("Only boolean types are currently supported for literal constants."));
-					FString LiteralConstantAlias = LiteralConstant.GetName().ToString() + TEXT("_") + (LiteralConstant.GetValue<bool>() ? TEXT("true") : TEXT("false"));
+					checkf(InputVar.GetType() == FNiagaraTypeDefinition::GetBoolDef(), TEXT("Only boolean types are currently supported for literal constants."));
+					FString LiteralConstantAlias = InputVar.GetName().ToString() + TEXT("_") + (InputVar.GetValue<bool>() ? TEXT("true") : TEXT("false"));
 					InName += TEXT("_") + GetSanitizedSymbolName(LiteralConstantAlias.Replace(TEXT("."), TEXT("_")));
 					ConstantInputIndicesToRemove.Add(i);
 				}
 				else
 				{
-					InputVars.Add(InputsNodes[i]->Input);
-					if (bHadNumericInputs)
+					InputVars.Add(InputVar);
+					if (InputVar.GetType().IsDataInterface())
+					{
+						bHasDIParameters = true;
+					}
+					else if (bHasNumericInputs)
 					{
 						InName += TEXT("_In");
-						InName += InputsNodes[i]->Input.GetType().GetName();
+						InName += InputVar.GetType().GetName();
 					}
 				}
 			}
@@ -359,7 +364,7 @@ void FHlslNiagaraTranslator::GenerateFunctionSignature(ENiagaraScriptUsage Scrip
 		for (int32 i = 0; i < OutputVars.Num(); ++i)
 		{
 			//Only add to the signature if the caller has provided it, otherwise we use a local default.
-			if (bHadNumericInputs)
+			if (bHasNumericInputs)
 			{
 				InName += TEXT("_Out");
 				InName += OutputVars[i].GetType().GetName();
@@ -373,7 +378,7 @@ void FHlslNiagaraTranslator::GenerateFunctionSignature(ENiagaraScriptUsage Scrip
 	// to be written within each call.
 	if ((ScriptUsage == ENiagaraScriptUsage::Module || ScriptUsage == ENiagaraScriptUsage::DynamicInput ||
 		ScriptUsage == ENiagaraScriptUsage::EmitterSpawnScript || ScriptUsage == ENiagaraScriptUsage::EmitterUpdateScript
-		|| bHasParameterMapParameters)
+		|| bHasParameterMapParameters || bHasDIParameters)
 		&& (ModuleAliasStr != nullptr || EmitterAliasStr != nullptr))
 	{
 		SCOPE_CYCLE_COUNTER(STAT_NiagaraEditor_Module_NiagaraHLSLTranslator_GenerateFunctionSignature_UniqueDueToMaps);
