Trivial fix for shader compiler warning

diff --git a/Engine/Plugins/FX/Niagara/Shaders/Private/NiagaraMeshVertexFactory.ush b/Engine/Plugins/FX/Niagara/Shaders/Private/NiagaraMeshVertexFactory.ush
--- a/Engine/Plugins/FX/Niagara/Shaders/Private/NiagaraMeshVertexFactory.ush
+++ b/Engine/Plugins/FX/Niagara/Shaders/Private/NiagaraMeshVertexFactory.ush
@@ -777,7 +777,7 @@ FVertexFactoryIntermediates GetVertexFactoryIntermediates(FVertexFactoryInput In
 		Intermediates.WorldToParticle = WorldToParticle;
 	}
 
-	Intermediates.ParticlePosition = Intermediates.ParticleToWorld[3];
+	Intermediates.ParticlePosition = Intermediates.ParticleToWorld[3].xyz;
 
 	// Translated world position.	
 	Intermediates.ParticleTranslatedWorldPosition = Intermediates.ParticlePosition.xyz + ResolvedView.PreViewTranslation.xyz;
