fx.EnableMinimalGPUBuffers is taken into account when computing the max instance count for CPU emitters.

#rnx

diff --git a/Engine/Plugins/FX/Niagara/Source/Niagara/Private/NiagaraEmitterInstance.cpp b/Engine/Plugins/FX/Niagara/Source/Niagara/Private/NiagaraEmitterInstance.cpp
--- a/Engine/Plugins/FX/Niagara/Source/Niagara/Private/NiagaraEmitterInstance.cpp
+++ b/Engine/Plugins/FX/Niagara/Source/Niagara/Private/NiagaraEmitterInstance.cpp
@@ -490,7 +490,7 @@ void FNiagaraEmitterInstance::Init(int32 InEmitterIdx, FNiagaraSystemInstanceID
 
 	// Prevent division by 0 in case there are no renderers.
 	uint32 MaxGPUBufferComponents = 1;
-	if (CachedEmitter->SimTarget == ENiagaraSimTarget::CPUSim)
+	if (CachedEmitter->SimTarget == ENiagaraSimTarget::CPUSim && GbEnableMinimalGPUBuffers)
 	{
 		// CPU emitters only upload the data needed by the renderers to the GPU. Compute the maximum number of components per particle
 		// among all the enabled renderers, since this will decide how many particles we can upload.
