Fixed a bug where the default value of a module variable could be set when it was used in a param map set node (instead of a get node).

diff --git a/Engine/Plugins/FX/Niagara/Source/NiagaraEditor/Private/NiagaraGraph.cpp b/Engine/Plugins/FX/Niagara/Source/NiagaraEditor/Private/NiagaraGraph.cpp
--- a/Engine/Plugins/FX/Niagara/Source/NiagaraEditor/Private/NiagaraGraph.cpp
+++ b/Engine/Plugins/FX/Niagara/Source/NiagaraEditor/Private/NiagaraGraph.cpp
@@ -901,26 +901,6 @@ UEdGraphPin* UNiagaraGraph::FindParameterMapDefaultValuePin(const FName Variable
 				break;
 			}
 		}
-
-		if (UNiagaraNodeParameterMapSet* SetNode = Cast<UNiagaraNodeParameterMapSet>(Node))
-		{
-			TArray<UEdGraphPin*> InputPins;
-			SetNode->GetInputPins(InputPins);
-			for (UEdGraphPin* InputPin : InputPins)
-			{
-				if (VariableName == InputPin->PinName)
-				{
-					DefaultInputPin = InputPin;
-					break;
-				}
-			}
-
-			if (DefaultInputPin != nullptr)
-			{
-				break;
-			}
-		}
-
 	}
 
 
