- Fix for skipping the DestroyComponent call on a Niagara Anim Notify if the component was deactive already

#rnx

diff --git a/Engine/Plugins/FX/Niagara/Source/NiagaraAnimNotifies/Private/AnimNotifyState_TimedNiagaraEffect.cpp b/Engine/Plugins/FX/Niagara/Source/NiagaraAnimNotifies/Private/AnimNotifyState_TimedNiagaraEffect.cpp
--- a/Engine/Plugins/FX/Niagara/Source/NiagaraAnimNotifies/Private/AnimNotifyState_TimedNiagaraEffect.cpp
+++ b/Engine/Plugins/FX/Niagara/Source/NiagaraAnimNotifies/Private/AnimNotifyState_TimedNiagaraEffect.cpp
@@ -58,7 +58,7 @@ void UAnimNotifyState_TimedNiagaraEffect::NotifyEnd(USkeletalMeshComponent * Mes
 			bTemplateMatch |= PreviousTemplates.Contains(FXSystemComponent->GetFXSystemAsset());
 #endif
 
-			if (bSocketMatch && bTemplateMatch && FXSystemComponent->IsActive())
+			if (bSocketMatch && bTemplateMatch)
 			{
 				// Either destroy the component or deactivate it to have it's active FXSystems finish.
 				// The component will auto destroy once all FXSystem are gone.
