Performance improvements for creation of uniform buffers for GPU compute scripts
-global/system/owner buffers are shared between emitter and stages
-emitter/external buffers are shared between stages

diff --git a/Engine/Plugins/FX/Niagara/Source/Niagara/Classes/NiagaraEmitterInstanceBatcher.h b/Engine/Plugins/FX/Niagara/Source/Niagara/Classes/NiagaraEmitterInstanceBatcher.h
--- a/Engine/Plugins/FX/Niagara/Source/Niagara/Classes/NiagaraEmitterInstanceBatcher.h
+++ b/Engine/Plugins/FX/Niagara/Source/Niagara/Classes/NiagaraEmitterInstanceBatcher.h
@@ -202,6 +202,9 @@ private:
 	void ClearFreeIDsListSizesBuffer(FRHICommandList& RHICmdList);
 	void UpdateFreeIDBuffers(FRHICommandList& RHICmdList, FEmitterInstanceList& Instances);
 
+	void SetConstantBuffers(FRHICommandList &RHICmdList, const FNiagaraShaderRef& Shader, const FNiagaraGPUSystemTick& Tick, const FNiagaraComputeInstanceData* Instance);
+	void BuildConstantBuffers(FNiagaraGPUSystemTick& Tick);
+
 	/** Feature level of this effects system */
 	ERHIFeatureLevel::Type FeatureLevel;
 	/** Shader platform that will be rendering this effects system */
@@ -215,6 +218,12 @@ private:
 	// GPU emitter instance count buffer. Contains the actual particle / instance count generate in the GPU tick.
 	FNiagaraGPUInstanceCountManager GPUInstanceCounterManager;
 
+	// persistent layouts used to create the constant buffers for the compute sim shader
+	FRHIUniformBufferLayout GlobalCBufferLayout;
+	FRHIUniformBufferLayout SystemCBufferLayout;
+	FRHIUniformBufferLayout OwnerCBufferLayout;
+	FRHIUniformBufferLayout EmitterCBufferLayout;
+
 	// @todo REMOVE THIS HACK
 	uint32 LastFrameThatDrainedData;
 	
diff --git a/Engine/Plugins/FX/Niagara/Source/Niagara/Classes/NiagaraScriptExecutionContext.h b/Engine/Plugins/FX/Niagara/Source/Niagara/Classes/NiagaraScriptExecutionContext.h
--- a/Engine/Plugins/FX/Niagara/Source/Niagara/Classes/NiagaraScriptExecutionContext.h
+++ b/Engine/Plugins/FX/Niagara/Source/Niagara/Classes/NiagaraScriptExecutionContext.h
@@ -258,10 +258,6 @@ public:
 	class FNiagaraShaderScript*  GPUScript_RT;
 
 	// persistent layouts used to create the constant buffers for the compute sim shader
-	FRHIUniformBufferLayout GlobalCBufferLayout;
-	FRHIUniformBufferLayout SystemCBufferLayout;
-	FRHIUniformBufferLayout OwnerCBufferLayout;
-	FRHIUniformBufferLayout EmitterCBufferLayout;
 	FRHIUniformBufferLayout ExternalCBufferLayout;
 
 	//Dynamic state updated either from GT via RT commands or from the RT side sim code itself.
@@ -340,9 +336,6 @@ struct FNiagaraSimStageData
 struct FNiagaraComputeInstanceData
 {
 	FNiagaraGpuSpawnInfo SpawnInfo;
-	uint8* GlobalParamData = nullptr;
-	uint8* SystemParamData = nullptr;
-	uint8* OwnerParamData = nullptr;
 	uint8* EmitterParamData = nullptr;
 	uint8* ExternalParamData = nullptr;
 	FNiagaraComputeExecutionContext* Context = nullptr;
@@ -376,27 +369,49 @@ struct FNiagaraComputeInstanceData
 class FNiagaraGPUSystemTick
 {
 public:
-	FNiagaraGPUSystemTick()
-		: Count(0)
-		, DIInstanceData(nullptr)
-		, InstanceData_ParamData_Packed(nullptr)
-	{}
-
 	void Init(FNiagaraSystemInstance* InSystemInstance);
 	void Destroy();
 
 	FORCEINLINE bool IsValid()const{ return InstanceData_ParamData_Packed != nullptr; }
 	FORCEINLINE FNiagaraComputeInstanceData* GetInstanceData()const{ return reinterpret_cast<FNiagaraComputeInstanceData*>(InstanceData_ParamData_Packed); }
 
-	uint32 Count;
-	uint32 TotalDispatches;
-	FNiagaraSystemInstanceID SystemInstanceID;
-	FNiagaraDataInterfaceInstanceData* DIInstanceData;
-	uint8* InstanceData_ParamData_Packed;
+	enum EUniformBufferType
+	{
+		UBT_FirstSystemType = 0,
+		UBT_Global = UBT_FirstSystemType,
+		UBT_System,
+		UBT_Owner,
+		UBT_NumSystemTypes,
+
+		UBT_FirstInstanceType = UBT_NumSystemTypes,
+		UBT_Emitter = UBT_FirstInstanceType,
+		UBT_External,
+
+		UBT_NumTypes,
+
+		UBT_NumInstanceTypes = UBT_NumTypes - UBT_NumSystemTypes,
+	};
+
+	FUniformBufferRHIRef GetUniformBuffer(EUniformBufferType Type, const FNiagaraComputeInstanceData* InstanceData, bool Previous) const;
+	const uint8* GetUniformBufferSource(EUniformBufferType Type, const FNiagaraComputeInstanceData* InstanceData, bool Previous) const;
+
+public:
+	// Transient data used by the RT
+	TArray<FUniformBufferRHIRef> UniformBuffers;
+
+	// data assigned by GT
+	FNiagaraSystemInstanceID SystemInstanceID = 0LL;
+	FNiagaraDataInterfaceInstanceData* DIInstanceData = nullptr;
+	uint8* InstanceData_ParamData_Packed = nullptr;
+	uint8* GlobalParamData = nullptr;
+	uint8* SystemParamData = nullptr;
+	uint8* OwnerParamData = nullptr;
+	uint32 Count = 0;
+	uint32 TotalDispatches = 0;
+	uint32 NumInstancesWithSimStages = 0;
 	bool bRequiresDistanceFieldData = false;
 	bool bRequiresDepthBuffer = false;
 	bool bRequiresEarlyViewData = false;
 	bool bNeedsReset = false;
 	bool bIsFinalTick = false;
-	uint32 NumInstancesWithSimStages = 0;
 };
diff --git a/Engine/Plugins/FX/Niagara/Source/Niagara/Private/NiagaraEmitterInstanceBatcher.cpp b/Engine/Plugins/FX/Niagara/Source/Niagara/Private/NiagaraEmitterInstanceBatcher.cpp
--- a/Engine/Plugins/FX/Niagara/Source/Niagara/Private/NiagaraEmitterInstanceBatcher.cpp
+++ b/Engine/Plugins/FX/Niagara/Source/Niagara/Private/NiagaraEmitterInstanceBatcher.cpp
@@ -72,6 +72,10 @@ NiagaraEmitterInstanceBatcher::NiagaraEmitterInstanceBatcher(ERHIFeatureLevel::T
 	: FeatureLevel(InFeatureLevel)
 	, ShaderPlatform(InShaderPlatform)
 	, GPUSortManager(InGPUSortManager)
+	, GlobalCBufferLayout(TEXT("Niagara GPU Global CBuffer"))
+	, SystemCBufferLayout(TEXT("Niagara GPU System CBuffer"))
+	, OwnerCBufferLayout(TEXT("Niagara GPU Owner CBuffer"))
+	, EmitterCBufferLayout(TEXT("Niagara GPU Emitter CBuffer"))
 	// @todo REMOVE THIS HACK
 	, LastFrameThatDrainedData(GFrameNumberRenderThread)
 	, NumAllocatedFreeIDListSizes(0)
@@ -101,6 +105,18 @@ NiagaraEmitterInstanceBatcher::NiagaraEmitterInstanceBatcher(ERHIFeatureLevel::T
 			);
 		}
 	}
+
+	GlobalCBufferLayout.ConstantBufferSize = sizeof(FNiagaraGlobalParameters);
+	GlobalCBufferLayout.ComputeHash();
+
+	SystemCBufferLayout.ConstantBufferSize = sizeof(FNiagaraSystemParameters);
+	SystemCBufferLayout.ComputeHash();
+
+	OwnerCBufferLayout.ConstantBufferSize = sizeof(FNiagaraOwnerParameters);
+	OwnerCBufferLayout.ComputeHash();
+
+	EmitterCBufferLayout.ConstantBufferSize = sizeof(FNiagaraEmitterParameters);
+	EmitterCBufferLayout.ComputeHash();
 }
 
 NiagaraEmitterInstanceBatcher::~NiagaraEmitterInstanceBatcher()
@@ -127,6 +143,108 @@ void NiagaraEmitterInstanceBatcher::InstanceDeallocated_RenderThread(const FNiag
 	}
 }
 
+void NiagaraEmitterInstanceBatcher::BuildConstantBuffers(FNiagaraGPUSystemTick& Tick)
+{
+	if (!Tick.Count)
+	{
+		return;
+	}
+
+	FNiagaraComputeInstanceData* EmittersData = Tick.GetInstanceData();
+
+	// first go through and figure out if we need to support interpolated spawning
+	bool HasInterpolationParameters = false;
+	bool HasMultipleStages = false;
+	for (uint32 CountIt = 0; CountIt < Tick.Count; ++CountIt)
+	{
+		HasInterpolationParameters = HasInterpolationParameters || EmittersData[CountIt].Context->HasInterpolationParameters;
+		HasMultipleStages = HasMultipleStages || EmittersData[CountIt].bUsesOldShaderStages || EmittersData[CountIt].bUsesSimStages;
+	}
+
+	int32 BoundParameterCounts[FNiagaraGPUSystemTick::UBT_NumTypes][2];
+	for (int32 i = 0; i < FNiagaraGPUSystemTick::UBT_NumTypes; ++i)
+	{
+		for (int32 j = 0; j < 2; ++j)
+		{
+			BoundParameterCounts[i][j] = 0;
+		}
+	}
+
+	for (uint32 CountIt = 0; CountIt < Tick.Count; ++CountIt)
+	{
+		FNiagaraComputeInstanceData& EmitterData = EmittersData[CountIt];
+
+		const FNiagaraShaderRef& Shader = EmitterData.Context->GPUScript_RT->GetShader();
+		const int32 HasExternalConstants = EmitterData.Context->ExternalCBufferLayout.ConstantBufferSize > 0 ? 1 : 0;
+
+		for (int32 InterpIt = 0; InterpIt < (HasInterpolationParameters ? 2 : 1); ++InterpIt)
+		{
+			BoundParameterCounts[FNiagaraGPUSystemTick::UBT_Global][InterpIt] += Shader->GlobalConstantBufferParam[InterpIt].IsBound() ? 1 : 0;
+			BoundParameterCounts[FNiagaraGPUSystemTick::UBT_System][InterpIt] += Shader->SystemConstantBufferParam[InterpIt].IsBound() ? 1 : 0;
+			BoundParameterCounts[FNiagaraGPUSystemTick::UBT_Owner][InterpIt] += Shader->OwnerConstantBufferParam[InterpIt].IsBound() ? 1 : 0;
+			BoundParameterCounts[FNiagaraGPUSystemTick::UBT_Emitter][InterpIt] += Shader->EmitterConstantBufferParam[InterpIt].IsBound() ? 1 : 0;
+			BoundParameterCounts[FNiagaraGPUSystemTick::UBT_External][InterpIt] += Shader->ExternalConstantBufferParam[InterpIt].IsBound() ? 1 : 0;
+		}
+	}
+
+	const int32 InterpScale = HasInterpolationParameters ? 2 : 1;
+	const int32 BufferCount = InterpScale * (FNiagaraGPUSystemTick::UBT_NumSystemTypes + FNiagaraGPUSystemTick::UBT_NumInstanceTypes * Tick.Count);
+
+	Tick.UniformBuffers.Empty(BufferCount);
+
+	const FRHIUniformBufferLayout* SystemLayouts[FNiagaraGPUSystemTick::UBT_NumSystemTypes] =
+	{
+		&GlobalCBufferLayout,
+		&SystemCBufferLayout,
+		&OwnerCBufferLayout
+	};
+
+	for (int32 InterpIt = 0; InterpIt < InterpScale; ++InterpIt)
+	{
+		for (int32 SystemTypeIt = FNiagaraGPUSystemTick::UBT_FirstSystemType; SystemTypeIt < FNiagaraGPUSystemTick::UBT_NumSystemTypes; ++SystemTypeIt)
+		{
+			FUniformBufferRHIRef BufferRef;
+
+			if (BoundParameterCounts[SystemTypeIt][InterpIt])
+			{
+				BufferRef = RHICreateUniformBuffer(
+					Tick.GetUniformBufferSource((FNiagaraGPUSystemTick::EUniformBufferType) SystemTypeIt, nullptr, !InterpIt),
+					*SystemLayouts[SystemTypeIt],
+					((BoundParameterCounts[SystemTypeIt][InterpIt] > 1) || HasMultipleStages)
+						? EUniformBufferUsage::UniformBuffer_SingleFrame
+						: EUniformBufferUsage::UniformBuffer_SingleDraw);
+			}
+
+			Tick.UniformBuffers.Add(BufferRef);
+		}
+
+		for (int32 InstanceTypeIt = FNiagaraGPUSystemTick::UBT_FirstInstanceType; InstanceTypeIt < FNiagaraGPUSystemTick::UBT_NumTypes; ++InstanceTypeIt)
+		{
+			for (uint32 InstanceIt = 0; InstanceIt < Tick.Count; ++InstanceIt)
+			{
+				FNiagaraComputeInstanceData& EmitterData = EmittersData[InstanceIt];
+
+				FUniformBufferRHIRef BufferRef;
+
+				if (BoundParameterCounts[InstanceTypeIt][InterpIt])
+				{
+					BufferRef = RHICreateUniformBuffer(
+						Tick.GetUniformBufferSource((FNiagaraGPUSystemTick::EUniformBufferType) InstanceTypeIt, &EmitterData, !InterpIt),
+						InstanceTypeIt == FNiagaraGPUSystemTick::UBT_Emitter
+							? EmitterCBufferLayout
+							: EmitterData.Context->ExternalCBufferLayout,
+						((BoundParameterCounts[InstanceTypeIt][InterpIt] > 1) || HasMultipleStages)
+							? EUniformBufferUsage::UniformBuffer_SingleFrame
+							: EUniformBufferUsage::UniformBuffer_SingleDraw);
+				}
+
+				Tick.UniformBuffers.Add(BufferRef);
+			}
+		}
+
+	}
+}
+
 void NiagaraEmitterInstanceBatcher::GiveSystemTick_RenderThread(FNiagaraGPUSystemTick& Tick)
 {
 	check(IsInRenderingThread());
@@ -163,7 +281,9 @@ void NiagaraEmitterInstanceBatcher::GiveSystemTick_RenderThread(FNiagaraGPUSyste
 	// A note:
 	// This is making a copy of Tick. That structure is small now and we take a copy to avoid
 	// making a bunch of small allocations on the game thread. We may need to revisit this.
-	Ticks_RT.Add(Tick);
+	FNiagaraGPUSystemTick& AddedTick = Ticks_RT.Add_GetRef(Tick);
+
+	BuildConstantBuffers(AddedTick);
 }
 
 void NiagaraEmitterInstanceBatcher::ReleaseInstanceCounts_RenderThread(FNiagaraComputeExecutionContext* ExecContext, FNiagaraDataSet* DataSet)
@@ -1340,16 +1460,32 @@ void NiagaraEmitterInstanceBatcher::UnsetDataInterfaceParameters(const TArray<FN
 	}
 }
 
-static void SetConstantBuffer(FRHICommandList &RHICmdList, FRHIComputeShader* ComputeShader, const FShaderUniformBufferParameter& BufferParam, const FRHIUniformBufferLayout& Layout, const uint8* ParamData)
+static void SetConstantBuffer(FRHICommandList &RHICmdList, FRHIComputeShader* ComputeShader, const FShaderUniformBufferParameter& BufferParam, const FUniformBufferRHIRef& UniformBuffer)
 {
-	if (!BufferParam.IsBound())
-		return;
+	if (BufferParam.IsBound() && UniformBuffer.IsValid())
+	{
+		RHICmdList.SetShaderUniformBuffer(ComputeShader, BufferParam.GetBaseIndex(), UniformBuffer);
+	}
+}
+
+void NiagaraEmitterInstanceBatcher::SetConstantBuffers(FRHICommandList &RHICmdList, const FNiagaraShaderRef& Shader, const FNiagaraGPUSystemTick& Tick, const FNiagaraComputeInstanceData* Instance)
+{
+	FNiagaraComputeExecutionContext* Context = Instance->Context;
+	FRHIComputeShader* ComputeShader = Shader.GetComputeShader();
+
+	SetConstantBuffer(RHICmdList, ComputeShader, Shader->GlobalConstantBufferParam[0], Tick.GetUniformBuffer(FNiagaraGPUSystemTick::UBT_Global, nullptr, true));
+	SetConstantBuffer(RHICmdList, ComputeShader, Shader->SystemConstantBufferParam[0], Tick.GetUniformBuffer(FNiagaraGPUSystemTick::UBT_System, nullptr, true));
+	SetConstantBuffer(RHICmdList, ComputeShader, Shader->OwnerConstantBufferParam[0], Tick.GetUniformBuffer(FNiagaraGPUSystemTick::UBT_Owner, nullptr, true));
+	SetConstantBuffer(RHICmdList, ComputeShader, Shader->EmitterConstantBufferParam[0], Tick.GetUniformBuffer(FNiagaraGPUSystemTick::UBT_Emitter, Instance, true));
+	SetConstantBuffer(RHICmdList, ComputeShader, Shader->ExternalConstantBufferParam[0], Tick.GetUniformBuffer(FNiagaraGPUSystemTick::UBT_External, Instance, true));
 
-	if (Layout.ConstantBufferSize)
+	if (Context->HasInterpolationParameters)
 	{
-		check(Layout.Resources.Num() == 0);
-		FUniformBufferRHIRef CBuffer = RHICreateUniformBuffer(ParamData, Layout, EUniformBufferUsage::UniformBuffer_SingleDraw);
-		RHICmdList.SetShaderUniformBuffer(ComputeShader, BufferParam.GetBaseIndex(), CBuffer);
+		SetConstantBuffer(RHICmdList, ComputeShader, Shader->GlobalConstantBufferParam[1], Tick.GetUniformBuffer(FNiagaraGPUSystemTick::UBT_Global, nullptr, false));
+		SetConstantBuffer(RHICmdList, ComputeShader, Shader->SystemConstantBufferParam[1], Tick.GetUniformBuffer(FNiagaraGPUSystemTick::UBT_System, nullptr, false));
+		SetConstantBuffer(RHICmdList, ComputeShader, Shader->OwnerConstantBufferParam[1], Tick.GetUniformBuffer(FNiagaraGPUSystemTick::UBT_Owner, nullptr, false));
+		SetConstantBuffer(RHICmdList, ComputeShader, Shader->EmitterConstantBufferParam[1], Tick.GetUniformBuffer(FNiagaraGPUSystemTick::UBT_Emitter, Instance, false));
+		SetConstantBuffer(RHICmdList, ComputeShader, Shader->ExternalConstantBufferParam[1], Tick.GetUniformBuffer(FNiagaraGPUSystemTick::UBT_External, Instance, false));
 	}
 }
 
@@ -1474,20 +1610,7 @@ void NiagaraEmitterInstanceBatcher::Run(const FNiagaraGPUSystemTick& Tick, const
 		NumThreadGroups = FMath::Min(NIAGARA_MAX_COMPUTE_THREADGROUPS, FMath::DivideAndRoundUp(TotalNumInstances, ShaderThreadGroupSize));
 	}
 
-	SetConstantBuffer(RHICmdList, ComputeShader, Shader->GlobalConstantBufferParam[0], Context->GlobalCBufferLayout, Instance->GlobalParamData);
-	SetConstantBuffer(RHICmdList, ComputeShader, Shader->SystemConstantBufferParam[0], Context->SystemCBufferLayout, Instance->SystemParamData);
-	SetConstantBuffer(RHICmdList, ComputeShader, Shader->OwnerConstantBufferParam[0], Context->OwnerCBufferLayout, Instance->OwnerParamData);
-	SetConstantBuffer(RHICmdList, ComputeShader, Shader->EmitterConstantBufferParam[0], Context->EmitterCBufferLayout, Instance->EmitterParamData);
-	SetConstantBuffer(RHICmdList, ComputeShader, Shader->ExternalConstantBufferParam[0], Context->ExternalCBufferLayout, Instance->ExternalParamData);
-	// setup script parameters
-	if (Context->HasInterpolationParameters)
-	{
-		SetConstantBuffer(RHICmdList, ComputeShader, Shader->GlobalConstantBufferParam[1], Context->GlobalCBufferLayout, Instance->GlobalParamData + sizeof(FNiagaraGlobalParameters));
-		SetConstantBuffer(RHICmdList, ComputeShader, Shader->SystemConstantBufferParam[1], Context->SystemCBufferLayout, Instance->SystemParamData + sizeof(FNiagaraSystemParameters));
-		SetConstantBuffer(RHICmdList, ComputeShader, Shader->OwnerConstantBufferParam[1], Context->OwnerCBufferLayout, Instance->OwnerParamData + sizeof(FNiagaraOwnerParameters));
-		SetConstantBuffer(RHICmdList, ComputeShader, Shader->EmitterConstantBufferParam[1], Context->EmitterCBufferLayout, Instance->EmitterParamData + sizeof(FNiagaraEmitterParameters));
-		SetConstantBuffer(RHICmdList, ComputeShader, Shader->ExternalConstantBufferParam[1], Context->ExternalCBufferLayout, Instance->ExternalParamData + Context->ExternalCBufferLayout.ConstantBufferSize);
-	}
+	SetConstantBuffers(RHICmdList, Shader, Tick, Instance);
 
 	//UE_LOG(LogNiagara, Log, TEXT("Num Instance : %d | Num Group : %d | Spawned Istance : %d | Start Instance : %d | Num Indices : %d | Stage Index : %d"), 
 		//TotalNumInstances, NumThreadGroups, InstancesToSpawnThisFrame, UpdateStartInstance, Context->NumIndicesPerInstance, SimulationStageIndex);
diff --git a/Engine/Plugins/FX/Niagara/Source/Niagara/Private/NiagaraScriptExecutionContext.cpp b/Engine/Plugins/FX/Niagara/Source/Niagara/Private/NiagaraScriptExecutionContext.cpp
--- a/Engine/Plugins/FX/Niagara/Source/Niagara/Private/NiagaraScriptExecutionContext.cpp
+++ b/Engine/Plugins/FX/Niagara/Source/Niagara/Private/NiagaraScriptExecutionContext.cpp
@@ -382,12 +382,29 @@ void FNiagaraGPUSystemTick::Init(FNiagaraSystemInstance* InSystemInstance)
 	NumInstancesWithSimStages = 0;
 
 	TotalDispatches = 0;
-	uint8* GlobalParamData = ParamDataBufferPtr;
-	ParamDataBufferPtr += 2 * sizeof(FNiagaraGlobalParameters);
-	uint8* SystemParamData = ParamDataBufferPtr;
-	ParamDataBufferPtr += 2 * sizeof(FNiagaraSystemParameters);
-	uint8* OwnerParamData = ParamDataBufferPtr;
-	ParamDataBufferPtr += 2 * sizeof(FNiagaraOwnerParameters);
+
+	// we want to include interpolation parameters (current and previous frame) if any of the emitters in the system
+	// require it
+	const bool IncludeInterpolationParameters = InSystemInstance->GPUParamIncludeInterpolation;
+	const int32 InterpFactor = IncludeInterpolationParameters ? 2 : 1;
+
+	GlobalParamData = ParamDataBufferPtr;
+	SystemParamData = GlobalParamData + InterpFactor * sizeof(FNiagaraGlobalParameters);
+	OwnerParamData = SystemParamData + InterpFactor * sizeof(FNiagaraSystemParameters);
+
+	// actually copy all of the data over, for the system data we only need to do it once (rather than per-emitter)
+	FMemory::Memcpy(GlobalParamData, &InSystemInstance->GetGlobalParameters(), sizeof(FNiagaraGlobalParameters));
+	FMemory::Memcpy(SystemParamData, &InSystemInstance->GetSystemParameters(), sizeof(FNiagaraSystemParameters));
+	FMemory::Memcpy(OwnerParamData, &InSystemInstance->GetOwnerParameters(), sizeof(FNiagaraOwnerParameters));
+
+	if (IncludeInterpolationParameters)
+	{
+		FMemory::Memcpy(GlobalParamData + sizeof(FNiagaraGlobalParameters), &InSystemInstance->GetGlobalParameters(true), sizeof(FNiagaraGlobalParameters));
+		FMemory::Memcpy(SystemParamData + sizeof(FNiagaraSystemParameters), &InSystemInstance->GetSystemParameters(true), sizeof(FNiagaraSystemParameters));
+		FMemory::Memcpy(OwnerParamData + sizeof(FNiagaraOwnerParameters), &InSystemInstance->GetOwnerParameters(true), sizeof(FNiagaraOwnerParameters));
+	}
+
+	ParamDataBufferPtr = OwnerParamData + InterpFactor * sizeof(FNiagaraOwnerParameters);
 
 	// Now we will generate instance data for every GPU simulation we want to run on the render thread.
 	// This is spawn rate as well as DataInterface per instance data and the ParameterData for the emitter.
@@ -413,10 +430,21 @@ void FNiagaraGPUSystemTick::Init(FNiagaraSystemInstance* InSystemInstance)
 
 			int32 ParmSize = GPUContext->CombinedParamStore.GetPaddedParameterSizeInBytes();
 
-			InstanceData->GlobalParamData = GlobalParamData;
-			InstanceData->SystemParamData = SystemParamData;
-			InstanceData->OwnerParamData = OwnerParamData;
 			InstanceData->EmitterParamData = ParamDataBufferPtr;
+			ParamDataBufferPtr += InterpFactor * sizeof(FNiagaraEmitterParameters);
+
+			InstanceData->ExternalParamData = ParamDataBufferPtr;
+			ParamDataBufferPtr += ParmSize;
+
+			// actually copy all of the data over
+			FMemory::Memcpy(InstanceData->EmitterParamData, &InSystemInstance->GetEmitterParameters(EmitterIdx), sizeof(FNiagaraEmitterParameters));
+			if (IncludeInterpolationParameters)
+			{
+				FMemory::Memcpy(InstanceData->EmitterParamData + sizeof(FNiagaraEmitterParameters), &InSystemInstance->GetEmitterParameters(EmitterIdx, true), sizeof(FNiagaraEmitterParameters));
+			}
+
+			GPUContext->CombinedParamStore.CopyParameterDataToPaddedBuffer(InstanceData->ExternalParamData, ParmSize);
+
 			UNiagaraEmitter* EmitterRaw = Emitter->GetCachedEmitter();
 			if (EmitterRaw)
 			{
@@ -432,21 +460,6 @@ void FNiagaraGPUSystemTick::Init(FNiagaraSystemInstance* InSystemInstance)
 			{
 				NumInstancesWithSimStages++;
 			}
-			ParamDataBufferPtr += 2 * sizeof(FNiagaraEmitterParameters);
-			InstanceData->ExternalParamData = ParamDataBufferPtr;
-			ParamDataBufferPtr += ParmSize;
-
-			// actually copy all of the data over
-			FMemory::Memcpy(InstanceData->GlobalParamData, &InSystemInstance->GetGlobalParameters(), sizeof(FNiagaraGlobalParameters));
-			FMemory::Memcpy(InstanceData->GlobalParamData + sizeof(FNiagaraGlobalParameters), &InSystemInstance->GetGlobalParameters(true), sizeof(FNiagaraGlobalParameters));
-			FMemory::Memcpy(InstanceData->SystemParamData, &InSystemInstance->GetSystemParameters(), sizeof(FNiagaraSystemParameters));
-			FMemory::Memcpy(InstanceData->SystemParamData + sizeof(FNiagaraSystemParameters), &InSystemInstance->GetSystemParameters(true), sizeof(FNiagaraSystemParameters));
-			FMemory::Memcpy(InstanceData->OwnerParamData, &InSystemInstance->GetOwnerParameters(), sizeof(FNiagaraOwnerParameters));
-			FMemory::Memcpy(InstanceData->OwnerParamData + sizeof(FNiagaraOwnerParameters), &InSystemInstance->GetOwnerParameters(true), sizeof(FNiagaraOwnerParameters));
-			FMemory::Memcpy(InstanceData->EmitterParamData, &InSystemInstance->GetEmitterParameters(EmitterIdx), sizeof(FNiagaraEmitterParameters));
-			FMemory::Memcpy(InstanceData->EmitterParamData + sizeof(FNiagaraEmitterParameters), &InSystemInstance->GetEmitterParameters(EmitterIdx, true), sizeof(FNiagaraEmitterParameters));
-
-			GPUContext->CombinedParamStore.CopyParameterDataToPaddedBuffer(InstanceData->ExternalParamData, ParmSize);
 
 			check(GPUContext->MaxUpdateIterations > 0);
 			InstanceData->SimStageData.AddZeroed(GPUContext->MaxUpdateIterations);
@@ -486,31 +499,67 @@ void FNiagaraGPUSystemTick::Destroy()
 	}
 }
 
+FUniformBufferRHIRef FNiagaraGPUSystemTick::GetUniformBuffer(EUniformBufferType Type, const FNiagaraComputeInstanceData* Instance, bool Current) const
+{
+	const int32 InterpOffset = Current
+		? 0
+		: (UBT_NumSystemTypes + Count * UBT_NumInstanceTypes);
+
+	if (Instance)
+	{
+		check(Type >= UBT_FirstInstanceType);
+		check(Type < UBT_NumTypes);
+
+		const int32 InstanceTypeIndex = Type - UBT_FirstInstanceType;
+
+		const int32 InstanceIndex = (Instance - GetInstanceData());
+		return UniformBuffers[InterpOffset + UBT_NumSystemTypes + Count * InstanceTypeIndex + InstanceIndex];
+	}
+
+	check(Type >= UBT_FirstSystemType);
+	check(Type < UBT_FirstInstanceType);
+
+	return UniformBuffers[InterpOffset + Type];
+}
+
+const uint8* FNiagaraGPUSystemTick::GetUniformBufferSource(EUniformBufferType Type, const FNiagaraComputeInstanceData* Instance, bool Current) const
+{
+	check(Type >= UBT_FirstSystemType);
+	check(Type < UBT_NumTypes);
+
+	switch (Type)
+	{
+		case UBT_Global:
+			return GlobalParamData + (Current ? 0 : sizeof(FNiagaraGlobalParameters));
+		case UBT_System:
+			return SystemParamData + (Current ? 0 : sizeof(FNiagaraSystemParameters));
+		case UBT_Owner:
+			return OwnerParamData + (Current ? 0 : sizeof(FNiagaraOwnerParameters));
+		case UBT_Emitter:
+		{
+			check(Instance);
+			return Instance->EmitterParamData + (Current ? 0 : sizeof(FNiagaraEmitterParameters));
+		}
+		case UBT_External:
+		{
+			check(Instance);
+			return Instance->ExternalParamData + (Current ? 0 : Instance->Context->ExternalCBufferLayout.ConstantBufferSize);
+		}
+	}
+
+	return nullptr;
+}
+
 //////////////////////////////////////////////////////////////////////////
 
 FNiagaraComputeExecutionContext::FNiagaraComputeExecutionContext()
 	: MainDataSet(nullptr)
 	, GPUScript(nullptr)
 	, GPUScript_RT(nullptr)
-	, GlobalCBufferLayout(TEXT("Niagara GPU Global CBuffer"))
-	, SystemCBufferLayout(TEXT("Niagara GPU System CBuffer"))
-	, OwnerCBufferLayout(TEXT("Niagara GPU Owner CBuffer"))
-	, EmitterCBufferLayout(TEXT("Niagara GPU Emitter CBuffer"))
 	, ExternalCBufferLayout(TEXT("Niagara GPU External CBuffer"))
 	, DataToRender(nullptr)
 
 {
-	GlobalCBufferLayout.ConstantBufferSize = sizeof(FNiagaraGlobalParameters);
-	GlobalCBufferLayout.ComputeHash();
-
-	SystemCBufferLayout.ConstantBufferSize = sizeof(FNiagaraSystemParameters);
-	SystemCBufferLayout.ComputeHash();
-
-	OwnerCBufferLayout.ConstantBufferSize = sizeof(FNiagaraOwnerParameters);
-	OwnerCBufferLayout.ComputeHash();
-
-	EmitterCBufferLayout.ConstantBufferSize = sizeof(FNiagaraEmitterParameters);
-	EmitterCBufferLayout.ComputeHash();
 }
 
 FNiagaraComputeExecutionContext::~FNiagaraComputeExecutionContext()
diff --git a/Engine/Plugins/FX/Niagara/Source/Niagara/Private/NiagaraSystemInstance.cpp b/Engine/Plugins/FX/Niagara/Source/Niagara/Private/NiagaraSystemInstance.cpp
--- a/Engine/Plugins/FX/Niagara/Source/Niagara/Private/NiagaraSystemInstance.cpp
+++ b/Engine/Plugins/FX/Niagara/Source/Niagara/Private/NiagaraSystemInstance.cpp
@@ -927,6 +927,7 @@ void FNiagaraSystemInstance::ResetInternal(bool bResetSimulations)
 	bLODDistanceIsValid = false;
 	TotalGPUParamSize = 0;
 	ActiveGPUEmitterCount = 0;
+	GPUParamIncludeInterpolation = false;
 	// Note: We do not need to update our bounds here as they are still valid
 
 	UNiagaraSystem* System = GetSystem();
@@ -2027,6 +2028,7 @@ void FNiagaraSystemInstance::Tick_Concurrent()
 	// Reset values that will be accumulated during emitter tick.
 	TotalGPUParamSize = 0;
 	ActiveGPUEmitterCount = 0;
+	GPUParamIncludeInterpolation = false;
 	UNiagaraSystem* System = GetSystem();
 
 	const int32 NumEmitters = Emitters.Num();
@@ -2064,7 +2066,7 @@ void FNiagaraSystemInstance::Tick_Concurrent()
 		}
 	}
 
-	bool FirstGpuEmitter = true;
+	int32 TotalCombinedParamStoreSize = 0;
 
 	// now tick all emitters
 	for (int32 EmitterIdx : EmitterExecutionOrder)
@@ -2075,22 +2077,26 @@ void FNiagaraSystemInstance::Tick_Concurrent()
 			Inst.Tick(CachedDeltaSeconds);
 		}
 
-		if (Inst.GetCachedEmitter() && Inst.GetCachedEmitter()->SimTarget == ENiagaraSimTarget::GPUComputeSim && Inst.GetGPUContext() != nullptr && !Inst.IsComplete())
+		if (Inst.GetCachedEmitter() && Inst.GetCachedEmitter()->SimTarget == ENiagaraSimTarget::GPUComputeSim && !Inst.IsComplete())
 		{
-			if (FirstGpuEmitter)
+			if (const FNiagaraComputeExecutionContext* GPUContext = Inst.GetGPUContext())
 			{
-				TotalGPUParamSize += 2 * sizeof(FNiagaraGlobalParameters);
-				TotalGPUParamSize += 2 * sizeof(FNiagaraSystemParameters);
-				TotalGPUParamSize += 2 * sizeof(FNiagaraOwnerParameters);
-				FirstGpuEmitter = false;
+				TotalCombinedParamStoreSize += GPUContext->CombinedParamStore.GetPaddedParameterSizeInBytes();
+				GPUParamIncludeInterpolation = GPUContext->HasInterpolationParameters || GPUParamIncludeInterpolation;
+				ActiveGPUEmitterCount++;
 			}
-
-			TotalGPUParamSize += 2 * sizeof(FNiagaraEmitterParameters);
-			TotalGPUParamSize += Inst.GetGPUContext()->CombinedParamStore.GetPaddedParameterSizeInBytes();
-			ActiveGPUEmitterCount++;
 		}
 	}
 
+	if (ActiveGPUEmitterCount)
+	{
+		const int32 InterpFactor = GPUParamIncludeInterpolation ? 2 : 1;
+
+		TotalGPUParamSize = InterpFactor * (sizeof(FNiagaraGlobalParameters) + sizeof(FNiagaraSystemParameters) + sizeof(FNiagaraOwnerParameters));
+		TotalGPUParamSize += InterpFactor * ActiveGPUEmitterCount * sizeof(FNiagaraEmitterParameters);
+		TotalGPUParamSize += TotalCombinedParamStoreSize;
+	}
+
 	// Update local bounds
 	if ( System->bFixedBounds )
 	{
diff --git a/Engine/Plugins/FX/Niagara/Source/Niagara/Public/NiagaraSystemInstance.h b/Engine/Plugins/FX/Niagara/Source/Niagara/Public/NiagaraSystemInstance.h
--- a/Engine/Plugins/FX/Niagara/Source/Niagara/Public/NiagaraSystemInstance.h
+++ b/Engine/Plugins/FX/Niagara/Source/Niagara/Public/NiagaraSystemInstance.h
@@ -445,6 +445,7 @@ public:
 	uint32 TotalGPUParamSize = 0;
 	uint32 ActiveGPUEmitterCount = 0;
 	int32 GPUDataInterfaceInstanceDataSize = 0;
+	bool GPUParamIncludeInterpolation = false;
 
 	struct FInstanceParameters
 	{
