Do not allow spawning world cascade or niagara systems on dedicated servers.
They will never tick so bAutoDestroy systems will leak.

diff --git a/Engine/Plugins/FX/Niagara/Source/Niagara/Private/NiagaraFunctionLibrary.cpp b/Engine/Plugins/FX/Niagara/Source/Niagara/Private/NiagaraFunctionLibrary.cpp
--- a/Engine/Plugins/FX/Niagara/Source/Niagara/Private/NiagaraFunctionLibrary.cpp
+++ b/Engine/Plugins/FX/Niagara/Source/Niagara/Private/NiagaraFunctionLibrary.cpp
@@ -50,21 +50,25 @@ UNiagaraFunctionLibrary::UNiagaraFunctionLibrary(const FObjectInitializer& Objec
 
 UNiagaraComponent* CreateNiagaraSystem(UNiagaraSystem* SystemTemplate, UWorld* World, AActor* Actor, bool bAutoDestroy, ENCPoolMethod PoolingMethod)
 {
-	UNiagaraComponent* NiagaraComponent;
-	if (PoolingMethod == ENCPoolMethod::None)
-	{
-		NiagaraComponent = NewObject<UNiagaraComponent>((Actor ? Actor : (UObject*)World));
-		NiagaraComponent->SetAsset(SystemTemplate);
-		NiagaraComponent->bAutoActivate = false;
-	}
-	else
+	UNiagaraComponent* NiagaraComponent = nullptr;
+
+	if (FApp::CanEverRender() && World && !World->IsNetMode(NM_DedicatedServer))
 	{
-		UNiagaraComponentPool* ComponentPool = FNiagaraWorldManager::Get(World)->GetComponentPool();
-		NiagaraComponent = ComponentPool->CreateWorldParticleSystem(SystemTemplate, World, PoolingMethod);
-	}	
+		if (PoolingMethod == ENCPoolMethod::None)
+		{
+			NiagaraComponent = NewObject<UNiagaraComponent>((Actor ? Actor : (UObject*)World));
+			NiagaraComponent->SetAsset(SystemTemplate);
+			NiagaraComponent->bAutoActivate = false;
+		}
+		else
+		{
+			UNiagaraComponentPool* ComponentPool = FNiagaraWorldManager::Get(World)->GetComponentPool();
+			NiagaraComponent = ComponentPool->CreateWorldParticleSystem(SystemTemplate, World, PoolingMethod);
+		}
 
-	NiagaraComponent->SetAutoDestroy(bAutoDestroy);
-	NiagaraComponent->bAllowAnyoneToDestroyMe = true;
+		NiagaraComponent->SetAutoDestroy(bAutoDestroy);
+		NiagaraComponent->bAllowAnyoneToDestroyMe = true;
+	}
 	return NiagaraComponent;
 }
 
@@ -91,17 +95,20 @@ UNiagaraComponent* UNiagaraFunctionLibrary::SpawnSystemAtLocation(const UObject*
 			if (!bShouldCull)
 			{
 				PSC = CreateNiagaraSystem(SystemTemplate, World, World->GetWorldSettings(), bAutoDestroy, PoolingMethod);
+				if(PSC)
+				{
 #if WITH_EDITORONLY_DATA
-				PSC->bWaitForCompilationOnActivate = true;
+					PSC->bWaitForCompilationOnActivate = true;
 #endif
-				PSC->RegisterComponentWithWorld(World);
+					PSC->RegisterComponentWithWorld(World);
 
-				PSC->SetAbsolute(true, true, true);
-				PSC->SetWorldLocationAndRotation(SpawnLocation, SpawnRotation);
-				PSC->SetRelativeScale3D(Scale);
-				if (bAutoActivate)
-				{
-					PSC->Activate(true);
+					PSC->SetAbsolute(true, true, true);
+					PSC->SetWorldLocationAndRotation(SpawnLocation, SpawnRotation);
+					PSC->SetRelativeScale3D(Scale);
+					if (bAutoActivate)
+					{
+						PSC->Activate(true);
+					}
 				}
 			}
 		}
@@ -139,28 +146,31 @@ UNiagaraComponent* UNiagaraFunctionLibrary::SpawnSystemAttached(UNiagaraSystem*
 			if (!bShouldCull)
 			{
 				PSC = CreateNiagaraSystem(SystemTemplate, AttachToComponent->GetWorld(), AttachToComponent->GetOwner(), bAutoDestroy, PoolingMethod);
-#if WITH_EDITOR
-				if (GForceNiagaraSpawnAttachedSolo > 0)
+				if(PSC)
 				{
-					PSC->SetForceSolo(true);
-				}
+#if WITH_EDITOR
+					if (GForceNiagaraSpawnAttachedSolo > 0)
+					{
+						PSC->SetForceSolo(true);
+					}
 #endif
-				PSC->RegisterComponentWithWorld(AttachToComponent->GetWorld());
+					PSC->RegisterComponentWithWorld(AttachToComponent->GetWorld());
 
-				PSC->AttachToComponent(AttachToComponent, FAttachmentTransformRules::KeepRelativeTransform, AttachPointName);
-				if (LocationType == EAttachLocation::KeepWorldPosition)
-				{
-					PSC->SetWorldLocationAndRotation(Location, Rotation);
-				}
-				else
-				{
-					PSC->SetRelativeLocationAndRotation(Location, Rotation);
-				}
-				PSC->SetRelativeScale3D(FVector(1.f));
+					PSC->AttachToComponent(AttachToComponent, FAttachmentTransformRules::KeepRelativeTransform, AttachPointName);
+					if (LocationType == EAttachLocation::KeepWorldPosition)
+					{
+						PSC->SetWorldLocationAndRotation(Location, Rotation);
+					}
+					else
+					{
+						PSC->SetRelativeLocationAndRotation(Location, Rotation);
+					}
+					PSC->SetRelativeScale3D(FVector(1.f));
 
-				if (bAutoActivate)
-				{
-					PSC->Activate();
+					if (bAutoActivate)
+					{
+						PSC->Activate();
+					}
 				}
 			}
 		}
@@ -210,52 +220,52 @@ UNiagaraComponent* UNiagaraFunctionLibrary::SpawnSystemAttached(
 				if (!bShouldCull)
 				{
 					PSC = CreateNiagaraSystem(SystemTemplate, World, AttachToComponent->GetOwner(), bAutoDestroy, PoolingMethod);
-				if (PSC)
-				{
-#if WITH_EDITOR
-					if (GForceNiagaraSpawnAttachedSolo > 0)
+					if (PSC)
 					{
-						PSC->SetForceSolo(true);
-					}
+#if WITH_EDITOR
+						if (GForceNiagaraSpawnAttachedSolo > 0)
+						{
+							PSC->SetForceSolo(true);
+						}
 #endif
-					PSC->SetupAttachment(AttachToComponent, AttachPointName);
+						PSC->SetupAttachment(AttachToComponent, AttachPointName);
 
-					if (LocationType == EAttachLocation::KeepWorldPosition)
-					{
-						const FTransform ParentToWorld = AttachToComponent->GetSocketTransform(AttachPointName);
-						const FTransform ComponentToWorld(Rotation, Location, Scale);
-						const FTransform RelativeTM = ComponentToWorld.GetRelativeTransform(ParentToWorld);
-						PSC->SetRelativeLocation_Direct(RelativeTM.GetLocation());
-						PSC->SetRelativeRotation_Direct(RelativeTM.GetRotation().Rotator());
-						PSC->SetRelativeScale3D_Direct(RelativeTM.GetScale3D());
-					}
-					else
-					{
-						PSC->SetRelativeLocation_Direct(Location);
-						PSC->SetRelativeRotation_Direct(Rotation);
-
-						if (LocationType == EAttachLocation::SnapToTarget)
+						if (LocationType == EAttachLocation::KeepWorldPosition)
 						{
-							// SnapToTarget indicates we "keep world scale", this indicates we we want the inverse of the parent-to-world scale 
-							// to calculate world scale at Scale 1, and then apply the passed in Scale
 							const FTransform ParentToWorld = AttachToComponent->GetSocketTransform(AttachPointName);
-							PSC->SetRelativeScale3D_Direct(Scale * ParentToWorld.GetSafeScaleReciprocal(ParentToWorld.GetScale3D()));
+							const FTransform ComponentToWorld(Rotation, Location, Scale);
+							const FTransform RelativeTM = ComponentToWorld.GetRelativeTransform(ParentToWorld);
+							PSC->SetRelativeLocation_Direct(RelativeTM.GetLocation());
+							PSC->SetRelativeRotation_Direct(RelativeTM.GetRotation().Rotator());
+							PSC->SetRelativeScale3D_Direct(RelativeTM.GetScale3D());
 						}
 						else
 						{
-							PSC->SetRelativeScale3D_Direct(Scale);
+							PSC->SetRelativeLocation_Direct(Location);
+							PSC->SetRelativeRotation_Direct(Rotation);
+
+							if (LocationType == EAttachLocation::SnapToTarget)
+							{
+								// SnapToTarget indicates we "keep world scale", this indicates we we want the inverse of the parent-to-world scale 
+								// to calculate world scale at Scale 1, and then apply the passed in Scale
+								const FTransform ParentToWorld = AttachToComponent->GetSocketTransform(AttachPointName);
+								PSC->SetRelativeScale3D_Direct(Scale * ParentToWorld.GetSafeScaleReciprocal(ParentToWorld.GetScale3D()));
+							}
+							else
+							{
+								PSC->SetRelativeScale3D_Direct(Scale);
+							}
 						}
-					}
 
-					PSC->RegisterComponentWithWorld(World);
-					if (bAutoActivate)
-					{
-						PSC->Activate(true);
-					}
+						PSC->RegisterComponentWithWorld(World);
+						if (bAutoActivate)
+						{
+							PSC->Activate(true);
+						}
 
-					// Notify the texture streamer so that PSC gets managed as a dynamic component.
-					IStreamingManager::Get().NotifyPrimitiveUpdated(PSC);
-				}
+						// Notify the texture streamer so that PSC gets managed as a dynamic component.
+						IStreamingManager::Get().NotifyPrimitiveUpdated(PSC);
+					}
 				}
 			}
 		}
diff --git a/Engine/Source/Runtime/Engine/Private/GameplayStatics.cpp b/Engine/Source/Runtime/Engine/Private/GameplayStatics.cpp
--- a/Engine/Source/Runtime/Engine/Private/GameplayStatics.cpp
+++ b/Engine/Source/Runtime/Engine/Private/GameplayStatics.cpp
@@ -966,22 +966,25 @@ void UGameplayStatics::PlayWorldCameraShake(const UObject* WorldContextObject, T
 UParticleSystemComponent* CreateParticleSystem(UParticleSystem* EmitterTemplate, UWorld* World, AActor* Actor, bool bAutoDestroy, EPSCPoolMethod PoolingMethod)
 {
 	//Defaulting to creating systems from a pool. Can be disabled via fx.ParticleSystemPool.Enable 0
-	UParticleSystemComponent* PSC;
+	UParticleSystemComponent* PSC = nullptr;
 
-	if (PoolingMethod != EPSCPoolMethod::None)
+	if (FApp::CanEverRender() && World && !World->IsNetMode(NM_DedicatedServer))
 	{
-		//If system is set to auto destroy the we should be safe to automatically allocate from a the world pool.
-		PSC = World->GetPSCPool().CreateWorldParticleSystem(EmitterTemplate, World, PoolingMethod);
-	}
-	else
-	{
-		PSC = NewObject<UParticleSystemComponent>((Actor ? Actor : (UObject*)World));
-		PSC->bAutoDestroy = bAutoDestroy;
-		PSC->bAllowAnyoneToDestroyMe = true;
-		PSC->SecondsBeforeInactive = 0.0f;
-		PSC->bAutoActivate = false;
-		PSC->SetTemplate(EmitterTemplate);
-		PSC->bOverrideLODMethod = false;
+		if (PoolingMethod != EPSCPoolMethod::None)
+		{
+			//If system is set to auto destroy the we should be safe to automatically allocate from a the world pool.
+			PSC = World->GetPSCPool().CreateWorldParticleSystem(EmitterTemplate, World, PoolingMethod);
+		}
+		else
+		{
+			PSC = NewObject<UParticleSystemComponent>((Actor ? Actor : (UObject*)World));
+			PSC->bAutoDestroy = bAutoDestroy;
+			PSC->bAllowAnyoneToDestroyMe = true;
+			PSC->SecondsBeforeInactive = 0.0f;
+			PSC->bAutoActivate = false;
+			PSC->SetTemplate(EmitterTemplate);
+			PSC->bOverrideLODMethod = false;
+		}
 	}
 
 	return PSC;
