Fixed cast shadows flag was ignored for Niagara meshes when ray tracing was enabled

diff --git a/Engine/Plugins/FX/Niagara/Source/Niagara/Private/NiagaraRendererMeshes.cpp b/Engine/Plugins/FX/Niagara/Source/Niagara/Private/NiagaraRendererMeshes.cpp
--- a/Engine/Plugins/FX/Niagara/Source/Niagara/Private/NiagaraRendererMeshes.cpp
+++ b/Engine/Plugins/FX/Niagara/Source/Niagara/Private/NiagaraRendererMeshes.cpp
@@ -731,7 +731,9 @@ void FNiagaraRendererMeshes::GetDynamicRayTracingInstances(FRayTracingMaterialGa
 #if !(UE_BUILD_SHIPPING || UE_BUILD_TEST)
 		MeshBatch.VisualizeLODIndex = LODIndex;
 #endif
-		MeshBatch.CastRayTracedShadow = MeshBatch.CastShadow;
+		MeshBatch.CastShadow = SceneProxy->CastsDynamicShadow();
+		MeshBatch.CastRayTracedShadow = SceneProxy->CastsDynamicShadow();
+
 		FMeshBatchElement& MeshBatchElement = MeshBatch.Elements[0];
 		MeshBatchElement.VertexFactoryUserData = VFs.VertexFactory.GetUniformBuffer();
 		MeshBatchElement.MinVertexIndex = Section.MinVertexIndex;
