Niagara - Fix an issue where the unapplied changes * would interfere with renaming in the scratch pad.

diff --git a/Engine/Plugins/FX/Niagara/Source/NiagaraEditor/Private/ViewModels/NiagaraScratchPadScriptViewModel.cpp b/Engine/Plugins/FX/Niagara/Source/NiagaraEditor/Private/ViewModels/NiagaraScratchPadScriptViewModel.cpp
--- a/Engine/Plugins/FX/Niagara/Source/NiagaraEditor/Private/ViewModels/NiagaraScratchPadScriptViewModel.cpp
+++ b/Engine/Plugins/FX/Niagara/Source/NiagaraEditor/Private/ViewModels/NiagaraScratchPadScriptViewModel.cpp
@@ -227,7 +227,7 @@ FSimpleDelegate& FNiagaraScratchPadScriptViewModel::OnRequestDiscardChanges()
 
 FText FNiagaraScratchPadScriptViewModel::GetDisplayNameInternal() const
 {
-	return FText::Format(LOCTEXT("DisplayNameFormat", "{0}{1}"), FText::FromString(OriginalScript->GetName()), bHasPendingChanges ? LOCTEXT("HasPendingChanges", "*") : FText());
+	return FText::FromString(OriginalScript->GetName());
 }
 
 void FNiagaraScratchPadScriptViewModel::OnScriptGraphChanged(const FEdGraphEditAction &Action)
diff --git a/Engine/Plugins/FX/Niagara/Source/NiagaraEditorWidgets/Private/SNiagaraScratchPad.cpp b/Engine/Plugins/FX/Niagara/Source/NiagaraEditorWidgets/Private/SNiagaraScratchPad.cpp
--- a/Engine/Plugins/FX/Niagara/Source/NiagaraEditorWidgets/Private/SNiagaraScratchPad.cpp
+++ b/Engine/Plugins/FX/Niagara/Source/NiagaraEditorWidgets/Private/SNiagaraScratchPad.cpp
@@ -18,6 +18,7 @@
 
 #include "Widgets/Text/SInlineEditableTextBlock.h"
 #include "Widgets/Input/SButton.h"
+#include "Widgets/Layout/SSpacer.h"
 #include "EditorStyleSet.h"
 #include "Framework/MultiBox/MultiBoxBuilder.h"
 #include "Framework/Application/SlateApplication.h"
@@ -149,6 +150,7 @@ class SNiagaraScratchPadScriptRow : public SCompoundWidget
 			SNew(SHorizontalBox)
 			+ SHorizontalBox::Slot()
 			.VAlign(VAlign_Center)
+			.AutoWidth()
 			.Padding(3, 0, 0, 0)
 			[
 				SAssignNew(NameEditableText, SInlineEditableTextBlock)
@@ -158,6 +160,21 @@ class SNiagaraScratchPadScriptRow : public SCompoundWidget
 				.OnTextCommitted(this, &SNiagaraScratchPadScriptRow::OnNameTextCommitted)
 			]
 			+ SHorizontalBox::Slot()
+			.VAlign(VAlign_Center)
+			.AutoWidth()
+			.Padding(2, 0, 0, 0)
+			[
+				SNew(STextBlock)
+				.Visibility(this, &SNiagaraScratchPadScriptRow::GetUnappliedChangesVisibility)
+				.TextStyle(FNiagaraEditorWidgetsStyle::Get(), "NiagaraEditor.ScratchPad.EditorHeaderText")
+				.Text(FText::FromString(TEXT("*")))
+				.ToolTipText(ScriptViewModel.ToSharedRef(), &FNiagaraScratchPadScriptViewModel::GetToolTip)
+			]
+			+ SHorizontalBox::Slot()
+			[
+				SNew(SSpacer)
+			]
+			+ SHorizontalBox::Slot()
 			.AutoWidth()
 			.Padding(1)
 			[
@@ -228,6 +245,11 @@ private:
 		return IsHovered();
 	}
 
+	EVisibility GetUnappliedChangesVisibility() const
+	{
+		return ScriptViewModel->HasUnappliedChanges() ? EVisibility::Visible : EVisibility::Collapsed;
+	}
+
 private:
 	UNiagaraScratchPadViewModel* ScratchPadViewModel;
 	TSharedPtr<FNiagaraScratchPadScriptViewModel> ScriptViewModel;
@@ -509,6 +531,17 @@ class SNiagaraScratchPadScriptEditor : public SCompoundWidget
 					.ToolTipText(ScriptViewModel.ToSharedRef(), &FNiagaraScratchPadScriptViewModel::GetToolTip)
 				]
 				+ SHorizontalBox::Slot()
+				.VAlign(VAlign_Center)
+				.AutoWidth()
+				.Padding(0, 0, 2, 0)
+				[
+					SNew(STextBlock)
+					.Visibility(this, &SNiagaraScratchPadScriptEditor::GetUnappliedChangesVisibility)
+					.TextStyle(FNiagaraEditorWidgetsStyle::Get(), "NiagaraEditor.ScratchPad.EditorHeaderText")
+					.Text(FText::FromString(TEXT("*")))
+					.ToolTipText(ScriptViewModel.ToSharedRef(), &FNiagaraScratchPadScriptViewModel::GetToolTip)
+				]
+				+ SHorizontalBox::Slot()
 				.AutoWidth()
 				.Padding(1)
 				[
@@ -533,6 +566,11 @@ private:
 		return ScriptViewModel->GetDisplayName();
 	}
 
+	EVisibility GetUnappliedChangesVisibility() const
+	{
+		return ScriptViewModel->HasUnappliedChanges() ? EVisibility::Visible : EVisibility::Collapsed;
+	}
+
 	FReply OnApplyButtonClicked()
 	{
 		ScriptViewModel->ApplyChanges();
