Niagara - Minor fix to tooltip formatting.

diff --git a/Engine/Plugins/FX/Niagara/Source/NiagaraEditor/Private/NiagaraEditorUtilities.cpp b/Engine/Plugins/FX/Niagara/Source/NiagaraEditor/Private/NiagaraEditorUtilities.cpp
--- a/Engine/Plugins/FX/Niagara/Source/NiagaraEditor/Private/NiagaraEditorUtilities.cpp
+++ b/Engine/Plugins/FX/Niagara/Source/NiagaraEditor/Private/NiagaraEditorUtilities.cpp
@@ -1483,10 +1483,10 @@ FText FNiagaraEditorUtilities::FormatVariableDescription(FText Description, FTex
 {
 	if (Description.IsEmptyOrWhitespace() == false)
 	{
-		return FText::Format(LOCTEXT("VariableDescriptionFormat", "{0}\nName: \"{1}\"\nType: {2}"), Description, Name, Type);
+		return FText::Format(LOCTEXT("VariableDescriptionFormat", "{0}\n\nName: {1}\n\nType: {2}"), Description, Name, Type);
 	}
 
-	return FText::Format(LOCTEXT("VariableDescriptionFormat_NoDesc", "Name: \"{0}\"\nType: {1}"), Name, Type);
+	return FText::Format(LOCTEXT("VariableDescriptionFormat_NoDesc", "Name: {0}\n\nType: {1}"), Name, Type);
 }
 
 void FNiagaraEditorUtilities::ResetSystemsThatReferenceSystemViewModel(const FNiagaraSystemViewModel& ReferencedSystemViewModel)
