Change the fallback texture used by the VolumeTetureDI to be a volume texture; silences validation errors with D3D11

diff --git a/Engine/Plugins/FX/Niagara/Source/Niagara/Private/NiagaraDataInterfaceVolumeTexture.cpp b/Engine/Plugins/FX/Niagara/Source/Niagara/Private/NiagaraDataInterfaceVolumeTexture.cpp
--- a/Engine/Plugins/FX/Niagara/Source/Niagara/Private/NiagaraDataInterfaceVolumeTexture.cpp
+++ b/Engine/Plugins/FX/Niagara/Source/Niagara/Private/NiagaraDataInterfaceVolumeTexture.cpp
@@ -269,8 +269,8 @@ public:
 				ComputeShaderRHI,
 				TextureParam,
 				SamplerParam,
-				GBlackTexture->SamplerStateRHI,
-				GBlackTexture->TextureRHI
+				GBlackVolumeTexture->SamplerStateRHI,
+				GBlackVolumeTexture->TextureRHI
 			);
 			SetShaderValue(RHICmdList, ComputeShaderRHI, Dimensions, FVector::ZeroVector);
 		}
