Fix for oddness with particle attributes defined in sim stages.

Dataset Compiled data was reading from the wrong script so had an attribute set inconstent with the shader.

diff --git a/Engine/Plugins/FX/Niagara/Source/Niagara/Private/NiagaraSystem.cpp b/Engine/Plugins/FX/Niagara/Source/Niagara/Private/NiagaraSystem.cpp
--- a/Engine/Plugins/FX/Niagara/Source/Niagara/Private/NiagaraSystem.cpp
+++ b/Engine/Plugins/FX/Niagara/Source/Niagara/Private/NiagaraSystem.cpp
@@ -1116,10 +1116,18 @@ void UNiagaraSystem::InitEmitterDataSetCompiledData(FNiagaraDataSetCompiledData&
 {
 	DataSetToInit.Empty();
 
-	DataSetToInit.Variables = InAssociatedEmitter->UpdateScriptProps.Script->GetVMExecutableData().Attributes;
-	for (const FNiagaraVariable& Var : InAssociatedEmitter->SpawnScriptProps.Script->GetVMExecutableData().Attributes)
+	if (InAssociatedEmitter->SimTarget == ENiagaraSimTarget::GPUComputeSim)
 	{
-		DataSetToInit.Variables.AddUnique(Var);
+		DataSetToInit.Variables = InAssociatedEmitter->GetGPUComputeScript()->GetVMExecutableData().Attributes;
+	}
+	else
+	{
+		DataSetToInit.Variables = InAssociatedEmitter->UpdateScriptProps.Script->GetVMExecutableData().Attributes;
+
+		for (const FNiagaraVariable& Var : InAssociatedEmitter->SpawnScriptProps.Script->GetVMExecutableData().Attributes)
+		{
+			DataSetToInit.Variables.AddUnique(Var);
+		}
 	}
 
 	DataSetToInit.bRequiresPersistentIDs = InAssociatedEmitter->RequiresPersistentIDs() || DataSetToInit.Variables.Contains(SYS_PARAM_PARTICLES_ID);
diff --git a/Engine/Plugins/FX/Niagara/Source/NiagaraCore/Private/NiagaraCustomVersion.cpp b/Engine/Plugins/FX/Niagara/Source/NiagaraCore/Private/NiagaraCustomVersion.cpp
--- a/Engine/Plugins/FX/Niagara/Source/NiagaraCore/Private/NiagaraCustomVersion.cpp
+++ b/Engine/Plugins/FX/Niagara/Source/NiagaraCore/Private/NiagaraCustomVersion.cpp
@@ -8,4 +8,4 @@ const FGuid FNiagaraCustomVersion::GUID(0xFCF57AFA, 0x50764283, 0xB9A9E658, 0xFF
 // Register the custom version with core
 FCustomVersionRegistration GRegisterNiagaraCustomVersion(FNiagaraCustomVersion::GUID, FNiagaraCustomVersion::LatestVersion, TEXT("NiagaraVer"));
 
-const FGuid FNiagaraCustomVersion::LatestScriptCompileVersion(0xD52A505F, 0x0CAD485D, 0x8DE36D97, 0x3857F2F6);
\ No newline at end of file
+const FGuid FNiagaraCustomVersion::LatestScriptCompileVersion(0x05940817, 0x4BC44989, 0x8583FE3D, 0x7EFE35B9);
\ No newline at end of file
