- Fix for skinning mode none accessing bones with scale applied to them

#rnx

diff --git a/Engine/Plugins/FX/Niagara/Source/Niagara/Classes/NDISkeletalMeshCommon.h b/Engine/Plugins/FX/Niagara/Source/Niagara/Classes/NDISkeletalMeshCommon.h
--- a/Engine/Plugins/FX/Niagara/Source/Niagara/Classes/NDISkeletalMeshCommon.h
+++ b/Engine/Plugins/FX/Niagara/Source/Niagara/Classes/NDISkeletalMeshCommon.h
@@ -236,7 +236,7 @@ struct FSkinnedPositionAccessorHelper<TNDISkelMesh_SkinningModeNone>
 		const int32 NumRealBones = Accessor.Mesh->RefSkeleton.GetRawBoneNum();
 		if (BoneIndex < NumRealBones)
 		{
-			return Accessor.Mesh->GetComposedRefPoseMatrix(BoneIndex).ToQuat();
+			return Accessor.Mesh->GetComposedRefPoseMatrix(BoneIndex).GetMatrixWithoutScale().ToQuat();
 		}
 
 		const FTransform& RefTransform = Accessor.Mesh->RefSkeleton.GetRefBonePose()[BoneIndex];
