Add optional command list hint to reduce risk of GPU hangs while running many shader stages (by default will be inserted each 10 stages)

diff --git a/Engine/Plugins/FX/Niagara/Source/Niagara/Private/NiagaraEmitterInstanceBatcher.cpp b/Engine/Plugins/FX/Niagara/Source/Niagara/Private/NiagaraEmitterInstanceBatcher.cpp
--- a/Engine/Plugins/FX/Niagara/Source/Niagara/Private/NiagaraEmitterInstanceBatcher.cpp
+++ b/Engine/Plugins/FX/Niagara/Source/Niagara/Private/NiagaraEmitterInstanceBatcher.cpp
@@ -61,6 +61,14 @@ static FAutoConsoleVariableRef CVarNiagaraGpuMaxQueuedRenderFrames(
 	ECVF_Default
 );
 
+int32 GNiagaraGpuSubmitCommandHint = 10;
+static FAutoConsoleVariableRef CVarNiagaraGpuSubmitCommandHint(
+	TEXT("fx.NiagaraGpuSubmitCommandHint"),
+	GNiagaraGpuSubmitCommandHint,
+	TEXT("If non-zero, a hint will be issued between the set number of dispatches within sequence of shader stages.\n"),
+	ECVF_Default
+);
+
 const FName NiagaraEmitterInstanceBatcher::Name(TEXT("NiagaraEmitterInstanceBatcher"));
 
 FFXSystemInterface* NiagaraEmitterInstanceBatcher::GetInterface(const FName& InName)
@@ -505,6 +513,12 @@ void NiagaraEmitterInstanceBatcher::DispatchMultipleStages(const FNiagaraGPUSyst
 			}
 
 			PostStageInterface(Tick, Instance, RHICmdList, ComputeShader, SimulationStageIndex);
+
+			// for long running dispatches we may want to issue a hint to the command list to break things up
+			if (GNiagaraGpuSubmitCommandHint && ((SimulationStageIndex + 1) % GNiagaraGpuSubmitCommandHint) == 0)
+			{
+				RHICmdList.SubmitCommandsHint();
+			}
 		}
 	}
 	else
@@ -615,13 +629,9 @@ void NiagaraEmitterInstanceBatcher::ResizeBuffersAndGatherResources(FOverlappabl
 			uint32 NumBufferIterations = 1;
 			if (Tick->NumInstancesWithSimStages > 0)
 			{
-
-				bool HasRunParticleStage = false;
-
 				const uint32 NumStages = Instance.Context->MaxUpdateIterations;
 				if (NumStages > 1)
 				{
-					HasRunParticleStage = true;
 					for (uint32 SimulationStageIndex = 0; SimulationStageIndex < NumStages; SimulationStageIndex++)
 					{
 						if (SimulationStageIndex != 0)
@@ -1416,6 +1426,7 @@ void NiagaraEmitterInstanceBatcher::SetDataInterfaceParameters(const TArray<FNia
 			Context.ComputeInstanceData = Instance;
 			Context.SimulationStageIndex = SimulationStageIndex;
 			Context.IsOutputStage = Instance->IsOutputStage(Interface, SimulationStageIndex);
+			Context.IsIterationStage = Instance->IsIterationStage(Interface, SimulationStageIndex);
 			DIParam.DIType.Get(PointerTable.DITypes)->SetParameters(DIParam.Parameters.Get(), RHICmdList, Context);
 		}
 
