Deprecated BP functions which should only be used for debugging

diff --git a/Engine/Plugins/FX/Niagara/Source/Niagara/Public/NiagaraComponent.h b/Engine/Plugins/FX/Niagara/Source/Niagara/Public/NiagaraComponent.h
--- a/Engine/Plugins/FX/Niagara/Source/Niagara/Public/NiagaraComponent.h
+++ b/Engine/Plugins/FX/Niagara/Source/Niagara/Public/NiagaraComponent.h
@@ -359,15 +359,15 @@ public:
 	void SetVariableMaterial(FName InVariableName, UMaterialInterface* Object);
 
 	/** Debug accessors for getting positions in blueprints. */
-	UFUNCTION(BlueprintCallable, Category = Niagara, meta = (DisplayName = "Get Niagara Emitter Positions"))
+	UFUNCTION(BlueprintCallable, Category = Niagara, meta = (DisplayName = "Get Niagara Emitter Positions", DeprecatedFunction, DeprecationMessage = "Get Niagara Emitter Positions is deprecated, use the particle export DI inside your emitter instead."))
 	TArray<FVector> GetNiagaraParticlePositions_DebugOnly(const FString& InEmitterName);
 	
 	/** Debug accessors for getting a float attribute array in blueprints.  The attribute name should be without namespaces. For example for "Particles.Position", send "Position". */
-	UFUNCTION(BlueprintCallable, Category = Niagara, meta = (DisplayName = "Get Niagara Emitter Float Attrib"))
+	UFUNCTION(BlueprintCallable, Category = Niagara, meta = (DisplayName = "Get Niagara Emitter Float Attrib", DeprecatedFunction, DeprecationMessage = "Get Niagara Emitter Float Attrib is deprecated, use the particle export DI inside your emitter instead."))
 	TArray<float> GetNiagaraParticleValues_DebugOnly(const FString& InEmitterName, const FString& InValueName);
 	
 	/** Debug accessors for getting a FVector attribute array in blueprints. The attribute name should be without namespaces. For example for "Particles.Position", send "Position". */
-	UFUNCTION(BlueprintCallable, Category = Niagara, meta = (DisplayName = "Get Niagara Emitter Vec3 Attrib"))
+	UFUNCTION(BlueprintCallable, Category = Niagara, meta = (DisplayName = "Get Niagara Emitter Vec3 Attrib", DeprecatedFunction, DeprecationMessage = "Get Niagara Emitter Vec3 Attrib is deprecated, use the particle export DI inside your emitter instead."))
 	TArray<FVector> GetNiagaraParticleValueVec3_DebugOnly(const FString& InEmitterName, const FString& InValueName);
 
 	/** Resets the System to it's initial pre-simulated state. */
