- Fix crash due to no objects being modified in the property notification when fixing cpu access on a NPC

diff --git a/Engine/Plugins/FX/Niagara/Source/NiagaraEditor/Private/Widgets/SNiagaraParameterCollection.cpp b/Engine/Plugins/FX/Niagara/Source/NiagaraEditor/Private/Widgets/SNiagaraParameterCollection.cpp
--- a/Engine/Plugins/FX/Niagara/Source/NiagaraEditor/Private/Widgets/SNiagaraParameterCollection.cpp
+++ b/Engine/Plugins/FX/Niagara/Source/NiagaraEditor/Private/Widgets/SNiagaraParameterCollection.cpp
@@ -833,6 +833,11 @@ void SNiagaraParameterCollection::ParameterContentColumnWidthChanged(float Width
 
 void SNiagaraParameterCollection::NotifyPostChange(const FPropertyChangedEvent& PropertyChangedEvent, FProperty* PropertyThatChanged)
 {
+	if (PropertyChangedEvent.GetNumObjectsBeingEdited() == 0)
+	{
+		return;
+	}
+
 	for (TSharedRef<INiagaraParameterViewModel> Parameter : Collection->GetParameters())
 	{
 		if (Parameter->GetDefaultValueType() == INiagaraParameterViewModel::EDefaultValueType::Object)
