Fixed multi-ribbon Niagara assets draw order changing randomly between frames as particle are being spawned.

diff --git a/Engine/Plugins/FX/Niagara/Source/Niagara/Private/NiagaraRendererRibbons.cpp b/Engine/Plugins/FX/Niagara/Source/Niagara/Private/NiagaraRendererRibbons.cpp
--- a/Engine/Plugins/FX/Niagara/Source/Niagara/Private/NiagaraRendererRibbons.cpp
+++ b/Engine/Plugins/FX/Niagara/Source/Niagara/Private/NiagaraRendererRibbons.cpp
@@ -735,6 +735,9 @@ FNiagaraDynamicDataBase* FNiagaraRendererRibbons::GenerateDynamicData(const FNia
 			}
 			DynamicData->MultiRibbonInfos.AddZeroed(MultiRibbonSortedIndices.Num());
 
+			// Sort the ribbons by ID so that the draw order stays consistent.
+			MultiRibbonSortedIndices.KeySort(TLess<FNiagaraID>());
+
 			uint32 RibbonIndex = 0;
 			for (TPair<FNiagaraID, TArray<int32>>& Pair : MultiRibbonSortedIndices)
 			{
@@ -759,6 +762,9 @@ FNiagaraDynamicDataBase* FNiagaraRendererRibbons::GenerateDynamicData(const FNia
 			}
 			DynamicData->MultiRibbonInfos.AddZeroed(MultiRibbonSortedIndices.Num());
 
+			// Sort the ribbons by ID so that the draw order stays consistent.
+			MultiRibbonSortedIndices.KeySort(TLess<int32>());
+
 			uint32 RibbonIndex = 0;
 			for (TPair<int32, TArray<int32>>& Pair : MultiRibbonSortedIndices)
 			{
diff --git a/Engine/Plugins/FX/Niagara/Source/Niagara/Public/NiagaraTypes.h b/Engine/Plugins/FX/Niagara/Source/Niagara/Public/NiagaraTypes.h
--- a/Engine/Plugins/FX/Niagara/Source/Niagara/Public/NiagaraTypes.h
+++ b/Engine/Plugins/FX/Niagara/Source/Niagara/Public/NiagaraTypes.h
@@ -233,7 +233,7 @@ struct FNiagaraID
 
 	bool operator==(const FNiagaraID& Other)const { return Index == Other.Index && AcquireTag == Other.AcquireTag; }
 	bool operator!=(const FNiagaraID& Other)const { return !(*this == Other); }
-	bool operator<(const FNiagaraID& Other)const { return Index < Other.Index && AcquireTag < Other.AcquireTag; }
+	bool operator<(const FNiagaraID& Other)const { return Index < Other.Index || (Index == Other.Index && AcquireTag < Other.AcquireTag); }
 };
 
 #define NIAGARA_INVALID_ID (FNiagaraID({(INDEX_NONE), (INDEX_NONE)}))
