Preventing making non-parameter pins on map get nodes by dragging pin links.

diff --git a/Engine/Plugins/FX/Niagara/Source/NiagaraEditor/Private/EdGraphSchema_Niagara.cpp b/Engine/Plugins/FX/Niagara/Source/NiagaraEditor/Private/EdGraphSchema_Niagara.cpp
--- a/Engine/Plugins/FX/Niagara/Source/NiagaraEditor/Private/EdGraphSchema_Niagara.cpp
+++ b/Engine/Plugins/FX/Niagara/Source/NiagaraEditor/Private/EdGraphSchema_Niagara.cpp
@@ -919,6 +919,23 @@ const FPinConnectionResponse UEdGraphSchema_Niagara::CanCreateConnection(const U
 		return FPinConnectionResponse(CONNECT_RESPONSE_DISALLOW, TEXT("Directions are not compatible"));
 	}
 
+	// Do not allow making connections off of dynamic add pins to non parameter map associated pins 
+	auto GetPinsAreInvalidAddPinCombination = [](const UEdGraphPin* A, const UEdGraphPin* B)->bool {
+		if (A->PinType.PinSubCategory == UNiagaraNodeWithDynamicPins::AddPinSubCategory)
+		{
+			if (B->PinType.PinCategory != PinCategoryType)
+			{
+				return true;
+			}
+		}
+		return false;
+	};
+
+	if (GetPinsAreInvalidAddPinCombination(PinA, PinB) || GetPinsAreInvalidAddPinCombination(PinB, PinA))
+	{
+		return FPinConnectionResponse(CONNECT_RESPONSE_DISALLOW, TEXT("Cannot make connections to or from add pins for non-parameter types"));
+	}
+
 	// Check for a circular connection before checking any type compatibility
 	TSet<const UEdGraphNode*> VisitedNodes;
 	if (UEdGraphSchema_Niagara::CheckCircularConnection(VisitedNodes, OutputPin->GetOwningNode(), InputPin->GetOwningNode()))
diff --git a/Engine/Plugins/FX/Niagara/Source/NiagaraEditor/Private/NiagaraNodeParameterMapGet.cpp b/Engine/Plugins/FX/Niagara/Source/NiagaraEditor/Private/NiagaraNodeParameterMapGet.cpp
--- a/Engine/Plugins/FX/Niagara/Source/NiagaraEditor/Private/NiagaraNodeParameterMapGet.cpp
+++ b/Engine/Plugins/FX/Niagara/Source/NiagaraEditor/Private/NiagaraNodeParameterMapGet.cpp
@@ -94,6 +94,13 @@ UEdGraphPin* UNiagaraNodeParameterMapGet::CreateDefaultPin(UEdGraphPin* OutputPi
 
 	UEdGraphPin* DefaultPin = CreatePin(EEdGraphPinDirection::EGPD_Input, OutputPin->PinType, TEXT(""));
 
+	// make sure the new pin name is legal
+	FName OutName;
+	if (FNiagaraEditorUtilities::DecomposeVariableNamespace(OutputPin->GetFName(), OutName).Num() == 0)
+	{
+		OutputPin->PinName = FName(FNiagaraConstants::LocalNamespace.ToString() + "." + OutputPin->GetName());
+	}
+
 	// we make the pin read only because the default value is set in the parameter panel unless the default mode is set to "custom" by the user
 	DefaultPin->bNotConnectable = true;
 	DefaultPin->bDefaultValueIsReadOnly = true;
