Fix for issue where the Oscilloscope NDI not working in PIE when Create New Audio Device For PIE is enabled in the Editor Settings. Issue was that the fix for the Oscilloscope only using half the buffer wasn't properly set up for FNiagaraDataInterfaceProxyOscilloscope::OnNewDeviceCreated the same way it is in FNiagaraDataInterfaceProxyOscilloscope::RegisterToAllAudioDevices(). To prevent future errors like this, I'm going to consolidate these to one function in Main.

diff --git a/Engine/Plugins/FX/Niagara/Source/Niagara/Private/NiagaraDataInterfaceAudioOscilloscope.cpp b/Engine/Plugins/FX/Niagara/Source/Niagara/Private/NiagaraDataInterfaceAudioOscilloscope.cpp
--- a/Engine/Plugins/FX/Niagara/Source/Niagara/Private/NiagaraDataInterfaceAudioOscilloscope.cpp
+++ b/Engine/Plugins/FX/Niagara/Source/Niagara/Private/NiagaraDataInterfaceAudioOscilloscope.cpp
@@ -99,8 +99,14 @@ void FNiagaraDataInterfaceProxyOscilloscope::OnNewDeviceCreated(Audio::FDeviceId
 	if (bIsSubmixListenerRegistered)
 	{
 		check(!SubmixListeners.Contains(InID));
-		SubmixListeners.Emplace(InID, new FNiagaraSubmixListener(PatchMixer, Resolution * 2, InID, SubmixRegisteredTo));
-		SubmixListeners[InID]->RegisterToSubmix();
+
+		FAudioDeviceHandle DeviceHandle = FAudioDeviceManager::Get()->GetAudioDevice(InID);
+		if (ensure(DeviceHandle))
+		{
+			const int32 NumSamplesToPop = Align(FMath::FloorToInt(ScopeInMilliseconds / 1000.0f * DeviceHandle->GetSampleRate()) * AUDIO_MIXER_MAX_OUTPUT_CHANNELS, 4);
+			SubmixListeners.Emplace(InID, new FNiagaraSubmixListener(PatchMixer, NumSamplesToPop, InID, SubmixRegisteredTo));
+			SubmixListeners[InID]->RegisterToSubmix();
+		}
 	}
 }
 
