Niagara - Prevent changing the namespace on user parameters until the operation can be better supported.

diff --git a/Engine/Plugins/FX/Niagara/Source/NiagaraEditor/Private/NiagaraEditorSettings.cpp b/Engine/Plugins/FX/Niagara/Source/NiagaraEditor/Private/NiagaraEditorSettings.cpp
--- a/Engine/Plugins/FX/Niagara/Source/NiagaraEditor/Private/NiagaraEditorSettings.cpp
+++ b/Engine/Plugins/FX/Niagara/Source/NiagaraEditor/Private/NiagaraEditorSettings.cpp
@@ -120,7 +120,8 @@ void UNiagaraEditorSettings::SetupNamespaceMetadata()
 			.SetBackgroundColor(FLinearColor(FColor(91, 161, 194)))
 			.SetSortId(0)
 			.AddOption(ENiagaraNamespaceMetadataOptions::HideInScript)
-			.AddOption(ENiagaraNamespaceMetadataOptions::PreventEditingNamespaceModifier),
+			.AddOption(ENiagaraNamespaceMetadataOptions::PreventEditingNamespaceModifier)
+			.AddOption(ENiagaraNamespaceMetadataOptions::PreventEditingNamespace),
 		FNiagaraNamespaceMetadata({FNiagaraConstants::ParameterCollectionNamespace})
 			.SetDisplayName(LOCTEXT("NiagaraParameterCollectionDisplayName", "NPC"))
 			.SetDisplayNameLong(LOCTEXT("NiagaraParameterCollectionDisplayNameLong", "Niagara Parameter Collection"))
