Niagara - Fix some last issues with the new parameters workflow reported this week.
+ Change the display of outputs so that they show the module namespace modifier in scripts.
+ Added the concept of required modifiers, which allows an additional modifier to be added to Output.Module.
+ Remove the * for added parameters and instead show a lock on parameters from external graphs with a tooltip about why they can't be edited.
+ Add duplication options to the parameters panel for duplicating parameters, and duplicating parameters into new namespaces, and duplicating with new namespace modifiers.
+ Rework the namespace modifier menu so that instead of Add/Remove/Edit it lists compatible namespace modifiers, and then "Custom..." which will add a custom modifier or set the current modifier to edit, and "None" which will remove the modifier.
+ Fix parameter renaming in the system view so that it also renames set variable node usages and linked inputs if they don't have external graph references.

diff --git a/Engine/Plugins/FX/Niagara/Source/NiagaraEditor/Private/NiagaraActions.cpp b/Engine/Plugins/FX/Niagara/Source/NiagaraEditor/Private/NiagaraActions.cpp
--- a/Engine/Plugins/FX/Niagara/Source/NiagaraEditor/Private/NiagaraActions.cpp
+++ b/Engine/Plugins/FX/Niagara/Source/NiagaraEditor/Private/NiagaraActions.cpp
@@ -47,6 +47,7 @@ FNiagaraParameterAction::FNiagaraParameterAction(const FNiagaraVariable& InParam
 	: FEdGraphSchemaAction(MoveTemp(InNodeCategory), MoveTemp(InMenuDesc), MoveTemp(InToolTip), InGrouping, MoveTemp(InKeywords), InSectionID)
 	, Parameter(InParameter)
 	, ReferenceCollection(InReferenceCollection)
+	, bIsExternallyReferenced(false)
 	, bNamespaceModifierRenamePending(false)
 {
 }
@@ -55,6 +56,7 @@ FNiagaraParameterAction::FNiagaraParameterAction(const FNiagaraVariable& InParam
 	FText InNodeCategory, FText InMenuDesc, FText InToolTip, const int32 InGrouping, FText InKeywords, int32 InSectionID)
 	: FEdGraphSchemaAction(MoveTemp(InNodeCategory), MoveTemp(InMenuDesc), MoveTemp(InToolTip), InGrouping, MoveTemp(InKeywords), InSectionID)
 	, Parameter(InParameter)
+	, bIsExternallyReferenced(false)
 	, bNamespaceModifierRenamePending(false)
 {
 }
diff --git a/Engine/Plugins/FX/Niagara/Source/NiagaraEditor/Private/NiagaraEditorSettings.cpp b/Engine/Plugins/FX/Niagara/Source/NiagaraEditor/Private/NiagaraEditorSettings.cpp
--- a/Engine/Plugins/FX/Niagara/Source/NiagaraEditor/Private/NiagaraEditorSettings.cpp
+++ b/Engine/Plugins/FX/Niagara/Source/NiagaraEditor/Private/NiagaraEditorSettings.cpp
@@ -10,8 +10,9 @@ FNiagaraNamespaceMetadata::FNiagaraNamespaceMetadata()
 {
 }
 
-FNiagaraNamespaceMetadata::FNiagaraNamespaceMetadata(TArray<FName> InNamespaces)
+FNiagaraNamespaceMetadata::FNiagaraNamespaceMetadata(TArray<FName> InNamespaces, FName InRequiredNamespaceModifier)
 	: Namespaces(InNamespaces)
+	, RequiredNamespaceModifier(InRequiredNamespaceModifier)
 	, BackgroundColor(FLinearColor::Black)
 	, ForegroundStyle("NiagaraEditor.ParameterName.NamespaceText")
 	, SortId(TNumericLimits<int32>::Max())
@@ -33,33 +34,40 @@ UNiagaraEditorSettings::UNiagaraEditorSettings(const FObjectInitializer& ObjectI
 
 void UNiagaraEditorSettings::SetupNamespaceMetadata()
 {
+	DefaultNamespaceMetadata = FNiagaraNamespaceMetadata({ NAME_None })
+		.SetDisplayName(LOCTEXT("DefaultDisplayName", "None"))
+		.SetDescription(LOCTEXT("DefaultDescription", "Non-standard unknown namespace."))
+		.SetBackgroundColor(FLinearColor(FColor(102, 102, 102)))
+		.AddOption(ENiagaraNamespaceMetadataOptions::PreventEditingNamespace)
+		.AddOption(ENiagaraNamespaceMetadataOptions::PreventEditingNamespaceModifier)
+		.AddOption(ENiagaraNamespaceMetadataOptions::PreventEditingName);
+
 	NamespaceMetadata =
 	{
-		FNiagaraNamespaceMetadata({NAME_None})
-			.SetDisplayName(LOCTEXT("DefaultDisplayName", "None"))
-			.SetDescription(LOCTEXT("DefaultDescription", "Arbitrary sub-namespace for specifying module specific dataset attributes, or calling nested modules."))
-			.SetBackgroundColor(FLinearColor(FColor(102, 102, 102)))
-			.AddOption(ENiagaraNamespaceMetadataOptions::PreventEditingNamespace)
-			.AddOption(ENiagaraNamespaceMetadataOptions::PreventEditingNamespaceModifier)
-			.AddOption(ENiagaraNamespaceMetadataOptions::PreventEditingName),
 		FNiagaraNamespaceMetadata({FNiagaraConstants::SystemNamespace})
 			.SetDisplayName(LOCTEXT("SystemDisplayName", "System"))
 			.SetDisplayNameLong(LOCTEXT("SystemDisplayNameLong", "System Attributes"))
 			.SetDescription(LOCTEXT("SystemDescription", "Persistent attribute in the system which is written in a system\n stage and can be read anywhere."))
 			.SetBackgroundColor(FLinearColor(FColor(49, 113, 142)))
-			.SetSortId(10),
+			.SetSortId(10)
+			.AddOptionalNamespaceModifier(FNiagaraConstants::ModuleNamespace)
+			.AddOptionalNamespaceModifier(FNiagaraConstants::InitialNamespace),
 		FNiagaraNamespaceMetadata({FNiagaraConstants::EmitterNamespace})
 			.SetDisplayName(LOCTEXT("EmitterDisplayName", "Emitter"))
 			.SetDisplayNameLong(LOCTEXT("EmitterDisplayNameLong", "Emitter Attributes"))
 			.SetDescription(LOCTEXT("EmitterDescription", "Persistent attribute which is written in a emitter\nstage and can be read in emitter and particle stages."))
 			.SetBackgroundColor(FLinearColor(FColor(145, 99, 56)))
-			.SetSortId(20),
+			.SetSortId(20)
+			.AddOptionalNamespaceModifier(FNiagaraConstants::ModuleNamespace)
+			.AddOptionalNamespaceModifier(FNiagaraConstants::InitialNamespace),
 		FNiagaraNamespaceMetadata({FNiagaraConstants::ParticleAttributeNamespace})
 			.SetDisplayName(LOCTEXT("ParticleDisplayName", "Particles"))
 			.SetDisplayNameLong(LOCTEXT("ParticleDisplayNameLong", "Particle Attributes"))
 			.SetDescription(LOCTEXT("ParticleDescription", "Persistent attribute which is written in a particle\nstage and can be read in particle stages."))
 			.SetBackgroundColor(FLinearColor(FColor(72, 130, 71)))
-			.SetSortId(30),
+			.SetSortId(30)
+			.AddOptionalNamespaceModifier(FNiagaraConstants::ModuleNamespace)
+			.AddOptionalNamespaceModifier(FNiagaraConstants::InitialNamespace),
 		FNiagaraNamespaceMetadata({FNiagaraConstants::ModuleNamespace})
 			.SetDisplayName(LOCTEXT("ModuleDisplayName", "Input"))
 			.SetDisplayNameLong(LOCTEXT("ModuleDisplayNameLong", "Module Inputs"))
@@ -67,7 +75,7 @@ void UNiagaraEditorSettings::SetupNamespaceMetadata()
 			.SetBackgroundColor(FLinearColor(FColor(136, 66, 65)))
 			.SetSortId(40)
 			.AddOption(ENiagaraNamespaceMetadataOptions::HideInSystem),
-		FNiagaraNamespaceMetadata({FNiagaraConstants::OutputNamespace, FNiagaraConstants::ModuleNamespace})
+		FNiagaraNamespaceMetadata({FNiagaraConstants::OutputNamespace}, FNiagaraConstants::ModuleNamespace)
 			.SetDisplayName(LOCTEXT("ModuleOutputDisplayName", "Output"))
 			.SetDisplayNameLong(LOCTEXT("ModuleOutputDisplayNameLong", "Module Outputs"))
 			.SetDescription(LOCTEXT("ModuleOutputDescription", "A transient value which the module author has decided might be useful to other modules further down in the stage.\nTransient values do not persist from frame to frame, or between stages, e.g. emitter to particle, or spawn to update."))
@@ -76,15 +84,6 @@ void UNiagaraEditorSettings::SetupNamespaceMetadata()
 			.AddOption(ENiagaraNamespaceMetadataOptions::AdvancedInScript)
 			.AddOption(ENiagaraNamespaceMetadataOptions::AdvancedInSystem)
 			.AddOption(ENiagaraNamespaceMetadataOptions::PreventCreatingInSystemEditor),
-		FNiagaraNamespaceMetadata({FNiagaraConstants::OutputNamespace})
-			.SetDisplayName(LOCTEXT("OutputDisplayName", "Output"))
-			.SetDescription(LOCTEXT("OutputDescription", "A transient value which the module author has decided might be useful to other modules further down in the stage.\nTransient values do not persist from frame to frame, or between stages, e.g. emitter to particle, or spawn to update."))
-			.SetBackgroundColor(FLinearColor(FColor(108, 87, 131)))
-			.AddOption(ENiagaraNamespaceMetadataOptions::AdvancedInScript)
-			.AddOption(ENiagaraNamespaceMetadataOptions::AdvancedInSystem)
-			.AddOption(ENiagaraNamespaceMetadataOptions::PreventEditingNamespace)
-			.AddOption(ENiagaraNamespaceMetadataOptions::PreventEditingNamespaceModifier)
-			.AddOption(ENiagaraNamespaceMetadataOptions::PreventEditingName),
 		FNiagaraNamespaceMetadata({FNiagaraConstants::LocalNamespace, FNiagaraConstants::ModuleNamespace})
 			.SetDisplayName(LOCTEXT("ModuleLocalDisplayName", "Local"))
 			.SetDisplayNameLong(LOCTEXT("ModuleLocalDisplayNameLong", "Module Locals"))
@@ -147,7 +146,7 @@ void UNiagaraEditorSettings::SetupNamespaceMetadata()
 			.AddOption(ENiagaraNamespaceMetadataOptions::PreventEditingNamespaceModifier)
 			.AddOption(ENiagaraNamespaceMetadataOptions::PreventEditingName),
 		FNiagaraNamespaceMetadata({FNiagaraConstants::StaticSwitchNamespace})
-			.SetDisplayName(LOCTEXT("StatisSwitchDisplayName", "Static Switch"))
+			.SetDisplayName(LOCTEXT("StatisSwitchDisplayName", "Static Switch Inputs"))
 			.SetDescription(LOCTEXT("StaticSwitchDescription", "Values which can only be set at edit time."))
 			.SetSortId(45)
 			.AddOption(ENiagaraNamespaceMetadataOptions::HideInSystem)
@@ -156,6 +155,14 @@ void UNiagaraEditorSettings::SetupNamespaceMetadata()
 			.AddOption(ENiagaraNamespaceMetadataOptions::PreventEditingName),
 	};
 
+	DefaultNamespaceModifierMetadata = FNiagaraNamespaceMetadata({ NAME_None })
+		.SetDisplayName(LOCTEXT("DefaultModifierDisplayName", "None"))
+		.SetDescription(LOCTEXT("DefaultModifierDescription", "Arbitrary sub-namespace for specifying module specific dataset attributes, or calling nested modules."))
+		.SetBackgroundColor(FLinearColor(FColor(102, 102, 102)))
+		.AddOption(ENiagaraNamespaceMetadataOptions::PreventEditingNamespace)
+		.AddOption(ENiagaraNamespaceMetadataOptions::PreventEditingNamespaceModifier)
+		.AddOption(ENiagaraNamespaceMetadataOptions::PreventEditingName);
+
 	NamespaceModifierMetadata =
 	{
 		FNiagaraNamespaceMetadata({FNiagaraConstants::InitialNamespace})
@@ -166,7 +173,8 @@ void UNiagaraEditorSettings::SetupNamespaceMetadata()
 		FNiagaraNamespaceMetadata({FNiagaraConstants::ModuleNamespace})
 			.SetDisplayName(LOCTEXT("ModuleModifierDisplayName", "Module"))
 			.SetDescription(LOCTEXT("ModuleModifierDescription", "A namespace modifier which makes that attribute unique to the module\ninstance by appending the unique module name."))
-			.SetBackgroundColor(FLinearColor(FColor(102, 102, 152))),
+			.SetBackgroundColor(FLinearColor(FColor(102, 102, 152)))
+			.AddOption(ENiagaraNamespaceMetadataOptions::HideInSystem),
 		FNiagaraNamespaceMetadata({FNiagaraConstants::SystemNamespace})
 			.SetDisplayName(LOCTEXT("SystemModifierDisplayName", "System"))
 			.SetDescription(LOCTEXT("SystemModifierDescription", "A namespace modifier which specifies that an engine provided parameter comes from the system."))
@@ -286,6 +294,11 @@ void UNiagaraEditorSettings::SetNewAssetDialogConfig(FName InDialogConfigKey, co
 	SaveConfig();
 }
 
+FNiagaraNamespaceMetadata UNiagaraEditorSettings::GetDefaultNamespaceMetadata() const
+{
+	return DefaultNamespaceMetadata;
+}
+
 FNiagaraNamespaceMetadata UNiagaraEditorSettings::GetMetaDataForNamespaces(TArray<FName> InNamespaces) const
 {
 	TArray<FNiagaraNamespaceMetadata> MatchingMetadata;
@@ -334,6 +347,11 @@ const TArray<FNiagaraNamespaceMetadata>& UNiagaraEditorSettings::GetAllNamespace
 	return NamespaceMetadata;
 }
 
+FNiagaraNamespaceMetadata UNiagaraEditorSettings::GetDefaultNamespaceModifierMetadata() const
+{
+	return DefaultNamespaceModifierMetadata;
+}
+
 FNiagaraNamespaceMetadata UNiagaraEditorSettings::GetMetaDataForNamespaceModifier(FName NamespaceModifier) const
 {
 	for (const FNiagaraNamespaceMetadata& NamespaceModifierMetadataItem : NamespaceModifierMetadata)
@@ -346,6 +364,11 @@ FNiagaraNamespaceMetadata UNiagaraEditorSettings::GetMetaDataForNamespaceModifie
 	return FNiagaraNamespaceMetadata();
 }
 
+const TArray<FNiagaraNamespaceMetadata>& UNiagaraEditorSettings::GetAllNamespaceModifierMetadata() const
+{
+	return NamespaceModifierMetadata;
+}
+
 FName UNiagaraEditorSettings::GetCategoryName() const
 {
 	return TEXT("Plugins");
diff --git a/Engine/Plugins/FX/Niagara/Source/NiagaraEditor/Private/NiagaraEditorUtilities.cpp b/Engine/Plugins/FX/Niagara/Source/NiagaraEditor/Private/NiagaraEditorUtilities.cpp
--- a/Engine/Plugins/FX/Niagara/Source/NiagaraEditor/Private/NiagaraEditorUtilities.cpp
+++ b/Engine/Plugins/FX/Niagara/Source/NiagaraEditor/Private/NiagaraEditorUtilities.cpp
@@ -2525,10 +2525,10 @@ bool FNiagaraParameterUtilities::DoesParameterNameMatchSearchText(FName Paramete
 			return true;
 		}
 
-		// Check the namespace modifier if it has one
-		if (HandleParts.Num() - NamespaceMetadata.Namespaces.Num() > 1)
+		// Check the namespace modifiers if it has them
+		for(int HandlePartIndex = NamespaceMetadata.Namespaces.Num(); HandlePartIndex < HandleParts.Num() - 2; HandlePartIndex++)
 		{
-			FNiagaraNamespaceMetadata NamespaceModifierMetadata = GetDefault<UNiagaraEditorSettings>()->GetMetaDataForNamespaceModifier(HandleParts[NamespaceMetadata.Namespaces.Num()]);
+			FNiagaraNamespaceMetadata NamespaceModifierMetadata = GetDefault<UNiagaraEditorSettings>()->GetMetaDataForNamespaceModifier(HandleParts[HandlePartIndex]);
 			if (NamespaceModifierMetadata.IsValid())
 			{
 				// Check first by modifier metadata display name.
@@ -2540,7 +2540,7 @@ bool FNiagaraParameterUtilities::DoesParameterNameMatchSearchText(FName Paramete
 			else
 			{
 				// Otherwise just check the string.
-				if (HandleParts[NamespaceMetadata.Namespaces.Num()].ToString().Contains(SearchTextString))
+				if (HandleParts[HandlePartIndex].ToString().Contains(SearchTextString))
 				{
 					return true;
 				}
@@ -2565,7 +2565,6 @@ FText FNiagaraParameterUtilities::FormatParameterNameForTextDisplay(FName Parame
 {
 	FNiagaraParameterHandle ParameterHandle(ParameterName);
 	TArray<FName> HandleParts = ParameterHandle.GetHandleParts();
-	FString DisplayString;
 
 	if (HandleParts.Num() == 0)
 	{
@@ -2576,63 +2575,43 @@ FText FNiagaraParameterUtilities::FormatParameterNameForTextDisplay(FName Parame
 		return FText::FromName(HandleParts[0]);
 	}
 
-	FText NamespaceText;
+	FText NamespaceFormat = LOCTEXT("NamespaceFormat", "({0})");
+	TArray<FText> ParameterNameTextParts;
 	FNiagaraNamespaceMetadata NamespaceMetadata = GetDefault<UNiagaraEditorSettings>()->GetMetaDataForNamespaces(HandleParts);
 	if (NamespaceMetadata.IsValid() && NamespaceMetadata.DisplayName.IsEmptyOrWhitespace() == false)
 	{
-		NamespaceText = NamespaceMetadata.DisplayName.ToUpper();
+		ParameterNameTextParts.Add(FText::Format(NamespaceFormat, NamespaceMetadata.DisplayName.ToUpper()));
 		HandleParts.RemoveAt(0, NamespaceMetadata.Namespaces.Num());
 	}
 	else
 	{
-		NamespaceText = FText::FromName(HandleParts[0]).ToUpper();
+		ParameterNameTextParts.Add(FText::FromName(HandleParts[0]).ToUpper());
 		HandleParts.RemoveAt(0);
 	}
 
-	FText NamespaceModfierText;
-	if(HandleParts.Num() > 1)
+	for(int32 HandlePartIndex = 0; HandlePartIndex < HandleParts.Num() - 1; HandlePartIndex++)
 	{
-		FNiagaraNamespaceMetadata NamespaceModifierMetadata = GetDefault<UNiagaraEditorSettings>()->GetMetaDataForNamespaceModifier(HandleParts[0]);
+		FNiagaraNamespaceMetadata NamespaceModifierMetadata = GetDefault<UNiagaraEditorSettings>()->GetMetaDataForNamespaceModifier(HandleParts[HandlePartIndex]);
 		if (NamespaceModifierMetadata.IsValid() && NamespaceModifierMetadata.DisplayName.IsEmptyOrWhitespace() == false)
 		{
-			NamespaceModfierText = NamespaceModifierMetadata.DisplayName.ToUpper();
+			ParameterNameTextParts.Add(FText::Format(NamespaceFormat, NamespaceModifierMetadata.DisplayName.ToUpper()));
 		}
 		else
 		{
-			NamespaceModfierText = FText::FromName(HandleParts[0]).ToUpper();
+			ParameterNameTextParts.Add(FText::Format(NamespaceFormat, FText::FromName(HandleParts[0]).ToUpper()));
 		}
-		HandleParts.RemoveAt(0);
 	}
 
-	FText NameText;
 	if (HandleParts.Num() == 0)
 	{
-		NameText = FText::FromName(NAME_None);
-	}
-	else if (HandleParts.Num() == 1)
-	{
-		NameText = FText::FromName(HandleParts[0]);
+		ParameterNameTextParts.Add(FText::FromName(NAME_None));
 	}
 	else
 	{
-		// This shouldn't happen for new data, and old data should have been fixed up, but handle it
-		// anyway just in case.
-		TArray<FString> HandlePartStrings;
-		for (FName HandlePart : HandleParts)
-		{
-			HandlePartStrings.Add(HandlePart.ToString());
-		}
-		NameText = FText::FromString(FString::Join(HandlePartStrings, TEXT(".")));
+		ParameterNameTextParts.Add(FText::FromName(HandleParts.Last()));
 	}
 
-	if (NamespaceModfierText.IsEmptyOrWhitespace())
-	{
-		return FText::Format(LOCTEXT("ParameterNameFormat", "({0}) {1}"), NamespaceText, NameText);
-	}
-	else
-	{
-		return FText::Format(LOCTEXT("ParameterNameWithModifierFormat", "({0}) ({1}) {2}"), NamespaceText, NamespaceModfierText, NameText);
-	}
+	return FText::Join(FText::FromString(TEXT(" ")), ParameterNameTextParts);
 }
 
 FName NamePartsToName(const TArray<FName>& NameParts)
@@ -2791,28 +2770,45 @@ FName FNiagaraParameterUtilities::ChangeNamespace(FName ParameterName, const FNi
 	return NAME_None;
 }
 
-bool FNiagaraParameterUtilities::TestCanAddNamespaceModifierWithMessage(FName ParameterName, FText& OutMessage)
+int32 FNiagaraParameterUtilities::GetNumberOfNamePartsBeforeEditableModifier(const FNiagaraNamespaceMetadata& NamespaceMetadata)
+{
+	if (NamespaceMetadata.IsValid() && NamespaceMetadata.Options.Contains(ENiagaraNamespaceMetadataOptions::PreventEditingNamespaceModifier) == false)
+	{
+		// If the namespace has a required modifier then we can add and edit a modifier after the required one.
+		return NamespaceMetadata.RequiredNamespaceModifier != NAME_None ? NamespaceMetadata.Namespaces.Num() + 1 : NamespaceMetadata.Namespaces.Num();
+	}
+	else
+	{
+		return INDEX_NONE;
+	}
+}
+
+void FNiagaraParameterUtilities::GetOptionalNamespaceModifiers(FName ParameterName, EParameterContext InParameterContext, TArray<FName>& OutOptionalNamespaceModifiers)
 {
 	FNiagaraParameterHandle ParameterHandle;
 	FNiagaraNamespaceMetadata NamespaceMetadata;
-	if (GetNamespaceModifierEditData(ParameterName, ParameterHandle, NamespaceMetadata, OutMessage))
+	FText Unused;
+	if (GetNamespaceModifierEditData(ParameterName, ParameterHandle, NamespaceMetadata, Unused))
 	{
-		int32 NumberOfNamePartsAfterNamespace = ParameterHandle.GetHandleParts().Num() - NamespaceMetadata.Namespaces.Num();
-		if (NumberOfNamePartsAfterNamespace == 1)
-		{
-			OutMessage = LOCTEXT("AddNamespaceModifier", "Add a namespace modifier to this parameter.");
-			return true;
-		}
-		else
+		for (FName OptionalNamespaceModifier : NamespaceMetadata.OptionalNamespaceModifiers)
 		{
-			OutMessage = LOCTEXT("CantAddAnotherNamespaceModfier", "Only one namespace modifier is supported.");
-			return false;
+			FNiagaraNamespaceMetadata NamespaceModifierMetadata = GetDefault<UNiagaraEditorSettings>()->GetMetaDataForNamespaceModifier(OptionalNamespaceModifier);
+			if (NamespaceModifierMetadata.IsValid())
+			{
+				bool bShouldHideInContext =
+					(InParameterContext == EParameterContext::Script && NamespaceModifierMetadata.Options.Contains(ENiagaraNamespaceMetadataOptions::HideInScript)) ||
+					(InParameterContext == EParameterContext::System && NamespaceModifierMetadata.Options.Contains(ENiagaraNamespaceMetadataOptions::HideInSystem));
+
+				if(bShouldHideInContext == false)
+				{
+					OutOptionalNamespaceModifiers.Add(OptionalNamespaceModifier);
+				}
+			}
 		}
 	}
-	return false;
 }
 
-FName FNiagaraParameterUtilities::AddNamespaceModifier(FName InParameterName)
+FName FNiagaraParameterUtilities::GetEditableNamespaceModifierForParameter(FName InParameterName)
 {
 	FNiagaraParameterHandle ParameterHandle;
 	FNiagaraNamespaceMetadata NamespaceMetadata;
@@ -2820,38 +2816,37 @@ FName FNiagaraParameterUtilities::AddNamespaceModifier(FName InParameterName)
 	if (GetNamespaceModifierEditData(InParameterName, ParameterHandle, NamespaceMetadata, Unused))
 	{
 		TArray<FName> NameParts = ParameterHandle.GetHandleParts();
-		int32 NumberOfNamePartsAfterNamespace = NameParts.Num() - NamespaceMetadata.Namespaces.Num();
-		if (NumberOfNamePartsAfterNamespace == 1)
+		int32 NumberOfNamePartsBeforeEditableModifier = GetNumberOfNamePartsBeforeEditableModifier(NamespaceMetadata);
+		if (NameParts.Num() - NumberOfNamePartsBeforeEditableModifier == 2)
 		{
-			NameParts.Insert(FNiagaraConstants::ModuleNamespace, NamespaceMetadata.Namespaces.Num());
-			return NamePartsToName(NameParts);
+			return NameParts[NumberOfNamePartsBeforeEditableModifier];
 		}
 	}
 	return NAME_None;
 }
 
-bool FNiagaraParameterUtilities::TestCanRemoveNamespaceModifierWithMessage(FName ParameterName, FText& OutMessage)
+bool FNiagaraParameterUtilities::TestCanSetSpecificNamespaceModifierWithMessage(FName InParameterName, FName InNamespaceModifier, FText& OutMessage)
 {
 	FNiagaraParameterHandle ParameterHandle;
 	FNiagaraNamespaceMetadata NamespaceMetadata;
-	if (GetNamespaceModifierEditData(ParameterName, ParameterHandle, NamespaceMetadata, OutMessage))
+	if (GetNamespaceModifierEditData(InParameterName, ParameterHandle, NamespaceMetadata, OutMessage))
 	{
-		int32 NumberOfNamePartsAfterNamespace = ParameterHandle.GetHandleParts().Num() - NamespaceMetadata.Namespaces.Num();
-		if (NumberOfNamePartsAfterNamespace == 2)
+		FNiagaraNamespaceMetadata ModifierMetadata = GetDefault<UNiagaraEditorSettings>()->GetMetaDataForNamespaceModifier(InNamespaceModifier);
+		if (ModifierMetadata.IsValid() && ModifierMetadata.Description.IsEmptyOrWhitespace() == false)
 		{
-			OutMessage = LOCTEXT("RemoveNamespaceModifier", "Remove the namespace modifier from this parameter.");
-			return true;
+			OutMessage = FText::Format(LOCTEXT("SetNamespaceModifierWithDescriptionFormat", "Set the namespace modifier for this parameter to {0}.\n\nDescription: {1}"),
+				FText::FromName(InNamespaceModifier), ModifierMetadata.Description);
 		}
 		else
 		{
-			OutMessage = LOCTEXT("NoNamespaceModifierToRemove", "No namespace modifier to remove.");
-			return false;
+			OutMessage = FText::Format(LOCTEXT("SetNamespaceModifierFormat", "Set the namespace modifier for this parameter to {0}."), FText::FromName(InNamespaceModifier));
 		}
+		return true;
 	}
 	return false;
 }
 
-FName FNiagaraParameterUtilities::RemoveNamespaceModifier(FName InParameterName)
+FName FNiagaraParameterUtilities::SetSpecificNamespaceModifier(FName InParameterName, FName InNamespaceModifier)
 {
 	FNiagaraParameterHandle ParameterHandle;
 	FNiagaraNamespaceMetadata NamespaceMetadata;
@@ -2859,37 +2854,83 @@ FName FNiagaraParameterUtilities::RemoveNamespaceModifier(FName InParameterName)
 	if (GetNamespaceModifierEditData(InParameterName, ParameterHandle, NamespaceMetadata, Unused))
 	{
 		TArray<FName> NameParts = ParameterHandle.GetHandleParts();
-		int32 NumberOfNamePartsAfterNamespace = NameParts.Num() - NamespaceMetadata.Namespaces.Num();
-		if (NumberOfNamePartsAfterNamespace == 2)
+		int32 NumberOfNamePartsBeforeEditableModifier = GetNumberOfNamePartsBeforeEditableModifier(NamespaceMetadata);
+		if (NumberOfNamePartsBeforeEditableModifier != INDEX_NONE)
 		{
-			NameParts.RemoveAt(NamespaceMetadata.Namespaces.Num());
-			return NamePartsToName(NameParts);
+			if (NameParts.Num() == NumberOfNamePartsBeforeEditableModifier + 1)
+			{
+				// Doesn't have a modifier.  Insert the supplied modifier if it's not none
+				if (InNamespaceModifier != NAME_None)
+				{
+					NameParts.Insert(InNamespaceModifier, NumberOfNamePartsBeforeEditableModifier);
+				}
+				return NamePartsToName(NameParts);
+			}
+			else if (NameParts.Num() == NumberOfNamePartsBeforeEditableModifier + 2)
+			{
+				// Already has a namespace modifier.
+				if (InNamespaceModifier == NAME_None || InNamespaceModifier.ToString().IsEmpty())
+				{
+					// If the name was none or it was empty, remove the modifier.
+					NameParts.RemoveAt(NumberOfNamePartsBeforeEditableModifier);
+				}
+				else
+				{
+					NameParts[NumberOfNamePartsBeforeEditableModifier] = InNamespaceModifier;
+				}
+				return NamePartsToName(NameParts);
+			}
 		}
 	}
 	return NAME_None;
 }
 
-bool FNiagaraParameterUtilities::TestCanEditNamespaceModifierWithMessage(FName ParameterName, FText& OutMessage)
+bool FNiagaraParameterUtilities::TestCanSetCustomNamespaceModifierWithMessage(FName InParameterName, FText& OutMessage)
 {
 	FNiagaraParameterHandle ParameterHandle;
 	FNiagaraNamespaceMetadata NamespaceMetadata;
-	if (GetNamespaceModifierEditData(ParameterName, ParameterHandle, NamespaceMetadata, OutMessage))
+	if (GetNamespaceModifierEditData(InParameterName, ParameterHandle, NamespaceMetadata, OutMessage))
 	{
-		int32 NumberOfNamePartsAfterNamespace = ParameterHandle.GetHandleParts().Num() - NamespaceMetadata.Namespaces.Num();
-		if (NumberOfNamePartsAfterNamespace == 2)
+		FName CurrentModifier = GetEditableNamespaceModifierForParameter(InParameterName);
+		if (CurrentModifier == NAME_None)
 		{
-			OutMessage = LOCTEXT("EditNamespaceModifier", "Edit the namespace modifier for this parameter.");
-			return true;
+			OutMessage = LOCTEXT("SetCustomNamespaceModifier", "Add a new custom modifier to this parameter.");
 		}
 		else
 		{
-			OutMessage = LOCTEXT("NoNamespaceModifierToEdit", "No namespace modifier to edit.");
-			return false;
+			OutMessage = LOCTEXT("EditCurrentNamespaceModifier", "Edit the current custom modifier on this parameter.");
 		}
+		return true;
 	}
 	return false;
 }
 
+FName FNiagaraParameterUtilities::SetCustomNamespaceModifier(FName InParameterName)
+{
+	TSet<FName> CurrentParameterNames;
+	return SetCustomNamespaceModifier(InParameterName, CurrentParameterNames);
+}
+
+FName FNiagaraParameterUtilities::SetCustomNamespaceModifier(FName InParameterName, TSet<FName>& CurrentParameterNames)
+{
+	FName CurrentModifier = GetEditableNamespaceModifierForParameter(InParameterName);
+	if (CurrentModifier == NAME_None)
+	{
+		FName ParameterNameWithCustomModifier = SetSpecificNamespaceModifier(InParameterName, "Custom");
+		int32 CustomIndex = 1;
+		while(CurrentParameterNames.Contains(ParameterNameWithCustomModifier))
+		{
+			ParameterNameWithCustomModifier = SetSpecificNamespaceModifier(InParameterName, *FString::Printf(TEXT("Custom%02i"), CustomIndex));
+			CustomIndex++;
+		}
+		return ParameterNameWithCustomModifier;
+	}
+	else
+	{
+		return InParameterName;
+	}
+}
+
 bool FNiagaraParameterUtilities::TestCanRenameWithMessage(FName ParameterName, FText& OutMessage)
 {
 	FNiagaraParameterHandle ParameterHandle;
diff --git a/Engine/Plugins/FX/Niagara/Source/NiagaraEditor/Private/NiagaraNodeParameterMapBase.cpp b/Engine/Plugins/FX/Niagara/Source/NiagaraEditor/Private/NiagaraNodeParameterMapBase.cpp
--- a/Engine/Plugins/FX/Niagara/Source/NiagaraEditor/Private/NiagaraNodeParameterMapBase.cpp
+++ b/Engine/Plugins/FX/Niagara/Source/NiagaraEditor/Private/NiagaraNodeParameterMapBase.cpp
@@ -373,56 +373,61 @@ void UNiagaraNodeParameterMapBase::GetChangeNamespaceModifierSubMenuForPin(UTool
 {
 	FToolMenuSection& Section = Menu->AddSection("Section");
 
-	TAttribute<FText> AddToolTip = TAttribute<FText>::Create(TAttribute<FText>::FGetter::CreateUObject(
-		this, &UNiagaraNodeParameterMapBase::GetAddNamespaceModifierForPinToolTip, (const UEdGraphPin*)InPin));
-	Section.AddMenuEntry(
-		"AddNamespaceModifier",
-		LOCTEXT("AddNamespaceModifierForPin", "Add"),
-		AddToolTip,
-		FSlateIcon(),
-		FUIAction(
-			FExecuteAction::CreateUObject(this, &UNiagaraNodeParameterMapBase::AddNamespaceModifierForPin, InPin),
-			FCanExecuteAction::CreateUObject(this, &UNiagaraNodeParameterMapBase::CanAddNamespaceModifierForPin, (const UEdGraphPin*)InPin)));
+	TArray<FName> OptionalNamespaceModifiers;
+	FNiagaraParameterUtilities::GetOptionalNamespaceModifiers(InPin->PinName, FNiagaraParameterUtilities::EParameterContext::Script, OptionalNamespaceModifiers);
+	for (FName OptionalNamespaceModifier : OptionalNamespaceModifiers)
+	{
+		TAttribute<FText> SetToolTip = TAttribute<FText>::Create(TAttribute<FText>::FGetter::CreateUObject(
+			this, &UNiagaraNodeParameterMapBase::GetSetNamespaceModifierForPinToolTip, (const UEdGraphPin*)InPin, OptionalNamespaceModifier));
+		Section.AddMenuEntry(
+			OptionalNamespaceModifier,
+			FText::FromName(OptionalNamespaceModifier),
+			SetToolTip,
+			FSlateIcon(),
+			FUIAction(
+				FExecuteAction::CreateUObject(this, &UNiagaraNodeParameterMapBase::SetNamespaceModifierForPin, InPin, OptionalNamespaceModifier),
+				FCanExecuteAction::CreateUObject(this, &UNiagaraNodeParameterMapBase::CanSetNamespaceModifierForPin, (const UEdGraphPin*)InPin, OptionalNamespaceModifier)));
+	}
 
-	TAttribute<FText> RemoveToolTip = TAttribute<FText>::Create(TAttribute<FText>::FGetter::CreateUObject(
-		this, &UNiagaraNodeParameterMapBase::GetRemoveNamespaceModifierForPinToolTip, (const UEdGraphPin*)InPin));
+	TAttribute<FText> SetCustomToolTip = TAttribute<FText>::Create(TAttribute<FText>::FGetter::CreateUObject(
+		this, &UNiagaraNodeParameterMapBase::GetSetCustomNamespaceModifierForPinToolTip, (const UEdGraphPin*)InPin));
 	Section.AddMenuEntry(
-		"RemoveNamespceModifier",
-		LOCTEXT("RemoveNamespaceModifierForPin", "Remove"),
-		RemoveToolTip,
+		"AddCustomNamespaceModifier",
+		LOCTEXT("SetCustomNamespaceModifierForPin", "Custom..."),
+		SetCustomToolTip,
 		FSlateIcon(),
 		FUIAction(
-			FExecuteAction::CreateUObject(this, &UNiagaraNodeParameterMapBase::RemoveNamespaceModifierForPin, InPin),
-			FCanExecuteAction::CreateUObject(this, &UNiagaraNodeParameterMapBase::CanRemoveNamespaceModifierForPin, (const UEdGraphPin*)InPin)));
+			FExecuteAction::CreateUObject(this, &UNiagaraNodeParameterMapBase::SetCustomNamespaceModifierForPin, InPin),
+			FCanExecuteAction::CreateUObject(this, &UNiagaraNodeParameterMapBase::CanSetCustomNamespaceModifierForPin, (const UEdGraphPin*)InPin)));
 
-	TAttribute<FText> EditToolTip = TAttribute<FText>::Create(TAttribute<FText>::FGetter::CreateUObject(
-		this, &UNiagaraNodeParameterMapBase::GetEditNamespaceModifierForPinToolTip, (const UEdGraphPin*)InPin));
+	TAttribute<FText> SetNoneToolTip = TAttribute<FText>::Create(TAttribute<FText>::FGetter::CreateUObject(
+		this, &UNiagaraNodeParameterMapBase::GetSetNamespaceModifierForPinToolTip, (const UEdGraphPin*)InPin, FName(NAME_None)));
 	Section.AddMenuEntry(
-		"EditNamespaceModifier",
-		LOCTEXT("EditNamespaceModifierForPin", "Edit"),
-		EditToolTip,
+		"AddNoneNamespaceModifier",
+		LOCTEXT("SetNoneNamespaceModifierForPin", "None"),
+		SetNoneToolTip,
 		FSlateIcon(),
 		FUIAction(
-			FExecuteAction::CreateUObject(this, &UNiagaraNodeParameterMapBase::EditNamespaceModifierForPin, (const UEdGraphPin*)InPin),
-			FCanExecuteAction::CreateUObject(this, &UNiagaraNodeParameterMapBase::CanEditNamespaceModifierForPin, (const UEdGraphPin*)InPin)));
+			FExecuteAction::CreateUObject(this, &UNiagaraNodeParameterMapBase::SetNamespaceModifierForPin, InPin, FName(NAME_None)),
+			FCanExecuteAction::CreateUObject(this, &UNiagaraNodeParameterMapBase::CanSetNamespaceModifierForPin, (const UEdGraphPin*)InPin, FName(NAME_None))));
 }
 
-FText UNiagaraNodeParameterMapBase::GetAddNamespaceModifierForPinToolTip(const UEdGraphPin* InPin) const
+FText UNiagaraNodeParameterMapBase::GetSetNamespaceModifierForPinToolTip(const UEdGraphPin* InPin, FName InNamespaceModifier) const
 {
-	FText AddMessage;
-	FNiagaraParameterUtilities::TestCanAddNamespaceModifierWithMessage(InPin->PinName, AddMessage);
-	return AddMessage;
+	FText SetMessage;
+	FNiagaraParameterUtilities::TestCanSetSpecificNamespaceModifierWithMessage(InPin->PinName, InNamespaceModifier, SetMessage);
+	return SetMessage;
 }
 
-bool UNiagaraNodeParameterMapBase::CanAddNamespaceModifierForPin(const UEdGraphPin* InPin) const
+bool UNiagaraNodeParameterMapBase::CanSetNamespaceModifierForPin(const UEdGraphPin* InPin, FName InNamespaceModifier) const
 {
 	FText Unused;
-	return FNiagaraParameterUtilities::TestCanAddNamespaceModifierWithMessage(InPin->PinName, Unused);
+	return FNiagaraParameterUtilities::TestCanSetSpecificNamespaceModifierWithMessage(InPin->PinName, InNamespaceModifier, Unused);
 }
 
-void UNiagaraNodeParameterMapBase::AddNamespaceModifierForPin(UEdGraphPin* InPin)
+void UNiagaraNodeParameterMapBase::SetNamespaceModifierForPin(UEdGraphPin* InPin, FName InNamespaceModifier)
 {
-	FName NewName = FNiagaraParameterUtilities::AddNamespaceModifier(InPin->PinName);
+	FName NewName = FNiagaraParameterUtilities::SetSpecificNamespaceModifier(InPin->PinName, InNamespaceModifier);
 	if (NewName != NAME_None)
 	{
 		FScopedTransaction Transaction(LOCTEXT("AddNamespaceModifierTransaction", "Add namespace modifier"));
@@ -430,46 +435,31 @@ void UNiagaraNodeParameterMapBase::AddNamespaceModifierForPin(UEdGraphPin* InPin
 	}
 }
 
-FText UNiagaraNodeParameterMapBase::GetRemoveNamespaceModifierForPinToolTip(const UEdGraphPin* InPin) const
+FText UNiagaraNodeParameterMapBase::GetSetCustomNamespaceModifierForPinToolTip(const UEdGraphPin* InPin) const
 {
-	FText RemoveMessage;
-	FNiagaraParameterUtilities::TestCanRemoveNamespaceModifierWithMessage(InPin->PinName, RemoveMessage);
-	return RemoveMessage;
+	FText SetMessage;
+	FNiagaraParameterUtilities::TestCanSetCustomNamespaceModifierWithMessage(InPin->PinName, SetMessage);
+	return SetMessage;
 }
 
-bool UNiagaraNodeParameterMapBase::CanRemoveNamespaceModifierForPin(const UEdGraphPin* InPin) const
+bool UNiagaraNodeParameterMapBase::CanSetCustomNamespaceModifierForPin(const UEdGraphPin* InPin) const
 {
 	FText Unused;
-	return FNiagaraParameterUtilities::TestCanRemoveNamespaceModifierWithMessage(InPin->PinName, Unused);
+	return FNiagaraParameterUtilities::TestCanSetCustomNamespaceModifierWithMessage(InPin->PinName, Unused);
 }
 
-void UNiagaraNodeParameterMapBase::RemoveNamespaceModifierForPin(UEdGraphPin* InPin)
+void UNiagaraNodeParameterMapBase::SetCustomNamespaceModifierForPin(UEdGraphPin* InPin)
 {
-	FName NewName = FNiagaraParameterUtilities::RemoveNamespaceModifier(InPin->PinName);
+	FName NewName = FNiagaraParameterUtilities::SetCustomNamespaceModifier(InPin->PinName);
 	if (NewName != NAME_None)
 	{
-		FScopedTransaction Transaction(LOCTEXT("RemoveNamespaceModifierTransaction", "Remove namespace modifier"));
-		CommitEditablePinName(FText::FromName(NewName), InPin);
+		if (NewName != InPin->PinName)
+		{
+			FScopedTransaction Transaction(LOCTEXT("AddNamespaceModifierTransaction", "Add custom namespace modifier"));
+			CommitEditablePinName(FText::FromName(NewName), InPin);
+		}
+		SetIsPinEditNamespaceModifierPending(InPin, true);
 	}
 }
 
-FText UNiagaraNodeParameterMapBase::GetEditNamespaceModifierForPinToolTip(const UEdGraphPin* InPin) const
-{
-	FText EditMessage;
-	FNiagaraParameterUtilities::TestCanEditNamespaceModifierWithMessage(InPin->PinName, EditMessage);
-	return EditMessage;
-}
-
-bool UNiagaraNodeParameterMapBase::CanEditNamespaceModifierForPin(const UEdGraphPin* InPin) const
-{
-	FText Unused;
-	return FNiagaraParameterUtilities::TestCanEditNamespaceModifierWithMessage(InPin->PinName, Unused);
-}
-
-void UNiagaraNodeParameterMapBase::EditNamespaceModifierForPin(const UEdGraphPin* InPin)
-{
-	SetIsPinEditNamespaceModifierPending(InPin, true);
-}
-
-
 #undef LOCTEXT_NAMESPACE
\ No newline at end of file
diff --git a/Engine/Plugins/FX/Niagara/Source/NiagaraEditor/Private/NiagaraNodeParameterMapBase.h b/Engine/Plugins/FX/Niagara/Source/NiagaraEditor/Private/NiagaraNodeParameterMapBase.h
--- a/Engine/Plugins/FX/Niagara/Source/NiagaraEditor/Private/NiagaraNodeParameterMapBase.h
+++ b/Engine/Plugins/FX/Niagara/Source/NiagaraEditor/Private/NiagaraNodeParameterMapBase.h
@@ -61,17 +61,13 @@ protected:
 
 	void GetChangeNamespaceModifierSubMenuForPin(UToolMenu* Menu, UEdGraphPin* InPin);
 
-	FText GetAddNamespaceModifierForPinToolTip(const UEdGraphPin* InPin) const;
-	bool CanAddNamespaceModifierForPin(const UEdGraphPin* InPin) const;
-	void AddNamespaceModifierForPin(UEdGraphPin* InPin);
+	FText GetSetNamespaceModifierForPinToolTip(const UEdGraphPin* InPin, FName InNamespaceModifier) const;
+	bool CanSetNamespaceModifierForPin(const UEdGraphPin* InPin, FName InNamespaceModifier) const;
+	void SetNamespaceModifierForPin(UEdGraphPin* InPin, FName InNamespaceModifier);
 
-	FText GetRemoveNamespaceModifierForPinToolTip(const UEdGraphPin* InPin) const;
-	bool CanRemoveNamespaceModifierForPin(const UEdGraphPin* InPin) const;
-	void RemoveNamespaceModifierForPin(UEdGraphPin* InPin);
-
-	FText GetEditNamespaceModifierForPinToolTip(const UEdGraphPin* InPin) const;
-	bool CanEditNamespaceModifierForPin(const UEdGraphPin* InPin) const;
-	void EditNamespaceModifierForPin(const UEdGraphPin* InPin);
+	FText GetSetCustomNamespaceModifierForPinToolTip(const UEdGraphPin* InPin) const;
+	bool CanSetCustomNamespaceModifierForPin(const UEdGraphPin* InPin) const;
+	void SetCustomNamespaceModifierForPin(UEdGraphPin* InPin);
 
 	virtual void OnPinRenamed(UEdGraphPin* RenamedPin, const FString& OldName) override;
 
diff --git a/Engine/Plugins/FX/Niagara/Source/NiagaraEditor/Private/SNiagaraParameterName.cpp b/Engine/Plugins/FX/Niagara/Source/NiagaraEditor/Private/SNiagaraParameterName.cpp
--- a/Engine/Plugins/FX/Niagara/Source/NiagaraEditor/Private/SNiagaraParameterName.cpp
+++ b/Engine/Plugins/FX/Niagara/Source/NiagaraEditor/Private/SNiagaraParameterName.cpp
@@ -4,6 +4,7 @@
 #include "NiagaraEditorSettings.h"
 #include "NiagaraEditorStyle.h"
 #include "NiagaraNodeParameterMapBase.h"
+#include "NiagaraEditorUtilities.h"
 
 #include "Widgets/SBoxPanel.h"
 #include "Widgets/Layout/SBorder.h"
@@ -23,7 +24,9 @@ void SNiagaraParameterName::Construct(const FArguments& InArgs)
 	OnDoubleClickedDelegate = InArgs._OnDoubleClicked;
 	IsSelected = InArgs._IsSelected;
 	DecoratorHAlign = InArgs._DecoratorHAlign;
+	DecoratorPadding = InArgs._DecoratorPadding;
 	Decorator = InArgs._Decorator.Widget;
+	bModifierIsPendingEdit = false;
 
 	UpdateContent(ParameterName.Get());
 }
@@ -35,6 +38,12 @@ void SNiagaraParameterName::Tick(const FGeometry& AllottedGeometry, const double
 	{
 		UpdateContent(CurrentParameterName);
 	}
+
+	if (EditableModifierTextBlock.IsValid() && bModifierIsPendingEdit)
+	{
+		bModifierIsPendingEdit = false;
+		EditableModifierTextBlock->EnterEditingMode();
+	}
 }
 
 FReply SNiagaraParameterName::OnMouseButtonDoubleClick(const FGeometry& InMyGeometry, const FPointerEvent& InMouseEvent)
@@ -65,9 +74,11 @@ TSharedRef<SBorder> SNiagaraParameterName::CreateNamespaceWidget(FText Namespace
 	];
 }
 
-void SNiagaraParameterName::UpdateContent(FName InDisplayedParameterName)
+void SNiagaraParameterName::UpdateContent(FName InDisplayedParameterName, int32 InEditableNamespaceModifierIndex)
 {
 	DisplayedParameterName = InDisplayedParameterName;
+	EditableNameTextBlock.Reset();
+	EditableModifierTextBlock.Reset();
 
 	FString DisplayedParameterNameString = DisplayedParameterName.ToString();
 	TArray<FString> NamePartStrings;
@@ -86,23 +97,27 @@ void SNiagaraParameterName::UpdateContent(FName InDisplayedParameterName)
 
 	TSharedRef<SHorizontalBox> ContentBox = SNew(SHorizontalBox);
 
+	bool bAddNameWidget = true;
 	if (NameParts.Num() > 1 || SingleNameDisplayMode == ESingleNameDisplayMode::Namespace)
 	{
+		FNiagaraNamespaceMetadata DefaultNamespaceMetadata = GetDefault<UNiagaraEditorSettings>()->GetDefaultNamespaceMetadata();
+		FNiagaraNamespaceMetadata DefaultNamespaceModifierMetadata = GetDefault<UNiagaraEditorSettings>()->GetDefaultNamespaceModifierMetadata();
+
 		// Add the namespace widget.
-		FNiagaraNamespaceMetadata DefaultNamespaceMetadata = GetDefault<UNiagaraEditorSettings>()->GetMetaDataForNamespaces({ NAME_None });
 		FNiagaraNamespaceMetadata NamespaceMetadata = GetDefault<UNiagaraEditorSettings>()->GetMetaDataForNamespaces(NameParts);
+		int32 NamespaceNamePartCount;
 		TSharedPtr<SWidget> NamespaceWidget;
 		if (NamespaceMetadata.IsValid())
 		{
-			NameParts.RemoveAt(0, NamespaceMetadata.Namespaces.Num());
+			NamespaceNamePartCount = NamespaceMetadata.Namespaces.Num();
 			NamespaceWidget = CreateNamespaceWidget(
 				NamespaceMetadata.DisplayName.ToUpper(), NamespaceMetadata.Description,
 				NamespaceMetadata.BackgroundColor, NamespaceMetadata.ForegroundStyle);
 		}
 		else
 		{
+			NamespaceNamePartCount = 1;
 			FText NamespaceDisplayName = FText::FromString(FName::NameToDisplayString(NameParts[0].ToString(), false).ToUpper());
-			NameParts.RemoveAt(0);
 			NamespaceWidget = CreateNamespaceWidget(
 				NamespaceDisplayName, DefaultNamespaceMetadata.Description,
 				DefaultNamespaceMetadata.BackgroundColor, DefaultNamespaceMetadata.ForegroundStyle);
@@ -116,77 +131,84 @@ void SNiagaraParameterName::UpdateContent(FName InDisplayedParameterName)
 				NamespaceWidget.ToSharedRef()
 			];
 
-		// Next the namespace modifier widget if there is a namespace modifier.
-		if (NameParts.Num() > 1)
+		if (NamespaceNamePartCount == NameParts.Num())
 		{
-			DisplayedNamespaceModifier = NameParts[0];
-			FNiagaraNamespaceMetadata DisplayedNamespaceModifierMetadata = GetDefault<UNiagaraEditorSettings>()->GetMetaDataForNamespaceModifier(DisplayedNamespaceModifier);
-			if (DisplayedNamespaceModifierMetadata.IsValid() == false)
-			{
-				DisplayedNamespaceModifierMetadata = DefaultNamespaceMetadata;
-			}
-
-			NameParts.RemoveAt(0);
-			NamespaceModifierBorder = CreateNamespaceWidget(
-				FText::FromString(FName::NameToDisplayString(DisplayedNamespaceModifier.ToString(), false).ToUpper()), DisplayedNamespaceModifierMetadata.Description,
-				DisplayedNamespaceModifierMetadata.BackgroundColor, DisplayedNamespaceModifierMetadata.ForegroundStyle);
-
-			ContentBox->AddSlot()
-				.VAlign(VAlign_Center)
-				.AutoWidth()
-				.Padding(0.0f, 0.0f, 5.0f, 0.0f)
-				[
-					NamespaceModifierBorder.ToSharedRef()
-				];
-		}
-		else
-		{
-			DisplayedNamespaceModifier = NAME_None;
-			NamespaceModifierBorder.Reset();
+			bAddNameWidget = false;
 		}
 
-		// If there are extra namespaces found, add them to the UI without metadata.
-		while (NameParts.Num() > 1)
+		// Next add namespace modifier widgets if needed.
+		for(int32 NamePartIndex = NamespaceNamePartCount; NamePartIndex < NameParts.Num() - 1; NamePartIndex++)
 		{
-			FName ExtraNamespace = NameParts[0];
-			NameParts.RemoveAt(0);
-
-			ContentBox->AddSlot()
-				.VAlign(VAlign_Center)
-				.AutoWidth()
-				.Padding(0.0f, 0.0f, 5.0f, 0.0f)
-				[
-					CreateNamespaceWidget(
-						FText::FromString(FName::NameToDisplayString(ExtraNamespace.ToString(), false).ToUpper()), DefaultNamespaceMetadata.Description,
-						DefaultNamespaceMetadata.BackgroundColor, DefaultNamespaceMetadata.ForegroundStyle)
-				];
+			if (NamePartIndex == InEditableNamespaceModifierIndex)
+			{
+				ContentBox->AddSlot()
+					.VAlign(VAlign_Center)
+					.AutoWidth()
+					.Padding(0.0f, 0.0f, 5.0f, 0.0f)
+					[
+						SNew(SBox)
+						.MinDesiredWidth(100)
+						[
+							SAssignNew(EditableModifierTextBlock, SInlineEditableTextBlock)
+							.Style(EditableTextStyle)
+							.Text(FText::FromName(NameParts[NamePartIndex]))
+							.OnVerifyTextChanged(this, &SNiagaraParameterName::VerifyNamespaceModifierTextChange)
+							.OnTextCommitted(this, &SNiagaraParameterName::NamespaceModifierTextCommitted)
+						]
+					];
+				bModifierIsPendingEdit = true;
+			}
+			else
+			{
+				FName NamespaceModifier = NameParts[NamePartIndex];
+				FNiagaraNamespaceMetadata NamespaceModifierMetadata = GetDefault<UNiagaraEditorSettings>()->GetMetaDataForNamespaceModifier(NamespaceModifier);
+				FText NamespaceModifierDisplayName;
+				if (NamespaceModifierMetadata.IsValid())
+				{
+					NamespaceModifierDisplayName = NamespaceModifierMetadata.DisplayName.ToUpper();
+				}
+				else
+				{
+					NamespaceModifierMetadata = DefaultNamespaceModifierMetadata;
+					NamespaceModifierDisplayName = FText::FromString(FName::NameToDisplayString(NamespaceModifier.ToString(), false).ToUpper());
+				}
+
+				TSharedRef<SBorder> NamespaceModifierBorder = CreateNamespaceWidget(
+					NamespaceModifierDisplayName, NamespaceModifierMetadata.Description,
+					NamespaceModifierMetadata.BackgroundColor, NamespaceModifierMetadata.ForegroundStyle);
+
+				ContentBox->AddSlot()
+					.VAlign(VAlign_Center)
+					.AutoWidth()
+					.Padding(0.0f, 0.0f, 5.0f, 0.0f)
+					[
+						NamespaceModifierBorder
+					];
+			}
 		}
 	}
 
-	TSharedPtr<SWidget> NameWidget;
-	if (NameParts.Num() > 0)
+	if (bAddNameWidget)
 	{
+		TSharedPtr<SWidget> NameWidget;
 		if (bIsReadOnly)
 		{
 			NameWidget = SNew(STextBlock)
 				.TextStyle(ReadOnlyTextStyle)
-				.Text(FText::FromName(NameParts[0]))
+				.Text(FText::FromName(NameParts.Last()))
 				.HighlightText(HighlightText);
 		}
 		else
 		{
-			NameWidget = SAssignNew(EditableTextBlock, SInlineEditableTextBlock)
+			NameWidget = SAssignNew(EditableNameTextBlock, SInlineEditableTextBlock)
 				.Style(EditableTextStyle)
-				.Text(FText::FromName(NameParts[0]))
+				.Text(FText::FromName(NameParts.Last()))
 				.IsSelected(IsSelected)
 				.OnVerifyTextChanged(this, &SNiagaraParameterName::VerifyNameTextChange)
 				.OnTextCommitted(this, &SNiagaraParameterName::NameTextCommitted)
 				.HighlightText(HighlightText);
 		}
-	}
 
-	if (NameWidget.IsValid())
-	{
 		ContentBox->AddSlot()
 		.VAlign(VAlign_Center)
 		.AutoWidth()
@@ -200,7 +222,7 @@ void SNiagaraParameterName::UpdateContent(FName InDisplayedParameterName)
 		ContentBox->AddSlot()
 		.VAlign(VAlign_Center)
 		.HAlign(DecoratorHAlign)
-		.Padding(3.0f, 0.0f, 0.0f, 0.0f)
+		.Padding(DecoratorPadding)
 		[
 			Decorator.ToSharedRef()
 		];
@@ -223,47 +245,6 @@ FName SNiagaraParameterName::ReconstructNameFromEditText(const FText& InEditText
 	return *NewParameterNameString;
 }
 
-FName SNiagaraParameterName::ReconstructNameFromNamespaceModifierEditText(const FText& InEditText)
-{
-	FString CurrentParameterNameString = ParameterName.Get().ToString();
-	TArray<FString> NameParts;
-	CurrentParameterNameString.ParseIntoArray(NameParts, TEXT("."));
-
-	TArray<FName> Namespaces;
-	for (int32 NamePartIndex = 0; NamePartIndex < NameParts.Num() - 1; NamePartIndex++)
-	{
-		Namespaces.Add(*NameParts[NamePartIndex]);
-	}
-
-	FNiagaraNamespaceMetadata NamespaceMetadata = GetDefault<UNiagaraEditorSettings>()->GetMetaDataForNamespaces(Namespaces);
-
-	if (NamespaceMetadata.IsValid())
-	{
-		if (NameParts.Num() > NamespaceMetadata.Namespaces.Num() + 1)
-		{
-			NameParts[NamespaceMetadata.Namespaces.Num()] = InEditText.ToString().Replace(TEXT("."), TEXT("_"));
-		}
-		FString NewParameterNameString = FString::Join(NameParts, TEXT("."));
-		return *NewParameterNameString;
-	}
-	else
-	{
-		// If there isn't valid namespace metadata for this parameter we can't safely edit the modifier
-		// since there is not way to know how many parts the namespace actually has.
-		return ParameterName.Get();
-	}
-}
-
-FText SNiagaraParameterName::GetNamespaceModifierText()
-{
-	FString CurrentParameterNameString = ParameterName.Get().ToString();
-	TArray<FString> NameParts;
-	CurrentParameterNameString.ParseIntoArray(NameParts, TEXT("."));
-	return NameParts.Num() >= 3 
-		? FText::FromString(NameParts[1]) 
-		: FText::FromName(NAME_None);
-}
-
 bool SNiagaraParameterName::VerifyNameTextChange(const FText& InNewNameText, FText& OutErrorMessage)
 {
 	FName NewParameterName = ReconstructNameFromEditText(InNewNameText);
@@ -285,8 +266,9 @@ void SNiagaraParameterName::NameTextCommitted(const FText& InNewNameText, ETextC
 
 bool SNiagaraParameterName::VerifyNamespaceModifierTextChange(const FText& InNewNameText, FText& OutErrorMessage)
 {
-	FName NewParameterName = ReconstructNameFromNamespaceModifierEditText(InNewNameText);
-	if (OnVerifyNameChangeDelegate.IsBound())
+	FName NewNamespaceModifier = *InNewNameText.ToString().Replace(TEXT("."), TEXT(""));
+	FName NewParameterName = FNiagaraParameterUtilities::SetSpecificNamespaceModifier(ParameterName.Get(), NewNamespaceModifier);
+	if (NewParameterName != NAME_None && OnVerifyNameChangeDelegate.IsBound())
 	{
 		return OnVerifyNameChangeDelegate.Execute(NewParameterName, OutErrorMessage);
 	}
@@ -297,38 +279,35 @@ void SNiagaraParameterName::NamespaceModifierTextCommitted(const FText& InNewNam
 {
 	if (CommitType == ETextCommit::OnEnter)
 	{
-		FName NewParameterName = ReconstructNameFromNamespaceModifierEditText(InNewNameText);
-		OnNameChangedDelegate.ExecuteIfBound(NewParameterName);
-	}
-
-	if (NamespaceModifierBorder.IsValid())
-	{
-		NamespaceModifierBorder->SetContent(SNew(STextBlock)
-			.TextStyle(FNiagaraEditorStyle::Get(), "NiagaraEditor.ParameterName.NamespaceText")
-			.Text(GetNamespaceModifierText()));
+		FName NewNamespaceModifier = *InNewNameText.ToString().Replace(TEXT("."), TEXT(""));
+		FName NewParameterName = FNiagaraParameterUtilities::SetSpecificNamespaceModifier(ParameterName.Get(), NewNamespaceModifier);
+		if (NewParameterName != NAME_None)
+		{
+			OnNameChangedDelegate.ExecuteIfBound(NewParameterName);
+		}
 	}
+	UpdateContent(ParameterName.Get());
 }
 
 void SNiagaraParameterName::EnterEditingMode()
 {
-	if (EditableTextBlock.IsValid())
+	if (EditableNameTextBlock.IsValid())
 	{
-		EditableTextBlock->EnterEditingMode();
+		EditableNameTextBlock->EnterEditingMode();
 	}
 }
 
 void SNiagaraParameterName::EnterNamespaceModifierEditingMode()
 {
-	if (NamespaceModifierBorder.IsValid())
+	FNiagaraParameterHandle ParameterHandle(ParameterName.Get());
+	FNiagaraNamespaceMetadata NamespaceMetadata = GetDefault<UNiagaraEditorSettings>()->GetMetaDataForNamespaces(ParameterHandle.GetHandleParts());
+	if (NamespaceMetadata.IsValid())
 	{
-		TSharedRef<SInlineEditableTextBlock> NamespaceModifierEditableTextBlock = SNew(SInlineEditableTextBlock)
-			.Style(EditableTextStyle)
-			.Text(GetNamespaceModifierText())
-			.IsSelected(IsSelected)
-			.OnVerifyTextChanged(this, &SNiagaraParameterName::VerifyNamespaceModifierTextChange)
-			.OnTextCommitted(this, &SNiagaraParameterName::NamespaceModifierTextCommitted);
-		NamespaceModifierBorder->SetContent(NamespaceModifierEditableTextBlock);
-		NamespaceModifierEditableTextBlock->EnterEditingMode();
+		int32 NamePartIndexForEditableModifier = FNiagaraParameterUtilities::GetNumberOfNamePartsBeforeEditableModifier(NamespaceMetadata);
+		if (NamePartIndexForEditableModifier != INDEX_NONE)
+		{
+			UpdateContent(ParameterName.Get(), NamePartIndexForEditableModifier);
+		}
 	}
 }
 
@@ -348,6 +327,8 @@ void SNiagaraParameterNameTextBlock::Construct(const FArguments& InArgs)
 		.IsSelected(InArgs._IsSelected)
 		.OnVerifyNameChange(this, &SNiagaraParameterNameTextBlock::VerifyNameChange)
 		.OnNameChanged(this, &SNiagaraParameterNameTextBlock::NameChanged)
+		.DecoratorHAlign(InArgs._DecoratorHAlign)
+		.DecoratorPadding(InArgs._DecoratorPadding)
 		.Decorator()
 		[
 			InArgs._Decorator.Widget
@@ -405,6 +386,8 @@ void SNiagaraParameterNamePinLabel::Construct(const FArguments& InArgs, UEdGraph
 		.OnVerifyTextChanged(InArgs._OnVerifyTextChanged)
 		.OnTextCommitted(InArgs._OnTextCommitted)
 		.IsSelected(InArgs._IsSelected)
+		.DecoratorHAlign(InArgs._DecoratorHAlign)
+		.DecoratorPadding(InArgs._DecoratorPadding)
 		.Decorator()
 		[
 			InArgs._Decorator.Widget
diff --git a/Engine/Plugins/FX/Niagara/Source/NiagaraEditor/Private/ViewModels/Stack/NiagaraStackFunctionInput.cpp b/Engine/Plugins/FX/Niagara/Source/NiagaraEditor/Private/ViewModels/Stack/NiagaraStackFunctionInput.cpp
--- a/Engine/Plugins/FX/Niagara/Source/NiagaraEditor/Private/ViewModels/Stack/NiagaraStackFunctionInput.cpp
+++ b/Engine/Plugins/FX/Niagara/Source/NiagaraEditor/Private/ViewModels/Stack/NiagaraStackFunctionInput.cpp
@@ -1542,46 +1542,17 @@ bool UNiagaraStackFunctionInput::SupportsRename() const
 void UNiagaraStackFunctionInput::OnRenamed(FText NewNameText)
 {
 	FName NewName(*NewNameText.ToString());
-	if (InputParameterHandle.GetName() != NewName && OwningAssignmentNode.IsValid())
+	if (InputParameterHandle.GetName() != NewName && OwningAssignmentNode.IsValid() && SourceScript.IsValid())
 	{
-		bool bIsCurrentlyExpanded = GetStackEditorData().GetStackEntryIsExpanded(FNiagaraStackGraphUtilities::GenerateStackModuleEditorDataKey(*OwningAssignmentNode.Get()), false);
-
 		FScopedTransaction ScopedTransaction(LOCTEXT("RenameInput", "Rename this function's input."));
-		if (ensureMsgf(OwningAssignmentNode->RenameAssignmentTarget(InputParameterHandle.GetName(), NewName), TEXT("Failed to rename assignment node input.")))
-		{
-			// Fixing up the stack graph and rapid iteration parameters must happen first so that when the stack is refreshed the UI is correct.
-			FNiagaraParameterHandle NewInputParameterHandle = FNiagaraParameterHandle(InputParameterHandle.GetNamespace(), NewName);
-			FNiagaraParameterHandle NewAliasedInputParameterHandle = FNiagaraParameterHandle::CreateAliasedModuleParameterHandle(NewInputParameterHandle, OwningAssignmentNode.Get());
-			UEdGraphPin* OverridePin = GetOverridePin();
-			if (OverridePin != nullptr)
-			{
-				// If there is an override pin then the only thing that needs to happen is that it's name needs to be updated so that the value it
-				// holds or is linked to stays intact.
-				OverridePin->Modify();
-				OverridePin->PinName = NewAliasedInputParameterHandle.GetParameterHandleString();
-			}
-			else if (IsRapidIterationCandidate())
-			{
-				// Otherwise if this is a valid rapid iteration parameter the values in the affected scripts need to be updated.
-				FNiagaraVariable NewRapidIterationParameter = CreateRapidIterationVariable(NewAliasedInputParameterHandle.GetParameterHandleString());
-				for (TWeakObjectPtr<UNiagaraScript> AffectedScript : AffectedScripts)
-				{
-					if (AffectedScript.IsValid())
-					{
-						AffectedScript->Modify();
-						AffectedScript->RapidIterationParameters.RenameParameter(RapidIterationParameter, *NewRapidIterationParameter.GetName().ToString());
-					}
-				}
-			}
-
-			// Restore the expanded state with the new editor data key.
-			FString NewStackEditorDataKey = FNiagaraStackGraphUtilities::GenerateStackFunctionInputEditorDataKey(*OwningAssignmentNode.Get(), NewInputParameterHandle);
-			GetStackEditorData().SetStackEntryIsExpanded(NewStackEditorDataKey, bIsCurrentlyExpanded);
-
-			// This refresh call must come last because it will finalize this input entry which would cause earlier fixup to fail.
-			OwningAssignmentNode->RefreshFromExternalChanges();
-			ensureMsgf(IsFinalized(), TEXT("Input not finalized when renamed."));
-		}
+		FNiagaraStackGraphUtilities::RenameAssignmentTarget(
+			GetSystemViewModel()->GetSystem(),
+			GetEmitterViewModel().IsValid() ? GetEmitterViewModel()->GetEmitter() : nullptr,
+			*SourceScript.Get(),
+			*OwningAssignmentNode.Get(),
+			FNiagaraVariable(InputType, InputParameterHandle.GetName()),
+			NewName);
+		ensureMsgf(IsFinalized(), TEXT("Input not finalized when renamed."));
 	}
 }
 
diff --git a/Engine/Plugins/FX/Niagara/Source/NiagaraEditor/Private/ViewModels/Stack/NiagaraStackGraphUtilities.cpp b/Engine/Plugins/FX/Niagara/Source/NiagaraEditor/Private/ViewModels/Stack/NiagaraStackGraphUtilities.cpp
--- a/Engine/Plugins/FX/Niagara/Source/NiagaraEditor/Private/ViewModels/Stack/NiagaraStackGraphUtilities.cpp
+++ b/Engine/Plugins/FX/Niagara/Source/NiagaraEditor/Private/ViewModels/Stack/NiagaraStackGraphUtilities.cpp
@@ -22,6 +22,9 @@
 #include "ViewModels/Stack/NiagaraStackFunctionInput.h"
 #include "ViewModels/Stack/NiagaraStackErrorItem.h"
 #include "NiagaraSystem.h"
+#include "NiagaraSystemEditorData.h"
+#include "NiagaraEmitterEditorData.h"
+#include "NiagaraStackEditorData.h"
 #include "NiagaraEditorModule.h"
 #include "NiagaraEditorUtilities.h"
 #include "NiagaraConstants.h"
@@ -2675,4 +2678,101 @@ void FNiagaraStackGraphUtilities::GetNamespacesForNewWriteParameters(EStackEditC
 	OutNamespacesForNewParameters.Add(FNiagaraConstants::TransientNamespace);
 }
 
+bool FNiagaraStackGraphUtilities::TryRenameAssignmentTarget(UNiagaraNodeAssignment& OwningAssignmentNode, FNiagaraVariable CurrentAssignmentTarget, FName NewAssignmentTargetName)
+{
+	UNiagaraNodeOutput* OutputNode = GetEmitterOutputNodeForStackNode(OwningAssignmentNode);
+	if (OutputNode != nullptr)
+	{
+		UNiagaraSystem* OwningSystem = OwningAssignmentNode.GetTypedOuter<UNiagaraSystem>();
+		UNiagaraEmitter* OwningEmitter = OwningAssignmentNode.GetTypedOuter<UNiagaraEmitter>();
+		UNiagaraScript* OwningScript = nullptr;
+		if (OwningEmitter != nullptr)
+		{
+			OwningScript = OwningEmitter->GetScript(OutputNode->GetUsage(), OutputNode->GetUsageId());
+		}
+		else if(OwningSystem != nullptr)
+		{
+			if (OutputNode->GetUsage() == ENiagaraScriptUsage::SystemSpawnScript)
+			{
+				OwningScript = OwningSystem->GetSystemSpawnScript();
+			}
+			else if (OutputNode->GetUsage() == ENiagaraScriptUsage::SystemUpdateScript)
+			{
+				OwningScript = OwningSystem->GetSystemUpdateScript();
+			}
+		}
+
+		if (OwningSystem != nullptr && OwningScript != nullptr)
+		{
+			RenameAssignmentTarget(*OwningSystem, OwningEmitter, *OwningScript, OwningAssignmentNode, CurrentAssignmentTarget, NewAssignmentTargetName);
+			return true;
+		}
+	}
+	return false;
+}
+
+void FNiagaraStackGraphUtilities::RenameAssignmentTarget(
+	UNiagaraSystem& OwningSystem,
+	UNiagaraEmitter* OwningEmitter,
+	UNiagaraScript& OwningScript,
+	UNiagaraNodeAssignment& OwningAssignmentNode,
+	FNiagaraVariable CurrentAssignmentTarget,
+	FName NewAssignmentTargetName)
+{
+	UNiagaraStackEditorData* StackEditorData;
+	if (OwningEmitter != nullptr)
+	{
+		UNiagaraEmitterEditorData* EmitterEditorData = Cast<UNiagaraEmitterEditorData>(OwningEmitter->GetEditorData());
+		StackEditorData = EmitterEditorData != nullptr ? &EmitterEditorData->GetStackEditorData() : nullptr;
+	}
+	else
+	{
+		UNiagaraSystemEditorData* SystemEditorData = Cast<UNiagaraSystemEditorData>(OwningSystem.GetEditorData());
+		StackEditorData = SystemEditorData != nullptr ? &SystemEditorData->GetStackEditorData() : nullptr;
+	}
+
+	bool bIsCurrentlyExpanded = StackEditorData != nullptr 
+		? StackEditorData->GetStackEntryIsExpanded(GenerateStackModuleEditorDataKey(OwningAssignmentNode), false) 
+		: false;
+
+	if (ensureMsgf(OwningAssignmentNode.RenameAssignmentTarget(CurrentAssignmentTarget.GetName(), NewAssignmentTargetName), TEXT("Failed to rename assignment node input.")))
+	{
+		// Fixing up the stack graph and rapid iteration parameters must happen first so that when the stack is refreshed the UI is correct.
+		FNiagaraParameterHandle CurrentInputParameterHandle = FNiagaraParameterHandle::CreateModuleParameterHandle(CurrentAssignmentTarget.GetName());
+		FNiagaraParameterHandle CurrentAliasedInputParameterHandle = FNiagaraParameterHandle::CreateAliasedModuleParameterHandle(CurrentInputParameterHandle, &OwningAssignmentNode);
+		FNiagaraParameterHandle NewInputParameterHandle = FNiagaraParameterHandle(CurrentInputParameterHandle.GetNamespace(), NewAssignmentTargetName);
+		FNiagaraParameterHandle NewAliasedInputParameterHandle = FNiagaraParameterHandle::CreateAliasedModuleParameterHandle(NewInputParameterHandle, &OwningAssignmentNode);
+		UEdGraphPin* OverridePin = GetStackFunctionInputOverridePin(OwningAssignmentNode, CurrentAliasedInputParameterHandle);
+		if (OverridePin != nullptr)
+		{
+			// If there is an override pin then the only thing that needs to happen is that it's name needs to be updated so that the value it
+			// holds or is linked to stays intact.
+			OverridePin->Modify();
+			OverridePin->PinName = NewAliasedInputParameterHandle.GetParameterHandleString();
+		}
+		else if (IsRapidIterationType(CurrentAssignmentTarget.GetType()))
+		{
+			// Otherwise if this is a rapid iteration type check to see if there is an existing rapid iteration value, and if so, rename it.
+			FString UniqueEmitterName = OwningEmitter != nullptr ? OwningEmitter->GetUniqueEmitterName() : FString();
+			FNiagaraVariable CurrentRapidIterationParameter = CreateRapidIterationParameter(UniqueEmitterName, OwningScript.GetUsage(), CurrentAliasedInputParameterHandle.GetParameterHandleString(), CurrentAssignmentTarget.GetType());
+			if (OwningScript.RapidIterationParameters.IndexOf(CurrentRapidIterationParameter) != INDEX_NONE)
+			{
+				FNiagaraVariable NewRapidIterationParameter = CreateRapidIterationParameter(UniqueEmitterName, OwningScript.GetUsage(), NewAliasedInputParameterHandle.GetParameterHandleString(), CurrentAssignmentTarget.GetType());
+				OwningScript.Modify();
+				OwningScript.RapidIterationParameters.RenameParameter(CurrentRapidIterationParameter, NewRapidIterationParameter.GetName());
+			}
+		}
+
+		if (StackEditorData != nullptr)
+		{
+			// Restore the expanded state with the new editor data key.
+			FString NewStackEditorDataKey = GenerateStackFunctionInputEditorDataKey(OwningAssignmentNode, NewInputParameterHandle);
+			StackEditorData->SetStackEntryIsExpanded(NewStackEditorDataKey, bIsCurrentlyExpanded);
+		}
+
+		// This refresh call must come last because it will finalize this input entry which would cause earlier fixup to fail.
+		OwningAssignmentNode.RefreshFromExternalChanges();
+	}
+}
+
 #undef LOCTEXT_NAMESPACE
diff --git a/Engine/Plugins/FX/Niagara/Source/NiagaraEditor/Private/Widgets/SNiagaraGraphActionWidget.cpp b/Engine/Plugins/FX/Niagara/Source/NiagaraEditor/Private/Widgets/SNiagaraGraphActionWidget.cpp
--- a/Engine/Plugins/FX/Niagara/Source/NiagaraEditor/Private/Widgets/SNiagaraGraphActionWidget.cpp
+++ b/Engine/Plugins/FX/Niagara/Source/NiagaraEditor/Private/Widgets/SNiagaraGraphActionWidget.cpp
@@ -27,15 +27,12 @@ void SNiagaraGraphActionWidget::Construct(const FArguments& InArgs, const FCreat
 			.IsReadOnly(true)
 			.HighlightText(InArgs._HighlightText)
 			.DecoratorHAlign(HAlign_Right)
+			.DecoratorPadding(FMargin(7.0f, 0.0f, 0.0f, 0.0f))
 			.Decorator()
 			[
-				SNew(SBox)
-				.Padding(FMargin(7.0f, 0.0f, 0.0f, 0.0f))
-				[
-					SNew(STextBlock)
-					.TextStyle(FNiagaraEditorStyle::Get(), "NiagaraEditor.ParameterName.TypeText")
-					.Text(NiagaraAction->GetParameterVariable()->GetType().GetNameText())
-				]
+				SNew(STextBlock)
+				.TextStyle(FNiagaraEditorStyle::Get(), "NiagaraEditor.ParameterName.TypeText")
+				.Text(NiagaraAction->GetParameterVariable()->GetType().GetNameText())
 			];
 	}
 	else
diff --git a/Engine/Plugins/FX/Niagara/Source/NiagaraEditor/Private/Widgets/SNiagaraParameterMapPaletteItem.cpp b/Engine/Plugins/FX/Niagara/Source/NiagaraEditor/Private/Widgets/SNiagaraParameterMapPaletteItem.cpp
--- a/Engine/Plugins/FX/Niagara/Source/NiagaraEditor/Private/Widgets/SNiagaraParameterMapPaletteItem.cpp
+++ b/Engine/Plugins/FX/Niagara/Source/NiagaraEditor/Private/Widgets/SNiagaraParameterMapPaletteItem.cpp
@@ -12,6 +12,8 @@
 #include "Widgets/SNiagaraParameterName.h"
 #include "Widgets/SNiagaraParameterMapView.h"
 #include "NiagaraEditorSettings.h"
+#include "EditorFontGlyphs.h"
+#include "EditorStyleSet.h"
 
 #define LOCTEXT_NAMESPACE "NiagaraParameterMapPalleteItem"
 
@@ -52,26 +54,25 @@ void SNiagaraParameterMapPalleteItem::Construct(const FArguments& InArgs, FCreat
 	ParameterToolTipText.Bind(this, &SNiagaraParameterMapPalleteItem::GetItemTooltip);
 	SetToolTipText(ParameterToolTipText);
 
-	bWasCreated = ParameterAction->ReferenceCollection.ContainsByPredicate(
-		[](const FNiagaraGraphParameterReferenceCollection& ReferenceCollection) { return ReferenceCollection.WasCreated(); });
-
 	TArray<FName> Namespaces;
 	NiagaraParameterMapSectionID::OnGetSectionNamespaces((NiagaraParameterMapSectionID::Type)InCreateData->Action->SectionID, Namespaces);
 	FNiagaraNamespaceMetadata NamespaceMetadata = GetDefault<UNiagaraEditorSettings>()->GetMetaDataForNamespaces(Namespaces);
 	bool bForceReadOnly = NamespaceMetadata.IsValid() == false || NamespaceMetadata.Options.Contains(ENiagaraNamespaceMetadataOptions::PreventEditingName);
 
 	ParameterNameTextBlock = SNew(SNiagaraParameterNameTextBlock)
-		.ParameterText(this, &SNiagaraParameterMapPalleteItem::GetDisplayText)
+		.ParameterText(FText::FromName(ParameterAction->Parameter.GetName()))
 		.HighlightText(InCreateData->HighlightText)
 		.OnTextCommitted(this, &SNiagaraParameterMapPalleteItem::OnNameTextCommitted)
 		.OnVerifyTextChanged(this, &SNiagaraParameterMapPalleteItem::OnNameTextVerifyChanged)
 		.IsSelected(InCreateData->IsRowSelectedDelegate)
-		.IsReadOnly(InCreateData->bIsReadOnly || bForceReadOnly)
+		.IsReadOnly(InCreateData->bIsReadOnly || bForceReadOnly || ParameterAction->bIsExternallyReferenced)
 		.Decorator()
 		[
 			SNew(STextBlock)
-			.Visibility(bWasCreated ? EVisibility::Visible : EVisibility::Collapsed)
-			.Text(LOCTEXT("AddedText", "*"))
+			.Visibility(ParameterAction->bIsExternallyReferenced ? EVisibility::Visible : EVisibility::Collapsed)
+			.Font(FEditorStyle::Get().GetFontStyle("FontAwesome.8"))
+			.Text(FEditorFontGlyphs::Lock)
+			.ToolTipText(LOCTEXT("LockedToolTip", "This parameter is used in a referenced external graph and can't be edited directly."))
 		];
 
 	InCreateData->OnRenameRequest->BindSP(ParameterNameTextBlock.ToSharedRef(), &SNiagaraParameterNameTextBlock::EnterEditingMode);
@@ -134,26 +135,16 @@ void SNiagaraParameterMapPalleteItem::Tick(const FGeometry& AllottedGeometry, co
 
 FText SNiagaraParameterMapPalleteItem::GetItemTooltip() const
 {
-	if (bWasCreated)
-	{
-		FText CurrentToolTip = SGraphPaletteItem::GetItemTooltip();
-		if (CurrentToolTip.IdenticalTo(ToolTipCache) == false)
-		{
-			ToolTipCache = CurrentToolTip;
-			CreatedToolTipCache = FText::Format(LOCTEXT("CreatedToolTipFormat", "{0}\n* Created through this panel."), ToolTipCache);
-		}
-		return CreatedToolTipCache;
-	}
-	else
-	{
-		return SGraphPaletteItem::GetItemTooltip();
-	}
+	return SGraphPaletteItem::GetItemTooltip();
 }
 
 void SNiagaraParameterMapPalleteItem::OnNameTextCommitted(const FText& NewText, ETextCommit::Type InTextCommit)
 {
 	TSharedPtr<FNiagaraParameterAction> ParameterAction = StaticCastSharedPtr<FNiagaraParameterAction>(ActionPtr.Pin());
-	OnItemRenamed.ExecuteIfBound(NewText, *ParameterAction.Get());
+	if (ParameterAction.IsValid())
+	{
+		OnItemRenamed.ExecuteIfBound(NewText, ParameterAction.ToSharedRef());
+	}
 }
 
 FText SNiagaraParameterMapPalleteItem::GetReferenceCount() const
diff --git a/Engine/Plugins/FX/Niagara/Source/NiagaraEditor/Private/Widgets/SNiagaraParameterMapPaletteItem.h b/Engine/Plugins/FX/Niagara/Source/NiagaraEditor/Private/Widgets/SNiagaraParameterMapPaletteItem.h
--- a/Engine/Plugins/FX/Niagara/Source/NiagaraEditor/Private/Widgets/SNiagaraParameterMapPaletteItem.h
+++ b/Engine/Plugins/FX/Niagara/Source/NiagaraEditor/Private/Widgets/SNiagaraParameterMapPaletteItem.h
@@ -16,7 +16,7 @@ class SNiagaraParameterMapPalleteItem : public SGraphPaletteItem
 {
 public:
 
-	DECLARE_DELEGATE_TwoParams(FOnItemRenamed, const FText&, struct FNiagaraParameterAction&)
+	DECLARE_DELEGATE_TwoParams(FOnItemRenamed, const FText&, TSharedRef<struct FNiagaraParameterAction>)
 
 	SLATE_BEGIN_ARGS(SNiagaraParameterMapPalleteItem)
 	{}
@@ -38,7 +38,6 @@ protected:
 private:
 	FOnItemRenamed OnItemRenamed;
 	TSharedPtr<SNiagaraParameterNameTextBlock> ParameterNameTextBlock;
-	bool bWasCreated;
 	mutable FText ToolTipCache;
 	mutable FText CreatedToolTipCache;
 };
\ No newline at end of file
diff --git a/Engine/Plugins/FX/Niagara/Source/NiagaraEditor/Private/Widgets/SNiagaraParameterMapView.cpp b/Engine/Plugins/FX/Niagara/Source/NiagaraEditor/Private/Widgets/SNiagaraParameterMapView.cpp
--- a/Engine/Plugins/FX/Niagara/Source/NiagaraEditor/Private/Widgets/SNiagaraParameterMapView.cpp
+++ b/Engine/Plugins/FX/Niagara/Source/NiagaraEditor/Private/Widgets/SNiagaraParameterMapView.cpp
@@ -16,6 +16,8 @@
 #include "NiagaraParameterStore.h"
 #include "NiagaraNodeWithDynamicPins.h"
 #include "NiagaraNodeParameterMapBase.h"
+#include "NiagaraNodeAssignment.h"
+#include "NiagaraNodeParameterMapGet.h"
 #include "NiagaraActions.h"
 #include "SGraphActionMenu.h"
 #include "Framework/MultiBox/MultiBoxBuilder.h"
@@ -88,7 +90,6 @@ void NiagaraParameterMapSectionID::OnGetSectionNamespaces(const NiagaraParameter
 		break;
 	case NiagaraParameterMapSectionID::MODULE_OUTPUT:
 		OutSectionNamespaces.Add(FNiagaraConstants::OutputNamespace);
-		OutSectionNamespaces.Add(FNiagaraConstants::ModuleNamespace);
 		break;
 	case NiagaraParameterMapSectionID::MODULE_LOCAL:
 		OutSectionNamespaces.Add(FNiagaraConstants::LocalNamespace);
@@ -436,7 +437,7 @@ void SNiagaraParameterMapView::AddParameter(FNiagaraVariable NewVariable, bool b
 	if (bSuccess)
 	{
 		GraphActionMenu->RefreshAllActions(true);
-		GraphActionMenu->SelectItemByName(NewVariable.GetName());
+		GraphActionMenu->SelectItemByName(*FNiagaraParameterUtilities::FormatParameterNameForTextDisplay(NewVariable.GetName()).ToString());
 		if (bEnterRenameModeOnAdd)
 		{
 			TArray<FName> Namespaces;
@@ -541,9 +542,44 @@ void SNiagaraParameterMapView::CollectAllActions(FGraphActionListBuilderBase& Ou
 			continue;
 		}
 
-		const FText Name = FText::FromName(Parameter.GetName());
-		const FText Tooltip = FText::Format(TooltipFormat, FText::FromName(Parameter.GetName()), Parameter.GetType().GetNameText());
+		bool bIsExternallyReferenced = false;
+		for (const FNiagaraGraphParameterReferenceCollection& ReferenceCollection : ParameterEntry.Value)
+		{
+			for (const FNiagaraGraphParameterReference& ParameterReference : ReferenceCollection.ParameterReferences)
+			{
+				UNiagaraNode* ReferenceNode = Cast<UNiagaraNode>(ParameterReference.Value.Get());
+				if (ReferenceNode != nullptr)
+				{
+					UNiagaraGraph* ReferenceGraph;
+					UNiagaraNodeAssignment* OwningAssignmentNode = ReferenceNode->GetTypedOuter<UNiagaraNodeAssignment>();
+					if (OwningAssignmentNode != nullptr)
+					{
+						ReferenceGraph = Cast<UNiagaraGraph>(OwningAssignmentNode->GetGraph());
+					}
+					else
+					{
+						ReferenceGraph = Cast<UNiagaraGraph>(ReferenceNode->GetGraph());
+					}
+					if (ReferenceGraph != nullptr)
+					{
+						if (Graphs.ContainsByPredicate([ReferenceGraph](TWeakObjectPtr<UNiagaraGraph> Graph) { return Graph.Get() == ReferenceGraph; }) == false)
+						{
+							bIsExternallyReferenced = true;
+							break;
+						}
+					}
+				}
+			}
+			if (bIsExternallyReferenced)
+			{
+				break;
+			}
+		}
+
+		const FText Name = FNiagaraParameterUtilities::FormatParameterNameForTextDisplay(Parameter.GetName());
+		const FText Tooltip = FText::Format(TooltipFormat, Name, Parameter.GetType().GetNameText());
 		TSharedPtr<FNiagaraParameterAction> ParameterAction(new FNiagaraParameterAction(Parameter, ParameterEntry.Value, FText::GetEmpty(), Name, Tooltip, 0, FText(), Section));
+		ParameterAction->bIsExternallyReferenced = bIsExternallyReferenced;
 		OutAllActions.AddAction(ParameterAction);
 		LastCollectedParameters.Add(Parameter);
 	}
@@ -788,15 +824,39 @@ TSharedPtr<SWidget> SNiagaraParameterMapView::OnContextMenuOpening()
 			RenameToolTip.Bind(this, &SNiagaraParameterMapView::GetRenameOnActionNodeToolTip);
 			MenuBuilder.AddMenuEntry(FGenericCommands::Get().Rename, NAME_None, LOCTEXT("Rename", "Rename"), RenameToolTip);
 
+			MenuBuilder.AddMenuSeparator();
+
 			MenuBuilder.AddSubMenu(
-				LOCTEXT("ChangeNamespace", "Namespace"),
+				LOCTEXT("ChangeNamespace", "Change Namespace"),
 				LOCTEXT("ChangeNamespaceToolTip", "Select a new namespace for the selected parameter."),
-				FNewMenuDelegate::CreateSP(this, &SNiagaraParameterMapView::GetChangeNamespaceSubMenu));
+				FNewMenuDelegate::CreateSP(this, &SNiagaraParameterMapView::GetChangeNamespaceSubMenu, false));
 
 			MenuBuilder.AddSubMenu(
-				LOCTEXT("ChangeNamespaceModifier", "Namespace Modifier"),
+				LOCTEXT("ChangeNamespaceModifier", "Change Namespace Modifier"),
 				LOCTEXT("ChangeNamespaceModifierToolTip", "Edit the namespace modifier for the selected parameter."),
-				FNewMenuDelegate::CreateSP(this, &SNiagaraParameterMapView::GetChangeNamespaceModifierSubMenu));
+				FNewMenuDelegate::CreateSP(this, &SNiagaraParameterMapView::GetChangeNamespaceModifierSubMenu, false));
+
+			MenuBuilder.AddMenuSeparator();
+
+			TAttribute<FText> DuplicateToolTip;
+			DuplicateToolTip.Bind(this, &SNiagaraParameterMapView::GetDuplicateParameterToolTip);
+			MenuBuilder.AddMenuEntry(
+				LOCTEXT("DuplicateParameter", "Duplicate"),
+				DuplicateToolTip,
+				FSlateIcon(),
+				FUIAction(
+					FExecuteAction::CreateSP(this, &SNiagaraParameterMapView::OnDuplicateParameter),
+					FCanExecuteAction::CreateSP(this, &SNiagaraParameterMapView::CanDuplicateParameter)));
+
+			MenuBuilder.AddSubMenu(
+				LOCTEXT("DuplicateToNewNamespace", "Duplicate to Namespace"),
+				LOCTEXT("DuplicateToNewNamespaceToolTip", "Duplicate this parameter to a new namespace."),
+				FNewMenuDelegate::CreateSP(this, &SNiagaraParameterMapView::GetChangeNamespaceSubMenu, true));
+
+			MenuBuilder.AddSubMenu(
+				LOCTEXT("DuplicateWithNewNamespaceModifier", "Duplicate with Namespace Modifier"),
+				LOCTEXT("ChangeNamespaceModifierToolTip", "Duplicate this parameter with a different namespace modifier."),
+				FNewMenuDelegate::CreateSP(this, &SNiagaraParameterMapView::GetChangeNamespaceModifierSubMenu, true));
 		}
 		MenuBuilder.EndSection();
 
@@ -1047,118 +1107,84 @@ void SNiagaraParameterMapView::OnSystemParameterStoreChanged()
 
 FText SNiagaraParameterMapView::GetDeleteEntryToolTip() const
 {
-	if (ToolkitType == SYSTEM)
+	TSharedPtr<FNiagaraParameterAction> ParameterAction;
+	FText ErrorMessage;
+	if (GetSingleParameterActionForSelection(ParameterAction, ErrorMessage) == false)
 	{
-		TArray<TSharedPtr<FEdGraphSchemaAction>> SelectedActions;
-		GraphActionMenu->GetSelectedActions(SelectedActions);
-		for (TSharedPtr<FEdGraphSchemaAction> SelectedAction : SelectedActions)
-		{
-			TSharedPtr<FNiagaraParameterAction> ParameterAction = StaticCastSharedPtr<FNiagaraParameterAction>(SelectedAction);
-			bool bWasCreated = ParameterAction->ReferenceCollection.ContainsByPredicate(
-				[](FNiagaraGraphParameterReferenceCollection& ParameterReferenceCollection) { return ParameterReferenceCollection.WasCreated(); });
-			if (bWasCreated == false)
-			{
-				return LOCTEXT("CantDeleteNonAddedSystem", "Can only delete system parameteters which were added directly.");
-			}
-		}
+		return ErrorMessage;
 	}
+
+	if (ParameterAction->bIsExternallyReferenced)
+	{
+		return LOCTEXT("CantDeleteExternal", "This parameter is referenced in an external script and can't be deleted.");
+	}
+
 	return LOCTEXT("DeleteSelected", "Delete the selected parameter.");
 }
 
 void SNiagaraParameterMapView::OnDeleteEntry()
 {
-	if (ToolkitType == SCRIPT)
+	TSharedPtr<FNiagaraParameterAction> ParameterAction;
+	FText Unused;
+	if (GetSingleParameterActionForSelection(ParameterAction, Unused) && ParameterAction->bIsExternallyReferenced == false)
 	{
-		TArray<TSharedPtr<FEdGraphSchemaAction>> SelectedActions;
-		GraphActionMenu->GetSelectedActions(SelectedActions);
-		for (auto& Action : SelectedActions)
+		if (IsScriptToolkit())
 		{
-			TSharedPtr<FNiagaraParameterAction> ParameterAction = StaticCastSharedPtr<FNiagaraParameterAction>(Action);
-			if (ParameterAction.Get())
+			FScopedTransaction RemoveParametersWithPins(LOCTEXT("RemoveParametersWithPins", "Remove parameter and referenced pins"));
+			for (const TWeakObjectPtr<UNiagaraGraph>& GraphWeakPtr : Graphs)
 			{
-				FScopedTransaction RemoveParametersWithPins(LOCTEXT("RemoveParametersWithPins", "Remove parameter and referenced pins"));
-				for (const TWeakObjectPtr<UNiagaraGraph>& GraphWeakPtr : Graphs)
+				if (GraphWeakPtr.IsValid())
 				{
-					if (GraphWeakPtr.IsValid())
-					{
-						UNiagaraGraph* Graph = GraphWeakPtr.Get();
-						Graph->RemoveParameter(ParameterAction->Parameter);
-					}
+					UNiagaraGraph* Graph = GraphWeakPtr.Get();
+					Graph->RemoveParameter(ParameterAction->Parameter);
 				}
 			}
 		}
-	}
-	else if (ToolkitType == SYSTEM)
-	{
-		UNiagaraSystem* System = CachedSystem.Get();
-		if (System != nullptr)
+		else if (IsSystemToolkit() && CachedSystem.IsValid())
 		{
-			TSharedPtr<FNiagaraParameterAction> ParameterAction;
-			FNiagaraParameterHandle ParameterHandle;
-			FNiagaraNamespaceMetadata NamespaceMetadata;
-			FText Unused;
-			if (GetSingleParameterActionForSelection(ParameterAction, Unused) &&
-				FNiagaraParameterUtilities::GetNamespaceEditData(ParameterAction->GetParameter().GetName(), ParameterHandle, NamespaceMetadata, Unused))
-			{
-				FScopedTransaction RemoveParametersWithPins(LOCTEXT("RemoveParametersFromSystem", "Remove parameter"));
-				System->Modify();
-				System->GetExposedParameters().RemoveParameter(ParameterAction->Parameter);
-				System->EditorOnlyAddedParameters.RemoveParameter(ParameterAction->Parameter);
-			}
+			FScopedTransaction RemoveParametersWithPins(LOCTEXT("RemoveParametersFromSystem", "Remove parameter"));
+			UNiagaraSystem* System = CachedSystem.Get();
+			System->Modify();
+			System->GetExposedParameters().RemoveParameter(ParameterAction->Parameter);
+			System->EditorOnlyAddedParameters.RemoveParameter(ParameterAction->Parameter);
 		}
 	}
 }
 
 bool SNiagaraParameterMapView::CanDeleteEntry() const
 {
-	if (ToolkitType == SCRIPT)
-	{
-		return true;
-	}
-	if (ToolkitType == SYSTEM)
-	{
-		TArray<TSharedPtr<FEdGraphSchemaAction>> SelectedActions;
-		GraphActionMenu->GetSelectedActions(SelectedActions);
-		if (SelectedActions.Num() == 1)
-		{
-			TSharedPtr<FNiagaraParameterAction> ParameterAction = StaticCastSharedPtr<FNiagaraParameterAction>(SelectedActions[0]);
-			bool bWasCreated = false;
-			for (const FNiagaraGraphParameterReferenceCollection& ParameterReferenceCollection : ParameterAction->ReferenceCollection)
-			{
-				if (ParameterReferenceCollection.WasCreated())
-				{
-					bWasCreated = true;
-					break;
-				}
-			}
-			return bWasCreated;
-		}
-	}
-	return false;
+	TSharedPtr<FNiagaraParameterAction> ParameterAction;
+	FText Unused;
+	return GetSingleParameterActionForSelection(ParameterAction, Unused) && ParameterAction->bIsExternallyReferenced == false;
 }
 
 FText SNiagaraParameterMapView::GetRenameOnActionNodeToolTip() const
 {
 	TSharedPtr<FNiagaraParameterAction> ParameterAction;
-	FNiagaraParameterHandle ParameterHandle;
-	FNiagaraNamespaceMetadata NamespaceMetadata;
 	FText ErrorMessage;
-	if (GetSingleParameterActionForSelection(ParameterAction, ErrorMessage) &&
-		FNiagaraParameterUtilities::GetNamespaceEditData(ParameterAction->GetParameter().GetName(), ParameterHandle, NamespaceMetadata, ErrorMessage))
+	if (GetSingleParameterActionForSelection(ParameterAction, ErrorMessage) == false)
 	{
-		if (NamespaceMetadata.Options.Contains(ENiagaraNamespaceMetadataOptions::PreventEditingName))
-		{	
-			return LOCTEXT("RenamingNotSupported", "The namespace for this parameter doesn't support renaming.");
-		}
-		else
-		{
-			return LOCTEXT("RenameParameter", "Rename this parameter.");
-		}
+		return ErrorMessage;
 	}
-	else
+
+	if (ParameterAction->bIsExternallyReferenced)
+	{
+		return LOCTEXT("CantRenameExternal", "This parameter is referenced in an external script and can't be renamed.");
+	}
+
+	FNiagaraParameterHandle ParameterHandle;
+	FNiagaraNamespaceMetadata NamespaceMetadata;
+	if (FNiagaraParameterUtilities::GetNamespaceEditData(ParameterAction->GetParameter().GetName(), ParameterHandle, NamespaceMetadata, ErrorMessage) == false)
 	{
 		return ErrorMessage;
 	}
+
+	if (NamespaceMetadata.Options.Contains(ENiagaraNamespaceMetadataOptions::PreventEditingName))
+	{	
+		return LOCTEXT("RenamingNotSupported", "The namespace for this parameter doesn't support renaming.");
+	}
+
+	return LOCTEXT("RenameParameter", "Rename this parameter.");
 }
 
 void SNiagaraParameterMapView::OnRequestRenameOnActionNode()
@@ -1174,19 +1200,14 @@ bool SNiagaraParameterMapView::CanRequestRenameOnActionNode() const
 	FNiagaraParameterHandle ParameterHandle;
 	FNiagaraNamespaceMetadata NamespaceMetadata;
 	FText Unused;
-	if (GetSingleParameterActionForSelection(ParameterAction, Unused) &&
-		FNiagaraParameterUtilities::GetNamespaceEditData(ParameterAction->GetParameter().GetName(), ParameterHandle, NamespaceMetadata, Unused))
-	{
-		bool bPreventRenameing = NamespaceMetadata.Options.Contains(ENiagaraNamespaceMetadataOptions::PreventEditingName);
-		return bPreventRenameing == false;
-	}
-	else 
-	{
-		return false;
-	}
+	return 
+		GetSingleParameterActionForSelection(ParameterAction, Unused) &&
+		ParameterAction->bIsExternallyReferenced == false &&
+		FNiagaraParameterUtilities::GetNamespaceEditData(ParameterAction->GetParameter().GetName(), ParameterHandle, NamespaceMetadata, Unused) &&
+		NamespaceMetadata.Options.Contains(ENiagaraNamespaceMetadataOptions::PreventEditingName) == false;
 }
 
-void SNiagaraParameterMapView::OnPostRenameActionNode(const FText& InText, FNiagaraParameterAction& InAction)
+void SNiagaraParameterMapView::OnPostRenameActionNode(const FText& InText, TSharedRef<FNiagaraParameterAction> InAction)
 {
 	FText TransactionName;
 	if (IsScriptToolkit())
@@ -1198,7 +1219,7 @@ void SNiagaraParameterMapView::OnPostRenameActionNode(const FText& InText, FNiag
 		TransactionName = LOCTEXT("RenameParameterSystemTransaction", "Rename parameter.");
 	}
 	FScopedTransaction RenameTransaction(TransactionName);
-	RenameParameter(InAction.Parameter, *InText.ToString());
+	RenameParameter(InAction, *InText.ToString());
 }
 
 bool SNiagaraParameterMapView::GetSingleParameterActionForSelection(
@@ -1223,6 +1244,7 @@ bool SNiagaraParameterMapView::GetSingleParameterActionForSelection(
 		OutErrorMessage = LOCTEXT("InvalidParameterAction", "Parameter action is invalid.");
 		return false;
 	}
+
 	return true;
 }
 
@@ -1232,7 +1254,7 @@ bool SNiagaraParameterMapView::ParameterExistsByName(FName ParameterName) const
 	return LastCollectedParameters.ContainsByPredicate(MatchVariableByName);
 }
 
-void SNiagaraParameterMapView::GetChangeNamespaceSubMenu(FMenuBuilder& MenuBuilder)
+void SNiagaraParameterMapView::GetChangeNamespaceSubMenu(FMenuBuilder& MenuBuilder, bool bDuplicateParameter)
 {
 	TSharedPtr<FNiagaraParameterAction> ParameterAction;
 	FText Unused;
@@ -1245,19 +1267,27 @@ void SNiagaraParameterMapView::GetChangeNamespaceSubMenu(FMenuBuilder& MenuBuild
 		{
 			bool bCanChange = MenuDataItem.bCanChange;
 			FText CanChangeToolTip = MenuDataItem.CanChangeToolTip;
-			if (bCanChange)
+			if (bCanChange && bDuplicateParameter == false)
 			{
-				// Check for an existing duplicate by name.
-				FName NewName = FNiagaraParameterUtilities::ChangeNamespace(ParameterAction->Parameter.GetName(), MenuDataItem.Metadata);
-				if (ParameterExistsByName(NewName))
+				if (ParameterAction->bIsExternallyReferenced)
 				{
 					bCanChange = false;
-					CanChangeToolTip = LOCTEXT("CantMoveAlreadyExits", "Can not move to this namespace because a parameter with this name already exists.");
+					CanChangeToolTip = LOCTEXT("CantChangeNamespaceExternallyReferenced", "Parameter is from an externally referenced script and can't be directly edited.");
+				}
+				else
+				{
+					// Check for an existing duplicate by name.
+					FName NewName = FNiagaraParameterUtilities::ChangeNamespace(ParameterAction->Parameter.GetName(), MenuDataItem.Metadata);
+					if (ParameterExistsByName(NewName))
+					{
+						bCanChange = false;
+						CanChangeToolTip = LOCTEXT("CantMoveAlreadyExits", "Can not move to this namespace because a parameter with this name already exists.");
+					}
 				}
 			}
 
 			FUIAction Action = FUIAction(
-				FExecuteAction::CreateSP(this, &SNiagaraParameterMapView::OnChangeNamespace, MenuDataItem.Metadata),
+				FExecuteAction::CreateSP(this, &SNiagaraParameterMapView::OnChangeNamespace, MenuDataItem.Metadata, bDuplicateParameter),
 				FCanExecuteAction::CreateLambda([bCanChange]() { return bCanChange; }));
 
 			TSharedRef<SWidget> MenuItemWidget = FNiagaraParameterUtilities::CreateNamespaceMenuItemWidget(MenuDataItem.NamespaceParameterName, CanChangeToolTip);
@@ -1266,152 +1296,334 @@ void SNiagaraParameterMapView::GetChangeNamespaceSubMenu(FMenuBuilder& MenuBuild
 	}
 }
 
-void SNiagaraParameterMapView::OnChangeNamespace(FNiagaraNamespaceMetadata Metadata)
+void SNiagaraParameterMapView::OnChangeNamespace(FNiagaraNamespaceMetadata Metadata, bool bDuplicateParameter)
 {
 	TSharedPtr<FNiagaraParameterAction> ParameterAction;
 	FText Unused;
-	if (GetSingleParameterActionForSelection(ParameterAction, Unused))
+	if (GetSingleParameterActionForSelection(ParameterAction, Unused) && (bDuplicateParameter || ParameterAction->bIsExternallyReferenced == false))
 	{
 		FName NewName = FNiagaraParameterUtilities::ChangeNamespace(ParameterAction->Parameter.GetName(), Metadata);
 		if (NewName != NAME_None)
 		{
-			FScopedTransaction Transaction(LOCTEXT("ChangeNamespaceTransaction", "Change namespace"));
-			RenameParameter(ParameterAction->Parameter, NewName);
+			bool bParameterExists = ParameterExistsByName(NewName);
+			if (bDuplicateParameter)
+			{
+				FName NewUniqueName;
+				if (bParameterExists)
+				{
+					TSet<FName> ParameterNames;
+					for (FNiagaraVariable& LastCollectedParameter : LastCollectedParameters)
+					{
+						ParameterNames.Add(LastCollectedParameter.GetName());
+					}
+					NewUniqueName = FNiagaraUtilities::GetUniqueName(NewName, ParameterNames);
+				}
+				else
+				{
+					NewUniqueName = NewName;
+				}
+				FScopedTransaction Transaction(LOCTEXT("DuplicateParameterToNewNamespaceTransaction", "Duplicate parameter to new namespace"));
+				AddParameter(FNiagaraVariable(ParameterAction->Parameter.GetType(), NewUniqueName), false);
+			}
+			else if (bParameterExists == false)
+			{
+				FScopedTransaction Transaction(LOCTEXT("ChangeNamespaceTransaction", "Change namespace"));
+				RenameParameter(ParameterAction, NewName);
+			}
 		}
 	}
 }
 
-void SNiagaraParameterMapView::GetChangeNamespaceModifierSubMenu(FMenuBuilder& MenuBuilder)
+TArray<FName> SNiagaraParameterMapView::GetOptionalNamespaceModifiers() const
 {
-	TAttribute<FText> AddToolTip;
-	AddToolTip.Bind(this, &SNiagaraParameterMapView::GetAddNamespaceModifierToolTip);
-	MenuBuilder.AddMenuEntry(
-		LOCTEXT("AddNamespaceModifier", "Add"),
-		AddToolTip,
-		FSlateIcon(),
-		FUIAction(
-			FExecuteAction::CreateSP(this, &SNiagaraParameterMapView::OnAddNamespaceModifier),
-			FCanExecuteAction::CreateSP(this, &SNiagaraParameterMapView::CanAddNamespaceModifier)));
+	TArray<FName> OptionalNamespaceModifiers;
+	TSharedPtr<FNiagaraParameterAction> ParameterAction;
+	FText Unused;
+	if (GetSingleParameterActionForSelection(ParameterAction, Unused))
+	{
+		FNiagaraParameterUtilities::EParameterContext ParameterContext = IsScriptToolkit() 
+			? FNiagaraParameterUtilities::EParameterContext::Script 
+			: FNiagaraParameterUtilities::EParameterContext::System;
+		FNiagaraParameterUtilities::GetOptionalNamespaceModifiers(ParameterAction->Parameter.GetName(), ParameterContext, OptionalNamespaceModifiers);
+	}
+	return OptionalNamespaceModifiers;
+}
+
+void SNiagaraParameterMapView::GetChangeNamespaceModifierSubMenu(FMenuBuilder& MenuBuilder, bool bDuplicateParameter)
+{
+	for (FName OptionalNamespaceModifier : GetOptionalNamespaceModifiers())
+	{
+		TAttribute<FText> SetToolTip;
+		SetToolTip.Bind(TAttribute<FText>::FGetter::CreateSP(this, &SNiagaraParameterMapView::GetSetNamespaceModifierToolTip, OptionalNamespaceModifier, bDuplicateParameter));
+		MenuBuilder.AddMenuEntry(
+			FText::FromName(OptionalNamespaceModifier),
+			SetToolTip,
+			FSlateIcon(),
+			FUIAction(
+				FExecuteAction::CreateSP(this, &SNiagaraParameterMapView::OnSetNamespaceModifier, OptionalNamespaceModifier, bDuplicateParameter),
+				FCanExecuteAction::CreateSP(this, &SNiagaraParameterMapView::CanSetNamespaceModifier, OptionalNamespaceModifier, bDuplicateParameter)));
+	}
 
-	TAttribute<FText> RemoveToolTip;
-	RemoveToolTip.Bind(this, &SNiagaraParameterMapView::GetRemoveNamespaceModifierToolTip);
+	TAttribute<FText> SetCustomToolTip;
+	SetCustomToolTip.Bind(TAttribute<FText>::FGetter::CreateSP(this, &SNiagaraParameterMapView::GetSetCustomNamespaceModifierToolTip, bDuplicateParameter));
 	MenuBuilder.AddMenuEntry(
-		LOCTEXT("RemoveNamespaceModifier", "Remove"),
-		RemoveToolTip,
+		LOCTEXT("CustomNamespaceModifier", "Custom..."),
+		SetCustomToolTip,
 		FSlateIcon(),
 		FUIAction(
-			FExecuteAction::CreateSP(this, &SNiagaraParameterMapView::OnRemoveNamespaceModifier),
-			FCanExecuteAction::CreateSP(this, &SNiagaraParameterMapView::CanRemoveNamespaceModifier)));
+			FExecuteAction::CreateSP(this, &SNiagaraParameterMapView::OnSetCustomNamespaceModifier, bDuplicateParameter),
+			FCanExecuteAction::CreateSP(this, &SNiagaraParameterMapView::CanSetCustomNamespaceModifier, bDuplicateParameter)));
 
-	TAttribute<FText> EditToolTip;
-	EditToolTip.Bind(this, &SNiagaraParameterMapView::GetEditNamespaceModifierToolTip);
+	TAttribute<FText> SetNoneToolTip;
+	SetNoneToolTip.Bind(TAttribute<FText>::FGetter::CreateSP(this, &SNiagaraParameterMapView::GetSetNamespaceModifierToolTip, FName(NAME_None), bDuplicateParameter));
 	MenuBuilder.AddMenuEntry(
-		LOCTEXT("EditNamespaceModifier", "Edit"),
-		EditToolTip,
+		LOCTEXT("NoneNamespaceModifier", "None"),
+		SetNoneToolTip,
 		FSlateIcon(),
 		FUIAction(
-			FExecuteAction::CreateSP(this, &SNiagaraParameterMapView::OnEditNamespaceModifier),
-			FCanExecuteAction::CreateSP(this, &SNiagaraParameterMapView::CanEditNamespaceModifier)));
+			FExecuteAction::CreateSP(this, &SNiagaraParameterMapView::OnSetNamespaceModifier, FName(NAME_None), bDuplicateParameter),
+			FCanExecuteAction::CreateSP(this, &SNiagaraParameterMapView::CanSetNamespaceModifier, FName(NAME_None), bDuplicateParameter)));
 }
 
-FText SNiagaraParameterMapView::GetAddNamespaceModifierToolTip() const
+bool SNiagaraParameterMapView::TestCanSetNamespaceModifierWithMessage(FName InNamespaceModifier, bool bDuplicateParameter, FText& OutMessage) const
 {
 	TSharedPtr<FNiagaraParameterAction> ParameterAction;
-	FText AddMessage;
-	if (GetSingleParameterActionForSelection(ParameterAction, AddMessage))
+	FText SetMessage;
+	if (GetSingleParameterActionForSelection(ParameterAction, OutMessage) == false)
+	{
+		return false;
+	}
+
+	if (FNiagaraParameterUtilities::TestCanSetSpecificNamespaceModifierWithMessage(ParameterAction->Parameter.GetName(), InNamespaceModifier, OutMessage) == false)
 	{
-		FNiagaraParameterUtilities::TestCanAddNamespaceModifierWithMessage(ParameterAction->Parameter.GetName(), AddMessage);
+		return false;
 	}
-	return AddMessage;
+
+	if (bDuplicateParameter == false)
+	{
+		if (ParameterAction->bIsExternallyReferenced)
+		{
+			OutMessage = LOCTEXT("CantChangeNamespaceModifierExternallyReferenced", "Parameter is from an externally referenced script and can't be directly edited.");
+			return false;
+		}
+
+		if (InNamespaceModifier != NAME_None)
+		{
+			FName NewName = FNiagaraParameterUtilities::SetSpecificNamespaceModifier(ParameterAction->Parameter.GetName(), InNamespaceModifier);
+			if (ParameterExistsByName(NewName))
+			{
+				OutMessage = LOCTEXT("CantChangeNamespaceModifierAlreadyExits", "Can't set this namespace modifier because it would create a parameter that already exists.");
+				return false;
+			}
+		}
+	}
+
+	return true;
 }
 
-bool SNiagaraParameterMapView::CanAddNamespaceModifier() const
+FText SNiagaraParameterMapView::GetSetNamespaceModifierToolTip(FName InNamespaceModifier, bool bDuplicateParameter) const
+{
+	FText SetMessage;
+	TestCanSetNamespaceModifierWithMessage(InNamespaceModifier, bDuplicateParameter, SetMessage);
+	return SetMessage;
+}
+
+bool SNiagaraParameterMapView::CanSetNamespaceModifier(FName InNamespaceModifier, bool bDuplicateParameter) const
 {
-	TSharedPtr<FNiagaraParameterAction> ParameterAction;
 	FText Unused;
-	return GetSingleParameterActionForSelection(ParameterAction, Unused) && 
-		FNiagaraParameterUtilities::TestCanAddNamespaceModifierWithMessage(ParameterAction->Parameter.GetName(), Unused);
+	return TestCanSetNamespaceModifierWithMessage(InNamespaceModifier, bDuplicateParameter, Unused);
 }
 
-void SNiagaraParameterMapView::OnAddNamespaceModifier()
+void SNiagaraParameterMapView::OnSetNamespaceModifier(FName InNamespaceModifier, bool bDuplicateParameter)
 {
 	TSharedPtr<FNiagaraParameterAction> ParameterAction;
 	FText Unused;
-	if (GetSingleParameterActionForSelection(ParameterAction, Unused))
+	if (GetSingleParameterActionForSelection(ParameterAction, Unused) && 
+		(bDuplicateParameter || ParameterAction->bIsExternallyReferenced == false))
 	{
-		FName NewName = FNiagaraParameterUtilities::AddNamespaceModifier(ParameterAction->GetParameter().GetName());
-		if(NewName != NAME_None)
+		FName NewName = FNiagaraParameterUtilities::SetSpecificNamespaceModifier(ParameterAction->Parameter.GetName(), InNamespaceModifier);
+		if (NewName != NAME_None)
 		{
-			FScopedTransaction Transaction(LOCTEXT("AddNamespaceModifierTransaction", "Add namespace modifier."));
-			RenameParameter(ParameterAction->Parameter, NewName);
+			bool bParameterExists = ParameterExistsByName(NewName);
+			if (bDuplicateParameter)
+			{
+				FName NewUniqueName;
+				if (bParameterExists)
+				{
+					TSet<FName> ParameterNames;
+					for (FNiagaraVariable& LastCollectedParameter : LastCollectedParameters)
+					{
+						ParameterNames.Add(LastCollectedParameter.GetName());
+					}
+					NewUniqueName = FNiagaraUtilities::GetUniqueName(NewName, ParameterNames);
+				}
+				else
+				{
+					NewUniqueName = NewName;
+				}
+				FScopedTransaction Transaction(LOCTEXT("DuplicateParameterToWithNamespaceModifierTransaction", "Duplicate parameter with namespace modifier"));
+				AddParameter(FNiagaraVariable(ParameterAction->Parameter.GetType(), NewUniqueName), false);
+			}
+			else if (bParameterExists == false)
+			{
+				FScopedTransaction Transaction(LOCTEXT("ChangeNamespaceModifierTransaction", "Change namespace modifier"));
+				RenameParameter(ParameterAction, NewName);
+			}
 		}
 	}
 }
 
-FText SNiagaraParameterMapView::GetRemoveNamespaceModifierToolTip() const
+bool SNiagaraParameterMapView::TestCanSetCustomNamespaceModifierWithMessage(bool bDuplicateParameter, FText& OutMessage) const
 {
 	TSharedPtr<FNiagaraParameterAction> ParameterAction;
-	FText RemoveMessage;
-	if (GetSingleParameterActionForSelection(ParameterAction, RemoveMessage))
+	FText SetMessage;
+	if (GetSingleParameterActionForSelection(ParameterAction, OutMessage) == false)
+	{
+		return false;
+	}
+
+	if (FNiagaraParameterUtilities::TestCanSetCustomNamespaceModifierWithMessage(ParameterAction->Parameter.GetName(), OutMessage) == false)
+	{
+		return false;
+	}
+
+	if (bDuplicateParameter == false && ParameterAction->bIsExternallyReferenced)
 	{
-		FNiagaraParameterUtilities::TestCanRemoveNamespaceModifierWithMessage(ParameterAction->Parameter.GetName(), RemoveMessage);
+		OutMessage = LOCTEXT("CantChangeNamespaceModifierExternallyReferenced", "Parameter is from an externally referenced script and can't be directly edited.");
+		return false;
 	}
-	return RemoveMessage;
+
+	return true;
 }
 
-bool SNiagaraParameterMapView::CanRemoveNamespaceModifier() const
+FText SNiagaraParameterMapView::GetSetCustomNamespaceModifierToolTip(bool bDuplicateParameter) const
+{
+	FText SetMessage;
+	TestCanSetCustomNamespaceModifierWithMessage(bDuplicateParameter, SetMessage);
+	return SetMessage;
+}
+
+bool SNiagaraParameterMapView::CanSetCustomNamespaceModifier(bool bDuplicateParameter) const
 {
-	TSharedPtr<FNiagaraParameterAction> ParameterAction;
 	FText Unused;
-	return GetSingleParameterActionForSelection(ParameterAction, Unused) &&
-		FNiagaraParameterUtilities::TestCanRemoveNamespaceModifierWithMessage(ParameterAction->Parameter.GetName(), Unused);
+	return TestCanSetCustomNamespaceModifierWithMessage(bDuplicateParameter, Unused);
 }
 
-void SNiagaraParameterMapView::OnRemoveNamespaceModifier()
+void SNiagaraParameterMapView::OnSetCustomNamespaceModifier(bool bDuplicateParameter)
 {
 	TSharedPtr<FNiagaraParameterAction> ParameterAction;
 	FText Unused;
-	if (GetSingleParameterActionForSelection(ParameterAction, Unused))
+	if (GetSingleParameterActionForSelection(ParameterAction, Unused) && (bDuplicateParameter || ParameterAction->bIsExternallyReferenced == false))
 	{
-		FName NewName = FNiagaraParameterUtilities::RemoveNamespaceModifier(ParameterAction->GetParameter().GetName());
+		TSet<FName> ParameterNames;
+		for (FNiagaraVariable& LastCollectedParameter : LastCollectedParameters)
+		{
+			ParameterNames.Add(LastCollectedParameter.GetName());
+		}
+		FName NewName = FNiagaraParameterUtilities::SetCustomNamespaceModifier(ParameterAction->Parameter.GetName(), ParameterNames);
 		if (NewName != NAME_None)
 		{
-			FScopedTransaction Transaction(LOCTEXT("RemoveNamespaceModifierTransaction", "Remove namespace modifier."));
-			RenameParameter(ParameterAction->Parameter, NewName);
+			if (bDuplicateParameter)
+			{
+				bool bParameterExists = ParameterNames.Contains(NewName);
+				FName NewUniqueName;
+				if (bParameterExists)
+				{
+					NewUniqueName = FNiagaraUtilities::GetUniqueName(NewName, ParameterNames);
+				}
+				else
+				{
+					NewUniqueName = NewName;
+				}
+				FScopedTransaction Transaction(LOCTEXT("DuplicateParameterToWithNamespaceModifierTransaction", "Duplicate parameter with custom namespace modifier"));
+				AddParameter(FNiagaraVariable(ParameterAction->Parameter.GetType(), NewUniqueName), false);
+
+				FName ActionName = *FNiagaraParameterUtilities::FormatParameterNameForTextDisplay(NewUniqueName).ToString();
+				GraphActionMenu->SelectItemByName(ActionName);
+				TSharedPtr<FNiagaraParameterAction> RenamedParameterAction;
+				if (GetSingleParameterActionForSelection(RenamedParameterAction, Unused))
+				{
+					RenamedParameterAction->bNamespaceModifierRenamePending = true;
+				}
+			}
+			else
+			{
+				if (ParameterAction->Parameter.GetName() != NewName)
+				{
+					FScopedTransaction Transaction(LOCTEXT("SetCustomNamespaceModifierTransaction", "Set custom namespace modifier"));
+					RenameParameter(ParameterAction, NewName);
+
+					FName ActionName = *FNiagaraParameterUtilities::FormatParameterNameForTextDisplay(NewName).ToString();
+					GraphActionMenu->SelectItemByName(ActionName);
+					TSharedPtr<FNiagaraParameterAction> RenamedParameterAction;
+					if (GetSingleParameterActionForSelection(RenamedParameterAction, Unused))
+					{
+						RenamedParameterAction->bNamespaceModifierRenamePending = true;
+					}
+				}
+				else
+				{
+					ParameterAction->bNamespaceModifierRenamePending = true;
+				}
+			}
 		}
 	}
 }
 
-FText SNiagaraParameterMapView::GetEditNamespaceModifierToolTip() const
+bool SNiagaraParameterMapView::TestCanDuplicateParameterWithMessage(FText& OutMessage) const
 {
 	TSharedPtr<FNiagaraParameterAction> ParameterAction;
-	FText EditMessage;
-	if (GetSingleParameterActionForSelection(ParameterAction, EditMessage))
+	if (GetSingleParameterActionForSelection(ParameterAction, OutMessage) == false)
+	{
+		return false;
+	}
+
+	FNiagaraParameterHandle ParameterHandle;
+	FNiagaraNamespaceMetadata NamespaceMetadata;
+	if (FNiagaraParameterUtilities::GetNamespaceEditData(ParameterAction->Parameter.GetName(), ParameterHandle, NamespaceMetadata, OutMessage) == false)
 	{
-		FNiagaraParameterUtilities::TestCanEditNamespaceModifierWithMessage(ParameterAction->Parameter.GetName(), EditMessage);
+		return false;
 	}
-	return EditMessage;
+
+	if (NamespaceMetadata.Options.Contains(ENiagaraNamespaceMetadataOptions::PreventEditingName))
+	{
+		OutMessage = LOCTEXT("CantDuplicateWhenCantRename", "This parameter can not be duplicated because it does not support editing its name.");
+		return false;
+	}
+
+	return true;
 }
 
-bool SNiagaraParameterMapView::CanEditNamespaceModifier() const
+FText SNiagaraParameterMapView::GetDuplicateParameterToolTip() const
+{
+	FText DuplicateMessage;
+	TestCanDuplicateParameterWithMessage(DuplicateMessage);
+	return DuplicateMessage;
+}
+
+bool SNiagaraParameterMapView::CanDuplicateParameter() const
 {
-	TSharedPtr<FNiagaraParameterAction> ParameterAction;
 	FText Unused;
-	return GetSingleParameterActionForSelection(ParameterAction, Unused) &&
-		FNiagaraParameterUtilities::TestCanEditNamespaceModifierWithMessage(ParameterAction->Parameter.GetName(), Unused);
+	return TestCanDuplicateParameterWithMessage(Unused);
 }
 
-void SNiagaraParameterMapView::OnEditNamespaceModifier()
+void SNiagaraParameterMapView::OnDuplicateParameter()
 {
-	TArray<TSharedPtr<FEdGraphSchemaAction>> SelectedActions;
-	GraphActionMenu->GetSelectedActions(SelectedActions);
-	if (SelectedActions.Num() == 1)
+	TSharedPtr<FNiagaraParameterAction> ParameterAction;
+	FNiagaraParameterHandle ParameterHandle;
+	FNiagaraNamespaceMetadata NamespaceMetadata;
+	FText Unused;
+	if (GetSingleParameterActionForSelection(ParameterAction, Unused) &&
+		FNiagaraParameterUtilities::GetNamespaceEditData(ParameterAction->Parameter.GetName(), ParameterHandle, NamespaceMetadata, Unused))
 	{
-		TSharedPtr<FNiagaraParameterAction> ParameterAction = StaticCastSharedPtr<FNiagaraParameterAction>(SelectedActions[0]);
-		if (ParameterAction.IsValid())
+		TSet<FName> ParameterNames;
+		for (FNiagaraVariable& LastCollectedParameter : LastCollectedParameters)
 		{
-			ParameterAction->bNamespaceModifierRenamePending = true;
+			ParameterNames.Add(LastCollectedParameter.GetName());
 		}
+		FName NewUniqueName = FNiagaraUtilities::GetUniqueName(ParameterAction->Parameter.GetName(), ParameterNames);
+		FScopedTransaction Transaction(LOCTEXT("DuplicateParameterTransaction", "Duplicate parameter"));
+		AddParameter(FNiagaraVariable(ParameterAction->Parameter.GetType(), NewUniqueName));
 	}
 }
 
@@ -1459,8 +1671,14 @@ void SNiagaraParameterMapView::OnCopyParameterReference()
 	}
 }
 
-void SNiagaraParameterMapView::RenameParameter(FNiagaraVariable& Parameter, FName NewName)
+void SNiagaraParameterMapView::RenameParameter(TSharedPtr<FNiagaraParameterAction> ParameterAction, FName NewName)
 {
+	if (ensureMsgf(ParameterAction->bIsExternallyReferenced == false, TEXT("Can not modify an externally referenced parameter.")) == false)
+	{
+		return;
+	}
+
+	FNiagaraVariable Parameter = ParameterAction->Parameter;
 	if (ToolkitType == SCRIPT)
 	{
 		if (Graphs.Num() > 0)
@@ -1481,42 +1699,7 @@ void SNiagaraParameterMapView::RenameParameter(FNiagaraVariable& Parameter, FNam
 					continue;
 				}
 
-				if (ReferenceCollection->WasCreated())
-				{
-					// If the parameter was created directly by the user just rename it.
-					Graph->RenameParameter(Parameter, NewName);
-				}
-				else
-				{
-					// Otherwise see if it has local graph references.  If so it can be renamed, if not than we have to
-					// create a new parameter with the new name since we can't rename parameters from nested graphs.
-					bool bHasLocalGraphReference = false;
-					for (const FNiagaraGraphParameterReference& Reference : ReferenceCollection->ParameterReferences)
-					{
-						UNiagaraNode* ReferencingNode = Cast<UNiagaraNode>(Reference.Value.Get());
-						if (ReferencingNode != nullptr && ReferencingNode->GetGraph() == Graph)
-						{
-							bHasLocalGraphReference = true;
-							break;
-						}
-					}
-
-					if (bHasLocalGraphReference)
-					{
-						Graph->RenameParameter(Parameter, NewName);
-					}
-					else
-					{
-						FNiagaraEditorUtilities::InfoWithToastAndLog(FText::Format(
-							LOCTEXT("AddInsteadOfRenameInScriptFormat", "The parameter {0} was not added directly\n or was discovered in a nested graph and can't be renamed.\nA new parameter {1} has been added instead."),
-							FNiagaraParameterUtilities::FormatParameterNameForTextDisplay(Parameter.GetName()),
-							FNiagaraParameterUtilities::FormatParameterNameForTextDisplay(NewName)), 7.0f);
-
-						FNiagaraVariable NewVariable(Parameter.GetType(), NewName);
-						bool bEnterRenameModeOnAdd = false;
-						AddParameter(NewVariable, bEnterRenameModeOnAdd);
-					}
-				}		
+				Graph->RenameParameter(Parameter, NewName);
 			}
 		}
 	}
@@ -1525,34 +1708,53 @@ void SNiagaraParameterMapView::RenameParameter(FNiagaraVariable& Parameter, FNam
 		UNiagaraSystem* System = CachedSystem.Get();
 		if (System != nullptr)
 		{
-			System->Modify();
-			if (System->GetExposedParameters().IndexOf(Parameter) == INDEX_NONE && System->EditorOnlyAddedParameters.IndexOf(Parameter) == INDEX_NONE)
-			{
-				// If the parameter isn't in one of the existing collections than it's a discovered parameter and instead of renaming it needs to be added.
-				FNiagaraEditorUtilities::InfoWithToastAndLog(FText::Format(
-					LOCTEXT("AddInsteadOfRenameInSystemFormat", "The parameter {0}\n was not added directly and can't be renamed.\nA new parameter {1} has been added instead."),
-					FNiagaraParameterUtilities::FormatParameterNameForTextDisplay(Parameter.GetName()),
-					FNiagaraParameterUtilities::FormatParameterNameForTextDisplay(NewName)), 7.0f);
-
-				FNiagaraVariable NewVariable(Parameter.GetType(), NewName);
-				bool bEnterRenameModeOnAdd = false;
-				AddParameter(NewVariable, bEnterRenameModeOnAdd);
-			}
-			else
+			// Rename the parameter in the parameter stores.
+			System->GetExposedParameters().RenameParameter(Parameter, NewName);
+			System->EditorOnlyAddedParameters.RenameParameter(Parameter, NewName);
+
+			// Look for set variables nodes or linked inputs which reference this parameter.
+
+			for (FNiagaraGraphParameterReferenceCollection& ReferenceCollection : ParameterAction->ReferenceCollection)
 			{
-				System->GetExposedParameters().RenameParameter(Parameter, NewName);
-				System->EditorOnlyAddedParameters.RenameParameter(Parameter, NewName);
+				for (FNiagaraGraphParameterReference& ParameterReference : ReferenceCollection.ParameterReferences)
+				{
+					UNiagaraNode* ReferenceNode = Cast<UNiagaraNode>(ParameterReference.Value);
+					if (ReferenceNode != nullptr)
+					{
+						UNiagaraNodeAssignment* OwningAssignmentNode = ReferenceNode->GetTypedOuter<UNiagaraNodeAssignment>();
+						if (OwningAssignmentNode != nullptr)
+						{
+							// If this is owned by a set variables node and it's not locked, update the assignment target on the assignment node.
+							FNiagaraStackGraphUtilities::TryRenameAssignmentTarget(*OwningAssignmentNode, Parameter, NewName);
+						}
+						else
+						{
+							// Otherwise if the reference node is a get node it's for a linked input so we can just update pin name.
+							UNiagaraNodeParameterMapGet* ReferenceGetNode = Cast<UNiagaraNodeParameterMapGet>(ReferenceNode);
+							if (ReferenceGetNode != nullptr)
+							{
+								UEdGraphPin** LinkedInputPinPtr = ReferenceGetNode->Pins.FindByPredicate([&ParameterReference](UEdGraphPin* Pin) { return Pin->PersistentGuid == ParameterReference.Key; });
+								if (LinkedInputPinPtr != nullptr)
+								{
+									UEdGraphPin* LinkedInputPin = *LinkedInputPinPtr;
+									LinkedInputPin->Modify();
+									LinkedInputPin->PinName = NewName;
+								}
+							}
+						}
+					}
+				}
 			}
 		}
 	}
 }
 
-bool SNiagaraParameterMapView::IsSystemToolkit()
+bool SNiagaraParameterMapView::IsSystemToolkit() const
 {
 	return ToolkitType == EToolkitType::SYSTEM;
 }
 
-bool SNiagaraParameterMapView::IsScriptToolkit()
+bool SNiagaraParameterMapView::IsScriptToolkit() const
 {
 	return ToolkitType == EToolkitType::SCRIPT;
 }
@@ -1571,7 +1773,9 @@ void SNiagaraParameterMapView::HandleGraphSubObjectSelectionChanged(const UObjec
 {
 	if (NewSelection->IsA<UNiagaraScriptVariable>())
 	{
-		GraphActionMenu->SelectItemByName(static_cast<const UNiagaraScriptVariable*>(NewSelection)->Variable.GetName());
+		FName VariableName = static_cast<const UNiagaraScriptVariable*>(NewSelection)->Variable.GetName();
+		FName VariableActionName = *FNiagaraParameterUtilities::FormatParameterNameForTextDisplay(VariableName).ToString();
+		GraphActionMenu->SelectItemByName(VariableActionName);
 	}
 	SelectedVariableObjects->SetSelectedObject(const_cast<UObject*>(NewSelection));
 }
@@ -2000,13 +2204,21 @@ void SNiagaraAddParameterMenu::AddParameterSelected(FNiagaraVariable NewVariable
 			? TEXT("New Variable")
 			: TEXT("New ") + TypeDisplayName;
 
+		TArray<FString> NameParts;
+
 		TArray<FName> SectionNamespaces;
 		NiagaraParameterMapSectionID::OnGetSectionNamespaces(InSection, SectionNamespaces);
-		TArray<FString> NameParts;
 		for (FName SectionNamespace : SectionNamespaces)
 		{
 			NameParts.Add(SectionNamespace.ToString());
 		}
+
+		FNiagaraNamespaceMetadata NamespaceMetadata = GetDefault<UNiagaraEditorSettings>()->GetMetaDataForNamespaces(SectionNamespaces);
+		if (NamespaceMetadata.IsValid() && NamespaceMetadata.RequiredNamespaceModifier != NAME_None)
+		{
+			NameParts.Add(NamespaceMetadata.RequiredNamespaceModifier.ToString());
+		}
+
 		NameParts.Add(NewVariableDefaultName);
 		const FString ResultName = FString::Join(NameParts, TEXT("."));
 		NewVariable.SetName(FName(*ResultName));
diff --git a/Engine/Plugins/FX/Niagara/Source/NiagaraEditor/Public/NiagaraActions.h b/Engine/Plugins/FX/Niagara/Source/NiagaraEditor/Public/NiagaraActions.h
--- a/Engine/Plugins/FX/Niagara/Source/NiagaraEditor/Public/NiagaraActions.h
+++ b/Engine/Plugins/FX/Niagara/Source/NiagaraEditor/Public/NiagaraActions.h
@@ -61,7 +61,8 @@ struct NIAGARAEDITOR_API FNiagaraScriptVarAndViewInfoAction : public FEdGraphSch
 struct NIAGARAEDITOR_API FNiagaraParameterAction : public FEdGraphSchemaAction
 {
 	FNiagaraParameterAction()
-		: bNamespaceModifierRenamePending(false)
+		: bIsExternallyReferenced(false)
+		, bNamespaceModifierRenamePending(false)
 	{
 	}
 
@@ -77,6 +78,7 @@ struct NIAGARAEDITOR_API FNiagaraParameterAction : public FEdGraphSchemaAction
 
 	TArray<FNiagaraGraphParameterReferenceCollection> ReferenceCollection;
 
+	bool bIsExternallyReferenced;
 	bool bNamespaceModifierRenamePending;
 };
 
diff --git a/Engine/Plugins/FX/Niagara/Source/NiagaraEditor/Public/NiagaraEditorSettings.h b/Engine/Plugins/FX/Niagara/Source/NiagaraEditor/Public/NiagaraEditorSettings.h
--- a/Engine/Plugins/FX/Niagara/Source/NiagaraEditor/Public/NiagaraEditorSettings.h
+++ b/Engine/Plugins/FX/Niagara/Source/NiagaraEditor/Public/NiagaraEditorSettings.h
@@ -47,24 +47,29 @@ struct FNiagaraNamespaceMetadata
 
 	FNiagaraNamespaceMetadata();
 
-	FNiagaraNamespaceMetadata(TArray<FName> InNamespaces);
+	FNiagaraNamespaceMetadata(TArray<FName> InNamespaces, FName InRequiredNamespaceModifier = NAME_None);
 
 	bool operator==(const FNiagaraNamespaceMetadata& Other) const
 	{
 		return
 			Namespaces == Other.Namespaces &&
+			RequiredNamespaceModifier == Other.RequiredNamespaceModifier &&
 			DisplayName.IdenticalTo(Other.DisplayName) &&
 			DisplayNameLong.IdenticalTo(Other.DisplayNameLong) &&
 			Description.IdenticalTo(Other.Description) &&
 			BackgroundColor == Other.BackgroundColor &&
 			ForegroundStyle == Other.ForegroundStyle &&
 			SortId == Other.SortId &&
+			OptionalNamespaceModifiers == Other.OptionalNamespaceModifiers &&
 			Options == Other.Options;
 	}
 
 	UPROPERTY()
 	TArray<FName> Namespaces;
 
+	UPROPERTY()
+	FName RequiredNamespaceModifier;
+
 	UPROPERTY()
 	FText DisplayName;
 
@@ -83,6 +88,9 @@ struct FNiagaraNamespaceMetadata
 	UPROPERTY()
 	int32 SortId;
 
+	UPROPERTY()
+	TArray<FName> OptionalNamespaceModifiers;
+
 	UPROPERTY()
 	TArray<ENiagaraNamespaceMetadataOptions> Options;
 
@@ -122,6 +130,12 @@ struct FNiagaraNamespaceMetadata
 		return *this;
 	}
 
+	FNiagaraNamespaceMetadata& AddOptionalNamespaceModifier(FName InOptionalNamespaceModifier)
+	{
+		OptionalNamespaceModifiers.Add(InOptionalNamespaceModifier);
+		return *this;
+	}
+
 	FNiagaraNamespaceMetadata& AddOption(ENiagaraNamespaceMetadataOptions Option)
 	{
 		Options.Add(Option);
@@ -201,9 +215,13 @@ public:
 
 	void SetNewAssetDialogConfig(FName InDialogConfigKey, const FNiagaraNewAssetDialogConfig& InNewAssetDialogConfig);
 
+	FNiagaraNamespaceMetadata GetDefaultNamespaceMetadata() const;
 	FNiagaraNamespaceMetadata GetMetaDataForNamespaces(TArray<FName> Namespaces) const;
 	const TArray<FNiagaraNamespaceMetadata>& GetAllNamespaceMetadata() const;
+
+	FNiagaraNamespaceMetadata GetDefaultNamespaceModifierMetadata() const;
 	FNiagaraNamespaceMetadata GetMetaDataForNamespaceModifier(FName NamespaceModifier) const;
+	const TArray<FNiagaraNamespaceMetadata>& GetAllNamespaceModifierMetadata() const;
 	
 	// Begin UDeveloperSettings Interface
 	virtual FName GetCategoryName() const override;
@@ -261,4 +279,10 @@ private:
 
 	UPROPERTY()
 	TArray<FNiagaraNamespaceMetadata> NamespaceModifierMetadata;
+
+	UPROPERTY()
+	FNiagaraNamespaceMetadata DefaultNamespaceMetadata;
+
+	UPROPERTY()
+	FNiagaraNamespaceMetadata DefaultNamespaceModifierMetadata;
 };
diff --git a/Engine/Plugins/FX/Niagara/Source/NiagaraEditor/Public/NiagaraEditorUtilities.h b/Engine/Plugins/FX/Niagara/Source/NiagaraEditor/Public/NiagaraEditorUtilities.h
--- a/Engine/Plugins/FX/Niagara/Source/NiagaraEditor/Public/NiagaraEditorUtilities.h
+++ b/Engine/Plugins/FX/Niagara/Source/NiagaraEditor/Public/NiagaraEditorUtilities.h
@@ -347,15 +347,21 @@ namespace FNiagaraParameterUtilities
 
 	NIAGARAEDITOR_API FName ChangeNamespace(FName ParameterName, const FNiagaraNamespaceMetadata& NewNamespaceMetadata);
 
-	NIAGARAEDITOR_API bool TestCanAddNamespaceModifierWithMessage(FName ParameterName, FText& OutMessage);
+	NIAGARAEDITOR_API int32 GetNumberOfNamePartsBeforeEditableModifier(const FNiagaraNamespaceMetadata& NamespaceMetadata);
 
-	NIAGARAEDITOR_API FName AddNamespaceModifier(FName InParameterName);
+	NIAGARAEDITOR_API void GetOptionalNamespaceModifiers(FName ParameterName, EParameterContext InParameterContext, TArray<FName>& OutOptionalNamespaceModifiers);
 
-	NIAGARAEDITOR_API bool TestCanRemoveNamespaceModifierWithMessage(FName ParameterName, FText& OutMessage);
+	NIAGARAEDITOR_API FName GetEditableNamespaceModifierForParameter(FName ParameterName);
 
-	NIAGARAEDITOR_API FName RemoveNamespaceModifier(FName InParameterName);
+	NIAGARAEDITOR_API bool TestCanSetSpecificNamespaceModifierWithMessage(FName InParameterName, FName InNamespaceModifier, FText& OutMessage);
 
-	NIAGARAEDITOR_API bool TestCanEditNamespaceModifierWithMessage(FName ParameterName, FText& OutMessage);
+	NIAGARAEDITOR_API FName SetSpecificNamespaceModifier(FName InParameterName, FName InNamespaceModifier);
+
+	NIAGARAEDITOR_API bool TestCanSetCustomNamespaceModifierWithMessage(FName InParameterName, FText& OutMessage);
+
+	NIAGARAEDITOR_API FName SetCustomNamespaceModifier(FName InParameterName);
+
+	NIAGARAEDITOR_API FName SetCustomNamespaceModifier(FName InParameterName, TSet<FName>& CurrentParameterNames);
 
 	NIAGARAEDITOR_API bool TestCanRenameWithMessage(FName ParameterName, FText& OutMessage);
 };
diff --git a/Engine/Plugins/FX/Niagara/Source/NiagaraEditor/Public/ViewModels/Stack/NiagaraStackGraphUtilities.h b/Engine/Plugins/FX/Niagara/Source/NiagaraEditor/Public/ViewModels/Stack/NiagaraStackGraphUtilities.h
--- a/Engine/Plugins/FX/Niagara/Source/NiagaraEditor/Public/ViewModels/Stack/NiagaraStackGraphUtilities.h
+++ b/Engine/Plugins/FX/Niagara/Source/NiagaraEditor/Public/ViewModels/Stack/NiagaraStackGraphUtilities.h
@@ -207,4 +207,14 @@ namespace FNiagaraStackGraphUtilities
 
 	/** Gets the valid namespaces which new parameters for this usage can write to. */
 	void GetNamespacesForNewWriteParameters(EStackEditContext EditContext, ENiagaraScriptUsage Usage, TArray<FName>& OutNamespacesForNewParameters);
+
+	bool TryRenameAssignmentTarget(UNiagaraNodeAssignment& OwningAssignmentNode, FNiagaraVariable CurrentAssignmentTarget, FName NewAssignmentTargetName);
+
+	void RenameAssignmentTarget(
+		UNiagaraSystem& OwningSystem,
+		UNiagaraEmitter* OwningEmitter,
+		UNiagaraScript& OwningScript,
+		UNiagaraNodeAssignment& OwningAssignmentNode,
+		FNiagaraVariable CurrentAssignmentTarget,
+		FName NewAssignmentTargetName);
 }
diff --git a/Engine/Plugins/FX/Niagara/Source/NiagaraEditor/Public/Widgets/SNiagaraParameterMapView.h b/Engine/Plugins/FX/Niagara/Source/NiagaraEditor/Public/Widgets/SNiagaraParameterMapView.h
--- a/Engine/Plugins/FX/Niagara/Source/NiagaraEditor/Public/Widgets/SNiagaraParameterMapView.h
+++ b/Engine/Plugins/FX/Niagara/Source/NiagaraEditor/Public/Widgets/SNiagaraParameterMapView.h
@@ -149,39 +149,42 @@ private:
 	FText GetRenameOnActionNodeToolTip() const;
 	void OnRequestRenameOnActionNode();
 	bool CanRequestRenameOnActionNode() const;
-	void OnPostRenameActionNode(const FText& InText, FNiagaraParameterAction& InAction);
+	void OnPostRenameActionNode(const FText& InText, TSharedRef<FNiagaraParameterAction> InAction);
 
-	bool GetSingleParameterActionForSelection(
-		TSharedPtr<FNiagaraParameterAction>& OutParameterAction,
-		FText& OutErrorMessage) const;
+	bool GetSingleParameterActionForSelection(TSharedPtr<FNiagaraParameterAction>& OutParameterAction, FText& OutErrorMessage) const;
 
 	bool ParameterExistsByName(FName ParameterName) const;
 
-	void GetChangeNamespaceSubMenu(FMenuBuilder& MenuBuilder);
-	void OnChangeNamespace(FNiagaraNamespaceMetadata Metadata);
+	void GetChangeNamespaceSubMenu(FMenuBuilder& MenuBuilder, bool bDuplicateParameter);
+	void OnChangeNamespace(FNiagaraNamespaceMetadata Metadata, bool bDuplicateParameter);
 
-	void GetChangeNamespaceModifierSubMenu(FMenuBuilder& MenuBuilder);
+	TArray<FName> GetOptionalNamespaceModifiers() const;
 
-	FText GetAddNamespaceModifierToolTip() const;
-	bool CanAddNamespaceModifier() const;
-	void OnAddNamespaceModifier();
+	void GetChangeNamespaceModifierSubMenu(FMenuBuilder& MenuBuilder, bool bDuplicateParameter);
 
-	FText GetRemoveNamespaceModifierToolTip() const;
-	bool CanRemoveNamespaceModifier() const;
-	void OnRemoveNamespaceModifier();
+	bool TestCanSetNamespaceModifierWithMessage(FName InNamespaceModifier, bool bDuplicateParameter, FText& OutMessage) const;
+	FText GetSetNamespaceModifierToolTip(FName InNamespaceModifier, bool bDuplicateParameter) const;
+	bool CanSetNamespaceModifier(FName InNamespaceModifier, bool bDuplicateParameter) const;
+	void OnSetNamespaceModifier(FName InNamespaceModifier, bool bDuplicateParameter);
 
-	FText GetEditNamespaceModifierToolTip() const;
-	bool CanEditNamespaceModifier() const;
-	void OnEditNamespaceModifier();
+	bool TestCanSetCustomNamespaceModifierWithMessage(bool bDuplicateParameter, FText& OutMessage) const;
+	FText GetSetCustomNamespaceModifierToolTip(bool bDuplicateParameter) const;
+	bool CanSetCustomNamespaceModifier(bool bDuplicateParameter) const;
+	void OnSetCustomNamespaceModifier(bool bDuplicateParameter);
+
+	bool TestCanDuplicateParameterWithMessage(FText& OutMessage) const;
+	FText GetDuplicateParameterToolTip() const;
+	bool CanDuplicateParameter() const;
+	void OnDuplicateParameter();
 
 	FText GetCopyParameterReferenceToolTip() const;
 	bool CanCopyParameterReference() const;
 	void OnCopyParameterReference();
 
-	void RenameParameter(FNiagaraVariable& Parameter, FName NewName);
+	void RenameParameter(TSharedPtr<FNiagaraParameterAction> ParameterAction, FName NewName);
 
-	bool IsSystemToolkit();
-	bool IsScriptToolkit();
+	bool IsSystemToolkit() const;
+	bool IsScriptToolkit() const;
 
 	/** Delegate handler used to match an FName to an action in the list, used for renaming keys */
 	bool HandleActionMatchesName(struct FEdGraphSchemaAction* InAction, const FName& InName) const;
diff --git a/Engine/Plugins/FX/Niagara/Source/NiagaraEditor/Public/Widgets/SNiagaraParameterName.h b/Engine/Plugins/FX/Niagara/Source/NiagaraEditor/Public/Widgets/SNiagaraParameterName.h
--- a/Engine/Plugins/FX/Niagara/Source/NiagaraEditor/Public/Widgets/SNiagaraParameterName.h
+++ b/Engine/Plugins/FX/Niagara/Source/NiagaraEditor/Public/Widgets/SNiagaraParameterName.h
@@ -32,6 +32,7 @@ public:
 		, _IsReadOnly(false)
 		, _SingleNameDisplayMode(ESingleNameDisplayMode::Name)
 		, _DecoratorHAlign(HAlign_Left)
+		, _DecoratorPadding(5.0f, 0.0f, 0.0f, 0.0f)
 	{
 		_Clipping = EWidgetClipping::OnDemand;
 	}
@@ -46,6 +47,7 @@ public:
 		SLATE_EVENT(FIsSelected, IsSelected)
 		SLATE_EVENT(FPointerEventHandler, OnDoubleClicked)
 		SLATE_ARGUMENT(EHorizontalAlignment, DecoratorHAlign)
+		SLATE_ARGUMENT(FMargin, DecoratorPadding)
 		SLATE_NAMED_SLOT(FArguments, Decorator)
 	SLATE_END_ARGS()
 
@@ -62,14 +64,10 @@ public:
 private:
 	TSharedRef<SBorder> CreateNamespaceWidget(FText NamespaceDisplayName, FText NamespaceDescription, FLinearColor NamespaceBorderColor, FName NamespaceForegroundStyle);
 
-	void UpdateContent(FName InDisplayedParameterName);
+	void UpdateContent(FName InDisplayedParameterName, int32 InEditableNamespaceModifierIndex = INDEX_NONE);
 
 	FName ReconstructNameFromEditText(const FText& InEditText);
 
-	FName ReconstructNameFromNamespaceModifierEditText(const FText& InEditText);
-
-	FText GetNamespaceModifierText();
-
 	bool VerifyNameTextChange(const FText& InNewNameText, FText& OutErrorMessage);
 
 	void NameTextCommitted(const FText& InNewNameText, ETextCommit::Type CommitType);
@@ -81,19 +79,24 @@ private:
 private:
 	const FInlineEditableTextBlockStyle* EditableTextStyle;
 	const FTextBlockStyle* ReadOnlyTextStyle;
+
 	TAttribute<FName> ParameterName;
 	bool bIsReadOnly;
 	ESingleNameDisplayMode SingleNameDisplayMode;
 	TAttribute<FText> HighlightText;
+
 	FOnVerifyNameChange OnVerifyNameChangeDelegate;
 	FOnNameChanged OnNameChangedDelegate;
 	FPointerEventHandler OnDoubleClickedDelegate;
+
 	FIsSelected IsSelected;
 	FName DisplayedParameterName;
-	FName DisplayedNamespaceModifier;
-	TSharedPtr<SInlineEditableTextBlock> EditableTextBlock;
-	TSharedPtr<SBorder> NamespaceModifierBorder;
+	TSharedPtr<SInlineEditableTextBlock> EditableNameTextBlock;
+	TSharedPtr<SInlineEditableTextBlock> EditableModifierTextBlock;
+	bool bModifierIsPendingEdit;
+
 	EHorizontalAlignment DecoratorHAlign;
+	FMargin DecoratorPadding;
 	TSharedPtr<SWidget> Decorator;
 };
 
@@ -103,6 +106,8 @@ public:
 	SLATE_BEGIN_ARGS(SNiagaraParameterNameTextBlock)
 		: _EditableTextStyle(&FEditorStyle::Get().GetWidgetStyle<FInlineEditableTextBlockStyle>("InlineEditableTextBlockStyle"))
 		, _IsReadOnly(false)
+		, _DecoratorHAlign(HAlign_Left)
+		, _DecoratorPadding(5.0f, 0.0f, 0.0f, 0.0f)
 	{
 		_Clipping = EWidgetClipping::OnDemand;
 	}
@@ -113,6 +118,8 @@ public:
 		SLATE_EVENT(FOnVerifyTextChanged, OnVerifyTextChanged)
 		SLATE_EVENT(FOnTextCommitted, OnTextCommitted)
 		SLATE_EVENT(FIsSelected, IsSelected)
+		SLATE_ARGUMENT(EHorizontalAlignment, DecoratorHAlign)
+		SLATE_ARGUMENT(FMargin, DecoratorPadding)
 		SLATE_NAMED_SLOT(FArguments, Decorator)
 		
 	SLATE_END_ARGS()
@@ -146,6 +153,8 @@ public:
 	SLATE_BEGIN_ARGS(SNiagaraParameterNamePinLabel)
 		: _EditableTextStyle(&FEditorStyle::Get().GetWidgetStyle<FInlineEditableTextBlockStyle>("InlineEditableTextBlockStyle"))
 		, _IsReadOnly(false)
+		, _DecoratorHAlign(HAlign_Left)
+		, _DecoratorPadding(5.0f, 0.0f, 0.0f, 0.0f)
 	{
 		_Clipping = EWidgetClipping::OnDemand;
 	}
@@ -156,13 +165,15 @@ public:
 		SLATE_EVENT(FOnVerifyTextChanged, OnVerifyTextChanged)
 		SLATE_EVENT(FOnTextCommitted, OnTextCommitted)
 		SLATE_EVENT(FIsSelected, IsSelected)
+		SLATE_ARGUMENT(EHorizontalAlignment, DecoratorHAlign)
+		SLATE_ARGUMENT(FMargin, DecoratorPadding)
 		SLATE_NAMED_SLOT(FArguments, Decorator)
-		SLATE_END_ARGS()
+	SLATE_END_ARGS()
 
-		void Construct(const FArguments& InArgs, UEdGraphPin* InTargetPin);
+	void Construct(const FArguments& InArgs, UEdGraphPin* InTargetPin);
 
 	virtual void Tick(const FGeometry& AllottedGeometry, const double InCurrentTime, const float InDeltaTime) override;
 
 private:
 	UEdGraphPin* TargetPin;
-};
\ No newline at end of file
+};
diff --git a/Engine/Plugins/FX/Niagara/Source/NiagaraEditorWidgets/Private/Stack/SNiagaraStackFunctionInputName.cpp b/Engine/Plugins/FX/Niagara/Source/NiagaraEditorWidgets/Private/Stack/SNiagaraStackFunctionInputName.cpp
--- a/Engine/Plugins/FX/Niagara/Source/NiagaraEditorWidgets/Private/Stack/SNiagaraStackFunctionInputName.cpp
+++ b/Engine/Plugins/FX/Niagara/Source/NiagaraEditorWidgets/Private/Stack/SNiagaraStackFunctionInputName.cpp
@@ -208,56 +208,62 @@ void SNiagaraStackFunctionInputName::OnChangeNamespace(FNiagaraNamespaceMetadata
 
 void SNiagaraStackFunctionInputName::GetChangeNamespaceModifierSubMenu(FMenuBuilder& MenuBuilder)
 {
-	TAttribute<FText> AddToolTip;
-	AddToolTip.Bind(this, &SNiagaraStackFunctionInputName::GetAddNamespaceModifierToolTip);
-	MenuBuilder.AddMenuEntry(
-		LOCTEXT("AddNamespaceModifier", "Add"),
-		AddToolTip,
-		FSlateIcon(),
-		FUIAction(
-			FExecuteAction::CreateSP(this, &SNiagaraStackFunctionInputName::OnAddNamespaceModifier),
-			FCanExecuteAction::CreateSP(this, &SNiagaraStackFunctionInputName::CanAddNamespaceModifier)));
+	FName InputParameterName = *FunctionInput->GetDisplayName().ToString();
+	TArray<FName> OptionalNamespaceModifiers;
+	FNiagaraParameterUtilities::GetOptionalNamespaceModifiers(InputParameterName, FNiagaraParameterUtilities::EParameterContext::System, OptionalNamespaceModifiers);
+	for (FName OptionalNamespaceModifier : OptionalNamespaceModifiers)
+	{
+		TAttribute<FText> SetToolTip;
+		SetToolTip.Bind(TAttribute<FText>::FGetter::CreateSP(this, &SNiagaraStackFunctionInputName::GetSetNamespaceModifierToolTip, OptionalNamespaceModifier));
+		MenuBuilder.AddMenuEntry(
+			FText::FromName(OptionalNamespaceModifier),
+			SetToolTip,
+			FSlateIcon(),
+			FUIAction(
+				FExecuteAction::CreateSP(this, &SNiagaraStackFunctionInputName::OnSetNamespaceModifier, OptionalNamespaceModifier),
+				FCanExecuteAction::CreateSP(this, &SNiagaraStackFunctionInputName::CanSetNamespaceModifier, OptionalNamespaceModifier)));
+	}
 
-	TAttribute<FText> RemoveToolTip;
-	RemoveToolTip.Bind(this, &SNiagaraStackFunctionInputName::GetRemoveNamespaceModifierToolTip);
+	TAttribute<FText> SetCustomToolTip;
+	SetCustomToolTip.Bind(TAttribute<FText>::FGetter::CreateSP(this, &SNiagaraStackFunctionInputName::GetSetCustomNamespaceModifierToolTip));
 	MenuBuilder.AddMenuEntry(
-		LOCTEXT("RemoveNamespaceModifier", "Remove"),
-		RemoveToolTip,
+		LOCTEXT("CustomNamespaceModifier", "Custom..."),
+		SetCustomToolTip,
 		FSlateIcon(),
 		FUIAction(
-			FExecuteAction::CreateSP(this, &SNiagaraStackFunctionInputName::OnRemoveNamespaceModifier),
-			FCanExecuteAction::CreateSP(this, &SNiagaraStackFunctionInputName::CanRemoveNamespaceModifier)));
+			FExecuteAction::CreateSP(this, &SNiagaraStackFunctionInputName::OnSetCustomNamespaceModifier),
+			FCanExecuteAction::CreateSP(this, &SNiagaraStackFunctionInputName::CanSetCustomNamespaceModifier)));
 
-	TAttribute<FText> EditToolTip;
-	EditToolTip.Bind(this, &SNiagaraStackFunctionInputName::GetEditNamespaceModifierToolTip);
+	TAttribute<FText> SetNoneToolTip;
+	SetNoneToolTip.Bind(TAttribute<FText>::FGetter::CreateSP(this, &SNiagaraStackFunctionInputName::GetSetNamespaceModifierToolTip, FName(NAME_None)));
 	MenuBuilder.AddMenuEntry(
-		LOCTEXT("EditNamespaceModifier", "Edit"),
-		EditToolTip,
+		LOCTEXT("NoneNamespaceModifier", "None"),
+		SetNoneToolTip,
 		FSlateIcon(),
 		FUIAction(
-			FExecuteAction::CreateSP(this, &SNiagaraStackFunctionInputName::OnEditNamespaceModifier),
-			FCanExecuteAction::CreateSP(this, &SNiagaraStackFunctionInputName::CanEditNamespaceModifier)));
+			FExecuteAction::CreateSP(this, &SNiagaraStackFunctionInputName::OnSetNamespaceModifier, FName(NAME_None)),
+			FCanExecuteAction::CreateSP(this, &SNiagaraStackFunctionInputName::CanSetNamespaceModifier, FName(NAME_None))));
 }
 
-FText SNiagaraStackFunctionInputName::GetAddNamespaceModifierToolTip() const
+FText SNiagaraStackFunctionInputName::GetSetNamespaceModifierToolTip(FName NamespaceModifier) const
 {
 	FName InputParameterName = *FunctionInput->GetDisplayName().ToString();
 	FText AddMessage;
-	FNiagaraParameterUtilities::TestCanAddNamespaceModifierWithMessage(InputParameterName, AddMessage);
+	FNiagaraParameterUtilities::TestCanSetSpecificNamespaceModifierWithMessage(InputParameterName, NamespaceModifier, AddMessage);
 	return AddMessage;
 }
 
-bool SNiagaraStackFunctionInputName::CanAddNamespaceModifier() const
+bool SNiagaraStackFunctionInputName::CanSetNamespaceModifier(FName NamespaceModifier) const
 {
 	FName InputParameterName = *FunctionInput->GetDisplayName().ToString();
 	FText Unused;
-	return FNiagaraParameterUtilities::TestCanAddNamespaceModifierWithMessage(InputParameterName, Unused);
+	return FNiagaraParameterUtilities::TestCanSetSpecificNamespaceModifierWithMessage(InputParameterName, NamespaceModifier, Unused);
 }
 
-void SNiagaraStackFunctionInputName::OnAddNamespaceModifier()
+void SNiagaraStackFunctionInputName::OnSetNamespaceModifier(FName NamespaceModifier)
 {
 	FName InputParameterName = *FunctionInput->GetDisplayName().ToString();
-	FName NewName = FNiagaraParameterUtilities::AddNamespaceModifier(InputParameterName);
+	FName NewName = FNiagaraParameterUtilities::SetSpecificNamespaceModifier(InputParameterName, NamespaceModifier);
 	if (NewName != NAME_None)
 	{
 		FScopedTransaction Transaction(LOCTEXT("AddNamespaceModifierTransaction", "Add namespace modifier."));
@@ -265,53 +271,36 @@ void SNiagaraStackFunctionInputName::OnAddNamespaceModifier()
 	}
 }
 
-FText SNiagaraStackFunctionInputName::GetRemoveNamespaceModifierToolTip() const
+FText SNiagaraStackFunctionInputName::GetSetCustomNamespaceModifierToolTip() const
 {
 	FName InputParameterName = *FunctionInput->GetDisplayName().ToString();
-	FText RemoveMessage;
-	FNiagaraParameterUtilities::TestCanRemoveNamespaceModifierWithMessage(InputParameterName, RemoveMessage);
-	return RemoveMessage;
+	FText AddMessage;
+	FNiagaraParameterUtilities::TestCanSetCustomNamespaceModifierWithMessage(InputParameterName, AddMessage);
+	return AddMessage;
 }
 
-bool SNiagaraStackFunctionInputName::CanRemoveNamespaceModifier() const
+bool SNiagaraStackFunctionInputName::CanSetCustomNamespaceModifier() const
 {
 	FName InputParameterName = *FunctionInput->GetDisplayName().ToString();
 	FText Unused;
-	return FNiagaraParameterUtilities::TestCanRemoveNamespaceModifierWithMessage(InputParameterName, Unused);
+	return FNiagaraParameterUtilities::TestCanSetCustomNamespaceModifierWithMessage(InputParameterName, Unused);
 }
 
-void SNiagaraStackFunctionInputName::OnRemoveNamespaceModifier()
+void SNiagaraStackFunctionInputName::OnSetCustomNamespaceModifier()
 {
 	FName InputParameterName = *FunctionInput->GetDisplayName().ToString();
-	FText Unused;
-	FName NewName = FNiagaraParameterUtilities::RemoveNamespaceModifier(InputParameterName);
+	FName NewName = FNiagaraParameterUtilities::SetCustomNamespaceModifier(InputParameterName);
 	if (NewName != NAME_None)
 	{
-		FScopedTransaction Transaction(LOCTEXT("RemoveNamespaceModifierTransaction", "Remove namespace modifier."));
-		FunctionInput->OnRenamed(FText::FromName(NewName));
-	}
-}
-
-FText SNiagaraStackFunctionInputName::GetEditNamespaceModifierToolTip() const
-{
-	FName InputParameterName = *FunctionInput->GetDisplayName().ToString();
-	FText EditMessage;
-	FNiagaraParameterUtilities::TestCanEditNamespaceModifierWithMessage(InputParameterName, EditMessage);
-	return EditMessage;
-}
-
-bool SNiagaraStackFunctionInputName::CanEditNamespaceModifier() const
-{
-	FName InputParameterName = *FunctionInput->GetDisplayName().ToString();
-	FText Unused;
-	return FNiagaraParameterUtilities::TestCanEditNamespaceModifierWithMessage(InputParameterName, Unused);
-}
-
-void SNiagaraStackFunctionInputName::OnEditNamespaceModifier()
-{
-	if (ParameterTextBlock.IsValid())
-	{
-		ParameterTextBlock->EnterNamespaceModifierEditingMode();
+		if (NewName != InputParameterName)
+		{
+			FScopedTransaction Transaction(LOCTEXT("AddNamespaceModifierTransaction", "Add namespace modifier."));
+			FunctionInput->OnRenamed(FText::FromName(NewName));
+		}
+		if (ParameterTextBlock.IsValid())
+		{
+			ParameterTextBlock->EnterNamespaceModifierEditingMode();
+		}
 	}
 }
 
diff --git a/Engine/Plugins/FX/Niagara/Source/NiagaraEditorWidgets/Private/Stack/SNiagaraStackFunctionInputName.h b/Engine/Plugins/FX/Niagara/Source/NiagaraEditorWidgets/Private/Stack/SNiagaraStackFunctionInputName.h
--- a/Engine/Plugins/FX/Niagara/Source/NiagaraEditorWidgets/Private/Stack/SNiagaraStackFunctionInputName.h
+++ b/Engine/Plugins/FX/Niagara/Source/NiagaraEditorWidgets/Private/Stack/SNiagaraStackFunctionInputName.h
@@ -49,17 +49,13 @@ private:
 
 	void GetChangeNamespaceModifierSubMenu(FMenuBuilder& MenuBuilder);
 
-	FText GetAddNamespaceModifierToolTip() const;
-	bool CanAddNamespaceModifier() const;
-	void OnAddNamespaceModifier();
+	FText GetSetNamespaceModifierToolTip(FName InNamespaceModifier) const;
+	bool CanSetNamespaceModifier(FName InNamespaceModifier) const;
+	void OnSetNamespaceModifier(FName InNamespaceModifier);
 
-	FText GetRemoveNamespaceModifierToolTip() const;
-	bool CanRemoveNamespaceModifier() const;
-	void OnRemoveNamespaceModifier();
-
-	FText GetEditNamespaceModifierToolTip() const;
-	bool CanEditNamespaceModifier() const;
-	void OnEditNamespaceModifier();
+	FText GetSetCustomNamespaceModifierToolTip() const;
+	bool CanSetCustomNamespaceModifier() const;
+	void OnSetCustomNamespaceModifier();
 
 	FText GetCopyParameterReferenceToolTip() const;
 	bool CanCopyParameterReference() const;
