Improving updates for the thumbnail previews of the renderers
#rnx

diff --git a/Engine/Plugins/FX/Niagara/Source/NiagaraEditorWidgets/Private/SNiagaraOverviewStackNode.cpp b/Engine/Plugins/FX/Niagara/Source/NiagaraEditorWidgets/Private/SNiagaraOverviewStackNode.cpp
--- a/Engine/Plugins/FX/Niagara/Source/NiagaraEditorWidgets/Private/SNiagaraOverviewStackNode.cpp
+++ b/Engine/Plugins/FX/Niagara/Source/NiagaraEditorWidgets/Private/SNiagaraOverviewStackNode.cpp
@@ -26,6 +26,7 @@
 #include "Widgets/SWidget.h"
 #include "Subsystems/AssetEditorSubsystem.h"
 #include "ViewModels/Stack/NiagaraStackRendererItem.h"
+#include "Materials/MaterialInterface.h"
 
 #define LOCTEXT_NAMESPACE "NiagaraOverviewStackNode"
 
@@ -62,8 +63,10 @@ void SNiagaraOverviewStackNode::Construct(const FArguments& InArgs, UNiagaraOver
 			}
 			if (StackViewModel)
 			{
+				StackViewModel->OnStructureChanged().AddSP(this, &SNiagaraOverviewStackNode::RefreshThumbnailBar);
 				StackViewModel->OnDataObjectChanged().AddSP(this, &SNiagaraOverviewStackNode::FillThumbnailBar, true);
 			}
+			UMaterial::OnMaterialCompilationFinished().AddSP(this, &SNiagaraOverviewStackNode::OnMaterialCompiled);
 			OverviewSelectionViewModel = OwningSystemViewModel->GetSelectionViewModel();
 		}
 	}
@@ -250,6 +253,34 @@ void SNiagaraOverviewStackNode::Tick(const FGeometry& AllottedGeometry, const do
 	}
 }
 
+void SNiagaraOverviewStackNode::OnMaterialCompiled(class UMaterialInterface* MaterialInterface)
+{
+	bool bUsingThisMaterial = false;
+	if (EmitterHandleViewModelWeak.IsValid())
+	{
+		EmitterHandleViewModelWeak.Pin()->GetRendererEntries(PreviewStackEntries);
+		FNiagaraEmitterInstance* InInstance = EmitterHandleViewModelWeak.Pin()->GetEmitterViewModel()->GetSimulation().IsValid() ? EmitterHandleViewModelWeak.Pin()->GetEmitterViewModel()->GetSimulation().Pin().Get() : nullptr;
+		for (UNiagaraStackEntry* Entry : PreviewStackEntries)
+		{
+			if (UNiagaraStackRendererItem* RendererItem = Cast<UNiagaraStackRendererItem>(Entry))
+			{
+				TArray<UMaterialInterface*> Materials;
+				RendererItem->GetRendererProperties()->GetUsedMaterials(InInstance, Materials);
+				if (Materials.Contains(MaterialInterface))
+				{
+					bUsingThisMaterial = true;
+					break;
+				}
+			}
+		}
+
+		if (bUsingThisMaterial)
+		{
+			RefreshThumbnailBar();
+		}
+	}
+}
+
 TSharedRef<SWidget> SNiagaraOverviewStackNode::CreateNodeContentArea()
 {
 	TSharedPtr<SWidget> ContentWidget;
@@ -325,6 +356,11 @@ TSharedRef<SWidget> SNiagaraOverviewStackNode::CreateNodeContentArea()
 
 }
 
+void SNiagaraOverviewStackNode::RefreshThumbnailBar()
+{
+	FillThumbnailBar(nullptr, false);
+}
+
 void SNiagaraOverviewStackNode::FillThumbnailBar(UObject* ChangedObject, const bool bIsTriggeredByObjectUpdate)
 {
 	UNiagaraRendererProperties* RendererProperties = Cast< UNiagaraRendererProperties>(ChangedObject);
diff --git a/Engine/Plugins/FX/Niagara/Source/NiagaraEditorWidgets/Private/SNiagaraOverviewStackNode.h b/Engine/Plugins/FX/Niagara/Source/NiagaraEditorWidgets/Private/SNiagaraOverviewStackNode.h
--- a/Engine/Plugins/FX/Niagara/Source/NiagaraEditorWidgets/Private/SNiagaraOverviewStackNode.h
+++ b/Engine/Plugins/FX/Niagara/Source/NiagaraEditorWidgets/Private/SNiagaraOverviewStackNode.h
@@ -28,8 +28,9 @@ protected:
 	virtual TSharedRef<SWidget> CreateTitleWidget(TSharedPtr<SNodeTitle> NodeTitle) override;
 	virtual TSharedRef<SWidget> CreateTitleRightWidget() override;
 	virtual TSharedRef<SWidget> CreateNodeContentArea() override;
+	void RefreshThumbnailBar();
 	void FillThumbnailBar(UObject* ChangedObject, const bool bIsTriggeredByObjectUpdate);
-
+	void OnMaterialCompiled(class UMaterialInterface* MaterialInterface);
 private:
 	EVisibility GetIssueIconVisibility() const;
 	EVisibility GetEnabledCheckBoxVisibility() const;
