Adding missing header
#rnx

diff --git a/Engine/Plugins/FX/Niagara/Source/NiagaraEditorWidgets/Private/SNiagaraOverviewStackNode.cpp b/Engine/Plugins/FX/Niagara/Source/NiagaraEditorWidgets/Private/SNiagaraOverviewStackNode.cpp
--- a/Engine/Plugins/FX/Niagara/Source/NiagaraEditorWidgets/Private/SNiagaraOverviewStackNode.cpp
+++ b/Engine/Plugins/FX/Niagara/Source/NiagaraEditorWidgets/Private/SNiagaraOverviewStackNode.cpp
@@ -27,6 +27,7 @@
 #include "Subsystems/AssetEditorSubsystem.h"
 #include "ViewModels/Stack/NiagaraStackRendererItem.h"
 #include "Materials/MaterialInterface.h"
+#include "Materials/Material.h"
 
 #define LOCTEXT_NAMESPACE "NiagaraOverviewStackNode"
 
