Niagara - Parameter namespace modifier fixes.
+ Fix issues where custom namespace modifiers wouldn't go into rename mode in the parameters panel.
+ Fix an issue where moving an output parameter to particles wouldn't remove the extra modifiers.
+ Fix an issue where moving a parameter to the output namespace wouldn't add in the required module modifier.
+ Change the "None" menu entriy for namespace modifiers to "Clear" to be more consistent with the rest of the engine.
+ Change the assignment node to show "Set Parameters" instead of "Set Variables"
+ Fix duplicate loctext issues.
+ Fix an issue using the "Read from new..." menu options on a Set Variables module could make parameters with non-supported namespaces.

diff --git a/Engine/Plugins/FX/Niagara/Source/NiagaraEditor/Private/NiagaraActions.cpp b/Engine/Plugins/FX/Niagara/Source/NiagaraEditor/Private/NiagaraActions.cpp
--- a/Engine/Plugins/FX/Niagara/Source/NiagaraEditor/Private/NiagaraActions.cpp
+++ b/Engine/Plugins/FX/Niagara/Source/NiagaraEditor/Private/NiagaraActions.cpp
@@ -43,22 +43,52 @@ void FNiagaraMenuAction::SetParamterVariable(const FNiagaraVariable& InParameter
 /************************************************************************/
 FNiagaraParameterAction::FNiagaraParameterAction(const FNiagaraVariable& InParameter,
 	const TArray<FNiagaraGraphParameterReferenceCollection>& InReferenceCollection,
-	FText InNodeCategory, FText InMenuDesc, FText InToolTip, const int32 InGrouping, FText InKeywords, int32 InSectionID)
+	FText InNodeCategory, FText InMenuDesc, FText InToolTip, const int32 InGrouping, FText InKeywords,
+	TSharedPtr<TArray<FName>> ParameterWithNamespaceModifierRenamePending, 
+	int32 InSectionID)
 	: FEdGraphSchemaAction(MoveTemp(InNodeCategory), MoveTemp(InMenuDesc), MoveTemp(InToolTip), InGrouping, MoveTemp(InKeywords), InSectionID)
 	, Parameter(InParameter)
 	, ReferenceCollection(InReferenceCollection)
 	, bIsExternallyReferenced(false)
-	, bNamespaceModifierRenamePending(false)
+	, ParameterWithNamespaceModifierRenamePendingWeak(ParameterWithNamespaceModifierRenamePending)
 {
 }
 
 FNiagaraParameterAction::FNiagaraParameterAction(const FNiagaraVariable& InParameter,
-	FText InNodeCategory, FText InMenuDesc, FText InToolTip, const int32 InGrouping, FText InKeywords, int32 InSectionID)
+	FText InNodeCategory, FText InMenuDesc, FText InToolTip, const int32 InGrouping, FText InKeywords,
+	TSharedPtr<TArray<FName>> ParameterWithNamespaceModifierRenamePending, 
+	int32 InSectionID)
 	: FEdGraphSchemaAction(MoveTemp(InNodeCategory), MoveTemp(InMenuDesc), MoveTemp(InToolTip), InGrouping, MoveTemp(InKeywords), InSectionID)
 	, Parameter(InParameter)
 	, bIsExternallyReferenced(false)
-	, bNamespaceModifierRenamePending(false)
+	, ParameterWithNamespaceModifierRenamePendingWeak(ParameterWithNamespaceModifierRenamePending)
+{
+}
+
+bool FNiagaraParameterAction::GetIsNamespaceModifierRenamePending() const
+{
+	TSharedPtr<TArray<FName>> ParameterNamesWithNamespaceModifierRenamePending = ParameterWithNamespaceModifierRenamePendingWeak.Pin();
+	if (ParameterNamesWithNamespaceModifierRenamePending.IsValid())
+	{
+		return ParameterNamesWithNamespaceModifierRenamePending->Contains(Parameter.GetName());
+	}
+	return false;
+}
+
+void FNiagaraParameterAction::SetIsNamespaceModifierRenamePending(bool bIsNamespaceModifierRenamePending)
 {
+	TSharedPtr<TArray<FName>> ParameterNamesWithNamespaceModifierRenamePending = ParameterWithNamespaceModifierRenamePendingWeak.Pin();
+	if (ParameterNamesWithNamespaceModifierRenamePending.IsValid())
+	{
+		if (bIsNamespaceModifierRenamePending)
+		{
+			ParameterNamesWithNamespaceModifierRenamePending->AddUnique(Parameter.GetName());
+		}
+		else
+		{
+			ParameterNamesWithNamespaceModifierRenamePending->Remove(Parameter.GetName());
+		}
+	}
 }
 
 /************************************************************************/
diff --git a/Engine/Plugins/FX/Niagara/Source/NiagaraEditor/Private/NiagaraEditorUtilities.cpp b/Engine/Plugins/FX/Niagara/Source/NiagaraEditor/Private/NiagaraEditorUtilities.cpp
--- a/Engine/Plugins/FX/Niagara/Source/NiagaraEditor/Private/NiagaraEditorUtilities.cpp
+++ b/Engine/Plugins/FX/Niagara/Source/NiagaraEditor/Private/NiagaraEditorUtilities.cpp
@@ -2757,12 +2757,25 @@ FName FNiagaraParameterUtilities::ChangeNamespace(FName ParameterName, const FNi
 	{
 		TArray<FName> NameParts = ParameterHandle.GetHandleParts();
 		NameParts.RemoveAt(0, NamespaceMetadata.Namespaces.Num());
-		bool bHasNamespaceModifier = NameParts.Num() > 1;
-		bool bNewNamespaceCanHaveNamespaceModifier =
-			NewNamespaceMetadata.Options.Contains(ENiagaraNamespaceMetadataOptions::PreventEditingNamespaceModifier) == false;
-		if (bHasNamespaceModifier && bNewNamespaceCanHaveNamespaceModifier == false)
+		int32 NumberOfNamespaceModifiers = NameParts.Num() - 1;
+		if (NumberOfNamespaceModifiers > 0)
 		{
-			NameParts.RemoveAt(0);
+			bool bNewNamespaceCanHaveNamespaceModifier =
+				NewNamespaceMetadata.Options.Contains(ENiagaraNamespaceMetadataOptions::PreventEditingNamespaceModifier) == false;
+			if (bNewNamespaceCanHaveNamespaceModifier == false)
+			{
+				// Remove all modifiers.
+				NameParts.RemoveAt(0, NumberOfNamespaceModifiers);
+			}
+			else 
+			{
+				// Remove all but the last modifier.
+				NameParts.RemoveAt(0, NumberOfNamespaceModifiers - 1);
+			}
+		}
+		if (NewNamespaceMetadata.RequiredNamespaceModifier != NAME_None && NameParts.Num() > 1 && NameParts[0] != NewNamespaceMetadata.RequiredNamespaceModifier)
+		{
+			NameParts.Insert(NewNamespaceMetadata.RequiredNamespaceModifier, 0);
 		}
 		NameParts.Insert(NewNamespaceMetadata.Namespaces, 0);
 		return NamePartsToName(NameParts);
@@ -2839,7 +2852,14 @@ bool FNiagaraParameterUtilities::TestCanSetSpecificNamespaceModifierWithMessage(
 		}
 		else
 		{
-			OutMessage = FText::Format(LOCTEXT("SetNamespaceModifierFormat", "Set the namespace modifier for this parameter to {0}."), FText::FromName(InNamespaceModifier));
+			if(InNamespaceModifier == NAME_None)
+			{ 
+				OutMessage = LOCTEXT("ClearNamespaceModifier", "Clears the namespace modifier on this parameter.");
+			}
+			else 
+			{
+				OutMessage = FText::Format(LOCTEXT("SetNamespaceModifierFormat", "Set the namespace modifier for this parameter to {0}."), FText::FromName(InNamespaceModifier));
+			}
 		}
 		return true;
 	}
diff --git a/Engine/Plugins/FX/Niagara/Source/NiagaraEditor/Private/NiagaraNodeAssignment.cpp b/Engine/Plugins/FX/Niagara/Source/NiagaraEditor/Private/NiagaraNodeAssignment.cpp
--- a/Engine/Plugins/FX/Niagara/Source/NiagaraEditor/Private/NiagaraNodeAssignment.cpp
+++ b/Engine/Plugins/FX/Niagara/Source/NiagaraEditor/Private/NiagaraNodeAssignment.cpp
@@ -52,7 +52,7 @@ void UNiagaraNodeAssignment::RefreshTitle()
 	}
 	else
 	{
-		Title = LOCTEXT("NodeTitle", "Set Variables");
+		Title = LOCTEXT("NodeTitle", "Set Parameters");
 	}
 }
 
diff --git a/Engine/Plugins/FX/Niagara/Source/NiagaraEditor/Private/NiagaraNodeParameterMapBase.cpp b/Engine/Plugins/FX/Niagara/Source/NiagaraEditor/Private/NiagaraNodeParameterMapBase.cpp
--- a/Engine/Plugins/FX/Niagara/Source/NiagaraEditor/Private/NiagaraNodeParameterMapBase.cpp
+++ b/Engine/Plugins/FX/Niagara/Source/NiagaraEditor/Private/NiagaraNodeParameterMapBase.cpp
@@ -404,7 +404,7 @@ void UNiagaraNodeParameterMapBase::GetChangeNamespaceModifierSubMenuForPin(UTool
 		this, &UNiagaraNodeParameterMapBase::GetSetNamespaceModifierForPinToolTip, (const UEdGraphPin*)InPin, FName(NAME_None)));
 	Section.AddMenuEntry(
 		"AddNoneNamespaceModifier",
-		LOCTEXT("SetNoneNamespaceModifierForPin", "None"),
+		LOCTEXT("SetNoneNamespaceModifierForPin", "Clear"),
 		SetNoneToolTip,
 		FSlateIcon(),
 		FUIAction(
@@ -455,7 +455,7 @@ void UNiagaraNodeParameterMapBase::SetCustomNamespaceModifierForPin(UEdGraphPin*
 	{
 		if (NewName != InPin->PinName)
 		{
-			FScopedTransaction Transaction(LOCTEXT("AddNamespaceModifierTransaction", "Add custom namespace modifier"));
+			FScopedTransaction Transaction(LOCTEXT("AddCustomNamespaceModifierTransaction", "Add custom namespace modifier"));
 			CommitEditablePinName(FText::FromName(NewName), InPin);
 		}
 		SetIsPinEditNamespaceModifierPending(InPin, true);
diff --git a/Engine/Plugins/FX/Niagara/Source/NiagaraEditor/Private/ViewModels/Stack/NiagaraStackModuleItem.cpp b/Engine/Plugins/FX/Niagara/Source/NiagaraEditor/Private/ViewModels/Stack/NiagaraStackModuleItem.cpp
--- a/Engine/Plugins/FX/Niagara/Source/NiagaraEditor/Private/ViewModels/Stack/NiagaraStackModuleItem.cpp
+++ b/Engine/Plugins/FX/Niagara/Source/NiagaraEditor/Private/ViewModels/Stack/NiagaraStackModuleItem.cpp
@@ -255,7 +255,7 @@ void UNiagaraStackModuleItem::RefreshChildrenInternal(const TArray<UNiagaraStack
 			UNiagaraNodeAssignment* AssignmentNode = CastChecked<UNiagaraNodeAssignment>(FunctionCallNode);
 			if (AssignmentNode->GetAssignmentTargets().Num() == 0)
 			{
-				FText EmptyAssignmentNodeMessageText = LOCTEXT("EmptyAssignmentNodeMessage", "No variables\n\nTo add a variable use the add button in the header, or drag a parameter from the parameters tab to the header.");
+				FText EmptyAssignmentNodeMessageText = LOCTEXT("EmptyAssignmentNodeMessage", "No Parameters\n\nTo add a parameter use the add button in the header, or drag a parameter from the parameters tab to the header.");
 				UNiagaraStackItemTextContent* EmptyAssignmentNodeMessage = FindCurrentChildOfTypeByPredicate<UNiagaraStackItemTextContent>(NewChildren,
 					[&](UNiagaraStackItemTextContent* CurrentStackItemTextContent) { return CurrentStackItemTextContent->GetDisplayName().IdenticalTo(EmptyAssignmentNodeMessageText); });
 
diff --git a/Engine/Plugins/FX/Niagara/Source/NiagaraEditor/Private/ViewModels/Stack/NiagaraStackScriptItemGroup.cpp b/Engine/Plugins/FX/Niagara/Source/NiagaraEditor/Private/ViewModels/Stack/NiagaraStackScriptItemGroup.cpp
--- a/Engine/Plugins/FX/Niagara/Source/NiagaraEditor/Private/ViewModels/Stack/NiagaraStackScriptItemGroup.cpp
+++ b/Engine/Plugins/FX/Niagara/Source/NiagaraEditor/Private/ViewModels/Stack/NiagaraStackScriptItemGroup.cpp
@@ -108,8 +108,8 @@ public:
 
 	static TSharedRef<FScriptGroupAddAction> CreateNewSetSpecificModuleAction()
 	{
-		FText DisplayName = LOCTEXT("NewSetSpecificModuleName", "Set new or existing value directly");
-		FText Description = LOCTEXT("NewSetSpecificModuleDescription", "Description: Create a new module that can set new orexisting values directly.");
+		FText DisplayName = LOCTEXT("NewSetSpecificModuleName", "Set new or existing parameter directly");
+		FText Description = LOCTEXT("NewSetSpecificModuleDescription", "Description: Create a new module that can set new or existing parameters directly.");
 		return MakeShareable(new FScriptGroupAddAction(FText(), DisplayName, Description, FText(), FNiagaraVariable(), false, FAssetData(), nullptr, false, true));
 	}
 
diff --git a/Engine/Plugins/FX/Niagara/Source/NiagaraEditor/Private/Widgets/SNiagaraParameterMapPaletteItem.cpp b/Engine/Plugins/FX/Niagara/Source/NiagaraEditor/Private/Widgets/SNiagaraParameterMapPaletteItem.cpp
--- a/Engine/Plugins/FX/Niagara/Source/NiagaraEditor/Private/Widgets/SNiagaraParameterMapPaletteItem.cpp
+++ b/Engine/Plugins/FX/Niagara/Source/NiagaraEditor/Private/Widgets/SNiagaraParameterMapPaletteItem.cpp
@@ -125,9 +125,9 @@ void SNiagaraParameterMapPalleteItem::Tick(const FGeometry& AllottedGeometry, co
 	if (ParameterNameTextBlock.IsValid())
 	{
 		TSharedPtr<FNiagaraParameterAction> ParameterAction = StaticCastSharedPtr<FNiagaraParameterAction>(ActionPtr.Pin());
-		if (ParameterAction.IsValid() && ParameterAction->bNamespaceModifierRenamePending)
+		if (ParameterAction.IsValid() && ParameterAction->GetIsNamespaceModifierRenamePending())
 		{
-			ParameterAction->bNamespaceModifierRenamePending = false;
+			ParameterAction->SetIsNamespaceModifierRenamePending(false);
 			ParameterNameTextBlock->EnterNamespaceModifierEditingMode();
 		}
 	}
diff --git a/Engine/Plugins/FX/Niagara/Source/NiagaraEditor/Private/Widgets/SNiagaraParameterMapView.cpp b/Engine/Plugins/FX/Niagara/Source/NiagaraEditor/Private/Widgets/SNiagaraParameterMapView.cpp
--- a/Engine/Plugins/FX/Niagara/Source/NiagaraEditor/Private/Widgets/SNiagaraParameterMapView.cpp
+++ b/Engine/Plugins/FX/Niagara/Source/NiagaraEditor/Private/Widgets/SNiagaraParameterMapView.cpp
@@ -225,6 +225,7 @@ void SNiagaraParameterMapView::Construct(const FArguments& InArgs, const TArray<
 	ToolkitCommands = InToolkitCommands;
 	AddParameterButtons.SetNum(NiagaraParameterMapSectionID::Num);
 	const FVector2D ViewOptionsShadowOffset = FNiagaraEditorStyle::Get().GetVector("NiagaraEditor.Stack.ViewOptionsShadowOffset");
+	ParametersWithNamespaceModifierRenamePending = MakeShared<TArray<FName>>();
 
 	SelectedScriptObjects = InSelectedObjects[0];
 	SelectedScriptObjects->OnSelectedObjectsChanged().AddSP(this, &SNiagaraParameterMapView::SelectedObjectsChanged);
@@ -578,7 +579,7 @@ void SNiagaraParameterMapView::CollectAllActions(FGraphActionListBuilderBase& Ou
 
 		const FText Name = FNiagaraParameterUtilities::FormatParameterNameForTextDisplay(Parameter.GetName());
 		const FText Tooltip = FText::Format(TooltipFormat, Name, Parameter.GetType().GetNameText());
-		TSharedPtr<FNiagaraParameterAction> ParameterAction(new FNiagaraParameterAction(Parameter, ParameterEntry.Value, FText::GetEmpty(), Name, Tooltip, 0, FText(), Section));
+		TSharedPtr<FNiagaraParameterAction> ParameterAction(new FNiagaraParameterAction(Parameter, ParameterEntry.Value, FText::GetEmpty(), Name, Tooltip, 0, FText(), ParametersWithNamespaceModifierRenamePending, Section));
 		ParameterAction->bIsExternallyReferenced = bIsExternallyReferenced;
 		OutAllActions.AddAction(ParameterAction);
 		LastCollectedParameters.Add(Parameter);
@@ -855,7 +856,7 @@ TSharedPtr<SWidget> SNiagaraParameterMapView::OnContextMenuOpening()
 
 			MenuBuilder.AddSubMenu(
 				LOCTEXT("DuplicateWithNewNamespaceModifier", "Duplicate with Namespace Modifier"),
-				LOCTEXT("ChangeNamespaceModifierToolTip", "Duplicate this parameter with a different namespace modifier."),
+				LOCTEXT("DupilcateWithNewNamespaceModifierToolTip", "Duplicate this parameter with a different namespace modifier."),
 				FNewMenuDelegate::CreateSP(this, &SNiagaraParameterMapView::GetChangeNamespaceModifierSubMenu, true));
 		}
 		MenuBuilder.EndSection();
@@ -1377,7 +1378,7 @@ void SNiagaraParameterMapView::GetChangeNamespaceModifierSubMenu(FMenuBuilder& M
 	TAttribute<FText> SetNoneToolTip;
 	SetNoneToolTip.Bind(TAttribute<FText>::FGetter::CreateSP(this, &SNiagaraParameterMapView::GetSetNamespaceModifierToolTip, FName(NAME_None), bDuplicateParameter));
 	MenuBuilder.AddMenuEntry(
-		LOCTEXT("NoneNamespaceModifier", "None"),
+		LOCTEXT("NoneNamespaceModifier", "Clear"),
 		SetNoneToolTip,
 		FSlateIcon(),
 		FUIAction(
@@ -1535,16 +1536,10 @@ void SNiagaraParameterMapView::OnSetCustomNamespaceModifier(bool bDuplicateParam
 				{
 					NewUniqueName = NewName;
 				}
-				FScopedTransaction Transaction(LOCTEXT("DuplicateParameterToWithNamespaceModifierTransaction", "Duplicate parameter with custom namespace modifier"));
+				FScopedTransaction Transaction(LOCTEXT("DuplicateParameterToWithCustomNamespaceModifierTransaction", "Duplicate parameter with custom namespace modifier"));
 				AddParameter(FNiagaraVariable(ParameterAction->Parameter.GetType(), NewUniqueName), false);
 
-				FName ActionName = *FNiagaraParameterUtilities::FormatParameterNameForTextDisplay(NewUniqueName).ToString();
-				GraphActionMenu->SelectItemByName(ActionName);
-				TSharedPtr<FNiagaraParameterAction> RenamedParameterAction;
-				if (GetSingleParameterActionForSelection(RenamedParameterAction, Unused))
-				{
-					RenamedParameterAction->bNamespaceModifierRenamePending = true;
-				}
+				ParametersWithNamespaceModifierRenamePending->Add(NewUniqueName);
 			}
 			else
 			{
@@ -1552,19 +1547,8 @@ void SNiagaraParameterMapView::OnSetCustomNamespaceModifier(bool bDuplicateParam
 				{
 					FScopedTransaction Transaction(LOCTEXT("SetCustomNamespaceModifierTransaction", "Set custom namespace modifier"));
 					RenameParameter(ParameterAction, NewName);
-
-					FName ActionName = *FNiagaraParameterUtilities::FormatParameterNameForTextDisplay(NewName).ToString();
-					GraphActionMenu->SelectItemByName(ActionName);
-					TSharedPtr<FNiagaraParameterAction> RenamedParameterAction;
-					if (GetSingleParameterActionForSelection(RenamedParameterAction, Unused))
-					{
-						RenamedParameterAction->bNamespaceModifierRenamePending = true;
-					}
-				}
-				else
-				{
-					ParameterAction->bNamespaceModifierRenamePending = true;
 				}
+				ParametersWithNamespaceModifierRenamePending->Add(NewName);
 			}
 		}
 	}
@@ -1678,6 +1662,7 @@ void SNiagaraParameterMapView::RenameParameter(TSharedPtr<FNiagaraParameterActio
 		return;
 	}
 
+	bool bSuccess = false;
 	FNiagaraVariable Parameter = ParameterAction->Parameter;
 	if (ToolkitType == SCRIPT)
 	{
@@ -1700,6 +1685,7 @@ void SNiagaraParameterMapView::RenameParameter(TSharedPtr<FNiagaraParameterActio
 				}
 
 				Graph->RenameParameter(Parameter, NewName);
+				bSuccess = true;
 			}
 		}
 	}
@@ -1709,11 +1695,18 @@ void SNiagaraParameterMapView::RenameParameter(TSharedPtr<FNiagaraParameterActio
 		if (System != nullptr)
 		{
 			// Rename the parameter in the parameter stores.
-			System->GetExposedParameters().RenameParameter(Parameter, NewName);
-			System->EditorOnlyAddedParameters.RenameParameter(Parameter, NewName);
+			if (System->GetExposedParameters().IndexOf(Parameter) != INDEX_NONE)
+			{
+				System->GetExposedParameters().RenameParameter(Parameter, NewName);
+				bSuccess = true;
+			}
+			if (System->EditorOnlyAddedParameters.IndexOf(Parameter) != INDEX_NONE)
+			{
+				System->EditorOnlyAddedParameters.RenameParameter(Parameter, NewName);
+				bSuccess = true;
+			}
 
 			// Look for set variables nodes or linked inputs which reference this parameter.
-
 			for (FNiagaraGraphParameterReferenceCollection& ReferenceCollection : ParameterAction->ReferenceCollection)
 			{
 				for (FNiagaraGraphParameterReference& ParameterReference : ReferenceCollection.ParameterReferences)
@@ -1725,7 +1718,7 @@ void SNiagaraParameterMapView::RenameParameter(TSharedPtr<FNiagaraParameterActio
 						if (OwningAssignmentNode != nullptr)
 						{
 							// If this is owned by a set variables node and it's not locked, update the assignment target on the assignment node.
-							FNiagaraStackGraphUtilities::TryRenameAssignmentTarget(*OwningAssignmentNode, Parameter, NewName);
+							bSuccess = FNiagaraStackGraphUtilities::TryRenameAssignmentTarget(*OwningAssignmentNode, Parameter, NewName);
 						}
 						else
 						{
@@ -1739,6 +1732,7 @@ void SNiagaraParameterMapView::RenameParameter(TSharedPtr<FNiagaraParameterActio
 									UEdGraphPin* LinkedInputPin = *LinkedInputPinPtr;
 									LinkedInputPin->Modify();
 									LinkedInputPin->PinName = NewName;
+									bSuccess = true;
 								}
 							}
 						}
@@ -1747,6 +1741,12 @@ void SNiagaraParameterMapView::RenameParameter(TSharedPtr<FNiagaraParameterActio
 			}
 		}
 	}
+
+	if (bSuccess)
+	{
+		GraphActionMenu->RefreshAllActions(true);
+		GraphActionMenu->SelectItemByName(*FNiagaraParameterUtilities::FormatParameterNameForTextDisplay(NewName).ToString());
+	}
 }
 
 bool SNiagaraParameterMapView::IsSystemToolkit() const
diff --git a/Engine/Plugins/FX/Niagara/Source/NiagaraEditor/Public/NiagaraActions.h b/Engine/Plugins/FX/Niagara/Source/NiagaraEditor/Public/NiagaraActions.h
--- a/Engine/Plugins/FX/Niagara/Source/NiagaraEditor/Public/NiagaraActions.h
+++ b/Engine/Plugins/FX/Niagara/Source/NiagaraEditor/Public/NiagaraActions.h
@@ -62,24 +62,34 @@ struct NIAGARAEDITOR_API FNiagaraParameterAction : public FEdGraphSchemaAction
 {
 	FNiagaraParameterAction()
 		: bIsExternallyReferenced(false)
-		, bNamespaceModifierRenamePending(false)
 	{
 	}
 
 	FNiagaraParameterAction(const FNiagaraVariable& InParameter,
 		const TArray<FNiagaraGraphParameterReferenceCollection>& InReferenceCollection,
-		FText InNodeCategory, FText InMenuDesc, FText InToolTip, const int32 InGrouping, FText InKeywords, int32 InSectionID = 0);
+		FText InNodeCategory, FText InMenuDesc, FText InToolTip, const int32 InGrouping, FText InKeywords,
+		TSharedPtr<TArray<FName>> ParameterWithNamespaceModifierRenamePending,
+		int32 InSectionID = 0);
+
 	FNiagaraParameterAction(const FNiagaraVariable& InParameter, 
-		FText InNodeCategory, FText InMenuDesc, FText InToolTip, const int32 InGrouping, FText InKeywords, int32 InSectionID = 0);
+		FText InNodeCategory, FText InMenuDesc, FText InToolTip, const int32 InGrouping, FText InKeywords,
+		TSharedPtr<TArray<FName>> ParameterWithNamespaceModifierRenamePending,
+		int32 InSectionID = 0);
 
 	const FNiagaraVariable& GetParameter() const { return Parameter; }
 
+	bool GetIsNamespaceModifierRenamePending() const;
+
+	void SetIsNamespaceModifierRenamePending(bool bIsNamespaceModifierRenamePending);
+
 	FNiagaraVariable Parameter;
 
 	TArray<FNiagaraGraphParameterReferenceCollection> ReferenceCollection;
 
 	bool bIsExternallyReferenced;
-	bool bNamespaceModifierRenamePending;
+
+private:
+	TWeakPtr<TArray<FName>> ParameterWithNamespaceModifierRenamePendingWeak;
 };
 
 struct NIAGARAEDITOR_API FNiagaraScriptParameterAction : public FEdGraphSchemaAction
diff --git a/Engine/Plugins/FX/Niagara/Source/NiagaraEditor/Public/Widgets/SNiagaraParameterMapView.h b/Engine/Plugins/FX/Niagara/Source/NiagaraEditor/Public/Widgets/SNiagaraParameterMapView.h
--- a/Engine/Plugins/FX/Niagara/Source/NiagaraEditor/Public/Widgets/SNiagaraParameterMapView.h
+++ b/Engine/Plugins/FX/Niagara/Source/NiagaraEditor/Public/Widgets/SNiagaraParameterMapView.h
@@ -225,6 +225,7 @@ private:
 	TArray<int32> HiddenSectionIDs;
 
 	TArray<FNiagaraVariable> LastCollectedParameters;
+	TSharedPtr<TArray<FName>> ParametersWithNamespaceModifierRenamePending;
 
 	bool bNeedsRefresh;
 	bool bIsAddingParameter;
diff --git a/Engine/Plugins/FX/Niagara/Source/NiagaraEditorWidgets/Private/Stack/SNiagaraStackFunctionInputName.cpp b/Engine/Plugins/FX/Niagara/Source/NiagaraEditorWidgets/Private/Stack/SNiagaraStackFunctionInputName.cpp
--- a/Engine/Plugins/FX/Niagara/Source/NiagaraEditorWidgets/Private/Stack/SNiagaraStackFunctionInputName.cpp
+++ b/Engine/Plugins/FX/Niagara/Source/NiagaraEditorWidgets/Private/Stack/SNiagaraStackFunctionInputName.cpp
@@ -237,7 +237,7 @@ void SNiagaraStackFunctionInputName::GetChangeNamespaceModifierSubMenu(FMenuBuil
 	TAttribute<FText> SetNoneToolTip;
 	SetNoneToolTip.Bind(TAttribute<FText>::FGetter::CreateSP(this, &SNiagaraStackFunctionInputName::GetSetNamespaceModifierToolTip, FName(NAME_None)));
 	MenuBuilder.AddMenuEntry(
-		LOCTEXT("NoneNamespaceModifier", "None"),
+		LOCTEXT("NoneNamespaceModifier", "Clear"),
 		SetNoneToolTip,
 		FSlateIcon(),
 		FUIAction(
diff --git a/Engine/Plugins/FX/Niagara/Source/NiagaraEditorWidgets/Private/Stack/SNiagaraStackFunctionInputValue.cpp b/Engine/Plugins/FX/Niagara/Source/NiagaraEditorWidgets/Private/Stack/SNiagaraStackFunctionInputValue.cpp
--- a/Engine/Plugins/FX/Niagara/Source/NiagaraEditorWidgets/Private/Stack/SNiagaraStackFunctionInputValue.cpp
+++ b/Engine/Plugins/FX/Niagara/Source/NiagaraEditorWidgets/Private/Stack/SNiagaraStackFunctionInputValue.cpp
@@ -609,7 +609,7 @@ void SNiagaraStackFunctionInputValue::CollectAllActions(FGraphActionListBuilderB
 		{
 			InputNames.Add(FunctionInput->GetInputParameterHandlePath()[i].GetName().ToString());
 		}
-		FName InputName = *FString::Join(InputNames, TEXT("_"));
+		FName InputName = *FString::Join(InputNames, TEXT("_")).Replace(TEXT("."), TEXT("_"));
 
 		for (const FName AvailableNamespace : AvailableNamespaces)
 		{
