fix for changing the Resolution and Scope properties causing an assert- add clamps to the ScopeInMilliseconds and Resolution properties on the Oscilloscope, and make Resolution an advanced property. Also fix the erroneous Memcpy size for the VectorVMReadBuffer.

diff --git a/Engine/Plugins/FX/Niagara/Source/Niagara/Classes/NiagaraDataInterfaceAudioOscilloscope.h b/Engine/Plugins/FX/Niagara/Source/Niagara/Classes/NiagaraDataInterfaceAudioOscilloscope.h
--- a/Engine/Plugins/FX/Niagara/Source/Niagara/Classes/NiagaraDataInterfaceAudioOscilloscope.h
+++ b/Engine/Plugins/FX/Niagara/Source/Niagara/Classes/NiagaraDataInterfaceAudioOscilloscope.h
@@ -12,7 +12,6 @@
 #include "NiagaraDataInterfaceAudio.h"
 #include "NiagaraDataInterfaceAudioOscilloscope.generated.h"
 
-
 struct FNiagaraDataInterfaceProxyOscilloscope : public FNiagaraDataInterfaceProxy
 {
 	FNiagaraDataInterfaceProxyOscilloscope(int32 InResolution, float InScopeInMillseconds);
@@ -99,12 +98,16 @@ public:
 	UPROPERTY(EditAnywhere, Category = "Oscilloscope")
 	USoundSubmix* Submix;
 
+	static const int32 MaxBufferResolution = 8192;
+
 	// The number of samples of audio to pass to the GPU. Audio will be resampled to fit this resolution.
-	UPROPERTY(EditAnywhere, Category = "Oscilloscope")
+	// Increasing this number will increase the resolution of the waveform, but will also increase usage of the GPU memory bus,
+	// potentially causing issues across the application.
+	UPROPERTY(EditAnywhere, Category = "Oscilloscope", AdvancedDisplay, meta = (ClampMin = "64", ClampMax = "8192"))
 	int32 Resolution;
 
 	// The number of milliseconds of audio to show.
-	UPROPERTY(EditAnywhere, Category = "Oscilloscope")
+	UPROPERTY(EditAnywhere, Category = "Oscilloscope", meta = (ClampMin = "5.0", ClampMax = "400.0"))
 	float ScopeInMilliseconds;
 
 	//VM function overrides:
diff --git a/Engine/Plugins/FX/Niagara/Source/Niagara/Private/NiagaraDataInterfaceAudioOscilloscope.cpp b/Engine/Plugins/FX/Niagara/Source/Niagara/Private/NiagaraDataInterfaceAudioOscilloscope.cpp
--- a/Engine/Plugins/FX/Niagara/Source/Niagara/Private/NiagaraDataInterfaceAudioOscilloscope.cpp
+++ b/Engine/Plugins/FX/Niagara/Source/Niagara/Private/NiagaraDataInterfaceAudioOscilloscope.cpp
@@ -38,7 +38,7 @@ FNiagaraDataInterfaceProxyOscilloscope::FNiagaraDataInterfaceProxyOscilloscope(i
 	DeviceDestroyedHandle = FAudioDeviceManagerDelegates::OnAudioDeviceDestroyed.AddRaw(this, &FNiagaraDataInterfaceProxyOscilloscope::OnDeviceDestroyed);
 
 	VectorVMReadBuffer.Reset();
-	VectorVMReadBuffer.AddZeroed(Resolution * AUDIO_MIXER_MAX_OUTPUT_CHANNELS);
+	VectorVMReadBuffer.AddZeroed(UNiagaraDataInterfaceAudioOscilloscope::MaxBufferResolution * AUDIO_MIXER_MAX_OUTPUT_CHANNELS);
 }
 
 FNiagaraDataInterfaceProxyOscilloscope::~FNiagaraDataInterfaceProxyOscilloscope()
@@ -363,7 +363,7 @@ void FNiagaraDataInterfaceProxyOscilloscope::OnUpdateResampling(int32 InResoluti
 			GPUDownsampledBuffer.Initialize(sizeof(float), NumSamplesInBuffer, EPixelFormat::PF_R32_FLOAT, BUF_Static);
 		});
 
-		VectorVMReadBuffer.SetNumZeroed(Resolution * AUDIO_MIXER_MAX_OUTPUT_CHANNELS);
+		VectorVMReadBuffer.SetNumZeroed(UNiagaraDataInterfaceAudioOscilloscope::MaxBufferResolution * AUDIO_MIXER_MAX_OUTPUT_CHANNELS);
 	}
 }
 
@@ -557,7 +557,7 @@ int32 FNiagaraDataInterfaceProxyOscilloscope::DownsampleAudioToBuffer()
 	DownsampledBuffer.SetNumZeroed(Resolution * NumChannels);
 
 	check(DownsampledBuffer.Num() <= VectorVMReadBuffer.Num());
-	FMemory::Memcpy(VectorVMReadBuffer.GetData(), DownsampledBuffer.GetData(), VectorVMReadBuffer.Num());
+	FMemory::Memcpy(VectorVMReadBuffer.GetData(), DownsampledBuffer.GetData(), DownsampledBuffer.Num());
 
 	return Resolution;
 }
