Niagara - Missed in last checkin.

diff --git a/Engine/Plugins/FX/Niagara/Source/NiagaraEditor/Private/ViewModels/NiagaraParameterPanelViewModel.cpp b/Engine/Plugins/FX/Niagara/Source/NiagaraEditor/Private/ViewModels/NiagaraParameterPanelViewModel.cpp
--- a/Engine/Plugins/FX/Niagara/Source/NiagaraEditor/Private/ViewModels/NiagaraParameterPanelViewModel.cpp
+++ b/Engine/Plugins/FX/Niagara/Source/NiagaraEditor/Private/ViewModels/NiagaraParameterPanelViewModel.cpp
@@ -280,7 +280,7 @@ FReply FNiagaraSystemToolkitParameterPanelViewModel::HandleActionDragged(const T
 	const FText UsageText = StaticEnum<ENiagaraScriptParameterUsage>()->GetDisplayNameTextByValue((int64)ScriptVarAndViewInfo.MetaData.GetUsage());
 	const FText Tooltip = FText::Format(TooltipFormat, FText::FromName(Var.GetName()), Var.GetType().GetNameText(), ScopeText, UsageText);
 
-	TSharedPtr<FNiagaraParameterAction> ParameterAction = MakeShared<FNiagaraParameterAction>(Var, FText::GetEmpty(), Name, Tooltip, 0, FText(), 0);
+	TSharedPtr<FNiagaraParameterAction> ParameterAction = MakeShared<FNiagaraParameterAction>(Var, FText::GetEmpty(), Name, Tooltip, 0, FText(), TSharedPtr<TArray<FName>>(), 0);
 	TSharedRef<FNiagaraParameterDragOperation> DragOperation = MakeShared<FNiagaraParameterDragOperation>(ParameterAction);
 	DragOperation->CurrentHoverText = InAction->GetMenuDescription();
 	DragOperation->SetupDefaults();
