- Add Descriptions to ParticleRead functions that include known caveats also

diff --git a/Engine/Plugins/FX/Niagara/Source/Niagara/Private/NiagaraDataInterfaceParticleRead.cpp b/Engine/Plugins/FX/Niagara/Source/Niagara/Private/NiagaraDataInterfaceParticleRead.cpp
--- a/Engine/Plugins/FX/Niagara/Source/Niagara/Private/NiagaraDataInterfaceParticleRead.cpp
+++ b/Engine/Plugins/FX/Niagara/Source/Niagara/Private/NiagaraDataInterfaceParticleRead.cpp
@@ -630,22 +630,25 @@ void UNiagaraDataInterfaceParticleRead::GetFunctions(TArray<FNiagaraFunctionSign
 		Sig.bMemberFunction = true;
 		Sig.bRequiresContext = false;
 		Sig.bExperimental = true;
+#if WITH_EDITORONLY_DATA
+		Sig.Description = NSLOCTEXT("Niagara", "NiagaraDataInterfaceParticleRead_GetNumSpawnedDesc", "Returns the current number of spawned particles this frame.  Note: This is the request spawn amount, when killing particles during spawn it will not be accurate.");
+#endif
 		OutFunctions.Add(Sig);
 	}
 
 	{
 		FNiagaraFunctionSignature Sig;
 		Sig.Name = GetSpawnedIDAtIndexFunctionName;
-
 		Sig.Inputs.Add(FNiagaraVariable(FNiagaraTypeDefinition(GetClass()), TEXT("Particle Reader")));
 		Sig.Inputs.Add(FNiagaraVariable(FNiagaraTypeDefinition::GetIntDef(), TEXT("Spawn Index")));
-
 		Sig.Outputs.Add(FNiagaraVariable(FNiagaraTypeDefinition::GetBoolDef(), TEXT("Valid")));
 		Sig.Outputs.Add(FNiagaraVariable(FNiagaraTypeDefinition::GetIDDef(), TEXT("ID")));
-
 		Sig.bMemberFunction = true;
 		Sig.bRequiresContext = false;
 		Sig.bExperimental = true;
+#if WITH_EDITORONLY_DATA
+		Sig.Description = NSLOCTEXT("Niagara", "NiagaraDataInterfaceParticleRead_GetSpawnedIDAtIndexDesc", "Returns the Niagara ID for the particle spawned at Index.");
+#endif
 		OutFunctions.Add(Sig);
 	}
 
@@ -656,6 +659,9 @@ void UNiagaraDataInterfaceParticleRead::GetFunctions(TArray<FNiagaraFunctionSign
 		Sig.Outputs.Add(FNiagaraVariable(FNiagaraTypeDefinition::GetIntDef(), TEXT("Num Particles")));
 		Sig.bMemberFunction = true;
 		Sig.bRequiresContext = false;
+#if WITH_EDITORONLY_DATA
+		Sig.Description = NSLOCTEXT("Niagara", "NiagaraDataInterfaceParticleRead_GetNumParticlesDesc", "Returns the current number of particles in the buffer.  Note: When reading from self this will be the previous frames data.");
+#endif
 		OutFunctions.Add(Sig);
 	}
 
@@ -668,6 +674,9 @@ void UNiagaraDataInterfaceParticleRead::GetFunctions(TArray<FNiagaraFunctionSign
 		Sig.bMemberFunction = true;
 		Sig.bRequiresContext = false;
 		Sig.bExperimental = true;
+#if WITH_EDITORONLY_DATA
+		Sig.Description = NSLOCTEXT("Niagara", "NiagaraDataInterfaceParticleRead_GetParticleIndex", "Returns the particle index from the NiagaraID or -1 if not valid.");
+#endif
 		OutFunctions.Add(Sig);
 	}
 
@@ -688,6 +697,9 @@ void UNiagaraDataInterfaceParticleRead::GetFunctions(TArray<FNiagaraFunctionSign
 		Sig.bMemberFunction = true;
 		Sig.bRequiresContext = false;
 		Sig.bExperimental = true;
+#if WITH_EDITORONLY_DATA
+		Sig.Description = NSLOCTEXT("Niagara", "NiagaraDataInterfaceParticleRead_GetIntByIDDesc", "Returns an Int value from a particle by NiagaraID.  Note: When reading from self this will be the previous frames data.");
+#endif
 		OutFunctions.Add(Sig);
 	}
 
@@ -696,15 +708,15 @@ void UNiagaraDataInterfaceParticleRead::GetFunctions(TArray<FNiagaraFunctionSign
 		Sig.Name = GetBoolByIDFunctionName;
 		Sig.Inputs.Add(FNiagaraVariable(FNiagaraTypeDefinition(GetClass()), TEXT("Particle Reader")));
 		Sig.Inputs.Add(FNiagaraVariable(FNiagaraTypeDefinition::GetIDDef(), TEXT("Particle ID")));
-
 		Sig.Outputs.Add(FNiagaraVariable(FNiagaraTypeDefinition::GetBoolDef(), TEXT("Valid")));
 		Sig.Outputs.Add(FNiagaraVariable(FNiagaraTypeDefinition::GetBoolDef(), TEXT("Value")));
-
 		Sig.FunctionSpecifiers.Add(FName("Attribute"));
-
 		Sig.bMemberFunction = true;
 		Sig.bRequiresContext = false;
 		Sig.bExperimental = true;
+#if WITH_EDITORONLY_DATA
+		Sig.Description = NSLOCTEXT("Niagara", "NiagaraDataInterfaceParticleRead_GetBoolByIDDesc", "Returns a Bool value from a particle by NiagaraID.  Note: When reading from self this will be the previous frames data.");
+#endif
 		OutFunctions.Add(Sig);
 	}
 
@@ -713,15 +725,15 @@ void UNiagaraDataInterfaceParticleRead::GetFunctions(TArray<FNiagaraFunctionSign
 		Sig.Name = GetFloatByIDFunctionName;
 		Sig.Inputs.Add(FNiagaraVariable(FNiagaraTypeDefinition(GetClass()), TEXT("Particle Reader")));
 		Sig.Inputs.Add(FNiagaraVariable(FNiagaraTypeDefinition::GetIDDef(), TEXT("Particle ID")));
-
 		Sig.Outputs.Add(FNiagaraVariable(FNiagaraTypeDefinition::GetBoolDef(), TEXT("Valid")));
 		Sig.Outputs.Add(FNiagaraVariable(FNiagaraTypeDefinition::GetFloatDef(), TEXT("Value")));
-
 		Sig.FunctionSpecifiers.Add(FName("Attribute"));
-
 		Sig.bMemberFunction = true;
 		Sig.bRequiresContext = false;
 		Sig.bExperimental = true;
+#if WITH_EDITORONLY_DATA
+		Sig.Description = NSLOCTEXT("Niagara", "NiagaraDataInterfaceParticleRead_GetFloatByIDDesc", "Returns a Float value from a particle by NiagaraID.  Note: When reading from self this will be the previous frames data.");
+#endif
 		OutFunctions.Add(Sig);
 	}
 
@@ -730,15 +742,15 @@ void UNiagaraDataInterfaceParticleRead::GetFunctions(TArray<FNiagaraFunctionSign
 		Sig.Name = GetVec2ByIDFunctionName;
 		Sig.Inputs.Add(FNiagaraVariable(FNiagaraTypeDefinition(GetClass()), TEXT("Particle Reader")));
 		Sig.Inputs.Add(FNiagaraVariable(FNiagaraTypeDefinition::GetIDDef(), TEXT("Particle ID")));
-
 		Sig.Outputs.Add(FNiagaraVariable(FNiagaraTypeDefinition::GetBoolDef(), TEXT("Valid")));
 		Sig.Outputs.Add(FNiagaraVariable(FNiagaraTypeDefinition::GetVec2Def(), TEXT("Value")));
-
 		Sig.FunctionSpecifiers.Add(FName("Attribute"));
-
 		Sig.bMemberFunction = true;
 		Sig.bRequiresContext = false;
 		Sig.bExperimental = true;
+#if WITH_EDITORONLY_DATA
+		Sig.Description = NSLOCTEXT("Niagara", "NiagaraDataInterfaceParticleRead_GetVector2ByIDDesc", "Returns a Vector2 value from a particle by NiagaraID.  Note: When reading from self this will be the previous frames data.");
+#endif
 		OutFunctions.Add(Sig);
 	}
 
@@ -747,15 +759,15 @@ void UNiagaraDataInterfaceParticleRead::GetFunctions(TArray<FNiagaraFunctionSign
 		Sig.Name = GetVec3ByIDFunctionName;
 		Sig.Inputs.Add(FNiagaraVariable(FNiagaraTypeDefinition(GetClass()), TEXT("Particle Reader")));
 		Sig.Inputs.Add(FNiagaraVariable(FNiagaraTypeDefinition::GetIDDef(), TEXT("Particle ID")));
-
 		Sig.Outputs.Add(FNiagaraVariable(FNiagaraTypeDefinition::GetBoolDef(), TEXT("Valid")));
 		Sig.Outputs.Add(FNiagaraVariable(FNiagaraTypeDefinition::GetVec3Def(), TEXT("Value")));
-
 		Sig.FunctionSpecifiers.Add(FName("Attribute"));
-
 		Sig.bMemberFunction = true;
 		Sig.bRequiresContext = false;
 		Sig.bExperimental = true;
+#if WITH_EDITORONLY_DATA
+		Sig.Description = NSLOCTEXT("Niagara", "NiagaraDataInterfaceParticleRead_GetVector3ByIDDesc", "Returns a Vector3 value from a particle by NiagaraID.  Note: When reading from self this will be the previous frames data.");
+#endif
 		OutFunctions.Add(Sig);
 	}
 
@@ -764,15 +776,15 @@ void UNiagaraDataInterfaceParticleRead::GetFunctions(TArray<FNiagaraFunctionSign
 		Sig.Name = GetVec4ByIDFunctionName;
 		Sig.Inputs.Add(FNiagaraVariable(FNiagaraTypeDefinition(GetClass()), TEXT("Particle Reader")));
 		Sig.Inputs.Add(FNiagaraVariable(FNiagaraTypeDefinition::GetIDDef(), TEXT("Particle ID")));
-
 		Sig.Outputs.Add(FNiagaraVariable(FNiagaraTypeDefinition::GetBoolDef(), TEXT("Valid")));
 		Sig.Outputs.Add(FNiagaraVariable(FNiagaraTypeDefinition::GetVec4Def(), TEXT("Value")));
-
 		Sig.FunctionSpecifiers.Add(FName("Attribute"));
-
 		Sig.bMemberFunction = true;
 		Sig.bRequiresContext = false;
 		Sig.bExperimental = true;
+#if WITH_EDITORONLY_DATA
+		Sig.Description = NSLOCTEXT("Niagara", "NiagaraDataInterfaceParticleRead_GetVector4ByIDDesc", "Returns a Vector4 value from a particle by NiagaraID.  Note: When reading from self this will be the previous frames data.");
+#endif
 		OutFunctions.Add(Sig);
 	}
 
@@ -781,15 +793,15 @@ void UNiagaraDataInterfaceParticleRead::GetFunctions(TArray<FNiagaraFunctionSign
 		Sig.Name = GetColorByIDFunctionName;
 		Sig.Inputs.Add(FNiagaraVariable(FNiagaraTypeDefinition(GetClass()), TEXT("Particle Reader")));
 		Sig.Inputs.Add(FNiagaraVariable(FNiagaraTypeDefinition::GetIDDef(), TEXT("Particle ID")));
-
 		Sig.Outputs.Add(FNiagaraVariable(FNiagaraTypeDefinition::GetBoolDef(), TEXT("Valid")));
 		Sig.Outputs.Add(FNiagaraVariable(FNiagaraTypeDefinition::GetColorDef(), TEXT("Value")));
-
 		Sig.FunctionSpecifiers.Add(FName("Attribute"));
-
 		Sig.bMemberFunction = true;
 		Sig.bRequiresContext = false;
 		Sig.bExperimental = true;
+#if WITH_EDITORONLY_DATA
+		Sig.Description = NSLOCTEXT("Niagara", "NiagaraDataInterfaceParticleRead_GetColorByIDDesc", "Returns a Color value from a particle by NiagaraID.  Note: When reading from self this will be the previous frames data.");
+#endif
 		OutFunctions.Add(Sig);
 	}
 
@@ -798,15 +810,15 @@ void UNiagaraDataInterfaceParticleRead::GetFunctions(TArray<FNiagaraFunctionSign
 		Sig.Name = GetQuatByIDFunctionName;
 		Sig.Inputs.Add(FNiagaraVariable(FNiagaraTypeDefinition(GetClass()), TEXT("Particle Reader")));
 		Sig.Inputs.Add(FNiagaraVariable(FNiagaraTypeDefinition::GetIDDef(), TEXT("Particle ID")));
-
 		Sig.Outputs.Add(FNiagaraVariable(FNiagaraTypeDefinition::GetBoolDef(), TEXT("Valid")));
 		Sig.Outputs.Add(FNiagaraVariable(FNiagaraTypeDefinition::GetQuatDef(), TEXT("Value")));
-
 		Sig.FunctionSpecifiers.Add(FName("Attribute"));
-
 		Sig.bMemberFunction = true;
 		Sig.bRequiresContext = false;
 		Sig.bExperimental = true;
+#if WITH_EDITORONLY_DATA
+		Sig.Description = NSLOCTEXT("Niagara", "NiagaraDataInterfaceParticleRead_GetQuatByIDDesc", "Returns a Quat value from a particle by NiagaraID.  Note: When reading from self this will be the previous frames data.");
+#endif
 		OutFunctions.Add(Sig);
 	}
 
@@ -815,15 +827,15 @@ void UNiagaraDataInterfaceParticleRead::GetFunctions(TArray<FNiagaraFunctionSign
 		Sig.Name = GetIDByIDFunctionName;
 		Sig.Inputs.Add(FNiagaraVariable(FNiagaraTypeDefinition(GetClass()), TEXT("Particle Reader")));
 		Sig.Inputs.Add(FNiagaraVariable(FNiagaraTypeDefinition::GetIDDef(), TEXT("Particle ID")));
-
 		Sig.Outputs.Add(FNiagaraVariable(FNiagaraTypeDefinition::GetBoolDef(), TEXT("Valid")));
 		Sig.Outputs.Add(FNiagaraVariable(FNiagaraTypeDefinition::GetIDDef(), TEXT("Value")));
-
 		Sig.FunctionSpecifiers.Add(FName("Attribute"));
-
 		Sig.bMemberFunction = true;
 		Sig.bRequiresContext = false;
 		Sig.bExperimental = true;
+#if WITH_EDITORONLY_DATA
+		Sig.Description = NSLOCTEXT("Niagara", "NiagaraDataInterfaceParticleRead_GetIDByIDDesc", "Returns a NiagaraID value from a particle by NiagaraID.  Note: When reading from self this will be the previous frames data.");
+#endif
 		OutFunctions.Add(Sig);
 	}
 
@@ -835,15 +847,15 @@ void UNiagaraDataInterfaceParticleRead::GetFunctions(TArray<FNiagaraFunctionSign
 		Sig.Name = GetIntByIndexFunctionName;
 		Sig.Inputs.Add(FNiagaraVariable(FNiagaraTypeDefinition(GetClass()), TEXT("Particle Reader")));
 		Sig.Inputs.Add(FNiagaraVariable(FNiagaraTypeDefinition::GetIntDef(), TEXT("Particle Index")));
-
 		Sig.Outputs.Add(FNiagaraVariable(FNiagaraTypeDefinition::GetBoolDef(), TEXT("Valid")));
 		Sig.Outputs.Add(FNiagaraVariable(FNiagaraTypeDefinition::GetIntDef(), TEXT("Value")));
-
 		Sig.FunctionSpecifiers.Add(FName("Attribute"));
-
 		Sig.bMemberFunction = true;
 		Sig.bRequiresContext = false;
 		Sig.bExperimental = true;
+#if WITH_EDITORONLY_DATA
+		Sig.Description = NSLOCTEXT("Niagara", "NiagaraDataInterfaceParticleRead_GetIntByIndexDesc", "Returns an Int value from a particle by index.  Note: When reading from self this will be the previous frames data.");
+#endif
 		OutFunctions.Add(Sig);
 	}
 
@@ -852,15 +864,15 @@ void UNiagaraDataInterfaceParticleRead::GetFunctions(TArray<FNiagaraFunctionSign
 		Sig.Name = GetBoolByIndexFunctionName;
 		Sig.Inputs.Add(FNiagaraVariable(FNiagaraTypeDefinition(GetClass()), TEXT("Particle Reader")));
 		Sig.Inputs.Add(FNiagaraVariable(FNiagaraTypeDefinition::GetIntDef(), TEXT("Particle Index")));
-
 		Sig.Outputs.Add(FNiagaraVariable(FNiagaraTypeDefinition::GetBoolDef(), TEXT("Valid")));
 		Sig.Outputs.Add(FNiagaraVariable(FNiagaraTypeDefinition::GetBoolDef(), TEXT("Value")));
-
 		Sig.FunctionSpecifiers.Add(FName("Attribute"));
-
 		Sig.bMemberFunction = true;
 		Sig.bRequiresContext = false;
 		Sig.bExperimental = true;
+#if WITH_EDITORONLY_DATA
+		Sig.Description = NSLOCTEXT("Niagara", "NiagaraDataInterfaceParticleRead_GetBoolByIndexDesc", "Returns a Bool value from a particle by index.  Note: When reading from self this will be the previous frames data.");
+#endif
 		OutFunctions.Add(Sig);
 	}
 
@@ -869,15 +881,15 @@ void UNiagaraDataInterfaceParticleRead::GetFunctions(TArray<FNiagaraFunctionSign
 		Sig.Name = GetFloatByIndexFunctionName;
 		Sig.Inputs.Add(FNiagaraVariable(FNiagaraTypeDefinition(GetClass()), TEXT("Particle Reader")));
 		Sig.Inputs.Add(FNiagaraVariable(FNiagaraTypeDefinition::GetIntDef(), TEXT("Particle Index")));
-
 		Sig.Outputs.Add(FNiagaraVariable(FNiagaraTypeDefinition::GetBoolDef(), TEXT("Valid")));
 		Sig.Outputs.Add(FNiagaraVariable(FNiagaraTypeDefinition::GetFloatDef(), TEXT("Value")));
-
 		Sig.FunctionSpecifiers.Add(FName("Attribute"));
-
 		Sig.bMemberFunction = true;
 		Sig.bRequiresContext = false;
 		Sig.bExperimental = true;
+#if WITH_EDITORONLY_DATA
+		Sig.Description = NSLOCTEXT("Niagara", "NiagaraDataInterfaceParticleRead_GetFloatByIndexDesc", "Returns a Float value from a particle by index.  Note: When reading from self this will be the previous frames data.");
+#endif
 		OutFunctions.Add(Sig);
 	}
 
@@ -886,15 +898,15 @@ void UNiagaraDataInterfaceParticleRead::GetFunctions(TArray<FNiagaraFunctionSign
 		Sig.Name = GetVec2ByIndexFunctionName;
 		Sig.Inputs.Add(FNiagaraVariable(FNiagaraTypeDefinition(GetClass()), TEXT("Particle Reader")));
 		Sig.Inputs.Add(FNiagaraVariable(FNiagaraTypeDefinition::GetIntDef(), TEXT("Particle Index")));
-
 		Sig.Outputs.Add(FNiagaraVariable(FNiagaraTypeDefinition::GetBoolDef(), TEXT("Valid")));
 		Sig.Outputs.Add(FNiagaraVariable(FNiagaraTypeDefinition::GetVec2Def(), TEXT("Value")));
-
 		Sig.FunctionSpecifiers.Add(FName("Attribute"));
-
 		Sig.bMemberFunction = true;
 		Sig.bRequiresContext = false;
 		Sig.bExperimental = true;
+#if WITH_EDITORONLY_DATA
+		Sig.Description = NSLOCTEXT("Niagara", "NiagaraDataInterfaceParticleRead_GetVector2ByIndexDesc", "Returns a Vector2 value from a particle by index.  Note: When reading from self this will be the previous frames data.");
+#endif
 		OutFunctions.Add(Sig);
 	}
 
@@ -903,15 +915,15 @@ void UNiagaraDataInterfaceParticleRead::GetFunctions(TArray<FNiagaraFunctionSign
 		Sig.Name = GetVec3ByIndexFunctionName;
 		Sig.Inputs.Add(FNiagaraVariable(FNiagaraTypeDefinition(GetClass()), TEXT("Particle Reader")));
 		Sig.Inputs.Add(FNiagaraVariable(FNiagaraTypeDefinition::GetIntDef(), TEXT("Particle Index")));
-
 		Sig.Outputs.Add(FNiagaraVariable(FNiagaraTypeDefinition::GetBoolDef(), TEXT("Valid")));
 		Sig.Outputs.Add(FNiagaraVariable(FNiagaraTypeDefinition::GetVec3Def(), TEXT("Value")));
-
 		Sig.FunctionSpecifiers.Add(FName("Attribute"));
-
 		Sig.bMemberFunction = true;
 		Sig.bRequiresContext = false;
 		Sig.bExperimental = true;
+#if WITH_EDITORONLY_DATA
+		Sig.Description = NSLOCTEXT("Niagara", "NiagaraDataInterfaceParticleRead_GetVector3ByIndexDesc", "Returns a Vector3 value from a particle by index.  Note: When reading from self this will be the previous frames data.");
+#endif
 		OutFunctions.Add(Sig);
 	}
 
@@ -920,15 +932,15 @@ void UNiagaraDataInterfaceParticleRead::GetFunctions(TArray<FNiagaraFunctionSign
 		Sig.Name = GetVec4ByIndexFunctionName;
 		Sig.Inputs.Add(FNiagaraVariable(FNiagaraTypeDefinition(GetClass()), TEXT("Particle Reader")));
 		Sig.Inputs.Add(FNiagaraVariable(FNiagaraTypeDefinition::GetIntDef(), TEXT("Particle Index")));
-
 		Sig.Outputs.Add(FNiagaraVariable(FNiagaraTypeDefinition::GetBoolDef(), TEXT("Valid")));
 		Sig.Outputs.Add(FNiagaraVariable(FNiagaraTypeDefinition::GetVec4Def(), TEXT("Value")));
-
 		Sig.FunctionSpecifiers.Add(FName("Attribute"));
-
 		Sig.bMemberFunction = true;
 		Sig.bRequiresContext = false;
 		Sig.bExperimental = true;
+#if WITH_EDITORONLY_DATA
+		Sig.Description = NSLOCTEXT("Niagara", "NiagaraDataInterfaceParticleRead_GetVector4ByIndexDesc", "Returns a Vector4 value from a particle by index.  Note: When reading from self this will be the previous frames data.");
+#endif
 		OutFunctions.Add(Sig);
 	}
 
@@ -937,15 +949,15 @@ void UNiagaraDataInterfaceParticleRead::GetFunctions(TArray<FNiagaraFunctionSign
 		Sig.Name = GetColorByIndexFunctionName;
 		Sig.Inputs.Add(FNiagaraVariable(FNiagaraTypeDefinition(GetClass()), TEXT("Particle Reader")));
 		Sig.Inputs.Add(FNiagaraVariable(FNiagaraTypeDefinition::GetIntDef(), TEXT("Particle Index")));
-
 		Sig.Outputs.Add(FNiagaraVariable(FNiagaraTypeDefinition::GetBoolDef(), TEXT("Valid")));
 		Sig.Outputs.Add(FNiagaraVariable(FNiagaraTypeDefinition::GetColorDef(), TEXT("Value")));
-
 		Sig.FunctionSpecifiers.Add(FName("Attribute"));
-
 		Sig.bMemberFunction = true;
 		Sig.bRequiresContext = false;
 		Sig.bExperimental = true;
+#if WITH_EDITORONLY_DATA
+		Sig.Description = NSLOCTEXT("Niagara", "NiagaraDataInterfaceParticleRead_GetColorByIndexDesc", "Returns a Color value from a particle by index.  Note: When reading from self this will be the previous frames data.");
+#endif
 		OutFunctions.Add(Sig);
 	}
 
@@ -954,15 +966,15 @@ void UNiagaraDataInterfaceParticleRead::GetFunctions(TArray<FNiagaraFunctionSign
 		Sig.Name = GetQuatByIndexFunctionName;
 		Sig.Inputs.Add(FNiagaraVariable(FNiagaraTypeDefinition(GetClass()), TEXT("Particle Reader")));
 		Sig.Inputs.Add(FNiagaraVariable(FNiagaraTypeDefinition::GetIntDef(), TEXT("Particle Index")));
-
 		Sig.Outputs.Add(FNiagaraVariable(FNiagaraTypeDefinition::GetBoolDef(), TEXT("Valid")));
 		Sig.Outputs.Add(FNiagaraVariable(FNiagaraTypeDefinition::GetQuatDef(), TEXT("Value")));
-
 		Sig.FunctionSpecifiers.Add(FName("Attribute"));
-
 		Sig.bMemberFunction = true;
 		Sig.bRequiresContext = false;
 		Sig.bExperimental = true;
+#if WITH_EDITORONLY_DATA
+		Sig.Description = NSLOCTEXT("Niagara", "NiagaraDataInterfaceParticleRead_GetQuatByIndexDesc", "Returns a Quat value from a particle by index.  Note: When reading from self this will be the previous frames data.");
+#endif
 		OutFunctions.Add(Sig);
 	} 
 
@@ -971,15 +983,15 @@ void UNiagaraDataInterfaceParticleRead::GetFunctions(TArray<FNiagaraFunctionSign
 		Sig.Name = GetIDByIndexFunctionName;
 		Sig.Inputs.Add(FNiagaraVariable(FNiagaraTypeDefinition(GetClass()), TEXT("Particle Reader")));
 		Sig.Inputs.Add(FNiagaraVariable(FNiagaraTypeDefinition::GetIntDef(), TEXT("Particle Index")));
-
 		Sig.Outputs.Add(FNiagaraVariable(FNiagaraTypeDefinition::GetBoolDef(), TEXT("Valid")));
 		Sig.Outputs.Add(FNiagaraVariable(FNiagaraTypeDefinition::GetIDDef(), TEXT("Value")));
-
 		Sig.FunctionSpecifiers.Add(FName("Attribute"));
-
 		Sig.bMemberFunction = true;
 		Sig.bRequiresContext = false;
 		Sig.bExperimental = true;
+#if WITH_EDITORONLY_DATA
+		Sig.Description = NSLOCTEXT("Niagara", "NiagaraDataInterfaceParticleRead_GetIDByIndexDesc", "Returns a NiagaraID value from a particle by index.  Note: When reading from self this will be the previous frames data.");
+#endif
 		OutFunctions.Add(Sig);
 	}
 }
