Fix local-space, velocity-facing meshes

diff --git a/Engine/Plugins/FX/Niagara/Shaders/Private/NiagaraMeshVertexFactory.ush b/Engine/Plugins/FX/Niagara/Shaders/Private/NiagaraMeshVertexFactory.ush
--- a/Engine/Plugins/FX/Niagara/Shaders/Private/NiagaraMeshVertexFactory.ush
+++ b/Engine/Plugins/FX/Niagara/Shaders/Private/NiagaraMeshVertexFactory.ush
@@ -39,6 +39,18 @@ float3 TransformVector(float3 InVector)
 	}
 }
 
+float3 TransformVectorInverse(float3 InVector)
+{
+	if (NiagaraMeshVF.bLocalSpace)
+	{
+		return RotateWorldToLocal(InVector);
+	}
+	else
+	{
+		return InVector;
+	}
+}
+
 // Only support previous transform for deferred feature level
 //#define USE_PREVIOUS_TRANSFORM (FEATURE_LEVEL >= FEATURE_LEVEL_SM4)
 #define USE_PREVIOUS_TRANSFORM 0
@@ -496,7 +508,14 @@ float3x3 GetFacingMatrix(in float3 ParticlePosition, in float4 ParticleVelocity)
 	{
 		case MESH_FACING_VELOCITY:
 		{
-			FacingDir = ParticleVelocity.w > 0.0f ? ParticleVelocity.xyz : WorldZ;
+			if (ParticleVelocity.w > 0.0f)
+			{
+				FacingDir = ParticleVelocity.xyz;				
+			}
+			else
+			{
+				FacingDir = TransformVectorInverse(WorldZ);
+			}
 			break;
 		}
 
@@ -504,24 +523,18 @@ float3x3 GetFacingMatrix(in float3 ParticlePosition, in float4 ParticleVelocity)
 		{
 			float3 WorldPosition = TransformPosition(ParticlePosition);
 			float3 CameraVector = ResolvedView.WorldCameraOrigin - WorldPosition;
-			FacingDir = normalize(CameraVector);
+			FacingDir = TransformVectorInverse(normalize(CameraVector));
 			break;
 		}
 
 		default:
 		case MESH_FACING_CAMERA_PLANE:
 		{
-			FacingDir = -ResolvedView.ViewForward;
+			FacingDir = TransformVectorInverse(-ResolvedView.ViewForward);
 			break;
 		}
 	}
 
-	// Check to rotate the facing direction to local-space
-	if (bLocalSpace)
-	{
-		FacingDir = RotateWorldToLocal(FacingDir);
-	}
-
 	float3 XAxis = { 1, 0, 0 };
 	float3 YAxis = { 0, 1, 0 };
 	float3 ZAxis = { 0, 0, 1 };
@@ -563,10 +576,7 @@ float3x3 GetFacingMatrix(in float3 ParticlePosition, in float4 ParticleVelocity)
 		}
 
 		// Check to rotate the reference direction to local-space
-		if (bLocalSpace)
-		{
-			RefVector = RotateWorldToLocal(RefVector);
-		}
+		RefVector = TransformVectorInverse(RefVector);
 
 		// Orthonormalize the look-at rotation and generate a matrix
 		XAxis = FacingDir;
@@ -747,7 +757,7 @@ FVertexFactoryIntermediates GetVertexFactoryIntermediates(FVertexFactoryInput In
 
 	float4x4 ParticleToWorld;
 	float4x4 WorldToParticle;
-
+	
 	CalculateTransforms(Input.InstanceId, Intermediates, ParticleToWorld, WorldToParticle);
 
 	//-TODO: If / when we store previous transform data use them here (providing they are in local space, if not handle below)
