Fixed buffer types in the sort shaders.

These were causing validation errors in D3D11 and may be the cause for the Mac autotest failures.

diff --git a/Engine/Plugins/FX/Niagara/Shaders/Private/NiagaraSortKeyGen.usf b/Engine/Plugins/FX/Niagara/Shaders/Private/NiagaraSortKeyGen.usf
--- a/Engine/Plugins/FX/Niagara/Shaders/Private/NiagaraSortKeyGen.usf
+++ b/Engine/Plugins/FX/Niagara/Shaders/Private/NiagaraSortKeyGen.usf
@@ -18,8 +18,8 @@
 
 #define NIAGARA_KEY_GEN_MAX_CULL_PLANES 6 // must match the code
 Buffer<float> NiagaraParticleDataFloat;
-Buffer<float> NiagaraParticleDataHalf;
-Buffer<uint> NiagaraParticleDataInt;
+Buffer<half> NiagaraParticleDataHalf;
+Buffer<int> NiagaraParticleDataInt;
 Buffer<uint> GPUParticleCountBuffer;
 
 uint FloatDataStride;
@@ -50,29 +50,24 @@ float2 CullDistanceRangeSquared;
 float4 LocalBoundingSphere;
 
 RWBuffer<uint> OutKeys;
-RWBuffer<uint> OutParticleIndices;
+RWBuffer<int> OutParticleIndices;
 RWBuffer<uint> OutCulledParticleCounts;
 
 #if ENABLE_CULLING && !USE_WAVE_INTRINSICS
 groupshared uint CullGroupCount;
 #endif
 
-float GetFloat(Buffer<float> Buff, uint DataStride, uint RegisterIdx, uint InstanceID)
-{
-	return Buff[RegisterIdx * DataStride + InstanceID];
-}
-
 float GetFloat(uint RegisterIdx, uint InstanceID)
 {
 	BRANCH
 	if (RegisterIdx & (1 << 31))
 	{
 		RegisterIdx &= ~(1 << 31);
-		return GetFloat(NiagaraParticleDataHalf, HalfDataStride, RegisterIdx, InstanceID);
+		return NiagaraParticleDataHalf[RegisterIdx * HalfDataStride + InstanceID];
 	}
 	else
 	{
-		return GetFloat(NiagaraParticleDataFloat, FloatDataStride, RegisterIdx, InstanceID);
+		return NiagaraParticleDataFloat[RegisterIdx * FloatDataStride + InstanceID];
 	}
 }
 
diff --git a/Engine/Plugins/FX/Niagara/Source/NiagaraCore/Private/NiagaraCustomVersion.cpp b/Engine/Plugins/FX/Niagara/Source/NiagaraCore/Private/NiagaraCustomVersion.cpp
--- a/Engine/Plugins/FX/Niagara/Source/NiagaraCore/Private/NiagaraCustomVersion.cpp
+++ b/Engine/Plugins/FX/Niagara/Source/NiagaraCore/Private/NiagaraCustomVersion.cpp
@@ -8,4 +8,4 @@ const FGuid FNiagaraCustomVersion::GUID(0xFCF57AFA, 0x50764283, 0xB9A9E658, 0xFF
 // Register the custom version with core
 FCustomVersionRegistration GRegisterNiagaraCustomVersion(FNiagaraCustomVersion::GUID, FNiagaraCustomVersion::LatestVersion, TEXT("NiagaraVer"));
 
-const FGuid FNiagaraCustomVersion::LatestScriptCompileVersion(0x05940817, 0x4BC44989, 0x8583FE3D, 0x7EFE35B9);
\ No newline at end of file
+const FGuid FNiagaraCustomVersion::LatestScriptCompileVersion(0x62B48DE8, 0xACA74510, 0x865C6828, 0xE205B7D8);
