Fix incorrect/flickering shadows with culled Niagara mesh particles
Also fixes the flickering in editor outline rendering for culled meshes

diff --git a/Engine/Plugins/FX/Niagara/Shaders/Private/NiagaraMeshVertexFactory.ush b/Engine/Plugins/FX/Niagara/Shaders/Private/NiagaraMeshVertexFactory.ush
--- a/Engine/Plugins/FX/Niagara/Shaders/Private/NiagaraMeshVertexFactory.ush
+++ b/Engine/Plugins/FX/Niagara/Shaders/Private/NiagaraMeshVertexFactory.ush
@@ -661,13 +661,10 @@ FVertexFactoryIntermediates GetVertexFactoryIntermediates(FVertexFactoryInput In
 {
 	FVertexFactoryIntermediates Intermediates = (FVertexFactoryIntermediates)0;
 
-#if !SHADOW_DEPTH_SHADER && !SPHERICAL_OPACITY_FOR_SHADOW_DEPTHS
-	// Don't use sort buffer in the shadow because they are only valid after opaque.
 	if (SortedIndicesOffset != 0xFFFFFFFF)
 	{
 		Input.InstanceId = SortedIndices[SortedIndicesOffset + Input.InstanceId];
 	}
-#endif
 
 	if (NiagaraMeshVF.ColorDataOffset == -1)
 	{
diff --git a/Engine/Plugins/FX/Niagara/Source/Niagara/Private/NiagaraRendererMeshes.cpp b/Engine/Plugins/FX/Niagara/Source/Niagara/Private/NiagaraRendererMeshes.cpp
--- a/Engine/Plugins/FX/Niagara/Source/Niagara/Private/NiagaraRendererMeshes.cpp
+++ b/Engine/Plugins/FX/Niagara/Source/Niagara/Private/NiagaraRendererMeshes.cpp
@@ -208,8 +208,8 @@ int32 FNiagaraRendererMeshes::GetMaxIndirectArgs() const
 	{
 		MaxSectionCount = FMath::Max(MaxSectionCount, IndexInfo.Num());
 	}
-
-	//REVERB_WIP: This needs to be multiplied by the number of active viewsv
+	
+	//TODO: This needs to be multiplied by the number of active viewsv
 	return MaxSectionCount;
 }
 
@@ -607,7 +607,6 @@ void FNiagaraRendererMeshes::GetDynamicMeshElements(const TArray<const FSceneVie
 					Mesh.CastRayTracedShadow = SceneProxy->CastsDynamicShadow();
 #endif
 					Mesh.DepthPriorityGroup = (ESceneDepthPriorityGroup)SceneProxy->GetDepthPriorityGroup(View);
-					Mesh.bUseSelectionOutline = !bDoGPUCulling; //TODO: selection outline can't use the culled instance index buffer, so it causes bad flickering
 
 					FMeshBatchElement& BatchElement = Mesh.Elements[0];
 					BatchElement.PrimitiveUniformBuffer = IsMotionBlurEnabled() ? SceneProxy->GetUniformBuffer() : SceneProxy->GetUniformBufferNoVelocity();
