- Fix crash when Asset is nullptr

diff --git a/Engine/Plugins/FX/Niagara/Source/Niagara/Private/NiagaraDataInterfaceSkeletalMesh.cpp b/Engine/Plugins/FX/Niagara/Source/Niagara/Private/NiagaraDataInterfaceSkeletalMesh.cpp
--- a/Engine/Plugins/FX/Niagara/Source/Niagara/Private/NiagaraDataInterfaceSkeletalMesh.cpp
+++ b/Engine/Plugins/FX/Niagara/Source/Niagara/Private/NiagaraDataInterfaceSkeletalMesh.cpp
@@ -2164,6 +2164,11 @@ bool UNiagaraDataInterfaceSkeletalMesh::PerInstanceTick(void* PerInstanceData, F
 void UNiagaraDataInterfaceSkeletalMesh::GetFeedback(UNiagaraSystem* Asset, UNiagaraComponent* Component, TArray<FNiagaraDataInterfaceError>& OutErrors,
 	TArray<FNiagaraDataInterfaceFeedback>& OutWarnings, TArray<FNiagaraDataInterfaceFeedback>& OutInfo)
 {
+	if (Asset == nullptr)
+	{
+		return;
+	}
+
 	bool bHasCPUAccessWarning = false;
 	bool bHasNoMeshAssignedError = false;
 	
