- Fix crash seen in crash reporter where UNiagaraComponent::GetNumMaterials does not null check the CachedEmitter

#rnx

diff --git a/Engine/Plugins/FX/Niagara/Source/Niagara/Private/NiagaraComponent.cpp b/Engine/Plugins/FX/Niagara/Source/Niagara/Private/NiagaraComponent.cpp
--- a/Engine/Plugins/FX/Niagara/Source/Niagara/Private/NiagaraComponent.cpp
+++ b/Engine/Plugins/FX/Niagara/Source/Niagara/Private/NiagaraComponent.cpp
@@ -1395,10 +1395,12 @@ int32 UNiagaraComponent::GetNumMaterials() const
 		for (int32 i = 0; i < SystemInstance->GetEmitters().Num(); i++)
 		{
 			FNiagaraEmitterInstance* EmitterInst = &SystemInstance->GetEmitters()[i].Get();
-			UNiagaraEmitter* Emitter = EmitterInst->GetCachedEmitter();
-			for (UNiagaraRendererProperties* Properties : Emitter->GetEnabledRenderers())
+			if ( UNiagaraEmitter* Emitter = EmitterInst->GetCachedEmitter() )
 			{
-				Properties->GetUsedMaterials(EmitterInst, UsedMaterials);
+				for (UNiagaraRendererProperties* Properties : Emitter->GetEnabledRenderers())
+				{
+					Properties->GetUsedMaterials(EmitterInst, UsedMaterials);
+				}
 			}
 		}
 	}
