- Fix for requiring CPU access on bone sampling
- Fix for virtual bones when using skinning mode none
- Fix for scale applied to bones when using skinning mode none

#rnx

diff --git a/Engine/Plugins/FX/Niagara/Source/Niagara/Classes/NDISkeletalMeshCommon.h b/Engine/Plugins/FX/Niagara/Source/Niagara/Classes/NDISkeletalMeshCommon.h
--- a/Engine/Plugins/FX/Niagara/Source/Niagara/Classes/NDISkeletalMeshCommon.h
+++ b/Engine/Plugins/FX/Niagara/Source/Niagara/Classes/NDISkeletalMeshCommon.h
@@ -144,22 +144,36 @@ struct FSkinnedPositionAccessorHelper<TNDISkelMesh_SkinningModeInvalid>
 	// The bone accessor functions are valid if a mesh is present, so don't stub them entirely
 	FORCEINLINE_DEBUGGABLE FVector GetSkinnedBonePosition(FSkeletalMeshAccessorHelper& Accessor, int32 BoneIndex)
 	{
-		return Accessor.Mesh ? Accessor.Mesh->GetComposedRefPoseMatrix(BoneIndex).GetOrigin() : FVector::ZeroVector;
+		const int32 NumRealBones = Accessor.Mesh->RefSkeleton.GetRawBoneNum();
+		if (BoneIndex < NumRealBones)
+		{
+			return Accessor.Mesh->GetComposedRefPoseMatrix(BoneIndex).GetOrigin();
+		}
+
+		const FTransform& RefTransform = Accessor.Mesh->RefSkeleton.GetRefBonePose()[BoneIndex];
+		return RefTransform.GetLocation();
 	}
 
 	FORCEINLINE_DEBUGGABLE FVector GetSkinnedBonePreviousPosition(FSkeletalMeshAccessorHelper& Accessor, int32 BoneIndex)
 	{
-		return Accessor.Mesh ? Accessor.Mesh->GetComposedRefPoseMatrix(BoneIndex).GetOrigin() : FVector::ZeroVector;
+		return GetSkinnedBonePosition(Accessor, BoneIndex);
 	}
 
 	FORCEINLINE_DEBUGGABLE FQuat GetSkinnedBoneRotation(FSkeletalMeshAccessorHelper& Accessor, int32 BoneIndex)
 	{
-		return Accessor.Mesh ? Accessor.Mesh->GetComposedRefPoseMatrix(BoneIndex).ToQuat() : FQuat::Identity;
+		const int32 NumRealBones = Accessor.Mesh->RefSkeleton.GetRawBoneNum();
+		if (BoneIndex < NumRealBones)
+		{
+			return Accessor.Mesh->GetComposedRefPoseMatrix(BoneIndex).GetMatrixWithoutScale().ToQuat();
+		}
+
+		const FTransform& RefTransform = Accessor.Mesh->RefSkeleton.GetRefBonePose()[BoneIndex];
+		return RefTransform.GetRotation();
 	}
 
 	FORCEINLINE_DEBUGGABLE FQuat GetSkinnedBonePreviousRotation(FSkeletalMeshAccessorHelper& Accessor, int32 BoneIndex)
 	{
-		return Accessor.Mesh ? Accessor.Mesh->GetComposedRefPoseMatrix(BoneIndex).ToQuat() : FQuat::Identity;
+		return GetSkinnedBoneRotation(Accessor, BoneIndex);
 	}
 };
 
@@ -537,16 +551,12 @@ template<typename NextBinder>
 struct TSkinningModeBinder
 {
 	template<typename... ParamTypes>
-	static void Bind(UNiagaraDataInterface* Interface, const FVMExternalFunctionBindingInfo& BindingInfo, void* InstanceData, FVMExternalFunction &OutFunc)
+	static void BindIgnoreCPUAccess(UNiagaraDataInterface* Interface, const FVMExternalFunctionBindingInfo& BindingInfo, void* InstanceData, FVMExternalFunction& OutFunc)
 	{
 		FNDISkeletalMesh_InstanceData* InstData = (FNDISkeletalMesh_InstanceData*)InstanceData;
 		UNiagaraDataInterfaceSkeletalMesh* MeshInterface = CastChecked<UNiagaraDataInterfaceSkeletalMesh>(Interface);
 		USkeletalMeshComponent* Component = Cast<USkeletalMeshComponent>(InstData->Component.Get());
-		if (!InstData->bAllowCPUMeshDataAccess) // No-op when we can't access the mesh on CPU
-		{
-			NextBinder::template Bind<ParamTypes..., FSkinnedPositionAccessorHelper<TNDISkelMesh_SkinningModeInvalid>>(Interface, BindingInfo, InstanceData, OutFunc);
-		}
-		else if (MeshInterface->SkinningMode == ENDISkeletalMesh_SkinningMode::None || !Component) // Can't skin if we have no component.
+		if (MeshInterface->SkinningMode == ENDISkeletalMesh_SkinningMode::None || !Component) // Can't skin if we have no component.
 		{
 			NextBinder::template Bind<ParamTypes..., FSkinnedPositionAccessorHelper<TNDISkelMesh_SkinningModeNone>>(Interface, BindingInfo, InstanceData, OutFunc);
 		}
@@ -565,4 +575,18 @@ struct TSkinningModeBinder
 			checkf(false, TEXT("Invalid skinning mode in %s"), *Interface->GetPathName());
 		}
 	}
-};
\ No newline at end of file
+
+	template<typename... ParamTypes>
+	static void BindCheckCPUAccess(UNiagaraDataInterface* Interface, const FVMExternalFunctionBindingInfo& BindingInfo, void* InstanceData, FVMExternalFunction& OutFunc)
+	{
+		FNDISkeletalMesh_InstanceData* InstData = (FNDISkeletalMesh_InstanceData*)InstanceData;
+		if (!InstData->bAllowCPUMeshDataAccess) // No-op when we can't access the mesh on CPU
+		{
+			NextBinder::template Bind<ParamTypes..., FSkinnedPositionAccessorHelper<TNDISkelMesh_SkinningModeInvalid>>(Interface, BindingInfo, InstanceData, OutFunc);
+		}
+		else
+		{
+			BindIgnoreCPUAccess(Interface, BindingInfo, InstanceData, OutFunc);
+		}
+	}
+};
diff --git a/Engine/Plugins/FX/Niagara/Source/Niagara/Private/NDISkeletalMesh_BoneSampling.cpp b/Engine/Plugins/FX/Niagara/Source/Niagara/Private/NDISkeletalMesh_BoneSampling.cpp
--- a/Engine/Plugins/FX/Niagara/Source/Niagara/Private/NDISkeletalMesh_BoneSampling.cpp
+++ b/Engine/Plugins/FX/Niagara/Source/Niagara/Private/NDISkeletalMesh_BoneSampling.cpp
@@ -344,7 +344,7 @@ void UNiagaraDataInterfaceSkeletalMesh::BindSkeletonSamplingFunction(const FVMEx
 		ensure(BindingInfo.GetNumInputs() == 2 && BindingInfo.GetNumOutputs() == 10);
 		if (InstanceData->Mesh)
 		{
-			TSkinningModeBinder<TNDIExplicitBinder<FNDITransformHandlerNoop, TNDIExplicitBinder<TInterpOff, NDI_FUNC_BINDER(UNiagaraDataInterfaceSkeletalMesh, GetSkinnedBoneData)>>>::Bind(this, BindingInfo, InstanceData, OutFunc);
+			TSkinningModeBinder<TNDIExplicitBinder<FNDITransformHandlerNoop, TNDIExplicitBinder<TInterpOff, NDI_FUNC_BINDER(UNiagaraDataInterfaceSkeletalMesh, GetSkinnedBoneData)>>>::BindIgnoreCPUAccess(this, BindingInfo, InstanceData, OutFunc);
 		}
 		else
 		{
@@ -356,7 +356,7 @@ void UNiagaraDataInterfaceSkeletalMesh::BindSkeletonSamplingFunction(const FVMEx
 		ensure(BindingInfo.GetNumInputs() == 2 && BindingInfo.GetNumOutputs() == 10);
 		if (InstanceData->Mesh)
 		{
-			TSkinningModeBinder<TNDIExplicitBinder<FNDITransformHandler, TNDIExplicitBinder<TInterpOff, NDI_FUNC_BINDER(UNiagaraDataInterfaceSkeletalMesh, GetSkinnedBoneData)>>>::Bind(this, BindingInfo, InstanceData, OutFunc);
+			TSkinningModeBinder<TNDIExplicitBinder<FNDITransformHandler, TNDIExplicitBinder<TInterpOff, NDI_FUNC_BINDER(UNiagaraDataInterfaceSkeletalMesh, GetSkinnedBoneData)>>>::BindIgnoreCPUAccess(this, BindingInfo, InstanceData, OutFunc);
 		}
 		else
 		{
@@ -368,7 +368,7 @@ void UNiagaraDataInterfaceSkeletalMesh::BindSkeletonSamplingFunction(const FVMEx
 		ensure(BindingInfo.GetNumInputs() == 3 && BindingInfo.GetNumOutputs() == 10);
 		if (InstanceData->Mesh)
 		{
-			TSkinningModeBinder<TNDIExplicitBinder<FNDITransformHandlerNoop, TNDIExplicitBinder<TInterpOn, NDI_FUNC_BINDER(UNiagaraDataInterfaceSkeletalMesh, GetSkinnedBoneData)>>>::Bind(this, BindingInfo, InstanceData, OutFunc);
+			TSkinningModeBinder<TNDIExplicitBinder<FNDITransformHandlerNoop, TNDIExplicitBinder<TInterpOn, NDI_FUNC_BINDER(UNiagaraDataInterfaceSkeletalMesh, GetSkinnedBoneData)>>>::BindIgnoreCPUAccess(this, BindingInfo, InstanceData, OutFunc);
 		}
 		else
 		{
@@ -380,7 +380,7 @@ void UNiagaraDataInterfaceSkeletalMesh::BindSkeletonSamplingFunction(const FVMEx
 		ensure(BindingInfo.GetNumInputs() == 3 && BindingInfo.GetNumOutputs() == 10);
 		if (InstanceData->Mesh)
 		{
-			TSkinningModeBinder<TNDIExplicitBinder<FNDITransformHandler, TNDIExplicitBinder<TInterpOn, NDI_FUNC_BINDER(UNiagaraDataInterfaceSkeletalMesh, GetSkinnedBoneData)>>>::Bind(this, BindingInfo, InstanceData, OutFunc);
+			TSkinningModeBinder<TNDIExplicitBinder<FNDITransformHandler, TNDIExplicitBinder<TInterpOn, NDI_FUNC_BINDER(UNiagaraDataInterfaceSkeletalMesh, GetSkinnedBoneData)>>>::BindIgnoreCPUAccess(this, BindingInfo, InstanceData, OutFunc);
 		}
 		else
 		{
diff --git a/Engine/Plugins/FX/Niagara/Source/Niagara/Private/NDISkeletalMesh_TriangleSampling.cpp b/Engine/Plugins/FX/Niagara/Source/Niagara/Private/NDISkeletalMesh_TriangleSampling.cpp
--- a/Engine/Plugins/FX/Niagara/Source/Niagara/Private/NDISkeletalMesh_TriangleSampling.cpp
+++ b/Engine/Plugins/FX/Niagara/Source/Niagara/Private/NDISkeletalMesh_TriangleSampling.cpp
@@ -253,7 +253,7 @@ void UNiagaraDataInterfaceSkeletalMesh::BindTriangleSamplingFunction(const FVMEx
 		check(BindingInfo.GetNumInputs() == 5 && BindingInfo.GetNumOutputs() == 15);
 		if (InstanceData->bAllowCPUMeshDataAccess)
 		{
-			TSkinningModeBinder<TNDIExplicitBinder<FNDITransformHandlerNoop, TVertexAccessorBinder<TNDIExplicitBinder<TInterpOff, NDI_FUNC_BINDER(UNiagaraDataInterfaceSkeletalMesh, GetTriCoordSkinnedData)>>>>::Bind(this, BindingInfo, InstanceData, OutFunc);
+			TSkinningModeBinder<TNDIExplicitBinder<FNDITransformHandlerNoop, TVertexAccessorBinder<TNDIExplicitBinder<TInterpOff, NDI_FUNC_BINDER(UNiagaraDataInterfaceSkeletalMesh, GetTriCoordSkinnedData)>>>>::BindIgnoreCPUAccess(this, BindingInfo, InstanceData, OutFunc);
 		}
 		else
 		{
@@ -265,7 +265,7 @@ void UNiagaraDataInterfaceSkeletalMesh::BindTriangleSamplingFunction(const FVMEx
 		check(BindingInfo.GetNumInputs() == 5 && BindingInfo.GetNumOutputs() == 15);
 		if (InstanceData->bAllowCPUMeshDataAccess)
 		{
-			TSkinningModeBinder<TNDIExplicitBinder<FNDITransformHandler, TVertexAccessorBinder<TNDIExplicitBinder<TIntegralConstant<bool, false>, NDI_FUNC_BINDER(UNiagaraDataInterfaceSkeletalMesh, GetTriCoordSkinnedData)>>>>::Bind(this, BindingInfo, InstanceData, OutFunc);
+			TSkinningModeBinder<TNDIExplicitBinder<FNDITransformHandler, TVertexAccessorBinder<TNDIExplicitBinder<TIntegralConstant<bool, false>, NDI_FUNC_BINDER(UNiagaraDataInterfaceSkeletalMesh, GetTriCoordSkinnedData)>>>>::BindIgnoreCPUAccess(this, BindingInfo, InstanceData, OutFunc);
 		}
 		else
 		{
@@ -277,7 +277,7 @@ void UNiagaraDataInterfaceSkeletalMesh::BindTriangleSamplingFunction(const FVMEx
 		check(BindingInfo.GetNumInputs() == 6 && BindingInfo.GetNumOutputs() == 15);
 		if (InstanceData->bAllowCPUMeshDataAccess)
 		{
-			TSkinningModeBinder<TNDIExplicitBinder<FNDITransformHandlerNoop, TVertexAccessorBinder<TNDIExplicitBinder<TInterpOn, NDI_FUNC_BINDER(UNiagaraDataInterfaceSkeletalMesh, GetTriCoordSkinnedData)>>>>::Bind(this, BindingInfo, InstanceData, OutFunc);
+			TSkinningModeBinder<TNDIExplicitBinder<FNDITransformHandlerNoop, TVertexAccessorBinder<TNDIExplicitBinder<TInterpOn, NDI_FUNC_BINDER(UNiagaraDataInterfaceSkeletalMesh, GetTriCoordSkinnedData)>>>>::BindIgnoreCPUAccess(this, BindingInfo, InstanceData, OutFunc);
 		}
 		else
 		{
@@ -289,7 +289,7 @@ void UNiagaraDataInterfaceSkeletalMesh::BindTriangleSamplingFunction(const FVMEx
 		check(BindingInfo.GetNumInputs() == 6 && BindingInfo.GetNumOutputs() == 15);
 		if (InstanceData->bAllowCPUMeshDataAccess)
 		{
-			TSkinningModeBinder<TNDIExplicitBinder< FNDITransformHandler, TVertexAccessorBinder<TNDIExplicitBinder<TIntegralConstant<bool, true>, NDI_FUNC_BINDER(UNiagaraDataInterfaceSkeletalMesh, GetTriCoordSkinnedData)>>>>::Bind(this, BindingInfo, InstanceData, OutFunc);
+			TSkinningModeBinder<TNDIExplicitBinder< FNDITransformHandler, TVertexAccessorBinder<TNDIExplicitBinder<TIntegralConstant<bool, true>, NDI_FUNC_BINDER(UNiagaraDataInterfaceSkeletalMesh, GetTriCoordSkinnedData)>>>>::BindIgnoreCPUAccess(this, BindingInfo, InstanceData, OutFunc);
 		}
 		else
 		{
@@ -316,7 +316,7 @@ void UNiagaraDataInterfaceSkeletalMesh::BindTriangleSamplingFunction(const FVMEx
 	else if (BindingInfo.Name == FSkeletalMeshInterfaceHelper::GetTriCoordVerticesName)
 	{
 		check(BindingInfo.GetNumInputs() == 2 && BindingInfo.GetNumOutputs() == 3);
-		TSkinningModeBinder<NDI_FUNC_BINDER(UNiagaraDataInterfaceSkeletalMesh, GetTriCoordVertices)>::Bind(this, BindingInfo, InstanceData, OutFunc);		
+		TSkinningModeBinder<NDI_FUNC_BINDER(UNiagaraDataInterfaceSkeletalMesh, GetTriCoordVertices)>::BindCheckCPUAccess(this, BindingInfo, InstanceData, OutFunc);		
 	}
 	else if (BindingInfo.Name == FSkeletalMeshInterfaceHelper::RandomTriangleName)
 	{
diff --git a/Engine/Plugins/FX/Niagara/Source/Niagara/Private/NDISkeletalMesh_VertexSampling.cpp b/Engine/Plugins/FX/Niagara/Source/Niagara/Private/NDISkeletalMesh_VertexSampling.cpp
--- a/Engine/Plugins/FX/Niagara/Source/Niagara/Private/NDISkeletalMesh_VertexSampling.cpp
+++ b/Engine/Plugins/FX/Niagara/Source/Niagara/Private/NDISkeletalMesh_VertexSampling.cpp
@@ -177,12 +177,12 @@ void UNiagaraDataInterfaceSkeletalMesh::BindVertexSamplingFunction(const FVMExte
 	if (BindingInfo.Name == FSkeletalMeshInterfaceHelper::GetSkinnedVertexDataName)
 	{
 		check(BindingInfo.GetNumInputs() == 2 && BindingInfo.GetNumOutputs() == 15);
-		TSkinningModeBinder<TNDIExplicitBinder<FNDITransformHandlerNoop, TVertexAccessorBinder<NDI_FUNC_BINDER(UNiagaraDataInterfaceSkeletalMesh, GetVertexSkinnedData)>>>::Bind(this, BindingInfo, InstanceData, OutFunc);
+		TSkinningModeBinder<TNDIExplicitBinder<FNDITransformHandlerNoop, TVertexAccessorBinder<NDI_FUNC_BINDER(UNiagaraDataInterfaceSkeletalMesh, GetVertexSkinnedData)>>>::BindCheckCPUAccess(this, BindingInfo, InstanceData, OutFunc);
 	}
 	else if (BindingInfo.Name == FSkeletalMeshInterfaceHelper::GetSkinnedVertexDataWSName)
 	{
 		check(BindingInfo.GetNumInputs() == 2 && BindingInfo.GetNumOutputs() == 15);
-		TSkinningModeBinder<TNDIExplicitBinder<FNDITransformHandler, TVertexAccessorBinder<NDI_FUNC_BINDER(UNiagaraDataInterfaceSkeletalMesh, GetVertexSkinnedData)>>>::Bind(this, BindingInfo, InstanceData, OutFunc);
+		TSkinningModeBinder<TNDIExplicitBinder<FNDITransformHandler, TVertexAccessorBinder<NDI_FUNC_BINDER(UNiagaraDataInterfaceSkeletalMesh, GetVertexSkinnedData)>>>::BindCheckCPUAccess(this, BindingInfo, InstanceData, OutFunc);
 	}
 	else if (BindingInfo.Name == FSkeletalMeshInterfaceHelper::GetVertexColorName)
 	{
