- Clamp Skel Mesh sampling LODIndex to the first available inline LOD as we do not yet handle streaming LOD levels in Niagara
- This also fixes a crash with cooked out LOD levels and vertex / triangle sampling

#rnx

diff --git a/Engine/Plugins/FX/Niagara/Source/Niagara/Private/NiagaraDataInterfaceSkeletalMesh.cpp b/Engine/Plugins/FX/Niagara/Source/Niagara/Private/NiagaraDataInterfaceSkeletalMesh.cpp
--- a/Engine/Plugins/FX/Niagara/Source/Niagara/Private/NiagaraDataInterfaceSkeletalMesh.cpp
+++ b/Engine/Plugins/FX/Niagara/Source/Niagara/Private/NiagaraDataInterfaceSkeletalMesh.cpp
@@ -1506,6 +1506,12 @@ bool FNDISkeletalMesh_InstanceData::Init(UNiagaraDataInterfaceSkeletalMesh* Inte
 	}
 #endif
 
+	// Get the first LOD level we can guarantee to be in memory
+	//-TODO: Support skeletal mesh streaming
+	const FSkeletalMeshRenderData* SkelMeshRenderData = Mesh->GetResourceForRendering();
+	check(SkelMeshRenderData != nullptr);
+	const int32 FirstInlineLODLevel = SkelMeshRenderData->LODRenderData.Num() - SkelMeshRenderData->NumInlinedLODs;
+
 	//Setup where to spawn from
 	SamplingRegionIndices.Empty();
 	bool bAllRegionsAreAreaWeighting = true;
@@ -1525,7 +1531,7 @@ bool FNDISkeletalMesh_InstanceData::Init(UNiagaraDataInterfaceSkeletalMesh* Inte
 		}
 		else
 		{
-			LODIndex = FMath::Clamp(LODIndex, 0, Mesh->GetLODNum() - 1);
+			LODIndex = FMath::Clamp(LODIndex, FirstInlineLODLevel, Mesh->GetLODNum() - 1);
 		}
 	}
 	else
@@ -1554,6 +1560,17 @@ bool FNDISkeletalMesh_InstanceData::Init(UNiagaraDataInterfaceSkeletalMesh* Inte
 					LODIndex = RegionLODIndex;
 				}
 
+				// Ensure we aren't using a region from a streaming / culled LOD level
+				if (LODIndex < FirstInlineLODLevel)
+				{
+					UE_LOG(LogNiagara, Warning, TEXT("Skeletal Mesh Data Interface is trying to use regions on LODs levels that are either streamed or cooked out. This is currently unsupported.\nInterface: %s\nMesh: %s\nRegion: %s"),
+						*Interface->GetFullName(),
+						*Mesh->GetFullName(),
+						*RegionName.ToString());
+
+					return false;
+				}
+
 				//ensure we don't try to use two regions from different LODs.
 				if (LODIndex != RegionLODIndex)
 				{
