- Fix CIS issue

#rnx

diff --git a/Engine/Plugins/FX/Niagara/Source/Niagara/Private/NiagaraDataInterfaceSkeletalMesh.cpp b/Engine/Plugins/FX/Niagara/Source/Niagara/Private/NiagaraDataInterfaceSkeletalMesh.cpp
--- a/Engine/Plugins/FX/Niagara/Source/Niagara/Private/NiagaraDataInterfaceSkeletalMesh.cpp
+++ b/Engine/Plugins/FX/Niagara/Source/Niagara/Private/NiagaraDataInterfaceSkeletalMesh.cpp
@@ -1508,9 +1508,14 @@ bool FNDISkeletalMesh_InstanceData::Init(UNiagaraDataInterfaceSkeletalMesh* Inte
 
 	// Get the first LOD level we can guarantee to be in memory
 	//-TODO: Support skeletal mesh streaming
-	const FSkeletalMeshRenderData* SkelMeshRenderData = Mesh->GetResourceForRendering();
-	check(SkelMeshRenderData != nullptr);
-	const int32 FirstInlineLODLevel = SkelMeshRenderData->LODRenderData.Num() - SkelMeshRenderData->NumInlinedLODs;
+	int32 FirstInlineLODLevel = 0;
+	if (Mesh != nullptr)
+	{
+		if (const FSkeletalMeshRenderData* SkelMeshRenderData = Mesh->GetResourceForRendering())
+		{
+			FirstInlineLODLevel = SkelMeshRenderData->LODRenderData.Num() - SkelMeshRenderData->NumInlinedLODs;
+		}
+	}
 
 	//Setup where to spawn from
 	SamplingRegionIndices.Empty();
