- Fix crash for pre skinned verts if we do not have CPU Access enabled

#rnx

diff --git a/Engine/Plugins/FX/Niagara/Source/Niagara/Private/NiagaraDataInterfaceSkeletalMesh.cpp b/Engine/Plugins/FX/Niagara/Source/Niagara/Private/NiagaraDataInterfaceSkeletalMesh.cpp
--- a/Engine/Plugins/FX/Niagara/Source/Niagara/Private/NiagaraDataInterfaceSkeletalMesh.cpp
+++ b/Engine/Plugins/FX/Niagara/Source/Niagara/Private/NiagaraDataInterfaceSkeletalMesh.cpp
@@ -166,7 +166,7 @@ void FSkeletalMeshSkinningData::RegisterUser(FSkeletalMeshSkinningDataUsage Usag
 			PrevComponentTransforms() = CurrComponentTransforms();
 		}
 
-		if (Usage.NeedPreSkinnedVerts() && CurrSkinnedPositions(LODIndex).Num() == 0)
+		if (Usage.NeedPreSkinnedVerts() && CurrSkinnedPositions(LODIndex).Num() == 0 && SkelComp->SkeletalMesh->GetLODInfo(LODIndex)->bAllowCPUAccess)
 		{
 			FSkeletalMeshLODRenderData& SkelMeshLODData = SkelComp->SkeletalMesh->GetResourceForRendering()->LODRenderData[LODIndex];
 			FSkinWeightVertexBuffer* SkinWeightBuffer = SkelComp->GetSkinWeightBuffer(LODIndex);
@@ -318,7 +318,7 @@ bool FSkeletalMeshSkinningData::Tick(float InDeltaSeconds, bool bRequirePreskin)
 		for (int32 LODIndex = 0; LODIndex < LODData.Num(); ++LODIndex)
 		{
 			FLODData& LOD = LODData[LODIndex];
-			if (LOD.PreSkinnedVertsUsers > 0)
+			if (LOD.PreSkinnedVertsUsers > 0 && SkelComp->SkeletalMesh->GetLODInfo(LODIndex)->bAllowCPUAccess)
 			{
 				//TODO: If we pass the sections in the usage too, we can probably skin a minimal set of verts just for the used regions.
 				FSkeletalMeshLODRenderData& SkelMeshLODData = SkelComp->SkeletalMesh->GetResourceForRendering()->LODRenderData[LODIndex];
