Fixed a couple of missing transitions on the culled particles count UAV.

diff --git a/Engine/Plugins/FX/Niagara/Source/Niagara/Private/NiagaraEmitterInstanceBatcher.cpp b/Engine/Plugins/FX/Niagara/Source/Niagara/Private/NiagaraEmitterInstanceBatcher.cpp
--- a/Engine/Plugins/FX/Niagara/Source/Niagara/Private/NiagaraEmitterInstanceBatcher.cpp
+++ b/Engine/Plugins/FX/Niagara/Source/Niagara/Private/NiagaraEmitterInstanceBatcher.cpp
@@ -1248,6 +1248,10 @@ void NiagaraEmitterInstanceBatcher::GenerateSortKeys(FRHICommandListImmediate& R
 	check(EnumHasAnyFlags(Flags, EGPUSortFlags::KeyGenAfterPreRender));
 
 	FRWBuffer* CulledCountsBuffer = GPUInstanceCounterManager.AcquireCulledCountsBuffer(RHICmdList, FeatureLevel);
+	if (CulledCountsBuffer)
+	{
+		RHICmdList.TransitionResource(EResourceTransitionAccess::ERWNoBarrier, EResourceTransitionPipeline::EComputeToCompute, CulledCountsBuffer->UAV);
+	}
 
 	const FGPUSortManager::FKeyGenInfo KeyGenInfo((uint32)NumElementsInBatch, EnumHasAnyFlags(Flags, EGPUSortFlags::HighPrecisionKeys));
 
@@ -1270,7 +1274,7 @@ void NiagaraEmitterInstanceBatcher::GenerateSortKeys(FRHICommandListImmediate& R
 	Params.OutParticleIndices = ValuesUAV;
 	Params.OutCulledParticleCounts = CulledCountsBuffer ? (FRHIUnorderedAccessView*)CulledCountsBuffer->UAV : GetEmptyUAVFromPool(RHICmdList, PF_R32_UINT, false);
 
-	FRHIUnorderedAccessView* OutputUAVs[] = { KeysUAV, ValuesUAV };
+	FRHIUnorderedAccessView* OutputUAVs[] = { KeysUAV, ValuesUAV, Params.OutCulledParticleCounts };
 	for (const FNiagaraGPUSortInfo& SortInfo : SimulationsToSort)
 	{
 		if (SortInfo.AllocationInfo.SortBatchId == BatchId)
diff --git a/Engine/Plugins/FX/Niagara/Source/Niagara/Private/NiagaraGPUInstanceCountManager.cpp b/Engine/Plugins/FX/Niagara/Source/Niagara/Private/NiagaraGPUInstanceCountManager.cpp
--- a/Engine/Plugins/FX/Niagara/Source/Niagara/Private/NiagaraGPUInstanceCountManager.cpp
+++ b/Engine/Plugins/FX/Niagara/Source/Niagara/Private/NiagaraGPUInstanceCountManager.cpp
@@ -146,6 +146,7 @@ FRWBuffer* FNiagaraGPUInstanceCountManager::AcquireCulledCountsBuffer(FRHIComman
 			CulledCountBuffer.AcquireTransientResource();
 
 			// Initialize the buffer by clearing it to zero then transition it to be ready to write to
+			RHICmdList.TransitionResource(EResourceTransitionAccess::EWritable, EResourceTransitionPipeline::EComputeToCompute, CulledCountBuffer.UAV);
 			RHICmdList.ClearUAVUint(CulledCountBuffer.UAV, FUintVector4(EForceInit::ForceInitToZero));
 			RHICmdList.TransitionResource(EResourceTransitionAccess::ERWBarrier, EResourceTransitionPipeline::EComputeToCompute, CulledCountBuffer.UAV);
 		}
