* first draft of a Landscape DI by building a texture from heightfield data and putting it into a buffer.

* added FTextureReadBuffer2D for read only buffers that can be easily used in Niagara

* functionality in Landscape to build a TArray of height values to pass to DI

diff --git a/Engine/Plugins/FX/Niagara/Source/Niagara/Classes/NiagaraDataInterfaceLandscape.h b/Engine/Plugins/FX/Niagara/Source/Niagara/Classes/NiagaraDataInterfaceLandscape.h
new file mode 100644
--- /dev/null
+++ b/Engine/Plugins/FX/Niagara/Source/Niagara/Classes/NiagaraDataInterfaceLandscape.h
@@ -0,0 +1,91 @@
+// Copyright Epic Games, Inc. All Rights Reserved.
+
+#pragma once
+
+#include "NiagaraDataInterface.h"
+#include "NiagaraCommon.h"
+#include "VectorVM.h"
+
+#include "NiagaraDataInterfaceLandscape.generated.h"
+
+// Landscape data used on the render thread
+struct FNDILandscapeData_RenderThread
+{	
+	TUniquePtr<FTextureReadBuffer2D> LandscapeTextureBuffer = nullptr;
+	FIntPoint NumCells = FIntPoint(EForceInit::ForceInitToZero);
+	FMatrix WorldToActorTransform = FMatrix::Identity;
+	bool IsSet = false;
+};
+
+// Landscape data used for the game thread
+struct FNDILandscapeData_GameThread
+{
+	// #todo(dmp): CPU particles access via a TArray?
+	// TArray<float> Values;
+
+	FIntPoint NumCells = FIntPoint(EForceInit::ForceInitToZero);
+	FMatrix WorldToActorTransform = FMatrix::Identity;
+	bool IsSet = false;
+};
+
+/** Data Interface allowing sampling of a Landscape */
+UCLASS(EditInlineNew, Category = "Landscape", meta = (DisplayName = "Landscape Sample"))
+class NIAGARA_API UNiagaraDataInterfaceLandscape : public UNiagaraDataInterface
+{
+	GENERATED_UCLASS_BODY()
+public:
+
+	DECLARE_NIAGARA_DI_PARAMETER();
+
+	UPROPERTY(EditAnywhere, Category = "Landscape")
+	AActor* SourceLandscape;
+
+	//UObject Interface
+	virtual void PostInitProperties() override;	
+#if WITH_EDITOR
+	virtual void PostEditChangeProperty(struct FPropertyChangedEvent& PropertyChangedEvent) override;
+#endif	
+	//UObject Interface End
+
+	//UNiagaraDataInterface Interface
+	virtual void GetFunctions(TArray<FNiagaraFunctionSignature>& OutFunctions)override;
+	virtual void GetVMExternalFunction(const FVMExternalFunctionBindingInfo& BindingInfo, void* InstanceData, FVMExternalFunction &OutFunc) override;
+	virtual bool CanExecuteOnTarget(ENiagaraSimTarget Target)const override { return Target==ENiagaraSimTarget::GPUComputeSim; }
+	//UNiagaraDataInterface Interface End
+
+	void EmptyVMFunction(FVectorVMContext& Context) {}
+
+	virtual bool Equals(const UNiagaraDataInterface* Other) const override;
+
+	// GPU sim functionality
+	virtual void GetParameterDefinitionHLSL(const FNiagaraDataInterfaceGPUParamInfo& ParamInfo, FString& OutHLSL) override;
+	virtual bool GetFunctionHLSL(const FNiagaraDataInterfaceGPUParamInfo& ParamInfo, const FNiagaraDataInterfaceGeneratedFunction& FunctionInfo, int FunctionInstanceIndex, FString& OutHLSL) override;	
+
+	virtual void ProvidePerInstanceDataForRenderThread(void* DataForRenderThread, void* PerInstanceData, const FNiagaraSystemInstanceID& SystemInstance) override {}
+	virtual bool InitPerInstanceData(void* PerInstanceData, FNiagaraSystemInstance* SystemInstance) override;
+	virtual void DestroyPerInstanceData(void* PerInstanceData, FNiagaraSystemInstance* SystemInstance) override;
+	virtual bool PerInstanceTick(void* PerInstanceData, FNiagaraSystemInstance* SystemInstance, float DeltaSeconds) override;
+	virtual int32 PerInstanceDataSize() const override { return sizeof(FNDILandscapeData_GameThread); }	
+	
+	static const FString LandscapeTextureName;
+	static const FString SamplerName;
+	static const FString NumCellsBaseName;
+	static const FString WorldToActorBaseName;
+protected:
+	virtual bool CopyToInternal(UNiagaraDataInterface* Destination) const override;
+		
+	static const FName GetHeightName;	
+	static const FName GetNumCellsName;
+};
+
+struct FNiagaraDataInterfaceProxyLandscape : public FNiagaraDataInterfaceProxy
+{
+
+	FNiagaraDataInterfaceProxyLandscape() {}
+	
+	virtual void ConsumePerInstanceDataFromGameThread(void* PerInstanceData, const FNiagaraSystemInstanceID& Instance) override	{ check(false);	}
+
+	virtual int32 PerInstanceDataPassedToRenderThreadSize() const override { return 0; }
+
+	TMap<FNiagaraSystemInstanceID, FNDILandscapeData_RenderThread> SystemInstancesToProxyData_RT;
+};
\ No newline at end of file
diff --git a/Engine/Plugins/FX/Niagara/Source/Niagara/Niagara.Build.cs b/Engine/Plugins/FX/Niagara/Source/Niagara/Niagara.Build.cs
--- a/Engine/Plugins/FX/Niagara/Source/Niagara/Niagara.Build.cs
+++ b/Engine/Plugins/FX/Niagara/Source/Niagara/Niagara.Build.cs
@@ -17,10 +17,11 @@ public class Niagara : ModuleRules
                 "TimeManagement",
                 "Renderer",
                 "JsonUtilities",
-		"Json",
-		"AudioPlatformConfiguration",
-		"SignalProcessing",
-		"ApplicationCore"
+				"Landscape",
+				"Json",
+				"AudioPlatformConfiguration",
+				"SignalProcessing",
+				"ApplicationCore"
             }
         );
 
diff --git a/Engine/Plugins/FX/Niagara/Source/Niagara/Private/NiagaraDataInterfaceLandscape.cpp b/Engine/Plugins/FX/Niagara/Source/Niagara/Private/NiagaraDataInterfaceLandscape.cpp
new file mode 100644
--- /dev/null
+++ b/Engine/Plugins/FX/Niagara/Source/Niagara/Private/NiagaraDataInterfaceLandscape.cpp
@@ -0,0 +1,332 @@
+// Copyright Epic Games, Inc. All Rights Reserved.
+
+#include "NiagaraDataInterfaceLandscape.h"
+#include "NiagaraShader.h"
+#include "ShaderParameterUtils.h"
+#include "NiagaraCustomVersion.h"
+#include "NiagaraComponent.h"
+#include "LandscapeHeightfieldCollisionComponent.h"
+#include "Landscape.h"
+#include "LandscapeProxy.h"
+
+#define LOCTEXT_NAMESPACE "UNiagaraDataInterfaceLandscape"
+
+const FName UNiagaraDataInterfaceLandscape::GetHeightName(TEXT("GetHeight"));
+const FName UNiagaraDataInterfaceLandscape::GetNumCellsName(TEXT("GetNumCells"));
+
+const FString UNiagaraDataInterfaceLandscape::LandscapeTextureName(TEXT("LandscapeTexture_"));
+const FString UNiagaraDataInterfaceLandscape::SamplerName(TEXT("Sampler_"));
+const FString UNiagaraDataInterfaceLandscape::NumCellsBaseName(TEXT("NumCells_"));
+const FString UNiagaraDataInterfaceLandscape::WorldToActorBaseName(TEXT("WorldToActor_"));
+
+UNiagaraDataInterfaceLandscape::UNiagaraDataInterfaceLandscape(FObjectInitializer const& ObjectInitializer)
+	: Super(ObjectInitializer)		
+{
+	Proxy.Reset(new FNiagaraDataInterfaceProxyLandscape());
+}
+
+void UNiagaraDataInterfaceLandscape::PostInitProperties()
+{
+	Super::PostInitProperties();
+
+	if (HasAnyFlags(RF_ClassDefaultObject))
+	{
+		FNiagaraTypeRegistry::Register(FNiagaraTypeDefinition(GetClass()), true, false, false);
+	}	
+}
+
+#if WITH_EDITOR
+
+void UNiagaraDataInterfaceLandscape::PostEditChangeProperty(struct FPropertyChangedEvent& PropertyChangedEvent)
+{
+	Super::PostEditChangeProperty(PropertyChangedEvent);	
+}
+
+#endif
+
+bool UNiagaraDataInterfaceLandscape::CopyToInternal(UNiagaraDataInterface* Destination) const
+{
+	if (!Super::CopyToInternal(Destination))
+	{
+		return false;
+	}
+	UNiagaraDataInterfaceLandscape* DestinationLandscape = CastChecked<UNiagaraDataInterfaceLandscape>(Destination);
+	DestinationLandscape->SourceLandscape = SourceLandscape;	
+
+	return true;
+}
+
+bool UNiagaraDataInterfaceLandscape::Equals(const UNiagaraDataInterface* Other) const
+{
+	if (!Super::Equals(Other))
+	{
+		return false;
+	}
+	const UNiagaraDataInterfaceLandscape* OtherLandscape = CastChecked<const UNiagaraDataInterfaceLandscape>(Other);
+	return OtherLandscape->SourceLandscape == SourceLandscape;
+}
+
+void UNiagaraDataInterfaceLandscape::GetFunctions(TArray<FNiagaraFunctionSignature>& OutFunctions)
+{
+	{
+		FNiagaraFunctionSignature Sig;
+		Sig.Name = GetNumCellsName;
+		Sig.bExperimental = true;
+		Sig.bMemberFunction = true;
+		Sig.bRequiresContext = false;
+		Sig.Inputs.Add(FNiagaraVariable(FNiagaraTypeDefinition(GetClass()), TEXT("Landscape")));		
+		Sig.Outputs.Add(FNiagaraVariable(FNiagaraTypeDefinition::GetVec2Def(), TEXT("NumCells")));
+		//Sig.Owner = *GetName();
+
+		OutFunctions.Add(Sig);
+	}
+
+	{
+		FNiagaraFunctionSignature Sig;
+		Sig.Name = GetHeightName;
+		Sig.bExperimental = true;
+		Sig.bMemberFunction = true;
+		Sig.bRequiresContext = false;
+		Sig.Inputs.Add(FNiagaraVariable(FNiagaraTypeDefinition(GetClass()), TEXT("Landscape")));
+		Sig.Inputs.Add(FNiagaraVariable(FNiagaraTypeDefinition::GetVec3Def(), TEXT("WorldPos")));		
+		Sig.Outputs.Add(FNiagaraVariable(FNiagaraTypeDefinition::GetFloatDef(), TEXT("Value")));
+		//Sig.Owner = *GetName();
+
+		OutFunctions.Add(Sig);
+	}
+}
+
+// #todo(dmp): this is gpu only for now
+void UNiagaraDataInterfaceLandscape::GetVMExternalFunction(const FVMExternalFunctionBindingInfo& BindingInfo, void* InstanceData, FVMExternalFunction &OutFunc)
+{
+	if (BindingInfo.Name == GetNumCellsName)
+	{
+		OutFunc = FVMExternalFunction::CreateUObject(this, &UNiagaraDataInterfaceLandscape::EmptyVMFunction);
+	}
+	else if (BindingInfo.Name == GetHeightName)
+	{
+		OutFunc = FVMExternalFunction::CreateUObject(this, &UNiagaraDataInterfaceLandscape::EmptyVMFunction);
+	}
+}
+
+bool UNiagaraDataInterfaceLandscape::GetFunctionHLSL(const FNiagaraDataInterfaceGPUParamInfo& ParamInfo, const FNiagaraDataInterfaceGeneratedFunction& FunctionInfo, int FunctionInstanceIndex, FString& OutHLSL)
+{
+	if (FunctionInfo.DefinitionName == GetNumCellsName)
+	{
+		FString NumCellsVar = NumCellsBaseName + ParamInfo.DataInterfaceHLSLSymbol;
+		OutHLSL += TEXT("void ") + FunctionInfo.InstanceName + TEXT("(out int2 Out_Value) \n{\n");
+		OutHLSL += TEXT("\t Out_Value = ") + NumCellsVar + TEXT(";\n");
+		OutHLSL += TEXT("\n}\n");
+		return true;
+	}
+	else if (FunctionInfo.DefinitionName == GetHeightName)
+	{
+		static const TCHAR* FormatBounds = TEXT(R"(
+			void {FunctionName}(float3 In_WorldPos, out float Out_Val)
+			{
+				float3 ActorPos = mul(float4(In_WorldPos, 1.0), {WorldToActorTransform}).xyz;				
+				float2 UV = (ActorPos.xy + .5) / {NumCells};
+				Out_Val = {Grid}.SampleLevel({SamplerName}, UV, 0);
+			}
+		)");
+		TMap<FString, FStringFormatArg> ArgsBounds = {
+			{TEXT("FunctionName"), FunctionInfo.InstanceName},
+			{TEXT("Grid"), LandscapeTextureName + ParamInfo.DataInterfaceHLSLSymbol},
+			{TEXT("SamplerName"),    SamplerName + ParamInfo.DataInterfaceHLSLSymbol },			
+			{TEXT("NumCells"),    NumCellsBaseName + ParamInfo.DataInterfaceHLSLSymbol },
+			{TEXT("WorldToActorTransform"),    WorldToActorBaseName + ParamInfo.DataInterfaceHLSLSymbol },
+		};
+		OutHLSL += FString::Format(FormatBounds, ArgsBounds);
+		return true;
+	}
+	return false;
+}
+
+void UNiagaraDataInterfaceLandscape::GetParameterDefinitionHLSL(const FNiagaraDataInterfaceGPUParamInfo& ParamInfo, FString& OutHLSL)
+{
+	FString HLSLLandscapeTextureName = LandscapeTextureName + ParamInfo.DataInterfaceHLSLSymbol;
+	FString HLSLSamplerName = SamplerName + ParamInfo.DataInterfaceHLSLSymbol;
+	OutHLSL += TEXT("Texture2D ") + HLSLLandscapeTextureName + TEXT(";\n");
+	OutHLSL += TEXT("SamplerState ") + HLSLSamplerName + TEXT(";\n");
+	OutHLSL += TEXT("int2 ") + NumCellsBaseName + ParamInfo.DataInterfaceHLSLSymbol + TEXT(";\n");
+	OutHLSL += TEXT("float4x4 ") + WorldToActorBaseName + ParamInfo.DataInterfaceHLSLSymbol + TEXT(";\n");	
+}
+
+bool UNiagaraDataInterfaceLandscape::InitPerInstanceData(void* PerInstanceData, FNiagaraSystemInstance* SystemInstance)
+{
+	check(Proxy);
+
+	// Creates empty instance data since the landscape info isn't available until after the first tick
+
+	FNDILandscapeData_GameThread* InstanceData = new (PerInstanceData) FNDILandscapeData_GameThread();
+
+	FNiagaraDataInterfaceProxyLandscape* RT_Proxy = GetProxyAs<FNiagaraDataInterfaceProxyLandscape>();
+	ENQUEUE_RENDER_COMMAND(FNiagaraDIDestroyInstanceData) (
+		[RT_Proxy, InstanceID = SystemInstance->GetId(), Batcher = SystemInstance->GetBatcher()](FRHICommandListImmediate& CmdList)
+	{
+		check(!RT_Proxy->SystemInstancesToProxyData_RT.Contains(InstanceID));
+		FNDILandscapeData_RenderThread* TargetData = &RT_Proxy->SystemInstancesToProxyData_RT.Add(InstanceID);
+	});
+
+	return true;
+}
+
+void UNiagaraDataInterfaceLandscape::DestroyPerInstanceData(void* PerInstanceData, FNiagaraSystemInstance* SystemInstance)
+{
+	FNDILandscapeData_GameThread* InstanceData = static_cast<FNDILandscapeData_GameThread*>(PerInstanceData);
+
+	InstanceData->~FNDILandscapeData_GameThread();
+
+	FNiagaraDataInterfaceProxyLandscape* RT_Proxy = GetProxyAs<FNiagaraDataInterfaceProxyLandscape>();
+	ENQUEUE_RENDER_COMMAND(FNiagaraDIDestroyInstanceData) (
+		[RT_Proxy, InstanceID = SystemInstance->GetId(), Batcher = SystemInstance->GetBatcher()](FRHICommandListImmediate& CmdList)
+	{
+		//check(ThisProxy->SystemInstancesToProxyData.Contains(InstanceID));
+		RT_Proxy->SystemInstancesToProxyData_RT.Remove(InstanceID);
+	}
+	);
+}
+
+// Physics data for landscape that we use to fill the texture with is only available on tick, so we fill it once at runtime here
+bool UNiagaraDataInterfaceLandscape::PerInstanceTick(void* PerInstanceData, FNiagaraSystemInstance* SystemInstance, float DeltaSeconds)
+{
+	FNDILandscapeData_GameThread* InstanceData = static_cast<FNDILandscapeData_GameThread*>(PerInstanceData);
+
+	FNiagaraDataInterfaceProxyLandscape* RT_Proxy = GetProxyAs<FNiagaraDataInterfaceProxyLandscape>();
+
+	ALandscape* TheLandscape = static_cast<ALandscape*>(SourceLandscape);
+
+	if (InstanceData->IsSet || !TheLandscape || !TheLandscape->CollisionComponents[0]->HeightfieldRef.IsValid() || TheLandscape->CollisionComponents[0]->HeightfieldRef->Heightfield == nullptr)
+	{
+		return false;
+	}
+
+	int32 SizeX = 0;
+	int32 SizeY = 0;
+	TArray<float> HeightValues;
+	TheLandscape->GetHeightValues(SizeX, SizeY, HeightValues);
+
+	if (HeightValues.Num() == 0)
+	{
+		return  false;
+	}
+
+	// check for max resolution and fail if too high
+	// #todo(dmp): downsample to some resolution
+	int32 MaxDim = 16384;
+	if (SizeX > MaxDim || SizeY > MaxDim)
+	{
+		UE_LOG(LogNiagara, Error, TEXT("Landscape grid dimensions execeed maximum of 16384 %s"), *FNiagaraUtilities::SystemInstanceIDToString(SystemInstance->GetId()));
+		return false;
+	}
+
+	FMatrix WorldToLocal = TheLandscape->GetTransform().ToMatrixWithScale().Inverse();
+
+	InstanceData->IsSet = true;
+	InstanceData->NumCells = FIntPoint(SizeX, SizeY);
+	InstanceData->WorldToActorTransform = WorldToLocal;
+
+	ENQUEUE_RENDER_COMMAND(FPushDILandscapeToRT) (
+		[RT_Proxy, SizeX, SizeY, HeightValues, WorldToLocal, InstanceID = SystemInstance->GetId()](FRHICommandListImmediate& RHICmdList)
+		{
+			FNDILandscapeData_RenderThread* ProxyData = RT_Proxy->SystemInstancesToProxyData_RT.Find(InstanceID);
+
+			ProxyData->IsSet = HeightValues.Num() > 0;
+
+			// allocate the new texture based on the size of the landscape
+			ProxyData->LandscapeTextureBuffer = TUniquePtr<FTextureReadBuffer2D>(new FTextureReadBuffer2D());
+			ProxyData->LandscapeTextureBuffer->Initialize(4, SizeX, SizeY, EPixelFormat::PF_R32_FLOAT);
+
+			// fill texture will landscape data derived outside of here
+			uint32 DestStride = sizeof(float);
+			float* DestArray = (float*)RHILockTexture2D(ProxyData->LandscapeTextureBuffer->Buffer, 0, RLM_WriteOnly, DestStride, true);
+
+			memcpy(DestArray, HeightValues.GetData(), HeightValues.Num() * sizeof(float));
+
+			RHIUnlockTexture2D(ProxyData->LandscapeTextureBuffer->Buffer, 0, true);
+
+			ProxyData->NumCells.X = SizeX;
+			ProxyData->NumCells.Y = SizeY;
+			ProxyData->WorldToActorTransform = WorldToLocal;
+		}
+	);
+
+	return false;
+}
+
+struct FNiagaraDataInterfaceParametersCS_Landscape : public FNiagaraDataInterfaceParametersCS
+{
+	DECLARE_TYPE_LAYOUT(FNiagaraDataInterfaceParametersCS_Landscape, NonVirtual);
+public:
+	void Bind(const FNiagaraDataInterfaceGPUParamInfo& ParameterInfo, const class FShaderParameterMap& ParameterMap)
+	{
+		FString TexName = UNiagaraDataInterfaceLandscape::LandscapeTextureName + ParameterInfo.DataInterfaceHLSLSymbol;
+		FString SampleName = (UNiagaraDataInterfaceLandscape::SamplerName + ParameterInfo.DataInterfaceHLSLSymbol);
+		LandscapeTextureParam.Bind(ParameterMap, *TexName);
+		SamplerParam.Bind(ParameterMap, *SampleName);
+		
+		if (!LandscapeTextureParam.IsBound())
+		{
+			UE_LOG(LogNiagara, Warning, TEXT("Binding failed for FNiagaraDataInterfaceParametersCS_Landscape Landscape Texture %s. Was it optimized out?"), *TexName)
+		}
+
+		if (!SamplerParam.IsBound())
+		{
+			UE_LOG(LogNiagara, Warning, TEXT("Binding failed for FNiagaraDataInterfaceParametersCS_Landscape Sampler %s. Was it optimized out?"), *SampleName)
+		}
+
+		NumCells.Bind(ParameterMap, *(UNiagaraDataInterfaceLandscape::NumCellsBaseName + ParameterInfo.DataInterfaceHLSLSymbol));
+
+		WorldToActorTransform.Bind(ParameterMap, *(UNiagaraDataInterfaceLandscape::WorldToActorBaseName + ParameterInfo.DataInterfaceHLSLSymbol));
+	}
+
+	void Set(FRHICommandList& RHICmdList, const FNiagaraDataInterfaceSetArgs& Context) const
+	{
+		check(IsInRenderingThread());
+
+		FRHIComputeShader* ComputeShaderRHI = Context.Shader.GetComputeShader();
+
+		FNiagaraDataInterfaceProxyLandscape* RT_Proxy = static_cast<FNiagaraDataInterfaceProxyLandscape*>(Context.DataInterface);
+
+		if (RT_Proxy && RT_Proxy->SystemInstancesToProxyData_RT.Find(Context.SystemInstance))
+		{
+			FNDILandscapeData_RenderThread* ProxyData = RT_Proxy->SystemInstancesToProxyData_RT.Find(Context.SystemInstance);
+			
+			if (ProxyData->LandscapeTextureBuffer)
+			{
+				FRHISamplerState* SamplerState = TStaticSamplerState<SF_Bilinear, AM_Clamp, AM_Clamp, AM_Clamp>::GetRHI();
+				SetSamplerParameter(RHICmdList, ComputeShaderRHI, SamplerParam, SamplerState);
+
+				SetSRVParameter(RHICmdList, Context.Shader.GetComputeShader(), LandscapeTextureParam, ProxyData->LandscapeTextureBuffer->SRV);
+				SetShaderValue(RHICmdList, ComputeShaderRHI, NumCells, ProxyData->NumCells);
+				SetShaderValue(RHICmdList, ComputeShaderRHI, WorldToActorTransform, ProxyData->WorldToActorTransform);
+			}
+			else
+			{
+				SetSamplerParameter(RHICmdList, ComputeShaderRHI, SamplerParam, GBlackTextureWithSRV->SamplerStateRHI);
+				SetSRVParameter(RHICmdList, Context.Shader.GetComputeShader(), LandscapeTextureParam, GBlackTextureWithSRV->ShaderResourceViewRHI);
+				SetShaderValue(RHICmdList, ComputeShaderRHI, NumCells, FIntPoint(0, 0));
+				SetShaderValue(RHICmdList, ComputeShaderRHI, WorldToActorTransform, FMatrix::Identity);
+			}
+		}
+		else
+		{			
+			SetSamplerParameter(RHICmdList, ComputeShaderRHI, SamplerParam, GBlackTextureWithSRV->SamplerStateRHI);
+			SetSRVParameter(RHICmdList, Context.Shader.GetComputeShader(), LandscapeTextureParam, GBlackTextureWithSRV->ShaderResourceViewRHI);
+			SetShaderValue(RHICmdList, ComputeShaderRHI, NumCells, FIntPoint(0,0));
+			SetShaderValue(RHICmdList, ComputeShaderRHI, WorldToActorTransform, FMatrix::Identity);
+		}
+	}
+private:
+
+	LAYOUT_FIELD(FShaderResourceParameter, LandscapeTextureParam);
+	LAYOUT_FIELD(FShaderResourceParameter, SamplerParam);
+	LAYOUT_FIELD(FShaderParameter, NumCells);
+	LAYOUT_FIELD(FShaderParameter, WorldToActorTransform);
+};
+
+
+IMPLEMENT_TYPE_LAYOUT(FNiagaraDataInterfaceParametersCS_Landscape);
+
+IMPLEMENT_NIAGARA_DI_PARAMETER(UNiagaraDataInterfaceLandscape, FNiagaraDataInterfaceParametersCS_Landscape);
+#undef LOCTEXT_NAMESPACE
\ No newline at end of file
diff --git a/Engine/Source/Runtime/Landscape/Classes/LandscapeProxy.h b/Engine/Source/Runtime/Landscape/Classes/LandscapeProxy.h
--- a/Engine/Source/Runtime/Landscape/Classes/LandscapeProxy.h
+++ b/Engine/Source/Runtime/Landscape/Classes/LandscapeProxy.h
@@ -421,6 +421,9 @@ protected:
 public:
 	LANDSCAPE_API TOptional<float> GetHeightAtLocation(FVector Location) const;
 
+	/** Fills an array with height values **/
+	LANDSCAPE_API void GetHeightValues(int32& SizeX, int32& SizeY, TArray<float>& ArrayValue) const;
+
 	/** Offset in quads from global components grid origin (in quads) **/
 	UPROPERTY()
 	FIntPoint LandscapeSectionOffset;
diff --git a/Engine/Source/Runtime/Landscape/Private/LandscapeCollision.cpp b/Engine/Source/Runtime/Landscape/Private/LandscapeCollision.cpp
--- a/Engine/Source/Runtime/Landscape/Private/LandscapeCollision.cpp
+++ b/Engine/Source/Runtime/Landscape/Private/LandscapeCollision.cpp
@@ -2636,7 +2636,7 @@ TOptional<float> ULandscapeHeightfieldCollisionComponent::GetHeight(float X, flo
 		Height = HeightfieldRef->RBHeightfield->getHeight(HeightfieldRef->RBHeightfield->getNbRows() - 1 - X, Y) * ZScale;
 	}
 #elif WITH_CHAOS
-	if(IsValidRef(HeightfieldRef) && HeightfieldRef->Heightfield.Get())
+	if (IsValidRef(HeightfieldRef) && HeightfieldRef->Heightfield.Get())
 	{
 		Height = HeightfieldRef->Heightfield->GetHeightAt({ X, Y }) * ZScale;
 	}
@@ -2644,7 +2644,7 @@ TOptional<float> ULandscapeHeightfieldCollisionComponent::GetHeight(float X, flo
 	return Height;
 }
 
-LANDSCAPE_API TOptional<float> ALandscapeProxy::GetHeightAtLocation(FVector Location) const
+TOptional<float> ALandscapeProxy::GetHeightAtLocation(FVector Location) const
 {
 	TOptional<float> Height;
 	ULandscapeInfo* Info = GetLandscapeInfo();
@@ -2662,3 +2662,90 @@ LANDSCAPE_API TOptional<float> ALandscapeProxy::GetHeightAtLocation(FVector Loca
 	}
 	return Height;
 }
+
+void ALandscapeProxy::GetHeightValues(int32& SizeX, int32& SizeY, TArray<float> &ArrayValues) const
+{			
+	SizeX = 0;
+	SizeY = 0;
+	ArrayValues.SetNum(0);
+	
+	// Exit if we have no landscape data
+	if (LandscapeComponents.Num() == 0 || CollisionComponents.Num() == 0)
+	{
+		return;
+	}
+
+	// find index coordinate range for landscape
+	int32 MinX = MAX_int32;
+	int32 MinY = MAX_int32;
+	int32 MaxX = -MAX_int32;
+	int32 MaxY = -MAX_int32;
+
+	for (ULandscapeComponent* LandscapeComponent : LandscapeComponents)
+	{
+		// expecting a valid pointer to a landscape component
+		if (!LandscapeComponent)
+		{
+			return;
+		}
+
+		LandscapeComponent->GetComponentExtent(MinX, MinY, MaxX, MaxY);
+	}
+
+	if (MinX == MAX_int32)
+	{
+		return;
+	}			
+		
+	SizeX = (MaxX - MinX + 1);
+	SizeY = (MaxY - MinY + 1);
+	ArrayValues.SetNumUninitialized(SizeX * SizeY);
+	
+	for (ULandscapeHeightfieldCollisionComponent *CollisionComponent : CollisionComponents)
+	{
+		// Make sure we have a valid collision component and a heightfield
+		if (!CollisionComponent || !IsValidRef(CollisionComponent->HeightfieldRef))
+		{
+			SizeX = 0;
+			SizeY = 0;
+			ArrayValues.SetNum(0);
+			return;
+		}
+
+		TUniquePtr<Chaos::FHeightField> &HeightFieldData = CollisionComponent->HeightfieldRef->Heightfield;
+
+		// If we are expecting height data, but it isn't there, clear the return array, and exit
+		if (!HeightFieldData.IsValid())
+		{
+			SizeX = 0;
+			SizeY = 0;
+			ArrayValues.SetNum(0);
+			return;
+		}
+
+		const int32 BaseX = CollisionComponent->SectionBaseX;
+		const int32 BaseY = CollisionComponent->SectionBaseY;
+
+		const int32 NumX = HeightFieldData->GetNumCols();
+		const int32 NumY = HeightFieldData->GetNumRows();
+
+		const FTransform& ComponentToWorld = CollisionComponent->GetComponentToWorld();
+		const float ZScale = ComponentToWorld.GetScale3D().Z * LANDSCAPE_ZSCALE;
+
+		// Write all values to output array
+		for (int32 x = 0; x < NumX; ++x)
+		{
+			for (int32 y = 0; y < NumY; ++y)
+			{
+				const float CurrHeight = HeightFieldData->GetHeight(x, y) * ZScale;				
+				const float WorldHeight = ComponentToWorld.TransformPositionNoScale(FVector(0, 0, CurrHeight)).Z;
+				
+				// write output
+				const int32 WriteX = BaseX + x;
+				const int32 WriteY = BaseY + y;
+				const int32 Idx = WriteY * SizeX + WriteX;
+				ArrayValues[Idx] = WorldHeight;
+			}
+		}
+	}
+}
\ No newline at end of file
diff --git a/Engine/Source/Runtime/RHI/Public/RHIUtilities.h b/Engine/Source/Runtime/RHI/Public/RHIUtilities.h
--- a/Engine/Source/Runtime/RHI/Public/RHIUtilities.h
+++ b/Engine/Source/Runtime/RHI/Public/RHIUtilities.h
@@ -286,6 +286,70 @@ struct FRWBuffer
 	}
 };
 
+/** Encapsulates a GPU read only texture 2D with its SRV. */
+struct FTextureReadBuffer2D
+{
+	FTexture2DRHIRef Buffer;	
+	FShaderResourceViewRHIRef SRV;
+	uint32 NumBytes;
+
+	FTextureReadBuffer2D()
+		: NumBytes(0)
+	{}
+
+	~FTextureReadBuffer2D()
+	{
+		Release();
+	}
+
+	// @param AdditionalUsage passed down to RHICreateVertexBuffer(), get combined with "BUF_UnorderedAccess | BUF_ShaderResource" e.g. BUF_Static
+	const static uint32 DefaultTextureInitFlag = TexCreate_ShaderResource;
+	void Initialize(const uint32 BytesPerElement, const uint32 SizeX, const uint32 SizeY, const EPixelFormat Format, uint32 Flags = DefaultTextureInitFlag)
+	{
+		check(GMaxRHIFeatureLevel == ERHIFeatureLevel::SM5
+			|| IsVulkanPlatform(GMaxRHIShaderPlatform)
+			|| IsMetalPlatform(GMaxRHIShaderPlatform)
+			|| (GMaxRHIFeatureLevel == ERHIFeatureLevel::ES3_1 && GSupportsResourceView)
+		);
+
+		NumBytes = SizeX * SizeY * BytesPerElement;
+
+		FRHIResourceCreateInfo CreateInfo;
+		Buffer = RHICreateTexture2D(
+			SizeX, SizeY, Format, //PF_R32_FLOAT,
+			/*NumMips=*/ 1,
+			1,
+			Flags,
+			/*BulkData=*/ CreateInfo);
+
+		
+		SRV = RHICreateShaderResourceView(Buffer, 0);
+	}
+
+	void AcquireTransientResource()
+	{
+		RHIAcquireTransientResource(Buffer);
+	}
+	void DiscardTransientResource()
+	{
+		RHIDiscardTransientResource(Buffer);
+	}
+
+	void Release()
+	{
+		int32 BufferRefCount = Buffer ? Buffer->GetRefCount() : -1;
+
+		if (BufferRefCount == 1)
+		{
+			DiscardTransientResource();
+		}
+
+		NumBytes = 0;
+		Buffer.SafeRelease();		
+		SRV.SafeRelease();
+	}
+};
+
 /** Encapsulates a GPU read buffer with its SRV. */
 struct FReadBuffer
 {
