Fixed missing per ribbon data, when some ribbons get culled out because they don't have enough vertices.

diff --git a/Engine/Plugins/FX/Niagara/Source/Niagara/Private/NiagaraRendererRibbons.cpp b/Engine/Plugins/FX/Niagara/Source/Niagara/Private/NiagaraRendererRibbons.cpp
--- a/Engine/Plugins/FX/Niagara/Source/Niagara/Private/NiagaraRendererRibbons.cpp
+++ b/Engine/Plugins/FX/Niagara/Source/Niagara/Private/NiagaraRendererRibbons.cpp
@@ -702,6 +702,10 @@ FNiagaraDynamicDataBase* FNiagaraRendererRibbons::GenerateDynamicData(const FNia
 
 			DynamicData->PackPerRibbonData(U0Scale, U0Offset, U1Scale, U1Offset, NumSegments, StartIndex);
 		}
+		else
+		{
+			DynamicData->PackPerRibbonData(0, 0, 0, 0, 0, 0);
+		}
 	};
 
 	DynamicData->MultiRibbonInfos.Reset();
