Mark graph for recompile when changing a static switch default value

diff --git a/Engine/Plugins/FX/Niagara/Source/NiagaraEditor/Private/Customizations/NiagaraStaticSwitchNodeDetails.cpp b/Engine/Plugins/FX/Niagara/Source/NiagaraEditor/Private/Customizations/NiagaraStaticSwitchNodeDetails.cpp
--- a/Engine/Plugins/FX/Niagara/Source/NiagaraEditor/Private/Customizations/NiagaraStaticSwitchNodeDetails.cpp
+++ b/Engine/Plugins/FX/Niagara/Source/NiagaraEditor/Private/Customizations/NiagaraStaticSwitchNodeDetails.cpp
@@ -33,7 +33,6 @@ FNiagaraStaticSwitchNodeDetails::FNiagaraStaticSwitchNodeDetails()
 
 	for (FNiagaraTypeDefinition Type : FNiagaraTypeRegistry::GetRegisteredParameterTypes())
 	{
-		//TODO are these all relevant niagara enum types?
 		if (Type.IsEnum())
 		{
 			DropdownOptions.Add(MakeShareable(new SwitchDropdownOption(Type.GetEnum()->GetName(), Type.GetEnum())));
@@ -498,6 +497,7 @@ void FNiagaraStaticSwitchNodeDetails::SetSwitchParameterMetadata(const FNiagaraV
 		return;
 	}
 	Node->GetNiagaraGraph()->SetMetaData(FNiagaraVariable(Node->GetInputType(), Node->InputParameterName), MetaData);
+	Node->GetNiagaraGraph()->NotifyGraphNeedsRecompile();
 }
 
 void FNiagaraStaticSwitchNodeDetails::UpdateSelectionFromNode()
