Mark graph for recompile when renaming a static switch node

diff --git a/Engine/Plugins/FX/Niagara/Source/NiagaraEditor/Private/Customizations/NiagaraStaticSwitchNodeDetails.cpp b/Engine/Plugins/FX/Niagara/Source/NiagaraEditor/Private/Customizations/NiagaraStaticSwitchNodeDetails.cpp
--- a/Engine/Plugins/FX/Niagara/Source/NiagaraEditor/Private/Customizations/NiagaraStaticSwitchNodeDetails.cpp
+++ b/Engine/Plugins/FX/Niagara/Source/NiagaraEditor/Private/Customizations/NiagaraStaticSwitchNodeDetails.cpp
@@ -569,6 +569,7 @@ void FNiagaraStaticSwitchNodeDetails::OnParameterNameCommited(const FText& InTex
 	if (Node.IsValid())
 	{
 		Node->ChangeSwitchParameterName(FName(*InText.ToString()));
+		Node->GetNiagaraGraph()->NotifyGraphNeedsRecompile();
 	}
 }
 
