Adjusted some comments and the "canAutoBind" tooltip for more clarity

diff --git a/Engine/Plugins/FX/Niagara/Source/NiagaraEditor/Private/NiagaraNodeFunctionCall.cpp b/Engine/Plugins/FX/Niagara/Source/NiagaraEditor/Private/NiagaraNodeFunctionCall.cpp
--- a/Engine/Plugins/FX/Niagara/Source/NiagaraEditor/Private/NiagaraNodeFunctionCall.cpp
+++ b/Engine/Plugins/FX/Niagara/Source/NiagaraEditor/Private/NiagaraNodeFunctionCall.cpp
@@ -1076,7 +1076,9 @@ bool UNiagaraNodeFunctionCall::FindAutoBoundInput(UNiagaraNodeInput* InputNode,
 {
 	check(InputNode && InputNode->IsExposed());
 	if (PinToAutoBind->LinkedTo.Num() > 0 || !InputNode->CanAutoBind())
+	{
 		return false;
+	}
 
 	const UEdGraphSchema_Niagara* Schema = CastChecked<UEdGraphSchema_Niagara>(GetSchema());
 	FNiagaraVariable PinVar = Schema->PinToNiagaraVariable(PinToAutoBind);
@@ -1087,7 +1089,7 @@ bool UNiagaraNodeFunctionCall::FindAutoBoundInput(UNiagaraNodeInput* InputNode,
 	UNiagaraNodeOutput* CallerOutputNodeSpawn = CallerGraph->FindOutputNode(ENiagaraScriptUsage::ParticleSpawnScript);
 	UNiagaraNodeOutput* CallerOutputNodeUpdate = CallerGraph->FindOutputNode(ENiagaraScriptUsage::ParticleUpdateScript);
 
-	//First, lest see if we're an attribute of this emitter. Only valid if we're a module call off the primary script.
+	//First, let's see if we're an attribute of this emitter. Only valid if we're a module call off the primary script.
 	if (CallerOutputNodeSpawn || CallerOutputNodeUpdate)
 	{
 		UNiagaraNodeOutput* CallerOutputNode = CallerOutputNodeSpawn != nullptr ? CallerOutputNodeSpawn : CallerOutputNodeUpdate;
@@ -1113,25 +1115,6 @@ bool UNiagaraNodeFunctionCall::FindAutoBoundInput(UNiagaraNodeInput* InputNode,
 		return true;
 	}
 
-	//Not sure it's a good idea to allow binding to user made parameters.
-// 	if (AutoBindOptions.bBindToParameters)
-// 	{
-// 		//Finally, lets see if we're a parameter of this emitter.
-// 		TArray<UNiagaraNodeInput*> CallerInputNodes;
-// 		CallerGraph->FindInputNodes(CallerInputNodes);
-// 		UNiagaraNodeInput** MatchingParamPtr = CallerInputNodes.FindByPredicate([&](UNiagaraNodeInput*& CallerInputNode)
-// 		{
-// 			return CallerInputNode->Input.IsEquivalent(PinVar);
-// 		});
-// 
-// 		if (MatchingParamPtr)
-// 		{
-// 			OutFoundVar = (*MatchingParamPtr)->Input;
-// 			OutNodeUsage = ENiagaraInputNodeUsage::Parameter;
-// 			return true;
-// 		}
-// 	}
-
 	//Unable to auto bind.
 	return false;
 }
diff --git a/Engine/Plugins/FX/Niagara/Source/NiagaraEditor/Private/Widgets/SNiagaraGraphNodeInput.cpp b/Engine/Plugins/FX/Niagara/Source/NiagaraEditor/Private/Widgets/SNiagaraGraphNodeInput.cpp
--- a/Engine/Plugins/FX/Niagara/Source/NiagaraEditor/Private/Widgets/SNiagaraGraphNodeInput.cpp
+++ b/Engine/Plugins/FX/Niagara/Source/NiagaraEditor/Private/Widgets/SNiagaraGraphNodeInput.cpp
@@ -32,7 +32,7 @@ public:
 	{
 		UI_COMMAND(ToggleExposed, "Exposed", "Toggles whether or not this parameter is exposed.", EUserInterfaceActionType::ToggleButton, FInputChord());
 		UI_COMMAND(ToggleRequired, "Required", "Toggles whether or not this parameter is required.", EUserInterfaceActionType::ToggleButton, FInputChord());
-		UI_COMMAND(ToggleCanAutoBind, "CanAutoBind", "Toggles whether or not this parameter can be automatically bound.", EUserInterfaceActionType::ToggleButton, FInputChord());
+		UI_COMMAND(ToggleCanAutoBind, "CanAutoBind", "Toggles whether or not this parameter can auto-bind to system parameters and emitter attributes. Will never auto bind to custom parameters.", EUserInterfaceActionType::ToggleButton, FInputChord());
 		UI_COMMAND(ToggleHidden, "Hidden", "Toggles whether or not this parameter is hidden.", EUserInterfaceActionType::ToggleButton, FInputChord());
 	}
 
