Fixed crash ray tracing Niagara ribbons with zero polygons.

diff --git a/Engine/Plugins/FX/Niagara/Source/Niagara/Private/NiagaraRendererRibbons.cpp b/Engine/Plugins/FX/Niagara/Source/Niagara/Private/NiagaraRendererRibbons.cpp
--- a/Engine/Plugins/FX/Niagara/Source/Niagara/Private/NiagaraRendererRibbons.cpp
+++ b/Engine/Plugins/FX/Niagara/Source/Niagara/Private/NiagaraRendererRibbons.cpp
@@ -1130,42 +1130,45 @@ void FNiagaraRendererRibbons::GetDynamicRayTracingInstances(FRayTracingMaterialG
 		return;
 	}
 
-	FRayTracingInstance RayTracingInstance;
-	RayTracingInstance.Geometry = &RayTracingGeometry;
-	RayTracingInstance.InstanceTransforms.Add(FMatrix::Identity);
+	auto& View = Context.ReferenceView;
+	auto& ViewFamily = Context.ReferenceViewFamily;
+	// Setup material for our ray tracing instance
+	FNiagaraMeshCollectorResourcesRibbon& CollectorResources = Context.RayTracingMeshResourceCollector.AllocateOneFrameResource<FNiagaraMeshCollectorResourcesRibbon>();
+
+	FGlobalDynamicIndexBuffer::FAllocationEx DynamicIndexAllocation;
+	CreatePerViewResources(Context.ReferenceView, Context.ReferenceViewFamily, SceneProxy, Context.RayTracingMeshResourceCollector, CollectorResources.UniformBuffer, DynamicIndexAllocation);
 
+	if (DynamicIndexAllocation.MaxUsedIndex <= 0)
 	{
-		auto& View = Context.ReferenceView;
-		auto& ViewFamily = Context.ReferenceViewFamily;
-		// Setup material for our ray tracing instance
-		FNiagaraMeshCollectorResourcesRibbon& CollectorResources = Context.RayTracingMeshResourceCollector.AllocateOneFrameResource<FNiagaraMeshCollectorResourcesRibbon>();
+		return;
+	}
 
-		FGlobalDynamicIndexBuffer::FAllocationEx DynamicIndexAllocation;
-		CreatePerViewResources(Context.ReferenceView, Context.ReferenceViewFamily, SceneProxy, Context.RayTracingMeshResourceCollector, CollectorResources.UniformBuffer, DynamicIndexAllocation);
+	FRayTracingInstance RayTracingInstance;
+	RayTracingInstance.Geometry = &RayTracingGeometry;
+	RayTracingInstance.InstanceTransforms.Add(FMatrix::Identity);
 
-		RayTracingGeometry.Initializer.IndexBuffer = DynamicIndexAllocation.IndexBuffer->IndexBufferRHI;
-		RayTracingGeometry.Initializer.IndexBufferOffset = DynamicIndexAllocation.FirstIndex * DynamicIndexAllocation.IndexStride;
+	RayTracingGeometry.Initializer.IndexBuffer = DynamicIndexAllocation.IndexBuffer->IndexBufferRHI;
+	RayTracingGeometry.Initializer.IndexBufferOffset = DynamicIndexAllocation.FirstIndex * DynamicIndexAllocation.IndexStride;
 
-		FMeshBatch MeshBatch;
+	FMeshBatch MeshBatch;
 
-		SetupMeshBatchAndCollectorResourceForView(Context.ReferenceView, Context.ReferenceViewFamily, SceneProxy, Context.RayTracingMeshResourceCollector, DynamicDataRibbon, DynamicIndexAllocation, MeshBatch, CollectorResources);
+	SetupMeshBatchAndCollectorResourceForView(Context.ReferenceView, Context.ReferenceViewFamily, SceneProxy, Context.RayTracingMeshResourceCollector, DynamicDataRibbon, DynamicIndexAllocation, MeshBatch, CollectorResources);
 
-		RayTracingInstance.Materials.Add(MeshBatch);
+	RayTracingInstance.Materials.Add(MeshBatch);
 
-		const uint32 VertexCount = DynamicIndexAllocation.MaxUsedIndex;
-		Context.DynamicRayTracingGeometriesToUpdate.Add(
-			FRayTracingDynamicGeometryUpdateParams
-			{
-				RayTracingInstance.Materials,
-				false,
-				VertexCount,
-				VertexCount * (uint32)sizeof(FVector),
-				MeshBatch.Elements[0].NumPrimitives,
-				&RayTracingGeometry,
-				&RayTracingDynamicVertexBuffer,
-			}
-		);
-	}
+	const uint32 VertexCount = DynamicIndexAllocation.MaxUsedIndex;
+	Context.DynamicRayTracingGeometriesToUpdate.Add(
+		FRayTracingDynamicGeometryUpdateParams
+		{
+			RayTracingInstance.Materials,
+			false,
+			VertexCount,
+			VertexCount * (uint32)sizeof(FVector),
+			MeshBatch.Elements[0].NumPrimitives,
+			&RayTracingGeometry,
+			&RayTracingDynamicVertexBuffer,
+		}
+	);
 
 	RayTracingInstance.BuildInstanceMaskAndFlags();
 
