Fixes to NiagaraLightRenderer
-skip generation of dynamic data if we're not using deferred shading
-make it draw relevant so that they will actually render without relying on other renderers

diff --git a/Engine/Plugins/FX/Niagara/Source/Niagara/Private/NiagaraRendererLights.cpp b/Engine/Plugins/FX/Niagara/Source/Niagara/Private/NiagaraRendererLights.cpp
--- a/Engine/Plugins/FX/Niagara/Source/Niagara/Private/NiagaraRendererLights.cpp
+++ b/Engine/Plugins/FX/Niagara/Source/Niagara/Private/NiagaraRendererLights.cpp
@@ -28,24 +28,31 @@ struct FNiagaraDynamicDataLights : public FNiagaraDynamicDataBase
 FNiagaraRendererLights::FNiagaraRendererLights(ERHIFeatureLevel::Type FeatureLevel, const UNiagaraRendererProperties *InProps, const FNiagaraEmitterInstance* Emitter)
 	: FNiagaraRenderer(FeatureLevel, InProps, Emitter)
 {
+	// todo - for platforms where we know we can't support deferred shading we can just set this to false
 	bHasLights = true;
 }
 
 FPrimitiveViewRelevance FNiagaraRendererLights::GetViewRelevance(const FSceneView* View, const FNiagaraSceneProxy *SceneProxy)const
 {
 	FPrimitiveViewRelevance Result;
-	Result.bDrawRelevance = false;
+	Result.bDrawRelevance = bHasLights;
 	Result.bShadowRelevance = false;
 	Result.bDynamicRelevance = false;
 	Result.bOpaque = false;
-	Result.bHasSimpleLights = true;
-	//MaterialRelevance.SetPrimitiveViewRelevance(Result);
+	Result.bHasSimpleLights = bHasLights;
+
 	return Result;
 }
 
 /** Update render data buffer from attributes */
 FNiagaraDynamicDataBase* FNiagaraRendererLights::GenerateDynamicData(const FNiagaraSceneProxy* Proxy, const UNiagaraRendererProperties* InProperties, const FNiagaraEmitterInstance* Emitter) const
 {
+	// particle (simple) lights are only supported with deferred shading
+	if (!bHasLights || Proxy->GetScene().GetShadingPath() != EShadingPath::Deferred)
+	{
+		return nullptr;
+	}
+
 	SCOPE_CYCLE_COUNTER(STAT_NiagaraGenLights);
 
 	//Bail if we don't have the required attributes to render this emitter.
@@ -99,7 +106,8 @@ FNiagaraDynamicDataBase* FNiagaraRendererLights::GenerateDynamicData(const FNiag
 		float LightRadius = RadiusAccessor.GetSafe(ParticleIndex, DefaultRadius) * Properties->RadiusScale;
 		if (ShouldRenderParticleLight && LightRadius > 0)
 		{
-			SimpleLightData LightData;
+			SimpleLightData& LightData = DynamicData->LightArray.AddDefaulted_GetRef();
+
 			LightData.LightEntry.Radius = LightRadius;
 			LightData.LightEntry.Color = FVector(ColAccessor.GetSafe(ParticleIndex, DefaultColor)) + Properties->ColorAdd;
 			LightData.LightEntry.Exponent = Properties->bUseInverseSquaredFalloff ? 0 : ExponentAccessor.GetSafe(ParticleIndex, 1.0f);
@@ -110,8 +118,6 @@ FNiagaraDynamicDataBase* FNiagaraRendererLights::GenerateDynamicData(const FNiag
 			{
 				LightData.PerViewEntry.Position = LocalToWorldMatrix.TransformPosition(LightData.PerViewEntry.Position);
 			}
-
-			DynamicData->LightArray.Add(LightData);
 		}
 	}
 
@@ -120,15 +126,14 @@ FNiagaraDynamicDataBase* FNiagaraRendererLights::GenerateDynamicData(const FNiag
 
 void FNiagaraRendererLights::GatherSimpleLights(FSimpleLightArray& OutParticleLights)const
 {
-	FNiagaraDynamicDataLights* DynamicData = (FNiagaraDynamicDataLights*)DynamicDataRender;
-	if (DynamicData)
+	if (const FNiagaraDynamicDataLights* DynamicData = static_cast<const FNiagaraDynamicDataLights*>(DynamicDataRender))
 	{
-		int32 LightCount = DynamicData->LightArray.Num();
+		const int32 LightCount = DynamicData->LightArray.Num();
 
 		OutParticleLights.InstanceData.Reserve(OutParticleLights.InstanceData.Num() + LightCount);
 		OutParticleLights.PerViewData.Reserve(OutParticleLights.PerViewData.Num() + LightCount);
 
-		for (FNiagaraRendererLights::SimpleLightData &LightData : DynamicData->LightArray)
+		for (const FNiagaraRendererLights::SimpleLightData &LightData : DynamicData->LightArray)
 		{
 			// When not using camera-offset, output one position for all views to share. 
 			OutParticleLights.PerViewData.Add(LightData.PerViewEntry);
