- Niagara commandlet to find assets with warmup time / lights / GPU emitters

diff --git a/Engine/Plugins/FX/Niagara/Source/Niagara/Public/NiagaraLightRendererProperties.h b/Engine/Plugins/FX/Niagara/Source/Niagara/Public/NiagaraLightRendererProperties.h
--- a/Engine/Plugins/FX/Niagara/Source/Niagara/Public/NiagaraLightRendererProperties.h
+++ b/Engine/Plugins/FX/Niagara/Source/Niagara/Public/NiagaraLightRendererProperties.h
@@ -11,7 +11,7 @@
 class FNiagaraEmitterInstance;
 class SWidget;
 
-UCLASS(editinlinenew, meta = (DisplayName = "Light Renderer"))
+UCLASS(editinlinenew, MinimalAPI, meta = (DisplayName = "Light Renderer"))
 class UNiagaraLightRendererProperties : public UNiagaraRendererProperties
 {
 public:
diff --git a/Engine/Plugins/FX/Niagara/Source/NiagaraEditor/Private/Commandlets/NiagaraSystemAuditCommandlet.cpp b/Engine/Plugins/FX/Niagara/Source/NiagaraEditor/Private/Commandlets/NiagaraSystemAuditCommandlet.cpp
new file mode 100644
--- /dev/null
+++ b/Engine/Plugins/FX/Niagara/Source/NiagaraEditor/Private/Commandlets/NiagaraSystemAuditCommandlet.cpp
@@ -0,0 +1,224 @@
+// Copyright Epic Games, Inc. All Rights Reserved.
+
+#include "Commandlets/NiagaraSystemAuditCommandlet.h"
+#include "HAL/FileManager.h"
+#include "Misc/Paths.h"
+#include "Modules/ModuleManager.h"
+#include "Misc/PackageName.h"
+#include "AssetData.h"
+#include "ARFilter.h"
+#include "AssetRegistryModule.h"
+#include "CollectionManagerTypes.h"
+#include "ICollectionManager.h"
+#include "CollectionManagerModule.h"
+
+#include "NiagaraSystem.h"
+#include "NiagaraRendererProperties.h"
+#include "NiagaraLightRendererProperties.h"
+
+DEFINE_LOG_CATEGORY_STATIC(LogNiagaraSystemAuditCommandlet, Log, All);
+
+UNiagaraSystemAuditCommandlet::UNiagaraSystemAuditCommandlet(const FObjectInitializer& ObjectInitializer)
+	: Super(ObjectInitializer)
+{
+}
+
+int32 UNiagaraSystemAuditCommandlet::Main(const FString& Params)
+{
+	if (!FParse::Value(*Params, TEXT("AuditOutputFolder="), AuditOutputFolder))
+	{
+		// No output folder specified. Use the default folder.
+		AuditOutputFolder = FPaths::ProjectSavedDir() / TEXT("Audit");
+	}
+
+	// Add a timestamp to the folder
+	AuditOutputFolder /= FDateTime::Now().ToString();
+
+	FParse::Value(*Params, TEXT("FilterCollection="), FilterCollection);
+
+	FString PackagePathsString;
+	if (FParse::Value(*Params, TEXT("PackagePaths="), PackagePathsString))
+	{
+		TArray<FString> PackagePathsStrings;
+		PackagePathsString.ParseIntoArray(PackagePathsStrings, TEXT(","));
+		for (const FString& v : PackagePathsStrings)
+		{
+			PackagePaths.Add(FName(v));
+		}
+	}
+
+	if (PackagePaths.Num() == 0)
+	{
+		PackagePaths.Add(FName(TEXT("/Game")));
+	}
+
+	ProcessNiagaraSystems();
+	DumpResults();
+
+	return 0;
+}
+
+bool UNiagaraSystemAuditCommandlet::ProcessNiagaraSystems()
+{
+	FAssetRegistryModule& AssetRegistryModule = FModuleManager::LoadModuleChecked<FAssetRegistryModule>(TEXT("AssetRegistry"));
+	IAssetRegistry& AssetRegistry = AssetRegistryModule.Get();
+	AssetRegistry.SearchAllAssets(true);
+
+	FARFilter Filter;
+	Filter.PackagePaths = PackagePaths;
+	Filter.bRecursivePaths = true;
+
+	Filter.ClassNames.Add(UNiagaraSystem::StaticClass()->GetFName());
+	if (!FilterCollection.IsEmpty())
+	{
+		FCollectionManagerModule& CollectionManagerModule = FCollectionManagerModule::GetModule();
+		CollectionManagerModule.Get().GetObjectsInCollection(FName(*FilterCollection), ECollectionShareType::CST_All, Filter.ObjectPaths, ECollectionRecursionFlags::SelfAndChildren);
+	}
+
+	TArray<FAssetData> AssetList;
+	AssetRegistry.GetAssets(Filter, AssetList);
+
+	double StartProcessNiagaraSystemsTime = FPlatformTime::Seconds();
+
+	//  Iterate over all systems
+	const FString DevelopersFolder = FPackageName::FilenameToLongPackageName(FPaths::GameDevelopersDir().LeftChop(1));
+	FString LastPackageName = TEXT("");
+	int32 PackageSwitches = 0;
+	UPackage* CurrentPackage = nullptr;
+	for (const FAssetData& AssetIt : AssetList)
+	{
+		const FString SystemName = AssetIt.ObjectPath.ToString();
+		const FString PackageName = AssetIt.PackageName.ToString();
+
+		if (PackageName.StartsWith(DevelopersFolder))
+		{
+			// Skip developer folders
+			continue;
+		}
+
+		if (PackageName != LastPackageName)
+		{
+			UPackage* Package = ::LoadPackage(nullptr, *PackageName, LOAD_None);
+			if (Package != nullptr)
+			{
+				LastPackageName = PackageName;
+				Package->FullyLoad();
+				CurrentPackage = Package;
+			}
+			else
+			{
+				UE_LOG(LogNiagaraSystemAuditCommandlet, Warning, TEXT("Failed to load package %s processing %s"), *PackageName, *SystemName);
+				CurrentPackage = nullptr;
+			}
+		}
+
+		const FString ShorterSystemName = AssetIt.AssetName.ToString();
+		UNiagaraSystem* NiagaraSystem = FindObject<UNiagaraSystem>(CurrentPackage, *ShorterSystemName);
+		if (NiagaraSystem == nullptr)
+		{
+			UE_LOG(LogNiagaraSystemAuditCommandlet, Warning, TEXT("Failed to load Niagara system %s"), *SystemName);
+			continue;
+		}
+
+		// Iterate over all emitters
+		bool bHasGPUEmitters = false;
+		bool bHasLights = false;
+
+		for (const FNiagaraEmitterHandle& EmitterHandle : NiagaraSystem->GetEmitterHandles())
+		{
+			UNiagaraEmitter* NiagaraEmitter = EmitterHandle.GetInstance();
+			if (NiagaraEmitter == nullptr)
+			{
+				continue;
+			}
+
+			bHasGPUEmitters |= NiagaraEmitter->SimTarget == ENiagaraSimTarget::GPUComputeSim;
+
+			for (UNiagaraRendererProperties* RendererProperties : NiagaraEmitter->GetRenderers())
+			{
+				if (UNiagaraLightRendererProperties* LightRendererProperties = Cast< UNiagaraLightRendererProperties>(RendererProperties))
+				{
+					bHasLights = true;
+				}
+			}
+		}
+
+		// Add to different charts we will write out
+		if (NiagaraSystem->GetWarmupTime() > 0.0f)
+		{
+			NiagaraSystemsWithWarmup.Add(FString::Printf(TEXT("%s,%f"), *NiagaraSystem->GetPathName(), NiagaraSystem->GetWarmupTime()));
+		}
+
+		if (bHasLights)
+		{
+			NiagaraSystemsWithLights.Add(NiagaraSystem->GetPathName());
+		}
+
+		if (bHasGPUEmitters)
+		{
+			NiagaraSystemsWithGPUEmitters.Add(NiagaraSystem->GetPathName());
+		}
+	}
+
+	// Probably don't need to do this, but just in case we have any 'hanging' packages 
+	// and more processing steps are added later, let's clean up everything...
+	::CollectGarbage(RF_NoFlags);
+
+	double ProcessNiagaraSystemsTime = FPlatformTime::Seconds() - StartProcessNiagaraSystemsTime;
+	UE_LOG(LogNiagaraSystemAuditCommandlet, Log, TEXT("Took %5.3f seconds to process referenced Niagara systems..."), ProcessNiagaraSystemsTime);
+
+	return true;
+}
+
+/** Dump the results of the audit */
+void UNiagaraSystemAuditCommandlet::DumpResults()
+{
+	// Dump all the simple mappings...
+	DumpSimpleSet(NiagaraSystemsWithWarmup, TEXT("NiagaraSystemsWithWarmup"), TEXT("Name,WarmupTime"));
+	DumpSimpleSet(NiagaraSystemsWithLights, TEXT("NiagaraSystemsWithLights"), TEXT("Name"));
+	DumpSimpleSet(NiagaraSystemsWithGPUEmitters, TEXT("NiagaraSystemsWithGPUEmitters"), TEXT("Name"));
+
+}
+
+bool UNiagaraSystemAuditCommandlet::DumpSimpleSet(TSet<FString>& InSet, const TCHAR* InShortFilename, const TCHAR* OptionalHeader)
+{
+	if (InSet.Num() > 0)
+	{
+		check(InShortFilename != NULL);
+
+		FArchive* OutputStream = GetOutputFile(InShortFilename);
+		if (OutputStream != NULL)
+		{
+			UE_LOG(LogNiagaraSystemAuditCommandlet, Log, TEXT("Dumping '%s' results..."), InShortFilename);
+			if (OptionalHeader != nullptr)
+			{
+				OutputStream->Logf(TEXT("%s"), OptionalHeader);
+			}
+
+			for (TSet<FString>::TIterator DumpIt(InSet); DumpIt; ++DumpIt)
+			{
+				FString ObjName = *DumpIt;
+				OutputStream->Logf(TEXT("%s"), *ObjName);
+			}
+
+			OutputStream->Close();
+			delete OutputStream;
+		}
+		else
+		{
+			return false;
+		}
+	}
+	return true;
+}
+
+FArchive* UNiagaraSystemAuditCommandlet::GetOutputFile(const TCHAR* InShortFilename)
+{
+	const FString Filename = FString::Printf(TEXT("%s/%s.csv"), *AuditOutputFolder, InShortFilename);
+	FArchive* OutputStream = IFileManager::Get().CreateDebugFileWriter(*Filename);
+	if (OutputStream == NULL)
+	{
+		UE_LOG(LogNiagaraSystemAuditCommandlet, Warning, TEXT("Failed to create output stream %s"), *Filename);
+	}
+	return OutputStream;
+}
\ No newline at end of file
diff --git a/Engine/Plugins/FX/Niagara/Source/NiagaraEditor/Public/Commandlets/NiagaraSystemAuditCommandlet.h b/Engine/Plugins/FX/Niagara/Source/NiagaraEditor/Public/Commandlets/NiagaraSystemAuditCommandlet.h
new file mode 100644
--- /dev/null
+++ b/Engine/Plugins/FX/Niagara/Source/NiagaraEditor/Public/Commandlets/NiagaraSystemAuditCommandlet.h
@@ -0,0 +1,45 @@
+// Copyright Epic Games, Inc. All Rights Reserved.
+
+#pragma once
+
+#include "CoreMinimal.h"
+#include "UObject/ObjectMacros.h"
+#include "Commandlets/Commandlet.h"
+#include "NiagaraSystemAuditCommandlet.generated.h"
+
+UCLASS(config=Editor)
+class UNiagaraSystemAuditCommandlet : public UCommandlet
+{
+	GENERATED_UCLASS_BODY()
+
+	/** All Niagara systems with warmup time */
+	TSet<FString> NiagaraSystemsWithWarmup;
+	/** All Niagara systems with lights and the details */
+	TSet<FString> NiagaraSystemsWithLights;
+	/** All Niagara systems with GPU emitters */
+	TSet<FString> NiagaraSystemsWithGPUEmitters;
+
+	/** The folder in which the commandlet's output files will be stored */
+	FString AuditOutputFolder;
+
+	/** Only assets in this collection will be considered. If this is left blank, no assets will be filtered by collection */
+	FString FilterCollection;
+
+	/** Package paths to include */
+	TArray<FName> PackagePaths;
+
+	/** Entry point */
+	int32 Main(const FString& Params) override;
+
+	/** Process all referenced Niagara systems */
+	bool ProcessNiagaraSystems();
+
+	/** Dump the results of the audit */
+	virtual void DumpResults();
+
+	/** Generic function to handle dumping values to a CSV file */
+	bool DumpSimpleSet(TSet<FString>& InSet, const TCHAR* InShortFilename, const TCHAR* OptionalHeader = nullptr);
+
+	/** Gets an archive to write to an output file */
+	FArchive* GetOutputFile(const TCHAR* InShortFilename);
+};
