Fix for handling of user override parameters on a component and how it handles booleans

diff --git a/Engine/Plugins/FX/Niagara/Source/Niagara/Private/NiagaraComponent.cpp b/Engine/Plugins/FX/Niagara/Source/Niagara/Private/NiagaraComponent.cpp
--- a/Engine/Plugins/FX/Niagara/Source/Niagara/Private/NiagaraComponent.cpp
+++ b/Engine/Plugins/FX/Niagara/Source/Niagara/Private/NiagaraComponent.cpp
@@ -1618,9 +1618,10 @@ void UNiagaraComponent::SetNiagaraVariableInt(const FString& InVariableName, int
 void UNiagaraComponent::SetVariableBool(FName InVariableName, bool InValue)
 {
 	const FNiagaraVariable VariableDesc(FNiagaraTypeDefinition::GetBoolDef(), InVariableName);
-	OverrideParameters.SetParameterValue(InValue ? FNiagaraBool::True : FNiagaraBool::False, VariableDesc, true);
+	const FNiagaraBool BoolValue(InValue);
+	OverrideParameters.SetParameterValue(BoolValue, VariableDesc, true);
 #if WITH_EDITOR
-	SetParameterOverride(VariableDesc, FNiagaraVariant(&InValue, sizeof(bool)));
+	SetParameterOverride(VariableDesc, FNiagaraVariant(&BoolValue, sizeof(FNiagaraBool)));
 #endif
 }
 
