- Potential fix for OOB array access and also add an check for null FNiagaraScriptExecutionParameterStore

#rnx

diff --git a/Engine/Plugins/FX/Niagara/Source/Niagara/Private/NiagaraScriptExecutionContext.cpp b/Engine/Plugins/FX/Niagara/Source/Niagara/Private/NiagaraScriptExecutionContext.cpp
--- a/Engine/Plugins/FX/Niagara/Source/Niagara/Private/NiagaraScriptExecutionContext.cpp
+++ b/Engine/Plugins/FX/Niagara/Source/Niagara/Private/NiagaraScriptExecutionContext.cpp
@@ -71,6 +71,9 @@ bool FNiagaraScriptExecutionContext::Tick(FNiagaraSystemInstance* ParentSystemIn
 				return false;
 			}
 
+			const FNiagaraScriptExecutionParameterStore* ScriptParameterStore = Script->GetExecutionReadyParameterStore(SimTarget);
+			check(ScriptParameterStore != nullptr);
+
 			//Fill the instance data table.
 			if (ParentSystemInstance)
 			{
@@ -98,7 +101,7 @@ bool FNiagaraScriptExecutionContext::Tick(FNiagaraSystemInstance* ParentSystemIn
 			TArray<int32> LocalFunctionTableIndices;
 			LocalFunctionTableIndices.Reserve(FunctionCount);
 
-			const auto& ScriptDataInterfaces = Script->GetExecutionReadyParameterStore(SimTarget)->GetDataInterfaces();
+			const auto& ScriptDataInterfaces = ScriptParameterStore->GetDataInterfaces();
 
 			bool bSuccessfullyMapped = true;
 
@@ -123,6 +126,7 @@ bool FNiagaraScriptExecutionContext::Tick(FNiagaraSystemInstance* ParentSystemIn
 					{
 						// first check to see if we should just use the one from the script
 						if (ScriptExecutableData.CalledVMExternalFunctionBindings.IsValidIndex(FunctionIt)
+							&& ScriptDataInterfaces.IsValidIndex(i)
 							&& ExternalInterface == ScriptDataInterfaces[i])
 						{
 							const FVMExternalFunction& ScriptFuncBind = ScriptExecutableData.CalledVMExternalFunctionBindings[FunctionIt];
