- Remove incorrect call to GetSkinWeights that is not required and can potentially crash, found and fixed by Uriel just submitting it for him

diff --git a/Engine/Plugins/FX/Niagara/Source/Niagara/Private/NDISkeletalMesh_BoneSampling.cpp b/Engine/Plugins/FX/Niagara/Source/Niagara/Private/NDISkeletalMesh_BoneSampling.cpp
--- a/Engine/Plugins/FX/Niagara/Source/Niagara/Private/NDISkeletalMesh_BoneSampling.cpp
+++ b/Engine/Plugins/FX/Niagara/Source/Niagara/Private/NDISkeletalMesh_BoneSampling.cpp
@@ -787,8 +787,6 @@ void UNiagaraDataInterfaceSkeletalMesh::GetSkinnedBoneData(FVectorVMContext& Con
 	const FQuat InstanceRotation = Output.bNeedsRotation ? InstanceTransform.ToQuat() : FQuat::Identity;
 	const FQuat PrevInstanceRotation = Output.bNeedsRotation ? PrevInstanceTransform.ToQuat() : FQuat::Identity;
 
-	FSkinWeightVertexBuffer* SkinWeightBuffer = InstData->GetSkinWeights();
-
 	FSkeletalMeshAccessorHelper Accessor;
 	Accessor.Init<TNDISkelMesh_FilterModeNone, TNDISkelMesh_AreaWeightingOff>(InstData);
 	check(Accessor.AreBonesAccessible());
