Removed a limitation in the Niagara skinned data cache where we assumed only one emitter was sampling a given skeletal mesh component.
Allowed Niagara Skeletal Mesh DI to bound to the skeletal mesh component from the scene actor hierarchy, as an alternative way manage binding without requiring blueprints.
Prevented skeletal mesh LODs from being streamed out while a Niagara emitters are sampling it.
Skeletal mesh stream out now waits for external references to the streamed out LOD to be released before discarding the LOD CPU data.
FSkeletalMeshLODRenderData is now ref counted to track the number of active references to prevent stream out while LODs are being sampled.

diff --git a/Engine/Plugins/FX/Niagara/Source/Niagara/Classes/NDISkeletalMeshCommon.h b/Engine/Plugins/FX/Niagara/Source/Niagara/Classes/NDISkeletalMeshCommon.h
--- a/Engine/Plugins/FX/Niagara/Source/Niagara/Classes/NDISkeletalMeshCommon.h
+++ b/Engine/Plugins/FX/Niagara/Source/Niagara/Classes/NDISkeletalMeshCommon.h
@@ -40,7 +40,8 @@ struct FSkeletalMeshAccessorHelper
 	{
 		Comp = Cast<USkeletalMeshComponent>(InstData->Component.Get());
 		Mesh = InstData->Mesh;
-		LODData = InstData->GetLODRenderDataAndSkinWeights(SkinWeightBuffer);
+		LODData = InstData->CachedLODData;
+		SkinWeightBuffer = InstData->GetSkinWeights();
 		IndexBuffer = LODData ? LODData->MultiSizeIndexContainer.GetIndexBuffer() : nullptr;
 		SkinningData = InstData->SkinningData.SkinningData.Get();
 		Usage = InstData->SkinningData.Usage;
@@ -64,9 +65,9 @@ struct FSkeletalMeshAccessorHelper
 	USkeletalMeshComponent* Comp = nullptr;
 	USkeletalMesh* Mesh = nullptr;
 	TWeakObjectPtr<USkeletalMesh> MeshSafe;
-	FSkeletalMeshLODRenderData* LODData = nullptr;
-	FSkinWeightVertexBuffer* SkinWeightBuffer = nullptr;
-	FRawStaticIndexBuffer16or32Interface* IndexBuffer = nullptr;
+	const FSkeletalMeshLODRenderData* LODData = nullptr;
+	const FSkinWeightVertexBuffer* SkinWeightBuffer = nullptr;
+	const FRawStaticIndexBuffer16or32Interface* IndexBuffer = nullptr;
 	const FSkeletalMeshSamplingRegion* SamplingRegion = nullptr;
 	const FSkeletalMeshSamplingRegionBuiltData* SamplingRegionBuiltData = nullptr;
 	FSkeletalMeshSkinningData* SkinningData = nullptr;
@@ -412,7 +413,7 @@ struct FSkinnedPositionAccessorHelper<TNDISkelMesh_SkinningModePreSkin>
 template<bool bUseFullPrecisionUVs>
 struct FSkelMeshVertexAccessor
 {
-	FORCEINLINE FVector2D GetVertexUV(FSkeletalMeshLODRenderData* LODData, int32 VertexIdx, int32 UVChannel)const
+	FORCEINLINE FVector2D GetVertexUV(const FSkeletalMeshLODRenderData* LODData, int32 VertexIdx, int32 UVChannel)const
 	{
 		check(LODData);
 		if (bUseFullPrecisionUVs)
@@ -434,12 +435,12 @@ struct FSkelMeshVertexAccessor
 
 struct FSkelMeshVertexAccessorNoop
 {
-	FORCEINLINE FVector2D GetVertexUV(FSkeletalMeshLODRenderData* LODData, int32 VertexIdx, int32 UVChannel)const
+	FORCEINLINE FVector2D GetVertexUV(const FSkeletalMeshLODRenderData* LODData, int32 VertexIdx, int32 UVChannel)const
 	{
 		return FVector2D(0.0f, 0.0f);
 	}
 
-	FORCEINLINE FLinearColor GetVertexColor(FSkeletalMeshLODRenderData* LODData, int32 VertexIdx)const
+	FORCEINLINE FLinearColor GetVertexColor(const FSkeletalMeshLODRenderData* LODData, int32 VertexIdx)const
 	{
 		return FLinearColor::White;
 	}
@@ -528,8 +529,7 @@ struct TVertexAccessorBinder
 
 		if (InstData->bAllowCPUMeshDataAccess)
 		{
-			FSkinWeightVertexBuffer* SkinWeightBuffer = nullptr;
-			FSkeletalMeshLODRenderData* LODData = InstData->GetLODRenderDataAndSkinWeights(SkinWeightBuffer);
+			const FSkeletalMeshLODRenderData* LODData = InstData->CachedLODData;
 			if (LODData->StaticVertexBuffers.StaticMeshVertexBuffer.GetUseFullPrecisionUVs())
 			{
 				NextBinder::template Bind<ParamTypes..., FSkelMeshVertexAccessor<true>>(Interface, BindingInfo, InstanceData, OutFunc);
diff --git a/Engine/Plugins/FX/Niagara/Source/Niagara/Classes/NiagaraDataInterfaceSkeletalMesh.h b/Engine/Plugins/FX/Niagara/Source/Niagara/Classes/NiagaraDataInterfaceSkeletalMesh.h
--- a/Engine/Plugins/FX/Niagara/Source/Niagara/Classes/NiagaraDataInterfaceSkeletalMesh.h
+++ b/Engine/Plugins/FX/Niagara/Source/Niagara/Classes/NiagaraDataInterfaceSkeletalMesh.h
@@ -27,32 +27,27 @@ struct FSkeletalMeshSkinningDataUsage
 		: LODIndex(INDEX_NONE)
 		, bUsesBoneMatrices(false)
 		, bUsesPreSkinnedVerts(false)
-		, bNeedDataImmediately(false)
 	{}
 
-	FSkeletalMeshSkinningDataUsage(int32 InLODIndex, bool bInUsesBoneMatrices, bool bInUsesPreSkinnedVerts, bool bInNeedDataImmediately)
+	FSkeletalMeshSkinningDataUsage(int32 InLODIndex, bool bInUsesBoneMatrices, bool bInUsesPreSkinnedVerts)
 		: LODIndex(InLODIndex)
 		, bUsesBoneMatrices(bInUsesBoneMatrices)
 		, bUsesPreSkinnedVerts(bInUsesPreSkinnedVerts)
-		, bNeedDataImmediately(bInNeedDataImmediately)
 	{}
 
 	FORCEINLINE bool NeedBoneMatrices()const { return bUsesBoneMatrices || bUsesPreSkinnedVerts; }
 	FORCEINLINE bool NeedPreSkinnedVerts()const { return bUsesPreSkinnedVerts; }
-	FORCEINLINE bool NeedsDataImmediately()const { return bNeedDataImmediately; }
 	FORCEINLINE int32 GetLODIndex()const { return LODIndex; }
 private:
 	int32 LODIndex;
 	uint32 bUsesBoneMatrices : 1;
 	uint32 bUsesPreSkinnedVerts : 1;
-	/** Some users need valid data immediately after the register call rather than being able to wait until the next tick. */
-	uint32 bNeedDataImmediately : 1;
 };
 
 struct FSkeletalMeshSkinningDataHandle
 {
 	FSkeletalMeshSkinningDataHandle();
-	FSkeletalMeshSkinningDataHandle(FSkeletalMeshSkinningDataUsage InUsage, TSharedPtr<struct FSkeletalMeshSkinningData> InSkinningData);
+	FSkeletalMeshSkinningDataHandle(FSkeletalMeshSkinningDataUsage InUsage, const TSharedPtr<struct FSkeletalMeshSkinningData>& InSkinningData, bool bNeedsDataImmediately);
 	FSkeletalMeshSkinningDataHandle(const FSkeletalMeshSkinningDataHandle& Other) = delete;
 	FSkeletalMeshSkinningDataHandle(FSkeletalMeshSkinningDataHandle&& Other);
 	~FSkeletalMeshSkinningDataHandle();
@@ -71,10 +66,11 @@ struct FSkeletalMeshSkinningData
 		, DeltaSeconds(.0333f)
 		, CurrIndex(0)
 		, BoneMatrixUsers(0)
+		, TotalPreSkinnedVertsUsers(0)
 		, bForceDataRefresh(false)
 	{}
 
-	void RegisterUser(FSkeletalMeshSkinningDataUsage Usage);
+	void RegisterUser(FSkeletalMeshSkinningDataUsage Usage, bool bNeedsDataImmediately);
 	void UnregisterUser(FSkeletalMeshSkinningDataUsage Usage);
 	bool IsUsed()const;
 	void ForceDataRefresh();
@@ -173,6 +169,14 @@ struct FSkeletalMeshSkinningData
 		return ComponentTransforms[CurrIndex ^ 1];
 	}
 
+	FORCEINLINE bool NeedPreSkinnedVerts() const
+	{
+		return TotalPreSkinnedVertsUsers > 0;
+	}
+
+	/** Whether this has been ticked this frame.*/
+	mutable bool bHasTicked = false;
+
 private:
 
 	void UpdateBoneTransforms();
@@ -189,6 +193,8 @@ private:
 
 	/** Number of users for cached bone matrices. */
 	volatile int32 BoneMatrixUsers;
+	/** Total number of users for pre skinned verts.  (From LODData) */
+	volatile int32 TotalPreSkinnedVertsUsers;
 
 	/** Cached bone matrices. */
 	TArray<FMatrix> BoneRefToLocals[2];
@@ -216,20 +222,11 @@ private:
 
 class FNDI_SkeletalMesh_GeneratedData
 {
-	// Encapsulates skinning data and mesh usage information. 
-	// Set by GetCachedSkinningData and used by TickGeneratedData to determine whether we need to pre-skin or not.
-	struct CachedSkinningDataAndUsage
-	{
-		bool bHasTicked = false;
-		TSharedPtr<FSkeletalMeshSkinningData> SkinningData;
-		FSkeletalMeshSkinningDataUsage Usage;
-	};
-
 	FRWLock CachedSkinningDataGuard;
-	TMap<TWeakObjectPtr<USkeletalMeshComponent>, CachedSkinningDataAndUsage> CachedSkinningData;
+	TMap<TWeakObjectPtr<USkeletalMeshComponent>, TSharedPtr<FSkeletalMeshSkinningData> > CachedSkinningData;
 
 public:
-	FSkeletalMeshSkinningDataHandle GetCachedSkinningData(TWeakObjectPtr<USkeletalMeshComponent>& InComponent, FSkeletalMeshSkinningDataUsage Usage);
+	FSkeletalMeshSkinningDataHandle GetCachedSkinningData(TWeakObjectPtr<USkeletalMeshComponent>& InComponent, FSkeletalMeshSkinningDataUsage Usage, bool bNeedsDataImmediately);
 	void TickGeneratedData(ETickingGroup TickGroup, float DeltaSeconds);
 };
 
@@ -519,42 +516,29 @@ struct FNDISkeletalMesh_InstanceData
 	/** Flag to stub VM functions that rely on mesh data being accessible on the CPU */
 	bool bAllowCPUMeshDataAccess;
 
+	/** The MinLOD applicable to the skeletal mesh, based on USkeletalMesh::MinLod which is platform specific.*/
+	int32 MinLODIdx = 0;
+	/** Whether to reset the emitter if any LOD get streamed in. Used when the required LOD was not initially available. */
+	bool bResetOnLODStreamedIn = false;
+	/** The cached LODIdx used to initialize the FNDIStaticMesh_InstanceData.*/
+	int32 CachedLODIdx = 0;
+	/** The referenced LOD data, used to prevent streaming out LODs while they are being referenced*/
+	TRefCountPtr<const FSkeletalMeshLODRenderData> CachedLODData;
+
 	FORCEINLINE_DEBUGGABLE bool ResetRequired(UNiagaraDataInterfaceSkeletalMesh* Interface)const;
 
 	bool Init(UNiagaraDataInterfaceSkeletalMesh* Interface, FNiagaraSystemInstance* SystemInstance);
 	FORCEINLINE_DEBUGGABLE bool Tick(UNiagaraDataInterfaceSkeletalMesh* Interface, FNiagaraSystemInstance* SystemInstance, float InDeltaSeconds);
 	FORCEINLINE_DEBUGGABLE void Release();
-	FORCEINLINE int32 GetLODIndex()const { return SkinningData.Usage.GetLODIndex(); }
+	FORCEINLINE int32 GetLODIndex()const { return CachedLODIdx; }
 
-	FORCEINLINE_DEBUGGABLE FSkeletalMeshLODRenderData* GetLODRenderDataAndSkinWeights(FSkinWeightVertexBuffer*& OutSkinWeightBuffer)
-	{
-		FSkeletalMeshLODRenderData* Ret = nullptr;
-		OutSkinWeightBuffer = nullptr;
-		if (Mesh)
-		{
-			Ret = &Mesh->GetResourceForRendering()->LODRenderData[GetLODIndex()];
-			if (bAllowCPUMeshDataAccess)
-			{
-				if (USkeletalMeshComponent* SkelComp = Cast<USkeletalMeshComponent>(Component.Get()))
-				{
-					OutSkinWeightBuffer = SkelComp->GetSkinWeightBuffer(GetLODIndex());
-				}
-				else
-				{
-					OutSkinWeightBuffer = &Ret->SkinWeightVertexBuffer; // Todo @sckime fix this when main comes in for real
-				}
-			}
-		}
-		return Ret;
-	}
-
-	FORCEINLINE_DEBUGGABLE FSkinWeightVertexBuffer* GetSkinWeights()
+	FORCEINLINE_DEBUGGABLE const FSkinWeightVertexBuffer* GetSkinWeights()
 	{
 		if (USkeletalMeshComponent* SkelComp = Cast<USkeletalMeshComponent>(Component.Get()))
 		{
-			return SkelComp->GetSkinWeightBuffer(GetLODIndex());
+			return SkelComp->GetSkinWeightBuffer(CachedLODIdx);
 		}
-		return &Mesh->GetResourceForRendering()->LODRenderData[GetLODIndex()].SkinWeightVertexBuffer;
+		return CachedLODData ? &CachedLODData->SkinWeightVertexBuffer : nullptr;
 	}
 
 	void UpdateFilteredSocketTransforms();
@@ -662,6 +646,8 @@ public:
 	virtual bool UpgradeFunctionCall(FNiagaraFunctionSignature& FunctionSignature) override;
 #endif
 
+	int32 CalculateLODIndexAndSamplingRegions(USkeletalMesh* InMesh, TArray<int32>& OutSamplingRegionIndices, bool& OutAllRegionsAreAreaWeighting) const;
+
 	virtual void ProvidePerInstanceDataForRenderThread(void* DataForRenderThread, void* PerInstanceData, const FNiagaraSystemInstanceID& SystemInstance) override;
 
 	static const FString MeshIndexBufferName;
diff --git a/Engine/Plugins/FX/Niagara/Source/Niagara/Private/NDISkeletalMesh_TriangleSampling.cpp b/Engine/Plugins/FX/Niagara/Source/Niagara/Private/NDISkeletalMesh_TriangleSampling.cpp
--- a/Engine/Plugins/FX/Niagara/Source/Niagara/Private/NDISkeletalMesh_TriangleSampling.cpp
+++ b/Engine/Plugins/FX/Niagara/Source/Niagara/Private/NDISkeletalMesh_TriangleSampling.cpp
@@ -358,7 +358,7 @@ FORCEINLINE int32 UNiagaraDataInterfaceSkeletalMesh::RandomTriIndex<TNDISkelMesh
 {
 	int32 SecIdx = RandHelper.RandRange(InstanceIndex, 0, Accessor.LODData->RenderSections.Num() - 1);
 	//int32 SecIdx = RandStream.RandRange(0, Accessor.LODData->RenderSections.Num() - 1);
-	FSkelMeshRenderSection& Sec = Accessor.LODData->RenderSections[SecIdx];
+	const FSkelMeshRenderSection& Sec = Accessor.LODData->RenderSections[SecIdx];
 	int32 Tri = RandHelper.RandRange(InstanceIndex, 0, Sec.NumTriangles - 1);
 	return (Sec.BaseIndex / 3) + Tri;
 }
@@ -702,7 +702,7 @@ FORCEINLINE int32 UNiagaraDataInterfaceSkeletalMesh::GetFilteredTriangleAt<TNDIS
 	{
 		if (Accessor.LODData->RenderSections[i].NumTriangles > (uint32)FilteredIndex)
 		{
-			FSkelMeshRenderSection& Sec = Accessor.LODData->RenderSections[i];
+			const FSkelMeshRenderSection& Sec = Accessor.LODData->RenderSections[i];
 			return Sec.BaseIndex + FilteredIndex;
 		}
 		FilteredIndex -= Accessor.LODData->RenderSections[i].NumTriangles;
@@ -837,14 +837,12 @@ void UNiagaraDataInterfaceSkeletalMesh::GetTriCoordColor(FVectorVMContext& Conte
 	VectorVM::FExternalFuncRegisterHandler<float> OutColorA(Context);
 
 	USkeletalMeshComponent* Comp = Cast<USkeletalMeshComponent>(InstData->Component.Get());
-	FSkinWeightVertexBuffer* SkinWeightBuffer;
-	FSkeletalMeshLODRenderData* LODData = InstData->GetLODRenderDataAndSkinWeights(SkinWeightBuffer);
+	const FSkeletalMeshLODRenderData* LODData = InstData->CachedLODData;
 	check(LODData);
-
 	const FColorVertexBuffer& Colors = LODData->StaticVertexBuffers.ColorVertexBuffer;
 	checkfSlow(Colors.GetNumVertices() != 0, TEXT("Trying to access vertex colors from mesh without any."));
 
-	FMultiSizeIndexContainer& Indices = LODData->MultiSizeIndexContainer;
+	const FMultiSizeIndexContainer& Indices = LODData->MultiSizeIndexContainer;
 	const FRawStaticIndexBuffer16or32Interface* IndexBuffer = Indices.GetIndexBuffer();
 	const int32 TriMax = (IndexBuffer->Num() / 3) - 1;
 	for (int32 i = 0; i < Context.NumInstances; ++i)
@@ -913,12 +911,11 @@ void UNiagaraDataInterfaceSkeletalMesh::GetTriCoordUV(FVectorVMContext& Context)
 	VectorVM::FExternalFuncRegisterHandler<float> OutUVX(Context);	VectorVM::FExternalFuncRegisterHandler<float> OutUVY(Context);
 
 	USkeletalMeshComponent* Comp = Cast<USkeletalMeshComponent>(InstData->Component.Get());
-	FSkinWeightVertexBuffer* SkinWeightBuffer;
-	FSkeletalMeshLODRenderData* LODData = InstData->GetLODRenderDataAndSkinWeights(SkinWeightBuffer);
+	const FSkeletalMeshLODRenderData* LODData = InstData->CachedLODData;
 	check(LODData);
 
-	FMultiSizeIndexContainer& Indices = LODData->MultiSizeIndexContainer;
-	FRawStaticIndexBuffer16or32Interface* IndexBuffer = Indices.GetIndexBuffer();
+	const FMultiSizeIndexContainer& Indices = LODData->MultiSizeIndexContainer;
+	const FRawStaticIndexBuffer16or32Interface* IndexBuffer = Indices.GetIndexBuffer();
 	const int32 TriMax = (IndexBuffer->Num() / 3) - 1;
 	const int32 UVSetMax = LODData->StaticVertexBuffers.StaticMeshVertexBuffer.GetNumTexCoords() - 1;
 	const float InvDt = 1.0f / InstData->DeltaSeconds;
@@ -1070,7 +1067,7 @@ void UNiagaraDataInterfaceSkeletalMesh::GetTriCoordSkinnedData(FVectorVMContext&
 
 	check(Accessor.IsSkinAccessible()); // supposed to use the fallback for invalid mesh
 
-	FSkeletalMeshLODRenderData* LODData = Accessor.LODData;
+	const FSkeletalMeshLODRenderData* LODData = Accessor.LODData;
 
 	const int32 TriMax = (Accessor.IndexBuffer->Num() / 3) - 1;
 	const float InvDt = 1.0f / InstData->DeltaSeconds;
diff --git a/Engine/Plugins/FX/Niagara/Source/Niagara/Private/NDISkeletalMesh_VertexSampling.cpp b/Engine/Plugins/FX/Niagara/Source/Niagara/Private/NDISkeletalMesh_VertexSampling.cpp
--- a/Engine/Plugins/FX/Niagara/Source/Niagara/Private/NDISkeletalMesh_VertexSampling.cpp
+++ b/Engine/Plugins/FX/Niagara/Source/Niagara/Private/NDISkeletalMesh_VertexSampling.cpp
@@ -559,13 +559,12 @@ void UNiagaraDataInterfaceSkeletalMesh::GetVertexColor(FVectorVMContext& Context
 	VectorVM::FExternalFuncRegisterHandler<float> OutColorA(Context);
 
 	USkeletalMeshComponent* Comp = Cast<USkeletalMeshComponent>(InstData->Component.Get());
-	FSkinWeightVertexBuffer* SkinWeightBuffer;
-	FSkeletalMeshLODRenderData* LODData = InstData->GetLODRenderDataAndSkinWeights(SkinWeightBuffer);
+	const FSkeletalMeshLODRenderData* LODData = InstData->CachedLODData;
 	check(LODData);
 	const FColorVertexBuffer& Colors = LODData->StaticVertexBuffers.ColorVertexBuffer;
 	checkfSlow(Colors.GetNumVertices() != 0, TEXT("Trying to access vertex colors from mesh without any."));
 
-	FMultiSizeIndexContainer& Indices = LODData->MultiSizeIndexContainer;
+	const FMultiSizeIndexContainer& Indices = LODData->MultiSizeIndexContainer;
 	const FRawStaticIndexBuffer16or32Interface* IndexBuffer = Indices.GetIndexBuffer();
 	int32 VertMax = LODData->GetNumVertices();
 	for (int32 i = 0; i < Context.NumInstances; ++i)
@@ -616,12 +615,11 @@ void UNiagaraDataInterfaceSkeletalMesh::GetVertexUV(FVectorVMContext& Context)
 	VectorVM::FExternalFuncRegisterHandler<float> OutUVX(Context);	VectorVM::FExternalFuncRegisterHandler<float> OutUVY(Context);
 
 	USkeletalMeshComponent* Comp = Cast<USkeletalMeshComponent>(InstData->Component.Get());
-	FSkinWeightVertexBuffer* SkinWeightBuffer;
-	FSkeletalMeshLODRenderData* LODData = InstData->GetLODRenderDataAndSkinWeights(SkinWeightBuffer);
+	const FSkeletalMeshLODRenderData* LODData = InstData->CachedLODData;
 	check(LODData);
 
-	FMultiSizeIndexContainer& Indices = LODData->MultiSizeIndexContainer;
-	FRawStaticIndexBuffer16or32Interface* IndexBuffer = Indices.GetIndexBuffer();
+	const FMultiSizeIndexContainer& Indices = LODData->MultiSizeIndexContainer;
+	const FRawStaticIndexBuffer16or32Interface* IndexBuffer = Indices.GetIndexBuffer();
 	const int32 VertMax = LODData->GetNumVertices() - 1;
 	for (int32 i = 0; i < Context.NumInstances; ++i)
 	{
diff --git a/Engine/Plugins/FX/Niagara/Source/Niagara/Private/NiagaraDataInterfaceSkeletalMesh.cpp b/Engine/Plugins/FX/Niagara/Source/Niagara/Private/NiagaraDataInterfaceSkeletalMesh.cpp
--- a/Engine/Plugins/FX/Niagara/Source/Niagara/Private/NiagaraDataInterfaceSkeletalMesh.cpp
+++ b/Engine/Plugins/FX/Niagara/Source/Niagara/Private/NiagaraDataInterfaceSkeletalMesh.cpp
@@ -82,13 +82,13 @@ FSkeletalMeshSkinningDataHandle::FSkeletalMeshSkinningDataHandle()
 {
 }
 
-FSkeletalMeshSkinningDataHandle::FSkeletalMeshSkinningDataHandle(FSkeletalMeshSkinningDataUsage InUsage, TSharedPtr<FSkeletalMeshSkinningData> InSkinningData)
+FSkeletalMeshSkinningDataHandle::FSkeletalMeshSkinningDataHandle(FSkeletalMeshSkinningDataUsage InUsage, const TSharedPtr<FSkeletalMeshSkinningData>& InSkinningData, bool bNeedsDataImmediately)
 	: Usage(InUsage)
 	, SkinningData(InSkinningData)
 {
 	if (FSkeletalMeshSkinningData* SkinData = SkinningData.Get())
 	{
-		SkinData->RegisterUser(Usage);
+		SkinData->RegisterUser(Usage, bNeedsDataImmediately);
 	}
 }
 
@@ -125,7 +125,7 @@ void FSkeletalMeshSkinningData::ForceDataRefresh()
 	bForceDataRefresh = true;
 }
 
-void FSkeletalMeshSkinningData::RegisterUser(FSkeletalMeshSkinningDataUsage Usage)
+void FSkeletalMeshSkinningData::RegisterUser(FSkeletalMeshSkinningDataUsage Usage, bool bNeedsDataImmediately)
 {
 	FRWScopeLock Lock(RWGuard, SLT_Write);
 	USkeletalMeshComponent* SkelComp = MeshComp.Get();
@@ -145,9 +145,10 @@ void FSkeletalMeshSkinningData::RegisterUser(FSkeletalMeshSkinningDataUsage Usag
 	if (Usage.NeedPreSkinnedVerts())
 	{
 		++LOD.PreSkinnedVertsUsers;
+		++TotalPreSkinnedVertsUsers;
 	}
 
-	if (Usage.NeedsDataImmediately())
+	if (bNeedsDataImmediately)
 	{
 		check(IsInGameThread());
 		if (CurrBoneRefToLocals().Num() == 0 || CurrComponentTransforms().Num() == 0 )
@@ -196,6 +197,7 @@ void FSkeletalMeshSkinningData::UnregisterUser(FSkeletalMeshSkinningDataUsage Us
 	if (Usage.NeedPreSkinnedVerts())
 	{
 		--LOD.PreSkinnedVertsUsers;
+		--TotalPreSkinnedVertsUsers;
 	}
 }
 
@@ -315,16 +317,22 @@ bool FSkeletalMeshSkinningData::Tick(float InDeltaSeconds, bool bRequirePreskin)
 
 	if (bRequirePreskin)
 	{
-		for (int32 LODIndex = 0; LODIndex < LODData.Num(); ++LODIndex)
+		const USkeletalMesh* SkeletalMesh = SkelComp->SkeletalMesh;
+		check(SkeletalMesh);
+		const FSkeletalMeshRenderData* RenderData = SkeletalMesh->GetResourceForRendering();
+		check(RenderData);
+
+		for (int32 LODIndex = RenderData->PendingFirstLODIdx; LODIndex < LODData.Num(); ++LODIndex)
 		{
 			FLODData& LOD = LODData[LODIndex];
-			if (LOD.PreSkinnedVertsUsers > 0 && SkelComp->SkeletalMesh->GetLODInfo(LODIndex)->bAllowCPUAccess)
+			if (LOD.PreSkinnedVertsUsers > 0 && SkeletalMesh->GetLODInfo(LODIndex)->bAllowCPUAccess)
 			{
+				// Increment ref count to prevent stream out for happening while we are processing the CPU data.
+				TRefCountPtr<const FSkeletalMeshLODRenderData> SkelMeshLODData = &RenderData->LODRenderData[LODIndex];
 				//TODO: If we pass the sections in the usage too, we can probably skin a minimal set of verts just for the used regions.
-				FSkeletalMeshLODRenderData& SkelMeshLODData = SkelComp->SkeletalMesh->GetResourceForRendering()->LODRenderData[LODIndex];
 				FSkinWeightVertexBuffer* SkinWeightBuffer = SkelComp->GetSkinWeightBuffer(LODIndex);
-				USkeletalMeshComponent::ComputeSkinnedPositions(SkelComp, CurrSkinnedPositions(LODIndex), CurrBoneRefToLocals(), SkelMeshLODData, *SkinWeightBuffer);
-				USkeletalMeshComponent::ComputeSkinnedTangentBasis(SkelComp, CurrSkinnedTangentBasis(LODIndex), CurrBoneRefToLocals(), SkelMeshLODData, *SkinWeightBuffer);
+				USkeletalMeshComponent::ComputeSkinnedPositions(SkelComp, CurrSkinnedPositions(LODIndex), CurrBoneRefToLocals(), *SkelMeshLODData, *SkinWeightBuffer);
+				USkeletalMeshComponent::ComputeSkinnedTangentBasis(SkelComp, CurrSkinnedTangentBasis(LODIndex), CurrBoneRefToLocals(), *SkelMeshLODData, *SkinWeightBuffer);
 				//check(CurrSkinnedPositions(LODIndex).Num() == SkelMeshLODData.NumVertices);
 				//Prime the previous positions if they're missing
 				if (PrevSkinnedPositions(LODIndex).Num() != CurrSkinnedPositions(LODIndex).Num())
@@ -341,31 +349,25 @@ bool FSkeletalMeshSkinningData::Tick(float InDeltaSeconds, bool bRequirePreskin)
 
 //////////////////////////////////////////////////////////////////////////
 
-FSkeletalMeshSkinningDataHandle FNDI_SkeletalMesh_GeneratedData::GetCachedSkinningData(TWeakObjectPtr<USkeletalMeshComponent>& Component, FSkeletalMeshSkinningDataUsage Usage)
+FSkeletalMeshSkinningDataHandle FNDI_SkeletalMesh_GeneratedData::GetCachedSkinningData(TWeakObjectPtr<USkeletalMeshComponent>& Component, FSkeletalMeshSkinningDataUsage Usage, bool bNeedsDataImmediately)
 {
 	check(Component.Get() != nullptr);
 
 	// Attempt to Find data
 	{
 		FRWScopeLock ReadLock(CachedSkinningDataGuard, SLT_ReadOnly);
-		if ( CachedSkinningDataAndUsage* Existing = CachedSkinningData.Find(Component) )
+		if (TSharedPtr<FSkeletalMeshSkinningData>* Existing = CachedSkinningData.Find(Component))
 		{
-			check(Existing->SkinningData.IsValid());
-			ensure(Existing->Usage.NeedBoneMatrices() == Usage.NeedBoneMatrices());
-			ensure(Existing->Usage.NeedPreSkinnedVerts() == Usage.NeedPreSkinnedVerts());
-			ensure(Existing->Usage.NeedsDataImmediately() == Usage.NeedsDataImmediately());
-			ensure(Existing->Usage.GetLODIndex() == Usage.GetLODIndex());
-
-			return FSkeletalMeshSkinningDataHandle(Existing->Usage, Existing->SkinningData);
+			return FSkeletalMeshSkinningDataHandle(Usage, *Existing, bNeedsDataImmediately);
 		}
 	}
 
 	// We need to add
 	FRWScopeLock WriteLock(CachedSkinningDataGuard, SLT_Write);
-	CachedSkinningDataAndUsage& NewData = CachedSkinningData.FindOrAdd(Component);
-	NewData.Usage = Usage;
-	NewData.SkinningData = MakeShared<FSkeletalMeshSkinningData>(Component);
-	return FSkeletalMeshSkinningDataHandle(NewData.Usage, NewData.SkinningData);
+	return FSkeletalMeshSkinningDataHandle(
+		Usage, 
+		CachedSkinningData.Add(Component, MakeShared<FSkeletalMeshSkinningData>(Component)),
+		bNeedsDataImmediately);
 }
 
 void FNDI_SkeletalMesh_GeneratedData::TickGeneratedData(ETickingGroup TickGroup, float DeltaSeconds)
@@ -385,23 +387,24 @@ void FNDI_SkeletalMesh_GeneratedData::TickGeneratedData(ETickingGroup TickGroup,
 	ToTickBonesOnly.Reserve(CachedSkinningData.Num());
 	ToTickPreskin.Reserve(CachedSkinningData.Num());
 
-	for (TPair<TWeakObjectPtr<USkeletalMeshComponent>, CachedSkinningDataAndUsage>& Pair : CachedSkinningData)
+	for (TPair<TWeakObjectPtr<USkeletalMeshComponent>, TSharedPtr<FSkeletalMeshSkinningData>>& Pair : CachedSkinningData)
 	{
+		USkeletalMeshComponent* Component = Pair.Key.Get();
+		TSharedPtr<FSkeletalMeshSkinningData>& SkinningData = Pair.Value;
+
 		if ( TickGroup == NiagaraFirstTickGroup )
 		{
-			Pair.Value.bHasTicked = false;
+			SkinningData->bHasTicked = false;
 		}
 
 		// Should remove?
-		TSharedPtr<FSkeletalMeshSkinningData>& SkinningData = Pair.Value.SkinningData;
-		USkeletalMeshComponent* Component = Pair.Key.Get();
 		if ( (Component == nullptr) || SkinningData.IsUnique() || !SkinningData->IsUsed() )
 		{
 			ToRemove.Add(Pair.Key);
 			continue;
 		}
 
-		if ( Pair.Value.bHasTicked == true )
+		if ( SkinningData->bHasTicked == true )
 		{
 			continue;
 		}
@@ -417,13 +420,12 @@ void FNDI_SkeletalMesh_GeneratedData::TickGeneratedData(ETickingGroup TickGroup,
 		}
 
 		// We are going to tick this one
-		Pair.Value.bHasTicked = true;
+		SkinningData->bHasTicked = true;
 
-		FSkeletalMeshSkinningDataUsage Usage = Pair.Value.Usage;
 		FSkeletalMeshSkinningData* SkinningDataPtr = SkinningData.Get();
 		check(SkinningDataPtr);
 
-		if (Usage.NeedPreSkinnedVerts())
+		if (SkinningDataPtr->NeedPreSkinnedVerts())
 		{
 			ToTickPreskin.Add(SkinningDataPtr);
 		}
@@ -1418,21 +1420,29 @@ USkeletalMesh* UNiagaraDataInterfaceSkeletalMesh::GetSkeletalMesh(UNiagaraCompon
 				FoundSkelComp = OuterComp;
 				Mesh = OuterComp->SkeletalMesh;
 			}
-			else if (AActor* Owner = SimComp->GetAttachmentRootActor())
+			else
 			{
-				for (UActorComponent* ActorComp : Owner->GetComponents())
+				AActor* Owner = SimComp->GetAttachmentRootActor();
+				while (Owner && !Mesh)
 				{
-					USkeletalMeshComponent* SourceComp = Cast<USkeletalMeshComponent>(ActorComp);
-					if (SourceComp)
+					for (UActorComponent* ActorComp : Owner->GetComponents())
 					{
-						USkeletalMesh* PossibleMesh = SourceComp->SkeletalMesh;
-						if (PossibleMesh != nullptr/* && PossibleMesh->bAllowCPUAccess*/)
+						USkeletalMeshComponent* SourceComp = Cast<USkeletalMeshComponent>(ActorComp);
+						if (SourceComp)
 						{
-							Mesh = PossibleMesh;
-							FoundSkelComp = SourceComp;
-							break;
+							USkeletalMesh* PossibleMesh = SourceComp->SkeletalMesh;
+							if (PossibleMesh != nullptr/* && PossibleMesh->bAllowCPUAccess*/)
+							{
+								Mesh = PossibleMesh;
+								FoundSkelComp = SourceComp;
+
+								break;
+							}
 						}
 					}
+
+					// Iterate on the actor hierarchy.
+					Owner = Owner->GetParentActor();
 				}
 			}
 
@@ -1506,137 +1516,66 @@ bool FNDISkeletalMesh_InstanceData::Init(UNiagaraDataInterfaceSkeletalMesh* Inte
 	}
 #endif
 
-	// Get the first LOD level we can guarantee to be in memory
-	//-TODO: Support skeletal mesh streaming
-	int32 FirstInlineLODLevel = 0;
-	if (Mesh != nullptr)
-	{
-		if (const FSkeletalMeshRenderData* SkelMeshRenderData = Mesh->GetResourceForRendering())
-		{
-			FirstInlineLODLevel = SkelMeshRenderData->LODRenderData.Num() - SkelMeshRenderData->NumInlinedLODs;
-		}
-	}
+	MinLODIdx = 0;
+	bResetOnLODStreamedIn = false;
+	CachedLODIdx = 0;
+	CachedLODData.SafeRelease();
 
 	//Setup where to spawn from
 	SamplingRegionIndices.Empty();
 	bool bAllRegionsAreAreaWeighting = true;
-	int32 LODIndex = INDEX_NONE;
-	if (Mesh == nullptr)
-	{
-		// Just say we're sampling LOD 0, even though there are no LODs
-		LODIndex = 0;
-	}
-	else if (Interface->SamplingRegions.Num() == 0)
-	{
-		//If we have no regions, sample the whole mesh at the specified LOD.
-		LODIndex = Interface->WholeMeshLOD;
-		if (LODIndex == INDEX_NONE)
-		{
-			LODIndex = Mesh->GetLODNum() - 1;
-		}
-		else
-		{
-			LODIndex = FMath::Clamp(LODIndex, FirstInlineLODLevel, Mesh->GetLODNum() - 1);
-		}
-	}
-	else
+
+	if (Mesh)
 	{
-		//Sampling from regions. Gather the indices of the regions we'll sample from.
-		const FSkeletalMeshSamplingInfo& SamplingInfo = Mesh->GetSamplingInfo();
-		for (FName RegionName : Interface->SamplingRegions)
+		MinLODIdx = Mesh->MinLod.GetValueForFeatureLevel(SystemInstance->GetFeatureLevel());
+		const int32 PendingFirstLODIndex = Mesh->GetResourceForRendering()->GetPendingFirstLODIdx(MinLODIdx);
+
+		const int32 DesiredLODIndex = Interface->CalculateLODIndexAndSamplingRegions(Mesh, SamplingRegionIndices, bAllRegionsAreAreaWeighting);
+		if (DesiredLODIndex != INDEX_NONE)
 		{
-			int32 RegionIdx = SamplingInfo.IndexOfRegion(RegionName);
-			if (RegionIdx != INDEX_NONE)
+			if (DesiredLODIndex >= PendingFirstLODIndex)
 			{
-				const FSkeletalMeshSamplingRegion& Region = SamplingInfo.GetRegion(RegionIdx);
-				const FSkeletalMeshSamplingRegionBuiltData& RegionBuiltData = SamplingInfo.GetRegionBuiltData(RegionIdx);
-				int32 RegionLODIndex = Region.LODIndex;
-				if (RegionLODIndex == INDEX_NONE)
-				{
-					RegionLODIndex = Mesh->GetLODInfoArray().Num() - 1;
-				}
-				else
-				{
-					RegionLODIndex = FMath::Clamp(RegionLODIndex, 0, Mesh->GetLODInfoArray().Num() - 1);
-				}
-
-				if (LODIndex == INDEX_NONE)
-				{
-					LODIndex = RegionLODIndex;
-				}
-
-				// Ensure we aren't using a region from a streaming / culled LOD level
-				if (LODIndex < FirstInlineLODLevel)
-				{
-					UE_LOG(LogNiagara, Warning, TEXT("Skeletal Mesh Data Interface is trying to use regions on LODs levels that are either streamed or cooked out. This is currently unsupported.\nInterface: %s\nMesh: %s\nRegion: %s"),
-						*Interface->GetFullName(),
-						*Mesh->GetFullName(),
-						*RegionName.ToString());
-
-					return false;
-				}
-
-				//ensure we don't try to use two regions from different LODs.
-				if (LODIndex != RegionLODIndex)
-				{
-					UE_LOG(LogNiagara, Warning, TEXT("Skeletal Mesh Data Interface is trying to use regions on different LODs of the mesh. This is currently unsupported.\nInterface: %s\nMesh: %s\nRegion: %s"),
-						*Interface->GetFullName(),
-						*Mesh->GetFullName(),
-						*RegionName.ToString());
-
-					return false;
-				}
-
-				if (RegionBuiltData.TriangleIndices.Num() > 0)
-				{
-					SamplingRegionIndices.Add(RegionIdx);
-					bAllRegionsAreAreaWeighting &= Region.bSupportUniformlyDistributedSampling;
-				}
-				else
-				{
-					UE_LOG(LogNiagara, Warning, TEXT("Skeletal Mesh Data Interface is trying to use a region with no associated triangles.\nLOD: %d\nInterface: %s\nMesh: %s\nRegion: %s"),
-						LODIndex,
-						*Interface->GetFullName(),
-						*Mesh->GetFullName(),
-						*RegionName.ToString());
-
-					return false;
-				}
+				CachedLODIdx = DesiredLODIndex;
 			}
 			else
 			{
-				UE_LOG(LogNiagara, Warning, TEXT("Skeletal Mesh Data Interface is trying to use a region on a mesh that does not provide this region.\nInterface: %s\nMesh: %s\nRegion: %s"),
-					*Interface->GetFullName(),
-					*Mesh->GetFullName(),
-					*RegionName.ToString());
-
-				return false;
+				CachedLODIdx = PendingFirstLODIndex;
+				bResetOnLODStreamedIn = true;
 			}
+
+			CachedLODData = &Mesh->GetResourceForRendering()->LODRenderData[CachedLODIdx];
+		}
+		else
+		{
+			return false;
 		}
 	}
 
-	// TODO: This change is temporary to work around a crash that happens when you change the
-	// source mesh on a system which is running in the level from the details panel.
-	// bool bNeedDataImmediately = SystemInstance->IsSolo();
-	bool bNeedDataImmediately = true;
-		
+	check(CachedLODIdx >= 0);
+
 	//Grab a handle to the skinning data if we have a component to skin.
 	const ENDISkeletalMesh_SkinningMode SkinningMode = Interface->SkinningMode;
 	FSkeletalMeshSkinningDataUsage Usage(
-		LODIndex,
+		CachedLODIdx,
 		SkinningMode == ENDISkeletalMesh_SkinningMode::SkinOnTheFly || SkinningMode == ENDISkeletalMesh_SkinningMode::PreSkin,
-		SkinningMode == ENDISkeletalMesh_SkinningMode::PreSkin,
-		bNeedDataImmediately);
+		SkinningMode == ENDISkeletalMesh_SkinningMode::PreSkin
+	);
 
+	// GetSkeletalMeshGeneratedData() is meant to match with the required lod, so don't access unless we are using it.
 	if (NewSkelComp)
 	{
+		// TODO: This change is temporary to work around a crash that happens when you change the
+		// source mesh on a system which is running in the level from the details panel.
+		// bool bNeedsDataImmediately = SystemInstance->IsSolo();
+		const bool bNeedsDataImmediately = true;
+
 		TWeakObjectPtr<USkeletalMeshComponent> SkelWeakCompPtr = NewSkelComp;
 		FNDI_SkeletalMesh_GeneratedData& GeneratedData = SystemInstance->GetWorldManager()->GetSkeletalMeshGeneratedData();
-		SkinningData = GeneratedData.GetCachedSkinningData(SkelWeakCompPtr, Usage);
+		SkinningData = GeneratedData.GetCachedSkinningData(SkelWeakCompPtr, Usage, bNeedsDataImmediately);
 	}
 	else
 	{
-		SkinningData = FSkeletalMeshSkinningDataHandle(Usage, nullptr);
+		SkinningData = FSkeletalMeshSkinningDataHandle(Usage, nullptr, false);
 	}
 
 	//Init area weighting sampler for Sampling regions.
@@ -1646,24 +1585,23 @@ bool FNDISkeletalMesh_InstanceData::Init(UNiagaraDataInterfaceSkeletalMesh* Inte
 		SamplingRegionAreaWeightedSampler.Init(this);
 	}
 
-
 	if (Mesh)
 	{
+		check(CachedLODData);
+
 		bAllowCPUMeshDataAccess = true; // Assume accessibility until proven otherwise below
-		FSkinWeightVertexBuffer* SkinWeightBuffer = nullptr;
-		FSkeletalMeshLODRenderData* LODData = GetLODRenderDataAndSkinWeights(SkinWeightBuffer);
-		check(LODData);
+		const FSkinWeightVertexBuffer* SkinWeightBuffer = GetSkinWeights();
 		check(SkinWeightBuffer);
 
 		// Check for the validity of the Mesh's cpu data.
-		if (Mesh->GetLODInfo(LODIndex)->bAllowCPUAccess)
+		if (Mesh->GetLODInfo(CachedLODIdx)->bAllowCPUAccess)
 		{
-			const bool LODDataNumVerticesCorrect = LODData->GetNumVertices() > 0;
-			const bool LODDataPositonNumVerticesCorrect = LODData->StaticVertexBuffers.PositionVertexBuffer.GetNumVertices() > 0;
+			const bool LODDataNumVerticesCorrect = CachedLODData->GetNumVertices() > 0;
+			const bool LODDataPositonNumVerticesCorrect = CachedLODData->StaticVertexBuffers.PositionVertexBuffer.GetNumVertices() > 0;
 			const bool SkinWeightBufferNumVerticesCorrect = SkinWeightBuffer->GetNumVertices() > 0;
-			const bool bIndexBufferValid = LODData->MultiSizeIndexContainer.IsIndexBufferValid();
-			const bool bIndexBufferFound = bIndexBufferValid && (LODData->MultiSizeIndexContainer.GetIndexBuffer() != nullptr);
-			const bool bIndexBufferNumCorrect = bIndexBufferFound && (LODData->MultiSizeIndexContainer.GetIndexBuffer()->Num() > 0);
+			const bool bIndexBufferValid = CachedLODData->MultiSizeIndexContainer.IsIndexBufferValid();
+			const bool bIndexBufferFound = bIndexBufferValid && (CachedLODData->MultiSizeIndexContainer.GetIndexBuffer() != nullptr);
+			const bool bIndexBufferNumCorrect = bIndexBufferFound && (CachedLODData->MultiSizeIndexContainer.GetIndexBuffer()->Num() > 0);
 
 			bAllowCPUMeshDataAccess = LODDataNumVerticesCorrect &&
 				LODDataPositonNumVerticesCorrect &&
@@ -1823,25 +1761,25 @@ bool FNDISkeletalMesh_InstanceData::Init(UNiagaraDataInterfaceSkeletalMesh* Inte
 			//check(MeshSkinWeightBufferSrv->IsValid()); // not available in this stream
 			MeshSkinWeightLookupBuffer = SkinWeightBuffer->GetLookupVertexBuffer();
 
-			FSkeletalMeshLODInfo* LODInfo = Mesh->GetLODInfo(LODIndex);
+			FSkeletalMeshLODInfo* LODInfo = Mesh->GetLODInfo(CachedLODIdx);
 			bIsGpuUniformlyDistributedSampling = LODInfo->bSupportUniformlyDistributedSampling && bAllRegionsAreAreaWeighting;
 
 			if (Mesh->HasActiveClothingAssets())
 			{
 				UE_LOG(LogNiagara, Warning, TEXT("Skeletal Mesh %s has cloth asset on it: spawning from it might not work properly."), *Mesh->GetName());
 			}
-			if (LODData->GetVertexBufferMaxBoneInfluences() > MAX_INFLUENCES_PER_STREAM)
+			if (CachedLODData->GetVertexBufferMaxBoneInfluences() > MAX_INFLUENCES_PER_STREAM)
 			{
 				UE_LOG(LogNiagara, Warning, TEXT("Skeletal Mesh %s has bones extra influence: spawning from it might not work properly."), *Mesh->GetName());
 			}
 
 			const FSkeletalMeshSamplingInfo& SamplingInfo = Mesh->GetSamplingInfo();
 			MeshGpuSpawnStaticBuffers = new FSkeletalMeshGpuSpawnStaticBuffers();
-			MeshGpuSpawnStaticBuffers->Initialise(this, *LODData, SamplingInfo.GetBuiltData().WholeMeshBuiltData[LODIndex]);
+			MeshGpuSpawnStaticBuffers->Initialise(this, *CachedLODData, SamplingInfo.GetBuiltData().WholeMeshBuiltData[CachedLODIdx]);
 			BeginInitResource(MeshGpuSpawnStaticBuffers);
 
 			MeshGpuSpawnDynamicBuffers = new FSkeletalMeshGpuDynamicBufferProxy();
-			MeshGpuSpawnDynamicBuffers->Initialise(RefSkel, *LODData, FilteredSocketInfo.Num());
+			MeshGpuSpawnDynamicBuffers->Initialise(RefSkel, *CachedLODData, FilteredSocketInfo.Num());
 			BeginInitResource(MeshGpuSpawnDynamicBuffers);
 		}
 	}
@@ -1864,6 +1802,16 @@ bool FNDISkeletalMesh_InstanceData::ResetRequired(UNiagaraDataInterfaceSkeletalM
 		return true;
 	}
 	
+	// Reset if the LOD we relied on was streamed out, of if the lod we need could now be available.
+	if (Mesh)
+	{
+		const int32 PendingFirstLODIndex = Mesh->GetResourceForRendering()->GetPendingFirstLODIdx(MinLODIdx);
+		if (PendingFirstLODIndex > CachedLODIdx || (PendingFirstLODIndex < CachedLODIdx && bResetOnLODStreamedIn))
+		{
+			return true;
+		}
+	}
+
 	if (USkeletalMeshComponent* SkelComp = Cast<USkeletalMeshComponent>(Comp))
 	{
 		if (!SkelComp->SkeletalMesh)//TODO: Handle clearing the mesh gracefully.
@@ -1960,10 +1908,7 @@ void FNDISkeletalMesh_InstanceData::UpdateFilteredSocketTransforms()
 
 bool FNDISkeletalMesh_InstanceData::HasColorData()
 {
-	FSkinWeightVertexBuffer* SkinWeightBuffer;
-	FSkeletalMeshLODRenderData* LODData = GetLODRenderDataAndSkinWeights(SkinWeightBuffer);
-
-	return LODData && LODData->StaticVertexBuffers.ColorVertexBuffer.GetNumVertices() != 0;
+	return CachedLODData && CachedLODData->StaticVertexBuffers.ColorVertexBuffer.GetNumVertices() != 0;
 }
 
 void FNDISkeletalMesh_InstanceData::Release()
@@ -2889,6 +2834,85 @@ ETickingGroup UNiagaraDataInterfaceSkeletalMesh::CalculateTickGroup(void* PerIns
 	return NiagaraFirstTickGroup;
 }
 
+
+int32 UNiagaraDataInterfaceSkeletalMesh::CalculateLODIndexAndSamplingRegions(USkeletalMesh* InMesh, TArray<int32>& OutSamplingRegionIndices, bool& OutAllRegionsAreAreaWeighting) const
+{
+	check(InMesh);
+
+	if (!SamplingRegions.Num())
+	{
+		//If we have no regions, sample the whole mesh at the specified LOD.
+		if (WholeMeshLOD == INDEX_NONE)
+		{
+			return InMesh->GetLODNum() - 1;
+		}
+		else
+		{
+			return FMath::Clamp(WholeMeshLOD, 0, InMesh->GetLODNum() - 1);
+		}
+	}
+	else
+	{
+		int32 LastRegionLODIndex = INDEX_NONE;
+
+		//Sampling from regions. Gather the indices of the regions we'll sample from.
+		const FSkeletalMeshSamplingInfo& SamplingInfo = InMesh->GetSamplingInfo();
+		for (FName RegionName : SamplingRegions)
+		{
+			const int32 RegionIdx = SamplingInfo.IndexOfRegion(RegionName);
+			if (RegionIdx != INDEX_NONE)
+			{
+				const FSkeletalMeshSamplingRegion& Region = SamplingInfo.GetRegion(RegionIdx);
+				const FSkeletalMeshSamplingRegionBuiltData& RegionBuiltData = SamplingInfo.GetRegionBuiltData(RegionIdx);
+
+				const int32 RegionLODIndex = Region.LODIndex == INDEX_NONE ? 
+					(InMesh->GetLODNum() - 1) :
+					(FMath::Clamp(Region.LODIndex, 0, InMesh->GetLODNum() - 1));
+
+				if (LastRegionLODIndex == INDEX_NONE)
+				{
+					LastRegionLODIndex = RegionLODIndex;
+				}
+				else if (RegionLODIndex != LastRegionLODIndex)
+				{
+					UE_LOG(LogNiagara, Warning, TEXT("Skeletal Mesh Data Interface is trying to use regions on LODs levels that are either streamed or cooked out. This is currently unsupported.\nInterface: %s\nMesh: %s\nRegion: %s"),
+						*GetFullName(),
+						*InMesh->GetFullName(),
+						*RegionName.ToString());
+					return INDEX_NONE;
+				}
+				
+				if (RegionBuiltData.TriangleIndices.Num() > 0)
+				{
+					OutSamplingRegionIndices.Add(RegionIdx);
+					OutAllRegionsAreAreaWeighting &= Region.bSupportUniformlyDistributedSampling;
+				}
+				else
+				{
+					UE_LOG(LogNiagara, Warning, TEXT("Skeletal Mesh Data Interface is trying to use a region with no associated triangles.\nLOD: %d\nInterface: %s\nMesh: %s\nRegion: %s"),
+						 RegionLODIndex,
+						*GetFullName(),
+						*InMesh->GetFullName(),
+						*RegionName.ToString());
+
+					return INDEX_NONE;
+				}
+			}
+			else
+			{
+				UE_LOG(LogNiagara, Warning, TEXT("Skeletal Mesh Data Interface is trying to use a region on a mesh that does not provide this region.\nInterface: %s\nMesh: %s\nRegion: %s"),
+					*GetFullName(),
+					*InMesh->GetFullName(),
+					*RegionName.ToString());
+
+				return INDEX_NONE;
+			}
+		}
+		return LastRegionLODIndex;
+	}
+}
+
+
 //UNiagaraDataInterfaceSkeletalMesh END
 //////////////////////////////////////////////////////////////////////////
 
@@ -2897,7 +2921,8 @@ void FSkeletalMeshAccessorHelper::Init<TNDISkelMesh_FilterModeSingle, TNDISkelMe
 {
 	Comp = Cast<USkeletalMeshComponent>(InstData->Component.Get());
 	Mesh = InstData->Mesh;
-	LODData = InstData->GetLODRenderDataAndSkinWeights(SkinWeightBuffer);
+	LODData = InstData->CachedLODData;
+	SkinWeightBuffer = InstData->GetSkinWeights();
 	IndexBuffer = LODData ? LODData->MultiSizeIndexContainer.GetIndexBuffer() : nullptr;
 	SkinningData = InstData->SkinningData.SkinningData.Get();
 	Usage = InstData->SkinningData.Usage;
@@ -2925,7 +2950,8 @@ void FSkeletalMeshAccessorHelper::Init<TNDISkelMesh_FilterModeSingle, TNDISkelMe
 {
 	Comp = Cast<USkeletalMeshComponent>(InstData->Component.Get());
 	Mesh = InstData->Mesh;
-	LODData = InstData->GetLODRenderDataAndSkinWeights(SkinWeightBuffer);
+	LODData = InstData->CachedLODData;
+	SkinWeightBuffer = InstData->GetSkinWeights();
 	IndexBuffer = LODData ? LODData->MultiSizeIndexContainer.GetIndexBuffer() : nullptr;
 	SkinningData = InstData->SkinningData.SkinningData.Get();
 	Usage = InstData->SkinningData.Usage;
diff --git a/Engine/Source/Runtime/Engine/Classes/Components/SkeletalMeshComponent.h b/Engine/Source/Runtime/Engine/Classes/Components/SkeletalMeshComponent.h
--- a/Engine/Source/Runtime/Engine/Classes/Components/SkeletalMeshComponent.h
+++ b/Engine/Source/Runtime/Engine/Classes/Components/SkeletalMeshComponent.h
@@ -1698,12 +1698,12 @@ public:
 	virtual void SetPhysicsAsset(class UPhysicsAsset* NewPhysicsAsset,bool bForceReInit = false) override;
 	virtual void SetSkeletalMesh(class USkeletalMesh* NewMesh, bool bReinitPose = true) override;
 
-	static FVector GetSkinnedVertexPosition(USkeletalMeshComponent* Component, int32 VertexIndex, const FSkeletalMeshLODRenderData& Model, FSkinWeightVertexBuffer& SkinWeightBuffer);
-	static FVector GetSkinnedVertexPosition(USkeletalMeshComponent* Component, int32 VertexIndex, const FSkeletalMeshLODRenderData& Model, FSkinWeightVertexBuffer& SkinWeightBuffer, TArray<FMatrix>& CachedRefToLocals);
-	static void ComputeSkinnedPositions(USkeletalMeshComponent* Component, TArray<FVector> & OutPositions, TArray<FMatrix>& CachedRefToLocals, const FSkeletalMeshLODRenderData& Model, FSkinWeightVertexBuffer& SkinWeightBuffer);
+	static FVector GetSkinnedVertexPosition(USkeletalMeshComponent* Component, int32 VertexIndex, const FSkeletalMeshLODRenderData& Model, const FSkinWeightVertexBuffer& SkinWeightBuffer);
+	static FVector GetSkinnedVertexPosition(USkeletalMeshComponent* Component, int32 VertexIndex, const FSkeletalMeshLODRenderData& Model, const FSkinWeightVertexBuffer& SkinWeightBuffer, TArray<FMatrix>& CachedRefToLocals);
+	static void ComputeSkinnedPositions(USkeletalMeshComponent* Component, TArray<FVector> & OutPositions, TArray<FMatrix>& CachedRefToLocals, const FSkeletalMeshLODRenderData& Model, const FSkinWeightVertexBuffer& SkinWeightBuffer);
 
-	static void GetSkinnedTangentBasis(USkeletalMeshComponent* Component, int32 VertexIndex, const FSkeletalMeshLODRenderData& Model, FSkinWeightVertexBuffer& SkinWeightBuffer, TArray<FMatrix>& CachedRefToLocals, FVector& OutTangentX, FVector& OutTangentZ);
-	static void ComputeSkinnedTangentBasis(USkeletalMeshComponent* Component, TArray<FVector>& OutTangenXZ, TArray<FMatrix>& CachedRefToLocals, const FSkeletalMeshLODRenderData& Model, FSkinWeightVertexBuffer& SkinWeightBuffer);
+	static void GetSkinnedTangentBasis(USkeletalMeshComponent* Component, int32 VertexIndex, const FSkeletalMeshLODRenderData& Model, const FSkinWeightVertexBuffer& SkinWeightBuffer, TArray<FMatrix>& CachedRefToLocals, FVector& OutTangentX, FVector& OutTangentZ);
+	static void ComputeSkinnedTangentBasis(USkeletalMeshComponent* Component, TArray<FVector>& OutTangenXZ, TArray<FMatrix>& CachedRefToLocals, const FSkeletalMeshLODRenderData& Model, const FSkinWeightVertexBuffer& SkinWeightBuffer);
 
 	void SetSkeletalMeshWithoutResettingAnimation(class USkeletalMesh* NewMesh);
 
diff --git a/Engine/Source/Runtime/Engine/Private/SkeletalMeshComponentPhysics.cpp b/Engine/Source/Runtime/Engine/Private/SkeletalMeshComponentPhysics.cpp
--- a/Engine/Source/Runtime/Engine/Private/SkeletalMeshComponentPhysics.cpp
+++ b/Engine/Source/Runtime/Engine/Private/SkeletalMeshComponentPhysics.cpp
@@ -1893,7 +1893,7 @@ void USkeletalMeshComponent::UpdatePhysicsToRBChannels()
 }
 
 template<bool bCachedMatrices>
-FVector GetTypedSkinnedVertexPositionWithCloth(USkeletalMeshComponent* Component, int32 VertexIndex, const FSkeletalMeshLODRenderData& LODData, FSkinWeightVertexBuffer& SkinWeightBuffer, TArray<FMatrix>& CachedRefToLocals)
+FVector GetTypedSkinnedVertexPositionWithCloth(USkeletalMeshComponent* Component, int32 VertexIndex, const FSkeletalMeshLODRenderData& LODData, const FSkinWeightVertexBuffer& SkinWeightBuffer, TArray<FMatrix>& CachedRefToLocals)
 {
 	// Find the chunk and vertex within that chunk, and skinning type, for this vertex.
 	int32 SectionIndex;
@@ -1944,18 +1944,18 @@ FVector GetTypedSkinnedVertexPositionWithCloth(USkeletalMeshComponent* Component
 	return GetTypedSkinnedVertexPosition<bCachedMatrices>(Component, Section, LODData.StaticVertexBuffers.PositionVertexBuffer, SkinWeightBuffer, VertIndexInChunk, CachedRefToLocals);
 }
 
-FVector USkeletalMeshComponent::GetSkinnedVertexPosition(USkeletalMeshComponent* Component, int32 VertexIndex, const FSkeletalMeshLODRenderData& LODData, FSkinWeightVertexBuffer& SkinWeightBuffer)
+FVector USkeletalMeshComponent::GetSkinnedVertexPosition(USkeletalMeshComponent* Component, int32 VertexIndex, const FSkeletalMeshLODRenderData& LODData, const FSkinWeightVertexBuffer& SkinWeightBuffer)
 {
 	TArray<FMatrix> Dummy;
 	return GetTypedSkinnedVertexPositionWithCloth<false>(Component, VertexIndex, LODData, SkinWeightBuffer, Dummy);
 }
 
-FVector USkeletalMeshComponent::GetSkinnedVertexPosition(USkeletalMeshComponent* Component, int32 VertexIndex, const FSkeletalMeshLODRenderData& LODData, FSkinWeightVertexBuffer& SkinWeightBuffer, TArray<FMatrix>& CachedRefToLocals)
+FVector USkeletalMeshComponent::GetSkinnedVertexPosition(USkeletalMeshComponent* Component, int32 VertexIndex, const FSkeletalMeshLODRenderData& LODData, const FSkinWeightVertexBuffer& SkinWeightBuffer, TArray<FMatrix>& CachedRefToLocals)
 {
 	return GetTypedSkinnedVertexPositionWithCloth<true>(Component, VertexIndex, LODData, SkinWeightBuffer, CachedRefToLocals);
 }
 
-void USkeletalMeshComponent::ComputeSkinnedPositions(USkeletalMeshComponent* Component, TArray<FVector> & OutPositions, TArray<FMatrix>& CachedRefToLocals, const FSkeletalMeshLODRenderData& LODData, FSkinWeightVertexBuffer& SkinWeightBuffer)
+void USkeletalMeshComponent::ComputeSkinnedPositions(USkeletalMeshComponent* Component, TArray<FVector> & OutPositions, TArray<FMatrix>& CachedRefToLocals, const FSkeletalMeshLODRenderData& LODData, const FSkinWeightVertexBuffer& SkinWeightBuffer)
 {
 	// Fail if no mesh
 	if (!Component->SkeletalMesh)
@@ -2053,7 +2053,7 @@ void USkeletalMeshComponent::ComputeSkinnedPositions(USkeletalMeshComponent* Com
 	}
 }
 
-void USkeletalMeshComponent::GetSkinnedTangentBasis(USkeletalMeshComponent* Component, int32 VertexIndex, const FSkeletalMeshLODRenderData& LODData, FSkinWeightVertexBuffer& SkinWeightBuffer, TArray<FMatrix>& CachedRefToLocals, FVector& OutTangentX, FVector& OutTangentZ)
+void USkeletalMeshComponent::GetSkinnedTangentBasis(USkeletalMeshComponent* Component, int32 VertexIndex, const FSkeletalMeshLODRenderData& LODData, const FSkinWeightVertexBuffer& SkinWeightBuffer, TArray<FMatrix>& CachedRefToLocals, FVector& OutTangentX, FVector& OutTangentZ)
 {
 	int32 SectionIndex;
 	int32 VertIndexInChunk;
@@ -2065,7 +2065,7 @@ void USkeletalMeshComponent::GetSkinnedTangentBasis(USkeletalMeshComponent* Comp
 	return GetTypedSkinnedTangentBasis(Component, Section, LODData.StaticVertexBuffers, SkinWeightBuffer, VertIndexInChunk, CachedRefToLocals, OutTangentX, OutTangentZ);
 }
 
-void USkeletalMeshComponent::ComputeSkinnedTangentBasis(USkeletalMeshComponent* Component, TArray<FVector>& OutTangenXZ, TArray<FMatrix>& CachedRefToLocals, const FSkeletalMeshLODRenderData& LODData, FSkinWeightVertexBuffer& SkinWeightBuffer)
+void USkeletalMeshComponent::ComputeSkinnedTangentBasis(USkeletalMeshComponent* Component, TArray<FVector>& OutTangenXZ, TArray<FMatrix>& CachedRefToLocals, const FSkeletalMeshLODRenderData& LODData, const FSkinWeightVertexBuffer& SkinWeightBuffer)
 {
 	// Fail if no mesh
 	if (!Component->SkeletalMesh)
diff --git a/Engine/Source/Runtime/Engine/Private/SkeletalMeshRenderData.cpp b/Engine/Source/Runtime/Engine/Private/SkeletalMeshRenderData.cpp
--- a/Engine/Source/Runtime/Engine/Private/SkeletalMeshRenderData.cpp
+++ b/Engine/Source/Runtime/Engine/Private/SkeletalMeshRenderData.cpp
@@ -384,6 +384,19 @@ FSkeletalMeshRenderData::FSkeletalMeshRenderData()
 	, bInitialized(false)
 {}
 
+FSkeletalMeshRenderData::~FSkeletalMeshRenderData()
+{
+	FSkeletalMeshLODRenderData** LODRenderDataArray = LODRenderData.GetData();
+	for (int32 LODIndex = 0; LODIndex < LODRenderData.Num(); ++LODIndex)
+	{
+		LODRenderDataArray[LODIndex]->Release();
+		// Prevent the array from calling the destructor to handle correctly the refcount.
+		// For compatibility reason, LODRenderDataArray is using ptr directly instead of TRefCountPtr.
+		LODRenderDataArray[LODIndex] = nullptr;
+	}
+	LODRenderData.Empty();
+}
+
 void FSkeletalMeshRenderData::Serialize(FArchive& Ar, USkeletalMesh* Owner)
 {
 	DECLARE_SCOPE_CYCLE_COUNTER(TEXT("FSkeletalMeshRenderData::Serialize"), STAT_SkeletalMeshRenderData_Serialize, STATGROUP_LoadTime);
diff --git a/Engine/Source/Runtime/Engine/Private/Streaming/SkeletalMeshUpdate.cpp b/Engine/Source/Runtime/Engine/Private/Streaming/SkeletalMeshUpdate.cpp
--- a/Engine/Source/Runtime/Engine/Private/Streaming/SkeletalMeshUpdate.cpp
+++ b/Engine/Source/Runtime/Engine/Private/Streaming/SkeletalMeshUpdate.cpp
@@ -15,6 +15,8 @@ SkeletalMeshUpdate.cpp: Helpers to stream in and out skeletal mesh LODs.
 #include "Components/SkinnedMeshComponent.h"
 #include "Streaming/RenderAssetUpdate.inl"
 
+extern int32 GStreamingMaxReferenceChecks;
+
 template class TRenderAssetUpdate<FSkelMeshUpdateContext>;
 
 static constexpr uint32 GSkelMeshMaxNumResourceUpdatesPerLOD = 16;
@@ -239,12 +241,12 @@ void FSkeletalMeshStreamIn::DoCancel(const FContext& Context)
 FSkeletalMeshStreamOut::FSkeletalMeshStreamOut(USkeletalMesh* InMesh, int32 InRequestedMips)
 	: FSkeletalMeshUpdate(InMesh, InRequestedMips)
 {
-	PushTask(FContext(InMesh, TT_None), TT_GameThread, SRA_UPDATE_CALLBACK(DoConditionalMarkComponentsDirty), TT_None, nullptr);
+	PushTask(FContext(InMesh, TT_None), TT_GameThread, SRA_UPDATE_CALLBACK(ConditionalMarkComponentsDirty), TT_None, nullptr);
 }
 
-void FSkeletalMeshStreamOut::DoConditionalMarkComponentsDirty(const FContext& Context)
+void FSkeletalMeshStreamOut::ConditionalMarkComponentsDirty(const FContext& Context)
 {
-	QUICK_SCOPE_CYCLE_COUNTER(STAT_FSkeletalMeshStreamOut_DoConditionalMarkComponentsDirty);
+	DECLARE_SCOPE_CYCLE_COUNTER(TEXT("FSkeletalMeshStreamOut::ConditionalMarkComponentsDirty"), STAT_SkeletalMeshStreamOut_ConditionalMarkComponentsDirty, STATGROUP_StreamingDetails);
 	CSV_SCOPED_TIMING_STAT_GLOBAL(SkStreamingMarkDirtyTime);
 	check(Context.CurrentThread == TT_GameThread);
 
@@ -275,11 +277,53 @@ void FSkeletalMeshStreamOut::DoConditionalMarkComponentsDirty(const FContext& Co
 	{
 		Abort();
 	}
-	PushTask(Context, TT_Render, SRA_UPDATE_CALLBACK(DoReleaseBuffers), (EThreadType)Context.CurrentThread, SRA_UPDATE_CALLBACK(DoCancel));
+	PushTask(Context, TT_Async, SRA_UPDATE_CALLBACK(WaitForReferences), (EThreadType)Context.CurrentThread, SRA_UPDATE_CALLBACK(Cancel));
 }
 
-void FSkeletalMeshStreamOut::DoReleaseBuffers(const FContext& Context)
+void FSkeletalMeshStreamOut::WaitForReferences(const FContext& Context)
 {
+	DECLARE_SCOPE_CYCLE_COUNTER(TEXT("FSkeletalMeshStreamOut::WaitForReferences"), STAT_SkeletalMeshStreamOut_WaitForReferences, STATGROUP_StreamingDetails);
+	check(Context.CurrentThread == TT_Async);
+
+	USkeletalMesh* Mesh = Context.Mesh;
+	FSkeletalMeshRenderData* RenderData = Context.RenderData;
+	uint32 NumExternalReferences = 0;
+
+	if (Mesh && RenderData)
+	{
+		for (int32 LODIdx = CurrentFirstLODIdx; LODIdx < PendingFirstMip; ++LODIdx)
+		{
+			// Minus 1 since the LODResources reference is not considered external
+			NumExternalReferences += RenderData->LODRenderData[LODIdx].GetRefCount() - 1;
+		}
+
+		if (NumExternalReferences > PreviousNumberOfExternalReferences && NumReferenceChecks > 0)
+		{
+			PreviousNumberOfExternalReferences = NumExternalReferences;
+			UE_LOG(LogSkeletalMesh, Warning, TEXT("[%s] Streamed out LODResources got referenced while in pending stream out."), *Mesh->GetName());
+		}
+	}
+
+	if (!NumExternalReferences || NumReferenceChecks >= GStreamingMaxReferenceChecks)
+	{
+		PushTask(Context, TT_Render, SRA_UPDATE_CALLBACK(ReleaseBuffers), (EThreadType)Context.CurrentThread, SRA_UPDATE_CALLBACK(Cancel));
+	}
+	else
+	{
+		++NumReferenceChecks;
+		if (NumReferenceChecks >= GStreamingMaxReferenceChecks)
+		{
+			UE_LOG(LogSkeletalMesh, Log, TEXT("[%s] Streamed out LODResources references are not getting released."), *Mesh->GetName());
+		}
+
+		bDeferExecution = true;
+		PushTask(Context, TT_Async, SRA_UPDATE_CALLBACK(WaitForReferences),  (EThreadType)Context.CurrentThread, SRA_UPDATE_CALLBACK(Cancel));
+	}
+}
+
+void FSkeletalMeshStreamOut::ReleaseBuffers(const FContext& Context)
+{
+	DECLARE_SCOPE_CYCLE_COUNTER(TEXT("FSkeletalMeshStreamOut::ReleaseBuffers"), STAT_SkeletalMeshStreamOut_ReleaseBuffers, STATGROUP_StreamingDetails);
 	check(Context.CurrentThread == TT_Render);
 	USkeletalMesh* Mesh = Context.Mesh;
 	FSkeletalMeshRenderData* RenderData = Context.RenderData;
@@ -305,12 +349,18 @@ void FSkeletalMeshStreamOut::DoReleaseBuffers(const FContext& Context)
 			LODResource.MultiSizeIndexContainer.ReleaseRHIForStreaming(Batcher);
 			LODResource.AdjacencyMultiSizeIndexContainer.ReleaseRHIForStreaming(Batcher);
 			LODResource.SkinWeightProfilesData.ReleaseRHIForStreaming(Batcher);
+
+			if (!FPlatformProperties::HasEditorOnlyData())
+			{
+				// TODO requires more testing : LODResource.ReleaseCPUResources(true);
+			}
 		}
 	}
 }
 
-void FSkeletalMeshStreamOut::DoCancel(const FContext& Context)
+void FSkeletalMeshStreamOut::Cancel(const FContext& Context)
 {
+	DECLARE_SCOPE_CYCLE_COUNTER(TEXT("FSkeletalMeshStreamOut::Cancel"), STAT_SkeletalMeshStreamOut_Cancel, STATGROUP_StreamingDetails);
 	FSkeletalMeshRenderData* RenderData = Context.RenderData;
 	if (RenderData)
 	{
diff --git a/Engine/Source/Runtime/Engine/Private/Streaming/SkeletalMeshUpdate.h b/Engine/Source/Runtime/Engine/Private/Streaming/SkeletalMeshUpdate.h
--- a/Engine/Source/Runtime/Engine/Private/Streaming/SkeletalMeshUpdate.h
+++ b/Engine/Source/Runtime/Engine/Private/Streaming/SkeletalMeshUpdate.h
@@ -132,15 +132,22 @@ public:
 	virtual ~FSkeletalMeshStreamOut() {}
 
 private:
-	void DoConditionalMarkComponentsDirty(const FContext& Context);
 
-	/** Release RHI buffers and update SRVs */
-	void DoReleaseBuffers(const FContext& Context);
+	int32 NumReferenceChecks = 0;
+	uint32 PreviousNumberOfExternalReferences = 0;
 
-	/** */
-	void DoCancel(const FContext& Context);
 
-	uint32 StartFrameNumber;
+	/** Notify components that the LOD is being streamed out so that they can release references. */
+	void ConditionalMarkComponentsDirty(const FContext& Context);
+
+	/** Wait for all references to be released. */
+	void WaitForReferences(const FContext& Context);
+
+	/** Release RHI buffers and update SRVs. */
+	void ReleaseBuffers(const FContext& Context);
+
+	/** Cancel the pending mip change. */
+	void Cancel(const FContext& Context);
 };
 
 class FSkeletalMeshStreamIn_IO : public FSkeletalMeshStreamIn
diff --git a/Engine/Source/Runtime/Engine/Public/Rendering/SkeletalMeshLODRenderData.h b/Engine/Source/Runtime/Engine/Public/Rendering/SkeletalMeshLODRenderData.h
--- a/Engine/Source/Runtime/Engine/Public/Rendering/SkeletalMeshLODRenderData.h
+++ b/Engine/Source/Runtime/Engine/Public/Rendering/SkeletalMeshLODRenderData.h
@@ -96,7 +96,7 @@ struct FSkelMeshRenderSection
 	friend FArchive& operator<<(FArchive& Ar, FSkelMeshRenderSection& S);
 };
 
-class FSkeletalMeshLODRenderData
+class FSkeletalMeshLODRenderData : public FRefCountBase
 {
 public:
 
@@ -159,11 +159,20 @@ public:
 	void ReleaseCPUResources(bool bForStreaming = false);
 
 	/** Constructor (default) */
-	FSkeletalMeshLODRenderData()
+	ENGINE_API FSkeletalMeshLODRenderData(bool bAddRef = true)
 		: BuffersSize(0)
 		, bStreamedDataInlined(true)
 		, bIsLODOptional(false)
 	{
+		if (bAddRef)
+		{
+			AddRef();
+		}
+	}
+
+	FORCEINLINE ~FSkeletalMeshLODRenderData()
+	{
+		check(GetRefCount() == 0);
 	}
 
 	/**
diff --git a/Engine/Source/Runtime/Engine/Public/Rendering/SkeletalMeshRenderData.h b/Engine/Source/Runtime/Engine/Public/Rendering/SkeletalMeshRenderData.h
--- a/Engine/Source/Runtime/Engine/Public/Rendering/SkeletalMeshRenderData.h
+++ b/Engine/Source/Runtime/Engine/Public/Rendering/SkeletalMeshRenderData.h
@@ -39,6 +39,7 @@ public:
 #endif
 
 	FSkeletalMeshRenderData();
+	~FSkeletalMeshRenderData();
 
 #if WITH_EDITOR
 	void Cache(const ITargetPlatform* TargetPlatform, USkeletalMesh* Owner);
@@ -69,13 +70,24 @@ public:
 	*/
 	ENGINE_API int32 GetMaxBonesPerSection() const;
 
+
 	/** Return first valid LOD index starting at MinLODIdx. */
 	ENGINE_API int32 GetFirstValidLODIdx(int32 MinLODIdx) const;
 
-	/** Return the current first LODIdx that can be used. */
-	FORCEINLINE int32 GetCurrentFirstLODIdx(int32 MinLODIdx) const
+	/** Return the pending first LODIdx that can be used. */
+	FORCEINLINE int32 GetPendingFirstLODIdx(int32 MinLODIdx) const
+	{
+		return GetFirstValidLODIdx(FMath::Max<int32>(PendingFirstLODIdx, MinLODIdx));
+	}
+
+	/** 
+	 * Return the pending first LOD that can be used for rendering starting at MinLODIdx.
+	 * This takes into account the streaming status from PendingFirstLODIdx, 
+	 * and MinLODIdx is expected to be UStaticMesh::MinLOD, which is platform specific.
+	 */
+	FORCEINLINE const FSkeletalMeshLODRenderData* GetPendingFirstLOD(int32 MinLODIdx) const
 	{
-		return GetFirstValidLODIdx(FMath::Max<int32>(CurrentFirstLODIdx, MinLODIdx));
+		return &LODRenderData[GetPendingFirstLODIdx(MinLODIdx)];
 	}
 
 private:
