- Fix a couple of issues with scaled transforms and bone sampling

diff --git a/Engine/Plugins/FX/Niagara/Source/Niagara/Private/NDISkeletalMesh_BoneSampling.cpp b/Engine/Plugins/FX/Niagara/Source/Niagara/Private/NDISkeletalMesh_BoneSampling.cpp
--- a/Engine/Plugins/FX/Niagara/Source/Niagara/Private/NDISkeletalMesh_BoneSampling.cpp
+++ b/Engine/Plugins/FX/Niagara/Source/Niagara/Private/NDISkeletalMesh_BoneSampling.cpp
@@ -784,8 +784,8 @@ void UNiagaraDataInterfaceSkeletalMesh::GetSkinnedBoneData(FVectorVMContext& Con
 	//Also need to pull in a per particle interpolation factor.
 	const FMatrix& InstanceTransform = InstData->Transform;
 	const FMatrix& PrevInstanceTransform = InstData->PrevTransform;
-	const FQuat InstanceRotation = Output.bNeedsRotation ? InstanceTransform.ToQuat() : FQuat::Identity;
-	const FQuat PrevInstanceRotation = Output.bNeedsRotation ? PrevInstanceTransform.ToQuat() : FQuat::Identity;
+	const FQuat InstanceRotation = Output.bNeedsRotation ? InstanceTransform.GetMatrixWithoutScale().ToQuat() : FQuat::Identity;
+	const FQuat PrevInstanceRotation = Output.bNeedsRotation ? PrevInstanceTransform.GetMatrixWithoutScale().ToQuat() : FQuat::Identity;
 
 	FSkeletalMeshAccessorHelper Accessor;
 	Accessor.Init<TNDISkelMesh_FilterModeNone, TNDISkelMesh_AreaWeightingOff>(InstData);
@@ -1027,7 +1027,7 @@ void UNiagaraDataInterfaceSkeletalMesh::GetFilteredSocketTransform(FVectorVMCont
 	{
 		const bool bNeedsRotation = OutSocketRotationX.IsValid() || OutSocketRotationY.IsValid() || OutSocketRotationZ.IsValid() || OutSocketRotationW.IsValid();
 		const FMatrix InstanceTransform = InstData->Transform;
-		const FQuat InstanceRotation = bNeedsRotation ? InstanceTransform.ToQuat() : FQuat::Identity;
+		const FQuat InstanceRotation = bNeedsRotation ? InstanceTransform.GetMatrixWithoutScale().ToQuat() : FQuat::Identity;
 
 		for (int32 i = 0; i < Context.NumInstances; ++i)
 		{
