Fix deprecated GetValueForFeatureLevel API breaking build

diff --git a/Engine/Plugins/FX/Niagara/Source/Niagara/Private/NiagaraDataInterfaceSkeletalMesh.cpp b/Engine/Plugins/FX/Niagara/Source/Niagara/Private/NiagaraDataInterfaceSkeletalMesh.cpp
--- a/Engine/Plugins/FX/Niagara/Source/Niagara/Private/NiagaraDataInterfaceSkeletalMesh.cpp
+++ b/Engine/Plugins/FX/Niagara/Source/Niagara/Private/NiagaraDataInterfaceSkeletalMesh.cpp
@@ -1527,7 +1527,7 @@ bool FNDISkeletalMesh_InstanceData::Init(UNiagaraDataInterfaceSkeletalMesh* Inte
 
 	if (Mesh)
 	{
-		MinLODIdx = Mesh->MinLod.GetValueForFeatureLevel(SystemInstance->GetFeatureLevel());
+		MinLODIdx = Mesh->MinLod.GetValue();
 		const int32 PendingFirstLODIndex = Mesh->GetResourceForRendering()->GetPendingFirstLODIdx(MinLODIdx);
 
 		const int32 DesiredLODIndex = Interface->CalculateLODIndexAndSamplingRegions(Mesh, SamplingRegionIndices, bAllRegionsAreAreaWeighting);
