Fix creation of niagara constant buffers so it is done on the frame that it is being executed (rather than when it is queued on the RT).  Fixes crashes/corruption when ticks are accumulated (application context switches with the editor).

diff --git a/Engine/Plugins/FX/Niagara/Source/Niagara/Classes/NiagaraEmitterInstanceBatcher.h b/Engine/Plugins/FX/Niagara/Source/Niagara/Classes/NiagaraEmitterInstanceBatcher.h
--- a/Engine/Plugins/FX/Niagara/Source/Niagara/Classes/NiagaraEmitterInstanceBatcher.h
+++ b/Engine/Plugins/FX/Niagara/Source/Niagara/Classes/NiagaraEmitterInstanceBatcher.h
@@ -203,7 +203,7 @@ private:
 	void UpdateFreeIDBuffers(FRHICommandList& RHICmdList, FEmitterInstanceList& Instances);
 
 	void SetConstantBuffers(FRHICommandList &RHICmdList, const FNiagaraShaderRef& Shader, const FNiagaraGPUSystemTick& Tick, const FNiagaraComputeInstanceData* Instance);
-	void BuildConstantBuffers(FNiagaraGPUSystemTick& Tick);
+	void BuildConstantBuffers(FNiagaraGPUSystemTick& Tick) const;
 
 	/** Feature level of this effects system */
 	ERHIFeatureLevel::Type FeatureLevel;
diff --git a/Engine/Plugins/FX/Niagara/Source/Niagara/Private/NiagaraEmitterInstanceBatcher.cpp b/Engine/Plugins/FX/Niagara/Source/Niagara/Private/NiagaraEmitterInstanceBatcher.cpp
--- a/Engine/Plugins/FX/Niagara/Source/Niagara/Private/NiagaraEmitterInstanceBatcher.cpp
+++ b/Engine/Plugins/FX/Niagara/Source/Niagara/Private/NiagaraEmitterInstanceBatcher.cpp
@@ -151,7 +151,7 @@ void NiagaraEmitterInstanceBatcher::InstanceDeallocated_RenderThread(const FNiag
 	}
 }
 
-void NiagaraEmitterInstanceBatcher::BuildConstantBuffers(FNiagaraGPUSystemTick& Tick)
+void NiagaraEmitterInstanceBatcher::BuildConstantBuffers(FNiagaraGPUSystemTick& Tick) const
 {
 	if (!Tick.Count)
 	{
@@ -289,9 +289,7 @@ void NiagaraEmitterInstanceBatcher::GiveSystemTick_RenderThread(FNiagaraGPUSyste
 	// A note:
 	// This is making a copy of Tick. That structure is small now and we take a copy to avoid
 	// making a bunch of small allocations on the game thread. We may need to revisit this.
-	FNiagaraGPUSystemTick& AddedTick = Ticks_RT.Add_GetRef(Tick);
-
-	BuildConstantBuffers(AddedTick);
+	Ticks_RT.Add(Tick);
 }
 
 void NiagaraEmitterInstanceBatcher::ReleaseInstanceCounts_RenderThread(FNiagaraComputeExecutionContext* ExecContext, FNiagaraDataSet* DataSet)
@@ -908,6 +906,8 @@ void NiagaraEmitterInstanceBatcher::ExecuteAll(FRHICommandList& RHICmdList, FRHI
 				continue;
 			}
 
+			BuildConstantBuffers(Tick);
+
 			Tick.bIsFinalTick = false; // @todo : this is true sometimes, needs investigation
 
 			const bool bResetCounts = SharedContext->ScratchIndex == INDEX_NONE;
