Make niagara mesh normals match static mesh better, by eliminating scale from the TangentToWorld transformation

diff --git a/Engine/Plugins/FX/Niagara/Shaders/Private/NiagaraMeshVertexFactory.ush b/Engine/Plugins/FX/Niagara/Shaders/Private/NiagaraMeshVertexFactory.ush
--- a/Engine/Plugins/FX/Niagara/Shaders/Private/NiagaraMeshVertexFactory.ush
+++ b/Engine/Plugins/FX/Niagara/Shaders/Private/NiagaraMeshVertexFactory.ush
@@ -212,6 +212,7 @@ struct FVertexFactoryIntermediates
 #endif
 
 	float4x4 ParticleToWorld;
+	float4x4 ParticleToWorldNoScale;
 	float4x4 PrevParticleToWorld;
 
 	float4x4 WorldToParticle;
@@ -222,6 +223,11 @@ float4x4 GetParticleTransform(FVertexFactoryIntermediates Inters)
 	return Inters.ParticleToWorld;
 }
 
+float4x4 GetParticleTransformNoScale(FVertexFactoryIntermediates Inters)
+{
+	return Inters.ParticleToWorldNoScale;
+}
+
 float4x4 GetParticleInvTransform(FVertexFactoryIntermediates Inters)
 {
 	return Inters.WorldToParticle;
@@ -411,7 +417,8 @@ FMaterialVertexParameters GetMaterialVertexParameters(FVertexFactoryInput Input,
 	Result.Particle.ParticleToWorld = Transform;
 	Result.Particle.WorldToParticle = GetParticleInvTransform(Intermediates);
 
-	Result.TangentToWorld = mul(TangentToLocal, (float3x3)Transform );
+	// Use the particle's unscaled transform for rotating normals to match static mesh
+	Result.TangentToWorld = mul(TangentToLocal, (float3x3)GetParticleTransformNoScale(Intermediates));
 
 	Result.PrevFrameLocalToWorld = GetParticlePrevTransform(Intermediates);
 
@@ -605,7 +612,8 @@ float3 CalculateCameraOffset(float3 WorldPosition, uint InstanceId)
 	return CameraDir * Offset;
 }
 
-void CalculateTransforms(int InstanceId, in FVertexFactoryIntermediates Intermediates, out float4x4 LocalToWorld, out float4x4 WorldToLocal)
+void CalculateTransforms(int InstanceId, in FVertexFactoryIntermediates Intermediates, out float4x4 LocalToWorld,
+	out float4x4 LocalToWorldNoScale, out float4x4 WorldToLocal)
 {
 	////////////////////////////////////////////////////////////////////////////////
 	// Translation
@@ -654,6 +662,7 @@ void CalculateTransforms(int InstanceId, in FVertexFactoryIntermediates Intermed
 	////////////////////////////////////////////////////////////////////////////////
 	// Compose Transforms
 	LocalToWorld = ComposeTransformMatrix(Scale, RotationMat, WorldParticlePosition);
+	LocalToWorldNoScale = ComposeTransformMatrix((float3)1.0f, RotationMat, WorldParticlePosition);
 	WorldToLocal = ComposeInvTransformMatrix(Scale, RotationMat, WorldParticlePosition);
 }
 
@@ -753,9 +762,10 @@ FVertexFactoryIntermediates GetVertexFactoryIntermediates(FVertexFactoryInput In
 	Intermediates.VertexColor = Input.VertexColor FCOLOR_COMPONENT_SWIZZLE;
 
 	float4x4 ParticleToWorld;
+	float4x4 ParticleToWorldNoScale;
 	float4x4 WorldToParticle;
 	
-	CalculateTransforms(Input.InstanceId, Intermediates, ParticleToWorld, WorldToParticle);
+	CalculateTransforms(Input.InstanceId, Intermediates, ParticleToWorld, ParticleToWorldNoScale, WorldToParticle);
 
 	//-TODO: If / when we store previous transform data use them here (providing they are in local space, if not handle below)
 	//if (MeshParticleVF.PrevTransformAvailable > 0)
@@ -777,10 +787,17 @@ FVertexFactoryIntermediates GetVertexFactoryIntermediates(FVertexFactoryInput In
 		Intermediates.PrevParticleToWorld = mul(Intermediates.PrevParticleToWorld, Primitive.PreviousLocalToWorld);
 		Intermediates.ParticleToWorld = mul(ParticleToWorld, Primitive.LocalToWorld);
 		Intermediates.WorldToParticle = mul(Primitive.WorldToLocal, WorldToParticle);
+
+		Intermediates.ParticleToWorldNoScale = mul(ParticleToWorldNoScale, Primitive.LocalToWorld);
+		float3 InvScale = Primitive.InvNonUniformScaleAndDeterminantSign.xyz;
+		Intermediates.ParticleToWorldNoScale[0].xyz *= InvScale.x;
+		Intermediates.ParticleToWorldNoScale[1].xyz *= InvScale.y;
+		Intermediates.ParticleToWorldNoScale[2].xyz *= InvScale.z;
 	}
 	else
 	{
 		Intermediates.ParticleToWorld = ParticleToWorld;
+		Intermediates.ParticleToWorldNoScale = ParticleToWorldNoScale;
 		Intermediates.WorldToParticle = WorldToParticle;
 	}
 
@@ -877,10 +894,10 @@ void CalcTangentToWorld(FVertexFactoryInput Input, FVertexFactoryIntermediates I
 {
 	float3x3 TangentToLocal = VertexFactoryGetTangentToLocal(Input, Intermediates);
 
-	float4x4 Transform = GetParticleTransform(Intermediates);
-	float3x3 TangentToWorld = mul(TangentToLocal, (float3x3)Transform);
+	// Use the particle's unscaled transform for rotating normals to match static mesh
+	float3x3 TransformNoScale = (float3x3)GetParticleTransformNoScale(Intermediates);
+	float3x3 TangentToWorld = mul(TangentToLocal, TransformNoScale);
 
-	// Normalize to remove scaling. Incorrect but faster than computing the inverse-transpose.
 	TangentToWorld0 = normalize(TangentToWorld[0]);
 	TangentToWorld2 = float4(normalize(TangentToWorld[2]), TangentBias(Input.TangentZ.w) * Primitive.InvNonUniformScaleAndDeterminantSign.w);
 }
