Fixes to particle space material transforms with GLSL shaders, which didn't support matrices as interpolators

diff --git a/Engine/Plugins/FX/Niagara/Shaders/Private/NiagaraMeshVertexFactory.ush b/Engine/Plugins/FX/Niagara/Shaders/Private/NiagaraMeshVertexFactory.ush
--- a/Engine/Plugins/FX/Niagara/Shaders/Private/NiagaraMeshVertexFactory.ush
+++ b/Engine/Plugins/FX/Niagara/Shaders/Private/NiagaraMeshVertexFactory.ush
@@ -149,11 +149,11 @@ struct FVertexFactoryInterpolantsVSToPS
 #endif
 
 #if USE_PARTICLE_LOCAL_TO_WORLD
-	nointerpolation float3x4 ParticleToWorld : PARTICLE_LOCAL_TO_WORLD;
+	nointerpolation float4 ParticleToWorld[3] : PARTICLE_LOCAL_TO_WORLD;
 #endif
 
 #if USE_PARTICLE_WORLD_TO_LOCAL
-	nointerpolation float3x4 WorldToParticle : PARTICLE_WORLD_TO_LOCAL;
+	nointerpolation float4 WorldToParticle[3] : PARTICLE_WORLD_TO_LOCAL;
 #endif
 
 #if INSTANCED_STEREO
@@ -987,11 +987,17 @@ FVertexFactoryInterpolantsVSToPS VertexFactoryGetInterpolantsVSToPS(FVertexFacto
 #endif
 
 #if USE_PARTICLE_LOCAL_TO_WORLD
-	Interpolants.ParticleToWorld = (float3x4) transpose(Intermediates.ParticleToWorld);
+	float4x4 ParticleToWorldT = transpose(Intermediates.ParticleToWorld);
+	Interpolants.ParticleToWorld[0] = ParticleToWorldT[0];
+	Interpolants.ParticleToWorld[1] = ParticleToWorldT[1];
+	Interpolants.ParticleToWorld[2] = ParticleToWorldT[2];
 #endif
 
 #if USE_PARTICLE_WORLD_TO_LOCAL
-	Interpolants.WorldToParticle = (float3x4) transpose(Intermediates.WorldToParticle);
+	float4x4 WorldToParticleT = transpose(Intermediates.WorldToParticle);
+	Interpolants.WorldToParticle[0] = WorldToParticleT[0];
+	Interpolants.WorldToParticle[1] = WorldToParticleT[1];
+	Interpolants.WorldToParticle[2] = WorldToParticleT[2];
 #endif
 
 #if INSTANCED_STEREO
diff --git a/Engine/Shaders/Private/MeshParticleVertexFactory.ush b/Engine/Shaders/Private/MeshParticleVertexFactory.ush
--- a/Engine/Shaders/Private/MeshParticleVertexFactory.ush
+++ b/Engine/Shaders/Private/MeshParticleVertexFactory.ush
@@ -121,7 +121,7 @@ struct FVertexFactoryInterpolantsVSToPS
 #endif
 
 #if NEEDS_PARTICLE_LOCAL_TO_WORLD || NEEDS_PARTICLE_WORLD_TO_LOCAL
-	float3x4 LocalToWorld : PARTICLE_TRANSFORM;
+	float4 LocalToWorld[3] : PARTICLE_TRANSFORM;
 #endif
 
 #if INSTANCED_STEREO
@@ -589,7 +589,10 @@ FVertexFactoryInterpolantsVSToPS VertexFactoryGetInterpolantsVSToPS(FVertexFacto
 #endif
 
 #if NEEDS_PARTICLE_LOCAL_TO_WORLD || NEEDS_PARTICLE_WORLD_TO_LOCAL
-	Interpolants.LocalToWorld = (float3x4) transpose(Intermediates.LocalToWorld);
+	float4x4 LocalToWorldT = transpose(Intermediates.LocalToWorld);
+	Interpolants.LocalToWorld[0] = LocalToWorldT[0];
+	Interpolants.LocalToWorld[1] = LocalToWorldT[1];
+	Interpolants.LocalToWorld[2] = LocalToWorldT[2];
 #endif
 
 #if INSTANCED_STEREO
