Fixing incremental CIS some more. The issue with CIS seems to be missing forward declares propagated from removing inclues in other changes today

diff --git a/Engine/Plugins/FX/Niagara/Source/NiagaraEditor/Public/Widgets/SNiagaraParameterMapView.h b/Engine/Plugins/FX/Niagara/Source/NiagaraEditor/Public/Widgets/SNiagaraParameterMapView.h
--- a/Engine/Plugins/FX/Niagara/Source/NiagaraEditor/Public/Widgets/SNiagaraParameterMapView.h
+++ b/Engine/Plugins/FX/Niagara/Source/NiagaraEditor/Public/Widgets/SNiagaraParameterMapView.h
@@ -29,6 +29,7 @@ class UNiagaraSystem;
 struct FNiagaraNamespaceMetadata;
 class IToolTip;
 class UNiagaraEditorSettings;
+class FMenuBuilder;
 
 /* Enums to use when grouping the blueprint members in the list panel. The order here will determine the order in the list */
 namespace NiagaraParameterMapSectionID
diff --git a/Engine/Plugins/FX/Niagara/Source/NiagaraEditorWidgets/Private/DetailCustomizations/SNiagaraNamePropertySelector.h b/Engine/Plugins/FX/Niagara/Source/NiagaraEditorWidgets/Private/DetailCustomizations/SNiagaraNamePropertySelector.h
--- a/Engine/Plugins/FX/Niagara/Source/NiagaraEditorWidgets/Private/DetailCustomizations/SNiagaraNamePropertySelector.h
+++ b/Engine/Plugins/FX/Niagara/Source/NiagaraEditorWidgets/Private/DetailCustomizations/SNiagaraNamePropertySelector.h
@@ -7,7 +7,9 @@
 #include "Types/SlateEnums.h"
 #include "Widgets/Input/SComboBox.h"
 #include "PropertyCustomizationHelpers.h"
+
 class IPropertyHandle;
+class SSearchBox;
 
 /** Details customization for arrays composed of FName properties (or wrappers). The array contents are selected from a predetermined source list.*/
 class SNiagaraNamePropertySelector : public SCompoundWidget 
diff --git a/Engine/Plugins/Runtime/AR/Microsoft/HoloLensAR/Source/HoloLensPlatformEditor/Private/HoloLensTargetSettingsCustomization.h b/Engine/Plugins/Runtime/AR/Microsoft/HoloLensAR/Source/HoloLensPlatformEditor/Private/HoloLensTargetSettingsCustomization.h
--- a/Engine/Plugins/Runtime/AR/Microsoft/HoloLensAR/Source/HoloLensPlatformEditor/Private/HoloLensTargetSettingsCustomization.h
+++ b/Engine/Plugins/Runtime/AR/Microsoft/HoloLensAR/Source/HoloLensPlatformEditor/Private/HoloLensTargetSettingsCustomization.h
@@ -9,6 +9,7 @@
 #include "TargetPlatformAudioCustomization.h"
 
 class SErrorHint;
+class STextComboBox;
 
 class FHoloLensTargetSettingsCustomization : public IDetailCustomization
 {
diff --git a/Engine/Source/Editor/LandscapeEditor/Private/LandscapeEditorDetailCustomization_LayersBrushStack.h b/Engine/Source/Editor/LandscapeEditor/Private/LandscapeEditorDetailCustomization_LayersBrushStack.h
--- a/Engine/Source/Editor/LandscapeEditor/Private/LandscapeEditorDetailCustomization_LayersBrushStack.h
+++ b/Engine/Source/Editor/LandscapeEditor/Private/LandscapeEditorDetailCustomization_LayersBrushStack.h
@@ -22,6 +22,7 @@ class FDetailWidgetRow;
 class IDetailChildrenBuilder;
 class IDetailLayoutBuilder;
 class SDragAndDropVerticalBox;
+class FMenuBuilder;
 
 /**
  * Slate widgets customizer for the layers list in the Landscape Editor
diff --git a/Engine/Source/Editor/SequenceRecorder/Private/SClassPropertyRecorderSettings.cpp b/Engine/Source/Editor/SequenceRecorder/Private/SClassPropertyRecorderSettings.cpp
--- a/Engine/Source/Editor/SequenceRecorder/Private/SClassPropertyRecorderSettings.cpp
+++ b/Engine/Source/Editor/SequenceRecorder/Private/SClassPropertyRecorderSettings.cpp
@@ -12,6 +12,7 @@
 #include "Widgets/Input/SCheckBox.h"
 #include "IDetailsView.h"
 #include "Sections/MovieSceneMultiPropertyRecorder.h"
+#include "Framework/Docking/TabManager.h"
 
 #define LOCTEXT_NAMESPACE "SClassPropertyRecorderSettings"
 
