- Check IsValid on game thread shader map before doing various operations to avoid pulling a null shader if we have not completed successfully

#[fyi] bob.tellez
#rnx

diff --git a/Engine/Plugins/FX/Niagara/Source/NiagaraShader/Private/NiagaraShared.cpp b/Engine/Plugins/FX/Niagara/Source/NiagaraShader/Private/NiagaraShared.cpp
--- a/Engine/Plugins/FX/Niagara/Source/NiagaraShader/Private/NiagaraShared.cpp
+++ b/Engine/Plugins/FX/Niagara/Source/NiagaraShader/Private/NiagaraShared.cpp
@@ -167,6 +167,11 @@ bool FNiagaraShaderScript::IsShaderMapComplete() const
 		return false;
 	}
 
+	if (!GameThreadShaderMap->IsValid())
+	{
+		return false;
+	}
+
 	for (int i=0; i < GetNumPermutations(); ++i)
 	{
 		if (GameThreadShaderMap->GetShader<FNiagaraShader>(i).IsNull())
@@ -487,7 +492,7 @@ void FNiagaraShaderScript::UpdateCachedData_PostCompile(bool bCalledFromSerializ
 	CachedData.bExternalConstantBufferUsed = 0;
 	CachedData.bViewUniformBufferUsed = 0;
 
-	if (GameThreadShaderMap != nullptr)
+	if (GameThreadShaderMap != nullptr && GameThreadShaderMap->IsValid())
 	{
 		for (int32 iPermutation = 0; iPermutation < CachedData.NumPermutations; ++iPermutation)
 		{
