[NO NON-PII COMMIT MESSAGE]

diff --git a/Engine/Plugins/FX/Niagara/Source/NiagaraEditor/Private/Customizations/NiagaraPlatformSetCustomization.cpp b/Engine/Plugins/FX/Niagara/Source/NiagaraEditor/Private/Customizations/NiagaraPlatformSetCustomization.cpp
--- a/Engine/Plugins/FX/Niagara/Source/NiagaraEditor/Private/Customizations/NiagaraPlatformSetCustomization.cpp
+++ b/Engine/Plugins/FX/Niagara/Source/NiagaraEditor/Private/Customizations/NiagaraPlatformSetCustomization.cpp
@@ -1432,11 +1432,6 @@ FReply SNiagaraConsoleInputBox::OnPreviewKeyDown(const FGeometry& MyGeometry, co
 					Suggestions.StepSelectedSuggestion(-1);
 					MarkActiveSuggestion();
 				}
-			}
-
-			// Need to always handle this for single-line controls to avoid them invoking widget navigation
-			if (!bIsMultiLine || bShowHistory)
-			{
 				return FReply::Handled();
 			}
 		}
