Make some niagara data-interfaces have non-inline type layouts, so they work with new memory freezing

diff --git a/Engine/Plugins/FX/Niagara/Source/Niagara/Private/NiagaraDataInterface2DArrayTexture.cpp b/Engine/Plugins/FX/Niagara/Source/Niagara/Private/NiagaraDataInterface2DArrayTexture.cpp
--- a/Engine/Plugins/FX/Niagara/Source/Niagara/Private/NiagaraDataInterface2DArrayTexture.cpp
+++ b/Engine/Plugins/FX/Niagara/Source/Niagara/Private/NiagaraDataInterface2DArrayTexture.cpp
@@ -166,7 +166,7 @@ void UNiagaraDataInterface2DArrayTexture::GetParameterDefinitionHLSL(const FNiag
 
 struct FNiagaraDataInterfaceParametersCS_2DArrayTexture : public FNiagaraDataInterfaceParametersCS
 {
-	DECLARE_INLINE_TYPE_LAYOUT(FNiagaraDataInterfaceParametersCS_2DArrayTexture, NonVirtual);
+	DECLARE_TYPE_LAYOUT(FNiagaraDataInterfaceParametersCS_2DArrayTexture, NonVirtual);
 public:
 	void Bind(const FNiagaraDataInterfaceGPUParamInfo& ParameterInfo, const class FShaderParameterMap& ParameterMap)
 	{
@@ -224,6 +224,8 @@ private:
 	LAYOUT_FIELD(FShaderParameter, Dimensions);
 };
 
+IMPLEMENT_TYPE_LAYOUT(FNiagaraDataInterfaceParametersCS_2DArrayTexture);
+
 IMPLEMENT_NIAGARA_DI_PARAMETER(UNiagaraDataInterface2DArrayTexture, FNiagaraDataInterfaceParametersCS_2DArrayTexture);
 
 void UNiagaraDataInterface2DArrayTexture::PushToRenderThreadImpl()
diff --git a/Engine/Plugins/FX/Niagara/Source/Niagara/Private/NiagaraDataInterfaceAudioOscilloscope.cpp b/Engine/Plugins/FX/Niagara/Source/Niagara/Private/NiagaraDataInterfaceAudioOscilloscope.cpp
--- a/Engine/Plugins/FX/Niagara/Source/Niagara/Private/NiagaraDataInterfaceAudioOscilloscope.cpp
+++ b/Engine/Plugins/FX/Niagara/Source/Niagara/Private/NiagaraDataInterfaceAudioOscilloscope.cpp
@@ -317,7 +317,7 @@ void UNiagaraDataInterfaceAudioOscilloscope::GetParameterDefinitionHLSL(const FN
 
 struct FNiagaraDataInterfaceParametersCS_AudioOscilloscope : public FNiagaraDataInterfaceParametersCS
 {
-	DECLARE_INLINE_TYPE_LAYOUT(FNiagaraDataInterfaceParametersCS_AudioOscilloscope, NonVirtual);
+	DECLARE_TYPE_LAYOUT(FNiagaraDataInterfaceParametersCS_AudioOscilloscope, NonVirtual);
 
 	void Bind(const FNiagaraDataInterfaceGPUParamInfo& ParameterInfo, const class FShaderParameterMap& ParameterMap)
 	{
@@ -342,6 +342,8 @@ struct FNiagaraDataInterfaceParametersCS_AudioOscilloscope : public FNiagaraData
 	LAYOUT_FIELD(FShaderResourceParameter, AudioBuffer);
 };
 
+IMPLEMENT_TYPE_LAYOUT(FNiagaraDataInterfaceParametersCS_AudioOscilloscope);
+
 IMPLEMENT_NIAGARA_DI_PARAMETER(UNiagaraDataInterfaceAudioOscilloscope, FNiagaraDataInterfaceParametersCS_AudioOscilloscope);
 
 void FNiagaraDataInterfaceProxyOscilloscope::OnUpdateResampling(int32 InResolution, float InScopeInMilliseconds)
diff --git a/Engine/Plugins/FX/Niagara/Source/Niagara/Private/NiagaraDataInterfaceAudioSpectrum.cpp b/Engine/Plugins/FX/Niagara/Source/Niagara/Private/NiagaraDataInterfaceAudioSpectrum.cpp
--- a/Engine/Plugins/FX/Niagara/Source/Niagara/Private/NiagaraDataInterfaceAudioSpectrum.cpp
+++ b/Engine/Plugins/FX/Niagara/Source/Niagara/Private/NiagaraDataInterfaceAudioSpectrum.cpp
@@ -746,7 +746,7 @@ void UNiagaraDataInterfaceAudioSpectrum::GetParameterDefinitionHLSL(const FNiaga
 
 struct FNiagaraDataInterfaceParametersCS_AudioSpectrum : public FNiagaraDataInterfaceParametersCS
 {
-	DECLARE_INLINE_TYPE_LAYOUT(FNiagaraDataInterfaceParametersCS_AudioSpectrum, NonVirtual);
+	DECLARE_TYPE_LAYOUT(FNiagaraDataInterfaceParametersCS_AudioSpectrum, NonVirtual);
 
 	void Bind(const FNiagaraDataInterfaceGPUParamInfo& ParameterInfo, const class FShaderParameterMap& ParameterMap)
 	{
@@ -776,6 +776,8 @@ struct FNiagaraDataInterfaceParametersCS_AudioSpectrum : public FNiagaraDataInte
 	LAYOUT_FIELD(FShaderResourceParameter,SpectrumBuffer);
 };
 
+IMPLEMENT_TYPE_LAYOUT(FNiagaraDataInterfaceParametersCS_AudioSpectrum);
+
 IMPLEMENT_NIAGARA_DI_PARAMETER(UNiagaraDataInterfaceAudioSpectrum, FNiagaraDataInterfaceParametersCS_AudioSpectrum);
 
 
diff --git a/Engine/Plugins/FX/Niagara/Source/Niagara/Private/NiagaraDataInterfaceVolumeTexture.cpp b/Engine/Plugins/FX/Niagara/Source/Niagara/Private/NiagaraDataInterfaceVolumeTexture.cpp
--- a/Engine/Plugins/FX/Niagara/Source/Niagara/Private/NiagaraDataInterfaceVolumeTexture.cpp
+++ b/Engine/Plugins/FX/Niagara/Source/Niagara/Private/NiagaraDataInterfaceVolumeTexture.cpp
@@ -213,7 +213,7 @@ void UNiagaraDataInterfaceVolumeTexture::GetParameterDefinitionHLSL(const FNiaga
 
 struct FNiagaraDataInterfaceParametersCS_VolumeTexture : public FNiagaraDataInterfaceParametersCS
 {
-	DECLARE_INLINE_TYPE_LAYOUT(FNiagaraDataInterfaceParametersCS_VolumeTexture, NonVirtual);
+	DECLARE_TYPE_LAYOUT(FNiagaraDataInterfaceParametersCS_VolumeTexture, NonVirtual);
 public:
 	void Bind(const FNiagaraDataInterfaceGPUParamInfo& ParameterInfo, const class FShaderParameterMap& ParameterMap)
 	{
@@ -281,6 +281,7 @@ private:
 	LAYOUT_FIELD(FShaderParameter, Dimensions);
 };
 
+IMPLEMENT_TYPE_LAYOUT(FNiagaraDataInterfaceParametersCS_VolumeTexture);
 IMPLEMENT_NIAGARA_DI_PARAMETER(UNiagaraDataInterfaceVolumeTexture, FNiagaraDataInterfaceParametersCS_VolumeTexture);
 
 void UNiagaraDataInterfaceVolumeTexture::PushToRenderThreadImpl()
