Merging CL 14641055 from //UE5/Release-Engine-Staging to Main (//UE5/Main) to fix non-unity issues

diff --git a/Engine/Plugins/FX/Niagara/Source/Niagara/Private/NiagaraComponentRendererProperties.cpp b/Engine/Plugins/FX/Niagara/Source/Niagara/Private/NiagaraComponentRendererProperties.cpp
--- a/Engine/Plugins/FX/Niagara/Source/Niagara/Private/NiagaraComponentRendererProperties.cpp
+++ b/Engine/Plugins/FX/Niagara/Source/Niagara/Private/NiagaraComponentRendererProperties.cpp
@@ -6,6 +6,7 @@
 #include "NiagaraRendererComponents.h"
 #include "Modules/ModuleManager.h"
 #if WITH_EDITOR
+#include "Editor.h"
 #include "Widgets/Images/SImage.h"
 #include "Styling/SlateIconFinder.h"
 #include "Widgets/SWidget.h"
diff --git a/Engine/Plugins/Media/ImgMedia/Source/ImgMedia/Private/Readers/IImgMediaReader.h b/Engine/Plugins/Media/ImgMedia/Source/ImgMedia/Private/Readers/IImgMediaReader.h
--- a/Engine/Plugins/Media/ImgMedia/Source/ImgMedia/Private/Readers/IImgMediaReader.h
+++ b/Engine/Plugins/Media/ImgMedia/Source/ImgMedia/Private/Readers/IImgMediaReader.h
@@ -7,7 +7,8 @@
 #include "IMediaTextureSample.h"
 #include "Math/IntPoint.h"
 #include "Templates/SharedPointer.h"
-
+#include "RHI.h"
+#include "RHIResources.h"
 
 /**
  * Information about an image sequence frame.
diff --git a/Engine/Plugins/Runtime/Oculus/OculusOpenXR/Source/OculusOpenXRHMD/Private/OculusOpenXRHMD.cpp b/Engine/Plugins/Runtime/Oculus/OculusOpenXR/Source/OculusOpenXRHMD/Private/OculusOpenXRHMD.cpp
--- a/Engine/Plugins/Runtime/Oculus/OculusOpenXR/Source/OculusOpenXRHMD/Private/OculusOpenXRHMD.cpp
+++ b/Engine/Plugins/Runtime/Oculus/OculusOpenXR/Source/OculusOpenXRHMD/Private/OculusOpenXRHMD.cpp
@@ -4,6 +4,7 @@
 #include "OpenXRCore.h"
 #include "OpenXRPlatformRHI.h"
 #include "DefaultSpectatorScreenController.h"
+#include "Modules/ModuleManager.h"
 
 #if PLATFORM_ANDROID
 #include <openxr_oculus.h>
