- Editgrate a couple of changes to fix cook issues on specific platforms

#rnx

diff --git a/Engine/Plugins/FX/Niagara/Shaders/Private/NiagaraDebugDraw.ush b/Engine/Plugins/FX/Niagara/Shaders/Private/NiagaraDebugDraw.ush
--- a/Engine/Plugins/FX/Niagara/Shaders/Private/NiagaraDebugDraw.ush
+++ b/Engine/Plugins/FX/Niagara/Shaders/Private/NiagaraDebugDraw.ush
@@ -32,7 +32,7 @@ void PackLine(inout RWBuffer<float> OutBuffer, uint InstanceIndex, float3 InPosi
 	PackLine(OutBuffer, InstanceIndex, InPosition0, InPosition1, PackedColor);
 }
 
-void UnpackLine(inout Buffer<float> InBuffer, uint InstanceIndex, out float3 OutPosition0, out float3 OutPosition1, out float4 OutColor)
+void UnpackLine(in Buffer<float> InBuffer, uint InstanceIndex, out float3 OutPosition0, out float3 OutPosition1, out float4 OutColor)
 {
 	uint FloatOffset = InstanceIndex * 7;
 
diff --git a/Engine/Plugins/FX/Niagara/Source/NiagaraShader/Private/NiagaraDebugShaders.cpp b/Engine/Plugins/FX/Niagara/Source/NiagaraShader/Private/NiagaraDebugShaders.cpp
--- a/Engine/Plugins/FX/Niagara/Source/NiagaraShader/Private/NiagaraDebugShaders.cpp
+++ b/Engine/Plugins/FX/Niagara/Source/NiagaraShader/Private/NiagaraDebugShaders.cpp
@@ -43,7 +43,7 @@ public:
 	static bool ShouldCompilePermutation(const FGlobalShaderPermutationParameters& Parameters)
 	{
 		FPermutationDomain PermutationVector(Parameters.PermutationId);
-		return true;
+		return !IsMobilePlatform(Parameters.Platform);
 	}
 
 	BEGIN_SHADER_PARAMETER_STRUCT(FParameters, )
@@ -81,7 +81,7 @@ public:
 	static bool ShouldCompilePermutation(const FGlobalShaderPermutationParameters& Parameters)
 	{
 		//FPermutationDomain PermutationVector(Parameters.PermutationId);
-		return true;
+		return !IsMobilePlatform(Parameters.Platform);
 	}
 
 	static void ModifyCompilationEnvironment(const FGlobalShaderPermutationParameters& Parameters, FShaderCompilerEnvironment& OutEnvironment)
@@ -112,7 +112,7 @@ public:
 
 	static bool ShouldCompilePermutation(const FGlobalShaderPermutationParameters& Parameters)
 	{
-		return true;
+		return !IsMobilePlatform(Parameters.Platform);
 	}
 
 	static void ModifyCompilationEnvironment(const FGlobalShaderPermutationParameters& Parameters, FShaderCompilerEnvironment& OutEnvironment)
@@ -135,7 +135,7 @@ public:
 
 	static bool ShouldCompilePermutation(const FGlobalShaderPermutationParameters& Parameters)
 	{
-		return true;
+		return !IsMobilePlatform(Parameters.Platform);
 	}
 
 	static void ModifyCompilationEnvironment(const FGlobalShaderPermutationParameters& Parameters, FShaderCompilerEnvironment& OutEnvironment)
