Niagara edit conditions for property rows have been restored, eg. the checkbox for Fixed Bounds.

diff --git a/Engine/Plugins/FX/Niagara/Source/NiagaraEditorWidgets/Private/SNiagaraStack.cpp b/Engine/Plugins/FX/Niagara/Source/NiagaraEditorWidgets/Private/SNiagaraStack.cpp
--- a/Engine/Plugins/FX/Niagara/Source/NiagaraEditorWidgets/Private/SNiagaraStack.cpp
+++ b/Engine/Plugins/FX/Niagara/Source/NiagaraEditorWidgets/Private/SNiagaraStack.cpp
@@ -679,7 +679,7 @@ TSharedRef<ITableRow> SNiagaraStack::OnGenerateRowForStackItem(UNiagaraStackEntr
 {
 	TSharedRef<SNiagaraStackTableRow> Container = ConstructContainerForItem(Item);
 	FRowWidgets RowWidgets = ConstructNameAndValueWidgetsForItem(Item, Container);
-	Container->SetNameAndValueContent(RowWidgets.NameWidget, RowWidgets.ValueWidget);
+	Container->SetNameAndValueContent(RowWidgets.NameWidget, RowWidgets.ValueWidget, RowWidgets.EditConditionWidget);
 	return Container;
 }
 
@@ -931,8 +931,7 @@ SNiagaraStack::FRowWidgets SNiagaraStack::ConstructNameAndValueWidgetsForItem(UN
 			.IsSelected(Container, &SNiagaraStackTableRow::IsSelected);
 		Container->AddFillRowContextMenuHandler(SNiagaraStackTableRow::FOnFillRowContextMenu::CreateSP(FunctionInputNameWidget, &SNiagaraStackFunctionInputName::FillRowContextMenu));
 
-		return FRowWidgets(FunctionInputNameWidget,
-			SNew(SNiagaraStackFunctionInputValue, FunctionInput));
+		return FRowWidgets(FunctionInputNameWidget,	SNew(SNiagaraStackFunctionInputValue, FunctionInput));
 	}
 	else if (Item->IsA<UNiagaraStackErrorItem>())
 	{
@@ -1015,17 +1014,23 @@ SNiagaraStack::FRowWidgets SNiagaraStack::ConstructNameAndValueWidgetsForItem(UN
 			Container->SetOverrideNameWidth(PropertyRowWidgets.WholeRowWidgetLayoutData.MinWidth, PropertyRowWidgets.WholeRowWidgetLayoutData.MaxWidth);
 			Container->SetOverrideNameAlignment(PropertyRowWidgets.WholeRowWidgetLayoutData.HorizontalAlignment, PropertyRowWidgets.WholeRowWidgetLayoutData.VerticalAlignment);
 			PropertyRowWidgets.WholeRowWidget->SetEnabled(IsEnabled);
-			return FRowWidgets(PropertyRowWidgets.WholeRowWidget.ToSharedRef());
+
+			FRowWidgets RowWidgets(PropertyRowWidgets.WholeRowWidget.ToSharedRef()); 
+			RowWidgets.EditConditionWidget = PropertyRowWidgets.EditConditionWidget;
+			return RowWidgets;
 		}
 		else
 		{
 			Container->SetOverrideNameWidth(PropertyRowWidgets.NameWidgetLayoutData.MinWidth, PropertyRowWidgets.NameWidgetLayoutData.MaxWidth);
-			Container->SetOverrideNameAlignment(PropertyRowWidgets.NameWidgetLayoutData.HorizontalAlignment, PropertyRowWidgets.NameWidgetLayoutData.VerticalAlignment);
+			//Container->SetOverrideNameAlignment(PropertyRowWidgets.NameWidgetLayoutData.HorizontalAlignment, PropertyRowWidgets.NameWidgetLayoutData.VerticalAlignment);
 			Container->SetOverrideValueWidth(PropertyRowWidgets.ValueWidgetLayoutData.MinWidth, PropertyRowWidgets.ValueWidgetLayoutData.MaxWidth);
-			Container->SetOverrideValueAlignment(PropertyRowWidgets.ValueWidgetLayoutData.HorizontalAlignment, PropertyRowWidgets.ValueWidgetLayoutData.VerticalAlignment);
+			//Container->SetOverrideValueAlignment(PropertyRowWidgets.ValueWidgetLayoutData.HorizontalAlignment, PropertyRowWidgets.ValueWidgetLayoutData.VerticalAlignment);
 			PropertyRowWidgets.NameWidget->SetEnabled(IsEnabled);
 			PropertyRowWidgets.ValueWidget->SetEnabled(IsEnabled);
-			return FRowWidgets(PropertyRowWidgets.NameWidget.ToSharedRef(), PropertyRowWidgets.ValueWidget.ToSharedRef());
+
+			FRowWidgets RowWidgets(PropertyRowWidgets.NameWidget.ToSharedRef(), PropertyRowWidgets.ValueWidget.ToSharedRef());
+			RowWidgets.EditConditionWidget = PropertyRowWidgets.EditConditionWidget;
+			return RowWidgets;
 		}
 	}
 	else if (Item->IsA<UNiagaraStackItem>())
diff --git a/Engine/Plugins/FX/Niagara/Source/NiagaraEditorWidgets/Private/SNiagaraStack.h b/Engine/Plugins/FX/Niagara/Source/NiagaraEditorWidgets/Private/SNiagaraStack.h
--- a/Engine/Plugins/FX/Niagara/Source/NiagaraEditorWidgets/Private/SNiagaraStack.h
+++ b/Engine/Plugins/FX/Niagara/Source/NiagaraEditorWidgets/Private/SNiagaraStack.h
@@ -46,6 +46,7 @@ private:
 
 		TSharedRef<SWidget> NameWidget;
 		TSharedPtr<SWidget> ValueWidget;
+		TSharedPtr<SWidget> EditConditionWidget;
 	};
 
 	void SynchronizeTreeExpansion();
diff --git a/Engine/Plugins/FX/Niagara/Source/NiagaraEditorWidgets/Private/Stack/SNiagaraStackTableRow.cpp b/Engine/Plugins/FX/Niagara/Source/NiagaraEditorWidgets/Private/Stack/SNiagaraStackTableRow.cpp
--- a/Engine/Plugins/FX/Niagara/Source/NiagaraEditorWidgets/Private/Stack/SNiagaraStackTableRow.cpp
+++ b/Engine/Plugins/FX/Niagara/Source/NiagaraEditorWidgets/Private/Stack/SNiagaraStackTableRow.cpp
@@ -130,7 +130,7 @@ bool HasVisibleChildren(UNiagaraStackEntry* Entry)
 	return Children.Num() > 0;
 }
 
-void SNiagaraStackTableRow::SetNameAndValueContent(TSharedRef<SWidget> InNameWidget, TSharedPtr<SWidget> InValueWidget)
+void SNiagaraStackTableRow::SetNameAndValueContent(TSharedRef<SWidget> InNameWidget, TSharedPtr<SWidget> InValueWidget, TSharedPtr<SWidget> InEditConditionWidget)
 {
 	FSlateColor IconColor = FNiagaraEditorWidgetsStyle::Get().GetColor(FNiagaraStackEditorWidgetsUtilities::GetColorNameForExecutionCategory(StackEntry->GetExecutionCategoryName()));
 	if (bIsCategoryIconHighlighted)
@@ -152,21 +152,17 @@ void SNiagaraStackTableRow::SetNameAndValueContent(TSharedRef<SWidget> InNameWid
 	.VAlign(VAlign_Center)
 	.Padding(0, 0, 1, 0)
 	[
-		SNew(SBox)
-		.WidthOverride(14)
+		SNew(SButton)
+		.ButtonStyle(FCoreStyle::Get(), "NoBorder")
+		.Visibility(this, &SNiagaraStackTableRow::GetExpanderVisibility)
+		.OnClicked(this, &SNiagaraStackTableRow::ExpandButtonClicked)
+		.ForegroundColor(FSlateColor::UseForeground())
+		.ContentPadding(2)
+		.HAlign(HAlign_Center)
 		[
-			SNew(SButton)
-			.ButtonStyle(FCoreStyle::Get(), "NoBorder")
-			.Visibility(this, &SNiagaraStackTableRow::GetExpanderVisibility)
-			.OnClicked(this, &SNiagaraStackTableRow::ExpandButtonClicked)
-			.ForegroundColor(FSlateColor::UseForeground())
-			.ContentPadding(2)
-			.HAlign(HAlign_Center)
-			[
-				SNew(SImage)
-				.Image(this, &SNiagaraStackTableRow::GetExpandButtonImage)
-				.ColorAndOpacity(FSlateColor::UseForeground())
-			]
+			SNew(SImage)
+			.Image(this, &SNiagaraStackTableRow::GetExpandButtonImage)
+			.ColorAndOpacity(FSlateColor::UseForeground())
 		]
 	]
 	// Execution sub-category icon
@@ -190,8 +186,17 @@ void SNiagaraStackTableRow::SetNameAndValueContent(TSharedRef<SWidget> InNameWid
 			.ColorAndOpacity(IconColor)
 		]
 	]
+	// Edit condition
+	+ SHorizontalBox::Slot()
+	.AutoWidth()
+	.Padding(0, 0, 3, 0)
+	[
+		InEditConditionWidget.IsValid() ? InEditConditionWidget.ToSharedRef() : SNullWidget::NullWidget
+	]
 	// Name content
 	+ SHorizontalBox::Slot()
+	.HAlign(NameHorizontalAlignment)
+	.VAlign(NameVerticalAlignment)
 	[
 		InNameWidget
 	];
@@ -239,8 +244,6 @@ void SNiagaraStackTableRow::SetNameAndValueContent(TSharedRef<SWidget> InNameWid
             [
 				SNew(SBox)
 				.Padding(FMargin(0, ContentPadding.Top, 5, ContentPadding.Bottom))
-				.HAlign(NameHorizontalAlignment)
-				.VAlign(NameVerticalAlignment)
 				.MinDesiredWidth(NameMinWidth.IsSet() ? NameMinWidth.GetValue() : FOptionalSize())
 				.MaxDesiredWidth(NameMaxWidth.IsSet() ? NameMaxWidth.GetValue() : FOptionalSize())
 				[
diff --git a/Engine/Plugins/FX/Niagara/Source/NiagaraEditorWidgets/Private/Stack/SNiagaraStackTableRow.h b/Engine/Plugins/FX/Niagara/Source/NiagaraEditorWidgets/Private/Stack/SNiagaraStackTableRow.h
--- a/Engine/Plugins/FX/Niagara/Source/NiagaraEditorWidgets/Private/Stack/SNiagaraStackTableRow.h
+++ b/Engine/Plugins/FX/Niagara/Source/NiagaraEditorWidgets/Private/Stack/SNiagaraStackTableRow.h
@@ -59,7 +59,7 @@ public:
 
 	void SetContentPadding(FMargin InContentPadding);
 
-	void SetNameAndValueContent(TSharedRef<SWidget> InNameWidget, TSharedPtr<SWidget> InValueWidget);
+	void SetNameAndValueContent(TSharedRef<SWidget> InNameWidget, TSharedPtr<SWidget> InValueWidget, TSharedPtr<SWidget> InEditConditionWidget);
 
 	void AddFillRowContextMenuHandler(FOnFillRowContextMenu FillRowContextMenuHandler);
 
diff --git a/Engine/Source/Editor/PropertyEditor/Private/DetailItemNode.cpp b/Engine/Source/Editor/PropertyEditor/Private/DetailItemNode.cpp
--- a/Engine/Source/Editor/PropertyEditor/Private/DetailItemNode.cpp
+++ b/Engine/Source/Editor/PropertyEditor/Private/DetailItemNode.cpp
@@ -276,7 +276,7 @@ bool FDetailItemNode::GenerateStandaloneWidget(FDetailWidgetRow& OutRow) const
 
 		bResult = true;
 	}
-	else if(Customization.IsValidCustomization())
+	else if (Customization.IsValidCustomization())
 	{
 		FDetailWidgetRow Row = Customization.GetWidgetRow();
 
@@ -323,6 +323,14 @@ bool FDetailItemNode::GenerateStandaloneWidget(FDetailWidgetRow& OutRow) const
 			];
 		}
 
+		if (Row.CustomResetToDefault.IsSet())
+		{
+			OutRow.OverrideResetToDefault(Row.CustomResetToDefault.GetValue());
+		}
+
+		OutRow.EditConditionValue = Row.EditConditionValue;
+		OutRow.OnEditConditionValueChanged = Row.OnEditConditionValueChanged;
+
 		bResult = true;
 	}
 
diff --git a/Engine/Source/Editor/PropertyEditor/Private/DetailTreeNode.cpp b/Engine/Source/Editor/PropertyEditor/Private/DetailTreeNode.cpp
--- a/Engine/Source/Editor/PropertyEditor/Private/DetailTreeNode.cpp
+++ b/Engine/Source/Editor/PropertyEditor/Private/DetailTreeNode.cpp
@@ -2,6 +2,7 @@
 
 #include "DetailTreeNode.h"
 #include "DetailWidgetRow.h"
+#include "PropertyEditorHelpers.h"
 
 FNodeWidgets FDetailTreeNode::CreateNodeWidgets() const
 {
@@ -10,7 +11,7 @@ FNodeWidgets FDetailTreeNode::CreateNodeWidgets() const
 
 	FNodeWidgets Widgets;
 
-	if(Row.HasAnyContent())
+	if (Row.HasAnyContent())
 	{
 		if (Row.HasColumns())
 		{
@@ -27,6 +28,10 @@ FNodeWidgets FDetailTreeNode::CreateNodeWidgets() const
 			Widgets.WholeRowWidgetLayoutData = FNodeWidgetLayoutData(
 				Row.WholeRowWidget.HorizontalAlignment, Row.WholeRowWidget.VerticalAlignment, Row.WholeRowWidget.MinWidth, Row.WholeRowWidget.MaxWidth);
 		}
+
+		Widgets.EditConditionWidget = SNew(SEditConditionWidget)
+			.EditConditionValue(Row.EditConditionValue)
+			.OnEditConditionValueChanged(Row.OnEditConditionValueChanged);
 	}
 
 	return Widgets;
diff --git a/Engine/Source/Editor/PropertyEditor/Public/IDetailTreeNode.h b/Engine/Source/Editor/PropertyEditor/Public/IDetailTreeNode.h
--- a/Engine/Source/Editor/PropertyEditor/Public/IDetailTreeNode.h
+++ b/Engine/Source/Editor/PropertyEditor/Public/IDetailTreeNode.h
@@ -68,6 +68,9 @@ struct FNodeWidgets
 
 	/** Layout data for the whole row widget. */
 	FNodeWidgetLayoutData WholeRowWidgetLayoutData;
+
+	/** Edit condition widget. */
+	TSharedPtr<SWidget> EditConditionWidget;
 };
 
 class IDetailTreeNode
