Replace a newly introduce legacy UnlockVertexBuffer() call with the unified UnlockBuffer()

diff --git a/Engine/Plugins/FX/Niagara/Source/Niagara/Private/NiagaraDataInterfaceMeshRendererInfo.cpp b/Engine/Plugins/FX/Niagara/Source/Niagara/Private/NiagaraDataInterfaceMeshRendererInfo.cpp
--- a/Engine/Plugins/FX/Niagara/Source/Niagara/Private/NiagaraDataInterfaceMeshRendererInfo.cpp
+++ b/Engine/Plugins/FX/Niagara/Source/Niagara/Private/NiagaraDataInterfaceMeshRendererInfo.cpp
@@ -83,7 +83,7 @@ public:
 				*BufferDataVector++ = Mesh.MinLocalBounds;
 				*BufferDataVector++ = Mesh.MaxLocalBounds;
 			}
-			RHIUnlockVertexBuffer(BufferMeshDataRHI);
+			RHIUnlockBuffer(BufferMeshDataRHI);
 
 			BufferMeshDataSRV = RHICreateShaderResourceView(BufferMeshDataRHI, sizeof(float), PF_R32_FLOAT);
 		}
