Editgrate null-check to fix an assertion failure

#rnx

diff --git a/Engine/Plugins/FX/Niagara/Source/Niagara/Private/NiagaraDataInterfaceMeshRendererInfo.cpp b/Engine/Plugins/FX/Niagara/Source/Niagara/Private/NiagaraDataInterfaceMeshRendererInfo.cpp
--- a/Engine/Plugins/FX/Niagara/Source/Niagara/Private/NiagaraDataInterfaceMeshRendererInfo.cpp
+++ b/Engine/Plugins/FX/Niagara/Source/Niagara/Private/NiagaraDataInterfaceMeshRendererInfo.cpp
@@ -261,11 +261,9 @@ void FNDIMeshRendererInfo::Release(UNiagaraMeshRendererProperties& Renderer, FND
 				}
 			);
 		}
-		else
-		{
-			Info = nullptr;
-		}
 	}
+
+	Info = nullptr;
 }
 
 void FNDIMeshRendererInfo::ResetMeshData(const UNiagaraMeshRendererProperties& Renderer, FMeshDataArray& OutMeshData)
@@ -538,9 +536,9 @@ bool UNiagaraDataInterfaceMeshRendererInfo::CopyToInternal(UNiagaraDataInterface
 
 void UNiagaraDataInterfaceMeshRendererInfo::PushToRenderThreadImpl()
 {
-	auto TypedProxy = GetProxyAs<FNDIMeshRendererInfoProxy>();
-	if (MeshRenderer)
+	if (MeshRenderer && Info.IsValid())
 	{
+		auto TypedProxy = GetProxyAs<FNDIMeshRendererInfoProxy>();
 		ENQUEUE_RENDER_COMMAND(FDIMeshRendererInfoPushToRT)
 		(
 			[TypedProxy, GPUData_RT=Info->GetOrCreateGPUData()](FRHICommandList& RHICmdList)
