Prototype Wrapped Object Pointers (ObjectHandle/ObjectPtr)

Call-site changes for non-scalar wrapped object pointer upgrades in Engine + Shootergame

diff --git a/Engine/Plugins/2D/Paper2D/Source/Paper2D/Private/PaperSprite.cpp b/Engine/Plugins/2D/Paper2D/Source/Paper2D/Private/PaperSprite.cpp
--- a/Engine/Plugins/2D/Paper2D/Source/Paper2D/Private/PaperSprite.cpp
+++ b/Engine/Plugins/2D/Paper2D/Source/Paper2D/Private/PaperSprite.cpp
@@ -1954,7 +1954,7 @@ UTexture2D* UPaperSprite::GetBakedTexture() const
 
 void UPaperSprite::GetBakedAdditionalSourceTextures(FAdditionalSpriteTextureArray& OutTextureList) const
 {
-	OutTextureList = AdditionalSourceTextures;
+	OutTextureList = ToRawPtrTArrayUnsafe(AdditionalSourceTextures);
 }
 
 UMaterialInterface* UPaperSprite::GetMaterial(int32 MaterialIndex) const
diff --git a/Engine/Plugins/Animation/GameplayInsights/Source/GameplayInsights/Private/GameplayTimingViewExtender.cpp b/Engine/Plugins/Animation/GameplayInsights/Source/GameplayInsights/Private/GameplayTimingViewExtender.cpp
--- a/Engine/Plugins/Animation/GameplayInsights/Source/GameplayInsights/Private/GameplayTimingViewExtender.cpp
+++ b/Engine/Plugins/Animation/GameplayInsights/Source/GameplayInsights/Private/GameplayTimingViewExtender.cpp
@@ -106,7 +106,7 @@ UWorld* FGameplayTimingViewExtender::GetWorldToVisualize()
 	if (GIsEditor && EditorEngine != nullptr && World == nullptr)
 	{
 		// lets use PlayWorld during PIE/Simulate and regular world from editor otherwise, to draw debug information
-		World = EditorEngine->PlayWorld != nullptr ? EditorEngine->PlayWorld : EditorEngine->GetEditorWorldContext().World();
+		World = EditorEngine->PlayWorld != nullptr ? ToRawPtr(EditorEngine->PlayWorld) : EditorEngine->GetEditorWorldContext().World();
 	}
 
 #endif
diff --git a/Engine/Plugins/Animation/LiveLinkCurveDebugUI/Source/LiveLinkCurveDebugUI/Private/LiveLinkCurveDebugUIModule.cpp b/Engine/Plugins/Animation/LiveLinkCurveDebugUI/Source/LiveLinkCurveDebugUI/Private/LiveLinkCurveDebugUIModule.cpp
--- a/Engine/Plugins/Animation/LiveLinkCurveDebugUI/Source/LiveLinkCurveDebugUI/Private/LiveLinkCurveDebugUIModule.cpp
+++ b/Engine/Plugins/Animation/LiveLinkCurveDebugUI/Source/LiveLinkCurveDebugUI/Private/LiveLinkCurveDebugUIModule.cpp
@@ -292,7 +292,7 @@ UWorld* FLiveLinkCurveDebugUIModule::GetWorldForDebugUIModule()
 	if (GIsEditor && EEngine != nullptr)
 	{
 		// lets use PlayWorld during PIE/Simulate and regular world from editor otherwise, to draw debug information
-		World = EEngine->PlayWorld != nullptr ? EEngine->PlayWorld : EEngine->GetEditorWorldContext().World();
+		World = EEngine->PlayWorld != nullptr ? ToRawPtr(EEngine->PlayWorld) : EEngine->GetEditorWorldContext().World();
 	}
 #endif
 
diff --git a/Engine/Plugins/Enterprise/DatasmithContent/Source/DatasmithContent/Private/DatasmithScene.cpp b/Engine/Plugins/Enterprise/DatasmithContent/Source/DatasmithContent/Private/DatasmithScene.cpp
--- a/Engine/Plugins/Enterprise/DatasmithContent/Source/DatasmithContent/Private/DatasmithScene.cpp
+++ b/Engine/Plugins/Enterprise/DatasmithContent/Source/DatasmithContent/Private/DatasmithScene.cpp
@@ -155,7 +155,7 @@ void UDatasmithScene::RemoveUserDataOfClass( TSubclassOf<UAssetUserData> InUserD
 const TArray<UAssetUserData*>* UDatasmithScene::GetAssetUserDataArray() const
 {
 #if WITH_EDITORONLY_DATA
-	return &AssetUserData;
+	return &ToRawPtrTArrayUnsafe(AssetUserData);
 #else
 	return NULL;
 #endif // #if WITH_EDITORONLY_DATA
diff --git a/Engine/Plugins/Enterprise/DatasmithContent/Source/DatasmithContent/Private/ObjectTemplates/DatasmithObjectTemplate.cpp b/Engine/Plugins/Enterprise/DatasmithContent/Source/DatasmithContent/Private/ObjectTemplates/DatasmithObjectTemplate.cpp
--- a/Engine/Plugins/Enterprise/DatasmithContent/Source/DatasmithContent/Private/ObjectTemplates/DatasmithObjectTemplate.cpp
+++ b/Engine/Plugins/Enterprise/DatasmithContent/Source/DatasmithContent/Private/ObjectTemplates/DatasmithObjectTemplate.cpp
@@ -76,7 +76,7 @@ TMap< TSubclassOf< UDatasmithObjectTemplate >, UDatasmithObjectTemplate* >* FDat
 #if WITH_EDITORONLY_DATA
 	if (UDatasmithAssetUserData* UserAssetData = FindOrCreateDatasmithUserData(Outer))
 	{
-		return &UserAssetData->ObjectTemplates;
+		return &ToRawPtrTMapUnsafe(UserAssetData->ObjectTemplates);
 	}
 #endif // #if WITH_EDITORONLY_DATA
 return nullptr;
diff --git a/Engine/Plugins/Enterprise/MDLImporter/Source/MDLImporter/Private/material/BakedMaterialFactory.cpp b/Engine/Plugins/Enterprise/MDLImporter/Source/MDLImporter/Private/material/BakedMaterialFactory.cpp
--- a/Engine/Plugins/Enterprise/MDLImporter/Source/MDLImporter/Private/material/BakedMaterialFactory.cpp
+++ b/Engine/Plugins/Enterprise/MDLImporter/Source/MDLImporter/Private/material/BakedMaterialFactory.cpp
@@ -58,7 +58,7 @@ namespace Mat
 		// get under clear coat output
 		UMaterialExpressionClearCoatNormalCustomOutput* UnderClearCoat = nullptr;
 		{
-			UMaterialExpression** Found = Material.Expressions.FindByPredicate(
+			UE_TRANSITIONAL_OBJECT_PTR(UMaterialExpression)* Found = Material.Expressions.FindByPredicate(
 			    [](const UMaterialExpression* Expr) { return Expr->IsA<UMaterialExpressionClearCoatNormalCustomOutput>(); });
 			if (Found)
 				UnderClearCoat = Cast<UMaterialExpressionClearCoatNormalCustomOutput>(*Found);
@@ -174,7 +174,7 @@ PRAGMA_ENABLE_DEPRECATION_WARNINGS
 	{
 		FString               Name;
 
-		UMaterialExpression** Found = Material.Expressions.FindByPredicate(
+		UE_TRANSITIONAL_OBJECT_PTR(UMaterialExpression)* Found = Material.Expressions.FindByPredicate(
 			[&Name](UMaterialExpression* Expression)  //
 		{
 			if (UMaterialExpressionScalarParameter* ScalarParameter = Cast<UMaterialExpressionScalarParameter>(Expression))
@@ -191,7 +191,7 @@ PRAGMA_ENABLE_DEPRECATION_WARNINGS
 			return Name.Find(TEXT("Tiling Factor")) != INDEX_NONE;
 		});
 
-		UMaterialExpression** FoundU = Material.Expressions.FindByPredicate(
+		UE_TRANSITIONAL_OBJECT_PTR(UMaterialExpression)* FoundU = Material.Expressions.FindByPredicate(
 		    [&Name](UMaterialExpression* Expression)  //
 		    {
 			    if (UMaterialExpressionScalarParameter* ScalarParameter = Cast<UMaterialExpressionScalarParameter>(Expression))
@@ -208,7 +208,7 @@ PRAGMA_ENABLE_DEPRECATION_WARNINGS
 			    return Name.Find(TEXT("U Tiling")) != INDEX_NONE;
 		    });
 
-		UMaterialExpression** FoundV = Material.Expressions.FindByPredicate(
+		UE_TRANSITIONAL_OBJECT_PTR(UMaterialExpression)* FoundV = Material.Expressions.FindByPredicate(
 			[&Name](UMaterialExpression* Expression)  //
 		{
 			if (UMaterialExpressionScalarParameter* ScalarParameter = Cast<UMaterialExpressionScalarParameter>(Expression))
diff --git a/Engine/Plugins/Enterprise/VariantManagerContent/Source/VariantManagerContent/Private/LevelVariantSets.cpp b/Engine/Plugins/Enterprise/VariantManagerContent/Source/VariantManagerContent/Private/LevelVariantSets.cpp
--- a/Engine/Plugins/Enterprise/VariantManagerContent/Source/VariantManagerContent/Private/LevelVariantSets.cpp
+++ b/Engine/Plugins/Enterprise/VariantManagerContent/Source/VariantManagerContent/Private/LevelVariantSets.cpp
@@ -281,7 +281,7 @@ UVariantSet* ULevelVariantSets::GetVariantSet(int32 VariantSetIndex)
 
 UVariantSet* ULevelVariantSets::GetVariantSetByName(FString VariantSetName)
 {
-	UVariantSet** VarSetPtr = VariantSets.FindByPredicate([VariantSetName](const UVariantSet* VarSet)
+	UE_TRANSITIONAL_OBJECT_PTR(UVariantSet)* VarSetPtr = VariantSets.FindByPredicate([VariantSetName](const UVariantSet* VarSet)
 	{
 		return VarSet->GetDisplayText().ToString() == VariantSetName;
 	});
diff --git a/Engine/Plugins/Enterprise/VariantManagerContent/Source/VariantManagerContent/Private/Variant.cpp b/Engine/Plugins/Enterprise/VariantManagerContent/Source/VariantManagerContent/Private/Variant.cpp
--- a/Engine/Plugins/Enterprise/VariantManagerContent/Source/VariantManagerContent/Private/Variant.cpp
+++ b/Engine/Plugins/Enterprise/VariantManagerContent/Source/VariantManagerContent/Private/Variant.cpp
@@ -234,7 +234,7 @@ AActor* UVariant::GetActor(int32 ActorIndex)
 
 UVariantObjectBinding* UVariant::GetBindingByName(const FString& ActorName)
 {
-	UVariantObjectBinding** FoundBindingPtr = ObjectBindings.FindByPredicate([&ActorName](const UVariantObjectBinding* Binding)
+	UE_TRANSITIONAL_OBJECT_PTR(UVariantObjectBinding)* FoundBindingPtr = ObjectBindings.FindByPredicate([&ActorName](const UVariantObjectBinding* Binding)
 	{
 		UObject* ThisActor = Binding->GetObject();
 		return ThisActor && ThisActor->GetName() == ActorName;
diff --git a/Engine/Plugins/Enterprise/VariantManagerContent/Source/VariantManagerContent/Private/VariantSet.cpp b/Engine/Plugins/Enterprise/VariantManagerContent/Source/VariantManagerContent/Private/VariantSet.cpp
--- a/Engine/Plugins/Enterprise/VariantManagerContent/Source/VariantManagerContent/Private/VariantSet.cpp
+++ b/Engine/Plugins/Enterprise/VariantManagerContent/Source/VariantManagerContent/Private/VariantSet.cpp
@@ -280,7 +280,7 @@ UVariant* UVariantSet::GetVariant(int32 VariantIndex)
 
 UVariant* UVariantSet::GetVariantByName(FString VariantName)
 {
-	UVariant** VarPtr = Variants.FindByPredicate([VariantName](const UVariant* Var)
+	UE_TRANSITIONAL_OBJECT_PTR(UVariant)* VarPtr = Variants.FindByPredicate([VariantName](const UVariant* Var)
 	{
 		return Var->GetDisplayText().ToString() == VariantName;
 	});
diff --git a/Engine/Plugins/Experimental/ChaosEditor/Source/FractureEditor/Private/FractureToolGenerators.cpp b/Engine/Plugins/Experimental/ChaosEditor/Source/FractureEditor/Private/FractureToolGenerators.cpp
--- a/Engine/Plugins/Experimental/ChaosEditor/Source/FractureEditor/Private/FractureToolGenerators.cpp
+++ b/Engine/Plugins/Experimental/ChaosEditor/Source/FractureEditor/Private/FractureToolGenerators.cpp
@@ -182,7 +182,7 @@ AGeometryCollectionActor* UFractureToolGenerateAsset::ConvertStaticMeshToGeometr
 
 				// Record the contributing source on the asset.
 				FSoftObjectPath SourceSoftObjectPath(ComponentStaticMesh);
-				TArray<UMaterialInterface*> SourceMaterials = StaticMeshComponent->GetMaterials();
+				decltype(FGeometryCollectionSource::SourceMaterial) SourceMaterials(StaticMeshComponent->GetMaterials());
 				FracturedGeometryCollection->GeometrySource.Add({ SourceSoftObjectPath, ComponentTransform, SourceMaterials });
 
 				FGeometryCollectionConversion::AppendStaticMesh(ComponentStaticMesh, SourceMaterials, ComponentTransform, FracturedGeometryCollection, true);
diff --git a/Engine/Plugins/Experimental/GPULightmass/Source/GPULightmass/Private/Scene/Scene.cpp b/Engine/Plugins/Experimental/GPULightmass/Source/GPULightmass/Private/Scene/Scene.cpp
--- a/Engine/Plugins/Experimental/GPULightmass/Source/GPULightmass/Private/Scene/Scene.cpp
+++ b/Engine/Plugins/Experimental/GPULightmass/Source/GPULightmass/Private/Scene/Scene.cpp
@@ -2037,7 +2037,7 @@ void FScene::ApplyFinishedLightmapsToWorld()
 		}
 
 		{
-			ULevel* SubLevelStorageLevel = LightingScenario ? LightingScenario : World->PersistentLevel;
+			ULevel* SubLevelStorageLevel = LightingScenario ? LightingScenario : ToRawPtr(World->PersistentLevel);
 			UMapBuildDataRegistry* SubLevelRegistry = SubLevelStorageLevel->GetOrCreateMapBuildData();
 			FPrecomputedVolumetricLightmapData& SubLevelData = SubLevelRegistry->AllocateLevelPrecomputedVolumetricLightmapBuildData(World->PersistentLevel->LevelBuildDataId);
 
diff --git a/Engine/Plugins/FX/Niagara/Source/Niagara/Private/MovieScene/MovieSceneNiagaraSystemTrack.cpp b/Engine/Plugins/FX/Niagara/Source/Niagara/Private/MovieScene/MovieSceneNiagaraSystemTrack.cpp
--- a/Engine/Plugins/FX/Niagara/Source/Niagara/Private/MovieScene/MovieSceneNiagaraSystemTrack.cpp
+++ b/Engine/Plugins/FX/Niagara/Source/Niagara/Private/MovieScene/MovieSceneNiagaraSystemTrack.cpp
@@ -29,7 +29,7 @@ FMovieSceneEvalTemplatePtr UMovieSceneNiagaraSystemTrack::CreateTemplateForSecti
 
 void UMovieSceneNiagaraSystemTrack::PostCompile(FMovieSceneEvaluationTrack& OutTrack, const FMovieSceneTrackCompilerArgs& Args) const
 {
-	UMovieSceneSection* const* SpawnSectionPtr = Sections.FindByPredicate([](const UMovieSceneSection* Section) { return Section->GetClass() == UMovieSceneNiagaraSystemSpawnSection::StaticClass(); });
+	UE_TRANSITIONAL_OBJECT_PTR(UMovieSceneSection) const* SpawnSectionPtr = Sections.FindByPredicate([](const UMovieSceneSection* Section) { return Section->GetClass() == UMovieSceneNiagaraSystemSpawnSection::StaticClass(); });
 	if (SpawnSectionPtr != nullptr)
 	{
 		UMovieSceneNiagaraSystemSpawnSection* SpawnSection = CastChecked<UMovieSceneNiagaraSystemSpawnSection>(*SpawnSectionPtr);
diff --git a/Engine/Plugins/FX/Niagara/Source/Niagara/Private/NiagaraComponentPool.cpp b/Engine/Plugins/FX/Niagara/Source/Niagara/Private/NiagaraComponentPool.cpp
--- a/Engine/Plugins/FX/Niagara/Source/Niagara/Private/NiagaraComponentPool.cpp
+++ b/Engine/Plugins/FX/Niagara/Source/Niagara/Private/NiagaraComponentPool.cpp
@@ -271,7 +271,7 @@ UNiagaraComponentPool::~UNiagaraComponentPool()
 
 void UNiagaraComponentPool::Cleanup(UWorld* World)
 {
-	for (TPair<UNiagaraSystem*, FNCPool>& Pool : WorldParticleSystemPools)
+	for (auto& Pool : WorldParticleSystemPools)
 	{
 		FNiagaraCrashReporterScope CRScope(Pool.Key);//In practice this may be null by now :(
 		Pool.Value.Cleanup();
@@ -455,7 +455,7 @@ void UNiagaraComponentPool::ReclaimWorldParticleSystem(UNiagaraComponent* Compon
 		if (CurrentTime - LastParticleSytemPoolCleanTime > GNiagaraSystemPoolingCleanTime)
 		{
 			LastParticleSytemPoolCleanTime = CurrentTime;
-			for (TPair<UNiagaraSystem*, FNCPool>& Pair : WorldParticleSystemPools)
+			for (auto& Pair : WorldParticleSystemPools)
 			{
 				Pair.Value.KillUnusedComponents(CurrentTime - GNiagaraSystemPoolKillUnusedTime, Asset);
 			}
@@ -550,7 +550,7 @@ void UNiagaraComponentPool::Dump()
 	FString DumpStr;
 
 	uint32 TotalMemUsage = 0;
-	for (TPair<UNiagaraSystem*, FNCPool>& Pair : WorldParticleSystemPools)
+	for (auto& Pair : WorldParticleSystemPools)
 	{
 		UNiagaraSystem* System = Pair.Key;
 		FNCPool& Pool = Pair.Value;		
diff --git a/Engine/Plugins/FX/Niagara/Source/Niagara/Private/NiagaraDataInterfaceRenderTarget2D.cpp b/Engine/Plugins/FX/Niagara/Source/Niagara/Private/NiagaraDataInterfaceRenderTarget2D.cpp
--- a/Engine/Plugins/FX/Niagara/Source/Niagara/Private/NiagaraDataInterfaceRenderTarget2D.cpp
+++ b/Engine/Plugins/FX/Niagara/Source/Niagara/Private/NiagaraDataInterfaceRenderTarget2D.cpp
@@ -461,7 +461,8 @@ void UNiagaraDataInterfaceRenderTarget2D::DestroyPerInstanceData(void* PerInstan
 	);
 
 	// Make sure to clear out the reference to the render target if we created one.
-	UTextureRenderTarget2D* ExistingRenderTarget = nullptr;
+	using RenderTargetType = decltype(decltype(ManagedRenderTargets)::ElementType::Value);
+	RenderTargetType ExistingRenderTarget = nullptr;
 	if ( ManagedRenderTargets.RemoveAndCopyValue(SystemInstance->GetId(), ExistingRenderTarget) && GNiagaraReleaseResourceOnRemove)
 	{
 		ExistingRenderTarget->ReleaseResource();
@@ -551,7 +552,8 @@ bool UNiagaraDataInterfaceRenderTarget2D::PerInstanceTickPostSimulate(void* PerI
 				{
 					InstanceData->TargetTexture = UserTargetTexture;
 
-					UTextureRenderTarget2D* ExistingRenderTarget = nullptr;
+					using RenderTargetType = decltype(decltype(ManagedRenderTargets)::ElementType::Value);
+					RenderTargetType ExistingRenderTarget = nullptr;
 					if ( ManagedRenderTargets.RemoveAndCopyValue(SystemInstance->GetId(), ExistingRenderTarget) && GNiagaraReleaseResourceOnRemove)
 					{
 						ExistingRenderTarget->ReleaseResource();
diff --git a/Engine/Plugins/FX/Niagara/Source/Niagara/Private/NiagaraDataInterfaceRenderTarget2DArray.cpp b/Engine/Plugins/FX/Niagara/Source/Niagara/Private/NiagaraDataInterfaceRenderTarget2DArray.cpp
--- a/Engine/Plugins/FX/Niagara/Source/Niagara/Private/NiagaraDataInterfaceRenderTarget2DArray.cpp
+++ b/Engine/Plugins/FX/Niagara/Source/Niagara/Private/NiagaraDataInterfaceRenderTarget2DArray.cpp
@@ -416,7 +416,8 @@ void UNiagaraDataInterfaceRenderTarget2DArray::DestroyPerInstanceData(void* PerI
 
 	// Make sure to clear out the reference to the render target if we created one.
 	extern int32 GNiagaraReleaseResourceOnRemove;
-	UTextureRenderTarget2DArray* ExistingRenderTarget = nullptr;
+	using RenderTargetType = decltype(decltype(ManagedRenderTargets)::ElementType::Value);
+	RenderTargetType ExistingRenderTarget = nullptr;
 	if (ManagedRenderTargets.RemoveAndCopyValue(SystemInstance->GetId(), ExistingRenderTarget) && GNiagaraReleaseResourceOnRemove)
 	{
 		ExistingRenderTarget->ReleaseResource();
@@ -508,7 +509,8 @@ bool UNiagaraDataInterfaceRenderTarget2DArray::PerInstanceTickPostSimulate(void*
 				InstanceData->TargetTexture = UserTargetTexture;
 
 				extern int32 GNiagaraReleaseResourceOnRemove;
-				UTextureRenderTarget2DArray* ExistingRenderTarget = nullptr;
+				using RenderTargetType = decltype(decltype(ManagedRenderTargets)::ElementType::Value);
+				RenderTargetType ExistingRenderTarget = nullptr;
 				if (ManagedRenderTargets.RemoveAndCopyValue(SystemInstance->GetId(), ExistingRenderTarget) && GNiagaraReleaseResourceOnRemove)
 				{
 					ExistingRenderTarget->ReleaseResource();
diff --git a/Engine/Plugins/FX/Niagara/Source/Niagara/Private/NiagaraDataInterfaceRenderTargetVolume.cpp b/Engine/Plugins/FX/Niagara/Source/Niagara/Private/NiagaraDataInterfaceRenderTargetVolume.cpp
--- a/Engine/Plugins/FX/Niagara/Source/Niagara/Private/NiagaraDataInterfaceRenderTargetVolume.cpp
+++ b/Engine/Plugins/FX/Niagara/Source/Niagara/Private/NiagaraDataInterfaceRenderTargetVolume.cpp
@@ -409,7 +409,8 @@ void UNiagaraDataInterfaceRenderTargetVolume::DestroyPerInstanceData(void* PerIn
 
 	// Make sure to clear out the reference to the render target if we created one.
 	extern int32 GNiagaraReleaseResourceOnRemove;
-	UTextureRenderTargetVolume* ExistingRenderTarget = nullptr;
+	using RenderTargetType = decltype(decltype(ManagedRenderTargets)::ElementType::Value);
+	RenderTargetType ExistingRenderTarget = nullptr;
 	if (ManagedRenderTargets.RemoveAndCopyValue(SystemInstance->GetId(), ExistingRenderTarget) && GNiagaraReleaseResourceOnRemove)
 	{
 		ExistingRenderTarget->ReleaseResource();
@@ -502,7 +503,8 @@ bool UNiagaraDataInterfaceRenderTargetVolume::PerInstanceTickPostSimulate(void*
 				InstanceData->TargetTexture = UserTargetTexture;
 
 				extern int32 GNiagaraReleaseResourceOnRemove;
-				UTextureRenderTargetVolume* ExistingRenderTarget = nullptr;
+				using RenderTargetType = decltype(decltype(ManagedRenderTargets)::ElementType::Value);
+				RenderTargetType ExistingRenderTarget = nullptr;
 				if (ManagedRenderTargets.RemoveAndCopyValue(SystemInstance->GetId(), ExistingRenderTarget) && GNiagaraReleaseResourceOnRemove)
 				{
 					ExistingRenderTarget->ReleaseResource();
diff --git a/Engine/Plugins/FX/Niagara/Source/Niagara/Private/NiagaraEmitter.cpp b/Engine/Plugins/FX/Niagara/Source/Niagara/Private/NiagaraEmitter.cpp
--- a/Engine/Plugins/FX/Niagara/Source/Niagara/Private/NiagaraEmitter.cpp
+++ b/Engine/Plugins/FX/Niagara/Source/Niagara/Private/NiagaraEmitter.cpp
@@ -269,8 +269,8 @@ UNiagaraEmitter* UNiagaraEmitter::DuplicateWithoutMerging(UObject* InOuter)
 {
 	UNiagaraEmitter* Duplicate;
 	{
-		TGuardValue<UNiagaraEmitter*> ParentGuard(Parent, nullptr);
-		TGuardValue<UNiagaraEmitter*> ParentAtLastMergeGuard(ParentAtLastMerge, nullptr);
+		TGuardValue<decltype(Parent)> ParentGuard(Parent, nullptr);
+		TGuardValue<decltype(ParentAtLastMerge)> ParentAtLastMergeGuard(ParentAtLastMerge, nullptr);
 		Duplicate = Cast<UNiagaraEmitter>(StaticDuplicateObject(this, InOuter));
 	}
 	return Duplicate;
@@ -1493,7 +1493,7 @@ void UNiagaraEmitter::UpdateFromMergedCopy(const INiagaraMergeManager& MergeMana
 	}
 
 	// Copy shader stages
-	for (UNiagaraSimulationStageBase*& SimulationStage : SimulationStages)
+	for (UE_TRANSITIONAL_OBJECT_PTR(UNiagaraSimulationStageBase)& SimulationStage : SimulationStages)
 	{
 		SimulationStage->OnChanged().RemoveAll(this);
 		SimulationStage->Script->RapidIterationParameters.RemoveAllOnChangedHandlers(this);
@@ -1688,7 +1688,7 @@ void UNiagaraEmitter::RemoveEventHandlerByUsageId(FGuid EventHandlerUsageId)
 
 UNiagaraSimulationStageBase* UNiagaraEmitter::GetSimulationStageById(FGuid ScriptUsageId) const
 {
-	UNiagaraSimulationStageBase*const* FoundSimulationStagePtr = SimulationStages.FindByPredicate([&ScriptUsageId](UNiagaraSimulationStageBase* SimulationStage) { return SimulationStage->Script->GetUsageId() == ScriptUsageId; });
+	UE_TRANSITIONAL_OBJECT_PTR(UNiagaraSimulationStageBase) const* FoundSimulationStagePtr = SimulationStages.FindByPredicate([&ScriptUsageId](UNiagaraSimulationStageBase* SimulationStage) { return SimulationStage->Script->GetUsageId() == ScriptUsageId; });
 	return FoundSimulationStagePtr != nullptr ? *FoundSimulationStagePtr : nullptr;
 }
 
diff --git a/Engine/Plugins/FX/Niagara/Source/Niagara/Private/NiagaraPlatformSet.cpp b/Engine/Plugins/FX/Niagara/Source/Niagara/Private/NiagaraPlatformSet.cpp
--- a/Engine/Plugins/FX/Niagara/Source/Niagara/Private/NiagaraPlatformSet.cpp
+++ b/Engine/Plugins/FX/Niagara/Source/Niagara/Private/NiagaraPlatformSet.cpp
@@ -259,7 +259,7 @@ void FNiagaraPlatformSet::GetOverridenDeviceProfiles(int32 QualityLevel, TArray<
 	int32 QLMask = CreateQualityLevelMask(QualityLevel);
 	for (const FNiagaraDeviceProfileStateEntry& Entry : DeviceProfileStates)
 	{
-		if (UObject** DeviceProfile = UDeviceProfileManager::Get().Profiles.FindByPredicate([&](UObject* CheckProfile) {return CheckProfile->GetFName() == Entry.ProfileName;}))
+		if (UE_TRANSITIONAL_OBJECT_PTR(UObject)* DeviceProfile = UDeviceProfileManager::Get().Profiles.FindByPredicate([&](UObject* CheckProfile) {return CheckProfile->GetFName() == Entry.ProfileName;}))
 		{
 			UDeviceProfile* Profile = CastChecked<UDeviceProfile>(*DeviceProfile);
 			//If this platform cannot change at runtime then we store all EQs in the state so that the device is still overridden if someone changes it's EQ CVar.
diff --git a/Engine/Plugins/FX/Niagara/Source/NiagaraEditor/Private/NiagaraGraph.cpp b/Engine/Plugins/FX/Niagara/Source/NiagaraEditor/Private/NiagaraGraph.cpp
--- a/Engine/Plugins/FX/Niagara/Source/NiagaraEditor/Private/NiagaraGraph.cpp
+++ b/Engine/Plugins/FX/Niagara/Source/NiagaraEditor/Private/NiagaraGraph.cpp
@@ -136,7 +136,7 @@ void UNiagaraGraph::PostLoad()
 		for (auto It = VariableToScriptVariable.CreateIterator(); It; ++It)
 		{
 			FNiagaraVariable Var = It.Key();
-			UNiagaraScriptVariable*& ScriptVar = It.Value();
+			UE_TRANSITIONAL_OBJECT_PTR(UNiagaraScriptVariable)& ScriptVar = It.Value();
 
 			if (ScriptVar == nullptr)
 			{
@@ -313,7 +313,7 @@ void UNiagaraGraph::PostLoad()
 
 		for (auto& VariableToScriptVariableItem : VariableToScriptVariable)
 		{
-			UNiagaraScriptVariable*& MetaData = VariableToScriptVariableItem.Value;
+			UE_TRANSITIONAL_OBJECT_PTR(UNiagaraScriptVariable)& MetaData = VariableToScriptVariableItem.Value;
 
 			// Migrate advanced display.
 			if (MetaData->Metadata.PropertyMetaData.Contains("AdvancedDisplay"))
@@ -1320,7 +1320,7 @@ UNiagaraScriptVariable* UNiagaraGraph::GetScriptVariable(FNiagaraVariable Parame
 	{
 		RefreshParameterReferences();
 	}
-	if (UNiagaraScriptVariable** FoundScriptVariable = VariableToScriptVariable.Find(Parameter))
+	if (UE_TRANSITIONAL_OBJECT_PTR(UNiagaraScriptVariable)* FoundScriptVariable = VariableToScriptVariable.Find(Parameter))
 	{
 		return *FoundScriptVariable;
 	}
@@ -1353,7 +1353,7 @@ UNiagaraScriptVariable* UNiagaraGraph::AddParameter(const FNiagaraVariable& Para
 		ParameterToReferencesMap.Add(Parameter, NewReferenceCollection);
 	}
 
-	UNiagaraScriptVariable** FoundScriptVariable = VariableToScriptVariable.Find(Parameter);
+	UE_TRANSITIONAL_OBJECT_PTR(UNiagaraScriptVariable)* FoundScriptVariable = VariableToScriptVariable.Find(Parameter);
 	if (!FoundScriptVariable)
 	{
 		Modify();
@@ -1396,7 +1396,7 @@ UNiagaraScriptVariable* UNiagaraGraph::AddParameter(FNiagaraVariable& Parameter,
 	Modify();
 
 	// First we need to determine the new name
-	UNiagaraScriptVariable** FoundScriptVariable = VariableToScriptVariable.Find(Parameter);
+	UE_TRANSITIONAL_OBJECT_PTR(UNiagaraScriptVariable)* FoundScriptVariable = VariableToScriptVariable.Find(Parameter);
 	if (!FoundScriptVariable)
 	{
 		if (Options.bMakeParameterNameUnique)
@@ -1550,8 +1550,8 @@ bool UNiagaraGraph::RenameParameterFromPin(const FNiagaraVariable& Parameter, FN
 	NewParameter.SetName(NewName);
 
 	FNiagaraVariableMetaData MetaData;
-	UNiagaraScriptVariable** OldScriptVariable = VariableToScriptVariable.Find(Parameter);
-	UNiagaraScriptVariable** NewScriptVariableFound = VariableToScriptVariable.Find(NewParameter);
+	UE_TRANSITIONAL_OBJECT_PTR(UNiagaraScriptVariable)* OldScriptVariable = VariableToScriptVariable.Find(Parameter);
+	UE_TRANSITIONAL_OBJECT_PTR(UNiagaraScriptVariable)* NewScriptVariableFound = VariableToScriptVariable.Find(NewParameter);
 
 	if (OldScriptVariable != nullptr)
 	{
@@ -1648,7 +1648,7 @@ bool UNiagaraGraph::RenameParameter(const FNiagaraVariable& Parameter, FName New
 	FNiagaraVariable NewParameter = Parameter;
 	NewParameter.SetName(NewName);
 
-	UNiagaraScriptVariable** OldScriptVariable = VariableToScriptVariable.Find(Parameter);
+	UE_TRANSITIONAL_OBJECT_PTR(UNiagaraScriptVariable)* OldScriptVariable = VariableToScriptVariable.Find(Parameter);
 	FNiagaraVariableMetaData OldMetaData;
 	if (OldScriptVariable)
 	{
@@ -1673,7 +1673,7 @@ bool UNiagaraGraph::RenameParameter(const FNiagaraVariable& Parameter, FName New
 		OldMetaData.SetScopeName(NewScopeName);
 	}
 		
-	UNiagaraScriptVariable** NewScriptVariableFound = VariableToScriptVariable.Find(NewParameter);
+	UE_TRANSITIONAL_OBJECT_PTR(UNiagaraScriptVariable)* NewScriptVariableFound = VariableToScriptVariable.Find(NewParameter);
 
 	// Swap metadata to the new parameter; put the new parameter into VariableToScriptVariable
 	if (OldScriptVariable)
@@ -1942,7 +1942,7 @@ bool UNiagaraGraph::AppendCompileHash(FNiagaraCompileHashVisitor* InVisitor, con
 
 	// We need to sort the variables in a stable manner.
 	TArray<UNiagaraScriptVariable*> Values;
-	VariableToScriptVariable.GenerateValueArray(Values);
+	ToRawPtrTMapUnsafe(VariableToScriptVariable).GenerateValueArray(Values);
 	Values.Remove(nullptr);
 	Values.Sort([&](const UNiagaraScriptVariable& A, const UNiagaraScriptVariable& B)
 	{
@@ -2520,7 +2520,7 @@ void UNiagaraGraph::MarkGraphRequiresSynchronization(FString Reason)
 
 TOptional<FNiagaraVariableMetaData> UNiagaraGraph::GetMetaData(const FNiagaraVariable& InVar) const
 {
-	if (UNiagaraScriptVariable** MetaData = VariableToScriptVariable.Find(InVar))
+	if (UE_TRANSITIONAL_OBJECT_PTR(UNiagaraScriptVariable)* MetaData = VariableToScriptVariable.Find(InVar))
 	{
 		if (*MetaData)
 		{
@@ -2539,7 +2539,7 @@ void UNiagaraGraph::SetMetaData(const FNiagaraVariable& InVar, const FNiagaraVar
 	//InMetaData.GetParameterName(CachedName);
 	//UE_LOG(LogNiagaraEditor, Log, TEXT("SetMetaData %s %s!"), *InVar.GetName().ToString(), *CachedName.ToString());
 
-	if (UNiagaraScriptVariable** FoundMetaData = VariableToScriptVariable.Find(InVar))
+	if (UE_TRANSITIONAL_OBJECT_PTR(UNiagaraScriptVariable)* FoundMetaData = VariableToScriptVariable.Find(InVar))
 	{
 		if (*FoundMetaData)
 		{
@@ -2551,7 +2551,7 @@ void UNiagaraGraph::SetMetaData(const FNiagaraVariable& InVar, const FNiagaraVar
 	else 
 	{
 		Modify();
-		UNiagaraScriptVariable*& NewScriptVariable = VariableToScriptVariable.Add(InVar, NewObject<UNiagaraScriptVariable>(this, FName(), RF_Transactional));
+		UE_TRANSITIONAL_OBJECT_PTR(UNiagaraScriptVariable)& NewScriptVariable = VariableToScriptVariable.Add(InVar, NewObject<UNiagaraScriptVariable>(this, FName(), RF_Transactional));
 		NewScriptVariable->Variable = InVar;
 		NewScriptVariable->Metadata = InMetaData;
 	}
@@ -2559,7 +2559,7 @@ void UNiagaraGraph::SetMetaData(const FNiagaraVariable& InVar, const FNiagaraVar
 
 void UNiagaraGraph::SetPerScriptMetaData(const FNiagaraVariable& InVar, const FNiagaraVariableMetaData& InMetaData)
 {
-	if (UNiagaraScriptVariable** FoundMetaData = VariableToScriptVariable.Find(InVar))
+	if (UE_TRANSITIONAL_OBJECT_PTR(UNiagaraScriptVariable)* FoundMetaData = VariableToScriptVariable.Find(InVar))
 	{
 		if (*FoundMetaData)
 		{
@@ -2575,7 +2575,7 @@ void UNiagaraGraph::SetPerScriptMetaData(const FNiagaraVariable& InVar, const FN
 	else
 	{
 		Modify();
-		UNiagaraScriptVariable*& NewScriptVariable = VariableToScriptVariable.Add(InVar, NewObject<UNiagaraScriptVariable>(this, FName(), RF_Transactional));
+		UE_TRANSITIONAL_OBJECT_PTR(UNiagaraScriptVariable)& NewScriptVariable = VariableToScriptVariable.Add(InVar, NewObject<UNiagaraScriptVariable>(this, FName(), RF_Transactional));
 		NewScriptVariable->Variable = InVar;
 		NewScriptVariable->Metadata = InMetaData;
 	}
@@ -2735,7 +2735,7 @@ void UNiagaraGraph::RefreshParameterReferences() const
 	
 	auto AddMissingScriptVariable = [&](const FNiagaraVariable& Variable)
 	{
-		UNiagaraScriptVariable** FoundScriptVariable = VariableToScriptVariable.Find(Variable);
+		UE_TRANSITIONAL_OBJECT_PTR(UNiagaraScriptVariable)* FoundScriptVariable = VariableToScriptVariable.Find(Variable);
 		if (!FoundScriptVariable)
 		{
 			UNiagaraScriptVariable* NewScriptVariable = NewObject<UNiagaraScriptVariable>(const_cast<UNiagaraGraph*>(this));
diff --git a/Engine/Plugins/FX/Niagara/Source/NiagaraEditor/Private/NiagaraMessages.cpp b/Engine/Plugins/FX/Niagara/Source/NiagaraEditor/Private/NiagaraMessages.cpp
--- a/Engine/Plugins/FX/Niagara/Source/NiagaraEditor/Private/NiagaraMessages.cpp
+++ b/Engine/Plugins/FX/Niagara/Source/NiagaraEditor/Private/NiagaraMessages.cpp
@@ -251,7 +251,7 @@ bool FNiagaraMessageJobCompileEvent::RecursiveGetScriptNamesAndAssetPathsFromCon
 
 	if (NodeGuid.IsValid())
 	{
-		UEdGraphNode* const* EventNodePtr = InGraphToSearch->Nodes.FindByPredicate([NodeGuid](UEdGraphNode* Node) { return Node->NodeGuid == NodeGuid; });
+		UE_TRANSITIONAL_OBJECT_PTR(UEdGraphNode) const* EventNodePtr = InGraphToSearch->Nodes.FindByPredicate([NodeGuid](UEdGraphNode* Node) { return Node->NodeGuid == NodeGuid; });
 		if (EventNodePtr != nullptr)
 		{
 			OutContextNodeObjectKeys.AddUnique(FObjectKey(*EventNodePtr));
diff --git a/Engine/Plugins/FX/Niagara/Source/NiagaraEditor/Private/Sequencer/NiagaraSequence/MovieSceneNiagaraEmitterTrack.cpp b/Engine/Plugins/FX/Niagara/Source/NiagaraEditor/Private/Sequencer/NiagaraSequence/MovieSceneNiagaraEmitterTrack.cpp
--- a/Engine/Plugins/FX/Niagara/Source/NiagaraEditor/Private/Sequencer/NiagaraSequence/MovieSceneNiagaraEmitterTrack.cpp
+++ b/Engine/Plugins/FX/Niagara/Source/NiagaraEditor/Private/Sequencer/NiagaraSequence/MovieSceneNiagaraEmitterTrack.cpp
@@ -248,7 +248,7 @@ void UMovieSceneNiagaraEmitterTrack::CreateSections(const FFrameRate& InFrameRes
 					UMovieSceneNiagaraEmitterSectionBase* EmitterSection = nullptr;
 					bool bSectionCreated = false;
 
-					UMovieSceneSection** SharedEmitterSection = Sections.FindByPredicate([=](UMovieSceneSection* Section) {	return Section->IsA(SectionClass); });
+					UE_TRANSITIONAL_OBJECT_PTR(UMovieSceneSection)* SharedEmitterSection = Sections.FindByPredicate([=](UMovieSceneSection* Section) {	return Section->IsA(SectionClass); });
 					if (SharedEmitterSection != nullptr)
 					{
 						EmitterSection = CastChecked<UMovieSceneNiagaraEmitterSectionBase>(*SharedEmitterSection);
diff --git a/Engine/Plugins/FX/Niagara/Source/NiagaraEditor/Private/ViewModels/NiagaraScratchPadUtilities.cpp b/Engine/Plugins/FX/Niagara/Source/NiagaraEditor/Private/ViewModels/NiagaraScratchPadUtilities.cpp
--- a/Engine/Plugins/FX/Niagara/Source/NiagaraEditor/Private/ViewModels/NiagaraScratchPadUtilities.cpp
+++ b/Engine/Plugins/FX/Niagara/Source/NiagaraEditor/Private/ViewModels/NiagaraScratchPadUtilities.cpp
@@ -54,7 +54,8 @@ void FNiagaraScratchPadUtilities::FixExternalScratchPadScriptsForEmitter(UNiagar
 
 			// Determine the destination for the copies of the external scratch pad scripts.
 			UObject* TargetScratchPadScriptOuter = nullptr;
-			TArray<UNiagaraScript*>* TargetScratchPadScriptArray = nullptr;
+			using ScriptArrayType = decltype(TargetEmitter.ScratchPadScripts);
+			ScriptArrayType* TargetScratchPadScriptArray = nullptr;
 			if (TargetSystem != nullptr)
 			{
 				TargetScratchPadScriptOuter = TargetSystem;
diff --git a/Engine/Plugins/FX/Niagara/Source/NiagaraEditor/Private/ViewModels/NiagaraScratchPadViewModel.cpp b/Engine/Plugins/FX/Niagara/Source/NiagaraEditor/Private/ViewModels/NiagaraScratchPadViewModel.cpp
--- a/Engine/Plugins/FX/Niagara/Source/NiagaraEditor/Private/ViewModels/NiagaraScratchPadViewModel.cpp
+++ b/Engine/Plugins/FX/Niagara/Source/NiagaraEditor/Private/ViewModels/NiagaraScratchPadViewModel.cpp
@@ -60,7 +60,7 @@ void GetOuterAndTargetScripts(TSharedRef<FNiagaraSystemViewModel> SystemViewMode
 	if (SystemViewModel->GetEditMode() == ENiagaraSystemViewModelEditMode::SystemAsset)
 	{
 		OutOuter = &SystemViewModel->GetSystem();
-		OutTargetScripts = &SystemViewModel->GetSystem().ScratchPadScripts;
+		OutTargetScripts = &ToRawPtrTArrayUnsafe(SystemViewModel->GetSystem().ScratchPadScripts);
 	}
 	else
 	{
@@ -68,7 +68,7 @@ void GetOuterAndTargetScripts(TSharedRef<FNiagaraSystemViewModel> SystemViewMode
 		{
 			UNiagaraEmitter* TargetEmitter = SystemViewModel->GetSystem().GetEmitterHandles()[0].GetInstance();
 			OutOuter = TargetEmitter;
-			OutTargetScripts = &TargetEmitter->ScratchPadScripts;
+			OutTargetScripts = &ToRawPtrTArrayUnsafe(TargetEmitter->ScratchPadScripts);
 		}
 		else
 		{
diff --git a/Engine/Plugins/FX/Niagara/Source/NiagaraEditor/Private/ViewModels/Stack/NiagaraStackEntry.cpp b/Engine/Plugins/FX/Niagara/Source/NiagaraEditor/Private/ViewModels/Stack/NiagaraStackEntry.cpp
--- a/Engine/Plugins/FX/Niagara/Source/NiagaraEditor/Private/ViewModels/Stack/NiagaraStackEntry.cpp
+++ b/Engine/Plugins/FX/Niagara/Source/NiagaraEditor/Private/ViewModels/Stack/NiagaraStackEntry.cpp
@@ -672,7 +672,7 @@ void UNiagaraStackEntry::RefreshStackErrorChildren()
 	for (FStackIssue Issue : StackIssues)
 	{
 		UNiagaraStackErrorItem* ErrorEntry = nullptr;
-		UNiagaraStackErrorItem** Found = ErrorChildren.FindByPredicate(
+		UE_TRANSITIONAL_OBJECT_PTR(UNiagaraStackErrorItem)* Found = ErrorChildren.FindByPredicate(
 			[&](UNiagaraStackErrorItem* CurrentChild) { return CurrentChild->GetStackIssue().GetUniqueIdentifier() == Issue.GetUniqueIdentifier(); });
 		if (Found == nullptr)
 		{
diff --git a/Engine/Plugins/FX/Niagara/Source/NiagaraEditor/Private/Widgets/SNiagaraScriptGraph.cpp b/Engine/Plugins/FX/Niagara/Source/NiagaraEditor/Private/Widgets/SNiagaraScriptGraph.cpp
--- a/Engine/Plugins/FX/Niagara/Source/NiagaraEditor/Private/Widgets/SNiagaraScriptGraph.cpp
+++ b/Engine/Plugins/FX/Niagara/Source/NiagaraEditor/Private/Widgets/SNiagaraScriptGraph.cpp
@@ -454,7 +454,7 @@ void SNiagaraScriptGraph::FocusGraphElement(const INiagaraScriptGraphFocusInfo*
 	{
 		const FNiagaraScriptGraphNodeToFocusInfo* NodeFocusInfo = static_cast<const FNiagaraScriptGraphNodeToFocusInfo*>(FocusInfo);
 		const FGuid& NodeGuidToMatch = NodeFocusInfo->GetNodeGuidToFocus();
-		UEdGraphNode* const* NodeToFocus = ViewModel->GetGraph()->Nodes.FindByPredicate([&NodeGuidToMatch](const UEdGraphNode* Node) {return Node->NodeGuid == NodeGuidToMatch; });
+		UE_TRANSITIONAL_OBJECT_PTR(UEdGraphNode) const* NodeToFocus = ViewModel->GetGraph()->Nodes.FindByPredicate([&NodeGuidToMatch](const UEdGraphNode* Node) {return Node->NodeGuid == NodeGuidToMatch; });
 		if (NodeToFocus != nullptr && *NodeToFocus != nullptr)
 		{
 			GetGraphEditor()->JumpToNode(*NodeToFocus);
diff --git a/Engine/Plugins/Lumin/MagicLeapPassableWorld/Source/MagicLeapARPin/Private/MagicLeapARPinRenderer.cpp b/Engine/Plugins/Lumin/MagicLeapPassableWorld/Source/MagicLeapARPin/Private/MagicLeapARPinRenderer.cpp
--- a/Engine/Plugins/Lumin/MagicLeapPassableWorld/Source/MagicLeapARPin/Private/MagicLeapARPinRenderer.cpp
+++ b/Engine/Plugins/Lumin/MagicLeapPassableWorld/Source/MagicLeapARPin/Private/MagicLeapARPinRenderer.cpp
@@ -59,7 +59,7 @@ void AMagicLeapARPinRenderer::OnARPinsUpdated(const TArray<FGuid>& Added, const
 {
 	for (const FGuid& DeletedPinID : Deleted)
 	{
-		AMagicLeapARPinInfoActorBase** ActorPtr = AllInfoActors.Find(DeletedPinID);
+		UE_TRANSITIONAL_OBJECT_PTR(AMagicLeapARPinInfoActorBase)* ActorPtr = AllInfoActors.Find(DeletedPinID);
 		if (ActorPtr != nullptr && *ActorPtr != nullptr)
 		{
 			(*ActorPtr)->K2_DestroyActor();
@@ -81,7 +81,7 @@ void AMagicLeapARPinRenderer::OnARPinsUpdated(const TArray<FGuid>& Added, const
 
 	for (const FGuid& UpdatedPinID : Updated)
 	{
-		AMagicLeapARPinInfoActorBase** ActorPtr = AllInfoActors.Find(UpdatedPinID);
+		UE_TRANSITIONAL_OBJECT_PTR(AMagicLeapARPinInfoActorBase)* ActorPtr = AllInfoActors.Find(UpdatedPinID);
 		if (ActorPtr != nullptr && *ActorPtr != nullptr)
 		{
 			(*ActorPtr)->OnUpdateARPinState();
diff --git a/Engine/Plugins/MeshPainting/Source/MeshPaintingToolset/Private/MeshTexturePaintingTool.cpp b/Engine/Plugins/MeshPainting/Source/MeshPaintingToolset/Private/MeshTexturePaintingTool.cpp
--- a/Engine/Plugins/MeshPainting/Source/MeshPaintingToolset/Private/MeshTexturePaintingTool.cpp
+++ b/Engine/Plugins/MeshPainting/Source/MeshPaintingToolset/Private/MeshTexturePaintingTool.cpp
@@ -1154,7 +1154,7 @@ void UMeshTexturePaintingTool::FinishPaintingTexture()
 
 void UMeshTexturePaintingTool::CycleTextures(int32 Direction)
 {
-	UTexture2D*& SelectedTexture = TextureProperties->PaintTexture;
+	UE_TRANSITIONAL_OBJECT_PTR(UTexture2D)& SelectedTexture = TextureProperties->PaintTexture;
 
 	const int32 TextureIndex = (SelectedTexture != nullptr) ? PaintableTextures.IndexOfByKey(SelectedTexture) : 0;
 	if (TextureIndex != INDEX_NONE)
@@ -1187,7 +1187,7 @@ void UMeshTexturePaintingTool::CommitAllPaintedTextures()
 		int32 CurStep = 1;
 		int32 TotalSteps = GetNumberOfPendingPaintChanges();
 
-		for (TMap< UTexture2D*, FPaintTexture2DData >::TIterator It(PaintTargetData); It; ++It)
+		for (decltype(PaintTargetData)::TIterator It(PaintTargetData); It; ++It)
 		{
 			FPaintTexture2DData* TextureData = &It.Value();
 
@@ -1251,7 +1251,7 @@ void UMeshTexturePaintingTool::ClearAllTextureOverrides()
 {
 	const auto FeatureLevel = ERHIFeatureLevel::SM5;// GEditor->GetEditorWorldContext().World()->FeatureLevel;
 	/** Remove all texture overrides which are currently stored and active */
-	for (TMap< UTexture2D*, FPaintTexture2DData >::TIterator It(PaintTargetData); It; ++It)
+	for (decltype(PaintTargetData)::TIterator It(PaintTargetData); It; ++It)
 	{
 		FPaintTexture2DData* TextureData = &It.Value();
 
@@ -1268,7 +1268,7 @@ void UMeshTexturePaintingTool::ClearAllTextureOverrides()
 int32 UMeshTexturePaintingTool::GetNumberOfPendingPaintChanges() const
 {
 	int32 Result = 0;
-	for (TMap< UTexture2D*, FPaintTexture2DData >::TConstIterator It(PaintTargetData); It; ++It)
+	for (decltype(PaintTargetData)::TConstIterator It(PaintTargetData); It; ++It)
 	{
 		const FPaintTexture2DData* TextureData = &It.Value();
 
diff --git a/Engine/Plugins/Runtime/ProceduralMeshComponent/Source/ProceduralMeshComponent/Private/ProceduralMeshComponent.cpp b/Engine/Plugins/Runtime/ProceduralMeshComponent/Source/ProceduralMeshComponent/Private/ProceduralMeshComponent.cpp
--- a/Engine/Plugins/Runtime/ProceduralMeshComponent/Source/ProceduralMeshComponent/Private/ProceduralMeshComponent.cpp
+++ b/Engine/Plugins/Runtime/ProceduralMeshComponent/Source/ProceduralMeshComponent/Private/ProceduralMeshComponent.cpp
@@ -998,7 +998,7 @@ void UProceduralMeshComponent::UpdateCollision()
 		CreateProcMeshBodySetup();
 	}
 	
-	UBodySetup* UseBodySetup = bUseAsyncCook ? AsyncBodySetupQueue.Last() : ToRawPtr(ProcMeshBodySetup);
+	UBodySetup* UseBodySetup = bUseAsyncCook ? AsyncBodySetupQueue.Last() : ProcMeshBodySetup;
 
 	// Fill in simple collision convex elements
 	UseBodySetup->AggGeom.ConvexElems = CollisionConvexElems;
diff --git a/Engine/Plugins/Runtime/Steam/SteamSockets/Source/SteamSockets/Private/SteamSocketsNetDriver.cpp b/Engine/Plugins/Runtime/Steam/SteamSockets/Source/SteamSockets/Private/SteamSocketsNetDriver.cpp
--- a/Engine/Plugins/Runtime/Steam/SteamSockets/Source/SteamSockets/Private/SteamSocketsNetDriver.cpp
+++ b/Engine/Plugins/Runtime/Steam/SteamSockets/Source/SteamSockets/Private/SteamSocketsNetDriver.cpp
@@ -290,7 +290,7 @@ void USteamSocketsNetDriver::TickDispatch(float DeltaTime)
 			if (Socket->RecvRaw(Message, 1, BytesRead) && BytesRead > 0 && Message != nullptr)
 			{
 				USteamSocketsNetConnection* ConnectionToHandleMessage = static_cast<USteamSocketsNetConnection*>((bIsAServer) ? 
-					FindClientConnectionForHandle(Message->m_conn) : ServerConnection);
+					FindClientConnectionForHandle(Message->m_conn) : ToRawPtr(ServerConnection));
 
 				// Grab sender information for the purposes of logging
 				FInternetAddrSteamSockets MessageSender(Message->m_identityPeer);
@@ -416,7 +416,7 @@ void USteamSocketsNetDriver::ResetSocketInfo(const FSteamSocket* RemovedSocket)
 {
 	const SteamSocketHandles SocketHandle = RemovedSocket->InternalHandle;
 	USteamSocketsNetConnection* SocketConnection = 
-		static_cast<USteamSocketsNetConnection*>(ServerConnection ? ServerConnection : FindClientConnectionForHandle(SocketHandle));
+		static_cast<USteamSocketsNetConnection*>(ServerConnection ? ToRawPtr(ServerConnection) : FindClientConnectionForHandle(SocketHandle));
 
 	if (SocketConnection)
 	{
@@ -537,7 +537,7 @@ void USteamSocketsNetDriver::OnConnectionUpdated(SteamSocketHandles SocketHandle
 	// care about the connected state flag.
 	if (NewState == k_ESteamNetworkingConnectionState_Connected)
 	{
-		UNetConnection* SocketConnection = (ServerConnection) ? ServerConnection : FindClientConnectionForHandle(SocketHandle);
+		UNetConnection* SocketConnection = (ServerConnection) ? ToRawPtr(ServerConnection) : FindClientConnectionForHandle(SocketHandle);
 		if (SocketConnection)
 		{
 			SocketConnection->State = USOCK_Open;
@@ -560,7 +560,7 @@ void USteamSocketsNetDriver::OnConnectionUpdated(SteamSocketHandles SocketHandle
 
 void USteamSocketsNetDriver::OnConnectionDisconnected(SteamSocketHandles SocketHandle)
 {
-	UNetConnection* SocketConnection = ServerConnection ? ServerConnection : FindClientConnectionForHandle(SocketHandle);
+	UNetConnection* SocketConnection = ServerConnection ? ToRawPtr(ServerConnection) : FindClientConnectionForHandle(SocketHandle);
 	if (SocketConnection)
 	{
 		SocketConnection->State = USOCK_Closed;
diff --git a/Engine/Source/Developer/FunctionalTesting/Private/FunctionalTest.cpp b/Engine/Source/Developer/FunctionalTesting/Private/FunctionalTest.cpp
--- a/Engine/Source/Developer/FunctionalTesting/Private/FunctionalTest.cpp
+++ b/Engine/Source/Developer/FunctionalTesting/Private/FunctionalTest.cpp
@@ -381,7 +381,7 @@ void AFunctionalTest::FinishTest(EFunctionalTestResult TestResult, const FString
 
 	OnTestFinished.Broadcast();
 
-	AActor** ActorToDestroy = AutoDestroyActors.GetData();
+	UE_TRANSITIONAL_OBJECT_PTR(AActor)* ActorToDestroy = AutoDestroyActors.GetData();
 
 	for (int32 ActorIndex = 0; ActorIndex < AutoDestroyActors.Num(); ++ActorIndex, ++ActorToDestroy)
 	{
diff --git a/Engine/Source/Developer/GameplayDebugger/Private/GameplayDebuggerPlayerManager.cpp b/Engine/Source/Developer/GameplayDebugger/Private/GameplayDebuggerPlayerManager.cpp
--- a/Engine/Source/Developer/GameplayDebugger/Private/GameplayDebuggerPlayerManager.cpp
+++ b/Engine/Source/Developer/GameplayDebugger/Private/GameplayDebuggerPlayerManager.cpp
@@ -183,7 +183,7 @@ void AGameplayDebuggerPlayerManager::RegisterReplicator(AGameplayDebuggerCategor
 	{
 		APlayerController* OwnerPC = Replicator.GetReplicationOwner();
 
-		NewData.InputComponent = OwnerPC ? NewObject<UInputComponent>(OwnerPC, TEXT("GameplayDebug_Input")) : InputComponent;
+		NewData.InputComponent = OwnerPC ? NewObject<UInputComponent>(OwnerPC, TEXT("GameplayDebug_Input")) : ToRawPtr(InputComponent);
 		check(NewData.InputComponent);
 		NewData.InputComponent->Priority = -1;
 		NewData.Controller = NewObject<UGameplayDebuggerLocalController>(OwnerPC ? OwnerPC : (AActor*)this, TEXT("GameplayDebug_Controller"));
diff --git a/Engine/Source/Developer/ToolMenus/Private/ToolMenus.cpp b/Engine/Source/Developer/ToolMenus/Private/ToolMenus.cpp
--- a/Engine/Source/Developer/ToolMenus/Private/ToolMenus.cpp
+++ b/Engine/Source/Developer/ToolMenus/Private/ToolMenus.cpp
@@ -394,13 +394,13 @@ bool UToolMenus::GetDisplayUIExtensionPoints() const
 
 UToolMenu* UToolMenus::FindMenu(const FName Name)
 {
-	UToolMenu** Found = Menus.Find(Name);
+	UE_TRANSITIONAL_OBJECT_PTR(UToolMenu)* Found = Menus.Find(Name);
 	return Found ? *Found : nullptr;
 }
 
 bool UToolMenus::IsMenuRegistered(const FName Name) const
 {
-	const UToolMenu* const * Found = Menus.Find(Name);
+	UE_TRANSITIONAL_OBJECT_PTR(UToolMenu) const * Found = Menus.Find(Name);
 	return Found && *Found && (*Found)->IsRegistered();
 }
 
diff --git a/Engine/Source/Editor/BehaviorTreeEditor/Private/SBehaviorTreeDiff.cpp b/Engine/Source/Editor/BehaviorTreeEditor/Private/SBehaviorTreeDiff.cpp
--- a/Engine/Source/Editor/BehaviorTreeEditor/Private/SBehaviorTreeDiff.cpp
+++ b/Engine/Source/Editor/BehaviorTreeEditor/Private/SBehaviorTreeDiff.cpp
@@ -369,7 +369,7 @@ void SBehaviorTreeDiff::OnSelectionChanged(FSharedDiffOnGraph Item, ESelectInfo:
 				{
 					// This is a sub-node, we need to find our parent node in the graph
 					// todo: work out why BTNode->ParentNode is always null
-					UEdGraphNode** ParentNodePtr = NodeGraph->Nodes.FindByPredicate([BTNode](UEdGraphNode* PotentialParentNode) -> bool
+					UE_TRANSITIONAL_OBJECT_PTR(UEdGraphNode)* ParentNodePtr = NodeGraph->Nodes.FindByPredicate([BTNode](UEdGraphNode* PotentialParentNode) -> bool
 					{
 						UBehaviorTreeGraphNode* BTPotentialParentNode = Cast<UBehaviorTreeGraphNode>(PotentialParentNode);
 						return BTPotentialParentNode && (BTPotentialParentNode->Decorators.Contains(BTNode) || BTPotentialParentNode->Services.Contains(BTNode));
diff --git a/Engine/Source/Editor/DetailCustomizations/Private/DialogueStructsCustomizations.cpp b/Engine/Source/Editor/DetailCustomizations/Private/DialogueStructsCustomizations.cpp
--- a/Engine/Source/Editor/DetailCustomizations/Private/DialogueStructsCustomizations.cpp
+++ b/Engine/Source/Editor/DetailCustomizations/Private/DialogueStructsCustomizations.cpp
@@ -850,7 +850,7 @@ void SValidatedDialogueContextHeaderWidget::Tick( const FGeometry& AllottedGeome
 
 					if( bTargetSetNeedsReset && TargetSet )
 					{
-						bTargetSetNeedsReset = ( DialogueWave->ContextMappings[i].Context.Targets != *TargetSet );
+						bTargetSetNeedsReset = ( ToRawPtrTArrayUnsafe(DialogueWave->ContextMappings[i].Context.Targets) != *TargetSet );
 					}
 
 					if( !bSpeakerNeedsReset && !bTargetSetNeedsReset)
diff --git a/Engine/Source/Editor/Kismet/Private/BlueprintEditor.cpp b/Engine/Source/Editor/Kismet/Private/BlueprintEditor.cpp
--- a/Engine/Source/Editor/Kismet/Private/BlueprintEditor.cpp
+++ b/Engine/Source/Editor/Kismet/Private/BlueprintEditor.cpp
@@ -245,7 +245,7 @@ public:
 				[
 					SNew(SListView<UEdGraph*>)
 					.ItemHeight(20)
-					.ListItemsSource(&(ParentGraph->SubGraphs))
+					.ListItemsSource(&ToRawPtrTArrayUnsafe(ParentGraph->SubGraphs))
 					.OnGenerateRow( this, &SChildGraphPicker::GenerateMenuItemRow )
 				]
 			]
diff --git a/Engine/Source/Editor/Kismet/Private/FindInBlueprintManager.cpp b/Engine/Source/Editor/Kismet/Private/FindInBlueprintManager.cpp
--- a/Engine/Source/Editor/Kismet/Private/FindInBlueprintManager.cpp
+++ b/Engine/Source/Editor/Kismet/Private/FindInBlueprintManager.cpp
@@ -1052,7 +1052,7 @@ namespace BlueprintSearchMetaDataHelpers
 		// Collect all macro graphs
 		InWriter->WriteArrayStart(FFindInBlueprintSearchTags::FiB_Nodes);
 		{
-			for(auto* Node : InGraph->Nodes)
+			for(auto& Node : InGraph->Nodes)
 			{
 				if(Node)
 				{
diff --git a/Engine/Source/Editor/KismetCompiler/Public/KismetCompiler.h b/Engine/Source/Editor/KismetCompiler/Public/KismetCompiler.h
--- a/Engine/Source/Editor/KismetCompiler/Public/KismetCompiler.h
+++ b/Engine/Source/Editor/KismetCompiler/Public/KismetCompiler.h
@@ -328,9 +328,9 @@ protected:
 		void AddObject(const UObject* const InObject);
 
 		template<typename TOBJ>
-		void AddObjects(const TArray<TOBJ*>& InObjects)
+		void AddObjects(const TArray<TOBJ>& InObjects)
 		{
-			for ( const auto ObjPtr : InObjects )
+			for ( const auto& ObjPtr : InObjects )
 			{
 				AddObject(ObjPtr);
 			}
diff --git a/Engine/Source/Editor/LocalizationDashboard/Private/LocalizationTargetDetailCustomization.cpp b/Engine/Source/Editor/LocalizationDashboard/Private/LocalizationTargetDetailCustomization.cpp
--- a/Engine/Source/Editor/LocalizationDashboard/Private/LocalizationTargetDetailCustomization.cpp
+++ b/Engine/Source/Editor/LocalizationDashboard/Private/LocalizationTargetDetailCustomization.cpp
@@ -761,7 +761,7 @@ void FLocalizationTargetDetailCustomization::RebuildTargetsList()
 
 		for (const FGuid& OtherTargetDependencyGuid : OtherTarget->Settings.TargetDependencies)
 		{
-			ULocalizationTarget** const OtherTargetDependency = TargetSet->TargetObjects.FindByPredicate([OtherTargetDependencyGuid](ULocalizationTarget* SomeLocalizationTarget)->bool{return SomeLocalizationTarget->Settings.Guid == OtherTargetDependencyGuid;});
+			UE_TRANSITIONAL_OBJECT_PTR(ULocalizationTarget)* const OtherTargetDependency = TargetSet->TargetObjects.FindByPredicate([OtherTargetDependencyGuid](ULocalizationTarget* SomeLocalizationTarget)->bool{return SomeLocalizationTarget->Settings.Guid == OtherTargetDependencyGuid;});
 			if (OtherTargetDependency && DoesTargetDependOnUs(*OtherTargetDependency))
 			{
 				return true;
diff --git a/Engine/Source/Editor/MaterialEditor/Private/MaterialEditor.cpp b/Engine/Source/Editor/MaterialEditor/Private/MaterialEditor.cpp
--- a/Engine/Source/Editor/MaterialEditor/Private/MaterialEditor.cpp
+++ b/Engine/Source/Editor/MaterialEditor/Private/MaterialEditor.cpp
@@ -4022,7 +4022,7 @@ void FMaterialEditor::OnConvertNamedRerouteToReroute()
 					FVector2D KnotPosition(DeclarationGraphNode->NodePosX + 50, DeclarationGraphNode->NodePosY);
 
 					UMaterialExpression* Reroute = FMaterialEditorUtilities::CreateNewMaterialExpression(Graph, UMaterialExpressionReroute::StaticClass(), KnotPosition, false, true);
-					auto* KnotGraphNode = CastChecked<UMaterialGraphNode_Knot>(Reroute->GraphNode);
+					auto KnotGraphNode = CastChecked<UMaterialGraphNode_Knot>(Reroute->GraphNode);
 
 					for (UEdGraphPin* Pin : DeclarationGraphNode->GetAllPins())
 					{
diff --git a/Engine/Source/Editor/MaterialEditor/Private/MaterialEditorUtilities.cpp b/Engine/Source/Editor/MaterialEditor/Private/MaterialEditorUtilities.cpp
--- a/Engine/Source/Editor/MaterialEditor/Private/MaterialEditorUtilities.cpp
+++ b/Engine/Source/Editor/MaterialEditor/Private/MaterialEditorUtilities.cpp
@@ -523,7 +523,7 @@ void FMaterialEditorUtilities::GetVisibleMaterialParametersFromExpression(
 
 			if (LayersExpression->bIsLayerGraphBuilt)
 			{
-				for (auto* Layer : LayersExpression->LayerCallers)
+				for (auto& Layer : LayersExpression->LayerCallers)
 				{
 					if (Layer && Layer->MaterialFunction)
 					{
@@ -536,7 +536,7 @@ void FMaterialEditorUtilities::GetVisibleMaterialParametersFromExpression(
 					}
 				}
 
-				for (auto* Blend : LayersExpression->BlendCallers)
+				for (auto& Blend : LayersExpression->BlendCallers)
 				{
 					if (Blend && Blend->MaterialFunction)
 					{
diff --git a/Engine/Source/Editor/MovieSceneTools/Private/MovieSceneToolsModule.cpp b/Engine/Source/Editor/MovieSceneTools/Private/MovieSceneToolsModule.cpp
--- a/Engine/Source/Editor/MovieSceneTools/Private/MovieSceneToolsModule.cpp
+++ b/Engine/Source/Editor/MovieSceneTools/Private/MovieSceneToolsModule.cpp
@@ -344,9 +344,9 @@ bool FMovieSceneToolsModule::UpgradeLegacyEventEndpointForSection(UMovieSceneEve
 			{
 				if (EntryPoint.NodeGuid_DEPRECATED.IsValid())
 				{
-					if (UEdGraph* const* GraphPtr = Algo::FindBy(SequenceDirectorBP->UbergraphPages, EntryPoint.GraphGuid_DEPRECATED, &UEdGraph::GraphGuid))
+					if (UE_TRANSITIONAL_OBJECT_PTR(UEdGraph) const* GraphPtr = Algo::FindBy(SequenceDirectorBP->UbergraphPages, EntryPoint.GraphGuid_DEPRECATED, &UEdGraph::GraphGuid))
 					{
-						UEdGraphNode* const* NodePtr  = Algo::FindBy((*GraphPtr)->Nodes, EntryPoint.NodeGuid_DEPRECATED, &UEdGraphNode::NodeGuid);
+						UE_TRANSITIONAL_OBJECT_PTR(UEdGraphNode) const* NodePtr  = Algo::FindBy((*GraphPtr)->Nodes, EntryPoint.NodeGuid_DEPRECATED, &UEdGraphNode::NodeGuid);
 						if (NodePtr)
 						{
 							UK2Node_CustomEvent* CustomEvent = Cast<UK2Node_CustomEvent>(*NodePtr);
@@ -359,9 +359,9 @@ bool FMovieSceneToolsModule::UpgradeLegacyEventEndpointForSection(UMovieSceneEve
 					}
 				}
 				// If the node guid is invalid, this must be a function graph on the BP
-				else if (UEdGraph* const* GraphPtr = Algo::FindBy(SequenceDirectorBP->FunctionGraphs, EntryPoint.GraphGuid_DEPRECATED, &UEdGraph::GraphGuid))
+				else if (UE_TRANSITIONAL_OBJECT_PTR(UEdGraph) const* GraphPtr = Algo::FindBy(SequenceDirectorBP->FunctionGraphs, EntryPoint.GraphGuid_DEPRECATED, &UEdGraph::GraphGuid))
 				{
-					UEdGraphNode* const* NodePtr = Algo::FindByPredicate((*GraphPtr)->Nodes, [](UEdGraphNode* InNode){ return InNode && InNode->IsA<UK2Node_FunctionEntry>(); });
+					UE_TRANSITIONAL_OBJECT_PTR(UEdGraphNode) const* NodePtr = Algo::FindByPredicate((*GraphPtr)->Nodes, [](UEdGraphNode* InNode){ return InNode && InNode->IsA<UK2Node_FunctionEntry>(); });
 					if (NodePtr)
 					{
 						UK2Node_FunctionEntry* FunctionEntry = CastChecked<UK2Node_FunctionEntry>(*NodePtr);
diff --git a/Engine/Source/Editor/SkeletonEditor/Private/EditableSkeleton.cpp b/Engine/Source/Editor/SkeletonEditor/Private/EditableSkeleton.cpp
--- a/Engine/Source/Editor/SkeletonEditor/Private/EditableSkeleton.cpp
+++ b/Engine/Source/Editor/SkeletonEditor/Private/EditableSkeleton.cpp
@@ -242,7 +242,7 @@ bool FEditableSkeleton::DoesSocketAlreadyExist(const USkeletalMeshSocket* InSock
 	}
 	else if(SocketParentType == ESocketParentType::Skeleton)
 	{
-		SocketArrayPtr = &Skeleton->Sockets;
+		SocketArrayPtr = &ToRawPtrTArrayUnsafe(Skeleton->Sockets);
 	}
 
 	if (SocketArrayPtr != nullptr)
diff --git a/Engine/Source/Editor/UnrealEd/Private/Animation/DebugSkelMeshComponent.cpp b/Engine/Source/Editor/UnrealEd/Private/Animation/DebugSkelMeshComponent.cpp
--- a/Engine/Source/Editor/UnrealEd/Private/Animation/DebugSkelMeshComponent.cpp
+++ b/Engine/Source/Editor/UnrealEd/Private/Animation/DebugSkelMeshComponent.cpp
@@ -370,7 +370,7 @@ void UDebugSkelMeshComponent::PostInitMeshObject(FSkeletalMeshObject* InMeshObje
 		}
 		else if (bDrawMorphTargetVerts)
 		{
-			InMeshObject->EnableOverlayRendering(true, nullptr, &MorphTargetOfInterests);
+			InMeshObject->EnableOverlayRendering(true, nullptr, &ToRawPtrTArrayUnsafe(MorphTargetOfInterests));
 		}
 	}
 }
diff --git a/Engine/Source/Editor/UnrealEd/Private/EditorServer.cpp b/Engine/Source/Editor/UnrealEd/Private/EditorServer.cpp
--- a/Engine/Source/Editor/UnrealEd/Private/EditorServer.cpp
+++ b/Engine/Source/Editor/UnrealEd/Private/EditorServer.cpp
@@ -6570,7 +6570,7 @@ bool UEditorEngine::HandleAddSelectedCommand( const TCHAR* Str, FOutputDevice& A
 
 		if (PrecompOverride != NULL)
 		{
-			TArray<class AActor*>* OverrideActorList = 
+			auto* OverrideActorList = 
 				bVisible ? &(PrecompOverride->OverrideVisibleActors) : &(PrecompOverride->OverrideInvisibleActors);
 			// Grab a list of selected actors...
 			for (FSelectionIterator ActorIt(GetSelectedActorIterator()) ; ActorIt; ++ActorIt)
diff --git a/Engine/Source/Editor/UnrealEd/Private/Fbx/FbxMainExport.cpp b/Engine/Source/Editor/UnrealEd/Private/Fbx/FbxMainExport.cpp
--- a/Engine/Source/Editor/UnrealEd/Private/Fbx/FbxMainExport.cpp
+++ b/Engine/Source/Editor/UnrealEd/Private/Fbx/FbxMainExport.cpp
@@ -954,7 +954,7 @@ void FFbxExporter::ExportStaticMesh(AActor* Actor, UStaticMeshComponent* StaticM
 
 			const int32 LightmapUVChannel = -1;
 			const TArray<FStaticMaterial>* MaterialOrderOverride = nullptr;
-			ExportStaticMeshToFbx(StaticMesh, CurrentLodIndex, *FbxMeshName, FbxActorLOD, LightmapUVChannel, ColorBuffer, MaterialOrderOverride, &StaticMeshComponent->OverrideMaterials);
+			ExportStaticMeshToFbx(StaticMesh, CurrentLodIndex, *FbxMeshName, FbxActorLOD, LightmapUVChannel, ColorBuffer, MaterialOrderOverride, &ToRawPtrTArrayUnsafe(StaticMeshComponent->OverrideMaterials));
 		}
 	}
 	else
@@ -967,7 +967,7 @@ void FFbxExporter::ExportStaticMesh(AActor* Actor, UStaticMeshComponent* StaticM
 		FbxNode* FbxActor = ExportActor(Actor, false, NodeNameAdapter);
 		const int32 LightmapUVChannel = -1;
 		const TArray<FStaticMaterial>* MaterialOrderOverride = nullptr;
-		ExportStaticMeshToFbx(StaticMesh, LODIndex, *FbxMeshName, FbxActor, LightmapUVChannel, ColorBuffer, MaterialOrderOverride, &StaticMeshComponent->OverrideMaterials);
+		ExportStaticMeshToFbx(StaticMesh, LODIndex, *FbxMeshName, FbxActor, LightmapUVChannel, ColorBuffer, MaterialOrderOverride, &ToRawPtrTArrayUnsafe(StaticMeshComponent->OverrideMaterials));
 	}
 }
 
@@ -2281,7 +2281,7 @@ FbxNode* FFbxExporter::ExportActor(AActor* Actor, bool bExportComponents, INodeN
 						const int32 LightmapUVChannel = -1;
 						const TArray<FStaticMaterial>* MaterialOrderOverride = nullptr;
 						const FColorVertexBuffer* ColorBuffer = nullptr;
-						ExportStaticMeshToFbx(StaticMeshComp->GetStaticMesh(), LODIndex, *StaticMeshComp->GetName(), ExportNode, LightmapUVChannel, ColorBuffer, MaterialOrderOverride, &StaticMeshComp->OverrideMaterials);
+						ExportStaticMeshToFbx(StaticMeshComp->GetStaticMesh(), LODIndex, *StaticMeshComp->GetName(), ExportNode, LightmapUVChannel, ColorBuffer, MaterialOrderOverride, &ToRawPtrTArrayUnsafe(StaticMeshComp->OverrideMaterials));
 					}
 				}
 				else if (SkelMeshComp && SkelMeshComp->SkeletalMesh)
@@ -4927,7 +4927,7 @@ void FFbxExporter::ExportInstancedMeshToFbx(const UInstancedStaticMeshComponent*
 			const int32 LightmapUVChannel = -1;
 			const TArray<FStaticMaterial>* MaterialOrderOverride = nullptr;
 			const FColorVertexBuffer* ColorBuffer = nullptr;
-			ExportStaticMeshToFbx(StaticMesh, LODIndex, *FString::Printf(TEXT("%d"), InstanceIndex), InstNode, LightmapUVChannel, ColorBuffer, MaterialOrderOverride, &InstancedMeshComp->OverrideMaterials);
+			ExportStaticMeshToFbx(StaticMesh, LODIndex, *FString::Printf(TEXT("%d"), InstanceIndex), InstNode, LightmapUVChannel, ColorBuffer, MaterialOrderOverride, &ToRawPtrTArrayUnsafe(InstancedMeshComp->OverrideMaterials));
 			FbxActor->AddChild(InstNode);
 		}
 	}
diff --git a/Engine/Source/Editor/UnrealEd/Private/Fbx/FbxSkeletalMeshExport.cpp b/Engine/Source/Editor/UnrealEd/Private/Fbx/FbxSkeletalMeshExport.cpp
--- a/Engine/Source/Editor/UnrealEd/Private/Fbx/FbxSkeletalMeshExport.cpp
+++ b/Engine/Source/Editor/UnrealEd/Private/Fbx/FbxSkeletalMeshExport.cpp
@@ -591,7 +591,7 @@ void FFbxExporter::ExportSkeletalMeshComponent(USkeletalMeshComponent* SkelMeshC
 	{
 		UAnimSequence* AnimSeq = (bSaveAnimSeq && SkelMeshComp->GetAnimationMode() == EAnimationMode::AnimationSingleNode) ? 
 			Cast<UAnimSequence>(SkelMeshComp->AnimationData.AnimToPlay) : NULL;
-		FbxNode* SkeletonRootNode = ExportSkeletalMeshToFbx(SkelMeshComp->SkeletalMesh, AnimSeq, MeshName, ActorRootNode, &SkelMeshComp->OverrideMaterials);
+		FbxNode* SkeletonRootNode = ExportSkeletalMeshToFbx(SkelMeshComp->SkeletalMesh, AnimSeq, MeshName, ActorRootNode, &ToRawPtrTArrayUnsafe(SkelMeshComp->OverrideMaterials));
 		if(SkeletonRootNode)
 		{
 			FbxSkeletonRoots.Add(SkelMeshComp, SkeletonRootNode);
diff --git a/Engine/Source/Editor/UnrealEd/Private/MaterialGraphSchema.cpp b/Engine/Source/Editor/UnrealEd/Private/MaterialGraphSchema.cpp
--- a/Engine/Source/Editor/UnrealEd/Private/MaterialGraphSchema.cpp
+++ b/Engine/Source/Editor/UnrealEd/Private/MaterialGraphSchema.cpp
@@ -926,7 +926,7 @@ void UMaterialGraphSchema::GetNamedRerouteActions(FGraphActionMenuBuilder& Actio
 		{
 			if (auto* MaterialGraphNode = Cast<UMaterialGraphNode>(GraphNode))
 			{
-				if (auto* Declaration = Cast<UMaterialExpressionNamedRerouteDeclaration>(MaterialGraphNode->MaterialExpression))
+				if (auto Declaration = Cast<UMaterialExpressionNamedRerouteDeclaration>(MaterialGraphNode->MaterialExpression))
 				{
 					static const FText Category = LOCTEXT("NamedRerouteCategory", "Named Reroutes");
 					const FText Name = FText::FromString(Declaration->Name.ToString());
diff --git a/Engine/Source/Editor/VREditor/UI/VREditorUISystem.cpp b/Engine/Source/Editor/VREditor/UI/VREditorUISystem.cpp
--- a/Engine/Source/Editor/VREditor/UI/VREditorUISystem.cpp
+++ b/Engine/Source/Editor/VREditor/UI/VREditorUISystem.cpp
@@ -1102,7 +1102,7 @@ void UVREditorUISystem::ShowEditorUIPanel(AVREditorFloatingUI* Panel, UVREditorI
 			else if (PanelID.ToString().Contains(ActorPreviewPrefix))
 			{
 				// if we are closing a preview panel, make sure to unpin it also
-				AActor** PreviewedActorPtr = PreviewWindowInfo.Find(PanelID);
+				UE_TRANSITIONAL_OBJECT_PTR(AActor)* PreviewedActorPtr = PreviewWindowInfo.Find(PanelID);
 				if (PreviewedActorPtr != nullptr)
 				{
 					AActor* PreviewedActor = *PreviewedActorPtr;
@@ -1832,7 +1832,7 @@ AVREditorFloatingUI* UVREditorUISystem::GetPanel(const VREditorPanelID& InPanelI
 	AVREditorFloatingUI* Result = nullptr;
 	if (FloatingUIs.Num() > 0)
 	{
-		AVREditorFloatingUI *const * FoundPanel = FloatingUIs.Find(InPanelID);
+		UE_TRANSITIONAL_OBJECT_PTR(AVREditorFloatingUI) const * FoundPanel = FloatingUIs.Find(InPanelID);
 		if (FoundPanel != nullptr)
 		{
 			Result = *FoundPanel;
diff --git a/Engine/Source/Runtime/AIModule/Private/Actions/PawnAction_Sequence.cpp b/Engine/Source/Runtime/AIModule/Private/Actions/PawnAction_Sequence.cpp
--- a/Engine/Source/Runtime/AIModule/Private/Actions/PawnAction_Sequence.cpp
+++ b/Engine/Source/Runtime/AIModule/Private/Actions/PawnAction_Sequence.cpp
@@ -98,7 +98,7 @@ bool UPawnAction_Sequence::PushNextActionCopy()
 
 	UPawnAction* ActionCopy = SubActionTriggeringPolicy == EPawnSubActionTriggeringPolicy::CopyBeforeTriggering
 		? Cast<UPawnAction>(StaticDuplicateObject(ActionSequence[CurrentActionIndex], this))
-		: ActionSequence[CurrentActionIndex];
+		: ToRawPtr(ActionSequence[CurrentActionIndex]);
 
 	UE_VLOG(GetPawn(), LogPawnAction, Log, TEXT("%s> pushing action %s")
 		, *GetName(), *GetNameSafe(ActionCopy));
diff --git a/Engine/Source/Runtime/AIModule/Private/BehaviorTree/BlackboardComponent.cpp b/Engine/Source/Runtime/AIModule/Private/BehaviorTree/BlackboardComponent.cpp
--- a/Engine/Source/Runtime/AIModule/Private/BehaviorTree/BlackboardComponent.cpp
+++ b/Engine/Source/Runtime/AIModule/Private/BehaviorTree/BlackboardComponent.cpp
@@ -246,8 +246,8 @@ void UBlackboardComponent::PopulateSynchronizedKeys()
 						uint8* RawData = GetKeyRawData(KeyID) + DataOffset;
 						uint8* RawSource = OtherBlackboard->GetKeyRawData(OtherKeyID) + DataOffset;
 
-						UBlackboardKeyType* KeyOb = bKeyHasInstance ? KeyInstances[KeyID] : ToRawPtr(Key.KeyType);
-						const UBlackboardKeyType* SourceKeyOb = bKeyHasInstance ? OtherBlackboard->KeyInstances[OtherKeyID] : ToRawPtr(Key.KeyType);
+						UBlackboardKeyType* KeyOb = bKeyHasInstance ? KeyInstances[KeyID] : Key.KeyType;
+						const UBlackboardKeyType* SourceKeyOb = bKeyHasInstance ? OtherBlackboard->KeyInstances[OtherKeyID] : Key.KeyType;
 
 						KeyOb->CopyValues(*this, RawData, SourceKeyOb, RawSource);
 					}
@@ -517,7 +517,7 @@ EBlackboardCompare::Type UBlackboardComponent::CompareKeyValues(TSubclassOf<UBla
 	const uint8* KeyAMemory = GetKeyRawData(KeyA) + (KeyInstances[KeyA] ? sizeof(FBlackboardInstancedKeyMemory) : 0);
 	const uint8* KeyBMemory = GetKeyRawData(KeyB) + (KeyInstances[KeyB] ? sizeof(FBlackboardInstancedKeyMemory) : 0);
 
-	const UBlackboardKeyType* KeyAOb = KeyInstances[KeyA] ? KeyInstances[KeyA] : KeyType->GetDefaultObject<UBlackboardKeyType>();
+	const UBlackboardKeyType* KeyAOb = KeyInstances[KeyA] ? ToRawPtr(KeyInstances[KeyA]) : KeyType->GetDefaultObject<UBlackboardKeyType>();
 	return KeyAOb->CompareValues(*this, KeyAMemory, KeyInstances[KeyB], KeyBMemory);
 }
 
@@ -786,7 +786,7 @@ void UBlackboardComponent::ClearValue(FBlackboard::FKey KeyID)
 			{
 				const bool bKeyHasInstance = EntryInfo->KeyType->HasInstance();
 
-				UBlackboardKeyType* KeyOb = bKeyHasInstance ? KeyInstances[KeyID] : ToRawPtr(EntryInfo->KeyType);
+				UBlackboardKeyType* KeyOb = bKeyHasInstance ? KeyInstances[KeyID] : EntryInfo->KeyType;
 				const uint16 DataOffset = bKeyHasInstance ? sizeof(FBlackboardInstancedKeyMemory) : 0;
 				uint8* InstancedRawData = RawData + DataOffset;
 
@@ -803,7 +803,7 @@ void UBlackboardComponent::ClearValue(FBlackboard::FKey KeyID)
 						if (OtherKeyID != FBlackboard::InvalidKey)
 						{
 							const FBlackboardEntry* OtherEntryInfo = OtherBlackboard->BlackboardAsset->GetKey(OtherKeyID);
-							UBlackboardKeyType* OtherKeyOb = bKeyHasInstance ? OtherBlackboard->KeyInstances[OtherKeyID] : ToRawPtr(EntryInfo->KeyType);
+							UBlackboardKeyType* OtherKeyOb = bKeyHasInstance ? OtherBlackboard->KeyInstances[OtherKeyID] : EntryInfo->KeyType;
 							uint8* OtherRawData = OtherBlackboard->GetKeyRawData(OtherKeyID) + DataOffset;
 
 							OtherKeyOb->CopyValues(*OtherBlackboard, OtherRawData, KeyOb, InstancedRawData);
@@ -850,8 +850,8 @@ bool UBlackboardComponent::CopyKeyValue(FBlackboard::FKey SourceKeyID, FBlackboa
 	const uint8* SourceValueMem = GetKeyRawData(SourceKeyID) + MemDataOffset;
 	uint8* DestinationValueMem = GetKeyRawData(DestinationKeyID) + MemDataOffset;
 
-	const UBlackboardKeyType* SourceKeyOb = bKeyHasInstance ? KeyInstances[SourceKeyID] : ToRawPtr(SourceValueEntryInfo->KeyType);
-	UBlackboardKeyType*	DestKeyOb = bKeyHasInstance ? KeyInstances[DestinationKeyID] : ToRawPtr(DestinationValueEntryInfo->KeyType);
+	const UBlackboardKeyType* SourceKeyOb = bKeyHasInstance ? KeyInstances[SourceKeyID] : SourceValueEntryInfo->KeyType;
+	UBlackboardKeyType*	DestKeyOb = bKeyHasInstance ? KeyInstances[DestinationKeyID] : DestinationValueEntryInfo->KeyType;
 
 	DestKeyOb->CopyValues(*this, DestinationValueMem, SourceKeyOb, SourceValueMem);
 
diff --git a/Engine/Source/Runtime/AIModule/Private/Perception/AIPerceptionComponent.cpp b/Engine/Source/Runtime/AIModule/Private/Perception/AIPerceptionComponent.cpp
--- a/Engine/Source/Runtime/AIModule/Private/Perception/AIPerceptionComponent.cpp
+++ b/Engine/Source/Runtime/AIModule/Private/Perception/AIPerceptionComponent.cpp
@@ -118,7 +118,7 @@ void UAIPerceptionComponent::ConfigureSense(UAISenseConfig& Config)
 {
 	// first check if we're reconfiguring a sense
 	bool bIsNewConfig = true;
-	for (UAISenseConfig*& SenseConfig : SensesConfig)
+	for (UE_TRANSITIONAL_OBJECT_PTR(UAISenseConfig)& SenseConfig : SensesConfig)
 	{
 		if (SenseConfig != nullptr && SenseConfig->GetClass() == Config.GetClass())
 		{
@@ -154,7 +154,7 @@ void UAIPerceptionComponent::ConfigureSense(UAISenseConfig& Config)
 
 UAIPerceptionComponent::TAISenseConfigConstIterator UAIPerceptionComponent::GetSensesConfigIterator() const
 {
-	return SensesConfig.CreateConstIterator();
+	return ToRawPtrTArrayUnsafe(SensesConfig).CreateConstIterator();
 }
 
 void UAIPerceptionComponent::SetMaxStimulusAge(FAISenseID SenseID, float MaxAge)
diff --git a/Engine/Source/Runtime/ClothingSystemRuntimeCommon/Private/ClothingAsset.cpp b/Engine/Source/Runtime/ClothingSystemRuntimeCommon/Private/ClothingAsset.cpp
--- a/Engine/Source/Runtime/ClothingSystemRuntimeCommon/Private/ClothingAsset.cpp
+++ b/Engine/Source/Runtime/ClothingSystemRuntimeCommon/Private/ClothingAsset.cpp
@@ -873,7 +873,7 @@ void UClothingAssetCommon::PostLoad()
 	// Add any missing configs for the available cloth factories, and try to migrate them from any existing one
 	// TODO: Remove all UObject PostLoad dependencies.
 	//       Even with these ConditionalPostLoad calls, the UObject PostLoads' order of execution cannot be guaranteed.
-	for (TPair<FName, UClothConfigBase*>& ClothConfig : ClothConfigs)
+	for (auto& ClothConfig : ClothConfigs)
 	{
 		if (ClothConfig.Value)
 		{
@@ -903,7 +903,7 @@ void UClothingAssetCommon::PostLoad()
 	if (ClothingCustomVersion < FClothingAssetCustomVersion::MovePropertiesToCommonBaseClasses)
 	{
 		// Remap legacy struct FClothConfig to new config objects
-		for (TPair<FName, UClothConfigBase*>& ClothConfig : ClothConfigs)
+		for (auto& ClothConfig : ClothConfigs)
 		{
 			if (UClothConfigCommon* const ClothConfigCommon = Cast<UClothConfigCommon>(ClothConfig.Value))
 			{
@@ -924,7 +924,7 @@ void UClothingAssetCommon::PostLoad()
 				FClothConfig_Legacy ClothConfigLegacy;
 				if (ClothSimConfigCommon->MigrateTo(ClothConfigLegacy))
 				{
-					for (TPair<FName, UClothConfigBase*>& ClothConfig : ClothConfigs)
+					for (auto& ClothConfig : ClothConfigs)
 					{
 						if (UClothConfigCommon* const ClothConfigCommon = Cast<UClothConfigCommon>(ClothConfig.Value))
 						{
@@ -1043,7 +1043,7 @@ void UClothingAssetCommon::PropagateSharedConfigs(bool bMigrateSharedConfigToCon
 				ClothSharedConfigs.Reserve(Max);
 
 				// Iterate through all configs, and find the shared ones
-				for (const TPair<FName, UClothConfigBase*>& ClothSharedConfigItem : ClothingAsset->ClothConfigs)
+				for (const auto& ClothSharedConfigItem : ClothingAsset->ClothConfigs)
 				{
 					if (Cast<UClothSharedConfigCommon>(ClothSharedConfigItem.Value) &&  // Only needs shared configs
 						!ClothSharedConfigs.Find(ClothSharedConfigItem.Key))            // Only needs a single shared config per type
@@ -1058,7 +1058,7 @@ void UClothingAssetCommon::PropagateSharedConfigs(bool bMigrateSharedConfigToCon
 		for (const TPair<FName, UClothConfigBase*>& ClothSharedConfigItem : ClothSharedConfigs)
 		{
 			// Set share config
-			if (UClothConfigBase** const ClothConfigBase = ClothConfigs.Find(ClothSharedConfigItem.Key))
+			if (UE_TRANSITIONAL_OBJECT_PTR(UClothConfigBase)* const ClothConfigBase = ClothConfigs.Find(ClothSharedConfigItem.Key))
 			{
 				// Reset this shared config
 				*ClothConfigBase = ClothSharedConfigItem.Value;
@@ -1074,12 +1074,12 @@ void UClothingAssetCommon::PropagateSharedConfigs(bool bMigrateSharedConfigToCon
 		if (bMigrateSharedConfigToConfig)
 		{
 			// Iterate through all this asset's shared configs
-			for (const TPair<FName, UClothConfigBase*>& ClothSharedConfigItem : ClothConfigs)
+			for (const auto& ClothSharedConfigItem : ClothConfigs)
 			{
 				if (const UClothSharedConfigCommon* const ClothSharedConfig = Cast<UClothSharedConfigCommon>(ClothSharedConfigItem.Value))
 				{
 					// Iterate through all this asset's configs, and migrate from the shared ones
-					for (const TPair<FName, UClothConfigBase*>& ClothConfigItem : ClothConfigs)
+					for (const auto& ClothConfigItem : ClothConfigs)
 					{
 						if (Cast<UClothSharedConfigCommon>(ClothConfigItem.Value))
 						{
diff --git a/Engine/Source/Runtime/Engine/Private/Animation/Skeleton.cpp b/Engine/Source/Runtime/Engine/Private/Animation/Skeleton.cpp
--- a/Engine/Source/Runtime/Engine/Private/Animation/Skeleton.cpp
+++ b/Engine/Source/Runtime/Engine/Private/Animation/Skeleton.cpp
@@ -1919,7 +1919,7 @@ UBlendProfile* USkeleton::CreateNewBlendProfile(const FName& InProfileName)
 
 UBlendProfile* USkeleton::GetBlendProfile(const FName& InProfileName)
 {
-	UBlendProfile** FoundProfile = BlendProfiles.FindByPredicate([InProfileName](const UBlendProfile* Profile)
+	UE_TRANSITIONAL_OBJECT_PTR(UBlendProfile)* FoundProfile = BlendProfiles.FindByPredicate([InProfileName](const UBlendProfile* Profile)
 	{
 		return Profile->GetName() == InProfileName.ToString();
 	});
@@ -2002,7 +2002,7 @@ void USkeleton::RemoveUserDataOfClass(TSubclassOf<UAssetUserData> InUserDataClas
 
 const TArray<UAssetUserData*>* USkeleton::GetAssetUserDataArray() const
 {
-	return &AssetUserData;
+	return &ToRawPtrTArrayUnsafe(AssetUserData);
 }
 
 void USkeleton::HandlePackageReloaded(const EPackageReloadPhase InPackageReloadPhase, FPackageReloadedEvent* InPackageReloadedEvent)
diff --git a/Engine/Source/Runtime/Engine/Private/AnimationAsset.cpp b/Engine/Source/Runtime/Engine/Private/AnimationAsset.cpp
--- a/Engine/Source/Runtime/Engine/Private/AnimationAsset.cpp
+++ b/Engine/Source/Runtime/Engine/Private/AnimationAsset.cpp
@@ -592,7 +592,7 @@ void UAnimationAsset::RemoveUserDataOfClass(TSubclassOf<UAssetUserData> InUserDa
 
 const TArray<UAssetUserData*>* UAnimationAsset::GetAssetUserDataArray() const
 {
-	return &AssetUserData;
+	return &ToRawPtrTArrayUnsafe(AssetUserData);
 }
 
 #if WITH_EDITOR
diff --git a/Engine/Source/Runtime/Engine/Private/Components/SceneComponent.cpp b/Engine/Source/Runtime/Engine/Private/Components/SceneComponent.cpp
--- a/Engine/Source/Runtime/Engine/Private/Components/SceneComponent.cpp
+++ b/Engine/Source/Runtime/Engine/Private/Components/SceneComponent.cpp
@@ -2297,7 +2297,7 @@ void FSceneComponentInstanceData::ApplyToComponent(UActorComponent* Component, c
 		SceneComponent->UpdateComponentToWorld();
 	}
 
-	for (const TPair<USceneComponent*, FTransform>& ChildComponentPair : AttachedInstanceComponents)
+	for (const auto& ChildComponentPair : AttachedInstanceComponents)
 	{
 		USceneComponent* ChildComponent = ChildComponentPair.Key;
 		// If the ChildComponent now has a "good" attach parent it was set by the transaction and it means we are undoing/redoing attachment
@@ -2313,7 +2313,7 @@ void FSceneComponentInstanceData::ApplyToComponent(UActorComponent* Component, c
 void FSceneComponentInstanceData::AddReferencedObjects(FReferenceCollector& Collector)
 {
 	FActorComponentInstanceData::AddReferencedObjects(Collector);
-	for (TPair<USceneComponent*, FTransform>& ChildComponentPair : AttachedInstanceComponents)
+	for (auto& ChildComponentPair : AttachedInstanceComponents)
 	{
 		Collector.AddReferencedObject(ChildComponentPair.Key);
 	}
@@ -2322,7 +2322,7 @@ void FSceneComponentInstanceData::AddReferencedObjects(FReferenceCollector& Coll
 void FSceneComponentInstanceData::FindAndReplaceInstances(const TMap<UObject*, UObject*>& OldToNewInstanceMap)
 {
 	TArray<USceneComponent*> SceneComponents;
-	AttachedInstanceComponents.GenerateKeyArray(SceneComponents);
+	ToRawPtrTMapUnsafe(AttachedInstanceComponents).GenerateKeyArray(SceneComponents);
 
 	for (USceneComponent* SceneComponent : SceneComponents)
 	{
diff --git a/Engine/Source/Runtime/Engine/Private/Components/SkeletalMeshComponent.cpp b/Engine/Source/Runtime/Engine/Private/Components/SkeletalMeshComponent.cpp
--- a/Engine/Source/Runtime/Engine/Private/Components/SkeletalMeshComponent.cpp
+++ b/Engine/Source/Runtime/Engine/Private/Components/SkeletalMeshComponent.cpp
@@ -3476,7 +3476,8 @@ float USkeletalMeshComponent::CalculateMass(FName BoneName)
 	}
 	else	//We want to calculate mass before we've initialized body instances - in this case use physics asset setup
 	{
-		TArray<USkeletalBodySetup*>* BodySetups = nullptr;
+		using BodySetupContainerType = decltype(UPhysicsAsset::SkeletalBodySetups);
+		BodySetupContainerType* BodySetups = nullptr;
 		if (UPhysicsAsset* PhysicsAsset = GetPhysicsAsset())
 		{
 			BodySetups = &PhysicsAsset->SkeletalBodySetups;
diff --git a/Engine/Source/Runtime/Engine/Private/CompositeCurveTable.cpp b/Engine/Source/Runtime/Engine/Private/CompositeCurveTable.cpp
--- a/Engine/Source/Runtime/Engine/Private/CompositeCurveTable.cpp
+++ b/Engine/Source/Runtime/Engine/Private/CompositeCurveTable.cpp
@@ -23,11 +23,11 @@ UCompositeCurveTable::UCompositeCurveTable(const FObjectInitializer& ObjectIniti
 void UCompositeCurveTable::GetPreloadDependencies(TArray<UObject*>& OutDeps)
 {
 	Super::GetPreloadDependencies(OutDeps);
-	for (const TSoftObjectPtr<UCurveTable>& ParentTable : ParentTables)
+	for (const UE_TRANSITIONAL_OBJECT_PTR(UCurveTable)& ParentTable : ParentTables)
 	{
-		if (UCurveTable* Parent = ParentTable.Get())
+		if (ParentTable)
 		{
-			OutDeps.Add(Parent);
+			OutDeps.Add(ParentTable);
 		}
 	}
 }
@@ -264,10 +264,10 @@ const UCompositeCurveTable* UCompositeCurveTable::FindLoops(TArray<const UCompos
 {
 	AlreadySeenTables.Add(this);
 
-	for (const TSoftObjectPtr<UCurveTable>& CurveTable : ParentTables)
+	for (const UE_TRANSITIONAL_OBJECT_PTR(UCurveTable)& CurveTable : ParentTables)
 	{
 		// we only care about composite tables since regular tables terminate the chain and can't be in loops
-		if (UCompositeCurveTable* CompositeCurveTable = Cast<UCompositeCurveTable>(CurveTable.Get()))
+		if (UCompositeCurveTable* CompositeCurveTable = Cast<UCompositeCurveTable>(CurveTable))
 		{
 			// if we've seen this table before then we have a loop
 			for (const UCompositeCurveTable* SeenTable : AlreadySeenTables)
diff --git a/Engine/Source/Runtime/Engine/Private/DataChannel.cpp b/Engine/Source/Runtime/Engine/Private/DataChannel.cpp
--- a/Engine/Source/Runtime/Engine/Private/DataChannel.cpp
+++ b/Engine/Source/Runtime/Engine/Private/DataChannel.cpp
@@ -3126,7 +3126,7 @@ int64 UActorChannel::ReplicateActor()
 	{
 		UE_NET_TRACE_SCOPE(NewActor, Bunch, GetTraceCollector(Bunch), ENetTraceVerbosity::Trace);
 
-		Connection->PackageMap->SerializeNewActor(Bunch, this, Actor);
+		Connection->PackageMap->SerializeNewActor(Bunch, this, static_cast<AActor*&>(Actor));
 		WroteSomethingImportant = true;
 
 		Actor->OnSerializeNewActor(Bunch);
@@ -4191,7 +4191,7 @@ static void	DebugNetGUIDs( UWorld* InWorld )
 		return;
 	}
 
-	UNetConnection * Connection = (NetDriver->ServerConnection ? ToRawPtr(NetDriver->ServerConnection) : (NetDriver->ClientConnections.Num() > 0 ? NetDriver->ClientConnections[0] : NULL));
+	UNetConnection * Connection = (NetDriver->ServerConnection ? ToRawPtr(NetDriver->ServerConnection) : (NetDriver->ClientConnections.Num() > 0 ? ToRawPtr(NetDriver->ClientConnections[0]) : NULL));
 	if (!Connection)
 	{
 		return;
@@ -4216,7 +4216,7 @@ static void	ListOpenActorChannels( UWorld* InWorld )
 		return;
 	}
 
-	UNetConnection * Connection = (NetDriver->ServerConnection ? ToRawPtr(NetDriver->ServerConnection) : (NetDriver->ClientConnections.Num() > 0 ? NetDriver->ClientConnections[0] : NULL));
+	UNetConnection * Connection = (NetDriver->ServerConnection ? ToRawPtr(NetDriver->ServerConnection) : (NetDriver->ClientConnections.Num() > 0 ? ToRawPtr(NetDriver->ClientConnections[0]) : NULL));
 	if (!Connection)
 	{
 		return;
@@ -4273,7 +4273,7 @@ static void	DeleteDormantActor( UWorld* InWorld )
 		return;
 	}
 
-	UNetConnection * Connection = (NetDriver->ServerConnection ? ToRawPtr(NetDriver->ServerConnection) : (NetDriver->ClientConnections.Num() > 0 ? NetDriver->ClientConnections[0] : NULL));
+	UNetConnection * Connection = (NetDriver->ServerConnection ? ToRawPtr(NetDriver->ServerConnection) : (NetDriver->ClientConnections.Num() > 0 ? ToRawPtr(NetDriver->ClientConnections[0]) : NULL));
 	if (!Connection)
 	{
 		return;
diff --git a/Engine/Source/Runtime/Engine/Private/DemoNetDriver.cpp b/Engine/Source/Runtime/Engine/Private/DemoNetDriver.cpp
--- a/Engine/Source/Runtime/Engine/Private/DemoNetDriver.cpp
+++ b/Engine/Source/Runtime/Engine/Private/DemoNetDriver.cpp
@@ -2856,7 +2856,7 @@ bool UDemoNetDriver::SpawnSplitscreenViewer(ULocalPlayer* NewPlayer, UWorld* InW
 		return false;
 	}
 
-	UNetConnection* ChildConnection = CreateChild((ClientConnections.Num() > 0) ? ClientConnections[0] : ToRawPtr(ServerConnection));
+	UNetConnection* ChildConnection = CreateChild((ClientConnections.Num() > 0) ? ClientConnections[0] : ServerConnection);
 
 	APlayerController* NewSplitscreenController = ReplayHelper.CreateSpectatorController(ChildConnection);
 	if (NewSplitscreenController == nullptr)
@@ -5034,7 +5034,7 @@ void UDemoNetDriver::QueueNetStartupActorForRollbackViaDeletion(AActor* Actor)
 				FRepShadowDataBuffer ShadowData(ReceivingRepState->StaticBuffer.GetData());
 				FConstRepObjectDataBuffer ActorData(Actor);
 
-				if (NewReplicator->RepLayout->DiffStableProperties(nullptr, &RollbackActor.ObjReferences, ShadowData, ActorData))
+				if (NewReplicator->RepLayout->DiffStableProperties(nullptr, &ToRawPtrTArrayUnsafe(RollbackActor.ObjReferences), ShadowData, ActorData))
 				{
 					RollbackActor.RepState = MakeShareable(NewReplicator->RepState.Release());
 				}
@@ -5054,7 +5054,7 @@ void UDemoNetDriver::QueueNetStartupActorForRollbackViaDeletion(AActor* Actor)
 					FRepShadowDataBuffer ShadowData(ReceivingRepState->StaticBuffer.GetData());
 					FConstRepObjectDataBuffer ActorCompData(ActorComp);
 
-					if (SubObjReplicator->RepLayout->DiffStableProperties(nullptr, &RollbackActor.ObjReferences, ShadowData, ActorCompData))
+					if (SubObjReplicator->RepLayout->DiffStableProperties(nullptr, &ToRawPtrTArrayUnsafe(RollbackActor.ObjReferences), ShadowData, ActorCompData))
 					{
 						RollbackActor.SubObjRepState.Add(ActorComp->GetFullName(), MakeShareable(SubObjReplicator->RepState.Release()));
 					}
diff --git a/Engine/Source/Runtime/Engine/Private/GameInstance.cpp b/Engine/Source/Runtime/Engine/Private/GameInstance.cpp
--- a/Engine/Source/Runtime/Engine/Private/GameInstance.cpp
+++ b/Engine/Source/Runtime/Engine/Private/GameInstance.cpp
@@ -302,7 +302,7 @@ FGameInstancePIEResult UGameInstance::InitializeForPlayInEditor(int32 PIEInstanc
 	NewWorld->SetPlayInEditorInitialNetMode(GetNetModeFromPlayNetMode(Params.NetMode, Params.bRunAsDedicated));
 	NewWorld->SetGameInstance(this);
 	WorldContext->SetCurrentWorld(NewWorld);
-	WorldContext->AddRef(EditorEngine->PlayWorld);	// Tie this context to this UEngine::PlayWorld*		// @fixme, needed still?
+	WorldContext->AddRef(static_cast<UWorld*&>(EditorEngine->PlayWorld));	// Tie this context to this UEngine::PlayWorld*		// @fixme, needed still?
 
 	// make sure we can clean up this world!
 	NewWorld->ClearFlags(RF_Standalone);
@@ -997,7 +997,7 @@ void UGameInstance::CleanupGameViewport()
 
 TArray<class ULocalPlayer*>::TConstIterator	UGameInstance::GetLocalPlayerIterator() const
 {
-	return LocalPlayers.CreateConstIterator();
+	return ToRawPtrTArrayUnsafe(LocalPlayers).CreateConstIterator();
 }
 
 const TArray<class ULocalPlayer*>& UGameInstance::GetLocalPlayers() const
diff --git a/Engine/Source/Runtime/Engine/Private/HLODProxy.cpp b/Engine/Source/Runtime/Engine/Private/HLODProxy.cpp
--- a/Engine/Source/Runtime/Engine/Private/HLODProxy.cpp
+++ b/Engine/Source/Runtime/Engine/Private/HLODProxy.cpp
@@ -119,7 +119,7 @@ void UHLODProxy::Clean()
 	}
 	else if (HLODActors.Num() > 0)
 	{
-		for (TMap<UHLODProxyDesc*, FHLODProxyMesh>::TIterator ItHLODActor = HLODActors.CreateIterator(); ItHLODActor; ++ItHLODActor)
+		for (auto ItHLODActor = HLODActors.CreateIterator(); ItHLODActor; ++ItHLODActor)
 		{
 			RemoveAssets(ItHLODActor.Value());
 		}
@@ -193,7 +193,7 @@ void UHLODProxy::UpdateHLODDescs(const ULevel* InLevel)
 
 	// For each HLODProxyDesc stored in this proxy, ensure that it is up to date with the associated LODActor
 	// Purge the HLODProxyDesc that are unused (not referenced by any LODActor)
-	for (TMap<UHLODProxyDesc*, FHLODProxyMesh>::TIterator ItHLODActor = HLODActors.CreateIterator(); ItHLODActor; ++ItHLODActor)
+	for (auto ItHLODActor = HLODActors.CreateIterator(); ItHLODActor; ++ItHLODActor)
 	{
 		UHLODProxyDesc* HLODProxyDesc = ItHLODActor.Key();
 		ALODActor** LODActor = LODActors.Find(HLODProxyDesc);
diff --git a/Engine/Source/Runtime/Engine/Private/LODActor.cpp b/Engine/Source/Runtime/Engine/Private/LODActor.cpp
--- a/Engine/Source/Runtime/Engine/Private/LODActor.cpp
+++ b/Engine/Source/Runtime/Engine/Private/LODActor.cpp
@@ -1204,8 +1204,8 @@ UInstancedStaticMeshComponent* ALODActor::GetOrCreateISMComponent(const FHLODIns
 		AddInstanceComponent(LODComponent);
 		LODComponent->SetupAttachment(GetRootComponent());
 		
-		LODComponent->SetStaticMesh(const_cast<UStaticMesh*>(InstancingKey.StaticMesh));
-		LODComponent->SetMaterial(0, const_cast<UMaterialInterface*>(InstancingKey.Material));
+		LODComponent->SetStaticMesh(const_cast<UStaticMesh*>(ToRawPtr(InstancingKey.StaticMesh)));
+		LODComponent->SetMaterial(0, const_cast<UMaterialInterface*>(ToRawPtr(InstancingKey.Material)));
 
 		if (StaticMeshComponent->IsRegistered())
 		{
diff --git a/Engine/Source/Runtime/Engine/Private/Materials/HLSLMaterialTranslator.cpp b/Engine/Source/Runtime/Engine/Private/Materials/HLSLMaterialTranslator.cpp
--- a/Engine/Source/Runtime/Engine/Private/Materials/HLSLMaterialTranslator.cpp
+++ b/Engine/Source/Runtime/Engine/Private/Materials/HLSLMaterialTranslator.cpp
@@ -266,7 +266,7 @@ void FHLSLMaterialTranslator::GatherCustomVertexInterpolators(TArray<UMaterialEx
 
 			if (LayersExpression->bIsLayerGraphBuilt)
 			{
-				for (auto* Layer : LayersExpression->LayerCallers)
+				for (auto& Layer : LayersExpression->LayerCallers)
 				{
 					if (Layer && Layer->MaterialFunction)
 					{
@@ -285,7 +285,7 @@ void FHLSLMaterialTranslator::GatherCustomVertexInterpolators(TArray<UMaterialEx
 					}
 				}
 
-				for (auto* Blend : LayersExpression->BlendCallers)
+				for (auto& Blend : LayersExpression->BlendCallers)
 				{
 					if (Blend && Blend->MaterialFunction)
 					{
@@ -402,7 +402,7 @@ EMaterialExpressionVisitResult FHLSLMaterialTranslator::VisitExpressionsRecursiv
 
 			if (LayersExpression->bIsLayerGraphBuilt)
 			{
-				for (auto* Layer : LayersExpression->LayerCallers)
+				for (auto& Layer : LayersExpression->LayerCallers)
 				{
 					if (Layer && Layer->MaterialFunction)
 					{
@@ -426,7 +426,7 @@ EMaterialExpressionVisitResult FHLSLMaterialTranslator::VisitExpressionsRecursiv
 					}
 				}
 
-				for (auto* Blend : LayersExpression->BlendCallers)
+				for (auto& Blend : LayersExpression->BlendCallers)
 				{
 					if (Blend && Blend->MaterialFunction)
 					{
diff --git a/Engine/Source/Runtime/Engine/Private/Materials/MaterialExpressions.cpp b/Engine/Source/Runtime/Engine/Private/Materials/MaterialExpressions.cpp
--- a/Engine/Source/Runtime/Engine/Private/Materials/MaterialExpressions.cpp
+++ b/Engine/Source/Runtime/Engine/Private/Materials/MaterialExpressions.cpp
@@ -6793,7 +6793,7 @@ void UMaterialExpressionMaterialAttributeLayers::RebuildLayerGraph(bool bReportE
 				{
 					// Active blend input is previous blend or direct layer if previous is opaque
 					UMaterialExpressionMaterialFunctionCall* BlendInput = BlendCallers[LayerIndex - 2];
-					BlendInput = BlendInput->MaterialFunction ? BlendInput : LayerCallers[LayerIndex - 1];
+					BlendInput = BlendInput->MaterialFunction ? BlendInput : ToRawPtr(LayerCallers[LayerIndex - 1]);
 
 					BlendCallers[LayerIndex - 1]->FunctionInputs[0].Input.Connect(0, BlendInput);
 					BlendCallers[LayerIndex - 1]->FunctionInputs[1].Input.Connect(0, LayerCallers[LayerIndex]);
diff --git a/Engine/Source/Runtime/Engine/Private/NetDriver.cpp b/Engine/Source/Runtime/Engine/Private/NetDriver.cpp
--- a/Engine/Source/Runtime/Engine/Private/NetDriver.cpp
+++ b/Engine/Source/Runtime/Engine/Private/NetDriver.cpp
@@ -5115,7 +5115,7 @@ void UNetDriver::PrintDebugRelevantActors()
 void UNetDriver::DrawNetDriverDebug()
 {
 #if ENABLE_DRAW_DEBUG
-	UNetConnection *Connection = (ServerConnection ? ToRawPtr(ServerConnection) : (ClientConnections.Num() >= 1 ? ClientConnections[0] : NULL));
+	UNetConnection *Connection = (ServerConnection ? ToRawPtr(ServerConnection) : (ClientConnections.Num() >= 1 ? ToRawPtr(ClientConnections[0]) : NULL));
 	if (!Connection)
 	{
 		return;
diff --git a/Engine/Source/Runtime/Engine/Private/Particles/WorldPSCPool.cpp b/Engine/Source/Runtime/Engine/Private/Particles/WorldPSCPool.cpp
--- a/Engine/Source/Runtime/Engine/Private/Particles/WorldPSCPool.cpp
+++ b/Engine/Source/Runtime/Engine/Private/Particles/WorldPSCPool.cpp
@@ -221,7 +221,7 @@ FWorldPSCPool::~FWorldPSCPool()
 
 void FWorldPSCPool::Cleanup(UWorld* World)
 {
-	for (TPair<UParticleSystem*, FPSCPool>& Pool : WorldParticleSystemPools)
+	for (auto& Pool : WorldParticleSystemPools)
 	{
 		Pool.Value.Cleanup();
 	}
@@ -313,7 +313,7 @@ void FWorldPSCPool::ReclaimWorldParticleSystem(UParticleSystemComponent* PSC)
 		if (CurrentTime - LastParticleSytemPoolCleanTime > GParticleSystemPoolingCleanTime)
 		{
 			LastParticleSytemPoolCleanTime = CurrentTime;
-			for (TPair<UParticleSystem*, FPSCPool>& Pair : WorldParticleSystemPools)
+			for (auto& Pair : WorldParticleSystemPools)
 			{
 				Pair.Value.KillUnusedComponents(CurrentTime - GParticleSystemPoolKillUnusedTime, PSC->Template);
 			}
@@ -341,7 +341,7 @@ void FWorldPSCPool::Dump()
 	FString DumpStr;
 
 	uint32 TotalMemUsage = 0;
-	for (TPair<UParticleSystem*, FPSCPool>& Pair : WorldParticleSystemPools)
+	for (auto& Pair : WorldParticleSystemPools)
 	{
 		UParticleSystem* System = Pair.Key;
 		FPSCPool& Pool = Pair.Value;		
diff --git a/Engine/Source/Runtime/Engine/Private/SkeletalMesh.cpp b/Engine/Source/Runtime/Engine/Private/SkeletalMesh.cpp
--- a/Engine/Source/Runtime/Engine/Private/SkeletalMesh.cpp
+++ b/Engine/Source/Runtime/Engine/Private/SkeletalMesh.cpp
@@ -4117,7 +4117,7 @@ void USkeletalMesh::RemoveUserDataOfClass(TSubclassOf<UAssetUserData> InUserData
 
 const TArray<UAssetUserData*>* USkeletalMesh::GetAssetUserDataArray() const
 {
-	return &AssetUserData;
+	return &ToRawPtrTArrayUnsafe(AssetUserData);
 }
 
 ////// SKELETAL MESH THUMBNAIL SUPPORT ////////
diff --git a/Engine/Source/Runtime/Engine/Private/SoundBase.cpp b/Engine/Source/Runtime/Engine/Private/SoundBase.cpp
--- a/Engine/Source/Runtime/Engine/Private/SoundBase.cpp
+++ b/Engine/Source/Runtime/Engine/Private/SoundBase.cpp
@@ -273,5 +273,5 @@ void USoundBase::RemoveUserDataOfClass(TSubclassOf<UAssetUserData> InUserDataCla
 
 const TArray<UAssetUserData*>* USoundBase::GetAssetUserDataArray() const
 {
-	return &AssetUserData;
+	return &ToRawPtrTArrayUnsafe(AssetUserData);
 }
\ No newline at end of file
diff --git a/Engine/Source/Runtime/Engine/Private/StaticMesh.cpp b/Engine/Source/Runtime/Engine/Private/StaticMesh.cpp
--- a/Engine/Source/Runtime/Engine/Private/StaticMesh.cpp
+++ b/Engine/Source/Runtime/Engine/Private/StaticMesh.cpp
@@ -6533,7 +6533,7 @@ void UStaticMesh::RemoveUserDataOfClass(TSubclassOf<UAssetUserData> InUserDataCl
 
 const TArray<UAssetUserData*>* UStaticMesh::GetAssetUserDataArray() const 
 {
-	return &AssetUserData;
+	return &ToRawPtrTArrayUnsafe(AssetUserData);
 }
 
 /**
diff --git a/Engine/Source/Runtime/Engine/Private/Texture.cpp b/Engine/Source/Runtime/Engine/Private/Texture.cpp
--- a/Engine/Source/Runtime/Engine/Private/Texture.cpp
+++ b/Engine/Source/Runtime/Engine/Private/Texture.cpp
@@ -850,7 +850,7 @@ void UTexture::RemoveUserDataOfClass(TSubclassOf<UAssetUserData> InUserDataClass
 
 const TArray<UAssetUserData*>* UTexture::GetAssetUserDataArray() const
 {
-	return &AssetUserData;
+	return &ToRawPtrTArrayUnsafe(AssetUserData);
 }
 
 FStreamableRenderResourceState UTexture::GetResourcePostInitState(FTexturePlatformData* PlatformData, bool bAllowStreaming, int32 MinRequestMipCount, int32 MaxMipCount) const
diff --git a/Engine/Source/Runtime/Engine/Private/World.cpp b/Engine/Source/Runtime/Engine/Private/World.cpp
--- a/Engine/Source/Runtime/Engine/Private/World.cpp
+++ b/Engine/Source/Runtime/Engine/Private/World.cpp
@@ -6985,7 +6985,7 @@ int32 UWorld::GetActorCount()
 
 FConstLevelIterator	UWorld::GetLevelIterator() const
 {
-	return Levels.CreateConstIterator();
+	return ToRawPtrTArrayUnsafe(Levels).CreateConstIterator();
 }
 
 ULevel* UWorld::GetLevel( int32 InLevelIndex ) const
diff --git a/Engine/Source/Runtime/Experimental/ChaosSolverEngine/Private/Chaos/ChaosGameplayEventDispatcher.cpp b/Engine/Source/Runtime/Experimental/ChaosSolverEngine/Private/Chaos/ChaosGameplayEventDispatcher.cpp
--- a/Engine/Source/Runtime/Experimental/ChaosSolverEngine/Private/Chaos/ChaosGameplayEventDispatcher.cpp
+++ b/Engine/Source/Runtime/Experimental/ChaosSolverEngine/Private/Chaos/ChaosGameplayEventDispatcher.cpp
@@ -266,7 +266,7 @@ void UChaosGameplayEventDispatcher::HandleCollisionEvents(const Chaos::FCollisio
 			// look through all the components that someone is interested in, and see if they had a collision
 			// note that we only need to care about the interaction from the POV of the registered component,
 			// since if anyone wants notifications for the other component it hit, it's also registered and we'll get to that elsewhere in the list
-			for (TMap<UPrimitiveComponent*, FChaosHandlerSet>::TIterator It(CollisionEventRegistrations); It; ++It)
+			for (decltype(CollisionEventRegistrations)::TIterator It(CollisionEventRegistrations); It; ++It)
 			{
 				const FChaosHandlerSet& HandlerSet = It.Value();
 
diff --git a/Engine/Source/Runtime/Experimental/GeometryCollectionEngine/Private/GeometryCollection/SkeletalMeshSimulationComponent.cpp b/Engine/Source/Runtime/Experimental/GeometryCollectionEngine/Private/GeometryCollection/SkeletalMeshSimulationComponent.cpp
--- a/Engine/Source/Runtime/Experimental/GeometryCollectionEngine/Private/GeometryCollection/SkeletalMeshSimulationComponent.cpp
+++ b/Engine/Source/Runtime/Experimental/GeometryCollectionEngine/Private/GeometryCollection/SkeletalMeshSimulationComponent.cpp
@@ -149,7 +149,7 @@ void USkeletalMeshSimulationComponent::OnCreatePhysicsState()
 			USkeletalMesh* SkeletalMesh = SkelMeshComponent->SkeletalMesh;
 			if (SkeletalMesh)
 			{
-				UPhysicsAsset* PhysicsAsset = OverridePhysicsAsset ? OverridePhysicsAsset : SkelMeshComponent->SkeletalMesh->GetPhysicsAsset();
+				UPhysicsAsset* PhysicsAsset = OverridePhysicsAsset ? ToRawPtr(OverridePhysicsAsset) : SkelMeshComponent->SkeletalMesh->GetPhysicsAsset();
 				FPhysicsAssetSimulationUtil::BuildParams(this, OwningActor, SkelMeshComponent, PhysicsAsset, OutPhysicsParams);
 			}
 
diff --git a/Engine/Source/Runtime/Experimental/InteractiveToolsFramework/Private/InteractiveTool.cpp b/Engine/Source/Runtime/Experimental/InteractiveToolsFramework/Private/InteractiveTool.cpp
--- a/Engine/Source/Runtime/Experimental/InteractiveToolsFramework/Private/InteractiveTool.cpp
+++ b/Engine/Source/Runtime/Experimental/InteractiveToolsFramework/Private/InteractiveTool.cpp
@@ -234,9 +234,9 @@ bool UInteractiveTool::CanAccept() const
 
 void UInteractiveTool::Tick(float DeltaTime)
 {
-	for (auto* Object : ToolPropertyObjects)
+	for (auto& Object : ToolPropertyObjects)
 	{
-		auto* Propset = Cast<UInteractiveToolPropertySet>(Object);
+		auto* Propset = Cast<UInteractiveToolPropertySet>(ToRawPtr(Object));
 		if ( Propset != nullptr )
 		{
 			if ( Propset->IsPropertySetEnabled() )
diff --git a/Engine/Source/Runtime/GameplayTags/Private/GameplayTagContainer.cpp b/Engine/Source/Runtime/GameplayTags/Private/GameplayTagContainer.cpp
--- a/Engine/Source/Runtime/GameplayTags/Private/GameplayTagContainer.cpp
+++ b/Engine/Source/Runtime/GameplayTags/Private/GameplayTagContainer.cpp
@@ -1612,7 +1612,8 @@ void FQueryEvaluator::ReadEditableQueryTags(UEditableGameplayTagQueryExpression*
 void FQueryEvaluator::ReadEditableQueryExprList(UEditableGameplayTagQueryExpression* EditableQueryExpr)
 {
 	// find the tag container to read into
-	TArray<UEditableGameplayTagQueryExpression*>* ExprList = nullptr;
+	using ExpressionArrayType = decltype(UEditableGameplayTagQueryExpression_AnyExprMatch::Expressions);
+	ExpressionArrayType* ExprList = nullptr;
 	if (EditableQueryExpr->IsA(UEditableGameplayTagQueryExpression_AnyExprMatch::StaticClass()))
 	{
 		ExprList = &((UEditableGameplayTagQueryExpression_AnyExprMatch*)EditableQueryExpr)->Expressions;
diff --git a/Engine/Source/Runtime/Interchange/Core/Private/Nodes/InterchangeBaseNodeContainer.cpp b/Engine/Source/Runtime/Interchange/Core/Private/Nodes/InterchangeBaseNodeContainer.cpp
--- a/Engine/Source/Runtime/Interchange/Core/Private/Nodes/InterchangeBaseNodeContainer.cpp
+++ b/Engine/Source/Runtime/Interchange/Core/Private/Nodes/InterchangeBaseNodeContainer.cpp
@@ -54,7 +54,7 @@ bool UInterchangeBaseNodeContainer::IsNodeUIDValid(const FString& NodeUniqueID)
 
 void UInterchangeBaseNodeContainer::IterateNodes(TFunctionRef<void(const FString&, UInterchangeBaseNode*)> IterationLambda)
 {
-	for (TPair<FString, UInterchangeBaseNode*>& NodeKeyValue : Nodes)
+	for (auto& NodeKeyValue : Nodes)
 	{
 		IterationLambda(NodeKeyValue.Key, NodeKeyValue.Value);
 	}
@@ -62,7 +62,7 @@ void UInterchangeBaseNodeContainer::IterateNodes(TFunctionRef<void(const FString
 
 void UInterchangeBaseNodeContainer::GetRoots(TArray<FString>& RootNodes)
 {
-	for (TPair<FString, UInterchangeBaseNode*>& NodeKeyValue : Nodes)
+	for (auto& NodeKeyValue : Nodes)
 	{
 		if (NodeKeyValue.Value->GetParentUID() == UInterchangeBaseNode::InvalidNodeUID())
 		{
@@ -114,7 +114,7 @@ int32 UInterchangeBaseNodeContainer::GetNodeChildrenCount(const FString& NodeUni
 TArray<FString> UInterchangeBaseNodeContainer::GetNodeChildrenUIDs(const FString& NodeUniqueID) const
 {
 	TArray<FString> OutChildrenUIDs;
-	for (const TPair<FString, UInterchangeBaseNode*>& NodeKeyValue : Nodes)
+	for (const auto& NodeKeyValue : Nodes)
 	{
 		if (NodeKeyValue.Value->GetParentUID() == NodeUniqueID)
 		{
diff --git a/Engine/Source/Runtime/Landscape/Private/Landscape.cpp b/Engine/Source/Runtime/Landscape/Private/Landscape.cpp
--- a/Engine/Source/Runtime/Landscape/Private/Landscape.cpp
+++ b/Engine/Source/Runtime/Landscape/Private/Landscape.cpp
@@ -368,13 +368,13 @@ void ULandscapeComponent::Serialize(FArchive& Ar)
 		TArray<UObject**> TexturesAndMaterials;
 		TexturesAndMaterials.Add((UObject**)&HeightmapTexture);
 		TexturesAndMaterials.Add((UObject**)&XYOffsetmapTexture);
-		for (UTexture2D*& WeightmapTexture : WeightmapTextures)
+		for (UE_TRANSITIONAL_OBJECT_PTR(UTexture2D)& WeightmapTexture : WeightmapTextures)
 		{
-			TexturesAndMaterials.Add((UObject**)&WeightmapTexture);
+			TexturesAndMaterials.Add((UObject**)&static_cast<UTexture2D*&>(WeightmapTexture));
 		}
-		for (UTexture2D*& MobileWeightmapTexture : MobileWeightmapTextures)
+		for (UE_TRANSITIONAL_OBJECT_PTR(UTexture2D)& MobileWeightmapTexture : MobileWeightmapTextures)
 		{
-			TexturesAndMaterials.Add((UObject**)&MobileWeightmapTexture);
+			TexturesAndMaterials.Add((UObject**)&static_cast<UTexture2D*&>(MobileWeightmapTexture));
 		}
 		for (auto& ItPair : LayersData)
 		{
@@ -385,17 +385,17 @@ void ULandscapeComponent::Serialize(FArchive& Ar)
 				TexturesAndMaterials.Add((UObject**)&WeightmapTexture);
 			}
 		}
-		for (UMaterialInstanceConstant*& MaterialInstance : MaterialInstances)
+		for (UE_TRANSITIONAL_OBJECT_PTR(UMaterialInstanceConstant)& MaterialInstance : MaterialInstances)
 		{
-			TexturesAndMaterials.Add((UObject**)&MaterialInstance);
+			TexturesAndMaterials.Add((UObject**)&static_cast<UMaterialInstanceConstant*&>(MaterialInstance));
 		}
-		for (UMaterialInterface*& MobileMaterialInterface : MobileMaterialInterfaces)
+		for (UE_TRANSITIONAL_OBJECT_PTR(UMaterialInterface)& MobileMaterialInterface : MobileMaterialInterfaces)
 		{
-			TexturesAndMaterials.Add((UObject**)(&MobileMaterialInterface));
+			TexturesAndMaterials.Add((UObject**)(&static_cast<UMaterialInterface*&>(MobileMaterialInterface)));
 		}
-		for (UMaterialInstanceConstant*& MobileCombinationMaterialInstance : MobileCombinationMaterialInstances)
+		for (UE_TRANSITIONAL_OBJECT_PTR(UMaterialInstanceConstant)& MobileCombinationMaterialInstance : MobileCombinationMaterialInstances)
 		{
-			TexturesAndMaterials.Add((UObject**)&MobileCombinationMaterialInstance);
+			TexturesAndMaterials.Add((UObject**)&static_cast<UMaterialInstanceConstant*&>(MobileCombinationMaterialInstance));
 		}
 
 		if (Ar.IsSaving())
@@ -1389,7 +1389,7 @@ void ULandscapeComponent::BeginDestroy()
 			if (WeightmapTextures.IsValidIndex(WeightmapIndex))
 			{
 				UTexture2D* WeightmapTexture = WeightmapTextures[WeightmapIndex];
-				ULandscapeWeightmapUsage** Usage = Proxy->WeightmapUsageMap.Find(WeightmapTexture);
+				UE_TRANSITIONAL_OBJECT_PTR(ULandscapeWeightmapUsage)* Usage = Proxy->WeightmapUsageMap.Find(WeightmapTexture);
 				if (Usage != nullptr && (*Usage) != nullptr)
 				{
 					(*Usage)->ChannelUsage[WeightmapLayerAllocations[LayerIdx].WeightmapTextureChannel] = nullptr;
diff --git a/Engine/Source/Runtime/Landscape/Private/LandscapeCollision.cpp b/Engine/Source/Runtime/Landscape/Private/LandscapeCollision.cpp
--- a/Engine/Source/Runtime/Landscape/Private/LandscapeCollision.cpp
+++ b/Engine/Source/Runtime/Landscape/Private/LandscapeCollision.cpp
@@ -990,7 +990,7 @@ bool ULandscapeHeightfieldCollisionComponent::CookCollisionData(const FName& For
 					else if (RenderPhysicalMaterialIds)
 					{
 						uint8 RenderIdx = RenderPhysicalMaterialIds[SrcSampleIndex];
-						UPhysicalMaterial* DominantMaterial = RenderIdx > 0 ? PhysicalMaterialRenderObjects[RenderIdx - 1] : DefMaterial;
+						UPhysicalMaterial* DominantMaterial = RenderIdx > 0 ? ToRawPtr(PhysicalMaterialRenderObjects[RenderIdx - 1]) : DefMaterial;
 						MaterialIndex = InOutMaterials.AddUnique(DominantMaterial);
 					}
 					else
diff --git a/Engine/Source/Runtime/Landscape/Private/LandscapeEdit.cpp b/Engine/Source/Runtime/Landscape/Private/LandscapeEdit.cpp
--- a/Engine/Source/Runtime/Landscape/Private/LandscapeEdit.cpp
+++ b/Engine/Source/Runtime/Landscape/Private/LandscapeEdit.cpp
@@ -335,7 +335,7 @@ void ULandscapeComponent::UpdateMaterialInstances_Internal(FMaterialUpdateContex
 	check(GIsEditor);
 
 	int32 MaxLOD = FMath::CeilLogTwo(SubsectionSizeQuads + 1) - 1;
-	TMap<UMaterialInterface*, int8> NewMaterialPerLOD;
+	decltype(MaterialPerLOD) NewMaterialPerLOD;
 	LODIndexToMaterialIndex.SetNumUninitialized(MaxLOD+1);
 	int8 LastLODIndex = INDEX_NONE;
 
@@ -746,7 +746,7 @@ void ULandscapeComponent::FixupWeightmaps()
 
 				UTexture2D* WeightmapTexture = WeightmapTextures[Allocation.WeightmapTextureIndex];
 
-				ULandscapeWeightmapUsage** TempUsage = Proxy->WeightmapUsageMap.Find(WeightmapTexture);
+				UE_TRANSITIONAL_OBJECT_PTR(ULandscapeWeightmapUsage)* TempUsage = Proxy->WeightmapUsageMap.Find(WeightmapTexture);
 
 				if (TempUsage == nullptr)
 				{
@@ -3088,7 +3088,7 @@ LANDSCAPE_API void ALandscapeProxy::Import(const FGuid& InGuid, int32 InMinX, in
 			ImportSettings.DestinationLayerGuid = DefaultLayer->Guid;
 			LayerImportSettings.Add(ImportSettings);
 
-			ComponentsToProcess.Append(LandscapeComponents);
+			ComponentsToProcess.Append(ToRawPtrTArrayUnsafe(LandscapeComponents));
 		}
 
 		check(LayerImportSettings.Num() != 0);
@@ -3128,7 +3128,7 @@ LANDSCAPE_API void ALandscapeProxy::Import(const FGuid& InGuid, int32 InMinX, in
 				check(ComponentLayerData != nullptr);
 
 				LayersTextures.Add(ComponentLayerData->HeightmapData.Texture);
-				LayersTextures.Append(ComponentLayerData->WeightmapData.Textures);
+				LayersTextures.Append(ToRawPtrTArrayUnsafe(ComponentLayerData->WeightmapData.Textures));
 			}
 		}
 
@@ -5853,7 +5853,7 @@ void ALandscapeProxy::RemoveInvalidWeightmaps()
 {
 	if (GIsEditor)
 	{
-		for (TMap< UTexture2D*, ULandscapeWeightmapUsage* >::TIterator It(WeightmapUsageMap); It; ++It)
+		for (decltype(WeightmapUsageMap)::TIterator It(WeightmapUsageMap); It; ++It)
 		{
 			UTexture2D* Tex = It.Key();
 			ULandscapeWeightmapUsage* Usage = It.Value();
diff --git a/Engine/Source/Runtime/Landscape/Private/LandscapeEditLayers.cpp b/Engine/Source/Runtime/Landscape/Private/LandscapeEditLayers.cpp
--- a/Engine/Source/Runtime/Landscape/Private/LandscapeEditLayers.cpp
+++ b/Engine/Source/Runtime/Landscape/Private/LandscapeEditLayers.cpp
@@ -1849,7 +1849,7 @@ void ALandscapeProxy::InitializeProxyLayersWeightmapUsage()
 					{
 						FWeightmapLayerAllocationInfo& Allocation = LayerData->WeightmapData.LayerAllocations[LayerIdx];
 						UTexture2D* WeightmapTexture = LayerData->WeightmapData.Textures[Allocation.WeightmapTextureIndex];
-						ULandscapeWeightmapUsage** TempUsage = WeightmapUsageMap.Find(WeightmapTexture);
+						UE_TRANSITIONAL_OBJECT_PTR(ULandscapeWeightmapUsage)* TempUsage = WeightmapUsageMap.Find(WeightmapTexture);
 
 						if (TempUsage == nullptr)
 						{
@@ -4272,7 +4272,7 @@ void ALandscape::UpdateLayersMaterialInstances(const TArray<ULandscapeComponent*
 	for (ULandscapeComponent* Component : ComponentsToUpdate)
 	{
 		int32 MaxLOD = FMath::CeilLogTwo(SubsectionSizeQuads + 1) - 1;
-		TMap<UMaterialInterface*, int8> NewMaterialPerLOD;
+		decltype(Component->MaterialPerLOD) NewMaterialPerLOD;
 		Component->LODIndexToMaterialIndex.SetNumUninitialized(MaxLOD + 1);
 		int8 LastLODIndex = INDEX_NONE;
 
@@ -5410,7 +5410,7 @@ void ALandscapeProxy::DeleteLayer(const FGuid& InLayerGuid)
 			for (const FWeightmapLayerAllocationInfo& Allocation : LayerComponentData->WeightmapData.LayerAllocations)
 			{
 				UTexture2D* WeightmapTexture = LayerComponentData->WeightmapData.Textures[Allocation.WeightmapTextureIndex];
-				ULandscapeWeightmapUsage** Usage = WeightmapUsageMap.Find(WeightmapTexture);
+				UE_TRANSITIONAL_OBJECT_PTR(ULandscapeWeightmapUsage)* Usage = WeightmapUsageMap.Find(WeightmapTexture);
 
 				if (Usage != nullptr && (*Usage) != nullptr)
 				{
@@ -6082,9 +6082,9 @@ void ALandscape::UpdateLandscapeSplines(const FGuid& InTargetLayer, bool bInUpda
 
 			// Clear layers without affecting weightmap allocations
 			const bool bMarkPackageDirty = false;
-			ClearLayer(LandscapeSplinesTargetLayerGuid, (!bInForceUpdateAllCompoments && LandscapeSplinesAffectedComponents.Num()) ? &LandscapeSplinesAffectedComponents : nullptr, Clear_All, bMarkPackageDirty);
+			ClearLayer(LandscapeSplinesTargetLayerGuid, (!bInForceUpdateAllCompoments && LandscapeSplinesAffectedComponents.Num()) ? &ToRawPtrTSetUnsafe(LandscapeSplinesAffectedComponents) : nullptr, Clear_All, bMarkPackageDirty);
 			LandscapeSplinesAffectedComponents.Empty();
-			ModifiedComponent = &LandscapeSplinesAffectedComponents;
+			ModifiedComponent = &ToRawPtrTSetUnsafe(LandscapeSplinesAffectedComponents);
 			// For now, in Landscape Layer System Mode with a reserved layer for splines, we always update all the splines since we clear the whole layer first
 			bInUpdateOnlySelected = false;
 
diff --git a/Engine/Source/Runtime/Landscape/Private/LandscapeSplines.cpp b/Engine/Source/Runtime/Landscape/Private/LandscapeSplines.cpp
--- a/Engine/Source/Runtime/Landscape/Private/LandscapeSplines.cpp
+++ b/Engine/Source/Runtime/Landscape/Private/LandscapeSplines.cpp
@@ -910,7 +910,7 @@ void ULandscapeSplinesComponent::PostLoad()
 
 			for (auto& ForeignSplineSegmentData : ForeignWorldSplineData.ForeignSplineSegmentData)
 			{
-				for (auto* MeshComponent : ForeignSplineSegmentData.MeshComponents)
+				for (auto& MeshComponent : ForeignSplineSegmentData.MeshComponents)
 				{
 					MeshComponentForeignOwnersMap.Add(MeshComponent, ForeignSplineSegmentData.Identifier);
 				}
@@ -1210,7 +1210,7 @@ void ULandscapeSplinesComponent::RemoveAllForeignMeshComponents(ULandscapeSpline
 		auto* ForeignSplineSegmentData = ForeignWorldSplineData->FindSegmentData(Owner);
 		if (ForeignSplineSegmentData)
 		{
-			for (auto* MeshComponent : ForeignSplineSegmentData->MeshComponents)
+			for (auto& MeshComponent : ForeignSplineSegmentData->MeshComponents)
 			{
 				checkSlow(MeshComponentForeignOwnersMap.FindRef(MeshComponent) == Owner);
 				verifySlow(MeshComponentForeignOwnersMap.Remove(MeshComponent) == 1);
@@ -1411,7 +1411,7 @@ void ULandscapeSplinesComponent::DestroyOrphanedForeignSplineMeshComponents(UWor
 
 			if (!ForeignSplineSegment)
 			{
-				for (auto* MeshComponent : SegmentData.MeshComponents)
+				for (auto& MeshComponent : SegmentData.MeshComponents)
 				{
 					if (MeshComponent)
 					{
@@ -2445,7 +2445,7 @@ void ULandscapeSplineSegment::PostLoad()
 		ULandscapeSplinesComponent* OuterSplines = GetOuterULandscapeSplinesComponent();
 		if (OuterSplines->SplineEditorMesh != nullptr)
 		{
-			for (auto* LocalMeshComponent : LocalMeshComponents)
+			for (auto& LocalMeshComponent : LocalMeshComponents)
 			{
 				if (LocalMeshComponent->GetStaticMesh() == nullptr || LocalMeshComponent->GetStaticMesh() == OuterSplines->SplineEditorMesh)
 				{
@@ -2463,7 +2463,7 @@ void ULandscapeSplineSegment::PostLoad()
 			}
 		}
 
-		for (auto* LocalMeshComponent : LocalMeshComponents)
+		for (auto& LocalMeshComponent : LocalMeshComponents)
 		{
 			OuterSplines->MeshComponentLocalOwnersMap.Add(LocalMeshComponent, this);
 		}
@@ -2472,7 +2472,7 @@ void ULandscapeSplineSegment::PostLoad()
 	if (GetLinkerUE4Version() < VER_UE4_LANDSCAPE_SPLINE_CROSS_LEVEL_MESHES)
 	{
 		// Fix collision profile
-		for (auto* LocalMeshComponent : LocalMeshComponents) // ForeignMeshComponents didn't exist yet
+		for (auto& LocalMeshComponent : LocalMeshComponents) // ForeignMeshComponents didn't exist yet
 		{
 			UpdateMeshCollisionProfile(LocalMeshComponent);
 			LocalMeshComponent->SetFlags(RF_TextExportTransient);
@@ -2502,7 +2502,7 @@ void ULandscapeSplineSegment::PostLoad()
 			}
 		}
 
-		for (auto* LocalMeshComponent : LocalMeshComponents)
+		for (auto& LocalMeshComponent : LocalMeshComponents)
 		{
 			if (LocalMeshComponent != nullptr)
 			{
@@ -2525,7 +2525,7 @@ void ULandscapeSplineSegment::PostLoad()
 	// for our spline meshes.
 	if (SplinesAlwaysUseBlockAll)
 	{
-		for (auto* LocalMeshComponent : LocalMeshComponents)
+		for (auto& LocalMeshComponent : LocalMeshComponents)
 		{
 			UpdateMeshCollisionProfile(LocalMeshComponent);
 			LocalMeshComponent->Modify();
@@ -2558,7 +2558,7 @@ void ULandscapeSplineSegment::SetSplineSelected(bool bInSelected)
 	bSelected = bInSelected;
 	GetOuterULandscapeSplinesComponent()->MarkRenderStateDirty();
 
-	for (auto* LocalMeshComponent : LocalMeshComponents)
+	for (auto& LocalMeshComponent : LocalMeshComponents)
 	{
 		LocalMeshComponent->bSelected = bInSelected;
 		LocalMeshComponent->PushSelectionToProxy();
@@ -3149,7 +3149,7 @@ void ULandscapeSplineSegment::UpdateSplineEditorMesh()
 {
 	ULandscapeSplinesComponent* OuterSplines = CastChecked<ULandscapeSplinesComponent>(GetOuter());
 
-	for (auto* LocalMeshComponent : LocalMeshComponents)
+	for (auto& LocalMeshComponent : LocalMeshComponents)
 	{
 		if (OuterSplines->IsUsingEditorMesh(LocalMeshComponent))
 		{
@@ -3186,7 +3186,7 @@ void ULandscapeSplineSegment::DeleteSplinePoints()
 	if (LocalMeshComponents.Num() > 0)
 	{
 		OuterSplines->Modify();
-		for (auto* LocalMeshComponent : LocalMeshComponents)
+		for (auto& LocalMeshComponent : LocalMeshComponents)
 		{
 			checkSlow(OuterSplines->MeshComponentLocalOwnersMap.FindRef(LocalMeshComponent) == this);
 			verifySlow(OuterSplines->MeshComponentLocalOwnersMap.Remove(LocalMeshComponent) == 1);
@@ -3406,7 +3406,7 @@ void ULandscapeInfo::MoveSegment(ULandscapeSplineSegment* InSegment, TScriptInte
 	ToSplineComponent->Modify();
 		
 	// Delete all Mesh Components associated with the Segment. (Will get recreated in UpdateSplinePoints)
-	for (auto* MeshComponent : InSegment->LocalMeshComponents)
+	for (auto& MeshComponent : InSegment->LocalMeshComponents)
 	{
 		MeshComponent->Modify();
 		MeshComponent->UnregisterComponent();
diff --git a/Engine/Source/Runtime/LevelSequence/Private/LevelSequence.cpp b/Engine/Source/Runtime/LevelSequence/Private/LevelSequence.cpp
--- a/Engine/Source/Runtime/LevelSequence/Private/LevelSequence.cpp
+++ b/Engine/Source/Runtime/LevelSequence/Private/LevelSequence.cpp
@@ -697,5 +697,5 @@ void ULevelSequence::RemoveUserDataOfClass(TSubclassOf<UAssetUserData> InUserDat
 
 const TArray<UAssetUserData*>* ULevelSequence::GetAssetUserDataArray() const
 {
-	return &AssetUserData;
+	return &ToRawPtrTArrayUnsafe(AssetUserData);
 }
\ No newline at end of file
diff --git a/Engine/Source/Runtime/MaterialShaderQualitySettings/Private/MaterialShaderQualitySettings.cpp b/Engine/Source/Runtime/MaterialShaderQualitySettings/Private/MaterialShaderQualitySettings.cpp
--- a/Engine/Source/Runtime/MaterialShaderQualitySettings/Private/MaterialShaderQualitySettings.cpp
+++ b/Engine/Source/Runtime/MaterialShaderQualitySettings/Private/MaterialShaderQualitySettings.cpp
@@ -40,7 +40,7 @@ const FName& UMaterialShaderQualitySettings::GetPreviewPlatform()
 
 void UMaterialShaderQualitySettings::SetPreviewPlatform(FName PlatformName)
 {
-	 UShaderPlatformQualitySettings** FoundPlatform = ForwardSettingMap.Find(PlatformName);
+	 UE_TRANSITIONAL_OBJECT_PTR(UShaderPlatformQualitySettings)* FoundPlatform = ForwardSettingMap.Find(PlatformName);
 	 PreviewPlatformSettings = FoundPlatform == nullptr ? nullptr : *FoundPlatform;
 	 PreviewPlatformName = PlatformName;
 }
diff --git a/Engine/Source/Runtime/MediaAssets/Private/Assets/PlatformMediaSource.cpp b/Engine/Source/Runtime/MediaAssets/Private/Assets/PlatformMediaSource.cpp
--- a/Engine/Source/Runtime/MediaAssets/Private/Assets/PlatformMediaSource.cpp
+++ b/Engine/Source/Runtime/MediaAssets/Private/Assets/PlatformMediaSource.cpp
@@ -18,7 +18,7 @@ void UPlatformMediaSource::PreSave(const class ITargetPlatform* TargetPlatform)
 	// Do this only if we are cooking (aka: have a target platform)
 	if (TargetPlatform)
 	{
-		UMediaSource** PlatformMediaSource = PlatformMediaSources.Find(TargetPlatform->IniPlatformName());
+		UE_TRANSITIONAL_OBJECT_PTR(UMediaSource)* PlatformMediaSource = PlatformMediaSources.Find(TargetPlatform->IniPlatformName());
 		MediaSource = (PlatformMediaSource != nullptr) ? *PlatformMediaSource : nullptr;
 	}
 #endif
@@ -189,7 +189,7 @@ UMediaSource* UPlatformMediaSource::GetMediaSource() const
 {
 #if WITH_EDITORONLY_DATA
 	const FString RunningPlatformName(FPlatformProperties::IniPlatformName());
-	UMediaSource*const* PlatformMediaSource = PlatformMediaSources.Find(RunningPlatformName);
+	UE_TRANSITIONAL_OBJECT_PTR(UMediaSource) const* PlatformMediaSource = PlatformMediaSources.Find(RunningPlatformName);
 
 	if (PlatformMediaSource == nullptr)
 	{
diff --git a/Engine/Source/Runtime/MovieSceneTracks/Private/Tracks/MovieSceneCameraAnimTrack.cpp b/Engine/Source/Runtime/MovieSceneTracks/Private/Tracks/MovieSceneCameraAnimTrack.cpp
--- a/Engine/Source/Runtime/MovieSceneTracks/Private/Tracks/MovieSceneCameraAnimTrack.cpp
+++ b/Engine/Source/Runtime/MovieSceneTracks/Private/Tracks/MovieSceneCameraAnimTrack.cpp
@@ -110,7 +110,7 @@ void UMovieSceneCameraAnimTrack::GetCameraAnimSectionsAtTime(FFrameNumber Time,
 
 	for (auto Section : CameraAnimSections)
 	{
-		UMovieSceneCameraAnimSection* const CASection = dynamic_cast<UMovieSceneCameraAnimSection*>(Section);
+		UMovieSceneCameraAnimSection* const CASection = dynamic_cast<UMovieSceneCameraAnimSection*>(ToRawPtr(Section));
 		if (CASection && CASection->GetRange().Contains(Time))
 		{
 			OutSections.Add(CASection);
diff --git a/Engine/Source/Runtime/UMG/Private/Components/ComboBox.cpp b/Engine/Source/Runtime/UMG/Private/Components/ComboBox.cpp
--- a/Engine/Source/Runtime/UMG/Private/Components/ComboBox.cpp
+++ b/Engine/Source/Runtime/UMG/Private/Components/ComboBox.cpp
@@ -27,7 +27,7 @@ TSharedRef<SWidget> UComboBox::RebuildWidget()
 
 	MyComboBox =
 		SNew(SComboBox<UObject*>)
-		.OptionsSource(&Items)
+		.OptionsSource(&ToRawPtrTArrayUnsafe(Items))
 		.OnGenerateWidget(BIND_UOBJECT_DELEGATE(SComboBox<UObject*>::FOnGenerateWidget, HandleGenerateWidget))
 		.IsFocusable(bIsFocusable);
 
diff --git a/Engine/Source/Runtime/UMG/Private/UserWidget.cpp b/Engine/Source/Runtime/UMG/Private/UserWidget.cpp
--- a/Engine/Source/Runtime/UMG/Private/UserWidget.cpp
+++ b/Engine/Source/Runtime/UMG/Private/UserWidget.cpp
@@ -371,7 +371,7 @@ UWorld* UUserWidget::GetWorld() const
 
 UUMGSequencePlayer* UUserWidget::GetSequencePlayer(const UWidgetAnimation* InAnimation) const
 {
-	UUMGSequencePlayer*const* FoundPlayer = ActiveSequencePlayers.FindByPredicate(
+	UE_TRANSITIONAL_OBJECT_PTR(UUMGSequencePlayer) const* FoundPlayer = ActiveSequencePlayers.FindByPredicate(
 		[&](const UUMGSequencePlayer* Player)
 	{
 		return Player->GetAnimation() == InAnimation;
@@ -653,7 +653,7 @@ bool UUserWidget::IsAnimationPlayingForward(const UWidgetAnimation* InAnimation)
 {
 	if (InAnimation)
 	{
-		UUMGSequencePlayer** FoundPlayer = ActiveSequencePlayers.FindByPredicate([&](const UUMGSequencePlayer* Player) { return Player->GetAnimation() == InAnimation; });
+		UE_TRANSITIONAL_OBJECT_PTR(UUMGSequencePlayer)* FoundPlayer = ActiveSequencePlayers.FindByPredicate([&](const UUMGSequencePlayer* Player) { return Player->GetAnimation() == InAnimation; });
 
 		if (FoundPlayer)
 		{
