Fixing PVS 7.7 Issues: warning V758: The 'DIParam' reference becomes invalid when temporary object returned by a function is destroyed.

diff --git a/Engine/Plugins/FX/Niagara/Source/Niagara/Private/NiagaraEmitterInstanceBatcher.cpp b/Engine/Plugins/FX/Niagara/Source/Niagara/Private/NiagaraEmitterInstanceBatcher.cpp
--- a/Engine/Plugins/FX/Niagara/Source/Niagara/Private/NiagaraEmitterInstanceBatcher.cpp
+++ b/Engine/Plugins/FX/Niagara/Source/Niagara/Private/NiagaraEmitterInstanceBatcher.cpp
@@ -524,11 +524,12 @@ bool NiagaraEmitterInstanceBatcher::ResetDataInterfaces(const FNiagaraGPUSystemT
 	{
 		// Note: All stages will contain the same bindings so if they are valid for one they are valid for all, this could change in the future
 		const FNiagaraShaderRef ComputeShader = ShaderScript->GetShader(0);
+		const TMemoryImageArray<FNiagaraDataInterfaceParamRef>& DIParameters = ComputeShader->GetDIParameters();
 
 		uint32 InterfaceIndex = 0;
 		for (FNiagaraDataInterfaceProxy* Interface : Instance->DataInterfaceProxies)
 		{
-			const FNiagaraDataInterfaceParamRef& DIParam = ComputeShader->GetDIParameters()[InterfaceIndex];
+			const FNiagaraDataInterfaceParamRef& DIParam = DIParameters[InterfaceIndex];
 			if (DIParam.Parameters.IsValid())
 			{
 				const FNiagaraDataInterfaceArgs TmpContext(Interface, Tick.SystemInstanceID, this);
@@ -555,11 +556,12 @@ void NiagaraEmitterInstanceBatcher::PreStageInterface(const FNiagaraGPUSystemTic
 
 	// Note: All stages will contain the same bindings so if they are valid for one they are valid for all, this could change in the future
 	const FNiagaraShaderRef& ComputeShader = Instance->Context->GPUScript_RT->GetShader(0);
+	const TMemoryImageArray<FNiagaraDataInterfaceParamRef>& DIParameters = ComputeShader->GetDIParameters();
 
 	uint32 InterfaceIndex = 0;
 	for (FNiagaraDataInterfaceProxy* Interface : Instance->DataInterfaceProxies)
 	{
-		const FNiagaraDataInterfaceParamRef& DIParam = ComputeShader->GetDIParameters()[InterfaceIndex];
+		const FNiagaraDataInterfaceParamRef& DIParam = DIParameters[InterfaceIndex];
 		if (DIParam.Parameters.IsValid())
 		{
 			const FNiagaraDataInterfaceStageArgs TmpContext(Interface, Tick.SystemInstanceID, this, Instance, SimulationStageIndex, Instance->IsOutputStage(Interface, SimulationStageIndex), Instance->IsIterationStage(Interface, SimulationStageIndex));
@@ -578,11 +580,12 @@ void NiagaraEmitterInstanceBatcher::PostStageInterface(const FNiagaraGPUSystemTi
 
 	// Note: All stages will contain the same bindings so if they are valid for one they are valid for all, this could change in the future
 	const FNiagaraShaderRef& ComputeShader = Instance->Context->GPUScript_RT->GetShader(0);
+	const TMemoryImageArray<FNiagaraDataInterfaceParamRef>& DIParameters = ComputeShader->GetDIParameters();
 
 	uint32 InterfaceIndex = 0;
 	for (FNiagaraDataInterfaceProxy* Interface : Instance->DataInterfaceProxies)
 	{
-		const FNiagaraDataInterfaceParamRef& DIParam = ComputeShader->GetDIParameters()[InterfaceIndex];
+		const FNiagaraDataInterfaceParamRef& DIParam = DIParameters[InterfaceIndex];
 		if (DIParam.Parameters.IsValid())
 		{
 			const FNiagaraDataInterfaceStageArgs TmpContext(Interface, Tick.SystemInstanceID, this, Instance, SimulationStageIndex, Instance->IsOutputStage(Interface, SimulationStageIndex), Instance->IsIterationStage(Interface, SimulationStageIndex));
@@ -596,11 +599,12 @@ void NiagaraEmitterInstanceBatcher::PostSimulateInterface(const FNiagaraGPUSyste
 {
 	// Note: All stages will contain the same bindings so if they are valid for one they are valid for all, this could change in the future
 	const FNiagaraShaderRef ComputeShader = ShaderScript->GetShader(0);
+	const TMemoryImageArray<FNiagaraDataInterfaceParamRef>& DIParameters = ComputeShader->GetDIParameters();
 
 	uint32 InterfaceIndex = 0;
 	for (FNiagaraDataInterfaceProxy* Interface : Instance->DataInterfaceProxies)
 	{
-		const FNiagaraDataInterfaceParamRef& DIParam = ComputeShader->GetDIParameters()[InterfaceIndex];
+		const FNiagaraDataInterfaceParamRef& DIParam = DIParameters[InterfaceIndex];
 		if (DIParam.Parameters.IsValid())
 		{
 			const FNiagaraDataInterfaceArgs TmpContext(Interface, Tick.SystemInstanceID, this);
