Editgrate a couple fixes for Niagara automation tests and update ground truths of some tests that seem to have changed

#rnx

diff --git a/Engine/Plugins/FX/Niagara/Source/Niagara/Private/NiagaraDataInterfaceMeshRendererInfo.cpp b/Engine/Plugins/FX/Niagara/Source/Niagara/Private/NiagaraDataInterfaceMeshRendererInfo.cpp
--- a/Engine/Plugins/FX/Niagara/Source/Niagara/Private/NiagaraDataInterfaceMeshRendererInfo.cpp
+++ b/Engine/Plugins/FX/Niagara/Source/Niagara/Private/NiagaraDataInterfaceMeshRendererInfo.cpp
@@ -526,11 +526,17 @@ bool UNiagaraDataInterfaceMeshRendererInfo::CopyToInternal(UNiagaraDataInterface
 	OtherTyped->MeshRenderer = MeshRenderer;	
 	OtherTyped->Info = Info;
 
-	// Check to add a reference the the per-renderer data
+	// Check to add a reference to the per-renderer data
 	if (Info.IsValid())
 	{
 		Info->AddRef();
 	}
+	else if (MeshRenderer != nullptr && !EnumHasAnyFlags(OtherTyped->GetFlags(), RF_NeedPostLoad))
+	{
+		// This is currently necessary because, for some reason, it's possible that data interfaces that have not been post-loaded to be copied to
+		// another data interface object that has, and was previously causing this data to never get acquired on the copy
+		OtherTyped->Info = FNDIMeshRendererInfo::Acquire(*MeshRenderer);
+	}
 	return true;
 }
 
diff --git a/Engine/Plugins/FX/Niagara/Source/Niagara/Private/NiagaraRendererMeshes.cpp b/Engine/Plugins/FX/Niagara/Source/Niagara/Private/NiagaraRendererMeshes.cpp
--- a/Engine/Plugins/FX/Niagara/Source/Niagara/Private/NiagaraRendererMeshes.cpp
+++ b/Engine/Plugins/FX/Niagara/Source/Niagara/Private/NiagaraRendererMeshes.cpp
@@ -491,27 +491,37 @@ void FNiagaraRendererMeshes::InitializeSortInfo(
 
 	if (bEnableFrustumCulling)
 	{
-		OutSortInfo.CullPlanes.SetNumZeroed(6);
-
-		// Gather the culling planes from the view projection matrix
-		const FMatrix& ViewProj = ViewMatrices.GetViewProjectionMatrix();
-		ViewProj.GetFrustumNearPlane(OutSortInfo.CullPlanes[0]);
-		ViewProj.GetFrustumFarPlane(OutSortInfo.CullPlanes[1]);
-		ViewProj.GetFrustumTopPlane(OutSortInfo.CullPlanes[2]);
-		ViewProj.GetFrustumBottomPlane(OutSortInfo.CullPlanes[3]);
-		ViewProj.GetFrustumLeftPlane(OutSortInfo.CullPlanes[4]);
-
-		if (bIsInstancedStereo)
+		if (const FConvexVolume* ShadowFrustum = View.GetDynamicMeshElementsShadowCullFrustum())
 		{
-			// For Instanced Stereo, cull using an extended frustum that encompasses both eyes
-			ensure(View.StereoPass == eSSP_LEFT_EYE); // Sanity check that the primary eye is the left			
-			const FSceneView* RightEyeView = AllViewsInFamily[ViewIndex + 1];
-			check(RightEyeView);
-			GetViewMatrices(*RightEyeView).GetViewProjectionMatrix().GetFrustumRightPlane(OutSortInfo.CullPlanes[5]);
+			// Ensure we don't break the maximum number of planes here
+			// (For an accurate shadow frustum, a tight hull is formed from the silhouette and back-facing planes of the view frustum)
+			check(ShadowFrustum->Planes.Num() <= FNiagaraGPUSortInfo::MaxCullPlanes);
+			OutSortInfo.CullPlanes = ShadowFrustum->Planes;
 		}
 		else
 		{
-			ViewProj.GetFrustumRightPlane(OutSortInfo.CullPlanes[5]);
+			OutSortInfo.CullPlanes.SetNumZeroed(6);
+
+			// Gather the culling planes from the view projection matrix
+			const FMatrix& ViewProj = ViewMatrices.GetViewProjectionMatrix();
+			ViewProj.GetFrustumNearPlane(OutSortInfo.CullPlanes[0]);
+			ViewProj.GetFrustumFarPlane(OutSortInfo.CullPlanes[1]);
+			ViewProj.GetFrustumTopPlane(OutSortInfo.CullPlanes[2]);
+			ViewProj.GetFrustumBottomPlane(OutSortInfo.CullPlanes[3]);
+
+			ViewProj.GetFrustumLeftPlane(OutSortInfo.CullPlanes[4]);
+			if (bIsInstancedStereo)
+			{
+				// For Instanced Stereo, cull using an extended frustum that encompasses both eyes
+				ensure(View.StereoPass == eSSP_LEFT_EYE); // Sanity check that the primary eye is the left
+				const FSceneView* RightEyeView = AllViewsInFamily[ViewIndex + 1];
+				check(RightEyeView);
+				GetViewMatrices(*RightEyeView).GetViewProjectionMatrix().GetFrustumRightPlane(OutSortInfo.CullPlanes[5]);
+			}
+			else
+			{
+				ViewProj.GetFrustumRightPlane(OutSortInfo.CullPlanes[5]);
+			}
 		}
 	}
 
