Editgrate fixes to Niagara mesh particle shadows.
Temporarily submit a ground truth to fix the engine test until the full change is integrated from UE4.

#rnx

diff --git a/Engine/Plugins/FX/Niagara/Source/Niagara/Private/NiagaraRendererMeshes.cpp b/Engine/Plugins/FX/Niagara/Source/Niagara/Private/NiagaraRendererMeshes.cpp
--- a/Engine/Plugins/FX/Niagara/Source/Niagara/Private/NiagaraRendererMeshes.cpp
+++ b/Engine/Plugins/FX/Niagara/Source/Niagara/Private/NiagaraRendererMeshes.cpp
@@ -462,22 +462,31 @@ void FNiagaraRendererMeshes::InitializeSortInfo(
 	OutSortInfo.MeshIndexAttributeOffset = MeshIdxOffset;
 	OutSortInfo.SortAttributeOffset = VFVariables[SortVarIdx].GetGPUOffset();
 
-	auto GetViewMatrices = [](const FSceneView& View) -> const FViewMatrices&
+	auto GetViewMatrices = [](const FSceneView& View, FVector& OutViewOrigin) -> const FViewMatrices&
 	{
+		OutViewOrigin = View.ViewLocation;
+
 #if !(UE_BUILD_SHIPPING || UE_BUILD_TEST)
 		const FSceneViewState* ViewState = View.State != nullptr ? View.State->GetConcreteViewState() : nullptr;
 		if (ViewState && ViewState->bIsFrozen && ViewState->bIsFrozenViewMatricesCached)
 		{
 			// Use the frozen view for culling so we can test that it's working
-			return ViewState->CachedViewMatrices;
+			OutViewOrigin = ViewState->CachedViewMatrices.GetViewOrigin();
+
+			// Don't retrieve the cached matrices for shadow views
+			bool bIsShadow = View.GetDynamicMeshElementsShadowCullFrustum() != nullptr;
+			if (!bIsShadow)
+			{
+				return ViewState->CachedViewMatrices;
+			}
 		}
 #endif
+
 		return View.ViewMatrices;
 	};
 
 	const TArray<const FSceneView*>& AllViewsInFamily = View.Family->Views;
-	const FViewMatrices& ViewMatrices = GetViewMatrices(View);
-	OutSortInfo.ViewOrigin = ViewMatrices.GetViewOrigin();
+	const FViewMatrices& ViewMatrices = GetViewMatrices(View, OutSortInfo.ViewOrigin);
 	OutSortInfo.ViewDirection = ViewMatrices.GetViewMatrix().GetColumn(2);
 
 	if (View.StereoPass != eSSP_FULL)
@@ -486,7 +495,10 @@ void FNiagaraRendererMeshes::InitializeSortInfo(
 		const uint32 PairedViewIdx = (ViewIndex & 1) ? (ViewIndex - 1) : (ViewIndex + 1);
 		const FSceneView* PairedView = AllViewsInFamily[PairedViewIdx];
 		check(PairedView);
-		OutSortInfo.ViewOrigin = 0.5f * (OutSortInfo.ViewOrigin + GetViewMatrices(*PairedView).GetViewOrigin());
+
+		FVector PairedViewOrigin;
+		GetViewMatrices(*PairedView, PairedViewOrigin);
+		OutSortInfo.ViewOrigin = 0.5f * (OutSortInfo.ViewOrigin + PairedViewOrigin);
 	}
 
 	if (bEnableFrustumCulling)
@@ -497,6 +509,13 @@ void FNiagaraRendererMeshes::InitializeSortInfo(
 			// (For an accurate shadow frustum, a tight hull is formed from the silhouette and back-facing planes of the view frustum)
 			check(ShadowFrustum->Planes.Num() <= FNiagaraGPUSortInfo::MaxCullPlanes);
 			OutSortInfo.CullPlanes = ShadowFrustum->Planes;
+
+			// Remove pre-shadow translation to get the planes in world space
+			const FVector PreShadowTranslation = View.GetPreShadowTranslation();
+			for (FPlane& Plane : OutSortInfo.CullPlanes)
+			{
+				Plane.W -= FVector::DotProduct(FVector(Plane), PreShadowTranslation);
+			}
 		}
 		else
 		{
@@ -516,7 +535,8 @@ void FNiagaraRendererMeshes::InitializeSortInfo(
 				ensure(View.StereoPass == eSSP_LEFT_EYE); // Sanity check that the primary eye is the left
 				const FSceneView* RightEyeView = AllViewsInFamily[ViewIndex + 1];
 				check(RightEyeView);
-				GetViewMatrices(*RightEyeView).GetViewProjectionMatrix().GetFrustumRightPlane(OutSortInfo.CullPlanes[5]);
+				FVector RightEyePos;
+				GetViewMatrices(*RightEyeView, RightEyePos).GetViewProjectionMatrix().GetFrustumRightPlane(OutSortInfo.CullPlanes[5]);
 			}
 			else
 			{
