Fixing PVS 7.7 Issues:

warning V547: Expression 'XYZ' is always true/false.

diff --git a/Engine/Plugins/FX/Niagara/Source/NiagaraEditorWidgets/Private/Stack/SNiagaraStackFunctionInputValue.cpp b/Engine/Plugins/FX/Niagara/Source/NiagaraEditorWidgets/Private/Stack/SNiagaraStackFunctionInputValue.cpp
--- a/Engine/Plugins/FX/Niagara/Source/NiagaraEditorWidgets/Private/Stack/SNiagaraStackFunctionInputValue.cpp
+++ b/Engine/Plugins/FX/Niagara/Source/NiagaraEditorWidgets/Private/Stack/SNiagaraStackFunctionInputValue.cpp
@@ -652,7 +652,7 @@ void SNiagaraStackFunctionInputValue::CollectAllActions(FGraphActionListBuilderB
 			: LOCTEXT("ExpressionToolTipl", "Resolve this variable with a custom expression.");
 		TSharedPtr<FNiagaraMenuAction> ExpressionAction(new FNiagaraMenuAction(FText(), DisplayName, Tooltip, 0, FText(),
 			FNiagaraMenuAction::FOnExecuteStackAction::CreateSP(this, &SNiagaraStackFunctionInputValue::CustomExpressionSelected),
-			FNiagaraMenuAction::FCanExecuteStackAction::CreateLambda([bIsDataInterfaceOrObject]() { return bIsDataInterfaceOrObject == false; })));
+			FNiagaraMenuAction::FCanExecuteStackAction::CreateLambda([bIsDataInterfaceOrObject]() { return bIsDataInterfaceOrObject == false; }))); //-V547
 		OutAllActions.AddAction(ExpressionAction);
 	}
 
@@ -665,7 +665,7 @@ void SNiagaraStackFunctionInputValue::CollectAllActions(FGraphActionListBuilderB
 			: LOCTEXT("ScratchToolTipl", "Create a new dynamic input in the scratch pad.");
 		TSharedPtr<FNiagaraMenuAction> CreateScratchAction(new FNiagaraMenuAction(FText(), CreateDisplayName, CreateTooltip, 0, FText(),
 			FNiagaraMenuAction::FOnExecuteStackAction::CreateSP(this, &SNiagaraStackFunctionInputValue::CreateScratchSelected),
-			FNiagaraMenuAction::FCanExecuteStackAction::CreateLambda([bIsDataInterfaceOrObject]() { return bIsDataInterfaceOrObject == false; })));
+			FNiagaraMenuAction::FCanExecuteStackAction::CreateLambda([bIsDataInterfaceOrObject]() { return bIsDataInterfaceOrObject == false; }))); //-V547
 		OutAllActions.AddAction(CreateScratchAction);
 	}
 
diff --git a/Engine/Source/Runtime/Core/Private/ProfilingDebugging/Windows/MemoryTrace.cpp b/Engine/Source/Runtime/Core/Private/ProfilingDebugging/Windows/MemoryTrace.cpp
--- a/Engine/Source/Runtime/Core/Private/ProfilingDebugging/Windows/MemoryTrace.cpp
+++ b/Engine/Source/Runtime/Core/Private/ProfilingDebugging/Windows/MemoryTrace.cpp
@@ -821,17 +821,12 @@ FMalloc* MemoryTrace_Create(FMalloc* InMalloc)
 		GSummaryTrace->Initialize();
 #endif
 
-		bool bLight = (Mode == 1);
+		Backtracer_Create(InMalloc);
 
-		if (!bLight)
-		{
-			Backtracer_Create(InMalloc);
-		}
-
-		FVirtualWinApiHooks::Initialize(bLight);
+		FVirtualWinApiHooks::Initialize(false);
 
 		static FUndestructed<FMallocWrapper> MemoryTrace;
-		MemoryTrace.Construct(InMalloc, bLight);
+		MemoryTrace.Construct(InMalloc, false);
 
 		return &MemoryTrace;
 	}
diff --git a/Engine/Source/Runtime/Core/Private/Tests/Containers/BitArrayTests.cpp b/Engine/Source/Runtime/Core/Private/Tests/Containers/BitArrayTests.cpp
--- a/Engine/Source/Runtime/Core/Private/Tests/Containers/BitArrayTests.cpp
+++ b/Engine/Source/Runtime/Core/Private/Tests/Containers/BitArrayTests.cpp
@@ -35,7 +35,7 @@ namespace BitArrayTest
 	{
 		FString Out;
 		int32 Index = 0;
-		for (TBitArray<>::FConstIterator It(BitArray); It; ++It)
+		for (TBitArray<>::FConstIterator It(BitArray); It; ++It, ++Index)
 		{
 			if (Index != 0 && Index % 8 == 0)
 			{
diff --git a/Engine/Source/Runtime/Engine/Private/TimerManagerTests.cpp b/Engine/Source/Runtime/Engine/Private/TimerManagerTests.cpp
--- a/Engine/Source/Runtime/Engine/Private/TimerManagerTests.cpp
+++ b/Engine/Source/Runtime/Engine/Private/TimerManagerTests.cpp
@@ -205,7 +205,7 @@ bool TimerManagerTest_ValidTimer_HandleLoopingSetDuringExecute(UWorld* World, FA
 	FLoopingTestFunc::Handle = &Handle;
 	FLoopingTestFunc::TimerCalled = 0;
 
-	Test->TestTrue(TIMER_TEST_TEXT("Timer called count starts at 0"), FLoopingTestFunc::TimerCalled == 0);
+	Test->TestTrue(TIMER_TEST_TEXT("Timer called count starts at 0"), FLoopingTestFunc::TimerCalled == 0); //-V547
 
 	TimerManager.SetTimer(Handle, FTimerDelegate::CreateStatic(FLoopingTestFunc::TimerExecute), Rate, true);
 
@@ -213,12 +213,12 @@ bool TimerManagerTest_ValidTimer_HandleLoopingSetDuringExecute(UWorld* World, FA
 	TimerTest_TickWorld(World, KINDA_SMALL_NUMBER);
 
 	TimerTest_TickWorld(World, 3.0f);
-	Test->TestTrue(TIMER_TEST_TEXT("Timer was called first time"), FLoopingTestFunc::TimerCalled == 1);
+	Test->TestTrue(TIMER_TEST_TEXT("Timer was called first time"), FLoopingTestFunc::TimerCalled == 1); //-V547
 	Test->TestTrue(TIMER_TEST_TEXT("Timer was readded"), TimerManager.IsTimerActive(Handle));
 	Test->TestTrue(TIMER_TEST_TEXT("Timer was readded with correct time"), FMath::IsNearlyEqual(TimerManager.GetTimerRemaining(Handle), FLoopingTestFunc::NewTime, 1e-2f));
 
 	TimerTest_TickWorld(World, 1.1f);
-	Test->TestTrue(TIMER_TEST_TEXT("Timer was called second time"), FLoopingTestFunc::TimerCalled == 2);
+	Test->TestTrue(TIMER_TEST_TEXT("Timer was called second time"), FLoopingTestFunc::TimerCalled == 2); //-V547
 	Test->TestFalse(TIMER_TEST_TEXT("Timer handle no longer active"), TimerManager.IsTimerActive(Handle));
 
 	return true;
