Fixing PVS 7.7 Issues:

warning V700: Consider inspecting the 'T foo = foo = ...' expression. It is odd that variable is initialized through itself.

diff --git a/Engine/Plugins/2D/Paper2D/Source/PaperTiledImporter/Private/PaperTiledImporterFactory.cpp b/Engine/Plugins/2D/Paper2D/Source/PaperTiledImporter/Private/PaperTiledImporterFactory.cpp
--- a/Engine/Plugins/2D/Paper2D/Source/PaperTiledImporter/Private/PaperTiledImporterFactory.cpp
+++ b/Engine/Plugins/2D/Paper2D/Source/PaperTiledImporter/Private/PaperTiledImporterFactory.cpp
@@ -398,7 +398,7 @@ UTexture2D* UPaperTiledImporterFactory::ImportTexture(const FString& SourceFilen
 
 	//@TODO: Avoid the first compression, since we're going to recompress
 	TArray<UObject*> ImportedAssets = AssetToolsModule.Get().ImportAssets(FileNames, TargetSubPath);
-	UTexture2D* ImportedTexture = ImportedTexture = (ImportedAssets.Num() > 0) ? Cast<UTexture2D>(ImportedAssets[0]) : nullptr;
+	UTexture2D* ImportedTexture = (ImportedAssets.Num() > 0) ? Cast<UTexture2D>(ImportedAssets[0]) : nullptr;
 
 	if (ImportedTexture != nullptr)
 	{
diff --git a/Engine/Plugins/Enterprise/DatasmithImporter/Source/DatasmithImporter/Private/DatasmithFileProducer.cpp b/Engine/Plugins/Enterprise/DatasmithImporter/Source/DatasmithImporter/Private/DatasmithFileProducer.cpp
--- a/Engine/Plugins/Enterprise/DatasmithImporter/Source/DatasmithImporter/Private/DatasmithFileProducer.cpp
+++ b/Engine/Plugins/Enterprise/DatasmithImporter/Source/DatasmithImporter/Private/DatasmithFileProducer.cpp
@@ -911,7 +911,7 @@ void UDatasmithDirProducer::FixPlmXmlHierarchy()
 			SceneActor->SetActorLabel(Label);
 			SceneActor->SpriteScale = 0.1f;
 
-			USceneComponent* RootComponent = RootComponent = NewObject<USceneComponent>(SceneActor, *Label, RF_Transactional);
+			USceneComponent* RootComponent = NewObject<USceneComponent>(SceneActor, *Label, RF_Transactional);
 			RootComponent->SetWorldTransform(FTransform::Identity);
 			RootComponent->Mobility = EComponentMobility::Static;
 			RootComponent->bVisualizeComponent = true;
diff --git a/Engine/Plugins/Enterprise/MDLImporter/Source/MDLImporter/Private/generator/FunctionLoader.cpp b/Engine/Plugins/Enterprise/MDLImporter/Source/MDLImporter/Private/generator/FunctionLoader.cpp
--- a/Engine/Plugins/Enterprise/MDLImporter/Source/MDLImporter/Private/generator/FunctionLoader.cpp
+++ b/Engine/Plugins/Enterprise/MDLImporter/Source/MDLImporter/Private/generator/FunctionLoader.cpp
@@ -261,7 +261,7 @@ namespace Generator
 		FString   FunctionPackageName = UPackageTools::SanitizePackageName(*(AssetPath / FunctionName));
 		UPackage* Package             = CreatePackage(*FunctionPackageName);
 
-		UMaterialFunction* Function = Function = dynamic_cast<UMaterialFunction*>(
+		UMaterialFunction* Function = dynamic_cast<UMaterialFunction*>(
 		    FunctionFactory->FactoryCreateNew(UMaterialFunction::StaticClass(), Package, *FunctionName, RF_Public | RF_Standalone, nullptr, GWarn));
 
 		check(Function);
diff --git a/Engine/Plugins/FX/Niagara/Source/NiagaraEditor/Private/ViewModels/NiagaraScratchPadViewModel.cpp b/Engine/Plugins/FX/Niagara/Source/NiagaraEditor/Private/ViewModels/NiagaraScratchPadViewModel.cpp
--- a/Engine/Plugins/FX/Niagara/Source/NiagaraEditor/Private/ViewModels/NiagaraScratchPadViewModel.cpp
+++ b/Engine/Plugins/FX/Niagara/Source/NiagaraEditor/Private/ViewModels/NiagaraScratchPadViewModel.cpp
@@ -438,7 +438,7 @@ TSharedPtr<FNiagaraScratchPadScriptViewModel> UNiagaraScratchPadViewModel::Creat
 	TArray<UNiagaraScript*>* TargetScripts;
 	GetOuterAndTargetScripts(GetSystemViewModel(), ScriptOuter, TargetScripts);
 
-	UNiagaraScript* NewScript = NewScript = CastChecked<UNiagaraScript>(StaticDuplicateObject(ScriptToDuplicate, ScriptOuter, GetUniqueScriptName(ScriptOuter, *ScriptToDuplicate->GetFName().ToString())));
+	UNiagaraScript* NewScript = CastChecked<UNiagaraScript>(StaticDuplicateObject(ScriptToDuplicate, ScriptOuter, GetUniqueScriptName(ScriptOuter, *ScriptToDuplicate->GetFName().ToString())));
 	NewScript->ClearFlags(RF_Public | RF_Standalone);
 	ScriptOuter->Modify();
 	TargetScripts->Add(NewScript);
diff --git a/Engine/Source/Editor/UnrealEd/Private/ThumbnailManager.cpp b/Engine/Source/Editor/UnrealEd/Private/ThumbnailManager.cpp
--- a/Engine/Source/Editor/UnrealEd/Private/ThumbnailManager.cpp
+++ b/Engine/Source/Editor/UnrealEd/Private/ThumbnailManager.cpp
@@ -121,7 +121,7 @@ FThumbnailRenderingInfo* UThumbnailManager::GetRenderingInfo(UObject* Object)
 	TArray<FThumbnailRenderingInfo>& ThumbnailTypes = RenderableThumbnailTypes;
 
 	// Get the class to check against.
-	UClass *ClassToCheck = ClassToCheck = Object->GetClass();
+	UClass *ClassToCheck = Object->GetClass();
 	
 	// Search for the cached entry and do the slower if not found
 	FThumbnailRenderingInfo* RenderInfo = RenderInfoMap.FindRef(ClassToCheck);
diff --git a/Engine/Source/Runtime/Engine/Private/Particles/ParticleModules.cpp b/Engine/Source/Runtime/Engine/Private/Particles/ParticleModules.cpp
--- a/Engine/Source/Runtime/Engine/Private/Particles/ParticleModules.cpp
+++ b/Engine/Source/Runtime/Engine/Private/Particles/ParticleModules.cpp
@@ -1031,7 +1031,7 @@ void UParticleModuleOrientationAxisLock::PostEditChangeProperty(FPropertyChanged
 		check(Emitter);
 		OuterObj = Emitter->GetOuter();
 	}
-	UParticleSystem* PartSys = PartSys = CastChecked<UParticleSystem>(OuterObj);
+	UParticleSystem* PartSys = CastChecked<UParticleSystem>(OuterObj);
 
 	FProperty* PropertyThatChanged = PropertyChangedEvent.Property;
 	if (PropertyThatChanged)
diff --git a/Engine/Source/Runtime/Engine/Private/Particles/ParticleModules_Event.cpp b/Engine/Source/Runtime/Engine/Private/Particles/ParticleModules_Event.cpp
--- a/Engine/Source/Runtime/Engine/Private/Particles/ParticleModules_Event.cpp
+++ b/Engine/Source/Runtime/Engine/Private/Particles/ParticleModules_Event.cpp
@@ -102,7 +102,7 @@ void UParticleModuleEventGenerator::PostEditChangeProperty(FPropertyChangedEvent
 		check(Emitter);
 		OuterObj = Emitter->GetOuter();
 	}
-	UParticleSystem* PartSys = PartSys = CastChecked<UParticleSystem>(OuterObj);
+	UParticleSystem* PartSys = CastChecked<UParticleSystem>(OuterObj);
 	if (PartSys)
 	{
 		PartSys->PostEditChangeProperty(PropertyChangedEvent);
