U5 - Remove some of the Niagara Profiling scopes as they pollute the unreal Insights capture with statnamedevents

diff --git a/Engine/Plugins/FX/Niagara/Source/NiagaraEditor/Private/NiagaraGraph.cpp b/Engine/Plugins/FX/Niagara/Source/NiagaraEditor/Private/NiagaraGraph.cpp
--- a/Engine/Plugins/FX/Niagara/Source/NiagaraEditor/Private/NiagaraGraph.cpp
+++ b/Engine/Plugins/FX/Niagara/Source/NiagaraEditor/Private/NiagaraGraph.cpp
@@ -37,6 +37,8 @@ DECLARE_CYCLE_STAT(TEXT("NiagaraEditor - Graph - FindInputNodes_FindInputNodes_S
 DECLARE_CYCLE_STAT(TEXT("NiagaraEditor - Graph - FindOutputNode"), STAT_NiagaraEditor_Graph_FindOutputNode, STATGROUP_NiagaraEditor);
 DECLARE_CYCLE_STAT(TEXT("NiagaraEditor - Graph - BuildTraversalHelper"), STAT_NiagaraEditor_Graph_BuildTraversalHelper, STATGROUP_NiagaraEditor);
 
+#define NIAGARA_SCOPE_CYCLE_COUNTER(x) //SCOPE_CYCLE_COUNTER(x)
+
 bool bWriteToLog = false;
 
 #define LOCTEXT_NAMESPACE "NiagaraGraph"
@@ -978,7 +980,7 @@ void UNiagaraGraph::FindEquivalentOutputNodes(ENiagaraScriptUsage TargetUsageTyp
 
 UNiagaraNodeOutput* UNiagaraGraph::FindOutputNode(ENiagaraScriptUsage TargetUsageType, FGuid TargetUsageId) const
 {
-	SCOPE_CYCLE_COUNTER(STAT_NiagaraEditor_Graph_FindOutputNode);
+	NIAGARA_SCOPE_CYCLE_COUNTER(STAT_NiagaraEditor_Graph_FindOutputNode);
 	for (UEdGraphNode* Node : Nodes)
 	{
 		if (UNiagaraNodeOutput* OutNode = Cast<UNiagaraNodeOutput>(Node))
@@ -994,7 +996,7 @@ UNiagaraNodeOutput* UNiagaraGraph::FindOutputNode(ENiagaraScriptUsage TargetUsag
 
 UNiagaraNodeOutput* UNiagaraGraph::FindEquivalentOutputNode(ENiagaraScriptUsage TargetUsageType, FGuid TargetUsageId) const
 {
-	SCOPE_CYCLE_COUNTER(STAT_NiagaraEditor_Graph_FindOutputNode);
+	NIAGARA_SCOPE_CYCLE_COUNTER(STAT_NiagaraEditor_Graph_FindOutputNode);
 	for (UEdGraphNode* Node : Nodes)
 	{
 		if (UNiagaraNodeOutput* OutNode = Cast<UNiagaraNodeOutput>(Node))
@@ -1016,7 +1018,7 @@ void BuildTraversalHelper(TArray<class UNiagaraNode*>& OutNodesTraversed, UNiaga
 		return;
 	}
 
-	SCOPE_CYCLE_COUNTER(STAT_NiagaraEditor_Graph_BuildTraversalHelper);
+	NIAGARA_SCOPE_CYCLE_COUNTER(STAT_NiagaraEditor_Graph_BuildTraversalHelper);
 
 	for (UEdGraphPin* Pin : CurrentNode->GetAllPins())
 	{
@@ -1061,12 +1063,12 @@ void UNiagaraGraph::BuildTraversal(TArray<class UNiagaraNode*>& OutNodesTraverse
 
 void UNiagaraGraph::FindInputNodes(TArray<UNiagaraNodeInput*>& OutInputNodes, UNiagaraGraph::FFindInputNodeOptions Options) const
 {
-	SCOPE_CYCLE_COUNTER(STAT_NiagaraEditor_Graph_FindInputNodes);
+	NIAGARA_SCOPE_CYCLE_COUNTER(STAT_NiagaraEditor_Graph_FindInputNodes);
 	TArray<UNiagaraNodeInput*> InputNodes;
 
 	if (!Options.bFilterByScriptUsage)
 	{
-		SCOPE_CYCLE_COUNTER(STAT_NiagaraEditor_Graph_FindInputNodes_NotFilterUsage);
+		NIAGARA_SCOPE_CYCLE_COUNTER(STAT_NiagaraEditor_Graph_FindInputNodes_NotFilterUsage);
 
 		for (UEdGraphNode* Node : Nodes)
 		{
@@ -1083,7 +1085,7 @@ void UNiagaraGraph::FindInputNodes(TArray<UNiagaraNodeInput*>& OutInputNodes, UN
 	}
 	else
 	{
-		SCOPE_CYCLE_COUNTER(STAT_NiagaraEditor_Graph_FindInputNodes_FilterUsage);
+		NIAGARA_SCOPE_CYCLE_COUNTER(STAT_NiagaraEditor_Graph_FindInputNodes_FilterUsage);
 
 		TArray<class UNiagaraNode*> Traversal;
 		BuildTraversal(Traversal, Options.TargetScriptUsage, Options.TargetScriptUsageId);
@@ -1102,7 +1104,7 @@ void UNiagaraGraph::FindInputNodes(TArray<UNiagaraNodeInput*>& OutInputNodes, UN
 
 	if (Options.bFilterDuplicates)
 	{
-		SCOPE_CYCLE_COUNTER(STAT_NiagaraEditor_Graph_FindInputNodes_FilterDupes);
+		NIAGARA_SCOPE_CYCLE_COUNTER(STAT_NiagaraEditor_Graph_FindInputNodes_FilterDupes);
 
 		for (UNiagaraNodeInput* InputNode : InputNodes)
 		{
@@ -1131,7 +1133,7 @@ void UNiagaraGraph::FindInputNodes(TArray<UNiagaraNodeInput*>& OutInputNodes, UN
 
 	if (Options.bSort)
 	{
-		SCOPE_CYCLE_COUNTER(STAT_NiagaraEditor_Graph_FindInputNodes_Sort);
+		NIAGARA_SCOPE_CYCLE_COUNTER(STAT_NiagaraEditor_Graph_FindInputNodes_Sort);
 
 		UNiagaraNodeInput::SortNodes(OutInputNodes);
 	}
@@ -2843,4 +2845,5 @@ const TMap<FNiagaraVariable, FInputPinsAndOutputPins> UNiagaraGraph::CollectVars
 	return VarToPinsMap;
 }
 
+#undef NIAGARA_SCOPE_CYCLE_COUNTER
 #undef LOCTEXT_NAMESPACE
diff --git a/Engine/Plugins/FX/Niagara/Source/NiagaraEditor/Private/NiagaraHlslTranslator.cpp b/Engine/Plugins/FX/Niagara/Source/NiagaraEditor/Private/NiagaraHlslTranslator.cpp
--- a/Engine/Plugins/FX/Niagara/Source/NiagaraEditor/Private/NiagaraHlslTranslator.cpp
+++ b/Engine/Plugins/FX/Niagara/Source/NiagaraEditor/Private/NiagaraHlslTranslator.cpp
@@ -80,6 +80,8 @@ DECLARE_CYCLE_STAT(TEXT("Niagara - HlslTranslator - GenerateFunctionSignature_Ou
 DECLARE_CYCLE_STAT(TEXT("Niagara - HlslTranslator - GenerateFunctionSignature_Inputs"), STAT_NiagaraEditor_Module_NiagaraHLSLTranslator_GenerateFunctionSignature_Inputs, STATGROUP_NiagaraEditor);
 DECLARE_CYCLE_STAT(TEXT("Niagara - HlslTranslator - GenerateFunctionSignature_FindInputNodes"), STAT_NiagaraEditor_Module_NiagaraHLSLTranslator_GenerateFunctionSignature_FindInputNodes, STATGROUP_NiagaraEditor);
 
+#define NIAGARA_SCOPE_CYCLE_COUNTER(x) //SCOPE_CYCLE_COUNTER(x)
+
 static FNiagaraShaderQueueTickable NiagaraShaderQueueProcessor;
 
 // not pretty. TODO: refactor
@@ -430,14 +432,14 @@ bool FHlslNiagaraTranslator::ValidateTypePins(UNiagaraNode* NodeToValidate)
 void FHlslNiagaraTranslator::GenerateFunctionSignature(ENiagaraScriptUsage ScriptUsage, FString InName, const FString& InFullName, UNiagaraGraph* FuncGraph, TArray<int32>& Inputs,
 	bool bHasNumericInputs, bool bHasParameterMapParameters, TArray<UEdGraphPin*> StaticSwitchValues, FNiagaraFunctionSignature& OutSig)const
 {
-	SCOPE_CYCLE_COUNTER(STAT_NiagaraEditor_Module_NiagaraHLSLTranslator_GenerateFunctionSignature);
+	NIAGARA_SCOPE_CYCLE_COUNTER(STAT_NiagaraEditor_Module_NiagaraHLSLTranslator_GenerateFunctionSignature);
 
 	TArray<FNiagaraVariable> InputVars;
 	TArray<UNiagaraNodeInput*> InputsNodes;
 	bool bHasDIParameters = false;
 
 	{
-		SCOPE_CYCLE_COUNTER(STAT_NiagaraEditor_Module_NiagaraHLSLTranslator_GenerateFunctionSignature_FindInputNodes);
+		NIAGARA_SCOPE_CYCLE_COUNTER(STAT_NiagaraEditor_Module_NiagaraHLSLTranslator_GenerateFunctionSignature_FindInputNodes);
 		InputsNodes.Reserve(100);
 		UNiagaraGraph::FFindInputNodeOptions Options;
 		Options.bSort = true;
@@ -459,7 +461,7 @@ void FHlslNiagaraTranslator::GenerateFunctionSignature(ENiagaraScriptUsage Scrip
 	}
 
 	{
-		SCOPE_CYCLE_COUNTER(STAT_NiagaraEditor_Module_NiagaraHLSLTranslator_GenerateFunctionSignature_Inputs);
+		NIAGARA_SCOPE_CYCLE_COUNTER(STAT_NiagaraEditor_Module_NiagaraHLSLTranslator_GenerateFunctionSignature_Inputs);
 
 		InName.Reserve(100 * InputsNodes.Num());
 		InputVars.Reserve(InputsNodes.Num());
@@ -505,7 +507,7 @@ void FHlslNiagaraTranslator::GenerateFunctionSignature(ENiagaraScriptUsage Scrip
 
 	TArray<FNiagaraVariable> OutputVars;
 	{
-		SCOPE_CYCLE_COUNTER(STAT_NiagaraEditor_Module_NiagaraHLSLTranslator_GenerateFunctionSignature_Outputs);
+		NIAGARA_SCOPE_CYCLE_COUNTER(STAT_NiagaraEditor_Module_NiagaraHLSLTranslator_GenerateFunctionSignature_Outputs);
 
 		OutputVars.Reserve(10);
 		FuncGraph->GetOutputNodeVariables(ScriptUsage, OutputVars);
@@ -530,7 +532,7 @@ void FHlslNiagaraTranslator::GenerateFunctionSignature(ENiagaraScriptUsage Scrip
 		|| bHasParameterMapParameters || bHasDIParameters)
 		&& (ModuleAliasStr != nullptr || EmitterAliasStr != nullptr))
 	{
-		SCOPE_CYCLE_COUNTER(STAT_NiagaraEditor_Module_NiagaraHLSLTranslator_GenerateFunctionSignature_UniqueDueToMaps);
+		NIAGARA_SCOPE_CYCLE_COUNTER(STAT_NiagaraEditor_Module_NiagaraHLSLTranslator_GenerateFunctionSignature_UniqueDueToMaps);
 		FString SignatureName;
 		SignatureName.Reserve(1024);
 		if (ModuleAliasStr != nullptr)
@@ -671,7 +673,7 @@ void FHlslNiagaraTranslator::BuildMissingDefaults()
 
 FString FHlslNiagaraTranslator::BuildParameterMapHlslDefinitions(TArray<FNiagaraVariable>& PrimaryDataSetOutputEntries)
 {
-	SCOPE_CYCLE_COUNTER(STAT_NiagaraEditor_HlslTranslator_BuildParameterMapHlslDefinitions);
+	NIAGARA_SCOPE_CYCLE_COUNTER(STAT_NiagaraEditor_HlslTranslator_BuildParameterMapHlslDefinitions);
 	FString HlslOutputString;
 
 	// Determine the unique parameter map structs...
@@ -1133,7 +1135,7 @@ const FNiagaraTranslateResults &FHlslNiagaraTranslator::Translate(const FNiagara
 	TRACE_CPUPROFILER_EVENT_SCOPE(NiagaraHlslTranslate);
 	TRACE_CPUPROFILER_EVENT_SCOPE_TEXT_ON_CHANNEL(*InCompileOptions.GetPathName(), NiagaraChannel);
 
-	SCOPE_CYCLE_COUNTER(STAT_NiagaraEditor_HlslTranslator_Translate);
+	NIAGARA_SCOPE_CYCLE_COUNTER(STAT_NiagaraEditor_HlslTranslator_Translate);
 	check(InCompileData);
 
 	CompileOptions = InCompileOptions;
@@ -4275,7 +4277,7 @@ int32 FHlslNiagaraTranslator::GetRapidIterationParameter(const FNiagaraVariable&
 int32 FHlslNiagaraTranslator::GetParameter(const FNiagaraVariable& Parameter)
 {
 
-	SCOPE_CYCLE_COUNTER(STAT_NiagaraEditor_HlslTranslator_GetParameter);
+	NIAGARA_SCOPE_CYCLE_COUNTER(STAT_NiagaraEditor_HlslTranslator_GetParameter);
 	if (!AddStructToDefinitionSet(Parameter.GetType()))
 	{
 		Error(FText::Format(LOCTEXT("GetParameterFail", "Cannot handle type {0}! Variable: {1}"), Parameter.GetType().GetNameText(), FText::FromName(Parameter.GetName())), nullptr, nullptr);
@@ -4663,7 +4665,7 @@ void FHlslNiagaraTranslator::InitializeParameterMapDefaults(int32 ParamMapHistor
 
 void FHlslNiagaraTranslator::Output(UNiagaraNodeOutput* OutputNode, const TArray<int32>& ComputedInputs)
 {
-	SCOPE_CYCLE_COUNTER(STAT_NiagaraEditor_HlslTranslator_Output);
+	NIAGARA_SCOPE_CYCLE_COUNTER(STAT_NiagaraEditor_HlslTranslator_Output);
 
 	TArray<FNiagaraVariable> Attributes;
 	TArray<int32> Inputs;
@@ -4804,21 +4806,21 @@ FString FHlslNiagaraTranslator::GetDataSetAccessSymbol(FNiagaraDataSetID DataSet
 
 void FHlslNiagaraTranslator::ParameterMapForBegin(UNiagaraNodeParameterMapFor* ForNode, int32 IterationCount)
 {
-	SCOPE_CYCLE_COUNTER(STAT_NiagaraEditor_HlslTranslator_MapForBegin);
+	NIAGARA_SCOPE_CYCLE_COUNTER(STAT_NiagaraEditor_HlslTranslator_MapForBegin);
 
 	AddBodyChunk(TEXT(""), TEXT("for(int index = 0; index < {0}; ++index)\n\t{"), FNiagaraTypeDefinition::GetIntDef(), IterationCount, false, false);
 }
 
 void FHlslNiagaraTranslator::ParameterMapForEnd(UNiagaraNodeParameterMapFor* ForNode)
 {
-	SCOPE_CYCLE_COUNTER(STAT_NiagaraEditor_HlslTranslator_MapForEnd);
+	NIAGARA_SCOPE_CYCLE_COUNTER(STAT_NiagaraEditor_HlslTranslator_MapForEnd);
 
 	AddBodyChunk(TEXT(""), TEXT("}"), FNiagaraTypeDefinition::GetIntDef(), false, false);
 }
 
 void FHlslNiagaraTranslator::ParameterMapSet(UNiagaraNodeParameterMapSet* SetNode, TArrayView<const FCompiledPin> Inputs, TArray<int32>& Outputs)
 {
-	SCOPE_CYCLE_COUNTER(STAT_NiagaraEditor_HlslTranslator_MapSet);
+	NIAGARA_SCOPE_CYCLE_COUNTER(STAT_NiagaraEditor_HlslTranslator_MapSet);
 
 	Outputs.SetNum(1);
 
@@ -5270,7 +5272,7 @@ FString FHlslNiagaraTranslator::GetParameterMapInstanceName(int32 ParamMapHistor
 
 void FHlslNiagaraTranslator::Emitter(class UNiagaraNodeEmitter* EmitterNode, TArray<int32>& Inputs, TArray<int32>& Outputs)
 {
-	SCOPE_CYCLE_COUNTER(STAT_NiagaraEditor_Module_NiagaraHLSLTranslator_Emitter);
+	NIAGARA_SCOPE_CYCLE_COUNTER(STAT_NiagaraEditor_Module_NiagaraHLSLTranslator_Emitter);
 
 	// Just pass through the input parameter map pin if the node isn't enabled...
 	if (!EmitterNode->IsNodeEnabled())
@@ -5367,7 +5369,7 @@ void FHlslNiagaraTranslator::Emitter(class UNiagaraNodeEmitter* EmitterNode, TAr
 
 void FHlslNiagaraTranslator::ParameterMapGet(UNiagaraNodeParameterMapGet* GetNode, TArrayView<const int32> Inputs, TArray<int32>& Outputs)
 {
-	SCOPE_CYCLE_COUNTER(STAT_NiagaraEditor_Module_NiagaraHLSLTranslator_MapGet);
+	NIAGARA_SCOPE_CYCLE_COUNTER(STAT_NiagaraEditor_Module_NiagaraHLSLTranslator_MapGet);
 
 	FPinCollectorArray OutputPins;
 	GetNode->GetOutputPins(OutputPins);
@@ -6062,7 +6064,7 @@ int32 FHlslNiagaraTranslator::RegisterDataInterface(FNiagaraVariable& Var, UNiag
 
 void FHlslNiagaraTranslator::Operation(class UNiagaraNodeOp* Operation, TArray<int32>& Inputs, TArray<int32>& Outputs)
 {
-	SCOPE_CYCLE_COUNTER(STAT_NiagaraEditor_HlslTranslator_Operation);
+	NIAGARA_SCOPE_CYCLE_COUNTER(STAT_NiagaraEditor_HlslTranslator_Operation);
 
 	// Use the pins to determine the output type here since they may have been changed due to numeric pin fix up.
 	const FNiagaraOpInfo* OpInfo = FNiagaraOpInfo::GetOpInfo(Operation->OpName);
@@ -6101,7 +6103,7 @@ void FHlslNiagaraTranslator::Operation(class UNiagaraNodeOp* Operation, TArray<i
 
 void FHlslNiagaraTranslator::FunctionCall(UNiagaraNodeFunctionCall* FunctionNode, TArray<int32>& Inputs, TArray<int32>& Outputs)
 {
-	SCOPE_CYCLE_COUNTER(STAT_NiagaraEditor_Module_NiagaraHLSLTranslator_FunctionCall);
+	NIAGARA_SCOPE_CYCLE_COUNTER(STAT_NiagaraEditor_Module_NiagaraHLSLTranslator_FunctionCall);
 
 	FPinCollectorArray CallOutputs;
 	FPinCollectorArray CallInputs;
@@ -6747,7 +6749,7 @@ void FHlslNiagaraTranslator::ProcessCustomHlsl(const FString& InCustomHlsl, ENia
 void FHlslNiagaraTranslator::HandleCustomHlslNode(UNiagaraNodeCustomHlsl* CustomFunctionHlsl, ENiagaraScriptUsage& OutScriptUsage, FString& OutName, FString& OutFullName, bool& bOutCustomHlsl, FString& OutCustomHlsl,
 	FNiagaraFunctionSignature& OutSignature, TArray<int32>& Inputs)
 {
-	SCOPE_CYCLE_COUNTER(STAT_NiagaraEditor_Module_NiagaraHLSLTranslator_CustomHLSL);
+	NIAGARA_SCOPE_CYCLE_COUNTER(STAT_NiagaraEditor_Module_NiagaraHLSLTranslator_CustomHLSL);
 	if (!CustomFunctionHlsl)
 	{
 		return;
@@ -6815,12 +6817,12 @@ void FHlslNiagaraTranslator::RegisterFunctionCall(ENiagaraScriptUsage ScriptUsag
 	FNiagaraFunctionSignature& InSignature, bool bIsCustomHlsl, const FString& InCustomHlsl, TArray<int32>& Inputs, TArrayView<UEdGraphPin* const> CallInputs, TArrayView<UEdGraphPin* const> CallOutputs,
 	FNiagaraFunctionSignature& OutSignature)
 {
-	SCOPE_CYCLE_COUNTER(STAT_NiagaraEditor_Module_NiagaraHLSLTranslator_RegisterFunctionCall);
+	NIAGARA_SCOPE_CYCLE_COUNTER(STAT_NiagaraEditor_Module_NiagaraHLSLTranslator_RegisterFunctionCall);
 
 	//////////////////////////////////////////////////////////////////////////
 	if (Source)
 	{
-		SCOPE_CYCLE_COUNTER(STAT_NiagaraEditor_Module_NiagaraHLSLTranslator_RegisterFunctionCall_Source);
+		NIAGARA_SCOPE_CYCLE_COUNTER(STAT_NiagaraEditor_Module_NiagaraHLSLTranslator_RegisterFunctionCall_Source);
 		UNiagaraGraph* SourceGraph = CastChecked<UNiagaraGraph>(Source->NodeGraph);
 
 		bool bHasNumericInputs = false;
@@ -6872,7 +6874,7 @@ void FHlslNiagaraTranslator::RegisterFunctionCall(ENiagaraScriptUsage ScriptUsag
 		FNiagaraFunctionBody* FuncBody = Functions.Find(OutSignature);
 		if (!FuncBody)
 		{
-			SCOPE_CYCLE_COUNTER(STAT_NiagaraEditor_Module_NiagaraHLSLTranslator_FuncBody);
+			NIAGARA_SCOPE_CYCLE_COUNTER(STAT_NiagaraEditor_Module_NiagaraHLSLTranslator_FuncBody);
 
 			if (OutSignature.Name == NAME_None)
 			{
@@ -6996,13 +6998,13 @@ void FHlslNiagaraTranslator::RegisterFunctionCall(ENiagaraScriptUsage ScriptUsag
 			ENiagaraCodeChunkMode OldMode = CurrentBodyChunkMode;
 			CurrentBodyChunkMode = ENiagaraCodeChunkMode::Body;
 			{
-				SCOPE_CYCLE_COUNTER(STAT_NiagaraEditor_Module_NiagaraHLSLTranslator_RegisterFunctionCall_Compile);
+				NIAGARA_SCOPE_CYCLE_COUNTER(STAT_NiagaraEditor_Module_NiagaraHLSLTranslator_RegisterFunctionCall_Compile);
 				FuncOutput->Compile(ThisTranslator, FuncOutputChunks);
 			}
 			CurrentBodyChunkMode = OldMode;
 
 			{
-				SCOPE_CYCLE_COUNTER(STAT_NiagaraEditor_Module_NiagaraHLSLTranslator_RegisterFunctionCall_FunctionDefStr);
+				NIAGARA_SCOPE_CYCLE_COUNTER(STAT_NiagaraEditor_Module_NiagaraHLSLTranslator_RegisterFunctionCall_FunctionDefStr);
 
 				FNiagaraFunctionBody& FunctionBody = Functions.Add(OutSignature);
 
@@ -7116,7 +7118,7 @@ void FHlslNiagaraTranslator::RegisterFunctionCall(ENiagaraScriptUsage ScriptUsag
 	}
 	else
 	{
-		SCOPE_CYCLE_COUNTER(STAT_NiagaraEditor_Module_NiagaraHLSLTranslator_RegisterFunctionCall_Signature);
+		NIAGARA_SCOPE_CYCLE_COUNTER(STAT_NiagaraEditor_Module_NiagaraHLSLTranslator_RegisterFunctionCall_Signature);
 
 		check(InSignature.IsValid());
 		check(Inputs.Num() > 0);
@@ -7232,7 +7234,7 @@ void FHlslNiagaraTranslator::RegisterFunctionCall(ENiagaraScriptUsage ScriptUsag
 
 void FHlslNiagaraTranslator::GenerateFunctionCall(ENiagaraScriptUsage ScriptUsage, FNiagaraFunctionSignature& FunctionSignature, TArrayView<const int32> Inputs, TArray<int32>& Outputs)
 {
-	SCOPE_CYCLE_COUNTER(STAT_NiagaraEditor_Module_NiagaraHLSLTranslator_GenerateFunctionCall);
+	NIAGARA_SCOPE_CYCLE_COUNTER(STAT_NiagaraEditor_Module_NiagaraHLSLTranslator_GenerateFunctionCall);
 
 	bool bEnteredStatScope = false;   
 	if (ScriptUsage == ENiagaraScriptUsage::Module)
@@ -7735,7 +7737,7 @@ void FHlslNiagaraTranslator::Convert(class UNiagaraNodeConvert* Convert, TArrayV
 
 void FHlslNiagaraTranslator::If(UNiagaraNodeIf* IfNode, TArray<FNiagaraVariable>& Vars, int32 Condition, TArray<int32>& PathA, TArray<int32>& PathB, TArray<int32>& Outputs)
 {
-	SCOPE_CYCLE_COUNTER(STAT_NiagaraEditor_HlslTranslator_If);
+	NIAGARA_SCOPE_CYCLE_COUNTER(STAT_NiagaraEditor_HlslTranslator_If);
 
 	int32 NumVars = Vars.Num();
 	check(PathA.Num() == NumVars);
@@ -7779,7 +7781,7 @@ void FHlslNiagaraTranslator::If(UNiagaraNodeIf* IfNode, TArray<FNiagaraVariable>
 
 int32 FHlslNiagaraTranslator::CompilePin(const UEdGraphPin* Pin)
 {
-	SCOPE_CYCLE_COUNTER(STAT_NiagaraEditor_HlslTranslator_CompilePin);
+	NIAGARA_SCOPE_CYCLE_COUNTER(STAT_NiagaraEditor_HlslTranslator_CompilePin);
 
 	check(Pin);
 	int32 Ret = INDEX_NONE;
@@ -7821,7 +7823,7 @@ int32 FHlslNiagaraTranslator::CompilePin(const UEdGraphPin* Pin)
 
 int32 FHlslNiagaraTranslator::CompileOutputPin(const UEdGraphPin* InPin)
 {
-	SCOPE_CYCLE_COUNTER(STAT_NiagaraEditor_HlslTranslator_CompileOutputPin);
+	NIAGARA_SCOPE_CYCLE_COUNTER(STAT_NiagaraEditor_HlslTranslator_CompileOutputPin);
 
 	if (InPin == nullptr)
 	{
@@ -8371,5 +8373,5 @@ TArray<FName> FHlslNiagaraTranslator::ConditionPropertyPath(const FNiagaraTypeDe
 
 
 
-
+#undef NIAGARA_SCOPE_CYCLE_COUNTER
 #undef LOCTEXT_NAMESPACE
diff --git a/Engine/Plugins/FX/Niagara/Source/NiagaraEditor/Private/NiagaraNodeEmitter.cpp b/Engine/Plugins/FX/Niagara/Source/NiagaraEditor/Private/NiagaraNodeEmitter.cpp
--- a/Engine/Plugins/FX/Niagara/Source/NiagaraEditor/Private/NiagaraNodeEmitter.cpp
+++ b/Engine/Plugins/FX/Niagara/Source/NiagaraEditor/Private/NiagaraNodeEmitter.cpp
@@ -22,7 +22,7 @@
 #define LOCTEXT_NAMESPACE "NiagaraNodeEmitter"
 
 DECLARE_CYCLE_STAT(TEXT("Niagara - Module - NiagaraNodeEmitter_Compile"), STAT_NiagaraEditor_Module_NiagaraNodeEmitter_Compile, STATGROUP_NiagaraEditor);
-
+#define NIAGARA_SCOPE_CYCLE_COUNTER(x) //SCOPE_CYCLE_COUNTER(x)
 
 void UNiagaraNodeEmitter::PostInitProperties()
 {
@@ -415,7 +415,7 @@ void UNiagaraNodeEmitter::BuildParameterMapHistory(FNiagaraParameterMapHistoryBu
 
 void UNiagaraNodeEmitter::Compile(FHlslNiagaraTranslator *Translator, TArray<int32>& Outputs)
 {
-	SCOPE_CYCLE_COUNTER(STAT_NiagaraEditor_Module_NiagaraNodeEmitter_Compile);
+	NIAGARA_SCOPE_CYCLE_COUNTER(STAT_NiagaraEditor_Module_NiagaraNodeEmitter_Compile);
 	FPinCollectorArray InputPins;
 	GetInputPins(InputPins);
 	InputPins.RemoveAll([](UEdGraphPin* InputPin) { return (InputPin->PinType.PinCategory != UEdGraphSchema_Niagara::PinCategoryType) && (InputPin->PinType.PinCategory != UEdGraphSchema_Niagara::PinCategoryEnum); });
@@ -499,4 +499,5 @@ void UNiagaraNodeEmitter::GatherExternalDependencyData(ENiagaraScriptUsage InMas
 	}
 }
 
+#undef NIAGARA_SCOPE_CYCLE_COUNTER
 #undef LOCTEXT_NAMESPACE // NiagaraNodeEmitter
