Build fix - Use new buffer type

diff --git a/Engine/Plugins/FX/Niagara/Source/Niagara/Private/NiagaraDataInterfaceSkeletalMeshUvMapping.h b/Engine/Plugins/FX/Niagara/Source/Niagara/Private/NiagaraDataInterfaceSkeletalMeshUvMapping.h
--- a/Engine/Plugins/FX/Niagara/Source/Niagara/Private/NiagaraDataInterfaceSkeletalMeshUvMapping.h
+++ b/Engine/Plugins/FX/Niagara/Source/Niagara/Private/NiagaraDataInterfaceSkeletalMeshUvMapping.h
@@ -23,7 +23,7 @@ public:
 private:
 	TResourceArray<uint8> FrozenQuadTree;
 
-	FVertexBufferRHIRef UvMappingBuffer;
+	FBufferRHIRef UvMappingBuffer;
 	FShaderResourceViewRHIRef UvMappingSrv;
 
 #if STATS
