Fix for crash when freeing some pooled Niagara components after the instance batcher has already been freed and can contain garbage.
Need to prevent currently in use systems re-entering the pool once it's been flushed ready for world destruction.

diff --git a/Engine/Plugins/FX/Niagara/Source/Niagara/Private/NiagaraComponentPool.cpp b/Engine/Plugins/FX/Niagara/Source/Niagara/Private/NiagaraComponentPool.cpp
--- a/Engine/Plugins/FX/Niagara/Source/Niagara/Private/NiagaraComponentPool.cpp
+++ b/Engine/Plugins/FX/Niagara/Source/Niagara/Private/NiagaraComponentPool.cpp
@@ -153,8 +153,8 @@ void FNCPool::Reclaim(UNiagaraComponent* Component, const float CurrentTimeSecon
 	}
 #endif
 
-	//Don't add back to the pool if we're no longer pooling or we've hit our max resident pool size.
-	if (GbEnableNiagaraSystemPooling != 0 && FreeElements.Num() < (int32)Component->GetAsset()->MaxPoolSize)
+	//Don't add back to the pool if we're no longer pooling or we've hit our max resident pool size. Also if the component's world is in the process of tearing down.
+	if (GbEnableNiagaraSystemPooling != 0 && FreeElements.Num() < (int32)Component->GetAsset()->MaxPoolSize && Component->GetWorld()->bIsTearingDown == false)
 	{
 		Component->DeactivateImmediate();
 
