Automated wrapped object pointer upgrade for Engine + ShooterGame + ShooterGame referenced plugins

diff --git a/Engine/Plugins/2D/Paper2D/Source/Paper2D/Classes/PaperCharacter.h b/Engine/Plugins/2D/Paper2D/Source/Paper2D/Classes/PaperCharacter.h
--- a/Engine/Plugins/2D/Paper2D/Source/Paper2D/Classes/PaperCharacter.h
+++ b/Engine/Plugins/2D/Paper2D/Source/Paper2D/Classes/PaperCharacter.h
@@ -20,7 +20,7 @@ class PAPER2D_API APaperCharacter : public ACharacter
 private:
 	/** The main skeletal mesh associated with this Character (optional sub-object). */
 	UPROPERTY(Category = Character, VisibleAnywhere, BlueprintReadOnly, meta = (AllowPrivateAccess = "true"))
-	class UPaperFlipbookComponent* Sprite;
+	TObjectPtr<class UPaperFlipbookComponent> Sprite;
 public:
 
 	virtual void PostInitializeComponents() override;
diff --git a/Engine/Plugins/2D/Paper2D/Source/Paper2D/Classes/PaperFlipbook.h b/Engine/Plugins/2D/Paper2D/Source/Paper2D/Classes/PaperFlipbook.h
--- a/Engine/Plugins/2D/Paper2D/Source/Paper2D/Classes/PaperFlipbook.h
+++ b/Engine/Plugins/2D/Paper2D/Source/Paper2D/Classes/PaperFlipbook.h
@@ -19,7 +19,7 @@ public:
 	GENERATED_USTRUCT_BODY()
 
 	UPROPERTY(Category=Sprite, EditAnywhere)
-	UPaperSprite* Sprite;
+	TObjectPtr<UPaperSprite> Sprite;
 
 	UPROPERTY(Category=Sprite, EditAnywhere, meta=(ClampMin=1))
 	int32 FrameRun;
@@ -66,7 +66,7 @@ protected:
 
 	// The material to use on a flipbook player instance if not overridden
 	UPROPERTY(Category=Sprite, EditAnywhere, BlueprintReadOnly)
-	UMaterialInterface* DefaultMaterial;
+	TObjectPtr<UMaterialInterface> DefaultMaterial;
 
 	// Collision source
 	UPROPERTY(Category=Sprite, EditAnywhere, BlueprintReadOnly)
diff --git a/Engine/Plugins/2D/Paper2D/Source/Paper2D/Classes/PaperFlipbookActor.h b/Engine/Plugins/2D/Paper2D/Source/Paper2D/Classes/PaperFlipbookActor.h
--- a/Engine/Plugins/2D/Paper2D/Source/Paper2D/Classes/PaperFlipbookActor.h
+++ b/Engine/Plugins/2D/Paper2D/Source/Paper2D/Classes/PaperFlipbookActor.h
@@ -20,7 +20,7 @@ class PAPER2D_API APaperFlipbookActor : public AActor
 
 private:
 	UPROPERTY(Category=Sprite, VisibleAnywhere, BlueprintReadOnly, meta=(ExposeFunctionCategories="Sprite,Rendering,Physics,Components|Flipbook", AllowPrivateAccess="true"))
-	class UPaperFlipbookComponent* RenderComponent;
+	TObjectPtr<class UPaperFlipbookComponent> RenderComponent;
 
 public:
 
diff --git a/Engine/Plugins/2D/Paper2D/Source/Paper2D/Classes/PaperFlipbookComponent.h b/Engine/Plugins/2D/Paper2D/Source/Paper2D/Classes/PaperFlipbookComponent.h
--- a/Engine/Plugins/2D/Paper2D/Source/Paper2D/Classes/PaperFlipbookComponent.h
+++ b/Engine/Plugins/2D/Paper2D/Source/Paper2D/Classes/PaperFlipbookComponent.h
@@ -26,11 +26,11 @@ class PAPER2D_API UPaperFlipbookComponent : public UMeshComponent
 protected:
 	/** Flipbook currently being played */
 	UPROPERTY(Category=Sprite, EditAnywhere, meta=(DisplayThumbnail = "true"), ReplicatedUsing=OnRep_SourceFlipbook)
-	UPaperFlipbook* SourceFlipbook;
+	TObjectPtr<UPaperFlipbook> SourceFlipbook;
 
 	// DEPRECATED in 4.5: The material override for this flipbook component (if any); replaced by the Materials array inherited from UMeshComponent
 	UPROPERTY()
-	UMaterialInterface* Material_DEPRECATED;
+	TObjectPtr<UMaterialInterface> Material_DEPRECATED;
 
 	/** Current play rate of the flipbook */
 	UPROPERTY(Category=Sprite, EditAnywhere)
@@ -62,7 +62,7 @@ protected:
 
 	/** The cached body setup */
 	UPROPERTY(Transient)
-	UBodySetup* CachedBodySetup;
+	TObjectPtr<UBodySetup> CachedBodySetup;
 
 public:
 	/** Event called whenever a non-looping flipbook finishes playing (either reaching the beginning or the end, depending on the play direction) */
diff --git a/Engine/Plugins/2D/Paper2D/Source/Paper2D/Classes/PaperGroupedSpriteActor.h b/Engine/Plugins/2D/Paper2D/Source/Paper2D/Classes/PaperGroupedSpriteActor.h
--- a/Engine/Plugins/2D/Paper2D/Source/Paper2D/Classes/PaperGroupedSpriteActor.h
+++ b/Engine/Plugins/2D/Paper2D/Source/Paper2D/Classes/PaperGroupedSpriteActor.h
@@ -22,7 +22,7 @@ class PAPER2D_API APaperGroupedSpriteActor : public AActor
 
 private:
 	UPROPERTY(Category=Sprite, VisibleAnywhere, BlueprintReadOnly, meta=(ExposeFunctionCategories = "Sprite,Rendering,Physics,Components|Sprite", AllowPrivateAccess="true"))
-	UPaperGroupedSpriteComponent* RenderComponent;
+	TObjectPtr<UPaperGroupedSpriteComponent> RenderComponent;
 
 public:
 	// AActor interface
diff --git a/Engine/Plugins/2D/Paper2D/Source/Paper2D/Classes/PaperGroupedSpriteComponent.h b/Engine/Plugins/2D/Paper2D/Source/Paper2D/Classes/PaperGroupedSpriteComponent.h
--- a/Engine/Plugins/2D/Paper2D/Source/Paper2D/Classes/PaperGroupedSpriteComponent.h
+++ b/Engine/Plugins/2D/Paper2D/Source/Paper2D/Classes/PaperGroupedSpriteComponent.h
@@ -26,7 +26,7 @@ public:
 	FMatrix Transform;
 
 	UPROPERTY(EditAnywhere, Category=Instances)
-	UPaperSprite* SourceSprite;
+	TObjectPtr<UPaperSprite> SourceSprite;
 
 	UPROPERTY(EditAnywhere, Category=Instances)
 	FColor VertexColor;
@@ -63,7 +63,7 @@ class PAPER2D_API UPaperGroupedSpriteComponent : public UMeshComponent
 protected:
 	/** Array of materials used by the instances */
 	UPROPERTY()
-	TArray<UMaterialInterface*> InstanceMaterials;
+	TArray<TObjectPtr<UMaterialInterface>> InstanceMaterials;
 
 	/** Array of instances */
 	UPROPERTY(EditAnywhere, DisplayName="Instances", Category=Instances, meta=(MakeEditWidget=true))
diff --git a/Engine/Plugins/2D/Paper2D/Source/Paper2D/Classes/PaperSprite.h b/Engine/Plugins/2D/Paper2D/Source/Paper2D/Classes/PaperSprite.h
--- a/Engine/Plugins/2D/Paper2D/Source/Paper2D/Classes/PaperSprite.h
+++ b/Engine/Plugins/2D/Paper2D/Source/Paper2D/Classes/PaperSprite.h
@@ -84,12 +84,12 @@ protected:
 	TSoftObjectPtr<UTexture2D> SourceTexture;
 
 	UPROPERTY(Transient)
-	mutable UTexture2D* SourceTextureCacheNeverSerialized;
+	mutable TObjectPtr<UTexture2D> SourceTextureCacheNeverSerialized;
 #endif
 
 	// Additional source textures for other slots
 	UPROPERTY(Category=Sprite, EditAnywhere, AssetRegistrySearchable, meta=(DisplayName="Additional Textures"))
-	TArray<UTexture*> AdditionalSourceTextures;
+	TArray<TObjectPtr<UTexture>> AdditionalSourceTextures;
 
 	// Position within BakedSourceTexture (in pixels)
 	UPROPERTY()
@@ -100,15 +100,15 @@ protected:
 	FVector2D BakedSourceDimension;
 
 	UPROPERTY()
-	UTexture2D* BakedSourceTexture;
+	TObjectPtr<UTexture2D> BakedSourceTexture;
 
 	// The material to use on a sprite instance if not overridden (this is the default material when only one is being used, and is the translucent/masked material for Diced render geometry, slot 0)
 	UPROPERTY(Category=Sprite, EditAnywhere, BlueprintReadOnly)
-	UMaterialInterface* DefaultMaterial;
+	TObjectPtr<UMaterialInterface> DefaultMaterial;
 
 	// The alternate material to use on a sprite instance if not overridden (this is only used for Diced render geometry, and will be the opaque material in that case, slot 1)
 	UPROPERTY(Category=Rendering, EditAnywhere, BlueprintReadOnly)
-	UMaterialInterface* AlternateMaterial;
+	TObjectPtr<UMaterialInterface> AlternateMaterial;
 
 	// List of sockets on this sprite
 	UPROPERTY(Category=Sprite, EditAnywhere)
@@ -125,7 +125,7 @@ protected:
 public:
 	// Baked physics data.
 	UPROPERTY(EditAnywhere, Category=NeverShown)
-	class UBodySetup* BodySetup;
+	TObjectPtr<class UBodySetup> BodySetup;
 
 #if WITH_EDITORONLY_DATA
 
@@ -156,7 +156,7 @@ protected:
 
 	// Spritesheet group that this sprite belongs to
 	UPROPERTY(Category=Sprite, EditAnywhere, AssetRegistrySearchable)
-	class UPaperSpriteAtlas* AtlasGroup;
+	TObjectPtr<class UPaperSpriteAtlas> AtlasGroup;
 
 	// The previous spritesheet group this belonged to
 	// To make sure we remove ourselves from it if changed or nulled out
diff --git a/Engine/Plugins/2D/Paper2D/Source/Paper2D/Classes/PaperSpriteActor.h b/Engine/Plugins/2D/Paper2D/Source/Paper2D/Classes/PaperSpriteActor.h
--- a/Engine/Plugins/2D/Paper2D/Source/Paper2D/Classes/PaperSpriteActor.h
+++ b/Engine/Plugins/2D/Paper2D/Source/Paper2D/Classes/PaperSpriteActor.h
@@ -20,7 +20,7 @@ class PAPER2D_API APaperSpriteActor : public AActor
 
 private:
 	UPROPERTY(Category=Sprite, VisibleAnywhere, BlueprintReadOnly, meta=(ExposeFunctionCategories="Sprite,Rendering,Physics,Components|Sprite", AllowPrivateAccess="true"))
-	class UPaperSpriteComponent* RenderComponent;
+	TObjectPtr<class UPaperSpriteComponent> RenderComponent;
 public:
 
 	// AActor interface
diff --git a/Engine/Plugins/2D/Paper2D/Source/Paper2D/Classes/PaperSpriteAtlas.h b/Engine/Plugins/2D/Paper2D/Source/Paper2D/Classes/PaperSpriteAtlas.h
--- a/Engine/Plugins/2D/Paper2D/Source/Paper2D/Classes/PaperSpriteAtlas.h
+++ b/Engine/Plugins/2D/Paper2D/Source/Paper2D/Classes/PaperSpriteAtlas.h
@@ -89,7 +89,7 @@ class UPaperSpriteAtlas : public UObject
 
 	// List of generated atlas textures
 	UPROPERTY(VisibleAnywhere, Category=AtlasTexture)
-	TArray<UTexture*> GeneratedTextures;
+	TArray<TObjectPtr<UTexture>> GeneratedTextures;
 
 	// The GUID of the atlas group, used to match up sprites that belong to this group even thru atlas renames
 	UPROPERTY(VisibleAnywhere, Category=General)
diff --git a/Engine/Plugins/2D/Paper2D/Source/Paper2D/Classes/PaperSpriteComponent.h b/Engine/Plugins/2D/Paper2D/Source/Paper2D/Classes/PaperSpriteComponent.h
--- a/Engine/Plugins/2D/Paper2D/Source/Paper2D/Classes/PaperSpriteComponent.h
+++ b/Engine/Plugins/2D/Paper2D/Source/Paper2D/Classes/PaperSpriteComponent.h
@@ -31,11 +31,11 @@ class PAPER2D_API UPaperSpriteComponent : public UMeshComponent
 protected:
 	// The sprite asset used by this component
 	UPROPERTY(Category=Sprite, EditAnywhere, BlueprintReadOnly, meta=(DisplayThumbnail = "true"))
-	UPaperSprite* SourceSprite;
+	TObjectPtr<UPaperSprite> SourceSprite;
 
 	// DEPRECATED in 4.4: The material override for this sprite component (if any); replaced by the Materials array inherited from UMeshComponent
 	UPROPERTY()
-	UMaterialInterface* MaterialOverride_DEPRECATED;
+	TObjectPtr<UMaterialInterface> MaterialOverride_DEPRECATED;
 
 	// The color of the sprite (passed to the sprite material as a vertex color)
 	UPROPERTY(BlueprintReadOnly, Interp, Category=Sprite)
diff --git a/Engine/Plugins/2D/Paper2D/Source/Paper2D/Classes/PaperTerrainActor.h b/Engine/Plugins/2D/Paper2D/Source/Paper2D/Classes/PaperTerrainActor.h
--- a/Engine/Plugins/2D/Paper2D/Source/Paper2D/Classes/PaperTerrainActor.h
+++ b/Engine/Plugins/2D/Paper2D/Source/Paper2D/Classes/PaperTerrainActor.h
@@ -19,13 +19,13 @@ class PAPER2D_API APaperTerrainActor : public AActor
 
 private:
 	UPROPERTY()
-	class USceneComponent* DummyRoot;
+	TObjectPtr<class USceneComponent> DummyRoot;
 
 	UPROPERTY()
-	class UPaperTerrainSplineComponent* SplineComponent;
+	TObjectPtr<class UPaperTerrainSplineComponent> SplineComponent;
 
 	UPROPERTY(Category=Sprite, VisibleAnywhere, BlueprintReadOnly, meta=(ExposeFunctionCategories="Sprite,Rendering,Physics,Components|Sprite", AllowPrivateAccess="true"))
-	class UPaperTerrainComponent* RenderComponent;
+	TObjectPtr<class UPaperTerrainComponent> RenderComponent;
 public:
 
 	// AActor interface
diff --git a/Engine/Plugins/2D/Paper2D/Source/Paper2D/Classes/PaperTerrainComponent.h b/Engine/Plugins/2D/Paper2D/Source/Paper2D/Classes/PaperTerrainComponent.h
--- a/Engine/Plugins/2D/Paper2D/Source/Paper2D/Classes/PaperTerrainComponent.h
+++ b/Engine/Plugins/2D/Paper2D/Source/Paper2D/Classes/PaperTerrainComponent.h
@@ -55,7 +55,7 @@ class PAPER2D_API UPaperTerrainComponent : public UPrimitiveComponent
 public:
 	/** The terrain material to apply to this component (set of rules for which sprites are used on different surfaces or the interior) */
 	UPROPERTY(Category=Sprite, EditAnywhere, BlueprintReadOnly)
-	class UPaperTerrainMaterial* TerrainMaterial;
+	TObjectPtr<class UPaperTerrainMaterial> TerrainMaterial;
 
 	UPROPERTY(Category = Sprite, EditAnywhere, BlueprintReadOnly)
 	bool bClosedSpline;
@@ -64,7 +64,7 @@ public:
 	bool bFilledSpline;
 
 	UPROPERTY()
-	class UPaperTerrainSplineComponent* AssociatedSpline;
+	TObjectPtr<class UPaperTerrainSplineComponent> AssociatedSpline;
 
 	/** Random seed used for choosing which spline meshes to use. */
 	UPROPERTY(Category=Sprite, EditAnywhere)
@@ -128,7 +128,7 @@ protected:
 
 	/** Description of collision */
 	UPROPERTY(Transient, DuplicateTransient)
-	class UBodySetup* CachedBodySetup;
+	TObjectPtr<class UBodySetup> CachedBodySetup;
 
 	TArray<FPaperTerrainSpriteGeometry> GeneratedSpriteGeometry;
 	
diff --git a/Engine/Plugins/2D/Paper2D/Source/Paper2D/Classes/PaperTerrainMaterial.h b/Engine/Plugins/2D/Paper2D/Source/Paper2D/Classes/PaperTerrainMaterial.h
--- a/Engine/Plugins/2D/Paper2D/Source/Paper2D/Classes/PaperTerrainMaterial.h
+++ b/Engine/Plugins/2D/Paper2D/Source/Paper2D/Classes/PaperTerrainMaterial.h
@@ -17,15 +17,15 @@ struct FPaperTerrainMaterialRule
 
 	// The sprite to use at the 'left' (closest to spline start) edge of the terrain segment
 	UPROPERTY(Category=Sprite, EditAnywhere, meta=(DisplayThumbnail="true"))
-	class UPaperSprite* StartCap;
+	TObjectPtr<class UPaperSprite> StartCap;
 
 	// A set of sprites to randomly choose to fill up the interior space between the caps in a terrain segment
 	UPROPERTY(Category=Sprite, EditAnywhere, meta=(DisplayThumbnail="true"))
-	TArray<class UPaperSprite*> Body;
+	TArray<TObjectPtr<class UPaperSprite>> Body;
 
 	// The sprite to use at the 'right' (closest to spline end) edge of the terrain segment
 	UPROPERTY(Category=Sprite, EditAnywhere, meta=(DisplayThumbnail="true"))
-	class UPaperSprite* EndCap;
+	TObjectPtr<class UPaperSprite> EndCap;
 
 	// Minimum slope angle (in degrees) to apply this rule
 	UPROPERTY(Category=Sprite, EditAnywhere)
@@ -81,7 +81,7 @@ class PAPER2D_API UPaperTerrainMaterial : public UDataAsset
 
 	// The sprite to use for an interior region fill
 	UPROPERTY(Category = Sprite, EditAnywhere, meta=(DisplayThumbnail="true"))
-	class UPaperSprite* InteriorFill;
+	TObjectPtr<class UPaperSprite> InteriorFill;
 
 #if WITH_EDITOR
 	// UObject interface
diff --git a/Engine/Plugins/2D/Paper2D/Source/Paper2D/Classes/PaperTileLayer.h b/Engine/Plugins/2D/Paper2D/Source/Paper2D/Classes/PaperTileLayer.h
--- a/Engine/Plugins/2D/Paper2D/Source/Paper2D/Classes/PaperTileLayer.h
+++ b/Engine/Plugins/2D/Paper2D/Source/Paper2D/Classes/PaperTileLayer.h
@@ -30,7 +30,7 @@ struct FPaperTileInfo
 
 	// The tile set that this tile comes from
 	UPROPERTY(EditAnywhere, Category=Sprite)
-	UPaperTileSet* TileSet;
+	TObjectPtr<UPaperTileSet> TileSet;
 
 	// This is the index of the current tile within the tile set
 	UPROPERTY(EditAnywhere, Category=Sprite)
@@ -171,7 +171,7 @@ private:
 	TArray<FPaperTileInfo> AllocatedCells;
 
 	UPROPERTY()
-	UPaperTileSet* TileSet_DEPRECATED;
+	TObjectPtr<UPaperTileSet> TileSet_DEPRECATED;
 
 	UPROPERTY()
 	TArray<int32> AllocatedGrid_DEPRECATED;
diff --git a/Engine/Plugins/2D/Paper2D/Source/Paper2D/Classes/PaperTileMap.h b/Engine/Plugins/2D/Paper2D/Source/Paper2D/Classes/PaperTileMap.h
--- a/Engine/Plugins/2D/Paper2D/Source/Paper2D/Classes/PaperTileMap.h
+++ b/Engine/Plugins/2D/Paper2D/Source/Paper2D/Classes/PaperTileMap.h
@@ -77,11 +77,11 @@ class PAPER2D_API UPaperTileMap : public UObject
 
 	// The material to use on a tile map instance if not overridden
 	UPROPERTY(Category=Setup, EditAnywhere, BlueprintReadOnly)
-	UMaterialInterface* Material;
+	TObjectPtr<UMaterialInterface> Material;
 
 	// The list of layers
 	UPROPERTY(Instanced, Category=Sprite, BlueprintReadOnly, EditAnywhere)
-	TArray<UPaperTileLayer*> TileLayers;
+	TArray<TObjectPtr<UPaperTileLayer>> TileLayers;
 
 protected:
 	// The extrusion thickness of collision geometry when using a 3D collision domain
@@ -104,13 +104,13 @@ public:
 
 	// Baked physics data.
 	UPROPERTY()
-	UBodySetup* BodySetup;
+	TObjectPtr<UBodySetup> BodySetup;
 
 public:
 #if WITH_EDITORONLY_DATA
 	/** Importing data and options used for this tile map */
 	UPROPERTY(Category=ImportSettings, VisibleAnywhere, Instanced)
-	class UAssetImportData* AssetImportData;
+	TObjectPtr<class UAssetImportData> AssetImportData;
 
 	/** The currently selected layer index */
 	UPROPERTY()
diff --git a/Engine/Plugins/2D/Paper2D/Source/Paper2D/Classes/PaperTileMapActor.h b/Engine/Plugins/2D/Paper2D/Source/Paper2D/Classes/PaperTileMapActor.h
--- a/Engine/Plugins/2D/Paper2D/Source/Paper2D/Classes/PaperTileMapActor.h
+++ b/Engine/Plugins/2D/Paper2D/Source/Paper2D/Classes/PaperTileMapActor.h
@@ -21,7 +21,7 @@ class PAPER2D_API APaperTileMapActor : public AActor
 
 private:
 	UPROPERTY(Category=TileMapActor, VisibleAnywhere, BlueprintReadOnly, meta=(ExposeFunctionCategories="Sprite,Rendering,Physics,Components|Sprite", AllowPrivateAccess="true"))
-	class UPaperTileMapComponent* RenderComponent;
+	TObjectPtr<class UPaperTileMapComponent> RenderComponent;
 public:
 
 	// AActor interface
diff --git a/Engine/Plugins/2D/Paper2D/Source/Paper2D/Classes/PaperTileMapComponent.h b/Engine/Plugins/2D/Paper2D/Source/Paper2D/Classes/PaperTileMapComponent.h
--- a/Engine/Plugins/2D/Paper2D/Source/Paper2D/Classes/PaperTileMapComponent.h
+++ b/Engine/Plugins/2D/Paper2D/Source/Paper2D/Classes/PaperTileMapComponent.h
@@ -51,13 +51,13 @@ private:
 	int32 TileHeight_DEPRECATED;
 
 	UPROPERTY()
-	class UPaperTileSet* DefaultLayerTileSet_DEPRECATED;
+	TObjectPtr<class UPaperTileSet> DefaultLayerTileSet_DEPRECATED;
 
 	UPROPERTY()
-	UMaterialInterface* Material_DEPRECATED;
+	TObjectPtr<UMaterialInterface> Material_DEPRECATED;
 
 	UPROPERTY()
-	TArray<UPaperTileLayer*> TileLayers_DEPRECATED;
+	TArray<TObjectPtr<UPaperTileLayer>> TileLayers_DEPRECATED;
 
 	// The color of the tile map (multiplied with the per-layer color and passed to the material as a vertex color)
 	UPROPERTY(EditAnywhere, Category=Materials)
@@ -82,7 +82,7 @@ private:
 public:
 	// The tile map used by this component
 	UPROPERTY(Category=Setup, EditAnywhere, BlueprintReadOnly)
-	class UPaperTileMap* TileMap;
+	TObjectPtr<class UPaperTileMap> TileMap;
 
 #if WITH_EDITORONLY_DATA
 	// Should this component show a tile grid when the component is selected?
diff --git a/Engine/Plugins/2D/Paper2D/Source/Paper2D/Classes/PaperTileSet.h b/Engine/Plugins/2D/Paper2D/Source/Paper2D/Classes/PaperTileSet.h
--- a/Engine/Plugins/2D/Paper2D/Source/Paper2D/Classes/PaperTileSet.h
+++ b/Engine/Plugins/2D/Paper2D/Source/Paper2D/Classes/PaperTileSet.h
@@ -86,11 +86,11 @@ private:
 
 	// The tile sheet texture associated with this tile set
 	UPROPERTY(Category=TileSet, BlueprintReadOnly, EditAnywhere, meta=(DisplayName="Tile Sheet Texture", AllowPrivateAccess="true"))
-	UTexture2D* TileSheet;
+	TObjectPtr<UTexture2D> TileSheet;
 
 	// Additional source textures for other slots
 	UPROPERTY(Category = TileSet, EditAnywhere, AssetRegistrySearchable, meta = (DisplayName = "Additional Textures"))
-	TArray<UTexture*> AdditionalSourceTextures;
+	TArray<TObjectPtr<UTexture>> AdditionalSourceTextures;
 
 	// The amount of padding around the border of the tile sheet (in pixels)
 	UPROPERTY(Category=TileSet, BlueprintReadOnly, EditAnywhere, meta=(UIMin=0, ClampMin=0, AllowPrivateAccess="true"))
diff --git a/Engine/Plugins/2D/Paper2D/Source/Paper2D/Classes/SpriteDrawCall.h b/Engine/Plugins/2D/Paper2D/Source/Paper2D/Classes/SpriteDrawCall.h
--- a/Engine/Plugins/2D/Paper2D/Source/Paper2D/Classes/SpriteDrawCall.h
+++ b/Engine/Plugins/2D/Paper2D/Source/Paper2D/Classes/SpriteDrawCall.h
@@ -18,7 +18,7 @@ public:
 	FVector Destination;
 
 	UPROPERTY()
-	UTexture* BaseTexture;
+	TObjectPtr<UTexture> BaseTexture;
 
 	FAdditionalSpriteTextureArray AdditionalTextures;
 
diff --git a/Engine/Plugins/2D/Paper2D/Source/Paper2DEditor/Classes/PaperTileMapPromotionFactory.h b/Engine/Plugins/2D/Paper2D/Source/Paper2DEditor/Classes/PaperTileMapPromotionFactory.h
--- a/Engine/Plugins/2D/Paper2D/Source/Paper2DEditor/Classes/PaperTileMapPromotionFactory.h
+++ b/Engine/Plugins/2D/Paper2D/Source/Paper2DEditor/Classes/PaperTileMapPromotionFactory.h
@@ -18,7 +18,7 @@ class UPaperTileMapPromotionFactory : public UFactory
 
 	// Object being promoted to an asset
 	UPROPERTY()
-	class UPaperTileMap* AssetToRename;
+	TObjectPtr<class UPaperTileMap> AssetToRename;
 
 	// UFactory interface
 	virtual UObject* FactoryCreateNew(UClass* Class, UObject* InParent, FName Name, EObjectFlags Flags, UObject* Context, FFeedbackContext* Warn) override;
diff --git a/Engine/Plugins/2D/Paper2D/Source/Paper2DEditor/Private/TileSetEditor/TileSheetPaddingFactory.h b/Engine/Plugins/2D/Paper2D/Source/Paper2DEditor/Private/TileSetEditor/TileSheetPaddingFactory.h
--- a/Engine/Plugins/2D/Paper2D/Source/Paper2DEditor/Private/TileSetEditor/TileSheetPaddingFactory.h
+++ b/Engine/Plugins/2D/Paper2D/Source/Paper2DEditor/Private/TileSetEditor/TileSheetPaddingFactory.h
@@ -21,7 +21,7 @@ class UTileSheetPaddingFactory : public UFactory
 
 	// Source tile sheet texture
 	UPROPERTY()
-	UPaperTileSet* SourceTileSet;
+	TObjectPtr<UPaperTileSet> SourceTileSet;
 
 	// The amount to extrude out from each tile (in pixels)
 	UPROPERTY(Category = Settings, EditAnywhere, meta=(UIMin=1, ClampMin=1))
diff --git a/Engine/Plugins/2D/Paper2D/Source/PaperSpriteSheetImporter/Private/PaperSpriteSheet.h b/Engine/Plugins/2D/Paper2D/Source/PaperSpriteSheetImporter/Private/PaperSpriteSheet.h
--- a/Engine/Plugins/2D/Paper2D/Source/PaperSpriteSheetImporter/Private/PaperSpriteSheet.h
+++ b/Engine/Plugins/2D/Paper2D/Source/PaperSpriteSheetImporter/Private/PaperSpriteSheet.h
@@ -27,7 +27,7 @@ public:
 
 	// The asset that was created for TextureName
 	UPROPERTY(VisibleAnywhere, Category=Data)
-	UTexture2D* Texture;
+	TObjectPtr<UTexture2D> Texture;
 
 	// The name of the normal map texture during import (if any)
 	UPROPERTY(VisibleAnywhere, Category=Data)
@@ -35,12 +35,12 @@ public:
 
 	// The asset that was created for NormalMapTextureName (if any)
 	UPROPERTY(VisibleAnywhere, Category=Data)
-	UTexture2D* NormalMapTexture;
+	TObjectPtr<UTexture2D> NormalMapTexture;
 
 #if WITH_EDITORONLY_DATA
 	// Import data for this 
 	UPROPERTY(VisibleAnywhere, Instanced, Category=ImportSettings)
-	class UAssetImportData* AssetImportData;
+	TObjectPtr<class UAssetImportData> AssetImportData;
 
 	// UObject interface
 	virtual void PostInitProperties() override;
diff --git a/Engine/Plugins/Cameras/GameplayCameras/Source/GameplayCameras/Public/CompositeCameraShakePattern.h b/Engine/Plugins/Cameras/GameplayCameras/Source/GameplayCameras/Public/CompositeCameraShakePattern.h
--- a/Engine/Plugins/Cameras/GameplayCameras/Source/GameplayCameras/Public/CompositeCameraShakePattern.h
+++ b/Engine/Plugins/Cameras/GameplayCameras/Source/GameplayCameras/Public/CompositeCameraShakePattern.h
@@ -29,7 +29,7 @@ public:
 
 	/** The list of child shake patterns */
 	UPROPERTY(EditAnywhere, Instanced, Category=CameraShake)
-	TArray<UCameraShakePattern*> ChildPatterns;
+	TArray<TObjectPtr<UCameraShakePattern>> ChildPatterns;
 
 private:
 
diff --git a/Engine/Plugins/Cameras/GameplayCameras/Source/GameplayCameras/Public/MatineeCameraShake.h b/Engine/Plugins/Cameras/GameplayCameras/Source/GameplayCameras/Public/MatineeCameraShake.h
--- a/Engine/Plugins/Cameras/GameplayCameras/Source/GameplayCameras/Public/MatineeCameraShake.h
+++ b/Engine/Plugins/Cameras/GameplayCameras/Source/GameplayCameras/Public/MatineeCameraShake.h
@@ -175,11 +175,11 @@ public:
 
 	/** Source camera animation to play. Can be null. */
 	UPROPERTY(EditAnywhere, Category = AnimShake)
-	class UCameraAnim* Anim;
+	TObjectPtr<class UCameraAnim> Anim;
 
 	/** Source camera animation sequence to play. Can be null, but can't have both Anim and AnimSequence. */
 	UPROPERTY(EditAnywhere, Category = AnimShake)
-	class UCameraAnimationSequence* AnimSequence;
+	TObjectPtr<class UCameraAnimationSequence> AnimSequence;
 
 	/**
 	* If true, play a random snippet of the animation of length Duration.  Implies bLoop and bRandomStartTime = true for the CameraAnim.
@@ -196,7 +196,7 @@ public:
 
 	/** The playing instance of the CameraAnim-based shake, if any. */
 	UPROPERTY(transient, BlueprintReadOnly, Category = CameraShake)
-	class UCameraAnimInst* AnimInst;
+	TObjectPtr<class UCameraAnimInst> AnimInst;
 
 public:
 
@@ -282,7 +282,7 @@ protected:
 
 	/** Sequence shake pattern for when using a sequence instead of a camera anim */
 	UPROPERTY(Instanced)
-	class USequenceCameraShakePattern* SequenceShakePattern;
+	TObjectPtr<class USequenceCameraShakePattern> SequenceShakePattern;
 
 	/** State tracking for the sequence shake pattern */
 	FCameraShakeState SequenceShakeState;
diff --git a/Engine/Plugins/Developer/AnimationSharing/Source/AnimationSharing/Public/AnimationSharingInstances.h b/Engine/Plugins/Developer/AnimationSharing/Source/AnimationSharing/Public/AnimationSharingInstances.h
--- a/Engine/Plugins/Developer/AnimationSharing/Source/AnimationSharing/Public/AnimationSharingInstances.h
+++ b/Engine/Plugins/Developer/AnimationSharing/Source/AnimationSharing/Public/AnimationSharingInstances.h
@@ -16,7 +16,7 @@ class ANIMATIONSHARING_API UAnimSharingStateInstance : public UAnimInstance
 
 protected:
 	UPROPERTY(EditDefaultsOnly, BlueprintReadOnly, Transient, Category = AnimationSharing)
-	UAnimSequence* AnimationToPlay;
+	TObjectPtr<UAnimSequence> AnimationToPlay;
 
 	UPROPERTY(EditDefaultsOnly, BlueprintReadOnly, Transient, Category = AnimationSharing)
 	float PermutationTimeOffset;
@@ -35,7 +35,7 @@ private:
 	uint8 ComponentIndex;
 	
 	UPROPERTY(Transient)
-	class UAnimSharingInstance* Instance;
+	TObjectPtr<class UAnimSharingInstance> Instance;
 };
 
 UCLASS(transient, Blueprintable)
diff --git a/Engine/Plugins/Developer/AnimationSharing/Source/AnimationSharing/Public/AnimationSharingManager.h b/Engine/Plugins/Developer/AnimationSharing/Source/AnimationSharing/Public/AnimationSharingManager.h
--- a/Engine/Plugins/Developer/AnimationSharing/Source/AnimationSharing/Public/AnimationSharingManager.h
+++ b/Engine/Plugins/Developer/AnimationSharing/Source/AnimationSharing/Public/AnimationSharingManager.h
@@ -385,7 +385,7 @@ public:
 
 	/** Actors currently registered to be animation driven by the AnimManager using this setup */
 	UPROPERTY(VisibleAnywhere, Transient, Category = AnimationSharing)
-	TArray<AActor*> RegisteredActors;
+	TArray<TObjectPtr<AActor>> RegisteredActors;
 	
 	/** Per actor data, matches RegisteredActors*/
 	TArray<FPerActorData> PerActorData;
@@ -400,7 +400,7 @@ public:
 
 	/** (Blueprint)class instance used for determining the state enum value for each registered actor */
 	UPROPERTY(EditAnywhere, Transient, Category = AnimationSharing)
-	UAnimationSharingStateProcessor* StateProcessor;
+	TObjectPtr<UAnimationSharingStateProcessor> StateProcessor;
 	bool bNativeStateProcessor;
 
 	/** Currently running blend instances */
@@ -411,7 +411,7 @@ public:
 	TArray<FAdditiveInstance> AdditiveInstances;
 
 	UPROPERTY(VisibleAnywhere, Transient, Category = AnimationSharing)
-	TArray<UAnimSequence*> UsedAnimationSequences;
+	TArray<TObjectPtr<UAnimSequence>> UsedAnimationSequences;
 
 	/** Significance manager used for retrieve AI actor significance values */
 	USignificanceManager* SignificanceManager;
@@ -421,11 +421,11 @@ public:
 
 	/** Enum class set up by the user to 'describe' the animation states */
 	UPROPERTY(VisibleAnywhere, Transient, Category = AnimationSharing)
-	UEnum* StateEnum;
+	TObjectPtr<UEnum> StateEnum;
 
 	/** Actor to which all the running SkeletalMeshComponents used for the sharing are attached to */
 	UPROPERTY(VisibleAnywhere, Transient, Category = AnimationSharing)
-	AActor* SharingActor;
+	TObjectPtr<AActor> SharingActor;
 
 	/** Platform specific scalability settings */
 	const FAnimationSharingScalability* ScalabilitySettings;
@@ -538,11 +538,11 @@ protected:
 protected:
 	/** Array of unique skeletons, matches PerSkeletonData array entries*/
 	UPROPERTY(Transient)
-	TArray<const USkeleton*> Skeletons;
+	TArray<TObjectPtr<const USkeleton>> Skeletons;
 
 	/** Sharing data required for the unique Skeleton setups */
 	UPROPERTY(VisibleAnywhere, Transient, Category = AnimationSharing)
-	TArray<UAnimSharingInstance*> PerSkeletonData;
+	TArray<TObjectPtr<UAnimSharingInstance>> PerSkeletonData;
 	
 	/** Platform specific scalability settings */
 	FAnimationSharingScalability ScalabilitySettings;
diff --git a/Engine/Plugins/Developer/AnimationSharing/Source/AnimationSharing/Public/AnimationSharingTypes.h b/Engine/Plugins/Developer/AnimationSharing/Source/AnimationSharing/Public/AnimationSharingTypes.h
--- a/Engine/Plugins/Developer/AnimationSharing/Source/AnimationSharing/Public/AnimationSharingTypes.h
+++ b/Engine/Plugins/Developer/AnimationSharing/Source/AnimationSharing/Public/AnimationSharingTypes.h
@@ -22,7 +22,7 @@ public:
 	
 	/** Animation Sequence to play for this particular setup */
 	UPROPERTY(EditAnywhere, Category = AnimationSharing)
-	UAnimSequence* AnimSequence;
+	TObjectPtr<UAnimSequence> AnimSequence;
 
 	/** Animation blueprint to use for playing back the Animation Sequence */
 	UPROPERTY(EditAnywhere, AdvancedDisplay, Category = AnimationSharing)
@@ -132,11 +132,11 @@ public:
 
 	/** Skeleton compatible with the animation sharing setup */
 	UPROPERTY(EditAnywhere, Category = AnimationSharing)
-	USkeleton* Skeleton;
+	TObjectPtr<USkeleton> Skeleton;
 
 	/** Skeletal mesh used to setup skeletal mesh components */
 	UPROPERTY(EditAnywhere, Category = AnimationSharing)
-	USkeletalMesh* SkeletalMesh;
+	TObjectPtr<USkeletalMesh> SkeletalMesh;
 
 	/** Animation blueprint used to perform the blending between states */
 	UPROPERTY(EditAnywhere, Category = AnimationSharing, meta = (DisplayName="Animation Blueprint for Blending"))
diff --git a/Engine/Plugins/Editor/ContentBrowser/ContentBrowserClassDataSource/Source/ContentBrowserClassDataSource/Public/ContentBrowserClassDataSource.h b/Engine/Plugins/Editor/ContentBrowser/ContentBrowserClassDataSource/Source/ContentBrowserClassDataSource/Public/ContentBrowserClassDataSource.h
--- a/Engine/Plugins/Editor/ContentBrowser/ContentBrowserClassDataSource/Source/ContentBrowserClassDataSource/Public/ContentBrowserClassDataSource.h
+++ b/Engine/Plugins/Editor/ContentBrowser/ContentBrowserClassDataSource/Source/ContentBrowserClassDataSource/Public/ContentBrowserClassDataSource.h
@@ -21,7 +21,7 @@ struct CONTENTBROWSERCLASSDATASOURCE_API FContentBrowserCompiledClassDataFilter
 
 public:
 	UPROPERTY()
-	TSet<UClass*> ValidClasses;
+	TSet<TObjectPtr<UClass>> ValidClasses;
 
 	UPROPERTY()
 	TSet<FName> ValidFolders;
diff --git a/Engine/Plugins/Editor/DataValidation/Source/DataValidation/Public/EditorValidatorSubsystem.h b/Engine/Plugins/Editor/DataValidation/Source/DataValidation/Public/EditorValidatorSubsystem.h
--- a/Engine/Plugins/Editor/DataValidation/Source/DataValidation/Public/EditorValidatorSubsystem.h
+++ b/Engine/Plugins/Editor/DataValidation/Source/DataValidation/Public/EditorValidatorSubsystem.h
@@ -127,7 +127,7 @@ protected:
 	TArray<FName> SavedPackagesToValidate;
 
 	UPROPERTY(Transient)
-	TMap<UClass*, UEditorValidatorBase*> Validators;
+	TMap<TObjectPtr<UClass>, TObjectPtr<UEditorValidatorBase>> Validators;
 
 	/** Specifies whether or not to validate assets on save when saving for a cook */
 	UPROPERTY(config)
diff --git a/Engine/Plugins/Editor/GeometryMode/Source/GeometryMode/Classes/GeomModifier.h b/Engine/Plugins/Editor/GeometryMode/Source/GeometryMode/Classes/GeomModifier.h
--- a/Engine/Plugins/Editor/GeometryMode/Source/GeometryMode/Classes/GeomModifier.h
+++ b/Engine/Plugins/Editor/GeometryMode/Source/GeometryMode/Classes/GeomModifier.h
@@ -59,7 +59,7 @@ class UGeomModifier : public UObject
 private:
 	/** Stored state of polys in case the brush state needs to be restroed */
 	UPROPERTY()
-	class UPolys* CachedPolys;
+	TObjectPtr<class UPolys> CachedPolys;
 
 public:
 
diff --git a/Engine/Plugins/Enterprise/DatasmithContent/Source/DatasmithContent/Public/DatasmithAreaLightActor.h b/Engine/Plugins/Enterprise/DatasmithContent/Source/DatasmithContent/Public/DatasmithAreaLightActor.h
--- a/Engine/Plugins/Enterprise/DatasmithContent/Source/DatasmithContent/Public/DatasmithAreaLightActor.h
+++ b/Engine/Plugins/Enterprise/DatasmithContent/Source/DatasmithContent/Public/DatasmithAreaLightActor.h
@@ -54,7 +54,7 @@ public:
 	float Temperature;
 
 	UPROPERTY(EditAnywhere, BlueprintReadWrite, Category = "Light")
-	class UTextureLightProfile* IESTexture;
+	TObjectPtr<class UTextureLightProfile> IESTexture;
 
 	UPROPERTY(EditAnywhere, BlueprintReadWrite, Category = "Light", meta=(DisplayName = "Use IES Intensity"))
 	bool bUseIESBrightness;
diff --git a/Engine/Plugins/Enterprise/DatasmithContent/Source/DatasmithContent/Public/DatasmithAssetImportData.h b/Engine/Plugins/Enterprise/DatasmithContent/Source/DatasmithContent/Public/DatasmithAssetImportData.h
--- a/Engine/Plugins/Enterprise/DatasmithContent/Source/DatasmithContent/Public/DatasmithAssetImportData.h
+++ b/Engine/Plugins/Enterprise/DatasmithContent/Source/DatasmithContent/Public/DatasmithAssetImportData.h
@@ -19,7 +19,7 @@ public:
 	FDatasmithAssetImportOptions AssetImportOptions;
 
 	UPROPERTY(EditAnywhere, Instanced, Category = Asset, meta = (ShowOnlyInnerProperties))
-	TArray<class UDatasmithAdditionalData*> AdditionalData;
+	TArray<TObjectPtr<class UDatasmithAdditionalData>> AdditionalData;
 #endif		// WITH_EDITORONLY_DATA
 };
 
@@ -112,7 +112,7 @@ class DATASMITHCONTENT_API UDatasmithTranslatedSceneImportData : public UDatasmi
 #if WITH_EDITORONLY_DATA
 public:
 	UPROPERTY(EditAnywhere, Category = "Options", meta = (ShowOnlyInnerProperties))
-	TArray<UDatasmithOptionsBase*> AdditionalOptions;
+	TArray<TObjectPtr<UDatasmithOptionsBase>> AdditionalOptions;
 #endif // WITH_EDITORONLY_DATA
 };
 
diff --git a/Engine/Plugins/Enterprise/DatasmithContent/Source/DatasmithContent/Public/DatasmithAssetUserData.h b/Engine/Plugins/Enterprise/DatasmithContent/Source/DatasmithContent/Public/DatasmithAssetUserData.h
--- a/Engine/Plugins/Enterprise/DatasmithContent/Source/DatasmithContent/Public/DatasmithAssetUserData.h
+++ b/Engine/Plugins/Enterprise/DatasmithContent/Source/DatasmithContent/Public/DatasmithAssetUserData.h
@@ -23,7 +23,7 @@ public:
 
 #if WITH_EDITORONLY_DATA
 	UPROPERTY()
-	TMap< TSubclassOf< class UDatasmithObjectTemplate >, UDatasmithObjectTemplate* > ObjectTemplates;
+	TMap< TSubclassOf< class UDatasmithObjectTemplate >, TObjectPtr<UDatasmithObjectTemplate> > ObjectTemplates;
 
 	virtual bool IsPostLoadThreadSafe() const override;
 	virtual void PostLoad() override;
diff --git a/Engine/Plugins/Enterprise/DatasmithContent/Source/DatasmithContent/Public/DatasmithImportedSequencesActor.h b/Engine/Plugins/Enterprise/DatasmithContent/Source/DatasmithContent/Public/DatasmithImportedSequencesActor.h
--- a/Engine/Plugins/Enterprise/DatasmithContent/Source/DatasmithContent/Public/DatasmithImportedSequencesActor.h
+++ b/Engine/Plugins/Enterprise/DatasmithContent/Source/DatasmithContent/Public/DatasmithImportedSequencesActor.h
@@ -17,7 +17,7 @@ public:
 	ADatasmithImportedSequencesActor(const FObjectInitializer& Init);
 
 	UPROPERTY(EditAnywhere, BlueprintReadWrite, Category = "ImportedSequences")
-	TArray<ULevelSequence*> ImportedSequences;
+	TArray<TObjectPtr<ULevelSequence>> ImportedSequences;
 
     UFUNCTION(BlueprintCallable, Category="ImportedSequences")
 	void PlayLevelSequence(ULevelSequence* SequenceToPlay);
diff --git a/Engine/Plugins/Enterprise/DatasmithContent/Source/DatasmithContent/Public/DatasmithScene.h b/Engine/Plugins/Enterprise/DatasmithContent/Source/DatasmithContent/Public/DatasmithScene.h
--- a/Engine/Plugins/Enterprise/DatasmithContent/Source/DatasmithContent/Public/DatasmithScene.h
+++ b/Engine/Plugins/Enterprise/DatasmithContent/Source/DatasmithContent/Public/DatasmithScene.h
@@ -30,7 +30,7 @@ public:
 #if WITH_EDITORONLY_DATA
 	/** Importing data and options used for this Datasmith scene */
 	UPROPERTY(EditAnywhere, Instanced, Category=ImportSettings)
-	class UDatasmithSceneImportData* AssetImportData;
+	TObjectPtr<class UDatasmithSceneImportData> AssetImportData;
 
 	UPROPERTY()
 	int32 BulkDataVersion; // Need an external version number because loading of the bulk data is handled externally
@@ -63,7 +63,7 @@ public:
 
 	/** Array of user data stored with the asset */
 	UPROPERTY()
-	TArray< UAssetUserData* > AssetUserData;
+	TArray< TObjectPtr<UAssetUserData> > AssetUserData;
 #endif // #if WITH_EDITORONLY_DATA
 
 	/** Register the DatasmithScene to the PreWorldRename callback as needed*/
diff --git a/Engine/Plugins/Enterprise/DatasmithContent/Source/DatasmithContent/Public/DatasmithSceneActor.h b/Engine/Plugins/Enterprise/DatasmithContent/Source/DatasmithContent/Public/DatasmithSceneActor.h
--- a/Engine/Plugins/Enterprise/DatasmithContent/Source/DatasmithContent/Public/DatasmithSceneActor.h
+++ b/Engine/Plugins/Enterprise/DatasmithContent/Source/DatasmithContent/Public/DatasmithSceneActor.h
@@ -20,7 +20,7 @@ public:
 	virtual ~ADatasmithSceneActor();
 
 	UPROPERTY(VisibleAnywhere, Category="Datasmith")
-	class UDatasmithScene* Scene;
+	TObjectPtr<class UDatasmithScene> Scene;
 
 	/** Map of all the actors related to this Datasmith Scene */
 	UPROPERTY(VisibleAnywhere, Category="Datasmith", AdvancedDisplay)
diff --git a/Engine/Plugins/Enterprise/DatasmithContent/Source/DatasmithContent/Public/ObjectTemplates/DatasmithDecalComponentTemplate.h b/Engine/Plugins/Enterprise/DatasmithContent/Source/DatasmithContent/Public/ObjectTemplates/DatasmithDecalComponentTemplate.h
--- a/Engine/Plugins/Enterprise/DatasmithContent/Source/DatasmithContent/Public/ObjectTemplates/DatasmithDecalComponentTemplate.h
+++ b/Engine/Plugins/Enterprise/DatasmithContent/Source/DatasmithContent/Public/ObjectTemplates/DatasmithDecalComponentTemplate.h
@@ -21,7 +21,7 @@ public:
 	FVector DecalSize;
 
 	UPROPERTY()
-	class UMaterialInterface* Material;
+	TObjectPtr<class UMaterialInterface> Material;
 
 	virtual UObject* UpdateObject( UObject* Destination, bool bForce = false ) override;
 	virtual void Load( const UObject* Source ) override;
diff --git a/Engine/Plugins/Enterprise/DatasmithContent/Source/DatasmithContent/Public/ObjectTemplates/DatasmithLandscapeTemplate.h b/Engine/Plugins/Enterprise/DatasmithContent/Source/DatasmithContent/Public/ObjectTemplates/DatasmithLandscapeTemplate.h
--- a/Engine/Plugins/Enterprise/DatasmithContent/Source/DatasmithContent/Public/ObjectTemplates/DatasmithLandscapeTemplate.h
+++ b/Engine/Plugins/Enterprise/DatasmithContent/Source/DatasmithContent/Public/ObjectTemplates/DatasmithLandscapeTemplate.h
@@ -19,7 +19,7 @@ public:
 	{}
 
 	UPROPERTY()
-	UMaterialInterface* LandscapeMaterial;
+	TObjectPtr<UMaterialInterface> LandscapeMaterial;
 
 	UPROPERTY()
 	int32 StaticLightingLOD;
diff --git a/Engine/Plugins/Enterprise/DatasmithContent/Source/DatasmithContent/Public/ObjectTemplates/DatasmithLightComponentTemplate.h b/Engine/Plugins/Enterprise/DatasmithContent/Source/DatasmithContent/Public/ObjectTemplates/DatasmithLightComponentTemplate.h
--- a/Engine/Plugins/Enterprise/DatasmithContent/Source/DatasmithContent/Public/ObjectTemplates/DatasmithLightComponentTemplate.h
+++ b/Engine/Plugins/Enterprise/DatasmithContent/Source/DatasmithContent/Public/ObjectTemplates/DatasmithLightComponentTemplate.h
@@ -39,10 +39,10 @@ public:
 	FLinearColor LightColor;
 
 	UPROPERTY()
-	class UMaterialInterface* LightFunctionMaterial;
+	TObjectPtr<class UMaterialInterface> LightFunctionMaterial;
 
 	UPROPERTY()
-	class UTextureLightProfile* IESTexture;
+	TObjectPtr<class UTextureLightProfile> IESTexture;
 
 	virtual UObject* UpdateObject( UObject* Destination, bool bForce = false ) override;
 	virtual void Load( const UObject* Source ) override;
diff --git a/Engine/Plugins/Enterprise/DatasmithContent/Source/DatasmithContent/Public/ObjectTemplates/DatasmithSkyLightComponentTemplate.h b/Engine/Plugins/Enterprise/DatasmithContent/Source/DatasmithContent/Public/ObjectTemplates/DatasmithSkyLightComponentTemplate.h
--- a/Engine/Plugins/Enterprise/DatasmithContent/Source/DatasmithContent/Public/ObjectTemplates/DatasmithSkyLightComponentTemplate.h
+++ b/Engine/Plugins/Enterprise/DatasmithContent/Source/DatasmithContent/Public/ObjectTemplates/DatasmithSkyLightComponentTemplate.h
@@ -23,7 +23,7 @@ public:
 	int32 CubemapResolution;
 
 	UPROPERTY()
-	class UTextureCube* Cubemap;
+	TObjectPtr<class UTextureCube> Cubemap;
 
 	virtual UObject* UpdateObject( UObject* Destination, bool bForce = false ) override;
 	virtual void Load( const UObject* Source ) override;
diff --git a/Engine/Plugins/Enterprise/DatasmithContent/Source/DatasmithContent/Public/ObjectTemplates/DatasmithStaticMeshComponentTemplate.h b/Engine/Plugins/Enterprise/DatasmithContent/Source/DatasmithContent/Public/ObjectTemplates/DatasmithStaticMeshComponentTemplate.h
--- a/Engine/Plugins/Enterprise/DatasmithContent/Source/DatasmithContent/Public/ObjectTemplates/DatasmithStaticMeshComponentTemplate.h
+++ b/Engine/Plugins/Enterprise/DatasmithContent/Source/DatasmithContent/Public/ObjectTemplates/DatasmithStaticMeshComponentTemplate.h
@@ -14,10 +14,10 @@ class DATASMITHCONTENT_API UDatasmithStaticMeshComponentTemplate : public UDatas
 public:
 
 	UPROPERTY()
-	class UStaticMesh* StaticMesh;
+	TObjectPtr<class UStaticMesh> StaticMesh;
 
 	UPROPERTY()
-	TArray< class UMaterialInterface* > OverrideMaterials;
+	TArray< TObjectPtr<class UMaterialInterface> > OverrideMaterials;
 
 	virtual UObject* UpdateObject( UObject* Destination, bool bForce = false ) override;
 	virtual void Load( const UObject* Source ) override;
diff --git a/Engine/Plugins/Enterprise/DatasmithContent/Source/DatasmithContent/Public/ObjectTemplates/DatasmithStaticMeshTemplate.h b/Engine/Plugins/Enterprise/DatasmithContent/Source/DatasmithContent/Public/ObjectTemplates/DatasmithStaticMeshTemplate.h
--- a/Engine/Plugins/Enterprise/DatasmithContent/Source/DatasmithContent/Public/ObjectTemplates/DatasmithStaticMeshTemplate.h
+++ b/Engine/Plugins/Enterprise/DatasmithContent/Source/DatasmithContent/Public/ObjectTemplates/DatasmithStaticMeshTemplate.h
@@ -70,7 +70,7 @@ public:
 	FName MaterialSlotName;
 
 	UPROPERTY()
-	class UMaterialInterface* MaterialInterface;
+	TObjectPtr<class UMaterialInterface> MaterialInterface;
 
 	void Apply( FStaticMaterial* Destination, FDatasmithStaticMaterialTemplate* PreviousTemplate );
 	void Load( const FStaticMaterial& Source );
diff --git a/Engine/Plugins/Enterprise/VariantManagerContent/Source/VariantManagerContent/Public/LevelVariantSets.h b/Engine/Plugins/Enterprise/VariantManagerContent/Source/VariantManagerContent/Public/LevelVariantSets.h
--- a/Engine/Plugins/Enterprise/VariantManagerContent/Source/VariantManagerContent/Public/LevelVariantSets.h
+++ b/Engine/Plugins/Enterprise/VariantManagerContent/Source/VariantManagerContent/Public/LevelVariantSets.h
@@ -91,7 +91,7 @@ private:
 
 	// A pointer to the director blueprint that generates this levelvariantsets' DirectorClass
 	UPROPERTY()
-	UObject* DirectorBlueprint;
+	TObjectPtr<UObject> DirectorBlueprint;
 
 	FDelegateHandle OnBlueprintCompiledHandle;
 	FDelegateHandle EndPlayDelegateHandle;
@@ -100,10 +100,10 @@ private:
 	// The class that is used to spawn this levelvariantsets' director instance.
 	// Director instances are allocated one per instance
 	UPROPERTY()
-	UBlueprintGeneratedClass* DirectorClass;
+	TObjectPtr<UBlueprintGeneratedClass> DirectorClass;
 
 	UPROPERTY()
-	TArray<UVariantSet*> VariantSets;
+	TArray<TObjectPtr<UVariantSet>> VariantSets;
 
 	// We keep one director instance per world to execute our functions
 	TMap<UWorld*, UObject*> WorldToDirectorInstance;
diff --git a/Engine/Plugins/Enterprise/VariantManagerContent/Source/VariantManagerContent/Public/PropertyValue.h b/Engine/Plugins/Enterprise/VariantManagerContent/Source/VariantManagerContent/Public/PropertyValue.h
--- a/Engine/Plugins/Enterprise/VariantManagerContent/Source/VariantManagerContent/Public/PropertyValue.h
+++ b/Engine/Plugins/Enterprise/VariantManagerContent/Source/VariantManagerContent/Public/PropertyValue.h
@@ -208,7 +208,7 @@ protected:
 	// We use these mainly to know how to serialize/deserialize the values of properties that need special care
 	// (e.g. UObjectProperties, name properties, text properties, etc)
 	UPROPERTY()
-	UClass* LeafPropertyClass_DEPRECATED;
+	TObjectPtr<UClass> LeafPropertyClass_DEPRECATED;
 	FFieldClass* LeafPropertyClass;
 
 	UPROPERTY()
diff --git a/Engine/Plugins/Enterprise/VariantManagerContent/Source/VariantManagerContent/Public/SwitchActor.h b/Engine/Plugins/Enterprise/VariantManagerContent/Source/VariantManagerContent/Public/SwitchActor.h
--- a/Engine/Plugins/Enterprise/VariantManagerContent/Source/VariantManagerContent/Public/SwitchActor.h
+++ b/Engine/Plugins/Enterprise/VariantManagerContent/Source/VariantManagerContent/Public/SwitchActor.h
@@ -48,7 +48,7 @@ private:
 
 	// Exposing our root component like this allows manual Mobility control on the details panel
 	UPROPERTY(Category = SwitchActor, VisibleAnywhere)
-	class USceneComponent* SceneComponent;
+	TObjectPtr<class USceneComponent> SceneComponent;
 
 	UPROPERTY()
 	int32 LastSelectedOption;
diff --git a/Engine/Plugins/Enterprise/VariantManagerContent/Source/VariantManagerContent/Public/Variant.h b/Engine/Plugins/Enterprise/VariantManagerContent/Source/VariantManagerContent/Public/Variant.h
--- a/Engine/Plugins/Enterprise/VariantManagerContent/Source/VariantManagerContent/Public/Variant.h
+++ b/Engine/Plugins/Enterprise/VariantManagerContent/Source/VariantManagerContent/Public/Variant.h
@@ -134,10 +134,10 @@ private:
 	FText DisplayText;
 
 	UPROPERTY()
-	TArray<UVariantObjectBinding*> ObjectBindings;
+	TArray<TObjectPtr<UVariantObjectBinding>> ObjectBindings;
 
 	UPROPERTY()
-	UTexture2D* Thumbnail;
+	TObjectPtr<UTexture2D> Thumbnail;
 
 #if WITH_EDITOR
 	// Whether we already tried restoring a thumbnail from the actual package (backwards compatibility)
diff --git a/Engine/Plugins/Enterprise/VariantManagerContent/Source/VariantManagerContent/Public/VariantObjectBinding.h b/Engine/Plugins/Enterprise/VariantManagerContent/Source/VariantManagerContent/Public/VariantObjectBinding.h
--- a/Engine/Plugins/Enterprise/VariantManagerContent/Source/VariantManagerContent/Public/VariantObjectBinding.h
+++ b/Engine/Plugins/Enterprise/VariantManagerContent/Source/VariantManagerContent/Public/VariantObjectBinding.h
@@ -62,7 +62,7 @@ private:
 	mutable TLazyObjectPtr<UObject> LazyObjectPtr;
 
 	UPROPERTY()
-	TArray<UPropertyValue*> CapturedProperties;
+	TArray<TObjectPtr<UPropertyValue>> CapturedProperties;
 
 	UPROPERTY()
 	TArray<FFunctionCaller> FunctionCallers;
diff --git a/Engine/Plugins/Enterprise/VariantManagerContent/Source/VariantManagerContent/Public/VariantSet.h b/Engine/Plugins/Enterprise/VariantManagerContent/Source/VariantManagerContent/Public/VariantSet.h
--- a/Engine/Plugins/Enterprise/VariantManagerContent/Source/VariantManagerContent/Public/VariantSet.h
+++ b/Engine/Plugins/Enterprise/VariantManagerContent/Source/VariantManagerContent/Public/VariantSet.h
@@ -86,8 +86,8 @@ private:
 	bool bExpanded;
 
 	UPROPERTY()
-	TArray<UVariant*> Variants;
+	TArray<TObjectPtr<UVariant>> Variants;
 
 	UPROPERTY()
-	UTexture2D* Thumbnail;
+	TObjectPtr<UTexture2D> Thumbnail;
 };
diff --git a/Engine/Plugins/Enterprise/VariantManagerContent/Source/VariantManagerContentEditor/Public/VariantManagerTestActor.h b/Engine/Plugins/Enterprise/VariantManagerContent/Source/VariantManagerContentEditor/Public/VariantManagerTestActor.h
--- a/Engine/Plugins/Enterprise/VariantManagerContent/Source/VariantManagerContentEditor/Public/VariantManagerTestActor.h
+++ b/Engine/Plugins/Enterprise/VariantManagerContent/Source/VariantManagerContentEditor/Public/VariantManagerTestActor.h
@@ -66,7 +66,7 @@ public:
 	bool bCapturedBoolProperty;
 
 	UPROPERTY(EditAnywhere, BlueprintReadWrite, Category="Template", meta=(ToolTip="Captured UObject property"))
-	UObject* CapturedObjectProperty;
+	TObjectPtr<UObject> CapturedObjectProperty;
 
 	UPROPERTY(EditAnywhere, BlueprintReadWrite, Category="Template", meta=(ToolTip="Captured UInterface property"))
 	FScriptInterface CapturedInterfaceProperty;
@@ -105,7 +105,7 @@ public:
 	FIntPoint CapturedIntPointProperty;
 
 	UPROPERTY(EditAnywhere, BlueprintReadWrite, Category="Template", meta=(ToolTip="Captured UObject array property"))
-	TArray<UObject*> CapturedUObjectArrayProperty;
+	TArray<TObjectPtr<UObject>> CapturedUObjectArrayProperty;
 
 	UPROPERTY(EditAnywhere, BlueprintReadWrite, Category="Template", meta=(ToolTip="Captured FVector array property"))
 	TArray<FVector> CapturedVectorArrayProperty;
diff --git a/Engine/Plugins/Experimental/AlembicImporter/Source/AlembicImporter/Classes/AlembicImportFactory.h b/Engine/Plugins/Experimental/AlembicImporter/Source/AlembicImporter/Classes/AlembicImportFactory.h
--- a/Engine/Plugins/Experimental/AlembicImporter/Source/AlembicImporter/Classes/AlembicImportFactory.h
+++ b/Engine/Plugins/Experimental/AlembicImporter/Source/AlembicImporter/Classes/AlembicImportFactory.h
@@ -24,7 +24,7 @@ class UAlembicImportFactory : public UFactory, public FReimportHandler
 
 	/** Object used to show import options for Alembic */
 	UPROPERTY()
-	UAbcImportSettings* ImportSettings;
+	TObjectPtr<UAbcImportSettings> ImportSettings;
 
 	UPROPERTY()
 	bool bShowOption;
diff --git a/Engine/Plugins/Experimental/ChaosEditor/Source/FractureEditor/Private/FractureToolAutoCluster.h b/Engine/Plugins/Experimental/ChaosEditor/Source/FractureEditor/Private/FractureToolAutoCluster.h
--- a/Engine/Plugins/Experimental/ChaosEditor/Source/FractureEditor/Private/FractureToolAutoCluster.h
+++ b/Engine/Plugins/Experimental/ChaosEditor/Source/FractureEditor/Private/FractureToolAutoCluster.h
@@ -47,5 +47,5 @@ public:
 	virtual void Execute(TWeakPtr<FFractureEditorModeToolkit> InToolkit) override;
 
 	UPROPERTY(EditAnywhere, Category = AutoCluster)
-	UFractureAutoClusterSettings* AutoClusterSettings;
+	TObjectPtr<UFractureAutoClusterSettings> AutoClusterSettings;
 };
\ No newline at end of file
diff --git a/Engine/Plugins/Experimental/ChaosEditor/Source/FractureEditor/Private/FractureToolBrick.h b/Engine/Plugins/Experimental/ChaosEditor/Source/FractureEditor/Private/FractureToolBrick.h
--- a/Engine/Plugins/Experimental/ChaosEditor/Source/FractureEditor/Private/FractureToolBrick.h
+++ b/Engine/Plugins/Experimental/ChaosEditor/Source/FractureEditor/Private/FractureToolBrick.h
@@ -101,7 +101,7 @@ public:
 
 private:
 	UPROPERTY()
-	UFractureBrickSettings* BrickSettings;
+	TObjectPtr<UFractureBrickSettings> BrickSettings;
 
 	void AddBoxEdges(const FVector& Min, const FVector& Max);
 
diff --git a/Engine/Plugins/Experimental/ChaosEditor/Source/FractureEditor/Private/FractureToolClusterCutter.h b/Engine/Plugins/Experimental/ChaosEditor/Source/FractureEditor/Private/FractureToolClusterCutter.h
--- a/Engine/Plugins/Experimental/ChaosEditor/Source/FractureEditor/Private/FractureToolClusterCutter.h
+++ b/Engine/Plugins/Experimental/ChaosEditor/Source/FractureEditor/Private/FractureToolClusterCutter.h
@@ -73,7 +73,7 @@ public:
 
 	
 	UPROPERTY(EditAnywhere, Category = Cluster)
-	UFractureClusterCutterSettings* ClusterSettings;
+	TObjectPtr<UFractureClusterCutterSettings> ClusterSettings;
 
 protected:
 	void GenerateVoronoiSites(const FFractureToolContext& Context, TArray<FVector>& Sites) override;
diff --git a/Engine/Plugins/Experimental/ChaosEditor/Source/FractureEditor/Private/FractureToolClusterMagnet.h b/Engine/Plugins/Experimental/ChaosEditor/Source/FractureEditor/Private/FractureToolClusterMagnet.h
--- a/Engine/Plugins/Experimental/ChaosEditor/Source/FractureEditor/Private/FractureToolClusterMagnet.h
+++ b/Engine/Plugins/Experimental/ChaosEditor/Source/FractureEditor/Private/FractureToolClusterMagnet.h
@@ -58,7 +58,7 @@ public:
 	virtual void Execute(TWeakPtr<FFractureEditorModeToolkit> InToolkit) override;
 
 	UPROPERTY(EditAnywhere, Category = AutoCluster)
-	UFractureClusterMagnetSettings* ClusterMagnetSettings;
+	TObjectPtr<UFractureClusterMagnetSettings> ClusterMagnetSettings;
 
 private:
 	bool CheckPresenceOfNecessaryAttributes(const FGeometryCollectionPtr GeometryCollection) const;
diff --git a/Engine/Plugins/Experimental/ChaosEditor/Source/FractureEditor/Private/FractureToolCutter.h b/Engine/Plugins/Experimental/ChaosEditor/Source/FractureEditor/Private/FractureToolCutter.h
--- a/Engine/Plugins/Experimental/ChaosEditor/Source/FractureEditor/Private/FractureToolCutter.h
+++ b/Engine/Plugins/Experimental/ChaosEditor/Source/FractureEditor/Private/FractureToolCutter.h
@@ -89,7 +89,7 @@ public:
 
 protected:
 	UPROPERTY(EditAnywhere, Category = Slicing)
-	UFractureCutterSettings* CutterSettings;
+	TObjectPtr<UFractureCutterSettings> CutterSettings;
 
 };
 
diff --git a/Engine/Plugins/Experimental/ChaosEditor/Source/FractureEditor/Private/FractureToolPlaneCut.h b/Engine/Plugins/Experimental/ChaosEditor/Source/FractureEditor/Private/FractureToolPlaneCut.h
--- a/Engine/Plugins/Experimental/ChaosEditor/Source/FractureEditor/Private/FractureToolPlaneCut.h
+++ b/Engine/Plugins/Experimental/ChaosEditor/Source/FractureEditor/Private/FractureToolPlaneCut.h
@@ -55,7 +55,7 @@ public:
 private:
 	// Slicing
 	UPROPERTY(EditAnywhere, Category = Slicing)
-	UFracturePlaneCutSettings* PlaneCutSettings;
+	TObjectPtr<UFracturePlaneCutSettings> PlaneCutSettings;
 
 	void GenerateSliceTransforms(const FFractureToolContext& Context, TArray<FTransform>& CuttingPlaneTransforms);
 
diff --git a/Engine/Plugins/Experimental/ChaosEditor/Source/FractureEditor/Private/FractureToolProperties.h b/Engine/Plugins/Experimental/ChaosEditor/Source/FractureEditor/Private/FractureToolProperties.h
--- a/Engine/Plugins/Experimental/ChaosEditor/Source/FractureEditor/Private/FractureToolProperties.h
+++ b/Engine/Plugins/Experimental/ChaosEditor/Source/FractureEditor/Private/FractureToolProperties.h
@@ -57,6 +57,6 @@ public:
 	virtual void Execute(TWeakPtr<FFractureEditorModeToolkit> InToolkit) override;
 
 	UPROPERTY(EditAnywhere, Category = InitialDynamicState)
-	UFractureInitialDynamicStateSettings* StateSettings;
+	TObjectPtr<UFractureInitialDynamicStateSettings> StateSettings;
 
 };
diff --git a/Engine/Plugins/Experimental/ChaosEditor/Source/FractureEditor/Private/FractureToolRadial.h b/Engine/Plugins/Experimental/ChaosEditor/Source/FractureEditor/Private/FractureToolRadial.h
--- a/Engine/Plugins/Experimental/ChaosEditor/Source/FractureEditor/Private/FractureToolRadial.h
+++ b/Engine/Plugins/Experimental/ChaosEditor/Source/FractureEditor/Private/FractureToolRadial.h
@@ -74,7 +74,7 @@ public:
 
 	// Radial Voronoi Fracture Input Settings
 	UPROPERTY(EditAnywhere, Category = Uniform)
-	UFractureRadialSettings* RadialSettings;
+	TObjectPtr<UFractureRadialSettings> RadialSettings;
 
 protected:
 	void GenerateVoronoiSites(const FFractureToolContext& Context, TArray<FVector>& Sites) override;
diff --git a/Engine/Plugins/Experimental/ChaosEditor/Source/FractureEditor/Private/FractureToolSlice.h b/Engine/Plugins/Experimental/ChaosEditor/Source/FractureEditor/Private/FractureToolSlice.h
--- a/Engine/Plugins/Experimental/ChaosEditor/Source/FractureEditor/Private/FractureToolSlice.h
+++ b/Engine/Plugins/Experimental/ChaosEditor/Source/FractureEditor/Private/FractureToolSlice.h
@@ -71,7 +71,7 @@ public:
 
 	// Slicing
 	UPROPERTY(EditAnywhere, Category = Slicing)
-	UFractureSliceSettings* SliceSettings;
+	TObjectPtr<UFractureSliceSettings> SliceSettings;
 
 private:
 	float RenderCuttingPlaneSize;
diff --git a/Engine/Plugins/Experimental/ChaosEditor/Source/FractureEditor/Private/FractureToolUniform.h b/Engine/Plugins/Experimental/ChaosEditor/Source/FractureEditor/Private/FractureToolUniform.h
--- a/Engine/Plugins/Experimental/ChaosEditor/Source/FractureEditor/Private/FractureToolUniform.h
+++ b/Engine/Plugins/Experimental/ChaosEditor/Source/FractureEditor/Private/FractureToolUniform.h
@@ -49,7 +49,7 @@ public:
 
 	// Uniform Voronoi Fracture Input Settings
 	UPROPERTY(EditAnywhere, Category = Uniform)
-	UFractureUniformSettings* UniformSettings;
+	TObjectPtr<UFractureUniformSettings> UniformSettings;
 
 protected:
 	void GenerateVoronoiSites(const FFractureToolContext& Context, TArray<FVector>& Sites) override;
diff --git a/Engine/Plugins/Experimental/ChaosEditor/Source/FractureEditor/Public/FractureTool.h b/Engine/Plugins/Experimental/ChaosEditor/Source/FractureEditor/Public/FractureTool.h
--- a/Engine/Plugins/Experimental/ChaosEditor/Source/FractureEditor/Public/FractureTool.h
+++ b/Engine/Plugins/Experimental/ChaosEditor/Source/FractureEditor/Public/FractureTool.h
@@ -32,7 +32,7 @@ public:
 	void PostEditChangeChainProperty(struct FPropertyChangedChainEvent& PropertyChangedEvent) override;
 
 	UPROPERTY()
-	UFractureModalTool* OwnerTool;
+	TObjectPtr<UFractureModalTool> OwnerTool;
 };
 
 
diff --git a/Engine/Plugins/Experimental/ChaosNiagara/Source/ChaosNiagara/Classes/NiagaraDataInterfaceChaosDestruction.h b/Engine/Plugins/Experimental/ChaosNiagara/Source/ChaosNiagara/Classes/NiagaraDataInterfaceChaosDestruction.h
--- a/Engine/Plugins/Experimental/ChaosNiagara/Source/ChaosNiagara/Classes/NiagaraDataInterfaceChaosDestruction.h
+++ b/Engine/Plugins/Experimental/ChaosNiagara/Source/ChaosNiagara/Classes/NiagaraDataInterfaceChaosDestruction.h
@@ -357,7 +357,7 @@ public:
 
 	/* Chaos Solver */
 	UPROPERTY(EditAnywhere, Category = "Solver", meta = (DisplayName = "Chaos Solver"))
-	TSet<AChaosSolverActor*> ChaosSolverActorSet;
+	TSet<TObjectPtr<AChaosSolverActor>> ChaosSolverActorSet;
 
 	/* */
 	UPROPERTY(EditAnywhere, Category = "Solver Data", meta = (DisplayName = "Data Source"))
@@ -479,7 +479,7 @@ public:
 
 	/* TODO: Explanatory comment */
 	UPROPERTY(EditAnywhere, Category = "Breaking Data Settings", meta = (DisplayName = "Breaking Filtered Materials", EditCondition = bApplyMaterialsFilter))
-	TSet<UPhysicalMaterial*> ChaosBreakingMaterialSet;
+	TSet<TObjectPtr<UPhysicalMaterial>> ChaosBreakingMaterialSet;
 
 	/* TODO: Explanatory comment */
 	UPROPERTY(EditAnywhere, Category = "Breaking Data Settings", meta = (DisplayName = "Get External Breaking Mesh and Physical Data"))
diff --git a/Engine/Plugins/Experimental/GeometryCache/Source/GeometryCache/Classes/GeometryCache.h b/Engine/Plugins/Experimental/GeometryCache/Source/GeometryCache/Classes/GeometryCache.h
--- a/Engine/Plugins/Experimental/GeometryCache/Source/GeometryCache/Classes/GeometryCache.h
+++ b/Engine/Plugins/Experimental/GeometryCache/Source/GeometryCache/Classes/GeometryCache.h
@@ -50,18 +50,18 @@ public:
 #if WITH_EDITORONLY_DATA
 	/** Importing data and options used for this Geometry cache object*/
 	UPROPERTY(Category = ImportSettings, VisibleAnywhere, Instanced)
-	class UAssetImportData* AssetImportData;	
+	TObjectPtr<class UAssetImportData> AssetImportData;	
 
 	/** Information for thumbnail rendering */
 	UPROPERTY(VisibleAnywhere, Instanced, Category = Thumbnail)
-	class UThumbnailInfo* ThumbnailInfo;
+	TObjectPtr<class UThumbnailInfo> ThumbnailInfo;
 #endif
 	UPROPERTY(EditAnywhere, Category = GeometryCache)
-	TArray<UMaterialInterface*> Materials;
+	TArray<TObjectPtr<UMaterialInterface>> Materials;
 	
 	/** GeometryCache track defining the samples/geometry data for this GeomCache instance */
 	UPROPERTY(VisibleAnywhere, Category=GeometryCache)
-	TArray<UGeometryCacheTrack*> Tracks;
+	TArray<TObjectPtr<UGeometryCacheTrack>> Tracks;
 
 	/** Set the start and end frames for the GeometryCache */
 	void SetFrameStartEnd(int32 InStartFrame, int32 InEndFrame);
diff --git a/Engine/Plugins/Experimental/GeometryCache/Source/GeometryCache/Classes/GeometryCacheActor.h b/Engine/Plugins/Experimental/GeometryCache/Source/GeometryCache/Classes/GeometryCacheActor.h
--- a/Engine/Plugins/Experimental/GeometryCache/Source/GeometryCache/Classes/GeometryCacheActor.h
+++ b/Engine/Plugins/Experimental/GeometryCache/Source/GeometryCache/Classes/GeometryCacheActor.h
@@ -22,7 +22,7 @@ class GEOMETRYCACHE_API AGeometryCacheActor : public AActor
 	// End AActor overrides.
 private:
 	UPROPERTY(Category = GeometryCacheActor, VisibleAnywhere, BlueprintReadOnly, meta = (ExposeFunctionCategories = "Mesh,Rendering,Physics,Components|GeometryCache", AllowPrivateAccess = "true"))
-	UGeometryCacheComponent* GeometryCacheComponent;
+	TObjectPtr<UGeometryCacheComponent> GeometryCacheComponent;
 public:
 	/** Returns GeometryCacheComponent subobject **/
 	UFUNCTION(BlueprintCallable, Category = "Components|GeometryCache")
diff --git a/Engine/Plugins/Experimental/GeometryCache/Source/GeometryCache/Classes/GeometryCacheComponent.h b/Engine/Plugins/Experimental/GeometryCache/Source/GeometryCache/Classes/GeometryCacheComponent.h
--- a/Engine/Plugins/Experimental/GeometryCache/Source/GeometryCache/Classes/GeometryCacheComponent.h
+++ b/Engine/Plugins/Experimental/GeometryCache/Source/GeometryCache/Classes/GeometryCacheComponent.h
@@ -157,7 +157,7 @@ public:
 
 	/** Geometry Cache instance referenced by the component */
 	UPROPERTY(EditAnywhere, BlueprintReadOnly, Category = GeometryCache)
-	UGeometryCache* GeometryCache;
+	TObjectPtr<UGeometryCache> GeometryCache;
 		
 	/** Get the duration of the playback */
 	UFUNCTION(BlueprintCallable, Category = "Components|GeometryCache")
diff --git a/Engine/Plugins/Experimental/GeometryCache/Source/GeometryCache/Classes/GeometryCacheTrackStreamable.h b/Engine/Plugins/Experimental/GeometryCache/Source/GeometryCache/Classes/GeometryCacheTrackStreamable.h
--- a/Engine/Plugins/Experimental/GeometryCache/Source/GeometryCache/Classes/GeometryCacheTrackStreamable.h
+++ b/Engine/Plugins/Experimental/GeometryCache/Source/GeometryCache/Classes/GeometryCacheTrackStreamable.h
@@ -142,7 +142,7 @@ class GEOMETRYCACHE_API UGeometryCacheTrackStreamable : public UGeometryCacheTra
 
 	/** Codec for this track */
 	UPROPERTY(VisibleAnywhere, Category = GeometryCache)
-	UGeometryCacheCodecBase *Codec;
+	TObjectPtr<UGeometryCacheCodecBase> Codec;
 	FGeometryCachePreprocessor *Preprocessor;
 
 	/**
diff --git a/Engine/Plugins/Experimental/GeometryCache/Source/GeometryCacheTracks/Classes/MovieSceneGeometryCacheSection.h b/Engine/Plugins/Experimental/GeometryCache/Source/GeometryCacheTracks/Classes/MovieSceneGeometryCacheSection.h
--- a/Engine/Plugins/Experimental/GeometryCache/Source/GeometryCacheTracks/Classes/MovieSceneGeometryCacheSection.h
+++ b/Engine/Plugins/Experimental/GeometryCache/Source/GeometryCacheTracks/Classes/MovieSceneGeometryCacheSection.h
@@ -22,7 +22,7 @@ USTRUCT()
 	 float GetSequenceLength() const;
 	/** The animation this section plays */
 	UPROPERTY(EditAnywhere, Category = "GeometryCache", DisplayName = "Geometry Cache")
-	UGeometryCache* GeometryCacheAsset;
+	TObjectPtr<UGeometryCache> GeometryCacheAsset;
 
 	/** The offset for the first loop of the animation clip */
 	UPROPERTY(EditAnywhere, Category = "GeometryCache")
diff --git a/Engine/Plugins/Experimental/GeometryCache/Source/GeometryCacheTracks/Classes/MovieSceneGeometryCacheTrack.h b/Engine/Plugins/Experimental/GeometryCache/Source/GeometryCacheTracks/Classes/MovieSceneGeometryCacheTrack.h
--- a/Engine/Plugins/Experimental/GeometryCache/Source/GeometryCacheTracks/Classes/MovieSceneGeometryCacheTrack.h
+++ b/Engine/Plugins/Experimental/GeometryCache/Source/GeometryCacheTracks/Classes/MovieSceneGeometryCacheTrack.h
@@ -52,6 +52,6 @@ private:
 
 	/** List of all animation sections */
 	UPROPERTY()
-	TArray<UMovieSceneSection*> AnimationSections;
+	TArray<TObjectPtr<UMovieSceneSection>> AnimationSections;
 
 };
diff --git a/Engine/Plugins/Experimental/GeometryCollectionPlugin/Source/GeometryCollectionEditor/Public/GeometryCollection/GeometryCollectionCacheFactory.h b/Engine/Plugins/Experimental/GeometryCollectionPlugin/Source/GeometryCollectionEditor/Public/GeometryCollection/GeometryCollectionCacheFactory.h
--- a/Engine/Plugins/Experimental/GeometryCollectionPlugin/Source/GeometryCollectionEditor/Public/GeometryCollection/GeometryCollectionCacheFactory.h
+++ b/Engine/Plugins/Experimental/GeometryCollectionPlugin/Source/GeometryCollectionEditor/Public/GeometryCollection/GeometryCollectionCacheFactory.h
@@ -20,7 +20,7 @@ public:
 	/** Config properties required for CreateNew */
 
 	UPROPERTY()
-	UGeometryCollection* TargetCollection;
+	TObjectPtr<UGeometryCollection> TargetCollection;
 
 	/** End required properties */
 
diff --git a/Engine/Plugins/Experimental/GeometryCollectionPlugin/Source/GeometryCollectionTracks/Public/MovieSceneGeometryCollectionTrack.h b/Engine/Plugins/Experimental/GeometryCollectionPlugin/Source/GeometryCollectionTracks/Public/MovieSceneGeometryCollectionTrack.h
--- a/Engine/Plugins/Experimental/GeometryCollectionPlugin/Source/GeometryCollectionTracks/Public/MovieSceneGeometryCollectionTrack.h
+++ b/Engine/Plugins/Experimental/GeometryCollectionPlugin/Source/GeometryCollectionTracks/Public/MovieSceneGeometryCollectionTrack.h
@@ -51,6 +51,6 @@ private:
 
 	/** List of all animation sections */
 	UPROPERTY()
-	TArray<UMovieSceneSection*> AnimationSections;
+	TArray<TObjectPtr<UMovieSceneSection>> AnimationSections;
 
 };
diff --git a/Engine/Plugins/Experimental/MotoSynth/Source/MotoSynth/Classes/SynthComponents/SynthComponentMoto.h b/Engine/Plugins/Experimental/MotoSynth/Source/MotoSynth/Classes/SynthComponents/SynthComponentMoto.h
--- a/Engine/Plugins/Experimental/MotoSynth/Source/MotoSynth/Classes/SynthComponents/SynthComponentMoto.h
+++ b/Engine/Plugins/Experimental/MotoSynth/Source/MotoSynth/Classes/SynthComponents/SynthComponentMoto.h
@@ -19,7 +19,7 @@ class MOTOSYNTH_API USynthComponentMoto : public USynthComponent
 public:
 	/* The moto synth preset to use for the moto synth component. */
 	UPROPERTY(EditAnywhere, BlueprintReadWrite, Category = "MotoSynth")
-	UMotoSynthPreset* MotoSynthPreset;
+	TObjectPtr<UMotoSynthPreset> MotoSynthPreset;
 
 	/* Sets the starting RPM of the engine */
 	UPROPERTY(EditAnywhere, BlueprintReadWrite, Category = "MotoSynth", meta = (ClampMin = "500.0", ClampMax = "20000.0", UIMin = "500.0", UIMax = "20000.0"))
diff --git a/Engine/Plugins/Experimental/MotoSynth/Source/MotoSynth/Public/MotoSynthPreset.h b/Engine/Plugins/Experimental/MotoSynth/Source/MotoSynth/Public/MotoSynthPreset.h
--- a/Engine/Plugins/Experimental/MotoSynth/Source/MotoSynth/Public/MotoSynthPreset.h
+++ b/Engine/Plugins/Experimental/MotoSynth/Source/MotoSynth/Public/MotoSynthPreset.h
@@ -67,11 +67,11 @@ struct MOTOSYNTH_API FMotoSynthRuntimeSettings
 
 	// Motosynth source to use for granular engine acceleration
 	UPROPERTY(EditAnywhere, Category = "Granular Engine", meta = (EditCondition = "bGranularEngineEnabled"))
-	UMotoSynthSource* AccelerationSource = nullptr;
+	TObjectPtr<UMotoSynthSource> AccelerationSource = nullptr;
 
 	// Motosynth source to use for granular engine deceleration
 	UPROPERTY(EditAnywhere, Category = "Granular Engine", meta = (EditCondition = "bGranularEngineEnabled"))
-	UMotoSynthSource* DecelerationSource = nullptr;
+	TObjectPtr<UMotoSynthSource> DecelerationSource = nullptr;
 
 	// If the stereo widener is enabled
 	UPROPERTY(EditAnywhere, BlueprintReadWrite, Category = "Effects")
diff --git a/Engine/Plugins/Experimental/MotoSynth/Source/MotoSynth/Public/MotoSynthSourceAsset.h b/Engine/Plugins/Experimental/MotoSynth/Source/MotoSynth/Public/MotoSynthSourceAsset.h
--- a/Engine/Plugins/Experimental/MotoSynth/Source/MotoSynth/Public/MotoSynthSourceAsset.h
+++ b/Engine/Plugins/Experimental/MotoSynth/Source/MotoSynth/Public/MotoSynthSourceAsset.h
@@ -100,7 +100,7 @@ public:
 #if WITH_EDITORONLY_DATA
 	// The source to use for the moto synth source
 	UPROPERTY(EditAnywhere, BlueprintReadWrite, Category = "Grain Table | Analysis")
-	USoundWave* SoundWaveSource;
+	TObjectPtr<USoundWave> SoundWaveSource;
 #endif // #if WITH_EDITORONLY_DATA
 
 	// Whether or not to convert this moto synth source to 8 bit on load to use less memory
diff --git a/Engine/Plugins/FX/Niagara/Source/Niagara/Classes/NiagaraDataInterface2DArrayTexture.h b/Engine/Plugins/FX/Niagara/Source/Niagara/Classes/NiagaraDataInterface2DArrayTexture.h
--- a/Engine/Plugins/FX/Niagara/Source/Niagara/Classes/NiagaraDataInterface2DArrayTexture.h
+++ b/Engine/Plugins/FX/Niagara/Source/Niagara/Classes/NiagaraDataInterface2DArrayTexture.h
@@ -19,7 +19,7 @@ public:
 	DECLARE_NIAGARA_DI_PARAMETER();
 
 	UPROPERTY(EditAnywhere, Category = "Texture")
-	UTexture2DArray* Texture;
+	TObjectPtr<UTexture2DArray> Texture;
 
 	//UObject Interface
 	virtual void PostInitProperties()override;
diff --git a/Engine/Plugins/FX/Niagara/Source/Niagara/Classes/NiagaraDataInterfaceAudio.h b/Engine/Plugins/FX/Niagara/Source/Niagara/Classes/NiagaraDataInterfaceAudio.h
--- a/Engine/Plugins/FX/Niagara/Source/Niagara/Classes/NiagaraDataInterfaceAudio.h
+++ b/Engine/Plugins/FX/Niagara/Source/Niagara/Classes/NiagaraDataInterfaceAudio.h
@@ -137,7 +137,7 @@ public:
 	
 	/** The audio submix where analyzed. */
 	UPROPERTY(EditAnywhere, Category = "Audio")
-	USoundSubmix* Submix;
+	TObjectPtr<USoundSubmix> Submix;
 
 
 	virtual bool Equals(const UNiagaraDataInterface* Other) const override;
diff --git a/Engine/Plugins/FX/Niagara/Source/Niagara/Classes/NiagaraDataInterfaceAudioOscilloscope.h b/Engine/Plugins/FX/Niagara/Source/Niagara/Classes/NiagaraDataInterfaceAudioOscilloscope.h
--- a/Engine/Plugins/FX/Niagara/Source/Niagara/Classes/NiagaraDataInterfaceAudioOscilloscope.h
+++ b/Engine/Plugins/FX/Niagara/Source/Niagara/Classes/NiagaraDataInterfaceAudioOscilloscope.h
@@ -96,7 +96,7 @@ public:
 	DECLARE_NIAGARA_DI_PARAMETER();
 	
 	UPROPERTY(EditAnywhere, Category = "Oscilloscope")
-	USoundSubmix* Submix;
+	TObjectPtr<USoundSubmix> Submix;
 
 	static const int32 MaxBufferResolution = 8192;
 
diff --git a/Engine/Plugins/FX/Niagara/Source/Niagara/Classes/NiagaraDataInterfaceAudioPlayer.h b/Engine/Plugins/FX/Niagara/Source/Niagara/Classes/NiagaraDataInterfaceAudioPlayer.h
--- a/Engine/Plugins/FX/Niagara/Source/Niagara/Classes/NiagaraDataInterfaceAudioPlayer.h
+++ b/Engine/Plugins/FX/Niagara/Source/Niagara/Classes/NiagaraDataInterfaceAudioPlayer.h
@@ -57,15 +57,15 @@ class NIAGARA_API UNiagaraDataInterfaceAudioPlayer : public UNiagaraDataInterfac
 public:
 	/** Reference to the audio asset to play */
 	UPROPERTY(EditAnywhere, Category = "Audio")
-    USoundBase* SoundToPlay;
+    TObjectPtr<USoundBase> SoundToPlay;
 
 	/** Optional sound attenuation setting to use */
 	UPROPERTY(EditAnywhere, Category = "Audio")
-	USoundAttenuation* Attenuation;
+	TObjectPtr<USoundAttenuation> Attenuation;
 
 	/** Optional sound concurrency setting to use */
 	UPROPERTY(EditAnywhere, Category = "Audio")
-	USoundConcurrency* Concurrency;
+	TObjectPtr<USoundConcurrency> Concurrency;
 	
 	/** A set of parameter names that can be referenced via index when setting sound cue parameters on persistent audio */
 	UPROPERTY(EditAnywhere, Category = "Parameters")
diff --git a/Engine/Plugins/FX/Niagara/Source/Niagara/Classes/NiagaraDataInterfaceCurveBase.h b/Engine/Plugins/FX/Niagara/Source/Niagara/Classes/NiagaraDataInterfaceCurveBase.h
--- a/Engine/Plugins/FX/Niagara/Source/Niagara/Classes/NiagaraDataInterfaceCurveBase.h
+++ b/Engine/Plugins/FX/Niagara/Source/Niagara/Classes/NiagaraDataInterfaceCurveBase.h
@@ -107,7 +107,7 @@ public:
 
 	/** The texture generated and exposed to materials, will be nullptr if we do not expose to the renderers. */
 	UPROPERTY()
-	class UTexture2D* ExposedTexture;
+	TObjectPtr<class UTexture2D> ExposedTexture;
 
 #if WITH_EDITOR	
 	/** Refreshes and returns the errors detected with the corresponding data, if any.*/
diff --git a/Engine/Plugins/FX/Niagara/Source/Niagara/Classes/NiagaraDataInterfaceGrid2DCollection.h b/Engine/Plugins/FX/Niagara/Source/Niagara/Classes/NiagaraDataInterfaceGrid2DCollection.h
--- a/Engine/Plugins/FX/Niagara/Source/Niagara/Classes/NiagaraDataInterfaceGrid2DCollection.h
+++ b/Engine/Plugins/FX/Niagara/Source/Niagara/Classes/NiagaraDataInterfaceGrid2DCollection.h
@@ -290,6 +290,6 @@ protected:
 	TMap<FNiagaraSystemInstanceID, FGrid2DCollectionRWInstanceData_GameThread*> SystemInstancesToProxyData_GT;
 
 	UPROPERTY(Transient)
-	TMap< uint64, UTextureRenderTarget2DArray*> ManagedRenderTargets;
+	TMap< uint64, TObjectPtr<UTextureRenderTarget2DArray>> ManagedRenderTargets;
 	
 };
diff --git a/Engine/Plugins/FX/Niagara/Source/Niagara/Classes/NiagaraDataInterfaceLandscape.h b/Engine/Plugins/FX/Niagara/Source/Niagara/Classes/NiagaraDataInterfaceLandscape.h
--- a/Engine/Plugins/FX/Niagara/Source/Niagara/Classes/NiagaraDataInterfaceLandscape.h
+++ b/Engine/Plugins/FX/Niagara/Source/Niagara/Classes/NiagaraDataInterfaceLandscape.h
@@ -38,7 +38,7 @@ public:
 	DECLARE_NIAGARA_DI_PARAMETER();
 
 	UPROPERTY(EditAnywhere, Category = "Landscape")
-	AActor* SourceLandscape;
+	TObjectPtr<AActor> SourceLandscape;
 
 	//UObject Interface
 	virtual void PostInitProperties() override;	
diff --git a/Engine/Plugins/FX/Niagara/Source/Niagara/Classes/NiagaraDataInterfaceMeshRendererInfo.h b/Engine/Plugins/FX/Niagara/Source/Niagara/Classes/NiagaraDataInterfaceMeshRendererInfo.h
--- a/Engine/Plugins/FX/Niagara/Source/Niagara/Classes/NiagaraDataInterfaceMeshRendererInfo.h
+++ b/Engine/Plugins/FX/Niagara/Source/Niagara/Classes/NiagaraDataInterfaceMeshRendererInfo.h
@@ -58,7 +58,7 @@ protected:
 
 	/** The name of the mesh renderer */
 	UPROPERTY(EditAnywhere, Category = "Source")
-	UNiagaraMeshRendererProperties* MeshRenderer;
+	TObjectPtr<UNiagaraMeshRendererProperties> MeshRenderer;
 
 	FNDIMeshRendererInfoPtr Info;
 };
diff --git a/Engine/Plugins/FX/Niagara/Source/Niagara/Classes/NiagaraDataInterfaceRenderTarget2D.h b/Engine/Plugins/FX/Niagara/Source/Niagara/Classes/NiagaraDataInterfaceRenderTarget2D.h
--- a/Engine/Plugins/FX/Niagara/Source/Niagara/Classes/NiagaraDataInterfaceRenderTarget2D.h
+++ b/Engine/Plugins/FX/Niagara/Source/Niagara/Classes/NiagaraDataInterfaceRenderTarget2D.h
@@ -149,5 +149,5 @@ protected:
 	static FNiagaraVariableBase ExposedRTVar;
 	
 	UPROPERTY(Transient, DuplicateTransient)
-	TMap<uint64, UTextureRenderTarget2D*> ManagedRenderTargets;
+	TMap<uint64, TObjectPtr<UTextureRenderTarget2D>> ManagedRenderTargets;
 };
diff --git a/Engine/Plugins/FX/Niagara/Source/Niagara/Classes/NiagaraDataInterfaceRenderTarget2DArray.h b/Engine/Plugins/FX/Niagara/Source/Niagara/Classes/NiagaraDataInterfaceRenderTarget2DArray.h
--- a/Engine/Plugins/FX/Niagara/Source/Niagara/Classes/NiagaraDataInterfaceRenderTarget2DArray.h
+++ b/Engine/Plugins/FX/Niagara/Source/Niagara/Classes/NiagaraDataInterfaceRenderTarget2DArray.h
@@ -138,5 +138,5 @@ protected:
 	static FNiagaraVariableBase ExposedRTVar;
 	
 	UPROPERTY(Transient, DuplicateTransient)
-	TMap<uint64, UTextureRenderTarget2DArray*> ManagedRenderTargets;
+	TMap<uint64, TObjectPtr<UTextureRenderTarget2DArray>> ManagedRenderTargets;
 };
diff --git a/Engine/Plugins/FX/Niagara/Source/Niagara/Classes/NiagaraDataInterfaceRenderTargetVolume.h b/Engine/Plugins/FX/Niagara/Source/Niagara/Classes/NiagaraDataInterfaceRenderTargetVolume.h
--- a/Engine/Plugins/FX/Niagara/Source/Niagara/Classes/NiagaraDataInterfaceRenderTargetVolume.h
+++ b/Engine/Plugins/FX/Niagara/Source/Niagara/Classes/NiagaraDataInterfaceRenderTargetVolume.h
@@ -132,5 +132,5 @@ protected:
 	static FNiagaraVariableBase ExposedRTVar;
 	
 	UPROPERTY(Transient, DuplicateTransient)
-	TMap<uint64, UTextureRenderTargetVolume*> ManagedRenderTargets;
+	TMap<uint64, TObjectPtr<UTextureRenderTargetVolume>> ManagedRenderTargets;
 };
diff --git a/Engine/Plugins/FX/Niagara/Source/Niagara/Classes/NiagaraDataInterfaceSkeletalMesh.h b/Engine/Plugins/FX/Niagara/Source/Niagara/Classes/NiagaraDataInterfaceSkeletalMesh.h
--- a/Engine/Plugins/FX/Niagara/Source/Niagara/Classes/NiagaraDataInterfaceSkeletalMesh.h
+++ b/Engine/Plugins/FX/Niagara/Source/Niagara/Classes/NiagaraDataInterfaceSkeletalMesh.h
@@ -669,7 +669,7 @@ public:
 	
 	/** The source actor from which to sample. Takes precedence over the direct mesh. Note that this can only be set when used as a user variable on a component in the world.*/
 	UPROPERTY(EditAnywhere, Category = "Mesh")
-	AActor* Source;
+	TObjectPtr<AActor> Source;
 
 	/** Reference to a user parameter if we're reading one. */
 	UPROPERTY(EditAnywhere, Category = "Mesh")
@@ -677,7 +677,7 @@ public:
 	
 	/** The source component from which to sample. Takes precedence over the direct mesh. Not exposed to the user, only indirectly accessible from blueprints. */
 	UPROPERTY(Transient)
-	USkeletalMeshComponent* SourceComponent;
+	TObjectPtr<USkeletalMeshComponent> SourceComponent;
 
 	/** Selects which skinning mode to use, for most cases Skin On The Fly will cover your requirements, see individual tooltips for more information. */
 	UPROPERTY(EditAnywhere, Category="Mesh")
diff --git a/Engine/Plugins/FX/Niagara/Source/Niagara/Classes/NiagaraDataInterfaceSpline.h b/Engine/Plugins/FX/Niagara/Source/Niagara/Classes/NiagaraDataInterfaceSpline.h
--- a/Engine/Plugins/FX/Niagara/Source/Niagara/Classes/NiagaraDataInterfaceSpline.h
+++ b/Engine/Plugins/FX/Niagara/Source/Niagara/Classes/NiagaraDataInterfaceSpline.h
@@ -54,7 +54,7 @@ public:
 	
 	/** The source actor from which to sample.  Note that this can only be set when used as a user variable on a component in the world.*/
 	UPROPERTY(EditAnywhere, Category = "Spline")
-	AActor* Source;
+	TObjectPtr<AActor> Source;
 	
 	//UObject Interface
 	virtual void PostInitProperties()override;
diff --git a/Engine/Plugins/FX/Niagara/Source/Niagara/Classes/NiagaraDataInterfaceStaticMesh.h b/Engine/Plugins/FX/Niagara/Source/Niagara/Classes/NiagaraDataInterfaceStaticMesh.h
--- a/Engine/Plugins/FX/Niagara/Source/Niagara/Classes/NiagaraDataInterfaceStaticMesh.h
+++ b/Engine/Plugins/FX/Niagara/Source/Niagara/Classes/NiagaraDataInterfaceStaticMesh.h
@@ -262,15 +262,15 @@ public:
 
 	/** Mesh used to sample from when not overridden by a source actor from the scene. This mesh is NOT removed from cooked builds. */
 	UPROPERTY(EditAnywhere, Category = "Mesh")
-	UStaticMesh* DefaultMesh;
+	TObjectPtr<UStaticMesh> DefaultMesh;
 
 	/** The source actor from which to sample. Takes precedence over the direct mesh. Note that this can only be set when used as a user variable on a component in the world. */
 	UPROPERTY(EditAnywhere, Category = "Mesh")
-	AActor* Source;
+	TObjectPtr<AActor> Source;
 	
 	/** The source component from which to sample. Takes precedence over the direct mesh. Not exposed to the user, only indirectly accessible from blueprints. */
 	UPROPERTY(Transient)
-	UStaticMeshComponent* SourceComponent;
+	TObjectPtr<UStaticMeshComponent> SourceComponent;
 
 	/** Array of filters the can be used to limit sampling to certain sections of the mesh. */
 	UPROPERTY(EditAnywhere, Category = "Mesh")
diff --git a/Engine/Plugins/FX/Niagara/Source/Niagara/Classes/NiagaraDataInterfaceTexture.h b/Engine/Plugins/FX/Niagara/Source/Niagara/Classes/NiagaraDataInterfaceTexture.h
--- a/Engine/Plugins/FX/Niagara/Source/Niagara/Classes/NiagaraDataInterfaceTexture.h
+++ b/Engine/Plugins/FX/Niagara/Source/Niagara/Classes/NiagaraDataInterfaceTexture.h
@@ -17,7 +17,7 @@ public:
 	DECLARE_NIAGARA_DI_PARAMETER();
 
 	UPROPERTY(EditAnywhere, Category = "Texture")
-	UTexture* Texture;
+	TObjectPtr<UTexture> Texture;
 
 	//UObject Interface
 	virtual void PostInitProperties()override;
diff --git a/Engine/Plugins/FX/Niagara/Source/Niagara/Classes/NiagaraDataInterfaceVectorField.h b/Engine/Plugins/FX/Niagara/Source/Niagara/Classes/NiagaraDataInterfaceVectorField.h
--- a/Engine/Plugins/FX/Niagara/Source/Niagara/Classes/NiagaraDataInterfaceVectorField.h
+++ b/Engine/Plugins/FX/Niagara/Source/Niagara/Classes/NiagaraDataInterfaceVectorField.h
@@ -15,7 +15,7 @@ class NIAGARA_API UNiagaraDataInterfaceVectorField : public UNiagaraDataInterfac
 public:
 	/** Vector field to sample from. */
 	UPROPERTY(EditAnywhere, Category = VectorField)
-	UVectorField* Field;
+	TObjectPtr<UVectorField> Field;
 
 	UPROPERTY(EditAnywhere, Category = VectorField)
 	bool bTileX;
diff --git a/Engine/Plugins/FX/Niagara/Source/Niagara/Classes/NiagaraDataInterfaceVolumeTexture.h b/Engine/Plugins/FX/Niagara/Source/Niagara/Classes/NiagaraDataInterfaceVolumeTexture.h
--- a/Engine/Plugins/FX/Niagara/Source/Niagara/Classes/NiagaraDataInterfaceVolumeTexture.h
+++ b/Engine/Plugins/FX/Niagara/Source/Niagara/Classes/NiagaraDataInterfaceVolumeTexture.h
@@ -19,7 +19,7 @@ public:
 	DECLARE_NIAGARA_DI_PARAMETER();
 
 	UPROPERTY(EditAnywhere, Category = "Texture")
-	UVolumeTexture* Texture;
+	TObjectPtr<UVolumeTexture> Texture;
 
 	//UObject Interface
 	virtual void PostInitProperties()override;
diff --git a/Engine/Plugins/FX/Niagara/Source/Niagara/Classes/NiagaraEffectType.h b/Engine/Plugins/FX/Niagara/Source/Niagara/Classes/NiagaraEffectType.h
--- a/Engine/Plugins/FX/Niagara/Source/Niagara/Classes/NiagaraEffectType.h
+++ b/Engine/Plugins/FX/Niagara/Source/Niagara/Classes/NiagaraEffectType.h
@@ -263,7 +263,7 @@ class NIAGARA_API UNiagaraEffectType : public UObject
 
 	/** Used to determine the relative significance of FX in the scene which is used in other scalability systems such as instance count culling. */
 	UPROPERTY(EditAnywhere, Instanced, Category = "Scalability")
-	UNiagaraSignificanceHandler* SignificanceHandler;
+	TObjectPtr<UNiagaraSignificanceHandler> SignificanceHandler;
 
 	/** Cull settings to use at each detail level. */
 	UPROPERTY()
@@ -333,7 +333,7 @@ private:
 
 	/** Controls generation of performance baseline data for this effect type. */
 	UPROPERTY(EditAnywhere, Category="Performance", Instanced)
-	UNiagaraBaselineController* PerformanceBaselineController;
+	TObjectPtr<UNiagaraBaselineController> PerformanceBaselineController;
 
 	/**
 	Performance data gathered from the Baseline System. 
diff --git a/Engine/Plugins/FX/Niagara/Source/Niagara/Classes/NiagaraEmitter.h b/Engine/Plugins/FX/Niagara/Source/Niagara/Classes/NiagaraEmitter.h
--- a/Engine/Plugins/FX/Niagara/Source/Niagara/Classes/NiagaraEmitter.h
+++ b/Engine/Plugins/FX/Niagara/Source/Niagara/Classes/NiagaraEmitter.h
@@ -120,7 +120,7 @@ struct FNiagaraEmitterScriptProperties
 	GENERATED_BODY()
 	
 	UPROPERTY()
-	UNiagaraScript *Script;
+	TObjectPtr<UNiagaraScript> Script;
 
 	UPROPERTY()
 	TArray<FNiagaraEventReceiverProperties> EventReceivers;
@@ -389,7 +389,7 @@ public:
 #if WITH_EDITORONLY_DATA
 	/** 'Source' data/graphs for the scripts used by this emitter. */
 	UPROPERTY()
-	class UNiagaraScriptSourceBase*	GraphSource;
+	TObjectPtr<class UNiagaraScriptSourceBase>	GraphSource;
 
 	/** Should we enable rapid iteration removal if the system is also set to remove rapid iteration parameters on compile? This value defaults to true.*/
 	UPROPERTY(EditAnywhere, AdvancedDisplay, Category = "Emitter")
@@ -411,7 +411,7 @@ public:
 
 	/** Internal: The thumbnail image.*/
 	UPROPERTY()
-	class UTexture2D* ThumbnailImage;
+	TObjectPtr<class UTexture2D> ThumbnailImage;
 
 	/** Internal: Indicates the thumbnail image is out of date.*/
 	UPROPERTY()
@@ -432,10 +432,10 @@ public:
 	FText Category;
 
 	UPROPERTY()
-	TArray<UNiagaraScript*> ScratchPadScripts;
+	TArray<TObjectPtr<UNiagaraScript>> ScratchPadScripts;
 
 	UPROPERTY()
-	TArray<UNiagaraScript*> ParentScratchPadScripts;
+	TArray<TObjectPtr<UNiagaraScript>> ParentScratchPadScripts;
 	
 	/** Callback issued whenever a VM compilation successfully happened (even if the results are a script that cannot be executed due to errors)*/
 	NIAGARA_API FOnEmitterCompiled& OnEmitterVMCompiled();
@@ -594,7 +594,7 @@ private:
 
 	/** Data used by the editor to maintain UI state etc.. */
 	UPROPERTY()
-	UNiagaraEditorDataBase* EditorData;
+	TObjectPtr<UNiagaraEditorDataBase> EditorData;
 
 	/** A multicast delegate which is called whenever all the scripts for this emitter have been compiled (successfully or not). */
 	FOnEmitterCompiled OnVMScriptCompiledDelegate;
@@ -613,26 +613,26 @@ private:
 	FString UniqueEmitterName;
 
 	UPROPERTY()
-	TArray<UNiagaraRendererProperties*> RendererProperties;
+	TArray<TObjectPtr<UNiagaraRendererProperties>> RendererProperties;
 
 	UPROPERTY(EditAnywhere, Category = "Events", meta=(NiagaraNoMerge))
 	TArray<FNiagaraEventScriptProperties> EventHandlerScriptProps;
 
 	UPROPERTY(meta = (NiagaraNoMerge))
-	TArray<UNiagaraSimulationStageBase*> SimulationStages;
+	TArray<TObjectPtr<UNiagaraSimulationStageBase>> SimulationStages;
 
 	UPROPERTY()
-	UNiagaraScript* GPUComputeScript;
+	TObjectPtr<UNiagaraScript> GPUComputeScript;
 
 	UPROPERTY()
 	TArray<FName> SharedEventGeneratorIds;
 
 #if WITH_EDITORONLY_DATA
 	UPROPERTY()
-	UNiagaraEmitter* Parent;
+	TObjectPtr<UNiagaraEmitter> Parent;
 
 	UPROPERTY()
-	UNiagaraEmitter* ParentAtLastMerge;
+	TObjectPtr<UNiagaraEmitter> ParentAtLastMerge;
 #endif
 
 #if WITH_EDITOR
@@ -668,7 +668,7 @@ private:
 #if WITH_EDITORONLY_DATA
 	/** Messages associated with the Emitter asset. */
 	UPROPERTY()
-	TMap<FGuid, UNiagaraMessageDataBase*> MessageKeyToMessageMap;
+	TMap<FGuid, TObjectPtr<UNiagaraMessageDataBase>> MessageKeyToMessageMap;
 #endif
 };
 
diff --git a/Engine/Plugins/FX/Niagara/Source/Niagara/Classes/NiagaraEmitterHandle.h b/Engine/Plugins/FX/Niagara/Source/Niagara/Classes/NiagaraEmitterHandle.h
--- a/Engine/Plugins/FX/Niagara/Source/Niagara/Classes/NiagaraEmitterHandle.h
+++ b/Engine/Plugins/FX/Niagara/Source/Niagara/Classes/NiagaraEmitterHandle.h
@@ -99,16 +99,16 @@ private:
 #if WITH_EDITORONLY_DATA
 	/** The source emitter this emitter handle was built from. */
 	UPROPERTY()
-	UNiagaraEmitter* Source_DEPRECATED;
+	TObjectPtr<UNiagaraEmitter> Source_DEPRECATED;
 
 	/** An unmodified copy of the emitter this handle references for use when merging change from the source emitter. */
 	UPROPERTY()
-	UNiagaraEmitter* LastMergedSource_DEPRECATED;
+	TObjectPtr<UNiagaraEmitter> LastMergedSource_DEPRECATED;
 
 	bool bIsolated;
 #endif
 
 	/** The copied instance of the emitter this handle references. */
 	UPROPERTY()
-	UNiagaraEmitter* Instance;
+	TObjectPtr<UNiagaraEmitter> Instance;
 };
diff --git a/Engine/Plugins/FX/Niagara/Source/Niagara/Classes/NiagaraParameterCollection.h b/Engine/Plugins/FX/Niagara/Source/Niagara/Classes/NiagaraParameterCollection.h
--- a/Engine/Plugins/FX/Niagara/Source/Niagara/Classes/NiagaraParameterCollection.h
+++ b/Engine/Plugins/FX/Niagara/Source/Niagara/Classes/NiagaraParameterCollection.h
@@ -39,7 +39,7 @@ public:
 
 	//TODO: Abstract to some interface to allow a hierarchy like UMaterialInstance?
 	UPROPERTY(EditAnywhere, Category=Instance)
-	UNiagaraParameterCollection* Collection;
+	TObjectPtr<UNiagaraParameterCollection> Collection;
 
 	/**
 	When editing instances, we must track which parameters are overridden so we can pull in any changes to the default.
@@ -200,10 +200,10 @@ protected:
 
 	/** Optional set of MPC that can drive scalar and vector parameters */
 	UPROPERTY(EditAnywhere, Category = "Parameter Collection")
-	UMaterialParameterCollection* SourceMaterialCollection;
+	TObjectPtr<UMaterialParameterCollection> SourceMaterialCollection;
 	
 	UPROPERTY()
-	UNiagaraParameterCollectionInstance* DefaultInstance;
+	TObjectPtr<UNiagaraParameterCollectionInstance> DefaultInstance;
 
 	/** Used to track whenever something of note changes in this parameter collection that might invalidate a compilation downstream of a script/emitter/system.*/
 	UPROPERTY()
diff --git a/Engine/Plugins/FX/Niagara/Source/Niagara/Classes/NiagaraPerfBaseline.h b/Engine/Plugins/FX/Niagara/Source/Niagara/Classes/NiagaraPerfBaseline.h
--- a/Engine/Plugins/FX/Niagara/Source/Niagara/Classes/NiagaraPerfBaseline.h
+++ b/Engine/Plugins/FX/Niagara/Source/Niagara/Classes/NiagaraPerfBaseline.h
@@ -102,11 +102,11 @@ public:
 
 	/** The effect type this controller is in use by. */
 	UPROPERTY(Category = "Baseline", BlueprintReadOnly)
-	UNiagaraEffectType* EffectType = nullptr;
+	TObjectPtr<UNiagaraEffectType> EffectType = nullptr;
 
 	/** The owning actor for this baseline controller. */
 	UPROPERTY(Category = "Baseline", BlueprintReadOnly)
-	ANiagaraPerfBaselineActor* Owner;
+	TObjectPtr<ANiagaraPerfBaselineActor> Owner;
 
 private:
 	/** The baseline system to spawn. */
@@ -132,7 +132,7 @@ class NIAGARA_API UNiagaraBaselineController_Basic : public UNiagaraBaselineCont
 	int32 NumInstances = 1;
 
 	UPROPERTY()
-	TArray<UNiagaraComponent*> SpawnedComponents;
+	TArray<TObjectPtr<UNiagaraComponent>> SpawnedComponents;
 };
 
 /** Actor that controls how the baseline system behaves and also controls stats gathering for. */
@@ -145,10 +145,10 @@ class NIAGARA_API ANiagaraPerfBaselineActor : public AActor
 public:
 
 	UPROPERTY(EditAnywhere, Category="Baseline")
-	UNiagaraBaselineController* Controller;
+	TObjectPtr<UNiagaraBaselineController> Controller;
 
 	UPROPERTY(EditAnywhere, Category="Baseline")
-	UTextRenderComponent* Label;
+	TObjectPtr<UTextRenderComponent> Label;
 
 #if NIAGARA_PERF_BASELINES
 
diff --git a/Engine/Plugins/FX/Niagara/Source/Niagara/Classes/NiagaraScript.h b/Engine/Plugins/FX/Niagara/Source/Niagara/Classes/NiagaraScript.h
--- a/Engine/Plugins/FX/Niagara/Source/Niagara/Classes/NiagaraScript.h
+++ b/Engine/Plugins/FX/Niagara/Source/Niagara/Classes/NiagaraScript.h
@@ -428,7 +428,7 @@ public:
 
 	/* Which script to use if this is deprecated.*/
 	UPROPERTY(EditAnywhere, Category = Script, meta = (EditCondition = "bDeprecated"))
-	UNiagaraScript* DeprecationRecommendation;
+	TObjectPtr<UNiagaraScript> DeprecationRecommendation;
 
 	/* Custom logic to convert the contents of an existing script assignment to this script.*/
 	UPROPERTY(EditAnywhere, Category = Script)
@@ -755,7 +755,7 @@ private:
 
 	/** 'Source' data/graphs for this script */
 	UPROPERTY()
-	class UNiagaraScriptSourceBase*	Source;
+	TObjectPtr<class UNiagaraScriptSourceBase>	Source;
 	
 	/** A multicast delegate which is called whenever the script has been compiled (successfully or not). */
 	FOnScriptCompiled OnVMScriptCompiledDelegate;
@@ -799,7 +799,7 @@ private:
 	TMap<const class ITargetPlatform*, TArray<FNiagaraShaderScript*>> CachedScriptResourcesForCooking;
 
 	UPROPERTY(Transient)
-	TArray<UObject*> ActiveCompileRoots;
+	TArray<TObjectPtr<UObject>> ActiveCompileRoots;
 
 	/* Flag set on load based on whether the serialized data includes editor only data */
 	bool IsCooked;
@@ -810,7 +810,7 @@ private:
 	FNiagaraVMExecutableData CachedScriptVM;
 	
 	UPROPERTY()
-	TArray<UNiagaraParameterCollection*> CachedParameterCollectionReferences;
+	TArray<TObjectPtr<UNiagaraParameterCollection>> CachedParameterCollectionReferences;
 
 	UPROPERTY()
 	TArray<FNiagaraScriptDataInterfaceInfo> CachedDefaultDataInterfaces;
diff --git a/Engine/Plugins/FX/Niagara/Source/Niagara/Classes/NiagaraSystem.h b/Engine/Plugins/FX/Niagara/Source/Niagara/Classes/NiagaraSystem.h
--- a/Engine/Plugins/FX/Niagara/Source/Niagara/Classes/NiagaraSystem.h
+++ b/Engine/Plugins/FX/Niagara/Source/Niagara/Classes/NiagaraSystem.h
@@ -165,7 +165,7 @@ struct FNiagaraSystemCompileRequest
 	double StartTime = 0.0;
 
 	UPROPERTY()
-	TArray<UObject*> RootObjects;
+	TArray<TObjectPtr<UObject>> RootObjects;
 
 	TArray<FEmitterCompiledScriptPair> EmitterCompiledScriptPairs;
 	
@@ -344,7 +344,7 @@ public:
 
 	/** Internal: The thumbnail image.*/
 	UPROPERTY()
-	class UTexture2D* ThumbnailImage;
+	TObjectPtr<class UTexture2D> ThumbnailImage;
 
 	/** Internal: Indicates the thumbnail image is out of date.*/
 	UPROPERTY()
@@ -361,7 +361,7 @@ public:
 	FText TemplateAssetDescription;
 
 	UPROPERTY()
-	TArray<UNiagaraScript*> ScratchPadScripts;
+	TArray<TObjectPtr<UNiagaraScript>> ScratchPadScripts;
 
 	UPROPERTY(transient)
 	FNiagaraParameterStore EditorOnlyAddedParameters;
@@ -544,7 +544,7 @@ private:
 
 protected:
 	UPROPERTY(EditAnywhere, Category = "System")
-	UNiagaraEffectType* EffectType;
+	TObjectPtr<UNiagaraEffectType> EffectType;
 
 	UPROPERTY(EditAnywhere, Category = "Scalability")
 	bool bOverrideScalabilitySettings;
@@ -560,7 +560,7 @@ protected:
 	TArray<FNiagaraEmitterHandle> EmitterHandles;
 
 	UPROPERTY(EditAnywhere, Category="System")
-	TArray<UNiagaraParameterCollectionInstance*> ParameterCollectionOverrides;
+	TArray<TObjectPtr<UNiagaraParameterCollectionInstance>> ParameterCollectionOverrides;
 
 #if WITH_EDITORONLY_DATA
 	UPROPERTY(Transient)
@@ -574,12 +574,12 @@ protected:
 	/** The script which defines the System parameters, and which generates the bindings from System
 		parameter to emitter parameter. */
 	UPROPERTY()
-	UNiagaraScript* SystemSpawnScript;
+	TObjectPtr<UNiagaraScript> SystemSpawnScript;
 
 	/** The script which defines the System parameters, and which generates the bindings from System
 	parameter to emitter parameter. */
 	UPROPERTY()
-	UNiagaraScript* SystemUpdateScript;
+	TObjectPtr<UNiagaraScript> SystemUpdateScript;
 
 	//** Post compile generated data used for initializing Emitter Instances during runtime. */
 	TArray<TSharedRef<const FNiagaraEmitterCompiledData>> EmitterCompiledData;
@@ -595,7 +595,7 @@ protected:
 #if WITH_EDITORONLY_DATA
 	/** Data used by the editor to maintain UI state etc.. */
 	UPROPERTY()
-	UNiagaraEditorDataBase* EditorData;
+	TObjectPtr<UNiagaraEditorDataBase> EditorData;
 
 	bool bIsolateEnabled;
 
@@ -675,7 +675,7 @@ protected:
 #if WITH_EDITORONLY_DATA
 	/** Messages associated with the System asset. */
 	UPROPERTY()
-	TMap<FGuid, UNiagaraMessageDataBase*> MessageKeyToMessageMap;
+	TMap<FGuid, TObjectPtr<UNiagaraMessageDataBase>> MessageKeyToMessageMap;
 
 	FGuid AssetGuid;
 #endif
diff --git a/Engine/Plugins/FX/Niagara/Source/Niagara/Private/NiagaraPrecompileContainer.h b/Engine/Plugins/FX/Niagara/Source/Niagara/Private/NiagaraPrecompileContainer.h
--- a/Engine/Plugins/FX/Niagara/Source/Niagara/Private/NiagaraPrecompileContainer.h
+++ b/Engine/Plugins/FX/Niagara/Source/Niagara/Private/NiagaraPrecompileContainer.h
@@ -13,8 +13,8 @@ class NIAGARA_API UNiagaraPrecompileContainer : public UObject
 
 public:
 	UPROPERTY()
-	TArray<UNiagaraScript*> Scripts;
+	TArray<TObjectPtr<UNiagaraScript>> Scripts;
 
 	UPROPERTY()
-	UNiagaraSystem* System;
+	TObjectPtr<UNiagaraSystem> System;
 };
\ No newline at end of file
diff --git a/Engine/Plugins/FX/Niagara/Source/Niagara/Public/MovieScene/MovieSceneNiagaraTrack.h b/Engine/Plugins/FX/Niagara/Source/Niagara/Public/MovieScene/MovieSceneNiagaraTrack.h
--- a/Engine/Plugins/FX/Niagara/Source/Niagara/Public/MovieScene/MovieSceneNiagaraTrack.h
+++ b/Engine/Plugins/FX/Niagara/Source/Niagara/Public/MovieScene/MovieSceneNiagaraTrack.h
@@ -21,5 +21,5 @@ public:
 
 protected:
 	UPROPERTY()
-	TArray<UMovieSceneSection*> Sections;
+	TArray<TObjectPtr<UMovieSceneSection>> Sections;
 };
\ No newline at end of file
diff --git a/Engine/Plugins/FX/Niagara/Source/Niagara/Public/NiagaraActor.h b/Engine/Plugins/FX/Niagara/Source/Niagara/Public/NiagaraActor.h
--- a/Engine/Plugins/FX/Niagara/Source/Niagara/Public/NiagaraActor.h
+++ b/Engine/Plugins/FX/Niagara/Source/Niagara/Public/NiagaraActor.h
@@ -23,16 +23,16 @@ public:
 private:
 	/** Pointer to System component */
 	UPROPERTY(VisibleAnywhere, BlueprintReadOnly, Category=NiagaraActor, meta = (AllowPrivateAccess = "true"))
-	class UNiagaraComponent* NiagaraComponent;
+	TObjectPtr<class UNiagaraComponent> NiagaraComponent;
 
 #if WITH_EDITORONLY_DATA
 	// Reference to sprite visualization component
 	UPROPERTY()
-	class UBillboardComponent* SpriteComponent;
+	TObjectPtr<class UBillboardComponent> SpriteComponent;
 
 	// Reference to arrow visualization component
 	UPROPERTY()
-	class UArrowComponent* ArrowComponent;
+	TObjectPtr<class UArrowComponent> ArrowComponent;
 
 #endif
 
diff --git a/Engine/Plugins/FX/Niagara/Source/Niagara/Public/NiagaraCommon.h b/Engine/Plugins/FX/Niagara/Source/Niagara/Public/NiagaraCommon.h
--- a/Engine/Plugins/FX/Niagara/Source/Niagara/Public/NiagaraCommon.h
+++ b/Engine/Plugins/FX/Niagara/Source/Niagara/Public/NiagaraCommon.h
@@ -480,7 +480,7 @@ public:
 	}
 
 	UPROPERTY()
-	class UNiagaraDataInterface* DataInterface;
+	TObjectPtr<class UNiagaraDataInterface> DataInterface;
 	
 	UPROPERTY()
 	FName Name;
@@ -671,13 +671,13 @@ private:
 	void AddInternal(class UNiagaraComponent* Comp, bool bReInit);
 
 	UPROPERTY(transient)
-	TArray<UNiagaraComponent*> ComponentsToReset;
+	TArray<TObjectPtr<UNiagaraComponent>> ComponentsToReset;
 	UPROPERTY(transient)
-	TArray<UNiagaraComponent*> ComponentsToReInit;
+	TArray<TObjectPtr<UNiagaraComponent>> ComponentsToReInit;
 	UPROPERTY(transient)
-	TArray<UNiagaraComponent*> ComponentsToNotifySimDestroy;
+	TArray<TObjectPtr<UNiagaraComponent>> ComponentsToNotifySimDestroy;
 	UPROPERTY(transient)
-	TArray<UNiagaraSystem*> SystemSimsToDestroy;
+	TArray<TObjectPtr<UNiagaraSystem>> SystemSimsToDestroy;
 
 	bool bDestroyOnAdd;
 	bool bOnlyActive;
@@ -853,7 +853,7 @@ struct FNiagaraVariableInfo
 	FText Definition;
 
 	UPROPERTY()
-	UNiagaraDataInterface* DataInterface;
+	TObjectPtr<UNiagaraDataInterface> DataInterface;
 };
 
 /** This enum decides how a renderer will attempt to process the incoming data from the stack.*/
diff --git a/Engine/Plugins/FX/Niagara/Source/Niagara/Public/NiagaraComponent.h b/Engine/Plugins/FX/Niagara/Source/Niagara/Public/NiagaraComponent.h
--- a/Engine/Plugins/FX/Niagara/Source/Niagara/Public/NiagaraComponent.h
+++ b/Engine/Plugins/FX/Niagara/Source/Niagara/Public/NiagaraComponent.h
@@ -35,13 +35,13 @@ struct FNiagaraMaterialOverride
 	GENERATED_USTRUCT_BODY();
 
 	UPROPERTY()
-	class UMaterialInterface* Material = nullptr;
+	TObjectPtr<class UMaterialInterface> Material = nullptr;
 
 	UPROPERTY()
 	uint32 MaterialSubIndex;
 
 	UPROPERTY()
-	UNiagaraRendererProperties* EmitterRendererProperty = nullptr;
+	TObjectPtr<UNiagaraRendererProperties> EmitterRendererProperty = nullptr;
 };
 
 /**
@@ -79,7 +79,7 @@ public:
 
 private:
 	UPROPERTY(EditAnywhere, Category="Niagara", meta = (DisplayName = "Niagara System Asset"))
-	UNiagaraSystem* Asset;
+	TObjectPtr<UNiagaraSystem> Asset;
 
 	/** Allows you to control how Niagara selects the tick group, changing this while an instance is active will result in not change as it is cached. */
 	UPROPERTY(EditAnywhere, Category = "Niagara", meta = (DisplayName = "Niagara Tick Behavior"))
diff --git a/Engine/Plugins/FX/Niagara/Source/Niagara/Public/NiagaraComponentPool.h b/Engine/Plugins/FX/Niagara/Source/Niagara/Public/NiagaraComponentPool.h
--- a/Engine/Plugins/FX/Niagara/Source/Niagara/Public/NiagaraComponentPool.h
+++ b/Engine/Plugins/FX/Niagara/Source/Niagara/Public/NiagaraComponentPool.h
@@ -50,7 +50,7 @@ struct FNCPoolElement
 	GENERATED_BODY()
 
 	UPROPERTY(transient)
-	UNiagaraComponent* Component;
+	TObjectPtr<UNiagaraComponent> Component;
 
 	float LastUsedTime;
 
@@ -114,7 +114,7 @@ class NIAGARA_API UNiagaraComponentPool : public UObject
 
 private:
 	UPROPERTY()
-	TMap<UNiagaraSystem*, FNCPool> WorldParticleSystemPools;
+	TMap<TObjectPtr<UNiagaraSystem>, FNCPool> WorldParticleSystemPools;
 
 	float LastParticleSytemPoolCleanTime;
 public:
diff --git a/Engine/Plugins/FX/Niagara/Source/Niagara/Public/NiagaraComponentRendererProperties.h b/Engine/Plugins/FX/Niagara/Source/Niagara/Public/NiagaraComponentRendererProperties.h
--- a/Engine/Plugins/FX/Niagara/Source/Niagara/Public/NiagaraComponentRendererProperties.h
+++ b/Engine/Plugins/FX/Niagara/Source/Niagara/Public/NiagaraComponentRendererProperties.h
@@ -119,7 +119,7 @@ public:
 
 	/** The object template used to create new components at runtime. */
 	UPROPERTY(Export, Instanced, EditAnywhere, Category = "Component Properties")
-	USceneComponent* TemplateComponent;
+	TObjectPtr<USceneComponent> TemplateComponent;
 
 	UPROPERTY()
 	TArray<FNiagaraComponentPropertyBinding> PropertyBindings;
