Due to refactoring, old neighbor grids could end up trying to allocate crazy amounts of VRAM. This puts a hard cap on the behavior, allowing subsequent edits and fix-ups.

diff --git a/Engine/Plugins/FX/Niagara/Source/Niagara/Classes/NiagaraDataInterfaceNeighborGrid3D.h b/Engine/Plugins/FX/Niagara/Source/Niagara/Classes/NiagaraDataInterfaceNeighborGrid3D.h
--- a/Engine/Plugins/FX/Niagara/Source/Niagara/Classes/NiagaraDataInterfaceNeighborGrid3D.h
+++ b/Engine/Plugins/FX/Niagara/Source/Niagara/Classes/NiagaraDataInterfaceNeighborGrid3D.h
@@ -26,20 +26,7 @@ public:
 #endif
 	}
 
-	void ResizeBuffers()
-	{
-		const int32 NumTotalCells = NumCells.X * NumCells.Y * NumCells.Z;
-		const int32 NumIntsInGridBuffer = NumTotalCells * MaxNeighborsPerCell;
-
-		NeighborhoodCountBuffer.Initialize(sizeof(int32), NumTotalCells, EPixelFormat::PF_R32_SINT, BUF_Static, TEXT("NiagaraNeighborGrid3D::NeighborCount"));
-		NeighborhoodBuffer.Initialize(sizeof(int32), NumIntsInGridBuffer, EPixelFormat::PF_R32_SINT, BUF_Static, TEXT("NiagaraNeighborGrid3D::NeighborsGrid"));
-
-#if STATS
-		DEC_MEMORY_STAT_BY(STAT_NiagaraGPUDataInterfaceMemory, GPUMemory);
-		GPUMemory = NumTotalCells * sizeof(int32) + NumIntsInGridBuffer * sizeof(int32);
-		INC_MEMORY_STAT_BY(STAT_NiagaraGPUDataInterfaceMemory, GPUMemory);
-#endif
-	}
+	void ResizeBuffers();
 
 	FIntVector NumCells;
 	float CellSize;
diff --git a/Engine/Plugins/FX/Niagara/Source/Niagara/Private/NiagaraDataInterfaceNeighborGrid3D.cpp b/Engine/Plugins/FX/Niagara/Source/Niagara/Private/NiagaraDataInterfaceNeighborGrid3D.cpp
--- a/Engine/Plugins/FX/Niagara/Source/Niagara/Private/NiagaraDataInterfaceNeighborGrid3D.cpp
+++ b/Engine/Plugins/FX/Niagara/Source/Niagara/Private/NiagaraDataInterfaceNeighborGrid3D.cpp
@@ -28,6 +28,14 @@ static const FName SetParticleNeighborFunctionName("SetParticleNeighbor");
 static const FName GetParticleNeighborCountFunctionName("GetParticleNeighborCount");
 static const FName SetParticleNeighborCountFunctionName("SetParticleNeighborCount");
 
+static int32 GMaxNiagaraNeighborGridCells = (100*100*100);
+static FAutoConsoleVariableRef CVarMaxNiagaraCPUParticlesPerEmitter(
+	TEXT("fx.MaxNiagaraNeighborGridCells"),
+	GMaxNiagaraNeighborGridCells,
+	TEXT("The max number of supported grid cells in Niagara. Overflowing this threshold will cause the sim to warn and fail. \n"),
+	ECVF_Default
+);
+
 
 /*--------------------------------------------------------------------------------------------------------------------------*/
 struct FNiagaraDataInterfaceParametersCS_NeighborGrid3D : public FNiagaraDataInterfaceParametersCS
@@ -165,6 +173,25 @@ private:
 	LAYOUT_FIELD(FRWShaderParameter, OutputParticleNeighborsGridParam);
 };
 
+void NeighborGrid3DRWInstanceData::ResizeBuffers()
+{
+	uint32 NumTotalCells = NumCells.X * NumCells.Y * NumCells.Z;
+
+	uint32 NumIntsInGridBuffer = NumTotalCells * MaxNeighborsPerCell;
+
+	if (NumTotalCells > (uint32)GMaxNiagaraNeighborGridCells)
+		return;
+
+	NeighborhoodCountBuffer.Initialize(sizeof(int32), NumTotalCells, EPixelFormat::PF_R32_SINT, BUF_Static, TEXT("NiagaraNeighborGrid3D::NeighborCount"));
+	NeighborhoodBuffer.Initialize(sizeof(int32), NumIntsInGridBuffer, EPixelFormat::PF_R32_SINT, BUF_Static, TEXT("NiagaraNeighborGrid3D::NeighborsGrid"));
+
+	#if STATS
+		DEC_MEMORY_STAT_BY(STAT_NiagaraGPUDataInterfaceMemory, GPUMemory);
+		GPUMemory = NumTotalCells * sizeof(int32) + NumIntsInGridBuffer * sizeof(int32);
+		INC_MEMORY_STAT_BY(STAT_NiagaraGPUDataInterfaceMemory, GPUMemory);
+	#endif
+}
+
 IMPLEMENT_TYPE_LAYOUT(FNiagaraDataInterfaceParametersCS_NeighborGrid3D);
 
 IMPLEMENT_NIAGARA_DI_PARAMETER(UNiagaraDataInterfaceNeighborGrid3D, FNiagaraDataInterfaceParametersCS_NeighborGrid3D);
@@ -596,6 +623,12 @@ bool UNiagaraDataInterfaceNeighborGrid3D::InitPerInstanceData(void* PerInstanceD
 	InstanceData->MaxNeighborsPerCell = RT_MaxNeighborsPerCell;	
 	InstanceData->NumCells = RT_NumCells;
 
+	if ((RT_NumCells.X * RT_NumCells.Y * RT_NumCells.Z) > GMaxNiagaraNeighborGridCells)
+	{
+		UE_LOG(LogNiagara, Warning, TEXT("Dimensions are too big! Please adjust! %d x %d x %d > %d for ==> %s"), RT_NumCells.X, RT_NumCells.Y, RT_NumCells.Z, GMaxNiagaraNeighborGridCells , *GetFullNameSafe(this))
+			return false;
+	}
+
 	// @todo-threadsafety. This would be a race but I'm taking a ref here. Not ideal in the long term.
 	// Push Updates to Proxy.
 	ENQUEUE_RENDER_COMMAND(FUpdateData)(
