- Fixed a Non-unity error

diff --git a/Engine/Plugins/FX/Niagara/Source/Niagara/Private/NiagaraDataInterfaceNeighborGrid3D.cpp b/Engine/Plugins/FX/Niagara/Source/Niagara/Private/NiagaraDataInterfaceNeighborGrid3D.cpp
--- a/Engine/Plugins/FX/Niagara/Source/Niagara/Private/NiagaraDataInterfaceNeighborGrid3D.cpp
+++ b/Engine/Plugins/FX/Niagara/Source/Niagara/Private/NiagaraDataInterfaceNeighborGrid3D.cpp
@@ -29,7 +29,7 @@ static const FName GetParticleNeighborCountFunctionName("GetParticleNeighborCoun
 static const FName SetParticleNeighborCountFunctionName("SetParticleNeighborCount");
 
 static int32 GMaxNiagaraNeighborGridCells = (100*100*100);
-static FAutoConsoleVariableRef CVarMaxNiagaraCPUParticlesPerEmitter(
+static FAutoConsoleVariableRef CVarMaxNiagaraNeighborGridCells(
 	TEXT("fx.MaxNiagaraNeighborGridCells"),
 	GMaxNiagaraNeighborGridCells,
 	TEXT("The max number of supported grid cells in Niagara. Overflowing this threshold will cause the sim to warn and fail. \n"),
