Fix another TPair/TObjectPtr type issue, this time surfaced by MSVC static analysis.

diff --git a/Engine/Plugins/FX/Niagara/Source/NiagaraEditor/Private/NiagaraGraph.cpp b/Engine/Plugins/FX/Niagara/Source/NiagaraEditor/Private/NiagaraGraph.cpp
--- a/Engine/Plugins/FX/Niagara/Source/NiagaraEditor/Private/NiagaraGraph.cpp
+++ b/Engine/Plugins/FX/Niagara/Source/NiagaraEditor/Private/NiagaraGraph.cpp
@@ -769,7 +769,7 @@ void UNiagaraGraph::StandardizeParameterNames()
 	// repopulated.
 	TMap<FNiagaraVariable, TObjectPtr<UNiagaraScriptVariable>> OldVariableToScriptVariable = VariableToScriptVariable;
 	VariableToScriptVariable.Empty();
-	for (TPair<FNiagaraVariable, UNiagaraScriptVariable*> VariableScriptVariablePair : OldVariableToScriptVariable)
+	for (TPair<FNiagaraVariable, TObjectPtr<UNiagaraScriptVariable>> VariableScriptVariablePair : OldVariableToScriptVariable)
 	{
 		FNiagaraVariable Variable = VariableScriptVariablePair.Key;
 		UNiagaraScriptVariable* ScriptVariable = VariableScriptVariablePair.Value;
