Remove use of ByteAddressBuffers for Niagara adjacency buffer information (#ushell-cherrypick of 15283990 by rob.krajcarski)
Update Niagara skelmesh DI connectivity (#ushell-cherrypick of 15353148 by rob.krajcarski)

diff --git a/Engine/Plugins/FX/Niagara/Shaders/Private/NiagaraDataInterfaceSkeletalMesh.ush b/Engine/Plugins/FX/Niagara/Shaders/Private/NiagaraDataInterfaceSkeletalMesh.ush
--- a/Engine/Plugins/FX/Niagara/Shaders/Private/NiagaraDataInterfaceSkeletalMesh.ush
+++ b/Engine/Plugins/FX/Niagara/Shaders/Private/NiagaraDataInterfaceSkeletalMesh.ush
@@ -64,7 +64,7 @@ struct FDISkelMeshContext
 	Buffer<int>		UvMappingBuffer;
 	uint			UvMappingBufferLength;
 	uint			UvMappingSet;
-	ByteAddressBuffer	ConnectivityBuffer;
+	Buffer<uint>		ConnectivityBuffer;
 	uint			ConnectivityMaxAdjacentPerVertex;
 	uint			ConnectivityBufferLength;
 	float4x4		InstanceTransform;
@@ -994,39 +994,55 @@ void DISkelMesh_GetFilteredSocketOrBoneAt(in FDISkelMeshContext DIContext, int F
 ////////////////////////////////////////////////////////////////////////////////////////////////////////////////////////////////////////////////////////////////
 // Connectivity
 
-void DISkelMesh_GetAdjacentTriangleIndex(in FDISkelMeshContext DIContext, int VertexId, int AdjacencyIndex, out int TriangleIndex)
+void DISkelMesh_GetAdjacentTriangleIndex(in FDISkelMeshContext DIContext, int VertexId, uint AdjacencyIndex, out int TriangleIndex)
 {
-	uint BufferLength = 0;
-	DIContext.ConnectivityBuffer.GetDimensions(BufferLength);
+	TriangleIndex = -1;
 
-#if DISKELMESH_ADJ_INDEX_FORMAT == DISKELMESH_ADJ_INDEX_FORMAT_FULL
-	const uint ReadIndex = 4 * (VertexId * DIContext.ConnectivityMaxAdjacentPerVertex + AdjacencyIndex);
-	if (ReadIndex < DIContext.ConnectivityBufferLength)
+	if (AdjacencyIndex < DIContext.ConnectivityMaxAdjacentPerVertex)
 	{
-		TriangleIndex = DIContext.ConnectivityBuffer.Load(ReadIndex);
-		return;
-	}
+#if DISKELMESH_ADJ_INDEX_FORMAT == DISKELMESH_ADJ_INDEX_FORMAT_FULL
+		const uint ReadIndex = (VertexId * DIContext.ConnectivityMaxAdjacentPerVertex + AdjacencyIndex);
+		if (ReadIndex < DIContext.ConnectivityBufferLength)
+		{
+			TriangleIndex = DIContext.ConnectivityBuffer[ReadIndex];
+		}
 #elif DISKELMESH_ADJ_INDEX_FORMAT == DISKELMESH_ADJ_INDEX_FORMAT_HALF
-	const uint ReadIndex = (2 * (VertexId * DIContext.ConnectivityMaxAdjacentPerVertex + AdjacencyIndex)) & ~0x3;
-	if (ReadIndex < DIContext.ConnectivityBufferLength)
-	{
-		uint PackedIndices = DIContext.ConnectivityBuffer.Load(ReadIndex);
-		uint2 UnpackedIndices = uint2(PackedIndices >> 16, PackedIndices & 0xFFFF);
-		TriangleIndex = UnpackedIndices[AdjacencyIndex & 1];
-		return;
-	}
+		const uint ReadIndex = ((VertexId * DIContext.ConnectivityMaxAdjacentPerVertex + AdjacencyIndex) & ~0x1) >> 1;
+		if (ReadIndex < DIContext.ConnectivityBufferLength)
+		{
+			uint PackedIndices = DIContext.ConnectivityBuffer[ReadIndex];
+			uint UnpackedIndex = uint2(PackedIndices & 0xFFFF, PackedIndices >> 16)[AdjacencyIndex & 1];
+			TriangleIndex = UnpackedIndex == 0xFFFF ? -1 : int(UnpackedIndex);
+		}
 #else
-	#error Unsupported index format for adjacency information
+		#error Unsupported index format for adjacency information
 #endif
+	}
+}
 
-	TriangleIndex = -1;
+bool IsAdjacent(FDISkelMeshContext DIContext, int VertexIndex, int TriangleIndex)
+{
+	for (uint AdjIt = 0; AdjIt < DIContext.ConnectivityMaxAdjacentPerVertex; ++AdjIt)
+	{
+		int AdjTriangleIndex = -1;
+		DISkelMesh_GetAdjacentTriangleIndex(DIContext, VertexIndex, AdjIt, AdjTriangleIndex);
+		if (AdjTriangleIndex == -1)
+		{
+			break;
+		}
+		if (AdjTriangleIndex == TriangleIndex)
+		{
+			return true;
+		}
+	}
+	return false;
 }
 
 void DISkelMesh_GetTriangleNeighbor(in FDISkelMeshContext DIContext, int TriangleIndex, int EdgeIndex, out int NeighborTriangleIndex, out int NeighborEdgeIndex)
 {
 	NeighborTriangleIndex = -1;
 	NeighborEdgeIndex = -1;
-	
+
 	int3 SourceVertexIndices;
 
 	DISkelMesh_GetTriVertices(DIContext, TriangleIndex, SourceVertexIndices.x, SourceVertexIndices.y, SourceVertexIndices.z);
@@ -1034,36 +1050,46 @@ void DISkelMesh_GetTriangleNeighbor(in FDISkelMeshContext DIContext, int Triangl
 	int StartSourceIndex = SourceVertexIndices[(EdgeIndex + 1) % 3];
 	int EndSourceIndex = SourceVertexIndices[(EdgeIndex + 2) % 3];
 
-	for (uint SourceTriangleIt = 0; SourceTriangleIt < DIContext.ConnectivityMaxAdjacentPerVertex; ++SourceTriangleIt)
+	for (uint SourceAdjIt = 0; SourceAdjIt < DIContext.ConnectivityMaxAdjacentPerVertex; ++SourceAdjIt)
 	{
-		int PossibleTriangleIndex = -1;
-		DISkelMesh_GetAdjacentTriangleIndex(DIContext, StartSourceIndex, SourceTriangleIt, PossibleTriangleIndex);
-		if (PossibleTriangleIndex == TriangleIndex)
+		int SourceAdjTriangleIndex = -1;
+		DISkelMesh_GetAdjacentTriangleIndex(DIContext, StartSourceIndex, SourceAdjIt, SourceAdjTriangleIndex);
+
+		if (SourceAdjTriangleIndex == TriangleIndex) // ignore self
 		{
 			continue;
 		}
-		else if (PossibleTriangleIndex == -1)
+		if (SourceAdjTriangleIndex == -1) // stop looking
 		{
 			break;
 		}
 
-		int3 PossibleVertexIndices;
-		DISkelMesh_GetTriVertices(DIContext, PossibleTriangleIndex, PossibleVertexIndices.x, PossibleVertexIndices.y, PossibleVertexIndices.z);
+		if (IsAdjacent(DIContext, EndSourceIndex, SourceAdjTriangleIndex))
+		{
+			NeighborTriangleIndex = SourceAdjTriangleIndex;
 
-		bool3 MatchStart = StartSourceIndex.xxx == PossibleVertexIndices;
-		bool3 MatchEnd = EndSourceIndex.xxx == PossibleVertexIndices;
+			// now to find the edge index for the neighbor
+			int3 NeighborVertexIndices;
+			DISkelMesh_GetTriVertices(DIContext, NeighborTriangleIndex, NeighborVertexIndices.x, NeighborVertexIndices.y, NeighborVertexIndices.z);
 
-		if (any(MatchStart) && any(MatchEnd))
-		{
-			NeighborTriangleIndex = PossibleTriangleIndex;
-			for (int EdgeIt = 0; EdgeIt < 3; ++EdgeIt)
+			bool3 NeighborVertMatches = bool3(
+				IsAdjacent(DIContext, NeighborVertexIndices.x, TriangleIndex),
+				IsAdjacent(DIContext, NeighborVertexIndices.y, TriangleIndex),
+				IsAdjacent(DIContext, NeighborVertexIndices.z, TriangleIndex));
+
+			if (all(NeighborVertMatches.xy))
 			{
-				if (!MatchStart[EdgeIt] && !MatchEnd[EdgeIt])
-				{
-					NeighborEdgeIndex = EdgeIt;
-					break;
-				}
+				NeighborEdgeIndex = 2;
 			}
+			else if (all(NeighborVertMatches.yz))
+			{
+				NeighborEdgeIndex = 0;
+			}
+			else if (all(NeighborVertMatches.zx))
+			{
+				NeighborEdgeIndex = 1;
+			}
+			break;
 		}
 	}
 }
@@ -1106,7 +1132,7 @@ void DISkelMesh_GetTriangleNeighbor(in FDISkelMeshContext DIContext, int Triangl
 	Buffer<int>		UvMappingBuffer_ ## NAME; \
 	uint			UvMappingBufferLength_ ## NAME; \
 	uint			UvMappingSet_ ## NAME; \
-	ByteAddressBuffer		ConnectivityBuffer_ ## NAME; \
+	Buffer<uint>		ConnectivityBuffer_ ## NAME; \
 	uint			ConnectivityBufferLength_ ## NAME; \
 	uint			ConnectivityMaxAdjacentPerVertex_ ## NAME; \
 	float4x4		InstanceTransform_ ## NAME; \
diff --git a/Engine/Plugins/FX/Niagara/Source/Niagara/Private/NDISkeletalMesh_TriangleSampling.cpp b/Engine/Plugins/FX/Niagara/Source/Niagara/Private/NDISkeletalMesh_TriangleSampling.cpp
--- a/Engine/Plugins/FX/Niagara/Source/Niagara/Private/NDISkeletalMesh_TriangleSampling.cpp
+++ b/Engine/Plugins/FX/Niagara/Source/Niagara/Private/NDISkeletalMesh_TriangleSampling.cpp
@@ -289,6 +289,7 @@ void UNiagaraDataInterfaceSkeletalMesh::GetTriangleSamplingFunctions(TArray<FNia
 		Sig.Inputs.Add(FNiagaraVariable(FNiagaraTypeDefinition::GetIntDef(), TEXT("Vertex ID")));
 		Sig.Inputs.Add(FNiagaraVariable(FNiagaraTypeDefinition::GetIntDef(), TEXT("Adjacency Index")));
 		Sig.Outputs.Add(FNiagaraVariable(FNiagaraTypeDefinition::GetIntDef(), TEXT("Triangle Index")));
+		Sig.Outputs.Add(FNiagaraVariable(FNiagaraTypeDefinition::GetBoolDef(), TEXT("IsValid")));
 		Sig.bMemberFunction = true;
 		Sig.bRequiresContext = false;
 		Sig.bExperimental = true;
@@ -304,6 +305,7 @@ void UNiagaraDataInterfaceSkeletalMesh::GetTriangleSamplingFunctions(TArray<FNia
 		Sig.Inputs.Add(FNiagaraVariable(FNiagaraTypeDefinition::GetIntDef(), TEXT("Edge Index")));
 		Sig.Outputs.Add(FNiagaraVariable(FNiagaraTypeDefinition::GetIntDef(), TEXT("Neighbor Triangle Index")));
 		Sig.Outputs.Add(FNiagaraVariable(FNiagaraTypeDefinition::GetIntDef(), TEXT("Neighbor Edge Index")));
+		Sig.Outputs.Add(FNiagaraVariable(FNiagaraTypeDefinition::GetBoolDef(), TEXT("IsValid")));
 		Sig.bMemberFunction = true;
 		Sig.bRequiresContext = false;
 		Sig.bExperimental = true;
diff --git a/Engine/Plugins/FX/Niagara/Source/Niagara/Private/NiagaraDataInterfaceSkeletalMesh.cpp b/Engine/Plugins/FX/Niagara/Source/Niagara/Private/NiagaraDataInterfaceSkeletalMesh.cpp
--- a/Engine/Plugins/FX/Niagara/Source/Niagara/Private/NiagaraDataInterfaceSkeletalMesh.cpp
+++ b/Engine/Plugins/FX/Niagara/Source/Niagara/Private/NiagaraDataInterfaceSkeletalMesh.cpp
@@ -38,6 +38,7 @@ struct FNiagaraSkelMeshDIFunctionVersion
 		AddTangentBasisToGetSkinnedVertexData = 4,
 		RemoveUvSetFromMapping = 5,
 		AddedEnabledUvMapping = 6,
+		AddedValidConnectivity = 7,
 
 		VersionPlusOne,
 		LatestVersion = VersionPlusOne - 1
@@ -686,9 +687,9 @@ FSkeletalMeshConnectivityHandle FNDI_SkeletalMesh_GeneratedData::GetCachedConnec
 	{
 		FRWScopeLock ReadLock(CachedConnectivityGuard, SLT_ReadOnly);
 		TSharedPtr<FSkeletalMeshConnectivity>* Existing = CachedConnectivity.FindByPredicate([&](const TSharedPtr<FSkeletalMeshConnectivity>& Connectivity)
-			{
-				return Connectivity->Matches(MeshObject, InLodIndex);
-			});
+		{
+			return Connectivity->CanBeUsed(MeshObject, InLodIndex);
+		});
 
 		if (Existing)
 		{
@@ -1498,13 +1499,14 @@ public:
 
 			if (InstanceData->ConnectivityBuffer)
 			{
+				const uint32 NumBufferElements = FMath::DivideAndRoundUp<uint32>(InstanceData->ConnectivityBuffer->GetBufferSize(), sizeof(uint32));
 				SetSRVParameter(RHICmdList, ComputeShaderRHI, ConnectivityBuffer, InstanceData->ConnectivityBuffer->GetSrv());
-				SetShaderValue(RHICmdList, ComputeShaderRHI, ConnectivityBufferLength, InstanceData->ConnectivityBuffer->GetBufferSize());
+				SetShaderValue(RHICmdList, ComputeShaderRHI, ConnectivityBufferLength, NumBufferElements);
 				SetShaderValue(RHICmdList, ComputeShaderRHI, ConnectivityMaxAdjacentPerVertex, InstanceData->ConnectivityBuffer->MaxAdjacentTriangleCount);
 			}
 			else
 			{
-				SetSRVParameter(RHICmdList, ComputeShaderRHI, ConnectivityBuffer, FNiagaraRenderer::GetDummyIntBuffer());
+				SetSRVParameter(RHICmdList, ComputeShaderRHI, ConnectivityBuffer, FNiagaraRenderer::GetDummyUIntBuffer());
 				SetShaderValue(RHICmdList, ComputeShaderRHI, ConnectivityBufferLength, 0);
 				SetShaderValue(RHICmdList, ComputeShaderRHI, ConnectivityMaxAdjacentPerVertex, 0);
 			}
@@ -1563,7 +1565,7 @@ public:
 			SetShaderValue(RHICmdList, ComputeShaderRHI, UvMappingBufferLength, 0);
 			SetShaderValue(RHICmdList, ComputeShaderRHI, UvMappingSet, 0);
 
-			SetSRVParameter(RHICmdList, ComputeShaderRHI, ConnectivityBuffer, FNiagaraRenderer::GetDummyIntBuffer());
+			SetSRVParameter(RHICmdList, ComputeShaderRHI, ConnectivityBuffer, FNiagaraRenderer::GetDummyUIntBuffer());
 			SetShaderValue(RHICmdList, ComputeShaderRHI, ConnectivityBufferLength, 0);
 			SetShaderValue(RHICmdList, ComputeShaderRHI, ConnectivityMaxAdjacentPerVertex, 0);
 
@@ -3227,12 +3229,12 @@ bool UNiagaraDataInterfaceSkeletalMesh::GetFunctionHLSL(const FNiagaraDataInterf
 	// Adjacency
 	else if (FunctionInfo.DefinitionName == FSkeletalMeshInterfaceHelper::GetAdjacentTriangleIndexName)
 	{
-		static const TCHAR* FormatSample = TEXT("void {InstanceFunctionName} (in int VertexId, in int AdjacencyIndex, out int TriangleIndex) { {GetDISkelMeshContextName} DISkelMesh_GetAdjacentTriangleIndex(DIContext, VertexId, AdjacencyIndex, TriangleIndex); }");
+		static const TCHAR* FormatSample = TEXT("void {InstanceFunctionName} (in int VertexId, in int AdjacencyIndex, out int TriangleIndex, out bool IsValid) { {GetDISkelMeshContextName} DISkelMesh_GetAdjacentTriangleIndex(DIContext, VertexId, max(0, AdjacencyIndex), TriangleIndex); IsValid = TriangleIndex != -1; }");
 		OutHLSL += FString::Format(FormatSample, ArgsSample);
 	}
 	else if (FunctionInfo.DefinitionName == FSkeletalMeshInterfaceHelper::GetTriangleNeighborName)
 	{
-		static const TCHAR* FormatSample = TEXT("void {InstanceFunctionName} (in int TriangleIndex, in int EdgeIndex, out int NeighborTriangleIndex, out int NeighborEdgeIndex) { {GetDISkelMeshContextName} DISkelMesh_GetTriangleNeighbor(DIContext, TriangleIndex, EdgeIndex, NeighborTriangleIndex, NeighborEdgeIndex); }");
+		static const TCHAR* FormatSample = TEXT("void {InstanceFunctionName} (in int TriangleIndex, in int EdgeIndex, out int NeighborTriangleIndex, out int NeighborEdgeIndex, out bool IsValid) { {GetDISkelMeshContextName} DISkelMesh_GetTriangleNeighbor(DIContext, TriangleIndex, EdgeIndex, NeighborTriangleIndex, NeighborEdgeIndex); IsValid = NeighborTriangleIndex != -1;} ");
 		OutHLSL += FString::Format(FormatSample, ArgsSample);
 	}
 	//////////////////////////////////////////////////////////////////////////////////////////////////////////////////////////////
@@ -3397,6 +3399,16 @@ bool UNiagaraDataInterfaceSkeletalMesh::UpgradeFunctionCall(FNiagaraFunctionSign
 		}
 	}
 
+	if (FunctionSignature.FunctionVersion < FNiagaraSkelMeshDIFunctionVersion::AddedValidConnectivity)
+	{
+		if ((FunctionSignature.Name == FSkeletalMeshInterfaceHelper::GetAdjacentTriangleIndexName)
+			|| (FunctionSignature.Name == FSkeletalMeshInterfaceHelper::GetTriangleNeighborName))
+		{
+			FunctionSignature.Outputs.Add(FNiagaraVariable(FNiagaraTypeDefinition::GetBoolDef(), TEXT("IsValid")));
+			bWasChanged = true;
+		}
+	}
+
 	// Set latest version
 	FunctionSignature.FunctionVersion = FNiagaraSkelMeshDIFunctionVersion::LatestVersion;
 
diff --git a/Engine/Plugins/FX/Niagara/Source/Niagara/Private/NiagaraDataInterfaceSkeletalMeshConnectivity.cpp b/Engine/Plugins/FX/Niagara/Source/Niagara/Private/NiagaraDataInterfaceSkeletalMeshConnectivity.cpp
--- a/Engine/Plugins/FX/Niagara/Source/Niagara/Private/NiagaraDataInterfaceSkeletalMeshConnectivity.cpp
+++ b/Engine/Plugins/FX/Niagara/Source/Niagara/Private/NiagaraDataInterfaceSkeletalMeshConnectivity.cpp
@@ -3,6 +3,7 @@
 #include "NiagaraDataInterfaceSkeletalMeshConnectivity.h"
 
 #include "Algo/StableSort.h"
+#include "Algo/Unique.h"
 #include "NDISkeletalMeshCommon.h"
 #include "NiagaraResourceArrayWriter.h"
 #include "NiagaraSettings.h"
@@ -125,9 +126,9 @@ void FSkeletalMeshConnectivity::UnregisterUser(FSkeletalMeshConnectivityUsage Us
 	}
 }
 
-bool FSkeletalMeshConnectivity::Matches(const TWeakObjectPtr<USkeletalMesh>& InMeshObject, int32 InLodIndex) const
+bool FSkeletalMeshConnectivity::CanBeUsed(const TWeakObjectPtr<USkeletalMesh>& InMeshObject, int32 InLodIndex) const
 {
-	return LodIndex == InLodIndex && MeshObject == InMeshObject;
+	return !QueuedForRelease && LodIndex == InLodIndex && MeshObject == InMeshObject;
 }
 
 bool FSkeletalMeshConnectivity::IsValidMeshObject(TWeakObjectPtr<USkeletalMesh>& MeshObject, int32 InLodIndex)
@@ -196,13 +197,37 @@ const FSkeletalMeshLODRenderData* FSkeletalMeshConnectivity::GetLodRenderData()
 	return nullptr;
 }
 
+FString FSkeletalMeshConnectivity::GetMeshName() const
+{
+	if (const USkeletalMesh* Mesh = MeshObject.Get())
+	{
+		return Mesh->GetPathName();
+	}
+	return TEXT("<none>");
+}
+
 const FSkeletalMeshConnectivityProxy* FSkeletalMeshConnectivity::GetProxy() const
 {
 	return &Proxy;
 }
 
+struct FAdjacencyVertexOverlapKey
+{
+	FVector Position;
+};
+
+uint32 GetTypeHash(const FAdjacencyVertexOverlapKey& Key)
+{
+	return GetTypeHash(Key.Position);
+}
+
+bool operator==(const FAdjacencyVertexOverlapKey& Lhs, const FAdjacencyVertexOverlapKey& Rhs)
+{
+	return Lhs.Position == Rhs.Position;
+}
+
 template<typename TriangleIndexType, bool SortBySize>
-static bool BuildAdjacencyBuffer(const FSkeletalMeshLODRenderData& LodRenderData, int32 MaxAdjacencyCount, TResourceArray<uint8>& Buffer)
+static bool BuildAdjacencyBuffer(const FSkeletalMeshLODRenderData& LodRenderData, int32 MaxAdjacencyCount, TResourceArray<uint8>& Buffer, int32& MaxFoundAdjacentTriangleCount)
 {
 	const FRawStaticIndexBuffer16or32Interface* IndexBuffer = LodRenderData.MultiSizeIndexContainer.GetIndexBuffer();
 	const FPositionVertexBuffer* VertexBuffer = SortBySize ? &LodRenderData.StaticVertexBuffers.PositionVertexBuffer : nullptr;
@@ -216,18 +241,46 @@ static bool BuildAdjacencyBuffer(const FSkeletalMeshLODRenderData& LodRenderData
 		return false;
 	}
 
-	TMultiMap<int32, TriangleIndexType> RawAdjacency;
-
 	int32 MinVertexValue = std::numeric_limits<int32>::max();
 	int32 MaxVertexValue = std::numeric_limits<int32>::min();
 
+	int32 UniqueVertexCount = 0;
+	TMultiMap<int32 /*UniqueVertexIndex*/, TriangleIndexType> RawAdjacency;
+	TMap<FAdjacencyVertexOverlapKey, int32 /*UniqueVertexIndex*/> UniqueIndexMap;
+	TMap<int32 /*VertexIndex*/, int32 /*UniqueVertexIndex*/> VertexToUniqueIndexMap;
+
 	for (uint32 IndexIt = 0; IndexIt < IndexCount; ++IndexIt)
 	{
+		const TriangleIndexType TriangleId = IndexIt / 3;
 		const int32 VertexId = IndexBuffer->Get(IndexIt);
 		MinVertexValue = FMath::Min(MinVertexValue, VertexId);
 		MaxVertexValue = FMath::Max(MaxVertexValue, VertexId);
 
-		RawAdjacency.AddUnique(VertexId, IndexIt / 3);
+		FAdjacencyVertexOverlapKey OverlapKey;
+		OverlapKey.Position = VertexBuffer->VertexPosition(VertexId);
+
+		int32 UniqueIndex = UniqueVertexCount;
+		if (int32* ExistingGroup = UniqueIndexMap.Find(OverlapKey))
+		{
+			UniqueIndex = *ExistingGroup;
+		}
+		else
+		{
+			UniqueIndexMap.Add(OverlapKey, UniqueVertexCount);
+			++UniqueVertexCount;
+		}
+
+		int32* ExistingUniqueIndex = VertexToUniqueIndexMap.Find(VertexId);
+		if (!ExistingUniqueIndex)
+		{
+			VertexToUniqueIndexMap.Add(VertexId, UniqueIndex);
+		}
+		else
+		{
+			check(*ExistingUniqueIndex == UniqueIndex);
+		}
+
+		RawAdjacency.Add(UniqueIndex, TriangleId);
 	}
 
 	const int32 SizePerVertex = MaxAdjacencyCount * sizeof(TriangleIndexType);
@@ -235,7 +288,7 @@ static bool BuildAdjacencyBuffer(const FSkeletalMeshLODRenderData& LodRenderData
 	const int32 PaddedBufferSize = 4 * FMath::DivideAndRoundUp(BufferSize, 4);
 	Buffer.Init(0xFF, PaddedBufferSize);
 
-	int32 MaxAdjCount = 0;
+	MaxFoundAdjacentTriangleCount = 0;
 
 	{
 		FNiagaraResourceArrayWriter Ar(Buffer);
@@ -258,13 +311,16 @@ static bool BuildAdjacencyBuffer(const FSkeletalMeshLODRenderData& LodRenderData
 			Ar.Seek(VertexIt * SizePerVertex);
 
 			TriangleValues.Reset();
-			RawAdjacency.MultiFind(VertexIt, TriangleValues, true);
 
-			if (TriangleValues.Num() > MaxAdjCount)
 			{
-				MaxAdjCount = TriangleValues.Num();
+				RawAdjacency.MultiFind(VertexToUniqueIndexMap[VertexIt], TriangleValues);
+			
+				TriangleValues.Sort();
+				TriangleValues.SetNum(Algo::Unique(TriangleValues));
 			}
 
+			MaxFoundAdjacentTriangleCount = FMath::Max(MaxFoundAdjacentTriangleCount, TriangleValues.Num());
+
 			if (SortBySize)
 			{
 				const int32 AdjacentCount = TriangleValues.Num();
@@ -310,14 +366,28 @@ FSkeletalMeshConnectivityProxy::Initialize(const FSkeletalMeshConnectivity& Conn
 {
 	if (const FSkeletalMeshLODRenderData* LodRenderData = Connectivity.GetLodRenderData())
 	{
+		int32 MaxFoundAdjacentTriangleCount = 0;
+		bool AdjacencySuccess = false;
+
 		const auto IndexFormat = GetDefault<UNiagaraSettings>()->NDISkelMesh_AdjacencyTriangleIndexFormat;
 		if (IndexFormat == ENDISkelMesh_AdjacencyTriangleIndexFormat::Full)
 		{
-			BuildAdjacencyBuffer<uint32, true>(*LodRenderData, MaxAdjacentTriangleCount, AdjacencyResource);
+			AdjacencySuccess = BuildAdjacencyBuffer<uint32, true>(*LodRenderData, MaxAdjacentTriangleCount, AdjacencyResource, MaxFoundAdjacentTriangleCount);
 		}
 		else if (IndexFormat == ENDISkelMesh_AdjacencyTriangleIndexFormat::Half)
 		{
-			BuildAdjacencyBuffer<uint16, true>(*LodRenderData, MaxAdjacentTriangleCount, AdjacencyResource);
+			AdjacencySuccess = BuildAdjacencyBuffer<uint16, true>(*LodRenderData, MaxAdjacentTriangleCount, AdjacencyResource, MaxFoundAdjacentTriangleCount);
+		}
+
+		if (!AdjacencySuccess)
+		{
+			UE_LOG(LogNiagara, Warning, TEXT("Failed to build adjacency for %s.  Check project settings for NDISkelMesh_AdjacencyTriangleIndexFormat.  Currently using %s."),
+				*Connectivity.GetMeshName(), *StaticEnum<ENDISkelMesh_AdjacencyTriangleIndexFormat::Type>()->GetValueAsString(IndexFormat));
+		}
+		if (MaxFoundAdjacentTriangleCount > MaxAdjacentTriangleCount)
+		{
+			UE_LOG(LogNiagara, Warning, TEXT("Max adjacency limit of %d exceeded (up to %d found) when processing %s.  Some connections will be ignored."),
+				MaxAdjacentTriangleCount, MaxFoundAdjacentTriangleCount, *Connectivity.GetMeshName());
 		}
 	}
 }
@@ -331,7 +401,7 @@ FSkeletalMeshConnectivityProxy::InitRHI()
 	const int32 BufferSize = AdjacencyResource.Num();
 
 	AdjacencyBuffer = RHICreateVertexBuffer(BufferSize, BUF_ShaderResource | BUF_Static, CreateInfo);
-	AdjacencySrv = RHICreateShaderResourceView(AdjacencyBuffer, sizeof(int32), PF_R32_SINT);
+	AdjacencySrv = RHICreateShaderResourceView(AdjacencyBuffer, sizeof(uint32), PF_R32_UINT);
 
 #if STATS
 	check(GpuMemoryUsage == 0);
diff --git a/Engine/Plugins/FX/Niagara/Source/Niagara/Private/NiagaraDataInterfaceSkeletalMeshConnectivity.h b/Engine/Plugins/FX/Niagara/Source/Niagara/Private/NiagaraDataInterfaceSkeletalMeshConnectivity.h
--- a/Engine/Plugins/FX/Niagara/Source/Niagara/Private/NiagaraDataInterfaceSkeletalMeshConnectivity.h
+++ b/Engine/Plugins/FX/Niagara/Source/Niagara/Private/NiagaraDataInterfaceSkeletalMeshConnectivity.h
@@ -41,7 +41,7 @@ struct FSkeletalMeshConnectivity
 	bool CanBeDestroyed() const;
 	void RegisterUser(FSkeletalMeshConnectivityUsage Usage, bool bNeedsDataImmediately);
 	void UnregisterUser(FSkeletalMeshConnectivityUsage Usage);
-	bool Matches(const TWeakObjectPtr<USkeletalMesh>& InMeshObject, int32 InLodIndex) const;
+	bool CanBeUsed(const TWeakObjectPtr<USkeletalMesh>& InMeshObject, int32 InLodIndex) const;
 	const FSkeletalMeshConnectivityProxy* GetProxy() const;
 
 	static bool IsValidMeshObject(TWeakObjectPtr<USkeletalMesh>& MeshObject, int32 InLodIndex);
@@ -49,6 +49,7 @@ struct FSkeletalMeshConnectivity
 	int32 GetAdjacentTriangleIndex(int32 VertexIndex, int32 AdjacencyIndex) const;
 	void GetTriangleVertices(int32 TriangleIndex, int32& OutVertex0, int32& OutVertex1, int32& OutVertex2) const;
 	const FSkeletalMeshLODRenderData* GetLodRenderData() const;
+	FString GetMeshName() const;
 
 	const int32 LodIndex;
 
