UE5 - Add missing CFLAG when using wave ops

diff --git a/Engine/Plugins/FX/Niagara/Source/NiagaraVertexFactories/Private/NiagaraSortingGPU.cpp b/Engine/Plugins/FX/Niagara/Source/NiagaraVertexFactories/Private/NiagaraSortingGPU.cpp
--- a/Engine/Plugins/FX/Niagara/Source/NiagaraVertexFactories/Private/NiagaraSortingGPU.cpp
+++ b/Engine/Plugins/FX/Niagara/Source/NiagaraVertexFactories/Private/NiagaraSortingGPU.cpp
@@ -7,6 +7,7 @@ NiagaraSortingGPU.cpp: Niagara sorting shaders
 #include "NiagaraSortingGPU.h"
 #include "NiagaraGPUSortInfo.h"
 #include "ShaderParameterUtils.h"
+#include "ShaderCompilerCore.h"
 
 int32 GNiagaraGPUCulling = 1;
 static FAutoConsoleVariableRef CVarNiagaraGPUCulling(
@@ -46,5 +47,9 @@ void FNiagaraSortKeyGenCS::ModifyCompilationEnvironment(const FGlobalShaderPermu
 
 	bool bUseWaveIntrinsics = FDataDrivenShaderPlatformInfo::GetSupportsWaveOperations(Parameters.Platform);
 	OutEnvironment.SetDefine(TEXT("USE_WAVE_INTRINSICS"), bUseWaveIntrinsics);
+	if (bUseWaveIntrinsics)
+	{
+		OutEnvironment.CompilerFlags.Add(CFLAG_WaveOperations);
+	}
 }
 
