Repairing possible crash due to mismatches in data interfaces array and sortedoffsets array on PostLoad.

diff --git a/Engine/Plugins/FX/Niagara/Source/Niagara/Private/NiagaraComponent.cpp b/Engine/Plugins/FX/Niagara/Source/Niagara/Private/NiagaraComponent.cpp
--- a/Engine/Plugins/FX/Niagara/Source/Niagara/Private/NiagaraComponent.cpp
+++ b/Engine/Plugins/FX/Niagara/Source/Niagara/Private/NiagaraComponent.cpp
@@ -2490,6 +2490,7 @@ void UNiagaraComponent::PostLoad()
 	Super::PostLoad();
 
 	OverrideParameters.PostLoad();
+	OverrideParameters.SanityCheckData();
 
 #if WITH_EDITOR
 	if (Asset != nullptr)
