Deprecated UEditorEngine::OnObjectsReplaced in favor of FCoreUObjectDelegates::OnObjectsReplaced

Also removed the UEditorEngine::BroadcastBlueprintCompiled that took a map of objects, as nothing calls it

diff --git a/Engine/Plugins/Developer/TraceSourceFiltering/Source/SourceFilteringEditor/Private/EditorSessionSourceFilterService.cpp b/Engine/Plugins/Developer/TraceSourceFiltering/Source/SourceFilteringEditor/Private/EditorSessionSourceFilterService.cpp
--- a/Engine/Plugins/Developer/TraceSourceFiltering/Source/SourceFilteringEditor/Private/EditorSessionSourceFilterService.cpp
+++ b/Engine/Plugins/Developer/TraceSourceFiltering/Source/SourceFilteringEditor/Private/EditorSessionSourceFilterService.cpp
@@ -46,10 +46,7 @@ FEditorSessionSourceFilterService::FEditorSessionSourceFilterService()
 	FilterCollection = FTraceSourceFiltering::Get().GetFilterCollection();
 	FilterCollection->GetSourceFiltersUpdated().AddRaw(this, &FEditorSessionSourceFilterService::StateChanged);
 
-	if (GEditor)
-	{
-		GEditor->OnObjectsReplaced().AddRaw(this, &FEditorSessionSourceFilterService::OnObjectsReplaced);
-	}
+	FCoreUObjectDelegates::OnObjectsReplaced.AddRaw(this, &FEditorSessionSourceFilterService::OnObjectsReplaced);
 
 	// Register delegate to catch engine-level trace filtering system changes 
 	FTraceWorldFiltering::OnFilterStateChanged().AddRaw(this, &FEditorSessionSourceFilterService::StateChanged);	
@@ -59,10 +56,7 @@ FEditorSessionSourceFilterService::FEditorSessionSourceFilterService()
 
 FEditorSessionSourceFilterService::~FEditorSessionSourceFilterService()
 {
-	if (GEditor)
-	{
-		GEditor->OnObjectsReplaced().RemoveAll(this);
-	}
+	FCoreUObjectDelegates::OnObjectsReplaced.RemoveAll(this);
 	
 	FilterCollection->GetSourceFiltersUpdated().RemoveAll(this);
 	FTraceWorldFiltering::OnFilterStateChanged().RemoveAll(this);
diff --git a/Engine/Plugins/Enterprise/DataprepEditor/Source/DataprepEditor/Private/Widgets/DataprepWidgets.cpp b/Engine/Plugins/Enterprise/DataprepEditor/Source/DataprepEditor/Private/Widgets/DataprepWidgets.cpp
--- a/Engine/Plugins/Enterprise/DataprepEditor/Source/DataprepEditor/Private/Widgets/DataprepWidgets.cpp
+++ b/Engine/Plugins/Enterprise/DataprepEditor/Source/DataprepEditor/Private/Widgets/DataprepWidgets.cpp
@@ -618,11 +618,8 @@ void SDataprepDetailsView::Construct(const FArguments& InArgs)
 
 	OnPropertyChangedHandle = Generator->OnFinishedChangingProperties().AddSP( this, &SDataprepDetailsView::OnPropertyChanged );
 
-	if ( GEditor )
-	{
-		OnObjectReplacedHandle = GEditor->OnObjectsReplaced().AddSP(this, &SDataprepDetailsView::OnObjectReplaced);
-		OnObjectTransactedHandle = FCoreUObjectDelegates::OnObjectTransacted.AddSP( this, &SDataprepDetailsView::OnObjectTransacted );
-	}
+	OnObjectReplacedHandle = FCoreUObjectDelegates::OnObjectsReplaced.AddSP(this, &SDataprepDetailsView::OnObjectReplaced);
+	OnObjectTransactedHandle = FCoreUObjectDelegates::OnObjectTransacted.AddSP( this, &SDataprepDetailsView::OnObjectTransacted );
 
 	Construct();
 }
@@ -696,11 +693,8 @@ SDataprepDetailsView::~SDataprepDetailsView()
 {
 	Generator->OnFinishedChangingProperties().Remove( OnPropertyChangedHandle );
 
-	if ( GEditor )
-	{
-		GEditor->OnObjectsReplaced().Remove( OnObjectReplacedHandle );
-		FCoreUObjectDelegates::OnObjectTransacted.Remove( OnObjectTransactedHandle );
-	}
+	FCoreUObjectDelegates::OnObjectsReplaced.Remove( OnObjectReplacedHandle );
+	FCoreUObjectDelegates::OnObjectTransacted.Remove( OnObjectTransactedHandle );
 
 	if ( UDataprepAsset* DataprepAsset = DataprepAssetForParameterization.Get() )
 	{
diff --git a/Engine/Plugins/Experimental/ControlRig/Source/ControlRig/Private/AnimNode_ControlRig.cpp b/Engine/Plugins/Experimental/ControlRig/Source/ControlRig/Private/AnimNode_ControlRig.cpp
--- a/Engine/Plugins/Experimental/ControlRig/Source/ControlRig/Private/AnimNode_ControlRig.cpp
+++ b/Engine/Plugins/Experimental/ControlRig/Source/ControlRig/Private/AnimNode_ControlRig.cpp
@@ -50,10 +50,7 @@ void FAnimNode_ControlRig::OnInitializeAnimInstance(const FAnimInstanceProxy* In
 	FAnimNode_ControlRigBase::OnInitializeAnimInstance(InProxy, InAnimInstance);
 
 #if WITH_EDITOR
-	if (GEditor)
-	{
-		GEditor->OnObjectsReplaced().AddRaw(this, &FAnimNode_ControlRig::OnObjectsReplaced);
-	}
+	FCoreUObjectDelegates::OnObjectsReplaced.AddRaw(this, &FAnimNode_ControlRig::OnObjectsReplaced);
 #endif // WITH_EDITOR
 
 	InitializeProperties(InAnimInstance, GetTargetClass());
@@ -62,10 +59,7 @@ void FAnimNode_ControlRig::OnInitializeAnimInstance(const FAnimInstanceProxy* In
 FAnimNode_ControlRig::~FAnimNode_ControlRig()
 {
 #if WITH_EDITOR
-	if (GEditor)
-	{
-		GEditor->OnObjectsReplaced().RemoveAll(this);
-	}
+	FCoreUObjectDelegates::OnObjectsReplaced.RemoveAll(this);
 #endif // WITH_EDITOR
 }
 
diff --git a/Engine/Plugins/Experimental/ControlRig/Source/ControlRigEditor/Private/EditMode/ControlRigEditMode.cpp b/Engine/Plugins/Experimental/ControlRig/Source/ControlRigEditor/Private/EditMode/ControlRigEditMode.cpp
--- a/Engine/Plugins/Experimental/ControlRig/Source/ControlRigEditor/Private/EditMode/ControlRigEditMode.cpp
+++ b/Engine/Plugins/Experimental/ControlRig/Source/ControlRigEditor/Private/EditMode/ControlRigEditMode.cpp
@@ -132,7 +132,7 @@ FControlRigEditMode::FControlRigEditMode()
 	BindCommands();
 
 #if WITH_EDITOR
-	GEditor->OnObjectsReplaced().AddRaw(this, &FControlRigEditMode::OnObjectsReplaced);
+	FCoreUObjectDelegates::OnObjectsReplaced.AddRaw(this, &FControlRigEditMode::OnObjectsReplaced);
 #endif
 }
 
@@ -159,10 +159,7 @@ FControlRigEditMode::~FControlRigEditMode()
 	}
 
 #if WITH_EDITOR
-	if (GEditor)
-	{
-		GEditor->OnObjectsReplaced().RemoveAll(this);
-	}
+	FCoreUObjectDelegates::OnObjectsReplaced.RemoveAll(this);
 #endif
 }
 
diff --git a/Engine/Plugins/Experimental/ControlRig/Source/ControlRigEditor/Private/Sequencer/ControlRigParameterTrackEditor.cpp b/Engine/Plugins/Experimental/ControlRig/Source/ControlRigEditor/Private/Sequencer/ControlRigParameterTrackEditor.cpp
--- a/Engine/Plugins/Experimental/ControlRig/Source/ControlRigEditor/Private/Sequencer/ControlRigParameterTrackEditor.cpp
+++ b/Engine/Plugins/Experimental/ControlRig/Source/ControlRigEditor/Private/Sequencer/ControlRigParameterTrackEditor.cpp
@@ -184,7 +184,7 @@ FControlRigParameterTrackEditor::FControlRigParameterTrackEditor(TSharedRef<ISeq
 	{
 		//we check for two things, one if the control rig has been replaced if so we need to switch.
 		//the other is if bound object on the edit mode is null we request a re-evaluate which will reset it up.
-		FDelegateHandle OnObjectsReplacedHandle = GEditor->OnObjectsReplaced().AddLambda([&](const TMap<UObject*, UObject*>& ReplacementMap)
+		FDelegateHandle OnObjectsReplacedHandle = FCoreUObjectDelegates::OnObjectsReplaced.AddLambda([&](const TMap<UObject*, UObject*>& ReplacementMap)
 		{
 			if (GetSequencer().IsValid())
 			{
@@ -270,7 +270,7 @@ FControlRigParameterTrackEditor::FControlRigParameterTrackEditor(TSharedRef<ISeq
 			}
 
 		});
-		AcquiredResources.Add([=] { GEditor->OnObjectsReplaced().Remove(OnObjectsReplacedHandle); });
+		AcquiredResources.Add([=] { FCoreUObjectDelegates::OnObjectsReplaced.Remove(OnObjectsReplacedHandle); });
 	}
 	//register all modified/selections for control rigs
 	const TArray<FMovieSceneBinding>& Bindings = MovieScene->GetBindings();
diff --git a/Engine/Plugins/Experimental/VirtualCamera/Source/VCamCore/Private/VCamComponent.cpp b/Engine/Plugins/Experimental/VirtualCamera/Source/VCamCore/Private/VCamComponent.cpp
--- a/Engine/Plugins/Experimental/VirtualCamera/Source/VCamCore/Private/VCamComponent.cpp
+++ b/Engine/Plugins/Experimental/VirtualCamera/Source/VCamCore/Private/VCamComponent.cpp
@@ -54,10 +54,7 @@ UVCamComponent::UVCamComponent()
 		FEditorDelegates::BeginPIE.AddUObject(this, &UVCamComponent::OnBeginPIE);
 		FEditorDelegates::EndPIE.AddUObject(this, &UVCamComponent::OnEndPIE);
 
-		if (GEditor)
-		{
-			GEditor->OnObjectsReplaced().AddUObject(this, &UVCamComponent::HandleObjectReplaced);
-		}
+		FCoreUObjectDelegates::OnObjectsReplaced.AddUObject(this, &UVCamComponent::HandleObjectReplaced);
 #endif
 	}
 }
@@ -85,10 +82,7 @@ void UVCamComponent::OnComponentDestroyed(bool bDestroyingHierarchy)
 	FEditorDelegates::BeginPIE.RemoveAll(this);
 	FEditorDelegates::EndPIE.RemoveAll(this);
 
-	if (GEditor)
-	{
-		GEditor->OnObjectsReplaced().RemoveAll(this);
-	}
+	FCoreUObjectDelegates::OnObjectsReplaced.RemoveAll(this);
 #endif
 }
 
diff --git a/Engine/Plugins/FX/Niagara/Source/Niagara/Private/NiagaraComponentRendererProperties.cpp b/Engine/Plugins/FX/Niagara/Source/Niagara/Private/NiagaraComponentRendererProperties.cpp
--- a/Engine/Plugins/FX/Niagara/Source/Niagara/Private/NiagaraComponentRendererProperties.cpp
+++ b/Engine/Plugins/FX/Niagara/Source/Niagara/Private/NiagaraComponentRendererProperties.cpp
@@ -112,11 +112,8 @@ UNiagaraComponentRendererProperties::UNiagaraComponentRendererProperties()
 #endif
 	, TemplateComponent(nullptr)
 {
-#if WITH_EDITORONLY_DATA
-	if (GEditor)
-	{
-		GEditor->OnObjectsReplaced().AddUObject(this, &UNiagaraComponentRendererProperties::OnObjectsReplacedCallback);
-	}
+#if WITH_EDITOR
+	FCoreUObjectDelegates::OnObjectsReplaced.AddUObject(this, &UNiagaraComponentRendererProperties::OnObjectsReplacedCallback);
 #endif
 
 	AttributeBindings.Reserve(2);
@@ -126,11 +123,8 @@ UNiagaraComponentRendererProperties::UNiagaraComponentRendererProperties()
 
 UNiagaraComponentRendererProperties::~UNiagaraComponentRendererProperties()
 {
-#if WITH_EDITORONLY_DATA
-	if (GEditor)
-	{
-		GEditor->OnObjectsReplaced().RemoveAll(this);
-	}
+#if WITH_EDITOR
+	FCoreUObjectDelegates::OnObjectsReplaced.RemoveAll(this);
 #endif
 }
 
@@ -320,6 +314,8 @@ void UNiagaraComponentRendererProperties::CreateTemplateComponent()
 	TemplateComponent->SetAbsolute(IsWorldSpace, IsWorldSpace, IsWorldSpace);
 }
 
+#if WITH_EDITOR
+
 void UNiagaraComponentRendererProperties::OnObjectsReplacedCallback(const TMap<UObject*, UObject*>& ReplacementsMap)
 {
 	// When a custom component class is recompiled in the editor, we need to switch to the new template component object
@@ -334,6 +330,8 @@ void UNiagaraComponentRendererProperties::OnObjectsReplacedCallback(const TMap<U
 	}
 }
 
+#endif
+
 bool UNiagaraComponentRendererProperties::HasPropertyBinding(FName PropertyName) const
 {
 	for (const FNiagaraComponentPropertyBinding& Binding : PropertyBindings)
diff --git a/Engine/Plugins/FX/Niagara/Source/Niagara/Private/NiagaraRendererComponents.cpp b/Engine/Plugins/FX/Niagara/Source/Niagara/Private/NiagaraRendererComponents.cpp
--- a/Engine/Plugins/FX/Niagara/Source/Niagara/Private/NiagaraRendererComponents.cpp
+++ b/Engine/Plugins/FX/Niagara/Source/Niagara/Private/NiagaraRendererComponents.cpp
@@ -143,12 +143,9 @@ FNiagaraRendererComponents::FNiagaraRendererComponents(ERHIFeatureLevel::Type Fe
 
 	ComponentPool.Reserve(Properties->ComponentCountLimit);
 
-#if WITH_EDITORONLY_DATA
-	if (GEditor)
-	{
-		// for the component renderer we need to listen for class changes so we can clean up old component renderer instances
-		GEditor->OnObjectsReplaced().AddRaw(this, &FNiagaraRendererComponents::OnObjectsReplacedCallback);
-	}
+#if WITH_EDITOR
+	// for the component renderer we need to listen for class changes so we can clean up old component renderer instances
+	FCoreUObjectDelegates::OnObjectsReplaced.AddRaw(this, &FNiagaraRendererComponents::OnObjectsReplacedCallback);
 #endif
 }
 
@@ -179,12 +176,9 @@ void FNiagaraRendererComponents::DestroyRenderState_Concurrent()
 				OwnerActor->Destroy();
 			}
 
-#if WITH_EDITORONLY_DATA
+#if WITH_EDITOR
 			// TODO: This has the potential to race, as renderers are destroyed on the render thread, but it should be a rarity
-			if (GEditor)
-			{
-				GEditor->OnObjectsReplaced().RemoveAll(this);
-			}
+			FCoreUObjectDelegates::OnObjectsReplaced.RemoveAll(this);
 #endif
 		}
 	);
@@ -577,7 +571,7 @@ FNiagaraRendererComponents::FComponentPropertyAddress FNiagaraRendererComponents
 	return FComponentPropertyAddress();
 }
 
-#if WITH_EDITORONLY_DATA
+#if WITH_EDITOR
 
 void FNiagaraRendererComponents::OnObjectsReplacedCallback(const TMap<UObject*, UObject*>& ReplacementsMap)
 {
diff --git a/Engine/Plugins/FX/Niagara/Source/Niagara/Public/NiagaraComponentRendererProperties.h b/Engine/Plugins/FX/Niagara/Source/Niagara/Public/NiagaraComponentRendererProperties.h
--- a/Engine/Plugins/FX/Niagara/Source/Niagara/Public/NiagaraComponentRendererProperties.h
+++ b/Engine/Plugins/FX/Niagara/Source/Niagara/Public/NiagaraComponentRendererProperties.h
@@ -151,6 +151,8 @@ private:
 
 	bool HasPropertyBinding(FName PropertyName) const;
 
+#if WITH_EDITOR
 	/** Callback for whenever any blueprint components are reinstanced */
 	void OnObjectsReplacedCallback(const TMap<UObject*, UObject*>& ReplacementsMap);
+#endif
 };
diff --git a/Engine/Plugins/FX/Niagara/Source/Niagara/Public/NiagaraRendererComponents.h b/Engine/Plugins/FX/Niagara/Source/Niagara/Public/NiagaraRendererComponents.h
--- a/Engine/Plugins/FX/Niagara/Source/Niagara/Public/NiagaraRendererComponents.h
+++ b/Engine/Plugins/FX/Niagara/Source/Niagara/Public/NiagaraRendererComponents.h
@@ -56,7 +56,7 @@ private:
 	};
 
 	void ResetComponentPool(bool bResetOwner);
-#if WITH_EDITORONLY_DATA
+#if WITH_EDITOR
 	void OnObjectsReplacedCallback(const TMap<UObject*, UObject*>& ReplacementsMap);
 #endif
 	
diff --git a/Engine/Plugins/MovieScene/ActorSequence/Source/ActorSequenceEditor/Private/ActorSequenceComponentCustomization.cpp b/Engine/Plugins/MovieScene/ActorSequence/Source/ActorSequenceEditor/Private/ActorSequenceComponentCustomization.cpp
--- a/Engine/Plugins/MovieScene/ActorSequence/Source/ActorSequenceEditor/Private/ActorSequenceComponentCustomization.cpp
+++ b/Engine/Plugins/MovieScene/ActorSequence/Source/ActorSequenceEditor/Private/ActorSequenceComponentCustomization.cpp
@@ -29,7 +29,7 @@ class SActorSequenceEditorWidgetWrapper : public SActorSequenceEditorWidget
 public:
 	~SActorSequenceEditorWidgetWrapper()
 	{
-		GEditor->OnObjectsReplaced().Remove(OnObjectsReplacedHandle);
+		FCoreUObjectDelegates::OnObjectsReplaced.Remove(OnObjectsReplacedHandle);
 	}
 
 	void Construct(const FArguments& InArgs, TWeakObjectPtr<UActorSequenceComponent> InSequenceComponent)
@@ -39,7 +39,7 @@ public:
 		WeakSequenceComponent = InSequenceComponent;
 		AssignSequence(GetActorSequence());
 
-		OnObjectsReplacedHandle = GEditor->OnObjectsReplaced().AddSP(this, &SActorSequenceEditorWidgetWrapper::OnObjectsReplaced);
+		OnObjectsReplacedHandle = FCoreUObjectDelegates::OnObjectsReplaced.AddSP(this, &SActorSequenceEditorWidgetWrapper::OnObjectsReplaced);
 	}
 
 protected:
diff --git a/Engine/Plugins/MovieScene/LevelSequenceEditor/Source/LevelSequenceEditor/Private/Misc/LevelSequenceEditorSpawnRegister.cpp b/Engine/Plugins/MovieScene/LevelSequenceEditor/Source/LevelSequenceEditor/Private/Misc/LevelSequenceEditorSpawnRegister.cpp
--- a/Engine/Plugins/MovieScene/LevelSequenceEditor/Source/LevelSequenceEditor/Private/Misc/LevelSequenceEditorSpawnRegister.cpp
+++ b/Engine/Plugins/MovieScene/LevelSequenceEditor/Source/LevelSequenceEditor/Private/Misc/LevelSequenceEditorSpawnRegister.cpp
@@ -31,7 +31,7 @@ FLevelSequenceEditorSpawnRegister::FLevelSequenceEditorSpawnRegister()
 	OnActorSelectionChangedHandle = LevelEditor.OnActorSelectionChanged().AddRaw(this, &FLevelSequenceEditorSpawnRegister::HandleActorSelectionChanged);
 
 #if WITH_EDITOR
-	GEditor->OnObjectsReplaced().AddRaw(this, &FLevelSequenceEditorSpawnRegister::OnObjectsReplaced);
+	FCoreUObjectDelegates::OnObjectsReplaced.AddRaw(this, &FLevelSequenceEditorSpawnRegister::OnObjectsReplaced);
 
 	OnObjectModifiedHandle = FCoreUObjectDelegates::OnObjectModified.AddRaw(this, &FLevelSequenceEditorSpawnRegister::OnObjectModified);
 	OnObjectSavedHandle    = FCoreUObjectDelegates::OnObjectSaved.AddRaw(this, &FLevelSequenceEditorSpawnRegister::OnPreObjectSaved);
@@ -54,7 +54,7 @@ FLevelSequenceEditorSpawnRegister::~FLevelSequenceEditorSpawnRegister()
 	}
 
 #if WITH_EDITOR
-	GEditor->OnObjectsReplaced().RemoveAll(this);
+	FCoreUObjectDelegates::OnObjectsReplaced.RemoveAll(this);
 	FCoreUObjectDelegates::OnObjectModified.Remove(OnObjectModifiedHandle);
 	FCoreUObjectDelegates::OnObjectSaved.Remove(OnObjectSavedHandle);
 #endif
diff --git a/Engine/Plugins/VirtualProduction/RemoteControl/Source/RemoteControlUI/Private/UI/SRemoteControlPanel.cpp b/Engine/Plugins/VirtualProduction/RemoteControl/Source/RemoteControlUI/Private/UI/SRemoteControlPanel.cpp
--- a/Engine/Plugins/VirtualProduction/RemoteControl/Source/RemoteControlUI/Private/UI/SRemoteControlPanel.cpp
+++ b/Engine/Plugins/VirtualProduction/RemoteControl/Source/RemoteControlUI/Private/UI/SRemoteControlPanel.cpp
@@ -1115,7 +1115,7 @@ void SRemoteControlPanel::RegisterEvents()
 		FEditorDelegates::PostPIEStarted.AddSP(this, &SRemoteControlPanel::OnPieEvent);
 		FEditorDelegates::EndPIE.AddSP(this, &SRemoteControlPanel::OnPieEvent);
 
-		GEditor->OnObjectsReplaced().AddSP(this, &SRemoteControlPanel::OnObjectsReplaced);
+		FCoreUObjectDelegates::OnObjectsReplaced.AddSP(this, &SRemoteControlPanel::OnObjectsReplaced);
 		GEditor->OnBlueprintReinstanced().AddSP(this, &SRemoteControlPanel::RefreshBlueprintLibraryNodes);
 	}
 }
@@ -1125,7 +1125,7 @@ void SRemoteControlPanel::UnregisterEvents()
 	if (GEditor)
 	{
 		GEditor->OnBlueprintReinstanced().RemoveAll(this);
-		GEditor->OnObjectsReplaced().RemoveAll(this);
+		FCoreUObjectDelegates::OnObjectsReplaced.RemoveAll(this);
 
 		FEditorDelegates::EndPIE.RemoveAll(this);
 		FEditorDelegates::PostPIEStarted.RemoveAll(this);
diff --git a/Engine/Source/Editor/DetailCustomizations/Private/ComponentTransformDetails.cpp b/Engine/Source/Editor/DetailCustomizations/Private/ComponentTransformDetails.cpp
--- a/Engine/Source/Editor/DetailCustomizations/Private/ComponentTransformDetails.cpp
+++ b/Engine/Source/Editor/DetailCustomizations/Private/ComponentTransformDetails.cpp
@@ -91,18 +91,12 @@ FComponentTransformDetails::FComponentTransformDetails( const TArray< TWeakObjec
 	, HiddenFieldMask( 0 )
 {
 	GConfig->GetBool(TEXT("SelectionDetails"), TEXT("PreserveScaleRatio"), bPreserveScaleRatio, GEditorPerProjectIni);
-	if (GEditor)
-	{
-		GEditor->OnObjectsReplaced().AddRaw(this, &FComponentTransformDetails::OnObjectsReplaced);
-	}
+	FCoreUObjectDelegates::OnObjectsReplaced.AddRaw(this, &FComponentTransformDetails::OnObjectsReplaced);
 }
 
 FComponentTransformDetails::~FComponentTransformDetails()
 {
-	if (GEditor)
-	{
-		GEditor->OnObjectsReplaced().RemoveAll(this);
-	}
+	FCoreUObjectDelegates::OnObjectsReplaced.RemoveAll(this);
 }
 
 TSharedRef<SWidget> FComponentTransformDetails::BuildTransformFieldLabel( ETransformField::Type TransformField )
diff --git a/Engine/Source/Editor/DetailCustomizations/Private/MovieSceneCaptureCustomization.cpp b/Engine/Source/Editor/DetailCustomizations/Private/MovieSceneCaptureCustomization.cpp
--- a/Engine/Source/Editor/DetailCustomizations/Private/MovieSceneCaptureCustomization.cpp
+++ b/Engine/Source/Editor/DetailCustomizations/Private/MovieSceneCaptureCustomization.cpp
@@ -17,14 +17,14 @@ TSharedRef<IDetailCustomization> FMovieSceneCaptureCustomization::MakeInstance()
 
 FMovieSceneCaptureCustomization::FMovieSceneCaptureCustomization()
 {
-	ObjectsReplacedHandle = GEditor->OnObjectsReplaced().AddRaw(this, &FMovieSceneCaptureCustomization::OnObjectsReplaced);
+	ObjectsReplacedHandle = FCoreUObjectDelegates::OnObjectsReplaced.AddRaw(this, &FMovieSceneCaptureCustomization::OnObjectsReplaced);
 	PropertyChangedHandle = FCoreUObjectDelegates::OnObjectPropertyChanged.AddRaw(this, &FMovieSceneCaptureCustomization::OnObjectPostEditChange);
 }
 
 FMovieSceneCaptureCustomization::~FMovieSceneCaptureCustomization()
 {
 	FCoreUObjectDelegates::OnObjectPropertyChanged.Remove(PropertyChangedHandle);
-	GEditor->OnObjectsReplaced().Remove(ObjectsReplacedHandle);
+	FCoreUObjectDelegates::OnObjectsReplaced.Remove(ObjectsReplacedHandle);
 }
 
 void FMovieSceneCaptureCustomization::CustomizeDetails(IDetailLayoutBuilder& DetailBuilder)
diff --git a/Engine/Source/Editor/FoliageEdit/Private/FoliageEdMode.cpp b/Engine/Source/Editor/FoliageEdit/Private/FoliageEdMode.cpp
--- a/Engine/Source/Editor/FoliageEdit/Private/FoliageEdMode.cpp
+++ b/Engine/Source/Editor/FoliageEdit/Private/FoliageEdMode.cpp
@@ -383,7 +383,7 @@ void FEdModeFoliage::Enter()
 	FEdMode::Enter();
 
 	// register for any objects replaced
-	GEditor->OnObjectsReplaced().AddRaw(this, &FEdModeFoliage::OnObjectsReplaced);
+	FCoreUObjectDelegates::OnObjectsReplaced.AddRaw(this, &FEdModeFoliage::OnObjectsReplaced);
 	FEditorDelegates::EndPIE.AddRaw(this, &FEdModeFoliage::OnEndPIE);
 
 
@@ -481,7 +481,7 @@ void FEdModeFoliage::Exit()
 		AssetRegistryModule.Get().OnAssetRemoved().RemoveAll(this);
 	}
 
-	GEditor->OnObjectsReplaced().RemoveAll(this);
+	FCoreUObjectDelegates::OnObjectsReplaced.RemoveAll(this);
 
 
 	// Remove the brush
diff --git a/Engine/Source/Editor/Kismet/Private/SSCSEditor.cpp b/Engine/Source/Editor/Kismet/Private/SSCSEditor.cpp
--- a/Engine/Source/Editor/Kismet/Private/SSCSEditor.cpp
+++ b/Engine/Source/Editor/Kismet/Private/SSCSEditor.cpp
@@ -4093,7 +4093,7 @@ void SSCSEditor::Construct( const FArguments& InArgs )
 	if (EditorMode == EComponentEditorMode::ActorInstance)
 	{
 		GEngine->OnLevelComponentRequestRename().AddSP(this, &SSCSEditor::OnLevelComponentRequestRename);
-		GEditor->OnObjectsReplaced().AddSP(this, &SSCSEditor::OnObjectsReplaced);
+		FCoreUObjectDelegates::OnObjectsReplaced.AddSP(this, &SSCSEditor::OnObjectsReplaced);
 	}
 }
 END_SLATE_FUNCTION_BUILD_OPTIMIZATION
diff --git a/Engine/Source/Editor/Matinee/Private/Matinee.cpp b/Engine/Source/Editor/Matinee/Private/Matinee.cpp
--- a/Engine/Source/Editor/Matinee/Private/Matinee.cpp
+++ b/Engine/Source/Editor/Matinee/Private/Matinee.cpp
@@ -976,7 +976,7 @@ void FMatinee::InitMatinee(const EToolkitMode::Type Mode, const TSharedPtr< clas
 	OnActorMovedDelegateHandle = GEngine->OnActorMoved().AddRaw(this, &FMatinee::OnActorMoved);
 
 	// register for any objects replaced
-	GEditor->OnObjectsReplaced().AddSP(this, &FMatinee::OnObjectsReplaced);
+	FCoreUObjectDelegates::OnObjectsReplaced.AddSP(this, &FMatinee::OnObjectsReplaced);
 
 	// Make sure any particle replay tracks are filled in with the correct state
 	UpdateParticleReplayTracks();
@@ -2002,7 +2002,7 @@ void FMatinee::OnClose()
 
 	// Unregister call back events
 	GEngine->OnActorMoved().Remove(OnActorMovedDelegateHandle);
-	GEditor->OnObjectsReplaced().RemoveAll(this);
+	FCoreUObjectDelegates::OnObjectsReplaced.RemoveAll(this);
 
 	// Restore the perspective viewport audio settings when matinee closes.
 	SetAudioRealtimeOverride(false);
diff --git a/Engine/Source/Editor/MeshPaint/Private/IMeshPaintMode.cpp b/Engine/Source/Editor/MeshPaint/Private/IMeshPaintMode.cpp
--- a/Engine/Source/Editor/MeshPaint/Private/IMeshPaintMode.cpp
+++ b/Engine/Source/Editor/MeshPaint/Private/IMeshPaintMode.cpp
@@ -68,7 +68,7 @@ void IMeshPaintEdMode::Enter()
 	bWasSelectionLockedOnStart = GEdSelectionLock;	
 
 	// Catch when objects are replaced when a construction script is rerun
-	GEditor->OnObjectsReplaced().AddSP(this, &IMeshPaintEdMode::OnObjectsReplaced);
+	FCoreUObjectDelegates::OnObjectsReplaced.AddSP(this, &IMeshPaintEdMode::OnObjectsReplaced);
 
 	// Hook into pre/post world save, so that the original collision volumes can be temporarily reinstated
 	FEditorDelegates::PreSaveWorld.AddSP(this, &IMeshPaintEdMode::OnPreSaveWorld);
@@ -141,7 +141,7 @@ void IMeshPaintEdMode::Exit()
 	GEditor->GetEditorSubsystem<UImportSubsystem>()->OnAssetPostImport.RemoveAll(this);
 	FEditorDelegates::PreSaveWorld.RemoveAll(this);
 	FEditorDelegates::PostSaveWorld.RemoveAll(this);
-	GEditor->OnObjectsReplaced().RemoveAll(this);
+	FCoreUObjectDelegates::OnObjectsReplaced.RemoveAll(this);
 	USelection::SelectionChangedEvent.Remove(SelectionChangedHandle);
 
 	// Call parent implementation
diff --git a/Engine/Source/Editor/MovieSceneTools/Private/MovieSceneToolsModule.cpp b/Engine/Source/Editor/MovieSceneTools/Private/MovieSceneToolsModule.cpp
--- a/Engine/Source/Editor/MovieSceneTools/Private/MovieSceneToolsModule.cpp
+++ b/Engine/Source/Editor/MovieSceneTools/Private/MovieSceneToolsModule.cpp
@@ -187,22 +187,7 @@ void FMovieSceneToolsModule::StartupModule()
 		}
 	};
 
-	if (GEditor)
-	{
-		this->OnObjectsReplacedHandle = GEditor->OnObjectsReplaced().AddLambda(OnObjectsReplaced);
-	}
-	else
-	{
-		FCoreDelegates::OnFEngineLoopInitComplete.AddLambda(
-			[this, OnObjectsReplaced]
-			{
-				if (GEditor)
-				{
-					this->OnObjectsReplacedHandle = GEditor->OnObjectsReplaced().AddLambda(OnObjectsReplaced);
-				}
-			}
-		);
-	}
+	OnObjectsReplacedHandle = FCoreUObjectDelegates::OnObjectsReplaced.AddLambda(OnObjectsReplaced);
 
 	// EditorStyle must be initialized by now
 	FModuleManager::Get().LoadModule("EditorStyle");
@@ -230,10 +215,7 @@ void FMovieSceneToolsModule::ShutdownModule()
 		SettingsModule->UnregisterSettings("Project", "Editor", "Level Sequences");
 	}
 
-	if (GEditor)
-	{
-		GEditor->OnObjectsReplaced().Remove(OnObjectsReplacedHandle);
-	}
+	FCoreUObjectDelegates::OnObjectsReplaced.Remove(OnObjectsReplacedHandle);
 
 	if (!FModuleManager::Get().IsModuleLoaded("Sequencer"))
 	{
diff --git a/Engine/Source/Editor/PropertyEditor/Private/Presentation/PropertyTable/PropertyTable.cpp b/Engine/Source/Editor/PropertyEditor/Private/Presentation/PropertyTable/PropertyTable.cpp
--- a/Engine/Source/Editor/PropertyEditor/Private/Presentation/PropertyTable/PropertyTable.cpp
+++ b/Engine/Source/Editor/PropertyEditor/Private/Presentation/PropertyTable/PropertyTable.cpp
@@ -39,9 +39,9 @@ FPropertyTable::FPropertyTable()
 
 FPropertyTable::~FPropertyTable()
 {
-	if (GEditor && ObjectsReplacedHandle.IsValid())
+	if (ObjectsReplacedHandle.IsValid())
 	{
-		GEditor->OnObjectsReplaced().Remove(ObjectsReplacedHandle);
+		FCoreUObjectDelegates::OnObjectsReplaced.Remove(ObjectsReplacedHandle);
 	}
 }
 
@@ -1021,9 +1021,9 @@ void FPropertyTable::SetObjects( const TArray< TWeakObjectPtr< UObject > >& Obje
 	UpdateRows();
 
 	// Bind to object delegates, we can't do this at construction because the shared pointer isn't set up yet
-	if (GEditor && !ObjectsReplacedHandle.IsValid())
+	if (!ObjectsReplacedHandle.IsValid())
 	{
-		ObjectsReplacedHandle = GEditor->OnObjectsReplaced().AddSP(this, &FPropertyTable::OnObjectsReplaced);
+		ObjectsReplacedHandle = FCoreUObjectDelegates::OnObjectsReplaced.AddSP(this, &FPropertyTable::OnObjectsReplaced);
 	}
 }
 
diff --git a/Engine/Source/Editor/Sequencer/Private/Sequencer.cpp b/Engine/Source/Editor/Sequencer/Private/Sequencer.cpp
--- a/Engine/Source/Editor/Sequencer/Private/Sequencer.cpp
+++ b/Engine/Source/Editor/Sequencer/Private/Sequencer.cpp
@@ -347,7 +347,7 @@ void FSequencer::InitSequencer(const FSequencerInitParams& InitParams, const TSh
 	}
 
 	{
-		FDelegateHandle OnObjectsReplacedHandle = GEditor->OnObjectsReplaced().AddLambda([&](const TMap<UObject*, UObject*>& ReplacementMap)
+		FDelegateHandle OnObjectsReplacedHandle = FCoreUObjectDelegates::OnObjectsReplaced.AddLambda([&](const TMap<UObject*, UObject*>& ReplacementMap)
 		{
 			// Close sequencer if any of the objects being replaced is itself
 			TArray<UPackage*> AllSequences;
@@ -392,7 +392,7 @@ void FSequencer::InitSequencer(const FSequencerInitParams& InitParams, const TSh
 			PreAnimatedState.OnObjectsReplaced(ReplacementMap);
 
 		});
-		AcquiredResources.Add([=] { GEditor->OnObjectsReplaced().Remove(OnObjectsReplacedHandle); });
+		AcquiredResources.Add([=] { FCoreUObjectDelegates::OnObjectsReplaced.Remove(OnObjectsReplacedHandle); });
 	}
 
 
diff --git a/Engine/Source/Editor/UMGEditor/Private/Hierarchy/SHierarchyView.cpp b/Engine/Source/Editor/UMGEditor/Private/Hierarchy/SHierarchyView.cpp
--- a/Engine/Source/Editor/UMGEditor/Private/Hierarchy/SHierarchyView.cpp
+++ b/Engine/Source/Editor/UMGEditor/Private/Hierarchy/SHierarchyView.cpp
@@ -29,7 +29,7 @@ void SHierarchyView::Construct(const FArguments& InArgs, TSharedPtr<FWidgetBluep
 	bIsUpdatingSelection = false;
 
 	// register for any objects replaced
-	GEditor->OnObjectsReplaced().AddRaw(this, &SHierarchyView::OnObjectsReplaced);
+	FCoreUObjectDelegates::OnObjectsReplaced.AddRaw(this, &SHierarchyView::OnObjectsReplaced);
 
 	// Create the filter for searching in the tree
 	SearchBoxWidgetFilter = MakeShareable(new WidgetTextFilter(WidgetTextFilter::FItemToStringArray::CreateSP(this, &SHierarchyView::GetWidgetFilterStrings)));
@@ -98,7 +98,7 @@ SHierarchyView::~SHierarchyView()
 		BlueprintEditor.Pin()->OnSelectedWidgetsChanged.RemoveAll(this);
 	}
 
-	GEditor->OnObjectsReplaced().RemoveAll(this);
+	FCoreUObjectDelegates::OnObjectsReplaced.RemoveAll(this);
 }
 
 void SHierarchyView::Tick( const FGeometry& AllottedGeometry, const double InCurrentTime, const float InDeltaTime )
diff --git a/Engine/Source/Editor/UMGEditor/Private/Palette/SPaletteViewModel.cpp b/Engine/Source/Editor/UMGEditor/Private/Palette/SPaletteViewModel.cpp
--- a/Engine/Source/Editor/UMGEditor/Private/Palette/SPaletteViewModel.cpp
+++ b/Engine/Source/Editor/UMGEditor/Private/Palette/SPaletteViewModel.cpp
@@ -117,7 +117,7 @@ void FPaletteViewModel::RegisterToEvents()
 	IHotReloadModule::Get().OnHotReload().AddSP(this, &FPaletteViewModel::HandleOnHotReload);
 
 	// register for any objects replaced
-	GEditor->OnObjectsReplaced().AddRaw(this, &FPaletteViewModel::OnObjectsReplaced);
+	FCoreUObjectDelegates::OnObjectsReplaced.AddRaw(this, &FPaletteViewModel::OnObjectsReplaced);
 
 	// Register for favorite list update to handle the case where a favorite is added in another window of the UMG Designer
 	UWidgetPaletteFavorites* Favorites = GetDefault<UWidgetDesignerSettings>()->Favorites;
@@ -129,7 +129,7 @@ FPaletteViewModel::~FPaletteViewModel()
 	GEditor->OnBlueprintReinstanced().RemoveAll(this);
 	FEditorDelegates::OnAssetsDeleted.RemoveAll(this);
 	IHotReloadModule::Get().OnHotReload().RemoveAll(this);
-	GEditor->OnObjectsReplaced().RemoveAll(this);
+	FCoreUObjectDelegates::OnObjectsReplaced.RemoveAll(this);
 
 	UWidgetPaletteFavorites* Favorites = GetDefault<UWidgetDesignerSettings>()->Favorites;
 	Favorites->OnFavoritesUpdated.RemoveAll(this);
diff --git a/Engine/Source/Editor/UMGEditor/Private/Templates/WidgetTemplateClass.cpp b/Engine/Source/Editor/UMGEditor/Private/Templates/WidgetTemplateClass.cpp
--- a/Engine/Source/Editor/UMGEditor/Private/Templates/WidgetTemplateClass.cpp
+++ b/Engine/Source/Editor/UMGEditor/Private/Templates/WidgetTemplateClass.cpp
@@ -17,7 +17,7 @@ FWidgetTemplateClass::FWidgetTemplateClass()
 	: WidgetClass(nullptr)
 {
 	// register for any objects replaced
-	GEditor->OnObjectsReplaced().AddRaw(this, &FWidgetTemplateClass::OnObjectsReplaced);
+	FCoreUObjectDelegates::OnObjectsReplaced.AddRaw(this, &FWidgetTemplateClass::OnObjectsReplaced);
 }
 
 FWidgetTemplateClass::FWidgetTemplateClass(TSubclassOf<UWidget> InWidgetClass)
@@ -26,7 +26,7 @@ FWidgetTemplateClass::FWidgetTemplateClass(TSubclassOf<UWidget> InWidgetClass)
 	Name = WidgetClass->GetDisplayNameText();
 
 	// register for any objects replaced
-	GEditor->OnObjectsReplaced().AddRaw(this, &FWidgetTemplateClass::OnObjectsReplaced);
+	FCoreUObjectDelegates::OnObjectsReplaced.AddRaw(this, &FWidgetTemplateClass::OnObjectsReplaced);
 }
 
 FWidgetTemplateClass::FWidgetTemplateClass(const FAssetData& InWidgetAssetData, TSubclassOf<UWidget> InWidgetClass)
@@ -45,7 +45,7 @@ FWidgetTemplateClass::FWidgetTemplateClass(const FAssetData& InWidgetAssetData,
 
 FWidgetTemplateClass::~FWidgetTemplateClass()
 {
-	GEditor->OnObjectsReplaced().RemoveAll(this);
+	FCoreUObjectDelegates::OnObjectsReplaced.RemoveAll(this);
 }
 
 FText FWidgetTemplateClass::GetCategory() const
diff --git a/Engine/Source/Editor/UMGEditor/Private/WidgetBlueprintEditor.cpp b/Engine/Source/Editor/UMGEditor/Private/WidgetBlueprintEditor.cpp
--- a/Engine/Source/Editor/UMGEditor/Private/WidgetBlueprintEditor.cpp
+++ b/Engine/Source/Editor/UMGEditor/Private/WidgetBlueprintEditor.cpp
@@ -89,7 +89,7 @@ FWidgetBlueprintEditor::~FWidgetBlueprintEditor()
 		Blueprint->OnCompiled().RemoveAll(this);
 	}
 
-	GEditor->OnObjectsReplaced().RemoveAll(this);
+	FCoreUObjectDelegates::OnObjectsReplaced.RemoveAll(this);
 	
 	if ( Sequencer.IsValid() )
 	{
@@ -123,7 +123,7 @@ void FWidgetBlueprintEditor::InitWidgetBlueprintEditor(const EToolkitMode::Type
 	InitBlueprintEditor(Mode, InitToolkitHost, InBlueprints, bShouldOpenInDefaultsMode);
 
 	// register for any objects replaced
-	GEditor->OnObjectsReplaced().AddSP(this, &FWidgetBlueprintEditor::OnObjectsReplaced);
+	FCoreUObjectDelegates::OnObjectsReplaced.AddSP(this, &FWidgetBlueprintEditor::OnObjectsReplaced);
 
 	// for change selected widgets on sequencer tree view
 	UWidgetBlueprint* Blueprint = GetWidgetBlueprintObj();
diff --git a/Engine/Source/Editor/UnrealEd/Classes/Editor/EditorEngine.h b/Engine/Source/Editor/UnrealEd/Classes/Editor/EditorEngine.h
--- a/Engine/Source/Editor/UnrealEd/Classes/Editor/EditorEngine.h
+++ b/Engine/Source/Editor/UnrealEd/Classes/Editor/EditorEngine.h
@@ -597,13 +597,10 @@ public:
 	void BroadcastBlueprintReinstanced() { BlueprintReinstanced.Broadcast(); }
 
 	/** Called when UObjects have been replaced to allow others a chance to fix their references. */
-	// TODO: Deprecate in favor of using FCoreUObjectDelegates::OnObjectsReplaced directly?
 	using FObjectsReplacedEvent = FCoreUObjectDelegates::FOnObjectsReplaced;
+	UE_DEPRECATED(5.0, "Use FCoreUObjectDelegates::OnObjectsReplaced instead.")
 	FObjectsReplacedEvent& OnObjectsReplaced() { return FCoreUObjectDelegates::OnObjectsReplaced; }
 
-	/**	Broadcasts that objects have been replaced*/
-	void BroadcastBlueprintCompiled(const TMap<UObject*, UObject*>& ReplacementMap) { FCoreUObjectDelegates::OnObjectsReplaced.Broadcast(ReplacementMap); }
-
 	/** Called when a package with data-driven classes becomes loaded or unloaded */
 	DECLARE_EVENT( UEditorEngine, FClassPackageLoadedOrUnloadedEvent );
 	FClassPackageLoadedOrUnloadedEvent& OnClassPackageLoadedOrUnloaded() { return ClassPackageLoadedOrUnloadedEvent; }
diff --git a/Engine/Source/Editor/UnrealEd/Private/Kismet2/KismetReinstanceUtilities.cpp b/Engine/Source/Editor/UnrealEd/Private/Kismet2/KismetReinstanceUtilities.cpp
--- a/Engine/Source/Editor/UnrealEd/Private/Kismet2/KismetReinstanceUtilities.cpp
+++ b/Engine/Source/Editor/UnrealEd/Private/Kismet2/KismetReinstanceUtilities.cpp
@@ -2078,11 +2078,7 @@ void FBlueprintCompileReinstancer::ReplaceInstancesOfClass_Inner(TMap<UClass*, U
 		TMap<UObject*, UObject*> ReplacedObjects;
 	} ObjectRemappingHelper;
 
-	FDelegateHandle OnObjectsReplacedHandle = FDelegateHandle();
-	if(GEditor)
-	{
-		OnObjectsReplacedHandle = GEditor->OnObjectsReplaced().AddRaw(&ObjectRemappingHelper, &FObjectRemappingHelper::OnObjectsReplaced);
-	}
+	FDelegateHandle OnObjectsReplacedHandle = FCoreUObjectDelegates::OnObjectsReplaced.AddRaw(&ObjectRemappingHelper, &FObjectRemappingHelper::OnObjectsReplaced);
 
 	auto UpdateObjectBeingDebugged = [](UObject* InOldObject, UObject* InNewObject)
 	{
@@ -2270,10 +2266,7 @@ void FBlueprintCompileReinstancer::ReplaceInstancesOfClass_Inner(TMap<UClass*, U
 		}
 	}
 
-	if(GEditor)
-	{
-		GEditor->OnObjectsReplaced().Remove(OnObjectsReplacedHandle);
-	}
+	FCoreUObjectDelegates::OnObjectsReplaced.Remove(OnObjectsReplacedHandle);
 
 	// Now replace any pointers to the old archetypes/instances with pointers to the new one
 	TArray<UObject*> SourceObjects;
diff --git a/Engine/Source/Editor/UnrealEd/Private/Toolkits/SimpleAssetEditor.cpp b/Engine/Source/Editor/UnrealEd/Private/Toolkits/SimpleAssetEditor.cpp
--- a/Engine/Source/Editor/UnrealEd/Private/Toolkits/SimpleAssetEditor.cpp
+++ b/Engine/Source/Editor/UnrealEd/Private/Toolkits/SimpleAssetEditor.cpp
@@ -40,7 +40,7 @@ const FName FSimpleAssetEditor::SimpleEditorAppIdentifier( TEXT( "GenericEditorA
 FSimpleAssetEditor::~FSimpleAssetEditor()
 {
 	GEditor->GetEditorSubsystem<UImportSubsystem>()->OnAssetPostImport.RemoveAll(this);
-	GEditor->OnObjectsReplaced().RemoveAll(this);
+	FCoreUObjectDelegates::OnObjectsReplaced.RemoveAll(this);
 
 	DetailsView.Reset();
 }
@@ -50,7 +50,7 @@ void FSimpleAssetEditor::InitEditor( const EToolkitMode::Type Mode, const TShare
 {
 	EditingObjects = ObjectsToEdit;
 	GEditor->GetEditorSubsystem<UImportSubsystem>()->OnAssetPostImport.AddSP(this, &FSimpleAssetEditor::HandleAssetPostImport);
-	GEditor->OnObjectsReplaced().AddSP(this, &FSimpleAssetEditor::OnObjectsReplaced);
+	FCoreUObjectDelegates::OnObjectsReplaced.AddSP(this, &FSimpleAssetEditor::OnObjectsReplaced);
 
 	FPropertyEditorModule& PropertyEditorModule = FModuleManager::GetModuleChecked<FPropertyEditorModule>( "PropertyEditor" );
 	FDetailsViewArgs DetailsViewArgs;
