bugfix for setting resolution from emitter scripts

diff --git a/Engine/Plugins/FX/Niagara/Source/Niagara/Private/NiagaraDataInterfaceNeighborGrid3D.cpp b/Engine/Plugins/FX/Niagara/Source/Niagara/Private/NiagaraDataInterfaceNeighborGrid3D.cpp
--- a/Engine/Plugins/FX/Niagara/Source/Niagara/Private/NiagaraDataInterfaceNeighborGrid3D.cpp
+++ b/Engine/Plugins/FX/Niagara/Source/Niagara/Private/NiagaraDataInterfaceNeighborGrid3D.cpp
@@ -678,7 +678,7 @@ void UNiagaraDataInterfaceNeighborGrid3D::SetNumCells(FVectorVMContext& Context)
 			InstData->NumCells.Z = NewNumCellsZ;
 			InstData->MaxNeighborsPerCell = NewMaxNeighborsPerCell;
 		
-			InstData->NeedsRealloc = OldNumCells != InstData->NumCells && OldMaxNeighborsPerCell != InstData->MaxNeighborsPerCell;
+			InstData->NeedsRealloc = OldNumCells != InstData->NumCells || OldMaxNeighborsPerCell != InstData->MaxNeighborsPerCell;
 		}
 	}
 }
