Changing a check to a log as it isn't a fatal issue.

diff --git a/Engine/Plugins/FX/Niagara/Source/Niagara/Private/NiagaraParameterStore.cpp b/Engine/Plugins/FX/Niagara/Source/Niagara/Private/NiagaraParameterStore.cpp
--- a/Engine/Plugins/FX/Niagara/Source/Niagara/Private/NiagaraParameterStore.cpp
+++ b/Engine/Plugins/FX/Niagara/Source/Niagara/Private/NiagaraParameterStore.cpp
@@ -515,7 +515,11 @@ Returns true if we added a new parameter.
 bool FNiagaraParameterStore::AddParameter(const FNiagaraVariable& Param, bool bInitInterfaces, bool bTriggerRebind, int32* OutOffset)
 {
 #if WITH_EDITORONLY_DATA
-	check(!ParameterOffsets.Num()); // Migration to SortedParameterOffsets
+	if (ParameterOffsets.Num())
+	{
+		UE_LOG(LogNiagara, Warning, TEXT("ParameterOffsets.Num() should be 0 is %d, please investigate for %s"), ParameterOffsets.Num() , *GetPathNameSafe(Owner));
+	}
+	//check(!ParameterOffsets.Num()); // Migration to SortedParameterOffsets
 #endif
 
 	auto ParameterVariables = ReadParameterVariables();
@@ -593,7 +597,11 @@ bool FNiagaraParameterStore::AddParameter(const FNiagaraVariable& Param, bool bI
 bool FNiagaraParameterStore::RemoveParameter(const FNiagaraVariableBase& ToRemove)
 {
 #if WITH_EDITORONLY_DATA
-	check(!ParameterOffsets.Num()); // Migration to SortedParameterOffsets
+	if (ParameterOffsets.Num())
+	{
+		UE_LOG(LogNiagara, Warning, TEXT("ParameterOffsets.Num() should be 0 is %d, please investigate for %s"), ParameterOffsets.Num(), *GetPathNameSafe(Owner));
+	}
+	//check(!ParameterOffsets.Num()); // Migration to SortedParameterOffsets
 #endif
 
 	if (IndexOf(ToRemove) != INDEX_NONE)
@@ -650,7 +658,11 @@ bool FNiagaraParameterStore::RemoveParameter(const FNiagaraVariableBase& ToRemov
 void FNiagaraParameterStore::RenameParameter(const FNiagaraVariableBase& Param, FName NewName)
 {
 #if WITH_EDITORONLY_DATA
-	check(!ParameterOffsets.Num()); // Migration to SortedParameterOffsets
+	if (ParameterOffsets.Num())
+	{
+		UE_LOG(LogNiagara, Warning, TEXT("ParameterOffsets.Num() should be 0 is %d, please investigate for %s"), ParameterOffsets.Num(), *GetPathNameSafe(Owner));
+	}
+	//check(!ParameterOffsets.Num()); // Migration to SortedParameterOffsets
 #endif
 
 	if (Param.GetName() == NewName)
@@ -994,7 +1006,11 @@ const FNiagaraVariableBase* FNiagaraParameterStore::FindVariable(const UNiagaraD
 const int32* FNiagaraParameterStore::FindParameterOffset(const FNiagaraVariableBase& Parameter, bool IgnoreType) const
 {
 #if WITH_EDITORONLY_DATA
-	check(!ParameterOffsets.Num()); // Migration to SortedParameterOffsets
+	if (ParameterOffsets.Num())
+	{
+		UE_LOG(LogNiagara, Warning, TEXT("ParameterOffsets.Num() should be 0 is %d, please investigate for %s"), ParameterOffsets.Num(), *GetPathNameSafe(Owner));
+	}
+	//check(!ParameterOffsets.Num()); // Migration to SortedParameterOffsets
 #endif
 
 	auto ParameterVariables = ReadParameterVariables();
