- Editgrate CL 15497574: Flipboo updates

#rnx

diff --git a/Engine/Plugins/FX/Niagara/Source/Niagara/Classes/NiagaraFlipbookSettings.h b/Engine/Plugins/FX/Niagara/Source/Niagara/Classes/NiagaraFlipbookSettings.h
--- a/Engine/Plugins/FX/Niagara/Source/Niagara/Classes/NiagaraFlipbookSettings.h
+++ b/Engine/Plugins/FX/Niagara/Source/Niagara/Classes/NiagaraFlipbookSettings.h
@@ -40,7 +40,7 @@ struct FNiagaraFlipbookTextureSettings
 	/** Optional output name, if left empty a name will be auto generated using the index of the texture/ */
 	UPROPERTY(EditAnywhere, Category = "Texture")
 	FName OutputName;
-
+	
 	/** Source visualization we should capture, i.e. Scene Color, World Normal, etc */
 	UPROPERTY(EditAnywhere, Category = "Texture")
 	FNiagaraFlipbookTextureSource SourceBinding;
@@ -63,6 +63,8 @@ struct FNiagaraFlipbookTextureSettings
 	/** Final texture generated, an existing entry will be updated with new capture data. */
 	UPROPERTY(EditAnywhere, Category = "Texture")
 	UTexture2D* GeneratedTexture = nullptr;
+
+	bool Equals(const FNiagaraFlipbookTextureSettings& Other) const;
 };
 
 UCLASS()
@@ -71,9 +73,20 @@ class NIAGARA_API UNiagaraFlipbookSettings : public UObject
 	GENERATED_BODY()
 
 public:
+	struct FDisplayInfo
+	{
+		float NormalizedTime;
+		int FrameIndexA;
+		int FrameIndexB;
+		float Interp;
+	};
+
 	UNiagaraFlipbookSettings(const FObjectInitializer& Init);
 
-	/** Time we start the capture at. */
+	/**
+	This is the start time of the simultion where we being the capture.
+	I.e. 2.0 would mean the simulation warms up by 2 seconds before we begin capturing.
+	*/
 	UPROPERTY(EditAnywhere, Category="Timeline")
 	float StartSeconds = 0.0f;
 
@@ -81,6 +94,18 @@ public:
 	UPROPERTY(EditAnywhere, Category = "Timeline")
 	float DurationSeconds = 4.0f;
 
+	/**
+	The frame rate to run the simulation at during capturing.
+	This is only used for the preview view and calculating the number of ticks to execute
+	as we capture the generated texture.
+	*/
+	UPROPERTY(EditAnywhere, Category = "Timeline", meta = (ClampMin=1, ClampMax=480))
+	int FramesPerSecond = 60;
+
+	/** Should the preview playback as looping or not. */
+	UPROPERTY(EditAnywhere, Category = "Preview")
+	uint8 bPreviewLooping : 1;
+
 	/** Number of frames in each dimension. */
 	UPROPERTY(EditAnywhere, Category = "Texture")
 	FIntPoint FramesPerDimension = FIntPoint(8, 8);
@@ -101,30 +126,55 @@ public:
 	UPROPERTY(EditAnywhere, Category = "Camera")
 	FRotator CameraViewportRotation[(int)ENiagaraFlipbookViewMode::Num];
 
+	/** Perspective camera orbit distance. */
+	UPROPERTY(EditAnywhere, Category = "Camera", meta = (EditCondition = "CameraViewportMode == ENiagaraFlipbookViewMode::Perspective", ClampMin = "0.01"))
+	float CameraOrbitDistance = 200.f;
+
 	/** Camera FOV to use when in perspective mode. */
-	UPROPERTY(EditAnywhere, Category = "Camera", meta=(EditCondition="CameraViewportMode == ENiagaraFlipbookViewMode::Perspective"))
+	UPROPERTY(EditAnywhere, Category = "Camera", meta=(EditCondition="CameraViewportMode == ENiagaraFlipbookViewMode::Perspective", ClampMin="1.0", ClampMax="180.0"))
 	float CameraFOV = 90.0f;
 
-	/** Camera Orthographic size to use with in orthographic mode. */
-	UPROPERTY(EditAnywhere, Category = "Camera", meta=(EditCondition="CameraViewportMode != ENiagaraFlipbookViewMode::Perspective"))
-	FVector2D CameraOrthoSize = FVector2D(512.0f, 512.0f);
+	/** Camera Orthographic width to use with in orthographic mode. */
+	UPROPERTY(EditAnywhere, Category = "Camera", meta = (EditCondition = "CameraViewportMode != ENiagaraFlipbookViewMode::Perspective", ClampMin="1.0"))
+	float CameraOrthoWidth = 512.0f;
+
+	UPROPERTY(EditAnywhere, Category = "Camera", meta = (PinHiddenByDefault, InlineEditConditionToggle))
+	uint8 bUseCameraAspectRatio : 1;
 
-	/** Should we render just the component or the whole scene . */
+	/** Custom aspect ratio to use rather than using the width & height to automatically calculate. */
+	UPROPERTY(EditAnywhere, Category = "Camera", meta = (EditCondition = "bUseCameraAspectRatio", ClampMin="0.01"))
+	float CameraAspectRatio = 1.0f;
+
+	/** Should we render just the component or the whole scene. */
 	UPROPERTY(EditAnywhere, Category = "Environment")
 	uint8 bRenderComponentOnly : 1;
 
-	/** Type of level setup to use. */
-	UPROPERTY(EditAnywhere, Category = "Environment", meta = (PinHiddenByDefault, InlineEditConditionToggle))
-	uint8 bLoadLevel : 1;
+	///** Type of level setup to use. */
+	//UPROPERTY(EditAnywhere, Category = "Environment", meta = (PinHiddenByDefault, InlineEditConditionToggle))
+	//uint8 bLoadLevel : 1;
+
+	///** Used to determine type of level setup. */
+	//UPROPERTY(EditAnywhere, Category = "Environment", meta = (EditCondition = "bLoadLevel"))
+	//TSoftObjectPtr<class ULevel> LevelEnvironment;
+
+	bool Equals(const UNiagaraFlipbookSettings& Other) const;
 
-	/** Used to determine type of level setup. */
-	UPROPERTY(EditAnywhere, Category = "Environment", meta = (EditCondition = "bLoadLevel"))
-	TSoftObjectPtr<class ULevel> LevelEnvironment;
+	int GetNumFrames() const { return FramesPerDimension.X * FramesPerDimension.Y; }
 
+	float GetSeekDelta() const { return 1.0f / float(FramesPerSecond); }
+
+	float GetAspectRatio(int32 iOutputTextureIndex) const;
+	FVector2D GetOrthoSize(int32 iOutputTextureIndex) const;
 	FVector GetCameraLocation() const;
 	FMatrix GetViewMatrix() const;
 	FMatrix GetProjectionMatrixForTexture(int32 iOutputTextureIndex) const;
 
+	bool IsOrthographic() const { return CameraViewportMode != ENiagaraFlipbookViewMode::Perspective; }
+	bool IsPerspective() const { return CameraViewportMode == ENiagaraFlipbookViewMode::Perspective; }
+
+	// Get display info, the input time is expected to tbe relative, i.e. StartDuration is not taking into account
+	FDisplayInfo GetDisplayInfo(float Time, bool bLooping) const;
+
 	virtual void PostLoad() override;
 
 #if WITH_EDITORONLY_DATA
diff --git a/Engine/Plugins/FX/Niagara/Source/Niagara/Classes/NiagaraSystem.h b/Engine/Plugins/FX/Niagara/Source/Niagara/Classes/NiagaraSystem.h
--- a/Engine/Plugins/FX/Niagara/Source/Niagara/Classes/NiagaraSystem.h
+++ b/Engine/Plugins/FX/Niagara/Source/Niagara/Classes/NiagaraSystem.h
@@ -515,6 +515,8 @@ public:
 
 #if WITH_EDITORONLY_DATA
 	UNiagaraFlipbookSettings* GetFlipbookSettings();
+	const UNiagaraFlipbookSettings* GetFlipbookGeneratedSettings() const { return FlipbookGeneratedSettings; }
+	void SetFlipbookGeneratedSettings(UNiagaraFlipbookSettings* Settings) { FlipbookGeneratedSettings = Settings; }
 #endif
 
 private:
@@ -634,9 +636,13 @@ protected:
 	float WarmupTickDelta;
 
 #if WITH_EDITORONLY_DATA
-	/** Flipbook generation settings. */
-	UPROPERTY(EditAnywhere, Export, Category = Flipbook, meta = (EditInline))
+	/** Settings used when generating the flipbook */
+	UPROPERTY(Export)
 	UNiagaraFlipbookSettings* FlipbookSettings;
+
+	/** Generated flipbook settings, will be null until we have generated at least once. */
+	UPROPERTY(Export)
+	UNiagaraFlipbookSettings* FlipbookGeneratedSettings;
 #endif
 
 	UPROPERTY()
diff --git a/Engine/Plugins/FX/Niagara/Source/Niagara/Private/NiagaraFlipbookSettings.cpp b/Engine/Plugins/FX/Niagara/Source/Niagara/Private/NiagaraFlipbookSettings.cpp
--- a/Engine/Plugins/FX/Niagara/Source/Niagara/Private/NiagaraFlipbookSettings.cpp
+++ b/Engine/Plugins/FX/Niagara/Source/Niagara/Private/NiagaraFlipbookSettings.cpp
@@ -5,6 +5,16 @@
 #include "CoreMinimal.h"
 #include "Engine/Texture2D.h"
 
+bool FNiagaraFlipbookTextureSettings::Equals(const FNiagaraFlipbookTextureSettings& Other) const
+{
+	return
+		OutputName == Other.OutputName &&
+		SourceBinding.SourceName == Other.SourceBinding.SourceName &&
+		bUseFrameSize == Other.bUseFrameSize &&
+		FrameSize == Other.FrameSize &&
+		TextureSize == Other.TextureSize;
+}
+
 UNiagaraFlipbookSettings::UNiagaraFlipbookSettings(const FObjectInitializer& Init)
 	: Super(Init)
 {
@@ -17,12 +27,74 @@ UNiagaraFlipbookSettings::UNiagaraFlipbookSettings(const FObjectInitializer& Ini
 		CameraViewportRotation[i] = FRotator::ZeroRotator;
 	}
 	CameraViewportLocation[(int)ENiagaraFlipbookViewMode::Perspective] = FVector(0.0f, -200.0f, 0.0f);
-	CameraViewportRotation[(int)ENiagaraFlipbookViewMode::Perspective] = FRotator(0.0f, 0.0f, 90.0f);
+	CameraViewportRotation[(int)ENiagaraFlipbookViewMode::Perspective] = FRotator(180.0f, 0.0f, 90.0f);
+}
+
+bool UNiagaraFlipbookSettings::Equals(const UNiagaraFlipbookSettings& Other) const
+{
+	if (OutputTextures.Num() != Other.OutputTextures.Num())
+	{
+		return false;
+	}
+
+	for ( int i=0; i < OutputTextures.Num(); ++i )
+	{
+		if ( !OutputTextures[i].Equals(Other.OutputTextures[i]) )
+		{
+			return false;
+		}
+	}
+
+	for ( int i=0; i < (int)ENiagaraFlipbookViewMode::Num; ++i )
+	{
+		if ( !CameraViewportLocation[i].Equals(Other.CameraViewportLocation[i]) ||
+			 !CameraViewportRotation[i].Equals(Other.CameraViewportRotation[i]) )
+		{
+			return false;
+		}
+	}
+
+	return
+		FMath::IsNearlyEqual(StartSeconds, Other.StartSeconds) &&
+		FMath::IsNearlyEqual(DurationSeconds, Other.DurationSeconds) &&
+		FramesPerSecond == Other.FramesPerSecond &&
+		bPreviewLooping == Other.bPreviewLooping &&
+		FramesPerDimension == Other.FramesPerDimension &&
+		CameraViewportMode == Other.CameraViewportMode &&
+		FMath::IsNearlyEqual(CameraOrbitDistance, Other.CameraOrbitDistance) &&
+		FMath::IsNearlyEqual(CameraFOV, Other.CameraFOV) &&
+		FMath::IsNearlyEqual(CameraOrthoWidth, Other.CameraOrthoWidth) &&
+		bUseCameraAspectRatio == Other.bUseCameraAspectRatio &&
+		FMath::IsNearlyEqual(CameraAspectRatio, Other.CameraAspectRatio) &&
+		bRenderComponentOnly == Other.bRenderComponentOnly;
+}
+
+float UNiagaraFlipbookSettings::GetAspectRatio(int32 iOutputTextureIndex) const
+{
+	if (OutputTextures.IsValidIndex(iOutputTextureIndex))
+	{
+		const float TextureAspectRatio = float(OutputTextures[iOutputTextureIndex].FrameSize.Y) / float(OutputTextures[iOutputTextureIndex].FrameSize.X);
+		return bUseCameraAspectRatio ? CameraAspectRatio : TextureAspectRatio;
+	}
+	return 1.0f;
+}
+
+FVector2D UNiagaraFlipbookSettings::GetOrthoSize(int32 iOutputTextureIndex) const
+{
+	return FVector2D(CameraOrthoWidth, CameraOrthoWidth * GetAspectRatio(iOutputTextureIndex));
 }
 
 FVector UNiagaraFlipbookSettings::GetCameraLocation() const
 {
-	return CameraViewportLocation[(int)CameraViewportMode];
+	if (IsPerspective())
+	{
+		const FVector OrbitOffset = CameraViewportRotation[(int)CameraViewportMode].RotateVector(FVector(0.0f, 0.0f, CameraOrbitDistance));
+		return CameraViewportLocation[(int)CameraViewportMode] - OrbitOffset;
+	}
+	else
+	{
+		return CameraViewportLocation[(int)CameraViewportMode];
+	}
 }
 
 FMatrix UNiagaraFlipbookSettings::GetViewMatrix() const
@@ -47,15 +119,34 @@ FMatrix UNiagaraFlipbookSettings::GetProjectionMatrixForTexture(int32 iOutputTex
 {
 	if (CameraViewportMode == ENiagaraFlipbookViewMode::Perspective)
 	{
-		return FReversedZPerspectiveMatrix(CameraFOV, OutputTextures[iOutputTextureIndex].FrameSize.X, OutputTextures[iOutputTextureIndex].FrameSize.Y, 1.0f);
+		const float AspectRatio = GetAspectRatio(iOutputTextureIndex);
+		const float HalfXFOV = FMath::DegreesToRadians(CameraFOV) * 0.5f;
+		const float HalfYFOV = FMath::Atan(FMath::Tan(HalfXFOV) / AspectRatio);
+		return FReversedZPerspectiveMatrix(HalfXFOV, HalfYFOV, 1.0f, 1.0, GNearClippingPlane, GNearClippingPlane);
 	}
 	else
 	{
 		const float ZRange = WORLD_MAX;
-		return FReversedZOrthoMatrix(CameraOrthoSize.X / 2.0f, CameraOrthoSize.Y / 2.0f, 0.5f / ZRange, ZRange);
+		const FVector2D OrthoSize = GetOrthoSize(iOutputTextureIndex);
+		return FReversedZOrthoMatrix(OrthoSize.X / 2.0f, OrthoSize.Y / 2.0f, 0.5f / ZRange, ZRange);
 	}
 }
 
+UNiagaraFlipbookSettings::FDisplayInfo UNiagaraFlipbookSettings::GetDisplayInfo(float Time, bool bLooping) const
+{
+	FDisplayInfo DisplayInfo;
+	DisplayInfo.NormalizedTime = FMath::Max(Time / DurationSeconds, 0.0f);
+	DisplayInfo.NormalizedTime = bLooping ? FMath::Fractional(DisplayInfo.NormalizedTime) : FMath::Min(DisplayInfo.NormalizedTime, 0.9999f);
+
+	const int NumFrames = GetNumFrames();
+	const float FrameTime = DisplayInfo.NormalizedTime * float(NumFrames);
+	DisplayInfo.FrameIndexA = FMath::FloorToInt(FrameTime);
+	DisplayInfo.FrameIndexB = bLooping ? (DisplayInfo.FrameIndexA + 1) % NumFrames : FMath::Min(DisplayInfo.FrameIndexA + 1, NumFrames - 1);
+	DisplayInfo.Interp = FrameTime - float(DisplayInfo.FrameIndexA);
+
+	return DisplayInfo;
+}
+
 void UNiagaraFlipbookSettings::PostLoad()
 {
 	Super::PostLoad();
diff --git a/Engine/Plugins/FX/Niagara/Source/NiagaraEditor/Private/ViewModels/NiagaraFlipbookViewModel.cpp b/Engine/Plugins/FX/Niagara/Source/NiagaraEditor/Private/ViewModels/NiagaraFlipbookViewModel.cpp
--- a/Engine/Plugins/FX/Niagara/Source/NiagaraEditor/Private/ViewModels/NiagaraFlipbookViewModel.cpp
+++ b/Engine/Plugins/FX/Niagara/Source/NiagaraEditor/Private/ViewModels/NiagaraFlipbookViewModel.cpp
@@ -57,7 +57,6 @@ void FNiagaraFlipbookViewModel::Initialize(TWeakPtr<FNiagaraSystemViewModel> InW
 		PreviewComponent->SetForceSolo(true);
 		PreviewComponent->SetAgeUpdateMode(ENiagaraAgeUpdateMode::DesiredAge);
 		PreviewComponent->SetCanRenderWhileSeeking(false);
-		PreviewComponent->SetAgeUpdateMode(ENiagaraAgeUpdateMode::DesiredAge);
 		PreviewComponent->Activate(true);
 
 		AdvancedPreviewScene = MakeShareable(new FAdvancedPreviewScene(FPreviewScene::ConstructionValues()));
@@ -74,63 +73,13 @@ void FNiagaraFlipbookViewModel::RefreshView()
 	//-TOOD:
 }
 
-FNiagaraFlipbookViewModel::FDisplayData FNiagaraFlipbookViewModel::GetDisplayDataFromAbsoluteTime(float AbsoluteTime) const
-{
-	FDisplayData DisplayData;
-	if (IsGeneratedDataValid())
-	{
-		DisplayData.NumFrames = GeneratedNumFrames;
-		DisplayData.StartSeconds = GeneratedStartSeconds;
-		DisplayData.DurationSeconds = GeneratedDurationSeconds;
-	}
-	else if ( UNiagaraFlipbookSettings* FlipbookSettings = GetFlipbookSettings() )
-	{
-		DisplayData.NumFrames = FlipbookSettings->FramesPerDimension.X * FlipbookSettings->FramesPerDimension.Y;
-		DisplayData.StartSeconds = FlipbookSettings->StartSeconds;
-		DisplayData.DurationSeconds = FlipbookSettings->DurationSeconds;
-	}
-
-	if ( DisplayData.DurationSeconds > 0.0f )
-	{
-		DisplayData.NormalizedTime = FMath::Clamp((AbsoluteTime - DisplayData.StartSeconds) / DisplayData.DurationSeconds, 0.0f, 1.0f);
-		DisplayData.FrameTime = DisplayData.NormalizedTime * float(DisplayData.NumFrames);
-		DisplayData.FrameIndex = FMath::FloorToInt(DisplayData.FrameTime);
-	}
-	return DisplayData;
-}
-
-FNiagaraFlipbookViewModel::FDisplayData FNiagaraFlipbookViewModel::GetDisplayDataFromRelativeTime(float RelativeTime) const
-{
-	FDisplayData DisplayData;
-	if (IsGeneratedDataValid())
-	{
-		DisplayData.NumFrames = GeneratedNumFrames;
-		DisplayData.StartSeconds = GeneratedStartSeconds;
-		DisplayData.DurationSeconds = GeneratedDurationSeconds;
-	}
-	else if (UNiagaraFlipbookSettings* FlipbookSettings = GetFlipbookSettings())
-	{
-		DisplayData.NumFrames = FlipbookSettings->FramesPerDimension.X * FlipbookSettings->FramesPerDimension.Y;
-		DisplayData.StartSeconds = FlipbookSettings->StartSeconds;
-		DisplayData.DurationSeconds = FlipbookSettings->DurationSeconds;
-	}
-
-	if (DisplayData.DurationSeconds > 0.0f)
-	{
-		DisplayData.NormalizedTime = FMath::Clamp(RelativeTime / DisplayData.DurationSeconds, 0.0f, 1.0f);
-		DisplayData.FrameTime = DisplayData.NormalizedTime * float(DisplayData.NumFrames);
-		DisplayData.FrameIndex = FMath::FloorToInt(DisplayData.FrameTime);
-	}
-	return DisplayData;
-}
-
 void FNiagaraFlipbookViewModel::SetDisplayTimeFromNormalized(float NormalizeTime)
 {
 	if ( Widget )
 	{
-		if ( IsGeneratedDataValid() )
+		if (const UNiagaraFlipbookSettings* GeneratedSettings = GetFlipbookGeneratedSettings())
 		{
-			const float RelativeTime = FMath::Clamp(NormalizeTime, 0.0f, 1.0f) * GeneratedDurationSeconds;
+			const float RelativeTime = FMath::Clamp(NormalizeTime, 0.0f, 1.0f) * GeneratedSettings->DurationSeconds;
 			Widget->SetPreviewRelativeTime(RelativeTime);
 		}
 		else if (UNiagaraFlipbookSettings* FlipbookSettings = GetFlipbookSettings())
@@ -157,30 +106,40 @@ UNiagaraFlipbookSettings* FNiagaraFlipbookViewModel::GetFlipbookSettings() const
 	return nullptr;
 }
 
-TSharedPtr<class SWidget> FNiagaraFlipbookViewModel::GetWidget()
+const UNiagaraFlipbookSettings* FNiagaraFlipbookViewModel::GetFlipbookGeneratedSettings() const
 {
-	return Widget;
+	if (PreviewComponent)
+	{
+		UNiagaraSystem* Asset = PreviewComponent->GetAsset();
+		return Asset ? Asset->GetFlipbookGeneratedSettings() : nullptr;
+	}
+
+	return nullptr;
 }
 
-void FNiagaraFlipbookViewModel::AddReferencedObjects(FReferenceCollector& Collector)
+const FNiagaraFlipbookTextureSettings* FNiagaraFlipbookViewModel::GetPreviewTextureSettings() const
 {
-	for (FNiagaraFlipbookTextureSettings& Texture : GeneratedTextures)
+	if ( UNiagaraFlipbookSettings* Settings = GetFlipbookSettings() )
 	{
-		if (Texture.GeneratedTexture != nullptr)
+		if ( Settings->OutputTextures.IsValidIndex(PreviewTextureIndex) )
 		{
-			Collector.AddReferencedObject(Texture.GeneratedTexture);
+			return &Settings->OutputTextures[PreviewTextureIndex];
 		}
 	}
+	return nullptr;
 }
 
-void FNiagaraFlipbookViewModel::RenderFlipbook()
+TSharedPtr<class SWidget> FNiagaraFlipbookViewModel::GetWidget()
 {
-	GeneratedNumFrames = 0;
-	GeneratedStartSeconds = 0.0f;
-	GeneratedDurationSeconds = 0.0f;
-	GeneratedFramesPerDimension = FIntPoint(0, 0);
-	GeneratedTextures.Empty();
+	return Widget;
+}
+
+void FNiagaraFlipbookViewModel::AddReferencedObjects(FReferenceCollector& Collector)
+{
+}
 
+void FNiagaraFlipbookViewModel::RenderFlipbook()
+{
 	UNiagaraFlipbookSettings* FlipbookSettings = GetFlipbookSettings();
 	if ( (PreviewComponent == nullptr) || (FlipbookSettings == nullptr) )
 	{
@@ -233,8 +192,9 @@ void FNiagaraFlipbookViewModel::RenderFlipbook()
 
 		// Step component to time
 		const float FrameTime = FlipbookSettings->StartSeconds + (float(iFrame) * FrameDeltaSeconds);
+		PreviewComponent->SetSeekDelta(FlipbookSettings->GetSeekDelta());
 		PreviewComponent->SeekToDesiredAge(FrameTime);
-		PreviewComponent->TickComponent(1.0f / 60.0f, ELevelTick::LEVELTICK_All, nullptr);
+		PreviewComponent->TickComponent(FlipbookSettings->GetSeekDelta(), ELevelTick::LEVELTICK_All, nullptr);
 
 		// We need any GPU sims to flush pending ticks
 		if (NiagaraBatcher)
@@ -277,11 +237,6 @@ void FNiagaraFlipbookViewModel::RenderFlipbook()
 	}
 
 	// Send data to generated textures
-	GeneratedNumFrames = TotalFrames;
-	GeneratedStartSeconds = FlipbookSettings->StartSeconds;
-	GeneratedDurationSeconds = FlipbookSettings->DurationSeconds;
-	GeneratedFramesPerDimension = FlipbookSettings->FramesPerDimension;
-
 	for (int32 iOutputTexture = 0; iOutputTexture < FlipbookSettings->OutputTextures.Num(); ++iOutputTexture)
 	{
 		FNiagaraFlipbookTextureSettings& OutputTexture = FlipbookSettings->OutputTextures[iOutputTexture];
@@ -313,8 +268,12 @@ void FNiagaraFlipbookViewModel::RenderFlipbook()
 		OutputTexture.GeneratedTexture->UpdateResource();
 		OutputTexture.GeneratedTexture->PostEditChange();
 		OutputTexture.GeneratedTexture->MarkPackageDirty();
+	}
 
-		GeneratedTextures.Add(OutputTexture);
+	// Duplicate and set as generated data
+	if (UNiagaraSystem* Asset = PreviewComponent->GetAsset())
+	{
+		Asset->SetFlipbookGeneratedSettings(DuplicateObject<UNiagaraFlipbookSettings>(FlipbookSettings, Asset));
 	}
 
 	// Clean up render targets
diff --git a/Engine/Plugins/FX/Niagara/Source/NiagaraEditor/Private/ViewModels/NiagaraFlipbookViewModel.h b/Engine/Plugins/FX/Niagara/Source/NiagaraEditor/Private/ViewModels/NiagaraFlipbookViewModel.h
--- a/Engine/Plugins/FX/Niagara/Source/NiagaraEditor/Private/ViewModels/NiagaraFlipbookViewModel.h
+++ b/Engine/Plugins/FX/Niagara/Source/NiagaraEditor/Private/ViewModels/NiagaraFlipbookViewModel.h
@@ -13,15 +13,6 @@
 class FNiagaraFlipbookViewModel : public TSharedFromThis<FNiagaraFlipbookViewModel>, public FGCObject
 {
 public:
-	struct FGeneratedData
-	{
-		UTexture2D*	Texture = nullptr;
-		int			TotalFrames = 0;
-		FIntPoint	FrameSize = FIntPoint::ZeroValue;
-		FIntPoint	FramesPerDimension = FIntPoint::ZeroValue;
-		float		DurationSeconds = 0.0f;
-	};
-
 	struct FDisplayData
 	{
 		int32		NumFrames = 0;
@@ -45,22 +36,14 @@ public:
 
 	int32 GetPreviewTextureIndex() const { return PreviewTextureIndex; }
 	void SetPreviewTextureIndex(int32 InPreviewTextureIndex) { PreviewTextureIndex = InPreviewTextureIndex; }
-	const FNiagaraFlipbookTextureSettings* GetPreviewTexture() const { return GeneratedTextures.IsValidIndex(PreviewTextureIndex) ? &GeneratedTextures[PreviewTextureIndex] : nullptr; }
-
-	bool IsGeneratedDataValid() const { return GeneratedNumFrames > 0; }
-	int32 GetGeneratedNumFrames() const { return GeneratedNumFrames; }
-	float GetGeneratedStartSeconds() const { return GeneratedStartSeconds; }
-	float GetGeneratedDurationSeconds() const { return GeneratedDurationSeconds; }
-	FIntPoint GetGeneratedFramesPerDimension() const { return GeneratedFramesPerDimension; }
-	TConstArrayView<FNiagaraFlipbookTextureSettings> GetGeneratedTextures() const { return MakeArrayView(GeneratedTextures); }
-
-	FDisplayData GetDisplayDataFromAbsoluteTime(float AbsoluteTime) const;
-	FDisplayData GetDisplayDataFromRelativeTime(float RelativeTime) const;
 
 	void SetDisplayTimeFromNormalized(float NormalizeTime);
 
 	class UNiagaraComponent* GetPreviewComponent() const;
 	class UNiagaraFlipbookSettings* GetFlipbookSettings() const;
+	const class UNiagaraFlipbookSettings* GetFlipbookGeneratedSettings() const;
+
+	const struct FNiagaraFlipbookTextureSettings* GetPreviewTextureSettings() const;
 
 	bool RenderView(const FRenderTarget* RenderTarget, FCanvas* Canvas, float WorldTime, int32 iOutputTextureIndex, bool bFillCanvas = false) const;
 
@@ -75,10 +58,4 @@ private:
 	TSharedPtr<class SNiagaraFlipbookWidget> Widget;
 
 	int32 PreviewTextureIndex = 0;
-
-	int32 GeneratedNumFrames = 0;
-	float GeneratedStartSeconds = 0.0f;
-	float GeneratedDurationSeconds = 0.0f;
-	FIntPoint GeneratedFramesPerDimension = FIntPoint(0, 0);
-	TArray<FNiagaraFlipbookTextureSettings> GeneratedTextures;
 };
diff --git a/Engine/Plugins/FX/Niagara/Source/NiagaraEditor/Private/Widgets/SNiagaraFlipbookTimelineWidget.cpp b/Engine/Plugins/FX/Niagara/Source/NiagaraEditor/Private/Widgets/SNiagaraFlipbookTimelineWidget.cpp
--- a/Engine/Plugins/FX/Niagara/Source/NiagaraEditor/Private/Widgets/SNiagaraFlipbookTimelineWidget.cpp
+++ b/Engine/Plugins/FX/Niagara/Source/NiagaraEditor/Private/Widgets/SNiagaraFlipbookTimelineWidget.cpp
@@ -46,7 +46,8 @@ int32 SNiagaraFlipbookTimelineWidget::OnPaint(const FPaintArgs& Args, const FGeo
 	static const FLinearColor SelectedBarColor(FLinearColor::White);
 
 	// Gather information to display
-	const auto FlipbookDisplayData = ViewModel->GetDisplayDataFromRelativeTime(RelativeTime);
+	const int NumFrames = FlipbookSettings->GetNumFrames();
+	const auto FlipbookDisplayData = FlipbookSettings->GetDisplayInfo(RelativeTime, FlipbookSettings->bPreviewLooping);
 
 	//////////////////////////////////////////////////////////////////////////
 	// Draw boxes to show frames
@@ -54,15 +55,15 @@ int32 SNiagaraFlipbookTimelineWidget::OnPaint(const FPaintArgs& Args, const FGeo
 		const FSlateBrush* BoxBrush = FEditorStyle::GetBrush("Sequencer.Section.BackgroundTint");
 		const FLinearColor BoxTints[2] = { FLinearColor(0.5f, 0.5f, 0.5f, 1.0f), FLinearColor(0.3f, 0.3f, 0.3f, 1.0f) };
 		const FLinearColor CurrentTint(0.5f, 0.5f, 1.0f, 1.0f);
-		if (FlipbookDisplayData.NumFrames > 0)
+		if (NumFrames > 0)
 		{
-			const float UStep = AllottedGeometry.Size.X / float(FlipbookDisplayData.NumFrames);
+			const float UStep = AllottedGeometry.Size.X / float(NumFrames);
 
-			for (int32 i = 0; i < FlipbookDisplayData.NumFrames; ++i)
+			for (int32 i = 0; i < NumFrames; ++i)
 			{
 				const FVector2D BoxLocation(UStep * float(i), 0.0f);
 				const FVector2D BoxSize(UStep, AllottedGeometry.Size.Y);
-				const FLinearColor& Tint = (i == FlipbookDisplayData.FrameIndex) ? CurrentTint : BoxTints[i & 1];
+				const FLinearColor& Tint = (i == FlipbookDisplayData.FrameIndexA) ? CurrentTint : BoxTints[i & 1];
 				FSlateDrawElement::MakeBox(OutDrawElements, RetLayerId++, AllottedGeometry.ToPaintGeometry(BoxLocation, BoxSize), BoxBrush, ESlateDrawEffect::None, Tint);
 			}
 		}
diff --git a/Engine/Plugins/FX/Niagara/Source/NiagaraEditor/Private/Widgets/SNiagaraFlipbookViewport.cpp b/Engine/Plugins/FX/Niagara/Source/NiagaraEditor/Private/Widgets/SNiagaraFlipbookViewport.cpp
--- a/Engine/Plugins/FX/Niagara/Source/NiagaraEditor/Private/Widgets/SNiagaraFlipbookViewport.cpp
+++ b/Engine/Plugins/FX/Niagara/Source/NiagaraEditor/Private/Widgets/SNiagaraFlipbookViewport.cpp
@@ -10,6 +10,7 @@
 #include "Framework/MultiBox/MultiBoxBuilder.h"
 #include "Modules/ModuleManager.h"
 #include "CanvasTypes.h"
+#include "EditorViewportCommands.h"
 #include "EngineModule.h"
 #include "LegacyScreenPercentageDriver.h"
 #include "ImageUtils.h"
@@ -30,43 +31,197 @@ public:
 
 	virtual bool InputKey(FViewport* InViewport, int32 ControllerId, FKey Key, EInputEvent Event, float AmountDepressed, bool bGamepad) override
 	{
-		return false;
-	}
+		// Only allow movement when in preview mode
+		auto ViewModel = WeakViewModel.Pin();
+		UNiagaraFlipbookSettings* FlipbookSettings = ViewModel ? ViewModel->GetFlipbookSettings() : nullptr;
+		if ( bShowPreview && FlipbookSettings )
+		{
+			bool bForwardKeyState = false;
+			bool bBackwardKeyState = false;
+			bool bRightKeyState = false;
+			bool bLeftKeyState = false;
 
-	virtual void CapturedMouseMove(FViewport* InViewport, int32 X, int32 Y) override
-	{
-	}
+			bool bUpKeyState = false;
+			bool bDownKeyState = false;
+			bool bZoomOutKeyState = false;
+			bool bZoomInKeyState = false;
 
-	FIntRect ConstrainRect(const FIntRect& InRect, FIntPoint InSize)
-	{
-		const float Scale = FMath::Min(float(InRect.Width()) /  float(InSize.X), float(InRect.Height()) / float(InSize.Y));
-		const FIntPoint ScaledSize(FMath::FloorToInt(float(InSize.X) * Scale), FMath::FloorToInt(float(InSize.Y) * Scale));
-		FIntPoint RectMin(
-			InRect.Min.X + ((InRect.Width()  - ScaledSize.X) >> 1),
-			InRect.Min.Y + ((InRect.Height() - ScaledSize.Y) >> 1)
-		);
+			bool bFocus = false;
 
-		return FIntRect(RectMin, RectMin + ScaledSize);
+			// Iterate through all key mappings to generate key state flags
+			for (uint32 i = 0; i < static_cast<uint8>(EMultipleKeyBindingIndex::NumChords); ++i)
+			{
+				EMultipleKeyBindingIndex ChordIndex = static_cast<EMultipleKeyBindingIndex> (i);
+				bForwardKeyState |= Viewport->KeyState(FViewportNavigationCommands::Get().Forward->GetActiveChord(ChordIndex)->Key);
+				bBackwardKeyState |= Viewport->KeyState(FViewportNavigationCommands::Get().Backward->GetActiveChord(ChordIndex)->Key);
+				bRightKeyState |= Viewport->KeyState(FViewportNavigationCommands::Get().Right->GetActiveChord(ChordIndex)->Key);
+				bLeftKeyState |= Viewport->KeyState(FViewportNavigationCommands::Get().Left->GetActiveChord(ChordIndex)->Key);
+
+				bUpKeyState |= Viewport->KeyState(FViewportNavigationCommands::Get().Up->GetActiveChord(ChordIndex)->Key);
+				bDownKeyState |= Viewport->KeyState(FViewportNavigationCommands::Get().Down->GetActiveChord(ChordIndex)->Key);
+				bZoomOutKeyState |= Viewport->KeyState(FViewportNavigationCommands::Get().FovZoomOut->GetActiveChord(ChordIndex)->Key);
+				bZoomInKeyState |= Viewport->KeyState(FViewportNavigationCommands::Get().FovZoomIn->GetActiveChord(ChordIndex)->Key);
+
+				bFocus |= Viewport->KeyState(FEditorViewportCommands::Get().FocusViewportToSelection->GetActiveChord(ChordIndex)->Key);
+			}
+
+			if ( FlipbookSettings->IsOrthographic() )
+			{
+				LocalMovement.X += bLeftKeyState ? KeyboardMoveSpeed : 0.0f;
+				LocalMovement.X -= bRightKeyState ? KeyboardMoveSpeed : 0.0f;
+				LocalMovement.Y += bBackwardKeyState ? KeyboardMoveSpeed : 0.0f;
+				LocalMovement.Y -= bForwardKeyState ? KeyboardMoveSpeed : 0.0f;
+			}
+			else
+			{
+				//LocalMovement.X += bLeftKeyState ? KeyboardMoveSpeed : 0.0f;
+				//LocalMovement.X -= bRightKeyState ? KeyboardMoveSpeed : 0.0f;
+				//LocalMovement.Y += bUpKeyState ? KeyboardMoveSpeed : 0.0f;
+				//LocalMovement.Y -= bDownKeyState ? KeyboardMoveSpeed : 0.0f;
+				//LocalMovement.Z += bBackwardKeyState ? KeyboardMoveSpeed : 0.0f;
+				//LocalMovement.Z -= bForwardKeyState ? KeyboardMoveSpeed : 0.0f;
+			}
+
+			LocalZoom += bZoomOutKeyState ? KeyboardMoveSpeed : 0.0f;
+			LocalZoom -= bZoomInKeyState ? KeyboardMoveSpeed : 0.0f;
+
+			// Focus
+			if (bFocus)
+			{
+				FocusCamera();
+			}
+		}
+
+		return true;
 	}
 
-	void ClearViewArea(FCanvas* Canvas, const FIntRect& InRect)
+	virtual bool InputAxis(FViewport* InViewport, int32 ControllerId, FKey Key, float Delta, float InDeltaTime, int32 NumSamples, bool bGamepad) override
 	{
-		UTexture2D* Texture = nullptr;
-		FLinearColor Color = ClearColor;
-		FVector2D EndUV(1.0f, 1.0f);
-
-		if (bShowCheckerBoard)
+		// Viewport movement only enabled when preview is enabled otherwise there is no feedback
+		auto ViewModel = WeakViewModel.Pin();
+		UNiagaraFlipbookSettings* FlipbookSettings = ViewModel ? ViewModel->GetFlipbookSettings() : nullptr;
+		if ( bShowPreview && FlipbookSettings && (Key == EKeys::MouseX || Key == EKeys::MouseY) )
 		{
-			Texture = GetCheckerboardTexture();
-			Color = FLinearColor::White;
-			EndUV.X = float(InRect.Width()) / float(FMath::Max(Texture->GetSizeX(), 1));
-			EndUV.Y = float(InRect.Height()) / float(FMath::Max(Texture->GetSizeY(), 1));
+			if (FlipbookSettings->IsOrthographic())
+			{
+				if (InViewport->KeyState(EKeys::RightMouseButton))
+				{
+					// Zoom
+					if (InViewport->KeyState(EKeys::LeftControl) || InViewport->KeyState(EKeys::RightControl))
+					{
+						LocalZoom += (Key == EKeys::MouseY) ? Delta : 0.0f;
+					}
+					// Aspect
+					else if (InViewport->KeyState(EKeys::LeftAlt) || InViewport->KeyState(EKeys::RightAlt))
+					{
+						LocalAspect += (Key == EKeys::MouseY) ? Delta : 0.0f;
+					}
+					// Move
+					else
+					{
+						LocalMovement.X += (Key == EKeys::MouseX) ? Delta : 0.0f;
+						LocalMovement.Y += (Key == EKeys::MouseY) ? Delta : 0.0f;
+					}
+				}
+			}
+			else
+			{
+				// Zoom
+				if (InViewport->KeyState(EKeys::RightMouseButton))
+				{
+					if (InViewport->KeyState(EKeys::LeftAlt) || InViewport->KeyState(EKeys::RightAlt))
+					{
+						LocalAspect += (Key == EKeys::MouseY) ? Delta : 0.0f;
+					}
+					else
+					{
+						LocalMovement.Z += (Key == EKeys::MouseY) ? Delta : 0.0f;
+					}
+				}
+				// Middle button translate orbit location
+				else if (InViewport->KeyState(EKeys::MiddleMouseButton))
+				{
+					LocalMovement.X += (Key == EKeys::MouseX) ? Delta : 0.0f;
+					LocalMovement.Y += (Key == EKeys::MouseY) ? Delta : 0.0f;
+				}
+				// Rotation
+				else if (InViewport->KeyState(EKeys::LeftMouseButton))
+				{
+					LocalRotation.X += (Key == EKeys::MouseX) ? Delta : 0.0f;
+					LocalRotation.Y += (Key == EKeys::MouseY) ? Delta : 0.0f;
+				}
+			}
 		}
-		Canvas->DrawTile(InRect.Min.X, InRect.Min.Y, InRect.Width(), InRect.Height(), 0.0f, 0.0f, EndUV.X, EndUV.Y, Color, Texture ? Texture->Resource : nullptr, false);
+
+		return true;
 	}
 
 	virtual void Tick(float DeltaSeconds) override
 	{
+		//FEditorViewportClient::Tick(DeltaSeconds);
+
+		// Apply local movement
+		auto ViewModel = WeakViewModel.Pin();
+		UNiagaraFlipbookSettings* FlipbookSettings = ViewModel ? ViewModel->GetFlipbookSettings() : nullptr;
+		if ( FlipbookSettings )
+		{
+			const FMatrix ViewMatrix = FlipbookSettings->GetViewMatrix().Inverse();
+			const FVector XAxis = ViewMatrix.GetUnitAxis(EAxis::X);
+			const FVector YAxis = ViewMatrix.GetUnitAxis(EAxis::Y);
+			const FVector ZAxis = ViewMatrix.GetUnitAxis(EAxis::Z);
+
+			FVector WorldMovement = FVector::ZeroVector;
+			if (FlipbookSettings->IsOrthographic())
+			{
+				const FVector2D MoveSpeed = GetPreviewOrthoUnits();
+
+				WorldMovement -= LocalMovement.X * MoveSpeed.X * XAxis;
+				WorldMovement -= LocalMovement.Y * MoveSpeed.Y * YAxis;
+				//WorldMovement += LocalMovement.Z * MoveSpeed * ZAxis;
+
+				FlipbookSettings->CameraViewportLocation[(int)FlipbookSettings->CameraViewportMode] += WorldMovement;
+			}
+			else
+			{
+				FRotator& WorldRotation = FlipbookSettings->CameraViewportRotation[(int)FlipbookSettings->CameraViewportMode];
+				WorldRotation.Yaw = FRotator::ClampAxis(WorldRotation.Yaw + LocalRotation.X);
+				WorldRotation.Roll = FMath::Clamp(WorldRotation.Roll + LocalRotation.Y, 0.0f, 180.0f);
+
+				const float MoveSpeed = PerspectiveMoveSpeed;
+				WorldMovement -= LocalMovement.X * MoveSpeed * XAxis;
+				WorldMovement -= LocalMovement.Y * MoveSpeed * YAxis;
+				//WorldMovement -= LocalMovement.Z * MoveSpeed * ZAxis;
+				FlipbookSettings->CameraViewportLocation[(int)FlipbookSettings->CameraViewportMode] += WorldMovement;
+
+				FlipbookSettings->CameraOrbitDistance = FMath::Max(FlipbookSettings->CameraOrbitDistance + LocalMovement.Z, 0.01f);
+			}
+
+			if (!FMath::IsNearlyZero(LocalZoom))
+			{
+				if (FlipbookSettings->IsPerspective())
+				{
+					FlipbookSettings->CameraFOV = FMath::Clamp(FlipbookSettings->CameraFOV + LocalZoom, 0.001f, 179.0f);
+				}
+				else
+				{
+					FlipbookSettings->CameraOrthoWidth = FMath::Max(FlipbookSettings->CameraOrthoWidth + LocalZoom, 1.0f);
+				}
+			}
+
+			if (!FMath::IsNearlyZero(LocalAspect))
+			{
+				if (FlipbookSettings->bUseCameraAspectRatio)
+				{
+					FlipbookSettings->CameraAspectRatio = FMath::Max(FlipbookSettings->CameraAspectRatio + (LocalAspect / 50.0f), 0.01f);
+				}
+			}
+		}
+
+		// Clear data
+		LocalMovement = FVector::ZeroVector;
+		LocalZoom = 0.0f;
+		LocalAspect = 0.0f;
+		LocalRotation = FVector::ZeroVector;
 	}
 
 	/** FViewportClient interface */
@@ -90,12 +245,12 @@ public:
 		const FVector2D TextStartOffset(5.0f, 30.0f);
 
 		const UNiagaraFlipbookSettings* FlipbookSettings = ViewModel->GetFlipbookSettings();
-		const FNiagaraFlipbookTextureSettings* PreviewTextureSettings = ViewModel->GetPreviewTexture();
+		const UNiagaraFlipbookSettings* FlipbookGeneratedSettings = ViewModel->GetFlipbookGeneratedSettings();
 		const int32 PreviewTextureIndex = ViewModel->GetPreviewTextureIndex();
 
 		// Determine view rects
-		FIntRect PreviewViewRect;
-		FIntRect FlipbookViewRect;
+		PreviewViewRect = FIntRect();
+		FlipbookViewRect = FIntRect();
 		{
 			const FIntRect ViewRect = Canvas->GetViewRect();
 			const int32 Border = 3;
@@ -115,13 +270,13 @@ public:
 
 			if ( bShowFlipbook )
 			{
-				if (FlipbookSettings->OutputTextures.IsValidIndex(PreviewTextureIndex))
+				if (FlipbookGeneratedSettings && FlipbookGeneratedSettings->OutputTextures.IsValidIndex(PreviewTextureIndex))
 				{
-					FlipbookViewRect = ConstrainRect(FIntRect(ViewRect.Max.X - ViewWidth, ViewRect.Min.Y, ViewRect.Max.X, ViewRect.Max.Y), FlipbookSettings->OutputTextures[PreviewTextureIndex].FrameSize);
+					FlipbookViewRect = ConstrainRect(FIntRect(ViewRect.Max.X - ViewWidth, ViewRect.Min.Y, ViewRect.Max.X, ViewRect.Max.Y), FlipbookGeneratedSettings->OutputTextures[PreviewTextureIndex].FrameSize);
 				}
 				else
 				{
-					FlipbookViewRect = FIntRect(ViewRect.Max.X - ViewWidth, ViewRect.Min.Y, ViewRect.Min.X, ViewRect.Max.Y);
+					FlipbookViewRect = FIntRect(ViewRect.Max.X - ViewWidth, ViewRect.Min.Y, ViewRect.Max.X, ViewRect.Max.Y);
 				}
 			}
 		}
@@ -150,15 +305,17 @@ public:
 
 				// Seek to correct time and render
 				UNiagaraComponent* PreviewComponent = ViewModel->GetPreviewComponent();
+				PreviewComponent->SetSeekDelta(FlipbookSettings->GetSeekDelta());
 				PreviewComponent->SeekToDesiredAge(WorldTime);
-				PreviewComponent->TickComponent(1.0f / 60.0f, ELevelTick::LEVELTICK_All, nullptr);
+				PreviewComponent->TickComponent(FlipbookSettings->GetSeekDelta(), ELevelTick::LEVELTICK_All, nullptr);
 
 				FNiagaraFlipbookRenderer FlipbookRenderer(PreviewComponent, WorldTime);
 				FlipbookRenderer.RenderView(RealtimeRenderTarget, PreviewTextureIndex);
 
+				const FVector2D HalfPixel(0.5f / float(OutputTexture.FrameSize.X), 0.5f / float(OutputTexture.FrameSize.Y));
 				Canvas->DrawTile(
 					PreviewViewRect.Min.X, PreviewViewRect.Min.Y, PreviewViewRect.Width(), PreviewViewRect.Height(),
-					0.0f, 0.0f, 1.0f, 1.0f,
+					HalfPixel.X, HalfPixel.Y, 1.0f - HalfPixel.X, 1.0f - HalfPixel.Y,
 					FLinearColor::White,
 					RealtimeRenderTarget->Resource,
 					false	//-TODO: Preview with alpha?
@@ -169,6 +326,9 @@ public:
 				{
 					Canvas->DrawShadowedString(TextPosition.X, TextPosition.Y, TEXT("Live Preview"), DisplayFont, FLinearColor::White);
 					TextPosition.Y += FontHeight;
+
+					Canvas->DrawShadowedString(TextPosition.X + 5.0f, TextPosition.Y, *FString::Printf(TEXT("Texture(%d) FrameSize(%d x %d)"), PreviewTextureIndex, OutputTexture.FrameSize.X, OutputTexture.FrameSize.Y), DisplayFont, FLinearColor::White);
+					TextPosition.Y += FontHeight;
 				}
 			}
 			else
@@ -185,22 +345,29 @@ public:
 		{
 			ClearViewArea(Canvas, FlipbookViewRect);
 
-			const bool bFlipbookValid = PreviewTextureSettings != nullptr;
-			const auto DisplayData = ViewModel->GetDisplayDataFromAbsoluteTime(CurrentTime);
+			const bool bFlipbookValid =
+				(FlipbookGeneratedSettings != nullptr) &&
+				FlipbookGeneratedSettings->OutputTextures.IsValidIndex(PreviewTextureIndex) &&
+				(FlipbookGeneratedSettings->OutputTextures[PreviewTextureIndex].GeneratedTexture != nullptr);
 			FVector2D TextPosition(FlipbookViewRect.Min.X + TextStartOffset.X, FlipbookViewRect.Min.Y + TextStartOffset.Y);
 
 			if (bFlipbookValid)
 			{
-				const FIntPoint TextureSize = FIntPoint(FMath::Max(PreviewTextureSettings->GeneratedTexture->GetSizeX(), 1), FMath::Max(PreviewTextureSettings->GeneratedTexture->GetSizeY(), 1));
-				const FIntPoint FramesPerDimension = ViewModel->GetGeneratedFramesPerDimension();
-				const FIntPoint FrameIndex2D = FIntPoint(DisplayData.FrameIndex % FramesPerDimension.X, DisplayData.FrameIndex / FramesPerDimension.X);
-				const FIntPoint FramePixel = FIntPoint(FrameIndex2D.X * PreviewTextureSettings->FrameSize.X, FrameIndex2D.Y * PreviewTextureSettings->FrameSize.Y);
+				const auto DisplayData = FlipbookGeneratedSettings->GetDisplayInfo(CurrentTime - FlipbookGeneratedSettings->StartSeconds, FlipbookGeneratedSettings->bPreviewLooping);
+				const FNiagaraFlipbookTextureSettings& OutputTexture = FlipbookGeneratedSettings->OutputTextures[PreviewTextureIndex];
+
+				const FIntPoint TextureSize = FIntPoint(FMath::Max(OutputTexture.GeneratedTexture->GetSizeX(), 1), FMath::Max(OutputTexture.GeneratedTexture->GetSizeY(), 1));
+				const FIntPoint FramesPerDimension = FlipbookGeneratedSettings->FramesPerDimension;
+				const FIntPoint FrameIndexA = FIntPoint(DisplayData.FrameIndexA % FramesPerDimension.X, DisplayData.FrameIndexA / FramesPerDimension.X);
+				const FIntPoint FrameIndexB = FIntPoint(DisplayData.FrameIndexB % FramesPerDimension.X, DisplayData.FrameIndexB / FramesPerDimension.X);
+				const FIntPoint FramePixelA = FIntPoint(FrameIndexA.X * OutputTexture.FrameSize.X, FrameIndexA.Y * OutputTexture.FrameSize.Y);
+				const FIntPoint FramePixelB = FIntPoint(FrameIndexB.X * OutputTexture.FrameSize.X, FrameIndexB.Y * OutputTexture.FrameSize.Y);
 
 				Canvas->DrawTile(
 					FlipbookViewRect.Min.X, FlipbookViewRect.Min.Y, FlipbookViewRect.Width(), FlipbookViewRect.Height(),
-					float(FramePixel.X) / float(TextureSize.X), float(FramePixel.Y) / float(TextureSize.Y), float(FramePixel.X + PreviewTextureSettings->FrameSize.X) / float(TextureSize.X), float(FramePixel.Y + PreviewTextureSettings->FrameSize.Y) / float(TextureSize.Y),
+					(float(FramePixelA.X) + 0.5f) / float(TextureSize.X), (float(FramePixelA.Y) + 0.5f) / float(TextureSize.Y), (float(FramePixelA.X + OutputTexture.FrameSize.X) - 0.5f) / float(TextureSize.X), (float(FramePixelA.Y + OutputTexture.FrameSize.Y) - 0.5f) / float(TextureSize.Y),
 					FLinearColor::White,
-					PreviewTextureSettings->GeneratedTexture->Resource,
+					OutputTexture.GeneratedTexture->Resource,
 					false	//-TODO: Preview with alpha?
 				);
 
@@ -210,21 +377,52 @@ public:
 					Canvas->DrawShadowedString(TextPosition.X, TextPosition.Y, TEXT("Flipbook"), DisplayFont, FLinearColor::White);
 					TextPosition.Y += FontHeight;
 
-					Canvas->DrawShadowedString(TextPosition.X + 5.0f, TextPosition.Y, *FString::Printf(TEXT("Texture(%d) FrameSize(%d x %d) Frame(%d/%d)"), PreviewTextureIndex, PreviewTextureSettings->FrameSize.X, PreviewTextureSettings->FrameSize.Y, DisplayData.FrameIndex, DisplayData.NumFrames), DisplayFont, FLinearColor::White);
+					Canvas->DrawShadowedString(TextPosition.X + 5.0f, TextPosition.Y, *FString::Printf(TEXT("Texture(%d) FrameSize(%d x %d) Frame(%d/%d)"), PreviewTextureIndex, OutputTexture.FrameSize.X, OutputTexture.FrameSize.Y, DisplayData.FrameIndexA, FlipbookGeneratedSettings->GetNumFrames()), DisplayFont, FLinearColor::White);
 					TextPosition.Y += FontHeight;
+
+					if (!FlipbookGeneratedSettings->Equals(*FlipbookSettings))
+					{
+						Canvas->DrawShadowedString(TextPosition.X + 5.0f, TextPosition.Y, TEXT("Warning: Out Of Date"), DisplayFont, FLinearColor::White);
+						TextPosition.Y += FontHeight;
+					}
 				}
 			}
 			else
 			{
-				Canvas->DrawShadowedString(TextPosition.X, TextPosition.Y, TEXT("Flipbook"), DisplayFont, FLinearColor::White);
+				Canvas->DrawShadowedString(TextPosition.X, TextPosition.Y, TEXT("Flipbook Not Generated"), DisplayFont, FLinearColor::White);
 				TextPosition.Y += FontHeight;
+				}
+		}
+	}
 
-				Canvas->DrawShadowedString(TextPosition.X + 5.0f, TextPosition.Y, TEXT("Texture Not Generated"), DisplayFont, FLinearColor::White);
-			}
+	FIntRect ConstrainRect(const FIntRect& InRect, FIntPoint InSize) const
+	{
+		const float Scale = FMath::Min(float(InRect.Width()) /  float(InSize.X), float(InRect.Height()) / float(InSize.Y));
+		const FIntPoint ScaledSize(FMath::FloorToInt(float(InSize.X) * Scale), FMath::FloorToInt(float(InSize.Y) * Scale));
+		FIntPoint RectMin(
+			InRect.Min.X + ((InRect.Width()  - ScaledSize.X) >> 1),
+			InRect.Min.Y + ((InRect.Height() - ScaledSize.Y) >> 1)
+		);
+
+		return FIntRect(RectMin, RectMin + ScaledSize);
+	}
+
+	void ClearViewArea(FCanvas* Canvas, const FIntRect& InRect)
+	{
+		UTexture2D* Texture = nullptr;
+		FLinearColor Color = ClearColor;
+		FVector2D EndUV(1.0f, 1.0f);
+
+		if (bShowCheckerBoard)
+		{
+			Texture = GetCheckerboardTexture();
+			Color = FLinearColor::White;
+			EndUV.X = float(InRect.Width()) / float(FMath::Max(Texture->GetSizeX(), 1));
+			EndUV.Y = float(InRect.Height()) / float(FMath::Max(Texture->GetSizeY(), 1));
 		}
+		Canvas->DrawTile(InRect.Min.X, InRect.Min.Y, InRect.Width(), InRect.Height(), 0.0f, 0.0f, EndUV.X, EndUV.Y, Color, Texture ? Texture->Resource : nullptr, false);
 	}
 
-	
 	UTexture2D* GetCheckerboardTexture()
 	{
 		if (CheckerboardTexture == nullptr)
@@ -255,13 +453,69 @@ public:
 		Collector.AddReferencedObject(RealtimeRenderTarget);
 	}
 
-	virtual UWorld* GetWorld() const override { return nullptr; }
+	void FocusCamera()
+	{
+		auto ViewModel = WeakViewModel.Pin();
+		UNiagaraFlipbookSettings* FlipbookSettings = ViewModel ? ViewModel->GetFlipbookSettings() : nullptr;
+		UNiagaraComponent* NiagaraComponent = ViewModel->GetPreviewComponent();
+		if ( bShowPreview && FlipbookSettings && NiagaraComponent )
+		{
+			//-TODO: Should take aspect ratio into account here
+			const FBoxSphereBounds ComponentBounds = NiagaraComponent->CalcBounds(NiagaraComponent->GetComponentTransform());
+			if ( FlipbookSettings->IsOrthographic() )
+			{
+				FlipbookSettings->CameraViewportLocation[(int)FlipbookSettings->CameraViewportMode] = ComponentBounds.Origin;
+				FlipbookSettings->CameraOrthoWidth = ComponentBounds.SphereRadius * 2.0f;
+			}
+			else
+			{
+				const float HalfFOVRadians = FMath::DegreesToRadians(FlipbookSettings->CameraFOV) * 0.5f;
+				const float CameraDistance = ComponentBounds.SphereRadius / FMath::Tan(HalfFOVRadians);
+				//const FVector CameraOffset = FlipbookSettings->GetViewMatrix().Inverse().GetUnitAxis(EAxis::Z) * CameraDistance;
+				//FlipbookSettings->CameraViewportLocation[(int)ENiagaraFlipbookViewMode::Perspective] = ComponentBounds.Origin - CameraOffset;
+				FlipbookSettings->CameraViewportLocation[(int)FlipbookSettings->CameraViewportMode] = ComponentBounds.Origin;
+				FlipbookSettings->CameraOrbitDistance = CameraDistance;
+			}
+		}
+	}
+
+	FVector2D GetPreviewOrthoUnits() const
+	{
+		FVector2D OrthoUnits = FVector2D::ZeroVector;
+
+		auto ViewModel = WeakViewModel.Pin();
+		const UNiagaraFlipbookSettings* FlipbookSettings = ViewModel ? ViewModel->GetFlipbookSettings() : nullptr;
+		if (FlipbookSettings && (PreviewViewRect.Area() > 0))
+		{
+			OrthoUnits = FlipbookSettings->GetOrthoSize(ViewModel->GetPreviewTextureIndex());
+			OrthoUnits.X = OrthoUnits.X / float(PreviewViewRect.Width());
+			OrthoUnits.Y = OrthoUnits.Y / float(PreviewViewRect.Height());
+		}
+		return OrthoUnits;
+	}
+
+
+	virtual UWorld* GetWorld() const override
+	{
+		auto ViewModel = WeakViewModel.Pin();
+		return ViewModel ? ViewModel->GetPreviewComponent()->GetWorld() : nullptr;
+	}
 
 	UFont* GetFont() const { return GetStatsFont(); }
 
 public:
 	UTextureRenderTarget2D*						RealtimeRenderTarget = nullptr;
 
+	FVector										LocalMovement = FVector::ZeroVector;
+	float										LocalZoom = 0.0f;
+	float										LocalAspect = 0.0f;
+	FVector										LocalRotation = FVector::ZeroVector;
+	float										KeyboardMoveSpeed = 5.0f;
+	float										PerspectiveMoveSpeed = 2.0f;
+
+	FIntRect									PreviewViewRect;
+	FIntRect									FlipbookViewRect;
+
 	bool										bShowCheckerBoard = true;
 	UTexture2D*									CheckerboardTexture = nullptr;
 	FColor										CheckerboardColorOne = FColor(128, 128, 128);
diff --git a/Engine/Plugins/FX/Niagara/Source/NiagaraEditor/Private/Widgets/SNiagaraFlipbookWidget.cpp b/Engine/Plugins/FX/Niagara/Source/NiagaraEditor/Private/Widgets/SNiagaraFlipbookWidget.cpp
--- a/Engine/Plugins/FX/Niagara/Source/NiagaraEditor/Private/Widgets/SNiagaraFlipbookWidget.cpp
+++ b/Engine/Plugins/FX/Niagara/Source/NiagaraEditor/Private/Widgets/SNiagaraFlipbookWidget.cpp
@@ -143,10 +143,10 @@ void SNiagaraFlipbookWidget::Tick(const FGeometry& AllottedGeometry, const doubl
 	{
 		float StartSeconds = 0.0f;
 		float DurationSeconds = 0.0f;
-		if ( ViewModel->IsGeneratedDataValid() )
+		if (const UNiagaraFlipbookSettings* GeneratedSettings = GetFlipbookGeneratedSettings())
 		{
-			StartSeconds = ViewModel->GetGeneratedStartSeconds();
-			DurationSeconds = ViewModel->GetGeneratedDurationSeconds();
+			StartSeconds = GeneratedSettings->StartSeconds;
+			DurationSeconds = GeneratedSettings->DurationSeconds;
 		}
 		else if (UNiagaraFlipbookSettings* FlipbookSettings = GetFlipbookSettings())
 		{
@@ -197,6 +197,15 @@ UNiagaraFlipbookSettings* SNiagaraFlipbookWidget::GetFlipbookSettings() const
 	return nullptr;
 }
 
+const UNiagaraFlipbookSettings* SNiagaraFlipbookWidget::GetFlipbookGeneratedSettings() const
+{
+	if (auto ViewModel = WeakViewModel.Pin())
+	{
+		return ViewModel->GetFlipbookGeneratedSettings();
+	}
+	return nullptr;
+}
+
 FReply SNiagaraFlipbookWidget::OnTransportBackwardEnd()
 {
 	bIsPlaying = false;
@@ -210,9 +219,13 @@ FReply SNiagaraFlipbookWidget::OnTransportBackwardStep()
 	bIsPlaying = false;
 	if (auto ViewModel = WeakViewModel.Pin())
 	{
-		const auto DisplayData = ViewModel->GetDisplayDataFromRelativeTime(PreviewRelativeTime);
-		const int32 NewFrame = FMath::Clamp(FMath::RoundToInt(DisplayData.FrameTime - 0.6f), 0, DisplayData.NumFrames);
-		PreviewRelativeTime = (DisplayData.DurationSeconds / float(DisplayData.NumFrames)) * float(NewFrame);
+		if(UNiagaraFlipbookSettings* Settings = ViewModel->GetFlipbookSettings())
+		{
+			const auto DisplayData = Settings->GetDisplayInfo(PreviewRelativeTime, Settings->bPreviewLooping);
+			const int NumFrames = Settings->GetNumFrames();
+			const int NewFrame = DisplayData.Interp > 0.25f ? DisplayData.FrameIndexA : FMath::Max(DisplayData.FrameIndexA - 1, 0);
+			PreviewRelativeTime = (Settings->DurationSeconds / float(NumFrames)) * float(NewFrame);
+		}
 	}
 
 	return FReply::Handled();
@@ -229,9 +242,13 @@ FReply SNiagaraFlipbookWidget::OnTransportForwardStep()
 	bIsPlaying = false;
 	if (auto ViewModel = WeakViewModel.Pin())
 	{
-		const auto DisplayData = ViewModel->GetDisplayDataFromRelativeTime(PreviewRelativeTime);
-		const int32 NewFrame = FMath::Clamp(FMath::RoundToInt(DisplayData.FrameTime + 0.6f), 0, DisplayData.NumFrames);
-		PreviewRelativeTime = (DisplayData.DurationSeconds / float(DisplayData.NumFrames)) * float(NewFrame);
+		if (UNiagaraFlipbookSettings* Settings = ViewModel->GetFlipbookSettings())
+		{
+			const auto DisplayData = Settings->GetDisplayInfo(PreviewRelativeTime, Settings->bPreviewLooping);
+			const int NumFrames = Settings->GetNumFrames();
+			const int NewFrame = DisplayData.Interp < 0.75f ? DisplayData.FrameIndexB : FMath::Min(DisplayData.FrameIndexB + 1, NumFrames - 1);
+			PreviewRelativeTime = (Settings->DurationSeconds / float(NumFrames)) * float(NewFrame);
+		}
 	}
 
 	return FReply::Handled();
diff --git a/Engine/Plugins/FX/Niagara/Source/NiagaraEditor/Private/Widgets/SNiagaraFlipbookWidget.h b/Engine/Plugins/FX/Niagara/Source/NiagaraEditor/Private/Widgets/SNiagaraFlipbookWidget.h
--- a/Engine/Plugins/FX/Niagara/Source/NiagaraEditor/Private/Widgets/SNiagaraFlipbookWidget.h
+++ b/Engine/Plugins/FX/Niagara/Source/NiagaraEditor/Private/Widgets/SNiagaraFlipbookWidget.h
@@ -27,6 +27,7 @@ public:
 	FText GetSelectedTextureAsText() const;
 
 	class UNiagaraFlipbookSettings* GetFlipbookSettings() const;
+	const class UNiagaraFlipbookSettings* GetFlipbookGeneratedSettings() const;
 
 	void SetPreviewRelativeTime(float RelativeTime);
 
@@ -43,6 +44,6 @@ private:
 	TSharedPtr<IDetailsView>					FlipbookSettingsDetails;
 	TSharedPtr<SWidget>							TransportControls;
 
-	bool										bIsPlaying = false;
+	bool										bIsPlaying = true;
 	float										PreviewRelativeTime = 0.0f;
 };
