Fix FNiagaraDataInterfaceArrayBool to use a Buffer<uint> on the GPU

diff --git a/Engine/Plugins/FX/Niagara/Source/Niagara/Private/NiagaraDataInterfaceArrayInt.cpp b/Engine/Plugins/FX/Niagara/Source/Niagara/Private/NiagaraDataInterfaceArrayInt.cpp
--- a/Engine/Plugins/FX/Niagara/Source/Niagara/Private/NiagaraDataInterfaceArrayInt.cpp
+++ b/Engine/Plugins/FX/Niagara/Source/Niagara/Private/NiagaraDataInterfaceArrayInt.cpp
@@ -19,7 +19,7 @@ struct FNDIArrayImplHelper<bool> : public FNDIArrayImplHelperBase<bool>
 {
 	typedef FNiagaraBool TVMArrayType;
 	static constexpr TCHAR const* HLSLValueTypeName = TEXT("bool");
-	static constexpr TCHAR const* HLSLBufferTypeName = TEXT("int");
+	static constexpr TCHAR const* HLSLBufferTypeName = TEXT("uint");
 	static constexpr EPixelFormat PixelFormat = PF_R8_UINT;
 	static const FNiagaraTypeDefinition& GetTypeDefinition() { return FNiagaraTypeDefinition::GetBoolDef(); }
 	static const bool GetDefaultValue() { return false; }
diff --git a/Engine/Plugins/FX/Niagara/Source/NiagaraCore/Private/NiagaraCustomVersion.cpp b/Engine/Plugins/FX/Niagara/Source/NiagaraCore/Private/NiagaraCustomVersion.cpp
--- a/Engine/Plugins/FX/Niagara/Source/NiagaraCore/Private/NiagaraCustomVersion.cpp
+++ b/Engine/Plugins/FX/Niagara/Source/NiagaraCore/Private/NiagaraCustomVersion.cpp
@@ -8,5 +8,4 @@ const FGuid FNiagaraCustomVersion::GUID(0xFCF57AFA, 0x50764283, 0xB9A9E658, 0xFF
 // Register the custom version with core
 FCustomVersionRegistration GRegisterNiagaraCustomVersion(FNiagaraCustomVersion::GUID, FNiagaraCustomVersion::LatestVersion, TEXT("NiagaraVer"));
 
-
-const FGuid FNiagaraCustomVersion::LatestScriptCompileVersion(0x7D0D5E8B, 0x14654CE3, 0x8A8C2594, 0x18FCEDA9);
+const FGuid FNiagaraCustomVersion::LatestScriptCompileVersion(0x6F6C61A2, 0xDFB54DF1, 0xBBBC7C87, 0x3173A7E1);
