- Editgrate CL 15487758: Fix culled counts being aquired and not released

diff --git a/Engine/Plugins/FX/Niagara/Source/Niagara/Private/NiagaraGPUInstanceCountManager.cpp b/Engine/Plugins/FX/Niagara/Source/Niagara/Private/NiagaraGPUInstanceCountManager.cpp
--- a/Engine/Plugins/FX/Niagara/Source/Niagara/Private/NiagaraGPUInstanceCountManager.cpp
+++ b/Engine/Plugins/FX/Niagara/Source/Niagara/Private/NiagaraGPUInstanceCountManager.cpp
@@ -368,14 +368,16 @@ void FNiagaraGPUInstanceCountManager::UpdateDrawIndirectBuffer(FRHICommandList&
 		DrawIndirectArgGenTasks.Empty();
 		DrawIndirectArgMap.Empty();
 		InstanceCountClearTasks.Empty();
+	}
 
-		if (bAcquiredCulledCounts && RequiredCulledCounts > 0)
-		{
-			CulledCountBuffer.DiscardTransientResource();
-		}
-		bAcquiredCulledCounts = false;
-		RequiredCulledCounts = 0; // reset this counter now that we're done with them
+	// Release culled count buffers
+	// This is done outside of the above if as a mesh renderer could request a culled count but never add any indirect draws
+	if (bAcquiredCulledCounts && RequiredCulledCounts > 0)
+	{
+		CulledCountBuffer.DiscardTransientResource();
 	}
+	bAcquiredCulledCounts = false;
+	RequiredCulledCounts = 0; // reset this counter now that we're done with them
 }
 
 const uint32* FNiagaraGPUInstanceCountManager::GetGPUReadback() 
