Fix GPU crash caused by nullptr SRV in Niagara 3D neighbor grid system.
 - The "fx.MaxNiagaraNeighborGridCells" cvar added in CL 15259924 can cause the Neighbor buffer/srv/uav to be uninitialized during the draw call for the emitter, leaving an shader buffer parameter unbound.

diff --git a/Engine/Plugins/FX/Niagara/Source/Niagara/Private/NiagaraDataInterfaceNeighborGrid3D.cpp b/Engine/Plugins/FX/Niagara/Source/Niagara/Private/NiagaraDataInterfaceNeighborGrid3D.cpp
--- a/Engine/Plugins/FX/Niagara/Source/Niagara/Private/NiagaraDataInterfaceNeighborGrid3D.cpp
+++ b/Engine/Plugins/FX/Niagara/Source/Niagara/Private/NiagaraDataInterfaceNeighborGrid3D.cpp
@@ -70,7 +70,7 @@ public:
 		NeighborGrid3DRWInstanceData* ProxyData = VFDI->SystemInstancesToProxyData.Find(Context.SystemInstanceID);
 		float CellSizeTmp[3];
 
-		if (!ProxyData)
+		if (!(ProxyData && ProxyData->NeighborhoodBuffer.Buffer.IsValid()))
 		{
 			CellSizeTmp[0] = CellSizeTmp[1] = CellSizeTmp[2] = 1.0f;
 			SetShaderValue(RHICmdList, ComputeShaderRHI, NumCellsParam, 0);
@@ -775,4 +775,4 @@ bool UNiagaraDataInterfaceNeighborGrid3D::CopyToInternal(UNiagaraDataInterface*
 	return true;
 }
 
-#undef LOCTEXT_NAMESPACE
\ No newline at end of file
+#undef LOCTEXT_NAMESPACE
