- volume texture source refactor to be closer to particle source
- support for low res velocity
- functions on grid3d for querying whole grid rather than tiled attributes
- optimized uvw clamping for grid3d sample

diff --git a/Engine/Plugins/FX/Niagara/Source/Niagara/Classes/NiagaraDataInterfaceGrid3DCollection.h b/Engine/Plugins/FX/Niagara/Source/Niagara/Classes/NiagaraDataInterfaceGrid3DCollection.h
--- a/Engine/Plugins/FX/Niagara/Source/Niagara/Classes/NiagaraDataInterfaceGrid3DCollection.h
+++ b/Engine/Plugins/FX/Niagara/Source/Niagara/Classes/NiagaraDataInterfaceGrid3DCollection.h
@@ -212,6 +212,10 @@ public:
 	static const FName GetValueFunctionName;
 	static const FName SampleGridFunctionName;
 
+	static const FName SetFullGridValueFunctionName;
+	static const FName GetFullGridPreviousValueFunctionName;
+	static const FName SamplePreviousFullGridFunctionName;
+
 	static const FName SetVector4ValueFunctionName;
 	static const FName SetVector3ValueFunctionName;
 	static const FName SetVector2ValueFunctionName;
diff --git a/Engine/Plugins/FX/Niagara/Source/Niagara/Private/NiagaraDataInterfaceGrid3DCollection.cpp b/Engine/Plugins/FX/Niagara/Source/Niagara/Private/NiagaraDataInterfaceGrid3DCollection.cpp
--- a/Engine/Plugins/FX/Niagara/Source/Niagara/Private/NiagaraDataInterfaceGrid3DCollection.cpp
+++ b/Engine/Plugins/FX/Niagara/Source/Niagara/Private/NiagaraDataInterfaceGrid3DCollection.cpp
@@ -32,6 +32,10 @@ const FName UNiagaraDataInterfaceGrid3DCollection::CopyPreviousToCurrentForCellF
 const FName UNiagaraDataInterfaceGrid3DCollection::SetValueFunctionName("SetGridValue");
 const FName UNiagaraDataInterfaceGrid3DCollection::GetValueFunctionName("GetGridValue");
 
+const FName UNiagaraDataInterfaceGrid3DCollection::SetFullGridValueFunctionName("SetFullGridValue");
+const FName UNiagaraDataInterfaceGrid3DCollection::GetFullGridPreviousValueFunctionName("GetFullGridPreviousValue");
+const FName UNiagaraDataInterfaceGrid3DCollection::SamplePreviousFullGridFunctionName("SamplePreviousFullGrid");
+
 const FName UNiagaraDataInterfaceGrid3DCollection::SetVector4ValueFunctionName("SetVector4Value");
 const FName UNiagaraDataInterfaceGrid3DCollection::SetVector3ValueFunctionName("SetVector3Value");
 const FName UNiagaraDataInterfaceGrid3DCollection::SetVector2ValueFunctionName("SetVector2Value");
@@ -442,6 +446,28 @@ void UNiagaraDataInterfaceGrid3DCollection::GetFunctions(TArray<FNiagaraFunction
 		OutFunctions.Add(Sig);
 	}
 
+	{
+		FNiagaraFunctionSignature Sig;
+		Sig.Name = GetFullGridPreviousValueFunctionName;
+		Sig.Inputs.Add(FNiagaraVariable(FNiagaraTypeDefinition(GetClass()), TEXT("Grid")));
+		Sig.Inputs.Add(FNiagaraVariable(FNiagaraTypeDefinition::GetIntDef(), TEXT("IndexX")));
+		Sig.Inputs.Add(FNiagaraVariable(FNiagaraTypeDefinition::GetIntDef(), TEXT("IndexY")));
+		Sig.Inputs.Add(FNiagaraVariable(FNiagaraTypeDefinition::GetIntDef(), TEXT("IndexZ")));		
+		Sig.Outputs.Add(FNiagaraVariable(FNiagaraTypeDefinition::GetFloatDef(), TEXT("Value")));
+
+		Sig.bExperimental = true;
+		Sig.bMemberFunction = true;
+		Sig.bRequiresContext = false;
+		Sig.ModuleUsageBitmask = ENiagaraScriptUsageMask::Particle;
+		Sig.bSupportsCPU = false;
+		Sig.bSupportsGPU = true;
+#if WITH_EDITORONLY_DATA
+		Sig.Description = NSLOCTEXT("Niagara", "NiagaraDataInterfaceGridColl2D_PreviousValueAtIndexFunction", "Get the value at a specific index.");
+#endif
+
+		OutFunctions.Add(Sig);
+	}
+
 	{
 		FNiagaraFunctionSignature Sig;
 		Sig.Name = SetValueFunctionName;
@@ -467,6 +493,29 @@ void UNiagaraDataInterfaceGrid3DCollection::GetFunctions(TArray<FNiagaraFunction
 		OutFunctions.Add(Sig);
 	}
 
+	{
+		FNiagaraFunctionSignature Sig;
+		Sig.Name = SetFullGridValueFunctionName;
+		Sig.Inputs.Add(FNiagaraVariable(FNiagaraTypeDefinition(GetClass()), TEXT("Grid")));
+		Sig.Inputs.Add(FNiagaraVariable(FNiagaraTypeDefinition::GetIntDef(), TEXT("IndexX")));
+		Sig.Inputs.Add(FNiagaraVariable(FNiagaraTypeDefinition::GetIntDef(), TEXT("IndexY")));
+		Sig.Inputs.Add(FNiagaraVariable(FNiagaraTypeDefinition::GetIntDef(), TEXT("IndexZ")));	
+		Sig.Inputs.Add(FNiagaraVariable(FNiagaraTypeDefinition::GetFloatDef(), TEXT("Value")));
+		Sig.Outputs.Add(FNiagaraVariable(FNiagaraTypeDefinition::GetIntDef(), TEXT("IGNORE")));
+
+		Sig.bExperimental = true;
+		Sig.bMemberFunction = true;
+		Sig.bRequiresContext = false;
+		Sig.bWriteFunction = true;
+		Sig.ModuleUsageBitmask = ENiagaraScriptUsageMask::Particle;
+		Sig.bSoftDeprecatedFunction = true;
+		Sig.bSupportsCPU = false;
+		Sig.bSupportsGPU = true;
+#if WITH_EDITORONLY_DATA
+		Sig.Description = NSLOCTEXT("Niagara", "NiagaraDataInterfaceGridColl2D_SetValueFunction", "Set the value at a specific index. Note that this is an older way of working with Grids. Consider using the SetFloat or other typed, named functions or parameter map variables with StackContext namespace instead.");
+#endif
+		OutFunctions.Add(Sig);
+	}
 
 	{
 		FNiagaraFunctionSignature Sig;
@@ -813,6 +862,22 @@ void UNiagaraDataInterfaceGrid3DCollection::GetFunctions(TArray<FNiagaraFunction
 		OutFunctions.Add(Sig);
 	}
 
+	{
+		FNiagaraFunctionSignature Sig;
+		Sig.Name = SamplePreviousFullGridFunctionName;
+		Sig.Inputs.Add(FNiagaraVariable(FNiagaraTypeDefinition(GetClass()), TEXT("Grid")));
+		Sig.Inputs.Add(FNiagaraVariable(FNiagaraTypeDefinition::GetVec3Def(), TEXT("Unit")));		
+		Sig.Outputs.Add(FNiagaraVariable(FNiagaraTypeDefinition::GetFloatDef(), TEXT("Value")));
+
+		Sig.bExperimental = true;
+		Sig.bMemberFunction = true;
+		Sig.bRequiresContext = false;
+		Sig.ModuleUsageBitmask = ENiagaraScriptUsageMask::Particle;
+		Sig.bSupportsCPU = false;
+		Sig.bSupportsGPU = true;
+		OutFunctions.Add(Sig);
+	}
+
 	{
 		FNiagaraFunctionSignature Sig;
 		Sig.Name = GetVector4AttributeIndexFunctionName;
@@ -1392,6 +1457,17 @@ bool UNiagaraDataInterfaceGrid3DCollection::GetFunctionHLSL(const FNiagaraDataIn
 		OutHLSL += FString::Format(FormatBounds, ArgsBounds);
 		return true;
 	}
+	else if (FunctionInfo.DefinitionName == GetFullGridPreviousValueFunctionName)
+	{
+		static const TCHAR* FormatBounds = TEXT(R"(
+			void {FunctionName}(int In_IndexX, int In_IndexY, int In_IndexZ, out float Out_Val)
+			{
+				Out_Val = {Grid}.Load(int4(In_IndexX, In_IndexY, In_IndexZ, 0));
+			}
+		)");
+		OutHLSL += FString::Format(FormatBounds, ArgsBounds);
+		return true;
+	}
 	else if (FunctionInfo.DefinitionName == SetValueFunctionName )
 	{
 		static const TCHAR *FormatBounds = TEXT(R"(
@@ -1408,6 +1484,18 @@ bool UNiagaraDataInterfaceGrid3DCollection::GetFunctionHLSL(const FNiagaraDataIn
 		OutHLSL += FString::Format(FormatBounds, ArgsBounds);
 		return true;
 	}
+	else if (FunctionInfo.DefinitionName == SetFullGridValueFunctionName)
+	{
+		static const TCHAR* FormatBounds = TEXT(R"(
+			void {FunctionName}(int In_IndexX, int In_IndexY, int In_IndexZ, float In_Value, out int val)
+			{			
+				val = 0;
+				RW{OutputGrid}[int3(In_IndexX, In_IndexY, In_IndexZ)] = In_Value;
+			}
+		)");
+		OutHLSL += FString::Format(FormatBounds, ArgsBounds);
+		return true;
+	}
 	else if (FunctionInfo.DefinitionName == CopyPreviousToCurrentForCellFunctionName)
 	{
 		static const TCHAR* FormatBounds = TEXT(R"(
@@ -1537,26 +1625,12 @@ bool UNiagaraDataInterfaceGrid3DCollection::GetFunctionHLSL(const FNiagaraDataIn
 					int TileIndexX = In_AttributeIndex % {NumTiles}.x;
 					int TileIndexY = (In_AttributeIndex / {NumTiles}.x) % {NumTiles}.y;
 					int TileIndexZ = In_AttributeIndex / ({NumTiles}.x * {NumTiles}.y);		
+					
+					float3 UVW = float3(In_UnitX, In_UnitY, In_UnitZ);
+					UVW = clamp(UVW, float3(0,0,0), float3(1,1,1));
 
-					float3 UVW =
-					{
-						In_UnitX / {NumTiles}.x + 1.0*TileIndexX/{NumTiles}.x,
-						In_UnitY / {NumTiles}.y + 1.0*TileIndexY/{NumTiles}.y,
-						In_UnitZ / {NumTiles}.z + 1.0*TileIndexZ/{NumTiles}.y
-					};
-					float3 TileMin =
-					{
-						(TileIndexX * {NumCellsName}.x + 0.5) / ({NumTiles}.x * {NumCellsName}.x),
-						(TileIndexY * {NumCellsName}.y + 0.5) / ({NumTiles}.y * {NumCellsName}.y),
-						(TileIndexZ * {NumCellsName}.z + 0.5) / ({NumTiles}.z * {NumCellsName}.z)
-					};
-					float3 TileMax =
-					{
-						((TileIndexX + 1) * {NumCellsName}.x - 0.5) / ({NumTiles}.x * {NumCellsName}.x),
-						((TileIndexY + 1) * {NumCellsName}.y - 0.5) / ({NumTiles}.y * {NumCellsName}.y),
-						((TileIndexZ + 1) * {NumCellsName}.z - 0.5) / ({NumTiles}.z * {NumCellsName}.z)
-					};
-					UVW = clamp(UVW, TileMin, TileMax);
+					UVW /= {NumTiles};
+					UVW += float3(1.0*TileIndexX, 1.0*TileIndexY, 1.0*TileIndexZ) / {NumTiles};
 
 					Out_Val = {Grid}.SampleLevel({SamplerName}, UVW, 0);
 				}
@@ -1564,6 +1638,17 @@ bool UNiagaraDataInterfaceGrid3DCollection::GetFunctionHLSL(const FNiagaraDataIn
 		OutHLSL += FString::Format(FormatBounds, ArgsBounds);
 		return true;
 	}
+	else if (FunctionInfo.DefinitionName == SamplePreviousFullGridFunctionName)
+	{
+	static const TCHAR* FormatBounds = TEXT(R"(
+				void {FunctionName}(float3 In_Unit, out float Out_Val)
+				{
+					Out_Val = {Grid}.SampleLevel({SamplerName}, In_Unit, 0);
+				}
+			)");
+	OutHLSL += FString::Format(FormatBounds, ArgsBounds);
+	return true;
+	}
 	return false;
 }
 #if WITH_EDITOR
