Fix for warnings being generated when trying to handle rename rapid iteration parameters.

diff --git a/Engine/Plugins/FX/Niagara/Source/Niagara/Private/NiagaraCommon.cpp b/Engine/Plugins/FX/Niagara/Source/Niagara/Private/NiagaraCommon.cpp
--- a/Engine/Plugins/FX/Niagara/Source/Niagara/Private/NiagaraCommon.cpp
+++ b/Engine/Plugins/FX/Niagara/Source/Niagara/Private/NiagaraCommon.cpp
@@ -830,12 +830,10 @@ FName NIAGARA_API FNiagaraUtilities::GetUniqueName(FName CandidateName, const TS
 	return UniqueName;
 }
 
-FNiagaraVariable FNiagaraUtilities::ConvertVariableToRapidIterationConstantName(FNiagaraVariable InVar, const TCHAR* InEmitterName, ENiagaraScriptUsage InUsage)
+FString FNiagaraUtilities::CreateRapidIterationConstantName(FName InVariableName, const TCHAR* InEmitterName, ENiagaraScriptUsage InUsage)
 {
-	FNiagaraVariable Var = InVar;
-
 	TArray<FString> SplitName;
-	Var.GetName().ToString().ParseIntoArray(SplitName, TEXT("."));
+	InVariableName.ToString().ParseIntoArray(SplitName, TEXT("."));
 	int32 NumSlots = SplitName.Num();
 	if (InEmitterName != nullptr)
 	{
@@ -863,8 +861,15 @@ FNiagaraVariable FNiagaraUtilities::ConvertVariableToRapidIterationConstantName(
 		SplitName.Insert(TEXT("Constants"), 0);
 	}
 
-	FString OutVarStrName = FString::Join(SplitName, TEXT("."));
-	Var.SetName(*OutVarStrName);
+	return FString::Join(SplitName, TEXT("."));
+}
+
+FNiagaraVariable FNiagaraUtilities::ConvertVariableToRapidIterationConstantName(FNiagaraVariable InVar, const TCHAR* InEmitterName, ENiagaraScriptUsage InUsage)
+{
+	FNiagaraVariable Var = InVar;
+
+	Var.SetName(*CreateRapidIterationConstantName(Var.GetName(), InEmitterName, InUsage));
+
 	return Var;
 }
 
diff --git a/Engine/Plugins/FX/Niagara/Source/Niagara/Public/NiagaraCommon.h b/Engine/Plugins/FX/Niagara/Source/Niagara/Public/NiagaraCommon.h
--- a/Engine/Plugins/FX/Niagara/Source/Niagara/Public/NiagaraCommon.h
+++ b/Engine/Plugins/FX/Niagara/Source/Niagara/Public/NiagaraCommon.h
@@ -1020,6 +1020,7 @@ namespace FNiagaraUtilities
 	if necessary by adding a 3 digit index to the end. */
 	FName NIAGARA_API GetUniqueName(FName CandidateName, const TSet<FName>& ExistingNames);
 
+	NIAGARA_API FString CreateRapidIterationConstantName(FName InVariableName, const TCHAR* InEmitterName, ENiagaraScriptUsage InUsage);
 	FNiagaraVariable NIAGARA_API ConvertVariableToRapidIterationConstantName(FNiagaraVariable InVar, const TCHAR* InEmitterName, ENiagaraScriptUsage InUsage);
 
 	void CollectScriptDataInterfaceParameters(const UObject& Owner, const TArrayView<UNiagaraScript*>& Scripts, FNiagaraParameterStore& OutDataInterfaceParameters);
diff --git a/Engine/Plugins/FX/Niagara/Source/NiagaraEditor/Private/NiagaraScriptSource.cpp b/Engine/Plugins/FX/Niagara/Source/NiagaraEditor/Private/NiagaraScriptSource.cpp
--- a/Engine/Plugins/FX/Niagara/Source/NiagaraEditor/Private/NiagaraScriptSource.cpp
+++ b/Engine/Plugins/FX/Niagara/Source/NiagaraEditor/Private/NiagaraScriptSource.cpp
@@ -197,7 +197,7 @@ bool UNiagaraScriptSource::AddModuleIfMissing(FString ModulePath, ENiagaraScript
 	return false;
 }
 
-void UNiagaraScriptSource::FixupRenamedParameters(UNiagaraNode* Node, FNiagaraParameterStore& RapidIterationParameters, TArray<FNiagaraVariable> OldRapidIterationVariables, TSet<FName> ValidRapidIterationParameterNames, const UNiagaraEmitter* Emitter, ENiagaraScriptUsage ScriptUsage) const
+void UNiagaraScriptSource::FixupRenamedParameters(UNiagaraNode* Node, FNiagaraParameterStore& RapidIterationParameters, const TArray<FNiagaraVariable>& OldRapidIterationVariables, const TSet<FName>& ValidRapidIterationParameterNames, const UNiagaraEmitter* Emitter, ENiagaraScriptUsage ScriptUsage) const
 {
 	UNiagaraNodeFunctionCall* FunctionCallNode = Cast<UNiagaraNodeFunctionCall>(Node);
 	if (FunctionCallNode == nullptr || FunctionCallNode->FunctionScript == nullptr)
@@ -220,24 +220,46 @@ void UNiagaraScriptSource::FixupRenamedParameters(UNiagaraNode* Node, FNiagaraPa
 	// go through the existing rapid iteration params to see if they are either still valid or were renamed
 	for (FNiagaraVariable OldRapidIterationVar : OldRapidIterationVariables)
 	{
-		if (ValidRapidIterationParameterNames.Contains(OldRapidIterationVar.GetName()) || !RapidIterationParameters.ParameterGuidMapping.Contains(OldRapidIterationVar))
+		if (ValidRapidIterationParameterNames.Contains(OldRapidIterationVar.GetName()))
 		{
 			continue;
 		}
 
-		FGuid BoundGuid = RapidIterationParameters.ParameterGuidMapping[OldRapidIterationVar];
-		for (const UEdGraphPin* ModulePin : ModuleInputPins)
+		const FGuid* BoundGuid = RapidIterationParameters.ParameterGuidMapping.Find(OldRapidIterationVar);
+		if (BoundGuid)
 		{
-			FNiagaraParameterHandle AliasedFunctionInputHandle = FNiagaraParameterHandle::CreateAliasedModuleParameterHandle(FNiagaraParameterHandle(ModulePin->PinName), FunctionCallNode);
-			FNiagaraVariable InputVar = NiagaraSchema->PinToNiagaraVariable(ModulePin);
-			TOptional<FNiagaraVariableMetaData> VariableMetaData = Graph->GetMetaData(InputVar);
-
-			// if the guid matches but the names differ then the parameter was renamed, so lets update the rapid iteration parameter
-			if (VariableMetaData.IsSet() && VariableMetaData->GetVariableGuid() == BoundGuid && AliasedFunctionInputHandle.GetParameterHandleString() != OldRapidIterationVar.GetName())
+			for (const UEdGraphPin* ModulePin : ModuleInputPins)
 			{
-				FNiagaraTypeDefinition InputType = NiagaraSchema->PinToTypeDefinition(ModulePin);
-				FNiagaraVariable NewRapidIterationVar = FNiagaraStackGraphUtilities::CreateRapidIterationParameter(UniqueEmitterName, ScriptUsage, AliasedFunctionInputHandle.GetParameterHandleString(), InputType);
-				RapidIterationParameters.RenameParameter(OldRapidIterationVar, NewRapidIterationVar.GetName());
+				FNiagaraVariable InputVar = NiagaraSchema->PinToNiagaraVariable(ModulePin);
+				TOptional<FNiagaraVariableMetaData> VariableMetaData = Graph->GetMetaData(InputVar);
+
+				if (VariableMetaData.IsSet() && VariableMetaData->GetVariableGuid() == *BoundGuid)
+				{
+					FNiagaraParameterHandle AliasedFunctionInputHandle = FNiagaraParameterHandle::CreateAliasedModuleParameterHandle(FNiagaraParameterHandle(ModulePin->PinName), FunctionCallNode);
+
+					const FString ConstantName = FNiagaraUtilities::CreateRapidIterationConstantName(AliasedFunctionInputHandle.GetParameterHandleString(), *UniqueEmitterName, ScriptUsage);
+
+					// if the names match move on
+					if (OldRapidIterationVar.GetName() == *ConstantName)
+					{
+						continue;
+					}
+
+					// before we try to rename make sure we check the namespace so that we're dealing with the right variable
+					int32 NameSpaceEnd;
+					if (ConstantName.FindLastChar(TEXT('.'), NameSpaceEnd))
+					{
+						if (!OldRapidIterationVar.IsInNameSpace(ConstantName.Left(NameSpaceEnd)))
+						{
+							continue;
+						}
+					}
+
+					// if the guid matches but the names differ then the parameter was renamed, so lets update the rapid iteration parameter
+					FNiagaraTypeDefinition InputType = NiagaraSchema->PinToTypeDefinition(ModulePin);
+					FNiagaraVariable NewRapidIterationVar = FNiagaraStackGraphUtilities::CreateRapidIterationParameter(UniqueEmitterName, ScriptUsage, AliasedFunctionInputHandle.GetParameterHandleString(), InputType);
+					RapidIterationParameters.RenameParameter(OldRapidIterationVar, NewRapidIterationVar.GetName());
+				}
 			}
 		}
 	}
diff --git a/Engine/Plugins/FX/Niagara/Source/NiagaraEditor/Public/NiagaraScriptSource.h b/Engine/Plugins/FX/Niagara/Source/NiagaraEditor/Public/NiagaraScriptSource.h
--- a/Engine/Plugins/FX/Niagara/Source/NiagaraEditor/Public/NiagaraScriptSource.h
+++ b/Engine/Plugins/FX/Niagara/Source/NiagaraEditor/Public/NiagaraScriptSource.h
@@ -45,7 +45,7 @@ class UNiagaraScriptSource : public UNiagaraScriptSourceBase
 
 	NIAGARAEDITOR_API virtual bool AddModuleIfMissing(FString ModulePath, ENiagaraScriptUsage Usage, bool& bOutFoundModule)override;
 
-	virtual void FixupRenamedParameters(UNiagaraNode* Node, FNiagaraParameterStore& RapidIterationParameters, TArray<FNiagaraVariable> OldRapidIterationVariables, TSet<FName> ValidRapidIterationParameterNames, const UNiagaraEmitter* Emitter, ENiagaraScriptUsage ScriptUsage) const;
+	virtual void FixupRenamedParameters(UNiagaraNode* Node, FNiagaraParameterStore& RapidIterationParameters, const TArray<FNiagaraVariable>& OldRapidIterationVariables, const TSet<FName>& ValidRapidIterationParameterNames, const UNiagaraEmitter* Emitter, ENiagaraScriptUsage ScriptUsage) const;
 	virtual void CleanUpOldAndInitializeNewRapidIterationParameters(const UNiagaraEmitter* Emitter, ENiagaraScriptUsage ScriptUsage, FGuid ScriptUsageId, FNiagaraParameterStore& RapidIterationParameters) const override;
 	virtual void ForceGraphToRecompileOnNextCheck() override;
 	virtual void RefreshFromExternalChanges() override;
