- Allow console variable to adjust existing buffers for debug draws
- Add some sensible defaults to the nodes and rename bDraw to Execute as this matches with other data interfaces

#rnx

diff --git a/Engine/Plugins/FX/Niagara/Source/Niagara/Private/DataInterface/NiagaraDataInterfaceDebugDraw.cpp b/Engine/Plugins/FX/Niagara/Source/Niagara/Private/DataInterface/NiagaraDataInterfaceDebugDraw.cpp
--- a/Engine/Plugins/FX/Niagara/Source/Niagara/Private/DataInterface/NiagaraDataInterfaceDebugDraw.cpp
+++ b/Engine/Plugins/FX/Niagara/Source/Niagara/Private/DataInterface/NiagaraDataInterfaceDebugDraw.cpp
@@ -90,14 +90,14 @@ namespace NDIDebugDrawLocal
 		};
 
 #if NIAGARA_COMPUTEDEBUG_ENABLED
-		static void Draw(FNDIDebugDrawInstanceData_GameThread* InstanceData, VMBindings& Bindings, bool bDraw)
+		static void Draw(FNDIDebugDrawInstanceData_GameThread* InstanceData, VMBindings& Bindings, bool bExecute)
 		{
 			const FVector Location = Bindings.LocationParam.GetAndAdvance();
 			const FQuat Rotation = Bindings.RotationParam.GetAndAdvance();
 			const FVector Extents = Bindings.ExtentsParam.GetAndAdvance();
 			const FLinearColor Color = Bindings.ColorParam.GetAndAdvance();
 
-			if ( bDraw )
+			if (bExecute)
 			{
 				const FVector Points[] =
 				{
@@ -152,7 +152,7 @@ namespace NDIDebugDrawLocal
 		};
 
 #if NIAGARA_COMPUTEDEBUG_ENABLED
-		static void Draw(FNDIDebugDrawInstanceData_GameThread* InstanceData, VMBindings& Bindings, bool bDraw)
+		static void Draw(FNDIDebugDrawInstanceData_GameThread* InstanceData, VMBindings& Bindings, bool bExecute)
 		{
 			const FVector Location = Bindings.LocationParam.GetAndAdvance();
 			const FVector XAxis = Bindings.XAxisParam.GetAndAdvance();
@@ -161,7 +161,7 @@ namespace NDIDebugDrawLocal
 			const int32 Segments = FMath::Clamp(Bindings.SegmentsParam.GetAndAdvance(), 4, 16);
 			const FLinearColor Color = Bindings.ColorParam.GetAndAdvance();
 
-			if (bDraw)
+			if (bExecute)
 			{
 				const FVector X = XAxis * Scale;
 				const FVector Y = YAxis * Scale;
@@ -200,13 +200,13 @@ namespace NDIDebugDrawLocal
 		};
 
 #if NIAGARA_COMPUTEDEBUG_ENABLED
-		static void Draw(FNDIDebugDrawInstanceData_GameThread* InstanceData, VMBindings& Bindings, bool bDraw)
+		static void Draw(FNDIDebugDrawInstanceData_GameThread* InstanceData, VMBindings& Bindings, bool bExecute)
 		{
 			const FVector Location = Bindings.LocationParam.GetAndAdvance();
 			const FQuat Rotation = Bindings.RotationParam.GetAndAdvance();
 			const float Scale = Bindings.ScaleParam.GetAndAdvance();
 
-			if (bDraw)
+			if (bExecute)
 			{
 				const FVector XAxis = Rotation.RotateVector(FVector(Scale, 0.0f, 0.0f));
 				const FVector YAxis = Rotation.RotateVector(FVector(0.0f, Scale, 0.0f));
@@ -243,7 +243,7 @@ namespace NDIDebugDrawLocal
 		};
 
 #if NIAGARA_COMPUTEDEBUG_ENABLED
-		static void Draw(FNDIDebugDrawInstanceData_GameThread* InstanceData, VMBindings& Bindings, bool bDraw)
+		static void Draw(FNDIDebugDrawInstanceData_GameThread* InstanceData, VMBindings& Bindings, bool bExecute)
 		{
 			const FVector Center = Bindings.CenterParam.GetAndAdvance();
 			const FQuat Rotation = Bindings.RotationParam.GetAndAdvance();
@@ -251,7 +251,7 @@ namespace NDIDebugDrawLocal
 			const FIntPoint NumCells(Bindings.NumCellsXParam.GetAndAdvance(), Bindings.NumCellsYParam.GetAndAdvance());
 			const FLinearColor Color = Bindings.ColorParam.GetAndAdvance();
 
-			if (bDraw && NumCells.X > 0 && NumCells.Y > 0)
+			if (bExecute && NumCells.X > 0 && NumCells.Y > 0)
 			{
 				const FVector Corner = Center - Rotation.RotateVector(FVector(Extents.X, Extents.Y, 0.0f));
 				const FVector XLength = Rotation.RotateVector(FVector(Extents.X * 2.0f, 0.0f, 0.0f));
@@ -299,7 +299,7 @@ namespace NDIDebugDrawLocal
 		};
 
 #if NIAGARA_COMPUTEDEBUG_ENABLED
-		static void Draw(FNDIDebugDrawInstanceData_GameThread* InstanceData, VMBindings& Bindings, bool bDraw)
+		static void Draw(FNDIDebugDrawInstanceData_GameThread* InstanceData, VMBindings& Bindings, bool bExecute)
 		{
 			const FVector Center = Bindings.CenterParam.GetAndAdvance();
 			const FQuat Rotation = Bindings.RotationParam.GetAndAdvance();
@@ -307,7 +307,7 @@ namespace NDIDebugDrawLocal
 			const FIntVector NumCells(Bindings.NumCellsXParam.GetAndAdvance(), Bindings.NumCellsYParam.GetAndAdvance(), Bindings.NumCellsZParam.GetAndAdvance());
 			const FLinearColor Color = Bindings.ColorParam.GetAndAdvance();
 
-			if (bDraw && NumCells.X > 0 && NumCells.Y > 0 && NumCells.Z > 0)
+			if (bExecute && NumCells.X > 0 && NumCells.Y > 0 && NumCells.Z > 0)
 			{
 				const FVector Corner = Center - Rotation.RotateVector(Extents);
 				const FVector XLength = Rotation.RotateVector(FVector(Extents.X * 2.0f, 0.0f, 0.0f));
@@ -357,13 +357,13 @@ namespace NDIDebugDrawLocal
 		};
 
 #if NIAGARA_COMPUTEDEBUG_ENABLED
-		static void Draw(FNDIDebugDrawInstanceData_GameThread* InstanceData, VMBindings& Bindings, bool bDraw)
+		static void Draw(FNDIDebugDrawInstanceData_GameThread* InstanceData, VMBindings& Bindings, bool bExecute)
 		{
 			const FVector LineStart = Bindings.LineStartParam.GetAndAdvance();
 			const FVector LineEnd = Bindings.LineEndParam.GetAndAdvance();
 			const FLinearColor Color = Bindings.ColorParam.GetAndAdvance();
 
-			if (bDraw)
+			if (bExecute)
 			{
 				InstanceData->AddLine(LineStart, LineEnd, Color);
 			}
@@ -390,14 +390,14 @@ namespace NDIDebugDrawLocal
 		};
 
 #if NIAGARA_COMPUTEDEBUG_ENABLED
-		static void Draw(FNDIDebugDrawInstanceData_GameThread* InstanceData, VMBindings& Bindings, bool bDraw)
+		static void Draw(FNDIDebugDrawInstanceData_GameThread* InstanceData, VMBindings& Bindings, bool bExecute)
 		{
 			const FVector Location = Bindings.LocationParam.GetAndAdvance();
 			const float Radius = Bindings.RadiusParam.GetAndAdvance();
 			const int32 Segments = FMath::Clamp(Bindings.SegmentsParam.GetAndAdvance(), 4, 16);
 			const FLinearColor Color = Bindings.ColorParam.GetAndAdvance();
 
-			if (bDraw)
+			if (bExecute)
 			{
 				const float uinc = 2.0f * PI / float(Segments);
 
@@ -441,7 +441,7 @@ namespace NDIDebugDrawLocal
 	void DrawDebug(FVectorVMContext& Context)
 	{
 		VectorVM::FUserPtrHandler<FNDIDebugDrawInstanceData_GameThread> InstanceData(Context);
-		FNDIInputParam<FNiagaraBool> DoDrawParam(Context);
+		FNDIInputParam<FNiagaraBool> ExecuteParam(Context);
 		typename TPrimType::VMBindings Bindings(Context);
 
 #if NIAGARA_COMPUTEDEBUG_ENABLED
@@ -452,8 +452,8 @@ namespace NDIDebugDrawLocal
 
 		for ( int i=0; i < Context.NumInstances; ++i )
 		{
-			const bool bDraw = DoDrawParam.GetAndAdvance();
-			TPrimType::Draw(InstanceData, Bindings, bDraw);
+			const bool bExecute = ExecuteParam.GetAndAdvance();
+			TPrimType::Draw(InstanceData, Bindings, bExecute);
 		}
 #endif
 	}
@@ -590,10 +590,10 @@ void UNiagaraDataInterfaceDebugDraw::GetFunctions(TArray<FNiagaraFunctionSignatu
 		FNiagaraFunctionSignature& Signature = OutFunctions.Add_GetRef(DefaultSignature);
 		Signature.Name = NDIDebugDrawLocal::DrawBoxName;
 		Signature.Inputs.Add(FNiagaraVariable(FNiagaraTypeDefinition(GetClass()), TEXT("DebugDrawInterface")));
-		Signature.Inputs.Add(FNiagaraVariable(FNiagaraTypeDefinition::GetBoolDef(), TEXT("bDraw")));
+		Signature.Inputs.Add_GetRef(FNiagaraVariable(FNiagaraTypeDefinition::GetBoolDef(), TEXT("Execute"))).SetValue(true);
 		Signature.Inputs.Add(FNiagaraVariable(FNiagaraTypeDefinition::GetVec3Def(), TEXT("Center")));
 		Signature.Inputs.Add(FNiagaraVariable(FNiagaraTypeDefinition::GetQuatDef(), TEXT("Rotation")));
-		Signature.Inputs.Add(FNiagaraVariable(FNiagaraTypeDefinition::GetVec3Def(), TEXT("Extents")));
+		Signature.Inputs.Add_GetRef(FNiagaraVariable(FNiagaraTypeDefinition::GetVec3Def(), TEXT("Extents"))).SetValue(FVector(10.0f));
 		Signature.Inputs.Add(FNiagaraVariable(FNiagaraTypeDefinition::GetColorDef(), TEXT("Color")));
 	}
 
@@ -601,12 +601,12 @@ void UNiagaraDataInterfaceDebugDraw::GetFunctions(TArray<FNiagaraFunctionSignatu
 		FNiagaraFunctionSignature & Signature = OutFunctions.Add_GetRef(DefaultSignature);
 		Signature.Name = NDIDebugDrawLocal::DrawCircleName;
 		Signature.Inputs.Add(FNiagaraVariable(FNiagaraTypeDefinition(GetClass()), TEXT("DebugDrawInterface")));
-		Signature.Inputs.Add(FNiagaraVariable(FNiagaraTypeDefinition::GetBoolDef(), TEXT("bDraw")));
+		Signature.Inputs.Add_GetRef(FNiagaraVariable(FNiagaraTypeDefinition::GetBoolDef(), TEXT("Execute"))).SetValue(true);
 		Signature.Inputs.Add(FNiagaraVariable(FNiagaraTypeDefinition::GetVec3Def(), TEXT("Center")));
-		Signature.Inputs.Add(FNiagaraVariable(FNiagaraTypeDefinition::GetVec3Def(), TEXT("XAxis")));
-		Signature.Inputs.Add(FNiagaraVariable(FNiagaraTypeDefinition::GetVec3Def(), TEXT("YAxis")));
-		Signature.Inputs.Add(FNiagaraVariable(FNiagaraTypeDefinition::GetFloatDef(), TEXT("Radius")));
-		Signature.Inputs.Add(FNiagaraVariable(FNiagaraTypeDefinition::GetIntDef(), TEXT("Num Segments")));
+		Signature.Inputs.Add_GetRef(FNiagaraVariable(FNiagaraTypeDefinition::GetVec3Def(), TEXT("XAxis"))).SetValue(FVector::XAxisVector);
+		Signature.Inputs.Add_GetRef(FNiagaraVariable(FNiagaraTypeDefinition::GetVec3Def(), TEXT("YAxis"))).SetValue(FVector::YAxisVector);
+		Signature.Inputs.Add_GetRef(FNiagaraVariable(FNiagaraTypeDefinition::GetFloatDef(), TEXT("Radius"))).SetValue(10.0f);
+		Signature.Inputs.Add_GetRef(FNiagaraVariable(FNiagaraTypeDefinition::GetIntDef(), TEXT("Num Segments"))).SetValue(6);
 		Signature.Inputs.Add(FNiagaraVariable(FNiagaraTypeDefinition::GetColorDef(), TEXT("Color")));
 	}
 
@@ -614,22 +614,22 @@ void UNiagaraDataInterfaceDebugDraw::GetFunctions(TArray<FNiagaraFunctionSignatu
 		FNiagaraFunctionSignature& Signature = OutFunctions.Add_GetRef(DefaultSignature);
 		Signature.Name = NDIDebugDrawLocal::DrawCoordinateSystemName;
 		Signature.Inputs.Add(FNiagaraVariable(FNiagaraTypeDefinition(GetClass()), TEXT("DebugDrawInterface")));
-		Signature.Inputs.Add(FNiagaraVariable(FNiagaraTypeDefinition::GetBoolDef(), TEXT("bDraw")));
+		Signature.Inputs.Add_GetRef(FNiagaraVariable(FNiagaraTypeDefinition::GetBoolDef(), TEXT("Execute"))).SetValue(true);
 		Signature.Inputs.Add(FNiagaraVariable(FNiagaraTypeDefinition::GetVec3Def(), TEXT("Location")));
 		Signature.Inputs.Add(FNiagaraVariable(FNiagaraTypeDefinition::GetQuatDef(), TEXT("Rotation")));
-		Signature.Inputs.Add(FNiagaraVariable(FNiagaraTypeDefinition::GetFloatDef(), TEXT("Scale")));
+		Signature.Inputs.Add_GetRef(FNiagaraVariable(FNiagaraTypeDefinition::GetFloatDef(), TEXT("Scale"))).SetValue(1.0f);
 	}
 
 	{
 		FNiagaraFunctionSignature& Signature = OutFunctions.Add_GetRef(DefaultSignature);
 		Signature.Name = NDIDebugDrawLocal::DrawGrid2DName;
 		Signature.Inputs.Add(FNiagaraVariable(FNiagaraTypeDefinition(GetClass()), TEXT("DebugDrawInterface")));
-		Signature.Inputs.Add(FNiagaraVariable(FNiagaraTypeDefinition::GetBoolDef(), TEXT("bDraw")));
+		Signature.Inputs.Add_GetRef(FNiagaraVariable(FNiagaraTypeDefinition::GetBoolDef(), TEXT("Execute"))).SetValue(true);
 		Signature.Inputs.Add(FNiagaraVariable(FNiagaraTypeDefinition::GetVec3Def(), TEXT("Center")));
 		Signature.Inputs.Add(FNiagaraVariable(FNiagaraTypeDefinition::GetQuatDef(), TEXT("Rotation")));
-		Signature.Inputs.Add(FNiagaraVariable(FNiagaraTypeDefinition::GetVec2Def(), TEXT("Extents")));
-		Signature.Inputs.Add(FNiagaraVariable(FNiagaraTypeDefinition::GetIntDef(), TEXT("NumCellsX")));
-		Signature.Inputs.Add(FNiagaraVariable(FNiagaraTypeDefinition::GetIntDef(), TEXT("NumCellsY")));
+		Signature.Inputs.Add_GetRef(FNiagaraVariable(FNiagaraTypeDefinition::GetVec2Def(), TEXT("Extents"))).SetValue(FVector2D(10.0f));
+		Signature.Inputs.Add_GetRef(FNiagaraVariable(FNiagaraTypeDefinition::GetIntDef(), TEXT("NumCellsX"))).SetValue(1);
+		Signature.Inputs.Add_GetRef(FNiagaraVariable(FNiagaraTypeDefinition::GetIntDef(), TEXT("NumCellsY"))).SetValue(1);
 		Signature.Inputs.Add(FNiagaraVariable(FNiagaraTypeDefinition::GetColorDef(), TEXT("Color")));
 	}
 
@@ -637,13 +637,13 @@ void UNiagaraDataInterfaceDebugDraw::GetFunctions(TArray<FNiagaraFunctionSignatu
 		FNiagaraFunctionSignature& Signature = OutFunctions.Add_GetRef(DefaultSignature);
 		Signature.Name = NDIDebugDrawLocal::DrawGrid3DName;
 		Signature.Inputs.Add(FNiagaraVariable(FNiagaraTypeDefinition(GetClass()), TEXT("DebugDrawInterface")));
-		Signature.Inputs.Add(FNiagaraVariable(FNiagaraTypeDefinition::GetBoolDef(), TEXT("bDraw")));
+		Signature.Inputs.Add_GetRef(FNiagaraVariable(FNiagaraTypeDefinition::GetBoolDef(), TEXT("Execute"))).SetValue(true);
 		Signature.Inputs.Add(FNiagaraVariable(FNiagaraTypeDefinition::GetVec3Def(), TEXT("Center")));
 		Signature.Inputs.Add(FNiagaraVariable(FNiagaraTypeDefinition::GetQuatDef(), TEXT("Rotation")));
-		Signature.Inputs.Add(FNiagaraVariable(FNiagaraTypeDefinition::GetVec3Def(), TEXT("Extents")));
-		Signature.Inputs.Add(FNiagaraVariable(FNiagaraTypeDefinition::GetIntDef(), TEXT("NumCellsX")));
-		Signature.Inputs.Add(FNiagaraVariable(FNiagaraTypeDefinition::GetIntDef(), TEXT("NumCellsY")));
-		Signature.Inputs.Add(FNiagaraVariable(FNiagaraTypeDefinition::GetIntDef(), TEXT("NumCellsZ")));
+		Signature.Inputs.Add_GetRef(FNiagaraVariable(FNiagaraTypeDefinition::GetVec3Def(), TEXT("Extents"))).SetValue(FVector(10.0f));
+		Signature.Inputs.Add_GetRef(FNiagaraVariable(FNiagaraTypeDefinition::GetIntDef(), TEXT("NumCellsX"))).SetValue(1);
+		Signature.Inputs.Add_GetRef(FNiagaraVariable(FNiagaraTypeDefinition::GetIntDef(), TEXT("NumCellsY"))).SetValue(1);
+		Signature.Inputs.Add_GetRef(FNiagaraVariable(FNiagaraTypeDefinition::GetIntDef(), TEXT("NumCellsZ"))).SetValue(1);
 		Signature.Inputs.Add(FNiagaraVariable(FNiagaraTypeDefinition::GetColorDef(), TEXT("Color")));
 	}
 
@@ -651,7 +651,7 @@ void UNiagaraDataInterfaceDebugDraw::GetFunctions(TArray<FNiagaraFunctionSignatu
 		FNiagaraFunctionSignature & Signature = OutFunctions.Add_GetRef(DefaultSignature);
 		Signature.Name = NDIDebugDrawLocal::DrawLineName;
 		Signature.Inputs.Add(FNiagaraVariable(FNiagaraTypeDefinition(GetClass()), TEXT("DebugDrawInterface")));
-		Signature.Inputs.Add(FNiagaraVariable(FNiagaraTypeDefinition::GetBoolDef(), TEXT("bDraw")));
+		Signature.Inputs.Add_GetRef(FNiagaraVariable(FNiagaraTypeDefinition::GetBoolDef(), TEXT("Execute"))).SetValue(true);
 		Signature.Inputs.Add(FNiagaraVariable(FNiagaraTypeDefinition::GetVec3Def(), TEXT("Start Location")));
 		Signature.Inputs.Add(FNiagaraVariable(FNiagaraTypeDefinition::GetVec3Def(), TEXT("End Location")));
 		Signature.Inputs.Add(FNiagaraVariable(FNiagaraTypeDefinition::GetColorDef(), TEXT("Color")));
@@ -661,10 +661,10 @@ void UNiagaraDataInterfaceDebugDraw::GetFunctions(TArray<FNiagaraFunctionSignatu
 		FNiagaraFunctionSignature& Signature = OutFunctions.Add_GetRef(DefaultSignature);
 		Signature.Name = NDIDebugDrawLocal::DrawSphereName;
 		Signature.Inputs.Add(FNiagaraVariable(FNiagaraTypeDefinition(GetClass()), TEXT("DebugDrawInterface")));
-		Signature.Inputs.Add(FNiagaraVariable(FNiagaraTypeDefinition::GetBoolDef(), TEXT("bDraw")));
+		Signature.Inputs.Add_GetRef(FNiagaraVariable(FNiagaraTypeDefinition::GetBoolDef(), TEXT("Execute"))).SetValue(true);
 		Signature.Inputs.Add(FNiagaraVariable(FNiagaraTypeDefinition::GetVec3Def(), TEXT("Center")));
-		Signature.Inputs.Add(FNiagaraVariable(FNiagaraTypeDefinition::GetFloatDef(), TEXT("Radius")));
-		Signature.Inputs.Add(FNiagaraVariable(FNiagaraTypeDefinition::GetIntDef(), TEXT("Num Segments")));
+		Signature.Inputs.Add_GetRef(FNiagaraVariable(FNiagaraTypeDefinition::GetFloatDef(), TEXT("Radius"))).SetValue(10.0f);
+		Signature.Inputs.Add_GetRef(FNiagaraVariable(FNiagaraTypeDefinition::GetIntDef(), TEXT("Num Segments"))).SetValue(6);
 		Signature.Inputs.Add(FNiagaraVariable(FNiagaraTypeDefinition::GetColorDef(), TEXT("Color")));
 	}
 }
@@ -726,43 +726,43 @@ bool UNiagaraDataInterfaceDebugDraw::GetFunctionHLSL(const FNiagaraDataInterface
 
 	if (FunctionInfo.DefinitionName == NDIDebugDrawLocal::DrawBoxName)
 	{
-		static const TCHAR* FormatSample = TEXT("void {InstanceFunctionName}(bool bDraw, float3 Location, float4 Rotation, float3 Extents, float4 Color) { NDIDebugDraw_DrawBox(bDraw, Location, Rotation, Extents, Color); }\n");
+		static const TCHAR* FormatSample = TEXT("void {InstanceFunctionName}(bool bExecute, float3 Location, float4 Rotation, float3 Extents, float4 Color) { NDIDebugDraw_DrawBox(bExecute, Location, Rotation, Extents, Color); }\n");
 		OutHLSL += FString::Format(FormatSample, ArgsSample);
 		return true;
 	}
 	else if (FunctionInfo.DefinitionName == NDIDebugDrawLocal::DrawCircleName)
 	{
-		static const TCHAR* FormatSample = TEXT("void {InstanceFunctionName}(bool bDraw, float3 Location, float3 XAxis, float3 YAxis, float Scale, int Segments, float4 Color) { NDIDebugDraw_Circle(bDraw, Location, XAxis, YAxis, Scale, Segments, Color); }\n");
+		static const TCHAR* FormatSample = TEXT("void {InstanceFunctionName}(bool bExecute, float3 Location, float3 XAxis, float3 YAxis, float Scale, int Segments, float4 Color) { NDIDebugDraw_Circle(bExecute, Location, XAxis, YAxis, Scale, Segments, Color); }\n");
 		OutHLSL += FString::Format(FormatSample, ArgsSample);
 		return true;
 	}
 	else if (FunctionInfo.DefinitionName == NDIDebugDrawLocal::DrawCoordinateSystemName)
 	{
-		static const TCHAR* FormatSample = TEXT("void {InstanceFunctionName}(bool bDraw, float3 Location, float4 Rotation, float Scale) { NDIDebugDraw_CoordinateSystem(bDraw, Location, Rotation, Scale); }\n");
+		static const TCHAR* FormatSample = TEXT("void {InstanceFunctionName}(bool bExecute, float3 Location, float4 Rotation, float Scale) { NDIDebugDraw_CoordinateSystem(bExecute, Location, Rotation, Scale); }\n");
 		OutHLSL += FString::Format(FormatSample, ArgsSample);
 		return true;
 	}
 	else if (FunctionInfo.DefinitionName == NDIDebugDrawLocal::DrawGrid2DName)
 	{
-		static const TCHAR* FormatSample = TEXT("void {InstanceFunctionName}(bool bDraw, float3 Center, float4 Rotation, float2 Extents, int NumCellsX, int NumCellsY, float4 Color) { NDIDebugDraw_Grid2D(bDraw, Center, Rotation, Extents, int2(NumCellsX, NumCellsY), Color); }\n");
+		static const TCHAR* FormatSample = TEXT("void {InstanceFunctionName}(bool bExecute, float3 Center, float4 Rotation, float2 Extents, int NumCellsX, int NumCellsY, float4 Color) { NDIDebugDraw_Grid2D(bExecute, Center, Rotation, Extents, int2(NumCellsX, NumCellsY), Color); }\n");
 		OutHLSL += FString::Format(FormatSample, ArgsSample);
 		return true;
 	}
 	else if (FunctionInfo.DefinitionName == NDIDebugDrawLocal::DrawGrid3DName)
 	{
-		static const TCHAR* FormatSample = TEXT("void {InstanceFunctionName}(bool bDraw, float3 Center, float4 Rotation, float3 Extents, int NumCellsX, int NumCellsY, int NumCellsZ, float4 Color) { NDIDebugDraw_Grid3D(bDraw, Center, Rotation, Extents, int3(NumCellsX, NumCellsY, NumCellsZ), Color); }\n");
+		static const TCHAR* FormatSample = TEXT("void {InstanceFunctionName}(bool bExecute, float3 Center, float4 Rotation, float3 Extents, int NumCellsX, int NumCellsY, int NumCellsZ, float4 Color) { NDIDebugDraw_Grid3D(bExecute, Center, Rotation, Extents, int3(NumCellsX, NumCellsY, NumCellsZ), Color); }\n");
 		OutHLSL += FString::Format(FormatSample, ArgsSample);
 		return true;
 	}
 	else if (FunctionInfo.DefinitionName == NDIDebugDrawLocal::DrawLineName)
 	{
-		static const TCHAR* FormatSample = TEXT("void {InstanceFunctionName}(bool bDraw, float3 LineStart, float3 LineEnd, float4 Color) { NDIDebugDraw_Line(bDraw, LineStart, LineEnd, Color); }\n");
+		static const TCHAR* FormatSample = TEXT("void {InstanceFunctionName}(bool bExecute, float3 LineStart, float3 LineEnd, float4 Color) { NDIDebugDraw_Line(bExecute, LineStart, LineEnd, Color); }\n");
 		OutHLSL += FString::Format(FormatSample, ArgsSample);
 		return true;
 	}
 	else if (FunctionInfo.DefinitionName == NDIDebugDrawLocal::DrawSphereName)
 	{
-		static const TCHAR* FormatSample = TEXT("void {InstanceFunctionName}(bool bDraw, float3 LineStart, float Radius, int Segments, float4 Color) { NDIDebugDraw_Sphere(bDraw, LineStart, Radius, Segments, Color); }\n");
+		static const TCHAR* FormatSample = TEXT("void {InstanceFunctionName}(bool bExecute, float3 LineStart, float Radius, int Segments, float4 Color) { NDIDebugDraw_Sphere(bExecute, LineStart, Radius, Segments, Color); }\n");
 		OutHLSL += FString::Format(FormatSample, ArgsSample);
 		return true;
 	}
diff --git a/Engine/Plugins/FX/Niagara/Source/Niagara/Private/NiagaraGpuComputeDebug.cpp b/Engine/Plugins/FX/Niagara/Source/Niagara/Private/NiagaraGpuComputeDebug.cpp
--- a/Engine/Plugins/FX/Niagara/Source/Niagara/Private/NiagaraGpuComputeDebug.cpp
+++ b/Engine/Plugins/FX/Niagara/Source/Niagara/Private/NiagaraGpuComputeDebug.cpp
@@ -220,9 +220,11 @@ FNiagaraSimulationDebugDrawData* FNiagaraGpuComputeDebug::GetSimulationDebugDraw
 		DebugDrawDataPtr.Reset(new FNiagaraSimulationDebugDrawData());
 	}
 
-	if (bRequiresGpuBuffers && DebugDrawDataPtr->GpuLineMaxInstances == 0)
+	if (bRequiresGpuBuffers && (DebugDrawDataPtr->GpuLineMaxInstances != GNiagaraGpuComputeDebug_MaxLineInstances))
 	{
 		check(IsInRenderingThread());
+		DebugDrawDataPtr->GpuLineBufferArgs.Release();
+		DebugDrawDataPtr->GpuLineVertexBuffer.Release();
 		DebugDrawDataPtr->GpuLineMaxInstances = GNiagaraGpuComputeDebug_MaxLineInstances;
 		if (DebugDrawDataPtr->GpuLineMaxInstances > 0)
 		{
