Fix for warning being generated when the parameter store change notification was being triggered for a component that had not finished loading

diff --git a/Engine/Plugins/FX/Niagara/Source/Niagara/Private/NiagaraComponent.cpp b/Engine/Plugins/FX/Niagara/Source/Niagara/Private/NiagaraComponent.cpp
--- a/Engine/Plugins/FX/Niagara/Source/Niagara/Private/NiagaraComponent.cpp
+++ b/Engine/Plugins/FX/Niagara/Source/Niagara/Private/NiagaraComponent.cpp
@@ -3001,8 +3001,12 @@ void UNiagaraComponent::FixInvalidUserParameterOverrideData()
 
 void UNiagaraComponent::AssetExposedParametersChanged()
 {
-	SynchronizeWithSourceSystem();
-	ReinitializeSystem();
+	// don't worry about doing anything based on the change broadcast if we haven't finished loading
+	if (!HasAnyFlags(RF_NeedPostLoad))
+	{
+		SynchronizeWithSourceSystem();
+		ReinitializeSystem();
+	}
 }
 
 #if WITH_EDITOR
