Fixed a bug where module outputs could be wrongly reassigned when multiple modules of the same type were used in the stack

diff --git a/Engine/Plugins/FX/Niagara/Source/NiagaraEditor/Private/NiagaraNodeFunctionCall.cpp b/Engine/Plugins/FX/Niagara/Source/NiagaraEditor/Private/NiagaraNodeFunctionCall.cpp
--- a/Engine/Plugins/FX/Niagara/Source/NiagaraEditor/Private/NiagaraNodeFunctionCall.cpp
+++ b/Engine/Plugins/FX/Niagara/Source/NiagaraEditor/Private/NiagaraNodeFunctionCall.cpp
@@ -421,7 +421,7 @@ bool UNiagaraNodeFunctionCall::FixupPinNames()
 
 		// see if any of the linked inputs were renamed
 		UNiagaraGraph* Owner = GetNiagaraGraph();
-		TMap<FGuid, FNiagaraVariable> GraphVariablesGuidMapping;
+		TMap<FGuid, TArray<FNiagaraVariable>> GraphVariablesGuidMapping;
 		if (Owner && LinkedModuleOutputs.Num() > 0)
 		{
 			// Search for guids of the linked inputs by looking backwards in the current script graph (not the graph this function node calls, but the one it's embedded in)
@@ -448,7 +448,7 @@ bool UNiagaraNodeFunctionCall::FixupPinNames()
 							FNiagaraParameterHandle AliasedFunctionInputHandle(FName(AliasedNamespace), MetadataName);
 							FNiagaraVariable OutputVar = Entry.Key;
 							OutputVar.SetName(AliasedFunctionInputHandle.GetParameterHandleString());
-							GraphVariablesGuidMapping.Add(Entry.Value->Metadata.GetVariableGuid(), OutputVar);
+							GraphVariablesGuidMapping.FindOrAdd(Entry.Value->Metadata.GetVariableGuid()).Add(OutputVar);
 						}
 					}
 				}
@@ -463,13 +463,32 @@ bool UNiagaraNodeFunctionCall::FixupPinNames()
 				for (FGuid Guid : GuidMatches)
 				{
 					// see if the name still matches and rename the linked input if not
-					FNiagaraVariable* VarFromGraph = GraphVariablesGuidMapping.Find(Guid);
-					if (VarFromGraph)
+					TArray<FNiagaraVariable>* VarsFromGraph = GraphVariablesGuidMapping.Find(Guid);
+					if (VarsFromGraph)
 					{
 						bFoundMatch = true;
-						if (VarFromGraph->GetName() != LinkedOutPin->PinName)
+
+						// there might be several modules in different version giving us the same guid for the output, so we check them all
+						if (!VarsFromGraph->ContainsByPredicate([LinkedOutPin](const FNiagaraVariable& Var) { return LinkedOutPin->PinName == Var.GetName(); }))
 						{
-							LinkedOutPin->PinName = VarFromGraph->GetName();
+							FNiagaraParameterHandle PinHandle(LinkedOutPin->PinName);
+							FNiagaraVariable* OriginalModuleVar = VarsFromGraph->FindByPredicate([PinHandle](const FNiagaraVariable& Var)
+							{
+								FNiagaraParameterHandle ModuleVarHandle(Var.GetName());
+								return PinHandle.GetNamespace() == ModuleVarHandle.GetNamespace();
+							});
+							
+							if (OriginalModuleVar)
+							{
+								// we found a reference from the original module namespace, so use that
+								LinkedOutPin->PinName = OriginalModuleVar->GetName();
+							}
+							else
+							{
+								// seems like the module was lost or got replaced, so just use any one of the new handles
+								LinkedOutPin->PinName = (*VarsFromGraph)[0].GetName();
+							}
+							
 							bNodeChanged = true;
 							UpdateInputNameBinding(Guid, LinkedOutPin->PinName);
 							break;
@@ -483,7 +502,8 @@ bool UNiagaraNodeFunctionCall::FixupPinNames()
 					FNiagaraVariable InputVariable = NiagaraSchema->PinToNiagaraVariable(LinkedOutPin);
 					for (auto& Entry : GraphVariablesGuidMapping)
 					{
-						if (Entry.Value == InputVariable)
+						FNiagaraVariable* MatchingVar = Entry.Value.FindByPredicate([InputVariable](const FNiagaraVariable& Var){ return InputVariable == Var; });
+						if (MatchingVar)
 						{
 							UpdateInputNameBinding(Entry.Key, LinkedOutPin->PinName);
 						}
