- Integrate CL 15668475: Fixing array data interface values for GPU sims

diff --git a/Engine/Plugins/FX/Niagara/Source/Niagara/Private/NiagaraDataInterfaceArray.cpp b/Engine/Plugins/FX/Niagara/Source/Niagara/Private/NiagaraDataInterfaceArray.cpp
--- a/Engine/Plugins/FX/Niagara/Source/Niagara/Private/NiagaraDataInterfaceArray.cpp
+++ b/Engine/Plugins/FX/Niagara/Source/Niagara/Private/NiagaraDataInterfaceArray.cpp
@@ -45,13 +45,15 @@ bool UNiagaraDataInterfaceArray::CopyToInternal(UNiagaraDataInterface* Destinati
 		return false;
 	}
 	UNiagaraDataInterfaceArray* OtherTyped = CastChecked<UNiagaraDataInterfaceArray>(Destination);
-	OtherTyped->MarkRenderDataDirty();
 	OtherTyped->MaxElements = MaxElements;
+	bool bCopied = true;
 	if (ensureMsgf(Impl.IsValid(), TEXT("Impl should always be valid for %s"), *GetNameSafe(GetClass())))
 	{
-		return Impl->CopyToInternal(OtherTyped->Impl.Get());
+		bCopied = Impl->CopyToInternal(OtherTyped->Impl.Get());
 	}
-	return true;
+
+	OtherTyped->MarkRenderDataDirty();
+	return bCopied;
 }
 
 bool UNiagaraDataInterfaceArray::Equals(const UNiagaraDataInterface* Other) const
