- Integrate CL 15671597: Transform view space velocity into world space

#rnx

diff --git a/Engine/Plugins/FX/Niagara/Shaders/Private/NiagaraDataInterfaceGBuffer.ush b/Engine/Plugins/FX/Niagara/Shaders/Private/NiagaraDataInterfaceGBuffer.ush
--- a/Engine/Plugins/FX/Niagara/Shaders/Private/NiagaraDataInterfaceGBuffer.ush
+++ b/Engine/Plugins/FX/Niagara/Shaders/Private/NiagaraDataInterfaceGBuffer.ush
@@ -50,7 +50,11 @@ void DIGBuffer_DecodeScreenVelocity(in float2 ScreenUV, out bool IsValid, out fl
 	IsValid = true;
 	if ( EncodedVelocity.x > 0 )
 	{
+	#if VELOCITY_ENCODE_DEPTH
+		OutValue = DecodeVelocityFromTexture(EncodedVelocity);
+	#else
 		OutValue.xy = DecodeVelocityFromTexture(EncodedVelocity).xy;
+	#endif
 	}
 #endif
 }
@@ -71,6 +75,7 @@ void DIGBuffer_DecodeWorldVelocity(in float2 ScreenUV, out bool IsValid, out flo
 	{
 	#if VELOCITY_ENCODE_DEPTH
 		OutValue = DecodeVelocityFromTexture(EncodedVelocity);
+		OutValue = mul(float4(OutValue, 0), View.ViewToTranslatedWorld).xyz;
 	#else
 		float2 ScreenVelocity = DecodeVelocityFromTexture(EncodedVelocity).xy;
 		ScreenVelocity.y = -ScreenVelocity.y;
