- Editgrate CL 15617533: Staticswitch engine test fix. Forgot a -1 during refactoring

diff --git a/Engine/Plugins/FX/Niagara/Source/NiagaraEditor/Private/NiagaraNodeStaticSwitch.cpp b/Engine/Plugins/FX/Niagara/Source/NiagaraEditor/Private/NiagaraNodeStaticSwitch.cpp
--- a/Engine/Plugins/FX/Niagara/Source/NiagaraEditor/Private/NiagaraNodeStaticSwitch.cpp
+++ b/Engine/Plugins/FX/Niagara/Source/NiagaraEditor/Private/NiagaraNodeStaticSwitch.cpp
@@ -345,7 +345,7 @@ bool UNiagaraNodeStaticSwitch::GetVarIndex(FHlslNiagaraTranslator* Translator, i
 	}
 	else if (SwitchTypeData.SwitchType == ENiagaraStaticSwitchType::Integer)
 	{
-		int32 MaxValue = NumOptionsPerVariable;
+		int32 MaxValue = NumOptionsPerVariable - 1;
 		if (MaxValue >= 0)
 		{
 			if (Translator && (Value > MaxValue || Value < 0))
