- Fix missing transitions when copying data out from Grid3D

#rnx

diff --git a/Engine/Plugins/FX/Niagara/Source/Niagara/Private/NiagaraDataInterfaceGrid3DCollection.cpp b/Engine/Plugins/FX/Niagara/Source/Niagara/Private/NiagaraDataInterfaceGrid3DCollection.cpp
--- a/Engine/Plugins/FX/Niagara/Source/Niagara/Private/NiagaraDataInterfaceGrid3DCollection.cpp
+++ b/Engine/Plugins/FX/Niagara/Source/Niagara/Private/NiagaraDataInterfaceGrid3DCollection.cpp
@@ -2694,18 +2694,23 @@ void FNiagaraDataInterfaceProxyGrid3DCollectionProxy::PostSimulate(FRHICommandLi
 	{
 		FRHITexture* Source = ProxyData->CurrentData->GridBuffer.Buffer;
 		FRHITexture* Destination = ProxyData->RenderTargetToCopyTo;
-		FRHITransitionInfo TransitionsBefore[] = {
-			FRHITransitionInfo(Source, ERHIAccess::SRVMask, ERHIAccess::CopySrc),
-			FRHITransitionInfo(Destination, ERHIAccess::SRVMask, ERHIAccess::CopyDest)
-		};
+
+		RHICmdList.Transition(
+			{
+				FRHITransitionInfo(Source, ERHIAccess::SRVMask, ERHIAccess::CopySrc),
+				FRHITransitionInfo(Destination, ERHIAccess::SRVMask, ERHIAccess::CopyDest)
+			}
+		);
 
 		FRHICopyTextureInfo CopyInfo;
 		RHICmdList.CopyTexture(ProxyData->CurrentData->GridBuffer.Buffer, ProxyData->RenderTargetToCopyTo, CopyInfo);
-		FRHITransitionInfo TransitionsAfter[] = {
-			FRHITransitionInfo(Source, ERHIAccess::CopySrc, ERHIAccess::SRVMask),
-			FRHITransitionInfo(Destination, ERHIAccess::CopyDest, ERHIAccess::SRVMask)
-		};
 
+		RHICmdList.Transition(
+			{
+				FRHITransitionInfo(Source, ERHIAccess::CopySrc, ERHIAccess::SRVMask),
+				FRHITransitionInfo(Destination, ERHIAccess::CopyDest, ERHIAccess::SRVMask)
+			}
+		);
 	}
 
 #if WITH_EDITOR
