- Integrate CL 15683149 : Fix incorrect transition for new buffers

#rnx

diff --git a/Engine/Plugins/FX/Niagara/Source/Niagara/Private/NiagaraGPUInstanceCountManager.cpp b/Engine/Plugins/FX/Niagara/Source/Niagara/Private/NiagaraGPUInstanceCountManager.cpp
--- a/Engine/Plugins/FX/Niagara/Source/Niagara/Private/NiagaraGPUInstanceCountManager.cpp
+++ b/Engine/Plugins/FX/Niagara/Source/Niagara/Private/NiagaraGPUInstanceCountManager.cpp
@@ -383,7 +383,7 @@ void FNiagaraGPUInstanceCountManager::UpdateDrawIndirectBuffers(NiagaraEmitterIn
 			TArray<FRHITransitionInfo, TInlineAllocator<10>> Transitions;
 			for (auto& PoolEntry : DrawIndirectPool)
 			{
-				Transitions.Emplace(PoolEntry->Buffer.UAV, ERHIAccess::IndirectArgs, ERHIAccess::UAVCompute);
+				Transitions.Emplace(PoolEntry->Buffer.UAV, ERHIAccess::Unknown, ERHIAccess::UAVCompute);
 			}
 			Transitions.Emplace(CountBuffer.UAV, kCountBufferDefaultState, ERHIAccess::UAVCompute);
 			if (CulledCountBuffer.SRV.IsValid())
