Fix binding a null counts UAV that could occur when updating args but no GPU counts had been allocated.

diff --git a/Engine/Plugins/FX/Niagara/Source/Niagara/Private/NiagaraGPUInstanceCountManager.cpp b/Engine/Plugins/FX/Niagara/Source/Niagara/Private/NiagaraGPUInstanceCountManager.cpp
--- a/Engine/Plugins/FX/Niagara/Source/Niagara/Private/NiagaraGPUInstanceCountManager.cpp
+++ b/Engine/Plugins/FX/Niagara/Source/Niagara/Private/NiagaraGPUInstanceCountManager.cpp
@@ -378,14 +378,25 @@ void FNiagaraGPUInstanceCountManager::UpdateDrawIndirectBuffers(NiagaraEmitterIn
 				RHIUnlockBuffer(TaskInfosBuffer.Buffer);
 			}
 
-			FRHIShaderResourceView* CulledCountsSRV;
-
+			FUnorderedAccessViewRHIRef CountsUAV = nullptr;
+			FShaderResourceViewRHIRef CulledCountsSRV = nullptr;
 			TArray<FRHITransitionInfo, TInlineAllocator<10>> Transitions;
 			for (auto& PoolEntry : DrawIndirectPool)
 			{
 				Transitions.Emplace(PoolEntry->Buffer.UAV, ERHIAccess::Unknown, ERHIAccess::UAVCompute);
 			}
-			Transitions.Emplace(CountBuffer.UAV, kCountBufferDefaultState, ERHIAccess::UAVCompute);
+
+			if (CountBuffer.UAV.IsValid())
+			{
+				Transitions.Emplace(CountBuffer.UAV, kCountBufferDefaultState, ERHIAccess::UAVCompute);
+				CountsUAV = CountBuffer.UAV;
+			}
+			else
+			{
+				// This can happen if there are no InstanceCountClearTasks and all DrawIndirectArgGenTasks are using culled counts
+				CountsUAV = Batcher.GetEmptyUAVFromPool(RHICmdList, PF_R32_UINT, false);
+			}
+
 			if (CulledCountBuffer.SRV.IsValid())
 			{
 				if (bAcquiredCulledCounts)
@@ -426,7 +437,7 @@ void FNiagaraGPUInstanceCountManager::UpdateDrawIndirectBuffers(NiagaraEmitterIn
 				const int32 NumInstanceCountClearTasks = bIsLastDispatch ? InstanceCountClearTasks.Num() : 0;
 
 				RHICmdList.SetComputeShader(DrawIndirectArgsGenCS.GetComputeShader());
-				DrawIndirectArgsGenCS->SetOutput(RHICmdList, ArgsUAV, CountBuffer.UAV);
+				DrawIndirectArgsGenCS->SetOutput(RHICmdList, ArgsUAV, CountsUAV);
 				DrawIndirectArgsGenCS->SetParameters(RHICmdList, TaskInfosBuffer.SRV, CulledCountsSRV, ArgGenTaskOffset, NumArgGenTasks, NumInstanceCountClearTasks);
 
 				// If the device supports RW Texture buffers then we can use a single compute pass, otherwise we need to split into two passes
