fix naming conflict

diff --git a/Engine/Plugins/FX/Niagara/Source/Niagara/Private/NiagaraDataInterfaceNeighborGrid3D.cpp b/Engine/Plugins/FX/Niagara/Source/Niagara/Private/NiagaraDataInterfaceNeighborGrid3D.cpp
--- a/Engine/Plugins/FX/Niagara/Source/Niagara/Private/NiagaraDataInterfaceNeighborGrid3D.cpp
+++ b/Engine/Plugins/FX/Niagara/Source/Niagara/Private/NiagaraDataInterfaceNeighborGrid3D.cpp
@@ -9,7 +9,7 @@
 
 #define LOCTEXT_NAMESPACE "NiagaraDataInterfaceNeighborGrid3D"
 
-static const FString MaxNeighborsPerCellName(TEXT("MaxNeighborsPerCell_"));
+static const FString MaxNeighborsPerCellName(TEXT("MaxNeighborsPerCellValue_"));
 static const FString ParticleNeighborsName(TEXT("ParticleNeighbors_"));
 static const FString ParticleNeighborCountName(TEXT("ParticleNeighborCount_"));
 static const FString OutputParticleNeighborsName(TEXT("OutputParticleNeighbors_"));
