Reskin and update of Niagara to apply current style guidelines

diff --git a/Engine/Plugins/FX/Niagara/Source/NiagaraEditor/Private/Customizations/NiagaraScriptVariableCustomization.cpp b/Engine/Plugins/FX/Niagara/Source/NiagaraEditor/Private/Customizations/NiagaraScriptVariableCustomization.cpp
--- a/Engine/Plugins/FX/Niagara/Source/NiagaraEditor/Private/Customizations/NiagaraScriptVariableCustomization.cpp
+++ b/Engine/Plugins/FX/Niagara/Source/NiagaraEditor/Private/Customizations/NiagaraScriptVariableCustomization.cpp
@@ -207,8 +207,9 @@ void FNiagaraScriptVariableDetails::CustomizeDetails(IDetailLayoutBuilder& Detai
 				.NameContent()
 				[
 					SNew(STextBlock)
-					.Font(FNiagaraEditorStyle::Get().GetFontStyle("NiagaraEditor.ParameterFont"))
+					.Font(FAppStyle::Get().GetFontStyle("PropertyWindow.NormalFont"))
 					.Text(FText::FromString(TEXT("Default Value")))
+					.ColorAndOpacity(FSlateColor::UseForeground())
 				]
 				.ValueContent()
 				.HAlign(HAlign_Fill)
@@ -230,8 +231,9 @@ void FNiagaraScriptVariableDetails::CustomizeDetails(IDetailLayoutBuilder& Detai
 					.HAlign(HAlign_Fill)
 					[
 						SNew(STextBlock)
-						.Font(FNiagaraEditorStyle::Get().GetFontStyle("NiagaraEditor.ParameterFont"))
-					.Text(NSLOCTEXT("ScriptVariableCustomization", "MissingDefaults", "To set default, add to a Map Get node that is wired to the graph."))
+						.Font(FAppStyle::Get().GetFontStyle("PropertyWindow.NormalFont"))
+						.Text(NSLOCTEXT("ScriptVariableCustomization", "MissingDefaults", "To set default, add to a Map Get node that is wired to the graph."))
+						.ColorAndOpacity(FSlateColor::UseForeground())
 					];
 			}
 		}
@@ -349,4 +351,4 @@ void FNiagaraScriptVariableDetails::OnStaticSwitchValueChanged()
 	} 
 }
  
-#undef LOCTEXT_NAMESPACE
\ No newline at end of file
+#undef LOCTEXT_NAMESPACE
diff --git a/Engine/Plugins/FX/Niagara/Source/NiagaraEditor/Private/Customizations/NiagaraTypeCustomizations.cpp b/Engine/Plugins/FX/Niagara/Source/NiagaraEditor/Private/Customizations/NiagaraTypeCustomizations.cpp
--- a/Engine/Plugins/FX/Niagara/Source/NiagaraEditor/Private/Customizations/NiagaraTypeCustomizations.cpp
+++ b/Engine/Plugins/FX/Niagara/Source/NiagaraEditor/Private/Customizations/NiagaraTypeCustomizations.cpp
@@ -417,8 +417,6 @@ void FNiagaraVariableAttributeBindingCustomization::CustomizeHeader(TSharedRef<I
 							.OnGetMenuContent(this, &FNiagaraVariableAttributeBindingCustomization::OnGetMenuContent)
 							.ContentPadding(1)
 							.ToolTipText(this, &FNiagaraVariableAttributeBindingCustomization::GetTooltipText)
-							.ButtonStyle(FEditorStyle::Get(), "PropertyEditor.AssetComboStyle")
-							.ForegroundColor(FEditorStyle::GetColor("PropertyEditor.AssetName.ColorAndOpacity"))
 							.ButtonContent()
 							[
 								/*SNew(STextBlock)
@@ -447,6 +445,7 @@ void FNiagaraVariableAttributeBindingCustomization::CustomizeHeader(TSharedRef<I
 							[
 								SNew(SImage)
 								.Image(FEditorStyle::GetBrush("PropertyWindow.DiffersFromDefault"))
+								.ColorAndOpacity(FSlateColor::UseForeground())
 							]
 						]
 					];
@@ -509,17 +508,16 @@ TSharedRef<SWidget> FNiagaraUserParameterBindingCustomization::OnGetMenuContent(
 	FGraphActionMenuBuilder MenuBuilder;
 
 	return SNew(SBorder)
-		.BorderImage(FEditorStyle::GetBrush("Menu.Background"))
 		.Padding(5)
 		[
 			SNew(SBox)
 			[
 				SNew(SGraphActionMenu)
 				.OnActionSelected(const_cast<FNiagaraUserParameterBindingCustomization*>(this), &FNiagaraUserParameterBindingCustomization::OnActionSelected)
-		.OnCreateWidgetForAction(SGraphActionMenu::FOnCreateWidgetForAction::CreateSP(const_cast<FNiagaraUserParameterBindingCustomization*>(this), &FNiagaraUserParameterBindingCustomization::OnCreateWidgetForAction))
-		.OnCollectAllActions(const_cast<FNiagaraUserParameterBindingCustomization*>(this), &FNiagaraUserParameterBindingCustomization::CollectAllActions)
-		.AutoExpandActionMenu(false)
-		.ShowFilterTextBox(true)
+				.OnCreateWidgetForAction(SGraphActionMenu::FOnCreateWidgetForAction::CreateSP(const_cast<FNiagaraUserParameterBindingCustomization*>(this), &FNiagaraUserParameterBindingCustomization::OnCreateWidgetForAction))
+				.OnCollectAllActions(const_cast<FNiagaraUserParameterBindingCustomization*>(this), &FNiagaraUserParameterBindingCustomization::CollectAllActions)
+				.AutoExpandActionMenu(false)
+				.ShowFilterTextBox(true)
 			]
 		];
 }
@@ -627,15 +625,13 @@ void FNiagaraUserParameterBindingCustomization::CustomizeHeader(TSharedRef<IProp
 				[
 					PropertyHandle->CreatePropertyNameWidget()
 				]
-			.ValueContent()
+				.ValueContent()
 				.MaxDesiredWidth(200.f)
 				[
 					SNew(SComboButton)
 					.OnGetMenuContent(this, &FNiagaraUserParameterBindingCustomization::OnGetMenuContent)
 					.ContentPadding(1)
 					.ToolTipText(this, &FNiagaraUserParameterBindingCustomization::GetTooltipText)
-					.ButtonStyle(FEditorStyle::Get(), "PropertyEditor.AssetComboStyle")
-					.ForegroundColor(FEditorStyle::GetColor("PropertyEditor.AssetName.ColorAndOpacity"))
 					.ButtonContent()
 					[
 
diff --git a/Engine/Plugins/FX/Niagara/Source/NiagaraEditor/Private/NiagaraEditorStyle.cpp b/Engine/Plugins/FX/Niagara/Source/NiagaraEditor/Private/NiagaraEditorStyle.cpp
--- a/Engine/Plugins/FX/Niagara/Source/NiagaraEditor/Private/NiagaraEditorStyle.cpp
+++ b/Engine/Plugins/FX/Niagara/Source/NiagaraEditor/Private/NiagaraEditorStyle.cpp
@@ -63,16 +63,13 @@ TSharedRef< FSlateStyleSet > FNiagaraEditorStyle::Create()
 {
 	const FTextBlockStyle NormalText = FEditorStyle::GetWidgetStyle<FTextBlockStyle>("NormalText");
 	const FEditableTextBoxStyle NormalEditableTextBox = FCoreStyle::Get().GetWidgetStyle<FEditableTextBoxStyle>("NormalEditableTextBox");
-	const FSpinBoxStyle NormalSpinBox = FEditorStyle::GetWidgetStyle<FSpinBoxStyle>("SpinBox");
+	const FSpinBoxStyle NormalSpinBox = FAppStyle::Get().GetWidgetStyle<FSpinBoxStyle>("SpinBox");
 
 	TSharedRef< FSlateStyleSet > Style = MakeShareable(new FSlateStyleSet("NiagaraEditorStyle"));
 	Style->SetContentRoot(FPaths::EngineContentDir() / TEXT("Editor/Slate/Niagara"));
 
 	// Stats
-	FTextBlockStyle CategoryText = FTextBlockStyle(NormalText)
-		.SetFont(DEFAULT_FONT("Regular", 10))
-		.SetShadowOffset(FVector2D(0, 1))
-		.SetShadowColorAndOpacity(FLinearColor(0, 0, 0, 0.9f));
+	const FTextBlockStyle CategoryText = FAppStyle::Get().GetWidgetStyle<FTextBlockStyle>("DetailsView.CategoryTextStyle");
 
 	Style->Set("NiagaraEditor.StatsText", CategoryText);
 
@@ -130,19 +127,17 @@ TSharedRef< FSlateStyleSet > FNiagaraEditorStyle::Create()
 	Style->Set("NiagaraEditor.NewAssetDialog.SubBorder", new BOX_CORE_BRUSH("Common/GroupBorderLight", FMargin(4.0f / 16.0f)));
 
 	// Emitter Header
-	FTextBlockStyle HeadingText = FTextBlockStyle(NormalText)
-		.SetFont(DEFAULT_FONT("Regular", 14));
-
-	FEditableTextBoxStyle HeadingEditableTextBox = FEditableTextBoxStyle(NormalEditableTextBox)
-		.SetFont(DEFAULT_FONT("Regular", 14));
+	FTextBlockStyle StackHeaderText = FTextBlockStyle(NormalText);
+	StackHeaderText.SetFont(DEFAULT_FONT("Regular", 11))
+		.SetColorAndOpacity(FSlateColor(EStyleColor::White));
 
-	Style->Set("NiagaraEditor.HeadingTextBlock", HeadingText);
+	Style->Set("NiagaraEditor.HeadingTextBlock", StackHeaderText);
 
-	Style->Set("NiagaraEditor.HeadingEditableTextBox", HeadingEditableTextBox);
+	FTextBlockStyle StackHeaderTextSubdued = FTextBlockStyle(NormalText);
+	StackHeaderTextSubdued.SetFont(DEFAULT_FONT("Regular", 11))
+		.SetColorAndOpacity(FStyleColors::Foreground);
 
-	Style->Set("NiagaraEditor.HeadingInlineEditableText", FInlineEditableTextBlockStyle()
-		.SetTextStyle(HeadingText)
-		.SetEditableTextBoxStyle(HeadingEditableTextBox));
+	Style->Set("NiagaraEditor.HeadingTextBlockSubdued", StackHeaderTextSubdued);
 
 	FTextBlockStyle TabText = FTextBlockStyle(NormalText)
 		.SetFont(DEFAULT_FONT("Regular", 12))
@@ -164,28 +159,21 @@ TSharedRef< FSlateStyleSet > FNiagaraEditorStyle::Create()
 	Style->Set("NiagaraEditor.DetailsHeadingText", DetailsHeadingText);
 
 	// Parameters
-	FSlateFontInfo ParameterFont = DEFAULT_FONT("Regular", 8);
-
-	Style->Set("NiagaraEditor.ParameterFont", ParameterFont);
-
+	
+	FSlateFontInfo NormalFont = FAppStyle::Get().GetFontStyle(TEXT("PropertyWindow.NormalFont"));
 	FTextBlockStyle ParameterText = FTextBlockStyle(NormalText)
-		.SetFont(ParameterFont);
+		.SetFont(NormalFont);
 
 	Style->Set("NiagaraEditor.ParameterText", ParameterText);
 
-	FEditableTextBoxStyle ParameterEditableTextBox = FEditableTextBoxStyle(NormalEditableTextBox)
-		.SetFont(ParameterFont);
-
-	Style->Set("NiagaraEditor.ParameterEditableTextBox", ParameterEditableTextBox);
-
-	Style->Set("NiagaraEditor.ParameterInlineEditableText", FInlineEditableTextBlockStyle()
-		.SetTextStyle(ParameterText)
-		.SetEditableTextBoxStyle(ParameterEditableTextBox));
+	FEditableTextBoxStyle ParameterEditableText = FEditableTextBoxStyle(NormalEditableTextBox)
+		.SetFont(NormalFont);
 
-	FSpinBoxStyle ParameterSpinBox = FSpinBoxStyle(NormalSpinBox)
-		.SetTextPadding(1);
+	FInlineEditableTextBlockStyle ParameterEditableTextBox = FInlineEditableTextBlockStyle()
+		.SetEditableTextBoxStyle(ParameterEditableText)
+		.SetTextStyle(ParameterText);
+	Style->Set("NiagaraEditor.ParameterInlineEditableText", ParameterEditableTextBox);
 
-	Style->Set("NiagaraEditor.ParameterSpinbox", ParameterSpinBox);
 
 	Style->Set("NiagaraEditor.ParameterName.NamespaceBorder", new BOX_PLUGIN_BRUSH("Icons/NamespaceBorder", FMargin(4.0f / 16.0f)));
 
@@ -205,16 +193,10 @@ TSharedRef< FSlateStyleSet > FNiagaraEditorStyle::Create()
 		.SetFont(DEFAULT_FONT("Regular", 8))
 		.SetColorAndOpacity(FLinearColor(1.0f, 1.0f, 1.0f, 0.5f)));
 
-	Style->Set("NiagaraEditor.Stack.HighlightedButtonBrush", new BOX_CORE_BRUSH("Common/ButtonHoverHint", FMargin(4 / 16.0f), FStyleColors::Primary));
-
 	// Parameter Map View
 	Style->Set("NiagaraEditor.Stack.DepressedHighlightedButtonBrush", new BOX_CORE_BRUSH("Common/ButtonHoverHint", FMargin(4 / 16.0f), FStyleColors::PrimaryPress));
-	Style->Set("NiagaraEditor.Stack.ViewOptionsShadowColor", FLinearColor::Black);
 	Style->Set("NiagaraEditor.Stack.FlatButtonColor", FLinearColor(FColor(205, 205, 205)));
 
-	const FVector2D ViewOptionsShadowOffset = FVector2D(0, 1);
-	Style->Set("NiagaraEditor.Stack.ViewOptionsShadowOffset", ViewOptionsShadowOffset);
-
 	// Code View
 	{
 		Style->Set("NiagaraEditor.CodeView.Checkbox.Text", FTextBlockStyle(NormalText)
@@ -240,10 +222,6 @@ TSharedRef< FSlateStyleSet > FNiagaraEditorStyle::Create()
 	Style->Set("NiagaraEditor.SelectedEmitter.UnsupportedSelectionText", SelectedEmitterUnsupportedSelectionText);
 
 	// Toolbar Icons
-	Style->Set("NiagaraEditor.Apply", new IMAGE_BRUSH("Icons/icon_Niagara_Apply_40x", Icon40x40));
-	Style->Set("NiagaraEditor.Apply.Small", new IMAGE_BRUSH("Icons/icon_Niagara_Apply_40x", Icon20x20));
-	Style->Set("NiagaraEditor.Compile", new IMAGE_BRUSH("Icons/icon_compile_40x", Icon40x40));
-	Style->Set("NiagaraEditor.Compile.Small", new IMAGE_BRUSH("Icons/icon_compile_40x", Icon20x20));
 	Style->Set("NiagaraEditor.AddEmitter", new IMAGE_BRUSH("Icons/icon_AddObject_40x", Icon40x40));
 	Style->Set("NiagaraEditor.AddEmitter.Small", new IMAGE_BRUSH("Icons/icon_AddObject_40x", Icon20x20));
 	Style->Set("NiagaraEditor.UnlockToChanges", new IMAGE_BRUSH("Icons/icon_levels_unlocked_40x", Icon40x40));
@@ -253,14 +231,6 @@ TSharedRef< FSlateStyleSet > FNiagaraEditorStyle::Create()
 	Style->Set("NiagaraEditor.SimulationOptions", new IMAGE_PLUGIN_BRUSH("Icons/Commands/icon_simulationOptions_40x", Icon40x40));
 	Style->Set("NiagaraEditor.SimulationOptions.Small", new IMAGE_PLUGIN_BRUSH("Icons/Commands/icon_simulationOptions_40x", Icon20x20));
 
-	Style->Set("Niagara.CompileStatus.Unknown", new IMAGE_BRUSH("Icons/CompileStatus_Working", Icon40x40));
-	Style->Set("Niagara.CompileStatus.Unknown.Small", new IMAGE_BRUSH("Icons/CompileStatus_Working", Icon20x20));
-	Style->Set("Niagara.CompileStatus.Error",   new IMAGE_BRUSH("Icons/CompileStatus_Fail", Icon40x40));
-	Style->Set("Niagara.CompileStatus.Error.Small", new IMAGE_BRUSH("Icons/CompileStatus_Fail", Icon20x20));
-	Style->Set("Niagara.CompileStatus.Good",    new IMAGE_BRUSH("Icons/CompileStatus_Good", Icon40x40));
-	Style->Set("Niagara.CompileStatus.Good.Small", new IMAGE_BRUSH("Icons/CompileStatus_Good", Icon20x20));
-	Style->Set("Niagara.CompileStatus.Warning", new IMAGE_BRUSH("Icons/CompileStatus_Warning", Icon40x40));
-	Style->Set("Niagara.CompileStatus.Warning.Small", new IMAGE_BRUSH("Icons/CompileStatus_Warning", Icon20x20));
 	Style->Set("Niagara.Asset.ReimportAsset.Needed", new IMAGE_BRUSH("Icons/icon_Reimport_Needed_40x", Icon40x40));
 	Style->Set("Niagara.Asset.ReimportAsset.Default", new IMAGE_BRUSH("Icons/icon_Reimport_40x", Icon40x40));
 	
diff --git a/Engine/Plugins/FX/Niagara/Source/NiagaraEditor/Private/Toolkits/NiagaraScriptToolkit.cpp b/Engine/Plugins/FX/Niagara/Source/NiagaraEditor/Private/Toolkits/NiagaraScriptToolkit.cpp
--- a/Engine/Plugins/FX/Niagara/Source/NiagaraEditor/Private/Toolkits/NiagaraScriptToolkit.cpp
+++ b/Engine/Plugins/FX/Niagara/Source/NiagaraEditor/Private/Toolkits/NiagaraScriptToolkit.cpp
@@ -509,19 +509,23 @@ void FNiagaraScriptToolkit::ExtendToolbar()
 FSlateIcon FNiagaraScriptToolkit::GetCompileStatusImage() const
 {
 	ENiagaraScriptCompileStatus Status = ScriptViewModel->GetLatestCompileStatus();
-
+	static const FName CompileStatusBackground("AssetEditor.CompileStatus.Background");
+	static const FName CompileStatusUnknown("AssetEditor.CompileStatus.Overlay.Unknown");
+	static const FName CompileStatusError("AssetEditor.CompileStatus.Overlay.Error");
+	static const FName CompileStatusGood("AssetEditor.CompileStatus.Overlay.Good");
+	static const FName CompileStatusWarning("AssetEditor.CompileStatus.Overlay.Warning");
 	switch (Status)
 	{
 	default:
 	case ENiagaraScriptCompileStatus::NCS_Unknown:
 	case ENiagaraScriptCompileStatus::NCS_Dirty:
-		return FSlateIcon(FNiagaraEditorStyle::GetStyleSetName(), "Niagara.CompileStatus.Unknown");
+		return FSlateIcon(FAppStyle::Get().GetStyleSetName(), CompileStatusBackground, NAME_None, CompileStatusUnknown);
 	case ENiagaraScriptCompileStatus::NCS_Error:
-		return FSlateIcon(FNiagaraEditorStyle::GetStyleSetName(), "Niagara.CompileStatus.Error");
+		return FSlateIcon(FAppStyle::Get().GetStyleSetName(), CompileStatusBackground, NAME_None, CompileStatusError);
 	case ENiagaraScriptCompileStatus::NCS_UpToDate:
-		return FSlateIcon(FNiagaraEditorStyle::GetStyleSetName(), "Niagara.CompileStatus.Good");
+		return FSlateIcon(FAppStyle::Get().GetStyleSetName(), CompileStatusBackground, NAME_None, CompileStatusGood);
 	case ENiagaraScriptCompileStatus::NCS_UpToDateWithWarnings:
-		return FSlateIcon(FNiagaraEditorStyle::GetStyleSetName(), "Niagara.CompileStatus.Warning");
+		return FSlateIcon(FAppStyle::Get().GetStyleSetName(), CompileStatusBackground, NAME_None, CompileStatusWarning);
 	}
 }
 void FNiagaraScriptToolkit::Tick(float DeltaTime)
diff --git a/Engine/Plugins/FX/Niagara/Source/NiagaraEditor/Private/Toolkits/NiagaraSystemToolkit.cpp b/Engine/Plugins/FX/Niagara/Source/NiagaraEditor/Private/Toolkits/NiagaraSystemToolkit.cpp
--- a/Engine/Plugins/FX/Niagara/Source/NiagaraEditor/Private/Toolkits/NiagaraSystemToolkit.cpp
+++ b/Engine/Plugins/FX/Niagara/Source/NiagaraEditor/Private/Toolkits/NiagaraSystemToolkit.cpp
@@ -1006,19 +1006,18 @@ void FNiagaraSystemToolkit::ExtendToolbar()
 
 		static void FillToolbar(FToolBarBuilder& ToolbarBuilder, FNiagaraSystemToolkit* Toolkit)
 		{
-			if (Toolkit->Emitter != nullptr)
+			ToolbarBuilder.BeginSection("Compile");
 			{
-				ToolbarBuilder.BeginSection("Apply");
+				if (Toolkit->Emitter != nullptr)
 				{
-					ToolbarBuilder.AddToolBarButton(FNiagaraEditorCommands::Get().Apply,
-						NAME_None, TAttribute<FText>(), TAttribute<FText>(),
-						FSlateIcon(FNiagaraEditorStyle::GetStyleSetName(), "NiagaraEditor.Apply"),
-						FName(TEXT("ApplyNiagaraEmitter")));
+					{
+						ToolbarBuilder.AddToolBarButton(FNiagaraEditorCommands::Get().Apply,
+							NAME_None, TAttribute<FText>(), TAttribute<FText>(),
+							FSlateIcon(FAppStyle::Get().GetStyleSetName(), "AssetEditor.Apply"),
+							FName(TEXT("ApplyNiagaraEmitter")));
+					}
 				}
-				ToolbarBuilder.EndSection();
-			}
-			ToolbarBuilder.BeginSection("Compile");
-			{
+		
 				ToolbarBuilder.AddToolBarButton(FNiagaraEditorCommands::Get().Compile,
 					NAME_None,
 					TAttribute<FText>(),
@@ -1035,27 +1034,23 @@ void FNiagaraSystemToolkit::ExtendToolbar()
 			}
 			ToolbarBuilder.EndSection();
 
-			ToolbarBuilder.BeginSection("NiagaraThumbnail");
+			ToolbarBuilder.BeginSection("NiagaraTools");
 			{
 				ToolbarBuilder.AddToolBarButton(FNiagaraEditorCommands::Get().SaveThumbnailImage, NAME_None,
 					LOCTEXT("GenerateThumbnail", "Thumbnail"),
 					LOCTEXT("GenerateThumbnailTooltip","Generate a thumbnail image."),
-					FSlateIcon(FEditorStyle::GetStyleSetName(), "Cascade.SaveThumbnailImage"));
-			}
-			ToolbarBuilder.EndSection();
-
-			ToolbarBuilder.BeginSection("NiagaraPreviewOptions");
-			{
+					FSlateIcon(FAppStyle::Get().GetStyleSetName(), "AssetEditor.SaveThumbnail"));
+		
 				ToolbarBuilder.AddToolBarButton(FNiagaraEditorCommands::Get().ToggleBounds, NAME_None,
 					LOCTEXT("ShowBounds", "Bounds"),
 					LOCTEXT("ShowBoundsTooltip", "Show the bounds for the scene."),
-					FSlateIcon(FEditorStyle::GetStyleSetName(), "Cascade.ToggleBounds"));
+					FSlateIcon(FAppStyle::Get().GetStyleSetName(), "AssetEditor.ToggleShowBounds"));
 				ToolbarBuilder.AddComboButton(
 					FUIAction(),
 					FOnGetContent::CreateRaw(Toolkit, &FNiagaraSystemToolkit::GenerateBoundsMenuContent, Toolkit->GetToolkitCommands()),
 					LOCTEXT("BoundsMenuCombo_Label", "Bounds Options"),
 					LOCTEXT("BoundsMenuCombo_ToolTip", "Bounds options"),
-					FSlateIcon(FEditorStyle::GetStyleSetName(), "Cascade.ToggleBounds"),
+					FSlateIcon(FAppStyle::Get().GetStyleSetName(), "AssetEditor.ToggleShowBounds"),
 					true
 				);
 			}
@@ -1065,15 +1060,15 @@ void FNiagaraSystemToolkit::ExtendToolbar()
 			ToolbarBuilder.BeginSection("NiagaraStatisticsOptions");
 			{
 				ToolbarBuilder.AddToolBarButton(FNiagaraEditorCommands::Get().ToggleStatPerformance, NAME_None,
-                    LOCTEXT("NiagaraShowPerformance", "Performance"),
-                    LOCTEXT("NiagaraShowPerformanceTooltip", "Show runtime performance for particle scripts."),
-                    FSlateIcon(FEditorStyle::GetStyleSetName(), "MaterialEditor.ToggleMaterialStats"));
+					LOCTEXT("NiagaraShowPerformance", "Performance"),
+					LOCTEXT("NiagaraShowPerformanceTooltip", "Show runtime performance for particle scripts."),
+					FSlateIcon(FAppStyle::Get().GetStyleSetName(), "AssetEditor.ToggleStats"));
 				ToolbarBuilder.AddComboButton(
                     FUIAction(),
                     FOnGetContent::CreateRaw(Toolkit, &FNiagaraSystemToolkit::GenerateStatConfigMenuContent, Toolkit->GetToolkitCommands()),
                     FText(),
                     LOCTEXT("NiagaraShowPerformanceCombo_ToolTip", "Runtime performance options"),
-                    FSlateIcon(FEditorStyle::GetStyleSetName(), "MaterialEditor.ToggleMaterialStats"),
+					FSlateIcon(FAppStyle::Get().GetStyleSetName(), "AssetEditor.ToggleStats"),
                     true);
 			}
 			ToolbarBuilder.EndSection();
@@ -1086,7 +1081,7 @@ void FNiagaraSystemToolkit::ExtendToolbar()
 					FOnGetContent::CreateStatic(Local::FillSimulationOptionsMenu, Toolkit),
 					LOCTEXT("SimulationOptions", "Simulation"),
 					LOCTEXT("SimulationOptionsTooltip", "Simulation options"),
-					FSlateIcon(FNiagaraEditorStyle::GetStyleSetName(), "NiagaraEditor.SimulationOptions")
+					FSlateIcon(FAppStyle::Get().GetStyleSetName(), "AssetEditor.Simulate")
 				);
 			}
 			ToolbarBuilder.EndSection();
@@ -1303,9 +1298,14 @@ TSharedRef<SWidget> FNiagaraSystemToolkit::GenerateCompileMenuContent()
 	FUIAction FullRebuildAction(
 		FExecuteAction::CreateRaw(this, &FNiagaraSystemToolkit::CompileSystem, true));
 
+
+	static const FName CompileStatusBackground("AssetEditor.CompileStatus.Background");
+	static const FName CompileStatusUnknown("AssetEditor.CompileStatus.Overlay.Unknown");
+
+
 	MenuBuilder.AddMenuEntry(LOCTEXT("FullRebuild", "Full Rebuild"),
 		LOCTEXT("FullRebuildTooltip", "Triggers a full rebuild of this system, ignoring the change tracking."),
-		FSlateIcon(FNiagaraEditorStyle::GetStyleSetName(), "Niagara.CompileStatus.Unknown"),
+		FSlateIcon(FAppStyle::Get().GetStyleSetName(), CompileStatusBackground, NAME_None, CompileStatusUnknown),
 		FullRebuildAction, NAME_None, EUserInterfaceActionType::Button);
 	MenuBuilder.AddMenuEntry(LOCTEXT("AutoCompile", "Automatically compile when graph changes"),
 		FText(), FSlateIcon(), Action, NAME_None, EUserInterfaceActionType::ToggleButton);
@@ -1317,18 +1317,24 @@ FSlateIcon FNiagaraSystemToolkit::GetCompileStatusImage() const
 {
 	ENiagaraScriptCompileStatus Status = SystemViewModel->GetLatestCompileStatus();
 
+	static const FName CompileStatusBackground("AssetEditor.CompileStatus.Background");
+	static const FName CompileStatusUnknown("AssetEditor.CompileStatus.Overlay.Unknown");
+	static const FName CompileStatusError("AssetEditor.CompileStatus.Overlay.Error");
+	static const FName CompileStatusGood("AssetEditor.CompileStatus.Overlay.Good");
+	static const FName CompileStatusWarning("AssetEditor.CompileStatus.Overlay.Warning");
+
 	switch (Status)
 	{
 	default:
 	case ENiagaraScriptCompileStatus::NCS_Unknown:
 	case ENiagaraScriptCompileStatus::NCS_Dirty:
-		return FSlateIcon(FNiagaraEditorStyle::GetStyleSetName(), "Niagara.CompileStatus.Unknown");
+		return FSlateIcon(FAppStyle::Get().GetStyleSetName(), CompileStatusBackground, NAME_None, CompileStatusUnknown);
 	case ENiagaraScriptCompileStatus::NCS_Error:
-		return FSlateIcon(FNiagaraEditorStyle::GetStyleSetName(), "Niagara.CompileStatus.Error");
+		return FSlateIcon(FAppStyle::Get().GetStyleSetName(), CompileStatusBackground, NAME_None, CompileStatusError);
 	case ENiagaraScriptCompileStatus::NCS_UpToDate:
-		return FSlateIcon(FNiagaraEditorStyle::GetStyleSetName(), "Niagara.CompileStatus.Good");
+		return FSlateIcon(FAppStyle::Get().GetStyleSetName(), CompileStatusBackground, NAME_None, CompileStatusGood);
 	case ENiagaraScriptCompileStatus::NCS_UpToDateWithWarnings:
-		return FSlateIcon(FNiagaraEditorStyle::GetStyleSetName(), "Niagara.CompileStatus.Warning");
+		return FSlateIcon(FAppStyle::Get().GetStyleSetName(), CompileStatusBackground, NAME_None, CompileStatusWarning);
 	}
 }
 
diff --git a/Engine/Plugins/FX/Niagara/Source/NiagaraEditor/Private/TypeEditorUtilities/NiagaraFloatTypeEditorUtilities.cpp b/Engine/Plugins/FX/Niagara/Source/NiagaraEditor/Private/TypeEditorUtilities/NiagaraFloatTypeEditorUtilities.cpp
--- a/Engine/Plugins/FX/Niagara/Source/NiagaraEditor/Private/TypeEditorUtilities/NiagaraFloatTypeEditorUtilities.cpp
+++ b/Engine/Plugins/FX/Niagara/Source/NiagaraEditor/Private/TypeEditorUtilities/NiagaraFloatTypeEditorUtilities.cpp
@@ -22,8 +22,7 @@ public:
 		ChildSlot
 		[
 			SNew(SSpinBox<float>)
-			.Style(FNiagaraEditorStyle::Get(), "NiagaraEditor.ParameterSpinBox")
-			.Font(FNiagaraEditorStyle::Get().GetFontStyle("NiagaraEditor.ParameterFont"))
+			.Font(FAppStyle::Get().GetFontStyle("PropertyWindow.NormalFont"))
 			.MinValue(TOptional<float>())
 			.MaxValue(TOptional<float>())
 			.MaxSliderValue(TOptional<float>())
diff --git a/Engine/Plugins/FX/Niagara/Source/NiagaraEditor/Private/TypeEditorUtilities/NiagaraIntegerTypeEditorUtilities.cpp b/Engine/Plugins/FX/Niagara/Source/NiagaraEditor/Private/TypeEditorUtilities/NiagaraIntegerTypeEditorUtilities.cpp
--- a/Engine/Plugins/FX/Niagara/Source/NiagaraEditor/Private/TypeEditorUtilities/NiagaraIntegerTypeEditorUtilities.cpp
+++ b/Engine/Plugins/FX/Niagara/Source/NiagaraEditor/Private/TypeEditorUtilities/NiagaraIntegerTypeEditorUtilities.cpp
@@ -23,8 +23,7 @@ public:
 		ChildSlot
 		[
 			SNew(SSpinBox<int32>)
-			.Style(FNiagaraEditorStyle::Get(), "NiagaraEditor.ParameterSpinBox")
-			.Font(FNiagaraEditorStyle::Get().GetFontStyle("NiagaraEditor.ParameterFont"))
+			.Font(FAppStyle::Get().GetFontStyle("PropertyWindow.NormalFont"))
 			.MinValue(TOptional<int32>())
 			.MaxValue(TOptional<int32>())
 			.MaxSliderValue(TOptional<int32>())
diff --git a/Engine/Plugins/FX/Niagara/Source/NiagaraEditor/Private/TypeEditorUtilities/NiagaraVectorTypeEditorUtilities.cpp b/Engine/Plugins/FX/Niagara/Source/NiagaraEditor/Private/TypeEditorUtilities/NiagaraVectorTypeEditorUtilities.cpp
--- a/Engine/Plugins/FX/Niagara/Source/NiagaraEditor/Private/TypeEditorUtilities/NiagaraVectorTypeEditorUtilities.cpp
+++ b/Engine/Plugins/FX/Niagara/Source/NiagaraEditor/Private/TypeEditorUtilities/NiagaraVectorTypeEditorUtilities.cpp
@@ -53,8 +53,7 @@ private:
 	TSharedRef<SWidget> ConstructComponentWidget(int32 Index)
 	{
 		return SNew(SNumericEntryBox<float>)
-		.Font(FNiagaraEditorStyle::Get().GetFontStyle("NiagaraEditor.ParameterFont"))
-		.OverrideTextMargin(2)
+		.Font(FAppStyle::Get().GetFontStyle("PropertyWindow.NormalFont"))
 		.MinValue(TOptional<float>())
 		.MaxValue(TOptional<float>())
 		.MaxSliderValue(TOptional<float>())
diff --git a/Engine/Plugins/FX/Niagara/Source/NiagaraEditor/Private/Widgets/SNiagaraParameterMapPaletteItem.cpp b/Engine/Plugins/FX/Niagara/Source/NiagaraEditor/Private/Widgets/SNiagaraParameterMapPaletteItem.cpp
--- a/Engine/Plugins/FX/Niagara/Source/NiagaraEditor/Private/Widgets/SNiagaraParameterMapPaletteItem.cpp
+++ b/Engine/Plugins/FX/Niagara/Source/NiagaraEditor/Private/Widgets/SNiagaraParameterMapPaletteItem.cpp
@@ -58,30 +58,39 @@ void SNiagaraParameterMapPalleteItem::Construct(const FArguments& InArgs, FCreat
 	bool bForceReadOnly = NamespaceMetadata.IsValid() == false || NamespaceMetadata.Options.Contains(ENiagaraNamespaceMetadataOptions::PreventEditingName);
 
 	ParameterNameTextBlock = SNew(SNiagaraParameterNameTextBlock)
+		.EditableTextStyle(FNiagaraEditorStyle::Get(), "NiagaraEditor.ParameterInlineEditableText")
+		.ReadOnlyTextStyle(&FAppStyle::Get().GetWidgetStyle<FTextBlockStyle>("NormalText"))
 		.ParameterText(FText::FromName(ParameterAction->Parameter.GetName()))
 		.HighlightText(InCreateData->HighlightText)
 		.OnTextCommitted(this, &SNiagaraParameterMapPalleteItem::OnNameTextCommitted)
 		.OnVerifyTextChanged(this, &SNiagaraParameterMapPalleteItem::OnNameTextVerifyChanged)
 		.IsSelected(InCreateData->IsRowSelectedDelegate)
 		.IsReadOnly(InCreateData->bIsReadOnly || bForceReadOnly || ParameterAction->bIsExternallyReferenced)
+		.DecoratorHAlign(HAlign_Fill)
 		.Decorator()
 		[
 			SNew(SHorizontalBox)
 			+ SHorizontalBox::Slot()
+			.FillWidth(1.0)
+			[
+				SNew(SSpacer)
+			]
+			+ SHorizontalBox::Slot()
 			.AutoWidth()
-			.Padding(0, 0, 5, 0)
+			.Padding(0, 0, 2, 0)
 			[
-				SNew(STextBlock)
+				SNew(SImage)
 				.Visibility(ParameterAction->bIsExternallyReferenced ? EVisibility::Visible : EVisibility::Collapsed)
-				.Font(FEditorStyle::Get().GetFontStyle("FontAwesome.8"))
-				.Text(FEditorFontGlyphs::Lock)
+				.Image(FAppStyle::Get().GetBrush("Icons.Lock"))
 				.ToolTipText(LOCTEXT("LockedToolTip", "This parameter is used in a referenced external graph and can't be edited directly."))
+				.ColorAndOpacity(FSlateColor::UseForeground())
 			]
 			+ SHorizontalBox::Slot()
 			.AutoWidth()
 			.Padding(0, 0, 5, 0)
 			[
 				SNew(STextBlock)
+				.ColorAndOpacity(FSlateColor::UseForeground())
 				.Visibility(ParameterAction->bIsSourcedFromCustomStackContext ? EVisibility::Visible : EVisibility::Collapsed)
 				.Font(FEditorStyle::Get().GetFontStyle("FontAwesome.8"))
 				.Text(FEditorFontGlyphs::Database)
@@ -115,18 +124,16 @@ void SNiagaraParameterMapPalleteItem::Construct(const FArguments& InArgs, FCreat
 		.VAlign(VAlign_Center)
 		.Padding(3, 0)
 		[
-			SNew(SComboButton)
-			.HasDownArrow(false)
-			.ButtonStyle(FEditorStyle::Get(), "RoundButton")
-			.ForegroundColor(FSlateColor::UseForeground())
-			.ContentPadding(FMargin(2.0f))
+			SNew(SBox)
 			.HAlign(HAlign_Right)
 			.VAlign(VAlign_Center)
-			.ButtonContent()
+			.MinDesiredWidth(30.0f)
+			.Content()
 			[
 				SNew(STextBlock)
 				.Text(this, &SNiagaraParameterMapPalleteItem::GetReferenceCount)
-				.Font(Font)
+				.TextStyle(&FAppStyle::Get().GetWidgetStyle<FTextBlockStyle>("HintText"))
+				.ColorAndOpacity(FSlateColor::UseForeground())
 			]
 		]
 
diff --git a/Engine/Plugins/FX/Niagara/Source/NiagaraEditor/Private/Widgets/SNiagaraParameterMapView.cpp b/Engine/Plugins/FX/Niagara/Source/NiagaraEditor/Private/Widgets/SNiagaraParameterMapView.cpp
--- a/Engine/Plugins/FX/Niagara/Source/NiagaraEditor/Private/Widgets/SNiagaraParameterMapView.cpp
+++ b/Engine/Plugins/FX/Niagara/Source/NiagaraEditor/Private/Widgets/SNiagaraParameterMapView.cpp
@@ -230,7 +230,6 @@ void SNiagaraParameterMapView::Construct(const FArguments& InArgs, const TArray<
 	ToolkitType = InToolkitType;
 	ToolkitCommands = InToolkitCommands;
 	AddParameterButtons.SetNum(NiagaraParameterMapSectionID::Num);
-	const FVector2D ViewOptionsShadowOffset = FNiagaraEditorStyle::Get().GetVector("NiagaraEditor.Stack.ViewOptionsShadowOffset");
 	ParametersWithNamespaceModifierRenamePending = MakeShared<TArray<FName>>();
 
 	SelectedScriptObjects = InSelectedObjects[0];
@@ -269,30 +268,16 @@ void SNiagaraParameterMapView::Construct(const FArguments& InArgs, const TArray<
 		[
 			SNew(SComboButton)
 			.ContentPadding(0)
-			.ForegroundColor(FSlateColor::UseForeground())
-			.ButtonStyle(FEditorStyle::Get(), "ToggleButton")
+			.ComboButtonStyle(FAppStyle::Get(), "SimpleComboButton") // Use the tool bar item style for this button
+			.HasDownArrow(false)
 			.AddMetaData<FTagMetaData>(FTagMetaData(TEXT("ViewOptions")))
 			.ToolTipText(LOCTEXT("ViewOptionsToolTip", "View Options"))
 			.OnGetMenuContent(this, &SNiagaraParameterMapView::GetViewOptionsMenu)
 			.ButtonContent()
 			[
-				SNew(SOverlay)
-				// drop shadow
-				+ SOverlay::Slot()
-				.VAlign(VAlign_Top)
-				.Padding(ViewOptionsShadowOffset.X, ViewOptionsShadowOffset.Y, 0, 0)
-				[
-					SNew(SImage)
-					.Image(FEditorStyle::GetBrush("GenericViewButton"))
-					.ColorAndOpacity(FNiagaraEditorStyle::Get().GetColor("NiagaraEditor.Stack.ViewOptionsShadowColor"))
-				]
-				+ SOverlay::Slot()
-				.VAlign(VAlign_Top)
-				[
-					SNew(SImage)
-					.Image(FEditorStyle::GetBrush("GenericViewButton"))
-					.ColorAndOpacity(FNiagaraEditorStyle::Get().GetColor("NiagaraEditor.Stack.FlatButtonColor"))
-				]
+				SNew(SImage)
+				.Image(FAppStyle::Get().GetBrush("Icons.Filter"))
+				.ColorAndOpacity(FSlateColor::UseForeground())
 			]
 		];
 
@@ -1013,8 +998,7 @@ TSharedRef<SWidget> SNiagaraParameterMapView::CreateAddToSectionButton(const Nia
 {
 	TSharedPtr<SComboButton> Button;
 	SAssignNew(Button, SComboButton)
-	.ButtonStyle(FEditorStyle::Get(), "RoundButton")
-	.ForegroundColor(FEditorStyle::GetSlateColor("DefaultForeground"))
+	.ButtonStyle(FAppStyle::Get(), "SimpleButton")
 	.ContentPadding(FMargin(2, 0))
 	.OnGetMenuContent(this, &SNiagaraParameterMapView::OnGetParameterMenu, InSection)
 	.IsEnabled(this, &SNiagaraParameterMapView::ParameterAddEnabled)
@@ -1024,27 +1008,9 @@ TSharedRef<SWidget> SNiagaraParameterMapView::CreateAddToSectionButton(const Nia
 	.AddMetaData<FTagMetaData>(FTagMetaData(MetaDataTag))
 	.ButtonContent()
 	[
-		SNew(SHorizontalBox)
-
-		+ SHorizontalBox::Slot()
-		.AutoWidth()
-		.Padding(FMargin(0, 1))
-		[
-			SNew(SImage)
-			.Image(FEditorStyle::GetBrush("Plus"))
-		]
-
-		+ SHorizontalBox::Slot()
-		.VAlign(VAlign_Center)
-		.AutoWidth()
-		.Padding(FMargin(2,0,0,0))
-		[
-			SNew(STextBlock)
-			.Font(IDetailLayoutBuilder::GetDetailFontBold())
-			.Text(AddNewText)
-			.Visibility(this, &SNiagaraParameterMapView::OnAddButtonTextVisibility, WeakRowWidget, InSection)
-			.ShadowOffset(FVector2D(1,1))
-		]
+		SNew(SImage)
+		.Image(FAppStyle::Get().GetBrush("Icons.PlusCircle"))
+		.ColorAndOpacity(FSlateColor::UseForeground())
 	];
 	AddParameterButtons[InSection] = Button;
 
diff --git a/Engine/Plugins/FX/Niagara/Source/NiagaraEditor/Public/Widgets/SNiagaraParameterName.h b/Engine/Plugins/FX/Niagara/Source/NiagaraEditor/Public/Widgets/SNiagaraParameterName.h
--- a/Engine/Plugins/FX/Niagara/Source/NiagaraEditor/Public/Widgets/SNiagaraParameterName.h
+++ b/Engine/Plugins/FX/Niagara/Source/NiagaraEditor/Public/Widgets/SNiagaraParameterName.h
@@ -8,6 +8,7 @@
 #include "Styling/SlateTypes.h"
 #include "Framework/SlateDelegates.h"
 #include "Styling/SlateWidgetStyleAsset.h"
+#include "NiagaraEditorStyle.h"
 
 class SInlineEditableTextBlock;
 class SBorder;
@@ -27,8 +28,8 @@ public:
 	};
 public:
 	SLATE_BEGIN_ARGS(SNiagaraParameterName) 
-		: _EditableTextStyle(&FEditorStyle::Get().GetWidgetStyle<FInlineEditableTextBlockStyle>("InlineEditableTextBlockStyle"))
-		, _ReadOnlyTextStyle(&FEditorStyle::Get().GetWidgetStyle<FTextBlockStyle>("NormalText"))
+		: _EditableTextStyle(&FNiagaraEditorStyle::Get().GetWidgetStyle<FInlineEditableTextBlockStyle>("NiagaraEditor.ParameterInlineEditableText"))
+		, _ReadOnlyTextStyle(&FNiagaraEditorStyle::Get().GetWidgetStyle<FTextBlockStyle>("NiagaraEditor.ParameterText"))
 		, _IsReadOnly(false)
 		, _SingleNameDisplayMode(ESingleNameDisplayMode::Name)
 		, _DecoratorHAlign(HAlign_Left)
@@ -104,8 +105,8 @@ class NIAGARAEDITOR_API SNiagaraParameterNameTextBlock : public SCompoundWidget
 {
 public:
 	SLATE_BEGIN_ARGS(SNiagaraParameterNameTextBlock)
-		: _EditableTextStyle(&FEditorStyle::Get().GetWidgetStyle<FInlineEditableTextBlockStyle>("InlineEditableTextBlockStyle"))
-		, _ReadOnlyTextStyle(&FEditorStyle::Get().GetWidgetStyle<FTextBlockStyle>("NormalText"))
+		: _EditableTextStyle(&FAppStyle::Get().GetWidgetStyle<FInlineEditableTextBlockStyle>("InlineEditableTextBlockStyle"))
+		, _ReadOnlyTextStyle(&FAppStyle::Get().GetWidgetStyle<FTextBlockStyle>("NormalText"))
 		, _IsReadOnly(false)
 		, _DecoratorHAlign(HAlign_Left)
 		, _DecoratorPadding(5.0f, 0.0f, 0.0f, 0.0f)
diff --git a/Engine/Plugins/FX/Niagara/Source/NiagaraEditorWidgets/Private/NiagaraEditorWidgetsStyle.cpp b/Engine/Plugins/FX/Niagara/Source/NiagaraEditorWidgets/Private/NiagaraEditorWidgetsStyle.cpp
--- a/Engine/Plugins/FX/Niagara/Source/NiagaraEditorWidgets/Private/NiagaraEditorWidgetsStyle.cpp
+++ b/Engine/Plugins/FX/Niagara/Source/NiagaraEditorWidgets/Private/NiagaraEditorWidgetsStyle.cpp
@@ -3,13 +3,14 @@
 #include "NiagaraEditorWidgetsStyle.h"
 
 #include "Framework/Application/SlateApplication.h"
-#include "EditorStyleSet.h"
 #include "Slate/SlateGameResources.h"
 #include "Styling/SlateStyleRegistry.h"
 #include "Styling/SlateTypes.h"
 #include "Styling/CoreStyle.h"
 #include "Interfaces/IPluginManager.h"
 #include "Classes/EditorStyleSettings.h"
+#include "Styling/AppStyle.h"
+#include "Styling/StyleColors.h"
 
 TSharedPtr< FSlateStyleSet > FNiagaraEditorWidgetsStyle::NiagaraEditorWidgetsStyleInstance = NULL;
  
@@ -57,9 +58,9 @@ const FVector2D Icon40x40(40.0f, 40.0f);
 
 TSharedRef< FSlateStyleSet > FNiagaraEditorWidgetsStyle::Create()
 {
-	const FTextBlockStyle NormalText = FEditorStyle::GetWidgetStyle<FTextBlockStyle>("NormalText");
-	const FEditableTextBoxStyle NormalEditableTextBox = FCoreStyle::Get().GetWidgetStyle<FEditableTextBoxStyle>("NormalEditableTextBox");
-	const FSpinBoxStyle NormalSpinBox = FEditorStyle::GetWidgetStyle<FSpinBoxStyle>("SpinBox");
+	const FTextBlockStyle NormalText = FAppStyle::Get().GetWidgetStyle<FTextBlockStyle>("NormalText");
+	const FEditableTextBoxStyle NormalEditableTextBox = FAppStyle::Get().GetWidgetStyle<FEditableTextBoxStyle>("NormalEditableTextBox");
+	const FSpinBoxStyle NormalSpinBox = FAppStyle::Get().GetWidgetStyle<FSpinBoxStyle>("SpinBox");
 
 	TSharedRef< FSlateStyleSet > Style = MakeShareable(new FSlateStyleSet("NiagaraEditorWidgetsStyle"));
 	Style->SetContentRoot(FPaths::EngineContentDir() / TEXT("Editor/Slate/Niagara"));
@@ -67,30 +68,43 @@ TSharedRef< FSlateStyleSet > FNiagaraEditorWidgetsStyle::Create()
 	// Stack
 	Style->Set("NiagaraEditor.Stack.IconSize", FVector2D(18.0f, 18.0f));
 
-	FSlateFontInfo StackGroupFont = DEFAULT_FONT("Bold", 10);
-	FTextBlockStyle StackGroupText = FTextBlockStyle(NormalText)
-		.SetFont(StackGroupFont)
-		.SetColorAndOpacity(FLinearColor(1.0f, 1.0f, 1.0f, 1.0f));
+	const FTextBlockStyle CategoryText = FAppStyle::Get().GetWidgetStyle<FTextBlockStyle>("DetailsView.CategoryTextStyle");
+	FSlateFontInfo CategoryFont = FAppStyle::Get().GetFontStyle(TEXT("PropertyWindow.BoldFont"));
+	FTextBlockStyle StackGroupText = FTextBlockStyle(CategoryText)
+		.SetFont(CategoryFont)
+		.SetTransformPolicy(ETextTransformPolicy::ToUpper);
 
 	Style->Set("NiagaraEditor.Stack.GroupText", StackGroupText);
-	
+
 	FEditableTextBoxStyle StackEditableGroupText = FEditableTextBoxStyle(NormalEditableTextBox)
-		.SetFont(StackGroupFont);
+		.SetFont(CategoryFont);
 
 	FInlineEditableTextBlockStyle StackInlineEditableGroupText = FInlineEditableTextBlockStyle()
 		.SetEditableTextBoxStyle(StackEditableGroupText)
-		.SetTextStyle(StackGroupText); 
+		.SetTextStyle(StackGroupText);
 
 	Style->Set("NiagaraEditor.Stack.EditableGroupText", StackInlineEditableGroupText);
 
-	FSlateFontInfo StackDefaultFont = DEFAULT_FONT("Regular", 10);
+	FSlateFontInfo NormalFont = FAppStyle::Get().GetFontStyle(TEXT("PropertyWindow.NormalFont"));
+	FTextBlockStyle StackItemText = FTextBlockStyle(CategoryText)
+		.SetFont(NormalFont)
+		.SetTransformPolicy(ETextTransformPolicy::ToUpper);
+
+	Style->Set("NiagaraEditor.Stack.ItemText", StackItemText);
+	
+	FEditableTextBoxStyle StackEditableItemText = FEditableTextBoxStyle(NormalEditableTextBox)
+		.SetFont(NormalFont);
+
+	FInlineEditableTextBlockStyle StackInlineEditableItemText = FInlineEditableTextBlockStyle()
+		.SetEditableTextBoxStyle(StackEditableItemText)
+		.SetTextStyle(StackItemText);
+	Style->Set("NiagaraEditor.Stack.EditableItemText", StackInlineEditableItemText);
+
 	FTextBlockStyle StackDefaultText = FTextBlockStyle(NormalText)
-		.SetFont(StackDefaultFont);
+		.SetFont(NormalFont);
 	Style->Set("NiagaraEditor.Stack.DefaultText", StackDefaultText);
 
-	FSlateFontInfo StackCategoryFont = DEFAULT_FONT("Bold", 10);
-	FTextBlockStyle StackCategoryText = FTextBlockStyle(NormalText)
-		.SetFont(StackCategoryFont);
+	FTextBlockStyle StackCategoryText = FTextBlockStyle(NormalText);
 	Style->Set("NiagaraEditor.Stack.CategoryText", StackCategoryText);
 	Style->Set("NiagaraEditor.SystemOverview.GroupHeaderText", StackCategoryText);
 
@@ -104,11 +118,6 @@ TSharedRef< FSlateStyleSet > FNiagaraEditorWidgetsStyle::Create()
 		.SetFont(TextContentFont);
 	Style->Set("NiagaraEditor.Stack.TextContentText", TextContentText);
 
-	FSlateFontInfo StackItemFont = DEFAULT_FONT("Regular", 11);
-	FTextBlockStyle StackItemText = FTextBlockStyle(NormalText)
-		.SetFont(StackItemFont);
-	Style->Set("NiagaraEditor.Stack.ItemText", StackItemText);
-
 	FSlateFontInfo PerfWidgetDetailFont = DEFAULT_FONT("Regular", 7);
 	Style->Set("NiagaraEditor.Stack.Stats.DetailFont", PerfWidgetDetailFont);
 	FSlateFontInfo PerfWidgetGroupFont = DEFAULT_FONT("Regular", 8);
@@ -116,13 +125,7 @@ TSharedRef< FSlateStyleSet > FNiagaraEditorWidgetsStyle::Create()
 	FSlateFontInfo PerfWidgetEvalTypeFont = DEFAULT_FONT("Regular", 7);
 	Style->Set("NiagaraEditor.Stack.Stats.EvalTypeFont", PerfWidgetEvalTypeFont);
 	
-	FEditableTextBoxStyle StackEditableItemText = FEditableTextBoxStyle(NormalEditableTextBox)
-		.SetFont(StackItemFont);
 
-	FInlineEditableTextBlockStyle StackInlineEditableItemText = FInlineEditableTextBlockStyle()
-		.SetTextStyle(StackItemText)
-		.SetEditableTextBoxStyle(StackEditableItemText);
-	Style->Set("NiagaraEditor.Stack.EditableItemText", StackInlineEditableItemText);
 
 	FSlateFontInfo StackSubduedItemFont = DEFAULT_FONT("Regular", 9);
 	FTextBlockStyle StackSubduedItemText = FTextBlockStyle(NormalText)
@@ -151,6 +154,7 @@ TSharedRef< FSlateStyleSet > FNiagaraEditorWidgetsStyle::Create()
 	Style->Set("NiagaraEditor.SystemOverview.NodeBackgroundBorder", new BOX_PLUGIN_BRUSH("Icons/SystemOverviewNodeBackground", FMargin(1.0f / 4.0f)));
 	Style->Set("NiagaraEditor.SystemOverview.NodeBackgroundColor", FLinearColor(FColor(48, 48, 48)));
 
+	const FTableRowStyle& NormalTableRowStyle = FAppStyle::Get().GetWidgetStyle<FTableRowStyle>("TableView.Row");
 	FSlateFontInfo CurveOverviewTopLevelFont = DEFAULT_FONT("Bold", 10);
 	FTextBlockStyle CurveOverviewTopLevelText = FTextBlockStyle(NormalText)
 		.SetFont(CurveOverviewTopLevelFont);
@@ -197,36 +201,22 @@ TSharedRef< FSlateStyleSet > FNiagaraEditorWidgetsStyle::Create()
 
 	Style->Set("NiagaraEditor.CurveDetails.ShowInOverview.Small", new IMAGE_CORE_BRUSH("Common/GoToSource", Icon12x12, FLinearColor(.9f, .9f, .9f, 1.0f)));
 
-	const FTableRowStyle& NormalTableRowStyle = FEditorStyle::Get().GetWidgetStyle<FTableRowStyle>("TableView.Row");
-
-	FSlateBrush StackRowSelectionBrush = BOX_PLUGIN_BRUSH("Icons/StackSelectionBorder", FMargin(2.0f / 8.0f), GetDefault<UEditorStyleSettings>()->SelectionColor);
-	FSlateBrush StackRowSubduedSelectionBrush = BOX_PLUGIN_BRUSH("Icons/StackSelectionBorder", FMargin(2.0f / 8.0f), GetDefault<UEditorStyleSettings>()->GetSubduedSelectionColor());
+	FSlateBrush StackRowSelectionBrush = BOX_PLUGIN_BRUSH("Icons/StackSelectionBorder", FMargin(2.0f / 8.0f), FStyleColors::Select);
+	FSlateBrush StackRowSubduedSelectionBrush = BOX_PLUGIN_BRUSH("Icons/StackSelectionBorder", FMargin(2.0f / 8.0f), FStyleColors::SelectInactive);
 	Style->Set("NiagaraEditor.Stack.TableViewRow", FTableRowStyle(NormalTableRowStyle)
+		.SetEvenRowBackgroundBrush(FSlateColorBrush(FStyleColors::White))
+		.SetOddRowBackgroundBrush(FSlateColorBrush(FStyleColors::White))
 		.SetActiveBrush(StackRowSelectionBrush)
 		.SetActiveHoveredBrush(StackRowSelectionBrush)
 		.SetInactiveBrush(StackRowSubduedSelectionBrush)
 		.SetInactiveHoveredBrush(StackRowSelectionBrush));
 
 	Style->Set("NiagaraEditor.SystemOverview.TableViewRow", FTableRowStyle(NormalTableRowStyle)
+		.SetEvenRowBackgroundBrush(FSlateColorBrush(FStyleColors::Panel))
+		.SetOddRowBackgroundBrush(FSlateColorBrush(FStyleColors::Panel))
+		.SetSelectedTextColor(FStyleColors::Background)
 		.SetInactiveBrush(IMAGE_CORE_BRUSH("Common/Selection", Icon8x8, GetDefault<UEditorStyleSettings>()->GetSubduedSelectionColor())));
 
-	Style->Set("NiagaraEditor.Stack.BackgroundColor", FLinearColor(FColor(96, 96, 96)));
-	Style->Set("NiagaraEditor.Stack.Item.HeaderBackgroundColor", FLinearColor(FColor(48, 48, 48)));
-	Style->Set("NiagaraEditor.Stack.Item.ContentBackgroundColor", FLinearColor(FColor(62, 62, 62)));
-	Style->Set("NiagaraEditor.Stack.Item.ContentAdvancedBackgroundColor", FLinearColor(FColor(53, 53, 53)));
-	Style->Set("NiagaraEditor.Stack.Item.FooterBackgroundColor", FLinearColor(FColor(75, 75, 75)));
-	Style->Set("NiagaraEditor.Stack.Item.InfoBackgroundColor", FLinearColor(FColor(68, 100, 106)));
-	Style->Set("NiagaraEditor.Stack.Item.WarningBackgroundColor", FLinearColor(FColor(97, 97, 68)));
-	Style->Set("NiagaraEditor.Stack.Item.ErrorBackgroundColor", FLinearColor(FColor(126, 78, 68)));
-
-	Style->Set("NiagaraEditor.Stack.UnknownColor", FLinearColor(1, 0, 1));
-
-	Style->Set("NiagaraEditor.Stack.ItemHeaderFooter.BackgroundBrush", new FSlateColorBrush(FLinearColor(FColor(20, 20, 20))));
-
-	Style->Set("NiagaraEditor.Stack.ForegroundColor", FLinearColor(FColor(220, 220, 220)));
-	Style->Set("NiagaraEditor.Stack.GroupForegroundColor", FLinearColor(FColor(220, 220, 220)));
-	Style->Set("NiagaraEditor.Stack.FlatButtonColor", FLinearColor(FColor(205, 205, 205)));
-	Style->Set("NiagaraEditor.Stack.DividerColor", FLinearColor(FColor(92, 92, 92)));
 	
 	Style->Set("NiagaraEditor.Stack.Stats.EvalTypeColor", FLinearColor(FColor(168, 168, 168)));
 	Style->Set("NiagaraEditor.Stack.Stats.RuntimePlaceholderColor", FLinearColor(FColor(86, 86, 86)));
@@ -240,21 +230,6 @@ TSharedRef< FSlateStyleSet > FNiagaraEditorWidgetsStyle::Create()
 	Style->Set("NiagaraEditor.Stack.Stats.HighCostColor", FLinearColor(FColor(205, 114, 69)));
 	Style->Set("NiagaraEditor.Stack.Stats.MaxCostColor", FLinearColor(FColor(200, 60, 60)));
 
-	Style->Set("NiagaraEditor.Stack.HighlightedButtonBrush", new BOX_CORE_BRUSH("Common/ButtonHoverHint", FMargin(4 / 16.0f), GetDefault<UEditorStyleSettings>()->SelectionColor));
-
-	const FVector2D ViewOptionsShadowOffset = FVector2D(0, 1);
-	Style->Set("NiagaraEditor.Stack.ViewOptionsShadowOffset", ViewOptionsShadowOffset);
-
-	FComboButtonStyle ViewOptionsComboButtonStyle = FCoreStyle::Get().GetWidgetStyle<FComboButtonStyle>("ComboButton");
-
-	const FLinearColor ViewOptionsShadowColor = FLinearColor::Black;
-	Style->Set("NiagaraEditor.Stack.ViewOptionsShadowColor", FLinearColor::Black);
-	Style->Set("NiagaraEditor.Stack.ViewOptionsButton",	ViewOptionsComboButtonStyle
-		.SetButtonStyle(FEditorStyle::Get().GetWidgetStyle<FButtonStyle>("HoverHintOnly"))
-		.SetShadowOffset(ViewOptionsShadowOffset)
-		.SetShadowColorAndOpacity(ViewOptionsShadowColor)
-	);
-
 	Style->Set("NiagaraEditor.Stack.AccentColor.System", FLinearColor(FColor(67, 105, 124)));
 	Style->Set("NiagaraEditor.Stack.AccentColor.Emitter", FLinearColor(FColor(126, 87, 67)));
 	Style->Set("NiagaraEditor.Stack.AccentColor.Particle", FLinearColor(FColor(87, 107, 61)));
@@ -302,18 +277,6 @@ TSharedRef< FSlateStyleSet > FNiagaraEditorWidgetsStyle::Create()
 	Style->Set("NiagaraEditor.Stack.SearchHighlightColor", FEditorStyle::GetColor("TextBlock.HighlighColor"));
 	Style->Set("NiagaraEditor.Stack.SearchResult", new BOX_PLUGIN_BRUSH("Icons/SearchResultBorder", FMargin(1.f/8.f)));
 
-	Style->Set("NiagaraEditor.Stack.AddButton", FButtonStyle()
-		.SetNormal(BOX_CORE_BRUSH("Common/FlatButton", 2.0f / 8.0f, FLinearColor(0.5f, 0.5f, 0.5f, 1.0f)))
-		.SetHovered(BOX_CORE_BRUSH("Common/FlatButton", 2.0f / 8.0f, FLinearColor(1.0f, 1.0f, 1.0f, 1.0f)))
-		.SetPressed(BOX_CORE_BRUSH("Common/FlatButton", 2.0f / 8.0f, FLinearColor(.8f, .8f, .8f, 1.0f)))
-	);
-
-	FTextBlockStyle AddButtonText = FTextBlockStyle(NormalText)
-		.SetColorAndOpacity(FLinearColor(0.8f, 0.8f, 0.8f, 1.0f))
-		.SetShadowOffset(FVector2D(0, 1))
-		.SetShadowColorAndOpacity(FLinearColor(0, 0, 0, 0.9f));
-	Style->Set("NiagaraEditor.Stack.AddButtonText", AddButtonText);
-
 	Style->Set("NiagaraEditor.Stack.ModuleHighlight", new IMAGE_PLUGIN_BRUSH("Icons/ModuleHighlight", Icon6x6, FLinearColor::White));
 	Style->Set("NiagaraEditor.Stack.ModuleHighlightMore", new IMAGE_PLUGIN_BRUSH("Icons/ModuleHighlightMore", Icon6x6, FLinearColor::White));
 	Style->Set("NiagaraEditor.Stack.ModuleHighlightLarge", new IMAGE_PLUGIN_BRUSH("Icons/ModuleHighlightLarge", Icon8x8, FLinearColor::White));
diff --git a/Engine/Plugins/FX/Niagara/Source/NiagaraEditorWidgets/Private/SNiagaraOverviewStack.cpp b/Engine/Plugins/FX/Niagara/Source/NiagaraEditorWidgets/Private/SNiagaraOverviewStack.cpp
--- a/Engine/Plugins/FX/Niagara/Source/NiagaraEditorWidgets/Private/SNiagaraOverviewStack.cpp
+++ b/Engine/Plugins/FX/Niagara/Source/NiagaraEditorWidgets/Private/SNiagaraOverviewStack.cpp
@@ -36,6 +36,7 @@
 #include "Framework/Application/SlateApplication.h"
 #include "Framework/MultiBox/MultiBoxBuilder.h"
 #include "Subsystems/AssetEditorSubsystem.h"
+#include "Styling/StyleColors.h"
 
 #define LOCTEXT_NAMESPACE "NiagaraOverviewStack"
 
@@ -61,17 +62,17 @@ class SNiagaraSystemOverviewEntryListRow : public STableRow<UNiagaraStackEntry*>
 		FMargin ContentPadding;
 		if (StackEntry->IsA<UNiagaraStackItem>())
 		{
-			BackgroundColor = FNiagaraEditorWidgetsStyle::Get().GetColor("NiagaraEditor.SystemOverview.Item.BackgroundColor");
+			BackgroundColor = FStyleColors::Panel;
 			ContentPadding = FMargin(0, 1, 0, 1);
 		}
 		else
 		{
-			BackgroundColor = FNiagaraEditorWidgetsStyle::Get().GetColor("NiagaraEditor.SystemOverview.Group.BackgroundColor");
+			BackgroundColor = FStyleColors::Header;
 			ContentPadding = FMargin(0, 2, 0, 2);
 		}
 
-		DisabledBackgroundColor = BackgroundColor + FLinearColor(.02f, .02f, .02f, 0.0f);
-		IsolatedBackgroundColor = BackgroundColor - FLinearColor(.03f, .03f, .03f, 0.0f);
+		DisabledBackgroundColor = FStyleColors::Recessed;
+		IsolatedBackgroundColor = FStyleColors::Select;
 
 		TSharedRef<FNiagaraSystemViewModel> SystemViewModel = StackEntry->GetSystemViewModel();
 		TSharedPtr<FNiagaraEmitterViewModel> EmitterViewModel = StackEntry->GetEmitterViewModel();
@@ -87,17 +88,15 @@ class SNiagaraSystemOverviewEntryListRow : public STableRow<UNiagaraStackEntry*>
 			.OnCanAcceptDrop(InArgs._OnCanAcceptDrop)
 			.OnAcceptDrop(InArgs._OnAcceptDrop)
 		[
-			SNew(SBox)
-			.Padding(FMargin(0, 1, 0, 1))
-			[
-				SNew(SBorder)
+			SNew(SBorder)
 				.BorderImage(FEditorStyle::GetBrush("WhiteBrush"))
-				.BorderBackgroundColor(this, &SNiagaraSystemOverviewEntryListRow::GetBackgroundColor)
+				.BorderBackgroundColor(FStyleColors::Recessed)
 				.ToolTipText_UObject(StackEntry, &UNiagaraStackEntry::GetTooltipText)
-				.Padding(FMargin(0))
+				.Padding(FMargin(0, 1, 0, 1))
 				[
 					SNew(SBorder)
-					.BorderImage(this, &SNiagaraSystemOverviewEntryListRow::GetBorder)
+					.BorderImage(FEditorStyle::GetBrush("WhiteBrush"))
+					.BorderBackgroundColor(TAttribute<FSlateColor>(this, &SNiagaraSystemOverviewEntryListRow::GetBackgroundColor))
 					.Padding(ContentPadding)
 					[
 						SNew(SHorizontalBox)
@@ -116,7 +115,6 @@ class SNiagaraSystemOverviewEntryListRow : public STableRow<UNiagaraStackEntry*>
 						]
 					]
 				]
-			]
 		],
 		InOwnerTableView);
 	}
@@ -204,6 +202,10 @@ class SNiagaraSystemOverviewEntryListRow : public STableRow<UNiagaraStackEntry*>
 private:
 	FSlateColor GetBackgroundColor() const
 	{
+		if (IsSelected())
+		{
+			return FStyleColors::Select;
+		}
 		TSharedRef<FNiagaraSystemViewModel> SystemViewModel = StackEntry->GetSystemViewModel();
 		if (SystemViewModel->GetSystem().GetIsolateEnabled())
 		{
@@ -233,9 +235,9 @@ private:
 	UNiagaraStackViewModel* StackViewModel;
 	UNiagaraStackEntry* StackEntry;
 	TSharedPtr<FNiagaraStackCommandContext> StackCommandContext;
-	FLinearColor BackgroundColor;
-	FLinearColor DisabledBackgroundColor;
-	FLinearColor IsolatedBackgroundColor;
+	FSlateColor BackgroundColor;
+	FSlateColor DisabledBackgroundColor;
+	FSlateColor IsolatedBackgroundColor;
 	TAttribute<EVisibility> IssueIconVisibility;
 	TWeakPtr<FNiagaraEmitterHandleViewModel> EmitterHandleViewModel;
 };
@@ -257,33 +259,26 @@ public:
 
 		ChildSlot
 		[
-			SNew(SButton)
-			.ButtonStyle(FEditorStyle::Get(), "HoverHintOnly")
-			.ForegroundColor(FSlateColor::UseSubduedForeground())
-			.OnClicked(this, &SNiagaraSystemOverviewEnabledCheckBox::OnButtonClicked)
+			SNew(SCheckBox)
+			.OnCheckStateChanged(this, &SNiagaraSystemOverviewEnabledCheckBox::OnButtonClicked)
+			.IsChecked(this, &SNiagaraSystemOverviewEnabledCheckBox::IsEnabled)
 			.ToolTipText(LOCTEXT("EnableCheckBoxToolTip", "Enable or disable this item."))
-			.ContentPadding(FMargin(3, 2, 2, 2))
-			.IsFocusable(false)
-			[
-				SNew(STextBlock)
-				.Font(FEditorStyle::Get().GetFontStyle("FontAwesome.10"))
-				.Text(this, &SNiagaraSystemOverviewEnabledCheckBox::GetButtonText)
-			]
 		];
 	}
 
 private:
-	FReply OnButtonClicked()
+	void OnButtonClicked(ECheckBoxState InNewState)
 	{
 		OnCheckedChanged.ExecuteIfBound(IsChecked.IsBound() && !IsChecked.Get());
-		return FReply::Handled();
 	}
 
-	FText GetButtonText() const
+	ECheckBoxState IsEnabled() const
 	{
-		return IsChecked.IsBound() && IsChecked.Get()
-			? LOCTEXT("CheckedText", "\xf14a")
-			: LOCTEXT("UncheckedText", "\xf0c8");
+		if (IsChecked.IsBound())
+		{
+			return IsChecked.Get() ? ECheckBoxState::Checked : ECheckBoxState::Unchecked;
+		}
+		return ECheckBoxState::Unchecked;
 	}
 
 private:
@@ -379,8 +374,9 @@ private:
 				.Padding(0.0f, 0.0f, 6.0f, 0.0f)
 				[
 					SNew(STextBlock)
-					.TextStyle(FNiagaraEditorWidgetsStyle::Get(), "NiagaraEditor.SystemOverview.ItemText")
+					.TextStyle(FNiagaraEditorWidgetsStyle::Get(), "NiagaraEditor.Stack.ItemText")
 					.Text(LOCTEXT("SetVariablesPrefix", "Set:"))
+					.ColorAndOpacity(FSlateColor::UseForeground())
 				]
 				+ SHorizontalBox::Slot()
 				[
@@ -396,6 +392,7 @@ private:
 				.Text(this, &SNiagaraSystemOverviewItemName::GetItemDisplayName)
 				.ToolTipText(this, &SNiagaraSystemOverviewItemName::GetItemToolTip)
 				.IsEnabled_UObject(StackItem.Get(), &UNiagaraStackEntry::GetIsEnabledAndOwnerIsEnabled)
+				.ColorAndOpacity(FSlateColor::UseForeground())
 			];
 			UpdateTextStyle();
 		}
@@ -411,7 +408,7 @@ private:
 			}
 			else
 			{
-				NameTextBlock->SetTextStyle(&FNiagaraEditorWidgetsStyle::Get().GetWidgetStyle<FTextBlockStyle>("NiagaraEditor.SystemOverview.ItemText"));
+				NameTextBlock->SetTextStyle(&FNiagaraEditorWidgetsStyle::Get().GetWidgetStyle<FTextBlockStyle>("NiagaraEditor.Stack.ItemText"));
 			}
 		}
 	}
@@ -744,7 +741,7 @@ TSharedRef<ITableRow> SNiagaraOverviewStack::OnGenerateRowForEntry(UNiagaraStack
 			+ SHorizontalBox::Slot()
 			.VAlign(VAlign_Center)
 			.AutoWidth()
-			.Padding(3, 0, 0, 0)
+			.Padding(0)
 			[
 				SNew(SNiagaraSystemOverviewEnabledCheckBox)
 				.Visibility(this, &SNiagaraOverviewStack::GetEnabledCheckBoxVisibility, StackItem)
@@ -781,9 +778,10 @@ TSharedRef<ITableRow> SNiagaraOverviewStack::OnGenerateRowForEntry(UNiagaraStack
 			.Padding(3, 2, 0, 2)
 			[
 				SNew(STextBlock)
-				.TextStyle(FNiagaraEditorWidgetsStyle::Get(), "NiagaraEditor.SystemOverview.GroupHeaderText")
+				.TextStyle(FNiagaraEditorWidgetsStyle::Get(), "NiagaraEditor.Stack.GroupText")
 				.Text_UObject(Item, &UNiagaraStackEntry::GetDisplayName)
 				.IsEnabled_UObject(Item, &UNiagaraStackEntry::GetIsEnabledAndOwnerIsEnabled)
+				.ColorAndOpacity(FSlateColor::UseForeground())
 			]
 			+ SHorizontalBox::Slot()
             .AutoWidth()
diff --git a/Engine/Plugins/FX/Niagara/Source/NiagaraEditorWidgets/Private/SNiagaraOverviewStackNode.cpp b/Engine/Plugins/FX/Niagara/Source/NiagaraEditorWidgets/Private/SNiagaraOverviewStackNode.cpp
--- a/Engine/Plugins/FX/Niagara/Source/NiagaraEditorWidgets/Private/SNiagaraOverviewStackNode.cpp
+++ b/Engine/Plugins/FX/Niagara/Source/NiagaraEditorWidgets/Private/SNiagaraOverviewStackNode.cpp
@@ -28,6 +28,7 @@
 #include "ViewModels/Stack/NiagaraStackRendererItem.h"
 #include "Materials/MaterialInterface.h"
 #include "Materials/Material.h"
+#include "Styling/StyleColors.h"
 
 #define LOCTEXT_NAMESPACE "NiagaraOverviewStackNode"
 
@@ -131,7 +132,6 @@ TSharedRef<SWidget> SNiagaraOverviewStackNode::CreateTitleRightWidget()
 			.IsFocusable(false)
 			.ToolTipText(LOCTEXT("OpenAndFocusParentEmitterToolTip", "Open and Focus Parent Emitter"))
 			.ButtonStyle(FEditorStyle::Get(), "HoverHintOnly")
-			.ForegroundColor(FNiagaraEditorWidgetsStyle::Get().GetColor("NiagaraEditor.Stack.ForegroundColor"))
 			.ContentPadding(2)
 			.OnClicked(this, &SNiagaraOverviewStackNode::OpenParentEmitter)
 			.Visibility(this, &SNiagaraOverviewStackNode::GetOpenParentEmitterVisibility)
@@ -140,7 +140,7 @@ TSharedRef<SWidget> SNiagaraOverviewStackNode::CreateTitleRightWidget()
 			[
 				SNew(SImage)
 				.Image(FNiagaraEditorWidgetsStyle::Get().GetBrush("NiagaraEditor.Stack.GoToSourceIcon"))
-				.ColorAndOpacity(FNiagaraEditorWidgetsStyle::Get().GetColor("NiagaraEditor.Stack.FlatButtonColor"))
+				.ColorAndOpacity(FSlateColor::UseForeground())
 			]
 		]
 		+ SHorizontalBox::Slot()
@@ -333,8 +333,8 @@ TSharedRef<SWidget> SNiagaraOverviewStackNode::CreateNodeContentArea()
 			+SHorizontalBox::Slot()
 			[
 				SNew(SBorder)
-				.BorderImage(FNiagaraEditorWidgetsStyle::Get().GetBrush("NiagaraEditor.SystemOverview.NodeBackgroundBorder"))
-				.BorderBackgroundColor(FNiagaraEditorWidgetsStyle::Get().GetColor("NiagaraEditor.SystemOverview.NodeBackgroundColor"))
+ 				.BorderImage(FNiagaraEditorWidgetsStyle::Get().GetBrush("NiagaraEditor.SystemOverview.NodeBackgroundBorder"))
+ 				.BorderBackgroundColor(FStyleColors::Panel)
 				.HAlign(HAlign_Fill)
 				.VAlign(VAlign_Fill)
 				.Padding(FMargin(0, 0, 0, 4))
diff --git a/Engine/Plugins/FX/Niagara/Source/NiagaraEditorWidgets/Private/SNiagaraScratchPad.cpp b/Engine/Plugins/FX/Niagara/Source/NiagaraEditorWidgets/Private/SNiagaraScratchPad.cpp
--- a/Engine/Plugins/FX/Niagara/Source/NiagaraEditorWidgets/Private/SNiagaraScratchPad.cpp
+++ b/Engine/Plugins/FX/Niagara/Source/NiagaraEditorWidgets/Private/SNiagaraScratchPad.cpp
@@ -418,21 +418,23 @@ private:
 		.Padding(3, 0, 0, 0)
 		[
 			SNew(STextBlock)
-			.TextStyle(FEditorStyle::Get(), "DetailsView.CategoryTextStyle")
+			.TextStyle(FNiagaraEditorWidgetsStyle::Get(), "NiagaraEditor.Stack.GroupText")
 			.Text(ViewModel->GetDisplayNameForUsage(Category))
+			.ColorAndOpacity(FSlateColor::UseForeground())
 		]
 		+ SHorizontalBox::Slot()
 		.AutoWidth()
-		.Padding(0.0f, 4.0f, 3.0f, 4.0f)
+		.Padding(0.0f, 2.0f, 3.0f, 2.0f)
 		[
 			SNew(SButton)
-			.ButtonStyle(FEditorStyle::Get(), "RoundButton")
+			.ButtonStyle(FAppStyle::Get(), "SimpleButton")
 			.OnClicked(this, &SNiagaraScratchPadScriptSelector::ScriptSelectorAddButtonClicked, Category)
 			.ContentPadding(FMargin(3.0f, 2.0f, 2.0f, 2.0f))
 			.Content()
 			[
 				SNew(SImage)
-				.Image(FEditorStyle::GetBrush("Plus"))
+				.Image(FAppStyle::Get().GetBrush("Icons.PlusCircle"))
+				.ColorAndOpacity(FSlateColor::UseForeground())
 			]
 		];
 	}
@@ -910,8 +912,9 @@ class SNiagaraScratchPadSectionBox : public SCompoundWidget
 			.Padding(1.0f, 5.0f, 0.0f, 4.0f)
 			[
 				SNew(STextBlock)
-				.TextStyle(FEditorStyle::Get(), "Docking.TabFont")
+				.TextStyle(FNiagaraEditorWidgetsStyle::Get(), "NiagaraEditor.Stack.GroupText")
 				.Text(InArgs._HeaderText)
+				.ColorAndOpacity(FSlateColor::UseForeground())
 			]
 			+ SVerticalBox::Slot()
 			.Padding(5.0f, 0.0f, 3.0f, 5.0f)
@@ -1159,4 +1162,4 @@ EVisibility SNiagaraScratchPad::GetObjectSelectionNoSelectionTextVisibility() co
 	return ViewModel->GetObjectSelection()->GetSelectedObjects().Num() == 0 ? EVisibility::Visible : EVisibility::Collapsed;
 }
 
-#undef LOCTEXT_NAMESPACE
\ No newline at end of file
+#undef LOCTEXT_NAMESPACE
diff --git a/Engine/Plugins/FX/Niagara/Source/NiagaraEditorWidgets/Private/SNiagaraStack.cpp b/Engine/Plugins/FX/Niagara/Source/NiagaraEditorWidgets/Private/SNiagaraStack.cpp
--- a/Engine/Plugins/FX/Niagara/Source/NiagaraEditorWidgets/Private/SNiagaraStack.cpp
+++ b/Engine/Plugins/FX/Niagara/Source/NiagaraEditorWidgets/Private/SNiagaraStack.cpp
@@ -59,6 +59,7 @@
 #include "Framework/Commands/UICommandList.h"
 #include "ViewModels/NiagaraOverviewGraphViewModel.h"
 #include "Widgets/SNiagaraParameterName.h"
+#include "Styling/StyleColors.h"
 
 #define LOCTEXT_NAMESPACE "NiagaraStack"
 
@@ -80,102 +81,109 @@ public:
 		ParentStackPtr = InArgs._ParentStack;
 		ChildSlot
 		[
-			SNew(SVerticalBox)
-
-			//~ Enable check box, view source emitter button, and external header controls.
-			+ SVerticalBox::Slot()
-			.AutoHeight()
-			.HAlign(HAlign_Fill)
+			SNew(SBorder)
+			.BorderImage(FEditorStyle::GetBrush("WhiteBrush"))
+			.BorderBackgroundColor(FStyleColors::Panel)
 			[
-				SNew(SHorizontalBox)
-				//~ Enabled
-				+ SHorizontalBox::Slot()
-				.AutoWidth()
-				.Padding(2, 4, 0, 0)
-				.VAlign(VAlign_Top)
-				[
-					SNew(SCheckBox)
-					.ToolTipText(LOCTEXT("EnabledToolTip", "Toggles whether this emitter is enabled. Disabled emitters don't simulate or render."))
-					.IsChecked(EmitterHandleViewModel.ToSharedRef(), &FNiagaraEmitterHandleViewModel::GetIsEnabledCheckState)
-					.OnCheckStateChanged(EmitterHandleViewModel.ToSharedRef(), &FNiagaraEmitterHandleViewModel::OnIsEnabledCheckStateChanged)
-				]
-				+ SHorizontalBox::Slot()
-				.Padding(2)
+				SNew(SVerticalBox)
+
+				//~ Enable check box, view source emitter button, and external header controls.
+				+ SVerticalBox::Slot()
+				.AutoHeight()
 				.HAlign(HAlign_Fill)
-				.VAlign(VAlign_Center)
-				.FillWidth(1.0f)
 				[
-					// Name and Source Emitter Name
-					SNew(SWrapBox)
-					.Clipping(EWidgetClipping::ClipToBoundsAlways) 
-					.UseAllottedSize(true)
-					+ SWrapBox::Slot()
+					SNew(SHorizontalBox)
+					//~ Enabled
+					+ SHorizontalBox::Slot()
+					.AutoWidth()
+					.Padding(3)
+					.VAlign(VAlign_Center)
+					[
+						SNew(SCheckBox)
+						.ToolTipText(LOCTEXT("EnabledToolTip", "Toggles whether this emitter is enabled. Disabled emitters don't simulate or render."))
+						.IsChecked(EmitterHandleViewModel.ToSharedRef(), &FNiagaraEmitterHandleViewModel::GetIsEnabledCheckState)
+						.OnCheckStateChanged(EmitterHandleViewModel.ToSharedRef(), &FNiagaraEmitterHandleViewModel::OnIsEnabledCheckStateChanged)
+					]
+					+ SHorizontalBox::Slot()
+					.Padding(2)
+					.HAlign(HAlign_Fill)
+					.VAlign(VAlign_Center)
+					.FillWidth(1.0f)
 					[
+						// Name and Source Emitter Name
+						SNew(SWrapBox)
+						.Clipping(EWidgetClipping::ClipToBoundsAlways) 
+						.UseAllottedSize(true)
+						+ SWrapBox::Slot()
+						[
 				
-						SAssignNew(EmitterNameTextBlock, SInlineEditableTextBlock)
-						.ToolTipText(this, &SNiagaraStackEmitterHeader::GetEmitterNameToolTip)
-						.Style(FNiagaraEditorStyle::Get(), "NiagaraEditor.HeadingInlineEditableText") 
-						.Clipping(EWidgetClipping::ClipToBoundsAlways)
-						.Text(EmitterHandleViewModel.ToSharedRef(), &FNiagaraEmitterHandleViewModel::GetNameText)
-						.OnTextCommitted(EmitterHandleViewModel.ToSharedRef(), &FNiagaraEmitterHandleViewModel::OnNameTextComitted)
-						.OnVerifyTextChanged(EmitterHandleViewModel.ToSharedRef(), &FNiagaraEmitterHandleViewModel::VerifyNameTextChanged)
-						.IsReadOnly(EmitterHandleViewModel->CanRenameEmitter() == false)
+							SAssignNew(EmitterNameTextBlock, SInlineEditableTextBlock)
+							.ToolTipText(this, &SNiagaraStackEmitterHeader::GetEmitterNameToolTip)
+							.Style(FAppStyle::Get(), "DetailsView.NameTextBlockStyle")
+							.Clipping(EWidgetClipping::ClipToBoundsAlways)
+							.Text(EmitterHandleViewModel.ToSharedRef(), &FNiagaraEmitterHandleViewModel::GetNameText)
+							.OnTextCommitted(EmitterHandleViewModel.ToSharedRef(), &FNiagaraEmitterHandleViewModel::OnNameTextComitted)
+							.OnVerifyTextChanged(EmitterHandleViewModel.ToSharedRef(), &FNiagaraEmitterHandleViewModel::VerifyNameTextChanged)
+							.IsReadOnly(EmitterHandleViewModel->CanRenameEmitter() == false)
+						]
+						+ SWrapBox::Slot()
+						.Padding(4, 0, 0, 0)
+						[
+							SNew(STextBlock)
+							.ToolTipText(this, &SNiagaraStackEmitterHeader::GetEmitterNameToolTip)
+							.TextStyle(FNiagaraEditorStyle::Get(), "NiagaraEditor.HeadingTextBlockSubdued")
+							.Clipping(EWidgetClipping::ClipToBoundsAlways)
+							.Text(EmitterHandleViewModel->GetEmitterViewModel(), &FNiagaraEmitterViewModel::GetParentNameText)
+							.Visibility(this, &SNiagaraStackEmitterHeader::GetSourceEmitterNameVisibility) 
+						]
 					]
-					+ SWrapBox::Slot()
-					.Padding(4, 0, 0, 0)
+					// Issue Icon
+					+ SHorizontalBox::Slot()
+					.AutoWidth()
+					.HAlign(HAlign_Right)
+					.VAlign(VAlign_Center)
+					.Padding(0, 0, 2, 0)
 					[
-						SNew(STextBlock)
-						.ToolTipText(this, &SNiagaraStackEmitterHeader::GetEmitterNameToolTip)
-						.TextStyle(FNiagaraEditorStyle::Get(), "NiagaraEditor.SubduedHeadingTextBox")
-						.Clipping(EWidgetClipping::ClipToBoundsAlways)
-						.Text(EmitterHandleViewModel->GetEmitterViewModel(), &FNiagaraEmitterViewModel::GetParentNameText)
-						.Visibility(this, &SNiagaraStackEmitterHeader::GetSourceEmitterNameVisibility) 
+						SNew(SNiagaraStackIssueIcon, StackViewModel, InRootEntry)
+						.Visibility(InArgs._IssueIconVisibility)
+						.OnClicked(this, &SNiagaraStackEmitterHeader::OnIssueIconClicked)
 					]
-				]
-				// Issue Icon
-				+ SHorizontalBox::Slot()
-				.AutoWidth()
-				.HAlign(HAlign_Right)
-				.VAlign(VAlign_Center)
-				.Padding(0, 0, 2, 0)
-				[
-					SNew(SNiagaraStackIssueIcon, StackViewModel, InRootEntry)
-					.Visibility(InArgs._IssueIconVisibility)
-					.OnClicked(this, &SNiagaraStackEmitterHeader::OnIssueIconClicked)
-				]
-				+ SHorizontalBox::Slot()
-				.AutoWidth()
-				.HAlign(HAlign_Fill)
-				.VAlign(VAlign_Center)
-				.Padding(0, 0, 2, 0)
-				[
-					SAssignNew(SettingsAnchor, SMenuAnchor)
-					.Placement(MenuPlacement_MenuLeft)
-					.OnGetMenuContent(this, &SNiagaraStackEmitterHeader::OnGetContent)
+					+ SHorizontalBox::Slot()
+					.AutoWidth()
+					.HAlign(HAlign_Fill)
+					.VAlign(VAlign_Center)
+					.Padding(0, 0, 2, 0)
 					[
-						SNew(SButton)
-						.ForegroundColor(FSlateColor::UseForeground())
-						.ButtonStyle(FEditorStyle::Get(), TEXT("FlatButton"))
-						.OnClicked(this, &SNiagaraStackEmitterHeader::OpenSubmenu)
-						.VAlign(VAlign_Center)
-						.HAlign(HAlign_Fill)
+						SAssignNew(SettingsAnchor, SMenuAnchor)
+						.Placement(MenuPlacement_MenuLeft)
+						.OnGetMenuContent(this, &SNiagaraStackEmitterHeader::OnGetContent)
 						[
-							SNew(STextBlock)
-							.Font(FEditorStyle::Get().GetFontStyle("FontAwesome.16"))
-							.Text(FEditorFontGlyphs::Cogs)
+							SNew(SButton)
+							.ContentPadding(0)
+							.ForegroundColor(FSlateColor::UseForeground())
+							.ButtonStyle(FAppStyle::Get(), "SimpleButton")
+							.OnClicked(this, &SNiagaraStackEmitterHeader::OpenSubmenu)
+							.VAlign(VAlign_Center)
+							.HAlign(HAlign_Fill)
+							[
+								SNew(SImage)
+								.ColorAndOpacity(FSlateColor::UseForeground())
+								.Image(FAppStyle::Get().GetBrush("DetailsView.ViewOptions"))
+							]
 						]
 					]
 				]
-			]
 
-			//~ Stats
-			+ SVerticalBox::Slot()
-			.AutoHeight()
-			.HAlign(HAlign_Fill)
-			.Padding(2)
-			[
-				SNew(STextBlock)
-				.Text(EmitterHandleViewModel->GetEmitterViewModel(), &FNiagaraEmitterViewModel::GetStatsText)
+				//~ Stats
+				+ SVerticalBox::Slot()
+				.AutoHeight()
+				.HAlign(HAlign_Fill)
+				.Padding(2)
+				[
+					SNew(STextBlock)
+					.Text(EmitterHandleViewModel->GetEmitterViewModel(), &FNiagaraEmitterViewModel::GetStatsText)
+					.ColorAndOpacity(FSlateColor::UseForeground())
+				]
 			]
 		];
 	}
@@ -273,7 +281,8 @@ void SNiagaraStack::Construct(const FArguments& InArgs, UNiagaraStackViewModel*
 		.Padding(1, 1, 1, 4)
 		[
 			SNew(SBorder)
-			.BorderImage(FEditorStyle::GetBrush("ToolPanel.GroupBorder"))
+			.BorderImage(FEditorStyle::GetBrush("WhiteBrush"))
+			.BorderBackgroundColor(FStyleColors::Panel)
 			.Padding(5)
 			[
 				ConstructHeaderWidget()
@@ -282,19 +291,14 @@ void SNiagaraStack::Construct(const FArguments& InArgs, UNiagaraStackViewModel*
 		+ SVerticalBox::Slot()
 		.Padding(1)
 		[
-			SNew(SBorder)
-			.BorderImage(FEditorStyle::GetBrush("WhiteBrush"))
-			.BorderBackgroundColor(FNiagaraEditorWidgetsStyle::Get().GetColor("NiagaraEditor.Stack.BackgroundColor"))
-			[
-				SAssignNew(StackTree, STreeView<UNiagaraStackEntry*>)
-				.OnGenerateRow(this, &SNiagaraStack::OnGenerateRowForStackItem)
-				.OnGetChildren(this, &SNiagaraStack::OnGetChildren)
-				.TreeItemsSource(&StackViewModel->GetRootEntryAsArray())
-				.OnTreeViewScrolled(this, &SNiagaraStack::StackTreeScrolled)
-				.OnSelectionChanged(this, &SNiagaraStack::StackTreeSelectionChanged)
-				.SelectionMode(ESelectionMode::Multi)
-				.OnItemToString_Debug_Static(&FNiagaraStackEditorWidgetsUtilities::StackEntryToStringForListDebug)
-			]
+			SAssignNew(StackTree, STreeView<UNiagaraStackEntry*>)
+			.OnGenerateRow(this, &SNiagaraStack::OnGenerateRowForStackItem)
+			.OnGetChildren(this, &SNiagaraStack::OnGetChildren)
+			.TreeItemsSource(&StackViewModel->GetRootEntryAsArray())
+			.OnTreeViewScrolled(this, &SNiagaraStack::StackTreeScrolled)
+			.OnSelectionChanged(this, &SNiagaraStack::StackTreeSelectionChanged)
+			.SelectionMode(ESelectionMode::Multi)
+			.OnItemToString_Debug_Static(&FNiagaraStackEditorWidgetsUtilities::StackEntryToStringForListDebug)
 		]
 	];
 
@@ -325,8 +329,6 @@ void SNiagaraStack::SynchronizeTreeExpansion()
 
 TSharedRef<SWidget> SNiagaraStack::ConstructHeaderWidget()
 {
-	const FVector2D ViewOptionsShadowOffset = FNiagaraEditorWidgetsStyle::Get().GetVector("NiagaraEditor.Stack.ViewOptionsShadowOffset");
-
 	return SNew(SVerticalBox)
 		//~ Top level object list view
 		+ SVerticalBox::Slot()
@@ -363,54 +365,23 @@ TSharedRef<SWidget> SNiagaraStack::ConstructHeaderWidget()
 			.AutoWidth()
 			.Padding(4, 0, 0, 0)
 			[
-				SNew(SBorder)
-				.Padding(0)
-				.BorderImage(this, &SNiagaraStack::GetViewOptionsIconBrush)
+				SNew(SComboButton)
+				.ComboButtonStyle(FAppStyle::Get(), "SimpleComboButton") // Use the tool bar item style for this button
+				.HasDownArrow(false)
+				.ToolTipText(LOCTEXT("ViewOptionsToolTip", "View Options"))
+				.OnGetMenuContent(this, &SNiagaraStack::GetViewOptionsMenu)
+				.ContentPadding(1)
+				.MenuPlacement(MenuPlacement_BelowRightAnchor)
+				.ButtonContent()
 				[
-					SNew(SComboButton)
-					.ComboButtonStyle(FNiagaraEditorWidgetsStyle::Get(), "NiagaraEditor.Stack.ViewOptionsButton")
-					.ForegroundColor(FSlateColor::UseForeground())
-					.ToolTipText(LOCTEXT("ViewOptionsToolTip", "View Options"))
-					.OnGetMenuContent(this, &SNiagaraStack::GetViewOptionsMenu)
-					.ContentPadding(1)
-					.MenuPlacement(MenuPlacement_BelowRightAnchor)
-					.ButtonContent()
-					[
-						SNew(SOverlay)
-						// drop shadow
-						+ SOverlay::Slot()
-						.VAlign(VAlign_Top)
-						.Padding(ViewOptionsShadowOffset.X, ViewOptionsShadowOffset.Y, 0, 0)
-						[
-							SNew(SImage)
-							.Image(FEditorStyle::GetBrush("GenericViewButton"))
-							.ColorAndOpacity(FNiagaraEditorWidgetsStyle::Get().GetColor("NiagaraEditor.Stack.ViewOptionsShadowColor"))
-						]
-						+ SOverlay::Slot()
-						.VAlign(VAlign_Top)
-						[
-							SNew(SImage)
-							.Image(FEditorStyle::GetBrush("GenericViewButton"))
-							.ColorAndOpacity(FNiagaraEditorWidgetsStyle::Get().GetColor("NiagaraEditor.Stack.FlatButtonColor"))
-						]
-					]
+					SNew(SImage)
+					.Image(FAppStyle::Get().GetBrush("Icons.Filter"))
+					.ColorAndOpacity(FSlateColor::UseForeground())
 				]
 			]
 		];
 }
 
-const FSlateBrush* SNiagaraStack::GetViewOptionsIconBrush() const
-{
-	bool bIsDefault = StackViewModel->GetShowAllAdvanced() == false &&
-		StackViewModel->GetShowLinkedInputs() == false &&
-		StackViewModel->GetShowOutputs() == false &&
-		StackViewModel->GetShowOnlyIssues() == false;
-
-	return bIsDefault ? 
-		FCoreStyle::Get().GetBrush("NoBrush") :
-		FNiagaraEditorWidgetsStyle::Get().GetBrush("NiagaraEditor.Stack.HighlightedButtonBrush");
-}
-
 void SNiagaraStack::OnSearchTextChanged(const FText& SearchText)
 {
 	StackViewModel->OnSearchTextChanged(SearchText);
@@ -688,23 +659,29 @@ TSharedRef<ITableRow> SNiagaraStack::OnGenerateRowForTopLevelObject(TSharedRef<U
 	TSharedPtr<SWidget> Content;
 	if (Item->SystemViewModel.IsValid())
 	{
-		Content = SNew(SHorizontalBox)
-			// System name
-			+ SHorizontalBox::Slot()
-			[
-				SNew(STextBlock)
-				.TextStyle(FNiagaraEditorStyle::Get(), "NiagaraEditor.HeadingTextBlock")
-				.Text(Item->SystemViewModel.ToSharedRef(), &FNiagaraSystemViewModel::GetDisplayName)
-			]
-			// Issue Icon
-			+ SHorizontalBox::Slot()
-			.AutoWidth()
-			.VAlign(VAlign_Center)
-			.Padding(4, 0, 2, 0)
+		Content =
+			SNew(SBorder)
+			.BorderImage(FEditorStyle::GetBrush("WhiteBrush"))
+			.BorderBackgroundColor(FStyleColors::Panel)
 			[
-				SNew(SNiagaraStackIssueIcon, StackViewModel, Item->RootEntry.Get())
-				.Visibility(this, &SNiagaraStack::GetIssueIconVisibility)
-				.OnClicked(this, &SNiagaraStack::OnCycleThroughSystemIssues, Item->SystemViewModel)
+				SNew(SHorizontalBox)
+				// System name
+				+ SHorizontalBox::Slot()
+				[
+					SNew(STextBlock)
+					.TextStyle(FAppStyle::Get(), "DetailsView.ConstantTextBlockStyle")
+					.Text(Item->SystemViewModel.ToSharedRef(), &FNiagaraSystemViewModel::GetDisplayName)
+				]
+				// Issue Icon
+				+ SHorizontalBox::Slot()
+				.AutoWidth()
+				.VAlign(VAlign_Center)
+				.Padding(4, 0, 2, 0)
+				[
+					SNew(SNiagaraStackIssueIcon, StackViewModel, Item->RootEntry.Get())
+					.Visibility(this, &SNiagaraStack::GetIssueIconVisibility)
+					.OnClicked(this, &SNiagaraStack::OnCycleThroughSystemIssues, Item->SystemViewModel)
+				]
 			];
 	}
 	else if (Item->EmitterHandleViewModel.IsValid())
@@ -778,29 +755,28 @@ TSharedRef<SNiagaraStackTableRow> SNiagaraStack::ConstructContainerForItem(UNiag
 	float RightContentPadding = 6;
 	TAttribute<FMargin> RowPadding = FMargin(0, 0, 0, 0);
 	FMargin ContentPadding(LeftContentPadding, 0, RightContentPadding, 0);
-	FLinearColor ItemBackgroundColor;
-	FLinearColor ItemForegroundColor = FNiagaraEditorWidgetsStyle::Get().GetColor("NiagaraEditor.Stack.ForegroundColor");
+	FSlateColor ItemBackgroundColor = FStyleColors::Panel;
+	FSlateColor ItemForegroundColor = FStyleColors::Foreground;
+	FSlateColor IndicatorColor = FLinearColor::Transparent;
 	bool bIsCategoryIconHighlighted;
 	bool bShowExecutionCategoryIcon;
 	switch (Item->GetStackRowStyle())
 	{
 	case UNiagaraStackEntry::EStackRowStyle::None:
-		ItemBackgroundColor = FLinearColor::Transparent;
 		bIsCategoryIconHighlighted = false;
 		bShowExecutionCategoryIcon = false;
 		break;
 	case UNiagaraStackEntry::EStackRowStyle::GroupHeader:
 		ContentPadding = FMargin(LeftContentPadding, 3, 0, 3);
-		ItemBackgroundColor = FLinearColor::Transparent;
-		ItemForegroundColor = FNiagaraEditorWidgetsStyle::Get().GetColor("NiagaraEditor.Stack.GroupForegroundColor");
+		ItemBackgroundColor = FStyleColors::Header;
 		bIsCategoryIconHighlighted = true;
 		bShowExecutionCategoryIcon = true;
 		break;
 	case UNiagaraStackEntry::EStackRowStyle::GroupFooter:
 		RowPadding = FMargin(0, 0, 0, 6);
-		ItemBackgroundColor = FLinearColor::Transparent;
 		bIsCategoryIconHighlighted = false;
 		bShowExecutionCategoryIcon = false;
+		ItemBackgroundColor = FStyleColors::Recessed;
 		break;
 	case UNiagaraStackEntry::EStackRowStyle::ItemHeader:
 		RowPadding.Bind(TAttribute<FMargin>::FGetter::CreateLambda([Item]()
@@ -815,31 +791,27 @@ TSharedRef<SNiagaraStackTableRow> SNiagaraStack::ConstructContainerForItem(UNiag
 			}
 		}));
 		ContentPadding = FMargin(LeftContentPadding, 2, 2, 2);
-		ItemBackgroundColor = FNiagaraEditorWidgetsStyle::Get().GetColor("NiagaraEditor.Stack.Item.HeaderBackgroundColor");
 		bIsCategoryIconHighlighted = false;
 		bShowExecutionCategoryIcon = true;
 		break;
 	case UNiagaraStackEntry::EStackRowStyle::ItemContent:
 		ContentPadding = FMargin(LeftContentPadding, 3, RightContentPadding, 3);
-		ItemBackgroundColor = FNiagaraEditorWidgetsStyle::Get().GetColor("NiagaraEditor.Stack.Item.ContentBackgroundColor");
 		bIsCategoryIconHighlighted = false;
 		bShowExecutionCategoryIcon = false;
 		break;
 	case UNiagaraStackEntry::EStackRowStyle::ItemContentAdvanced:
 		ContentPadding = FMargin(LeftContentPadding, 3, RightContentPadding, 3);
-		ItemBackgroundColor = FNiagaraEditorWidgetsStyle::Get().GetColor("NiagaraEditor.Stack.Item.ContentAdvancedBackgroundColor");
 		bIsCategoryIconHighlighted = false;
 		bShowExecutionCategoryIcon = false;
+		ItemBackgroundColor = FStyleColors::Recessed;
 		break;
 	case UNiagaraStackEntry::EStackRowStyle::ItemFooter:
 		RowPadding = FMargin(0, 0, 0, 2);
-		ItemBackgroundColor = FNiagaraEditorWidgetsStyle::Get().GetColor("NiagaraEditor.Stack.Item.FooterBackgroundColor");
 		bIsCategoryIconHighlighted = false;
 		bShowExecutionCategoryIcon = false;
 		break;
 	case UNiagaraStackEntry::EStackRowStyle::ItemCategory:
 		ContentPadding = FMargin(LeftContentPadding, 3, RightContentPadding, 3);
-		ItemBackgroundColor = FNiagaraEditorWidgetsStyle::Get().GetColor("NiagaraEditor.Stack.Item.ContentBackgroundColor");
 		bIsCategoryIconHighlighted = false;
 		bShowExecutionCategoryIcon = false;
 		break;
@@ -848,13 +820,16 @@ TSharedRef<SNiagaraStackTableRow> SNiagaraStack::ConstructContainerForItem(UNiag
 		switch (Item->GetIssueSeverity())
 		{
 			case EStackIssueSeverity::Error:
-				ItemBackgroundColor = FNiagaraEditorWidgetsStyle::Get().GetColor("NiagaraEditor.Stack.Item.ErrorBackgroundColor");
+				ItemBackgroundColor = FStyleColors::Header;
+				IndicatorColor = FStyleColors::Error;
 				break;
 			case EStackIssueSeverity::Warning:
-				ItemBackgroundColor = FNiagaraEditorWidgetsStyle::Get().GetColor("NiagaraEditor.Stack.Item.WarningBackgroundColor");
+				ItemBackgroundColor = FStyleColors::Header;
+				IndicatorColor = FStyleColors::Warning;
 				break;
 			case EStackIssueSeverity::Info:
-				ItemBackgroundColor = FNiagaraEditorWidgetsStyle::Get().GetColor("NiagaraEditor.Stack.Item.InfoBackgroundColor");
+				ItemBackgroundColor = FStyleColors::Header;
+				IndicatorColor = FStyleColors::Foreground;
 				break;
 			default:
 				checkf(false, TEXT("Issue severity not set for stack issue."));
@@ -865,7 +840,6 @@ TSharedRef<SNiagaraStackTableRow> SNiagaraStack::ConstructContainerForItem(UNiag
 		bShowExecutionCategoryIcon = false;
 		break;
 	default:
-		ItemBackgroundColor = FNiagaraEditorWidgetsStyle::Get().GetColor("NiagaraEditor.Stack.UnknownColor");
 		bIsCategoryIconHighlighted = false;
 		bShowExecutionCategoryIcon = false;
 		break;
@@ -874,8 +848,9 @@ TSharedRef<SNiagaraStackTableRow> SNiagaraStack::ConstructContainerForItem(UNiag
 	return SNew(SNiagaraStackTableRow, StackViewModel, Item, StackCommandContext.ToSharedRef(), StackTree.ToSharedRef())
 		.RowPadding(RowPadding)
 		.ContentPadding(ContentPadding)
-		.ItemBackgroundColor(ItemBackgroundColor)
+ 		.ItemBackgroundColor(ItemBackgroundColor)
 		.ItemForegroundColor(ItemForegroundColor)
+		.IndicatorColor(IndicatorColor)
 		.IsCategoryIconHighlighted(bIsCategoryIconHighlighted)
 		.ShowExecutionCategoryIcon(bShowExecutionCategoryIcon)
 		.NameColumnWidth(this, &SNiagaraStack::GetNameColumnWidth)
@@ -941,7 +916,7 @@ SNiagaraStack::FRowWidgets SNiagaraStack::ConstructNameAndValueWidgetsForItem(UN
 	{
 		Container->SetOverrideNameAlignment(EHorizontalAlignment::HAlign_Fill, EVerticalAlignment::VAlign_Center);
 		return FRowWidgets(SNew(STextBlock)
-			.TextStyle(FNiagaraEditorWidgetsStyle::Get(), "NiagaraEditor.Stack.ParameterText")
+			.TextStyle(FNiagaraEditorStyle::Get(), "NiagaraEditor.ParameterText")
 			.ToolTipText_UObject(Item, &UNiagaraStackEntry::GetTooltipText)
 			.Text_UObject(Item, &UNiagaraStackEntry::GetDisplayName)
 			.ColorAndOpacity(this, &SNiagaraStack::GetTextColorForItem, Item)
@@ -968,7 +943,7 @@ SNiagaraStack::FRowWidgets SNiagaraStack::ConstructNameAndValueWidgetsForItem(UN
 	else if (Item->IsA<UNiagaraStackInputCategory>())
 	{
 		return FRowWidgets(SNew(STextBlock)
-			.TextStyle(FNiagaraEditorWidgetsStyle::Get(), "NiagaraEditor.Stack.CategoryText")
+			.TextStyle(FNiagaraEditorWidgetsStyle::Get(), "NiagaraEditor.Stack.ItemText")
 			.ToolTipText_UObject(Item, &UNiagaraStackEntry::GetTooltipText)
 			.Text_UObject(Item, &UNiagaraStackEntry::GetDisplayName)
 			.ColorAndOpacity(this, &SNiagaraStack::GetTextColorForItem, Item)
@@ -981,7 +956,7 @@ SNiagaraStack::FRowWidgets SNiagaraStack::ConstructNameAndValueWidgetsForItem(UN
 		UNiagaraStackModuleItemOutput* ModuleItemOutput = CastChecked<UNiagaraStackModuleItemOutput>(Item);
 		return FRowWidgets(
 			SNew(SNiagaraParameterName)
-				.ReadOnlyTextStyle(FNiagaraEditorWidgetsStyle::Get(), "NiagaraEditor.Stack.DefaultText")
+				.ReadOnlyTextStyle(FNiagaraEditorWidgetsStyle::Get(), "NiagaraEditor.Stack.ItemText")
 				.ToolTipText_UObject(Item, &UNiagaraStackEntry::GetTooltipText)
 				.ParameterName(ModuleItemOutput->GetOutputParameterHandle().GetParameterHandleString())
 				.HighlightText_UObject(StackViewModel, &UNiagaraStackViewModel::GetCurrentSearchText)
@@ -995,7 +970,7 @@ SNiagaraStack::FRowWidgets SNiagaraStack::ConstructNameAndValueWidgetsForItem(UN
 	{
 		return FRowWidgets(
 			SNew(STextBlock)
-			.TextStyle(FNiagaraEditorWidgetsStyle::Get(), "NiagaraEditor.Stack.DefaultText")
+			.TextStyle(FNiagaraEditorWidgetsStyle::Get(), "NiagaraEditor.Stack.ItemText")
 			.ToolTipText_UObject(Item, &UNiagaraStackEntry::GetTooltipText)
 			.Text_UObject(Item, &UNiagaraStackEntry::GetDisplayName)
 			.ColorAndOpacity(this, &SNiagaraStack::GetTextColorForItem, Item)
@@ -1026,9 +1001,8 @@ SNiagaraStack::FRowWidgets SNiagaraStack::ConstructNameAndValueWidgetsForItem(UN
 		else
 		{
 			Container->SetOverrideNameWidth(PropertyRowWidgets.NameWidgetLayoutData.MinWidth, PropertyRowWidgets.NameWidgetLayoutData.MaxWidth);
-			//Container->SetOverrideNameAlignment(PropertyRowWidgets.NameWidgetLayoutData.HorizontalAlignment, PropertyRowWidgets.NameWidgetLayoutData.VerticalAlignment);
 			Container->SetOverrideValueWidth(PropertyRowWidgets.ValueWidgetLayoutData.MinWidth, PropertyRowWidgets.ValueWidgetLayoutData.MaxWidth);
-			//Container->SetOverrideValueAlignment(PropertyRowWidgets.ValueWidgetLayoutData.HorizontalAlignment, PropertyRowWidgets.ValueWidgetLayoutData.VerticalAlignment);
+			Container->SetOverrideNameAlignment(HAlign_Left, VAlign_Center);
 			PropertyRowWidgets.NameWidget->SetEnabled(IsEnabled);
 			PropertyRowWidgets.ValueWidget->SetEnabled(IsEnabled);
 
diff --git a/Engine/Plugins/FX/Niagara/Source/NiagaraEditorWidgets/Private/SNiagaraStack.h b/Engine/Plugins/FX/Niagara/Source/NiagaraEditorWidgets/Private/SNiagaraStack.h
--- a/Engine/Plugins/FX/Niagara/Source/NiagaraEditorWidgets/Private/SNiagaraStack.h
+++ b/Engine/Plugins/FX/Niagara/Source/NiagaraEditorWidgets/Private/SNiagaraStack.h
@@ -102,7 +102,6 @@ private:
 	void CollapseAll();
 
 	TSharedRef<SWidget> GetViewOptionsMenu() const;
-	const FSlateBrush* GetViewOptionsIconBrush() const;
 
 	// Drag/Drop
 	FReply OnRowDragDetected(const FGeometry& InGeometry, const FPointerEvent& InPointerEvent, UNiagaraStackEntry* InStackEntry);
diff --git a/Engine/Plugins/FX/Niagara/Source/NiagaraEditorWidgets/Private/SNiagaraStackEntryWidget.cpp b/Engine/Plugins/FX/Niagara/Source/NiagaraEditorWidgets/Private/SNiagaraStackEntryWidget.cpp
--- a/Engine/Plugins/FX/Niagara/Source/NiagaraEditorWidgets/Private/SNiagaraStackEntryWidget.cpp
+++ b/Engine/Plugins/FX/Niagara/Source/NiagaraEditorWidgets/Private/SNiagaraStackEntryWidget.cpp
@@ -8,6 +8,7 @@
 #include "Widgets/Text/STextBlock.h"
 #include "Widgets/Text/SInlineEditableTextBlock.h"
 #include "Framework/Application/SlateApplication.h"
+#include "Styling/StyleColors.h"
 
 #define LOCTEXT_NAMESPACE "SNiagaraStackEntryWidget"
 
@@ -17,8 +18,8 @@ void SNiagaraStackDisplayName::Construct(const FArguments& InArgs, UNiagaraStack
 	StackViewModel = &InStackViewModel;
 
 	NameStyle = InArgs._NameStyle;
-	EditableNameStyle = InArgs._EditableNameStyle;
-	OriginalNameStyle = InArgs._OriginalNameStyle;
+ 	EditableNameStyle = InArgs._EditableNameStyle;
+ 	OriginalNameStyle = InArgs._OriginalNameStyle;
 
 	TAttribute<FText> EntryToolTipText;
 	EntryToolTipText.Bind(this, &SNiagaraStackDisplayName::GetEntryToolTipText);
@@ -214,7 +215,7 @@ FSlateColor SNiagaraStackEntryWidget::GetTextColorForSearch(FSlateColor DefaultC
 {
 	if (IsCurrentSearchMatch())
 	{
-		return FSlateColor(FLinearColor(FColor::Orange));
+		return FStyleColors::Select;
 	}
 	return DefaultColor;
 }
diff --git a/Engine/Plugins/FX/Niagara/Source/NiagaraEditorWidgets/Private/SNiagaraStackEntryWidget.h b/Engine/Plugins/FX/Niagara/Source/NiagaraEditorWidgets/Private/SNiagaraStackEntryWidget.h
--- a/Engine/Plugins/FX/Niagara/Source/NiagaraEditorWidgets/Private/SNiagaraStackEntryWidget.h
+++ b/Engine/Plugins/FX/Niagara/Source/NiagaraEditorWidgets/Private/SNiagaraStackEntryWidget.h
@@ -7,6 +7,7 @@
 #include "ViewModels/Stack/NiagaraStackViewModel.h"
 #include "Framework/SlateDelegates.h"
 #include "NiagaraEditorWidgetsStyle.h"
+#include "Styling/AppStyle.h"
 
 class UNiagaraStackViewModel;
 class UNiagaraStackEntry;
@@ -32,9 +33,9 @@ class SNiagaraStackDisplayName : public SNiagaraStackEntryWidget
 {
 public:
 	SLATE_BEGIN_ARGS(SNiagaraStackDisplayName)
-		: _NameStyle(&FNiagaraEditorWidgetsStyle::Get().GetWidgetStyle<FTextBlockStyle>("NiagaraEditor.Stack.ItemText"))
-		, _EditableNameStyle(&FNiagaraEditorWidgetsStyle::Get().GetWidgetStyle<FInlineEditableTextBlockStyle>("NiagaraEditor.Stack.EditableItemText"))
-		, _OriginalNameStyle(&FNiagaraEditorWidgetsStyle::Get().GetWidgetStyle<FTextBlockStyle>("NiagaraEditor.Stack.SubduedItemText"))
+		: _NameStyle(&FAppStyle::Get().GetWidgetStyle<FTextBlockStyle>("NormalText"))
+		, _EditableNameStyle(&FAppStyle::Get().GetWidgetStyle<FInlineEditableTextBlockStyle>("InlineEditableTextBlockStyle"))
+ 		, _OriginalNameStyle(&FAppStyle::Get().GetWidgetStyle<FTextBlockStyle>("NormalText"))
 	{
 		_Clipping = EWidgetClipping::OnDemand;
 	}
diff --git a/Engine/Plugins/FX/Niagara/Source/NiagaraEditorWidgets/Private/Stack/SNiagaraStackErrorItem.cpp b/Engine/Plugins/FX/Niagara/Source/NiagaraEditorWidgets/Private/Stack/SNiagaraStackErrorItem.cpp
--- a/Engine/Plugins/FX/Niagara/Source/NiagaraEditorWidgets/Private/Stack/SNiagaraStackErrorItem.cpp
+++ b/Engine/Plugins/FX/Niagara/Source/NiagaraEditorWidgets/Private/Stack/SNiagaraStackErrorItem.cpp
@@ -11,6 +11,7 @@
 #include "Widgets/Input/SButton.h"
 #include "Widgets/Text/STextBlock.h"
 #include "Widgets/Input/SHyperlink.h"
+#include "Styling/StyleColors.h"
 
 #define LOCTEXT_NAMESPACE "SNiagaraStackErrorItem"
 
@@ -18,21 +19,20 @@ void SNiagaraStackErrorItem::Construct(const FArguments& InArgs, UNiagaraStackEr
 {
 	ErrorItem = InErrorItem;
 	StackViewModel = InStackViewModel;
-
 	const FSlateBrush* IconBrush;
 	switch (ErrorItem->GetStackIssue().GetSeverity())
 	{
 	case EStackIssueSeverity::Error:
-		IconBrush = FEditorStyle::GetBrush("Icons.Error");
+		IconBrush = FAppStyle::Get().GetBrush("Icons.ErrorWithColor");
 		break;
 	case EStackIssueSeverity::Warning:
-		IconBrush = FEditorStyle::GetBrush("Icons.Warning");
+		IconBrush = FAppStyle::Get().GetBrush("Icons.WarningWithColor");
 		break;
 	case EStackIssueSeverity::Info:
-		IconBrush = FEditorStyle::GetBrush("Icons.Info");
+		IconBrush = FAppStyle::Get().GetBrush("Icons.InfoWithColor");
 		break;
 	default:
-		IconBrush = FEditorStyle::GetBrush("NoBrush");
+		IconBrush = FAppStyle::Get().GetBrush("NoBrush");
 		break;
 	}
 
diff --git a/Engine/Plugins/FX/Niagara/Source/NiagaraEditorWidgets/Private/Stack/SNiagaraStackFunctionInputName.cpp b/Engine/Plugins/FX/Niagara/Source/NiagaraEditorWidgets/Private/Stack/SNiagaraStackFunctionInputName.cpp
--- a/Engine/Plugins/FX/Niagara/Source/NiagaraEditorWidgets/Private/Stack/SNiagaraStackFunctionInputName.cpp
+++ b/Engine/Plugins/FX/Niagara/Source/NiagaraEditorWidgets/Private/Stack/SNiagaraStackFunctionInputName.cpp
@@ -56,7 +56,7 @@ void SNiagaraStackFunctionInputName::Construct(const FArguments& InArgs, UNiagar
 			.IsEnabled(this, &SNiagaraStackFunctionInputName::GetIsEnabled)
 			.IsSelected(this, &SNiagaraStackFunctionInputName::GetIsNameWidgetSelected)
 			.OnTextCommitted(this, &SNiagaraStackFunctionInputName::OnNameTextCommitted)
-			//.HighlightText_UObject(InStackViewModel, &UNiagaraStackViewModel::GetCurrentSearchText)
+			.HighlightText_UObject(InStackViewModel, &UNiagaraStackViewModel::GetCurrentSearchText)
 			.ToolTipText(this, &SNiagaraStackFunctionInputName::GetToolTipText)
 		];
 	}
diff --git a/Engine/Plugins/FX/Niagara/Source/NiagaraEditorWidgets/Private/Stack/SNiagaraStackFunctionInputValue.cpp b/Engine/Plugins/FX/Niagara/Source/NiagaraEditorWidgets/Private/Stack/SNiagaraStackFunctionInputValue.cpp
--- a/Engine/Plugins/FX/Niagara/Source/NiagaraEditorWidgets/Private/Stack/SNiagaraStackFunctionInputValue.cpp
+++ b/Engine/Plugins/FX/Niagara/Source/NiagaraEditorWidgets/Private/Stack/SNiagaraStackFunctionInputValue.cpp
@@ -97,17 +97,28 @@ void SNiagaraStackFunctionInputValue::Construct(const FArguments& InArgs, UNiaga
 			+ SHorizontalBox::Slot()
 			.AutoWidth()
 			.VAlign(VAlign_Center)
-			.Padding(3, 0, 0, 0)
+			.Padding(5, 0, 2, 0)
 			[
 				SAssignNew(SetFunctionInputButton, SComboButton)
-				.ButtonStyle(FEditorStyle::Get(), "HoverHintOnly")
-				.ForegroundColor(FSlateColor::UseForeground())
+				.HasDownArrow(false)
+				.ComboButtonStyle(FAppStyle::Get(), "SimpleComboButton") // Use the tool bar item style for this button
 				.OnGetMenuContent(this, &SNiagaraStackFunctionInputValue::OnGetAvailableHandleMenu)
 				.ContentPadding(FMargin(2))
 				.Visibility(this, &SNiagaraStackFunctionInputValue::GetDropdownButtonVisibility)
 				.MenuPlacement(MenuPlacement_BelowRightAnchor)
 				.HAlign(HAlign_Center)
 				.VAlign(VAlign_Center)
+				.ButtonContent()
+				[
+					SNew(SBox)
+					.WidthOverride(12)
+					.HeightOverride(12)
+					[
+						SNew(SImage)
+						.Image(FAppStyle::Get().GetBrush("Icons.Link"))
+						.ColorAndOpacity(FSlateColor::UseForeground())
+					]
+				]
 			]
 
 			// Reset Button
@@ -119,7 +130,7 @@ void SNiagaraStackFunctionInputValue::Construct(const FArguments& InArgs, UNiaga
 				SNew(SButton)
 				.IsFocusable(false)
 				.ToolTipText(LOCTEXT("ResetToolTip", "Reset to the default value"))
-				.ButtonStyle(FEditorStyle::Get(), "NoBorder")
+				.ButtonStyle(FEditorStyle::Get(), "SimpleButton")
 				.ContentPadding(0)
 				.Visibility(this, &SNiagaraStackFunctionInputValue::GetResetButtonVisibility)
 				.OnClicked(this, &SNiagaraStackFunctionInputValue::ResetButtonPressed)
@@ -127,6 +138,7 @@ void SNiagaraStackFunctionInputValue::Construct(const FArguments& InArgs, UNiaga
 				[
 					SNew(SImage)
 					.Image(FEditorStyle::GetBrush("PropertyWindow.DiffersFromDefault"))
+					.ColorAndOpacity(FSlateColor::UseForeground())
 				]
 			]
 
@@ -139,7 +151,7 @@ void SNiagaraStackFunctionInputValue::Construct(const FArguments& InArgs, UNiaga
 				SNew(SButton)
 				.IsFocusable(false)
 				.ToolTipText(LOCTEXT("ResetToBaseToolTip", "Reset this input to the value defined by the parent emitter"))
-				.ButtonStyle(FEditorStyle::Get(), "NoBorder")
+				.ButtonStyle(FEditorStyle::Get(), "SimpleButton")
 				.ContentPadding(0)
 				.Visibility(this, &SNiagaraStackFunctionInputValue::GetResetToBaseButtonVisibility)
 				.OnClicked(this, &SNiagaraStackFunctionInputValue::ResetToBaseButtonPressed)
@@ -495,16 +507,12 @@ TSharedRef<SWidget> SNiagaraStackFunctionInputValue::OnGetAvailableHandleMenu()
 	TSharedPtr<SNiagaraLibraryOnlyToggleHeader> LibraryOnlyToggle;
 	TSharedPtr<SGraphActionMenu> GraphActionMenu;
 
-	TSharedRef<SBorder> MenuWidget = SNew(SBorder)
-	.BorderImage(FEditorStyle::GetBrush("Menu.Background"))
-	.Padding(5)
-	[
-		SNew(SBox)
+	TSharedRef<SBox> MenuWidget = SNew(SBox)
 		.MinDesiredWidth(300)
 		.HeightOverride(400)
 		[
 			SNew(SVerticalBox)
-			+SVerticalBox::Slot()
+			+ SVerticalBox::Slot()
 			.Padding(1.0f)
 			[
 				SAssignNew(LibraryOnlyToggle, SNiagaraLibraryOnlyToggleHeader)
@@ -512,7 +520,7 @@ TSharedRef<SWidget> SNiagaraStackFunctionInputValue::OnGetAvailableHandleMenu()
 				.LibraryOnly(this, &SNiagaraStackFunctionInputValue::GetLibraryOnly)
 				.LibraryOnlyChanged(this, &SNiagaraStackFunctionInputValue::SetLibraryOnly)
 			]
-			+SVerticalBox::Slot()
+			+ SVerticalBox::Slot()
 			.FillHeight(15)
 			[
 				SAssignNew(GraphActionMenu, SGraphActionMenu)
@@ -526,9 +534,7 @@ TSharedRef<SWidget> SNiagaraStackFunctionInputValue::OnGetAvailableHandleMenu()
 					return SNew(SNiagaraGraphActionWidget, InData);
 				})
 			]
-		]
-	];
-
+		];
 	LibraryOnlyToggle->SetActionMenu(GraphActionMenu.ToSharedRef());
 	SetFunctionInputButton->SetMenuContentWidgetToFocus(GraphActionMenu->GetFilterTextBox()->AsShared());
 	return MenuWidget;
diff --git a/Engine/Plugins/FX/Niagara/Source/NiagaraEditorWidgets/Private/Stack/SNiagaraStackIssueIcon.cpp b/Engine/Plugins/FX/Niagara/Source/NiagaraEditorWidgets/Private/Stack/SNiagaraStackIssueIcon.cpp
--- a/Engine/Plugins/FX/Niagara/Source/NiagaraEditorWidgets/Private/Stack/SNiagaraStackIssueIcon.cpp
+++ b/Engine/Plugins/FX/Niagara/Source/NiagaraEditorWidgets/Private/Stack/SNiagaraStackIssueIcon.cpp
@@ -39,9 +39,8 @@ void SNiagaraStackIssueIcon::Construct(const FArguments& InArgs, UNiagaraStackVi
 	{
 		TSharedRef<SButton> IconButton =
 			SNew(SButton)
-			.ButtonStyle(FEditorStyle::Get(), "HoverHintOnly")
-			.ForegroundColor(FNiagaraEditorWidgetsStyle::Get().GetColor("NiagaraEditor.Stack.ForegroundColor"))
-			.ContentPadding(FMargin(1, 0, 0, 0))
+			.ContentPadding(0)
+			.ButtonStyle(FAppStyle::Get(), "SimpleButton")
 			.OnClicked(InArgs._OnClicked)
 			.Content()
 			[
@@ -206,15 +205,15 @@ void SNiagaraStackIssueIcon::UpdateFromEntry()
 
 	if (StackEntry->GetTotalNumberOfErrorIssues() > 0)
 	{
-		IconBrush = FEditorStyle::GetBrush("Icons.Error");
+		IconBrush = FAppStyle::Get().GetBrush("Icons.ErrorWithColor");
 	}
 	else if (StackEntry->GetTotalNumberOfWarningIssues() > 0)
 	{
-		IconBrush = FEditorStyle::GetBrush("Icons.Warning");
+		IconBrush = FAppStyle::Get().GetBrush("Icons.WarningWithColor");
 	}
 	else if (StackEntry->GetTotalNumberOfInfoIssues() > 0)
 	{
-		IconBrush = FEditorStyle::GetBrush("Icons.Info");
+		IconBrush = FAppStyle::Get().GetBrush("Icons.InfoWithColor");
 	}
 
 	IconToolTipCache.Reset();
diff --git a/Engine/Plugins/FX/Niagara/Source/NiagaraEditorWidgets/Private/Stack/SNiagaraStackItem.cpp b/Engine/Plugins/FX/Niagara/Source/NiagaraEditorWidgets/Private/Stack/SNiagaraStackItem.cpp
--- a/Engine/Plugins/FX/Niagara/Source/NiagaraEditorWidgets/Private/Stack/SNiagaraStackItem.cpp
+++ b/Engine/Plugins/FX/Niagara/Source/NiagaraEditorWidgets/Private/Stack/SNiagaraStackItem.cpp
@@ -14,6 +14,7 @@
 #include "Widgets/Images/SImage.h"
 #include "Widgets/Layout/SBox.h"
 #include "ViewModels/Stack/NiagaraStackClipboardUtilities.h"
+#include "Styling/StyleColors.h"
 
 #define LOCTEXT_NAMESPACE "NiagaraStackItem"
 
@@ -44,6 +45,8 @@ void SNiagaraStackItem::Construct(const FArguments& InArgs, UNiagaraStackItem& I
 		.VAlign(VAlign_Center)
 		[
 			SAssignNew(DisplayNameWidget, SNiagaraStackDisplayName, InItem, *InStackViewModel)
+			.NameStyle(FNiagaraEditorWidgetsStyle::Get(), "NiagaraEditor.Stack.ItemText")
+			.EditableNameStyle(FNiagaraEditorWidgetsStyle::Get(), "NiagaraEditor.Stack.EditableItemText")
 		];
 
 	// Allow derived classes to add additional widgets.
@@ -84,15 +87,14 @@ void SNiagaraStackItem::Construct(const FArguments& InArgs, UNiagaraStackItem& I
 			SNew(SButton)
 			.ButtonStyle(FEditorStyle::Get(), "HoverHintOnly")
 			.IsFocusable(false)
-			.ForegroundColor(FNiagaraEditorWidgetsStyle::Get().GetColor("NiagaraEditor.Stack.FlatButtonColor"))
 			.ToolTipText(this, &SNiagaraStackItem::GetDeleteButtonToolTipText)
 			.IsEnabled(this, &SNiagaraStackItem::GetDeleteButtonEnabled)
 			.OnClicked(this, &SNiagaraStackItem::DeleteClicked)
 			.Content()
 			[
-				SNew(STextBlock)
-				.Font(FEditorStyle::Get().GetFontStyle("FontAwesome.10"))
-				.Text(FText::FromString(FString(TEXT("\xf1f8"))))
+				SNew(SImage)
+				.Image(FAppStyle::Get().GetBrush("Icons.Delete"))
+				.ColorAndOpacity(FSlateColor::UseForeground())
 			]
 		];
 	}
diff --git a/Engine/Plugins/FX/Niagara/Source/NiagaraEditorWidgets/Private/Stack/SNiagaraStackItemFooter.cpp b/Engine/Plugins/FX/Niagara/Source/NiagaraEditorWidgets/Private/Stack/SNiagaraStackItemFooter.cpp
--- a/Engine/Plugins/FX/Niagara/Source/NiagaraEditorWidgets/Private/Stack/SNiagaraStackItemFooter.cpp
+++ b/Engine/Plugins/FX/Niagara/Source/NiagaraEditorWidgets/Private/Stack/SNiagaraStackItemFooter.cpp
@@ -51,7 +51,7 @@ void SNiagaraStackItemFooter::Construct(const FArguments& InArgs, UNiagaraStackI
 
 EVisibility SNiagaraStackItemFooter::GetExpandButtonVisibility() const
 {
-	return ItemFooter->GetHasAdvancedContent() ? EVisibility::Visible : EVisibility::Hidden;
+	return ItemFooter->GetHasAdvancedContent() ? EVisibility::Visible : EVisibility::Collapsed;
 }
 
 EVisibility SNiagaraStackItemFooter::GetOverrideIconVisibility() const
diff --git a/Engine/Plugins/FX/Niagara/Source/NiagaraEditorWidgets/Private/Stack/SNiagaraStackItemGroup.cpp b/Engine/Plugins/FX/Niagara/Source/NiagaraEditorWidgets/Private/Stack/SNiagaraStackItemGroup.cpp
--- a/Engine/Plugins/FX/Niagara/Source/NiagaraEditorWidgets/Private/Stack/SNiagaraStackItemGroup.cpp
+++ b/Engine/Plugins/FX/Niagara/Source/NiagaraEditorWidgets/Private/Stack/SNiagaraStackItemGroup.cpp
@@ -11,6 +11,7 @@
 #include "Widgets/Input/SButton.h"
 #include "Widgets/Input/SCheckBox.h"
 #include "Widgets/Text/STextBlock.h"
+#include "Styling/StyleColors.h"
 
 #define LOCTEXT_NAMESPACE "NiagaraStackItemGroup"
 
@@ -30,6 +31,7 @@ void SNiagaraStackItemGroup::Construct(const FArguments& InArgs, UNiagaraStackIt
 	[
 		SNew(SNiagaraStackDisplayName, InGroup, *InStackViewModel)
 		.NameStyle(FNiagaraEditorWidgetsStyle::Get(), "NiagaraEditor.Stack.GroupText")
+		.EditableNameStyle(FNiagaraEditorWidgetsStyle::Get(), "NiagaraEditor.Stack.EditableGroupText")
 	];
 
 	// Delete group button
@@ -39,16 +41,15 @@ void SNiagaraStackItemGroup::Construct(const FArguments& InArgs, UNiagaraStackIt
 		SNew(SButton)
 		.ButtonStyle(FEditorStyle::Get(), "HoverHintOnly")
 		.IsFocusable(false)
-		.ForegroundColor(FNiagaraEditorWidgetsStyle::Get().GetColor("NiagaraEditor.Stack.ForegroundColor"))
 		.ToolTipText(this, &SNiagaraStackItemGroup::GetDeleteButtonToolTip)
 		.OnClicked(this, &SNiagaraStackItemGroup::DeleteClicked)
 		.IsEnabled(this, &SNiagaraStackItemGroup::GetDeleteButtonIsEnabled)
 		.Visibility(this, &SNiagaraStackItemGroup::GetDeleteButtonVisibility)
 		.Content()
 		[
-			SNew(STextBlock)
-			.Font(FEditorStyle::Get().GetFontStyle("FontAwesome.10"))
-			.Text(FText::FromString(FString(TEXT("\xf1f8"))))
+			SNew(SImage)
+			.Image(FAppStyle::Get().GetBrush("Icons.Delete"))
+			.ColorAndOpacity(FSlateColor::UseForeground())
 		]
 	];
 
diff --git a/Engine/Plugins/FX/Niagara/Source/NiagaraEditorWidgets/Private/Stack/SNiagaraStackItemGroupAddButton.cpp b/Engine/Plugins/FX/Niagara/Source/NiagaraEditorWidgets/Private/Stack/SNiagaraStackItemGroupAddButton.cpp
--- a/Engine/Plugins/FX/Niagara/Source/NiagaraEditorWidgets/Private/Stack/SNiagaraStackItemGroupAddButton.cpp
+++ b/Engine/Plugins/FX/Niagara/Source/NiagaraEditorWidgets/Private/Stack/SNiagaraStackItemGroupAddButton.cpp
@@ -28,53 +28,33 @@ void SNiagaraStackItemGroupAddButton::Construct(const FArguments& InArgs, UNiaga
 		if (AddUtilities->GetAddMode() == INiagaraStackItemGroupAddUtilities::AddFromAction)
 		{
 			Content = SAssignNew(AddActionButton, SComboButton)
-				.ButtonStyle(FNiagaraEditorWidgetsStyle::Get(), "NiagaraEditor.Stack.AddButton")
-				.ButtonColorAndOpacity(FNiagaraEditorWidgetsStyle::Get().GetColor(FNiagaraStackEditorWidgetsUtilities::GetIconColorNameForExecutionCategory(InStackItemGroup.GetExecutionCategoryName())))
+				.ButtonStyle(FAppStyle::Get(), "SimpleButton")
 				.ToolTipText(FText::Format(AddToGroupFormat, AddUtilities->GetAddItemName()))
 				.HasDownArrow(false)
 				.OnGetMenuContent(this, &SNiagaraStackItemGroupAddButton::GetAddMenu)
 				.IsEnabled_UObject(&InStackItemGroup, &UNiagaraStackEntry::GetIsEnabledAndOwnerIsEnabled)
-				.ContentPadding(0)
+				.ContentPadding(1.0f)
 				.MenuPlacement(MenuPlacement_BelowRightAnchor)
 				.ButtonContent()
 				[
-					SNew(SBox)
-					.WidthOverride(InArgs._Width)
-					.HeightOverride(TextIconSize)
-					.HAlign(HAlign_Center)
-					.VAlign(VAlign_Center)
-					[
-						SNew(STextBlock)
-						//.TextStyle(FEditorStyle::Get(), "NormalText.Important")
-						.TextStyle(FNiagaraEditorWidgetsStyle::Get(), "NiagaraEditor.Stack.AddButtonText")
-						.Font(FEditorStyle::Get().GetFontStyle("FontAwesome.10"))
-						.Text(FText::FromString(FString(TEXT("\xf067"))) /*fa-plus*/)
-					]
+					SNew(SImage)
+					.Image(FAppStyle::Get().GetBrush("Icons.PlusCircle"))
+					.ColorAndOpacity(FNiagaraEditorWidgetsStyle::Get().GetColor(FNiagaraStackEditorWidgetsUtilities::GetIconColorNameForExecutionCategory(InStackItemGroup.GetExecutionCategoryName())))
 				];
 		}
 		else if (AddUtilities->GetAddMode() == INiagaraStackItemGroupAddUtilities::AddDirectly)
 		{
 			Content = SNew(SButton)
-				.ButtonStyle(FNiagaraEditorWidgetsStyle::Get(), "NiagaraEditor.Stack.AddButton")
-				.ButtonColorAndOpacity(FNiagaraEditorWidgetsStyle::Get().GetColor(FNiagaraStackEditorWidgetsUtilities::GetIconColorNameForExecutionCategory(InStackItemGroup.GetExecutionCategoryName())))
+				.ButtonStyle(FAppStyle::Get(), "SimpleButton")
 				.ToolTipText(FText::Format(AddToGroupFormat, AddUtilities->GetAddItemName()))
 				.IsEnabled_UObject(&InStackItemGroup, &UNiagaraStackEntry::GetIsEnabledAndOwnerIsEnabled)
-				.ContentPadding(0)
+				.ContentPadding(1.0f)
 				.OnClicked(this, &SNiagaraStackItemGroupAddButton::AddDirectlyButtonClicked)
 				.Content()
 				[
-					SNew(SBox)
-					.WidthOverride(InArgs._Width)
-					.HeightOverride(TextIconSize)
-					.HAlign(HAlign_Center)
-					.VAlign(VAlign_Center)
-					[
-						SNew(STextBlock)
-						//.TextStyle(FEditorStyle::Get(), "NormalText.Important")
-						.TextStyle(FNiagaraEditorWidgetsStyle::Get(), "NiagaraEditor.Stack.AddButtonText")
-						.Font(FEditorStyle::Get().GetFontStyle("FontAwesome.10"))
-						.Text(FText::FromString(FString(TEXT("\xf067"))) /*fa-plus*/)
-					]
+					SNew(SImage)
+					.Image(FAppStyle::Get().GetBrush("Icons.PlusCircle"))
+					.ColorAndOpacity(FNiagaraEditorWidgetsStyle::Get().GetColor(FNiagaraStackEditorWidgetsUtilities::GetIconColorNameForExecutionCategory(InStackItemGroup.GetExecutionCategoryName())))
 				];
 		}
 	}
diff --git a/Engine/Plugins/FX/Niagara/Source/NiagaraEditorWidgets/Private/Stack/SNiagaraStackModuleItem.cpp b/Engine/Plugins/FX/Niagara/Source/NiagaraEditorWidgets/Private/Stack/SNiagaraStackModuleItem.cpp
--- a/Engine/Plugins/FX/Niagara/Source/NiagaraEditorWidgets/Private/Stack/SNiagaraStackModuleItem.cpp
+++ b/Engine/Plugins/FX/Niagara/Source/NiagaraEditorWidgets/Private/Stack/SNiagaraStackModuleItem.cpp
@@ -112,8 +112,7 @@ void SNiagaraStackModuleItem::AddCustomRowWidgets(TSharedRef<SHorizontalBox> Hor
 	[
 		SAssignNew(AddButton, SComboButton)
 		.HasDownArrow(false)
-		.ButtonStyle(FEditorStyle::Get(), "HoverHintOnly")
-		.ForegroundColor(FSlateColor::UseForeground())
+		.ButtonStyle(FAppStyle::Get(), "SimpleButton")
 		.OnGetMenuContent(this, &SNiagaraStackModuleItem::RaiseActionMenuClicked)
 		.ContentPadding(FMargin(2))
 		.HAlign(HAlign_Center)
@@ -123,7 +122,8 @@ void SNiagaraStackModuleItem::AddCustomRowWidgets(TSharedRef<SHorizontalBox> Hor
 		.ButtonContent()
 		[
 			SNew(SImage)
-			.Image(FEditorStyle::Get().GetBrush("Icons.PlusCircle"))
+			.Image(FAppStyle::Get().GetBrush("Icons.PlusCircle"))
+			.ColorAndOpacity(FSlateColor::UseForeground())
 		]
 	];
 
@@ -133,9 +133,8 @@ void SNiagaraStackModuleItem::AddCustomRowWidgets(TSharedRef<SHorizontalBox> Hor
 	.AutoWidth()
 	[
 		SNew(SButton)
-		.ButtonStyle(FEditorStyle::Get(), "HoverHintOnly")
+		.ButtonStyle(FAppStyle::Get(), "SimpleButton")
 		.IsFocusable(false)
-		.ForegroundColor(FNiagaraEditorWidgetsStyle::Get().GetColor("NiagaraEditor.Stack.FlatButtonColor"))
 		.ToolTipText(LOCTEXT("RefreshTooltip", "Refresh this module"))
 		.Visibility(this, &SNiagaraStackModuleItem::GetRefreshVisibility)
 		.IsEnabled(this, &SNiagaraStackModuleItem::GetButtonsEnabled)
@@ -144,7 +143,8 @@ void SNiagaraStackModuleItem::AddCustomRowWidgets(TSharedRef<SHorizontalBox> Hor
 		[
 			SNew(STextBlock)
 			.Font(FEditorStyle::Get().GetFontStyle("FontAwesome.10"))
-		.Text(FEditorFontGlyphs::Refresh)
+			.Text(FEditorFontGlyphs::Refresh)
+			.ColorAndOpacity(FSlateColor::UseForeground())
 		]
 	];
 }
diff --git a/Engine/Plugins/FX/Niagara/Source/NiagaraEditorWidgets/Private/Stack/SNiagaraStackTableRow.cpp b/Engine/Plugins/FX/Niagara/Source/NiagaraEditorWidgets/Private/Stack/SNiagaraStackTableRow.cpp
--- a/Engine/Plugins/FX/Niagara/Source/NiagaraEditorWidgets/Private/Stack/SNiagaraStackTableRow.cpp
+++ b/Engine/Plugins/FX/Niagara/Source/NiagaraEditorWidgets/Private/Stack/SNiagaraStackTableRow.cpp
@@ -14,6 +14,7 @@
 #include "Widgets/Images/SImage.h"
 #include "Widgets/Layout/SBox.h"
 #include "Framework/Application/SlateApplication.h"
+#include "Styling/StyleColors.h"
 
 #define LOCTEXT_NAMESPACE "NiagaraStackTableRow"
 
@@ -39,8 +40,9 @@ void SNiagaraStackTableRow::Construct(const FArguments& InArgs, UNiagaraStackVie
 	CollapsedImage = FCoreStyle::Get().GetBrush("TreeArrow_Collapsed");
 
 	ItemBackgroundColor = InArgs._ItemBackgroundColor;
-	DisabledItemBackgroundColor = ItemBackgroundColor + FLinearColor(.02f, .02f, .02f, 0.0f);
+	DisabledItemBackgroundColor = FStyleColors::Recessed;
 	ForegroundColor = InArgs._ItemForegroundColor;
+	IndicatorColor = InArgs._IndicatorColor;
 
 	ExecutionCategoryToolTipText = InStackEntry->GetExecutionSubcategoryName() != NAME_None
 		? FText::Format(LOCTEXT("ExecutionCategoryToolTipFormat", "{0} - {1}"), FText::FromName(InStackEntry->GetExecutionCategoryName()), FText::FromName(InStackEntry->GetExecutionSubcategoryName()))
@@ -156,7 +158,6 @@ void SNiagaraStackTableRow::SetNameAndValueContent(TSharedRef<SWidget> InNameWid
 		.ButtonStyle(FCoreStyle::Get(), "NoBorder")
 		.Visibility(this, &SNiagaraStackTableRow::GetExpanderVisibility)
 		.OnClicked(this, &SNiagaraStackTableRow::ExpandButtonClicked)
-		.ForegroundColor(FSlateColor::UseForeground())
 		.ContentPadding(2)
 		.HAlign(HAlign_Center)
 		[
@@ -218,7 +219,7 @@ void SNiagaraStackTableRow::SetNameAndValueContent(TSharedRef<SWidget> InNameWid
 	if (InValueWidget.IsValid())
 	{
 		ChildContent = SNew(SSplitter)
-		.Style(FNiagaraEditorWidgetsStyle::Get(), "NiagaraEditor.Stack.Splitter")
+		.Style(FAppStyle::Get(), "DetailsView.Splitter")
 		.PhysicalSplitterHandleSize(1.0f)
 		.HitDetectionSplitterHandleSize(5.0f)
 
@@ -231,14 +232,8 @@ void SNiagaraStackTableRow::SetNameAndValueContent(TSharedRef<SWidget> InNameWid
 			.Padding(FMargin(ContentPadding.Left, 0, 0, 0))
 		    .AutoWidth()
 		    [
-				 SNew(SBorder)
-		        .BorderImage(FEditorStyle::GetBrush("WhiteBrush"))
-		        .BorderBackgroundColor(FNiagaraEditorWidgetsStyle::Get().GetColor("NiagaraEditor.Stack.DividerColor"))
-		        .Padding(0)
-		        [
-		            SNew(SBox)
-		            .WidthOverride(StackEntry->HasFrontDivider() ? 1 : 0)
-		        ]
+				SNew(SBox)
+				.WidthOverride(StackEntry->HasFrontDivider() ? 1 : 0)
 		    ]
 		    + SHorizontalBox::Slot()
             [
@@ -282,7 +277,9 @@ void SNiagaraStackTableRow::SetNameAndValueContent(TSharedRef<SWidget> InNameWid
 	}
 
 	FName AccentColorName = FNiagaraStackEditorWidgetsUtilities::GetIconColorNameForExecutionCategory(StackEntry->GetExecutionCategoryName());
-	FLinearColor AccentColor = AccentColorName != NAME_None ? FNiagaraEditorWidgetsStyle::Get().GetColor(AccentColorName) : FLinearColor::Transparent;
+	const bool bDisplayingIndicator = IndicatorColor != FStyleColors::Transparent;
+	FSlateColor AccentColor = bDisplayingIndicator ? IndicatorColor : (AccentColorName != NAME_None ? FNiagaraEditorWidgetsStyle::Get().GetColor(AccentColorName) : FStyleColors::Transparent);
+
 	ChildSlot
 	[
 		SNew(SHorizontalBox)
@@ -290,70 +287,74 @@ void SNiagaraStackTableRow::SetNameAndValueContent(TSharedRef<SWidget> InNameWid
 		// Accent color.
 		+ SHorizontalBox::Slot()
 		.AutoWidth()
-		.Padding(1, 0, 6, 0)
+		.Padding(bDisplayingIndicator ? FMargin(0, 0, 5, 0) : FMargin(1, 0, 6, 0) )
 		[
 			SNew(SBorder)
-			.BorderImage(FEditorStyle::GetBrush("WhiteBrush"))
+ 			.BorderImage(FEditorStyle::GetBrush("WhiteBrush"))
 			.BorderBackgroundColor(AccentColor)
 			.Padding(0)
 			[
 				SNew(SBox)
-				.WidthOverride(4)
+				.WidthOverride(bDisplayingIndicator ? 6 : 4)
 			]
 		]
 		// Content
 		+ SHorizontalBox::Slot()
+		.Padding(0.0f)
 		[
 			// Row content
-			SNew(SBox)
-			.Padding(RowPadding)
+			SNew(SBorder)
+			.BorderImage(bDisplayingIndicator ? FAppStyle::Get().GetBrush("Brushes.Header") : FAppStyle::Get().GetBrush("DetailsView.GridLine"))
+			.Padding(FMargin(0, 0, 0, 1))
 			[
-				SNew(SBorder)
-				.BorderImage(FEditorStyle::GetBrush("WhiteBrush"))
-				.BorderBackgroundColor(this, &SNiagaraStackTableRow::GetItemBackgroundColor)
-				.ForegroundColor(ForegroundColor)
-				.Padding(0)
-				[
-					SNew(SVerticalBox)
-					+SVerticalBox::Slot()
-					.Padding(bInsertDivAbove ? FMargin(8, 8, 8, 4) : 0)
-					.AutoHeight()
-                    [
-	                    SNew(SBorder)
-			            .BorderImage(FEditorStyle::GetBrush("WhiteBrush"))
-			            .BorderBackgroundColor(FNiagaraEditorWidgetsStyle::Get().GetColor("NiagaraEditor.Stack.DividerColor"))
-			            .Visibility(bInsertDivAbove ? EVisibility::Visible : EVisibility::Collapsed)
-			            .Padding(0)
-			            [
-			                SNew(SBox)
-			                .HeightOverride(1)
-			            ]
-                    ]
-					+SVerticalBox::Slot()
+					SNew(SBorder)
+ 					.BorderImage(FEditorStyle::GetBrush("WhiteBrush"))
+					.BorderBackgroundColor(this, &SNiagaraStackTableRow::GetItemBackgroundColor)
+ 					.ForegroundColor(ForegroundColor)
+					.Padding(0)
 					[
-						SNew(SBorder)
-						.BorderImage(this, &SNiagaraStackTableRow::GetSearchResultBorderBrush)
-						.BorderBackgroundColor(FNiagaraEditorWidgetsStyle::Get().GetColor("NiagaraEditor.Stack.SearchHighlightColor"))
-						.Padding(0)
+						SNew(SHorizontalBox)
+						+ SHorizontalBox::Slot()
+						.Padding(0.0f)
 						[
-							SNew(SBorder)
-							.BorderImage(this, &SNiagaraStackTableRow::GetSelectionBorderBrush)
+							SNew(SVerticalBox)
+							+SVerticalBox::Slot()
+							.Padding(bInsertDivAbove ? FMargin(8, 8, 8, 4) : 0)
+							.AutoHeight()
+							[
+								SNew(SBorder)
+								.BorderImage(FEditorStyle::GetBrush("WhiteBrush"))
+ 								.BorderBackgroundColor(FStyleColors::Panel)
+								.Visibility(bInsertDivAbove ? EVisibility::Visible : EVisibility::Collapsed)
+								.Padding(0)
+								[
+									SNew(SBox)
+									.HeightOverride(1)
+								]
+							]
+							+SVerticalBox::Slot()
 							.Padding(0)
 							[
-								ChildContent.ToSharedRef()
+								SNew(SBorder)
+								.BorderImage(this, &SNiagaraStackTableRow::GetSearchResultBorderBrush)
+ 								.BorderBackgroundColor(FStyleColors::Select)
+								.Padding(0)
+								[
+									ChildContent.ToSharedRef()
+								]
 							]
 						]
+						+ SHorizontalBox::Slot()
+						.AutoWidth()
+						.Padding(3, 0)
+						[
+							SNew(SNiagaraStackIssueIcon, StackViewModel, StackEntry)
+							.Visibility(IssueIconVisibility)
+						]
 					]
 				]
 			]
-		]
-		+ SHorizontalBox::Slot()
-		.AutoWidth()
-		.Padding(3, 0, 0, 0)
-		[
-			SNew(SNiagaraStackIssueIcon, StackViewModel, StackEntry)
-			.Visibility(IssueIconVisibility)
-		]
+		
 	];
 }
 
diff --git a/Engine/Plugins/FX/Niagara/Source/NiagaraEditorWidgets/Private/Stack/SNiagaraStackTableRow.h b/Engine/Plugins/FX/Niagara/Source/NiagaraEditorWidgets/Private/Stack/SNiagaraStackTableRow.h
--- a/Engine/Plugins/FX/Niagara/Source/NiagaraEditorWidgets/Private/Stack/SNiagaraStackTableRow.h
+++ b/Engine/Plugins/FX/Niagara/Source/NiagaraEditorWidgets/Private/Stack/SNiagaraStackTableRow.h
@@ -29,8 +29,9 @@ public:
 		, _RowPadding(FMargin(0, 0, 0, 0))
 	{}
 	SLATE_ARGUMENT(FMargin, ContentPadding)
-		SLATE_ARGUMENT(FLinearColor, ItemBackgroundColor)
-		SLATE_ARGUMENT(FLinearColor, ItemForegroundColor)
+		SLATE_ARGUMENT(FSlateColor, ItemBackgroundColor)
+		SLATE_ARGUMENT(FSlateColor, ItemForegroundColor)
+		SLATE_ARGUMENT(FSlateColor, IndicatorColor)
 		SLATE_ARGUMENT(bool, IsCategoryIconHighlighted)
 		SLATE_ARGUMENT(bool, ShowExecutionCategoryIcon)
 		SLATE_ATTRIBUTE(float, NameColumnWidth)
@@ -114,9 +115,10 @@ private:
 	const FSlateBrush* ExpandedImage;
 	const FSlateBrush* CollapsedImage;
 
-	FLinearColor ItemBackgroundColor;
-	FLinearColor DisabledItemBackgroundColor;
-	FLinearColor ForegroundColor;
+	FSlateColor ItemBackgroundColor;
+	FSlateColor DisabledItemBackgroundColor;
+	FSlateColor ForegroundColor;
+	FSlateColor IndicatorColor;
 
 	FText ExecutionCategoryToolTipText;
 
diff --git a/Engine/Source/Editor/EditorStyle/Private/StarshipStyle.cpp b/Engine/Source/Editor/EditorStyle/Private/StarshipStyle.cpp
--- a/Engine/Source/Editor/EditorStyle/Private/StarshipStyle.cpp
+++ b/Engine/Source/Editor/EditorStyle/Private/StarshipStyle.cpp
@@ -2587,6 +2587,16 @@ void FStarshipEditorStyle::FStyle::SetupGeneralIcons()
 	Set("Cross", new IMAGE_BRUSH("Icons/Cross_12x", Icon12x12));
 	Set("ArrowUp", new IMAGE_BRUSH("Icons/ArrowUp_12x", Icon12x12));
 	Set("ArrowDown", new IMAGE_BRUSH("Icons/ArrowDown_12x", Icon12x12));
+	Set("AssetEditor.SaveThumbnail", new IMAGE_BRUSH_SVG("Starship/AssetEditors/SaveThumbnail", Icon20x20));
+	Set("AssetEditor.ToggleShowBounds", new IMAGE_BRUSH_SVG("Starship/Common/SetShowBounds", Icon20x20));
+	Set("AssetEditor.Apply", new IMAGE_BRUSH_SVG("Starship/Common/Apply", Icon20x20));
+	Set("AssetEditor.Simulate", new IMAGE_BRUSH_SVG("Starship/MainToolbar/simulate", Icon20x20));
+	Set("AssetEditor.ToggleStats", new IMAGE_BRUSH_SVG("Starship/Common/Statistics", Icon20x20));
+	Set("AssetEditor.CompileStatus.Background", new IMAGE_BRUSH_SVG("Starship/Blueprints/CompileStatus_Background", Icon20x20));
+	Set("AssetEditor.CompileStatus.Overlay.Unknown", new IMAGE_BRUSH_SVG("Starship/Blueprints/CompileStatus_Unknown_Badge", Icon20x20, FStyleColors::AccentYellow));
+	Set("AssetEditor.CompileStatus.Overlay.Warning", new IMAGE_BRUSH_SVG("Starship/Blueprints/CompileStatus_Warning_Badge", Icon20x20, FStyleColors::Warning));
+	Set("AssetEditor.CompileStatus.Overlay.Good", new IMAGE_BRUSH_SVG("Starship/Blueprints/CompileStatus_Good_Badge", Icon20x20, FStyleColors::AccentGreen));
+	Set("AssetEditor.CompileStatus.Overlay.Error", new IMAGE_BRUSH_SVG("Starship/Blueprints/CompileStatus_Fail_Badge", Icon20x20, FStyleColors::Error));
 }
 
 void FStarshipEditorStyle::FStyle::SetupWindowStyles()
@@ -3155,7 +3165,7 @@ void FStarshipEditorStyle::FStyle::SetupPropertyEditorStyles()
 			.SetForegroundColor(FSlateColor(EStyleColor::White));
 		NameStyle.TextStyle.SetFont(DEFAULT_FONT("Regular", 11))
 			.SetColorAndOpacity(FSlateColor(EStyleColor::White));
-
+		Set( "DetailsView.ConstantTextBlockStyle", NameStyle.TextStyle);
 		Set( "DetailsView.NameTextBlockStyle",  NameStyle );
 
 		Set( "DetailsView.NameChangeCommitted", new BOX_BRUSH( "Common/EditableTextSelectionBackground", FMargin(4.f/16.f) ) );
@@ -3217,6 +3227,12 @@ void FStarshipEditorStyle::FStyle::SetupPropertyEditorStyles()
 			.SetColorAndOpacity(FStyleColors::ForegroundHeader)
 		);
 
+		Set("DetailsView.CategoryTextStyleUpdate", FTextBlockStyle(NormalText)
+			.SetFont(FStyleFonts::Get().Small)
+			.SetColorAndOpacity(FStyleColors::ForegroundHeader)
+			.SetTransformPolicy(ETextTransformPolicy::ToUpper)
+		);
+
 		FButtonStyle DetailsExtensionMenuButton = FButtonStyle(FStarshipCoreStyle::GetCoreStyle().GetWidgetStyle<FButtonStyle>("NoBorder"))
 			.SetNormalForeground(FStyleColors::Foreground)
 			.SetHoveredForeground(FStyleColors::ForegroundHover)
diff --git a/Engine/Source/Editor/PropertyEditor/Private/DetailCategoryBuilderImpl.cpp b/Engine/Source/Editor/PropertyEditor/Private/DetailCategoryBuilderImpl.cpp
--- a/Engine/Source/Editor/PropertyEditor/Private/DetailCategoryBuilderImpl.cpp
+++ b/Engine/Source/Editor/PropertyEditor/Private/DetailCategoryBuilderImpl.cpp
@@ -15,6 +15,7 @@
 #include "IPropertyGenerationUtilities.h"
 #include "DetailBuilderTypes.h"
 #include "PropertyCustomizationHelpers.h"
+#include "Styling/StyleColors.h"
 
 namespace DetailLayoutConstants
 {
@@ -711,13 +712,13 @@ bool FDetailCategoryImpl::GenerateStandaloneWidget(FDetailWidgetRow& OutRow) con
 	}
 
 	const bool bIsInnerCategory = DetailLayoutBuilder.Pin()->IsLayoutForExternalRoot();
-
+	FTextBlockStyle NameStyle = bIsInnerCategory ? FAppStyle::Get().GetWidgetStyle<FTextBlockStyle>("NormalText") : FAppStyle::Get().GetWidgetStyle<FTextBlockStyle>("DetailsView.CategoryTextStyleUpdate");
 	OutRow.NameContent()
 	[
 		SNew(STextBlock)
 		.Text(GetDisplayName())
-		.Font(FEditorStyle::GetFontStyle(bIsInnerCategory ? "PropertyWindow.NormalFont" : "DetailsView.CategoryFontStyle"))
-		.ShadowOffset(bIsInnerCategory ? FVector2D::ZeroVector : FVector2D(1.0f, 1.0f))
+		.TextStyle(&NameStyle)
+		.ShadowOffset(FVector2D::ZeroVector)
 	];
 
 	if(HeaderContentWidget.IsValid())
diff --git a/Engine/Source/Runtime/SlateCore/Private/Styling/StarshipCoreStyle.cpp b/Engine/Source/Runtime/SlateCore/Private/Styling/StarshipCoreStyle.cpp
--- a/Engine/Source/Runtime/SlateCore/Private/Styling/StarshipCoreStyle.cpp
+++ b/Engine/Source/Runtime/SlateCore/Private/Styling/StarshipCoreStyle.cpp
@@ -290,7 +290,6 @@ TSharedRef<ISlateStyle> FStarshipCoreStyle::Create()
 		Style->Set("Icons.Denied", new IMAGE_BRUSH("Icons/denied_16x", Icon16x16));
 	
 		Style->Set("Icons.Help", new IMAGE_BRUSH("Icons/icon_help_16x", Icon16x16));
-		Style->Set("Icons.Info", new IMAGE_BRUSH("Icons/icon_info_16x", Icon16x16));
 	
 		Style->Set("Icons.Download", new IMAGE_BRUSH("Icons/icon_Downloads_16x", Icon16x16));
 
@@ -300,6 +299,9 @@ TSharedRef<ISlateStyle> FStarshipCoreStyle::Create()
 		Style->Set("Icons.Warning", new IMAGE_BRUSH_SVG("Starship/Common/alert-triangle", Icon16x16));
 		Style->Set("Icons.WarningWithColor", new IMAGE_BRUSH_SVG("Starship/Common/alert-triangle", Icon16x16, FStyleColors::Warning));
 
+		Style->Set("Icons.Info", new IMAGE_BRUSH_SVG("Starship/Common/Info", Icon16x16));
+		Style->Set("Icons.InfoWithColor", new IMAGE_BRUSH_SVG("Starship/Common/Info", Icon16x16, FStyleColors::Foreground));
+
 		Style->Set("Icons.box-perspective", new IMAGE_BRUSH_SVG("Starship/Common/box-perspective", Icon16x16));
 		Style->Set("Icons.cylinder", new IMAGE_BRUSH_SVG("Starship/Common/cylinder", Icon16x16));
 		Style->Set("Icons.pyramid", new IMAGE_BRUSH_SVG("Starship/Common/pyriamid", Icon16x16));
@@ -358,6 +360,9 @@ TSharedRef<ISlateStyle> FStarshipCoreStyle::Create()
 		Style->Set("Icons.Refresh", new IMAGE_BRUSH_SVG("Starship/Common/refresh", Icon16x16));
 
 		Style->Set("Icons.Star", new IMAGE_BRUSH_SVG("Starship/Common/Favorite", Icon16x16));
+		Style->Set("Icons.Link", new IMAGE_BRUSH_SVG("Starship/Common/Linked", Icon16x16));
+		Style->Set("Icons.Unlink", new IMAGE_BRUSH_SVG("Starship/Common/Unlinked", Icon16x16));
+
 	}
 
 	// Tool panels
