Editgrate Niagara precise motion vector options

diff --git a/Engine/Plugins/FX/Niagara/Shaders/Private/NiagaraMeshVertexFactory.ush b/Engine/Plugins/FX/Niagara/Shaders/Private/NiagaraMeshVertexFactory.ush
--- a/Engine/Plugins/FX/Niagara/Shaders/Private/NiagaraMeshVertexFactory.ush
+++ b/Engine/Plugins/FX/Niagara/Shaders/Private/NiagaraMeshVertexFactory.ush
@@ -16,50 +16,71 @@
 #define USE_PARTICLE_LOCAL_TO_WORLD (NEEDS_PARTICLE_LOCAL_TO_WORLD)
 #define USE_PARTICLE_WORLD_TO_LOCAL (NEEDS_PARTICLE_WORLD_TO_LOCAL)
 
-float3 TransformPosition(float3 InPosition)
+#ifndef ENABLE_PRECISE_MOTION_VECTORS
+#define ENABLE_PRECISE_MOTION_VECTORS 0
+#endif
+
+float3 LocalToWorldPos(float3 LocalSpacePos, float4x4 LocalToWorld)
+{
+	return
+		LocalToWorld[0].xyz * LocalSpacePos.x +
+		LocalToWorld[1].xyz * LocalSpacePos.y +
+		LocalToWorld[2].xyz * LocalSpacePos.z +
+		LocalToWorld[3].xyz;
+}
+
+float3 SimToWorldPos(float3 SimSpacePos, float4x4 LocalToWorld)
 {
 	if (NiagaraMeshVF.bLocalSpace)
 	{
-		return TransformLocalToWorld(InPosition).xyz;
+		return LocalToWorldPos(SimSpacePos, LocalToWorld);
 	}
-	else
+	
+	return SimSpacePos;
+}
+
+float3 LocalToWorldVec(float3 LocalSpaceVec, float4x4 LocalToWorld)
+{
+	const float3 InvScale = Primitive.InvNonUniformScaleAndDeterminantSign.xyz;
+	return
+		InvScale.x * LocalToWorld[0].xyz * LocalSpaceVec.xxx +
+		InvScale.y * LocalToWorld[1].xyz * LocalSpaceVec.yyy +
+		InvScale.z * LocalToWorld[2].xyz * LocalSpaceVec.zzz;
+}
+
+float3 WorldToLocalVec(float3 WorldSpaceVec, float4x4 LocalToWorld)
+{
+	const float3 InvScale = Primitive.InvNonUniformScaleAndDeterminantSign.xyz;
+	float3x3 InvRot =
 	{
-		return InPosition;
-	}
+		InvScale.x * LocalToWorld[0].xyz,
+		InvScale.y * LocalToWorld[1].xyz,
+		InvScale.z * LocalToWorld[2].xyz
+	};
+	InvRot = transpose(InvRot);
+	return mul(WorldSpaceVec, InvRot);
 }
 
-float3 TransformVector(float3 InVector)
+float3 SimToWorldVec(float3 SimSpaceVec, float4x4 LocalToWorld)
 {
 	if (NiagaraMeshVF.bLocalSpace)
 	{
-		return RotateLocalToWorld(InVector);
-	}
-	else
-	{
-		return InVector;
+		return LocalToWorldVec(SimSpaceVec, LocalToWorld);
 	}
+	
+	return SimSpaceVec;
 }
 
-float3 TransformVectorInverse(float3 InVector)
+float3 WorldToSimVec(float3 InVector, float4x4 LocalToWorld)
 {
 	if (NiagaraMeshVF.bLocalSpace)
 	{
-		return RotateWorldToLocal(InVector);
-	}
-	else
-	{
-		return InVector;
+		return WorldToLocalVec(InVector, LocalToWorld);
 	}
+	
+	return InVector;
 }
 
-// Only support previous transform for deferred feature level
-//#define USE_PREVIOUS_TRANSFORM (FEATURE_LEVEL >= FEATURE_LEVEL_SM4)
-#define USE_PREVIOUS_TRANSFORM 0
-
-#if USE_PREVIOUS_TRANSFORM
-	Buffer<float4> PrevTransformBuffer;
-#endif
-
 #define MESH_FACING_DEFAULT 0
 #define MESH_FACING_VELOCITY 1
 #define MESH_FACING_CAMERA_POSITION 2
@@ -214,7 +235,7 @@ struct FVertexFactoryIntermediates
 	float4x4 PrevParticleToWorld;
 
 	float4x4 WorldToParticle;
-	
+
 	/** The instance ID when unsorted/culled, otherwise is an index buffer indirection */
 	uint SortedInstanceId;
 };
@@ -236,21 +257,7 @@ float4x4 GetParticleInvTransform(FVertexFactoryIntermediates Inters)
 
 float4x4 GetParticlePrevTransform(FVertexFactoryIntermediates Inters)
 {
-#if USE_PREVIOUS_TRANSFORM
-	uint TransformBaseIdx = Inters.SortedInstanceId * 3;
-
-	float4x4 PrevTransform4x4 = float4x4(
-		PrevTransformBuffer[TransformBaseIdx + 0],
-		PrevTransformBuffer[TransformBaseIdx + 1],
-		PrevTransformBuffer[TransformBaseIdx + 2],
-		float4(0, 0, 0, 1)
-		);
-
-	//Transform is transposed on CPU and position is packed into W. Transpose again to get the correct transform back.
-	return transpose(PrevTransform4x4);
-#else
 	return Inters.PrevParticleToWorld;
-#endif
 }
 
 /** Converts from vertex factory specific interpolants to a FMaterialPixelParameters, which is used by material inputs. */
@@ -439,7 +446,7 @@ float4 CalcPrevWorldPosition(FVertexFactoryInput Input, FVertexFactoryIntermedia
     return float4(WorldPosition, Input.Position.w);
 }
 
-float3x3 GetFacingMatrix(in float3 ParticlePosition, in float4 ParticleVelocity)
+float3x3 GetFacingMatrix(in float3 ParticlePosition, in float4 ParticleVelocity, in float4x4 LocalToWorld, in float3 WorldCamOrigin, in float3 CamViewForward, in float3 CamViewUp)
 {
 	float3 WorldX = float3(1, 0, 0);
 	float3 WorldZ = float3(0, 0, 1);
@@ -453,27 +460,27 @@ float3x3 GetFacingMatrix(in float3 ParticlePosition, in float4 ParticleVelocity)
 		{
 			if (ParticleVelocity.w > 0.0f)
 			{
-				FacingDir = ParticleVelocity.xyz;				
+				FacingDir = ParticleVelocity.xyz;
 			}
 			else
 			{
-				FacingDir = TransformVectorInverse(WorldZ);
+				FacingDir = WorldToSimVec(WorldZ, LocalToWorld);
 			}
 			break;
 		}
 
 		case MESH_FACING_CAMERA_POSITION:
 		{
-			float3 WorldPosition = TransformPosition(ParticlePosition);
-			float3 CameraVector = ResolvedView.WorldCameraOrigin - WorldPosition;
-			FacingDir = TransformVectorInverse(normalize(CameraVector));
+			float3 WorldPosition = SimToWorldPos(ParticlePosition, LocalToWorld);
+			float3 CameraVector = WorldCamOrigin - WorldPosition;
+			FacingDir = WorldToSimVec(normalize(CameraVector), LocalToWorld);
 			break;
 		}
 
 		default:
 		case MESH_FACING_CAMERA_PLANE:
 		{
-			FacingDir = TransformVectorInverse(-ResolvedView.ViewForward);
+			FacingDir = WorldToSimVec(-CamViewForward, LocalToWorld);
 			break;
 		}
 	}
@@ -489,14 +496,14 @@ float3x3 GetFacingMatrix(in float3 ParticlePosition, in float4 ParticleVelocity)
 		if (bWorldSpaceAxis && bLocalSpace)
 		{
 			// Transform the world-space axis to local space
-			LockedAxis = RotateWorldToLocal(LockedAxis);
+			LockedAxis = WorldToLocalVec(LockedAxis, LocalToWorld);
 		}
 		else if (!bWorldSpaceAxis && !bLocalSpace)
 		{
 			// Transform the local-space axis to world space
-			LockedAxis = RotateLocalToWorld(LockedAxis);
+			LockedAxis = LocalToWorldVec(LockedAxis, LocalToWorld);
 		}
-		
+
 		if (abs(dot(FacingDir, LockedAxis)) > 0.99f)
 		{
 			// The facing dir and locked axis are too similar, choose a reference direction for the facing dir
@@ -515,7 +522,7 @@ float3x3 GetFacingMatrix(in float3 ParticlePosition, in float4 ParticleVelocity)
 		{
 			// Use the camera upwards direction as a reference vector
 			//-TODO: Add ability to remove HMD roll in VR
-			RefVector = ResolvedView.ViewUp;
+			RefVector = CamViewUp;
 		}
 		else
 		{
@@ -524,107 +531,138 @@ float3x3 GetFacingMatrix(in float3 ParticlePosition, in float4 ParticleVelocity)
 			RefVector = abs(DotWorldZ) > 0.99f ? (-sign(DotWorldZ) * WorldX) : WorldZ;
 		}
 
-		// Check to rotate the reference direction to local-space
-		RefVector = TransformVectorInverse(RefVector);
+		// rotate the reference direction to simulation space, if necessary
+		RefVector = WorldToSimVec(RefVector, LocalToWorld);
 
 		// Orthonormalize the look-at rotation and generate a matrix
 		XAxis = FacingDir;
 		YAxis = normalize(cross(RefVector, FacingDir));
 		ZAxis = cross(XAxis, YAxis);
 	}
-	
+
 	return float3x3(XAxis, YAxis, ZAxis);
 }
 
-float3 CalculateCameraOffset(float3 WorldPosition, uint InstanceId)
+float4 GetVelocityDirMag(int VelocityDataOffset, float3 DefaultVelocity, uint InstanceId)
 {
-	float Offset = NiagaraMeshVF.CameraOffsetDataOffset != -1 ? GetFloat(NiagaraMeshVF.CameraOffsetDataOffset, InstanceId) : NiagaraMeshVF.DefaultCamOffset;
-	float3 CameraVector = ResolvedView.WorldCameraOrigin - WorldPosition;
-	float3 CameraDir = normalize(CameraVector);
-	if (NiagaraMeshVF.bLocalSpace)
-	{
-		CameraDir = RotateWorldToLocal(CameraDir);
-	}
-
-	return CameraDir * Offset;
+	float4 DirMag;
+	DirMag.xyz = SafeGetVec3(VelocityDataOffset, InstanceId, DefaultVelocity);
+	DirMag.w = length(DirMag.xyz);
+	
+	DirMag.xyz = DirMag.w > 0.0f ? DirMag.xyz / DirMag.w : float3(0.0f, 0.0f, 0.0f);
+	
+	return DirMag;
 }
 
-void CalculateTransforms(in FVertexFactoryIntermediates Intermediates, out float4x4 LocalToWorld,
-	out float4x4 LocalToWorldNoScale, out float4x4 WorldToLocal)
+float3 CalculateCameraOffset(float3 ParticlePosition, float3 CamWorldPos, uint InstanceId, int CameraOffsetDataOffset, float DefaultCamOffset, float4x4 LocalToWorld)
 {
-	int InstanceId = Intermediates.SortedInstanceId;
-	
-	////////////////////////////////////////////////////////////////////////////////
-	// Translation
-	float3 WorldParticlePosition = float3(0.0f, 0.0f, 0.0f);
-	if (NiagaraMeshVF.PositionDataOffset == -1)
-	{
-		WorldParticlePosition = NiagaraMeshVF.DefaultPos.xyz;
-	}
-	else
-	{
-		WorldParticlePosition = GetVec3(NiagaraMeshVF.PositionDataOffset, InstanceId);
-	}
+	float3 WorldPosition = SimToWorldPos(ParticlePosition, LocalToWorld);
+	float Offset = SafeGetFloat(CameraOffsetDataOffset, InstanceId, DefaultCamOffset);
+	float3 CameraVector = WorldToSimVec(normalize(CamWorldPos - WorldPosition), LocalToWorld);
 
-	// Apply CameraOffset
-	float3 CameraOffsetWorldPosition = WorldParticlePosition;
-	if (NiagaraMeshVF.bLocalSpace)
-	{
-		CameraOffsetWorldPosition = TransformLocalToWorld(WorldParticlePosition).xyz;
-	}
-	WorldParticlePosition += CalculateCameraOffset(CameraOffsetWorldPosition, InstanceId);
+	return CameraVector * Offset;
+}
+
+void GetParticleSRT(out float3 Scale, out float3x3 Rotation, out float3 Translation, in uint InstanceId, in int ScaleDataOffset, float3 DefaultScale,
+	in int RotationDataOffset, in float4 DefaultRotation, in int PositionDataOffset, in float3 DefaultPos, in int CamOffsetDataOffset, in float DefaultCamOffset,
+	in float4 VelocityDirMag, in float3 CamViewPos, in float3 CamViewForward, in float3 CamViewUp, in float4x4 LocalToWorld)
+{
+	float3 ParticlePosition = SafeGetVec3(PositionDataOffset, InstanceId, DefaultPos);
+	float4 ParticleRotation = SafeGetVec4(RotationDataOffset, InstanceId, DefaultRotation);
+	float3 ParticleScale = SafeGetVec3(ScaleDataOffset, InstanceId, DefaultScale);
 
+	////////////////////////////////////////////////////////////////////////////////
+	// Scale
+	Scale = ParticleScale * NiagaraMeshVF.MeshScale;
+	
 	////////////////////////////////////////////////////////////////////////////////
 	// Rotation
-	float3x3 RotationMat = {
-		1, 0, 0,
-		0, 1, 0,
-		0, 0, 1
-	};
-
-	if(NiagaraMeshVF.TransformDataOffset != -1)
-	{
-		float4 Quat = GetVec4(NiagaraMeshVF.TransformDataOffset, InstanceId);
-		RotationMat = NiagaraQuatTo3x3(normalize(Quat));
-	}
-	else
-	{
-		float4 Quat = NiagaraMeshVF.DefaultTransform;
-		RotationMat = NiagaraQuatTo3x3(normalize(Quat));		
-	}
-
+	Rotation = NiagaraQuatTo3x3(normalize(ParticleRotation));
 	if (NiagaraMeshVF.FacingMode != MESH_FACING_DEFAULT)
 	{
 		// Factor in facing rotation
-		float3x3 FacingMat = GetFacingMatrix(WorldParticlePosition, Intermediates.ParticleVelocity);
-		RotationMat = mul(RotationMat, FacingMat);
-	}
-
-	////////////////////////////////////////////////////////////////////////////////
-	// Scale
-	float3 Scale = (float3)1.0f;
-	if(NiagaraMeshVF.ScaleDataOffset != -1)
-	{
-		Scale = GetVec3(NiagaraMeshVF.ScaleDataOffset, InstanceId);		
-	}
-	else
-	{
-		Scale = NiagaraMeshVF.DefaultScale.xyz;
+		float3x3 FacingMat = GetFacingMatrix(ParticlePosition, VelocityDirMag, LocalToWorld, CamViewPos, CamViewForward, CamViewUp);
+		Rotation = mul(Rotation, FacingMat);
 	}
 	
 	////////////////////////////////////////////////////////////////////////////////
 	// Translation
-	float3 Translation = WorldParticlePosition;
+	Translation = ParticlePosition;
+
+	// Apply CameraOffset and WS Pivot offset
+	Translation += CalculateCameraOffset(Translation, CamViewPos, InstanceId, CamOffsetDataOffset, DefaultCamOffset, LocalToWorld);
 	if (!NiagaraMeshVF.bPivotOffsetIsWorldSpace)
 	{
-		Translation += mul(NiagaraMeshVF.PivotOffset * Scale, RotationMat);
+		Translation += mul(NiagaraMeshVF.PivotOffset * ParticleScale, Rotation);
 	}
-	
+}
+
+void CalculateTransforms(inout FVertexFactoryIntermediates Intermediates)
+{
+	int InstanceId = Intermediates.SortedInstanceId;
+
+	float3 Scale;	
+	float3x3 Rotation;
+	float3 Translation;
+	GetParticleSRT(Scale, Rotation, Translation, InstanceId, NiagaraMeshVF.ScaleDataOffset, NiagaraMeshVF.DefaultScale,
+		NiagaraMeshVF.RotationDataOffset, NiagaraMeshVF.DefaultRotation, NiagaraMeshVF.PositionDataOffset, NiagaraMeshVF.DefaultPos.xyz,
+		NiagaraMeshVF.CameraOffsetDataOffset, NiagaraMeshVF.DefaultCamOffset, Intermediates.ParticleVelocity, ResolvedView.WorldCameraOrigin.xyz,
+		ResolvedView.ViewForward, ResolvedView.ViewUp, Primitive.LocalToWorld);
+
 	////////////////////////////////////////////////////////////////////////////////
 	// Compose Transforms
-	LocalToWorld = NiagaraComposeTransformMatrix(Scale * NiagaraMeshVF.MeshScale, RotationMat, Translation);
-	LocalToWorldNoScale = NiagaraComposeTransformMatrix((float3)1.0f, RotationMat, Translation);
-	WorldToLocal = NiagaraComposeInvTransformMatrix(Scale, RotationMat, Translation);
+	Intermediates.ParticleToWorld = NiagaraComposeTransformMatrix(Scale, Rotation, Translation);
+	Intermediates.ParticleToWorldNoScale = NiagaraComposeTransformMatrix((float3)1.0f, Rotation, Translation);
+	Intermediates.WorldToParticle = NiagaraComposeInvTransformMatrix(Scale, Rotation, Translation);	
+	
+#if ENABLE_PRECISE_MOTION_VECTORS
+	// Run through the whole process of generating the components of the transform with previous frame data
+	float4 PrevVelocity = GetVelocityDirMag(NiagaraMeshVF.PrevVelocityDataOffset, NiagaraMeshVF.DefaultPrevVelocity, InstanceId);
+	float4x4 PrevViewToWorld = ResolvedView.PrevViewToTranslatedWorld;
+	float3 PrevScale;
+	float3x3 PrevRotation;
+	float3 PrevTranslation;
+	GetParticleSRT(PrevScale, PrevRotation, PrevTranslation, InstanceId, NiagaraMeshVF.PrevScaleDataOffset, NiagaraMeshVF.DefaultPrevScale,
+		NiagaraMeshVF.PrevRotationDataOffset, NiagaraMeshVF.DefaultPrevRotation, NiagaraMeshVF.PrevPositionDataOffset, NiagaraMeshVF.DefaultPrevPos.xyz,
+		NiagaraMeshVF.PrevCameraOffsetDataOffset, NiagaraMeshVF.DefaultPrevCamOffset, PrevVelocity, ResolvedView.PrevWorldCameraOrigin.xyz,
+		PrevViewToWorld[2].xyz, PrevViewToWorld[1].xyz, Primitive.PreviousLocalToWorld);
+		
+	Intermediates.PrevParticleToWorld = NiagaraComposeTransformMatrix(PrevScale, PrevRotation, PrevTranslation);
+#else
+	// Do a cheaper means of calculating the previous matrix that just extrapolates based on velocity
+	const float3 ParticleVelocity = Intermediates.ParticleVelocity.xyz * Intermediates.ParticleVelocity.w;
+	Intermediates.PrevParticleToWorld = Intermediates.ParticleToWorld;
+	Intermediates.PrevParticleToWorld[3].xyz -= ParticleVelocity * NiagaraMeshVF.DeltaSeconds;
+#endif
+	
+	// Handle local space by multiplying in the primitive transforms. TODO: Optimize behind define.
+	if (NiagaraMeshVF.bLocalSpace)
+	{
+		Intermediates.PrevParticleToWorld = mul(Intermediates.PrevParticleToWorld, Primitive.PreviousLocalToWorld);
+		Intermediates.ParticleToWorld = mul(Intermediates.ParticleToWorld, Primitive.LocalToWorld);
+		Intermediates.WorldToParticle = mul(Primitive.WorldToLocal, Intermediates.WorldToParticle);
+
+		Intermediates.ParticleToWorldNoScale = mul(Intermediates.ParticleToWorldNoScale, Primitive.LocalToWorld);
+		float3 InvScale = Primitive.InvNonUniformScaleAndDeterminantSign.xyz;
+		Intermediates.ParticleToWorldNoScale[0].xyz *= InvScale.x;
+		Intermediates.ParticleToWorldNoScale[1].xyz *= InvScale.y;
+		Intermediates.ParticleToWorldNoScale[2].xyz *= InvScale.z;
+	}
+
+	// If the pivot offset is in world space, apply it here to both the current and previous transforms
+	if (NiagaraMeshVF.bPivotOffsetIsWorldSpace)
+	{
+		Intermediates.ParticleToWorld[3].xyz += NiagaraMeshVF.PivotOffset;
+		Intermediates.PrevParticleToWorld[3].xyz += NiagaraMeshVF.PivotOffset;
+
+		float4x4 OffsetMtx = float4x4(
+			float4(1, 0, 0, 0),
+			float4(0, 1, 0, 0),
+			float4(0, 0, 1, 0),
+			float4(-NiagaraMeshVF.PivotOffset, 1)
+		);
+		Intermediates.WorldToParticle = mul(OffsetMtx, Intermediates.WorldToParticle);
+	}
 }
 
 FVertexFactoryIntermediates GetVertexFactoryIntermediates(FVertexFactoryInput Input)
@@ -638,153 +676,37 @@ FVertexFactoryIntermediates GetVertexFactoryIntermediates(FVertexFactoryInput In
 		Intermediates.SortedInstanceId = SortedIndices[SortedIndicesOffset + Intermediates.SortedInstanceId];
 	}
 
-	if (NiagaraMeshVF.ColorDataOffset == -1)
-	{
-		Intermediates.ParticleColor = NiagaraMeshVF.DefaultColor;
-	}
-	else
-	{
-		Intermediates.ParticleColor = GetVec4(NiagaraMeshVF.ColorDataOffset, Intermediates.SortedInstanceId);
-	}
-
-	float3 ParticleVelocity = (float3)0.0f;
-	if (NiagaraMeshVF.VelocityDataOffset != -1)
-	{
-		ParticleVelocity = float3(GetVec3(NiagaraMeshVF.VelocityDataOffset, Intermediates.SortedInstanceId));
-	}
-	else
-	{
-		ParticleVelocity = NiagaraMeshVF.DefaultVelocity.xyz;
-	}
-
-	// Calculate the particle velocity (NOTE: always do this regardless of NEEDS_PARTICLE_VELOCITY because of velocity facing particles)
-	const float ParticleVelocityMagnitude = length(ParticleVelocity);
-	Intermediates.ParticleVelocity.xyz = (ParticleVelocityMagnitude > 0) ? ParticleVelocity / ParticleVelocityMagnitude : float3(0.0f, 0.0f, 0.0f);
-	Intermediates.ParticleVelocity.w = ParticleVelocityMagnitude;
+	Intermediates.ParticleColor = SafeGetVec4(NiagaraMeshVF.ColorDataOffset, Intermediates.SortedInstanceId, NiagaraMeshVF.DefaultColor);
+	Intermediates.ParticleVelocity = GetVelocityDirMag(NiagaraMeshVF.VelocityDataOffset, NiagaraMeshVF.DefaultVelocity, Intermediates.SortedInstanceId);
 
 #if (DYNAMIC_PARAMETERS_MASK & 1)
-	if (NiagaraMeshVF.MaterialParamDataOffset == -1)
-	{
-		Intermediates.DynamicParameter = NiagaraMeshVF.DefaultDynamicMaterialParameter0;
-	}
-	else
-	{
-		Intermediates.DynamicParameter = GetVec4(NiagaraMeshVF.MaterialParamDataOffset, Intermediates.SortedInstanceId);
-	}
+	Intermediates.DynamicParameter = SafeGetVec4(NiagaraMeshVF.MaterialParamDataOffset, Intermediates.SortedInstanceId, NiagaraMeshVF.DefaultDynamicMaterialParameter0);
 #endif
 #if (DYNAMIC_PARAMETERS_MASK & 2)
-	if (NiagaraMeshVF.MaterialParam1DataOffset == -1)
-	{
-		Intermediates.DynamicParameter1 = NiagaraMeshVF.DefaultDynamicMaterialParameter1;
-	}	
-	else
-	{
-		Intermediates.DynamicParameter1 = GetVec4(NiagaraMeshVF.MaterialParam1DataOffset, Intermediates.SortedInstanceId);
-	}
+	Intermediates.DynamicParameter1 = SafeGetVec4(NiagaraMeshVF.MaterialParam1DataOffset, Intermediates.SortedInstanceId, NiagaraMeshVF.DefaultDynamicMaterialParameter1);
 #endif
 #if (DYNAMIC_PARAMETERS_MASK & 4)
-	if (NiagaraMeshVF.MaterialParam2DataOffset == -1)
-	{
-		Intermediates.DynamicParameter2 = NiagaraMeshVF.DefaultDynamicMaterialParameter2;
-	}
-	else
-	{
-		Intermediates.DynamicParameter2 = GetVec4(NiagaraMeshVF.MaterialParam2DataOffset, Intermediates.SortedInstanceId);
-	}
+	Intermediates.DynamicParameter2 = SafeGetVec4(NiagaraMeshVF.MaterialParam2DataOffset, Intermediates.SortedInstanceId, NiagaraMeshVF.DefaultDynamicMaterialParameter2);	
 #endif
 #if (DYNAMIC_PARAMETERS_MASK & 8)
-	if (NiagaraMeshVF.MaterialParam3DataOffset == -1)
-	{
-		Intermediates.DynamicParameter3 = NiagaraMeshVF.DefaultDynamicMaterialParameter3;
-	}
-	else
-	{
-		Intermediates.DynamicParameter3 = GetVec4(NiagaraMeshVF.MaterialParam3DataOffset, Intermediates.SortedInstanceId);
-	}
+	Intermediates.DynamicParameter3 = SafeGetVec4(NiagaraMeshVF.MaterialParam3DataOffset, Intermediates.SortedInstanceId, NiagaraMeshVF.DefaultDynamicMaterialParameter3);	
 #endif
 
 #if NEEDS_PARTICLE_TIME
-	if (NiagaraMeshVF.NormalizedAgeDataOffset == -1)
-	{
-		Intermediates.RelativeTime = NiagaraMeshVF.DefaultNormAge;
-	}
-	else
-	{
-		Intermediates.RelativeTime = GetFloat(NiagaraMeshVF.NormalizedAgeDataOffset, Intermediates.SortedInstanceId);
-	}
+	Intermediates.RelativeTime = SafeGetFloat(NiagaraMeshVF.NormalizedAgeDataOffset, Intermediates.SortedInstanceId, NiagaraMeshVF.DefaultNormAge);	
 #endif
 
 #if NEEDS_PARTICLE_RANDOM
-	if (NiagaraMeshVF.MaterialRandomDataOffset == -1)
-	{
-		Intermediates.ParticleRandom = NiagaraMeshVF.DefaultMatRandom;
-	}
-	else
-	{
-		Intermediates.ParticleRandom = GetFloat(NiagaraMeshVF.MaterialRandomDataOffset, Intermediates.SortedInstanceId);
-	}
-#endif	
+	Intermediates.ParticleRandom = SafeGetFloat(NiagaraMeshVF.MaterialRandomDataOffset, Intermediates.SortedInstanceId, NiagaraMeshVF.DefaultMatRandom);	
+#endif
 
 	// Swizzle vertex color.
 	Intermediates.VertexColor = Input.VertexColor FCOLOR_COMPONENT_SWIZZLE;
 
-	float4x4 ParticleToWorld;
-	float4x4 ParticleToWorldNoScale;
-	float4x4 WorldToParticle;
-	
-	CalculateTransforms(Intermediates, ParticleToWorld, ParticleToWorldNoScale, WorldToParticle);
-
-	//-TODO: If / when we store previous transform data use them here (providing they are in local space, if not handle below)
-	//if (MeshParticleVF.PrevTransformAvailable > 0)
-	//{
-	//   //-TODO:
-	//}
-	//else
-	{
-		// NOTE: Velocity is in simulation space, so calculate prev transform here, before handling local space
-		Intermediates.PrevParticleToWorld = ParticleToWorld;
-		Intermediates.PrevParticleToWorld[3].xyz -= ParticleVelocity * NiagaraMeshVF.DeltaSeconds;
-	}
-
-	//Handle local space. TODO: Optimize behind define.
-	//Apply local space to local to world, also pretranspose here to avoid it in the vs&ps
-
-	if (NiagaraMeshVF.bLocalSpace)
-	{
-		Intermediates.PrevParticleToWorld = mul(Intermediates.PrevParticleToWorld, Primitive.PreviousLocalToWorld);
-		Intermediates.ParticleToWorld = mul(ParticleToWorld, Primitive.LocalToWorld);
-		Intermediates.WorldToParticle = mul(Primitive.WorldToLocal, WorldToParticle);
-
-		Intermediates.ParticleToWorldNoScale = mul(ParticleToWorldNoScale, Primitive.LocalToWorld);
-		float3 InvScale = Primitive.InvNonUniformScaleAndDeterminantSign.xyz;
-		Intermediates.ParticleToWorldNoScale[0].xyz *= InvScale.x;
-		Intermediates.ParticleToWorldNoScale[1].xyz *= InvScale.y;
-		Intermediates.ParticleToWorldNoScale[2].xyz *= InvScale.z;
-	}
-	else
-	{
-		Intermediates.ParticleToWorld = ParticleToWorld;
-		Intermediates.ParticleToWorldNoScale = ParticleToWorldNoScale;
-		Intermediates.WorldToParticle = WorldToParticle;
-	}
-	
-	if (NiagaraMeshVF.bPivotOffsetIsWorldSpace)
-	{
-		Intermediates.ParticleToWorld[3].xyz += NiagaraMeshVF.PivotOffset;
-		Intermediates.PrevParticleToWorld[3].xyz += NiagaraMeshVF.PivotOffset;
-		
-		float4x4 OffsetMtx = float4x4(
-			float4(1, 0, 0, 0),
-			float4(0, 1, 0, 0),
-			float4(0, 0, 1, 0),
-			float4(-NiagaraMeshVF.PivotOffset, 1)
-		);
-		Intermediates.WorldToParticle = mul(OffsetMtx, Intermediates.WorldToParticle);		
-	}
+	CalculateTransforms(Intermediates);	
 
+	// Particle position and translated world position.
 	Intermediates.ParticlePosition = Intermediates.ParticleToWorld[3].xyz;
-
-	// Translated world position.	
 	Intermediates.ParticleTranslatedWorldPosition = Intermediates.ParticlePosition.xyz + ResolvedView.PreViewTranslation.xyz;
 
 #if NUM_MATERIAL_TEXCOORDS_VERTEX
@@ -796,7 +718,7 @@ FVertexFactoryIntermediates GetVertexFactoryIntermediates(FVertexFactoryInput In
 
 #if USE_PARTICLE_SUBUVS
 	// SubUV.
-	float SubImageIndex = NiagaraMeshVF.SubImageDataOffset != -1? GetFloat(NiagaraMeshVF.SubImageDataOffset, Intermediates.SortedInstanceId) : NiagaraMeshVF.DefaultSubImage;		
+	float SubImageIndex = SafeGetFloat(NiagaraMeshVF.SubImageDataOffset, Intermediates.SortedInstanceId, NiagaraMeshVF.DefaultSubImage);
 	float SubImageLerp = frac(SubImageIndex);
 	float SubImageA = SubImageIndex - SubImageLerp;
 	float SubImageB = SubImageA + 1;
diff --git a/Engine/Plugins/FX/Niagara/Shaders/Private/NiagaraSpriteVertexFactory.ush b/Engine/Plugins/FX/Niagara/Shaders/Private/NiagaraSpriteVertexFactory.ush
--- a/Engine/Plugins/FX/Niagara/Shaders/Private/NiagaraSpriteVertexFactory.ush
+++ b/Engine/Plugins/FX/Niagara/Shaders/Private/NiagaraSpriteVertexFactory.ush
@@ -35,6 +35,10 @@
 /** Blends between FaceCamera and FaceCameraPosition.*/
 #define FACING_MODE_FACECAMERADISTANCEBLEND 4
 
+#ifndef ENABLE_PRECISE_MOTION_VECTORS
+#define ENABLE_PRECISE_MOTION_VECTORS 0
+#endif
+
 struct FVertexFactoryInput
 {
 	float2	TexCoord				: ATTRIBUTE0;
@@ -185,7 +189,6 @@ struct FVertexFactoryIntermediates
 
 	/** The particle worldspace position */
 	float3 Position;
-	float3 OldPosition;
 	
 	float4 SizeRotSubImage;
 	
@@ -439,7 +442,7 @@ float3 SafeNormalizeWithFallback(float3 V, float3 Fallback)
 	return L > 0.000001 ? V / L : Fallback; 
 }
 
-void GetTangents(FVertexFactoryInput Input, FVertexFactoryIntermediates Intermediates, float3 ParticleDirection, out float3 OutRight, out float3 OutUp)
+void GetTangents(FVertexFactoryInput Input, float3 ParticlePosition, float Rotation, float3 CustomFacing, float3 CustomAlignment, float3 ParticleDirection, out float3 OutRight, out float3 OutUp)
 {
 	// Select camera up/right vectors.
 	float3 ResolvedViewRight = lerp(ResolvedView.ViewRight, ResolvedView.HMDViewNoRollRight, NiagaraSpriteVF.RemoveHMDRoll);
@@ -451,13 +454,13 @@ void GetTangents(FVertexFactoryInput Input, FVertexFactoryIntermediates Intermed
 	float3 CameraFacingVector = -ResolvedView.ViewForward;
 
 	// Determine the vector from the particle to the camera and the particle's movement direction.
-	const float3 CameraVector = ResolvedView.WorldCameraOrigin - Intermediates.Position;
+	const float3 CameraVector = ResolvedView.WorldCameraOrigin - ParticlePosition;
 	const float3 CameraDirection = SafeNormalize(CameraVector);
 
-	float3 CustomFacingVector = NiagaraSpriteVFLooseParameters.ParticleFacingMode == FACING_MODE_CUSTOMFACINGVECTOR ? Intermediates.CustomFacingVector : CameraDirection;
+	float3 CustomFacingVector = NiagaraSpriteVFLooseParameters.ParticleFacingMode == FACING_MODE_CUSTOMFACINGVECTOR ? CustomFacing : CameraDirection;
 	
 	// the particle's movement direction.
-	float3 DirectionVector	= lerp(ParticleDirection, Intermediates.CustomAlignmentVector, NiagaraSpriteVFLooseParameters.ParticleAlignmentMode == ALIGNMENT_MODE_CUSTOMALIGNMENT);
+	float3 DirectionVector	= lerp(ParticleDirection, CustomAlignment, NiagaraSpriteVFLooseParameters.ParticleAlignmentMode == ALIGNMENT_MODE_CUSTOMALIGNMENT);
 	
 	// Tangent vectors for any non-default alignment mode.
 	float3 CameraRightAligned	= SafeNormalize(cross(CustomFacingVector, DirectionVector));
@@ -523,7 +526,7 @@ void GetTangents(FVertexFactoryInput Input, FVertexFactoryIntermediates Intermed
 	// Determine the angle of rotation.
 	float SinRotation;
 	float CosRotation;
-	const float SpriteRotation = Intermediates.SizeRotSubImage.z * NiagaraSpriteVF.RotationScale + NiagaraSpriteVF.RotationBias;
+	const float SpriteRotation = Rotation * NiagaraSpriteVF.RotationScale + NiagaraSpriteVF.RotationBias;
 	sincos(SpriteRotation, SinRotation, CosRotation);
 
 	float3 SR = Right * SinRotation;
@@ -567,10 +570,10 @@ float3x3 CalcTangentBasis(FVertexFactoryIntermediates Intermediates)
 	return Result;
 }
 
-void ComputeBillboardUVs(FVertexFactoryInput Input, FVertexFactoryIntermediates Intermediates, out float2 UVForPosition, out float2 UVForTexturing, out float2 UVForTexturingUnflipped)
+void ComputeBillboardUVs(FVertexFactoryInput Input, float2 ParticleSize, float SubImageIndex, out float2 UVForPosition, out float2 UVForTexturing, out float2 UVForTexturingUnflipped)
 {
 	// Encoding the UV flip in the sign of the size data.
-	float2 UVFlip = sign(Intermediates.SizeRotSubImage.xy);
+	float2 UVFlip = sign(ParticleSize);
 
 #if FEATURE_LEVEL >= FEATURE_LEVEL_SM4
 	BRANCH
@@ -580,7 +583,7 @@ void ComputeBillboardUVs(FVertexFactoryInput Input, FVertexFactoryIntermediates
 		const uint VertexIndexMask = NiagaraSpriteVFLooseParameters.NumCutoutVerticesPerFrame == 8 ? 7 : 3; // BVC_FourVertices || BVC_EightVertices
 		uint CutoutVertexIndex = Input.VertexId & VertexIndexMask;
 		float NumFrames = NiagaraSpriteVF.SubImageSize.x * NiagaraSpriteVF.SubImageSize.y;
-		uint SubImageIndexInt = (uint)fmod(Intermediates.SizeRotSubImage.w, NumFrames);
+		uint SubImageIndexInt = (uint)fmod(SubImageIndex, NumFrames);
 
 		FLATTEN
 		if (UVFlip.x * UVFlip.y < 0)
@@ -608,40 +611,42 @@ void ComputeBillboardUVs(FVertexFactoryInput Input, FVertexFactoryIntermediates
 	}
 }
 
-float3 TransformPosition(float3 InPosition)
+float3 LocalToWorldPos(float3 LocalSpacePos, float4x4 LocalToWorld)
 {
-	if (NiagaraSpriteVF.bLocalSpace)
-	{
-		return TransformLocalToWorld(InPosition).xyz;
-	}
-	else
-	{
-		return InPosition;
-	}
+	return
+		LocalToWorld[0].xyz * LocalSpacePos.x +
+		LocalToWorld[1].xyz * LocalSpacePos.y +
+		LocalToWorld[2].xyz * LocalSpacePos.z +
+		LocalToWorld[3].xyz;
 }
 
-float3 TransformPrevPosition(float3 InPosition)
+float3 SimToWorldPos(float3 SimSpacePos, float4x4 LocalToWorld)
 {
 	if (NiagaraSpriteVF.bLocalSpace)
 	{
-		return mul(float4(InPosition,1), Primitive.PreviousLocalToWorld).xyz;
-	}
-	else
-	{
-		return InPosition;
+		return LocalToWorldPos(SimSpacePos, LocalToWorld);
 	}
+	
+	return SimSpacePos;
+}
+
+float3 LocalToWorldVec(float3 LocalSpaceVec, float4x4 LocalToWorld)
+{
+	const float3 InvScale = Primitive.InvNonUniformScaleAndDeterminantSign.xyz;
+	return
+		InvScale.x * LocalToWorld[0].xyz * LocalSpaceVec.xxx +
+		InvScale.y * LocalToWorld[1].xyz * LocalSpaceVec.yyy +
+		InvScale.z * LocalToWorld[2].xyz * LocalSpaceVec.zzz;
 }
 
-float3 TransformVector(float3 InVector)
+float3 SimToWorldVec(float3 SimSpaceVec, float4x4 LocalToWorld)
 {
 	if (NiagaraSpriteVF.bLocalSpace)
 	{
-		return RotateLocalToWorld(InVector);
-	}
-	else
-	{
-		return InVector;
+		return LocalToWorldVec(SimSpaceVec, LocalToWorld);
 	}
+	
+	return SimSpaceVec;
 }
 
 /*********************************************************************************
@@ -650,195 +655,140 @@ float3 TransformVector(float3 InVector)
 
 float3 GetNiagaraParticlePosition(uint InstanceID)
 {
-	if(NiagaraSpriteVF.PositionDataOffset  == -1)
-	{
-		// See FNiagaraConstants for a synchronized value...
-		return TransformPosition(NiagaraSpriteVF.DefaultPos.xyz);
-	}
-	return TransformPosition(GetVec3(NiagaraSpriteVF.PositionDataOffset, InstanceID));
+	return SimToWorldPos(SafeGetVec3(NiagaraSpriteVF.PositionDataOffset, InstanceID, NiagaraSpriteVF.DefaultPos.xyz), Primitive.LocalToWorld);
 }
-float3 GetNiagaraParticlePrevPosition(uint InstanceID)
+
+float3 GetNiagaraPreviousParticlePosition(uint InstanceID)
 {
-	if(NiagaraSpriteVF.PositionDataOffset  == -1)
-	{
-		// See FNiagaraConstants for a synchronized value...
-		return TransformPrevPosition(NiagaraSpriteVF.DefaultPos.xyz);
-	}
-	return TransformPrevPosition(GetVec3(NiagaraSpriteVF.PositionDataOffset, InstanceID));
+	return SimToWorldPos(SafeGetVec3(NiagaraSpriteVF.PrevPositionDataOffset, InstanceID, NiagaraSpriteVF.DefaultPrevPos.xyz), Primitive.PreviousLocalToWorld);
 }
+
 float3 GetNiagaraParticleVelocity(uint InstanceID)
 {
-	if(NiagaraSpriteVF.VelocityDataOffset  == -1)
-	{
-		// See FNiagaraConstants for a synchronized value...
-		return TransformVector(NiagaraSpriteVF.DefaultVelocity.xyz); 
-	}
-	return TransformVector(GetVec3(NiagaraSpriteVF.VelocityDataOffset, InstanceID));
+	return SimToWorldVec(SafeGetVec3(NiagaraSpriteVF.VelocityDataOffset, InstanceID, NiagaraSpriteVF.DefaultVelocity.xyz), Primitive.LocalToWorld);
+}
+
+float3 GetNiagaraPreviousParticleVelocity(uint InstanceID)
+{
+	return SimToWorldVec(SafeGetVec3(NiagaraSpriteVF.PrevVelocityDataOffset, InstanceID, NiagaraSpriteVF.DefaultPrevVelocity.xyz), Primitive.PreviousLocalToWorld);
 }
+
 float2 GetNiagaraParticleSize(uint InstanceID)
 {
-	if(NiagaraSpriteVF.SizeDataOffset  == -1)
-	{
-		// See FNiagaraConstants for a synchronized value...
-		return NiagaraSpriteVF.DefaultSize.xy;
-	}
-	return GetVec2(NiagaraSpriteVF.SizeDataOffset, InstanceID);
+	return SafeGetVec2(NiagaraSpriteVF.SizeDataOffset, InstanceID, NiagaraSpriteVF.DefaultSize.xy);
 }
+
+float2 GetNiagaraPreviousParticleSize(uint InstanceID)
+{
+	return SafeGetVec2(NiagaraSpriteVF.PrevSizeDataOffset, InstanceID, NiagaraSpriteVF.DefaultPrevSize.xy);
+}
+
 float GetNiagaraParticleRotation(uint InstanceID)
 {
-	if(NiagaraSpriteVF.RotationDataOffset  == -1)
-	{
-		// See FNiagaraConstants for a synchronized value...
-		return  (NiagaraSpriteVF.DefaultRotation / 180.0f) * PI;
-	}
-	return (GetFloat(NiagaraSpriteVF.RotationDataOffset, InstanceID) / 180.0f) * PI;
+	// See FNiagaraConstants for a synchronized value...
+	return (SafeGetFloat(NiagaraSpriteVF.RotationDataOffset, InstanceID, NiagaraSpriteVF.DefaultRotation) / 180.0f) * PI;
+}
+
+float GetNiagaraPreviousParticleRotation(uint InstanceID)
+{
+	// See FNiagaraConstants for a synchronized value...
+	return (SafeGetFloat(NiagaraSpriteVF.PrevRotationDataOffset, InstanceID, NiagaraSpriteVF.DefaultPrevRotation) / 180.0f) * PI;
 }
+
 float4 GetNiagaraParticleColor(uint InstanceID)
 {
-	if(NiagaraSpriteVF.ColorDataOffset  == -1)
-	{
-		// See FNiagaraConstants for a synchronized value...
-		return NiagaraSpriteVF.DefaultColor;
-	}
-	return GetVec4(NiagaraSpriteVF.ColorDataOffset, InstanceID);
+	return SafeGetVec4(NiagaraSpriteVF.ColorDataOffset, InstanceID, NiagaraSpriteVF.DefaultColor);
 }
 
 float2 GetNiagaraUVScale(uint InstanceID)
 {
-	if(NiagaraSpriteVF.UVScaleDataOffset  == -1)
-	{
-		// See FNiagaraConstants for a synchronized value...
-		return NiagaraSpriteVF.DefaultUVScale.xy;
-	}
-	return GetVec2(NiagaraSpriteVF.UVScaleDataOffset, InstanceID);	
+	return SafeGetVec2(NiagaraSpriteVF.UVScaleDataOffset, InstanceID, NiagaraSpriteVF.DefaultUVScale.xy);
 }
 
 float2 GetNiagaraPivotOffset(uint InstanceID)
 {
-	if(NiagaraSpriteVF.PivotOffsetDataOffset == -1)
-	{
-		return NiagaraSpriteVF.DefaultPivotOffset;
-	}
-	return GetVec2(NiagaraSpriteVF.PivotOffsetDataOffset, InstanceID);
+	return SafeGetVec2(NiagaraSpriteVF.PivotOffsetDataOffset, InstanceID, NiagaraSpriteVF.DefaultPivotOffset);
+}
+
+float2 GetNiagaraPreviousPivotOffset(uint InstanceID)
+{
+	return SafeGetVec2(NiagaraSpriteVF.PrevPivotOffsetDataOffset, InstanceID, NiagaraSpriteVF.DefaultPrevPivotOffset);
 }
 
 float GetNiagaraParticleRandom(uint InstanceID)
 {
-	if(NiagaraSpriteVF.MaterialRandomDataOffset  == -1)
-	{
-		// See FNiagaraConstants for a synchronized value...
-		return  NiagaraSpriteVF.DefaultMatRandom;
-	}
-	return GetFloat(NiagaraSpriteVF.MaterialRandomDataOffset, InstanceID);	
+	return SafeGetFloat(NiagaraSpriteVF.MaterialRandomDataOffset, InstanceID, NiagaraSpriteVF.DefaultMatRandom);
 }
 
 float GetNiagaraCameraOffset(uint InstanceID)
 {
-	if (NiagaraSpriteVF.CameraOffsetDataOffset == -1)
-	{
-		// See FNiagaraConstants for a synchronized value...
-		return NiagaraSpriteVF.DefaultCamOffset;
-	}
-	return GetFloat(NiagaraSpriteVF.CameraOffsetDataOffset, InstanceID);
+	return SafeGetFloat(NiagaraSpriteVF.CameraOffsetDataOffset, InstanceID, NiagaraSpriteVF.DefaultCamOffset);
+}
+
+float GetNiagaraPreviousCameraOffset(uint InstanceID)
+{
+	return SafeGetFloat(NiagaraSpriteVF.PrevCameraOffsetDataOffset, InstanceID, NiagaraSpriteVF.DefaultPrevCamOffset);
 }
 
 #if USE_PARTICLE_TIME
 float GetNiagaraNormalizedAge(uint InstanceID)
 {
-	if (NiagaraSpriteVF.NormalizedAgeDataOffset == -1)
-	{
-		return NiagaraSpriteVF.DefaultNormAge;
-	}
-	return GetFloat(NiagaraSpriteVF.NormalizedAgeDataOffset, InstanceID);
+	return SafeGetFloat(NiagaraSpriteVF.NormalizedAgeDataOffset, InstanceID, NiagaraSpriteVF.DefaultNormAge);
 }
 #endif
 
 #if (DYNAMIC_PARAMETERS_MASK & 1)
 float4 GetNiagaraParticleDynamicParameters(uint InstanceID)
 {
-	if (NiagaraSpriteVF.MaterialParamDataOffset != -1)
-	{
-		return GetVec4(NiagaraSpriteVF.MaterialParamDataOffset, InstanceID);
-	}
-	else
-	{
-		return NiagaraSpriteVF.DefaultDynamicMaterialParameter0;
-	}
+	return SafeGetVec4(NiagaraSpriteVF.MaterialParamDataOffset, InstanceID, NiagaraSpriteVF.DefaultDynamicMaterialParameter0);
 }
 #endif
 #if (DYNAMIC_PARAMETERS_MASK & 2)
 float4 GetNiagaraParticleDynamicParameters1(uint InstanceID)
 {
-	if (NiagaraSpriteVF.MaterialParam1DataOffset != -1)
-	{
-		return GetVec4(NiagaraSpriteVF.MaterialParam1DataOffset, InstanceID);
-	}
-	else
-	{
-		return  NiagaraSpriteVF.DefaultDynamicMaterialParameter1;
-	}
+	return SafeGetVec4(NiagaraSpriteVF.MaterialParam1DataOffset, InstanceID, NiagaraSpriteVF.DefaultDynamicMaterialParameter1);
 }
 #endif
 #if (DYNAMIC_PARAMETERS_MASK & 4)
 float4 GetNiagaraParticleDynamicParameters2(uint InstanceID)
 {
-	if (NiagaraSpriteVF.MaterialParam2DataOffset != -1)
-	{
-		return GetVec4(NiagaraSpriteVF.MaterialParam2DataOffset, InstanceID);
-	}
-	else
-	{
-		return NiagaraSpriteVF.DefaultDynamicMaterialParameter2;
-	}
+	return SafeGetVec4(NiagaraSpriteVF.MaterialParam2DataOffset, InstanceID, NiagaraSpriteVF.DefaultDynamicMaterialParameter2);
 }
 #endif
 #if (DYNAMIC_PARAMETERS_MASK & 8)
 float4 GetNiagaraParticleDynamicParameters3(uint InstanceID)
 {
-	if (NiagaraSpriteVF.MaterialParam3DataOffset != -1)
-	{
-		return GetVec4(NiagaraSpriteVF.MaterialParam3DataOffset, InstanceID);
-	}
-	else
-	{
-		return NiagaraSpriteVF.DefaultDynamicMaterialParameter3;
-	}
+	return SafeGetVec4(NiagaraSpriteVF.MaterialParam3DataOffset, InstanceID, NiagaraSpriteVF.DefaultDynamicMaterialParameter3);
 }
 #endif
 
-
 float GetNiagaraParticleSubimage(uint InstanceID)
 {
 	// See FNiagaraConstants for a synchronized value...
-	float Val = NiagaraSpriteVF.DefaultSubImage;
-    if(NiagaraSpriteVF.SubimageDataOffset != -1)
-	{
-		Val = GetFloat(NiagaraSpriteVF.SubimageDataOffset, InstanceID);
-	}
-	return Val;
+	return SafeGetFloat(NiagaraSpriteVF.SubimageDataOffset, InstanceID, NiagaraSpriteVF.DefaultSubImage);
 }
 
 float3 GetNiagaraParticleFacingVector(uint InstanceID)
 {
 	// See FNiagaraConstants for a synchronized value...
-	float3 Val = NiagaraSpriteVF.DefaultFacing.xyz;
-    if(NiagaraSpriteVF.FacingDataOffset != -1)
-	{
-		Val = GetVec3(NiagaraSpriteVF.FacingDataOffset, InstanceID);
-	}
-	return TransformVector(Val);
+	return SafeGetVec3(NiagaraSpriteVF.FacingDataOffset, InstanceID, NiagaraSpriteVF.DefaultFacing.xyz);
+}
+
+float3 GetNiagaraPreviousParticleFacingVector(uint InstanceID)
+{
+	// See FNiagaraConstants for a synchronized value...
+	return SafeGetVec3(NiagaraSpriteVF.PrevFacingDataOffset, InstanceID, NiagaraSpriteVF.DefaultPrevFacing.xyz);
 }
 
 float3 GetNiagaraParticleAlignmentVector(uint InstanceID)
 {
-	float3 Val = NiagaraSpriteVF.DefaultAlignment.xyz;
-    if(NiagaraSpriteVF.AlignmentDataOffset != -1)
-	{
-		Val = GetVec3(NiagaraSpriteVF.AlignmentDataOffset, InstanceID);
-	}
-	return TransformVector(Val);
+	return SafeGetVec3(NiagaraSpriteVF.AlignmentDataOffset, InstanceID, NiagaraSpriteVF.DefaultAlignment.xyz);
 }
 
+float3 GetNiagaraPreviousParticleAlignmentVector(uint InstanceID)
+{
+	return SafeGetVec3(NiagaraSpriteVF.PrevAlignmentDataOffset, InstanceID, NiagaraSpriteVF.DefaultPrevAlignment.xyz);
+}
 
 FVertexFactoryIntermediates GetVertexFactoryIntermediates(FVertexFactoryInput Input)
 {
@@ -850,10 +800,21 @@ FVertexFactoryIntermediates GetVertexFactoryIntermediates(FVertexFactoryInput In
 		ParticleID = NiagaraSpriteVFLooseParameters.SortedIndices[NiagaraSpriteVFLooseParameters.SortedIndicesOffset + ParticleID];
 	}
 
-	float2 ParticleSize = GetNiagaraParticleSize(ParticleID);
-	Intermediates.SizeRotSubImage = float4(ParticleSize.x, ParticleSize.y, GetNiagaraParticleRotation(ParticleID), GetNiagaraParticleSubimage(ParticleID));
+	const float3 ParticlePosition = GetNiagaraParticlePosition(ParticleID);
+	const float ParticleRotation = GetNiagaraParticleRotation(ParticleID);
+	const float2 ParticleSize = GetNiagaraParticleSize(ParticleID);
+	const float3 ParticleVelocity = GetNiagaraParticleVelocity(ParticleID);
+	const float SubImageIndex = GetNiagaraParticleSubimage(ParticleID);
+	const float3 CustomFacing = SafeNormalize(GetNiagaraParticleFacingVector(ParticleID));
+	const float3 CustomAlignment = SafeNormalize(GetNiagaraParticleAlignmentVector(ParticleID));	
+	const float2 PivotOffset = GetNiagaraPivotOffset(ParticleID);
+	const float2 UVScale = GetNiagaraUVScale(ParticleID);
+	const float3 CameraOffset = SafeNormalize(ResolvedView.WorldCameraOrigin - ParticlePosition) * GetNiagaraCameraOffset(ParticleID);
+	
+	Intermediates.Position = ParticlePosition;
+	Intermediates.SizeRotSubImage = float4(ParticleSize.x, ParticleSize.y, ParticleRotation, SubImageIndex);
 	Intermediates.Color	= GetNiagaraParticleColor(ParticleID);
-	Intermediates.CustomFacingVector = SafeNormalize(GetNiagaraParticleFacingVector(ParticleID));
+	Intermediates.CustomFacingVector = CustomFacing;
 	Intermediates.CustomAlignmentVector = SafeNormalize(GetNiagaraParticleAlignmentVector(ParticleID));
 
 #if (DYNAMIC_PARAMETERS_MASK & 1)
@@ -869,17 +830,12 @@ FVertexFactoryIntermediates GetVertexFactoryIntermediates(FVertexFactoryInput In
 	Intermediates.DynamicParameter3 = GetNiagaraParticleDynamicParameters3(ParticleID);
 #endif
 
-	// Particle world position, before any sprite aligment.
-	Intermediates.Position = GetNiagaraParticlePosition(ParticleID);
-
-	const float3 ParticleVelocity = GetNiagaraParticleVelocity(ParticleID);
-
 	// Tangents.
 	{
 		const float3 ParticleDirection = SafeNormalizeWithFallback(ParticleVelocity, float3(0,0,1));
 		float3 Right,Up;
 
-		GetTangents(Input, Intermediates, ParticleDirection, Right, Up);
+		GetTangents(Input, ParticlePosition, ParticleRotation, CustomFacing, CustomAlignment, ParticleDirection, Right, Up);
 		Intermediates.TangentUp = Up;
 		Intermediates.TangentRight = Right;
 	}
@@ -887,38 +843,60 @@ FVertexFactoryIntermediates GetVertexFactoryIntermediates(FVertexFactoryInput In
 	float2 UVForPosition;
 	float2 UVForTexturing;
 	float2 UVForTexturingUnflipped;
-	ComputeBillboardUVs(Input, Intermediates, UVForPosition, UVForTexturing, UVForTexturingUnflipped);
+	ComputeBillboardUVs(Input, ParticleSize, SubImageIndex, UVForPosition, UVForTexturing, UVForTexturingUnflipped);
 
-	const float2 UVScale = GetNiagaraUVScale(ParticleID);
 	UVForTexturing.xy = float2(0.5f, 0.5f) + (UVForTexturing.xy - float2(0.5f, 0.5f) ) *  UVScale.xy; 
 	const float2 Size = abs(Intermediates.SizeRotSubImage.xy);
 	
-	const float2 PivotOffset = GetNiagaraPivotOffset(ParticleID);
 
-	// Vertex position and prev vertex position.
-	{
-		// Do we need the hack that forces particles at the camera origin to be *really* behind the camera origin?
-		const float3 CameraOffset = SafeNormalize(ResolvedView.WorldCameraOrigin - Intermediates.Position) * GetNiagaraCameraOffset(ParticleID);
-
-		const float2x3 Tangents = float2x3(Intermediates.TangentRight, Intermediates.TangentUp);
-		const float3 VertexOffset = CameraOffset + mul(Size * (UVForPosition - PivotOffset), Tangents);
+	// Vertex position
+	const float2x3 Tangents = float2x3(Intermediates.TangentRight, Intermediates.TangentUp);
+	const float3 VertexOffset = CameraOffset + mul(Size * (UVForPosition - PivotOffset), Tangents);
 
-		Intermediates.VertexWorldPosition = Intermediates.Position + VertexOffset;
-		Intermediates.PreviousVertexWorldPosition = GetNiagaraParticlePrevPosition(ParticleID) + VertexOffset + ParticleVelocity * NiagaraSpriteVF.DeltaSeconds;
+	Intermediates.VertexWorldPosition = ParticlePosition + VertexOffset;
+		
+	// Prev vertex position (approximated from current velocity unless ENABLE_PRECISE_MOTION_VECTORS)
+#if ENABLE_PRECISE_MOTION_VECTORS
+	{
+		const float3 PrevParticlePosition = GetNiagaraPreviousParticlePosition(ParticleID);
+		const float PrevParticleRotation = GetNiagaraPreviousParticleRotation(ParticleID);
+		const float2 PrevParticleSize = GetNiagaraPreviousParticleSize(ParticleID);
+		const float3 PrevParticleVelocity = GetNiagaraPreviousParticleVelocity(ParticleID);		
+		const float3 PrevCustomFacing = SafeNormalize(GetNiagaraPreviousParticleFacingVector(ParticleID));
+		const float3 PrevCustomAlignment = SafeNormalize(GetNiagaraPreviousParticleAlignmentVector(ParticleID));
+		const float2 PrevPivotOffset = GetNiagaraPreviousPivotOffset(ParticleID);
+		const float3 PrevCameraOffset = SafeNormalize(ResolvedView.PrevWorldCameraOrigin - PrevParticlePosition) * GetNiagaraPreviousCameraOffset(ParticleID);
+		const float3 PrevParticleDirection = SafeNormalizeWithFallback(PrevParticleVelocity, float3(0,0,1));
+			
+		float3 PrevRight, PrevUp;
+		GetTangents(Input, PrevParticlePosition, PrevParticleRotation, PrevCustomFacing, PrevCustomAlignment, PrevParticleDirection, PrevRight, PrevUp);
+		
+		// NOTE: Don't use previous SubImageIndex here because cutout verts don't correlate between sub images
+		float2 PrevUVForPosition;
+		float2 PrevUVForTexturing;
+		float2 PrevUVForTexturingUnflipped;
+		ComputeBillboardUVs(Input, PrevParticleSize, SubImageIndex, PrevUVForPosition, PrevUVForTexturing, PrevUVForTexturingUnflipped);
+			
+		const float2x3 PrevTangents = float2x3(PrevRight, PrevUp);
+		const float3 PrevVertexOffset = PrevCameraOffset + mul(abs(PrevParticleSize) * (PrevUVForPosition - PivotOffset), PrevTangents);
+		
+		Intermediates.PreviousVertexWorldPosition = PrevParticlePosition + PrevVertexOffset;
+	}
+#else
+	Intermediates.PreviousVertexWorldPosition = Intermediates.VertexWorldPosition - ParticleVelocity * NiagaraSpriteVF.DeltaSeconds;
+#endif
 
 #if USE_PARTICLE_POSITION
-		// Do we really want the camera offset here ? What about when GENERATE_SPHERICAL_PARTICLE_NORMALS
-		Intermediates.TranslatedWorldPositionAndSize.xyz = Intermediates.Position + CameraOffset + ResolvedView.PreViewTranslation;
-		Intermediates.TranslatedWorldPositionAndSize.w = .5f * min(Size.x, Size.y);
+	// Do we really want the camera offset here ? What about when GENERATE_SPHERICAL_PARTICLE_NORMALS
+	Intermediates.TranslatedWorldPositionAndSize.xyz = Intermediates.Position + CameraOffset + ResolvedView.PreViewTranslation;
+	Intermediates.TranslatedWorldPositionAndSize.w = .5f * min(Size.x, Size.y);
 #if GENERATE_SPHERICAL_PARTICLE_NORMALS
-		// Recenter the particle position for GENERATE_SPHERICAL_PARTICLE_NORMALS to use the right center.
-		Intermediates.TranslatedWorldPositionAndSize.xyz += mul(Size * (float2(.5, .5) - PivotOffset), Tangents);
+	// Recenter the particle position for GENERATE_SPHERICAL_PARTICLE_NORMALS to use the right center.
+	Intermediates.TranslatedWorldPositionAndSize.xyz += mul(Size * (float2(.5, .5) - PivotOffset), Tangents);
 #endif // GENERATE_SPHERICAL_PARTICLE_NORMALS
 #endif // USE_PARTICLE_POSITION
-	}
 	
 	// SubUV.
-	float SubImageIndex = Intermediates.SizeRotSubImage.w;
 	float SubImageLerp = frac(SubImageIndex);
 	float SubImageA = SubImageIndex - SubImageLerp;
 	float SubImageB = SubImageA + 1;
diff --git a/Engine/Plugins/FX/Niagara/Shaders/Private/NiagaraVFParticleAccess.usf b/Engine/Plugins/FX/Niagara/Shaders/Private/NiagaraVFParticleAccess.usf
--- a/Engine/Plugins/FX/Niagara/Shaders/Private/NiagaraVFParticleAccess.usf
+++ b/Engine/Plugins/FX/Niagara/Shaders/Private/NiagaraVFParticleAccess.usf
@@ -99,3 +99,38 @@ float4 GetVec4(int RegisterIdx, uint InstanceID)
 	}
 }
 
+float SafeGetFloat(int RegisterIdx, uint InstanceID, float Default)
+{
+	// NOTE: Branch to prevent unnecessary loads
+	BRANCH
+	if (RegisterIdx == -1)
+		return Default;
+	return GetFloat(RegisterIdx, InstanceID);
+}
+
+float2 SafeGetVec2(int RegisterIdx, uint InstanceID, float2 Default)
+{
+	// NOTE: Branch to prevent unnecessary loads
+	BRANCH
+	if (RegisterIdx == -1)
+		return Default;
+	return GetVec2(RegisterIdx, InstanceID);
+}
+
+float3 SafeGetVec3(int RegisterIdx, uint InstanceID, float3 Default)
+{
+	// NOTE: Branch to prevent unnecessary loads
+	BRANCH
+	if (RegisterIdx == -1)
+		return Default;
+	return GetVec3(RegisterIdx, InstanceID);
+}
+
+float4 SafeGetVec4(int RegisterIdx, uint InstanceID, float4 Default)
+{
+	// NOTE: Branch to prevent unnecessary loads
+	BRANCH
+	if (RegisterIdx == -1)
+		return Default;
+	return GetVec4(RegisterIdx, InstanceID);
+}
\ No newline at end of file
diff --git a/Engine/Plugins/FX/Niagara/Source/Niagara/Classes/NiagaraScript.h b/Engine/Plugins/FX/Niagara/Source/Niagara/Classes/NiagaraScript.h
--- a/Engine/Plugins/FX/Niagara/Source/Niagara/Classes/NiagaraScript.h
+++ b/Engine/Plugins/FX/Niagara/Source/Niagara/Classes/NiagaraScript.h
@@ -167,6 +167,11 @@ public:
 	/** Configuration options*/
 	UPROPERTY()
 	TArray<FString> AdditionalDefines;
+
+	UPROPERTY()
+	TArray<FNiagaraVariableBase> AdditionalVariables;
+
+	TArray<FString> GetAdditionalVariableStrings();
 #endif
 
 	/** Whether or not we need to bake Rapid Iteration params. True to keep params, false to bake.*/
diff --git a/Engine/Plugins/FX/Niagara/Source/Niagara/Classes/NiagaraScriptSourceBase.h b/Engine/Plugins/FX/Niagara/Source/Niagara/Classes/NiagaraScriptSourceBase.h
--- a/Engine/Plugins/FX/Niagara/Source/Niagara/Classes/NiagaraScriptSourceBase.h
+++ b/Engine/Plugins/FX/Niagara/Source/Niagara/Classes/NiagaraScriptSourceBase.h
@@ -68,6 +68,7 @@ public:
 	FString Name;
 	int32 TargetUsageBitmask;
 	TArray<FString> AdditionalDefines;
+	TArray<FNiagaraVariableBase> AdditionalVariables;
 };
 
 struct FNiagaraParameterStore;
diff --git a/Engine/Plugins/FX/Niagara/Source/Niagara/Private/NiagaraCommon.cpp b/Engine/Plugins/FX/Niagara/Source/Niagara/Private/NiagaraCommon.cpp
--- a/Engine/Plugins/FX/Niagara/Source/Niagara/Private/NiagaraCommon.cpp
+++ b/Engine/Plugins/FX/Niagara/Source/Niagara/Private/NiagaraCommon.cpp
@@ -472,7 +472,43 @@ void  FNiagaraVariableAttributeBinding::SetValue(const FName& InValue, const UNi
 	CacheValues(InEmitter, InSourceMode);	
 }
 
-void FNiagaraVariableAttributeBinding::Setup(const FNiagaraVariableBase& InRootVar, const FNiagaraVariableBase& InDataSetVar, const FNiagaraVariable& InDefaultValue, ENiagaraRendererSourceDataMode InSourceMode)
+void FNiagaraVariableAttributeBinding::SetAsPreviousValue(const FNiagaraVariableBase& Src, const UNiagaraEmitter* InEmitter, ENiagaraRendererSourceDataMode InSourceMode)
+{
+	static const FString PreviousNamespace = FNiagaraConstants::PreviousNamespace.ToString();
+
+	RootVariable = ParamMapVariable = DataSetVariable = Src;
+
+	// Split out the name and it's namespace
+	TArray<FString> SplitName;
+	Src.GetName().ToString().ParseIntoArray(SplitName, TEXT("."));
+
+	// If the name already contains a "Previous" in the name, just go with that
+	bool bIsPrev = false;
+	for (const FString& Split : SplitName)
+	{
+		if (Split.Equals(PreviousNamespace, ESearchCase::IgnoreCase))
+		{
+			bIsPrev = true;
+			break;
+		}
+	}
+
+	if (bIsPrev)
+	{
+		SetValue(Src.GetName(), InEmitter, InSourceMode);
+	}
+	else
+	{
+		// insert "Previous" into the name, after the first namespace. Or the beginning, if it has none
+		const int32 Location = SplitName.Num() > 1 ? 1 : 0;
+		SplitName.Insert(PreviousNamespace, Location);
+
+		FString PrevName = FString::Join(SplitName, TEXT("."));
+		SetValue(*PrevName, InEmitter, InSourceMode);
+	}
+}
+
+void FNiagaraVariableAttributeBinding::Setup(const FNiagaraVariableBase& InRootVar, const FNiagaraVariable& InDefaultValue, ENiagaraRendererSourceDataMode InSourceMode)
 {
 	RootVariable = InRootVar;
 	if (InDefaultValue.IsDataAllocated() && InDefaultValue.GetType() == InRootVar.GetType())
@@ -495,10 +531,6 @@ FString FNiagaraVariableAttributeBinding::GetDefaultValueString() const
 	return DefaultValueStr;
 }
 
-const FName& FNiagaraVariableAttributeBinding::GetName(ENiagaraRendererSourceDataMode InSourceMode) const
-{
-	return CachedDisplayName;
-}
 #endif
 
 void FNiagaraVariableAttributeBinding::PostLoad(ENiagaraRendererSourceDataMode InSourceMode)
diff --git a/Engine/Plugins/FX/Niagara/Source/Niagara/Private/NiagaraComponent.cpp b/Engine/Plugins/FX/Niagara/Source/Niagara/Private/NiagaraComponent.cpp
--- a/Engine/Plugins/FX/Niagara/Source/Niagara/Private/NiagaraComponent.cpp
+++ b/Engine/Plugins/FX/Niagara/Source/Niagara/Private/NiagaraComponent.cpp
@@ -288,7 +288,10 @@ void FNiagaraSceneProxy::CreateRenderers(const UNiagaraComponent* Component)
 					if (Properties->GetIsActive() && EmitterInst->GetData().IsInitialized() && !EmitterInst->IsDisabled())
 					{
 						NewRenderer = Properties->CreateEmitterRenderer(FeatureLevel, &EmitterInst.Get(), Component);
-						bAlwaysHasVelocity |= Properties->bMotionBlurEnabled;
+						if (Properties->MotionVectorSetting != ENiagaraRendererMotionVectorSetting::Disable)
+						{
+							bAlwaysHasVelocity = true;
+						}
 					}
 					EmitterRenderers.Add(NewRenderer);
 				}
diff --git a/Engine/Plugins/FX/Niagara/Source/Niagara/Private/NiagaraComponentRendererProperties.cpp b/Engine/Plugins/FX/Niagara/Source/Niagara/Private/NiagaraComponentRendererProperties.cpp
--- a/Engine/Plugins/FX/Niagara/Source/Niagara/Private/NiagaraComponentRendererProperties.cpp
+++ b/Engine/Plugins/FX/Niagara/Source/Niagara/Private/NiagaraComponentRendererProperties.cpp
@@ -419,12 +419,12 @@ void UNiagaraComponentRendererProperties::PostEditChangeProperty(struct FPropert
 			CreateTemplateComponent();
 
 			FNiagaraComponentPropertyBinding PositionBinding;
-			PositionBinding.AttributeBinding.Setup(SYS_PARAM_PARTICLES_POSITION,  FNiagaraConstants::GetAttributeAsParticleDataSetKey(SYS_PARAM_PARTICLES_POSITION), SYS_PARAM_PARTICLES_POSITION);
+			PositionBinding.AttributeBinding.Setup(SYS_PARAM_PARTICLES_POSITION, SYS_PARAM_PARTICLES_POSITION);
 			PositionBinding.PropertyName = FName("RelativeLocation");
 			PropertyBindings.Add(PositionBinding);
 
 			FNiagaraComponentPropertyBinding ScaleBinding;
-			ScaleBinding.AttributeBinding.Setup(SYS_PARAM_PARTICLES_SCALE, FNiagaraConstants::GetAttributeAsParticleDataSetKey(SYS_PARAM_PARTICLES_SCALE), SYS_PARAM_PARTICLES_SCALE);
+			ScaleBinding.AttributeBinding.Setup(SYS_PARAM_PARTICLES_SCALE, SYS_PARAM_PARTICLES_SCALE);
 			ScaleBinding.PropertyName = FName("RelativeScale3D");
 			PropertyBindings.Add(ScaleBinding);
 		}
diff --git a/Engine/Plugins/FX/Niagara/Source/Niagara/Private/NiagaraConstants.cpp b/Engine/Plugins/FX/Niagara/Source/Niagara/Private/NiagaraConstants.cpp
--- a/Engine/Plugins/FX/Niagara/Source/Niagara/Private/NiagaraConstants.cpp
+++ b/Engine/Plugins/FX/Niagara/Source/Niagara/Private/NiagaraConstants.cpp
@@ -721,8 +721,7 @@ FNiagaraVariableAttributeBinding FNiagaraConstants::GetAttributeDefaultBinding(c
 	}
 
 	FNiagaraVariableAttributeBinding Binding;
-	const FNiagaraVariable* FoundVar = AttrDataSetKeyMap.Find(InVar);
-	Binding.Setup(InVar, FoundVar ? *FoundVar : InVar, GetAttributeWithDefaultValue(InVar));
+	Binding.Setup(InVar, GetAttributeWithDefaultValue(InVar));
 	return Binding;
 }
 
diff --git a/Engine/Plugins/FX/Niagara/Source/Niagara/Private/NiagaraMeshRendererProperties.cpp b/Engine/Plugins/FX/Niagara/Source/Niagara/Private/NiagaraMeshRendererProperties.cpp
--- a/Engine/Plugins/FX/Niagara/Source/Niagara/Private/NiagaraMeshRendererProperties.cpp
+++ b/Engine/Plugins/FX/Niagara/Source/Niagara/Private/NiagaraMeshRendererProperties.cpp
@@ -75,7 +75,7 @@ UNiagaraMeshRendererProperties::UNiagaraMeshRendererProperties()
 	NumFlipbookFrames = 1;
 #endif
 
-	AttributeBindings.Reserve(16);
+	AttributeBindings.Reserve(21);
 	AttributeBindings.Add(&PositionBinding);
 	AttributeBindings.Add(&VelocityBinding);
 	AttributeBindings.Add(&ColorBinding);
@@ -91,6 +91,13 @@ UNiagaraMeshRendererProperties::UNiagaraMeshRendererProperties()
 	AttributeBindings.Add(&DynamicMaterial3Binding);
 	AttributeBindings.Add(&CameraOffsetBinding);
 
+	// These are associated with attributes in the VF layout only if bGenerateAccurateMotionVectors is true
+	AttributeBindings.Add(&PrevPositionBinding);
+	AttributeBindings.Add(&PrevScaleBinding);
+	AttributeBindings.Add(&PrevMeshOrientationBinding);
+	AttributeBindings.Add(&PrevCameraOffsetBinding);
+	AttributeBindings.Add(&PrevVelocityBinding);
+
 	// The remaining bindings are not associated with attributes in the VF layout
 	AttributeBindings.Add(&RendererVisibilityTagBinding);
 	AttributeBindings.Add(&MeshIndexBinding);
@@ -155,14 +162,14 @@ FNiagaraBoundsCalculator* UNiagaraMeshRendererProperties::CreateBoundsCalculator
 					Offset = Offset.ComponentMax(MeshProperties.PivotOffset.GetAbs());
 				}
 				break;
-			}			
+			}
 
 			LocalBounds += MeshBounds;
 		}
 	}
 
 	if (LocalBounds.IsValid)
-	{		
+	{
 		// Take the bounding center into account with the extents, as it may not be at the origin
 		const FVector Extents = LocalBounds.Max.GetAbs().ComponentMax(LocalBounds.Min.GetAbs());
 		FNiagaraBoundsCalculatorHelper<false, true, false>* BoundsCalculator
@@ -247,10 +254,23 @@ void UNiagaraMeshRendererProperties::InitBindings()
 		//Default custom sorting to age
 		CustomSortingBinding = FNiagaraConstants::GetAttributeDefaultBinding(SYS_PARAM_PARTICLES_NORMALIZED_AGE);
 	}
+
+	SetPreviousBindings(nullptr, SourceMode);
+}
+
+void UNiagaraMeshRendererProperties::SetPreviousBindings(const UNiagaraEmitter* SrcEmitter, ENiagaraRendererSourceDataMode InSourceMode)
+{
+	PrevPositionBinding.SetAsPreviousValue(PositionBinding.GetParamMapBindableVariable(), SrcEmitter, InSourceMode);
+	PrevScaleBinding.SetAsPreviousValue(ScaleBinding.GetParamMapBindableVariable(), SrcEmitter, InSourceMode);
+	PrevMeshOrientationBinding.SetAsPreviousValue(MeshOrientationBinding.GetParamMapBindableVariable(), SrcEmitter, InSourceMode);
+	PrevCameraOffsetBinding.SetAsPreviousValue(CameraOffsetBinding.GetParamMapBindableVariable(), SrcEmitter, InSourceMode);
+	PrevVelocityBinding.SetAsPreviousValue(VelocityBinding.GetParamMapBindableVariable(), SrcEmitter, InSourceMode);
 }
 
 void UNiagaraMeshRendererProperties::UpdateSourceModeDerivates(ENiagaraRendererSourceDataMode InSourceMode, bool bFromPropertyEdit)
 {
+	Super::UpdateSourceModeDerivates(InSourceMode, bFromPropertyEdit);
+
 	UNiagaraEmitter* SrcEmitter = GetTypedOuter<UNiagaraEmitter>();
 	if (SrcEmitter)
 	{
@@ -258,9 +278,9 @@ void UNiagaraMeshRendererProperties::UpdateSourceModeDerivates(ENiagaraRendererS
 		{
 			MaterialParamBinding.CacheValues(SrcEmitter);
 		}
-	}
 
-	Super::UpdateSourceModeDerivates(InSourceMode, bFromPropertyEdit);
+		SetPreviousBindings(SrcEmitter, InSourceMode);
+	}
 }
 
 void UNiagaraMeshRendererProperties::CacheFromCompiledData(const FNiagaraDataSetCompiledData* CompiledData)
@@ -268,12 +288,13 @@ void UNiagaraMeshRendererProperties::CacheFromCompiledData(const FNiagaraDataSet
 	UpdateSourceModeDerivates(SourceMode);
 
 	// Initialize layout
-	RendererLayoutWithCustomSorting.Initialize(ENiagaraMeshVFLayout::Num);
+	const int32 NumLayoutVars = NeedsPreciseMotionVectors() ? ENiagaraMeshVFLayout::Num_Max : ENiagaraMeshVFLayout::Num_Default;
+	RendererLayoutWithCustomSorting.Initialize(NumLayoutVars);
 	RendererLayoutWithCustomSorting.SetVariableFromBinding(CompiledData, PositionBinding, ENiagaraMeshVFLayout::Position);
 	RendererLayoutWithCustomSorting.SetVariableFromBinding(CompiledData, VelocityBinding, ENiagaraMeshVFLayout::Velocity);
 	RendererLayoutWithCustomSorting.SetVariableFromBinding(CompiledData, ColorBinding, ENiagaraMeshVFLayout::Color);
 	RendererLayoutWithCustomSorting.SetVariableFromBinding(CompiledData, ScaleBinding, ENiagaraMeshVFLayout::Scale);
-	RendererLayoutWithCustomSorting.SetVariableFromBinding(CompiledData, MeshOrientationBinding, ENiagaraMeshVFLayout::Transform);
+	RendererLayoutWithCustomSorting.SetVariableFromBinding(CompiledData, MeshOrientationBinding, ENiagaraMeshVFLayout::Rotation);
 	RendererLayoutWithCustomSorting.SetVariableFromBinding(CompiledData, MaterialRandomBinding, ENiagaraMeshVFLayout::MaterialRandom);
 	RendererLayoutWithCustomSorting.SetVariableFromBinding(CompiledData, NormalizedAgeBinding, ENiagaraMeshVFLayout::NormalizedAge);
 	RendererLayoutWithCustomSorting.SetVariableFromBinding(CompiledData, CustomSortingBinding, ENiagaraMeshVFLayout::CustomSorting);
@@ -283,14 +304,22 @@ void UNiagaraMeshRendererProperties::CacheFromCompiledData(const FNiagaraDataSet
 	MaterialParamValidMask |= RendererLayoutWithCustomSorting.SetVariableFromBinding(CompiledData, DynamicMaterial1Binding, ENiagaraMeshVFLayout::DynamicParam1) ? 0x2 : 0;
 	MaterialParamValidMask |= RendererLayoutWithCustomSorting.SetVariableFromBinding(CompiledData, DynamicMaterial2Binding, ENiagaraMeshVFLayout::DynamicParam2) ? 0x4 : 0;
 	MaterialParamValidMask |= RendererLayoutWithCustomSorting.SetVariableFromBinding(CompiledData, DynamicMaterial3Binding, ENiagaraMeshVFLayout::DynamicParam3) ? 0x8 : 0;
+	if (NeedsPreciseMotionVectors())
+	{
+		RendererLayoutWithCustomSorting.SetVariableFromBinding(CompiledData, PrevPositionBinding, ENiagaraMeshVFLayout::PrevPosition);
+		RendererLayoutWithCustomSorting.SetVariableFromBinding(CompiledData, PrevScaleBinding, ENiagaraMeshVFLayout::PrevScale);
+		RendererLayoutWithCustomSorting.SetVariableFromBinding(CompiledData, PrevMeshOrientationBinding, ENiagaraMeshVFLayout::PrevRotation);
+		RendererLayoutWithCustomSorting.SetVariableFromBinding(CompiledData, PrevCameraOffsetBinding, ENiagaraMeshVFLayout::PrevCameraOffset);
+		RendererLayoutWithCustomSorting.SetVariableFromBinding(CompiledData, PrevVelocityBinding, ENiagaraMeshVFLayout::PrevVelocity);
+	}
 	RendererLayoutWithCustomSorting.Finalize();
 
-	RendererLayoutWithoutCustomSorting.Initialize(ENiagaraMeshVFLayout::Num);
+	RendererLayoutWithoutCustomSorting.Initialize(NumLayoutVars);
 	RendererLayoutWithoutCustomSorting.SetVariableFromBinding(CompiledData, PositionBinding, ENiagaraMeshVFLayout::Position);
 	RendererLayoutWithoutCustomSorting.SetVariableFromBinding(CompiledData, VelocityBinding, ENiagaraMeshVFLayout::Velocity);
 	RendererLayoutWithoutCustomSorting.SetVariableFromBinding(CompiledData, ColorBinding, ENiagaraMeshVFLayout::Color);
 	RendererLayoutWithoutCustomSorting.SetVariableFromBinding(CompiledData, ScaleBinding, ENiagaraMeshVFLayout::Scale);
-	RendererLayoutWithoutCustomSorting.SetVariableFromBinding(CompiledData, MeshOrientationBinding, ENiagaraMeshVFLayout::Transform);
+	RendererLayoutWithoutCustomSorting.SetVariableFromBinding(CompiledData, MeshOrientationBinding, ENiagaraMeshVFLayout::Rotation);
 	RendererLayoutWithoutCustomSorting.SetVariableFromBinding(CompiledData, MaterialRandomBinding, ENiagaraMeshVFLayout::MaterialRandom);
 	RendererLayoutWithoutCustomSorting.SetVariableFromBinding(CompiledData, NormalizedAgeBinding, ENiagaraMeshVFLayout::NormalizedAge);
 	RendererLayoutWithoutCustomSorting.SetVariableFromBinding(CompiledData, SubImageIndexBinding, ENiagaraMeshVFLayout::SubImage);
@@ -299,6 +328,14 @@ void UNiagaraMeshRendererProperties::CacheFromCompiledData(const FNiagaraDataSet
 	MaterialParamValidMask |= RendererLayoutWithoutCustomSorting.SetVariableFromBinding(CompiledData, DynamicMaterial1Binding, ENiagaraMeshVFLayout::DynamicParam1) ? 0x2 : 0;
 	MaterialParamValidMask |= RendererLayoutWithoutCustomSorting.SetVariableFromBinding(CompiledData, DynamicMaterial2Binding, ENiagaraMeshVFLayout::DynamicParam2) ? 0x4 : 0;
 	MaterialParamValidMask |= RendererLayoutWithoutCustomSorting.SetVariableFromBinding(CompiledData, DynamicMaterial3Binding, ENiagaraMeshVFLayout::DynamicParam3) ? 0x8 : 0;
+	if (NeedsPreciseMotionVectors())
+	{
+		RendererLayoutWithoutCustomSorting.SetVariableFromBinding(CompiledData, PrevPositionBinding, ENiagaraMeshVFLayout::PrevPosition);
+		RendererLayoutWithoutCustomSorting.SetVariableFromBinding(CompiledData, PrevScaleBinding, ENiagaraMeshVFLayout::PrevScale);
+		RendererLayoutWithoutCustomSorting.SetVariableFromBinding(CompiledData, PrevMeshOrientationBinding, ENiagaraMeshVFLayout::PrevRotation);
+		RendererLayoutWithoutCustomSorting.SetVariableFromBinding(CompiledData, PrevCameraOffsetBinding, ENiagaraMeshVFLayout::PrevCameraOffset);
+		RendererLayoutWithoutCustomSorting.SetVariableFromBinding(CompiledData, PrevVelocityBinding, ENiagaraMeshVFLayout::PrevVelocity);
+	}
 	RendererLayoutWithoutCustomSorting.Finalize();
 }
 
@@ -319,7 +356,7 @@ bool UNiagaraMeshRendererProperties::IsSupportedVariableForBinding(const FNiagar
 void UNiagaraMeshRendererProperties::GetUsedMeshMaterials(int32 MeshIndex, const FNiagaraEmitterInstance* Emitter, TArray<UMaterialInterface*>& OutMaterials) const
 {
 	check(Meshes.IsValidIndex(MeshIndex));
-	
+
 	const UStaticMesh* Mesh = Meshes[MeshIndex].Mesh;
 	check(Mesh);
 
@@ -328,7 +365,7 @@ void UNiagaraMeshRendererProperties::GetUsedMeshMaterials(int32 MeshIndex, const
 
 	OutMaterials.SetNum(0, false);
 
-	// Retrieve a list of materials whose indices match up with the mesh, and only fill it in with materials that are used by any section of any LOD	
+	// Retrieve a list of materials whose indices match up with the mesh, and only fill it in with materials that are used by any section of any LOD
 	for (const FStaticMeshLODResources& LODModel : RenderData->LODResources)
 	{
 		for (const FStaticMeshSection& Section : LODModel.Sections)
@@ -370,7 +407,7 @@ void UNiagaraMeshRendererProperties::GetUsedMeshMaterials(int32 MeshIndex, const
 				{
 					OverrideMat = Cast<UMaterialInterface>(Emitter->FindBinding(OverrideMaterials[OverrideIndex].UserParamBinding.Parameter));
 				}
-				
+
 				if (!OverrideMat)
 				{
 					OverrideMat = OverrideMaterials[OverrideIndex].ExplicitMat;
@@ -456,7 +493,11 @@ void UNiagaraMeshRendererProperties::PostLoad()
 		}
 	}
 
+
 	PostLoadBindings(SourceMode);
+	
+	// Fix up these bindings from their loaded source bindings
+	SetPreviousBindings(nullptr, SourceMode);
 
 	for ( const FNiagaraMeshMaterialOverride& OverrideMaterial : OverrideMaterials )
 	{
@@ -507,6 +548,19 @@ const TArray<FNiagaraVariable>& UNiagaraMeshRendererProperties::GetOptionalAttri
 	return Attrs;
 }
 
+void UNiagaraMeshRendererProperties::GetAdditionalVariables(TArray<FNiagaraVariableBase>& OutArray) const
+{
+	if (NeedsPreciseMotionVectors())
+	{
+		OutArray.Reserve(5);
+		OutArray.AddUnique(PrevPositionBinding.GetParamMapBindableVariable());
+		OutArray.AddUnique(PrevScaleBinding.GetParamMapBindableVariable());
+		OutArray.AddUnique(PrevMeshOrientationBinding.GetParamMapBindableVariable());
+		OutArray.AddUnique(PrevCameraOffsetBinding.GetParamMapBindableVariable());
+		OutArray.AddUnique(PrevVelocityBinding.GetParamMapBindableVariable());		
+	}
+}
+
 void UNiagaraMeshRendererProperties::GetRendererWidgets(const FNiagaraEmitterInstance* InEmitter, TArray<TSharedPtr<SWidget>>& OutWidgets, TSharedPtr<FAssetThumbnailPool> InThumbnailPool) const
 {
 	TSharedRef<SWidget> ThumbnailWidget = SNullWidget::NullWidget;
@@ -642,7 +696,7 @@ void UNiagaraMeshRendererProperties::PostEditChangeProperty(FPropertyChangedEven
 			RebuildMeshList();
 		}
 	}
-		
+
 	if (bIsRedirect || bRebuildMeshList)
 	{
 		// We only need to check material usage as we will invalidate any renderers later on
@@ -762,11 +816,11 @@ void UNiagaraMeshRendererProperties::CheckMaterialUsage()
 }
 
 bool UNiagaraMeshRendererProperties::ChangeRequiresMeshListRebuild(const FProperty* Property)
-{	
+{
 	if (Property == nullptr)
 	{
 		return false;
-	}	
+	}
 
 	// If any of these are changed, we have to rebuild the mesh list
 	static const TArray<FName, TInlineAllocator<4>> RebuildMeshPropertyNames
@@ -795,7 +849,7 @@ void UNiagaraMeshRendererProperties::RebuildMeshList()
 		// No first page mesh selected
 		return;
 	}
-	
+
 	Meshes.AddDefaulted_GetRef().Mesh = FirstFlipbookFrame;
 
 	if (NumFlipbookFrames <= 1)
@@ -824,7 +878,7 @@ void UNiagaraMeshRendererProperties::RebuildMeshList()
 	}
 
 	FSoftObjectPath ParticleMeshPath = FirstFlipbookFrame->GetPathName();
-	FString BaseName = ParticleMeshPath.GetAssetName();	
+	FString BaseName = ParticleMeshPath.GetAssetName();
 	int32 FirstFrameIdx = 0;
 
 	// Build a regex pattern string to use to attempt to find the first frame number in the first frame mesh
@@ -882,7 +936,7 @@ void UNiagaraMeshRendererProperties::RebuildMeshList()
 
 		TMap<FString, FStringFormatArg> Args = {
 			{ TEXT("frame_number"), NumString }
-		};		
+		};
 
 		FString FrameName = BaseName + FString::Format(*FlipbookSuffixFormat, Args);
 		FSoftObjectPath ObjPath(BasePackageLocation / (FrameName + TCHAR('.') + FrameName));
diff --git a/Engine/Plugins/FX/Niagara/Source/Niagara/Private/NiagaraRenderer.cpp b/Engine/Plugins/FX/Niagara/Source/Niagara/Private/NiagaraRenderer.cpp
--- a/Engine/Plugins/FX/Niagara/Source/Niagara/Private/NiagaraRenderer.cpp
+++ b/Engine/Plugins/FX/Niagara/Source/Niagara/Private/NiagaraRenderer.cpp
@@ -275,7 +275,7 @@ FNiagaraRenderer::FNiagaraRenderer(ERHIFeatureLevel::Type InFeatureLevel, const
 	: DynamicDataRender(nullptr)
 	, bLocalSpace(Emitter->GetCachedEmitter()->bLocalSpace)
 	, bHasLights(false)
-	, bMotionBlurEnabled(InProps ? InProps->bMotionBlurEnabled : false)
+	, bMotionBlurEnabled(InProps ? InProps->MotionVectorSetting != ENiagaraRendererMotionVectorSetting::Disable : false)
 	, SimTarget(Emitter->GetCachedEmitter()->SimTarget)
 	, FeatureLevel(InFeatureLevel)
 {
diff --git a/Engine/Plugins/FX/Niagara/Source/Niagara/Private/NiagaraRendererMeshes.cpp b/Engine/Plugins/FX/Niagara/Source/Niagara/Private/NiagaraRendererMeshes.cpp
--- a/Engine/Plugins/FX/Niagara/Source/Niagara/Private/NiagaraRendererMeshes.cpp
+++ b/Engine/Plugins/FX/Niagara/Source/Niagara/Private/NiagaraRendererMeshes.cpp
@@ -2,7 +2,6 @@
 
 #include "NiagaraRendererMeshes.h"
 #include "ParticleResources.h"
-#include "NiagaraMeshVertexFactory.h"
 #include "NiagaraDataSet.h"
 #include "NiagaraStats.h"
 #include "Async/ParallelFor.h"
@@ -55,20 +54,6 @@ struct FNiagaraDynamicDataMesh : public FNiagaraDynamicDataBase
 
 //////////////////////////////////////////////////////////////////////////
 
-class FNiagaraMeshCollectorResourcesMesh : public FOneFrameResource
-{
-public:
-	FNiagaraMeshVertexFactory VertexFactory;
-	FNiagaraMeshUniformBufferRef UniformBuffer;
-
-	virtual ~FNiagaraMeshCollectorResourcesMesh()
-	{
-		VertexFactory.ReleaseResource();
-	}
-};
-
-//////////////////////////////////////////////////////////////////////////
-
 FNiagaraRendererMeshes::FNiagaraRendererMeshes(ERHIFeatureLevel::Type FeatureLevel, const UNiagaraRendererProperties* Props, const FNiagaraEmitterInstance* Emitter)
 	: FNiagaraRenderer(FeatureLevel, Props, Emitter)
 	, MaterialParamValidMask(0)
@@ -91,6 +76,7 @@ FNiagaraRendererMeshes::FNiagaraRendererMeshes(ERHIFeatureLevel::Type FeatureLev
 	bEnableCulling = bEnableFrustumCulling;
 	DistanceCullRange = FVector2D(0, FLT_MAX);
 	RendererVisibility = Properties->RendererVisibility;
+	bAccurateMotionVectors = Properties->NeedsPreciseMotionVectors();
 	MaxSectionCount = 0;
 
 	if (Properties->bEnableCameraDistanceCulling)
@@ -128,16 +114,16 @@ FNiagaraRendererMeshes::FNiagaraRendererMeshes(ERHIFeatureLevel::Type FeatureLev
 	RendererLayoutWithCustomSorting = &Properties->RendererLayoutWithCustomSorting;
 	RendererLayoutWithoutCustomSorting = &Properties->RendererLayoutWithoutCustomSorting;
 
-
 	bSetAnyBoundVars = false;
 	if (Emitter->GetRendererBoundVariables().IsEmpty() == false)
 	{
 		const TArray< const FNiagaraVariableAttributeBinding*>& VFBindings = Properties->GetAttributeBindings();
-		check(VFBindings.Num() >= ENiagaraMeshVFLayout::Type::Num);
-		for (int32 i = 0; i < ENiagaraMeshVFLayout::Type::Num; i++)
+		const int32 NumBindings = bAccurateMotionVectors ? ENiagaraMeshVFLayout::Type::Num_Max : ENiagaraMeshVFLayout::Type::Num_Default;
+		check(VFBindings.Num() >= ENiagaraMeshVFLayout::Type::Num_Max);
+		for (int32 i = 0; i < ENiagaraMeshVFLayout::Type::Num_Max; i++)
 		{
 			VFBoundOffsetsInParamStore[i] = INDEX_NONE;
-			if (VFBindings[i] && VFBindings[i]->CanBindToHostParameterMap())
+			if (i < NumBindings && VFBindings[i] && VFBindings[i]->CanBindToHostParameterMap())
 			{
 				VFBoundOffsetsInParamStore[i] = Emitter->GetRendererBoundVariables().IndexOf(VFBindings[i]->GetParamMapBindableVariable());
 				if (VFBoundOffsetsInParamStore[i] != INDEX_NONE)
@@ -147,7 +133,7 @@ FNiagaraRendererMeshes::FNiagaraRendererMeshes(ERHIFeatureLevel::Type FeatureLev
 	}
 	else
 	{
-		for (int32 i = 0; i < ENiagaraMeshVFLayout::Type::Num; i++)
+		for (int32 i = 0; i < ENiagaraMeshVFLayout::Type::Num_Max; i++)
 		{
 			VFBoundOffsetsInParamStore[i] = INDEX_NONE;
 		}
@@ -223,15 +209,15 @@ void FNiagaraRendererMeshes::ReleaseRenderThreadResources()
 {
 }
 
-void FNiagaraRendererMeshes::SetupVertexFactory(FNiagaraMeshVertexFactory* InVertexFactory, const FStaticMeshLODResources& LODResources) const
+void FNiagaraRendererMeshes::SetupVertexFactory(FNiagaraMeshVertexFactory& InVertexFactory, const FStaticMeshLODResources& LODResources) const
 {
 	FStaticMeshDataType Data;
 
-	LODResources.VertexBuffers.PositionVertexBuffer.BindPositionVertexBuffer(InVertexFactory, Data);
-	LODResources.VertexBuffers.StaticMeshVertexBuffer.BindTangentVertexBuffer(InVertexFactory, Data);
-	LODResources.VertexBuffers.StaticMeshVertexBuffer.BindTexCoordVertexBuffer(InVertexFactory, Data, MAX_TEXCOORDS);
-	LODResources.VertexBuffers.ColorVertexBuffer.BindColorVertexBuffer(InVertexFactory, Data);
-	InVertexFactory->SetData(Data);
+	LODResources.VertexBuffers.PositionVertexBuffer.BindPositionVertexBuffer(&InVertexFactory, Data);
+	LODResources.VertexBuffers.StaticMeshVertexBuffer.BindTangentVertexBuffer(&InVertexFactory, Data);
+	LODResources.VertexBuffers.StaticMeshVertexBuffer.BindTexCoordVertexBuffer(&InVertexFactory, Data, MAX_TEXCOORDS);
+	LODResources.VertexBuffers.ColorVertexBuffer.BindColorVertexBuffer(&InVertexFactory, Data);
+	InVertexFactory.SetData(Data);
 }
 
 int32 FNiagaraRendererMeshes::GetLODIndex(int32 MeshIndex) const
@@ -390,7 +376,6 @@ FNiagaraMeshUniformBufferRef FNiagaraRendererMeshes::CreatePerViewUniformBuffer(
 	FMemory::Memzero(&PerViewUniformParameters, sizeof(PerViewUniformParameters)); // Clear unset bytes
 
 	PerViewUniformParameters.bLocalSpace = bLocalSpace;
-	PerViewUniformParameters.PrevTransformAvailable = false;
 	PerViewUniformParameters.DeltaSeconds = View.Family->DeltaWorldTime;
 	PerViewUniformParameters.MeshScale = MeshData.Scale;
 
@@ -428,12 +413,15 @@ FNiagaraMeshUniformBufferRef FNiagaraRendererMeshes::CreatePerViewUniformBuffer(
 	PerViewUniformParameters.NiagaraParticleDataFloat = BufferData.FloatSRV;
 	PerViewUniformParameters.NiagaraParticleDataHalf = BufferData.HalfSRV;
 
-
 	PerViewUniformParameters.DefaultPos = bLocalSpace ? FVector4(0.0f, 0.0f, 0.0f, 1.0f) : FVector4(SceneProxy.GetLocalToWorld().GetOrigin());
+	PerViewUniformParameters.DefaultPrevPos = PerViewUniformParameters.DefaultPos;
 	PerViewUniformParameters.DefaultVelocity = FVector(0.f, 0.0f, 0.0f);
+	PerViewUniformParameters.DefaultPrevVelocity = PerViewUniformParameters.DefaultVelocity;
 	PerViewUniformParameters.DefaultColor = FVector4(1.0f, 1.0f, 1.0f, 1.0f);
 	PerViewUniformParameters.DefaultScale = FVector(1.0f, 1.0f, 1.0f);
-	PerViewUniformParameters.DefaultTransform = FVector4(0.0f, 0.0f, 0.0f, 1.0f);
+	PerViewUniformParameters.DefaultPrevScale = PerViewUniformParameters.DefaultScale;
+	PerViewUniformParameters.DefaultRotation = FVector4(0.0f, 0.0f, 0.0f, 1.0f);
+	PerViewUniformParameters.DefaultPrevRotation = PerViewUniformParameters.DefaultRotation;
 	PerViewUniformParameters.DefaultMatRandom = 0.0f;
 	PerViewUniformParameters.DefaultNormAge = 0.0f;
 
@@ -443,6 +431,13 @@ FNiagaraMeshUniformBufferRef FNiagaraRendererMeshes::CreatePerViewUniformBuffer(
 	PerViewUniformParameters.DefaultDynamicMaterialParameter2 = FVector4(1.0f, 1.0f, 1.0f, 1.0f);
 	PerViewUniformParameters.DefaultDynamicMaterialParameter3 = FVector4(1.0f, 1.0f, 1.0f, 1.0f);
 	PerViewUniformParameters.DefaultCamOffset = 0.0f;
+	PerViewUniformParameters.DefaultPrevCamOffset = PerViewUniformParameters.DefaultCamOffset;
+
+	PerViewUniformParameters.PrevPositionDataOffset = INDEX_NONE;
+	PerViewUniformParameters.PrevScaleDataOffset = INDEX_NONE;
+	PerViewUniformParameters.PrevRotationDataOffset = INDEX_NONE;
+	PerViewUniformParameters.PrevCameraOffsetDataOffset = INDEX_NONE;
+	PerViewUniformParameters.PrevVelocityDataOffset = INDEX_NONE;
 
 	TConstArrayView<FNiagaraRendererVariableInfo> VFVariables = RendererLayout.GetVFVariables_RenderThread();
 	if (SourceMode == ENiagaraRendererSourceDataMode::Particles)
@@ -451,7 +446,7 @@ FNiagaraMeshUniformBufferRef FNiagaraRendererMeshes::CreatePerViewUniformBuffer(
 		PerViewUniformParameters.VelocityDataOffset = VFVariables[ENiagaraMeshVFLayout::Velocity].GetGPUOffset();
 		PerViewUniformParameters.ColorDataOffset = VFVariables[ENiagaraMeshVFLayout::Color].GetGPUOffset();
 		PerViewUniformParameters.ScaleDataOffset = VFVariables[ENiagaraMeshVFLayout::Scale].GetGPUOffset();
-		PerViewUniformParameters.TransformDataOffset = VFVariables[ENiagaraMeshVFLayout::Transform].GetGPUOffset();
+		PerViewUniformParameters.RotationDataOffset = VFVariables[ENiagaraMeshVFLayout::Rotation].GetGPUOffset();
 		PerViewUniformParameters.MaterialRandomDataOffset = VFVariables[ENiagaraMeshVFLayout::MaterialRandom].GetGPUOffset();
 		PerViewUniformParameters.NormalizedAgeDataOffset = VFVariables[ENiagaraMeshVFLayout::NormalizedAge].GetGPUOffset();
 
@@ -461,6 +456,15 @@ FNiagaraMeshUniformBufferRef FNiagaraRendererMeshes::CreatePerViewUniformBuffer(
 		PerViewUniformParameters.MaterialParam2DataOffset = VFVariables[ENiagaraMeshVFLayout::DynamicParam2].GetGPUOffset();
 		PerViewUniformParameters.MaterialParam3DataOffset = VFVariables[ENiagaraMeshVFLayout::DynamicParam3].GetGPUOffset();
 		PerViewUniformParameters.CameraOffsetDataOffset = VFVariables[ENiagaraMeshVFLayout::CameraOffset].GetGPUOffset();
+		
+		if (bAccurateMotionVectors)
+		{
+			PerViewUniformParameters.PrevPositionDataOffset = VFVariables[ENiagaraMeshVFLayout::PrevPosition].GetGPUOffset();
+			PerViewUniformParameters.PrevScaleDataOffset = VFVariables[ENiagaraMeshVFLayout::PrevScale].GetGPUOffset();
+			PerViewUniformParameters.PrevRotationDataOffset = VFVariables[ENiagaraMeshVFLayout::PrevRotation].GetGPUOffset();
+			PerViewUniformParameters.PrevCameraOffsetDataOffset = VFVariables[ENiagaraMeshVFLayout::PrevCameraOffset].GetGPUOffset();
+			PerViewUniformParameters.PrevVelocityDataOffset = VFVariables[ENiagaraMeshVFLayout::PrevVelocity].GetGPUOffset();
+		}
 	}
 	else if (SourceMode == ENiagaraRendererSourceDataMode::Emitter) // Clear all these out because we will be using the defaults to specify them
 	{
@@ -468,7 +472,7 @@ FNiagaraMeshUniformBufferRef FNiagaraRendererMeshes::CreatePerViewUniformBuffer(
 		PerViewUniformParameters.VelocityDataOffset = INDEX_NONE;
 		PerViewUniformParameters.ColorDataOffset = INDEX_NONE;
 		PerViewUniformParameters.ScaleDataOffset = INDEX_NONE;
-		PerViewUniformParameters.TransformDataOffset = INDEX_NONE;
+		PerViewUniformParameters.RotationDataOffset = INDEX_NONE;
 		PerViewUniformParameters.MaterialRandomDataOffset = INDEX_NONE;
 		PerViewUniformParameters.NormalizedAgeDataOffset = INDEX_NONE;
 
@@ -485,12 +489,12 @@ FNiagaraMeshUniformBufferRef FNiagaraRendererMeshes::CreatePerViewUniformBuffer(
 		check(SourceMode <= ENiagaraRendererSourceDataMode::Emitter);
 	}
 
-
 	if (bSetAnyBoundVars && DynamicDataMesh)
 	{
 		const uint8* ParameterBoundData = DynamicDataMesh->ParameterDataBound.GetData();
 
-		for (int32 i = 0; i < ENiagaraMeshVFLayout::Type::Num; i++)
+		const int32 NumVFOffsets = bAccurateMotionVectors ? ENiagaraMeshVFLayout::Type::Num_Max : ENiagaraMeshVFLayout::Type::Num_Default;
+		for (int32 i = 0; i < NumVFOffsets; i++)
 		{
 			if (VFBoundOffsetsInParamStore[i] != INDEX_NONE && DynamicDataMesh->ParameterDataBound.IsValidIndex(VFBoundOffsetsInParamStore[i]))
 			{
@@ -500,43 +504,76 @@ FNiagaraMeshUniformBufferRef FNiagaraRendererMeshes::CreatePerViewUniformBuffer(
 					FMemory::Memcpy(&PerViewUniformParameters.DefaultPos, ParameterBoundData + VFBoundOffsetsInParamStore[i], sizeof(FVector));
 					break;
 				case ENiagaraMeshVFLayout::Type::Velocity:
-					memcpy(&PerViewUniformParameters.DefaultVelocity, ParameterBoundData + VFBoundOffsetsInParamStore[i], sizeof(FVector));
+					FMemory::Memcpy(&PerViewUniformParameters.DefaultVelocity, ParameterBoundData + VFBoundOffsetsInParamStore[i], sizeof(FVector));
 					break;
 				case ENiagaraMeshVFLayout::Type::Color:
-					memcpy(&PerViewUniformParameters.DefaultColor, ParameterBoundData + VFBoundOffsetsInParamStore[i], sizeof(FLinearColor));
+					FMemory::Memcpy(&PerViewUniformParameters.DefaultColor, ParameterBoundData + VFBoundOffsetsInParamStore[i], sizeof(FLinearColor));
 					break;
 				case ENiagaraMeshVFLayout::Type::Scale:
-					memcpy(&PerViewUniformParameters.DefaultScale, ParameterBoundData + VFBoundOffsetsInParamStore[i], sizeof(FVector));
+					FMemory::Memcpy(&PerViewUniformParameters.DefaultScale, ParameterBoundData + VFBoundOffsetsInParamStore[i], sizeof(FVector));
 					break;
-				case ENiagaraMeshVFLayout::Type::Transform:
-					memcpy(&PerViewUniformParameters.DefaultTransform, ParameterBoundData + VFBoundOffsetsInParamStore[i], sizeof(FVector4));
+				case ENiagaraMeshVFLayout::Type::Rotation:
+					FMemory::Memcpy(&PerViewUniformParameters.DefaultRotation, ParameterBoundData + VFBoundOffsetsInParamStore[i], sizeof(FVector4));
 					break;
 				case ENiagaraMeshVFLayout::Type::MaterialRandom:
-					memcpy(&PerViewUniformParameters.DefaultMatRandom, ParameterBoundData + VFBoundOffsetsInParamStore[i], sizeof(float));
+					FMemory::Memcpy(&PerViewUniformParameters.DefaultMatRandom, ParameterBoundData + VFBoundOffsetsInParamStore[i], sizeof(float));
 					break;
 				case ENiagaraMeshVFLayout::Type::NormalizedAge:
-					memcpy(&PerViewUniformParameters.DefaultNormAge, ParameterBoundData + VFBoundOffsetsInParamStore[i], sizeof(float));
+					FMemory::Memcpy(&PerViewUniformParameters.DefaultNormAge, ParameterBoundData + VFBoundOffsetsInParamStore[i], sizeof(float));
 					break;
 				case ENiagaraMeshVFLayout::Type::CustomSorting:
 					// unsupported for now...
 					break;
 				case ENiagaraMeshVFLayout::Type::SubImage:
-					memcpy(&PerViewUniformParameters.DefaultSubImage, ParameterBoundData + VFBoundOffsetsInParamStore[i], sizeof(float));
+					FMemory::Memcpy(&PerViewUniformParameters.DefaultSubImage, ParameterBoundData + VFBoundOffsetsInParamStore[i], sizeof(float));
 					break;
 				case ENiagaraMeshVFLayout::Type::DynamicParam0:
-					memcpy(&PerViewUniformParameters.DefaultDynamicMaterialParameter0, ParameterBoundData + VFBoundOffsetsInParamStore[i], sizeof(FVector4));
+					FMemory::Memcpy(&PerViewUniformParameters.DefaultDynamicMaterialParameter0, ParameterBoundData + VFBoundOffsetsInParamStore[i], sizeof(FVector4));
 					break;
 				case ENiagaraMeshVFLayout::Type::DynamicParam1:
-					memcpy(&PerViewUniformParameters.DefaultDynamicMaterialParameter1, ParameterBoundData + VFBoundOffsetsInParamStore[i], sizeof(FVector4));
+					FMemory::Memcpy(&PerViewUniformParameters.DefaultDynamicMaterialParameter1, ParameterBoundData + VFBoundOffsetsInParamStore[i], sizeof(FVector4));
 					break;
 				case ENiagaraMeshVFLayout::Type::DynamicParam2:
-					memcpy(&PerViewUniformParameters.DefaultDynamicMaterialParameter2, ParameterBoundData + VFBoundOffsetsInParamStore[i], sizeof(FVector4));
+					FMemory::Memcpy(&PerViewUniformParameters.DefaultDynamicMaterialParameter2, ParameterBoundData + VFBoundOffsetsInParamStore[i], sizeof(FVector4));
 					break;
 				case ENiagaraMeshVFLayout::Type::DynamicParam3:
-					memcpy(&PerViewUniformParameters.DefaultDynamicMaterialParameter3, ParameterBoundData + VFBoundOffsetsInParamStore[i], sizeof(FVector4));
+					FMemory::Memcpy(&PerViewUniformParameters.DefaultDynamicMaterialParameter3, ParameterBoundData + VFBoundOffsetsInParamStore[i], sizeof(FVector4));
 					break;
 				case ENiagaraMeshVFLayout::Type::CameraOffset:
-					memcpy(&PerViewUniformParameters.DefaultCamOffset, ParameterBoundData + VFBoundOffsetsInParamStore[i], sizeof(float));
+					FMemory::Memcpy(&PerViewUniformParameters.DefaultCamOffset, ParameterBoundData + VFBoundOffsetsInParamStore[i], sizeof(float));
+					break;
+				case ENiagaraMeshVFLayout::Type::PrevPosition:
+					FMemory::Memcpy(&PerViewUniformParameters.DefaultPrevPos, ParameterBoundData + VFBoundOffsetsInParamStore[i], sizeof(FVector));
+					break;
+				case ENiagaraMeshVFLayout::Type::PrevScale:
+					FMemory::Memcpy(&PerViewUniformParameters.DefaultPrevScale, ParameterBoundData + VFBoundOffsetsInParamStore[i], sizeof(FVector));
+					break;
+				case ENiagaraMeshVFLayout::Type::PrevRotation:
+					FMemory::Memcpy(&PerViewUniformParameters.DefaultPrevRotation, ParameterBoundData + VFBoundOffsetsInParamStore[i], sizeof(FVector4));
+					break;
+				case ENiagaraMeshVFLayout::Type::PrevCameraOffset:
+					FMemory::Memcpy(&PerViewUniformParameters.DefaultPrevCamOffset, ParameterBoundData + VFBoundOffsetsInParamStore[i], sizeof(float));
+					break;
+				}
+			}
+			else
+			{
+				// If these prev values aren't bound to the host parameters, but their current values are, copy them
+				switch (i)
+				{
+				case ENiagaraMeshVFLayout::Type::PrevPosition:
+					PerViewUniformParameters.DefaultPrevPos = PerViewUniformParameters.DefaultPos;
+					break;
+				case ENiagaraMeshVFLayout::Type::PrevScale:
+					PerViewUniformParameters.DefaultPrevScale = PerViewUniformParameters.DefaultScale;
+					break;
+				case ENiagaraMeshVFLayout::Type::PrevRotation:
+					PerViewUniformParameters.DefaultPrevRotation = PerViewUniformParameters.DefaultRotation;
+					break;
+				case ENiagaraMeshVFLayout::Type::PrevCameraOffset:
+					PerViewUniformParameters.DefaultPrevCamOffset = PerViewUniformParameters.DefaultCamOffset;
+					break;
+				default:
 					break;
 				}
 			}
@@ -693,7 +730,7 @@ void FNiagaraRendererMeshes::InitializeSortInfo(
 	if (bDoGPUCulling)
 	{
 		OutSortInfo.CullPositionAttributeOffset = VFVariables[ENiagaraMeshVFLayout::Position].GetGPUOffset();
-		OutSortInfo.CullOrientationAttributeOffset = VFVariables[ENiagaraMeshVFLayout::Transform].GetGPUOffset();
+		OutSortInfo.CullOrientationAttributeOffset = VFVariables[ENiagaraMeshVFLayout::Rotation].GetGPUOffset();
 		OutSortInfo.CullScaleAttributeOffset = VFVariables[ENiagaraMeshVFLayout::Scale].GetGPUOffset();
 	}
 }
@@ -909,22 +946,30 @@ void FNiagaraRendererMeshes::GetDynamicMeshElements(const TArray<const FSceneVie
 				const FStaticMeshLODResources& LODModel = MeshData.RenderData->LODResources[LODIndex];
 				const int32 SectionCount = LODModel.Sections.Num();
 
-				FNiagaraMeshCollectorResourcesMesh& CollectorResources = Collector.AllocateOneFrameResource<FNiagaraMeshCollectorResourcesMesh>();
+				FMeshCollectorResourcesBase* CollectorResources;
+				if (bAccurateMotionVectors)
+				{
+					CollectorResources = &Collector.AllocateOneFrameResource<FMeshCollectorResourcesEx>();
+				}
+				else
+				{
+					CollectorResources = &Collector.AllocateOneFrameResource<FMeshCollectorResources>();
+				}				
 
 				// Get the next vertex factory to use
 				// TODO: Find a way to safely pool these such that they won't be concurrently accessed by multiple views
-				FNiagaraMeshVertexFactory* VertexFactory = &CollectorResources.VertexFactory;
-				VertexFactory->SetParticleFactoryType(NVFT_Mesh);
-				VertexFactory->SetMeshIndex(MeshIndex);
-				VertexFactory->SetLODIndex(LODIndex);
-				VertexFactory->InitResource();
+				FNiagaraMeshVertexFactory& VertexFactory = CollectorResources->GetVertexFactory();
+				VertexFactory.SetParticleFactoryType(NVFT_Mesh);
+				VertexFactory.SetMeshIndex(MeshIndex);
+				VertexFactory.SetLODIndex(LODIndex);
+				VertexFactory.InitResource();
 				SetupVertexFactory(VertexFactory, LODModel);
 
-				VertexFactory->SetUniformBuffer(PerViewUniformBuffer);
-				CollectorResources.UniformBuffer = PerViewUniformBuffer;
+				VertexFactory.SetUniformBuffer(PerViewUniformBuffer);
+				CollectorResources->UniformBuffer = PerViewUniformBuffer;
 
 				// Sort/Cull particles if needed.
-				VertexFactory->SetSortedIndices(nullptr, 0xFFFFFFFF);
+				VertexFactory.SetSortedIndices(nullptr, 0xFFFFFFFF);
 				if (bShouldSort || bDoGPUCulling)
 				{
 					// Set up mesh-specific sorting parameters
@@ -941,7 +986,7 @@ void FNiagaraRendererMeshes::GetDynamicMeshElements(const TArray<const FSceneVie
 						// We need to run the sort shader on the GPU
 						if (Batcher->AddSortedGPUSimulation(SortInfo))
 						{
-							VertexFactory->SetSortedIndices(SortInfo.AllocationInfo.BufferSRV, SortInfo.AllocationInfo.BufferOffset);
+							VertexFactory.SetSortedIndices(SortInfo.AllocationInfo.BufferSRV, SortInfo.AllocationInfo.BufferOffset);
 						}
 					}
 					else
@@ -951,7 +996,7 @@ void FNiagaraRendererMeshes::GetDynamicMeshElements(const TArray<const FSceneVie
 						FGlobalDynamicReadBuffer::FAllocation SortedIndices;
 						SortedIndices = DynamicReadBuffer.AllocateInt32(NumInstances);
 						SortIndices(SortInfo, VFVariables[SortVarIdx], *SourceParticleData, SortedIndices);
-						VertexFactory->SetSortedIndices(SortedIndices.SRV, 0);
+						VertexFactory.SetSortedIndices(SortedIndices.SRV, 0);
 					}
 				}
 
@@ -983,7 +1028,7 @@ void FNiagaraRendererMeshes::GetDynamicMeshElements(const TArray<const FSceneVie
 					}
 
 					const uint32 GPUCountBufferOffset = bDoGPUCulling ? SortInfo.CulledGPUParticleCountOffset : SourceParticleData->GetGPUInstanceCountBufferOffset();
-					CreateMeshBatchForSection(Collector, *VertexFactory, *MaterialProxy, *SceneProxy, LODModel, Section, *View, ViewIndex, NumInstances,
+					CreateMeshBatchForSection(Collector, VertexFactory, *MaterialProxy, *SceneProxy, LODModel, Section, *View, ViewIndex, NumInstances,
 						GPUCountBufferOffset, bIsWireframe, bIsInstancedStereo, bDoGPUCulling);
 				}
 			}
@@ -1094,7 +1139,7 @@ void FNiagaraRendererMeshes::GetDynamicRayTracingInstances(FRayTracingMaterialGa
 		//ENiagaraMeshVFLayout::Transform just contains a Quat, not the whole transform
 		const FNiagaraRendererVariableInfo& VarPositionInfo = VFVariables[ENiagaraMeshVFLayout::Position];
 		const FNiagaraRendererVariableInfo& VarScaleInfo = VFVariables[ENiagaraMeshVFLayout::Scale];
-		const FNiagaraRendererVariableInfo& VarTransformInfo = VFVariables[ENiagaraMeshVFLayout::Transform];
+		const FNiagaraRendererVariableInfo& VarTransformInfo = VFVariables[ENiagaraMeshVFLayout::Rotation];
 
 		int32 PositionBaseCompOffset = VarPositionInfo.DatasetOffset;
 		int32 ScaleBaseCompOffset = VarScaleInfo.DatasetOffset;
@@ -1223,7 +1268,7 @@ void FNiagaraRendererMeshes::GetDynamicRayTracingInstances(FRayTracingMaterialGa
 
 			const FUintVector4 NiagaraOffsets(
 				VFVariables[ENiagaraMeshVFLayout::Position].GetGPUOffset(),
-				VFVariables[ENiagaraMeshVFLayout::Transform].GetGPUOffset(),
+				VFVariables[ENiagaraMeshVFLayout::Rotation].GetGPUOffset(),
 				VFVariables[ENiagaraMeshVFLayout::Scale].GetGPUOffset(),
 				bLocalSpace ? 1 : 0);
 
diff --git a/Engine/Plugins/FX/Niagara/Source/Niagara/Private/NiagaraRendererProperties.cpp b/Engine/Plugins/FX/Niagara/Source/Niagara/Private/NiagaraRendererProperties.cpp
--- a/Engine/Plugins/FX/Niagara/Source/Niagara/Private/NiagaraRendererProperties.cpp
+++ b/Engine/Plugins/FX/Niagara/Source/Niagara/Private/NiagaraRendererProperties.cpp
@@ -5,10 +5,11 @@
 #include "NiagaraDataSet.h"
 #include "NiagaraConstants.h"
 #include "NiagaraEmitter.h"
-#include "Interfaces/ITargetPlatform.h"
-#include "Styling/SlateIconFinder.h"
 #include "NiagaraScriptSourceBase.h"
+#include "NiagaraSettings.h"
 #include "NiagaraSystem.h"
+#include "Interfaces/ITargetPlatform.h"
+#include "Styling/SlateIconFinder.h"
 
 void FNiagaraRendererLayout::Initialize(int32 NumVariables)
 {
@@ -18,7 +19,7 @@ void FNiagaraRendererLayout::Initialize(int32 NumVariables)
 	TotalHalfComponents_GT = 0;
 }
 
-bool FNiagaraRendererLayout::SetVariable(const FNiagaraDataSetCompiledData* CompiledData, const FNiagaraVariable& Variable, int32 VFVarOffset)
+bool FNiagaraRendererLayout::SetVariable(const FNiagaraDataSetCompiledData* CompiledData, const FNiagaraVariableBase& Variable, int32 VFVarOffset)
 {
 	// No compiled data, nothing to bind
 	if (CompiledData == nullptr)
@@ -328,6 +329,44 @@ void UNiagaraRendererProperties::PostInitProperties()
 #endif
 }
 
+void UNiagaraRendererProperties::PostLoad()
+{
+	Super::PostLoad();
+
+	if (bMotionBlurEnabled_DEPRECATED == false)
+	{
+		MotionVectorSetting = ENiagaraRendererMotionVectorSetting::Disable;
+	}
+}
+
+#if WITH_EDITORONLY_DATA
+
+void UNiagaraRendererProperties::PostEditChangeProperty(struct FPropertyChangedEvent& PropertyChangedEvent)
+{
+	Super::PostEditChangeProperty(PropertyChangedEvent);
+
+	if (UNiagaraEmitter* Emitter = GetTypedOuter<UNiagaraEmitter>())
+	{
+		// Check for properties changing that invalidate the current script compilation for the emitter
+		bool bNeedsRecompile = false;
+		if (PropertyChangedEvent.GetPropertyName() == GET_MEMBER_NAME_CHECKED(UNiagaraRendererProperties, MotionVectorSetting))
+		{
+			if (Emitter->GraphSource)
+			{
+				Emitter->GraphSource->MarkNotSynchronized(TEXT("Renderer MotionVectorSetting changed"));
+			}
+			bNeedsRecompile = true;
+		}
+
+		if (bNeedsRecompile)
+		{
+			UNiagaraSystem::RequestCompileForEmitter(Emitter);
+		}
+	}
+}
+
+#endif
+
 void UNiagaraRendererProperties::SetIsEnabled(bool bInIsEnabled)
 {
 	if (bIsEnabled != bInIsEnabled)
@@ -370,4 +409,16 @@ void UNiagaraRendererProperties::UpdateSourceModeDerivates(ENiagaraRendererSourc
 		}
 #endif
 	}
+}
+
+bool UNiagaraRendererProperties::NeedsPreciseMotionVectors() const
+{
+	if (MotionVectorSetting == ENiagaraRendererMotionVectorSetting::AutoDetect)
+	{
+		// TODO - We could get even smarter here and early return with false if we know that the material can absolutely not be overridden by the user and
+		// it doesn't need to render velocity
+		return GetDefault<UNiagaraSettings>()->DefaultRendererMotionVectorSetting == ENiagaraDefaultRendererMotionVectorSetting::Precise;
+	}
+	
+	return MotionVectorSetting == ENiagaraRendererMotionVectorSetting::Precise;
 }
\ No newline at end of file
diff --git a/Engine/Plugins/FX/Niagara/Source/Niagara/Private/NiagaraRendererSprites.cpp b/Engine/Plugins/FX/Niagara/Source/Niagara/Private/NiagaraRendererSprites.cpp
--- a/Engine/Plugins/FX/Niagara/Source/Niagara/Private/NiagaraRendererSprites.cpp
+++ b/Engine/Plugins/FX/Niagara/Source/Niagara/Private/NiagaraRendererSprites.cpp
@@ -2,7 +2,6 @@
 
 #include "NiagaraRendererSprites.h"
 #include "ParticleResources.h"
-#include "NiagaraSpriteVertexFactory.h"
 #include "NiagaraDataSet.h"
 #include "NiagaraStats.h"
 #include "NiagaraEmitterInstanceBatcher.h"
@@ -60,20 +59,6 @@ struct FNiagaraDynamicDataSprites : public FNiagaraDynamicDataBase
 	TArray<uint8> ParameterDataBound;
 };
 
-/* Mesh collector classes */
-class FNiagaraMeshCollectorResourcesSprite : public FOneFrameResource
-{
-public:
-	FNiagaraSpriteVertexFactory VertexFactory;
-	FNiagaraSpriteUniformBufferRef UniformBuffer;
-
-	virtual ~FNiagaraMeshCollectorResourcesSprite()
-	{
-		VertexFactory.ReleaseResource();
-	}
-};
-
-
 //////////////////////////////////////////////////////////////////////////
 
 FNiagaraRendererSprites::FNiagaraRendererSprites(ERHIFeatureLevel::Type FeatureLevel, const UNiagaraRendererProperties* InProps, const FNiagaraEmitterInstance* Emitter)
@@ -113,6 +98,7 @@ FNiagaraRendererSprites::FNiagaraRendererSprites(ERHIFeatureLevel::Type FeatureL
 	MinFacingCameraBlendDistance = Properties->MinFacingCameraBlendDistance;
 	MaxFacingCameraBlendDistance = Properties->MaxFacingCameraBlendDistance;
 	RendererVisibility = Properties->RendererVisibility;
+	bAccurateMotionVectors = Properties->NeedsPreciseMotionVectors();
 
 	bEnableDistanceCulling = Properties->bEnableCameraDistanceCulling;
 	if (Properties->bEnableCameraDistanceCulling)
@@ -149,12 +135,13 @@ FNiagaraRendererSprites::FNiagaraRendererSprites(ERHIFeatureLevel::Type FeatureL
 	if (Emitter->GetRendererBoundVariables().IsEmpty() == false)
 	{
 		const TArray< const FNiagaraVariableAttributeBinding*>& VFBindings = Properties->GetAttributeBindings();
-		check(VFBindings.Num() >= ENiagaraSpriteVFLayout::Type::Num);
+		const int32 NumBindings = bAccurateMotionVectors ? ENiagaraSpriteVFLayout::Num_Max : ENiagaraSpriteVFLayout::Num_Default;
+		check(VFBindings.Num() >= ENiagaraSpriteVFLayout::Type::Num_Max);
 
-		for (int32 i = 0; i < ENiagaraSpriteVFLayout::Type::Num; i++)
+		for (int32 i = 0; i < ENiagaraSpriteVFLayout::Type::Num_Max; i++)
 		{
 			VFBoundOffsetsInParamStore[i] = INDEX_NONE;
-			if (VFBindings[i] && VFBindings[i]->CanBindToHostParameterMap())
+			if (i < NumBindings && VFBindings[i] && VFBindings[i]->CanBindToHostParameterMap())
 			{
 				VFBoundOffsetsInParamStore[i] = Emitter->GetRendererBoundVariables().IndexOf(VFBindings[i]->GetParamMapBindableVariable());
 				if (VFBoundOffsetsInParamStore[i] != INDEX_NONE)
@@ -164,7 +151,7 @@ FNiagaraRendererSprites::FNiagaraRendererSprites(ERHIFeatureLevel::Type FeatureL
 	}
 	else
 	{
-		for (int32 i = 0; i < ENiagaraSpriteVFLayout::Type::Num; i++)
+		for (int32 i = 0; i < ENiagaraSpriteVFLayout::Type::Num_Max; i++)
 		{
 			VFBoundOffsetsInParamStore[i] = INDEX_NONE;
 		}
@@ -277,23 +264,40 @@ FNiagaraSpriteUniformBufferRef FNiagaraRendererSprites::CreatePerViewUniformBuff
 
 
 	PerViewUniformParameters.DefaultPos = bLocalSpace ? FVector4(0.0f, 0.0f, 0.0f, 1.0f) : FVector4(SceneProxy->GetLocalToWorld().GetOrigin());
+	PerViewUniformParameters.DefaultPrevPos = PerViewUniformParameters.DefaultPos;
 	PerViewUniformParameters.DefaultSize = FVector2D(50.f, 50.0f);
+	PerViewUniformParameters.DefaultPrevSize = PerViewUniformParameters.DefaultSize;
 	PerViewUniformParameters.DefaultUVScale = FVector2D(1.0f, 1.0f);
 	PerViewUniformParameters.DefaultPivotOffset = PivotInUVSpace;
+	PerViewUniformParameters.DefaultPrevPivotOffset = PerViewUniformParameters.DefaultPivotOffset;
 	PerViewUniformParameters.DefaultVelocity = FVector(0.f, 0.0f, 0.0f);
+	PerViewUniformParameters.DefaultPrevVelocity = PerViewUniformParameters.DefaultVelocity;
 	PerViewUniformParameters.DefaultRotation =  0.0f;
+	PerViewUniformParameters.DefaultPrevRotation =  PerViewUniformParameters.DefaultRotation;
 	PerViewUniformParameters.DefaultColor = FVector4(1.0f, 1.0f, 1.0f, 1.0f);
 	PerViewUniformParameters.DefaultMatRandom = 0.0f;
 	PerViewUniformParameters.DefaultCamOffset = 0.0f;
+	PerViewUniformParameters.DefaultPrevCamOffset = PerViewUniformParameters.DefaultCamOffset;
 	PerViewUniformParameters.DefaultNormAge = 0.0f;
 	PerViewUniformParameters.DefaultSubImage = 0.0f;
 	PerViewUniformParameters.DefaultFacing = FVector4(1.0f, 0.0f, 0.0f, 0.0f);
+	PerViewUniformParameters.DefaultPrevFacing = PerViewUniformParameters.DefaultFacing;
 	PerViewUniformParameters.DefaultAlignment = FVector4(1.0f, 0.0f, 0.0f, 0.0f);
+	PerViewUniformParameters.DefaultPrevAlignment = PerViewUniformParameters.DefaultAlignment;
 	PerViewUniformParameters.DefaultDynamicMaterialParameter0 = FVector4(1.0f, 1.0f, 1.0f, 1.0f);
 	PerViewUniformParameters.DefaultDynamicMaterialParameter1 = FVector4(1.0f, 1.0f, 1.0f, 1.0f);
 	PerViewUniformParameters.DefaultDynamicMaterialParameter2 = FVector4(1.0f, 1.0f, 1.0f, 1.0f);
 	PerViewUniformParameters.DefaultDynamicMaterialParameter3 = FVector4(1.0f, 1.0f, 1.0f, 1.0f);
 
+	PerViewUniformParameters.PrevPositionDataOffset = INDEX_NONE;
+	PerViewUniformParameters.PrevVelocityDataOffset = INDEX_NONE;
+	PerViewUniformParameters.PrevRotationDataOffset = INDEX_NONE;
+	PerViewUniformParameters.PrevSizeDataOffset = INDEX_NONE;
+	PerViewUniformParameters.PrevFacingDataOffset = INDEX_NONE;
+	PerViewUniformParameters.PrevAlignmentDataOffset = INDEX_NONE;
+	PerViewUniformParameters.PrevCameraOffsetDataOffset = INDEX_NONE;
+	PerViewUniformParameters.PrevPivotOffsetDataOffset = INDEX_NONE;
+
 	TConstArrayView<FNiagaraRendererVariableInfo> VFVariables = RendererLayout->GetVFVariables_RenderThread();
 	if (SourceMode == ENiagaraRendererSourceDataMode::Particles)
 	{
@@ -314,6 +318,17 @@ FNiagaraSpriteUniformBufferRef FNiagaraRendererSprites::CreatePerViewUniformBuff
 		PerViewUniformParameters.PivotOffsetDataOffset = VFVariables[ENiagaraSpriteVFLayout::PivotOffset].GetGPUOffset();
 		PerViewUniformParameters.NormalizedAgeDataOffset = VFVariables[ENiagaraSpriteVFLayout::NormalizedAge].GetGPUOffset();
 		PerViewUniformParameters.MaterialRandomDataOffset = VFVariables[ENiagaraSpriteVFLayout::MaterialRandom].GetGPUOffset();
+		if (bAccurateMotionVectors)
+		{
+			PerViewUniformParameters.PrevPositionDataOffset = VFVariables[ENiagaraSpriteVFLayout::PrevPosition].GetGPUOffset();
+			PerViewUniformParameters.PrevVelocityDataOffset = VFVariables[ENiagaraSpriteVFLayout::PrevVelocity].GetGPUOffset();
+			PerViewUniformParameters.PrevRotationDataOffset = VFVariables[ENiagaraSpriteVFLayout::PrevRotation].GetGPUOffset();
+			PerViewUniformParameters.PrevSizeDataOffset = VFVariables[ENiagaraSpriteVFLayout::PrevSize].GetGPUOffset();
+			PerViewUniformParameters.PrevFacingDataOffset = VFVariables[ENiagaraSpriteVFLayout::PrevFacing].GetGPUOffset();
+			PerViewUniformParameters.PrevAlignmentDataOffset = VFVariables[ENiagaraSpriteVFLayout::PrevAlignment].GetGPUOffset();
+			PerViewUniformParameters.PrevCameraOffsetDataOffset = VFVariables[ENiagaraSpriteVFLayout::PrevCameraOffset].GetGPUOffset();
+			PerViewUniformParameters.PrevPivotOffsetDataOffset = VFVariables[ENiagaraSpriteVFLayout::PrevPivotOffset].GetGPUOffset();
+		}
 	}
 	else if (SourceMode == ENiagaraRendererSourceDataMode::Emitter) // Clear all these out because we will be using the defaults to specify them
 	{
@@ -347,71 +362,128 @@ FNiagaraSpriteUniformBufferRef FNiagaraRendererSprites::CreatePerViewUniformBuff
 
 	if (bSetAnyBoundVars && DynamicDataSprites)
 	{
-		for (int32 i = 0; i < ENiagaraSpriteVFLayout::Type::Num; i++)
+		const int32 NumLayoutVars = bAccurateMotionVectors ? ENiagaraSpriteVFLayout::Num_Max : ENiagaraSpriteVFLayout::Num_Default;
+		for (int32 i = 0; i < NumLayoutVars; i++)
 		{
 			if (VFBoundOffsetsInParamStore[i] != INDEX_NONE && DynamicDataSprites->ParameterDataBound.IsValidIndex(VFBoundOffsetsInParamStore[i]))
 			{
 				switch (i)
 				{
 				case ENiagaraSpriteVFLayout::Type::Position:
-					memcpy(&PerViewUniformParameters.DefaultPos, DynamicDataSprites->ParameterDataBound.GetData() + VFBoundOffsetsInParamStore[i], sizeof(FVector));
+					FMemory::Memcpy(&PerViewUniformParameters.DefaultPos, DynamicDataSprites->ParameterDataBound.GetData() + VFBoundOffsetsInParamStore[i], sizeof(FVector));
 					break;
 				case ENiagaraSpriteVFLayout::Type::Color:
-					memcpy(&PerViewUniformParameters.DefaultColor, DynamicDataSprites->ParameterDataBound.GetData() + VFBoundOffsetsInParamStore[i], sizeof(FLinearColor));
+					FMemory::Memcpy(&PerViewUniformParameters.DefaultColor, DynamicDataSprites->ParameterDataBound.GetData() + VFBoundOffsetsInParamStore[i], sizeof(FLinearColor));
 					break;
 				case ENiagaraSpriteVFLayout::Type::Velocity:
-					memcpy(&PerViewUniformParameters.DefaultVelocity, DynamicDataSprites->ParameterDataBound.GetData() + VFBoundOffsetsInParamStore[i], sizeof(FVector));
+					FMemory::Memcpy(&PerViewUniformParameters.DefaultVelocity, DynamicDataSprites->ParameterDataBound.GetData() + VFBoundOffsetsInParamStore[i], sizeof(FVector));
 					break;
 				case ENiagaraSpriteVFLayout::Type::Rotation:
-					memcpy(&PerViewUniformParameters.DefaultRotation, DynamicDataSprites->ParameterDataBound.GetData() + VFBoundOffsetsInParamStore[i], sizeof(float));
+					FMemory::Memcpy(&PerViewUniformParameters.DefaultRotation, DynamicDataSprites->ParameterDataBound.GetData() + VFBoundOffsetsInParamStore[i], sizeof(float));
 					break;
 				case ENiagaraSpriteVFLayout::Type::Size:
-					memcpy(&PerViewUniformParameters.DefaultSize, DynamicDataSprites->ParameterDataBound.GetData() + VFBoundOffsetsInParamStore[i], sizeof(FVector2D));
+					FMemory::Memcpy(&PerViewUniformParameters.DefaultSize, DynamicDataSprites->ParameterDataBound.GetData() + VFBoundOffsetsInParamStore[i], sizeof(FVector2D));
 					break;
 				case ENiagaraSpriteVFLayout::Type::Facing:
-					memcpy(&PerViewUniformParameters.DefaultFacing, DynamicDataSprites->ParameterDataBound.GetData() + VFBoundOffsetsInParamStore[i], sizeof(FVector));
+					FMemory::Memcpy(&PerViewUniformParameters.DefaultFacing, DynamicDataSprites->ParameterDataBound.GetData() + VFBoundOffsetsInParamStore[i], sizeof(FVector));
 					bCustomFacingSet = true;
 					break;
 				case ENiagaraSpriteVFLayout::Type::Alignment:
-					memcpy(&PerViewUniformParameters.DefaultAlignment, DynamicDataSprites->ParameterDataBound.GetData() + VFBoundOffsetsInParamStore[i], sizeof(FVector));
+					FMemory::Memcpy(&PerViewUniformParameters.DefaultAlignment, DynamicDataSprites->ParameterDataBound.GetData() + VFBoundOffsetsInParamStore[i], sizeof(FVector));
 					bCustomAlignmentSet = true;
 					break;
 				case ENiagaraSpriteVFLayout::Type::SubImage:
-					memcpy(&PerViewUniformParameters.DefaultSubImage, DynamicDataSprites->ParameterDataBound.GetData() + VFBoundOffsetsInParamStore[i], sizeof(float));
+					FMemory::Memcpy(&PerViewUniformParameters.DefaultSubImage, DynamicDataSprites->ParameterDataBound.GetData() + VFBoundOffsetsInParamStore[i], sizeof(float));
 					break;
 				case ENiagaraSpriteVFLayout::Type::MaterialParam0:
-					memcpy(&PerViewUniformParameters.DefaultDynamicMaterialParameter0, DynamicDataSprites->ParameterDataBound.GetData() + VFBoundOffsetsInParamStore[i], sizeof(FVector4));
+					FMemory::Memcpy(&PerViewUniformParameters.DefaultDynamicMaterialParameter0, DynamicDataSprites->ParameterDataBound.GetData() + VFBoundOffsetsInParamStore[i], sizeof(FVector4));
 					PerViewUniformParameters.MaterialParamValidMask |= 0x1;
 					break;
 				case ENiagaraSpriteVFLayout::Type::MaterialParam1:
-					memcpy(&PerViewUniformParameters.DefaultDynamicMaterialParameter1, DynamicDataSprites->ParameterDataBound.GetData() + VFBoundOffsetsInParamStore[i], sizeof(FVector4));
+					FMemory::Memcpy(&PerViewUniformParameters.DefaultDynamicMaterialParameter1, DynamicDataSprites->ParameterDataBound.GetData() + VFBoundOffsetsInParamStore[i], sizeof(FVector4));
 					PerViewUniformParameters.MaterialParamValidMask |= 0x2;
 					break;
 				case ENiagaraSpriteVFLayout::Type::MaterialParam2:
-					memcpy(&PerViewUniformParameters.DefaultDynamicMaterialParameter2, DynamicDataSprites->ParameterDataBound.GetData() + VFBoundOffsetsInParamStore[i], sizeof(FVector4));
+					FMemory::Memcpy(&PerViewUniformParameters.DefaultDynamicMaterialParameter2, DynamicDataSprites->ParameterDataBound.GetData() + VFBoundOffsetsInParamStore[i], sizeof(FVector4));
 					PerViewUniformParameters.MaterialParamValidMask |= 0x4;
 					break;
 				case ENiagaraSpriteVFLayout::Type::MaterialParam3:
-					memcpy(&PerViewUniformParameters.DefaultDynamicMaterialParameter3, DynamicDataSprites->ParameterDataBound.GetData() + VFBoundOffsetsInParamStore[i], sizeof(FVector4));
+					FMemory::Memcpy(&PerViewUniformParameters.DefaultDynamicMaterialParameter3, DynamicDataSprites->ParameterDataBound.GetData() + VFBoundOffsetsInParamStore[i], sizeof(FVector4));
 					PerViewUniformParameters.MaterialParamValidMask |= 0x8;
 					break;
 				case ENiagaraSpriteVFLayout::Type::CameraOffset:
-					memcpy(&PerViewUniformParameters.DefaultCamOffset, DynamicDataSprites->ParameterDataBound.GetData() + VFBoundOffsetsInParamStore[i], sizeof(float));
+					FMemory::Memcpy(&PerViewUniformParameters.DefaultCamOffset, DynamicDataSprites->ParameterDataBound.GetData() + VFBoundOffsetsInParamStore[i], sizeof(float));
 					break;
 				case ENiagaraSpriteVFLayout::Type::UVScale:
-					memcpy(&PerViewUniformParameters.DefaultUVScale, DynamicDataSprites->ParameterDataBound.GetData() + VFBoundOffsetsInParamStore[i], sizeof(FVector2D));
+					FMemory::Memcpy(&PerViewUniformParameters.DefaultUVScale, DynamicDataSprites->ParameterDataBound.GetData() + VFBoundOffsetsInParamStore[i], sizeof(FVector2D));
 					break;
 				case ENiagaraSpriteVFLayout::Type::PivotOffset:
-					memcpy(&PerViewUniformParameters.DefaultPivotOffset, DynamicDataSprites->ParameterDataBound.GetData() + VFBoundOffsetsInParamStore[i], sizeof(FVector2D));
+					FMemory::Memcpy(&PerViewUniformParameters.DefaultPivotOffset, DynamicDataSprites->ParameterDataBound.GetData() + VFBoundOffsetsInParamStore[i], sizeof(FVector2D));
 					break;
 				case ENiagaraSpriteVFLayout::Type::MaterialRandom:
-					memcpy(&PerViewUniformParameters.DefaultMatRandom, DynamicDataSprites->ParameterDataBound.GetData() + VFBoundOffsetsInParamStore[i], sizeof(float));
+					FMemory::Memcpy(&PerViewUniformParameters.DefaultMatRandom, DynamicDataSprites->ParameterDataBound.GetData() + VFBoundOffsetsInParamStore[i], sizeof(float));
 					break;
 				case ENiagaraSpriteVFLayout::Type::CustomSorting:
 					// unsupport for now...
 					break;
 				case ENiagaraSpriteVFLayout::Type::NormalizedAge:
-					memcpy(&PerViewUniformParameters.DefaultNormAge, DynamicDataSprites->ParameterDataBound.GetData() + VFBoundOffsetsInParamStore[i], sizeof(float));
+					FMemory::Memcpy(&PerViewUniformParameters.DefaultNormAge, DynamicDataSprites->ParameterDataBound.GetData() + VFBoundOffsetsInParamStore[i], sizeof(float));
+					break;
+				case ENiagaraSpriteVFLayout::Type::PrevPosition:
+					FMemory::Memcpy(&PerViewUniformParameters.DefaultPrevPos, DynamicDataSprites->ParameterDataBound.GetData() + VFBoundOffsetsInParamStore[i], sizeof(FVector));
+					break;
+				case ENiagaraSpriteVFLayout::Type::PrevVelocity:
+					FMemory::Memcpy(&PerViewUniformParameters.DefaultPrevVelocity, DynamicDataSprites->ParameterDataBound.GetData() + VFBoundOffsetsInParamStore[i], sizeof(FVector));
+					break;
+				case ENiagaraSpriteVFLayout::Type::PrevRotation:
+					FMemory::Memcpy(&PerViewUniformParameters.DefaultPrevRotation, DynamicDataSprites->ParameterDataBound.GetData() + VFBoundOffsetsInParamStore[i], sizeof(float));
+					break;
+				case ENiagaraSpriteVFLayout::Type::PrevSize:
+					FMemory::Memcpy(&PerViewUniformParameters.DefaultPrevSize, DynamicDataSprites->ParameterDataBound.GetData() + VFBoundOffsetsInParamStore[i], sizeof(FVector2D));
+					break;
+				case ENiagaraSpriteVFLayout::Type::PrevFacing:
+					FMemory::Memcpy(&PerViewUniformParameters.DefaultPrevFacing, DynamicDataSprites->ParameterDataBound.GetData() + VFBoundOffsetsInParamStore[i], sizeof(FVector));
+					break;
+				case ENiagaraSpriteVFLayout::Type::PrevAlignment:
+					FMemory::Memcpy(&PerViewUniformParameters.DefaultPrevAlignment, DynamicDataSprites->ParameterDataBound.GetData() + VFBoundOffsetsInParamStore[i], sizeof(FVector));
+					break;
+				case ENiagaraSpriteVFLayout::Type::PrevCameraOffset:
+					FMemory::Memcpy(&PerViewUniformParameters.DefaultPrevCamOffset, DynamicDataSprites->ParameterDataBound.GetData() + VFBoundOffsetsInParamStore[i], sizeof(float));					
+					break;
+				case ENiagaraSpriteVFLayout::Type::PrevPivotOffset:
+					FMemory::Memcpy(&PerViewUniformParameters.DefaultPrevPivotOffset, DynamicDataSprites->ParameterDataBound.GetData() + VFBoundOffsetsInParamStore[i], sizeof(FVector2D));
+					break;
+				}
+			}
+			else
+			{
+				switch (i)
+				{
+				case ENiagaraSpriteVFLayout::Type::PrevPosition:
+					PerViewUniformParameters.DefaultPrevPos = PerViewUniformParameters.DefaultPos;
+					break;
+				case ENiagaraSpriteVFLayout::Type::PrevVelocity:
+					PerViewUniformParameters.DefaultPrevVelocity = PerViewUniformParameters.DefaultVelocity;
+					break;
+				case ENiagaraSpriteVFLayout::Type::PrevRotation:
+					PerViewUniformParameters.DefaultPrevRotation = PerViewUniformParameters.DefaultRotation;
+					break;
+				case ENiagaraSpriteVFLayout::Type::PrevSize:
+					PerViewUniformParameters.DefaultPrevSize = PerViewUniformParameters.DefaultSize;
+					break;
+				case ENiagaraSpriteVFLayout::Type::PrevFacing:
+					PerViewUniformParameters.DefaultPrevFacing = PerViewUniformParameters.DefaultFacing;
+					break;
+				case ENiagaraSpriteVFLayout::Type::PrevAlignment:
+					PerViewUniformParameters.DefaultPrevAlignment = PerViewUniformParameters.DefaultAlignment;
+					break;
+				case ENiagaraSpriteVFLayout::Type::PrevCameraOffset:
+					PerViewUniformParameters.DefaultPrevCamOffset = PerViewUniformParameters.DefaultCamOffset;
+					break;
+				case ENiagaraSpriteVFLayout::Type::PrevPivotOffset:
+					PerViewUniformParameters.DefaultPrevPivotOffset = PerViewUniformParameters.DefaultPivotOffset;
+					break;
+				default:
 					break;
 				}
 			}
@@ -644,10 +716,11 @@ void FNiagaraRendererSprites::CreateMeshBatchForView(
 	FNiagaraDynamicDataSprites* DynamicDataSprites,
 	FMeshBatch& MeshBatch,
 	FNiagaraSpriteVFLooseParameters& VFLooseParams,
-	FNiagaraMeshCollectorResourcesSprite& CollectorResources,
+	FMeshCollectorResourcesBase& CollectorResources,
 	const FNiagaraRendererLayout* RendererLayout
 ) const
 {
+	FNiagaraSpriteVertexFactory& VertexFactory = CollectorResources.GetVertexFactory();
 	FNiagaraDataBuffer* SourceParticleData = DynamicDataSprites->GetParticleDataToRender();
 	check(SourceParticleData);//Can be null but should be checked before here.
 	int32 NumInstances = SourceMode == ENiagaraRendererSourceDataMode::Particles ? SourceParticleData->GetNumInstances() : 1;
@@ -677,19 +750,19 @@ void FNiagaraRendererSprites::CreateMeshBatchForView(
 			ActualAlignmentMode = ENiagaraSpriteAlignment::Unaligned;
 		}
 
-		CollectorResources.VertexFactory.SetAlignmentMode((uint32)ActualAlignmentMode);
-		CollectorResources.VertexFactory.SetFacingMode((uint32)FacingMode);
+		VertexFactory.SetAlignmentMode((uint32)ActualAlignmentMode);
+		VertexFactory.SetFacingMode((uint32)FacingMode);
 	}
-	CollectorResources.VertexFactory.SetParticleFactoryType(NVFT_Sprite);
-	CollectorResources.VertexFactory.InitResource();
-	CollectorResources.VertexFactory.SetSpriteUniformBuffer(CollectorResources.UniformBuffer);
+	VertexFactory.SetParticleFactoryType(NVFT_Sprite);
+	VertexFactory.InitResource();
+	VertexFactory.SetSpriteUniformBuffer(CollectorResources.UniformBuffer);
 
-	VFLooseParams.NumCutoutVerticesPerFrame = CollectorResources.VertexFactory.GetNumCutoutVerticesPerFrame();
-	VFLooseParams.CutoutGeometry = CollectorResources.VertexFactory.GetCutoutGeometrySRV() ? CollectorResources.VertexFactory.GetCutoutGeometrySRV() : GFNiagaraNullCutoutVertexBuffer.VertexBufferSRV.GetReference();
-	VFLooseParams.ParticleAlignmentMode = CollectorResources.VertexFactory.GetAlignmentMode();
-	VFLooseParams.ParticleFacingMode = CollectorResources.VertexFactory.GetFacingMode();
-	VFLooseParams.SortedIndices = CollectorResources.VertexFactory.GetSortedIndicesSRV() ? CollectorResources.VertexFactory.GetSortedIndicesSRV() : GFNiagaraNullSortedIndicesVertexBuffer.VertexBufferSRV.GetReference();
-	VFLooseParams.SortedIndicesOffset = CollectorResources.VertexFactory.GetSortedIndicesOffset();
+	VFLooseParams.NumCutoutVerticesPerFrame = VertexFactory.GetNumCutoutVerticesPerFrame();
+	VFLooseParams.CutoutGeometry = VertexFactory.GetCutoutGeometrySRV() ? VertexFactory.GetCutoutGeometrySRV() : GFNiagaraNullCutoutVertexBuffer.VertexBufferSRV.GetReference();
+	VFLooseParams.ParticleAlignmentMode = VertexFactory.GetAlignmentMode();
+	VFLooseParams.ParticleFacingMode = VertexFactory.GetFacingMode();
+	VFLooseParams.SortedIndices = VertexFactory.GetSortedIndicesSRV() ? VertexFactory.GetSortedIndicesSRV() : GFNiagaraNullSortedIndicesVertexBuffer.VertexBufferSRV.GetReference();
+	VFLooseParams.SortedIndicesOffset = VertexFactory.GetSortedIndicesOffset();
 
 	const bool bGPUCulled = CulledGPUParticleCountOffset != INDEX_NONE;
 	auto IndirectDraw = FNiagaraGPUInstanceCountManager::FIndirectArgSlot();
@@ -715,9 +788,9 @@ void FNiagaraRendererSprites::CreateMeshBatchForView(
 		VFLooseParams.IndirectArgsOffset = 0;
 	}
 
-	CollectorResources.VertexFactory.LooseParameterUniformBuffer = FNiagaraSpriteVFLooseParametersRef::CreateUniformBufferImmediate(VFLooseParams, UniformBuffer_SingleFrame);
+	VertexFactory.LooseParameterUniformBuffer = FNiagaraSpriteVFLooseParametersRef::CreateUniformBufferImmediate(VFLooseParams, UniformBuffer_SingleFrame);
 
-	MeshBatch.VertexFactory = &CollectorResources.VertexFactory;
+	MeshBatch.VertexFactory = &VertexFactory;
 	MeshBatch.CastShadow = SceneProxy->CastsDynamicShadow();
 #if RHI_RAYTRACING
 	MeshBatch.CastRayTracedShadow = SceneProxy->CastsDynamicShadow();
@@ -790,7 +863,7 @@ void FNiagaraRendererSprites::GetDynamicMeshElements(const TArray<const FSceneVi
 	if (bVisTagInParamStore && DynamicDataSprites->ParameterDataBound.IsValidIndex(RendererVisTagOffset))
 	{
 		int32 VisTag = 0;
-		memcpy(&VisTag, DynamicDataSprites->ParameterDataBound.GetData() + RendererVisTagOffset, sizeof(int32));
+		FMemory::Memcpy(&VisTag, DynamicDataSprites->ParameterDataBound.GetData() + RendererVisTagOffset, sizeof(int32));
 		if (RendererVisibility != VisTag)
 		{
 			return;
@@ -829,7 +902,7 @@ void FNiagaraRendererSprites::GetDynamicMeshElements(const TArray<const FSceneVi
 				if (BoundPosOffset != INDEX_NONE && DynamicDataSprites->ParameterDataBound.IsValidIndex(BoundPosOffset))
 				{
 					// retrieve the reference position from the parameter store
-					memcpy(&RefPosition, DynamicDataSprites->ParameterDataBound.GetData() + BoundPosOffset, sizeof(FVector));
+					FMemory::Memcpy(&RefPosition, DynamicDataSprites->ParameterDataBound.GetData() + BoundPosOffset, sizeof(FVector));
 					if (bLocalSpace)
 					{
 						RefPosition = SceneProxy->GetLocalToWorld().TransformPosition(RefPosition);
@@ -848,14 +921,23 @@ void FNiagaraRendererSprites::GetDynamicMeshElements(const TArray<const FSceneVi
 				}
 			}
 
-			FNiagaraMeshCollectorResourcesSprite& CollectorResources = Collector.AllocateOneFrameResource<FNiagaraMeshCollectorResourcesSprite>();
+			FMeshCollectorResourcesBase* CollectorResources;
+			if (bAccurateMotionVectors)
+			{
+				CollectorResources = &Collector.AllocateOneFrameResource<FMeshCollectorResourcesEx>();
+			}
+			else
+			{
+				CollectorResources = &Collector.AllocateOneFrameResource<FMeshCollectorResources>();
+			}
+
 			FNiagaraSpriteVFLooseParameters VFLooseParams;
 			int32 CulledGPUParticleCountOffset = INDEX_NONE;
-			SetVertexFactoryParticleData(CollectorResources.VertexFactory, CulledGPUParticleCountOffset, DynamicDataSprites, CPUSimParticleDataAllocation, View, VFLooseParams, SceneProxy, RendererLayout);
-			CollectorResources.UniformBuffer = CreatePerViewUniformBuffer(View, ViewFamily, SceneProxy, RendererLayout, DynamicDataSprites);
+			SetVertexFactoryParticleData(CollectorResources->GetVertexFactory(), CulledGPUParticleCountOffset, DynamicDataSprites, CPUSimParticleDataAllocation, View, VFLooseParams, SceneProxy, RendererLayout);
+			CollectorResources->UniformBuffer = CreatePerViewUniformBuffer(View, ViewFamily, SceneProxy, RendererLayout, DynamicDataSprites);
 			FMeshBatch& MeshBatch = Collector.AllocateMesh();
 
-			CreateMeshBatchForView(View, ViewFamily, SceneProxy, CulledGPUParticleCountOffset, DynamicDataSprites, MeshBatch, VFLooseParams, CollectorResources, RendererLayout);
+			CreateMeshBatchForView(View, ViewFamily, SceneProxy, CulledGPUParticleCountOffset, DynamicDataSprites, MeshBatch, VFLooseParams, *CollectorResources, RendererLayout);
 
 			Collector.AddMesh(ViewIndex, MeshBatch);
 		}
@@ -901,15 +983,24 @@ void FNiagaraRendererSprites::GetDynamicRayTracingInstances(FRayTracingMaterialG
 	const bool bNeedsGPUVis = !bVisTagInParamStore && RendererVisTagOffset != INDEX_NONE && GNiagaraGPUCulling && FNiagaraUtilities::AllowComputeShaders(Batcher->GetShaderPlatform());
 
 	{
+		FMeshCollectorResourcesBase* CollectorResources;
+		if (bAccurateMotionVectors)
+		{
+			CollectorResources = &Context.RayTracingMeshResourceCollector.AllocateOneFrameResource<FMeshCollectorResourcesEx>();
+		}
+		else
+		{
+			CollectorResources = &Context.RayTracingMeshResourceCollector.AllocateOneFrameResource<FMeshCollectorResources>();
+		}
+
 		// Setup material for our ray tracing instance
 		FCPUSimParticleDataAllocation CPUSimParticleDataAllocation = ConditionalAllocateCPUSimParticleData(DynamicDataSprites, RendererLayout, Context.RayTracingMeshResourceCollector.GetDynamicReadBuffer(), bNeedsGPUVis);
-		FNiagaraMeshCollectorResourcesSprite& CollectorResources = Context.RayTracingMeshResourceCollector.AllocateOneFrameResource<FNiagaraMeshCollectorResourcesSprite>();
 		FNiagaraSpriteVFLooseParameters VFLooseParams;
 		int32 CulledGPUParticleCountOffset = INDEX_NONE;
-		SetVertexFactoryParticleData(CollectorResources.VertexFactory, CulledGPUParticleCountOffset, DynamicDataSprites, CPUSimParticleDataAllocation, Context.ReferenceView, VFLooseParams, SceneProxy, RendererLayout);
-		CollectorResources.UniformBuffer = CreatePerViewUniformBuffer(Context.ReferenceView, Context.ReferenceViewFamily, SceneProxy, RendererLayout, DynamicDataSprites);
+		SetVertexFactoryParticleData(CollectorResources->GetVertexFactory(), CulledGPUParticleCountOffset, DynamicDataSprites, CPUSimParticleDataAllocation, Context.ReferenceView, VFLooseParams, SceneProxy, RendererLayout);
+		CollectorResources->UniformBuffer = CreatePerViewUniformBuffer(Context.ReferenceView, Context.ReferenceViewFamily, SceneProxy, RendererLayout, DynamicDataSprites);
 		FMeshBatch MeshBatch;
-		CreateMeshBatchForView(Context.ReferenceView, Context.ReferenceViewFamily, SceneProxy, CulledGPUParticleCountOffset, DynamicDataSprites, MeshBatch, VFLooseParams, CollectorResources, RendererLayout);
+		CreateMeshBatchForView(Context.ReferenceView, Context.ReferenceViewFamily, SceneProxy, CulledGPUParticleCountOffset, DynamicDataSprites, MeshBatch, VFLooseParams, *CollectorResources, RendererLayout);
 
 		RayTracingInstance.Materials.Add(MeshBatch);
 
diff --git a/Engine/Plugins/FX/Niagara/Source/Niagara/Private/NiagaraScript.cpp b/Engine/Plugins/FX/Niagara/Source/Niagara/Private/NiagaraScript.cpp
--- a/Engine/Plugins/FX/Niagara/Source/Niagara/Private/NiagaraScript.cpp
+++ b/Engine/Plugins/FX/Niagara/Source/Niagara/Private/NiagaraScript.cpp
@@ -85,7 +85,7 @@ UNiagaraScriptSourceBase::UNiagaraScriptSourceBase(const FObjectInitializer& Obj
 }
 
 
-FNiagaraVMExecutableData::FNiagaraVMExecutableData() 
+FNiagaraVMExecutableData::FNiagaraVMExecutableData()
 	: NumTempRegisters(0)
 	, NumUserPtrs(0)
 #if WITH_EDITORONLY_DATA
@@ -105,7 +105,7 @@ bool FNiagaraVMExecutableData::IsValid() const
 	return LastCompileStatus != ENiagaraScriptCompileStatus::NCS_Unknown;
 }
 
-void FNiagaraVMExecutableData::Reset() 
+void FNiagaraVMExecutableData::Reset()
 {
 	*this = FNiagaraVMExecutableData();
 }
@@ -138,7 +138,7 @@ void FNiagaraVMExecutableData::BakeScriptLiterals(TArray<uint8>& OutLiterals) co
 	{
 		const FNiagaraVariable& Variable = Variables[Index];
 		const int32 VariableSize = Variable.GetAllocatedSizeInBytes();
-		
+
 		FMemory::Memcpy(LiteralData, Variable.GetData(), VariableSize);
 		LiteralData += VariableSize;
 	}
@@ -155,6 +155,21 @@ void FNiagaraVMExecutableDataId::Invalidate()
 	*this = FNiagaraVMExecutableDataId();
 }
 
+#if WITH_EDITORONLY_DATA
+
+TArray<FString> FNiagaraVMExecutableDataId::GetAdditionalVariableStrings()
+{
+	TArray<FString> Vars;
+	for (const FNiagaraVariableBase& Var : AdditionalVariables)
+	{
+		Vars.Emplace(Var.GetName().ToString() + TEXT(" ") + Var.GetType().GetName());
+	}
+
+	return Vars;
+}
+
+#endif
+
 bool FNiagaraVMExecutableDataId::HasInterpolatedParameters() const
 {
 	return bInterpolatedSpawn;
@@ -174,7 +189,7 @@ bool FNiagaraVMExecutableDataId::RequiresPersistentIDs() const
 bool FNiagaraVMExecutableDataId::operator==(const FNiagaraVMExecutableDataId& ReferenceSet) const
 {
 	if (CompilerVersionID != ReferenceSet.CompilerVersionID ||
-		ScriptUsageType != ReferenceSet.ScriptUsageType || 
+		ScriptUsageType != ReferenceSet.ScriptUsageType ||
 		ScriptUsageTypeID != ReferenceSet.ScriptUsageTypeID ||
 #if WITH_EDITORONLY_DATA
 		BaseScriptCompileHash != ReferenceSet.BaseScriptCompileHash ||
@@ -205,6 +220,7 @@ bool FNiagaraVMExecutableDataId::operator==(const FNiagaraVMExecutableDataId& Re
 		return false;
 	}
 
+
 	for (int32 Idx = 0; Idx < ReferenceSet.AdditionalDefines.Num(); Idx++)
 	{
 		const FString& ReferenceStr = ReferenceSet.AdditionalDefines[Idx];
@@ -214,6 +230,21 @@ bool FNiagaraVMExecutableDataId::operator==(const FNiagaraVMExecutableDataId& Re
 			return false;
 		}
 	}
+
+	if (AdditionalVariables.Num() != ReferenceSet.AdditionalVariables.Num())
+	{
+		return false;
+	}
+
+	for (int32 Idx = 0; Idx < ReferenceSet.AdditionalVariables.Num(); Idx++)
+	{
+		const FNiagaraVariableBase& ReferenceVar = ReferenceSet.AdditionalVariables[Idx];
+
+		if (AdditionalVariables[Idx] != ReferenceVar)
+		{
+			return false;
+		}
+	}
 #endif
 
 
@@ -230,14 +261,14 @@ void FNiagaraVMExecutableDataId::AppendKeyString(FString& KeyString, const FStri
 		KeyString += TEXT(" [ScriptUsageType]");
 	}
 	KeyString += Delimiter;
-	
+
 	KeyString += CompilerVersionID.ToString();
 	if (bAppendObjectForDebugging)
 	{
 		KeyString += TEXT(" [CompilerVersionID]");
 	}
 	KeyString += Delimiter;
-	
+
 	KeyString += BaseScriptCompileHash.ToString();
 	if (bAppendObjectForDebugging)
 	{
@@ -264,7 +295,15 @@ void FNiagaraVMExecutableDataId::AppendKeyString(FString& KeyString, const FStri
 		KeyString += AdditionalDefines[Idx];
 		KeyString += Delimiter;
 	}
-	
+
+	for (int32 Idx = 0; Idx < AdditionalVariables.Num(); Idx++)
+	{
+		KeyString += AdditionalVariables[Idx].GetName().ToString();
+		KeyString += Delimiter;
+		KeyString += AdditionalVariables[Idx].GetType().GetName();
+		KeyString += Delimiter;
+	}
+
 	// Add any referenced script compile hashes to the key so that we will recompile when they are changed
 	for (int32 HashIndex = 0; HashIndex < ReferencedCompileHashes.Num(); HashIndex++)
 	{
@@ -304,7 +343,7 @@ UNiagaraScript::UNiagaraScript(const FObjectInitializer& ObjectInitializer)
 	ScriptResource->OnCompilationComplete().AddUniqueDynamic(this, &UNiagaraScript::RaiseOnGPUCompilationComplete);
 
 	RapidIterationParameters.DebugName = *GetFullName();
-#endif	
+#endif
 }
 
 UNiagaraScript::~UNiagaraScript()
@@ -409,7 +448,7 @@ void UNiagaraScript::ComputeVMCompilationId(FNiagaraVMExecutableDataId& Id) cons
 				Id.AdditionalDefines.Add(TEXT("TrimAttributes"));
 
 				TArray<FString> PreserveAttributes;
-				
+
 				// preserve the attributes that have been defined on the emitter directly
 				for (const FString& Attribute : Emitter->AttributesToPreserve)
 				{
@@ -432,9 +471,26 @@ void UNiagaraScript::ComputeVMCompilationId(FNiagaraVMExecutableDataId& Id) cons
 
 				Id.AdditionalDefines.Append(PreserveAttributes);
 			}
+
+			// Gather additional variables from renderers
+			for (UNiagaraRendererProperties* RendererProperty : Emitter->GetRenderers())
+			{
+				TArray<FNiagaraVariableBase> AdditionalVariables;
+				RendererProperty->GetAdditionalVariables(AdditionalVariables);
+				for (const FNiagaraVariableBase& AdditionalVariable : AdditionalVariables)
+				{
+					if (AdditionalVariable.IsValid())
+					{
+						Id.AdditionalVariables.AddUnique(AdditionalVariable);
+					}
+				}
+			}
+
+			// Sort the additional variables by name lexically so they are always in the same order
+			Id.AdditionalVariables.Sort([](const FNiagaraVariableBase& A, const FNiagaraVariableBase& B) { return A.GetName().LexicalLess(B.GetName()); });
 		}
 
-		if ((Emitter->bInterpolatedSpawning && Usage == ENiagaraScriptUsage::ParticleGPUComputeScript) || 
+		if ((Emitter->bInterpolatedSpawning && Usage == ENiagaraScriptUsage::ParticleGPUComputeScript) ||
 			(Emitter->bInterpolatedSpawning && Usage == ENiagaraScriptUsage::ParticleSpawnScript) ||
 			Usage == ENiagaraScriptUsage::ParticleSpawnScriptInterpolated)
 		{
@@ -537,8 +593,8 @@ void UNiagaraScript::ComputeVMCompilationId(FNiagaraVMExecutableDataId& Id) cons
 	}
 
 
-	// If we aren't using rapid iteration parameters, we need to bake them into the hashstate for the compile id. This 
-	// makes their values part of the lookup. 
+	// If we aren't using rapid iteration parameters, we need to bake them into the hashstate for the compile id. This
+	// makes their values part of the lookup.
 	if (false == Id.bUsesRapidIterationParams)
 	{
 		FSHA1 HashState;
@@ -564,7 +620,7 @@ void UNiagaraScript::ComputeVMCompilationId(FNiagaraVMExecutableDataId& Id) cons
 					//UE_LOG(LogNiagara, Display, TEXT("Param %s %s %s"), *VarTypeName, *VarName, *ByteStr);
 					HashState.Update(VarData, Vars[i].GetType().GetSize());
 				}
-			}			
+			}
 		}
 		HashState.Final();
 
@@ -835,10 +891,10 @@ void UNiagaraScript::PreSave(const class ITargetPlatform* TargetPlatform)
 		UE_LOG(LogNiagara, Log, TEXT("Cached DataInterfaceInfos:"));
 		for (const FNiagaraScriptDataInterfaceInfo& DataInterfaceCacheInfo : CachedDefaultDataInterfaces)
 		{
-			UE_LOG(LogNiagara, Log, TEXT("Name:%s, Type: %s, Path:%s"), 
-				*DataInterfaceCacheInfo.Name.ToString(), *DataInterfaceCacheInfo.Type.GetName(), 
-				DataInterfaceCacheInfo.DataInterface != nullptr 
-					? *DataInterfaceCacheInfo.DataInterface->GetPathName() 
+			UE_LOG(LogNiagara, Log, TEXT("Name:%s, Type: %s, Path:%s"),
+				*DataInterfaceCacheInfo.Name.ToString(), *DataInterfaceCacheInfo.Type.GetName(),
+				DataInterfaceCacheInfo.DataInterface != nullptr
+					? *DataInterfaceCacheInfo.DataInterface->GetPathName()
 					: TEXT("None"));
 		}
 
@@ -1020,7 +1076,7 @@ bool UNiagaraScript::IsUsageDependentOn(ENiagaraScriptUsage InUsageA, ENiagaraSc
 		return false;
 	}
 
-	// Usages of the same phase are interdependent because we copy the attributes from one to the other and if those got 
+	// Usages of the same phase are interdependent because we copy the attributes from one to the other and if those got
 	// out of sync, there could be problems.
 
 	if ((InUsageA == ENiagaraScriptUsage::ParticleSpawnScript || InUsageA == ENiagaraScriptUsage::ParticleSpawnScriptInterpolated || InUsageA == ENiagaraScriptUsage::ParticleUpdateScript || InUsageA == ENiagaraScriptUsage::ParticleEventScript)
@@ -1199,7 +1255,7 @@ void UNiagaraScript::PostLoad()
 		}
 	}
 #endif
-	
+
 	ProcessSerializedShaderMaps();
 
 #if WITH_EDITORONLY_DATA
@@ -1635,7 +1691,7 @@ UNiagaraDataInterface* UNiagaraScript::CopyDataInterface(UNiagaraDataInterface*
 }
 
 
-UNiagaraDataInterface* ResolveDataInterface(FNiagaraCompileRequestDataBase* InBase, FName VariableName) 
+UNiagaraDataInterface* ResolveDataInterface(FNiagaraCompileRequestDataBase* InBase, FName VariableName)
 {
 	UNiagaraDataInterface* const* FoundDI = InBase->GetObjectNameMap().Find(VariableName);
 	if (FoundDI && *(FoundDI))
@@ -1645,7 +1701,7 @@ UNiagaraDataInterface* ResolveDataInterface(FNiagaraCompileRequestDataBase* InBa
 	return nullptr;
 }
 
-void DumpNameMap(FNiagaraCompileRequestDataBase* InBase) 
+void DumpNameMap(FNiagaraCompileRequestDataBase* InBase)
 {
 	for (const TPair<FName, UNiagaraDataInterface*>& Pair : InBase->GetObjectNameMap())
 	{
@@ -1664,10 +1720,10 @@ void UNiagaraScript::SetVMCompilationResults(const FNiagaraVMExecutableDataId& I
 	CachedParameterCollectionReferences.Empty();
 	// Proactively clear out the script resource, because it might be stale now.
 	ScriptResource->Invalidate();
-	
+
 	if (CachedScriptVM.LastCompileStatus == ENiagaraScriptCompileStatus::NCS_Error)
 	{
-		// Compiler errors for Niagara will have a strong UI impact but the game should still function properly, there 
+		// Compiler errors for Niagara will have a strong UI impact but the game should still function properly, there
 		// will just be oddities in the visuals. It should be acted upon, but in no way should the game be blocked from
 		// a successful cook because of it. Therefore, we do a warning.
 		if (!CachedScriptVM.ErrorMsg.IsEmpty())
@@ -1694,8 +1750,8 @@ void UNiagaraScript::SetVMCompilationResults(const FNiagaraVMExecutableDataId& I
 		{
 			CachedDefaultDataInterfaces[Idx].DataInterface = CopyDataInterface(FindDIById, this);
 			check(CachedDefaultDataInterfaces[Idx].DataInterface != nullptr);
-		}			
-		
+		}
+
 		if (CachedDefaultDataInterfaces[Idx].DataInterface == nullptr)
 		{
 			// Use the CDO since we didn't have a default..
@@ -1730,7 +1786,7 @@ void UNiagaraScript::SetVMCompilationResults(const FNiagaraVMExecutableDataId& I
 	}
 
 	InvalidateExecutionReadyParameterStores();
-	
+
 	AsyncOptimizeByteCode();
 
 	OnVMScriptCompiled().Broadcast(this);
@@ -1835,6 +1891,7 @@ bool UNiagaraScript::RequestExternallyManagedAsyncCompile(const TSharedPtr<FNiag
 
 		FNiagaraCompileOptions Options(GetUsage(), GetUsageId(), ModuleUsageBitmask, GetPathName(), GetFullName(), GetName());
 		Options.AdditionalDefines = LastGeneratedVMId.AdditionalDefines;
+		Options.AdditionalVariables = LastGeneratedVMId.AdditionalVariables;
 		OutAsyncHandle = NiagaraModule.StartScriptCompileJob(RequestData.Get(), Options);
 		UE_LOG(LogNiagara, Verbose, TEXT("Script '%s' is requesting compile.."), *GetFullName());
 		return true;
@@ -1995,7 +2052,7 @@ void UNiagaraScript::BeginCacheForCookedPlatformData(const ITargetPlatform *Targ
 	if (ShouldCacheShadersForCooking(TargetPlatform))
 	{
 		// Commandlets like DerivedDataCacheCommandlet call BeginCacheForCookedPlatformData directly on objects. This may mean that
-		// we have not properly gotten the HLSL script generated by the time that we get here. This does the awkward work of 
+		// we have not properly gotten the HLSL script generated by the time that we get here. This does the awkward work of
 		// waiting on the parent system to finish generating the HLSL before we can begin compiling it for the GPU.
 		UNiagaraSystem* SystemOwner = FindRootSystem();
 		if (SystemOwner)
@@ -2088,14 +2145,15 @@ void UNiagaraScript::CacheResourceShadersForCooking(EShaderPlatform ShaderPlatfo
 			{
 				return;
 			}
-			
+
 			FNiagaraShaderScript *ResourceToCache = nullptr;
 			FNiagaraShaderScript* NewResource = AllocateResource();
 			check(CachedScriptVMId.CompilerVersionID.IsValid());
 			check(CachedScriptVMId.BaseScriptCompileHash.IsValid());
 
 			NewResource->SetScript(this, TargetFeatureLevel, ShaderPlatform, CachedScriptVMId.CompilerVersionID, CachedScriptVMId.AdditionalDefines,
-				CachedScriptVMId.BaseScriptCompileHash,	CachedScriptVMId.ReferencedCompileHashes, 
+				CachedScriptVMId.GetAdditionalVariableStrings(),
+				CachedScriptVMId.BaseScriptCompileHash,	CachedScriptVMId.ReferencedCompileHashes,
 				CachedScriptVMId.bUsesRapidIterationParams, GetFriendlyName());
 			ResourceToCache = NewResource;
 
@@ -2164,7 +2222,8 @@ void UNiagaraScript::CacheResourceShadersForRendering(bool bRegenerateId, bool b
 			const EShaderPlatform ShaderPlatform = GShaderPlatformForFeatureLevel[CacheFeatureLevel];
 
 			ScriptResource->SetScript(this, CacheFeatureLevel, ShaderPlatform, CachedScriptVMId.CompilerVersionID, CachedScriptVMId.AdditionalDefines,
-				CachedScriptVMId.BaseScriptCompileHash, CachedScriptVMId.ReferencedCompileHashes, 
+				CachedScriptVMId.GetAdditionalVariableStrings(),
+				CachedScriptVMId.BaseScriptCompileHash, CachedScriptVMId.ReferencedCompileHashes,
 				CachedScriptVMId.bUsesRapidIterationParams, GetFriendlyName());
 
 			if (FNiagaraUtilities::SupportsComputeShaders(ShaderPlatform))
diff --git a/Engine/Plugins/FX/Niagara/Source/Niagara/Private/NiagaraSpriteRendererProperties.cpp b/Engine/Plugins/FX/Niagara/Source/Niagara/Private/NiagaraSpriteRendererProperties.cpp
--- a/Engine/Plugins/FX/Niagara/Source/Niagara/Private/NiagaraSpriteRendererProperties.cpp
+++ b/Engine/Plugins/FX/Niagara/Source/Niagara/Private/NiagaraSpriteRendererProperties.cpp
@@ -59,7 +59,7 @@ UNiagaraSpriteRendererProperties::UNiagaraSpriteRendererProperties()
 	FNiagaraTypeDefinition MaterialDef(UMaterialInterface::StaticClass());
 	MaterialUserParamBinding.Parameter.SetType(MaterialDef);
 
-	AttributeBindings.Reserve(19);
+	AttributeBindings.Reserve(27);
 
 	// NOTE: These bindings' indices have to align to their counterpart in ENiagaraSpriteVFLayout
 	AttributeBindings.Add(&PositionBinding);
@@ -81,6 +81,16 @@ UNiagaraSpriteRendererProperties::UNiagaraSpriteRendererProperties()
 	AttributeBindings.Add(&CustomSortingBinding);
 	AttributeBindings.Add(&NormalizedAgeBinding);
 
+	// These bindings are only actually used with accurate motion vectors (indices still need to align)
+	AttributeBindings.Add(&PrevPositionBinding);
+	AttributeBindings.Add(&PrevVelocityBinding);
+	AttributeBindings.Add(&PrevSpriteRotationBinding);
+	AttributeBindings.Add(&PrevSpriteSizeBinding);
+	AttributeBindings.Add(&PrevSpriteFacingBinding);
+	AttributeBindings.Add(&PrevSpriteAlignmentBinding);
+	AttributeBindings.Add(&PrevCameraOffsetBinding);
+	AttributeBindings.Add(&PrevPivotOffsetBinding);
+
 	// The remaining bindings are not associated with attributes in the VF layout
 	AttributeBindings.Add(&RendererVisibilityTagBinding);
 }
@@ -133,6 +143,9 @@ void UNiagaraSpriteRendererProperties::PostLoad()
 		CacheDerivedData();
 	}
 	PostLoadBindings(SourceMode);
+
+	// Fix up these bindings from their loaded source bindings
+	SetPreviousBindings(nullptr, SourceMode);
 #endif // WITH_EDITORONLY_DATA
 }
 
@@ -218,13 +231,28 @@ void UNiagaraSpriteRendererProperties::InitBindings()
 		//Default custom sorting to age
 		CustomSortingBinding = FNiagaraConstants::GetAttributeDefaultBinding(SYS_PARAM_PARTICLES_NORMALIZED_AGE);
 	}
+
+	SetPreviousBindings(nullptr, SourceMode);
+}
+
+void UNiagaraSpriteRendererProperties::SetPreviousBindings(const UNiagaraEmitter* SrcEmitter, ENiagaraRendererSourceDataMode InSourceMode)
+{
+	PrevPositionBinding.SetAsPreviousValue(PositionBinding.GetParamMapBindableVariable(), SrcEmitter, InSourceMode);
+	PrevVelocityBinding.SetAsPreviousValue(VelocityBinding.GetParamMapBindableVariable(), SrcEmitter, InSourceMode);
+	PrevSpriteRotationBinding.SetAsPreviousValue(SpriteRotationBinding.GetParamMapBindableVariable(), SrcEmitter, InSourceMode);
+	PrevSpriteSizeBinding.SetAsPreviousValue(SpriteSizeBinding.GetParamMapBindableVariable(), SrcEmitter, InSourceMode);
+	PrevSpriteFacingBinding.SetAsPreviousValue(SpriteFacingBinding.GetParamMapBindableVariable(), SrcEmitter, InSourceMode);
+	PrevSpriteAlignmentBinding.SetAsPreviousValue(SpriteAlignmentBinding.GetParamMapBindableVariable(), SrcEmitter, InSourceMode);
+	PrevCameraOffsetBinding.SetAsPreviousValue(CameraOffsetBinding.GetParamMapBindableVariable(), SrcEmitter, InSourceMode);
+	PrevPivotOffsetBinding.SetAsPreviousValue(PivotOffsetBinding.GetParamMapBindableVariable(), SrcEmitter, InSourceMode);
 }
 
 void UNiagaraSpriteRendererProperties::CacheFromCompiledData(const FNiagaraDataSetCompiledData* CompiledData)
 {
 	UpdateSourceModeDerivates(SourceMode);
 
-	RendererLayoutWithCustomSort.Initialize(ENiagaraSpriteVFLayout::Num);
+	const int32 NumLayoutVars = NeedsPreciseMotionVectors() ? ENiagaraSpriteVFLayout::Num_Max : ENiagaraSpriteVFLayout::Num_Default;
+	RendererLayoutWithCustomSort.Initialize(NumLayoutVars);
 	RendererLayoutWithCustomSort.SetVariableFromBinding(CompiledData, PositionBinding, ENiagaraSpriteVFLayout::Position);
 	RendererLayoutWithCustomSort.SetVariableFromBinding(CompiledData, VelocityBinding, ENiagaraSpriteVFLayout::Velocity);
 	RendererLayoutWithCustomSort.SetVariableFromBinding(CompiledData, ColorBinding, ENiagaraSpriteVFLayout::Color);
@@ -243,9 +271,20 @@ void UNiagaraSpriteRendererProperties::CacheFromCompiledData(const FNiagaraDataS
 	MaterialParamValidMask |= RendererLayoutWithCustomSort.SetVariableFromBinding(CompiledData, DynamicMaterial1Binding, ENiagaraSpriteVFLayout::MaterialParam1) ? 0x2 : 0;
 	MaterialParamValidMask |= RendererLayoutWithCustomSort.SetVariableFromBinding(CompiledData, DynamicMaterial2Binding, ENiagaraSpriteVFLayout::MaterialParam2) ? 0x4 : 0;
 	MaterialParamValidMask |= RendererLayoutWithCustomSort.SetVariableFromBinding(CompiledData, DynamicMaterial3Binding, ENiagaraSpriteVFLayout::MaterialParam3) ? 0x8 : 0;
+	if (NeedsPreciseMotionVectors())
+	{
+		RendererLayoutWithCustomSort.SetVariableFromBinding(CompiledData, PrevPositionBinding, ENiagaraSpriteVFLayout::PrevPosition);
+		RendererLayoutWithCustomSort.SetVariableFromBinding(CompiledData, PrevVelocityBinding, ENiagaraSpriteVFLayout::PrevVelocity);
+		RendererLayoutWithCustomSort.SetVariableFromBinding(CompiledData, PrevSpriteRotationBinding, ENiagaraSpriteVFLayout::PrevRotation);
+		RendererLayoutWithCustomSort.SetVariableFromBinding(CompiledData, PrevSpriteSizeBinding, ENiagaraSpriteVFLayout::PrevSize);
+		RendererLayoutWithCustomSort.SetVariableFromBinding(CompiledData, PrevSpriteFacingBinding, ENiagaraSpriteVFLayout::PrevFacing);
+		RendererLayoutWithCustomSort.SetVariableFromBinding(CompiledData, PrevSpriteAlignmentBinding, ENiagaraSpriteVFLayout::PrevAlignment);
+		RendererLayoutWithCustomSort.SetVariableFromBinding(CompiledData, PrevCameraOffsetBinding, ENiagaraSpriteVFLayout::PrevCameraOffset);
+		RendererLayoutWithCustomSort.SetVariableFromBinding(CompiledData, PrevPivotOffsetBinding, ENiagaraSpriteVFLayout::PrevPivotOffset);
+	}
 	RendererLayoutWithCustomSort.Finalize();
 
-	RendererLayoutWithoutCustomSort.Initialize(ENiagaraSpriteVFLayout::Num);
+	RendererLayoutWithoutCustomSort.Initialize(NumLayoutVars);
 	RendererLayoutWithoutCustomSort.SetVariableFromBinding(CompiledData, PositionBinding, ENiagaraSpriteVFLayout::Position);
 	RendererLayoutWithoutCustomSort.SetVariableFromBinding(CompiledData, VelocityBinding, ENiagaraSpriteVFLayout::Velocity);
 	RendererLayoutWithoutCustomSort.SetVariableFromBinding(CompiledData, ColorBinding, ENiagaraSpriteVFLayout::Color);
@@ -263,6 +302,17 @@ void UNiagaraSpriteRendererProperties::CacheFromCompiledData(const FNiagaraDataS
 	MaterialParamValidMask |= RendererLayoutWithoutCustomSort.SetVariableFromBinding(CompiledData, DynamicMaterial1Binding, ENiagaraSpriteVFLayout::MaterialParam1) ? 0x2 : 0;
 	MaterialParamValidMask |= RendererLayoutWithoutCustomSort.SetVariableFromBinding(CompiledData, DynamicMaterial2Binding, ENiagaraSpriteVFLayout::MaterialParam2) ? 0x4 : 0;
 	MaterialParamValidMask |= RendererLayoutWithoutCustomSort.SetVariableFromBinding(CompiledData, DynamicMaterial3Binding, ENiagaraSpriteVFLayout::MaterialParam3) ? 0x8 : 0;
+	if (NeedsPreciseMotionVectors())
+	{
+		RendererLayoutWithoutCustomSort.SetVariableFromBinding(CompiledData, PrevPositionBinding, ENiagaraSpriteVFLayout::PrevPosition);
+		RendererLayoutWithoutCustomSort.SetVariableFromBinding(CompiledData, PrevVelocityBinding, ENiagaraSpriteVFLayout::PrevVelocity);
+		RendererLayoutWithoutCustomSort.SetVariableFromBinding(CompiledData, PrevSpriteRotationBinding, ENiagaraSpriteVFLayout::PrevRotation);
+		RendererLayoutWithoutCustomSort.SetVariableFromBinding(CompiledData, PrevSpriteSizeBinding, ENiagaraSpriteVFLayout::PrevSize);
+		RendererLayoutWithoutCustomSort.SetVariableFromBinding(CompiledData, PrevSpriteFacingBinding, ENiagaraSpriteVFLayout::PrevFacing);
+		RendererLayoutWithoutCustomSort.SetVariableFromBinding(CompiledData, PrevSpriteAlignmentBinding, ENiagaraSpriteVFLayout::PrevAlignment);
+		RendererLayoutWithoutCustomSort.SetVariableFromBinding(CompiledData, PrevCameraOffsetBinding, ENiagaraSpriteVFLayout::PrevCameraOffset);
+		RendererLayoutWithoutCustomSort.SetVariableFromBinding(CompiledData, PrevPivotOffsetBinding, ENiagaraSpriteVFLayout::PrevPivotOffset);
+	}
 	RendererLayoutWithoutCustomSort.Finalize();
 
 }
@@ -292,6 +342,8 @@ bool UNiagaraSpriteRendererProperties::PopulateRequiredBindings(FNiagaraParamete
 
 void UNiagaraSpriteRendererProperties::UpdateSourceModeDerivates(ENiagaraRendererSourceDataMode InSourceMode, bool bFromPropertyEdit)
 {
+	Super::UpdateSourceModeDerivates(InSourceMode, bFromPropertyEdit);
+	
 	UNiagaraEmitter* SrcEmitter = GetTypedOuter<UNiagaraEmitter>();
 	if (SrcEmitter)
 	{
@@ -299,9 +351,9 @@ void UNiagaraSpriteRendererProperties::UpdateSourceModeDerivates(ENiagaraRendere
 		{
 			MaterialParamBinding.CacheValues(SrcEmitter);
 		}
-	}
 
-	Super::UpdateSourceModeDerivates(InSourceMode, bFromPropertyEdit);
+		SetPreviousBindings(SrcEmitter, InSourceMode);
+	}
 }
 
 #if WITH_EDITORONLY_DATA
@@ -412,7 +464,6 @@ void UNiagaraSpriteRendererProperties::RemoveVariable(const FNiagaraVariableBase
 	}
 }
 
-
 const TArray<FNiagaraVariable>& UNiagaraSpriteRendererProperties::GetOptionalAttributes()
 {
 	static TArray<FNiagaraVariable> Attrs;
@@ -441,6 +492,22 @@ const TArray<FNiagaraVariable>& UNiagaraSpriteRendererProperties::GetOptionalAtt
 	return Attrs;
 }
 
+void UNiagaraSpriteRendererProperties::GetAdditionalVariables(TArray<FNiagaraVariableBase>& OutArray) const
+{
+	if (NeedsPreciseMotionVectors())
+	{
+		OutArray.Reserve(8);
+		OutArray.Add(PrevPositionBinding.GetParamMapBindableVariable());
+		OutArray.Add(PrevVelocityBinding.GetParamMapBindableVariable());
+		OutArray.Add(PrevSpriteRotationBinding.GetParamMapBindableVariable());
+		OutArray.Add(PrevSpriteSizeBinding.GetParamMapBindableVariable());
+		OutArray.Add(PrevSpriteFacingBinding.GetParamMapBindableVariable());
+		OutArray.Add(PrevSpriteAlignmentBinding.GetParamMapBindableVariable());
+		OutArray.Add(PrevCameraOffsetBinding.GetParamMapBindableVariable());
+		OutArray.Add(PrevPivotOffsetBinding.GetParamMapBindableVariable());		
+	}
+}
+
 void UNiagaraSpriteRendererProperties::GetRendererWidgets(const FNiagaraEmitterInstance* InEmitter, TArray<TSharedPtr<SWidget>>& OutWidgets, TSharedPtr<FAssetThumbnailPool> InThumbnailPool) const
 {
 	TSharedRef<SWidget> ThumbnailWidget = SNullWidget::NullWidget;
diff --git a/Engine/Plugins/FX/Niagara/Source/Niagara/Public/NiagaraCommon.h b/Engine/Plugins/FX/Niagara/Source/Niagara/Public/NiagaraCommon.h
--- a/Engine/Plugins/FX/Niagara/Source/Niagara/Public/NiagaraCommon.h
+++ b/Engine/Plugins/FX/Niagara/Source/Niagara/Public/NiagaraCommon.h
@@ -109,6 +109,36 @@ enum class ENiagaraDefaultMode : uint8
 	FailIfPreviouslyNotSet
 };
 
+/** How to handle how Niagara rendered effects should generate motion vectors by default (can still be overridden on a case-by-case basis) */
+UENUM()
+enum class ENiagaraDefaultRendererMotionVectorSetting
+{
+	/**
+	 * Motion vectors generated are precise (ideal for motion blur and temporal anti-aliasing).
+	 * Requires relevant emitters to store more data per particle, and may affect vertex processing performance.
+	 */
+	Precise,
+	/**
+	 * Approximates the motion vectors from current velocity.
+	 * Saves memory and performance, but can result in lower quality motion blur and/or anti-aliasing.
+	 */
+	Approximate
+};
+
+/** How a given Niagara renderer should handle motion vector generation. */
+UENUM()
+enum class ENiagaraRendererMotionVectorSetting
+{
+	/** Determines the best method to employ when generating motion vectors (accurate vs. approximate) based on project and renderer settings */
+	AutoDetect,
+	/** Force motion vectors to be precise for this renderer. */
+	Precise,
+	/** Force motion vectors to be approximate for this renderer (higher performance, lower particle memory usage). */
+	Approximate,
+	/** Do not generate motion vectors (i.e. render the object as though it is stationary). */
+	Disable
+};
+
 UENUM()
 enum class ENiagaraSimTarget : uint8
 {
@@ -886,12 +916,13 @@ struct FNiagaraVariableAttributeBinding
 	bool NIAGARA_API DoesBindingExistOnSource() const { return bBindingExistsOnSource; }
 	bool NIAGARA_API CanBindToHostParameterMap() const { return bBindingExistsOnSource && !bIsCachedParticleValue; }
 	void NIAGARA_API SetValue(const FName& InValue, const UNiagaraEmitter* InEmitter, ENiagaraRendererSourceDataMode InSourceMode);
+	void NIAGARA_API SetAsPreviousValue(const FNiagaraVariableBase& Src, const UNiagaraEmitter* InEmitter, ENiagaraRendererSourceDataMode InSourceMode);
 	void NIAGARA_API CacheValues(const UNiagaraEmitter* InEmitter, ENiagaraRendererSourceDataMode InSourceMode);
 	bool NIAGARA_API RenameVariableIfMatching(const FNiagaraVariableBase& OldVariable, const FNiagaraVariableBase& NewVariable, const UNiagaraEmitter* InEmitter, ENiagaraRendererSourceDataMode InSourceMode);
 	bool NIAGARA_API Matches(const FNiagaraVariableBase& OldVariable, const UNiagaraEmitter* InEmitter, ENiagaraRendererSourceDataMode InSourceMode);
 
 #if WITH_EDITORONLY_DATA
-	NIAGARA_API const FName& GetName(ENiagaraRendererSourceDataMode InSourceMode) const;
+	NIAGARA_API const FName& GetName() const { return CachedDisplayName; }
 	NIAGARA_API FString GetDefaultValueString() const;
 #endif
 	NIAGARA_API const FNiagaraVariableBase& GetParamMapBindableVariable() const { return ParamMapVariable; }
@@ -906,7 +937,7 @@ struct FNiagaraVariableAttributeBinding
 		return RootVariable.GetValue<T>();
 	}
 
-	void NIAGARA_API Setup(const FNiagaraVariableBase& InRootVar, const FNiagaraVariableBase& InDataSetVar, const FNiagaraVariable& InDefaultValue, ENiagaraRendererSourceDataMode InSourceMode = ENiagaraRendererSourceDataMode::Particles);
+	void NIAGARA_API Setup(const FNiagaraVariableBase& InRootVar, const FNiagaraVariable& InDefaultValue, ENiagaraRendererSourceDataMode InSourceMode = ENiagaraRendererSourceDataMode::Particles);
 
 	void NIAGARA_API PostLoad(ENiagaraRendererSourceDataMode InSourceMode);
 	void NIAGARA_API Dump() const;
diff --git a/Engine/Plugins/FX/Niagara/Source/Niagara/Public/NiagaraMeshRendererProperties.h b/Engine/Plugins/FX/Niagara/Source/Niagara/Public/NiagaraMeshRendererProperties.h
--- a/Engine/Plugins/FX/Niagara/Source/Niagara/Public/NiagaraMeshRendererProperties.h
+++ b/Engine/Plugins/FX/Niagara/Source/Niagara/Public/NiagaraMeshRendererProperties.h
@@ -22,7 +22,7 @@ enum class ENiagaraMeshFacingMode : uint8
 	/** The mesh aligns it's local-space X-axis with the particle's Particles.Velocity vector.*/
 	Velocity,
 	/** Has the mesh local-space X-axis point towards the camera's position.*/
-	CameraPosition, 
+	CameraPosition,
 	/** Has the mesh local-space X-axis point towards the closest point on the camera view plane.*/
 	CameraPlane
 };
@@ -50,12 +50,12 @@ enum class ENiagaraMeshLockedAxisSpace : uint8 {
 };
 
 USTRUCT()
-struct NIAGARA_API FNiagaraMeshMaterialOverride 
+struct NIAGARA_API FNiagaraMeshMaterialOverride
 {
 	GENERATED_USTRUCT_BODY()
-public:	
+public:
 	FNiagaraMeshMaterialOverride();
-	
+
 	/** Used to upgrade a serialized FNiagaraParameterStore property to our own struct */
 	bool SerializeFromMismatchedTag(const struct FPropertyTag& Tag, FStructuredArchive::FSlot Slot);
 
@@ -85,7 +85,7 @@ namespace ENiagaraMeshVFLayout
 		Velocity,
 		Color,
 		Scale,
-		Transform,
+		Rotation,
 		MaterialRandom,
 		NormalizedAge,
 		CustomSorting,
@@ -96,7 +96,16 @@ namespace ENiagaraMeshVFLayout
 		DynamicParam3,
 		CameraOffset,
 
-		Num,
+		Num_Default,
+
+		// The remaining layout params aren't needed unless accurate motion vectors are required
+		PrevPosition = Num_Default,
+		PrevScale,
+		PrevRotation,
+		PrevCameraOffset,
+		PrevVelocity,
+
+		Num_Max,
 	};
 };
 
@@ -157,6 +166,7 @@ public:
 	virtual bool IsMaterialValidForRenderer(UMaterial* Material, FText& InvalidMessage) override;
 	virtual void FixMaterial(UMaterial* Material) override;
 	virtual const TArray<FNiagaraVariable>& GetOptionalAttributes() override;
+	virtual void GetAdditionalVariables(TArray<FNiagaraVariableBase>& OutArray) const override;
 	virtual	void GetRendererWidgets(const FNiagaraEmitterInstance* InEmitter, TArray<TSharedPtr<SWidget>>& OutWidgets, TSharedPtr<FAssetThumbnailPool> InThumbnailPool) const override;
 	virtual	void GetRendererTooltipWidgets(const FNiagaraEmitterInstance* InEmitter, TArray<TSharedPtr<SWidget>>& OutWidgets, TSharedPtr<FAssetThumbnailPool> InThumbnailPool) const override;
 	virtual void GetRendererFeedback(const UNiagaraEmitter* InEmitter, TArray<FText>& OutErrors, TArray<FText>& OutWarnings, TArray<FText>& OutInfo) const override;
@@ -174,11 +184,11 @@ public:
 	virtual ENiagaraRendererSourceDataMode GetCurrentSourceMode() const override { return SourceMode; }
 	//UNiagaraRendererProperties Interface END
 
-	void GetUsedMeshMaterials(int32 MeshIndex, const FNiagaraEmitterInstance* Emitter, TArray<UMaterialInterface*>& OutMaterials) const;
+	void GetUsedMeshMaterials(int32 MeshIndex, const FNiagaraEmitterInstance* Emitter, TArray<UMaterialInterface*>& OutMaterials) const;	
 
 	/**
 	 * The static mesh(es) to be instanced when rendering mesh particles.
-	 * 
+	 *
 	 * NOTES:
 	 * - If "Override Material" is not specified, the mesh's material is used. Override materials must have the Niagara Mesh Particles flag checked.
 	 * - If "Enable Mesh Flipbook" is specified, this mesh is assumed to be the first frame of the flipbook.
@@ -205,7 +215,7 @@ public:
 	/** If true, blends the sub-image UV lookup with its next adjacent member using the fractional part of the SubImageIndex float value as the linear interpolation factor.*/
 	UPROPERTY(EditAnywhere, Category = "SubUV", meta = (DisplayName = "Sub UV Blending Enabled"))
 	uint32 bSubImageBlend : 1;
-	
+
 	/** Enables frustum culling of individual mesh particles */
 	UPROPERTY(EditAnywhere, Category = "Visibility")
 	uint32 bEnableFrustumCulling : 1;
@@ -218,7 +228,7 @@ public:
 	UPROPERTY(EditAnywhere, Category = "Mesh Flipbook")
 	uint32 bEnableMeshFlipbook : 1;
 
-	/** The materials to be used instead of the StaticMesh's materials. Note that each material must have the Niagara Mesh Particles flag checked. If the ParticleMesh 
+	/** The materials to be used instead of the StaticMesh's materials. Note that each material must have the Niagara Mesh Particles flag checked. If the ParticleMesh
 	requires more materials than exist in this array or any entry in this array is set to None, we will use the ParticleMesh's existing Material instead.*/
 	UPROPERTY(EditAnywhere, Category = "Mesh Rendering", meta = (EditCondition = "bOverrideMaterials"))
 	TArray<FNiagaraMeshMaterialOverride> OverrideMaterials;
@@ -242,17 +252,17 @@ public:
 	/** Specifies what space the locked axis is in */
 	UPROPERTY(EditAnywhere, Category = "Mesh Rendering", meta = (EditCondition = "bLockedAxisEnable"))
 	ENiagaraMeshLockedAxisSpace LockedAxisSpace;
-	
+
 	UPROPERTY(EditAnywhere, Category = "Visibility", meta = (EditCondition = "bEnableCameraDistanceCulling", ClampMin = 0.0f))
 	float MinCameraDistance;
-	
+
 	UPROPERTY(EditAnywhere, Category = "Visibility", meta = (EditCondition = "bEnableCameraDistanceCulling", ClampMin = 0.0f))
 	float MaxCameraDistance = 1000.0f;
 
 	/** If a render visibility tag is present, particles whose tag matches this value will be visible in this renderer. */
 	UPROPERTY(EditAnywhere, Category = "Visibility")
 	uint32 RendererVisibility = 0;
-	
+
 	/** Which attribute should we use for position when generating instanced meshes?*/
 	UPROPERTY(EditAnywhere, Category = "Bindings")
 	FNiagaraVariableAttributeBinding PositionBinding;
@@ -272,7 +282,7 @@ public:
 	/** Which attribute should we use for scale when generating instanced meshes?*/
 	UPROPERTY(EditAnywhere, Category = "Bindings")
 	FNiagaraVariableAttributeBinding ScaleBinding;
-	
+
 	/** Which attribute should we use for sprite sub-image indexing when generating sprites?*/
 	UPROPERTY(EditAnywhere, Category = "Bindings")
 	FNiagaraVariableAttributeBinding SubImageIndexBinding;
@@ -288,7 +298,7 @@ public:
 	/** Which attribute should we use for dynamic material parameters when generating instanced meshes?*/
 	UPROPERTY(EditAnywhere, Category = "Bindings")
 	FNiagaraVariableAttributeBinding DynamicMaterial2Binding;
-	
+
 	/** Which attribute should we use for dynamic material parameters when generating instanced meshes?*/
 	UPROPERTY(EditAnywhere, Category = "Bindings")
 	FNiagaraVariableAttributeBinding DynamicMaterial3Binding;
@@ -320,9 +330,22 @@ public:
 	/** If this array has entries, we will create a MaterialInstanceDynamic per Emitter instance from Material and set the Material parameters using the Niagara simulation variables listed.*/
 	UPROPERTY(EditAnywhere, Category = "Bindings")
 	TArray<FNiagaraMaterialAttributeBinding> MaterialParameterBindings;
+
+	// The following bindings are not provided by the user, but are filled based on what other bindings are set to, and the value of bGenerateAccurateMotionVectors
+
+	UPROPERTY(Transient)
+	FNiagaraVariableAttributeBinding PrevPositionBinding;
+	UPROPERTY(Transient)
+	FNiagaraVariableAttributeBinding PrevScaleBinding;
+	UPROPERTY(Transient)
+	FNiagaraVariableAttributeBinding PrevMeshOrientationBinding;
+	UPROPERTY(Transient)
+	FNiagaraVariableAttributeBinding PrevCameraOffsetBinding;
+	UPROPERTY(Transient)
+	FNiagaraVariableAttributeBinding PrevVelocityBinding;
 	
 #if WITH_EDITORONLY_DATA
-	/** 
+	/**
 	 * The static mesh to use for the first frame of the flipbook. Its name will also be used to find subsequent frames of a similar name.
 	 * NOTE: The subsequent frames are expected to exist in the same content directory as the first frame of the flipbook, otherwise they
 	 * will not be found or used.
@@ -349,18 +372,17 @@ public:
 	/** The number of frames (static meshes) to be included in the flipbook. */
 	UPROPERTY(EditAnywhere, Category = "Mesh Flipbook", meta = (EditCondition = "bEnableMeshFlipbook && FirstFlipbookFrame != nullptr", ClampMin = 1, NoSpinbox = true))
 	uint32 NumFlipbookFrames;
-#endif	
+#endif
 
 	uint32 MaterialParamValidMask = 0;
 	FNiagaraRendererLayout RendererLayoutWithCustomSorting;
 	FNiagaraRendererLayout RendererLayoutWithoutCustomSorting;
 
 protected:
-	bool FindBinding(const FNiagaraUserParameterBinding& InBinding, const FNiagaraEmitterInstance* InEmitter, TArray<UMaterialInterface*>& OutMaterials);
 	void InitBindings();
-
-	void UpdateSourceModeDerivates(ENiagaraRendererSourceDataMode InSourceMode, bool bFromPropertyEdit = false);
-	virtual bool NeedsMIDsForMaterials() const { return MaterialParameterBindings.Num() > 0; }
+	void SetPreviousBindings(const UNiagaraEmitter* SrcEmitter, ENiagaraRendererSourceDataMode InSourceMode);
+	virtual void UpdateSourceModeDerivates(ENiagaraRendererSourceDataMode InSourceMode, bool bFromPropertyEdit = false) override;
+	virtual bool NeedsMIDsForMaterials() const override { return MaterialParameterBindings.Num() > 0; }	
 #if WITH_EDITORONLY_DATA
 	bool ChangeRequiresMeshListRebuild(const FProperty* Property);
 	void RebuildMeshList();
diff --git a/Engine/Plugins/FX/Niagara/Source/Niagara/Public/NiagaraRendererMeshes.h b/Engine/Plugins/FX/Niagara/Source/Niagara/Public/NiagaraRendererMeshes.h
--- a/Engine/Plugins/FX/Niagara/Source/Niagara/Public/NiagaraRendererMeshes.h
+++ b/Engine/Plugins/FX/Niagara/Source/Niagara/Public/NiagaraRendererMeshes.h
@@ -5,9 +5,9 @@ NiagaraRenderer.h: Base class for Niagara render modules
 ==============================================================================*/
 #pragma once
 
-#include "NiagaraMeshVertexFactory.h"
 #include "NiagaraRenderer.h"
 #include "NiagaraMeshRendererProperties.h"
+#include "NiagaraMeshVertexFactory.h"
 
 class FNiagaraDataSet;
 struct FNiagaraDynamicDataMesh;
@@ -34,7 +34,7 @@ public:
 #endif
 	//FNiagaraRenderer Interface END
 
-	void SetupVertexFactory(FNiagaraMeshVertexFactory *InVertexFactory, const FStaticMeshLODResources& LODResources) const;
+	void SetupVertexFactory(FNiagaraMeshVertexFactory& InVertexFactory, const FStaticMeshLODResources& LODResources) const;
 
 protected:
 	struct FParticleGPUBufferData
@@ -59,6 +59,28 @@ protected:
 		TArray<uint32, TInlineAllocator<4>> MaterialRemapTable;
 	};
 
+	class FMeshCollectorResourcesBase : public FOneFrameResource
+	{
+	public:
+		FNiagaraMeshUniformBufferRef UniformBuffer;
+
+		virtual ~FMeshCollectorResourcesBase() {}
+		virtual FNiagaraMeshVertexFactory& GetVertexFactory() = 0;
+	};
+
+	template <typename TVertexFactory>
+	class TMeshCollectorResources : public FMeshCollectorResourcesBase
+	{
+	public:
+		TVertexFactory VertexFactory;
+
+		virtual ~TMeshCollectorResources() { VertexFactory.ReleaseResource(); }
+		virtual FNiagaraMeshVertexFactory& GetVertexFactory() override { return VertexFactory; }
+	};
+
+	using FMeshCollectorResources = TMeshCollectorResources<FNiagaraMeshVertexFactory>;
+	using FMeshCollectorResourcesEx = TMeshCollectorResources<FNiagaraMeshVertexFactoryEx>;
+
 	int32 GetLODIndex(int32 MeshIndex) const;
 
 	void PrepareParticleBuffers(
@@ -122,6 +144,7 @@ private:
 	uint32 bLockedAxisEnable : 1;
 	uint32 bEnableCulling : 1;
 	uint32 bEnableFrustumCulling : 1;
+	uint32 bAccurateMotionVectors : 1;
 
 	uint32 bSubImageBlend : 1;
 	FVector2D SubImageSize;
@@ -136,7 +159,7 @@ private:
 	uint32 MaterialParamValidMask;
 	uint32 MaxSectionCount;
 
-	int32 VFBoundOffsetsInParamStore[ENiagaraMeshVFLayout::Type::Num];
+	int32 VFBoundOffsetsInParamStore[ENiagaraMeshVFLayout::Type::Num_Max];
 	uint32 bSetAnyBoundVars : 1;
 
 	const FNiagaraRendererLayout* RendererLayoutWithCustomSorting;
diff --git a/Engine/Plugins/FX/Niagara/Source/Niagara/Public/NiagaraRendererProperties.h b/Engine/Plugins/FX/Niagara/Source/Niagara/Public/NiagaraRendererProperties.h
--- a/Engine/Plugins/FX/Niagara/Source/Niagara/Public/NiagaraRendererProperties.h
+++ b/Engine/Plugins/FX/Niagara/Source/Niagara/Public/NiagaraRendererProperties.h
@@ -82,7 +82,7 @@ public:
 		return FixDescription;
 	}
 
-	bool IsDismissable() const 
+	bool IsDismissable() const
 	{
 		return Dismissable;
 	}
@@ -130,7 +130,7 @@ struct FNiagaraRendererVariableInfo
 struct FNiagaraRendererLayout
 {
 	void Initialize(int32 NumVariables);
-	bool SetVariable(const FNiagaraDataSetCompiledData* CompiledData, const FNiagaraVariable& Variable, int32 VFVarOffset);
+	bool SetVariable(const FNiagaraDataSetCompiledData* CompiledData, const FNiagaraVariableBase& Variable, int32 VFVarOffset);
 	bool SetVariableFromBinding(const FNiagaraDataSetCompiledData* CompiledData, const FNiagaraVariableAttributeBinding& VariableBinding, int32 VFVarOffset);
 	void Finalize();
 
@@ -160,16 +160,20 @@ class NIAGARA_API UNiagaraRendererProperties : public UNiagaraMergeable
 	GENERATED_BODY()
 
 public:
-	UNiagaraRendererProperties()
+	UNiagaraRendererProperties()		
 		: bIsEnabled(true)
-		, bMotionBlurEnabled(true)
+		, bMotionBlurEnabled_DEPRECATED(true)
 	{
 	}
 
-
+	//UObject Interface Begin
 	virtual void PostInitProperties() override;
-
+	virtual void PostLoad() override;
+#if WITH_EDITORONLY_DATA
+	virtual void PostEditChangeProperty(struct FPropertyChangedEvent& PropertyChangedEvent) override;
+#endif
 	//UObject Interface End
+	
 	virtual FNiagaraRenderer* CreateEmitterRenderer(ERHIFeatureLevel::Type FeatureLevel, const FNiagaraEmitterInstance* Emitter, const UNiagaraComponent* InComponent) PURE_VIRTUAL ( UNiagaraRendererProperties::CreateEmitterRenderer, return nullptr;);
 	virtual class FNiagaraBoundsCalculator* CreateBoundsCalculator() PURE_VIRTUAL(UNiagaraRendererProperties::CreateBoundsCalculator, return nullptr;);
 	virtual void GetUsedMaterials(const FNiagaraEmitterInstance* InEmitter, TArray<UMaterialInterface*>& OutMaterials) const PURE_VIRTUAL(UNiagaraRendererProperties::GetUsedMaterials,);
@@ -185,8 +189,8 @@ public:
 	virtual void GetAssetTagsForContext(const UObject* InAsset, const TArray<const UNiagaraRendererProperties*>& InProperties, TMap<FName, uint32>& NumericKeys, TMap<FName, FString>& StringKeys) const;
 
 	/** In the case that we need parameters bound in that aren't Particle variables, these should be set up here so that the data is appropriately populated after the simulation.*/
-	virtual bool PopulateRequiredBindings(FNiagaraParameterStore& InParameterStore) { return false; }
-	
+	virtual bool PopulateRequiredBindings(FNiagaraParameterStore& InParameterStore) { return false; }	
+
 #if WITH_EDITORONLY_DATA
 
 	virtual bool IsSupportedVariableForBinding(const FNiagaraVariableBase& InSourceForBinding, const FName& InTargetBindingName) const;
@@ -203,6 +207,8 @@ public:
 	virtual const TArray<FNiagaraVariable>& GetRequiredAttributes() { static TArray<FNiagaraVariable> Vars; return Vars; };
 	virtual const TArray<FNiagaraVariable>& GetOptionalAttributes() { static TArray<FNiagaraVariable> Vars; return Vars; };
 
+	virtual void GetAdditionalVariables(TArray<FNiagaraVariableBase>& OutArray) const { };
+
 	UNiagaraRendererProperties* StaticDuplicateWithNewMergeId(UObject* InOuter) const
 	{
 		return CastChecked<UNiagaraRendererProperties>(Super::StaticDuplicateWithNewMergeIdInternal(InOuter));
@@ -212,7 +218,7 @@ public:
 	virtual void GetRendererTooltipWidgets(const FNiagaraEmitterInstance* InEmitter, TArray<TSharedPtr<SWidget>>& OutWidgets, TSharedPtr<FAssetThumbnailPool> InThumbnailPool) const PURE_VIRTUAL(UNiagaraRendererProperties::GetRendererTooltipWidgets, );
 	virtual void GetRendererFeedback(const UNiagaraEmitter* InEmitter, TArray<FText>& OutErrors, TArray<FText>& OutWarnings, TArray<FText>& OutInfo) const {};
 	virtual void GetRendererFeedback(UNiagaraEmitter* InEmitter, TArray<FNiagaraRendererFeedback>& OutErrors, TArray<FNiagaraRendererFeedback>& OutWarnings, TArray<FNiagaraRendererFeedback>& OutInfo) const;
-	
+
 	// The icon to display in the niagara stack widget under the renderer section
 	virtual const FSlateBrush* GetStackIcon() const;
 
@@ -235,24 +241,29 @@ public:
 	virtual bool NeedsSystemPostTick() const { return false; }
 	/** When true, the renderer will be opted in to get its OnSystemComplete_GameThread overload called */
 	virtual bool NeedsSystemCompletion() const { return false; }
-	
+
+	bool NeedsPreciseMotionVectors() const;
+
 	/** Platforms on which this renderer is enabled. */
 	UPROPERTY(EditAnywhere, Category = "Scalability")
 	FNiagaraPlatformSet Platforms;
 
-	/** By default, emitters are drawn in the order that they are added to the system. This value will allow you to control the order in a more fine-grained manner. 
+	/** By default, emitters are drawn in the order that they are added to the system. This value will allow you to control the order in a more fine-grained manner.
 	Materials of the same type (i.e. Transparent) will draw in order from lowest to highest within the system. The default value is 0.*/
 	UPROPERTY(EditAnywhere, Category = "Sort Order")
 	int32 SortOrderHint;
-	
+
+	/** Hint about how to generate motion (velocity) vectors for this renderer. */
+	UPROPERTY(EditAnywhere, Category = "Motion Blur")
+	ENiagaraRendererMotionVectorSetting MotionVectorSetting;
+
 	UPROPERTY()
 	bool bIsEnabled;
 
-	/** Is motion blur enabled on this renderer or not, the material must also have motion blur enabled. */
-	UPROPERTY(EditAnywhere, Category = "Rendering")
-	bool bMotionBlurEnabled;
-
 protected:
+	UPROPERTY()
+	bool bMotionBlurEnabled_DEPRECATED; // This has been rolled into MotionVectorSetting
+
 	TArray<const FNiagaraVariableAttributeBinding*> AttributeBindings;
 
 	virtual void PostLoadBindings(ENiagaraRendererSourceDataMode InSourceMode);
diff --git a/Engine/Plugins/FX/Niagara/Source/Niagara/Public/NiagaraRendererSprites.h b/Engine/Plugins/FX/Niagara/Source/Niagara/Public/NiagaraRendererSprites.h
--- a/Engine/Plugins/FX/Niagara/Source/Niagara/Public/NiagaraRendererSprites.h
+++ b/Engine/Plugins/FX/Niagara/Source/Niagara/Public/NiagaraRendererSprites.h
@@ -8,6 +8,7 @@ NiagaraRendererSprites.h: Renderer for rendering Niagara particles as sprites.
 
 #include "NiagaraRenderer.h"
 #include "NiagaraSpriteRendererProperties.h"
+#include "NiagaraSpriteVertexFactory.h"
 
 struct FNiagaraDynamicDataSprites;
 
@@ -42,6 +43,29 @@ private:
 		FGlobalDynamicReadBuffer::FAllocation IntData;
 	};
 
+	/* Mesh collector classes */
+	class FMeshCollectorResourcesBase : public FOneFrameResource
+	{
+	public:
+		FNiagaraSpriteUniformBufferRef UniformBuffer;
+
+		virtual ~FMeshCollectorResourcesBase() {}
+		virtual FNiagaraSpriteVertexFactory& GetVertexFactory() = 0;
+	};
+
+	template <typename TVertexFactory>
+	class TMeshCollectorResources : public FMeshCollectorResourcesBase
+	{
+	public:
+		TVertexFactory VertexFactory;
+
+		virtual ~TMeshCollectorResources() { VertexFactory.ReleaseResource(); }
+		virtual FNiagaraSpriteVertexFactory& GetVertexFactory() override { return VertexFactory; }
+	};
+
+	using FMeshCollectorResources = TMeshCollectorResources<FNiagaraSpriteVertexFactory>;
+	using FMeshCollectorResourcesEx = TMeshCollectorResources<FNiagaraSpriteVertexFactoryEx>;
+
 	FCPUSimParticleDataAllocation ConditionalAllocateCPUSimParticleData(FNiagaraDynamicDataSprites *DynamicDataSprites, const FNiagaraRendererLayout* RendererLayout, FGlobalDynamicReadBuffer& DynamicReadBuffer, bool bNeedsGPUVis) const;
 	TUniformBufferRef<class FNiagaraSpriteUniformParameters> CreatePerViewUniformBuffer(const FSceneView* View, const FSceneViewFamily& ViewFamily, const FNiagaraSceneProxy *SceneProxy, const FNiagaraRendererLayout* RendererLayout, const FNiagaraDynamicDataSprites* DynamicDataSprites) const;
 	void SetVertexFactoryParticleData(
@@ -62,7 +86,7 @@ private:
 		FNiagaraDynamicDataSprites* DynamicDataSprites,
 		FMeshBatch& OutMeshBatch,
 		class FNiagaraSpriteVFLooseParameters& VFLooseParams,
-		class FNiagaraMeshCollectorResourcesSprite& OutCollectorResources,
+		class FMeshCollectorResourcesBase& OutCollectorResources,
 		const FNiagaraRendererLayout* RendererLayout
 	) const;
 
@@ -82,6 +106,7 @@ private:
 	uint32 bGpuLowLatencyTranslucency : 1;
 	uint32 bEnableCulling : 1;
 	uint32 bEnableDistanceCulling : 1;
+	uint32 bAccurateMotionVectors : 1;
 	uint32 bSetAnyBoundVars : 1;
 	uint32 bVisTagInParamStore : 1;
 
@@ -95,7 +120,7 @@ private:
 	int32 RendererVisTagOffset;
 	int32 RendererVisibility;
 
-	int32 VFBoundOffsetsInParamStore[ENiagaraSpriteVFLayout::Type::Num];
+	int32 VFBoundOffsetsInParamStore[ENiagaraSpriteVFLayout::Type::Num_Max];
 
 	const FNiagaraRendererLayout* RendererLayoutWithCustomSort;
 	const FNiagaraRendererLayout* RendererLayoutWithoutCustomSort;
diff --git a/Engine/Plugins/FX/Niagara/Source/Niagara/Public/NiagaraSettings.h b/Engine/Plugins/FX/Niagara/Source/Niagara/Public/NiagaraSettings.h
--- a/Engine/Plugins/FX/Niagara/Source/Niagara/Public/NiagaraSettings.h
+++ b/Engine/Plugins/FX/Niagara/Source/Niagara/Public/NiagaraSettings.h
@@ -92,6 +92,10 @@ class NIAGARA_API UNiagaraSettings : public UDeveloperSettings
 	UPROPERTY(config, EditAnywhere, Category = Renderer)
 	ENiagaraGpuBufferFormat DefaultGridFormat = ENiagaraGpuBufferFormat::HalfFloat;
 
+	/** The default setting for motion vectors in Niagara renderers */
+	UPROPERTY(config, EditAnywhere, Category = Renderer)
+	ENiagaraDefaultRendererMotionVectorSetting DefaultRendererMotionVectorSetting = ENiagaraDefaultRendererMotionVectorSetting::Precise;
+
 	UPROPERTY(config, EditAnywhere, Category=SkeletalMeshDI, meta = ( DisplayName = "Gpu Max Bone Influences", ToolTip = "Controls the maximum number of influences we allow the Skeletal Mesh Data Interface to use on the GPU.  Changing this setting requires restarting the editor.", ConfigRestartRequired = true))
 	TEnumAsByte<ENDISkelMesh_GpuMaxInfluences::Type> NDISkelMesh_GpuMaxInfluences;
 
diff --git a/Engine/Plugins/FX/Niagara/Source/Niagara/Public/NiagaraSpriteRendererProperties.h b/Engine/Plugins/FX/Niagara/Source/Niagara/Public/NiagaraSpriteRendererProperties.h
--- a/Engine/Plugins/FX/Niagara/Source/Niagara/Public/NiagaraSpriteRendererProperties.h
+++ b/Engine/Plugins/FX/Niagara/Source/Niagara/Public/NiagaraSpriteRendererProperties.h
@@ -61,7 +61,19 @@ namespace ENiagaraSpriteVFLayout
 		CustomSorting,
 		NormalizedAge,
 
-		Num,
+		Num_Default,
+
+		// The remaining layout params aren't needed unless accurate motion vectors are required
+		PrevPosition = Num_Default,
+		PrevVelocity,
+		PrevRotation,
+		PrevSize,
+		PrevFacing,
+		PrevAlignment,
+		PrevCameraOffset,
+		PrevPivotOffset,
+
+		Num_Max,
 	};
 };
 
@@ -100,6 +112,7 @@ public:
 	virtual bool IsMaterialValidForRenderer(UMaterial* Material, FText& InvalidMessage) override;
 	virtual void FixMaterial(UMaterial* Material) override;
 	virtual const TArray<FNiagaraVariable>& GetOptionalAttributes() override;
+	virtual void GetAdditionalVariables(TArray<FNiagaraVariableBase>& OutArray) const override;
 	virtual void GetRendererWidgets(const FNiagaraEmitterInstance* InEmitter, TArray<TSharedPtr<SWidget>>& OutWidgets, TSharedPtr<FAssetThumbnailPool> InThumbnailPool) const override;
 	virtual void GetRendererTooltipWidgets(const FNiagaraEmitterInstance* InEmitter, TArray<TSharedPtr<SWidget>>& OutWidgets, TSharedPtr<FAssetThumbnailPool> InThumbnailPool) const override;
 	virtual void GetRendererFeedback(const UNiagaraEmitter* InEmitter, TArray<FText>& OutErrors, TArray<FText>& OutWarnings, TArray<FText>& OutInfo) const override;
@@ -269,7 +282,24 @@ public:
 	UPROPERTY(EditAnywhere, Category = "Bindings")
 	TArray< FNiagaraMaterialAttributeBinding > MaterialParameterBindings;
 
-	void InitBindings();
+	// The following bindings are only needed for accurate motion vectors
+
+	UPROPERTY(Transient)
+	FNiagaraVariableAttributeBinding PrevPositionBinding;
+	UPROPERTY(Transient)
+	FNiagaraVariableAttributeBinding PrevVelocityBinding;
+	UPROPERTY(Transient)
+	FNiagaraVariableAttributeBinding PrevSpriteRotationBinding;
+	UPROPERTY(Transient)
+	FNiagaraVariableAttributeBinding PrevSpriteSizeBinding;
+	UPROPERTY(Transient)
+	FNiagaraVariableAttributeBinding PrevSpriteFacingBinding;
+	UPROPERTY(Transient)
+	FNiagaraVariableAttributeBinding PrevSpriteAlignmentBinding;
+	UPROPERTY(Transient)
+	FNiagaraVariableAttributeBinding PrevCameraOffsetBinding;
+	UPROPERTY(Transient)
+	FNiagaraVariableAttributeBinding PrevPivotOffsetBinding;
 
 	virtual bool NeedsMIDsForMaterials() const { return MaterialParameterBindings.Num() > 0; }
 
@@ -314,7 +344,9 @@ public:
 	uint32 MaterialParamValidMask = 0;
 	
 protected:
-	void UpdateSourceModeDerivates(ENiagaraRendererSourceDataMode InSourceMode, bool bFromPropertyEdit = false) override;
+	void InitBindings();
+	void SetPreviousBindings(const UNiagaraEmitter* SrcEmitter, ENiagaraRendererSourceDataMode InSourceMode);
+	virtual void UpdateSourceModeDerivates(ENiagaraRendererSourceDataMode InSourceMode, bool bFromPropertyEdit = false) override;
 
 private:
 	/** Derived data for this asset, generated off of SubUVTexture. */
diff --git a/Engine/Plugins/FX/Niagara/Source/NiagaraCore/Private/NiagaraCustomVersion.cpp b/Engine/Plugins/FX/Niagara/Source/NiagaraCore/Private/NiagaraCustomVersion.cpp
--- a/Engine/Plugins/FX/Niagara/Source/NiagaraCore/Private/NiagaraCustomVersion.cpp
+++ b/Engine/Plugins/FX/Niagara/Source/NiagaraCore/Private/NiagaraCustomVersion.cpp
@@ -8,4 +8,4 @@ const FGuid FNiagaraCustomVersion::GUID(0xFCF57AFA, 0x50764283, 0xB9A9E658, 0xFF
 // Register the custom version with core
 FCustomVersionRegistration GRegisterNiagaraCustomVersion(FNiagaraCustomVersion::GUID, FNiagaraCustomVersion::LatestVersion, TEXT("NiagaraVer"));
 
-const FGuid FNiagaraCustomVersion::LatestScriptCompileVersion(0x2D41776A, 0x3DBF2044, 0xAAA80348, 0x0EA7C7BD);
+const FGuid FNiagaraCustomVersion::LatestScriptCompileVersion(0xC2E1AF4A,  0xB3CF4652,  0x8BD6FF64, 0xF8BAA37D);
diff --git a/Engine/Plugins/FX/Niagara/Source/NiagaraEditor/Private/Customizations/NiagaraComponentRendererPropertiesDetails.cpp b/Engine/Plugins/FX/Niagara/Source/NiagaraEditor/Private/Customizations/NiagaraComponentRendererPropertiesDetails.cpp
--- a/Engine/Plugins/FX/Niagara/Source/NiagaraEditor/Private/Customizations/NiagaraComponentRendererPropertiesDetails.cpp
+++ b/Engine/Plugins/FX/Niagara/Source/NiagaraEditor/Private/Customizations/NiagaraComponentRendererPropertiesDetails.cpp
@@ -283,7 +283,7 @@ void FNiagaraComponentRendererPropertiesDetails::ChangePropertyBinding(TSharedPt
 	{
 		FScopedTransaction Transaction(FText::Format(LOCTEXT("ChangePropertyBinding", "Change component property binding to \"{0}\" "), FText::FromName(BindingVar.GetName())));
 		FNiagaraComponentPropertyBinding NewBinding = ToPropertyBinding(PropertyHandle, ComponentProperties);
-		NewBinding.AttributeBinding.Setup(BindingVar, FNiagaraConstants::GetAttributeAsParticleDataSetKey(BindingVar), BindingVar);
+		NewBinding.AttributeBinding.Setup(BindingVar, BindingVar);
 
 		ComponentProperties->Modify();
 		PropertyHandle->NotifyPreChange();
@@ -373,7 +373,7 @@ FText FNiagaraComponentRendererPropertiesDetails::GetCurrentBindingText(TSharedP
 		const FNiagaraComponentPropertyBinding* PropertyBinding = FindBinding(PropertyHandle);
 		if (PropertyBinding)
 		{
-			return FText::FromName(PropertyBinding->AttributeBinding.GetName(ENiagaraRendererSourceDataMode::Particles));
+			return FText::FromName(PropertyBinding->AttributeBinding.GetName());
 		}
 	}
 	return FText::FromString(TEXT("Missing"));
diff --git a/Engine/Plugins/FX/Niagara/Source/NiagaraEditor/Private/Customizations/NiagaraTypeCustomizations.cpp b/Engine/Plugins/FX/Niagara/Source/NiagaraEditor/Private/Customizations/NiagaraTypeCustomizations.cpp
--- a/Engine/Plugins/FX/Niagara/Source/NiagaraEditor/Private/Customizations/NiagaraTypeCustomizations.cpp
+++ b/Engine/Plugins/FX/Niagara/Source/NiagaraEditor/Private/Customizations/NiagaraTypeCustomizations.cpp
@@ -196,7 +196,7 @@ FName FNiagaraVariableAttributeBindingCustomization::GetVariableName() const
 {
 	if (BaseEmitter && TargetVariableBinding)
 	{
-		return (TargetVariableBinding->GetName(RenderProps->GetCurrentSourceMode()));
+		return (TargetVariableBinding->GetName());
 	}
 	return FName();
 }
@@ -205,7 +205,7 @@ FText FNiagaraVariableAttributeBindingCustomization::GetCurrentText() const
 {
 	if (BaseEmitter && TargetVariableBinding)
 	{
-		return FText::FromName(TargetVariableBinding->GetName(RenderProps->GetCurrentSourceMode()));
+		return FText::FromName(TargetVariableBinding->GetName());
 	}
 	return FText::FromString(TEXT("Missing"));
 }
@@ -216,7 +216,7 @@ FText FNiagaraVariableAttributeBindingCustomization::GetTooltipText() const
 	{
 		FString DefaultValueStr = TargetVariableBinding->GetDefaultValueString();
 
-		FText TooltipDesc = FText::Format(LOCTEXT("AttributeBindingTooltip", "Use the variable \"{0}\" if it exists, otherwise use the default \"{1}\" "), FText::FromName(TargetVariableBinding->GetName(RenderProps->GetCurrentSourceMode())),
+		FText TooltipDesc = FText::Format(LOCTEXT("AttributeBindingTooltip", "Use the variable \"{0}\" if it exists, otherwise use the default \"{1}\" "), FText::FromName(TargetVariableBinding->GetName()),
 			FText::FromString(DefaultValueStr));
 		return TooltipDesc;
 	}
diff --git a/Engine/Plugins/FX/Niagara/Source/NiagaraEditor/Private/NiagaraHlslTranslator.cpp b/Engine/Plugins/FX/Niagara/Source/NiagaraEditor/Private/NiagaraHlslTranslator.cpp
--- a/Engine/Plugins/FX/Niagara/Source/NiagaraEditor/Private/NiagaraHlslTranslator.cpp
+++ b/Engine/Plugins/FX/Niagara/Source/NiagaraEditor/Private/NiagaraHlslTranslator.cpp
@@ -760,9 +760,29 @@ FString FHlslNiagaraTranslator::BuildParameterMapHlslDefinitions(TArray<FNiagara
 	// For now we only care about attributes from the other output parameter map histories.
 	for (int32 ParamMapIdx = 0; ParamMapIdx < OtherOutputParamMapHistories.Num(); ParamMapIdx++)
 	{
-		for (int32 VarIdx = 0; VarIdx < OtherOutputParamMapHistories[ParamMapIdx].Variables.Num(); VarIdx++)
+		TArray<FNiagaraVariable> Vars = OtherOutputParamMapHistories[ParamMapIdx].Variables;
+		for (const FNiagaraVariableBase& Var : CompileOptions.AdditionalVariables)
 		{
-			FNiagaraVariable& Var = OtherOutputParamMapHistories[ParamMapIdx].Variables[VarIdx];
+			bool bAddVar = true;
+			if (FNiagaraParameterMapHistory::IsPreviousValue(Var))
+			{
+				FNiagaraVariable Source = FNiagaraParameterMapHistory::GetSourceForPreviousValue(FNiagaraVariable(Var));
+				if (!UniqueVariables.Contains(Source))
+				{
+					// Disallow the addition previous values if its source is not used
+					bAddVar = false;
+				}
+			}
+
+			if (bAddVar)
+			{
+				Vars.AddUnique(FNiagaraVariable(Var.GetType(), Var.GetName()));
+			}
+		}
+
+		for (int32 VarIdx = 0; VarIdx < Vars.Num(); VarIdx++)
+		{
+			FNiagaraVariable& Var = Vars[VarIdx];
 			if (OtherOutputParamMapHistories[ParamMapIdx].IsPrimaryDataSetOutput(Var, CompileOptions.TargetUsage))
 			{
 				int32 PreviousMax = UniqueVariables.Num();
diff --git a/Engine/Plugins/FX/Niagara/Source/NiagaraEditor/Private/NiagaraParameterMapHistory.cpp b/Engine/Plugins/FX/Niagara/Source/NiagaraEditor/Private/NiagaraParameterMapHistory.cpp
--- a/Engine/Plugins/FX/Niagara/Source/NiagaraEditor/Private/NiagaraParameterMapHistory.cpp
+++ b/Engine/Plugins/FX/Niagara/Source/NiagaraEditor/Private/NiagaraParameterMapHistory.cpp
@@ -339,7 +339,7 @@ FNiagaraVariable FNiagaraParameterMapHistory::VariableToNamespacedVariable(const
 	return RetVar;
 }
 
-bool FNiagaraParameterMapHistory::IsInNamespace(const FNiagaraVariable& InVar, const FString& Namespace)
+bool FNiagaraParameterMapHistory::IsInNamespace(const FNiagaraVariableBase& InVar, const FString& Namespace)
 {
 	if (Namespace.EndsWith(TEXT(".")))
 	{
@@ -422,7 +422,7 @@ bool FNiagaraParameterMapHistory::SplitRapidIterationParameterName(const FNiagar
 	return false;
 }
 
-bool FNiagaraParameterMapHistory::IsAttribute(const FNiagaraVariable& InVar)
+bool FNiagaraParameterMapHistory::IsAttribute(const FNiagaraVariableBase& InVar)
 {
 	return IsInNamespace(InVar, PARAM_MAP_ATTRIBUTE_STR);
 }
@@ -476,7 +476,7 @@ const UEdGraphPin* FNiagaraParameterMapHistory::GetDefaultValuePin(int32 VarIdx)
 	return nullptr;
 }
 
-bool FNiagaraParameterMapHistory::IsInitialValue(const FNiagaraVariable& InVar)
+bool FNiagaraParameterMapHistory::IsInitialValue(const FNiagaraVariableBase& InVar)
 {
 	TArray<FString> SplitName;
 	InVar.GetName().ToString().ParseIntoArray(SplitName, TEXT("."));
@@ -491,7 +491,7 @@ bool FNiagaraParameterMapHistory::IsInitialValue(const FNiagaraVariable& InVar)
 	return false;
 }
 
-bool FNiagaraParameterMapHistory::IsPreviousValue(const FNiagaraVariable& InVar)
+bool FNiagaraParameterMapHistory::IsPreviousValue(const FNiagaraVariableBase& InVar)
 {
 	TArray<FString> SplitName;
 	InVar.GetName().ToString().ParseIntoArray(SplitName, TEXT("."));
diff --git a/Engine/Plugins/FX/Niagara/Source/NiagaraEditor/Public/NiagaraParameterMapHistory.h b/Engine/Plugins/FX/Niagara/Source/NiagaraEditor/Public/NiagaraParameterMapHistory.h
--- a/Engine/Plugins/FX/Niagara/Source/NiagaraEditor/Public/NiagaraParameterMapHistory.h
+++ b/Engine/Plugins/FX/Niagara/Source/NiagaraEditor/Public/NiagaraParameterMapHistory.h
@@ -173,7 +173,7 @@ public:
 	static FString MakeSafeNamespaceString(const FString& InStr);
 
 	/** Does the variable start with this namespace?*/
-	static bool IsInNamespace(const FNiagaraVariable& InVar, const FString& Namespace);
+	static bool IsInNamespace(const FNiagaraVariableBase& InVar, const FString& Namespace);
 
 	/** Given the script type and flags, should we treat the System namespace as read-only?*/
 	static bool IsSystemNamespaceReadOnly(const UNiagaraScript* InScript);
@@ -203,13 +203,13 @@ public:
 	/**
 	* Does this parameter start with the "Particles" namespace?
 	*/
-	static bool IsAttribute(const FNiagaraVariable& InVar);
-	
+	static bool IsAttribute(const FNiagaraVariableBase& InVar);
+
 	/** Does this parameter contain the "Initial" namespace as one of its intermediate namespaces?*/
-	static bool IsInitialValue(const FNiagaraVariable& InVar);
+	static bool IsInitialValue(const FNiagaraVariableBase& InVar);
 
 	/** Does this parameter contain the "Previous" namespace as one of its intermediate namespaces?*/
-	static bool IsPreviousValue(const FNiagaraVariable& InVar);
+	static bool IsPreviousValue(const FNiagaraVariableBase& InVar);
 
 	/** Get the output node associated with this graph.*/
 	const UNiagaraNodeOutput* GetFinalOutputNode() const;
diff --git a/Engine/Plugins/FX/Niagara/Source/NiagaraShader/Private/NiagaraShader.cpp b/Engine/Plugins/FX/Niagara/Source/NiagaraShader/Private/NiagaraShader.cpp
--- a/Engine/Plugins/FX/Niagara/Source/NiagaraShader/Private/NiagaraShader.cpp
+++ b/Engine/Plugins/FX/Niagara/Source/NiagaraShader/Private/NiagaraShader.cpp
@@ -150,12 +150,17 @@ void FNiagaraShaderMapId::GetScriptHash(FSHAHash& OutHash) const
 	//HashState.Update((const uint8*)&BaseScriptID, sizeof(BaseScriptID));
 	HashState.Update(BaseCompileHash.Hash, FNiagaraCompileHash::HashSize);
 	HashState.Update((const uint8*)&FeatureLevel, sizeof(FeatureLevel));
-		
+
 	for (int32 Index = 0; Index < AdditionalDefines.Num(); Index++)
 	{
 		HashState.UpdateWithString(*AdditionalDefines[Index], AdditionalDefines[Index].Len());
 	}
 
+	for (int32 Index = 0; Index < AdditionalVariables.Num(); Index++)
+	{
+		HashState.UpdateWithString(*AdditionalVariables[Index], AdditionalVariables[Index].Len());
+	}
+
 	for (int32 Index = 0; Index < ReferencedCompileHashes.Num(); Index++)
 	{
 		HashState.Update(ReferencedCompileHashes[Index].Hash, FNiagaraCompileHash::HashSize);
@@ -195,6 +200,11 @@ bool FNiagaraShaderMapId::operator==(const FNiagaraShaderMapId& ReferenceSet) co
 		return false;
 	}
 
+	if (AdditionalVariables.Num() != ReferenceSet.AdditionalVariables.Num())
+	{
+		return false;
+	}
+
 	if (ReferencedCompileHashes.Num() != ReferenceSet.ReferencedCompileHashes.Num())
 	{
 		return false;
@@ -210,6 +220,16 @@ bool FNiagaraShaderMapId::operator==(const FNiagaraShaderMapId& ReferenceSet) co
 		}
 	}
 
+	for (int32 Idx = 0; Idx < ReferenceSet.AdditionalVariables.Num(); Idx++)
+	{
+		const FString& ReferenceStr = ReferenceSet.AdditionalVariables[Idx];
+
+		if (AdditionalVariables[Idx] != ReferenceStr)
+		{
+			return false;
+		}
+	}
+
 	for (int32 i = 0; i < ReferencedCompileHashes.Num(); i++)
 	{
 		if (ReferencedCompileHashes[i] != ReferenceSet.ReferencedCompileHashes[i])
@@ -303,6 +323,16 @@ void FNiagaraShaderMapId::AppendKeyString(FString& KeyString) const
 		}
 	}
 
+	// Add additional variables
+	for (int32 VariablesIndex = 0; VariablesIndex < AdditionalVariables.Num(); VariablesIndex++)
+	{
+		KeyString += AdditionalVariables[VariablesIndex];
+		if (VariablesIndex < AdditionalVariables.Num() - 1)
+		{
+			KeyString += TEXT("_");
+		}
+	}
+
 	// Add any referenced top level compile hashes to the key so that we will recompile when they are changed
 	for (int32 HashIndex = 0; HashIndex < ReferencedCompileHashes.Num(); HashIndex++)
 	{
diff --git a/Engine/Plugins/FX/Niagara/Source/NiagaraShader/Private/NiagaraShared.cpp b/Engine/Plugins/FX/Niagara/Source/NiagaraShader/Private/NiagaraShared.cpp
--- a/Engine/Plugins/FX/Niagara/Source/NiagaraShader/Private/NiagaraShared.cpp
+++ b/Engine/Plugins/FX/Niagara/Source/NiagaraShader/Private/NiagaraShared.cpp
@@ -129,7 +129,8 @@ void FNiagaraShaderScript::LegacySerialize(FArchive& Ar)
 bool FNiagaraShaderScript::IsSame(const FNiagaraShaderMapId& InId) const
 {
 	if (InId.ReferencedCompileHashes.Num() != ReferencedCompileHashes.Num() ||
-		InId.AdditionalDefines.Num() != AdditionalDefines.Num())
+		InId.AdditionalDefines.Num() != AdditionalDefines.Num() ||
+		InId.AdditionalVariables.Num() != AdditionalVariables.Num())
 	{
 		return false;
 	}
@@ -148,6 +149,13 @@ bool FNiagaraShaderScript::IsSame(const FNiagaraShaderMapId& InId) const
 			return false;
 		}
 	}
+	for (int32 i = 0; i < AdditionalVariables.Num(); ++i)
+	{
+		if (AdditionalVariables[i] != *InId.AdditionalVariables[i])
+		{
+			return false;
+		}
+	}
 
 	return
 		InId.FeatureLevel == FeatureLevel &&/*
@@ -251,6 +259,12 @@ void FNiagaraShaderScript::GetShaderMapId(EShaderPlatform Platform, const ITarge
 			OutId.AdditionalDefines.Emplace(Define);
 		}
 
+		OutId.AdditionalVariables.Empty(AdditionalVariables.Num());
+		for(const FString& Variable : AdditionalVariables)
+		{
+			OutId.AdditionalVariables.Emplace(Variable);
+		}
+
 		TArray<FShaderType*> DependentShaderTypes;
 		GetDependentShaderTypes(Platform, DependentShaderTypes);
 		for (FShaderType* ShaderType : DependentShaderTypes)
@@ -377,7 +391,7 @@ void FNiagaraShaderScript::SerializeShaderMap(FArchive& Ar)
 	}
 }
 
-void FNiagaraShaderScript::SetScript(UNiagaraScriptBase* InScript, ERHIFeatureLevel::Type InFeatureLevel, EShaderPlatform InShaderPlatform, const FGuid& InCompilerVersionID,  const TArray<FString>& InAdditionalDefines,
+void FNiagaraShaderScript::SetScript(UNiagaraScriptBase* InScript, ERHIFeatureLevel::Type InFeatureLevel, EShaderPlatform InShaderPlatform, const FGuid& InCompilerVersionID,  const TArray<FString>& InAdditionalDefines, const TArray<FString>& InAdditionalVariables,
 		const FNiagaraCompileHash& InBaseCompileHash, const TArray<FNiagaraCompileHash>& InReferencedCompileHashes, 
 		bool bInUsesRapidIterationParams, FString InFriendlyName)
 {
@@ -386,6 +400,7 @@ void FNiagaraShaderScript::SetScript(UNiagaraScriptBase* InScript, ERHIFeatureLe
 	CompilerVersionId = InCompilerVersionID;
 	//BaseScriptId = InBaseScriptID;
 	AdditionalDefines = InAdditionalDefines;
+	AdditionalVariables = InAdditionalVariables;
 	bUsesRapidIterationParams = bInUsesRapidIterationParams;
 	BaseCompileHash = InBaseCompileHash;
 	ReferencedCompileHashes = InReferencedCompileHashes;
diff --git a/Engine/Plugins/FX/Niagara/Source/NiagaraShader/Public/NiagaraShared.h b/Engine/Plugins/FX/Niagara/Source/NiagaraShader/Public/NiagaraShared.h
--- a/Engine/Plugins/FX/Niagara/Source/NiagaraShader/Public/NiagaraShared.h
+++ b/Engine/Plugins/FX/Niagara/Source/NiagaraShader/Public/NiagaraShared.h
@@ -239,6 +239,9 @@ public:
 	/** Configuration options */
 	LAYOUT_FIELD(TMemoryImageArray<FMemoryImageString>, AdditionalDefines);
 
+	/** Configuration options */
+	LAYOUT_FIELD(TMemoryImageArray<FMemoryImageString>, AdditionalVariables);
+
 	/**
 	* The hash of the subgraph this shader primarily represents.
 	*/
@@ -764,7 +767,7 @@ public:
 
 	const FString& GetFriendlyName() const { return FriendlyName; }
 
-	NIAGARASHADER_API void SetScript(UNiagaraScriptBase* InScript, ERHIFeatureLevel::Type InFeatureLevel, EShaderPlatform InShaderPlatform, const FGuid& InCompilerVersion, const TArray<FString>& InAdditionalDefines,
+	NIAGARASHADER_API void SetScript(UNiagaraScriptBase* InScript, ERHIFeatureLevel::Type InFeatureLevel, EShaderPlatform InShaderPlatform, const FGuid& InCompilerVersion, const TArray<FString>& InAdditionalDefines, const TArray<FString>& InAdditionalVariables,
 		const FNiagaraCompileHash& InBaseCompileHash, const TArray<FNiagaraCompileHash>& InReferencedCompileHashes, 
 		bool bInUsesRapidIterationParams, FString InFriendlyName);
 
@@ -861,6 +864,9 @@ private:
 	/** Configuration options */
 	TArray<FString> AdditionalDefines;
 
+	/** Additional variables needed*/
+	TArray<FString> AdditionalVariables;
+
 	/** Whether or not we need to bake Rapid Iteration params. True to keep params, false to bake.*/
 	bool bUsesRapidIterationParams = true;
 
diff --git a/Engine/Plugins/FX/Niagara/Source/NiagaraVertexFactories/Private/NiagaraMeshVertexFactory.cpp b/Engine/Plugins/FX/Niagara/Source/NiagaraVertexFactories/Private/NiagaraMeshVertexFactory.cpp
--- a/Engine/Plugins/FX/Niagara/Source/NiagaraVertexFactories/Private/NiagaraMeshVertexFactory.cpp
+++ b/Engine/Plugins/FX/Niagara/Source/NiagaraVertexFactories/Private/NiagaraMeshVertexFactory.cpp
@@ -20,12 +20,6 @@ public:
 	{
 		FNiagaraVertexFactoryShaderParametersBase::Bind(ParameterMap);
 
-		//PrevTransformBuffer.Bind(ParameterMap, TEXT("PrevTransformBuffer"));
-
-		// 		NiagaraParticleDataInt.Bind(ParameterMap, TEXT("NiagaraParticleDataInt"));
-		// 		FloatDataOffset.Bind(ParameterMap, TEXT("NiagaraInt32DataOffset"));
-		// 		FloatDataStride.Bind(ParameterMap, TEXT("NiagaraInt3DataStride"));
-
 		SortedIndices.Bind(ParameterMap, TEXT("SortedIndices"));
 		SortedIndicesOffset.Bind(ParameterMap, TEXT("SortedIndicesOffset"));
 	}
@@ -52,16 +46,8 @@ public:
 	}
 
 private:
-
-
-	//LAYOUT_FIELD(FShaderResourceParameter, PrevTransformBuffer)
-
-	// 	LAYOUT_FIELD(FShaderResourceParameter, NiagaraParticleDataInt)
-	// 	LAYOUT_FIELD(FShaderParameter, Int32DataOffset)
-	// 	LAYOUT_FIELD(FShaderParameter, Int32DataStride)
-	LAYOUT_FIELD(FShaderResourceParameter, SortedIndices);
-	LAYOUT_FIELD(FShaderParameter, SortedIndicesOffset);
-
+	LAYOUT_FIELD(FShaderResourceParameter, SortedIndices)
+	LAYOUT_FIELD(FShaderParameter, SortedIndicesOffset)
 };
 
 IMPLEMENT_TYPE_LAYOUT(FNiagaraMeshVertexFactoryShaderParametersVS);
@@ -86,9 +72,6 @@ public:
 		const FNiagaraMeshVertexFactory* NiagaraMeshVF = static_cast<const FNiagaraMeshVertexFactory*>(VertexFactory);
 		ShaderBindings.Add(Shader->GetUniformBufferParameter<FNiagaraMeshUniformParameters>(), NiagaraMeshVF->GetUniformBuffer());
 	}
-
-	
-	
 };
 
 IMPLEMENT_TYPE_LAYOUT(FNiagaraMeshVertexFactoryShaderParametersPS);
@@ -160,39 +143,6 @@ void FNiagaraMeshVertexFactory::InitRHI()
 	}
 }
 
-/*
-uint8* FNiagaraMeshVertexFactory::LockPreviousTransformBuffer(uint32 ParticleCount)
-{
-	const static uint32 ElementSize = sizeof(FVector4);
-	const static uint32 ParticleSize = ElementSize * 3;
-	const uint32 AllocationRequest = ParticleCount * ParticleSize;
-
-	check(!PrevTransformBuffer.MappedBuffer);
-
-	if (AllocationRequest > PrevTransformBuffer.NumBytes)
-	{
-		PrevTransformBuffer.Release();
-		PrevTransformBuffer.Initialize(ElementSize, ParticleCount * 3, PF_A32B32G32R32F, BUF_Dynamic);
-	}
-
-	PrevTransformBuffer.Lock();
-
-	return PrevTransformBuffer.MappedBuffer;
-}
-
-void FNiagaraMeshVertexFactory::UnlockPreviousTransformBuffer()
-{
-	check(PrevTransformBuffer.MappedBuffer);
-
-	PrevTransformBuffer.Unlock();
-}
-
-FRHIShaderResourceView* FNiagaraMeshVertexFactory::GetPreviousTransformBufferSRV() const
-{
-	return PrevTransformBuffer.SRV;
-}
-*/
-
 bool FNiagaraMeshVertexFactory::ShouldCompilePermutation(const FVertexFactoryShaderPermutationParameters& Parameters)
 {
 	return (FNiagaraUtilities::SupportsNiagaraRendering(Parameters.Platform)) && (Parameters.MaterialParameters.bIsUsedWithNiagaraMeshParticles || Parameters.MaterialParameters.bIsSpecialEngineMaterial);
@@ -216,3 +166,18 @@ IMPLEMENT_VERTEX_FACTORY_TYPE(FNiagaraMeshVertexFactory, "/Plugin/FX/Niagara/Pri
 	  EVertexFactoryFlags::UsedWithMaterials 
 	| EVertexFactoryFlags::SupportsDynamicLighting
 );
+
+/////////////////////////////////////////////////////////////////////////////////////////////////////
+
+IMPLEMENT_VERTEX_FACTORY_PARAMETER_TYPE(FNiagaraMeshVertexFactoryEx, SF_Vertex, FNiagaraMeshVertexFactoryShaderParametersVS);
+#if RHI_RAYTRACING
+IMPLEMENT_VERTEX_FACTORY_PARAMETER_TYPE(FNiagaraMeshVertexFactoryEx, SF_Compute, FNiagaraMeshVertexFactoryShaderParametersVS);
+IMPLEMENT_VERTEX_FACTORY_PARAMETER_TYPE(FNiagaraMeshVertexFactoryEx, SF_RayHitGroup, FNiagaraMeshVertexFactoryShaderParametersVS);
+#endif
+IMPLEMENT_VERTEX_FACTORY_PARAMETER_TYPE(FNiagaraMeshVertexFactoryEx, SF_Pixel, FNiagaraMeshVertexFactoryShaderParametersPS);
+
+IMPLEMENT_VERTEX_FACTORY_TYPE(FNiagaraMeshVertexFactoryEx, "/Plugin/FX/Niagara/Private/NiagaraMeshVertexFactory.ush",
+	  EVertexFactoryFlags::UsedWithMaterials
+	| EVertexFactoryFlags::SupportsDynamicLighting
+	| EVertexFactoryFlags::SupportsPrecisePrevWorldPos
+);
diff --git a/Engine/Plugins/FX/Niagara/Source/NiagaraVertexFactories/Private/NiagaraSpriteVertexFactory.cpp b/Engine/Plugins/FX/Niagara/Source/NiagaraVertexFactories/Private/NiagaraSpriteVertexFactory.cpp
--- a/Engine/Plugins/FX/Niagara/Source/NiagaraVertexFactories/Private/NiagaraSpriteVertexFactory.cpp
+++ b/Engine/Plugins/FX/Niagara/Source/NiagaraVertexFactories/Private/NiagaraSpriteVertexFactory.cpp
@@ -211,3 +211,18 @@ IMPLEMENT_VERTEX_FACTORY_TYPE(FNiagaraSpriteVertexFactory,"/Plugin/FX/Niagara/Pr
 	| EVertexFactoryFlags::SupportsDynamicLighting
 	| EVertexFactoryFlags::SupportsRayTracing
 );
+
+/////////////////////////////////////////////////////////////////////////////
+
+IMPLEMENT_VERTEX_FACTORY_PARAMETER_TYPE(FNiagaraSpriteVertexFactoryEx, SF_Vertex, FNiagaraSpriteVertexFactoryShaderParametersVS);
+IMPLEMENT_VERTEX_FACTORY_PARAMETER_TYPE(FNiagaraSpriteVertexFactoryEx, SF_Pixel, FNiagaraSpriteVertexFactoryShaderParametersPS);
+#if RHI_RAYTRACING
+IMPLEMENT_VERTEX_FACTORY_PARAMETER_TYPE(FNiagaraSpriteVertexFactoryEx, SF_Compute, FNiagaraSpriteVertexFactoryShaderParametersVS);
+IMPLEMENT_VERTEX_FACTORY_PARAMETER_TYPE(FNiagaraSpriteVertexFactoryEx, SF_RayHitGroup, FNiagaraSpriteVertexFactoryShaderParametersVS);
+#endif // RHI_RAYTRACING
+IMPLEMENT_VERTEX_FACTORY_TYPE(FNiagaraSpriteVertexFactoryEx,"/Plugin/FX/Niagara/Private/NiagaraSpriteVertexFactory.ush",
+	  EVertexFactoryFlags::UsedWithMaterials
+	| EVertexFactoryFlags::SupportsDynamicLighting
+	| EVertexFactoryFlags::SupportsRayTracing
+	| EVertexFactoryFlags::SupportsPrecisePrevWorldPos
+);
\ No newline at end of file
diff --git a/Engine/Plugins/FX/Niagara/Source/NiagaraVertexFactories/Public/NiagaraMeshVertexFactory.h b/Engine/Plugins/FX/Niagara/Source/NiagaraVertexFactories/Public/NiagaraMeshVertexFactory.h
--- a/Engine/Plugins/FX/Niagara/Source/NiagaraVertexFactories/Public/NiagaraMeshVertexFactory.h
+++ b/Engine/Plugins/FX/Niagara/Source/NiagaraVertexFactories/Public/NiagaraMeshVertexFactory.h
@@ -35,15 +35,18 @@ BEGIN_GLOBAL_SHADER_PARAMETER_STRUCT(FNiagaraMeshUniformParameters, NIAGARAVERTE
 	SHADER_PARAMETER(FVector4, SubImageSize)
 	SHADER_PARAMETER(uint32, TexCoordWeightA)
 	SHADER_PARAMETER(uint32, TexCoordWeightB)
-	SHADER_PARAMETER(uint32, PrevTransformAvailable)
 	SHADER_PARAMETER(float, DeltaSeconds)
 	SHADER_PARAMETER(uint32, MaterialParamValidMask)
 
 	SHADER_PARAMETER(int, PositionDataOffset)
+	SHADER_PARAMETER(int, PrevPositionDataOffset)
 	SHADER_PARAMETER(int, VelocityDataOffset)
+	SHADER_PARAMETER(int, PrevVelocityDataOffset)
 	SHADER_PARAMETER(int, ColorDataOffset)
-	SHADER_PARAMETER(int, TransformDataOffset)
+	SHADER_PARAMETER(int, RotationDataOffset)
+	SHADER_PARAMETER(int, PrevRotationDataOffset)
 	SHADER_PARAMETER(int, ScaleDataOffset)
+	SHADER_PARAMETER(int, PrevScaleDataOffset)
 	SHADER_PARAMETER(int, MaterialParamDataOffset)
 	SHADER_PARAMETER(int, MaterialParam1DataOffset)
 	SHADER_PARAMETER(int, MaterialParam2DataOffset)
@@ -52,12 +55,17 @@ BEGIN_GLOBAL_SHADER_PARAMETER_STRUCT(FNiagaraMeshUniformParameters, NIAGARAVERTE
 	SHADER_PARAMETER(int, SubImageDataOffset)
 	SHADER_PARAMETER(int, MaterialRandomDataOffset)
 	SHADER_PARAMETER(int, CameraOffsetDataOffset)
+	SHADER_PARAMETER(int, PrevCameraOffsetDataOffset)
 
 	SHADER_PARAMETER(FVector4, DefaultPos)
+	SHADER_PARAMETER(FVector4, DefaultPrevPos)
 	SHADER_PARAMETER(FVector, DefaultVelocity)
+	SHADER_PARAMETER(FVector, DefaultPrevVelocity)
 	SHADER_PARAMETER(FVector4, DefaultColor)
-	SHADER_PARAMETER(FVector4, DefaultTransform)
+	SHADER_PARAMETER(FVector4, DefaultRotation)
+	SHADER_PARAMETER(FVector4, DefaultPrevRotation)
 	SHADER_PARAMETER(FVector, DefaultScale)
+	SHADER_PARAMETER(FVector, DefaultPrevScale)
 	SHADER_PARAMETER(FVector4, DefaultDynamicMaterialParameter0)
 	SHADER_PARAMETER(FVector4, DefaultDynamicMaterialParameter1)
 	SHADER_PARAMETER(FVector4, DefaultDynamicMaterialParameter2)
@@ -66,6 +74,7 @@ BEGIN_GLOBAL_SHADER_PARAMETER_STRUCT(FNiagaraMeshUniformParameters, NIAGARAVERTE
 	SHADER_PARAMETER(float, DefaultSubImage)
 	SHADER_PARAMETER(float, DefaultMatRandom)
 	SHADER_PARAMETER(float, DefaultCamOffset)
+	SHADER_PARAMETER(float, DefaultPrevCamOffset)
 
 	SHADER_PARAMETER(int, SubImageBlendMode)
 	SHADER_PARAMETER(uint32, FacingMode)
@@ -206,12 +215,23 @@ protected:
 	uint32 SortedIndicesOffset;
 };
 
-inline FNiagaraMeshVertexFactory* ConstructNiagaraMeshVertexFactory()
+/**
+* Advanced mesh vertex factory. Used for enabling accurate motion vector output
+*/
+class NIAGARAVERTEXFACTORIES_API FNiagaraMeshVertexFactoryEx : public FNiagaraMeshVertexFactory
 {
-	return new FNiagaraMeshVertexFactory();
-}
+	DECLARE_VERTEX_FACTORY_TYPE(FNiagaraMeshVertexFactoryEx);
+public:
+	FNiagaraMeshVertexFactoryEx(ENiagaraVertexFactoryType InType, ERHIFeatureLevel::Type InFeatureLevel)
+		: FNiagaraMeshVertexFactory(InType, InFeatureLevel)
+	{}
 
-inline FNiagaraMeshVertexFactory* ConstructNiagaraMeshVertexFactory(ENiagaraVertexFactoryType InType, ERHIFeatureLevel::Type InFeatureLevel)
-{
-	return new FNiagaraMeshVertexFactory(InType, InFeatureLevel);
-}
+	FNiagaraMeshVertexFactoryEx() {}
+
+	static void ModifyCompilationEnvironment(const FVertexFactoryShaderPermutationParameters& Parameters, FShaderCompilerEnvironment& OutEnvironment)
+	{
+		FNiagaraMeshVertexFactory::ModifyCompilationEnvironment(Parameters, OutEnvironment);
+
+		OutEnvironment.SetDefine(TEXT("ENABLE_PRECISE_MOTION_VECTORS"), TEXT("1"));
+	}
+};
\ No newline at end of file
diff --git a/Engine/Plugins/FX/Niagara/Source/NiagaraVertexFactories/Public/NiagaraSpriteVertexFactory.h b/Engine/Plugins/FX/Niagara/Source/NiagaraVertexFactories/Public/NiagaraSpriteVertexFactory.h
--- a/Engine/Plugins/FX/Niagara/Source/NiagaraVertexFactories/Public/NiagaraSpriteVertexFactory.h
+++ b/Engine/Plugins/FX/Niagara/Source/NiagaraVertexFactories/Public/NiagaraSpriteVertexFactory.h
@@ -32,10 +32,15 @@ BEGIN_GLOBAL_SHADER_PARAMETER_STRUCT( FNiagaraSpriteUniformParameters, NIAGARAVE
 	SHADER_PARAMETER_EX( float, NormalsType, EShaderPrecisionModifier::Half )
 	SHADER_PARAMETER_EX( float, DeltaSeconds, EShaderPrecisionModifier::Half )
 	SHADER_PARAMETER_EX( FVector2D, DefaultPivotOffset, EShaderPrecisionModifier::Half )
+	SHADER_PARAMETER_EX( FVector2D, DefaultPrevPivotOffset, EShaderPrecisionModifier::Half )
 	SHADER_PARAMETER(int, PositionDataOffset)
+	SHADER_PARAMETER(int, PrevPositionDataOffset)
 	SHADER_PARAMETER(int, VelocityDataOffset)
+	SHADER_PARAMETER(int, PrevVelocityDataOffset)
 	SHADER_PARAMETER(int, RotationDataOffset)
+	SHADER_PARAMETER(int, PrevRotationDataOffset)
 	SHADER_PARAMETER(int, SizeDataOffset)
+	SHADER_PARAMETER(int, PrevSizeDataOffset)
 	SHADER_PARAMETER(int, SubimageDataOffset)
 	SHADER_PARAMETER(int, ColorDataOffset)
 	SHADER_PARAMETER(uint32, MaterialParamValidMask)
@@ -44,25 +49,36 @@ BEGIN_GLOBAL_SHADER_PARAMETER_STRUCT( FNiagaraSpriteUniformParameters, NIAGARAVE
 	SHADER_PARAMETER(int, MaterialParam2DataOffset)
 	SHADER_PARAMETER(int, MaterialParam3DataOffset)
 	SHADER_PARAMETER(int, FacingDataOffset)
+	SHADER_PARAMETER(int, PrevFacingDataOffset)
 	SHADER_PARAMETER(int, AlignmentDataOffset)
+	SHADER_PARAMETER(int, PrevAlignmentDataOffset)
 	SHADER_PARAMETER(int, SubImageBlendMode)
 	SHADER_PARAMETER(int, CameraOffsetDataOffset)
+	SHADER_PARAMETER(int, PrevCameraOffsetDataOffset)
 	SHADER_PARAMETER(int, UVScaleDataOffset)
 	SHADER_PARAMETER(int, PivotOffsetDataOffset)
+	SHADER_PARAMETER(int, PrevPivotOffsetDataOffset)
 	SHADER_PARAMETER(int, NormalizedAgeDataOffset)
 	SHADER_PARAMETER(int, MaterialRandomDataOffset)
 	SHADER_PARAMETER(FVector4, DefaultPos)
+	SHADER_PARAMETER(FVector4, DefaultPrevPos)
 	SHADER_PARAMETER(FVector2D, DefaultSize)
+	SHADER_PARAMETER(FVector2D, DefaultPrevSize)
 	SHADER_PARAMETER(FVector2D, DefaultUVScale)
 	SHADER_PARAMETER(FVector, DefaultVelocity)
+	SHADER_PARAMETER(FVector, DefaultPrevVelocity)
 	SHADER_PARAMETER(float, DefaultRotation)
+	SHADER_PARAMETER(float, DefaultPrevRotation)
 	SHADER_PARAMETER(FVector4, DefaultColor)
 	SHADER_PARAMETER(float, DefaultMatRandom)
 	SHADER_PARAMETER(float, DefaultCamOffset)
+	SHADER_PARAMETER(float, DefaultPrevCamOffset)
 	SHADER_PARAMETER(float, DefaultNormAge)
 	SHADER_PARAMETER(float, DefaultSubImage)
 	SHADER_PARAMETER(FVector4, DefaultFacing)
+	SHADER_PARAMETER(FVector4, DefaultPrevFacing)
 	SHADER_PARAMETER(FVector4, DefaultAlignment)
+	SHADER_PARAMETER(FVector4, DefaultPrevAlignment)
 	SHADER_PARAMETER(FVector4, DefaultDynamicMaterialParameter0)
 	SHADER_PARAMETER(FVector4, DefaultDynamicMaterialParameter1)
 	SHADER_PARAMETER(FVector4, DefaultDynamicMaterialParameter2)
@@ -221,3 +237,25 @@ private:
 	FShaderResourceViewRHIRef SortedIndicesSRV;
 	uint32 SortedIndicesOffset;
 };
+
+class NIAGARAVERTEXFACTORIES_API FNiagaraSpriteVertexFactoryEx : public FNiagaraSpriteVertexFactory
+{
+	DECLARE_VERTEX_FACTORY_TYPE(FNiagaraSpriteVertexFactoryEx);
+
+public:
+
+	FNiagaraSpriteVertexFactoryEx(ENiagaraVertexFactoryType InType, ERHIFeatureLevel::Type InFeatureLevel )
+		: FNiagaraSpriteVertexFactory(InType, InFeatureLevel)
+	{}
+
+	FNiagaraSpriteVertexFactoryEx()
+		: FNiagaraSpriteVertexFactory(NVFT_MAX, ERHIFeatureLevel::Num)
+	{}
+
+	static void ModifyCompilationEnvironment(const FVertexFactoryShaderPermutationParameters& Parameters, FShaderCompilerEnvironment& OutEnvironment)
+	{
+		FNiagaraSpriteVertexFactory::ModifyCompilationEnvironment(Parameters, OutEnvironment);
+
+		OutEnvironment.SetDefine(TEXT("ENABLE_PRECISE_MOTION_VECTORS"), TEXT("1"));
+	}
+};
