Adding missing SSpacer and SCheckbox includes

diff --git a/Engine/Plugins/FX/Niagara/Source/NiagaraEditor/Private/Widgets/SNiagaraParameterMapPaletteItem.cpp b/Engine/Plugins/FX/Niagara/Source/NiagaraEditor/Private/Widgets/SNiagaraParameterMapPaletteItem.cpp
--- a/Engine/Plugins/FX/Niagara/Source/NiagaraEditor/Private/Widgets/SNiagaraParameterMapPaletteItem.cpp
+++ b/Engine/Plugins/FX/Niagara/Source/NiagaraEditor/Private/Widgets/SNiagaraParameterMapPaletteItem.cpp
@@ -12,6 +12,7 @@
 #include "Widgets/SNiagaraParameterMapView.h"
 #include "Widgets/SNiagaraParameterName.h"
 #include "Widgets/Input/SComboButton.h"
+#include "Widgets/Layout/SSpacer.h"
 
 #define LOCTEXT_NAMESPACE "NiagaraParameterMapPalleteItem"
 
diff --git a/Engine/Plugins/FX/Niagara/Source/NiagaraEditorWidgets/Private/SNiagaraOverviewStack.cpp b/Engine/Plugins/FX/Niagara/Source/NiagaraEditorWidgets/Private/SNiagaraOverviewStack.cpp
--- a/Engine/Plugins/FX/Niagara/Source/NiagaraEditorWidgets/Private/SNiagaraOverviewStack.cpp
+++ b/Engine/Plugins/FX/Niagara/Source/NiagaraEditorWidgets/Private/SNiagaraOverviewStack.cpp
@@ -37,6 +37,7 @@
 #include "Framework/MultiBox/MultiBoxBuilder.h"
 #include "Subsystems/AssetEditorSubsystem.h"
 #include "Styling/StyleColors.h"
+#include "Widgets/Input/SCheckBox.h"
 
 #define LOCTEXT_NAMESPACE "NiagaraOverviewStack"
 
