Additional CIS fix for non-unity compile

diff --git a/Engine/Plugins/FX/Niagara/Source/NiagaraEditorWidgets/Private/NiagaraEditorWidgetsStyle.cpp b/Engine/Plugins/FX/Niagara/Source/NiagaraEditorWidgets/Private/NiagaraEditorWidgetsStyle.cpp
--- a/Engine/Plugins/FX/Niagara/Source/NiagaraEditorWidgets/Private/NiagaraEditorWidgetsStyle.cpp
+++ b/Engine/Plugins/FX/Niagara/Source/NiagaraEditorWidgets/Private/NiagaraEditorWidgetsStyle.cpp
@@ -274,7 +274,7 @@ TSharedRef< FSlateStyleSet > FNiagaraEditorWidgetsStyle::Create()
 		.SetHandleHighlightBrush(IMAGE_CORE_BRUSH("Common/SplitterHandleHighlight", Icon8x8, FLinearColor::White))
 	);
 
-	Style->Set("NiagaraEditor.Stack.SearchHighlightColor", FEditorStyle::GetColor("TextBlock.HighlighColor"));
+	Style->Set("NiagaraEditor.Stack.SearchHighlightColor", FStyleColors::Highlight);
 	Style->Set("NiagaraEditor.Stack.SearchResult", new BOX_PLUGIN_BRUSH("Icons/SearchResultBorder", FMargin(1.f/8.f)));
 
 	Style->Set("NiagaraEditor.Stack.ModuleHighlight", new IMAGE_PLUGIN_BRUSH("Icons/ModuleHighlight", Icon6x6, FLinearColor::White));
