- Fix handling around texture references on RenderTarget data interface
- Move UAV creation into TextureRenderTarget as it is not safe to create a UAV off a texture reference since a modification can be in flight

diff --git a/Engine/Plugins/FX/Niagara/Source/Niagara/Classes/NiagaraDataInterfaceRenderTarget2D.h b/Engine/Plugins/FX/Niagara/Source/Niagara/Classes/NiagaraDataInterfaceRenderTarget2D.h
--- a/Engine/Plugins/FX/Niagara/Source/Niagara/Classes/NiagaraDataInterfaceRenderTarget2D.h
+++ b/Engine/Plugins/FX/Niagara/Source/Niagara/Classes/NiagaraDataInterfaceRenderTarget2D.h
@@ -49,8 +49,8 @@ struct FRenderTarget2DRWInstanceData_RenderThread
 	bool bWasWrittenTo = false;
 
 	FSamplerStateRHIRef SamplerStateRHI;
-	FTextureReferenceRHIRef TextureReferenceRHI;
-	FUnorderedAccessViewRHIRef UAV;
+	FTexture2DRHIRef TextureRHI;
+	FUnorderedAccessViewRHIRef UnorderedAccessViewRHI;
 #if WITH_EDITORONLY_DATA
 	uint32 bPreviewTexture : 1;
 #endif
diff --git a/Engine/Plugins/FX/Niagara/Source/Niagara/Classes/NiagaraDataInterfaceRenderTarget2DArray.h b/Engine/Plugins/FX/Niagara/Source/Niagara/Classes/NiagaraDataInterfaceRenderTarget2DArray.h
--- a/Engine/Plugins/FX/Niagara/Source/Niagara/Classes/NiagaraDataInterfaceRenderTarget2DArray.h
+++ b/Engine/Plugins/FX/Niagara/Source/Niagara/Classes/NiagaraDataInterfaceRenderTarget2DArray.h
@@ -42,8 +42,8 @@ struct FRenderTarget2DArrayRWInstanceData_RenderThread
 	FIntVector Size = FIntVector(EForceInit::ForceInitToZero);
 	
 	FSamplerStateRHIRef SamplerStateRHI;
-	FTextureReferenceRHIRef TextureReferenceRHI;
-	FUnorderedAccessViewRHIRef UAV;
+	FTexture2DArrayRHIRef TextureRHI;
+	FUnorderedAccessViewRHIRef UnorderedAccessViewRHI;
 #if WITH_EDITORONLY_DATA
 	uint32 bPreviewTexture : 1;
 #endif
diff --git a/Engine/Plugins/FX/Niagara/Source/Niagara/Classes/NiagaraDataInterfaceRenderTargetVolume.h b/Engine/Plugins/FX/Niagara/Source/Niagara/Classes/NiagaraDataInterfaceRenderTargetVolume.h
--- a/Engine/Plugins/FX/Niagara/Source/Niagara/Classes/NiagaraDataInterfaceRenderTargetVolume.h
+++ b/Engine/Plugins/FX/Niagara/Source/Niagara/Classes/NiagaraDataInterfaceRenderTargetVolume.h
@@ -41,8 +41,8 @@ struct FRenderTargetVolumeRWInstanceData_RenderThread
 	FIntVector Size = FIntVector(EForceInit::ForceInitToZero);
 	
 	FSamplerStateRHIRef SamplerStateRHI;
-	FTextureReferenceRHIRef TextureReferenceRHI;
-	FUnorderedAccessViewRHIRef UAV;
+	FTexture3DRHIRef TextureRHI;
+	FUnorderedAccessViewRHIRef UnorderedAccessViewRHI;
 #if WITH_EDITORONLY_DATA
 	uint32 bPreviewTexture : 1;
 #endif
diff --git a/Engine/Plugins/FX/Niagara/Source/Niagara/Private/NiagaraDataInterfaceRenderTarget2D.cpp b/Engine/Plugins/FX/Niagara/Source/Niagara/Private/NiagaraDataInterfaceRenderTarget2D.cpp
--- a/Engine/Plugins/FX/Niagara/Source/Niagara/Private/NiagaraDataInterfaceRenderTarget2D.cpp
+++ b/Engine/Plugins/FX/Niagara/Source/Niagara/Private/NiagaraDataInterfaceRenderTarget2D.cpp
@@ -5,7 +5,10 @@
 #include "ClearQuad.h"
 #include "TextureResource.h"
 #include "GenerateMips.h"
+#include "CanvasTypes.h"
 #include "Engine/Texture2D.h"
+#include "Engine/TextureRenderTarget2D.h"
+
 #include "NiagaraEmitterInstanceBatcher.h"
 #include "NiagaraSettings.h"
 #include "NiagaraSystemInstance.h"
@@ -14,8 +17,6 @@
 #include "NiagaraGpuComputeDebug.h"
 #endif
 #include "NiagaraStats.h"
-#include "Engine/TextureRenderTarget2D.h"
-#include "CanvasTypes.h"
 
 #define LOCTEXT_NAMESPACE "NiagaraDataInterfaceRenderTarget2D"
 
@@ -101,24 +102,15 @@ public:
 	
 		if (OutputParam.IsUAVBound())
 		{
-			FRHIUnorderedAccessView* OutputUAV = nullptr;
-
-			if (FRHITexture* TextureRHI = ProxyData->TextureReferenceRHI->GetReferencedTexture())
+			FRHIUnorderedAccessView* OutputUAV = ProxyData->UnorderedAccessViewRHI;
+			if (OutputUAV)
 			{
-				// Note: Because we control the render target it should not changed underneath us without queueing a request to recreate the UAV.  If that assumption changes we would need to track the UAV.
-				if (!ProxyData->UAV.IsValid())
-				{
-					ProxyData->UAV = RHICreateUnorderedAccessView(TextureRHI, 0);
-				}
-
-				OutputUAV = ProxyData->UAV;
 				// FIXME: this transition needs to happen in FNiagaraDataInterfaceProxyRenderTarget2DProxy::PreStage so it doesn't break up the overlap group,
 				// but for some reason it stops working if I move it in there.
-				RHICmdList.Transition(FRHITransitionInfo(OutputUAV, ERHIAccess::SRVMask, ERHIAccess::UAVCompute));
+				RHICmdList.Transition(FRHITransitionInfo(OutputUAV, ERHIAccess::Unknown, ERHIAccess::UAVCompute));
 				ProxyData->bWasWrittenTo = true;
 			}
-
-			if (OutputUAV == nullptr)
+			else
 			{
 				OutputUAV = Context.Batcher->GetEmptyRWTextureFromPool(RHICmdList, EPixelFormat::PF_A16B16G16R16);
 			}
@@ -128,7 +120,13 @@ public:
 
 		if ( InputParam.IsBound() )
 		{
-			FRHITexture* TextureRHI = ProxyData->TextureReferenceRHI->GetReferencedTexture();
+			FRHITexture* TextureRHI = ProxyData->TextureRHI;
+			if (!ensureMsgf(!OutputParam.IsUAVBound(), TEXT("NiagaraDIRenderTarget2D(%s) is bound as both read & write, read will be ignored."), *Context.DataInterface->SourceDIName.ToString()))
+			{
+				//-TODO: Feedback to the user that read & write is bound
+				TextureRHI = nullptr;
+			}
+
 			if (TextureRHI == nullptr)
 			{
 				TextureRHI = GBlackTexture->TextureRHI;
@@ -159,7 +157,7 @@ public:
 			FNiagaraDataInterfaceProxyRenderTarget2DProxy* VFDI = static_cast<FNiagaraDataInterfaceProxyRenderTarget2DProxy*>(Context.DataInterface);
 			if ( FRenderTarget2DRWInstanceData_RenderThread* ProxyData = VFDI->SystemInstancesToProxyData_RT.Find(Context.SystemInstanceID) )
 			{
-				if (FRHIUnorderedAccessView* OutputUAV = ProxyData->UAV)
+				if (FRHIUnorderedAccessView* OutputUAV = ProxyData->UnorderedAccessViewRHI)
 				{
 					// FIXME: move to FNiagaraDataInterfaceProxyRenderTarget2DProxy::PostStage, same as for the transition in Set() above.
 					RHICmdList.Transition(FRHITransitionInfo(OutputUAV, ERHIAccess::UAVCompute, ERHIAccess::SRVMask));
@@ -190,12 +188,9 @@ void FRenderTarget2DRWInstanceData_RenderThread::UpdateMemoryStats()
 	DEC_MEMORY_STAT_BY(STAT_NiagaraRT2DMemory, MemorySize);
 
 	MemorySize = 0;
-	if (TextureReferenceRHI.IsValid())
+	if (FRHITexture* RHITexture = TextureRHI)
 	{
-		if (FRHITexture* RHITexture = TextureReferenceRHI->GetReferencedTexture())
-		{
-			MemorySize = RHIComputeMemorySize(RHITexture);
-		}
+		MemorySize = RHIComputeMemorySize(RHITexture);
 	}
 
 	INC_MEMORY_STAT_BY(STAT_NiagaraRT2DMemory, MemorySize);
@@ -577,7 +572,7 @@ bool UNiagaraDataInterfaceRenderTarget2D::InitPerInstanceData(void* PerInstanceD
 	// Push Updates to Proxy.
 	FNiagaraDataInterfaceProxyRenderTarget2DProxy* RT_Proxy = GetProxyAs<FNiagaraDataInterfaceProxyRenderTarget2DProxy>();
 	ENQUEUE_RENDER_COMMAND(FUpdateData)(
-		[RT_Proxy, RT_InstanceID=SystemInstance->GetId(), RT_InstanceData=*InstanceData, RT_TargetTexture=InstanceData->TargetTexture](FRHICommandListImmediate& RHICmdList)
+		[RT_Proxy, RT_InstanceID=SystemInstance->GetId(), RT_InstanceData=*InstanceData, RT_TargetTexture=InstanceData->TargetTexture ? InstanceData->TargetTexture->GameThread_GetRenderTargetResource() : nullptr](FRHICommandListImmediate& RHICmdList)
 		{
 			check(!RT_Proxy->SystemInstancesToProxyData_RT.Contains(RT_InstanceID));
 			FRenderTarget2DRWInstanceData_RenderThread* TargetData = &RT_Proxy->SystemInstancesToProxyData_RT.Add(RT_InstanceID);
@@ -587,9 +582,15 @@ bool UNiagaraDataInterfaceRenderTarget2D::InitPerInstanceData(void* PerInstanceD
 #if WITH_EDITORONLY_DATA
 			TargetData->bPreviewTexture = RT_InstanceData.bPreviewTexture;
 #endif
-			TargetData->SamplerStateRHI = RT_TargetTexture->Resource && RT_TargetTexture->Resource->SamplerStateRHI ? RT_TargetTexture->Resource->SamplerStateRHI : nullptr;
-			TargetData->TextureReferenceRHI = RT_TargetTexture->TextureReference.TextureReferenceRHI;
-			TargetData->UAV.SafeRelease();
+			if (RT_TargetTexture)
+			{
+				if (FTextureRenderTarget2DResource* Resource2D = RT_TargetTexture->GetTextureRenderTarget2DResource())
+				{
+					TargetData->SamplerStateRHI = Resource2D->SamplerStateRHI;
+					TargetData->TextureRHI = Resource2D->GetTextureRHI();
+					TargetData->UnorderedAccessViewRHI = Resource2D->GetUnorderedAccessViewRHI();
+				}
+			}
 
 #if STATS
 			TargetData->UpdateMemoryStats();
@@ -617,7 +618,7 @@ void UNiagaraDataInterfaceRenderTarget2D::DestroyPerInstanceData(void* PerInstan
 			if ( ensure(TargetData != nullptr) )
 			{
 				TargetData->SamplerStateRHI = nullptr;
-				TargetData->TextureReferenceRHI = nullptr;
+				TargetData->TextureRHI = nullptr;
 				TargetData->UpdateMemoryStats();
 			}
 #endif
@@ -807,7 +808,7 @@ bool UNiagaraDataInterfaceRenderTarget2D::PerInstanceTickPostSimulate(void* PerI
 			FNiagaraDataInterfaceProxyRenderTarget2DProxy* RT_Proxy = GetProxyAs<FNiagaraDataInterfaceProxyRenderTarget2DProxy>();
 			ENQUEUE_RENDER_COMMAND(FUpdateData)
 			(
-				[RT_Proxy, RT_InstanceID=SystemInstance->GetId(), RT_InstanceData=*InstanceData, RT_TargetTexture=InstanceData->TargetTexture](FRHICommandListImmediate& RHICmdList)
+				[RT_Proxy, RT_InstanceID=SystemInstance->GetId(), RT_InstanceData=*InstanceData, RT_TargetTexture=InstanceData->TargetTexture ? InstanceData->TargetTexture->GameThread_GetRenderTargetResource() : nullptr](FRHICommandListImmediate& RHICmdList)
 				{
 					FRenderTarget2DRWInstanceData_RenderThread* TargetData = RT_Proxy->SystemInstancesToProxyData_RT.Find(RT_InstanceID);
 					if (ensureMsgf(TargetData != nullptr, TEXT("InstanceData was not found for %llu"), RT_InstanceID))
@@ -817,9 +818,18 @@ bool UNiagaraDataInterfaceRenderTarget2D::PerInstanceTickPostSimulate(void* PerI
 	#if WITH_EDITORONLY_DATA
 						TargetData->bPreviewTexture = RT_InstanceData.bPreviewTexture;
 	#endif
-						TargetData->SamplerStateRHI = RT_TargetTexture->Resource && RT_TargetTexture->Resource->SamplerStateRHI ? RT_TargetTexture->Resource->SamplerStateRHI : nullptr;
-						TargetData->TextureReferenceRHI = RT_TargetTexture->TextureReference.TextureReferenceRHI;
-						TargetData->UAV.SafeRelease();
+						TargetData->SamplerStateRHI.SafeRelease();
+						TargetData->TextureRHI.SafeRelease();
+						TargetData->UnorderedAccessViewRHI.SafeRelease();
+						if (RT_TargetTexture)
+						{
+							if (FTextureRenderTarget2DResource* Resource2D = RT_TargetTexture->GetTextureRenderTarget2DResource())
+							{
+								TargetData->SamplerStateRHI = Resource2D->SamplerStateRHI;
+								TargetData->TextureRHI = Resource2D->GetTextureRHI();
+								TargetData->UnorderedAccessViewRHI = Resource2D->GetUnorderedAccessViewRHI();
+							}
+						}
 #if STATS
 						TargetData->UpdateMemoryStats();
 #endif
@@ -843,7 +853,7 @@ void FNiagaraDataInterfaceProxyRenderTarget2DProxy::PostStage(FRHICommandList& R
 
 			FMemMark MemMark(FMemStack::Get());
 			FRDGBuilder GraphBuilder(FRHICommandListExecutor::GetImmediateCommandList());
-			TRefCountPtr<IPooledRenderTarget> PoolRenderTarget = CreateRenderTarget(ProxyData->TextureReferenceRHI->GetReferencedTexture(), TEXT("NiagaraRenderTarget2D"));
+			TRefCountPtr<IPooledRenderTarget> PoolRenderTarget = CreateRenderTarget(ProxyData->TextureRHI, TEXT("NiagaraRenderTarget2D"));
 			FRDGTextureRef MipOutputTexture = GraphBuilder.RegisterExternalTexture(PoolRenderTarget);
 			FGenerateMips::Execute(GraphBuilder, MipOutputTexture);
 			GraphBuilder.Execute();
@@ -863,18 +873,18 @@ void FNiagaraDataInterfaceProxyRenderTarget2DProxy::PostSimulate(FRHICommandList
 
 			FMemMark MemMark(FMemStack::Get());
 			FRDGBuilder GraphBuilder(FRHICommandListExecutor::GetImmediateCommandList());
-			TRefCountPtr<IPooledRenderTarget> PoolRenderTarget = CreateRenderTarget(ProxyData->TextureReferenceRHI->GetReferencedTexture(), TEXT("NiagaraRenderTarget2D"));
+			TRefCountPtr<IPooledRenderTarget> PoolRenderTarget = CreateRenderTarget(ProxyData->TextureRHI, TEXT("NiagaraRenderTarget2D"));
 			FRDGTextureRef MipOutputTexture = GraphBuilder.RegisterExternalTexture(PoolRenderTarget);
 			FGenerateMips::Execute(GraphBuilder, MipOutputTexture);
 			GraphBuilder.Execute();
 		}
 
 #if NIAGARA_COMPUTEDEBUG_ENABLED
-		if (ProxyData->bPreviewTexture && ProxyData->TextureReferenceRHI.IsValid())
+		if (ProxyData->bPreviewTexture)
 		{
 			if (FNiagaraGpuComputeDebug* GpuComputeDebug = Context.Batcher->GetGpuComputeDebug())
 			{
-				if (FRHITexture* RHITexture = ProxyData->TextureReferenceRHI->GetReferencedTexture())
+				if (FRHITexture* RHITexture = ProxyData->TextureRHI)
 				{
 					GpuComputeDebug->AddTexture(RHICmdList, Context.SystemInstanceID, SourceDIName, RHITexture);
 				}
diff --git a/Engine/Plugins/FX/Niagara/Source/Niagara/Private/NiagaraDataInterfaceRenderTarget2DArray.cpp b/Engine/Plugins/FX/Niagara/Source/Niagara/Private/NiagaraDataInterfaceRenderTarget2DArray.cpp
--- a/Engine/Plugins/FX/Niagara/Source/Niagara/Private/NiagaraDataInterfaceRenderTarget2DArray.cpp
+++ b/Engine/Plugins/FX/Niagara/Source/Niagara/Private/NiagaraDataInterfaceRenderTarget2DArray.cpp
@@ -4,7 +4,10 @@
 #include "ShaderParameterUtils.h"
 #include "ClearQuad.h"
 #include "TextureResource.h"
+#include "TextureRenderTarget2DArrayResource.h"
 #include "Engine/Texture2DArray.h"
+#include "Engine/TextureRenderTarget2DArray.h"
+
 #include "NiagaraEmitterInstanceBatcher.h"
 #include "NiagaraSystemInstance.h"
 #include "NiagaraRenderer.h"
@@ -12,7 +15,6 @@
 #if WITH_EDITOR
 #include "NiagaraGpuComputeDebug.h"
 #endif
-#include "Engine/TextureRenderTarget2DArray.h"
 
 #define LOCTEXT_NAMESPACE "NiagaraDataInterfaceRenderTarget2DArray"
 
@@ -74,23 +76,12 @@ public:
 	
 		if ( OutputParam.IsUAVBound())
 		{
-			FRHIUnorderedAccessView* OutputUAV = nullptr;
-
-			if (FRHITexture* TextureRHI = ProxyData->TextureReferenceRHI->GetReferencedTexture())
+			FRHIUnorderedAccessView* OutputUAV = ProxyData->UnorderedAccessViewRHI;
+			if (OutputUAV)
 			{
-				// Note: Because we control the render target it should not changed underneath us without queueing a request to recreate the UAV.  If that assumption changes we would need to track the UAV.
-				ERHIAccess InitialState = ERHIAccess::SRVMask;
-				if (!ProxyData->UAV.IsValid())
-				{
-					ProxyData->UAV = RHICreateUnorderedAccessView(TextureRHI, 0);
-					InitialState = ERHIAccess::Unknown;
-				}
-
-				OutputUAV = ProxyData->UAV;
-				RHICmdList.Transition(FRHITransitionInfo(OutputUAV, InitialState, ERHIAccess::UAVCompute));
+				RHICmdList.Transition(FRHITransitionInfo(OutputUAV, ERHIAccess::Unknown, ERHIAccess::UAVCompute));
 			}
-
-			if (OutputUAV == nullptr)
+			else
 			{
 				OutputUAV = Context.Batcher->GetEmptyRWTextureFromPool(RHICmdList, EPixelFormat::PF_A16B16G16R16);
 			}
@@ -100,7 +91,13 @@ public:
 
 		if (InputParam.IsBound())
 		{
-			FRHITexture* TextureRHI = ProxyData->TextureReferenceRHI->GetReferencedTexture();
+			FRHITexture* TextureRHI = ProxyData->TextureRHI;
+			if (!ensureMsgf(!OutputParam.IsUAVBound(), TEXT("NiagaraDIRenderTarget2DArray(%s) is bound as both read & write, read will be ignored."), *Context.DataInterface->SourceDIName.ToString()))
+			{
+				//-TODO: Feedback to the user that read & write is bound
+				TextureRHI = nullptr;
+			}
+
 			if (TextureRHI == nullptr)
 			{
 				TextureRHI = GBlackTexture->TextureRHI;
@@ -131,7 +128,7 @@ public:
 			FNiagaraDataInterfaceProxyRenderTarget2DArrayProxy* VFDI = static_cast<FNiagaraDataInterfaceProxyRenderTarget2DArrayProxy*>(Context.DataInterface);
 			if (FRenderTarget2DArrayRWInstanceData_RenderThread* ProxyData = VFDI->SystemInstancesToProxyData_RT.Find(Context.SystemInstanceID))
 			{
-				if (FRHIUnorderedAccessView* OutputUAV = ProxyData->UAV)
+				if (FRHIUnorderedAccessView* OutputUAV = ProxyData->UnorderedAccessViewRHI)
 				{
 					RHICmdList.Transition(FRHITransitionInfo(OutputUAV, ERHIAccess::UAVCompute, ERHIAccess::SRVMask));
 				}
@@ -528,7 +525,7 @@ bool UNiagaraDataInterfaceRenderTarget2DArray::InitPerInstanceData(void* PerInst
 	// Push Updates to Proxy.
 	FNiagaraDataInterfaceProxyRenderTarget2DArrayProxy* RT_Proxy = GetProxyAs<FNiagaraDataInterfaceProxyRenderTarget2DArrayProxy>();
 	ENQUEUE_RENDER_COMMAND(FUpdateData)(
-		[RT_Proxy, RT_InstanceID = SystemInstance->GetId(), RT_InstanceData = *InstanceData, RT_TargetTexture = InstanceData->TargetTexture](FRHICommandListImmediate& RHICmdList)
+		[RT_Proxy, RT_InstanceID=SystemInstance->GetId(), RT_InstanceData=*InstanceData, RT_TargetTexture=InstanceData->TargetTexture ? InstanceData->TargetTexture->GameThread_GetRenderTargetResource() : nullptr](FRHICommandListImmediate& RHICmdList)
 		{
 			check(!RT_Proxy->SystemInstancesToProxyData_RT.Contains(RT_InstanceID));
 			FRenderTarget2DArrayRWInstanceData_RenderThread* TargetData = &RT_Proxy->SystemInstancesToProxyData_RT.Add(RT_InstanceID);
@@ -537,9 +534,15 @@ bool UNiagaraDataInterfaceRenderTarget2DArray::InitPerInstanceData(void* PerInst
 #if WITH_EDITORONLY_DATA
 			TargetData->bPreviewTexture = RT_InstanceData.bPreviewTexture;
 #endif
-			TargetData->SamplerStateRHI = RT_TargetTexture->Resource && RT_TargetTexture->Resource->SamplerStateRHI ? RT_TargetTexture->Resource->SamplerStateRHI : nullptr;
-			TargetData->TextureReferenceRHI = RT_TargetTexture->TextureReference.TextureReferenceRHI;
-			TargetData->UAV.SafeRelease();
+			if (RT_TargetTexture)
+			{
+				if (FTextureRenderTarget2DArrayResource* Resource2DArray = RT_TargetTexture->GetTextureRenderTarget2DArrayResource())
+				{
+					TargetData->SamplerStateRHI = Resource2DArray->SamplerStateRHI;
+					TargetData->TextureRHI = Resource2DArray->GetTextureRHI();
+					TargetData->UnorderedAccessViewRHI = Resource2DArray->GetUnorderedAccessViewRHI();
+				}
+			}
 		}
 	);
 	return true;
@@ -707,7 +710,7 @@ bool UNiagaraDataInterfaceRenderTarget2DArray::PerInstanceTickPostSimulate(void*
 		FNiagaraDataInterfaceProxyRenderTarget2DArrayProxy* RT_Proxy = GetProxyAs<FNiagaraDataInterfaceProxyRenderTarget2DArrayProxy>();
 		ENQUEUE_RENDER_COMMAND(FUpdateData)
 		(
-			[RT_Proxy, RT_InstanceID=SystemInstance->GetId(), RT_InstanceData=*InstanceData, RT_TargetTexture=InstanceData->TargetTexture](FRHICommandListImmediate& RHICmdList)
+			[RT_Proxy, RT_InstanceID=SystemInstance->GetId(), RT_InstanceData=*InstanceData, RT_TargetTexture=InstanceData->TargetTexture ? InstanceData->TargetTexture->GameThread_GetRenderTargetResource() : nullptr](FRHICommandListImmediate& RHICmdList)
 			{
 				FRenderTarget2DArrayRWInstanceData_RenderThread* TargetData = RT_Proxy->SystemInstancesToProxyData_RT.Find(RT_InstanceID);
 				if (ensureMsgf(TargetData != nullptr, TEXT("InstanceData was not found for %llu"), RT_InstanceID))
@@ -716,9 +719,18 @@ bool UNiagaraDataInterfaceRenderTarget2DArray::PerInstanceTickPostSimulate(void*
 #if WITH_EDITORONLY_DATA
 					TargetData->bPreviewTexture = RT_InstanceData.bPreviewTexture;
 #endif
-					TargetData->SamplerStateRHI = RT_TargetTexture->Resource && RT_TargetTexture->Resource->SamplerStateRHI ? RT_TargetTexture->Resource->SamplerStateRHI : nullptr;
-					TargetData->TextureReferenceRHI = RT_TargetTexture->TextureReference.TextureReferenceRHI;
-					TargetData->UAV.SafeRelease();
+					TargetData->SamplerStateRHI.SafeRelease();
+					TargetData->TextureRHI.SafeRelease();
+					TargetData->UnorderedAccessViewRHI.SafeRelease();
+					if (RT_TargetTexture)
+					{
+						if (FTextureRenderTarget2DArrayResource* Resource2DArray = RT_TargetTexture->GetTextureRenderTarget2DArrayResource())
+						{
+							TargetData->SamplerStateRHI = Resource2DArray->SamplerStateRHI;
+							TargetData->TextureRHI = Resource2DArray->GetTextureRHI();
+							TargetData->UnorderedAccessViewRHI = Resource2DArray->GetUnorderedAccessViewRHI();
+						}
+					}
 				}
 			}
 		);
@@ -732,11 +744,11 @@ void FNiagaraDataInterfaceProxyRenderTarget2DArrayProxy::PostSimulate(FRHIComman
 #if NIAGARA_COMPUTEDEBUG_ENABLED
 	FRenderTarget2DArrayRWInstanceData_RenderThread* ProxyData = SystemInstancesToProxyData_RT.Find(Context.SystemInstanceID);
 
-	if (ProxyData && ProxyData->bPreviewTexture && ProxyData->TextureReferenceRHI.IsValid())
+	if (ProxyData && ProxyData->bPreviewTexture)
 	{
 		if (FNiagaraGpuComputeDebug* GpuComputeDebug = Context.Batcher->GetGpuComputeDebug())
 		{
-			if (FRHITexture* RHITexture = ProxyData->TextureReferenceRHI->GetReferencedTexture())
+			if (FRHITexture* RHITexture = ProxyData->TextureRHI)
 			{
 				GpuComputeDebug->AddTexture(RHICmdList, Context.SystemInstanceID, SourceDIName, RHITexture);
 			}
diff --git a/Engine/Plugins/FX/Niagara/Source/Niagara/Private/NiagaraDataInterfaceRenderTargetVolume.cpp b/Engine/Plugins/FX/Niagara/Source/Niagara/Private/NiagaraDataInterfaceRenderTargetVolume.cpp
--- a/Engine/Plugins/FX/Niagara/Source/Niagara/Private/NiagaraDataInterfaceRenderTargetVolume.cpp
+++ b/Engine/Plugins/FX/Niagara/Source/Niagara/Private/NiagaraDataInterfaceRenderTargetVolume.cpp
@@ -4,6 +4,9 @@
 #include "ShaderParameterUtils.h"
 #include "ClearQuad.h"
 #include "TextureResource.h"
+#include "TextureRenderTargetVolumeResource.h"
+#include "Engine/TextureRenderTargetVolume.h"
+
 #include "NiagaraEmitterInstanceBatcher.h"
 #include "NiagaraSystemInstance.h"
 #include "NiagaraRenderer.h"
@@ -11,7 +14,6 @@
 #if WITH_EDITOR
 #include "NiagaraGpuComputeDebug.h"
 #endif
-#include "Engine/TextureRenderTargetVolume.h"
 
 #define LOCTEXT_NAMESPACE "NiagaraDataInterfaceRenderTargetVolume"
 
@@ -73,23 +75,12 @@ public:
 		
 		if ( OutputParam.IsUAVBound())
 		{
-			FRHIUnorderedAccessView* OutputUAV = nullptr;
-
-			if (FRHITexture* TextureRHI = ProxyData->TextureReferenceRHI->GetReferencedTexture())
+			FRHIUnorderedAccessView* OutputUAV = ProxyData->UnorderedAccessViewRHI;
+			if (OutputUAV)
 			{
-				// Note: Because we control the render target it should not changed underneath us without queueing a request to recreate the UAV.  If that assumption changes we would need to track the UAV.
-				ERHIAccess InitialState = ERHIAccess::SRVMask;
-				if (!ProxyData->UAV.IsValid())
-				{
-					ProxyData->UAV = RHICreateUnorderedAccessView(TextureRHI, 0);
-					InitialState = ERHIAccess::Unknown;
-				}
-
-				OutputUAV = ProxyData->UAV;
-				RHICmdList.Transition(FRHITransitionInfo(OutputUAV, InitialState, ERHIAccess::UAVCompute));
+				RHICmdList.Transition(FRHITransitionInfo(OutputUAV, ERHIAccess::Unknown, ERHIAccess::UAVCompute));
 			}
-
-			if (OutputUAV == nullptr)
+			else
 			{
 				OutputUAV = Context.Batcher->GetEmptyRWTextureFromPool(RHICmdList, EPixelFormat::PF_A16B16G16R16);
 			}
@@ -98,7 +89,13 @@ public:
 
 		if (InputParam.IsBound())
 		{
-			FRHITexture* TextureRHI = ProxyData->TextureReferenceRHI->GetReferencedTexture();
+			FRHITexture* TextureRHI = ProxyData->TextureRHI;
+			if (!ensureMsgf(!OutputParam.IsUAVBound(), TEXT("NiagaraDIRenderTargetVolume(%s) is bound as both read & write, read will be ignored."), *Context.DataInterface->SourceDIName.ToString()))
+			{
+				//-TODO: Feedback to the user that read & write is bound
+				TextureRHI = nullptr;
+			}
+
 			if (TextureRHI == nullptr)
 			{
 				TextureRHI = GBlackTexture->TextureRHI;
@@ -129,7 +126,7 @@ public:
 			FNiagaraDataInterfaceProxyRenderTargetVolumeProxy* VFDI = static_cast<FNiagaraDataInterfaceProxyRenderTargetVolumeProxy*>(Context.DataInterface);
 			if (FRenderTargetVolumeRWInstanceData_RenderThread* ProxyData = VFDI->SystemInstancesToProxyData_RT.Find(Context.SystemInstanceID))
 			{
-				if (FRHIUnorderedAccessView* OutputUAV = ProxyData->UAV)
+				if (FRHIUnorderedAccessView* OutputUAV = ProxyData->UnorderedAccessViewRHI)
 				{
 					RHICmdList.Transition(FRHITransitionInfo(OutputUAV, ERHIAccess::UAVCompute, ERHIAccess::SRVMask));
 				}
@@ -524,7 +521,7 @@ bool UNiagaraDataInterfaceRenderTargetVolume::InitPerInstanceData(void* PerInsta
 	FNiagaraDataInterfaceProxyRenderTargetVolumeProxy* RT_Proxy = GetProxyAs<FNiagaraDataInterfaceProxyRenderTargetVolumeProxy>();
 	ENQUEUE_RENDER_COMMAND(FUpdateData)
 	(
-		[RT_Proxy, RT_InstanceID=SystemInstance->GetId(), RT_InstanceData=*InstanceData, RT_TargetTexture=InstanceData->TargetTexture](FRHICommandListImmediate& RHICmdList)
+		[RT_Proxy, RT_InstanceID=SystemInstance->GetId(), RT_InstanceData=*InstanceData, RT_TargetTexture= InstanceData->TargetTexture ? InstanceData->TargetTexture->GameThread_GetRenderTargetResource() : nullptr](FRHICommandListImmediate& RHICmdList)
 		{
 			check(!RT_Proxy->SystemInstancesToProxyData_RT.Contains(RT_InstanceID));
 			FRenderTargetVolumeRWInstanceData_RenderThread* TargetData = &RT_Proxy->SystemInstancesToProxyData_RT.Add(RT_InstanceID);
@@ -533,9 +530,15 @@ bool UNiagaraDataInterfaceRenderTargetVolume::InitPerInstanceData(void* PerInsta
 #if WITH_EDITORONLY_DATA
 			TargetData->bPreviewTexture = RT_InstanceData.bPreviewTexture;
 #endif
-			TargetData->SamplerStateRHI = RT_TargetTexture->Resource && RT_TargetTexture->Resource->SamplerStateRHI ? RT_TargetTexture->Resource->SamplerStateRHI : nullptr;
-			TargetData->TextureReferenceRHI = RT_TargetTexture->TextureReference.TextureReferenceRHI;
-			TargetData->UAV.SafeRelease();
+			if (RT_TargetTexture)
+			{
+				if (FTextureRenderTargetVolumeResource* ResourceVolume = RT_TargetTexture->GetTextureRenderTargetVolumeResource())
+				{
+					TargetData->SamplerStateRHI = ResourceVolume->SamplerStateRHI;
+					TargetData->TextureRHI = ResourceVolume->GetTextureRHI();
+					TargetData->UnorderedAccessViewRHI = ResourceVolume->GetUnorderedAccessViewRHI();
+				}
+			}
 		}
 	);
 	return true;
@@ -703,7 +706,7 @@ bool UNiagaraDataInterfaceRenderTargetVolume::PerInstanceTickPostSimulate(void*
 		FNiagaraDataInterfaceProxyRenderTargetVolumeProxy* RT_Proxy = GetProxyAs<FNiagaraDataInterfaceProxyRenderTargetVolumeProxy>();
 		ENQUEUE_RENDER_COMMAND(FUpdateData)
 		(
-			[RT_Proxy, RT_InstanceID=SystemInstance->GetId(), RT_InstanceData=*InstanceData, RT_TargetTexture=InstanceData->TargetTexture](FRHICommandListImmediate& RHICmdList)
+			[RT_Proxy, RT_InstanceID=SystemInstance->GetId(), RT_InstanceData=*InstanceData, RT_TargetTexture=InstanceData->TargetTexture ? InstanceData->TargetTexture->GameThread_GetRenderTargetResource() : nullptr](FRHICommandListImmediate& RHICmdList)
 			{
 				FRenderTargetVolumeRWInstanceData_RenderThread* TargetData = RT_Proxy->SystemInstancesToProxyData_RT.Find(RT_InstanceID);
 				if (ensureMsgf(TargetData != nullptr, TEXT("InstanceData was not found for %llu"), RT_InstanceID))
@@ -712,9 +715,18 @@ bool UNiagaraDataInterfaceRenderTargetVolume::PerInstanceTickPostSimulate(void*
 #if WITH_EDITORONLY_DATA
 					TargetData->bPreviewTexture = RT_InstanceData.bPreviewTexture;
 #endif
-					TargetData->SamplerStateRHI = RT_TargetTexture->Resource && RT_TargetTexture->Resource->SamplerStateRHI ? RT_TargetTexture->Resource->SamplerStateRHI : nullptr;
-					TargetData->TextureReferenceRHI = RT_TargetTexture->TextureReference.TextureReferenceRHI;
-					TargetData->UAV.SafeRelease();
+					TargetData->SamplerStateRHI.SafeRelease();
+					TargetData->TextureRHI.SafeRelease();
+					TargetData->UnorderedAccessViewRHI.SafeRelease();
+					if (RT_TargetTexture)
+					{
+						if (FTextureRenderTargetVolumeResource* ResourceVolume = RT_TargetTexture->GetTextureRenderTargetVolumeResource())
+						{
+							TargetData->SamplerStateRHI = ResourceVolume->SamplerStateRHI;
+							TargetData->TextureRHI = ResourceVolume->GetTextureRHI();
+							TargetData->UnorderedAccessViewRHI = ResourceVolume->GetUnorderedAccessViewRHI();
+						}
+					}
 				}
 			}
 		);
@@ -728,11 +740,11 @@ void FNiagaraDataInterfaceProxyRenderTargetVolumeProxy::PostSimulate(FRHICommand
 #if NIAGARA_COMPUTEDEBUG_ENABLED
 	FRenderTargetVolumeRWInstanceData_RenderThread* ProxyData = SystemInstancesToProxyData_RT.Find(Context.SystemInstanceID);
 
-	if (ProxyData && ProxyData->bPreviewTexture && ProxyData->TextureReferenceRHI.IsValid())
+	if (ProxyData && ProxyData->bPreviewTexture && ProxyData->TextureRHI.IsValid())
 	{
 		if (FNiagaraGpuComputeDebug* GpuComputeDebug = Context.Batcher->GetGpuComputeDebug())
 		{
-			if ( FRHITexture* RHITexture = ProxyData->TextureReferenceRHI->GetReferencedTexture() )
+			if ( FRHITexture* RHITexture = ProxyData->TextureRHI )
 			{
 				GpuComputeDebug->AddTexture(RHICmdList, Context.SystemInstanceID, SourceDIName, RHITexture);
 			}
diff --git a/Engine/Source/Runtime/Engine/Private/TextureRenderTarget2D.cpp b/Engine/Source/Runtime/Engine/Private/TextureRenderTarget2D.cpp
--- a/Engine/Source/Runtime/Engine/Private/TextureRenderTarget2D.cpp
+++ b/Engine/Source/Runtime/Engine/Private/TextureRenderTarget2D.cpp
@@ -524,6 +524,12 @@ void FTextureRenderTarget2DResource::InitDynamicRHI()
 			RenderTargetTextureRHI,
 			Texture2DRHI
 			);
+
+		if ((TexCreateFlags & TexCreate_UAV) != 0)
+		{
+			UnorderedAccessViewRHI = RHICreateUnorderedAccessView(RenderTargetTextureRHI);
+		}
+
 		SetGPUMask(CreateInfo.GPUMask);
 		TextureRHI = (FTextureRHIRef&)Texture2DRHI;
 		RHIUpdateTextureReference(Owner->TextureReference.TextureReferenceRHI,TextureRHI);
diff --git a/Engine/Source/Runtime/Engine/Private/TextureRenderTarget2DArray.cpp b/Engine/Source/Runtime/Engine/Private/TextureRenderTarget2DArray.cpp
--- a/Engine/Source/Runtime/Engine/Private/TextureRenderTarget2DArray.cpp
+++ b/Engine/Source/Runtime/Engine/Private/TextureRenderTarget2DArray.cpp
@@ -255,6 +255,11 @@ void FTextureRenderTarget2DArrayResource::InitDynamicRHI()
 			);
 		}
 
+		if ((TexCreateFlags & TexCreate_UAV) != 0)
+		{
+			UnorderedAccessViewRHI = RHICreateUnorderedAccessView(RenderTarget2DArrayRHI);
+		}
+
 		TextureRHI = Texture2DArrayRHI;
 		RHIUpdateTextureReference(Owner->TextureReference.TextureReferenceRHI, TextureRHI);
 
diff --git a/Engine/Source/Runtime/Engine/Private/TextureRenderTargetVolume.cpp b/Engine/Source/Runtime/Engine/Private/TextureRenderTargetVolume.cpp
--- a/Engine/Source/Runtime/Engine/Private/TextureRenderTargetVolume.cpp
+++ b/Engine/Source/Runtime/Engine/Private/TextureRenderTargetVolume.cpp
@@ -262,6 +262,11 @@ void FTextureRenderTargetVolumeResource::InitDynamicRHI()
 			);
 		}
 
+		if ((TexCreateFlags & TexCreate_UAV) != 0)
+		{
+			UnorderedAccessViewRHI = RHICreateUnorderedAccessView(RenderTargetVolumeRHI);
+		}
+
 		TextureRHI = TextureVolumeRHI;
 		RHIUpdateTextureReference(Owner->TextureReference.TextureReferenceRHI,TextureRHI);
 
diff --git a/Engine/Source/Runtime/Engine/Public/TextureRenderTarget2DArrayResource.h b/Engine/Source/Runtime/Engine/Public/TextureRenderTarget2DArrayResource.h
--- a/Engine/Source/Runtime/Engine/Public/TextureRenderTarget2DArrayResource.h
+++ b/Engine/Source/Runtime/Engine/Public/TextureRenderTarget2DArrayResource.h
@@ -68,6 +68,11 @@ public:
 	 */
 	FTexture2DArrayRHIRef GetTextureRHI() { return Texture2DArrayRHI; }
 
+	/**
+	 * @return UnorderedAccessView for rendering
+	 */
+	FUnorderedAccessViewRHIRef GetUnorderedAccessViewRHI() { return UnorderedAccessViewRHI; }
+
 	/**
 	* Render target resource should be sampled in linear color space
 	*
@@ -109,4 +114,6 @@ private:
 	FTexture2DArrayRHIRef RenderTarget2DArrayRHI;
 	/** Texture resource used for rendering with and resolving to */
 	FTexture2DArrayRHIRef Texture2DArrayRHI;
+	/** Optional Unordered Access View for the resource, automatically created if bCanCreateUAV is true */
+	FUnorderedAccessViewRHIRef UnorderedAccessViewRHI;
 };
diff --git a/Engine/Source/Runtime/Engine/Public/TextureRenderTargetVolumeResource.h b/Engine/Source/Runtime/Engine/Public/TextureRenderTargetVolumeResource.h
--- a/Engine/Source/Runtime/Engine/Public/TextureRenderTargetVolumeResource.h
+++ b/Engine/Source/Runtime/Engine/Public/TextureRenderTargetVolumeResource.h
@@ -68,6 +68,11 @@ public:
 	 */
 	FTexture3DRHIRef GetTextureRHI() { return TextureVolumeRHI; }
 
+	/**
+	 * @return UnorderedAccessView for rendering
+	 */
+	FUnorderedAccessViewRHIRef GetUnorderedAccessViewRHI() { return UnorderedAccessViewRHI; }
+
 	/**
 	* Render target resource should be sampled in linear color space
 	*
@@ -109,4 +114,6 @@ private:
 	FTexture3DRHIRef RenderTargetVolumeRHI;
 	/** Texture resource used for rendering with and resolving to */
 	FTexture3DRHIRef TextureVolumeRHI;
+	/** Optional Unordered Access View for the resource, automatically created if bCanCreateUAV is true */
+	FUnorderedAccessViewRHIRef UnorderedAccessViewRHI;
 };
diff --git a/Engine/Source/Runtime/Engine/Public/TextureResource.h b/Engine/Source/Runtime/Engine/Public/TextureResource.h
--- a/Engine/Source/Runtime/Engine/Public/TextureResource.h
+++ b/Engine/Source/Runtime/Engine/Public/TextureResource.h
@@ -315,10 +315,10 @@ public:
 
 	// FTextureRenderTargetResource interface
 	
-	virtual class FTextureRenderTarget2DResource* GetTextureRenderTarget2DResource()
-	{
-		return NULL;
-	}
+	virtual class FTextureRenderTarget2DResource* GetTextureRenderTarget2DResource() { return nullptr; }
+	virtual class FTextureRenderTarget2DArrayResource* GetTextureRenderTarget2DArrayResource() { return nullptr; }
+	virtual class FTextureRenderTargetVolumeResource* GetTextureRenderTargetVolumeResource() { return nullptr; }
+
 	virtual void ClampSize(int32 SizeX,int32 SizeY) {}
 
 	// FRenderTarget interface.
@@ -441,6 +441,12 @@ public:
 	 * @return TextureRHI for rendering 
 	 */
 	FTexture2DRHIRef GetTextureRHI() { return Texture2DRHI; }
+
+	/**
+	 * @return UnorderedAccessView for rendering
+	 */
+	FUnorderedAccessViewRHIRef GetUnorderedAccessViewRHI() { return UnorderedAccessViewRHI; }
+
 protected:
 	/**
 	 * Updates (resolves) the render target texture.
@@ -456,6 +462,8 @@ private:
 	const class UTextureRenderTarget2D* Owner;
 	/** Texture resource used for rendering with and resolving to */
 	FTexture2DRHIRef Texture2DRHI;
+	/** Optional Unordered Access View for the resource, automatically created if bCanCreateUAV is true */
+	FUnorderedAccessViewRHIRef UnorderedAccessViewRHI;
 	/** the color the texture is cleared to */
 	FLinearColor ClearColor;
 	EPixelFormat Format;
