Fix for missing memmark.

diff --git a/Engine/Plugins/FX/Niagara/Source/Niagara/Private/NiagaraEmitterInstanceBatcher.cpp b/Engine/Plugins/FX/Niagara/Source/Niagara/Private/NiagaraEmitterInstanceBatcher.cpp
--- a/Engine/Plugins/FX/Niagara/Source/Niagara/Private/NiagaraEmitterInstanceBatcher.cpp
+++ b/Engine/Plugins/FX/Niagara/Source/Niagara/Private/NiagaraEmitterInstanceBatcher.cpp
@@ -480,6 +480,7 @@ void NiagaraEmitterInstanceBatcher::ProcessPendingTicksFlush(FRHICommandListImme
 			View.ViewUniformBuffer = TUniformBufferRef<FViewUniformShaderParameters>::CreateUniformBufferImmediate(*View.CachedViewUniformShaderParameters, UniformBuffer_SingleFrame);
 
 			// Execute all ticks that we can support without invalid simulations
+			FMemMark MemMark(FMemStack::Get());
 			FRDGBuilder GraphBuilder(RHICmdList);
 			PreInitViews(GraphBuilder, true);
 			PostInitViews(GraphBuilder, View.ViewUniformBuffer, true);
