Fix up some UE4 strings in Niagara

diff --git a/Engine/Plugins/FX/Niagara/Source/Niagara/Private/NiagaraConstants.cpp b/Engine/Plugins/FX/Niagara/Source/Niagara/Private/NiagaraConstants.cpp
--- a/Engine/Plugins/FX/Niagara/Source/Niagara/Private/NiagaraConstants.cpp
+++ b/Engine/Plugins/FX/Niagara/Source/Niagara/Private/NiagaraConstants.cpp
@@ -523,7 +523,7 @@ void FNiagaraConstants::Init()
 		AttrDescStrMap.Add(SYS_PARAM_PARTICLES_VISIBILITY_TAG, LOCTEXT("VisibilityTag", "Used for selecting renderers to use when rendering this particle. Without this, the particle will render in all renderers"));
 		AttrDescStrMap.Add(SYS_PARAM_PARTICLES_COMPONENTS_ENABLED, LOCTEXT("ComponentRenderEnabledParamDesc", "Used to check if component rendering should be enabled on a per-particle basis. Without this, the each particle will spawn a component."));
 		AttrDescStrMap.Add(SYS_PARAM_PARTICLES_RIBBONID, LOCTEXT("RibbonIDDesc", "Sets the ribbon id for a particle. Particles with the same ribbon id will be connected into a ribbon."));
-		AttrDescStrMap.Add(SYS_PARAM_PARTICLES_RIBBONWIDTH, LOCTEXT("RibbonWidthDesc", "Sets the ribbon width for a particle, in UE4 units."));
+		AttrDescStrMap.Add(SYS_PARAM_PARTICLES_RIBBONWIDTH, LOCTEXT("RibbonWidthDesc", "Sets the ribbon width for a particle, in Unreal units."));
 		AttrDescStrMap.Add(SYS_PARAM_PARTICLES_RIBBONTWIST, LOCTEXT("RibbonTwistDesc", "Sets the ribbon twist for a particle, in degrees."));
 		AttrDescStrMap.Add(SYS_PARAM_PARTICLES_RIBBONFACING, LOCTEXT("RibbonFacingDesc", "Sets the facing vector of the ribbon at the particle position, or the side vector the ribbon's width is extended along, depending on the selected facing mode."));
 		AttrDescStrMap.Add(SYS_PARAM_PARTICLES_RIBBONLINKORDER, LOCTEXT("RibbonLinkOrderDesc", "Explicit order for linking particles within a ribbon. Particles of the same ribbon id will be connected into a ribbon in incrementing order of this attribute value."));
diff --git a/Engine/Plugins/FX/Niagara/Source/NiagaraEditor/Private/NiagaraEditorModule.cpp b/Engine/Plugins/FX/Niagara/Source/NiagaraEditor/Private/NiagaraEditorModule.cpp
--- a/Engine/Plugins/FX/Niagara/Source/NiagaraEditor/Private/NiagaraEditorModule.cpp
+++ b/Engine/Plugins/FX/Niagara/Source/NiagaraEditor/Private/NiagaraEditorModule.cpp
@@ -673,7 +673,7 @@ void ExecuteInvalidateNiagaraCachedScripts(const TArray< FString >& Args)
 		}
 		else if (SourceControlState->IsCheckedOutOther())
 		{
-			UE_LOG(LogConsoleResponse, Display, TEXT("fx.InvalidateCachedScripts failed\n\"NiagaraShaderVersion.ush\" is already checked out by someone else\n(UE4 SourceControl needs to be fixed to allow multiple checkout.)"));
+			UE_LOG(LogConsoleResponse, Display, TEXT("fx.InvalidateCachedScripts failed\n\"NiagaraShaderVersion.ush\" is already checked out by someone else\n(Unreal SourceControl needs to be fixed to allow multiple checkout.)"));
 			return;
 		}
 		else if (SourceControlState->IsDeleted())
