Adjust FFXSystemInterface API to take an array of views rather than individual elements from the views.  No functional change, just sets the stage for being able to get additional data out of the views.

diff --git a/Engine/Plugins/FX/Niagara/Source/Niagara/Classes/NiagaraEmitterInstanceBatcher.h b/Engine/Plugins/FX/Niagara/Source/Niagara/Classes/NiagaraEmitterInstanceBatcher.h
--- a/Engine/Plugins/FX/Niagara/Source/Niagara/Classes/NiagaraEmitterInstanceBatcher.h
+++ b/Engine/Plugins/FX/Niagara/Source/Niagara/Classes/NiagaraEmitterInstanceBatcher.h
@@ -8,21 +8,21 @@ the same VectorVM byte code / compute shader code
 #pragma once
 
 #include "CoreMinimal.h"
-#include "NiagaraCommon.h"
-#include "RendererInterface.h"
-#include "NiagaraParameters.h"
-#include "NiagaraEmitter.h"
-#include "Tickable.h"
-#include "Modules/ModuleManager.h"
-#include "RHIResources.h"
 #include "FXSystem.h"
-#include "NiagaraRendererProperties.h"
+#include "Modules/ModuleManager.h"
 #include "ParticleResources.h"
-#include "Runtime/Engine/Private/Particles/ParticleSortingGPU.h"
-#include "NiagaraGPUSortInfo.h"
-#include "NiagaraScriptExecutionContext.h"
+#include "NiagaraCommon.h"
+#include "NiagaraEmitter.h"
 #include "NiagaraGPUInstanceCountManager.h"
 #include "NiagaraGPUProfiler.h"
+#include "NiagaraGPUSortInfo.h"
+#include "NiagaraParameters.h"
+#include "NiagaraRendererProperties.h"
+#include "NiagaraScriptExecutionContext.h"
+#include "RendererInterface.h"
+#include "RHIResources.h"
+#include "Runtime/Engine/Private/Particles/ParticleSortingGPU.h"
+#include "Tickable.h"
 
 class FGPUSortManager;
 class FNiagaraGpuComputeDebug;
@@ -73,20 +73,17 @@ public:
 	virtual void AddVectorField(UVectorFieldComponent* VectorFieldComponent) override {}
 	virtual void RemoveVectorField(UVectorFieldComponent* VectorFieldComponent) override {}
 	virtual void UpdateVectorField(UVectorFieldComponent* VectorFieldComponent) override {}
-	virtual void PreInitViews(class FRDGBuilder& GraphBuilder, bool bAllowGPUParticleUpdate) override;
-	virtual void PostInitViews(class FRDGBuilder& GraphBuilder, FRHIUniformBuffer* ViewUniformBuffer, bool bAllowGPUParticleUpdate) override;
+	virtual void PreInitViews(FRDGBuilder& GraphBuilder, bool bAllowGPUParticleUpdate) override;
+	virtual void PostInitViews(FRDGBuilder& GraphBuilder, TConstArrayView<FViewInfo> Views, bool bAllowGPUParticleUpdate) override;
 	virtual bool UsesGlobalDistanceField() const override;
 	virtual bool UsesDepthBuffer() const override;
 	virtual bool RequiresEarlyViewUniformBuffer() const override;
-	virtual void PreRender(class FRDGBuilder& GraphBuilder, FRHIUniformBuffer* ViewUniformBuffer, const class FGlobalDistanceFieldParameterData* GlobalDistanceFieldParameterData, bool bAllowGPUParticleUpdate) override;
+	virtual void PreRender(FRDGBuilder& GraphBuilder, TConstArrayView<FViewInfo> Views, bool bAllowGPUParticleUpdate) override;
 	virtual void OnDestroy() override; // Called on the gamethread to delete the batcher on the renderthread.
 
 	virtual void Tick(float DeltaTime) override;
 
-	virtual void PostRenderOpaque(
-		class FRDGBuilder& GraphBuilder,
-		FRHIUniformBuffer* ViewUniformBuffer,
-		bool bAllowGPUParticleUpdate) override;
+	virtual void PostRenderOpaque(FRDGBuilder& GraphBuilder, TConstArrayView<FViewInfo> Views, bool bAllowGPUParticleUpdate) override;
 
 	/**
 	 * Process and respond to a build up of excessive ticks inside the batcher.
diff --git a/Engine/Plugins/FX/Niagara/Source/Niagara/Private/NiagaraEmitterInstanceBatcher.cpp b/Engine/Plugins/FX/Niagara/Source/Niagara/Private/NiagaraEmitterInstanceBatcher.cpp
--- a/Engine/Plugins/FX/Niagara/Source/Niagara/Private/NiagaraEmitterInstanceBatcher.cpp
+++ b/Engine/Plugins/FX/Niagara/Source/Niagara/Private/NiagaraEmitterInstanceBatcher.cpp
@@ -1,29 +1,30 @@
 // Copyright Epic Games, Inc. All Rights Reserved.
 
 #include "NiagaraEmitterInstanceBatcher.h"
-#include "NiagaraScriptExecutionContext.h"
-#include "RHI.h"
-#include "RHIGPUReadback.h"
-#include "NiagaraStats.h"
-#include "NiagaraShader.h"
-#include "NiagaraSortingGPU.h"
-#include "NiagaraWorldManager.h"
-#include "NiagaraShaderParticleID.h"
-#include "NiagaraRenderer.h"
+
+#include "Async/Async.h"
+#include "CanvasTypes.h"
+#include "ClearQuad.h"
+#include "Engine/Canvas.h"
+#include "GPUSort.h"
 #include "NiagaraDataInterfaceRW.h"
 #if WITH_EDITOR
 #include "NiagaraGpuComputeDebug.h"
 #endif
 #include "NiagaraGpuReadbackManager.h"
-#include "ShaderParameterUtils.h"
-#include "SceneUtils.h"
-#include "ClearQuad.h"
-#include "Async/Async.h"
-#include "GPUSort.h"
-#include "CanvasTypes.h"
-#include "Engine/Canvas.h"
+#include "NiagaraRenderer.h"
+#include "NiagaraScriptExecutionContext.h"
+#include "NiagaraShader.h"
+#include "NiagaraShaderParticleID.h"
+#include "NiagaraSortingGPU.h"
+#include "NiagaraStats.h"
+#include "NiagaraWorldManager.h"
+#include "RHI.h"
+#include "RHIGPUReadback.h"
 #include "Runtime/Renderer/Private/SceneRendering.h"
 #include "Runtime/Renderer/Private/PostProcess/SceneRenderTargets.h"
+#include "SceneUtils.h"
+#include "ShaderParameterUtils.h"
 
 DECLARE_CYCLE_STAT(TEXT("Niagara Dispatch Setup"), STAT_NiagaraGPUDispatchSetup_RT, STATGROUP_Niagara);
 DECLARE_CYCLE_STAT(TEXT("GPU Emitter Dispatch [RT]"), STAT_NiagaraGPUSimTick_RT, STATGROUP_Niagara);
@@ -458,6 +459,8 @@ void NiagaraEmitterInstanceBatcher::ProcessPendingTicksFlush(FRHICommandListImme
 
 			// Make a temporary ViewInfo
 			//-TODO: We could gather some more information here perhaps?
+			FMemMark Mark(FMemStack::Get());
+
 			FSceneViewFamily ViewFamily(FSceneViewFamily::ConstructionValues(nullptr, nullptr, FEngineShowFlags(ESFIM_Game))
 				.SetWorldTimes(0, 0, 0)
 				.SetGammaCorrection(1.0f));
@@ -470,21 +473,33 @@ void NiagaraEmitterInstanceBatcher::ProcessPendingTicksFlush(FRHICommandListImme
 			ViewInitOptions.ViewRotationMatrix = FMatrix::Identity;
 			ViewInitOptions.ProjectionMatrix = FMatrix::Identity;
 
-			FViewInfo View(ViewInitOptions);
-			View.ViewRect = View.UnscaledViewRect;
-			View.CachedViewUniformShaderParameters = MakeUnique<FViewUniformShaderParameters>();
+			FViewInfo* DummyView = nullptr;
+
+			if (RHICmdList.Bypass())
+			{
+				DummyView = new(FMemStack::Get().Alloc(sizeof(FViewInfo), alignof(FViewInfo))) FViewInfo(ViewInitOptions);
+			}
+			else
+			{
+				DummyView = new(RHICmdList.Alloc(sizeof(FViewInfo), alignof(FViewInfo))) FViewInfo(ViewInitOptions);
+			}
+
+			DummyView->ViewRect = DummyView->UnscaledViewRect;
+			DummyView->CachedViewUniformShaderParameters = MakeUnique<FViewUniformShaderParameters>();
 
 			FBox UnusedVolumeBounds[TVC_MAX];
-			View.SetupUniformBufferParameters(UnusedVolumeBounds, TVC_MAX, *View.CachedViewUniformShaderParameters);
+			DummyView->SetupUniformBufferParameters(UnusedVolumeBounds, TVC_MAX, *DummyView->CachedViewUniformShaderParameters);
 
-			View.ViewUniformBuffer = TUniformBufferRef<FViewUniformShaderParameters>::CreateUniformBufferImmediate(*View.CachedViewUniformShaderParameters, UniformBuffer_SingleFrame);
+			DummyView->ViewUniformBuffer = TUniformBufferRef<FViewUniformShaderParameters>::CreateUniformBufferImmediate(*DummyView->CachedViewUniformShaderParameters, UniformBuffer_SingleFrame);
 
+			TConstArrayView<FViewInfo> DummyViews = MakeArrayView(DummyView, 1);
+			const bool AllowGPUParticleUpdate = false;
+			
 			// Execute all ticks that we can support without invalid simulations
-			FMemMark MemMark(FMemStack::Get());
 			FRDGBuilder GraphBuilder(RHICmdList);
 			PreInitViews(GraphBuilder, true);
-			PostInitViews(GraphBuilder, View.ViewUniformBuffer, true);
-			PostRenderOpaque(GraphBuilder, View.ViewUniformBuffer, true);
+			PostInitViews(GraphBuilder, DummyViews, AllowGPUParticleUpdate);
+			PostRenderOpaque(GraphBuilder, DummyViews, AllowGPUParticleUpdate);
 			GraphBuilder.Execute();
 			break;
 		}
@@ -1497,12 +1512,16 @@ void NiagaraEmitterInstanceBatcher::PreInitViews(FRDGBuilder& GraphBuilder, bool
 	}
 }
 
-void NiagaraEmitterInstanceBatcher::PostInitViews(FRDGBuilder& GraphBuilder, FRHIUniformBuffer* ViewUniformBuffer, bool bAllowGPUParticleUpdate)
+void NiagaraEmitterInstanceBatcher::PostInitViews(FRDGBuilder& GraphBuilder, TArrayView<const class FViewInfo> Views, bool bAllowGPUParticleUpdate)
 {
 	LLM_SCOPE(ELLMTag::Niagara);
 
+	bAllowGPUParticleUpdate = bAllowGPUParticleUpdate && GetReferenceAllowGPUUpdate(Views);
+
 	if (bAllowGPUParticleUpdate && FNiagaraUtilities::AllowGPUParticles(GetShaderPlatform()))
 	{
+		FRHIUniformBuffer* ViewUniformBuffer = GetReferenceViewUniformBuffer(Views);
+
 		AddPass(
 			GraphBuilder,
 			RDG_EVENT_NAME("Niagara::PostInitViews"),
@@ -1515,10 +1534,13 @@ void NiagaraEmitterInstanceBatcher::PostInitViews(FRDGBuilder& GraphBuilder, FRH
 	}
 }
 
-void NiagaraEmitterInstanceBatcher::PostRenderOpaque(FRDGBuilder& GraphBuilder, FRHIUniformBuffer* ViewUniformBuffer, bool bAllowGPUParticleUpdate)
+void NiagaraEmitterInstanceBatcher::PostRenderOpaque(FRDGBuilder& GraphBuilder, TConstArrayView<FViewInfo> Views, bool bAllowGPUParticleUpdate)
 {
 	LLM_SCOPE(ELLMTag::Niagara);
 
+	bAllowGPUParticleUpdate = bAllowGPUParticleUpdate && GetReferenceAllowGPUUpdate(Views);
+	FRHIUniformBuffer* ViewUniformBuffer = GetReferenceViewUniformBuffer(Views);
+
 	AddPass(
 		GraphBuilder,
 		RDG_EVENT_NAME("Niagara::PostRenderOpaque"),
@@ -1629,7 +1651,7 @@ bool NiagaraEmitterInstanceBatcher::RequiresEarlyViewUniformBuffer() const
 	return NumTicksThatRequireEarlyViewData > 0;
 }
 
-void NiagaraEmitterInstanceBatcher::PreRender(FRDGBuilder& GraphBuilder, FRHIUniformBuffer* ViewUniformBuffer, const class FGlobalDistanceFieldParameterData* GlobalDistanceFieldParameterData, bool bAllowGPUParticleUpdate)
+void NiagaraEmitterInstanceBatcher::PreRender(FRDGBuilder& GraphBuilder, TConstArrayView<FViewInfo> Views, bool bAllowGPUParticleUpdate)
 {
 	if (!FNiagaraUtilities::AllowGPUParticles(GetShaderPlatform()))
 	{
@@ -1638,7 +1660,8 @@ void NiagaraEmitterInstanceBatcher::PreRender(FRDGBuilder& GraphBuilder, FRHIUni
 
 	LLM_SCOPE(ELLMTag::Niagara);
 
-	GlobalDistanceFieldParams = GlobalDistanceFieldParameterData ? *GlobalDistanceFieldParameterData : FGlobalDistanceFieldParameterData();
+	const FGlobalDistanceFieldParameterData* ViewGlobalDistanceFieldParams = GetReferenceGlobalDistanceFieldData(Views);
+	GlobalDistanceFieldParams = ViewGlobalDistanceFieldParams ? *ViewGlobalDistanceFieldParams : FGlobalDistanceFieldParameterData();
 }
 
 void NiagaraEmitterInstanceBatcher::OnDestroy()
diff --git a/Engine/Source/Runtime/Engine/Private/Particles/FXSystem.cpp b/Engine/Source/Runtime/Engine/Private/Particles/FXSystem.cpp
--- a/Engine/Source/Runtime/Engine/Private/Particles/FXSystem.cpp
+++ b/Engine/Source/Runtime/Engine/Private/Particles/FXSystem.cpp
@@ -83,6 +83,24 @@ void FFXSystemInterface::UnregisterCustomFXSystem(const FName& InterfaceName)
 	CreateCustomFXDelegates.Remove(InterfaceName);
 }
 
+FRHIUniformBuffer* FFXSystemInterface::GetReferenceViewUniformBuffer(TConstArrayView<FViewInfo> Views)
+{
+	check(Views.Num() > 0);
+	return Views[0].ViewUniformBuffer;
+}
+
+bool FFXSystemInterface::GetReferenceAllowGPUUpdate(TConstArrayView<FViewInfo> Views)
+{
+	check(Views.Num() > 0);
+	return Views[0].AllowGPUParticleUpdate();
+}
+
+const FGlobalDistanceFieldParameterData* FFXSystemInterface::GetReferenceGlobalDistanceFieldData(TConstArrayView<FViewInfo> Views)
+{
+	check(Views.Num() > 0);
+	return &Views[0].GlobalDistanceFieldInfo.ParameterData;
+}
+
 /*------------------------------------------------------------------------------
 	FX system console variables.
 ------------------------------------------------------------------------------*/
@@ -394,7 +412,7 @@ void FFXSystem::PreInitViews(FRDGBuilder& GraphBuilder, bool bAllowGPUParticleUp
 	}
 }
 
-void FFXSystem::PostInitViews(FRDGBuilder& GraphBuilder, FRHIUniformBuffer* ViewUniformBuffer, bool bAllowGPUParticleUpdate)
+void FFXSystem::PostInitViews(FRDGBuilder& GraphBuilder, TConstArrayView<FViewInfo> Views, bool bAllowGPUParticleUpdate)
 {
 	// nothing to do here
 }
@@ -437,10 +455,15 @@ DECLARE_CYCLE_STAT(TEXT("FXPreRender_SimulateCDF"), STAT_CLM_FXPreRender_Simulat
 DECLARE_CYCLE_STAT(TEXT("FXPreRender_FinalizeCDF"), STAT_CLM_FXPreRender_FinalizeCDF, STATGROUP_CommandListMarkers);
 
 
-void FFXSystem::PreRender(FRDGBuilder& GraphBuilder, FRHIUniformBuffer* ViewUniformBuffer, const FGlobalDistanceFieldParameterData* GlobalDistanceFieldParameterData, bool bAllowGPUParticleSceneUpdate)
+void FFXSystem::PreRender(FRDGBuilder& GraphBuilder, TConstArrayView<FViewInfo> Views, bool bAllowGPUParticleSceneUpdate)
 {
+	bAllowGPUParticleSceneUpdate = bAllowGPUParticleSceneUpdate && GetReferenceAllowGPUUpdate(Views);
+
 	if (RHISupportsGPUParticles() && bAllowGPUParticleSceneUpdate)
 	{
+		FRHIUniformBuffer* ViewUniformBuffer = GetReferenceViewUniformBuffer(Views);
+		const FGlobalDistanceFieldParameterData* GlobalDistanceFieldParameterData = GetReferenceGlobalDistanceFieldData(Views);
+
 		AddPass(
 			GraphBuilder,
 			RDG_EVENT_NAME("FFXSystem::PreRender"),
@@ -476,13 +499,14 @@ void FFXSystem::PreRender(FRDGBuilder& GraphBuilder, FRHIUniformBuffer* ViewUnif
     }
 }
 
-void FFXSystem::PostRenderOpaque(
-	FRDGBuilder& GraphBuilder,
-	FRHIUniformBuffer* ViewUniformBuffer,
-	bool bAllowGPUParticleUpdate)
+void FFXSystem::PostRenderOpaque(FRDGBuilder& GraphBuilder, TConstArrayView<FViewInfo> Views, bool bAllowGPUParticleUpdate)
 {
+	bAllowGPUParticleUpdate = bAllowGPUParticleUpdate && GetReferenceAllowGPUUpdate(Views);
+
 	if (RHISupportsGPUParticles() && IsParticleCollisionModeSupported(GetShaderPlatform(), PCM_DepthBuffer) && bAllowGPUParticleUpdate)
 	{
+		FRHIUniformBuffer* ViewUniformBuffer = GetReferenceViewUniformBuffer(Views);
+
 		AddPass(GraphBuilder, RDG_EVENT_NAME("FFXSystem::PostRenderOpaque"), 
 			[this, ViewUniformBuffer](FRHICommandListImmediate& RHICmdList)
 			{
diff --git a/Engine/Source/Runtime/Engine/Private/Particles/FXSystemPrivate.h b/Engine/Source/Runtime/Engine/Private/Particles/FXSystemPrivate.h
--- a/Engine/Source/Runtime/Engine/Private/Particles/FXSystemPrivate.h
+++ b/Engine/Source/Runtime/Engine/Private/Particles/FXSystemPrivate.h
@@ -151,15 +151,12 @@ public:
 	virtual void UpdateVectorField(UVectorFieldComponent* VectorFieldComponent) override;
 	FParticleEmitterInstance* CreateGPUSpriteEmitterInstance(FGPUSpriteEmitterInfo& EmitterInfo);
 	virtual void PreInitViews(class FRDGBuilder& GraphBuilder, bool bAllowGPUParticleUpdate) override;
-	virtual void PostInitViews(class FRDGBuilder& GraphBuilder, FRHIUniformBuffer* ViewUniformBuffer, bool bAllowGPUParticleUpdate) override;
+	virtual void PostInitViews(FRDGBuilder& GraphBuilder, TConstArrayView<FViewInfo> Views, bool bAllowGPUParticleUpdate) override;
 	virtual bool UsesGlobalDistanceField() const override;
 	virtual bool UsesDepthBuffer() const override;
 	virtual bool RequiresEarlyViewUniformBuffer() const override;
-	virtual void PreRender(class FRDGBuilder& GraphBuilder, FRHIUniformBuffer* ViewUniformBuffer, const FGlobalDistanceFieldParameterData* GlobalDistanceFieldParameterData, bool bAllowGPUParticleUpdate) override;
-	virtual void PostRenderOpaque(
-		class FRDGBuilder& GraphBuilder,
-		FRHIUniformBuffer* ViewUniformBuffer,
-		bool bAllowGPUParticleUpdate) override;
+	virtual void PreRender(FRDGBuilder& GraphBuilder, TConstArrayView<FViewInfo> Views, bool bAllowGPUParticleUpdate) override;
+	virtual void PostRenderOpaque(FRDGBuilder& GraphBuilder, TConstArrayView<FViewInfo> Views, bool bAllowGPUParticleUpdate) override;
 	// End FFXSystemInterface.
 
 	/*--------------------------------------------------------------------------
diff --git a/Engine/Source/Runtime/Engine/Private/Particles/FXSystemSet.cpp b/Engine/Source/Runtime/Engine/Private/Particles/FXSystemSet.cpp
--- a/Engine/Source/Runtime/Engine/Private/Particles/FXSystemSet.cpp
+++ b/Engine/Source/Runtime/Engine/Private/Particles/FXSystemSet.cpp
@@ -133,12 +133,12 @@ void FFXSystemSet::PreInitViews(class FRDGBuilder& GraphBuilder, bool bAllowGPUP
 	}
 }
 
-void FFXSystemSet::PostInitViews(class FRDGBuilder& GraphBuilder, FRHIUniformBuffer* ViewUniformBuffer, bool bAllowGPUParticleUpdate)
+void FFXSystemSet::PostInitViews(class FRDGBuilder& GraphBuilder, TConstArrayView<FViewInfo> Views, bool bAllowGPUParticleUpdate)
 {
 	for (FFXSystemInterface* FXSystem : FXSystems)
 	{
 		check(FXSystem);
-		FXSystem->PostInitViews(GraphBuilder, ViewUniformBuffer, bAllowGPUParticleUpdate);
+		FXSystem->PostInitViews(GraphBuilder, Views, bAllowGPUParticleUpdate);
 	}
 }
 
@@ -181,12 +181,12 @@ bool FFXSystemSet::RequiresEarlyViewUniformBuffer() const
 	return false;
 }
 
-void FFXSystemSet::PreRender(class FRDGBuilder& GraphBuilder, FRHIUniformBuffer* ViewUniformBuffer, const class FGlobalDistanceFieldParameterData* GlobalDistanceFieldParameterData, bool bAllowGPUParticleSceneUpdate)
+void FFXSystemSet::PreRender(class FRDGBuilder& GraphBuilder, TConstArrayView<FViewInfo> Views, bool bAllowGPUParticleSceneUpdate)
 {
 	for (FFXSystemInterface* FXSystem : FXSystems)
 	{
 		check(FXSystem);
-		FXSystem->PreRender(GraphBuilder, ViewUniformBuffer, GlobalDistanceFieldParameterData, bAllowGPUParticleSceneUpdate);
+		FXSystem->PreRender(GraphBuilder, Views, bAllowGPUParticleSceneUpdate);
 	}
 }
 
@@ -199,19 +199,12 @@ void FFXSystemSet::SetSceneTexturesUniformBuffer(FRHIUniformBuffer* InSceneTextu
 	}
 }
 
-void FFXSystemSet::PostRenderOpaque(
-	class FRDGBuilder& GraphBuilder,
-	FRHIUniformBuffer* ViewUniformBuffer,
-	bool bAllowGPUParticleUpdate)
+void FFXSystemSet::PostRenderOpaque(FRDGBuilder& GraphBuilder, TConstArrayView<const class FViewInfo> Views, bool bAllowGPUParticleSceneUpdate)
 {
 	for (FFXSystemInterface* FXSystem : FXSystems)
 	{
 		check(FXSystem);
-		FXSystem->PostRenderOpaque(
-			GraphBuilder,
-			ViewUniformBuffer,
-			bAllowGPUParticleUpdate
-		);
+		FXSystem->PostRenderOpaque(GraphBuilder, Views, bAllowGPUParticleSceneUpdate);
 	}
 }
 
diff --git a/Engine/Source/Runtime/Engine/Private/Particles/FXSystemSet.h b/Engine/Source/Runtime/Engine/Private/Particles/FXSystemSet.h
--- a/Engine/Source/Runtime/Engine/Private/Particles/FXSystemSet.h
+++ b/Engine/Source/Runtime/Engine/Private/Particles/FXSystemSet.h
@@ -39,16 +39,13 @@ public:
 	virtual void AddVectorField(UVectorFieldComponent* VectorFieldComponent) override;
 	virtual void RemoveVectorField(UVectorFieldComponent* VectorFieldComponent) override;
 	virtual void UpdateVectorField(UVectorFieldComponent* VectorFieldComponent) override;
-	virtual void PreInitViews(class FRDGBuilder& GraphBuilder, bool bAllowGPUParticleUpdate) override;
-	virtual void PostInitViews(class FRDGBuilder& GraphBuilder, FRHIUniformBuffer* ViewUniformBuffer, bool bAllowGPUParticleUpdate) override;
+	virtual void PreInitViews(FRDGBuilder& GraphBuilder, bool bAllowGPUParticleUpdate) override;
+	virtual void PostInitViews(FRDGBuilder& GraphBuilder, TConstArrayView<FViewInfo> Views, bool bAllowGPUParticleUpdate) override;
 	virtual bool UsesGlobalDistanceField() const override;
 	virtual bool UsesDepthBuffer() const override;
 	virtual bool RequiresEarlyViewUniformBuffer() const override;
-	virtual void PreRender(class FRDGBuilder& GraphBuilder, FRHIUniformBuffer* ViewUniformBuffer, const class FGlobalDistanceFieldParameterData* GlobalDistanceFieldParameterData, bool bAllowGPUParticleSceneUpdate) override;
-	virtual void PostRenderOpaque(
-		class FRDGBuilder& GraphBuilder,
-		FRHIUniformBuffer* ViewUniformBuffer,
-		bool bAllowGPUParticleUpdate) override;
+	virtual void PreRender(FRDGBuilder& GraphBuilder, TConstArrayView<FViewInfo> Views, bool bAllowGPUParticleSceneUpdate) override;
+	virtual void PostRenderOpaque(FRDGBuilder& GraphBuilder, TConstArrayView<FViewInfo> Views, bool bAllowGPUParticleSceneUpdate) override;
 
 	virtual void OnDestroy() override;
 	virtual void DestroyGPUSimulation() override;
diff --git a/Engine/Source/Runtime/Engine/Public/FXSystem.h b/Engine/Source/Runtime/Engine/Public/FXSystem.h
--- a/Engine/Source/Runtime/Engine/Public/FXSystem.h
+++ b/Engine/Source/Runtime/Engine/Public/FXSystem.h
@@ -232,9 +232,9 @@ public:
 	 * Notification from the renderer that it is about to perform visibility
 	 * checks on FX belonging to this system.
 	 */
-	virtual void PreInitViews(class FRDGBuilder& GraphBuilder, bool bAllowGPUParticleUpdate) = 0;
+	virtual void PreInitViews(FRDGBuilder& GraphBuilder, bool bAllowGPUParticleUpdate) = 0;
 
-	virtual void PostInitViews(class FRDGBuilder& GraphBuilder, FRHIUniformBuffer* ViewUniformBuffer, bool bAllowGPUParticleUpdate) = 0;
+	virtual void PostInitViews(FRDGBuilder& GraphBuilder, TConstArrayView<FViewInfo> Views, bool bAllowGPUParticleUpdate) = 0;
 
 	virtual bool UsesGlobalDistanceField() const = 0;
 
@@ -246,15 +246,12 @@ public:
 	 * Notification from the renderer that it is about to draw FX belonging to
 	 * this system.
 	 */
-	virtual void PreRender(class FRDGBuilder& GraphBuilder, FRHIUniformBuffer* ViewUniformBuffer, const class FGlobalDistanceFieldParameterData* GlobalDistanceFieldParameterData, bool bAllowGPUParticleSceneUpdate) = 0;
+	virtual void PreRender(FRDGBuilder& GraphBuilder, TConstArrayView<FViewInfo> Views, bool bAllowGPUParticleUpdate) = 0;
 
 	/**
 	 * Notification from the renderer that opaque primitives have rendered.
 	 */
-	virtual void PostRenderOpaque(
-		class FRDGBuilder& GraphBuilder,
-		FRHIUniformBuffer* ViewUniformBuffer,
-		bool bAllowGPUParticleUpdate) = 0;
+	virtual void PostRenderOpaque(FRDGBuilder& GraphBuilder, TConstArrayView<FViewInfo> Views, bool bAllowGPUParticleUpdate) = 0;
 
 	bool IsPendingKill() const { return bIsPendingKill; }
 
@@ -270,6 +267,10 @@ protected:
 	/** By making the destructor protected, an instance must be destroyed via FFXSystemInterface::Destroy. */
 	ENGINE_API virtual ~FFXSystemInterface() {}
 
+	ENGINE_API static FRHIUniformBuffer* GetReferenceViewUniformBuffer(TConstArrayView<FViewInfo> Views);
+	ENGINE_API static bool GetReferenceAllowGPUUpdate(TConstArrayView<FViewInfo> Views);
+	ENGINE_API static const FGlobalDistanceFieldParameterData* GetReferenceGlobalDistanceFieldData(TConstArrayView<FViewInfo> Views);
+
 private:
 
 	bool bIsPendingKill = false;
diff --git a/Engine/Source/Runtime/Renderer/Private/DeferredShadingRenderer.cpp b/Engine/Source/Runtime/Renderer/Private/DeferredShadingRenderer.cpp
--- a/Engine/Source/Runtime/Renderer/Private/DeferredShadingRenderer.cpp
+++ b/Engine/Source/Runtime/Renderer/Private/DeferredShadingRenderer.cpp
@@ -441,11 +441,7 @@ static void RenderOpaqueFX(
 			FXSystem->SetSceneTexturesUniformBuffer(ExtractUBPassParameters->SceneTextures->GetRHIRef());
 		});
 
-		FXSystem->PostRenderOpaque(
-			GraphBuilder,
-			Views[0].ViewUniformBuffer,
-			Views[0].AllowGPUParticleUpdate()
-		);
+		FXSystem->PostRenderOpaque(GraphBuilder, Views, true /*bAllowGPUParticleUpdate*/);
 
 		// Clear the scene textures UB pointer on the FX system.
 		AddPass(GraphBuilder, [FXSystem](FRHICommandList&)
@@ -1924,7 +1920,7 @@ void FDeferredShadingSceneRenderer::Render(FRDGBuilder& GraphBuilder)
 	{
 		SCOPE_CYCLE_COUNTER(STAT_FDeferredShadingSceneRenderer_FXSystem_PreRender);
 		GraphBuilder.SetCommandListStat(GET_STATID(STAT_CLM_FXPreRender));
-		FXSystem->PreRender(GraphBuilder, Views[0].ViewUniformBuffer, &Views[0].GlobalDistanceFieldInfo.ParameterData, Views[0].AllowGPUParticleUpdate());
+		FXSystem->PreRender(GraphBuilder, Views, true /*bAllowGPUParticleUpdate*/);
 		if (FGPUSortManager* GPUSortManager = FXSystem->GetGPUSortManager())
 		{
 			GPUSortManager->OnPreRender(GraphBuilder);
diff --git a/Engine/Source/Runtime/Renderer/Private/MobileShadingRenderer.cpp b/Engine/Source/Runtime/Renderer/Private/MobileShadingRenderer.cpp
--- a/Engine/Source/Runtime/Renderer/Private/MobileShadingRenderer.cpp
+++ b/Engine/Source/Runtime/Renderer/Private/MobileShadingRenderer.cpp
@@ -732,7 +732,7 @@ void FMobileSceneRenderer::Render(FRDGBuilder& GraphBuilder)
 	// Notify the FX system that the scene is about to be rendered.
 	if (FXSystem && ViewFamily.EngineShowFlags.Particles)
 	{
-		FXSystem->PreRender(GraphBuilder, Views[0].ViewUniformBuffer, NULL, !Views[0].bIsPlanarReflection);
+		FXSystem->PreRender(GraphBuilder, Views, true /*bAllowGPUParticleUpdate*/);
 		if (FGPUSortManager* GPUSortManager = FXSystem->GetGPUSortManager())
 		{
 			GPUSortManager->OnPreRender(GraphBuilder);
@@ -817,11 +817,7 @@ void FMobileSceneRenderer::Render(FRDGBuilder& GraphBuilder)
 
 	if (FXSystem && Views.IsValidIndex(0))
 	{
-		FXSystem->PostRenderOpaque(
-			GraphBuilder,
-			Views[0].ViewUniformBuffer,
-			Views[0].AllowGPUParticleUpdate()
-		);
+		FXSystem->PostRenderOpaque(GraphBuilder, Views, true /*bAllowGPUParticleUpdate*/);
 		if (FGPUSortManager* GPUSortManager = FXSystem->GetGPUSortManager())
 		{
 			GPUSortManager->OnPostRenderOpaque(GraphBuilder);
diff --git a/Engine/Source/Runtime/Renderer/Private/SceneHitProxyRendering.cpp b/Engine/Source/Runtime/Renderer/Private/SceneHitProxyRendering.cpp
--- a/Engine/Source/Runtime/Renderer/Private/SceneHitProxyRendering.cpp
+++ b/Engine/Source/Runtime/Renderer/Private/SceneHitProxyRendering.cpp
@@ -651,14 +651,14 @@ void FDeferredShadingSceneRenderer::RenderHitProxies(FRDGBuilder& GraphBuilder)
 	if (FXSystem && Views.IsValidIndex(0))
 	{
 		FGPUSortManager* GPUSortManager = FXSystem->GetGPUSortManager();
-		FXSystem->PreRender(GraphBuilder, Views[0].ViewUniformBuffer, &Views[0].GlobalDistanceFieldInfo.ParameterData, false);
+		FXSystem->PreRender(GraphBuilder, Views, false);
 		if (GPUSortManager)
 		{
 			GPUSortManager->OnPreRender(GraphBuilder);
 		}
 		// Call PostRenderOpaque now as this is irrelevant for when rendering hit proxies.
 		// because we don't tick the particles in the render loop (see last param being "false").
-		FXSystem->PostRenderOpaque(GraphBuilder, Views[0].ViewUniformBuffer, false);
+		FXSystem->PostRenderOpaque(GraphBuilder, Views, false /*bAllowGPUParticleUpdate*/);
 		if (GPUSortManager)
 		{
 			GPUSortManager->OnPostRenderOpaque(GraphBuilder);
diff --git a/Engine/Source/Runtime/Renderer/Private/SceneVisibility.cpp b/Engine/Source/Runtime/Renderer/Private/SceneVisibility.cpp
--- a/Engine/Source/Runtime/Renderer/Private/SceneVisibility.cpp
+++ b/Engine/Source/Runtime/Renderer/Private/SceneVisibility.cpp
@@ -4536,7 +4536,7 @@ bool FDeferredShadingSceneRenderer::InitViews(FRDGBuilder& GraphBuilder, const F
 		if (FXSystem && FXSystem->RequiresEarlyViewUniformBuffer() && Views.IsValidIndex(0))
 		{
 			Views[0].InitRHIResources();
-			FXSystem->PostInitViews(GraphBuilder, Views[0].ViewUniformBuffer, Views[0].AllowGPUParticleUpdate() && !ViewFamily.EngineShowFlags.HitProxies);
+			FXSystem->PostInitViews(GraphBuilder, Views, !ViewFamily.EngineShowFlags.HitProxies);
 		}
 	}
 
