Hack bandaid fix for crash when deleting a Niagara module in the stack

diff --git a/Engine/Plugins/FX/Niagara/Source/NiagaraEditorWidgets/Private/SNiagaraOverviewStack.cpp b/Engine/Plugins/FX/Niagara/Source/NiagaraEditorWidgets/Private/SNiagaraOverviewStack.cpp
--- a/Engine/Plugins/FX/Niagara/Source/NiagaraEditorWidgets/Private/SNiagaraOverviewStack.cpp
+++ b/Engine/Plugins/FX/Niagara/Source/NiagaraEditorWidgets/Private/SNiagaraOverviewStack.cpp
@@ -207,6 +207,11 @@ private:
 		{
 			return FStyleColors::Select;
 		}
+		// robk - hack bandaid fix
+		if (StackEntry->IsFinalized())
+		{
+			return DisabledBackgroundColor;
+		}
 		TSharedRef<FNiagaraSystemViewModel> SystemViewModel = StackEntry->GetSystemViewModel();
 		if (SystemViewModel->GetSystem().GetIsolateEnabled())
 		{
