Integration of CL 15780864

diff --git a/Engine/Plugins/FX/Niagara/Source/NiagaraEditor/Private/NiagaraEditorStyle.cpp b/Engine/Plugins/FX/Niagara/Source/NiagaraEditor/Private/NiagaraEditorStyle.cpp
--- a/Engine/Plugins/FX/Niagara/Source/NiagaraEditor/Private/NiagaraEditorStyle.cpp
+++ b/Engine/Plugins/FX/Niagara/Source/NiagaraEditor/Private/NiagaraEditorStyle.cpp
@@ -291,8 +291,8 @@ TSharedRef< FSlateStyleSet > FNiagaraEditorStyle::Create()
 
 	Style->Set("NiagaraEditor.Module.Pin.TypeSelector.Button", FButtonStyle()
 		.SetNormal(FSlateNoResource())
-		.SetPressed(BOX_BRUSH("Common/Button_Pressed", 8.0f / 32.0f, SelectionColor_Pressed))
-		.SetHovered(BOX_BRUSH("Common/Button_Hovered", 8.0f / 32.0f, SelectionColor))
+		.SetPressed(BOX_CORE_BRUSH("Common/Button_Pressed", 8.0f / 32.0f, SelectionColor_Pressed))
+		.SetHovered(BOX_CORE_BRUSH("Common/Button_Hovered", 8.0f / 32.0f, SelectionColor))
 		.SetNormalPadding(FMargin(0, 0, 0, 0))
 		.SetPressedPadding(FMargin(0, 0, 0, 0)));
 	{
