Updates the NiagaraDataInterfaceCollisionQuery to support HW raytraces.  Includes migrating the hlsl into an external shader file rather than having it inline.  InstnaceBatcher is currently responsible for binding the raytracing shader.

diff --git a/Engine/Plugins/FX/Niagara/Shaders/Private/NiagaraDataInterfaceCollisionQuery.ush b/Engine/Plugins/FX/Niagara/Shaders/Private/NiagaraDataInterfaceCollisionQuery.ush
new file mode 100644
--- /dev/null
+++ b/Engine/Plugins/FX/Niagara/Shaders/Private/NiagaraDataInterfaceCollisionQuery.ush
@@ -0,0 +1,127 @@
+// Copyright Epic Games, Inc. All Rights Reserved.
+
+#pragma once
+
+#include "/Engine/Private/RayTracing/RayTracingCommon.ush"
+#include "NiagaraRayTracingCommon.ush"
+
+#if FEATURE_LEVEL >= FEATURE_LEVEL_SM5
+int											MaxRayTraceCount;
+RWStructuredBuffer<FBasicRayData>			RWRayRequests;
+StructuredBuffer<FNiagaraRayTracingPayload>	IntersectionResults;
+#endif
+
+void DICollisionQuery_SceneDepth(float3 In_SamplePos, out float Out_SceneDepth, out float3 Out_CameraPosWorld, out bool Out_IsInsideView, out float3 Out_WorldPos, out float3 Out_WorldNormal)
+{
+	Out_SceneDepth = -1;
+	Out_WorldPos = float3(0.0, 0.0, 0.0);
+	Out_WorldNormal = float3(0.0, 0.0, 1.0);
+	Out_IsInsideView = true;
+	Out_CameraPosWorld.xyz = View.WorldCameraOrigin.xyz;
+
+#if FEATURE_LEVEL >= FEATURE_LEVEL_SM5
+	float4 SamplePosition = float4(In_SamplePos + View.PreViewTranslation, 1);
+	float4 ClipPosition = mul(SamplePosition, View.TranslatedWorldToClip);
+	float2 ScreenPosition = ClipPosition.xy / ClipPosition.w;
+	// Check if the sample is inside the view.
+	if (all(abs(ScreenPosition.xy) <= float2(1, 1)))
+	{
+		// Sample the depth buffer to get a world position near the sample position.
+		float2 ScreenUV = ScreenPosition * View.ScreenPositionScaleBias.xy + View.ScreenPositionScaleBias.wz;
+		float SceneDepth = CalcSceneDepth(ScreenUV);
+		Out_SceneDepth = SceneDepth;
+		// Reconstruct world position.
+		Out_WorldPos = WorldPositionFromSceneDepth(ScreenPosition.xy, SceneDepth);
+		// Sample the normal buffer
+		Out_WorldNormal = Texture2DSampleLevel(SceneTexturesStruct.GBufferATexture, SceneTexturesStruct_GBufferATextureSampler, ScreenUV, 0).xyz * 2.0 - 1.0;
+	}
+	else
+	{
+		Out_IsInsideView = false;
+	}
+#endif
+}
+
+void DICollisionQuery_CustomDepth(float3 In_SamplePos, out float Out_SceneDepth, out float3 Out_CameraPosWorld, out bool Out_IsInsideView, out float3 Out_WorldPos, out float3 Out_WorldNormal)
+{
+	Out_SceneDepth = -1;
+	Out_WorldPos = float3(0.0, 0.0, 0.0);
+	Out_WorldNormal = float3(0.0, 0.0, 1.0);
+	Out_IsInsideView = true;
+	Out_CameraPosWorld.xyz = View.WorldCameraOrigin.xyz;
+
+#if FEATURE_LEVEL >= FEATURE_LEVEL_SM5
+	float4 SamplePosition = float4(In_SamplePos + View.PreViewTranslation, 1);
+	float4 ClipPosition = mul(SamplePosition, View.TranslatedWorldToClip);
+	float2 ScreenPosition = ClipPosition.xy / ClipPosition.w;
+	// Check if the sample is inside the view.
+	if (all(abs(ScreenPosition.xy) <= float2(1, 1)))
+	{
+		// Sample the depth buffer to get a world position near the sample position.
+		float2 ScreenUV = ScreenPosition * View.ScreenPositionScaleBias.xy + View.ScreenPositionScaleBias.wz;
+		float SceneDepth = ConvertFromDeviceZ(Texture2DSampleLevel(SceneTexturesStruct.CustomDepthTexture, SceneTexturesStruct_SceneDepthTextureSampler, ScreenUV, 0).r);
+		Out_SceneDepth = SceneDepth;
+		// Reconstruct world position.
+		Out_WorldPos = WorldPositionFromSceneDepth(ScreenPosition.xy, SceneDepth);
+		// Sample the normal buffer
+		Out_WorldNormal = Texture2DSampleLevel(SceneTexturesStruct.GBufferATexture, SceneTexturesStruct_GBufferATextureSampler, ScreenUV, 0).xyz * 2.0 - 1.0;
+	}
+	else
+	{
+		Out_IsInsideView = false;
+	}
+#endif
+}
+
+void DICollisionQuery_DistanceField(float3 In_SamplePos, out float Out_DistanceToNearestSurface, out float3 Out_FieldGradient, out bool Out_IsDistanceFieldValid)
+{
+#if PLATFORM_SUPPORTS_DISTANCE_FIELDS && (FEATURE_LEVEL >= FEATURE_LEVEL_SM5)
+	Out_DistanceToNearestSurface = GetDistanceToNearestSurfaceGlobal(In_SamplePos);
+	Out_FieldGradient = GetDistanceFieldGradientGlobal(In_SamplePos);
+	Out_IsDistanceFieldValid = MaxGlobalDFAOConeDistance > 0 && !(Out_DistanceToNearestSurface > 0 && all(Out_FieldGradient == float3(0,0,0)));
+#else
+	Out_DistanceToNearestSurface = 0;
+	Out_FieldGradient = (float3)0;
+	Out_IsDistanceFieldValid = false;
+#endif
+}
+
+void DICollisionQuery_IssueAsyncRayTrace(int In_QueryID, float3 In_TraceStart, float3 In_TraceEnd, int In_TraceChannel, out bool Out_IsQueryValid)
+{
+	Out_IsQueryValid = false;
+
+#if FEATURE_LEVEL >= FEATURE_LEVEL_SM5
+	if (In_QueryID >= 0 && In_QueryID < MaxRayTraceCount)
+	{
+		float3 Trace = In_TraceEnd - In_TraceStart;
+		float TraceLength = length(Trace);
+
+		RWRayRequests[In_QueryID].Origin = In_TraceStart;
+		RWRayRequests[In_QueryID].Mask = 0xFF; // In_TraceChannel - for now we'll ignore the TraceChannel till we can mark up the environment
+		RWRayRequests[In_QueryID].Direction = normalize(Trace);
+		RWRayRequests[In_QueryID].TFar = TraceLength;
+		Out_IsQueryValid = true;
+	}
+#endif
+}
+
+void DICollisionQuery_ReadAsyncRayTrace(int In_PreviousFrameQueryID, out bool Out_CollisionValid, out float Out_CollisionDistance, out float3 Out_CollisionPosWorld, out float3 Out_CollisionNormal)
+{
+	Out_CollisionValid = false;
+	Out_CollisionDistance = 0.0f;
+	Out_CollisionPosWorld = float3(0.0f, 0.0f, 0.0f);
+	Out_CollisionNormal = float3(0.0f, 0.0f, 0.0f);
+
+#if FEATURE_LEVEL >= FEATURE_LEVEL_SM5
+	if (In_PreviousFrameQueryID >= 0 && In_PreviousFrameQueryID < MaxRayTraceCount)
+	{
+		Out_CollisionValid = IntersectionResults[In_PreviousFrameQueryID].IsHit();
+		if (Out_CollisionValid)
+		{
+			Out_CollisionDistance = IntersectionResults[In_PreviousFrameQueryID].HitT;
+			Out_CollisionPosWorld = IntersectionResults[In_PreviousFrameQueryID].WorldPosition;
+			Out_CollisionNormal = IntersectionResults[In_PreviousFrameQueryID].WorldNormal;
+		}
+	}
+#endif
+}
\ No newline at end of file
diff --git a/Engine/Plugins/FX/Niagara/Shaders/Private/NiagaraRayTracingCommon.ush b/Engine/Plugins/FX/Niagara/Shaders/Private/NiagaraRayTracingCommon.ush
new file mode 100644
--- /dev/null
+++ b/Engine/Plugins/FX/Niagara/Shaders/Private/NiagaraRayTracingCommon.ush
@@ -0,0 +1,12 @@
+// Copyright Epic Games, Inc. All Rights Reserved.
+
+#pragma once
+
+#include "/Engine/Private/RayTracing/RayTracingCommon.ush"
+
+// payload structure for Niagara ray traces
+struct FNiagaraRayTracingPayload : FMinimalPayload
+{
+	float3 WorldPosition;
+	float3 WorldNormal;
+};
diff --git a/Engine/Plugins/FX/Niagara/Shaders/Private/NiagaraRayTracingShaders.usf b/Engine/Plugins/FX/Niagara/Shaders/Private/NiagaraRayTracingShaders.usf
new file mode 100644
--- /dev/null
+++ b/Engine/Plugins/FX/Niagara/Shaders/Private/NiagaraRayTracingShaders.usf
@@ -0,0 +1,81 @@
+// Copyright Epic Games, Inc. All Rights Reserved.
+
+#include "/Engine/Private/Common.ush"
+#include "/Engine/Private/ShadingCommon.ush"
+#include "/Engine/Private/RayTracing/RayTracingCommon.ush"
+#include "/Engine/Private/RayTracing/RayTracingHitGroupCommon.ush"
+#include "NiagaraRayTracingCommon.ush"
+
+#define RT_CONCATENATE2(a, b) a##b
+#define RT_CONCATENATE(a, b) RT_CONCATENATE2(a, b)
+#define RT_REGISTER(InType, InIndex, InSpace) register(RT_CONCATENATE(InType, InIndex), RT_CONCATENATE(space, InSpace))
+
+RaytracingAccelerationStructure               TLAS               : RT_REGISTER(t, 0, RAY_TRACING_REGISTER_SPACE_GLOBAL);
+StructuredBuffer<FBasicRayData>               Rays               : RT_REGISTER(t, 1, RAY_TRACING_REGISTER_SPACE_GLOBAL);
+RWStructuredBuffer<FNiagaraRayTracingPayload> CollisionOutput    : RT_REGISTER(u, 0, RAY_TRACING_REGISTER_SPACE_GLOBAL);
+
+#undef RT_REGISTER
+#undef RT_CONCATENATE
+#undef RT_CONCATENATE2
+
+RAY_TRACING_ENTRY_RAYGEN(NiagaraCollisionRayTraceRG)
+{
+	const uint RayIndex = DispatchRaysIndex().x;
+	FBasicRayData InputRay = Rays[RayIndex];
+
+	RayDesc Ray;
+	Ray.Origin = InputRay.Origin;
+	Ray.Direction = InputRay.Direction;
+	Ray.TMin = 0.0f;
+	Ray.TMax = InputRay.TFar;
+
+	uint RayFlags = RAY_FLAG_FORCE_OPAQUE; // don't run anyhit shader
+
+	FNiagaraRayTracingPayload Payload = (FNiagaraRayTracingPayload)0;
+
+	TraceRay(
+		TLAS,   // AccelerationStructure
+		RayFlags,
+		RAY_TRACING_MASK_OPAQUE,          // InstanceInclusionMask
+		RAY_TRACING_SHADER_SLOT_MATERIAL, // RayContributionToHitGroupIndex
+		RAY_TRACING_NUM_SHADER_SLOTS,     // MultiplierForGeometryContributionToShaderIndex
+		0,      // MissShaderIndex
+		Ray,    // RayDesc
+		Payload // Payload
+	);
+
+	CollisionOutput[RayIndex].HitT = Payload.HitT;
+	if (CollisionOutput[RayIndex].IsHit())
+	{
+		CollisionOutput[RayIndex].WorldPosition = Payload.WorldPosition;
+		CollisionOutput[RayIndex].WorldNormal = Payload.WorldNormal;
+	}
+}
+
+float3 GetBarycentricPosition(FDefaultAttributes Attributes, FTriangleBaseAttributes Tri)
+{
+	return (1.0f - Attributes.Barycentrics.x - Attributes.Barycentrics.y) * Tri.LocalPositions[0]
+		+ Attributes.Barycentrics.x * Tri.LocalPositions[1]
+		+ Attributes.Barycentrics.y * Tri.LocalPositions[2];
+}
+
+RAY_TRACING_ENTRY_CLOSEST_HIT(NiagaraCollisionRayTraceCH,
+	FNiagaraRayTracingPayload, Payload,
+	FDefaultAttributes, Attributes)
+{
+	Payload.HitT = RayTCurrent();
+
+	FTriangleBaseAttributes Tri = LoadTriangleBaseAttributes(PrimitiveIndex());
+	float3 LocalPosition = GetBarycentricPosition(Attributes, Tri);
+	float3 LocalEdges[2] = { Tri.LocalPositions[1] - Tri.LocalPositions[0], Tri.LocalPositions[2] - Tri.LocalPositions[0] };
+	float3 LocalNormal = cross(LocalEdges[1], LocalEdges[0]);
+	float3x3 InverseTranspose3x3 = transpose((float3x3)WorldToObject4x3());
+
+	Payload.WorldPosition = mul(float4(LocalPosition, 1.0f), ObjectToWorld4x3()).xyz;
+	Payload.WorldNormal = normalize(mul(LocalNormal, InverseTranspose3x3));
+}
+
+RAY_TRACING_ENTRY_MISS(NiagaraCollisionRayTraceMiss, FNiagaraRayTracingPayload, Payload)
+{
+	Payload.SetMiss();
+}
diff --git a/Engine/Plugins/FX/Niagara/Source/Niagara/Classes/NiagaraDataInterface.h b/Engine/Plugins/FX/Niagara/Source/Niagara/Classes/NiagaraDataInterface.h
--- a/Engine/Plugins/FX/Niagara/Source/Niagara/Classes/NiagaraDataInterface.h
+++ b/Engine/Plugins/FX/Niagara/Source/Niagara/Classes/NiagaraDataInterface.h
@@ -340,6 +340,7 @@ public:
 	virtual bool RequiresDistanceFieldData() const { return false; }
 	virtual bool RequiresDepthBuffer() const { return false; }
 	virtual bool RequiresEarlyViewData() const { return false; }
+	virtual bool RequiresRayTracingScene() const { return false; }
 
 	virtual bool HasTickGroupPrereqs() const { return false; }
 	virtual ETickingGroup CalculateTickGroup(const void* PerInstanceData) const { return NiagaraFirstTickGroup; }
diff --git a/Engine/Plugins/FX/Niagara/Source/Niagara/Classes/NiagaraDataInterfaceCollisionQuery.h b/Engine/Plugins/FX/Niagara/Source/Niagara/Classes/NiagaraDataInterfaceCollisionQuery.h
--- a/Engine/Plugins/FX/Niagara/Source/Niagara/Classes/NiagaraDataInterfaceCollisionQuery.h
+++ b/Engine/Plugins/FX/Niagara/Source/Niagara/Classes/NiagaraDataInterfaceCollisionQuery.h
@@ -29,10 +29,16 @@ public:
 
 	DECLARE_NIAGARA_DI_PARAMETER();
 
+	UPROPERTY(EditAnywhere, Category = "Ray Trace")
+	int32 MaxRayTraceCount = 0;
+
 	FNiagaraSystemInstance *SystemInstance;
 
 	//UObject Interface
 	virtual void PostInitProperties() override;
+#if WITH_EDITOR
+	virtual void PostEditChangeProperty(struct FPropertyChangedEvent& PropertyChangedEvent) override;
+#endif
 	//UObject Interface End
 
 	/** Initializes the per instance data for this interface. Returns false if there was some error and the simulation should be disabled. */
@@ -45,11 +51,12 @@ public:
 	virtual bool PerInstanceTickPostSimulate(void* PerInstanceData, FNiagaraSystemInstance* SystemInstance, float DeltaSeconds) override;
 	virtual int32 PerInstanceDataSize() const override { return sizeof(CQDIPerInstanceData); }
 	
+	virtual void GetCommonHLSL(FString& OutHlsl) override;
 	virtual void GetFunctions(TArray<FNiagaraFunctionSignature>& OutFunctions) override;
 	virtual void GetVMExternalFunction(const FVMExternalFunctionBindingInfo& BindingInfo, void* InstanceData, FVMExternalFunction &OutFunc) override;
 	virtual bool AppendCompileHash(FNiagaraCompileHashVisitor* InVisitor) const override;
 	virtual void GetAssetTagsForContext(const UObject* InAsset, const TArray<const UNiagaraDataInterface*>& InProperties, TMap<FName, uint32>& NumericKeys, TMap<FName, FString>& StringKeys) const override;
-
+	virtual void PushToRenderThreadImpl() override;
 
 	// VM functions
 	void PerformQuerySyncCPU(FVectorVMContext& Context);
@@ -60,6 +67,7 @@ public:
 	virtual bool CanExecuteOnTarget(ENiagaraSimTarget Target) const override { return true; }
 	virtual bool RequiresDistanceFieldData() const override { return true; }
 	virtual bool RequiresDepthBuffer() const override { return true; }
+	virtual bool RequiresRayTracingScene() const override;
 
 	virtual void GetParameterDefinitionHLSL(const FNiagaraDataInterfaceGPUParamInfo& ParamInfo, FString& OutHLSL) override;
 	virtual bool GetFunctionHLSL(const FNiagaraDataInterfaceGPUParamInfo& ParamInfo, const FNiagaraDataInterfaceGeneratedFunction& FunctionInfo, int FunctionInstanceIndex, FString& OutHLSL) override;
@@ -73,6 +81,12 @@ public:
 	
 	virtual bool HasPreSimulateTick() const override{ return true; }
 	virtual bool HasPostSimulateTick() const override { return true; }
+
+	virtual bool Equals(const UNiagaraDataInterface* Other) const override;
+
+protected:
+	virtual bool CopyToInternal(UNiagaraDataInterface* Destination) const override;
+
 private:
 
 	static FCriticalSection CriticalSection;
@@ -81,16 +95,8 @@ private:
 	const static FName SceneDepthName;
 	const static FName CustomDepthName;
 	const static FName DistanceFieldName;
+	const static FName IssueAsyncRayTraceName;
+	const static FName ReadAsyncRayTraceName;
 	const static FName SyncTraceName;
 	const static FName AsyncTraceName;
-};
-
-struct FNiagaraDataIntefaceProxyCollisionQuery : public FNiagaraDataInterfaceProxy
-{
-	// There's nothing in this proxy. It just reads from scene textures.
-
-	virtual int32 PerInstanceDataPassedToRenderThreadSize() const override
-	{
-		return 0;
-	}
-};
+};
\ No newline at end of file
diff --git a/Engine/Plugins/FX/Niagara/Source/Niagara/Classes/NiagaraEmitterInstanceBatcher.h b/Engine/Plugins/FX/Niagara/Source/Niagara/Classes/NiagaraEmitterInstanceBatcher.h
--- a/Engine/Plugins/FX/Niagara/Source/Niagara/Classes/NiagaraEmitterInstanceBatcher.h
+++ b/Engine/Plugins/FX/Niagara/Source/Niagara/Classes/NiagaraEmitterInstanceBatcher.h
@@ -27,6 +27,7 @@ the same VectorVM byte code / compute shader code
 class FGPUSortManager;
 class FNiagaraGpuComputeDebug;
 class FNiagaraGpuReadbackManager;
+class FNiagaraRayTracingHelper;
 
 enum class ETickStage
 {
@@ -78,6 +79,7 @@ public:
 	virtual bool UsesGlobalDistanceField() const override;
 	virtual bool UsesDepthBuffer() const override;
 	virtual bool RequiresEarlyViewUniformBuffer() const override;
+	virtual bool RequiresRayTracingScene() const override;
 	virtual void PreRender(FRDGBuilder& GraphBuilder, TConstArrayView<FViewInfo> Views, bool bAllowGPUParticleUpdate) override;
 	virtual void OnDestroy() override; // Called on the gamethread to delete the batcher on the renderthread.
 
@@ -165,6 +167,11 @@ public:
 #endif
 	FNiagaraGpuReadbackManager* GetGpuReadbackManager() const { return GpuReadbackManagerPtr.Get(); }
 
+#if RHI_RAYTRACING
+	void IssueRayTraces(FRHICommandList& RHICmdList, const FIntPoint& RayTraceCounts, FRHIShaderResourceView* RayTraceRequests, FRHIUnorderedAccessView* RayTraceResults) const;
+	bool HasRayTracingScene() const;
+#endif
+
 private:
 	using FEmitterInstanceList = TArray<FNiagaraComputeInstanceData*>;
 
@@ -217,6 +224,11 @@ private:
 	 */
 	void GenerateSortKeys(FRHICommandListImmediate& RHICmdList, int32 BatchId, int32 NumElementsInBatch, EGPUSortFlags Flags, FRHIUnorderedAccessView* KeysUAV, FRHIUnorderedAccessView* ValuesUAV);
 
+#if RHI_RAYTRACING
+	void BuildRayTracingSceneInfo(FRHICommandList& RHICmdList, TConstArrayView<FViewInfo> Views);
+	void ResetRayTracingSceneInfo();
+#endif
+
 	inline uint32 UnpackEmitterDispatchCount(uint8* PackedData)
 	{
 		return *(uint32*)PackedData;
@@ -247,6 +259,11 @@ private:
 	/** All sort tasks registered in AddSortedGPUSimulation(). Holds all the data required in GenerateSortKeys(). */
 	TArray<FNiagaraGPUSortInfo> SimulationsToSort;
 
+#if RHI_RAYTRACING
+	// helper object which is valid during PostRenderOpaque if DataInterfaces are present that require a RayTrace scene
+	TUniquePtr<FNiagaraRayTracingHelper> RayTracingHelper;
+#endif
+
 	// GPU emitter instance count buffer. Contains the actual particle / instance count generate in the GPU tick.
 	FNiagaraGPUInstanceCountManager GPUInstanceCounterManager;
 
@@ -299,6 +316,7 @@ private:
 	uint32 NumTicksThatRequireDistanceFieldData = 0;
 	uint32 NumTicksThatRequireDepthBuffer = 0;
 	uint32 NumTicksThatRequireEarlyViewData = 0;
+	uint32 NumTicksThatRequireRayTracingScene = 0;
 
 	int32 TotalDispatchesThisFrame = 0;
 
diff --git a/Engine/Plugins/FX/Niagara/Source/Niagara/Classes/NiagaraScriptExecutionContext.h b/Engine/Plugins/FX/Niagara/Source/Niagara/Classes/NiagaraScriptExecutionContext.h
--- a/Engine/Plugins/FX/Niagara/Source/Niagara/Classes/NiagaraScriptExecutionContext.h
+++ b/Engine/Plugins/FX/Niagara/Source/Niagara/Classes/NiagaraScriptExecutionContext.h
@@ -528,6 +528,7 @@ public:
 	bool bRequiresDepthBuffer = false;
 	bool bRequiresEarlyViewData = false;
 	bool bRequiresViewUniformBuffer = false;
+	bool bRequiresRayTracingScene = false;
 	bool bNeedsReset = false;
 	bool bIsFinalTick = false;
 };
diff --git a/Engine/Plugins/FX/Niagara/Source/Niagara/Private/NiagaraDataInterfaceCollisionQuery.cpp b/Engine/Plugins/FX/Niagara/Source/Niagara/Private/NiagaraDataInterfaceCollisionQuery.cpp
--- a/Engine/Plugins/FX/Niagara/Source/Niagara/Private/NiagaraDataInterfaceCollisionQuery.cpp
+++ b/Engine/Plugins/FX/Niagara/Source/Niagara/Private/NiagaraDataInterfaceCollisionQuery.cpp
@@ -1,13 +1,19 @@
 // Copyright Epic Games, Inc. All Rights Reserved.
 
 #include "NiagaraDataInterfaceCollisionQuery.h"
-#include "NiagaraTypes.h"
-#include "NiagaraWorldManager.h"
-#include "ShaderParameterUtils.h"
+
 #include "GlobalDistanceFieldParameters.h"
 #include "NiagaraComponent.h"
 #include "NiagaraEmitterInstanceBatcher.h"
+#include "NiagaraRayTracingHelper.h"
+#include "NiagaraStats.h"
+#include "NiagaraTypes.h"
+#include "NiagaraWorldManager.h"
+#include "RayTracingInstanceUtils.h"
+#include "RenderResource.h"
 #include "Shader.h"
+#include "ShaderCore.h"
+#include "ShaderParameterUtils.h"
 
 FCriticalSection UNiagaraDataInterfaceCollisionQuery::CriticalSection;
 
@@ -25,9 +31,100 @@ struct FNiagaraCollisionDIFunctionVersion
 	};
 };
 
+struct FNiagaraDataIntefaceProxyCollisionQuery : public FNiagaraDataInterfaceProxy
+{
+#if RHI_RAYTRACING
+	FRWBufferStructured RayTraceRequests;
+	FRWBufferStructured RayTraceIntersections;
+
+	int32 MaxRayTraceCount = 0;
+#endif
+
+	virtual ~FNiagaraDataIntefaceProxyCollisionQuery()
+	{
+#if RHI_RAYTRACING
+		DEC_MEMORY_STAT_BY(STAT_NiagaraGPUDataInterfaceMemory, RayTraceRequests.NumBytes);
+		RayTraceRequests.Release();
+
+		DEC_MEMORY_STAT_BY(STAT_NiagaraGPUDataInterfaceMemory, RayTraceIntersections.NumBytes);
+		RayTraceIntersections.Release();
+#endif
+	}
+
+	virtual int32 PerInstanceDataPassedToRenderThreadSize() const override
+	{
+		return 0;
+	}
+
+	virtual void PostSimulate(FRHICommandList& RHICmdList, const FNiagaraDataInterfaceArgs& Context) override
+	{
+#if RHI_RAYTRACING
+		if (MaxRayTraceCount > 0 && Context.Batcher->HasRayTracingScene())
+		{
+			{
+				FRHITransitionInfo PreTransitions[] =
+				{
+					FRHITransitionInfo(RayTraceRequests.UAV, ERHIAccess::UAVMask, ERHIAccess::SRVMask),
+					FRHITransitionInfo(RayTraceIntersections.UAV, ERHIAccess::SRVMask, ERHIAccess::UAVMask)
+				};
+				RHICmdList.Transition(PreTransitions);
+			}
+
+			Context.Batcher->IssueRayTraces(RHICmdList, FIntPoint(MaxRayTraceCount, 1), RayTraceRequests.SRV, RayTraceIntersections.UAV);
+
+			{
+				FRHITransitionInfo PostTransitions[] =
+				{
+					FRHITransitionInfo(RayTraceRequests.UAV, ERHIAccess::SRVMask, ERHIAccess::UAVMask),
+					FRHITransitionInfo(RayTraceIntersections.UAV, ERHIAccess::UAVMask, ERHIAccess::SRVMask)
+				};
+				RHICmdList.Transition(PostTransitions);
+			}
+		}
+#endif
+	}
+
+	void RenderThreadInitialize(int32 InMaxRayTraceRequests)
+	{
+#if RHI_RAYTRACING
+		MaxRayTraceCount = 0;
+
+		DEC_MEMORY_STAT_BY(STAT_NiagaraGPUDataInterfaceMemory, RayTraceRequests.NumBytes);
+		RayTraceRequests.Release();
+
+		DEC_MEMORY_STAT_BY(STAT_NiagaraGPUDataInterfaceMemory, RayTraceIntersections.NumBytes);
+		RayTraceIntersections.Release();
+
+		if (IsRayTracingEnabled() && InMaxRayTraceRequests > 0)
+		{
+			MaxRayTraceCount = 16 * FMath::DivideAndRoundUp(InMaxRayTraceRequests, 16);
+
+			RayTraceRequests.Initialize(
+				TEXT("NiagaraRayTraceRequests"),
+				sizeof(FBasicRayData),
+				MaxRayTraceCount,
+				BUF_Static);
+			INC_MEMORY_STAT_BY(STAT_NiagaraGPUDataInterfaceMemory, RayTraceRequests.NumBytes);
+
+			RayTraceIntersections.Initialize(
+				TEXT("NiagaraRayTraceIntersections"),
+				sizeof(FNiagaraRayTracingPayload),
+				MaxRayTraceCount,
+				BUF_Static,
+				false /*bUseUavCounter*/,
+				false /*bAppendBuffer*/,
+				ERHIAccess::SRVMask);
+			INC_MEMORY_STAT_BY(STAT_NiagaraGPUDataInterfaceMemory, RayTraceIntersections.NumBytes);
+		}
+#endif
+	}
+};
+
 const FName UNiagaraDataInterfaceCollisionQuery::SceneDepthName(TEXT("QuerySceneDepthGPU"));
 const FName UNiagaraDataInterfaceCollisionQuery::CustomDepthName(TEXT("QueryCustomDepthGPU"));
 const FName UNiagaraDataInterfaceCollisionQuery::DistanceFieldName(TEXT("QueryMeshDistanceFieldGPU"));
+const FName UNiagaraDataInterfaceCollisionQuery::IssueAsyncRayTraceName(TEXT("IssueAsyncRayTraceGpu"));
+const FName UNiagaraDataInterfaceCollisionQuery::ReadAsyncRayTraceName(TEXT("ReadAsyncRayTraceGpu"));
 const FName UNiagaraDataInterfaceCollisionQuery::SyncTraceName(TEXT("PerformCollisionQuerySyncCPU"));
 const FName UNiagaraDataInterfaceCollisionQuery::AsyncTraceName(TEXT("PerformCollisionQueryAsyncCPU"));
 
@@ -203,6 +300,41 @@ void UNiagaraDataInterfaceCollisionQuery::GetFunctions(TArray<FNiagaraFunctionSi
 	SigMeshField.Outputs.Add(FNiagaraVariable(FNiagaraTypeDefinition::GetBoolDef(), TEXT("IsDistanceFieldValid")));
     OutFunctions.Add(SigMeshField);
 
+	{
+		FNiagaraFunctionSignature& IssueRayTrace = OutFunctions.AddDefaulted_GetRef();
+		IssueRayTrace.Name = UNiagaraDataInterfaceCollisionQuery::IssueAsyncRayTraceName;
+		IssueRayTrace.bRequiresExecPin = true;
+		IssueRayTrace.bMemberFunction = true;
+		IssueRayTrace.bRequiresContext = false;
+		IssueRayTrace.bSupportsCPU = false;
+#if WITH_EDITORONLY_DATA
+		IssueRayTrace.FunctionVersion = FNiagaraCollisionDIFunctionVersion::LatestVersion;
+#endif
+		IssueRayTrace.Inputs.Add(FNiagaraVariable(FNiagaraTypeDefinition(GetClass()), TEXT("CollisionQuery")));
+		IssueRayTrace.Inputs.Add(FNiagaraVariable(FNiagaraTypeDefinition::GetIntDef(), TEXT("QueryID")));
+		IssueRayTrace.Inputs.Add(FNiagaraVariable(FNiagaraTypeDefinition::GetVec3Def(), TEXT("TraceStartWorld")));
+		IssueRayTrace.Inputs.Add(FNiagaraVariable(FNiagaraTypeDefinition::GetVec3Def(), TEXT("TraceEndWorld")));
+		IssueRayTrace.Inputs.Add(FNiagaraVariable(FNiagaraTypeDefinition::GetIntDef(), TEXT("TraceChannel")));
+		IssueRayTrace.Outputs.Add(FNiagaraVariable(FNiagaraTypeDefinition::GetBoolDef(), TEXT("IsQueryValid")));
+	}
+
+	{
+		FNiagaraFunctionSignature& ReadRayTrace = OutFunctions.AddDefaulted_GetRef();
+		ReadRayTrace.Name = UNiagaraDataInterfaceCollisionQuery::ReadAsyncRayTraceName;
+		ReadRayTrace.bMemberFunction = true;
+		ReadRayTrace.bRequiresContext = false;
+		ReadRayTrace.bSupportsCPU = false;
+#if WITH_EDITORONLY_DATA
+		ReadRayTrace.FunctionVersion = FNiagaraCollisionDIFunctionVersion::LatestVersion;
+#endif
+		ReadRayTrace.Inputs.Add(FNiagaraVariable(FNiagaraTypeDefinition(GetClass()), TEXT("CollisionQuery")));
+		ReadRayTrace.Inputs.Add(FNiagaraVariable(FNiagaraTypeDefinition::GetIntDef(), TEXT("PreviousFrameQueryID")));
+		ReadRayTrace.Outputs.Add(FNiagaraVariable(FNiagaraTypeDefinition::GetBoolDef(), TEXT("CollisionValid")));
+		ReadRayTrace.Outputs.Add(FNiagaraVariable(FNiagaraTypeDefinition::GetFloatDef(), TEXT("CollisionDistance")));
+		ReadRayTrace.Outputs.Add(FNiagaraVariable(FNiagaraTypeDefinition::GetVec3Def(), TEXT("CollisionPosWorld")));
+		ReadRayTrace.Outputs.Add(FNiagaraVariable(FNiagaraTypeDefinition::GetVec3Def(), TEXT("CollisionNormal")));
+	}
+
 	FNiagaraFunctionSignature SigCpuSync;
 	SigCpuSync.Name = UNiagaraDataInterfaceCollisionQuery::SyncTraceName;
 	SigCpuSync.bMemberFunction = true;
@@ -257,68 +389,70 @@ void UNiagaraDataInterfaceCollisionQuery::GetFunctions(TArray<FNiagaraFunctionSi
 // 
 bool UNiagaraDataInterfaceCollisionQuery::GetFunctionHLSL(const FNiagaraDataInterfaceGPUParamInfo& ParamInfo, const FNiagaraDataInterfaceGeneratedFunction& FunctionInfo, int FunctionInstanceIndex, FString& OutHLSL)
 {
+	TMap<FString, FStringFormatArg> Args;
+	Args.Add(TEXT("FunctionName"), FunctionInfo.InstanceName);
 
-	if (FunctionInfo.DefinitionName == SceneDepthName || FunctionInfo.DefinitionName == CustomDepthName)
-	{
-		static const FString SceneDepthSampleExpr = TEXT("CalcSceneDepth(ScreenUV)");
-		static const FString CustomDepthSampleExpr = TEXT("ConvertFromDeviceZ(Texture2DSampleLevel(SceneTexturesStruct.CustomDepthTexture, SceneTexturesStruct_SceneDepthTextureSampler, ScreenUV, 0).r)");
-		const FStringFormatOrderedArguments Args = {
-			FunctionInfo.InstanceName,
-			FunctionInfo.DefinitionName == SceneDepthName ? SceneDepthSampleExpr : CustomDepthSampleExpr
-		};
-
-		OutHLSL += FString::Format(TEXT(R"(
-			void {0}(in float3 In_SamplePos, out float Out_SceneDepth, out float3 Out_CameraPosWorld, out bool Out_IsInsideView, out float3 Out_WorldPos, out float3 Out_WorldNormal)
-			{				
-				Out_SceneDepth = -1;
-				Out_WorldPos = float3(0.0, 0.0, 0.0);
-				Out_WorldNormal = float3(0.0, 0.0, 1.0);
-				Out_IsInsideView = true;
-				Out_CameraPosWorld.xyz = View.WorldCameraOrigin.xyz;
-
-			#if FEATURE_LEVEL >= FEATURE_LEVEL_SM5
-				float4 SamplePosition = float4(In_SamplePos + View.PreViewTranslation, 1);
-				float4 ClipPosition = mul(SamplePosition, View.TranslatedWorldToClip);
-				float2 ScreenPosition = ClipPosition.xy / ClipPosition.w;
-				// Check if the sample is inside the view.
-				if (all(abs(ScreenPosition.xy) <= float2(1, 1)))
-				{
-					// Sample the depth buffer to get a world position near the sample position.
-					float2 ScreenUV = ScreenPosition * View.ScreenPositionScaleBias.xy + View.ScreenPositionScaleBias.wz;
-					float SceneDepth = {1};
-					Out_SceneDepth = SceneDepth;
-					// Reconstruct world position.
-					Out_WorldPos = WorldPositionFromSceneDepth(ScreenPosition.xy, SceneDepth);
-					// Sample the normal buffer
-					Out_WorldNormal = Texture2DSampleLevel(SceneTexturesStruct.GBufferATexture, SceneTexturesStruct_GBufferATextureSampler, ScreenUV, 0).xyz * 2.0 - 1.0;
-				}
-				else
-				{
-					Out_IsInsideView = false;
-				}
-			#endif
+	if (FunctionInfo.DefinitionName == SceneDepthName)
+	{
+		static const TCHAR* FormatSample = TEXT(R"(
+			void {FunctionName}(in float3 In_SamplePos, out float Out_SceneDepth, out float3 Out_CameraPosWorld, out bool Out_IsInsideView, out float3 Out_WorldPos, out float3 Out_WorldNormal)
+			{
+				DICollisionQuery_SceneDepth(In_SamplePos, Out_SceneDepth, Out_CameraPosWorld, Out_IsInsideView, Out_WorldPos, Out_WorldNormal);
 			}
-		)"), Args);
-		return true;
+		)");
+
+		OutHLSL += FString::Format(FormatSample, Args);
+	}
+	else if (FunctionInfo.DefinitionName == CustomDepthName)
+	{
+		static const TCHAR* FormatSample = TEXT(R"(
+			void {FunctionName}(in float3 In_SamplePos, out float Out_SceneDepth, out float3 Out_CameraPosWorld, out bool Out_IsInsideView, out float3 Out_WorldPos, out float3 Out_WorldNormal)
+			{
+				DICollisionQuery_CustomDepth(In_SamplePos, Out_SceneDepth, Out_CameraPosWorld, Out_IsInsideView, Out_WorldPos, Out_WorldNormal);
+			}
+		)");
+
+		OutHLSL += FString::Format(FormatSample, Args);
 	}
 	else if (FunctionInfo.DefinitionName == DistanceFieldName)
 	{
-		OutHLSL += TEXT("void ") + FunctionInfo.InstanceName + TEXT("(in float3 In_SamplePos, out float Out_DistanceToNearestSurface, out float3 Out_FieldGradient, out bool Out_IsDistanceFieldValid) \n{\n");
-		OutHLSL += TEXT("\
-			#if PLATFORM_SUPPORTS_DISTANCE_FIELDS && (FEATURE_LEVEL >= FEATURE_LEVEL_SM5)\n\
-			Out_DistanceToNearestSurface = GetDistanceToNearestSurfaceGlobal(In_SamplePos);\n\
-			Out_FieldGradient = GetDistanceFieldGradientGlobal(In_SamplePos);\n\
-			Out_IsDistanceFieldValid = MaxGlobalDFAOConeDistance > 0 && !(Out_DistanceToNearestSurface > 0 && Out_FieldGradient == float3(0,0,0));\n\
-			#else\n\
-			Out_DistanceToNearestSurface = 0;\n\
-			Out_FieldGradient = (float3)0;\n\
-			Out_IsDistanceFieldValid = false;\n\
-			#endif\n\
-			}\n\n");
-		return true;
+		static const TCHAR* FormatSample = TEXT(R"(
+			void {FunctionName}(in float3 In_SamplePos, out float Out_DistanceToNearestSurface, out float3 Out_FieldGradient, out bool Out_IsDistanceFieldValid)
+			{
+				DICollisionQuery_DistanceField(In_SamplePos, Out_DistanceToNearestSurface, Out_FieldGradient, Out_IsDistanceFieldValid);
+			}
+		)");
+
+		OutHLSL += TEXT("void ") + FunctionInfo.InstanceName + TEXT(" \n{\n");
+	}
+	else if (FunctionInfo.DefinitionName == IssueAsyncRayTraceName)
+	{
+		static const TCHAR* FormatSample = TEXT(R"(
+			void {FunctionName}(int In_QueryID, float3 In_TraceStart, float3 In_TraceEnd, int In_TraceChannel, out bool Out_IsQueryValid)
+			{
+				DICollisionQuery_IssueAsyncRayTrace(In_QueryID, In_TraceStart, In_TraceEnd, In_TraceChannel, Out_IsQueryValid);
+			}
+		)");
+
+		OutHLSL += FString::Format(FormatSample, Args);
+	}
+	else if (FunctionInfo.DefinitionName == ReadAsyncRayTraceName)
+	{
+		static const TCHAR* FormatSample = TEXT(R"(
+			void {FunctionName}(int In_PreviousFrameQueryID, out bool Out_CollisionValid, out float Out_CollisionDistance, out float3 Out_CollisionPosWorld, out float3 Out_CollisionNormal)
+			{
+				DICollisionQuery_ReadAsyncRayTrace(In_PreviousFrameQueryID, Out_CollisionValid, Out_CollisionDistance, Out_CollisionPosWorld, Out_CollisionNormal);
+			}
+		)");
+
+		OutHLSL += FString::Format(FormatSample, Args);
+	}
+	else
+	{
+		return false;
 	}
 
-	return false;
+	return true;
 }
 
 #if WITH_EDITORONLY_DATA
@@ -385,6 +519,11 @@ void UNiagaraDataInterfaceCollisionQuery::ValidateFunction(const FNiagaraFunctio
 }
 #endif
 
+bool UNiagaraDataInterfaceCollisionQuery::RequiresRayTracingScene() const
+{
+	return IsRayTracingEnabled() && MaxRayTraceCount > 0;
+}
+
 void UNiagaraDataInterfaceCollisionQuery::GetParameterDefinitionHLSL(const FNiagaraDataInterfaceGPUParamInfo& ParamInfo, FString& OutHLSL)
 {
 	// we don't need to add these to hlsl, as they're already in common.ush
@@ -421,6 +560,11 @@ void UNiagaraDataInterfaceCollisionQuery::GetVMExternalFunction(const FVMExterna
 	}
 }
 
+void UNiagaraDataInterfaceCollisionQuery::GetCommonHLSL(FString& OutHlsl)
+{
+	OutHlsl += TEXT("#include \"/Plugin/FX/Niagara/Private/NiagaraDataInterfaceCollisionQuery.ush\"\n");
+}
+
 bool UNiagaraDataInterfaceCollisionQuery::AppendCompileHash(FNiagaraCompileHashVisitor* InVisitor) const
 {
 	if (!Super::AppendCompileHash(InVisitor))
@@ -429,6 +573,9 @@ bool UNiagaraDataInterfaceCollisionQuery::AppendCompileHash(FNiagaraCompileHashV
 	}
 	bool bDistanceFieldEnabled = IsDistanceFieldEnabled();
 	InVisitor->UpdatePOD(TEXT("NiagaraCollisionDI_DistanceField"), bDistanceFieldEnabled);
+
+	FSHAHash Hash = GetShaderFileHash(TEXT("/Plugin/FX/Niagara/Private/NiagaraDataInterfaceCollisionQuery.ush"), EShaderPlatform::SP_PCD3D_SM5);
+	InVisitor->UpdateString(TEXT("NiagaraDataInterfaceCollisionQueryHlslSource"), Hash.ToString());
 	return true;
 }
 
@@ -659,6 +806,59 @@ bool UNiagaraDataInterfaceCollisionQuery::PerInstanceTickPostSimulate(void* PerI
 	return false;
 }
 
+bool UNiagaraDataInterfaceCollisionQuery::Equals(const UNiagaraDataInterface* Other) const
+{
+	if (!Super::Equals(Other))
+	{
+		return false;
+	}
+
+	const UNiagaraDataInterfaceCollisionQuery* OtherTyped = CastChecked<const UNiagaraDataInterfaceCollisionQuery>(Other);
+	return OtherTyped->MaxRayTraceCount == MaxRayTraceCount;
+}
+
+bool UNiagaraDataInterfaceCollisionQuery::CopyToInternal(UNiagaraDataInterface* Destination) const
+{
+	if (!Super::CopyToInternal(Destination))
+	{
+		return false;
+	}
+
+	UNiagaraDataInterfaceCollisionQuery* OtherTyped = CastChecked<UNiagaraDataInterfaceCollisionQuery>(Destination);
+	OtherTyped->MaxRayTraceCount = MaxRayTraceCount;
+	OtherTyped->MarkRenderDataDirty();
+	return true;
+}
+
+void UNiagaraDataInterfaceCollisionQuery::PushToRenderThreadImpl()
+{
+	if (!GSupportsResourceView)
+	{
+		return;
+	}
+
+	FNiagaraDataIntefaceProxyCollisionQuery* RT_Proxy = GetProxyAs<FNiagaraDataIntefaceProxyCollisionQuery>();
+
+	// Push Updates to Proxy, first release any resources
+	ENQUEUE_RENDER_COMMAND(FUpdateDI)(
+		[RT_Proxy, RT_MaxRayTraceRequests = MaxRayTraceCount](FRHICommandListImmediate& RHICmdList)
+		{
+			RT_Proxy->RenderThreadInitialize(RT_MaxRayTraceRequests);
+		});
+}
+
+#if WITH_EDITOR
+void UNiagaraDataInterfaceCollisionQuery::PostEditChangeProperty(struct FPropertyChangedEvent& PropertyChangedEvent)
+{
+	Super::PostEditChangeProperty(PropertyChangedEvent);
+
+	if (PropertyChangedEvent.Property && PropertyChangedEvent.Property->GetFName() == GET_MEMBER_NAME_CHECKED(UNiagaraDataInterfaceCollisionQuery, MaxRayTraceCount))
+	{
+		MarkRenderDataDirty();
+	}
+}
+#endif
+
 //////////////////////////////////////////////////////////////////////////
 
 struct FNiagaraDataInterfaceParametersCS_CollisionQuery : public FNiagaraDataInterfaceParametersCS
@@ -669,12 +869,18 @@ public:
 	{
 		PassUniformBuffer.Bind(ParameterMap, FSceneTextureUniformParameters::StaticStructMetadata.GetShaderVariableName());
 		GlobalDistanceFieldParameters.Bind(ParameterMap);
+#if RHI_RAYTRACING
+		MaxRayTraceCountParam.Bind(ParameterMap, TEXT("MaxRayTraceCount"));
+		RayRequestsParam.Bind(ParameterMap, TEXT("RayRequests"));
+		IntersectionResultsParam.Bind(ParameterMap, TEXT("IntersectionResults"));
+#endif
 	}
 
 	void Set(FRHICommandList& RHICmdList, const FNiagaraDataInterfaceSetArgs& Context) const
 	{
 		check(IsInRenderingThread());
 
+		FNiagaraDataIntefaceProxyCollisionQuery* QueryDI = (FNiagaraDataIntefaceProxyCollisionQuery*)Context.DataInterface;
 		FRHIComputeShader* ComputeShaderRHI = Context.Shader.GetComputeShader();
 		
 		//-Note: Scene textures will not exist in the Mobile rendering path
@@ -685,14 +891,47 @@ public:
 		if (GlobalDistanceFieldParameters.IsBound() && Context.Batcher)
 		{
 			GlobalDistanceFieldParameters.Set(RHICmdList, ComputeShaderRHI, Context.Batcher->GetGlobalDistanceFieldParameters());
-		}		
+		}
+
+#if RHI_RAYTRACING
+		SetShaderValue(RHICmdList, ComputeShaderRHI, MaxRayTraceCountParam, QueryDI->MaxRayTraceCount);
+
+		if (RayRequestsParam.IsUAVBound())
+		{
+			RHICmdList.SetUAVParameter(ComputeShaderRHI, RayRequestsParam.GetUAVIndex(), QueryDI->RayTraceRequests.UAV);
+		}
+
+		if (IntersectionResultsParam.IsBound())
+		{
+			SetSRVParameter(RHICmdList, ComputeShaderRHI, IntersectionResultsParam, QueryDI->RayTraceIntersections.SRV);
+		}
+#endif
 	}
 
+#if RHI_RAYTRACING
+	void Unset(FRHICommandList& RHICmdList, const FNiagaraDataInterfaceSetArgs& Context) const
+	{
+		FRHIComputeShader* ComputeShaderRHI = Context.Shader.GetComputeShader();
+		FNiagaraDataIntefaceProxyCollisionQuery* QueryDI = (FNiagaraDataIntefaceProxyCollisionQuery*)Context.DataInterface;
+
+		if (RayRequestsParam.IsUAVBound())
+		{
+			RayRequestsParam.UnsetUAV(RHICmdList, ComputeShaderRHI);
+		}
+	}
+#endif
+
 private:
 	/** The SceneDepthTexture parameter for depth buffer collision. */
 	LAYOUT_FIELD(FShaderUniformBufferParameter, PassUniformBuffer);
 
 	LAYOUT_FIELD(FGlobalDistanceFieldParameters, GlobalDistanceFieldParameters);
+
+#if RHI_RAYTRACING
+	LAYOUT_FIELD(FShaderParameter, MaxRayTraceCountParam);
+	LAYOUT_FIELD(FRWShaderParameter, RayRequestsParam);
+	LAYOUT_FIELD(FShaderResourceParameter, IntersectionResultsParam);
+#endif
 };
 
 IMPLEMENT_TYPE_LAYOUT(FNiagaraDataInterfaceParametersCS_CollisionQuery);
diff --git a/Engine/Plugins/FX/Niagara/Source/Niagara/Private/NiagaraEmitterInstanceBatcher.cpp b/Engine/Plugins/FX/Niagara/Source/Niagara/Private/NiagaraEmitterInstanceBatcher.cpp
--- a/Engine/Plugins/FX/Niagara/Source/Niagara/Private/NiagaraEmitterInstanceBatcher.cpp
+++ b/Engine/Plugins/FX/Niagara/Source/Niagara/Private/NiagaraEmitterInstanceBatcher.cpp
@@ -12,6 +12,7 @@
 #include "NiagaraGpuComputeDebug.h"
 #endif
 #include "NiagaraGpuReadbackManager.h"
+#include "NiagaraRayTracingHelper.h"
 #include "NiagaraRenderer.h"
 #include "NiagaraScriptExecutionContext.h"
 #include "NiagaraShader.h"
@@ -123,7 +124,7 @@ namespace NiagaraEmitterInstanceBatcherLocal
 
 	static ETickStage CalculateTickStage(const FNiagaraGPUSystemTick& Tick)
 	{
-		if (!GNiagaraAllowTickBeforeRender || Tick.bRequiresDistanceFieldData || Tick.bRequiresDepthBuffer)
+		if (!GNiagaraAllowTickBeforeRender || Tick.bRequiresDistanceFieldData || Tick.bRequiresDepthBuffer || Tick.bRequiresRayTracingScene)
 		{
 			return ETickStage::PostOpaqueRender;
 		}
@@ -205,6 +206,10 @@ NiagaraEmitterInstanceBatcher::NiagaraEmitterInstanceBatcher(ERHIFeatureLevel::T
 	EmitterCBufferLayout->UBLayout.ConstantBufferSize = sizeof(FNiagaraEmitterParameters);
 	EmitterCBufferLayout->UBLayout.ComputeHash();
 
+#if RHI_RAYTRACING
+	RayTracingHelper.Reset(new FNiagaraRayTracingHelper(InShaderPlatform));
+#endif
+
 #if NIAGARA_COMPUTEDEBUG_ENABLED
 	GpuComputeDebugPtr.Reset(new FNiagaraGpuComputeDebug(FeatureLevel));
 #endif
@@ -253,10 +258,12 @@ void NiagaraEmitterInstanceBatcher::InstanceDeallocated_RenderThread(const FNiag
 			ensure(NumTicksThatRequireDistanceFieldData >= (Tick.bRequiresDistanceFieldData ? 1u : 0u));
 			ensure(NumTicksThatRequireDepthBuffer >= (Tick.bRequiresDepthBuffer ? 1u : 0u));
 			ensure(NumTicksThatRequireEarlyViewData >= (Tick.bRequiresEarlyViewData ? 1u : 0u));
+			ensure(NumTicksThatRequireRayTracingScene >= (Tick.bRequiresRayTracingScene ? 1u : 0u));
 
 			NumTicksThatRequireDistanceFieldData -= Tick.bRequiresDistanceFieldData ? 1 : 0;
 			NumTicksThatRequireDepthBuffer -= Tick.bRequiresDepthBuffer ? 1 : 0;
 			NumTicksThatRequireEarlyViewData -= Tick.bRequiresEarlyViewData ? 1 : 0;
+			NumTicksThatRequireRayTracingScene -= Tick.bRequiresRayTracingScene ? 1 : 0;
 
 			//-OPT: Since we can't RemoveAtSwap (due to ordering issues) if may be better to not remove and flag as dead
 			Tick.Destroy();
@@ -397,6 +404,7 @@ void NiagaraEmitterInstanceBatcher::GiveSystemTick_RenderThread(FNiagaraGPUSyste
 	NumTicksThatRequireDistanceFieldData += Tick.bRequiresDistanceFieldData ? 1 : 0;
 	NumTicksThatRequireDepthBuffer += Tick.bRequiresDepthBuffer ? 1 : 0;
 	NumTicksThatRequireEarlyViewData += Tick.bRequiresEarlyViewData ? 1 : 0;
+	NumTicksThatRequireRayTracingScene += Tick.bRequiresRayTracingScene ? 1 : 0;
 }
 
 void NiagaraEmitterInstanceBatcher::ReleaseInstanceCounts_RenderThread(FNiagaraComputeExecutionContext* ExecContext, FNiagaraDataSet* DataSet)
@@ -522,6 +530,7 @@ void NiagaraEmitterInstanceBatcher::FinishDispatches()
 	NumTicksThatRequireDistanceFieldData = 0;
 	NumTicksThatRequireDepthBuffer = 0;
 	NumTicksThatRequireEarlyViewData = 0;
+	NumTicksThatRequireRayTracingScene = 0;
 
 	for (int32 iTickStage=0; iTickStage < (int)ETickStage::Max; ++iTickStage)
 	{
@@ -1539,15 +1548,20 @@ void NiagaraEmitterInstanceBatcher::PostRenderOpaque(FRDGBuilder& GraphBuilder,
 	LLM_SCOPE(ELLMTag::Niagara);
 
 	bAllowGPUParticleUpdate = bAllowGPUParticleUpdate && GetReferenceAllowGPUUpdate(Views);
-	FRHIUniformBuffer* ViewUniformBuffer = GetReferenceViewUniformBuffer(Views);
 
 	AddPass(
 		GraphBuilder,
 		RDG_EVENT_NAME("Niagara::PostRenderOpaque"),
-		[this, ViewUniformBuffer, bAllowGPUParticleUpdate](FRHICommandListImmediate& RHICmdList)
+		[this, Views, bAllowGPUParticleUpdate](FRHICommandListImmediate& RHICmdList)
 		{
 			if (bAllowGPUParticleUpdate && FNiagaraUtilities::AllowGPUParticles(GetShaderPlatform()))
 			{
+#if RHI_RAYTRACING
+				BuildRayTracingSceneInfo(RHICmdList, Views);
+#endif
+
+				FRHIUniformBuffer* ViewUniformBuffer = GetReferenceViewUniformBuffer(Views);
+
 				BuildTickStagePasses(RHICmdList, ETickStage::PostOpaqueRender);
 
 				// Setup new readback since if there is no pending request, there is no risk of having invalid data read (offset being allocated after the readback was sent).
@@ -1559,6 +1573,10 @@ void NiagaraEmitterInstanceBatcher::PostRenderOpaque(FRDGBuilder& GraphBuilder,
 				FinishDispatches();
 
 				ProcessDebugReadbacks(RHICmdList, false);
+
+#if RHI_RAYTRACING
+				ResetRayTracingSceneInfo();
+#endif
 			}
 
 			if (!GPUInstanceCounterManager.HasPendingGPUReadback())
@@ -1651,6 +1669,13 @@ bool NiagaraEmitterInstanceBatcher::RequiresEarlyViewUniformBuffer() const
 	return NumTicksThatRequireEarlyViewData > 0;
 }
 
+bool NiagaraEmitterInstanceBatcher::RequiresRayTracingScene() const
+{
+	checkSlow(Ticks_RT.ContainsByPredicate([](const FNiagaraGPUSystemTick& Tick) { return Tick.bRequiresRayTracingScene; }) == NumTicksThatRequireRayTracingScene > 0);
+
+	return NumTicksThatRequireRayTracingScene > 0;
+}
+
 void NiagaraEmitterInstanceBatcher::PreRender(FRDGBuilder& GraphBuilder, TConstArrayView<FViewInfo> Views, bool bAllowGPUParticleUpdate)
 {
 	if (!FNiagaraUtilities::AllowGPUParticles(GetShaderPlatform()))
@@ -1783,6 +1808,37 @@ void NiagaraEmitterInstanceBatcher::GenerateSortKeys(FRHICommandListImmediate& R
 	RHICmdList.EndUAVOverlap(MakeArrayView(OverlapUAVs, NumOverlapUAVs));
 }
 
+#if RHI_RAYTRACING
+void NiagaraEmitterInstanceBatcher::BuildRayTracingSceneInfo(FRHICommandList& RHICmdList, TConstArrayView<FViewInfo> Views)
+{
+	if (NumTicksThatRequireRayTracingScene > 0)
+	{
+		RayTracingHelper->BuildRayTracingSceneInfo(RHICmdList, Views);
+	}
+	else
+	{
+		RayTracingHelper->Reset();
+	}
+}
+
+void NiagaraEmitterInstanceBatcher::ResetRayTracingSceneInfo()
+{
+	RayTracingHelper->Reset();
+}
+
+void NiagaraEmitterInstanceBatcher::IssueRayTraces(FRHICommandList& RHICmdList, const FIntPoint& RayTraceCounts, FRHIShaderResourceView* RayTraceRequests, FRHIUnorderedAccessView* RayTraceResults) const
+{
+	check(NumTicksThatRequireRayTracingScene > 0);
+
+	RayTracingHelper->IssueRayTraces(RHICmdList, RayTraceCounts, RayTraceRequests, RayTraceResults);
+}
+
+bool NiagaraEmitterInstanceBatcher::HasRayTracingScene() const
+{
+	return NumTicksThatRequireRayTracingScene && RayTracingHelper->IsValid();
+}
+#endif
+
 /* Set shader parameters for data interfaces
  */
 void NiagaraEmitterInstanceBatcher::SetDataInterfaceParameters(const TArray<FNiagaraDataInterfaceProxy*>& DataInterfaceProxies, const FNiagaraShaderRef& Shader, FRHICommandList& RHICmdList, const FNiagaraComputeInstanceData* Instance, const FNiagaraGPUSystemTick& Tick, uint32 SimulationStageIndex) const
diff --git a/Engine/Plugins/FX/Niagara/Source/Niagara/Private/NiagaraRayTracingHelper.cpp b/Engine/Plugins/FX/Niagara/Source/Niagara/Private/NiagaraRayTracingHelper.cpp
new file mode 100644
--- /dev/null
+++ b/Engine/Plugins/FX/Niagara/Source/Niagara/Private/NiagaraRayTracingHelper.cpp
@@ -0,0 +1,231 @@
+// Copyright Epic Games, Inc. All Rights Reserved.
+
+#include "NiagaraRayTracingHelper.h"
+
+#if RHI_RAYTRACING
+
+#include "GlobalShader.h"
+#include "MeshPassProcessor.h"
+#include "Renderer/Private/SceneRendering.h"
+#include "ShaderParameterStruct.h"
+
+/// TODO
+///  -get DispatchIndirect working where we can atomically increment the collision request index and then use that to
+///		control the dispatch index (or, at the least feed it into the RayGen shader to early out)
+///  -get geometry masking working when an environmental mask is implemented
+///  -template the CollisionDI to allow for multiple instances in the same emitter (even if that doesn't make a whole lot of sense)
+
+class FNiagaraCollisionRayTraceRG : public FGlobalShader
+{
+	DECLARE_GLOBAL_SHADER(FNiagaraCollisionRayTraceRG)
+	SHADER_USE_ROOT_PARAMETER_STRUCT(FNiagaraCollisionRayTraceRG, FGlobalShader)
+
+	BEGIN_SHADER_PARAMETER_STRUCT(FParameters, )
+		SHADER_PARAMETER_RDG_BUFFER_SRV(RaytracingAccelerationStructure, TLAS)
+		SHADER_PARAMETER_RDG_BUFFER_SRV(StructuredBuffer<FBasicRayData>, Rays)
+		SHADER_PARAMETER_RDG_BUFFER_UAV(RWStructuredBuffer<FNiagaraRayTracingPayload>, CollisionOutput)
+	END_SHADER_PARAMETER_STRUCT()
+
+	static bool ShouldCompilePermutation(const FGlobalShaderPermutationParameters& Parameters)
+	{
+		return ShouldCompileRayTracingShadersForProject(Parameters.Platform);
+	}
+};
+
+class FNiagaraCollisionRayTraceCH : public FGlobalShader
+{
+	DECLARE_SHADER_TYPE(FNiagaraCollisionRayTraceCH, Global);
+
+	static bool ShouldCompilePermutation(const FGlobalShaderPermutationParameters& Parameters)
+	{
+		return ShouldCompileRayTracingShadersForProject(Parameters.Platform);
+	}
+
+	FNiagaraCollisionRayTraceCH() = default;
+	FNiagaraCollisionRayTraceCH(const ShaderMetaType::CompiledShaderInitializerType& Initializer)
+		: FGlobalShader(Initializer)
+	{}
+};
+
+class FNiagaraCollisionRayTraceMiss : public FGlobalShader
+{
+	DECLARE_SHADER_TYPE(FNiagaraCollisionRayTraceMiss, Global);
+
+	static bool ShouldCompilePermutation(const FGlobalShaderPermutationParameters& Parameters)
+	{
+		return ShouldCompileRayTracingShadersForProject(Parameters.Platform);
+	}
+
+	FNiagaraCollisionRayTraceMiss() = default;
+	FNiagaraCollisionRayTraceMiss(const ShaderMetaType::CompiledShaderInitializerType& Initializer)
+		: FGlobalShader(Initializer)
+	{}
+};
+
+IMPLEMENT_GLOBAL_SHADER(FNiagaraCollisionRayTraceRG, "/Plugin/FX/Niagara/Private/NiagaraRayTracingShaders.usf", "NiagaraCollisionRayTraceRG", SF_RayGen);
+IMPLEMENT_GLOBAL_SHADER(FNiagaraCollisionRayTraceCH, "/Plugin/FX/Niagara/Private/NiagaraRayTracingShaders.usf", "NiagaraCollisionRayTraceCH", SF_RayHitGroup);
+IMPLEMENT_GLOBAL_SHADER(FNiagaraCollisionRayTraceMiss, "/Plugin/FX/Niagara/Private/NiagaraRayTracingShaders.usf", "NiagaraCollisionRayTraceMiss", SF_RayMiss);
+
+void FNiagaraRayTracingHelper::Reset()
+{
+	RayTracingPipelineState = nullptr;
+	RayTracingScene = nullptr;
+	RayTracingSceneView = nullptr;
+	ViewUniformBuffer = nullptr;
+}
+
+static FRayTracingPipelineState* CreateNiagaraRayTracingPipelineState(
+	EShaderPlatform Platform,
+	FRHICommandList& RHICmdList,
+	FRHIRayTracingShader* RayGenShader,
+	FRHIRayTracingShader* ClosestHitShader,
+	FRHIRayTracingShader* MissShader)
+{
+	FRayTracingPipelineStateInitializer Initializer;
+	Initializer.MaxPayloadSizeInBytes = sizeof(FNiagaraRayTracingPayload);
+
+	FRHIRayTracingShader* RayGenShaderTable[] = { RayGenShader };
+	Initializer.SetRayGenShaderTable(RayGenShaderTable);
+
+	FRHIRayTracingShader* HitGroupTable[] = { ClosestHitShader };
+	Initializer.SetHitGroupTable(HitGroupTable);
+
+	FRHIRayTracingShader* MissTable[] = { MissShader };
+	Initializer.SetMissShaderTable(MissTable);
+
+	Initializer.bAllowHitGroupIndexing = true; // Use the same hit shader for all geometry in the scene by disabling SBT indexing.
+
+	return PipelineStateCache::GetAndOrCreateRayTracingPipelineState(RHICmdList, Initializer);
+}
+
+static void BindNiagaraRayTracingMeshCommands(
+	FRHICommandList& RHICmdList,
+	FRayTracingSceneRHIRef RayTracingScene,
+	FRHIUniformBuffer* ViewUniformBuffer,
+	TConstArrayView<FVisibleRayTracingMeshCommand> RayTracingMeshCommands,
+	FRayTracingPipelineState* Pipeline,
+	TFunctionRef<uint32(const FRayTracingMeshCommand&)> PackUserData)
+{
+	const int32 NumTotalBindings = RayTracingMeshCommands.Num();
+
+	const uint32 MergedBindingsSize = sizeof(FRayTracingLocalShaderBindings) * NumTotalBindings;
+
+	const uint32 NumUniformBuffers = 1;
+	const uint32 UniformBufferArraySize = sizeof(FRHIUniformBuffer*) * NumUniformBuffers;
+
+	FRayTracingLocalShaderBindings* Bindings = nullptr;
+	FRHIUniformBuffer** UniformBufferArray = nullptr;
+
+	if (RHICmdList.Bypass())
+	{
+		FMemStack& MemStack = FMemStack::Get();
+
+		Bindings = reinterpret_cast<FRayTracingLocalShaderBindings*>(MemStack.Alloc(MergedBindingsSize, alignof(FRayTracingLocalShaderBindings)));
+		UniformBufferArray = reinterpret_cast<FRHIUniformBuffer**>(MemStack.Alloc(UniformBufferArraySize, alignof(FRHIUniformBuffer*)));
+	}
+	else
+	{
+		Bindings = reinterpret_cast<FRayTracingLocalShaderBindings*>(RHICmdList.Alloc(MergedBindingsSize, alignof(FRayTracingLocalShaderBindings)));
+		UniformBufferArray = reinterpret_cast<FRHIUniformBuffer**>(RHICmdList.Alloc(UniformBufferArraySize, alignof(FRHIUniformBuffer*)));
+	}
+
+	UniformBufferArray[0] = ViewUniformBuffer;
+
+	uint32 BindingIndex = 0;
+	for (const FVisibleRayTracingMeshCommand VisibleMeshCommand : RayTracingMeshCommands)
+	{
+		const FRayTracingMeshCommand& MeshCommand = *VisibleMeshCommand.RayTracingMeshCommand;
+
+		FRayTracingLocalShaderBindings Binding = {};
+		Binding.InstanceIndex = VisibleMeshCommand.InstanceIndex;
+		Binding.SegmentIndex = MeshCommand.GeometrySegmentIndex;
+		Binding.UserData = PackUserData(MeshCommand);
+		Binding.UniformBuffers = UniformBufferArray;
+		Binding.NumUniformBuffers = NumUniformBuffers;
+
+		Bindings[BindingIndex] = Binding;
+		BindingIndex++;
+	}
+
+	const bool bCopyDataToInlineStorage = false; // Storage is already allocated from RHICmdList, no extra copy necessary
+	RHICmdList.SetRayTracingHitGroups(
+		RayTracingScene,
+		Pipeline,
+		NumTotalBindings,
+		Bindings,
+		bCopyDataToInlineStorage);
+}
+
+void FNiagaraRayTracingHelper::BuildRayTracingSceneInfo(FRHICommandList& RHICmdList, TConstArrayView<FViewInfo> Views)
+{
+	const FViewInfo& ReferenceView = Views[0];
+	if (ReferenceView.HasRayTracingScene())
+	{
+		RayTracingScene = ReferenceView.GetRayTracingSceneChecked();
+		RayTracingSceneView = ReferenceView.GetRayTracingSceneViewChecked();
+		ViewUniformBuffer = ReferenceView.ViewUniformBuffer;
+
+		FGlobalShaderMap* ShaderMap = GetGlobalShaderMap(ShaderPlatform);
+		auto RayGenShader = ShaderMap->GetShader<FNiagaraCollisionRayTraceRG>().GetRayTracingShader();
+		auto ClosestHitShader = ShaderMap->GetShader<FNiagaraCollisionRayTraceCH>().GetRayTracingShader();
+		auto MissShader = ShaderMap->GetShader<FNiagaraCollisionRayTraceMiss>().GetRayTracingShader();
+
+		RayTracingPipelineState = CreateNiagaraRayTracingPipelineState(
+			ShaderPlatform,
+			RHICmdList,
+			RayGenShader,
+			ClosestHitShader,
+			MissShader);
+
+		// some options for what we want with our per MeshCommand user data.  For now we'll ignore it, but possibly
+		// something we'd want to incorporate.  Some examples could be if the material is translucent, or possibly the physical material?
+		auto BakeTranslucent = [&](const FRayTracingMeshCommand& MeshCommand) {	return (MeshCommand.bIsTranslucent != 0) & 0x1;	};
+		auto BakeDefault = [&](const FRayTracingMeshCommand& MeshCommand) { return 0; };
+
+		BindNiagaraRayTracingMeshCommands(
+			RHICmdList,
+			RayTracingScene,
+			ViewUniformBuffer,
+			ReferenceView.VisibleRayTracingMeshCommands,
+			RayTracingPipelineState,
+			BakeDefault);
+	}
+	else
+	{
+		Reset();
+	}
+}
+
+void FNiagaraRayTracingHelper::IssueRayTraces(FRHICommandList& RHICmdList, const FIntPoint& RayTraceCounts, FRHIShaderResourceView* RayTraceRequests, FRHIUnorderedAccessView* RayTraceResults) const
+{
+	check(RayTracingPipelineState);
+	check(RayTracingScene);
+	check(RayTracingSceneView);
+
+	FGlobalShaderMap* ShaderMap = GetGlobalShaderMap(ShaderPlatform);
+
+	FRayTracingShaderBindings GlobalBindings;
+	GlobalBindings.SRVs[0] = RayTracingSceneView;
+	GlobalBindings.SRVs[1] = RayTraceRequests;
+	GlobalBindings.UAVs[0] = RayTraceResults;
+
+	// todo - add support for RayTraceDispatchIndirect
+	RHICmdList.RayTraceDispatch(
+		RayTracingPipelineState,
+		ShaderMap->GetShader<FNiagaraCollisionRayTraceRG>().GetRayTracingShader(),
+		RayTracingScene,
+		GlobalBindings,
+		RayTraceCounts.X,
+		RayTraceCounts.Y
+	);
+}
+
+bool FNiagaraRayTracingHelper::IsValid() const
+{
+	return RayTracingPipelineState
+		&& RayTracingScene
+		&& RayTracingSceneView
+		&& ViewUniformBuffer;
+}
+
+#endif
\ No newline at end of file
diff --git a/Engine/Plugins/FX/Niagara/Source/Niagara/Private/NiagaraRayTracingHelper.h b/Engine/Plugins/FX/Niagara/Source/Niagara/Private/NiagaraRayTracingHelper.h
new file mode 100644
--- /dev/null
+++ b/Engine/Plugins/FX/Niagara/Source/Niagara/Private/NiagaraRayTracingHelper.h
@@ -0,0 +1,47 @@
+// Copyright Epic Games, Inc. All Rights Reserved.
+#pragma once
+
+#include "CoreMinimal.h"
+
+#if RHI_RAYTRACING
+
+class FViewInfo;
+class FRayTracingPipelineState;
+class FRHICommandListImmediate;
+class FRHIRayTracingScene;
+class FRHIShaderResourceView;
+class FRHIUniformBuffer;
+class FRHIUnorderedAccessView;
+
+// c++ mirror of the struct defined in NiagaraRayTracingCommon.ush
+struct FNiagaraRayTracingPayload
+{
+	float	HitT;            // Distance from ray origin to the intersection point in the ray direction. Negative on miss.
+	float	WorldPosition[3];
+	float	WorldNormal[3];
+};
+
+class FNiagaraRayTracingHelper
+{
+	public:
+		FNiagaraRayTracingHelper(EShaderPlatform InShaderPlatform)
+		: ShaderPlatform(InShaderPlatform)
+		{}
+
+		FNiagaraRayTracingHelper() = delete;
+
+		void Reset();
+		void BuildRayTracingSceneInfo(FRHICommandList& RHICmdList, TConstArrayView<FViewInfo> Views);
+		void IssueRayTraces(FRHICommandList& RHICmdList, const FIntPoint& RayTraceCounts, FRHIShaderResourceView* RayTraceRequests, FRHIUnorderedAccessView* RayTraceResults) const;
+		bool IsValid() const;
+
+	private:
+		const EShaderPlatform ShaderPlatform;
+
+		mutable FRayTracingPipelineState* RayTracingPipelineState = nullptr;
+		mutable FRHIRayTracingScene* RayTracingScene = nullptr;
+		mutable FRHIShaderResourceView* RayTracingSceneView = nullptr;
+		mutable FRHIUniformBuffer* ViewUniformBuffer = nullptr;
+};
+
+#endif
diff --git a/Engine/Plugins/FX/Niagara/Source/Niagara/Private/NiagaraScriptExecutionContext.cpp b/Engine/Plugins/FX/Niagara/Source/Niagara/Private/NiagaraScriptExecutionContext.cpp
--- a/Engine/Plugins/FX/Niagara/Source/Niagara/Private/NiagaraScriptExecutionContext.cpp
+++ b/Engine/Plugins/FX/Niagara/Source/Niagara/Private/NiagaraScriptExecutionContext.cpp
@@ -583,6 +583,7 @@ void FNiagaraGPUSystemTick::Init(FNiagaraSystemInstance* InSystemInstance)
 	bRequiresDepthBuffer = InSystemInstance->RequiresDepthBuffer();
 	bRequiresEarlyViewData = InSystemInstance->RequiresEarlyViewData();
 	bRequiresViewUniformBuffer = InSystemInstance->RequiresViewUniformBuffer();
+	bRequiresRayTracingScene = InSystemInstance->RequiresRayTracingScene();
  
 	uint32 DataSizeForGPU = InSystemInstance->GPUDataInterfaceInstanceDataSize;
 
diff --git a/Engine/Plugins/FX/Niagara/Source/Niagara/Private/NiagaraSystemInstance.cpp b/Engine/Plugins/FX/Niagara/Source/Niagara/Private/NiagaraSystemInstance.cpp
--- a/Engine/Plugins/FX/Niagara/Source/Niagara/Private/NiagaraSystemInstance.cpp
+++ b/Engine/Plugins/FX/Niagara/Source/Niagara/Private/NiagaraSystemInstance.cpp
@@ -1304,6 +1304,31 @@ bool FNiagaraSystemInstance::RequiresViewUniformBuffer() const
 	return false;
 }
 
+bool FNiagaraSystemInstance::RequiresRayTracingScene() const
+{
+	if (!bHasGPUEmitters)
+	{
+		return false;
+	}
+
+	for (const TSharedRef<FNiagaraEmitterInstance, ESPMode::ThreadSafe>& Emitter : Emitters)
+	{
+		FNiagaraComputeExecutionContext* GPUContext = Emitter->GetGPUContext();
+		if (GPUContext)
+		{
+			for (UNiagaraDataInterface* DataInterface : GPUContext->CombinedParamStore.GetDataInterfaces())
+			{
+				if (DataInterface && DataInterface->RequiresRayTracingScene())
+				{
+					return true;
+				}
+			}
+		}
+	}
+
+	return false;
+}
+
 void FNiagaraSystemInstance::InitDataInterfaces()
 {
 	bDataInterfacesHaveTickPrereqs = false;
diff --git a/Engine/Plugins/FX/Niagara/Source/Niagara/Public/NiagaraSystemInstance.h b/Engine/Plugins/FX/Niagara/Source/Niagara/Public/NiagaraSystemInstance.h
--- a/Engine/Plugins/FX/Niagara/Source/Niagara/Public/NiagaraSystemInstance.h
+++ b/Engine/Plugins/FX/Niagara/Source/Niagara/Public/NiagaraSystemInstance.h
@@ -111,6 +111,7 @@ public:
 	bool RequiresDepthBuffer() const;
 	bool RequiresEarlyViewData() const;
 	bool RequiresViewUniformBuffer() const;
+	bool RequiresRayTracingScene() const;
 
 	/** Requests the the simulation be reset on the next tick. */
 	void Reset(EResetMode Mode);
diff --git a/Engine/Plugins/FX/Niagara/Source/NiagaraCore/Private/NiagaraCustomVersion.cpp b/Engine/Plugins/FX/Niagara/Source/NiagaraCore/Private/NiagaraCustomVersion.cpp
--- a/Engine/Plugins/FX/Niagara/Source/NiagaraCore/Private/NiagaraCustomVersion.cpp
+++ b/Engine/Plugins/FX/Niagara/Source/NiagaraCore/Private/NiagaraCustomVersion.cpp
@@ -8,4 +8,4 @@ const FGuid FNiagaraCustomVersion::GUID(0xFCF57AFA, 0x50764283, 0xB9A9E658, 0xFF
 // Register the custom version with core
 FCustomVersionRegistration GRegisterNiagaraCustomVersion(FNiagaraCustomVersion::GUID, FNiagaraCustomVersion::LatestVersion, TEXT("NiagaraVer"));
 
-const FGuid FNiagaraCustomVersion::LatestScriptCompileVersion(0xC2E1AF4A,  0xB3CF4652,  0x8BD6FF64, 0xF8BAA37D);
+const FGuid FNiagaraCustomVersion::LatestScriptCompileVersion(0x96A8A2A7, 0x2BB144DF, 0x9E0FCC1F, 0x92DA8F52);
diff --git a/Engine/Source/Runtime/Engine/Private/Particles/FXSystem.cpp b/Engine/Source/Runtime/Engine/Private/Particles/FXSystem.cpp
--- a/Engine/Source/Runtime/Engine/Private/Particles/FXSystem.cpp
+++ b/Engine/Source/Runtime/Engine/Private/Particles/FXSystem.cpp
@@ -447,6 +447,16 @@ bool FFXSystem::RequiresEarlyViewUniformBuffer() const
 	return false;
 }
 
+bool FFXSystem::RequiresRayTracingScene() const
+{
+	if (RHISupportsGPUParticles())
+	{
+		return RequiresRayTracingSceneInternal();
+	}
+
+	return false;
+}
+
 DECLARE_CYCLE_STAT(TEXT("FXPreRender_Prepare"), STAT_CLM_FXPreRender_Prepare, STATGROUP_CommandListMarkers);
 DECLARE_CYCLE_STAT(TEXT("FXPreRender_Simulate"), STAT_CLM_FXPreRender_Simulate, STATGROUP_CommandListMarkers);
 DECLARE_CYCLE_STAT(TEXT("FXPreRender_Finalize"), STAT_CLM_FXPreRender_Finalize, STATGROUP_CommandListMarkers);
diff --git a/Engine/Source/Runtime/Engine/Private/Particles/FXSystemPrivate.h b/Engine/Source/Runtime/Engine/Private/Particles/FXSystemPrivate.h
--- a/Engine/Source/Runtime/Engine/Private/Particles/FXSystemPrivate.h
+++ b/Engine/Source/Runtime/Engine/Private/Particles/FXSystemPrivate.h
@@ -155,6 +155,7 @@ public:
 	virtual bool UsesGlobalDistanceField() const override;
 	virtual bool UsesDepthBuffer() const override;
 	virtual bool RequiresEarlyViewUniformBuffer() const override;
+	virtual bool RequiresRayTracingScene() const override;
 	virtual void PreRender(FRDGBuilder& GraphBuilder, TConstArrayView<FViewInfo> Views, bool bAllowGPUParticleUpdate) override;
 	virtual void PostRenderOpaque(FRDGBuilder& GraphBuilder, TConstArrayView<FViewInfo> Views, bool bAllowGPUParticleUpdate) override;
 	// End FFXSystemInterface.
@@ -272,6 +273,7 @@ private:
 	bool UsesGlobalDistanceFieldInternal() const;
 	bool UsesDepthBufferInternal() const;
 	bool RequiresEarlyViewUniformBufferInternal() const;
+	bool RequiresRayTracingSceneInternal() const;
 
 	/**
 	* Updates resources used in a multi-GPU context
diff --git a/Engine/Source/Runtime/Engine/Private/Particles/FXSystemSet.cpp b/Engine/Source/Runtime/Engine/Private/Particles/FXSystemSet.cpp
--- a/Engine/Source/Runtime/Engine/Private/Particles/FXSystemSet.cpp
+++ b/Engine/Source/Runtime/Engine/Private/Particles/FXSystemSet.cpp
@@ -181,6 +181,19 @@ bool FFXSystemSet::RequiresEarlyViewUniformBuffer() const
 	return false;
 }
 
+bool FFXSystemSet::RequiresRayTracingScene() const
+{
+	for (FFXSystemInterface* FXSystem : FXSystems)
+	{
+		check(FXSystem);
+		if (FXSystem->RequiresRayTracingScene())
+		{
+			return true;
+		}
+	}
+	return false;
+}
+
 void FFXSystemSet::PreRender(class FRDGBuilder& GraphBuilder, TConstArrayView<FViewInfo> Views, bool bAllowGPUParticleSceneUpdate)
 {
 	for (FFXSystemInterface* FXSystem : FXSystems)
diff --git a/Engine/Source/Runtime/Engine/Private/Particles/FXSystemSet.h b/Engine/Source/Runtime/Engine/Private/Particles/FXSystemSet.h
--- a/Engine/Source/Runtime/Engine/Private/Particles/FXSystemSet.h
+++ b/Engine/Source/Runtime/Engine/Private/Particles/FXSystemSet.h
@@ -44,6 +44,7 @@ public:
 	virtual bool UsesGlobalDistanceField() const override;
 	virtual bool UsesDepthBuffer() const override;
 	virtual bool RequiresEarlyViewUniformBuffer() const override;
+	virtual bool RequiresRayTracingScene() const override;
 	virtual void PreRender(FRDGBuilder& GraphBuilder, TConstArrayView<FViewInfo> Views, bool bAllowGPUParticleSceneUpdate) override;
 	virtual void PostRenderOpaque(FRDGBuilder& GraphBuilder, TConstArrayView<FViewInfo> Views, bool bAllowGPUParticleSceneUpdate) override;
 
diff --git a/Engine/Source/Runtime/Engine/Private/Particles/ParticleGpuSimulation.cpp b/Engine/Source/Runtime/Engine/Private/Particles/ParticleGpuSimulation.cpp
--- a/Engine/Source/Runtime/Engine/Private/Particles/ParticleGpuSimulation.cpp
+++ b/Engine/Source/Runtime/Engine/Private/Particles/ParticleGpuSimulation.cpp
@@ -4614,6 +4614,11 @@ bool FFXSystem::RequiresEarlyViewUniformBufferInternal() const
 	return false;
 }
 
+bool FFXSystem::RequiresRayTracingSceneInternal() const
+{
+	return false;
+}
+
 void FFXSystem::PrepareGPUSimulation(FRHICommandListImmediate& RHICmdList)
 {
 	// Grab resources.
diff --git a/Engine/Source/Runtime/Engine/Public/FXSystem.h b/Engine/Source/Runtime/Engine/Public/FXSystem.h
--- a/Engine/Source/Runtime/Engine/Public/FXSystem.h
+++ b/Engine/Source/Runtime/Engine/Public/FXSystem.h
@@ -243,6 +243,8 @@ public:
 
 	virtual bool RequiresEarlyViewUniformBuffer() const = 0;
 
+	virtual bool RequiresRayTracingScene() const = 0;
+
 	/**
 	 * Notification from the renderer that it is about to draw FX belonging to
 	 * this system.
