Fix a crash with GPU Niagara meshes in Emiter source mode, and make the transform bindings work.

diff --git a/Engine/Plugins/FX/Niagara/Source/Niagara/Private/NiagaraRendererMeshes.cpp b/Engine/Plugins/FX/Niagara/Source/Niagara/Private/NiagaraRendererMeshes.cpp
--- a/Engine/Plugins/FX/Niagara/Source/Niagara/Private/NiagaraRendererMeshes.cpp
+++ b/Engine/Plugins/FX/Niagara/Source/Niagara/Private/NiagaraRendererMeshes.cpp
@@ -402,7 +402,6 @@ FNiagaraMeshUniformBufferRef FNiagaraRendererMeshes::CreatePerViewUniformBuffer(
 	}
 
 	PerViewUniformParameters.MaterialParamValidMask = MaterialParamValidMask;
-	PerViewUniformParameters.DefaultPos = bLocalSpace ? FVector4(0.0f, 0.0f, 0.0f, 1.0f) : FVector4(SceneProxy.GetLocalToWorld().GetOrigin());
 	PerViewUniformParameters.SubImageSize = FVector4(SubImageSize.X, SubImageSize.Y, 1.0f / SubImageSize.X, 1.0f / SubImageSize.Y);
 	PerViewUniformParameters.SubImageBlendMode = bSubImageBlend;
 	PerViewUniformParameters.FacingMode = (uint32)FacingMode;
@@ -1158,144 +1157,208 @@ void FNiagaraRendererMeshes::GetDynamicRayTracingInstances(FRayTracingMaterialGa
 		float* RESTRICT QuatArrayZ = (float*)SourceParticleData->GetComponentPtrFloat(TransformBaseCompOffset + 2);
 		float* RESTRICT QuatArrayW = (float*)SourceParticleData->GetComponentPtrFloat(TransformBaseCompOffset + 3);
 
-		auto GetInstancePosition = [&PositionX, &PositionY, &PositionZ](int32 Idx)
-		{
-			return FVector4(PositionX[Idx], PositionY[Idx], PositionZ[Idx], 1);
-		};
+		FMatrix LocalTransform(SceneProxy->GetLocalToWorld());
 
-		auto GetInstanceScale = [&ScaleX, &ScaleY, &ScaleZ](int32 Idx)
+		if (SourceMode == ENiagaraRendererSourceDataMode::Emitter)
 		{
-			return FVector(ScaleX[Idx], ScaleY[Idx], ScaleZ[Idx]);
-		};
+			FVector Pos = bLocalSpace ? FVector() : LocalTransform.GetOrigin();
+			FVector Scale { 1.0f, 1.0f, 1.0f };
+			FQuat Rot = FQuat::Identity;
 
-		auto GetInstanceQuat = [&QuatArrayX, &QuatArrayY, &QuatArrayZ, &QuatArrayW](int32 Idx)
-		{
-			return FQuat(QuatArrayX[Idx], QuatArrayY[Idx], QuatArrayZ[Idx], QuatArrayW[Idx]);
-		};
+			if (bSetAnyBoundVars)
+			{
+				const uint8* ParameterBoundData = DynamicDataMesh->ParameterDataBound.GetData();
+				if (VFBoundOffsetsInParamStore[ENiagaraMeshVFLayout::Type::Position] != INDEX_NONE
+					&& DynamicDataMesh->ParameterDataBound.IsValidIndex(VFBoundOffsetsInParamStore[ENiagaraMeshVFLayout::Type::Position]))
+				{
+					FMemory::Memcpy(&Pos, ParameterBoundData + VFBoundOffsetsInParamStore[ENiagaraMeshVFLayout::Type::Position], sizeof(FVector));
+				}
+				if (VFBoundOffsetsInParamStore[ENiagaraMeshVFLayout::Type::Scale] != INDEX_NONE
+					&& DynamicDataMesh->ParameterDataBound.IsValidIndex(VFBoundOffsetsInParamStore[ENiagaraMeshVFLayout::Type::Scale]))
+				{
+					FMemory::Memcpy(&Scale, ParameterBoundData + VFBoundOffsetsInParamStore[ENiagaraMeshVFLayout::Type::Scale], sizeof(FVector));
+				}
+				if (VFBoundOffsetsInParamStore[ENiagaraMeshVFLayout::Type::Rotation] != INDEX_NONE
+					&& DynamicDataMesh->ParameterDataBound.IsValidIndex(VFBoundOffsetsInParamStore[ENiagaraMeshVFLayout::Type::Rotation]))
+				{
+					FMemory::Memcpy(&Rot, ParameterBoundData + VFBoundOffsetsInParamStore[ENiagaraMeshVFLayout::Type::Rotation], sizeof(FVector4));
+				}
+			}
 
-		//#dxr_todo: handle MESH_FACING_VELOCITY, MESH_FACING_CAMERA_POSITION, MESH_FACING_CAMERA_PLANE
-		bool bHasPosition = PositionBaseCompOffset > 0;
-		bool bHasRotation = TransformBaseCompOffset > 0;
-		bool bHasScale = ScaleBaseCompOffset > 0;
+			FVector4 Transform1 = FVector4(1.0f, 0.0f, 0.0f, Pos.X);
+			FVector4 Transform2 = FVector4(0.0f, 1.0f, 0.0f, Pos.Y);
+			FVector4 Transform3 = FVector4(0.0f, 0.0f, 1.0f, Pos.Z);
 
-		FMatrix LocalTransform(SceneProxy->GetLocalToWorld());
+			FTransform RotationTransform(Rot.GetNormalized());
+			FMatrix RotationMatrix = RotationTransform.ToMatrixWithScale();
 
-		for (int32 InstanceIndex = 0; InstanceIndex < NumInstances; InstanceIndex++)
+			Transform1.X = RotationMatrix.M[0][0];
+			Transform1.Y = RotationMatrix.M[0][1];
+			Transform1.Z = RotationMatrix.M[0][2];
+
+			Transform2.X = RotationMatrix.M[1][0];
+			Transform2.Y = RotationMatrix.M[1][1];
+			Transform2.Z = RotationMatrix.M[1][2];
+
+			Transform3.X = RotationMatrix.M[2][0];
+			Transform3.Y = RotationMatrix.M[2][1];
+			Transform3.Z = RotationMatrix.M[2][2];
+
+			FMatrix ScaleMatrix(FMatrix::Identity);
+			ScaleMatrix.M[0][0] *= Scale.X;
+			ScaleMatrix.M[1][1] *= Scale.Y;
+			ScaleMatrix.M[2][2] *= Scale.Z;
+
+			FMatrix InstanceTransform = FMatrix(FPlane(Transform1), FPlane(Transform2), FPlane(Transform3), FPlane(0.0, 0.0, 0.0, 1.0));
+			InstanceTransform = InstanceTransform * ScaleMatrix;
+			InstanceTransform = InstanceTransform.GetTransposed();
+
+			if (bLocalSpace)
+			{
+				InstanceTransform = InstanceTransform * LocalTransform;
+			}
+
+			RayTracingInstance.InstanceTransforms.Add(InstanceTransform);
+		}
+		else
 		{
-			FMatrix InstanceTransform(FMatrix::Identity);
+			auto GetInstancePosition = [&PositionX, &PositionY, &PositionZ](int32 Idx)
+			{
+				return FVector4(PositionX[Idx], PositionY[Idx], PositionZ[Idx], 1);
+			};
+
+			auto GetInstanceScale = [&ScaleX, &ScaleY, &ScaleZ](int32 Idx)
+			{
+				return FVector(ScaleX[Idx], ScaleY[Idx], ScaleZ[Idx]);
+			};
 
-			if (SimTarget == ENiagaraSimTarget::CPUSim)
+			auto GetInstanceQuat = [&QuatArrayX, &QuatArrayY, &QuatArrayZ, &QuatArrayW](int32 Idx)
 			{
-				FVector4 InstancePos = bHasPosition ? GetInstancePosition(InstanceIndex) : FVector4(0, 0, 0, 0);
+				return FQuat(QuatArrayX[Idx], QuatArrayY[Idx], QuatArrayZ[Idx], QuatArrayW[Idx]);
+			};
+
+			//#dxr_todo: handle MESH_FACING_VELOCITY, MESH_FACING_CAMERA_POSITION, MESH_FACING_CAMERA_PLANE
+			bool bHasPosition = PositionBaseCompOffset > 0;
+			bool bHasRotation = TransformBaseCompOffset > 0;
+			bool bHasScale = ScaleBaseCompOffset > 0;
 
-				FVector4 Transform1 = FVector4(1.0f, 0.0f, 0.0f, InstancePos.X);
-				FVector4 Transform2 = FVector4(0.0f, 1.0f, 0.0f, InstancePos.Y);
-				FVector4 Transform3 = FVector4(0.0f, 0.0f, 1.0f, InstancePos.Z);
+			for (int32 InstanceIndex = 0; InstanceIndex < NumInstances; InstanceIndex++)
+			{
+				FMatrix InstanceTransform(FMatrix::Identity);
 
-				if (bHasRotation)
+				if (SimTarget == ENiagaraSimTarget::CPUSim)
 				{
-					FQuat InstanceQuat = GetInstanceQuat(InstanceIndex);
-					FTransform RotationTransform(InstanceQuat.GetNormalized());
-					FMatrix RotationMatrix = RotationTransform.ToMatrixWithScale();
+					FVector4 InstancePos = bHasPosition ? GetInstancePosition(InstanceIndex) : FVector4(0, 0, 0, 0);
 
-					Transform1.X = RotationMatrix.M[0][0];
-					Transform1.Y = RotationMatrix.M[0][1];
-					Transform1.Z = RotationMatrix.M[0][2];
+					FVector4 Transform1 = FVector4(1.0f, 0.0f, 0.0f, InstancePos.X);
+					FVector4 Transform2 = FVector4(0.0f, 1.0f, 0.0f, InstancePos.Y);
+					FVector4 Transform3 = FVector4(0.0f, 0.0f, 1.0f, InstancePos.Z);
 
-					Transform2.X = RotationMatrix.M[1][0];
-					Transform2.Y = RotationMatrix.M[1][1];
-					Transform2.Z = RotationMatrix.M[1][2];
+					if (bHasRotation)
+					{
+						FQuat InstanceQuat = GetInstanceQuat(InstanceIndex);
+						FTransform RotationTransform(InstanceQuat.GetNormalized());
+						FMatrix RotationMatrix = RotationTransform.ToMatrixWithScale();
 
-					Transform3.X = RotationMatrix.M[2][0];
-					Transform3.Y = RotationMatrix.M[2][1];
-					Transform3.Z = RotationMatrix.M[2][2];
-				}
+						Transform1.X = RotationMatrix.M[0][0];
+						Transform1.Y = RotationMatrix.M[0][1];
+						Transform1.Z = RotationMatrix.M[0][2];
 
-				FMatrix ScaleMatrix(FMatrix::Identity);
-				if (bHasScale)
-				{
-					FVector InstanceSca(GetInstanceScale(InstanceIndex));
-					ScaleMatrix.M[0][0] *= InstanceSca.X;
-					ScaleMatrix.M[1][1] *= InstanceSca.Y;
-					ScaleMatrix.M[2][2] *= InstanceSca.Z;
-				}
+						Transform2.X = RotationMatrix.M[1][0];
+						Transform2.Y = RotationMatrix.M[1][1];
+						Transform2.Z = RotationMatrix.M[1][2];
 
-				InstanceTransform = FMatrix(FPlane(Transform1), FPlane(Transform2), FPlane(Transform3), FPlane(0.0, 0.0, 0.0, 1.0));
-				InstanceTransform = InstanceTransform * ScaleMatrix;
-				InstanceTransform = InstanceTransform.GetTransposed();
+						Transform3.X = RotationMatrix.M[2][0];
+						Transform3.Y = RotationMatrix.M[2][1];
+						Transform3.Z = RotationMatrix.M[2][2];
+					}
+
+					FMatrix ScaleMatrix(FMatrix::Identity);
+					if (bHasScale)
+					{
+						FVector InstanceScale(GetInstanceScale(InstanceIndex));
+						ScaleMatrix.M[0][0] *= InstanceScale.X;
+						ScaleMatrix.M[1][1] *= InstanceScale.Y;
+						ScaleMatrix.M[2][2] *= InstanceScale.Z;
+					}
+
+					InstanceTransform = FMatrix(FPlane(Transform1), FPlane(Transform2), FPlane(Transform3), FPlane(0.0, 0.0, 0.0, 1.0));
+					InstanceTransform = InstanceTransform * ScaleMatrix;
+					InstanceTransform = InstanceTransform.GetTransposed();
 
-				if (bLocalSpace)
+					if (bLocalSpace)
+					{
+						InstanceTransform = InstanceTransform * LocalTransform;
+					}
+				}
+				else
 				{
-					InstanceTransform = InstanceTransform * LocalTransform;
+					// Indirect instancing dispatching: transforms are not available at this point but computed in GPU instead
+					// Set invalid transforms so ray tracing ignores them. Valid transforms will be set later directly in the GPU
+					FMatrix ScaleTransform = FMatrix::Identity;
+					ScaleTransform.M[0][0] = 0.0;
+					ScaleTransform.M[1][1] = 0.0;
+					ScaleTransform.M[2][2] = 0.0;
+
+					InstanceTransform = ScaleTransform * InstanceTransform;
 				}
-			}
-			else
-			{
-				// Indirect instancing dispatching: transforms are not available at this point but computed in GPU instead
-				// Set invalid transforms so ray tracing ignores them. Valid transforms will be set later directly in the GPU
-				FMatrix ScaleTransform = FMatrix::Identity;
-				ScaleTransform.M[0][0] = 0.0;
-				ScaleTransform.M[1][1] = 0.0;
-				ScaleTransform.M[2][2] = 0.0;
-
-				InstanceTransform = ScaleTransform * InstanceTransform;
+
+				RayTracingInstance.InstanceTransforms.Add(InstanceTransform);
 			}
 
-			RayTracingInstance.InstanceTransforms.Add(InstanceTransform);
-		}
+			// Set indirect transforms for GPU instances
+			if (SimTarget == ENiagaraSimTarget::GPUComputeSim
+				&& FNiagaraUtilities::AllowComputeShaders(GShaderPlatformForFeatureLevel[FeatureLevel])
+				&& FDataDrivenShaderPlatformInfo::GetSupportsRayTracingIndirectInstanceData(GShaderPlatformForFeatureLevel[FeatureLevel])
+				)
+			{
+				FRHICommandListImmediate& RHICmdList = Context.RHICmdList;
 
-		// Set indirect transforms for GPU instances
-		if (SimTarget == ENiagaraSimTarget::GPUComputeSim
-			&& FNiagaraUtilities::AllowComputeShaders(GShaderPlatformForFeatureLevel[FeatureLevel])
-			&& FDataDrivenShaderPlatformInfo::GetSupportsRayTracingIndirectInstanceData(GShaderPlatformForFeatureLevel[FeatureLevel])
-			)
-		{
-			FRHICommandListImmediate& RHICmdList = Context.RHICmdList;
+				uint32 CPUInstancesCount = SourceParticleData->GetNumInstances();
 
-			uint32 CPUInstancesCount = SourceParticleData->GetNumInstances();
+				RayTracingInstance.NumTransforms = CPUInstancesCount;
 
-			RayTracingInstance.NumTransforms = CPUInstancesCount;
+				FRWBufferStructured InstanceGPUTransformsBuffer;
+				//InstanceGPUTransformsBuffer.Initialize(sizeof(FMatrix), CPUInstancesCount, BUF_Static);
+				InstanceGPUTransformsBuffer.Initialize(TEXT("InstanceGPUTransformsBuffer"), 3 * 4 * sizeof(float), CPUInstancesCount, BUF_Static);
+				RayTracingInstance.InstanceGPUTransformsSRV = InstanceGPUTransformsBuffer.SRV;
 
-			FRWBufferStructured InstanceGPUTransformsBuffer;
-			//InstanceGPUTransformsBuffer.Initialize(sizeof(FMatrix), CPUInstancesCount, BUF_Static);
-			InstanceGPUTransformsBuffer.Initialize(TEXT("InstanceGPUTransformsBuffer"), 3 * 4 * sizeof(float), CPUInstancesCount, BUF_Static);
-			RayTracingInstance.InstanceGPUTransformsSRV = InstanceGPUTransformsBuffer.SRV;
+				FNiagaraGPURayTracingTransformsCS::FPermutationDomain PermutationVector;
 
-			FNiagaraGPURayTracingTransformsCS::FPermutationDomain PermutationVector;
+				TShaderMapRef<FNiagaraGPURayTracingTransformsCS> GPURayTracingTransformsCS(GetGlobalShaderMap(FeatureLevel), PermutationVector);
+				RHICmdList.SetComputeShader(GPURayTracingTransformsCS.GetComputeShader());
 
-			TShaderMapRef<FNiagaraGPURayTracingTransformsCS> GPURayTracingTransformsCS(GetGlobalShaderMap(FeatureLevel), PermutationVector);
-			RHICmdList.SetComputeShader(GPURayTracingTransformsCS.GetComputeShader());
+				const FUintVector4 NiagaraOffsets(
+					VFVariables[ENiagaraMeshVFLayout::Position].GetGPUOffset(),
+					VFVariables[ENiagaraMeshVFLayout::Rotation].GetGPUOffset(),
+					VFVariables[ENiagaraMeshVFLayout::Scale].GetGPUOffset(),
+					bLocalSpace ? 1 : 0);
 
-			const FUintVector4 NiagaraOffsets(
-				VFVariables[ENiagaraMeshVFLayout::Position].GetGPUOffset(),
-				VFVariables[ENiagaraMeshVFLayout::Rotation].GetGPUOffset(),
-				VFVariables[ENiagaraMeshVFLayout::Scale].GetGPUOffset(),
-				bLocalSpace ? 1 : 0);
+				uint32 FloatDataOffset = 0;
+				uint32 FloatDataStride = SourceParticleData->GetFloatStride() / sizeof(float);
 
-			uint32 FloatDataOffset = 0;
-			uint32 FloatDataStride = SourceParticleData->GetFloatStride() / sizeof(float);
+				GPURayTracingTransformsCS->SetParameters(
+					RHICmdList,
+					CPUInstancesCount,
+					SourceParticleData->GetGPUBufferFloat().SRV,
+					FloatDataOffset,
+					FloatDataStride,
+					SourceParticleData->GetGPUInstanceCountBufferOffset(),
+					Batcher->GetGPUInstanceCounterManager().GetInstanceCountBuffer().SRV,
+					NiagaraOffsets,
+					LocalTransform,
+					InstanceGPUTransformsBuffer.UAV);
 
-			GPURayTracingTransformsCS->SetParameters(
-				RHICmdList,
-				CPUInstancesCount,
-				SourceParticleData->GetGPUBufferFloat().SRV,
-				FloatDataOffset,
-				FloatDataStride,
-				SourceParticleData->GetGPUInstanceCountBufferOffset(),
-				Batcher->GetGPUInstanceCounterManager().GetInstanceCountBuffer().SRV,
-				NiagaraOffsets,
-				LocalTransform,
-				InstanceGPUTransformsBuffer.UAV);
+				uint32 NGroups = FMath::DivideAndRoundUp(CPUInstancesCount, FNiagaraGPURayTracingTransformsCS::ThreadGroupSize);
+				DispatchComputeShader(RHICmdList, GPURayTracingTransformsCS, NGroups, 1, 1);
+				GPURayTracingTransformsCS->UnbindBuffers(RHICmdList);
 
-			uint32 NGroups = FMath::DivideAndRoundUp(CPUInstancesCount, FNiagaraGPURayTracingTransformsCS::ThreadGroupSize);
-			DispatchComputeShader(RHICmdList, GPURayTracingTransformsCS, NGroups, 1, 1);
-			GPURayTracingTransformsCS->UnbindBuffers(RHICmdList);
+				RHICmdList.Transition(FRHITransitionInfo(InstanceGPUTransformsBuffer.UAV, ERHIAccess::Unknown, ERHIAccess::SRVCompute));
+			}
 
-			RHICmdList.Transition(FRHITransitionInfo(InstanceGPUTransformsBuffer.UAV, ERHIAccess::Unknown, ERHIAccess::SRVCompute));
+			RayTracingInstance.BuildInstanceMaskAndFlags();
+			OutRayTracingInstances.Add(RayTracingInstance);
 		}
-
-		RayTracingInstance.BuildInstanceMaskAndFlags();
-		OutRayTracingInstances.Add(RayTracingInstance);
 	}
 }
 #endif
