Editgrate 15863104 - Initialize Niagara Previous values to the current value at the end of the spawn script to fix motion blur issues on the first frame of a particle's life.

diff --git a/Engine/Plugins/FX/Niagara/Source/NiagaraCore/Private/NiagaraCustomVersion.cpp b/Engine/Plugins/FX/Niagara/Source/NiagaraCore/Private/NiagaraCustomVersion.cpp
--- a/Engine/Plugins/FX/Niagara/Source/NiagaraCore/Private/NiagaraCustomVersion.cpp
+++ b/Engine/Plugins/FX/Niagara/Source/NiagaraCore/Private/NiagaraCustomVersion.cpp
@@ -8,4 +8,4 @@ const FGuid FNiagaraCustomVersion::GUID(0xFCF57AFA, 0x50764283, 0xB9A9E658, 0xFF
 // Register the custom version with core
 FCustomVersionRegistration GRegisterNiagaraCustomVersion(FNiagaraCustomVersion::GUID, FNiagaraCustomVersion::LatestVersion, TEXT("NiagaraVer"));
 
-const FGuid FNiagaraCustomVersion::LatestScriptCompileVersion(0x96A8A2A7, 0x2BB144DF, 0x9E0FCC1F, 0x92DA8F52);
+const FGuid FNiagaraCustomVersion::LatestScriptCompileVersion(0xE76B1567, 0x609C40DE, 0xA2ABD4E7, 0x133C1623);
diff --git a/Engine/Plugins/FX/Niagara/Source/NiagaraEditor/Private/NiagaraHlslTranslator.cpp b/Engine/Plugins/FX/Niagara/Source/NiagaraEditor/Private/NiagaraHlslTranslator.cpp
--- a/Engine/Plugins/FX/Niagara/Source/NiagaraEditor/Private/NiagaraHlslTranslator.cpp
+++ b/Engine/Plugins/FX/Niagara/Source/NiagaraEditor/Private/NiagaraHlslTranslator.cpp
@@ -1589,6 +1589,10 @@ const FNiagaraTranslateResults &FHlslNiagaraTranslator::Translate(const FNiagara
 				ActiveHistoryForFunctionCalls.EndUsage();
 				InstanceWrite = FDataSetAccessInfo(); // Reset after building the output..
 				AddBodyComment(TEXT("//End Spawn Script!\n\n"));
+
+				AddBodyComment(TEXT("//Handle resetting previous values at the end of spawn so that they match outputs! (Needed for motion blur/etc)"));
+				AddBodyChunk(TEXT("HandlePreviousValuesForSpawn(Context);"));
+				
 				BuildMissingDefaults();
 			}
 
@@ -1897,6 +1901,7 @@ const FNiagaraTranslateResults &FHlslNiagaraTranslator::Translate(const FNiagara
 		}
 
 		DefineInterpolatedParametersFunction(HlslOutput);
+		DefinePreviousParametersFunction(HlslOutput, DataSetVariables, DataSetReads, DataSetWrites);
 
 		// define functions for reading and writing all secondary data sets
 		DefineDataSetReadFunction(HlslOutput, CompilationOutput.ScriptData.ReadDataSets);
@@ -2417,6 +2422,38 @@ void FHlslNiagaraTranslator::GatherComponentsForDataSetAccess(UScriptStruct* Str
 	}
 }
 
+void FHlslNiagaraTranslator::DefinePreviousParametersFunction(FString& HlslOutputString, TArray<TArray<FNiagaraVariable>>& DataSetVariables, TMap<FNiagaraDataSetID, int32>& DataSetReads, TMap<FNiagaraDataSetID, int32>& DataSetWrites)
+{
+	HlslOutputString +=
+		TEXT("#if (((SimulationStageIndex >= 0) && (SimulationStageIndex < 1))) // MapSpawn\n")
+		TEXT("void HandlePreviousValuesForSpawn(inout FSimulationContext Context)\n{\n");
+	if (UNiagaraScript::IsParticleSpawnScript(CompileOptions.TargetUsage) || UNiagaraScript::IsGPUScript(CompileOptions.TargetUsage))
+	{
+		TArray<FNiagaraDataSetID> ReadDataSetIDs;
+		TArray<FNiagaraDataSetID> WriteDataSetIDs;
+
+		DataSetReads.GetKeys(ReadDataSetIDs);
+		DataSetWrites.GetKeys(WriteDataSetIDs);
+
+		for (int32 DataSetIndex = 0; DataSetIndex < DataSetWrites.Num(); ++DataSetIndex)
+		{
+			const FNiagaraDataSetID DataSetID = ReadDataSetIDs[DataSetIndex];
+			const TArray<FNiagaraVariable>& NiagaraVariables = DataSetVariables[DataSetWrites[DataSetID]];
+			for (const FNiagaraVariable& Var : NiagaraVariables)
+			{
+				if (FNiagaraParameterMapHistory::IsPreviousValue(Var))
+				{
+					FString CurMap = TranslationStages[0].PassNamespace;
+					FString Value = FString::Printf(TEXT("Context.%s.%s = Context.%s.%s;\n"), *CurMap, *GetSanitizedSymbolName(Var.GetName().ToString()),
+						*CurMap, *GetSanitizedSymbolName(FNiagaraParameterMapHistory::GetSourceForPreviousValue(Var).GetName().ToString()));
+					HlslOutputString += Value + TEXT("\n");
+				}
+			}
+		}
+	}
+	HlslOutputString += TEXT("}\n#endif\n\n");
+}
+
 void FHlslNiagaraTranslator::DefineInterpolatedParametersFunction(FString &HlslOutputString)
 {
 	for (int32 i = 0; i < TranslationStages.Num(); i++)
diff --git a/Engine/Plugins/FX/Niagara/Source/NiagaraEditor/Private/NiagaraHlslTranslator.h b/Engine/Plugins/FX/Niagara/Source/NiagaraEditor/Private/NiagaraHlslTranslator.h
--- a/Engine/Plugins/FX/Niagara/Source/NiagaraEditor/Private/NiagaraHlslTranslator.h
+++ b/Engine/Plugins/FX/Niagara/Source/NiagaraEditor/Private/NiagaraHlslTranslator.h
@@ -496,6 +496,8 @@ public:
 	virtual void ParameterMapForEnd(class UNiagaraNodeParameterMapFor* ForNode);
 
 	void DefineInterpolatedParametersFunction(FString &HlslOutput);
+	void DefinePreviousParametersFunction(FString& HlslOutput, TArray<TArray<FNiagaraVariable>>& DataSetVariables, TMap<FNiagaraDataSetID, int32>& DataSetReads, TMap<FNiagaraDataSetID, int32>& DataSetWrites);
+
 	void DefineDataSetReadFunction(FString &HlslOutput, TArray<FNiagaraDataSetID> &ReadDataSets);
 	void DefineDataSetWriteFunction(FString &HlslOutput, TArray<FNiagaraDataSetProperties> &WriteDataSets, TArray<int32>& WriteConditionVarIndices);
 	void DefineMain(FString &HLSLOutput, const TArray<TArray<FNiagaraVariable>>& DataSetVariables, const TMap<FNiagaraDataSetID, int32>& DataSetReads, const TMap<FNiagaraDataSetID, int32>& DataSetWrites);
