Fixed shader compile error in NiagaraRayTracingShaders.usf due to using explicit resource slot assignment.

diff --git a/Engine/Plugins/FX/Niagara/Shaders/Private/NiagaraRayTracingShaders.usf b/Engine/Plugins/FX/Niagara/Shaders/Private/NiagaraRayTracingShaders.usf
--- a/Engine/Plugins/FX/Niagara/Shaders/Private/NiagaraRayTracingShaders.usf
+++ b/Engine/Plugins/FX/Niagara/Shaders/Private/NiagaraRayTracingShaders.usf
@@ -6,17 +6,9 @@
 #include "/Engine/Private/RayTracing/RayTracingHitGroupCommon.ush"
 #include "NiagaraRayTracingCommon.ush"
 
-#define RT_CONCATENATE2(a, b) a##b
-#define RT_CONCATENATE(a, b) RT_CONCATENATE2(a, b)
-#define RT_REGISTER(InType, InIndex, InSpace) register(RT_CONCATENATE(InType, InIndex), RT_CONCATENATE(space, InSpace))
-
-RaytracingAccelerationStructure               TLAS               : RT_REGISTER(t, 0, RAY_TRACING_REGISTER_SPACE_GLOBAL);
-StructuredBuffer<FBasicRayData>               Rays               : RT_REGISTER(t, 1, RAY_TRACING_REGISTER_SPACE_GLOBAL);
-RWStructuredBuffer<FNiagaraRayTracingPayload> CollisionOutput    : RT_REGISTER(u, 0, RAY_TRACING_REGISTER_SPACE_GLOBAL);
-
-#undef RT_REGISTER
-#undef RT_CONCATENATE
-#undef RT_CONCATENATE2
+RaytracingAccelerationStructure               TLAS;
+StructuredBuffer<FBasicRayData>               Rays;
+RWStructuredBuffer<FNiagaraRayTracingPayload> CollisionOutput;
 
 RAY_TRACING_ENTRY_RAYGEN(NiagaraCollisionRayTraceRG)
 {
