re-add correct boundaries for sampling tiled grids

diff --git a/Engine/Plugins/FX/Niagara/Source/Niagara/Private/NiagaraDataInterfaceGrid3DCollection.cpp b/Engine/Plugins/FX/Niagara/Source/Niagara/Private/NiagaraDataInterfaceGrid3DCollection.cpp
--- a/Engine/Plugins/FX/Niagara/Source/Niagara/Private/NiagaraDataInterfaceGrid3DCollection.cpp
+++ b/Engine/Plugins/FX/Niagara/Source/Niagara/Private/NiagaraDataInterfaceGrid3DCollection.cpp
@@ -1625,12 +1625,26 @@ bool UNiagaraDataInterfaceGrid3DCollection::GetFunctionHLSL(const FNiagaraDataIn
 					int TileIndexX = In_AttributeIndex % {NumTiles}.x;
 					int TileIndexY = (In_AttributeIndex / {NumTiles}.x) % {NumTiles}.y;
 					int TileIndexZ = In_AttributeIndex / ({NumTiles}.x * {NumTiles}.y);		
-					
-					float3 UVW = float3(In_UnitX, In_UnitY, In_UnitZ);
-					UVW = clamp(UVW, float3(0,0,0), float3(1,1,1));
 
-					UVW /= {NumTiles};
-					UVW += float3(1.0*TileIndexX, 1.0*TileIndexY, 1.0*TileIndexZ) / {NumTiles};
+					float3 UVW =
+					{
+						In_UnitX / {NumTiles}.x + 1.0*TileIndexX/{NumTiles}.x,
+						In_UnitY / {NumTiles}.y + 1.0*TileIndexY/{NumTiles}.y,
+						In_UnitZ / {NumTiles}.z + 1.0*TileIndexZ/{NumTiles}.y
+					};
+					float3 TileMin =
+					{
+						(TileIndexX * {NumCellsName}.x + 0.5) / ({NumTiles}.x * {NumCellsName}.x),
+						(TileIndexY * {NumCellsName}.y + 0.5) / ({NumTiles}.y * {NumCellsName}.y),
+						(TileIndexZ * {NumCellsName}.z + 0.5) / ({NumTiles}.z * {NumCellsName}.z)
+					};
+					float3 TileMax =
+					{
+						((TileIndexX + 1) * {NumCellsName}.x - 0.5) / ({NumTiles}.x * {NumCellsName}.x),
+						((TileIndexY + 1) * {NumCellsName}.y - 0.5) / ({NumTiles}.y * {NumCellsName}.y),
+						((TileIndexZ + 1) * {NumCellsName}.z - 0.5) / ({NumTiles}.z * {NumCellsName}.z)
+					};
+					UVW = clamp(UVW, TileMin, TileMax);
 
 					Out_Val = {Grid}.SampleLevel({SamplerName}, UVW, 0);
 				}
