Fix for niagara using RDG resources outside of their respective lifetimes.

diff --git a/Engine/Plugins/FX/Niagara/Source/Niagara/Private/DataInterface/NiagaraDataInterfaceGBuffer.cpp b/Engine/Plugins/FX/Niagara/Source/Niagara/Private/DataInterface/NiagaraDataInterfaceGBuffer.cpp
--- a/Engine/Plugins/FX/Niagara/Source/Niagara/Private/DataInterface/NiagaraDataInterfaceGBuffer.cpp
+++ b/Engine/Plugins/FX/Niagara/Source/Niagara/Private/DataInterface/NiagaraDataInterfaceGBuffer.cpp
@@ -98,7 +98,6 @@ struct FNiagaraDataInterfaceParametersCS_GBuffer : public FNiagaraDataInterfaceP
 public:
 	void Bind(const FNiagaraDataInterfaceGPUParamInfo& ParameterInfo, const class FShaderParameterMap& ParameterMap)
 	{
-		PassUniformBuffer.Bind(ParameterMap, FSceneTextureUniformParameters::StaticStructMetadata.GetShaderVariableName());
 		VelocityTextureParam.Bind(ParameterMap, TEXT("NDIGBuffer_VelocityTexture"));
 		VelocityTextureSamplerParam.Bind(ParameterMap, TEXT("NDIGBuffer_VelocityTextureSampler"));
 	}
@@ -108,18 +107,12 @@ public:
 		check(IsInRenderingThread());
 		FRHIComputeShader* ComputeShaderRHI = RHICmdList.GetBoundComputeShader();
 
-		//-Note: Scene textures will not exist in the Mobile rendering path
-		TUniformBufferRef<FSceneTextureUniformParameters> SceneTextureUniformParams = GNiagaraViewDataManager.GetSceneTextureUniformParameters();
-		check(!PassUniformBuffer.IsBound() || SceneTextureUniformParams);
-		SetUniformBufferParameter(RHICmdList, ComputeShaderRHI, PassUniformBuffer, SceneTextureUniformParams);
-
 		FRHISamplerState* VelocitySamplerState = TStaticSamplerState<SF_Bilinear, AM_Clamp, AM_Clamp, AM_Clamp>::GetRHI();
 		FRHITexture* VelocityTexture = GNiagaraViewDataManager.GetSceneVelocityTexture() ? GNiagaraViewDataManager.GetSceneVelocityTexture() : GBlackTexture->TextureRHI;
 		SetTextureParameter(RHICmdList, ComputeShaderRHI, VelocityTextureParam, VelocityTextureSamplerParam, VelocitySamplerState, VelocityTexture);
 	}
 
 private:
-	LAYOUT_FIELD(FShaderUniformBufferParameter, PassUniformBuffer);
 	LAYOUT_FIELD(FShaderResourceParameter, VelocityTextureParam);
 	LAYOUT_FIELD(FShaderResourceParameter, VelocityTextureSamplerParam);
 };
diff --git a/Engine/Plugins/FX/Niagara/Source/Niagara/Private/NiagaraDataInterfaceCollisionQuery.cpp b/Engine/Plugins/FX/Niagara/Source/Niagara/Private/NiagaraDataInterfaceCollisionQuery.cpp
--- a/Engine/Plugins/FX/Niagara/Source/Niagara/Private/NiagaraDataInterfaceCollisionQuery.cpp
+++ b/Engine/Plugins/FX/Niagara/Source/Niagara/Private/NiagaraDataInterfaceCollisionQuery.cpp
@@ -867,7 +867,6 @@ struct FNiagaraDataInterfaceParametersCS_CollisionQuery : public FNiagaraDataInt
 public:
 	void Bind(const FNiagaraDataInterfaceGPUParamInfo& ParameterInfo, const class FShaderParameterMap& ParameterMap)
 	{
-		PassUniformBuffer.Bind(ParameterMap, FSceneTextureUniformParameters::StaticStructMetadata.GetShaderVariableName());
 		GlobalDistanceFieldParameters.Bind(ParameterMap);
 #if RHI_RAYTRACING
 		MaxRayTraceCountParam.Bind(ParameterMap, TEXT("MaxRayTraceCount"));
@@ -882,11 +881,6 @@ public:
 
 		FNiagaraDataIntefaceProxyCollisionQuery* QueryDI = (FNiagaraDataIntefaceProxyCollisionQuery*)Context.DataInterface;
 		FRHIComputeShader* ComputeShaderRHI = Context.Shader.GetComputeShader();
-		
-		//-Note: Scene textures will not exist in the Mobile rendering path
-		TUniformBufferRef<FSceneTextureUniformParameters> SceneTextureUniformParams = GNiagaraViewDataManager.GetSceneTextureUniformParameters();
-		check(!PassUniformBuffer.IsBound() || SceneTextureUniformParams);
-		SetUniformBufferParameter(RHICmdList, ComputeShaderRHI, PassUniformBuffer/*Shader->GetUniformBufferParameter(SceneTexturesUniformBufferStruct)*/, SceneTextureUniformParams);
 
 		if (GlobalDistanceFieldParameters.IsBound() && Context.Batcher)
 		{
@@ -922,9 +916,6 @@ public:
 #endif
 
 private:
-	/** The SceneDepthTexture parameter for depth buffer collision. */
-	LAYOUT_FIELD(FShaderUniformBufferParameter, PassUniformBuffer);
-
 	LAYOUT_FIELD(FGlobalDistanceFieldParameters, GlobalDistanceFieldParameters);
 
 #if RHI_RAYTRACING
diff --git a/Engine/Plugins/FX/Niagara/Source/Niagara/Private/NiagaraDataInterfaceOcclusion.cpp b/Engine/Plugins/FX/Niagara/Source/Niagara/Private/NiagaraDataInterfaceOcclusion.cpp
--- a/Engine/Plugins/FX/Niagara/Source/Niagara/Private/NiagaraDataInterfaceOcclusion.cpp
+++ b/Engine/Plugins/FX/Niagara/Source/Niagara/Private/NiagaraDataInterfaceOcclusion.cpp
@@ -189,23 +189,11 @@ struct FNiagaraDataInterfaceParametersCS_OcclusionQuery : public FNiagaraDataInt
 public:
 	void Bind(const FNiagaraDataInterfaceGPUParamInfo& ParameterInfo, const class FShaderParameterMap& ParameterMap)
 	{
-		PassUniformBuffer.Bind(ParameterMap, FSceneTextureUniformParameters::StaticStructMetadata.GetShaderVariableName());
 	}
 
 	void Set(FRHICommandList& RHICmdList, const FNiagaraDataInterfaceSetArgs& Context) const
 	{
-		check(IsInRenderingThread());
-		FRHIComputeShader* ComputeShaderRHI = RHICmdList.GetBoundComputeShader();
-
-		//-Note: Scene textures will not exist in the Mobile rendering path
-		TUniformBufferRef<FSceneTextureUniformParameters> SceneTextureUniformParams = GNiagaraViewDataManager.GetSceneTextureUniformParameters();
-		check(!PassUniformBuffer.IsBound() || SceneTextureUniformParams);
-		SetUniformBufferParameter(RHICmdList, ComputeShaderRHI, PassUniformBuffer, SceneTextureUniformParams);
 	}
-
-private:
-	/** The SceneDepthTexture parameter for depth buffer query. */
-	LAYOUT_FIELD(FShaderUniformBufferParameter, PassUniformBuffer);
 };
 
 IMPLEMENT_TYPE_LAYOUT(FNiagaraDataInterfaceParametersCS_OcclusionQuery);
diff --git a/Engine/Plugins/FX/Niagara/Source/Niagara/Private/NiagaraEmitterInstanceBatcher.cpp b/Engine/Plugins/FX/Niagara/Source/Niagara/Private/NiagaraEmitterInstanceBatcher.cpp
--- a/Engine/Plugins/FX/Niagara/Source/Niagara/Private/NiagaraEmitterInstanceBatcher.cpp
+++ b/Engine/Plugins/FX/Niagara/Source/Niagara/Private/NiagaraEmitterInstanceBatcher.cpp
@@ -505,6 +505,7 @@ void NiagaraEmitterInstanceBatcher::ProcessPendingTicksFlush(FRHICommandListImme
 			
 			// Execute all ticks that we can support without invalid simulations
 			FRDGBuilder GraphBuilder(RHICmdList);
+			CreateSystemTextures(GraphBuilder);
 			PreInitViews(GraphBuilder, true);
 			PostInitViews(GraphBuilder, DummyViews, AllowGPUParticleUpdate);
 			PostRenderOpaque(GraphBuilder, DummyViews, AllowGPUParticleUpdate);
@@ -1305,13 +1306,6 @@ void NiagaraEmitterInstanceBatcher::ExecuteAll(FRHICommandList& RHICmdList, FRHI
 
 	SCOPED_DRAW_EVENTF(RHICmdList, NiagaraEmitterInstanceBatcher_ExecuteAll, TEXT("NiagaraEmitterInstanceBatcher_ExecuteAll - TickStage(%d)"), TickStage);
 
-	FUniformBufferStaticBindings StaticUniformBuffers;
-	if (FRHIUniformBuffer* SceneTexturesUniformBuffer = GNiagaraViewDataManager.GetSceneTextureUniformParameters())
-	{
-		StaticUniformBuffers.AddUniformBuffer(SceneTexturesUniformBuffer);
-	}
-	SCOPED_UNIFORM_BUFFER_STATIC_BINDINGS(RHICmdList, StaticUniformBuffers);
-
 	FMemMark Mark(FMemStack::Get());
 	TArray<FOverlappableTicks, TMemStackAllocator<> > SimPasses;
 	{
@@ -1441,6 +1435,7 @@ void NiagaraEmitterInstanceBatcher::ExecuteAll(FRHICommandList& RHICmdList, FRHI
 
 void NiagaraEmitterInstanceBatcher::PreInitViews(FRDGBuilder& GraphBuilder, bool bAllowGPUParticleUpdate)
 {
+	GNiagaraViewDataManager.ClearSceneTextureParameters();
 	FramesBeforeTickFlush = 0;
 
 	GpuReadbackManagerPtr->Tick();
@@ -1503,15 +1498,19 @@ void NiagaraEmitterInstanceBatcher::PreInitViews(FRDGBuilder& GraphBuilder, bool
 
 		if (GNiagaraAllowTickBeforeRender)
 		{
-			AddPass(
-				GraphBuilder,
+			FNiagaraSceneTextureParameters* PassParameters = GraphBuilder.AllocParameters<FNiagaraSceneTextureParameters>();
+			PassParameters->SceneTextures = CreateSceneTextureUniformBuffer(GraphBuilder, ERHIFeatureLevel::SM5, ESceneTextureSetupMode::None);
+
+			GraphBuilder.AddPass(
 				RDG_EVENT_NAME("Niagara::PreInitViews"),
+				PassParameters,
+				ERDGPassFlags::Compute | ERDGPassFlags::NeverCull,
 				[this](FRHICommandListImmediate& RHICmdList)
-				{
-					BuildTickStagePasses(RHICmdList, ETickStage::PreInitViews);
-					ExecuteAll(RHICmdList, nullptr, ETickStage::PreInitViews);
-				}
-			);
+			{
+				FNiagaraSceneTextureScope Scope;
+				BuildTickStagePasses(RHICmdList, ETickStage::PreInitViews);
+				ExecuteAll(RHICmdList, nullptr, ETickStage::PreInitViews);
+			});
 		}
 	}
 	else
@@ -1530,16 +1529,19 @@ void NiagaraEmitterInstanceBatcher::PostInitViews(FRDGBuilder& GraphBuilder, TAr
 	if (bAllowGPUParticleUpdate && FNiagaraUtilities::AllowGPUParticles(GetShaderPlatform()))
 	{
 		FRHIUniformBuffer* ViewUniformBuffer = GetReferenceViewUniformBuffer(Views);
+		FNiagaraSceneTextureParameters* PassParameters = GraphBuilder.AllocParameters<FNiagaraSceneTextureParameters>();
+		PassParameters->SceneTextures = CreateSceneTextureUniformBuffer(GraphBuilder, ERHIFeatureLevel::SM5, ESceneTextureSetupMode::None);
 
-		AddPass(
-			GraphBuilder,
+		GraphBuilder.AddPass(
 			RDG_EVENT_NAME("Niagara::PostInitViews"),
+			PassParameters,
+			ERDGPassFlags::Compute | ERDGPassFlags::NeverCull,
 			[this, ViewUniformBuffer](FRHICommandListImmediate& RHICmdList)
-			{
-				BuildTickStagePasses(RHICmdList, ETickStage::PostInitViews);
-				ExecuteAll(RHICmdList, ViewUniformBuffer, ETickStage::PostInitViews);
-			}
-		);
+		{
+			FNiagaraSceneTextureScope Scope;
+			BuildTickStagePasses(RHICmdList, ETickStage::PostInitViews);
+			ExecuteAll(RHICmdList, ViewUniformBuffer, ETickStage::PostInitViews);
+		});
 	}
 }
 
@@ -1549,42 +1551,49 @@ void NiagaraEmitterInstanceBatcher::PostRenderOpaque(FRDGBuilder& GraphBuilder,
 
 	bAllowGPUParticleUpdate = bAllowGPUParticleUpdate && GetReferenceAllowGPUUpdate(Views);
 
-	AddPass(
-		GraphBuilder,
+	FNiagaraSceneTextureParameters* PassParameters = GraphBuilder.AllocParameters<FNiagaraSceneTextureParameters>();
+	GNiagaraViewDataManager.GetSceneTextureParameters(GraphBuilder, *PassParameters);
+
+	GraphBuilder.AddPass(
 		RDG_EVENT_NAME("Niagara::PostRenderOpaque"),
+		PassParameters,
+		ERDGPassFlags::Compute | ERDGPassFlags::NeverCull,
 		[this, Views, bAllowGPUParticleUpdate](FRHICommandListImmediate& RHICmdList)
+	{
+		FNiagaraSceneTextureScope Scope;
+
+		if (bAllowGPUParticleUpdate && FNiagaraUtilities::AllowGPUParticles(GetShaderPlatform()))
 		{
-			if (bAllowGPUParticleUpdate && FNiagaraUtilities::AllowGPUParticles(GetShaderPlatform()))
-			{
 #if RHI_RAYTRACING
-				BuildRayTracingSceneInfo(RHICmdList, Views);
+			BuildRayTracingSceneInfo(RHICmdList, Views);
 #endif
 
-				FRHIUniformBuffer* ViewUniformBuffer = GetReferenceViewUniformBuffer(Views);
+			FRHIUniformBuffer* ViewUniformBuffer = GetReferenceViewUniformBuffer(Views);
 
-				BuildTickStagePasses(RHICmdList, ETickStage::PostOpaqueRender);
+			BuildTickStagePasses(RHICmdList, ETickStage::PostOpaqueRender);
 
-				// Setup new readback since if there is no pending request, there is no risk of having invalid data read (offset being allocated after the readback was sent).
-				ExecuteAll(RHICmdList, ViewUniformBuffer, ETickStage::PostOpaqueRender);
+			// Setup new readback since if there is no pending request, there is no risk of having invalid data read (offset being allocated after the readback was sent).
+			ExecuteAll(RHICmdList, ViewUniformBuffer, ETickStage::PostOpaqueRender);
 
-				UpdateFreeIDBuffers(RHICmdList, DeferredIDBufferUpdates);
-				DeferredIDBufferUpdates.SetNum(0, false);
+			UpdateFreeIDBuffers(RHICmdList, DeferredIDBufferUpdates);
+			DeferredIDBufferUpdates.SetNum(0, false);
 
-				FinishDispatches();
+			FinishDispatches();
 
-				ProcessDebugReadbacks(RHICmdList, false);
+			ProcessDebugReadbacks(RHICmdList, false);
 
 #if RHI_RAYTRACING
-				ResetRayTracingSceneInfo();
+			ResetRayTracingSceneInfo();
 #endif
-			}
+		}
 
-			if (!GPUInstanceCounterManager.HasPendingGPUReadback())
-			{
-				GPUInstanceCounterManager.EnqueueGPUReadback(RHICmdList);
-			}
+		if (!GPUInstanceCounterManager.HasPendingGPUReadback())
+		{
+			GPUInstanceCounterManager.EnqueueGPUReadback(RHICmdList);
 		}
-	);
+	});
+
+	GNiagaraViewDataManager.ClearSceneTextureParameters();
 }
 
 void NiagaraEmitterInstanceBatcher::ProcessDebugReadbacks(FRHICommandListImmediate& RHICmdList, bool bWaitCompletion)
diff --git a/Engine/Plugins/FX/Niagara/Source/Niagara/Private/NiagaraWorldManager.cpp b/Engine/Plugins/FX/Niagara/Source/Niagara/Private/NiagaraWorldManager.cpp
--- a/Engine/Plugins/FX/Niagara/Source/Niagara/Private/NiagaraWorldManager.cpp
+++ b/Engine/Plugins/FX/Niagara/Source/Niagara/Private/NiagaraWorldManager.cpp
@@ -201,43 +201,44 @@ namespace FNiagaraUtilities
 	}
 }
 
-BEGIN_SHADER_PARAMETER_STRUCT(FExtractUniformBufferParameters, )
-	SHADER_PARAMETER_RDG_UNIFORM_BUFFER(FSceneTextureUniformParameters, SceneTextures)
-	RDG_TEXTURE_ACCESS(Depth, ERHIAccess::SRVGraphics)
-	RDG_TEXTURE_ACCESS(Normal, ERHIAccess::SRVGraphics)
-	RDG_TEXTURE_ACCESS(Velocity, ERHIAccess::SRVGraphics)
-END_SHADER_PARAMETER_STRUCT()
-
 void FNiagaraViewDataMgr::PostOpaqueRender(FPostOpaqueRenderParameters& Params)
 {
-	FRDGBuilder& GraphBuilder = *Params.GraphBuilder;
-
 	ViewUniformBuffer = Params.ViewUniformBuffer;
+	Parameters.SceneTextures = Params.SceneTexturesUniformParams;
+	Parameters.Depth = Params.DepthTexture;
+	Parameters.Normal = Params.NormalTexture;
+	Parameters.Velocity = Params.VelocityTexture;
+}
 
-	FExtractUniformBufferParameters* PassParameters = GraphBuilder.AllocParameters<FExtractUniformBufferParameters>();
-	PassParameters->SceneTextures = Params.SceneTexturesUniformParams;
-	PassParameters->Depth = Params.DepthTexture;
-	PassParameters->Normal = Params.NormalTexture;
-	PassParameters->Velocity = Params.VelocityTexture;
-
-	// Niagara is not ported to RDG and is manually accessing RHI resources. This extraction is technically not safe,
-	// because the lifetime of the textures are not guaranteed outside of the pass. However, we happen to know that
-	// they will outlast Niagara.
+void FNiagaraViewDataMgr::GetSceneTextureParameters(FRDGBuilder& GraphBuilder, FNiagaraSceneTextureParameters& OutParameters) const
+{
+	OutParameters = Parameters;
 
-	GraphBuilder.AddPass({}, PassParameters, ERDGPassFlags::Raster | ERDGPassFlags::SkipRenderPass | ERDGPassFlags::NeverCull, [this, PassParameters](FRHICommandList&)
+	if (!OutParameters.SceneTextures)
 	{
-		SceneTexturesUniformParams = PassParameters->SceneTextures->GetRHIRef();
-		SceneDepthTexture = static_cast<FRHITexture2D*>(PassParameters->Depth ? PassParameters->Depth->GetRHI() : nullptr);
-		SceneNormalTexture = static_cast<FRHITexture2D*>(PassParameters->Normal ? PassParameters->Normal->GetRHI() : nullptr);
-		SceneVelocityTexture = static_cast<FRHITexture2D*>(PassParameters->Velocity ? PassParameters->Velocity->GetRHI() : nullptr);
-	});
+		OutParameters.SceneTextures = CreateSceneTextureUniformBuffer(GraphBuilder, ERHIFeatureLevel::SM5, ESceneTextureSetupMode::None);
+	}
+}
+
+void FNiagaraViewDataMgr::BeginPass()
+{
+	check(!bInsidePass);
+	bInsidePass = true;
+}
+
+void FNiagaraViewDataMgr::EndPass()
+{
+	check(bInsidePass);
+	bInsidePass = false;
+}
+
+void FNiagaraViewDataMgr::ClearSceneTextureParameters()
+{
+	Parameters = {};
 }
 
 FNiagaraViewDataMgr::FNiagaraViewDataMgr()
 	: FRenderResource()
-	, SceneDepthTexture(nullptr)
-	, SceneNormalTexture(nullptr)
-	, ViewUniformBuffer(nullptr)
 {
 
 }
@@ -265,10 +266,6 @@ void FNiagaraViewDataMgr::InitDynamicRHI()
 
 void FNiagaraViewDataMgr::ReleaseDynamicRHI()
 {
-	SceneDepthTexture = nullptr;
-	SceneNormalTexture = nullptr;
-	ViewUniformBuffer = nullptr;
-	SceneTexturesUniformParams.SafeRelease();
 }
 
 //////////////////////////////////////////////////////////////////////////
diff --git a/Engine/Plugins/FX/Niagara/Source/Niagara/Public/NiagaraWorldManager.h b/Engine/Plugins/FX/Niagara/Source/Niagara/Public/NiagaraWorldManager.h
--- a/Engine/Plugins/FX/Niagara/Source/Niagara/Public/NiagaraWorldManager.h
+++ b/Engine/Plugins/FX/Niagara/Source/Niagara/Public/NiagaraWorldManager.h
@@ -29,6 +29,13 @@ class UNiagaraParameterCollectionInstance;
 class UNiagaraComponentPool;
 struct FNiagaraScalabilityState;
 
+BEGIN_SHADER_PARAMETER_STRUCT(FNiagaraSceneTextureParameters, )
+	SHADER_PARAMETER_RDG_UNIFORM_BUFFER(FSceneTextureUniformParameters, SceneTextures)
+	RDG_TEXTURE_ACCESS(Depth, ERHIAccess::SRVCompute)
+	RDG_TEXTURE_ACCESS(Normal, ERHIAccess::SRVCompute)
+	RDG_TEXTURE_ACCESS(Velocity, ERHIAccess::SRVCompute)
+END_SHADER_PARAMETER_STRUCT()
+
 class FNiagaraViewDataMgr : public FRenderResource
 {
 public:
@@ -39,29 +46,46 @@ public:
 
 	NIAGARA_API void PostOpaqueRender(FPostOpaqueRenderParameters& Params);
 
-	FRHITexture2D* GetSceneDepthTexture() { return SceneDepthTexture; }
-	FRHITexture2D* GetSceneNormalTexture() { return SceneNormalTexture; }
-	FRHITexture2D* GetSceneVelocityTexture() { return SceneVelocityTexture; }
+	void BeginPass();
+	void EndPass();
+	void ClearSceneTextureParameters();
+
+	void GetSceneTextureParameters(FRDGBuilder& GraphBuilder, FNiagaraSceneTextureParameters& InParameters) const;
+
+	FRHITexture2D* GetSceneDepthTexture() { check(bInsidePass); return (FRHITexture2D*)TryGetRHI(Parameters.Depth); }
+	FRHITexture2D* GetSceneNormalTexture() { check(bInsidePass); return (FRHITexture2D*)TryGetRHI(Parameters.Normal); }
+	FRHITexture2D* GetSceneVelocityTexture() { check(bInsidePass); return (FRHITexture2D*)TryGetRHI(Parameters.Velocity); }
 	FRHIUniformBuffer* GetViewUniformBuffer() { return ViewUniformBuffer; }
-	TUniformBufferRef<FSceneTextureUniformParameters> GetSceneTextureUniformParameters() { return SceneTexturesUniformParams; }
 
 	virtual void InitDynamicRHI() override;
 
 	virtual void ReleaseDynamicRHI() override;
 
 private:
-	FRHITexture2D* SceneDepthTexture = nullptr;
-	FRHITexture2D* SceneNormalTexture = nullptr;
-	FRHITexture2D* SceneVelocityTexture = nullptr;
+	FNiagaraSceneTextureParameters Parameters;
 	FRHIUniformBuffer* ViewUniformBuffer = nullptr;
 
-	TUniformBufferRef<FSceneTextureUniformParameters> SceneTexturesUniformParams;
 	FPostOpaqueRenderDelegate PostOpaqueDelegate;
 	FDelegateHandle PostOpaqueDelegateHandle;
+	bool bInsidePass = false;
 };
 
 extern TGlobalResource<FNiagaraViewDataMgr> GNiagaraViewDataManager;
 
+class FNiagaraSceneTextureScope
+{
+public:
+	FNiagaraSceneTextureScope()
+	{
+		GNiagaraViewDataManager.BeginPass();
+	}
+
+	~FNiagaraSceneTextureScope()
+	{
+		GNiagaraViewDataManager.EndPass();
+	}
+};
+
 USTRUCT()
 struct FNiagaraWorldManagerTickFunction : public FTickFunction
 {
diff --git a/Engine/Shaders/Public/ShaderVersion.ush b/Engine/Shaders/Public/ShaderVersion.ush
--- a/Engine/Shaders/Public/ShaderVersion.ush
+++ b/Engine/Shaders/Public/ShaderVersion.ush
@@ -2,4 +2,4 @@
 // Each time the console command is executed it generates a new GUID. As this file is included
 // in Platform.ush (which should be included in any shader) it allows to invalidate the shader DDC.
 // 
-// GUID = 4B4392B38CAB46E39C0C10B2CA18FF52_A
+// GUID = FB1C97C486D94B569EE3D3A099F20C42
diff --git a/Engine/Source/Runtime/Renderer/Private/SceneTextures.cpp b/Engine/Source/Runtime/Renderer/Private/SceneTextures.cpp
--- a/Engine/Source/Runtime/Renderer/Private/SceneTextures.cpp
+++ b/Engine/Source/Runtime/Renderer/Private/SceneTextures.cpp
@@ -1209,4 +1209,9 @@ FIntPoint GetSceneTextureExtent()
 ERHIFeatureLevel::Type GetSceneTextureFeatureLevel()
 {
 	return FSceneTexturesConfig::Get().FeatureLevel;
+}
+
+void CreateSystemTextures(FRDGBuilder& GraphBuilder)
+{
+	FRDGSystemTextures::Create(GraphBuilder);
 }
\ No newline at end of file
diff --git a/Engine/Source/Runtime/Renderer/Public/SceneRenderTargetParameters.h b/Engine/Source/Runtime/Renderer/Public/SceneRenderTargetParameters.h
--- a/Engine/Source/Runtime/Renderer/Public/SceneRenderTargetParameters.h
+++ b/Engine/Source/Runtime/Renderer/Public/SceneRenderTargetParameters.h
@@ -202,6 +202,9 @@ extern RENDERER_API ERHIFeatureLevel::Type GetSceneTextureFeatureLevel();
  */
 extern RENDERER_API void ResetSceneTextureExtentHistory();
 
+/** Registers system textures into RDG. */
+extern RENDERER_API void CreateSystemTextures(FRDGBuilder& GraphBuilder);
+
 ///////////////////////////////////////////////////////////////////////////
 // Deprecated APIs
 
