Fix for shader compile errors related to Niagara raytracing shaders using const IsHit()

diff --git a/Engine/Plugins/FX/Niagara/Shaders/Private/NiagaraDataInterfaceCollisionQuery.ush b/Engine/Plugins/FX/Niagara/Shaders/Private/NiagaraDataInterfaceCollisionQuery.ush
--- a/Engine/Plugins/FX/Niagara/Shaders/Private/NiagaraDataInterfaceCollisionQuery.ush
+++ b/Engine/Plugins/FX/Niagara/Shaders/Private/NiagaraDataInterfaceCollisionQuery.ush
@@ -115,12 +115,14 @@ void DICollisionQuery_ReadAsyncRayTrace(int In_PreviousFrameQueryID, out bool Ou
 #if FEATURE_LEVEL >= FEATURE_LEVEL_SM5
 	if (In_PreviousFrameQueryID >= 0 && In_PreviousFrameQueryID < MaxRayTraceCount)
 	{
-		Out_CollisionValid = IntersectionResults[In_PreviousFrameQueryID].IsHit();
+		FNiagaraRayTracingPayload HitResult = IntersectionResults[In_PreviousFrameQueryID];
+
+		Out_CollisionValid = HitResult.IsHit();
 		if (Out_CollisionValid)
 		{
-			Out_CollisionDistance = IntersectionResults[In_PreviousFrameQueryID].HitT;
-			Out_CollisionPosWorld = IntersectionResults[In_PreviousFrameQueryID].WorldPosition;
-			Out_CollisionNormal = IntersectionResults[In_PreviousFrameQueryID].WorldNormal;
+			Out_CollisionDistance = HitResult.HitT;
+			Out_CollisionPosWorld = HitResult.WorldPosition;
+			Out_CollisionNormal = HitResult.WorldNormal;
 		}
 	}
 #endif
