Move Niagara's access to the raytracing scene through the FScene rather than the FViewInfo.  Migrates expoirted API functions to supporrt this as well.

diff --git a/Engine/Plugins/FX/Niagara/Source/Niagara/Private/NiagaraRayTracingHelper.cpp b/Engine/Plugins/FX/Niagara/Source/Niagara/Private/NiagaraRayTracingHelper.cpp
--- a/Engine/Plugins/FX/Niagara/Source/Niagara/Private/NiagaraRayTracingHelper.cpp
+++ b/Engine/Plugins/FX/Niagara/Source/Niagara/Private/NiagaraRayTracingHelper.cpp
@@ -6,6 +6,8 @@
 
 #include "GlobalShader.h"
 #include "MeshPassProcessor.h"
+#include "Renderer/Private/RayTracing/RayTracingScene.h"
+#include "Renderer/Private/ScenePrivate.h"
 #include "Renderer/Private/SceneRendering.h"
 #include "ShaderParameterStruct.h"
 
@@ -158,11 +160,15 @@ static void BindNiagaraRayTracingMeshCommands(
 
 void FNiagaraRayTracingHelper::BuildRayTracingSceneInfo(FRHICommandList& RHICmdList, TConstArrayView<FViewInfo> Views)
 {
+	check(Views.Num() > 0);
+
 	const FViewInfo& ReferenceView = Views[0];
-	if (ReferenceView.HasRayTracingScene())
+	const FScene* RenderScene = (ReferenceView.Family && ReferenceView.Family->Scene) ? ReferenceView.Family->Scene->GetRenderScene() : nullptr;
+
+	if (RenderScene && RenderScene->RayTracingScene.IsCreated())
 	{
-		RayTracingScene = ReferenceView.GetRayTracingSceneChecked();
-		RayTracingSceneView = ReferenceView.GetRayTracingSceneViewChecked();
+		RayTracingScene = RenderScene->RayTracingScene.GetRHIRayTracingSceneChecked();
+		RayTracingSceneView = RenderScene->RayTracingScene.GetShaderResourceViewChecked();
 		ViewUniformBuffer = ReferenceView.ViewUniformBuffer;
 
 		FGlobalShaderMap* ShaderMap = GetGlobalShaderMap(ShaderPlatform);
diff --git a/Engine/Plugins/FX/Niagara/Source/Niagara/Private/NiagaraRayTracingHelper.h b/Engine/Plugins/FX/Niagara/Source/Niagara/Private/NiagaraRayTracingHelper.h
--- a/Engine/Plugins/FX/Niagara/Source/Niagara/Private/NiagaraRayTracingHelper.h
+++ b/Engine/Plugins/FX/Niagara/Source/Niagara/Private/NiagaraRayTracingHelper.h
@@ -3,6 +3,8 @@
 
 #include "CoreMinimal.h"
 
+#include "RHIDefinitions.h"
+
 #if RHI_RAYTRACING
 
 class FViewInfo;
diff --git a/Engine/Source/Runtime/Renderer/Private/RayTracing/RayTracingScene.h b/Engine/Source/Runtime/Renderer/Private/RayTracing/RayTracingScene.h
--- a/Engine/Source/Runtime/Renderer/Private/RayTracing/RayTracingScene.h
+++ b/Engine/Source/Runtime/Renderer/Private/RayTracing/RayTracingScene.h
@@ -40,18 +40,18 @@ public:
 
 	// Returns true if RHI ray tracing scene has been created or asynchronous creation has been kicked.
 	// i.e. returns true after BeginCreate() and before Reset().
-	bool IsCreated() const;
+	RENDERER_API bool IsCreated() const;
 
 	// Waits for the async creation task if it's active and then returns RayTracingSceneRHI object (may return null).
-	FRHIRayTracingScene* GetRHIRayTracingScene() const;
+	RENDERER_API  FRHIRayTracingScene* GetRHIRayTracingScene() const;
 
 	// Similar to GetRayTracingScene, but checks that ray tracing scene RHI object is valid.
-	FRHIRayTracingScene* GetRHIRayTracingSceneChecked() const;
+	RENDERER_API  FRHIRayTracingScene* GetRHIRayTracingSceneChecked() const;
 
 	// Returns Buffer and SRV for this ray tracing scene.
 	// Valid to call immediately after BeginCreate() and does not block.
-	FRHIShaderResourceView* GetShaderResourceViewChecked() const;
-	FRHIBuffer* GetBufferChecked() const;
+	RENDERER_API FRHIShaderResourceView* GetShaderResourceViewChecked() const;
+	RENDERER_API FRHIBuffer* GetBufferChecked() const;
 
 public:
 
diff --git a/Engine/Source/Runtime/Renderer/Private/SceneRendering.h b/Engine/Source/Runtime/Renderer/Private/SceneRendering.h
--- a/Engine/Source/Runtime/Renderer/Private/SceneRendering.h
+++ b/Engine/Source/Runtime/Renderer/Private/SceneRendering.h
@@ -1400,10 +1400,10 @@ public:
 	FLumenTranslucencyGIVolume LumenTranslucencyGIVolume;
 
 #if RHI_RAYTRACING
-	RENDERER_API bool HasRayTracingScene() const;
-	RENDERER_API FRHIRayTracingScene* GetRayTracingScene() const; // Soft-deprecated method, use FScene.RayTracingScene instead.
-	RENDERER_API FRHIRayTracingScene* GetRayTracingSceneChecked() const; // Soft-deprecated method, use FScene.RayTracingScene instead.
-	RENDERER_API FRHIShaderResourceView* GetRayTracingSceneViewChecked() const; // Soft-deprecated method, use FScene.RayTracingScene instead.
+	bool HasRayTracingScene() const;
+	FRHIRayTracingScene* GetRayTracingScene() const; // Soft-deprecated method, use FScene.RayTracingScene instead.
+	FRHIRayTracingScene* GetRayTracingSceneChecked() const; // Soft-deprecated method, use FScene.RayTracingScene instead.
+	FRHIShaderResourceView* GetRayTracingSceneViewChecked() const; // Soft-deprecated method, use FScene.RayTracingScene instead.
 
 	// Primary pipeline state object to be used with the ray tracing scene for this view.
 	// Material shaders are only available when using this pipeline.
