Ensure that we don't worry about trying to optimize the Niagara bytecode unless it is valid.  Out of date code could have gotten to the VM (for standalone scripts during cooking) which can result in crashes.

diff --git a/Engine/Plugins/FX/Niagara/Source/Niagara/Private/NiagaraScript.cpp b/Engine/Plugins/FX/Niagara/Source/Niagara/Private/NiagaraScript.cpp
--- a/Engine/Plugins/FX/Niagara/Source/Niagara/Private/NiagaraScript.cpp
+++ b/Engine/Plugins/FX/Niagara/Source/Niagara/Private/NiagaraScript.cpp
@@ -765,7 +765,7 @@ TOptional<ENiagaraSimTarget> UNiagaraScript::GetSimTarget() const
 
 void UNiagaraScript::AsyncOptimizeByteCode()
 {
-	if ( !CachedScriptVM.IsValid() || (CachedScriptVM.OptimizedByteCode.Num() > 0) || (CachedScriptVM.ByteCode.Num() == 0) )
+	if ( !CachedScriptVMId.IsValid() || !CachedScriptVM.IsValid() || (CachedScriptVM.OptimizedByteCode.Num() > 0) || (CachedScriptVM.ByteCode.Num() == 0) )
 	{
 		return;
 	}
