Fix for crash due to binding deleted View uniform buffer.

diff --git a/Engine/Plugins/FX/Niagara/Source/Niagara/Private/NiagaraEmitterInstanceBatcher.cpp b/Engine/Plugins/FX/Niagara/Source/Niagara/Private/NiagaraEmitterInstanceBatcher.cpp
--- a/Engine/Plugins/FX/Niagara/Source/Niagara/Private/NiagaraEmitterInstanceBatcher.cpp
+++ b/Engine/Plugins/FX/Niagara/Source/Niagara/Private/NiagaraEmitterInstanceBatcher.cpp
@@ -1528,7 +1528,6 @@ void NiagaraEmitterInstanceBatcher::PostInitViews(FRDGBuilder& GraphBuilder, TAr
 
 	if (bAllowGPUParticleUpdate && FNiagaraUtilities::AllowGPUParticles(GetShaderPlatform()))
 	{
-		FRHIUniformBuffer* ViewUniformBuffer = GetReferenceViewUniformBuffer(Views);
 		FNiagaraSceneTextureParameters* PassParameters = GraphBuilder.AllocParameters<FNiagaraSceneTextureParameters>();
 		PassParameters->SceneTextures = CreateSceneTextureUniformBuffer(GraphBuilder, ERHIFeatureLevel::SM5, ESceneTextureSetupMode::None);
 
@@ -1536,8 +1535,9 @@ void NiagaraEmitterInstanceBatcher::PostInitViews(FRDGBuilder& GraphBuilder, TAr
 			RDG_EVENT_NAME("Niagara::PostInitViews"),
 			PassParameters,
 			ERDGPassFlags::Compute | ERDGPassFlags::NeverCull,
-			[this, ViewUniformBuffer](FRHICommandListImmediate& RHICmdList)
+			[this, Views](FRHICommandListImmediate& RHICmdList)
 		{
+			FRHIUniformBuffer* ViewUniformBuffer = GetReferenceViewUniformBuffer(Views);
 			FNiagaraSceneTextureScope Scope;
 			BuildTickStagePasses(RHICmdList, ETickStage::PostInitViews);
 			ExecuteAll(RHICmdList, ViewUniformBuffer, ETickStage::PostInitViews);
