Corrected a forward declare.

#rnx

diff --git a/Engine/Plugins/FX/Niagara/Source/Niagara/Private/NiagaraRayTracingHelper.h b/Engine/Plugins/FX/Niagara/Source/Niagara/Private/NiagaraRayTracingHelper.h
--- a/Engine/Plugins/FX/Niagara/Source/Niagara/Private/NiagaraRayTracingHelper.h
+++ b/Engine/Plugins/FX/Niagara/Source/Niagara/Private/NiagaraRayTracingHelper.h
@@ -9,7 +9,7 @@
 
 class FViewInfo;
 class FRayTracingPipelineState;
-class FRHICommandListImmediate;
+class FRHICommandList;
 class FRHIRayTracingScene;
 class FRHIShaderResourceView;
 class FRHIUniformBuffer;
