Fixes to the niagara instance batcher dealing with flushing built up ticks

diff --git a/Engine/Plugins/FX/Niagara/Source/Niagara/Private/NiagaraEmitterInstanceBatcher.cpp b/Engine/Plugins/FX/Niagara/Source/Niagara/Private/NiagaraEmitterInstanceBatcher.cpp
--- a/Engine/Plugins/FX/Niagara/Source/Niagara/Private/NiagaraEmitterInstanceBatcher.cpp
+++ b/Engine/Plugins/FX/Niagara/Source/Niagara/Private/NiagaraEmitterInstanceBatcher.cpp
@@ -7,6 +7,7 @@
 #include "ClearQuad.h"
 #include "Engine/Canvas.h"
 #include "GPUSort.h"
+#include "Misc/ScopeExit.h"
 #include "NiagaraDataInterfaceRW.h"
 #if WITH_EDITOR
 #include "NiagaraGpuComputeDebug.h"
@@ -501,12 +502,12 @@ void NiagaraEmitterInstanceBatcher::ProcessPendingTicksFlush(FRHICommandListImme
 			DummyView->ViewUniformBuffer = TUniformBufferRef<FViewUniformShaderParameters>::CreateUniformBufferImmediate(*DummyView->CachedViewUniformShaderParameters, UniformBuffer_SingleFrame);
 
 			TConstArrayView<FViewInfo> DummyViews = MakeArrayView(DummyView, 1);
-			const bool AllowGPUParticleUpdate = false;
+			const bool AllowGPUParticleUpdate = true;
 			
 			// Execute all ticks that we can support without invalid simulations
 			FRDGBuilder GraphBuilder(RHICmdList);
 			CreateSystemTextures(GraphBuilder);
-			PreInitViews(GraphBuilder, true);
+			PreInitViews(GraphBuilder, AllowGPUParticleUpdate);
 			PostInitViews(GraphBuilder, DummyViews, AllowGPUParticleUpdate);
 			PostRenderOpaque(GraphBuilder, DummyViews, AllowGPUParticleUpdate);
 			GraphBuilder.Execute();
@@ -1491,7 +1492,7 @@ void NiagaraEmitterInstanceBatcher::PreInitViews(FRDGBuilder& GraphBuilder, bool
 	// (note that currently there are no callback appropriate to do it)
 	SimulationsToSort.Reset();
 
-	// Update draw indirect buffer to max possible size.
+
 	if (bAllowGPUParticleUpdate && FNiagaraUtilities::AllowGPUParticles(GetShaderPlatform()))
 	{
 		UpdateInstanceCountManager(GraphBuilder.RHICmdList);
@@ -1566,6 +1567,10 @@ void NiagaraEmitterInstanceBatcher::PostRenderOpaque(FRDGBuilder& GraphBuilder,
 		{
 #if RHI_RAYTRACING
 			BuildRayTracingSceneInfo(RHICmdList, Views);
+			ON_SCOPE_EXIT
+			{
+				ResetRayTracingSceneInfo();
+			};
 #endif
 
 			FRHIUniformBuffer* ViewUniformBuffer = GetReferenceViewUniformBuffer(Views);
