Integrate 15823512 and 15823993:

diff --git a/Engine/Plugins/FX/Niagara/Source/Niagara/Private/NiagaraDataInterfaceGrid2DCollection.cpp b/Engine/Plugins/FX/Niagara/Source/Niagara/Private/NiagaraDataInterfaceGrid2DCollection.cpp
--- a/Engine/Plugins/FX/Niagara/Source/Niagara/Private/NiagaraDataInterfaceGrid2DCollection.cpp
+++ b/Engine/Plugins/FX/Niagara/Source/Niagara/Private/NiagaraDataInterfaceGrid2DCollection.cpp
@@ -211,7 +211,7 @@ struct FNiagaraDataInterfaceParametersCS_Grid2DCollection : public FNiagaraDataI
 
 public:
 	void Bind(const FNiagaraDataInterfaceGPUParamInfo& ParameterInfo, const class FShaderParameterMap& ParameterMap)
-	{			
+	{
 		NumAttributesParam.Bind(ParameterMap, *(UNiagaraDataInterfaceRWBase::NumAttributesName + ParameterInfo.DataInterfaceHLSLSymbol));
 		NumCellsParam.Bind(ParameterMap, *(UNiagaraDataInterfaceRWBase::NumCellsName + ParameterInfo.DataInterfaceHLSLSymbol));
 		UnitToUVParam.Bind(ParameterMap, *(UNiagaraDataInterfaceRWBase::UnitToUVName + ParameterInfo.DataInterfaceHLSLSymbol));
@@ -224,7 +224,7 @@ public:
 		SamplerParam.Bind(ParameterMap, *(UNiagaraDataInterfaceGrid2DCollection::SamplerName + ParameterInfo.DataInterfaceHLSLSymbol));
 		AttributeIndicesParam.Bind(ParameterMap, *(UNiagaraDataInterfaceGrid2DCollection::AttributeIndicesBaseName + ParameterInfo.DataInterfaceHLSLSymbol));
 
-		// Gather up all the attribute names referenced. Note that there may be multiple in the list of the same name, 
+		// Gather up all the attribute names referenced. Note that there may be multiple in the list of the same name,
 		// but we only deal with this by the number of bound methods.
 		{
 			int32 NumFuncs = ParameterInfo.GeneratedFunctions.Num();
@@ -256,7 +256,7 @@ public:
 		// Get shader and DI
 		FRHIComputeShader* ComputeShaderRHI = Context.Shader.GetComputeShader();
 		FNiagaraDataInterfaceProxyGrid2DCollectionProxy* VFDI = static_cast<FNiagaraDataInterfaceProxyGrid2DCollectionProxy*>(Context.DataInterface);
-		
+
 		FGrid2DCollectionRWInstanceData_RenderThread* ProxyData = VFDI->SystemInstancesToProxyData_RT.Find(Context.SystemInstanceID);
 		check(ProxyData);
 
@@ -322,7 +322,7 @@ public:
 		}
 	}
 
-	void Unset(FRHICommandList& RHICmdList, const FNiagaraDataInterfaceSetArgs& Context) const 
+	void Unset(FRHICommandList& RHICmdList, const FNiagaraDataInterfaceSetArgs& Context) const
 	{
 		if (OutputGridParam.IsBound())
 		{
@@ -340,7 +340,7 @@ private:
 	LAYOUT_FIELD(FShaderResourceParameter, GridParam);
 	LAYOUT_FIELD(FRWShaderParameter, OutputGridParam);
 	LAYOUT_FIELD(FShaderParameter, AttributeIndicesParam);
-	
+
 	LAYOUT_FIELD(FShaderResourceParameter, SamplerParam);
 	LAYOUT_FIELD(TMemoryImageArray<FName>, AttributeNames);
 	LAYOUT_FIELD(TMemoryImageArray<uint32>, AttributeChannelCount);
@@ -379,12 +379,12 @@ bool UNiagaraDataInterfaceGrid2DCollection::AppendCompileHash(FNiagaraCompileHas
 	if (!Super::AppendCompileHash(InVisitor))
 		return false;
 
-	
+
 	InVisitor->UpdatePOD(TEXT("UNiagaraDataInterfaceGrid2DCollectionVersion"), (int32)FNiagaraGridCollection2DDIFunctionVersion::LatestVersion);
 	return true;
 }
 
-#if WITH_EDITOR	
+#if WITH_EDITOR
 void UNiagaraDataInterfaceGrid2DCollection::GetFeedback(UNiagaraSystem* Asset, UNiagaraComponent* Component, TArray<FNiagaraDataInterfaceError>& OutErrors,
 	TArray<FNiagaraDataInterfaceFeedback>& OutWarnings, TArray<FNiagaraDataInterfaceFeedback>& OutInfo)
 {
@@ -400,7 +400,7 @@ bool UNiagaraDataInterfaceGrid2DCollection::UpgradeFunctionCall(FNiagaraFunction
 {
 	if (Super::UpgradeFunctionCall(FunctionSignature))
 		return true;
-	
+
 	FName UpgradeName;
 	if (FunctionSignature.Name == SetVector3ValueFunctionName)
 		UpgradeName = SetVectorValueFunctionName;
@@ -419,8 +419,8 @@ bool UNiagaraDataInterfaceGrid2DCollection::UpgradeFunctionCall(FNiagaraFunction
 	else if (FunctionSignature.Name == GetFloatValueFunctionName)
 		UpgradeName = GetPreviousFloatValueFunctionName;
 	else if (FunctionSignature.Name == SampleGridFloatFunctionName)
-		UpgradeName = SamplePreviousGridFloatFunctionName;	
-	
+		UpgradeName = SamplePreviousGridFloatFunctionName;
+
 
 	if (UpgradeName != NAME_None)
 	{
@@ -432,7 +432,7 @@ bool UNiagaraDataInterfaceGrid2DCollection::UpgradeFunctionCall(FNiagaraFunction
 			if (Sig.Name == UpgradeName)
 			{
 				FNiagaraFunctionSignature Backup = FunctionSignature;
-				FunctionSignature = Sig;	
+				FunctionSignature = Sig;
 				return true;
 			}
 		}
@@ -466,7 +466,7 @@ void UNiagaraDataInterfaceGrid2DCollection::GetFunctions(TArray<FNiagaraFunction
 		OutFunctions.Add(Sig);
 	}
 
-	{		
+	{
 		// Older, deprecated form
 		FNiagaraFunctionSignature Sig;
 		Sig.Name = GetValueFunctionName;
@@ -526,7 +526,7 @@ void UNiagaraDataInterfaceGrid2DCollection::GetFunctions(TArray<FNiagaraFunction
 		Sig.bExperimental = true;
 		Sig.bMemberFunction = true;
 		Sig.bRequiresContext = false;
-		Sig.bWriteFunction = true; 
+		Sig.bWriteFunction = true;
 		Sig.ModuleUsageBitmask = ENiagaraScriptUsageMask::Particle;
 		Sig.bSoftDeprecatedFunction = true;
 		Sig.bSupportsCPU = false;
@@ -572,7 +572,7 @@ void UNiagaraDataInterfaceGrid2DCollection::GetFunctions(TArray<FNiagaraFunction
 		Sig.bMemberFunction = true;
 		Sig.bRequiresContext = false;
 		Sig.bWriteFunction = true;
-		Sig.bRequiresExecPin = true; 
+		Sig.bRequiresExecPin = true;
 		Sig.ModuleUsageBitmask = ENiagaraScriptUsageMask::Particle;
 		Sig.bSupportsCPU = false;
 		Sig.bSupportsGPU = true;
@@ -646,7 +646,7 @@ void UNiagaraDataInterfaceGrid2DCollection::GetFunctions(TArray<FNiagaraFunction
 
 		// Add older version for back compat with CustomHLSL.
 		Sig.Name = GetVector4ValueFunctionName;
-		Sig.bSoftDeprecatedFunction = true; 
+		Sig.bSoftDeprecatedFunction = true;
 		OutFunctions.Add(Sig);
 	}
 
@@ -902,7 +902,7 @@ void UNiagaraDataInterfaceGrid2DCollection::GetFunctions(TArray<FNiagaraFunction
 		Sig.Name = SampleGridFunctionName;
 		Sig.Inputs.Add(FNiagaraVariable(FNiagaraTypeDefinition(GetClass()), TEXT("Grid")));
 		Sig.Inputs.Add(FNiagaraVariable(FNiagaraTypeDefinition::GetFloatDef(), TEXT("UnitX")));
-		Sig.Inputs.Add(FNiagaraVariable(FNiagaraTypeDefinition::GetFloatDef(), TEXT("UnitY")));		
+		Sig.Inputs.Add(FNiagaraVariable(FNiagaraTypeDefinition::GetFloatDef(), TEXT("UnitY")));
 		Sig.Inputs.Add(FNiagaraVariable(FNiagaraTypeDefinition::GetIntDef(), TEXT("AttributeIndex")));
 		Sig.Outputs.Add(FNiagaraVariable(FNiagaraTypeDefinition::GetFloatDef(), TEXT("Value")));
 
@@ -1041,7 +1041,7 @@ void UNiagaraDataInterfaceGrid2DCollection::GetVMExternalFunction(const FVMExter
 		check(BindingInfo.GetNumInputs() == 3 && BindingInfo.GetNumOutputs() == 1);
 		NDI_FUNC_BINDER(UNiagaraDataInterfaceGrid2DCollection, SetNumCells)::Bind(this, OutFunc);
 	}
-	
+
 	else if (BindingInfo.Name == GetVector4AttributeIndexFunctionName)
 	{
 		FName AttributeName = BindingInfo.FindSpecifier(NAME_Attribute)->Value;
@@ -1090,7 +1090,7 @@ void UNiagaraDataInterfaceGrid2DCollection::GetParameterDefinitionHLSL(const FNi
 {
 	Super::GetParameterDefinitionHLSL(ParamInfo, OutHLSL);
 
-	static const TCHAR *FormatDeclarations = TEXT(R"(				
+	static const TCHAR *FormatDeclarations = TEXT(R"(
 		Texture2DArray<float> {GridName};
 		RWTexture2DArray<float> RW{OutputGridName};
 		SamplerState {SamplerName};
@@ -1104,7 +1104,7 @@ void UNiagaraDataInterfaceGrid2DCollection::GetParameterDefinitionHLSL(const FNi
 	// Note that we have to have at least one here because hlsl doesn't support arrays of size 0.
 	const int AttributeInt4Count = FMath::Max(1, FMath::DivideAndRoundUp(ParamInfo.GeneratedFunctions.Num(), 4));
 
-	TMap<FString, FStringFormatArg> ArgsDeclarations = {				
+	TMap<FString, FStringFormatArg> ArgsDeclarations = {
 		{ TEXT("GridName"),    GridName + ParamInfo.DataInterfaceHLSLSymbol },
 		{ TEXT("SamplerName"),    SamplerName + ParamInfo.DataInterfaceHLSLSymbol },
 		{ TEXT("OutputGridName"),    OutputGridName + ParamInfo.DataInterfaceHLSLSymbol },
@@ -1117,10 +1117,10 @@ void UNiagaraDataInterfaceGrid2DCollection::GetParameterDefinitionHLSL(const FNi
 
 void UNiagaraDataInterfaceGrid2DCollection::WriteSetHLSL(const FNiagaraDataInterfaceGPUParamInfo& ParamInfo, const FNiagaraDataInterfaceGeneratedFunction& FunctionInfo, int FunctionInstanceIndex, int32 InNumChannels, FString& OutHLSL)
 {
-	
+
 	FString FormatBounds = TEXT(R"(
 			void {FunctionName}(int In_IndexX, int In_IndexY, float{NumChannelsVariableSuffix} In_Value)
-			{			
+			{
 				int In_AttributeIndex = {AttributeIndicesName}[{AttributeIndexGroup}]{AttributeIndexComponent};
 
 			    for (int i = 0; i < {NumChannels}; i++)
@@ -1138,7 +1138,7 @@ void UNiagaraDataInterfaceGrid2DCollection::WriteSetHLSL(const FNiagaraDataInter
 					{
 						case 0:
 							Val = In_Value.x;
-							break; 
+							break;
 						case 1:
 							Val = In_Value.y;
 							break; )");
@@ -1160,10 +1160,10 @@ void UNiagaraDataInterfaceGrid2DCollection::WriteSetHLSL(const FNiagaraDataInter
 	}
 	if (InNumChannels > 1)
 	{
-		FormatBounds += TEXT(R"(	
+		FormatBounds += TEXT(R"(
 					})");
 	}
-	FormatBounds += TEXT(R"(	
+	FormatBounds += TEXT(R"(
 					RW{OutputGrid}[int3(In_IndexX, In_IndexY, In_AttributeIndex + i)] = Val;
 				}
 			}
@@ -1185,7 +1185,7 @@ void UNiagaraDataInterfaceGrid2DCollection::WriteSetHLSL(const FNiagaraDataInter
 
 void UNiagaraDataInterfaceGrid2DCollection::WriteGetHLSL(const FNiagaraDataInterfaceGPUParamInfo& ParamInfo, const FNiagaraDataInterfaceGeneratedFunction& FunctionInfo, int FunctionInstanceIndex, int32 InNumChannels, FString& OutHLSL)
 {
-	
+
 	FString FormatBounds = TEXT(R"(
 			void {FunctionName}(int In_IndexX, int In_IndexY, out float{NumChannelsVariableSuffix} Out_Val)
 			{
@@ -1206,7 +1206,7 @@ void UNiagaraDataInterfaceGrid2DCollection::WriteGetHLSL(const FNiagaraDataInter
 					{
 						case 0:
 							Out_Val.x = Val;
-							break; 
+							break;
 						case 1:
 							Out_Val.y = Val;
 							break; )");
@@ -1228,10 +1228,10 @@ void UNiagaraDataInterfaceGrid2DCollection::WriteGetHLSL(const FNiagaraDataInter
 	}
 	if (InNumChannels > 1)
 	{
-		FormatBounds += TEXT(R"(	
+		FormatBounds += TEXT(R"(
 					})");
 	}
-	FormatBounds += TEXT(R"(	
+	FormatBounds += TEXT(R"(
 				}
 			}
 		)");
@@ -1273,7 +1273,7 @@ void UNiagaraDataInterfaceGrid2DCollection::WriteSampleHLSL(const FNiagaraDataIn
 					{
 						case 0:
 							Out_Val.x = Val;
-							break; 
+							break;
 						case 1:
 							Out_Val.y = Val;
 							break; )");
@@ -1295,10 +1295,10 @@ void UNiagaraDataInterfaceGrid2DCollection::WriteSampleHLSL(const FNiagaraDataIn
 	}
 	if (InNumChannels > 1)
 	{
-		FormatBounds += TEXT(R"(	
+		FormatBounds += TEXT(R"(
 					})");
 	}
-	FormatBounds += TEXT(R"(	
+	FormatBounds += TEXT(R"(
 				}
 			}
 		)");
@@ -1328,7 +1328,7 @@ void UNiagaraDataInterfaceGrid2DCollection::WriteAttributeGetIndexHLSL(const FNi
 			}
 	)");
 
-	
+
 	TMap<FString, FStringFormatArg> ArgsBounds = {
 		{TEXT("FunctionName"), FunctionInfo.InstanceName},
 		{TEXT("AttributeIndicesName"), AttributeIndicesBaseName + ParamInfo.DataInterfaceHLSLSymbol},
@@ -1381,7 +1381,7 @@ bool UNiagaraDataInterfaceGrid2DCollection::GetFunctionHLSL(const FNiagaraDataIn
 	if (ParentRet)
 	{
 		return true;
-	} 
+	}
 
 	TMap<FString, FStringFormatArg> ArgsBounds =
 	{
@@ -1409,7 +1409,7 @@ bool UNiagaraDataInterfaceGrid2DCollection::GetFunctionHLSL(const FNiagaraDataIn
 	{
 		static const TCHAR *FormatBounds = TEXT(R"(
 			void {FunctionName}(int In_IndexX, int In_IndexY, int In_AttributeIndex, float In_Value, out int val)
-			{			
+			{
 				val = 0;
 				RW{OutputGrid}[int3(In_IndexX, In_IndexY, In_AttributeIndex)] = In_Value;
 			}
@@ -1423,7 +1423,7 @@ bool UNiagaraDataInterfaceGrid2DCollection::GetFunctionHLSL(const FNiagaraDataIn
 			void {FunctionName}(int In_IndexX, int In_IndexY)
 			{
 				for (int AttributeIndex = 0; AttributeIndex < {NumAttributes}.x; AttributeIndex++)
-				{			
+				{
 					float Val = {Grid}.Load(int4(In_IndexX, In_IndexY, AttributeIndex, 0));
 					RW{OutputGrid}[int3(In_IndexX, In_IndexY, AttributeIndex)] = Val;
 				}
@@ -1438,7 +1438,7 @@ bool UNiagaraDataInterfaceGrid2DCollection::GetFunctionHLSL(const FNiagaraDataIn
 			void {FunctionName}(int In_IndexX, int In_IndexY)
 			{
 				for (int AttributeIndex = 0; AttributeIndex < {NumAttributes}.x; AttributeIndex++)
-				{			
+				{
 					float Val = 0.0f;
 					RW{OutputGrid}[int3(In_IndexX, In_IndexY, AttributeIndex)] = Val;
 				}
@@ -1555,7 +1555,7 @@ bool UNiagaraDataInterfaceGrid2DCollection::GenerateIterationSourceNamespaceRead
 			DIVarName = InArguments[i].GetName().ToString();
 		}
 	}
-	
+
 
 	if (InAttributes.Num() != InAttributeHLSLNames.Num())
 		return false;
@@ -1609,7 +1609,7 @@ bool UNiagaraDataInterfaceGrid2DCollection::GenerateIterationSourceNamespaceWrit
 	OutHLSL += TEXT("\t//Generated by UNiagaraDataInterfaceGrid2DCollection::GenerateIterationSourceNamespaceWriteAttributesHLSL\n");
 	for (int32 i = 0; i < InArguments.Num(); i++)
 	{
-		OutHLSL += FString::Printf(TEXT("\t// Argument Name \"%s\" Type \"%s\"\n"), *InArguments[i].GetName().ToString(), *InArguments[i].GetType().GetName());		
+		OutHLSL += FString::Printf(TEXT("\t// Argument Name \"%s\" Type \"%s\"\n"), *InArguments[i].GetName().ToString(), *InArguments[i].GetType().GetName());
 		if (InArguments[i].GetType().GetClass() == GetClass())
 		{
 			DIVarName = InArguments[i].GetName().ToString();
@@ -1624,14 +1624,14 @@ bool UNiagaraDataInterfaceGrid2DCollection::GenerateIterationSourceNamespaceWrit
 	{
 		OutHLSL += FString::Printf(TEXT("\tint X, Y;\n\t%s.ExecutionIndexToGridIndex(X, Y);\n"), *DIVarName);
 	}
-	
+
 	TArray<FString> RootArray;
 	IterationSourceVar.GetName().ToString().ParseIntoArray(RootArray, TEXT("."));
 
 	for (int32 i = 0; i < InAttributes.Num(); i++)
 	{
 		OutHLSL += FString::Printf(TEXT("\t// Name \"%s\" Type \"%s\" Var \"%s\"\n"), *InAttributes[i].GetName().ToString(), *InAttributes[i].GetType().GetName(), *InAttributeHLSLNames[i]);
-		
+
 		TArray<FString> OutArray;
 		if (InAttributes[i].GetName().ToString().ParseIntoArray(OutArray, TEXT(".")) > 0)
 		{
@@ -1641,7 +1641,7 @@ bool UNiagaraDataInterfaceGrid2DCollection::GenerateIterationSourceNamespaceWrit
 				OutErrors.Add(Error);
 				continue;
 			}
-			
+
 			// Clear out the shared namespace with the root variable...
 			FString AttributeName;
 			for (int32 NamespaceIdx = 0; NamespaceIdx < OutArray.Num(); NamespaceIdx++)
@@ -1700,7 +1700,7 @@ bool UNiagaraDataInterfaceGrid2DCollection::GenerateSetupHLSL(FNiagaraDataInterf
 bool UNiagaraDataInterfaceGrid2DCollection::GenerateTeardownHLSL(FNiagaraDataInterfaceGPUParamInfo& DIInstanceInfo, TConstArrayView<FNiagaraVariable> InArguments, bool bSpawnOnly, bool bPartialWrites, TArray<FText>& OutErrors, FString& OutHLSL) const
 {
 	OutHLSL += TEXT("\t//Generated by UNiagaraDataInterfaceGrid2DCollection::GenerateTeardownHLSL\n");
-	
+
 
 	return true;
 }
@@ -1758,7 +1758,7 @@ bool UNiagaraDataInterfaceGrid2DCollection::InitPerInstanceData(void* PerInstanc
 		InstanceData->NumCells.Y = FMath::Max(1, int32(float(InstanceData->NumCells.Y) * GNiagaraGrid2DResolutionMultiplier));
 	}
 
-	// If we are setting the grid from the voxel size, then recompute NumVoxels and change bbox	
+	// If we are setting the grid from the voxel size, then recompute NumVoxels and change bbox
 	if (SetGridFromMaxAxis)
 	{
 		float CellSize = FMath::Max(WorldBBoxSize.X, WorldBBoxSize.Y) / NumCellsMaxAxis;
@@ -1774,12 +1774,12 @@ bool UNiagaraDataInterfaceGrid2DCollection::InitPerInstanceData(void* PerInstanc
 		else if (WorldBBoxSize.X < WorldBBoxSize.Y && !FMath::IsNearlyEqual(CellSize * InstanceData->NumCells.X, WorldBBoxSize.X))
 		{
 			InstanceData->NumCells.X++;
-		}		
-		
+		}
+
 		InstanceData->WorldBBoxSize = FVector2D(InstanceData->NumCells.X, InstanceData->NumCells.Y) * CellSize;
 		NumCellsX = InstanceData->NumCells.X;
 		NumCellsY = InstanceData->NumCells.Y;
-	}	
+	}
 
 	InstanceData->CellSize = InstanceData->WorldBBoxSize / FVector2D(InstanceData->NumCells.X, InstanceData->NumCells.Y);
 
@@ -1854,7 +1854,7 @@ bool UNiagaraDataInterfaceGrid2DCollection::InitPerInstanceData(void* PerInstanc
 		if (RT_Resource && RT_Resource->TextureRHI.IsValid())
 		{
 			TargetData->RenderTargetToCopyTo = RT_Resource->TextureRHI;
-		}	
+		}
 		else
 		{
 			TargetData->RenderTargetToCopyTo = nullptr;
@@ -1888,7 +1888,7 @@ void UNiagaraDataInterfaceGrid2DCollection::DestroyPerInstanceData(void* PerInst
 }
 
 bool UNiagaraDataInterfaceGrid2DCollection::PerInstanceTick(void* PerInstanceData, FNiagaraSystemInstance* SystemInstance, float DeltaSeconds)
-{	
+{
 	FGrid2DCollectionRWInstanceData_GameThread* InstanceData = SystemInstancesToProxyData_GT.FindRef(SystemInstance->GetId());
 
 	ENiagaraGpuBufferFormat BufferFormat = bOverrideFormat ? OverrideBufferFormat : GetDefault<UNiagaraSettings>()->DefaultGridFormat;
@@ -2019,7 +2019,7 @@ void UNiagaraDataInterfaceGrid2DCollection::FindAttributes(TArray<FNiagaraVariab
 		if (Emitter && EmitterHandle.GetIsEnabled() && Emitter->IsValid() && (Emitter->SimTarget == ENiagaraSimTarget::GPUComputeSim))
 		{
 			// Search scripts for this data interface so we get the variable name
-			auto FindDataInterfaceVariable = 
+			auto FindDataInterfaceVariable =
 				[&OwnerSystem, &Emitter](const UNiagaraDataInterface* DataInterface) -> FName
 				{
 					UNiagaraScript* Scripts[] =
@@ -2137,8 +2137,8 @@ bool UNiagaraDataInterfaceGrid2DCollection::FillTexture2D(const UNiagaraComponen
 		return false;
 	}
 
-	FNiagaraSystemInstance *SystemInstance = Component->GetSystemInstance();
-	if (!SystemInstance)
+	FNiagaraSystemInstanceControllerConstPtr SystemInstanceController = Component->GetSystemInstanceController();
+	if (!SystemInstanceController.IsValid())
 	{
 		return false;
 	}
@@ -2159,7 +2159,7 @@ bool UNiagaraDataInterfaceGrid2DCollection::FillTexture2D(const UNiagaraComponen
 
 	FNiagaraDataInterfaceProxyGrid2DCollectionProxy* RT_Proxy = GetProxyAs<FNiagaraDataInterfaceProxyGrid2DCollectionProxy>();
 	ENQUEUE_RENDER_COMMAND(FUpdateDIColorCurve)(
-		[RT_Proxy, InstanceID=SystemInstance->GetId(), RT_TextureResource=Dest->Resource, AttributeIndex](FRHICommandListImmediate& RHICmdList)
+		[RT_Proxy, InstanceID=SystemInstanceController->GetSystemInstanceID(), RT_TextureResource=Dest->Resource, AttributeIndex](FRHICommandListImmediate& RHICmdList)
 	{
 		FGrid2DCollectionRWInstanceData_RenderThread* Grid2DInstanceData = RT_Proxy->SystemInstancesToProxyData_RT.Find(InstanceID);
 
@@ -2184,22 +2184,23 @@ bool UNiagaraDataInterfaceGrid2DCollection::FillRawTexture2D(const UNiagaraCompo
 		return false;
 	}
 
-	FNiagaraSystemInstance* SystemInstance = Component->GetSystemInstance();
-	if (!SystemInstance)
+	FNiagaraSystemInstanceControllerConstPtr SystemInstanceController = Component->GetSystemInstanceController();
+	if (!SystemInstanceController.IsValid())
 	{
 		TilesX = -1;
 		TilesY = -1;
 		return false;
 	}
+	const FNiagaraSystemInstanceID SystemID = SystemInstanceController->GetSystemInstanceID();
 
-	FGrid2DCollectionRWInstanceData_GameThread* Grid2DInstanceData = SystemInstancesToProxyData_GT.FindRef(SystemInstance->GetId());
+	FGrid2DCollectionRWInstanceData_GameThread* Grid2DInstanceData = SystemInstancesToProxyData_GT.FindRef(SystemID);
 	if (!Grid2DInstanceData)
 	{
 		TilesX = -1;
 		TilesY = -1;
 		return false;
 	}
-	
+
 	const FNiagaraGrid2DLegacyTiled2DInfo Tiled2DInfo(Grid2DInstanceData->NumCells, Grid2DInstanceData->NumAttributes);
 	TilesX = Tiled2DInfo.NumTiles.X;
 	TilesY = Tiled2DInfo.NumTiles.Y;
@@ -2214,7 +2215,7 @@ bool UNiagaraDataInterfaceGrid2DCollection::FillRawTexture2D(const UNiagaraCompo
 
 	FNiagaraDataInterfaceProxyGrid2DCollectionProxy* RT_Proxy = GetProxyAs<FNiagaraDataInterfaceProxyGrid2DCollectionProxy>();
 	ENQUEUE_RENDER_COMMAND(FUpdateDIColorCurve)(
-		[RT_Proxy, RT_InstanceID=SystemInstance->GetId(), RT_TextureResource=Dest->Resource](FRHICommandListImmediate& RHICmdList)
+		[RT_Proxy, RT_InstanceID=SystemID, RT_TextureResource=Dest->Resource](FRHICommandListImmediate& RHICmdList)
 	{
 		FGrid2DCollectionRWInstanceData_RenderThread* RT_Grid2DInstanceData = RT_Proxy->SystemInstancesToProxyData_RT.Find(RT_InstanceID);
 		if (RT_TextureResource && RT_TextureResource->TextureRHI.IsValid() && RT_Grid2DInstanceData && RT_Grid2DInstanceData->CurrentData)
@@ -2236,14 +2237,14 @@ void UNiagaraDataInterfaceGrid2DCollection::GetRawTextureSize(const UNiagaraComp
 		return;
 	}
 
-	FNiagaraSystemInstance *SystemInstance = Component->GetSystemInstance();
-	if (!SystemInstance)
+	FNiagaraSystemInstanceControllerConstPtr SystemInstanceController = Component->GetSystemInstanceController();
+	if (!SystemInstanceController.IsValid())
 	{
 		SizeX = -1;
 		SizeY = -1;
 		return;
 	}
-	FNiagaraSystemInstanceID InstanceID = SystemInstance->GetId();
+	FNiagaraSystemInstanceID InstanceID = SystemInstanceController->GetSystemInstanceID();
 
 	FGrid2DCollectionRWInstanceData_GameThread* Grid2DInstanceData = SystemInstancesToProxyData_GT.FindRef(InstanceID);
 	if (!Grid2DInstanceData)
@@ -2267,14 +2268,14 @@ void UNiagaraDataInterfaceGrid2DCollection::GetTextureSize(const UNiagaraCompone
 		return;
 	}
 
-	FNiagaraSystemInstance *SystemInstance = Component->GetSystemInstance();
-	if (!SystemInstance)
+	FNiagaraSystemInstanceControllerConstPtr SystemInstanceController = Component->GetSystemInstanceController();
+	if (!SystemInstanceController.IsValid())
 	{
 		SizeX = -1;
 		SizeY = -1;
 		return;
 	}
-	FNiagaraSystemInstanceID InstanceID = SystemInstance->GetId();
+	FNiagaraSystemInstanceID InstanceID = SystemInstanceController->GetSystemInstanceID();
 
 	FGrid2DCollectionRWInstanceData_GameThread* Grid2DInstanceData = SystemInstancesToProxyData_GT.FindRef(InstanceID);
 	if (!Grid2DInstanceData)
@@ -2306,7 +2307,7 @@ void UNiagaraDataInterfaceGrid2DCollection::GetCellSize(FVectorVMContext& Contex
 	FNDIOutputParam<FVector2D> OutCellSize(Context);
 
 	for (int32 InstanceIdx = 0; InstanceIdx < Context.NumInstances; ++InstanceIdx)
-	{	
+	{
 		OutCellSize.SetAndAdvance(InstData->CellSize);
 	}
 }
@@ -2341,7 +2342,7 @@ void UNiagaraDataInterfaceGrid2DCollection::SetNumCells(FVectorVMContext& Contex
 		if (bSuccess)
 		{
 			FIntPoint OldNumCells = InstData->NumCells;
-	
+
 			InstData->NumCells.X = NewNumCellsX;
 			InstData->NumCells.Y = NewNumCellsY;
 
@@ -2364,7 +2365,7 @@ bool UNiagaraDataInterfaceGrid2DCollection::PerInstanceTickPostSimulate(void* Pe
 	if (InstanceData->NeedsRealloc && InstanceData->NumCells.X > 0 && InstanceData->NumCells.Y > 0)
 	{
 		InstanceData->NeedsRealloc = false;
-		
+
 		InstanceData->CellSize = InstanceData->WorldBBoxSize / FVector2D(InstanceData->NumCells.X, InstanceData->NumCells.Y);
 
 		if (InstanceData->TargetTexture)
@@ -2385,15 +2386,15 @@ bool UNiagaraDataInterfaceGrid2DCollection::PerInstanceTickPostSimulate(void* Pe
 		{
 			check(RT_Proxy->SystemInstancesToProxyData_RT.Contains(InstanceID));
 			FGrid2DCollectionRWInstanceData_RenderThread* TargetData = RT_Proxy->SystemInstancesToProxyData_RT.Find(InstanceID);
-			
+
 			TargetData->NumCells = RT_InstanceData.NumCells;
 			TargetData->NumAttributes = RT_InstanceData.NumAttributes;
 			TargetData->CellSize = RT_InstanceData.CellSize;
-			
+
 			TargetData->Buffers.Empty();
 			TargetData->CurrentData = nullptr;
 			TargetData->DestinationData = nullptr;
-			
+
 			if (RT_Resource && RT_Resource->TextureRHI.IsValid())
 			{
 				TargetData->RenderTargetToCopyTo = RT_Resource->TextureRHI;
@@ -2549,19 +2550,19 @@ void FNiagaraDataInterfaceProxyGrid2DCollectionProxy::PreStage(FRHICommandList&
 		// we start.  If a user wants to access data from the previous stage, then they can read from the current data.
 
 		// #todo(dmp): we might want to expose an option where we have buffers that are write only and need a clear (ie: no buffering like the neighbor grid).  They would be considered transient perhaps?  It'd be more
-		// memory efficient since it would theoretically not require any double buffering.		
+		// memory efficient since it would theoretically not require any double buffering.
 		RHICmdList.Transition(FRHITransitionInfo(ProxyData->DestinationData->GridUAV, ERHIAccess::SRVMask, ERHIAccess::UAVCompute));
 		if (!Context.IsIterationStage)
 		{
 			SCOPED_DRAW_EVENT(RHICmdList, Grid2DCollection_PreStage);
 			RHICmdList.ClearUAVFloat(ProxyData->DestinationData->GridUAV, FVector4(ForceInitToZero));
 			RHICmdList.Transition(FRHITransitionInfo(ProxyData->DestinationData->GridUAV, ERHIAccess::UAVCompute, ERHIAccess::UAVCompute));
-		}		
+		}
 	}
 }
 
 void FNiagaraDataInterfaceProxyGrid2DCollectionProxy::PostStage(FRHICommandList& RHICmdList, const FNiagaraDataInterfaceStageArgs& Context)
-{	
+{
 	if (Context.IsOutputStage)
 	{
 		FGrid2DCollectionRWInstanceData_RenderThread* ProxyData = SystemInstancesToProxyData_RT.Find(Context.SystemInstanceID);
@@ -2608,7 +2609,7 @@ void FNiagaraDataInterfaceProxyGrid2DCollectionProxy::PostSimulate(FRHICommandLi
 }
 
 void FNiagaraDataInterfaceProxyGrid2DCollectionProxy::ResetData(FRHICommandList& RHICmdList, const FNiagaraDataInterfaceArgs& Context)
-{	
+{
 	FGrid2DCollectionRWInstanceData_RenderThread* ProxyData = SystemInstancesToProxyData_RT.Find(Context.SystemInstanceID);
 	if (!ProxyData)
 	{
@@ -2636,8 +2637,8 @@ void FNiagaraDataInterfaceProxyGrid2DCollectionProxy::ResetData(FRHICommandList&
 
 			RHICmdList.ClearUAVFloat(Buffer->GridUAV, FVector4(ForceInitToZero));
 			RHICmdList.Transition(FRHITransitionInfo(Buffer->GridUAV, ERHIAccess::UAVCompute, AccessAfter));
-		}		
-	}	
+		}
+	}
 }
 
 FIntVector FNiagaraDataInterfaceProxyGrid2DCollectionProxy::GetElementCount(FNiagaraSystemInstanceID SystemInstanceID) const
diff --git a/Engine/Plugins/FX/Niagara/Source/Niagara/Private/NiagaraDataInterfaceGrid3DCollection.cpp b/Engine/Plugins/FX/Niagara/Source/Niagara/Private/NiagaraDataInterfaceGrid3DCollection.cpp
--- a/Engine/Plugins/FX/Niagara/Source/Niagara/Private/NiagaraDataInterfaceGrid3DCollection.cpp
+++ b/Engine/Plugins/FX/Niagara/Source/Niagara/Private/NiagaraDataInterfaceGrid3DCollection.cpp
@@ -1625,7 +1625,7 @@ bool UNiagaraDataInterfaceGrid3DCollection::GetFunctionHLSL(const FNiagaraDataIn
 					int TileIndexX = In_AttributeIndex % {NumTiles}.x;
 					int TileIndexY = (In_AttributeIndex / {NumTiles}.x) % {NumTiles}.y;
 					int TileIndexZ = In_AttributeIndex / ({NumTiles}.x * {NumTiles}.y);		
-
+					
 					float3 UVW =
 					{
 						In_UnitX / {NumTiles}.x + 1.0*TileIndexX/{NumTiles}.x,
@@ -2365,15 +2365,15 @@ void UNiagaraDataInterfaceGrid3DCollection::GetRawTextureSize(const UNiagaraComp
 		return;
 	}
 
-	FNiagaraSystemInstance *SystemInstance = Component->GetSystemInstance();
-	if (!SystemInstance)
+	FNiagaraSystemInstanceControllerConstPtr Controller = Component->GetSystemInstanceController();
+	if (!Controller)
 	{
 		SizeX = -1;
 		SizeY = -1;
 		SizeZ = -1;
 		return;
 	}
-	FNiagaraSystemInstanceID InstanceID = SystemInstance->GetId();
+	FNiagaraSystemInstanceID InstanceID = Controller->GetSystemInstanceID();
 
 	FGrid3DCollectionRWInstanceData_GameThread* Grid3DInstanceData = SystemInstancesToProxyData_GT.FindRef(InstanceID);
 	if (!Grid3DInstanceData)
@@ -2400,15 +2400,15 @@ void UNiagaraDataInterfaceGrid3DCollection::GetTextureSize(const UNiagaraCompone
 		return;
 	}
 
-	FNiagaraSystemInstance *SystemInstance = Component->GetSystemInstance();
-	if (!SystemInstance)
+	FNiagaraSystemInstanceControllerConstPtr Controller = Component->GetSystemInstanceController();
+	if (!Controller)
 	{
 		SizeX = -1;
 		SizeY = -1;
 		SizeZ = -1;
 		return;
 	}
-	FNiagaraSystemInstanceID InstanceID = SystemInstance->GetId();
+	FNiagaraSystemInstanceID InstanceID = Controller->GetSystemInstanceID();
 
 	FGrid3DCollectionRWInstanceData_GameThread* Grid3DInstanceData = SystemInstancesToProxyData_GT.FindRef(InstanceID);
 	if (!Grid3DInstanceData)
diff --git a/Engine/Plugins/FX/Niagara/Source/Niagara/Private/NiagaraDataInterfaceSkeletalMesh.cpp b/Engine/Plugins/FX/Niagara/Source/Niagara/Private/NiagaraDataInterfaceSkeletalMesh.cpp
--- a/Engine/Plugins/FX/Niagara/Source/Niagara/Private/NiagaraDataInterfaceSkeletalMesh.cpp
+++ b/Engine/Plugins/FX/Niagara/Source/Niagara/Private/NiagaraDataInterfaceSkeletalMesh.cpp
@@ -167,7 +167,7 @@ float FSkeletalMeshSamplingRegionAreaWeightedSampler::GetWeights(TArray<float>&
 		OutWeights.Empty();
 		return 0.0f;
 	}
-	
+
 	const FSkeletalMeshSamplingInfo& SamplingInfo = SkelMesh->GetSamplingInfo();
 	OutWeights.Empty(NumUsedRegions);
 	for (int32 i = 0; i < NumUsedRegions; ++i)
@@ -240,7 +240,7 @@ void FSkeletalMeshSkinningData::RegisterUser(FSkeletalMeshSkinningDataUsage Usag
 	USkeletalMesh* SkelMesh = SkelComp->SkeletalMesh;
 	int32 LODIndex = 0;
 	int32 NumLODInfo = 1;
-	
+
 	if (SkelMesh != nullptr)
 	{
 		NumLODInfo = SkelMesh->GetLODInfoArray().Num();
@@ -270,7 +270,7 @@ void FSkeletalMeshSkinningData::RegisterUser(FSkeletalMeshSkinningDataUsage Usag
 		{
 			UpdateBoneTransforms();
 		}
-		
+
 		//Prime the prev matrices if they're missing.
 		if (PrevBoneRefToLocals().Num() != CurrBoneRefToLocals().Num())
 		{
@@ -505,7 +505,7 @@ FSkeletalMeshSkinningDataHandle FNDI_SkeletalMesh_GeneratedData::GetCachedSkinni
 	// We need to add
 	FRWScopeLock WriteLock(CachedSkinningDataGuard, SLT_Write);
 	return FSkeletalMeshSkinningDataHandle(
-		Usage, 
+		Usage,
 		CachedSkinningData.Add(Component, MakeShared<FSkeletalMeshSkinningData>(Component)),
 		bNeedsDataImmediately);
 }
@@ -579,14 +579,14 @@ void FNDI_SkeletalMesh_GeneratedData::Tick(ETickingGroup TickGroup, float DeltaS
 	{
 		CachedSkinningData.Remove(Key);
 	}
-	
-	// First tick the meshes that don't need pre-skinning. 
+
+	// First tick the meshes that don't need pre-skinning.
 	// This prevents additional threading overhead when we don't need to pre-skin.
 	for (int i = 0; i < ToTickBonesOnly.Num(); i++)
 	{
 		ToTickBonesOnly[i]->Tick(DeltaSeconds, false);
 	}
-		
+
 	// Then tick the remaining meshes requiring pre-skinning in parallel
 	if (ToTickPreskin.Num() != 0)
 	{
@@ -739,7 +739,7 @@ void FSkeletalMeshGpuSpawnStaticBuffers::Initialise(FNDISkeletalMesh_InstanceDat
 		TriangleCount = SkeletalMeshLODRenderData.MultiSizeIndexContainer.GetIndexBuffer()->Num() / 3;
 		VertexCount = SkeletalMeshLODRenderData.GetNumVertices();
 		
-		// TODO: Bring these back when we can know if they are for sure sampling from them. Disabled for now to suppress log spam		
+		// TODO: Bring these back when we can know if they are for sure sampling from them. Disabled for now to suppress log spam
 		//if (TriangleCount == 0)
 		//{
 		//	UE_LOG(LogNiagara, Warning, TEXT("FSkeletalMeshGpuSpawnStaticBuffers> TriangleCount(0) is invalid for SkelMesh(%s) System(%s)"), *GetFullNameSafe(InstData->SkeletalMesh.Get()), *GetFullNameSafe(SystemInstance->GetSystem()));
@@ -1685,7 +1685,7 @@ void UNiagaraDataInterfaceSkeletalMesh::ProvidePerInstanceDataForRenderThread(vo
 	Data->ConnectivityBuffer = SourceData->Connectivity.GetProxy();
 }
 
-USkeletalMesh* UNiagaraDataInterfaceSkeletalMesh::GetSkeletalMesh(FNiagaraSystemInstance* SystemInstance, TWeakObjectPtr<USceneComponent>& SceneComponent, USkeletalMeshComponent*& FoundSkelComp, FNDISkeletalMesh_InstanceData* InstData)
+USkeletalMesh* UNiagaraDataInterfaceSkeletalMesh::GetSkeletalMesh(FNiagaraSystemInstance* SystemInstance, USceneComponent* AttachComponent, TWeakObjectPtr<USceneComponent>& SceneComponent, USkeletalMeshComponent*& FoundSkelComp, FNDISkeletalMesh_InstanceData* InstData)
 {
 	// Helper to scour an actor (or its parents) for a valid skeletal mesh component
 	auto FindActorSkelMeshComponent = [](AActor* Actor, bool bRecurseParents = false) -> USkeletalMeshComponent*
@@ -1734,7 +1734,7 @@ USkeletalMesh* UNiagaraDataInterfaceSkeletalMesh::GetSkeletalMesh(FNiagaraSystem
 		UObject* UserParamObject = InstData->UserParamBinding.Init(SystemInstance->GetInstanceParameters(), MeshUserParameter.Parameter);
 		InstData->CachedUserParam = UserParamObject;
 		if (UserParamObject)
-		{			
+		{
 			if (USkeletalMeshComponent* UserSkelMeshComp = Cast<USkeletalMeshComponent>(UserParamObject))
 			{
 				if (!UserSkelMeshComp->IsPendingKill())
@@ -1748,10 +1748,10 @@ USkeletalMesh* UNiagaraDataInterfaceSkeletalMesh::GetSkeletalMesh(FNiagaraSystem
 			}
 			else
 			{
-				//We have a valid, non-null UObject parameter type but it is not a type we can use to get a skeletal mesh from. 
+				//We have a valid, non-null UObject parameter type but it is not a type we can use to get a skeletal mesh from.
 				UE_LOG(LogNiagara, Warning, TEXT("SkeletalMesh data interface using object parameter with invalid type. Skeletal Mesh Data Interfaces can only get a valid mesh from SkeletalMeshComponents, SkeletalMeshActors or Actors."));
 				UE_LOG(LogNiagara, Warning, TEXT("Invalid Parameter : %s"), *UserParamObject->GetFullName());
-				UE_LOG(LogNiagara, Warning, TEXT("Niagara Component : %s"), *GetFullNameSafe(Cast<UNiagaraComponent>(SystemInstance->GetAttachComponent())));
+				UE_LOG(LogNiagara, Warning, TEXT("Niagara Component : %s"), *GetFullNameSafe(Cast<UNiagaraComponent>(AttachComponent)));
 				UE_LOG(LogNiagara, Warning, TEXT("System : %s"), *GetFullNameSafe(SystemInstance->GetSystem()));
 			}
 		}
@@ -1768,34 +1768,31 @@ USkeletalMesh* UNiagaraDataInterfaceSkeletalMesh::GetSkeletalMesh(FNiagaraSystem
 	{
 		FoundSkelComp = FindActorSkelMeshComponent(Source);
 	}
-	else if (bTryAttachParent && SystemInstance)
+	else if (bTryAttachParent && AttachComponent)
 	{
-		if (USceneComponent* AttachComponent = SystemInstance->GetAttachComponent())
+		// First, try to find the mesh component up the attachment hierarchy
+		for (USceneComponent* Curr = AttachComponent; Curr; Curr = Curr->GetAttachParent())
 		{
-			// First, try to find the mesh component up the attachment hierarchy
-			for (USceneComponent* Curr = AttachComponent; Curr; Curr = Curr->GetAttachParent())
+			USkeletalMeshComponent* ParentComp = Cast<USkeletalMeshComponent>(Curr);
+			if (ParentComp && !ParentComp->IsPendingKill())
 			{
-				USkeletalMeshComponent* ParentComp = Cast<USkeletalMeshComponent>(Curr);
-				if (ParentComp && !ParentComp->IsPendingKill())
-				{
-					FoundSkelComp = ParentComp;
-					break;
-				}
+				FoundSkelComp = ParentComp;
+				break;
 			}
-			
-			if (!FoundSkelComp)
+		}
+
+		if (!FoundSkelComp)
+		{
+			// Next, try to find one in our outer chain
+			USkeletalMeshComponent* OuterComp = AttachComponent->GetTypedOuter<USkeletalMeshComponent>();
+			if (OuterComp && !OuterComp->IsPendingKill())
 			{
-				// Next, try to find one in our outer chain
-				USkeletalMeshComponent* OuterComp = AttachComponent->GetTypedOuter<USkeletalMeshComponent>();
-				if (OuterComp && !OuterComp->IsPendingKill())
-				{
-					FoundSkelComp = OuterComp;
-				}
-				else if (AActor* Actor = AttachComponent->GetAttachmentRootActor())
-				{
-					// Final fall-back, look for any mesh component on our root actor or any of its parents
-					FoundSkelComp = FindActorSkelMeshComponent(Actor, true);
-				}
+				FoundSkelComp = OuterComp;
+			}
+			else if (AActor* Actor = AttachComponent->GetAttachmentRootActor())
+			{
+				// Final fall-back, look for any mesh component on our root actor or any of its parents
+				FoundSkelComp = FindActorSkelMeshComponent(Actor, true);
 			}
 		}
 	}
@@ -1811,13 +1808,21 @@ USkeletalMesh* UNiagaraDataInterfaceSkeletalMesh::GetSkeletalMesh(FNiagaraSystem
 	else if (!SystemInstance || !SystemInstance->GetWorld()->IsGameWorld())
 	{
 		// NOTE: We don't fall back on the preview mesh if we have a valid skeletal mesh component referenced
-		Mesh = PreviewMesh.LoadSynchronous();		
+		Mesh = PreviewMesh.LoadSynchronous();
 	}
 #endif
 
 	return Mesh;
 }
 
+USkeletalMesh* UNiagaraDataInterfaceSkeletalMesh::GetSkeletalMesh(UNiagaraComponent* Component)
+{
+	// NOTE: We don't need the system instance when not initializing instance data, and when using a UNiagaraComponent, it is always the attach component
+	TWeakObjectPtr<USceneComponent> SceneComponent;
+	USkeletalMeshComponent* FoundSkelComp = nullptr;
+	return GetSkeletalMesh(nullptr, Component, SceneComponent, FoundSkelComp, nullptr);
+}
+
 bool FNDISkeletalMesh_InstanceData::Init(UNiagaraDataInterfaceSkeletalMesh* Interface, FNiagaraSystemInstance* SystemInstance)
 {
 	check(Interface);
@@ -1843,17 +1848,18 @@ bool FNDISkeletalMesh_InstanceData::Init(UNiagaraDataInterfaceSkeletalMesh* Inte
 
 	// Get skel mesh and confirm have valid data
 	USkeletalMeshComponent* NewSkelComp = nullptr;
-	USkeletalMesh* Mesh = Interface->GetSkeletalMesh(SystemInstance, SceneComponent, NewSkelComp, this);
-	
+	USceneComponent* AttachComponent = SystemInstance->GetAttachComponent();
+	USkeletalMesh* Mesh = Interface->GetSkeletalMesh(SystemInstance, AttachComponent, SceneComponent, NewSkelComp, this);
+
 	SkeletalMesh = Mesh;
 	bMeshValid = Mesh != nullptr;
 	bComponentValid = SceneComponent.IsValid();
 
 	Transform = (bComponentValid ? SceneComponent->GetComponentToWorld() : SystemInstance->GetWorldTransform()).ToMatrixWithScale();
 	TransformInverseTransposed = Transform.Inverse().GetTransposed();
-	PrevTransform = Transform;	
+	PrevTransform = Transform;
 
-	if (USceneComponent* AttachComponent = SystemInstance->GetAttachComponent())
+	if (AttachComponent)
 	{
 		CachedAttachParent = AttachComponent->GetAttachParent();
 	}
@@ -1862,7 +1868,7 @@ bool FNDISkeletalMesh_InstanceData::Init(UNiagaraDataInterfaceSkeletalMesh* Inte
 	if (Mesh != nullptr)
 	{
 		// HACK! This only works on systems created by a Niagara component...should maybe move somewhere else to cover non-component systems
-		if (UNiagaraComponent* NiagaraComponent = Cast<UNiagaraComponent>(SystemInstance->GetAttachComponent()))
+		if (UNiagaraComponent* NiagaraComponent = Cast<UNiagaraComponent>(AttachComponent))
 		{
 			Mesh->GetOnMeshChanged().AddUObject(NiagaraComponent, &UNiagaraComponent::ReinitializeSystem);
 			if (USkeleton* Skeleton = Mesh->GetSkeleton())
@@ -2232,7 +2238,7 @@ bool FNDISkeletalMesh_InstanceData::ResetRequired(UNiagaraDataInterfaceSkeletalM
 		{
 			return true;
 		}
-	}		
+	}
 	else if (Interface->SourceComponent)
 	{
 		// Reset if the source component changed (or there wasn't one and now there is)
@@ -2253,7 +2259,7 @@ bool FNDISkeletalMesh_InstanceData::ResetRequired(UNiagaraDataInterfaceSkeletalM
 			return true;
 		}
 	}
-	
+
 	// Reset if the LOD we relied on was streamed out, or if the LOD we need could now be available.
 	if (SkelMesh != nullptr)
 	{
@@ -2282,7 +2288,7 @@ bool FNDISkeletalMesh_InstanceData::ResetRequired(UNiagaraDataInterfaceSkeletalM
 	{
 		return true;
 	}
-	
+
 	return false;
 }
 
@@ -2479,7 +2485,7 @@ void UNiagaraDataInterfaceSkeletalMesh::GetFunctions(TArray<FNiagaraFunctionSign
 void UNiagaraDataInterfaceSkeletalMesh::GetVMExternalFunction(const FVMExternalFunctionBindingInfo& BindingInfo, void* InstanceData, FVMExternalFunction &OutFunc)
 {
 	FNDISkeletalMesh_InstanceData* InstData = (FNDISkeletalMesh_InstanceData*)InstanceData;
-	
+
 	if (!InstData)
 	{
 		OutFunc = FVMExternalFunction();
@@ -2584,7 +2590,7 @@ void UNiagaraDataInterfaceSkeletalMesh::DestroyPerInstanceData(void* PerInstance
 
 #if WITH_EDITOR
 	if (USkeletalMesh* SkeletalMesh = Inst->SkeletalMesh.Get())
-	{		
+	{
 		if (UNiagaraComponent* NiagaraComponent = Cast<UNiagaraComponent>(SystemInstance->GetAttachComponent()))
 		{
 			SkeletalMesh->GetOnMeshChanged().RemoveAll(NiagaraComponent);
@@ -2618,7 +2624,7 @@ bool UNiagaraDataInterfaceSkeletalMesh::PerInstanceTick(void* PerInstanceData, F
 	return Inst->Tick(this, SystemInstance, InDeltaSeconds);
 }
 
-#if WITH_EDITOR	
+#if WITH_EDITOR
 void UNiagaraDataInterfaceSkeletalMesh::GetFeedback(UNiagaraSystem* Asset, UNiagaraComponent* Component, TArray<FNiagaraDataInterfaceError>& OutErrors,
 	TArray<FNiagaraDataInterfaceFeedback>& OutWarnings, TArray<FNiagaraDataInterfaceFeedback>& OutInfo)
 {
@@ -2632,10 +2638,7 @@ void UNiagaraDataInterfaceSkeletalMesh::GetFeedback(UNiagaraSystem* Asset, UNiag
 
 	// Collect Errors
 #if WITH_EDITORONLY_DATA
-	FNiagaraSystemInstance* SystemInstance = Component ? Component->GetSystemInstance() : nullptr;
-	TWeakObjectPtr<USceneComponent> SceneComponent;
-	USkeletalMeshComponent* SkelMeshComponent = nullptr;
-	USkeletalMesh* SkelMesh = GetSkeletalMesh(SystemInstance, SceneComponent, SkelMeshComponent);
+	USkeletalMesh* SkelMesh = GetSkeletalMesh(Component);
 
 	if (SkelMesh != nullptr)
 	{
@@ -2651,7 +2654,7 @@ void UNiagaraDataInterfaceSkeletalMesh::GetFeedback(UNiagaraSystem* Asset, UNiag
 
 		// Check for the possibility that this mesh won't behave properly because of no CPU access
 		if (!bHasCPUAccess)
-		{			
+		{
 			// Collect all scripts used by the system
 			// NOTE: We don't descriminate between CPU or GPU scripts here because while GPU access will "Just Work"
 			// on some platforms, other platforms (like Mobile or OpenGL) do not create a shader resource view for the
@@ -2661,10 +2664,10 @@ void UNiagaraDataInterfaceSkeletalMesh::GetFeedback(UNiagaraSystem* Asset, UNiag
 			Scripts.Add(Asset->GetSystemUpdateScript());
 			for (auto&& EmitterHandle : Asset->GetEmitterHandles())
 			{
-				EmitterHandle.GetInstance()->GetScripts(Scripts, false);				
+				EmitterHandle.GetInstance()->GetScripts(Scripts, false);
 			}
 
-			// Now check if any script uses functions that require CPU access			
+			// Now check if any script uses functions that require CPU access
 			// TODO: This isn't complete enough. It doesn't guarantee that the DI used by these functions are THIS DI.
 			// Finding that information out is currently non-trivial so just pop a warning with the possibility of false
 			// positives
@@ -2695,7 +2698,7 @@ void UNiagaraDataInterfaceSkeletalMesh::GetFeedback(UNiagaraSystem* Asset, UNiag
 					}
 				}
 				return false;
-			}();			
+			}();
 		}
 	}
 	else
@@ -3141,12 +3144,12 @@ bool UNiagaraDataInterfaceSkeletalMesh::GetFunctionHLSL(const FNiagaraDataInterf
 	{
 		static const TCHAR* FormatSample = TEXT("void {InstanceFunctionName} (in int SocketIndex, out int Bone) { {GetDISkelMeshContextName} DISkelMesh_GetFilteredSocketBoneAt(DIContext, SocketIndex, Bone); }");
 		OutHLSL += FString::Format(FormatSample, ArgsSample);
-	} 
+	}
 	else if (FunctionInfo.DefinitionName == FSkeletalMeshInterfaceHelper::GetFilteredSocketTransformName)
 	{
 		static const TCHAR* FormatSample = TEXT("void {InstanceFunctionName} (in int SocketIndex, in bool bApplyTransform, out float3 OutTranslation, out float4 OutRotation, out float3 OutScale) { {GetDISkelMeshContextName} DISkelMesh_GetFilteredSocketTransform(DIContext, SocketIndex, bApplyTransform, OutTranslation, OutRotation, OutScale); }");
 		OutHLSL += FString::Format(FormatSample, ArgsSample);
-	} 
+	}
 	else if (FunctionInfo.DefinitionName == FSkeletalMeshInterfaceHelper::RandomFilteredSocketName)
 	{
 		static const TCHAR* FormatSample = TEXT("void {InstanceFunctionName} (NiagaraRandInfo InRandomInfo, out int SocketBone) { {GetDISkelMeshContextName} DISkelMesh_RandomFilteredSocket(DIContext, InRandomInfo.Seed1, InRandomInfo.Seed2, InRandomInfo.Seed3, SocketBone); }");
@@ -3379,7 +3382,7 @@ void UNiagaraDataInterfaceSkeletalMesh::GetParameterDefinitionHLSL(const FNiagar
 
 void UNiagaraDataInterfaceSkeletalMesh::SetSourceComponentFromBlueprints(USkeletalMeshComponent* ComponentToUse)
 {
-	// NOTE: When ChangeId changes the next tick will be skipped and a reset of the per-instance data will be initiated. 
+	// NOTE: When ChangeId changes the next tick will be skipped and a reset of the per-instance data will be initiated.
 	ChangeId++;
 	SourceComponent = ComponentToUse;
 	Source = ComponentToUse->GetOwner();
@@ -3387,7 +3390,7 @@ void UNiagaraDataInterfaceSkeletalMesh::SetSourceComponentFromBlueprints(USkelet
 
 void UNiagaraDataInterfaceSkeletalMesh::SetSamplingRegionsFromBlueprints(const TArray<FName>& InSamplingRegions)
 {
-	// NOTE: When ChangeId changes the next tick will be skipped and a reset of the per-instance data will be initiated. 
+	// NOTE: When ChangeId changes the next tick will be skipped and a reset of the per-instance data will be initiated.
 	ChangeId++;
 	SamplingRegions = InSamplingRegions;
 }
@@ -3434,7 +3437,7 @@ int32 UNiagaraDataInterfaceSkeletalMesh::CalculateLODIndexAndSamplingRegions(USk
 				const FSkeletalMeshSamplingRegion& Region = SamplingInfo.GetRegion(RegionIdx);
 				const FSkeletalMeshSamplingRegionBuiltData& RegionBuiltData = SamplingInfo.GetRegionBuiltData(RegionIdx);
 
-				const int32 RegionLODIndex = Region.LODIndex == INDEX_NONE ? 
+				const int32 RegionLODIndex = Region.LODIndex == INDEX_NONE ?
 					(InMesh->GetLODNum() - 1) :
 					(FMath::Clamp(Region.LODIndex, 0, InMesh->GetLODNum() - 1));
 
@@ -3450,7 +3453,7 @@ int32 UNiagaraDataInterfaceSkeletalMesh::CalculateLODIndexAndSamplingRegions(USk
 						*RegionName.ToString());
 					return INDEX_NONE;
 				}
-				
+
 				if (RegionBuiltData.TriangleIndices.Num() > 0)
 				{
 					OutSamplingRegionIndices.Add(RegionIdx);
@@ -3507,7 +3510,7 @@ void FSkeletalMeshAccessorHelper::Init<TNDISkelMesh_FilterModeSingle, TNDISkelMe
 		SamplingRegion = nullptr;
 		SamplingRegionBuiltData = nullptr;
 	}
-	
+
 	if (SkinningData != nullptr)
 	{
 		SkinningData->EnterRead();
diff --git a/Engine/Plugins/FX/Niagara/Source/Niagara/Private/NiagaraDebugHud.cpp b/Engine/Plugins/FX/Niagara/Source/Niagara/Private/NiagaraDebugHud.cpp
--- a/Engine/Plugins/FX/Niagara/Source/Niagara/Private/NiagaraDebugHud.cpp
+++ b/Engine/Plugins/FX/Niagara/Source/Niagara/Private/NiagaraDebugHud.cpp
@@ -278,12 +278,12 @@ namespace NiagaraDebugLocal
 			if ( *Text == '\n' )
 			{
 				CurrSize.X = 0.0f;
-				CurrSize.Y = CurrSize.Y + fAdvanceHeight; 
+				CurrSize.Y = CurrSize.Y + fAdvanceHeight;
 				PrevChar = nullptr;
 				++Text;
 				continue;
 			}
-			
+
 			float TmpWidth, TmpHeight;
 			Font->GetCharSize(*Text, TmpWidth, TmpHeight);
 
@@ -409,7 +409,7 @@ FNiagaraDebugHud::~FNiagaraDebugHud()
 void FNiagaraDebugHud::UpdateSettings(const FNiagaraDebugHUDSettingsData& NewSettings)
 {
 	using namespace NiagaraDebugLocal;
-	
+
 	FNiagaraDebugHUDSettingsData::StaticStruct()->CopyScriptStruct(&Settings, &NewSettings);
 	GCachedSystemVariables.Empty();
 }
@@ -470,7 +470,8 @@ void FNiagaraDebugHud::GatherSystemInfo()
 			continue;
 		}
 
-		FNiagaraSystemInstance* SystemInstance = NiagaraComponent->GetSystemInstance();
+		FNiagaraSystemInstanceControllerPtr SystemInstanceController = NiagaraComponent->GetSystemInstanceController();
+		FNiagaraSystemInstance* SystemInstance = SystemInstanceController.IsValid() ? SystemInstanceController->GetSystemInstance_Unsafe() : nullptr;
 		if (!SystemInstance)
 		{
 			continue;
@@ -977,7 +978,8 @@ void FNiagaraDebugHud::DrawValidation(class FNiagaraWorldManager* WorldManager,
 		}
 
 		UNiagaraSystem* NiagaraSystem = NiagaraComponent->GetAsset();
-		FNiagaraSystemInstance* SystemInstance = NiagaraComponent->GetSystemInstance();
+		FNiagaraSystemInstanceControllerPtr SystemInstanceController = NiagaraComponent->GetSystemInstanceController();
+		FNiagaraSystemInstance* SystemInstance = SystemInstanceController.IsValid() ? SystemInstanceController->GetSystemInstance_Unsafe() : nullptr;
 		if ((NiagaraSystem == nullptr) || (SystemInstance == nullptr))
 		{
 			continue;
@@ -1152,8 +1154,9 @@ void FNiagaraDebugHud::DrawComponents(FNiagaraWorldManager* WorldManager, UCanva
 		}
 
 		UNiagaraSystem* NiagaraSystem = NiagaraComponent->GetAsset();
-		FNiagaraSystemInstance* SystemInstance = NiagaraComponent->GetSystemInstance();
-		if ((NiagaraSystem == nullptr) || (SystemInstance == nullptr))
+		FNiagaraSystemInstanceControllerPtr SystemInstanceController = NiagaraComponent->GetSystemInstanceController();
+		FNiagaraSystemInstance* SystemInstance = SystemInstanceController.IsValid() ? SystemInstanceController->GetSystemInstance_Unsafe() : nullptr;
+		if (NiagaraSystem == nullptr || SystemInstance == nullptr)
 		{
 			continue;
 		}
@@ -1173,7 +1176,7 @@ void FNiagaraDebugHud::DrawComponents(FNiagaraWorldManager* WorldManager, UCanva
 		}
 
 		// Get system simulation
-		auto SystemSimulation = NiagaraComponent->GetSystemSimulation();
+		auto SystemSimulation = SystemInstance->GetSystemSimulation();
 		const bool bSystemSimulationValid = SystemSimulation.IsValid() && SystemSimulation->IsValid();
 		if (bSystemSimulationValid)
 		{
diff --git a/Engine/Plugins/FX/Niagara/Source/Niagara/Private/NiagaraEffectType.cpp b/Engine/Plugins/FX/Niagara/Source/Niagara/Private/NiagaraEffectType.cpp
--- a/Engine/Plugins/FX/Niagara/Source/Niagara/Private/NiagaraEffectType.cpp
+++ b/Engine/Plugins/FX/Niagara/Source/Niagara/Private/NiagaraEffectType.cpp
@@ -90,7 +90,7 @@ const FNiagaraEmitterScalabilitySettings& UNiagaraEffectType::GetActiveEmitterSc
 			return Settings;
 		}
 	}
-	
+
 	//UE_LOG(LogNiagara, Warning, TEXT("Could not find active emitter scalability settings for EffectType %s"), *GetFullName());
 
 	static FNiagaraEmitterScalabilitySettings Dummy;
@@ -204,7 +204,7 @@ void UNiagaraSignificanceHandlerDistance::CalculateSignificance(TArray<UNiagaraC
 		FNiagaraScalabilityState& State = OutState[CompIdx];
 
 		const bool AddIndex = !State.bCulled || State.IsDirty();
-		
+
 		if (State.bCulled)
 		{
 			State.Significance = 0.0f;
@@ -221,9 +221,9 @@ void UNiagaraSignificanceHandlerDistance::CalculateSignificance(TArray<UNiagaraC
 			}
 			else
 #endif
-			if(FNiagaraSystemInstance* Inst = Component->GetSystemInstance())
+			if (FNiagaraSystemInstanceControllerConstPtr Controller = Component->GetSystemInstanceController())
 			{
-				LODDistance = Inst->GetLODDistance();
+				LODDistance = Controller->GetLODDistance();
 			}
 
 			State.Significance = 1.0f / LODDistance;
@@ -254,9 +254,9 @@ void UNiagaraSignificanceHandlerAge::CalculateSignificance(TArray<UNiagaraCompon
 		{
 			UNiagaraComponent* Component = Components[CompIdx];
 
-			if (FNiagaraSystemInstance* Inst = Component->GetSystemInstance())
+			if (FNiagaraSystemInstanceControllerConstPtr Controller = Component->GetSystemInstanceController())
 			{
-				State.Significance = 1.0f / Inst->GetAge();//Newer Systems are higher significance.
+				State.Significance = 1.0f / Controller->GetAge();//Newer Systems are higher significance.
 			}
 		}
 
diff --git a/Engine/Plugins/FX/Niagara/Source/Niagara/Private/NiagaraLightRendererProperties.cpp b/Engine/Plugins/FX/Niagara/Source/Niagara/Private/NiagaraLightRendererProperties.cpp
--- a/Engine/Plugins/FX/Niagara/Source/Niagara/Private/NiagaraLightRendererProperties.cpp
+++ b/Engine/Plugins/FX/Niagara/Source/Niagara/Private/NiagaraLightRendererProperties.cpp
@@ -98,10 +98,10 @@ void UNiagaraLightRendererProperties::InitCDOPropertiesAfterModuleStartup()
 	}
 }
 
-FNiagaraRenderer* UNiagaraLightRendererProperties::CreateEmitterRenderer(ERHIFeatureLevel::Type FeatureLevel, const FNiagaraEmitterInstance* Emitter, const UNiagaraComponent* InComponent)
+FNiagaraRenderer* UNiagaraLightRendererProperties::CreateEmitterRenderer(ERHIFeatureLevel::Type FeatureLevel, const FNiagaraEmitterInstance* Emitter, const FNiagaraSystemInstanceController& InController)
 {
 	FNiagaraRenderer* NewRenderer = new FNiagaraRendererLights(FeatureLevel, this, Emitter);
-	NewRenderer->Initialize(this, Emitter, InComponent);
+	NewRenderer->Initialize(this, Emitter, InController);
 	return NewRenderer;
 }
 
diff --git a/Engine/Plugins/FX/Niagara/Source/Niagara/Private/NiagaraMeshRendererProperties.cpp b/Engine/Plugins/FX/Niagara/Source/Niagara/Private/NiagaraMeshRendererProperties.cpp
--- a/Engine/Plugins/FX/Niagara/Source/Niagara/Private/NiagaraMeshRendererProperties.cpp
+++ b/Engine/Plugins/FX/Niagara/Source/Niagara/Private/NiagaraMeshRendererProperties.cpp
@@ -6,6 +6,7 @@
 #include "NiagaraConstants.h"
 #include "NiagaraBoundsCalculatorHelper.h"
 #include "NiagaraCustomVersion.h"
+#include "NiagaraEmitterInstance.h"
 #include "Modules/ModuleManager.h"
 
 #if WITH_EDITOR
@@ -103,7 +104,7 @@ UNiagaraMeshRendererProperties::UNiagaraMeshRendererProperties()
 	AttributeBindings.Add(&MeshIndexBinding);
 }
 
-FNiagaraRenderer* UNiagaraMeshRendererProperties::CreateEmitterRenderer(ERHIFeatureLevel::Type FeatureLevel, const FNiagaraEmitterInstance* Emitter, const UNiagaraComponent* InComponent)
+FNiagaraRenderer* UNiagaraMeshRendererProperties::CreateEmitterRenderer(ERHIFeatureLevel::Type FeatureLevel, const FNiagaraEmitterInstance* Emitter, const FNiagaraSystemInstanceController& InController)
 {
 	for (const auto& MeshProperties : Meshes)
 	{
@@ -111,7 +112,7 @@ FNiagaraRenderer* UNiagaraMeshRendererProperties::CreateEmitterRenderer(ERHIFeat
 		{
 			// There's at least one valid mesh
 			FNiagaraRenderer* NewRenderer = new FNiagaraRendererMeshes(FeatureLevel, this, Emitter);
-			NewRenderer->Initialize(this, Emitter, InComponent);
+			NewRenderer->Initialize(this, Emitter, InController);
 			return NewRenderer;
 		}
 	}
diff --git a/Engine/Plugins/FX/Niagara/Source/Niagara/Private/NiagaraRenderer.cpp b/Engine/Plugins/FX/Niagara/Source/Niagara/Private/NiagaraRenderer.cpp
--- a/Engine/Plugins/FX/Niagara/Source/Niagara/Private/NiagaraRenderer.cpp
+++ b/Engine/Plugins/FX/Niagara/Source/Niagara/Private/NiagaraRenderer.cpp
@@ -6,10 +6,13 @@
 #include "NiagaraDataSet.h"
 #include "NiagaraStats.h"
 #include "NiagaraVertexFactory.h"
+#include "NiagaraComponent.h"
 #include "Engine/Engine.h"
 #include "DynamicBufferAllocator.h"
 #include "NiagaraEmitterInstanceBatcher.h"
 #include "NiagaraGPUSortInfo.h"
+#include "NiagaraEmitterInstance.h"
+#include "NiagaraSystemInstanceController.h"
 #include "Materials/MaterialInstanceDynamic.h"
 
 DECLARE_CYCLE_STAT(TEXT("Sort Particles"), STAT_NiagaraSortParticles, STATGROUP_Niagara);
@@ -284,7 +287,7 @@ FNiagaraRenderer::FNiagaraRenderer(ERHIFeatureLevel::Type InFeatureLevel, const
 #endif
 }
 
-void FNiagaraRenderer::Initialize(const UNiagaraRendererProperties *InProps, const FNiagaraEmitterInstance* Emitter, const UNiagaraComponent* InComponent)
+void FNiagaraRenderer::Initialize(const UNiagaraRendererProperties* InProps, const FNiagaraEmitterInstance* Emitter, const FNiagaraSystemInstanceController& InController)
 {
 	//Get our list of valid base materials. Fall back to default material if they're not valid.
 	InProps->GetUsedMaterials(Emitter, BaseMaterials_GT);
@@ -299,17 +302,9 @@ void FNiagaraRenderer::Initialize(const UNiagaraRendererProperties *InProps, con
 		}
 		else if (Mat && bCreateMidsForUsedMaterials && !Mat->IsA<UMaterialInstanceDynamic>())
 		{
-			const UNiagaraComponent* Comp = InComponent;
-			for (const FNiagaraMaterialOverride& Override : Comp->EmitterMaterials)
+			if (UMaterialInterface* Override = InController.GetMaterialOverride(InProps, Index))
 			{
-				if (Override.EmitterRendererProperty == InProps)
-				{
-					if (Index == Override.MaterialSubIndex)
-					{
-						Mat = Override.Material;
-						continue;
-					}
-				}
+				Mat = Override;
 			}
 		}
 
diff --git a/Engine/Plugins/FX/Niagara/Source/Niagara/Private/NiagaraRendererComponents.cpp b/Engine/Plugins/FX/Niagara/Source/Niagara/Private/NiagaraRendererComponents.cpp
--- a/Engine/Plugins/FX/Niagara/Source/Niagara/Private/NiagaraRendererComponents.cpp
+++ b/Engine/Plugins/FX/Niagara/Source/Niagara/Private/NiagaraRendererComponents.cpp
@@ -5,6 +5,7 @@
 #include "NiagaraDataSet.h"
 #include "NiagaraStats.h"
 #include "NiagaraComponentRendererProperties.h"
+#include "NiagaraSystemInstance.h"
 #include "Async/Async.h"
 
 #if WITH_EDITOR
@@ -179,7 +180,7 @@ void FNiagaraRendererComponents::DestroyRenderState_Concurrent()
 		[this, Pool_GT = MoveTemp(ComponentPool), Owner_GT = MoveTemp(SpawnedOwner)]()
 		{
 			// we do not reset ParticlesWithComponents here because it's possible the render state is destroyed without destroying the renderer. In this case we want to know which particles
-			// had spawned some components previously 
+			// had spawned some components previously
 			for (auto& PoolEntry : Pool_GT)
 			{
 				if (PoolEntry.Component.IsValid())
@@ -325,7 +326,7 @@ void FNiagaraRendererComponents::PostSystemTick_GameThread(const UNiagaraRendere
 			// Get the particle ID and see if we have any components already assigned to the particle
 			ParticleID = UniqueIDAccessor.GetSafe(ParticleIndex, -1);
 			ParticlesWithComponents.RemoveAndCopyValue(ParticleID, PoolIndex);
-			
+
 			if (PoolIndex == -1 && Properties->bOnlyCreateComponentsOnParticleSpawn)
 			{
 				// Don't allow this particle to acquire a component unless it was just spawned or had a component assigned to it previously
@@ -422,7 +423,7 @@ void FNiagaraRendererComponents::PostSystemTick_GameThread(const UNiagaraRendere
 		{
 			Emitter->SetParticleComponentActive(ComponentKey, ParticleID);
 		}
-		
+
 		++ComponentCount;
 		if (ComponentCount > GNiagaraWarnComponentRenderCount)
 		{
@@ -443,7 +444,7 @@ void FNiagaraRendererComponents::PostSystemTick_GameThread(const UNiagaraRendere
 			break;
 		}
 	}
-		
+
 	if (ComponentCount < ComponentPool.Num())
 	{
 		// go over the pooled components we didn't need this tick to see if we can destroy some and deactivate the rest
@@ -462,8 +463,8 @@ void FNiagaraRendererComponents::PostSystemTick_GameThread(const UNiagaraRendere
 			{
 				continue;
 			}
-		
-			USceneComponent* Component = PoolEntry.Component.Get();		
+
+			USceneComponent* Component = PoolEntry.Component.Get();
 			if (!Component || (CurrentTime - PoolEntry.LastActiveTime) >= GNiagaraComponentRenderPoolInactiveTimeLimit)
 			{
 				if (Component)
@@ -608,7 +609,7 @@ void FNiagaraRendererComponents::OnObjectsReplacedCallback(const TMap<UObject*,
 		{
 			ResetComponentPool(false);
 			break;
-		}		
+		}
 	}
 }
 
diff --git a/Engine/Plugins/FX/Niagara/Source/Niagara/Private/NiagaraRendererLights.cpp b/Engine/Plugins/FX/Niagara/Source/Niagara/Private/NiagaraRendererLights.cpp
--- a/Engine/Plugins/FX/Niagara/Source/Niagara/Private/NiagaraRendererLights.cpp
+++ b/Engine/Plugins/FX/Niagara/Source/Niagara/Private/NiagaraRendererLights.cpp
@@ -6,6 +6,7 @@
 #include "NiagaraDataSetAccessor.h"
 #include "NiagaraStats.h"
 #include "NiagaraVertexFactory.h"
+#include "NiagaraComponent.h"
 #include "Engine/Engine.h"
 
 #include "NiagaraLightRendererProperties.h"
@@ -20,7 +21,7 @@ struct FNiagaraDynamicDataLights : public FNiagaraDynamicDataBase
 		: FNiagaraDynamicDataBase(InEmitter)
 	{
 	}
-	
+
 	TArray<FNiagaraRendererLights::SimpleLightData> LightArray;
 };
 
@@ -108,7 +109,7 @@ FNiagaraDynamicDataBase* FNiagaraRendererLights::GenerateDynamicData(const FNiag
 
 			const FLinearColor Color = ColorReader.GetSafe(ParticleIndex, DefaultColor);
 			const float Brightness = Properties->bAlphaScalesBrightness ? Color.A : 1.0f;
-			
+
 			LightData.LightEntry.Radius = LightRadius;
 			LightData.LightEntry.Color = FVector(Color) * Brightness + Properties->ColorAdd;
 			LightData.LightEntry.Exponent = Properties->bUseInverseSquaredFalloff ? 0 : ExponentReader.GetSafe(ParticleIndex, DefaultExponent);
@@ -136,7 +137,7 @@ void FNiagaraRendererLights::GatherSimpleLights(FSimpleLightArray& OutParticleLi
 
 		for (const FNiagaraRendererLights::SimpleLightData &LightData : DynamicData->LightArray)
 		{
-			// When not using camera-offset, output one position for all views to share. 
+			// When not using camera-offset, output one position for all views to share.
 			OutParticleLights.PerViewData.Add(LightData.PerViewEntry);
 
 			// Add an entry for the light instance.
diff --git a/Engine/Plugins/FX/Niagara/Source/Niagara/Private/NiagaraRendererMeshes.cpp b/Engine/Plugins/FX/Niagara/Source/Niagara/Private/NiagaraRendererMeshes.cpp
--- a/Engine/Plugins/FX/Niagara/Source/Niagara/Private/NiagaraRendererMeshes.cpp
+++ b/Engine/Plugins/FX/Niagara/Source/Niagara/Private/NiagaraRendererMeshes.cpp
@@ -3,7 +3,9 @@
 #include "NiagaraRendererMeshes.h"
 #include "ParticleResources.h"
 #include "NiagaraDataSet.h"
+#include "NiagaraEmitterInstance.h"
 #include "NiagaraStats.h"
+#include "NiagaraComponent.h"
 #include "Async/ParallelFor.h"
 #include "Engine/StaticMesh.h"
 #include "NiagaraEmitterInstanceBatcher.h"
@@ -144,9 +146,9 @@ FNiagaraRendererMeshes::~FNiagaraRendererMeshes()
 {
 }
 
-void FNiagaraRendererMeshes::Initialize(const UNiagaraRendererProperties* InProps, const FNiagaraEmitterInstance* Emitter, const UNiagaraComponent* InComponent)
+void FNiagaraRendererMeshes::Initialize(const UNiagaraRendererProperties* InProps, const FNiagaraEmitterInstance* Emitter, const FNiagaraSystemInstanceController& InController)
 {
-	FNiagaraRenderer::Initialize(InProps, Emitter, InComponent);
+	FNiagaraRenderer::Initialize(InProps, Emitter, InController);
 
 	check(Emitter);
 	check(InProps);
diff --git a/Engine/Plugins/FX/Niagara/Source/Niagara/Private/NiagaraRendererRibbons.cpp b/Engine/Plugins/FX/Niagara/Source/Niagara/Private/NiagaraRendererRibbons.cpp
--- a/Engine/Plugins/FX/Niagara/Source/Niagara/Private/NiagaraRendererRibbons.cpp
+++ b/Engine/Plugins/FX/Niagara/Source/Niagara/Private/NiagaraRendererRibbons.cpp
@@ -6,6 +6,7 @@
 #include "NiagaraDataSet.h"
 #include "NiagaraDataSetAccessor.h"
 #include "NiagaraStats.h"
+#include "NiagaraComponent.h"
 #include "RayTracingDefinitions.h"
 #include "RayTracingDynamicGeometryCollection.h"
 #include "RayTracingInstance.h"
@@ -215,7 +216,7 @@ void FNiagaraRendererRibbons::ReleaseRenderThreadResources()
 	FNiagaraRenderer::ReleaseRenderThreadResources();
 #if RHI_RAYTRACING
 	if (IsRayTracingEnabled())
-	{	
+	{
 		RayTracingGeometry.ReleaseResource();
 		RayTracingDynamicVertexBuffer.Release();
 	}
@@ -265,7 +266,7 @@ TValue* FNiagaraRendererRibbons::AppendToIndexBuffer(
 	// This sets up the first and next vertex for each pair of triangles in the slice.
 	// For a plane this will just be a linear set
 	// For a multiplane it will be multiple separate linear sets
-	// For a tube it will be a linear set that wraps back around to itself, 
+	// For a tube it will be a linear set that wraps back around to itself,
 	// Same with the custom vertices.
 	TArray<int32> SliceTriangleToVertexIds;
 
@@ -322,7 +323,7 @@ TValue* FNiagaraRendererRibbons::AppendToIndexBuffer(
 
 			for (int32 TriangleId = 0; TriangleId < SliceTriangleToVertexIds.Num(); TriangleId += 2)
 			{
-				// Switch geometry layout based on above or below the centerline. 
+				// Switch geometry layout based on above or below the centerline.
 				// This has the effect of mirroring the triangle layout across the center
 				// except when it's an odd number of segments, then the center segment doesn't mirror.
 				bool bShouldFlipGeometry = TriangleId < (SliceTriangleToVertexIds.Num() / 2);
@@ -368,11 +369,11 @@ TValue* FNiagaraRendererRibbons::AppendToIndexBuffer(
 
 template <typename TValue>
 void FNiagaraRendererRibbons::GenerateIndexBuffer(
-	FGlobalDynamicIndexBuffer::FAllocationEx& InOutIndexAllocation, 
+	FGlobalDynamicIndexBuffer::FAllocationEx& InOutIndexAllocation,
 	const FRibbonRenderingIndexOffsets& Offsets,
-	int32 InterpCount, 
-	const FVector& ViewDirection, 
-	const FVector& ViewOriginForDistanceCulling, 
+	int32 InterpCount,
+	const FVector& ViewDirection,
+	const FVector& ViewOriginForDistanceCulling,
 	FNiagaraDynamicDataRibbon* DynamicData) const
 {
 	check(DynamicData);
@@ -843,12 +844,12 @@ FNiagaraDynamicDataBase* FNiagaraRendererRibbons::GenerateDynamicData(const FNia
 			{
 				SegmentData.Add(SegmentIndex);
 			}
-			
+
 			float U0Offset;
 			float U0Scale;
 			float U0DistributionScaler;
 			if(UV0Settings.bEnablePerParticleUOverride && U0OverrideIsBound)
-			{ 
+			{
 				U0Offset = 0;
 				U0Scale = 1.0f;
 				U0DistributionScaler = 1;
@@ -1195,11 +1196,11 @@ FNiagaraRendererRibbons::FRibbonRenderingIndexOffsets FNiagaraRendererRibbons::C
 }
 
 void FNiagaraRendererRibbons::CreatePerViewResources(
-	const FSceneView* View, 
-	const FSceneViewFamily& ViewFamily, 
-	const FNiagaraSceneProxy* SceneProxy, 
+	const FSceneView* View,
+	const FSceneViewFamily& ViewFamily,
+	const FNiagaraSceneProxy* SceneProxy,
 	FMeshElementCollector& Collector,
-	FNiagaraRibbonUniformBufferRef& OutUniformBuffer, 
+	FNiagaraRibbonUniformBufferRef& OutUniformBuffer,
 	FGlobalDynamicIndexBuffer::FAllocationEx& InOutIndexAllocation) const
 {
 	FNiagaraDynamicDataRibbon* DynamicDataRibbon = static_cast<FNiagaraDynamicDataRibbon*>(DynamicDataRender);
@@ -1338,7 +1339,7 @@ void FNiagaraRendererRibbons::CreatePerViewResources(
 	PerViewUniformParameters.MaterialParam1DataOffset = VFVariables[ENiagaraRibbonVFLayout::MaterialParam1].GetGPUOffset();
 	PerViewUniformParameters.MaterialParam2DataOffset = VFVariables[ENiagaraRibbonVFLayout::MaterialParam2].GetGPUOffset();
 	PerViewUniformParameters.MaterialParam3DataOffset = VFVariables[ENiagaraRibbonVFLayout::MaterialParam3].GetGPUOffset();
-	PerViewUniformParameters.DistanceFromStartOffset = 
+	PerViewUniformParameters.DistanceFromStartOffset =
 		(UV0Settings.DistributionMode == ENiagaraRibbonUVDistributionMode::TiledFromStartOverRibbonLength ||
 		UV1Settings.DistributionMode == ENiagaraRibbonUVDistributionMode::TiledFromStartOverRibbonLength)?
 		VFVariables[ENiagaraRibbonVFLayout::DistanceFromStart].GetGPUOffset() : -1;
diff --git a/Engine/Plugins/FX/Niagara/Source/Niagara/Private/NiagaraRendererSprites.cpp b/Engine/Plugins/FX/Niagara/Source/Niagara/Private/NiagaraRendererSprites.cpp
--- a/Engine/Plugins/FX/Niagara/Source/Niagara/Private/NiagaraRendererSprites.cpp
+++ b/Engine/Plugins/FX/Niagara/Source/Niagara/Private/NiagaraRendererSprites.cpp
@@ -6,6 +6,7 @@
 #include "NiagaraStats.h"
 #include "NiagaraEmitterInstanceBatcher.h"
 #include "NiagaraSortingGPU.h"
+#include "NiagaraComponent.h"
 #include "NiagaraCutoutVertexBuffer.h"
 #include "RayTracingDefinitions.h"
 #include "RayTracingDynamicGeometryCollection.h"
diff --git a/Engine/Plugins/FX/Niagara/Source/Niagara/Private/NiagaraRibbonRendererProperties.cpp b/Engine/Plugins/FX/Niagara/Source/Niagara/Private/NiagaraRibbonRendererProperties.cpp
--- a/Engine/Plugins/FX/Niagara/Source/Niagara/Private/NiagaraRibbonRendererProperties.cpp
+++ b/Engine/Plugins/FX/Niagara/Source/Niagara/Private/NiagaraRibbonRendererProperties.cpp
@@ -5,6 +5,7 @@
 #include "NiagaraConstants.h"
 #include "NiagaraBoundsCalculatorHelper.h"
 #include "NiagaraCustomVersion.h"
+#include "NiagaraEmitterInstance.h"
 #include "Modules/ModuleManager.h"
 #if WITH_EDITOR
 #include "Widgets/Images/SImage.h"
@@ -85,10 +86,10 @@ UNiagaraRibbonRendererProperties::UNiagaraRibbonRendererProperties()
 	AttributeBindings.Add(&V1RangeOverrideBinding);
 }
 
-FNiagaraRenderer* UNiagaraRibbonRendererProperties::CreateEmitterRenderer(ERHIFeatureLevel::Type FeatureLevel, const FNiagaraEmitterInstance* Emitter, const UNiagaraComponent* InComponent)
+FNiagaraRenderer* UNiagaraRibbonRendererProperties::CreateEmitterRenderer(ERHIFeatureLevel::Type FeatureLevel, const FNiagaraEmitterInstance* Emitter, const FNiagaraSystemInstanceController& InController)
 {
 	FNiagaraRenderer* NewRenderer = new FNiagaraRendererRibbons(FeatureLevel, this, Emitter);
-	NewRenderer->Initialize(this, Emitter, InComponent);
+	NewRenderer->Initialize(this, Emitter, InController);
 	return NewRenderer;
 }
 
@@ -260,8 +261,8 @@ void UNiagaraRibbonRendererProperties::CacheFromCompiledData(const FNiagaraDataS
 	{
 		bSortKeyDataSetAccessorIsAge = true;
 		SortKeyDataSetAccessor.Init(CompiledData, NormalizedAgeBinding.GetDataSetBindableVariable().GetName());
-	}	
-	
+	}
+
 	NormalizedAgeAccessor.Init(CompiledData, NormalizedAgeBinding.GetDataSetBindableVariable().GetName());
 
 	PositionDataSetAccessor.Init(CompiledData, PositionBinding.GetDataSetBindableVariable().GetName());
diff --git a/Engine/Plugins/FX/Niagara/Source/Niagara/Private/NiagaraScalabilityManager.cpp b/Engine/Plugins/FX/Niagara/Source/Niagara/Private/NiagaraScalabilityManager.cpp
--- a/Engine/Plugins/FX/Niagara/Source/Niagara/Private/NiagaraScalabilityManager.cpp
+++ b/Engine/Plugins/FX/Niagara/Source/Niagara/Private/NiagaraScalabilityManager.cpp
@@ -194,7 +194,7 @@ bool FNiagaraScalabilityManager::EvaluateCullState(FNiagaraWorldManager* WorldMa
 	//Though this does mean the sorted significance values will be using out of date distances etc.
 	//I'm somewhat on the fence currently as to whether it's better to pay this cost for correctness.
 	const bool UpdateScalability = Component->ScalabilityManagerHandle == ComponentIndex
-		&& (!Context.bNewOnly || Component->GetSystemInstance()->IsPendingSpawn());
+		&& (!Context.bNewOnly || Component->GetSystemInstanceController()->IsPendingSpawn());
 
 	if (UpdateScalability)
 	{
@@ -440,7 +440,7 @@ void FNiagaraScalabilityManager::Update(FNiagaraWorldManager* WorldMan, float De
 
 	const bool bResetUpdate = EffectType->UpdateFrequency == ENiagaraScalabilityUpdateFrequency::Continuous
 		|| ((TimeSinceUpdate >= UpdatePeriod) && !DefaultContext.ComponentRequiresUpdate.Contains(true));
-		
+
 	const int32 ComponentCount = ManagedComponents.Num();
 
 	if (bResetUpdate)
diff --git a/Engine/Plugins/FX/Niagara/Source/Niagara/Private/NiagaraSpriteRendererProperties.cpp b/Engine/Plugins/FX/Niagara/Source/Niagara/Private/NiagaraSpriteRendererProperties.cpp
--- a/Engine/Plugins/FX/Niagara/Source/Niagara/Private/NiagaraSpriteRendererProperties.cpp
+++ b/Engine/Plugins/FX/Niagara/Source/Niagara/Private/NiagaraSpriteRendererProperties.cpp
@@ -7,6 +7,7 @@
 #include "NiagaraRendererSprites.h"
 #include "NiagaraBoundsCalculatorHelper.h"
 #include "NiagaraCustomVersion.h"
+#include "NiagaraEmitterInstance.h"
 #include "Modules/ModuleManager.h"
 #include "NiagaraEmitter.h"
 #include "NiagaraScriptSourceBase.h"
@@ -95,10 +96,10 @@ UNiagaraSpriteRendererProperties::UNiagaraSpriteRendererProperties()
 	AttributeBindings.Add(&RendererVisibilityTagBinding);
 }
 
-FNiagaraRenderer* UNiagaraSpriteRendererProperties::CreateEmitterRenderer(ERHIFeatureLevel::Type FeatureLevel, const FNiagaraEmitterInstance* Emitter, const UNiagaraComponent* InComponent)
+FNiagaraRenderer* UNiagaraSpriteRendererProperties::CreateEmitterRenderer(ERHIFeatureLevel::Type FeatureLevel, const FNiagaraEmitterInstance* Emitter, const FNiagaraSystemInstanceController& InController)
 {
-	FNiagaraRenderer* NewRenderer = new FNiagaraRendererSprites(FeatureLevel, this, Emitter);	
-	NewRenderer->Initialize(this, Emitter, InComponent);
+	FNiagaraRenderer* NewRenderer = new FNiagaraRendererSprites(FeatureLevel, this, Emitter);
+	NewRenderer->Initialize(this, Emitter, InController);
 	return NewRenderer;
 }
 
@@ -227,7 +228,7 @@ void UNiagaraSpriteRendererProperties::InitBindings()
 		MaterialRandomBinding = FNiagaraConstants::GetAttributeDefaultBinding(SYS_PARAM_PARTICLES_MATERIAL_RANDOM);
 		NormalizedAgeBinding = FNiagaraConstants::GetAttributeDefaultBinding(SYS_PARAM_PARTICLES_NORMALIZED_AGE);
 		RendererVisibilityTagBinding = FNiagaraConstants::GetAttributeDefaultBinding(SYS_PARAM_PARTICLES_VISIBILITY_TAG);
-		
+
 		//Default custom sorting to age
 		CustomSortingBinding = FNiagaraConstants::GetAttributeDefaultBinding(SYS_PARAM_PARTICLES_NORMALIZED_AGE);
 	}
@@ -266,7 +267,7 @@ void UNiagaraSpriteRendererProperties::CacheFromCompiledData(const FNiagaraDataS
 	RendererLayoutWithCustomSort.SetVariableFromBinding(CompiledData, PivotOffsetBinding, ENiagaraSpriteVFLayout::PivotOffset);
 	RendererLayoutWithCustomSort.SetVariableFromBinding(CompiledData, NormalizedAgeBinding, ENiagaraSpriteVFLayout::NormalizedAge);
 	RendererLayoutWithCustomSort.SetVariableFromBinding(CompiledData, MaterialRandomBinding, ENiagaraSpriteVFLayout::MaterialRandom);
-	RendererLayoutWithCustomSort.SetVariableFromBinding(CompiledData, CustomSortingBinding, ENiagaraSpriteVFLayout::CustomSorting);	
+	RendererLayoutWithCustomSort.SetVariableFromBinding(CompiledData, CustomSortingBinding, ENiagaraSpriteVFLayout::CustomSorting);
 	MaterialParamValidMask  = RendererLayoutWithCustomSort.SetVariableFromBinding(CompiledData, DynamicMaterialBinding, ENiagaraSpriteVFLayout::MaterialParam0) ? 0x1 : 0;
 	MaterialParamValidMask |= RendererLayoutWithCustomSort.SetVariableFromBinding(CompiledData, DynamicMaterial1Binding, ENiagaraSpriteVFLayout::MaterialParam1) ? 0x2 : 0;
 	MaterialParamValidMask |= RendererLayoutWithCustomSort.SetVariableFromBinding(CompiledData, DynamicMaterial2Binding, ENiagaraSpriteVFLayout::MaterialParam2) ? 0x4 : 0;
@@ -318,7 +319,7 @@ void UNiagaraSpriteRendererProperties::CacheFromCompiledData(const FNiagaraDataS
 }
 
 
-bool UNiagaraSpriteRendererProperties::PopulateRequiredBindings(FNiagaraParameterStore& InParameterStore) 
+bool UNiagaraSpriteRendererProperties::PopulateRequiredBindings(FNiagaraParameterStore& InParameterStore)
 {
 	bool bAnyAdded = false;
 
@@ -385,7 +386,7 @@ bool UNiagaraSpriteRendererProperties::IsSupportedVariableForBinding(const FNiag
 //}
 //#endif
 
-void UNiagaraSpriteRendererProperties::PostEditChangeProperty(struct FPropertyChangedEvent& PropertyChangedEvent) 
+void UNiagaraSpriteRendererProperties::PostEditChangeProperty(struct FPropertyChangedEvent& PropertyChangedEvent)
 {
 	SubImageSize.X = FMath::Max<float>(SubImageSize.X, 1.f);
 	SubImageSize.Y = FMath::Max<float>(SubImageSize.Y, 1.f);
@@ -393,8 +394,7 @@ void UNiagaraSpriteRendererProperties::PostEditChangeProperty(struct FPropertyCh
 	// DerivedData.BoundingGeometry in case we cleared the CutoutTexture
 	if (bUseMaterialCutoutTexture || CutoutTexture || DerivedData.BoundingGeometry.Num())
 	{
-		const bool bUpdateCutoutDDC = 
-			PropertyChangedEvent.MemberProperty->GetFName() == GET_MEMBER_NAME_CHECKED(UNiagaraSpriteRendererProperties, SubImageSize) ||
+		const bool bUpdateCutoutDDC =
 			PropertyChangedEvent.GetPropertyName() == TEXT("bUseMaterialCutoutTexture") ||
 			PropertyChangedEvent.GetPropertyName() == TEXT("CutoutTexture") ||
 			PropertyChangedEvent.GetPropertyName() == TEXT("BoundingMode") ||
@@ -452,7 +452,7 @@ void UNiagaraSpriteRendererProperties::RenameVariable(const FNiagaraVariableBase
 void UNiagaraSpriteRendererProperties::RemoveVariable(const FNiagaraVariableBase& OldVariable, const UNiagaraEmitter* InEmitter)
 {
 	Super::RemoveVariable(OldVariable, InEmitter);
-	
+
 	// Handle resetting material bindings to defaults
 	for (FNiagaraMaterialAttributeBinding& Binding : MaterialParameterBindings)
 	{
@@ -543,7 +543,7 @@ void UNiagaraSpriteRendererProperties::GetRendererFeedback(const UNiagaraEmitter
 			if (UVScaleBinding.DoesBindingExistOnSource())
 			{
 				OutInfo.Add(LOCTEXT("SpriteRendererUVScaleWithCutout", "Cutouts will not be sized dynamically with UVScale variable. If scaling above 1.0, geometry may clip."));
-			}			
+			}
 		}
 	}
 }
diff --git a/Engine/Plugins/FX/Niagara/Source/Niagara/Private/NiagaraSystemInstance.cpp b/Engine/Plugins/FX/Niagara/Source/Niagara/Private/NiagaraSystemInstance.cpp
--- a/Engine/Plugins/FX/Niagara/Source/Niagara/Private/NiagaraSystemInstance.cpp
+++ b/Engine/Plugins/FX/Niagara/Source/Niagara/Private/NiagaraSystemInstance.cpp
@@ -133,8 +133,8 @@ FNiagaraSystemInstance::FNiagaraSystemInstance(UWorld& InWorld, UNiagaraSystem&
 		}
 		FeatureLevel = World->FeatureLevel;
 	}
-		
-	// In some cases the system may have already stated that you should ignore dependencies and tick as early as possible.		
+
+	// In some cases the system may have already stated that you should ignore dependencies and tick as early as possible.
 	if (!InAsset.bRequireCurrentFrameData)
 	{
 		TickBehavior = ENiagaraTickBehavior::ForceTickFirst;
@@ -175,7 +175,7 @@ void FNiagaraSystemInstance::SetEmitterEnable(FName EmitterName, bool bNewEnable
 		{
 			if (Emitters[FoundIdx]->IsAllowedToExecute())
 			{
-				
+
 				{
 					if (bNewEnableState)
 					{
@@ -262,7 +262,7 @@ void FNiagaraSystemInstance::SetActualExecutionState(ENiagaraExecutionState InSt
 			bNotifyOnCompletion = true;
 
 			// We may also end up calling HandleCompletion on each emitter.
-			// This may happen *before* we've successfully pulled data off of a 
+			// This may happen *before* we've successfully pulled data off of a
 			// simulation run. This means that we need to synchronize the execution
 			// states upon activation.
 			for (int32 EmitterIdx = 0; EmitterIdx < Emitters.Num(); ++EmitterIdx)
@@ -331,9 +331,9 @@ bool FNiagaraSystemInstance::RequestCapture(const FGuid& RequestId)
 	SetSolo(true);
 
 	// Go ahead and populate the shared array so that we don't have to do this on the game thread and potentially race.
-	TSharedRef<TArray<TSharedPtr<struct FNiagaraScriptDebuggerInfo, ESPMode::ThreadSafe>>, ESPMode::ThreadSafe> TempCaptureHolder = 
+	TSharedRef<TArray<TSharedPtr<struct FNiagaraScriptDebuggerInfo, ESPMode::ThreadSafe>>, ESPMode::ThreadSafe> TempCaptureHolder =
 		MakeShared<TArray<TSharedPtr<struct FNiagaraScriptDebuggerInfo, ESPMode::ThreadSafe>>, ESPMode::ThreadSafe>();
-	
+
 	TempCaptureHolder->Add(MakeShared<FNiagaraScriptDebuggerInfo, ESPMode::ThreadSafe>(NAME_None, ENiagaraScriptUsage::SystemSpawnScript, FGuid()));
 	TempCaptureHolder->Add(MakeShared<FNiagaraScriptDebuggerInfo, ESPMode::ThreadSafe>(NAME_None, ENiagaraScriptUsage::SystemUpdateScript, FGuid()));
 
@@ -402,7 +402,7 @@ bool FNiagaraSystemInstance::QueryCaptureResults(const FGuid& RequestId, TArray<
 					break;
 				}
 			}
-			
+
 			if (bWaitForGPU)
 			{
 				for (TSharedRef<FNiagaraEmitterInstance, ESPMode::ThreadSafe> CachedEmitter : Emitters)
@@ -470,7 +470,7 @@ void FNiagaraSystemInstance::SetSolo(bool bInSolo)
 		TSharedPtr<FNiagaraSystemSimulation, ESPMode::ThreadSafe> NewSoloSim = MakeShared<FNiagaraSystemSimulation, ESPMode::ThreadSafe>();
 		NewSoloSim->Init(System, World, true, TG_MAX);
 
-		NewSoloSim->TransferInstance(SystemSimulation.Get(), this);	
+		NewSoloSim->TransferInstance(SystemSimulation.Get(), this);
 
 		SystemSimulation = NewSoloSim;
 		bSolo = true;
@@ -481,7 +481,7 @@ void FNiagaraSystemInstance::SetSolo(bool bInSolo)
 		TSharedPtr<FNiagaraSystemSimulation, ESPMode::ThreadSafe> NewSim = GetWorldManager()->GetSystemSimulation(TickGroup, System);
 
 		NewSim->TransferInstance(SystemSimulation.Get(), this);
-		
+
 		SystemSimulation = NewSim;
 		bSolo = false;
 	}
@@ -573,7 +573,7 @@ void FNiagaraSystemInstance::Activate(EResetMode InResetMode)
 			WaitForAsyncTickAndFinalize();
 
 			DeferredResetMode = EResetMode::None;
-			Reset(InResetMode);			
+			Reset(InResetMode);
 		}
 	}
 	else
@@ -605,14 +605,14 @@ void FNiagaraSystemInstance::Deactivate(bool bImmediate)
 	}
 }
 
-bool FNiagaraSystemInstance::AllocateSystemInstance(TUniquePtr<FNiagaraSystemInstance>& OutSystemInstanceAllocation, UWorld& InWorld, UNiagaraSystem& InAsset,
+bool FNiagaraSystemInstance::AllocateSystemInstance(FNiagaraSystemInstancePtr& OutSystemInstanceAllocation, UWorld& InWorld, UNiagaraSystem& InAsset,
 	FNiagaraUserRedirectionParameterStore* InOverrideParameters, USceneComponent* InAttachComponent, ENiagaraTickBehavior InTickBehavior, bool bInPooled)
 {
-	OutSystemInstanceAllocation = MakeUnique<FNiagaraSystemInstance>(InWorld, InAsset, InOverrideParameters, InAttachComponent, InTickBehavior, bInPooled);
+	OutSystemInstanceAllocation = MakeShared<FNiagaraSystemInstance, ESPMode::ThreadSafe>(InWorld, InAsset, InOverrideParameters, InAttachComponent, InTickBehavior, bInPooled);
 	return true;
 }
 
-bool FNiagaraSystemInstance::DeallocateSystemInstance(TUniquePtr< FNiagaraSystemInstance >& SystemInstanceAllocation)
+bool FNiagaraSystemInstance::DeallocateSystemInstance(FNiagaraSystemInstancePtr& SystemInstanceAllocation)
 {
 	if (SystemInstanceAllocation.IsValid())
 	{
@@ -637,7 +637,7 @@ bool FNiagaraSystemInstance::DeallocateSystemInstance(TUniquePtr< FNiagaraSystem
 				}
 			);
 		}
-		
+
 		// Queue deferred deletion from the WorldManager
 		FNiagaraWorldManager* WorldManager = SystemInstanceAllocation->GetWorldManager();
 		check(WorldManager != nullptr);
@@ -652,7 +652,7 @@ bool FNiagaraSystemInstance::DeallocateSystemInstance(TUniquePtr< FNiagaraSystem
 		check(SystemInstanceAllocation == nullptr);
 	}
 	SystemInstanceAllocation = nullptr;
-	
+
 	return true;
 }
 
@@ -782,7 +782,7 @@ void FNiagaraSystemInstance::Reset(FNiagaraSystemInstance::EResetMode Mode)
 	{
 		Mode = EResetMode::ReInit;
 	}
-		
+
 	// Depending on the rest mode we may need to bind or can possibly skip it
 	// We must bind if we were previously complete as unbind will have been called, we can not get here if the system was disabled
 	bool bBindParams = IsComplete();
@@ -803,7 +803,7 @@ void FNiagaraSystemInstance::Reset(FNiagaraSystemInstance::EResetMode Mode)
 		ReInitInternal();
 		bBindParams = !IsDisabled();
 	}
-	
+
 	//If none of our emitters actually made it out of the init process we can just bail here before we ever tick.
 	bool bHasActiveEmitters = false;
 	for (auto& Inst : Emitters)
@@ -1023,12 +1023,15 @@ void FNiagaraSystemInstance::ReInitInternal()
 	RequestedExecutionState = ENiagaraExecutionState::Inactive;
 	ActualExecutionState = ENiagaraExecutionState::Inactive;
 
+	// Do we need to run in solo mode? NOTE: set this here before the early outs
+	bSolo = bForceSolo;
+
 	bool bAllReadyToRun = IsReadyToRun();
 	if (!bAllReadyToRun)
 	{
 		return;
 	}
-	
+
 	if (!System->IsValid())
 	{
 		SetRequestedExecutionState(ENiagaraExecutionState::Disabled);
@@ -1036,7 +1039,6 @@ void FNiagaraSystemInstance::ReInitInternal()
 		return;
 	}
 
-	/** Do we need to run in solo mode? */
 	bSolo |= DoSystemDataInterfacesRequireSolo(*System, OverrideParameters);
 	if (bSolo)
 	{
@@ -1137,12 +1139,12 @@ void FNiagaraSystemInstance::Cleanup()
 // {
 // 	OldInstance->GetParameterStore().Unbind(&InstanceParameters);
 // 	NewInstance->GetParameterStore().Bind(&InstanceParameters);
-// 
+//
 // 	for (TSharedRef<FNiagaraEmitterInstance, ESPMode::ThreadSafe> Simulation : Emitters)
 // 	{
 // 		Simulation->RebindParameterCollection(OldInstance, NewInstance);
 // 	}
-// 
+//
 // 	//Have to re init the instance data for data interfaces.
 // 	//This is actually lots more work than absolutely needed in some cases so we can improve it a fair bit.
 // 	InitDataInterfaces();
@@ -1389,7 +1391,7 @@ void FNiagaraSystemInstance::InitDataInterfaces()
 							NewPair.Key = Interface;
 							NewPair.Value = InstanceDataSize;
 
-							// Assume that some of our data is going to be 16 byte aligned, so enforce that 
+							// Assume that some of our data is going to be 16 byte aligned, so enforce that
 							// all per-instance data is aligned that way.
 							InstanceDataSize += Align(Size, 16);
 						}
@@ -1414,7 +1416,7 @@ void FNiagaraSystemInstance::InitDataInterfaces()
 		}
 	};
 
-	CalcInstDataSize(InstanceParameters, false, false);//This probably should be a proper exec context. 
+	CalcInstDataSize(InstanceParameters, false, false);//This probably should be a proper exec context.
 
 	if (SystemSimulation->GetIsSolo() && FNiagaraSystemSimulation::UseLegacySystemSimulationContexts())
 	{
@@ -1499,7 +1501,7 @@ void FNiagaraSystemInstance::InitDataInterfaces()
 			if (!bResult)
 			{
 				UE_LOG(LogNiagara, Error, TEXT("Error initializing data interface \"%s\" for system. %s"), *Interface->GetPathName(), Asset.IsValid() ? *Asset->GetName() : TEXT("nullptr"));
-			}		
+			}
 		}
 		else
 		{
@@ -1515,8 +1517,8 @@ void FNiagaraSystemInstance::InitDataInterfaces()
 		Complete(true);
 		return;
 	}
-	
-	//We have valid DI instance data so now generate the table of function calls.	
+
+	//We have valid DI instance data so now generate the table of function calls.
 	//When using the new exec contexts, each system instance builds it's own tables of DI function bindings for DI calls that require it.
 	//i.e. User DIs or those with per instance data that are called from system scripts.
 	if (FNiagaraSystemSimulation::UseLegacySystemSimulationContexts() == false)
@@ -1550,7 +1552,7 @@ void FNiagaraSystemInstance::TickDataInterfaces(float DeltaSeconds, bool bPostSi
 			TPair<TWeakObjectPtr<UNiagaraDataInterface>, int32>& Pair = DataInterfaceInstanceDataOffsets[DIPairIndex];
 			if (UNiagaraDataInterface* Interface = Pair.Key.Get())
 			{
-				//Ideally when we make the batching changes, we can keep the instance data in big single type blocks that can all be updated together with a single virtual call.				
+				//Ideally when we make the batching changes, we can keep the instance data in big single type blocks that can all be updated together with a single virtual call.
 				if (Interface->PerInstanceTickPostSimulate(&DataInterfaceInstanceData[Pair.Value], this, DeltaSeconds))
 				{
 					// Destroy per instance data in order to not cause any errors on check(...) inside DIs when initializing
@@ -1593,7 +1595,7 @@ void FNiagaraSystemInstance::TickDataInterfaces(float DeltaSeconds, bool bPostSi
 				Sim.DirtyDataInterfaces();
 			}
 		}
-		
+
 		// Rebind funcs for system scripts
 		if (FNiagaraSystemSimulation::UseLegacySystemSimulationContexts() == false)
 		{
@@ -1626,7 +1628,7 @@ float FNiagaraSystemInstance::GetLODDistance()
 	{
 		return LODDistance;
 	}
-	
+
 	constexpr float DefaultLODDistance = 0.0f;
 
 	FNiagaraWorldManager* WorldManager = GetWorldManager();
@@ -1634,7 +1636,7 @@ float FNiagaraSystemInstance::GetLODDistance()
 	{
 		return DefaultLODDistance;
 	}
-	
+
 	check(World);
 	const FVector EffectLocation = WorldTransform.GetLocation();
 	LODDistance = DefaultLODDistance;
@@ -1806,7 +1808,7 @@ void FNiagaraSystemInstance::TickInstanceParameters_GameThread(float DeltaSecond
 	//Bias the LastRenderTime slightly to account for any delay as it's written by the RT.
 	check(World);
 	GatheredInstanceParameters.TimeSeconds = World->TimeSeconds;
-	GatheredInstanceParameters.RealTimeSeconds = World->RealTimeSeconds;	
+	GatheredInstanceParameters.RealTimeSeconds = World->RealTimeSeconds;
 
 	// flip our buffered parameters
 	FlipParameterBuffers();
@@ -1905,7 +1907,7 @@ void FNiagaraSystemInstance::ClearEventDataSets()
 FNiagaraDataSet*
 FNiagaraSystemInstance::CreateEventDataSet(FName EmitterName, FName EventName)
 {
-	// TODO: find a better way of multiple events trying to write to the same data set; 
+	// TODO: find a better way of multiple events trying to write to the same data set;
 	// for example, if two analytical collision primitives want to send collision events, they need to push to the same data set
 	FNiagaraDataSet*& OutSet = EmitterEventDataSetMap.FindOrAdd(EmitterEventKey(EmitterName, EventName));
 
@@ -1953,7 +1955,7 @@ bool FNiagaraSystemInstance::UsesScript(const UNiagaraScript* Script)const
 
 // bool FNiagaraSystemInstance::UsesDataInterface(UNiagaraDataInterface* Interface)
 // {
-// 
+//
 // }
 
 bool FNiagaraSystemInstance::UsesCollection(const UNiagaraParameterCollection* Collection)const
@@ -2121,7 +2123,7 @@ void FNiagaraSystemInstance::Tick_GameThread(float DeltaSeconds)
 	SCOPE_CYCLE_COUNTER(STAT_NiagaraOverview_GT);
 	CSV_SCOPED_TIMING_STAT_EXCLUSIVE(Effects);
 	LLM_SCOPE(ELLMTag::Niagara);
-	
+
 	FNiagaraCrashReporterScope CRScope(this);
 
 	UNiagaraSystem* System = GetSystem();
@@ -2271,7 +2273,7 @@ void FNiagaraSystemInstance::Tick_Concurrent(bool bEnqueueGPUTickIfNeeded)
 
 		if (NewLocalBounds.IsValid)
 		{
-			LocalBounds = NewLocalBounds.ExpandBy(NewLocalBounds.GetExtent() * GNiagaraBoundsExpandByPercent);				
+			LocalBounds = NewLocalBounds.ExpandBy(NewLocalBounds.GetExtent() * GNiagaraBoundsExpandByPercent);
 		}
 		else
 		{
