Fix some compile errors caused by a bad merge

diff --git a/Engine/Plugins/FX/Niagara/Source/Niagara/Private/NiagaraComponent.cpp b/Engine/Plugins/FX/Niagara/Source/Niagara/Private/NiagaraComponent.cpp
--- a/Engine/Plugins/FX/Niagara/Source/Niagara/Private/NiagaraComponent.cpp
+++ b/Engine/Plugins/FX/Niagara/Source/Niagara/Private/NiagaraComponent.cpp
@@ -1503,7 +1503,7 @@ void UNiagaraComponent::OnEndOfFrameUpdateDuringTick()
 	Super::OnEndOfFrameUpdateDuringTick();
 	if (SystemInstanceController)
 	{
-		SystemInstanceController->WaitForConcurrentTickAndFinalize();
+		SystemInstanceController->WaitForAsyncTickAndFinalize();
 	}
 }
 
diff --git a/Engine/Plugins/FX/Niagara/Source/Niagara/Public/NiagaraSystemInstanceController.h b/Engine/Plugins/FX/Niagara/Source/Niagara/Public/NiagaraSystemInstanceController.h
--- a/Engine/Plugins/FX/Niagara/Source/Niagara/Public/NiagaraSystemInstanceController.h
+++ b/Engine/Plugins/FX/Niagara/Source/Niagara/Public/NiagaraSystemInstanceController.h
@@ -131,7 +131,7 @@ public:
 	NIAGARA_SYSTEM_INSTANCE_CONTROLLER_SHIM(Deactivate,)
 	NIAGARA_SYSTEM_INSTANCE_CONTROLLER_SHIM(CalculateTickGroup, const)
 	NIAGARA_SYSTEM_INSTANCE_CONTROLLER_SHIM(GetLocalBounds, const)
-	NIAGARA_SYSTEM_INSTANCE_CONTROLLER_SHIM(WaitForConcurrentTickAndFinalize,)
+	NIAGARA_SYSTEM_INSTANCE_CONTROLLER_SHIM(WaitForAsyncTickAndFinalize,)
 	NIAGARA_SYSTEM_INSTANCE_CONTROLLER_SHIM(OnSimulationDestroyed,)
 	NIAGARA_SYSTEM_INSTANCE_CONTROLLER_SHIM(GetAttachComponent,)
 	#if WITH_EDITOR
