no DI reset on updating target texture

diff --git a/Engine/Plugins/FX/Niagara/Source/Niagara/Private/NiagaraDataInterfaceGrid3DCollection.cpp b/Engine/Plugins/FX/Niagara/Source/Niagara/Private/NiagaraDataInterfaceGrid3DCollection.cpp
--- a/Engine/Plugins/FX/Niagara/Source/Niagara/Private/NiagaraDataInterfaceGrid3DCollection.cpp
+++ b/Engine/Plugins/FX/Niagara/Source/Niagara/Private/NiagaraDataInterfaceGrid3DCollection.cpp
@@ -2087,7 +2087,7 @@ bool UNiagaraDataInterfaceGrid3DCollection::PerInstanceTick(void* PerInstanceDat
 
 	});
 
-	return NeedsReset;
+	return false;
 }
 
 void UNiagaraDataInterfaceGrid3DCollection::GetExposedVariables(TArray<FNiagaraVariableBase>& OutVariables) const
