Edigrate 15870708 from UE4
Fixes inheritance for Sim stages

diff --git a/Engine/Plugins/FX/Niagara/Source/NiagaraEditor/Private/NiagaraScriptMergeManager.cpp b/Engine/Plugins/FX/Niagara/Source/NiagaraEditor/Private/NiagaraScriptMergeManager.cpp
--- a/Engine/Plugins/FX/Niagara/Source/NiagaraEditor/Private/NiagaraScriptMergeManager.cpp
+++ b/Engine/Plugins/FX/Niagara/Source/NiagaraEditor/Private/NiagaraScriptMergeManager.cpp
@@ -514,22 +514,22 @@ TSharedPtr<FNiagaraScriptStackMergeAdapter> FNiagaraEventHandlerMergeAdapter::Ge
 	return EventStack;
 }
 
-FNiagaraSimulationStageMergeAdapter::FNiagaraSimulationStageMergeAdapter(const UNiagaraEmitter& InEmitter, const UNiagaraSimulationStageBase* InSimulationStage, UNiagaraNodeOutput* InOutputNode)
+FNiagaraSimulationStageMergeAdapter::FNiagaraSimulationStageMergeAdapter(const UNiagaraEmitter& InEmitter, const UNiagaraSimulationStageBase* InSimulationStage, int32 InSimulationStageIndex, UNiagaraNodeOutput* InOutputNode)
 {
-	Initialize(InEmitter, InSimulationStage, nullptr, InOutputNode);
+	Initialize(InEmitter, InSimulationStage, nullptr, InSimulationStageIndex, InOutputNode);
 }
 
-FNiagaraSimulationStageMergeAdapter::FNiagaraSimulationStageMergeAdapter(const UNiagaraEmitter& InEmitter, UNiagaraSimulationStageBase* InSimulationStage, UNiagaraNodeOutput* InOutputNode)
+FNiagaraSimulationStageMergeAdapter::FNiagaraSimulationStageMergeAdapter(const UNiagaraEmitter& InEmitter, UNiagaraSimulationStageBase* InSimulationStage, int32 InSimulationStageIndex, UNiagaraNodeOutput* InOutputNode)
 {
-	Initialize(InEmitter, InSimulationStage, InSimulationStage, InOutputNode);
+	Initialize(InEmitter, InSimulationStage, InSimulationStage, InSimulationStageIndex, InOutputNode);
 }
 
 FNiagaraSimulationStageMergeAdapter::FNiagaraSimulationStageMergeAdapter(const UNiagaraEmitter& InEmitter, UNiagaraNodeOutput* InOutputNode)
 {
-	Initialize(InEmitter, nullptr, nullptr, InOutputNode);
+	Initialize(InEmitter, nullptr, nullptr, INDEX_NONE, InOutputNode);
 }
 
-void FNiagaraSimulationStageMergeAdapter::Initialize(const UNiagaraEmitter& InEmitter, const UNiagaraSimulationStageBase* InSimulationStage, UNiagaraSimulationStageBase* InEditableSimulationStage, UNiagaraNodeOutput* InOutputNode)
+void FNiagaraSimulationStageMergeAdapter::Initialize(const UNiagaraEmitter& InEmitter, const UNiagaraSimulationStageBase* InSimulationStage, UNiagaraSimulationStageBase* InEditableSimulationStage, int32 InSimulationStageIndex, UNiagaraNodeOutput* InOutputNode)
 {
 	Emitter = MakeWeakObjectPtr(const_cast<UNiagaraEmitter*>(&InEmitter));
 
@@ -537,6 +537,7 @@ void FNiagaraSimulationStageMergeAdapter::Initialize(const UNiagaraEmitter& InEm
 	EditableSimulationStage = InEditableSimulationStage;
 
 	OutputNode = InOutputNode;
+	SimulationStageIndex = InSimulationStageIndex;
 
 	if (SimulationStage != nullptr && OutputNode != nullptr)
 	{
@@ -582,6 +583,11 @@ UNiagaraNodeInput* FNiagaraSimulationStageMergeAdapter::GetInputNode() const
 	return InputNode.Get();
 }
 
+int32 FNiagaraSimulationStageMergeAdapter::GetSimulationStageIndex() const
+{
+	return SimulationStageIndex;
+}
+
 TSharedPtr<FNiagaraScriptStackMergeAdapter> FNiagaraSimulationStageMergeAdapter::GetSimulationStageStack() const
 {
 	return SimulationStageStack;
@@ -810,8 +816,10 @@ void FNiagaraEmitterMergeAdapter::Initialize(const UNiagaraEmitter& InEmitter, U
 
 	// Create an shader stage adapter for each usage id even if it's missing a shader stage object or an output node.  These
 	// incomplete adapters will be caught if they are diffed.
-	for (const UNiagaraSimulationStageBase* SimulationStage : Emitter->GetSimulationStages())
+	for (int32 SimulationStageIndex = 0; SimulationStageIndex < Emitter->GetSimulationStages().Num(); SimulationStageIndex++)
 	{
+		const UNiagaraSimulationStageBase* SimulationStage = Emitter->GetSimulationStages()[SimulationStageIndex];
+	
 		UNiagaraNodeOutput** MatchingOutputNodePtr = SimulationStageOutputNodes.FindByPredicate(
 			[=](UNiagaraNodeOutput* SimulationStageOutputNode) { return SimulationStageOutputNode->GetUsageId() == SimulationStage->Script->GetUsageId(); });
 
@@ -819,12 +827,12 @@ void FNiagaraEmitterMergeAdapter::Initialize(const UNiagaraEmitter& InEmitter, U
 
 		if (EditableEmitter == nullptr)
 		{
-			SimulationStages.Add(MakeShared<FNiagaraSimulationStageMergeAdapter>(*Emitter.Get(), SimulationStage, MatchingOutputNode));
+			SimulationStages.Add(MakeShared<FNiagaraSimulationStageMergeAdapter>(*Emitter.Get(), SimulationStage, SimulationStageIndex, MatchingOutputNode));
 		}
 		else
 		{
 			UNiagaraSimulationStageBase* EditableSimulationStage = EditableEmitter->GetSimulationStageById(SimulationStage->Script->GetUsageId());
-			SimulationStages.Add(MakeShared<FNiagaraSimulationStageMergeAdapter>(*Emitter.Get(), EditableSimulationStage, MatchingOutputNode));
+			SimulationStages.Add(MakeShared<FNiagaraSimulationStageMergeAdapter>(*Emitter.Get(), EditableSimulationStage, SimulationStageIndex, MatchingOutputNode));
 		}
 
 		if (MatchingOutputNode != nullptr)
@@ -1924,6 +1932,25 @@ void FNiagaraScriptMergeManager::DiffSimulationStages(const TArray<TSharedRef<FN
 		OtherSimulationStages,
 		[](TSharedRef<FNiagaraSimulationStageMergeAdapter> SimulationStage) { return SimulationStage->GetUsageId(); });
 
+	// Sort the diff results for easier diff applying and testing.
+	auto OrderSimulationStageByIndex = [](TSharedRef<FNiagaraSimulationStageMergeAdapter> SimulationStageA, TSharedRef<FNiagaraSimulationStageMergeAdapter> SimulationStageB)
+	{
+		return SimulationStageA->GetSimulationStageIndex() < SimulationStageB->GetSimulationStageIndex();
+	};
+
+	SimulationStageListDiffResults.RemovedBaseValues.Sort(OrderSimulationStageByIndex);
+	SimulationStageListDiffResults.AddedOtherValues.Sort(OrderSimulationStageByIndex);
+
+	auto OrderCommonSimulationStagePairByBaseIndex = [](
+		const FCommonValuePair<TSharedRef<FNiagaraSimulationStageMergeAdapter>>& CommonValuesA,
+		const FCommonValuePair<TSharedRef<FNiagaraSimulationStageMergeAdapter>>& CommonValuesB)
+	{
+		return CommonValuesA.BaseValue->GetSimulationStageIndex() < CommonValuesB.BaseValue->GetSimulationStageIndex();
+	};
+
+	SimulationStageListDiffResults.CommonValuePairs.Sort(OrderCommonSimulationStagePairByBaseIndex);
+
+	// Populate results from the sorted diff.
 	DiffResults.RemovedBaseSimulationStages.Append(SimulationStageListDiffResults.RemovedBaseValues);
 	DiffResults.AddedOtherSimulationStages.Append(SimulationStageListDiffResults.AddedOtherValues);
 
@@ -2941,6 +2968,9 @@ FNiagaraScriptMergeManager::FApplyDiffResults FNiagaraScriptMergeManager::ApplyS
 			AddedSimulationStage->Script->SetSource(EmitterSource);
 			BaseEmitter->AddSimulationStage(AddedSimulationStage);
 
+			int32 TargetIndex = FMath::Min(AddedOtherSimulationStage->GetSimulationStageIndex(), BaseEmitter->GetSimulationStages().Num() - 1);
+			BaseEmitter->MoveSimulationStageToIndex(AddedSimulationStage, TargetIndex);
+
 			FGuid PreferredOutputNodeGuid = AddedOtherSimulationStage->GetOutputNode()->NodeGuid;
 			FGuid PreferredInputNodeGuid = AddedOtherSimulationStage->GetInputNode()->NodeGuid;
 			UNiagaraNodeOutput* SimulationStageOutputNode = FNiagaraStackGraphUtilities::ResetGraphForOutput(*EmitterGraph, ENiagaraScriptUsage::ParticleSimulationStageScript, AddedSimulationStage->Script->GetUsageId(), PreferredOutputNodeGuid, PreferredInputNodeGuid);
diff --git a/Engine/Plugins/FX/Niagara/Source/NiagaraEditor/Private/NiagaraScriptMergeManager.h b/Engine/Plugins/FX/Niagara/Source/NiagaraEditor/Private/NiagaraScriptMergeManager.h
--- a/Engine/Plugins/FX/Niagara/Source/NiagaraEditor/Private/NiagaraScriptMergeManager.h
+++ b/Engine/Plugins/FX/Niagara/Source/NiagaraEditor/Private/NiagaraScriptMergeManager.h
@@ -165,8 +165,8 @@ private:
 class FNiagaraSimulationStageMergeAdapter
 {
 public:
-	FNiagaraSimulationStageMergeAdapter(const UNiagaraEmitter& InEmitter, const UNiagaraSimulationStageBase* InSimulationStage, UNiagaraNodeOutput* InOutputNode);
-	FNiagaraSimulationStageMergeAdapter(const UNiagaraEmitter& InEmitter, UNiagaraSimulationStageBase* InSimulationStage, UNiagaraNodeOutput* InOutputNode);
+	FNiagaraSimulationStageMergeAdapter(const UNiagaraEmitter& InEmitter, const UNiagaraSimulationStageBase* InSimulationStage, int32 InSimulationStageIndex, UNiagaraNodeOutput* InOutputNode);
+	FNiagaraSimulationStageMergeAdapter(const UNiagaraEmitter& InEmitter, UNiagaraSimulationStageBase* InSimulationStage, int32 InSimulationStageIndex, UNiagaraNodeOutput* InOutputNode);
 	FNiagaraSimulationStageMergeAdapter(const UNiagaraEmitter& InEmitter, UNiagaraNodeOutput* InOutputNode);
 
 	FGuid GetUsageId() const;
@@ -175,10 +175,11 @@ public:
 	UNiagaraSimulationStageBase* GetEditableSimulationStage() const;
 	UNiagaraNodeOutput* GetOutputNode() const;
 	UNiagaraNodeInput* GetInputNode() const;
+	int32 GetSimulationStageIndex() const;
 	TSharedPtr<FNiagaraScriptStackMergeAdapter> GetSimulationStageStack() const;
 
 private:
-	void Initialize(const UNiagaraEmitter& InEmitter, const UNiagaraSimulationStageBase* InSimulationStage, UNiagaraSimulationStageBase* InEditableSimulationStage, UNiagaraNodeOutput* InOutputNode);
+	void Initialize(const UNiagaraEmitter& InEmitter, const UNiagaraSimulationStageBase* InSimulationStage, UNiagaraSimulationStageBase* InEditableSimulationStage, int32 InSimulationStageIndex, UNiagaraNodeOutput* InOutputNode);
 
 private:
 	TWeakObjectPtr<UNiagaraEmitter> Emitter;
@@ -186,6 +187,7 @@ private:
 	UNiagaraSimulationStageBase* EditableSimulationStage;
 	TWeakObjectPtr<UNiagaraNodeOutput> OutputNode;
 	TWeakObjectPtr<UNiagaraNodeInput> InputNode;
+	int32 SimulationStageIndex;
 
 	TSharedPtr<FNiagaraScriptStackMergeAdapter> SimulationStageStack;
 };
