- Integrate CL 15908017
- Add missing output for SetSize
- Added CompareAndExchange, hidden from graph only usable in custom HLSL

#rnx

diff --git a/Engine/Plugins/FX/Niagara/Shaders/Private/NiagaraDataInterfaceIntRenderTarget2D.ush b/Engine/Plugins/FX/Niagara/Shaders/Private/NiagaraDataInterfaceIntRenderTarget2D.ush
--- a/Engine/Plugins/FX/Niagara/Shaders/Private/NiagaraDataInterfaceIntRenderTarget2D.ush
+++ b/Engine/Plugins/FX/Niagara/Shaders/Private/NiagaraDataInterfaceIntRenderTarget2D.ush
@@ -27,6 +27,18 @@ void AtomicAdd_{ParameterName}_UEImpureCall(in bool bExecute, int X, int Y, int
 	}
 }
 
+void AtomicCompareAndExchange_{ParameterName}_UEImpureCall(in bool bExecute, int X, int Y, int ComparisonValue, int Value, out int OriginalValue)
+{
+	if ( bExecute )
+	{
+		InterlockedCompareExchange(RWTextureUAV_{ParameterName}[int2(X, Y)], ComparisonValue, Value, OriginalValue);
+	}
+	else
+	{
+		OriginalValue = RWTextureUAV_{ParameterName}[int2(X, Y)];
+	}
+}
+
 void AtomicMax_{ParameterName}_UEImpureCall(in bool bExecute, int X, int Y, int InValue, out int CurrValue, out int PrevValue)
 {
 	if ( bExecute )
diff --git a/Engine/Plugins/FX/Niagara/Source/Niagara/Private/NiagaraDataInterfaceIntRenderTarget2D.cpp b/Engine/Plugins/FX/Niagara/Source/Niagara/Private/NiagaraDataInterfaceIntRenderTarget2D.cpp
--- a/Engine/Plugins/FX/Niagara/Source/Niagara/Private/NiagaraDataInterfaceIntRenderTarget2D.cpp
+++ b/Engine/Plugins/FX/Niagara/Source/Niagara/Private/NiagaraDataInterfaceIntRenderTarget2D.cpp
@@ -31,7 +31,7 @@ namespace NDIIntRenderTarget2DLocal
 	static const FName SetValueFunctionName("SetValue");
 	static const FName AtomicAddFunctionName("AtomicAdd");
 	//static const FName AtomicAndFunctionName("AtomicAnd");
-	//static const FName AtomicCASFunctionName("AtomicCompareAndExchange");
+	static const FName AtomicCASFunctionName("AtomicCompareAndExchange");
 	//static const FName AtomicCSFunctionName("AtomicCompareStore");
 	//static const FName AtomicExchangeFunctionName("AtomicExchange");
 	static const FName AtomicMaxFunctionName("AtomicMax");
@@ -277,7 +277,28 @@ void UNiagaraDataInterfaceIntRenderTarget2D::GetFunctions(TArray<FNiagaraFunctio
 #endif
 	}
 	////static const FName AtomicAndFunctionName("AtomicAnd");
-	////static const FName AtomicCASFunctionName("AtomicCompareAndExchange");
+	{
+		FNiagaraFunctionSignature& Sig = OutFunctions.AddDefaulted_GetRef();
+		Sig.Name = NDIIntRenderTarget2DLocal::AtomicCASFunctionName;
+		Sig.Inputs.Add(FNiagaraVariable(FNiagaraTypeDefinition(GetClass()), TEXT("RenderTarget")));
+		Sig.Inputs.Add_GetRef(FNiagaraVariable(FNiagaraTypeDefinition::GetBoolDef(), TEXT("Execute"))).SetValue(true);
+		Sig.Inputs.Add(FNiagaraVariable(FNiagaraTypeDefinition::GetIntDef(), TEXT("PixelX")));
+		Sig.Inputs.Add(FNiagaraVariable(FNiagaraTypeDefinition::GetIntDef(), TEXT("PixelY")));
+		Sig.Inputs.Add(FNiagaraVariable(FNiagaraTypeDefinition::GetIntDef(), TEXT("Comparison Value")));
+		Sig.Inputs.Add(FNiagaraVariable(FNiagaraTypeDefinition::GetIntDef(), TEXT("Value")));
+		Sig.Outputs.Add(FNiagaraVariable(FNiagaraTypeDefinition::GetIntDef(), TEXT("Original Value")));
+		Sig.bHidden = true;
+		Sig.bExperimental = true;
+		Sig.bMemberFunction = true;
+		Sig.bRequiresExecPin = true;
+		Sig.bRequiresContext = false;
+		Sig.bWriteFunction = true;
+		Sig.bSupportsCPU = false;
+		Sig.bSupportsGPU = true;
+#if WITH_EDITORONLY_DATA
+		Sig.Description = LOCTEXT("AtomicCASFunctionDesc", "Compares the pixel value against the comparison value, if they are equal the value is replaced.  Original Value is the pixel value before the operation completes.  This opertion is thread safe.");
+#endif
+	}
 	////static const FName AtomicCSFunctionName("AtomicCompareStore");
 	////static const FName AtomicExchangeFunctionName("AtomicExchange");
 	{
@@ -346,6 +367,7 @@ void UNiagaraDataInterfaceIntRenderTarget2D::GetFunctions(TArray<FNiagaraFunctio
 		Sig.Inputs.Add(FNiagaraVariable(FNiagaraTypeDefinition(GetClass()), TEXT("RenderTarget")));
 		Sig.Inputs.Add(FNiagaraVariable(FNiagaraTypeDefinition::GetIntDef(), TEXT("Width")));
 		Sig.Inputs.Add(FNiagaraVariable(FNiagaraTypeDefinition::GetIntDef(), TEXT("Height")));
+		Sig.Outputs.Add(FNiagaraVariable(FNiagaraTypeDefinition::GetBoolDef(), TEXT("Success")));
 		Sig.ModuleUsageBitmask = EmitterSystemOnlyBitmask;
 		Sig.bExperimental = true;
 		Sig.bMemberFunction = true;
@@ -447,15 +469,10 @@ bool UNiagaraDataInterfaceIntRenderTarget2D::CopyToInternal(UNiagaraDataInterfac
 #if WITH_EDITORONLY_DATA
 bool UNiagaraDataInterfaceIntRenderTarget2D::AppendCompileHash(FNiagaraCompileHashVisitor* InVisitor) const
 {
-	if (!Super::AppendCompileHash(InVisitor))
-		return false;
-
 	bool bSuccess = Super::AppendCompileHash(InVisitor);
 	FSHAHash Hash = GetShaderFileHash(NDIIntRenderTarget2DLocal::TemplateShaderFile, EShaderPlatform::SP_PCD3D_SM5);
 	InVisitor->UpdateString(TEXT("NiagaraDataInterfaceExportTemplateHLSLSource"), Hash.ToString());
 	return bSuccess;
-
-	return true;
 }
 
 void UNiagaraDataInterfaceIntRenderTarget2D::GetParameterDefinitionHLSL(const FNiagaraDataInterfaceGPUParamInfo& ParamInfo, FString& OutHLSL)
@@ -477,7 +494,7 @@ bool UNiagaraDataInterfaceIntRenderTarget2D::GetFunctionHLSL(const FNiagaraDataI
 		(FunctionInfo.DefinitionName == NDIIntRenderTarget2DLocal::SetValueFunctionName) ||
 		(FunctionInfo.DefinitionName == NDIIntRenderTarget2DLocal::AtomicAddFunctionName) ||
 		//(FunctionInfo.DefinitionName == NDIIntRenderTarget2DLocal::AtomicAndFunctionName) ||
-		//(FunctionInfo.DefinitionName == NDIIntRenderTarget2DLocal::AtomicCASFunctionName) ||
+		(FunctionInfo.DefinitionName == NDIIntRenderTarget2DLocal::AtomicCASFunctionName) ||
 		//(FunctionInfo.DefinitionName == NDIIntRenderTarget2DLocal::AtomicCSFunctionName) ||
 		//(FunctionInfo.DefinitionName == NDIIntRenderTarget2DLocal::AtomicExchangeFunctionName) ||
 		(FunctionInfo.DefinitionName == NDIIntRenderTarget2DLocal::AtomicMaxFunctionName) ||
