- Editgrate fix for TicksThisFrame not being cleared

#rnx

diff --git a/Engine/Plugins/FX/Niagara/Source/Niagara/Private/NiagaraEmitterInstanceBatcher.cpp b/Engine/Plugins/FX/Niagara/Source/Niagara/Private/NiagaraEmitterInstanceBatcher.cpp
--- a/Engine/Plugins/FX/Niagara/Source/Niagara/Private/NiagaraEmitterInstanceBatcher.cpp
+++ b/Engine/Plugins/FX/Niagara/Source/Niagara/Private/NiagaraEmitterInstanceBatcher.cpp
@@ -1142,7 +1142,17 @@ void NiagaraEmitterInstanceBatcher::BuildTickStagePasses(FRHICommandListImmediat
 			{
 				FNiagaraComputeInstanceData& InstanceData = Tick->GetInstanceData()[i];
 				FNiagaraComputeExecutionContext* ExecContext = InstanceData.Context;
-				if ((ExecContext == nullptr) || !ExecContext->GPUScript_RT->IsShaderMapComplete_RenderThread())
+
+				if (ExecContext == nullptr)
+				{
+					continue;
+				}
+
+				//Ensure we clear the temp tick count value in every emitter tick that we've processed.
+				//Can't do later in the final tick as that may not contain *all* emitter ticks.
+				ExecContext->TicksThisFrame = 0;
+
+				if (!ExecContext->GPUScript_RT->IsShaderMapComplete_RenderThread())
 				{
 					continue;
 				}
@@ -1335,14 +1345,6 @@ void NiagaraEmitterInstanceBatcher::ExecuteAll(FRHICommandList& RHICmdList, FRHI
 				{
 					// Reset to default as it will no longer be used.
 					SharedContext->ScratchIndex = INDEX_NONE;
-
-					//Reset per instance data also
-					for (uint32 i = 0; i < Tick->Count; ++i)
-					{
-						FNiagaraComputeInstanceData& InstanceData = Tick->GetInstanceData()[i];
-						FNiagaraComputeExecutionContext* ExecContext = InstanceData.Context;
-						ExecContext->TicksThisFrame = 0;
-					}
 				}
 				else
 				{
