Fix null SRVs being set by the skeletal mesh DI.

#rnx

diff --git a/Engine/Plugins/FX/Niagara/Source/Niagara/Private/NiagaraDataInterfaceSkeletalMesh.cpp b/Engine/Plugins/FX/Niagara/Source/Niagara/Private/NiagaraDataInterfaceSkeletalMesh.cpp
--- a/Engine/Plugins/FX/Niagara/Source/Niagara/Private/NiagaraDataInterfaceSkeletalMesh.cpp
+++ b/Engine/Plugins/FX/Niagara/Source/Niagara/Private/NiagaraDataInterfaceSkeletalMesh.cpp
@@ -1441,8 +1441,10 @@ public:
 			SetSRVParameter(RHICmdList, ComputeShaderRHI, MeshSampleRegionsTriangleIndices, StaticBuffers->GetNumSamplingRegionTriangles() > 0 ? StaticBuffers->GetSampleRegionsTriangleIndicesSRV().GetReference() : FNiagaraRenderer::GetDummyUIntBuffer());
 			SetSRVParameter(RHICmdList, ComputeShaderRHI, MeshSampleRegionsVertices, StaticBuffers->GetNumSamplingRegionVertices() > 0 ? StaticBuffers->GetSampleRegionsVerticesSRV().GetReference() : FNiagaraRenderer::GetDummyUIntBuffer());
 
-			SetSRVParameter(RHICmdList, ComputeShaderRHI, MeshSkinWeightBuffer, InstanceData->MeshSkinWeightBuffer->GetSRV());
-			SetSRVParameter(RHICmdList, ComputeShaderRHI, MeshSkinWeightLookupBuffer, InstanceData->MeshSkinWeightLookupBuffer->GetSRV());
+			FRHIShaderResourceView* MeshSkinWeightBufferSRV = InstanceData->MeshSkinWeightBuffer->GetSRV();
+			SetSRVParameter(RHICmdList, ComputeShaderRHI, MeshSkinWeightBuffer, MeshSkinWeightBufferSRV ? MeshSkinWeightBufferSRV : FNiagaraRenderer::GetDummyUIntBuffer());
+			FRHIShaderResourceView* MeshSkinWeightLookupBufferSRV = InstanceData->MeshSkinWeightLookupBuffer->GetSRV();
+			SetSRVParameter(RHICmdList, ComputeShaderRHI, MeshSkinWeightLookupBuffer, MeshSkinWeightLookupBufferSRV ? MeshSkinWeightLookupBufferSRV : FNiagaraRenderer::GetDummyUIntBuffer());
 
 			SetShaderValue(RHICmdList, ComputeShaderRHI, MeshWeightStride, InstanceData->MeshWeightStrideByte/4);
 
