Adding back missing reset to default options on renderer property rows

diff --git a/Engine/Plugins/FX/Niagara/Source/NiagaraEditorWidgets/Private/SNiagaraStack.cpp b/Engine/Plugins/FX/Niagara/Source/NiagaraEditorWidgets/Private/SNiagaraStack.cpp
--- a/Engine/Plugins/FX/Niagara/Source/NiagaraEditorWidgets/Private/SNiagaraStack.cpp
+++ b/Engine/Plugins/FX/Niagara/Source/NiagaraEditorWidgets/Private/SNiagaraStack.cpp
@@ -60,6 +60,7 @@
 #include "ViewModels/NiagaraOverviewGraphViewModel.h"
 #include "Widgets/SNiagaraParameterName.h"
 #include "Styling/StyleColors.h"
+#include "SResetToDefaultPropertyEditor.h"
 
 #define LOCTEXT_NAMESPACE "NiagaraStack"
 
@@ -650,7 +651,7 @@ TSharedRef<ITableRow> SNiagaraStack::OnGenerateRowForStackItem(UNiagaraStackEntr
 {
 	TSharedRef<SNiagaraStackTableRow> Container = ConstructContainerForItem(Item);
 	FRowWidgets RowWidgets = ConstructNameAndValueWidgetsForItem(Item, Container);
-	Container->SetNameAndValueContent(RowWidgets.NameWidget, RowWidgets.ValueWidget, RowWidgets.EditConditionWidget);
+	Container->SetNameAndValueContent(RowWidgets.NameWidget, RowWidgets.ValueWidget, RowWidgets.EditConditionWidget, RowWidgets.ResetWidget);
 	return Container;
 }
 
@@ -1008,6 +1009,7 @@ SNiagaraStack::FRowWidgets SNiagaraStack::ConstructNameAndValueWidgetsForItem(UN
 
 			FRowWidgets RowWidgets(PropertyRowWidgets.NameWidget.ToSharedRef(), PropertyRowWidgets.ValueWidget.ToSharedRef());
 			RowWidgets.EditConditionWidget = PropertyRowWidgets.EditConditionWidget;
+			RowWidgets.ResetWidget = SNew(SResetToDefaultPropertyEditor, PropertyRow->GetDetailTreeNode()->CreatePropertyHandle());
 			return RowWidgets;
 		}
 	}
diff --git a/Engine/Plugins/FX/Niagara/Source/NiagaraEditorWidgets/Private/SNiagaraStack.h b/Engine/Plugins/FX/Niagara/Source/NiagaraEditorWidgets/Private/SNiagaraStack.h
--- a/Engine/Plugins/FX/Niagara/Source/NiagaraEditorWidgets/Private/SNiagaraStack.h
+++ b/Engine/Plugins/FX/Niagara/Source/NiagaraEditorWidgets/Private/SNiagaraStack.h
@@ -47,6 +47,7 @@ private:
 		TSharedRef<SWidget> NameWidget;
 		TSharedPtr<SWidget> ValueWidget;
 		TSharedPtr<SWidget> EditConditionWidget;
+		TSharedPtr<SWidget> ResetWidget;
 	};
 
 	void SynchronizeTreeExpansion();
diff --git a/Engine/Plugins/FX/Niagara/Source/NiagaraEditorWidgets/Private/Stack/SNiagaraStackTableRow.cpp b/Engine/Plugins/FX/Niagara/Source/NiagaraEditorWidgets/Private/Stack/SNiagaraStackTableRow.cpp
--- a/Engine/Plugins/FX/Niagara/Source/NiagaraEditorWidgets/Private/Stack/SNiagaraStackTableRow.cpp
+++ b/Engine/Plugins/FX/Niagara/Source/NiagaraEditorWidgets/Private/Stack/SNiagaraStackTableRow.cpp
@@ -132,7 +132,7 @@ bool HasVisibleChildren(UNiagaraStackEntry* Entry)
 	return Children.Num() > 0;
 }
 
-void SNiagaraStackTableRow::SetNameAndValueContent(TSharedRef<SWidget> InNameWidget, TSharedPtr<SWidget> InValueWidget, TSharedPtr<SWidget> InEditConditionWidget)
+void SNiagaraStackTableRow::SetNameAndValueContent(TSharedRef<SWidget> InNameWidget, TSharedPtr<SWidget> InValueWidget, TSharedPtr<SWidget> InEditConditionWidget, TSharedPtr<SWidget> InResetWidget)
 {
 	FSlateColor IconColor = FNiagaraEditorWidgetsStyle::Get().GetColor(FNiagaraStackEditorWidgetsUtilities::GetColorNameForExecutionCategory(StackEntry->GetExecutionCategoryName()));
 	if (bIsCategoryIconHighlighted)
@@ -351,6 +351,14 @@ void SNiagaraStackTableRow::SetNameAndValueContent(TSharedRef<SWidget> InNameWid
 							SNew(SNiagaraStackIssueIcon, StackViewModel, StackEntry)
 							.Visibility(IssueIconVisibility)
 						]
+						// Reset To Default
+						+ SHorizontalBox::Slot()
+						.AutoWidth()
+						.VAlign(VAlign_Center)
+						.Padding(0, 0, 3, 0)
+						[
+							InResetWidget.IsValid() ? InResetWidget.ToSharedRef() : SNullWidget::NullWidget
+						]
 					]
 				]
 			]
diff --git a/Engine/Plugins/FX/Niagara/Source/NiagaraEditorWidgets/Private/Stack/SNiagaraStackTableRow.h b/Engine/Plugins/FX/Niagara/Source/NiagaraEditorWidgets/Private/Stack/SNiagaraStackTableRow.h
--- a/Engine/Plugins/FX/Niagara/Source/NiagaraEditorWidgets/Private/Stack/SNiagaraStackTableRow.h
+++ b/Engine/Plugins/FX/Niagara/Source/NiagaraEditorWidgets/Private/Stack/SNiagaraStackTableRow.h
@@ -60,7 +60,7 @@ public:
 
 	void SetContentPadding(FMargin InContentPadding);
 
-	void SetNameAndValueContent(TSharedRef<SWidget> InNameWidget, TSharedPtr<SWidget> InValueWidget, TSharedPtr<SWidget> InEditConditionWidget);
+	void SetNameAndValueContent(TSharedRef<SWidget> InNameWidget, TSharedPtr<SWidget> InValueWidget, TSharedPtr<SWidget> InEditConditionWidget, TSharedPtr<SWidget> InResetWidget);
 
 	void AddFillRowContextMenuHandler(FOnFillRowContextMenu FillRowContextMenuHandler);
 
