Fix an issue where Niagara meshes would not render if you added a material parameter override to a mesh with a material that was not the base

diff --git a/Engine/Plugins/FX/Niagara/Source/Niagara/Private/NiagaraRendererMeshes.cpp b/Engine/Plugins/FX/Niagara/Source/Niagara/Private/NiagaraRendererMeshes.cpp
--- a/Engine/Plugins/FX/Niagara/Source/Niagara/Private/NiagaraRendererMeshes.cpp
+++ b/Engine/Plugins/FX/Niagara/Source/Niagara/Private/NiagaraRendererMeshes.cpp
@@ -180,11 +180,16 @@ void FNiagaraRendererMeshes::Initialize(const UNiagaraRendererProperties* InProp
 			{
 				MeshData.MaterialRemapTable.Add(BaseMaterials_GT.IndexOfByPredicate([&](UMaterialInterface* LookMat)
 				{
-					UMaterialInterface* TestMat = LookMat->IsA<UMaterialInstanceDynamic>()? CastChecked<UMaterialInstanceDynamic>(LookMat)->GetBaseMaterial() : LookMat;
-					return TestMat == MeshMaterial;
+					if (LookMat == MeshMaterial)
+					{
+						return true;
+					}
+					if (UMaterialInstanceDynamic* MID = Cast<UMaterialInstanceDynamic>(LookMat))
+					{
+						return MeshMaterial == MID->Parent;
+					}
+					return false;
 				}));
-
-				//MeshData.MaterialRemapTable.Add(BaseMaterials_GT.Find(TestMat));
 			}
 
 			// Extend the local bounds by this mesh's bounds
