Make sure RT proxy gets properly initialized on load for Niagara collision query

diff --git a/Engine/Plugins/FX/Niagara/Source/Niagara/Classes/NiagaraDataInterfaceCollisionQuery.h b/Engine/Plugins/FX/Niagara/Source/Niagara/Classes/NiagaraDataInterfaceCollisionQuery.h
--- a/Engine/Plugins/FX/Niagara/Source/Niagara/Classes/NiagaraDataInterfaceCollisionQuery.h
+++ b/Engine/Plugins/FX/Niagara/Source/Niagara/Classes/NiagaraDataInterfaceCollisionQuery.h
@@ -36,6 +36,7 @@ public:
 
 	//UObject Interface
 	virtual void PostInitProperties() override;
+	virtual void PostLoad() override;
 #if WITH_EDITOR
 	virtual void PostEditChangeProperty(struct FPropertyChangedEvent& PropertyChangedEvent) override;
 #endif
diff --git a/Engine/Plugins/FX/Niagara/Source/Niagara/Private/NiagaraDataInterfaceCollisionQuery.cpp b/Engine/Plugins/FX/Niagara/Source/Niagara/Private/NiagaraDataInterfaceCollisionQuery.cpp
--- a/Engine/Plugins/FX/Niagara/Source/Niagara/Private/NiagaraDataInterfaceCollisionQuery.cpp
+++ b/Engine/Plugins/FX/Niagara/Source/Niagara/Private/NiagaraDataInterfaceCollisionQuery.cpp
@@ -167,6 +167,16 @@ void UNiagaraDataInterfaceCollisionQuery::PostInitProperties()
 	}
 }
 
+void UNiagaraDataInterfaceCollisionQuery::PostLoad()
+{
+	Super::PostLoad();
+
+	if (MaxRayTraceCount)
+	{
+		MarkRenderDataDirty();
+	}
+}
+
 void UNiagaraDataInterfaceCollisionQuery::GetAssetTagsForContext(const UObject* InAsset, const TArray<const UNiagaraDataInterface*>& InProperties, TMap<FName, uint32>& NumericKeys, TMap<FName, FString>& StringKeys) const
 {
 #if WITH_EDITOR
