Fix for code gen with Niagara collision DI and GPU distance field sampling

diff --git a/Engine/Plugins/FX/Niagara/Source/Niagara/Private/NiagaraDataInterfaceCollisionQuery.cpp b/Engine/Plugins/FX/Niagara/Source/Niagara/Private/NiagaraDataInterfaceCollisionQuery.cpp
--- a/Engine/Plugins/FX/Niagara/Source/Niagara/Private/NiagaraDataInterfaceCollisionQuery.cpp
+++ b/Engine/Plugins/FX/Niagara/Source/Niagara/Private/NiagaraDataInterfaceCollisionQuery.cpp
@@ -433,7 +433,7 @@ bool UNiagaraDataInterfaceCollisionQuery::GetFunctionHLSL(const FNiagaraDataInte
 			}
 		)");
 
-		OutHLSL += TEXT("void ") + FunctionInfo.InstanceName + TEXT(" \n{\n");
+		OutHLSL += FString::Format(FormatSample, Args);
 	}
 	else if (FunctionInfo.DefinitionName == IssueAsyncRayTraceName)
 	{
diff --git a/Engine/Plugins/FX/Niagara/Source/NiagaraCore/Private/NiagaraCustomVersion.cpp b/Engine/Plugins/FX/Niagara/Source/NiagaraCore/Private/NiagaraCustomVersion.cpp
--- a/Engine/Plugins/FX/Niagara/Source/NiagaraCore/Private/NiagaraCustomVersion.cpp
+++ b/Engine/Plugins/FX/Niagara/Source/NiagaraCore/Private/NiagaraCustomVersion.cpp
@@ -8,4 +8,4 @@ const FGuid FNiagaraCustomVersion::GUID(0xFCF57AFA, 0x50764283, 0xB9A9E658, 0xFF
 // Register the custom version with core
 FCustomVersionRegistration GRegisterNiagaraCustomVersion(FNiagaraCustomVersion::GUID, FNiagaraCustomVersion::LatestVersion, TEXT("NiagaraVer"));
 
-const FGuid FNiagaraCustomVersion::LatestScriptCompileVersion(0xE76B1567, 0x609C40DE, 0xA2ABD4E7, 0x133C1623);
+const FGuid FNiagaraCustomVersion::LatestScriptCompileVersion(0xDA91AA46, 0xFE7F41A3, 0xACCBC3C5, 0x012D61DB);
