Merging CL 15928544: Fixed the upgrade function for the camera DI

diff --git a/Engine/Plugins/FX/Niagara/Source/Niagara/Private/NiagaraDataInterfaceCamera.cpp b/Engine/Plugins/FX/Niagara/Source/Niagara/Private/NiagaraDataInterfaceCamera.cpp
--- a/Engine/Plugins/FX/Niagara/Source/Niagara/Private/NiagaraDataInterfaceCamera.cpp
+++ b/Engine/Plugins/FX/Niagara/Source/Niagara/Private/NiagaraDataInterfaceCamera.cpp
@@ -358,11 +358,14 @@ bool UNiagaraDataInterfaceCamera::UpgradeFunctionCall(FNiagaraFunctionSignature&
 	}
 
 	// Added a new output to the gpu properties
-	if (FunctionSignature.FunctionVersion < FNiagaraCameraDIFunctionVersion::AddedPreViewTranslation && FunctionSignature.Name == GetViewPropertiesName)
+	if (FunctionSignature.FunctionVersion < FNiagaraCameraDIFunctionVersion::AddedPreViewTranslation && FunctionSignature.Name == GetViewPropertiesName && FunctionSignature.Outputs.Num() == 8)
 	{
 		FunctionSignature.AddOutput(FNiagaraVariable(FNiagaraTypeDefinition::GetVec3Def(), TEXT("PreViewTranslation")), LOCTEXT("PreViewTranslationDescription", "Returns the translation to apply for the various 'Translated XX to XX' transforms."));
 		bWasChanged = true;
 	}
+	
+	// Set latest version
+    FunctionSignature.FunctionVersion = FNiagaraCameraDIFunctionVersion::LatestVersion;
 
 	return bWasChanged;
 }
