Add the ability to bind meshes in the mesh renderer to a user parameter

diff --git a/Engine/Plugins/FX/Niagara/Source/Niagara/Private/NiagaraComponent.cpp b/Engine/Plugins/FX/Niagara/Source/Niagara/Private/NiagaraComponent.cpp
--- a/Engine/Plugins/FX/Niagara/Source/Niagara/Private/NiagaraComponent.cpp
+++ b/Engine/Plugins/FX/Niagara/Source/Niagara/Private/NiagaraComponent.cpp
@@ -1837,6 +1837,20 @@ void UNiagaraComponent::SetVariableMaterial(FName InVariableName, UMaterialInter
 #endif
 }
 
+
+void UNiagaraComponent::SetVariableStaticMesh(FName InVariableName, UStaticMesh* InValue)
+{
+	if (SystemInstanceController.IsValid())
+	{
+		TWeakObjectPtr<UStaticMesh> ObjPtr = InValue;
+		SystemInstanceController->SetVariable_Deferred(InVariableName, ObjPtr);
+	}
+#if WITH_EDITOR
+	const FNiagaraVariable VariableDesc(FNiagaraTypeDefinition::GetUStaticMeshDef(), InVariableName);
+	SetParameterOverride(VariableDesc, FNiagaraVariant(InValue));
+#endif
+}
+
 void UNiagaraComponent::SetVariableTextureRenderTarget(FName InVariableName, UTextureRenderTarget* TextureRenderTarget)
 {
 	if (SystemInstanceController.IsValid())
diff --git a/Engine/Plugins/FX/Niagara/Source/Niagara/Private/NiagaraEmitterInstance.cpp b/Engine/Plugins/FX/Niagara/Source/Niagara/Private/NiagaraEmitterInstance.cpp
--- a/Engine/Plugins/FX/Niagara/Source/Niagara/Private/NiagaraEmitterInstance.cpp
+++ b/Engine/Plugins/FX/Niagara/Source/Niagara/Private/NiagaraEmitterInstance.cpp
@@ -844,7 +844,10 @@ void FNiagaraEmitterInstance::BindParameters(bool bExternalOnly)
 	//if (bAnyRendererBindingsAdded)
 	{
 		if (ParentSystemInstance)
+		{
+			ParentSystemInstance->GetOverrideParameters()->Bind(&RendererBindings);
 			ParentSystemInstance->GetInstanceParameters().Bind(&RendererBindings);
+		}
 
 		//SystemScriptDefinedDataInterfaceParameters.Bind(&RendererBindings);
 		//ScriptDefinedDataInterfaceParameters.Bind(&RendererBindings);
diff --git a/Engine/Plugins/FX/Niagara/Source/Niagara/Private/NiagaraMeshRendererProperties.cpp b/Engine/Plugins/FX/Niagara/Source/Niagara/Private/NiagaraMeshRendererProperties.cpp
--- a/Engine/Plugins/FX/Niagara/Source/Niagara/Private/NiagaraMeshRendererProperties.cpp
+++ b/Engine/Plugins/FX/Niagara/Source/Niagara/Private/NiagaraMeshRendererProperties.cpp
@@ -55,9 +55,39 @@ FNiagaraMeshRendererMeshProperties::FNiagaraMeshRendererMeshProperties()
 	, PivotOffset(ForceInitToZero)
 	, PivotOffsetSpace(ENiagaraMeshPivotOffsetSpace::Mesh)
 {	
+
+	FNiagaraTypeDefinition StaticMeshDef(UStaticMesh::StaticClass());
+	UserParamBinding.Parameter.SetType(StaticMeshDef);
+}
+
+UStaticMesh* FNiagaraMeshRendererMeshProperties::ResolveStaticMesh(const FNiagaraEmitterInstance* Emitter) const
+{
+	UStaticMesh* FoundMesh = nullptr;
+
+	if (UserParamBinding.Parameter.IsValid() && Emitter)
+	{
+		UStaticMesh* TestMesh = Cast<UStaticMesh>(Emitter->GetRendererBoundVariables().GetUObject(UserParamBinding.Parameter));
+		if (TestMesh && TestMesh->GetRenderData())
+		{
+			FoundMesh = TestMesh;
+		}
+	}
+
+	if (!FoundMesh && Mesh && Mesh->GetRenderData())
+	{
+		FoundMesh = Mesh;
+	}
+
+	return FoundMesh;
+}
+
+bool FNiagaraMeshRendererMeshProperties::HasValidMeshProperties() const
+{
+	return Mesh || UserParamBinding.Parameter.IsValid();
 }
 
 
+
 UNiagaraMeshRendererProperties::UNiagaraMeshRendererProperties()
 	: SourceMode(ENiagaraRendererSourceDataMode::Particles)
 	, SortMode(ENiagaraSortMode::None)
@@ -111,7 +141,7 @@ FNiagaraRenderer* UNiagaraMeshRendererProperties::CreateEmitterRenderer(ERHIFeat
 {
 	for (const auto& MeshProperties : Meshes)
 	{
-		if (MeshProperties.Mesh && MeshProperties.Mesh->GetRenderData())
+		if (MeshProperties.ResolveStaticMesh(Emitter))
 		{
 			// There's at least one valid mesh
 			FNiagaraRenderer* NewRenderer = new FNiagaraRendererMeshes(FeatureLevel, this, Emitter);
@@ -361,7 +391,7 @@ void UNiagaraMeshRendererProperties::GetUsedMeshMaterials(int32 MeshIndex, const
 {
 	check(Meshes.IsValidIndex(MeshIndex));
 
-	const UStaticMesh* Mesh = Meshes[MeshIndex].Mesh;
+	const UStaticMesh* Mesh = Meshes[MeshIndex].ResolveStaticMesh(Emitter);
 	check(Mesh);
 
 	const FStaticMeshRenderData* RenderData = Mesh->GetRenderData();
@@ -431,7 +461,7 @@ void UNiagaraMeshRendererProperties::GetUsedMaterials(const FNiagaraEmitterInsta
 	TArray<UMaterialInterface*> OrderedMeshMaterials;
 	for (int32 MeshIndex = 0; MeshIndex < Meshes.Num(); ++MeshIndex)
 	{
-		const UStaticMesh* Mesh = Meshes[MeshIndex].Mesh;
+		const UStaticMesh* Mesh = Meshes[MeshIndex].ResolveStaticMesh(InEmitter);
 		if (Mesh && Mesh->GetRenderData())
 		{
 			GetUsedMeshMaterials(MeshIndex, InEmitter, OrderedMeshMaterials);
@@ -466,6 +496,15 @@ bool UNiagaraMeshRendererProperties::PopulateRequiredBindings(FNiagaraParameterS
 		bAnyAdded = true;
 	}
 
+	for (FNiagaraMeshRendererMeshProperties& Binding : Meshes)
+	{
+		if (Binding.UserParamBinding.Parameter.IsValid())
+		{
+			InParameterStore.AddParameter(Binding.UserParamBinding.Parameter, false);
+			bAnyAdded = true;
+		}
+	}
+
 	return bAnyAdded;
 }
 
diff --git a/Engine/Plugins/FX/Niagara/Source/Niagara/Private/NiagaraModule.cpp b/Engine/Plugins/FX/Niagara/Source/Niagara/Private/NiagaraModule.cpp
--- a/Engine/Plugins/FX/Niagara/Source/Niagara/Private/NiagaraModule.cpp
+++ b/Engine/Plugins/FX/Niagara/Source/Niagara/Private/NiagaraModule.cpp
@@ -567,6 +567,7 @@ FNiagaraTypeDefinition FNiagaraTypeDefinition::UObjectDef;
 FNiagaraTypeDefinition FNiagaraTypeDefinition::UMaterialDef;
 FNiagaraTypeDefinition FNiagaraTypeDefinition::UTextureDef;
 FNiagaraTypeDefinition FNiagaraTypeDefinition::UTextureRenderTargetDef;
+FNiagaraTypeDefinition FNiagaraTypeDefinition::UStaticMeshDef;
 FNiagaraTypeDefinition FNiagaraTypeDefinition::WildcardDef;
 
 TSet<UScriptStruct*> FNiagaraTypeDefinition::NumericStructs;
@@ -638,6 +639,7 @@ void FNiagaraTypeDefinition::Init()
 	UMaterialDef = FNiagaraTypeDefinition(UMaterialClass);
 	UTextureDef = FNiagaraTypeDefinition(UTextureClass);
 	UTextureRenderTargetDef = FNiagaraTypeDefinition(UTextureRenderTargetClass);
+	UStaticMeshDef = FNiagaraTypeDefinition(UStaticMesh::StaticClass());
 
 	CollisionEventDef = FNiagaraTypeDefinition(FNiagaraCollisionEventPayload::StaticStruct());
 	NumericStructs.Add(NumericStruct);
@@ -831,6 +833,7 @@ void FNiagaraTypeDefinition::RecreateUserDefinedTypeRegistry()
 	FNiagaraTypeRegistry::Register(UMaterialDef, VarFlags);
 	FNiagaraTypeRegistry::Register(UTextureDef, VarFlags);
 	FNiagaraTypeRegistry::Register(UTextureRenderTargetDef, VarFlags);
+	FNiagaraTypeRegistry::Register(UStaticMeshDef, VarFlags);
 	FNiagaraTypeRegistry::Register(FNiagaraRandInfo::StaticStruct(), ParamFlags | PayloadFlags);
 	FNiagaraTypeRegistry::Register(StaticEnum<ENiagaraLegacyTrailWidthMode>(), ParamFlags | PayloadFlags);
 
diff --git a/Engine/Plugins/FX/Niagara/Source/Niagara/Private/NiagaraRenderer.cpp b/Engine/Plugins/FX/Niagara/Source/Niagara/Private/NiagaraRenderer.cpp
--- a/Engine/Plugins/FX/Niagara/Source/Niagara/Private/NiagaraRenderer.cpp
+++ b/Engine/Plugins/FX/Niagara/Source/Niagara/Private/NiagaraRenderer.cpp
@@ -290,6 +290,7 @@ FNiagaraRenderer::FNiagaraRenderer(ERHIFeatureLevel::Type InFeatureLevel, const
 void FNiagaraRenderer::Initialize(const UNiagaraRendererProperties* InProps, const FNiagaraEmitterInstance* Emitter, const FNiagaraSystemInstanceController& InController)
 {
 	//Get our list of valid base materials. Fall back to default material if they're not valid.
+	BaseMaterials_GT.Empty();
 	InProps->GetUsedMaterials(Emitter, BaseMaterials_GT);
 	bool bCreateMidsForUsedMaterials = InProps->NeedsMIDsForMaterials();
 
diff --git a/Engine/Plugins/FX/Niagara/Source/Niagara/Private/NiagaraRendererMeshes.cpp b/Engine/Plugins/FX/Niagara/Source/Niagara/Private/NiagaraRendererMeshes.cpp
--- a/Engine/Plugins/FX/Niagara/Source/Niagara/Private/NiagaraRendererMeshes.cpp
+++ b/Engine/Plugins/FX/Niagara/Source/Niagara/Private/NiagaraRendererMeshes.cpp
@@ -164,19 +164,21 @@ void FNiagaraRendererMeshes::Initialize(const UNiagaraRendererProperties* InProp
 
 	// Initialize the valid mesh slots, and prep them with the data for every mesh, LOD, and section we'll be needing over the lifetime of the renderer
 	const uint32 MaxMeshes = Properties->Meshes.Num();
-	Meshes.Reserve(MaxMeshes);
+	Meshes.Empty(MaxMeshes);
 	for (uint32 SourceMeshIndex = 0; SourceMeshIndex < MaxMeshes; ++SourceMeshIndex)
 	{
 		const auto& MeshProperties = Properties->Meshes[SourceMeshIndex];
-		if (MeshProperties.Mesh && MeshProperties.Mesh->GetRenderData())
+		UStaticMesh* Mesh = MeshProperties.ResolveStaticMesh(Emitter);
+		
+		if (Mesh)
 		{
 			FMeshData& MeshData = Meshes.AddDefaulted_GetRef();
-			MeshData.RenderData = MeshProperties.Mesh->GetRenderData();
+			MeshData.RenderData = Mesh->GetRenderData();
 			MeshData.SourceMeshIndex = SourceMeshIndex;
 			MeshData.PivotOffset = MeshProperties.PivotOffset;
 			MeshData.PivotOffsetSpace = MeshProperties.PivotOffsetSpace;
 			MeshData.Scale = MeshProperties.Scale;
-			MeshData.MinimumLOD = MeshProperties.Mesh->GetMinLOD().GetValue();
+			MeshData.MinimumLOD = Mesh->GetMinLOD().GetValue();
 
 			// Create an index remap from mesh material index to it's index in the master material list
 			TArray<UMaterialInterface*> MeshMaterials;
@@ -198,7 +200,7 @@ void FNiagaraRendererMeshes::Initialize(const UNiagaraRendererProperties* InProp
 			}
 
 			// Extend the local bounds by this mesh's bounds
-			FBox LocalBounds = MeshProperties.Mesh->GetExtendedBounds().GetBox();
+			FBox LocalBounds = Mesh->GetExtendedBounds().GetBox();
 			LocalBounds.Min *= MeshProperties.Scale;
 			LocalBounds.Max *= MeshProperties.Scale;
 			MeshData.LocalCullingSphere.Center = LocalBounds.GetCenter();
@@ -1404,7 +1406,8 @@ FNiagaraDynamicDataBase* FNiagaraRendererMeshes::GenerateDynamicData(const FNiag
 	// Bail if we have cached mesh render data for any meshes that are no longer valid
 	for (const auto& MeshData : Meshes)
 	{
-		if (!Properties->Meshes.IsValidIndex(MeshData.SourceMeshIndex) || !Properties->Meshes[MeshData.SourceMeshIndex].Mesh)
+		if (!Properties->Meshes.IsValidIndex(MeshData.SourceMeshIndex) || 
+			!Properties->Meshes[MeshData.SourceMeshIndex].HasValidMeshProperties())
 		{
 			return nullptr;
 		}
diff --git a/Engine/Plugins/FX/Niagara/Source/Niagara/Private/NiagaraSystemInstanceController.cpp b/Engine/Plugins/FX/Niagara/Source/Niagara/Private/NiagaraSystemInstanceController.cpp
--- a/Engine/Plugins/FX/Niagara/Source/Niagara/Private/NiagaraSystemInstanceController.cpp
+++ b/Engine/Plugins/FX/Niagara/Source/Niagara/Private/NiagaraSystemInstanceController.cpp
@@ -4,6 +4,9 @@
 #include "Materials/MaterialInstanceDynamic.h"
 
 FNiagaraSystemInstanceController::FNiagaraSystemInstanceController()
+	: bNeedsRendererRecache(false)
+	, bNeedsOverrideParametersTicked(false)
+	, bNeedsUpdateEmitterMaterials(false)
 {
 }
 
@@ -68,6 +71,21 @@ void FNiagaraSystemInstanceController::PostTickRenderers(FNiagaraSystemRenderDat
 {
 	if (auto SystemInst = SystemInstance.Get())
 	{
+		if (bNeedsOverrideParametersTicked)
+		{
+			bNeedsOverrideParametersTicked = false;
+			OverrideParameters->Tick();
+		}
+		if (bNeedsUpdateEmitterMaterials)
+		{
+			bNeedsUpdateEmitterMaterials = false;
+			UpdateEmitterMaterials();
+		}
+		if (bNeedsRendererRecache)
+		{
+			RenderData.RecacheRenderers(*SystemInst, *this);
+			bNeedsRendererRecache = false;
+		}
 		RenderData.PostTickRenderers(*SystemInst);
 	}
 }
@@ -403,7 +421,24 @@ void FNiagaraSystemInstanceController::SetVariable(FName InVariableName, TWeakOb
 		OverrideParameters->SetUObject(NewValue, VariableDesc);
 		if (CurrentValue != NewValue)
 		{
-			UpdateEmitterMaterials(); // Will need to update our internal tables. Maybe need a new MID.
+			bNeedsUpdateEmitterMaterials = true; // Will need to update our internal tables. Maybe need a new MID.
+		}
+	}
+}
+
+void FNiagaraSystemInstanceController::SetVariable(FName InVariableName, TWeakObjectPtr<UStaticMesh> InValue)
+{
+	if (OverrideParameters)
+	{
+		const FNiagaraVariable VariableDesc(FNiagaraTypeDefinition::GetUStaticMeshDef(), InVariableName);
+		UObject* CurrentValue = OverrideParameters->GetUObject(VariableDesc);
+		UStaticMesh* NewValue = InValue.Get();
+		OverrideParameters->SetUObject(NewValue, VariableDesc);
+		if (CurrentValue != NewValue)
+		{
+			bNeedsOverrideParametersTicked = true;
+			bNeedsUpdateEmitterMaterials = true;
+			OnNeedsRendererRecache();
 		}
 	}
 }
diff --git a/Engine/Plugins/FX/Niagara/Source/Niagara/Private/NiagaraSystemRenderData.cpp b/Engine/Plugins/FX/Niagara/Source/Niagara/Private/NiagaraSystemRenderData.cpp
--- a/Engine/Plugins/FX/Niagara/Source/Niagara/Private/NiagaraSystemRenderData.cpp
+++ b/Engine/Plugins/FX/Niagara/Source/Niagara/Private/NiagaraSystemRenderData.cpp
@@ -13,44 +13,19 @@ void FNiagaraSystemRenderData::ExecuteDynamicDataCommands_RenderThread(const FSe
 	}
 }
 
-FNiagaraSystemRenderData::FNiagaraSystemRenderData(const FNiagaraSystemInstanceController& SystemInstanceController, const FNiagaraSystemInstance& SystemInstance, ERHIFeatureLevel::Type FeatureLevel)
+FNiagaraSystemRenderData::FNiagaraSystemRenderData(const FNiagaraSystemInstanceController& SystemInstanceController, const FNiagaraSystemInstance& SystemInstance, ERHIFeatureLevel::Type InFeatureLevel)
 {
 	UNiagaraSystem* System = SystemInstance.GetSystem();
 	check(System);
+	FeatureLevel = InFeatureLevel;
 
 	RendererDrawOrder = System->GetRendererDrawOrder();
-	EmitterRenderers.Reserve(RendererDrawOrder.Num());
-
-	for (TSharedRef<const FNiagaraEmitterInstance, ESPMode::ThreadSafe> EmitterInst : SystemInstance.GetEmitters())
-	{
-		if (UNiagaraEmitter* Emitter = EmitterInst->GetCachedEmitter())
-		{
-			Emitter->ForEachEnabledRenderer(
-				[&](UNiagaraRendererProperties* Properties)
-				{
-					//We can skip creation of the renderer if the current quality level doesn't support it. If the quality level changes all systems are fully reinitialized.
-					FNiagaraRenderer* NewRenderer = nullptr;
-					if (Properties->GetIsActive() && EmitterInst->GetData().IsInitialized() && !EmitterInst->IsDisabled())
-					{
-						NewRenderer = Properties->CreateEmitterRenderer(FeatureLevel, &EmitterInst.Get(), SystemInstanceController);
-						if (Properties->MotionVectorSetting != ENiagaraRendererMotionVectorSetting::Disable)
-						{
-							bAnyMotionBlurEnabled = true;
-						}
-					}
-					EmitterRenderers.Add(NewRenderer);
-				}
-			);
-		}
-	}
-
-	// If we have renderers then the draw order on the system should match, when compiling the number of renderers can be zero
-	checkf((EmitterRenderers.Num() == 0) || (EmitterRenderers.Num() == RendererDrawOrder.Num()), TEXT("EmitterRenderers Num %d does not match System DrawOrder %d"), EmitterRenderers.Num(), RendererDrawOrder.Num());
+	RecacheRenderers(SystemInstance, SystemInstanceController);
 }
 
 FNiagaraSystemRenderData::~FNiagaraSystemRenderData()
 {
-	if(EmitterRenderers.Num() > 0)
+	if(EmitterRenderers_RT.Num() > 0)
 	{
 		Destroy_RenderThread();
 	}
@@ -58,7 +33,7 @@ FNiagaraSystemRenderData::~FNiagaraSystemRenderData()
 
 void FNiagaraSystemRenderData::CreateRenderThreadResources(NiagaraEmitterInstanceBatcher& Batcher)
 {
-	for (auto Renderer : EmitterRenderers)
+	for (auto Renderer : EmitterRenderers_RT)
 	{
 		if (Renderer)
 		{
@@ -69,7 +44,7 @@ void FNiagaraSystemRenderData::CreateRenderThreadResources(NiagaraEmitterInstanc
 
 void FNiagaraSystemRenderData::ReleaseRenderThreadResources()
 {
-	for (auto Renderer : EmitterRenderers)
+	for (auto Renderer : EmitterRenderers_RT)
 	{
 		if (Renderer)
 		{
@@ -81,13 +56,16 @@ void FNiagaraSystemRenderData::ReleaseRenderThreadResources()
 void FNiagaraSystemRenderData::DestroyRenderState_Concurrent()
 {
 	// Give the renderers an opportunity to free resources that must be freed on the game thread (or task)
-	for (auto Renderer : EmitterRenderers)
+	for (auto Renderer : EmitterRenderers_GT)
 	{
 		if (Renderer)
 		{
 			Renderer->DestroyRenderState_Concurrent();
 		}
 	}
+
+	// Don't delete the renderers on game/concurrent threads, they have to be deleted on the render thread
+	EmitterRenderers_GT.Reset();
 }
 
 void FNiagaraSystemRenderData::Destroy_RenderThread()
@@ -95,7 +73,7 @@ void FNiagaraSystemRenderData::Destroy_RenderThread()
 	check(IsInRenderingThread());
 
 	// Delete the renderers
-	for (auto Renderer : EmitterRenderers)
+	for (auto Renderer : EmitterRenderers_RT)
 	{
 		if (Renderer)
 		{
@@ -104,7 +82,7 @@ void FNiagaraSystemRenderData::Destroy_RenderThread()
 		}
 	}
 
-	EmitterRenderers.Empty();
+	EmitterRenderers_RT.Reset();
 }
 
 FPrimitiveViewRelevance FNiagaraSystemRenderData::GetViewRelevance(const FSceneView& View, const FNiagaraSceneProxy& SceneProxy) const
@@ -112,7 +90,7 @@ FPrimitiveViewRelevance FNiagaraSystemRenderData::GetViewRelevance(const FSceneV
 	FPrimitiveViewRelevance Relevance;
 	if (IsRenderingEnabled())
 	{
-		for (const auto Renderer : EmitterRenderers)
+		for (const auto Renderer : EmitterRenderers_GT)
 		{
 			if (Renderer)
 			{
@@ -127,7 +105,7 @@ FPrimitiveViewRelevance FNiagaraSystemRenderData::GetViewRelevance(const FSceneV
 uint32 FNiagaraSystemRenderData::GetDynamicDataSize() const
 {
 	uint32 Size = 0;
-	for (const auto Renderer : EmitterRenderers)
+	for (const auto Renderer : EmitterRenderers_GT)
 	{
 		if (Renderer)
 		{
@@ -142,7 +120,7 @@ void FNiagaraSystemRenderData::GenerateSetDynamicDataCommands(FSetDynamicDataCom
 {
 	if (!SystemInstance)
 	{
-		for (auto Renderer : EmitterRenderers)
+		for (auto Renderer : EmitterRenderers_GT)
 		{
 			if (Renderer)
 			{
@@ -162,7 +140,7 @@ void FNiagaraSystemRenderData::GenerateSetDynamicDataCommands(FSetDynamicDataCom
 	FScopeCycleCounter SystemStatCounter(SystemStatID);
 #endif
 
-	const int32 NumEmitterRenderers = EmitterRenderers.Num();
+	const int32 NumEmitterRenderers = EmitterRenderers_GT.Num();
 	if (NumEmitterRenderers == 0)
 	{
 		// Early out if we have no renderers
@@ -190,7 +168,7 @@ void FNiagaraSystemRenderData::GenerateSetDynamicDataCommands(FSetDynamicDataCom
 		Emitter->ForEachEnabledRenderer(
 			[&](UNiagaraRendererProperties* Properties)
 			{
-				FNiagaraRenderer* Renderer = EmitterRenderers[RendererIndex];
+				FNiagaraRenderer* Renderer = EmitterRenderers_GT[RendererIndex];
 				FNiagaraDynamicDataBase* NewData = nullptr;
 
 				if (Renderer && Properties->GetIsActive())
@@ -227,7 +205,7 @@ void FNiagaraSystemRenderData::GetDynamicMeshElements(const TArray<const FSceneV
 {
 	for (int32 RendererIdx : RendererDrawOrder)
 	{
-		FNiagaraRenderer* Renderer = EmitterRenderers[RendererIdx];
+		FNiagaraRenderer* Renderer = EmitterRenderers_RT[RendererIdx];
 		if (Renderer && (Renderer->GetSimTarget() != ENiagaraSimTarget::GPUComputeSim || FNiagaraUtilities::AllowGPUParticles(ViewFamily.GetShaderPlatform())))
 		{
 			Renderer->GetDynamicMeshElements(Views, ViewFamily, VisibilityMap, Collector, &SceneProxy);
@@ -238,7 +216,7 @@ void FNiagaraSystemRenderData::GetDynamicMeshElements(const TArray<const FSceneV
 #if RHI_RAYTRACING
 void FNiagaraSystemRenderData::GetDynamicRayTracingInstances(FRayTracingMaterialGatheringContext& Context, TArray<FRayTracingInstance>& OutRayTracingInstances, const FNiagaraSceneProxy& SceneProxy)
 {
-	for (auto Renderer : EmitterRenderers)
+	for (auto Renderer : EmitterRenderers_RT)
 	{
 		if (Renderer)
 		{
@@ -250,7 +228,7 @@ void FNiagaraSystemRenderData::GetDynamicRayTracingInstances(FRayTracingMaterial
 
 void FNiagaraSystemRenderData::GatherSimpleLights(FSimpleLightArray& OutParticleLights) const
 {
-	for (auto Renderer : EmitterRenderers)
+	for (auto Renderer : EmitterRenderers_RT)
 	{
 		if (Renderer && Renderer->HasLights())
 		{
@@ -259,19 +237,65 @@ void FNiagaraSystemRenderData::GatherSimpleLights(FSimpleLightArray& OutParticle
 	}
 }
 
+void FNiagaraSystemRenderData::RecacheRenderers(const FNiagaraSystemInstance& SystemInstance, const FNiagaraSystemInstanceController& Controller)
+{
+	DestroyRenderState_Concurrent();
+
+	EmitterRenderers_GT.Reserve(RendererDrawOrder.Num());
+
+	bAnyMotionBlurEnabled = false;
+	for (TSharedRef<const FNiagaraEmitterInstance, ESPMode::ThreadSafe> EmitterInst : SystemInstance.GetEmitters())
+	{
+		if (UNiagaraEmitter* Emitter = EmitterInst->GetCachedEmitter())
+		{
+			Emitter->ForEachEnabledRenderer(
+				[&](UNiagaraRendererProperties* Properties)
+			{
+				//We can skip creation of the renderer if the current quality level doesn't support it. If the quality level changes all systems are fully reinitialized.
+				FNiagaraRenderer* NewRenderer = nullptr;
+				if (Properties->GetIsActive() && EmitterInst->GetData().IsInitialized() && !EmitterInst->IsDisabled())
+				{
+					NewRenderer = Properties->CreateEmitterRenderer(FeatureLevel, &EmitterInst.Get(), Controller);
+					if (Properties->MotionVectorSetting != ENiagaraRendererMotionVectorSetting::Disable)
+					{
+						bAnyMotionBlurEnabled = true;
+					}
+				}
+				EmitterRenderers_GT.Add(NewRenderer);
+			}
+			);
+		}
+	}
+
+	// If we have renderers then the draw order on the system should match, when compiling the number of renderers can be zero
+	checkf((EmitterRenderers_GT.Num() == 0) || (EmitterRenderers_GT.Num() == RendererDrawOrder.Num()), TEXT("EmitterRenderers Num %d does not match System DrawOrder %d"), EmitterRenderers_GT.Num(), RendererDrawOrder.Num());
+	
+	// NOTE: Since this object and its render thread resources may be concurrently accessed on the render thread, we have to create the renderers and pass them off to
+	// replace the current ones on the render thread's time line
+	ENQUEUE_RENDER_COMMAND(NiagaraRecacheRenderers)(
+		[this, EmitterRenderers_Copy=EmitterRenderers_GT](FRHICommandListImmediate& RHICmdList) mutable
+		{
+			// Release/destroy the current renderers
+			Destroy_RenderThread();
+
+			EmitterRenderers_RT = MoveTemp(EmitterRenderers_Copy);
+		}
+	);
+}
+
 void FNiagaraSystemRenderData::PostTickRenderers(const FNiagaraSystemInstance& SystemInstance)
 {
 	// Give renderers a chance to do some processing PostTick
-	if (EmitterRenderers.Num() > 0)
+	if (EmitterRenderers_GT.Num() > 0)
 	{
 		const UNiagaraSystem* System = SystemInstance.GetSystem();
 		for (const FNiagaraRendererExecutionIndex& ExecIdx : System->GetRendererPostTickOrder())
 		{
-			if (EmitterRenderers.IsValidIndex(ExecIdx.SystemRendererIndex))
+			if (EmitterRenderers_GT.IsValidIndex(ExecIdx.SystemRendererIndex))
 			{
 				const FNiagaraEmitterInstance& EmitterInst = SystemInstance.GetEmitters()[ExecIdx.EmitterIndex].Get();
 				const UNiagaraEmitter* Emitter = EmitterInst.GetCachedEmitter();
-				FNiagaraRenderer* EmitterRenderer = EmitterRenderers[ExecIdx.SystemRendererIndex];
+				FNiagaraRenderer* EmitterRenderer = EmitterRenderers_GT[ExecIdx.SystemRendererIndex];
 				if (Emitter && EmitterRenderer)
 				{
 					const UNiagaraRendererProperties* RendererProperties = Emitter->GetRenderers()[ExecIdx.EmitterRendererIndex];
@@ -285,16 +309,16 @@ void FNiagaraSystemRenderData::PostTickRenderers(const FNiagaraSystemInstance& S
 void FNiagaraSystemRenderData::OnSystemComplete(const FNiagaraSystemInstance& SystemInstance)
 {
 	// Give renderers a chance to handle completion
-	if (EmitterRenderers.Num() > 0)
+	if (EmitterRenderers_GT.Num() > 0)
 	{
 		const UNiagaraSystem* System = SystemInstance.GetSystem();
 		for (const FNiagaraRendererExecutionIndex& ExecIdx : System->GetRendererCompletionOrder())
 		{
-			if (EmitterRenderers.IsValidIndex(ExecIdx.SystemRendererIndex))
+			if (EmitterRenderers_GT.IsValidIndex(ExecIdx.SystemRendererIndex))
 			{
 				const FNiagaraEmitterInstance& EmitterInst = SystemInstance.GetEmitters()[ExecIdx.EmitterIndex].Get();
 				const UNiagaraEmitter* Emitter = EmitterInst.GetCachedEmitter();
-				FNiagaraRenderer* EmitterRenderer = EmitterRenderers[ExecIdx.SystemRendererIndex];
+				FNiagaraRenderer* EmitterRenderer = EmitterRenderers_GT[ExecIdx.SystemRendererIndex];
 				if (Emitter && EmitterRenderer)
 				{
 					const UNiagaraRendererProperties* RendererProperties = Emitter->GetRenderers()[ExecIdx.EmitterRendererIndex];
diff --git a/Engine/Plugins/FX/Niagara/Source/Niagara/Public/NiagaraComponent.h b/Engine/Plugins/FX/Niagara/Source/Niagara/Public/NiagaraComponent.h
--- a/Engine/Plugins/FX/Niagara/Source/Niagara/Public/NiagaraComponent.h
+++ b/Engine/Plugins/FX/Niagara/Source/Niagara/Public/NiagaraComponent.h
@@ -407,6 +407,9 @@ public:
 	UFUNCTION(BlueprintCallable, Category = Niagara, meta = (DisplayName = "Set Niagara Variable (Material)"))
 	void SetVariableMaterial(FName InVariableName, UMaterialInterface* Object);
 
+	UFUNCTION(BlueprintCallable, Category = Niagara, meta = (DisplayName = "Set Niagara Variable (Static Mesh)"))
+	void SetVariableStaticMesh(FName InVariableName, UStaticMesh* InValue);
+
 	UFUNCTION(BlueprintCallable, Category = Niagara, meta = (DisplayName = "Set Niagara Variable (TextureRenderTarget)"))
 	void SetVariableTextureRenderTarget(FName InVariableName, class UTextureRenderTarget* TextureRenderTarget);
 
diff --git a/Engine/Plugins/FX/Niagara/Source/Niagara/Public/NiagaraMeshRendererProperties.h b/Engine/Plugins/FX/Niagara/Source/Niagara/Public/NiagaraMeshRendererProperties.h
--- a/Engine/Plugins/FX/Niagara/Source/Niagara/Public/NiagaraMeshRendererProperties.h
+++ b/Engine/Plugins/FX/Niagara/Source/Niagara/Public/NiagaraMeshRendererProperties.h
@@ -120,6 +120,10 @@ struct NIAGARA_API FNiagaraMeshRendererMeshProperties
 	UPROPERTY(EditAnywhere, Category = "Mesh")
 	TObjectPtr<UStaticMesh> Mesh;
 
+	/** Use the UStaticMesh bound to this user variable if it is set to a valid value. If this is bound to a valid value and Mesh is also set, UserParamBinding wins.*/
+	UPROPERTY(EditAnywhere, Category = "Mesh")
+	FNiagaraUserParameterBinding UserParamBinding;
+
 	/** Scale of the mesh */
 	UPROPERTY(EditAnywhere, Category = "Mesh")
 	FVector Scale;
@@ -131,6 +135,9 @@ struct NIAGARA_API FNiagaraMeshRendererMeshProperties
 	/** What space is the pivot offset in? */
 	UPROPERTY(EditAnywhere, Category = "Mesh")
 	ENiagaraMeshPivotOffsetSpace PivotOffsetSpace;
+
+    UStaticMesh* ResolveStaticMesh(const FNiagaraEmitterInstance* Emitter) const;
+	bool HasValidMeshProperties() const;
 };
 
 UCLASS(editinlinenew, meta = (DisplayName = "Mesh Renderer"))
diff --git a/Engine/Plugins/FX/Niagara/Source/Niagara/Public/NiagaraSystemInstanceController.h b/Engine/Plugins/FX/Niagara/Source/Niagara/Public/NiagaraSystemInstanceController.h
--- a/Engine/Plugins/FX/Niagara/Source/Niagara/Public/NiagaraSystemInstanceController.h
+++ b/Engine/Plugins/FX/Niagara/Source/Niagara/Public/NiagaraSystemInstanceController.h
@@ -162,14 +162,20 @@ private:
 	void SetVariable(FName InVariableName, FQuat InValue);
 	void SetVariable(FName InVariableName, TWeakObjectPtr<UObject> Object);
 	void SetVariable(FName InVariableName, TWeakObjectPtr<UMaterialInterface> Object);
+	void SetVariable(FName InVariableName, TWeakObjectPtr<UStaticMesh> Object);
 	void SetVariable(FName InVariableName, TWeakObjectPtr<UTextureRenderTarget> TextureRenderTarget);
 
 	NIAGARA_SYSTEM_INSTANCE_CONTROLLER_SHIM(SetRequestedExecutionState,)
 
+	void OnNeedsRendererRecache() { bNeedsRendererRecache = true; }
+
 private:
 	FNiagaraWorldManager* WorldManager = nullptr;
 	FNiagaraUserRedirectionParameterStore* OverrideParameters = nullptr;
 	FNiagaraSystemInstancePtr SystemInstance;
 	TArray<FMaterialOverride> EmitterMaterials;
 	FOnMaterialsUpdated OnMaterialsUpdatedDelegate;
+	uint32 bNeedsRendererRecache : 1;
+	uint32 bNeedsOverrideParametersTicked : 1;
+	uint32 bNeedsUpdateEmitterMaterials : 1;
 };
diff --git a/Engine/Plugins/FX/Niagara/Source/Niagara/Public/NiagaraSystemRenderData.h b/Engine/Plugins/FX/Niagara/Source/Niagara/Public/NiagaraSystemRenderData.h
--- a/Engine/Plugins/FX/Niagara/Source/Niagara/Public/NiagaraSystemRenderData.h
+++ b/Engine/Plugins/FX/Niagara/Source/Niagara/Public/NiagaraSystemRenderData.h
@@ -63,20 +63,23 @@ public:
 
 	void PostTickRenderers(const FNiagaraSystemInstance& SystemInstance);
 	void OnSystemComplete(const FNiagaraSystemInstance& SystemInstance);
+	void RecacheRenderers(const FNiagaraSystemInstance& SystemInstance, const FNiagaraSystemInstanceController& Controller);
 
-	FORCEINLINE bool IsRenderingEnabled() const { return bRenderingEnabled && EmitterRenderers.Num() > 0; }
+	FORCEINLINE bool IsRenderingEnabled() const { return bRenderingEnabled && EmitterRenderers_GT.Num() > 0; }
 	FORCEINLINE void SetRenderingEnabled(bool bInEnabled) { bRenderingEnabled = bInEnabled; }
-
 	FORCEINLINE bool HasAnyMotionBlurEnabled() const { return bAnyMotionBlurEnabled; }
 
-
 private:
-	/** Emitter Renderers in the order they appear in the emitters. */
-	TArray<FNiagaraRenderer*> EmitterRenderers;
+	/** Emitter Renderers in the order they appear in the emitters. To be accessed by the GameThread */
+	TArray<FNiagaraRenderer*> EmitterRenderers_GT;
+	/** Emitter Renderers in the order they appear in the emitters. To be accessed by the RenderThread */
+	TArray<FNiagaraRenderer*> EmitterRenderers_RT;
 
 	/** Indices of renderers in the order they should be rendered. */
 	TArray<int32> RendererDrawOrder;
 
 	bool bRenderingEnabled = true;
 	bool bAnyMotionBlurEnabled = false;
+
+	ERHIFeatureLevel::Type FeatureLevel;
 };
diff --git a/Engine/Plugins/FX/Niagara/Source/Niagara/Public/NiagaraTypes.h b/Engine/Plugins/FX/Niagara/Source/Niagara/Public/NiagaraTypes.h
--- a/Engine/Plugins/FX/Niagara/Source/Niagara/Public/NiagaraTypes.h
+++ b/Engine/Plugins/FX/Niagara/Source/Niagara/Public/NiagaraTypes.h
@@ -1048,6 +1048,7 @@ public:
 	static const FNiagaraTypeDefinition& GetIDDef() { return IDDef; }
 	static const FNiagaraTypeDefinition& GetUObjectDef() { return UObjectDef; }
 	static const FNiagaraTypeDefinition& GetUMaterialDef() { return UMaterialDef; }
+	static const FNiagaraTypeDefinition& GetUStaticMeshDef() { return UStaticMeshDef; }
 	static const FNiagaraTypeDefinition& GetUTextureDef() { return UTextureDef; }
 	static const FNiagaraTypeDefinition& GetUTextureRenderTargetDef() { return UTextureRenderTargetDef; }
 	static const FNiagaraTypeDefinition& GetWildcardDef() { return WildcardDef; }
@@ -1125,6 +1126,7 @@ private:
 	static FNiagaraTypeDefinition UMaterialDef;
 	static FNiagaraTypeDefinition UTextureDef;
 	static FNiagaraTypeDefinition UTextureRenderTargetDef;
+	static FNiagaraTypeDefinition UStaticMeshDef;
 	static FNiagaraTypeDefinition WildcardDef;
 
 	static FNiagaraTypeDefinition HalfDef;
diff --git a/Engine/Plugins/FX/Niagara/Source/NiagaraEditorWidgets/Private/Stack/SNiagaraStackParameterStoreEntryValue.cpp b/Engine/Plugins/FX/Niagara/Source/NiagaraEditorWidgets/Private/Stack/SNiagaraStackParameterStoreEntryValue.cpp
--- a/Engine/Plugins/FX/Niagara/Source/NiagaraEditorWidgets/Private/Stack/SNiagaraStackParameterStoreEntryValue.cpp
+++ b/Engine/Plugins/FX/Niagara/Source/NiagaraEditorWidgets/Private/Stack/SNiagaraStackParameterStoreEntryValue.cpp
@@ -214,6 +214,22 @@ TSharedRef<SWidget> SNiagaraStackParameterStoreEntryValue::ConstructValueStructW
 				.DisplayThumbnail(true)
 				.NewAssetFactories(TArray<UFactory*>());
 
+		}
+		else if (StackEntry->GetInputType().GetClass()->IsChildOf(UStaticMesh::StaticClass()))
+		{
+			TArray<const UClass*> AllowedClasses;
+			AllowedClasses.Add(UStaticMesh::StaticClass());
+
+			return SNew(SObjectPropertyEntryBox)
+				.ObjectPath_Raw(this, &SNiagaraStackParameterStoreEntryValue::GetCurrentAssetPath)
+				.AllowedClass(UStaticMesh::StaticClass())
+				.OnObjectChanged_Raw(this, &SNiagaraStackParameterStoreEntryValue::OnAssetSelectedFromPicker, UStaticMesh::StaticClass())
+				.AllowClear(false)
+				.DisplayUseSelected(true)
+				.DisplayBrowse(true)
+				.DisplayThumbnail(true)
+				.NewAssetFactories(TArray<UFactory*>());
+
 		}
 		else
 		{
