Use the render thread array in NiagaraSytemRenderData because view relevance is retrieved from the render thread.

diff --git a/Engine/Plugins/FX/Niagara/Source/Niagara/Private/NiagaraSystemRenderData.cpp b/Engine/Plugins/FX/Niagara/Source/Niagara/Private/NiagaraSystemRenderData.cpp
--- a/Engine/Plugins/FX/Niagara/Source/Niagara/Private/NiagaraSystemRenderData.cpp
+++ b/Engine/Plugins/FX/Niagara/Source/Niagara/Private/NiagaraSystemRenderData.cpp
@@ -90,7 +90,7 @@ FPrimitiveViewRelevance FNiagaraSystemRenderData::GetViewRelevance(const FSceneV
 	FPrimitiveViewRelevance Relevance;
 	if (IsRenderingEnabled())
 	{
-		for (const auto Renderer : EmitterRenderers_GT)
+		for (const auto Renderer : EmitterRenderers_RT)
 		{
 			if (Renderer)
 			{
