Fix non-unity compile error

diff --git a/Engine/Plugins/FX/Niagara/Source/Niagara/Private/NiagaraComponent.cpp b/Engine/Plugins/FX/Niagara/Source/Niagara/Private/NiagaraComponent.cpp
--- a/Engine/Plugins/FX/Niagara/Source/Niagara/Private/NiagaraComponent.cpp
+++ b/Engine/Plugins/FX/Niagara/Source/Niagara/Private/NiagaraComponent.cpp
@@ -23,6 +23,7 @@
 #include "ProfilingDebugging/CsvProfiler.h"
 #include "UObject/NameTypes.h"
 #include "VectorVM.h"
+#include "Engine/StaticMesh.h"
 
 DECLARE_CYCLE_STAT(TEXT("Sceneproxy create (GT)"), STAT_NiagaraCreateSceneProxy, STATGROUP_Niagara);
 DECLARE_CYCLE_STAT(TEXT("Component Tick (GT)"), STAT_NiagaraComponentTick, STATGROUP_Niagara);
