Couple of fixes for unlimited skel mesh influences with the skel mesh DI (fixed an incorrect log message and added a missing shader parameter)

diff --git a/Engine/Plugins/FX/Niagara/Source/Niagara/Private/NiagaraDataInterfaceSkeletalMesh.cpp b/Engine/Plugins/FX/Niagara/Source/Niagara/Private/NiagaraDataInterfaceSkeletalMesh.cpp
--- a/Engine/Plugins/FX/Niagara/Source/Niagara/Private/NiagaraDataInterfaceSkeletalMesh.cpp
+++ b/Engine/Plugins/FX/Niagara/Source/Niagara/Private/NiagaraDataInterfaceSkeletalMesh.cpp
@@ -1448,6 +1448,7 @@ public:
 			SetSRVParameter(RHICmdList, ComputeShaderRHI, MeshSkinWeightLookupBuffer, MeshSkinWeightLookupBufferSRV ? MeshSkinWeightLookupBufferSRV : FNiagaraRenderer::GetDummyUIntBuffer());
 
 			SetShaderValue(RHICmdList, ComputeShaderRHI, MeshWeightStride, InstanceData->MeshWeightStrideByte/4);
+			SetShaderValue(RHICmdList, ComputeShaderRHI, MeshSkinWeightIndexSize, InstanceData->MeshSkinWeightIndexSizeByte);
 
 			uint32 EnabledFeaturesBits = 0;
 			EnabledFeaturesBits |= StaticBuffers->IsUseGpuUniformlyDistributedSampling() ? 1 : 0;
@@ -2205,7 +2206,23 @@ bool FNDISkeletalMesh_InstanceData::Init(UNiagaraDataInterfaceSkeletalMesh* Inte
 			{
 				UE_LOG(LogNiagara, Warning, TEXT("Skeletal Mesh %s has cloth asset on it: spawning from it might not work properly."), *Mesh->GetName());
 			}
-			if (CachedLODData->GetVertexBufferMaxBoneInfluences() > MAX_INFLUENCES_PER_STREAM)
+
+			const auto MaxInfluenceType = GetDefault<UNiagaraSettings>()->NDISkelMesh_GpuMaxInfluences;
+			int32 MaxInfluenceCount = -1;
+			if (MaxInfluenceType == ENDISkelMesh_GpuMaxInfluences::Type::AllowMax4)
+			{
+				MaxInfluenceCount = 4;
+			}
+			else if (MaxInfluenceType == ENDISkelMesh_GpuMaxInfluences::Type::AllowMax8)
+			{
+				MaxInfluenceCount = 8;
+			}
+			else
+			{
+				checkf(MaxInfluenceType == ENDISkelMesh_GpuMaxInfluences::Type::Unlimited, TEXT("Unknown value for NDISkelMesh_GpuMaxInfluences: %d"), MaxInfluenceType);
+			}
+
+			if (MaxInfluenceCount > 0 && (static_cast<uint32>(MaxInfluenceCount) < CachedLODData->GetVertexBufferMaxBoneInfluences()))
 			{
 				UE_LOG(LogNiagara, Warning, TEXT("Skeletal Mesh %s has bones extra influence: spawning from it might not work properly."), *Mesh->GetName());
 			}
