- Remove divides and mods from attribute lookups

diff --git a/Engine/Plugins/FX/Niagara/Source/Niagara/Classes/NiagaraDataInterfaceGrid3DCollection.h b/Engine/Plugins/FX/Niagara/Source/Niagara/Classes/NiagaraDataInterfaceGrid3DCollection.h
--- a/Engine/Plugins/FX/Niagara/Source/Niagara/Classes/NiagaraDataInterfaceGrid3DCollection.h
+++ b/Engine/Plugins/FX/Niagara/Source/Niagara/Classes/NiagaraDataInterfaceGrid3DCollection.h
@@ -66,6 +66,7 @@ struct FGrid3DCollectionRWInstanceData_RenderThread
 	FVector WorldBBoxSize = FVector::ZeroVector;
 	EPixelFormat PixelFormat = EPixelFormat::PF_R32_FLOAT;
 	TArray<int32> AttributeIndices;
+
 	TArray<FName> Vars;
 	TArray<int32> VarComponents;
 	TArray<uint32> Offsets;
@@ -202,6 +203,7 @@ public:
 	void GetAttributeIndex(FVectorVMContext& Context, const FName& InName, int32 NumChannels);
 
 	static const FString NumTilesName;
+	static const FString OneOverNumTilesName;
 
 	static const FString GridName;
 	static const FString OutputGridName;
@@ -237,6 +239,10 @@ public:
 	static const FName GetPreviousFloatValueFunctionName;
 	static const FName SamplePreviousGridFloatFunctionName;
 	static const FString AttributeIndicesBaseName;
+	static const FString AttributePixelOffsetName;
+	static const FString AttributeUVOffsetName;
+	static const FString AttributeUVClampMinName;
+	static const FString AttributeUVClampMaxName;
 	static const TCHAR* VectorComponentNames[];
 
 	static const FName SetNumCellsFunctionName;
diff --git a/Engine/Plugins/FX/Niagara/Source/Niagara/Private/NiagaraDataInterfaceGrid3DCollection.cpp b/Engine/Plugins/FX/Niagara/Source/Niagara/Private/NiagaraDataInterfaceGrid3DCollection.cpp
--- a/Engine/Plugins/FX/Niagara/Source/Niagara/Private/NiagaraDataInterfaceGrid3DCollection.cpp
+++ b/Engine/Plugins/FX/Niagara/Source/Niagara/Private/NiagaraDataInterfaceGrid3DCollection.cpp
@@ -18,6 +18,7 @@
 #define LOCTEXT_NAMESPACE "NiagaraDataInterfaceGrid3DCollection"
 
 const FString UNiagaraDataInterfaceGrid3DCollection::NumTilesName(TEXT("NumTiles_"));
+const FString UNiagaraDataInterfaceGrid3DCollection::OneOverNumTilesName(TEXT("OneOverNumTiles_"));
 
 const FString UNiagaraDataInterfaceGrid3DCollection::GridName(TEXT("Grid_"));
 const FString UNiagaraDataInterfaceGrid3DCollection::OutputGridName(TEXT("OutputGrid_"));
@@ -63,6 +64,11 @@ const FName UNiagaraDataInterfaceGrid3DCollection::GetFloatAttributeIndexFunctio
 const FName UNiagaraDataInterfaceGrid3DCollection::SampleGridFunctionName("SampleGrid");
 
 const FString UNiagaraDataInterfaceGrid3DCollection::AttributeIndicesBaseName(TEXT("AttributeIndices_"));
+const FString UNiagaraDataInterfaceGrid3DCollection::AttributePixelOffsetName(TEXT("AttributePixelOffset_"));
+const FString UNiagaraDataInterfaceGrid3DCollection::AttributeUVOffsetName(TEXT("AttributeUVOffset_"));
+const FString UNiagaraDataInterfaceGrid3DCollection::AttributeUVClampMinName(TEXT("AttributeUVClampMin_"));
+const FString UNiagaraDataInterfaceGrid3DCollection::AttributeUVClampMaxName(TEXT("AttributeUVClampMax_"));
+
 const TCHAR* UNiagaraDataInterfaceGrid3DCollection::VectorComponentNames[] = { TEXT(".x"), TEXT(".y"), TEXT(".z"), TEXT(".w") };
 
 FNiagaraVariableBase UNiagaraDataInterfaceGrid3DCollection::ExposedRTVar;
@@ -146,6 +152,7 @@ public:
 		NumAttributesParam.Bind(ParameterMap, *(UNiagaraDataInterfaceRWBase::NumAttributesName + ParameterInfo.DataInterfaceHLSLSymbol));
 		NumCellsParam.Bind(ParameterMap, *(UNiagaraDataInterfaceRWBase::NumCellsName + ParameterInfo.DataInterfaceHLSLSymbol));
 		NumTilesParam.Bind(ParameterMap, *(UNiagaraDataInterfaceGrid3DCollection::NumTilesName + ParameterInfo.DataInterfaceHLSLSymbol));
+		OneOverNumTilesParam.Bind(ParameterMap, *(UNiagaraDataInterfaceGrid3DCollection::OneOverNumTilesName + ParameterInfo.DataInterfaceHLSLSymbol));
 		UnitToUVParam.Bind(ParameterMap, *(UNiagaraDataInterfaceRWBase::UnitToUVName + ParameterInfo.DataInterfaceHLSLSymbol));
 		CellSizeParam.Bind(ParameterMap, *(UNiagaraDataInterfaceRWBase::CellSizeName + ParameterInfo.DataInterfaceHLSLSymbol));
 
@@ -156,6 +163,10 @@ public:
 
 		SamplerParam.Bind(ParameterMap, *(UNiagaraDataInterfaceGrid3DCollection::SamplerName + ParameterInfo.DataInterfaceHLSLSymbol));
 		AttributeIndicesParam.Bind(ParameterMap, *(UNiagaraDataInterfaceGrid3DCollection::AttributeIndicesBaseName + ParameterInfo.DataInterfaceHLSLSymbol));
+		AttributePixelOffsetParam.Bind(ParameterMap, *(UNiagaraDataInterfaceGrid3DCollection::AttributePixelOffsetName + ParameterInfo.DataInterfaceHLSLSymbol));
+		AttributeUVOffsetParam.Bind(ParameterMap, *(UNiagaraDataInterfaceGrid3DCollection::AttributeUVOffsetName + ParameterInfo.DataInterfaceHLSLSymbol));
+		AttributeUVClampMinParam.Bind(ParameterMap, *(UNiagaraDataInterfaceGrid3DCollection::AttributeUVClampMinName + ParameterInfo.DataInterfaceHLSLSymbol));
+		AttributeUVClampMaxParam.Bind(ParameterMap, *(UNiagaraDataInterfaceGrid3DCollection::AttributeUVClampMaxName + ParameterInfo.DataInterfaceHLSLSymbol));
 
 
 		// Gather up all the attribute names referenced. Note that there may be multiple in the list of the same name,
@@ -194,17 +205,9 @@ public:
 		FGrid3DCollectionRWInstanceData_RenderThread* ProxyData = VFDI->SystemInstancesToProxyData_RT.Find(Context.SystemInstanceID);
 		check(ProxyData);
 
-		int NumCellsTmp[3];
-		NumCellsTmp[0] = ProxyData->NumCells.X;
-		NumCellsTmp[1] = ProxyData->NumCells.Y;
-		NumCellsTmp[2] = ProxyData->NumCells.Z;
-		SetShaderValue(RHICmdList, ComputeShaderRHI, NumCellsParam, NumCellsTmp);
-
-		int NumTilesTmp[3];
-		NumTilesTmp[0] = ProxyData->NumTiles.X;
-		NumTilesTmp[1] = ProxyData->NumTiles.Y;
-		NumTilesTmp[2] = ProxyData->NumTiles.Z;
-		SetShaderValue(RHICmdList, ComputeShaderRHI, NumTilesParam, NumTilesTmp);
+		SetShaderValue(RHICmdList, ComputeShaderRHI, NumCellsParam, ProxyData->NumCells);	
+		SetShaderValue(RHICmdList, ComputeShaderRHI, NumTilesParam, ProxyData->NumTiles);
+		SetShaderValue(RHICmdList, ComputeShaderRHI, OneOverNumTilesParam, FVector(1.0f) / FVector(ProxyData->NumTiles));
 
 		if (ProxyData->AttributeIndices.Num() == 0 && AttributeNames.Num() > 0)
 		{
@@ -226,7 +229,6 @@ public:
 				{
 					ProxyData->AttributeIndices[i] = -1; // We may need to protect against this in the hlsl as this might underflow an array lookup if used incorrectly.
 				}
-
 			}
 		}
 
@@ -239,6 +241,53 @@ public:
 		SetShaderValue(RHICmdList, ComputeShaderRHI, WorldBBoxSizeParam, ProxyData->WorldBBoxSize);
 
 		SetShaderValueArray(RHICmdList, ComputeShaderRHI, AttributeIndicesParam, ProxyData->AttributeIndices.GetData(), ProxyData->AttributeIndices.Num());
+		{
+			const int32 NumAttributes = ProxyData->TotalNumAttributes;
+			TArray<FIntVector4, TInlineAllocator<16>> AttributePixelOffset;
+			TArray<FVector4, TInlineAllocator<16>> AttributeUVOffset;
+			TArray<FVector4, TInlineAllocator<16>> AttributeUVClampMin;
+			TArray<FVector4, TInlineAllocator<16>> AttributeUVClampMax;
+			AttributePixelOffset.Reserve(NumAttributes);
+			AttributeUVOffset.Reserve(NumAttributes);
+			AttributeUVClampMin.Reserve(NumAttributes);
+			AttributeUVClampMax.Reserve(NumAttributes);
+
+			for ( int32 i=0; i < NumAttributes; ++i )
+			{
+				const FIntVector AttributeTileIndex(i % ProxyData->NumTiles.X, (i / ProxyData->NumTiles.X) % ProxyData->NumTiles.Y, i / (ProxyData->NumTiles.X * ProxyData->NumTiles.Y));
+
+				AttributePixelOffset.Emplace(
+					AttributeTileIndex.X * ProxyData->NumCells.X,
+					AttributeTileIndex.Y * ProxyData->NumCells.Y,
+					AttributeTileIndex.Z * ProxyData->NumCells.Z,
+					0
+				);
+				AttributeUVOffset.Emplace(
+					(1.0f / ProxyData->NumTiles.X) * float(AttributeTileIndex.X),
+					(1.0f / ProxyData->NumTiles.Y) * float(AttributeTileIndex.Y),
+					(1.0f / ProxyData->NumTiles.Z) * float(AttributeTileIndex.Z),
+					0.0f
+				);
+				AttributeUVClampMin.Emplace(
+					(float(AttributeTileIndex.X) * ProxyData->NumCells.X + 0.5) / (ProxyData->NumTiles.X * ProxyData->NumCells.X),
+					(float(AttributeTileIndex.Y) * ProxyData->NumCells.Y + 0.5) / (ProxyData->NumTiles.Y * ProxyData->NumCells.Y),
+					(float(AttributeTileIndex.Z) * ProxyData->NumCells.Z + 0.5) / (ProxyData->NumTiles.Z * ProxyData->NumCells.Z),
+					0.0f
+				);
+				AttributeUVClampMax.Emplace(
+					(float(AttributeTileIndex.X + 1) * ProxyData->NumCells.X - 0.5) / (ProxyData->NumTiles.X * ProxyData->NumCells.X),
+					(float(AttributeTileIndex.Y + 1) * ProxyData->NumCells.Y - 0.5) / (ProxyData->NumTiles.Y * ProxyData->NumCells.Y),
+					(float(AttributeTileIndex.Z + 1) * ProxyData->NumCells.Z - 0.5) / (ProxyData->NumTiles.Z * ProxyData->NumCells.Z),
+					0.0f
+				);
+			}
+
+			SetShaderValueArray(RHICmdList, ComputeShaderRHI, AttributePixelOffsetParam, AttributePixelOffset.GetData(), AttributePixelOffset.Num());
+			SetShaderValueArray(RHICmdList, ComputeShaderRHI, AttributeUVOffsetParam, AttributeUVOffset.GetData(), AttributeUVOffset.Num());
+			SetShaderValueArray(RHICmdList, ComputeShaderRHI, AttributeUVClampMinParam, AttributeUVClampMin.GetData(), AttributeUVClampMin.Num());
+			SetShaderValueArray(RHICmdList, ComputeShaderRHI, AttributeUVClampMaxParam, AttributeUVClampMax.GetData(), AttributeUVClampMax.Num());
+		}
+
 		FRHISamplerState *SamplerState = TStaticSamplerState<SF_Bilinear, AM_Clamp, AM_Clamp, AM_Clamp>::GetRHI();
 		SetSamplerParameter(RHICmdList, ComputeShaderRHI, SamplerParam, SamplerState);
 
@@ -285,12 +334,17 @@ private:
 	LAYOUT_FIELD(FShaderParameter, UnitToUVParam);
 	LAYOUT_FIELD(FShaderParameter, NumCellsParam);
 	LAYOUT_FIELD(FShaderParameter, NumTilesParam);
+	LAYOUT_FIELD(FShaderParameter, OneOverNumTilesParam);
 	LAYOUT_FIELD(FShaderParameter, CellSizeParam);
 	LAYOUT_FIELD(FShaderParameter, WorldBBoxSizeParam);
 
 	LAYOUT_FIELD(FShaderResourceParameter, GridParam);
 	LAYOUT_FIELD(FRWShaderParameter, OutputGridParam);
 	LAYOUT_FIELD(FShaderParameter, AttributeIndicesParam);
+	LAYOUT_FIELD(FShaderParameter, AttributePixelOffsetParam);
+	LAYOUT_FIELD(FShaderParameter, AttributeUVOffsetParam);
+	LAYOUT_FIELD(FShaderParameter, AttributeUVClampMinParam);
+	LAYOUT_FIELD(FShaderParameter, AttributeUVClampMaxParam);
 
 	LAYOUT_FIELD(FShaderResourceParameter, SamplerParam);
 	LAYOUT_FIELD(TMemoryImageArray<FName>, AttributeNames);
@@ -1075,8 +1129,13 @@ void UNiagaraDataInterfaceGrid3DCollection::GetParameterDefinitionHLSL(const FNi
 		Texture3D<float> {GridName};
 		RWTexture3D<float> RW{OutputGridName};
 		int3 {NumTiles};
+		float3 {OneOverNumTiles};
 		SamplerState {SamplerName};
 		int4 {AttributeIndicesName}[{AttributeInt4Count}];
+		int4 {AttributePixelOffsetName}[{AttributeCount}];
+		float4 {AttributeUVOffsetName}[{AttributeCount}];
+		float4 {AttributeUVClampMinName}[{AttributeCount}];
+		float4 {AttributeUVClampMaxName}[{AttributeCount}];
 		int {NumAttributesName};
 	)");
 
@@ -1084,21 +1143,27 @@ void UNiagaraDataInterfaceGrid3DCollection::GetParameterDefinitionHLSL(const FNi
 	// and leave 3 elements unused. Besides being wasteful, this means that the array we send to the CS would need to be padded,
 	// which is a hassle. Instead, use int4 explicitly, and access individual components in the generated code.
 	// Note that we have to have at least one here because hlsl doesn't support arrays of size 0.
-	const int AttributeInt4Count = FMath::Max(1, FMath::DivideAndRoundUp(ParamInfo.GeneratedFunctions.Num(), 4));
+	const int AttributeCount = ParamInfo.GeneratedFunctions.Num();
+	const int AttributeInt4Count = FMath::Max(1, FMath::DivideAndRoundUp(AttributeCount, 4));
 
-	TMap<FString, FStringFormatArg> ArgsDeclarations = {
+	TMap<FString, FStringFormatArg> ArgsDeclarations = {				
 		{ TEXT("GridName"),    GridName + ParamInfo.DataInterfaceHLSLSymbol },
 		{ TEXT("SamplerName"),    SamplerName + ParamInfo.DataInterfaceHLSLSymbol },
 		{ TEXT("OutputGridName"),    OutputGridName + ParamInfo.DataInterfaceHLSLSymbol },
 		{ TEXT("NumTiles"),    NumTilesName + ParamInfo.DataInterfaceHLSLSymbol },
+		{ TEXT("OneOverNumTiles"), OneOverNumTilesName + ParamInfo.DataInterfaceHLSLSymbol },
 		{ TEXT("UnitToUVName"), UNiagaraDataInterfaceRWBase::UnitToUVName + ParamInfo.DataInterfaceHLSLSymbol},
 		{ TEXT("AttributeIndicesName"), AttributeIndicesBaseName + ParamInfo.DataInterfaceHLSLSymbol},
+		{ TEXT("AttributePixelOffsetName"), AttributePixelOffsetName + ParamInfo.DataInterfaceHLSLSymbol},
+		{ TEXT("AttributeUVOffsetName"), AttributeUVOffsetName + ParamInfo.DataInterfaceHLSLSymbol},
+		{ TEXT("AttributeUVClampMinName"), AttributeUVClampMinName + ParamInfo.DataInterfaceHLSLSymbol},
+		{ TEXT("AttributeUVClampMaxName"), AttributeUVClampMaxName + ParamInfo.DataInterfaceHLSLSymbol},
 		{ TEXT("AttributeInt4Count"), AttributeInt4Count},
+		{ TEXT("AttributeCount"), AttributeCount },
 		{ TEXT("NumAttributesName"), UNiagaraDataInterfaceRWBase::NumAttributesName + ParamInfo.DataInterfaceHLSLSymbol},
 	};
 	OutHLSL += FString::Format(FormatDeclarations, ArgsDeclarations);
 }
-
 void UNiagaraDataInterfaceGrid3DCollection::WriteSetHLSL(const FNiagaraDataInterfaceGPUParamInfo& ParamInfo, const FNiagaraDataInterfaceGeneratedFunction& FunctionInfo, int FunctionInstanceIndex, int32 InNumChannels, FString& OutHLSL)
 {
 
@@ -1123,13 +1188,10 @@ void UNiagaraDataInterfaceGrid3DCollection::WriteSetHLSL(const FNiagaraDataInter
 		)");
 	}
 
-	FormatBounds += TEXT(R"(
+	FormatBounds += TEXT(R"(	
 					int CurAttributeIndex = In_AttributeIndex + i;
-					int TileIndexX = CurAttributeIndex % {NumTiles}.x;
-					int TileIndexY = (CurAttributeIndex / {NumTiles}.x) % {NumTiles}.y;
-					int TileIndexZ = CurAttributeIndex / ({NumTiles}.x * {NumTiles}.y);
-
-					RW{OutputGrid}[int3(In_IndexX + TileIndexX * {NumCellsName}.x, In_IndexY + TileIndexY * {NumCellsName}.y, In_IndexZ + TileIndexZ * {NumCellsName}.z)] = Val;
+					int3 TileOffset = {AttributePixelOffsetName}[CurAttributeIndex].xyz;
+					RW{OutputGrid}[int3(In_IndexX + TileOffset.x, In_IndexY + TileOffset.y, In_IndexZ + TileOffset.z)] = Val;
 				}
 			}
 		)");
@@ -1138,6 +1200,7 @@ void UNiagaraDataInterfaceGrid3DCollection::WriteSetHLSL(const FNiagaraDataInter
 		{TEXT("OutputGrid"), OutputGridName + ParamInfo.DataInterfaceHLSLSymbol},
 		{TEXT("NumCellsName"), UNiagaraDataInterfaceRWBase::NumCellsName + ParamInfo.DataInterfaceHLSLSymbol},
 		{TEXT("UnitToUVName"), UNiagaraDataInterfaceRWBase::UnitToUVName + ParamInfo.DataInterfaceHLSLSymbol},
+		{TEXT("AttributePixelOffsetName"), AttributePixelOffsetName + ParamInfo.DataInterfaceHLSLSymbol},
 		{TEXT("AttributeIndicesName"), AttributeIndicesBaseName + ParamInfo.DataInterfaceHLSLSymbol},
 		{TEXT("AttributeIndexGroup"), FunctionInstanceIndex / 4},
 		{TEXT("AttributeIndexComponent"), VectorComponentNames[FunctionInstanceIndex % 4]},
@@ -1161,13 +1224,9 @@ void UNiagaraDataInterfaceGrid3DCollection::WriteGetHLSL(const FNiagaraDataInter
 
 			    for (int i = 0; i < {NumChannels}; i++)
 				{
-
 					int CurAttributeIndex = In_AttributeIndex + i;
-					int TileIndexX = CurAttributeIndex % {NumTiles}.x;
-					int TileIndexY = (CurAttributeIndex / {NumTiles}.x) % {NumTiles}.y;
-					int TileIndexZ = CurAttributeIndex / ({NumTiles}.x * {NumTiles}.y);
-
-					float Val = {Grid}.Load(int4(In_IndexX + TileIndexX * {NumCellsName}.x, In_IndexY + TileIndexY * {NumCellsName}.y, In_IndexZ + TileIndexZ * {NumCellsName}.z, 0));
+					int3 TileOffset = {AttributePixelOffsetName}[CurAttributeIndex].xyz;
+					float Val = {Grid}.Load(int4(In_IndexX + TileOffset.x, In_IndexY + TileOffset.y, In_IndexZ + TileOffset.z, 0));
 					)");
 	if (InNumChannels == 1)
 	{
@@ -1187,6 +1246,7 @@ void UNiagaraDataInterfaceGrid3DCollection::WriteGetHLSL(const FNiagaraDataInter
 		{TEXT("Grid"), GridName + ParamInfo.DataInterfaceHLSLSymbol},
 		{TEXT("NumCellsName"), UNiagaraDataInterfaceRWBase::NumCellsName + ParamInfo.DataInterfaceHLSLSymbol},
 		{TEXT("UnitToUVName"), UNiagaraDataInterfaceRWBase::UnitToUVName + ParamInfo.DataInterfaceHLSLSymbol},
+		{TEXT("AttributePixelOffsetName"), AttributePixelOffsetName + ParamInfo.DataInterfaceHLSLSymbol},
 		{TEXT("AttributeIndicesName"), AttributeIndicesBaseName + ParamInfo.DataInterfaceHLSLSymbol},
 		{TEXT("AttributeIndexGroup"), FunctionInstanceIndex / 4},
 		{TEXT("AttributeIndexComponent"), VectorComponentNames[FunctionInstanceIndex % 4]},
@@ -1208,29 +1268,8 @@ void UNiagaraDataInterfaceGrid3DCollection::WriteSampleHLSL(const FNiagaraDataIn
 			    for (int i = 0; i < {NumChannels}; i++)
 				{
 					int CurAttributeIndex = In_AttributeIndex + i;
-					int TileIndexX = CurAttributeIndex % {NumTiles}.x;
-					int TileIndexY = (CurAttributeIndex / {NumTiles}.x) % {NumTiles}.y;
-					int TileIndexZ = CurAttributeIndex / ({NumTiles}.x * {NumTiles}.y);
-
-					float3 UVW =
-					{
-						In_UnitX / {NumTiles}.x + 1.0*TileIndexX/{NumTiles}.x,
-						In_UnitY / {NumTiles}.y + 1.0*TileIndexY/{NumTiles}.y,
-						In_UnitZ / {NumTiles}.z + 1.0*TileIndexZ/{NumTiles}.y
-					};
-					float3 TileMin =
-					{
-						(TileIndexX * {NumCellsName}.x + 0.5) / ({NumTiles}.x * {NumCellsName}.x),
-						(TileIndexY * {NumCellsName}.y + 0.5) / ({NumTiles}.y * {NumCellsName}.y),
-						(TileIndexZ * {NumCellsName}.z + 0.5) / ({NumTiles}.z * {NumCellsName}.z)
-					};
-					float3 TileMax =
-					{
-						((TileIndexX + 1) * {NumCellsName}.x - 0.5) / ({NumTiles}.x * {NumCellsName}.x),
-						((TileIndexY + 1) * {NumCellsName}.y - 0.5) / ({NumTiles}.y * {NumCellsName}.y),
-						((TileIndexZ + 1) * {NumCellsName}.z - 0.5) / ({NumTiles}.z * {NumCellsName}.z)
-					};
-					UVW = clamp(UVW, TileMin, TileMax);
+					float3 UVW = {AttributeUVOffsetName}[CurAttributeIndex].xyz + (In_Unit * {OneOverNumTiles});
+					UVW = clamp(UVW, {AttributeUVClampMinName}[CurAttributeIndex].xyz, {AttributeUVClampMaxName}[CurAttributeIndex].xyz);
 
 					float Val = {Grid}.SampleLevel({SamplerName}, UVW, 0);
 					)");
@@ -1257,9 +1296,14 @@ void UNiagaraDataInterfaceGrid3DCollection::WriteSampleHLSL(const FNiagaraDataIn
 		{TEXT("NumChannels"), FString::FromInt(InNumChannels)},
 		{TEXT("NumChannelsVariableSuffix"), InNumChannels > 1 ? FString::FromInt(InNumChannels) : TEXT("")},
 		{TEXT("AttributeIndicesName"), AttributeIndicesBaseName + ParamInfo.DataInterfaceHLSLSymbol},
+		{TEXT("AttributePixelOffsetName"), AttributePixelOffsetName + ParamInfo.DataInterfaceHLSLSymbol},
+		{TEXT("AttributeUVOffsetName"), AttributeUVOffsetName + ParamInfo.DataInterfaceHLSLSymbol},
+		{TEXT("AttributeUVClampMinName"), AttributeUVClampMinName + ParamInfo.DataInterfaceHLSLSymbol},
+		{TEXT("AttributeUVClampMaxName"), AttributeUVClampMaxName + ParamInfo.DataInterfaceHLSLSymbol},
 		{TEXT("AttributeIndexGroup"), FunctionInstanceIndex / 4},
 		{TEXT("AttributeIndexComponent"), VectorComponentNames[FunctionInstanceIndex % 4]},
 		{TEXT("NumTiles"), NumTilesName + ParamInfo.DataInterfaceHLSLSymbol},
+		{TEXT("OneOverNumTiles"), OneOverNumTilesName + ParamInfo.DataInterfaceHLSLSymbol },
 	};
 	OutHLSL += FString::Format(*FormatBounds, ArgsBounds);
 }
@@ -1339,8 +1383,12 @@ bool UNiagaraDataInterfaceGrid3DCollection::GetFunctionHLSL(const FNiagaraDataIn
 		{TEXT("UnitToUVName"), UNiagaraDataInterfaceRWBase::UnitToUVName + ParamInfo.DataInterfaceHLSLSymbol},
 		{TEXT("SamplerName"), SamplerName + ParamInfo.DataInterfaceHLSLSymbol},
 		{TEXT("NumTiles"), NumTilesName + ParamInfo.DataInterfaceHLSLSymbol},
-		{TEXT("UnitToUVName"), UNiagaraDataInterfaceRWBase::UnitToUVName + ParamInfo.DataInterfaceHLSLSymbol},
+		{TEXT("OneOverNumTiles"), OneOverNumTilesName + ParamInfo.DataInterfaceHLSLSymbol },
 		{TEXT("NumCellsName"), UNiagaraDataInterfaceRWBase::NumCellsName + ParamInfo.DataInterfaceHLSLSymbol},
+		{TEXT("AttributePixelOffsetName"), AttributePixelOffsetName + ParamInfo.DataInterfaceHLSLSymbol},
+		{TEXT("AttributeUVOffsetName"), AttributeUVOffsetName + ParamInfo.DataInterfaceHLSLSymbol},
+		{TEXT("AttributeUVClampMinName"), AttributeUVClampMinName + ParamInfo.DataInterfaceHLSLSymbol},
+		{TEXT("AttributeUVClampMaxName"), AttributeUVClampMaxName + ParamInfo.DataInterfaceHLSLSymbol},
 	};
 
 	if (FunctionInfo.DefinitionName == GetValueFunctionName || FunctionInfo.DefinitionName == GetPreviousValueAtIndexFunctionName)
@@ -1348,11 +1396,9 @@ bool UNiagaraDataInterfaceGrid3DCollection::GetFunctionHLSL(const FNiagaraDataIn
 		static const TCHAR *FormatBounds = TEXT(R"(
 			void {FunctionName}(int In_IndexX, int In_IndexY, int In_IndexZ, int In_AttributeIndex, out float Out_Val)
 			{
-				int TileIndexX = In_AttributeIndex % {NumTiles}.x;
-				int TileIndexY = (In_AttributeIndex / {NumTiles}.x) % {NumTiles}.y;
-				int TileIndexZ = In_AttributeIndex / ({NumTiles}.x * {NumTiles}.y);
+				int3 TileOffset = {AttributePixelOffsetName}[In_AttributeIndex].xyz;
 
-				Out_Val = {Grid}.Load(int4(In_IndexX + TileIndexX * {NumCellsName}.x, In_IndexY + TileIndexY * {NumCellsName}.y, In_IndexZ + TileIndexZ * {NumCellsName}.z, 0));
+				Out_Val = {Grid}.Load(int4(In_IndexX + TileOffset.x, In_IndexY + TileOffset.y, In_IndexZ + TileOffset.z, 0));
 			}
 		)");
 		OutHLSL += FString::Format(FormatBounds, ArgsBounds);
@@ -1373,13 +1419,11 @@ bool UNiagaraDataInterfaceGrid3DCollection::GetFunctionHLSL(const FNiagaraDataIn
 	{
 		static const TCHAR *FormatBounds = TEXT(R"(
 			void {FunctionName}(int In_IndexX, int In_IndexY, int In_IndexZ, int In_AttributeIndex, float In_Value, out int val)
-			{
-				int TileIndexX = In_AttributeIndex % {NumTiles}.x;
-				int TileIndexY = (In_AttributeIndex / {NumTiles}.x) % {NumTiles}.y;
-				int TileIndexZ = In_AttributeIndex / ({NumTiles}.x * {NumTiles}.y);
+			{			
+				int3 TileOffset = {AttributePixelOffsetName}[In_AttributeIndex].xyz;
 
 				val = 0;
-				RW{OutputGrid}[int3(In_IndexX + TileIndexX * {NumCellsName}.x, In_IndexY + TileIndexY * {NumCellsName}.y, In_IndexZ + TileIndexZ * {NumCellsName}.z)] = In_Value;
+				RW{OutputGrid}[int3(In_IndexX + TileOffset.x, In_IndexY + TileOffset.y, In_IndexZ + TileOffset.z)] = In_Value;
 			}
 		)");
 		OutHLSL += FString::Format(FormatBounds, ArgsBounds);
@@ -1403,15 +1447,11 @@ bool UNiagaraDataInterfaceGrid3DCollection::GetFunctionHLSL(const FNiagaraDataIn
 			void {FunctionName}(int In_IndexX, int In_IndexY, int In_IndexZ)
 			{
 				for (int AttributeIndex = 0; AttributeIndex < {NumAttributes}.x; AttributeIndex++)
-				{
-					int TileIndexX = AttributeIndex % {NumTiles}.x;
-					int TileIndexY = (AttributeIndex / {NumTiles}.x) % {NumTiles}.y;
-					int TileIndexZ = AttributeIndex / ({NumTiles}.x * {NumTiles}.y);
-
-					float Val = {Grid}.Load(int4(In_IndexX + TileIndexX * {NumCellsName}.x, In_IndexY + TileIndexY * {NumCellsName}.y, In_IndexZ + TileIndexZ * {NumCellsName}.z, 0));
-					//float Val = {Grid}.Load(int4(In_IndexX, In_IndexY, AttributeIndex, 0));
-					//RW{OutputGrid}[int3(In_IndexX, In_IndexY, In_IndexZ, AttributeIndex)] = Val;
-					RW{OutputGrid}[int3(In_IndexX + TileIndexX * {NumCellsName}.x, In_IndexY + TileIndexY * {NumCellsName}.y, In_IndexZ + TileIndexZ * {NumCellsName}.z)] = Val;
+				{		
+					int3 TileOffset = {AttributePixelOffsetName}[AttributeIndex].xyz;
+
+					float Val = {Grid}.Load(int4(In_IndexX + TileOffset.x, In_IndexY + TileOffset.y, In_IndexZ + TileOffset.z, 0));	
+					RW{OutputGrid}[int3(In_IndexX + TileOffset.x, In_IndexY + TileOffset.y, In_IndexZ + TileOffset.z)] = Val;
 				}
 			}
 		)");
@@ -1427,11 +1467,9 @@ bool UNiagaraDataInterfaceGrid3DCollection::GetFunctionHLSL(const FNiagaraDataIn
 				{
 					float Val = 0.0f;
 
-					int TileIndexX = AttributeIndex % {NumTiles}.x;
-					int TileIndexY = (AttributeIndex / {NumTiles}.x) % {NumTiles}.y;
-					int TileIndexZ = AttributeIndex / ({NumTiles}.x * {NumTiles}.y);
+					int3 TileOffset = {AttributePixelOffsetName}[AttributeIndex].xyz;
 
-					RW{OutputGrid}[int3(In_IndexX + TileIndexX * {NumCellsName}.x, In_IndexY + TileIndexY * {NumCellsName}.y, In_IndexZ + TileIndexZ * {NumCellsName}.z)] = In_Value;
+					RW{OutputGrid}[int3(In_IndexX + TileOffset.x, In_IndexY + TileOffset.y, In_IndexZ + TileOffset.z)] = In_Value;
 				}
 			}
 		)");
@@ -1523,13 +1561,8 @@ bool UNiagaraDataInterfaceGrid3DCollection::GetFunctionHLSL(const FNiagaraDataIn
 		static const TCHAR* FormatBounds = TEXT(R"(
 				void {FunctionName}(float In_UnitX, float In_UnitY, float In_UnitZ, int In_AttributeIndex, out float Out_Val)
 				{
-					int TileIndexX = In_AttributeIndex % {NumTiles}.x;
-					int TileIndexY = (In_AttributeIndex / {NumTiles}.x) % {NumTiles}.y;
-					int TileIndexZ = In_AttributeIndex / ({NumTiles}.x * {NumTiles}.y);		
-					
-					float3 UVW = clamp(float3(In_UnitX, In_UnitY, In_UnitZ), .5 / {NumCellsName}, 1. - .5 / {NumCellsName});
-					UVW = (UVW + float3(TileIndexX, TileIndexY, TileIndexZ)) / {NumTiles};
-
+					float3 UVW = {AttributeUVOffsetName}[In_AttributeIndex].xyz + (float3(In_UnitX, In_UnitY, In_UnitZ) * {OneOverNumTiles});
+					UVW = clamp(UVW, {AttributeUVClampMinName}[In_AttributeIndex].xyz, {AttributeUVClampMaxName}[In_AttributeIndex].xyz);
 					Out_Val = {Grid}.SampleLevel({SamplerName}, UVW, 0);
 				}
 			)");
diff --git a/Engine/Plugins/FX/Niagara/Source/NiagaraCore/Private/NiagaraCustomVersion.cpp b/Engine/Plugins/FX/Niagara/Source/NiagaraCore/Private/NiagaraCustomVersion.cpp
--- a/Engine/Plugins/FX/Niagara/Source/NiagaraCore/Private/NiagaraCustomVersion.cpp
+++ b/Engine/Plugins/FX/Niagara/Source/NiagaraCore/Private/NiagaraCustomVersion.cpp
@@ -8,4 +8,4 @@ const FGuid FNiagaraCustomVersion::GUID(0xFCF57AFA, 0x50764283, 0xB9A9E658, 0xFF
 // Register the custom version with core
 FCustomVersionRegistration GRegisterNiagaraCustomVersion(FNiagaraCustomVersion::GUID, FNiagaraCustomVersion::LatestVersion, TEXT("NiagaraVer"));
 
-const FGuid FNiagaraCustomVersion::LatestScriptCompileVersion(0x4E77ECF8, 0xD5584919, 0xB77BC6E7, 0x5F47CF41);
+const FGuid FNiagaraCustomVersion::LatestScriptCompileVersion(0x0FC7D74E, 0x22166A42, 0xBFA2127F, 0xC334B412);
