Fix for scratch pad scripts not necessarily being postloaded before they were accessed

diff --git a/Engine/Plugins/FX/Niagara/Source/Niagara/Private/NiagaraEmitter.cpp b/Engine/Plugins/FX/Niagara/Source/Niagara/Private/NiagaraEmitter.cpp
--- a/Engine/Plugins/FX/Niagara/Source/Niagara/Private/NiagaraEmitter.cpp
+++ b/Engine/Plugins/FX/Niagara/Source/Niagara/Private/NiagaraEmitter.cpp
@@ -499,6 +499,23 @@ void UNiagaraEmitter::PostLoad()
 		{
 			ParentAtLastMerge->ConditionalPostLoad();
 		}
+
+		for (auto ScratchPadScript : ScratchPadScripts)
+		{
+			if (ScratchPadScript)
+			{
+				ScratchPadScript->ConditionalPostLoad();
+			}
+		}
+
+		for (auto ParentScratchPadScript : ParentScratchPadScripts)
+		{
+			if (ParentScratchPadScript)
+			{
+				ParentScratchPadScript->ConditionalPostLoad();
+			}
+		}
+
 		if (IsSynchronizedWithParent() == false)
 		{
 			// Modify here so that the asset will be marked dirty when using the resave commandlet.  This will be ignored during regular post load.
