Fix for in-details panel curve toolbar styling

diff --git a/Engine/Plugins/FX/Niagara/Source/NiagaraEditorWidgets/Private/DetailCustomizations/NiagaraDataInterfaceCurveDetails.cpp b/Engine/Plugins/FX/Niagara/Source/NiagaraEditorWidgets/Private/DetailCustomizations/NiagaraDataInterfaceCurveDetails.cpp
--- a/Engine/Plugins/FX/Niagara/Source/NiagaraEditorWidgets/Private/DetailCustomizations/NiagaraDataInterfaceCurveDetails.cpp
+++ b/Engine/Plugins/FX/Niagara/Source/NiagaraEditorWidgets/Private/DetailCustomizations/NiagaraDataInterfaceCurveDetails.cpp
@@ -51,6 +51,7 @@
 #include "SColorGradientEditor.h"
 
 #include "EditorFontGlyphs.h"
+#include "Widgets/Input/SSegmentedControl.h"
 
 #define LOCTEXT_NAMESPACE "NiagaraDataInterfaceCurveDetails"
 
@@ -172,14 +173,14 @@ public:
 			.Padding(0, 0, 4, 0)
 			[
 				SNew(SButton)
-				.ButtonStyle(FEditorStyle::Get(), "FlatButton.Dark")
+				.ButtonStyle(FAppStyle::Get(), "SimpleButton")
 				.OnClicked(this, &SNiagaraDataInterfaceCurveKeySelector::ZoomToFitClicked)
 				.ToolTipText(LOCTEXT("ZoomToFitToolTip", "Zoom to fit all keys"))
 				.Content()
 				[
 					SNew(STextBlock)
 					.Font(FEditorStyle::Get().GetFontStyle("FontAwesome.10"))
-					.ColorAndOpacity(FNiagaraEditorWidgetsStyle::Get().GetColor("NiagaraEditor.CurveDetails.TextButtonForeground"))
+					.ColorAndOpacity(FSlateColor::UseForeground())
 					.Text(FEditorFontGlyphs::Expand)
 				]
 			]
@@ -190,16 +191,15 @@ public:
 			.Padding(0, 0, 1, 0)
 			[
 				SNew(SButton)
-				.ButtonStyle(FEditorStyle::Get(), "FlatButton.Dark")
+				.ButtonStyle(FAppStyle::Get(), "SimpleButton")
 				.OnClicked(this, &SNiagaraDataInterfaceCurveKeySelector::PreviousCurveClicked)
 				.ToolTipText(LOCTEXT("PreviousCurveToolTip", "Select the previous curve"))
 				.Visibility(this, &SNiagaraDataInterfaceCurveKeySelector::GetNextPreviousCurveButtonVisibility)
 				.Content()
 				[
-					SNew(STextBlock)
-					.Font(FEditorStyle::Get().GetFontStyle("FontAwesome.10"))
-					.ColorAndOpacity(FNiagaraEditorWidgetsStyle::Get().GetColor("NiagaraEditor.CurveDetails.TextButtonForeground"))
-					.Text(FEditorFontGlyphs::Arrow_Up)
+					SNew(SImage)
+					.ColorAndOpacity(FSlateColor::UseForeground())
+					.Image(FAppStyle::Get().GetBrush("Icons.ArrowUp"))
 				]
 			]
 
@@ -209,16 +209,15 @@ public:
 			.Padding(0, 0, 4, 0)
 			[
 				SNew(SButton)
-				.ButtonStyle(FEditorStyle::Get(), "FlatButton.Dark")
+				.ButtonStyle(FAppStyle::Get(), "SimpleButton")
 				.OnClicked(this, &SNiagaraDataInterfaceCurveKeySelector::NextCurveClicked)
 				.ToolTipText(LOCTEXT("NextCurveToolTip", "Select the next curve"))
 				.Visibility(this, &SNiagaraDataInterfaceCurveKeySelector::GetNextPreviousCurveButtonVisibility)
 				.Content()
 				[
-					SNew(STextBlock)
-					.Font(FEditorStyle::Get().GetFontStyle("FontAwesome.10"))
-					.ColorAndOpacity(FNiagaraEditorWidgetsStyle::Get().GetColor("NiagaraEditor.CurveDetails.TextButtonForeground"))
-					.Text(FEditorFontGlyphs::Arrow_Down)
+					SNew(SImage)
+					.Image(FAppStyle::Get().GetBrush("Icons.ArrowDown"))
+					.ColorAndOpacity(FSlateColor::UseForeground())
 				]
 			]
 
@@ -228,15 +227,14 @@ public:
 			.Padding(0, 0, 1, 0)
 			[
 				SNew(SButton)
-				.ButtonStyle(FEditorStyle::Get(), "FlatButton.Dark")
+				.ButtonStyle(FAppStyle::Get(), "SimpleButton")
 				.OnClicked(this, &SNiagaraDataInterfaceCurveKeySelector::PreviousKeyClicked)
 				.ToolTipText(LOCTEXT("PreviousKeyToolTip", "Select the previous key for the selected curve."))
 				.Content()
 				[
-					SNew(STextBlock)
-					.Font(FEditorStyle::Get().GetFontStyle("FontAwesome.10"))
-					.ColorAndOpacity(FNiagaraEditorWidgetsStyle::Get().GetColor("NiagaraEditor.CurveDetails.TextButtonForeground"))
-					.Text(FEditorFontGlyphs::Arrow_Left)
+					SNew(SImage)
+					.Image(FAppStyle::Get().GetBrush("Icons.ArrowLeft"))
+					.ColorAndOpacity(FSlateColor::UseForeground())
 				]
 			]
 
@@ -246,15 +244,14 @@ public:
 			.Padding(0, 0, 4, 0)
 			[
 				SNew(SButton)
-				.ButtonStyle(FEditorStyle::Get(), "FlatButton.Dark")
+				.ButtonStyle(FAppStyle::Get(), "SimpleButton")
 				.OnClicked(this, &SNiagaraDataInterfaceCurveKeySelector::NextKeyClicked)
 				.ToolTipText(LOCTEXT("NextKeyToolTip", "Select the next key for the selected curve."))
 				.Content()
 				[
-					SNew(STextBlock)
-					.Font(FEditorStyle::Get().GetFontStyle("FontAwesome.10"))
-					.ColorAndOpacity(FNiagaraEditorWidgetsStyle::Get().GetColor("NiagaraEditor.CurveDetails.TextButtonForeground"))
-					.Text(FEditorFontGlyphs::Arrow_Right)
+					SNew(SImage)
+					.Image(FAppStyle::Get().GetBrush("Icons.ArrowRight"))
+					.ColorAndOpacity(FSlateColor::UseForeground())
 				]
 			]
 
@@ -264,15 +261,14 @@ public:
 			.Padding(0, 0, 1, 0)
 			[
 				SNew(SButton)
-				.ButtonStyle(FEditorStyle::Get(), "FlatButton.Dark")
+				.ButtonStyle(FAppStyle::Get(), "SimpleButton")
 				.OnClicked(this, &SNiagaraDataInterfaceCurveKeySelector::AddKeyClicked)
 				.ToolTipText(LOCTEXT("AddKeyToolTip", "Add a key to the selected curve."))
 				.Content()
 				[
-					SNew(STextBlock)
-					.Font(FEditorStyle::Get().GetFontStyle("FontAwesome.10"))
-					.ColorAndOpacity(FNiagaraEditorWidgetsStyle::Get().GetColor("NiagaraEditor.CurveDetails.TextButtonForeground"))
-					.Text(FEditorFontGlyphs::Plus)
+					SNew(SImage)
+					.Image(FAppStyle::Get().GetBrush("Icons.PlusCircle"))
+					.ColorAndOpacity(FSlateColor::UseForeground())
 				]
 			]
 
@@ -281,15 +277,14 @@ public:
 			.AutoWidth()
 			[
 				SNew(SButton)
-				.ButtonStyle(FEditorStyle::Get(), "FlatButton.Dark")
+				.ButtonStyle(FAppStyle::Get(), "SimpleButton")
 				.OnClicked(this, &SNiagaraDataInterfaceCurveKeySelector::DeleteKeyClicked)
 				.ToolTipText(LOCTEXT("DeleteKeyToolTip", "Delete the currently selected keys."))
 				.Content()
 				[
-					SNew(STextBlock)
-					.Font(FEditorStyle::Get().GetFontStyle("FontAwesome.10"))
-					.ColorAndOpacity(FNiagaraEditorWidgetsStyle::Get().GetColor("NiagaraEditor.CurveDetails.TextButtonForeground"))
-					.Text(FEditorFontGlyphs::Trash)
+					SNew(SImage)
+					.Image(FAppStyle::Get().GetBrush("Icons.Delete"))
+					.ColorAndOpacity(FSlateColor::UseForeground())
 				]
 			]
 		];
@@ -775,7 +770,7 @@ class SNiagaraCurveThumbnail : public SLeafWidget
 
 	virtual int32 OnPaint(const FPaintArgs& Args, const FGeometry& AllottedGeometry, const FSlateRect& MyCullingRect, FSlateWindowElementList& OutDrawElements, int32 LayerId, const FWidgetStyle& InWidgetStyle, bool bParentEnabled) const override
 	{
-		FSlateDrawElement::MakeLines(OutDrawElements, LayerId, AllottedGeometry.ToPaintGeometry(), CurvePoints, ESlateDrawEffect::None, FLinearColor::White, true, 2.0f);
+		FSlateDrawElement::MakeLines(OutDrawElements, LayerId, AllottedGeometry.ToPaintGeometry(), CurvePoints, ESlateDrawEffect::None, InWidgetStyle.GetForegroundColor(), true, 2.0f);
 		return LayerId;
 	}
 
@@ -823,7 +818,7 @@ class SNiagaraCurveTemplateBar : public SCompoundWidget
 
 				ToolBarBuilder.AddWidget(
 					SNew(SButton)
-					.ButtonStyle(FEditorStyle::Get(), "FlatButton.Dark")
+					.ButtonStyle(FAppStyle::Get(), "SimpleButton")
 					.OnClicked(this, &SNiagaraCurveTemplateBar::CurveTemplateClicked, TWeakObjectPtr<UCurveFloat>(FloatCurveAsset))
 					.ToolTipText(FText::Format(LOCTEXT("ApplyCurveTemplateFormat", "{0}\nClick to apply this template to the selected curves."), CurveDisplayName))
 					.ContentPadding(FMargin(3))
@@ -1032,6 +1027,7 @@ public:
 				]
 				+ SHorizontalBox::Slot()
 				.VAlign(VAlign_Center)
+				.FillWidth(1.0f)
 				[
 					CurveEditorPanel->GetKeyDetailsView().ToSharedRef()
 				]
@@ -1404,7 +1400,7 @@ void FNiagaraDataInterfaceCurveDetailsBase::CustomizeDetails(IDetailLayoutBuilde
 	}
 
 	FToolBarBuilder ToolBarBuilder(CurveEditor->GetCommands(), FMultiBoxCustomization::None, nullptr, true);
-	ToolBarBuilder.SetStyle(&FEditorStyle::Get(), "Sequencer.ToolBar");
+	ToolBarBuilder.SetStyle(&FAppStyle::Get(), "SlimToolbar");
 
 	ToolBarBuilder.AddComboButton(
 		FUIAction(), 
@@ -1433,42 +1429,15 @@ void FNiagaraDataInterfaceCurveDetailsBase::CustomizeDetails(IDetailLayoutBuilde
 			SNew(SBox)
 			.Padding(FMargin(0, 0, 2, 0))
 			[
-				SNew(SButton)
-				.ButtonStyle(FEditorStyle::Get(), "FlatButton.Light")
-				.OnClicked(this, &FNiagaraDataInterfaceCurveDetailsBase::SetGradientVisibility, true)
-				.ToolTipText(LOCTEXT("ShowGradientToolTip", "Show the gradient editor."))
-				.ButtonColorAndOpacity(this, &FNiagaraDataInterfaceCurveDetailsBase::GetGradientButtonColor)
-				.ContentPadding(FMargin(0))
-				.Content()
-				[
-					SNew(STextBlock)
-					.Justification(ETextJustify::Center)
-					.MinDesiredWidth(50)
-					.Font(FCoreStyle::GetDefaultFontStyle("Bold", 9))
-					.ColorAndOpacity(FNiagaraEditorWidgetsStyle::Get().GetColor("NiagaraEditor.CurveDetails.TextButtonForeground"))
-					.Text(LOCTEXT("GradientButtonLabel", "Gradient"))
-				]
-			]);
-
-		ToolBarBuilder.AddWidget(
-			SNew(SBox)
-			.Padding(FMargin(2, 0, 0, 0))
-			[
-				SNew(SButton)
-				.ButtonStyle(FEditorStyle::Get(), "FlatButton.Light")
-				.OnClicked(this, &FNiagaraDataInterfaceCurveDetailsBase::SetGradientVisibility, false)
-				.ToolTipText(LOCTEXT("ShowCurvesToolTip", "Show the curve editor."))
-				.ButtonColorAndOpacity(this, &FNiagaraDataInterfaceCurveDetailsBase::GetCurveButtonColor)
-				.ContentPadding(FMargin(0))
-				.Content()
-				[
-					SNew(STextBlock)
-					.Justification(ETextJustify::Center)
-					.MinDesiredWidth(50)
-					.Font(FCoreStyle::GetDefaultFontStyle("Bold", 9))
-					.ColorAndOpacity(FNiagaraEditorWidgetsStyle::Get().GetColor("NiagaraEditor.CurveDetails.TextButtonForeground"))
-					.Text(LOCTEXT("CurvesButtonLabel", "Curves"))
-				]
+				SNew(SSegmentedControl<bool>)
+				.OnValueChanged(this, &FNiagaraDataInterfaceCurveDetailsBase::SetGradientVisibility)
+				.Value(this, &FNiagaraDataInterfaceCurveDetailsBase::IsGradientVisible)
+				+ SSegmentedControl<bool>::Slot(false)
+				.Text(LOCTEXT("CurvesButtonLabel", "Curves"))
+				.ToolTip(LOCTEXT("ShowCurvesToolTip", "Show the curve editor."))
+				+ SSegmentedControl<bool>::Slot(true)
+				.Text(LOCTEXT("GradientButtonLabel", "Gradient"))
+				.ToolTip(LOCTEXT("ShowGradientToolTip", "Show the gradient editor."))
 			]);
 	}
 
@@ -1477,14 +1446,13 @@ void FNiagaraDataInterfaceCurveDetailsBase::CustomizeDetails(IDetailLayoutBuilde
 	CurveCategory.AddCustomRow(NSLOCTEXT("NiagaraDataInterfaceCurveDetails", "CurveFilterText", "Curve")).WholeRowContent() [ CurveDataInterfaceEditor.ToSharedRef() ];
 }
 
-FReply FNiagaraDataInterfaceCurveDetailsBase::SetGradientVisibility(bool bInShowGradient)
+void FNiagaraDataInterfaceCurveDetailsBase::SetGradientVisibility(bool bInShowGradient)
 {
 	TSharedPtr<FNiagaraStackCurveEditorOptions> CurveEditorOptions = StackCurveEditorOptionsWeak.Pin();
 	if (CurveEditorOptions.IsValid())
 	{
 		CurveEditorOptions->SetIsGradientVisible(bInShowGradient);
 	}
-	return FReply::Handled();
 }
 
 int32 FNiagaraDataInterfaceCurveDetailsBase::GetGradientCurvesSwitcherIndex() const
@@ -1492,15 +1460,12 @@ int32 FNiagaraDataInterfaceCurveDetailsBase::GetGradientCurvesSwitcherIndex() co
 	return StackCurveEditorOptionsWeak.IsValid() && StackCurveEditorOptionsWeak.Pin()->GetIsGradientVisible() ? 0 : 1;
 }
 
-FSlateColor FNiagaraDataInterfaceCurveDetailsBase::GetGradientButtonColor() const
+bool FNiagaraDataInterfaceCurveDetailsBase::IsGradientVisible() const
 {
-	return StackCurveEditorOptionsWeak.IsValid() && StackCurveEditorOptionsWeak.Pin()->GetIsGradientVisible() ? FLinearColor(0.4f, 0.4f, 0.4f, 1) : FLinearColor(0.15f, 0.15f, 0.15f, 1);
+	return StackCurveEditorOptionsWeak.IsValid() ? StackCurveEditorOptionsWeak.Pin()->GetIsGradientVisible() : false;
 }
 
-FSlateColor FNiagaraDataInterfaceCurveDetailsBase::GetCurveButtonColor() const
-{
-	return StackCurveEditorOptionsWeak.IsValid() && StackCurveEditorOptionsWeak.Pin()->GetIsGradientVisible() ? FLinearColor(0.15f, 0.15f, 0.15f, 1) : FLinearColor(0.4f, 0.4f, 0.4f, 1);
-}
+
 
 void FNiagaraDataInterfaceCurveDetailsBase::OnShowInCurveEditor() const
 {
diff --git a/Engine/Plugins/FX/Niagara/Source/NiagaraEditorWidgets/Private/DetailCustomizations/NiagaraDataInterfaceCurveDetails.h b/Engine/Plugins/FX/Niagara/Source/NiagaraEditorWidgets/Private/DetailCustomizations/NiagaraDataInterfaceCurveDetails.h
--- a/Engine/Plugins/FX/Niagara/Source/NiagaraEditorWidgets/Private/DetailCustomizations/NiagaraDataInterfaceCurveDetails.h
+++ b/Engine/Plugins/FX/Niagara/Source/NiagaraEditorWidgets/Private/DetailCustomizations/NiagaraDataInterfaceCurveDetails.h
@@ -34,10 +34,9 @@ protected:
 
 private:
 	virtual TSharedRef<SWidget> GetCurveToCopyMenu();
-	FReply SetGradientVisibility(bool bInShowGradient);
+	void SetGradientVisibility(bool bInShowGradient);
 	int32 GetGradientCurvesSwitcherIndex() const;
-	FSlateColor GetGradientButtonColor() const;
-	FSlateColor GetCurveButtonColor() const;
+	bool IsGradientVisible() const;
 	void OnShowInCurveEditor() const;
 	void CurveToCopySelected(const FAssetData& AssetData);
 	TWeakObjectPtr<UNiagaraDataInterfaceCurveBase> CurveDataInterfaceWeak;
