Changing stack overview styling to improve readability

diff --git a/Engine/Plugins/FX/Niagara/Source/NiagaraEditorWidgets/Private/NiagaraEditorWidgetsStyle.cpp b/Engine/Plugins/FX/Niagara/Source/NiagaraEditorWidgets/Private/NiagaraEditorWidgetsStyle.cpp
--- a/Engine/Plugins/FX/Niagara/Source/NiagaraEditorWidgets/Private/NiagaraEditorWidgetsStyle.cpp
+++ b/Engine/Plugins/FX/Niagara/Source/NiagaraEditorWidgets/Private/NiagaraEditorWidgetsStyle.cpp
@@ -92,6 +92,10 @@ TSharedRef< FSlateStyleSet > FNiagaraEditorWidgetsStyle::Create()
 
 	Style->Set("NiagaraEditor.Stack.ItemText", StackItemText);
 	
+	FTextBlockStyle OverviewStackItemText = FTextBlockStyle(CategoryText)
+		.SetFont(NormalFont);
+	Style->Set("NiagaraEditor.OverviewStack.ItemText", OverviewStackItemText);
+
 	FEditableTextBoxStyle StackEditableItemText = FEditableTextBoxStyle(NormalEditableTextBox)
 		.SetFont(NormalFont);
 
diff --git a/Engine/Plugins/FX/Niagara/Source/NiagaraEditorWidgets/Private/SNiagaraOverviewStack.cpp b/Engine/Plugins/FX/Niagara/Source/NiagaraEditorWidgets/Private/SNiagaraOverviewStack.cpp
--- a/Engine/Plugins/FX/Niagara/Source/NiagaraEditorWidgets/Private/SNiagaraOverviewStack.cpp
+++ b/Engine/Plugins/FX/Niagara/Source/NiagaraEditorWidgets/Private/SNiagaraOverviewStack.cpp
@@ -381,7 +381,7 @@ private:
 				.Padding(0.0f, 0.0f, 6.0f, 0.0f)
 				[
 					SNew(STextBlock)
-					.TextStyle(FNiagaraEditorWidgetsStyle::Get(), "NiagaraEditor.Stack.ItemText")
+					.TextStyle(FNiagaraEditorWidgetsStyle::Get(), "NiagaraEditor.OverviewStack.ItemText")
 					.Text(LOCTEXT("SetVariablesPrefix", "Set:"))
 					.ColorAndOpacity(FSlateColor::UseForeground())
 				]
@@ -415,7 +415,7 @@ private:
 			}
 			else
 			{
-				NameTextBlock->SetTextStyle(&FNiagaraEditorWidgetsStyle::Get().GetWidgetStyle<FTextBlockStyle>("NiagaraEditor.Stack.ItemText"));
+				NameTextBlock->SetTextStyle(&FNiagaraEditorWidgetsStyle::Get().GetWidgetStyle<FTextBlockStyle>("NiagaraEditor.OverviewStack.ItemText"));
 			}
 		}
 	}
