Handling edge case in previous variable definition where the attribute existed later in the script order, causing mismatches in outputs across scripts.

diff --git a/Engine/Plugins/FX/Niagara/Source/NiagaraCore/Private/NiagaraCustomVersion.cpp b/Engine/Plugins/FX/Niagara/Source/NiagaraCore/Private/NiagaraCustomVersion.cpp
--- a/Engine/Plugins/FX/Niagara/Source/NiagaraCore/Private/NiagaraCustomVersion.cpp
+++ b/Engine/Plugins/FX/Niagara/Source/NiagaraCore/Private/NiagaraCustomVersion.cpp
@@ -8,4 +8,4 @@ const FGuid FNiagaraCustomVersion::GUID(0xFCF57AFA, 0x50764283, 0xB9A9E658, 0xFF
 // Register the custom version with core
 FCustomVersionRegistration GRegisterNiagaraCustomVersion(FNiagaraCustomVersion::GUID, FNiagaraCustomVersion::LatestVersion, TEXT("NiagaraVer"));
 
-const FGuid FNiagaraCustomVersion::LatestScriptCompileVersion(0x0FC7D74E, 0x22166A42, 0xBFA2127F, 0xC334B412);
+const FGuid FNiagaraCustomVersion::LatestScriptCompileVersion(0xB25900FB, 0xC1D24C63, 0x8D25CC89, 0x026FF460);
diff --git a/Engine/Plugins/FX/Niagara/Source/NiagaraEditor/Private/NiagaraHlslTranslator.cpp b/Engine/Plugins/FX/Niagara/Source/NiagaraEditor/Private/NiagaraHlslTranslator.cpp
--- a/Engine/Plugins/FX/Niagara/Source/NiagaraEditor/Private/NiagaraHlslTranslator.cpp
+++ b/Engine/Plugins/FX/Niagara/Source/NiagaraEditor/Private/NiagaraHlslTranslator.cpp
@@ -763,18 +763,23 @@ FString FHlslNiagaraTranslator::BuildParameterMapHlslDefinitions(TArray<FNiagara
 		TArray<FNiagaraVariable> Vars = OtherOutputParamMapHistories[ParamMapIdx].Variables;
 		for (const FNiagaraVariableBase& Var : CompileOptions.AdditionalVariables)
 		{
-			bool bAddVar = true;
+			uint32 Uses = 0;
 			if (FNiagaraParameterMapHistory::IsPreviousValue(Var))
 			{
 				FNiagaraVariable Source = FNiagaraParameterMapHistory::GetSourceForPreviousValue(FNiagaraVariable(Var));
-				if (!UniqueVariables.Contains(Source))
-				{
-					// Disallow the addition previous values if its source is not used
-					bAddVar = false;
+				for (int32 ParamMapIdxTest = 0; ParamMapIdxTest < OtherOutputParamMapHistories.Num(); ParamMapIdxTest++)
+				{				
+
+					if (OtherOutputParamMapHistories[ParamMapIdxTest].Variables.Contains(Source))
+					{
+						// Disallow the addition previous values if its source is not used
+						Uses++;
+						break;
+					}
 				}
 			}
 
-			if (bAddVar)
+			if (Uses > 0)
 			{
 				Vars.AddUnique(FNiagaraVariable(Var.GetType(), Var.GetName()));
 			}
