- Editgrate CL 16041510 - Add SetVariableTexture method

#rnx

diff --git a/Engine/Plugins/FX/Niagara/Source/Niagara/Private/NiagaraComponent.cpp b/Engine/Plugins/FX/Niagara/Source/Niagara/Private/NiagaraComponent.cpp
--- a/Engine/Plugins/FX/Niagara/Source/Niagara/Private/NiagaraComponent.cpp
+++ b/Engine/Plugins/FX/Niagara/Source/Niagara/Private/NiagaraComponent.cpp
@@ -1838,7 +1838,6 @@ void UNiagaraComponent::SetVariableMaterial(FName InVariableName, UMaterialInter
 #endif
 }
 
-
 void UNiagaraComponent::SetVariableStaticMesh(FName InVariableName, UStaticMesh* InValue)
 {
 	if (SystemInstanceController.IsValid())
@@ -1852,6 +1851,19 @@ void UNiagaraComponent::SetVariableStaticMesh(FName InVariableName, UStaticMesh*
 #endif
 }
 
+void UNiagaraComponent::SetVariableTexture(FName InVariableName, UTexture* Texture)
+{
+	if (SystemInstanceController.IsValid())
+	{
+		TWeakObjectPtr<UTexture> ObjPtr = Texture;
+		SystemInstanceController->SetVariable_Deferred(InVariableName, ObjPtr);
+	}
+#if WITH_EDITOR
+	const FNiagaraVariable VariableDesc(FNiagaraTypeDefinition::GetUTextureDef(), InVariableName);
+	SetParameterOverride(VariableDesc, FNiagaraVariant(Texture));
+#endif
+}
+
 void UNiagaraComponent::SetVariableTextureRenderTarget(FName InVariableName, UTextureRenderTarget* TextureRenderTarget)
 {
 	if (SystemInstanceController.IsValid())
diff --git a/Engine/Plugins/FX/Niagara/Source/Niagara/Public/NiagaraComponent.h b/Engine/Plugins/FX/Niagara/Source/Niagara/Public/NiagaraComponent.h
--- a/Engine/Plugins/FX/Niagara/Source/Niagara/Public/NiagaraComponent.h
+++ b/Engine/Plugins/FX/Niagara/Source/Niagara/Public/NiagaraComponent.h
@@ -410,6 +410,9 @@ public:
 	UFUNCTION(BlueprintCallable, Category = Niagara, meta = (DisplayName = "Set Niagara Variable (Static Mesh)"))
 	void SetVariableStaticMesh(FName InVariableName, UStaticMesh* InValue);
 
+	UFUNCTION(BlueprintCallable, Category = Niagara, meta = (DisplayName = "Set Niagara Variable (Texture)"))
+	void SetVariableTexture(FName InVariableName, class UTexture* Texture);
+
 	UFUNCTION(BlueprintCallable, Category = Niagara, meta = (DisplayName = "Set Niagara Variable (TextureRenderTarget)"))
 	void SetVariableTextureRenderTarget(FName InVariableName, class UTextureRenderTarget* TextureRenderTarget);
 
