Fix for Niagara skelmesh DI sampling triangles with unlimited bone influences - wasn't taking section bone map into account when grabbing bone matrices

diff --git a/Engine/Plugins/FX/Niagara/Shaders/Private/NiagaraDataInterfaceSkeletalMesh.ush b/Engine/Plugins/FX/Niagara/Shaders/Private/NiagaraDataInterfaceSkeletalMesh.ush
--- a/Engine/Plugins/FX/Niagara/Shaders/Private/NiagaraDataInterfaceSkeletalMesh.ush
+++ b/Engine/Plugins/FX/Niagara/Shaders/Private/NiagaraDataInterfaceSkeletalMesh.ush
@@ -109,16 +109,13 @@ float3x4 DISkelMesh_GetPrevBoneSkinningMatrix(in FDISkelMeshContext DIContext, u
     return float3x4(DIContext.MeshPrevBonesBuffer[Bone * 3], DIContext.MeshPrevBonesBuffer[Bone * 3 + 1], DIContext.MeshPrevBonesBuffer[Bone * 3 + 2]);
 }
 
-float3x4 DISkelMesh_GetPrevSkinningMatrix(in FDISkelMeshContext DIContext, uint VertexIndex, int4 BlendIndices, float4 BlendWeights)
+float3x4 DISkelMesh_GetPrevSkinningMatrix(in FDISkelMeshContext DIContext, uint SectionBoneOffset, int4 BlendIndices, float4 BlendWeights)
 {
-	// Get the matrix offset for each vertex because BlendIndices are stored relatively to each section start vertex.
-    uint MatrixOffset = DIContext.MeshTriangleMatricesOffsetBuffer[VertexIndex];
-
 	float3x4 Result;
-    Result = DISkelMesh_GetPrevBoneSkinningMatrix(DIContext, MatrixOffset + BlendIndices.x) * BlendWeights.x;
-    Result += DISkelMesh_GetPrevBoneSkinningMatrix(DIContext, MatrixOffset + BlendIndices.y) * BlendWeights.y;
-    Result += DISkelMesh_GetPrevBoneSkinningMatrix(DIContext, MatrixOffset + BlendIndices.z) * BlendWeights.z;
-    Result += DISkelMesh_GetPrevBoneSkinningMatrix(DIContext, MatrixOffset + BlendIndices.w) * BlendWeights.w;
+    Result = DISkelMesh_GetPrevBoneSkinningMatrix(DIContext, SectionBoneOffset + BlendIndices.x) * BlendWeights.x;
+    Result += DISkelMesh_GetPrevBoneSkinningMatrix(DIContext, SectionBoneOffset + BlendIndices.y) * BlendWeights.y;
+    Result += DISkelMesh_GetPrevBoneSkinningMatrix(DIContext, SectionBoneOffset + BlendIndices.z) * BlendWeights.z;
+    Result += DISkelMesh_GetPrevBoneSkinningMatrix(DIContext, SectionBoneOffset + BlendIndices.w) * BlendWeights.w;
 	return Result;
 }
 
@@ -127,16 +124,13 @@ float3x4 DISkelMesh_GetCurrBoneSkinningMatrix(in FDISkelMeshContext DIContext, u
 	return float3x4(DIContext.MeshCurrBonesBuffer[Bone * 3], DIContext.MeshCurrBonesBuffer[Bone * 3 + 1], DIContext.MeshCurrBonesBuffer[Bone * 3 + 2]);
 }
 
-float3x4 DISkelMesh_GetCurrSkinningMatrix(in FDISkelMeshContext DIContext, uint VertexIndex, int4 BlendIndices, float4 BlendWeights)
+float3x4 DISkelMesh_GetCurrSkinningMatrix(in FDISkelMeshContext DIContext, uint SectionBoneOffset, int4 BlendIndices, float4 BlendWeights)
 {
-	// Get the matrix offset for each vertex because BlendIndices are stored relatively to each section start vertex.
-	uint MatrixOffset = DIContext.MeshTriangleMatricesOffsetBuffer[VertexIndex];
-
 	float3x4 Result;
-    Result = DISkelMesh_GetCurrBoneSkinningMatrix(DIContext, MatrixOffset + BlendIndices.x) * BlendWeights.x;
-    Result += DISkelMesh_GetCurrBoneSkinningMatrix(DIContext, MatrixOffset + BlendIndices.y) * BlendWeights.y;
-    Result += DISkelMesh_GetCurrBoneSkinningMatrix(DIContext, MatrixOffset + BlendIndices.z) * BlendWeights.z;
-    Result += DISkelMesh_GetCurrBoneSkinningMatrix(DIContext, MatrixOffset + BlendIndices.w) * BlendWeights.w;
+    Result = DISkelMesh_GetCurrBoneSkinningMatrix(DIContext, SectionBoneOffset + BlendIndices.x) * BlendWeights.x;
+    Result += DISkelMesh_GetCurrBoneSkinningMatrix(DIContext, SectionBoneOffset + BlendIndices.y) * BlendWeights.y;
+    Result += DISkelMesh_GetCurrBoneSkinningMatrix(DIContext, SectionBoneOffset + BlendIndices.z) * BlendWeights.z;
+    Result += DISkelMesh_GetCurrBoneSkinningMatrix(DIContext, SectionBoneOffset + BlendIndices.w) * BlendWeights.w;
 	return Result;
 }
 
@@ -164,6 +158,9 @@ FDISkelMeshSkinnedVertex DISkelMesh_GetSkinnedVertex(in FDISkelMeshContext DICon
 		float3x4 PrevBoneMatrix;
 		float3x4 CurrBoneMatrix;
 
+		// Get the matrix offset for each vertex because BlendIndices are stored relatively to each section start vertex.
+		uint SectionBoneOffset = DIContext.MeshTriangleMatricesOffsetBuffer[VertexIndex];
+
 #if DISKELMESH_BONE_INFLUENCES == DISKELMESH_BONE_INFLUENCES_UNLIMITED
 		BRANCH
 		if (DIContext.UnlimitedBoneInfluences)
@@ -184,15 +181,17 @@ FDISkelMeshSkinnedVertex DISkelMesh_GetSkinnedVertex(in FDISkelMeshContext DICon
 					BoneIndex = DIContext.MeshSkinWeightBuffer[BoneIndexOffset + 1] << 8 | BoneIndex;
 					//@todo-lh: Workaround to fix issue in SPIRVEmitter of DXC; this block must be inside the if branch
 					float BoneWeight = float(DIContext.MeshSkinWeightBuffer[WeightsOffset + InfluenceIdx]) / 255.0f;
-					PrevBoneMatrix += DISkelMesh_GetPrevBoneSkinningMatrix(DIContext, BoneIndex) * BoneWeight;
-					CurrBoneMatrix += DISkelMesh_GetCurrBoneSkinningMatrix(DIContext, BoneIndex) * BoneWeight;
+
+					PrevBoneMatrix += DISkelMesh_GetPrevBoneSkinningMatrix(DIContext, SectionBoneOffset + BoneIndex) * BoneWeight;
+					CurrBoneMatrix += DISkelMesh_GetCurrBoneSkinningMatrix(DIContext, SectionBoneOffset + BoneIndex) * BoneWeight;
 				}
 				else
 				{
 					//@todo-lh: Workaround to fix issue in SPIRVEmitter of DXC; this block must be inside the if branch
 					float BoneWeight = float(DIContext.MeshSkinWeightBuffer[WeightsOffset + InfluenceIdx]) / 255.0f;
-					PrevBoneMatrix += DISkelMesh_GetPrevBoneSkinningMatrix(DIContext, BoneIndex) * BoneWeight;
-					CurrBoneMatrix += DISkelMesh_GetCurrBoneSkinningMatrix(DIContext, BoneIndex) * BoneWeight;
+
+					PrevBoneMatrix += DISkelMesh_GetPrevBoneSkinningMatrix(DIContext, SectionBoneOffset + BoneIndex) * BoneWeight;
+					CurrBoneMatrix += DISkelMesh_GetCurrBoneSkinningMatrix(DIContext, SectionBoneOffset + BoneIndex) * BoneWeight;
 				}
 			}	
 		}
@@ -202,16 +201,16 @@ FDISkelMeshSkinnedVertex DISkelMesh_GetSkinnedVertex(in FDISkelMeshContext DICon
 			uint WeightOffset = DIContext.MeshNumWeights == 8 ? 2 : 1;
 			int4 BlendIndices = DISkelMesh_UnpackIndices4(DIContext.MeshSkinWeightBuffer[VertexIndex * DIContext.MeshWeightStride]);
 	        float4 BlendWeights = DISkelMesh_UnpackWeights4(DIContext.MeshSkinWeightBuffer[VertexIndex * DIContext.MeshWeightStride + WeightOffset]);
-			CurrBoneMatrix = DISkelMesh_GetCurrSkinningMatrix(DIContext, VertexIndex, BlendIndices, BlendWeights);
-			PrevBoneMatrix = DISkelMesh_GetPrevSkinningMatrix(DIContext, VertexIndex, BlendIndices, BlendWeights);
+			CurrBoneMatrix = DISkelMesh_GetCurrSkinningMatrix(DIContext, SectionBoneOffset, BlendIndices, BlendWeights);
+			PrevBoneMatrix = DISkelMesh_GetPrevSkinningMatrix(DIContext, SectionBoneOffset, BlendIndices, BlendWeights);
 #if DISKELMESH_BONE_INFLUENCES != DISKELMESH_BONE_INFLUENCES_MAX4
 			// We are in 8 or unlimited mode so need to check for the extended weights
 			if (DIContext.MeshNumWeights == 8)
 			{
 			    BlendIndices = DISkelMesh_UnpackIndices4(DIContext.MeshSkinWeightBuffer[VertexIndex * DIContext.MeshWeightStride + 1]);
 			    BlendWeights = DISkelMesh_UnpackWeights4(DIContext.MeshSkinWeightBuffer[VertexIndex * DIContext.MeshWeightStride + WeightOffset + 1]);
-				CurrBoneMatrix += DISkelMesh_GetCurrSkinningMatrix(DIContext, VertexIndex, BlendIndices, BlendWeights);
-				PrevBoneMatrix += DISkelMesh_GetPrevSkinningMatrix(DIContext, VertexIndex, BlendIndices, BlendWeights);
+				CurrBoneMatrix += DISkelMesh_GetCurrSkinningMatrix(DIContext, SectionBoneOffset, BlendIndices, BlendWeights);
+				PrevBoneMatrix += DISkelMesh_GetPrevSkinningMatrix(DIContext, SectionBoneOffset, BlendIndices, BlendWeights);
 			}
 #endif
 		}
