Fixing nonunity CIS issues

diff --git a/Engine/Plugins/FX/Niagara/Source/Niagara/Private/NiagaraModule.cpp b/Engine/Plugins/FX/Niagara/Source/Niagara/Private/NiagaraModule.cpp
--- a/Engine/Plugins/FX/Niagara/Source/Niagara/Private/NiagaraModule.cpp
+++ b/Engine/Plugins/FX/Niagara/Source/Niagara/Private/NiagaraModule.cpp
@@ -25,6 +25,7 @@
 #include "Misc/CoreDelegates.h"
 #include "NiagaraDebuggerClient.h"
 #include "Particles/FXBudget.h"
+#include "Engine/StaticMesh.h"
 
 IMPLEMENT_MODULE(INiagaraModule, Niagara);
 
diff --git a/Engine/Plugins/FX/Niagara/Source/Niagara/Private/NiagaraSystemInstanceController.cpp b/Engine/Plugins/FX/Niagara/Source/Niagara/Private/NiagaraSystemInstanceController.cpp
--- a/Engine/Plugins/FX/Niagara/Source/Niagara/Private/NiagaraSystemInstanceController.cpp
+++ b/Engine/Plugins/FX/Niagara/Source/Niagara/Private/NiagaraSystemInstanceController.cpp
@@ -2,6 +2,7 @@
 
 #include "NiagaraSystemInstanceController.h"
 #include "Materials/MaterialInstanceDynamic.h"
+#include "Engine/StaticMesh.h"
 
 FNiagaraSystemInstanceController::FNiagaraSystemInstanceController()
 	: bNeedsRendererRecache(false)
