Editgrate CL 16061319 - Allow direct setting of override parameters before the Niagara component has activated (Fixes the user spline automated test)

diff --git a/Engine/Plugins/FX/Niagara/Source/Niagara/Private/NiagaraComponent.cpp b/Engine/Plugins/FX/Niagara/Source/Niagara/Private/NiagaraComponent.cpp
--- a/Engine/Plugins/FX/Niagara/Source/Niagara/Private/NiagaraComponent.cpp
+++ b/Engine/Plugins/FX/Niagara/Source/Niagara/Private/NiagaraComponent.cpp
@@ -1662,12 +1662,16 @@ void UNiagaraComponent::SetRandomSeedOffset(int32 NewRandomSeedOffset)
 
 void UNiagaraComponent::SetVariableLinearColor(FName InVariableName, const FLinearColor& InValue)
 {
+	const FNiagaraVariable VariableDesc(FNiagaraTypeDefinition::GetColorDef(), InVariableName);
 	if (SystemInstanceController.IsValid())
 	{
 		SystemInstanceController->SetVariable_Deferred(InVariableName, InValue);
 	}
+	else
+	{
+		OverrideParameters.SetParameterValue(InValue, VariableDesc, true);
+	}
 #if WITH_EDITOR
-	const FNiagaraVariable VariableDesc(FNiagaraTypeDefinition::GetColorDef(), InVariableName);
 	SetParameterOverride(VariableDesc, FNiagaraVariant(&InValue, sizeof(FLinearColor)));
 #endif
 }
@@ -1679,12 +1683,16 @@ void UNiagaraComponent::SetNiagaraVariableLinearColor(const FString& InVariableN
 
 void UNiagaraComponent::SetVariableQuat(FName InVariableName, const FQuat& InValue)
 {
+	const FNiagaraVariable VariableDesc(FNiagaraTypeDefinition::GetQuatDef(), InVariableName);
 	if (SystemInstanceController.IsValid())
 	{
 		SystemInstanceController->SetVariable_Deferred(InVariableName, InValue);
 	}
+	else
+	{
+		OverrideParameters.SetParameterValue(InValue, VariableDesc, true);
+	}
 #if WITH_EDITOR
-	const FNiagaraVariable VariableDesc(FNiagaraTypeDefinition::GetQuatDef(), InVariableName);
 	SetParameterOverride(VariableDesc, FNiagaraVariant(&InValue, sizeof(FQuat)));
 #endif
 }
@@ -1696,12 +1704,16 @@ void UNiagaraComponent::SetNiagaraVariableQuat(const FString& InVariableName, co
 
 void UNiagaraComponent::SetVariableVec4(FName InVariableName, const FVector4& InValue)
 {
+	const FNiagaraVariable VariableDesc(FNiagaraTypeDefinition::GetVec4Def(), InVariableName);
 	if (SystemInstanceController.IsValid())
 	{
 		SystemInstanceController->SetVariable_Deferred(InVariableName, InValue);
 	}
+	else
+	{
+		OverrideParameters.SetParameterValue(InValue, VariableDesc, true);
+	}
 #if WITH_EDITOR
-	const FNiagaraVariable VariableDesc(FNiagaraTypeDefinition::GetVec4Def(), InVariableName);
 	SetParameterOverride(VariableDesc, FNiagaraVariant(&InValue, sizeof(FVector4)));
 #endif
 }
@@ -1713,12 +1725,16 @@ void UNiagaraComponent::SetNiagaraVariableVec4(const FString& InVariableName, co
 
 void UNiagaraComponent::SetVariableVec3(FName InVariableName, FVector InValue)
 {
+	const FNiagaraVariable VariableDesc(FNiagaraTypeDefinition::GetVec3Def(), InVariableName);
 	if (SystemInstanceController.IsValid())
 	{
 		SystemInstanceController->SetVariable_Deferred(InVariableName, InValue);
 	}
+	else
+	{
+		OverrideParameters.SetParameterValue(InValue, VariableDesc, true);
+	}
 #if WITH_EDITOR
-	const FNiagaraVariable VariableDesc(FNiagaraTypeDefinition::GetVec3Def(), InVariableName);
 	SetParameterOverride(VariableDesc, FNiagaraVariant(&InValue, sizeof(FVector)));
 #endif
 }
@@ -1730,12 +1746,16 @@ void UNiagaraComponent::SetNiagaraVariableVec3(const FString& InVariableName, FV
 
 void UNiagaraComponent::SetVariableVec2(FName InVariableName, FVector2D InValue)
 {
+	const FNiagaraVariable VariableDesc(FNiagaraTypeDefinition::GetVec2Def(), InVariableName);
 	if (SystemInstanceController.IsValid())
 	{
 		SystemInstanceController->SetVariable_Deferred(InVariableName, InValue);
 	}
+	else
+	{
+		OverrideParameters.SetParameterValue(InValue, VariableDesc, true);
+	}
 #if WITH_EDITOR
-	const FNiagaraVariable VariableDesc(FNiagaraTypeDefinition::GetVec2Def(), InVariableName);
 	SetParameterOverride(VariableDesc, FNiagaraVariant(&InValue, sizeof(FVector2D)));
 #endif
 }
@@ -1747,12 +1767,16 @@ void UNiagaraComponent::SetNiagaraVariableVec2(const FString& InVariableName, FV
 
 void UNiagaraComponent::SetVariableFloat(FName InVariableName, float InValue)
 {
+	const FNiagaraVariable VariableDesc(FNiagaraTypeDefinition::GetFloatDef(), InVariableName);
 	if (SystemInstanceController.IsValid())
 	{
 		SystemInstanceController->SetVariable_Deferred(InVariableName, InValue);
 	}
+	else
+	{
+		OverrideParameters.SetParameterValue(InValue, VariableDesc, true);
+	}
 #if WITH_EDITOR
-	const FNiagaraVariable VariableDesc(FNiagaraTypeDefinition::GetFloatDef(), InVariableName);
 	SetParameterOverride(VariableDesc, FNiagaraVariant(&InValue, sizeof(float)));
 #endif
 }
@@ -1764,12 +1788,16 @@ void UNiagaraComponent::SetNiagaraVariableFloat(const FString& InVariableName, f
 
 void UNiagaraComponent::SetVariableInt(FName InVariableName, int32 InValue)
 {
+	const FNiagaraVariable VariableDesc(FNiagaraTypeDefinition::GetIntDef(), InVariableName);
 	if (SystemInstanceController.IsValid())
 	{
 		SystemInstanceController->SetVariable_Deferred(InVariableName, InValue);
 	}
+	else
+	{
+		OverrideParameters.SetParameterValue(InValue, VariableDesc, true);
+	}
 #if WITH_EDITOR
-	const FNiagaraVariable VariableDesc(FNiagaraTypeDefinition::GetIntDef(), InVariableName);
 	SetParameterOverride(VariableDesc, FNiagaraVariant(&InValue, sizeof(int32)));
 #endif
 }
@@ -1781,13 +1809,17 @@ void UNiagaraComponent::SetNiagaraVariableInt(const FString& InVariableName, int
 
 void UNiagaraComponent::SetVariableBool(FName InVariableName, bool InValue)
 {
+	const FNiagaraVariable VariableDesc(FNiagaraTypeDefinition::GetBoolDef(), InVariableName);
+	const FNiagaraBool BoolValue(InValue);
 	if (SystemInstanceController.IsValid())
 	{
 		SystemInstanceController->SetVariable_Deferred(InVariableName, InValue);
 	}
+	else
+	{
+		OverrideParameters.SetParameterValue(BoolValue, VariableDesc, true);
+	}
 #if WITH_EDITOR
-	const FNiagaraVariable VariableDesc(FNiagaraTypeDefinition::GetBoolDef(), InVariableName);
-	const FNiagaraBool BoolValue(InValue);
 	SetParameterOverride(VariableDesc, FNiagaraVariant(&BoolValue, sizeof(FNiagaraBool)));
 #endif
 }
@@ -1809,13 +1841,17 @@ void UNiagaraComponent::SetNiagaraVariableActor(const FString& InVariableName, A
 
 void UNiagaraComponent::SetVariableObject(FName InVariableName, UObject* InValue)
 {
+	const FNiagaraVariable VariableDesc(FNiagaraTypeDefinition::GetUObjectDef(), InVariableName);
 	if (SystemInstanceController.IsValid())
 	{
 		TWeakObjectPtr<UObject> ObjPtr = InValue;
 		SystemInstanceController->SetVariable_Deferred(InVariableName, ObjPtr);
 	}
+	else
+	{
+		OverrideParameters.SetUObject(InValue, VariableDesc);
+	}
 #if WITH_EDITOR
-	const FNiagaraVariable VariableDesc(FNiagaraTypeDefinition::GetUObjectDef(), InVariableName);
 	SetParameterOverride(VariableDesc, FNiagaraVariant(InValue));
 #endif
 }
@@ -1827,13 +1863,17 @@ void UNiagaraComponent::SetNiagaraVariableObject(const FString& InVariableName,
 
 void UNiagaraComponent::SetVariableMaterial(FName InVariableName, UMaterialInterface* InValue)
 {
+	const FNiagaraVariable VariableDesc(FNiagaraTypeDefinition::GetUMaterialDef(), InVariableName);
 	if (SystemInstanceController.IsValid())
 	{
 		TWeakObjectPtr<UMaterialInterface> ObjPtr = InValue;
 		SystemInstanceController->SetVariable_Deferred(InVariableName, ObjPtr);
 	}
+	else
+	{
+		OverrideParameters.SetUObject(InValue, VariableDesc);
+	}
 #if WITH_EDITOR
-	const FNiagaraVariable VariableDesc(FNiagaraTypeDefinition::GetUMaterialDef(), InVariableName);
 	SetParameterOverride(VariableDesc, FNiagaraVariant(InValue));
 #endif
 }
@@ -1866,13 +1906,17 @@ void UNiagaraComponent::SetVariableTexture(FName InVariableName, UTexture* Textu
 
 void UNiagaraComponent::SetVariableTextureRenderTarget(FName InVariableName, UTextureRenderTarget* TextureRenderTarget)
 {
+	const FNiagaraVariable VariableDesc(FNiagaraTypeDefinition::GetUTextureRenderTargetDef(), InVariableName);
 	if (SystemInstanceController.IsValid())
 	{
 		TWeakObjectPtr<UTextureRenderTarget> ObjPtr = TextureRenderTarget;
 		SystemInstanceController->SetVariable_Deferred(InVariableName, ObjPtr);
 	}
+	else
+	{
+		OverrideParameters.SetUObject(TextureRenderTarget, VariableDesc);
+	}
 #if WITH_EDITOR
-	const FNiagaraVariable VariableDesc(FNiagaraTypeDefinition::GetUTextureRenderTargetDef(), InVariableName);
 	SetParameterOverride(VariableDesc, FNiagaraVariant(TextureRenderTarget));
 #endif
 }
