- Fix issues around data interfaces not being able to understand the number of attributes they have as the CDO is called into to generate the parameter definition, not the actual data interface

#okforgitub public
#rnx

diff --git a/Engine/Plugins/FX/Niagara/Source/Niagara/Classes/NiagaraDataInterfaceGrid3DCollection.h b/Engine/Plugins/FX/Niagara/Source/Niagara/Classes/NiagaraDataInterfaceGrid3DCollection.h
--- a/Engine/Plugins/FX/Niagara/Source/Niagara/Classes/NiagaraDataInterfaceGrid3DCollection.h
+++ b/Engine/Plugins/FX/Niagara/Source/Niagara/Classes/NiagaraDataInterfaceGrid3DCollection.h
@@ -79,6 +79,8 @@ struct FGrid3DCollectionRWInstanceData_RenderThread
 	FGrid3DBuffer* CurrentData = nullptr;
 	FGrid3DBuffer* DestinationData = nullptr;
 
+	FReadBuffer PerAttributeData;
+
 	FTextureRHIRef RenderTargetToCopyTo;
 
 	void BeginSimulate(FRHICommandList& RHICmdList);
@@ -204,6 +206,8 @@ public:
 
 	static const FString NumTilesName;
 	static const FString OneOverNumTilesName;
+	static const FString UnitClampMinName;
+	static const FString UnitClampMaxName;
 
 	static const FString GridName;
 	static const FString OutputGridName;
@@ -239,10 +243,7 @@ public:
 	static const FName GetPreviousFloatValueFunctionName;
 	static const FName SamplePreviousGridFloatFunctionName;
 	static const FString AttributeIndicesBaseName;
-	static const FString AttributePixelOffsetName;
-	static const FString AttributeUVOffsetName;
-	static const FString AttributeUVClampMinName;
-	static const FString AttributeUVClampMaxName;
+	static const FString PerAttributeDataName;
 	static const TCHAR* VectorComponentNames[];
 
 	static const FName SetNumCellsFunctionName;
diff --git a/Engine/Plugins/FX/Niagara/Source/Niagara/Private/NiagaraDataInterfaceGrid3DCollection.cpp b/Engine/Plugins/FX/Niagara/Source/Niagara/Private/NiagaraDataInterfaceGrid3DCollection.cpp
--- a/Engine/Plugins/FX/Niagara/Source/Niagara/Private/NiagaraDataInterfaceGrid3DCollection.cpp
+++ b/Engine/Plugins/FX/Niagara/Source/Niagara/Private/NiagaraDataInterfaceGrid3DCollection.cpp
@@ -19,6 +19,8 @@
 
 const FString UNiagaraDataInterfaceGrid3DCollection::NumTilesName(TEXT("NumTiles_"));
 const FString UNiagaraDataInterfaceGrid3DCollection::OneOverNumTilesName(TEXT("OneOverNumTiles_"));
+const FString UNiagaraDataInterfaceGrid3DCollection::UnitClampMinName(TEXT("UnitClampMin_"));
+const FString UNiagaraDataInterfaceGrid3DCollection::UnitClampMaxName(TEXT("UnitClampMax_"));
 
 const FString UNiagaraDataInterfaceGrid3DCollection::GridName(TEXT("Grid_"));
 const FString UNiagaraDataInterfaceGrid3DCollection::OutputGridName(TEXT("OutputGrid_"));
@@ -64,10 +66,7 @@ const FName UNiagaraDataInterfaceGrid3DCollection::GetFloatAttributeIndexFunctio
 const FName UNiagaraDataInterfaceGrid3DCollection::SampleGridFunctionName("SampleGrid");
 
 const FString UNiagaraDataInterfaceGrid3DCollection::AttributeIndicesBaseName(TEXT("AttributeIndices_"));
-const FString UNiagaraDataInterfaceGrid3DCollection::AttributePixelOffsetName(TEXT("AttributePixelOffset_"));
-const FString UNiagaraDataInterfaceGrid3DCollection::AttributeUVOffsetName(TEXT("AttributeUVOffset_"));
-const FString UNiagaraDataInterfaceGrid3DCollection::AttributeUVClampMinName(TEXT("AttributeUVClampMin_"));
-const FString UNiagaraDataInterfaceGrid3DCollection::AttributeUVClampMaxName(TEXT("AttributeUVClampMax_"));
+const FString UNiagaraDataInterfaceGrid3DCollection::PerAttributeDataName(TEXT("PerAttributeDataName_"));
 
 const TCHAR* UNiagaraDataInterfaceGrid3DCollection::VectorComponentNames[] = { TEXT(".x"), TEXT(".y"), TEXT(".z"), TEXT(".w") };
 
@@ -153,6 +152,9 @@ public:
 		NumCellsParam.Bind(ParameterMap, *(UNiagaraDataInterfaceRWBase::NumCellsName + ParameterInfo.DataInterfaceHLSLSymbol));
 		NumTilesParam.Bind(ParameterMap, *(UNiagaraDataInterfaceGrid3DCollection::NumTilesName + ParameterInfo.DataInterfaceHLSLSymbol));
 		OneOverNumTilesParam.Bind(ParameterMap, *(UNiagaraDataInterfaceGrid3DCollection::OneOverNumTilesName + ParameterInfo.DataInterfaceHLSLSymbol));
+		UnitClampMinParam.Bind(ParameterMap, *(UNiagaraDataInterfaceGrid3DCollection::UnitClampMinName + ParameterInfo.DataInterfaceHLSLSymbol));
+		UnitClampMaxParam.Bind(ParameterMap, *(UNiagaraDataInterfaceGrid3DCollection::UnitClampMaxName + ParameterInfo.DataInterfaceHLSLSymbol));
+
 		UnitToUVParam.Bind(ParameterMap, *(UNiagaraDataInterfaceRWBase::UnitToUVName + ParameterInfo.DataInterfaceHLSLSymbol));
 		CellSizeParam.Bind(ParameterMap, *(UNiagaraDataInterfaceRWBase::CellSizeName + ParameterInfo.DataInterfaceHLSLSymbol));
 
@@ -163,10 +165,7 @@ public:
 
 		SamplerParam.Bind(ParameterMap, *(UNiagaraDataInterfaceGrid3DCollection::SamplerName + ParameterInfo.DataInterfaceHLSLSymbol));
 		AttributeIndicesParam.Bind(ParameterMap, *(UNiagaraDataInterfaceGrid3DCollection::AttributeIndicesBaseName + ParameterInfo.DataInterfaceHLSLSymbol));
-		AttributePixelOffsetParam.Bind(ParameterMap, *(UNiagaraDataInterfaceGrid3DCollection::AttributePixelOffsetName + ParameterInfo.DataInterfaceHLSLSymbol));
-		AttributeUVOffsetParam.Bind(ParameterMap, *(UNiagaraDataInterfaceGrid3DCollection::AttributeUVOffsetName + ParameterInfo.DataInterfaceHLSLSymbol));
-		AttributeUVClampMinParam.Bind(ParameterMap, *(UNiagaraDataInterfaceGrid3DCollection::AttributeUVClampMinName + ParameterInfo.DataInterfaceHLSLSymbol));
-		AttributeUVClampMaxParam.Bind(ParameterMap, *(UNiagaraDataInterfaceGrid3DCollection::AttributeUVClampMaxName + ParameterInfo.DataInterfaceHLSLSymbol));
+		PerAttributeDataParam.Bind(ParameterMap, *(UNiagaraDataInterfaceGrid3DCollection::PerAttributeDataName + ParameterInfo.DataInterfaceHLSLSymbol));
 
 
 		// Gather up all the attribute names referenced. Note that there may be multiple in the list of the same name,
@@ -209,6 +208,10 @@ public:
 		SetShaderValue(RHICmdList, ComputeShaderRHI, NumTilesParam, ProxyData->NumTiles);
 		SetShaderValue(RHICmdList, ComputeShaderRHI, OneOverNumTilesParam, FVector(1.0f) / FVector(ProxyData->NumTiles));
 
+		const FVector HalfPixelOffset = FVector(0.5f / ProxyData->NumCells.X, 0.5f / ProxyData->NumCells.Y, 0.5f / ProxyData->NumCells.Z);
+		SetShaderValue(RHICmdList, ComputeShaderRHI, UnitClampMinParam, HalfPixelOffset);
+		SetShaderValue(RHICmdList, ComputeShaderRHI, UnitClampMaxParam, FVector::OneVector - HalfPixelOffset);
+
 		if (ProxyData->AttributeIndices.Num() == 0 && AttributeNames.Num() > 0)
 		{
 			int NumAttrIndices = Align(AttributeNames.Num(), 4);
@@ -241,52 +244,8 @@ public:
 		SetShaderValue(RHICmdList, ComputeShaderRHI, WorldBBoxSizeParam, ProxyData->WorldBBoxSize);
 
 		SetShaderValueArray(RHICmdList, ComputeShaderRHI, AttributeIndicesParam, ProxyData->AttributeIndices.GetData(), ProxyData->AttributeIndices.Num());
-		{
-			const int32 NumAttributes = ProxyData->TotalNumAttributes;
-			TArray<FIntVector4, TInlineAllocator<16>> AttributePixelOffset;
-			TArray<FVector4, TInlineAllocator<16>> AttributeUVOffset;
-			TArray<FVector4, TInlineAllocator<16>> AttributeUVClampMin;
-			TArray<FVector4, TInlineAllocator<16>> AttributeUVClampMax;
-			AttributePixelOffset.Reserve(NumAttributes);
-			AttributeUVOffset.Reserve(NumAttributes);
-			AttributeUVClampMin.Reserve(NumAttributes);
-			AttributeUVClampMax.Reserve(NumAttributes);
-
-			for ( int32 i=0; i < NumAttributes; ++i )
-			{
-				const FIntVector AttributeTileIndex(i % ProxyData->NumTiles.X, (i / ProxyData->NumTiles.X) % ProxyData->NumTiles.Y, i / (ProxyData->NumTiles.X * ProxyData->NumTiles.Y));
-
-				AttributePixelOffset.Emplace(
-					AttributeTileIndex.X * ProxyData->NumCells.X,
-					AttributeTileIndex.Y * ProxyData->NumCells.Y,
-					AttributeTileIndex.Z * ProxyData->NumCells.Z,
-					0
-				);
-				AttributeUVOffset.Emplace(
-					(1.0f / ProxyData->NumTiles.X) * float(AttributeTileIndex.X),
-					(1.0f / ProxyData->NumTiles.Y) * float(AttributeTileIndex.Y),
-					(1.0f / ProxyData->NumTiles.Z) * float(AttributeTileIndex.Z),
-					0.0f
-				);
-				AttributeUVClampMin.Emplace(
-					(float(AttributeTileIndex.X) * ProxyData->NumCells.X + 0.5) / (ProxyData->NumTiles.X * ProxyData->NumCells.X),
-					(float(AttributeTileIndex.Y) * ProxyData->NumCells.Y + 0.5) / (ProxyData->NumTiles.Y * ProxyData->NumCells.Y),
-					(float(AttributeTileIndex.Z) * ProxyData->NumCells.Z + 0.5) / (ProxyData->NumTiles.Z * ProxyData->NumCells.Z),
-					0.0f
-				);
-				AttributeUVClampMax.Emplace(
-					(float(AttributeTileIndex.X + 1) * ProxyData->NumCells.X - 0.5) / (ProxyData->NumTiles.X * ProxyData->NumCells.X),
-					(float(AttributeTileIndex.Y + 1) * ProxyData->NumCells.Y - 0.5) / (ProxyData->NumTiles.Y * ProxyData->NumCells.Y),
-					(float(AttributeTileIndex.Z + 1) * ProxyData->NumCells.Z - 0.5) / (ProxyData->NumTiles.Z * ProxyData->NumCells.Z),
-					0.0f
-				);
-			}
-
-			SetShaderValueArray(RHICmdList, ComputeShaderRHI, AttributePixelOffsetParam, AttributePixelOffset.GetData(), AttributePixelOffset.Num());
-			SetShaderValueArray(RHICmdList, ComputeShaderRHI, AttributeUVOffsetParam, AttributeUVOffset.GetData(), AttributeUVOffset.Num());
-			SetShaderValueArray(RHICmdList, ComputeShaderRHI, AttributeUVClampMinParam, AttributeUVClampMin.GetData(), AttributeUVClampMin.Num());
-			SetShaderValueArray(RHICmdList, ComputeShaderRHI, AttributeUVClampMaxParam, AttributeUVClampMax.GetData(), AttributeUVClampMax.Num());
-		}
+		
+		SetSRVParameter(RHICmdList, ComputeShaderRHI, PerAttributeDataParam, ProxyData->PerAttributeData.SRV);
 
 		FRHISamplerState *SamplerState = TStaticSamplerState<SF_Bilinear, AM_Clamp, AM_Clamp, AM_Clamp>::GetRHI();
 		SetSamplerParameter(RHICmdList, ComputeShaderRHI, SamplerParam, SamplerState);
@@ -335,16 +294,15 @@ private:
 	LAYOUT_FIELD(FShaderParameter, NumCellsParam);
 	LAYOUT_FIELD(FShaderParameter, NumTilesParam);
 	LAYOUT_FIELD(FShaderParameter, OneOverNumTilesParam);
+	LAYOUT_FIELD(FShaderParameter, UnitClampMinParam);
+	LAYOUT_FIELD(FShaderParameter, UnitClampMaxParam);
 	LAYOUT_FIELD(FShaderParameter, CellSizeParam);
 	LAYOUT_FIELD(FShaderParameter, WorldBBoxSizeParam);
 
 	LAYOUT_FIELD(FShaderResourceParameter, GridParam);
 	LAYOUT_FIELD(FRWShaderParameter, OutputGridParam);
 	LAYOUT_FIELD(FShaderParameter, AttributeIndicesParam);
-	LAYOUT_FIELD(FShaderParameter, AttributePixelOffsetParam);
-	LAYOUT_FIELD(FShaderParameter, AttributeUVOffsetParam);
-	LAYOUT_FIELD(FShaderParameter, AttributeUVClampMinParam);
-	LAYOUT_FIELD(FShaderParameter, AttributeUVClampMaxParam);
+	LAYOUT_FIELD(FShaderResourceParameter, PerAttributeDataParam);
 
 	LAYOUT_FIELD(FShaderResourceParameter, SamplerParam);
 	LAYOUT_FIELD(TMemoryImageArray<FName>, AttributeNames);
@@ -1130,12 +1088,11 @@ void UNiagaraDataInterfaceGrid3DCollection::GetParameterDefinitionHLSL(const FNi
 		RWTexture3D<float> RW{OutputGridName};
 		int3 {NumTiles};
 		float3 {OneOverNumTiles};
+		float3 {UnitClampMin};
+		float3 {UnitClampMax};
 		SamplerState {SamplerName};
 		int4 {AttributeIndicesName}[{AttributeInt4Count}];
-		int4 {AttributePixelOffsetName}[{AttributeCount}];
-		float4 {AttributeUVOffsetName}[{AttributeCount}];
-		float4 {AttributeUVClampMinName}[{AttributeCount}];
-		float4 {AttributeUVClampMaxName}[{AttributeCount}];
+		Buffer<float4> {PerAttributeDataName};
 		int {NumAttributesName};
 	)");
 
@@ -1152,14 +1109,13 @@ void UNiagaraDataInterfaceGrid3DCollection::GetParameterDefinitionHLSL(const FNi
 		{ TEXT("OutputGridName"),    OutputGridName + ParamInfo.DataInterfaceHLSLSymbol },
 		{ TEXT("NumTiles"),    NumTilesName + ParamInfo.DataInterfaceHLSLSymbol },
 		{ TEXT("OneOverNumTiles"), OneOverNumTilesName + ParamInfo.DataInterfaceHLSLSymbol },
+		{ TEXT("UnitClampMin"), UnitClampMinName + ParamInfo.DataInterfaceHLSLSymbol },
+		{ TEXT("UnitClampMax"), UnitClampMaxName + ParamInfo.DataInterfaceHLSLSymbol },
+
 		{ TEXT("UnitToUVName"), UNiagaraDataInterfaceRWBase::UnitToUVName + ParamInfo.DataInterfaceHLSLSymbol},
 		{ TEXT("AttributeIndicesName"), AttributeIndicesBaseName + ParamInfo.DataInterfaceHLSLSymbol},
-		{ TEXT("AttributePixelOffsetName"), AttributePixelOffsetName + ParamInfo.DataInterfaceHLSLSymbol},
-		{ TEXT("AttributeUVOffsetName"), AttributeUVOffsetName + ParamInfo.DataInterfaceHLSLSymbol},
-		{ TEXT("AttributeUVClampMinName"), AttributeUVClampMinName + ParamInfo.DataInterfaceHLSLSymbol},
-		{ TEXT("AttributeUVClampMaxName"), AttributeUVClampMaxName + ParamInfo.DataInterfaceHLSLSymbol},
+		{ TEXT("PerAttributeDataName"), PerAttributeDataName + ParamInfo.DataInterfaceHLSLSymbol},
 		{ TEXT("AttributeInt4Count"), AttributeInt4Count},
-		{ TEXT("AttributeCount"), AttributeCount },
 		{ TEXT("NumAttributesName"), UNiagaraDataInterfaceRWBase::NumAttributesName + ParamInfo.DataInterfaceHLSLSymbol},
 	};
 	OutHLSL += FString::Format(FormatDeclarations, ArgsDeclarations);
@@ -1190,7 +1146,7 @@ void UNiagaraDataInterfaceGrid3DCollection::WriteSetHLSL(const FNiagaraDataInter
 
 	FormatBounds += TEXT(R"(	
 					int CurAttributeIndex = In_AttributeIndex + i;
-					int3 TileOffset = {AttributePixelOffsetName}[CurAttributeIndex].xyz;
+					int3 TileOffset = {PerAttributeDataName}[CurAttributeIndex].xyz;
 					RW{OutputGrid}[int3(In_IndexX + TileOffset.x, In_IndexY + TileOffset.y, In_IndexZ + TileOffset.z)] = Val;
 				}
 			}
@@ -1200,7 +1156,7 @@ void UNiagaraDataInterfaceGrid3DCollection::WriteSetHLSL(const FNiagaraDataInter
 		{TEXT("OutputGrid"), OutputGridName + ParamInfo.DataInterfaceHLSLSymbol},
 		{TEXT("NumCellsName"), UNiagaraDataInterfaceRWBase::NumCellsName + ParamInfo.DataInterfaceHLSLSymbol},
 		{TEXT("UnitToUVName"), UNiagaraDataInterfaceRWBase::UnitToUVName + ParamInfo.DataInterfaceHLSLSymbol},
-		{TEXT("AttributePixelOffsetName"), AttributePixelOffsetName + ParamInfo.DataInterfaceHLSLSymbol},
+		{TEXT("PerAttributeDataName"), PerAttributeDataName + ParamInfo.DataInterfaceHLSLSymbol},
 		{TEXT("AttributeIndicesName"), AttributeIndicesBaseName + ParamInfo.DataInterfaceHLSLSymbol},
 		{TEXT("AttributeIndexGroup"), FunctionInstanceIndex / 4},
 		{TEXT("AttributeIndexComponent"), VectorComponentNames[FunctionInstanceIndex % 4]},
@@ -1209,7 +1165,7 @@ void UNiagaraDataInterfaceGrid3DCollection::WriteSetHLSL(const FNiagaraDataInter
 		{TEXT("NumTiles"), NumTilesName + ParamInfo.DataInterfaceHLSLSymbol},
 		{TEXT("UnitToUVName"), UNiagaraDataInterfaceRWBase::UnitToUVName + ParamInfo.DataInterfaceHLSLSymbol},
 		{TEXT("NumCellsName"), UNiagaraDataInterfaceRWBase::NumCellsName + ParamInfo.DataInterfaceHLSLSymbol},
-
+		{TEXT("NumAttributesName"), UNiagaraDataInterfaceRWBase::NumAttributesName + ParamInfo.DataInterfaceHLSLSymbol},
 	};
 	OutHLSL += FString::Format(*FormatBounds, ArgsBounds);
 }
@@ -1225,7 +1181,7 @@ void UNiagaraDataInterfaceGrid3DCollection::WriteGetHLSL(const FNiagaraDataInter
 			    for (int i = 0; i < {NumChannels}; i++)
 				{
 					int CurAttributeIndex = In_AttributeIndex + i;
-					int3 TileOffset = {AttributePixelOffsetName}[CurAttributeIndex].xyz;
+					int3 TileOffset = {PerAttributeDataName}[CurAttributeIndex].xyz;
 					float Val = {Grid}.Load(int4(In_IndexX + TileOffset.x, In_IndexY + TileOffset.y, In_IndexZ + TileOffset.z, 0));
 					)");
 	if (InNumChannels == 1)
@@ -1246,14 +1202,14 @@ void UNiagaraDataInterfaceGrid3DCollection::WriteGetHLSL(const FNiagaraDataInter
 		{TEXT("Grid"), GridName + ParamInfo.DataInterfaceHLSLSymbol},
 		{TEXT("NumCellsName"), UNiagaraDataInterfaceRWBase::NumCellsName + ParamInfo.DataInterfaceHLSLSymbol},
 		{TEXT("UnitToUVName"), UNiagaraDataInterfaceRWBase::UnitToUVName + ParamInfo.DataInterfaceHLSLSymbol},
-		{TEXT("AttributePixelOffsetName"), AttributePixelOffsetName + ParamInfo.DataInterfaceHLSLSymbol},
+		{TEXT("PerAttributeDataName"), PerAttributeDataName + ParamInfo.DataInterfaceHLSLSymbol},
 		{TEXT("AttributeIndicesName"), AttributeIndicesBaseName + ParamInfo.DataInterfaceHLSLSymbol},
 		{TEXT("AttributeIndexGroup"), FunctionInstanceIndex / 4},
 		{TEXT("AttributeIndexComponent"), VectorComponentNames[FunctionInstanceIndex % 4]},
 		{TEXT("NumChannels"), FString::FromInt(InNumChannels)},
 		{TEXT("NumChannelsVariableSuffix"), InNumChannels > 1 ? FString::FromInt(InNumChannels) : TEXT("")},
 		{TEXT("NumTiles"), NumTilesName + ParamInfo.DataInterfaceHLSLSymbol},
-
+		{TEXT("NumAttributesName"), UNiagaraDataInterfaceRWBase::NumAttributesName + ParamInfo.DataInterfaceHLSLSymbol},
 	};
 	OutHLSL += FString::Format(*FormatBounds, ArgsBounds);
 }
@@ -1265,11 +1221,11 @@ void UNiagaraDataInterfaceGrid3DCollection::WriteSampleHLSL(const FNiagaraDataIn
 			{
 				int In_AttributeIndex = {AttributeIndicesName}[{AttributeIndexGroup}]{AttributeIndexComponent};
 
+				float3 TileUV = clamp(In_Unit, {UnitClampMin}, {UnitClampMax}) * {OneOverNumTiles};
 			    for (int i = 0; i < {NumChannels}; i++)
 				{
 					int CurAttributeIndex = In_AttributeIndex + i;
-					float3 UVW = {AttributeUVOffsetName}[CurAttributeIndex].xyz + (In_Unit * {OneOverNumTiles});
-					UVW = clamp(UVW, {AttributeUVClampMinName}[CurAttributeIndex].xyz, {AttributeUVClampMaxName}[CurAttributeIndex].xyz);
+					float3 UVW = {PerAttributeDataName}[CurAttributeIndex + {NumAttributesName}].xyz + TileUV;
 
 					float Val = {Grid}.SampleLevel({SamplerName}, UVW, 0);
 					)");
@@ -1296,14 +1252,14 @@ void UNiagaraDataInterfaceGrid3DCollection::WriteSampleHLSL(const FNiagaraDataIn
 		{TEXT("NumChannels"), FString::FromInt(InNumChannels)},
 		{TEXT("NumChannelsVariableSuffix"), InNumChannels > 1 ? FString::FromInt(InNumChannels) : TEXT("")},
 		{TEXT("AttributeIndicesName"), AttributeIndicesBaseName + ParamInfo.DataInterfaceHLSLSymbol},
-		{TEXT("AttributePixelOffsetName"), AttributePixelOffsetName + ParamInfo.DataInterfaceHLSLSymbol},
-		{TEXT("AttributeUVOffsetName"), AttributeUVOffsetName + ParamInfo.DataInterfaceHLSLSymbol},
-		{TEXT("AttributeUVClampMinName"), AttributeUVClampMinName + ParamInfo.DataInterfaceHLSLSymbol},
-		{TEXT("AttributeUVClampMaxName"), AttributeUVClampMaxName + ParamInfo.DataInterfaceHLSLSymbol},
+		{TEXT("PerAttributeDataName"), PerAttributeDataName + ParamInfo.DataInterfaceHLSLSymbol},
 		{TEXT("AttributeIndexGroup"), FunctionInstanceIndex / 4},
 		{TEXT("AttributeIndexComponent"), VectorComponentNames[FunctionInstanceIndex % 4]},
 		{TEXT("NumTiles"), NumTilesName + ParamInfo.DataInterfaceHLSLSymbol},
 		{TEXT("OneOverNumTiles"), OneOverNumTilesName + ParamInfo.DataInterfaceHLSLSymbol },
+		{TEXT("UnitClampMin"), UnitClampMinName + ParamInfo.DataInterfaceHLSLSymbol },
+		{TEXT("UnitClampMax"), UnitClampMaxName + ParamInfo.DataInterfaceHLSLSymbol },
+		{TEXT("NumAttributesName"), UNiagaraDataInterfaceRWBase::NumAttributesName + ParamInfo.DataInterfaceHLSLSymbol},
 	};
 	OutHLSL += FString::Format(*FormatBounds, ArgsBounds);
 }
@@ -1384,11 +1340,11 @@ bool UNiagaraDataInterfaceGrid3DCollection::GetFunctionHLSL(const FNiagaraDataIn
 		{TEXT("SamplerName"), SamplerName + ParamInfo.DataInterfaceHLSLSymbol},
 		{TEXT("NumTiles"), NumTilesName + ParamInfo.DataInterfaceHLSLSymbol},
 		{TEXT("OneOverNumTiles"), OneOverNumTilesName + ParamInfo.DataInterfaceHLSLSymbol },
+		{TEXT("UnitClampMin"), UnitClampMinName + ParamInfo.DataInterfaceHLSLSymbol },
+		{TEXT("UnitClampMax"), UnitClampMaxName + ParamInfo.DataInterfaceHLSLSymbol },
 		{TEXT("NumCellsName"), UNiagaraDataInterfaceRWBase::NumCellsName + ParamInfo.DataInterfaceHLSLSymbol},
-		{TEXT("AttributePixelOffsetName"), AttributePixelOffsetName + ParamInfo.DataInterfaceHLSLSymbol},
-		{TEXT("AttributeUVOffsetName"), AttributeUVOffsetName + ParamInfo.DataInterfaceHLSLSymbol},
-		{TEXT("AttributeUVClampMinName"), AttributeUVClampMinName + ParamInfo.DataInterfaceHLSLSymbol},
-		{TEXT("AttributeUVClampMaxName"), AttributeUVClampMaxName + ParamInfo.DataInterfaceHLSLSymbol},
+		{TEXT("PerAttributeDataName"), PerAttributeDataName + ParamInfo.DataInterfaceHLSLSymbol},
+		{TEXT("NumAttributesName"), UNiagaraDataInterfaceRWBase::NumAttributesName + ParamInfo.DataInterfaceHLSLSymbol},
 	};
 
 	if (FunctionInfo.DefinitionName == GetValueFunctionName || FunctionInfo.DefinitionName == GetPreviousValueAtIndexFunctionName)
@@ -1396,9 +1352,12 @@ bool UNiagaraDataInterfaceGrid3DCollection::GetFunctionHLSL(const FNiagaraDataIn
 		static const TCHAR *FormatBounds = TEXT(R"(
 			void {FunctionName}(int In_IndexX, int In_IndexY, int In_IndexZ, int In_AttributeIndex, out float Out_Val)
 			{
-				int3 TileOffset = {AttributePixelOffsetName}[In_AttributeIndex].xyz;
-
-				Out_Val = {Grid}.Load(int4(In_IndexX + TileOffset.x, In_IndexY + TileOffset.y, In_IndexZ + TileOffset.z, 0));
+				Out_Val = 0;
+				if ( In_AttributeIndex < {NumAttributesName} )
+				{
+					int3 TileOffset = {PerAttributeDataName}[In_AttributeIndex].xyz;
+					Out_Val = {Grid}.Load(int4(In_IndexX + TileOffset.x, In_IndexY + TileOffset.y, In_IndexZ + TileOffset.z, 0));
+				}
 			}
 		)");
 		OutHLSL += FString::Format(FormatBounds, ArgsBounds);
@@ -1420,10 +1379,12 @@ bool UNiagaraDataInterfaceGrid3DCollection::GetFunctionHLSL(const FNiagaraDataIn
 		static const TCHAR *FormatBounds = TEXT(R"(
 			void {FunctionName}(int In_IndexX, int In_IndexY, int In_IndexZ, int In_AttributeIndex, float In_Value, out int val)
 			{			
-				int3 TileOffset = {AttributePixelOffsetName}[In_AttributeIndex].xyz;
-
 				val = 0;
-				RW{OutputGrid}[int3(In_IndexX + TileOffset.x, In_IndexY + TileOffset.y, In_IndexZ + TileOffset.z)] = In_Value;
+				if ( In_AttributeIndex < {NumAttributesName} )
+				{
+					int3 TileOffset = {PerAttributeDataName}[In_AttributeIndex].xyz;
+					RW{OutputGrid}[int3(In_IndexX + TileOffset.x, In_IndexY + TileOffset.y, In_IndexZ + TileOffset.z)] = In_Value;
+				}
 			}
 		)");
 		OutHLSL += FString::Format(FormatBounds, ArgsBounds);
@@ -1448,7 +1409,7 @@ bool UNiagaraDataInterfaceGrid3DCollection::GetFunctionHLSL(const FNiagaraDataIn
 			{
 				for (int AttributeIndex = 0; AttributeIndex < {NumAttributes}.x; AttributeIndex++)
 				{		
-					int3 TileOffset = {AttributePixelOffsetName}[AttributeIndex].xyz;
+					int3 TileOffset = {PerAttributeDataName}[AttributeIndex].xyz;
 
 					float Val = {Grid}.Load(int4(In_IndexX + TileOffset.x, In_IndexY + TileOffset.y, In_IndexZ + TileOffset.z, 0));	
 					RW{OutputGrid}[int3(In_IndexX + TileOffset.x, In_IndexY + TileOffset.y, In_IndexZ + TileOffset.z)] = Val;
@@ -1467,7 +1428,7 @@ bool UNiagaraDataInterfaceGrid3DCollection::GetFunctionHLSL(const FNiagaraDataIn
 				{
 					float Val = 0.0f;
 
-					int3 TileOffset = {AttributePixelOffsetName}[AttributeIndex].xyz;
+					int3 TileOffset = {PerAttributeDataName}[AttributeIndex].xyz;
 
 					RW{OutputGrid}[int3(In_IndexX + TileOffset.x, In_IndexY + TileOffset.y, In_IndexZ + TileOffset.z)] = In_Value;
 				}
@@ -1561,9 +1522,15 @@ bool UNiagaraDataInterfaceGrid3DCollection::GetFunctionHLSL(const FNiagaraDataIn
 		static const TCHAR* FormatBounds = TEXT(R"(
 				void {FunctionName}(float In_UnitX, float In_UnitY, float In_UnitZ, int In_AttributeIndex, out float Out_Val)
 				{
-					float3 UVW = {AttributeUVOffsetName}[In_AttributeIndex].xyz + (float3(In_UnitX, In_UnitY, In_UnitZ) * {OneOverNumTiles});
-					UVW = clamp(UVW, {AttributeUVClampMinName}[In_AttributeIndex].xyz, {AttributeUVClampMaxName}[In_AttributeIndex].xyz);
-					Out_Val = {Grid}.SampleLevel({SamplerName}, UVW, 0);
+					Out_Val = 0;
+					if ( In_AttributeIndex < {NumAttributesName} )
+					{
+						float3 Unit = float3(In_UnitX, In_UnitY, In_UnitZ);
+						Unit = clamp(Unit, {UnitClampMin}, {UnitClampMax}) * {OneOverNumTiles};
+
+						float3 UVW = {PerAttributeDataName}[In_AttributeIndex + {NumAttributesName}].xyz + Unit;
+						Out_Val = {Grid}.SampleLevel({SamplerName}, UVW, 0);
+					}
 				}
 			)");
 		OutHLSL += FString::Format(FormatBounds, ArgsBounds);
@@ -2453,6 +2420,7 @@ bool UNiagaraDataInterfaceGrid3DCollection::PerInstanceTickPostSimulate(void* Pe
 			TargetData->Buffers.Empty();
 			TargetData->CurrentData = nullptr;
 			TargetData->DestinationData = nullptr;
+			TargetData->PerAttributeData.Release();
 		});
 
 	}
@@ -2572,11 +2540,36 @@ void FGrid3DCollectionRWInstanceData_RenderThread::EndSimulate(FRHICommandList&
 
 void FNiagaraDataInterfaceProxyGrid3DCollectionProxy::PreStage(FRHICommandList& RHICmdList, const FNiagaraDataInterfaceStageArgs& Context)
 {
+	FGrid3DCollectionRWInstanceData_RenderThread* ProxyData = SystemInstancesToProxyData_RT.Find(Context.SystemInstanceID);
+
+	// Generate per-attribute data
+	if (ProxyData->PerAttributeData.NumBytes == 0)
+	{
+		TResourceArray<FVector4> PerAttributeData;
+		PerAttributeData.AddUninitialized((ProxyData->TotalNumAttributes * 2) + 1);
+		for (int32 iAttribute = 0; iAttribute < ProxyData->TotalNumAttributes; ++iAttribute)
+		{
+			const FIntVector AttributeTileIndex(iAttribute % ProxyData->NumTiles.X, (iAttribute / ProxyData->NumTiles.X) % ProxyData->NumTiles.Y, iAttribute / (ProxyData->NumTiles.X * ProxyData->NumTiles.Y));
+			PerAttributeData[iAttribute] = FVector4(
+				AttributeTileIndex.X * ProxyData->NumCells.X,
+				AttributeTileIndex.Y * ProxyData->NumCells.Y,
+				AttributeTileIndex.Z * ProxyData->NumCells.Z,
+				0
+			);
+			PerAttributeData[iAttribute + ProxyData->TotalNumAttributes] = FVector4(
+				(1.0f / ProxyData->NumTiles.X) * float(AttributeTileIndex.X),
+				(1.0f / ProxyData->NumTiles.Y) * float(AttributeTileIndex.Y),
+				(1.0f / ProxyData->NumTiles.Z) * float(AttributeTileIndex.Z),
+				0.0f
+			);
+		}
+		PerAttributeData[ProxyData->TotalNumAttributes * 2] = FVector4(65535, 65535, 65535, 65535);
+		ProxyData->PerAttributeData.Initialize(TEXT("Grid3D::PerAttributeData"), sizeof(FVector4), PerAttributeData.Num(), EPixelFormat::PF_A32B32G32R32F, BUF_Static, &PerAttributeData);
+	}
+
 	// #todo(dmp): Context doesnt need to specify if a stage is output or not since we moved pre/post stage to the DI itself.  Not sure which design is better for the future
 	if (Context.IsOutputStage)
 	{
-		FGrid3DCollectionRWInstanceData_RenderThread* ProxyData = SystemInstancesToProxyData_RT.Find(Context.SystemInstanceID);
-
 		ProxyData->BeginSimulate(RHICmdList);
 
 		// If we don't have an iteration stage, then we should manually clear the buffer to make sure there is no residual data.  If we are doing something like rasterizing particles into a grid, we want it to be clear before
diff --git a/Engine/Plugins/FX/Niagara/Source/NiagaraCore/Private/NiagaraCustomVersion.cpp b/Engine/Plugins/FX/Niagara/Source/NiagaraCore/Private/NiagaraCustomVersion.cpp
--- a/Engine/Plugins/FX/Niagara/Source/NiagaraCore/Private/NiagaraCustomVersion.cpp
+++ b/Engine/Plugins/FX/Niagara/Source/NiagaraCore/Private/NiagaraCustomVersion.cpp
@@ -8,4 +8,4 @@ const FGuid FNiagaraCustomVersion::GUID(0xFCF57AFA, 0x50764283, 0xB9A9E658, 0xFF
 // Register the custom version with core
 FCustomVersionRegistration GRegisterNiagaraCustomVersion(FNiagaraCustomVersion::GUID, FNiagaraCustomVersion::LatestVersion, TEXT("NiagaraVer"));
 
-const FGuid FNiagaraCustomVersion::LatestScriptCompileVersion(0xB25900FB, 0xC1D24C63, 0x8D25CC89, 0x026FF460);
+const FGuid FNiagaraCustomVersion::LatestScriptCompileVersion(0x356395E7, 0x3682A948, 0x85C91E8A, 0x762086B9);
