Moved NiagaraRayTracingHelper to the NiagaraShader module so that it can be initialized in PostConfigInit (instead of PreDefault).

diff --git a/Engine/Plugins/FX/Niagara/Source/Niagara/Private/NiagaraRayTracingHelper.cpp b/Engine/Plugins/FX/Niagara/Source/Niagara/Private/NiagaraRayTracingHelper.cpp
deleted file mode 100644
diff --git a/Engine/Plugins/FX/Niagara/Source/Niagara/Private/NiagaraRayTracingHelper.h b/Engine/Plugins/FX/Niagara/Source/Niagara/Private/NiagaraRayTracingHelper.h
deleted file mode 100644
diff --git a/Engine/Plugins/FX/Niagara/Source/NiagaraShader/NiagaraShader.Build.cs b/Engine/Plugins/FX/Niagara/Source/NiagaraShader/NiagaraShader.Build.cs
--- a/Engine/Plugins/FX/Niagara/Source/NiagaraShader/NiagaraShader.Build.cs
+++ b/Engine/Plugins/FX/Niagara/Source/NiagaraShader/NiagaraShader.Build.cs
@@ -9,13 +9,13 @@ public class NiagaraShader : ModuleRules
         PrivateDependencyModuleNames.AddRange(
             new string[] {
                 "Core",
-                "Engine",
                 "CoreUObject",
-                "NiagaraCore"
+                "Engine",
+                "NiagaraCore",
+                "Renderer",
             }
         );
 
-
         PublicDependencyModuleNames.AddRange(
             new string[] {
                 "Core",
diff --git a/Engine/Plugins/FX/Niagara/Source/NiagaraShader/Private/NiagaraRayTracingHelper.cpp b/Engine/Plugins/FX/Niagara/Source/NiagaraShader/Private/NiagaraRayTracingHelper.cpp
new file mode 100644
--- /dev/null
+++ b/Engine/Plugins/FX/Niagara/Source/NiagaraShader/Private/NiagaraRayTracingHelper.cpp
@@ -0,0 +1,237 @@
+// Copyright Epic Games, Inc. All Rights Reserved.
+
+#include "NiagaraRayTracingHelper.h"
+
+#if RHI_RAYTRACING
+
+#include "GlobalShader.h"
+#include "MeshPassProcessor.h"
+#include "Renderer/Private/RayTracing/RayTracingScene.h"
+#include "Renderer/Private/ScenePrivate.h"
+#include "Renderer/Private/SceneRendering.h"
+#include "ShaderParameterStruct.h"
+
+/// TODO
+///  -get DispatchIndirect working where we can atomically increment the collision request index and then use that to
+///		control the dispatch index (or, at the least feed it into the RayGen shader to early out)
+///  -get geometry masking working when an environmental mask is implemented
+///  -template the CollisionDI to allow for multiple instances in the same emitter (even if that doesn't make a whole lot of sense)
+
+class FNiagaraCollisionRayTraceRG : public FGlobalShader
+{
+	DECLARE_GLOBAL_SHADER(FNiagaraCollisionRayTraceRG)
+	SHADER_USE_ROOT_PARAMETER_STRUCT(FNiagaraCollisionRayTraceRG, FGlobalShader)
+
+	BEGIN_SHADER_PARAMETER_STRUCT(FParameters, )
+		SHADER_PARAMETER_RDG_BUFFER_SRV(RaytracingAccelerationStructure, TLAS)
+		SHADER_PARAMETER_RDG_BUFFER_SRV(StructuredBuffer<FBasicRayData>, Rays)
+		SHADER_PARAMETER_RDG_BUFFER_UAV(RWStructuredBuffer<FNiagaraRayTracingPayload>, CollisionOutput)
+	END_SHADER_PARAMETER_STRUCT()
+
+	static bool ShouldCompilePermutation(const FGlobalShaderPermutationParameters& Parameters)
+	{
+		return ShouldCompileRayTracingShadersForProject(Parameters.Platform);
+	}
+};
+
+class FNiagaraCollisionRayTraceCH : public FGlobalShader
+{
+	DECLARE_SHADER_TYPE(FNiagaraCollisionRayTraceCH, Global);
+
+	static bool ShouldCompilePermutation(const FGlobalShaderPermutationParameters& Parameters)
+	{
+		return ShouldCompileRayTracingShadersForProject(Parameters.Platform);
+	}
+
+	FNiagaraCollisionRayTraceCH() = default;
+	FNiagaraCollisionRayTraceCH(const ShaderMetaType::CompiledShaderInitializerType& Initializer)
+		: FGlobalShader(Initializer)
+	{}
+};
+
+class FNiagaraCollisionRayTraceMiss : public FGlobalShader
+{
+	DECLARE_SHADER_TYPE(FNiagaraCollisionRayTraceMiss, Global);
+
+	static bool ShouldCompilePermutation(const FGlobalShaderPermutationParameters& Parameters)
+	{
+		return ShouldCompileRayTracingShadersForProject(Parameters.Platform);
+	}
+
+	FNiagaraCollisionRayTraceMiss() = default;
+	FNiagaraCollisionRayTraceMiss(const ShaderMetaType::CompiledShaderInitializerType& Initializer)
+		: FGlobalShader(Initializer)
+	{}
+};
+
+IMPLEMENT_GLOBAL_SHADER(FNiagaraCollisionRayTraceRG, "/Plugin/FX/Niagara/Private/NiagaraRayTracingShaders.usf", "NiagaraCollisionRayTraceRG", SF_RayGen);
+IMPLEMENT_GLOBAL_SHADER(FNiagaraCollisionRayTraceCH, "/Plugin/FX/Niagara/Private/NiagaraRayTracingShaders.usf", "NiagaraCollisionRayTraceCH", SF_RayHitGroup);
+IMPLEMENT_GLOBAL_SHADER(FNiagaraCollisionRayTraceMiss, "/Plugin/FX/Niagara/Private/NiagaraRayTracingShaders.usf", "NiagaraCollisionRayTraceMiss", SF_RayMiss);
+
+void FNiagaraRayTracingHelper::Reset()
+{
+	RayTracingPipelineState = nullptr;
+	RayTracingScene = nullptr;
+	RayTracingSceneView = nullptr;
+	ViewUniformBuffer = nullptr;
+}
+
+static FRayTracingPipelineState* CreateNiagaraRayTracingPipelineState(
+	EShaderPlatform Platform,
+	FRHICommandList& RHICmdList,
+	FRHIRayTracingShader* RayGenShader,
+	FRHIRayTracingShader* ClosestHitShader,
+	FRHIRayTracingShader* MissShader)
+{
+	FRayTracingPipelineStateInitializer Initializer;
+	Initializer.MaxPayloadSizeInBytes = sizeof(FNiagaraRayTracingPayload);
+
+	FRHIRayTracingShader* RayGenShaderTable[] = { RayGenShader };
+	Initializer.SetRayGenShaderTable(RayGenShaderTable);
+
+	FRHIRayTracingShader* HitGroupTable[] = { ClosestHitShader };
+	Initializer.SetHitGroupTable(HitGroupTable);
+
+	FRHIRayTracingShader* MissTable[] = { MissShader };
+	Initializer.SetMissShaderTable(MissTable);
+
+	Initializer.bAllowHitGroupIndexing = true; // Use the same hit shader for all geometry in the scene by disabling SBT indexing.
+
+	return PipelineStateCache::GetAndOrCreateRayTracingPipelineState(RHICmdList, Initializer);
+}
+
+static void BindNiagaraRayTracingMeshCommands(
+	FRHICommandList& RHICmdList,
+	FRayTracingSceneRHIRef RayTracingScene,
+	FRHIUniformBuffer* ViewUniformBuffer,
+	TConstArrayView<FVisibleRayTracingMeshCommand> RayTracingMeshCommands,
+	FRayTracingPipelineState* Pipeline,
+	TFunctionRef<uint32(const FRayTracingMeshCommand&)> PackUserData)
+{
+	const int32 NumTotalBindings = RayTracingMeshCommands.Num();
+
+	const uint32 MergedBindingsSize = sizeof(FRayTracingLocalShaderBindings) * NumTotalBindings;
+
+	const uint32 NumUniformBuffers = 1;
+	const uint32 UniformBufferArraySize = sizeof(FRHIUniformBuffer*) * NumUniformBuffers;
+
+	FRayTracingLocalShaderBindings* Bindings = nullptr;
+	FRHIUniformBuffer** UniformBufferArray = nullptr;
+
+	if (RHICmdList.Bypass())
+	{
+		FMemStack& MemStack = FMemStack::Get();
+
+		Bindings = reinterpret_cast<FRayTracingLocalShaderBindings*>(MemStack.Alloc(MergedBindingsSize, alignof(FRayTracingLocalShaderBindings)));
+		UniformBufferArray = reinterpret_cast<FRHIUniformBuffer**>(MemStack.Alloc(UniformBufferArraySize, alignof(FRHIUniformBuffer*)));
+	}
+	else
+	{
+		Bindings = reinterpret_cast<FRayTracingLocalShaderBindings*>(RHICmdList.Alloc(MergedBindingsSize, alignof(FRayTracingLocalShaderBindings)));
+		UniformBufferArray = reinterpret_cast<FRHIUniformBuffer**>(RHICmdList.Alloc(UniformBufferArraySize, alignof(FRHIUniformBuffer*)));
+	}
+
+	UniformBufferArray[0] = ViewUniformBuffer;
+
+	uint32 BindingIndex = 0;
+	for (const FVisibleRayTracingMeshCommand VisibleMeshCommand : RayTracingMeshCommands)
+	{
+		const FRayTracingMeshCommand& MeshCommand = *VisibleMeshCommand.RayTracingMeshCommand;
+
+		FRayTracingLocalShaderBindings Binding = {};
+		Binding.InstanceIndex = VisibleMeshCommand.InstanceIndex;
+		Binding.SegmentIndex = MeshCommand.GeometrySegmentIndex;
+		Binding.UserData = PackUserData(MeshCommand);
+		Binding.UniformBuffers = UniformBufferArray;
+		Binding.NumUniformBuffers = NumUniformBuffers;
+
+		Bindings[BindingIndex] = Binding;
+		BindingIndex++;
+	}
+
+	const bool bCopyDataToInlineStorage = false; // Storage is already allocated from RHICmdList, no extra copy necessary
+	RHICmdList.SetRayTracingHitGroups(
+		RayTracingScene,
+		Pipeline,
+		NumTotalBindings,
+		Bindings,
+		bCopyDataToInlineStorage);
+}
+
+void FNiagaraRayTracingHelper::BuildRayTracingSceneInfo(FRHICommandList& RHICmdList, TConstArrayView<FViewInfo> Views)
+{
+	check(Views.Num() > 0);
+
+	const FViewInfo& ReferenceView = Views[0];
+	const FScene* RenderScene = (ReferenceView.Family && ReferenceView.Family->Scene) ? ReferenceView.Family->Scene->GetRenderScene() : nullptr;
+
+	if (RenderScene && RenderScene->RayTracingScene.IsCreated())
+	{
+		RayTracingScene = RenderScene->RayTracingScene.GetRHIRayTracingSceneChecked();
+		RayTracingSceneView = RenderScene->RayTracingScene.GetShaderResourceViewChecked();
+		ViewUniformBuffer = ReferenceView.ViewUniformBuffer;
+
+		FGlobalShaderMap* ShaderMap = GetGlobalShaderMap(ShaderPlatform);
+		auto RayGenShader = ShaderMap->GetShader<FNiagaraCollisionRayTraceRG>().GetRayTracingShader();
+		auto ClosestHitShader = ShaderMap->GetShader<FNiagaraCollisionRayTraceCH>().GetRayTracingShader();
+		auto MissShader = ShaderMap->GetShader<FNiagaraCollisionRayTraceMiss>().GetRayTracingShader();
+
+		RayTracingPipelineState = CreateNiagaraRayTracingPipelineState(
+			ShaderPlatform,
+			RHICmdList,
+			RayGenShader,
+			ClosestHitShader,
+			MissShader);
+
+		// some options for what we want with our per MeshCommand user data.  For now we'll ignore it, but possibly
+		// something we'd want to incorporate.  Some examples could be if the material is translucent, or possibly the physical material?
+		auto BakeTranslucent = [&](const FRayTracingMeshCommand& MeshCommand) {	return (MeshCommand.bIsTranslucent != 0) & 0x1;	};
+		auto BakeDefault = [&](const FRayTracingMeshCommand& MeshCommand) { return 0; };
+
+		BindNiagaraRayTracingMeshCommands(
+			RHICmdList,
+			RayTracingScene,
+			ViewUniformBuffer,
+			ReferenceView.VisibleRayTracingMeshCommands,
+			RayTracingPipelineState,
+			BakeDefault);
+	}
+	else
+	{
+		Reset();
+	}
+}
+
+void FNiagaraRayTracingHelper::IssueRayTraces(FRHICommandList& RHICmdList, const FIntPoint& RayTraceCounts, FRHIShaderResourceView* RayTraceRequests, FRHIUnorderedAccessView* RayTraceResults) const
+{
+	check(RayTracingPipelineState);
+	check(RayTracingScene);
+	check(RayTracingSceneView);
+
+	FGlobalShaderMap* ShaderMap = GetGlobalShaderMap(ShaderPlatform);
+
+	FRayTracingShaderBindings GlobalBindings;
+	GlobalBindings.SRVs[0] = RayTracingSceneView;
+	GlobalBindings.SRVs[1] = RayTraceRequests;
+	GlobalBindings.UAVs[0] = RayTraceResults;
+
+	// todo - add support for RayTraceDispatchIndirect
+	RHICmdList.RayTraceDispatch(
+		RayTracingPipelineState,
+		ShaderMap->GetShader<FNiagaraCollisionRayTraceRG>().GetRayTracingShader(),
+		RayTracingScene,
+		GlobalBindings,
+		RayTraceCounts.X,
+		RayTraceCounts.Y
+	);
+}
+
+bool FNiagaraRayTracingHelper::IsValid() const
+{
+	return RayTracingPipelineState
+		&& RayTracingScene
+		&& RayTracingSceneView
+		&& ViewUniformBuffer;
+}
+
+#endif
\ No newline at end of file
diff --git a/Engine/Plugins/FX/Niagara/Source/NiagaraShader/Public/NiagaraRayTracingHelper.h b/Engine/Plugins/FX/Niagara/Source/NiagaraShader/Public/NiagaraRayTracingHelper.h
new file mode 100644
--- /dev/null
+++ b/Engine/Plugins/FX/Niagara/Source/NiagaraShader/Public/NiagaraRayTracingHelper.h
@@ -0,0 +1,49 @@
+// Copyright Epic Games, Inc. All Rights Reserved.
+#pragma once
+
+#include "CoreMinimal.h"
+
+#include "RHIDefinitions.h"
+
+#if RHI_RAYTRACING
+
+class FViewInfo;
+class FRayTracingPipelineState;
+class FRHICommandList;
+class FRHIRayTracingScene;
+class FRHIShaderResourceView;
+class FRHIUniformBuffer;
+class FRHIUnorderedAccessView;
+
+// c++ mirror of the struct defined in NiagaraRayTracingCommon.ush
+struct FNiagaraRayTracingPayload
+{
+	float	HitT;            // Distance from ray origin to the intersection point in the ray direction. Negative on miss.
+	float	WorldPosition[3];
+	float	WorldNormal[3];
+};
+
+class NIAGARASHADER_API FNiagaraRayTracingHelper
+{
+	public:
+		FNiagaraRayTracingHelper(EShaderPlatform InShaderPlatform)
+		: ShaderPlatform(InShaderPlatform)
+		{}
+
+		FNiagaraRayTracingHelper() = delete;
+
+		void Reset();
+		void BuildRayTracingSceneInfo(FRHICommandList& RHICmdList, TConstArrayView<FViewInfo> Views);
+		void IssueRayTraces(FRHICommandList& RHICmdList, const FIntPoint& RayTraceCounts, FRHIShaderResourceView* RayTraceRequests, FRHIUnorderedAccessView* RayTraceResults) const;
+		bool IsValid() const;
+
+	private:
+		const EShaderPlatform ShaderPlatform;
+
+		mutable FRayTracingPipelineState* RayTracingPipelineState = nullptr;
+		mutable FRHIRayTracingScene* RayTracingScene = nullptr;
+		mutable FRHIShaderResourceView* RayTracingSceneView = nullptr;
+		mutable FRHIUniformBuffer* ViewUniformBuffer = nullptr;
+};
+
+#endif
