Editgrate 16035091:
- Add a location where we can wait on async work before sending end of frame updates
- Hook Niagara into that location so any spawning started in PostActorTick will complete on the same frame

diff --git a/Engine/Plugins/FX/Niagara/Source/Niagara/Private/NiagaraWorldManager.cpp b/Engine/Plugins/FX/Niagara/Source/Niagara/Private/NiagaraWorldManager.cpp
--- a/Engine/Plugins/FX/Niagara/Source/Niagara/Private/NiagaraWorldManager.cpp
+++ b/Engine/Plugins/FX/Niagara/Source/Niagara/Private/NiagaraWorldManager.cpp
@@ -176,6 +176,7 @@ FDelegateHandle FNiagaraWorldManager::OnPostWorldCleanupHandle;
 FDelegateHandle FNiagaraWorldManager::OnPreWorldFinishDestroyHandle;
 FDelegateHandle FNiagaraWorldManager::OnWorldBeginTearDownHandle;
 FDelegateHandle FNiagaraWorldManager::TickWorldHandle;
+FDelegateHandle FNiagaraWorldManager::OnWorldPreSendAllEndOfFrameUpdatesHandle;
 FDelegateHandle FNiagaraWorldManager::PreGCHandle;
 FDelegateHandle FNiagaraWorldManager::PostReachabilityAnalysisHandle;
 FDelegateHandle FNiagaraWorldManager::PostGCHandle;
@@ -343,6 +344,15 @@ void FNiagaraWorldManager::OnStartup()
 	OnPreWorldFinishDestroyHandle = FWorldDelegates::OnPreWorldFinishDestroy.AddStatic(&FNiagaraWorldManager::OnPreWorldFinishDestroy);
 	OnWorldBeginTearDownHandle = FWorldDelegates::OnWorldBeginTearDown.AddStatic(&FNiagaraWorldManager::OnWorldBeginTearDown);
 	TickWorldHandle = FWorldDelegates::OnWorldPostActorTick.AddStatic(&FNiagaraWorldManager::TickWorld);
+	OnWorldPreSendAllEndOfFrameUpdatesHandle = FWorldDelegates::OnWorldPreSendAllEndOfFrameUpdates.AddLambda(
+		[](UWorld* InWorld)
+		{
+			if ( FNiagaraWorldManager* FoundManager = WorldManagers.FindRef(InWorld) )
+			{
+				FoundManager->PreSendAllEndOfFrameUpdates();
+			}
+		}
+	);
 
 	PreGCHandle = FCoreUObjectDelegates::GetPreGarbageCollectDelegate().AddStatic(&FNiagaraWorldManager::OnPreGarbageCollect);
 	PostReachabilityAnalysisHandle = FCoreUObjectDelegates::PostReachabilityAnalysis.AddStatic(&FNiagaraWorldManager::OnPostReachabilityAnalysis);
@@ -358,6 +368,7 @@ void FNiagaraWorldManager::OnShutdown()
 	FWorldDelegates::OnPreWorldFinishDestroy.Remove(OnPreWorldFinishDestroyHandle);
 	FWorldDelegates::OnWorldBeginTearDown.Remove(OnWorldBeginTearDownHandle);
 	FWorldDelegates::OnWorldPostActorTick.Remove(TickWorldHandle);
+	FWorldDelegates::OnWorldPreSendAllEndOfFrameUpdates.Remove(OnWorldPreSendAllEndOfFrameUpdatesHandle);
 
 	FCoreUObjectDelegates::GetPreGarbageCollectDelegate().Remove(PreGCHandle);
 	FCoreUObjectDelegates::PostReachabilityAnalysis.Remove(PostReachabilityAnalysisHandle);
@@ -745,6 +756,18 @@ void FNiagaraWorldManager::PostActorTick(float DeltaSeconds)
 	}
 }
 
+void FNiagaraWorldManager::PreSendAllEndOfFrameUpdates()
+{
+	for (const auto& Simulation : SimulationsWithPostActorWork)
+	{
+		if (Simulation->IsValid())
+		{
+			Simulation->WaitForInstancesTickComplete();
+		}
+	}
+	SimulationsWithPostActorWork.Reset();
+}
+
 void FNiagaraWorldManager::MarkSimulationForPostActorWork(FNiagaraSystemSimulation* SystemSimulation)
 {
 	check(SystemSimulation != nullptr);
diff --git a/Engine/Plugins/FX/Niagara/Source/Niagara/Public/NiagaraWorldManager.h b/Engine/Plugins/FX/Niagara/Source/Niagara/Public/NiagaraWorldManager.h
--- a/Engine/Plugins/FX/Niagara/Source/Niagara/Public/NiagaraWorldManager.h
+++ b/Engine/Plugins/FX/Niagara/Source/Niagara/Public/NiagaraWorldManager.h
@@ -150,6 +150,9 @@ public:
 	/** Called after all actor tick groups are complete. */
 	void PostActorTick(float DeltaSeconds);
 
+	/** Called before we run end of frame updates, allows us to wait on async work. */
+	void PreSendAllEndOfFrameUpdates();
+
 	void OnWorldCleanup(bool bSessionEnded, bool bCleanupResources);
 	void OnPostWorldCleanup(bool bSessionEnded, bool bCleanupResources);
 
@@ -278,6 +281,7 @@ private:
 	static FDelegateHandle OnPreWorldFinishDestroyHandle;
 	static FDelegateHandle OnWorldBeginTearDownHandle;
 	static FDelegateHandle TickWorldHandle;
+	static FDelegateHandle OnWorldPreSendAllEndOfFrameUpdatesHandle;
 	static FDelegateHandle PreGCHandle;
 	static FDelegateHandle PostReachabilityAnalysisHandle;
 	static FDelegateHandle PostGCHandle;
diff --git a/Engine/Source/Runtime/Engine/Classes/Engine/World.h b/Engine/Source/Runtime/Engine/Classes/Engine/World.h
--- a/Engine/Source/Runtime/Engine/Classes/Engine/World.h
+++ b/Engine/Source/Runtime/Engine/Classes/Engine/World.h
@@ -3920,6 +3920,9 @@ public:
 	DECLARE_MULTICAST_DELEGATE_ThreeParams(FOnWorldPostActorTick, UWorld* /*World*/, ELevelTick/**Tick Type*/, float/**Delta Seconds*/);
 	static FOnWorldPostActorTick OnWorldPostActorTick;
 
+	DECLARE_MULTICAST_DELEGATE_OneParam(FOnWorldPreSendAllEndOfFrameUpdates, UWorld* /*World*/);
+	static FOnWorldPreSendAllEndOfFrameUpdates OnWorldPreSendAllEndOfFrameUpdates;
+
 	// Callback for world creation
 	static FWorldEvent OnPostWorldCreation;
 	
diff --git a/Engine/Source/Runtime/Engine/Private/LevelTick.cpp b/Engine/Source/Runtime/Engine/Private/LevelTick.cpp
--- a/Engine/Source/Runtime/Engine/Private/LevelTick.cpp
+++ b/Engine/Source/Runtime/Engine/Private/LevelTick.cpp
@@ -1043,6 +1043,9 @@ void UWorld::SendAllEndOfFrameUpdates()
 	CSV_SCOPED_TIMING_STAT_EXCLUSIVE(EndOfFrameUpdates);
 	CSV_SCOPED_SET_WAIT_STAT(EndOfFrameUpdates);
 
+	// Allow systems to complete async work that could introduce additional components to end of frame updates
+	FWorldDelegates::OnWorldPreSendAllEndOfFrameUpdates.Broadcast(this);
+
 	if (!HasEndOfFrameUpdates())
 	{
 		return;
diff --git a/Engine/Source/Runtime/Engine/Private/World.cpp b/Engine/Source/Runtime/Engine/Private/World.cpp
--- a/Engine/Source/Runtime/Engine/Private/World.cpp
+++ b/Engine/Source/Runtime/Engine/Private/World.cpp
@@ -348,6 +348,7 @@ FWorldDelegates::FWorldGetAssetTags FWorldDelegates::GetAssetTags;
 FWorldDelegates::FOnWorldTickStart FWorldDelegates::OnWorldTickStart;
 FWorldDelegates::FOnWorldPreActorTick FWorldDelegates::OnWorldPreActorTick;
 FWorldDelegates::FOnWorldPostActorTick FWorldDelegates::OnWorldPostActorTick;
+FWorldDelegates::FOnWorldPreSendAllEndOfFrameUpdates FWorldDelegates::OnWorldPreSendAllEndOfFrameUpdates;
 FWorldDelegates::FWorldEvent FWorldDelegates::OnWorldBeginTearDown;
 #if WITH_EDITOR
 FWorldDelegates::FRefreshLevelScriptActionsEvent FWorldDelegates::RefreshLevelScriptActions;
