- Fix incorrect Engine.ExecutionCount on GPU systems using an iteration interface

diff --git a/Engine/Plugins/FX/Niagara/Source/NiagaraCore/Private/NiagaraCustomVersion.cpp b/Engine/Plugins/FX/Niagara/Source/NiagaraCore/Private/NiagaraCustomVersion.cpp
--- a/Engine/Plugins/FX/Niagara/Source/NiagaraCore/Private/NiagaraCustomVersion.cpp
+++ b/Engine/Plugins/FX/Niagara/Source/NiagaraCore/Private/NiagaraCustomVersion.cpp
@@ -8,4 +8,4 @@ const FGuid FNiagaraCustomVersion::GUID(0xFCF57AFA, 0x50764283, 0xB9A9E658, 0xFF
 // Register the custom version with core
 FCustomVersionRegistration GRegisterNiagaraCustomVersion(FNiagaraCustomVersion::GUID, FNiagaraCustomVersion::LatestVersion, TEXT("NiagaraVer"));
 
-const FGuid FNiagaraCustomVersion::LatestScriptCompileVersion(0xC5531FA3, 0x1FA74F4F, 0x8B391920, 0xD86B9FE3);
+const FGuid FNiagaraCustomVersion::LatestScriptCompileVersion(0x5E8FBE61, 0x53A1054B, 0xB543A46F, 0x6916AB4D);
diff --git a/Engine/Plugins/FX/Niagara/Source/NiagaraEditor/Private/NiagaraHlslTranslator.cpp b/Engine/Plugins/FX/Niagara/Source/NiagaraEditor/Private/NiagaraHlslTranslator.cpp
--- a/Engine/Plugins/FX/Niagara/Source/NiagaraEditor/Private/NiagaraHlslTranslator.cpp
+++ b/Engine/Plugins/FX/Niagara/Source/NiagaraEditor/Private/NiagaraHlslTranslator.cpp
@@ -3277,6 +3277,7 @@ void FHlslNiagaraTranslator::DefineMainGPUFunctions(
 			"	{\n"
 			"		bRunUpdateLogic = InstanceID < IterationInterfaceInstanceCount && GSimStart != 1;\n"
 			"		bRunSpawnLogic = InstanceID < IterationInterfaceInstanceCount && GSimStart == 1;\n"
+			"		GSpawnStartInstance = IterationInterfaceInstanceCount;\n"
 			"	}\n"
 			"	else\n"
 			"	#endif // USE_SIMULATION_STAGES\n"
