- Added GetParticleIndexFromIDTable which allows for stable zero based index access to particles if you only spawn and never kill from the source emitter

diff --git a/Engine/Plugins/FX/Niagara/Source/Niagara/Classes/NiagaraDataInterface.h b/Engine/Plugins/FX/Niagara/Source/Niagara/Classes/NiagaraDataInterface.h
--- a/Engine/Plugins/FX/Niagara/Source/Niagara/Classes/NiagaraDataInterface.h
+++ b/Engine/Plugins/FX/Niagara/Source/Niagara/Classes/NiagaraDataInterface.h
@@ -689,6 +689,14 @@ struct FNDIInputParam<FNiagaraBool>
 	FORCEINLINE bool GetAndAdvance() { return Data.GetAndAdvance().GetValue(); }
 };
 
+template<>
+struct FNDIInputParam<bool>
+{
+	VectorVM::FExternalFuncInputHandler<FNiagaraBool> Data;
+	FORCEINLINE FNDIInputParam(FVectorVMContext& Context) : Data(Context) {}
+	FORCEINLINE bool GetAndAdvance() { return Data.GetAndAdvance().GetValue(); }
+};
+
 template<>
 struct FNDIInputParam<FVector2D>
 {
@@ -769,6 +777,15 @@ struct FNDIOutputParam<FNiagaraBool>
 	FORCEINLINE void SetAndAdvance(bool Val) { Data.GetDestAndAdvance()->SetValue(Val); }
 };
 
+template<>
+struct FNDIOutputParam<bool>
+{
+	VectorVM::FExternalFuncRegisterHandler<FNiagaraBool> Data;
+	FORCEINLINE FNDIOutputParam(FVectorVMContext& Context) : Data(Context) {}
+	FORCEINLINE bool IsValid() const { return Data.IsValid(); }
+	FORCEINLINE void SetAndAdvance(bool Val) { Data.GetDestAndAdvance()->SetValue(Val); }
+};
+
 template<>
 struct FNDIOutputParam<FVector2D>
 {
diff --git a/Engine/Plugins/FX/Niagara/Source/Niagara/Classes/NiagaraDataInterfaceParticleRead.h b/Engine/Plugins/FX/Niagara/Source/Niagara/Classes/NiagaraDataInterfaceParticleRead.h
--- a/Engine/Plugins/FX/Niagara/Source/Niagara/Classes/NiagaraDataInterfaceParticleRead.h
+++ b/Engine/Plugins/FX/Niagara/Source/Niagara/Classes/NiagaraDataInterfaceParticleRead.h
@@ -53,6 +53,7 @@ public:
 	void GetSpawnedIDAtIndex(FVectorVMContext& Context);
 	void GetNumParticles(FVectorVMContext& Context);
 	void GetParticleIndex(FVectorVMContext& Context);
+	void GetParticleIndexFromIDTable(FVectorVMContext& Context);
 	void ReadInt(FVectorVMContext& Context, FName AttributeToRead);
 	void ReadBool(FVectorVMContext& Context, FName AttributeToRead);
 	void ReadFloat(FVectorVMContext& Context, FName AttributeToRead);
diff --git a/Engine/Plugins/FX/Niagara/Source/Niagara/Private/NiagaraDataInterfaceParticleRead.cpp b/Engine/Plugins/FX/Niagara/Source/Niagara/Private/NiagaraDataInterfaceParticleRead.cpp
--- a/Engine/Plugins/FX/Niagara/Source/Niagara/Private/NiagaraDataInterfaceParticleRead.cpp
+++ b/Engine/Plugins/FX/Niagara/Source/Niagara/Private/NiagaraDataInterfaceParticleRead.cpp
@@ -47,6 +47,8 @@ static const FName GetColorByIndexFunctionName("Get Color By Index");
 static const FName GetQuatByIndexFunctionName("Get Quaternion By Index");
 static const FName GetIDByIndexFunctionName("Get ID By Index");
 
+static const FName GetParticleIndexFromIDTableName("Get Particle Index From ID Table");
+
 static const FName ParticleReadIDName(TEXT("ID"));
 
 static const FString NumSpawnedParticlesBaseName(TEXT("NumSpawnedParticles_"));
@@ -765,6 +767,21 @@ void UNiagaraDataInterfaceParticleRead::GetPersistentIDFunctions(TArray<FNiagara
 		OutFunctions.Add(Sig);
 	}
 
+	{
+		FNiagaraFunctionSignature& Sig = OutFunctions.AddDefaulted_GetRef();
+		Sig.Name = GetParticleIndexFromIDTableName;
+		Sig.Inputs.Add(FNiagaraVariable(FNiagaraTypeDefinition(GetClass()), TEXT("Particle Reader")));
+		Sig.Inputs.Add(FNiagaraVariable(FNiagaraTypeDefinition::GetIntDef(), TEXT("ID Table Index")));
+		Sig.Outputs.Add(FNiagaraVariable(FNiagaraTypeDefinition::GetBoolDef(), TEXT("Valid")));
+		Sig.Outputs.Add(FNiagaraVariable(FNiagaraTypeDefinition::GetIntDef(), TEXT("Particle Index")));
+		Sig.bMemberFunction = true;
+		Sig.bRequiresContext = false;
+		Sig.bExperimental = true;
+#if WITH_EDITORONLY_DATA
+		Sig.Description = NSLOCTEXT("Niagara", "NiagaraDataInterfaceParticleRead_GetParticleIndexFromIDTable", "Returns the particle index from the ID table, or -1 if that index or particle is not valid.\nThe ID to Index table is guaranteed to give a stable particle index over a range should the source emitter only ever spawn and not kill particles.");
+#endif
+	}
+
 	//
 	// Get attribute by ID
 	//
@@ -1091,6 +1108,7 @@ DEFINE_NDI_DIRECT_FUNC_BINDER(UNiagaraDataInterfaceParticleRead, GetNumSpawnedPa
 DEFINE_NDI_DIRECT_FUNC_BINDER(UNiagaraDataInterfaceParticleRead, GetSpawnedIDAtIndex);
 DEFINE_NDI_DIRECT_FUNC_BINDER(UNiagaraDataInterfaceParticleRead, GetNumParticles);
 DEFINE_NDI_DIRECT_FUNC_BINDER(UNiagaraDataInterfaceParticleRead, GetParticleIndex);
+DEFINE_NDI_DIRECT_FUNC_BINDER(UNiagaraDataInterfaceParticleRead, GetParticleIndexFromIDTable);
 DEFINE_NDI_DIRECT_FUNC_BINDER_WITH_PAYLOAD(UNiagaraDataInterfaceParticleRead, ReadInt);
 DEFINE_NDI_DIRECT_FUNC_BINDER_WITH_PAYLOAD(UNiagaraDataInterfaceParticleRead, ReadBool);
 DEFINE_NDI_DIRECT_FUNC_BINDER_WITH_PAYLOAD(UNiagaraDataInterfaceParticleRead, ReadFloat);
@@ -1144,6 +1162,12 @@ void UNiagaraDataInterfaceParticleRead::GetVMExternalFunction(const FVMExternalF
 		return;
 	}
 
+	if (BindingInfo.Name == GetParticleIndexFromIDTableName)
+	{
+		NDI_FUNC_BINDER(UNiagaraDataInterfaceParticleRead, GetParticleIndexFromIDTable)::Bind(this, OutFunc);
+		return;
+	}
+
 	bool bBindSuccessful = false;
 	FNDIParticleRead_InstanceData* PIData = static_cast<FNDIParticleRead_InstanceData*>(InstanceData);
 	static const FName NAME_Attribute("Attribute");
@@ -1393,6 +1417,44 @@ void UNiagaraDataInterfaceParticleRead::GetParticleIndex(FVectorVMContext& Conte
 	}
 }
 
+void UNiagaraDataInterfaceParticleRead::GetParticleIndexFromIDTable(FVectorVMContext& Context)
+{
+	VectorVM::FUserPtrHandler<FNDIParticleRead_InstanceData> InstData(Context);
+	FNDIInputParam<int32> InIDTableIndex(Context);
+	FNDIOutputParam<bool> OutValid(Context);
+	FNDIOutputParam<int32> OutParticleIndex(Context);
+
+	const FNiagaraEmitterInstance* EmitterInstance = InstData.Get()->EmitterInstance;
+	const FNiagaraDataBuffer* CurrentData = EmitterInstance ? EmitterInstance->GetData().GetCurrentData() : nullptr;
+
+	if ( !CurrentData )
+	{
+		for ( int32 i=0; i < Context.NumInstances; ++i )
+		{
+			OutValid.SetAndAdvance(false);
+			OutParticleIndex.SetAndAdvance(INDEX_NONE);
+		}
+		return;
+	}
+
+	const TArray<int32>& IDTable = CurrentData->GetIDTable();
+	for (int32 i = 0; i < Context.NumInstances; ++i)
+	{
+		const int32 IDTableIndex = InIDTableIndex.GetAndAdvance();
+		if ( IDTable.IsValidIndex(IDTableIndex) )
+		{
+			const int32 ParticleIndex = IDTable[IDTableIndex];
+			OutValid.SetAndAdvance(ParticleIndex != INDEX_NONE);
+			OutParticleIndex.SetAndAdvance(ParticleIndex);
+		}
+		else
+		{
+			OutValid.SetAndAdvance(false);
+			OutParticleIndex.SetAndAdvance(INDEX_NONE);
+		}
+	}
+}
+
 //////////////////////////////////////////////////////////////////////////
 
 template<typename T>
@@ -1976,6 +2038,25 @@ bool UNiagaraDataInterfaceParticleRead::GetFunctionHLSL(const FNiagaraDataInterf
 		return true;
 	}
 
+	if (FunctionInfo.DefinitionName == GetParticleIndexFromIDTableName)
+	{
+		static const TCHAR* FuncTemplate = TEXT(
+			"void {FunctionName}(int IDTableIndex, out bool bValid, out int ParticleIndex)\n"
+			"{\n"
+			"    ParticleIndex = (IDTableIndex >= 0) ? {IDToIndexTableName}[IDTableIndex] : -1;\n"
+			"    bValid = (IDTableIndex >= 0) && (ParticleIndex != -1);\n"
+			"}\n\n"
+		);
+
+		TMap<FString, FStringFormatArg> FuncTemplateArgs;
+		FuncTemplateArgs.Add(TEXT("FunctionName"), FunctionInfo.InstanceName);
+		FuncTemplateArgs.Add(TEXT("IDToIndexTableName"), IDToIndexTableBaseName + ParamInfo.DataInterfaceHLSLSymbol);
+
+		OutHLSL += FString::Format(FuncTemplate, FuncTemplateArgs);
+		return true;
+	}
+
+
 	//
 	// Get attribute by ID
 	//
diff --git a/Engine/Plugins/FX/Niagara/Source/NiagaraEditor/Private/NiagaraHlslTranslator.cpp b/Engine/Plugins/FX/Niagara/Source/NiagaraEditor/Private/NiagaraHlslTranslator.cpp
--- a/Engine/Plugins/FX/Niagara/Source/NiagaraEditor/Private/NiagaraHlslTranslator.cpp
+++ b/Engine/Plugins/FX/Niagara/Source/NiagaraEditor/Private/NiagaraHlslTranslator.cpp
@@ -3129,13 +3129,13 @@ void FHlslNiagaraTranslator::DefineMainGPUFunctions(
 				ContextName = FString::Printf(TEXT("Context.%s."), *TranslationStages[i].PassNamespace);
 			}
 
-			HlslOutput += TEXT("\t\tif (bValid)\n\t\t{\n");
-
 			if (bRequiresPersistentIDs && !TranslationStages[i].bPartialParticleUpdate)
 			{
-				HlslOutput += FString::Printf(TEXT("\t\t\tUpdateID(0, %sParticles.ID.Index, WriteIndex);\n"), *ContextName);
+				HlslOutput += FString::Printf(TEXT("\t\tUpdateID(0, bValid ? %sParticles.ID.Index : -1, WriteIndex);\n"), *ContextName);
 			}
 
+			HlslOutput += TEXT("\t\tif (bValid)\n\t\t{\n");
+
 			for (int32 DataSetIndex = 0, IntCounter = 0, FloatCounter = 0, HalfCounter = 0; DataSetIndex < DataSetWrites.Num(); ++DataSetIndex)
 			{
 				const FNiagaraDataSetID DataSetID = ReadDataSetIDs[DataSetIndex];
