- Integrate and update CL 16013645 - Fix various Niagara data interfaces to work on mobile
- Ensure mobile scene textures are present for FXSystem to access

diff --git a/Engine/Plugins/FX/Niagara/Shaders/Private/NiagaraDataInterfaceCollisionQuery.ush b/Engine/Plugins/FX/Niagara/Shaders/Private/NiagaraDataInterfaceCollisionQuery.ush
--- a/Engine/Plugins/FX/Niagara/Shaders/Private/NiagaraDataInterfaceCollisionQuery.ush
+++ b/Engine/Plugins/FX/Niagara/Shaders/Private/NiagaraDataInterfaceCollisionQuery.ush
@@ -5,13 +5,37 @@
 #include "/Engine/Private/RayTracing/RayTracingCommon.ush"
 #include "NiagaraRayTracingCommon.ush"
 
+#define NDICOLLISIONQUERY_USE_GBUFFER_NORMAL	(FEATURE_LEVEL >= FEATURE_LEVEL_SM5)
+
 #if FEATURE_LEVEL >= FEATURE_LEVEL_SM5
 int											MaxRayTraceCount;
 RWStructuredBuffer<FBasicRayData>			RWRayRequests;
 StructuredBuffer<FNiagaraRayTracingPayload>	IntersectionResults;
 #endif
 
-void DICollisionQuery_SceneDepth(float3 In_SamplePos, out float Out_SceneDepth, out float3 Out_CameraPosWorld, out bool Out_IsInsideView, out float3 Out_WorldPos, out float3 Out_WorldNormal)
+float NDICollisionQuery_GetSceneDepth(float2 ScreenUV)
+{
+#if (FEATURE_LEVEL <= FEATURE_LEVEL_ES3_1)
+	#if MOBILE_DEFERRED_SHADING
+		return ConvertFromDeviceZ(Texture2DSample(MobileSceneTextures.SceneDepthTexture, MobileSceneTextures.SceneDepthTextureSampler, ScreenUV).r);
+	#else
+		return ConvertFromDeviceZ(Texture2DSample(MobileSceneTextures.SceneColorTexture, MobileSceneTextures.SceneColorTextureSampler, ScreenUV).a);
+	#endif
+#else
+		return CalcSceneDepth(ScreenUV);
+#endif
+}
+
+float NDICollisionQuery_GetCustomSceneDepth(float2 ScreenUV)
+{
+#if (FEATURE_LEVEL <= FEATURE_LEVEL_ES3_1)
+	return ConvertFromDeviceZ(Texture2DSampleLevel(MobileSceneTextures.CustomDepthTexture, MobileSceneTextures.CustomDepthTextureSampler, ScreenUV, 0).r);
+#else
+	return ConvertFromDeviceZ(Texture2DSampleLevel(SceneTexturesStruct.CustomDepthTexture, SceneTexturesStruct_SceneDepthTextureSampler, ScreenUV, 0).r);
+#endif
+}
+
+void NDICollisionQuery_QuerySceneDepthGPU(in float3 In_SamplePos, out float Out_SceneDepth, out float3 Out_CameraPosWorld, out bool Out_IsInsideView, out float3 Out_WorldPos, out float3 Out_WorldNormal)
 {
 	Out_SceneDepth = -1;
 	Out_WorldPos = float3(0.0, 0.0, 0.0);
@@ -19,30 +43,50 @@ void DICollisionQuery_SceneDepth(float3 In_SamplePos, out float Out_SceneDepth,
 	Out_IsInsideView = true;
 	Out_CameraPosWorld.xyz = View.WorldCameraOrigin.xyz;
 
-#if FEATURE_LEVEL >= FEATURE_LEVEL_SM5
 	float4 SamplePosition = float4(In_SamplePos + View.PreViewTranslation, 1);
 	float4 ClipPosition = mul(SamplePosition, View.TranslatedWorldToClip);
 	float2 ScreenPosition = ClipPosition.xy / ClipPosition.w;
+
 	// Check if the sample is inside the view.
 	if (all(abs(ScreenPosition.xy) <= float2(1, 1)))
 	{
 		// Sample the depth buffer to get a world position near the sample position.
 		float2 ScreenUV = ScreenPosition * View.ScreenPositionScaleBias.xy + View.ScreenPositionScaleBias.wz;
-		float SceneDepth = CalcSceneDepth(ScreenUV);
+		float SceneDepth = NDICollisionQuery_GetSceneDepth(ScreenUV);
+		float3 WorldPosition = WorldPositionFromSceneDepth(ScreenPosition.xy, SceneDepth);
+
+#if NDICOLLISIONQUERY_USE_GBUFFER_NORMAL
+		float3 WorldNormal = Texture2DSampleLevel(SceneTexturesStruct.GBufferATexture, SceneTexturesStruct_GBufferATextureSampler, ScreenUV, 0).xyz * 2.0 - 1.0;
+#else
+		float CollisionDepthBounds = 500.0f;
+		float SceneDepth0 = NDICollisionQuery_GetSceneDepth(ScreenUV + float2(View.BufferSizeAndInvSize.z, 0.0));
+		float SceneDepth1 = NDICollisionQuery_GetSceneDepth(ScreenUV + float2(0.0, View.BufferSizeAndInvSize.w));
+		// When using the forward shading, the normal of the pixel is approximated by the derivative of the world position
+		// of the pixel. But in on the visible edge this derivative can become very high, making CollisionPlane almost
+		// perpendicular to the view plane. In these case the particle may collide the visible edges of the diferent meshes
+		// in the view frustum. To avoid this, we disable the collision test if one of the derivate is above a threshold.
+		if (max(abs(SceneDepth - SceneDepth0), abs(SceneDepth - SceneDepth1)) > CollisionDepthBounds)
+		{
+			return;
+		}
+
+		float3 WorldPosition0 = WorldPositionFromSceneDepth(ScreenPosition.xy + float2(2 * View.ViewSizeAndInvSize.z, 0.0), SceneDepth0);
+		float3 WorldPosition1 = WorldPositionFromSceneDepth(ScreenPosition.xy - float2(0.0, 2 * View.ViewSizeAndInvSize.w), SceneDepth1);
+		float3 WorldNormal = normalize(cross(WorldPosition0 - WorldPosition, WorldPosition1 - WorldPosition));
+#endif
+
+		// Set outputs
 		Out_SceneDepth = SceneDepth;
-		// Reconstruct world position.
-		Out_WorldPos = WorldPositionFromSceneDepth(ScreenPosition.xy, SceneDepth);
-		// Sample the normal buffer
-		Out_WorldNormal = Texture2DSampleLevel(SceneTexturesStruct.GBufferATexture, SceneTexturesStruct_GBufferATextureSampler, ScreenUV, 0).xyz * 2.0 - 1.0;
+		Out_WorldPos = WorldPosition;
+		Out_WorldNormal = WorldNormal;
 	}
 	else
 	{
 		Out_IsInsideView = false;
 	}
-#endif
 }
 
-void DICollisionQuery_CustomDepth(float3 In_SamplePos, out float Out_SceneDepth, out float3 Out_CameraPosWorld, out bool Out_IsInsideView, out float3 Out_WorldPos, out float3 Out_WorldNormal)
+void NDICollisionQuery_QueryCustomDepthGPU(in float3 In_SamplePos, out float Out_SceneDepth, out float3 Out_CameraPosWorld, out bool Out_IsInsideView, out float3 Out_WorldPos, out float3 Out_WorldNormal)
 {
 	Out_SceneDepth = -1;
 	Out_WorldPos = float3(0.0, 0.0, 0.0);
@@ -50,30 +94,45 @@ void DICollisionQuery_CustomDepth(float3 In_SamplePos, out float Out_SceneDepth,
 	Out_IsInsideView = true;
 	Out_CameraPosWorld.xyz = View.WorldCameraOrigin.xyz;
 
-#if FEATURE_LEVEL >= FEATURE_LEVEL_SM5
 	float4 SamplePosition = float4(In_SamplePos + View.PreViewTranslation, 1);
 	float4 ClipPosition = mul(SamplePosition, View.TranslatedWorldToClip);
 	float2 ScreenPosition = ClipPosition.xy / ClipPosition.w;
+
 	// Check if the sample is inside the view.
 	if (all(abs(ScreenPosition.xy) <= float2(1, 1)))
 	{
 		// Sample the depth buffer to get a world position near the sample position.
 		float2 ScreenUV = ScreenPosition * View.ScreenPositionScaleBias.xy + View.ScreenPositionScaleBias.wz;
-		float SceneDepth = ConvertFromDeviceZ(Texture2DSampleLevel(SceneTexturesStruct.CustomDepthTexture, SceneTexturesStruct_SceneDepthTextureSampler, ScreenUV, 0).r);
-		Out_SceneDepth = SceneDepth;
-		// Reconstruct world position.
-		Out_WorldPos = WorldPositionFromSceneDepth(ScreenPosition.xy, SceneDepth);
-		// Sample the normal buffer
-		Out_WorldNormal = Texture2DSampleLevel(SceneTexturesStruct.GBufferATexture, SceneTexturesStruct_GBufferATextureSampler, ScreenUV, 0).xyz * 2.0 - 1.0;
+		float SceneDepth = NDICollisionQuery_GetCustomSceneDepth(ScreenUV);
+		float3 WorldPosition = WorldPositionFromSceneDepth(ScreenPosition.xy, SceneDepth);
+
+#if NDICOLLISIONQUERY_USE_GBUFFER_NORMAL
+		float3 WorldNormal = Texture2DSampleLevel(SceneTexturesStruct.GBufferATexture, SceneTexturesStruct_GBufferATextureSampler, ScreenUV, 0).xyz * 2.0 - 1.0;
+#else
+		float CollisionDepthBounds = 500.0f;
+		float SceneDepth0 = NDICollisionQuery_GetCustomSceneDepth(ScreenUV + float2(View.BufferSizeAndInvSize.z, 0.0));
+		float SceneDepth1 = NDICollisionQuery_GetCustomSceneDepth(ScreenUV + float2(0.0, View.BufferSizeAndInvSize.w));
+		// When using the forward shading, the normal of the pixel is approximated by the derivative of the world position
+		// of the pixel. But in on the visible edge this derivative can become very high, making CollisionPlane almost
+		// perpendicular to the view plane. In these case the particle may collide the visible edges of the diferent meshes
+		// in the view frustum. To avoid this, we disable the collision test if one of the derivate is above a threshold.
+		if (max(abs(SceneDepth - SceneDepth0), abs(SceneDepth - SceneDepth1)) > CollisionDepthBounds)
+		{
+			return;
+		}
+
+		float3 WorldPosition0 = WorldPositionFromSceneDepth(ScreenPosition.xy + float2(2 * View.ViewSizeAndInvSize.z, 0.0), SceneDepth0);
+		float3 WorldPosition1 = WorldPositionFromSceneDepth(ScreenPosition.xy - float2(0.0, 2 * View.ViewSizeAndInvSize.w), SceneDepth1);
+		float3 WorldNormal = normalize(cross(WorldPosition0 - WorldPosition, WorldPosition1 - WorldPosition));
+#endif
 	}
 	else
 	{
 		Out_IsInsideView = false;
 	}
-#endif
 }
 
-void DICollisionQuery_DistanceField(float3 In_SamplePos, out float Out_DistanceToNearestSurface, out float3 Out_FieldGradient, out bool Out_IsDistanceFieldValid)
+void NDICollisionQuery_QueryMeshDistanceFieldGPU(in float3 In_SamplePos, out float Out_DistanceToNearestSurface, out float3 Out_FieldGradient, out bool Out_IsDistanceFieldValid)
 {
 #if PLATFORM_SUPPORTS_DISTANCE_FIELDS && (FEATURE_LEVEL >= FEATURE_LEVEL_SM5)
 	Out_DistanceToNearestSurface = GetDistanceToNearestSurfaceGlobal(In_SamplePos);
@@ -86,7 +145,7 @@ void DICollisionQuery_DistanceField(float3 In_SamplePos, out float Out_DistanceT
 #endif
 }
 
-void DICollisionQuery_IssueAsyncRayTrace(int In_QueryID, float3 In_TraceStart, float3 In_TraceEnd, int In_TraceChannel, out bool Out_IsQueryValid)
+void NDICollisionQuery_IssueAsyncRayTrace(int In_QueryID, float3 In_TraceStart, float3 In_TraceEnd, int In_TraceChannel, out bool Out_IsQueryValid)
 {
 	Out_IsQueryValid = false;
 
@@ -105,7 +164,7 @@ void DICollisionQuery_IssueAsyncRayTrace(int In_QueryID, float3 In_TraceStart, f
 #endif
 }
 
-void DICollisionQuery_ReadAsyncRayTrace(int In_PreviousFrameQueryID, out bool Out_CollisionValid, out float Out_CollisionDistance, out float3 Out_CollisionPosWorld, out float3 Out_CollisionNormal)
+void NDICollisionQuery_ReadAsyncRayTrace(int In_PreviousFrameQueryID, out bool Out_CollisionValid, out float Out_CollisionDistance, out float3 Out_CollisionPosWorld, out float3 Out_CollisionNormal)
 {
 	Out_CollisionValid = false;
 	Out_CollisionDistance = 0.0f;
diff --git a/Engine/Plugins/FX/Niagara/Shaders/Private/NiagaraDataInterfaceCollisionQueryTemplate.ush b/Engine/Plugins/FX/Niagara/Shaders/Private/NiagaraDataInterfaceCollisionQueryTemplate.ush
new file mode 100644
--- /dev/null
+++ b/Engine/Plugins/FX/Niagara/Shaders/Private/NiagaraDataInterfaceCollisionQueryTemplate.ush
@@ -0,0 +1,26 @@
+// Copyright Epic Games, Inc. All Rights Reserved.
+
+void QuerySceneDepthGPU_{ParameterName}(in float3 In_SamplePos, out float Out_SceneDepth, out float3 Out_CameraPosWorld, out bool Out_IsInsideView, out float3 Out_WorldPos, out float3 Out_WorldNormal)
+{
+	NDICollisionQuery_QuerySceneDepthGPU(In_SamplePos, Out_SceneDepth, Out_CameraPosWorld, Out_IsInsideView, Out_WorldPos, Out_WorldNormal);
+}
+
+void QueryCustomDepthGPU_{ParameterName}(in float3 In_SamplePos, out float Out_SceneDepth, out float3 Out_CameraPosWorld, out bool Out_IsInsideView, out float3 Out_WorldPos, out float3 Out_WorldNormal)
+{
+	NDICollisionQuery_QueryCustomDepthGPU(In_SamplePos, Out_SceneDepth, Out_CameraPosWorld, Out_IsInsideView, Out_WorldPos, Out_WorldNormal);
+}
+
+void QueryMeshDistanceFieldGPU_{ParameterName}(in float3 In_SamplePos, out float Out_DistanceToNearestSurface, out float3 Out_FieldGradient, out bool Out_IsDistanceFieldValid)
+{
+	NDICollisionQuery_QueryMeshDistanceFieldGPU(In_SamplePos, Out_DistanceToNearestSurface, Out_FieldGradient, Out_IsDistanceFieldValid);
+}
+
+void IssueAsyncRayTraceGpu_{ParameterName}_UEImpureCall(int In_QueryID, float3 In_TraceStart, float3 In_TraceEnd, int In_TraceChannel, out bool Out_IsQueryValid)
+{
+	NDICollisionQuery_IssueAsyncRayTrace(In_QueryID, In_TraceStart, In_TraceEnd, In_TraceChannel, Out_IsQueryValid);
+}
+
+void ReadAsyncRayTraceGpu_{ParameterName}(int In_PreviousFrameQueryID, out bool Out_CollisionValid, out float Out_CollisionDistance, out float3 Out_CollisionPosWorld, out float3 Out_CollisionNormal)
+{
+	NDICollisionQuery_ReadAsyncRayTrace(In_PreviousFrameQueryID, Out_CollisionValid, Out_CollisionDistance, Out_CollisionPosWorld, Out_CollisionNormal);
+}
diff --git a/Engine/Plugins/FX/Niagara/Shaders/Private/NiagaraDataInterfaceGBuffer.ush b/Engine/Plugins/FX/Niagara/Shaders/Private/NiagaraDataInterfaceGBuffer.ush
--- a/Engine/Plugins/FX/Niagara/Shaders/Private/NiagaraDataInterfaceGBuffer.ush
+++ b/Engine/Plugins/FX/Niagara/Shaders/Private/NiagaraDataInterfaceGBuffer.ush
@@ -4,8 +4,9 @@
 
 #include "/Engine/Private/DeferredShadingCommon.ush"
 
-#define NDIGBUFFER_ENABLED				(FEATURE_LEVEL >= FEATURE_LEVEL_SM5)
-#define NDIGBUFFER_USE_GBUFFER_VELOCITY	0
+#define NDIGBUFFER_HAS_GBUFFER				(FEATURE_LEVEL >= FEATURE_LEVEL_SM5)
+#define NDIGBUFFER_USE_GBUFFER_VELOCITY		(false && NDIGBUFFER_HAS_GBUFFER)
+#define NDIGBUFFER_IS_MOBILE				(FEATURE_LEVEL <= FEATURE_LEVEL_ES3_1)
 
 ////////////////////////////////////////////////////////////////////////////////////////////////////////////////////////////////////////////////////////////////
 
@@ -14,7 +15,7 @@ SamplerState NDIGBuffer_VelocityTextureSampler;
 
 void DIGBuffer_DecodeDiffuseColor(in float2 ScreenUV, out bool IsValid, out float3 OutValue)
 {
-#if NDIGBUFFER_ENABLED
+#if NDIGBUFFER_HAS_GBUFFER
 	FGBufferData GBufferData = GetGBufferData(ScreenUV);
 	IsValid = true;
 	OutValue = GBufferData.DiffuseColor;
@@ -26,7 +27,7 @@ void DIGBuffer_DecodeDiffuseColor(in float2 ScreenUV, out bool IsValid, out floa
 
 void DIGBuffer_DecodeWorldNormal(in float2 ScreenUV, out bool IsValid, out float3 OutValue)
 {
-#if NDIGBUFFER_ENABLED
+#if NDIGBUFFER_HAS_GBUFFER
 	FGBufferData GBufferData = GetGBufferData(ScreenUV);
 	IsValid = true;
 	OutValue = GBufferData.WorldNormal;
@@ -40,7 +41,7 @@ void DIGBuffer_DecodeScreenVelocity(in float2 ScreenUV, out bool IsValid, out fl
 {
 	IsValid = false;
 	OutValue = 0;
-#if NDIGBUFFER_ENABLED
+#if NDIGBUFFER_HAS_GBUFFER
 	float4 EncodedVelocity = 0;
 	#if NDIGBUFFER_USE_GBUFFER_VELOCITY
 		EncodedVelocity = GetGBufferData(ScreenUV).Velocity;
@@ -63,7 +64,7 @@ void DIGBuffer_DecodeWorldVelocity(in float2 ScreenUV, out bool IsValid, out flo
 {
 	IsValid = false;
 	OutValue = 0;
-#if NDIGBUFFER_ENABLED
+#if NDIGBUFFER_HAS_GBUFFER
 	float4 EncodedVelocity = 0;
 	#if NDIGBUFFER_USE_GBUFFER_VELOCITY
 		EncodedVelocity = GetGBufferData(ScreenUV).Velocity;
@@ -101,7 +102,7 @@ void DIGBuffer_DecodeWorldVelocity(in float2 ScreenUV, out bool IsValid, out flo
 
 void DIGBuffer_DecodeBaseColor(in float2 ScreenUV, out bool IsValid, out float3 OutValue)
 {
-#if NDIGBUFFER_ENABLED
+#if NDIGBUFFER_HAS_GBUFFER
 	FGBufferData GBufferData = GetGBufferData(ScreenUV);
 	IsValid = true;
 	OutValue = GBufferData.BaseColor;
@@ -113,7 +114,7 @@ void DIGBuffer_DecodeBaseColor(in float2 ScreenUV, out bool IsValid, out float3
 
 void DIGBuffer_DecodeMetallic(in float2 ScreenUV, out bool IsValid, out float OutValue)
 {
-#if NDIGBUFFER_ENABLED
+#if NDIGBUFFER_HAS_GBUFFER
 	FGBufferData GBufferData = GetGBufferData(ScreenUV);
 	IsValid = true;
 	OutValue = GBufferData.Metallic;
@@ -125,7 +126,7 @@ void DIGBuffer_DecodeMetallic(in float2 ScreenUV, out bool IsValid, out float Ou
 
 void DIGBuffer_DecodeSpecular(in float2 ScreenUV, out bool IsValid, out float OutValue)
 {
-#if NDIGBUFFER_ENABLED
+#if NDIGBUFFER_HAS_GBUFFER
 	FGBufferData GBufferData = GetGBufferData(ScreenUV);
 	IsValid = true;
 	OutValue = GBufferData.Specular;
@@ -137,7 +138,7 @@ void DIGBuffer_DecodeSpecular(in float2 ScreenUV, out bool IsValid, out float Ou
 
 void DIGBuffer_DecodeRoughness(in float2 ScreenUV, out bool IsValid, out float OutValue)
 {
-#if NDIGBUFFER_ENABLED
+#if NDIGBUFFER_HAS_GBUFFER
 	FGBufferData GBufferData = GetGBufferData(ScreenUV);
 	IsValid = true;
 	OutValue = GBufferData.Roughness;
@@ -149,10 +150,17 @@ void DIGBuffer_DecodeRoughness(in float2 ScreenUV, out bool IsValid, out float O
 
 void DIGBuffer_DecodeDepth(in float2 ScreenUV, out bool IsValid, out float OutValue)
 {
-#if NDIGBUFFER_ENABLED
+#if NDIGBUFFER_HAS_GBUFFER
 	FGBufferData GBufferData = GetGBufferData(ScreenUV);
 	IsValid = true;
 	OutValue = GBufferData.Depth;
+#elif NDIGBUFFER_IS_MOBILE
+	IsValid = true;
+	#if MOBILE_DEFERRED_SHADING
+		OutValue = ConvertFromDeviceZ(Texture2DSample(MobileSceneTextures.SceneDepthTexture, MobileSceneTextures.SceneDepthTextureSampler, ScreenUV).r);
+	#else
+		OutValue = ConvertFromDeviceZ(Texture2DSample(MobileSceneTextures.SceneColorTexture, MobileSceneTextures.SceneColorTextureSampler, ScreenUV).a);
+	#endif
 #else
 	IsValid = false;
 	OutValue = 0;
@@ -161,10 +169,13 @@ void DIGBuffer_DecodeDepth(in float2 ScreenUV, out bool IsValid, out float OutVa
 
 void DIGBuffer_DecodeCustomDepth(in float2 ScreenUV, out bool IsValid, out float OutValue)
 {
-#if NDIGBUFFER_ENABLED
+#if NDIGBUFFER_HAS_GBUFFER
 	FGBufferData GBufferData = GetGBufferData(ScreenUV);
 	IsValid = true;
 	OutValue = GBufferData.CustomDepth;
+#elif NDIGBUFFER_IS_MOBILE
+	IsValid = true;
+	OutValue = ConvertFromDeviceZ(Texture2DSampleLevel(MobileSceneTextures.CustomDepthTexture, MobileSceneTextures.CustomDepthTextureSampler, ScreenUV, 0).r);
 #else
 	IsValid = false;
 	OutValue = 0;
@@ -173,10 +184,13 @@ void DIGBuffer_DecodeCustomDepth(in float2 ScreenUV, out bool IsValid, out float
 
 void DIGBuffer_DecodeCustomStencil(in float2 ScreenUV, out bool IsValid, out int OutValue)
 {
-#if NDIGBUFFER_ENABLED
+#if NDIGBUFFER_HAS_GBUFFER
 	FGBufferData GBufferData = GetGBufferData(ScreenUV);
 	IsValid = true;
 	OutValue = GBufferData.CustomStencil;
+#elif NDIGBUFFER_IS_MOBILE
+	IsValid = true;
+	OutValue = Texture2DSampleLevel(MobileSceneTextures.MobileCustomStencilTexture, MobileSceneTextures.MobileCustomStencilTextureSampler, ScreenUV, 0).r * 255.0f;
 #else
 	IsValid = false;
 	OutValue = 0;
@@ -185,10 +199,14 @@ void DIGBuffer_DecodeCustomStencil(in float2 ScreenUV, out bool IsValid, out int
 
 void DIGBuffer_DecodeSceneColor(in float2 ScreenUV, out bool IsValid, out float4 OutValue)
 {
-#if NDIGBUFFER_ENABLED
+#if NDIGBUFFER_HAS_GBUFFER
 	IsValid = true;
 	OutValue = Texture2DSampleLevel(SceneTexturesStruct_SceneColorTexture, SceneTexturesStruct_SceneColorTextureSampler, ScreenUV, 0);
 	OutValue.rgb *= View.OneOverPreExposure.xxx;
+#elif NDIGBUFFER_IS_MOBILE
+	IsValid = true;
+	OutValue = Texture2DSampleLevel(MobileSceneTextures.SceneColorTexture, MobileSceneTextures.SceneColorTextureSampler, ScreenUV, 0);
+	OutValue.rgb *= View.OneOverPreExposure.xxx;
 #else
 	IsValid = false;
 	OutValue = 0;
diff --git a/Engine/Plugins/FX/Niagara/Shaders/Private/NiagaraDataInterfaceOcclusion.ush b/Engine/Plugins/FX/Niagara/Shaders/Private/NiagaraDataInterfaceOcclusion.ush
--- a/Engine/Plugins/FX/Niagara/Shaders/Private/NiagaraDataInterfaceOcclusion.ush
+++ b/Engine/Plugins/FX/Niagara/Shaders/Private/NiagaraDataInterfaceOcclusion.ush
@@ -4,9 +4,21 @@
 
 #include "/Engine/Private/DeferredShadingCommon.ush"
 
+float DIOcclusion_CalcSceneDepth(float2 ScreenUV)
+{
+#if (FEATURE_LEVEL <= FEATURE_LEVEL_ES3_1)
+	#if MOBILE_DEFERRED_SHADING
+		return ConvertFromDeviceZ(Texture2DSample(MobileSceneTextures.SceneDepthTexture, MobileSceneTextures.SceneDepthTextureSampler, ScreenUV).r);
+	#else
+		return ConvertFromDeviceZ(Texture2DSample(MobileSceneTextures.SceneColorTexture, MobileSceneTextures.SceneColorTextureSampler, ScreenUV).a);
+	#endif
+#else
+		return CalcSceneDepth(ScreenUV);
+#endif
+}
+
 void DIOcclusion_Rectangle(in float3 In_SampleCenterWorldPos, in float In_SampleWindowWidthWorld, in float In_SampleWindowHeightWorld, in float In_SampleSteps, out float Out_VisibilityFraction, out float Out_SampleFraction)
 {
-#if FEATURE_LEVEL >= FEATURE_LEVEL_SM5
 	// we need to project the camera vector to a plane to compare it with scene depth
 	float CameraDistance = abs(dot(In_SampleCenterWorldPos.xyz - View.WorldViewOrigin.xyz, View.ViewForward.xyz));
 	float4 SamplePosition = float4(In_SampleCenterWorldPos + View.PreViewTranslation, 1);
@@ -44,7 +56,7 @@ void DIOcclusion_Rectangle(in float3 In_SampleCenterWorldPos, in float In_Sample
 					continue;
 				}
 				
-				float Depth = CalcSceneDepth(float2(SampleX, SampleY));
+				float Depth = DIOcclusion_CalcSceneDepth(float2(SampleX, SampleY));
 				if (Depth < CameraDistance) 
 				{
 					OccludedSamples++;
@@ -55,15 +67,10 @@ void DIOcclusion_Rectangle(in float3 In_SampleCenterWorldPos, in float In_Sample
 	}
 	Out_VisibilityFraction = TotalSamples > 0 ? 1 - OccludedSamples / TotalSamples : 0;
 	Out_SampleFraction = Steps == 0 ? 0 : (TotalSamples / (Steps * Steps));
-#else
-	Out_VisibilityFraction = 1;
-	Out_SampleFraction = 1;
-#endif
 }
 
 void DIOcclusion_Circle(in float3 In_SampleCenterWorldPos, in float In_SampleWindowDiameterWorld, in float In_SampleRays, in float In_SampleStepsPerRay, out float Out_VisibilityFraction, out float Out_SampleFraction)
 {
-#if FEATURE_LEVEL >= FEATURE_LEVEL_SM5
 	const float PI = 3.14159265;
 	const float SPIRAL_TURN = 2 * PI * 0.61803399; // use golden ratio to rotate sample pattern each ring so we get a spiral
 	float CameraDistance = abs(dot(In_SampleCenterWorldPos.xyz - View.WorldViewOrigin.xyz, View.ViewForward.xyz)); // we need to project the camera vector to a plane to compare it with scene depth
@@ -79,7 +86,7 @@ void DIOcclusion_Circle(in float3 In_SampleCenterWorldPos, in float In_SampleWin
 
 	if (ScreenUV.x <= 1 && ScreenUV.x >= 0 && ScreenUV.y <= 1 && ScreenUV.y >= 0)
 	{
-		float Depth = CalcSceneDepth(ScreenUV);
+		float Depth = DIOcclusion_CalcSceneDepth(ScreenUV);
 		if (Depth < CameraDistance) 
 		{
 			OccludedSamples++;
@@ -124,8 +131,4 @@ void DIOcclusion_Circle(in float3 In_SampleCenterWorldPos, in float In_SampleWin
 	}
 	Out_VisibilityFraction = TotalSamples > 0 ? 1 - OccludedSamples / TotalSamples : 0;
 	Out_SampleFraction = Steps == 0 ? 0 : (TotalSamples / (Rays * Steps + 1));
-#else
-	Out_VisibilityFraction = 1;
-	Out_SampleFraction = 1;
-#endif
 }
diff --git a/Engine/Plugins/FX/Niagara/Source/Niagara/Classes/NiagaraDataInterfaceCollisionQuery.h b/Engine/Plugins/FX/Niagara/Source/Niagara/Classes/NiagaraDataInterfaceCollisionQuery.h
--- a/Engine/Plugins/FX/Niagara/Source/Niagara/Classes/NiagaraDataInterfaceCollisionQuery.h
+++ b/Engine/Plugins/FX/Niagara/Source/Niagara/Classes/NiagaraDataInterfaceCollisionQuery.h
@@ -70,8 +70,8 @@ public:
 
 #if WITH_EDITORONLY_DATA
 	virtual bool AppendCompileHash(FNiagaraCompileHashVisitor* InVisitor) const override;
+	virtual void GetCommonHLSL(FString& OutHLSL) override;
 	virtual void GetParameterDefinitionHLSL(const FNiagaraDataInterfaceGPUParamInfo& ParamInfo, FString& OutHLSL) override;
-	virtual void GetCommonHLSL(FString& OutHlsl) override;
 	virtual bool GetFunctionHLSL(const FNiagaraDataInterfaceGPUParamInfo& ParamInfo, const FNiagaraDataInterfaceGeneratedFunction& FunctionInfo, int FunctionInstanceIndex, FString& OutHLSL) override;
 	virtual bool UpgradeFunctionCall(FNiagaraFunctionSignature& FunctionSignature) override;
 #endif
@@ -92,12 +92,4 @@ private:
 
 	static FCriticalSection CriticalSection;
 	UEnum* TraceChannelEnum;
-
-	const static FName SceneDepthName;
-	const static FName CustomDepthName;
-	const static FName DistanceFieldName;
-	const static FName IssueAsyncRayTraceName;
-	const static FName ReadAsyncRayTraceName;
-	const static FName SyncTraceName;
-	const static FName AsyncTraceName;
-};
\ No newline at end of file
+};
diff --git a/Engine/Plugins/FX/Niagara/Source/Niagara/Classes/NiagaraDataInterfaceOcclusion.h b/Engine/Plugins/FX/Niagara/Source/Niagara/Classes/NiagaraDataInterfaceOcclusion.h
--- a/Engine/Plugins/FX/Niagara/Source/Niagara/Classes/NiagaraDataInterfaceOcclusion.h
+++ b/Engine/Plugins/FX/Niagara/Source/Niagara/Classes/NiagaraDataInterfaceOcclusion.h
@@ -12,9 +12,6 @@ class NIAGARA_API UNiagaraDataInterfaceOcclusion : public UNiagaraDataInterface
 	GENERATED_UCLASS_BODY()
 
 public:
-
-	DECLARE_NIAGARA_DI_PARAMETER();
-
 	//UObject Interface
 	virtual void PostInitProperties() override;
 	//UObject Interface End
diff --git a/Engine/Plugins/FX/Niagara/Source/Niagara/Private/NiagaraDataInterfaceCollisionQuery.cpp b/Engine/Plugins/FX/Niagara/Source/Niagara/Private/NiagaraDataInterfaceCollisionQuery.cpp
--- a/Engine/Plugins/FX/Niagara/Source/Niagara/Private/NiagaraDataInterfaceCollisionQuery.cpp
+++ b/Engine/Plugins/FX/Niagara/Source/Niagara/Private/NiagaraDataInterfaceCollisionQuery.cpp
@@ -13,10 +13,25 @@
 #include "RenderResource.h"
 #include "Shader.h"
 #include "ShaderCore.h"
+#include "ShaderCompilerCore.h"
 #include "ShaderParameterUtils.h"
 
 #define LOCTEXT_NAMESPACE "NiagaraDataInterfaceCollisionQuery"
 
+namespace NDICollisionQueryLocal
+{
+	static const TCHAR* CommonShaderFile = TEXT("/Plugin/FX/Niagara/Private/NiagaraDataInterfaceCollisionQuery.ush");
+	static const TCHAR* TemplateShaderFile = TEXT("/Plugin/FX/Niagara/Private/NiagaraDataInterfaceCollisionQueryTemplate.ush");
+
+	static const FName SceneDepthName(TEXT("QuerySceneDepthGPU"));
+	static const FName CustomDepthName(TEXT("QueryCustomDepthGPU"));
+	static const FName DistanceFieldName(TEXT("QueryMeshDistanceFieldGPU"));
+	static const FName SyncTraceName(TEXT("PerformCollisionQuerySyncCPU"));
+	static const FName AsyncTraceName(TEXT("PerformCollisionQueryAsyncCPU"));
+	static const FName IssueAsyncRayTraceName(TEXT("IssueAsyncRayTraceGpu"));
+	static const FName ReadAsyncRayTraceName(TEXT("ReadAsyncRayTraceGpu"));
+}
+
 FCriticalSection UNiagaraDataInterfaceCollisionQuery::CriticalSection;
 
 struct FNiagaraCollisionDIFunctionVersion
@@ -32,7 +47,6 @@ struct FNiagaraCollisionDIFunctionVersion
 		LatestVersion = VersionPlusOne - 1
 	};
 };
-
 struct FNiagaraDataIntefaceProxyCollisionQuery : public FNiagaraDataInterfaceProxy
 {
 #if RHI_RAYTRACING
@@ -122,13 +136,6 @@ struct FNiagaraDataIntefaceProxyCollisionQuery : public FNiagaraDataInterfacePro
 	}
 };
 
-const FName UNiagaraDataInterfaceCollisionQuery::SceneDepthName(TEXT("QuerySceneDepthGPU"));
-const FName UNiagaraDataInterfaceCollisionQuery::CustomDepthName(TEXT("QueryCustomDepthGPU"));
-const FName UNiagaraDataInterfaceCollisionQuery::DistanceFieldName(TEXT("QueryMeshDistanceFieldGPU"));
-const FName UNiagaraDataInterfaceCollisionQuery::IssueAsyncRayTraceName(TEXT("IssueAsyncRayTraceGpu"));
-const FName UNiagaraDataInterfaceCollisionQuery::ReadAsyncRayTraceName(TEXT("ReadAsyncRayTraceGpu"));
-const FName UNiagaraDataInterfaceCollisionQuery::SyncTraceName(TEXT("PerformCollisionQuerySyncCPU"));
-const FName UNiagaraDataInterfaceCollisionQuery::AsyncTraceName(TEXT("PerformCollisionQueryAsyncCPU"));
 
 UNiagaraDataInterfaceCollisionQuery::UNiagaraDataInterfaceCollisionQuery(FObjectInitializer const& ObjectInitializer)
 	: Super(ObjectInitializer)
@@ -234,7 +241,7 @@ void UNiagaraDataInterfaceCollisionQuery::GetAssetTagsForContext(const UObject*
 					{
 						for (const auto& Func : Info.RegisteredFunctions)
 						{
-							if (Func.Name == SyncTraceName || Func.Name == AsyncTraceName)
+							if (Func.Name == NDICollisionQueryLocal::SyncTraceName || Func.Name == NDICollisionQueryLocal::AsyncTraceName)
 							{
 								return true;
 							}
@@ -264,7 +271,7 @@ void UNiagaraDataInterfaceCollisionQuery::GetAssetTagsForContext(const UObject*
 void UNiagaraDataInterfaceCollisionQuery::GetFunctions(TArray<FNiagaraFunctionSignature>& OutFunctions)
 {
 	FNiagaraFunctionSignature SigDepth;
-	SigDepth.Name = UNiagaraDataInterfaceCollisionQuery::SceneDepthName;
+	SigDepth.Name = NDICollisionQueryLocal::SceneDepthName;
 	SigDepth.bMemberFunction = true;
 	SigDepth.bRequiresContext = false;
 	SigDepth.bSupportsCPU = false;
@@ -290,7 +297,7 @@ void UNiagaraDataInterfaceCollisionQuery::GetFunctions(TArray<FNiagaraFunctionSi
 	OutFunctions.Add(SigDepth);
 
 	FNiagaraFunctionSignature SigCustomDepth;
-	SigCustomDepth.Name = UNiagaraDataInterfaceCollisionQuery::CustomDepthName;
+	SigCustomDepth.Name = NDICollisionQueryLocal::CustomDepthName;
 	SigCustomDepth.bMemberFunction = true;
 	SigCustomDepth.bRequiresContext = false;
 	SigCustomDepth.bSupportsCPU = false;
@@ -308,7 +315,7 @@ void UNiagaraDataInterfaceCollisionQuery::GetFunctions(TArray<FNiagaraFunctionSi
 	OutFunctions.Add(SigCustomDepth);
 
 	FNiagaraFunctionSignature SigMeshField;
-	SigMeshField.Name = UNiagaraDataInterfaceCollisionQuery::DistanceFieldName;
+	SigMeshField.Name = NDICollisionQueryLocal::DistanceFieldName;
 	SigMeshField.bMemberFunction = true;
 	SigMeshField.bRequiresContext = false;
 	SigMeshField.bSupportsCPU = false;
@@ -336,7 +343,7 @@ void UNiagaraDataInterfaceCollisionQuery::GetFunctions(TArray<FNiagaraFunctionSi
 		const FText IsQueryValidDescription = LOCTEXT("IsQueryValidDescription", "Returns true if the query was enqueued");
 
 		FNiagaraFunctionSignature& IssueRayTrace = OutFunctions.AddDefaulted_GetRef();
-		IssueRayTrace.Name = UNiagaraDataInterfaceCollisionQuery::IssueAsyncRayTraceName;
+		IssueRayTrace.Name = NDICollisionQueryLocal::IssueAsyncRayTraceName;
 		IssueRayTrace.bRequiresExecPin = true;
 		IssueRayTrace.bMemberFunction = true;
 		IssueRayTrace.bRequiresContext = false;
@@ -361,7 +368,7 @@ void UNiagaraDataInterfaceCollisionQuery::GetFunctions(TArray<FNiagaraFunctionSi
 		const FText CollisionNormalDescription = LOCTEXT("CollisionNormalDescription", "The surface normal of the world geometry at the point of intersection");
 
 		FNiagaraFunctionSignature& ReadRayTrace = OutFunctions.AddDefaulted_GetRef();
-		ReadRayTrace.Name = UNiagaraDataInterfaceCollisionQuery::ReadAsyncRayTraceName;
+		ReadRayTrace.Name = NDICollisionQueryLocal::ReadAsyncRayTraceName;
 		ReadRayTrace.bMemberFunction = true;
 		ReadRayTrace.bRequiresContext = false;
 		ReadRayTrace.bSupportsCPU = false;
@@ -378,7 +385,7 @@ void UNiagaraDataInterfaceCollisionQuery::GetFunctions(TArray<FNiagaraFunctionSi
 	}
 
 	FNiagaraFunctionSignature SigCpuSync;
-	SigCpuSync.Name = UNiagaraDataInterfaceCollisionQuery::SyncTraceName;
+	SigCpuSync.Name = NDICollisionQueryLocal::SyncTraceName;
 	SigCpuSync.bMemberFunction = true;
 	SigCpuSync.bRequiresContext = false;
 	SigCpuSync.bSupportsGPU = false;
@@ -413,7 +420,7 @@ void UNiagaraDataInterfaceCollisionQuery::GetFunctions(TArray<FNiagaraFunctionSi
 	OutFunctions.Add(SigCpuSync);
 
 	FNiagaraFunctionSignature SigCpuAsync;
-	SigCpuAsync.Name = UNiagaraDataInterfaceCollisionQuery::AsyncTraceName;
+	SigCpuAsync.Name = NDICollisionQueryLocal::AsyncTraceName;
 	SigCpuAsync.bMemberFunction = true;
 	SigCpuAsync.bRequiresContext = false;
 	SigCpuAsync.bSupportsGPU = false;
@@ -449,70 +456,16 @@ void UNiagaraDataInterfaceCollisionQuery::GetFunctions(TArray<FNiagaraFunctionSi
 #if WITH_EDITORONLY_DATA
 bool UNiagaraDataInterfaceCollisionQuery::GetFunctionHLSL(const FNiagaraDataInterfaceGPUParamInfo& ParamInfo, const FNiagaraDataInterfaceGeneratedFunction& FunctionInfo, int FunctionInstanceIndex, FString& OutHLSL)
 {
-	TMap<FString, FStringFormatArg> Args;
-	Args.Add(TEXT("FunctionName"), FunctionInfo.InstanceName);
-
-	if (FunctionInfo.DefinitionName == SceneDepthName)
-	{
-		static const TCHAR* FormatSample = TEXT(R"(
-			void {FunctionName}(in float3 In_SamplePos, out float Out_SceneDepth, out float3 Out_CameraPosWorld, out bool Out_IsInsideView, out float3 Out_WorldPos, out float3 Out_WorldNormal)
-			{
-				DICollisionQuery_SceneDepth(In_SamplePos, Out_SceneDepth, Out_CameraPosWorld, Out_IsInsideView, Out_WorldPos, Out_WorldNormal);
-			}
-		)");
-
-		OutHLSL += FString::Format(FormatSample, Args);
-	}
-	else if (FunctionInfo.DefinitionName == CustomDepthName)
+	if ( (FunctionInfo.DefinitionName == NDICollisionQueryLocal::SceneDepthName) ||
+		 (FunctionInfo.DefinitionName == NDICollisionQueryLocal::CustomDepthName) ||
+		 (FunctionInfo.DefinitionName == NDICollisionQueryLocal::DistanceFieldName) ||
+		 (FunctionInfo.DefinitionName == NDICollisionQueryLocal::IssueAsyncRayTraceName) ||
+		(FunctionInfo.DefinitionName == NDICollisionQueryLocal::ReadAsyncRayTraceName) )
 	{
-		static const TCHAR* FormatSample = TEXT(R"(
-			void {FunctionName}(in float3 In_SamplePos, out float Out_SceneDepth, out float3 Out_CameraPosWorld, out bool Out_IsInsideView, out float3 Out_WorldPos, out float3 Out_WorldNormal)
-			{
-				DICollisionQuery_CustomDepth(In_SamplePos, Out_SceneDepth, Out_CameraPosWorld, Out_IsInsideView, Out_WorldPos, Out_WorldNormal);
-			}
-		)");
-
-		OutHLSL += FString::Format(FormatSample, Args);
+		return true;
 	}
-	else if (FunctionInfo.DefinitionName == DistanceFieldName)
-	{
-		static const TCHAR* FormatSample = TEXT(R"(
-			void {FunctionName}(in float3 In_SamplePos, out float Out_DistanceToNearestSurface, out float3 Out_FieldGradient, out bool Out_IsDistanceFieldValid)
-			{
-				DICollisionQuery_DistanceField(In_SamplePos, Out_DistanceToNearestSurface, Out_FieldGradient, Out_IsDistanceFieldValid);
-			}
-		)");
 
-		OutHLSL += FString::Format(FormatSample, Args);
-	}
-	else if (FunctionInfo.DefinitionName == IssueAsyncRayTraceName)
-	{
-		static const TCHAR* FormatSample = TEXT(R"(
-			void {FunctionName}(int In_QueryID, float3 In_TraceStart, float3 In_TraceEnd, int In_TraceChannel, out bool Out_IsQueryValid)
-			{
-				DICollisionQuery_IssueAsyncRayTrace(In_QueryID, In_TraceStart, In_TraceEnd, In_TraceChannel, Out_IsQueryValid);
-			}
-		)");
-
-		OutHLSL += FString::Format(FormatSample, Args);
-	}
-	else if (FunctionInfo.DefinitionName == ReadAsyncRayTraceName)
-	{
-		static const TCHAR* FormatSample = TEXT(R"(
-			void {FunctionName}(int In_PreviousFrameQueryID, out bool Out_CollisionValid, out float Out_CollisionDistance, out float3 Out_CollisionPosWorld, out float3 Out_CollisionNormal)
-			{
-				DICollisionQuery_ReadAsyncRayTrace(In_PreviousFrameQueryID, Out_CollisionValid, Out_CollisionDistance, Out_CollisionPosWorld, Out_CollisionNormal);
-			}
-		)");
-
-		OutHLSL += FString::Format(FormatSample, Args);
-	}
-	else
-	{
-		return false;
-	}
-
-	return true;
+	return false;
 }
 
 bool UNiagaraDataInterfaceCollisionQuery::UpgradeFunctionCall(FNiagaraFunctionSignature& FunctionSignature)
@@ -520,7 +473,7 @@ bool UNiagaraDataInterfaceCollisionQuery::UpgradeFunctionCall(FNiagaraFunctionSi
 	bool bWasChanged = false;
 
 	// The distance field query got a new output at some point, but there exists no custom version for it
-	if (FunctionSignature.Name == UNiagaraDataInterfaceCollisionQuery::DistanceFieldName && FunctionSignature.Outputs.Num() == 2)
+	if (FunctionSignature.Name == NDICollisionQueryLocal::DistanceFieldName && FunctionSignature.Outputs.Num() == 2)
 	{
 		FunctionSignature.Outputs.Add(FNiagaraVariable(FNiagaraTypeDefinition::GetBoolDef(), TEXT("IsDistanceFieldValid")));
 		bWasChanged = true;
@@ -535,7 +488,7 @@ bool UNiagaraDataInterfaceCollisionQuery::UpgradeFunctionCall(FNiagaraFunctionSi
 	// Added the possibility to skip a line trace to increase performance when only a fraction of particles wants to do a line trace
 	if (FunctionSignature.FunctionVersion < FNiagaraCollisionDIFunctionVersion::AddedTraceSkip)
 	{
-		if (FunctionSignature.Name == UNiagaraDataInterfaceCollisionQuery::SyncTraceName || FunctionSignature.Name == UNiagaraDataInterfaceCollisionQuery::AsyncTraceName)
+		if (FunctionSignature.Name == NDICollisionQueryLocal::SyncTraceName || FunctionSignature.Name == NDICollisionQueryLocal::AsyncTraceName)
 		{
 			FunctionSignature.Inputs.Add(FNiagaraVariable(FNiagaraTypeDefinition::GetBoolDef(), TEXT("SkipTrace")));
 			bWasChanged = true;
@@ -545,7 +498,7 @@ bool UNiagaraDataInterfaceCollisionQuery::UpgradeFunctionCall(FNiagaraFunctionSi
 	// Added the physical material ID as a result for line traces
 	if (FunctionSignature.FunctionVersion < FNiagaraCollisionDIFunctionVersion::ReturnCollisionMaterialIdx)
 	{
-		if (FunctionSignature.Name == SyncTraceName || FunctionSignature.Name == AsyncTraceName)
+		if (FunctionSignature.Name == NDICollisionQueryLocal::SyncTraceName || FunctionSignature.Name == NDICollisionQueryLocal::AsyncTraceName)
 		{
 			FunctionSignature.Outputs.Add(FNiagaraVariable(FNiagaraTypeDefinition::GetIntDef(), TEXT("CollisionMaterialIndex")));
 			bWasChanged = true;
@@ -568,7 +521,7 @@ bool IsDistanceFieldEnabled()
 #if WITH_EDITOR
 void UNiagaraDataInterfaceCollisionQuery::ValidateFunction(const FNiagaraFunctionSignature& Function, TArray<FText>& OutValidationErrors)
 {
-	if (Function.Name == DistanceFieldName)
+	if (Function.Name == NDICollisionQueryLocal::DistanceFieldName)
 	{
 		if (!IsDistanceFieldEnabled())
 		{
@@ -586,7 +539,14 @@ bool UNiagaraDataInterfaceCollisionQuery::RequiresRayTracingScene() const
 #if WITH_EDITORONLY_DATA
 void UNiagaraDataInterfaceCollisionQuery::GetParameterDefinitionHLSL(const FNiagaraDataInterfaceGPUParamInfo& ParamInfo, FString& OutHLSL)
 {
-	// we don't need to add these to hlsl, as they're already in common.ush
+	TMap<FString, FStringFormatArg> TemplateArgs =
+	{
+		{TEXT("ParameterName"),	ParamInfo.DataInterfaceHLSLSymbol},
+	};
+
+	FString TemplateFile;
+	LoadShaderSourceFile(NDICollisionQueryLocal::TemplateShaderFile, EShaderPlatform::SP_PCD3D_SM5, &TemplateFile, nullptr);
+	OutHLSL += FString::Format(*TemplateFile, TemplateArgs);
 }
 #endif
 
@@ -597,20 +557,20 @@ DEFINE_NDI_DIRECT_FUNC_BINDER(UNiagaraDataInterfaceCollisionQuery, QueryMeshDist
 
 void UNiagaraDataInterfaceCollisionQuery::GetVMExternalFunction(const FVMExternalFunctionBindingInfo& BindingInfo, void* InstanceData, FVMExternalFunction &OutFunc)
 {
-	if (BindingInfo.Name == UNiagaraDataInterfaceCollisionQuery::SyncTraceName)
+	if (BindingInfo.Name == NDICollisionQueryLocal::SyncTraceName)
 	{
 		NDI_FUNC_BINDER(UNiagaraDataInterfaceCollisionQuery, PerformQuerySyncCPU)::Bind(this, OutFunc);
 	}
-	else if (BindingInfo.Name == UNiagaraDataInterfaceCollisionQuery::AsyncTraceName)
+	else if (BindingInfo.Name == NDICollisionQueryLocal::AsyncTraceName)
 	{
 		NDI_FUNC_BINDER(UNiagaraDataInterfaceCollisionQuery, PerformQueryAsyncCPU)::Bind(this, OutFunc);
 	}
-	else if (BindingInfo.Name == UNiagaraDataInterfaceCollisionQuery::SceneDepthName ||
-			 BindingInfo.Name == UNiagaraDataInterfaceCollisionQuery::CustomDepthName)
+	else if (BindingInfo.Name == NDICollisionQueryLocal::SceneDepthName ||
+			 BindingInfo.Name == NDICollisionQueryLocal::CustomDepthName)
 	{
 		NDI_FUNC_BINDER(UNiagaraDataInterfaceCollisionQuery, QuerySceneDepth)::Bind(this, OutFunc);
 	}
-	else if (BindingInfo.Name == UNiagaraDataInterfaceCollisionQuery::DistanceFieldName)
+	else if (BindingInfo.Name == NDICollisionQueryLocal::DistanceFieldName)
 	{
 		NDI_FUNC_BINDER(UNiagaraDataInterfaceCollisionQuery, QueryMeshDistanceField)::Bind(this, OutFunc);
 	}
@@ -624,7 +584,7 @@ void UNiagaraDataInterfaceCollisionQuery::GetVMExternalFunction(const FVMExterna
 #if WITH_EDITORONLY_DATA
 void UNiagaraDataInterfaceCollisionQuery::GetCommonHLSL(FString& OutHlsl)
 {
-	OutHlsl += TEXT("#include \"/Plugin/FX/Niagara/Private/NiagaraDataInterfaceCollisionQuery.ush\"\n");
+	OutHlsl.Appendf(TEXT("#include \"%s\"\n"), NDICollisionQueryLocal::CommonShaderFile);
 }
 
 bool UNiagaraDataInterfaceCollisionQuery::AppendCompileHash(FNiagaraCompileHashVisitor* InVisitor) const
@@ -633,11 +593,11 @@ bool UNiagaraDataInterfaceCollisionQuery::AppendCompileHash(FNiagaraCompileHashV
 	{
 		return false;
 	}
-	bool bDistanceFieldEnabled = IsDistanceFieldEnabled();
-	InVisitor->UpdatePOD(TEXT("NiagaraCollisionDI_DistanceField"), bDistanceFieldEnabled);
 
-	FSHAHash Hash = GetShaderFileHash(TEXT("/Plugin/FX/Niagara/Private/NiagaraDataInterfaceCollisionQuery.ush"), EShaderPlatform::SP_PCD3D_SM5);
-	InVisitor->UpdateString(TEXT("NiagaraDataInterfaceCollisionQueryHlslSource"), Hash.ToString());
+	InVisitor->UpdatePOD(TEXT("NiagaraCollisionDI_DistanceField"), IsDistanceFieldEnabled());
+	InVisitor->UpdateString(TEXT("NDICollisionQueryCommonHLSLSource"), GetShaderFileHash(NDICollisionQueryLocal::CommonShaderFile, EShaderPlatform::SP_PCD3D_SM5).ToString());
+	InVisitor->UpdateString(TEXT("NDICollisionQueryTemplateHLSLSource"), GetShaderFileHash(NDICollisionQueryLocal::TemplateShaderFile, EShaderPlatform::SP_PCD3D_SM5).ToString());
+
 	return true;
 }
 #endif
@@ -944,9 +904,10 @@ public:
 
 		FNiagaraDataIntefaceProxyCollisionQuery* QueryDI = (FNiagaraDataIntefaceProxyCollisionQuery*)Context.DataInterface;
 		FRHIComputeShader* ComputeShaderRHI = Context.Shader.GetComputeShader();
-
-		if (GlobalDistanceFieldParameters.IsBound() && Context.Batcher)
+		// Bind distance field parameters
+		if (GlobalDistanceFieldParameters.IsBound())
 		{
+			check(Context.Batcher);
 			GlobalDistanceFieldParameters.Set(RHICmdList, ComputeShaderRHI, Context.Batcher->GetGlobalDistanceFieldParameters());
 		}
 
diff --git a/Engine/Plugins/FX/Niagara/Source/Niagara/Private/NiagaraDataInterfaceOcclusion.cpp b/Engine/Plugins/FX/Niagara/Source/Niagara/Private/NiagaraDataInterfaceOcclusion.cpp
--- a/Engine/Plugins/FX/Niagara/Source/Niagara/Private/NiagaraDataInterfaceOcclusion.cpp
+++ b/Engine/Plugins/FX/Niagara/Source/Niagara/Private/NiagaraDataInterfaceOcclusion.cpp
@@ -189,21 +189,4 @@ void UNiagaraDataInterfaceOcclusion::QueryOcclusionFactorCircleGPU(FVectorVMCont
 
 // ------------------------------------------------------------
 
-struct FNiagaraDataInterfaceParametersCS_OcclusionQuery : public FNiagaraDataInterfaceParametersCS
-{
-	DECLARE_TYPE_LAYOUT(FNiagaraDataInterfaceParametersCS_OcclusionQuery, NonVirtual);
-public:
-	void Bind(const FNiagaraDataInterfaceGPUParamInfo& ParameterInfo, const class FShaderParameterMap& ParameterMap)
-	{
-	}
-
-	void Set(FRHICommandList& RHICmdList, const FNiagaraDataInterfaceSetArgs& Context) const
-	{
-	}
-};
-
-IMPLEMENT_TYPE_LAYOUT(FNiagaraDataInterfaceParametersCS_OcclusionQuery);
-
-IMPLEMENT_NIAGARA_DI_PARAMETER(UNiagaraDataInterfaceOcclusion, FNiagaraDataInterfaceParametersCS_OcclusionQuery);
-
 #undef LOCTEXT_NAMESPACE
diff --git a/Engine/Plugins/FX/Niagara/Source/Niagara/Private/NiagaraEmitterInstanceBatcher.cpp b/Engine/Plugins/FX/Niagara/Source/Niagara/Private/NiagaraEmitterInstanceBatcher.cpp
--- a/Engine/Plugins/FX/Niagara/Source/Niagara/Private/NiagaraEmitterInstanceBatcher.cpp
+++ b/Engine/Plugins/FX/Niagara/Source/Niagara/Private/NiagaraEmitterInstanceBatcher.cpp
@@ -1444,7 +1444,6 @@ void NiagaraEmitterInstanceBatcher::ExecuteAll(FRHICommandList& RHICmdList, FRHI
 				}
 			}
 		}
-
 	}
 
 	// Release counts
@@ -1518,7 +1517,7 @@ void NiagaraEmitterInstanceBatcher::PreInitViews(FRDGBuilder& GraphBuilder, bool
 		if (GNiagaraAllowTickBeforeRender)
 		{
 			FNiagaraSceneTextureParameters* PassParameters = GraphBuilder.AllocParameters<FNiagaraSceneTextureParameters>();
-			PassParameters->SceneTextures = CreateSceneTextureUniformBuffer(GraphBuilder, ERHIFeatureLevel::SM5, ESceneTextureSetupMode::None);
+			GNiagaraViewDataManager.GetSceneTextureParameters(GraphBuilder, *PassParameters);
 
 			GraphBuilder.AddPass(
 				RDG_EVENT_NAME("Niagara::PreInitViews"),
@@ -1548,7 +1547,7 @@ void NiagaraEmitterInstanceBatcher::PostInitViews(FRDGBuilder& GraphBuilder, TAr
 	if (bAllowGPUParticleUpdate && FNiagaraUtilities::AllowGPUParticles(GetShaderPlatform()))
 	{
 		FNiagaraSceneTextureParameters* PassParameters = GraphBuilder.AllocParameters<FNiagaraSceneTextureParameters>();
-		PassParameters->SceneTextures = CreateSceneTextureUniformBuffer(GraphBuilder, ERHIFeatureLevel::SM5, ESceneTextureSetupMode::None);
+		GNiagaraViewDataManager.GetSceneTextureParameters(GraphBuilder, *PassParameters);
 
 		GraphBuilder.AddPass(
 			RDG_EVENT_NAME("Niagara::PostInitViews"),
diff --git a/Engine/Plugins/FX/Niagara/Source/Niagara/Private/NiagaraWorldManager.cpp b/Engine/Plugins/FX/Niagara/Source/Niagara/Private/NiagaraWorldManager.cpp
--- a/Engine/Plugins/FX/Niagara/Source/Niagara/Private/NiagaraWorldManager.cpp
+++ b/Engine/Plugins/FX/Niagara/Source/Niagara/Private/NiagaraWorldManager.cpp
@@ -198,6 +198,7 @@ void FNiagaraViewDataMgr::PostOpaqueRender(FPostOpaqueRenderParameters& Params)
 {
 	ViewUniformBuffer = Params.ViewUniformBuffer;
 	Parameters.SceneTextures = Params.SceneTexturesUniformParams;
+	Parameters.MobileSceneTextures = Params.MobileSceneTexturesUniformParams;
 	Parameters.Depth = Params.DepthTexture;
 	Parameters.Normal = Params.NormalTexture;
 	Parameters.Velocity = Params.VelocityTexture;
@@ -207,9 +208,20 @@ void FNiagaraViewDataMgr::GetSceneTextureParameters(FRDGBuilder& GraphBuilder, F
 {
 	OutParameters = Parameters;
 
-	if (!OutParameters.SceneTextures)
+	ERHIFeatureLevel::Type LocalFeatureLevel = GetFeatureLevel();
+	if (FSceneInterface::GetShadingPath(LocalFeatureLevel) == EShadingPath::Deferred)
 	{
-		OutParameters.SceneTextures = CreateSceneTextureUniformBuffer(GraphBuilder, ERHIFeatureLevel::SM5, ESceneTextureSetupMode::None);
+		if ( !Parameters.SceneTextures )
+		{
+			OutParameters.SceneTextures = CreateSceneTextureUniformBuffer(GraphBuilder, LocalFeatureLevel, ESceneTextureSetupMode::None);
+		}
+	}
+	else if (FSceneInterface::GetShadingPath(LocalFeatureLevel) == EShadingPath::Mobile)
+	{
+		if ( !Parameters.MobileSceneTextures )
+		{
+			OutParameters.MobileSceneTextures = CreateMobileSceneTextureUniformBuffer(GraphBuilder, EMobileSceneTextureSetupMode::None);
+		}
 	}
 }
 
diff --git a/Engine/Plugins/FX/Niagara/Source/Niagara/Public/NiagaraWorldManager.h b/Engine/Plugins/FX/Niagara/Source/Niagara/Public/NiagaraWorldManager.h
--- a/Engine/Plugins/FX/Niagara/Source/Niagara/Public/NiagaraWorldManager.h
+++ b/Engine/Plugins/FX/Niagara/Source/Niagara/Public/NiagaraWorldManager.h
@@ -32,6 +32,7 @@ struct FNiagaraScalabilityState;
 
 BEGIN_SHADER_PARAMETER_STRUCT(FNiagaraSceneTextureParameters, )
 	SHADER_PARAMETER_RDG_UNIFORM_BUFFER(FSceneTextureUniformParameters, SceneTextures)
+	SHADER_PARAMETER_RDG_UNIFORM_BUFFER(FMobileSceneTextureUniformParameters, MobileSceneTextures)
 	RDG_TEXTURE_ACCESS(Depth, ERHIAccess::SRVCompute)
 	RDG_TEXTURE_ACCESS(Normal, ERHIAccess::SRVCompute)
 	RDG_TEXTURE_ACCESS(Velocity, ERHIAccess::SRVCompute)
diff --git a/Engine/Plugins/FX/Niagara/Source/NiagaraCore/Private/NiagaraCustomVersion.cpp b/Engine/Plugins/FX/Niagara/Source/NiagaraCore/Private/NiagaraCustomVersion.cpp
--- a/Engine/Plugins/FX/Niagara/Source/NiagaraCore/Private/NiagaraCustomVersion.cpp
+++ b/Engine/Plugins/FX/Niagara/Source/NiagaraCore/Private/NiagaraCustomVersion.cpp
@@ -8,4 +8,4 @@ const FGuid FNiagaraCustomVersion::GUID(0xFCF57AFA, 0x50764283, 0xB9A9E658, 0xFF
 // Register the custom version with core
 FCustomVersionRegistration GRegisterNiagaraCustomVersion(FNiagaraCustomVersion::GUID, FNiagaraCustomVersion::LatestVersion, TEXT("NiagaraVer"));
 
-const FGuid FNiagaraCustomVersion::LatestScriptCompileVersion(0x5E8FBE61, 0x53A1054B, 0xB543A46F, 0x6916AB4D);
+const FGuid FNiagaraCustomVersion::LatestScriptCompileVersion(0x6A47E183, 0xE402744D, 0xAACE0531, 0xC89547FF);
diff --git a/Engine/Source/Runtime/RenderCore/Public/RendererInterface.h b/Engine/Source/Runtime/RenderCore/Public/RendererInterface.h
--- a/Engine/Source/Runtime/RenderCore/Public/RendererInterface.h
+++ b/Engine/Source/Runtime/RenderCore/Public/RendererInterface.h
@@ -25,6 +25,7 @@ class FSceneInterface;
 class FSceneView;
 class FSceneViewFamily;
 class FSceneTextureUniformParameters;
+class FMobileSceneTextureUniformParameters;
 class FGlobalDistanceFieldParameterData;
 struct FMeshBatch;
 struct FSynthBenchmarkResults;
@@ -568,6 +569,7 @@ public:
 	FRDGBuilder* GraphBuilder = nullptr;
 	FRHIUniformBuffer* ViewUniformBuffer = nullptr;
 	TRDGUniformBufferRef<FSceneTextureUniformParameters> SceneTexturesUniformParams = nullptr;
+	TRDGUniformBufferRef<FMobileSceneTextureUniformParameters> MobileSceneTexturesUniformParams = nullptr;
 	const FGlobalDistanceFieldParameterData* GlobalDistanceFieldParams = nullptr;
 	void* Uid = nullptr; // A unique identifier for the view.
 };
diff --git a/Engine/Source/Runtime/Renderer/Private/MobileShadingRenderer.cpp b/Engine/Source/Runtime/Renderer/Private/MobileShadingRenderer.cpp
--- a/Engine/Source/Runtime/Renderer/Private/MobileShadingRenderer.cpp
+++ b/Engine/Source/Runtime/Renderer/Private/MobileShadingRenderer.cpp
@@ -150,9 +150,28 @@ static bool UsesCustomDepthStencilLookup(const FViewInfo& View)
 }
 
 BEGIN_SHADER_PARAMETER_STRUCT(FMobileRenderOpaqueFXPassParameters, )
-	SHADER_PARAMETER_RDG_UNIFORM_BUFFER(FMobileSceneTextureUniformParameters, SceneTextures)
+	SHADER_PARAMETER_RDG_UNIFORM_BUFFER(FMobileSceneTextureUniformParameters, MobileSceneTextures)
 END_SHADER_PARAMETER_STRUCT()
 
+static void RenderOpaqueFX(
+	FRDGBuilder& GraphBuilder,
+	TArrayView<const FViewInfo> Views,
+	FFXSystemInterface* FXSystem,
+	TRDGUniformBufferRef<FMobileSceneTextureUniformParameters> MobileSceneTexturesUniformBuffer)
+{
+	// Notify the FX system that opaque primitives have been rendered and we now have a valid depth buffer.
+	if (FXSystem && Views.Num() > 0)
+	{
+		FXSystem->PostRenderOpaque(GraphBuilder, Views, true /*bAllowGPUParticleUpdate*/);
+
+		if (FGPUSortManager* GPUSortManager = FXSystem->GetGPUSortManager())
+		{
+			GPUSortManager->OnPostRenderOpaque(GraphBuilder);
+		}
+	}
+}
+
+
 BEGIN_SHADER_PARAMETER_STRUCT(FMobileRenderPassParameters, )
 	SHADER_PARAMETER_STRUCT_INCLUDE(FViewShaderParameters, View)
 	SHADER_PARAMETER_RDG_UNIFORM_BUFFER(FMobileBasePassUniformParameters, MobileBasePass)
@@ -813,14 +832,10 @@ void FMobileSceneRenderer::Render(FRDGBuilder& GraphBuilder)
 		SceneTextures.MobileUniformBuffer = CreateMobileSceneTextureUniformBuffer(GraphBuilder, SceneTextures.MobileSetupMode);
 	}
 
-	if (FXSystem && Views.IsValidIndex(0))
-	{
-		FXSystem->PostRenderOpaque(GraphBuilder, Views, true /*bAllowGPUParticleUpdate*/);
-		if (FGPUSortManager* GPUSortManager = FXSystem->GetGPUSortManager())
-		{
-			GPUSortManager->OnPostRenderOpaque(GraphBuilder);
-		}
-	}
+	FRendererModule& RendererModule = static_cast<FRendererModule&>(GetRendererModule());
+	RendererModule.RenderPostOpaqueExtensions(GraphBuilder, Views, SceneTextures);
+
+	RenderOpaqueFX(GraphBuilder, Views, FXSystem, SceneTextures.MobileUniformBuffer);
 
 	if (bRequiresPixelProjectedPlanarRelfectionPass)
 	{
diff --git a/Engine/Source/Runtime/Renderer/Private/SceneRendering.cpp b/Engine/Source/Runtime/Renderer/Private/SceneRendering.cpp
--- a/Engine/Source/Runtime/Renderer/Private/SceneRendering.cpp
+++ b/Engine/Source/Runtime/Renderer/Private/SceneRendering.cpp
@@ -4061,6 +4061,7 @@ void FRendererModule::RenderPostOpaqueExtensions(
 			RenderParameters.SmallDepthTexture = SceneTextures.SmallDepth;
 			RenderParameters.ViewUniformBuffer = View.ViewUniformBuffer;
 			RenderParameters.SceneTexturesUniformParams = SceneTextures.UniformBuffer;
+			RenderParameters.MobileSceneTexturesUniformParams = SceneTextures.MobileUniformBuffer;
 			RenderParameters.GlobalDistanceFieldParams = &View.GlobalDistanceFieldInfo.ParameterData;
 
 			RenderParameters.ViewportRect = View.ViewRect;
