Fix PVS warning V502 when comparing against NAME_None in a trinary

diff --git a/Engine/Plugins/Editor/AssetManagerEditor/Source/AssetManagerEditor/Private/SSizeMap.cpp b/Engine/Plugins/Editor/AssetManagerEditor/Source/AssetManagerEditor/Private/SSizeMap.cpp
--- a/Engine/Plugins/Editor/AssetManagerEditor/Source/AssetManagerEditor/Private/SSizeMap.cpp
+++ b/Engine/Plugins/Editor/AssetManagerEditor/Source/AssetManagerEditor/Private/SSizeMap.cpp
@@ -405,7 +405,7 @@ void SSizeMap::GatherDependenciesRecursively(TSharedPtr<FAssetThumbnailPool>& In
 	for (const FAssetIdentifier& AssetIdentifier : AssetIdentifiers)
 	{
 		FName AssetPackageName = AssetIdentifier.IsPackage() ? AssetIdentifier.PackageName : NAME_None;
-		FString AssetPackageNameString = AssetPackageName != NAME_None ? AssetPackageName.ToString() : FString();
+		FString AssetPackageNameString = (AssetPackageName != NAME_None) ? AssetPackageName.ToString() : FString();
 		FPrimaryAssetId AssetPrimaryId = AssetIdentifier.GetPrimaryAssetId();
 		int32 ChunkId = UAssetManager::ExtractChunkIdFromPrimaryAssetId(AssetPrimaryId);
 		int32 FilterChunkId = UAssetManager::ExtractChunkIdFromPrimaryAssetId(FilterPrimaryAsset);
diff --git a/Engine/Plugins/Experimental/Animation/ContextualAnimation/Source/ContextualAnimation/Public/ContextualAnimSceneAsset.h b/Engine/Plugins/Experimental/Animation/ContextualAnimation/Source/ContextualAnimation/Public/ContextualAnimSceneAsset.h
--- a/Engine/Plugins/Experimental/Animation/ContextualAnimation/Source/ContextualAnimation/Public/ContextualAnimSceneAsset.h
+++ b/Engine/Plugins/Experimental/Animation/ContextualAnimation/Source/ContextualAnimation/Public/ContextualAnimSceneAsset.h
@@ -79,7 +79,7 @@ public:
 		return AlignmentSections.IsValidIndex(Index) ? AlignmentSections[Index].SectionName : NAME_None;
 	}
 
-	FORCEINLINE const FName& GetLeaderRole() const { return LeaderRole != NAME_None ? LeaderRole : PrimaryRole; }
+	FORCEINLINE const FName& GetLeaderRole() const { return (LeaderRole != NAME_None) ? LeaderRole : PrimaryRole; }
 
 	FORCEINLINE float GetRadius() const { return Radius; }
 
diff --git a/Engine/Plugins/Experimental/ChaosCaching/Source/ChaosCaching/Private/Chaos/CacheManagerActor.cpp b/Engine/Plugins/Experimental/ChaosCaching/Source/ChaosCaching/Private/Chaos/CacheManagerActor.cpp
--- a/Engine/Plugins/Experimental/ChaosCaching/Source/ChaosCaching/Private/Chaos/CacheManagerActor.cpp
+++ b/Engine/Plugins/Experimental/ChaosCaching/Source/ChaosCaching/Private/Chaos/CacheManagerActor.cpp
@@ -335,7 +335,7 @@ void AChaosCacheManager::BeginPlay()
 				{
 					// Make sure there's a cache available if we're going to record.
 					UPrimitiveComponent* Component = Observed.GetComponent();
-					FName CacheName = Observed.CacheName == NAME_None ? MakeUniqueObjectName(CacheCollection, UChaosCache::StaticClass(), "Cache") : Observed.CacheName;
+					FName CacheName = (Observed.CacheName == NAME_None) ? MakeUniqueObjectName(CacheCollection, UChaosCache::StaticClass(), "Cache") : Observed.CacheName;
 
 					UChaosCache*    RecordCache = CacheCollection->FindOrAddCache(CacheName);
 					FCacheUserToken Token       = RecordCache->BeginRecord(Observed.GetComponent(), CurrAdapter->GetGuid(), GetTransform());
diff --git a/Engine/Plugins/FX/Niagara/Source/NiagaraEditorWidgets/Private/Stack/SNiagaraStackTableRow.cpp b/Engine/Plugins/FX/Niagara/Source/NiagaraEditorWidgets/Private/Stack/SNiagaraStackTableRow.cpp
--- a/Engine/Plugins/FX/Niagara/Source/NiagaraEditorWidgets/Private/Stack/SNiagaraStackTableRow.cpp
+++ b/Engine/Plugins/FX/Niagara/Source/NiagaraEditorWidgets/Private/Stack/SNiagaraStackTableRow.cpp
@@ -278,7 +278,7 @@ void SNiagaraStackTableRow::SetNameAndValueContent(TSharedRef<SWidget> InNameWid
 
 	FName AccentColorName = FNiagaraStackEditorWidgetsUtilities::GetIconColorNameForExecutionCategory(StackEntry->GetExecutionCategoryName());
 	const bool bDisplayingIndicator = IndicatorColor != FStyleColors::Transparent;
-	FSlateColor AccentColor = bDisplayingIndicator ? IndicatorColor : (AccentColorName != NAME_None ? FNiagaraEditorWidgetsStyle::Get().GetColor(AccentColorName) : FStyleColors::Transparent);
+	FSlateColor AccentColor = bDisplayingIndicator ? IndicatorColor : ((AccentColorName != NAME_None) ? FNiagaraEditorWidgetsStyle::Get().GetColor(AccentColorName) : FStyleColors::Transparent);
 
 	ChildSlot
 	[
@@ -502,7 +502,7 @@ EVisibility SNiagaraStackTableRow::GetRowVisibility() const
 
 EVisibility SNiagaraStackTableRow::GetExecutionCategoryIconVisibility() const
 {
-	return bShowExecutionCategoryIcon && StackEntry->GetExecutionSubcategoryName() != NAME_None
+	return bShowExecutionCategoryIcon && (StackEntry->GetExecutionSubcategoryName() != NAME_None)
 		? EVisibility::Visible
 		: EVisibility::Collapsed;
 }
diff --git a/Engine/Plugins/Runtime/PropertyAccess/Source/PropertyAccessEditor/Private/SPropertyBinding.cpp b/Engine/Plugins/Runtime/PropertyAccess/Source/PropertyAccessEditor/Private/SPropertyBinding.cpp
--- a/Engine/Plugins/Runtime/PropertyAccess/Source/PropertyAccessEditor/Private/SPropertyBinding.cpp
+++ b/Engine/Plugins/Runtime/PropertyAccess/Source/PropertyAccessEditor/Private/SPropertyBinding.cpp
@@ -704,7 +704,7 @@ void SPropertyBinding::HandleCreateAndAddBinding()
 		WidgetName = Args.OnGenerateBindingName.Execute();
 	}
 
-	FString Post = PropertyName != NAME_None ? PropertyName.ToString() : TEXT("");
+	FString Post = (PropertyName != NAME_None) ? PropertyName.ToString() : FString();
 	Post.RemoveFromStart(TEXT("On"));
 	Post.RemoveFromEnd(TEXT("Event"));
 
diff --git a/Engine/Source/Developer/HierarchicalLODUtilities/Private/HierarchicalLODUtilities.cpp b/Engine/Source/Developer/HierarchicalLODUtilities/Private/HierarchicalLODUtilities.cpp
--- a/Engine/Source/Developer/HierarchicalLODUtilities/Private/HierarchicalLODUtilities.cpp
+++ b/Engine/Source/Developer/HierarchicalLODUtilities/Private/HierarchicalLODUtilities.cpp
@@ -110,7 +110,7 @@ static FString GetHLODPackageName(const ULevel* InLevel, const uint32 InHLODLeve
 	FString LevelPackageName;
 	if (ULevelStreaming* StreamingLevel = FLevelUtils::FindStreamingLevel(InLevel))
 	{
-		LevelPackageName = StreamingLevel->PackageNameToLoad != NAME_None ? StreamingLevel->PackageNameToLoad.ToString() : StreamingLevel->GetWorldAssetPackageName();
+		LevelPackageName = (StreamingLevel->PackageNameToLoad != NAME_None) ? StreamingLevel->PackageNameToLoad.ToString() : StreamingLevel->GetWorldAssetPackageName();
 	}
 	else
 	{
diff --git a/Engine/Source/Developer/TargetDeviceServices/Private/Proxies/TargetDeviceProxy.cpp b/Engine/Source/Developer/TargetDeviceServices/Private/Proxies/TargetDeviceProxy.cpp
--- a/Engine/Source/Developer/TargetDeviceServices/Private/Proxies/TargetDeviceProxy.cpp
+++ b/Engine/Source/Developer/TargetDeviceServices/Private/Proxies/TargetDeviceProxy.cpp
@@ -150,7 +150,7 @@ const FString FTargetDeviceProxy::GetTargetDeviceId(FName InVariant) const
 {
 	FString Variant;
 	// for an aggregate (All_<platform>_devices_on_<host>) proxy we have a list of associated devices
-	for (TSet<FString>::TConstIterator ItVariant(TargetDeviceVariants[InVariant == NAME_None? DefaultVariant: InVariant].DeviceIDs); ItVariant; ++ItVariant)
+	for (TSet<FString>::TConstIterator ItVariant(TargetDeviceVariants[(InVariant == NAME_None) ? DefaultVariant : InVariant].DeviceIDs); ItVariant; ++ItVariant)
 	{
 		// should return the first device
 		// this method is designed for physical devices, not aggregate proxies
diff --git a/Engine/Source/Developer/TraceInsights/Private/Insights/Table/ViewModels/Table.cpp b/Engine/Source/Developer/TraceInsights/Private/Insights/Table/ViewModels/Table.cpp
--- a/Engine/Source/Developer/TraceInsights/Private/Insights/Table/ViewModels/Table.cpp
+++ b/Engine/Source/Developer/TraceInsights/Private/Insights/Table/ViewModels/Table.cpp
@@ -153,7 +153,7 @@ void FTable::GetVisibleColumnsData(const TArray<Insights::FBaseTreeNodePtr>& InN
 	{
 		NumItems = Max_Rows;
 
-		FMessageLog ReportMessageLog(LogListingName != NAME_None ? LogListingName : TEXT("Other"));
+		FMessageLog ReportMessageLog((LogListingName != NAME_None) ? LogListingName : TEXT("Other"));
 		ReportMessageLog.Warning(FText::Format(LOCTEXT("TooManyRows", "Too many rows selected. Only the first {0} will be copied."), NumItems));
 		ReportMessageLog.Notify();
 	}
diff --git a/Engine/Source/Editor/AnimGraph/Private/AnimBlueprintCompiler.cpp b/Engine/Source/Editor/AnimGraph/Private/AnimBlueprintCompiler.cpp
--- a/Engine/Source/Editor/AnimGraph/Private/AnimBlueprintCompiler.cpp
+++ b/Engine/Source/Editor/AnimGraph/Private/AnimBlueprintCompiler.cpp
@@ -1511,7 +1511,7 @@ void FAnimBlueprintCompilerContext::CreateAnimGraphStubFunctions()
 				UK2Node_FunctionEntry* EntryNode = SpawnIntermediateNode<UK2Node_FunctionEntry>(RootNode, StubGraph);
 				EntryNode->NodePosX = -200;
 				EntryNode->CustomGeneratedFunctionName = GraphToUseforSignature->GetFName();	// Note that the function generated from this temporary graph is undecorated
-				EntryNode->MetaData.Category = RootNode->Node.GetGroup() == NAME_None ? FText::GetEmpty() : FText::FromName(RootNode->Node.GetGroup());
+				EntryNode->MetaData.Category = (RootNode->Node.GetGroup() == NAME_None) ? FText::GetEmpty() : FText::FromName(RootNode->Node.GetGroup());
 
 				// Add linked input poses as parameters
 				for(UAnimGraphNode_LinkedInputPose* LinkedInputPoseNode : LinkedInputPoseNodes)
diff --git a/Engine/Source/Editor/AnimGraph/Private/AnimGraphNode_LinkedAnimLayer.cpp b/Engine/Source/Editor/AnimGraph/Private/AnimGraphNode_LinkedAnimLayer.cpp
--- a/Engine/Source/Editor/AnimGraph/Private/AnimGraphNode_LinkedAnimLayer.cpp
+++ b/Engine/Source/Editor/AnimGraph/Private/AnimGraphNode_LinkedAnimLayer.cpp
@@ -95,7 +95,7 @@ FText UAnimGraphNode_LinkedAnimLayer::GetNodeTitle(ENodeTitleType::Type TitleTyp
 	FFormatNamedArguments Args;
 		Args.Add(TEXT("NodeTitle"), LOCTEXT("NodeTitle", "Linked Anim Layer"));
 		Args.Add(TEXT("TargetClass"), TargetAnimBlueprintInterface ? FText::FromString(TargetAnimBlueprintInterface->GetName()) : LOCTEXT("InterfaceNone", "None"));
-		Args.Add(TEXT("Layer"), Node.Layer == NAME_None ? LOCTEXT("LayerNone", "None") : FText::FromName(Node.Layer));
+		Args.Add(TEXT("Layer"), (Node.Layer == NAME_None) ? LOCTEXT("LayerNone", "None") : FText::FromName(Node.Layer));
 
 		if (FAnimNode_LinkedAnimLayer* PreviewNode = GetPreviewNode())
 		{
diff --git a/Engine/Source/Editor/DetailCustomizations/Private/SPerPlatformPropertiesWidget.cpp b/Engine/Source/Editor/DetailCustomizations/Private/SPerPlatformPropertiesWidget.cpp
--- a/Engine/Source/Editor/DetailCustomizations/Private/SPerPlatformPropertiesWidget.cpp
+++ b/Engine/Source/Editor/DetailCustomizations/Private/SPerPlatformPropertiesWidget.cpp
@@ -48,7 +48,7 @@ TSharedRef<SWidget> SPerPlatformPropertiesRow::MakePerPlatformWidget(FName InNam
 		.MinDesiredWidth(50.0f)
 		[
 			SAssignNew(HorizontalBox, SHorizontalBox)
-			.ToolTipText(InName != NAME_None 
+			.ToolTipText((InName != NAME_None)
 				? FText::Format(NSLOCTEXT("SPerPlatformPropertiesWidget", "PerPlatformDesc", "Override for {0}"), FText::AsCultureInvariant(InName.ToString()))
 				: NSLOCTEXT("SPerPlatformPropertiesWidget", "DefaultDesc", "Default value for properties without an override"))
 			+SHorizontalBox::Slot()
diff --git a/Engine/Source/Editor/Kismet/Private/SSCSEditor.cpp b/Engine/Source/Editor/Kismet/Private/SSCSEditor.cpp
--- a/Engine/Source/Editor/Kismet/Private/SSCSEditor.cpp
+++ b/Engine/Source/Editor/Kismet/Private/SSCSEditor.cpp
@@ -2871,7 +2871,7 @@ void ConformTransformRelativeToParent(USceneComponent* SceneComponentTemplate, U
 {
 	// If we find a match, calculate its new position relative to the scene root component instance in its current scene
 	FTransform ComponentToWorld(SceneComponentTemplate->GetRelativeRotation(), SceneComponentTemplate->GetRelativeLocation(), SceneComponentTemplate->GetRelativeScale3D());
-	FTransform ParentToWorld = SceneComponentTemplate->GetAttachSocketName() != NAME_None ? ParentSceneComponent->GetSocketTransform(SceneComponentTemplate->GetAttachSocketName(), RTS_World) : ParentSceneComponent->GetComponentToWorld();
+	FTransform ParentToWorld = (SceneComponentTemplate->GetAttachSocketName() != NAME_None) ? ParentSceneComponent->GetSocketTransform(SceneComponentTemplate->GetAttachSocketName(), RTS_World) : ParentSceneComponent->GetComponentToWorld();
 	FTransform RelativeTM = ComponentToWorld.GetRelativeTransform(ParentToWorld);
 
 	// Store new relative location value (if not set to absolute)
diff --git a/Engine/Source/Editor/Persona/Private/SAnimViewportToolBar.cpp b/Engine/Source/Editor/Persona/Private/SAnimViewportToolBar.cpp
--- a/Engine/Source/Editor/Persona/Private/SAnimViewportToolBar.cpp
+++ b/Engine/Source/Editor/Persona/Private/SAnimViewportToolBar.cpp
@@ -290,7 +290,7 @@ void SAnimViewportToolBar::Construct(const FArguments& InArgs, TSharedPtr<class
 
 	// Create our pinned commands before we bind commands
 	IPinnedCommandListModule& PinnedCommandListModule = FModuleManager::LoadModuleChecked<IPinnedCommandListModule>(TEXT("PinnedCommandList"));
-	PinnedCommands = PinnedCommandListModule.CreatePinnedCommandList(InArgs._ContextName != NAME_None ? InArgs._ContextName : TEXT("PersonaViewport"));
+	PinnedCommands = PinnedCommandListModule.CreatePinnedCommandList((InArgs._ContextName != NAME_None) ? InArgs._ContextName : TEXT("PersonaViewport"));
 	PinnedCommands->SetStyle(&FEditorStyle::Get(), TEXT("ViewportPinnedCommandList"));
 
 	ChildSlot
diff --git a/Engine/Source/Editor/PhysicsAssetEditor/Private/PhysicsAssetDetailsCustomization.cpp b/Engine/Source/Editor/PhysicsAssetEditor/Private/PhysicsAssetDetailsCustomization.cpp
--- a/Engine/Source/Editor/PhysicsAssetEditor/Private/PhysicsAssetDetailsCustomization.cpp
+++ b/Engine/Source/Editor/PhysicsAssetEditor/Private/PhysicsAssetDetailsCustomization.cpp
@@ -223,7 +223,7 @@ TSharedRef< SWidget > FPhysicsAssetDetailsCustomization::FillPhysicalAnimationPr
 							SNew(SBox)
 							.WidthOverride(MenuIconSize)
 							.HeightOverride(MenuIconSize)
-							.Visibility_Lambda([ProfileName](){ return ProfileName == NAME_None ? EVisibility::Collapsed : EVisibility::Visible; })
+							.Visibility_Lambda([ProfileName](){ return (ProfileName == NAME_None) ? EVisibility::Collapsed : EVisibility::Visible; })
 							[
 								SNew(SImage)
 								.Image(FSlateIcon(FEditorStyle::GetStyleSetName(), "Symbols.SearchGlass").GetIcon())
@@ -338,7 +338,7 @@ TSharedRef< SWidget > FPhysicsAssetDetailsCustomization::FillConstraintProfilesO
 						SNew(SBox)
 						.WidthOverride(MenuIconSize)
 						.HeightOverride(MenuIconSize)
-						.Visibility_Lambda([ProfileName]() { return ProfileName == NAME_None ? EVisibility::Collapsed : EVisibility::Visible; })
+						.Visibility_Lambda([ProfileName]() { return (ProfileName == NAME_None) ? EVisibility::Collapsed : EVisibility::Visible; })
 						[
 							SNew(SImage)
 							.Image(FSlateIcon(FEditorStyle::GetStyleSetName(), "Symbols.SearchGlass").GetIcon())
diff --git a/Engine/Source/Editor/SkeletonEditor/Private/SSkeletonTree.cpp b/Engine/Source/Editor/SkeletonEditor/Private/SSkeletonTree.cpp
--- a/Engine/Source/Editor/SkeletonEditor/Private/SSkeletonTree.cpp
+++ b/Engine/Source/Editor/SkeletonEditor/Private/SSkeletonTree.cpp
@@ -601,7 +601,7 @@ void SSkeletonTree::CreateTreeColumns()
 			.Text_Lambda([this] () -> FText
 			{
 				FName CurrentProfile = BlendProfilePicker->GetSelectedBlendProfileName();
-				return CurrentProfile != NAME_None ? FText::FromName(CurrentProfile) : LOCTEXT("NoBlendProfile", "No Blend");
+				return (CurrentProfile != NAME_None) ? FText::FromName(CurrentProfile) : LOCTEXT("NoBlendProfile", "No Blend");
 			})
 			.OnTextCommitted_Lambda([this](const FText& InText, ETextCommit::Type InCommitType)
 			{
diff --git a/Engine/Source/Editor/SubobjectDataInterface/Private/SubobjectDataSubsystem.cpp b/Engine/Source/Editor/SubobjectDataInterface/Private/SubobjectDataSubsystem.cpp
--- a/Engine/Source/Editor/SubobjectDataInterface/Private/SubobjectDataSubsystem.cpp
+++ b/Engine/Source/Editor/SubobjectDataInterface/Private/SubobjectDataSubsystem.cpp
@@ -493,7 +493,7 @@ static void ConformTransformRelativeToParent(USceneComponent* SceneComponentTemp
 {
 	// If we find a match, calculate its new position relative to the scene root component instance in its current scene
 	FTransform ComponentToWorld(SceneComponentTemplate->GetRelativeRotation(), SceneComponentTemplate->GetRelativeLocation(), SceneComponentTemplate->GetRelativeScale3D());
-	FTransform ParentToWorld = SceneComponentTemplate->GetAttachSocketName() != NAME_None ? ParentSceneComponent->GetSocketTransform(SceneComponentTemplate->GetAttachSocketName(), RTS_World) : ParentSceneComponent->GetComponentToWorld();
+	FTransform ParentToWorld = (SceneComponentTemplate->GetAttachSocketName() != NAME_None) ? ParentSceneComponent->GetSocketTransform(SceneComponentTemplate->GetAttachSocketName(), RTS_World) : ParentSceneComponent->GetComponentToWorld();
 	FTransform RelativeTM = ComponentToWorld.GetRelativeTransform(ParentToWorld);
 
 	// Store new relative location value (if not set to absolute)
diff --git a/Engine/Source/Editor/UnrealEd/Private/Fbx/SFbxSceneOptionWindow.cpp b/Engine/Source/Editor/UnrealEd/Private/Fbx/SFbxSceneOptionWindow.cpp
--- a/Engine/Source/Editor/UnrealEd/Private/Fbx/SFbxSceneOptionWindow.cpp
+++ b/Engine/Source/Editor/UnrealEd/Private/Fbx/SFbxSceneOptionWindow.cpp
@@ -1239,7 +1239,7 @@ void SFbxSceneOptionWindow::Construct(const FArguments& InArgs)
 
 	check(OwnerWindow.IsValid());
 
-	MaterialBasePath = GlobalImportSettings->MaterialBasePath == NAME_None ? TEXT("") : GlobalImportSettings->MaterialBasePath.ToString();
+	MaterialBasePath = (GlobalImportSettings->MaterialBasePath == NAME_None) ? TEXT("") : GlobalImportSettings->MaterialBasePath.ToString();
 
 	InitAllTabs();
 
diff --git a/Engine/Source/Editor/UnrealEd/Private/Kismet2/BlueprintEditorUtils.cpp b/Engine/Source/Editor/UnrealEd/Private/Kismet2/BlueprintEditorUtils.cpp
--- a/Engine/Source/Editor/UnrealEd/Private/Kismet2/BlueprintEditorUtils.cpp
+++ b/Engine/Source/Editor/UnrealEd/Private/Kismet2/BlueprintEditorUtils.cpp
@@ -6687,7 +6687,7 @@ struct FConformCallsToParentFunctionUtils
 				}
 
 				// Rename the graph if it does not match the actual function name.
-				const FName FunctionName = FunctionEntryNode->CustomGeneratedFunctionName != NAME_None ? FunctionEntryNode->CustomGeneratedFunctionName : FunctionEntryNode->FunctionReference.GetMemberName();
+				const FName FunctionName = (FunctionEntryNode->CustomGeneratedFunctionName != NAME_None) ? FunctionEntryNode->CustomGeneratedFunctionName : FunctionEntryNode->FunctionReference.GetMemberName();
 				if (FunctionEntryNode == FunctionEntryNodes[0]
 					&& !FBlueprintEditorUtils::IsEventGraph(CurrentGraph)
 					&& CurrentGraph->GetFName() != FunctionName)
diff --git a/Engine/Source/Runtime/Core/Private/Stats/StatsCommand.cpp b/Engine/Source/Runtime/Core/Private/Stats/StatsCommand.cpp
--- a/Engine/Source/Runtime/Core/Private/Stats/StatsCommand.cpp
+++ b/Engine/Source/Runtime/Core/Private/Stats/StatsCommand.cpp
@@ -1146,7 +1146,7 @@ struct FHUDGroupManager
 		FHudFrame& NewFrame = History.FindOrAdd( TargetFrame );
 
 		FName RootName = Params.Root.Get();
-		FString RootString = RootName == NAME_None ? FString() : RootName.ToString();
+		FString RootString = (RootName == NAME_None) ? FString() : RootName.ToString();
 
 		const bool bUseSlowMode = Params.bSlowMode;
 		const bool bUseBudgetMode = Params.BudgetSection != NAME_None;
diff --git a/Engine/Source/Runtime/CoreUObject/Private/UObject/Field.cpp b/Engine/Source/Runtime/CoreUObject/Private/UObject/Field.cpp
--- a/Engine/Source/Runtime/CoreUObject/Private/UObject/Field.cpp
+++ b/Engine/Source/Runtime/CoreUObject/Private/UObject/Field.cpp
@@ -967,7 +967,7 @@ void FField::PostDuplicate(const FField& InField)
 FField* FField::Duplicate(const FField* InField, FFieldVariant DestOwner, const FName DestName, EObjectFlags FlagMask, EInternalObjectFlags InternalFlagsMask)
 {
 	check(InField);
-	FField* NewField = InField->GetClass()->Construct(DestOwner, DestName == NAME_None ? InField->GetFName() : DestName, InField->GetFlags() & FlagMask);
+	FField* NewField = InField->GetClass()->Construct(DestOwner, (DestName == NAME_None) ? InField->GetFName() : DestName, InField->GetFlags() & FlagMask);
 	NewField->PostDuplicate(*InField);
 	return NewField;
 }
diff --git a/Engine/Source/Runtime/Engine/Private/AudioPluginUtilities.cpp b/Engine/Source/Runtime/Engine/Private/AudioPluginUtilities.cpp
--- a/Engine/Source/Runtime/Engine/Private/AudioPluginUtilities.cpp
+++ b/Engine/Source/Runtime/Engine/Private/AudioPluginUtilities.cpp
@@ -104,7 +104,7 @@ IAudioOcclusionFactory* AudioPluginUtilities::GetDesiredOcclusionPlugin()
 IAudioModulationFactory* AudioPluginUtilities::GetDesiredModulationPlugin()
 {
 	const FName& PlatformPluginName = FName(*GetDesiredPluginName(EAudioPlugin::MODULATION));
-	const FName& PluginName = PlatformPluginName == NAME_None ? GetDefaultModulationPluginName() : PlatformPluginName;
+	const FName& PluginName = (PlatformPluginName == NAME_None) ? GetDefaultModulationPluginName() : PlatformPluginName;
 	const FName& FeatureName = IAudioModulationFactory::GetModularFeatureName();
 
 	TArray<IAudioModulationFactory*> Factories = IModularFeatures::Get().GetModularFeatureImplementations<IAudioModulationFactory>(FeatureName);
diff --git a/Engine/Source/Runtime/Engine/Private/BlueprintGeneratedClass.cpp b/Engine/Source/Runtime/Engine/Private/BlueprintGeneratedClass.cpp
--- a/Engine/Source/Runtime/Engine/Private/BlueprintGeneratedClass.cpp
+++ b/Engine/Source/Runtime/Engine/Private/BlueprintGeneratedClass.cpp
@@ -829,7 +829,7 @@ UObject* UBlueprintGeneratedClass::FindArchetype(const UClass* ArchetypeClass, c
 		// node in a function graph, which can occur when the AddComponent node is wired to a flow-control node such as a ForEach loop, for example. Thus,
 		// we still look for the archetype by name, but we must first ensure that the instance name is converted to its "base" name by removing the index.
 #if WITH_EDITORONLY_DATA
-		const FName ArchetypeBaseName = NewArchetypeName != NAME_None ? NewArchetypeName : FName(ArchetypeName, 0);
+		const FName ArchetypeBaseName = (NewArchetypeName != NAME_None) ? NewArchetypeName : FName(ArchetypeName, 0);
 #else
 		const FName ArchetypeBaseName = FName(ArchetypeName, 0);
 #endif
diff --git a/Engine/Source/Runtime/Engine/Private/Level.cpp b/Engine/Source/Runtime/Engine/Private/Level.cpp
--- a/Engine/Source/Runtime/Engine/Private/Level.cpp
+++ b/Engine/Source/Runtime/Engine/Private/Level.cpp
@@ -2497,7 +2497,7 @@ TArray<FString> ULevel::GetOnDiskExternalActorPackages(const FString& ExternalAc
 TArray<FString> ULevel::GetOnDiskExternalActorPackages() const
 {
 	UWorld* World = GetTypedOuter<UWorld>();
-	FString ExternalActorsPath = ULevel::GetExternalActorsPath(World->GetPackage(), World->OriginalWorldName == NAME_None ? World->GetName() : World->OriginalWorldName.ToString());
+	FString ExternalActorsPath = ULevel::GetExternalActorsPath(World->GetPackage(), (World->OriginalWorldName == NAME_None) ? World->GetName() : World->OriginalWorldName.ToString());
 	return GetOnDiskExternalActorPackages(ExternalActorsPath);
 }
 
@@ -2510,7 +2510,7 @@ TArray<UPackage*> ULevel::GetLoadedExternalActorPackages() const
 
 	// We also need to provide empty packages (for actors that were converted to non-external)
 	UWorld* World = GetTypedOuter<UWorld>();
-	FString ExternalActorsPath = ULevel::GetExternalActorsPath(World->GetPackage(), World->OriginalWorldName == NAME_None ? World->GetName() : World->OriginalWorldName.ToString());
+	FString ExternalActorsPath = ULevel::GetExternalActorsPath(World->GetPackage(), (World->OriginalWorldName == NAME_None) ? World->GetName() : World->OriginalWorldName.ToString());
 	if (!ExternalActorsPath.IsEmpty())
 	{
 		// Make sure we filter out similar external folders
diff --git a/Engine/Source/Runtime/Engine/Private/LevelStreaming.cpp b/Engine/Source/Runtime/Engine/Private/LevelStreaming.cpp
--- a/Engine/Source/Runtime/Engine/Private/LevelStreaming.cpp
+++ b/Engine/Source/Runtime/Engine/Private/LevelStreaming.cpp
@@ -1483,7 +1483,7 @@ ULevelStreaming* ULevelStreaming::CreateInstance(const FString& InstanceUniqueNa
 		{
 			StreamingLevelInstance = NewObject<ULevelStreaming>(InWorld, GetClass(), NAME_None, RF_Transient, NULL);
 			// new level streaming instance will load the same map package as this object
-			StreamingLevelInstance->PackageNameToLoad = (PackageNameToLoad == NAME_None ? GetWorldAssetPackageFName() : PackageNameToLoad);
+			StreamingLevelInstance->PackageNameToLoad = ((PackageNameToLoad == NAME_None) ? GetWorldAssetPackageFName() : PackageNameToLoad);
 			// under a provided unique name
 			StreamingLevelInstance->SetWorldAssetByPackageName(InstanceUniquePackageName);
 			StreamingLevelInstance->SetShouldBeLoaded(false);
diff --git a/Engine/Source/Runtime/Engine/Private/PhysicsEngine/PhysicsAsset.cpp b/Engine/Source/Runtime/Engine/Private/PhysicsEngine/PhysicsAsset.cpp
--- a/Engine/Source/Runtime/Engine/Private/PhysicsEngine/PhysicsAsset.cpp
+++ b/Engine/Source/Runtime/Engine/Private/PhysicsEngine/PhysicsAsset.cpp
@@ -642,7 +642,7 @@ void SanitizeProfilesHelper(const TArray<T*>& SetupInstances, const TArray<FName
 		if(PropertyChangedEvent.ChangeType != EPropertyChangeType::Unspecified && PropertyChangedEvent.ChangeType != EPropertyChangeType::ArrayRemove)
 		{
 			int32 CollisionCount = 0;
-			FName NewName = PostProfiles[ArrayIdx] == NAME_None ? FName(TEXT("New")) : PostProfiles[ArrayIdx];
+			FName NewName = (PostProfiles[ArrayIdx] == NAME_None) ? FName(TEXT("New")) : PostProfiles[ArrayIdx];
 			const FString NewNameNoNumber = NewName.ToString();
 			while(PreProfiles.Contains(NewName))
 			{
diff --git a/Engine/Source/Runtime/Engine/Private/ShaderCompiler/ShaderCompiler.cpp b/Engine/Source/Runtime/Engine/Private/ShaderCompiler/ShaderCompiler.cpp
--- a/Engine/Source/Runtime/Engine/Private/ShaderCompiler/ShaderCompiler.cpp
+++ b/Engine/Source/Runtime/Engine/Private/ShaderCompiler/ShaderCompiler.cpp
@@ -2238,7 +2238,7 @@ void FShaderCompileUtilities::ExecuteShaderCompileJob(FShaderCommonCompileJob& J
 	auto* SingleJob = Job.GetSingleShaderJob();
 	if (SingleJob)
 	{
-		const FName Format = SingleJob->Input.ShaderFormat != NAME_None ? SingleJob->Input.ShaderFormat : LegacyShaderPlatformToShaderFormat(EShaderPlatform(SingleJob->Input.Target.Platform));
+		const FName Format = (SingleJob->Input.ShaderFormat != NAME_None) ? SingleJob->Input.ShaderFormat : LegacyShaderPlatformToShaderFormat(EShaderPlatform(SingleJob->Input.Target.Platform));
 		const IShaderFormat* Compiler = TPM.FindShaderFormat(Format);
 
 		if (!Compiler)
diff --git a/Engine/Source/Runtime/Landscape/Private/LandscapeEditLayers.cpp b/Engine/Source/Runtime/Landscape/Private/LandscapeEditLayers.cpp
--- a/Engine/Source/Runtime/Landscape/Private/LandscapeEditLayers.cpp
+++ b/Engine/Source/Runtime/Landscape/Private/LandscapeEditLayers.cpp
@@ -6502,7 +6502,7 @@ FName ALandscape::GenerateUniqueLayerName(FName InName) const
 		return InName;
 	}
 
-	FString BaseName = InName == NAME_None ? "Layer" : InName.ToString();
+	FString BaseName = (InName == NAME_None) ? "Layer" : InName.ToString();
 	FName NewName;
 	int32 LayerIndex = 0;
 	do
diff --git a/Engine/Source/Runtime/Landscape/Private/LandscapeGrass.cpp b/Engine/Source/Runtime/Landscape/Private/LandscapeGrass.cpp
--- a/Engine/Source/Runtime/Landscape/Private/LandscapeGrass.cpp
+++ b/Engine/Source/Runtime/Landscape/Private/LandscapeGrass.cpp
@@ -1237,7 +1237,7 @@ void UMaterialExpressionLandscapeGrassOutput::ValidateInputName(FGrassInput& InI
 		bool bFoundValidName = false;
 
 		// Parameters cannot be named Name_None, use the default name instead
-		FName PotentialName = InInput.Name == NAME_None ? UMaterialExpressionLandscapeGrassOutput::PinDefaultName : InInput.Name;
+		FName PotentialName = (InInput.Name == NAME_None) ? UMaterialExpressionLandscapeGrassOutput::PinDefaultName : InInput.Name;
 
 		// Find an available unique name
 		while (!bFoundValidName)
diff --git a/Engine/Source/Runtime/LevelSequence/Private/LevelSequencePlayer.cpp b/Engine/Source/Runtime/LevelSequence/Private/LevelSequencePlayer.cpp
--- a/Engine/Source/Runtime/LevelSequence/Private/LevelSequencePlayer.cpp
+++ b/Engine/Source/Runtime/LevelSequence/Private/LevelSequencePlayer.cpp
@@ -87,7 +87,7 @@ void ULevelSequencePlayer::Initialize(ULevelSequence* InLevelSequence, ULevel* I
 	if (LevelStreaming)
 	{
 		// StreamedLevelPackage is a package name of the form /Game/Folder/MapName, not a full asset path
-		FString StreamedLevelPackage = (LevelStreaming->PackageNameToLoad == NAME_None ? LevelStreaming->GetWorldAssetPackageFName() : LevelStreaming->PackageNameToLoad).ToString();
+		FString StreamedLevelPackage = ((LevelStreaming->PackageNameToLoad == NAME_None) ? LevelStreaming->GetWorldAssetPackageFName() : LevelStreaming->PackageNameToLoad).ToString();
 
 		int32 SlashPos = 0;
 		if (StreamedLevelPackage.FindLastChar('/', SlashPos) && SlashPos < StreamedLevelPackage.Len()-1)
diff --git a/Engine/Source/Runtime/MovieScene/Private/Evaluation/MovieSceneMotionVectorSimulation.cpp b/Engine/Source/Runtime/MovieScene/Private/Evaluation/MovieSceneMotionVectorSimulation.cpp
--- a/Engine/Source/Runtime/MovieScene/Private/Evaluation/MovieSceneMotionVectorSimulation.cpp
+++ b/Engine/Source/Runtime/MovieScene/Private/Evaluation/MovieSceneMotionVectorSimulation.cpp
@@ -164,7 +164,7 @@ void FMovieSceneMotionVectorSimulation::ApplySimulatedTransforms(USceneComponent
 	{
 		FName AttachSocketName = Child->GetAttachSocketName();
 
-		FTransform SocketTransform = AttachSocketName == NAME_None ? FTransform::Identity : GetSocketTransform(InComponent, AttachSocketName);
+		FTransform SocketTransform = (AttachSocketName == NAME_None) ? FTransform::Identity : GetSocketTransform(InComponent, AttachSocketName);
 		FTransform ParentToWorld = SocketTransform * InPreviousTransform;
 
 		if (Child->IsUsingAbsoluteLocation())
diff --git a/Engine/Source/Runtime/Slate/Private/Framework/MultiBox/MultiBox.cpp b/Engine/Source/Runtime/Slate/Private/Framework/MultiBox/MultiBox.cpp
--- a/Engine/Source/Runtime/Slate/Private/Framework/MultiBox/MultiBox.cpp
+++ b/Engine/Source/Runtime/Slate/Private/Framework/MultiBox/MultiBox.cpp
@@ -165,7 +165,7 @@ TSharedRef<IMultiBlockBaseWidget> FMultiBlock::MakeWidget(TSharedRef<SMultiBoxWi
 
 	// Work out what style the widget should be using
 	const ISlateStyle* const StyleSet = InOwnerMultiBoxWidget->GetStyleSet();
-	FName StyleName = StyleNameOverride != NAME_None ? StyleNameOverride : InOwnerMultiBoxWidget->GetStyleName();
+	FName StyleName = (StyleNameOverride != NAME_None) ? StyleNameOverride : InOwnerMultiBoxWidget->GetStyleName();
 
 	// Build up the widget
 	NewMultiBlockWidget->BuildMultiBlockWidget(StyleSet, StyleName);
diff --git a/Engine/Source/Runtime/Slate/Private/Framework/MultiBox/MultiBoxBuilder.cpp b/Engine/Source/Runtime/Slate/Private/Framework/MultiBox/MultiBoxBuilder.cpp
--- a/Engine/Source/Runtime/Slate/Private/Framework/MultiBox/MultiBoxBuilder.cpp
+++ b/Engine/Source/Runtime/Slate/Private/Framework/MultiBox/MultiBoxBuilder.cpp
@@ -22,7 +22,7 @@
 
 
 FMultiBoxBuilder::FMultiBoxBuilder( const EMultiBoxType InType, FMultiBoxCustomization InCustomization, const bool bInShouldCloseWindowAfterMenuSelection, const TSharedPtr< const FUICommandList >& InCommandList, TSharedPtr<FExtender> InExtender, FName InTutorialHighlightName, FName InMenuName )
-	: MultiBox( FMultiBox::Create( InType, InMenuName != NAME_None ? FMultiBoxCustomization::AllowCustomization(InMenuName) : InCustomization, bInShouldCloseWindowAfterMenuSelection ) )
+	: MultiBox( FMultiBox::Create( InType, (InMenuName != NAME_None) ? FMultiBoxCustomization::AllowCustomization(InMenuName) : InCustomization, bInShouldCloseWindowAfterMenuSelection ) )
 	, CommandListStack()
 	, TutorialHighlightName(InTutorialHighlightName)
 	, MenuName(InMenuName)
diff --git a/Engine/Source/Runtime/SlateCore/Private/Styling/SlateStyleSet.cpp b/Engine/Source/Runtime/SlateCore/Private/Styling/SlateStyleSet.cpp
--- a/Engine/Source/Runtime/SlateCore/Private/Styling/SlateStyleSet.cpp
+++ b/Engine/Source/Runtime/SlateCore/Private/Styling/SlateStyleSet.cpp
@@ -362,7 +362,7 @@ const TSharedPtr< FSlateDynamicImageBrush > FSlateStyleSet::GetDynamicImageBrush
 {
 	//create a resource name
 	FName ResourceName;
-	ResourceName = TextureName == NAME_None ? BrushTemplate : FName(*( BrushTemplate.ToString() + TextureName.ToString() ));
+	ResourceName = (TextureName == NAME_None) ? BrushTemplate : FName(*( BrushTemplate.ToString() + TextureName.ToString() ));
 
 	//see if we already have that brush
 	TWeakPtr< FSlateDynamicImageBrush > WeakImageBrush = DynamicBrushes.FindRef(ResourceName);
diff --git a/Engine/Source/Runtime/SlateCore/Private/Textures/SlateIcon.cpp b/Engine/Source/Runtime/SlateCore/Private/Textures/SlateIcon.cpp
--- a/Engine/Source/Runtime/SlateCore/Private/Textures/SlateIcon.cpp
+++ b/Engine/Source/Runtime/SlateCore/Private/Textures/SlateIcon.cpp
@@ -20,7 +20,7 @@ FSlateIcon::FSlateIcon( )
 FSlateIcon::FSlateIcon(const FName InStyleSetName, const FName InStyleName, const FName InSmallStyleName, const FName InStatusOverlayStyleName)
 	: StyleSetName(InStyleSetName)
 	, StyleName(InStyleName)
-	, SmallStyleName(InSmallStyleName == NAME_None ? ISlateStyle::Join(InStyleName, ".Small") :  InSmallStyleName)
+	, SmallStyleName((InSmallStyleName == NAME_None) ? ISlateStyle::Join(InStyleName, ".Small") :  InSmallStyleName)
 	, StatusOverlayStyleName(InStatusOverlayStyleName)
 { }
 
diff --git a/Engine/Source/Runtime/StaticMeshDescription/Private/StaticMeshOperations.cpp b/Engine/Source/Runtime/StaticMeshDescription/Private/StaticMeshOperations.cpp
--- a/Engine/Source/Runtime/StaticMeshDescription/Private/StaticMeshOperations.cpp
+++ b/Engine/Source/Runtime/StaticMeshDescription/Private/StaticMeshOperations.cpp
@@ -890,7 +890,7 @@ void FStaticMeshOperations::ConvertFromRawMesh(const FRawMesh& SourceRawMesh, FM
 		if (PolygonGroupID == INDEX_NONE)
 		{
 			PolygonGroupID = DestinationMeshDescription.CreatePolygonGroup();
-			PolygonGroupImportedMaterialSlotNames[PolygonGroupID] = PolygonGroupImportedMaterialSlotName == NAME_None ? FName(*FString::Printf(TEXT("MaterialSlot_%d"), MaterialIndex)) : PolygonGroupImportedMaterialSlotName;
+			PolygonGroupImportedMaterialSlotNames[PolygonGroupID] = (PolygonGroupImportedMaterialSlotName == NAME_None) ? FName(*FString::Printf(TEXT("MaterialSlot_%d"), MaterialIndex)) : PolygonGroupImportedMaterialSlotName;
 			PolygonGroups.Add(PolygonGroupID);
 			MaterialIndexToPolygonGroup.Add(MaterialIndex, PolygonGroupID);
 		}
