Add the SupportsRayTracingDynamicGeometry flag to FNiagaraSpriteVertexFactoryEx

diff --git a/Engine/Plugins/FX/Niagara/Source/NiagaraVertexFactories/Private/NiagaraSpriteVertexFactory.cpp b/Engine/Plugins/FX/Niagara/Source/NiagaraVertexFactories/Private/NiagaraSpriteVertexFactory.cpp
--- a/Engine/Plugins/FX/Niagara/Source/NiagaraVertexFactories/Private/NiagaraSpriteVertexFactory.cpp
+++ b/Engine/Plugins/FX/Niagara/Source/NiagaraVertexFactories/Private/NiagaraSpriteVertexFactory.cpp
@@ -225,5 +225,6 @@ IMPLEMENT_VERTEX_FACTORY_TYPE(FNiagaraSpriteVertexFactoryEx,"/Plugin/FX/Niagara/
 	  EVertexFactoryFlags::UsedWithMaterials
 	| EVertexFactoryFlags::SupportsDynamicLighting
 	| EVertexFactoryFlags::SupportsRayTracing
+	| EVertexFactoryFlags::SupportsRayTracingDynamicGeometry
 	| EVertexFactoryFlags::SupportsPrecisePrevWorldPos
 );
