Fix missing style warnings coming from Niagara styles

diff --git a/Engine/Plugins/FX/Niagara/Source/NiagaraEditor/Private/NiagaraEditorStyle.cpp b/Engine/Plugins/FX/Niagara/Source/NiagaraEditor/Private/NiagaraEditorStyle.cpp
--- a/Engine/Plugins/FX/Niagara/Source/NiagaraEditor/Private/NiagaraEditorStyle.cpp
+++ b/Engine/Plugins/FX/Niagara/Source/NiagaraEditor/Private/NiagaraEditorStyle.cpp
@@ -50,7 +50,7 @@ NIAGARAEDITOR_API FString RelativePathToPluginPath(const FString& RelativePath,
 #define BOX_CORE_BRUSH( RelativePath, ... ) FSlateBoxBrush( FPaths::EngineContentDir() / "Editor/Slate" / RelativePath + TEXT(".png"), __VA_ARGS__ )
 #define BOX_PLUGIN_BRUSH( RelativePath, ... ) FSlateBoxBrush( RelativePathToPluginPath( RelativePath, ".png"), __VA_ARGS__ )
 #define BORDER_BRUSH( RelativePath, ... ) FSlateBorderBrush( Style->RootToContentDir( RelativePath, TEXT(".png") ), __VA_ARGS__ )
-#define DEFAULT_FONT(...) FCoreStyle::GetDefaultFontStyle(__VA_ARGS__)
+#define DEFAULT_FONT(...) FStarshipCoreStyle::GetDefaultFontStyle(__VA_ARGS__)
 #define IMAGE_PLUGIN_BRUSH_SVG( RelativePath, ... ) FSlateVectorImageBrush( RelativePathToPluginPath( RelativePath, ".svg" ), __VA_ARGS__ )
 
 const FVector2D Icon8x8(8.0f, 8.0f);
@@ -63,7 +63,7 @@ const FVector2D Icon64x64(64.0f, 64.0f);
 TSharedRef< FSlateStyleSet > FNiagaraEditorStyle::Create()
 {
 	const FTextBlockStyle NormalText = FEditorStyle::GetWidgetStyle<FTextBlockStyle>("NormalText");
-	const FEditableTextBoxStyle NormalEditableTextBox = FCoreStyle::Get().GetWidgetStyle<FEditableTextBoxStyle>("NormalEditableTextBox");
+	const FEditableTextBoxStyle NormalEditableTextBox = FStarshipCoreStyle::GetCoreStyle().GetWidgetStyle<FEditableTextBoxStyle>("NormalEditableTextBox");
 	const FSpinBoxStyle NormalSpinBox = FAppStyle::Get().GetWidgetStyle<FSpinBoxStyle>("SpinBox");
 
 	TSharedRef< FSlateStyleSet > Style = MakeShareable(new FSlateStyleSet("NiagaraEditorStyle"));
@@ -161,19 +161,20 @@ TSharedRef< FSlateStyleSet > FNiagaraEditorStyle::Create()
             .SetEvenRowBackgroundHoveredBrush(IMAGE_BRUSH("Common/Selection", Icon8x8, FLinearColor(1.0f, 1.0f, 1.0f, 0.1f)))
             .SetOddRowBackgroundBrush(FSlateNoResource())
             .SetOddRowBackgroundHoveredBrush(IMAGE_BRUSH("Common/Selection", Icon8x8, FLinearColor(1.0f, 1.0f, 1.0f, 0.1f)))
-            .SetSelectorFocusedBrush(BORDER_BRUSH("Common/Selector", FMargin(4.f / 16.f), FCoreStyle::Get().GetSlateColor("SelectorColor")))
-            .SetActiveBrush(IMAGE_BRUSH("Common/Selection", Icon8x8, FCoreStyle::Get().GetSlateColor("SelectionColor")))
-            .SetActiveHoveredBrush(IMAGE_BRUSH("Common/Selection", Icon8x8, FCoreStyle::Get().GetSlateColor("SelectionColor")))
+            .SetSelectorFocusedBrush(BORDER_BRUSH("Common/Selector", FMargin(4.f / 16.f), FStarshipCoreStyle::GetCoreStyle().GetSlateColor("SelectorColor")))
+            .SetActiveBrush(FSlateColorBrush(FStyleColors::Select))
+            .SetActiveHoveredBrush(FSlateColorBrush(FStyleColors::Select))
 			// note: inactive brush is used for selections that don't have keyboard focus. SelectionColor_Inactive would be gray. We use the normal SelectionColor here
-            .SetInactiveBrush(IMAGE_BRUSH("Common/Selection", Icon8x8, FCoreStyle::Get().GetSlateColor("SelectionColor")))
-            .SetInactiveHoveredBrush(IMAGE_BRUSH("Common/Selection", Icon8x8, FCoreStyle::Get().GetSlateColor("SelectionColor")))
-            .SetActiveHighlightedBrush(IMAGE_BRUSH("Common/Selection", Icon8x8, FCoreStyle::Get().GetSlateColor("HighlightColor")))
-            .SetInactiveHighlightedBrush(IMAGE_BRUSH("Common/Selection", Icon8x8, FSlateColor(FLinearColor(.1f, .1f, .1f))))
-            .SetTextColor(FCoreStyle::Get().GetSlateColor("DefaultForeground"))
-            .SetSelectedTextColor(FCoreStyle::Get().GetSlateColor("InvertedForeground"))
-            .SetDropIndicator_Above(BOX_BRUSH("Common/DropZoneIndicator_Above", FMargin(10.0f / 16.0f, 10.0f / 16.0f, 0, 0), FCoreStyle::Get().GetColor("SelectorColor")))
-            .SetDropIndicator_Onto(BOX_BRUSH("Common/DropZoneIndicator_Onto", FMargin(4.0f / 16.0f), FCoreStyle::Get().GetColor("SelectorColor")))
-            .SetDropIndicator_Below(BOX_BRUSH("Common/DropZoneIndicator_Below", FMargin(10.0f / 16.0f, 0, 0, 10.0f / 16.0f), FCoreStyle::Get().GetColor("SelectorColor")));
+            .SetInactiveBrush(FSlateColorBrush(FStyleColors::SelectInactive))
+            .SetInactiveHoveredBrush(FSlateColorBrush(FStyleColors::SelectInactive))
+            .SetActiveHighlightedBrush(FSlateColorBrush(FStyleColors::PrimaryHover))
+            .SetInactiveHighlightedBrush(FSlateColorBrush(FStyleColors::SelectParent))
+			.SetTextColor(FStyleColors::Foreground)
+			.SetSelectedTextColor(FStyleColors::ForegroundInverted)
+
+            .SetDropIndicator_Above(BOX_BRUSH("Common/DropZoneIndicator_Above", FMargin(10.0f / 16.0f, 10.0f / 16.0f, 0, 0), FStarshipCoreStyle::GetCoreStyle().GetColor("SelectorColor")))
+            .SetDropIndicator_Onto(BOX_BRUSH("Common/DropZoneIndicator_Onto", FMargin(4.0f / 16.0f), FStarshipCoreStyle::GetCoreStyle().GetColor("SelectorColor")))
+            .SetDropIndicator_Below(BOX_BRUSH("Common/DropZoneIndicator_Below", FMargin(10.0f / 16.0f, 0, 0, 10.0f / 16.0f), FStarshipCoreStyle::GetCoreStyle().GetColor("SelectorColor")));
 	
 	Style->Set("ActionMenu.Row", ActionMenuRowStyle);
 	
diff --git a/Engine/Plugins/FX/Niagara/Source/NiagaraEditor/Private/Widgets/SNiagaraBakerViewportToolbar.cpp b/Engine/Plugins/FX/Niagara/Source/NiagaraEditor/Private/Widgets/SNiagaraBakerViewportToolbar.cpp
--- a/Engine/Plugins/FX/Niagara/Source/NiagaraEditor/Private/Widgets/SNiagaraBakerViewportToolbar.cpp
+++ b/Engine/Plugins/FX/Niagara/Source/NiagaraEditor/Private/Widgets/SNiagaraBakerViewportToolbar.cpp
@@ -79,7 +79,7 @@ void SNiagaraBakerViewportToolbar::Construct(const FArguments& InArgs)
 		SNew(SEditorViewportToolBarButton)
 		.Cursor(EMouseCursor::Default)
 		.ButtonType(EUserInterfaceActionType::Button)
-		.ButtonStyle(&FEditorStyle::Get().GetWidgetStyle<FButtonStyle>("EditorViewportToolBar.MenuButtonWarning"))
+		.ButtonStyle(&FEditorStyle::Get().GetWidgetStyle<FButtonStyle>("EditorViewportToolBar.WarningButton"))
 		.OnClicked(
 			FOnClicked::CreateLambda(
 				[WeakViewModel=WeakViewModel]()
diff --git a/Engine/Plugins/FX/Niagara/Source/NiagaraEditor/Private/Widgets/SNiagaraSystemViewportToolBar.cpp b/Engine/Plugins/FX/Niagara/Source/NiagaraEditor/Private/Widgets/SNiagaraSystemViewportToolBar.cpp
--- a/Engine/Plugins/FX/Niagara/Source/NiagaraEditor/Private/Widgets/SNiagaraSystemViewportToolBar.cpp
+++ b/Engine/Plugins/FX/Niagara/Source/NiagaraEditor/Private/Widgets/SNiagaraSystemViewportToolBar.cpp
@@ -102,7 +102,7 @@ void SNiagaraSystemViewportToolBar::ExtendLeftAlignedToolbarSlots(TSharedPtr<SHo
 		SNew(SEditorViewportToolBarButton)
 		.Cursor(EMouseCursor::Default)
 		.ButtonType(EUserInterfaceActionType::Button)
-		.ButtonStyle(&FEditorStyle::Get().GetWidgetStyle<FButtonStyle>("EditorViewportToolBar.MenuButtonWarning"))
+		.ButtonStyle(&FEditorStyle::Get().GetWidgetStyle<FButtonStyle>("EditorViewportToolBar.WarningButton"))
 		.OnClicked(this, &SNiagaraSystemViewportToolBar::OnSimulationRealtimeWarningClicked)
 		.Visibility(this, &SNiagaraSystemViewportToolBar::GetSimulationRealtimeWarningVisibility)
 		.ToolTipText(LOCTEXT("SimulationRealtimeOff_ToolTip", "This simulation is not updating in real time.  Click to turn on real time."))
