Fix more niagara style warnings

diff --git a/Engine/Plugins/FX/Niagara/Source/NiagaraEditor/Private/NiagaraEditorStyle.cpp b/Engine/Plugins/FX/Niagara/Source/NiagaraEditor/Private/NiagaraEditorStyle.cpp
--- a/Engine/Plugins/FX/Niagara/Source/NiagaraEditor/Private/NiagaraEditorStyle.cpp
+++ b/Engine/Plugins/FX/Niagara/Source/NiagaraEditor/Private/NiagaraEditorStyle.cpp
@@ -49,6 +49,7 @@ NIAGARAEDITOR_API FString RelativePathToPluginPath(const FString& RelativePath,
 #define BOX_BRUSH( RelativePath, ... ) FSlateBoxBrush( Style->RootToContentDir( RelativePath, TEXT(".png") ), __VA_ARGS__ )
 #define BOX_CORE_BRUSH( RelativePath, ... ) FSlateBoxBrush( FPaths::EngineContentDir() / "Editor/Slate" / RelativePath + TEXT(".png"), __VA_ARGS__ )
 #define BOX_PLUGIN_BRUSH( RelativePath, ... ) FSlateBoxBrush( RelativePathToPluginPath( RelativePath, ".png"), __VA_ARGS__ )
+#define BORDER_CORE_BRUSH( RelativePath, ... ) FSlateBorderBrush( FPaths::EngineContentDir() / "Editor/Slate" / RelativePath + TEXT(".png"), __VA_ARGS__ )
 #define BORDER_BRUSH( RelativePath, ... ) FSlateBorderBrush( Style->RootToContentDir( RelativePath, TEXT(".png") ), __VA_ARGS__ )
 #define DEFAULT_FONT(...) FStarshipCoreStyle::GetDefaultFontStyle(__VA_ARGS__)
 #define IMAGE_PLUGIN_BRUSH_SVG( RelativePath, ... ) FSlateVectorImageBrush( RelativePathToPluginPath( RelativePath, ".svg" ), __VA_ARGS__ )
@@ -150,18 +151,18 @@ TSharedRef< FSlateStyleSet > FNiagaraEditorStyle::Create()
 	const FCheckBoxStyle NiagaraGraphActionMenuFilterCheckBox = FCheckBoxStyle()
             .SetCheckBoxType(ESlateCheckBoxType::ToggleButton)
             .SetUncheckedImage( FSlateNoResource() )
-            .SetUncheckedHoveredImage( BOX_BRUSH("Common/RoundedSelection_16x", 4.0f/16.0f, FLinearColor(0.7f, 0.7f, 0.7f) ))
-            .SetUncheckedPressedImage( BOX_BRUSH("Common/RoundedSelection_16x", 4.0f/16.0f, FLinearColor(0.8f, 0.8f, 0.8f) ))
-            .SetCheckedImage( BOX_BRUSH("Common/RoundedSelection_16x",  4.0f/16.0f, FLinearColor(0.9f, 0.9f, 0.9f) ))
-            .SetCheckedHoveredImage( BOX_BRUSH("Common/RoundedSelection_16x",  4.0f/16.0f, FLinearColor(1.f, 1.f, 1.f) ))
-            .SetCheckedPressedImage( BOX_BRUSH("Common/RoundedSelection_16x",  4.0f/16.0f, FLinearColor(1.f, 1.f, 1.f) ));
+            .SetUncheckedHoveredImage(BOX_CORE_BRUSH("Common/RoundedSelection_16x", 4.0f/16.0f, FLinearColor(0.7f, 0.7f, 0.7f) ))
+            .SetUncheckedPressedImage(BOX_CORE_BRUSH("Common/RoundedSelection_16x", 4.0f/16.0f, FLinearColor(0.8f, 0.8f, 0.8f) ))
+            .SetCheckedImage(BOX_CORE_BRUSH("Common/RoundedSelection_16x",  4.0f/16.0f, FLinearColor(0.9f, 0.9f, 0.9f) ))
+            .SetCheckedHoveredImage(BOX_CORE_BRUSH("Common/RoundedSelection_16x",  4.0f/16.0f, FLinearColor(1.f, 1.f, 1.f) ))
+            .SetCheckedPressedImage(BOX_CORE_BRUSH("Common/RoundedSelection_16x",  4.0f/16.0f, FLinearColor(1.f, 1.f, 1.f) ));
 
 	const FTableRowStyle ActionMenuRowStyle = FTableRowStyle()
             .SetEvenRowBackgroundBrush(FSlateNoResource())
-            .SetEvenRowBackgroundHoveredBrush(IMAGE_BRUSH("Common/Selection", Icon8x8, FLinearColor(1.0f, 1.0f, 1.0f, 0.1f)))
+            .SetEvenRowBackgroundHoveredBrush(IMAGE_CORE_BRUSH("Common/Selection", Icon8x8, FLinearColor(1.0f, 1.0f, 1.0f, 0.1f)))
             .SetOddRowBackgroundBrush(FSlateNoResource())
-            .SetOddRowBackgroundHoveredBrush(IMAGE_BRUSH("Common/Selection", Icon8x8, FLinearColor(1.0f, 1.0f, 1.0f, 0.1f)))
-            .SetSelectorFocusedBrush(BORDER_BRUSH("Common/Selector", FMargin(4.f / 16.f), FStarshipCoreStyle::GetCoreStyle().GetSlateColor("SelectorColor")))
+            .SetOddRowBackgroundHoveredBrush(IMAGE_CORE_BRUSH("Common/Selection", Icon8x8, FLinearColor(1.0f, 1.0f, 1.0f, 0.1f)))
+            .SetSelectorFocusedBrush(BORDER_CORE_BRUSH("Common/Selector", FMargin(4.f / 16.f), FStarshipCoreStyle::GetCoreStyle().GetSlateColor("SelectorColor")))
             .SetActiveBrush(FSlateColorBrush(FStyleColors::Select))
             .SetActiveHoveredBrush(FSlateColorBrush(FStyleColors::Select))
 			// note: inactive brush is used for selections that don't have keyboard focus. SelectionColor_Inactive would be gray. We use the normal SelectionColor here
@@ -172,9 +173,9 @@ TSharedRef< FSlateStyleSet > FNiagaraEditorStyle::Create()
 			.SetTextColor(FStyleColors::Foreground)
 			.SetSelectedTextColor(FStyleColors::ForegroundInverted)
 
-            .SetDropIndicator_Above(BOX_BRUSH("Common/DropZoneIndicator_Above", FMargin(10.0f / 16.0f, 10.0f / 16.0f, 0, 0), FStarshipCoreStyle::GetCoreStyle().GetColor("SelectorColor")))
-            .SetDropIndicator_Onto(BOX_BRUSH("Common/DropZoneIndicator_Onto", FMargin(4.0f / 16.0f), FStarshipCoreStyle::GetCoreStyle().GetColor("SelectorColor")))
-            .SetDropIndicator_Below(BOX_BRUSH("Common/DropZoneIndicator_Below", FMargin(10.0f / 16.0f, 0, 0, 10.0f / 16.0f), FStarshipCoreStyle::GetCoreStyle().GetColor("SelectorColor")));
+            .SetDropIndicator_Above(BOX_CORE_BRUSH("Common/DropZoneIndicator_Above", FMargin(10.0f / 16.0f, 10.0f / 16.0f, 0, 0), FStarshipCoreStyle::GetCoreStyle().GetColor("SelectorColor")))
+            .SetDropIndicator_Onto(BOX_CORE_BRUSH("Common/DropZoneIndicator_Onto", FMargin(4.0f / 16.0f), FStarshipCoreStyle::GetCoreStyle().GetColor("SelectorColor")))
+            .SetDropIndicator_Below(BOX_CORE_BRUSH("Common/DropZoneIndicator_Below", FMargin(10.0f / 16.0f, 0, 0, 10.0f / 16.0f), FStarshipCoreStyle::GetCoreStyle().GetColor("SelectorColor")));
 	
 	Style->Set("ActionMenu.Row", ActionMenuRowStyle);
 	
@@ -348,9 +349,6 @@ TSharedRef< FSlateStyleSet > FNiagaraEditorStyle::Create()
 		.SetColorAndOpacity(FLinearColor(0.72f, 0.72f, 0.72f))
 		.SetHighlightColor(FLinearColor(1, 1, 1)));
 
-	// Separator in the action menus
-	Style->Set( "MenuSeparator", new BOX_BRUSH( "MenuSeparator", 1/4.0f, FLinearColor(1,1,1,0.2f) ) );
-	
 	const FString SmallRoundedButtonStart(TEXT("Common/SmallRoundedButtonLeft"));
 	const FString SmallRoundedButtonMiddle(TEXT("Common/SmallRoundedButtonCentre"));
 	const FString SmallRoundedButtonEnd(TEXT("Common/SmallRoundedButtonRight"));
diff --git a/Engine/Plugins/FX/Niagara/Source/NiagaraEditor/Private/Widgets/SNiagaraGraphActionWidget.cpp b/Engine/Plugins/FX/Niagara/Source/NiagaraEditor/Private/Widgets/SNiagaraGraphActionWidget.cpp
--- a/Engine/Plugins/FX/Niagara/Source/NiagaraEditor/Private/Widgets/SNiagaraGraphActionWidget.cpp
+++ b/Engine/Plugins/FX/Niagara/Source/NiagaraEditor/Private/Widgets/SNiagaraGraphActionWidget.cpp
@@ -142,7 +142,7 @@ void SNiagaraActionWidget::Construct(const FArguments& InArgs, const FCreateNiag
        .VAlign(VAlign_Fill)
        [
            SNew(SSeparator)
-           .SeparatorImage(FNiagaraEditorStyle::Get().GetBrush("MenuSeparator"))
+           .SeparatorImage(FAppStyle::Get().GetBrush("Separator"))
            .Orientation(EOrientation::Orient_Vertical)
            .Visibility(ActionPtr->SourceData.bDisplaySource ? EVisibility::Visible : EVisibility::Collapsed)
        ]
diff --git a/Engine/Plugins/FX/Niagara/Source/NiagaraEditorWidgets/Private/SNiagaraCurveOverview.cpp b/Engine/Plugins/FX/Niagara/Source/NiagaraEditorWidgets/Private/SNiagaraCurveOverview.cpp
--- a/Engine/Plugins/FX/Niagara/Source/NiagaraEditorWidgets/Private/SNiagaraCurveOverview.cpp
+++ b/Engine/Plugins/FX/Niagara/Source/NiagaraEditorWidgets/Private/SNiagaraCurveOverview.cpp
@@ -501,8 +501,7 @@ void SNiagaraCurveOverview::Construct(const FArguments& InArgs, TSharedRef<FNiag
 
 	RefreshCurveEditorTreeItems();
 
-	FToolBarBuilder ToolBarBuilder(CurveEditorPanel->GetCommands(), FMultiBoxCustomization::None, CurveEditorPanel->GetToolbarExtender(), true);
-	ToolBarBuilder.SetStyle(&FEditorStyle::Get(), "Sequencer.ToolBar");
+	FSlimHorizontalToolBarBuilder ToolBarBuilder(CurveEditorPanel->GetCommands(), FMultiBoxCustomization::None, CurveEditorPanel->GetToolbarExtender(), true);
 	ToolBarBuilder.BeginSection("Asset");
 	ToolBarBuilder.EndSection();
 
