Last set of missing niagara style warnings

diff --git a/Engine/Plugins/FX/Niagara/Source/NiagaraEditor/Private/NiagaraEditorStyle.cpp b/Engine/Plugins/FX/Niagara/Source/NiagaraEditor/Private/NiagaraEditorStyle.cpp
--- a/Engine/Plugins/FX/Niagara/Source/NiagaraEditor/Private/NiagaraEditorStyle.cpp
+++ b/Engine/Plugins/FX/Niagara/Source/NiagaraEditor/Private/NiagaraEditorStyle.cpp
@@ -173,9 +173,9 @@ TSharedRef< FSlateStyleSet > FNiagaraEditorStyle::Create()
 			.SetTextColor(FStyleColors::Foreground)
 			.SetSelectedTextColor(FStyleColors::ForegroundInverted)
 
-            .SetDropIndicator_Above(BOX_CORE_BRUSH("Common/DropZoneIndicator_Above", FMargin(10.0f / 16.0f, 10.0f / 16.0f, 0, 0), FStarshipCoreStyle::GetCoreStyle().GetColor("SelectorColor")))
-            .SetDropIndicator_Onto(BOX_CORE_BRUSH("Common/DropZoneIndicator_Onto", FMargin(4.0f / 16.0f), FStarshipCoreStyle::GetCoreStyle().GetColor("SelectorColor")))
-            .SetDropIndicator_Below(BOX_CORE_BRUSH("Common/DropZoneIndicator_Below", FMargin(10.0f / 16.0f, 0, 0, 10.0f / 16.0f), FStarshipCoreStyle::GetCoreStyle().GetColor("SelectorColor")));
+            .SetDropIndicator_Above(BOX_CORE_BRUSH("Common/DropZoneIndicator_Above", FMargin(10.0f / 16.0f, 10.0f / 16.0f, 0, 0), FStarshipCoreStyle::GetCoreStyle().GetSlateColor("SelectorColor")))
+            .SetDropIndicator_Onto(BOX_CORE_BRUSH("Common/DropZoneIndicator_Onto", FMargin(4.0f / 16.0f), FStarshipCoreStyle::GetCoreStyle().GetSlateColor("SelectorColor")))
+            .SetDropIndicator_Below(BOX_CORE_BRUSH("Common/DropZoneIndicator_Below", FMargin(10.0f / 16.0f, 0, 0, 10.0f / 16.0f), FStarshipCoreStyle::GetCoreStyle().GetSlateColor("SelectorColor")));
 	
 	Style->Set("ActionMenu.Row", ActionMenuRowStyle);
 	
@@ -433,9 +433,9 @@ TSharedRef< FSlateStyleSet > FNiagaraEditorStyle::Create()
 		Style->Set("NiagaraEditor.Debugger.Outliner.Filter", new IMAGE_PLUGIN_BRUSH("Icons/Debugger/Filter_24x", Icon24x24));
 
 		FButtonStyle OutlinerToolBarButton = FButtonStyle()
-			.SetNormal(BOX_BRUSH("Common/FlatButton", FMargin(4 / 16.0f), FLinearColor(0.0f, 0.0f, 0.0f, 0.0f)))
-			.SetHovered(BOX_BRUSH("Common/FlatButton", FMargin(4 / 16.0f), SelectionColor))
-			.SetPressed(BOX_BRUSH("Common/FlatButton", FMargin(4 / 16.0f), SelectionColor_Pressed))
+			.SetNormal(BOX_CORE_BRUSH("Common/FlatButton", FMargin(4 / 16.0f), FLinearColor(0.0f, 0.0f, 0.0f, 0.0f)))
+			.SetHovered(BOX_CORE_BRUSH("Common/FlatButton", FMargin(4 / 16.0f), SelectionColor))
+			.SetPressed(BOX_CORE_BRUSH("Common/FlatButton", FMargin(4 / 16.0f), SelectionColor_Pressed))
 			.SetNormalPadding(FMargin(0, 0, 0, 1))
 			.SetPressedPadding(FMargin(0, 1, 0, 0));
 		Style->Set("NiagaraEditor.Debugger.Outliner.Toolbar", OutlinerToolBarButton);
