Merging CL 16162815:
Fixed non-unity build error

diff --git a/Engine/Plugins/FX/Niagara/Source/NiagaraEditor/Private/Widgets/SNiagaraScriptVersionWidget.cpp b/Engine/Plugins/FX/Niagara/Source/NiagaraEditor/Private/Widgets/SNiagaraScriptVersionWidget.cpp
--- a/Engine/Plugins/FX/Niagara/Source/NiagaraEditor/Private/Widgets/SNiagaraScriptVersionWidget.cpp
+++ b/Engine/Plugins/FX/Niagara/Source/NiagaraEditor/Private/Widgets/SNiagaraScriptVersionWidget.cpp
@@ -4,14 +4,17 @@
 
 #include "AssetToolsModule.h"
 #include "DetailLayoutBuilder.h"
+#include "Editor.h"
 #include "IDetailsView.h"
 #include "PropertyEditorModule.h"
 #include "SGraphActionMenu.h"
 #include "SlateOptMacros.h"
+#include "Editor/EditorEngine.h"
 #include "Framework/MultiBox/MultiBoxBuilder.h"
 #include "Modules/ModuleManager.h"
 #include "NiagaraEditor/Private/NiagaraVersionMetaData.h"
 #include "NiagaraEditor/Public/NiagaraActions.h"
+#include "Subsystems/AssetEditorSubsystem.h"
 #include "Widgets/Images/SImage.h"
 #include "Widgets/Input/SButton.h"
 #include "Widgets/Layout/SBox.h"
