Updates to the niagara ray tracing interface
-support for indirect dispatch
-new function for reserving trace slots and another for having the slots created by the DI
-moved resources over to the template

diff --git a/Engine/Plugins/FX/Niagara/Shaders/Private/NiagaraDataInterfaceCollisionQuery.ush b/Engine/Plugins/FX/Niagara/Shaders/Private/NiagaraDataInterfaceCollisionQuery.ush
--- a/Engine/Plugins/FX/Niagara/Shaders/Private/NiagaraDataInterfaceCollisionQuery.ush
+++ b/Engine/Plugins/FX/Niagara/Shaders/Private/NiagaraDataInterfaceCollisionQuery.ush
@@ -5,13 +5,7 @@
 #include "/Engine/Private/RayTracing/RayTracingCommon.ush"
 #include "NiagaraRayTracingCommon.ush"
 
-#define NDICOLLISIONQUERY_USE_GBUFFER_NORMAL	(FEATURE_LEVEL >= FEATURE_LEVEL_SM5)
-
-#if FEATURE_LEVEL >= FEATURE_LEVEL_SM5
-int											MaxRayTraceCount;
-RWStructuredBuffer<FBasicRayData>			RWRayRequests;
-StructuredBuffer<FNiagaraRayTracingPayload>	IntersectionResults;
-#endif
+#define NDICOLLISIONQUERY_USE_GBUFFER_NORMAL    (FEATURE_LEVEL >= FEATURE_LEVEL_SM5)
 
 float NDICollisionQuery_GetSceneDepth(float2 ScreenUV)
 {
@@ -145,34 +139,85 @@ void NDICollisionQuery_QueryMeshDistanceFieldGPU(in float3 In_SamplePos, out flo
 #endif
 }
 
-void NDICollisionQuery_IssueAsyncRayTrace(int In_QueryID, float3 In_TraceStart, float3 In_TraceEnd, int In_TraceChannel, out bool Out_IsQueryValid)
+#if NIAGARA_USE_RAYTRACING
+
+uint BuildRayTraceMask(int TraceChannel)
+{
+	// for now we'll ignore the TraceChannel till we can mark up the environment
+	return 0xFF;
+}
+
+void NDICollisionQuery_IssueAsyncRayTrace(uint MaxRayTraceCount, RWStructuredBuffer<FBasicRayData> RWRayRequests, RWBuffer<uint> RWRayTraceCounts, int In_QueryID, float3 In_TraceStart, float3 In_TraceEnd, int In_TraceChannel, out bool Out_IsQueryValid)
 {
 	Out_IsQueryValid = false;
 
-#if FEATURE_LEVEL >= FEATURE_LEVEL_SM5
-	if (In_QueryID >= 0 && In_QueryID < MaxRayTraceCount)
+	if (In_QueryID >= 0 && uint(In_QueryID) < MaxRayTraceCount)
 	{
 		float3 Trace = In_TraceEnd - In_TraceStart;
 		float TraceLength = length(Trace);
 
 		RWRayRequests[In_QueryID].Origin = In_TraceStart;
-		RWRayRequests[In_QueryID].Mask = 0xFF; // In_TraceChannel - for now we'll ignore the TraceChannel till we can mark up the environment
+		RWRayRequests[In_QueryID].Mask = BuildRayTraceMask(In_TraceChannel);
 		RWRayRequests[In_QueryID].Direction = normalize(Trace);
 		RWRayRequests[In_QueryID].TFar = TraceLength;
 		Out_IsQueryValid = true;
+
+		// we want to keep track of the maximum index used
+		InterlockedMax(RWRayTraceCounts[0], In_QueryID + 1);
+		RWRayTraceCounts[1] = 1;
+		RWRayTraceCounts[2] = 1;
 	}
-#endif
 }
 
-void NDICollisionQuery_ReadAsyncRayTrace(int In_PreviousFrameQueryID, out bool Out_CollisionValid, out float Out_CollisionDistance, out float3 Out_CollisionPosWorld, out float3 Out_CollisionNormal)
+void NDICollisionQuery_ReserveRayTraceIndex(uint MaxRayTraceCount, RWBuffer<uint> RWRayTraceCounts, int In_TraceCount, out int Out_StartQueryID, out bool Out_IndicesValid)
+{
+	Out_IndicesValid = false;
+
+	if ((In_TraceCount + RWRayTraceCounts[0]) < MaxRayTraceCount)
+	{
+		uint RequestIndex = 0;
+		InterlockedAdd(RWRayTraceCounts[0], In_TraceCount, RequestIndex);
+		RWRayTraceCounts[1] = 1;
+		RWRayTraceCounts[2] = 1;
+
+		if ((In_TraceCount + RequestIndex) <= MaxRayTraceCount)
+		{
+			Out_StartQueryID = int(RequestIndex);
+			Out_IndicesValid = true;
+		}
+		else
+		{
+			InterlockedMax(RWRayTraceCounts[0], MaxRayTraceCount);
+		}
+	}
+}
+
+void NDICollisionQuery_CreateAsyncRayTrace(uint MaxRayTraceCount, RWStructuredBuffer<FBasicRayData> RWRayRequests, RWBuffer<uint> RWRayTraceCounts, float3 In_TraceStart, float3 In_TraceEnd, int In_TraceChannel, out int Out_QueryID, out bool Out_IsQueryValid)
+{
+	Out_IsQueryValid = false;
+
+	NDICollisionQuery_ReserveRayTraceIndex(MaxRayTraceCount, RWRayTraceCounts, 1, Out_QueryID, Out_IsQueryValid);
+
+	if (Out_IsQueryValid)
+	{
+		float3 Trace = In_TraceEnd - In_TraceStart;
+		float TraceLength = length(Trace);
+
+		RWRayRequests[Out_QueryID].Origin = In_TraceStart;
+		RWRayRequests[Out_QueryID].Mask = BuildRayTraceMask(In_TraceChannel);
+		RWRayRequests[Out_QueryID].Direction = normalize(Trace);
+		RWRayRequests[Out_QueryID].TFar = TraceLength;
+	}
+}
+
+void NDICollisionQuery_ReadAsyncRayTrace(uint MaxRayTraceCount, StructuredBuffer<FNiagaraRayTracingPayload> IntersectionResults, int In_PreviousFrameQueryID, out bool Out_CollisionValid, out float Out_CollisionDistance, out float3 Out_CollisionPosWorld, out float3 Out_CollisionNormal)
 {
 	Out_CollisionValid = false;
 	Out_CollisionDistance = 0.0f;
 	Out_CollisionPosWorld = float3(0.0f, 0.0f, 0.0f);
 	Out_CollisionNormal = float3(0.0f, 0.0f, 0.0f);
 
-#if FEATURE_LEVEL >= FEATURE_LEVEL_SM5
-	if (In_PreviousFrameQueryID >= 0 && In_PreviousFrameQueryID < MaxRayTraceCount)
+	if (In_PreviousFrameQueryID >= 0 && uint(In_PreviousFrameQueryID) < MaxRayTraceCount)
 	{
 		FNiagaraRayTracingPayload HitResult = IntersectionResults[In_PreviousFrameQueryID];
 
@@ -184,5 +229,6 @@ void NDICollisionQuery_ReadAsyncRayTrace(int In_PreviousFrameQueryID, out bool O
 			Out_CollisionNormal = HitResult.WorldNormal;
 		}
 	}
-#endif
-}
\ No newline at end of file
+}
+
+#endif
\ No newline at end of file
diff --git a/Engine/Plugins/FX/Niagara/Shaders/Private/NiagaraDataInterfaceCollisionQueryTemplate.ush b/Engine/Plugins/FX/Niagara/Shaders/Private/NiagaraDataInterfaceCollisionQueryTemplate.ush
--- a/Engine/Plugins/FX/Niagara/Shaders/Private/NiagaraDataInterfaceCollisionQueryTemplate.ush
+++ b/Engine/Plugins/FX/Niagara/Shaders/Private/NiagaraDataInterfaceCollisionQueryTemplate.ush
@@ -1,5 +1,12 @@
 // Copyright Epic Games, Inc. All Rights Reserved.
 
+#if NIAGARA_USE_RAYTRACING
+uint                                            MaxRayTraceCount_{ParameterName};
+RWStructuredBuffer<FBasicRayData>               RWRayRequests_{ParameterName};
+StructuredBuffer<FNiagaraRayTracingPayload>     IntersectionResults_{ParameterName};
+RWBuffer<uint>                                  RWRayTraceCounts_{ParameterName};
+#endif
+
 void QuerySceneDepthGPU_{ParameterName}(in float3 In_SamplePos, out float Out_SceneDepth, out float3 Out_CameraPosWorld, out bool Out_IsInsideView, out float3 Out_WorldPos, out float3 Out_WorldNormal)
 {
 	NDICollisionQuery_QuerySceneDepthGPU(In_SamplePos, Out_SceneDepth, Out_CameraPosWorld, Out_IsInsideView, Out_WorldPos, Out_WorldNormal);
@@ -17,10 +24,36 @@ void QueryMeshDistanceFieldGPU_{ParameterName}(in float3 In_SamplePos, out float
 
 void IssueAsyncRayTraceGpu_{ParameterName}_UEImpureCall(int In_QueryID, float3 In_TraceStart, float3 In_TraceEnd, int In_TraceChannel, out bool Out_IsQueryValid)
 {
-	NDICollisionQuery_IssueAsyncRayTrace(In_QueryID, In_TraceStart, In_TraceEnd, In_TraceChannel, Out_IsQueryValid);
+#if NIAGARA_USE_RAYTRACING
+	NDICollisionQuery_IssueAsyncRayTrace(MaxRayTraceCount_{ParameterName}, RWRayRequests_{ParameterName}, RWRayTraceCounts_{ParameterName}, In_QueryID, In_TraceStart, In_TraceEnd, In_TraceChannel, Out_IsQueryValid);
+#else
+	Out_IsQueryValid = false;
+#endif
 }
 
-void ReadAsyncRayTraceGpu_{ParameterName}(int In_PreviousFrameQueryID, out bool Out_CollisionValid, out float Out_CollisionDistance, out float3 Out_CollisionPosWorld, out float3 Out_CollisionNormal)
+void ReserveRayTraceIndex_{ParameterName}_UEImpureCall(int In_TraceCount, out int Out_StartQueryID, out bool Out_IndicesValid)
+{
+#if NIAGARA_USE_RAYTRACING
+	NDICollisionQuery_ReserveRayTraceIndex(MaxRayTraceCount_{ParameterName}, RWRayTraceCounts_{ParameterName}, In_TraceCount, Out_StartQueryID, Out_IndicesValid);
+#else
+	Out_IndicesValid = false;
+#endif
+}
+
+void CreateAsyncRayTraceGpu_{ParameterName}_UEImpureCall(float3 In_TraceStart, float3 In_TraceEnd, int In_TraceChannel, out int Out_QueryID, out bool Out_IsQueryValid)
 {
-	NDICollisionQuery_ReadAsyncRayTrace(In_PreviousFrameQueryID, Out_CollisionValid, Out_CollisionDistance, Out_CollisionPosWorld, Out_CollisionNormal);
+#if NIAGARA_USE_RAYTRACING
+	NDICollisionQuery_CreateAsyncRayTrace(MaxRayTraceCount_{ParameterName}, RWRayRequests_{ParameterName}, RWRayTraceCounts_{ParameterName}, In_TraceStart, In_TraceEnd, In_TraceChannel, Out_QueryID, Out_IsQueryValid);
+#else
+	Out_IsQueryValid = false;
+#endif
 }
+
+void ReadAsyncRayTraceGpu_{ParameterName}(int In_PreviousFrameQueryID, out bool Out_CollisionValid, out float Out_CollisionDistance, out float3 Out_CollisionPosWorld, out float3 Out_CollisionNormal)
+{
+#if NIAGARA_USE_RAYTRACING
+	NDICollisionQuery_ReadAsyncRayTrace(MaxRayTraceCount_{ParameterName}, IntersectionResults_{ParameterName}, In_PreviousFrameQueryID, Out_CollisionValid, Out_CollisionDistance, Out_CollisionPosWorld, Out_CollisionNormal);
+#else
+	Out_CollisionValid = false;
+#endif
+}
\ No newline at end of file
diff --git a/Engine/Plugins/FX/Niagara/Shaders/Private/NiagaraRayTracingCommon.ush b/Engine/Plugins/FX/Niagara/Shaders/Private/NiagaraRayTracingCommon.ush
--- a/Engine/Plugins/FX/Niagara/Shaders/Private/NiagaraRayTracingCommon.ush
+++ b/Engine/Plugins/FX/Niagara/Shaders/Private/NiagaraRayTracingCommon.ush
@@ -4,6 +4,8 @@
 
 #include "/Engine/Private/RayTracing/RayTracingCommon.ush"
 
+#define NIAGARA_USE_RAYTRACING        (FEATURE_LEVEL >= FEATURE_LEVEL_SM5)
+
 // payload structure for Niagara ray traces
 struct FNiagaraRayTracingPayload : FMinimalPayload
 {
diff --git a/Engine/Plugins/FX/Niagara/Shaders/Private/NiagaraRayTracingShaders.usf b/Engine/Plugins/FX/Niagara/Shaders/Private/NiagaraRayTracingShaders.usf
--- a/Engine/Plugins/FX/Niagara/Shaders/Private/NiagaraRayTracingShaders.usf
+++ b/Engine/Plugins/FX/Niagara/Shaders/Private/NiagaraRayTracingShaders.usf
@@ -10,9 +10,21 @@ RaytracingAccelerationStructure               TLAS;
 StructuredBuffer<FBasicRayData>               Rays;
 RWStructuredBuffer<FNiagaraRayTracingPayload> CollisionOutput;
 
+#if NIAGARA_RAYTRACE_FAKE_INDIRECT
+Buffer<uint>                                  RayTraceCounts;
+#endif
+
 RAY_TRACING_ENTRY_RAYGEN(NiagaraCollisionRayTraceRG)
 {
 	const uint RayIndex = DispatchRaysIndex().x;
+
+#if NIAGARA_RAYTRACE_FAKE_INDIRECT
+	if (RayIndex >= RayTraceCounts[0])
+	{
+		return;
+	}
+#endif
+
 	FBasicRayData InputRay = Rays[RayIndex];
 
 	RayDesc Ray;
diff --git a/Engine/Plugins/FX/Niagara/Source/Niagara/Classes/NiagaraEmitterInstanceBatcher.h b/Engine/Plugins/FX/Niagara/Source/Niagara/Classes/NiagaraEmitterInstanceBatcher.h
--- a/Engine/Plugins/FX/Niagara/Source/Niagara/Classes/NiagaraEmitterInstanceBatcher.h
+++ b/Engine/Plugins/FX/Niagara/Source/Niagara/Classes/NiagaraEmitterInstanceBatcher.h
@@ -176,7 +176,7 @@ public:
 	FNiagaraGpuReadbackManager* GetGpuReadbackManager() const { return GpuReadbackManagerPtr.Get(); }
 
 #if RHI_RAYTRACING
-	void IssueRayTraces(FRHICommandList& RHICmdList, const FIntPoint& RayTraceCounts, FRHIShaderResourceView* RayTraceRequests, FRHIUnorderedAccessView* RayTraceResults) const;
+	void IssueRayTraces(FRHICommandList& RHICmdList, const FIntPoint& RayTraceCounts, FRHIShaderResourceView* RayTraceRequests, FRWBuffer* IndirectArgsBuffer, uint32 IndirectArgsOffset, FRHIUnorderedAccessView* RayTraceResults) const;
 	bool HasRayTracingScene() const;
 #endif
 
diff --git a/Engine/Plugins/FX/Niagara/Source/Niagara/Private/NiagaraDataInterfaceCollisionQuery.cpp b/Engine/Plugins/FX/Niagara/Source/Niagara/Private/NiagaraDataInterfaceCollisionQuery.cpp
--- a/Engine/Plugins/FX/Niagara/Source/Niagara/Private/NiagaraDataInterfaceCollisionQuery.cpp
+++ b/Engine/Plugins/FX/Niagara/Source/Niagara/Private/NiagaraDataInterfaceCollisionQuery.cpp
@@ -29,7 +29,14 @@ namespace NDICollisionQueryLocal
 	static const FName SyncTraceName(TEXT("PerformCollisionQuerySyncCPU"));
 	static const FName AsyncTraceName(TEXT("PerformCollisionQueryAsyncCPU"));
 	static const FName IssueAsyncRayTraceName(TEXT("IssueAsyncRayTraceGpu"));
+	static const FName CreateAsyncRayTraceName(TEXT("CreateAsyncRayTraceGpu"));
+	static const FName ReserveAsyncRayTraceName(TEXT("ReserveAsyncRayTraceGpu"));
 	static const FName ReadAsyncRayTraceName(TEXT("ReadAsyncRayTraceGpu"));
+
+	static const FString MaxRayTraceCountParamName(TEXT("MaxRayTraceCount_"));
+	static const FString RayRequestsParamName(TEXT("RayRequests_"));
+	static const FString IntersectionResultsParamName(TEXT("IntersectionResults_"));
+	static const FString RayTraceCountsParamName(TEXT("RayTraceCounts_"));
 }
 
 FCriticalSection UNiagaraDataInterfaceCollisionQuery::CriticalSection;
@@ -52,6 +59,7 @@ struct FNiagaraDataIntefaceProxyCollisionQuery : public FNiagaraDataInterfacePro
 #if RHI_RAYTRACING
 	FRWBufferStructured RayTraceRequests;
 	FRWBufferStructured RayTraceIntersections;
+	FRWBuffer RayTraceCounts;
 
 	int32 MaxRayTraceCount = 0;
 #endif
@@ -64,6 +72,9 @@ struct FNiagaraDataIntefaceProxyCollisionQuery : public FNiagaraDataInterfacePro
 
 		DEC_MEMORY_STAT_BY(STAT_NiagaraGPUDataInterfaceMemory, RayTraceIntersections.NumBytes);
 		RayTraceIntersections.Release();
+
+		DEC_MEMORY_STAT_BY(STAT_NiagaraGPUDataInterfaceMemory, RayTraceCounts.NumBytes);
+		RayTraceCounts.Release();
 #endif
 	}
 
@@ -72,27 +83,43 @@ struct FNiagaraDataIntefaceProxyCollisionQuery : public FNiagaraDataInterfacePro
 		return 0;
 	}
 
+	virtual void PreStage(FRHICommandList& RHICmdList, const FNiagaraDataInterfaceStageArgs& Context) override
+	{
+		FNiagaraDataInterfaceProxy::PreStage(RHICmdList, Context);
+
+#if RHI_RAYTRACING
+		if (MaxRayTraceCount > 0)
+		{
+			// clear out indirect args
+			RHICmdList.ClearUAVUint(RayTraceCounts.UAV, FUintVector4(0, 0, 0, 0));
+			RHICmdList.Transition(FRHITransitionInfo(RayTraceCounts.UAV, ERHIAccess::Unknown, ERHIAccess::UAVCompute));
+		}
+#endif
+	}
+
 	virtual void PostSimulate(FRHICommandList& RHICmdList, const FNiagaraDataInterfaceArgs& Context) override
 	{
+		FNiagaraDataInterfaceProxy::PostSimulate(RHICmdList, Context);
+
 #if RHI_RAYTRACING
 		if (MaxRayTraceCount > 0 && Context.Batcher->HasRayTracingScene())
 		{
 			{
 				FRHITransitionInfo PreTransitions[] =
 				{
-					FRHITransitionInfo(RayTraceRequests.UAV, ERHIAccess::UAVMask, ERHIAccess::SRVMask),
-					FRHITransitionInfo(RayTraceIntersections.UAV, ERHIAccess::SRVMask, ERHIAccess::UAVMask)
+					FRHITransitionInfo(RayTraceRequests.UAV, ERHIAccess::UAVCompute, ERHIAccess::SRVCompute),
+					FRHITransitionInfo(RayTraceIntersections.UAV, ERHIAccess::SRVCompute, ERHIAccess::UAVCompute)
 				};
 				RHICmdList.Transition(PreTransitions);
 			}
 
-			Context.Batcher->IssueRayTraces(RHICmdList, FIntPoint(MaxRayTraceCount, 1), RayTraceRequests.SRV, RayTraceIntersections.UAV);
+			Context.Batcher->IssueRayTraces(RHICmdList, FIntPoint(MaxRayTraceCount, 1), RayTraceRequests.SRV, &RayTraceCounts, 0, RayTraceIntersections.UAV);
 
 			{
 				FRHITransitionInfo PostTransitions[] =
 				{
-					FRHITransitionInfo(RayTraceRequests.UAV, ERHIAccess::SRVMask, ERHIAccess::UAVMask),
-					FRHITransitionInfo(RayTraceIntersections.UAV, ERHIAccess::UAVMask, ERHIAccess::SRVMask)
+					FRHITransitionInfo(RayTraceRequests.UAV, ERHIAccess::SRVCompute, ERHIAccess::UAVCompute),
+					FRHITransitionInfo(RayTraceIntersections.UAV, ERHIAccess::UAVCompute, ERHIAccess::SRVCompute)
 				};
 				RHICmdList.Transition(PostTransitions);
 			}
@@ -111,6 +138,9 @@ struct FNiagaraDataIntefaceProxyCollisionQuery : public FNiagaraDataInterfacePro
 		DEC_MEMORY_STAT_BY(STAT_NiagaraGPUDataInterfaceMemory, RayTraceIntersections.NumBytes);
 		RayTraceIntersections.Release();
 
+		DEC_MEMORY_STAT_BY(STAT_NiagaraGPUDataInterfaceMemory, RayTraceCounts.NumBytes);
+		RayTraceCounts.Release();
+
 		if (IsRayTracingEnabled() && InMaxRayTraceRequests > 0)
 		{
 			MaxRayTraceCount = 16 * FMath::DivideAndRoundUp(InMaxRayTraceRequests, 16);
@@ -119,7 +149,10 @@ struct FNiagaraDataIntefaceProxyCollisionQuery : public FNiagaraDataInterfacePro
 				TEXT("NiagaraRayTraceRequests"),
 				sizeof(FBasicRayData),
 				MaxRayTraceCount,
-				BUF_Static);
+				BUF_Static,
+				false /*bUseUavCounter*/,
+				false /*bAppendBuffer*/,
+				ERHIAccess::UAVCompute);
 			INC_MEMORY_STAT_BY(STAT_NiagaraGPUDataInterfaceMemory, RayTraceRequests.NumBytes);
 
 			RayTraceIntersections.Initialize(
@@ -129,8 +162,17 @@ struct FNiagaraDataIntefaceProxyCollisionQuery : public FNiagaraDataInterfacePro
 				BUF_Static,
 				false /*bUseUavCounter*/,
 				false /*bAppendBuffer*/,
-				ERHIAccess::SRVMask);
+				ERHIAccess::SRVCompute);
 			INC_MEMORY_STAT_BY(STAT_NiagaraGPUDataInterfaceMemory, RayTraceIntersections.NumBytes);
+
+			RayTraceCounts.Initialize(
+				TEXT("NiagaraRayTraceIndirectArgs"),
+				sizeof(uint32),
+				4,
+				PF_R32_UINT,
+				ERHIAccess::UAVCompute,
+				BUF_Static | BUF_DrawIndirect);
+			INC_MEMORY_STAT_BY(STAT_NiagaraGPUDataInterfaceMemory, RayTraceCounts.NumBytes);
 		}
 #endif
 	}
@@ -360,6 +402,52 @@ void UNiagaraDataInterfaceCollisionQuery::GetFunctions(TArray<FNiagaraFunctionSi
 		IssueRayTrace.AddOutput(FNiagaraVariable(FNiagaraTypeDefinition::GetBoolDef(), TEXT("IsQueryValid")), IsQueryValidDescription);
 	}
 
+	{
+		const FText TraceStartWorldDescription = LOCTEXT("TraceStartWorldDescription", "Ray starting point in world space");
+		const FText TraceEndWorldDescription = LOCTEXT("TraceEndWorldDescription", "Ray end point in world space");
+		const FText TraceChannelDescription = LOCTEXT("TraceChannelDescription", "Currently unused, will represent the trace channels for which geometry the trace should test against");
+		const FText QueryIDDescription = LOCTEXT("QueryIDDescription", "Unique (for this frame) index of the query being enqueued (used in subsequent frames to retrieve results).");
+		const FText IsQueryValidDescription = LOCTEXT("IsQueryValidDescription", "Returns true if the query was enqueued");
+
+		FNiagaraFunctionSignature& IssueRayTrace = OutFunctions.AddDefaulted_GetRef();
+		IssueRayTrace.Name = NDICollisionQueryLocal::CreateAsyncRayTraceName;
+		IssueRayTrace.bRequiresExecPin = true;
+		IssueRayTrace.bMemberFunction = true;
+		IssueRayTrace.bRequiresContext = false;
+		IssueRayTrace.bSupportsCPU = false;
+#if WITH_EDITORONLY_DATA
+		IssueRayTrace.FunctionVersion = FNiagaraCollisionDIFunctionVersion::LatestVersion;
+		IssueRayTrace.Description = LOCTEXT("IssueAsyncRayTraceDescription", "Creates a GPU raytrace with the result being available the following frame (index is returned)");
+#endif
+		IssueRayTrace.AddInput(FNiagaraVariable(FNiagaraTypeDefinition(GetClass()), TEXT("CollisionQuery")));
+		IssueRayTrace.AddInput(FNiagaraVariable(FNiagaraTypeDefinition::GetVec3Def(), TEXT("TraceStartWorld")), TraceStartWorldDescription);
+		IssueRayTrace.AddInput(FNiagaraVariable(FNiagaraTypeDefinition::GetVec3Def(), TEXT("TraceEndWorld")), TraceEndWorldDescription);
+		IssueRayTrace.AddInput(FNiagaraVariable(FNiagaraTypeDefinition::GetIntDef(), TEXT("TraceChannel")), TraceChannelDescription);
+		IssueRayTrace.AddOutput(FNiagaraVariable(FNiagaraTypeDefinition::GetIntDef(), TEXT("QueryID")), QueryIDDescription);
+		IssueRayTrace.AddOutput(FNiagaraVariable(FNiagaraTypeDefinition::GetBoolDef(), TEXT("IsQueryValid")), IsQueryValidDescription);
+	}
+
+	{
+		const FText TraceCountDescription = LOCTEXT("QueryIDDescription", "Number of async raytrace requests to be reserved");
+		const FText FirstQueryIDValidDescription = LOCTEXT("TraceChannelDescription", "The first index in the block reserved through this call");
+		const FText IsQueryValidDescription = LOCTEXT("IsQueryValidDescription", "Returns true if the requested indices were reserved");
+
+		FNiagaraFunctionSignature& IssueRayTrace = OutFunctions.AddDefaulted_GetRef();
+		IssueRayTrace.Name = NDICollisionQueryLocal::ReserveAsyncRayTraceName;
+		IssueRayTrace.bRequiresExecPin = true;
+		IssueRayTrace.bMemberFunction = true;
+		IssueRayTrace.bRequiresContext = false;
+		IssueRayTrace.bSupportsCPU = false;
+#if WITH_EDITORONLY_DATA
+		IssueRayTrace.FunctionVersion = FNiagaraCollisionDIFunctionVersion::LatestVersion;
+		IssueRayTrace.Description = LOCTEXT("IssueAsyncRayTraceDescription", "Reserves a number of ray trace request slots");
+#endif
+		IssueRayTrace.AddInput(FNiagaraVariable(FNiagaraTypeDefinition(GetClass()), TEXT("CollisionQuery")));
+		IssueRayTrace.AddInput(FNiagaraVariable(FNiagaraTypeDefinition::GetIntDef(), TEXT("TraceCount")), TraceCountDescription);
+		IssueRayTrace.AddOutput(FNiagaraVariable(FNiagaraTypeDefinition::GetBoolDef(), TEXT("FirstQueryID")), FirstQueryIDValidDescription);
+		IssueRayTrace.AddOutput(FNiagaraVariable(FNiagaraTypeDefinition::GetBoolDef(), TEXT("IsQueryValid")), IsQueryValidDescription);
+	}
+
 	{
 		const FText PreviousFrameQueryIDDescription = LOCTEXT("PreviousFrameQueryIDDescription", "The index of the results being retrieved");
 		const FText CollisionValidDescription = LOCTEXT("CollisionValidDescription", "Returns true if the a Hit was encountered");
@@ -460,7 +548,9 @@ bool UNiagaraDataInterfaceCollisionQuery::GetFunctionHLSL(const FNiagaraDataInte
 		 (FunctionInfo.DefinitionName == NDICollisionQueryLocal::CustomDepthName) ||
 		 (FunctionInfo.DefinitionName == NDICollisionQueryLocal::DistanceFieldName) ||
 		 (FunctionInfo.DefinitionName == NDICollisionQueryLocal::IssueAsyncRayTraceName) ||
-		(FunctionInfo.DefinitionName == NDICollisionQueryLocal::ReadAsyncRayTraceName) )
+		 (FunctionInfo.DefinitionName == NDICollisionQueryLocal::CreateAsyncRayTraceName) ||
+		 (FunctionInfo.DefinitionName == NDICollisionQueryLocal::ReserveAsyncRayTraceName) ||
+		 (FunctionInfo.DefinitionName == NDICollisionQueryLocal::ReadAsyncRayTraceName) )
 	{
 		return true;
 	}
@@ -892,9 +982,10 @@ public:
 	{
 		GlobalDistanceFieldParameters.Bind(ParameterMap);
 #if RHI_RAYTRACING
-		MaxRayTraceCountParam.Bind(ParameterMap, TEXT("MaxRayTraceCount"));
-		RayRequestsParam.Bind(ParameterMap, TEXT("RayRequests"));
-		IntersectionResultsParam.Bind(ParameterMap, TEXT("IntersectionResults"));
+		MaxRayTraceCountParam.Bind(ParameterMap, *(NDICollisionQueryLocal::MaxRayTraceCountParamName + ParameterInfo.DataInterfaceHLSLSymbol));
+		RayRequestsParam.Bind(ParameterMap, *(NDICollisionQueryLocal::RayRequestsParamName + ParameterInfo.DataInterfaceHLSLSymbol));
+		IntersectionResultsParam.Bind(ParameterMap, *(NDICollisionQueryLocal::IntersectionResultsParamName + ParameterInfo.DataInterfaceHLSLSymbol));
+		RayTraceCountsParam.Bind(ParameterMap, *(NDICollisionQueryLocal::RayTraceCountsParamName + ParameterInfo.DataInterfaceHLSLSymbol));
 #endif
 	}
 
@@ -923,6 +1014,11 @@ public:
 		{
 			SetSRVParameter(RHICmdList, ComputeShaderRHI, IntersectionResultsParam, QueryDI->RayTraceIntersections.SRV);
 		}
+
+		if (RayTraceCountsParam.IsUAVBound())
+		{
+			RHICmdList.SetUAVParameter(ComputeShaderRHI, RayTraceCountsParam.GetUAVIndex(), QueryDI->RayTraceCounts.UAV);
+		}
 #endif
 	}
 
@@ -936,6 +1032,11 @@ public:
 		{
 			RayRequestsParam.UnsetUAV(RHICmdList, ComputeShaderRHI);
 		}
+
+		if (RayTraceCountsParam.IsUAVBound())
+		{
+			RayTraceCountsParam.UnsetUAV(RHICmdList, ComputeShaderRHI);
+		}
 	}
 #endif
 
@@ -946,6 +1047,7 @@ private:
 	LAYOUT_FIELD(FShaderParameter, MaxRayTraceCountParam);
 	LAYOUT_FIELD(FRWShaderParameter, RayRequestsParam);
 	LAYOUT_FIELD(FShaderResourceParameter, IntersectionResultsParam);
+	LAYOUT_FIELD(FRWShaderParameter, RayTraceCountsParam);
 #endif
 };
 
diff --git a/Engine/Plugins/FX/Niagara/Source/Niagara/Private/NiagaraEmitterInstanceBatcher.cpp b/Engine/Plugins/FX/Niagara/Source/Niagara/Private/NiagaraEmitterInstanceBatcher.cpp
--- a/Engine/Plugins/FX/Niagara/Source/Niagara/Private/NiagaraEmitterInstanceBatcher.cpp
+++ b/Engine/Plugins/FX/Niagara/Source/Niagara/Private/NiagaraEmitterInstanceBatcher.cpp
@@ -1860,11 +1860,11 @@ void NiagaraEmitterInstanceBatcher::ResetRayTracingSceneInfo()
 	RayTracingHelper->Reset();
 }
 
-void NiagaraEmitterInstanceBatcher::IssueRayTraces(FRHICommandList& RHICmdList, const FIntPoint& RayTraceCounts, FRHIShaderResourceView* RayTraceRequests, FRHIUnorderedAccessView* RayTraceResults) const
+void NiagaraEmitterInstanceBatcher::IssueRayTraces(FRHICommandList& RHICmdList, const FIntPoint& RayTraceCounts, FRHIShaderResourceView* RayTraceRequests, FRWBuffer* IndirectArgsBuffer, uint32 IndirectArgsOffset, FRHIUnorderedAccessView* RayTraceResults) const
 {
 	check(NumTicksThatRequireRayTracingScene > 0);
 
-	RayTracingHelper->IssueRayTraces(RHICmdList, RayTraceCounts, RayTraceRequests, RayTraceResults);
+	RayTracingHelper->IssueRayTraces(RHICmdList, RayTraceCounts, RayTraceRequests, IndirectArgsBuffer, IndirectArgsOffset, RayTraceResults);
 }
 
 bool NiagaraEmitterInstanceBatcher::HasRayTracingScene() const
diff --git a/Engine/Plugins/FX/Niagara/Source/NiagaraShader/Private/NiagaraRayTracingHelper.cpp b/Engine/Plugins/FX/Niagara/Source/NiagaraShader/Private/NiagaraRayTracingHelper.cpp
--- a/Engine/Plugins/FX/Niagara/Source/NiagaraShader/Private/NiagaraRayTracingHelper.cpp
+++ b/Engine/Plugins/FX/Niagara/Source/NiagaraShader/Private/NiagaraRayTracingHelper.cpp
@@ -12,10 +12,7 @@
 #include "ShaderParameterStruct.h"
 
 /// TODO
-///  -get DispatchIndirect working where we can atomically increment the collision request index and then use that to
-///		control the dispatch index (or, at the least feed it into the RayGen shader to early out)
 ///  -get geometry masking working when an environmental mask is implemented
-///  -template the CollisionDI to allow for multiple instances in the same emitter (even if that doesn't make a whole lot of sense)
 
 class FNiagaraCollisionRayTraceRG : public FGlobalShader
 {
@@ -26,12 +23,29 @@ class FNiagaraCollisionRayTraceRG : public FGlobalShader
 		SHADER_PARAMETER_RDG_BUFFER_SRV(RaytracingAccelerationStructure, TLAS)
 		SHADER_PARAMETER_RDG_BUFFER_SRV(StructuredBuffer<FBasicRayData>, Rays)
 		SHADER_PARAMETER_RDG_BUFFER_UAV(RWStructuredBuffer<FNiagaraRayTracingPayload>, CollisionOutput)
+		SHADER_PARAMETER_RDG_BUFFER_SRV(Buffer<UINT>, RayTraceCounts)
 	END_SHADER_PARAMETER_STRUCT()
 
+	class FFakeIndirectDispatch : SHADER_PERMUTATION_BOOL("NIAGARA_RAYTRACE_FAKE_INDIRECT");
+	using FPermutationDomain = TShaderPermutationDomain<FFakeIndirectDispatch>;
+
 	static bool ShouldCompilePermutation(const FGlobalShaderPermutationParameters& Parameters)
 	{
 		return ShouldCompileRayTracingShadersForProject(Parameters.Platform);
 	}
+
+	static FRHIRayTracingShader* GetShader(FGlobalShaderMap* ShaderMap)
+	{
+		FPermutationDomain PermutationVector;
+		PermutationVector.Set<FFakeIndirectDispatch>(!SupportsIndirectDispatch());
+
+		return ShaderMap->GetShader<FNiagaraCollisionRayTraceRG>(PermutationVector).GetRayTracingShader();
+	}
+
+	static bool SupportsIndirectDispatch()
+	{
+		return GRHISupportsRayTracingDispatchIndirect;
+	}
 };
 
 class FNiagaraCollisionRayTraceCH : public FGlobalShader
@@ -172,7 +186,7 @@ void FNiagaraRayTracingHelper::BuildRayTracingSceneInfo(FRHICommandList& RHICmdL
 		ViewUniformBuffer = ReferenceView.ViewUniformBuffer;
 
 		FGlobalShaderMap* ShaderMap = GetGlobalShaderMap(ShaderPlatform);
-		auto RayGenShader = ShaderMap->GetShader<FNiagaraCollisionRayTraceRG>().GetRayTracingShader();
+		auto RayGenShader = FNiagaraCollisionRayTraceRG::GetShader(ShaderMap);
 		auto ClosestHitShader = ShaderMap->GetShader<FNiagaraCollisionRayTraceCH>().GetRayTracingShader();
 		auto MissShader = ShaderMap->GetShader<FNiagaraCollisionRayTraceMiss>().GetRayTracingShader();
 
@@ -202,7 +216,7 @@ void FNiagaraRayTracingHelper::BuildRayTracingSceneInfo(FRHICommandList& RHICmdL
 	}
 }
 
-void FNiagaraRayTracingHelper::IssueRayTraces(FRHICommandList& RHICmdList, const FIntPoint& RayTraceCounts, FRHIShaderResourceView* RayTraceRequests, FRHIUnorderedAccessView* RayTraceResults) const
+void FNiagaraRayTracingHelper::IssueRayTraces(FRHICommandList& RHICmdList, const FIntPoint& RayTraceCounts, FRHIShaderResourceView* RayTraceRequests, FRWBuffer* IndirectArgsBuffer, uint32 IndirectArgsOffset, FRHIUnorderedAccessView* RayTraceResults) const
 {
 	check(RayTracingPipelineState);
 	check(RayTracingScene);
@@ -215,15 +229,36 @@ void FNiagaraRayTracingHelper::IssueRayTraces(FRHICommandList& RHICmdList, const
 	GlobalBindings.SRVs[1] = RayTraceRequests;
 	GlobalBindings.UAVs[0] = RayTraceResults;
 
-	// todo - add support for RayTraceDispatchIndirect
-	RHICmdList.RayTraceDispatch(
-		RayTracingPipelineState,
-		ShaderMap->GetShader<FNiagaraCollisionRayTraceRG>().GetRayTracingShader(),
-		RayTracingScene,
-		GlobalBindings,
-		RayTraceCounts.X,
-		RayTraceCounts.Y
-	);
+	if (FNiagaraCollisionRayTraceRG::SupportsIndirectDispatch())
+	{
+		RHICmdList.Transition(FRHITransitionInfo(IndirectArgsBuffer->UAV, ERHIAccess::UAVCompute, ERHIAccess::IndirectArgs | ERHIAccess::SRVCompute));
+
+		RHICmdList.RayTraceDispatchIndirect(
+			RayTracingPipelineState,
+			FNiagaraCollisionRayTraceRG::GetShader(ShaderMap),
+			RayTracingScene,
+			GlobalBindings,
+			IndirectArgsBuffer->Buffer,
+			IndirectArgsOffset);
+
+		RHICmdList.Transition(FRHITransitionInfo(IndirectArgsBuffer->UAV, ERHIAccess::IndirectArgs | ERHIAccess::SRVCompute, ERHIAccess::UAVCompute));
+	}
+	else
+	{
+		RHICmdList.Transition(FRHITransitionInfo(IndirectArgsBuffer->UAV, ERHIAccess::UAVCompute, ERHIAccess::SRVCompute));
+		GlobalBindings.SRVs[2] = IndirectArgsBuffer->SRV;
+
+		RHICmdList.RayTraceDispatch(
+			RayTracingPipelineState,
+			FNiagaraCollisionRayTraceRG::GetShader(ShaderMap),
+			RayTracingScene,
+			GlobalBindings,
+			RayTraceCounts.X,
+			RayTraceCounts.Y
+		);
+
+		RHICmdList.Transition(FRHITransitionInfo(IndirectArgsBuffer->UAV, ERHIAccess::SRVCompute, ERHIAccess::UAVCompute));
+	}
 }
 
 bool FNiagaraRayTracingHelper::IsValid() const
diff --git a/Engine/Plugins/FX/Niagara/Source/NiagaraShader/Public/NiagaraRayTracingHelper.h b/Engine/Plugins/FX/Niagara/Source/NiagaraShader/Public/NiagaraRayTracingHelper.h
--- a/Engine/Plugins/FX/Niagara/Source/NiagaraShader/Public/NiagaraRayTracingHelper.h
+++ b/Engine/Plugins/FX/Niagara/Source/NiagaraShader/Public/NiagaraRayTracingHelper.h
@@ -34,7 +34,7 @@ class NIAGARASHADER_API FNiagaraRayTracingHelper
 
 		void Reset();
 		void BuildRayTracingSceneInfo(FRHICommandList& RHICmdList, TConstArrayView<FViewInfo> Views);
-		void IssueRayTraces(FRHICommandList& RHICmdList, const FIntPoint& RayTraceCounts, FRHIShaderResourceView* RayTraceRequests, FRHIUnorderedAccessView* RayTraceResults) const;
+		void IssueRayTraces(FRHICommandList& RHICmdList, const FIntPoint& RayTraceCounts, FRHIShaderResourceView* RayTraceRequests, FRWBuffer* IndirectArgsBuffer, uint32 IndirectArgsOffset, FRHIUnorderedAccessView* RayTraceResults) const;
 		bool IsValid() const;
 
 	private:
