Unity compile fixes for NiagaraRayTracingHelper

diff --git a/Engine/Plugins/FX/Niagara/Source/NiagaraShader/Public/NiagaraRayTracingHelper.h b/Engine/Plugins/FX/Niagara/Source/NiagaraShader/Public/NiagaraRayTracingHelper.h
--- a/Engine/Plugins/FX/Niagara/Source/NiagaraShader/Public/NiagaraRayTracingHelper.h
+++ b/Engine/Plugins/FX/Niagara/Source/NiagaraShader/Public/NiagaraRayTracingHelper.h
@@ -14,6 +14,7 @@ class FRHIRayTracingScene;
 class FRHIShaderResourceView;
 class FRHIUniformBuffer;
 class FRHIUnorderedAccessView;
+struct FRWBuffer;
 
 // c++ mirror of the struct defined in NiagaraRayTracingCommon.ush
 struct FNiagaraRayTracingPayload
