- Editgrate CL 16180189 - Refactor the way we dispatch / hold tick information on the render thread
- Fixes low latency translucent path as we now correctly predict the output buffers with simulation stages
- Fixes various issues with persistent IDs and simulation stages
- Fixes issue with consuming all render thread tick information before we tick
- Fixes issue with tracking reset tick information between game / render thread
- Fixes various transition related issues
- Fixes issue with warmup and GPU systems

#rnx

diff --git a/Engine/Plugins/FX/Niagara/Shaders/Private/NiagaraComputeFreeIDs.usf b/Engine/Plugins/FX/Niagara/Shaders/Private/NiagaraComputeFreeIDs.usf
--- a/Engine/Plugins/FX/Niagara/Shaders/Private/NiagaraComputeFreeIDs.usf
+++ b/Engine/Plugins/FX/Niagara/Shaders/Private/NiagaraComputeFreeIDs.usf
@@ -17,6 +17,7 @@ uint FreeIDListIndex;
 
 #if USE_WAVE_INTRINSICS
 
+// This happens to be fixed on PS4 and XB1, but we'll need to configure it externally
 [numthreads(THREAD_COUNT, 1, 1)]
 void ComputeFreeIDs(uint3 GroupId : SV_GroupID, uint3 GroupThreadId : SV_GroupThreadID, uint3 DispatchThreadId : SV_DispatchThreadID)
 {
diff --git a/Engine/Plugins/FX/Niagara/Shaders/Private/NiagaraInitFreeIDBuffer.usf b/Engine/Plugins/FX/Niagara/Shaders/Private/NiagaraInitFreeIDBuffer.usf
--- a/Engine/Plugins/FX/Niagara/Shaders/Private/NiagaraInitFreeIDBuffer.usf
+++ b/Engine/Plugins/FX/Niagara/Shaders/Private/NiagaraInitFreeIDBuffer.usf
@@ -19,6 +19,6 @@ void InitIDBufferCS(uint Index : SV_DispatchThreadID)
 	else
 	{
 		// Copy existing IDs after the new IDs.
-		RWNewBuffer[Index.x] = ExistingBuffer[Index.x];
+		RWNewBuffer[Index.x] = ExistingBuffer[Index.x - NumNewElements];
 	}
 }
diff --git a/Engine/Plugins/FX/Niagara/Source/Niagara/Classes/NiagaraComputeExecutionContext.h b/Engine/Plugins/FX/Niagara/Source/Niagara/Classes/NiagaraComputeExecutionContext.h
new file mode 100644
--- /dev/null
+++ b/Engine/Plugins/FX/Niagara/Source/Niagara/Classes/NiagaraComputeExecutionContext.h
@@ -0,0 +1,235 @@
+// Copyright Epic Games, Inc. All Rights Reserved.
+
+#pragma once
+
+#include "NiagaraScriptExecutionContext.h"
+#include "RHIGPUReadback.h"
+
+class FNiagaraGPUInstanceCountManager;
+class NiagaraEmitterInstanceBatcher;
+class FNiagaraGPUSystemTick;
+struct FNiagaraComputeInstanceData;
+struct FNiagaraComputeExecutionContext;
+
+class FNiagaraRHIUniformBufferLayout : public FRHIResource
+{
+public:
+	explicit FNiagaraRHIUniformBufferLayout(const TCHAR* LayoutName) : FRHIResource(RRT_None), UBLayout(LayoutName) { }
+
+	FRHIUniformBufferLayout UBLayout;
+};
+
+struct FNiagaraSimStageData
+{
+	FNiagaraSimStageData()
+	{
+		bFirstStage = false;
+		bLastStage = false;
+	}
+
+	uint32 bFirstStage : 1;
+	uint32 bLastStage : 1;
+
+	uint32 StageIndex = INDEX_NONE;
+
+	FNiagaraDataBuffer* Source = nullptr;
+	uint32 SourceCountOffset = INDEX_NONE;
+	uint32 SourceNumInstances = 0;
+
+	FNiagaraDataBuffer* Destination = nullptr;
+	uint32 DestinationCountOffset = INDEX_NONE;
+	uint32 DestinationNumInstances = 0;
+
+	FNiagaraDataInterfaceProxyRW* AlternateIterationSource = nullptr;
+	const FSimulationStageMetaData* StageMetaData = nullptr;
+};
+
+struct FNiagaraGpuDispatchInstance
+{
+	FNiagaraGpuDispatchInstance(const FNiagaraGPUSystemTick& InTick, const FNiagaraComputeInstanceData& InInstanceData)
+		: Tick(InTick)
+		, InstanceData(InInstanceData)
+	{
+	}
+
+	const FNiagaraGPUSystemTick& Tick;
+	const FNiagaraComputeInstanceData& InstanceData;
+	FNiagaraSimStageData SimStageData;
+};
+
+struct FNiagaraGpuDispatchGroup
+{
+	TArray<FNiagaraGPUSystemTick*>				TicksWithPerInstanceData;
+	TArray<FNiagaraGpuDispatchInstance>			DispatchInstances;
+	TArray<FNiagaraComputeExecutionContext*>	FreeIDUpdates;
+};
+
+struct FNiagaraGpuDispatchList
+{
+	void PreAllocateGroups(int32 LastGroup)
+	{
+		const int32 GroupsToAllocate = LastGroup - DispatchGroups.Num();
+		if (GroupsToAllocate > 0)
+		{
+			DispatchGroups.AddDefaulted(GroupsToAllocate);
+		}
+	}
+
+	bool HasWork() const { return DispatchGroups.Num() > 0; }
+
+	TArray<uint32>						CountsToRelease;
+	TArray<FNiagaraGpuDispatchGroup>	DispatchGroups;
+};
+
+struct FNiagaraGpuSpawnInfoParams
+{
+	float IntervalDt;
+	float InterpStartDt;
+	int32 SpawnGroup;
+	int32 GroupSpawnStartIndex;
+};
+
+struct FNiagaraGpuSpawnInfo
+{
+	uint32 EventSpawnTotal = 0;
+	uint32 SpawnRateInstances = 0;
+	uint32 MaxParticleCount = 0;
+	int32 SpawnInfoStartOffsets[NIAGARA_MAX_GPU_SPAWN_INFOS];
+	FNiagaraGpuSpawnInfoParams SpawnInfoParams[NIAGARA_MAX_GPU_SPAWN_INFOS];
+
+	void Reset()
+	{
+		EventSpawnTotal = 0;
+		SpawnRateInstances = 0;
+		MaxParticleCount = 0;
+		for (int32 i = 0; i < NIAGARA_MAX_GPU_SPAWN_INFOS; ++i)
+		{
+			SpawnInfoStartOffsets[i] = 0;
+
+			SpawnInfoParams[i].IntervalDt = 0;
+			SpawnInfoParams[i].InterpStartDt = 0;
+			SpawnInfoParams[i].SpawnGroup = 0;
+			SpawnInfoParams[i].GroupSpawnStartIndex = 0;
+		}
+	}
+};
+
+struct FNiagaraComputeExecutionContext
+{
+	FNiagaraComputeExecutionContext();
+	~FNiagaraComputeExecutionContext();
+
+	void Reset(NiagaraEmitterInstanceBatcher* Batcher);
+
+	void InitParams(UNiagaraScript* InGPUComputeScript, ENiagaraSimTarget InSimTarget, const uint32 InDefaultSimulationStageIndex, int32 InMaxUpdateIterations, const TSet<uint32> InSpawnStages);
+	void DirtyDataInterfaces();
+	bool Tick(FNiagaraSystemInstance* ParentSystemInstance);
+
+	bool OptionalContexInit(FNiagaraSystemInstance* ParentSystemInstance);
+
+	void PostTick();
+
+	void SetDataToRender(FNiagaraDataBuffer* InDataToRender);
+	void SetTranslucentDataToRender(FNiagaraDataBuffer* InTranslucentDataToRender);
+	FNiagaraDataBuffer* GetDataToRender(bool bIsLowLatencyTranslucent) const { return bIsLowLatencyTranslucent && TranslucentDataToRender ? TranslucentDataToRender : DataToRender; }
+
+	struct 
+	{
+		// The offset at which the GPU instance count (see FNiagaraGPUInstanceCountManager()).
+		uint32 GPUCountOffset = INDEX_NONE;
+		// The CPU instance count at the time the GPU count readback was issued. Always bigger or equal to the GPU count.
+		uint32 CPUCount = 0;
+	}  EmitterInstanceReadback;
+	void ReleaseReadbackCounter(FNiagaraGPUInstanceCountManager& GPUInstanceCountManager);
+	
+#if !UE_BUILD_SHIPPING
+	const TCHAR* GetDebugSimName() const { return *DebugSimName; }
+	void SetDebugSimName(const TCHAR* InDebugSimName) { DebugSimName = InDebugSimName; }
+#else
+	const TCHAR* GetDebugSimName() const { return TEXT(""); }
+	void SetDebugSimName(const TCHAR*) { }
+#endif
+
+//-TOOD:private:
+	void ResetInternal(NiagaraEmitterInstanceBatcher* Batcher);
+
+public:
+	static uint32 TickCounter;
+
+#if !UE_BUILD_SHIPPING
+	FString DebugSimName;
+#endif
+#if STATS
+	TWeakObjectPtr<UNiagaraEmitter> EmitterPtr; // emitter pointer used to report captured gpu stats
+#endif
+
+	const TArray<UNiagaraDataInterface*>& GetDataInterfaces()const { return CombinedParamStore.GetDataInterfaces(); }
+
+	class FNiagaraDataSet *MainDataSet;
+	UNiagaraScript* GPUScript;
+	class FNiagaraShaderScript*  GPUScript_RT;
+
+	// persistent layouts used to create the constant buffers for the compute sim shader
+	TRefCountPtr<FNiagaraRHIUniformBufferLayout> ExternalCBufferLayout;
+
+	//Dynamic state updated either from GT via RT commands or from the RT side sim code itself.
+	//TArray<uint8, TAlignedHeapAllocator<16>> ParamData_RT;		// RT side copy of the parameter data
+	FNiagaraScriptInstanceParameterStore CombinedParamStore;
+#if DO_CHECK
+	TArray< FString >  DIClassNames;
+#endif
+
+	TArray<FNiagaraDataInterfaceProxy*> DataInterfaceProxies;
+
+	// Most current buffer that can be used for rendering.
+	FNiagaraDataBuffer* DataToRender = nullptr;
+
+	// Optional buffer which can be used to render translucent data with no latency (i.e. this frames data)
+	FNiagaraDataBuffer* TranslucentDataToRender = nullptr;
+
+	// Game thread spawn info will be sent to the render thread inside FNiagaraComputeInstanceData
+	FNiagaraGpuSpawnInfo GpuSpawnInfo_GT;
+
+	uint32 DefaultSimulationStageIndex = 0;
+	uint32 MaxUpdateIterations = 0;
+	TSet<uint32> SpawnStages;
+
+	bool HasInterpolationParameters = false;
+
+	// Do we have a reset pending, controlled by the game thread and passed to instance data
+	bool bResetPending_GT = true;
+
+	// Particle count read fence, used to allow us to snoop the count written by the render thread but also avoid racing on a reset value
+	uint32 ParticleCountReadFence = 1;
+	uint32 ParticleCountWriteFence = 0;
+
+	// Render thread data
+	FNiagaraDataBuffer* GetCurrDataBuffer() { check(IsInRenderingThread()); return DataBuffers_RT[DataBufferIndex_RT]; }
+	FNiagaraDataBuffer* GetNextDataBuffer() { check(IsInRenderingThread()); return DataBuffers_RT[DataBufferIndex_RT ^ 1]; }
+	void AdvanceDataBuffer() { DataBufferIndex_RT ^= 1; ++BufferSwapsThisFrame_RT; }
+
+	uint32 DataBufferIndex_RT = 0;
+	FNiagaraDataBuffer* DataBuffers_RT[2] = { nullptr, nullptr };
+	uint32 CountOffset_RT = INDEX_NONE;
+
+	uint32 BufferSwapsThisFrame_RT = 0;
+
+	// Used to track if we have processed any ticks for this context this frame
+	bool bHasTickedThisFrame_RT = false;
+
+	// The current number of instances on the RT
+	// Before ticks are processed on the RT this will be CurrentData's NumInstances
+	// As ticks are processed (or we generated the tick batches) this will change and won't be accurate until dispatches are executed
+	uint32 CurrentNumInstances_RT = 0;
+	// The current maximum of instances on the RT
+	uint32 CurrentMaxInstances_RT = 0;
+	// The current maximum instances we should allocate on the RT
+	uint32 CurrentMaxAllocateInstances_RT = 0;
+
+	TArray<FSimulationStageMetaData> SimStageInfo;
+
+	bool IsOutputStage(FNiagaraDataInterfaceProxy* DIProxy, uint32 CurrentStage) const;
+	bool IsIterationStage(FNiagaraDataInterfaceProxy* DIProxy, uint32 CurrentStage) const;
+	FNiagaraDataInterfaceProxyRW* FindIterationInterface(const TArray<FNiagaraDataInterfaceProxyRW*>& InProxies, uint32 SimulationStageIndex) const;
+	const FSimulationStageMetaData* GetSimStageMetaData(uint32 SimulationStageIndex) const;
+};
diff --git a/Engine/Plugins/FX/Niagara/Source/Niagara/Classes/NiagaraDataSet.h b/Engine/Plugins/FX/Niagara/Source/Niagara/Classes/NiagaraDataSet.h
--- a/Engine/Plugins/FX/Niagara/Source/Niagara/Classes/NiagaraDataSet.h
+++ b/Engine/Plugins/FX/Niagara/Source/Niagara/Classes/NiagaraDataSet.h
@@ -105,14 +105,19 @@ protected:
 class NIAGARA_API FNiagaraDataBuffer : public FNiagaraSharedObject
 {
 	friend class FScopedNiagaraDataSetGPUReadback;
+	friend class NiagaraEmitterInstanceBatcher;
+	
 protected:
 	virtual ~FNiagaraDataBuffer();
 
 public:
 	FNiagaraDataBuffer(FNiagaraDataSet* InOwner);
 	void Allocate(uint32 NumInstances, bool bMaintainExisting = false);
-	void AllocateGPU(uint32 InNumInstances, FNiagaraGPUInstanceCountManager& GPUInstanceCountManager, FRHICommandList& RHICmdList, ERHIFeatureLevel::Type FeatureLevel, const TCHAR* DebugSimName);
+
+	void AllocateGPU(FRHICommandList& RHICmdList, uint32 InNumInstances, ERHIFeatureLevel::Type FeatureLevel, const TCHAR* DebugSimName);
+	void AliasGPU(FNiagaraDataBuffer* BufferToAlias);
 	void ReleaseGPU();
+
 	void SwapInstances(uint32 OldIndex, uint32 NewIndex);
 	void KillInstance(uint32 InstanceIdx);
 	void CopyTo(FNiagaraDataBuffer& DestBuffer, int32 SrcStartIdx, int32 DestStartIdx, int32 NumInstances)const;
@@ -185,11 +190,14 @@ public:
 	static void SetOutputShaderParams(FRHICommandList& RHICmdList, class FNiagaraShader* Shader, FNiagaraDataBuffer* Buffer);
 	static void UnsetShaderParams(FRHICommandList& RHICmdList, class FNiagaraShader* Shader);
 
-	void ReleaseGPUInstanceCount(FNiagaraGPUInstanceCountManager& GPUInstanceCountManager);
+	void ClearGPUInstanceCount()
+	{
+		GPUInstanceCountBufferOffset = INDEX_NONE;
+	}
 
 	void BuildRegisterTable();
-private:
 
+private:
 	FORCEINLINE void CheckUsage(bool bReadOnly)const;
 
 	FORCEINLINE int32 GetSafeComponentBufferSize(int32 RequiredSize) const
@@ -227,6 +235,9 @@ private:
 	FRWBuffer GPUIDToIndexTable;
 	/** GPU Buffer containing half values for GPU simulations. */
 	FRWBuffer GPUBufferHalf;
+#if NIAGARA_MEMORY_TRACKING
+	int32 AllocationSizeBytes = 0;
+#endif
 	//////////////////////////////////////////////////////////////////////////
 
 	/** Number of instances in data. */
@@ -379,9 +390,6 @@ public:
 		return *DestinationData;
 	}
 
-	/** Release the GPU instance counts so that they can be reused */
-	void ReleaseGPUInstanceCounts(FNiagaraGPUInstanceCountManager& GPUInstanceCountManager);
-
 	void AllocateGPUFreeIDs(uint32 InNumInstances, FRHICommandList& RHICmdList, ERHIFeatureLevel::Type FeatureLevel, const TCHAR* DebugSimName);
 
 	void SetMaxInstanceCount(uint32 InInstanceCount) { MaxInstanceCount = InInstanceCount; }
diff --git a/Engine/Plugins/FX/Niagara/Source/Niagara/Classes/NiagaraEmitterInstanceBatcher.h b/Engine/Plugins/FX/Niagara/Source/Niagara/Classes/NiagaraEmitterInstanceBatcher.h
--- a/Engine/Plugins/FX/Niagara/Source/Niagara/Classes/NiagaraEmitterInstanceBatcher.h
+++ b/Engine/Plugins/FX/Niagara/Source/Niagara/Classes/NiagaraEmitterInstanceBatcher.h
@@ -12,13 +12,15 @@ the same VectorVM byte code / compute shader code
 #include "Modules/ModuleManager.h"
 #include "ParticleResources.h"
 #include "NiagaraCommon.h"
+#include "NiagaraComputeExecutionContext.h"
 #include "NiagaraEmitter.h"
 #include "NiagaraGPUInstanceCountManager.h"
 #include "NiagaraGPUProfiler.h"
 #include "NiagaraGPUSortInfo.h"
+#include "NiagaraGPUSystemTick.h"
 #include "NiagaraParameters.h"
 #include "NiagaraRendererProperties.h"
-#include "NiagaraScriptExecutionContext.h"
+#include "NiagaraSystemGpuComputeProxy.h"
 #include "RendererInterface.h"
 #include "RHIResources.h"
 #include "Runtime/Engine/Private/Particles/ParticleSortingGPU.h"
@@ -29,14 +31,6 @@ class FNiagaraGpuComputeDebug;
 class FNiagaraGpuReadbackManager;
 class FNiagaraRayTracingHelper;
 
-enum class ETickStage
-{
-	PreInitViews,
-	PostInitViews,
-	PostOpaqueRender,
-	Max
-};
-
 struct FNiagaraUAVPoolAccessScope
 {
 	FNiagaraUAVPoolAccessScope(class NiagaraEmitterInstanceBatcher& InBatcher);
@@ -50,9 +44,6 @@ class NiagaraEmitterInstanceBatcher : public FFXSystemInterface
 	friend FNiagaraUAVPoolAccessScope;
 
 public:
-	using FNiagaraTransitionList = TArray<FRHITransitionInfo, TMemStackAllocator<>>;
-	using FOverlappableTicks = TArray<FNiagaraGPUSystemTick*, TMemStackAllocator<>>;
-
 	NiagaraEmitterInstanceBatcher(ERHIFeatureLevel::Type InFeatureLevel, EShaderPlatform InShaderPlatform, FGPUSortManager* InGPUSortManager);
 
 	~NiagaraEmitterInstanceBatcher();
@@ -60,15 +51,11 @@ public:
 	static NIAGARA_API const FName Name;
 	virtual FFXSystemInterface* GetInterface(const FName& InName) override;
 
-	/** Notification that the InstanceID has been removed. */
-	void InstanceDeallocated_RenderThread(const FNiagaraSystemInstanceID InstanceID);
-
-	// The batcher assumes ownership of the data here.
+	/** Add system instance proxy to the batcher for tracking. */
+	void AddGpuComputeProxy(FNiagaraSystemGpuComputeProxy* ComputeProxy);
 
-	void GiveSystemTick_RenderThread(FNiagaraGPUSystemTick& Tick);
-
-	/** Called to release GPU instance counts that the batcher is tracking. */
-	void ReleaseInstanceCounts_RenderThread(FNiagaraComputeExecutionContext* ExecContext, FNiagaraDataSet* DataSet);
+	/** Remove system instance proxy from the batcher. */
+	void RemoveGpuComputeProxy(FNiagaraSystemGpuComputeProxy* ComputeProxy);
 
 #if WITH_EDITOR
 	virtual void Suspend() override {}
@@ -120,23 +107,8 @@ public:
 
 	const FGlobalDistanceFieldParameterData& GetGlobalDistanceFieldParameters() const { return GlobalDistanceFieldParams; }
 
-	void SetDataInterfaceParameters(const TArray<FNiagaraDataInterfaceProxy*>& DataInterfaceProxies, const FNiagaraShaderRef& Shader, FRHICommandList &RHICmdList, const FNiagaraComputeInstanceData* Instance, const FNiagaraGPUSystemTick& Tick, uint32 SimulationStageIndex) const;
-	void UnsetDataInterfaceParameters(const TArray<FNiagaraDataInterfaceProxy*>& DataInterfaceProxies, const FNiagaraShaderRef& Shader, FRHICommandList& RHICmdList, const FNiagaraComputeInstanceData* Instance, const FNiagaraGPUSystemTick& Tick, uint32 SimulationStageIndex) const;
-
-	void Run(	const FNiagaraGPUSystemTick& Tick,
-				const FNiagaraComputeInstanceData* Instance,
-				uint32 UpdateStartInstance,
-				const uint32 TotalNumInstances,
-				const FNiagaraShaderRef& Shader,
-				FRHICommandList &RHICmdList,
-				FRHIUniformBuffer* ViewUniformBuffer,
-				const FNiagaraGpuSpawnInfo& SpawnInfo,
-				bool bCopyBeforeStart = false,
-				uint32 DefaultSimulationStageIndex = 0,
-				uint32 SimulationStageIndex = 0,
-				FNiagaraDataInterfaceProxyRW* IterationInterface = nullptr,
-				bool HasRunParticleStage = false
-			);
+	void SetDataInterfaceParameters(FRHICommandList& RHICmdList, const FNiagaraGPUSystemTick& Tick, const FNiagaraComputeInstanceData& InstanceData, const FNiagaraShaderRef& ComputeShader, const FNiagaraSimStageData& SimStageData) const;
+	void UnsetDataInterfaceParameters(FRHICommandList& RHICmdList, const FNiagaraGPUSystemTick& Tick, const FNiagaraComputeInstanceData& InstanceData, const FNiagaraShaderRef& ComputeShader, const FNiagaraSimStageData& SimStageData) const;
 
 	FORCEINLINE FNiagaraGPUInstanceCountManager& GetGPUInstanceCounterManager() { check(IsInRenderingThread()); return GPUInstanceCounterManager; }
 
@@ -144,19 +116,19 @@ public:
 	FORCEINLINE ERHIFeatureLevel::Type GetFeatureLevel() const { return FeatureLevel; }
 
 	/** Reset the data interfaces and check if the spawn stages are valid */
-	bool ResetDataInterfaces(const FNiagaraGPUSystemTick& Tick, FNiagaraComputeInstanceData *Instance, FRHICommandList& RHICmdList, const FNiagaraShaderScript* ShaderScript) const;
+	void ResetDataInterfaces(FRHICommandList& RHICmdList, const FNiagaraGPUSystemTick& Tick, const FNiagaraComputeInstanceData& InstanceData) const;
 
 	/** Given a shader stage index, find the corresponding data interface */
 	FNiagaraDataInterfaceProxyRW* FindIterationInterface(FNiagaraComputeInstanceData* Instance, const uint32 SimulationStageIndex) const;
 
-	/** Loop over all the data interfaces and call the prestage methods */
-	void PreStageInterface(const FNiagaraGPUSystemTick& Tick, FNiagaraComputeInstanceData *Instance, FRHICommandList& RHICmdList, const uint32 SimulationStageIndex) const;
+	/** Loop over all the data interfaces and call the pre-stage methods */
+	void PreStageInterface(FRHICommandList& RHICmdList, const FNiagaraGPUSystemTick& Tick, const FNiagaraComputeInstanceData& InstanceData, const FNiagaraSimStageData& SimStageData) const;
 
-	/** Loop over all the data interfaces and call the poststage methods */
-	void PostStageInterface(const FNiagaraGPUSystemTick& Tick, FNiagaraComputeInstanceData *Instance, FRHICommandList& RHICmdList, const uint32 SimulationStageIndex) const;
+	/** Loop over all the data interfaces and call the post-stage methods */
+	void PostStageInterface(FRHICommandList& RHICmdList, const FNiagaraGPUSystemTick& Tick, const FNiagaraComputeInstanceData& InstanceData, const FNiagaraSimStageData& SimStageData) const;
 
-	/** Loop over all data interfaces and call the postsimulate methods */
-	void PostSimulateInterface(const FNiagaraGPUSystemTick& Tick, FNiagaraComputeInstanceData* Instance, FRHICommandList& RHICmdList, const FNiagaraShaderScript* ShaderScript) const;
+	/** Loop over all data interfaces and call the post-simulate methods */
+	void PostSimulateInterface(FRHICommandList& RHICmdList, const FNiagaraGPUSystemTick& Tick, const FNiagaraComputeInstanceData& InstanceData) const;
 
 	/** Grab a temporary dummy RW buffer from the pool.  Note: When doing this outside of Niagara you must be within a FNiagaraUAVPoolAccessScope. */
 	NIAGARA_API FRHIUnorderedAccessView* GetEmptyRWBufferFromPool(FRHICommandList& RHICmdList, EPixelFormat Format) const;
@@ -181,35 +153,12 @@ public:
 #endif
 
 private:
-	using FEmitterInstanceList = TArray<FNiagaraComputeInstanceData*>;
-
-	struct FDispatchInstance
-	{
-		FNiagaraGPUSystemTick* Tick = nullptr;
-		FNiagaraComputeInstanceData* InstanceData = nullptr;
-		int32 StageIndex = 0;
-		bool bFinalStage = false;
-	};
-	using FDispatchInstanceList = TArray<FDispatchInstance, TMemStackAllocator<>>;
-
-	struct FDispatchGroup
-	{
-		FDispatchInstanceList DispatchInstances;
-		FNiagaraTransitionList TransitionsBefore;
-		FNiagaraTransitionList TransitionsAfter;
-		TArray<uint32, TMemStackAllocator<>> InstancesWithPersistentIDUpdates;
-		TArray<FNiagaraComputeInstanceData*, TMemStackAllocator<>> InstancesWithFreeIDUpdates;
-	};
-	using FDispatchGroupList = TArray<FDispatchGroup, TMemStackAllocator<>>;
-
+	void DumpDebugFrame();
 	void UpdateInstanceCountManager(FRHICommandListImmediate& RHICmdList);
-	void BuildTickStagePasses(FRHICommandListImmediate& RHICmdList, ETickStage GenerateTickStage);
-
-	void ExecuteAll(FRHICommandList& RHICmdList, FRHIUniformBuffer* ViewUniformBuffer, ETickStage TickStage);
-	void AddDestinationBufferTransitions(FDispatchGroup* Group, FNiagaraDataBuffer* DestinationData);
-	void BuildDispatchGroups(FOverlappableTicks& OverlappableTick, FRHICommandList& RHICmdList, FDispatchGroupList& DispatchGroups, FEmitterInstanceList& InstancesWithPersistentIDs);
-	void DispatchStage(FDispatchInstance& DispatchInstance, uint32 StageIndex, FRHICommandList& RHICmdList, FRHIUniformBuffer* ViewUniformBuffer);
-	void DispatchAllOnCompute(FDispatchInstanceList& DispatchInstances, FRHICommandList& RHICmdList, FRHIUniformBuffer* ViewUniformBuffer);
+	void PrepareTicksForProxy(FRHICommandListImmediate& RHICmdList, FNiagaraSystemGpuComputeProxy* ComputeProxy, FNiagaraGpuDispatchList& GpuDispatchList);
+	void PrepareAllTicks(FRHICommandListImmediate& RHICmdList);
+	void ExecuteTicks(FRHICommandList& RHICmdList, FRHIUniformBuffer* ViewUniformBuffer, ENiagaraGpuComputeTickStage::Type TickStage);
+	void DispatchStage(FRHICommandList& RHICmdList, FRHIUniformBuffer* ViewUniformBuffer, const FNiagaraGPUSystemTick& Tick, const FNiagaraComputeInstanceData& InstanceData, const FNiagaraSimStageData& SimStageData);
 
 	/**
 	 * Generate all the initial keys and values for a GPUSortManager sort batch.
@@ -237,22 +186,12 @@ private:
 	void ResetRayTracingSceneInfo();
 #endif
 
-	inline uint32 UnpackEmitterDispatchCount(uint8* PackedData)
-	{
-		return *(uint32*)PackedData;
-	}
-
-	inline FNiagaraComputeInstanceData* UnpackEmitterComputeDispatchArray(uint8* PackedData)
-	{
-		return (FNiagaraComputeInstanceData*)(PackedData + sizeof(uint32));
-	}
-
 	bool UseOverlapCompute();
 	void FinishDispatches();
 	void ReleaseTicks();
+
 	void UpdateFreeIDsListSizesBuffer(FRHICommandList& RHICmdList, uint32 NumInstances);
-	void UpdateFreeIDBuffers(FRHICommandList& RHICmdList, FEmitterInstanceList& Instances);
-	void UpdateFreeIDBuffers(FRHICommandList& RHICmdList, TArrayView<FNiagaraComputeInstanceData*> Instances);
+	void UpdateFreeIDBuffers(FRHICommandList& RHICmdList, TConstArrayView<FNiagaraComputeExecutionContext*> Instances);
 
 	void SetConstantBuffers(FRHICommandList &RHICmdList, const FNiagaraShaderRef& Shader, const FNiagaraGPUSystemTick& Tick, const FNiagaraComputeInstanceData* Instance);
 	void BuildConstantBuffers(FNiagaraGPUSystemTick& Tick) const;
@@ -287,16 +226,11 @@ private:
 
 	uint32 FramesBeforeTickFlush = 0;
 
-	TArray<FNiagaraGPUSystemTick> Ticks_RT;
 	FGlobalDistanceFieldParameterData GlobalDistanceFieldParams;
 
 	/** A buffer of list sizes used by UpdateFreeIDBuffers to allow overlapping several dispatches. */
 	FRWBuffer FreeIDListSizesBuffer;
 	uint32 NumAllocatedFreeIDListSizes = 0;
-	bool bFreeIDListSizesBufferCleared = false;
-
-	/** List of emitter instances which need their free ID buffers updated post render. */
-	FEmitterInstanceList DeferredIDBufferUpdates;
 
 	struct DummyUAV
 	{
@@ -324,17 +258,17 @@ private:
 	void ResetEmptyUAVPool(TMap<EPixelFormat, DummyUAVPool>& UAVMap);
 	void ResetEmptyUAVPools();
 
-	uint32 NumTicksThatRequireDistanceFieldData = 0;
-	uint32 NumTicksThatRequireDepthBuffer = 0;
-	uint32 NumTicksThatRequireEarlyViewData = 0;
-	uint32 NumTicksThatRequireRayTracingScene = 0;
+	uint32 NumProxiesThatRequireDistanceFieldData = 0;
+	uint32 NumProxiesThatRequireDepthBuffer = 0;
+	uint32 NumProxiesThatRequireEarlyViewData = 0;
+	uint32 NumProxiesThatRequireRayTracingScene = 0;
 
 	int32 TotalDispatchesThisFrame = 0;
 
-	TArray<FNiagaraComputeSharedContext*> ContextsPerStage[(int)ETickStage::Max];
-	TArray<FNiagaraGPUSystemTick*> TicksPerStage[(int)ETickStage::Max];
+	bool bRequiresReadback = false;
+	TArray<FNiagaraSystemGpuComputeProxy*> ProxiesPerStage[ENiagaraGpuComputeTickStage::Max];
 
-	TArray<uint32> CountsToRelease[(int)ETickStage::Max];
+	FNiagaraGpuDispatchList DispatchListPerStage[ENiagaraGpuComputeTickStage::Max];
 
 #if NIAGARA_COMPUTEDEBUG_ENABLED
 	TUniquePtr<FNiagaraGpuComputeDebug> GpuComputeDebugPtr;
diff --git a/Engine/Plugins/FX/Niagara/Source/Niagara/Classes/NiagaraGPUSystemTick.h b/Engine/Plugins/FX/Niagara/Source/Niagara/Classes/NiagaraGPUSystemTick.h
new file mode 100644
--- /dev/null
+++ b/Engine/Plugins/FX/Niagara/Source/Niagara/Classes/NiagaraGPUSystemTick.h
@@ -0,0 +1,110 @@
+// Copyright Epic Games, Inc. All Rights Reserved.
+
+#pragma once
+
+#include "NiagaraComputeExecutionContext.h"
+#include "RHIGPUReadback.h"
+
+struct FNiagaraComputeDataInterfaceInstanceData
+{
+	void* PerInstanceDataForRT = nullptr;
+	TMap<FNiagaraDataInterfaceProxy*, int32> InterfaceProxiesToOffsets;
+	uint32 PerInstanceDataSize = 0;
+	uint32 Instances = 0;
+};
+
+struct FNiagaraComputeInstanceData
+{
+	FNiagaraComputeInstanceData()
+	{
+		bResetData = false;
+		bStartNewOverlapGroup = false;
+		bUsesSimStages = false;
+		bUsesOldShaderStages = false;
+	}
+
+	FNiagaraGpuSpawnInfo SpawnInfo;
+	uint8* EmitterParamData = nullptr;
+	uint8* ExternalParamData = nullptr;
+	FNiagaraComputeExecutionContext* Context = nullptr;
+	TArray<FNiagaraDataInterfaceProxy*> DataInterfaceProxies;
+	TArray<FNiagaraDataInterfaceProxyRW*> IterationDataInterfaceProxies;
+	uint32 ParticleCountFence = INDEX_NONE;
+	uint32 bResetData : 1;
+	uint32 bStartNewOverlapGroup : 1;
+	uint32 bUsesSimStages : 1;
+	uint32 bUsesOldShaderStages : 1;
+
+	bool IsOutputStage(FNiagaraDataInterfaceProxy* DIProxy, uint32 CurrentStage) const;
+	bool IsIterationStage(FNiagaraDataInterfaceProxy* DIProxy, uint32 CurrentStage) const;
+	FNiagaraDataInterfaceProxyRW* FindIterationInterface(uint32 SimulationStageIndex) const;
+};
+
+/*
+	Represents all the information needed to dispatch a single tick of a FNiagaraSystemInstance.
+	This object will be created on the game thread and passed to the renderthread.
+
+	It contains the PerInstance data buffer for every DataInterface referenced by the system as well
+	as the Data required to dispatch updates for each Emitter in the system.
+
+	DataInterface data is packed tightly. It includes a TMap that associates the data interface with
+	the offset into the packed buffer. At that offset is the Per-Instance data for this System.
+
+	InstanceData_ParamData_Packed packs FNiagaraComputeInstanceData and ParamData into one buffer.
+	There is padding after the array of FNiagaraComputeInstanceData so we can upload ParamData directly into a UniformBuffer
+	(it is 16 byte aligned).
+
+*/
+class FNiagaraGPUSystemTick
+{
+public:
+	void Init(FNiagaraSystemInstance* InSystemInstance);
+	void Destroy();
+
+	FORCEINLINE TArrayView<FNiagaraComputeInstanceData> GetInstances() const
+	{
+		return MakeArrayView(reinterpret_cast<FNiagaraComputeInstanceData*>(InstanceData_ParamData_Packed), InstanceCount);
+	};
+
+	enum EUniformBufferType
+	{
+		UBT_FirstSystemType = 0,
+		UBT_Global = UBT_FirstSystemType,
+		UBT_System,
+		UBT_Owner,
+		UBT_NumSystemTypes,
+
+		UBT_FirstInstanceType = UBT_NumSystemTypes,
+		UBT_Emitter = UBT_FirstInstanceType,
+		UBT_External,
+
+		UBT_NumTypes,
+
+		UBT_NumInstanceTypes = UBT_NumTypes - UBT_NumSystemTypes,
+	};
+
+	FUniformBufferRHIRef GetUniformBuffer(EUniformBufferType Type, const FNiagaraComputeInstanceData* InstanceData, bool Previous) const;
+	const uint8* GetUniformBufferSource(EUniformBufferType Type, const FNiagaraComputeInstanceData* InstanceData, bool Previous) const;
+
+public:
+	// Transient data used by the RT
+	TArray<FUniformBufferRHIRef> UniformBuffers;
+
+	// data assigned by GT
+	FNiagaraSystemInstanceID SystemInstanceID = 0LL;						//-TODO: Remove?
+	class FNiagaraSystemGpuComputeProxy* SystemGpuComputeProxy = nullptr;	//-TODO: Can we remove this?
+	FNiagaraComputeDataInterfaceInstanceData* DIInstanceData = nullptr;
+	uint8* InstanceData_ParamData_Packed = nullptr;
+	uint8* GlobalParamData = nullptr;
+	uint8* SystemParamData = nullptr;
+	uint8* OwnerParamData = nullptr;
+	uint32 InstanceCount = 0;
+	uint32 TotalDispatches = 0;
+	uint32 NumInstancesWithSimStages = 0;									//-TODO: Remove me
+	bool bIsFinalTick = false;
+
+#if !UE_BUILD_SHIPPING && !UE_BUILD_TEST
+	// Debugging only
+	TConstArrayView<FNiagaraComputeInstanceData> InstanceDataDebuggingOnly;
+#endif
+};
diff --git a/Engine/Plugins/FX/Niagara/Source/Niagara/Classes/NiagaraScriptExecutionContext.h b/Engine/Plugins/FX/Niagara/Source/Niagara/Classes/NiagaraScriptExecutionContext.h
--- a/Engine/Plugins/FX/Niagara/Source/Niagara/Classes/NiagaraScriptExecutionContext.h
+++ b/Engine/Plugins/FX/Niagara/Source/Niagara/Classes/NiagaraScriptExecutionContext.h
@@ -12,11 +12,8 @@ NiagaraEmitterInstance.h: Niagara emitter simulation class
 #include "NiagaraEmitter.h"
 #include "NiagaraScriptExecutionParameterStore.h"
 #include "NiagaraTypes.h"
-#include "RHIGPUReadback.h"
 
 struct FNiagaraDataInterfaceProxy;
-class FNiagaraGPUInstanceCountManager;
-class NiagaraEmitterInstanceBatcher;
 
 /** All scripts that will use the system script execution context. */
 enum class ENiagaraSystemSimulationScript : uint8
@@ -263,274 +260,3 @@ public:
 	/** Generates a table of DI calls unique to the passed system instance. These are then accesss inside the PerInstanceFunctionHook. */
 	virtual bool GeneratePerInstanceDIFunctionTable(FNiagaraSystemInstance* Inst, TArray<FNiagaraPerInstanceDIFuncInfo>& OutFunctions);
 };
-
-struct FNiagaraGpuSpawnInfoParams
-{
-	float IntervalDt;
-	float InterpStartDt;
-	int32 SpawnGroup;
-	int32 GroupSpawnStartIndex;
-};
-
-struct FNiagaraGpuSpawnInfo
-{
-	uint32 EventSpawnTotal = 0;
-	uint32 SpawnRateInstances = 0;
-	uint32 MaxParticleCount = 0;
-	int32 SpawnInfoStartOffsets[NIAGARA_MAX_GPU_SPAWN_INFOS];
-	FNiagaraGpuSpawnInfoParams SpawnInfoParams[NIAGARA_MAX_GPU_SPAWN_INFOS];
-
-	void Reset()
-	{
-		EventSpawnTotal = 0;
-		SpawnRateInstances = 0;
-		MaxParticleCount = 0;
-		for (int32 i = 0; i < NIAGARA_MAX_GPU_SPAWN_INFOS; ++i)
-		{
-			SpawnInfoStartOffsets[i] = 0;
-
-			SpawnInfoParams[i].IntervalDt = 0;
-			SpawnInfoParams[i].InterpStartDt = 0;
-			SpawnInfoParams[i].SpawnGroup = 0;
-			SpawnInfoParams[i].GroupSpawnStartIndex = 0;
-		}		
-	}
-};
-
-class FNiagaraRHIUniformBufferLayout : public FRHIResource
-{
-public:
-	explicit FNiagaraRHIUniformBufferLayout(const TCHAR* LayoutName) : FRHIResource(RRT_None), UBLayout(LayoutName) { }
-
-	FRHIUniformBufferLayout UBLayout;
-};
-
-struct FNiagaraComputeSharedContext
-{
-	int32 ScratchIndex = INDEX_NONE;
-	int32 ScratchTickStage = INDEX_NONE;
-
-	uint32 ParticleCountReadFence = 1;
-	uint32 ParticleCountWriteFence = 0;
-};
-
-struct FNiagaraComputeSharedContextDeleter
-{
-	void operator()(FNiagaraComputeSharedContext* Ptr) const
-	{
-		if (Ptr)
-		{
-			ENQUEUE_RENDER_COMMAND(NiagaraDeleteSharedContext)([RT_Ptr=Ptr](FRHICommandListImmediate& RHICmdList) { delete RT_Ptr; });
-		}
-	}
-};
-
-struct FNiagaraComputeExecutionContext
-{
-	FNiagaraComputeExecutionContext();
-	~FNiagaraComputeExecutionContext();
-
-	void Reset(NiagaraEmitterInstanceBatcher* Batcher);
-
-	void InitParams(UNiagaraScript* InGPUComputeScript, ENiagaraSimTarget InSimTarget, const uint32 InDefaultSimulationStageIndex, int32 InMaxUpdateIterations, const TSet<uint32> InSpawnStages);
-	void DirtyDataInterfaces();
-	bool Tick(FNiagaraSystemInstance* ParentSystemInstance);
-
-	bool OptionalContexInit(FNiagaraSystemInstance* ParentSystemInstance);
-
-	void PostTick();
-
-	void SetDataToRender(FNiagaraDataBuffer* InDataToRender);
-	void SetTranslucentDataToRender(FNiagaraDataBuffer* InTranslucentDataToRender);
-	FNiagaraDataBuffer* GetDataToRender(bool bIsLowLatencyTranslucent) const { return bIsLowLatencyTranslucent && TranslucentDataToRender ? TranslucentDataToRender : DataToRender; }
-
-	struct 
-	{
-		// The offset at which the GPU instance count (see FNiagaraGPUInstanceCountManager()).
-		uint32 GPUCountOffset = INDEX_NONE;
-		// The CPU instance count at the time the GPU count readback was issued. Always bigger or equal to the GPU count.
-		uint32 CPUCount = 0;
-	}  EmitterInstanceReadback;
-	
-#if !UE_BUILD_SHIPPING
-	const TCHAR* GetDebugSimName() const { return *DebugSimName; }
-	void SetDebugSimName(const TCHAR* InDebugSimName) { DebugSimName = InDebugSimName; }
-#else
-	const TCHAR* GetDebugSimName() const { return TEXT(""); }
-	void SetDebugSimName(const TCHAR*) { }
-#endif
-
-private:
-	void ResetInternal(NiagaraEmitterInstanceBatcher* Batcher);
-
-public:
-	static uint32 TickCounter;
-
-#if !UE_BUILD_SHIPPING
-	FString DebugSimName;
-#endif
-#if STATS
-	TWeakObjectPtr<UNiagaraEmitter> EmitterPtr; // emitter pointer used to report captured gpu stats
-#endif
-
-	const TArray<UNiagaraDataInterface*>& GetDataInterfaces()const { return CombinedParamStore.GetDataInterfaces(); }
-
-	class FNiagaraDataSet *MainDataSet;
-	UNiagaraScript* GPUScript;
-	class FNiagaraShaderScript*  GPUScript_RT;
-
-	// persistent layouts used to create the constant buffers for the compute sim shader
-	TRefCountPtr<FNiagaraRHIUniformBufferLayout> ExternalCBufferLayout;
-
-	//Dynamic state updated either from GT via RT commands or from the RT side sim code itself.
-	//TArray<uint8, TAlignedHeapAllocator<16>> ParamData_RT;		// RT side copy of the parameter data
-	FNiagaraScriptInstanceParameterStore CombinedParamStore;
-#if DO_CHECK
-	TArray< FString >  DIClassNames;
-#endif
-
-	TArray<FNiagaraDataInterfaceProxy*> DataInterfaceProxies;
-
-	// Most current buffer that can be used for rendering.
-	FNiagaraDataBuffer* DataToRender = nullptr;
-
-	// Optional buffer which can be used to render translucent data with no latency (i.e. this frames data)
-	FNiagaraDataBuffer* TranslucentDataToRender = nullptr;
-
-	// Game thread spawn info will be sent to the render thread inside FNiagaraComputeInstanceData
-	FNiagaraGpuSpawnInfo GpuSpawnInfo_GT;
-
-	uint32 DefaultSimulationStageIndex;
-	uint32 MaxUpdateIterations;
-	TSet<uint32> SpawnStages;
-
-	bool HasInterpolationParameters;
-
-	/** Temp data used in NiagaraEmitterInstanceBatcher::ExecuteAll() to avoid creating a map per FNiagaraComputeExecutionContext */
-	mutable uint32 ScratchNumInstances = 0;
-	mutable uint32 ScratchMaxInstances = 0;
-	/** Number of ticks we've processed for this emitter this frame. */
-	mutable uint32 TicksThisFrame = 0;
-
-	TArray<FSimulationStageMetaData> SimStageInfo;
-
-	bool IsOutputStage(FNiagaraDataInterfaceProxy* DIProxy, uint32 CurrentStage) const;
-	bool IsIterationStage(FNiagaraDataInterfaceProxy* DIProxy, uint32 CurrentStage) const;
-	FNiagaraDataInterfaceProxyRW* FindIterationInterface(const TArray<FNiagaraDataInterfaceProxyRW*>& InProxies, uint32 SimulationStageIndex) const;
-	const FSimulationStageMetaData* GetSimStageMetaData(uint32 SimulationStageIndex) const;
-};
-
-struct FNiagaraDataInterfaceInstanceData
-{
-	void* PerInstanceDataForRT;
-	TMap<FNiagaraDataInterfaceProxy*, int32> InterfaceProxiesToOffsets;
-	uint32 PerInstanceDataSize;
-	uint32 Instances;
-
-	~FNiagaraDataInterfaceInstanceData()
-	{}
-};
-
-//TODO: Rename FNiagaraGPUEmitterTick?
-
-
-struct FNiagaraSimStageData
-{
-	FNiagaraDataBuffer* Source = nullptr;
-	FNiagaraDataBuffer* Destination = nullptr;
-	FNiagaraDataInterfaceProxyRW* AlternateIterationSource = nullptr;
-	uint32 SourceCountOffset = 0;
-	uint32 DestinationCountOffset = 0;
-	uint32 SourceNumInstances = 0;
-	uint32 DestinationNumInstances = 0;
-	const FSimulationStageMetaData* StageMetaData = nullptr;
-};
-
-struct FNiagaraComputeInstanceData
-{
-	FNiagaraGpuSpawnInfo SpawnInfo;
-	uint8* EmitterParamData = nullptr;
-	uint8* ExternalParamData = nullptr;
-	FNiagaraComputeExecutionContext* Context = nullptr;
-	TArray<FNiagaraDataInterfaceProxy*> DataInterfaceProxies;
-	TArray<FNiagaraDataInterfaceProxyRW*> IterationDataInterfaceProxies;
-	bool bStartNewOverlapGroup = false;
-	bool bUsesSimStages = false;
-	bool bUsesOldShaderStages = false;
-	TArray<FNiagaraSimStageData, TInlineAllocator<1>> SimStageData;
-
-	bool IsOutputStage(FNiagaraDataInterfaceProxy* DIProxy, uint32 CurrentStage) const;
-	bool IsIterationStage(FNiagaraDataInterfaceProxy* DIProxy, uint32 CurrentStage) const;
-	FNiagaraDataInterfaceProxyRW* FindIterationInterface(uint32 SimulationStageIndex) const;
-};
-
-
-
-/*
-	Represents all the information needed to dispatch a single tick of a FNiagaraSystemInstance.
-	This object will be created on the game thread and passed to the renderthread.
-
-	It contains the PerInstance data buffer for every DataInterface referenced by the system as well
-	as the Data required to dispatch updates for each Emitter in the system.
-
-	DataInterface data is packed tightly. It includes a TMap that associates the data interface with
-	the offset into the packed buffer. At that offset is the Per-Instance data for this System.
-
-	InstanceData_ParamData_Packed packs FNiagaraComputeInstanceData and ParamData into one buffer.
-	There is padding after the array of FNiagaraComputeInstanceData so we can upload ParamData directly into a UniformBuffer
-	(it is 16 byte aligned).
-
-*/
-class FNiagaraGPUSystemTick
-{
-public:
-	void Init(FNiagaraSystemInstance* InSystemInstance);
-	void Destroy();
-
-	FORCEINLINE bool IsValid()const{ return InstanceData_ParamData_Packed != nullptr; }
-	FORCEINLINE FNiagaraComputeInstanceData* GetInstanceData()const{ return reinterpret_cast<FNiagaraComputeInstanceData*>(InstanceData_ParamData_Packed); }
-
-	enum EUniformBufferType
-	{
-		UBT_FirstSystemType = 0,
-		UBT_Global = UBT_FirstSystemType,
-		UBT_System,
-		UBT_Owner,
-		UBT_NumSystemTypes,
-
-		UBT_FirstInstanceType = UBT_NumSystemTypes,
-		UBT_Emitter = UBT_FirstInstanceType,
-		UBT_External,
-
-		UBT_NumTypes,
-
-		UBT_NumInstanceTypes = UBT_NumTypes - UBT_NumSystemTypes,
-	};
-
-	FUniformBufferRHIRef GetUniformBuffer(EUniformBufferType Type, const FNiagaraComputeInstanceData* InstanceData, bool Previous) const;
-	const uint8* GetUniformBufferSource(EUniformBufferType Type, const FNiagaraComputeInstanceData* InstanceData, bool Previous) const;
-
-public:
-	// Transient data used by the RT
-	TArray<FUniformBufferRHIRef> UniformBuffers;
-
-	// data assigned by GT
-	FNiagaraSystemInstanceID SystemInstanceID = 0LL;
-	FNiagaraComputeSharedContext* SharedContext = nullptr;
-	FNiagaraDataInterfaceInstanceData* DIInstanceData = nullptr;
-	uint8* InstanceData_ParamData_Packed = nullptr;
-	uint8* GlobalParamData = nullptr;
-	uint8* SystemParamData = nullptr;
-	uint8* OwnerParamData = nullptr;
-	uint32 Count = 0;
-	uint32 TotalDispatches = 0;
-	uint32 NumInstancesWithSimStages = 0;
-	uint32 ParticleCountFence = 0;
-	bool bRequiresDistanceFieldData = false;
-	bool bRequiresDepthBuffer = false;
-	bool bRequiresEarlyViewData = false;
-	bool bRequiresViewUniformBuffer = false;
-	bool bRequiresRayTracingScene = false;
-	bool bNeedsReset = false;
-	bool bIsFinalTick = false;
-};
diff --git a/Engine/Plugins/FX/Niagara/Source/Niagara/Private/NiagaraComputeExecutionContext.cpp b/Engine/Plugins/FX/Niagara/Source/Niagara/Private/NiagaraComputeExecutionContext.cpp
new file mode 100644
--- /dev/null
+++ b/Engine/Plugins/FX/Niagara/Source/Niagara/Private/NiagaraComputeExecutionContext.cpp
@@ -0,0 +1,438 @@
+// Copyright Epic Games, Inc. All Rights Reserved.
+
+#include "NiagaraComputeExecutionContext.h"
+#include "NiagaraStats.h"
+#include "NiagaraEmitterInstanceBatcher.h"
+#include "NiagaraDataInterface.h"
+#include "NiagaraSystemInstance.h"
+#include "NiagaraSystemGpuComputeProxy.h"
+#include "NiagaraWorldManager.h"
+#include "NiagaraEmitterInstance.h"
+#include "NiagaraFunctionLibrary.h"
+#include "NiagaraGPUInstanceCountManager.h"
+#include "NiagaraGPUSystemTick.h"
+#include "NiagaraDataInterfaceRW.h"
+
+FNiagaraComputeExecutionContext::FNiagaraComputeExecutionContext()
+	: MainDataSet(nullptr)
+	, GPUScript(nullptr)
+	, GPUScript_RT(nullptr)
+{
+	ExternalCBufferLayout = new FNiagaraRHIUniformBufferLayout(TEXT("Niagara GPU External CBuffer"));
+}
+
+FNiagaraComputeExecutionContext::~FNiagaraComputeExecutionContext()
+{
+	// EmitterInstanceReadback.GPUCountOffset should be INDEX_NONE at this point to ensure the index is reused.
+	// When the batcher is being destroyed though, we don't free the index, but this would not be leaking.
+	// check(EmitterInstanceReadback.GPUCountOffset == INDEX_NONE);
+	SetDataToRender(nullptr);
+
+	ExternalCBufferLayout = nullptr;
+}
+
+void FNiagaraComputeExecutionContext::Reset(NiagaraEmitterInstanceBatcher* Batcher)
+{
+	NiagaraEmitterInstanceBatcher* RT_Batcher = Batcher && !Batcher->IsPendingKill() ? Batcher : nullptr;
+	ENQUEUE_RENDER_COMMAND(ResetRT)(
+		[RT_Batcher, RT_Context=this](FRHICommandListImmediate& RHICmdList)
+		{
+			RT_Context->ResetInternal(RT_Batcher);
+		}
+	);
+}
+
+void FNiagaraComputeExecutionContext::InitParams(UNiagaraScript* InGPUComputeScript, ENiagaraSimTarget InSimTarget, const uint32 InDefaultSimulationStageIndex, const int32 InMaxUpdateIterations, const TSet<uint32> InSpawnStages)
+{
+	GPUScript = InGPUComputeScript;
+	CombinedParamStore.InitFromOwningContext(InGPUComputeScript, InSimTarget, true);
+	DefaultSimulationStageIndex = InDefaultSimulationStageIndex;
+	MaxUpdateIterations = InMaxUpdateIterations;
+	SpawnStages.Empty();
+
+	SpawnStages.Append(InSpawnStages);
+	
+	HasInterpolationParameters = GPUScript && GPUScript->GetComputedVMCompilationId().HasInterpolatedParameters();
+
+	if (InGPUComputeScript)
+	{
+		FNiagaraVMExecutableData& VMData = InGPUComputeScript->GetVMExecutableData();
+		if (VMData.IsValid() && VMData.SimulationStageMetaData.Num() > 0)
+		{
+			SimStageInfo = VMData.SimulationStageMetaData;
+
+			int32 FoundMaxUpdateIterations = SimStageInfo[SimStageInfo.Num() - 1].MaxStage;
+
+			// Some useful debugging code should we need to look up differences between old and new
+			const bool bDebugSimStages = false;
+			if (bDebugSimStages)
+			{
+				UE_LOG(LogNiagara, Log, TEXT("Stored vs:"));
+				bool bPass = FoundMaxUpdateIterations == MaxUpdateIterations;
+				UE_LOG(LogNiagara, Log, TEXT("MaxUpdateIterations: %d vs %d %s"), FoundMaxUpdateIterations, MaxUpdateIterations, bPass ? TEXT("Pass") : TEXT("FAIL!!!!!!!!"));
+
+				int32 NumSpawnFound = 0;
+				bool bMatchesFound = true;
+				for (int32 i = 0; i < SimStageInfo.Num(); i++)
+				{
+					if (SimStageInfo[i].bSpawnOnly)
+					{
+						NumSpawnFound++;
+
+						if (!SpawnStages.Contains(SimStageInfo[i].MinStage))
+						{
+							bMatchesFound = false;
+							UE_LOG(LogNiagara, Log, TEXT("Missing spawn stage: %d FAIL!!!!!!!!!"), SimStageInfo[i].MinStage);
+						}
+					}
+				}
+
+				bPass = SpawnStages.Num() == NumSpawnFound;
+				UE_LOG(LogNiagara, Log, TEXT("SpawnStages.Num(): %d vs %d %s"), NumSpawnFound, SpawnStages.Num(), bPass ? TEXT("Pass") : TEXT("FAIL!!!!!!!!"));
+
+				TArray<FNiagaraVariable> Params;
+				CombinedParamStore.GetParameters(Params);
+				for (FNiagaraVariable& Var : Params)
+				{
+					if (!Var.IsDataInterface())
+						continue;
+
+					UNiagaraDataInterface* DI = CombinedParamStore.GetDataInterface(Var);
+					UNiagaraDataInterfaceRWBase* DIRW = Cast<UNiagaraDataInterfaceRWBase>(DI);
+					if (DIRW)
+					{
+						for (int32 i = 0; i < SimStageInfo.Num(); i++)
+						{
+							if (SimStageInfo[i].IterationSource == Var.GetName())
+							{
+								if (!DIRW->IterationShaderStages.Contains(SimStageInfo[i].MinStage))
+								{
+									UE_LOG(LogNiagara, Log, TEXT("Missing iteration stage for %s: %d FAIL!!!!!!!!!"), *Var.GetName().ToString(), SimStageInfo[i].MinStage);
+								}
+							}
+
+							if (SimStageInfo[i].OutputDestinations.Contains(Var.GetName()))
+							{
+								if (!DIRW->OutputShaderStages.Contains(SimStageInfo[i].MinStage))
+								{
+									UE_LOG(LogNiagara, Log, TEXT("Missing output stage for %s: %d FAIL!!!!!!!!!"), *Var.GetName().ToString(), SimStageInfo[i].MinStage);
+								}
+							}
+						}
+					}
+
+				}
+			}
+
+
+			// Set the values that we are using from compiled data instead...
+			MaxUpdateIterations = SimStageInfo[SimStageInfo.Num() - 1].MaxStage;
+			SpawnStages.Empty();
+
+			for (int32 i = 0; i < SimStageInfo.Num(); i++)
+			{
+				if (SimStageInfo[i].bSpawnOnly)
+				{
+					SpawnStages.Add(SimStageInfo[i].MinStage);
+				}
+			}
+
+			
+		}
+	}
+
+	
+#if DO_CHECK
+	// DI Parameters are the same between all shader permutations so we can just get the first one
+	FNiagaraShaderRef Shader = InGPUComputeScript->GetRenderThreadScript()->GetShaderGameThread(0);
+	if (Shader.IsValid())
+	{
+		DIClassNames.Empty(Shader->GetDIParameters().Num());
+		for (const FNiagaraDataInterfaceParamRef& DIParams : Shader->GetDIParameters())
+		{
+			DIClassNames.Add(DIParams.DIType.Get(Shader.GetPointerTable().DITypes)->GetClass()->GetName());
+		}
+	}
+	else
+	{
+		DIClassNames.Empty(InGPUComputeScript->GetRenderThreadScript()->GetDataInterfaceParamInfo().Num());
+		for (const FNiagaraDataInterfaceGPUParamInfo& DIParams : InGPUComputeScript->GetRenderThreadScript()->GetDataInterfaceParamInfo())
+		{
+			DIClassNames.Add(DIParams.DIClassName);
+		}
+	}
+#endif
+}
+
+
+const FSimulationStageMetaData* FNiagaraComputeExecutionContext::GetSimStageMetaData(uint32 SimulationStageIndex) const
+{
+	if (SimStageInfo.Num() > 0)
+	{
+		for (int32 i = 0; i < SimStageInfo.Num(); i++)
+		{
+			if (SimulationStageIndex >= (uint32)SimStageInfo[i].MinStage  && SimulationStageIndex < (uint32)SimStageInfo[i].MaxStage)
+			{
+				return &SimStageInfo[i];
+			}
+		}
+	}
+	return nullptr;
+}
+
+bool FNiagaraComputeExecutionContext::IsOutputStage(FNiagaraDataInterfaceProxy* DIProxy, uint32 CurrentStage) const
+{
+	const FSimulationStageMetaData* MetaData = GetSimStageMetaData(CurrentStage);
+	if (MetaData && DIProxy && !DIProxy->SourceDIName.IsNone())
+	{
+		if (MetaData->OutputDestinations.Contains(DIProxy->SourceDIName))
+			return true;
+	}
+	else if (DIProxy && SimStageInfo.Num() == 0)
+	{
+		return DIProxy->IsOutputStage_DEPRECATED(CurrentStage);
+	}
+	return false;
+}
+
+bool FNiagaraComputeExecutionContext::IsIterationStage(FNiagaraDataInterfaceProxy* DIProxy, uint32 CurrentStage) const
+{
+	const FSimulationStageMetaData* MetaData = GetSimStageMetaData(CurrentStage);
+	if (MetaData && DIProxy && !DIProxy->SourceDIName.IsNone())
+	{
+		if (MetaData->IterationSource.IsNone()) // Per particle iteration...
+			return false;
+
+		if (MetaData->IterationSource == DIProxy->SourceDIName)
+			return true;
+	}
+	else if (DIProxy && SimStageInfo.Num() == 0)
+	{
+		return DIProxy->IsIterationStage_DEPRECATED(CurrentStage);
+	}
+	return false;
+}
+
+FNiagaraDataInterfaceProxyRW* FNiagaraComputeExecutionContext::FindIterationInterface(const TArray<FNiagaraDataInterfaceProxyRW*>& InProxies, uint32 CurrentStage) const
+{
+	const FSimulationStageMetaData* MetaData = GetSimStageMetaData(CurrentStage);
+	if (MetaData)
+	{
+		if (MetaData->IterationSource.IsNone()) // Per particle iteration...
+			return nullptr;
+
+		for (FNiagaraDataInterfaceProxyRW* Proxy : InProxies)
+		{
+			if (Proxy->SourceDIName == MetaData->IterationSource)
+				return Proxy;
+		}
+
+		UE_LOG(LogNiagara, Verbose, TEXT("FNiagaraComputeExecutionContext::FindIterationInterface could not find IterationInterface %s"), *MetaData->IterationSource.ToString());
+
+		return nullptr;
+	}
+	else if (SimStageInfo.Num() == 0)
+	{
+		// Fallback to old shader stages
+		for (FNiagaraDataInterfaceProxyRW* Proxy : InProxies)
+		{
+			if (Proxy->IsIterationStage_DEPRECATED(CurrentStage))
+				return Proxy;
+		}
+	}
+
+	return nullptr;
+}
+
+void FNiagaraComputeExecutionContext::DirtyDataInterfaces()
+{
+	CombinedParamStore.MarkInterfacesDirty();
+}
+
+bool FNiagaraComputeExecutionContext::Tick(FNiagaraSystemInstance* ParentSystemInstance)
+{
+	if (CombinedParamStore.GetInterfacesDirty())
+	{
+#if DO_CHECK
+		const TArray<UNiagaraDataInterface*> &DataInterfaces = CombinedParamStore.GetDataInterfaces();
+		// We must make sure that the data interfaces match up between the original script values and our overrides...
+		if (DIClassNames.Num() != DataInterfaces.Num())
+		{
+			UE_LOG(LogNiagara, Warning, TEXT("Mismatch between Niagara GPU Execution Context data interfaces and those in its script!"));
+			return false;
+		}
+
+		for (int32 i = 0; i < DIClassNames.Num(); ++i)
+		{
+			FString UsedClassName = DataInterfaces[i]->GetClass()->GetName();
+			if (DIClassNames[i] != UsedClassName)
+			{
+				UE_LOG(LogNiagara, Warning, TEXT("Mismatched class between Niagara GPU Execution Context data interfaces and those in its script!\nIndex:%d\nShader:%s\nScript:%s")
+					, i, *DIClassNames[i], *UsedClassName);
+			}
+		}
+#endif
+		CombinedParamStore.Tick();
+	}
+
+	return true;
+}
+
+bool FNiagaraComputeExecutionContext::OptionalContexInit(FNiagaraSystemInstance* ParentSystemInstance)
+{
+	if (GPUScript)
+	{
+		FNiagaraVMExecutableData& VMData = GPUScript->GetVMExecutableData();
+
+		if (VMData.IsValid() && VMData.bNeedsGPUContextInit)
+		{
+			const TArray<UNiagaraDataInterface*>& DataInterfaces = CombinedParamStore.GetDataInterfaces();
+			for (int32 i = 0; i < DataInterfaces.Num(); i++)
+			{
+				UNiagaraDataInterface* Interface = DataInterfaces[i];
+
+				int32 UserPtrIdx = VMData.DataInterfaceInfo[i].UserPtrIdx;
+				if (UserPtrIdx != INDEX_NONE)
+				{
+					void* InstData = ParentSystemInstance->FindDataInterfaceInstanceData(Interface);
+					if (Interface->NeedsGPUContextInit())
+					{
+						Interface->GPUContextInit(VMData.DataInterfaceInfo[i], InstData, ParentSystemInstance);
+					}
+				}
+			}
+		}
+	}
+	return true;
+}
+
+void FNiagaraComputeExecutionContext::PostTick()
+{
+	//If we're for interpolated spawn, copy over the previous frame's parameters into the Prev parameters.
+	if (HasInterpolationParameters)
+	{
+		CombinedParamStore.CopyCurrToPrev();
+	}
+}
+
+void FNiagaraComputeExecutionContext::ReleaseReadbackCounter(FNiagaraGPUInstanceCountManager& GPUInstanceCountManager)
+{
+	if ( EmitterInstanceReadback.GPUCountOffset != INDEX_NONE )
+	{
+		check(EmitterInstanceReadback.GPUCountOffset != CountOffset_RT);
+
+		GPUInstanceCountManager.FreeEntry(EmitterInstanceReadback.GPUCountOffset);
+		EmitterInstanceReadback.GPUCountOffset = INDEX_NONE;
+	}
+}
+
+void FNiagaraComputeExecutionContext::ResetInternal(NiagaraEmitterInstanceBatcher* Batcher)
+{
+	checkf(IsInRenderingThread(), TEXT("Can only reset the gpu context from the render thread"));
+
+	if (Batcher)
+	{
+		ReleaseReadbackCounter(Batcher->GetGPUInstanceCounterManager());
+		Batcher->GetGPUInstanceCounterManager().FreeEntry(CountOffset_RT);
+	}
+
+	CurrentNumInstances_RT = 0;
+	CountOffset_RT = INDEX_NONE;
+	EmitterInstanceReadback.GPUCountOffset = INDEX_NONE;
+
+	SetDataToRender(nullptr);
+}
+
+void FNiagaraComputeExecutionContext::SetDataToRender(FNiagaraDataBuffer* InDataToRender)
+{
+	if (DataToRender)
+	{
+		DataToRender->ReleaseReadRef();
+	}
+
+	DataToRender = InDataToRender;
+
+	if (DataToRender)
+	{
+		DataToRender->AddReadRef();
+	}
+
+	// This call the DataToRender should be equal to the TranslucentDataToRender so we can release the read ref
+	if (TranslucentDataToRender)
+	{
+		ensure((DataToRender == nullptr) || (DataToRender == TranslucentDataToRender));
+		TranslucentDataToRender->ReleaseReadRef();
+		TranslucentDataToRender = nullptr;
+	}
+}
+
+void FNiagaraComputeExecutionContext::SetTranslucentDataToRender(FNiagaraDataBuffer* InTranslucentDataToRender)
+{
+	if (TranslucentDataToRender)
+	{
+		TranslucentDataToRender->ReleaseReadRef();
+	}
+
+	TranslucentDataToRender = InTranslucentDataToRender;
+
+	if (TranslucentDataToRender)
+	{
+		TranslucentDataToRender->AddReadRef();
+	}
+}
+
+bool FNiagaraComputeInstanceData::IsOutputStage(FNiagaraDataInterfaceProxy* DIProxy, uint32 CurrentStage) const
+{
+	if (bUsesOldShaderStages)
+	{
+		return DIProxy->IsOutputStage_DEPRECATED(CurrentStage);
+	}
+	else if (bUsesSimStages)
+	{
+		return Context->IsOutputStage(DIProxy, CurrentStage);
+	}
+	return false;
+}
+
+bool FNiagaraComputeInstanceData::IsIterationStage(FNiagaraDataInterfaceProxy* DIProxy, uint32 CurrentStage) const
+{
+
+	if (bUsesOldShaderStages)
+	{
+		return DIProxy->IsIterationStage_DEPRECATED(CurrentStage);
+	}
+	else if (bUsesSimStages)
+	{
+		return Context->IsIterationStage(DIProxy, CurrentStage);
+	}
+	return false;
+}
+
+FNiagaraDataInterfaceProxyRW* FNiagaraComputeInstanceData::FindIterationInterface(uint32 SimulationStageIndex) const
+{
+	if (bUsesOldShaderStages)
+	{
+		FNiagaraDataInterfaceProxyRW* IterationInterface = nullptr;
+		for (FNiagaraDataInterfaceProxyRW* Interface : IterationDataInterfaceProxies)
+		{
+			if (Interface->IsIterationStage_DEPRECATED(SimulationStageIndex))
+			{
+				if (IterationInterface)
+				{
+					UE_LOG(LogNiagara, Error, TEXT("Multiple output Data Interfaces found for current stage"));
+				}
+				else
+				{
+					IterationInterface = Interface;
+				}
+			}
+		}
+
+		return IterationInterface;
+	}
+	else if (bUsesSimStages)
+	{
+		return Context->FindIterationInterface(IterationDataInterfaceProxies, SimulationStageIndex);
+	}
+	return nullptr;
+}
diff --git a/Engine/Plugins/FX/Niagara/Source/Niagara/Private/NiagaraDataInterfaceExport.cpp b/Engine/Plugins/FX/Niagara/Source/Niagara/Private/NiagaraDataInterfaceExport.cpp
--- a/Engine/Plugins/FX/Niagara/Source/Niagara/Private/NiagaraDataInterfaceExport.cpp
+++ b/Engine/Plugins/FX/Niagara/Source/Niagara/Private/NiagaraDataInterfaceExport.cpp
@@ -119,7 +119,7 @@ struct FNDIExportProxy : public FNiagaraDataInterfaceProxy
 
 		// Ensure our buffer if big enough to hold all the data
 		const int32 AllocationRounding = 64;
-		const int32 DataBufferNumInstances = FMath::Max(Context.ComputeInstanceData->SimStageData.Last().DestinationNumInstances, Context.ComputeInstanceData->SimStageData.Last().SourceNumInstances);
+		const int32 DataBufferNumInstances = Context.ComputeInstanceData->Context->CurrentMaxInstances_RT;
 		int32 NumInstances;
 		if (InstanceData->AllocationMode == ENDIExport_GPUAllocationMode::PerParticle)
 		{
diff --git a/Engine/Plugins/FX/Niagara/Source/Niagara/Private/NiagaraDataInterfaceGrid3DCollection.cpp b/Engine/Plugins/FX/Niagara/Source/Niagara/Private/NiagaraDataInterfaceGrid3DCollection.cpp
--- a/Engine/Plugins/FX/Niagara/Source/Niagara/Private/NiagaraDataInterfaceGrid3DCollection.cpp
+++ b/Engine/Plugins/FX/Niagara/Source/Niagara/Private/NiagaraDataInterfaceGrid3DCollection.cpp
@@ -2580,7 +2580,7 @@ void FNiagaraDataInterfaceProxyGrid3DCollectionProxy::PreStage(FRHICommandList&
 
 		// #todo(dmp): for now, if there is an output DI that is NOT an iteration DI AND if both the iteration DI and this DI have the same total number of instances, do not do the UAV clear prior to the stage.
 		// this isn't optimal, but should work to some degree to reduce overhead in cases where we have multiple grids of the same resolution being processed/written to in 1 stage
-		FNiagaraDataInterfaceProxyRW* IterationInterface = Context.ComputeInstanceData->SimStageData[Context.SimulationStageIndex].AlternateIterationSource;
+		FNiagaraDataInterfaceProxyRW* IterationInterface = Context.SimStageData->AlternateIterationSource;
 		RHICmdList.Transition(FRHITransitionInfo(ProxyData->DestinationData->GridBuffer.UAV, ERHIAccess::SRVMask, ERHIAccess::UAVCompute));
 		if (!Context.IsIterationStage)
 		{
diff --git a/Engine/Plugins/FX/Niagara/Source/Niagara/Private/NiagaraDataInterfaceParticleRead.cpp b/Engine/Plugins/FX/Niagara/Source/Niagara/Private/NiagaraDataInterfaceParticleRead.cpp
--- a/Engine/Plugins/FX/Niagara/Source/Niagara/Private/NiagaraDataInterfaceParticleRead.cpp
+++ b/Engine/Plugins/FX/Niagara/Source/Niagara/Private/NiagaraDataInterfaceParticleRead.cpp
@@ -497,48 +497,54 @@ struct FNiagaraDataInterfaceParametersCS_ParticleRead : public FNiagaraDataInter
 			return;
 		}
 
-		FNiagaraDataBuffer* SourceData;
-		uint32 NumSpawnedInstances = 0, IDAcquireTag = 0, InstanceCountOffset = 0xffffffff;
-		const bool bReadingOwnEmitter = Context.ComputeInstanceData->Context == InstanceData->SourceEmitterGPUContext;
-		if (bReadingOwnEmitter)
+		FNiagaraDataBuffer* SourceData = nullptr;
+		uint32 NumSpawnedInstances = 0;
+		uint32 IDAcquireTag = 0;
+
+		// Test for reading ourself
+		if (Context.ComputeInstanceData->Context == InstanceData->SourceEmitterGPUContext)
 		{
-			// If the current execution context is the same as the source emitter's context, it means we're reading from
-			// ourselves. We can't use SourceDataSet->GetCurrentData() in that case, because EndSimulate() has already been
-			// called on the current emitter, and the current data has been set to the destination data. We need to use the
-			// current compute instance data to get to the input buffers.
-			const FNiagaraSimStageData& SimStageData = Context.ComputeInstanceData->SimStageData[Context.SimulationStageIndex];
-			SourceData = SimStageData.Source;
-			InstanceCountOffset = SimStageData.SourceCountOffset;
-
-			// We still want to get the spawn count and ID acquire tag from the destination data, because that's where
-			// NiagaraEmitterInstanceBatcher::Run() stores them.
-			if (SimStageData.Destination != nullptr)
+			// Pull source data from SimStageData as this has the correct information until all stages are dispatched
+			// Note: In the case of an in place (aka partial update) stage this does not work as we need to bind the
+			//       same buffer as both UAV & SRV.  For the moment the user must disable partial writes.
+			//       Another option would be to read from the UAV, but that would equally race.
+			//-TODO: Automate this in some way or surface as feedback...
+			if (const FSimulationStageMetaData* StageMetaData = Context.SimStageData->StageMetaData )
+			{
+				if ( StageMetaData->bPartialParticleUpdate )
+				{
+					UE_LOG(LogNiagara, Error, TEXT("Particle read DI reading self '%s' on stage '%s' is unsafe, please disable partial writes on the stage."), *InstanceData->SourceEmitterName, *StageMetaData->SimulationStageName.ToString());
+					SetErrorParams(RHICmdList, ComputeShader, true);
+					return;
+				}
+			}
+
+			SourceData = Context.SimStageData->Source;
+
+			// Pull some spawned / acquire tag from destination data as we are reading from ourselves this is what we are currently doing in this stage.
+			if (Context.SimStageData->Destination)
 			{
-				NumSpawnedInstances = SimStageData.Destination->GetNumSpawnedInstances();
-				IDAcquireTag = SimStageData.Destination->GetIDAcquireTag();
+				ensure(SourceData != nullptr);
+				NumSpawnedInstances = Context.SimStageData->Destination->GetNumSpawnedInstances();
+				IDAcquireTag = Context.SimStageData->Destination->GetIDAcquireTag();
 			}
-			// When we don't write particle data the destination is invalid, therefore we can pull from the source
-			else if (SimStageData.Source != nullptr)
+			// In cases where we do not have destination data pull from source (i.e. this is an iteration on none particle stages)
+			else if (SourceData)
 			{
-				NumSpawnedInstances = SimStageData.Source->GetNumSpawnedInstances();
-				IDAcquireTag = SimStageData.Source->GetIDAcquireTag();
+				NumSpawnedInstances = SourceData->GetNumSpawnedInstances();
+				IDAcquireTag = SourceData->GetIDAcquireTag();
 			}
 		}
 		else
 		{
 			SourceData = SourceDataSet->GetCurrentData();
-			if (SourceData)
+			if ( SourceData )
 			{
 				NumSpawnedInstances = SourceData->GetNumSpawnedInstances();
 				IDAcquireTag = SourceData->GetIDAcquireTag();
-				InstanceCountOffset = SourceData->GetGPUInstanceCountBufferOffset();
 			}
 		}
 
-		SetShaderValue(RHICmdList, ComputeShader, NumSpawnedParticlesParam, NumSpawnedInstances);
-		SetShaderValue(RHICmdList, ComputeShader, SpawnedParticlesAcquireTagParam, IDAcquireTag);
-		SetSRVParameter(RHICmdList, ComputeShader, SpawnedIDsBufferParam, GetIntSRVWithFallback(SourceDataSet->GetGPUFreeIDs()));
-
 		if (!SourceData)
 		{
 			SetErrorParams(RHICmdList, ComputeShader, true);
@@ -559,7 +565,10 @@ struct FNiagaraDataInterfaceParametersCS_ParticleRead : public FNiagaraDataInter
 		// There's no need to transition the input buffers, because the grouping logic inside NiagaraEmitterInstanceBatcher ensures that our source emitter has ran before us,
 		// and its buffers have been transitioned to readable.
 
-		SetShaderValue(RHICmdList, ComputeShader, InstanceCountOffsetParam, InstanceCountOffset);
+		SetShaderValue(RHICmdList, ComputeShader, InstanceCountOffsetParam, SourceData->GetGPUInstanceCountBufferOffset());
+		SetShaderValue(RHICmdList, ComputeShader, NumSpawnedParticlesParam, NumSpawnedInstances);
+		SetShaderValue(RHICmdList, ComputeShader, SpawnedParticlesAcquireTagParam, IDAcquireTag);
+		SetSRVParameter(RHICmdList, ComputeShader, SpawnedIDsBufferParam, GetIntSRVWithFallback(SourceDataSet->GetGPUFreeIDs()));
 		SetSRVParameter(RHICmdList, ComputeShader, IDToIndexTableParam, GetIntSRVWithFallback(SourceData->GetGPUIDToIndexTable()));
 		SetSRVParameter(RHICmdList, ComputeShader, InputFloatBufferParam, GetFloatSRVWithFallback(SourceData->GetGPUBufferFloat()));
 		SetSRVParameter(RHICmdList, ComputeShader, InputIntBufferParam, GetIntSRVWithFallback(SourceData->GetGPUBufferInt()));
diff --git a/Engine/Plugins/FX/Niagara/Source/Niagara/Private/NiagaraDataSet.cpp b/Engine/Plugins/FX/Niagara/Source/Niagara/Private/NiagaraDataSet.cpp
--- a/Engine/Plugins/FX/Niagara/Source/Niagara/Private/NiagaraDataSet.cpp
+++ b/Engine/Plugins/FX/Niagara/Source/Niagara/Private/NiagaraDataSet.cpp
@@ -395,14 +395,6 @@ void FNiagaraDataSet::Dump(int32 StartIndex, int32 NumInstances, const FString&
 	}
 }
 
-void FNiagaraDataSet::ReleaseGPUInstanceCounts(FNiagaraGPUInstanceCountManager& GPUInstanceCountManager)
-{
-	for (FNiagaraDataBuffer* Buffer : Data)
-	{
-		Buffer->ReleaseGPUInstanceCount(GPUInstanceCountManager);
-	}
-}
-
 void FNiagaraDataSet::AllocateGPUFreeIDs(uint32 InNumInstances, FRHICommandList& RHICmdList, ERHIFeatureLevel::Type FeatureLevel, const TCHAR* DebugSimName)
 {
 	checkSlow(GetSimTarget() == ENiagaraSimTarget::GPUComputeSim && RequiresPersistentIDs());
@@ -430,24 +422,11 @@ void FNiagaraDataSet::AllocateGPUFreeIDs(uint32 InNumInstances, FRHICommandList&
 	FRWBuffer NewFreeIDsBuffer;
 	NewFreeIDsBuffer.Initialize(DebugBufferName, sizeof(int32), NumIDsToAlloc, EPixelFormat::PF_R32_SINT, BUF_Static);
 
-	FRHIShaderResourceView* ExistingBuffer;
-	if (GPUNumAllocatedIDs > 0)
-	{
-		// We must maintain the existing list of free IDs.
-		// The free IDs buffer was written in the previous simulation step, but hasn't been transitioned to read yet, so we must
-		// transition it explicitly here. The new buffer will be transitioned by NiagaraEmitterInstanceBatcher::DispatchAllOnCompute(),
-		// so there's no need for a barrier at the end of this function.
-		RHICmdList.Transition(FRHITransitionInfo(GPUFreeIDs.UAV, ERHIAccess::Unknown, ERHIAccess::SRVCompute));
-		ExistingBuffer = GPUFreeIDs.SRV;
-	}
-	else
-	{
-		ExistingBuffer = FNiagaraRenderer::GetDummyIntBuffer();
-	}
-
+	FRHIShaderResourceView* ExistingBuffer = GPUNumAllocatedIDs > 0 ? GPUFreeIDs.SRV.GetReference() : FNiagaraRenderer::GetDummyIntBuffer();
 
 	RHICmdList.Transition(FRHITransitionInfo(NewFreeIDsBuffer.UAV, ERHIAccess::Unknown, ERHIAccess::UAVCompute));
 	NiagaraInitGPUFreeIDList(RHICmdList, FeatureLevel, NumIDsToAlloc, NewFreeIDsBuffer, GPUNumAllocatedIDs, ExistingBuffer);
+	RHICmdList.Transition(FRHITransitionInfo(NewFreeIDsBuffer.UAV, ERHIAccess::UAVCompute, ERHIAccess::SRVCompute));
 
 	GPUFreeIDs = MoveTemp(NewFreeIDsBuffer);
 	GPUNumAllocatedIDs = NumIDsToAlloc;
@@ -563,13 +542,10 @@ FNiagaraDataBuffer::FNiagaraDataBuffer(FNiagaraDataSet* InOwner)
 FNiagaraDataBuffer::~FNiagaraDataBuffer()
 {
 	check(!IsInUse());
-	// If this is data for a GPU emitter, we have to release the GPU instance counts for reuse.
-	// The only exception is if the batcher was pending kill and we couldn't enqueue a rendering command, 
-	// in which case this would have been released on the game thread and not from the batcher DataSetsToDestroy_RT.
-	check(!IsInRenderingThread() || GPUInstanceCountBufferOffset == INDEX_NONE);
 	DEC_MEMORY_STAT_BY(STAT_NiagaraParticleMemory, FloatData.GetAllocatedSize() + Int32Data.GetAllocatedSize());
-
-	DEC_MEMORY_STAT_BY(STAT_NiagaraGPUParticleMemory, GPUBufferFloat.NumBytes + GPUBufferHalf.NumBytes + GPUBufferInt.NumBytes + GPUIDToIndexTable.NumBytes);
+#if NIAGARA_MEMORY_TRACKING
+	DEC_MEMORY_STAT_BY(STAT_NiagaraGPUParticleMemory, AllocationSizeBytes);
+#endif
 }
 
 int32 FNiagaraDataBuffer::TransferInstance(FNiagaraDataBuffer& SourceBuffer, int32 InstanceIndex, bool bRemoveFromSource)
@@ -751,20 +727,15 @@ void FNiagaraDataBuffer::Allocate(uint32 InNumInstances, bool bMaintainExisting)
 	}
 }
 
-void FNiagaraDataBuffer::AllocateGPU(uint32 InNumInstances, FNiagaraGPUInstanceCountManager& GPUInstanceCountManager, FRHICommandList& RHICmdList, ERHIFeatureLevel::Type FeatureLevel, const TCHAR* DebugSimName)
+void FNiagaraDataBuffer::AllocateGPU(FRHICommandList& RHICmdList, uint32 InNumInstances, ERHIFeatureLevel::Type FeatureLevel, const TCHAR* DebugSimName)
 {
 	static constexpr uint32 GPUBufferFlags = BUF_Static | BUF_SourceCopy;
 
 	CheckUsage(false);
-
 	checkSlow(Owner->GetSimTarget() == ENiagaraSimTarget::GPUComputeSim);
 
-	//uint32 OldOffset = GPUInstanceCountBufferOffset;
-	// Release previous entry if any.
-	GPUInstanceCountManager.FreeEntry(GPUInstanceCountBufferOffset);
-	// Get a new entry currently set to 0, since simulation will increment it to the actual instance count.
-	GPUInstanceCountBufferOffset = GPUInstanceCountManager.AcquireEntry();
-	//UE_LOG(LogNiagara, Log, TEXT("AllocateGPU %p GPUInstanceCountBufferOffsetOld: %d New: %d"), this, OldOffset, GPUInstanceCountBufferOffset);
+	// Clear count offset
+	GPUInstanceCountBufferOffset = INDEX_NONE;
 
 	NumInstancesAllocated = InNumInstances;
 
@@ -782,9 +753,10 @@ void FNiagaraDataBuffer::AllocateGPU(uint32 InNumInstances, FNiagaraGPUInstanceC
 	Int32Stride = PaddedNumInstances * sizeof(int32);
 	HalfStride = PaddedNumInstances * sizeof(FFloat16);
 
-	DEC_MEMORY_STAT_BY(STAT_NiagaraGPUParticleMemory, GPUBufferFloat.NumBytes + GPUBufferHalf.NumBytes + GPUBufferInt.NumBytes + GPUIDToIndexTable.NumBytes);
 #if NIAGARA_MEMORY_TRACKING
-	Owner->BufferSizeBytes -= GPUBufferFloat.NumBytes + GPUBufferHalf.NumBytes + GPUBufferInt.NumBytes + GPUIDToIndexTable.NumBytes;
+	DEC_MEMORY_STAT_BY(STAT_NiagaraGPUParticleMemory, AllocationSizeBytes);
+	Owner->BufferSizeBytes -= AllocationSizeBytes;
+	AllocationSizeBytes = 0;
 #endif
 
 	if (PaddedNumInstances == 0)
@@ -861,9 +833,9 @@ void FNiagaraDataBuffer::AllocateGPU(uint32 InNumInstances, FNiagaraGPUInstanceC
 				{
 					GPUIDToIndexTable.Release();
 				}
-				TCHAR DebugBufferName[128];
-				FCString::Snprintf(DebugBufferName, UE_ARRAY_COUNT(DebugBufferName), TEXT("NiagaraIDToIndexTable_%s_%p"), DebugSimName ? DebugSimName : TEXT(""), this);
-				GPUIDToIndexTable.Initialize(DebugBufferName, sizeof(int32), NumNeededElems, EPixelFormat::PF_R32_SINT, BUF_Static);
+				TStringBuilder<128> DebugBufferName;
+				DebugBufferName.Appendf(TEXT("NiagaraIDToIndexTable_%s_%p"), DebugSimName ? DebugSimName : TEXT(""), this);
+				GPUIDToIndexTable.Initialize(DebugBufferName.ToString(), sizeof(int32), NumNeededElems, EPixelFormat::PF_R32_SINT, BUF_Static);
 				Transitions.Add(FRHITransitionInfo(GPUIDToIndexTable.UAV, ERHIAccess::Unknown, ERHIAccess::SRVCompute));
 			}
 		}
@@ -871,19 +843,43 @@ void FNiagaraDataBuffer::AllocateGPU(uint32 InNumInstances, FNiagaraGPUInstanceC
 		// NiagaraEmitterInstanceBatcher expects the buffers to be readable before running the sim.
 		RHICmdList.Transition(MakeArrayView(Transitions.GetData(), Transitions.Num()));
 	}
-	INC_MEMORY_STAT_BY(STAT_NiagaraGPUParticleMemory, GPUBufferFloat.NumBytes + GPUBufferHalf.NumBytes + GPUBufferInt.NumBytes + GPUIDToIndexTable.NumBytes);
 #if NIAGARA_MEMORY_TRACKING
-	Owner->BufferSizeBytes += GPUBufferFloat.NumBytes + GPUBufferHalf.NumBytes + GPUBufferInt.NumBytes + GPUIDToIndexTable.NumBytes;
+	AllocationSizeBytes = GPUBufferFloat.NumBytes + GPUBufferHalf.NumBytes + GPUBufferInt.NumBytes + GPUIDToIndexTable.NumBytes;
+	Owner->BufferSizeBytes += AllocationSizeBytes;
+	INC_MEMORY_STAT_BY(STAT_NiagaraGPUParticleMemory, AllocationSizeBytes);
 #endif
 }
 
+void FNiagaraDataBuffer::AliasGPU(FNiagaraDataBuffer* BufferToAlias)
+{
+#if NIAGARA_MEMORY_TRACKING
+	DEC_MEMORY_STAT_BY(STAT_NiagaraGPUParticleMemory, AllocationSizeBytes);
+	Owner->BufferSizeBytes -= AllocationSizeBytes;
+	AllocationSizeBytes = 0;
+#endif
+
+	GPUInstanceCountBufferOffset = BufferToAlias->GPUInstanceCountBufferOffset;
+	GPUBufferFloat = BufferToAlias->GPUBufferFloat;
+	GPUBufferInt = BufferToAlias->GPUBufferInt;
+	GPUIDToIndexTable = BufferToAlias->GPUIDToIndexTable;
+	GPUBufferHalf = BufferToAlias->GPUBufferHalf;
+	NumInstances = BufferToAlias->NumInstances;
+	NumInstancesAllocated = BufferToAlias->NumInstancesAllocated;
+	FloatStride = BufferToAlias->FloatStride;
+	Int32Stride = BufferToAlias->Int32Stride;
+	HalfStride = BufferToAlias->HalfStride;
+	NumSpawnedInstances = BufferToAlias->NumSpawnedInstances;
+	IDAcquireTag = BufferToAlias->IDAcquireTag;
+}
+
 void FNiagaraDataBuffer::ReleaseGPU()
 {
 	check(IsInRenderingThread());
 
-	DEC_MEMORY_STAT_BY(STAT_NiagaraGPUParticleMemory, GPUBufferFloat.NumBytes + GPUBufferHalf.NumBytes + GPUBufferInt.NumBytes + GPUIDToIndexTable.NumBytes);
 #if NIAGARA_MEMORY_TRACKING
-	Owner->BufferSizeBytes -= GPUBufferFloat.NumBytes + GPUBufferHalf.NumBytes + GPUBufferInt.NumBytes + GPUIDToIndexTable.NumBytes;
+	DEC_MEMORY_STAT_BY(STAT_NiagaraGPUParticleMemory, AllocationSizeBytes);
+	Owner->BufferSizeBytes -= AllocationSizeBytes;
+	AllocationSizeBytes = 0;
 #endif
 
 	NumInstancesAllocated = 0;
@@ -1276,12 +1272,6 @@ void FNiagaraDataBuffer::UnsetShaderParams(FRHICommandList& RHICmdList, FNiagara
 	}
 }
 
-void FNiagaraDataBuffer::ReleaseGPUInstanceCount(FNiagaraGPUInstanceCountManager& GPUInstanceCountManager)
-{
-	GPUInstanceCountManager.FreeEntry(GPUInstanceCountBufferOffset);
-}
-
-
 #if WITH_EDITOR
 void FScopedNiagaraDataSetGPUReadback::ReadbackData(NiagaraEmitterInstanceBatcher* InBatcher, FNiagaraDataSet* InDataSet)
 {
diff --git a/Engine/Plugins/FX/Niagara/Source/Niagara/Private/NiagaraEmitterInstance.cpp b/Engine/Plugins/FX/Niagara/Source/Niagara/Private/NiagaraEmitterInstance.cpp
--- a/Engine/Plugins/FX/Niagara/Source/Niagara/Private/NiagaraEmitterInstance.cpp
+++ b/Engine/Plugins/FX/Niagara/Source/Niagara/Private/NiagaraEmitterInstance.cpp
@@ -5,6 +5,7 @@
 #include "NiagaraConstants.h"
 #include "NiagaraRenderer.h"
 #include "NiagaraSystemInstance.h"
+#include "NiagaraSystemGpuComputeProxy.h"
 #include "NiagaraDataInterface.h"
 #include "NiagaraEmitterInstanceBatcher.h"
 #include "NiagaraScriptExecutionContext.h"
@@ -156,10 +157,6 @@ FNiagaraEmitterInstance::~FNiagaraEmitterInstance()
 		ENQUEUE_RENDER_COMMAND(FDeleteContextCommand)(
 			[Batcher_RT, ExecContext=GPUExecContext, DataSet= ParticleDataSet](FRHICommandListImmediate& RHICmdList)
 			{
-				if ( Batcher_RT != nullptr )
-				{
-					Batcher_RT->ReleaseInstanceCounts_RenderThread(ExecContext, DataSet);
-				}
 				if ( ExecContext != nullptr )
 				{
 					delete ExecContext;
@@ -560,11 +557,6 @@ void FNiagaraEmitterInstance::ResetSimulation(bool bKillExisting /*= true*/)
 	{
 		bResetPending = true;
 		TotalSpawnedParticles = 0;
-
-		if (CachedEmitter->SimTarget == ENiagaraSimTarget::GPUComputeSim && GPUExecContext != nullptr)
-		{
-			GPUExecContext->Reset(Batcher);
-		}
 	}
 }
 
@@ -573,11 +565,6 @@ void FNiagaraEmitterInstance::OnPooledReuse()
 	// Ensure we kill any existing particles and mark our buffers for reset
 	bResetPending = true;
 	TotalSpawnedParticles = 0;
-
-	if (CachedEmitter != nullptr && CachedEmitter->SimTarget == ENiagaraSimTarget::GPUComputeSim && GPUExecContext != nullptr)
-	{
-		GPUExecContext->Reset(Batcher);
-	}
 }
 
 void FNiagaraEmitterInstance::SetParticleComponentActive(FObjectKey ComponentKey, int32 ParticleID) const
@@ -856,15 +843,16 @@ int32 FNiagaraEmitterInstance::GetNumParticlesGPUInternal() const
 {
 	check(GPUExecContext);
 
-	// If our fence has not been passed return the TotalSpawnedParticles as a 'guess' to the amount
-	FNiagaraComputeSharedContext* GPUSharedContext = GetParentSystemInstance()->GetComputeSharedContext();
-	if (GPUSharedContext->ParticleCountReadFence > GPUSharedContext->ParticleCountWriteFence)
+	if (GPUExecContext->ParticleCountReadFence <= GPUExecContext->ParticleCountWriteFence)
+	{
+		// Fence has passed we read directly from the GPU Exec Context which will have the most up-to-date information
+		return GPUExecContext->CurrentNumInstances_RT;
+	}
+	else
 	{
+		// Fence has not been passed return the TotalSpawnedParticles as a 'guess' to the amount
 		return TotalSpawnedParticles;
 	}
-
-	// If the fence has passed we read directly from the GPU Exec Context which will have the most up-to-date information
-	return GPUExecContext->ScratchNumInstances;
 }
 
 const FNiagaraEmitterHandle& FNiagaraEmitterInstance::GetEmitterHandle() const
@@ -1178,6 +1166,12 @@ void FNiagaraEmitterInstance::PreTick()
 		}
 
 		bResetPending = false;
+
+		if ( GPUExecContext )
+		{
+			GPUExecContext->bResetPending_GT = true;
+			GPUExecContext->GpuSpawnInfo_GT.Reset();
+		}
 	}
 
 	++TickCount;
diff --git a/Engine/Plugins/FX/Niagara/Source/Niagara/Private/NiagaraEmitterInstanceBatcher.cpp b/Engine/Plugins/FX/Niagara/Source/Niagara/Private/NiagaraEmitterInstanceBatcher.cpp
--- a/Engine/Plugins/FX/Niagara/Source/Niagara/Private/NiagaraEmitterInstanceBatcher.cpp
+++ b/Engine/Plugins/FX/Niagara/Source/Niagara/Private/NiagaraEmitterInstanceBatcher.cpp
@@ -16,6 +16,7 @@
 #include "NiagaraRayTracingHelper.h"
 #include "NiagaraRenderer.h"
 #include "NiagaraScriptExecutionContext.h"
+#include "NiagaraSystemGpuComputeProxy.h"
 #include "NiagaraShader.h"
 #include "NiagaraShaderParticleID.h"
 #include "NiagaraSortingGPU.h"
@@ -124,34 +125,55 @@ namespace NiagaraEmitterInstanceBatcherLocal
 		ECVF_Default
 	);
 
-	static ETickStage CalculateTickStage(const FNiagaraGPUSystemTick& Tick)
+	static int32 GAddDispatchGroupDrawEvent = 0;
+	static FAutoConsoleVariableRef CVarAddDispatchGroupDrawEvent(
+		TEXT("fx.Niagara.Batcher.AddDispatchGroupDrawEvent"),
+		GAddDispatchGroupDrawEvent,
+		TEXT("Add a draw event marker around each dispatch group."),
+		ECVF_Default
+	);
+
+	#if !WITH_EDITOR
+		constexpr int32 GDebugLogging = 0;
+	#else
+		static int32 GDebugLogging = 0;
+		static FAutoConsoleVariableRef CVarNiagaraDebugLogging(
+			TEXT("fx.Niagara.Batcher.DebugLogging"),
+			GDebugLogging,
+			TEXT("Enables a lot of spew to the log to debug the batcher."),
+			ECVF_Default
+		);
+	#endif
+
+	template<typename TTransitionArrayType>
+	static void AddDataBufferTransitions(TTransitionArrayType& BeforeTransitionArray, TTransitionArrayType& AfterTransitionArray, FNiagaraDataBuffer* DestinationData, ERHIAccess BeforeState = ERHIAccess::SRVMask, ERHIAccess AfterState = ERHIAccess::UAVCompute)
 	{
-		if (!GNiagaraAllowTickBeforeRender || Tick.bRequiresDistanceFieldData || Tick.bRequiresDepthBuffer || Tick.bRequiresRayTracingScene)
+		if (FRHIUnorderedAccessView* FloatUAV = DestinationData->GetGPUBufferFloat().UAV )
 		{
-			return ETickStage::PostOpaqueRender;
+			BeforeTransitionArray.Emplace(FloatUAV, BeforeState, AfterState);
+			AfterTransitionArray.Emplace(FloatUAV, AfterState, BeforeState);
 		}
-
-		if (Tick.bRequiresEarlyViewData)
+		if (FRHIUnorderedAccessView* HalfUAV = DestinationData->GetGPUBufferHalf().UAV)
 		{
-			return ETickStage::PostInitViews;
+			BeforeTransitionArray.Emplace(HalfUAV, BeforeState, AfterState);
+			AfterTransitionArray.Emplace(HalfUAV, AfterState, BeforeState);
 		}
-
-		if (Tick.bRequiresViewUniformBuffer)
+		if (FRHIUnorderedAccessView* IntUAV = DestinationData->GetGPUBufferInt().UAV)
 		{
-			return ETickStage::PostOpaqueRender;
+			BeforeTransitionArray.Emplace(IntUAV, BeforeState, AfterState);
+			AfterTransitionArray.Emplace(IntUAV, AfterState, BeforeState);
 		}
-		return ETickStage::PreInitViews;
 	}
 
-	static bool ShouldRunStage(const FNiagaraGPUSystemTick* Tick, FNiagaraComputeExecutionContext* Context, FNiagaraDataInterfaceProxyRW* IterationInterface, uint32 StageIndex)
+	static bool ShouldRunStage(FNiagaraComputeExecutionContext* ComputeContext, FNiagaraDataInterfaceProxyRW* IterationInterface, uint32 StageIndex, bool bResetData)
 	{
-		if (!IterationInterface || Context->SpawnStages.Num() == 0)
+		if (!IterationInterface || ComputeContext->SpawnStages.Num() == 0)
 		{
 			return true;
 		}
 
-		const bool bIsSpawnStage = Context->SpawnStages.Contains(StageIndex);
-		return Tick->bNeedsReset == bIsSpawnStage;
+		const bool bIsSpawnStage = ComputeContext->SpawnStages.Contains(StageIndex);
+		return bResetData == bIsSpawnStage;
 	}
 }
 
@@ -179,8 +201,6 @@ NiagaraEmitterInstanceBatcher::NiagaraEmitterInstanceBatcher(ERHIFeatureLevel::T
 	: FeatureLevel(InFeatureLevel)
 	, ShaderPlatform(InShaderPlatform)
 	, GPUSortManager(InGPUSortManager)
-	, NumAllocatedFreeIDListSizes(0)
-	, bFreeIDListSizesBufferCleared(false)
 {
 	// Register the batcher callback in the GPUSortManager.
 	// The callback is used to generate the initial keys and values for the GPU sort tasks,
@@ -243,12 +263,44 @@ NiagaraEmitterInstanceBatcher::~NiagaraEmitterInstanceBatcher()
 	EmitterCBufferLayout = nullptr;
 }
 
-void NiagaraEmitterInstanceBatcher::InstanceDeallocated_RenderThread(const FNiagaraSystemInstanceID InstanceID)
+void NiagaraEmitterInstanceBatcher::AddGpuComputeProxy(FNiagaraSystemGpuComputeProxy* ComputeProxy)
+{
+	check(ComputeProxy->BatcherIndex == INDEX_NONE);
+
+	const ENiagaraGpuComputeTickStage::Type TickStage = ComputeProxy->GetComputeTickStage();
+	ComputeProxy->BatcherIndex = ProxiesPerStage[TickStage].Num();
+	ProxiesPerStage[TickStage].Add(ComputeProxy);
+
+	NumProxiesThatRequireDistanceFieldData	+= ComputeProxy->RequiresDistanceFieldData() ? 1 : 0;
+	NumProxiesThatRequireDepthBuffer		+= ComputeProxy->RequiresDepthBuffer() ? 1 : 0;
+	NumProxiesThatRequireEarlyViewData		+= ComputeProxy->RequiresEarlyViewData() ? 1 : 0;
+	NumProxiesThatRequireRayTracingScene	+= ComputeProxy->RequiresRayTracingScene() ? 1 : 0;
+}
+
+void NiagaraEmitterInstanceBatcher::RemoveGpuComputeProxy(FNiagaraSystemGpuComputeProxy* ComputeProxy)
 {
+	check(ComputeProxy->BatcherIndex != INDEX_NONE);
+
+	const int32 TickStage = int32(ComputeProxy->GetComputeTickStage());
+	const int32 ProxyIndex = ComputeProxy->BatcherIndex;
+	check(ProxiesPerStage[TickStage][ProxyIndex] == ComputeProxy);
+
+	ProxiesPerStage[TickStage].RemoveAtSwap(ProxyIndex);
+	if (ProxiesPerStage[TickStage].IsValidIndex(ProxyIndex))
+	{
+		ProxiesPerStage[TickStage][ProxyIndex]->BatcherIndex = ProxyIndex;
+	}
+	ComputeProxy->BatcherIndex = INDEX_NONE;
+
+	NumProxiesThatRequireDistanceFieldData	-= ComputeProxy->RequiresDistanceFieldData() ? 1 : 0;
+	NumProxiesThatRequireDepthBuffer		-= ComputeProxy->RequiresDepthBuffer() ? 1 : 0;
+	NumProxiesThatRequireEarlyViewData		-= ComputeProxy->RequiresEarlyViewData() ? 1 : 0;
+	NumProxiesThatRequireRayTracingScene	-= ComputeProxy->RequiresRayTracingScene() ? 1 : 0;
+
 #if NIAGARA_COMPUTEDEBUG_ENABLED
 	if (FNiagaraGpuComputeDebug* GpuComputeDebug = GpuComputeDebugPtr.Get())
 	{
-		GpuComputeDebug->OnSystemDeallocated(InstanceID);
+		GpuComputeDebug->OnSystemDeallocated(ComputeProxy->GetSystemInstanceID());
 	}
 #endif
 #if !UE_BUILD_SHIPPING
@@ -256,80 +308,52 @@ void NiagaraEmitterInstanceBatcher::InstanceDeallocated_RenderThread(const FNiag
 		[&](const FDebugReadbackInfo& Info)
 		{
 			// In the unlikely event we have one in the queue make sure it's marked as complete with no data in it
-			if ( Info.InstanceID == InstanceID )
+			if ( Info.InstanceID == ComputeProxy->GetSystemInstanceID())
 			{
 				Info.DebugInfo->Frame.CopyFromGPUReadback(nullptr, nullptr, nullptr, 0, 0, 0, 0, 0);
 				Info.DebugInfo->bWritten = true;
 			}
-			return Info.InstanceID == InstanceID;
+			return Info.InstanceID == ComputeProxy->GetSystemInstanceID();
 		}
 	);
 #endif
-
-	int iTick = 0;
-	while ( iTick < Ticks_RT.Num() )
-	{
-		FNiagaraGPUSystemTick& Tick = Ticks_RT[iTick];
-		if (Tick.SystemInstanceID == InstanceID)
-		{
-			ensure(NumTicksThatRequireDistanceFieldData >= (Tick.bRequiresDistanceFieldData ? 1u : 0u));
-			ensure(NumTicksThatRequireDepthBuffer >= (Tick.bRequiresDepthBuffer ? 1u : 0u));
-			ensure(NumTicksThatRequireEarlyViewData >= (Tick.bRequiresEarlyViewData ? 1u : 0u));
-			ensure(NumTicksThatRequireRayTracingScene >= (Tick.bRequiresRayTracingScene ? 1u : 0u));
-
-			NumTicksThatRequireDistanceFieldData -= Tick.bRequiresDistanceFieldData ? 1 : 0;
-			NumTicksThatRequireDepthBuffer -= Tick.bRequiresDepthBuffer ? 1 : 0;
-			NumTicksThatRequireEarlyViewData -= Tick.bRequiresEarlyViewData ? 1 : 0;
-			NumTicksThatRequireRayTracingScene -= Tick.bRequiresRayTracingScene ? 1 : 0;
-
-			//-OPT: Since we can't RemoveAtSwap (due to ordering issues) if may be better to not remove and flag as dead
-			Tick.Destroy();
-			Ticks_RT.RemoveAt(iTick);
-		}
-		else
-		{
-			++iTick;
-		}
-	}
 }
 
 void NiagaraEmitterInstanceBatcher::BuildConstantBuffers(FNiagaraGPUSystemTick& Tick) const
 {
-	if (!Tick.Count)
+	if (!Tick.InstanceCount)
 	{
 		return;
 	}
 
-	FNiagaraComputeInstanceData* EmittersData = Tick.GetInstanceData();
+	TArrayView<FNiagaraComputeInstanceData> InstanceData = Tick.GetInstances();
 
 	// first go through and figure out if we need to support interpolated spawning
 	bool HasInterpolationParameters = false;
 	bool HasMultipleStages = false;
-	for (uint32 CountIt = 0; CountIt < Tick.Count; ++CountIt)
+	for (const FNiagaraComputeInstanceData& Instance : InstanceData)
 	{
-		HasInterpolationParameters = HasInterpolationParameters || EmittersData[CountIt].Context->HasInterpolationParameters;
-		HasMultipleStages = HasMultipleStages || EmittersData[CountIt].bUsesOldShaderStages || EmittersData[CountIt].bUsesSimStages;
+		HasInterpolationParameters = HasInterpolationParameters || Instance.Context->HasInterpolationParameters;
+		HasMultipleStages = HasMultipleStages || Instance.bUsesOldShaderStages || Instance.bUsesSimStages;
 	}
 
 	int32 BoundParameterCounts[FNiagaraGPUSystemTick::UBT_NumTypes][2];
 	FMemory::Memzero(BoundParameterCounts);
 
-	for (uint32 CountIt = 0; CountIt < Tick.Count; ++CountIt)
+	for (const FNiagaraComputeInstanceData& Instance : InstanceData)
 	{
-		FNiagaraComputeInstanceData& EmitterData = EmittersData[CountIt];
-
 		for (int32 InterpIt = 0; InterpIt < (HasInterpolationParameters ? 2 : 1); ++InterpIt)
 		{
-			BoundParameterCounts[FNiagaraGPUSystemTick::UBT_Global][InterpIt] += EmitterData.Context->GPUScript_RT->IsGlobalConstantBufferUsed_RenderThread(InterpIt) ? 1 : 0;
-			BoundParameterCounts[FNiagaraGPUSystemTick::UBT_System][InterpIt] += EmitterData.Context->GPUScript_RT->IsSystemConstantBufferUsed_RenderThread(InterpIt) ? 1 : 0;
-			BoundParameterCounts[FNiagaraGPUSystemTick::UBT_Owner][InterpIt] += EmitterData.Context->GPUScript_RT->IsOwnerConstantBufferUsed_RenderThread(InterpIt) ? 1 : 0;
-			BoundParameterCounts[FNiagaraGPUSystemTick::UBT_Emitter][InterpIt] += EmitterData.Context->GPUScript_RT->IsEmitterConstantBufferUsed_RenderThread(InterpIt) ? 1 : 0;
-			BoundParameterCounts[FNiagaraGPUSystemTick::UBT_External][InterpIt] += EmitterData.Context->GPUScript_RT->IsExternalConstantBufferUsed_RenderThread(InterpIt) ? 1 : 0;
+			BoundParameterCounts[FNiagaraGPUSystemTick::UBT_Global][InterpIt] += Instance.Context->GPUScript_RT->IsGlobalConstantBufferUsed_RenderThread(InterpIt) ? 1 : 0;
+			BoundParameterCounts[FNiagaraGPUSystemTick::UBT_System][InterpIt] += Instance.Context->GPUScript_RT->IsSystemConstantBufferUsed_RenderThread(InterpIt) ? 1 : 0;
+			BoundParameterCounts[FNiagaraGPUSystemTick::UBT_Owner][InterpIt] += Instance.Context->GPUScript_RT->IsOwnerConstantBufferUsed_RenderThread(InterpIt) ? 1 : 0;
+			BoundParameterCounts[FNiagaraGPUSystemTick::UBT_Emitter][InterpIt] += Instance.Context->GPUScript_RT->IsEmitterConstantBufferUsed_RenderThread(InterpIt) ? 1 : 0;
+			BoundParameterCounts[FNiagaraGPUSystemTick::UBT_External][InterpIt] += Instance.Context->GPUScript_RT->IsExternalConstantBufferUsed_RenderThread(InterpIt) ? 1 : 0;
 		}
 	}
 
 	const int32 InterpScale = HasInterpolationParameters ? 2 : 1;
-	const int32 BufferCount = InterpScale * (FNiagaraGPUSystemTick::UBT_NumSystemTypes + FNiagaraGPUSystemTick::UBT_NumInstanceTypes * Tick.Count);
+	const int32 BufferCount = InterpScale * (FNiagaraGPUSystemTick::UBT_NumSystemTypes + FNiagaraGPUSystemTick::UBT_NumInstanceTypes * Tick.InstanceCount);
 
 	Tick.UniformBuffers.Empty(BufferCount);
 
@@ -361,19 +385,17 @@ void NiagaraEmitterInstanceBatcher::BuildConstantBuffers(FNiagaraGPUSystemTick&
 
 		for (int32 InstanceTypeIt = FNiagaraGPUSystemTick::UBT_FirstInstanceType; InstanceTypeIt < FNiagaraGPUSystemTick::UBT_NumTypes; ++InstanceTypeIt)
 		{
-			for (uint32 InstanceIt = 0; InstanceIt < Tick.Count; ++InstanceIt)
+			for (const FNiagaraComputeInstanceData& Instance : InstanceData)
 			{
-				FNiagaraComputeInstanceData& EmitterData = EmittersData[InstanceIt];
-
 				FUniformBufferRHIRef BufferRef;
 
 				if (BoundParameterCounts[InstanceTypeIt][InterpIt])
 				{
-					FRHIUniformBufferLayout& Layout = InstanceTypeIt == FNiagaraGPUSystemTick::UBT_Emitter ? EmitterCBufferLayout->UBLayout : EmitterData.Context->ExternalCBufferLayout->UBLayout;
+					FRHIUniformBufferLayout& Layout = InstanceTypeIt == FNiagaraGPUSystemTick::UBT_Emitter ? EmitterCBufferLayout->UBLayout : Instance.Context->ExternalCBufferLayout->UBLayout;
 					if (Layout.Resources.Num() > 0 || Layout.ConstantBufferSize > 0)
 					{
 						BufferRef = RHICreateUniformBuffer(
-							Tick.GetUniformBufferSource((FNiagaraGPUSystemTick::EUniformBufferType) InstanceTypeIt, &EmitterData, !InterpIt),
+							Tick.GetUniformBufferSource((FNiagaraGPUSystemTick::EUniformBufferType) InstanceTypeIt, &Instance, !InterpIt),
 							Layout,
 							((BoundParameterCounts[InstanceTypeIt][InterpIt] > 1) || HasMultipleStages)
 							? EUniformBufferUsage::UniformBuffer_SingleFrame
@@ -384,57 +406,6 @@ void NiagaraEmitterInstanceBatcher::BuildConstantBuffers(FNiagaraGPUSystemTick&
 				Tick.UniformBuffers.Add(BufferRef);
 			}
 		}
-
-	}
-}
-
-void NiagaraEmitterInstanceBatcher::GiveSystemTick_RenderThread(FNiagaraGPUSystemTick& Tick)
-{
-	check(IsInRenderingThread());
-
-	if (!FNiagaraUtilities::AllowGPUParticles(GetShaderPlatform()))
-	{
-		return;
-	}
-
-	// Now we consume DataInterface instance data.
-	if (Tick.DIInstanceData)
-	{
-		uint8* BasePointer = (uint8*) Tick.DIInstanceData->PerInstanceDataForRT;
-
-		//UE_LOG(LogNiagara, Log, TEXT("RT Give DI (dipacket) %p (baseptr) %p"), Tick.DIInstanceData, BasePointer);
-		for(auto& Pair : Tick.DIInstanceData->InterfaceProxiesToOffsets)
-		{
-			FNiagaraDataInterfaceProxy* Proxy = Pair.Key;
-			uint8* InstanceDataPtr = BasePointer + Pair.Value;
-
-			//UE_LOG(LogNiagara, Log, TEXT("\tRT DI (proxy) %p (instancebase) %p"), Proxy, InstanceDataPtr);
-			Proxy->ConsumePerInstanceDataFromGameThread(InstanceDataPtr, Tick.SystemInstanceID);
-		}
-	}
-
-	// A note:
-	// This is making a copy of Tick. That structure is small now and we take a copy to avoid
-	// making a bunch of small allocations on the game thread. We may need to revisit this.
-	Ticks_RT.Add(Tick);
-
-	NumTicksThatRequireDistanceFieldData += Tick.bRequiresDistanceFieldData ? 1 : 0;
-	NumTicksThatRequireDepthBuffer += Tick.bRequiresDepthBuffer ? 1 : 0;
-	NumTicksThatRequireEarlyViewData += Tick.bRequiresEarlyViewData ? 1 : 0;
-	NumTicksThatRequireRayTracingScene += Tick.bRequiresRayTracingScene ? 1 : 0;
-}
-
-void NiagaraEmitterInstanceBatcher::ReleaseInstanceCounts_RenderThread(FNiagaraComputeExecutionContext* ExecContext, FNiagaraDataSet* DataSet)
-{
-	LLM_SCOPE(ELLMTag::Niagara);
-
-	if ( ExecContext != nullptr )
-	{
-		GPUInstanceCounterManager.FreeEntry(ExecContext->EmitterInstanceReadback.GPUCountOffset);
-	}
-	if ( DataSet != nullptr )
-	{
-		DataSet->ReleaseGPUInstanceCounts(GPUInstanceCounterManager);
 	}
 }
 
@@ -453,9 +424,18 @@ void NiagaraEmitterInstanceBatcher::Tick(float DeltaTime)
 void NiagaraEmitterInstanceBatcher::ProcessPendingTicksFlush(FRHICommandListImmediate& RHICmdList, bool bForceFlush)
 {
 	// No ticks are pending
-	if ( Ticks_RT.Num() == 0 )
+	bool bHasTicks = false;
+	for ( int iTickStage=0; iTickStage < ENiagaraGpuComputeTickStage::Max; ++iTickStage)
+	{
+		if ( ProxiesPerStage[iTickStage].Num() > 0 )
+		{
+			bHasTicks = true;
+			break;
+		}
+	}
+
+	if ( !bHasTicks )
 	{
-		FramesBeforeTickFlush = 0;
 		return;
 	}
 
@@ -545,15 +525,14 @@ void NiagaraEmitterInstanceBatcher::FinishDispatches()
 {
 	ReleaseTicks();
 
-	NumTicksThatRequireDistanceFieldData = 0;
-	NumTicksThatRequireDepthBuffer = 0;
-	NumTicksThatRequireEarlyViewData = 0;
-	NumTicksThatRequireRayTracingScene = 0;
-
-	for (int32 iTickStage=0; iTickStage < (int)ETickStage::Max; ++iTickStage)
+	for (FNiagaraGpuDispatchList& DispatchList : DispatchListPerStage)
 	{
-		ContextsPerStage[iTickStage].Reset();
-		TicksPerStage[iTickStage].Reset();
+		DispatchList.DispatchGroups.Empty();
+		if (DispatchList.CountsToRelease.Num() > 0)
+		{
+			GPUInstanceCounterManager.FreeEntryArray(DispatchList.CountsToRelease);
+			DispatchList.CountsToRelease.Empty();
+		}
 	}
 }
 
@@ -561,12 +540,13 @@ void NiagaraEmitterInstanceBatcher::ReleaseTicks()
 {
 	check(IsInRenderingThread());
 
-	for (FNiagaraGPUSystemTick& Tick : Ticks_RT)
+	for ( int iTickStage=0; iTickStage < ENiagaraGpuComputeTickStage::Max; ++iTickStage )
 	{
-		Tick.Destroy();
+		for ( FNiagaraSystemGpuComputeProxy* ComputeProxy : ProxiesPerStage[iTickStage] )
+		{
+			ComputeProxy->ReleaseTicks(GPUInstanceCounterManager);
+		}
 	}
-
-	Ticks_RT.Empty(0);
 }
 
 bool NiagaraEmitterInstanceBatcher::UseOverlapCompute()
@@ -574,31 +554,23 @@ bool NiagaraEmitterInstanceBatcher::UseOverlapCompute()
 	return !IsMobilePlatform(ShaderPlatform) && GNiagaraOverlapCompute;
 }
 
-bool NiagaraEmitterInstanceBatcher::ResetDataInterfaces(const FNiagaraGPUSystemTick& Tick, FNiagaraComputeInstanceData* Instance, FRHICommandList &RHICmdList, const FNiagaraShaderScript* ShaderScript) const
+void NiagaraEmitterInstanceBatcher::ResetDataInterfaces(FRHICommandList& RHICmdList, const FNiagaraGPUSystemTick& Tick, const FNiagaraComputeInstanceData& InstanceData) const
 {
-	bool ValidSpawnStage = true;
-	FNiagaraComputeExecutionContext* Context = Instance->Context;
+	// Note: All stages will contain the same bindings so if they are valid for one they are valid for all, this could change in the future
+	const FNiagaraShaderRef& ComputeShader = InstanceData.Context->GPUScript_RT->GetShader(0);
+	const TMemoryImageArray<FNiagaraDataInterfaceParamRef>& DIParameters = ComputeShader->GetDIParameters();
 
-	// Reset all rw data interface data
-	if (Tick.bNeedsReset)
+	uint32 InterfaceIndex = 0;
+	for (FNiagaraDataInterfaceProxy* Interface : InstanceData.DataInterfaceProxies)
 	{
-		// Note: All stages will contain the same bindings so if they are valid for one they are valid for all, this could change in the future
-		const FNiagaraShaderRef ComputeShader = ShaderScript->GetShader(0);
-		const TMemoryImageArray<FNiagaraDataInterfaceParamRef>& DIParameters = ComputeShader->GetDIParameters();
-
-		uint32 InterfaceIndex = 0;
-		for (FNiagaraDataInterfaceProxy* Interface : Instance->DataInterfaceProxies)
+		const FNiagaraDataInterfaceParamRef& DIParam = DIParameters[InterfaceIndex];
+		if (DIParam.Parameters.IsValid())
 		{
-			const FNiagaraDataInterfaceParamRef& DIParam = DIParameters[InterfaceIndex];
-			if (DIParam.Parameters.IsValid())
-			{
-				const FNiagaraDataInterfaceArgs TmpContext(Interface, Tick.SystemInstanceID, this);
-				Interface->ResetData(RHICmdList, TmpContext);
-			}
-			InterfaceIndex++;
+			const FNiagaraDataInterfaceArgs TmpContext(Interface, Tick.SystemInstanceID, this);
+			Interface->ResetData(RHICmdList, TmpContext);
 		}
+		InterfaceIndex++;
 	}
-	return ValidSpawnStage;
 }
 
 FNiagaraDataInterfaceProxyRW* NiagaraEmitterInstanceBatcher::FindIterationInterface(FNiagaraComputeInstanceData* Instance, const uint32 SimulationStageIndex) const
@@ -607,62 +579,52 @@ FNiagaraDataInterfaceProxyRW* NiagaraEmitterInstanceBatcher::FindIterationInterf
 	return Instance->FindIterationInterface(SimulationStageIndex);
 }
 
-void NiagaraEmitterInstanceBatcher::PreStageInterface(const FNiagaraGPUSystemTick& Tick, FNiagaraComputeInstanceData* Instance, FRHICommandList& RHICmdList, const uint32 SimulationStageIndex) const
+void NiagaraEmitterInstanceBatcher::PreStageInterface(FRHICommandList& RHICmdList, const FNiagaraGPUSystemTick& Tick, const FNiagaraComputeInstanceData& InstanceData, const FNiagaraSimStageData& SimStageData) const
 {
-	if (!Instance->Context || !Instance->Context->GPUScript_RT->IsShaderMapComplete_RenderThread())
-	{
-		return;
-	}
-
 	// Note: All stages will contain the same bindings so if they are valid for one they are valid for all, this could change in the future
-	const FNiagaraShaderRef& ComputeShader = Instance->Context->GPUScript_RT->GetShader(0);
+	const FNiagaraShaderRef& ComputeShader = InstanceData.Context->GPUScript_RT->GetShader(0);
 	const TMemoryImageArray<FNiagaraDataInterfaceParamRef>& DIParameters = ComputeShader->GetDIParameters();
 
 	uint32 InterfaceIndex = 0;
-	for (FNiagaraDataInterfaceProxy* Interface : Instance->DataInterfaceProxies)
+	for (FNiagaraDataInterfaceProxy* Interface : InstanceData.DataInterfaceProxies)
 	{
 		const FNiagaraDataInterfaceParamRef& DIParam = DIParameters[InterfaceIndex];
 		if (DIParam.Parameters.IsValid())
 		{
-			const FNiagaraDataInterfaceStageArgs TmpContext(Interface, Tick.SystemInstanceID, this, Instance, SimulationStageIndex, Instance->IsOutputStage(Interface, SimulationStageIndex), Instance->IsIterationStage(Interface, SimulationStageIndex));
+			const FNiagaraDataInterfaceStageArgs TmpContext(Interface, Tick.SystemInstanceID, this, &InstanceData, &SimStageData, InstanceData.IsOutputStage(Interface, SimStageData.StageIndex), InstanceData.IsIterationStage(Interface, SimStageData.StageIndex));
 			Interface->PreStage(RHICmdList, TmpContext);
 		}
 		InterfaceIndex++;
 	}
 }
 
-void NiagaraEmitterInstanceBatcher::PostStageInterface(const FNiagaraGPUSystemTick& Tick, FNiagaraComputeInstanceData* Instance, FRHICommandList &RHICmdList, const uint32 SimulationStageIndex) const
+void NiagaraEmitterInstanceBatcher::PostStageInterface(FRHICommandList& RHICmdList, const FNiagaraGPUSystemTick& Tick, const FNiagaraComputeInstanceData& InstanceData, const FNiagaraSimStageData& SimStageData) const
 {
-	if (!Instance->Context || !Instance->Context->GPUScript_RT->IsShaderMapComplete_RenderThread())
-	{
-		return;
-	}
-
 	// Note: All stages will contain the same bindings so if they are valid for one they are valid for all, this could change in the future
-	const FNiagaraShaderRef& ComputeShader = Instance->Context->GPUScript_RT->GetShader(0);
+	const FNiagaraShaderRef& ComputeShader = InstanceData.Context->GPUScript_RT->GetShader(0);
 	const TMemoryImageArray<FNiagaraDataInterfaceParamRef>& DIParameters = ComputeShader->GetDIParameters();
 
 	uint32 InterfaceIndex = 0;
-	for (FNiagaraDataInterfaceProxy* Interface : Instance->DataInterfaceProxies)
+	for (FNiagaraDataInterfaceProxy* Interface : InstanceData.DataInterfaceProxies)
 	{
 		const FNiagaraDataInterfaceParamRef& DIParam = DIParameters[InterfaceIndex];
 		if (DIParam.Parameters.IsValid())
 		{
-			const FNiagaraDataInterfaceStageArgs TmpContext(Interface, Tick.SystemInstanceID, this, Instance, SimulationStageIndex, Instance->IsOutputStage(Interface, SimulationStageIndex), Instance->IsIterationStage(Interface, SimulationStageIndex));
+			const FNiagaraDataInterfaceStageArgs TmpContext(Interface, Tick.SystemInstanceID, this, &InstanceData, &SimStageData, InstanceData.IsOutputStage(Interface, SimStageData.StageIndex), InstanceData.IsIterationStage(Interface, SimStageData.StageIndex));
 			Interface->PostStage(RHICmdList, TmpContext);
 		}
 		InterfaceIndex++;
 	}
 }
 
-void NiagaraEmitterInstanceBatcher::PostSimulateInterface(const FNiagaraGPUSystemTick& Tick, FNiagaraComputeInstanceData* Instance, FRHICommandList& RHICmdList, const FNiagaraShaderScript* ShaderScript) const
+void NiagaraEmitterInstanceBatcher::PostSimulateInterface(FRHICommandList& RHICmdList, const FNiagaraGPUSystemTick& Tick, const FNiagaraComputeInstanceData& InstanceData) const
 {
 	// Note: All stages will contain the same bindings so if they are valid for one they are valid for all, this could change in the future
-	const FNiagaraShaderRef ComputeShader = ShaderScript->GetShader(0);
+	const FNiagaraShaderRef& ComputeShader = InstanceData.Context->GPUScript_RT->GetShader(0);
 	const TMemoryImageArray<FNiagaraDataInterfaceParamRef>& DIParameters = ComputeShader->GetDIParameters();
 
 	uint32 InterfaceIndex = 0;
-	for (FNiagaraDataInterfaceProxy* Interface : Instance->DataInterfaceProxies)
+	for (FNiagaraDataInterfaceProxy* Interface : InstanceData.DataInterfaceProxies)
 	{
 		const FNiagaraDataInterfaceParamRef& DIParam = DIParameters[InterfaceIndex];
 		if (DIParam.Parameters.IsValid())
@@ -674,785 +636,941 @@ void NiagaraEmitterInstanceBatcher::PostSimulateInterface(const FNiagaraGPUSyste
 	}
 }
 
-void NiagaraEmitterInstanceBatcher::DispatchStage(FDispatchInstance& DispatchInstance, uint32 StageIndex, FRHICommandList& RHICmdList, FRHIUniformBuffer* ViewUniformBuffer)
-{
-	FNiagaraGPUSystemTick& Tick = *DispatchInstance.Tick;
-	FNiagaraComputeInstanceData* Instance = DispatchInstance.InstanceData;
-	FNiagaraComputeExecutionContext* Context = Instance->Context;
 
-	if (!Context || !Context->GPUScript_RT->IsShaderMapComplete_RenderThread())
-	{
-		return;
-	}
-
-	if (StageIndex == 0)
+void NiagaraEmitterInstanceBatcher::UpdateFreeIDsListSizesBuffer(FRHICommandList& RHICmdList, uint32 NumInstances)
+{
+	ERHIAccess BeforeState = ERHIAccess::UAVCompute;
+	if (NumInstances > NumAllocatedFreeIDListSizes)
 	{
-		FNiagaraComputeExecutionContext::TickCounter++;
-		if (!ResetDataInterfaces(Tick, Instance, RHICmdList, Context->GPUScript_RT))
+		constexpr uint32 ALLOC_CHUNK_SIZE = 128;
+		NumAllocatedFreeIDListSizes = Align(NumInstances, ALLOC_CHUNK_SIZE);
+		if (FreeIDListSizesBuffer.Buffer)
 		{
-			return;
+			FreeIDListSizesBuffer.Release();
 		}
-	}
-
-	const int32 PermutationId = Tick.NumInstancesWithSimStages > 0 ? Context->GPUScript_RT->ShaderStageIndexToPermutationId_RenderThread(StageIndex) : 0;
-	const FNiagaraShaderRef ComputeShader = Context->GPUScript_RT->GetShader(PermutationId);
+		FreeIDListSizesBuffer.Initialize(TEXT("NiagaraFreeIDListSizes"), sizeof(uint32), NumAllocatedFreeIDListSizes, EPixelFormat::PF_R32_SINT, BUF_Static);
 
-	const uint32 DefaultSimulationStageIndex = Instance->Context->DefaultSimulationStageIndex;
-	FNiagaraDataInterfaceProxyRW* IterationInterface = Instance->SimStageData[StageIndex].AlternateIterationSource;
-	if (!IterationInterface)
-	{
-		Run(Tick, Instance, 0, Instance->SimStageData[StageIndex].DestinationNumInstances, ComputeShader, RHICmdList, ViewUniformBuffer, Instance->SpawnInfo, false, DefaultSimulationStageIndex, StageIndex, nullptr, StageIndex > 0);
+		BeforeState = ERHIAccess::Unknown;
 	}
-	else
-	{
-		const FIntVector ElementCount = IterationInterface->GetElementCount(Tick.SystemInstanceID);
-		const uint64 TotalNumInstances = ElementCount.X * ElementCount.Y * ElementCount.Z;
-
-		// Verify the number of elements isn't higher that what we can handle
-		checkf(TotalNumInstances < uint64(TNumericLimits<int32>::Max()), TEXT("ElementCount(%d, %d, %d) for IterationInterface(%s) overflows an int32 this is not allowed"), ElementCount.X, ElementCount.Y, ElementCount.Z, *IterationInterface->SourceDIName.ToString());
 
-		Run(Tick, Instance, 0, uint32(TotalNumInstances), ComputeShader, RHICmdList, ViewUniformBuffer, Instance->SpawnInfo, false, DefaultSimulationStageIndex, StageIndex, IterationInterface);
+	{
+		SCOPED_DRAW_EVENT(RHICmdList, NiagaraGPUComputeClearFreeIDListSizes);
+		RHICmdList.Transition(FRHITransitionInfo(FreeIDListSizesBuffer.UAV, BeforeState, ERHIAccess::UAVCompute));
+		NiagaraFillGPUIntBuffer(RHICmdList, FeatureLevel, FreeIDListSizesBuffer, 0);
 	}
 }
 
-void NiagaraEmitterInstanceBatcher::AddDestinationBufferTransitions(FDispatchGroup* Group, FNiagaraDataBuffer* DestinationData)
+void NiagaraEmitterInstanceBatcher::UpdateFreeIDBuffers(FRHICommandList& RHICmdList, TConstArrayView<FNiagaraComputeExecutionContext*> Instances)
 {
-	auto AddBufferTransitions = [](FDispatchGroup* Group, FRHIUnorderedAccessView* Buffer)
+	if (Instances.Num() == 0)
 	{
-		if (Buffer)
-		{
-			Group->TransitionsBefore.Add(FRHITransitionInfo(Buffer, ERHIAccess::SRVMask, ERHIAccess::UAVCompute));
-			Group->TransitionsAfter.Add(FRHITransitionInfo(Buffer, ERHIAccess::UAVCompute, ERHIAccess::SRVMask));
-		}
-	};
+		return;
+	}
+
+	SCOPED_DRAW_EVENT(RHICmdList, NiagaraGPUComputeFreeIDs);
+	SCOPED_GPU_STAT(RHICmdList, NiagaraGPUComputeFreeIDs);
+
+	TArray<FRHITransitionInfo, TMemStackAllocator<>> TransitionsBefore;
+	TransitionsBefore.Emplace(FreeIDListSizesBuffer.UAV, ERHIAccess::UAVCompute, ERHIAccess::UAVCompute);
+	RHICmdList.Transition(TransitionsBefore);
 
-	AddBufferTransitions(Group, DestinationData->GetGPUBufferFloat().UAV);
-	AddBufferTransitions(Group, DestinationData->GetGPUBufferHalf().UAV);
-	AddBufferTransitions(Group, DestinationData->GetGPUBufferInt().UAV);
+	check((uint32)Instances.Num() <= NumAllocatedFreeIDListSizes);
 
+	RHICmdList.BeginUAVOverlap(FreeIDListSizesBuffer.UAV);
+	for (uint32 iInstance = 0; iInstance < (uint32)Instances.Num(); ++iInstance)
+	{
+		FNiagaraComputeExecutionContext* ComputeContext = Instances[iInstance];
+		FNiagaraDataSet* MainDataSet = ComputeContext->MainDataSet;
+		FNiagaraDataBuffer* CurrentData = ComputeContext->MainDataSet->GetCurrentData();
+
+		SCOPED_DRAW_EVENTF(RHICmdList, NiagaraGPUComputeFreeIDsEmitter, TEXT("Update Free ID Buffer - %s"), ComputeContext->GetDebugSimName());
+		NiagaraComputeGPUFreeIDs(RHICmdList, FeatureLevel, MainDataSet->GetGPUNumAllocatedIDs(), CurrentData->GetGPUIDToIndexTable().SRV, MainDataSet->GetGPUFreeIDs(), FreeIDListSizesBuffer, iInstance);
+	}
+	RHICmdList.EndUAVOverlap(FreeIDListSizesBuffer.UAV);
 }
 
-void NiagaraEmitterInstanceBatcher::BuildDispatchGroups(FOverlappableTicks& OverlappableTick, FRHICommandList& RHICmdList, FDispatchGroupList& DispatchGroups, FEmitterInstanceList& InstancesWithPersistentIDs)
+void NiagaraEmitterInstanceBatcher::DumpDebugFrame()
 {
-	SCOPE_CYCLE_COUNTER(STAT_NiagaraGPUDispatchSetup_RT);
+	// Anything doing?
+	bool bHasAnyWork = false;
+	for (int iTickStage = 0; iTickStage < ENiagaraGpuComputeTickStage::Max; ++iTickStage)
+	{
+		bHasAnyWork |= DispatchListPerStage[iTickStage].HasWork();
+	}
+	if ( !bHasAnyWork )
+	{
+		return;
+	}
 
-	DispatchGroups.Reserve(8);
-	FDispatchInstanceList EmittersWithStages;
-	EmittersWithStages.Reserve(16);
-	uint32 MaxStageCountForGroup = 0;
+	// Dump Frame
+	UE_LOG(LogNiagara, Warning, TEXT("====== BatcherFrame(%d)"), GFrameNumberRenderThread);
 
-	//UE_LOG(LogNiagara, Warning, TEXT("NiagaraEmitterInstanceBatcher::BuildDispatchGroups:  %0xP"), this);
-	for (FNiagaraGPUSystemTick* Tick : OverlappableTick)
+	for (int iTickStage = 0; iTickStage < ENiagaraGpuComputeTickStage::Max; ++iTickStage)
 	{
-		const uint32 DispatchCount = Tick->Count;
-		const bool bNeedsReset = Tick->bNeedsReset;
+		if (!DispatchListPerStage[iTickStage].HasWork())
+		{
+			continue;
+		}
+
+		FNiagaraGpuDispatchList& DispatchList = DispatchListPerStage[iTickStage];
+		UE_LOG(LogNiagara, Warning, TEXT("==== TickStage(%d) TotalGroups(%d)"), iTickStage, DispatchList.DispatchGroups.Num());
 
-		int32 InstanceGroupIndex = 0;
-		int32 PreviousInstanceFinalGroupIndex = 0;
-		for (uint32 Index=0; Index < DispatchCount; ++Index)
+		for ( int iDispatchGroup=0; iDispatchGroup < DispatchList.DispatchGroups.Num(); ++iDispatchGroup )
 		{
-			FNiagaraComputeInstanceData& InstanceData = Tick->GetInstanceData()[Index];
-			FNiagaraComputeExecutionContext* Context = InstanceData.Context;
-			if ((Context == nullptr) || !Context->GPUScript_RT->IsShaderMapComplete_RenderThread())
+			const FNiagaraGpuDispatchGroup& DispatchGroup = DispatchListPerStage[iTickStage].DispatchGroups[iDispatchGroup];
+
+			if (DispatchGroup.TicksWithPerInstanceData.Num() > 0)
 			{
-				continue;
+				UE_LOG(LogNiagara, Warning, TEXT("====== TicksWithPerInstanceData(%s)"), DispatchGroup.TicksWithPerInstanceData.Num());
+				for (FNiagaraGPUSystemTick* Tick : DispatchGroup.TicksWithPerInstanceData)
+				{
+					for (const auto& Pair : Tick->DIInstanceData->InterfaceProxiesToOffsets)
+					{
+						FNiagaraDataInterfaceProxy* Proxy = Pair.Key;
+						UE_LOG(LogNiagara, Warning, TEXT("Proxy(%s)"), *Proxy->SourceDIName.ToString());
+					}
+				}
 			}
 
-			// If we can't overlap the previous groups (i.e. particle read dependency) then we must start at the end to ensure all previous stages have completed
-			if (DispatchGroups.IsValidIndex(InstanceGroupIndex) && (InstanceData.bStartNewOverlapGroup && DispatchGroups[InstanceGroupIndex].DispatchInstances.Num() > 0))
+			UE_LOG(LogNiagara, Warning, TEXT("====== DispatchGroup(%d)"), iDispatchGroup);
+			for ( const FNiagaraGpuDispatchInstance& DispatchInstance : DispatchGroup.DispatchInstances )
 			{
-				InstanceGroupIndex = PreviousInstanceFinalGroupIndex;
-			}
+				const FNiagaraSimStageData& SimStageData = DispatchInstance.SimStageData;
+				const FNiagaraComputeInstanceData& InstanceData = DispatchInstance.InstanceData;
+
+				TStringBuilder<512> Builder;
+				Builder.Appendf(TEXT("Proxy(%p) "), DispatchInstance.Tick.SystemGpuComputeProxy);
+				Builder.Appendf(TEXT("ComputeContext(%p) "), InstanceData.Context);
+				Builder.Appendf(TEXT("Emitter(%s) "), InstanceData.Context->GetDebugSimName());
+				if (SimStageData.StageMetaData)
+				{
+					Builder.Appendf(TEXT("Stage(%d | %s) "), SimStageData.StageIndex, *SimStageData.StageMetaData->SimulationStageName.ToString());
+				}
+				else
+				{
+					Builder.Appendf(TEXT("Stage(%d) "), SimStageData.StageIndex);
+				}
 
-			// Add persistent ID transitions
-			const bool bRequiresPersistentIDs = Context->MainDataSet->RequiresPersistentIDs();
+				if (InstanceData.bResetData)
+				{
+					Builder.Append(TEXT("ResetData "));
+				}
 
-			FRHIUnorderedAccessView* FinalIDToIndexUAV = nullptr;
+				if (InstanceData.Context->MainDataSet->RequiresPersistentIDs())
+				{
+					Builder.Append(TEXT("HasPersistentIDs "));
+				}
 
-			int32 StageGroupIndex = InstanceGroupIndex;
-			FDispatchInstance* FinalDispatchInstance = nullptr;
-			for (uint32 iStage=0; iStage < Context->MaxUpdateIterations; ++iStage)
-			{
-				if (!NiagaraEmitterInstanceBatcherLocal::ShouldRunStage(Tick, Context, InstanceData.SimStageData[iStage].AlternateIterationSource, iStage))
+				if ( DispatchInstance.SimStageData.bFirstStage )
 				{
-					continue;
+					Builder.Append(TEXT("FirstStage "));
 				}
 
-				FDispatchGroup* StageGroup = DispatchGroups.IsValidIndex(StageGroupIndex) ? &DispatchGroups[StageGroupIndex] : &DispatchGroups.AddDefaulted_GetRef();
+				if ( DispatchInstance.SimStageData.bLastStage )
+				{
+					Builder.Append(TEXT("LastStage "));
+				}
 
-				FNiagaraDataInterfaceProxyRW* IterationInterface = InstanceData.SimStageData[iStage].AlternateIterationSource;
-				if (!IterationInterface)
+				if (InstanceData.Context->EmitterInstanceReadback.GPUCountOffset != INDEX_NONE)
 				{
-					if (FNiagaraDataBuffer* DestinationData = InstanceData.SimStageData[iStage].Destination)
+					if (InstanceData.Context->EmitterInstanceReadback.GPUCountOffset == SimStageData.SourceCountOffset)
 					{
-						AddDestinationBufferTransitions(StageGroup, DestinationData);
-
-						if (bRequiresPersistentIDs)
-						{
-							// Ensure GPU Free IDs are readable, this only needs to happen once as they will remain readable through all simulation stages
-							if (FinalIDToIndexUAV == nullptr)
-							{
-								StageGroup->TransitionsBefore.Add(FRHITransitionInfo(Context->MainDataSet->GetGPUFreeIDs().UAV, ERHIAccess::UAVCompute, ERHIAccess::SRVCompute));
-							}
-
-							// Insert a UAV barrier on the ID to index table, to prevent a race with the previous dispatch which also wrote to it.
-							FRHIUnorderedAccessView* IDToIndexUAV = DestinationData->GetGPUIDToIndexTable().UAV;
-							ensure((FinalIDToIndexUAV != nullptr) || (IDToIndexUAV == InstanceData.SimStageData[0].Destination->GetGPUIDToIndexTable().UAV));
-							StageGroup->TransitionsBefore.Add(FRHITransitionInfo(IDToIndexUAV, FinalIDToIndexUAV == nullptr ? ERHIAccess::UAVCompute : ERHIAccess::SRVCompute, ERHIAccess::UAVCompute));
-							StageGroup->TransitionsAfter.Add(FRHITransitionInfo(IDToIndexUAV, ERHIAccess::UAVCompute, ERHIAccess::SRVCompute));
-							FinalIDToIndexUAV = IDToIndexUAV;
-						}
+						Builder.Append(TEXT("ReadbackSource "));
 					}
 				}
-
-				// Add dispatch instances
-				FinalDispatchInstance = &StageGroup->DispatchInstances.AddDefaulted_GetRef();
-				FinalDispatchInstance->Tick = Tick;
-				FinalDispatchInstance->InstanceData = &InstanceData;
-				FinalDispatchInstance->StageIndex = iStage;
-
-				++StageGroupIndex;
+				Builder.Appendf(TEXT("Source(%p 0x%08x %d) "), SimStageData.Source, SimStageData.SourceCountOffset, SimStageData.SourceNumInstances);
+				Builder.Appendf(TEXT("Destination(%p 0x%08x %d) "), SimStageData.Destination, SimStageData.DestinationCountOffset, SimStageData.DestinationNumInstances);
+				Builder.Appendf(TEXT("Iteration(%s) "), SimStageData.AlternateIterationSource ? *SimStageData.AlternateIterationSource->SourceDIName.ToString() : TEXT("Particles"));
+				UE_LOG(LogNiagara, Warning, TEXT("%s"), Builder.ToString());
 			}
 
-			// Mark that this is the final stage run for this instance
-			if (FinalDispatchInstance != nullptr)
+			if (DispatchGroup.FreeIDUpdates.Num() > 0)
 			{
-				FinalDispatchInstance->bFinalStage = true;
+				UE_LOG(LogNiagara, Warning, TEXT("====== FreeIDUpdates"));
+				for (FNiagaraComputeExecutionContext* ComputeContext : DispatchGroup.FreeIDUpdates)
+				{
+					UE_LOG(LogNiagara, Warning, TEXT("ComputeContext(%p) Emitter(%s)"), ComputeContext, ComputeContext->GetDebugSimName());
+				}
 			}
+		}
+		if (DispatchList.CountsToRelease.Num() > 0)
+		{
+			UE_LOG(LogNiagara, Warning, TEXT("====== CountsToRelease"));
 
-			PreviousInstanceFinalGroupIndex = StageGroupIndex;
-
-			if (FinalIDToIndexUAV)
-			{
-				InstancesWithPersistentIDs.Add(&InstanceData);
-			}
+			const int NumPerLine = 16;
 
-			// Final tick, if so enqueue a readback and set the data to render
-			if (Tick->bIsFinalTick)
+			TStringBuilder<512> StringBuilder;
+			for ( int32 i=0; i < DispatchList.CountsToRelease.Num(); ++i )
 			{
-				// If the emitter has stages, we'll call SetDataToRender() after we add all the stages, since those can flip the source/destination buffers a number of times.
-				Context->SetDataToRender(Context->MainDataSet->GetCurrentData());
-			}
-		}
-	}
+				const bool bFirst = (i % NumPerLine) == 0;
+				const bool bLast = ((i % NumPerLine) == NumPerLine - 1) || (i == DispatchList.CountsToRelease.Num() -1);
 
-	uint32 NumInstancesWithPersistentIDs = (uint32)InstancesWithPersistentIDs.Num();
-	if (NumInstancesWithPersistentIDs > 0)
-	{
-		// These buffers will be needed by the simulation dispatches which come immediately after, so there will be a stall, but
-		// moving this step to a different place is difficult, so we'll live with it for now.
-		SCOPED_DRAW_EVENT(RHICmdList, NiagaraGPUClearIDTables);
-		SCOPED_GPU_STAT(RHICmdList, NiagaraGPUClearIDTables);
-
-		FNiagaraTransitionList IDToIndexTableTransitions;
-		IDToIndexTableTransitions.SetNum(NumInstancesWithPersistentIDs);
-		for (uint32 i = 0; i < NumInstancesWithPersistentIDs; ++i)
-		{
-			FNiagaraComputeInstanceData* Instance = InstancesWithPersistentIDs[i];
-			FRHIUnorderedAccessView* IDToIndexUAV = Instance->SimStageData[0].Destination->GetGPUIDToIndexTable().UAV;
-			IDToIndexTableTransitions[i] = FRHITransitionInfo(IDToIndexUAV, ERHIAccess::SRVCompute, ERHIAccess::UAVCompute);
-		}
-		RHICmdList.Transition(IDToIndexTableTransitions);
+				if ( !bFirst )
+				{
+					StringBuilder.Append(TEXT(", "));
+				}
+				StringBuilder.Appendf(TEXT("0x%08x"), DispatchList.CountsToRelease[i]);
 
-		for (uint32 i = 0; i < NumInstancesWithPersistentIDs; ++i)
-		{
-			FNiagaraComputeInstanceData* Instance = InstancesWithPersistentIDs[i];
-			SCOPED_DRAW_EVENTF(RHICmdList, NiagaraGPUComputeClearIDToIndexBuffer, TEXT("Clear ID To Index Table - %s"), Instance->Context->GetDebugSimName());
-			NiagaraFillGPUIntBuffer(RHICmdList, FeatureLevel, Instance->SimStageData[0].Destination->GetGPUIDToIndexTable(), INDEX_NONE);
+				if ( bLast )
+				{
+					UE_LOG(LogNiagara, Warning, TEXT("%s"), StringBuilder.ToString());
+					StringBuilder.Reset();
+				}
+			}
 		}
 	}
 }
 
-void NiagaraEmitterInstanceBatcher::DispatchAllOnCompute(FDispatchInstanceList& DispatchInstances, FRHICommandList& RHICmdList, FRHIUniformBuffer* ViewUniformBuffer)
+void NiagaraEmitterInstanceBatcher::UpdateInstanceCountManager(FRHICommandListImmediate& RHICmdList)
 {
-	// Run all the PreStage functions in bulk.
-	for (FDispatchInstance& DispatchInstance : DispatchInstances)
-	{
-		PreStageInterface(*DispatchInstance.Tick, DispatchInstance.InstanceData, RHICmdList, DispatchInstance.StageIndex);
-	}
-
-	// Run all the simulation compute shaders.
-	for (FDispatchInstance& DispatchInstance : DispatchInstances)
+	// Resize dispatch buffer count
+	//-OPT: No need to iterate over all the ticks, we can store this as ticks are queued
+	int32 TotalDispatchCount = 0;
 	{
-#if STATS
-		if (GPUProfiler.IsProfilingEnabled() && DispatchInstance.Tick->bIsFinalTick)
+		for (int iTickStage = 0; iTickStage < ENiagaraGpuComputeTickStage::Max; ++iTickStage)
 		{
-			FString StageName = "SpawnUpdate";
-			if (DispatchInstance.StageIndex > 0)
+			for (FNiagaraSystemGpuComputeProxy* ComputeProxy : ProxiesPerStage[iTickStage])
 			{
-				for (FSimulationStageMetaData& MetaData : DispatchInstance.InstanceData->Context->SimStageInfo)
+				for (FNiagaraGPUSystemTick& Tick : ComputeProxy->PendingTicks)
 				{
-					if (DispatchInstance.StageIndex >= MetaData.MinStage && DispatchInstance.StageIndex < MetaData.MaxStage)
+					TotalDispatchCount += (int32)Tick.TotalDispatches;
+
+					for ( FNiagaraComputeInstanceData& InstanceData : Tick.GetInstances() )
 					{
-						StageName = MetaData.SimulationStageName.ToString();
+						if (InstanceData.bResetData)
+						{
+							GPUInstanceCounterManager.FreeEntry(InstanceData.Context->EmitterInstanceReadback.GPUCountOffset);
+						}
 					}
 				}
 			}
-			TStatId StatId = FDynamicStats::CreateStatId<FStatGroup_STATGROUP_NiagaraDetailed>("GPU_Stage_" + StageName);
-			int32 TimerHandle = GPUProfiler.StartTimer((uint64)DispatchInstance.InstanceData->Context, StatId, RHICmdList);
-			DispatchStage(DispatchInstance, DispatchInstance.StageIndex, RHICmdList, ViewUniformBuffer);
-			GPUProfiler.EndTimer(TimerHandle, RHICmdList);
 		}
-		else
-		{
-			DispatchStage(DispatchInstance, DispatchInstance.StageIndex, RHICmdList, ViewUniformBuffer);
-		}
-#else
-		DispatchStage(DispatchInstance, DispatchInstance.StageIndex, RHICmdList, ViewUniformBuffer);
-#endif
+		GPUInstanceCounterManager.ResizeBuffers(RHICmdList, FeatureLevel, TotalDispatchCount);
 	}
 
-	// Run all the PostStage functions in bulk.
-	for (FDispatchInstance& DispatchInstance : DispatchInstances)
+	// Consume any pending readbacks that are ready
 	{
-		PostStageInterface(*DispatchInstance.Tick, DispatchInstance.InstanceData, RHICmdList, DispatchInstance.StageIndex);
-		if (DispatchInstance.bFinalStage)
+		SCOPE_CYCLE_COUNTER(STAT_NiagaraGPUReadback_RT);
+		if ( const uint32* Counts = GPUInstanceCounterManager.GetGPUReadback() )
 		{
-			PostSimulateInterface(*DispatchInstance.Tick, DispatchInstance.InstanceData, RHICmdList, DispatchInstance.InstanceData->Context->GPUScript_RT);
-		}
-	}
-}
+			if (NiagaraEmitterInstanceBatcherLocal::GDebugLogging)
+			{
+				UE_LOG(LogNiagara, Warning, TEXT("====== BatcherFrame(%d) Readback Complete"), GFrameNumberRenderThread);
+			}
 
-void NiagaraEmitterInstanceBatcher::UpdateFreeIDsListSizesBuffer(FRHICommandList& RHICmdList, uint32 NumInstances)
-{
-	ERHIAccess BeforeState = ERHIAccess::UAVCompute;
-	if (NumInstances > NumAllocatedFreeIDListSizes)
-	{
-		constexpr uint32 ALLOC_CHUNK_SIZE = 128;
-		NumAllocatedFreeIDListSizes = Align(NumInstances, ALLOC_CHUNK_SIZE);
-		if (FreeIDListSizesBuffer.Buffer)
-		{
-			FreeIDListSizesBuffer.Release();
-		}
-		FreeIDListSizesBuffer.Initialize(TEXT("NiagaraFreeIDListSizes"), sizeof(uint32), NumAllocatedFreeIDListSizes, EPixelFormat::PF_R32_SINT, BUF_Static);
+			for (int iTickStage=0; iTickStage < ENiagaraGpuComputeTickStage::Max; ++iTickStage)
+			{
+				for (FNiagaraSystemGpuComputeProxy* ComputeProxy : ProxiesPerStage[iTickStage])
+				{
+					for ( FNiagaraComputeExecutionContext* ComputeContext : ComputeProxy->ComputeContexts )
+					{
+						if (ComputeContext->EmitterInstanceReadback.GPUCountOffset == INDEX_NONE )
+						{
+							continue;
+						}
 
-		BeforeState = ERHIAccess::Unknown;
+						const uint32 DeadInstanceCount = ComputeContext->EmitterInstanceReadback.CPUCount - Counts[ComputeContext->EmitterInstanceReadback.GPUCountOffset];
+						if ( DeadInstanceCount <= ComputeContext->CurrentNumInstances_RT )
+						{
+							ComputeContext->CurrentNumInstances_RT -= DeadInstanceCount;
+						}
+						if (NiagaraEmitterInstanceBatcherLocal::GDebugLogging)
+						{
+							UE_LOG(LogNiagara, Warning, TEXT("ComputeContext(%p) Emitter(%s) DeadInstances(%d) CountReleased(0x%08x)"), ComputeContext, ComputeContext->GetDebugSimName(), DeadInstanceCount, ComputeContext->EmitterInstanceReadback.GPUCountOffset);
+						}
 
-		bFreeIDListSizesBufferCleared = false;
-	}
+						// Release the counter for the readback
+						ComputeContext->ReleaseReadbackCounter(GPUInstanceCounterManager);
+					}
+				}
+			}
 
-	if ( !bFreeIDListSizesBufferCleared )
-	{
-		SCOPED_DRAW_EVENT(RHICmdList, NiagaraGPUComputeClearFreeIDListSizes);
-		RHICmdList.Transition(FRHITransitionInfo(FreeIDListSizesBuffer.UAV, BeforeState, ERHIAccess::UAVCompute));
-		NiagaraFillGPUIntBuffer(RHICmdList, FeatureLevel, FreeIDListSizesBuffer, 0);
-		bFreeIDListSizesBufferCleared = true;
+			// Release the readback buffer
+			GPUInstanceCounterManager.ReleaseGPUReadback();
+		}
 	}
 }
 
-void NiagaraEmitterInstanceBatcher::UpdateFreeIDBuffers(FRHICommandList& RHICmdList, FEmitterInstanceList& Instances)
+void NiagaraEmitterInstanceBatcher::PrepareTicksForProxy(FRHICommandListImmediate& RHICmdList, FNiagaraSystemGpuComputeProxy* ComputeProxy, FNiagaraGpuDispatchList& GpuDispatchList)
 {
-	if (Instances.Num() == 0)
+	for (FNiagaraComputeExecutionContext* ComputeContext : ComputeProxy->ComputeContexts)
+	{
+		ComputeContext->CurrentMaxInstances_RT = 0;
+		ComputeContext->CurrentMaxAllocateInstances_RT = 0;
+		ComputeContext->BufferSwapsThisFrame_RT = 0;
+	}
+
+	if (ComputeProxy->PendingTicks.Num() == 0)
 	{
 		return;
 	}
 
-	SCOPED_DRAW_EVENT(RHICmdList, NiagaraGPUComputeFreeIDs);
-	SCOPED_GPU_STAT(RHICmdList, NiagaraGPUComputeFreeIDs);
+	const bool bEnqueueCountReadback = !GPUInstanceCounterManager.HasPendingGPUReadback();
 
-	FNiagaraTransitionList TransitionsBefore;
-	TransitionsBefore.Reserve(Instances.Num() * 2 + 1);
+	// Set final tick flag
+	ComputeProxy->PendingTicks.Last().bIsFinalTick = true;
 
-	// Insert a UAV barrier on the buffer, to make sure the previous clear finished writing to it.
-	TransitionsBefore.Add(FRHITransitionInfo(FreeIDListSizesBuffer.UAV, ERHIAccess::UAVCompute, ERHIAccess::UAVCompute));
-	for(FNiagaraComputeInstanceData* Instance : Instances)
+	// Process ticks
+	int32 iTickStartDispatchGroup = 0;
+
+	for (FNiagaraGPUSystemTick& Tick : ComputeProxy->PendingTicks)
 	{
-		FNiagaraDataSet* MainDataSet = Instance->Context->MainDataSet;
-		FNiagaraDataBuffer* FinalData = MainDataSet->GetCurrentData();
+		int32 iInstanceStartDispatchGroup = iTickStartDispatchGroup;
+		int32 iInstanceCurrDispatchGroup = iTickStartDispatchGroup;
+		bool bHasFreeIDUpdates = false;
 
-		// The ID to index tables should already be in SRVCompute. Make the free ID tables writable.
-		TransitionsBefore.Add(FRHITransitionInfo(MainDataSet->GetGPUFreeIDs().UAV, ERHIAccess::SRVCompute, ERHIAccess::UAVCompute));
-	}
+		// Track that we need to consume per instance data before executing the ticks
+		if (Tick.DIInstanceData)
+		{
+			GpuDispatchList.PreAllocateGroups(iTickStartDispatchGroup + 1);
+			GpuDispatchList.DispatchGroups[iTickStartDispatchGroup].TicksWithPerInstanceData.Add(&Tick);
+		}
 
-	RHICmdList.Transition(TransitionsBefore);
+		// Iterate over all instances preparing our number of instances
+		for (FNiagaraComputeInstanceData& InstanceData : Tick.GetInstances())
+		{
+			FNiagaraComputeExecutionContext* ComputeContext = InstanceData.Context;
 
-	check((uint32)Instances.Num() <= NumAllocatedFreeIDListSizes);
+			// Instance requires a reset?
+			if (InstanceData.bResetData)
+			{
+				ComputeContext->CurrentNumInstances_RT = 0;
+				if (ComputeContext->CountOffset_RT != INDEX_NONE)
+				{
+					GpuDispatchList.CountsToRelease.Add(ComputeContext->CountOffset_RT);
+					ComputeContext->CountOffset_RT = INDEX_NONE;
+				}
+			}
 
-	RHICmdList.BeginUAVOverlap(FreeIDListSizesBuffer.UAV);
-	for (uint32 InstanceIdx = 0; InstanceIdx < (uint32)Instances.Num(); ++InstanceIdx)
-	{
-		FNiagaraComputeInstanceData* Instance = Instances[InstanceIdx];
-		FNiagaraDataSet* MainDataSet = Instance->Context->MainDataSet;
-		FNiagaraDataBuffer* FinalData = MainDataSet->GetCurrentData();
+			// If shader is not ready don't do anything
+			if (!ComputeContext->GPUScript_RT->IsShaderMapComplete_RenderThread())
+			{
+				continue;
+			}
 
-		SCOPED_DRAW_EVENTF(RHICmdList, NiagaraGPUComputeFreeIDsEmitter, TEXT("Update Free ID Buffer - %s"), Instance->Context->GetDebugSimName());
-		NiagaraComputeGPUFreeIDs(RHICmdList, FeatureLevel, MainDataSet->GetGPUNumAllocatedIDs(), FinalData->GetGPUIDToIndexTable().SRV, MainDataSet->GetGPUFreeIDs(), FreeIDListSizesBuffer, InstanceIdx);
-	}
-	RHICmdList.EndUAVOverlap(FreeIDListSizesBuffer.UAV);
+#if WITH_EDITOR
+			//-TODO: Temporary change to avoid crashing on incorrect FeatureLevel
+			if (ComputeContext->GPUScript_RT->GetFeatureLevel() != FeatureLevel)
+			{
+				continue;
+			}
+#endif
 
-	bFreeIDListSizesBufferCleared = false;
-}
+			// Determine this instances start dispatch group, in the case of emitter dependencies (i.e. particle reads) we need to continue rather than starting again
+			iInstanceStartDispatchGroup = InstanceData.bStartNewOverlapGroup ? iInstanceCurrDispatchGroup : iInstanceStartDispatchGroup;
+			iInstanceCurrDispatchGroup = iInstanceStartDispatchGroup;
 
-void NiagaraEmitterInstanceBatcher::UpdateFreeIDBuffers(FRHICommandList& RHICmdList, TArrayView<FNiagaraComputeInstanceData*> Instances)
-{
-}
+			GpuDispatchList.PreAllocateGroups(iInstanceCurrDispatchGroup + ComputeContext->MaxUpdateIterations);
 
-void NiagaraEmitterInstanceBatcher::UpdateInstanceCountManager(FRHICommandListImmediate& RHICmdList)
-{
-	// Resize dispatch buffer count
-	int32 TotalDispatchCount = 0;
-	{
-		for (FNiagaraGPUSystemTick& Tick : Ticks_RT)
-		{
-			TotalDispatchCount += (int32)Tick.TotalDispatches;
+			// Calculate instance counts
+			const uint32 MaxBufferInstances = ComputeContext->MainDataSet->GetMaxInstanceCount();
+			const uint32 PrevNumInstances = ComputeContext->CurrentNumInstances_RT;
+
+			ComputeContext->CurrentNumInstances_RT = PrevNumInstances + InstanceData.SpawnInfo.SpawnRateInstances + InstanceData.SpawnInfo.EventSpawnTotal;
+			ComputeContext->CurrentNumInstances_RT = FMath::Min(ComputeContext->CurrentNumInstances_RT, MaxBufferInstances);
+
+			// Calculate new maximum count
+			ComputeContext->CurrentMaxInstances_RT = FMath::Max(ComputeContext->CurrentMaxInstances_RT, ComputeContext->CurrentNumInstances_RT);
+
+			if (!GNiagaraBatcherFreeBufferEarly || (ComputeContext->CurrentMaxInstances_RT > 0))
+			{
+				ComputeContext->CurrentMaxAllocateInstances_RT = FMath::Max3(ComputeContext->CurrentMaxAllocateInstances_RT, ComputeContext->CurrentMaxInstances_RT, InstanceData.SpawnInfo.MaxParticleCount);
+			}
+			else
+			{
+				ComputeContext->CurrentMaxAllocateInstances_RT = FMath::Max(ComputeContext->CurrentMaxAllocateInstances_RT, ComputeContext->CurrentMaxInstances_RT);
+			}
 
-			// Cancel any pending readback if the emitter is resetting.
-			if (Tick.bNeedsReset)
+			bHasFreeIDUpdates |= ComputeContext->MainDataSet->RequiresPersistentIDs();
+
+			// Add particle update / spawn stage
 			{
-				FNiagaraComputeInstanceData* Instances = Tick.GetInstanceData();
-				for (uint32 InstanceIndex = 0; InstanceIndex < Tick.Count; ++InstanceIndex)
+				FNiagaraGpuDispatchGroup& DispatchGroup = GpuDispatchList.DispatchGroups[iInstanceCurrDispatchGroup++];
+				FNiagaraGpuDispatchInstance& DispatchInstance = DispatchGroup.DispatchInstances.Emplace_GetRef(Tick, InstanceData);
+				FNiagaraSimStageData& SimStageData = DispatchInstance.SimStageData;
+
+				SimStageData.bFirstStage = true;
+				SimStageData.StageIndex = 0;
+				SimStageData.Source = ComputeContext->GetCurrDataBuffer();
+				SimStageData.SourceCountOffset = ComputeContext->CountOffset_RT;
+				SimStageData.SourceNumInstances = PrevNumInstances;
+
+				// Validation as a non zero instance count must have a valid count offset
+				check(SimStageData.SourceCountOffset != INDEX_NONE || SimStageData.SourceNumInstances == 0);
+
+				SimStageData.Destination = ComputeContext->GetNextDataBuffer();
+				SimStageData.DestinationCountOffset = GPUInstanceCounterManager.AcquireEntry();
+				SimStageData.DestinationNumInstances = ComputeContext->CurrentNumInstances_RT;
+
+				ComputeContext->AdvanceDataBuffer();
+				ComputeContext->CountOffset_RT = SimStageData.DestinationCountOffset;
+
+				// If we are the last tick then we may want to enqueue for a readback
+				// Note: Do not pull count from SimStageData as a reset tick will be INDEX_NONE
+				if (SimStageData.SourceCountOffset != INDEX_NONE)
 				{
-					FNiagaraComputeExecutionContext* Context = Instances[InstanceIndex].Context;
-					if (Context)
+					if (bEnqueueCountReadback && Tick.bIsFinalTick && (ComputeContext->EmitterInstanceReadback.GPUCountOffset == INDEX_NONE))
 					{
-						GPUInstanceCounterManager.FreeEntry(Context->EmitterInstanceReadback.GPUCountOffset);
+						bRequiresReadback = true;
+						ComputeContext->EmitterInstanceReadback.CPUCount = SimStageData.SourceNumInstances;
+						ComputeContext->EmitterInstanceReadback.GPUCountOffset = SimStageData.SourceCountOffset;
+					}
+					else if (SimStageData.SourceCountOffset != ComputeContext->EmitterInstanceReadback.GPUCountOffset)
+					{
+						GpuDispatchList.CountsToRelease.Add(SimStageData.SourceCountOffset);
 					}
 				}
+				if (Tick.NumInstancesWithSimStages > 0)
+				{
+					SimStageData.AlternateIterationSource = InstanceData.FindIterationInterface(0);
+				}
 			}
-		}
-		GPUInstanceCounterManager.ResizeBuffers(RHICmdList, FeatureLevel, TotalDispatchCount);
-	}
+			ComputeContext->bHasTickedThisFrame_RT = true;
 
-	// Don't consume the readback data if we have no ticks as they will be gone forever
-	// This causes an issue with sequencer ticks @ 30hz but the engine ticking @ 60hz
-	// The readback really needs to be changed to not rely on the tick being enqueued, we should be pushing this data back
-	if (TotalDispatchCount > 0)
-	{
-		SCOPE_CYCLE_COUNTER(STAT_NiagaraGPUReadback_RT);
-		const uint32* Counts = GPUInstanceCounterManager.GetGPUReadback();
-		if (Counts)
-		{
-			for (FNiagaraGPUSystemTick& Tick : Ticks_RT)
+			//-OPT: Seems like a weird test
+			if (Tick.NumInstancesWithSimStages > 0)
 			{
-				FNiagaraComputeInstanceData* Instances = Tick.GetInstanceData();
-				for (uint32 InstanceIndex = 0; InstanceIndex < Tick.Count; ++InstanceIndex)
+				//-OPT: Is this ever true?
+				for (uint32 SimStageIndex = 1; SimStageIndex < ComputeContext->MaxUpdateIterations; ++SimStageIndex)
 				{
-					FNiagaraComputeExecutionContext* Context = Instances[InstanceIndex].Context;
-					if (Context && Context->EmitterInstanceReadback.GPUCountOffset != INDEX_NONE)
+					// Determine if the iteration is outputting to a custom data size
+					FNiagaraDataInterfaceProxyRW* IterationInterface = InstanceData.FindIterationInterface(SimStageIndex);
+					if (!NiagaraEmitterInstanceBatcherLocal::ShouldRunStage(ComputeContext, IterationInterface, SimStageIndex, InstanceData.bResetData))
+					{
+						continue;
+					}
+
+					// Build SimStage data
+					FNiagaraGpuDispatchGroup& DispatchGroup = GpuDispatchList.DispatchGroups[iInstanceCurrDispatchGroup++];
+					FNiagaraGpuDispatchInstance& DispatchInstance = DispatchGroup.DispatchInstances.Emplace_GetRef(Tick, InstanceData);
+					FNiagaraSimStageData& SimStageData = DispatchInstance.SimStageData;
+					SimStageData.StageIndex = SimStageIndex;
+					SimStageData.StageMetaData = ComputeContext->GetSimStageMetaData(SimStageIndex);
+					SimStageData.AlternateIterationSource = IterationInterface;
+
+					//-TODO: REMOVE OLD SIM STAGES
+					if (SimStageData.StageMetaData == nullptr)
 					{
-						check(Context->MainDataSet);
-						FNiagaraDataBuffer* CurrentData = Context->MainDataSet->GetCurrentData();
-						if (CurrentData)
+						SimStageData.Source = ComputeContext->GetCurrDataBuffer();
+						SimStageData.SourceCountOffset = ComputeContext->CountOffset_RT;
+						SimStageData.SourceNumInstances = ComputeContext->CurrentNumInstances_RT;
+						SimStageData.Destination = ComputeContext->GetNextDataBuffer();
+						SimStageData.DestinationCountOffset = GPUInstanceCounterManager.AcquireEntry();
+						SimStageData.DestinationNumInstances = ComputeContext->CurrentNumInstances_RT;
+
+						if (ensure(SimStageData.SourceCountOffset != INDEX_NONE))
 						{
-							const uint32 DeadInstanceCount = Context->EmitterInstanceReadback.CPUCount - Counts[Context->EmitterInstanceReadback.GPUCountOffset];
-
-							// This will communicate the particle counts to the game thread. If DeadInstanceCount equals CurrentData->GetNumInstances() the game thread will know that the emitter has completed.
-							if (DeadInstanceCount <= CurrentData->GetNumInstances())
-							{
-								CurrentData->SetNumInstances(CurrentData->GetNumInstances() - DeadInstanceCount);
-								//UE_LOG(LogNiagara, Log, TEXT("GPU Readback Offset: %d %p = %d"), Context->EmitterInstanceReadback.GPUCountOffset, CurrentData, CurrentData->GetNumInstances());
-							}
+							GpuDispatchList.CountsToRelease.Add(SimStageData.SourceCountOffset);
 						}
 
-						// Now release the readback since another one will be enqueued in the tick.
-						// Also prevents processing the same data again.
-						GPUInstanceCounterManager.FreeEntry(Context->EmitterInstanceReadback.GPUCountOffset);
+						ComputeContext->AdvanceDataBuffer();
+						ComputeContext->CountOffset_RT = SimStageData.DestinationCountOffset;
+					}
+					// This stage does not modify particle data, i.e. read only or not related to particles at all
+					else if (!SimStageData.StageMetaData->bWritesParticles)
+					{
+						SimStageData.Source = ComputeContext->GetCurrDataBuffer();
+						SimStageData.SourceCountOffset = ComputeContext->CountOffset_RT;
+						SimStageData.SourceNumInstances = ComputeContext->CurrentNumInstances_RT;
+						SimStageData.Destination = nullptr;
+						SimStageData.DestinationCountOffset = ComputeContext->CountOffset_RT;
+						SimStageData.DestinationNumInstances = ComputeContext->CurrentNumInstances_RT;
+					}
+					// This stage writes particles but will not kill any, we can use the buffer as both source and destination
+					else if (SimStageData.StageMetaData->bPartialParticleUpdate)
+					{
+						SimStageData.Source = nullptr;
+						SimStageData.SourceCountOffset = ComputeContext->CountOffset_RT;
+						SimStageData.SourceNumInstances = ComputeContext->CurrentNumInstances_RT;
+						SimStageData.Destination = ComputeContext->GetCurrDataBuffer();
+						SimStageData.DestinationCountOffset = ComputeContext->CountOffset_RT;
+						SimStageData.DestinationNumInstances = ComputeContext->CurrentNumInstances_RT;
+					}
+					// This stage may kill particles, we need to allocate a new destination buffer
+					else
+					{
+						SimStageData.Source = ComputeContext->GetCurrDataBuffer();
+						SimStageData.SourceCountOffset = ComputeContext->CountOffset_RT;
+						SimStageData.SourceNumInstances = ComputeContext->CurrentNumInstances_RT;
+						SimStageData.Destination = ComputeContext->GetNextDataBuffer();
+						SimStageData.DestinationCountOffset = GPUInstanceCounterManager.AcquireEntry();
+						SimStageData.DestinationNumInstances = ComputeContext->CurrentNumInstances_RT;
+
+						if (ensure(SimStageData.SourceCountOffset != INDEX_NONE))
+						{
+							GpuDispatchList.CountsToRelease.Add(SimStageData.SourceCountOffset);
+						}
+
+						ComputeContext->AdvanceDataBuffer();
+						ComputeContext->CountOffset_RT = SimStageData.DestinationCountOffset;
 					}
 				}
 			}
-			// Readback is only valid for one frame, so that any newly allocated instance count
-			// are guarantied to be in the next valid readback data.
-			GPUInstanceCounterManager.ReleaseGPUReadback();
+
+			// The final instance in the last group we modified is the final stage for this tick
+			FNiagaraGpuDispatchGroup& FinalDispatchGroup = GpuDispatchList.DispatchGroups[iInstanceCurrDispatchGroup - 1];
+			FinalDispatchGroup.DispatchInstances.Last().SimStageData.bLastStage = true;
+
+			// Keep track of where the next set of dispatch should occur
+			iTickStartDispatchGroup = FMath::Max(iTickStartDispatchGroup, iInstanceCurrDispatchGroup);
+		}
+
+		// Accumulate Free ID updates
+		// Note: These must be done at the end of the tick due to the way spawned instances read from the free list
+		if (bHasFreeIDUpdates)
+		{
+			FNiagaraGpuDispatchGroup& FinalDispatchGroup = GpuDispatchList.DispatchGroups[iInstanceCurrDispatchGroup - 1];
+			for (FNiagaraComputeInstanceData& InstanceData : Tick.GetInstances())
+			{
+				FNiagaraComputeExecutionContext* ComputeContext = InstanceData.Context;
+				if (!ComputeContext->GPUScript_RT->IsShaderMapComplete_RenderThread())
+				{
+					continue;
+				}
+
+				if ( ComputeContext->MainDataSet->RequiresPersistentIDs() )
+				{
+					FinalDispatchGroup.FreeIDUpdates.Add(ComputeContext);
+				}
+			}
+		}
+
+		// Build constant buffers for tick
+		BuildConstantBuffers(Tick);
+
+		//// Consume DataInterface instance data
+		////-TODO: This is incorrect, it should be consumed as we execute the tick
+		//if (Tick.DIInstanceData)
+		//{
+		//	uint8* BasePointer = (uint8*)Tick.DIInstanceData->PerInstanceDataForRT;
+
+		//	for (auto& Pair : Tick.DIInstanceData->InterfaceProxiesToOffsets)
+		//	{
+		//		FNiagaraDataInterfaceProxy* Proxy = Pair.Key;
+		//		uint8* InstanceDataPtr = BasePointer + Pair.Value;
+		//		Proxy->ConsumePerInstanceDataFromGameThread(InstanceDataPtr, Tick.SystemInstanceID);
+		//	}
+		//}
+	}
+
+	// Now that all ticks have been processed we can adjust our output buffers to the correct size
+	// We will also set the translucent data to render, i.e. this frames data.
+	for (FNiagaraComputeExecutionContext* ComputeContext : ComputeProxy->ComputeContexts)
+	{
+		if (!ComputeContext->bHasTickedThisFrame_RT)
+		{
+			continue;
+		}
+
+		//-OPT: We should allocate all GPU free IDs together since they require a transition
+		if (ComputeContext->MainDataSet->RequiresPersistentIDs())
+		{
+			ComputeContext->MainDataSet->AllocateGPUFreeIDs(ComputeContext->CurrentMaxAllocateInstances_RT + 1, RHICmdList, FeatureLevel, ComputeContext->GetDebugSimName());
+		}
+
+		// Allocate space for all buffers we use
+		//-OPT: We can batch the allocation of persistent IDs together so the compute shaders overlap
+		const uint32 NumBuffersToResize = FMath::Min<uint32>(ComputeContext->BufferSwapsThisFrame_RT, UE_ARRAY_COUNT(ComputeContext->DataBuffers_RT));
+		for (uint32 i = 0; i < NumBuffersToResize; ++i)
+		{
+			const uint32 BufferIndex = (ComputeContext->DataBufferIndex_RT + i) % UE_ARRAY_COUNT(ComputeContext->DataBuffers_RT);
+			ComputeContext->DataBuffers_RT[BufferIndex]->AllocateGPU(RHICmdList, ComputeContext->CurrentMaxAllocateInstances_RT + 1, FeatureLevel, ComputeContext->GetDebugSimName());
+		}
+
+		// Set the translucent draw data and also set the GPU count offset for the draw.
+		// This is safe as simulations have all completed before the translucent pass and allow us to issue
+		// draw commands with this frames data rather than the previous frames.  Note an opaque draw where
+		// the simulation uses the depth buffer can not do this as the simulation can not run until depth
+		// is available for use.
+		extern int32 GNiagaraGpuLowLatencyTranslucencyEnabled;
+		if (GNiagaraGpuLowLatencyTranslucencyEnabled)
+		{
+			ComputeContext->GetCurrDataBuffer()->GPUInstanceCountBufferOffset = ComputeContext->CountOffset_RT;
+			ComputeContext->SetTranslucentDataToRender(ComputeContext->GetCurrDataBuffer());
 		}
 	}
 }
 
-void NiagaraEmitterInstanceBatcher::BuildTickStagePasses(FRHICommandListImmediate& RHICmdList, ETickStage GenerateTickStage)
+void NiagaraEmitterInstanceBatcher::PrepareAllTicks(FRHICommandListImmediate& RHICmdList)
 {
-	// Early out for bulk tick generation in PreInitViews
-	const bool bGeneratePerTickStage = GNiagaraGpuLowLatencyTranslucencyEnabled == 0;
-	if (!bGeneratePerTickStage && (GenerateTickStage != ETickStage::PreInitViews))
+	for (int iTickStage=0; iTickStage < ENiagaraGpuComputeTickStage::Max; ++iTickStage)
 	{
-		return;
+		for (FNiagaraSystemGpuComputeProxy* ComputeProxy : ProxiesPerStage[iTickStage])
+		{
+			PrepareTicksForProxy(RHICmdList, ComputeProxy, DispatchListPerStage[iTickStage]);
+		}
 	}
+}
 
-	for (int32 iTickStage = 0; iTickStage < (int)ETickStage::Max; ++iTickStage)
+void NiagaraEmitterInstanceBatcher::ExecuteTicks(FRHICommandList& RHICmdList, FRHIUniformBuffer* ViewUniformBuffer, ENiagaraGpuComputeTickStage::Type TickStage)
+{
+	// Anything to execute for this stage
+	FNiagaraGpuDispatchList& DispatchList = DispatchListPerStage[TickStage];
+	if ( !DispatchList.HasWork() )
 	{
-		ContextsPerStage[iTickStage].Reset(Ticks_RT.Num());
-		TicksPerStage[iTickStage].Reset(Ticks_RT.Num());
-		CountsToRelease[iTickStage].Reset(Ticks_RT.Num());		//-TODO: Improve what we need here, should be able to make a better guess below
+		return;
 	}
 
-	const bool bEnqueueReadback = !GPUInstanceCounterManager.HasPendingGPUReadback();
-	const bool bAlwaysAllocateBuffers = GNiagaraBatcherFreeBufferEarly == 0;
 
-	for (FNiagaraGPUSystemTick& Tick : Ticks_RT)
+	SCOPED_DRAW_EVENTF(RHICmdList, NiagaraEmitterInstanceBatcher_ExecuteTicks, TEXT("NiagaraEmitterInstanceBatcher_ExecuteTicks - TickStage(%d)"), TickStage);
+	SCOPE_CYCLE_COUNTER(STAT_NiagaraGPUSimTick_RT);
+	SCOPED_GPU_STAT(RHICmdList, NiagaraGPUSimulation);
+
+	FMemMark Mark(FMemStack::Get());
+	TArray<FRHITransitionInfo, TMemStackAllocator<>> TransitionsBefore;
+	TArray<FRHITransitionInfo, TMemStackAllocator<>> TransitionsAfter;
+	TArray<FNiagaraDataBuffer*, TMemStackAllocator<>> IDToIndexInit;
+
+	for ( const FNiagaraGpuDispatchGroup& DispatchGroup : DispatchList.DispatchGroups )
 	{
-		FNiagaraComputeSharedContext* SharedContext = Tick.SharedContext;
+		SCOPED_CONDITIONAL_DRAW_EVENT(RHICmdList, NiagaraDispatchGroup, NiagaraEmitterInstanceBatcherLocal::GAddDispatchGroupDrawEvent);
 
-		const ETickStage TickStage = NiagaraEmitterInstanceBatcherLocal::CalculateTickStage(Tick);
-		if (bGeneratePerTickStage && (GenerateTickStage != TickStage))
-		{
-			continue;
-		}
+		const bool bIsFirstGroup = &DispatchGroup == &DispatchList.DispatchGroups[0];
+		const bool bIsLastGroup = &DispatchGroup == &DispatchList.DispatchGroups.Last();
 
-		BuildConstantBuffers(Tick);
+		// Generate transitions and discover free / ID table updates
+		TransitionsBefore.Reserve(DispatchGroup.DispatchInstances.Num() * 3);
+		TransitionsAfter.Reserve(DispatchGroup.DispatchInstances.Num() * 3);
+		for (const FNiagaraGpuDispatchInstance& DispatchInstance : DispatchGroup.DispatchInstances)
+		{
+			if (FNiagaraDataBuffer* DestinationBuffer = DispatchInstance.SimStageData.Destination)
+			{
+				NiagaraEmitterInstanceBatcherLocal::AddDataBufferTransitions(TransitionsBefore, TransitionsAfter, DestinationBuffer);
+			}
 
-		Tick.bIsFinalTick = false;
+			FNiagaraComputeExecutionContext* ComputeContext = DispatchInstance.InstanceData.Context;
+			const bool bRequiresPersistentIDs = ComputeContext->MainDataSet->RequiresPersistentIDs();
+			if (bRequiresPersistentIDs)
+			{
+				if ( FNiagaraDataBuffer* IDtoIndexBuffer = DispatchInstance.SimStageData.Destination )
+				{
+					IDToIndexInit.Add(IDtoIndexBuffer);
+					TransitionsBefore.Emplace(IDtoIndexBuffer->GetGPUIDToIndexTable().UAV, ERHIAccess::SRVCompute, ERHIAccess::UAVCompute);
+					TransitionsAfter.Emplace(IDtoIndexBuffer->GetGPUIDToIndexTable().UAV, ERHIAccess::UAVCompute, ERHIAccess::SRVCompute);
+				}
+			}
+		}
 
-		const bool bResetCounts = SharedContext->ScratchIndex == INDEX_NONE;
-		if (bResetCounts)
+		TransitionsBefore.Emplace(GPUInstanceCounterManager.GetInstanceCountBuffer().UAV, bIsFirstGroup ? FNiagaraGPUInstanceCountManager::kCountBufferDefaultState : ERHIAccess::UAVCompute, ERHIAccess::UAVCompute);
+		if (bIsLastGroup)
 		{
-			check(!ContextsPerStage[(int)ETickStage::PreInitViews].Contains(SharedContext));
-			check(!ContextsPerStage[(int)ETickStage::PostInitViews].Contains(SharedContext));
-			check(!ContextsPerStage[(int)ETickStage::PostOpaqueRender].Contains(SharedContext));
-			ContextsPerStage[(int)TickStage].Add(SharedContext);
-			SharedContext->ScratchTickStage = (int)TickStage;
+			TransitionsAfter.Emplace(GPUInstanceCounterManager.GetInstanceCountBuffer().UAV, ERHIAccess::UAVCompute, FNiagaraGPUInstanceCountManager::kCountBufferDefaultState);
 		}
 
-		// Here scratch index represent the index of the last tick
-		SharedContext->ScratchIndex = TicksPerStage[(int)TickStage].Add(&Tick);
-		check(SharedContext->ScratchTickStage == (int)TickStage);
-
-		// Allows us to count total required instances across all ticks
-		for (uint32 i = 0; i < Tick.Count; ++i)
+		if (DispatchGroup.FreeIDUpdates.Num() > 0)
 		{
-			FNiagaraComputeInstanceData& InstanceData = Tick.GetInstanceData()[i];
-			FNiagaraComputeExecutionContext* ExecContext = InstanceData.Context;
-			if ((ExecContext == nullptr) || !ExecContext->GPUScript_RT->IsShaderMapComplete_RenderThread())
+			for (FNiagaraComputeExecutionContext* ComputeContext : DispatchGroup.FreeIDUpdates )
 			{
-				continue;
+				TransitionsAfter.Emplace(ComputeContext->MainDataSet->GetGPUFreeIDs().UAV, ERHIAccess::SRVCompute, ERHIAccess::UAVCompute);
 			}
+		}
 
-			uint32 PrevNumInstances = 0;
-			if (ExecContext->TicksThisFrame == 0)
-			{
-				PrevNumInstances = Tick.bNeedsReset ? 0 : ExecContext->MainDataSet->GetCurrentData()->GetNumInstances();
-				ExecContext->ScratchMaxInstances = PrevNumInstances;
-			}
-			else
+		// Consume per tick data from the game thread
+		for (FNiagaraGPUSystemTick* Tick : DispatchGroup.TicksWithPerInstanceData)
+		{
+			uint8* BasePointer = reinterpret_cast<uint8*>(Tick->DIInstanceData->PerInstanceDataForRT);
+
+			for (const auto& Pair : Tick->DIInstanceData->InterfaceProxiesToOffsets)
 			{
-				PrevNumInstances = Tick.bNeedsReset ? 0 : ExecContext->ScratchNumInstances;
+				FNiagaraDataInterfaceProxy* Proxy = Pair.Key;
+				uint8* InstanceDataPtr = BasePointer + Pair.Value;
+				Proxy->ConsumePerInstanceDataFromGameThread(InstanceDataPtr, Tick->SystemInstanceID);
 			}
+		}
+
+		// Execute Before Transitions
+		RHICmdList.Transition(TransitionsBefore);
+		TransitionsBefore.Reset();
 
-			/** Track the number of ticks processed so we can accumulate instance counts correctly across multiple ticks that may or may not include this emitter. */
-			++ExecContext->TicksThisFrame;
+		// Initialize the IDtoIndex tables
+		if (IDToIndexInit.Num() > 0)
+		{
+			SCOPED_DRAW_EVENT(RHICmdList, NiagaraGPUComputeClearIDToIndexBuffer);
 
-			ExecContext->ScratchNumInstances = InstanceData.SpawnInfo.SpawnRateInstances + InstanceData.SpawnInfo.EventSpawnTotal + PrevNumInstances;
+			TArray<FRHITransitionInfo, TMemStackAllocator<>> IDToIndexTransitions;
+			IDToIndexTransitions.Reserve(IDToIndexInit.Num());
 
-			const uint32 MaxInstanceCount = ExecContext->MainDataSet->GetMaxInstanceCount();
-			ExecContext->ScratchNumInstances = FMath::Min(ExecContext->ScratchNumInstances, MaxInstanceCount);
-			if (bAlwaysAllocateBuffers || (ExecContext->ScratchNumInstances > 0))
+			for (FNiagaraDataBuffer* IDtoIndexBuffer : IDToIndexInit)
 			{
-				ExecContext->ScratchMaxInstances = FMath::Max3(ExecContext->ScratchMaxInstances, InstanceData.SpawnInfo.MaxParticleCount, ExecContext->ScratchNumInstances);
+				NiagaraFillGPUIntBuffer(RHICmdList, FeatureLevel, IDtoIndexBuffer->GetGPUIDToIndexTable(), INDEX_NONE);
+				IDToIndexTransitions.Emplace(IDtoIndexBuffer->GetGPUIDToIndexTable().UAV, ERHIAccess::UAVCompute, ERHIAccess::UAVCompute);
 			}
-
-			InstanceData.SimStageData[0].SourceNumInstances = PrevNumInstances;
-			InstanceData.SimStageData[0].DestinationNumInstances = ExecContext->ScratchNumInstances;
+			IDToIndexInit.Reset();
+			RHICmdList.Transition(IDToIndexTransitions);
 		}
-	}
 
-	for (int32 iTickStage = 0; iTickStage < (int)ETickStage::Max; ++iTickStage)
-	{
-		// Set bIsFinalTick for the last tick of each context and reset the scratch index.
-		const int32 ScrachIndexReset = UseOverlapCompute() ? 0 : INDEX_NONE;
-		for (FNiagaraComputeSharedContext* SharedContext : ContextsPerStage[iTickStage])
+		// Execute PreStage
+		for (const FNiagaraGpuDispatchInstance& DispatchInstance : DispatchGroup.DispatchInstances)
 		{
-			TicksPerStage[iTickStage][SharedContext->ScratchIndex]->bIsFinalTick = true;
-			SharedContext->ScratchIndex = ScrachIndexReset;
+			PreStageInterface(RHICmdList, DispatchInstance.Tick, DispatchInstance.InstanceData, DispatchInstance.SimStageData);
 		}
 
-		for (FNiagaraGPUSystemTick* Tick : TicksPerStage[iTickStage])
+		// Execute Stage
+		RHICmdList.BeginUAVOverlap(GPUInstanceCounterManager.GetInstanceCountBuffer().UAV);
+		for (const FNiagaraGpuDispatchInstance& DispatchInstance : DispatchGroup.DispatchInstances)
 		{
-			for (uint32 i = 0; i < Tick->Count; ++i)
+			++FNiagaraComputeExecutionContext::TickCounter;
+			if (DispatchInstance.InstanceData.bResetData && DispatchInstance.SimStageData.bFirstStage )
 			{
-				FNiagaraComputeInstanceData& InstanceData = Tick->GetInstanceData()[i];
-				FNiagaraComputeExecutionContext* ExecContext = InstanceData.Context;
-
-				if (ExecContext == nullptr)
-				{
-					continue;
-				}
-
-				//Ensure we clear the temp tick count value in every emitter tick that we've processed.
-				//Can't do later in the final tick as that may not contain *all* emitter ticks.
-				ExecContext->TicksThisFrame = 0;
-
-				if (!ExecContext->GPUScript_RT->IsShaderMapComplete_RenderThread())
-				{
-					continue;
-				}
-
-				// First stage is presumed to be a particle stage
-				FNiagaraDataBuffer* CurrentData = ExecContext->MainDataSet->GetCurrentData();
-				FNiagaraDataBuffer* DestinationData = &ExecContext->MainDataSet->BeginSimulate();
-
-				if (ExecContext->MainDataSet->RequiresPersistentIDs())
+				ResetDataInterfaces(RHICmdList, DispatchInstance.Tick, DispatchInstance.InstanceData);
+			}
+#if STATS
+			if (GPUProfiler.IsProfilingEnabled() && DispatchInstance.Tick.bIsFinalTick)
+			{
+				TStatId StatId;
+				if (const FSimulationStageMetaData* StageMetaData = DispatchInstance.SimStageData.StageMetaData)
 				{
-					ExecContext->MainDataSet->AllocateGPUFreeIDs(ExecContext->ScratchMaxInstances + 1, RHICmdList, FeatureLevel, ExecContext->GetDebugSimName());
+					StatId = FDynamicStats::CreateStatId<FStatGroup_STATGROUP_NiagaraDetailed>("GPU_Stage_" + StageMetaData->SimulationStageName.ToString());
 				}
-
-				if (DestinationData->GetGPUInstanceCountBufferOffset() != INDEX_NONE)
+				else
 				{
-					CountsToRelease[iTickStage].Add(DestinationData->GetGPUInstanceCountBufferOffset());
-					DestinationData->ClearGPUInstanceCountBufferOffset();
+					static FString DefaultStageName = "GPU_Stage_SpawnUpdate";
+					StatId = FDynamicStats::CreateStatId<FStatGroup_STATGROUP_NiagaraDetailed>(DefaultStageName);
 				}
+				const int32 TimerHandle = GPUProfiler.StartTimer((uint64)DispatchInstance.InstanceData.Context, StatId, RHICmdList);
+				DispatchStage(RHICmdList, ViewUniformBuffer, DispatchInstance.Tick, DispatchInstance.InstanceData, DispatchInstance.SimStageData);
+				GPUProfiler.EndTimer(TimerHandle, RHICmdList);
+			}
+			else
+#endif
+			{
+				DispatchStage(RHICmdList, ViewUniformBuffer, DispatchInstance.Tick, DispatchInstance.InstanceData, DispatchInstance.SimStageData);
+			}
+		}
+		RHICmdList.EndUAVOverlap(GPUInstanceCounterManager.GetInstanceCountBuffer().UAV);
 
-				DestinationData->AllocateGPU(ExecContext->ScratchMaxInstances + 1, GPUInstanceCounterManager, RHICmdList, FeatureLevel, ExecContext->GetDebugSimName());
-
-				InstanceData.SimStageData[0].Source = CurrentData;
-				InstanceData.SimStageData[0].SourceCountOffset = CurrentData->GetGPUInstanceCountBufferOffset();
-				InstanceData.SimStageData[0].Destination = DestinationData;
-				InstanceData.SimStageData[0].DestinationCountOffset = DestinationData->GetGPUInstanceCountBufferOffset();
+		// Execute PostStage
+		for (const FNiagaraGpuDispatchInstance& DispatchInstance : DispatchGroup.DispatchInstances)
+		{
+			PostStageInterface(RHICmdList, DispatchInstance.Tick, DispatchInstance.InstanceData, DispatchInstance.SimStageData);
+			if ( DispatchInstance.SimStageData.bLastStage )
+			{
+				PostSimulateInterface(RHICmdList, DispatchInstance.Tick, DispatchInstance.InstanceData);
 
-				ExecContext->MainDataSet->EndSimulate();
+				// Update CurrentData with the latest information as things like ParticleReads can use this data
+				FNiagaraComputeExecutionContext* ComputeContext = DispatchInstance.InstanceData.Context;
+				const FNiagaraSimStageData& FinalSimStageData = DispatchInstance.SimStageData;
+				FNiagaraDataBuffer* CurrentData = ComputeContext->MainDataSet->GetCurrentData();
+				FNiagaraDataBuffer* FinalData = FinalSimStageData.Destination != nullptr ? FinalSimStageData.Destination : FinalSimStageData.Source;
+				check(FinalData);
 
-				// If we are going to enqueue a readback this frame
-				//-TODO: Currently this does not work with low latency translucency, we need to allocate counters and release / modify DataBuffers when executing to avoid issues around renders picking up the wrong ones that are opaque
-				bool bClearBufferFromReadback = false;
-				if (bEnqueueReadback && Tick->bIsFinalTick && (ExecContext->EmitterInstanceReadback.GPUCountOffset == INDEX_NONE) && (CurrentData->GetGPUInstanceCountBufferOffset() != INDEX_NONE))
-				{
-					bClearBufferFromReadback = true;
-					ExecContext->EmitterInstanceReadback.GPUCountOffset = InstanceData.SimStageData[0].Source->GetGPUInstanceCountBufferOffset();
-					ExecContext->EmitterInstanceReadback.CPUCount = InstanceData.SimStageData[0].SourceNumInstances;
-				}
+				CurrentData->AliasGPU(FinalData);
 
-				if (Tick->NumInstancesWithSimStages > 0)
-				{
-					// Setup iteration source for stage 0
-					InstanceData.SimStageData[0].AlternateIterationSource = FindIterationInterface(&InstanceData, 0);
+				ComputeContext->SetTranslucentDataToRender(CurrentData);
+				ComputeContext->SetDataToRender(CurrentData);
+			}
+		}
 
-					const uint32 NumStages = InstanceData.Context->MaxUpdateIterations;
-					if (NumStages > 1)
-					{
-						// Flip current / destination buffers so we read from the buffer we just wrote into on the next stage
-						Swap(CurrentData, DestinationData);
+		// Execute After Transitions
+		RHICmdList.Transition(TransitionsAfter);
+		TransitionsAfter.Reset();
 
-						uint32 CurrentNumInstances = InstanceData.SimStageData[0].DestinationNumInstances;
-						uint32 DestinationNumInstances = InstanceData.SimStageData[0].SourceNumInstances;
+		// Update free IDs
+		if (DispatchGroup.FreeIDUpdates.Num() > 0)
+		{
+			UpdateFreeIDsListSizesBuffer(RHICmdList, DispatchGroup.FreeIDUpdates.Num());
+			UpdateFreeIDBuffers(RHICmdList, DispatchGroup.FreeIDUpdates);
 
-						// Old simulation will always RW the instance count, so ensure we have a count present for the destination
-						if (InstanceData.bUsesOldShaderStages)
-						{
-							if (DestinationData->GetGPUInstanceCountBufferOffset() == INDEX_NONE)
-							{
-								DestinationData->AllocateGPU(ExecContext->ScratchMaxInstances + 1, GPUInstanceCounterManager, RHICmdList, FeatureLevel, ExecContext->GetDebugSimName());
-							}
-						}
+			for (FNiagaraComputeExecutionContext* ComputeContext : DispatchGroup.FreeIDUpdates)
+			{
+				TransitionsAfter.Emplace(ComputeContext->MainDataSet->GetGPUFreeIDs().UAV, ERHIAccess::UAVCompute, ERHIAccess::SRVCompute);
+			}
+			RHICmdList.Transition(TransitionsAfter);
+			TransitionsAfter.Reset();
+		}
+	}
 
+	DispatchList.DispatchGroups.Empty();
+	if (DispatchList.CountsToRelease.Num() > 0)
+	{
+		GPUInstanceCounterManager.FreeEntryArray(DispatchList.CountsToRelease);
+		DispatchList.CountsToRelease.Empty();
+	}
+}
 
-						for (uint32 SimulationStageIndex = 1; SimulationStageIndex < NumStages; ++SimulationStageIndex)
-						{
-							InstanceData.SimStageData[SimulationStageIndex].Source = CurrentData;
-							InstanceData.SimStageData[SimulationStageIndex].SourceCountOffset = CurrentData->GetGPUInstanceCountBufferOffset();
-							InstanceData.SimStageData[SimulationStageIndex].SourceNumInstances = CurrentNumInstances;
-							InstanceData.SimStageData[SimulationStageIndex].Destination = DestinationData;
-							InstanceData.SimStageData[SimulationStageIndex].DestinationCountOffset = InstanceData.bUsesOldShaderStages ? DestinationData->GetGPUInstanceCountBufferOffset() : INDEX_NONE;
-							InstanceData.SimStageData[SimulationStageIndex].DestinationNumInstances = DestinationNumInstances;
-
-							// Determine if the iteration is outputting to a custom data size
-							FNiagaraDataInterfaceProxyRW* IterationInterface = FindIterationInterface(&InstanceData, SimulationStageIndex);
-							InstanceData.SimStageData[SimulationStageIndex].AlternateIterationSource = IterationInterface;
-
-							if (!NiagaraEmitterInstanceBatcherLocal::ShouldRunStage(Tick, ExecContext, IterationInterface, SimulationStageIndex))
-							{
-								continue;
-							}
-
-							if (!InstanceData.bUsesOldShaderStages)
-							{
-								// This should never be nullptr with simulation stages
-								const FSimulationStageMetaData* StageMetaData = ExecContext->GetSimStageMetaData(SimulationStageIndex);
-								check(StageMetaData);
-								InstanceData.SimStageData[SimulationStageIndex].StageMetaData = StageMetaData;
-
-								// No particle data will be written we read only, i.e. scattering particles into a grid
-								if (!StageMetaData->bWritesParticles)
-								{
-									InstanceData.SimStageData[SimulationStageIndex].Destination = nullptr;
-									InstanceData.SimStageData[SimulationStageIndex].DestinationNumInstances = CurrentNumInstances;
-									continue;
-								}
-
-								// Particle counts are not changing and we can safely read / write to the same particle buffer
-								if (StageMetaData->bPartialParticleUpdate)
-								{
-									InstanceData.SimStageData[SimulationStageIndex].Source = nullptr;
-
-									InstanceData.SimStageData[SimulationStageIndex].Destination = CurrentData;
-									InstanceData.SimStageData[SimulationStageIndex].DestinationCountOffset = INDEX_NONE;
-									InstanceData.SimStageData[SimulationStageIndex].DestinationNumInstances = CurrentNumInstances;
-									continue;
-								}
-							}
-
-							// We need to allocate a new buffer as particle counts could be changing
-							ensure(CurrentData == ExecContext->MainDataSet->GetCurrentData());
-							DestinationData = &ExecContext->MainDataSet->BeginSimulate(false);
-							if (DestinationData->GetGPUInstanceCountBufferOffset() != INDEX_NONE)
-							{
-								if (bClearBufferFromReadback && (ExecContext->EmitterInstanceReadback.GPUCountOffset == DestinationData->GetGPUInstanceCountBufferOffset()))
-								{
-									bClearBufferFromReadback = false;
-								}
-								else
-								{
-									CountsToRelease[iTickStage].Add(DestinationData->GetGPUInstanceCountBufferOffset());
-								}
-								DestinationData->ClearGPUInstanceCountBufferOffset();
-							}
-							DestinationData->AllocateGPU(ExecContext->ScratchMaxInstances + 1, GPUInstanceCounterManager, RHICmdList, FeatureLevel, ExecContext->GetDebugSimName());
-							ExecContext->MainDataSet->EndSimulate();
-
-							DestinationNumInstances = CurrentNumInstances;
-
-							InstanceData.SimStageData[SimulationStageIndex].Destination = DestinationData;
-							InstanceData.SimStageData[SimulationStageIndex].DestinationCountOffset = DestinationData->GetGPUInstanceCountBufferOffset();
-							InstanceData.SimStageData[SimulationStageIndex].DestinationNumInstances = DestinationNumInstances;
-
-							Swap(CurrentData, DestinationData);
-						}
-					}
-				}
+void NiagaraEmitterInstanceBatcher::DispatchStage(FRHICommandList& RHICmdList, FRHIUniformBuffer* ViewUniformBuffer, const FNiagaraGPUSystemTick& Tick, const FNiagaraComputeInstanceData& InstanceData, const FNiagaraSimStageData& SimStageData)
+{
+	// Setup source buffer
+	if ( SimStageData.Source != nullptr )
+	{
+		SimStageData.Source->SetNumInstances(SimStageData.SourceNumInstances);
+		SimStageData.Source->GPUInstanceCountBufferOffset = SimStageData.SourceCountOffset;
+	}
 
-				// If we queued an emitter readback we need to clear out the count from the original buffer to avoid it being released
-				if (bClearBufferFromReadback)
-				{
-					check(InstanceData.SimStageData[0].Source && (InstanceData.SimStageData[0].Source->GetGPUInstanceCountBufferOffset() != INDEX_NONE));
-					InstanceData.SimStageData[0].Source->ClearGPUInstanceCountBufferOffset();
-				}
+	// Setup destination buffer
+	int32 InstancesToSpawn = 0;
+	if ( SimStageData.Destination != nullptr )
+	{
+		SimStageData.Destination->SetNumInstances(SimStageData.DestinationNumInstances);
+		SimStageData.Destination->GPUInstanceCountBufferOffset = SimStageData.DestinationCountOffset;
+		SimStageData.Destination->SetIDAcquireTag(FNiagaraComputeExecutionContext::TickCounter);
 
-				if (Tick->bIsFinalTick && GNiagaraGpuLowLatencyTranslucencyEnabled)
-				{
-					ExecContext->SetTranslucentDataToRender(ExecContext->MainDataSet->GetCurrentData());
-				}
-			}
+		if ( SimStageData.bFirstStage )
+		{
+			check(SimStageData.DestinationNumInstances >= SimStageData.SourceNumInstances);
+			InstancesToSpawn = SimStageData.DestinationNumInstances - SimStageData.SourceNumInstances;
 		}
+		SimStageData.Destination->SetNumSpawnedInstances(InstancesToSpawn);
 	}
-}
 
-void NiagaraEmitterInstanceBatcher::ExecuteAll(FRHICommandList& RHICmdList, FRHIUniformBuffer* ViewUniformBuffer, ETickStage TickStage)
-{
-	SCOPE_CYCLE_COUNTER(STAT_NiagaraGPUSimTick_RT);
+	// Get dispatch count
+	uint32 DispatchCount = 0;
+	if ( SimStageData.AlternateIterationSource )
+	{
+		const FIntVector ElementCount = SimStageData.AlternateIterationSource->GetElementCount(Tick.SystemInstanceID);
+		// Verify the number of elements isn't higher that what we can handle
+		checkf(ElementCount.X * ElementCount.Y * ElementCount.Z < uint64(TNumericLimits<uint32>::Max()), TEXT("ElementCount(%d, %d, %d) for IterationInterface(%s) overflows an int32 this is not allowed"), ElementCount.X, ElementCount.Y, ElementCount.Z, *SimStageData.AlternateIterationSource->SourceDIName.ToString());
+
+		DispatchCount = ElementCount.X * ElementCount.Y * ElementCount.Z;
+	}
+	else
+	{
+		DispatchCount = SimStageData.DestinationNumInstances;
+	}
 
-	// Anything to execute for this tick stage?
-	if (TicksPerStage[(int)TickStage].Num() == 0)
+	if ( DispatchCount == 0 )
 	{
 		return;
 	}
 
-	SCOPED_DRAW_EVENTF(RHICmdList, NiagaraEmitterInstanceBatcher_ExecuteAll, TEXT("NiagaraEmitterInstanceBatcher_ExecuteAll - TickStage(%d)"), TickStage);
+	// Get Shader
+	const int32 PermutationId = Tick.NumInstancesWithSimStages > 0 ? InstanceData.Context->GPUScript_RT->ShaderStageIndexToPermutationId_RenderThread(SimStageData.StageIndex) : 0;
+	const FNiagaraShaderRef ComputeShader = InstanceData.Context->GPUScript_RT->GetShader(PermutationId);
+	FRHIComputeShader* RHIComputeShader = ComputeShader.GetComputeShader();
+	RHICmdList.SetComputeShader(RHIComputeShader);
 
-	FMemMark Mark(FMemStack::Get());
-	TArray<FOverlappableTicks, TMemStackAllocator<> > SimPasses;
+	SCOPED_DRAW_EVENTF(RHICmdList, NiagaraGPUSimulationCS,
+		TEXT("NiagaraGpuSim(%s) DispatchCount(%u) Stage(%s %u) NumInstructions(%u)"),
+		InstanceData.Context->GetDebugSimName(),
+		DispatchCount,
+		SimStageData.StageMetaData ? *SimStageData.StageMetaData->SimulationStageName.ToString() : TEXT("Particles"),
+		SimStageData.StageIndex,
+		ComputeShader->GetNumInstructions()
+	);
+	FNiagaraUAVPoolAccessScope UAVPoolAccessScope(*this);
+
+	// Set Parameters
+	const bool bRequiresPersistentIDs = InstanceData.Context->MainDataSet->RequiresPersistentIDs();
+
+	SetShaderValue(RHICmdList, RHIComputeShader, ComputeShader->SimStartParam, InstanceData.bResetData ? 1U : 0U);
+	SetShaderValue(RHICmdList, RHIComputeShader, ComputeShader->EmitterTickCounterParam, FNiagaraComputeExecutionContext::TickCounter);
+	SetShaderValue(RHICmdList, RHIComputeShader, ComputeShader->UpdateStartInstanceParam, 0);
+	SetShaderValue(RHICmdList, RHIComputeShader, ComputeShader->NumSpawnedInstancesParam, InstancesToSpawn);
+	SetSRVParameter(RHICmdList, RHIComputeShader, ComputeShader->FreeIDBufferParam, bRequiresPersistentIDs ? InstanceData.Context->MainDataSet->GetGPUFreeIDs().SRV.GetReference() : FNiagaraRenderer::GetDummyIntBuffer());
+
+	//-TODO: REMOVE OLD SIM STAGES
+	SetShaderValue(RHICmdList, RHIComputeShader, ComputeShader->DefaultSimulationStageIndexParam, 0);
+	SetShaderValue(RHICmdList, RHIComputeShader, ComputeShader->SimulationStageIndexParam, SimStageData.StageIndex);
+	//-TODO: REMOVE OLD SIM STAGES
+
+	// Set spawn Information
+	// This parameter is an array of structs with 2 floats and 2 ints on CPU, but a float4 array on GPU. The shader uses asint() to cast the integer values. To set the parameter,
+	// we pass the structure array as a float* to SetShaderValueArray() and specify the number of floats (not float vectors).
+	static_assert((sizeof(InstanceData.SpawnInfo.SpawnInfoStartOffsets) % SHADER_PARAMETER_ARRAY_ELEMENT_ALIGNMENT) == 0, "sizeof SpawnInfoStartOffsets should be a multiple of SHADER_PARAMETER_ARRAY_ELEMENT_ALIGNMENT");
+	static_assert((sizeof(InstanceData.SpawnInfo.SpawnInfoParams) % SHADER_PARAMETER_ARRAY_ELEMENT_ALIGNMENT) == 0, "sizeof SpawnInfoParams should be a multiple of SHADER_PARAMETER_ARRAY_ELEMENT_ALIGNMENT");
+	SetShaderValueArray(RHICmdList, RHIComputeShader, ComputeShader->EmitterSpawnInfoOffsetsParam, InstanceData.SpawnInfo.SpawnInfoStartOffsets, NIAGARA_MAX_GPU_SPAWN_INFOS);
+	SetShaderValueArray(RHICmdList, RHIComputeShader, ComputeShader->EmitterSpawnInfoParamsParam, &InstanceData.SpawnInfo.SpawnInfoParams[0].IntervalDt, 4 * NIAGARA_MAX_GPU_SPAWN_INFOS);
+
+	if (ComputeShader->ViewUniformBufferParam.IsBound())
 	{
-		if (UseOverlapCompute())
+		if (ensureMsgf(ViewUniformBuffer != nullptr, TEXT("ViewUniformBuffer is required for '%s' but we do not have one to bind"), InstanceData.Context->GetDebugSimName()))
 		{
-			// Transpose now only once the data to get all independent tick per pass
-			SimPasses.Reserve(2); // Safe bet!
+			RHICmdList.SetShaderUniformBuffer(RHIComputeShader, ComputeShader->ViewUniformBufferParam.GetBaseIndex(), ViewUniformBuffer);
+		}
+	}
 
-			for (FNiagaraGPUSystemTick* Tick : TicksPerStage[(int)TickStage])
-			{
-				FNiagaraComputeSharedContext* SharedContext = Tick->SharedContext;
-				const int32 ScratchIndex = SharedContext->ScratchIndex;
-				check(ScratchIndex != INDEX_NONE);
+	// Setup instance counts
+	if (ComputeShader->InstanceCountsParam.IsBound())
+	{
+		ComputeShader->InstanceCountsParam.SetBuffer(RHICmdList, RHIComputeShader, GPUInstanceCounterManager.GetInstanceCountBuffer());
 
-				if (ScratchIndex >= SimPasses.Num())
-				{
-					SimPasses.AddDefaulted(SimPasses.Num() - ScratchIndex + 1);
-					if (ScratchIndex == 0)
-					{
-						SimPasses[0].Reserve(ContextsPerStage[(int)TickStage].Num()); // Guarantied!
-					}
-				}
-				SimPasses[ScratchIndex].Add(Tick);
-				// Scratch index is now the number of passes for this context.
-				if (Tick->bIsFinalTick)
-				{
-					// Reset to default as it will no longer be used.
-					SharedContext->ScratchIndex = INDEX_NONE;
-				}
-				else
-				{
-					SharedContext->ScratchIndex += 1;
-				}
-			}
+		if (SimStageData.AlternateIterationSource)
+		{
+			SetShaderValue(RHICmdList, RHIComputeShader, ComputeShader->ReadInstanceCountOffsetParam, -1);
+			SetShaderValue(RHICmdList, RHIComputeShader, ComputeShader->WriteInstanceCountOffsetParam, -1);
 		}
 		else
 		{
-			// Force dispatches to run individually, this should only be used for debugging as it is highly inefficient on the GPU
-			SimPasses.Reserve(ContextsPerStage[(int)TickStage].Num()); // Guarantied!
-			for (FNiagaraGPUSystemTick* Tick : TicksPerStage[(int)TickStage])
-			{
-				SimPasses.AddDefaulted_GetRef().Add(Tick);
-			}
+			SetShaderValue(RHICmdList, RHIComputeShader, ComputeShader->ReadInstanceCountOffsetParam, SimStageData.SourceCountOffset);
+			SetShaderValue(RHICmdList, RHIComputeShader, ComputeShader->WriteInstanceCountOffsetParam, SimStageData.DestinationCountOffset);
 		}
 	}
 
-	if (SimPasses.Num() > 0)
+	// Simulation Stage Information
+	// X = Count Buffer Instance Count Offset (INDEX_NONE == Use Instance Count)
+	// Y = Instance Count
+	// Z = Iteration Index
+	// W = Num Iterations
 	{
-		FEmitterInstanceList InstancesWithPersistentIDs;
-		for (int32 SimPassIdx = 0; SimPassIdx < SimPasses.Num(); ++SimPassIdx)
+		FIntVector4 SimulationStageIterationInfo = { INDEX_NONE, -1, 0, 0 };
+		float SimulationStageNormalizedIterationIndex = 0.0f;
+
+		if (SimStageData.AlternateIterationSource != nullptr)
 		{
-			FOverlappableTicks& SimPass = SimPasses[SimPassIdx];
-			InstancesWithPersistentIDs.SetNum(0, false);
+			const uint32 IterationInstanceCountOffset = SimStageData.AlternateIterationSource->GetGPUInstanceCountOffset(Tick.SystemInstanceID);
+			SimulationStageIterationInfo.X = IterationInstanceCountOffset;
+			SimulationStageIterationInfo.Y = IterationInstanceCountOffset == INDEX_NONE ? DispatchCount : 0;
+		}
 
-			FDispatchGroupList DispatchGroups;
-			BuildDispatchGroups(SimPass, RHICmdList, DispatchGroups, InstancesWithPersistentIDs);
-			const bool bHasInstancesWithPersistentIDs = (InstancesWithPersistentIDs.Num() > 0);
+		if (SimStageData.StageMetaData != nullptr)
+		{
+			const int32 NumStages = SimStageData.StageMetaData->MaxStage - SimStageData.StageMetaData->MinStage;
+			ensure((int32(SimStageData.StageIndex) >= SimStageData.StageMetaData->MinStage) && (int32(SimStageData.StageIndex) < SimStageData.StageMetaData->MaxStage));
+			const int32 IterationIndex = SimStageData.StageIndex - SimStageData.StageMetaData->MinStage;
+			SimulationStageIterationInfo.Z = SimStageData.StageIndex - SimStageData.StageMetaData->MinStage;
+			SimulationStageIterationInfo.W = NumStages;
+			SimulationStageNormalizedIterationIndex = NumStages > 1 ? float(IterationIndex) / float(NumStages - 1) : 1.0f;
+		}
+		SetShaderValue(RHICmdList, RHIComputeShader, ComputeShader->SimulationStageIterationInfoParam, SimulationStageIterationInfo);
+		SetShaderValue(RHICmdList, RHIComputeShader, ComputeShader->SimulationStageNormalizedIterationIndexParam, SimulationStageNormalizedIterationIndex);
+	}
 
-			{
-				SCOPED_DRAW_EVENT(RHICmdList, NiagaraGPUSimulation);
-				SCOPED_GPU_STAT(RHICmdList, NiagaraGPUSimulation);
-				for (int32 GroupIdx = 0; GroupIdx < DispatchGroups.Num(); ++GroupIdx)
-				{
-					FDispatchGroup& Group = DispatchGroups[GroupIdx];
-					//SCOPED_DRAW_EVENTF(RHICmdList, NiagaraOverlapGroup, TEXT("Overlap Group - %u emitters"), Group.DispatchInstances.Num());
+	// Set Input & Output buffers
+	FNiagaraDataBuffer::SetInputShaderParams(RHICmdList, ComputeShader.GetShader(), SimStageData.Source);
+	FNiagaraDataBuffer::SetOutputShaderParams(RHICmdList, ComputeShader.GetShader(), SimStageData.Destination);
 
-					const bool bIsFirstGroup = (GroupIdx == 0);
-					const bool bIsLastGroup = (GroupIdx == DispatchGroups.Num() - 1);
+	// Set data interface parameters
+	SetDataInterfaceParameters(RHICmdList, Tick, InstanceData, ComputeShader, SimStageData);
 
-					// If we're the first group, insert a transition from readable to writable on the instance count buffer. Otherwise we need a UAV to UAV barrier.
-					ERHIAccess InstanceCountBufferSrcAccess = bIsFirstGroup ? FNiagaraGPUInstanceCountManager::kCountBufferDefaultState : ERHIAccess::UAVCompute;
-					Group.TransitionsBefore.Add(FRHITransitionInfo(GPUInstanceCounterManager.GetInstanceCountBuffer().UAV, InstanceCountBufferSrcAccess, ERHIAccess::UAVCompute));
+	// Execute the dispatch
+	{
+		const uint32 ThreadGroupSize = FNiagaraShader::GetGroupSize(ShaderPlatform);
+		checkf(ThreadGroupSize <= NiagaraComputeMaxThreadGroupSize, TEXT("ShaderPlatform(%d) is set to ThreadGroup size (%d) which is > the NiagaraComputeMaxThreadGroupSize(%d)"), ShaderPlatform, ThreadGroupSize, NiagaraComputeMaxThreadGroupSize);
+		const uint32 TotalThreadGroups = FMath::DivideAndRoundUp(DispatchCount, ThreadGroupSize);
 
-					// If we're the last group, make the instance count buffer readable after the dispatch.
-					if (bIsLastGroup)
-					{
-						Group.TransitionsAfter.Add(FRHITransitionInfo(GPUInstanceCounterManager.GetInstanceCountBuffer().UAV, ERHIAccess::UAVCompute, FNiagaraGPUInstanceCountManager::kCountBufferDefaultState));
-					}
+		const uint32 ThreadGroupCountZ = 1;
+		const uint32 ThreadGroupCountY = FMath::DivideAndRoundUp(TotalThreadGroups, NiagaraMaxThreadGroupCountPerDimension);
+		const uint32 ThreadGroupCountX = FMath::DivideAndRoundUp(TotalThreadGroups, ThreadGroupCountY);
 
-					RHICmdList.Transition(Group.TransitionsBefore);
+		SetShaderValue(RHICmdList, RHIComputeShader, ComputeShader->DispatchThreadIdToLinearParam, ThreadGroupCountY > 1 ? ThreadGroupCountX * ThreadGroupSize : 0);
 
-					RHICmdList.BeginUAVOverlap(GPUInstanceCounterManager.GetInstanceCountBuffer().UAV);
-					DispatchAllOnCompute(Group.DispatchInstances, RHICmdList, ViewUniformBuffer);
-					RHICmdList.EndUAVOverlap(GPUInstanceCounterManager.GetInstanceCountBuffer().UAV);
+		SetConstantBuffers(RHICmdList, ComputeShader, Tick, &InstanceData);
 
-					RHICmdList.Transition(Group.TransitionsAfter);
+		DispatchComputeShader(RHICmdList, ComputeShader, ThreadGroupCountX, ThreadGroupCountY, ThreadGroupCountZ);
 
-					// If we're the last group and we have emitters with persistent IDs, clear the ID list sizes before doing the final transitions, to overlap that dispatch with the simulation dispatches.
-					if (bIsLastGroup && bHasInstancesWithPersistentIDs)
-					{
-						// If we're doing multiple ticks (e.g. when scrubbing the timeline in the editor), we must update the free ID buffers before running
-						// the next tick, which will cause stalls (because the ID to index buffer is written by DispatchAllOnCompute and read by UpdateFreeIDBuffers).
-						// However, when we're at the last tick, we can postpone the update until later in the frame and avoid the stall. This will be the case when
-						// running normally, with one tick per frame.
-						if (SimPassIdx < SimPasses.Num() - 1)
-						{
-							UpdateFreeIDsListSizesBuffer(RHICmdList, InstancesWithPersistentIDs.Num());
-							UpdateFreeIDBuffers(RHICmdList, InstancesWithPersistentIDs);
-						}
-						else
-						{
-							for (FNiagaraComputeInstanceData* Instance : InstancesWithPersistentIDs)
-							{
-								DeferredIDBufferUpdates.AddUnique(Instance);
-							}
-							UpdateFreeIDsListSizesBuffer(RHICmdList, DeferredIDBufferUpdates.Num());
-						}
-					}
-				}
-			}
-		}
+		INC_DWORD_STAT(STAT_NiagaraGPUDispatches);
 	}
 
-	// Release counts
-	GPUInstanceCounterManager.FreeEntryArray(MakeArrayView(CountsToRelease[int(TickStage)]));
-	CountsToRelease[int(TickStage)].Reset();
+	// Unset UAV parameters
+	UnsetDataInterfaceParameters(RHICmdList, Tick, InstanceData, ComputeShader, SimStageData);
+	FNiagaraDataBuffer::UnsetShaderParams(RHICmdList, ComputeShader.GetShader());
+	ComputeShader->InstanceCountsParam.UnsetUAV(RHICmdList, RHIComputeShader);
+
+	// Optionally submit commands to the GPU
+	// This can be used to avoid accidental TDR detection in the editor especially when issuing multiple ticks in the same frame
+	if (GNiagaraGpuSubmitCommandHint > 0)
+	{
+		++TotalDispatchesThisFrame;
+		if ((TotalDispatchesThisFrame % GNiagaraGpuSubmitCommandHint) == 0)
+		{
+			RHICmdList.SubmitCommandsHint();
+		}
+	}
 }
 
 void NiagaraEmitterInstanceBatcher::PreInitViews(FRDGBuilder& GraphBuilder, bool bAllowGPUParticleUpdate)
 {
+	bRequiresReadback = false;
 	GNiagaraViewDataManager.ClearSceneTextureParameters();
 	FramesBeforeTickFlush = 0;
 
@@ -1513,6 +1631,12 @@ void NiagaraEmitterInstanceBatcher::PreInitViews(FRDGBuilder& GraphBuilder, bool
 	if (bAllowGPUParticleUpdate && FNiagaraUtilities::AllowGPUParticles(GetShaderPlatform()))
 	{
 		UpdateInstanceCountManager(GraphBuilder.RHICmdList);
+		PrepareAllTicks(GraphBuilder.RHICmdList);
+
+		if ( NiagaraEmitterInstanceBatcherLocal::GDebugLogging)
+		{
+			DumpDebugFrame();
+		}
 
 		if (GNiagaraAllowTickBeforeRender)
 		{
@@ -1526,8 +1650,7 @@ void NiagaraEmitterInstanceBatcher::PreInitViews(FRDGBuilder& GraphBuilder, bool
 				[this](FRHICommandListImmediate& RHICmdList)
 			{
 				FNiagaraSceneTextureScope Scope;
-				BuildTickStagePasses(RHICmdList, ETickStage::PreInitViews);
-				ExecuteAll(RHICmdList, nullptr, ETickStage::PreInitViews);
+				ExecuteTicks(RHICmdList, nullptr, ENiagaraGpuComputeTickStage::PreInitViews);
 			});
 		}
 	}
@@ -1557,8 +1680,7 @@ void NiagaraEmitterInstanceBatcher::PostInitViews(FRDGBuilder& GraphBuilder, TAr
 		{
 			FRHIUniformBuffer* ViewUniformBuffer = GetReferenceViewUniformBuffer(Views);
 			FNiagaraSceneTextureScope Scope;
-			BuildTickStagePasses(RHICmdList, ETickStage::PostInitViews);
-			ExecuteAll(RHICmdList, ViewUniformBuffer, ETickStage::PostInitViews);
+			ExecuteTicks(RHICmdList, ViewUniformBuffer, ENiagaraGpuComputeTickStage::PostInitViews);
 		});
 	}
 }
@@ -1592,13 +1714,8 @@ void NiagaraEmitterInstanceBatcher::PostRenderOpaque(FRDGBuilder& GraphBuilder,
 
 			FRHIUniformBuffer* ViewUniformBuffer = GetReferenceViewUniformBuffer(Views);
 
-			BuildTickStagePasses(RHICmdList, ETickStage::PostOpaqueRender);
-
 			// Setup new readback since if there is no pending request, there is no risk of having invalid data read (offset being allocated after the readback was sent).
-			ExecuteAll(RHICmdList, ViewUniformBuffer, ETickStage::PostOpaqueRender);
-
-			UpdateFreeIDBuffers(RHICmdList, DeferredIDBufferUpdates);
-			DeferredIDBufferUpdates.SetNum(0, false);
+			ExecuteTicks(RHICmdList, ViewUniformBuffer, ENiagaraGpuComputeTickStage::PostOpaqueRender);
 
 			FinishDispatches();
 
@@ -1609,7 +1726,7 @@ void NiagaraEmitterInstanceBatcher::PostRenderOpaque(FRDGBuilder& GraphBuilder,
 #endif
 		}
 
-		if (!GPUInstanceCounterManager.HasPendingGPUReadback())
+	if (!GPUInstanceCounterManager.HasPendingGPUReadback() && bRequiresReadback)
 		{
 			GPUInstanceCounterManager.EnqueueGPUReadback(RHICmdList);
 		}
@@ -1682,30 +1799,22 @@ void NiagaraEmitterInstanceBatcher::ProcessDebugReadbacks(FRHICommandListImmedia
 
 bool NiagaraEmitterInstanceBatcher::UsesGlobalDistanceField() const
 {
-	checkSlow(Ticks_RT.ContainsByPredicate([](const FNiagaraGPUSystemTick& Tick) { return Tick.bRequiresDistanceFieldData; }) == NumTicksThatRequireDistanceFieldData > 0);
-
-	return NumTicksThatRequireDistanceFieldData > 0;
+	return NumProxiesThatRequireDistanceFieldData > 0;
 }
 
 bool NiagaraEmitterInstanceBatcher::UsesDepthBuffer() const
 {
-	checkSlow(Ticks_RT.ContainsByPredicate([](const FNiagaraGPUSystemTick& Tick) { return Tick.bRequiresDepthBuffer; }) == NumTicksThatRequireDepthBuffer > 0);
-
-	return NumTicksThatRequireDepthBuffer > 0;
+	return NumProxiesThatRequireDepthBuffer > 0;
 }
 
 bool NiagaraEmitterInstanceBatcher::RequiresEarlyViewUniformBuffer() const
 {
-	checkSlow(Ticks_RT.ContainsByPredicate([](const FNiagaraGPUSystemTick& Tick) { return Tick.bRequiresEarlyViewData; }) == NumTicksThatRequireEarlyViewData > 0);
-
-	return NumTicksThatRequireEarlyViewData > 0;
+	return NumProxiesThatRequireEarlyViewData > 0;
 }
 
 bool NiagaraEmitterInstanceBatcher::RequiresRayTracingScene() const
 {
-	checkSlow(Ticks_RT.ContainsByPredicate([](const FNiagaraGPUSystemTick& Tick) { return Tick.bRequiresRayTracingScene; }) == NumTicksThatRequireRayTracingScene > 0);
-
-	return NumTicksThatRequireRayTracingScene > 0;
+	return NumProxiesThatRequireRayTracingScene > 0;
 }
 
 void NiagaraEmitterInstanceBatcher::PreRender(FRDGBuilder& GraphBuilder, TConstArrayView<FViewInfo> Views, bool bAllowGPUParticleUpdate)
@@ -1845,7 +1954,7 @@ void NiagaraEmitterInstanceBatcher::GenerateSortKeys(FRHICommandListImmediate& R
 #if RHI_RAYTRACING
 void NiagaraEmitterInstanceBatcher::BuildRayTracingSceneInfo(FRHICommandList& RHICmdList, TConstArrayView<FViewInfo> Views)
 {
-	if (NumTicksThatRequireRayTracingScene > 0)
+	if (NumProxiesThatRequireRayTracingScene > 0)
 	{
 		RayTracingHelper->BuildRayTracingSceneInfo(RHICmdList, Views);
 	}
@@ -1862,35 +1971,31 @@ void NiagaraEmitterInstanceBatcher::ResetRayTracingSceneInfo()
 
 void NiagaraEmitterInstanceBatcher::IssueRayTraces(FRHICommandList& RHICmdList, const FIntPoint& RayTraceCounts, FRHIShaderResourceView* RayTraceRequests, FRWBuffer* IndirectArgsBuffer, uint32 IndirectArgsOffset, FRHIUnorderedAccessView* RayTraceResults) const
 {
-	check(NumTicksThatRequireRayTracingScene > 0);
+	check(NumProxiesThatRequireRayTracingScene > 0);
 
 	RayTracingHelper->IssueRayTraces(RHICmdList, RayTraceCounts, RayTraceRequests, IndirectArgsBuffer, IndirectArgsOffset, RayTraceResults);
 }
 
 bool NiagaraEmitterInstanceBatcher::HasRayTracingScene() const
 {
-	return NumTicksThatRequireRayTracingScene && RayTracingHelper->IsValid();
+	return NumProxiesThatRequireRayTracingScene && RayTracingHelper->IsValid();
 }
 #endif
 
 /* Set shader parameters for data interfaces
  */
-void NiagaraEmitterInstanceBatcher::SetDataInterfaceParameters(const TArray<FNiagaraDataInterfaceProxy*>& DataInterfaceProxies, const FNiagaraShaderRef& Shader, FRHICommandList& RHICmdList, const FNiagaraComputeInstanceData* Instance, const FNiagaraGPUSystemTick& Tick, uint32 SimulationStageIndex) const
+void NiagaraEmitterInstanceBatcher::SetDataInterfaceParameters(FRHICommandList& RHICmdList, const FNiagaraGPUSystemTick& Tick, const FNiagaraComputeInstanceData& InstanceData, const FNiagaraShaderRef& ComputeShader, const FNiagaraSimStageData& SimStageData) const
 {
-	// set up data interface buffers, as defined by the DIs during compilation
-	//
-
 	// @todo-threadsafety This is a bit gross. Need to rethink this api.
-	const FNiagaraSystemInstanceID& SystemInstanceID = Tick.SystemInstanceID;
-	const FNiagaraShaderMapPointerTable& PointerTable = Shader.GetPointerTable();
+	const FNiagaraShaderMapPointerTable& PointerTable = ComputeShader.GetPointerTable();
 
 	uint32 InterfaceIndex = 0;
-	for (FNiagaraDataInterfaceProxy* Interface : DataInterfaceProxies)
+	for (FNiagaraDataInterfaceProxy* Interface : InstanceData.DataInterfaceProxies)
 	{
-		const FNiagaraDataInterfaceParamRef& DIParam = Shader->GetDIParameters()[InterfaceIndex];
+		const FNiagaraDataInterfaceParamRef& DIParam = ComputeShader->GetDIParameters()[InterfaceIndex];
 		if (DIParam.Parameters.IsValid())
 		{
-			FNiagaraDataInterfaceSetArgs Context(Interface, SystemInstanceID, this, Shader, Instance, SimulationStageIndex, Instance->IsOutputStage(Interface, SimulationStageIndex), Instance->IsIterationStage(Interface, SimulationStageIndex));
+			FNiagaraDataInterfaceSetArgs Context(Interface, Tick.SystemInstanceID, this, ComputeShader, &InstanceData, &SimStageData, InstanceData.IsOutputStage(Interface, SimStageData.StageIndex), InstanceData.IsIterationStage(Interface, SimStageData.StageIndex));
 			DIParam.DIType.Get(PointerTable.DITypes)->SetParameters(DIParam.Parameters.Get(), RHICmdList, Context);
 		}
 
@@ -1898,32 +2003,18 @@ void NiagaraEmitterInstanceBatcher::SetDataInterfaceParameters(const TArray<FNia
 	}
 }
 
-void NiagaraEmitterInstanceBatcher::UnsetDataInterfaceParameters(const TArray<FNiagaraDataInterfaceProxy*> &DataInterfaceProxies, const FNiagaraShaderRef& Shader, FRHICommandList &RHICmdList, const FNiagaraComputeInstanceData* Instance, const FNiagaraGPUSystemTick& Tick, uint32 SimulationStageIndex) const
+void NiagaraEmitterInstanceBatcher::UnsetDataInterfaceParameters(FRHICommandList& RHICmdList, const FNiagaraGPUSystemTick& Tick, const FNiagaraComputeInstanceData& InstanceData, const FNiagaraShaderRef& ComputeShader, const FNiagaraSimStageData& SimStageData) const
 {
-	// set up data interface buffers, as defined by the DIs during compilation
-	//
-
 	// @todo-threadsafety This is a bit gross. Need to rethink this api.
-	const FNiagaraSystemInstanceID& SystemInstance = Tick.SystemInstanceID;
-	const FNiagaraShaderMapPointerTable& PointerTable = Shader.GetPointerTable();
+	const FNiagaraShaderMapPointerTable& PointerTable = ComputeShader.GetPointerTable();
 
 	uint32 InterfaceIndex = 0;
-	for (FNiagaraDataInterfaceProxy* Interface : DataInterfaceProxies)
+	for (FNiagaraDataInterfaceProxy* Interface : InstanceData.DataInterfaceProxies)
 	{
-		const FNiagaraDataInterfaceParamRef& DIParam = Shader->GetDIParameters()[InterfaceIndex];
+		const FNiagaraDataInterfaceParamRef& DIParam = ComputeShader->GetDIParameters()[InterfaceIndex];
 		if (DIParam.Parameters.IsValid())
 		{
-			void* PerInstanceData = nullptr;
-			int32* OffsetFound = nullptr;
-			if (Tick.DIInstanceData && Tick.DIInstanceData->PerInstanceDataSize != 0 && Tick.DIInstanceData->InterfaceProxiesToOffsets.Num() != 0)
-			{
-				OffsetFound = Tick.DIInstanceData->InterfaceProxiesToOffsets.Find(Interface);
-				if (OffsetFound != nullptr)
-				{
-					PerInstanceData = (*OffsetFound) + (uint8*)Tick.DIInstanceData->PerInstanceDataForRT;
-				}
-			}
-			FNiagaraDataInterfaceSetArgs Context(Interface, SystemInstance, this, Shader, Instance, SimulationStageIndex, Instance->IsOutputStage(Interface, SimulationStageIndex), Instance->IsIterationStage(Interface, SimulationStageIndex));
+			FNiagaraDataInterfaceSetArgs Context(Interface, Tick.SystemInstanceID, this, ComputeShader, &InstanceData, &SimStageData, InstanceData.IsOutputStage(Interface, SimStageData.StageIndex), InstanceData.IsIterationStage(Interface, SimStageData.StageIndex));
 			DIParam.DIType.Get(PointerTable.DITypes)->UnsetParameters(DIParam.Parameters.Get(), RHICmdList, Context);
 		}
 
@@ -1960,185 +2051,6 @@ void NiagaraEmitterInstanceBatcher::SetConstantBuffers(FRHICommandList &RHICmdLi
 	}
 }
 
-/* Kick off a simulation/spawn run
- */
-void NiagaraEmitterInstanceBatcher::Run(const FNiagaraGPUSystemTick& Tick, const FNiagaraComputeInstanceData* Instance, uint32 UpdateStartInstance, const uint32 TotalNumInstances, const FNiagaraShaderRef& Shader,
-	FRHICommandList& RHICmdList, FRHIUniformBuffer* ViewUniformBuffer, const FNiagaraGpuSpawnInfo& SpawnInfo, bool bCopyBeforeStart, uint32 DefaultSimulationStageIndex, uint32 SimulationStageIndex, FNiagaraDataInterfaceProxyRW* IterationInterface, bool HasRunParticleStage)
-{
-	FNiagaraComputeExecutionContext* Context = Instance->Context;
-
-	// We must always set the current counts to ensure they are up to date when we have 0 instances
-	FNiagaraDataBuffer* CurrentData = Instance->SimStageData[SimulationStageIndex].Source;
-	if (CurrentData != nullptr)
-	{
-		CurrentData->SetNumInstances(Instance->SimStageData[SimulationStageIndex].SourceNumInstances);
-	}
-
-	FNiagaraDataBuffer* DestinationData = Instance->SimStageData[SimulationStageIndex].Destination;
-	int32 InstancesToSpawnThisFrame = 0;
-	if (DestinationData)
-	{
-		DestinationData->SetNumInstances(Instance->SimStageData[SimulationStageIndex].DestinationNumInstances);
-		DestinationData->SetIDAcquireTag(FNiagaraComputeExecutionContext::TickCounter);
-
-		// Only spawn particles on the first stage
-		if (!HasRunParticleStage)
-		{
-			if (Instance->SimStageData[SimulationStageIndex].DestinationNumInstances > Instance->SimStageData[SimulationStageIndex].SourceNumInstances)
-			{
-				InstancesToSpawnThisFrame = Instance->SimStageData[SimulationStageIndex].DestinationNumInstances - Instance->SimStageData[SimulationStageIndex].SourceNumInstances;
-			}
-		}
-		DestinationData->SetNumSpawnedInstances(InstancesToSpawnThisFrame);
-	}
-
-	if (TotalNumInstances == 0)
-	{
-		return;
-	}
-
-	SCOPED_DRAW_EVENTF(RHICmdList, NiagaraGPUSimulationCS, TEXT("NiagaraGpuSim(%s) NumInstances(%u) Stage(%s %u) NumInstructions(%u)"),
-		Context->GetDebugSimName(),
-		TotalNumInstances,
-		Context->GetSimStageMetaData(SimulationStageIndex) ? *Context->GetSimStageMetaData(SimulationStageIndex)->SimulationStageName.ToString() : TEXT("Particles"),
-		SimulationStageIndex,
-		Shader->GetNumInstructions()
-	);
-	FNiagaraUAVPoolAccessScope UAVPoolAccessScope(*this);
-
-	const TArray<FNiagaraDataInterfaceProxy*>& DataInterfaceProxies = Instance->DataInterfaceProxies;
-
-	FRHIComputeShader* ComputeShader = Shader.GetComputeShader();
-	RHICmdList.SetComputeShader(ComputeShader);
-
-	// #todo(dmp): clean up this logic for shader stages on first frame
-	SetShaderValue(RHICmdList, ComputeShader, Shader->SimStartParam, Tick.bNeedsReset ? 1U : 0U);
-
-	// set the view uniform buffer param
-	if (Shader->ViewUniformBufferParam.IsBound())
-	{
-		if (ensureMsgf(ViewUniformBuffer != nullptr, TEXT("ViewUniformBuffer is required for '%s' but we do not have one to bind"), Context->GetDebugSimName()))
-		{
-			RHICmdList.SetShaderUniformBuffer(ComputeShader, Shader->ViewUniformBufferParam.GetBaseIndex(), ViewUniformBuffer);
-		}
-	}
-
-	SetDataInterfaceParameters(DataInterfaceProxies, Shader, RHICmdList, Instance, Tick, SimulationStageIndex);
-
-	// set the shader and data set params
-	//
-	const bool bRequiresPersistentIDs = Context->MainDataSet->RequiresPersistentIDs();
-	SetSRVParameter(RHICmdList, Shader.GetComputeShader(), Shader->FreeIDBufferParam, bRequiresPersistentIDs ? Context->MainDataSet->GetGPUFreeIDs().SRV.GetReference() : FNiagaraRenderer::GetDummyIntBuffer());
-
-	FNiagaraDataBuffer::SetInputShaderParams(RHICmdList, Shader.GetShader(), CurrentData);
-	FNiagaraDataBuffer::SetOutputShaderParams(RHICmdList, Shader.GetShader(), DestinationData);
-
-	// set the instance count uav
-	//
-	if (Shader->InstanceCountsParam.IsBound())
-	{
-		Shader->InstanceCountsParam.SetBuffer(RHICmdList, ComputeShader, GPUInstanceCounterManager.GetInstanceCountBuffer());
-
-		if (IterationInterface)
-		{
-			SetShaderValue(RHICmdList, ComputeShader, Shader->ReadInstanceCountOffsetParam, -1);
-			SetShaderValue(RHICmdList, ComputeShader, Shader->WriteInstanceCountOffsetParam, -1);
-		}
-		else
-		{
-			const uint32 ReadOffset = (Tick.bNeedsReset && SimulationStageIndex == 0) ? INDEX_NONE : Instance->SimStageData[SimulationStageIndex].SourceCountOffset;
-			const uint32 WriteOffset = Instance->SimStageData[SimulationStageIndex].DestinationCountOffset;
-			//UE_LOG(LogNiagara, Log, TEXT("Instance count setup R: %d W: %d reset? %s %d"), ReadOffset, WriteOffset, Tick.bNeedsReset ? TEXT("T") : TEXT("F"), CurrentData.GetGPUInstanceCountBufferOffset());
-			SetShaderValue(RHICmdList, ComputeShader, Shader->ReadInstanceCountOffsetParam, ReadOffset);
-			SetShaderValue(RHICmdList, ComputeShader, Shader->WriteInstanceCountOffsetParam, WriteOffset);
-		}
-	}
-
-	// set the execution parameters
-	//
-	SetShaderValue(RHICmdList, ComputeShader, Shader->EmitterTickCounterParam, FNiagaraComputeExecutionContext::TickCounter);
-
-	// set spawn info
-	//
-	static_assert((sizeof(SpawnInfo.SpawnInfoStartOffsets) % SHADER_PARAMETER_ARRAY_ELEMENT_ALIGNMENT) == 0, "sizeof SpawnInfoStartOffsets should be a multiple of SHADER_PARAMETER_ARRAY_ELEMENT_ALIGNMENT");
-	static_assert((sizeof(SpawnInfo.SpawnInfoParams) % SHADER_PARAMETER_ARRAY_ELEMENT_ALIGNMENT) == 0, "sizeof SpawnInfoParams should be a multiple of SHADER_PARAMETER_ARRAY_ELEMENT_ALIGNMENT");
-	SetShaderValueArray(RHICmdList, ComputeShader, Shader->EmitterSpawnInfoOffsetsParam, SpawnInfo.SpawnInfoStartOffsets, NIAGARA_MAX_GPU_SPAWN_INFOS);
-	// This parameter is an array of structs with 2 floats and 2 ints on CPU, but a float4 array on GPU. The shader uses asint() to cast the integer values. To set the parameter,
-	// we pass the structure array as a float* to SetShaderValueArray() and specify the number of floats (not float vectors).
-	SetShaderValueArray(RHICmdList, ComputeShader, Shader->EmitterSpawnInfoParamsParam, &SpawnInfo.SpawnInfoParams[0].IntervalDt, 4*NIAGARA_MAX_GPU_SPAWN_INFOS);
-
-	SetShaderValue(RHICmdList, ComputeShader, Shader->UpdateStartInstanceParam, UpdateStartInstance);					// 0, except for event handler runs
-	SetShaderValue(RHICmdList, ComputeShader, Shader->NumSpawnedInstancesParam, InstancesToSpawnThisFrame);				// number of instances in the spawn run
-	SetShaderValue(RHICmdList, ComputeShader, Shader->DefaultSimulationStageIndexParam, DefaultSimulationStageIndex);					// 0, except if several stages are defined
-	SetShaderValue(RHICmdList, ComputeShader, Shader->SimulationStageIndexParam, SimulationStageIndex);					// 0, except if several stages are defined
-
-	{
-		// Packed data where
-		// X = Count Buffer Instance Count Offset (INDEX_NONE == Use Instance Count)
-		// Y = Instance Count
-		// Z = Iteration Index
-		// W = Num Iterations
-		FIntVector4 SimulationStageIterationInfo = { INDEX_NONE, -1, 0, 0 };
-		float SimulationStageNormalizedIterationIndex = 0.0f;
-
-		if (IterationInterface != nullptr)
-		{
-			const uint32 IterationInstanceCountOffset = IterationInterface->GetGPUInstanceCountOffset(Tick.SystemInstanceID);
-			SimulationStageIterationInfo.X = IterationInstanceCountOffset;
-			SimulationStageIterationInfo.Y = IterationInstanceCountOffset == INDEX_NONE ? TotalNumInstances : 0;
-		}
-
-		if (const FSimulationStageMetaData* StageMetaData = Instance->SimStageData[SimulationStageIndex].StageMetaData)
-		{
-			const int32 NumStages = StageMetaData->MaxStage - StageMetaData->MinStage;
-			ensure((int32(SimulationStageIndex) >= StageMetaData->MinStage) && (int32(SimulationStageIndex) < StageMetaData->MaxStage));
-			const int32 IterationIndex = SimulationStageIndex - StageMetaData->MinStage;
-			SimulationStageIterationInfo.Z = SimulationStageIndex - StageMetaData->MinStage;
-			SimulationStageIterationInfo.W = NumStages;
-			SimulationStageNormalizedIterationIndex = NumStages > 1 ? float(IterationIndex) / float(NumStages - 1) : 1.0f;
-		}
-
-		SetShaderValue(RHICmdList, ComputeShader, Shader->SimulationStageIterationInfoParam, SimulationStageIterationInfo);
-		SetShaderValue(RHICmdList, ComputeShader, Shader->SimulationStageNormalizedIterationIndexParam, SimulationStageNormalizedIterationIndex);
-	}
-
-	// Execute the dispatch
-	{
-		const uint32 ThreadGroupSize = FNiagaraShader::GetGroupSize(ShaderPlatform);
-		checkf(ThreadGroupSize <= NiagaraComputeMaxThreadGroupSize, TEXT("ShaderPlatform(%d) is set to ThreadGroup size (%d) which is > the NiagaraComputeMaxThreadGroupSize(%d)"), ShaderPlatform, ThreadGroupSize, NiagaraComputeMaxThreadGroupSize);
-		const uint32 TotalThreadGroups = FMath::DivideAndRoundUp(TotalNumInstances, ThreadGroupSize);
-
-		const uint32 ThreadGroupCountZ = 1;
-		const uint32 ThreadGroupCountY = FMath::DivideAndRoundUp(TotalThreadGroups, NiagaraMaxThreadGroupCountPerDimension);
-		const uint32 ThreadGroupCountX = FMath::DivideAndRoundUp(TotalThreadGroups, ThreadGroupCountY);
-
-		SetShaderValue(RHICmdList, ComputeShader, Shader->DispatchThreadIdToLinearParam, ThreadGroupCountY > 1 ? ThreadGroupCountX * ThreadGroupSize : 0);
-
-		SetConstantBuffers(RHICmdList, Shader, Tick, Instance);
-
-		DispatchComputeShader(RHICmdList, Shader.GetShader(), ThreadGroupCountX, ThreadGroupCountY, ThreadGroupCountZ);
-
-		INC_DWORD_STAT(STAT_NiagaraGPUDispatches);
-	}
-
-	// Unset UAV parameters.
-	//
-	UnsetDataInterfaceParameters(DataInterfaceProxies, Shader, RHICmdList, Instance, Tick, SimulationStageIndex);
-	FNiagaraDataBuffer::UnsetShaderParams(RHICmdList, Shader.GetShader());
-	Shader->InstanceCountsParam.UnsetUAV(RHICmdList, ComputeShader);
-
-	// Optionally submit commands to the GPU
-	// This can be used to avoid accidental TDR detection in the editor especially when issuing multiple ticks in the same frame
-	if (GNiagaraGpuSubmitCommandHint > 0)
-	{
-		++TotalDispatchesThisFrame;
-		if ((TotalDispatchesThisFrame % GNiagaraGpuSubmitCommandHint) == 0)
-		{
-			RHICmdList.SubmitCommandsHint();
-		}
-	}
-}
-
 FGPUSortManager* NiagaraEmitterInstanceBatcher::GetGPUSortManager() const
 {
 	return GPUSortManager;
diff --git a/Engine/Plugins/FX/Niagara/Source/Niagara/Private/NiagaraGPUSystemTick.cpp b/Engine/Plugins/FX/Niagara/Source/Niagara/Private/NiagaraGPUSystemTick.cpp
new file mode 100644
--- /dev/null
+++ b/Engine/Plugins/FX/Niagara/Source/Niagara/Private/NiagaraGPUSystemTick.cpp
@@ -0,0 +1,281 @@
+// Copyright Epic Games, Inc. All Rights Reserved.
+
+#include "NiagaraGPUSystemTick.h"
+#include "NiagaraSystemGpuComputeProxy.h"
+#include "NiagaraSystemInstance.h"
+
+void FNiagaraGPUSystemTick::Init(FNiagaraSystemInstance* InSystemInstance)
+{
+	ensure(InSystemInstance != nullptr);
+	CA_ASSUME(InSystemInstance != nullptr);
+	ensure(!InSystemInstance->IsComplete());
+	SystemInstanceID = InSystemInstance->GetId();
+	SystemGpuComputeProxy = InSystemInstance->GetSystemGpuComputeProxy();
+ 
+	uint32 DataSizeForGPU = InSystemInstance->GPUDataInterfaceInstanceDataSize;
+
+	if (DataSizeForGPU > 0)
+	{
+		uint32 AllocationSize = DataSizeForGPU;
+
+		DIInstanceData = new FNiagaraComputeDataInterfaceInstanceData;
+		DIInstanceData->PerInstanceDataSize = AllocationSize;
+		DIInstanceData->PerInstanceDataForRT = FMemory::Malloc(AllocationSize);
+		DIInstanceData->Instances = InSystemInstance->DataInterfaceInstanceDataOffsets.Num();
+
+		uint8* InstanceDataBase = (uint8*) DIInstanceData->PerInstanceDataForRT;
+		uint32 RunningOffset = 0;
+
+		DIInstanceData->InterfaceProxiesToOffsets.Reserve(InSystemInstance->GPUDataInterfaces.Num());
+
+		for (const auto& Pair : InSystemInstance->GPUDataInterfaces)
+		{
+			UNiagaraDataInterface* Interface = Pair.Key.Get();
+			if (Interface == nullptr)
+			{
+				continue;
+			}
+
+			FNiagaraDataInterfaceProxy* Proxy = Interface->GetProxy();
+			const int32 Offset = Pair.Value;
+
+			const int32 RTDataSize = Interface->PerInstanceDataPassedToRenderThreadSize();
+			ensure(RTDataSize > 0);
+			check(Proxy);
+
+			void* PerInstanceData = &InSystemInstance->DataInterfaceInstanceData[Offset];
+
+			Interface->ProvidePerInstanceDataForRenderThread(InstanceDataBase, PerInstanceData, SystemInstanceID);
+
+			// @todo rethink this. So ugly.
+			DIInstanceData->InterfaceProxiesToOffsets.Add(Proxy, RunningOffset);
+
+			InstanceDataBase += RTDataSize;
+			RunningOffset += RTDataSize;
+		}
+	}
+
+	check(MAX_uint32 > InSystemInstance->ActiveGPUEmitterCount);
+
+	// Layout our packet.
+	const uint32 PackedDispatchesSize = InSystemInstance->ActiveGPUEmitterCount * sizeof(FNiagaraComputeInstanceData);
+	// We want the Params after the instance data to be aligned so we can upload to the gpu.
+	uint32 PackedDispatchesSizeAligned = Align(PackedDispatchesSize, SHADER_PARAMETER_STRUCT_ALIGNMENT);
+	uint32 TotalParamSize = InSystemInstance->TotalGPUParamSize;
+
+	uint32 TotalPackedBufferSize = PackedDispatchesSizeAligned + TotalParamSize;
+
+	InstanceData_ParamData_Packed = (uint8*)FMemory::Malloc(TotalPackedBufferSize);
+
+	FNiagaraComputeInstanceData* Instances = (FNiagaraComputeInstanceData*)(InstanceData_ParamData_Packed);
+	uint8* ParamDataBufferPtr = InstanceData_ParamData_Packed + PackedDispatchesSizeAligned;
+
+	NumInstancesWithSimStages = 0;
+
+	TotalDispatches = 0;
+
+	// we want to include interpolation parameters (current and previous frame) if any of the emitters in the system
+	// require it
+	const bool IncludeInterpolationParameters = InSystemInstance->GPUParamIncludeInterpolation;
+	const int32 InterpFactor = IncludeInterpolationParameters ? 2 : 1;
+
+	GlobalParamData = ParamDataBufferPtr;
+	SystemParamData = GlobalParamData + InterpFactor * sizeof(FNiagaraGlobalParameters);
+	OwnerParamData = SystemParamData + InterpFactor * sizeof(FNiagaraSystemParameters);
+
+	// actually copy all of the data over, for the system data we only need to do it once (rather than per-emitter)
+	FMemory::Memcpy(GlobalParamData, &InSystemInstance->GetGlobalParameters(), sizeof(FNiagaraGlobalParameters));
+	FMemory::Memcpy(SystemParamData, &InSystemInstance->GetSystemParameters(), sizeof(FNiagaraSystemParameters));
+	FMemory::Memcpy(OwnerParamData, &InSystemInstance->GetOwnerParameters(), sizeof(FNiagaraOwnerParameters));
+
+	if (IncludeInterpolationParameters)
+	{
+		FMemory::Memcpy(GlobalParamData + sizeof(FNiagaraGlobalParameters), &InSystemInstance->GetGlobalParameters(true), sizeof(FNiagaraGlobalParameters));
+		FMemory::Memcpy(SystemParamData + sizeof(FNiagaraSystemParameters), &InSystemInstance->GetSystemParameters(true), sizeof(FNiagaraSystemParameters));
+		FMemory::Memcpy(OwnerParamData + sizeof(FNiagaraOwnerParameters), &InSystemInstance->GetOwnerParameters(true), sizeof(FNiagaraOwnerParameters));
+	}
+
+	ParamDataBufferPtr = OwnerParamData + InterpFactor * sizeof(FNiagaraOwnerParameters);
+
+	// Now we will generate instance data for every GPU simulation we want to run on the render thread.
+	// This is spawn rate as well as DataInterface per instance data and the ParameterData for the emitter.
+	// @todo Ideally we would only update DataInterface and ParameterData bits if they have changed.
+	uint32 InstanceIndex = 0;
+	bool bStartNewOverlapGroup = false;
+
+	const TConstArrayView<FNiagaraEmitterExecutionIndex> EmitterExecutionOrder = InSystemInstance->GetEmitterExecutionOrder();
+	for (const FNiagaraEmitterExecutionIndex& EmiterExecIndex : EmitterExecutionOrder)
+	{
+		// The dependency resolution code does not consider CPU and GPU emitters separately, so the flag which marks the start of a new overlap group can be set on either
+		// a CPU or GPU emitter. We must turn on bStartNewOverlapGroup when we encounter the flag, and reset it when we've actually marked a GPU emitter as starting a new group.
+		bStartNewOverlapGroup |= EmiterExecIndex.bStartNewOverlapGroup;
+
+		const uint32 EmitterIdx = EmiterExecIndex.EmitterIndex;
+		if (FNiagaraEmitterInstance* EmitterInstance = &InSystemInstance->GetEmitters()[EmitterIdx].Get())
+		{
+			if (EmitterInstance->IsComplete() )
+			{
+				continue;
+			}
+
+			const UNiagaraEmitter* Emitter = EmitterInstance->GetCachedEmitter();
+			FNiagaraComputeExecutionContext* GPUContext = EmitterInstance->GetGPUContext();
+
+			check(Emitter);
+
+			if (!Emitter || !GPUContext || Emitter->SimTarget != ENiagaraSimTarget::GPUComputeSim)
+			{
+				continue;
+			}
+
+			// Handle edge case where an emitter was set to inactive on the first frame by scalability
+			// In which case it will never have ticked so we should not execute a GPU tick for this until it becomes active
+			// See FNiagaraSystemInstance::Tick_Concurrent for details
+			if (EmitterInstance->HasTicked() == false)
+			{
+				ensure((EmitterInstance->GetExecutionState() == ENiagaraExecutionState::Inactive) || (EmitterInstance->GetExecutionState() == ENiagaraExecutionState::InactiveClear));
+				continue;
+			}
+
+			FNiagaraComputeInstanceData* InstanceData = new (&Instances[InstanceIndex]) FNiagaraComputeInstanceData;
+			InstanceIndex++;
+
+			InstanceData->Context = GPUContext;
+			check(GPUContext->MainDataSet);
+
+			InstanceData->SpawnInfo = GPUContext->GpuSpawnInfo_GT;
+
+			// Consume pending reset
+			if ( GPUContext->bResetPending_GT )
+			{
+				InstanceData->bResetData = GPUContext->bResetPending_GT;
+				GPUContext->bResetPending_GT = false;
+
+				++GPUContext->ParticleCountReadFence;
+			}
+			InstanceData->ParticleCountFence = GPUContext->ParticleCountReadFence;
+
+			int32 ParmSize = GPUContext->CombinedParamStore.GetPaddedParameterSizeInBytes();
+
+			InstanceData->EmitterParamData = ParamDataBufferPtr;
+			ParamDataBufferPtr += InterpFactor * sizeof(FNiagaraEmitterParameters);
+
+			InstanceData->ExternalParamData = ParamDataBufferPtr;
+			ParamDataBufferPtr += ParmSize;
+
+			// actually copy all of the data over
+			FMemory::Memcpy(InstanceData->EmitterParamData, &InSystemInstance->GetEmitterParameters(EmitterIdx), sizeof(FNiagaraEmitterParameters));
+			if (IncludeInterpolationParameters)
+			{
+				FMemory::Memcpy(InstanceData->EmitterParamData + sizeof(FNiagaraEmitterParameters), &InSystemInstance->GetEmitterParameters(EmitterIdx, true), sizeof(FNiagaraEmitterParameters));
+			}
+
+			GPUContext->CombinedParamStore.CopyParameterDataToPaddedBuffer(InstanceData->ExternalParamData, ParmSize);
+
+			InstanceData->bStartNewOverlapGroup = bStartNewOverlapGroup;
+			bStartNewOverlapGroup = false;
+
+			InstanceData->bUsesSimStages = Emitter->bSimulationStagesEnabled/* TODO limit to just with stages in the future! Leaving like this so what can convert! && EmitterRaw->GetSimulationStages().Num() > 0*/;
+			InstanceData->bUsesOldShaderStages = Emitter->bDeprecatedShaderStagesEnabled;
+
+			if (InstanceData->bUsesSimStages || InstanceData->bUsesOldShaderStages)
+			{
+				NumInstancesWithSimStages++;
+			}
+
+			check(GPUContext->MaxUpdateIterations > 0);
+			TotalDispatches += FMath::Max<int32>(GPUContext->MaxUpdateIterations, 1);
+
+			// @todo-threadsafety Think of a better way to do this!
+			const TArray<UNiagaraDataInterface*>& DataInterfaces = GPUContext->CombinedParamStore.GetDataInterfaces();
+			InstanceData->DataInterfaceProxies.Reserve(DataInterfaces.Num());
+			InstanceData->IterationDataInterfaceProxies.Reserve(DataInterfaces.Num());
+
+			for (UNiagaraDataInterface* DI : DataInterfaces)
+			{
+				FNiagaraDataInterfaceProxy* DIProxy = DI->GetProxy();
+				check(DIProxy);
+				InstanceData->DataInterfaceProxies.Add(DIProxy);
+
+				if ( FNiagaraDataInterfaceProxyRW* RWProxy = DIProxy->AsIterationProxy() )
+				{
+					InstanceData->IterationDataInterfaceProxies.Add(RWProxy);
+				}
+			}
+		}
+	}
+
+	check(InSystemInstance->ActiveGPUEmitterCount == InstanceIndex);
+	InstanceCount = InstanceIndex;
+#if !UE_BUILD_SHIPPING && !UE_BUILD_TEST
+	InstanceDataDebuggingOnly = GetInstances();
+#endif
+}
+
+void FNiagaraGPUSystemTick::Destroy()
+{
+	for ( FNiagaraComputeInstanceData& Instance : GetInstances() )
+	{
+		Instance.Context->ParticleCountWriteFence = Instance.ParticleCountFence;
+		Instance.~FNiagaraComputeInstanceData();
+	}
+
+	FMemory::Free(InstanceData_ParamData_Packed);
+	if (DIInstanceData)
+	{
+		FMemory::Free(DIInstanceData->PerInstanceDataForRT);
+		delete DIInstanceData;
+	}
+}
+
+FUniformBufferRHIRef FNiagaraGPUSystemTick::GetUniformBuffer(EUniformBufferType Type, const FNiagaraComputeInstanceData* Instance, bool Current) const
+{
+	const int32 InterpOffset = Current
+		? 0
+		: (UBT_NumSystemTypes + InstanceCount * UBT_NumInstanceTypes);
+
+	if (Instance)
+	{
+		check(Type >= UBT_FirstInstanceType);
+		check(Type < UBT_NumTypes);
+
+		const int32 InstanceTypeIndex = Type - UBT_FirstInstanceType;
+
+		const int32 InstanceIndex = (Instance - GetInstances().GetData());
+		return UniformBuffers[InterpOffset + UBT_NumSystemTypes + InstanceCount * InstanceTypeIndex + InstanceIndex];
+	}
+
+	check(Type >= UBT_FirstSystemType);
+	check(Type < UBT_FirstInstanceType);
+
+	return UniformBuffers[InterpOffset + Type];
+}
+
+const uint8* FNiagaraGPUSystemTick::GetUniformBufferSource(EUniformBufferType Type, const FNiagaraComputeInstanceData* Instance, bool Current) const
+{
+	check(Type >= UBT_FirstSystemType);
+	check(Type < UBT_NumTypes);
+
+	switch (Type)
+	{
+		case UBT_Global:
+			return GlobalParamData + (Current ? 0 : sizeof(FNiagaraGlobalParameters));
+		case UBT_System:
+			return SystemParamData + (Current ? 0 : sizeof(FNiagaraSystemParameters));
+		case UBT_Owner:
+			return OwnerParamData + (Current ? 0 : sizeof(FNiagaraOwnerParameters));
+		case UBT_Emitter:
+		{
+			check(Instance);
+			return Instance->EmitterParamData + (Current ? 0 : sizeof(FNiagaraEmitterParameters));
+		}
+		case UBT_External:
+		{
+			// External is special and interpolated parameters are already included inside of the combined parameter store
+			check(Instance);
+			return Instance->ExternalParamData;
+		}
+	}
+
+	return nullptr;
+}
diff --git a/Engine/Plugins/FX/Niagara/Source/Niagara/Private/NiagaraScriptExecutionContext.cpp b/Engine/Plugins/FX/Niagara/Source/Niagara/Private/NiagaraScriptExecutionContext.cpp
--- a/Engine/Plugins/FX/Niagara/Source/Niagara/Private/NiagaraScriptExecutionContext.cpp
+++ b/Engine/Plugins/FX/Niagara/Source/Niagara/Private/NiagaraScriptExecutionContext.cpp
@@ -5,6 +5,7 @@
 #include "NiagaraEmitterInstanceBatcher.h"
 #include "NiagaraDataInterface.h"
 #include "NiagaraSystemInstance.h"
+#include "NiagaraSystemGpuComputeProxy.h"
 #include "NiagaraWorldManager.h"
 #include "NiagaraEmitterInstance.h"
 #include "NiagaraFunctionLibrary.h"
@@ -570,704 +571,3 @@ bool FNiagaraSystemScriptExecutionContext::GeneratePerInstanceDIFunctionTable(FN
 	return true;
 };
 
-//////////////////////////////////////////////////////////////////////////
-
-void FNiagaraGPUSystemTick::Init(FNiagaraSystemInstance* InSystemInstance)
-{
-	ensure(InSystemInstance != nullptr);
-	CA_ASSUME(InSystemInstance != nullptr);
-	ensure(!InSystemInstance->IsComplete());
-	SystemInstanceID = InSystemInstance->GetId();
-	SharedContext = InSystemInstance->GetComputeSharedContext();
-	bRequiresDistanceFieldData = InSystemInstance->RequiresDistanceFieldData();
-	bRequiresDepthBuffer = InSystemInstance->RequiresDepthBuffer();
-	bRequiresEarlyViewData = InSystemInstance->RequiresEarlyViewData();
-	bRequiresViewUniformBuffer = InSystemInstance->RequiresViewUniformBuffer();
-	bRequiresRayTracingScene = InSystemInstance->RequiresRayTracingScene();
- 
-	uint32 DataSizeForGPU = InSystemInstance->GPUDataInterfaceInstanceDataSize;
-
-	if (DataSizeForGPU > 0)
-	{
-		uint32 AllocationSize = DataSizeForGPU;
-
-		DIInstanceData = new FNiagaraDataInterfaceInstanceData;
-		DIInstanceData->PerInstanceDataSize = AllocationSize;
-		DIInstanceData->PerInstanceDataForRT = FMemory::Malloc(AllocationSize);
-		DIInstanceData->Instances = InSystemInstance->DataInterfaceInstanceDataOffsets.Num();
-
-		uint8* InstanceDataBase = (uint8*) DIInstanceData->PerInstanceDataForRT;
-		uint32 RunningOffset = 0;
-
-		DIInstanceData->InterfaceProxiesToOffsets.Reserve(InSystemInstance->GPUDataInterfaces.Num());
-
-		for (const auto& Pair : InSystemInstance->GPUDataInterfaces)
-		{
-			UNiagaraDataInterface* Interface = Pair.Key.Get();
-			if (Interface == nullptr)
-			{
-				continue;
-			}
-
-			FNiagaraDataInterfaceProxy* Proxy = Interface->GetProxy();
-			const int32 Offset = Pair.Value;
-
-			const int32 RTDataSize = Interface->PerInstanceDataPassedToRenderThreadSize();
-			ensure(RTDataSize > 0);
-			check(Proxy);
-
-			void* PerInstanceData = &InSystemInstance->DataInterfaceInstanceData[Offset];
-
-			Interface->ProvidePerInstanceDataForRenderThread(InstanceDataBase, PerInstanceData, SystemInstanceID);
-
-			// @todo rethink this. So ugly.
-			DIInstanceData->InterfaceProxiesToOffsets.Add(Proxy, RunningOffset);
-
-			InstanceDataBase += RTDataSize;
-			RunningOffset += RTDataSize;
-		}
-	}
-
-	check(MAX_uint32 > InSystemInstance->ActiveGPUEmitterCount);
-
-	// Layout our packet.
-	const uint32 PackedDispatchesSize = InSystemInstance->ActiveGPUEmitterCount * sizeof(FNiagaraComputeInstanceData);
-	// We want the Params after the instance data to be aligned so we can upload to the gpu.
-	uint32 PackedDispatchesSizeAligned = Align(PackedDispatchesSize, SHADER_PARAMETER_STRUCT_ALIGNMENT);
-	uint32 TotalParamSize = InSystemInstance->TotalGPUParamSize;
-
-	uint32 TotalPackedBufferSize = PackedDispatchesSizeAligned + TotalParamSize;
-
-	InstanceData_ParamData_Packed = (uint8*)FMemory::Malloc(TotalPackedBufferSize);
-
-	FNiagaraComputeInstanceData* Instances = (FNiagaraComputeInstanceData*)(InstanceData_ParamData_Packed);
-	uint8* ParamDataBufferPtr = InstanceData_ParamData_Packed + PackedDispatchesSizeAligned;
-
-	bNeedsReset = InSystemInstance->RequiresGpuBufferReset();
-	NumInstancesWithSimStages = 0;
-
-	if ( bNeedsReset )
-	{
-		++SharedContext->ParticleCountReadFence;
-	}
-	ParticleCountFence = SharedContext->ParticleCountReadFence;
-
-	TotalDispatches = 0;
-
-	// we want to include interpolation parameters (current and previous frame) if any of the emitters in the system
-	// require it
-	const bool IncludeInterpolationParameters = InSystemInstance->GPUParamIncludeInterpolation;
-	const int32 InterpFactor = IncludeInterpolationParameters ? 2 : 1;
-
-	GlobalParamData = ParamDataBufferPtr;
-	SystemParamData = GlobalParamData + InterpFactor * sizeof(FNiagaraGlobalParameters);
-	OwnerParamData = SystemParamData + InterpFactor * sizeof(FNiagaraSystemParameters);
-
-	// actually copy all of the data over, for the system data we only need to do it once (rather than per-emitter)
-	FMemory::Memcpy(GlobalParamData, &InSystemInstance->GetGlobalParameters(), sizeof(FNiagaraGlobalParameters));
-	FMemory::Memcpy(SystemParamData, &InSystemInstance->GetSystemParameters(), sizeof(FNiagaraSystemParameters));
-	FMemory::Memcpy(OwnerParamData, &InSystemInstance->GetOwnerParameters(), sizeof(FNiagaraOwnerParameters));
-
-	if (IncludeInterpolationParameters)
-	{
-		FMemory::Memcpy(GlobalParamData + sizeof(FNiagaraGlobalParameters), &InSystemInstance->GetGlobalParameters(true), sizeof(FNiagaraGlobalParameters));
-		FMemory::Memcpy(SystemParamData + sizeof(FNiagaraSystemParameters), &InSystemInstance->GetSystemParameters(true), sizeof(FNiagaraSystemParameters));
-		FMemory::Memcpy(OwnerParamData + sizeof(FNiagaraOwnerParameters), &InSystemInstance->GetOwnerParameters(true), sizeof(FNiagaraOwnerParameters));
-	}
-
-	ParamDataBufferPtr = OwnerParamData + InterpFactor * sizeof(FNiagaraOwnerParameters);
-
-	// Now we will generate instance data for every GPU simulation we want to run on the render thread.
-	// This is spawn rate as well as DataInterface per instance data and the ParameterData for the emitter.
-	// @todo Ideally we would only update DataInterface and ParameterData bits if they have changed.
-	uint32 InstanceIndex = 0;
-	bool bStartNewOverlapGroup = false;
-
-	const TConstArrayView<FNiagaraEmitterExecutionIndex> EmitterExecutionOrder = InSystemInstance->GetEmitterExecutionOrder();
-	for (const FNiagaraEmitterExecutionIndex& EmiterExecIndex : EmitterExecutionOrder)
-	{
-		// The dependency resolution code does not consider CPU and GPU emitters separately, so the flag which marks the start of a new overlap group can be set on either
-		// a CPU or GPU emitter. We must turn on bStartNewOverlapGroup when we encounter the flag, and reset it when we've actually marked a GPU emitter as starting a new group.
-		bStartNewOverlapGroup |= EmiterExecIndex.bStartNewOverlapGroup;
-
-		const uint32 EmitterIdx = EmiterExecIndex.EmitterIndex;
-		if (FNiagaraEmitterInstance* EmitterInstance = &InSystemInstance->GetEmitters()[EmitterIdx].Get())
-		{
-			if (EmitterInstance->IsComplete() )
-			{
-				continue;
-			}
-
-			const UNiagaraEmitter* Emitter = EmitterInstance->GetCachedEmitter();
-			FNiagaraComputeExecutionContext* GPUContext = EmitterInstance->GetGPUContext();
-
-			check(Emitter);
-
-			if (!Emitter || !GPUContext || Emitter->SimTarget != ENiagaraSimTarget::GPUComputeSim)
-			{
-				continue;
-			}
-
-			// Handle edge case where an emitter was set to inactive on the first frame by scalability
-			// In which case it will never have ticked so we should not execute a GPU tick for this until it becomes active
-			// See FNiagaraSystemInstance::Tick_Concurrent for details
-			if (EmitterInstance->HasTicked() == false)
-			{
-				ensure((EmitterInstance->GetExecutionState() == ENiagaraExecutionState::Inactive) || (EmitterInstance->GetExecutionState() == ENiagaraExecutionState::InactiveClear));
-				continue;
-			}
-
-			FNiagaraComputeInstanceData* InstanceData = new (&Instances[InstanceIndex]) FNiagaraComputeInstanceData;
-			InstanceIndex++;
-
-			InstanceData->Context = GPUContext;
-			check(GPUContext->MainDataSet);
-
-			InstanceData->SpawnInfo = GPUContext->GpuSpawnInfo_GT;
-
-			int32 ParmSize = GPUContext->CombinedParamStore.GetPaddedParameterSizeInBytes();
-
-			InstanceData->EmitterParamData = ParamDataBufferPtr;
-			ParamDataBufferPtr += InterpFactor * sizeof(FNiagaraEmitterParameters);
-
-			InstanceData->ExternalParamData = ParamDataBufferPtr;
-			ParamDataBufferPtr += ParmSize;
-
-			// actually copy all of the data over
-			FMemory::Memcpy(InstanceData->EmitterParamData, &InSystemInstance->GetEmitterParameters(EmitterIdx), sizeof(FNiagaraEmitterParameters));
-			if (IncludeInterpolationParameters)
-			{
-				FMemory::Memcpy(InstanceData->EmitterParamData + sizeof(FNiagaraEmitterParameters), &InSystemInstance->GetEmitterParameters(EmitterIdx, true), sizeof(FNiagaraEmitterParameters));
-			}
-
-			GPUContext->CombinedParamStore.CopyParameterDataToPaddedBuffer(InstanceData->ExternalParamData, ParmSize);
-
-			InstanceData->bStartNewOverlapGroup = bStartNewOverlapGroup;
-			bStartNewOverlapGroup = false;
-
-			InstanceData->bUsesSimStages = Emitter->bSimulationStagesEnabled/* TODO limit to just with stages in the future! Leaving like this so what can convert! && EmitterRaw->GetSimulationStages().Num() > 0*/;
-			InstanceData->bUsesOldShaderStages = Emitter->bDeprecatedShaderStagesEnabled;
-
-			if (InstanceData->bUsesSimStages || InstanceData->bUsesOldShaderStages)
-			{
-				NumInstancesWithSimStages++;
-			}
-
-			check(GPUContext->MaxUpdateIterations > 0);
-			InstanceData->SimStageData.AddZeroed(GPUContext->MaxUpdateIterations);
-			TotalDispatches += FMath::Max<int32>(GPUContext->MaxUpdateIterations, 1);
-
-			// @todo-threadsafety Think of a better way to do this!
-			const TArray<UNiagaraDataInterface*>& DataInterfaces = GPUContext->CombinedParamStore.GetDataInterfaces();
-			InstanceData->DataInterfaceProxies.Reserve(DataInterfaces.Num());
-			InstanceData->IterationDataInterfaceProxies.Reserve(DataInterfaces.Num());
-
-			for (UNiagaraDataInterface* DI : DataInterfaces)
-			{
-				FNiagaraDataInterfaceProxy* DIProxy = DI->GetProxy();
-				check(DIProxy);
-				InstanceData->DataInterfaceProxies.Add(DIProxy);
-
-				if ( FNiagaraDataInterfaceProxyRW* RWProxy = DIProxy->AsIterationProxy() )
-				{
-					InstanceData->IterationDataInterfaceProxies.Add(RWProxy);
-				}
-			}
-		}
-	}
-
-	check(InSystemInstance->ActiveGPUEmitterCount == InstanceIndex);
-	Count = InstanceIndex;
-}
-
-void FNiagaraGPUSystemTick::Destroy()
-{
-	SharedContext->ParticleCountWriteFence = ParticleCountFence;
-
-	FNiagaraComputeInstanceData* Instances = GetInstanceData();
-	for (uint32 i = 0; i < Count; i++)
-	{
-		FNiagaraComputeInstanceData& Instance = Instances[i];
-		Instance.~FNiagaraComputeInstanceData();
-	}
-
-	FMemory::Free(InstanceData_ParamData_Packed);
-	if (DIInstanceData)
-	{
-		FMemory::Free(DIInstanceData->PerInstanceDataForRT);
-		delete DIInstanceData;
-	}
-}
-
-FUniformBufferRHIRef FNiagaraGPUSystemTick::GetUniformBuffer(EUniformBufferType Type, const FNiagaraComputeInstanceData* Instance, bool Current) const
-{
-	const int32 InterpOffset = Current
-		? 0
-		: (UBT_NumSystemTypes + Count * UBT_NumInstanceTypes);
-
-	if (Instance)
-	{
-		check(Type >= UBT_FirstInstanceType);
-		check(Type < UBT_NumTypes);
-
-		const int32 InstanceTypeIndex = Type - UBT_FirstInstanceType;
-
-		const int32 InstanceIndex = (Instance - GetInstanceData());
-		return UniformBuffers[InterpOffset + UBT_NumSystemTypes + Count * InstanceTypeIndex + InstanceIndex];
-	}
-
-	check(Type >= UBT_FirstSystemType);
-	check(Type < UBT_FirstInstanceType);
-
-	return UniformBuffers[InterpOffset + Type];
-}
-
-const uint8* FNiagaraGPUSystemTick::GetUniformBufferSource(EUniformBufferType Type, const FNiagaraComputeInstanceData* Instance, bool Current) const
-{
-	check(Type >= UBT_FirstSystemType);
-	check(Type < UBT_NumTypes);
-
-	switch (Type)
-	{
-		case UBT_Global:
-			return GlobalParamData + (Current ? 0 : sizeof(FNiagaraGlobalParameters));
-		case UBT_System:
-			return SystemParamData + (Current ? 0 : sizeof(FNiagaraSystemParameters));
-		case UBT_Owner:
-			return OwnerParamData + (Current ? 0 : sizeof(FNiagaraOwnerParameters));
-		case UBT_Emitter:
-		{
-			check(Instance);
-			return Instance->EmitterParamData + (Current ? 0 : sizeof(FNiagaraEmitterParameters));
-		}
-		case UBT_External:
-		{
-			// External is special and interpolated parameters are already included inside of the combined parameter store
-			check(Instance);
-			return Instance->ExternalParamData;
-		}
-	}
-
-	return nullptr;
-}
-
-//////////////////////////////////////////////////////////////////////////
-
-FNiagaraComputeExecutionContext::FNiagaraComputeExecutionContext()
-	: MainDataSet(nullptr)
-	, GPUScript(nullptr)
-	, GPUScript_RT(nullptr)
-	, TranslucentDataToRender(nullptr)
-{
-	ExternalCBufferLayout = new FNiagaraRHIUniformBufferLayout(TEXT("Niagara GPU External CBuffer"));
-}
-
-FNiagaraComputeExecutionContext::~FNiagaraComputeExecutionContext()
-{
-	// EmitterInstanceReadback.GPUCountOffset should be INDEX_NONE at this point to ensure the index is reused.
-	// When the batcher is being destroyed though, we don't free the index, but this would not be leaking.
-	// check(EmitterInstanceReadback.GPUCountOffset == INDEX_NONE);
-	SetDataToRender(nullptr);
-
-	ExternalCBufferLayout = nullptr;
-}
-
-void FNiagaraComputeExecutionContext::Reset(NiagaraEmitterInstanceBatcher* Batcher)
-{
-	NiagaraEmitterInstanceBatcher* RT_Batcher = Batcher && !Batcher->IsPendingKill() ? Batcher : nullptr;
-	ENQUEUE_RENDER_COMMAND(ResetRT)(
-		[RT_Batcher, RT_Context=this](FRHICommandListImmediate& RHICmdList)
-		{
-			RT_Context->ResetInternal(RT_Batcher);
-		}
-	);
-}
-
-void FNiagaraComputeExecutionContext::InitParams(UNiagaraScript* InGPUComputeScript, ENiagaraSimTarget InSimTarget, const uint32 InDefaultSimulationStageIndex, const int32 InMaxUpdateIterations, const TSet<uint32> InSpawnStages)
-{
-	GPUScript = InGPUComputeScript;
-	CombinedParamStore.InitFromOwningContext(InGPUComputeScript, InSimTarget, true);
-	DefaultSimulationStageIndex = InDefaultSimulationStageIndex;
-	MaxUpdateIterations = InMaxUpdateIterations;
-	SpawnStages.Empty();
-
-	SpawnStages.Append(InSpawnStages);
-	
-	HasInterpolationParameters = GPUScript && GPUScript->GetComputedVMCompilationId().HasInterpolatedParameters();
-
-	if (InGPUComputeScript)
-	{
-		FNiagaraVMExecutableData& VMData = InGPUComputeScript->GetVMExecutableData();
-		if (VMData.IsValid() && VMData.SimulationStageMetaData.Num() > 0)
-		{
-			SimStageInfo = VMData.SimulationStageMetaData;
-
-			int32 FoundMaxUpdateIterations = SimStageInfo[SimStageInfo.Num() - 1].MaxStage;
-
-			// Some useful debugging code should we need to look up differences between old and new
-			const bool bDebugSimStages = false;
-			if (bDebugSimStages)
-			{
-				UE_LOG(LogNiagara, Log, TEXT("Stored vs:"));
-				bool bPass = FoundMaxUpdateIterations == MaxUpdateIterations;
-				UE_LOG(LogNiagara, Log, TEXT("MaxUpdateIterations: %d vs %d %s"), FoundMaxUpdateIterations, MaxUpdateIterations, bPass ? TEXT("Pass") : TEXT("FAIL!!!!!!!!"));
-
-				int32 NumSpawnFound = 0;
-				bool bMatchesFound = true;
-				for (int32 i = 0; i < SimStageInfo.Num(); i++)
-				{
-					if (SimStageInfo[i].bSpawnOnly)
-					{
-						NumSpawnFound++;
-
-						if (!SpawnStages.Contains(SimStageInfo[i].MinStage))
-						{
-							bMatchesFound = false;
-							UE_LOG(LogNiagara, Log, TEXT("Missing spawn stage: %d FAIL!!!!!!!!!"), SimStageInfo[i].MinStage);
-						}
-					}
-				}
-
-				bPass = SpawnStages.Num() == NumSpawnFound;
-				UE_LOG(LogNiagara, Log, TEXT("SpawnStages.Num(): %d vs %d %s"), NumSpawnFound, SpawnStages.Num(), bPass ? TEXT("Pass") : TEXT("FAIL!!!!!!!!"));
-
-				TArray<FNiagaraVariable> Params;
-				CombinedParamStore.GetParameters(Params);
-				for (FNiagaraVariable& Var : Params)
-				{
-					if (!Var.IsDataInterface())
-						continue;
-
-					UNiagaraDataInterface* DI = CombinedParamStore.GetDataInterface(Var);
-					UNiagaraDataInterfaceRWBase* DIRW = Cast<UNiagaraDataInterfaceRWBase>(DI);
-					if (DIRW)
-					{
-						for (int32 i = 0; i < SimStageInfo.Num(); i++)
-						{
-							if (SimStageInfo[i].IterationSource == Var.GetName())
-							{
-								if (!DIRW->IterationShaderStages.Contains(SimStageInfo[i].MinStage))
-								{
-									UE_LOG(LogNiagara, Log, TEXT("Missing iteration stage for %s: %d FAIL!!!!!!!!!"), *Var.GetName().ToString(), SimStageInfo[i].MinStage);
-								}
-							}
-
-							if (SimStageInfo[i].OutputDestinations.Contains(Var.GetName()))
-							{
-								if (!DIRW->OutputShaderStages.Contains(SimStageInfo[i].MinStage))
-								{
-									UE_LOG(LogNiagara, Log, TEXT("Missing output stage for %s: %d FAIL!!!!!!!!!"), *Var.GetName().ToString(), SimStageInfo[i].MinStage);
-								}
-							}
-						}
-					}
-
-				}
-			}
-
-
-			// Set the values that we are using from compiled data instead...
-			MaxUpdateIterations = SimStageInfo[SimStageInfo.Num() - 1].MaxStage;
-			SpawnStages.Empty();
-
-			for (int32 i = 0; i < SimStageInfo.Num(); i++)
-			{
-				if (SimStageInfo[i].bSpawnOnly)
-				{
-					SpawnStages.Add(SimStageInfo[i].MinStage);
-				}
-			}
-
-			
-		}
-
-	}
-
-	
-#if DO_CHECK
-	// DI Parameters are the same between all shader permutations so we can just get the first one
-	FNiagaraShaderRef Shader = InGPUComputeScript->GetRenderThreadScript()->GetShaderGameThread(0);
-	if (Shader.IsValid())
-	{
-		DIClassNames.Empty(Shader->GetDIParameters().Num());
-		for (const FNiagaraDataInterfaceParamRef& DIParams : Shader->GetDIParameters())
-		{
-			DIClassNames.Add(DIParams.DIType.Get(Shader.GetPointerTable().DITypes)->GetClass()->GetName());
-		}
-	}
-	else
-	{
-		DIClassNames.Empty(InGPUComputeScript->GetRenderThreadScript()->GetDataInterfaceParamInfo().Num());
-		for (const FNiagaraDataInterfaceGPUParamInfo& DIParams : InGPUComputeScript->GetRenderThreadScript()->GetDataInterfaceParamInfo())
-		{
-			DIClassNames.Add(DIParams.DIClassName);
-		}
-	}
-#endif
-}
-
-
-const FSimulationStageMetaData* FNiagaraComputeExecutionContext::GetSimStageMetaData(uint32 SimulationStageIndex) const
-{
-	if (SimStageInfo.Num() > 0)
-	{
-		for (int32 i = 0; i < SimStageInfo.Num(); i++)
-		{
-			if (SimulationStageIndex >= (uint32)SimStageInfo[i].MinStage  && SimulationStageIndex < (uint32)SimStageInfo[i].MaxStage)
-			{
-				return &SimStageInfo[i];
-			}
-		}
-	}
-	return nullptr;
-}
-
-bool FNiagaraComputeExecutionContext::IsOutputStage(FNiagaraDataInterfaceProxy* DIProxy, uint32 CurrentStage) const
-{
-	const FSimulationStageMetaData* MetaData = GetSimStageMetaData(CurrentStage);
-	if (MetaData && DIProxy && !DIProxy->SourceDIName.IsNone())
-	{
-		if (MetaData->OutputDestinations.Contains(DIProxy->SourceDIName))
-			return true;
-	}
-	else if (DIProxy && SimStageInfo.Num() == 0)
-	{
-		return DIProxy->IsOutputStage_DEPRECATED(CurrentStage);
-	}
-	return false;
-}
-
-bool FNiagaraComputeExecutionContext::IsIterationStage(FNiagaraDataInterfaceProxy* DIProxy, uint32 CurrentStage) const
-{
-	const FSimulationStageMetaData* MetaData = GetSimStageMetaData(CurrentStage);
-	if (MetaData && DIProxy && !DIProxy->SourceDIName.IsNone())
-	{
-		if (MetaData->IterationSource.IsNone()) // Per particle iteration...
-			return false;
-
-		if (MetaData->IterationSource == DIProxy->SourceDIName)
-			return true;
-	}
-	else if (DIProxy && SimStageInfo.Num() == 0)
-	{
-		return DIProxy->IsIterationStage_DEPRECATED(CurrentStage);
-	}
-	return false;
-}
-
-FNiagaraDataInterfaceProxyRW* FNiagaraComputeExecutionContext::FindIterationInterface(const TArray<FNiagaraDataInterfaceProxyRW*>& InProxies, uint32 CurrentStage) const
-{
-	const FSimulationStageMetaData* MetaData = GetSimStageMetaData(CurrentStage);
-	if (MetaData)
-	{
-		if (MetaData->IterationSource.IsNone()) // Per particle iteration...
-			return nullptr;
-
-		for (FNiagaraDataInterfaceProxyRW* Proxy : InProxies)
-		{
-			if (Proxy->SourceDIName == MetaData->IterationSource)
-				return Proxy;
-		}
-
-		UE_LOG(LogNiagara, Verbose, TEXT("FNiagaraComputeExecutionContext::FindIterationInterface could not find IterationInterface %s"), *MetaData->IterationSource.ToString());
-
-		return nullptr;
-	}
-	else if (SimStageInfo.Num() == 0)
-	{
-		// Fallback to old shader stages
-		for (FNiagaraDataInterfaceProxyRW* Proxy : InProxies)
-		{
-			if (Proxy->IsIterationStage_DEPRECATED(CurrentStage))
-				return Proxy;
-		}
-	}
-
-	return nullptr;
-}
-
-void FNiagaraComputeExecutionContext::DirtyDataInterfaces()
-{
-	CombinedParamStore.MarkInterfacesDirty();
-}
-
-bool FNiagaraComputeExecutionContext::OptionalContexInit(FNiagaraSystemInstance* ParentSystemInstance)
-{
-	if (GPUScript)
-	{
-		FNiagaraVMExecutableData& VMData = GPUScript->GetVMExecutableData();
-
-		if (VMData.IsValid() && VMData.bNeedsGPUContextInit)
-		{
-			const TArray<UNiagaraDataInterface*>& DataInterfaces = CombinedParamStore.GetDataInterfaces();
-			for (int32 i = 0; i < DataInterfaces.Num(); i++)
-			{
-				UNiagaraDataInterface* Interface = DataInterfaces[i];
-
-				int32 UserPtrIdx = VMData.DataInterfaceInfo[i].UserPtrIdx;
-				if (UserPtrIdx != INDEX_NONE)
-				{
-					void* InstData = ParentSystemInstance->FindDataInterfaceInstanceData(Interface);
-					if (Interface->NeedsGPUContextInit())
-					{
-						Interface->GPUContextInit(VMData.DataInterfaceInfo[i], InstData, ParentSystemInstance);
-					}
-				}
-			}
-		}
-	}
-	return true;
-}
-
-bool FNiagaraComputeExecutionContext::Tick(FNiagaraSystemInstance* ParentSystemInstance)
-{
-	if (CombinedParamStore.GetInterfacesDirty())
-	{
-#if DO_CHECK
-		const TArray<UNiagaraDataInterface*> &DataInterfaces = CombinedParamStore.GetDataInterfaces();
-		// We must make sure that the data interfaces match up between the original script values and our overrides...
-		if (DIClassNames.Num() != DataInterfaces.Num())
-		{
-			UE_LOG(LogNiagara, Warning, TEXT("Mismatch between Niagara GPU Execution Context data interfaces and those in its script!"));
-			return false;
-		}
-
-		for (int32 i = 0; i < DIClassNames.Num(); ++i)
-		{
-			FString UsedClassName = DataInterfaces[i]->GetClass()->GetName();
-			if (DIClassNames[i] != UsedClassName)
-			{
-				UE_LOG(LogNiagara, Warning, TEXT("Mismatched class between Niagara GPU Execution Context data interfaces and those in its script!\nIndex:%d\nShader:%s\nScript:%s")
-					, i, *DIClassNames[i], *UsedClassName);
-			}
-		}
-#endif
-
-		
-		CombinedParamStore.Tick();
-	}
-
-	return true;
-}
-
-void FNiagaraComputeExecutionContext::PostTick()
-{
-	//If we're for interpolated spawn, copy over the previous frame's parameters into the Prev parameters.
-	if (HasInterpolationParameters)
-	{
-		CombinedParamStore.CopyCurrToPrev();
-	}
-}
-
-void FNiagaraComputeExecutionContext::ResetInternal(NiagaraEmitterInstanceBatcher* Batcher)
-{
-	checkf(IsInRenderingThread(), TEXT("Can only reset the gpu context from the render thread"));
-
-	// Release and reset readback data.
-	if (Batcher)
-	{
-		Batcher->ReleaseInstanceCounts_RenderThread(this, MainDataSet);
-	}
-	else // In this case the batcher is pending kill so no need to putback entry in the pool.
-	{
-		EmitterInstanceReadback.GPUCountOffset = INDEX_NONE;
-	}
-
-	GpuSpawnInfo_GT.Reset();
-
-	SetDataToRender(nullptr);
-}
-
-void FNiagaraComputeExecutionContext::SetDataToRender(FNiagaraDataBuffer* InDataToRender)
-{
-	if (DataToRender)
-	{
-		DataToRender->ReleaseReadRef();
-	}
-
-	DataToRender = InDataToRender;
-
-	if (DataToRender)
-	{
-		DataToRender->AddReadRef();
-	}
-
-	// This call the DataToRender should be equal to the TranslucentDataToRender so we can release the read ref
-	if (TranslucentDataToRender)
-	{
-		ensure((DataToRender == nullptr) || (DataToRender == TranslucentDataToRender));
-		TranslucentDataToRender->ReleaseReadRef();
-		TranslucentDataToRender = nullptr;
-	}
-}
-
-void FNiagaraComputeExecutionContext::SetTranslucentDataToRender(FNiagaraDataBuffer* InTranslucentDataToRender)
-{
-	if (TranslucentDataToRender)
-	{
-		TranslucentDataToRender->ReleaseReadRef();
-	}
-
-	TranslucentDataToRender = InTranslucentDataToRender;
-
-	if (TranslucentDataToRender)
-	{
-		TranslucentDataToRender->AddReadRef();
-	}
-}
-
-bool FNiagaraComputeInstanceData::IsOutputStage(FNiagaraDataInterfaceProxy* DIProxy, uint32 CurrentStage) const
-{
-	if (bUsesOldShaderStages)
-	{
-		return DIProxy->IsOutputStage_DEPRECATED(CurrentStage);
-	}
-	else if (bUsesSimStages)
-	{
-		return Context->IsOutputStage(DIProxy, CurrentStage);
-	}
-	return false;
-}
-
-bool FNiagaraComputeInstanceData::IsIterationStage(FNiagaraDataInterfaceProxy* DIProxy, uint32 CurrentStage) const
-{
-
-	if (bUsesOldShaderStages)
-	{
-		return DIProxy->IsIterationStage_DEPRECATED(CurrentStage);
-	}
-	else if (bUsesSimStages)
-	{
-		return Context->IsIterationStage(DIProxy, CurrentStage);
-	}
-	return false;
-}
-
-FNiagaraDataInterfaceProxyRW* FNiagaraComputeInstanceData::FindIterationInterface(uint32 SimulationStageIndex) const
-{
-	if (bUsesOldShaderStages)
-	{
-		FNiagaraDataInterfaceProxyRW* IterationInterface = nullptr;
-		for (FNiagaraDataInterfaceProxyRW* Interface : IterationDataInterfaceProxies)
-		{
-			if (Interface->IsIterationStage_DEPRECATED(SimulationStageIndex))
-			{
-				if (IterationInterface)
-				{
-					UE_LOG(LogNiagara, Error, TEXT("Multiple output Data Interfaces found for current stage"));
-				}
-				else
-				{
-					IterationInterface = Interface;
-				}
-			}
-		}
-
-		return IterationInterface;
-	}
-	else if (bUsesSimStages)
-	{
-		return Context->FindIterationInterface(IterationDataInterfaceProxies, SimulationStageIndex);
-	}
-	return nullptr;
-}
diff --git a/Engine/Plugins/FX/Niagara/Source/Niagara/Private/NiagaraSystemGpuComputeProxy.cpp b/Engine/Plugins/FX/Niagara/Source/Niagara/Private/NiagaraSystemGpuComputeProxy.cpp
new file mode 100644
--- /dev/null
+++ b/Engine/Plugins/FX/Niagara/Source/Niagara/Private/NiagaraSystemGpuComputeProxy.cpp
@@ -0,0 +1,151 @@
+// Copyright Epic Games, Inc. All Rights Reserved.
+
+#include "NiagaraSystemGpuComputeProxy.h"
+#include "NiagaraSystemInstance.h"
+#include "NiagaraEmitterInstance.h"
+#include "NiagaraEmitterInstanceBatcher.h"
+#include "NiagaraComputeExecutionContext.h"
+#include "NiagaraGPUSystemTick.h"
+
+FNiagaraSystemGpuComputeProxy::FNiagaraSystemGpuComputeProxy(FNiagaraSystemInstance* OwnerInstance)
+	: DebugOwnerInstance(OwnerInstance)
+	, DebugOwnerBatcher(nullptr)
+	, SystemInstanceID(OwnerInstance->GetId())
+{
+	bRequiresDistanceFieldData = OwnerInstance->RequiresDistanceFieldData();
+	bRequiresDepthBuffer = OwnerInstance->RequiresDepthBuffer();
+	bRequiresEarlyViewData = OwnerInstance->RequiresEarlyViewData();
+	bRequiresViewUniformBuffer = OwnerInstance->RequiresViewUniformBuffer();
+
+	// Gather all emitter compute contexts
+	for ( auto& Emitter : OwnerInstance->GetEmitters() )
+	{
+		if (FNiagaraComputeExecutionContext* ComputeContext = Emitter->GetGPUContext() )
+		{
+			ComputeContexts.Emplace(ComputeContext);
+		}
+	}
+
+	// Calculate Tick Stage
+	if (bRequiresDistanceFieldData || bRequiresDepthBuffer)
+	{
+		ComputeTickStage = ENiagaraGpuComputeTickStage::PostOpaqueRender;
+	}
+	else if (bRequiresEarlyViewData)
+	{
+		ComputeTickStage = ENiagaraGpuComputeTickStage::PostInitViews;
+	}
+	else if (bRequiresViewUniformBuffer)
+	{
+		ComputeTickStage = ENiagaraGpuComputeTickStage::PostOpaqueRender;
+	}
+	else
+	{
+		ComputeTickStage = ENiagaraGpuComputeTickStage::PreInitViews;
+	}
+}
+
+FNiagaraSystemGpuComputeProxy::~FNiagaraSystemGpuComputeProxy()
+{
+	check(IsInRenderingThread());
+	check(DebugOwnerBatcher == nullptr);
+}
+
+void FNiagaraSystemGpuComputeProxy::AddToBatcher(NiagaraEmitterInstanceBatcher* Batcher)
+{
+	check(IsInGameThread());
+	check(DebugOwnerBatcher == nullptr);
+	DebugOwnerBatcher = Batcher;
+
+	ENQUEUE_RENDER_COMMAND(AddProxyToBatcher)(
+		[=](FRHICommandListImmediate& RHICmdList)
+		{
+			Batcher->AddGpuComputeProxy(this);
+
+			for (FNiagaraComputeExecutionContext* ComputeContext : ComputeContexts)
+			{
+				ComputeContext->bHasTickedThisFrame_RT = false;
+				ComputeContext->CurrentNumInstances_RT = 0;
+				ComputeContext->CurrentMaxInstances_RT = 0;
+
+				for (int i=0; i < UE_ARRAY_COUNT(ComputeContext->DataBuffers_RT); ++i)
+				{
+					check(ComputeContext->DataBuffers_RT[i] == nullptr);
+					ComputeContext->DataBuffers_RT[i] = new FNiagaraDataBuffer(ComputeContext->MainDataSet);
+				}
+			}
+		}
+	);
+}
+
+void FNiagaraSystemGpuComputeProxy::RemoveFromBatcher(NiagaraEmitterInstanceBatcher* Batcher, bool bDeleteProxy)
+{
+	check(IsInGameThread());
+	check(DebugOwnerBatcher == Batcher);
+	DebugOwnerBatcher = nullptr;
+
+	ENQUEUE_RENDER_COMMAND(RemoveFromBatcher)(
+		[=](FRHICommandListImmediate& RHICmdList)
+		{
+			Batcher->RemoveGpuComputeProxy(this);
+			ReleaseTicks(Batcher->GetGPUInstanceCounterManager());
+
+			for (FNiagaraComputeExecutionContext* ComputeContext : ComputeContexts)
+			{
+				ComputeContext->ResetInternal(Batcher);
+
+				//-TODO: Can we move this inside the context???
+				for (int i = 0; i < UE_ARRAY_COUNT(ComputeContext->DataBuffers_RT); ++i)
+				{
+					check(ComputeContext->DataBuffers_RT[i]);
+					ComputeContext->DataBuffers_RT[i]->ReleaseGPU();
+					ComputeContext->DataBuffers_RT[i]->Destroy();
+					ComputeContext->DataBuffers_RT[i] = nullptr;
+				}
+			}
+
+			if (bDeleteProxy)
+			{
+				delete this;
+			}
+		}
+	);
+}
+
+void FNiagaraSystemGpuComputeProxy::QueueTick(const FNiagaraGPUSystemTick& Tick)
+{
+	check(IsInRenderingThread());
+
+	//if ( !FNiagaraUtilities::AllowGPUParticles(GetShaderPlatform()) )
+	//{
+	//	return;
+	//}
+
+	//-OPT: Making a copy of the tick here, reduce this.
+	PendingTicks.Add(Tick);
+}
+
+void FNiagaraSystemGpuComputeProxy::ReleaseTicks(FNiagaraGPUInstanceCountManager& GPUInstanceCountManager)
+{
+	check(IsInRenderingThread());
+
+	// Release all the ticks
+	for (auto& Tick : PendingTicks)
+	{
+		Tick.Destroy();
+	}
+	PendingTicks.Empty();
+
+	for (FNiagaraComputeExecutionContext* ComputeContext : ComputeContexts)
+	{
+		// Reset pending information as this will have the readback incorporated into it
+		ComputeContext->bHasTickedThisFrame_RT = false;
+		ComputeContext->CurrentMaxInstances_RT = 0;
+
+		// Clear counter offsets
+		for (int i = 0; i < UE_ARRAY_COUNT(ComputeContext->DataBuffers_RT); ++i)
+		{
+			ComputeContext->DataBuffers_RT[i]->ClearGPUInstanceCount();
+		}
+	}
+}
diff --git a/Engine/Plugins/FX/Niagara/Source/Niagara/Private/NiagaraSystemInstance.cpp b/Engine/Plugins/FX/Niagara/Source/Niagara/Private/NiagaraSystemInstance.cpp
--- a/Engine/Plugins/FX/Niagara/Source/Niagara/Private/NiagaraSystemInstance.cpp
+++ b/Engine/Plugins/FX/Niagara/Source/Niagara/Private/NiagaraSystemInstance.cpp
@@ -1,6 +1,7 @@
 // Copyright Epic Games, Inc. All Rights Reserved.
 
 #include "NiagaraSystemInstance.h"
+#include "NiagaraSystemGpuComputeProxy.h"
 #include "NiagaraConstants.h"
 #include "NiagaraCommon.h"
 #include "NiagaraDataInterface.h"
@@ -107,7 +108,6 @@ FNiagaraSystemInstance::FNiagaraSystemInstance(UWorld& InWorld, UNiagaraSystem&
 	, bAlreadyBound(false)
 	, bLODDistanceIsValid(false)
 	, bPooled(bInPooled)
-	, bHasSimulationReset(false)
 	, CachedDeltaSeconds(0.0f)
 	, RequestedExecutionState(ENiagaraExecutionState::Complete)
 	, ActualExecutionState(ENiagaraExecutionState::Complete)
@@ -626,17 +626,11 @@ bool FNiagaraSystemInstance::DeallocateSystemInstance(FNiagaraSystemInstancePtr&
 		}
 		SystemInstanceAllocation->UnbindParameters();
 
-		// If we have active GPU emitters make sure we remove any pending ticks from the RT
-		NiagaraEmitterInstanceBatcher* InstanceBatcher = SystemInstanceAllocation->GetBatcher();
-		if (SystemInstanceAllocation->bHasGPUEmitters && FNiagaraUtilities::AllowGPUParticles(InstanceBatcher->GetShaderPlatform()))
+		// Release the render proxy
+		if (SystemInstanceAllocation->SystemGpuComputeProxy)
 		{
-			ENQUEUE_RENDER_COMMAND(NiagaraRemoveGPUSystem)
-			(
-				[InstanceBatcher, InstanceID=SystemInstanceAllocation->GetId()](FRHICommandListImmediate& RHICmdList) mutable
-				{
-					InstanceBatcher->InstanceDeallocated_RenderThread(InstanceID);
-				}
-			);
+			FNiagaraSystemGpuComputeProxy* Proxy = SystemInstanceAllocation->SystemGpuComputeProxy.Release();
+			Proxy->RemoveFromBatcher(SystemInstanceAllocation->GetBatcher(), true);
 		}
 
 		// Queue deferred deletion from the WorldManager
@@ -686,6 +680,12 @@ void FNiagaraSystemInstance::Complete(bool bExternalCompletion)
 
 	DestroyDataInterfaceInstanceData();
 
+	if (SystemGpuComputeProxy)
+	{
+		FNiagaraSystemGpuComputeProxy* Proxy = SystemGpuComputeProxy.Release();
+		Proxy->RemoveFromBatcher(GetBatcher(), true);
+	}
+
 	if (!bPooled)
 	{
 		UnbindParameters(true);
@@ -824,11 +824,25 @@ void FNiagaraSystemInstance::Reset(FNiagaraSystemInstance::EResetMode Mode)
 		{
 			InstanceParameters.Tick();//Make sure the owner has flushed it's parameters by now. Especially it's DIs.
 			InitDataInterfaces();
+
+			// Kill the current proxy as we are reinitializing the system
+			if (SystemGpuComputeProxy.IsValid())
+			{
+				FNiagaraSystemGpuComputeProxy* Proxy = SystemGpuComputeProxy.Release();
+				Proxy->RemoveFromBatcher(GetBatcher(), true);
+			}
 		}
 
 		//Interface init can disable the system.
 		if (!IsComplete())
 		{
+			// Create the shared context for the batcher if we have a single active GPU emitter in the system
+			if (bHasGPUEmitters && !SystemGpuComputeProxy.IsValid())
+			{
+				SystemGpuComputeProxy.Reset(new FNiagaraSystemGpuComputeProxy(this));
+				SystemGpuComputeProxy->AddToBatcher(GetBatcher());
+			}
+
 			SystemSimulation->AddInstance(this);
 
 			UNiagaraSystem* System = GetSystem();
@@ -844,6 +858,9 @@ void FNiagaraSystemInstance::Reset(FNiagaraSystemInstance::EResetMode Mode)
 				TickCount = 0;
 			}
 		}
+		else
+		{
+		}
 	}
 	else
 	{
@@ -861,7 +878,6 @@ void FNiagaraSystemInstance::ResetInternal(bool bResetSimulations)
 
 	Age = 0;
 	TickCount = 0;
-	bHasSimulationReset = bResetSimulations;
 	CachedDeltaSeconds = 0.0f;
 	bLODDistanceIsValid = false;
 	TotalGPUParamSize = 0;
@@ -999,7 +1015,6 @@ void FNiagaraSystemInstance::ReInitInternal()
 
 	Age = 0;
 	TickCount = 0;
-	bHasSimulationReset = true;
 	LocalBounds = FBox(FVector::ZeroVector, FVector::ZeroVector);
 	CachedDeltaSeconds = 0.0f;
 	bAlreadyBound = false;
@@ -1119,6 +1134,12 @@ void FNiagaraSystemInstance::Cleanup()
 
 	DestroyDataInterfaceInstanceData();
 
+	if (SystemGpuComputeProxy)
+	{
+		FNiagaraSystemGpuComputeProxy* Proxy = SystemGpuComputeProxy.Release();
+		Proxy->RemoveFromBatcher(GetBatcher(), true);
+	}
+
 	UnbindParameters();
 
 	// Clear out the emitters.
@@ -2029,12 +2050,6 @@ void FNiagaraSystemInstance::InitEmitters()
 			}
 		}
 
-		// Create the shared context for the batcher if we have a single active GPU emitter in the system
-		if (bHasGPUEmitters)
-		{
-			SharedContext.Reset(new FNiagaraComputeSharedContext());
-		}
-
 		if (System->bFixedBounds)
 		{
 			LocalBounds = System->GetFixedBounds();
@@ -2370,11 +2385,10 @@ void FNiagaraSystemInstance::GenerateAndSubmitGPUTick()
 		// We no longer own it and cannot modify it after this point.
 		// @todo We are taking a copy of the object here. This object is small so this overhead should
 		// not be very high. And we avoid making a bunch of small allocations here.
-		NiagaraEmitterInstanceBatcher* TheBatcher = GetBatcher();
 		ENQUEUE_RENDER_COMMAND(FNiagaraGiveSystemInstanceTickToRT)(
-			[TheBatcher, GPUTick](FRHICommandListImmediate& RHICmdList) mutable
+			[RT_Proxy=SystemGpuComputeProxy.Get(), GPUTick](FRHICommandListImmediate& RHICmdList) mutable
 			{
-				TheBatcher->GiveSystemTick_RenderThread(GPUTick);
+				RT_Proxy->QueueTick(GPUTick);
 			}
 		);
 	}
@@ -2383,6 +2397,7 @@ void FNiagaraSystemInstance::GenerateAndSubmitGPUTick()
 void FNiagaraSystemInstance::InitGPUTick(FNiagaraGPUSystemTick& OutTick)
 {
 	SCOPE_CYCLE_COUNTER(STAT_NiagaraInitGPUSystemTick);
+	check(SystemGpuComputeProxy.IsValid());
 	OutTick.Init(this);
 
 	//if (GPUTick.DIInstanceData)
@@ -2426,19 +2441,6 @@ bool FNiagaraSystemInstance::GetIsolateEnabled() const
 
 void FNiagaraSystemInstance::DestroyDataInterfaceInstanceData()
 {
-	NiagaraEmitterInstanceBatcher* InstanceBatcher = GetBatcher();
-	if (bHasGPUEmitters && FNiagaraUtilities::AllowGPUParticles(InstanceBatcher->GetShaderPlatform()))
-	{
-		ENQUEUE_RENDER_COMMAND(NiagaraRemoveGPUSystem)
-		(
-			[InstanceBatcher, InstanceID=GetId()](FRHICommandListImmediate& RHICmdList) mutable
-			{
-				InstanceBatcher->InstanceDeallocated_RenderThread(InstanceID);
-			}
-		);
-	}
-
-	//
 	for (TPair<TWeakObjectPtr<UNiagaraDataInterface>, int32>& Pair : DataInterfaceInstanceDataOffsets)
 	{
 		if (UNiagaraDataInterface* Interface = Pair.Key.Get())
diff --git a/Engine/Plugins/FX/Niagara/Source/Niagara/Private/NiagaraSystemSimulation.cpp b/Engine/Plugins/FX/Niagara/Source/Niagara/Private/NiagaraSystemSimulation.cpp
--- a/Engine/Plugins/FX/Niagara/Source/Niagara/Private/NiagaraSystemSimulation.cpp
+++ b/Engine/Plugins/FX/Niagara/Source/Niagara/Private/NiagaraSystemSimulation.cpp
@@ -8,6 +8,7 @@
 #include "NiagaraSettings.h"
 #include "NiagaraParameterCollection.h"
 #include "NiagaraSystemInstance.h"
+#include "NiagaraSystemGpuComputeProxy.h"
 #include "NiagaraConstants.h"
 #include "NiagaraStats.h"
 #include "Async/ParallelFor.h"
@@ -315,7 +316,7 @@ struct FNiagaraSystemInstanceTickConcurrentTask
 	{
 		if (SystemSimulation->GetGPUTickHandlingMode() == ENiagaraGPUTickHandlingMode::ConcurrentBatched)
 		{
-			TArray<FNiagaraGPUSystemTick, TInlineAllocator<NiagaraSystemTickBatchSize>> GPUTicks;
+			TArray<TPair<FNiagaraSystemGpuComputeProxy*, FNiagaraGPUSystemTick>, TInlineAllocator<NiagaraSystemTickBatchSize>> GPUTicks;
 			GPUTicks.Reserve(Batch.Num());
 			for (FNiagaraSystemInstance* Inst : Batch)
 			{
@@ -323,20 +324,20 @@ struct FNiagaraSystemInstanceTickConcurrentTask
 				Inst->Tick_Concurrent(false);
 				if (Inst->NeedsGPUTick())
 				{
-					Inst->InitGPUTick(GPUTicks.AddDefaulted_GetRef());
+					auto& Tick = GPUTicks.AddDefaulted_GetRef();
+					Tick.Key = Inst->GetSystemGpuComputeProxy();
+					Inst->InitGPUTick(Tick.Value);
 				}
 			}
 
 			if (GPUTicks.Num() > 0)
 			{
-				NiagaraEmitterInstanceBatcher* TheBatcher = SystemSimulation->GetBatcher();
-				ENQUEUE_RENDER_COMMAND(FNiagaraGiveSystemInstanceTickToRT)
-				(
-					[TheBatcher, GPUTicks](FRHICommandListImmediate& RHICmdList) mutable
+				ENQUEUE_RENDER_COMMAND(FNiagaraGiveSystemInstanceTickToRT)(
+					[GPUTicks](FRHICommandListImmediate& RHICmdList) mutable
 					{
 						for (auto& GPUTick : GPUTicks)
 						{
-							TheBatcher->GiveSystemTick_RenderThread(GPUTick);
+							GPUTick.Key->QueueTick(GPUTick.Value);
 						}
 					}
 				);
@@ -389,7 +390,7 @@ struct FNiagaraSystemInstanceFinalizeTask
 
 		if ( TickHandlingMode == ENiagaraGPUTickHandlingMode::GameThreadBatched )
 		{
-			TArray<FNiagaraGPUSystemTick, TInlineAllocator<NiagaraSystemTickBatchSize>> GPUTicks;
+			TArray<TPair<FNiagaraSystemGpuComputeProxy*, FNiagaraGPUSystemTick>, TInlineAllocator<NiagaraSystemTickBatchSize>> GPUTicks;
 			GPUTicks.Reserve(Batch.Num());
 			for (FNiagaraSystemInstance* Instance : Batch)
 			{
@@ -402,19 +403,20 @@ struct FNiagaraSystemInstanceFinalizeTask
 				Instance->FinalizeTick_GameThread(false);
 				if(Instance->NeedsGPUTick())
 				{
-					Instance->InitGPUTick(GPUTicks.AddDefaulted_GetRef());
+					auto& Tick = GPUTicks.AddDefaulted_GetRef();
+					Tick.Key = Instance->GetSystemGpuComputeProxy();
+					Instance->InitGPUTick(Tick.Value);
 				}
 			}
 
 			if(GPUTicks.Num() > 0)
 			{
-				NiagaraEmitterInstanceBatcher* TheBatcher = SystemSimulation->GetBatcher();
 				ENQUEUE_RENDER_COMMAND(FNiagaraGiveSystemInstanceTickToRT)(
-					[TheBatcher, GPUTicks](FRHICommandListImmediate& RHICmdList) mutable
+					[GPUTicks](FRHICommandListImmediate& RHICmdList) mutable
 					{
 						for(auto& GPUTick : GPUTicks)
 						{
-							TheBatcher->GiveSystemTick_RenderThread(GPUTick);
+							GPUTick.Key->QueueTick(GPUTick.Value);
 						}
 					}
 				);
diff --git a/Engine/Plugins/FX/Niagara/Source/Niagara/Public/NiagaraSystemGpuComputeProxy.h b/Engine/Plugins/FX/Niagara/Source/Niagara/Public/NiagaraSystemGpuComputeProxy.h
new file mode 100644
--- /dev/null
+++ b/Engine/Plugins/FX/Niagara/Source/Niagara/Public/NiagaraSystemGpuComputeProxy.h
@@ -0,0 +1,61 @@
+// Copyright Epic Games, Inc. All Rights Reserved.
+
+#pragma once
+
+#include "NiagaraCommon.h"
+
+class FNiagaraSystemInstance;
+class NiagaraEmitterInstanceBatcher;
+struct FNiagaraComputeExecutionContext;
+class FNiagaraGPUInstanceCountManager;
+class FNiagaraGPUSystemTick;
+
+namespace ENiagaraGpuComputeTickStage
+{
+	enum Type
+	{
+		PreInitViews,
+		PostInitViews,
+		PostOpaqueRender,
+		Max
+	};
+};
+
+class FNiagaraSystemGpuComputeProxy
+{
+	friend class NiagaraEmitterInstanceBatcher;
+
+public:
+	FNiagaraSystemGpuComputeProxy(FNiagaraSystemInstance* OwnerInstance);
+	~FNiagaraSystemGpuComputeProxy();
+
+	void AddToBatcher(NiagaraEmitterInstanceBatcher* Batcher);
+	void RemoveFromBatcher(NiagaraEmitterInstanceBatcher* Batcher, bool bDeleteProxy);
+
+	FNiagaraSystemInstanceID GetSystemInstanceID() const { return SystemInstanceID; }
+	ENiagaraGpuComputeTickStage::Type GetComputeTickStage() const { return ComputeTickStage; }
+	void QueueTick(const FNiagaraGPUSystemTick& Tick);
+	void ReleaseTicks(FNiagaraGPUInstanceCountManager& GPUInstanceCountManager);
+
+	bool RequiresDistanceFieldData() const { return bRequiresDistanceFieldData; }
+	bool RequiresDepthBuffer() const { return bRequiresDepthBuffer; }
+	bool RequiresEarlyViewData() const { return bRequiresEarlyViewData; }
+	bool RequiresViewUniformBuffer() const { return bRequiresViewUniformBuffer; }
+	bool RequiresRayTracingScene() const { return bRequiresRayTracingScene; }
+
+private:
+	FNiagaraSystemInstance*						DebugOwnerInstance = nullptr;
+	NiagaraEmitterInstanceBatcher*				DebugOwnerBatcher = nullptr;
+	int32										BatcherIndex = INDEX_NONE;
+
+	FNiagaraSystemInstanceID					SystemInstanceID = FNiagaraSystemInstanceID();
+	ENiagaraGpuComputeTickStage::Type			ComputeTickStage = ENiagaraGpuComputeTickStage::PostOpaqueRender;
+	uint32										bRequiresDistanceFieldData : 1;
+	uint32										bRequiresDepthBuffer : 1;
+	uint32										bRequiresEarlyViewData : 1;
+	uint32										bRequiresViewUniformBuffer : 1;
+	uint32										bRequiresRayTracingScene : 1;
+
+	TArray<FNiagaraComputeExecutionContext*>	ComputeContexts;
+	TArray<FNiagaraGPUSystemTick>				PendingTicks;
+};
diff --git a/Engine/Plugins/FX/Niagara/Source/Niagara/Public/NiagaraSystemInstance.h b/Engine/Plugins/FX/Niagara/Source/Niagara/Public/NiagaraSystemInstance.h
--- a/Engine/Plugins/FX/Niagara/Source/Niagara/Public/NiagaraSystemInstance.h
+++ b/Engine/Plugins/FX/Niagara/Source/Niagara/Public/NiagaraSystemInstance.h
@@ -16,6 +16,7 @@ class FNiagaraSystemInstance;
 class FNiagaraSystemSimulation;
 class NiagaraEmitterInstanceBatcher;
 class FNiagaraGPUSystemTick;
+class FNiagaraSystemGpuComputeProxy;
 
 using FNiagaraSystemInstancePtr = TSharedPtr<FNiagaraSystemInstance, ESPMode::ThreadSafe>;
 
@@ -236,7 +237,7 @@ public:
 
 	FORCEINLINE bool NeedsGPUTick() const { return ActiveGPUEmitterCount > 0 /*&& Component->IsRegistered()*/ && !IsComplete();}
 
-	struct FNiagaraComputeSharedContext* GetComputeSharedContext() { check(SharedContext.Get()); return SharedContext.Get(); }
+	FNiagaraSystemGpuComputeProxy* GetSystemGpuComputeProxy() { return SystemGpuComputeProxy.Get(); }
 
 	/** Gets a multicast delegate which is called after this instance has finished ticking for the frame on the game thread */
 	FORCEINLINE void SetOnPostTick(const FOnPostTick& InPostTickDelegate) { OnPostTickDelegate = InPostTickDelegate; }
@@ -285,7 +286,6 @@ public:
 
 	FORCEINLINE float GetAge() const { return Age; }
 	FORCEINLINE int32 GetTickCount() const { return TickCount; }
-	FORCEINLINE bool RequiresGpuBufferReset() const { return bHasSimulationReset && (TickCount == 1); }
 
 	FORCEINLINE float GetLastRenderTime() const { return LastRenderTime; }
 	FORCEINLINE void SetLastRenderTime(float TimeSeconds) { LastRenderTime = TimeSeconds; }
@@ -517,9 +517,6 @@ private:
 	/** True if the system instance is pooled. Prevents unbinding of parameters on completing the system */
 	uint32 bPooled : 1;
 
-	/** Will be set to true when the the simulation needs a full reset from ResetInternal() */
-	uint32 bHasSimulationReset : 1;
-
 #if WITH_EDITOR
 	uint32 bNeedsUIResync : 1;
 #endif
@@ -549,6 +546,7 @@ private:
 	ENiagaraExecutionState ActualExecutionState;
 
 	NiagaraEmitterInstanceBatcher* Batcher = nullptr;
+	TUniquePtr<FNiagaraSystemGpuComputeProxy> SystemGpuComputeProxy;
 
 	/** Tag we feed into crash reporter for this instance. */
 	mutable FString CrashReporterTag;
@@ -563,7 +561,6 @@ public:
 	// Transient data that is accumulated during tick.
 	uint32 TotalGPUParamSize = 0;
 	uint32 ActiveGPUEmitterCount = 0;
-	TUniquePtr<FNiagaraComputeSharedContext, FNiagaraComputeSharedContextDeleter> SharedContext;
 
 	int32 GPUDataInterfaceInstanceDataSize = 0;
 	bool GPUParamIncludeInterpolation = false;
diff --git a/Engine/Plugins/FX/Niagara/Source/NiagaraCore/Public/NiagaraDataInterfaceBase.h b/Engine/Plugins/FX/Niagara/Source/NiagaraCore/Public/NiagaraDataInterfaceBase.h
--- a/Engine/Plugins/FX/Niagara/Source/NiagaraCore/Public/NiagaraDataInterfaceBase.h
+++ b/Engine/Plugins/FX/Niagara/Source/NiagaraCore/Public/NiagaraDataInterfaceBase.h
@@ -22,6 +22,7 @@ class FRHICommandList;
 struct FNiagaraDataInterfaceProxy;
 class NiagaraEmitterInstanceBatcher;
 struct FNiagaraComputeInstanceData;
+struct FNiagaraSimStageData;
 
 DECLARE_EXPORTED_TEMPLATE_INTRINSIC_TYPE_LAYOUT(template<>, TIndexedPtr<UNiagaraDataInterfaceBase>, NIAGARACORE_API);
 
@@ -41,17 +42,17 @@ struct FNiagaraDataInterfaceArgs
 
 struct FNiagaraDataInterfaceStageArgs : public FNiagaraDataInterfaceArgs
 {
-	FNiagaraDataInterfaceStageArgs(FNiagaraDataInterfaceProxy* InDataInterface, FNiagaraSystemInstanceID InSystemInstanceID, const NiagaraEmitterInstanceBatcher* InBatcher, const FNiagaraComputeInstanceData* InComputeInstanceData, uint32 InSimulationStageIndex, bool InIsOutputStage, bool InIsIterationStage)
+	FNiagaraDataInterfaceStageArgs(FNiagaraDataInterfaceProxy* InDataInterface, FNiagaraSystemInstanceID InSystemInstanceID, const NiagaraEmitterInstanceBatcher* InBatcher, const FNiagaraComputeInstanceData* InComputeInstanceData, const FNiagaraSimStageData* InSimStageData, bool InIsOutputStage, bool InIsIterationStage)
 		: FNiagaraDataInterfaceArgs(InDataInterface, InSystemInstanceID, InBatcher)
 		, ComputeInstanceData(InComputeInstanceData)
-		, SimulationStageIndex(InSimulationStageIndex)
+		, SimStageData(InSimStageData)
 		, IsOutputStage(InIsOutputStage)
 		, IsIterationStage(InIsIterationStage)
 	{
 	}
 
 	const FNiagaraComputeInstanceData* ComputeInstanceData;
-	uint32 SimulationStageIndex;
+	const FNiagaraSimStageData* SimStageData;
 	bool IsOutputStage;
 	bool IsIterationStage;
 };
@@ -61,11 +62,11 @@ struct FNiagaraDataInterfaceSetArgs : public FNiagaraDataInterfaceArgs
 	typedef TShaderRefBase<FNiagaraShader, FNiagaraShaderMapPointerTable> FShaderReference;
 
 	FNiagaraDataInterfaceSetArgs(
-		FNiagaraDataInterfaceProxy* InDataInterface, FNiagaraSystemInstanceID InSystemInstanceID, const NiagaraEmitterInstanceBatcher* InBatcher, const FShaderReference& InShader, const FNiagaraComputeInstanceData* InComputeInstanceData, uint32 InSimulationStageIndex, bool InIsOutputStage, bool InIsIterationStage)
+		FNiagaraDataInterfaceProxy* InDataInterface, FNiagaraSystemInstanceID InSystemInstanceID, const NiagaraEmitterInstanceBatcher* InBatcher, const FShaderReference& InShader, const FNiagaraComputeInstanceData* InComputeInstanceData, const FNiagaraSimStageData* InSimStageData, bool InIsOutputStage, bool InIsIterationStage)
 		: FNiagaraDataInterfaceArgs(InDataInterface, InSystemInstanceID, InBatcher)
 		, Shader(InShader)
 		, ComputeInstanceData(InComputeInstanceData)
-		, SimulationStageIndex(InSimulationStageIndex)
+		, SimStageData(InSimStageData)
 		, IsOutputStage(InIsOutputStage)
 		, IsIterationStage(InIsIterationStage)
 	{
@@ -73,7 +74,7 @@ struct FNiagaraDataInterfaceSetArgs : public FNiagaraDataInterfaceArgs
 
 	FShaderReference Shader;
 	const FNiagaraComputeInstanceData* ComputeInstanceData;
-	uint32 SimulationStageIndex;
+	const FNiagaraSimStageData* SimStageData;
 	bool IsOutputStage;
 	bool IsIterationStage;
 };
diff --git a/Engine/Plugins/FX/Niagara/Source/NiagaraEditor/Private/NiagaraBakerRenderer.cpp b/Engine/Plugins/FX/Niagara/Source/NiagaraEditor/Private/NiagaraBakerRenderer.cpp
--- a/Engine/Plugins/FX/Niagara/Source/NiagaraEditor/Private/NiagaraBakerRenderer.cpp
+++ b/Engine/Plugins/FX/Niagara/Source/NiagaraEditor/Private/NiagaraBakerRenderer.cpp
@@ -3,6 +3,7 @@
 #include "NiagaraBakerRenderer.h"
 #include "NiagaraComponent.h"
 #include "NiagaraSystem.h"
+#include "NiagaraComputeExecutionContext.h"
 
 #include "Niagara/Classes/NiagaraDataInterfaceRenderTarget2D.h"
 #include "Niagara/Classes/NiagaraDataInterfaceGrid2DCollection.h"
diff --git a/Engine/Plugins/Runtime/HairStrands/Source/HairStrandsCore/Private/Niagara/NiagaraDataInterfaceHairStrands.cpp b/Engine/Plugins/Runtime/HairStrands/Source/HairStrandsCore/Private/Niagara/NiagaraDataInterfaceHairStrands.cpp
--- a/Engine/Plugins/Runtime/HairStrands/Source/HairStrandsCore/Private/Niagara/NiagaraDataInterfaceHairStrands.cpp
+++ b/Engine/Plugins/Runtime/HairStrands/Source/HairStrandsCore/Private/Niagara/NiagaraDataInterfaceHairStrands.cpp
@@ -4013,7 +4013,7 @@ void UNiagaraDataInterfaceHairStrands::ProvidePerInstanceDataForRenderThread(voi
 
 void FNDIHairStrandsProxy::PreStage(FRHICommandList& RHICmdList, const FNiagaraDataInterfaceStageArgs& Context)
 {
-	if (Context.SimulationStageIndex == 0)
+	if (Context.SimStageData->bFirstStage)
 	{
 		FNDIHairStrandsData* ProxyData =
 			SystemInstancesToProxyData.Find(Context.SystemInstanceID);
diff --git a/Engine/Plugins/Runtime/HairStrands/Source/HairStrandsCore/Private/Niagara/NiagaraDataInterfacePhysicsAsset.cpp b/Engine/Plugins/Runtime/HairStrands/Source/HairStrandsCore/Private/Niagara/NiagaraDataInterfacePhysicsAsset.cpp
--- a/Engine/Plugins/Runtime/HairStrands/Source/HairStrandsCore/Private/Niagara/NiagaraDataInterfacePhysicsAsset.cpp
+++ b/Engine/Plugins/Runtime/HairStrands/Source/HairStrandsCore/Private/Niagara/NiagaraDataInterfacePhysicsAsset.cpp
@@ -629,7 +629,7 @@ void FNDIPhysicsAssetProxy::PreStage(FRHICommandList& RHICmdList, const FNiagara
 
 	if (ProxyData != nullptr && ProxyData->AssetBuffer)
 	{
-		if (Context.SimulationStageIndex == 0)
+		if (Context.SimStageData->bFirstStage)
 		{
 			UpdateInternalBuffer<FVector4, EPixelFormat::PF_A32B32G32R32F, PHYSICS_ASSET_MAX_TRANSFORMS>(ProxyData->AssetArrays.CurrentTransform, ProxyData->AssetBuffer->CurrentTransformBuffer);
 			UpdateInternalBuffer<FVector4, EPixelFormat::PF_A32B32G32R32F, PHYSICS_ASSET_MAX_TRANSFORMS>(ProxyData->AssetArrays.PreviousTransform, ProxyData->AssetBuffer->PreviousTransformBuffer);
diff --git a/Engine/Plugins/Runtime/HairStrands/Source/HairStrandsCore/Private/Niagara/NiagaraDataInterfacePressureGrid.cpp b/Engine/Plugins/Runtime/HairStrands/Source/HairStrandsCore/Private/Niagara/NiagaraDataInterfacePressureGrid.cpp
--- a/Engine/Plugins/Runtime/HairStrands/Source/HairStrandsCore/Private/Niagara/NiagaraDataInterfacePressureGrid.cpp
+++ b/Engine/Plugins/Runtime/HairStrands/Source/HairStrandsCore/Private/Niagara/NiagaraDataInterfacePressureGrid.cpp
@@ -535,7 +535,7 @@ void FNDIPressureGridProxy::PreStage(FRHICommandList& RHICmdList, const FNiagara
 
 	if (ProxyData != nullptr)
 	{
-		if (Context.SimulationStageIndex == 0)
+		if (Context.SimStageData->bFirstStage)
 		{
 			ClearBuffer(RHICmdList, Context.Batcher->GetFeatureLevel(), ProxyData->CurrentGridBuffer, ProxyData->DestinationGridBuffer, ProxyData->GridSize, true);
 		}
diff --git a/Engine/Plugins/Runtime/HairStrands/Source/HairStrandsCore/Private/Niagara/NiagaraDataInterfaceVelocityGrid.cpp b/Engine/Plugins/Runtime/HairStrands/Source/HairStrandsCore/Private/Niagara/NiagaraDataInterfaceVelocityGrid.cpp
--- a/Engine/Plugins/Runtime/HairStrands/Source/HairStrandsCore/Private/Niagara/NiagaraDataInterfaceVelocityGrid.cpp
+++ b/Engine/Plugins/Runtime/HairStrands/Source/HairStrandsCore/Private/Niagara/NiagaraDataInterfaceVelocityGrid.cpp
@@ -633,7 +633,7 @@ void FNDIVelocityGridProxy::PreStage(FRHICommandList& RHICmdList, const FNiagara
 
 	if (ProxyData != nullptr )
 	{
-		if (Context.SimulationStageIndex == 0)
+		if (Context.SimStageData->bFirstStage)
 		{
 			ClearBuffer(RHICmdList, ProxyData->DestinationGridBuffer);
 		}
