Fixup logic around grabbing landscape tiles for the Niagara DI when the bounds of the niagara system are large

diff --git a/Engine/Plugins/FX/Niagara/Source/Niagara/Private/NiagaraDataInterfaceLandscape.cpp b/Engine/Plugins/FX/Niagara/Source/Niagara/Private/NiagaraDataInterfaceLandscape.cpp
--- a/Engine/Plugins/FX/Niagara/Source/Niagara/Private/NiagaraDataInterfaceLandscape.cpp
+++ b/Engine/Plugins/FX/Niagara/Source/Niagara/Private/NiagaraDataInterfaceLandscape.cpp
@@ -477,41 +477,64 @@ FNDI_Landscape_SharedResourceHandle FNDI_Landscape_GeneratedData::GetLandscapeDa
 	}
 
 	// we want to use the bounds of the system to figure out which cells of the landscape that we need to handle
-	const FBox WorldBounds = SystemInstance.GetLocalBounds().TransformBy(SystemInstance.GetWorldTransform());
+	const FBox SystemWorldBounds = SystemInstance.GetLocalBounds().TransformBy(SystemInstance.GetWorldTransform());
 	const FTransform& LandscapeActorToWorld = Landscape->LandscapeActorToWorld();
-	const FVector MinWorldBoundsLandscapeActor = LandscapeActorToWorld.InverseTransformPosition(WorldBounds.Min);
-	const FVector MaxWorldBoundsLandscapeActor = LandscapeActorToWorld.InverseTransformPosition(WorldBounds.Max);
+	const FVector SystemMinInLandscape = LandscapeActorToWorld.InverseTransformPosition(SystemWorldBounds.Min);
+	const FVector SystemMaxInLandscape = LandscapeActorToWorld.InverseTransformPosition(SystemWorldBounds.Max);
+
+	const FBox SystemBoundsInLandscape(
+		SystemMinInLandscape.ComponentMin(SystemMaxInLandscape),
+		SystemMinInLandscape.ComponentMax(SystemMaxInLandscape));
+
+	// at runtime we don't know the potential number of components in the landscape, so we can only clamp the values to start at 0.
+	// Ideally, we'd be able to specify an upper bound as well to further refine the region of overlap.
+	const FIntPoint MinBoundRegion = FIntPoint(SystemBoundsInLandscape.Min.X, SystemBoundsInLandscape.Min.Y).ComponentMax(FIntPoint(0, 0));
+	const FIntPoint MaxBoundRegion = FIntPoint(SystemBoundsInLandscape.Max.X, SystemBoundsInLandscape.Max.Y).ComponentMax(FIntPoint(0, 0));
 
-	const FVector LandscapeActorMin = MinWorldBoundsLandscapeActor.ComponentMin(MaxWorldBoundsLandscapeActor);
-	const FVector LandscapeActorMax = MinWorldBoundsLandscapeActor.ComponentMax(MaxWorldBoundsLandscapeActor);
+	const FIntRect SystemRect = FIntRect(MinBoundRegion, MaxBoundRegion).Scale(1.0f / Landscape->ComponentSizeQuads);
 
 	FNDI_Landscape_SharedResource::FResourceKey Key;
 	Key.Source = Landscape;
+	Key.MinCaptureRegion = FIntPoint(TNumericLimits<int32>::Max(), TNumericLimits<int32>::Max());
+	Key.MaxCaptureRegion = FIntPoint(TNumericLimits<int32>::Min(), TNumericLimits<int32>::Min());
 
-	const FIntPoint MinBoundRegion(
-		FMath::FloorToInt(LandscapeActorMin.X / Landscape->ComponentSizeQuads),
-		FMath::FloorToInt(LandscapeActorMin.Y / Landscape->ComponentSizeQuads));
+	const int32 MaxSystemRegionCount = (SystemRect.Max.X - SystemRect.Min.X) * (SystemRect.Max.Y - SystemRect.Min.Y);
+	const int32 MaxLandscapeRegionCount = LandscapeInfo->XYtoCollisionComponentMap.Num();
 
-	const FIntPoint MaxBoundRegion(
-		FMath::FloorToInt(LandscapeActorMax.X / Landscape->ComponentSizeQuads),
-		FMath::FloorToInt(LandscapeActorMax.Y / Landscape->ComponentSizeQuads));
+	Key.CapturedRegions.Reserve(FMath::Min(MaxSystemRegionCount, MaxLandscapeRegionCount));
 
-	Key.MinCaptureRegion = FIntPoint(TNumericLimits<int32>::Max(), TNumericLimits<int32>::Max());
-	Key.MaxCaptureRegion = FIntPoint(TNumericLimits<int32>::Min(), TNumericLimits<int32>::Min());
+	auto AddRegion = [&](const FIntPoint& Region)
+	{
+		Key.CapturedRegions.Add(Region);
+		Key.MinCaptureRegion = Key.MinCaptureRegion.ComponentMin(Region);
+		Key.MaxCaptureRegion = Key.MaxCaptureRegion.ComponentMax(Region);
+	};
 
-	const int32 MaxRegionCount = (Key.MaxCaptureRegion.X - Key.MinCaptureRegion.X + 1) * (Key.MaxCaptureRegion.Y - Key.MinCaptureRegion.Y + 1);
-	Key.CapturedRegions.Reserve(MaxRegionCount);
+	if (MaxSystemRegionCount > MaxLandscapeRegionCount)
+	{
+		Key.CapturedRegions.Reserve(MaxLandscapeRegionCount);
 
-	for (int32 GridY = MinBoundRegion.Y; GridY <= MaxBoundRegion.Y; ++GridY)
+		for (const auto& LandscapeComponent : LandscapeInfo->XYtoCollisionComponentMap)
+		{
+			if (SystemRect.Contains(LandscapeComponent.Key))
+			{
+				AddRegion(LandscapeComponent.Key);
+			}
+		}
+	}
+	else
 	{
-		for (int32 GridX = MinBoundRegion.X; GridX <= MaxBoundRegion.X; ++GridX)
+		Key.CapturedRegions.Reserve(MaxSystemRegionCount);
+
+		for (int32 GridY = SystemRect.Min.Y; GridY < SystemRect.Max.Y; ++GridY)
 		{
-			const FIntPoint CurrentRegion(GridX, GridY);
-			if (LandscapeInfo->XYtoCollisionComponentMap.Contains(CurrentRegion))
+			for (int32 GridX = SystemRect.Min.X; GridX < SystemRect.Max.X; ++GridX)
 			{
-				Key.CapturedRegions.Emplace(GridX, GridY);
-				Key.MinCaptureRegion = Key.MinCaptureRegion.ComponentMin(CurrentRegion);
-				Key.MaxCaptureRegion = Key.MaxCaptureRegion.ComponentMax(CurrentRegion);
+				const FIntPoint CurrentRegion(GridX, GridY);
+				if (LandscapeInfo->XYtoCollisionComponentMap.Contains(CurrentRegion))
+				{
+					AddRegion(CurrentRegion);
+				}
 			}
 		}
 	}
