Fixes a couple of merge errors in handling niagara ray tracing

diff --git a/Engine/Plugins/FX/Niagara/Source/Niagara/Private/NiagaraEmitterInstanceBatcher.cpp b/Engine/Plugins/FX/Niagara/Source/Niagara/Private/NiagaraEmitterInstanceBatcher.cpp
--- a/Engine/Plugins/FX/Niagara/Source/Niagara/Private/NiagaraEmitterInstanceBatcher.cpp
+++ b/Engine/Plugins/FX/Niagara/Source/Niagara/Private/NiagaraEmitterInstanceBatcher.cpp
@@ -1720,10 +1720,6 @@ void NiagaraEmitterInstanceBatcher::PostRenderOpaque(FRDGBuilder& GraphBuilder,
 			FinishDispatches();
 
 			ProcessDebugReadbacks(RHICmdList, false);
-
-#if RHI_RAYTRACING
-			ResetRayTracingSceneInfo();
-#endif
 		}
 
 	if (!GPUInstanceCounterManager.HasPendingGPUReadback() && bRequiresReadback)
diff --git a/Engine/Plugins/FX/Niagara/Source/Niagara/Private/NiagaraSystemGpuComputeProxy.cpp b/Engine/Plugins/FX/Niagara/Source/Niagara/Private/NiagaraSystemGpuComputeProxy.cpp
--- a/Engine/Plugins/FX/Niagara/Source/Niagara/Private/NiagaraSystemGpuComputeProxy.cpp
+++ b/Engine/Plugins/FX/Niagara/Source/Niagara/Private/NiagaraSystemGpuComputeProxy.cpp
@@ -16,6 +16,7 @@ FNiagaraSystemGpuComputeProxy::FNiagaraSystemGpuComputeProxy(FNiagaraSystemInsta
 	bRequiresDepthBuffer = OwnerInstance->RequiresDepthBuffer();
 	bRequiresEarlyViewData = OwnerInstance->RequiresEarlyViewData();
 	bRequiresViewUniformBuffer = OwnerInstance->RequiresViewUniformBuffer();
+	bRequiresRayTracingScene = OwnerInstance->RequiresRayTracingScene();
 
 	// Gather all emitter compute contexts
 	for ( auto& Emitter : OwnerInstance->GetEmitters() )
@@ -27,7 +28,7 @@ FNiagaraSystemGpuComputeProxy::FNiagaraSystemGpuComputeProxy(FNiagaraSystemInsta
 	}
 
 	// Calculate Tick Stage
-	if (bRequiresDistanceFieldData || bRequiresDepthBuffer)
+	if (bRequiresDistanceFieldData || bRequiresDepthBuffer || bRequiresRayTracingScene)
 	{
 		ComputeTickStage = ENiagaraGpuComputeTickStage::PostOpaqueRender;
 	}
