- Ensure DataToRender is set for PreInitViews simulations

#rnx

diff --git a/Engine/Plugins/FX/Niagara/Source/Niagara/Classes/NiagaraEmitterInstanceBatcher.h b/Engine/Plugins/FX/Niagara/Source/Niagara/Classes/NiagaraEmitterInstanceBatcher.h
--- a/Engine/Plugins/FX/Niagara/Source/Niagara/Classes/NiagaraEmitterInstanceBatcher.h
+++ b/Engine/Plugins/FX/Niagara/Source/Niagara/Classes/NiagaraEmitterInstanceBatcher.h
@@ -186,9 +186,7 @@ private:
 	void ResetRayTracingSceneInfo();
 #endif
 
-	bool UseOverlapCompute();
 	void FinishDispatches();
-	void ReleaseTicks();
 
 	void UpdateFreeIDsListSizesBuffer(FRHICommandList& RHICmdList, uint32 NumInstances);
 	void UpdateFreeIDBuffers(FRHICommandList& RHICmdList, TConstArrayView<FNiagaraComputeExecutionContext*> Instances);
diff --git a/Engine/Plugins/FX/Niagara/Source/Niagara/Private/NiagaraEmitterInstanceBatcher.cpp b/Engine/Plugins/FX/Niagara/Source/Niagara/Private/NiagaraEmitterInstanceBatcher.cpp
--- a/Engine/Plugins/FX/Niagara/Source/Niagara/Private/NiagaraEmitterInstanceBatcher.cpp
+++ b/Engine/Plugins/FX/Niagara/Source/Niagara/Private/NiagaraEmitterInstanceBatcher.cpp
@@ -49,24 +49,6 @@ DECLARE_GPU_STAT_NAMED(NiagaraGPUSorting, TEXT("Niagara GPU sorting"));
 
 uint32 FNiagaraComputeExecutionContext::TickCounter = 0;
 
-int32 GNiagaraAllowTickBeforeRender = 1;
-static FAutoConsoleVariableRef CVarNiagaraAllowTickBeforeRender(
-	TEXT("fx.NiagaraAllowTickBeforeRender"),
-	GNiagaraAllowTickBeforeRender,
-	TEXT("If 1, Niagara GPU systems that don't rely on view data will be rendered in sync\n")
-	TEXT("with the current frame simulation instead of the last frame one. (default=1)\n"),
-	ECVF_Default
-);
-
-int32 GNiagaraOverlapCompute = 1;
-static FAutoConsoleVariableRef CVarNiagaraUseAsyncCompute(
-	TEXT("fx.NiagaraOverlapCompute"),
-	GNiagaraOverlapCompute,
-	TEXT("0 - Disable compute dispatch overlap, this will result in poor performance due to resource barriers between each dispatch call, but can be used to debug resource transition issues.\n")
-	TEXT("1 - (Default) Enable compute dispatch overlap where possible, this increases GPU utilization.\n"),
-	ECVF_Default
-);
-
 int32 GNiagaraGpuMaxQueuedRenderFrames = 10;
 static FAutoConsoleVariableRef CVarNiagaraGpuMaxQueuedRenderFrames(
 	TEXT("fx.Niagara.Batcher.MaxQueuedFramesWithoutRender"),
@@ -523,37 +505,27 @@ void NiagaraEmitterInstanceBatcher::ProcessPendingTicksFlush(FRHICommandListImme
 
 void NiagaraEmitterInstanceBatcher::FinishDispatches()
 {
-	ReleaseTicks();
+	check(IsInRenderingThread());
 
-	for (FNiagaraGpuDispatchList& DispatchList : DispatchListPerStage)
+	for (int iTickStage = 0; iTickStage < ENiagaraGpuComputeTickStage::Max; ++iTickStage)
 	{
-		DispatchList.DispatchGroups.Empty();
-		if (DispatchList.CountsToRelease.Num() > 0)
+		for (FNiagaraSystemGpuComputeProxy* ComputeProxy : ProxiesPerStage[iTickStage])
 		{
-			GPUInstanceCounterManager.FreeEntryArray(DispatchList.CountsToRelease);
-			DispatchList.CountsToRelease.Empty();
+			ComputeProxy->ReleaseTicks(GPUInstanceCounterManager);
 		}
 	}
-}
 
-void NiagaraEmitterInstanceBatcher::ReleaseTicks()
-{
-	check(IsInRenderingThread());
-
-	for ( int iTickStage=0; iTickStage < ENiagaraGpuComputeTickStage::Max; ++iTickStage )
+	for (FNiagaraGpuDispatchList& DispatchList : DispatchListPerStage)
 	{
-		for ( FNiagaraSystemGpuComputeProxy* ComputeProxy : ProxiesPerStage[iTickStage] )
+		DispatchList.DispatchGroups.Empty();
+		if (DispatchList.CountsToRelease.Num() > 0)
 		{
-			ComputeProxy->ReleaseTicks(GPUInstanceCounterManager);
+			GPUInstanceCounterManager.FreeEntryArray(DispatchList.CountsToRelease);
+			DispatchList.CountsToRelease.Empty();
 		}
 	}
 }
 
-bool NiagaraEmitterInstanceBatcher::UseOverlapCompute()
-{
-	return !IsMobilePlatform(ShaderPlatform) && GNiagaraOverlapCompute;
-}
-
 void NiagaraEmitterInstanceBatcher::ResetDataInterfaces(FRHICommandList& RHICmdList, const FNiagaraGPUSystemTick& Tick, const FNiagaraComputeInstanceData& InstanceData) const
 {
 	// Note: All stages will contain the same bindings so if they are valid for one they are valid for all, this could change in the future
@@ -1137,20 +1109,6 @@ void NiagaraEmitterInstanceBatcher::PrepareTicksForProxy(FRHICommandListImmediat
 
 		// Build constant buffers for tick
 		BuildConstantBuffers(Tick);
-
-		//// Consume DataInterface instance data
-		////-TODO: This is incorrect, it should be consumed as we execute the tick
-		//if (Tick.DIInstanceData)
-		//{
-		//	uint8* BasePointer = (uint8*)Tick.DIInstanceData->PerInstanceDataForRT;
-
-		//	for (auto& Pair : Tick.DIInstanceData->InterfaceProxiesToOffsets)
-		//	{
-		//		FNiagaraDataInterfaceProxy* Proxy = Pair.Key;
-		//		uint8* InstanceDataPtr = BasePointer + Pair.Value;
-		//		Proxy->ConsumePerInstanceDataFromGameThread(InstanceDataPtr, Tick.SystemInstanceID);
-		//	}
-		//}
 	}
 
 	// Now that all ticks have been processed we can adjust our output buffers to the correct size
@@ -1177,16 +1135,24 @@ void NiagaraEmitterInstanceBatcher::PrepareTicksForProxy(FRHICommandListImmediat
 			ComputeContext->DataBuffers_RT[BufferIndex]->AllocateGPU(RHICmdList, ComputeContext->CurrentMaxAllocateInstances_RT + 1, FeatureLevel, ComputeContext->GetDebugSimName());
 		}
 
-		// Set the translucent draw data and also set the GPU count offset for the draw.
-		// This is safe as simulations have all completed before the translucent pass and allow us to issue
-		// draw commands with this frames data rather than the previous frames.  Note an opaque draw where
-		// the simulation uses the depth buffer can not do this as the simulation can not run until depth
-		// is available for use.
-		extern int32 GNiagaraGpuLowLatencyTranslucencyEnabled;
-		if (GNiagaraGpuLowLatencyTranslucencyEnabled)
+		// RDG will defer the execution of the actual dispatch calls until much later in the frame, so anything that is required for
+		// Mesh Processor commands to execute correct we must setup immediately.  I.e. the final buffer / final count.
+		if (ComputeProxy->GetComputeTickStage() == ENiagaraGpuComputeTickStage::PreInitViews)
 		{
-			ComputeContext->GetCurrDataBuffer()->GPUInstanceCountBufferOffset = ComputeContext->CountOffset_RT;
-			ComputeContext->SetTranslucentDataToRender(ComputeContext->GetCurrDataBuffer());
+			FNiagaraDataBuffer* FinalBuffer = ComputeContext->GetCurrDataBuffer();
+			FinalBuffer->GPUInstanceCountBufferOffset = ComputeContext->CountOffset_RT;
+			FinalBuffer->SetNumInstances(ComputeContext->CurrentNumInstances_RT);
+			ComputeContext->SetDataToRender(FinalBuffer);
+		}
+		// When low latency translucency is enabled we can setup the final buffer / final count here.
+		// This will allow our mesh processor commands to pickup the data for the same frame.
+		// This allows simulations that use the depth buffer, for example, to execute with no latency
+		else if ( GNiagaraGpuLowLatencyTranslucencyEnabled )
+		{
+			FNiagaraDataBuffer* FinalBuffer = ComputeContext->GetCurrDataBuffer();
+			FinalBuffer->GPUInstanceCountBufferOffset = ComputeContext->CountOffset_RT;
+			FinalBuffer->SetNumInstances(ComputeContext->CurrentNumInstances_RT);
+			ComputeContext->SetTranslucentDataToRender(FinalBuffer);
 		}
 	}
 }
@@ -1638,21 +1604,18 @@ void NiagaraEmitterInstanceBatcher::PreInitViews(FRDGBuilder& GraphBuilder, bool
 			DumpDebugFrame();
 		}
 
-		if (GNiagaraAllowTickBeforeRender)
-		{
-			FNiagaraSceneTextureParameters* PassParameters = GraphBuilder.AllocParameters<FNiagaraSceneTextureParameters>();
-			GNiagaraViewDataManager.GetSceneTextureParameters(GraphBuilder, *PassParameters);
+		FNiagaraSceneTextureParameters* PassParameters = GraphBuilder.AllocParameters<FNiagaraSceneTextureParameters>();
+		GNiagaraViewDataManager.GetSceneTextureParameters(GraphBuilder, *PassParameters);
 
-			GraphBuilder.AddPass(
-				RDG_EVENT_NAME("Niagara::PreInitViews"),
-				PassParameters,
-				ERDGPassFlags::Compute | ERDGPassFlags::NeverCull,
-				[this](FRHICommandListImmediate& RHICmdList)
-			{
-				FNiagaraSceneTextureScope Scope;
-				ExecuteTicks(RHICmdList, nullptr, ENiagaraGpuComputeTickStage::PreInitViews);
-			});
-		}
+		GraphBuilder.AddPass(
+			RDG_EVENT_NAME("Niagara::PreInitViews"),
+			PassParameters,
+			ERDGPassFlags::Compute | ERDGPassFlags::NeverCull,
+			[this](FRHICommandListImmediate& RHICmdList)
+		{
+			FNiagaraSceneTextureScope Scope;
+			ExecuteTicks(RHICmdList, nullptr, ENiagaraGpuComputeTickStage::PreInitViews);
+		});
 	}
 	else
 	{
@@ -1722,7 +1685,7 @@ void NiagaraEmitterInstanceBatcher::PostRenderOpaque(FRDGBuilder& GraphBuilder,
 			ProcessDebugReadbacks(RHICmdList, false);
 		}
 
-	if (!GPUInstanceCounterManager.HasPendingGPUReadback() && bRequiresReadback)
+		if (!GPUInstanceCounterManager.HasPendingGPUReadback() && bRequiresReadback)
 		{
 			GPUInstanceCounterManager.EnqueueGPUReadback(RHICmdList);
 		}
