- Redirect for flipbook to baker changes

#rnx

diff --git a/Engine/Plugins/FX/Niagara/Config/BaseNiagara.ini b/Engine/Plugins/FX/Niagara/Config/BaseNiagara.ini
--- a/Engine/Plugins/FX/Niagara/Config/BaseNiagara.ini
+++ b/Engine/Plugins/FX/Niagara/Config/BaseNiagara.ini
@@ -35,6 +35,13 @@
 
 +PropertyRedirects=(OldName="NiagaraRibbonRendererProperties.RibbonDistanceFromStart", NewName="NiagaraRibbonRendererProperties.RibbonUVDistance")
 
+; Flipbook to Baker
++ClassRedirects=(OldName="NiagaraFlipbookSettings", NewName="NiagaraBakerSettings")
++PropertyRedirects=(OldName="NiagaraSystem.FlipbookSettings", NewName="NiagaraSystem.BakerSettings")
++PropertyRedirects=(OldName="NiagaraSystem.FlipbookGeneratedSettings", NewName="NiagaraSystem.BakerGeneratedSettings")
++StructRedirects=(OldName="NiagaraFlipbookTextureSettings", NewName="NiagaraBakerTextureSettings")
++StructRedirects=(OldName="NiagaraFlipbookTextureSource", NewName="NiagaraBakerTextureSource")
+
 [/Script/Niagara.NiagaraSettings]
 +AdditionalParameterEnums=/Niagara/Enums/ENiagaraCoordinateSpace.ENiagaraCoordinateSpace
 +AdditionalParameterEnums=/Niagara/Enums/ENiagaraOrientationAxis.ENiagaraOrientationAxis
