Fix for stomp with dummy view during ProcessPendingTicksFlush

diff --git a/Engine/Plugins/FX/Niagara/Source/Niagara/Private/NiagaraEmitterInstanceBatcher.cpp b/Engine/Plugins/FX/Niagara/Source/Niagara/Private/NiagaraEmitterInstanceBatcher.cpp
--- a/Engine/Plugins/FX/Niagara/Source/Niagara/Private/NiagaraEmitterInstanceBatcher.cpp
+++ b/Engine/Plugins/FX/Niagara/Source/Niagara/Private/NiagaraEmitterInstanceBatcher.cpp
@@ -460,16 +460,7 @@ void NiagaraEmitterInstanceBatcher::ProcessPendingTicksFlush(FRHICommandListImme
 			ViewInitOptions.ViewRotationMatrix = FMatrix::Identity;
 			ViewInitOptions.ProjectionMatrix = FMatrix::Identity;
 
-			FViewInfo* DummyView = nullptr;
-
-			if (RHICmdList.Bypass())
-			{
-				DummyView = new(FMemStack::Get().Alloc(sizeof(FViewInfo), alignof(FViewInfo))) FViewInfo(ViewInitOptions);
-			}
-			else
-			{
-				DummyView = new(RHICmdList.Alloc(sizeof(FViewInfo), alignof(FViewInfo))) FViewInfo(ViewInitOptions);
-			}
+			FViewInfo* DummyView = new(FMemStack::Get().Alloc(sizeof(FViewInfo), alignof(FViewInfo))) FViewInfo(ViewInitOptions);
 
 			DummyView->ViewRect = DummyView->UnscaledViewRect;
 			DummyView->CachedViewUniformShaderParameters = MakeUnique<FViewUniformShaderParameters>();
