- Expose reading custom stencil

diff --git a/Engine/Plugins/FX/Niagara/Shaders/Private/NiagaraDataInterfaceGBuffer.ush b/Engine/Plugins/FX/Niagara/Shaders/Private/NiagaraDataInterfaceGBuffer.ush
--- a/Engine/Plugins/FX/Niagara/Shaders/Private/NiagaraDataInterfaceGBuffer.ush
+++ b/Engine/Plugins/FX/Niagara/Shaders/Private/NiagaraDataInterfaceGBuffer.ush
@@ -167,6 +167,23 @@ void DIGBuffer_DecodeDepth(in float2 ScreenUV, out bool IsValid, out float OutVa
 #endif
 }
 
+#if 0
+// This has not been tested yet
+void DIGBuffer_DecodeStencil(in float2 ScreenUV, out bool IsValid, out int OutValue)
+{
+#if NDIGBUFFER_HAS_GBUFFER
+	IsValid = true;
+
+	uint2 IntUV = (uint2)(ScreenUV.xy * (float2)View.BufferSizeAndInvSize.xy);
+	uint Stencil = SceneStencilTexture.Load(uint3(IntUV, 0)) STENCIL_COMPONENT_SWIZZLE;
+	OutValue = int(Stencil);
+#else
+	IsValid = false;
+	OutValue = 0;
+#endif
+}
+#endif
+
 void DIGBuffer_DecodeCustomDepth(in float2 ScreenUV, out bool IsValid, out float OutValue)
 {
 #if NDIGBUFFER_HAS_GBUFFER
diff --git a/Engine/Plugins/FX/Niagara/Source/Niagara/Private/DataInterface/NiagaraDataInterfaceGBuffer.cpp b/Engine/Plugins/FX/Niagara/Source/Niagara/Private/DataInterface/NiagaraDataInterfaceGBuffer.cpp
--- a/Engine/Plugins/FX/Niagara/Source/Niagara/Private/DataInterface/NiagaraDataInterfaceGBuffer.cpp
+++ b/Engine/Plugins/FX/Niagara/Source/Niagara/Private/DataInterface/NiagaraDataInterfaceGBuffer.cpp
@@ -73,10 +73,10 @@ namespace NiagaraDataInterfaceGBufferLocal
 			FGBufferAttribute(TEXT("Specular"),			TEXT("float"),	FNiagaraTypeDefinition::GetFloatDef(), FText::GetEmpty()),
 			FGBufferAttribute(TEXT("Roughness"),		TEXT("float"),	FNiagaraTypeDefinition::GetFloatDef(), FText::GetEmpty()),
 			FGBufferAttribute(TEXT("Depth"),			TEXT("float"),	FNiagaraTypeDefinition::GetFloatDef(), FText::GetEmpty()),
+			//FGBufferAttribute(TEXT("Stencil"),			TEXT("int"),	FNiagaraTypeDefinition::GetIntDef(), FText::GetEmpty()),
 
 			FGBufferAttribute(TEXT("CustomDepth"),		TEXT("float"),	FNiagaraTypeDefinition::GetFloatDef(), FText::GetEmpty()),
-			// CustomStencil appears broken currently across the board so not exposing until that's working
-			//FGBufferAttribute(TEXT("CustomStencil"),	TEXT("int"),	FNiagaraTypeDefinition::GetIntDef(), FText::GetEmpty()),
+			FGBufferAttribute(TEXT("CustomStencil"),	TEXT("int"),	FNiagaraTypeDefinition::GetIntDef(), FText::GetEmpty()),
 
 			FGBufferAttribute(TEXT("SceneColor"),		TEXT("float4"),	FNiagaraTypeDefinition::GetVec4Def(), GetDescription_SceneColor()),
 		};
