Hardens Niagara ray traced collisions to better handle platforms that don't support HW ray tracing as well as when the runtime environment doesn't have traces enabled
-adds the ShaderPlatfrom to the DI to allow them to better control the shaders being generated

diff --git a/Engine/Plugins/FX/Niagara/Shaders/Private/NiagaraDataInterfaceCollisionQuery.ush b/Engine/Plugins/FX/Niagara/Shaders/Private/NiagaraDataInterfaceCollisionQuery.ush
--- a/Engine/Plugins/FX/Niagara/Shaders/Private/NiagaraDataInterfaceCollisionQuery.ush
+++ b/Engine/Plugins/FX/Niagara/Shaders/Private/NiagaraDataInterfaceCollisionQuery.ush
@@ -2,7 +2,6 @@
 
 #pragma once
 
-#include "/Engine/Private/RayTracing/RayTracingCommon.ush"
 #include "NiagaraRayTracingCommon.ush"
 
 #define NDICOLLISIONQUERY_USE_GBUFFER_NORMAL    (FEATURE_LEVEL >= FEATURE_LEVEL_SM5)
diff --git a/Engine/Plugins/FX/Niagara/Shaders/Private/NiagaraDataInterfaceCollisionQueryTemplate.ush b/Engine/Plugins/FX/Niagara/Shaders/Private/NiagaraDataInterfaceCollisionQueryTemplate.ush
--- a/Engine/Plugins/FX/Niagara/Shaders/Private/NiagaraDataInterfaceCollisionQueryTemplate.ush
+++ b/Engine/Plugins/FX/Niagara/Shaders/Private/NiagaraDataInterfaceCollisionQueryTemplate.ush
@@ -1,6 +1,7 @@
 // Copyright Epic Games, Inc. All Rights Reserved.
 
 #if NIAGARA_USE_RAYTRACING
+bool											RayTracingEnabled_{ParameterName};
 uint                                            MaxRayTraceCount_{ParameterName};
 RWStructuredBuffer<FBasicRayData>               RWRayRequests_{ParameterName};
 StructuredBuffer<FNiagaraRayTracingPayload>     IntersectionResults_{ParameterName};
@@ -25,7 +26,10 @@ void QueryMeshDistanceFieldGPU_{ParameterName}(in float3 In_SamplePos, out float
 void IssueAsyncRayTraceGpu_{ParameterName}_UEImpureCall(int In_QueryID, float3 In_TraceStart, float3 In_TraceEnd, int In_TraceChannel, out bool Out_IsQueryValid)
 {
 #if NIAGARA_USE_RAYTRACING
-	NDICollisionQuery_IssueAsyncRayTrace(MaxRayTraceCount_{ParameterName}, RWRayRequests_{ParameterName}, RWRayTraceCounts_{ParameterName}, In_QueryID, In_TraceStart, In_TraceEnd, In_TraceChannel, Out_IsQueryValid);
+	if (RayTracingEnabled_{ParameterName})
+	{
+		NDICollisionQuery_IssueAsyncRayTrace(MaxRayTraceCount_{ParameterName}, RWRayRequests_{ParameterName}, RWRayTraceCounts_{ParameterName}, In_QueryID, In_TraceStart, In_TraceEnd, In_TraceChannel, Out_IsQueryValid);
+	}
 #else
 	Out_IsQueryValid = false;
 #endif
@@ -34,7 +38,10 @@ void IssueAsyncRayTraceGpu_{ParameterName}_UEImpureCall(int In_QueryID, float3 I
 void ReserveRayTraceIndex_{ParameterName}_UEImpureCall(int In_TraceCount, out int Out_StartQueryID, out bool Out_IndicesValid)
 {
 #if NIAGARA_USE_RAYTRACING
-	NDICollisionQuery_ReserveRayTraceIndex(MaxRayTraceCount_{ParameterName}, RWRayTraceCounts_{ParameterName}, In_TraceCount, Out_StartQueryID, Out_IndicesValid);
+	if (RayTracingEnabled_{ParameterName})
+	{
+		NDICollisionQuery_ReserveRayTraceIndex(MaxRayTraceCount_{ParameterName}, RWRayTraceCounts_{ParameterName}, In_TraceCount, Out_StartQueryID, Out_IndicesValid);
+	}
 #else
 	Out_IndicesValid = false;
 #endif
@@ -43,7 +50,10 @@ void ReserveRayTraceIndex_{ParameterName}_UEImpureCall(int In_TraceCount, out in
 void CreateAsyncRayTraceGpu_{ParameterName}_UEImpureCall(float3 In_TraceStart, float3 In_TraceEnd, int In_TraceChannel, out int Out_QueryID, out bool Out_IsQueryValid)
 {
 #if NIAGARA_USE_RAYTRACING
-	NDICollisionQuery_CreateAsyncRayTrace(MaxRayTraceCount_{ParameterName}, RWRayRequests_{ParameterName}, RWRayTraceCounts_{ParameterName}, In_TraceStart, In_TraceEnd, In_TraceChannel, Out_QueryID, Out_IsQueryValid);
+	if (RayTracingEnabled_{ParameterName})
+	{
+		NDICollisionQuery_CreateAsyncRayTrace(MaxRayTraceCount_{ParameterName}, RWRayRequests_{ParameterName}, RWRayTraceCounts_{ParameterName}, In_TraceStart, In_TraceEnd, In_TraceChannel, Out_QueryID, Out_IsQueryValid);
+	}
 #else
 	Out_IsQueryValid = false;
 #endif
@@ -52,7 +62,10 @@ void CreateAsyncRayTraceGpu_{ParameterName}_UEImpureCall(float3 In_TraceStart, f
 void ReadAsyncRayTraceGpu_{ParameterName}(int In_PreviousFrameQueryID, out bool Out_CollisionValid, out float Out_CollisionDistance, out float3 Out_CollisionPosWorld, out float3 Out_CollisionNormal)
 {
 #if NIAGARA_USE_RAYTRACING
-	NDICollisionQuery_ReadAsyncRayTrace(MaxRayTraceCount_{ParameterName}, IntersectionResults_{ParameterName}, In_PreviousFrameQueryID, Out_CollisionValid, Out_CollisionDistance, Out_CollisionPosWorld, Out_CollisionNormal);
+	if (RayTracingEnabled_{ParameterName})
+	{
+		NDICollisionQuery_ReadAsyncRayTrace(MaxRayTraceCount_{ParameterName}, IntersectionResults_{ParameterName}, In_PreviousFrameQueryID, Out_CollisionValid, Out_CollisionDistance, Out_CollisionPosWorld, Out_CollisionNormal);
+	}
 #else
 	Out_CollisionValid = false;
 #endif
diff --git a/Engine/Plugins/FX/Niagara/Shaders/Private/NiagaraRayTracingCommon.ush b/Engine/Plugins/FX/Niagara/Shaders/Private/NiagaraRayTracingCommon.ush
--- a/Engine/Plugins/FX/Niagara/Shaders/Private/NiagaraRayTracingCommon.ush
+++ b/Engine/Plugins/FX/Niagara/Shaders/Private/NiagaraRayTracingCommon.ush
@@ -2,9 +2,11 @@
 
 #pragma once
 
-#include "/Engine/Private/RayTracing/RayTracingCommon.ush"
+#define NIAGARA_USE_RAYTRACING        (NIAGARA_SUPPORTS_RAY_TRACING && (FEATURE_LEVEL >= FEATURE_LEVEL_SM5))
+
+#if NIAGARA_USE_RAYTRACING
 
-#define NIAGARA_USE_RAYTRACING        (FEATURE_LEVEL >= FEATURE_LEVEL_SM5)
+#include "/Engine/Private/RayTracing/RayTracingCommon.ush"
 
 // payload structure for Niagara ray traces
 struct FNiagaraRayTracingPayload : FMinimalPayload
@@ -12,3 +14,5 @@ struct FNiagaraRayTracingPayload : FMinimalPayload
 	float3 WorldPosition;
 	float3 WorldNormal;
 };
+
+#endif // NIAGARA_USE_RAYTRACING
\ No newline at end of file
diff --git a/Engine/Plugins/FX/Niagara/Source/Niagara/Classes/NiagaraDataInterface.h b/Engine/Plugins/FX/Niagara/Source/Niagara/Classes/NiagaraDataInterface.h
--- a/Engine/Plugins/FX/Niagara/Source/Niagara/Classes/NiagaraDataInterface.h
+++ b/Engine/Plugins/FX/Niagara/Source/Niagara/Classes/NiagaraDataInterface.h
@@ -294,7 +294,7 @@ public:
 #endif
 
 	/** Allows data interfaces to influence the compilation of GPU shaders and is only called on the CDO object not the instance. */
-	virtual void ModifyCompilationEnvironment(struct FShaderCompilerEnvironment& OutEnvironment) const {}
+	virtual void ModifyCompilationEnvironment(EShaderPlatform ShaderPlatform, struct FShaderCompilerEnvironment& OutEnvironment) const {}
 
 	/** 
 		Subclasses that wish to work with GPU systems/emitters must implement this.
diff --git a/Engine/Plugins/FX/Niagara/Source/Niagara/Classes/NiagaraDataInterfaceCollisionQuery.h b/Engine/Plugins/FX/Niagara/Source/Niagara/Classes/NiagaraDataInterfaceCollisionQuery.h
--- a/Engine/Plugins/FX/Niagara/Source/Niagara/Classes/NiagaraDataInterfaceCollisionQuery.h
+++ b/Engine/Plugins/FX/Niagara/Source/Niagara/Classes/NiagaraDataInterfaceCollisionQuery.h
@@ -70,6 +70,7 @@ public:
 
 #if WITH_EDITORONLY_DATA
 	virtual bool AppendCompileHash(FNiagaraCompileHashVisitor* InVisitor) const override;
+	virtual void ModifyCompilationEnvironment(EShaderPlatform ShaderPlatform, struct FShaderCompilerEnvironment& OutEnvironment) const override;
 	virtual void GetCommonHLSL(FString& OutHLSL) override;
 	virtual void GetParameterDefinitionHLSL(const FNiagaraDataInterfaceGPUParamInfo& ParamInfo, FString& OutHLSL) override;
 	virtual bool GetFunctionHLSL(const FNiagaraDataInterfaceGPUParamInfo& ParamInfo, const FNiagaraDataInterfaceGeneratedFunction& FunctionInfo, int FunctionInstanceIndex, FString& OutHLSL) override;
diff --git a/Engine/Plugins/FX/Niagara/Source/Niagara/Classes/NiagaraDataInterfaceSkeletalMesh.h b/Engine/Plugins/FX/Niagara/Source/Niagara/Classes/NiagaraDataInterfaceSkeletalMesh.h
--- a/Engine/Plugins/FX/Niagara/Source/Niagara/Classes/NiagaraDataInterfaceSkeletalMesh.h
+++ b/Engine/Plugins/FX/Niagara/Source/Niagara/Classes/NiagaraDataInterfaceSkeletalMesh.h
@@ -801,7 +801,7 @@ public:
 	virtual bool HasTickGroupPrereqs() const override { return true; }
 	virtual ETickingGroup CalculateTickGroup(const void* PerInstanceData) const override;
 
-	virtual void ModifyCompilationEnvironment(struct FShaderCompilerEnvironment& OutEnvironment) const override;
+	virtual void ModifyCompilationEnvironment(EShaderPlatform ShaderPlatform, struct FShaderCompilerEnvironment& OutEnvironment) const override;
 #if WITH_EDITORONLY_DATA
 	virtual bool AppendCompileHash(FNiagaraCompileHashVisitor* InVisitor) const override;
 	virtual void GetCommonHLSL(FString& OutHLSL) override;
diff --git a/Engine/Plugins/FX/Niagara/Source/Niagara/Classes/NiagaraScript.h b/Engine/Plugins/FX/Niagara/Source/Niagara/Classes/NiagaraScript.h
--- a/Engine/Plugins/FX/Niagara/Source/Niagara/Classes/NiagaraScript.h
+++ b/Engine/Plugins/FX/Niagara/Source/Niagara/Classes/NiagaraScript.h
@@ -825,7 +825,7 @@ public:
 	//~ End UObject interface
 
 	//~ Begin UNiagaraScriptBase interface
-	virtual void ModifyCompilationEnvironment(struct FShaderCompilerEnvironment& OutEnvironment) const override;
+	virtual void ModifyCompilationEnvironment(EShaderPlatform Platform, struct FShaderCompilerEnvironment& OutEnvironment) const override;
 	virtual TConstArrayView<FSimulationStageMetaData> GetSimulationStageMetaData() const override { return MakeArrayView(CachedScriptVM.SimulationStageMetaData); }
 	//~ End UNiagaraScriptBase interface
 
diff --git a/Engine/Plugins/FX/Niagara/Source/Niagara/Private/NiagaraDataInterfaceCollisionQuery.cpp b/Engine/Plugins/FX/Niagara/Source/Niagara/Private/NiagaraDataInterfaceCollisionQuery.cpp
--- a/Engine/Plugins/FX/Niagara/Source/Niagara/Private/NiagaraDataInterfaceCollisionQuery.cpp
+++ b/Engine/Plugins/FX/Niagara/Source/Niagara/Private/NiagaraDataInterfaceCollisionQuery.cpp
@@ -33,6 +33,7 @@ namespace NDICollisionQueryLocal
 	static const FName ReserveAsyncRayTraceName(TEXT("ReserveAsyncRayTraceGpu"));
 	static const FName ReadAsyncRayTraceName(TEXT("ReadAsyncRayTraceGpu"));
 
+	static const FString RayTracingEnabledParamName(TEXT("RayTracingEnabled_"));
 	static const FString MaxRayTraceCountParamName(TEXT("MaxRayTraceCount_"));
 	static const FString RayRequestsParamName(TEXT("RayRequests_"));
 	static const FString IntersectionResultsParamName(TEXT("IntersectionResults_"));
@@ -690,6 +691,14 @@ bool UNiagaraDataInterfaceCollisionQuery::AppendCompileHash(FNiagaraCompileHashV
 
 	return true;
 }
+
+void UNiagaraDataInterfaceCollisionQuery::ModifyCompilationEnvironment(EShaderPlatform ShaderPlatform, FShaderCompilerEnvironment& OutEnvironment) const
+{
+	Super::ModifyCompilationEnvironment(ShaderPlatform, OutEnvironment);
+
+	OutEnvironment.SetDefine(TEXT("NIAGARA_SUPPORTS_RAY_TRACING"), ShouldCompileRayTracingShadersForProject(ShaderPlatform) ? 1 : 0);
+}
+
 #endif
 
 void UNiagaraDataInterfaceCollisionQuery::PerformQuerySyncCPU(FVectorVMContext & Context)
@@ -982,6 +991,7 @@ public:
 	{
 		GlobalDistanceFieldParameters.Bind(ParameterMap);
 #if RHI_RAYTRACING
+		RayTracingEnabledParam.Bind(ParameterMap, *(NDICollisionQueryLocal::RayTracingEnabledParamName + ParameterInfo.DataInterfaceHLSLSymbol));
 		MaxRayTraceCountParam.Bind(ParameterMap, *(NDICollisionQueryLocal::MaxRayTraceCountParamName + ParameterInfo.DataInterfaceHLSLSymbol));
 		RayRequestsParam.Bind(ParameterMap, *(NDICollisionQueryLocal::RayRequestsParamName + ParameterInfo.DataInterfaceHLSLSymbol));
 		IntersectionResultsParam.Bind(ParameterMap, *(NDICollisionQueryLocal::IntersectionResultsParamName + ParameterInfo.DataInterfaceHLSLSymbol));
@@ -1003,6 +1013,7 @@ public:
 		}
 
 #if RHI_RAYTRACING
+		SetShaderValue(RHICmdList, ComputeShaderRHI, RayTracingEnabledParam, IsRayTracingEnabled() ? 1 : 0);
 		SetShaderValue(RHICmdList, ComputeShaderRHI, MaxRayTraceCountParam, QueryDI->MaxRayTraceCount);
 
 		if (RayRequestsParam.IsUAVBound())
@@ -1044,6 +1055,7 @@ private:
 	LAYOUT_FIELD(FGlobalDistanceFieldParameters, GlobalDistanceFieldParameters);
 
 #if RHI_RAYTRACING
+	LAYOUT_FIELD(FShaderParameter, RayTracingEnabledParam);
 	LAYOUT_FIELD(FShaderParameter, MaxRayTraceCountParam);
 	LAYOUT_FIELD(FRWShaderParameter, RayRequestsParam);
 	LAYOUT_FIELD(FShaderResourceParameter, IntersectionResultsParam);
diff --git a/Engine/Plugins/FX/Niagara/Source/Niagara/Private/NiagaraDataInterfaceSkeletalMesh.cpp b/Engine/Plugins/FX/Niagara/Source/Niagara/Private/NiagaraDataInterfaceSkeletalMesh.cpp
--- a/Engine/Plugins/FX/Niagara/Source/Niagara/Private/NiagaraDataInterfaceSkeletalMesh.cpp
+++ b/Engine/Plugins/FX/Niagara/Source/Niagara/Private/NiagaraDataInterfaceSkeletalMesh.cpp
@@ -2968,9 +2968,9 @@ bool UNiagaraDataInterfaceSkeletalMesh::AppendCompileHash(FNiagaraCompileHashVis
 }
 #endif
 
-void UNiagaraDataInterfaceSkeletalMesh::ModifyCompilationEnvironment(struct FShaderCompilerEnvironment& OutEnvironment) const
+void UNiagaraDataInterfaceSkeletalMesh::ModifyCompilationEnvironment(EShaderPlatform ShaderPlatform, struct FShaderCompilerEnvironment& OutEnvironment) const
 {
-	Super::ModifyCompilationEnvironment(OutEnvironment);
+	Super::ModifyCompilationEnvironment(ShaderPlatform, OutEnvironment);
 
 	OutEnvironment.SetDefine(TEXT("DISKELMESH_BONE_INFLUENCES"), int(GetDefault<UNiagaraSettings>()->NDISkelMesh_GpuMaxInfluences));
 	OutEnvironment.SetDefine(TEXT("DISKELMESH_PROBALIAS_FORMAT"), int(GetDefault<UNiagaraSettings>()->NDISkelMesh_GpuUniformSamplingFormat));
diff --git a/Engine/Plugins/FX/Niagara/Source/Niagara/Private/NiagaraScript.cpp b/Engine/Plugins/FX/Niagara/Source/Niagara/Private/NiagaraScript.cpp
--- a/Engine/Plugins/FX/Niagara/Source/Niagara/Private/NiagaraScript.cpp
+++ b/Engine/Plugins/FX/Niagara/Source/Niagara/Private/NiagaraScript.cpp
@@ -2333,7 +2333,7 @@ bool UNiagaraScript::IsEditorOnly() const
 	return Super::IsEditorOnly();
 }
 
-void UNiagaraScript::ModifyCompilationEnvironment(struct FShaderCompilerEnvironment& OutEnvironment) const
+void UNiagaraScript::ModifyCompilationEnvironment(EShaderPlatform Platform, struct FShaderCompilerEnvironment& OutEnvironment) const
 {
 	// Add all data interfaces
 	TSet<UClass*> DIUniqueClasses;
@@ -2350,7 +2350,7 @@ void UNiagaraScript::ModifyCompilationEnvironment(struct FShaderCompilerEnvironm
 	{
 		if ( UNiagaraDataInterface* DICDO = CastChecked<UNiagaraDataInterface>(DIClass->GetDefaultObject(true)) )
 		{
-			DICDO->ModifyCompilationEnvironment(OutEnvironment);
+			DICDO->ModifyCompilationEnvironment(Platform, OutEnvironment);
 		}
 	}
 }
diff --git a/Engine/Plugins/FX/Niagara/Source/NiagaraShader/Private/NiagaraRayTracingHelper.cpp b/Engine/Plugins/FX/Niagara/Source/NiagaraShader/Private/NiagaraRayTracingHelper.cpp
--- a/Engine/Plugins/FX/Niagara/Source/NiagaraShader/Private/NiagaraRayTracingHelper.cpp
+++ b/Engine/Plugins/FX/Niagara/Source/NiagaraShader/Private/NiagaraRayTracingHelper.cpp
@@ -34,6 +34,12 @@ class FNiagaraCollisionRayTraceRG : public FGlobalShader
 		return ShouldCompileRayTracingShadersForProject(Parameters.Platform);
 	}
 
+	static void ModifyCompilationEnvironment(const FGlobalShaderPermutationParameters& Parameters, FShaderCompilerEnvironment& OutEnvironment)
+	{
+		FGlobalShader::ModifyCompilationEnvironment(Parameters, OutEnvironment);
+		OutEnvironment.SetDefine(TEXT("NIAGARA_SUPPORTS_RAY_TRACING"), 1);
+	}
+
 	static FRHIRayTracingShader* GetShader(FGlobalShaderMap* ShaderMap)
 	{
 		FPermutationDomain PermutationVector;
@@ -57,6 +63,12 @@ class FNiagaraCollisionRayTraceCH : public FGlobalShader
 		return ShouldCompileRayTracingShadersForProject(Parameters.Platform);
 	}
 
+	static void ModifyCompilationEnvironment(const FGlobalShaderPermutationParameters& Parameters, FShaderCompilerEnvironment& OutEnvironment)
+	{
+		FGlobalShader::ModifyCompilationEnvironment(Parameters, OutEnvironment);
+		OutEnvironment.SetDefine(TEXT("NIAGARA_SUPPORTS_RAY_TRACING"), 1);
+	}
+
 	FNiagaraCollisionRayTraceCH() = default;
 	FNiagaraCollisionRayTraceCH(const ShaderMetaType::CompiledShaderInitializerType& Initializer)
 		: FGlobalShader(Initializer)
@@ -72,6 +84,12 @@ class FNiagaraCollisionRayTraceMiss : public FGlobalShader
 		return ShouldCompileRayTracingShadersForProject(Parameters.Platform);
 	}
 
+	static void ModifyCompilationEnvironment(const FGlobalShaderPermutationParameters& Parameters, FShaderCompilerEnvironment& OutEnvironment)
+	{
+		FGlobalShader::ModifyCompilationEnvironment(Parameters, OutEnvironment);
+		OutEnvironment.SetDefine(TEXT("NIAGARA_SUPPORTS_RAY_TRACING"), 1);
+	}
+
 	FNiagaraCollisionRayTraceMiss() = default;
 	FNiagaraCollisionRayTraceMiss(const ShaderMetaType::CompiledShaderInitializerType& Initializer)
 		: FGlobalShader(Initializer)
diff --git a/Engine/Plugins/FX/Niagara/Source/NiagaraShader/Private/NiagaraShader.cpp b/Engine/Plugins/FX/Niagara/Source/NiagaraShader/Private/NiagaraShader.cpp
--- a/Engine/Plugins/FX/Niagara/Source/NiagaraShader/Private/NiagaraShader.cpp
+++ b/Engine/Plugins/FX/Niagara/Source/NiagaraShader/Private/NiagaraShader.cpp
@@ -431,7 +431,7 @@ void FNiagaraShaderType::BeginCompileShader(
 
 	NewJob->Input.Environment.SetDefine(TEXT("NIAGARA_COMPRESSED_ATTRIBUTES_ENABLED"), Script->GetUsesCompressedAttributes() ? 1 : 0);
 
-	Script->ModifyCompilationEnvironment(NewJob->Input.Environment);
+	Script->ModifyCompilationEnvironment(Platform, NewJob->Input.Environment);
 
 	AddReferencedUniformBufferIncludes(NewJob->Input.Environment, NewJob->Input.SourceFilePrefix, (EShaderPlatform)Target.Platform);
 	
diff --git a/Engine/Plugins/FX/Niagara/Source/NiagaraShader/Private/NiagaraShared.cpp b/Engine/Plugins/FX/Niagara/Source/NiagaraShader/Private/NiagaraShared.cpp
--- a/Engine/Plugins/FX/Niagara/Source/NiagaraShader/Private/NiagaraShared.cpp
+++ b/Engine/Plugins/FX/Niagara/Source/NiagaraShader/Private/NiagaraShared.cpp
@@ -57,11 +57,11 @@ bool FNiagaraShaderScript::ShouldCache(EShaderPlatform Platform, const FShaderTy
 	return true;
 }
 
-void FNiagaraShaderScript::ModifyCompilationEnvironment(struct FShaderCompilerEnvironment& OutEnvironment) const
+void FNiagaraShaderScript::ModifyCompilationEnvironment(EShaderPlatform Platform, struct FShaderCompilerEnvironment& OutEnvironment) const
 {
 	if ( BaseVMScript )
 	{
-		BaseVMScript->ModifyCompilationEnvironment(OutEnvironment);
+		BaseVMScript->ModifyCompilationEnvironment(Platform, OutEnvironment);
 	}
 }
 
diff --git a/Engine/Plugins/FX/Niagara/Source/NiagaraShader/Public/NiagaraScriptBase.h b/Engine/Plugins/FX/Niagara/Source/NiagaraShader/Public/NiagaraScriptBase.h
--- a/Engine/Plugins/FX/Niagara/Source/NiagaraShader/Public/NiagaraScriptBase.h
+++ b/Engine/Plugins/FX/Niagara/Source/NiagaraShader/Public/NiagaraScriptBase.h
@@ -51,6 +51,6 @@ class UNiagaraScriptBase : public UObject
 {
 	GENERATED_UCLASS_BODY()
 
-	virtual void ModifyCompilationEnvironment(struct FShaderCompilerEnvironment& OutEnvironment) const PURE_VIRTUAL(UNiagaraScriptBase::ModifyCompilationEnvironment, );
+	virtual void ModifyCompilationEnvironment(EShaderPlatform Platform, struct FShaderCompilerEnvironment& OutEnvironment) const PURE_VIRTUAL(UNiagaraScriptBase::ModifyCompilationEnvironment, );
 	virtual TConstArrayView<FSimulationStageMetaData> GetSimulationStageMetaData() const PURE_VIRTUAL(UNiagaraScriptBase::GetSimulationStageMetaData(), return MakeArrayView<FSimulationStageMetaData>(nullptr, 0); )
 };
diff --git a/Engine/Plugins/FX/Niagara/Source/NiagaraShader/Public/NiagaraShared.h b/Engine/Plugins/FX/Niagara/Source/NiagaraShader/Public/NiagaraShared.h
--- a/Engine/Plugins/FX/Niagara/Source/NiagaraShader/Public/NiagaraShared.h
+++ b/Engine/Plugins/FX/Niagara/Source/NiagaraShader/Public/NiagaraShared.h
@@ -667,7 +667,7 @@ public:
 	 * Allow Niagara script the opportunity to modify the compilation environment for GPU simulations.
 	 * Generally used for data interfaces to inject project defines.
 	*/
-	NIAGARASHADER_API  virtual void ModifyCompilationEnvironment(struct FShaderCompilerEnvironment& OutEnvironment) const;
+	NIAGARASHADER_API  virtual void ModifyCompilationEnvironment(EShaderPlatform Platform, struct FShaderCompilerEnvironment& OutEnvironment) const;
 
 	/** Serializes the script. */
 	NIAGARASHADER_API  virtual void LegacySerialize(FArchive& Ar);
