Merging CLs 16272008, 16283208, 16286746, 16287853:
* Moved emitter merging and compilation out of system/emitter postload
* Updated emitter merge logic to be able to run python upgrade scripts during merge

diff --git a/Engine/Plugins/FX/Niagara/Source/Niagara/Classes/NiagaraEmitter.h b/Engine/Plugins/FX/Niagara/Source/Niagara/Classes/NiagaraEmitter.h
--- a/Engine/Plugins/FX/Niagara/Source/Niagara/Classes/NiagaraEmitter.h
+++ b/Engine/Plugins/FX/Niagara/Source/Niagara/Classes/NiagaraEmitter.h
@@ -387,6 +387,11 @@ public:
 	void CacheFromCompiledData(const FNiagaraDataSetCompiledData* CompiledData);
 	void CacheFromShaderCompiled();
 
+	// Reference to the async update task that is responsible to do any outstanding work like emitter merging on the game thread
+	FGraphEventRef UpdateTaskRef;
+
+	void UpdateEmitterAfterLoad();
+
 #if WITH_EDITORONLY_DATA
 	/** 'Source' data/graphs for the scripts used by this emitter. */
 	UPROPERTY()
@@ -606,6 +611,8 @@ private:
 	void RaiseOnEmitterGPUCompiled(UNiagaraScript* InScript, const FGuid& ScriptVersion);
 #endif
 
+	bool bFullyLoaded = false;
+
 #if !UE_BUILD_SHIPPING
 	FString DebugSimName;
 #endif
diff --git a/Engine/Plugins/FX/Niagara/Source/Niagara/Classes/NiagaraScript.h b/Engine/Plugins/FX/Niagara/Source/Niagara/Classes/NiagaraScript.h
--- a/Engine/Plugins/FX/Niagara/Source/Niagara/Classes/NiagaraScript.h
+++ b/Engine/Plugins/FX/Niagara/Source/Niagara/Classes/NiagaraScript.h
@@ -825,7 +825,7 @@ public:
 	//~ End UObject interface
 
 	//~ Begin UNiagaraScriptBase interface
-	virtual void ModifyCompilationEnvironment(EShaderPlatform Platform, struct FShaderCompilerEnvironment& OutEnvironment) const override;
+	virtual void ModifyCompilationEnvironment(struct FShaderCompilerEnvironment& OutEnvironment) const override;
 	virtual TConstArrayView<FSimulationStageMetaData> GetSimulationStageMetaData() const override { return MakeArrayView(CachedScriptVM.SimulationStageMetaData); }
 	//~ End UNiagaraScriptBase interface
 
diff --git a/Engine/Plugins/FX/Niagara/Source/Niagara/Classes/NiagaraSystem.h b/Engine/Plugins/FX/Niagara/Source/Niagara/Classes/NiagaraSystem.h
--- a/Engine/Plugins/FX/Niagara/Source/Niagara/Classes/NiagaraSystem.h
+++ b/Engine/Plugins/FX/Niagara/Source/Niagara/Classes/NiagaraSystem.h
@@ -239,6 +239,7 @@ public:
 
 private:
 	bool IsValidInternal() const;
+	void EnsureFullyLoaded() const;
 public:
 	/** Returns true if this system is valid and can be instanced. False otherwise. */
 	bool IsValid() const { return FPlatformProperties::RequiresCookedData() ? bIsValidCached : IsValidInternal(); }
@@ -248,6 +249,9 @@ public:
 		original asset. */
 	FNiagaraEmitterHandle AddEmitterHandle(UNiagaraEmitter& SourceEmitter, FName EmitterName);
 
+	/** Adds a new emitter handle to this system without copying the original asset. This should only be used for temporary systems and never for live assets. */
+	void AddEmitterHandleDirect(FNiagaraEmitterHandle& EmitterHandleToAdd);
+
 	/** Duplicates an existing emitter handle and adds it to the System.  The new handle will reference the same source asset,
 		but will have a copy of the duplicated Instance value. */
 	FNiagaraEmitterHandle DuplicateEmitterHandle(const FNiagaraEmitterHandle& EmitterHandleToDuplicate, FName EmitterName);
@@ -381,6 +385,8 @@ public:
 	FNiagaraSystemUpdateContext UpdateContext;
 #endif
 
+	void UpdateSystemAfterLoad();
+
 	bool ShouldAutoDeactivate() const { return bAutoDeactivate; }
 	bool IsLooping() const;
 
@@ -468,6 +474,7 @@ public:
 	UPROPERTY(EditAnywhere, Category = "Debug")
 	bool bDumpDebugEmitterInfo;
 
+	bool bFullyLoaded = false;
 
 	/** When enabled, we follow the settings on the UNiagaraComponent for tick order. When this option is disabled, we ignore any dependencies from data interfaces or other variables and instead fire off the simulation as early in the frame as possible. This greatly
 	reduces overhead and allows the game thread to run faster, but comes at a tradeoff if the dependencies might leave gaps or other visual artifacts.*/
@@ -651,6 +658,9 @@ protected:
 	UNiagaraBakerSettings* BakerGeneratedSettings;
 #endif
 
+	// Reference to the async system update task that is responsible to do any outstanding work like emitter merging or script compilation 
+	FGraphEventRef UpdateTaskRef;
+	
 	UPROPERTY()
 	bool bHasSystemScriptDIsWithPerInstanceData;
 
diff --git a/Engine/Plugins/FX/Niagara/Source/Niagara/Private/NiagaraEmitter.cpp b/Engine/Plugins/FX/Niagara/Source/Niagara/Private/NiagaraEmitter.cpp
--- a/Engine/Plugins/FX/Niagara/Source/Niagara/Private/NiagaraEmitter.cpp
+++ b/Engine/Plugins/FX/Niagara/Source/Niagara/Private/NiagaraEmitter.cpp
@@ -65,6 +65,33 @@
 	ECVF_Default
 );
 
+/**
+Game thread task to do any outstanding work on the loaded emitter
+*/
+struct FNiagaraEmitterUpdateTask
+{
+	FNiagaraEmitterUpdateTask(TWeakObjectPtr<UNiagaraEmitter> InEmitter) : WeakEmitter(InEmitter)
+	{
+	}
+
+	FORCEINLINE TStatId GetStatId() const { RETURN_QUICK_DECLARE_CYCLE_STAT(FNiagaraEmitterUpdateTask, STATGROUP_TaskGraphTasks); }
+	static FORCEINLINE ENamedThreads::Type GetDesiredThread() { return ENamedThreads::GameThread; }
+	static FORCEINLINE ESubsequentsMode::Type GetSubsequentsMode() { return ESubsequentsMode::TrackSubsequents; }
+
+	void DoTask(ENamedThreads::Type CurrentThread, const FGraphEventRef& MyCompletionGraphEvent)
+	{
+		UNiagaraEmitter* Emitter = WeakEmitter.Get();
+		if (Emitter == nullptr)
+		{
+			// probably got garbage collected in the meantime
+			return;
+		}
+		Emitter->UpdateEmitterAfterLoad();
+	}
+
+	TWeakObjectPtr<UNiagaraEmitter> WeakEmitter;
+};
+
 FNiagaraDetailsLevelScaleOverrides::FNiagaraDetailsLevelScaleOverrides()
 {
 	Low = 0.125f;
@@ -351,34 +378,6 @@
 		SetFlags(RF_Transactional);
 	}
 
-	for (int32 RendererIndex = RendererProperties.Num() - 1; RendererIndex >= 0; --RendererIndex)
-	{
-		if (RendererProperties[RendererIndex] == nullptr)
-		{
-#if WITH_EDITORONLY_DATA
-			//In cooked builds these can be cooked out and null on purpose.
-			ensureMsgf(IsCooked, TEXT("Null renderer found in %s at index %i, removing it to prevent crashes."), *GetPathName(), RendererIndex);
-#endif
-			RendererProperties.RemoveAt(RendererIndex);
-		}
-		else
-		{
-			RendererProperties[RendererIndex]->ConditionalPostLoad();
-		}
-	}
-
-	for (int32 SimulationStageIndex = SimulationStages.Num() - 1; SimulationStageIndex >= 0; --SimulationStageIndex)
-	{
-		if (ensureMsgf(SimulationStages[SimulationStageIndex] != nullptr && SimulationStages[SimulationStageIndex]->Script != nullptr, TEXT("Null simulation stage, or simulation stage with a null script found in %s at index %i, removing it to prevent crashes."), *GetPathName(), SimulationStageIndex) == false)
-		{
-			SimulationStages.RemoveAt(SimulationStageIndex);
-		}
-		else
-		{
-			SimulationStages[SimulationStageIndex]->ConditionalPostLoad();
-		}
-	}
-
 	const int32 NiagaraVer = GetLinkerCustomVersion(FNiagaraCustomVersion::GUID);
 	if (NiagaraVer < FNiagaraCustomVersion::PlatformScalingRefactor)
 	{
@@ -428,6 +427,39 @@
 		}
 	}
 
+	for (int32 RendererIndex = RendererProperties.Num() - 1; RendererIndex >= 0; --RendererIndex)
+	{
+		if (RendererProperties[RendererIndex] == nullptr)
+		{
 #if WITH_EDITORONLY_DATA
+			//In cooked builds these can be cooked out and null on purpose.
+			ensureMsgf(IsCooked, TEXT("Null renderer found in %s at index %i, removing it to prevent crashes."), *GetPathName(), RendererIndex);
+#endif
+			RendererProperties.RemoveAt(RendererIndex);
+		}
+		else
+		{
+			RendererProperties[RendererIndex]->ConditionalPostLoad();
+		}
+	}
+
+	for (int32 SimulationStageIndex = SimulationStages.Num() - 1; SimulationStageIndex >= 0; --SimulationStageIndex)
+	{
+		if (ensureMsgf(SimulationStages[SimulationStageIndex] != nullptr && SimulationStages[SimulationStageIndex]->Script != nullptr, TEXT("Null simulation stage, or simulation stage with a null script found in %s at index %i, removing it to prevent crashes."), *GetPathName(), SimulationStageIndex) == false)
+		{
+			SimulationStages.RemoveAt(SimulationStageIndex);
+		}
+		else
+		{
+			SimulationStages[SimulationStageIndex]->ConditionalPostLoad();
+		}
+	}
+
+	if (SpawnScriptProps.Script)
+	{
+		SpawnScriptProps.Script->ConditionalPostLoad();
+	}
+	
+#if WITH_EDITORONLY_DATA
 	if (!GPUComputeScript)
 	{
@@ -435,13 +467,8 @@
 		GPUComputeScript->SetUsage(ENiagaraScriptUsage::ParticleGPUComputeScript);
 		GPUComputeScript->SetLatestSource(SpawnScriptProps.Script ? SpawnScriptProps.Script->GetLatestSource() : nullptr);
 	}
-
 	GPUComputeScript->OnGPUScriptCompiled().AddUObject(this, &UNiagaraEmitter::RaiseOnEmitterGPUCompiled);
-#else
-	check(GPUComputeScript == nullptr || SimTarget == ENiagaraSimTarget::GPUComputeSim);
-#endif
 
-#if WITH_EDITORONLY_DATA
 	if (EmitterSpawnScriptProps.Script == nullptr || EmitterUpdateScriptProps.Script == nullptr)
 	{
 		EmitterSpawnScriptProps.Script = NewObject<UNiagaraScript>(this, "EmitterSpawnScript", EObjectFlags::RF_Transactional);
@@ -456,48 +483,13 @@
 			EmitterUpdateScriptProps.Script->SetLatestSource(SpawnScriptProps.Script->GetLatestSource());
 		}
 	}
-#endif
-
-	//Temporarily disabling interpolated spawn if the script type and flag don't match.
-	if (SpawnScriptProps.Script)
-	{
-		SpawnScriptProps.Script->ConditionalPostLoad();
-		bool bActualInterpolatedSpawning = SpawnScriptProps.Script->IsInterpolatedParticleSpawnScript();
-		if (bInterpolatedSpawning != bActualInterpolatedSpawning)
-		{
-			bInterpolatedSpawning = false;
-			if (bActualInterpolatedSpawning)
-			{
-#if WITH_EDITORONLY_DATA
-				SpawnScriptProps.Script->InvalidateCompileResults(TEXT("Interpolated spawn changed."));//clear out the script as it was compiled with interpolated spawn.
-#endif
-				SpawnScriptProps.Script->SetUsage(ENiagaraScriptUsage::ParticleSpawnScript);
-			}
-			UE_LOG(LogNiagara, Warning, TEXT("Disabling interpolated spawn because emitter flag and script type don't match. Did you adjust this value in the UI? Emitter may need recompile.. %s"), *GetFullName());
-		}
-	}
-
-#if WITH_EDITORONLY_DATA
-	if (GetOuter()->IsA<UNiagaraEmitter>())
-	{
-		// If this emitter is owned by another emitter, remove it's inheritance information so that it doesn't try to merge changes.
-		Parent = nullptr;
-		ParentAtLastMerge = nullptr;
-	}
 
 	if (!GetOutermost()->bIsCookedForEditor)
 	{
 		GraphSource->ConditionalPostLoad();
 		GraphSource->PostLoadFromEmitter(*this);
-	}
-#endif
-
-	EnsureScriptsPostLoaded();
-
-#if WITH_EDITORONLY_DATA
-	if (!GetOutermost()->bIsCookedForEditor)
-	{
-		// Handle emitter inheritance.
+		
+		// Prepare for emitter inheritance.
 		if (Parent != nullptr)
 		{
 			Parent->ConditionalPostLoad();
@@ -527,99 +519,34 @@
 		{
 			// Modify here so that the asset will be marked dirty when using the resave commandlet.  This will be ignored during regular post load.
 			Modify();
-			MergeChangesFromParent();
-		}
-
-		// Reset scripts if recompile is forced.
-		bool bGenerateNewChangeId = false;
-		FString GenerateNewChangeIdReason;
-		if (GetForceCompileOnLoad())
-		{
-			// If we are a standalone emitter, then we invalidate id's, which should cause systems dependent on us to regenerate.
-			UObject* OuterObj = GetOuter();
-			if (OuterObj == GetOutermost())
-			{
-				GraphSource->ForceGraphToRecompileOnNextCheck();
-				bGenerateNewChangeId = true;
-				GenerateNewChangeIdReason = TEXT("PostLoad - Force compile on load");
-				if (GEnableVerboseNiagaraChangeIdLogging)
-				{
-					UE_LOG(LogNiagara, Log, TEXT("InvalidateCachedCompileIds for %s because GbForceNiagaraCompileOnLoad = %d"), *GetPathName(), GbForceNiagaraCompileOnLoad);
-				}
-			}
-		}
-	
-		if (ChangeId.IsValid() == false)
-		{
-			// If the change id is already invalid we need to generate a new one, and can skip checking the owned scripts.
-			bGenerateNewChangeId = true;
-			GenerateNewChangeIdReason = TEXT("PostLoad - Change id was invalid.");
-			if (GEnableVerboseNiagaraChangeIdLogging)
-			{
-				UE_LOG(LogNiagara, Log, TEXT("Change ID updated for emitter %s because the ID was invalid."), *GetPathName());
-			}
-		}
-
-		if (bGenerateNewChangeId)
-		{
-			UpdateChangeId(GenerateNewChangeIdReason);
-		}
-
-		GraphSource->OnChanged().AddUObject(this, &UNiagaraEmitter::GraphSourceChanged);
-
-		EmitterSpawnScriptProps.Script->RapidIterationParameters.AddOnChangedHandler(
-			FNiagaraParameterStore::FOnChanged::FDelegate::CreateUObject(this, &UNiagaraEmitter::ScriptRapidIterationParameterChanged));
-		EmitterUpdateScriptProps.Script->RapidIterationParameters.AddOnChangedHandler(
-			FNiagaraParameterStore::FOnChanged::FDelegate::CreateUObject(this, &UNiagaraEmitter::ScriptRapidIterationParameterChanged));
-
-		if (SpawnScriptProps.Script)
-		{
-			SpawnScriptProps.Script->RapidIterationParameters.AddOnChangedHandler(
-				FNiagaraParameterStore::FOnChanged::FDelegate::CreateUObject(this, &UNiagaraEmitter::ScriptRapidIterationParameterChanged));
-		}
-	
-		if (UpdateScriptProps.Script)
-		{
-			UpdateScriptProps.Script->RapidIterationParameters.AddOnChangedHandler(
-				FNiagaraParameterStore::FOnChanged::FDelegate::CreateUObject(this, &UNiagaraEmitter::ScriptRapidIterationParameterChanged));
-		}
-
-		for (FNiagaraEventScriptProperties& EventScriptProperties : EventHandlerScriptProps)
-		{
-			EventScriptProperties.Script->RapidIterationParameters.AddOnChangedHandler(
-				FNiagaraParameterStore::FOnChanged::FDelegate::CreateUObject(this, &UNiagaraEmitter::ScriptRapidIterationParameterChanged));
-		}
-
-		for (UNiagaraSimulationStageBase* SimulationStage : SimulationStages)
-		{
-			SimulationStage->OnChanged().AddUObject(this, &UNiagaraEmitter::SimulationStageChanged);
-			SimulationStage->Script->RapidIterationParameters.AddOnChangedHandler(
-				FNiagaraParameterStore::FOnChanged::FDelegate::CreateUObject(this, &UNiagaraEmitter::ScriptRapidIterationParameterChanged));
-		}
-
-		for (UNiagaraRendererProperties* Renderer : RendererProperties)
-		{
-			Renderer->OnChanged().AddUObject(this, &UNiagaraEmitter::RendererChanged);
-		}
-
-		if (EditorData != nullptr)
-		{
-			EditorData->OnPersistentDataChanged().AddUObject(this, &UNiagaraEmitter::PersistentEditorDataChanged);
 		}
 	}
+#else
+	check(GPUComputeScript == nullptr || SimTarget == ENiagaraSimTarget::GPUComputeSim);
 #endif
 
-	ResolveScalabilitySettings();
-
-#if !UE_BUILD_SHIPPING
-	DebugSimName.Empty();
-	if (const UNiagaraSystem* SystemOwner = Cast<const UNiagaraSystem>(GetOuter()))
+	//Temporarily disabling interpolated spawn if the script type and flag don't match.
+	if (SpawnScriptProps.Script)
 	{
-		DebugSimName = SystemOwner->GetName();
-		DebugSimName.AppendChar(':');
+		bool bActualInterpolatedSpawning = SpawnScriptProps.Script->IsInterpolatedParticleSpawnScript();
+		if (bInterpolatedSpawning != bActualInterpolatedSpawning)
+		{
+			bInterpolatedSpawning = false;
+			if (bActualInterpolatedSpawning)
+			{
+#if WITH_EDITORONLY_DATA
+				//clear out the script as it was compiled with interpolated spawn.
+				SpawnScriptProps.Script->InvalidateCompileResults(TEXT("Interpolated spawn changed."));
+#endif
+				SpawnScriptProps.Script->SetUsage(ENiagaraScriptUsage::ParticleSpawnScript);
+			}
+			UE_LOG(LogNiagara, Warning, TEXT("Disabling interpolated spawn because emitter flag and script type don't match. Did you adjust this value in the UI? Emitter may need recompile.. %s"), *GetFullName());
+		}
 	}
-	DebugSimName.Append(GetName());
-#endif
+	EnsureScriptsPostLoaded();
+	
+	// we are not yet finished, but we do the rest of the work after postload in a separate task
+	UpdateTaskRef = TGraphTask<FNiagaraEmitterUpdateTask>::CreateTask().ConstructAndDispatchWhenReady(this);
 }
 
 bool UNiagaraEmitter::IsEditorOnly() const
@@ -1247,6 +1174,125 @@
 	}
 }
 
+void UNiagaraEmitter::UpdateEmitterAfterLoad()
+{
+	if (bFullyLoaded)
+	{
+		return;
+	}
+	bFullyLoaded = true;
+	UpdateTaskRef = nullptr;
+	
+#if WITH_EDITORONLY_DATA
+	check(IsInGameThread());
+	if (GetOuter()->IsA<UNiagaraEmitter>())
+	{
+		// If this emitter is owned by another emitter, remove it's inheritance information so that it doesn't try to merge changes.
+		Parent = nullptr;
+		ParentAtLastMerge = nullptr;
+	}
+	
+	if (!GetOutermost()->bIsCookedForEditor)
+	{
+		if (IsSynchronizedWithParent() == false)
+		{
+			// Modify here so that the asset will be marked dirty when using the resave commandlet.  This will be ignored during regular post load.
+			MergeChangesFromParent();
+		}
+
+		// Reset scripts if recompile is forced.
+		bool bGenerateNewChangeId = false;
+		FString GenerateNewChangeIdReason;
+		if (GetForceCompileOnLoad())
+		{
+			// If we are a standalone emitter, then we invalidate id's, which should cause systems dependent on us to regenerate.
+			UObject* OuterObj = GetOuter();
+			if (OuterObj == GetOutermost())
+			{
+				GraphSource->ForceGraphToRecompileOnNextCheck();
+				bGenerateNewChangeId = true;
+				GenerateNewChangeIdReason = TEXT("PostLoad - Force compile on load");
+				if (GEnableVerboseNiagaraChangeIdLogging)
+				{
+					UE_LOG(LogNiagara, Log, TEXT("InvalidateCachedCompileIds for %s because GbForceNiagaraCompileOnLoad = %d"), *GetPathName(), GbForceNiagaraCompileOnLoad);
+				}
+			}
+		}
+	
+		if (ChangeId.IsValid() == false)
+		{
+			// If the change id is already invalid we need to generate a new one, and can skip checking the owned scripts.
+			bGenerateNewChangeId = true;
+			GenerateNewChangeIdReason = TEXT("PostLoad - Change id was invalid.");
+			if (GEnableVerboseNiagaraChangeIdLogging)
+			{
+				UE_LOG(LogNiagara, Log, TEXT("Change ID updated for emitter %s because the ID was invalid."), *GetPathName());
+			}
+		}
+
+		if (bGenerateNewChangeId)
+		{
+			UpdateChangeId(GenerateNewChangeIdReason);
+		}
+
+		GraphSource->OnChanged().AddUObject(this, &UNiagaraEmitter::GraphSourceChanged);
+
+		EmitterSpawnScriptProps.Script->RapidIterationParameters.AddOnChangedHandler(
+			FNiagaraParameterStore::FOnChanged::FDelegate::CreateUObject(this, &UNiagaraEmitter::ScriptRapidIterationParameterChanged));
+		EmitterUpdateScriptProps.Script->RapidIterationParameters.AddOnChangedHandler(
+			FNiagaraParameterStore::FOnChanged::FDelegate::CreateUObject(this, &UNiagaraEmitter::ScriptRapidIterationParameterChanged));
+
+		if (SpawnScriptProps.Script)
+		{
+			SpawnScriptProps.Script->RapidIterationParameters.AddOnChangedHandler(
+				FNiagaraParameterStore::FOnChanged::FDelegate::CreateUObject(this, &UNiagaraEmitter::ScriptRapidIterationParameterChanged));
+		}
+	
+		if (UpdateScriptProps.Script)
+		{
+			UpdateScriptProps.Script->RapidIterationParameters.AddOnChangedHandler(
+				FNiagaraParameterStore::FOnChanged::FDelegate::CreateUObject(this, &UNiagaraEmitter::ScriptRapidIterationParameterChanged));
+		}
+
+		for (FNiagaraEventScriptProperties& EventScriptProperties : EventHandlerScriptProps)
+		{
+			EventScriptProperties.Script->RapidIterationParameters.AddOnChangedHandler(
+				FNiagaraParameterStore::FOnChanged::FDelegate::CreateUObject(this, &UNiagaraEmitter::ScriptRapidIterationParameterChanged));
+		}
+
+		for (UNiagaraSimulationStageBase* SimulationStage : SimulationStages)
+		{
+			SimulationStage->OnChanged().AddUObject(this, &UNiagaraEmitter::SimulationStageChanged);
+			SimulationStage->Script->RapidIterationParameters.AddOnChangedHandler(
+				FNiagaraParameterStore::FOnChanged::FDelegate::CreateUObject(this, &UNiagaraEmitter::ScriptRapidIterationParameterChanged));
+		}
+
+		for (UNiagaraRendererProperties* Renderer : RendererProperties)
+		{
+			Renderer->OnChanged().AddUObject(this, &UNiagaraEmitter::RendererChanged);
+		}
+
+		if (EditorData != nullptr)
+		{
+			EditorData->OnPersistentDataChanged().AddUObject(this, &UNiagaraEmitter::PersistentEditorDataChanged);
+		}
+	}
+#endif
+
+	ResolveScalabilitySettings();
+
+#if !UE_BUILD_SHIPPING
+	DebugSimName.Empty();
+	if (const UNiagaraSystem* SystemOwner = Cast<const UNiagaraSystem>(GetOuter()))
+	{
+		DebugSimName = SystemOwner->GetName();
+		DebugSimName.AppendChar(':');
+	}
+	DebugSimName.Append(GetName());
+#endif
+
+}
+
 bool UNiagaraEmitter::IsAllowedByScalability()const
 {
 	return Platforms.IsActive();
diff --git a/Engine/Plugins/FX/Niagara/Source/Niagara/Private/NiagaraSystem.cpp b/Engine/Plugins/FX/Niagara/Source/Niagara/Private/NiagaraSystem.cpp
--- a/Engine/Plugins/FX/Niagara/Source/Niagara/Private/NiagaraSystem.cpp
+++ b/Engine/Plugins/FX/Niagara/Source/Niagara/Private/NiagaraSystem.cpp
@@ -12,7 +12,6 @@
 #include "NiagaraEmitterInstanceBatcher.h"
 #include "NiagaraModule.h"
 #include "NiagaraPrecompileContainer.h"
-#include "NiagaraRenderer.h"
 #include "NiagaraRendererProperties.h"
 #include "NiagaraScriptSourceBase.h"
 #include "NiagaraSettings.h"
@@ -79,6 +78,44 @@
 
 //////////////////////////////////////////////////////////////////////////
 
+/**
+Game thread task to do any outstanding work on the loaded system
+*/
+struct FNiagaraSystemUpdateTask
+{
+	FNiagaraSystemUpdateTask(TWeakObjectPtr<UNiagaraSystem> InSystem) : WeakSystem(InSystem)
+	{
+	}
+
+	FORCEINLINE TStatId GetStatId() const { RETURN_QUICK_DECLARE_CYCLE_STAT(FNiagaraSystemUpdateTask, STATGROUP_TaskGraphTasks); }
+	static FORCEINLINE ENamedThreads::Type GetDesiredThread() { return ENamedThreads::GameThread; }
+	static FORCEINLINE ESubsequentsMode::Type GetSubsequentsMode() { return ESubsequentsMode::TrackSubsequents; }
+
+	void DoTask(ENamedThreads::Type CurrentThread, const FGraphEventRef& MyCompletionGraphEvent)
+	{
+		check(IsInGameThread());
+		UNiagaraSystem* System = WeakSystem.Get();
+		if (System == nullptr)
+		{
+			// probably got garbage collected in the meantime
+			return;
+		}
+#if WITH_EDITORONLY_DATA
+		for (const FNiagaraEmitterHandle& Handle : System->GetEmitterHandles())
+		{
+			if (Handle.GetInstance() && Handle.GetInstance()->UpdateTaskRef.IsValid())
+			{
+				// the emitter tasks should all be done by now, if not then something is very wrong
+				ensureMsgf(Handle.GetInstance()->UpdateTaskRef->IsComplete(), TEXT("Unfinished emitter update task for system %s"), *GetPathNameSafe(System));
+			}
+		}
+#endif
+		System->UpdateSystemAfterLoad();
+	}
+
+	TWeakObjectPtr<UNiagaraSystem> WeakSystem;
+};
+
 UNiagaraSystem::UNiagaraSystem(const FObjectInitializer& ObjectInitializer)
 : Super(ObjectInitializer)
 #if WITH_EDITORONLY_DATA
@@ -149,6 +186,7 @@
 #if WITH_EDITORONLY_DATA
 	WaitForCompilationComplete();
 #endif
+	EnsureFullyLoaded();
 }
 
 #if WITH_EDITOR
@@ -260,9 +298,192 @@
 	return false;
 }
 
+void UNiagaraSystem::UpdateSystemAfterLoad()
+{
+	// guard against deadlocks by having wait called on it during the update
+	if (bFullyLoaded)
+	{
+		return;
+	}
+	bFullyLoaded = true;
+	UpdateTaskRef = nullptr;
+	
 #if WITH_EDITORONLY_DATA
-bool UNiagaraSystem::UsesScript(const UNiagaraScript* Script)const
+	// We remove emitters and scripts on dedicated servers, so skip further work.
+	const bool bIsDedicatedServer = !GIsClient && GIsServer;
+
+	if (!GetOutermost()->bIsCookedForEditor && !bIsDedicatedServer)
+	{		
+		TArray<UNiagaraScript*> AllSystemScripts;
+		UNiagaraScriptSourceBase* SystemScriptSource;
+		if (SystemSpawnScript == nullptr)
+		{
+			SystemSpawnScript = NewObject<UNiagaraScript>(this, "SystemSpawnScript", RF_Transactional);
+			SystemSpawnScript->SetUsage(ENiagaraScriptUsage::SystemSpawnScript);
+			INiagaraModule& NiagaraModule = FModuleManager::GetModuleChecked<INiagaraModule>("Niagara");
+			SystemScriptSource = NiagaraModule.GetEditorOnlyDataUtilities().CreateDefaultScriptSource(this);
+			SystemSpawnScript->SetLatestSource(SystemScriptSource);
+		}
+		else
+		{
+			SystemSpawnScript->ConditionalPostLoad();
+			SystemScriptSource = SystemSpawnScript->GetLatestSource();
+		}
+		AllSystemScripts.Add(SystemSpawnScript);
+
+		if (SystemUpdateScript == nullptr)
+		{
+			SystemUpdateScript = NewObject<UNiagaraScript>(this, "SystemUpdateScript", RF_Transactional);
+			SystemUpdateScript->SetUsage(ENiagaraScriptUsage::SystemUpdateScript);
+			SystemUpdateScript->SetLatestSource(SystemScriptSource);
+		}
+		else
+		{
+			SystemUpdateScript->ConditionalPostLoad();
+		}
+		AllSystemScripts.Add(SystemUpdateScript);
+
+		bool bSystemScriptsAreSynchronized = true;
+		for (UNiagaraScript* SystemScript : AllSystemScripts)
+		{
+			bSystemScriptsAreSynchronized &= SystemScript->AreScriptAndSourceSynchronized();
+		}
+
+		bool bEmitterScriptsAreSynchronized = true;
+
+#if 0
+		UE_LOG(LogNiagara, Log, TEXT("PreMerger"));
+		for (FNiagaraEmitterHandle& EmitterHandle : EmitterHandles)
+		{
+			UE_LOG(LogNiagara, Log, TEXT("Emitter Handle: %s"), *EmitterHandle.GetUniqueInstanceName());
+			UNiagaraScript* UpdateScript = EmitterHandle.GetInstance()->GetScript(ENiagaraScriptUsage::ParticleUpdateScript, FGuid());
+			UNiagaraScript* SpawnScript = EmitterHandle.GetInstance()->GetScript(ENiagaraScriptUsage::ParticleSpawnScript, FGuid());
+			UE_LOG(LogNiagara, Log, TEXT("Spawn Parameters"));
+			SpawnScript->GetVMExecutableData().Parameters.DumpParameters();
+			UE_LOG(LogNiagara, Log, TEXT("Spawn RI Parameters"));
+			SpawnScript->RapidIterationParameters.DumpParameters();
+			UE_LOG(LogNiagara, Log, TEXT("Update Parameters"));
+			UpdateScript->GetVMExecutableData().Parameters.DumpParameters();
+			UE_LOG(LogNiagara, Log, TEXT("Update RI Parameters"));
+			UpdateScript->RapidIterationParameters.DumpParameters();
+		}
+#endif
+
+		for (FNiagaraEmitterHandle& EmitterHandle : EmitterHandles)
+		{
+			if (EmitterHandle.GetIsEnabled() && EmitterHandle.GetInstance() && !EmitterHandle.GetInstance()->AreAllScriptAndSourcesSynchronized())
+			{
+				bEmitterScriptsAreSynchronized = false;
+			}
+		}
+
+		if (UNiagaraEmitter::GetForceCompileOnLoad())
+		{
+			ForceGraphToRecompileOnNextCheck();
+			UE_LOG(LogNiagara, Log, TEXT("System %s being rebuilt because UNiagaraEmitter::GetForceCompileOnLoad() == true."), *GetPathName());
+		}
+
+		if (bSystemScriptsAreSynchronized == false && GEnableVerboseNiagaraChangeIdLogging)
+		{
+			UE_LOG(LogNiagara, Log, TEXT("System %s being compiled because there were changes to a system script Change ID."), *GetPathName());
+		}
+
+		if (bEmitterScriptsAreSynchronized == false && GEnableVerboseNiagaraChangeIdLogging)
+		{
+			UE_LOG(LogNiagara, Log, TEXT("System %s being compiled because there were changes to an emitter script Change ID."), *GetPathName());
+		}
+
+		if (EmitterCompiledData.Num() == 0 || EmitterCompiledData[0]->DataSetCompiledData.Variables.Num() == 0)
+		{
+			InitEmitterCompiledData();
+		}
+
+		if (SystemCompiledData.InstanceParamStore.ReadParameterVariables().Num() == 0 ||SystemCompiledData.DataSetCompiledData.Variables.Num() == 0)
+		{
+			InitSystemCompiledData();
+		}
+
+#if 0
+		UE_LOG(LogNiagara, Log, TEXT("Before"));
+		for (FNiagaraEmitterHandle& EmitterHandle : EmitterHandles)
+		{
+			UE_LOG(LogNiagara, Log, TEXT("Emitter Handle: %s"), *EmitterHandle.GetUniqueInstanceName());
+			UNiagaraScript* UpdateScript = EmitterHandle.GetInstance()->GetScript(ENiagaraScriptUsage::ParticleUpdateScript, FGuid());
+			UNiagaraScript* SpawnScript = EmitterHandle.GetInstance()->GetScript(ENiagaraScriptUsage::ParticleSpawnScript, FGuid());
+			UE_LOG(LogNiagara, Log, TEXT("Spawn Parameters"));
+			SpawnScript->GetVMExecutableData().Parameters.DumpParameters();
+			UE_LOG(LogNiagara, Log, TEXT("Spawn RI Parameters"));
+			SpawnScript->RapidIterationParameters.DumpParameters();
+			UE_LOG(LogNiagara, Log, TEXT("Update Parameters"));
+			UpdateScript->GetVMExecutableData().Parameters.DumpParameters();
+			UE_LOG(LogNiagara, Log, TEXT("Update RI Parameters"));
+			UpdateScript->RapidIterationParameters.DumpParameters();
+		}
+#endif
+
+		if (bSystemScriptsAreSynchronized == false || bEmitterScriptsAreSynchronized == false)
+		{
+			// Call modify here so that the system will resave the compile ids and script vm when running the resave
+			// commandlet.  In normal post load, it will be ignored.
+			Modify();
+			RequestCompile(false);
+		}
+
+#if 0
+		UE_LOG(LogNiagara, Log, TEXT("After"));
+		for (FNiagaraEmitterHandle& EmitterHandle : EmitterHandles)
+		{
+			UE_LOG(LogNiagara, Log, TEXT("Emitter Handle: %s"), *EmitterHandle.GetUniqueInstanceName());
+			UNiagaraScript* UpdateScript = EmitterHandle.GetInstance()->GetScript(ENiagaraScriptUsage::ParticleUpdateScript, FGuid());
+			UNiagaraScript* SpawnScript = EmitterHandle.GetInstance()->GetScript(ENiagaraScriptUsage::ParticleSpawnScript, FGuid());
+			UE_LOG(LogNiagara, Log, TEXT("Spawn Parameters"));
+			SpawnScript->GetVMExecutableData().Parameters.DumpParameters();
+			UE_LOG(LogNiagara, Log, TEXT("Spawn RI Parameters"));
+			SpawnScript->RapidIterationParameters.DumpParameters();
+			UE_LOG(LogNiagara, Log, TEXT("Update Parameters"));
+			UpdateScript->GetVMExecutableData().Parameters.DumpParameters();
+			UE_LOG(LogNiagara, Log, TEXT("Update RI Parameters"));
+			UpdateScript->RapidIterationParameters.DumpParameters();
+		}
+#endif
+	}
+	if (GNiagaraForceSystemsToCookOutRapidIterationOnLoad == 1 && !bBakeOutRapidIteration)
+	{
+		WaitForCompilationComplete();
+		bBakeOutRapidIteration = true;
+		RequestCompile(false);
+	}
+#endif
+
+	if ( FPlatformProperties::RequiresCookedData() )
+	{
+		bIsValidCached = IsValidInternal();
+		bIsReadyToRunCached = IsReadyToRunInternal();
+	}
+
+	ResolveScalabilitySettings();
+
+	ComputeEmittersExecutionOrder();
+
+	ComputeRenderersDrawOrder();
+
+	CacheFromCompiledData();
+
+	//TODO: Move to serialized properties?
+	UpdateDITickFlags();
+	UpdateHasGPUEmitters();
+
+	// Run task to prime pools this must happen on the GameThread
+	if (PoolPrimeSize > 0 && MaxPoolSize > 0)
+	{
+		FNiagaraWorldManager::PrimePoolForAllWorlds(this);
+	}
+}
+
+#if WITH_EDITORONLY_DATA
+bool UNiagaraSystem::UsesScript(const UNiagaraScript* Script) const
 {
+	EnsureFullyLoaded();
 	if (SystemSpawnScript == Script ||
 		SystemUpdateScript == Script)
 	{
@@ -530,212 +751,38 @@
 	ExposedParameters.RecreateRedirections();
 #endif
 
+	FGraphEventArray EmitterPrerequisiteTasks;
+	for (FNiagaraEmitterHandle& EmitterHandle : EmitterHandles)
+	{
 #if WITH_EDITORONLY_DATA
-	// We remove emitters and scripts on dedicated servers, so skip further work.
-	const bool bIsDedicatedServer = !GIsClient && GIsServer;
-
-	if (!GetOutermost()->bIsCookedForEditor && !bIsDedicatedServer)
-	{
-		TArray<UNiagaraScript*> AllSystemScripts;
-		UNiagaraScriptSourceBase* SystemScriptSource = nullptr;
-		if (SystemSpawnScript == nullptr)
+		EmitterHandle.ConditionalPostLoad(NiagaraVer);
+#else
+		if (UNiagaraEmitter* NiagaraEmitter = EmitterHandle.GetInstance())
 		{
-			SystemSpawnScript = NewObject<UNiagaraScript>(this, "SystemSpawnScript", RF_Transactional);
-			SystemSpawnScript->SetUsage(ENiagaraScriptUsage::SystemSpawnScript);
-			INiagaraModule& NiagaraModule = FModuleManager::GetModuleChecked<INiagaraModule>("Niagara");
-			SystemScriptSource = NiagaraModule.GetEditorOnlyDataUtilities().CreateDefaultScriptSource(this);
-			SystemSpawnScript->SetLatestSource(SystemScriptSource);
-		}
-		else
-		{
-			SystemSpawnScript->ConditionalPostLoad();
-			SystemScriptSource = SystemSpawnScript->GetLatestSource();
-		}
-		AllSystemScripts.Add(SystemSpawnScript);
-
-		if (SystemUpdateScript == nullptr)
-		{
-			SystemUpdateScript = NewObject<UNiagaraScript>(this, "SystemUpdateScript", RF_Transactional);
-			SystemUpdateScript->SetUsage(ENiagaraScriptUsage::SystemUpdateScript);
-			SystemUpdateScript->SetLatestSource(SystemScriptSource);
-		}
-		else
-		{
-			SystemUpdateScript->ConditionalPostLoad();
-		}
-		AllSystemScripts.Add(SystemUpdateScript);
-
-		//TODO: This causes a crash becuase the script source ptr is null? Fix
-		//For existing emitters before the lifecylce rework, ensure they have the system lifecycle module.
-		if (NiagaraVer < FNiagaraCustomVersion::LifeCycleRework)
-		{
-			/*UNiagaraScriptSourceBase* SystemScriptSource = SystemUpdateScript->GetSource();
-			if (SystemScriptSource)
-			{
-				bool bFoundModule;
-				if (SystemScriptSource->AddModuleIfMissing(TEXT("/Niagara/Modules/System/SystemLifeCycle.SystemLifeCycle"), ENiagaraScriptUsage::SystemUpdateScript, bFoundModule))
-				{
-					bNeedsRecompile = true;
-				}
-			}*/
-		}
-
-		bool bSystemScriptsAreSynchronized = true;
-		for (UNiagaraScript* SystemScript : AllSystemScripts)
-		{
-			bSystemScriptsAreSynchronized &= SystemScript->AreScriptAndSourceSynchronized();
-		}
-
-		bool bEmitterScriptsAreSynchronized = true;
-
-#if 0
-		UE_LOG(LogNiagara, Log, TEXT("PreMerger"));
-		for (FNiagaraEmitterHandle& EmitterHandle : EmitterHandles)
-		{
-			UE_LOG(LogNiagara, Log, TEXT("Emitter Handle: %s"), *EmitterHandle.GetUniqueInstanceName());
-			UNiagaraScript* UpdateScript = EmitterHandle.GetInstance()->GetScript(ENiagaraScriptUsage::ParticleUpdateScript, FGuid());
-			UNiagaraScript* SpawnScript = EmitterHandle.GetInstance()->GetScript(ENiagaraScriptUsage::ParticleSpawnScript, FGuid());
-			UE_LOG(LogNiagara, Log, TEXT("Spawn Parameters"));
-			SpawnScript->GetVMExecutableData().Parameters.DumpParameters();
-			UE_LOG(LogNiagara, Log, TEXT("Spawn RI Parameters"));
-			SpawnScript->RapidIterationParameters.DumpParameters();
-			UE_LOG(LogNiagara, Log, TEXT("Update Parameters"));
-			UpdateScript->GetVMExecutableData().Parameters.DumpParameters();
-			UE_LOG(LogNiagara, Log, TEXT("Update RI Parameters"));
-			UpdateScript->RapidIterationParameters.DumpParameters();
+			NiagaraEmitter->ConditionalPostLoad();
 		}
 #endif
-
-		for (FNiagaraEmitterHandle& EmitterHandle : EmitterHandles)
+		if (UNiagaraEmitter* NiagaraEmitter = EmitterHandle.GetInstance())
 		{
-			EmitterHandle.ConditionalPostLoad(NiagaraVer);
-			if (EmitterHandle.GetIsEnabled() && EmitterHandle.GetInstance() && !EmitterHandle.GetInstance()->AreAllScriptAndSourcesSynchronized())
-			{
-				bEmitterScriptsAreSynchronized = false;
-			}
+			EmitterPrerequisiteTasks.Add(NiagaraEmitter->UpdateTaskRef);
 		}
-
-		if (EditorData == nullptr)
-		{
-			INiagaraModule& NiagaraModule = FModuleManager::GetModuleChecked<INiagaraModule>("Niagara");
-			EditorData = NiagaraModule.GetEditorOnlyDataUtilities().CreateDefaultEditorData(this);
-		}
-		else
-		{
-			EditorData->PostLoadFromOwner(this);
-		}
-
-		if (UNiagaraEmitter::GetForceCompileOnLoad())
-		{
-			ForceGraphToRecompileOnNextCheck();
-			UE_LOG(LogNiagara, Log, TEXT("System %s being rebuilt because UNiagaraEmitter::GetForceCompileOnLoad() == true."), *GetPathName());
-		}
-
-		if (bSystemScriptsAreSynchronized == false && GEnableVerboseNiagaraChangeIdLogging)
-		{
-			UE_LOG(LogNiagara, Log, TEXT("System %s being compiled because there were changes to a system script Change ID."), *GetPathName());
-		}
-
-		if (bEmitterScriptsAreSynchronized == false && GEnableVerboseNiagaraChangeIdLogging)
-		{
-			UE_LOG(LogNiagara, Log, TEXT("System %s being compiled because there were changes to an emitter script Change ID."), *GetPathName());
-		}
-
-		if (EmitterCompiledData.Num() == 0 || EmitterCompiledData[0]->DataSetCompiledData.Variables.Num() == 0)
-		{
-			InitEmitterCompiledData();
-		}
-
-		if (SystemCompiledData.InstanceParamStore.ReadParameterVariables().Num() == 0 ||SystemCompiledData.DataSetCompiledData.Variables.Num() == 0)
-		{
-			InitSystemCompiledData();
-		}
-
-#if 0
-		UE_LOG(LogNiagara, Log, TEXT("Before"));
-		for (FNiagaraEmitterHandle& EmitterHandle : EmitterHandles)
-		{
-			UE_LOG(LogNiagara, Log, TEXT("Emitter Handle: %s"), *EmitterHandle.GetUniqueInstanceName());
-			UNiagaraScript* UpdateScript = EmitterHandle.GetInstance()->GetScript(ENiagaraScriptUsage::ParticleUpdateScript, FGuid());
-			UNiagaraScript* SpawnScript = EmitterHandle.GetInstance()->GetScript(ENiagaraScriptUsage::ParticleSpawnScript, FGuid());
-			UE_LOG(LogNiagara, Log, TEXT("Spawn Parameters"));
-			SpawnScript->GetVMExecutableData().Parameters.DumpParameters();
-			UE_LOG(LogNiagara, Log, TEXT("Spawn RI Parameters"));
-			SpawnScript->RapidIterationParameters.DumpParameters();
-			UE_LOG(LogNiagara, Log, TEXT("Update Parameters"));
-			UpdateScript->GetVMExecutableData().Parameters.DumpParameters();
-			UE_LOG(LogNiagara, Log, TEXT("Update RI Parameters"));
-			UpdateScript->RapidIterationParameters.DumpParameters();
-		}
-#endif
-
-		if (bSystemScriptsAreSynchronized == false || bEmitterScriptsAreSynchronized == false)
-		{
-			// Call modify here so that the system will resave the compile ids and script vm when running the resave
-			// commandlet.  In normal post load, it will be ignored.
-			Modify();
-			RequestCompile(false);
-		}
-
-#if 0
-		UE_LOG(LogNiagara, Log, TEXT("After"));
-		for (FNiagaraEmitterHandle& EmitterHandle : EmitterHandles)
-		{
-			UE_LOG(LogNiagara, Log, TEXT("Emitter Handle: %s"), *EmitterHandle.GetUniqueInstanceName());
-			UNiagaraScript* UpdateScript = EmitterHandle.GetInstance()->GetScript(ENiagaraScriptUsage::ParticleUpdateScript, FGuid());
-			UNiagaraScript* SpawnScript = EmitterHandle.GetInstance()->GetScript(ENiagaraScriptUsage::ParticleSpawnScript, FGuid());
-			UE_LOG(LogNiagara, Log, TEXT("Spawn Parameters"));
-			SpawnScript->GetVMExecutableData().Parameters.DumpParameters();
-			UE_LOG(LogNiagara, Log, TEXT("Spawn RI Parameters"));
-			SpawnScript->RapidIterationParameters.DumpParameters();
-			UE_LOG(LogNiagara, Log, TEXT("Update Parameters"));
-			UpdateScript->GetVMExecutableData().Parameters.DumpParameters();
-			UE_LOG(LogNiagara, Log, TEXT("Update RI Parameters"));
-			UpdateScript->RapidIterationParameters.DumpParameters();
-		}
-#endif
 	}
-	if (GNiagaraForceSystemsToCookOutRapidIterationOnLoad == 1 && !bBakeOutRapidIteration)
+	
+#if WITH_EDITORONLY_DATA
+	if (EditorData == nullptr)
 	{
-		WaitForCompilationComplete();
-		bBakeOutRapidIteration = true;
-		RequestCompile(false);
+		INiagaraModule& NiagaraModule = FModuleManager::GetModuleChecked<INiagaraModule>("Niagara");
+		EditorData = NiagaraModule.GetEditorOnlyDataUtilities().CreateDefaultEditorData(this);
 	}
+	else
+	{
+		EditorData->PostLoadFromOwner(this);
+	}
+	
 #endif // WITH_EDITORONLY_DATA
 
-	if ( FPlatformProperties::RequiresCookedData() )
-	{
-		bIsValidCached = IsValidInternal();
-		bIsReadyToRunCached = IsReadyToRunInternal();
-	}
-
-	ResolveScalabilitySettings();
-
-	ComputeEmittersExecutionOrder();
-
-	ComputeRenderersDrawOrder();
-
-	CacheFromCompiledData();
-
-	//TODO: Move to serialized properties?
-	UpdateDITickFlags();
-	UpdateHasGPUEmitters();
-
-	// Run task to prime pools this must happen on the GameThread
-	if (PoolPrimeSize > 0 && MaxPoolSize > 0)
-	{
-		AsyncTask(
-			ENamedThreads::GameThread,
-			[WeakSystem=TWeakObjectPtr<UNiagaraSystem>(this)]()
-			{
-				if ( UNiagaraSystem* NiagaraSystem = WeakSystem.Get() )
-				{
-					check(IsInGameThread());
-					FNiagaraWorldManager::PrimePoolForAllWorlds(NiagaraSystem);
-				}
-			}
-		);
-	}
+	// we are not yet finished, but we do the rest of the work after postload in a separate task
+	UpdateTaskRef = TGraphTask<FNiagaraSystemUpdateTask>::CreateTask(&EmitterPrerequisiteTasks).ConstructAndDispatchWhenReady(this);
 }
 
 #if WITH_EDITORONLY_DATA
@@ -817,6 +864,7 @@
 		return false;
 	}
 
+	EnsureFullyLoaded();
 	if (!SystemSpawnScript || !SystemUpdateScript)
 	{
 		if (FPlatformProperties::RequiresCookedData())
@@ -886,6 +934,7 @@
 
 void UNiagaraSystem::GetAssetRegistryTags(TArray<FAssetRegistryTag>& OutTags) const
 {
+	EnsureFullyLoaded();
 #if WITH_EDITOR
 	OutTags.Add(FAssetRegistryTag("HasGPUEmitter", HasAnyGPUEmitters() ? TEXT("True") : TEXT("False"), FAssetRegistryTag::TT_Alphabetical));
 
@@ -1631,6 +1680,23 @@
 	return true;
 }
 
+void UNiagaraSystem::EnsureFullyLoaded() const
+{
+	if (UpdateTaskRef.IsValid() && UpdateTaskRef->IsComplete() == false)
+	{
+		ensureMsgf(FUObjectThreadContext::Get().IsRoutingPostLoad == false, TEXT("Niagara system data was forced to load during PostLoad instead of the taskgraph! Asset %s"), *GetPathName());
+		for (const FNiagaraEmitterHandle& EmitterHandle : EmitterHandles)
+		{
+			if (UNiagaraEmitter* NiagaraEmitter = EmitterHandle.GetInstance())
+			{
+				NiagaraEmitter->UpdateEmitterAfterLoad();
+			}
+		}
+		UNiagaraSystem* System = const_cast<UNiagaraSystem*>(this); // Sadly necessary because we are called from other const functions
+		System->UpdateSystemAfterLoad();
+	}	
+}
 
+
 bool UNiagaraSystem::CanObtainEmitterAttribute(const FNiagaraVariableBase& InVarWithUniqueNameNamespace) const
 {
@@ -1666,6 +1732,12 @@
 	return EmitterHandle;
 }
 
+void UNiagaraSystem::AddEmitterHandleDirect(FNiagaraEmitterHandle& EmitterHandle)
+{
+	EmitterHandles.Add(EmitterHandle);
+	RefreshSystemParametersFromEmitter(EmitterHandle);
+}
+
 FNiagaraEmitterHandle UNiagaraSystem::DuplicateEmitterHandle(const FNiagaraEmitterHandle& EmitterHandleToDuplicate, FName EmitterName)
 {
 	UNiagaraEmitter* DuplicateEmitter = UNiagaraEmitter::CreateAsDuplicate(*EmitterHandleToDuplicate.GetInstance(), EmitterName, *this);
@@ -2306,11 +2378,9 @@
 
 
 		//Compile all emitters
-		bool bTrulyAsync = true;
 		bool bAnyUnsynchronized = false;	
 
 		// Pass one... determine if any need to be compiled.
-		bool bForceSystems = false;
 		{
 
 			for (int32 i = 0; i < EmitterHandles.Num(); i++)
@@ -2318,8 +2388,6 @@
 				FNiagaraEmitterHandle Handle = EmitterHandles[i];
 				if (Handle.GetInstance() && Handle.GetIsEnabled())
 				{
-					UNiagaraScriptSourceBase* GraphSource = Handle.GetInstance()->GraphSource;
-
 					TArray<UNiagaraScript*> EmitterScripts;
 					Handle.GetInstance()->GetScripts(EmitterScripts, false, true);
 					check(EmitterScripts.Num() > 0);
@@ -2343,7 +2411,6 @@
 				}
 			}
 
-			bForceSystems = bForce || bAnyUnsynchronized;
 			bAnyCompiled = bAnyUnsynchronized || bForce;
 
 			// Now add the system scripts for compilation...
@@ -2386,66 +2453,59 @@
 		}
 
 
-	
+		// We found things needing compilation, now we have to go through an static duplicate everything that will be translated...
 		{
+			UNiagaraPrecompileContainer* Container = NewObject<UNiagaraPrecompileContainer>(GetTransientPackage());
+			Container->System = this;
+			Container->Scripts = ScriptsNeedingCompile;
+			TSharedPtr<FNiagaraCompileRequestDataBase, ESPMode::ThreadSafe> SystemPrecompiledData = NiagaraModule.Precompile(Container, FGuid());
 
-			// We found things needing compilation, now we have to go through an static duplicate everything that will be translated...
+			if (SystemPrecompiledData.IsValid() == false)
 			{
-				UNiagaraPrecompileContainer* Container = NewObject<UNiagaraPrecompileContainer>(GetTransientPackage());
-				Container->System = this;
-				Container->Scripts = ScriptsNeedingCompile;
-				TSharedPtr<FNiagaraCompileRequestDataBase, ESPMode::ThreadSafe> SystemPrecompiledData = NiagaraModule.Precompile(Container, FGuid());
+				UE_LOG(LogNiagara, Error, TEXT("Failed to precompile %s.  This is due to unexpected invalid or broken data.  Additional details should be in the log."), *GetPathName());
+				return false;
+			}
 
-				if (SystemPrecompiledData.IsValid() == false)
-				{
-					UE_LOG(LogNiagara, Error, TEXT("Failed to precompile %s.  This is due to unexpected invalid or broken data.  Additional details should be in the log."), *GetPathName());
-					return false;
-				}
+			SystemPrecompiledData->GetReferencedObjects(ActiveCompilation.RootObjects);
+			ActiveCompilation.MappedData.Add(SystemSpawnScript, SystemPrecompiledData);
+			ActiveCompilation.MappedData.Add(SystemUpdateScript, SystemPrecompiledData);
 
-				SystemPrecompiledData->GetReferencedObjects(ActiveCompilation.RootObjects);
-				ActiveCompilation.MappedData.Add(SystemSpawnScript, SystemPrecompiledData);
-				ActiveCompilation.MappedData.Add(SystemUpdateScript, SystemPrecompiledData);
-
-				check(EmitterHandles.Num() == SystemPrecompiledData->GetDependentRequestCount());
+			check(EmitterHandles.Num() == SystemPrecompiledData->GetDependentRequestCount());
 
 
-				// Grab the list of user variables that were actually encountered so that we can add to them later.
-				TArray<FNiagaraVariable> EncounteredExposedVars;
-				SystemPrecompiledData->GatherPreCompiledVariables(TEXT("User"), EncounteredExposedVars);
+			// Grab the list of user variables that were actually encountered so that we can add to them later.
+			TArray<FNiagaraVariable> EncounteredExposedVars;
+			SystemPrecompiledData->GatherPreCompiledVariables(TEXT("User"), EncounteredExposedVars);
 
-				for (int32 i = 0; i < EmitterHandles.Num(); i++)
+			for (int32 i = 0; i < EmitterHandles.Num(); i++)
+			{
+				FNiagaraEmitterHandle Handle = EmitterHandles[i];
+				if (Handle.GetInstance() && Handle.GetIsEnabled())
 				{
-					FNiagaraEmitterHandle Handle = EmitterHandles[i];
-					if (Handle.GetInstance() && Handle.GetIsEnabled())
+					TSharedPtr<FNiagaraCompileRequestDataBase, ESPMode::ThreadSafe> EmitterPrecompiledData = SystemPrecompiledData->GetDependentRequest(i);
+					EmitterPrecompiledData->GetReferencedObjects(ActiveCompilation.RootObjects);
+
+					TArray<UNiagaraScript*> EmitterScripts;
+					Handle.GetInstance()->GetScripts(EmitterScripts, false, true);
+					check(EmitterScripts.Num() > 0);
+					for (UNiagaraScript* EmitterScript : EmitterScripts)
 					{
-						UNiagaraScriptSourceBase* GraphSource = Handle.GetInstance()->GraphSource;
-						TSharedPtr<FNiagaraCompileRequestDataBase, ESPMode::ThreadSafe> EmitterPrecompiledData = SystemPrecompiledData->GetDependentRequest(i);
-						EmitterPrecompiledData->GetReferencedObjects(ActiveCompilation.RootObjects);
+						ActiveCompilation.MappedData.Add(EmitterScript, EmitterPrecompiledData);
+					}
 
-						TArray<UNiagaraScript*> EmitterScripts;
-						Handle.GetInstance()->GetScripts(EmitterScripts, false, true);
-						check(EmitterScripts.Num() > 0);
-						for (UNiagaraScript* EmitterScript : EmitterScripts)
-						{
-							ActiveCompilation.MappedData.Add(EmitterScript, EmitterPrecompiledData);
-						}
+					// Add the emitter's User variables to the encountered list to expose for later.
+					EmitterPrecompiledData->GatherPreCompiledVariables(TEXT("User"), EncounteredExposedVars);
+				}
+			}
 
 
-						// Add the emitter's User variables to the encountered list to expose for later.
-						EmitterPrecompiledData->GatherPreCompiledVariables(TEXT("User"), EncounteredExposedVars);
-
-					}
-				}
-
-
-				// Now let's synchronize the variables that we actually encountered during precompile so that we can expose them to the end user.
-				for (int32 i = 0; i < EncounteredExposedVars.Num(); i++)
+			// Now let's synchronize the variables that we actually encountered during precompile so that we can expose them to the end user.
+			for (int32 i = 0; i < EncounteredExposedVars.Num(); i++)
+			{
+				if (OriginalExposedParams.Contains(EncounteredExposedVars[i]) == false)
 				{
-					if (OriginalExposedParams.Contains(EncounteredExposedVars[i]) == false)
-					{
-						// Just in case it wasn't added previously..
-						ExposedParameters.AddParameter(EncounteredExposedVars[i]);
-					}
+					// Just in case it wasn't added previously..
+					ExposedParameters.AddParameter(EncounteredExposedVars[i]);
 				}
 			}
 		}
diff --git a/Engine/Plugins/FX/Niagara/Source/NiagaraEditor/NiagaraEditor.Build.cs b/Engine/Plugins/FX/Niagara/Source/NiagaraEditor/NiagaraEditor.Build.cs
--- a/Engine/Plugins/FX/Niagara/Source/NiagaraEditor/NiagaraEditor.Build.cs
+++ b/Engine/Plugins/FX/Niagara/Source/NiagaraEditor/NiagaraEditor.Build.cs
@@ -27,6 +27,7 @@ public class NiagaraEditor : ModuleRules
 				"RenderCore",
 				"Slate", 
 				"SlateCore",
+				"SlateNullRenderer",
 				"Kismet",
                 "EditorStyle",
 				"UnrealEd", 
diff --git a/Engine/Plugins/FX/Niagara/Source/NiagaraEditor/Private/NiagaraEditorUtilities.cpp b/Engine/Plugins/FX/Niagara/Source/NiagaraEditor/Private/NiagaraEditorUtilities.cpp
--- a/Engine/Plugins/FX/Niagara/Source/NiagaraEditor/Private/NiagaraEditorUtilities.cpp
+++ b/Engine/Plugins/FX/Niagara/Source/NiagaraEditor/Private/NiagaraEditorUtilities.cpp
@@ -1942,7 +1942,7 @@ TArray<UNiagaraComponent*> FNiagaraEditorUtilities::GetComponentsThatReferenceSy
 	return ReferencingComponents;
 }
 
-const FGuid FNiagaraEditorUtilities::AddEmitterToSystem(UNiagaraSystem& InSystem, UNiagaraEmitter& InEmitterToAdd)
+const FGuid FNiagaraEditorUtilities::AddEmitterToSystem(UNiagaraSystem& InSystem, UNiagaraEmitter& InEmitterToAdd, bool bCreateCopy)
 {
 	// Kill all system instances before modifying the emitter handle list to prevent accessing deleted data.
 	KillSystemInstances(InSystem);
@@ -1960,7 +1960,7 @@ const FGuid FNiagaraEditorUtilities::AddEmitterToSystem(UNiagaraSystem& InSystem
 		InSystem.Modify();
 		EmitterHandle = InSystem.AddEmitterHandle(InEmitterToAdd, FNiagaraUtilities::GetUniqueName(InEmitterToAdd.GetFName(), EmitterHandleNames));
 	}
-	else
+	else if (bCreateCopy)
 	{
 		// When editing an emitter asset we add the emitter as a duplicate so that the parent emitter is duplicated, but it's parent emitter
 		// information is maintained.
@@ -1968,6 +1968,11 @@ const FGuid FNiagaraEditorUtilities::AddEmitterToSystem(UNiagaraSystem& InSystem
 		FNiagaraEmitterHandle TemporaryEmitterHandle(InEmitterToAdd);
 		EmitterHandle = InSystem.DuplicateEmitterHandle(TemporaryEmitterHandle, *InEmitterToAdd.GetUniqueEmitterName());
 	}
+	else
+	{
+		EmitterHandle = FNiagaraEmitterHandle(InEmitterToAdd);
+		InSystem.AddEmitterHandleDirect(EmitterHandle);		
+	}
 	
 	FNiagaraStackGraphUtilities::RebuildEmitterNodes(InSystem);
 	SystemEditorData->SynchronizeOverviewGraphWithSystem(InSystem);
diff --git a/Engine/Plugins/FX/Niagara/Source/NiagaraEditor/Private/NiagaraScriptMergeManager.cpp b/Engine/Plugins/FX/Niagara/Source/NiagaraEditor/Private/NiagaraScriptMergeManager.cpp
--- a/Engine/Plugins/FX/Niagara/Source/NiagaraEditor/Private/NiagaraScriptMergeManager.cpp
+++ b/Engine/Plugins/FX/Niagara/Source/NiagaraEditor/Private/NiagaraScriptMergeManager.cpp
@@ -18,9 +18,13 @@
 #include "NiagaraNodeParameterMapSet.h"
 #include "NiagaraEditorUtilities.h"
 #include "ViewModels/Stack/NiagaraStackGraphUtilities.h"
+#include "ViewModels/NiagaraSystemViewModel.h"
+#include "NiagaraClipboard.h"
 #include "NiagaraEditorModule.h"
 #include "NiagaraEmitterEditorData.h"
 #include "NiagaraStackEditorData.h"
+#include "Framework/Application/SlateApplication.h"
+#include "Interfaces/ISlateNullRendererModule.h"
 
 #include "UObject/PropertyPortFlags.h"
 
@@ -29,6 +33,11 @@
 #include "Modules/ModuleManager.h"
 #include "NiagaraConstants.h"
 #include "NiagaraScriptVariable.h"
+#include "NiagaraSystemEditorData.h"
+#include "NiagaraSystemFactoryNew.h"
+#include "ViewModels/NiagaraEmitterHandleViewModel.h"
+#include "ViewModels/Stack/NiagaraStackModuleItem.h"
+#include "ViewModels/Stack/NiagaraStackViewModel.h"
 
 #define LOCTEXT_NAMESPACE "NiagaraScriptMergeManager"
 
@@ -1105,6 +1114,28 @@ FString FNiagaraEmitterDiffResults::GetErrorMessagesString() const
 	return FString::Join(ErrorMessageStrings, TEXT("\n"));
 }
 
+bool FNiagaraEmitterDiffResults::HasVersionChanges() const
+{
+	for (auto& EventHandlerDiff : ModifiedEventHandlers)
+	{
+		if (EventHandlerDiff.ScriptDiffResults.ChangedVersionOtherModules.Num() > 0)
+		{
+			return true;
+		}
+	}
+	for (auto& SimStageDiff : ModifiedSimulationStages)
+	{
+		if (SimStageDiff.ScriptDiffResults.ChangedVersionOtherModules.Num() > 0)
+		{
+			return true;
+		}
+	}
+	return	EmitterSpawnDiffResults.ChangedVersionOtherModules.Num() > 0 ||
+			EmitterUpdateDiffResults.ChangedVersionOtherModules.Num() > 0 ||
+			ParticleSpawnDiffResults.ChangedVersionOtherModules.Num() > 0 ||
+			ParticleUpdateDiffResults.ChangedVersionOtherModules.Num() > 0;
+}
+
 TSharedRef<FNiagaraScriptMergeManager> FNiagaraScriptMergeManager::Get()
 {
 	FNiagaraEditorModule& NiagaraEditorModule = FModuleManager::GetModuleChecked<FNiagaraEditorModule>("NiagaraEditor");
@@ -1254,6 +1285,106 @@ FNiagaraScriptMergeManager::FApplyDiffResults FNiagaraScriptMergeManager::ForceI
 	return DiffResults;
 }
 
+void FNiagaraScriptMergeManager::UpdateModuleVersions(UNiagaraEmitter& Instance, const FNiagaraEmitterDiffResults& DiffResults) const
+{
+	// gather all the changes
+	TArray<TSharedRef<FNiagaraStackFunctionMergeAdapter>> ChangedVersionBaseModules;
+	ChangedVersionBaseModules.Append(DiffResults.EmitterSpawnDiffResults.ChangedVersionBaseModules);
+	ChangedVersionBaseModules.Append(DiffResults.EmitterUpdateDiffResults.ChangedVersionBaseModules);
+	ChangedVersionBaseModules.Append(DiffResults.ParticleSpawnDiffResults.ChangedVersionBaseModules);
+	ChangedVersionBaseModules.Append(DiffResults.ParticleUpdateDiffResults.ChangedVersionBaseModules);
+	for (auto& EventHandlerDiff : DiffResults.ModifiedEventHandlers)
+	{
+		ChangedVersionBaseModules.Append(EventHandlerDiff.ScriptDiffResults.ChangedVersionBaseModules);
+	}
+	for (auto& SimStageDiff : DiffResults.ModifiedSimulationStages)
+	{
+		ChangedVersionBaseModules.Append(SimStageDiff.ScriptDiffResults.ChangedVersionBaseModules);
+	}
+
+	TArray<TSharedRef<FNiagaraStackFunctionMergeAdapter>> ChangedVersionOtherModules;
+	ChangedVersionOtherModules.Append(DiffResults.EmitterSpawnDiffResults.ChangedVersionOtherModules);
+	ChangedVersionOtherModules.Append(DiffResults.EmitterUpdateDiffResults.ChangedVersionOtherModules);
+	ChangedVersionOtherModules.Append(DiffResults.ParticleSpawnDiffResults.ChangedVersionOtherModules);
+	ChangedVersionOtherModules.Append(DiffResults.ParticleUpdateDiffResults.ChangedVersionOtherModules);
+	for (auto& EventHandlerDiff : DiffResults.ModifiedEventHandlers)
+	{
+		ChangedVersionOtherModules.Append(EventHandlerDiff.ScriptDiffResults.ChangedVersionOtherModules);
+	}
+	for (auto& SimStageDiff : DiffResults.ModifiedSimulationStages)
+	{
+		ChangedVersionOtherModules.Append(SimStageDiff.ScriptDiffResults.ChangedVersionOtherModules);
+	}
+	
+	// the python update script needs a view model, so we need to create a temporary system here to instantiate the view models
+	TSharedPtr<FNiagaraSystemViewModel> SystemViewModel;
+	if (ChangedVersionOtherModules.Num() > 0)
+	{
+		if (!FSlateApplication::IsInitialized())
+		{
+			// it's possible that during the cook process we don't have slate initialized, but we need it set up for the view model
+			FSlateApplication::Create();
+			TSharedRef<FSlateRenderer> SlateRenderer = FModuleManager::Get().LoadModuleChecked<ISlateNullRendererModule>("SlateNullRenderer").CreateSlateNullRenderer();
+			FSlateApplication::Get().InitializeRenderer(SlateRenderer);
+		}
+		UNiagaraSystem* System = NewObject<UNiagaraSystem>(GetTransientPackage(), NAME_None, RF_Transient | RF_Standalone);
+		UNiagaraSystemFactoryNew::InitializeSystem(System, true);
+		SystemViewModel = MakeShared<FNiagaraSystemViewModel>();
+		FNiagaraSystemViewModelOptions SystemOptions;
+		SystemOptions.bCanModifyEmittersFromTimeline = false;
+		SystemOptions.bCanSimulate = false;
+		SystemOptions.bCanAutoCompile = false;
+		SystemOptions.EditMode = ENiagaraSystemViewModelEditMode::EmitterDuringMerge;
+		SystemOptions.MessageLogGuid = System->GetAssetGuid();
+		SystemViewModel->Initialize(*System, SystemOptions);
+		SystemViewModel->GetEditorData().SetOwningSystemIsPlaceholder(true, *System);
+		SystemViewModel->AddEmitter(Instance);
+	}
+
+	// apply version changes
+	for (int i = 0; i < ChangedVersionOtherModules.Num(); i++)
+	{
+		TSharedRef<FNiagaraStackFunctionMergeAdapter> BaseModule = ChangedVersionBaseModules[i];
+		TSharedRef<FNiagaraStackFunctionMergeAdapter> ChangedModule = ChangedVersionOtherModules[i];
+		FGuid NewScriptVersion = BaseModule->GetFunctionCallNode()->SelectedScriptVersion;
+
+		// search for the stack item of the module node we are merging
+		UNiagaraStackModuleItem* ModuleItem = nullptr;
+		UNiagaraStackViewModel* EmitterStackViewModel = SystemViewModel->GetEmitterHandleViewModels()[0]->GetEmitterStackViewModel();
+		TArray<UNiagaraStackEntry*> EntriesToCheck;
+		EmitterStackViewModel->GetRootEntry()->GetUnfilteredChildren(EntriesToCheck);
+		while (EntriesToCheck.Num() > 0)
+		{
+			UNiagaraStackEntry* Entry = EntriesToCheck.Pop();
+			UNiagaraStackModuleItem* ModuleItemToCheck = Cast<UNiagaraStackModuleItem>(Entry);
+			if (ModuleItemToCheck && ModuleItemToCheck->GetModuleNode().NodeGuid == ChangedModule->GetFunctionCallNode()->NodeGuid)
+			{
+				ModuleItem = ModuleItemToCheck;
+				break;
+			}
+			Entry->GetUnfilteredChildren(EntriesToCheck);
+		}
+		
+		FNiagaraScriptVersionUpgradeContext UpgradeContext;
+		UpgradeContext.ConstantResolver = FCompileConstantResolver(&Instance, FNiagaraStackGraphUtilities::GetOutputNodeUsage(*ChangedModule->GetFunctionCallNode()));
+		if (ModuleItem)
+		{
+			UpgradeContext.CreateClipboardCallback = [ModuleItem](UNiagaraClipboardContent* ClipboardContent)
+			{
+				ModuleItem->RefreshChildren();
+				ModuleItem->Copy(ClipboardContent);
+				if (ClipboardContent->Functions.Num() > 0)
+				{
+					ClipboardContent->FunctionInputs = ClipboardContent->Functions[0]->Inputs;
+					ClipboardContent->Functions.Empty();
+				}
+			};
+			UpgradeContext.ApplyClipboardCallback = [ModuleItem](UNiagaraClipboardContent* ClipboardContent, FText& OutWarning) { ModuleItem->Paste(ClipboardContent, OutWarning); };
+		}
+		ChangedModule->GetFunctionCallNode()->ChangeScriptVersion(NewScriptVersion, UpgradeContext, true);
+	}
+}
+
 INiagaraMergeManager::FMergeEmitterResults FNiagaraScriptMergeManager::MergeEmitter(UNiagaraEmitter& Parent, UNiagaraEmitter* ParentAtLastMerge, UNiagaraEmitter& Instance) const
 {
 	SCOPE_CYCLE_COUNTER(STAT_NiagaraEditor_ScriptMergeManager_MergeEmitter);
@@ -1304,17 +1435,36 @@ INiagaraMergeManager::FMergeEmitterResults FNiagaraScriptMergeManager::MergeEmit
 			// If there were differences from the parent or the parent diff failed we can just return a copy of the parent as the merged instance since there
 			// were no changes in the instance which need to be applied.
 			MergeResults.MergeResult = INiagaraMergeManager::EMergeEmitterResult::SucceededDifferencesApplied;
-			MergeResults.MergedInstance = Parent.DuplicateWithoutMerging((UObject*)GetTransientPackage());
+			MergeResults.MergedInstance = Parent.DuplicateWithoutMerging(GetTransientPackage());
 			MergeResults.MergedInstance->ParentScratchPadScripts.Append(MergeResults.MergedInstance->ScratchPadScripts);
 			MergeResults.MergedInstance->ScratchPadScripts.Empty();
 		}
 	}
 	else
 	{
-		DiffResults = DiffEmitters(Parent, Instance); // diff again for version changes
-		UNiagaraEmitter* MergedInstance = Parent.DuplicateWithoutMerging((UObject*)GetTransientPackage());
+		UNiagaraEmitter* MergedInstance = Parent.DuplicateWithoutMerging(GetTransientPackage());
 		TSharedRef<FNiagaraEmitterMergeAdapter> MergedInstanceAdapter = MakeShared<FNiagaraEmitterMergeAdapter>(*MergedInstance);
 
+		// diff for version changes and apply them first. We need to upgrade both the current emitter and the parent at last merge
+		// to the newest version and then diff them again. Otherwise we won't know which changes were made by the version upgrade
+		// and which changes were made to the parent emitter.
+		FNiagaraEmitterDiffResults VersionChangeDiffResults = DiffEmitters(Parent, Instance);
+		if (VersionChangeDiffResults.HasVersionChanges())
+		{
+			// apply version upgrade to the merged emitter instance
+			UpdateModuleVersions(Instance, VersionChangeDiffResults);
+
+			// apply version upgrade to the parent at last merge
+			UNiagaraEmitter* ParentAtLastMergeCopy = Cast<UNiagaraEmitter>(StaticDuplicateObject(FirstEmitterToDiffAgainst, GetTransientPackage()));
+			ParentAtLastMergeCopy->ClearFlags(RF_Standalone | RF_Public);
+			FNiagaraEmitterDiffResults ParentsDiffResults = DiffEmitters(Parent, *ParentAtLastMergeCopy);
+			UpdateModuleVersions(*ParentAtLastMergeCopy, ParentsDiffResults);
+
+			// now diff again
+			DiffResults = DiffEmitters(*ParentAtLastMergeCopy, Instance);
+			ensureMsgf(DiffResults.HasVersionChanges() == false, TEXT("Emitter still has pending version changes after merge! Asset %s"), *Instance.GetPathName());
+		}
+
 		TMap<FGuid, UNiagaraNodeFunctionCall*> ParentFunctionIdToNodeMap;
 		TMap<FGuid, UNiagaraNodeFunctionCall*> LastMergedParentFunctionIdToNodeMap;
 		TMap<FGuid, UNiagaraNodeFunctionCall*> InstanceFunctionIdToNodeMap;
@@ -2736,22 +2886,6 @@ FNiagaraScriptMergeManager::FApplyDiffResults FNiagaraScriptMergeManager::ApplyS
 		BaseScriptStackAdapter->GetScript()->SetUsage(DiffResults.ChangedOtherUsage.GetValue());
 		BaseScriptStackAdapter->GetOutputNode()->SetUsage(DiffResults.ChangedOtherUsage.GetValue());
 	}
-	
-	// Update module versions
-	for (int i = 0; i < DiffResults.ChangedVersionOtherModules.Num(); i++)
-	{
-		TSharedRef<FNiagaraStackFunctionMergeAdapter> BaseModule = DiffResults.ChangedVersionBaseModules[i];
-		TSharedRef<FNiagaraStackFunctionMergeAdapter> ChangedModule = DiffResults.ChangedVersionOtherModules[i];
-		FGuid NewScriptVersion = BaseModule->GetFunctionCallNode()->SelectedScriptVersion;
-		FNiagaraScriptVersionUpgradeContext UpgradeContext;
-		UpgradeContext.bSkipPythonScript = true;
-		UNiagaraNodeFunctionCall* FunctionCallNode = ChangedModule->GetFunctionCallNode();
-		FunctionCallNode->ChangeScriptVersion(NewScriptVersion, UpgradeContext, true);
-		if (FunctionCallNode->RefreshFromExternalChanges())
-		{
-			FunctionCallNode->GetNiagaraGraph()->NotifyGraphNeedsRecompile();
-		}
-	}
 
 	// Apply the graph actions.
 	for (TSharedRef<FNiagaraStackFunctionMergeAdapter> RemoveModule : RemoveModules)
diff --git a/Engine/Plugins/FX/Niagara/Source/NiagaraEditor/Private/NiagaraScriptMergeManager.h b/Engine/Plugins/FX/Niagara/Source/NiagaraEditor/Private/NiagaraScriptMergeManager.h
--- a/Engine/Plugins/FX/Niagara/Source/NiagaraEditor/Private/NiagaraScriptMergeManager.h
+++ b/Engine/Plugins/FX/Niagara/Source/NiagaraEditor/Private/NiagaraScriptMergeManager.h
@@ -332,6 +332,7 @@ struct FNiagaraEmitterDiffResults
 	const TArray<FText>& GetErrorMessages() const;
 
 	FString GetErrorMessagesString() const;
+	bool HasVersionChanges() const;
 
 	TArray<FProperty*> DifferentEmitterProperties;
 
@@ -382,6 +383,8 @@ public:
 	};
 
 public:
+	void UpdateModuleVersions(UNiagaraEmitter& Instance, const FNiagaraEmitterDiffResults& EmitterDiffResults) const;
+	
 	virtual INiagaraMergeManager::FMergeEmitterResults MergeEmitter(UNiagaraEmitter& Parent, UNiagaraEmitter* ParentAtLastMerge, UNiagaraEmitter& Instance) const override;
 
 	static TSharedRef<FNiagaraScriptMergeManager> Get();
diff --git a/Engine/Plugins/FX/Niagara/Source/NiagaraEditor/Private/ViewModels/NiagaraSystemViewModel.cpp b/Engine/Plugins/FX/Niagara/Source/NiagaraEditor/Private/ViewModels/NiagaraSystemViewModel.cpp
--- a/Engine/Plugins/FX/Niagara/Source/NiagaraEditor/Private/ViewModels/NiagaraSystemViewModel.cpp
+++ b/Engine/Plugins/FX/Niagara/Source/NiagaraEditor/Private/ViewModels/NiagaraSystemViewModel.cpp
@@ -337,14 +337,14 @@ TSharedPtr<FNiagaraEmitterHandleViewModel> FNiagaraSystemViewModel::AddEmitter(U
 		GEditor->BeginTransaction(LOCTEXT("AddEmitter", "Add emitter"));
 	}
 
-	const FGuid NewEmitterHandleId = FNiagaraEditorUtilities::AddEmitterToSystem(GetSystem(), Emitter);
+	const FGuid NewEmitterHandleId = FNiagaraEditorUtilities::AddEmitterToSystem(GetSystem(), Emitter, EditMode != ENiagaraSystemViewModelEditMode::EmitterDuringMerge);
 
 	if (false == bSystemIsPlaceholder)
 	{
 		GEditor->EndTransaction();
 	}
 
-	if (GetSystem().GetNumEmitters() == 1 && EditorSettings->GetAutoPlay())
+	if (GetSystem().GetNumEmitters() == 1 && EditorSettings->GetAutoPlay() && Sequencer.IsValid())
 	{
 		// When adding a new emitter to an empty system start playing.
 		Sequencer->SetPlaybackStatus(EMovieScenePlayerStatus::Playing);
@@ -362,10 +362,13 @@ TSharedPtr<FNiagaraEmitterHandleViewModel> FNiagaraSystemViewModel::AddEmitter(U
 
 	TGuardValue<bool> UpdateGuard(bUpdatingSequencerFromEmitterDataChange, true);
 
-	FFrameTime NewStartFrame = NewSystemPlaybackRange.GetLowerBoundValue() * NiagaraSequence->GetMovieScene()->GetTickResolution();
-	int32 NewDuration = (NewSystemPlaybackRange.Size<float>() * NiagaraSequence->GetMovieScene()->GetTickResolution()).FrameNumber.Value;
+	if (NiagaraSequence)
+	{
+		FFrameTime NewStartFrame = NewSystemPlaybackRange.GetLowerBoundValue() * NiagaraSequence->GetMovieScene()->GetTickResolution();
+		int32 NewDuration = (NewSystemPlaybackRange.Size<float>() * NiagaraSequence->GetMovieScene()->GetTickResolution()).FrameNumber.Value;
 
-	NiagaraSequence->GetMovieScene()->SetPlaybackRange(NewStartFrame.RoundToFrame(), NewDuration);
+		NiagaraSequence->GetMovieScene()->SetPlaybackRange(NewStartFrame.RoundToFrame(), NewDuration);
+	}
 
 	TSharedPtr<FNiagaraEmitterHandleViewModel> NewEmitterHandleViewModel = GetEmitterHandleViewModelById(NewEmitterHandleId);
 	TArray<UNiagaraStackEntry*> SelectedStackEntries;
@@ -1177,12 +1180,15 @@ void FNiagaraSystemViewModel::ResetEmitterHandleViewModelsAndTracks()
 		EmitterHandleViewModel->Reset();
 	}
 
-	TArray<UMovieSceneTrack*> MasterTracks = NiagaraSequence->GetMovieScene()->GetMasterTracks();
-	for (UMovieSceneTrack* MasterTrack : MasterTracks)
+	if (NiagaraSequence)
 	{
-		if (MasterTrack != nullptr)
+		TArray<UMovieSceneTrack*> MasterTracks = NiagaraSequence->GetMovieScene()->GetMasterTracks();
+		for (UMovieSceneTrack* MasterTrack : MasterTracks)
 		{
-			NiagaraSequence->GetMovieScene()->RemoveMasterTrack(*MasterTrack);
+			if (MasterTrack != nullptr)
+			{
+				NiagaraSequence->GetMovieScene()->RemoveMasterTrack(*MasterTrack);
+			}
 		}
 	}
 }
@@ -1289,64 +1295,72 @@ void FNiagaraSystemViewModel::RefreshSequencerTracks()
 {
 	TGuardValue<bool> UpdateGuard(bUpdatingSequencerFromEmitterDataChange, true);
 
-	TArray<UMovieSceneTrack*> MasterTracks = NiagaraSequence->GetMovieScene()->GetMasterTracks();
-	for (UMovieSceneTrack* MasterTrack : MasterTracks)
+	if (Sequencer.IsValid())
 	{
-		if (MasterTrack != nullptr)
+		TArray<UMovieSceneTrack*> MasterTracks = NiagaraSequence->GetMovieScene()->GetMasterTracks();
+		for (UMovieSceneTrack* MasterTrack : MasterTracks)
 		{
-			NiagaraSequence->GetMovieScene()->RemoveMasterTrack(*MasterTrack);
+			if (MasterTrack != nullptr)
+			{
+				NiagaraSequence->GetMovieScene()->RemoveMasterTrack(*MasterTrack);
+			}
 		}
-	}
 
-	TMap<FGuid, UMovieSceneNiagaraEmitterTrack*> EmitterHandleIdToTrackMap;
-	for (TSharedRef<FNiagaraEmitterHandleViewModel> EmitterHandleViewModel : EmitterHandleViewModels)
-	{
-		UMovieSceneNiagaraEmitterTrack* EmitterTrack = Cast<UMovieSceneNiagaraEmitterTrack>(NiagaraSequence->GetMovieScene()->AddMasterTrack(UMovieSceneNiagaraEmitterTrack::StaticClass()));
-		EmitterTrack->Initialize(*this, EmitterHandleViewModel, NiagaraSequence->GetMovieScene()->GetTickResolution());
-		EmitterHandleIdToTrackMap.Add(EmitterHandleViewModel->GetId(), EmitterTrack);
-	}
+		TMap<FGuid, UMovieSceneNiagaraEmitterTrack*> EmitterHandleIdToTrackMap;
+		for (TSharedRef<FNiagaraEmitterHandleViewModel> EmitterHandleViewModel : EmitterHandleViewModels)
+		{
+			UMovieSceneNiagaraEmitterTrack* EmitterTrack = Cast<UMovieSceneNiagaraEmitterTrack>(NiagaraSequence->GetMovieScene()->AddMasterTrack(UMovieSceneNiagaraEmitterTrack::StaticClass()));
+			EmitterTrack->Initialize(*this, EmitterHandleViewModel, NiagaraSequence->GetMovieScene()->GetTickResolution());
+			EmitterHandleIdToTrackMap.Add(EmitterHandleViewModel->GetId(), EmitterTrack);
+		}
 
-	TArray<UMovieSceneFolder*>& MovieSceneRootFolders = NiagaraSequence->GetMovieScene()->GetRootFolders();
-	MovieSceneRootFolders.Empty();
+		TArray<UMovieSceneFolder*>& MovieSceneRootFolders = NiagaraSequence->GetMovieScene()->GetRootFolders();
+		MovieSceneRootFolders.Empty();
 
-	const UNiagaraSystemEditorData& SystemEditorData = GetEditorData();
-	UNiagaraSystemEditorFolder& RootFolder = SystemEditorData.GetRootFolder();
-	for (const UNiagaraSystemEditorFolder* RootChildFolder : RootFolder.GetChildFolders())
-	{
-		UMovieSceneFolder* MovieSceneRootFolder = NewObject<UMovieSceneFolder>(NiagaraSequence->GetMovieScene(), RootChildFolder->GetFolderName(), RF_Transactional);
-		MovieSceneRootFolder->SetFolderName(RootChildFolder->GetFolderName());
-		MovieSceneRootFolders.Add(MovieSceneRootFolder);
-		PopulateChildMovieSceneFoldersFromNiagaraFolders(RootChildFolder, MovieSceneRootFolder, EmitterHandleIdToTrackMap);
+		const UNiagaraSystemEditorData& SystemEditorData = GetEditorData();
+		UNiagaraSystemEditorFolder& RootFolder = SystemEditorData.GetRootFolder();
+		for (const UNiagaraSystemEditorFolder* RootChildFolder : RootFolder.GetChildFolders())
+		{
+			UMovieSceneFolder* MovieSceneRootFolder = NewObject<UMovieSceneFolder>(NiagaraSequence->GetMovieScene(), RootChildFolder->GetFolderName(), RF_Transactional);
+			MovieSceneRootFolder->SetFolderName(RootChildFolder->GetFolderName());
+			MovieSceneRootFolders.Add(MovieSceneRootFolder);
+			PopulateChildMovieSceneFoldersFromNiagaraFolders(RootChildFolder, MovieSceneRootFolder, EmitterHandleIdToTrackMap);
+		}
+		Sequencer->NotifyMovieSceneDataChanged(EMovieSceneDataChangeType::MovieSceneStructureItemsChanged);
 	}
 
-	Sequencer->NotifyMovieSceneDataChanged(EMovieSceneDataChangeType::MovieSceneStructureItemsChanged);
-
 	// Since we just rebuilt all of the sequencer tracks, these updates don't need to be done.
 	EmitterIdsRequiringSequencerTrackUpdate.Empty();
 }
 
 void FNiagaraSystemViewModel::UpdateSequencerTracksForEmitters(const TArray<FGuid>& EmitterIdsRequiringUpdate)
 {
-	TGuardValue<bool> UpdateGuard(bUpdatingSequencerFromEmitterDataChange, true);
-	for (UMovieSceneTrack* Track : NiagaraSequence->GetMovieScene()->GetMasterTracks())
+	if (Sequencer.IsValid())
 	{
-		UMovieSceneNiagaraEmitterTrack* EmitterTrack = CastChecked<UMovieSceneNiagaraEmitterTrack>(Track);
-		if (EmitterIdsRequiringUpdate.Contains(EmitterTrack->GetEmitterHandleViewModel()->GetId()))
+		TGuardValue<bool> UpdateGuard(bUpdatingSequencerFromEmitterDataChange, true);
+		for (UMovieSceneTrack* Track : NiagaraSequence->GetMovieScene()->GetMasterTracks())
 		{
-			EmitterTrack->UpdateTrackFromEmitterGraphChange(NiagaraSequence->GetMovieScene()->GetTickResolution());
+			UMovieSceneNiagaraEmitterTrack* EmitterTrack = CastChecked<UMovieSceneNiagaraEmitterTrack>(Track);
+			if (EmitterIdsRequiringUpdate.Contains(EmitterTrack->GetEmitterHandleViewModel()->GetId()))
+			{
+				EmitterTrack->UpdateTrackFromEmitterGraphChange(NiagaraSequence->GetMovieScene()->GetTickResolution());
+			}
 		}
+		Sequencer->NotifyMovieSceneDataChanged(EMovieSceneDataChangeType::MovieSceneStructureItemsChanged);
 	}
-	Sequencer->NotifyMovieSceneDataChanged(EMovieSceneDataChangeType::MovieSceneStructureItemsChanged);
 }
 
 UMovieSceneNiagaraEmitterTrack* FNiagaraSystemViewModel::GetTrackForHandleViewModel(TSharedRef<FNiagaraEmitterHandleViewModel> EmitterHandleViewModel)
 {
-	for (UMovieSceneTrack* Track : NiagaraSequence->GetMovieScene()->GetMasterTracks())
+	if (NiagaraSequence)
 	{
-		UMovieSceneNiagaraEmitterTrack* EmitterTrack = CastChecked<UMovieSceneNiagaraEmitterTrack>(Track);
-		if (EmitterTrack->GetEmitterHandleViewModel() == EmitterHandleViewModel)
+		for (UMovieSceneTrack* Track : NiagaraSequence->GetMovieScene()->GetMasterTracks())
 		{
-			return EmitterTrack;
+			UMovieSceneNiagaraEmitterTrack* EmitterTrack = CastChecked<UMovieSceneNiagaraEmitterTrack>(Track);
+			if (EmitterTrack->GetEmitterHandleViewModel() == EmitterHandleViewModel)
+			{
+				return EmitterTrack;
+			}
 		}
 	}
 	return nullptr;
@@ -1354,6 +1368,11 @@ UMovieSceneNiagaraEmitterTrack* FNiagaraSystemViewModel::GetTrackForHandleViewMo
 
 void FNiagaraSystemViewModel::SetupSequencer()
 {
+	if (EditMode == ENiagaraSystemViewModelEditMode::EmitterDuringMerge)
+	{
+		// we don't need a sequencer when merging emitters
+		return;
+	}
 	NiagaraSequence = NewObject<UNiagaraSequence>(GetTransientPackage());
 	UMovieScene* MovieScene = NewObject<UMovieScene>(NiagaraSequence, FName("Niagara System MovieScene"), RF_Transactional);
 	MovieScene->SetDisplayRate(FFrameRate(240, 1));
@@ -1403,7 +1422,7 @@ void FNiagaraSystemViewModel::SetupSequencer()
 void FNiagaraSystemViewModel::SnapToNextSpeed(const FString& PropertyName, const UNiagaraEditorSettings* Settings)
 {
 	// we update the speed in any case
-	if (Sequencer != nullptr)
+	if (Sequencer.IsValid())
 	{
 		Sequencer->SnapToClosestPlaybackSpeed();
 	}
@@ -1428,19 +1447,23 @@ void FNiagaraSystemViewModel::ResetSystem()
 
 void FNiagaraSystemViewModel::ResetSystem(ETimeResetMode TimeResetMode, EMultiResetMode MultiResetMode, EReinitMode ReinitMode)
 {
-	bool bResetAge = TimeResetMode == ETimeResetMode::AllowResetTime && (Sequencer->GetPlaybackStatus() == EMovieScenePlayerStatus::Playing || Sequencer->GetPlaybackStatus() == EMovieScenePlayerStatus::Stopped || EditorSettings->GetResimulateOnChangeWhilePaused() == false);
-	if (bResetAge)
+	bool bResetAge = false;
+	if (Sequencer.IsValid())
 	{
-		TGuardValue<bool> Guard(bSettingSequencerTimeDirectly, true);
-		if (Sequencer->GetPlaybackStatus() == EMovieScenePlayerStatus::Playing)
+		bResetAge = TimeResetMode == ETimeResetMode::AllowResetTime && (Sequencer->GetPlaybackStatus() == EMovieScenePlayerStatus::Playing || Sequencer->GetPlaybackStatus() == EMovieScenePlayerStatus::Stopped || EditorSettings->GetResimulateOnChangeWhilePaused() == false);
+		if (bResetAge)
 		{
-			Sequencer->SetPlaybackStatus(EMovieScenePlayerStatus::Paused);
-			Sequencer->SetGlobalTime(0);
-			Sequencer->SetPlaybackStatus(EMovieScenePlayerStatus::Playing);
-		}
-		else
-		{
-			Sequencer->SetGlobalTime(0);
+			TGuardValue<bool> Guard(bSettingSequencerTimeDirectly, true);
+			if (Sequencer->GetPlaybackStatus() == EMovieScenePlayerStatus::Playing)
+			{
+				Sequencer->SetPlaybackStatus(EMovieScenePlayerStatus::Paused);
+				Sequencer->SetGlobalTime(0);
+				Sequencer->SetPlaybackStatus(EMovieScenePlayerStatus::Playing);
+			}
+			else
+			{
+				Sequencer->SetGlobalTime(0);
+			}
 		}
 	}
 
@@ -1589,7 +1612,7 @@ void FNiagaraSystemViewModel::EmitterHandlePropertyChanged(FGuid OwningEmitterHa
 {
 	// When the emitter handle changes, refresh the System scripts emitter nodes and the sequencer tracks just in case the
 	// property that changed was the handles emitter.
-	if (bUpdatingEmittersFromSequencerDataChange == false)
+	if (bUpdatingEmittersFromSequencerDataChange == false && NiagaraSequence)
 	{
 		TGuardValue<bool> UpdateGuard(bUpdatingSequencerFromEmitterDataChange, true);
 		for (UMovieSceneTrack* Track : NiagaraSequence->GetMovieScene()->GetMasterTracks())
@@ -1667,7 +1690,7 @@ void FNiagaraSystemViewModel::EmitterParameterStoreChanged(const FNiagaraParamet
 		{
 			UpdateSequencerTracksForEmitters(EmitterIdsRequiringSequencerTrackUpdate);
 		}
-		else
+		else if (Sequencer.IsValid())
 		{
 			TGuardValue<bool> UpdateGuard(bUpdatingSequencerFromEmitterDataChange, true);
 			for (UMovieSceneTrack* Track : NiagaraSequence->GetMovieScene()->GetMasterTracks())
@@ -1675,6 +1698,7 @@ void FNiagaraSystemViewModel::EmitterParameterStoreChanged(const FNiagaraParamet
 				UMovieSceneNiagaraEmitterTrack* EmitterTrack = CastChecked<UMovieSceneNiagaraEmitterTrack>(Track);
 				EmitterTrack->UpdateTrackFromEmitterParameterChange(NiagaraSequence->GetMovieScene()->GetTickResolution());
 			}
+			
 			Sequencer->NotifyMovieSceneDataChanged(EMovieSceneDataChangeType::TrackValueChanged);
 		}
 	}
@@ -1691,7 +1715,7 @@ void FNiagaraSystemViewModel::UpdateSimulationFromParameterChange()
 	}
 	else
 	{
-		if (Sequencer->GetPlaybackStatus() == EMovieScenePlayerStatus::Stopped)
+		if (Sequencer && Sequencer->GetPlaybackStatus() == EMovieScenePlayerStatus::Stopped)
 		{
 			// TODO: Update the view when paused and reset on change is turned off.
 		}
@@ -1761,7 +1785,7 @@ void PopulateNiagaraFoldersFromMovieSceneFolders(const TArray<UMovieSceneFolder*
 
 void FNiagaraSystemViewModel::SequencerDataChanged(EMovieSceneDataChangeType DataChangeType)
 {
-	if (bUpdatingSequencerFromEmitterDataChange == false && GIsTransacting == false)
+	if (bUpdatingSequencerFromEmitterDataChange == false && GIsTransacting == false && NiagaraSequence)
 	{
 		TGuardValue<bool> UpdateGuard(bUpdatingEmittersFromSequencerDataChange, true);
 
@@ -1871,7 +1895,7 @@ void FNiagaraSystemViewModel::SequencerDataChanged(EMovieSceneDataChangeType Dat
 
 void FNiagaraSystemViewModel::SequencerTimeChanged()
 {
-	if (!PreviewComponent || !SystemInstance || !SystemInstance->GetAreDataInterfacesInitialized())
+	if (!PreviewComponent || !SystemInstance || !SystemInstance->GetAreDataInterfacesInitialized() || !Sequencer.IsValid())
 	{
 		return;
 	}
@@ -1949,6 +1973,11 @@ void FNiagaraSystemViewModel::SequencerSectionSelectionChanged(TArray<UMovieScen
 
 void FNiagaraSystemViewModel::UpdateEmitterHandleSelectionFromSequencer()
 {
+	if (!Sequencer.IsValid())
+	{
+		return;
+	}
+	
 	TArray<FGuid> NewSelectedEmitterHandleIds;
 
 	TArray<UMovieSceneTrack*> SelectedTracks;
@@ -2006,6 +2035,10 @@ void FNiagaraSystemViewModel::UpdateEmitterHandleSelectionFromSequencer()
 
 void FNiagaraSystemViewModel::UpdateSequencerFromEmitterHandleSelection()
 {
+	if (!Sequencer.IsValid())
+	{
+		return;
+	}
 	TGuardValue<bool> UpdateGuard(bUpdatingSequencerSelectionFromSystem, true);
 	Sequencer->EmptySelection();
 	for (FGuid SelectedEmitterHandleId : SelectionViewModel->GetSelectedEmitterHandleIds())
diff --git a/Engine/Plugins/FX/Niagara/Source/NiagaraEditor/Public/NiagaraEditorUtilities.h b/Engine/Plugins/FX/Niagara/Source/NiagaraEditor/Public/NiagaraEditorUtilities.h
--- a/Engine/Plugins/FX/Niagara/Source/NiagaraEditor/Public/NiagaraEditorUtilities.h
+++ b/Engine/Plugins/FX/Niagara/Source/NiagaraEditor/Public/NiagaraEditorUtilities.h
@@ -240,7 +240,7 @@ namespace FNiagaraEditorUtilities
 
 	TArray<UNiagaraComponent*> GetComponentsThatReferenceSystemViewModel(const FNiagaraSystemViewModel& ReferencedSystemViewModel);
 
-	NIAGARAEDITOR_API const FGuid AddEmitterToSystem(UNiagaraSystem& InSystem, UNiagaraEmitter& InEmitterToAdd);
+	NIAGARAEDITOR_API const FGuid AddEmitterToSystem(UNiagaraSystem& InSystem, UNiagaraEmitter& InEmitterToAdd, bool bCreateCopy = true);
 
 	void RemoveEmittersFromSystemByEmitterHandleId(UNiagaraSystem& InSystem, TSet<FGuid> EmitterHandleIdsToDelete);
 
diff --git a/Engine/Plugins/FX/Niagara/Source/NiagaraEditor/Public/ViewModels/NiagaraSystemViewModel.h b/Engine/Plugins/FX/Niagara/Source/NiagaraEditor/Public/ViewModels/NiagaraSystemViewModel.h
--- a/Engine/Plugins/FX/Niagara/Source/NiagaraEditor/Public/ViewModels/NiagaraSystemViewModel.h
+++ b/Engine/Plugins/FX/Niagara/Source/NiagaraEditor/Public/ViewModels/NiagaraSystemViewModel.h
@@ -49,6 +49,8 @@ enum class NIAGARAEDITOR_API ENiagaraSystemViewModelEditMode
 	SystemAsset,
 	/** An emitter asset is being edited.  In this mode the system scripts will not be editable and all emitter editing options are available. */
 	EmitterAsset,
+	/** Special case where an emitter asset is being edited during a merge. In this mode changes made to the emitter being edited don't have to be applied, but directly affect the source. */
+	EmitterDuringMerge,
 };
 
 /** Defines options for the niagara System view model */
