Fixed build error from previous merge

diff --git a/Engine/Plugins/FX/Niagara/Source/Niagara/Classes/NiagaraScript.h b/Engine/Plugins/FX/Niagara/Source/Niagara/Classes/NiagaraScript.h
--- a/Engine/Plugins/FX/Niagara/Source/Niagara/Classes/NiagaraScript.h
+++ b/Engine/Plugins/FX/Niagara/Source/Niagara/Classes/NiagaraScript.h
@@ -825,7 +825,7 @@ public:
 	//~ End UObject interface
 
 	//~ Begin UNiagaraScriptBase interface
-	virtual void ModifyCompilationEnvironment(struct FShaderCompilerEnvironment& OutEnvironment) const override;
+	virtual void ModifyCompilationEnvironment(EShaderPlatform Platform, struct FShaderCompilerEnvironment& OutEnvironment) const override;
 	virtual TConstArrayView<FSimulationStageMetaData> GetSimulationStageMetaData() const override { return MakeArrayView(CachedScriptVM.SimulationStageMetaData); }
 	//~ End UNiagaraScriptBase interface
 
