- Editgate 16303301 - Fix custom depth in collision data interface

#rnx

diff --git a/Engine/Plugins/FX/Niagara/Shaders/Private/NiagaraDataInterfaceCollisionQuery.ush b/Engine/Plugins/FX/Niagara/Shaders/Private/NiagaraDataInterfaceCollisionQuery.ush
--- a/Engine/Plugins/FX/Niagara/Shaders/Private/NiagaraDataInterfaceCollisionQuery.ush
+++ b/Engine/Plugins/FX/Niagara/Shaders/Private/NiagaraDataInterfaceCollisionQuery.ush
@@ -24,7 +24,7 @@ float NDICollisionQuery_GetCustomSceneDepth(float2 ScreenUV)
 #if (FEATURE_LEVEL <= FEATURE_LEVEL_ES3_1)
 	return ConvertFromDeviceZ(Texture2DSampleLevel(MobileSceneTextures.CustomDepthTexture, MobileSceneTextures.CustomDepthTextureSampler, ScreenUV, 0).r);
 #else
-	return ConvertFromDeviceZ(Texture2DSampleLevel(SceneTexturesStruct.CustomDepthTexture, SceneTexturesStruct_SceneDepthTextureSampler, ScreenUV, 0).r);
+	return ConvertFromDeviceZ(Texture2DSampleLevel(SceneTexturesStruct.CustomDepthTexture, SceneTexturesStruct_CustomDepthTextureSampler, ScreenUV, 0).r);
 #endif
 }
 
@@ -118,6 +118,11 @@ void NDICollisionQuery_QueryCustomDepthGPU(in float3 In_SamplePos, out float Out
 		float3 WorldPosition1 = WorldPositionFromSceneDepth(ScreenPosition.xy - float2(0.0, 2 * View.ViewSizeAndInvSize.w), SceneDepth1);
 		float3 WorldNormal = normalize(cross(WorldPosition0 - WorldPosition, WorldPosition1 - WorldPosition));
 #endif
+
+		// Set outputs
+		Out_SceneDepth = SceneDepth;
+		Out_WorldPos = WorldPosition;
+		Out_WorldNormal = WorldNormal;
 	}
 	else
 	{
