fixing unrealgame after merge

diff --git a/Engine/Plugins/FX/Niagara/Source/Niagara/Private/NiagaraEmitter.cpp b/Engine/Plugins/FX/Niagara/Source/Niagara/Private/NiagaraEmitter.cpp
--- a/Engine/Plugins/FX/Niagara/Source/Niagara/Private/NiagaraEmitter.cpp
+++ b/Engine/Plugins/FX/Niagara/Source/Niagara/Private/NiagaraEmitter.cpp
@@ -577,12 +577,13 @@ void UNiagaraEmitter::PostLoad()
 	}
 	EnsureScriptsPostLoaded();
 	
-
+#if WITH_EDITORONLY_DATA
 	// this can only ever be true for old assets that haven't been loaded yet, so this won't overwrite subsequent changes to the template specification
 	if(bIsTemplateAsset_DEPRECATED)
 	{
 		TemplateSpecification = ENiagaraScriptTemplateSpecification::Template;
 	}
+#endif
 	
 	// we are not yet finished, but we do the rest of the work after postload in a separate task
 	UpdateTaskRef = TGraphTask<FNiagaraEmitterUpdateTask>::CreateTask().ConstructAndDispatchWhenReady(this);
