UE Fix crash during load of certain older assets that have nullptr entries in VariableToScriptVariable.

diff --git a/Engine/Plugins/FX/Niagara/Source/NiagaraEditor/Private/NiagaraGraph.cpp b/Engine/Plugins/FX/Niagara/Source/NiagaraEditor/Private/NiagaraGraph.cpp
--- a/Engine/Plugins/FX/Niagara/Source/NiagaraEditor/Private/NiagaraGraph.cpp
+++ b/Engine/Plugins/FX/Niagara/Source/NiagaraEditor/Private/NiagaraGraph.cpp
@@ -140,27 +140,20 @@ void UNiagaraGraph::PostLoad()
 
 	const int32 NiagaraVer = GetLinkerCustomVersion(FNiagaraCustomVersion::GUID);
 
-	if (NiagaraVer < FNiagaraCustomVersion::FixNullScriptVariables)
+	for (auto It = VariableToScriptVariable.CreateIterator(); It; ++It)
 	{
-		for (auto It = VariableToScriptVariable.CreateIterator(); It; ++It)
-		{
-			FNiagaraVariable Var = It.Key();
-			UE_TRANSITIONAL_OBJECT_PTR(UNiagaraScriptVariable)& ScriptVar = It.Value();
+		FNiagaraVariable Var = It.Key();
+		UE_TRANSITIONAL_OBJECT_PTR(UNiagaraScriptVariable)& ScriptVar = It.Value();
 
-			if (ScriptVar == nullptr)
-			{
-				ScriptVar = NewObject<UNiagaraScriptVariable>(const_cast<UNiagaraGraph*>(this));
-				UE_LOG(LogNiagaraEditor, Display, TEXT("Fixed null UNiagaraScriptVariable | variable %s | asset path %s"), *Var.GetName().ToString(), *GetPathName());
-			}
+		if (ScriptVar == nullptr)
+		{
+			ScriptVar = NewObject<UNiagaraScriptVariable>(const_cast<UNiagaraGraph*>(this));
+			UE_LOG(LogNiagaraEditor, Display, TEXT("Fixed null UNiagaraScriptVariable | variable %s | asset path %s"), *Var.GetName().ToString(), *GetPathName());
 		}
-	}
-
-	// Conditional postload all ScriptVars to ensure static switch default values are allocated as these are required when postloading all graph nodes later.
-	if (NiagaraVer < FNiagaraCustomVersion::MoveDefaultValueFromFNiagaraVariableMetaDataToUNiagaraScriptVariable)
-	{
-		for (auto It = VariableToScriptVariable.CreateIterator(); It; ++It)
+		else
 		{
-			It.Value()->ConditionalPostLoad();
+			// Conditional postload all ScriptVars to ensure static switch default values are allocated as these are required when postloading all graph nodes later.
+			ScriptVar->ConditionalPostLoad();
 		}
 	}
 
