- Integrate CL 16314302 - Fix various compile errors in debug draw

#rnx

diff --git a/Engine/Plugins/FX/Niagara/Shaders/Private/NiagaraDataInterfaceDebugDraw.ush b/Engine/Plugins/FX/Niagara/Shaders/Private/NiagaraDataInterfaceDebugDraw.ush
--- a/Engine/Plugins/FX/Niagara/Shaders/Private/NiagaraDataInterfaceDebugDraw.ush
+++ b/Engine/Plugins/FX/Niagara/Shaders/Private/NiagaraDataInterfaceDebugDraw.ush
@@ -104,10 +104,10 @@ void NDIDebugDraw_DrawRectangle(bool bDraw, float3 Location, float3 XAxis, float
 			float3 XAxisNorm = normalize(XAxis);
 			float3 YAxisNorm = normalize(YAxis);
 
-			float3 StartPosition = Location - ((XAxisNorm * Extents.X) + (YAxisNorm * Extents.Y));
+			float3 StartPosition = Location - ((XAxisNorm * Extents.x) + (YAxisNorm * Extents.y));
 			float2 Step = (Extents * 2) / float2(SegmentsX, SegmentsY);
-			float3 StepX = XAxisNorm * Step.X;
-			float3 StepY = YAxisNorm * Step.Y;
+			float3 StepX = XAxisNorm * Step.x;
+			float3 StepY = YAxisNorm * Step.y;
 
 			// Add all Y axis lines
 			for (int X = LowerBound; X <= UpperBoundX; X++)
@@ -115,7 +115,7 @@ void NDIDebugDraw_DrawRectangle(bool bDraw, float3 Location, float3 XAxis, float
 				float3 LineStart = StartPosition + (StepX * X);
 				float3 LineEnd = LineStart + (StepY * SegmentsY);
 
-				PackLine(RWNDIDebugDrawLineVertex, InstanceIndex + i, LineStart, LineEnd, Color);
+				PackLine(RWNDIDebugDrawLineVertex, InstanceIndex++, LineStart, LineEnd, Color);
 			}
 
 			// Add all X axis lines
@@ -124,7 +124,7 @@ void NDIDebugDraw_DrawRectangle(bool bDraw, float3 Location, float3 XAxis, float
 				float3 LineStart = StartPosition + (StepY * Y);
 				float3 LineEnd = LineStart + (StepX * SegmentsX);
 
-				PackLine(RWNDIDebugDrawLineVertex, InstanceIndex + i, LineStart, LineEnd, Color);
+				PackLine(RWNDIDebugDrawLineVertex, InstanceIndex++, LineStart, LineEnd, Color);
 			}
 		}
 		else
@@ -314,14 +314,14 @@ void NDIDebugDraw_Sphere(bool bDraw, float3 Location, float Radius, int Segments
 }
 
 
-void NDIDebugDraw_Cylinder(bool bDraw, float3 Location, float3 Axis, float Height, float Radius, int NumHeightSegments, NumRadiusSegments, float4 Color)
+void NDIDebugDraw_Cylinder(bool bDraw, float3 Location, float3 Axis, float Height, float Radius, int NumHeightSegments, int NumRadiusSegments, float4 Color)
 {
 	if ( bDraw && (NDIDebugDrawLineMaxInstances > 0) )
 	{
 		NumHeightSegments = clamp(NumHeightSegments, 1, 16);
 		NumRadiusSegments = clamp(NumRadiusSegments, 3, 16);
 
-		int NumLines = NumRadiusSegments + (NumRadiusSegments * NumHeightSegments + 1)
+		int NumLines = NumRadiusSegments + (NumRadiusSegments * NumHeightSegments + 1);
 
 		uint InstanceIndex = 0;
 		InterlockedAdd(RWNDIDebugDrawArgs[1], NumLines, InstanceIndex);
@@ -338,7 +338,7 @@ void NDIDebugDraw_Cylinder(bool bDraw, float3 Location, float3 Axis, float Heigh
 			float3 TangentYScaled = TangentY * Radius;
 
 			float3 HalfHeightOffset = AxisNorm * Height * 0.5f;
-			float3 AxisStep = AxisNorm * Height / NumHeightSegments);
+			float3 AxisStep = AxisNorm * Height / NumHeightSegments;
 
 			float CurrentRotation = 0.0f;
 			float3 LastPoint = Location + (TangentXScaled * cos(CurrentRotation)) + (TangentYScaled * sin(CurrentRotation));
@@ -384,7 +384,7 @@ void NDIDebugDraw_Cone(bool bDraw, float3 Location, float3 Axis, float Height, f
 		NumHeightSegments = clamp(NumHeightSegments, 1, 16);
 		NumRadiusSegments = clamp(NumRadiusSegments, 3, 16);
 
-		int NumLines = NumRadiusSegments + (NumRadiusSegments * NumHeightSegments + 1)
+		int NumLines = NumRadiusSegments + (NumRadiusSegments * NumHeightSegments + 1);
 
 		uint InstanceIndex = 0;
 		InterlockedAdd(RWNDIDebugDrawArgs[1], NumLines, InstanceIndex);
@@ -394,29 +394,31 @@ void NDIDebugDraw_Cone(bool bDraw, float3 Location, float3 Axis, float Height, f
 
 			const float3x3 AxisRotation = NiagaraFindRotationBetweenNormals(float3(0,0,1), AxisNorm);
 
-			float3 TangentX = mul(AxisRotation, float3(1, 0, 0));
-			float3 TangentY = mul(AxisRotation, float3(0, 1, 0));
+			float3 TangentXScaled = mul(AxisRotation, float3(1, 0, 0));
+			float3 TangentYScaled = mul(AxisRotation, float3(0, 1, 0));
 
 			float3 HalfHeightOffset = AxisNorm * Height * 0.5f;
-			float3 AxisStep = AxisNorm * Height / NumHeightSegments);
+			float3 AxisStep = AxisNorm * Height / NumHeightSegments;
 
 			float CurrentRotation = 0.0f;
-			float3 LastPoint = Location + (TangentXScaled * cos(CurrentRotation)) + (TangentYScaled * sin(CurrentRotation));
+			float3 LastPointNormal = Location + (TangentXScaled * cos(CurrentRotation)) + (TangentYScaled * sin(CurrentRotation));
+
+			float HeightSegmentStepAlpha = 1.0 / NumHeightSegments;
 
 			for (int Idx = 1; Idx <= NumRadiusSegments; Idx++)
 			{
 				CurrentRotation = 2.0f * PI / NumRadiusSegments * Idx;
-				float3 CurrentPoint = Location + (TangentXScaled * cos(CurrentRotation)) + (TangentYScaled * sin(CurrentRotation));
+				float3 CurrentPointNormal = (TangentXScaled * cos(CurrentRotation)) + (TangentYScaled * sin(CurrentRotation));
+				float3 CurrentPoint = Location + CurrentPointNormal;
 				
 				float3 LocationBottom = Location - HalfHeightOffset;
 				float3 LocationTop = Location + HalfHeightOffset;
 
+				float3 LastBottom = LocationBottom + (LastPointNormal * RadiusBottom);
+				float3 LastTop = LocationTop + (LastPointNormal * RadiusTop);
 
-				float3 LastBottom = (LocationBottom + (LastNormal * RadiusBottom));
-				float3 LastTop = (LocationTop + (LastNormal * RadiusTop));
-
-				float3 CurrentBottom = (LocationBottom + (CurrentNormal * RadiusBottom));
-				float3 CurrentTop = (LocationTop + (CurrentNormal * RadiusTop));
+				float3 CurrentBottom = (LocationBottom + (CurrentPointNormal * RadiusBottom));
+				float3 CurrentTop = (LocationTop + (CurrentPointNormal * RadiusTop));
 				
 				// Add Height Line
 				PackLine(RWNDIDebugDrawLineVertex, InstanceIndex++, CurrentBottom, CurrentTop, Color);
@@ -432,13 +434,13 @@ void NDIDebugDraw_Cone(bool bDraw, float3 Location, float3 Axis, float Height, f
 					float3 RingOffset = AxisStep * HeightIdx;
 					float3 LocationRing = LocationBottom + RingOffset;
 
-					float3 LastRingPosition = LocationRing + (LastNormal * RingRadius);
-					float3 CurrRingPosition = LocationRing + (CurrentNormal * RingRadius);
+					float3 LastRingPosition = LocationRing + (LastPointNormal * RingRadius);
+					float3 CurrRingPosition = LocationRing + (CurrentPointNormal * RingRadius);
 
 					PackLine(RWNDIDebugDrawLineVertex, InstanceIndex++, LastRingPosition, CurrRingPosition, Color);
 				}
 
-				LastPoint = CurrentPoint;
+				LastPointNormal = CurrentPointNormal;
 			}
 		}
 		else
diff --git a/Engine/Plugins/FX/Niagara/Source/Niagara/Private/DataInterface/NiagaraDataInterfaceDebugDraw.cpp b/Engine/Plugins/FX/Niagara/Source/Niagara/Private/DataInterface/NiagaraDataInterfaceDebugDraw.cpp
--- a/Engine/Plugins/FX/Niagara/Source/Niagara/Private/DataInterface/NiagaraDataInterfaceDebugDraw.cpp
+++ b/Engine/Plugins/FX/Niagara/Source/Niagara/Private/DataInterface/NiagaraDataInterfaceDebugDraw.cpp
@@ -3007,7 +3007,7 @@ bool UNiagaraDataInterfaceDebugDraw::GetFunctionHLSL(const FNiagaraDataInterface
 	}
 	else if (FunctionInfo.DefinitionName == NDIDebugDrawLocal::DrawCylinderName)
 	{
-		static const TCHAR* FormatSample = TEXT("void {InstanceFunctionName}(bool bExecute, float3 Location, float3 Axis, float Height, float Radius, int NumHeightSegments, NumRadiusSegments, float4 Color) { NDIDebugDraw_Cylinder(bExecute, Location, Axis, Height, Radius, NumHeightSegments, NumRadiusSegments, Color); }\n");
+		static const TCHAR* FormatSample = TEXT("void {InstanceFunctionName}(bool bExecute, float3 Location, float3 Axis, float Height, float Radius, int NumHeightSegments, int NumRadiusSegments, float4 Color) { NDIDebugDraw_Cylinder(bExecute, Location, Axis, Height, Radius, NumHeightSegments, NumRadiusSegments, Color); }\n");
 		OutHLSL += FString::Format(FormatSample, ArgsSample);
 		return true;
 	}
