Fix issues with attempting to set a StaticMesh or Texture variable on a NiagaraComponent before it has activated.

diff --git a/Engine/Plugins/FX/Niagara/Source/Niagara/Private/NiagaraComponent.cpp b/Engine/Plugins/FX/Niagara/Source/Niagara/Private/NiagaraComponent.cpp
--- a/Engine/Plugins/FX/Niagara/Source/Niagara/Private/NiagaraComponent.cpp
+++ b/Engine/Plugins/FX/Niagara/Source/Niagara/Private/NiagaraComponent.cpp
@@ -1903,26 +1903,34 @@ void UNiagaraComponent::SetVariableMaterial(FName InVariableName, UMaterialInter
 
 void UNiagaraComponent::SetVariableStaticMesh(FName InVariableName, UStaticMesh* InValue)
 {
+	const FNiagaraVariable VariableDesc(FNiagaraTypeDefinition::GetUStaticMeshDef(), InVariableName);
 	if (SystemInstanceController.IsValid())
 	{
 		TWeakObjectPtr<UStaticMesh> ObjPtr = InValue;
 		SystemInstanceController->SetVariable_Deferred(InVariableName, ObjPtr);
 	}
+	else
+	{
+		OverrideParameters.SetUObject(InValue, VariableDesc);
+	}
 #if WITH_EDITOR
-	const FNiagaraVariable VariableDesc(FNiagaraTypeDefinition::GetUStaticMeshDef(), InVariableName);
 	SetParameterOverride(VariableDesc, FNiagaraVariant(InValue));
 #endif
 }
 
 void UNiagaraComponent::SetVariableTexture(FName InVariableName, UTexture* Texture)
 {
+	const FNiagaraVariable VariableDesc(FNiagaraTypeDefinition::GetUTextureDef(), InVariableName);
 	if (SystemInstanceController.IsValid())
 	{
 		TWeakObjectPtr<UTexture> ObjPtr = Texture;
 		SystemInstanceController->SetVariable_Deferred(InVariableName, ObjPtr);
 	}
+	else
+	{
+		OverrideParameters.SetUObject(Texture, VariableDesc);
+	}
 #if WITH_EDITOR
-	const FNiagaraVariable VariableDesc(FNiagaraTypeDefinition::GetUTextureDef(), InVariableName);
 	SetParameterOverride(VariableDesc, FNiagaraVariant(Texture));
 #endif
 }
