Prepare for the deprecation of direct access to the Resource member of UTexture in favor of the GetResource() accessor.

#rnx

diff --git a/Engine/Plugins/2D/Paper2D/Source/Paper2D/Classes/SpriteDrawCall.h b/Engine/Plugins/2D/Paper2D/Source/Paper2D/Classes/SpriteDrawCall.h
--- a/Engine/Plugins/2D/Paper2D/Source/Paper2D/Classes/SpriteDrawCall.h
+++ b/Engine/Plugins/2D/Paper2D/Source/Paper2D/Classes/SpriteDrawCall.h
@@ -40,6 +40,6 @@ public:
 
 	bool IsValid() const
 	{
-		return (RenderVerts.Num() > 0) && (BaseTexture != nullptr) && (BaseTexture->Resource != nullptr);
+		return (RenderVerts.Num() > 0) && (BaseTexture != nullptr) && (BaseTexture->GetResource() != nullptr);
 	}
 };
diff --git a/Engine/Plugins/2D/Paper2D/Source/Paper2D/Private/PaperRenderSceneProxy.h b/Engine/Plugins/2D/Paper2D/Source/Paper2D/Private/PaperRenderSceneProxy.h
--- a/Engine/Plugins/2D/Paper2D/Source/Paper2D/Private/PaperRenderSceneProxy.h
+++ b/Engine/Plugins/2D/Paper2D/Source/Paper2D/Private/PaperRenderSceneProxy.h
@@ -61,7 +61,7 @@ struct PAPER2D_API FSpriteRenderSection
 
 	class FTexture* GetBaseTextureResource() const
 	{
-		return (BaseTexture != nullptr) ? BaseTexture->Resource : nullptr;
+		return (BaseTexture != nullptr) ? BaseTexture->GetResource() : nullptr;
 	}
 
 	bool IsValid() const
diff --git a/Engine/Plugins/2D/Paper2D/Source/Paper2DEditor/Private/ExtractSprites/SPaperExtractSpritesDialog.cpp b/Engine/Plugins/2D/Paper2D/Source/Paper2DEditor/Private/ExtractSprites/SPaperExtractSpritesDialog.cpp
--- a/Engine/Plugins/2D/Paper2D/Source/Paper2DEditor/Private/ExtractSprites/SPaperExtractSpritesDialog.cpp
+++ b/Engine/Plugins/2D/Paper2D/Source/Paper2DEditor/Private/ExtractSprites/SPaperExtractSpritesDialog.cpp
@@ -81,7 +81,7 @@ void FPaperExtractSpritesViewportClient::Draw(FViewport* InViewport, FCanvas* Ca
 		const float Width = Texture->GetSurfaceWidth() * ZoomAmount;
 		const float Height = Texture->GetSurfaceHeight() * ZoomAmount;
 
-		Canvas->DrawTile(XPos, YPos, Width, Height, 0.0f, 0.0f, 1.0f, 1.0f, TextureDrawColor, Texture->Resource, bUseTranslucentBlend);
+		Canvas->DrawTile(XPos, YPos, Width, Height, 0.0f, 0.0f, 1.0f, 1.0f, TextureDrawColor, Texture->GetResource(), bUseTranslucentBlend);
 
 		for (FPaperExtractedSprite Sprite : ExtractedSprites)
 		{
diff --git a/Engine/Plugins/2D/Paper2D/Source/Paper2DEditor/Private/PaperEditorViewportClient.cpp b/Engine/Plugins/2D/Paper2D/Source/Paper2DEditor/Private/PaperEditorViewportClient.cpp
--- a/Engine/Plugins/2D/Paper2D/Source/Paper2DEditor/Private/PaperEditorViewportClient.cpp
+++ b/Engine/Plugins/2D/Paper2D/Source/Paper2DEditor/Private/PaperEditorViewportClient.cpp
@@ -135,7 +135,7 @@ void FPaperEditorViewportClient::DestroyCheckerboardTexture()
 {
 	if (CheckerboardTexture)
 	{
-		if (CheckerboardTexture->Resource)
+		if (CheckerboardTexture->GetResource())
 		{
 			CheckerboardTexture->ReleaseResource();
 		}
diff --git a/Engine/Plugins/2D/Paper2D/Source/Paper2DEditor/Private/PaperSpriteThumbnailRenderer.cpp b/Engine/Plugins/2D/Paper2D/Source/Paper2DEditor/Private/PaperSpriteThumbnailRenderer.cpp
--- a/Engine/Plugins/2D/Paper2D/Source/Paper2DEditor/Private/PaperSpriteThumbnailRenderer.cpp
+++ b/Engine/Plugins/2D/Paper2D/Source/Paper2DEditor/Private/PaperSpriteThumbnailRenderer.cpp
@@ -42,7 +42,7 @@ void UPaperSpriteThumbnailRenderer::DrawGrid(int32 X, int32 Y, uint32 Width, uin
 		4.0f,
 		4.0f,
 		FLinearColor::White,
-		GridTexture->Resource,
+		GridTexture->GetResource(),
 		bAlphaBlend);
 }
 
@@ -106,7 +106,7 @@ void UPaperSpriteThumbnailRenderer::DrawFrame(class UPaperSprite* Sprite, int32
 		}
 
 		// Draw triangles
-		if ((CanvasPositions.Num() > 0) && (SourceTexture->Resource != nullptr))
+		if ((CanvasPositions.Num() > 0) && (SourceTexture->GetResource() != nullptr))
 		{
 			TArray<FCanvasUVTri> Triangles;
 			const FLinearColor SpriteColor(FLinearColor::White);
@@ -117,7 +117,7 @@ void UPaperSpriteThumbnailRenderer::DrawFrame(class UPaperSprite* Sprite, int32
 				Triangle->V1_Pos = CanvasPositions[Vertex + 1]; Triangle->V1_UV = CanvasUVs[Vertex + 1]; Triangle->V1_Color = SpriteColor;
 				Triangle->V2_Pos = CanvasPositions[Vertex + 2]; Triangle->V2_UV = CanvasUVs[Vertex + 2]; Triangle->V2_Color = SpriteColor;
 			}
-			FCanvasTriangleItem CanvasTriangle(Triangles, SourceTexture->Resource);
+			FCanvasTriangleItem CanvasTriangle(Triangles, SourceTexture->GetResource());
 			CanvasTriangle.BlendMode = bUseTranslucentBlend ? ESimpleElementBlendMode::SE_BLEND_Translucent : ESimpleElementBlendMode::SE_BLEND_Opaque;
 			Canvas->DrawItem(CanvasTriangle);
 		}
diff --git a/Engine/Plugins/2D/Paper2D/Source/Paper2DEditor/Private/PaperTileSetThumbnailRenderer.cpp b/Engine/Plugins/2D/Paper2D/Source/Paper2DEditor/Private/PaperTileSetThumbnailRenderer.cpp
--- a/Engine/Plugins/2D/Paper2D/Source/Paper2DEditor/Private/PaperTileSetThumbnailRenderer.cpp
+++ b/Engine/Plugins/2D/Paper2D/Source/Paper2DEditor/Private/PaperTileSetThumbnailRenderer.cpp
@@ -43,7 +43,7 @@ void UPaperTileSetThumbnailRenderer::Draw(UObject* Object, int32 X, int32 Y, uin
 				4.0f,
 				4.0f,
 				FLinearColor::White,
-				GridTexture->Resource,
+				GridTexture->GetResource(),
 				bAlphaBlend);
 		}
 
@@ -98,7 +98,7 @@ void UPaperTileSetThumbnailRenderer::Draw(UObject* Object, int32 X, int32 Y, uin
 			(TextureWidth - Margin.Right) * InvWidth,
 			(TextureHeight - Margin.Bottom) * InvHeight,
 			FLinearColor::White,
-			TileSheetTexture->Resource,
+			TileSheetTexture->GetResource(),
 			bUseTranslucentBlend);
 
 		// Draw a label overlay
diff --git a/Engine/Plugins/2D/Paper2D/Source/Paper2DEditor/Private/TileSetEditor/TileSetEditorViewportClient.cpp b/Engine/Plugins/2D/Paper2D/Source/Paper2DEditor/Private/TileSetEditor/TileSetEditorViewportClient.cpp
--- a/Engine/Plugins/2D/Paper2D/Source/Paper2DEditor/Private/TileSetEditor/TileSetEditorViewportClient.cpp
+++ b/Engine/Plugins/2D/Paper2D/Source/Paper2DEditor/Private/TileSetEditor/TileSetEditorViewportClient.cpp
@@ -54,7 +54,7 @@ void FTileSetEditorViewportClient::Draw(FViewport* InViewport, FCanvas* Canvas)
 			const float Width = Texture->GetSurfaceWidth() * ZoomAmount;
 			const float Height = Texture->GetSurfaceHeight() * ZoomAmount;
 
-			Canvas->DrawTile(XPos, YPos, Width, Height, 0.0f, 0.0f, 1.0f, 1.0f, TextureDrawColor, Texture->Resource, bUseTranslucentBlend);
+			Canvas->DrawTile(XPos, YPos, Width, Height, 0.0f, 0.0f, 1.0f, 1.0f, TextureDrawColor, Texture->GetResource(), bUseTranslucentBlend);
 		}
 
 
diff --git a/Engine/Plugins/Compositing/OpenCVLensDistortion/Source/OpenCVLensDistortion/Private/LensDistortionDisplacementMapRendering.cpp b/Engine/Plugins/Compositing/OpenCVLensDistortion/Source/OpenCVLensDistortion/Private/LensDistortionDisplacementMapRendering.cpp
--- a/Engine/Plugins/Compositing/OpenCVLensDistortion/Source/OpenCVLensDistortion/Private/LensDistortionDisplacementMapRendering.cpp
+++ b/Engine/Plugins/Compositing/OpenCVLensDistortion/Source/OpenCVLensDistortion/Private/LensDistortionDisplacementMapRendering.cpp
@@ -167,7 +167,7 @@ void FOpenCVLensDistortionParameters::DrawDisplacementMapToRenderTarget(UWorld*
 		return;
 	}
 
-	if(InPreComputedUndistortDisplacementMap == nullptr || InPreComputedUndistortDisplacementMap->Resource == nullptr)
+	if(InPreComputedUndistortDisplacementMap == nullptr || InPreComputedUndistortDisplacementMap->GetResource() == nullptr)
 	{
 		UE_LOG(LogOpenCVLensDistortion, Error, TEXT("Precomputed displacement map is required to generate final displacement maps."));
 		return;
@@ -176,7 +176,7 @@ void FOpenCVLensDistortionParameters::DrawDisplacementMapToRenderTarget(UWorld*
 	//Prepare parameters for render command
 	const FName TextureRenderTargetName = InOutputRenderTarget->GetFName();
 	FTextureRenderTargetResource* TextureRenderTargetResource = InOutputRenderTarget->GameThread_GetRenderTargetResource();
-	const FTextureResource* PreComputedMapResource = InPreComputedUndistortDisplacementMap->Resource;
+	const FTextureResource* PreComputedMapResource = InPreComputedUndistortDisplacementMap->GetResource();
 	ERHIFeatureLevel::Type FeatureLevel = InWorld->Scene->GetFeatureLevel();
 
 	ENQUEUE_RENDER_COMMAND(CaptureCommand)
diff --git a/Engine/Plugins/Compositing/OpenColorIO/Source/OpenColorIO/Private/OpenColorIOColorTransform.cpp b/Engine/Plugins/Compositing/OpenColorIO/Source/OpenColorIO/Private/OpenColorIOColorTransform.cpp
--- a/Engine/Plugins/Compositing/OpenColorIO/Source/OpenColorIO/Private/OpenColorIOColorTransform.cpp
+++ b/Engine/Plugins/Compositing/OpenColorIO/Source/OpenColorIO/Private/OpenColorIOColorTransform.cpp
@@ -319,7 +319,7 @@ bool UOpenColorIOColorTransform::GetShaderAndLUTResouces(ERHIFeatureLevel::Type
 		//Some color transform will only require shader code with no LUT involved.
 		if (Lut3dTexture != nullptr)
 		{
-			OutLUT3dResource = Lut3dTexture->Resource;
+			OutLUT3dResource = Lut3dTexture->GetResource();
 		}
 
 		return true;
diff --git a/Engine/Plugins/Compositing/OpenColorIO/Source/OpenColorIO/Private/OpenColorIORendering.cpp b/Engine/Plugins/Compositing/OpenColorIO/Source/OpenColorIO/Private/OpenColorIORendering.cpp
--- a/Engine/Plugins/Compositing/OpenColorIO/Source/OpenColorIO/Private/OpenColorIORendering.cpp
+++ b/Engine/Plugins/Compositing/OpenColorIO/Source/OpenColorIO/Private/OpenColorIORendering.cpp
@@ -125,8 +125,8 @@ bool FOpenColorIORendering::ApplyColorTransform(UWorld* InWorld, const FOpenColo
 	}
 
 
-	FTextureResource* InputResource = InTexture->Resource;
-	FTextureResource* OutputResource = OutRenderTarget->Resource;
+	FTextureResource* InputResource = InTexture->GetResource();
+	FTextureResource* OutputResource = OutRenderTarget->GetResource();
 	if (InputResource == nullptr)
 	{
 		UE_LOG(LogOpenColorIO, Warning, TEXT("Can't apply color transform - Invalid Input Texture resource"));
diff --git a/Engine/Plugins/Editor/BlueprintMaterialTextureNodes/Source/BlueprintMaterialTextureNodes/Private/BlueprintMaterialTextureNodesBPLibrary.cpp b/Engine/Plugins/Editor/BlueprintMaterialTextureNodes/Source/BlueprintMaterialTextureNodes/Private/BlueprintMaterialTextureNodesBPLibrary.cpp
--- a/Engine/Plugins/Editor/BlueprintMaterialTextureNodes/Source/BlueprintMaterialTextureNodes/Private/BlueprintMaterialTextureNodesBPLibrary.cpp
+++ b/Engine/Plugins/Editor/BlueprintMaterialTextureNodes/Source/BlueprintMaterialTextureNodes/Private/BlueprintMaterialTextureNodesBPLibrary.cpp
@@ -144,7 +144,7 @@ TArray<FLinearColor> UBlueprintMaterialTextureNodesBPLibrary::RenderTarget_Sampl
 		FMessageLog("Blueprint").Warning(LOCTEXT("RenderTargetSampleUV_InvalidRenderTarget", "RenderTargetSampleUVEditoOnly: Render Target must be non-null."));
 		return { FLinearColor(0,0,0,0) };
 	}
-	else if (!InRenderTarget->Resource)
+	else if (!InRenderTarget->GetResource())
 	{
 		FMessageLog("Blueprint").Warning(LOCTEXT("RenderTargetSampleUV_ReleasedRenderTarget", "RenderTargetSampleUVEditoOnly: Render Target has been released."));
 		return { FLinearColor(0,0,0,0) };
@@ -206,7 +206,7 @@ FLinearColor UBlueprintMaterialTextureNodesBPLibrary::RenderTarget_SampleUV_Edit
 		FMessageLog("Blueprint").Warning(LOCTEXT("RenderTargetSampleUV_InvalidRenderTarget", "RenderTargetSampleUVEditoOnly: Render Target must be non-null."));
 		return FLinearColor(0, 0, 0, 0);
 	}
-	else if (!InRenderTarget->Resource)
+	else if (!InRenderTarget->GetResource())
 	{
 		FMessageLog("Blueprint").Warning(LOCTEXT("RenderTargetSampleUV_ReleasedRenderTarget", "RenderTargetSampleUVEditoOnly: Render Target has been released."));
 		return FLinearColor(0, 0, 0, 0);
diff --git a/Engine/Plugins/Editor/GeometryMode/Source/GeometryMode/Private/GeometryEdMode.cpp b/Engine/Plugins/Editor/GeometryMode/Source/GeometryMode/Private/GeometryEdMode.cpp
--- a/Engine/Plugins/Editor/GeometryMode/Source/GeometryMode/Private/GeometryEdMode.cpp
+++ b/Engine/Plugins/Editor/GeometryMode/Source/GeometryMode/Private/GeometryEdMode.cpp
@@ -668,7 +668,7 @@ void FEdModeGeometry::RenderVertex( const FSceneView* View, FPrimitiveDrawInterf
 			Color = GeomVertex->IsSelected() ? FColor(255,128,64) : GeomObject->GetActualBrush()->GetWireColor();
 
 			PDI->SetHitProxy( new HGeomVertexProxy( GeomObject, VertIdx) );
-			PDI->DrawSprite( Location, 4.f * Scale, 4.f * Scale, GEngine->DefaultBSPVertexTexture->Resource, Color, SDPG_Foreground, 0.0, 0.0, 0.0, 0.0 );
+			PDI->DrawSprite( Location, 4.f * Scale, 4.f * Scale, GEngine->DefaultBSPVertexTexture->GetResource(), Color, SDPG_Foreground, 0.0, 0.0, 0.0, 0.0 );
 			PDI->SetHitProxy( NULL );
 		}
 	}
diff --git a/Engine/Plugins/Enterprise/VariantManagerContent/Source/VariantManagerContent/Private/ThumbnailGenerator.cpp b/Engine/Plugins/Enterprise/VariantManagerContent/Source/VariantManagerContent/Private/ThumbnailGenerator.cpp
--- a/Engine/Plugins/Enterprise/VariantManagerContent/Source/VariantManagerContent/Private/ThumbnailGenerator.cpp
+++ b/Engine/Plugins/Enterprise/VariantManagerContent/Source/VariantManagerContent/Private/ThumbnailGenerator.cpp
@@ -165,13 +165,13 @@ UTexture2D* ThumbnailGenerator::GenerateThumbnailFromTexture(UTexture2D* Texture
 	ENQUEUE_RENDER_COMMAND(RetrieveTextureDataForThumbnail)(
 		[CommandData, TargetWidth, TargetHeight](FRHICommandListImmediate& RHICmdList)
 		{
-			if (!CommandData->Texture->Resource)
+			if (!CommandData->Texture->GetResource())
 			{
 				CommandData->Promise.SetValue();
 				return;
 			}
 
-			FTexture2DRHIRef Texture2DRHI = CommandData->Texture->Resource->TextureRHI->GetTexture2D();
+			FTexture2DRHIRef Texture2DRHI = CommandData->Texture->GetResource()->TextureRHI->GetTexture2D();
 
 			CommandData->PackedPixelFormat = Texture2DRHI->GetFormat();
 			CommandData->bCanResize = ThumbnailGeneratorImpl::IsPixelFormatResizeable( CommandData->PackedPixelFormat );
@@ -181,7 +181,7 @@ UTexture2D* ThumbnailGenerator::GenerateThumbnailFromTexture(UTexture2D* Texture
 			// We start at CurrentFirstMip instead of 0 because Texture2DRHI will always match the CurrentFirstMip, so
 			// even if the texture is 256x256, we may be dealing with CurrentFirstMip 2, and so a 64x64 resource.
 			// This shouldn't happen if we wait for streaming (that we do), but just to be safe
-			int32 TargetMipIndex = ((FTexture2DResource*)CommandData->Texture->Resource)->GetCurrentFirstMip();
+			int32 TargetMipIndex = ((FTexture2DResource*)CommandData->Texture->GetResource())->GetCurrentFirstMip();
 			if (!CommandData->bCanResize)
 			{
 				for(int32 MipIndex = 0; MipIndex < CommandData->Texture->PlatformData->Mips.Num(); ++MipIndex)
diff --git a/Engine/Plugins/Experimental/AR/PointCloud/Source/PointCloud/Private/PointCloudSceneProxy.cpp b/Engine/Plugins/Experimental/AR/PointCloud/Source/PointCloud/Private/PointCloudSceneProxy.cpp
--- a/Engine/Plugins/Experimental/AR/PointCloud/Source/PointCloud/Private/PointCloudSceneProxy.cpp
+++ b/Engine/Plugins/Experimental/AR/PointCloud/Source/PointCloud/Private/PointCloudSceneProxy.cpp
@@ -70,13 +70,13 @@ void FNoFetchPointCloudSceneProxy::GetDynamicMeshElements(const TArray<const FSc
 			}
 			FPrimitiveDrawInterface* PDI = Collector.GetPDI(ViewIndex);
 			// Loop through manually drawing the points since PT_PointList isn't really a supported primitive type
-			if (SpriteTexture && SpriteTexture->Resource)
+			if (SpriteTexture && SpriteTexture->GetResource())
 			{
-				const auto SpriteSizeX = Size / (float)SpriteTexture->Resource->GetSizeX();
-				const auto SpriteSizeY = Size / (float)SpriteTexture->Resource->GetSizeY();
+				const auto SpriteSizeX = Size / (float)SpriteTexture->GetResource()->GetSizeX();
+				const auto SpriteSizeY = Size / (float)SpriteTexture->GetResource()->GetSizeY();
 				for (const FVector& Point : Points)
 				{
-					PDI->DrawSprite(Point, SpriteSizeX, SpriteSizeY, SpriteTexture->Resource, Color, SDPG_World, 0.f, 0.f, 0.f, 0.f, SE_BLEND_Masked);
+					PDI->DrawSprite(Point, SpriteSizeX, SpriteSizeY, SpriteTexture->GetResource(), Color, SDPG_World, 0.f, 0.f, 0.f, 0.f, SE_BLEND_Masked);
 				}
 			}
 			else
diff --git a/Engine/Plugins/Experimental/DigitalMultiplex/DMXPixelMapping/Source/DMXPixelMappingEditor/Private/DMXPixelMappingThumbnailRendering.cpp b/Engine/Plugins/Experimental/DigitalMultiplex/DMXPixelMapping/Source/DMXPixelMappingEditor/Private/DMXPixelMappingThumbnailRendering.cpp
--- a/Engine/Plugins/Experimental/DigitalMultiplex/DMXPixelMapping/Source/DMXPixelMappingEditor/Private/DMXPixelMappingThumbnailRendering.cpp
+++ b/Engine/Plugins/Experimental/DigitalMultiplex/DMXPixelMapping/Source/DMXPixelMappingEditor/Private/DMXPixelMappingThumbnailRendering.cpp
@@ -12,10 +12,10 @@ bool UDMXPixelMappingThumbnailRendering::CanVisualizeAsset(UObject* Object)
 void UDMXPixelMappingThumbnailRendering::GetThumbnailSize(UObject* Object, float Zoom, uint32& OutWidth, uint32& OutHeight) const
 {
 	UTexture* ObjectTexture = GetThumbnailTextureFromObject(Object);
-	if (ObjectTexture != nullptr && ObjectTexture->Resource != nullptr)
+	if (ObjectTexture != nullptr && ObjectTexture->GetResource() != nullptr)
 	{
-		OutWidth = Zoom * ObjectTexture->Resource->GetSizeX();
-		OutHeight = Zoom * ObjectTexture->Resource->GetSizeY();
+		OutWidth = Zoom * ObjectTexture->GetResource()->GetSizeX();
+		OutHeight = Zoom * ObjectTexture->GetResource()->GetSizeY();
 	}
 	else
 	{
@@ -29,7 +29,7 @@ void UDMXPixelMappingThumbnailRendering::Draw(UObject* Object, int32 X, int32 Y,
 	UTexture* ObjectTexture = GetThumbnailTextureFromObject(Object);
 	if (ObjectTexture != nullptr)
 	{
-		Canvas->DrawTile(X, Y, Width, Height, 0.0f, 0.0f, 1.0f, 1.0f, FLinearColor::White, ObjectTexture->Resource, false);
+		Canvas->DrawTile(X, Y, Width, Height, 0.0f, 0.0f, 1.0f, 1.0f, FLinearColor::White, ObjectTexture->GetResource(), false);
 	}
 }
 
diff --git a/Engine/Plugins/Experimental/DigitalMultiplex/DMXPixelMapping/Source/DMXPixelMappingEditor/Private/Viewports/DMXPixelMappingPreviewViewportClient.cpp b/Engine/Plugins/Experimental/DigitalMultiplex/DMXPixelMapping/Source/DMXPixelMappingEditor/Private/Viewports/DMXPixelMappingPreviewViewportClient.cpp
--- a/Engine/Plugins/Experimental/DigitalMultiplex/DMXPixelMapping/Source/DMXPixelMappingEditor/Private/Viewports/DMXPixelMappingPreviewViewportClient.cpp
+++ b/Engine/Plugins/Experimental/DigitalMultiplex/DMXPixelMapping/Source/DMXPixelMappingEditor/Private/Viewports/DMXPixelMappingPreviewViewportClient.cpp
@@ -41,7 +41,7 @@ void FDMXPixelMappingPreviewViewportClient::Draw(FViewport* InViewport, FCanvas*
 
 			if (Renderer.IsValid() && InputTexture != nullptr)
 			{
-				Renderer->RenderTextureToRectangle(InputTexture->Resource, InViewport->GetRenderTargetTexture(), InViewport->GetSizeXY(), InputTexture->SRGB);
+				Renderer->RenderTextureToRectangle(InputTexture->GetResource(), InViewport->GetRenderTargetTexture(), InViewport->GetSizeXY(), InputTexture->SRGB);
 			}
 		}
 	}
diff --git a/Engine/Plugins/Experimental/DigitalMultiplex/DMXPixelMapping/Source/DMXPixelMappingEditor/Private/Viewports/DMXPixelMappingSourceTextureViewportClient.cpp b/Engine/Plugins/Experimental/DigitalMultiplex/DMXPixelMapping/Source/DMXPixelMappingEditor/Private/Viewports/DMXPixelMappingSourceTextureViewportClient.cpp
--- a/Engine/Plugins/Experimental/DigitalMultiplex/DMXPixelMapping/Source/DMXPixelMappingEditor/Private/Viewports/DMXPixelMappingSourceTextureViewportClient.cpp
+++ b/Engine/Plugins/Experimental/DigitalMultiplex/DMXPixelMapping/Source/DMXPixelMappingEditor/Private/Viewports/DMXPixelMappingSourceTextureViewportClient.cpp
@@ -33,12 +33,12 @@ void FDMXPixelMappingSourceTextureViewportClient::Draw(FViewport* InViewport, FC
 		InputTexture = RendererComponent->GetRendererInputTexture();
 	}
 
-	if (InputTexture != nullptr && InputTexture->Resource != nullptr)
+	if (InputTexture != nullptr && InputTexture->GetResource() != nullptr)
 	{
 		uint32 Width = SceneViewport->GetSizeXY().X;
 		uint32 Height = SceneViewport->GetSizeXY().Y;
 
-		FCanvasTileItem TileItem(FVector2D(0, 0), InputTexture->Resource, FVector2D(Width, Height), FLinearColor::White);
+		FCanvasTileItem TileItem(FVector2D(0, 0), InputTexture->GetResource(), FVector2D(Width, Height), FLinearColor::White);
 		TileItem.BlendMode = ESimpleElementBlendMode::SE_BLEND_MAX;
 		Canvas->DrawItem(TileItem);
 	}
diff --git a/Engine/Plugins/Experimental/DigitalMultiplex/DMXPixelMapping/Source/DMXPixelMappingEditor/Private/Widgets/SDMXPixelMappingSourceTextureViewport.cpp b/Engine/Plugins/Experimental/DigitalMultiplex/DMXPixelMapping/Source/DMXPixelMappingEditor/Private/Widgets/SDMXPixelMappingSourceTextureViewport.cpp
--- a/Engine/Plugins/Experimental/DigitalMultiplex/DMXPixelMapping/Source/DMXPixelMappingEditor/Private/Widgets/SDMXPixelMappingSourceTextureViewport.cpp
+++ b/Engine/Plugins/Experimental/DigitalMultiplex/DMXPixelMapping/Source/DMXPixelMappingEditor/Private/Widgets/SDMXPixelMappingSourceTextureViewport.cpp
@@ -100,7 +100,7 @@ const FTextureResource* SDMXPixelMappingSourceTextureViewport::GetInputTextureRe
 {
 	if (const UTexture* Texture = GetInputTexture())
 	{
-		return Texture->Resource;
+		return Texture->GetResource();
 	}
 
 	return nullptr;
diff --git a/Engine/Plugins/Experimental/DigitalMultiplex/DMXPixelMapping/Source/DMXPixelMappingRuntime/Private/Components/DMXPixelMappingFixtureGroupItemComponent.cpp b/Engine/Plugins/Experimental/DigitalMultiplex/DMXPixelMapping/Source/DMXPixelMappingRuntime/Private/Components/DMXPixelMappingFixtureGroupItemComponent.cpp
--- a/Engine/Plugins/Experimental/DigitalMultiplex/DMXPixelMapping/Source/DMXPixelMappingRuntime/Private/Components/DMXPixelMappingFixtureGroupItemComponent.cpp
+++ b/Engine/Plugins/Experimental/DigitalMultiplex/DMXPixelMapping/Source/DMXPixelMappingRuntime/Private/Components/DMXPixelMappingFixtureGroupItemComponent.cpp
@@ -456,8 +456,8 @@ void UDMXPixelMappingFixtureGroupItemComponent::QueueDownsample()
 	// Store pixel position
 	DownsamplePixelIndex = RendererComponent->GetDownsamplePixelNum();
 	
-	const uint32 TextureSizeX = InputTexture->Resource->GetSizeX();
-	const uint32 TextureSizeY = InputTexture->Resource->GetSizeY();
+	const uint32 TextureSizeX = InputTexture->GetResource()->GetSizeX();
+	const uint32 TextureSizeY = InputTexture->GetResource()->GetSizeY();
 	check(TextureSizeX > 0 && TextureSizeY > 0);
 	const FIntPoint PixelPosition = RendererComponent->GetPixelPosition(DownsamplePixelIndex);
 	const FVector2D UV = FVector2D(PositionX / TextureSizeX, PositionY / TextureSizeY);
diff --git a/Engine/Plugins/Experimental/DigitalMultiplex/DMXPixelMapping/Source/DMXPixelMappingRuntime/Private/Components/DMXPixelMappingMatrixCellComponent.cpp b/Engine/Plugins/Experimental/DigitalMultiplex/DMXPixelMapping/Source/DMXPixelMappingRuntime/Private/Components/DMXPixelMappingMatrixCellComponent.cpp
--- a/Engine/Plugins/Experimental/DigitalMultiplex/DMXPixelMapping/Source/DMXPixelMappingRuntime/Private/Components/DMXPixelMappingMatrixCellComponent.cpp
+++ b/Engine/Plugins/Experimental/DigitalMultiplex/DMXPixelMapping/Source/DMXPixelMappingRuntime/Private/Components/DMXPixelMappingMatrixCellComponent.cpp
@@ -326,8 +326,8 @@ void UDMXPixelMappingMatrixCellComponent::QueueDownsample()
 	// Store downsample index
 	DownsamplePixelIndex = RendererComponent->GetDownsamplePixelNum();
 
-	const uint32 TextureSizeX = InputTexture->Resource->GetSizeX();
-	const uint32 TextureSizeY = InputTexture->Resource->GetSizeY();
+	const uint32 TextureSizeX = InputTexture->GetResource()->GetSizeX();
+	const uint32 TextureSizeY = InputTexture->GetResource()->GetSizeY();
 	check(TextureSizeX > 0 && TextureSizeY > 0);
 	const FIntPoint PixelPosition = RendererComponent->GetPixelPosition(DownsamplePixelIndex);
 	const FVector2D UV = FVector2D(PositionX / TextureSizeX, PositionY / TextureSizeY);
diff --git a/Engine/Plugins/Experimental/DigitalMultiplex/DMXPixelMapping/Source/DMXPixelMappingRuntime/Private/Components/DMXPixelMappingRendererComponent.cpp b/Engine/Plugins/Experimental/DigitalMultiplex/DMXPixelMapping/Source/DMXPixelMappingRuntime/Private/Components/DMXPixelMappingRendererComponent.cpp
--- a/Engine/Plugins/Experimental/DigitalMultiplex/DMXPixelMapping/Source/DMXPixelMappingRuntime/Private/Components/DMXPixelMappingRendererComponent.cpp
+++ b/Engine/Plugins/Experimental/DigitalMultiplex/DMXPixelMapping/Source/DMXPixelMappingRuntime/Private/Components/DMXPixelMappingRendererComponent.cpp
@@ -25,7 +25,7 @@
 #include "LevelEditor.h"
 #endif
 
-
+class FTextureResource;
 const FIntPoint UDMXPixelMappingRendererComponent::MaxDownsampleBufferTargetSize = FIntPoint(4096);
 const FLinearColor UDMXPixelMappingRendererComponent::ClearTextureColor = FLinearColor::Black;
 
@@ -174,7 +174,7 @@ void UDMXPixelMappingRendererComponent::RenderEditorPreviewTexture()
 		}
 	}, true);
 
-	Renderer->RenderPreview(GetPreviewRenderTarget()->Resource, DownsampleBufferTarget->Resource, MoveTemp(PixelPreviewParams));
+	Renderer->RenderPreview(GetPreviewRenderTarget()->GetResource(), DownsampleBufferTarget->GetResource(), MoveTemp(PixelPreviewParams));
 }
 
 UTextureRenderTarget2D* UDMXPixelMappingRendererComponent::GetPreviewRenderTarget()
@@ -192,7 +192,7 @@ FVector2D UDMXPixelMappingRendererComponent::GetSize() const
 	// Get a size from Input Texture
 	if (const UTexture* const RendererInputTexture = GetRendererInputTexture())
 	{
-		if (class FTextureResource* Resource = RendererInputTexture->Resource)
+		if (const FTextureResource* Resource = RendererInputTexture->GetResource())
 		{
 			return FVector2D(Resource->GetSizeX(), Resource->GetSizeY());
 		}
@@ -303,8 +303,8 @@ void UDMXPixelMappingRendererComponent::Render()
 
 	// 7. Downsample all pixels
 	GetRenderer()->DownsampleRender(
-		DownsampleInputTexture->Resource,
-		DownsampleBufferTarget->Resource,
+		DownsampleInputTexture->GetResource(),
+		DownsampleBufferTarget->GetResource(),
 		DownsampleBufferTarget->GameThread_GetRenderTargetResource(),
 		MoveTemp(DownsamplePixelParams), // Move Set to GPU thread, no empty function call needed
 		[this](TArray<FColor>&& InDownsampleBuffer, FIntRect InRect) { SetDownsampleBuffer(MoveTemp(InDownsampleBuffer), InRect); }
@@ -415,9 +415,9 @@ void UDMXPixelMappingRendererComponent::RendererInputTexture()
 	{
 	case(EDMXPixelMappingRendererType::Texture) :
 		// Nothing
-		if (InputTexture != nullptr &&  InputTexture->Resource != nullptr)
+		if (InputTexture != nullptr &&  InputTexture->GetResource() != nullptr)
 		{
-			ResizePreviewRenderTarget(InputTexture->Resource->GetSizeX(), InputTexture->Resource->GetSizeY());
+			ResizePreviewRenderTarget(InputTexture->GetResource()->GetSizeX(), InputTexture->GetResource()->GetSizeY());
 		}
 		break;
 	case(EDMXPixelMappingRendererType::Material):
diff --git a/Engine/Plugins/Experimental/DigitalMultiplex/DMXPixelMapping/Source/DMXPixelMappingRuntime/Private/Components/DMXPixelMappingScreenComponent.cpp b/Engine/Plugins/Experimental/DigitalMultiplex/DMXPixelMapping/Source/DMXPixelMappingRuntime/Private/Components/DMXPixelMappingScreenComponent.cpp
--- a/Engine/Plugins/Experimental/DigitalMultiplex/DMXPixelMapping/Source/DMXPixelMappingRuntime/Private/Components/DMXPixelMappingScreenComponent.cpp
+++ b/Engine/Plugins/Experimental/DigitalMultiplex/DMXPixelMapping/Source/DMXPixelMappingRuntime/Private/Components/DMXPixelMappingScreenComponent.cpp
@@ -493,8 +493,8 @@ void UDMXPixelMappingScreenComponent::QueueDownsample()
 		return;
 	}
 
-	const uint32 TextureSizeX = InputTexture->Resource->GetSizeX();
-	const uint32 TextureSizeY = InputTexture->Resource->GetSizeY();
+	const uint32 TextureSizeX = InputTexture->GetResource()->GetSizeX();
+	const uint32 TextureSizeY = InputTexture->GetResource()->GetSizeY();
 	check(TextureSizeX > 0 && TextureSizeY > 0);
 	
 	constexpr bool bStaticCalculateUV = true;
diff --git a/Engine/Plugins/Experimental/ImagePlate/Source/ImagePlate/Private/ImagePlateFileSequence.cpp b/Engine/Plugins/Experimental/ImagePlate/Source/ImagePlate/Private/ImagePlateFileSequence.cpp
--- a/Engine/Plugins/Experimental/ImagePlate/Source/ImagePlate/Private/ImagePlateFileSequence.cpp
+++ b/Engine/Plugins/Experimental/ImagePlate/Source/ImagePlate/Private/ImagePlateFileSequence.cpp
@@ -136,16 +136,16 @@ bool FImagePlateSourceFrame::EnsureTextureMetrics(UTexture* DestinationTexture)
 			}
 		}
 	}
-	else if (DestinationTexture->Resource)
+	else if (DestinationTexture->GetResource())
 	{
 		// We have to have a valid texture to check whether it's the right type or not
-		if (!DestinationTexture->Resource->TextureRHI)
+		if (!DestinationTexture->GetResource()->TextureRHI)
 		{
 			DestinationTexture->UpdateResource();
 			FlushRenderingCommands();
 		}
 
-		FTexture2DRHIRef Texture2DRHI = DestinationTexture->Resource->TextureRHI ? DestinationTexture->Resource->TextureRHI->GetTexture2D() : nullptr;
+		FTexture2DRHIRef Texture2DRHI = DestinationTexture->GetResource()->TextureRHI ? DestinationTexture->GetResource()->TextureRHI->GetTexture2D() : nullptr;
 		if (!Texture2DRHI)
 		{
 			UE_LOG(LogImagePlateFileSequence, Warning, TEXT("Unsupported texture type encountered: Unable to update texture to fit source frame size or bitdepth."));
@@ -195,12 +195,12 @@ TFuture<void> FImagePlateSourceFrame::CopyTo(UTexture* DestinationTexture)
 	ENQUEUE_RENDER_COMMAND(CopySourceBufferToTexture)(
 		[CommandData](FRHICommandListImmediate& RHICmdList)
 		{
-			if (!CommandData->DestinationTexture->Resource || !CommandData->DestinationTexture->Resource->TextureRHI)
+			if (!CommandData->DestinationTexture->GetResource() || !CommandData->DestinationTexture->GetResource()->TextureRHI)
 			{
 				return;
 			}
 
-			FTexture2DRHIRef Texture2DRHI = CommandData->DestinationTexture->Resource->TextureRHI->GetTexture2D();
+			FTexture2DRHIRef Texture2DRHI = CommandData->DestinationTexture->GetResource()->TextureRHI->GetTexture2D();
 			if (!Texture2DRHI)
 			{
 				return;
diff --git a/Engine/Plugins/Experimental/VirtualHeightfieldMesh/Source/VirtualHeightfieldMesh/Private/VirtualHeightfieldMeshSceneProxy.cpp b/Engine/Plugins/Experimental/VirtualHeightfieldMesh/Source/VirtualHeightfieldMesh/Private/VirtualHeightfieldMeshSceneProxy.cpp
--- a/Engine/Plugins/Experimental/VirtualHeightfieldMesh/Source/VirtualHeightfieldMesh/Private/VirtualHeightfieldMeshSceneProxy.cpp
+++ b/Engine/Plugins/Experimental/VirtualHeightfieldMesh/Source/VirtualHeightfieldMesh/Private/VirtualHeightfieldMeshSceneProxy.cpp
@@ -446,7 +446,7 @@ void FVirtualHeightfieldMeshSceneProxy::CreateRenderThreadResources()
 				UniformParams.PageTableTexture = AllocatedVirtualTexture->GetPageTableTexture(0);
 				UniformParams.HeightTexture = AllocatedVirtualTexture->GetPhysicalTextureSRV(0, false);
 				UniformParams.HeightSampler = TStaticSamplerState<SF_Bilinear>::GetRHI();
-				UniformParams.LodBiasTexture = LodBiasTexture ? LodBiasTexture->Resource->TextureRHI : GBlackTexture->TextureRHI;
+				UniformParams.LodBiasTexture = LodBiasTexture ? LodBiasTexture->GetResource()->TextureRHI : GBlackTexture->TextureRHI;
 				UniformParams.LodBiasSampler = TStaticSamplerState<SF_Point>::GetRHI();
 				UniformParams.NumQuadsPerTileSide = NumQuadsPerTileSide;
 
@@ -1218,8 +1218,8 @@ void FVirtualHeightfieldMeshRendererExtension::SubmitWork(FRDGBuilder& GraphBuil
 
 		VirtualHeightfieldMesh::FProxyDesc ProxyDesc;
 		ProxyDesc.PageTableTexture = AllocatedVirtualTexture->GetPageTableTexture(0);
-		ProxyDesc.HeightMinMaxTexture = Proxy->HeightMinMaxTexture ? Proxy->HeightMinMaxTexture->Resource->TextureRHI : VirtualHeightfieldMesh::GHeightMinMaxDefaultTexture->TextureRHI;
-		ProxyDesc.LodBiasMinMaxTexture = Proxy->LodBiasMinMaxTexture ? Proxy->LodBiasMinMaxTexture->Resource->TextureRHI : GBlackTexture->TextureRHI;
+		ProxyDesc.HeightMinMaxTexture = Proxy->HeightMinMaxTexture ? Proxy->HeightMinMaxTexture->GetResource()->TextureRHI : VirtualHeightfieldMesh::GHeightMinMaxDefaultTexture->TextureRHI;
+		ProxyDesc.LodBiasMinMaxTexture = Proxy->LodBiasMinMaxTexture ? Proxy->LodBiasMinMaxTexture->GetResource()->TextureRHI : GBlackTexture->TextureRHI;
 		ProxyDesc.MinMaxLevelOffset = ProxyDesc.HeightMinMaxTexture->GetNumMips() - 1 - AllocatedVirtualTexture->GetMaxLevel();
 		ProxyDesc.MaxLevel = AllocatedVirtualTexture->GetMaxLevel();
 		ProxyDesc.NumForceLoadLods = FMath::Min<uint32>(FMath::Max(Proxy->NumForceLoadLods, 0), ProxyDesc.MaxLevel + 1);
diff --git a/Engine/Plugins/FX/Niagara/Source/Niagara/Private/NiagaraDataInterface2DArrayTexture.cpp b/Engine/Plugins/FX/Niagara/Source/Niagara/Private/NiagaraDataInterface2DArrayTexture.cpp
--- a/Engine/Plugins/FX/Niagara/Source/Niagara/Private/NiagaraDataInterface2DArrayTexture.cpp
+++ b/Engine/Plugins/FX/Niagara/Source/Niagara/Private/NiagaraDataInterface2DArrayTexture.cpp
@@ -245,7 +245,7 @@ void UNiagaraDataInterface2DArrayTexture::PushToRenderThreadImpl()
 
 	ENQUEUE_RENDER_COMMAND(FPushDITextureToRT)
 	(
-		[RT_Proxy, RT_Resource=Texture ? Texture->Resource : nullptr, RT_TexDims](FRHICommandListImmediate& RHICmdList)
+		[RT_Proxy, RT_Resource=Texture ? Texture->GetResource() : nullptr, RT_TexDims](FRHICommandListImmediate& RHICmdList)
 		{
 			RT_Proxy->TextureRHI = RT_Resource ? RT_Resource->TextureRHI : nullptr;
 			RT_Proxy->SamplerStateRHI = RT_Resource ? RT_Resource->SamplerStateRHI : nullptr;
diff --git a/Engine/Plugins/FX/Niagara/Source/Niagara/Private/NiagaraDataInterfaceCubeTexture.cpp b/Engine/Plugins/FX/Niagara/Source/Niagara/Private/NiagaraDataInterfaceCubeTexture.cpp
--- a/Engine/Plugins/FX/Niagara/Source/Niagara/Private/NiagaraDataInterfaceCubeTexture.cpp
+++ b/Engine/Plugins/FX/Niagara/Source/Niagara/Private/NiagaraDataInterfaceCubeTexture.cpp
@@ -283,7 +283,7 @@ void UNiagaraDataInterfaceCubeTexture::PushToRenderThreadImpl()
 
 	ENQUEUE_RENDER_COMMAND(FPushDITextureToRT)
 	(
-		[RT_Proxy, RT_Resource=Texture ? Texture->Resource : nullptr, RT_TexDims](FRHICommandListImmediate& RHICmdList)
+		[RT_Proxy, RT_Resource=Texture ? Texture->GetResource() : nullptr, RT_TexDims](FRHICommandListImmediate& RHICmdList)
 		{
 			RT_Proxy->TextureRHI = RT_Resource ? RT_Resource->TextureRHI : nullptr;
 			RT_Proxy->SamplerStateRHI = RT_Resource ? RT_Resource->SamplerStateRHI : nullptr;
diff --git a/Engine/Plugins/FX/Niagara/Source/Niagara/Private/NiagaraDataInterfaceGrid2DCollection.cpp b/Engine/Plugins/FX/Niagara/Source/Niagara/Private/NiagaraDataInterfaceGrid2DCollection.cpp
--- a/Engine/Plugins/FX/Niagara/Source/Niagara/Private/NiagaraDataInterfaceGrid2DCollection.cpp
+++ b/Engine/Plugins/FX/Niagara/Source/Niagara/Private/NiagaraDataInterfaceGrid2DCollection.cpp
@@ -1831,7 +1831,7 @@ bool UNiagaraDataInterfaceGrid2DCollection::InitPerInstanceData(void* PerInstanc
 	// Push Updates to Proxy.
 	FNiagaraDataInterfaceProxyGrid2DCollectionProxy* RT_Proxy = GetProxyAs<FNiagaraDataInterfaceProxyGrid2DCollectionProxy>();
 	ENQUEUE_RENDER_COMMAND(FUpdateData)(
-		[RT_Resource=InstanceData->TargetTexture ? InstanceData->TargetTexture->Resource : nullptr, RT_Proxy, InstanceID=SystemInstance->GetId(), RT_InstanceData=*InstanceData](FRHICommandListImmediate& RHICmdList)
+		[RT_Resource=InstanceData->TargetTexture ? InstanceData->TargetTexture->GetResource() : nullptr, RT_Proxy, InstanceID=SystemInstance->GetId(), RT_InstanceData=*InstanceData](FRHICommandListImmediate& RHICmdList)
 	{
 		check(!RT_Proxy->SystemInstancesToProxyData_RT.Contains(InstanceID));
 		FGrid2DCollectionRWInstanceData_RenderThread* TargetData = &RT_Proxy->SystemInstancesToProxyData_RT.Add(InstanceID);
@@ -1903,7 +1903,7 @@ bool UNiagaraDataInterfaceGrid2DCollection::PerInstanceTick(void* PerInstanceDat
 
 	FNiagaraDataInterfaceProxyGrid2DCollectionProxy* RT_Proxy = GetProxyAs<FNiagaraDataInterfaceProxyGrid2DCollectionProxy>();
 	ENQUEUE_RENDER_COMMAND(FUpdateData)(
-		[RT_Resource=InstanceData->TargetTexture ? InstanceData->TargetTexture->Resource : nullptr, RT_Proxy, InstanceID=SystemInstance->GetId()](FRHICommandListImmediate& RHICmdList)
+		[RT_Resource=InstanceData->TargetTexture ? InstanceData->TargetTexture->GetResource() : nullptr, RT_Proxy, InstanceID=SystemInstance->GetId()](FRHICommandListImmediate& RHICmdList)
 	{
 		FGrid2DCollectionRWInstanceData_RenderThread* TargetData = RT_Proxy->SystemInstancesToProxyData_RT.Find(InstanceID);
 		if (RT_Resource && RT_Resource->TextureRHI.IsValid())
@@ -2161,7 +2161,7 @@ bool UNiagaraDataInterfaceGrid2DCollection::FillTexture2D(const UNiagaraComponen
 
 	FNiagaraDataInterfaceProxyGrid2DCollectionProxy* RT_Proxy = GetProxyAs<FNiagaraDataInterfaceProxyGrid2DCollectionProxy>();
 	ENQUEUE_RENDER_COMMAND(FUpdateDIColorCurve)(
-		[RT_Proxy, InstanceID=SystemInstanceController->GetSystemInstanceID(), RT_TextureResource=Dest->Resource, AttributeIndex](FRHICommandListImmediate& RHICmdList)
+		[RT_Proxy, InstanceID=SystemInstanceController->GetSystemInstanceID(), RT_TextureResource=Dest->GetResource(), AttributeIndex](FRHICommandListImmediate& RHICmdList)
 	{
 		FGrid2DCollectionRWInstanceData_RenderThread* Grid2DInstanceData = RT_Proxy->SystemInstancesToProxyData_RT.Find(InstanceID);
 
@@ -2217,7 +2217,7 @@ bool UNiagaraDataInterfaceGrid2DCollection::FillRawTexture2D(const UNiagaraCompo
 
 	FNiagaraDataInterfaceProxyGrid2DCollectionProxy* RT_Proxy = GetProxyAs<FNiagaraDataInterfaceProxyGrid2DCollectionProxy>();
 	ENQUEUE_RENDER_COMMAND(FUpdateDIColorCurve)(
-		[RT_Proxy, RT_InstanceID=SystemID, RT_TextureResource=Dest->Resource](FRHICommandListImmediate& RHICmdList)
+		[RT_Proxy, RT_InstanceID=SystemID, RT_TextureResource=Dest->GetResource()](FRHICommandListImmediate& RHICmdList)
 	{
 		FGrid2DCollectionRWInstanceData_RenderThread* RT_Grid2DInstanceData = RT_Proxy->SystemInstancesToProxyData_RT.Find(RT_InstanceID);
 		if (RT_TextureResource && RT_TextureResource->TextureRHI.IsValid() && RT_Grid2DInstanceData && RT_Grid2DInstanceData->CurrentData)
@@ -2384,7 +2384,7 @@ bool UNiagaraDataInterfaceGrid2DCollection::PerInstanceTickPostSimulate(void* Pe
 		// Push Updates to Proxy.
 		FNiagaraDataInterfaceProxyGrid2DCollectionProxy* RT_Proxy = GetProxyAs<FNiagaraDataInterfaceProxyGrid2DCollectionProxy>();
 		ENQUEUE_RENDER_COMMAND(FUpdateData)(
-			[RT_Resource=InstanceData->TargetTexture ? InstanceData->TargetTexture->Resource : nullptr, RT_Proxy, InstanceID=SystemInstance->GetId(), RT_InstanceData=*InstanceData](FRHICommandListImmediate& RHICmdList)
+			[RT_Resource=InstanceData->TargetTexture ? InstanceData->TargetTexture->GetResource() : nullptr, RT_Proxy, InstanceID=SystemInstance->GetId(), RT_InstanceData=*InstanceData](FRHICommandListImmediate& RHICmdList)
 		{
 			check(RT_Proxy->SystemInstancesToProxyData_RT.Contains(InstanceID));
 			FGrid2DCollectionRWInstanceData_RenderThread* TargetData = RT_Proxy->SystemInstancesToProxyData_RT.Find(InstanceID);
diff --git a/Engine/Plugins/FX/Niagara/Source/Niagara/Private/NiagaraDataInterfaceGrid3DCollection.cpp b/Engine/Plugins/FX/Niagara/Source/Niagara/Private/NiagaraDataInterfaceGrid3DCollection.cpp
--- a/Engine/Plugins/FX/Niagara/Source/Niagara/Private/NiagaraDataInterfaceGrid3DCollection.cpp
+++ b/Engine/Plugins/FX/Niagara/Source/Niagara/Private/NiagaraDataInterfaceGrid3DCollection.cpp
@@ -1886,7 +1886,7 @@ bool UNiagaraDataInterfaceGrid3DCollection::InitPerInstanceData(void* PerInstanc
 	// Push Updates to Proxy.
 	FNiagaraDataInterfaceProxyGrid3DCollectionProxy* RT_Proxy = GetProxyAs<FNiagaraDataInterfaceProxyGrid3DCollectionProxy>();
 	ENQUEUE_RENDER_COMMAND(FUpdateData)(
-		[RT_Proxy, RT_Resource=InstanceData->TargetTexture ? InstanceData->TargetTexture->Resource : nullptr, InstanceID = SystemInstance->GetId(), RT_InstanceData=*InstanceData](FRHICommandListImmediate& RHICmdList)
+		[RT_Proxy, RT_Resource=InstanceData->TargetTexture ? InstanceData->TargetTexture->GetResource() : nullptr, InstanceID = SystemInstance->GetId(), RT_InstanceData=*InstanceData](FRHICommandListImmediate& RHICmdList)
 	{
 		check(!RT_Proxy->SystemInstancesToProxyData_RT.Contains(InstanceID));
 		FGrid3DCollectionRWInstanceData_RenderThread* TargetData = &RT_Proxy->SystemInstancesToProxyData_RT.Add(InstanceID);
@@ -1955,7 +1955,7 @@ bool UNiagaraDataInterfaceGrid3DCollection::PerInstanceTick(void* PerInstanceDat
 
 	FNiagaraDataInterfaceProxyGrid3DCollectionProxy* RT_Proxy = GetProxyAs<FNiagaraDataInterfaceProxyGrid3DCollectionProxy>();
 	ENQUEUE_RENDER_COMMAND(FUpdateData)(
-		[RT_Resource= InstanceData->TargetTexture ? InstanceData->TargetTexture->Resource : nullptr, RT_Proxy, InstanceID = SystemInstance->GetId()](FRHICommandListImmediate& RHICmdList)
+		[RT_Resource= InstanceData->TargetTexture ? InstanceData->TargetTexture->GetResource() : nullptr, RT_Proxy, InstanceID = SystemInstance->GetId()](FRHICommandListImmediate& RHICmdList)
 	{
 		FGrid3DCollectionRWInstanceData_RenderThread* TargetData = RT_Proxy->SystemInstancesToProxyData_RT.Find(InstanceID);
 
diff --git a/Engine/Plugins/FX/Niagara/Source/Niagara/Private/NiagaraDataInterfaceRenderTarget2D.cpp b/Engine/Plugins/FX/Niagara/Source/Niagara/Private/NiagaraDataInterfaceRenderTarget2D.cpp
--- a/Engine/Plugins/FX/Niagara/Source/Niagara/Private/NiagaraDataInterfaceRenderTarget2D.cpp
+++ b/Engine/Plugins/FX/Niagara/Source/Niagara/Private/NiagaraDataInterfaceRenderTarget2D.cpp
@@ -620,7 +620,7 @@ bool UNiagaraDataInterfaceRenderTarget2D::RenderVariableToCanvas(FNiagaraSystemI
 		return false;
 	}
 
-	if ( !GT_InstanceData->TargetTexture || !GT_InstanceData->TargetTexture->Resource )
+	if ( !GT_InstanceData->TargetTexture || !GT_InstanceData->TargetTexture->GetResource())
 	{
 		return false;
 	}
@@ -629,7 +629,7 @@ bool UNiagaraDataInterfaceRenderTarget2D::RenderVariableToCanvas(FNiagaraSystemI
 		DrawRect.Min.X, DrawRect.Min.Y, DrawRect.Width(), DrawRect.Height(),
 		0.0f, 0.0f, 1.0f, 1.0f,
 		FLinearColor::White,
-		GT_InstanceData->TargetTexture->Resource,
+		GT_InstanceData->TargetTexture->GetResource(),
 		false
 	);
 
@@ -747,7 +747,7 @@ bool UNiagaraDataInterfaceRenderTarget2D::PerInstanceTickPostSimulate(void* PerI
 			(InstanceData->TargetTexture->RenderTargetFormat != InstanceData->Format) ||
 			!InstanceData->TargetTexture->bCanCreateUAV ||
 			(InstanceData->TargetTexture->bAutoGenerateMips != bAutoGenerateMips) ||
-			!InstanceData->TargetTexture->Resource )
+			!InstanceData->TargetTexture->GetResource())
 		{
 			// resize RT to match what we need for the output
 			InstanceData->TargetTexture->bCanCreateUAV = true;
diff --git a/Engine/Plugins/FX/Niagara/Source/Niagara/Private/NiagaraDataInterfaceRenderTarget2DArray.cpp b/Engine/Plugins/FX/Niagara/Source/Niagara/Private/NiagaraDataInterfaceRenderTarget2DArray.cpp
--- a/Engine/Plugins/FX/Niagara/Source/Niagara/Private/NiagaraDataInterfaceRenderTarget2DArray.cpp
+++ b/Engine/Plugins/FX/Niagara/Source/Niagara/Private/NiagaraDataInterfaceRenderTarget2DArray.cpp
@@ -690,7 +690,7 @@ bool UNiagaraDataInterfaceRenderTarget2DArray::PerInstanceTickPostSimulate(void*
 			(InstanceData->TargetTexture->OverrideFormat != InstanceData->Format) ||
 			!InstanceData->TargetTexture->bCanCreateUAV ||
 			//(InstanceData->TargetTexture->bAutoGenerateMips != bAutoGenerateMips) ||
-			!InstanceData->TargetTexture->Resource )
+			!InstanceData->TargetTexture->GetResource())
 		{
 			// resize RT to match what we need for the output
 			InstanceData->TargetTexture->bCanCreateUAV = true;
diff --git a/Engine/Plugins/FX/Niagara/Source/Niagara/Private/NiagaraDataInterfaceRenderTargetCube.cpp b/Engine/Plugins/FX/Niagara/Source/Niagara/Private/NiagaraDataInterfaceRenderTargetCube.cpp
--- a/Engine/Plugins/FX/Niagara/Source/Niagara/Private/NiagaraDataInterfaceRenderTargetCube.cpp
+++ b/Engine/Plugins/FX/Niagara/Source/Niagara/Private/NiagaraDataInterfaceRenderTargetCube.cpp
@@ -666,7 +666,7 @@ bool UNiagaraDataInterfaceRenderTargetCube::PerInstanceTickPostSimulate(void* Pe
 			(InstanceData->TargetTexture->OverrideFormat != InstanceData->Format) ||
 			!InstanceData->TargetTexture->bCanCreateUAV ||
 			//(InstanceData->TargetTexture->bAutoGenerateMips != bAutoGenerateMips) ||
-			!InstanceData->TargetTexture->Resource )
+			!InstanceData->TargetTexture->GetResource())
 		{
 			// resize RT to match what we need for the output
 			InstanceData->TargetTexture->bCanCreateUAV = true;
diff --git a/Engine/Plugins/FX/Niagara/Source/Niagara/Private/NiagaraDataInterfaceRenderTargetVolume.cpp b/Engine/Plugins/FX/Niagara/Source/Niagara/Private/NiagaraDataInterfaceRenderTargetVolume.cpp
--- a/Engine/Plugins/FX/Niagara/Source/Niagara/Private/NiagaraDataInterfaceRenderTargetVolume.cpp
+++ b/Engine/Plugins/FX/Niagara/Source/Niagara/Private/NiagaraDataInterfaceRenderTargetVolume.cpp
@@ -688,7 +688,7 @@ bool UNiagaraDataInterfaceRenderTargetVolume::PerInstanceTickPostSimulate(void*
 			(InstanceData->TargetTexture->OverrideFormat != InstanceData->Format) ||
 			!InstanceData->TargetTexture->bCanCreateUAV ||
 			//(InstanceData->TargetTexture->bAutoGenerateMips != bAutoGenerateMips) ||
-			!InstanceData->TargetTexture->Resource )
+			!InstanceData->TargetTexture->GetResource())
 		{
 			// resize RT to match what we need for the output
 			InstanceData->TargetTexture->bCanCreateUAV = true;
diff --git a/Engine/Plugins/FX/Niagara/Source/Niagara/Private/NiagaraDataInterfaceTexture.cpp b/Engine/Plugins/FX/Niagara/Source/Niagara/Private/NiagaraDataInterfaceTexture.cpp
--- a/Engine/Plugins/FX/Niagara/Source/Niagara/Private/NiagaraDataInterfaceTexture.cpp
+++ b/Engine/Plugins/FX/Niagara/Source/Niagara/Private/NiagaraDataInterfaceTexture.cpp
@@ -458,7 +458,7 @@ void UNiagaraDataInterfaceTexture::PushToRenderThreadImpl()
 			if (RT_Texture)
 			{
 				RT_Proxy->TextureReferenceRHI = RT_Texture->TextureReference.TextureReferenceRHI;
-				RT_Proxy->SamplerStateRHI = RT_Texture->Resource ? RT_Texture->Resource->SamplerStateRHI : nullptr;
+				RT_Proxy->SamplerStateRHI = RT_Texture->GetResource() ? RT_Texture->GetResource()->SamplerStateRHI : nullptr;
 			}
 			else
 			{
diff --git a/Engine/Plugins/FX/Niagara/Source/Niagara/Private/NiagaraDataInterfaceVolumeTexture.cpp b/Engine/Plugins/FX/Niagara/Source/Niagara/Private/NiagaraDataInterfaceVolumeTexture.cpp
--- a/Engine/Plugins/FX/Niagara/Source/Niagara/Private/NiagaraDataInterfaceVolumeTexture.cpp
+++ b/Engine/Plugins/FX/Niagara/Source/Niagara/Private/NiagaraDataInterfaceVolumeTexture.cpp
@@ -291,7 +291,7 @@ void UNiagaraDataInterfaceVolumeTexture::PushToRenderThreadImpl()
 
 	ENQUEUE_RENDER_COMMAND(FPushDITextureToRT)
 	(
-		[RT_Proxy, RT_Resource=Texture ? Texture->Resource : nullptr, RT_TexDims](FRHICommandListImmediate& RHICmdList)
+		[RT_Proxy, RT_Resource=Texture ? Texture->GetResource() : nullptr, RT_TexDims](FRHICommandListImmediate& RHICmdList)
 		{
 			RT_Proxy->TextureRHI = RT_Resource ? RT_Resource->TextureRHI : nullptr;
 			RT_Proxy->SamplerStateRHI = RT_Resource ? RT_Resource->SamplerStateRHI : nullptr;
diff --git a/Engine/Plugins/FX/Niagara/Source/Niagara/Private/NiagaraRenderer.cpp b/Engine/Plugins/FX/Niagara/Source/Niagara/Private/NiagaraRenderer.cpp
--- a/Engine/Plugins/FX/Niagara/Source/Niagara/Private/NiagaraRenderer.cpp
+++ b/Engine/Plugins/FX/Niagara/Source/Niagara/Private/NiagaraRenderer.cpp
@@ -496,7 +496,7 @@ void FNiagaraRenderer::ProcessMaterialParameterBindings(TConstArrayView< FNiagar
 							if (Var)
 							{
 								UTexture* Tex = Cast<UTexture>(Var);
-								if (Tex && Tex->Resource != nullptr)
+								if (Tex && Tex->GetResource() != nullptr)
 								{
 									MatDyn->SetTextureParameterValue(Binding.MaterialParameterName, Tex);
 								}
diff --git a/Engine/Plugins/FX/Niagara/Source/NiagaraEditor/Private/NiagaraThumbnailRenderer.cpp b/Engine/Plugins/FX/Niagara/Source/NiagaraEditor/Private/NiagaraThumbnailRenderer.cpp
--- a/Engine/Plugins/FX/Niagara/Source/NiagaraEditor/Private/NiagaraThumbnailRenderer.cpp
+++ b/Engine/Plugins/FX/Niagara/Source/NiagaraEditor/Private/NiagaraThumbnailRenderer.cpp
@@ -30,7 +30,7 @@ void UNiagaraThumbnailRendererBase::Draw(UObject* Object, int32 X, int32 Y, uint
 	UTexture2D* ObjectTexture = GetThumbnailTextureFromObject(Object);
 	if (ObjectTexture != nullptr)
 	{
-		Canvas->DrawTile(X, Y, Width, Height, 0.0f, 0.0f, 1.0f, 1.0f, FLinearColor::White, ObjectTexture->Resource, false);
+		Canvas->DrawTile(X, Y, Width, Height, 0.0f, 0.0f, 1.0f, 1.0f, FLinearColor::White, ObjectTexture->GetResource(), false);
 	}
 }
 
diff --git a/Engine/Plugins/FX/Niagara/Source/NiagaraEditor/Private/Widgets/SNiagaraBakerViewport.cpp b/Engine/Plugins/FX/Niagara/Source/NiagaraEditor/Private/Widgets/SNiagaraBakerViewport.cpp
--- a/Engine/Plugins/FX/Niagara/Source/NiagaraEditor/Private/Widgets/SNiagaraBakerViewport.cpp
+++ b/Engine/Plugins/FX/Niagara/Source/NiagaraEditor/Private/Widgets/SNiagaraBakerViewport.cpp
@@ -317,7 +317,7 @@ public:
 					PreviewViewRect.Min.X, PreviewViewRect.Min.Y, PreviewViewRect.Width(), PreviewViewRect.Height(),
 					HalfPixel.X, HalfPixel.Y, 1.0f - HalfPixel.X, 1.0f - HalfPixel.Y,
 					FLinearColor::White,
-					RealtimeRenderTarget->Resource,
+					RealtimeRenderTarget->GetResource(),
 					false	//-TODO: Preview with alpha?
 				);
 
@@ -367,7 +367,7 @@ public:
 					BakerViewRect.Min.X, BakerViewRect.Min.Y, BakerViewRect.Width(), BakerViewRect.Height(),
 					(float(FramePixelA.X) + 0.5f) / float(TextureSize.X), (float(FramePixelA.Y) + 0.5f) / float(TextureSize.Y), (float(FramePixelA.X + OutputTexture.FrameSize.X) - 0.5f) / float(TextureSize.X), (float(FramePixelA.Y + OutputTexture.FrameSize.Y) - 0.5f) / float(TextureSize.Y),
 					FLinearColor::White,
-					OutputTexture.GeneratedTexture->Resource,
+					OutputTexture.GeneratedTexture->GetResource(),
 					false	//-TODO: Preview with alpha?
 				);
 
@@ -420,7 +420,7 @@ public:
 			EndUV.X = float(InRect.Width()) / float(FMath::Max(Texture->GetSizeX(), 1));
 			EndUV.Y = float(InRect.Height()) / float(FMath::Max(Texture->GetSizeY(), 1));
 		}
-		Canvas->DrawTile(InRect.Min.X, InRect.Min.Y, InRect.Width(), InRect.Height(), 0.0f, 0.0f, EndUV.X, EndUV.Y, Color, Texture ? Texture->Resource : nullptr, false);
+		Canvas->DrawTile(InRect.Min.X, InRect.Min.Y, InRect.Width(), InRect.Height(), 0.0f, 0.0f, EndUV.X, EndUV.Y, Color, Texture ? Texture->GetResource() : nullptr, false);
 	}
 
 	UTexture2D* GetCheckerboardTexture()
@@ -436,7 +436,7 @@ public:
 	{
 		if (CheckerboardTexture)
 		{
-			if (CheckerboardTexture->Resource)
+			if (CheckerboardTexture->GetResource())
 			{
 				CheckerboardTexture->ReleaseResource();
 			}
diff --git a/Engine/Plugins/MeshPainting/Source/MeshPaintingToolset/Private/TexturePaintToolset.cpp b/Engine/Plugins/MeshPainting/Source/MeshPaintingToolset/Private/TexturePaintToolset.cpp
--- a/Engine/Plugins/MeshPainting/Source/MeshPaintingToolset/Private/TexturePaintToolset.cpp
+++ b/Engine/Plugins/MeshPainting/Source/MeshPaintingToolset/Private/TexturePaintToolset.cpp
@@ -46,7 +46,7 @@ void UTexturePaintToolset::CopyTextureToRenderTargetTexture(UTexture* SourceText
 	UTexture2D* Texture2D = Cast<UTexture2D>(SourceTexture);
 	if (Texture2D != nullptr)
 	{
-		TextureResource = Texture2D->Resource;
+		TextureResource = Texture2D->GetResource();
 	}
 	else
 	{
diff --git a/Engine/Plugins/Runtime/AR/AppleAR/AppleARKit/Source/AppleARKit/Private/AppleARKitTextures.cpp b/Engine/Plugins/Runtime/AR/AppleAR/AppleARKit/Source/AppleARKit/Private/AppleARKitTextures.cpp
--- a/Engine/Plugins/Runtime/AR/AppleAR/AppleARKit/Source/AppleARKit/Private/AppleARKitTextures.cpp
+++ b/Engine/Plugins/Runtime/AR/AppleAR/AppleARKit/Source/AppleARKit/Private/AppleARKitTextures.cpp
@@ -397,12 +397,12 @@ void UAppleARKitTextureCameraImage::UpdateCameraImage(float InTimestamp, CVPixel
 	
 	Timestamp = InTimestamp;
 	
-	if (!Resource)
+	if (!GetResource())
 	{
 		UpdateResource();
 	}
 	
-	if (auto MyResource = static_cast<FARKitTextureResource*>(Resource))
+	if (auto MyResource = static_cast<FARKitTextureResource*>(GetResource()))
 	{
 		MyResource->UpdateCameraImage(CameraImage, InPixelFormat, ColorSpace, BlurParams);
 	}
@@ -453,7 +453,7 @@ UAppleARKitEnvironmentCaptureProbeTexture::UAppleARKitEnvironmentCaptureProbeTex
 #if PLATFORM_APPLE
 void UAppleARKitEnvironmentCaptureProbeTexture::Init(float InTimestamp, id<MTLTexture> InEnvironmentTexture)
 {
-	if (Resource == nullptr)
+	if (GetResource() == nullptr)
 	{
 		UpdateResource();
 	}
@@ -481,10 +481,10 @@ void UAppleARKitEnvironmentCaptureProbeTexture::Init(float InTimestamp, id<MTLTe
 		Size.Y = MetalTexture.height;
 	}
 	// Force an update to our external texture on the render thread
-	if (Resource != nullptr)
+	if (GetResource() != nullptr)
 	{
 		ENQUEUE_RENDER_COMMAND(UpdateEnvironmentCapture)(
-			[InResource = Resource](FRHICommandListImmediate& RHICmdList)
+			[InResource = GetResource()](FRHICommandListImmediate& RHICmdList)
 			{
 				InResource->InitRHI();
 			});
@@ -743,12 +743,12 @@ void UAppleARKitOcclusionTexture::SetMetalTexture(float InTimestamp, id<MTLTextu
 		}
 	}
 	
-	if (Resource == nullptr)
+	if (GetResource() == nullptr)
 	{
 		UpdateResource();
 	}
 	
-	if (auto MyResource = static_cast<FARKitTextureResource*>(Resource))
+	if (auto MyResource = static_cast<FARKitTextureResource*>(GetResource()))
 	{
 		MyResource->UpdateMetalTexture(InMetalTexture, PixelFormat, ColorSpace);
 	}
@@ -1085,7 +1085,7 @@ UAppleARKitCameraVideoTexture::UAppleARKitCameraVideoTexture(const FObjectInitia
 
 void UAppleARKitCameraVideoTexture::Init()
 {
-	if (!Resource)
+	if (!GetResource())
 	{
 		UpdateResource();
 	}
@@ -1096,7 +1096,7 @@ void UAppleARKitCameraVideoTexture::UpdateFrame(const FAppleARKitFrame& InFrame)
 #if SUPPORTS_ARKIT_1_0
 	if (InFrame.CapturedYImage && InFrame.CapturedCbCrImage)
 	{
-		if (auto VideoResource = static_cast<FARKitCameraVideoResource*>(Resource))
+		if (auto VideoResource = static_cast<FARKitCameraVideoResource*>(GetResource()))
 		{
 			auto CapturedYImageCopy = InFrame.CapturedYImage;
 			auto CapturedCbCrImageCopy = InFrame.CapturedCbCrImage;
@@ -1120,7 +1120,7 @@ void UAppleARKitCameraVideoTexture::UpdateFrame(const FAppleARKitFrame& InFrame)
 #if SUPPORTS_ARKIT_1_0
 id<MTLTexture> UAppleARKitCameraVideoTexture::GetMetalTexture() const
 {
-	if (auto MyResource = (FARKitCameraVideoResource*)Resource)
+	if (auto MyResource = (FARKitCameraVideoResource*)GetResource())
 	{
 		return MyResource->GetMetalTexture();
 	}
diff --git a/Engine/Plugins/Runtime/AR/Google/GoogleARCore/Source/GoogleARCoreBase/Private/GoogleARCoreTexture.cpp b/Engine/Plugins/Runtime/AR/Google/GoogleARCore/Source/GoogleARCoreBase/Private/GoogleARCoreTexture.cpp
--- a/Engine/Plugins/Runtime/AR/Google/GoogleARCore/Source/GoogleARCoreBase/Private/GoogleARCoreTexture.cpp
+++ b/Engine/Plugins/Runtime/AR/Google/GoogleARCore/Source/GoogleARCoreBase/Private/GoogleARCoreTexture.cpp
@@ -67,7 +67,7 @@ FTextureResource* UARCoreCameraTexture::CreateResource()
 
 uint32 UARCoreCameraTexture::GetTextureId() const
 {
-	if (auto MyResource = static_cast<FARCoreCameraTextureResource*>(Resource))
+	if (const FARCoreCameraTextureResource* MyResource = static_cast<const FARCoreCameraTextureResource*>(GetResource()))
 	{
 		return MyResource->GetTextureId();
 	}
@@ -202,12 +202,12 @@ void UARCoreDepthTexture::UpdateDepthImage(const ArSession* SessionHandle, const
 	Size.X = ImageWidth;
 	Size.Y = ImageHeight;
 	
-	if (!Resource)
+	if (!GetResource())
 	{
 		UpdateResource();
 	}
 	
-	if (auto MyResource = static_cast<FARCoreDepthTextureResource*>(Resource))
+	if (auto MyResource = static_cast<FARCoreDepthTextureResource*>(GetResource()))
 	{
 		MyResource->UpdateDepthData((const uint16*)DepthData, PlaneSize, ImageWidth, ImageHeight);
 	}
diff --git a/Engine/Plugins/Runtime/GoogleVR/GoogleVRHMD/Source/GoogleVRHMD/Private/GoogleVRSplash.cpp b/Engine/Plugins/Runtime/GoogleVR/GoogleVRHMD/Source/GoogleVRHMD/Private/GoogleVRSplash.cpp
--- a/Engine/Plugins/Runtime/GoogleVR/GoogleVRHMD/Source/GoogleVRHMD/Private/GoogleVRSplash.cpp
+++ b/Engine/Plugins/Runtime/GoogleVR/GoogleVRHMD/Source/GoogleVRHMD/Private/GoogleVRSplash.cpp
@@ -236,7 +236,7 @@ void FGoogleVRSplash::RenderStereoSplashScreen(FRHICommandListImmediate& RHICmdL
 
 			SetGraphicsPipelineState(RHICmdList, GraphicsPSOInit);
 
-			PixelShader->SetParameters(RHICmdList, TStaticSamplerState<SF_Bilinear>::GetRHI(), SplashTexture->Resource->TextureRHI);
+			PixelShader->SetParameters(RHICmdList, TStaticSamplerState<SF_Bilinear>::GetRHI(), SplashTexture->GetResource()->TextureRHI);
 
 			// Modify the v to flip the texture. Somehow it renders flipped if not.
 			float U = SplashTextureUVOffset.X;
diff --git a/Engine/Plugins/Runtime/HairStrands/Source/HairStrandsCore/Private/GroomTextureBuilder.cpp b/Engine/Plugins/Runtime/HairStrands/Source/HairStrandsCore/Private/GroomTextureBuilder.cpp
--- a/Engine/Plugins/Runtime/HairStrands/Source/HairStrandsCore/Private/GroomTextureBuilder.cpp
+++ b/Engine/Plugins/Runtime/HairStrands/Source/HairStrandsCore/Private/GroomTextureBuilder.cpp
@@ -146,7 +146,7 @@ static void InternalGenerateFollicleTexture_GPU(
 {
 	const ERHIFeatureLevel::Type FeatureLevel = GMaxRHIFeatureLevel;
 
-	if (OutTexture == nullptr || !OutTexture->Resource || !OutTexture->Resource->GetTexture2DRHI() ||
+	if (OutTexture == nullptr || !OutTexture->GetResource() || !OutTexture->GetResource()->GetTexture2DRHI() ||
 		(InRootUVBuffers_R.Num() == 0 &&
 		 InRootUVBuffers_G.Num() == 0 &&
 		 InRootUVBuffers_B.Num() == 0 &&
@@ -218,13 +218,13 @@ static void InternalGenerateFollicleTexture_GPU(
 		FollicleMaskTexture,
 		[FollicleMaskTexture, MipCount, OutTexture](FRHICommandListImmediate& RHICmdList)
 	{
-		if (OutTexture && OutTexture->Resource && OutTexture->Resource->GetTexture2DRHI())
+		if (OutTexture && OutTexture->GetResource() && OutTexture->GetResource()->GetTexture2DRHI())
 		{
 			FRHICopyTextureInfo CopyInfo;
 			CopyInfo.NumMips = MipCount;
 			RHICmdList.CopyTexture(
 				FollicleMaskTexture->GetRHI(),
-				OutTexture->Resource->GetTexture2DRHI(),
+				OutTexture->GetResource()->GetTexture2DRHI(),
 				CopyInfo);
 		}
 	});
@@ -389,7 +389,7 @@ void FGroomTextureBuilder::BuildFollicleTexture(const TArray<FFollicleInfo>& InI
 			ENQUEUE_RENDER_COMMAND(FFollicleTextureQuery)(
 			[Resolution, MipCount, Format, InInfos, OutTexture](FRHICommandListImmediate& RHICmdList)
 			{
-				if (OutTexture->Resource)
+				if (OutTexture->GetResource())
 				{
 					GFollicleQueries.Enqueue({ InInfos, Resolution, MipCount, Format, OutTexture });
 				}
diff --git a/Engine/Plugins/Runtime/Oculus/OculusVR/Source/OculusHMD/Private/OculusHMD_Splash.cpp b/Engine/Plugins/Runtime/Oculus/OculusVR/Source/OculusHMD/Private/OculusHMD_Splash.cpp
--- a/Engine/Plugins/Runtime/Oculus/OculusVR/Source/OculusHMD/Private/OculusHMD_Splash.cpp
+++ b/Engine/Plugins/Runtime/Oculus/OculusVR/Source/OculusHMD/Private/OculusHMD_Splash.cpp
@@ -534,9 +534,9 @@ void FSplash::DoShow()
 		//@DBG BEGIN
 		if (SplashLayer.Desc.LoadingTexture->IsValidLowLevel())
 		{
-			if (SplashLayer.Desc.LoadingTexture->Resource && SplashLayer.Desc.LoadingTexture->Resource->TextureRHI)
+			if (SplashLayer.Desc.LoadingTexture->GetResource() && SplashLayer.Desc.LoadingTexture->GetResource()->TextureRHI)
 			{
-				SplashLayer.Desc.LoadedTexture = SplashLayer.Desc.LoadingTexture->Resource->TextureRHI;
+				SplashLayer.Desc.LoadedTexture = SplashLayer.Desc.LoadingTexture->GetResource()->TextureRHI;
 			}
 			else
 			{
diff --git a/Engine/Plugins/Runtime/Oculus/OculusVR/Source/OculusMR/Private/OculusMR_CastingCameraActor.cpp b/Engine/Plugins/Runtime/Oculus/OculusVR/Source/OculusMR/Private/OculusMR_CastingCameraActor.cpp
--- a/Engine/Plugins/Runtime/Oculus/OculusVR/Source/OculusMR/Private/OculusMR_CastingCameraActor.cpp
+++ b/Engine/Plugins/Runtime/Oculus/OculusVR/Source/OculusMR/Private/OculusMR_CastingCameraActor.cpp
@@ -419,8 +419,8 @@ void AOculusMR_CastingCameraActor::Tick(float DeltaTime)
 					{
 						// The Vulkan RHI's implementation of GetNativeResource is different and returns the VkImage cast
 						// as a void* instead of a pointer to the VkImage, so we need this workaround
-						BackgroundTexture = (void*)BackgroundRenderTargets[EncodeIndex]->Resource->TextureRHI->GetNativeResource();
-						ForegroundTexture = (void*)ForegroundRenderTargets[EncodeIndex]->Resource->TextureRHI->GetNativeResource();
+						BackgroundTexture = (void*)BackgroundRenderTargets[EncodeIndex]->GetResource()->TextureRHI->GetNativeResource();
+						ForegroundTexture = (void*)ForegroundRenderTargets[EncodeIndex]->GetResource()->TextureRHI->GetNativeResource();
 					});
 				});
 			}
@@ -430,8 +430,8 @@ void AOculusMR_CastingCameraActor::Tick(float DeltaTime)
 				{
 					ExecuteOnRHIThread([this, EncodeIndex, &BackgroundTexture, &ForegroundTexture]()
 					{
-						BackgroundTexture = *((void**)BackgroundRenderTargets[EncodeIndex]->Resource->TextureRHI->GetNativeResource());
-						ForegroundTexture = *((void**)ForegroundRenderTargets[EncodeIndex]->Resource->TextureRHI->GetNativeResource());
+						BackgroundTexture = *((void**)BackgroundRenderTargets[EncodeIndex]->GetResource()->TextureRHI->GetNativeResource());
+						ForegroundTexture = *((void**)ForegroundRenderTargets[EncodeIndex]->GetResource()->TextureRHI->GetNativeResource());
 					});
 				});
 			}
@@ -499,7 +499,7 @@ void AOculusMR_CastingCameraActor::UpdateCameraColorTexture(const ovrpSizei &fra
 	{
 		SrcData,
 		Pitch,
-		(FTexture2DResource*)CameraColorTexture->Resource,
+		(FTexture2DResource*)CameraColorTexture->GetResource(),
 		(uint32) frameSize.w,
 		(uint32) frameSize.h
 	};
diff --git a/Engine/Plugins/Runtime/TextureShare/Source/TextureShare/Private/Blueprints/TextureShareContainers.cpp b/Engine/Plugins/Runtime/TextureShare/Source/TextureShare/Private/Blueprints/TextureShareContainers.cpp
--- a/Engine/Plugins/Runtime/TextureShare/Source/TextureShare/Private/Blueprints/TextureShareContainers.cpp
+++ b/Engine/Plugins/Runtime/TextureShare/Source/TextureShare/Private/Blueprints/TextureShareContainers.cpp
@@ -57,7 +57,7 @@ bool FTextureShareBPTexture2D::GetTexture2DRHI_RenderThread(FTexture2DRHIRef& Ou
 
 	if (Texture)
 	{
-		FTexture2DResource* Texture2DResource = Texture ? (FTexture2DResource*)Texture->Resource : nullptr;
+		FTexture2DResource* Texture2DResource = Texture ? (FTexture2DResource*)Texture->GetResource() : nullptr;
 		if (Texture2DResource)
 		{
 			OutTexture2DRHIRef = Texture2DResource->GetTexture2DRHI();
diff --git a/Engine/Plugins/Runtime/WindowsMixedReality/Source/HoloLensAR/Private/HoloLensCameraImageTexture.cpp b/Engine/Plugins/Runtime/WindowsMixedReality/Source/HoloLensAR/Private/HoloLensCameraImageTexture.cpp
--- a/Engine/Plugins/Runtime/WindowsMixedReality/Source/HoloLensAR/Private/HoloLensCameraImageTexture.cpp
+++ b/Engine/Plugins/Runtime/WindowsMixedReality/Source/HoloLensAR/Private/HoloLensCameraImageTexture.cpp
@@ -406,9 +406,9 @@ void UHoloLensCameraImageTexture::Init(void* handle)
 	if (LastUpdateFrame != GFrameCounter)
 	{
 		LastUpdateFrame = GFrameCounter;
-		if (Resource != nullptr)
+		if (GetResource() != nullptr)
 		{
-			FHoloLensCameraImageResource* LambdaResource = static_cast<FHoloLensCameraImageResource*>(Resource);
+			FHoloLensCameraImageResource* LambdaResource = static_cast<FHoloLensCameraImageResource*>(GetResource());
 			ENQUEUE_RENDER_COMMAND(Init_RenderThread)(
 				[LambdaResource, handle](FRHICommandListImmediate&)
 			{
diff --git a/Engine/Plugins/Runtime/nDisplay/Source/DisplayCluster/Private/Game/EngineClasses/Scene/DisplayClusterPreviewComponent.cpp b/Engine/Plugins/Runtime/nDisplay/Source/DisplayCluster/Private/Game/EngineClasses/Scene/DisplayClusterPreviewComponent.cpp
--- a/Engine/Plugins/Runtime/nDisplay/Source/DisplayCluster/Private/Game/EngineClasses/Scene/DisplayClusterPreviewComponent.cpp
+++ b/Engine/Plugins/Runtime/nDisplay/Source/DisplayCluster/Private/Game/EngineClasses/Scene/DisplayClusterPreviewComponent.cpp
@@ -235,7 +235,7 @@ void UDisplayClusterPreviewComponent::UpdatePreviewRenderTarget()
 		if (RenderTarget)
 		{
 			// clear preview RTT to black in this case
-			FTextureRenderTarget2DResource* TexResource = (FTextureRenderTarget2DResource*)RenderTarget->Resource;
+			FTextureRenderTarget2DResource* TexResource = (FTextureRenderTarget2DResource*)RenderTarget->GetResource();
 			if (TexResource)
 			{
 				FCanvas Canvas(TexResource, NULL, 0, 0, 0, GMaxRHIFeatureLevel);
diff --git a/Engine/Plugins/Runtime/nDisplay/Source/DisplayCluster/Private/Render/Viewport/Configuration/DisplayClusterViewportConfigurationHelpers.h b/Engine/Plugins/Runtime/nDisplay/Source/DisplayCluster/Private/Render/Viewport/Configuration/DisplayClusterViewportConfigurationHelpers.h
--- a/Engine/Plugins/Runtime/nDisplay/Source/DisplayCluster/Private/Render/Viewport/Configuration/DisplayClusterViewportConfigurationHelpers.h
+++ b/Engine/Plugins/Runtime/nDisplay/Source/DisplayCluster/Private/Render/Viewport/Configuration/DisplayClusterViewportConfigurationHelpers.h
@@ -277,7 +277,7 @@ public:
 
 		if (InOverride.bAllowOverride && InOverride.SourceTexture != nullptr)
 		{
-			FTextureRHIRef& TextureRHI = InOverride.SourceTexture->Resource->TextureRHI;
+			FTextureRHIRef& TextureRHI = InOverride.SourceTexture->GetResource()->TextureRHI;
 
 			if (TextureRHI.IsValid())
 			{
@@ -568,7 +568,7 @@ public:
 			DstCameraSettings.ChromakeyMarkersDistance = InChromakeyMarkers.MarkerTileDistance;
 
 			// Assign texture RHI ref
-			FTextureResource* MarkersResource = InChromakeyMarkers.MarkerTileRGBA->Resource;
+			FTextureResource* MarkersResource = InChromakeyMarkers.MarkerTileRGBA->GetResource();
 			if (MarkersResource)
 			{
 				DstCameraSettings.ChromakeMarkerTextureRHI = MarkersResource->TextureRHI;
diff --git a/Engine/Plugins/Runtime/nDisplay/Source/DisplayClusterConfigurator/Private/Views/OutputMapping/GraphNodes/SDisplayClusterConfiguratorViewportNode.cpp b/Engine/Plugins/Runtime/nDisplay/Source/DisplayClusterConfigurator/Private/Views/OutputMapping/GraphNodes/SDisplayClusterConfiguratorViewportNode.cpp
--- a/Engine/Plugins/Runtime/nDisplay/Source/DisplayClusterConfigurator/Private/Views/OutputMapping/GraphNodes/SDisplayClusterConfiguratorViewportNode.cpp
+++ b/Engine/Plugins/Runtime/nDisplay/Source/DisplayClusterConfigurator/Private/Views/OutputMapping/GraphNodes/SDisplayClusterConfiguratorViewportNode.cpp
@@ -343,8 +343,8 @@ void SDisplayClusterConfiguratorViewportNode::UpdatePreviewTexture()
 			{
 				BackgroundActiveBrush = FSlateBrush();
 				BackgroundActiveBrush.SetResourceObject(CachedTexture);
-				BackgroundActiveBrush.ImageSize.X = CachedTexture->Resource->GetSizeX();
-				BackgroundActiveBrush.ImageSize.Y = CachedTexture->Resource->GetSizeY();
+				BackgroundActiveBrush.ImageSize.X = CachedTexture->GetResource()->GetSizeX();
+				BackgroundActiveBrush.ImageSize.Y = CachedTexture->GetResource()->GetSizeY();
 			}
 		}
 		else
diff --git a/Engine/Plugins/VirtualProduction/CameraCalibration/Source/CameraCalibration/Private/CalibratedMapProcessor.cpp b/Engine/Plugins/VirtualProduction/CameraCalibration/Source/CameraCalibration/Private/CalibratedMapProcessor.cpp
--- a/Engine/Plugins/VirtualProduction/CameraCalibration/Source/CameraCalibration/Private/CalibratedMapProcessor.cpp
+++ b/Engine/Plugins/VirtualProduction/CameraCalibration/Source/CameraCalibration/Private/CalibratedMapProcessor.cpp
@@ -148,9 +148,9 @@ void FCalibratedMapProcessor::ExecuteJob(TSharedPtr<FDerivedDistortionDataJob> J
 		if(SourceTexture == nullptr
 		|| DestinationUndistortionTexture == nullptr
 		|| DestinationDistortionTexture == nullptr
-		|| SourceTexture->Resource == nullptr
-		|| DestinationUndistortionTexture->Resource == nullptr
-		|| DestinationDistortionTexture->Resource == nullptr)
+		|| SourceTexture->GetResource() == nullptr
+		|| DestinationUndistortionTexture->GetResource() == nullptr
+		|| DestinationDistortionTexture->GetResource() == nullptr)
 		{
 			Job->State = EDerivedDistortionDataJobState::Completed;
 			Job->Output.Result = EDerivedDistortionDataResult::Error;
@@ -158,9 +158,9 @@ void FCalibratedMapProcessor::ExecuteJob(TSharedPtr<FDerivedDistortionDataJob> J
 			return;
 		}
 		
-		FTextureResource* SourceDistortionMap = SourceTexture->Resource;
-		FTextureResource* DestinationUndistortionDisplacementMap = DestinationUndistortionTexture->Resource;
-		FTextureResource* DestinationDistortionDisplacementMap = DestinationDistortionTexture->Resource;
+		FTextureResource* SourceDistortionMap = SourceTexture->GetResource();
+		FTextureResource* DestinationUndistortionDisplacementMap = DestinationUndistortionTexture->GetResource();
+		FTextureResource* DestinationDistortionDisplacementMap = DestinationDistortionTexture->GetResource();
 
 		{
 			FScopeLock Lock(&RunningJobsCriticalSection);
diff --git a/Engine/Plugins/VirtualProduction/DMX/DMXFixtures/Source/DMXFixtures/Private/DMXFixtureActorMatrix.cpp b/Engine/Plugins/VirtualProduction/DMX/DMXFixtures/Source/DMXFixtures/Private/DMXFixtureActorMatrix.cpp
--- a/Engine/Plugins/VirtualProduction/DMX/DMXFixtures/Source/DMXFixtures/Private/DMXFixtureActorMatrix.cpp
+++ b/Engine/Plugins/VirtualProduction/DMX/DMXFixtures/Source/DMXFixtures/Private/DMXFixtureActorMatrix.cpp
@@ -13,12 +13,12 @@ namespace
 {
 	void UpdateMatrixTexture(uint8* MatrixData, UTexture2D* DynamicTexture, int32 MipIndex, uint32 NumRegions, FUpdateTextureRegion2D Region, uint32 SrcPitch, uint32 SrcBpp)
 	{
-		if (DynamicTexture->Resource)
+		if (DynamicTexture->GetResource())
 		{
 			ENQUEUE_RENDER_COMMAND(UpdateTextureRegionsData)(
 				[=](FRHICommandListImmediate& RHICmdList)
 				{
-					FTexture2DResource* Resource = (FTexture2DResource*)DynamicTexture->Resource;
+					FTexture2DResource* Resource = (FTexture2DResource*)DynamicTexture->GetResource();
 					RHIUpdateTexture2D(
 						Resource->GetTexture2DRHI(),
 						MipIndex,
diff --git a/Engine/Plugins/VirtualProduction/DMX/DMXFixtures/Source/DMXFixtures/Private/DMXFixtureComponent.cpp b/Engine/Plugins/VirtualProduction/DMX/DMXFixtures/Source/DMXFixtures/Private/DMXFixtureComponent.cpp
--- a/Engine/Plugins/VirtualProduction/DMX/DMXFixtures/Source/DMXFixtures/Private/DMXFixtureComponent.cpp
+++ b/Engine/Plugins/VirtualProduction/DMX/DMXFixtures/Source/DMXFixtures/Private/DMXFixtureComponent.cpp
@@ -118,7 +118,7 @@ TArray<FLinearColor> UDMXFixtureComponent::GetTextureCenterColors(UTexture2D* Te
 		(
 			[Texture, &PixelColorArray, NumCells](FRHICommandListImmediate& RHICmdList)
 			{
-				FTexture2DResource* uTex2DRes = (FTexture2DResource*)Texture->Resource;
+				FTexture2DResource* uTex2DRes = (FTexture2DResource*)Texture->GetResource();
 				if (uTex2DRes)
 				{
 					FIntVector Size = uTex2DRes->TextureRHI->GetSizeXYZ();
