Fixing issue where if a script only had a fast path function, it would never get its delegates bound and could crash downstream.

diff --git a/Engine/Plugins/FX/Niagara/Source/Niagara/Private/NiagaraScriptExecutionContext.cpp b/Engine/Plugins/FX/Niagara/Source/Niagara/Private/NiagaraScriptExecutionContext.cpp
--- a/Engine/Plugins/FX/Niagara/Source/Niagara/Private/NiagaraScriptExecutionContext.cpp
+++ b/Engine/Plugins/FX/Niagara/Source/Niagara/Private/NiagaraScriptExecutionContext.cpp
@@ -325,6 +325,14 @@ bool FNiagaraScriptExecutionContext::Tick(FNiagaraSystemInstance* ParentSystemIn
 				FunctionTable[LocalFunctionTableIndices[LocalFunctionIt]] = &LocalFunctionTable[LocalFunctionIt];
 			}
 
+			for (int32 i = 0; i < FunctionTable.Num(); i++)
+			{
+				if (FunctionTable[i] == nullptr)
+				{
+					UE_LOG(LogNiagara, Warning, TEXT("Invalid Function Table Entry! %s"), *ScriptExecutableData.CalledVMExternalFunctions[i].Name.ToString());
+				}
+			}
+
 #if WITH_EDITOR	
 			// We may now have new errors that we need to broadcast about, so flush the asset parameters delegate..
 			if (ParentSystemInstance)
@@ -493,12 +501,7 @@ bool FNiagaraSystemScriptExecutionContext::Tick(class FNiagaraSystemInstance* In
 						}
 						break;
 					}
-				}
-
-				for (int32 FunctionIt = 0; FunctionIt < FunctionCount; ++FunctionIt)
-				{
-					FunctionTable[FunctionIt] = &ExtFunctionInfo[FunctionIt].Function;
-				}
+				}				
 
 				if (FuncInfo.Function.IsBound() == false)
 				{
@@ -507,7 +510,30 @@ bool FNiagaraSystemScriptExecutionContext::Tick(class FNiagaraSystemInstance* In
 					return false;
 				}
 			}
+
+			if (FunctionTable.Num() != ExtFunctionInfo.Num())
+			{
+				UE_LOG(LogNiagara, Warning, TEXT("Error building data interface function table for system script!"));
+				FunctionTable.Empty();
+				return false;
+			}
+
+			for (int32 FunctionIt = 0; FunctionIt < FunctionTable.Num(); ++FunctionIt)
+			{
+				FunctionTable[FunctionIt] = &ExtFunctionInfo[FunctionIt].Function;
+			}
+
+			for (int32 i = 0; i < FunctionTable.Num(); i++)
+			{
+				if (FunctionTable[i] == nullptr)
+				{
+					UE_LOG(LogNiagara, Warning, TEXT("Invalid Function Table Entry! %s"), *ScriptExecutableData.CalledVMExternalFunctions[i].Name.ToString());
+				}
+			}
+
 		}
+
+		
 	}
 
 	Parameters.Tick();
