Fix Engine.Emitter.SimulationPosition being backwards in respect to Emitter localspace flag. Now map directly to engine_owner_position instead of map.engine.owner.position.

diff --git a/Engine/Plugins/FX/Niagara/Source/NiagaraCore/Private/NiagaraCustomVersion.cpp b/Engine/Plugins/FX/Niagara/Source/NiagaraCore/Private/NiagaraCustomVersion.cpp
--- a/Engine/Plugins/FX/Niagara/Source/NiagaraCore/Private/NiagaraCustomVersion.cpp
+++ b/Engine/Plugins/FX/Niagara/Source/NiagaraCore/Private/NiagaraCustomVersion.cpp
@@ -8,4 +8,4 @@ const FGuid FNiagaraCustomVersion::GUID(0xFCF57AFA, 0x50764283, 0xB9A9E658, 0xFF
 // Register the custom version with core
 FCustomVersionRegistration GRegisterNiagaraCustomVersion(FNiagaraCustomVersion::GUID, FNiagaraCustomVersion::LatestVersion, TEXT("NiagaraVer"));
 
-const FGuid FNiagaraCustomVersion::LatestScriptCompileVersion(0x7F8E9CA7, 0x61806B41, 0xA6C7E8A3, 0xCB2120F4);
+const FGuid FNiagaraCustomVersion::LatestScriptCompileVersion(0x55518ACA, 0xDAEE46EA, 0xBB4C108E, 0xAFD271AB );
diff --git a/Engine/Plugins/FX/Niagara/Source/NiagaraEditor/Private/NiagaraHlslTranslator.cpp b/Engine/Plugins/FX/Niagara/Source/NiagaraEditor/Private/NiagaraHlslTranslator.cpp
--- a/Engine/Plugins/FX/Niagara/Source/NiagaraEditor/Private/NiagaraHlslTranslator.cpp
+++ b/Engine/Plugins/FX/Niagara/Source/NiagaraEditor/Private/NiagaraHlslTranslator.cpp
@@ -5132,11 +5132,11 @@ bool FHlslNiagaraTranslator::HandleBoundConstantVariableToDataSetRead(FNiagaraVa
 		// Simulation position is 0 for localspace emitters.
 		// If we are not in localspace then this will not be a literal constant and is instead a default linked variable as handled in GenerateConstantString().
 		const bool bEmitterLocalSpace = CompileOptions.AdditionalDefines.Contains(SYS_PARAM_EMITTER_LOCALSPACE.GetName().ToString());
-		if (bEmitterLocalSpace)
+		if (bEmitterLocalSpace == false)
 		{
-			const FString ConstantStr = FString::Printf(TEXT("%s"), *SYS_PARAM_ENGINE_POSITION.GetName().ToString());
-			const FString ParameterMapInstanceNameStr = GetParameterMapInstanceName(InParamMapHistoryIdx);
-			Output = AddBodyChunk(GetUniqueSymbolName(TEXT("Constant")), FString::Printf(TEXT("%s.%s"), *ParameterMapInstanceNameStr, *ConstantStr), InVariable.GetType());
+			const FString SysParamEnginePositionSymbolNameString = GetSanitizedSymbolName(SYS_PARAM_ENGINE_POSITION.GetName().ToString(), true);
+			const FString ConstantStr = FString::Printf(TEXT("%s%s"), *SysParamEnginePositionSymbolNameString, *FString(TEXT(".xyz")));
+			Output = AddBodyChunk(GetUniqueSymbolName(TEXT("Constant")), ConstantStr, InVariable.GetType());
 			return true;
 		}
 		else
