- Add Matrix4 support for RI bake out

diff --git a/Engine/Plugins/FX/Niagara/Source/NiagaraCore/Private/NiagaraCustomVersion.cpp b/Engine/Plugins/FX/Niagara/Source/NiagaraCore/Private/NiagaraCustomVersion.cpp
--- a/Engine/Plugins/FX/Niagara/Source/NiagaraCore/Private/NiagaraCustomVersion.cpp
+++ b/Engine/Plugins/FX/Niagara/Source/NiagaraCore/Private/NiagaraCustomVersion.cpp
@@ -8,4 +8,4 @@ const FGuid FNiagaraCustomVersion::GUID(0xFCF57AFA, 0x50764283, 0xB9A9E658, 0xFF
 // Register the custom version with core
 FCustomVersionRegistration GRegisterNiagaraCustomVersion(FNiagaraCustomVersion::GUID, FNiagaraCustomVersion::LatestVersion, TEXT("NiagaraVer"));
 
-const FGuid FNiagaraCustomVersion::LatestScriptCompileVersion(0x55518ACA, 0xDAEE46EA, 0xBB4C108E, 0xAFD271AB );
+const FGuid FNiagaraCustomVersion::LatestScriptCompileVersion(0x9C6C9663, 0x7C42D545, 0x9DC47339, 0x5624F6E6);
diff --git a/Engine/Plugins/FX/Niagara/Source/NiagaraEditor/Private/NiagaraHlslTranslator.cpp b/Engine/Plugins/FX/Niagara/Source/NiagaraEditor/Private/NiagaraHlslTranslator.cpp
--- a/Engine/Plugins/FX/Niagara/Source/NiagaraEditor/Private/NiagaraHlslTranslator.cpp
+++ b/Engine/Plugins/FX/Niagara/Source/NiagaraEditor/Private/NiagaraHlslTranslator.cpp
@@ -4564,6 +4564,17 @@ FString FHlslNiagaraTranslator::GenerateConstantString(const FNiagaraVariable& C
 			int32* ValuePtr = (int32*)Constant.GetData();
 			ConstantStr = FString::Printf(TEXT("%d"), *ValuePtr);
 		}
+		else if (Type == FNiagaraTypeDefinition::GetMatrix4Def() )
+		{
+			float* ValuePtr = (float*)Constant.GetData();
+			ConstantStr = FString::Printf(
+				TEXT("float4x4(%g,%g,%g,%g, %g,%g,%g,%g, %g,%g,%g,%g, %g,%g,%g,%g)"),
+				ValuePtr[ 0], ValuePtr[ 1], ValuePtr[ 2], ValuePtr[ 3],
+				ValuePtr[ 4], ValuePtr[ 5], ValuePtr[ 6], ValuePtr[ 7],
+				ValuePtr[ 8], ValuePtr[ 9], ValuePtr[10], ValuePtr[11],
+				ValuePtr[12], ValuePtr[13], ValuePtr[14], ValuePtr[15]
+			);
+		}
 		else if (Type == FNiagaraTypeDefinition::GetBoolDef())
 		{
 			if (Constant.GetValue<FNiagaraBool>().IsValid() == false)
