- Integrate CL 16315785 - Temporarily remove consuming per instance data per tick as not all data interfaces are currently safe for this (Skel Mesh DI)
- Clean up skeletal mesh data interface

#rnx

diff --git a/Engine/Plugins/FX/Niagara/Source/Niagara/Classes/NiagaraDataInterfaceSkeletalMesh.h b/Engine/Plugins/FX/Niagara/Source/Niagara/Classes/NiagaraDataInterfaceSkeletalMesh.h
--- a/Engine/Plugins/FX/Niagara/Source/Niagara/Classes/NiagaraDataInterfaceSkeletalMesh.h
+++ b/Engine/Plugins/FX/Niagara/Source/Niagara/Classes/NiagaraDataInterfaceSkeletalMesh.h
@@ -438,8 +438,8 @@ public:
 
 	virtual FString GetFriendlyName() const override { return TEXT("FSkeletalMeshGpuSpawnStaticBuffers"); }
 
-	FShaderResourceViewRHIRef GetBufferTriangleUniformSamplerProbAliasSRV() const { return BufferTriangleUniformSamplerProbAliasSRV; }
-	FShaderResourceViewRHIRef GetBufferTriangleMatricesOffsetSRV() const { return BufferTriangleMatricesOffsetSRV; }
+	FRHIShaderResourceView* GetBufferTriangleUniformSamplerProbAliasSRV() const { return BufferTriangleUniformSamplerProbAliasSRV; }
+	FRHIShaderResourceView* GetBufferTriangleMatricesOffsetSRV() const { return BufferTriangleMatricesOffsetSRV; }
 	uint32 GetTriangleCount() const { return TriangleCount; }
 	uint32 GetVertexCount() const { return VertexCount; }
 
@@ -447,9 +447,9 @@ public:
 	bool IsUseGpuUniformlyDistributedSampling() const { return bUseGpuUniformlyDistributedSampling; }
 	int32 GetNumSamplingRegionTriangles() const { return NumSamplingRegionTriangles; }
 	int32 GetNumSamplingRegionVertices() const { return NumSamplingRegionVertices; }
-	FShaderResourceViewRHIRef GetSampleRegionsProbAliasSRV() const { return SampleRegionsProbAliasSRV; }
-	FShaderResourceViewRHIRef GetSampleRegionsTriangleIndicesSRV() const { return SampleRegionsTriangleIndicesSRV; }
-	FShaderResourceViewRHIRef GetSampleRegionsVerticesSRV() const { return SampleRegionsVerticesSRV; }
+	FRHIShaderResourceView* GetSampleRegionsProbAliasSRV() const { return SampleRegionsProbAliasSRV; }
+	FRHIShaderResourceView* GetSampleRegionsTriangleIndicesSRV() const { return SampleRegionsTriangleIndicesSRV; }
+	FRHIShaderResourceView* GetSampleRegionsVerticesSRV() const { return SampleRegionsVerticesSRV; }
 
 	FRHIShaderResourceView* GetBufferPositionSRV() const { return MeshVertexBufferSrv; }
 	FRHIShaderResourceView* GetBufferIndexSRV() const { return MeshIndexBufferSrv; }
@@ -674,11 +674,12 @@ struct FNDISkeletalMesh_InstanceData
 	/** The referenced LOD data, used to prevent streaming out LODs while they are being referenced*/
 	TRefCountPtr<const FSkeletalMeshLODRenderData> CachedLODData;
 
-	FORCEINLINE_DEBUGGABLE bool ResetRequired(UNiagaraDataInterfaceSkeletalMesh* Interface, FNiagaraSystemInstance* SystemInstance) const;
+	bool ResetRequired(UNiagaraDataInterfaceSkeletalMesh* Interface, FNiagaraSystemInstance* SystemInstance) const;
 
 	bool Init(UNiagaraDataInterfaceSkeletalMesh* Interface, FNiagaraSystemInstance* SystemInstance);
-	FORCEINLINE_DEBUGGABLE bool Tick(UNiagaraDataInterfaceSkeletalMesh* Interface, FNiagaraSystemInstance* SystemInstance, float InDeltaSeconds);
-	FORCEINLINE_DEBUGGABLE void Release();
+	bool Tick(UNiagaraDataInterfaceSkeletalMesh* Interface, FNiagaraSystemInstance* SystemInstance, float InDeltaSeconds);
+	void Release();
+
 	FORCEINLINE int32 GetLODIndex()const { return CachedLODIdx; }
 
 	FORCEINLINE_DEBUGGABLE const FSkinWeightVertexBuffer* GetSkinWeights()
@@ -1078,22 +1079,22 @@ public:
 
 struct FNiagaraDISkeletalMeshPassedDataToRT
 {
-	FSkeletalMeshGpuSpawnStaticBuffers* StaticBuffers;
-	FSkeletalMeshGpuDynamicBufferProxy* DynamicBuffer;
-	const FSkinWeightDataVertexBuffer* MeshSkinWeightBuffer;
-	const FSkinWeightLookupVertexBuffer* MeshSkinWeightLookupBuffer;
-	const FSkeletalMeshUvMappingBufferProxy* UvMappingBuffer;
-	const FSkeletalMeshConnectivityProxy* ConnectivityBuffer;
-
-	bool bIsGpuUniformlyDistributedSampling;
-
-	bool bUnlimitedBoneInfluences;
-	uint32 MeshWeightStrideByte;
-	uint32 MeshSkinWeightIndexSizeByte;
-	FMatrix44f Transform;
-	FMatrix44f PrevTransform;
-	float DeltaSeconds;
-	uint32 UvMappingSet;
+	FSkeletalMeshGpuSpawnStaticBuffers* StaticBuffers = nullptr;
+	FSkeletalMeshGpuDynamicBufferProxy* DynamicBuffer = nullptr;
+	const FSkinWeightDataVertexBuffer* MeshSkinWeightBuffer = nullptr;
+	const FSkinWeightLookupVertexBuffer* MeshSkinWeightLookupBuffer = nullptr;
+	const FSkeletalMeshUvMappingBufferProxy* UvMappingBuffer = nullptr;
+	const FSkeletalMeshConnectivityProxy* ConnectivityBuffer = nullptr;
+
+	bool bIsGpuUniformlyDistributedSampling = false;
+
+	bool bUnlimitedBoneInfluences = false;
+	uint32 MeshWeightStrideByte = 0;
+	uint32 MeshSkinWeightIndexSizeByte = 0;
+	FMatrix44f Transform = FMatrix44f::Identity;
+	FMatrix44f PrevTransform = FMatrix44f::Identity;
+	float DeltaSeconds = 0.0f;
+	uint32 UvMappingSet = 0;
 };
 
 typedef FNiagaraDISkeletalMeshPassedDataToRT FNiagaraDataInterfaceProxySkeletalMeshData;
diff --git a/Engine/Plugins/FX/Niagara/Source/Niagara/Private/NiagaraDataInterfaceSkeletalMesh.cpp b/Engine/Plugins/FX/Niagara/Source/Niagara/Private/NiagaraDataInterfaceSkeletalMesh.cpp
--- a/Engine/Plugins/FX/Niagara/Source/Niagara/Private/NiagaraDataInterfaceSkeletalMesh.cpp
+++ b/Engine/Plugins/FX/Niagara/Source/Niagara/Private/NiagaraDataInterfaceSkeletalMesh.cpp
@@ -1399,48 +1399,20 @@ public:
 			SetSRVParameter(RHICmdList, ComputeShaderRHI, MeshTangentBuffer, StaticBuffers->GetBufferTangentSRV());
 
 			SetShaderValue(RHICmdList, ComputeShaderRHI, MeshNumTexCoord, StaticBuffers->GetNumTexCoord());
-			if (StaticBuffers->GetNumTexCoord() > 0)
-			{
-				SetSRVParameter(RHICmdList, ComputeShaderRHI, MeshTexCoordBuffer, StaticBuffers->GetBufferTexCoordSRV());
-			}
-			else
-			{
-				SetSRVParameter(RHICmdList, ComputeShaderRHI, MeshTexCoordBuffer, FNiagaraRenderer::GetDummyFloatBuffer());
-			}
-			if (StaticBuffers->GetBufferColorSRV())
-			{
-				SetSRVParameter(RHICmdList, ComputeShaderRHI, MeshColorBuffer, StaticBuffers->GetBufferColorSRV());
-			}
-			else
-			{
-				SetSRVParameter(RHICmdList, ComputeShaderRHI, MeshColorBuffer, FNiagaraRenderer::GetDummyFloat4Buffer());
-			}
+			SetSRVParameter(RHICmdList, ComputeShaderRHI, MeshTexCoordBuffer, StaticBuffers->GetBufferTexCoordSRV() != nullptr ? StaticBuffers->GetBufferTexCoordSRV() : FNiagaraRenderer::GetDummyFloatBuffer());
+			SetSRVParameter(RHICmdList, ComputeShaderRHI, MeshColorBuffer, StaticBuffers->GetBufferColorSRV() != nullptr ? StaticBuffers->GetBufferColorSRV() : FNiagaraRenderer::GetDummyFloat4Buffer());
 			SetShaderValue(RHICmdList, ComputeShaderRHI, MeshTriangleCount, StaticBuffers->GetTriangleCount());
 			SetShaderValue(RHICmdList, ComputeShaderRHI, MeshVertexCount, StaticBuffers->GetVertexCount());
 
 			// Set triangle sampling buffer
-			if (InstanceData->bIsGpuUniformlyDistributedSampling)
-			{
-				SetSRVParameter(RHICmdList, ComputeShaderRHI, MeshTriangleSamplerProbAliasBuffer, StaticBuffers->GetBufferTriangleUniformSamplerProbAliasSRV().GetReference());
-			}
-			else
-			{
-				SetSRVParameter(RHICmdList, ComputeShaderRHI, MeshTriangleSamplerProbAliasBuffer, FNiagaraRenderer::GetDummyUIntBuffer());
-			}
+			SetSRVParameter(RHICmdList, ComputeShaderRHI, MeshTriangleSamplerProbAliasBuffer, StaticBuffers->GetBufferTriangleUniformSamplerProbAliasSRV() != nullptr ? StaticBuffers->GetBufferTriangleUniformSamplerProbAliasSRV() : FNiagaraRenderer::GetDummyUIntBuffer());
 
 			// Set triangle sampling region buffer
 			SetShaderValue(RHICmdList, ComputeShaderRHI, MeshNumSamplingRegionTriangles, StaticBuffers->GetNumSamplingRegionTriangles());
 			SetShaderValue(RHICmdList, ComputeShaderRHI, MeshNumSamplingRegionVertices, StaticBuffers->GetNumSamplingRegionVertices());
-			if (StaticBuffers->IsSamplingRegionsAllAreaWeighted() && StaticBuffers->GetNumSamplingRegionTriangles() > 0)
-			{
-				SetSRVParameter(RHICmdList, ComputeShaderRHI, MeshSamplingRegionsProbAliasBuffer, StaticBuffers->GetSampleRegionsProbAliasSRV().GetReference());
-			}
-			else
-			{
-				SetSRVParameter(RHICmdList, ComputeShaderRHI, MeshSamplingRegionsProbAliasBuffer, FNiagaraRenderer::GetDummyUIntBuffer());
-			}
-			SetSRVParameter(RHICmdList, ComputeShaderRHI, MeshSampleRegionsTriangleIndices, StaticBuffers->GetNumSamplingRegionTriangles() > 0 ? StaticBuffers->GetSampleRegionsTriangleIndicesSRV().GetReference() : FNiagaraRenderer::GetDummyUIntBuffer());
-			SetSRVParameter(RHICmdList, ComputeShaderRHI, MeshSampleRegionsVertices, StaticBuffers->GetNumSamplingRegionVertices() > 0 ? StaticBuffers->GetSampleRegionsVerticesSRV().GetReference() : FNiagaraRenderer::GetDummyUIntBuffer());
+			SetSRVParameter(RHICmdList, ComputeShaderRHI, MeshSamplingRegionsProbAliasBuffer, StaticBuffers->GetSampleRegionsProbAliasSRV() != nullptr ? StaticBuffers->GetSampleRegionsProbAliasSRV() : FNiagaraRenderer::GetDummyUIntBuffer());
+			SetSRVParameter(RHICmdList, ComputeShaderRHI, MeshSampleRegionsTriangleIndices, StaticBuffers->GetSampleRegionsTriangleIndicesSRV() != nullptr ? StaticBuffers->GetSampleRegionsTriangleIndicesSRV() : FNiagaraRenderer::GetDummyUIntBuffer());
+			SetSRVParameter(RHICmdList, ComputeShaderRHI, MeshSampleRegionsVertices, StaticBuffers->GetSampleRegionsVerticesSRV() != nullptr ? StaticBuffers->GetSampleRegionsVerticesSRV() : FNiagaraRenderer::GetDummyUIntBuffer());
 
 			FRHIShaderResourceView* MeshSkinWeightBufferSRV = InstanceData->MeshSkinWeightBuffer->GetSRV();
 			SetSRVParameter(RHICmdList, ComputeShaderRHI, MeshSkinWeightBuffer, MeshSkinWeightBufferSRV ? MeshSkinWeightBufferSRV : FNiagaraRenderer::GetDummyUIntBuffer());
@@ -1479,7 +1451,7 @@ public:
 
 			SetShaderValue(RHICmdList, ComputeShaderRHI, NumBones, DynamicBuffers->GetNumBones());
 
-			FRHIShaderResourceView* FilteredAndUnfilteredBonesSRV = StaticBuffers->GetNumFilteredBones() > 0 ? StaticBuffers->GetFilteredAndUnfilteredBonesSRV() : FNiagaraRenderer::GetDummyUIntBuffer();
+			FRHIShaderResourceView* FilteredAndUnfilteredBonesSRV = StaticBuffers->GetFilteredAndUnfilteredBonesSRV() != nullptr ? StaticBuffers->GetFilteredAndUnfilteredBonesSRV() : FNiagaraRenderer::GetDummyUIntBuffer();
 			SetShaderValue(RHICmdList, ComputeShaderRHI, NumFilteredBones, StaticBuffers->GetNumFilteredBones());
 			SetShaderValue(RHICmdList, ComputeShaderRHI, NumUnfilteredBones, StaticBuffers->GetNumUnfilteredBones());
 			SetShaderValue(RHICmdList, ComputeShaderRHI, RandomMaxBone, StaticBuffers->GetExcludedBoneIndex() >= 0 ? DynamicBuffers->GetNumBones() - 2 : DynamicBuffers->GetNumBones() - 1);
@@ -2402,24 +2374,24 @@ bool FNDISkeletalMesh_InstanceData::HasColorData()
 
 void FNDISkeletalMesh_InstanceData::Release()
 {
-	if (MeshGpuSpawnStaticBuffers)
+	if (MeshGpuSpawnStaticBuffers || MeshGpuSpawnDynamicBuffers)
 	{
-		BeginReleaseResource(MeshGpuSpawnStaticBuffers);
 		ENQUEUE_RENDER_COMMAND(DeleteResource)(
-			[ParamPointerToRelease = MeshGpuSpawnStaticBuffers](FRHICommandListImmediate& RHICmdList)
+			[RT_StaticBuffers=MeshGpuSpawnStaticBuffers, RT_DynamicBuffers=MeshGpuSpawnDynamicBuffers](FRHICommandListImmediate& RHICmdList)
 			{
-				delete ParamPointerToRelease;
-			});
+				if (RT_StaticBuffers)
+				{
+					RT_StaticBuffers->ReleaseResource();
+					delete RT_StaticBuffers;
+				}
+				if (RT_DynamicBuffers)
+				{
+					RT_DynamicBuffers->ReleaseResource();
+					delete RT_DynamicBuffers;
+				}
+			}
+		);
 		MeshGpuSpawnStaticBuffers = nullptr;
-	}
-	if (MeshGpuSpawnDynamicBuffers)
-	{
-		BeginReleaseResource(MeshGpuSpawnDynamicBuffers);
-		ENQUEUE_RENDER_COMMAND(DeleteResource)(
-			[ParamPointerToRelease = MeshGpuSpawnDynamicBuffers](FRHICommandListImmediate& RHICmdList)
-			{
-				delete ParamPointerToRelease;
-			});
 		MeshGpuSpawnDynamicBuffers = nullptr;
 	}
 }
@@ -2660,20 +2632,15 @@ void UNiagaraDataInterfaceSkeletalMesh::DestroyPerInstanceData(void* PerInstance
 	}
 #endif
 
+	ENQUEUE_RENDER_COMMAND(FNiagaraDestroySkeletalMeshInstanceData) (
+		[RT_Proxy=GetProxyAs<FNiagaraDataInterfaceProxySkeletalMesh>(), InstanceID=SystemInstance->GetId()](FRHICommandListImmediate& CmdList)
+		{
+			RT_Proxy->SystemInstancesToData.Remove(InstanceID);
+		}
+	);
+
 	Inst->Release();
 	Inst->~FNDISkeletalMesh_InstanceData();
-
-	{
-		// @todo this races
-		FNiagaraDataInterfaceProxySkeletalMesh* ThisProxy = GetProxyAs<FNiagaraDataInterfaceProxySkeletalMesh>();
-		ENQUEUE_RENDER_COMMAND(FNiagaraDestroySkeletalMeshInstanceData) (
-			[ThisProxy, InstanceID = SystemInstance->GetId()](FRHICommandListImmediate& CmdList)
-			{
-				//check(ThisProxy->Contains(InstanceID));
-				ThisProxy->SystemInstancesToData.Remove(InstanceID);
-			}
-		);
-	}
 }
 
 bool UNiagaraDataInterfaceSkeletalMesh::PerInstanceTick(void* PerInstanceData, FNiagaraSystemInstance* SystemInstance, float InDeltaSeconds)
diff --git a/Engine/Plugins/FX/Niagara/Source/Niagara/Private/NiagaraEmitterInstanceBatcher.cpp b/Engine/Plugins/FX/Niagara/Source/Niagara/Private/NiagaraEmitterInstanceBatcher.cpp
--- a/Engine/Plugins/FX/Niagara/Source/Niagara/Private/NiagaraEmitterInstanceBatcher.cpp
+++ b/Engine/Plugins/FX/Niagara/Source/Niagara/Private/NiagaraEmitterInstanceBatcher.cpp
@@ -883,11 +883,11 @@ void NiagaraEmitterInstanceBatcher::PrepareTicksForProxy(FRHICommandListImmediat
 		bool bHasFreeIDUpdates = false;
 
 		// Track that we need to consume per instance data before executing the ticks
-		if (Tick.DIInstanceData)
-		{
-			GpuDispatchList.PreAllocateGroups(iTickStartDispatchGroup + 1);
-			GpuDispatchList.DispatchGroups[iTickStartDispatchGroup].TicksWithPerInstanceData.Add(&Tick);
-		}
+		//if (Tick.DIInstanceData)
+		//{
+		//	GpuDispatchList.PreAllocateGroups(iTickStartDispatchGroup + 1);
+		//	GpuDispatchList.DispatchGroups[iTickStartDispatchGroup].TicksWithPerInstanceData.Add(&Tick);
+		//}
 
 		// Iterate over all instances preparing our number of instances
 		for (FNiagaraComputeInstanceData& InstanceData : Tick.GetInstances())
@@ -1207,17 +1207,17 @@ void NiagaraEmitterInstanceBatcher::ExecuteTicks(FRHICommandList& RHICmdList, FR
 		}
 
 		// Consume per tick data from the game thread
-		for (FNiagaraGPUSystemTick* Tick : DispatchGroup.TicksWithPerInstanceData)
-		{
-			uint8* BasePointer = reinterpret_cast<uint8*>(Tick->DIInstanceData->PerInstanceDataForRT);
-
-			for (const auto& Pair : Tick->DIInstanceData->InterfaceProxiesToOffsets)
-			{
-				FNiagaraDataInterfaceProxy* Proxy = Pair.Key;
-				uint8* InstanceDataPtr = BasePointer + Pair.Value;
-				Proxy->ConsumePerInstanceDataFromGameThread(InstanceDataPtr, Tick->SystemInstanceID);
-			}
-		}
+		//for (FNiagaraGPUSystemTick* Tick : DispatchGroup.TicksWithPerInstanceData)
+		//{
+		//	uint8* BasePointer = reinterpret_cast<uint8*>(Tick->DIInstanceData->PerInstanceDataForRT);
+
+		//	for (const auto& Pair : Tick->DIInstanceData->InterfaceProxiesToOffsets)
+		//	{
+		//		FNiagaraDataInterfaceProxy* Proxy = Pair.Key;
+		//		uint8* InstanceDataPtr = BasePointer + Pair.Value;
+		//		Proxy->ConsumePerInstanceDataFromGameThread(InstanceDataPtr, Tick->SystemInstanceID);
+		//	}
+		//}
 
 		// Execute Before Transitions
 		RHICmdList.Transition(TransitionsBefore);
diff --git a/Engine/Plugins/FX/Niagara/Source/Niagara/Private/NiagaraSystemGpuComputeProxy.cpp b/Engine/Plugins/FX/Niagara/Source/Niagara/Private/NiagaraSystemGpuComputeProxy.cpp
--- a/Engine/Plugins/FX/Niagara/Source/Niagara/Private/NiagaraSystemGpuComputeProxy.cpp
+++ b/Engine/Plugins/FX/Niagara/Source/Niagara/Private/NiagaraSystemGpuComputeProxy.cpp
@@ -124,6 +124,21 @@ void FNiagaraSystemGpuComputeProxy::QueueTick(const FNiagaraGPUSystemTick& Tick)
 
 	//-OPT: Making a copy of the tick here, reduce this.
 	PendingTicks.Add(Tick);
+
+	// Consume DataInterface instance data
+	//-TODO: This should be consumed as the command in executed rather than here, otherwise ticks are not processed correctly
+	//       However an audit of all data interfaces is required to ensure they pass data safely, for example the skeletal mesh one does not
+	if (Tick.DIInstanceData)
+	{
+		uint8* BasePointer = (uint8*)Tick.DIInstanceData->PerInstanceDataForRT;
+
+		for (auto& Pair : Tick.DIInstanceData->InterfaceProxiesToOffsets)
+		{
+			FNiagaraDataInterfaceProxy* Proxy = Pair.Key;
+			uint8* InstanceDataPtr = BasePointer + Pair.Value;
+			Proxy->ConsumePerInstanceDataFromGameThread(InstanceDataPtr, Tick.SystemInstanceID);
+		}
+	}
 }
 
 void FNiagaraSystemGpuComputeProxy::ReleaseTicks(FNiagaraGPUInstanceCountManager& GPUInstanceCountManager)
