Merging using //UE5/Main_To_//UE5/Private-Frosty
Fix some missing delegate unbinds when cleaning up parameter panel and parameter definition panel viewmodels.
Parameter definitions panel viewmodel now has a weak ptr to the parameter definitions asset to prevent a crash on shutdown where the asset is removed before the viewmodel and owning toolkit.

diff --git a/Engine/Plugins/FX/Niagara/Source/NiagaraEditor/Private/Toolkits/NiagaraParameterDefinitionsToolkit.cpp b/Engine/Plugins/FX/Niagara/Source/NiagaraEditor/Private/Toolkits/NiagaraParameterDefinitionsToolkit.cpp
--- a/Engine/Plugins/FX/Niagara/Source/NiagaraEditor/Private/Toolkits/NiagaraParameterDefinitionsToolkit.cpp
+++ b/Engine/Plugins/FX/Niagara/Source/NiagaraEditor/Private/Toolkits/NiagaraParameterDefinitionsToolkit.cpp
@@ -30,7 +30,12 @@ FNiagaraParameterDefinitionsToolkit::FNiagaraParameterDefinitionsToolkit()
 {}
 
 FNiagaraParameterDefinitionsToolkit::~FNiagaraParameterDefinitionsToolkit()
-{}
+{
+	if (ParameterPanelViewModel.IsValid())
+	{
+		ParameterPanelViewModel->Cleanup();
+	}
+}
 
 void FNiagaraParameterDefinitionsToolkit::RegisterTabSpawners(const TSharedRef<class FTabManager>& InTabManager)
 {
diff --git a/Engine/Plugins/FX/Niagara/Source/NiagaraEditor/Private/ViewModels/NiagaraParameterPanelViewModel.cpp b/Engine/Plugins/FX/Niagara/Source/NiagaraEditor/Private/ViewModels/NiagaraParameterPanelViewModel.cpp
--- a/Engine/Plugins/FX/Niagara/Source/NiagaraEditor/Private/ViewModels/NiagaraParameterPanelViewModel.cpp
+++ b/Engine/Plugins/FX/Niagara/Source/NiagaraEditor/Private/ViewModels/NiagaraParameterPanelViewModel.cpp
@@ -648,8 +648,54 @@ FNiagaraSystemToolkitParameterPanelViewModel::FNiagaraSystemToolkitParameterPane
 
 void FNiagaraSystemToolkitParameterPanelViewModel::Cleanup()
 {
+	UNiagaraSystem& System = SystemViewModel->GetSystem();
+
+	auto GetGraphFromScript = [](UNiagaraScript* Script)->UNiagaraGraph* {
+		return CastChecked<UNiagaraScriptSource>(Script->GetLatestSource())->NodeGraph;
+	};
+
+	TArray<UNiagaraGraph*> GraphsToRemoveCallbacks;
+	GraphsToRemoveCallbacks.Add(GetGraphFromScript(System.GetSystemSpawnScript()));
+	GraphsToRemoveCallbacks.Add(GetGraphFromScript(System.GetSystemUpdateScript()));
+	for (TSharedRef<FNiagaraEmitterHandleViewModel> EmitterHandleViewModel : SystemViewModel->GetEmitterHandleViewModels())
+	{
+		if (FNiagaraEmitterHandle* EmitterHandle = EmitterHandleViewModel->GetEmitterHandle())
+		{
+			TArray<UNiagaraScript*> EmitterScripts;
+			const bool bCompilableOnly = false;
+			EmitterHandle->GetInstance()->GetScripts(EmitterScripts, bCompilableOnly);
+			for (UNiagaraScript* EmitterScript : EmitterScripts)
+			{
+				GraphsToRemoveCallbacks.Add(GetGraphFromScript(EmitterScript));
+			}
+		}
+	}
+
+	for (UNiagaraGraph* Graph : GraphsToRemoveCallbacks)
+	{
+		if (FDelegateHandle* OnGraphChangedHandle = GraphIdToOnGraphChangedHandleMap.Find(Graph->GetUniqueID()))
+		{
+			Graph->RemoveOnGraphChangedHandler(*OnGraphChangedHandle);
+		}
+	}
+
+	if (SystemViewModel->GetEditMode() == ENiagaraSystemViewModelEditMode::SystemAsset && SystemGraphSelectionViewModelWeak.IsValid())
+	{
+		SystemGraphSelectionViewModelWeak.Pin()->GetOnSelectedEmitterScriptGraphsRefreshedDelegate().RemoveAll(this);
+	}
+
+	if (SystemViewModel->GetEditMode() == ENiagaraSystemViewModelEditMode::SystemAsset)
+	{
+		System.GetExposedParameters().RemoveOnChangedHandler(UserParameterStoreChangedHandle);
+	}
+
+	SystemViewModel->OnSystemCompiled().RemoveAll(this);
+	SystemViewModel->OnParameterRemovedExternally().RemoveAll(this);
+	SystemViewModel->OnParameterRenamedExternally().RemoveAll(this);
+
 	SystemViewModel->GetOnSubscribedParameterDefinitionsChangedDelegate().RemoveAll(this);
-	SystemViewModel->GetSystem().GetExposedParameters().RemoveOnChangedHandler(UserParameterStoreChangedHandle);
+	SystemViewModel->OnEmitterHandleViewModelsChanged().RemoveAll(this);
+
 }
 
 void FNiagaraSystemToolkitParameterPanelViewModel::Init(const FSystemToolkitUIContext& InUIContext)
@@ -684,6 +730,7 @@ void FNiagaraSystemToolkitParameterPanelViewModel::Init(const FSystemToolkitUICo
 	{
 		FDelegateHandle OnGraphChangedHandle = Graph->AddOnGraphChangedHandler(
 			FOnGraphChanged::FDelegate::CreateRaw(this, &FNiagaraSystemToolkitParameterPanelViewModel::OnGraphChanged));
+		GraphIdToOnGraphChangedHandleMap.Add(Graph->GetUniqueID(), OnGraphChangedHandle);
 	}
 
 	// Bind OnChanged() callback for system edit mode active emitter handle selection changing.
@@ -707,6 +754,9 @@ void FNiagaraSystemToolkitParameterPanelViewModel::Init(const FSystemToolkitUICo
 	// Bind OnChanged() bindings for parameter definitions.
 	SystemViewModel->GetOnSubscribedParameterDefinitionsChangedDelegate().AddSP(this, &FNiagaraSystemToolkitParameterPanelViewModel::Refresh);
 
+	// Bind OnChanged() bindings for emitter handles changing.
+	SystemViewModel->OnEmitterHandleViewModelsChanged().AddSP(this, &FNiagaraSystemToolkitParameterPanelViewModel::ReconcileOnGraphChangedBindings);
+
 	// Init default categories
 	if (DefaultCategories.Num() == 0)
 	{
@@ -1681,6 +1731,51 @@ void FNiagaraSystemToolkitParameterPanelViewModel::OnParameterRemovedExternally(
 	Refresh();
 }
 
+void FNiagaraSystemToolkitParameterPanelViewModel::ReconcileOnGraphChangedBindings()
+{
+	UNiagaraSystem& System = SystemViewModel->GetSystem();
+
+	auto GetGraphFromScript = [](UNiagaraScript* Script)->UNiagaraGraph* {
+		return CastChecked<UNiagaraScriptSource>(Script->GetLatestSource())->NodeGraph;
+	};
+
+	TSet<uint32> UnvisitedGraphIds;
+	GraphIdToOnGraphChangedHandleMap.GetKeys(UnvisitedGraphIds);
+	UnvisitedGraphIds.Remove(GetGraphFromScript(System.GetSystemSpawnScript())->GetUniqueID());
+	UnvisitedGraphIds.Remove(GetGraphFromScript(System.GetSystemUpdateScript())->GetUniqueID());
+
+	TArray<UNiagaraGraph*> GraphsToAddCallbacks;
+	for (TSharedRef<FNiagaraEmitterHandleViewModel> EmitterHandleViewModel : SystemViewModel->GetEmitterHandleViewModels())
+	{
+		if (FNiagaraEmitterHandle* EmitterHandle = EmitterHandleViewModel->GetEmitterHandle())
+		{
+			TArray<UNiagaraScript*> EmitterScripts;
+			const bool bCompilableOnly = false;
+			EmitterHandle->GetInstance()->GetScripts(EmitterScripts, bCompilableOnly);
+			for (UNiagaraScript* EmitterScript : EmitterScripts)
+			{
+				GraphsToAddCallbacks.Add(GetGraphFromScript(EmitterScript));
+			}
+		}
+	}
+
+	for (UNiagaraGraph* Graph : GraphsToAddCallbacks)
+	{
+		UnvisitedGraphIds.Remove(Graph->GetUniqueID());
+		if (GraphIdToOnGraphChangedHandleMap.Find(Graph->GetUniqueID()) == nullptr)
+		{
+			FDelegateHandle OnGraphChangedHandle = Graph->AddOnGraphChangedHandler(
+				FOnGraphChanged::FDelegate::CreateRaw(this, &FNiagaraSystemToolkitParameterPanelViewModel::OnGraphChanged));
+			GraphIdToOnGraphChangedHandleMap.Add(Graph->GetUniqueID(), OnGraphChangedHandle);
+		}
+	}
+
+	for (const uint32& UnvisitedGraphId : UnvisitedGraphIds)
+	{
+		GraphIdToOnGraphChangedHandleMap.Remove(UnvisitedGraphId);
+	}
+}
+
 ///////////////////////////////////////////////////////////////////////////////
 /// Script Toolkit Parameter Panel View Model								///
 ///////////////////////////////////////////////////////////////////////////////
@@ -2305,25 +2400,20 @@ void FNiagaraScriptToolkitParameterPanelViewModel::OnGraphSubObjectSelectionChan
 }
 
 ///////////////////////////////////////////////////////////////////////////////
-/// Parameter Definitions Toolkit Parameter Panel View Model					///
+/// Parameter Definitions Toolkit Parameter Panel View Model				///
 ///////////////////////////////////////////////////////////////////////////////
 
 FNiagaraParameterDefinitionsToolkitParameterPanelViewModel::FNiagaraParameterDefinitionsToolkitParameterPanelViewModel(UNiagaraParameterDefinitions* InParameterDefinitions, const TSharedPtr<FNiagaraObjectSelection>& InObjectSelection)
 {
-	ParameterDefinitions = InParameterDefinitions;
+	ParameterDefinitionsWeak = InParameterDefinitions;
 	VariableObjectSelection = InObjectSelection;
 }
 
-FNiagaraParameterDefinitionsToolkitParameterPanelViewModel::~FNiagaraParameterDefinitionsToolkitParameterPanelViewModel()
-{
-	//ParameterDefinitions->GetOnParameterDefinitionsChanged().RemoveAll(this); //@todo(ng) invoke this before editor shutdown
-}
-
 void FNiagaraParameterDefinitionsToolkitParameterPanelViewModel::Init(const FParameterDefinitionsToolkitUIContext& InUIContext)
 {
 	UIContext = InUIContext;
 
-	ParameterDefinitions->GetOnParameterDefinitionsChanged().AddSP(this, &FNiagaraParameterDefinitionsToolkitParameterPanelViewModel::Refresh);
+	ParameterDefinitionsWeak.Get()->GetOnParameterDefinitionsChanged().AddSP(this, &FNiagaraParameterDefinitionsToolkitParameterPanelViewModel::Refresh);
 
 	// Init default categories
 	if (DefaultCategories.Num() == 0)
@@ -2334,13 +2424,21 @@ void FNiagaraParameterDefinitionsToolkitParameterPanelViewModel::Init(const FPar
 	}
 }
 
+void FNiagaraParameterDefinitionsToolkitParameterPanelViewModel::Cleanup()
+{
+	if (ParameterDefinitionsWeak.IsValid())
+	{
+		ParameterDefinitionsWeak.Get()->GetOnParameterDefinitionsChanged().RemoveAll(this);
+	}
+}
+
 void FNiagaraParameterDefinitionsToolkitParameterPanelViewModel::AddParameter(FNiagaraVariable NewVariable, const FNiagaraParameterPanelCategory Category, const bool bRequestRename) const
 {
 	TGuardValue<bool> AddParameterRefreshGuard(bIsAddingParameter, true);
 	bool bSuccess = false;
 
 	TSet<FName> Names;
-	for (const UNiagaraScriptVariable* ScriptVar : ParameterDefinitions->GetParametersConst())
+	for (const UNiagaraScriptVariable* ScriptVar : ParameterDefinitionsWeak.Get()->GetParametersConst())
 	{
 		Names.Add(ScriptVar->Variable.GetName());
 	}
@@ -2348,7 +2446,7 @@ void FNiagaraParameterDefinitionsToolkitParameterPanelViewModel::AddParameter(FN
 	NewVariable.SetName(NewUniqueName);
 
 	FScopedTransaction AddTransaction(LOCTEXT("ParameterDefinitionsToolkitParameterPanelViewModel_AddParameter", "Add Parameter"));
-	ParameterDefinitions->AddParameter(NewVariable);
+	ParameterDefinitionsWeak.Get()->AddParameter(NewVariable);
 	SelectParameterItemByName(NewVariable.GetName(), bRequestRename);
 }
 
@@ -2360,14 +2458,14 @@ bool FNiagaraParameterDefinitionsToolkitParameterPanelViewModel::GetCanAddParame
 void FNiagaraParameterDefinitionsToolkitParameterPanelViewModel::DeleteParameter(const FNiagaraParameterPanelItem& ItemToDelete) const
 {
 	FScopedTransaction RemoveParameter(LOCTEXT("ParameterDefinitionsToolkitParameterPanelViewModel_RemoveParameter", "Remove Parameter"));
-	ParameterDefinitions->RemoveParameter(ItemToDelete.GetVariable());
+	ParameterDefinitionsWeak.Get()->RemoveParameter(ItemToDelete.GetVariable());
 	UIContext.RefreshSelectionDetailsViewPanel();
 }
 
 void FNiagaraParameterDefinitionsToolkitParameterPanelViewModel::RenameParameter(const FNiagaraParameterPanelItem& ItemToRename, const FName NewName) const
 {
 	FScopedTransaction RenameParameter(LOCTEXT("ParameterDefinitionsToolkitParameterPanelViewModel_RenameParameter", "Rename Parameter"));
-	ParameterDefinitions->RenameParameter(ItemToRename.GetVariable(), NewName);
+	ParameterDefinitionsWeak.Get()->RenameParameter(ItemToRename.GetVariable(), NewName);
 }
 
 void FNiagaraParameterDefinitionsToolkitParameterPanelViewModel::SetParameterIsSubscribedToLibrary(const FNiagaraParameterPanelItem ItemToModify, const bool bSubscribed) const
@@ -2450,7 +2548,7 @@ FNiagaraParameterUtilities::EParameterContext FNiagaraParameterDefinitionsToolki
 
 const TArray<UNiagaraScriptVariable*> FNiagaraParameterDefinitionsToolkitParameterPanelViewModel::GetEditableScriptVariablesWithName(const FName ParameterName) const
 {
-	const TArray<UNiagaraScriptVariable*>& LibraryParameters = ParameterDefinitions->GetParametersConst();
+	const TArray<UNiagaraScriptVariable*>& LibraryParameters = ParameterDefinitionsWeak.Get()->GetParametersConst();
 	return LibraryParameters.FilterByPredicate([ParameterName](const UNiagaraScriptVariable* ScriptVariable){return ScriptVariable->Variable.GetName() == ParameterName;});
 }
 
@@ -2474,7 +2572,7 @@ TArray<FNiagaraVariable> FNiagaraParameterDefinitionsToolkitParameterPanelViewMo
 TArray<FNiagaraParameterPanelItem> FNiagaraParameterDefinitionsToolkitParameterPanelViewModel::GetViewedParameterItems() const
 {
 	CachedViewedItems.Reset();
-	for (const UNiagaraScriptVariable* const& ScriptVar : ParameterDefinitions->GetParametersConst())
+	for (const UNiagaraScriptVariable* const& ScriptVar : ParameterDefinitionsWeak.Get()->GetParametersConst())
 	{
 		const FNiagaraNamespaceMetadata NamespaceMetaData = FNiagaraEditorUtilities::GetNamespaceMetaDataForVariableName(ScriptVar->Variable.GetName());
 		const bool bExternallyReferenced = false;
@@ -2555,7 +2653,7 @@ bool FNiagaraParameterDefinitionsToolkitParameterPanelViewModel::GetCanHandleDra
 
 void FNiagaraParameterDefinitionsToolkitParameterPanelViewModel::OnParameterItemSelected(const FNiagaraParameterPanelItem& SelectedItem, ESelectInfo::Type SelectInfo) const
 {
-	if (UNiagaraScriptVariable* SelectedScriptVariable = ParameterDefinitions->GetScriptVariable(SelectedItem.GetVariable()))
+	if (UNiagaraScriptVariable* SelectedScriptVariable = ParameterDefinitionsWeak.Get()->GetScriptVariable(SelectedItem.GetVariable()))
 	{
 		VariableObjectSelection->SetSelectedObject(SelectedScriptVariable);
 	}
diff --git a/Engine/Plugins/FX/Niagara/Source/NiagaraEditor/Public/ViewModels/NiagaraParameterPanelViewModel.h b/Engine/Plugins/FX/Niagara/Source/NiagaraEditor/Public/ViewModels/NiagaraParameterPanelViewModel.h
--- a/Engine/Plugins/FX/Niagara/Source/NiagaraEditor/Public/ViewModels/NiagaraParameterPanelViewModel.h
+++ b/Engine/Plugins/FX/Niagara/Source/NiagaraEditor/Public/ViewModels/NiagaraParameterPanelViewModel.h
@@ -309,6 +309,8 @@ private:
 	void OnParameterRenamedExternally(const FNiagaraVariableBase& InOldVar, const FNiagaraVariableBase& InNewVar, UNiagaraEmitter* InOptionalEmitter);
 	void OnParameterRemovedExternally(const FNiagaraVariableBase& InOldVar, UNiagaraEmitter* InOptionalEmitter); 
 
+	void ReconcileOnGraphChangedBindings();
+
 private:
 	// Graphs viewed to gather UNiagaraScriptVariables that are displayed by the Parameter Panel.
 	TWeakObjectPtr<UNiagaraGraph> SystemScriptGraph;
@@ -317,6 +319,7 @@ private:
 	TWeakPtr<FNiagaraSystemGraphSelectionViewModel> SystemGraphSelectionViewModelWeak;
 
 	FDelegateHandle UserParameterStoreChangedHandle;
+	TMap<uint32, FDelegateHandle> GraphIdToOnGraphChangedHandleMap;
 
 	mutable FSystemToolkitUIContext UIContext;
 
@@ -415,10 +418,10 @@ public:
 	/** Construct a ParameterDefinitionsToolkit Parameter Panel View Model from a Parameter Definitions. */
 	FNiagaraParameterDefinitionsToolkitParameterPanelViewModel(UNiagaraParameterDefinitions* InParameterDefinitions, const TSharedPtr<FNiagaraObjectSelection>& InObjectSelection);
 
-	~FNiagaraParameterDefinitionsToolkitParameterPanelViewModel();
-
 	void Init(const FParameterDefinitionsToolkitUIContext& InUIContext);
 
+	void Cleanup();
+
 	//~ Begin INiagaraImmutableParameterPanelViewModel interface
 	//NOTE: The ParameterDefinitionsToolkitParameterPanelViewModel does not edit any graphs, so return an empty array.
 	virtual const TArray<UNiagaraGraph*> GetEditableGraphsConst() const override { return TArray<UNiagaraGraph*>(); };
@@ -467,7 +470,7 @@ public:
 	void AddParameterFromMenu(FNiagaraVariable NewVariable, const FNiagaraParameterPanelCategory Category, const bool bRequestRename) const;
 
 private:
-	UNiagaraParameterDefinitions* ParameterDefinitions;
+	TWeakObjectPtr<UNiagaraParameterDefinitions> ParameterDefinitionsWeak;
 	mutable FParameterDefinitionsToolkitUIContext UIContext;
 
 	TSharedPtr<FNiagaraObjectSelection> VariableObjectSelection;
