- Fix broken velocity texture management

#rnx

diff --git a/Engine/Plugins/FX/Niagara/Source/Niagara/Classes/NiagaraEmitterInstanceBatcher.h b/Engine/Plugins/FX/Niagara/Source/Niagara/Classes/NiagaraEmitterInstanceBatcher.h
--- a/Engine/Plugins/FX/Niagara/Source/Niagara/Classes/NiagaraEmitterInstanceBatcher.h
+++ b/Engine/Plugins/FX/Niagara/Source/Niagara/Classes/NiagaraEmitterInstanceBatcher.h
@@ -145,6 +145,9 @@ public:
 	/** Get the shared SortManager, used in the rendering loop to call FGPUSortManager::OnPreRender() and FGPUSortManager::OnPostRenderOpaque() */
 	virtual FGPUSortManager* GetGPUSortManager() const override;
 
+	/** When inside a pass returns the active scene textures for the pass, can be nullptr. */
+	class FNiagaraSceneTextureParameters* GetNiagaraSceneTextures() const { return NiagaraSceneTextures; }
+
 	/** Debug only function to readback data. */
 	void AddDebugReadback(FNiagaraSystemInstanceID InstanceID, TSharedPtr<struct FNiagaraScriptDebuggerInfo, ESPMode::ThreadSafe> DebugInfo, FNiagaraComputeExecutionContext* Context);
 
@@ -272,6 +275,8 @@ private:
 
 	FNiagaraGpuDispatchList DispatchListPerStage[ENiagaraGpuComputeTickStage::Max];
 
+	class FNiagaraSceneTextureParameters* NiagaraSceneTextures = nullptr;
+
 #if NIAGARA_COMPUTEDEBUG_ENABLED
 	TUniquePtr<FNiagaraGpuComputeDebug> GpuComputeDebugPtr;
 #endif
diff --git a/Engine/Plugins/FX/Niagara/Source/Niagara/Private/DataInterface/NiagaraDataInterfaceGBuffer.cpp b/Engine/Plugins/FX/Niagara/Source/Niagara/Private/DataInterface/NiagaraDataInterfaceGBuffer.cpp
--- a/Engine/Plugins/FX/Niagara/Source/Niagara/Private/DataInterface/NiagaraDataInterfaceGBuffer.cpp
+++ b/Engine/Plugins/FX/Niagara/Source/Niagara/Private/DataInterface/NiagaraDataInterfaceGBuffer.cpp
@@ -1,7 +1,9 @@
 // Copyright Epic Games, Inc. All Rights Reserved.
 
 #include "NiagaraDataInterfaceGBuffer.h"
+#include "NiagaraEmitterInstanceBatcher.h"
 #include "NiagaraTypes.h"
+#include "NiagaraRenderViewDataManager.h"
 #include "NiagaraWorldManager.h"
 #include "ShaderParameterUtils.h"
 #include "Internationalization/Internationalization.h"
@@ -108,8 +110,15 @@ public:
 		FRHIComputeShader* ComputeShaderRHI = RHICmdList.GetBoundComputeShader();
 
 		FRHISamplerState* VelocitySamplerState = TStaticSamplerState<SF_Bilinear, AM_Clamp, AM_Clamp, AM_Clamp>::GetRHI();
-		FRHITexture* VelocityTexture = GNiagaraViewDataManager.GetSceneVelocityTexture() ? GNiagaraViewDataManager.GetSceneVelocityTexture() : GBlackTexture->TextureRHI;
-		SetTextureParameter(RHICmdList, ComputeShaderRHI, VelocityTextureParam, VelocityTextureSamplerParam, VelocitySamplerState, VelocityTexture);
+		FRHITexture* VelocityRHITexture = GBlackTexture->TextureRHI;
+		if ( FNiagaraSceneTextureParameters* NiagaraSceneTextures = Context.Batcher->GetNiagaraSceneTextures() )
+		{
+			if ( FRDGTexture* VelocityRDGTexture = NiagaraSceneTextures->Velocity.GetTexture() )
+			{
+				VelocityRHITexture = VelocityRDGTexture->GetRHI();
+			}
+		}
+		SetTextureParameter(RHICmdList, ComputeShaderRHI, VelocityTextureParam, VelocityTextureSamplerParam, VelocitySamplerState, VelocityRHITexture);
 	}
 
 private:
diff --git a/Engine/Plugins/FX/Niagara/Source/Niagara/Private/NiagaraEmitterInstanceBatcher.cpp b/Engine/Plugins/FX/Niagara/Source/Niagara/Private/NiagaraEmitterInstanceBatcher.cpp
--- a/Engine/Plugins/FX/Niagara/Source/Niagara/Private/NiagaraEmitterInstanceBatcher.cpp
+++ b/Engine/Plugins/FX/Niagara/Source/Niagara/Private/NiagaraEmitterInstanceBatcher.cpp
@@ -21,6 +21,7 @@
 #include "NiagaraShaderParticleID.h"
 #include "NiagaraSortingGPU.h"
 #include "NiagaraStats.h"
+#include "NiagaraRenderViewDataManager.h"
 #include "NiagaraWorldManager.h"
 #include "RHI.h"
 #include "RHIGPUReadback.h"
@@ -1581,9 +1582,11 @@ void NiagaraEmitterInstanceBatcher::PreInitViews(FRDGBuilder& GraphBuilder, bool
 			RDG_EVENT_NAME("Niagara::PreInitViews"),
 			PassParameters,
 			ERDGPassFlags::Compute | ERDGPassFlags::NeverCull,
-			[this](FRHICommandListImmediate& RHICmdList)
+			[this, PassParameters](FRHICommandListImmediate& RHICmdList)
 		{
-			FNiagaraSceneTextureScope Scope;
+			NiagaraSceneTextures = PassParameters;
+			ON_SCOPE_EXIT { NiagaraSceneTextures = nullptr; };
+
 			ExecuteTicks(RHICmdList, nullptr, ENiagaraGpuComputeTickStage::PreInitViews);
 		});
 	}
@@ -1609,10 +1612,12 @@ void NiagaraEmitterInstanceBatcher::PostInitViews(FRDGBuilder& GraphBuilder, TAr
 			RDG_EVENT_NAME("Niagara::PostInitViews"),
 			PassParameters,
 			ERDGPassFlags::Compute | ERDGPassFlags::NeverCull,
-			[this, Views](FRHICommandListImmediate& RHICmdList)
+			[this, PassParameters, Views](FRHICommandListImmediate& RHICmdList)
 		{
+			NiagaraSceneTextures = PassParameters;
+			ON_SCOPE_EXIT{ NiagaraSceneTextures = nullptr; };
+
 			FRHIUniformBuffer* ViewUniformBuffer = GetReferenceViewUniformBuffer(Views);
-			FNiagaraSceneTextureScope Scope;
 			ExecuteTicks(RHICmdList, ViewUniformBuffer, ENiagaraGpuComputeTickStage::PostInitViews);
 		});
 	}
@@ -1631,12 +1636,13 @@ void NiagaraEmitterInstanceBatcher::PostRenderOpaque(FRDGBuilder& GraphBuilder,
 		RDG_EVENT_NAME("Niagara::PostRenderOpaque"),
 		PassParameters,
 		ERDGPassFlags::Compute | ERDGPassFlags::NeverCull,
-		[this, Views, bAllowGPUParticleUpdate](FRHICommandListImmediate& RHICmdList)
+		[this, PassParameters, Views, bAllowGPUParticleUpdate](FRHICommandListImmediate& RHICmdList)
 	{
-		FNiagaraSceneTextureScope Scope;
-
 		if (bAllowGPUParticleUpdate && FNiagaraUtilities::AllowGPUParticles(GetShaderPlatform()))
 		{
+			NiagaraSceneTextures = PassParameters;
+			ON_SCOPE_EXIT{ NiagaraSceneTextures = nullptr; };
+
 #if RHI_RAYTRACING
 			BuildRayTracingSceneInfo(RHICmdList, Views);
 			ON_SCOPE_EXIT
diff --git a/Engine/Plugins/FX/Niagara/Source/Niagara/Private/NiagaraModule.cpp b/Engine/Plugins/FX/Niagara/Source/Niagara/Private/NiagaraModule.cpp
--- a/Engine/Plugins/FX/Niagara/Source/Niagara/Private/NiagaraModule.cpp
+++ b/Engine/Plugins/FX/Niagara/Source/Niagara/Private/NiagaraModule.cpp
@@ -10,6 +10,7 @@
 #include "UObject/Class.h"
 #include "UObject/Package.h"
 #include "NiagaraWorldManager.h"
+#include "NiagaraRenderViewDataManager.h"
 #include "VectorVM.h"
 #include "NiagaraConstants.h"
 #include "NiagaraLightRendererProperties.h"
@@ -181,7 +182,7 @@ void INiagaraModule::StartupModule()
 {
 	VectorVM::Init();
 	FNiagaraTypeDefinition::Init();
-	FNiagaraViewDataMgr::Init();
+	FNiagaraRenderViewDataManager::Init();
 
 	FNiagaraWorldManager::OnStartup();
 
@@ -387,7 +388,7 @@ void INiagaraModule::ShutdownRenderingResources()
 {
 	FFXSystemInterface::UnregisterCustomFXSystem(NiagaraEmitterInstanceBatcher::Name);
 
-	FNiagaraViewDataMgr::Shutdown();
+	FNiagaraRenderViewDataManager::Shutdown();
 }
 
 void INiagaraModule::ShutdownModule()
diff --git a/Engine/Plugins/FX/Niagara/Source/Niagara/Private/NiagaraRenderViewDataManager.cpp b/Engine/Plugins/FX/Niagara/Source/Niagara/Private/NiagaraRenderViewDataManager.cpp
new file mode 100644
--- /dev/null
+++ b/Engine/Plugins/FX/Niagara/Source/Niagara/Private/NiagaraRenderViewDataManager.cpp
@@ -0,0 +1,78 @@
+// Copyright Epic Games, Inc. All Rights Reserved.
+
+#include "NiagaraRenderViewDataManager.h"
+
+#include "EngineModule.h"
+#include "SceneRenderTargetParameters.h"
+
+TGlobalResource<FNiagaraRenderViewDataManager> GNiagaraViewDataManager;
+
+void FNiagaraRenderViewDataManager::PostOpaqueRender(FPostOpaqueRenderParameters& Params)
+{
+	ViewUniformBuffer = Params.ViewUniformBuffer;
+	Parameters.SceneTextures = Params.SceneTexturesUniformParams;
+	Parameters.MobileSceneTextures = Params.MobileSceneTexturesUniformParams;
+	Parameters.Depth = Params.DepthTexture;
+	Parameters.Normal = Params.NormalTexture;
+	Parameters.Velocity = Params.VelocityTexture;
+}
+
+void FNiagaraRenderViewDataManager::GetSceneTextureParameters(FRDGBuilder& GraphBuilder, FNiagaraSceneTextureParameters& OutParameters) const
+{
+	OutParameters = Parameters;
+
+	ERHIFeatureLevel::Type LocalFeatureLevel = GetFeatureLevel();
+	if (FSceneInterface::GetShadingPath(LocalFeatureLevel) == EShadingPath::Deferred)
+	{
+		if ( !Parameters.SceneTextures )
+		{
+			OutParameters.SceneTextures = CreateSceneTextureUniformBuffer(GraphBuilder, LocalFeatureLevel, ESceneTextureSetupMode::SceneVelocity);
+		}
+	}
+	else if (FSceneInterface::GetShadingPath(LocalFeatureLevel) == EShadingPath::Mobile)
+	{
+		if ( !Parameters.MobileSceneTextures )
+		{
+			OutParameters.MobileSceneTextures = CreateMobileSceneTextureUniformBuffer(GraphBuilder, EMobileSceneTextureSetupMode::None);
+		}
+	}
+	OutParameters.Depth = Parameters.Depth;
+	OutParameters.Normal = Parameters.Normal;
+	OutParameters.Velocity = Parameters.Velocity;
+}
+
+void FNiagaraRenderViewDataManager::ClearSceneTextureParameters()
+{
+	Parameters = {};
+}
+
+FNiagaraRenderViewDataManager::FNiagaraRenderViewDataManager()
+	: FRenderResource()
+{
+
+}
+
+void FNiagaraRenderViewDataManager::Init()
+{
+	IRendererModule& RendererModule = GetRendererModule();
+
+	GNiagaraViewDataManager.PostOpaqueDelegate.BindRaw(&GNiagaraViewDataManager, &FNiagaraRenderViewDataManager::PostOpaqueRender);
+	GNiagaraViewDataManager.PostOpaqueDelegateHandle = RendererModule.RegisterPostOpaqueRenderDelegate(GNiagaraViewDataManager.PostOpaqueDelegate);
+}
+
+void FNiagaraRenderViewDataManager::Shutdown()
+{
+	IRendererModule& RendererModule = GetRendererModule();
+
+	RendererModule.RemovePostOpaqueRenderDelegate(GNiagaraViewDataManager.PostOpaqueDelegateHandle);
+	GNiagaraViewDataManager.ReleaseDynamicRHI();
+}
+
+void FNiagaraRenderViewDataManager::InitDynamicRHI()
+{
+
+}
+
+void FNiagaraRenderViewDataManager::ReleaseDynamicRHI()
+{
+}
diff --git a/Engine/Plugins/FX/Niagara/Source/Niagara/Private/NiagaraWorldManager.cpp b/Engine/Plugins/FX/Niagara/Source/Niagara/Private/NiagaraWorldManager.cpp
--- a/Engine/Plugins/FX/Niagara/Source/Niagara/Private/NiagaraWorldManager.cpp
+++ b/Engine/Plugins/FX/Niagara/Source/Niagara/Private/NiagaraWorldManager.cpp
@@ -183,8 +183,6 @@ FDelegateHandle FNiagaraWorldManager::PostGCHandle;
 FDelegateHandle FNiagaraWorldManager::PreGCBeginDestroyHandle; 
 TMap<class UWorld*, class FNiagaraWorldManager*> FNiagaraWorldManager::WorldManagers;
 
-TGlobalResource<FNiagaraViewDataMgr> GNiagaraViewDataManager;
-
 namespace FNiagaraUtilities
 {
 	int GetNiagaraTickGroup(ETickingGroup TickGroup)
@@ -194,85 +192,6 @@ namespace FNiagaraUtilities
 	}
 }
 
-void FNiagaraViewDataMgr::PostOpaqueRender(FPostOpaqueRenderParameters& Params)
-{
-	ViewUniformBuffer = Params.ViewUniformBuffer;
-	Parameters.SceneTextures = Params.SceneTexturesUniformParams;
-	Parameters.MobileSceneTextures = Params.MobileSceneTexturesUniformParams;
-	Parameters.Depth = Params.DepthTexture;
-	Parameters.Normal = Params.NormalTexture;
-	Parameters.Velocity = Params.VelocityTexture;
-}
-
-void FNiagaraViewDataMgr::GetSceneTextureParameters(FRDGBuilder& GraphBuilder, FNiagaraSceneTextureParameters& OutParameters) const
-{
-	OutParameters = Parameters;
-
-	ERHIFeatureLevel::Type LocalFeatureLevel = GetFeatureLevel();
-	if (FSceneInterface::GetShadingPath(LocalFeatureLevel) == EShadingPath::Deferred)
-	{
-		if ( !Parameters.SceneTextures )
-		{
-			OutParameters.SceneTextures = CreateSceneTextureUniformBuffer(GraphBuilder, LocalFeatureLevel, ESceneTextureSetupMode::None);
-		}
-	}
-	else if (FSceneInterface::GetShadingPath(LocalFeatureLevel) == EShadingPath::Mobile)
-	{
-		if ( !Parameters.MobileSceneTextures )
-		{
-			OutParameters.MobileSceneTextures = CreateMobileSceneTextureUniformBuffer(GraphBuilder, EMobileSceneTextureSetupMode::None);
-		}
-	}
-}
-
-void FNiagaraViewDataMgr::BeginPass()
-{
-	check(!bInsidePass);
-	bInsidePass = true;
-}
-
-void FNiagaraViewDataMgr::EndPass()
-{
-	check(bInsidePass);
-	bInsidePass = false;
-}
-
-void FNiagaraViewDataMgr::ClearSceneTextureParameters()
-{
-	Parameters = {};
-}
-
-FNiagaraViewDataMgr::FNiagaraViewDataMgr()
-	: FRenderResource()
-{
-
-}
-
-void FNiagaraViewDataMgr::Init()
-{
-	IRendererModule& RendererModule = GetRendererModule();
-
-	GNiagaraViewDataManager.PostOpaqueDelegate.BindRaw(&GNiagaraViewDataManager, &FNiagaraViewDataMgr::PostOpaqueRender);
-	GNiagaraViewDataManager.PostOpaqueDelegateHandle = RendererModule.RegisterPostOpaqueRenderDelegate(GNiagaraViewDataManager.PostOpaqueDelegate);
-}
-
-void FNiagaraViewDataMgr::Shutdown()
-{
-	IRendererModule& RendererModule = GetRendererModule();
-
-	RendererModule.RemovePostOpaqueRenderDelegate(GNiagaraViewDataManager.PostOpaqueDelegateHandle);
-	GNiagaraViewDataManager.ReleaseDynamicRHI();
-}
-
-void FNiagaraViewDataMgr::InitDynamicRHI()
-{
-
-}
-
-void FNiagaraViewDataMgr::ReleaseDynamicRHI()
-{
-}
-
 //////////////////////////////////////////////////////////////////////////
 
 void FNiagaraWorldManagerTickFunction::ExecuteTick(float DeltaTime, enum ELevelTick TickType, ENamedThreads::Type CurrentThread, const FGraphEventRef& MyCompletionGraphEvent)
diff --git a/Engine/Plugins/FX/Niagara/Source/Niagara/Public/NiagaraRenderViewDataManager.h b/Engine/Plugins/FX/Niagara/Source/Niagara/Public/NiagaraRenderViewDataManager.h
new file mode 100644
--- /dev/null
+++ b/Engine/Plugins/FX/Niagara/Source/Niagara/Public/NiagaraRenderViewDataManager.h
@@ -0,0 +1,46 @@
+// Copyright Epic Games, Inc. All Rights Reserved.
+
+#pragma once
+
+#include "CoreMinimal.h"
+#include "RendererInterface.h"
+#include "RenderResource.h"
+#include "ShaderParameterStruct.h"
+#include "ShaderParameterMacros.h"
+
+BEGIN_SHADER_PARAMETER_STRUCT(FNiagaraSceneTextureParameters, )
+	SHADER_PARAMETER_RDG_UNIFORM_BUFFER(FSceneTextureUniformParameters, SceneTextures)
+	SHADER_PARAMETER_RDG_UNIFORM_BUFFER(FMobileSceneTextureUniformParameters, MobileSceneTextures)
+	RDG_TEXTURE_ACCESS(Depth, ERHIAccess::SRVCompute)
+	RDG_TEXTURE_ACCESS(Normal, ERHIAccess::SRVCompute)
+	RDG_TEXTURE_ACCESS(Velocity, ERHIAccess::SRVCompute)
+END_SHADER_PARAMETER_STRUCT()
+
+class FNiagaraRenderViewDataManager : public FRenderResource
+{
+public:
+	FNiagaraRenderViewDataManager();
+
+	static void Init();
+	static void Shutdown();
+
+	NIAGARA_API void PostOpaqueRender(FPostOpaqueRenderParameters& Params);
+
+	void ClearSceneTextureParameters();
+
+	void GetSceneTextureParameters(FRDGBuilder& GraphBuilder, FNiagaraSceneTextureParameters& InParameters) const;
+
+	virtual void InitDynamicRHI() override;
+
+	virtual void ReleaseDynamicRHI() override;
+
+private:
+	FNiagaraSceneTextureParameters Parameters;
+	FRHIUniformBuffer* ViewUniformBuffer = nullptr;
+
+	FPostOpaqueRenderDelegate PostOpaqueDelegate;
+	FDelegateHandle PostOpaqueDelegateHandle;
+};
+
+extern TGlobalResource<FNiagaraRenderViewDataManager> GNiagaraViewDataManager;
+
diff --git a/Engine/Plugins/FX/Niagara/Source/Niagara/Public/NiagaraWorldManager.h b/Engine/Plugins/FX/Niagara/Source/Niagara/Public/NiagaraWorldManager.h
--- a/Engine/Plugins/FX/Niagara/Source/Niagara/Public/NiagaraWorldManager.h
+++ b/Engine/Plugins/FX/Niagara/Source/Niagara/Public/NiagaraWorldManager.h
@@ -30,64 +30,6 @@ class UNiagaraParameterCollectionInstance;
 class UNiagaraComponentPool;
 struct FNiagaraScalabilityState;
 
-BEGIN_SHADER_PARAMETER_STRUCT(FNiagaraSceneTextureParameters, )
-	SHADER_PARAMETER_RDG_UNIFORM_BUFFER(FSceneTextureUniformParameters, SceneTextures)
-	SHADER_PARAMETER_RDG_UNIFORM_BUFFER(FMobileSceneTextureUniformParameters, MobileSceneTextures)
-	RDG_TEXTURE_ACCESS(Depth, ERHIAccess::SRVCompute)
-	RDG_TEXTURE_ACCESS(Normal, ERHIAccess::SRVCompute)
-	RDG_TEXTURE_ACCESS(Velocity, ERHIAccess::SRVCompute)
-END_SHADER_PARAMETER_STRUCT()
-
-class FNiagaraViewDataMgr : public FRenderResource
-{
-public:
-	FNiagaraViewDataMgr();
-
-	static void Init();
-	static void Shutdown();
-
-	NIAGARA_API void PostOpaqueRender(FPostOpaqueRenderParameters& Params);
-
-	void BeginPass();
-	void EndPass();
-	void ClearSceneTextureParameters();
-
-	void GetSceneTextureParameters(FRDGBuilder& GraphBuilder, FNiagaraSceneTextureParameters& InParameters) const;
-
-	FRHITexture2D* GetSceneDepthTexture() { check(bInsidePass); return (FRHITexture2D*)TryGetRHI(Parameters.Depth); }
-	FRHITexture2D* GetSceneNormalTexture() { check(bInsidePass); return (FRHITexture2D*)TryGetRHI(Parameters.Normal); }
-	FRHITexture2D* GetSceneVelocityTexture() { check(bInsidePass); return (FRHITexture2D*)TryGetRHI(Parameters.Velocity); }
-	FRHIUniformBuffer* GetViewUniformBuffer() { return ViewUniformBuffer; }
-
-	virtual void InitDynamicRHI() override;
-
-	virtual void ReleaseDynamicRHI() override;
-
-private:
-	FNiagaraSceneTextureParameters Parameters;
-	FRHIUniformBuffer* ViewUniformBuffer = nullptr;
-
-	FPostOpaqueRenderDelegate PostOpaqueDelegate;
-	FDelegateHandle PostOpaqueDelegateHandle;
-	bool bInsidePass = false;
-};
-
-extern TGlobalResource<FNiagaraViewDataMgr> GNiagaraViewDataManager;
-
-class FNiagaraSceneTextureScope
-{
-public:
-	FNiagaraSceneTextureScope()
-	{
-		GNiagaraViewDataManager.BeginPass();
-	}
-
-	~FNiagaraSceneTextureScope()
-	{
-		GNiagaraViewDataManager.EndPass();
-	}
-};
-
 USTRUCT()
 struct FNiagaraWorldManagerTickFunction : public FTickFunction
 {
