Compilation time optimization for NiagaraEditorStyle
  - Split extremely long function leading to pathological codegen time with msvc
  - 2m39s -> 8.62s for NiagaraEditorStyle.cpp
  - 3m05s -> 41.42s for Module.NiagaraEditor.2_of_10.cpp
  - 6m56s -> 4m47s for a full QAGame build with XGE

diff --git a/Engine/Plugins/FX/Niagara/Source/NiagaraEditor/Private/NiagaraEditorStyle.cpp b/Engine/Plugins/FX/Niagara/Source/NiagaraEditor/Private/NiagaraEditorStyle.cpp
--- a/Engine/Plugins/FX/Niagara/Source/NiagaraEditor/Private/NiagaraEditorStyle.cpp
+++ b/Engine/Plugins/FX/Niagara/Source/NiagaraEditor/Private/NiagaraEditorStyle.cpp
@@ -61,21 +61,19 @@ const FVector2D Icon20x20(20.0f, 20.0f);
 const FVector2D Icon40x40(40.0f, 40.0f);
 const FVector2D Icon64x64(64.0f, 64.0f);
 
-TSharedRef< FSlateStyleSet > FNiagaraEditorStyle::Create()
-{
-	const FTextBlockStyle NormalText = FEditorStyle::GetWidgetStyle<FTextBlockStyle>("NormalText");
-	const FEditableTextBoxStyle NormalEditableTextBox = FStarshipCoreStyle::GetCoreStyle().GetWidgetStyle<FEditableTextBoxStyle>("NormalEditableTextBox");
-	const FSpinBoxStyle NormalSpinBox = FAppStyle::Get().GetWidgetStyle<FSpinBoxStyle>("SpinBox");
+namespace NiagaraEditorStyleImpl {
 
-	TSharedRef< FSlateStyleSet > Style = MakeShareable(new FSlateStyleSet("NiagaraEditorStyle"));
-	Style->SetContentRoot(FPaths::EngineContentDir() / TEXT("Editor/Slate/Niagara"));
-
-	// Stats
+void InitStats(TSharedRef< FSlateStyleSet > Style)
+{
 	const FTextBlockStyle CategoryText = FAppStyle::Get().GetWidgetStyle<FTextBlockStyle>("DetailsView.CategoryTextStyle");
 
 	Style->Set("NiagaraEditor.StatsText", CategoryText);
+}
+
+void InitAssetPicker(TSharedRef< FSlateStyleSet > Style)
+{
+	const FTextBlockStyle NormalText = FEditorStyle::GetWidgetStyle<FTextBlockStyle>("NormalText");
 
-	// Asset picker
 	FTextBlockStyle AssetPickerBoldAssetNameText = FTextBlockStyle(NormalText)
 		.SetColorAndOpacity(FLinearColor::White)
 		.SetFont(DEFAULT_FONT("Bold", 9));
@@ -127,6 +125,11 @@ TSharedRef< FSlateStyleSet > FNiagaraEditorStyle::Create()
 	Style->Set("NiagaraEditor.NewAssetDialog.ActiveOptionBorderColor", FLinearColor(FColor(96, 96, 96)));
 
 	Style->Set("NiagaraEditor.NewAssetDialog.SubBorder", new BOX_CORE_BRUSH("Common/GroupBorderLight", FMargin(4.0f / 16.0f)));
+}
+
+void InitActionMenu(TSharedRef< FSlateStyleSet > Style)
+{
+	const FTextBlockStyle NormalText = FEditorStyle::GetWidgetStyle<FTextBlockStyle>("NormalText");
 
 	// Action Menu
 	FTextBlockStyle ActionMenuHeadingText = FTextBlockStyle(NormalText)
@@ -195,7 +198,12 @@ TSharedRef< FSlateStyleSet > FNiagaraEditorStyle::Create()
 
 	Style->Set("GraphActionMenu.TemplateTabTextBlock", TemplateTabText);
 	
-	Style->Set( "GraphActionMenu.FilterCheckBox", NiagaraGraphActionMenuFilterCheckBox );
+	Style->Set("GraphActionMenu.FilterCheckBox", NiagaraGraphActionMenuFilterCheckBox);
+}
+
+void InitEmitterHeader(TSharedRef< FSlateStyleSet > Style)
+{
+	const FTextBlockStyle NormalText = FEditorStyle::GetWidgetStyle<FTextBlockStyle>("NormalText");
 	
 	// Emitter Header
 	FTextBlockStyle StackHeaderText = FTextBlockStyle(NormalText);
@@ -228,6 +236,12 @@ TSharedRef< FSlateStyleSet > FNiagaraEditorStyle::Create()
 		.SetFont(DEFAULT_FONT("Bold", 9));
 
 	Style->Set("NiagaraEditor.DetailsHeadingText", DetailsHeadingText);
+}
+
+void InitParameters(TSharedRef< FSlateStyleSet > Style)
+{
+	const FTextBlockStyle NormalText = FEditorStyle::GetWidgetStyle<FTextBlockStyle>("NormalText");
+	const FEditableTextBoxStyle NormalEditableTextBox = FStarshipCoreStyle::GetCoreStyle().GetWidgetStyle<FEditableTextBoxStyle>("NormalEditableTextBox");
 
 	// Parameters
 	
@@ -263,27 +277,38 @@ TSharedRef< FSlateStyleSet > FNiagaraEditorStyle::Create()
 	Style->Set("NiagaraEditor.ParameterName.TypeText", FTextBlockStyle(NormalText)
 		.SetFont(DEFAULT_FONT("Regular", 8))
 		.SetColorAndOpacity(FLinearColor(1.0f, 1.0f, 1.0f, 0.5f)));
+}
 
+void InitParameterMapView(TSharedRef< FSlateStyleSet > Style)
+{
 	// Parameter Map View
 	Style->Set("NiagaraEditor.Stack.DepressedHighlightedButtonBrush", new BOX_CORE_BRUSH("Common/ButtonHoverHint", FMargin(4 / 16.0f), FStyleColors::PrimaryPress));
 	Style->Set("NiagaraEditor.Stack.FlatButtonColor", FLinearColor(FColor(205, 205, 205)));
+}
+
+void InitCodeView(TSharedRef< FSlateStyleSet > Style)
+{
+	const FTextBlockStyle NormalText = FEditorStyle::GetWidgetStyle<FTextBlockStyle>("NormalText");
 
 	// Code View
-	{
-		Style->Set("NiagaraEditor.CodeView.Checkbox.Text", FTextBlockStyle(NormalText)
-			.SetFont(DEFAULT_FONT("Bold", 12))
-			.SetColorAndOpacity(FLinearColor(1.0f, 1.0f, 1.0f, 0.9f))
-			.SetShadowOffset(FVector2D(1, 1))
-			.SetShadowColorAndOpacity(FLinearColor(0, 0, 0, 0.9f)));
-
-		const int32 LogFontSize = 9;
-		FSlateFontInfo LogFont = DEFAULT_FONT("Mono", LogFontSize);
-		FTextBlockStyle NormalLogText = FTextBlockStyle(NormalText)
-			.SetFont(LogFont)
-			.SetColorAndOpacity(FLinearColor(FColor(0xffffffff)))
-			.SetSelectedBackgroundColor(FLinearColor(FColor(0xff666666)));
-		Style->Set("NiagaraEditor.CodeView.Hlsl.Normal", NormalLogText);
-	}
+	Style->Set("NiagaraEditor.CodeView.Checkbox.Text", FTextBlockStyle(NormalText)
+		.SetFont(DEFAULT_FONT("Bold", 12))
+		.SetColorAndOpacity(FLinearColor(1.0f, 1.0f, 1.0f, 0.9f))
+		.SetShadowOffset(FVector2D(1, 1))
+		.SetShadowColorAndOpacity(FLinearColor(0, 0, 0, 0.9f)));
+
+	const int32 LogFontSize = 9;
+	FSlateFontInfo LogFont = DEFAULT_FONT("Mono", LogFontSize);
+	FTextBlockStyle NormalLogText = FTextBlockStyle(NormalText)
+		.SetFont(LogFont)
+		.SetColorAndOpacity(FLinearColor(FColor(0xffffffff)))
+		.SetSelectedBackgroundColor(FLinearColor(FColor(0xff666666)));
+	Style->Set("NiagaraEditor.CodeView.Hlsl.Normal", NormalLogText);
+}
+
+void InitSelectedEmitter(TSharedRef< FSlateStyleSet > Style)
+{
+	const FTextBlockStyle NormalText = FEditorStyle::GetWidgetStyle<FTextBlockStyle>("NormalText");
 
 	// Selected Emitter
 	FSlateFontInfo SelectedEmitterUnsupportedSelectionFont = DEFAULT_FONT("Regular", 10);
@@ -291,8 +316,10 @@ TSharedRef< FSlateStyleSet > FNiagaraEditorStyle::Create()
 		.SetFont(SelectedEmitterUnsupportedSelectionFont)
 		.SetColorAndOpacity(FLinearColor(1.0f, 1.0f, 1.0f, 1.0f));
 	Style->Set("NiagaraEditor.SelectedEmitter.UnsupportedSelectionText", SelectedEmitterUnsupportedSelectionText);
+}
 
-	// Toolbar Icons
+void InitToolbarIcons(TSharedRef< FSlateStyleSet > Style)
+{
 	Style->Set("NiagaraEditor.AddEmitter", new IMAGE_BRUSH("Icons/icon_AddObject_40x", Icon40x40));
 	Style->Set("NiagaraEditor.AddEmitter.Small", new IMAGE_BRUSH("Icons/icon_AddObject_40x", Icon20x20));
 	Style->Set("NiagaraEditor.UnlockToChanges", new IMAGE_BRUSH("Icons/icon_levels_unlocked_40x", Icon40x40));
@@ -307,21 +334,33 @@ TSharedRef< FSlateStyleSet > FNiagaraEditorStyle::Create()
 	
 	Style->Set("Niagaraeditor.OverviewNode.IsolatedColor", FLinearColor::Yellow);
 	Style->Set("Niagaraeditor.OverviewNode.NotIsolatedColor", FLinearColor::Transparent);
+}
 
+void InitIcons(TSharedRef< FSlateStyleSet > Style)
+{
 	// Icons
 	Style->Set("NiagaraEditor.Isolate", new IMAGE_PLUGIN_BRUSH("Icons/Isolate", Icon16x16));
 	Style->Set("NiagaraEditor.Module.Pin.TypeSelector", new IMAGE_PLUGIN_BRUSH("Icons/Scratch", Icon16x16, FLinearColor::Gray));
 	Style->Set("NiagaraEditor.Module.AddPin", new IMAGE_PLUGIN_BRUSH("Icons/PlusSymbol_12x", Icon12x12, FLinearColor::Gray));
 	Style->Set("NiagaraEditor.Module.RemovePin", new IMAGE_PLUGIN_BRUSH("Icons/MinusSymbol_12x", Icon12x12, FLinearColor::Gray));
 	Style->Set("NiagaraEditor.Scratch", new IMAGE_PLUGIN_BRUSH("Icons/Scratch", Icon16x16, FLinearColor::Yellow));
+}
 
+void InitOverview(TSharedRef< FSlateStyleSet > Style)
+{
 	// Overview debug icons
 	Style->Set("NiagaraEditor.Overview.DebugActive", new IMAGE_PLUGIN_BRUSH("Icons/OverviewDebugActive", Icon16x16));
 	Style->Set("NiagaraEditor.Overview.DebugInactive", new IMAGE_PLUGIN_BRUSH("Icons/OverviewDebugInactive", Icon16x16));
-	
+}
+
+void InitEmitterDetails(TSharedRef< FSlateStyleSet > Style)
+{
 	// Emitter details customization
 	Style->Set("NiagaraEditor.MaterialWarningBorder", new BOX_CORE_BRUSH("Common/GroupBorderLight", FMargin(4.0f / 16.0f)));
+}
 
+void InitAssetColors(TSharedRef< FSlateStyleSet > Style)
+{
 	// Asset colors
 	Style->Set("NiagaraEditor.AssetColors.System", FLinearColor(1.0f, 0.0f, 0.0f));
 	Style->Set("NiagaraEditor.AssetColors.Emitter", FLinearColor(1.0f, 0.3f, 0.0f));
@@ -329,26 +368,42 @@ TSharedRef< FSlateStyleSet > FNiagaraEditorStyle::Create()
 	Style->Set("NiagaraEditor.AssetColors.ParameterCollection", FLinearColor(1.0f, 1.0f, 0.3f));
 	Style->Set("NiagaraEditor.AssetColors.ParameterCollectionInstance", FLinearColor(1.0f, 1.0f, 0.7f));
 	Style->Set("NiagaraEditor.AssetColors.ParameterDefinitions", FLinearColor(0.57f, 0.82f, 0.06f));
+}
 
+void InitThumbnails(TSharedRef< FSlateStyleSet > Style)
+{
 	// Script factory thumbnails
 	Style->Set("NiagaraEditor.Thumbnails.DynamicInputs", new IMAGE_BRUSH("Icons/NiagaraScriptDynamicInputs_64x", Icon64x64));
 	Style->Set("NiagaraEditor.Thumbnails.Functions", new IMAGE_BRUSH("Icons/NiagaraScriptFunction_64x", Icon64x64));
 	Style->Set("NiagaraEditor.Thumbnails.Modules", new IMAGE_BRUSH("Icons/NiagaraScriptModules_64x", Icon64x64));
+}
 
+void InitClassIcon(TSharedRef< FSlateStyleSet > Style)
+{
 	// Renderer class icons
 	Style->Set("ClassIcon.NiagaraSpriteRendererProperties", new IMAGE_PLUGIN_BRUSH("Icons/Renderers/renderer_sprite", Icon16x16));
 	Style->Set("ClassIcon.NiagaraMeshRendererProperties", new IMAGE_PLUGIN_BRUSH("Icons/Renderers/renderer_mesh", Icon16x16));
 	Style->Set("ClassIcon.NiagaraRibbonRendererProperties", new IMAGE_PLUGIN_BRUSH("Icons/Renderers/renderer_ribbon", Icon16x16));
 	Style->Set("ClassIcon.NiagaraLightRendererProperties", new IMAGE_PLUGIN_BRUSH("Icons/Renderers/renderer_light", Icon16x16));
 	Style->Set("ClassIcon.NiagaraRendererProperties", new IMAGE_PLUGIN_BRUSH("Icons/Renderers/renderer_default", Icon16x16));
+}
 
+void InitStackIcons(TSharedRef< FSlateStyleSet > Style)
+{
 	//GPU/CPU icons
 	Style->Set("NiagaraEditor.Stack.GPUIcon", new IMAGE_PLUGIN_BRUSH("Icons/Simulate_GPU_x40", Icon16x16));
 	Style->Set("NiagaraEditor.Stack.CPUIcon", new IMAGE_PLUGIN_BRUSH("Icons/Simulate_CPU_x40", Icon16x16));
+}
 
-
+void InitNiagaraSequence(TSharedRef< FSlateStyleSet > Style)
+{
 	// Niagara sequence
 	Style->Set("NiagaraEditor.NiagaraSequence.DefaultTrackColor", FLinearColor(0, .25f, 0));
+}
+
+void InitPlatformSet(TSharedRef< FSlateStyleSet > Style)
+{
+	const FTextBlockStyle NormalText = FEditorStyle::GetWidgetStyle<FTextBlockStyle>("NormalText");
 
 	// Niagara platform set customization
 	Style->Set("NiagaraEditor.PlatformSet.DropdownButton", new IMAGE_CORE_BRUSH("Common/ComboArrow", Icon8x8));
@@ -419,101 +474,154 @@ TSharedRef< FSlateStyleSet > FNiagaraEditorStyle::Create()
 		.SetInactiveBrush(IMAGE_CORE_BRUSH("Common/Selection", Icon8x8, SelectionColor_Inactive))
 		.SetInactiveHoveredBrush(IMAGE_CORE_BRUSH("Common/Selection", Icon8x8, SelectionColor_Inactive)));
 
+}
+
+void InitDropTarget(TSharedRef< FSlateStyleSet > Style)
+{
 	Style->Set("NiagaraEditor.DropTarget.BackgroundColor", FLinearColor(1.0f, 1.0f, 1.0f, 0.25f));
 	Style->Set("NiagaraEditor.DropTarget.BackgroundColorHover", FLinearColor(1.0f, 1.0f, 1.0f, 0.1f));
 	Style->Set("NiagaraEditor.DropTarget.BorderVertical", new IMAGE_PLUGIN_BRUSH("Icons/StackDropTargetBorder_Vertical", FVector2D(2, 8), FLinearColor::White, ESlateBrushTileType::Vertical));
 	Style->Set("NiagaraEditor.DropTarget.BorderHorizontal", new IMAGE_PLUGIN_BRUSH("Icons/StackDropTargetBorder_Horizontal", FVector2D(8, 2), FLinearColor::White, ESlateBrushTileType::Horizontal));
+}
 
+void InitScriptGraph(TSharedRef< FSlateStyleSet > Style)
+{
 	Style->Set("NiagaraEditor.ScriptGraph.SearchBorderColor", FLinearColor(.1f, .1f, .1f, 1.f));
+}
 
-	// Debugger Style
-	{
-		const FVector2D Icon24x24(24.0f, 24.0f);
-
-		Style->Set("NiagaraEditor.Debugger.TabIcon", new IMAGE_PLUGIN_BRUSH_SVG("Icons/NiagaraDebugger", Icon16x16));
-		Style->Set("NiagaraEditor.Debugger.PlayIcon", new IMAGE_PLUGIN_BRUSH("Icons/Debugger/Play", Icon24x24));
-		Style->Set("NiagaraEditor.Debugger.SpeedIcon", new IMAGE_PLUGIN_BRUSH("Icons/Debugger/Speed", Icon24x24));
-		Style->Set("NiagaraEditor.Debugger.PauseIcon", new IMAGE_PLUGIN_BRUSH("Icons/Debugger/Pause", Icon24x24));
-		Style->Set("NiagaraEditor.Debugger.LoopIcon", new IMAGE_PLUGIN_BRUSH("Icons/Debugger/Loop", Icon24x24));
-		Style->Set("NiagaraEditor.Debugger.StepIcon", new IMAGE_PLUGIN_BRUSH("Icons/Debugger/Step", Icon24x24));
-
-		Style->Set("NiagaraEditor.Debugger.Outliner.Capture", new IMAGE_PLUGIN_BRUSH("Icons/Debugger/Capture", Icon24x24));
-		Style->Set("NiagaraEditor.Debugger.Outliner.CapturePerf", new IMAGE_PLUGIN_BRUSH("Icons/Debugger/Perf_40x", Icon24x24));
-		Style->Set("NiagaraEditor.Debugger.Outliner.Filter", new IMAGE_PLUGIN_BRUSH("Icons/Debugger/Filter_24x", Icon24x24));
-
-		FButtonStyle OutlinerToolBarButton = FButtonStyle()
-			.SetNormal(BOX_CORE_BRUSH("Common/FlatButton", FMargin(4 / 16.0f), FLinearColor(0.0f, 0.0f, 0.0f, 0.0f)))
-			.SetHovered(BOX_CORE_BRUSH("Common/FlatButton", FMargin(4 / 16.0f), SelectionColor))
-			.SetPressed(BOX_CORE_BRUSH("Common/FlatButton", FMargin(4 / 16.0f), SelectionColor_Pressed))
-			.SetNormalPadding(FMargin(0, 0, 0, 1))
-			.SetPressedPadding(FMargin(0, 1, 0, 0));
-		Style->Set("NiagaraEditor.Debugger.Outliner.Toolbar", OutlinerToolBarButton);
-	}
+void InitDebuggerStyle(TSharedRef< FSlateStyleSet > Style)
+{
+	const FVector2D Icon24x24(24.0f, 24.0f);
 
-	// Baker
+	Style->Set("NiagaraEditor.Debugger.TabIcon", new IMAGE_PLUGIN_BRUSH_SVG("Icons/NiagaraDebugger", Icon16x16));
+	Style->Set("NiagaraEditor.Debugger.PlayIcon", new IMAGE_PLUGIN_BRUSH("Icons/Debugger/Play", Icon24x24));
+	Style->Set("NiagaraEditor.Debugger.SpeedIcon", new IMAGE_PLUGIN_BRUSH("Icons/Debugger/Speed", Icon24x24));
+	Style->Set("NiagaraEditor.Debugger.PauseIcon", new IMAGE_PLUGIN_BRUSH("Icons/Debugger/Pause", Icon24x24));
+	Style->Set("NiagaraEditor.Debugger.LoopIcon", new IMAGE_PLUGIN_BRUSH("Icons/Debugger/Loop", Icon24x24));
+	Style->Set("NiagaraEditor.Debugger.StepIcon", new IMAGE_PLUGIN_BRUSH("Icons/Debugger/Step", Icon24x24));
+
+	Style->Set("NiagaraEditor.Debugger.Outliner.Capture", new IMAGE_PLUGIN_BRUSH("Icons/Debugger/Capture", Icon24x24));
+	Style->Set("NiagaraEditor.Debugger.Outliner.CapturePerf", new IMAGE_PLUGIN_BRUSH("Icons/Debugger/Perf_40x", Icon24x24));
+	Style->Set("NiagaraEditor.Debugger.Outliner.Filter", new IMAGE_PLUGIN_BRUSH("Icons/Debugger/Filter_24x", Icon24x24));
+
+	const FSlateColor SelectionColor = FEditorStyle::GetSlateColor("SelectionColor");
+	const FSlateColor SelectionColor_Pressed = FEditorStyle::GetSlateColor("SelectionColor_Pressed");
+
+	FButtonStyle OutlinerToolBarButton = FButtonStyle()
+		.SetNormal(BOX_CORE_BRUSH("Common/FlatButton", FMargin(4 / 16.0f), FLinearColor(0.0f, 0.0f, 0.0f, 0.0f)))
+		.SetHovered(BOX_CORE_BRUSH("Common/FlatButton", FMargin(4 / 16.0f), SelectionColor))
+		.SetPressed(BOX_CORE_BRUSH("Common/FlatButton", FMargin(4 / 16.0f), SelectionColor_Pressed))
+		.SetNormalPadding(FMargin(0, 0, 0, 1))
+		.SetPressedPadding(FMargin(0, 1, 0, 0));
+	Style->Set("NiagaraEditor.Debugger.Outliner.Toolbar", OutlinerToolBarButton);
+}
+
+void InitBakerStyle(TSharedRef< FSlateStyleSet > Style)
+{
 	Style->Set("NiagaraEditor.Baker", new IMAGE_PLUGIN_BRUSH("Icons/Baker/BakerIcon", Icon40x40));
+}
 
+void InitCommonColors(TSharedRef< FSlateStyleSet > Style)
+{
 	Style->Set("NiagaraEditor.CommonColors.System", FLinearColor(FColor(1, 202, 252)));
 	Style->Set("NiagaraEditor.CommonColors.Emitter", FLinearColor(FColor(241, 99, 6)));
 	Style->Set("NiagaraEditor.CommonColors.Particle", FLinearColor(FColor(131, 218, 9)));
+}
 
-	//Outliner Style
-	{
-		Style->Set("NiagaraEditor.Outliner.WorldItem", FTableRowStyle()
-			.SetEvenRowBackgroundBrush(FSlateNoResource())
-			.SetEvenRowBackgroundHoveredBrush(IMAGE_CORE_BRUSH("Common/Selection", Icon8x8, SelectionColor_Inactive))
-			.SetOddRowBackgroundBrush(FSlateNoResource())
-			.SetOddRowBackgroundHoveredBrush(IMAGE_CORE_BRUSH("Common/Selection", Icon8x8, SelectionColor_Inactive))
-			.SetSelectorFocusedBrush(FSlateNoResource())
-			.SetActiveBrush(IMAGE_CORE_BRUSH("Common/Selection", Icon8x8, SelectionColor))
-			.SetActiveHoveredBrush(IMAGE_CORE_BRUSH("Common/Selection", Icon8x8, SelectionColor))
-			.SetInactiveBrush(IMAGE_CORE_BRUSH("Common/Selection", Icon8x8, SelectionColor_Inactive))
-			.SetInactiveHoveredBrush(IMAGE_CORE_BRUSH("Common/Selection", Icon8x8, SelectionColor_Inactive)));
-		
-
-		const FLinearColor SystemColor = Style->GetColor("NiagaraEditor.AssetColors.System") * 0.6f;
-		const FLinearColor SystemColorEven = SystemColor * 0.85f;
-		const FLinearColor SystemColorOdd = SystemColor * 0.7f;
-		Style->Set("NiagaraEditor.Outliner.SystemItem", FTableRowStyle()
-			.SetEvenRowBackgroundBrush(IMAGE_CORE_BRUSH("Common/Selection", Icon8x8, SystemColorEven))
-			.SetEvenRowBackgroundHoveredBrush(IMAGE_CORE_BRUSH("Common/Selection", Icon8x8, SystemColor))
-			.SetOddRowBackgroundBrush(IMAGE_CORE_BRUSH("Common/Selection", Icon8x8, SystemColorOdd))
-			.SetOddRowBackgroundHoveredBrush(IMAGE_CORE_BRUSH("Common/Selection", Icon8x8, SystemColor))
-			.SetSelectorFocusedBrush(FSlateNoResource())
-			.SetActiveBrush(IMAGE_CORE_BRUSH("Common/Selection", Icon8x8, SelectionColor))
-			.SetActiveHoveredBrush(IMAGE_CORE_BRUSH("Common/Selection", Icon8x8, SelectionColor))
-			.SetInactiveBrush(IMAGE_CORE_BRUSH("Common/Selection", Icon8x8, SelectionColor_Inactive))
-			.SetInactiveHoveredBrush(IMAGE_CORE_BRUSH("Common/Selection", Icon8x8, SelectionColor_Inactive)));
-
-
-		const FLinearColor SystemInstanceColor = Style->GetColor("NiagaraEditor.CommonColors.System") * 0.6f;
-		const FLinearColor SystemInstanceColorEven = SystemInstanceColor * 0.85f;
-		const FLinearColor SystemInstanceColorOdd = SystemInstanceColor * 0.7f;
-		Style->Set("NiagaraEditor.Outliner.ComponentItem", FTableRowStyle()
-			.SetEvenRowBackgroundBrush(IMAGE_CORE_BRUSH("Common/Selection", Icon8x8, SystemInstanceColorEven))
-			.SetEvenRowBackgroundHoveredBrush(IMAGE_CORE_BRUSH("Common/Selection", Icon8x8, SystemInstanceColor))
-			.SetOddRowBackgroundBrush(IMAGE_CORE_BRUSH("Common/Selection", Icon8x8, SystemInstanceColorOdd))
-			.SetOddRowBackgroundHoveredBrush(IMAGE_CORE_BRUSH("Common/Selection", Icon8x8, SystemInstanceColor))
-			.SetSelectorFocusedBrush(FSlateNoResource())
-			.SetActiveBrush(IMAGE_CORE_BRUSH("Common/Selection", Icon8x8, SelectionColor))
-			.SetActiveHoveredBrush(IMAGE_CORE_BRUSH("Common/Selection", Icon8x8, SelectionColor))
-			.SetInactiveBrush(IMAGE_CORE_BRUSH("Common/Selection", Icon8x8, SelectionColor_Inactive))
-			.SetInactiveHoveredBrush(IMAGE_CORE_BRUSH("Common/Selection", Icon8x8, SelectionColor_Inactive)));
-
-		const FLinearColor EmitterInstanceColor = Style->GetColor("NiagaraEditor.CommonColors.Emitter") * 0.6f;
-		const FLinearColor EmitterInstanceColorEven = EmitterInstanceColor * 0.85f;
-		const FLinearColor EmitterInstanceColorOdd = EmitterInstanceColor * 0.7f;
-		Style->Set("NiagaraEditor.Outliner.EmitterItem", FTableRowStyle()
-			.SetEvenRowBackgroundBrush(IMAGE_CORE_BRUSH("Common/Selection", Icon8x8, EmitterInstanceColorEven))
-			.SetEvenRowBackgroundHoveredBrush(IMAGE_CORE_BRUSH("Common/Selection", Icon8x8, EmitterInstanceColor))
-			.SetOddRowBackgroundBrush(IMAGE_CORE_BRUSH("Common/Selection", Icon8x8, EmitterInstanceColorOdd))
-			.SetOddRowBackgroundHoveredBrush(IMAGE_CORE_BRUSH("Common/Selection", Icon8x8, EmitterInstanceColor))
-			.SetSelectorFocusedBrush(FSlateNoResource())
-			.SetActiveBrush(IMAGE_CORE_BRUSH("Common/Selection", Icon8x8, SelectionColor))
-			.SetActiveHoveredBrush(IMAGE_CORE_BRUSH("Common/Selection", Icon8x8, SelectionColor))
-			.SetInactiveBrush(IMAGE_CORE_BRUSH("Common/Selection", Icon8x8, SelectionColor_Inactive))
-			.SetInactiveHoveredBrush(IMAGE_CORE_BRUSH("Common/Selection", Icon8x8, SelectionColor_Inactive)));
-	}
+void InitOutlinerStyle(TSharedRef< FSlateStyleSet > Style)
+{
+	const FSlateColor SelectionColor = FEditorStyle::GetSlateColor("SelectionColor");
+	const FSlateColor SelectionColor_Inactive = FEditorStyle::GetSlateColor("SelectionColor_Inactive");
+
+	Style->Set("NiagaraEditor.Outliner.WorldItem", FTableRowStyle()
+		.SetEvenRowBackgroundBrush(FSlateNoResource())
+		.SetEvenRowBackgroundHoveredBrush(IMAGE_CORE_BRUSH("Common/Selection", Icon8x8, SelectionColor_Inactive))
+		.SetOddRowBackgroundBrush(FSlateNoResource())
+		.SetOddRowBackgroundHoveredBrush(IMAGE_CORE_BRUSH("Common/Selection", Icon8x8, SelectionColor_Inactive))
+		.SetSelectorFocusedBrush(FSlateNoResource())
+		.SetActiveBrush(IMAGE_CORE_BRUSH("Common/Selection", Icon8x8, SelectionColor))
+		.SetActiveHoveredBrush(IMAGE_CORE_BRUSH("Common/Selection", Icon8x8, SelectionColor))
+		.SetInactiveBrush(IMAGE_CORE_BRUSH("Common/Selection", Icon8x8, SelectionColor_Inactive))
+		.SetInactiveHoveredBrush(IMAGE_CORE_BRUSH("Common/Selection", Icon8x8, SelectionColor_Inactive)));
+
+
+	const FLinearColor SystemColor = Style->GetColor("NiagaraEditor.AssetColors.System") * 0.6f;
+	const FLinearColor SystemColorEven = SystemColor * 0.85f;
+	const FLinearColor SystemColorOdd = SystemColor * 0.7f;
+	Style->Set("NiagaraEditor.Outliner.SystemItem", FTableRowStyle()
+		.SetEvenRowBackgroundBrush(IMAGE_CORE_BRUSH("Common/Selection", Icon8x8, SystemColorEven))
+		.SetEvenRowBackgroundHoveredBrush(IMAGE_CORE_BRUSH("Common/Selection", Icon8x8, SystemColor))
+		.SetOddRowBackgroundBrush(IMAGE_CORE_BRUSH("Common/Selection", Icon8x8, SystemColorOdd))
+		.SetOddRowBackgroundHoveredBrush(IMAGE_CORE_BRUSH("Common/Selection", Icon8x8, SystemColor))
+		.SetSelectorFocusedBrush(FSlateNoResource())
+		.SetActiveBrush(IMAGE_CORE_BRUSH("Common/Selection", Icon8x8, SelectionColor))
+		.SetActiveHoveredBrush(IMAGE_CORE_BRUSH("Common/Selection", Icon8x8, SelectionColor))
+		.SetInactiveBrush(IMAGE_CORE_BRUSH("Common/Selection", Icon8x8, SelectionColor_Inactive))
+		.SetInactiveHoveredBrush(IMAGE_CORE_BRUSH("Common/Selection", Icon8x8, SelectionColor_Inactive)));
+
+
+	const FLinearColor SystemInstanceColor = Style->GetColor("NiagaraEditor.CommonColors.System") * 0.6f;
+	const FLinearColor SystemInstanceColorEven = SystemInstanceColor * 0.85f;
+	const FLinearColor SystemInstanceColorOdd = SystemInstanceColor * 0.7f;
+	Style->Set("NiagaraEditor.Outliner.ComponentItem", FTableRowStyle()
+		.SetEvenRowBackgroundBrush(IMAGE_CORE_BRUSH("Common/Selection", Icon8x8, SystemInstanceColorEven))
+		.SetEvenRowBackgroundHoveredBrush(IMAGE_CORE_BRUSH("Common/Selection", Icon8x8, SystemInstanceColor))
+		.SetOddRowBackgroundBrush(IMAGE_CORE_BRUSH("Common/Selection", Icon8x8, SystemInstanceColorOdd))
+		.SetOddRowBackgroundHoveredBrush(IMAGE_CORE_BRUSH("Common/Selection", Icon8x8, SystemInstanceColor))
+		.SetSelectorFocusedBrush(FSlateNoResource())
+		.SetActiveBrush(IMAGE_CORE_BRUSH("Common/Selection", Icon8x8, SelectionColor))
+		.SetActiveHoveredBrush(IMAGE_CORE_BRUSH("Common/Selection", Icon8x8, SelectionColor))
+		.SetInactiveBrush(IMAGE_CORE_BRUSH("Common/Selection", Icon8x8, SelectionColor_Inactive))
+		.SetInactiveHoveredBrush(IMAGE_CORE_BRUSH("Common/Selection", Icon8x8, SelectionColor_Inactive)));
+
+	const FLinearColor EmitterInstanceColor = Style->GetColor("NiagaraEditor.CommonColors.Emitter") * 0.6f;
+	const FLinearColor EmitterInstanceColorEven = EmitterInstanceColor * 0.85f;
+	const FLinearColor EmitterInstanceColorOdd = EmitterInstanceColor * 0.7f;
+	Style->Set("NiagaraEditor.Outliner.EmitterItem", FTableRowStyle()
+		.SetEvenRowBackgroundBrush(IMAGE_CORE_BRUSH("Common/Selection", Icon8x8, EmitterInstanceColorEven))
+		.SetEvenRowBackgroundHoveredBrush(IMAGE_CORE_BRUSH("Common/Selection", Icon8x8, EmitterInstanceColor))
+		.SetOddRowBackgroundBrush(IMAGE_CORE_BRUSH("Common/Selection", Icon8x8, EmitterInstanceColorOdd))
+		.SetOddRowBackgroundHoveredBrush(IMAGE_CORE_BRUSH("Common/Selection", Icon8x8, EmitterInstanceColor))
+		.SetSelectorFocusedBrush(FSlateNoResource())
+		.SetActiveBrush(IMAGE_CORE_BRUSH("Common/Selection", Icon8x8, SelectionColor))
+		.SetActiveHoveredBrush(IMAGE_CORE_BRUSH("Common/Selection", Icon8x8, SelectionColor))
+		.SetInactiveBrush(IMAGE_CORE_BRUSH("Common/Selection", Icon8x8, SelectionColor_Inactive))
+		.SetInactiveHoveredBrush(IMAGE_CORE_BRUSH("Common/Selection", Icon8x8, SelectionColor_Inactive)));
+}
+
+} // namespace NiagaraEditorStyleImpl
+
+TSharedRef< FSlateStyleSet > FNiagaraEditorStyle::Create()
+{
+	using namespace NiagaraEditorStyleImpl;
+
+	TSharedRef< FSlateStyleSet > Style = MakeShareable(new FSlateStyleSet("NiagaraEditorStyle"));
+	Style->SetContentRoot(FPaths::EngineContentDir() / TEXT("Editor/Slate/Niagara"));
+
+	InitStats(Style);
+	InitAssetPicker(Style);
+	InitActionMenu(Style);
+	InitEmitterHeader(Style);
+	InitParameters(Style);
+	InitParameterMapView(Style);
+	InitCodeView(Style);
+	InitSelectedEmitter(Style);
+	InitToolbarIcons(Style);
+	InitIcons(Style);
+	InitEmitterDetails(Style);
+	InitAssetColors(Style);
+	InitThumbnails(Style);
+	InitClassIcon(Style);
+	InitStackIcons(Style);
+	InitNiagaraSequence(Style);
+	InitPlatformSet(Style);
+	InitDropTarget(Style);
+	InitScriptGraph(Style);
+	InitDebuggerStyle(Style);
+	InitBakerStyle(Style);
+	InitCommonColors(Style);
+	InitOutlinerStyle(Style);
+	
 	return Style;
 }
 
