fix for broken material

diff --git a/Engine/Plugins/FX/Niagara/Shaders/Private/NiagaraMeshVertexFactory.ush b/Engine/Plugins/FX/Niagara/Shaders/Private/NiagaraMeshVertexFactory.ush
--- a/Engine/Plugins/FX/Niagara/Shaders/Private/NiagaraMeshVertexFactory.ush
+++ b/Engine/Plugins/FX/Niagara/Shaders/Private/NiagaraMeshVertexFactory.ush
@@ -922,11 +922,13 @@ float4 VertexFactoryGetTranslatedPrimitiveVolumeBounds(FVertexFactoryInterpolant
 
 	float2 VertexFactoryGetRayTracingTextureCoordinate( FVertexFactoryRayTracingInterpolants Interpolants )
 	{
-	#if NUM_MATERIAL_TEXCOORDS
+	#if USE_PARTICLE_SUBUVS
+		return Interpolants.InterpolantsVSToPS.SubUV0AndTexCoord0.zw;
+	#elif NUM_TEX_COORD_INTERPOLATORS
 		return Interpolants.InterpolantsVSToPS.TexCoords[0].xy;
-	#else // #if NUM_MATERIAL_TEXCOORDS
+	#else
 		return float2(0,0);
-	#endif // #if NUM_MATERIAL_TEXCOORDS
+	#endif
 	}
 
 	FVertexFactoryInterpolantsVSToPS VertexFactoryAssignInterpolants(FVertexFactoryRayTracingInterpolants Input)
@@ -958,7 +960,9 @@ float4 VertexFactoryGetTranslatedPrimitiveVolumeBounds(FVertexFactoryInterpolant
 		INTERPOLATE_MEMBER(InterpolantsVSToPS.ParticleColor);
 #endif
 
-#if  NUM_TEX_COORD_INTERPOLATORS
+#if USE_PARTICLE_SUBUVS
+		INTERPOLATE_MEMBER(InterpolantsVSToPS.SubUV0AndTexCoord0);
+#elif NUM_TEX_COORD_INTERPOLATORS
 		for (int i = 0; i < ( NUM_TEX_COORD_INTERPOLATORS + 1) / 2; ++i)
 		{
 			INTERPOLATE_MEMBER(InterpolantsVSToPS.TexCoords[i]);
