- Integrate CL 16390833 - Fix OnGraphChanged callbacks getting double added but only removed once, resulting in hanging delegate bindings.

diff --git a/Engine/Plugins/FX/Niagara/Source/NiagaraEditor/Private/ViewModels/NiagaraParameterPanelViewModel.cpp b/Engine/Plugins/FX/Niagara/Source/NiagaraEditor/Private/ViewModels/NiagaraParameterPanelViewModel.cpp
--- a/Engine/Plugins/FX/Niagara/Source/NiagaraEditor/Private/ViewModels/NiagaraParameterPanelViewModel.cpp
+++ b/Engine/Plugins/FX/Niagara/Source/NiagaraEditor/Private/ViewModels/NiagaraParameterPanelViewModel.cpp
@@ -654,7 +654,7 @@ void FNiagaraSystemToolkitParameterPanelViewModel::Cleanup()
 		return CastChecked<UNiagaraScriptSource>(Script->GetLatestSource())->NodeGraph;
 	};
 
-	TArray<UNiagaraGraph*> GraphsToRemoveCallbacks;
+	TSet<UNiagaraGraph*> GraphsToRemoveCallbacks;
 	GraphsToRemoveCallbacks.Add(GetGraphFromScript(System.GetSystemSpawnScript()));
 	GraphsToRemoveCallbacks.Add(GetGraphFromScript(System.GetSystemUpdateScript()));
 	for (TSharedRef<FNiagaraEmitterHandleViewModel> EmitterHandleViewModel : SystemViewModel->GetEmitterHandleViewModels())
@@ -709,7 +709,7 @@ void FNiagaraSystemToolkitParameterPanelViewModel::Init(const FSystemToolkitUICo
 		return CastChecked<UNiagaraScriptSource>(Script->GetLatestSource())->NodeGraph;
 	};
 
-	TArray<UNiagaraGraph*> GraphsToAddCallbacks;
+	TSet<UNiagaraGraph*> GraphsToAddCallbacks;
 	GraphsToAddCallbacks.Add(GetGraphFromScript(System.GetSystemSpawnScript()));
 	GraphsToAddCallbacks.Add(GetGraphFromScript(System.GetSystemUpdateScript()));
 	for (TSharedRef<FNiagaraEmitterHandleViewModel> EmitterHandleViewModel : SystemViewModel->GetEmitterHandleViewModels())
@@ -1744,7 +1744,7 @@ void FNiagaraSystemToolkitParameterPanelViewModel::ReconcileOnGraphChangedBindin
 	UnvisitedGraphIds.Remove(GetGraphFromScript(System.GetSystemSpawnScript())->GetUniqueID());
 	UnvisitedGraphIds.Remove(GetGraphFromScript(System.GetSystemUpdateScript())->GetUniqueID());
 
-	TArray<UNiagaraGraph*> GraphsToAddCallbacks;
+	TSet<UNiagaraGraph*> GraphsToAddCallbacks;
 	for (TSharedRef<FNiagaraEmitterHandleViewModel> EmitterHandleViewModel : SystemViewModel->GetEmitterHandleViewModels())
 	{
 		if (FNiagaraEmitterHandle* EmitterHandle = EmitterHandleViewModel->GetEmitterHandle())
