Fixing include tool warning.

#rnx

diff --git a/Engine/Plugins/FX/Niagara/Source/NiagaraEditor/Private/Widgets/SNiagaraPinTypeSelector.h b/Engine/Plugins/FX/Niagara/Source/NiagaraEditor/Private/Widgets/SNiagaraPinTypeSelector.h
--- a/Engine/Plugins/FX/Niagara/Source/NiagaraEditor/Private/Widgets/SNiagaraPinTypeSelector.h
+++ b/Engine/Plugins/FX/Niagara/Source/NiagaraEditor/Private/Widgets/SNiagaraPinTypeSelector.h
@@ -13,6 +13,7 @@
 #include "NiagaraTypes.h"
 #include "Widgets/Input/SComboButton.h"
 #include "Widgets/SCompoundWidget.h"
+#include "SGraphPalette.h"
 
 class NIAGARAEDITOR_API SNiagaraPinTypeSelector : public SCompoundWidget
 {
@@ -31,10 +32,6 @@ private:
 	UEdGraphPin* Pin = nullptr;
 	TSharedPtr<SComboButton> SelectorButton;
 };
-#include "Widgets/DeclarativeSyntaxSupport.h"
-#include "Widgets/SWidget.h"
-#include "Widgets/SCompoundWidget.h"
-#include "SGraphPalette.h"
 
 class SNiagaraIconWidget : public SGraphPaletteItem
 {
