only copy prev to curr for named grid values

diff --git a/Engine/Plugins/FX/Niagara/Source/Niagara/Classes/NiagaraDataInterfaceGrid3DCollection.h b/Engine/Plugins/FX/Niagara/Source/Niagara/Classes/NiagaraDataInterfaceGrid3DCollection.h
--- a/Engine/Plugins/FX/Niagara/Source/Niagara/Classes/NiagaraDataInterfaceGrid3DCollection.h
+++ b/Engine/Plugins/FX/Niagara/Source/Niagara/Classes/NiagaraDataInterfaceGrid3DCollection.h
@@ -36,6 +36,7 @@ struct FGrid3DCollectionRWInstanceData_GameThread
 	FIntVector NumCells = FIntVector::ZeroValue;
 	FIntVector NumTiles = FIntVector::ZeroValue;
 	int32 TotalNumAttributes = 0;
+	int32 TotalNumNamedAttributes = 0;
 	FVector CellSize = FVector::ZeroVector;
 	FVector WorldBBoxSize = FVector::ZeroVector;
 	EPixelFormat PixelFormat = EPixelFormat::PF_R32_FLOAT;
@@ -62,6 +63,7 @@ struct FGrid3DCollectionRWInstanceData_RenderThread
 	FIntVector NumCells = FIntVector::ZeroValue;
 	FIntVector NumTiles = FIntVector::ZeroValue;
 	int32 TotalNumAttributes = 0;
+	int32 TotalNumNamedAttributes = 0;
 	FVector CellSize = FVector::ZeroVector;
 	FVector WorldBBoxSize = FVector::ZeroVector;
 	EPixelFormat PixelFormat = EPixelFormat::PF_R32_FLOAT;
diff --git a/Engine/Plugins/FX/Niagara/Source/Niagara/Classes/NiagaraDataInterfaceRW.h b/Engine/Plugins/FX/Niagara/Source/Niagara/Classes/NiagaraDataInterfaceRW.h
--- a/Engine/Plugins/FX/Niagara/Source/Niagara/Classes/NiagaraDataInterfaceRW.h
+++ b/Engine/Plugins/FX/Niagara/Source/Niagara/Classes/NiagaraDataInterfaceRW.h
@@ -40,6 +40,7 @@ class NIAGARA_API UNiagaraDataInterfaceRWBase : public UNiagaraDataInterface
 public:
 	// Global HLSL variable base names, used by HLSL.
 	static const FString NumAttributesName;
+	static const FString NumNamedAttributesName;
 	static const FString NumCellsName;
 	static const FString UnitToUVName;
 	static const FString CellSizeName;
diff --git a/Engine/Plugins/FX/Niagara/Source/Niagara/Private/NiagaraDataInterfaceGrid3DCollection.cpp b/Engine/Plugins/FX/Niagara/Source/Niagara/Private/NiagaraDataInterfaceGrid3DCollection.cpp
--- a/Engine/Plugins/FX/Niagara/Source/Niagara/Private/NiagaraDataInterfaceGrid3DCollection.cpp
+++ b/Engine/Plugins/FX/Niagara/Source/Niagara/Private/NiagaraDataInterfaceGrid3DCollection.cpp
@@ -149,6 +149,7 @@ public:
 	void Bind(const FNiagaraDataInterfaceGPUParamInfo& ParameterInfo, const class FShaderParameterMap& ParameterMap)
 	{
 		NumAttributesParam.Bind(ParameterMap, *(UNiagaraDataInterfaceRWBase::NumAttributesName + ParameterInfo.DataInterfaceHLSLSymbol));
+		NumNamedAttributesParam.Bind(ParameterMap, *(UNiagaraDataInterfaceRWBase::NumNamedAttributesName + ParameterInfo.DataInterfaceHLSLSymbol));
 		NumCellsParam.Bind(ParameterMap, *(UNiagaraDataInterfaceRWBase::NumCellsName + ParameterInfo.DataInterfaceHLSLSymbol));
 		NumTilesParam.Bind(ParameterMap, *(UNiagaraDataInterfaceGrid3DCollection::NumTilesName + ParameterInfo.DataInterfaceHLSLSymbol));
 		OneOverNumTilesParam.Bind(ParameterMap, *(UNiagaraDataInterfaceGrid3DCollection::OneOverNumTilesName + ParameterInfo.DataInterfaceHLSLSymbol));
@@ -236,6 +237,7 @@ public:
 		}
 
 		SetShaderValue(RHICmdList, ComputeShaderRHI, NumAttributesParam, ProxyData->TotalNumAttributes);
+		SetShaderValue(RHICmdList, ComputeShaderRHI, NumNamedAttributesParam, ProxyData->TotalNumNamedAttributes);
 
 		SetShaderValue(RHICmdList, ComputeShaderRHI, UnitToUVParam, FVector3f(1.0f) / FVector3f(ProxyData->NumCells));
 
@@ -290,6 +292,7 @@ public:
 
 private:
 	LAYOUT_FIELD(FShaderParameter, NumAttributesParam);
+	LAYOUT_FIELD(FShaderParameter, NumNamedAttributesParam);
 	LAYOUT_FIELD(FShaderParameter, UnitToUVParam);
 	LAYOUT_FIELD(FShaderParameter, NumCellsParam);
 	LAYOUT_FIELD(FShaderParameter, NumTilesParam);
@@ -1094,6 +1097,7 @@ void UNiagaraDataInterfaceGrid3DCollection::GetParameterDefinitionHLSL(const FNi
 		int4 {AttributeIndicesName}[{AttributeInt4Count}];
 		Buffer<float4> {PerAttributeDataName};
 		int {NumAttributesName};
+		int {NumNamedAttributesName};
 	)");
 
 	// If we use an int array for the attribute indices, the shader compiler will actually use int4 due to the packing rules,
@@ -1117,6 +1121,7 @@ void UNiagaraDataInterfaceGrid3DCollection::GetParameterDefinitionHLSL(const FNi
 		{ TEXT("PerAttributeDataName"), PerAttributeDataName + ParamInfo.DataInterfaceHLSLSymbol},
 		{ TEXT("AttributeInt4Count"), AttributeInt4Count},
 		{ TEXT("NumAttributesName"), UNiagaraDataInterfaceRWBase::NumAttributesName + ParamInfo.DataInterfaceHLSLSymbol},
+		{ TEXT("NumNamedAttributesName"), UNiagaraDataInterfaceRWBase::NumNamedAttributesName + ParamInfo.DataInterfaceHLSLSymbol},
 	};
 	OutHLSL += FString::Format(FormatDeclarations, ArgsDeclarations);
 }
@@ -1335,6 +1340,7 @@ bool UNiagaraDataInterfaceGrid3DCollection::GetFunctionHLSL(const FNiagaraDataIn
 		{TEXT("Grid"), GridName + ParamInfo.DataInterfaceHLSLSymbol},
 		{TEXT("OutputGrid"), OutputGridName + ParamInfo.DataInterfaceHLSLSymbol},
 		{TEXT("NumAttributes"), UNiagaraDataInterfaceRWBase::NumAttributesName + ParamInfo.DataInterfaceHLSLSymbol},
+		{TEXT("NumNamedAttributes"), UNiagaraDataInterfaceRWBase::NumNamedAttributesName + ParamInfo.DataInterfaceHLSLSymbol},
 		{TEXT("NumCells"), UNiagaraDataInterfaceRWBase::NumCellsName + ParamInfo.DataInterfaceHLSLSymbol},
 		{TEXT("UnitToUVName"), UNiagaraDataInterfaceRWBase::UnitToUVName + ParamInfo.DataInterfaceHLSLSymbol},
 		{TEXT("SamplerName"), SamplerName + ParamInfo.DataInterfaceHLSLSymbol},
@@ -1407,7 +1413,7 @@ bool UNiagaraDataInterfaceGrid3DCollection::GetFunctionHLSL(const FNiagaraDataIn
 		static const TCHAR* FormatBounds = TEXT(R"(
 			void {FunctionName}(int In_IndexX, int In_IndexY, int In_IndexZ)
 			{
-				for (int AttributeIndex = 0; AttributeIndex < {NumAttributes}.x; AttributeIndex++)
+				for (int AttributeIndex = 0; AttributeIndex < {NumNamedAttributes}.x; AttributeIndex++)
 				{		
 					int3 TileOffset = {PerAttributeDataName}[AttributeIndex].xyz;
 
@@ -1724,10 +1730,12 @@ bool UNiagaraDataInterfaceGrid3DCollection::InitPerInstanceData(void* PerInstanc
 
 	/* Go through all references to this data interface and build up the attribute list from the function metadata of those referenced.*/
 	int32 NumAttribChannelsFound = 0;
-	FindAttributes(InstanceData->Vars, InstanceData->Offsets, NumAttribChannelsFound);
+	int32 NumNamedAttribChannelsFound = 0;
+	FindAttributes(InstanceData->Vars, InstanceData->Offsets, NumNamedAttribChannelsFound);
 
-	NumAttribChannelsFound = NumAttributes + NumAttribChannelsFound;
+	NumAttribChannelsFound = NumAttributes + NumNamedAttribChannelsFound;
 	InstanceData->TotalNumAttributes = NumAttribChannelsFound;
+	InstanceData->TotalNumNamedAttributes = NumNamedAttribChannelsFound;
 
 	if (SetResolutionMethod == ESetResolutionMethod::Independent)
 	{
@@ -1894,6 +1902,7 @@ bool UNiagaraDataInterfaceGrid3DCollection::InitPerInstanceData(void* PerInstanc
 		TargetData->NumCells = RT_InstanceData.NumCells;
 		TargetData->NumTiles = RT_InstanceData.NumTiles;
 		TargetData->TotalNumAttributes = RT_InstanceData.TotalNumAttributes;
+		TargetData->TotalNumNamedAttributes = RT_InstanceData.TotalNumNamedAttributes;
 		TargetData->CellSize = RT_InstanceData.CellSize;
 		TargetData->WorldBBoxSize = RT_InstanceData.WorldBBoxSize;
 		TargetData->PixelFormat = RT_InstanceData.PixelFormat;
@@ -2412,6 +2421,7 @@ bool UNiagaraDataInterfaceGrid3DCollection::PerInstanceTickPostSimulate(void* Pe
 			TargetData->NumTiles = RT_InstanceData.NumTiles;
 
 			TargetData->TotalNumAttributes = RT_InstanceData.TotalNumAttributes;
+			TargetData->TotalNumNamedAttributes = RT_InstanceData.TotalNumNamedAttributes;
 			TargetData->CellSize = RT_InstanceData.CellSize;
 
 			TargetData->Buffers.Empty();
diff --git a/Engine/Plugins/FX/Niagara/Source/Niagara/Private/NiagaraDataInterfaceRW.cpp b/Engine/Plugins/FX/Niagara/Source/Niagara/Private/NiagaraDataInterfaceRW.cpp
--- a/Engine/Plugins/FX/Niagara/Source/Niagara/Private/NiagaraDataInterfaceRW.cpp
+++ b/Engine/Plugins/FX/Niagara/Source/Niagara/Private/NiagaraDataInterfaceRW.cpp
@@ -9,6 +9,7 @@
 
 // Global HLSL variable base names, used by HLSL.
 const FString UNiagaraDataInterfaceRWBase::NumAttributesName(TEXT("NumAttributes_"));
+const FString UNiagaraDataInterfaceRWBase::NumNamedAttributesName(TEXT("NumNamedAttributes_"));
 const FString UNiagaraDataInterfaceRWBase::NumCellsName(TEXT("NumCells_"));
 const FString UNiagaraDataInterfaceRWBase::UnitToUVName(TEXT("UnitToUV_"));
 const FString UNiagaraDataInterfaceRWBase::CellSizeName(TEXT("CellSize_"));
