Fix RDG validation assert when binding both mobile and desktop SceneTextures UB for Niagara rendering

diff --git a/Engine/Plugins/FX/Niagara/Source/Niagara/Private/NiagaraRenderViewDataManager.cpp b/Engine/Plugins/FX/Niagara/Source/Niagara/Private/NiagaraRenderViewDataManager.cpp
--- a/Engine/Plugins/FX/Niagara/Source/Niagara/Private/NiagaraRenderViewDataManager.cpp
+++ b/Engine/Plugins/FX/Niagara/Source/Niagara/Private/NiagaraRenderViewDataManager.cpp
@@ -29,6 +29,7 @@ void FNiagaraRenderViewDataManager::GetSceneTextureParameters(FRDGBuilder& Graph
 		{
 			OutParameters.SceneTextures = CreateSceneTextureUniformBuffer(GraphBuilder, LocalFeatureLevel, ESceneTextureSetupMode::SceneVelocity);
 		}
+		OutParameters.MobileSceneTextures = nullptr;
 	}
 	else if (FSceneInterface::GetShadingPath(LocalFeatureLevel) == EShadingPath::Mobile)
 	{
@@ -36,6 +37,7 @@ void FNiagaraRenderViewDataManager::GetSceneTextureParameters(FRDGBuilder& Graph
 		{
 			OutParameters.MobileSceneTextures = CreateMobileSceneTextureUniformBuffer(GraphBuilder, EMobileSceneTextureSetupMode::None);
 		}
+		OutParameters.SceneTextures = nullptr;
 	}
 	OutParameters.Depth = Parameters.Depth;
 	OutParameters.Normal = Parameters.Normal;
