Merging using //UE5/Main_to_//UE5/Release-Engine-Staging
Fix crashes on undo/redo in script and system toolkits.

diff --git a/Engine/Plugins/FX/Niagara/Source/NiagaraEditor/Private/ViewModels/NiagaraCurveSelectionViewModel.cpp b/Engine/Plugins/FX/Niagara/Source/NiagaraEditor/Private/ViewModels/NiagaraCurveSelectionViewModel.cpp
--- a/Engine/Plugins/FX/Niagara/Source/NiagaraEditor/Private/ViewModels/NiagaraCurveSelectionViewModel.cpp
+++ b/Engine/Plugins/FX/Niagara/Source/NiagaraEditor/Private/ViewModels/NiagaraCurveSelectionViewModel.cpp
@@ -886,7 +886,8 @@ void UNiagaraCurveSelectionViewModel::Refresh()
 
 		for (TSharedRef<FNiagaraEmitterHandleViewModel> EmitterHandleViewModel : SystemViewModel->GetEmitterHandleViewModels())
 		{
-			if (EmitterHandleViewModel->GetEmitterHandle())
+			// During transactions Emitters may not have their EditorData set if they were created during the transaction; skip if this is true.
+			if (EmitterHandleViewModel->GetEmitterHandle() && EmitterHandleViewModel->GetEmitterHandle()->GetInstance()->GetEditorData())
 			{
 				TSharedPtr<FNiagaraCurveSelectionTreeNode> EmitterNode = CreateNodeForEmitter(OldRootChildNodes, *EmitterHandleViewModel->GetEmitterHandle());
 				if (EmitterNode.IsValid())
diff --git a/Engine/Plugins/FX/Niagara/Source/NiagaraEditor/Private/ViewModels/NiagaraOverviewGraphViewModel.cpp b/Engine/Plugins/FX/Niagara/Source/NiagaraEditor/Private/ViewModels/NiagaraOverviewGraphViewModel.cpp
--- a/Engine/Plugins/FX/Niagara/Source/NiagaraEditor/Private/ViewModels/NiagaraOverviewGraphViewModel.cpp
+++ b/Engine/Plugins/FX/Niagara/Source/NiagaraEditor/Private/ViewModels/NiagaraOverviewGraphViewModel.cpp
@@ -577,7 +577,11 @@ void FNiagaraOverviewGraphViewModel::SystemSelectionChanged()
 void FNiagaraOverviewGraphViewModel::PostUndo(bool bSuccess)
 {
 	// This is neccessary to have the graph editor respond correctly to data changes due to undo.
-	OverviewGraph->NotifyGraphChanged();
+	// Validate the OverviewGraph exists before notifying as we may be entering PostUndo from a transaction that predates the OverviewGraph.
+	if (OverviewGraph)
+	{
+		OverviewGraph->NotifyGraphChanged();
+	}
 }
 
 #undef LOCTEXT_NAMESPACE // NiagaraScriptGraphViewModel
