Merging using //UE5/Main_to_//UE5/Release-Engine-Staging
Do not crash if parameter definitions subscriptiosn goes null.

diff --git a/Engine/Plugins/FX/Niagara/Source/Niagara/Classes/NiagaraParameterDefinitionsSubscriber.h b/Engine/Plugins/FX/Niagara/Source/Niagara/Classes/NiagaraParameterDefinitionsSubscriber.h
--- a/Engine/Plugins/FX/Niagara/Source/Niagara/Classes/NiagaraParameterDefinitionsSubscriber.h
+++ b/Engine/Plugins/FX/Niagara/Source/Niagara/Classes/NiagaraParameterDefinitionsSubscriber.h
@@ -63,8 +63,8 @@ public:
 	 */
 	virtual TArray<INiagaraParameterDefinitionsSubscriber*> GetOwnedParameterDefinitionsSubscribers() { return TArray<INiagaraParameterDefinitionsSubscriber*>(); };
 
-	const TArray<UNiagaraParameterDefinitionsBase*> GetSubscribedParameterDefinitions() const;
-	const TArray<UNiagaraParameterDefinitionsBase*> GetSubscribedParameterDefinitionsPendingSynchronization() const;
+	const TArray<UNiagaraParameterDefinitionsBase*> GetSubscribedParameterDefinitions();
+	const TArray<UNiagaraParameterDefinitionsBase*> GetSubscribedParameterDefinitionsPendingSynchronization();
 
 	void SubscribeToParameterDefinitions(UNiagaraParameterDefinitionsBase* NewParameterDefinitions, bool bDoNotAssertIfAlreadySubscribed = false);
 	void UnsubscribeFromParameterDefinitions(const FGuid& ParameterDefinitionsToRemoveId);
@@ -79,7 +79,7 @@ private:
 
 private:
 	/** Get the unique Id GUID of every UNiagaraScriptVariable owned by every subscribed parameter definitions. */
-	TArray<FGuid> GetSubscribedParameterDefinitionsParameterIds() const;
+	TArray<FGuid> GetSubscribedParameterDefinitionsParameterIds();
 
 	/** Update the cached change Id of specified Synchronized Subscribed Parameter Definition subscriptions so that they are not marked pending sync. */
 	void MarkSubscribedParameterDefinitionsSynchronized(TArray<FGuid> SynchronizedSubscribedParameterDefinitionsIds /*= TArray<FGuid>()*/);
diff --git a/Engine/Plugins/FX/Niagara/Source/Niagara/Private/NiagaraParameterDefinitionsSubscriber.cpp b/Engine/Plugins/FX/Niagara/Source/Niagara/Private/NiagaraParameterDefinitionsSubscriber.cpp
--- a/Engine/Plugins/FX/Niagara/Source/Niagara/Private/NiagaraParameterDefinitionsSubscriber.cpp
+++ b/Engine/Plugins/FX/Niagara/Source/Niagara/Private/NiagaraParameterDefinitionsSubscriber.cpp
@@ -8,6 +8,7 @@
 #include "NiagaraScriptSourceBase.h"
 
 #if WITH_EDITORONLY_DATA
+
 void INiagaraParameterDefinitionsSubscriber::InitParameterDefinitionsSubscriptions()
 {
 	for (UNiagaraParameterDefinitionsBase* ParameterDefinitions : GetSubscribedParameterDefinitions())
@@ -25,27 +26,37 @@ void INiagaraParameterDefinitionsSubscriber::CleanupParameterDefinitionsSubscrip
 	}
 }
 
-const TArray<UNiagaraParameterDefinitionsBase*> INiagaraParameterDefinitionsSubscriber::GetSubscribedParameterDefinitions() const
+const TArray<UNiagaraParameterDefinitionsBase*> INiagaraParameterDefinitionsSubscriber::GetSubscribedParameterDefinitions()
 {
-	const TArray<FParameterDefinitionsSubscription>& Subscriptions = GetParameterDefinitionsSubscriptions();
+	TArray<FParameterDefinitionsSubscription>& Subscriptions = GetParameterDefinitionsSubscriptions();
 	TArray<UNiagaraParameterDefinitionsBase*> OutParameterDefinitions;
-	OutParameterDefinitions.AddUninitialized(Subscriptions.Num());
-	for (int32 Idx = 0; Idx < Subscriptions.Num(); ++Idx)
+	for (int32 Idx = Subscriptions.Num() - 1; Idx > -1; --Idx)
 	{
-		OutParameterDefinitions[Idx] = Subscriptions[Idx].ParameterDefinitions;
+		if (Subscriptions[Idx].ParameterDefinitions == nullptr)
+		{
+			Subscriptions.RemoveAt(Idx);
+		}
+		else
+		{
+			OutParameterDefinitions.Add(Subscriptions[Idx].ParameterDefinitions);
+		}
 	}
 	return OutParameterDefinitions;
 }
 
-const TArray<UNiagaraParameterDefinitionsBase*> INiagaraParameterDefinitionsSubscriber::GetSubscribedParameterDefinitionsPendingSynchronization() const
+const TArray<UNiagaraParameterDefinitionsBase*> INiagaraParameterDefinitionsSubscriber::GetSubscribedParameterDefinitionsPendingSynchronization()
 {
-	const TArray<FParameterDefinitionsSubscription>& Subscriptions = GetParameterDefinitionsSubscriptions();
+	TArray<FParameterDefinitionsSubscription>& Subscriptions = GetParameterDefinitionsSubscriptions();
 	TArray<UNiagaraParameterDefinitionsBase*> OutParameterDefinitions;
-	for (const FParameterDefinitionsSubscription& Subscription : Subscriptions)
+	for (int32 Idx = Subscriptions.Num() - 1; Idx > -1; --Idx)
 	{
-		if (Subscription.ParameterDefinitions->GetChangeIdHash() != Subscription.CachedChangeIdHash)
+		if (Subscriptions[Idx].ParameterDefinitions == nullptr)
+		{
+			Subscriptions.RemoveAt(Idx);
+		}
+		else if (Subscriptions[Idx].ParameterDefinitions->GetChangeIdHash() != Subscriptions[Idx].CachedChangeIdHash)
 		{
-			OutParameterDefinitions.Add(Subscription.ParameterDefinitions);
+			OutParameterDefinitions.Add(Subscriptions[Idx].ParameterDefinitions);
 		}
 	}
 	return OutParameterDefinitions;
@@ -148,7 +159,7 @@ UNiagaraParameterDefinitionsBase* INiagaraParameterDefinitionsSubscriber::FindSu
 	return nullptr;
 }
 
-TArray<FGuid> INiagaraParameterDefinitionsSubscriber::GetSubscribedParameterDefinitionsParameterIds() const
+TArray<FGuid> INiagaraParameterDefinitionsSubscriber::GetSubscribedParameterDefinitionsParameterIds()
 {
 	TArray<FGuid> Ids;
 	for (const UNiagaraParameterDefinitionsBase* Definitions : GetSubscribedParameterDefinitions())
