Merging using //UE5/Main_to_//UE5/Release-Engine-Staging
Postload definitions subscriptions to remove null members so that we do not have hanging nullptrs during asset reference check.

diff --git a/Engine/Plugins/FX/Niagara/Source/Niagara/Classes/NiagaraParameterDefinitionsSubscriber.h b/Engine/Plugins/FX/Niagara/Source/Niagara/Classes/NiagaraParameterDefinitionsSubscriber.h
--- a/Engine/Plugins/FX/Niagara/Source/Niagara/Classes/NiagaraParameterDefinitionsSubscriber.h
+++ b/Engine/Plugins/FX/Niagara/Source/Niagara/Classes/NiagaraParameterDefinitionsSubscriber.h
@@ -43,6 +43,7 @@ public:
 
 	void InitParameterDefinitionsSubscriptions();
 	void CleanupParameterDefinitionsSubscriptions();
+	void PostLoadDefinitionsSubscriptions();
 
 	//~ Begin Pure Virtual Methods
 	virtual const TArray<FParameterDefinitionsSubscription>& GetParameterDefinitionsSubscriptions() const = 0;
diff --git a/Engine/Plugins/FX/Niagara/Source/Niagara/Private/NiagaraEmitter.cpp b/Engine/Plugins/FX/Niagara/Source/Niagara/Private/NiagaraEmitter.cpp
--- a/Engine/Plugins/FX/Niagara/Source/Niagara/Private/NiagaraEmitter.cpp
+++ b/Engine/Plugins/FX/Niagara/Source/Niagara/Private/NiagaraEmitter.cpp
@@ -379,6 +379,10 @@ void UNiagaraEmitter::PostLoad()
 {
 	Super::PostLoad();
 
+#if WITH_EDITORONLY_DATA
+	PostLoadDefinitionsSubscriptions();
+#endif
+
 	if (GIsEditor)
 	{
 		SetFlags(RF_Transactional);
diff --git a/Engine/Plugins/FX/Niagara/Source/Niagara/Private/NiagaraParameterDefinitionsSubscriber.cpp b/Engine/Plugins/FX/Niagara/Source/Niagara/Private/NiagaraParameterDefinitionsSubscriber.cpp
--- a/Engine/Plugins/FX/Niagara/Source/Niagara/Private/NiagaraParameterDefinitionsSubscriber.cpp
+++ b/Engine/Plugins/FX/Niagara/Source/Niagara/Private/NiagaraParameterDefinitionsSubscriber.cpp
@@ -26,6 +26,18 @@ void INiagaraParameterDefinitionsSubscriber::CleanupParameterDefinitionsSubscrip
 	}
 }
 
+void INiagaraParameterDefinitionsSubscriber::PostLoadDefinitionsSubscriptions()
+{
+	TArray<FParameterDefinitionsSubscription>& Subscriptions = GetParameterDefinitionsSubscriptions();
+	for (int32 Idx = Subscriptions.Num() - 1; Idx > -1; --Idx)
+	{
+		if (Subscriptions[Idx].ParameterDefinitions == nullptr)
+		{
+			Subscriptions.RemoveAt(Idx);
+		}
+	}
+}
+
 const TArray<UNiagaraParameterDefinitionsBase*> INiagaraParameterDefinitionsSubscriber::GetSubscribedParameterDefinitions()
 {
 	TArray<FParameterDefinitionsSubscription>& Subscriptions = GetParameterDefinitionsSubscriptions();
diff --git a/Engine/Plugins/FX/Niagara/Source/Niagara/Private/NiagaraScript.cpp b/Engine/Plugins/FX/Niagara/Source/Niagara/Private/NiagaraScript.cpp
--- a/Engine/Plugins/FX/Niagara/Source/Niagara/Private/NiagaraScript.cpp
+++ b/Engine/Plugins/FX/Niagara/Source/Niagara/Private/NiagaraScript.cpp
@@ -1474,12 +1474,17 @@ void UNiagaraScript::PostLoad()
 
 			// Synchronize with Definitions after source scripts have been postloaded.
 			FVersionedNiagaraScript& VersionedScriptAdapter = VersionedScriptAdapters.Emplace_GetRef(this, Data.Version.VersionGuid);
+			VersionedScriptAdapter.PostLoadDefinitionsSubscriptions();
 			const UNiagaraSettings* Settings = GetDefault<UNiagaraSettings>();
 			check(Settings);
 			TArray<FGuid> DefaultDefinitionsUniqueIds;
 			for (const FSoftObjectPath& DefaultLinkedParameterDefinitionObjPath : Settings->DefaultLinkedParameterDefinitions)
 			{
-				UNiagaraParameterDefinitionsBase* DefaultLinkedParameterDefinitions = CastChecked<UNiagaraParameterDefinitionsBase>(DefaultLinkedParameterDefinitionObjPath.TryLoad());
+				UNiagaraParameterDefinitionsBase* DefaultLinkedParameterDefinitions = Cast<UNiagaraParameterDefinitionsBase>(DefaultLinkedParameterDefinitionObjPath.TryLoad());
+				if (!DefaultLinkedParameterDefinitions)
+				{
+					continue;
+				}
 				DefaultDefinitionsUniqueIds.Add(DefaultLinkedParameterDefinitions->GetDefinitionsUniqueId());
 				const bool bDoNotAssertIfAlreadySubscribed = true;
 				VersionedScriptAdapter.SubscribeToParameterDefinitions(DefaultLinkedParameterDefinitions, bDoNotAssertIfAlreadySubscribed);
diff --git a/Engine/Plugins/FX/Niagara/Source/Niagara/Private/NiagaraSystem.cpp b/Engine/Plugins/FX/Niagara/Source/Niagara/Private/NiagaraSystem.cpp
--- a/Engine/Plugins/FX/Niagara/Source/Niagara/Private/NiagaraSystem.cpp
+++ b/Engine/Plugins/FX/Niagara/Source/Niagara/Private/NiagaraSystem.cpp
@@ -722,6 +722,10 @@ void UNiagaraSystem::PostLoad()
 {
 	Super::PostLoad();
 
+#if WITH_EDITORONLY_DATA
+	PostLoadDefinitionsSubscriptions();
+#endif
+
 	// Workaround for UE-104235 where a CDO loads a NiagaraSystem before the NiagaraModule has had a chance to load
 	// We force the module to load here we makes sure the type registry, etc, is all setup in time.
 	static bool bLoadChecked = false;
