Merging using //UE5/Main_to_//UE5/Release-Engine-Staging
Fix suggested parameters to add being culled by parameters not in current selection set.

diff --git a/Engine/Plugins/FX/Niagara/Source/NiagaraEditor/Public/Widgets/SNiagaraParameterMenu.cpp b/Engine/Plugins/FX/Niagara/Source/NiagaraEditor/Public/Widgets/SNiagaraParameterMenu.cpp
--- a/Engine/Plugins/FX/Niagara/Source/NiagaraEditor/Public/Widgets/SNiagaraParameterMenu.cpp
+++ b/Engine/Plugins/FX/Niagara/Source/NiagaraEditor/Public/Widgets/SNiagaraParameterMenu.cpp
@@ -455,7 +455,6 @@ void SNiagaraAddParameterFromPanelMenu::CollectAllActions(FGraphActionListBuilde
 	}
 
 	// Collect "add parameter from parameter definition asset" actions
-	const TSet<FName> GraphParameterNames = bShowGraphParameters ? VisitedParameterNames : GetAllGraphParameterNames();
 	for (UNiagaraParameterDefinitions* ParameterDefinitions : AvailableParameterDefinitions)
 	{
 		bool bTopLevelCategory = ParameterDefinitions->GetIsPromotedToTopInAddMenus();
@@ -475,7 +474,7 @@ void SNiagaraAddParameterFromPanelMenu::CollectAllActions(FGraphActionListBuilde
 			{
 				continue;
 			}
-			else if (GraphParameterNames.Contains(ScriptVar->Variable.GetName()))
+			else if (VisitedParameterNames.Contains(ScriptVar->Variable.GetName()))
 			{
 				continue;
 			}
