Merging using //UE5/Main_to_//UE5/Release-Engine-Staging
Fix applying changes to parameter definitions failing to apply to new assets. No longer use OnChanged() binding at UNiagaraParameterDefinitions level.
Fix FNiagaraScriptVariableCustomization not updating post change from parameter definitions synchronization.

diff --git a/Engine/Plugins/FX/Niagara/Source/Niagara/Classes/NiagaraParameterDefinitionsSubscriber.h b/Engine/Plugins/FX/Niagara/Source/Niagara/Classes/NiagaraParameterDefinitionsSubscriber.h
--- a/Engine/Plugins/FX/Niagara/Source/Niagara/Classes/NiagaraParameterDefinitionsSubscriber.h
+++ b/Engine/Plugins/FX/Niagara/Source/Niagara/Classes/NiagaraParameterDefinitionsSubscriber.h
@@ -41,8 +41,6 @@ class NIAGARA_API INiagaraParameterDefinitionsSubscriber
 public:
 	virtual ~INiagaraParameterDefinitionsSubscriber() = default;
 
-	void InitParameterDefinitionsSubscriptions();
-	void CleanupParameterDefinitionsSubscriptions();
 	void PostLoadDefinitionsSubscriptions();
 
 	//~ Begin Pure Virtual Methods
diff --git a/Engine/Plugins/FX/Niagara/Source/Niagara/Private/NiagaraEmitter.cpp b/Engine/Plugins/FX/Niagara/Source/Niagara/Private/NiagaraEmitter.cpp
--- a/Engine/Plugins/FX/Niagara/Source/Niagara/Private/NiagaraEmitter.cpp
+++ b/Engine/Plugins/FX/Niagara/Source/Niagara/Private/NiagaraEmitter.cpp
@@ -534,12 +534,17 @@ void UNiagaraEmitter::PostLoad()
 		}
 
 		// Synchronize with definitions before merging.
+		// First force sync with all definitions in the DefaultLinkedParameterDefinitions array.
 		const UNiagaraSettings* Settings = GetDefault<UNiagaraSettings>();
 		check(Settings);
 		TArray<FGuid> DefaultDefinitionsUniqueIds;
 		for (const FSoftObjectPath& DefaultLinkedParameterDefinitionObjPath : Settings->DefaultLinkedParameterDefinitions)
 		{
-			UNiagaraParameterDefinitionsBase* DefaultLinkedParameterDefinitions = CastChecked<UNiagaraParameterDefinitionsBase>(DefaultLinkedParameterDefinitionObjPath.TryLoad());
+			UNiagaraParameterDefinitionsBase* DefaultLinkedParameterDefinitions = Cast<UNiagaraParameterDefinitionsBase>(DefaultLinkedParameterDefinitionObjPath.TryLoad());
+			if (DefaultLinkedParameterDefinitions == nullptr)
+			{
+				continue;
+			}
 			DefaultDefinitionsUniqueIds.Add(DefaultLinkedParameterDefinitions->GetDefinitionsUniqueId());
 			const bool bDoNotAssertIfAlreadySubscribed = true;
 			SubscribeToParameterDefinitions(DefaultLinkedParameterDefinitions, bDoNotAssertIfAlreadySubscribed);
@@ -549,7 +554,9 @@ void UNiagaraEmitter::PostLoad()
 		Args.bForceSynchronizeDefinitions = true;
 		Args.bSubscribeAllNameMatchParameters = true;
 		SynchronizeWithParameterDefinitions(Args);
-		InitParameterDefinitionsSubscriptions();
+		
+		// After forcing syncing all DefaultLinkedParameterDefinitions, call SynchronizeWithParameterDefinitions again to sync with all definitions the emitter was already subscribed to, and do not force the sync.
+		SynchronizeWithParameterDefinitions();
 
 		if (IsSynchronizedWithParent() == false)
 		{
@@ -1967,8 +1974,6 @@ void UNiagaraEmitter::BeginDestroy()
 	{
 		GPUComputeScript->OnGPUScriptCompiled().RemoveAll(this);
 	}
-
-	CleanupParameterDefinitionsSubscriptions();
 #endif
 	Super::BeginDestroy();
 }
diff --git a/Engine/Plugins/FX/Niagara/Source/Niagara/Private/NiagaraParameterDefinitionsSubscriber.cpp b/Engine/Plugins/FX/Niagara/Source/Niagara/Private/NiagaraParameterDefinitionsSubscriber.cpp
--- a/Engine/Plugins/FX/Niagara/Source/Niagara/Private/NiagaraParameterDefinitionsSubscriber.cpp
+++ b/Engine/Plugins/FX/Niagara/Source/Niagara/Private/NiagaraParameterDefinitionsSubscriber.cpp
@@ -9,23 +9,6 @@
 
 #if WITH_EDITORONLY_DATA
 
-void INiagaraParameterDefinitionsSubscriber::InitParameterDefinitionsSubscriptions()
-{
-	for (UNiagaraParameterDefinitionsBase* ParameterDefinitions : GetSubscribedParameterDefinitions())
-	{
-		ParameterDefinitions->GetOnParameterDefinitionsChanged().AddRaw(this, &INiagaraParameterDefinitionsSubscriber::SynchronizeWithParameterDefinitions, FSynchronizeWithParameterDefinitionsArgs());
-	}
-	SynchronizeWithParameterDefinitions();
-}
-
-void INiagaraParameterDefinitionsSubscriber::CleanupParameterDefinitionsSubscriptions()
-{
-	for (UNiagaraParameterDefinitionsBase* ParameterDefinitions : GetSubscribedParameterDefinitions())
-	{
-		ParameterDefinitions->GetOnParameterDefinitionsChanged().RemoveAll(this);
-	}
-}
-
 void INiagaraParameterDefinitionsSubscriber::PostLoadDefinitionsSubscriptions()
 {
 	TArray<FParameterDefinitionsSubscription>& Subscriptions = GetParameterDefinitionsSubscriptions();
diff --git a/Engine/Plugins/FX/Niagara/Source/Niagara/Private/NiagaraScript.cpp b/Engine/Plugins/FX/Niagara/Source/Niagara/Private/NiagaraScript.cpp
--- a/Engine/Plugins/FX/Niagara/Source/Niagara/Private/NiagaraScript.cpp
+++ b/Engine/Plugins/FX/Niagara/Source/Niagara/Private/NiagaraScript.cpp
@@ -1473,6 +1473,7 @@ void UNiagaraScript::PostLoad()
         	Source->ConditionalPostLoad();
 
 			// Synchronize with Definitions after source scripts have been postloaded.
+			// First force sync with all definitions in the DefaultLinkedParameterDefinitions array.
 			FVersionedNiagaraScript& VersionedScriptAdapter = VersionedScriptAdapters.Emplace_GetRef(this, Data.Version.VersionGuid);
 			VersionedScriptAdapter.PostLoadDefinitionsSubscriptions();
 			const UNiagaraSettings* Settings = GetDefault<UNiagaraSettings>();
@@ -1494,7 +1495,9 @@ void UNiagaraScript::PostLoad()
 			Args.bForceSynchronizeDefinitions = true;
 			Args.bSubscribeAllNameMatchParameters = true;
 			VersionedScriptAdapter.SynchronizeWithParameterDefinitions(Args);
-			VersionedScriptAdapter.InitParameterDefinitionsSubscriptions();
+
+			// After forcing syncing all DefaultLinkedParameterDefinitions, call SynchronizeWithParameterDefinitions again to sync with all definitions the script was already subscribed to, and do not force the sync.
+			VersionedScriptAdapter.SynchronizeWithParameterDefinitions();
 
         	bool bScriptVMNeedsRebuild = false;
         	FString RebuildReason;
@@ -1961,7 +1964,6 @@ FVersionedNiagaraScript FVersionedNiagaraScriptWeakPtr::Pin()
 	if (Script.IsValid())
 	{
 		FVersionedNiagaraScript PinnedVersionedNiagaraScript = FVersionedNiagaraScript(Script.Get(), Version);
-		PinnedVersionedNiagaraScript.InitParameterDefinitionsSubscriptions(); //@todo(ng) cleanup
 		return PinnedVersionedNiagaraScript;
 	}
 	return FVersionedNiagaraScript();
@@ -1998,7 +2000,6 @@ FString FVersionedNiagaraScript::GetSourceObjectPathName() const
 FVersionedNiagaraScriptWeakPtr FVersionedNiagaraScript::ToWeakPtr()
 {
 	FVersionedNiagaraScriptWeakPtr WeakVersionedNiagaraScript = FVersionedNiagaraScriptWeakPtr(Script, Version);
-	WeakVersionedNiagaraScript.InitParameterDefinitionsSubscriptions(); //@todo(ng) cleanup
 	return WeakVersionedNiagaraScript;
 }
 
@@ -2352,13 +2353,6 @@ void UNiagaraScript::BeginDestroy()
 {
 	Super::BeginDestroy();
 
-#if WITH_EDITOR
-	for (FVersionedNiagaraScript& VersionedScriptAdapter : VersionedScriptAdapters)
-	{
-		VersionedScriptAdapter.CleanupParameterDefinitionsSubscriptions();
-	}
-#endif
-
 	if (!HasAnyFlags(RF_ClassDefaultObject) && ScriptResource)
 	{
 		if (!ScriptResource->QueueForRelease(ReleasedByRT))
diff --git a/Engine/Plugins/FX/Niagara/Source/Niagara/Private/NiagaraSystem.cpp b/Engine/Plugins/FX/Niagara/Source/Niagara/Private/NiagaraSystem.cpp
--- a/Engine/Plugins/FX/Niagara/Source/Niagara/Private/NiagaraSystem.cpp
+++ b/Engine/Plugins/FX/Niagara/Source/Niagara/Private/NiagaraSystem.cpp
@@ -184,8 +184,6 @@ void UNiagaraSystem::BeginDestroy()
 	{
 		QueryCompileComplete(true, false, true);
 	}
-
-	CleanupParameterDefinitionsSubscriptions();
 #endif
 
 	//Should we just destroy all system sims here to simplify cleanup?
@@ -818,6 +816,31 @@ void UNiagaraSystem::PostLoad()
 		EditorParameters = NiagaraModule.GetEditorOnlyDataUtilities().CreateDefaultEditorParameters(this);
 	}
 
+	// Synchronize with parameter definitions
+	// First force sync with all definitions in the DefaultLinkedParameterDefinitions array.
+	const UNiagaraSettings* Settings = GetDefault<UNiagaraSettings>();
+	check(Settings);
+	TArray<FGuid> DefaultDefinitionsUniqueIds;
+	for (const FSoftObjectPath& DefaultLinkedParameterDefinitionObjPath : Settings->DefaultLinkedParameterDefinitions)
+	{
+		UNiagaraParameterDefinitionsBase* DefaultLinkedParameterDefinitions = Cast<UNiagaraParameterDefinitionsBase>(DefaultLinkedParameterDefinitionObjPath.TryLoad());
+		if (DefaultLinkedParameterDefinitions == nullptr)
+		{
+			continue;
+		}
+		DefaultDefinitionsUniqueIds.Add(DefaultLinkedParameterDefinitions->GetDefinitionsUniqueId());
+		const bool bDoNotAssertIfAlreadySubscribed = true;
+		SubscribeToParameterDefinitions(DefaultLinkedParameterDefinitions, bDoNotAssertIfAlreadySubscribed);
+	}
+	FSynchronizeWithParameterDefinitionsArgs Args;
+	Args.SpecificDefinitionsUniqueIds = DefaultDefinitionsUniqueIds;
+	Args.bForceSynchronizeDefinitions = true;
+	Args.bSubscribeAllNameMatchParameters = true;
+	SynchronizeWithParameterDefinitions(Args);
+
+	// After forcing syncing all DefaultLinkedParameterDefinitions, call SynchronizeWithParameterDefinitions again to sync with all definitions the system was already subscribed to, and do not force the sync.
+	SynchronizeWithParameterDefinitions();
+	
 	// see the equivalent in NiagaraEmitter for details
 	if(bIsTemplateAsset_DEPRECATED)
 	{
diff --git a/Engine/Plugins/FX/Niagara/Source/NiagaraEditor/Private/Customizations/NiagaraScriptVariableCustomization.cpp b/Engine/Plugins/FX/Niagara/Source/NiagaraEditor/Private/Customizations/NiagaraScriptVariableCustomization.cpp
--- a/Engine/Plugins/FX/Niagara/Source/NiagaraEditor/Private/Customizations/NiagaraScriptVariableCustomization.cpp
+++ b/Engine/Plugins/FX/Niagara/Source/NiagaraEditor/Private/Customizations/NiagaraScriptVariableCustomization.cpp
@@ -80,6 +80,15 @@ void FNiagaraScriptVariableDetails::PostUndo(bool bSuccess)
 			ParameterEditorStaticSwitchValue->UpdateInternalValueFromStruct(ParameterValue.ToSharedRef());
 		}
 	}
+	else if (Variable->GetOuter()->IsA<UNiagaraParameterDefinitions>())
+	{
+		if (TypeUtilityLibraryValue && ParameterEditorLibraryValue)
+		{
+			TSharedPtr<FStructOnScope> ParameterValue = MakeShareable(new FStructOnScope(Variable->Variable.GetType().GetStruct()));
+			Variable->CopyDefaultValueDataTo(ParameterValue->GetStructMemory());
+			ParameterEditorLibraryValue->UpdateInternalValueFromStruct(ParameterValue.ToSharedRef());
+		}
+	}
 	else
 	{
 		if (UEdGraphPin* Pin = GetAnyDefaultPin())
@@ -354,8 +363,9 @@ void FNiagaraScriptVariableDetails::AddLibraryDefaultValueCustomRow(IDetailCateg
 		Variable->CopyDefaultValueDataTo(ParameterValue->GetStructMemory());
 		ParameterEditorLibraryValue->SetEnabled(bInLibraryAsset);
 		ParameterEditorLibraryValue->UpdateInternalValueFromStruct(ParameterValue.ToSharedRef());
+		ParameterEditorLibraryValue->SetOnBeginValueChange(SNiagaraParameterEditor::FOnValueChange::CreateSP(this, &FNiagaraScriptVariableDetails::OnBeginLibraryValueChanged));
 		ParameterEditorLibraryValue->SetOnValueChanged(SNiagaraParameterEditor::FOnValueChange::CreateSP(this, &FNiagaraScriptVariableDetails::OnLibraryValueChanged));
-		ParameterEditorLibraryValue->SetOnEndValueChange(SNiagaraParameterEditor::FOnValueChange::CreateSP(this, &FNiagaraScriptVariableDetails::OnEndValueChanged));
+		ParameterEditorLibraryValue->SetOnEndValueChange(SNiagaraParameterEditor::FOnValueChange::CreateSP(this, &FNiagaraScriptVariableDetails::OnEndLibraryValueChanged));
 
 		FDetailWidgetRow& DefaultValueWidget = CategoryBuilder.AddCustomRow(LOCTEXT("DefaultValueFilterText", "Default Value"));
 		DefaultValueWidget.NameContent()
@@ -416,6 +426,14 @@ void FNiagaraScriptVariableDetails::OnValueChanged()
 		}
 		else
 		{
+			// Values that can change continuously are not guaranteed to call OnBeginChange() before OnChange(), so if a transaction has not been begun by OnBeginChange(), create one now.
+			bool bOnValueChangedTransactionActive = false;
+			if (GEditor->IsTransactionActive() == false)
+			{
+				bOnValueChangedTransactionActive = true;
+				GEditor->BeginTransaction(NSLOCTEXT("ScriptVariableCustomization", "ChangeValue", "Change Default Value"));
+				Variable->Modify();
+			}
 			TSharedPtr<FStructOnScope> ParameterValue = MakeShareable(new FStructOnScope(Variable->Variable.GetType().GetStruct()));
 			ParameterEditorValue->UpdateStructFromInternalValue(ParameterValue.ToSharedRef());
 			Variable->Variable.SetData(ParameterValue->GetStructMemory());
@@ -423,8 +441,17 @@ void FNiagaraScriptVariableDetails::OnValueChanged()
 
 			for (UEdGraphPin* Pin : GetDefaultPins())
 			{
+				if (bOnValueChangedTransactionActive)
+				{
+					Pin->Modify();
+				}
 				GetDefault<UEdGraphSchema_Niagara>()->TrySetDefaultValue(*Pin, NewDefaultValue, true);
 			}
+
+			if (bOnValueChangedTransactionActive)
+			{
+				GEditor->EndTransaction();
+			}
 		}
 	}
 }
@@ -481,6 +508,44 @@ void FNiagaraScriptVariableDetails::OnStaticSwitchValueChanged()
 	} 
 }
 
+void FNiagaraScriptVariableDetails::OnBeginLibraryValueChanged()
+{
+	if (!ParameterEditorLibraryValue->CanChangeContinuously())
+	{
+		return;
+	}
+
+	if (TypeUtilityLibraryValue && ParameterEditorLibraryValue && CachedDetailBuilder.IsValid())
+	{
+		const TSharedPtr<IPropertyHandle> DefaultValueHandle = CachedDetailBuilder.Pin()->GetProperty("DefaultValueVariant", UNiagaraScriptVariable::StaticClass());
+		GEditor->BeginTransaction(NSLOCTEXT("ScriptVariableCustomization", "ChangeLibraryValue", "Change Default Value"));
+		DefaultValueHandle->NotifyPreChange();
+		Variable->Modify();
+		TSharedPtr<FStructOnScope> ParameterValue = MakeShareable(new FStructOnScope(Variable->Variable.GetType().GetStruct()));
+		ParameterEditorLibraryValue->UpdateStructFromInternalValue(ParameterValue.ToSharedRef());
+		Variable->SetDefaultValueData(ParameterValue->GetStructMemory());
+	}
+}
+
+void FNiagaraScriptVariableDetails::OnEndLibraryValueChanged()
+{
+	Variable->UpdateChangeId();
+
+	if (TypeUtilityLibraryValue && ParameterEditorLibraryValue && CachedDetailBuilder.IsValid())
+	{
+		const TSharedPtr<IPropertyHandle> DefaultValueHandle = CachedDetailBuilder.Pin()->GetProperty("DefaultValueVariant", UNiagaraScriptVariable::StaticClass());
+		DefaultValueHandle->NotifyPostChange(EPropertyChangeType::ValueSet);
+		DefaultValueHandle->NotifyFinishedChangingProperties();
+	}
+
+	if (UNiagaraParameterDefinitions* OuterParameterDefinitions = Cast<UNiagaraParameterDefinitions>(Variable->GetOuter()))
+	{
+		OuterParameterDefinitions->NotifyParameterDefinitionsChanged();
+	}
+
+	OnEndValueChanged();
+}
+
 void FNiagaraScriptVariableDetails::OnLibraryValueChanged()
 {
 	if (TypeUtilityLibraryValue && ParameterEditorLibraryValue && CachedDetailBuilder.IsValid())
@@ -489,21 +554,48 @@ void FNiagaraScriptVariableDetails::OnLibraryValueChanged()
 		if (!ParameterEditorLibraryValue->CanChangeContinuously())
 		{
 			const FScopedTransaction Transaction(NSLOCTEXT("ScriptVariableCustomization", "ChangeLibraryValue", "Change Default Value"));
+			DefaultValueHandle->NotifyPreChange();
 			Variable->Modify();
-
 			TSharedPtr<FStructOnScope> ParameterValue = MakeShareable(new FStructOnScope(Variable->Variable.GetType().GetStruct()));
 			ParameterEditorLibraryValue->UpdateStructFromInternalValue(ParameterValue.ToSharedRef());
-			DefaultValueHandle->NotifyPreChange();
 			Variable->SetDefaultValueData(ParameterValue->GetStructMemory());
+
+			if (UNiagaraParameterDefinitions* OuterParameterDefinitions = Cast<UNiagaraParameterDefinitions>(Variable->GetOuter()))
+			{
+				OuterParameterDefinitions->NotifyParameterDefinitionsChanged();
+			}
 			DefaultValueHandle->NotifyPostChange(EPropertyChangeType::ValueSet);
 			DefaultValueHandle->NotifyFinishedChangingProperties();
 		}
 		else
 		{
+			// Values that can change continuously are not guaranteed to call OnBeginChange() before OnChange(), so if a transaction has not been begun by OnBeginChange(), create one now.
+			bool bOnValueChangedTransactionActive = false;
+			if (GEditor->IsTransactionActive() == false)
+			{
+				bOnValueChangedTransactionActive = true;
+				GEditor->BeginTransaction(NSLOCTEXT("ScriptVariableCustomization", "ChangeValue", "Change Default Value"));
+				Variable->Modify();
+			}
+			DefaultValueHandle->NotifyPreChange();
 			TSharedPtr<FStructOnScope> ParameterValue = MakeShareable(new FStructOnScope(Variable->Variable.GetType().GetStruct()));
 			ParameterEditorLibraryValue->UpdateStructFromInternalValue(ParameterValue.ToSharedRef());
-			DefaultValueHandle->NotifyPreChange();
 			Variable->SetDefaultValueData(ParameterValue->GetStructMemory());
+
+			if (bOnValueChangedTransactionActive)
+			{
+				Variable->UpdateChangeId();
+
+				DefaultValueHandle->NotifyPostChange(EPropertyChangeType::ValueSet);
+				DefaultValueHandle->NotifyFinishedChangingProperties();
+
+				if (UNiagaraParameterDefinitions* OuterParameterDefinitions = Cast<UNiagaraParameterDefinitions>(Variable->GetOuter()))
+				{
+					OuterParameterDefinitions->NotifyParameterDefinitionsChanged();
+				}
+
+				GEditor->EndTransaction();
+			}
 			DefaultValueHandle->NotifyPostChange(EPropertyChangeType::ValueSet);
 			DefaultValueHandle->NotifyFinishedChangingProperties();
 		}
diff --git a/Engine/Plugins/FX/Niagara/Source/NiagaraEditor/Private/Customizations/NiagaraScriptVariableCustomization.h b/Engine/Plugins/FX/Niagara/Source/NiagaraEditor/Private/Customizations/NiagaraScriptVariableCustomization.h
--- a/Engine/Plugins/FX/Niagara/Source/NiagaraEditor/Private/Customizations/NiagaraScriptVariableCustomization.h
+++ b/Engine/Plugins/FX/Niagara/Source/NiagaraEditor/Private/Customizations/NiagaraScriptVariableCustomization.h
@@ -79,7 +79,10 @@ private:
 
 	void OnStaticSwitchValueChanged();
 
+	void OnBeginLibraryValueChanged();
+	void OnEndLibraryValueChanged();
 	void OnLibraryValueChanged();
+
 	int32 GetLibraryDefaultModeValue() const { return LibraryDefaultModeValue; };
 	int32 GetLibrarySourcedDefaultModeInitialValue() const;
 	int32 GetLibrarySynchronizedDefaultModeInitialValue() const;
diff --git a/Engine/Plugins/FX/Niagara/Source/NiagaraEditor/Private/NiagaraEditorModule.cpp b/Engine/Plugins/FX/Niagara/Source/NiagaraEditor/Private/NiagaraEditorModule.cpp
--- a/Engine/Plugins/FX/Niagara/Source/NiagaraEditor/Private/NiagaraEditorModule.cpp
+++ b/Engine/Plugins/FX/Niagara/Source/NiagaraEditor/Private/NiagaraEditorModule.cpp
@@ -1228,6 +1228,14 @@ void FNiagaraEditorModule::OnPostEngineInit()
 		GEditor->OnExecParticleInvoked().AddRaw(this, &FNiagaraEditorModule::OnExecParticleInvoked);
 
 		PreviewPlatformChangedHandle = CastChecked<UEditorEngine>(GEngine)->OnPreviewPlatformChanged().AddRaw(this, &FNiagaraEditorModule::OnPreviewPlatformChanged);
+
+		// Preload all parameter definitions in the default linked settings so that they will be postloaded before postload calls to scripts/emitters/systems that rely on them.
+		const UNiagaraSettings* Settings = GetDefault<UNiagaraSettings>();
+		check(Settings);
+		for (const FSoftObjectPath& DefaultLinkedParameterDefinitionSoftPath : Settings->DefaultLinkedParameterDefinitions)
+		{
+			DefaultLinkedParameterDefinitionSoftPath.TryLoad();
+		}
 	}
 	else
 	{
diff --git a/Engine/Plugins/FX/Niagara/Source/NiagaraEditor/Private/NiagaraParameterDefinitions.cpp b/Engine/Plugins/FX/Niagara/Source/NiagaraEditor/Private/NiagaraParameterDefinitions.cpp
--- a/Engine/Plugins/FX/Niagara/Source/NiagaraEditor/Private/NiagaraParameterDefinitions.cpp
+++ b/Engine/Plugins/FX/Niagara/Source/NiagaraEditor/Private/NiagaraParameterDefinitions.cpp
@@ -32,7 +32,6 @@ void UNiagaraParameterDefinitions::PostLoad()
 		}
 
 		SynchronizeWithSubscribedParameterDefinitions();
-		InitBindings();
 	}
 }
 
@@ -42,14 +41,6 @@ void UNiagaraParameterDefinitions::PostEditChangeProperty(struct FPropertyChange
 	OnParameterDefinitionsChangedDelegate.Broadcast();
 }
 
-void UNiagaraParameterDefinitions::InitBindings()
-{
-	for (FParameterDefinitionsBindingNameSubscription& Subscription : ExternalParameterDefinitionsSubscriptions)
-	{
-		Subscription.SubscribedParameterDefinitions->GetOnParameterDefinitionsChanged().AddUObject(this, &UNiagaraParameterDefinitions::SynchronizeWithSubscribedParameterDefinitions);
-	}
-}
-
 void UNiagaraParameterDefinitions::AddParameter(const FNiagaraVariable& NewVariable)
 {
 	Modify();
diff --git a/Engine/Plugins/FX/Niagara/Source/NiagaraEditor/Private/Toolkits/NiagaraParameterDefinitionsToolkit.cpp b/Engine/Plugins/FX/Niagara/Source/NiagaraEditor/Private/Toolkits/NiagaraParameterDefinitionsToolkit.cpp
--- a/Engine/Plugins/FX/Niagara/Source/NiagaraEditor/Private/Toolkits/NiagaraParameterDefinitionsToolkit.cpp
+++ b/Engine/Plugins/FX/Niagara/Source/NiagaraEditor/Private/Toolkits/NiagaraParameterDefinitionsToolkit.cpp
@@ -35,6 +35,10 @@ FNiagaraParameterDefinitionsToolkit::~FNiagaraParameterDefinitionsToolkit()
 	{
 		ParameterPanelViewModel->Cleanup();
 	}
+	if (SelectedScriptVarDetailsWidget.IsValid())
+	{
+		SelectedScriptVarDetailsWidget->OnFinishedChangingProperties().RemoveAll(this);
+	}
 }
 
 void FNiagaraParameterDefinitionsToolkit::RegisterTabSpawners(const TSharedRef<class FTabManager>& InTabManager)
@@ -49,7 +53,7 @@ void FNiagaraParameterDefinitionsToolkit::RegisterTabSpawners(const TSharedRef<c
 		.SetDisplayName(LOCTEXT("ParameterDefinitionsDetailsTab", "Options"))
 		.SetGroup(WorkspaceMenuCategoryRef);
 
-	InTabManager->RegisterTabSpawner(SelectedDetailsTabId, FOnSpawnTab::CreateSP(this, &FNiagaraParameterDefinitionsToolkit::SpawnTab_SelectedDetails))
+	InTabManager->RegisterTabSpawner(SelectedDetailsTabId, FOnSpawnTab::CreateSP(this, &FNiagaraParameterDefinitionsToolkit::SpawnTab_SelectedScriptVarDetails))
 		.SetDisplayName(LOCTEXT("SelectedDetailsTab", "Selected Details"))
 		.SetGroup(WorkspaceMenuCategoryRef)
 		.SetIcon(FSlateIcon(FEditorStyle::GetStyleSetName(), "LevelEditor.Tabs.Details"));
@@ -78,6 +82,7 @@ void FNiagaraParameterDefinitionsToolkit::Initialize(const EToolkitMode::Type Mo
 	ParameterDefinitionsInstance = Cast<UNiagaraParameterDefinitions>(StaticDuplicateObject(ParameterDefinitionsSource, GetTransientPackage()));
 	ParameterDefinitionsInstance->ClearFlags(RF_Standalone | RF_Public);
 	LastSyncedDefinitionsChangeIdHash = ParameterDefinitionsInstance->GetChangeIdHash();
+	ParameterDefinitionsInstance->GetOnParameterDefinitionsChanged().AddRaw(this, &FNiagaraParameterDefinitionsToolkit::OnEditedParameterDefinitionsChanged);
 
 	// Determine display name for panel heading based on asset usage type
 	FText DisplayName = LOCTEXT("NiagaraParameterDefinitionsDisplayName", "Niagara Parameter Definitions");
@@ -174,13 +179,15 @@ void FNiagaraParameterDefinitionsToolkit::GetSaveableObjects(TArray<UObject*>& O
 
 void FNiagaraParameterDefinitionsToolkit::SaveAsset_Execute()
 {
-	OnApply();
+	bool bInvokedFromSaveAsset = true;
+	OnApply(bInvokedFromSaveAsset);
 	FAssetEditorToolkit::SaveAsset_Execute();
 }
 
 void FNiagaraParameterDefinitionsToolkit::SaveAssetAs_Execute()
 {
-	OnApply();
+	bool bInvokedFromSaveAsset = true;
+	OnApply(bInvokedFromSaveAsset);
 	FAssetEditorToolkit::SaveAsset_Execute();
 }
 
@@ -228,7 +235,6 @@ TSharedRef<SDockTab> FNiagaraParameterDefinitionsToolkit::SpawnTab_ParameterDefi
 	DetailsViewArgs.bHideSelectionTip = true;
 	ParameterDefinitionsDetailsView = PropertyEditorModule.CreateDetailView(DetailsViewArgs);
 
-	ParameterDefinitionsDetailsView->OnFinishedChangingProperties().AddRaw(this, &FNiagaraParameterDefinitionsToolkit::OnEditedParameterDefinitionsPropertyFinishedChanging);
 	ParameterDefinitionsDetailsView->SetObject(ParameterDefinitionsInstance);
 
 	return SNew(SDockTab)
@@ -251,7 +257,7 @@ TSharedRef<SDockTab> FNiagaraParameterDefinitionsToolkit::SpawnTab_ParameterPane
 	];
 }
 
-TSharedRef<SDockTab> FNiagaraParameterDefinitionsToolkit::SpawnTab_SelectedDetails(const FSpawnTabArgs& Args)
+TSharedRef<SDockTab> FNiagaraParameterDefinitionsToolkit::SpawnTab_SelectedScriptVarDetails(const FSpawnTabArgs& Args)
 {
 	checkf(Args.GetTabId().TabType == SelectedDetailsTabId, TEXT("Wrong tab ID in NiagaraParameterDefinitionsToolkit!"));
 
@@ -259,9 +265,11 @@ TSharedRef<SDockTab> FNiagaraParameterDefinitionsToolkit::SpawnTab_SelectedDetai
 	.Label(LOCTEXT("SelectedDetailsTabLabel", "Selected Details"))
 	.TabColorScale(GetTabColorScale())
 	[
-		SAssignNew(SelectedDetailsWidget, SNiagaraSelectedObjectsDetails, DetailsScriptSelection.ToSharedRef())
+		SAssignNew(SelectedScriptVarDetailsWidget, SNiagaraSelectedObjectsDetails, DetailsScriptSelection.ToSharedRef())
 		.AllowEditingLibraryScriptVariables(true)
 	];
+
+	SelectedScriptVarDetailsWidget->OnFinishedChangingProperties().AddRaw(this, &FNiagaraParameterDefinitionsToolkit::OnEditedParameterDefinitionsPropertyFinishedChanging);
 }
 
 void FNiagaraParameterDefinitionsToolkit::ExtendToolbar()
@@ -301,55 +309,59 @@ void FNiagaraParameterDefinitionsToolkit::ExtendToolbar()
 
 void FNiagaraParameterDefinitionsToolkit::SetupCommands()
 {
+	bool bInvokedFromSaveAsset = false;
 	GetToolkitCommands()->MapAction(
 		FNiagaraEditorCommands::Get().Apply,
-		FExecuteAction::CreateSP(this, &FNiagaraParameterDefinitionsToolkit::OnApply),
+		FExecuteAction::CreateSP(this, &FNiagaraParameterDefinitionsToolkit::OnApply, bInvokedFromSaveAsset),
 		FCanExecuteAction::CreateSP(this, &FNiagaraParameterDefinitionsToolkit::OnApplyEnabled));
 }
 
-void FNiagaraParameterDefinitionsToolkit::OnApply()
+void FNiagaraParameterDefinitionsToolkit::OnApply(bool bInvokedFromSaveAsset /*= false*/)
 {
-	if (ParameterDefinitionsInstance->GetChangeIdHash() != LastSyncedDefinitionsChangeIdHash)
+	if (bInvokedFromSaveAsset == false && ParameterDefinitionsInstance->GetChangeIdHash() == LastSyncedDefinitionsChangeIdHash)
 	{
-		const FScopedBusyCursor BusyCursor;
-		const FText LocalizedScriptEditorApply = NSLOCTEXT("UnrealEd", "ToolTip_NiagaraParameterDefinitionsEditorApply", "Apply changes to parameter definitions.");
-		GWarn->BeginSlowTask(LocalizedScriptEditorApply, true);
-		GWarn->StatusUpdate(1, 1, LocalizedScriptEditorApply);
+		ensureMsgf(false, TEXT("Tried to apply changes to parameter definitions but the change ID does not differ with the last seen change ID!"));
+	}
 
-		if (ParameterDefinitionsSource->IsSelected())
-		{
-			GEditor->GetSelectedObjects()->Deselect(ParameterDefinitionsSource);
-		}
+	const FScopedBusyCursor BusyCursor;
+	const FText LocalizedScriptEditorApply = NSLOCTEXT("UnrealEd", "ToolTip_NiagaraParameterDefinitionsEditorApply", "Apply changes to parameter definitions.");
+	GWarn->BeginSlowTask(LocalizedScriptEditorApply, true);
+	GWarn->StatusUpdate(1, 1, LocalizedScriptEditorApply);
 
-		ResetLoaders(ParameterDefinitionsSource->GetOutermost()); // Make sure that we're not going to get invalid version number linkers into the package we are going into. 
-		ParameterDefinitionsSource->GetOutermost()->LinkerCustomVersion.Empty();
+	if (ParameterDefinitionsSource->IsSelected())
+	{
+		GEditor->GetSelectedObjects()->Deselect(ParameterDefinitionsSource);
+	}
 
-		ParameterDefinitionsSource->PreEditChange(nullptr);
-		// overwrite the original parameter definitions in place by constructing a new one with the same name
-		ParameterDefinitionsSource = (UNiagaraParameterDefinitions*)StaticDuplicateObject(ParameterDefinitionsInstance, ParameterDefinitionsSource->GetOuter(),
-			ParameterDefinitionsSource->GetFName(), RF_AllFlags, ParameterDefinitionsSource->GetClass());
+	ResetLoaders(ParameterDefinitionsSource->GetOutermost()); // Make sure that we're not going to get invalid version number linkers into the package we are going into. 
+	ParameterDefinitionsSource->GetOutermost()->LinkerCustomVersion.Empty();
 
-		// Restore RF_Standalone and RF_Public on the original parameter definitions.
-		ParameterDefinitionsSource->SetFlags(RF_Standalone | RF_Public);
+	ParameterDefinitionsSource->PreEditChange(nullptr);
+	// overwrite the original parameter definitions in place by constructing a new one with the same name
+	ParameterDefinitionsSource = (UNiagaraParameterDefinitions*)StaticDuplicateObject(ParameterDefinitionsInstance, ParameterDefinitionsSource->GetOuter(),
+		ParameterDefinitionsSource->GetFName(), RF_AllFlags, ParameterDefinitionsSource->GetClass());
 
-		ParameterDefinitionsSource->PostEditChange();
+	// Restore RF_Standalone and RF_Public on the original parameter definitions.
+	ParameterDefinitionsSource->SetFlags(RF_Standalone | RF_Public);
 
-		// Record the last synced change id to detect future changes.
-		LastSyncedDefinitionsChangeIdHash = ParameterDefinitionsInstance->GetChangeIdHash();
+	ParameterDefinitionsSource->PostEditChange();
 
-		FRefreshAllScriptsFromExternalChangesArgs Args;
-		Args.OriginatingParameterDefinitions = ParameterDefinitionsSource;
-		FNiagaraEditorUtilities::RefreshAllScriptsFromExternalChanges(Args);
+	// Record the last synced change id to detect future changes.
+	LastSyncedDefinitionsChangeIdHash = ParameterDefinitionsInstance->GetChangeIdHash();
+	bEditedParameterDefinitionsHasPendingChanges = false;
 
-		GWarn->EndSlowTask();
-		FNiagaraEditorModule::Get().InvalidateCachedScriptAssetData();
-	}
+	ParameterDefinitionsSource->NotifyParameterDefinitionsChanged();
+	FRefreshAllScriptsFromExternalChangesArgs Args;
+	Args.OriginatingParameterDefinitions = ParameterDefinitionsSource;
+	FNiagaraEditorUtilities::RefreshAllScriptsFromExternalChanges(Args);
+
+	GWarn->EndSlowTask();
+	FNiagaraEditorModule::Get().InvalidateCachedScriptAssetData();
 }
 
 bool FNiagaraParameterDefinitionsToolkit::OnApplyEnabled() const
 {
-	return true;
-	//return bEditedParameterDefinitionsHasPendingChanges; //@todo(ng) fixup
+	return bEditedParameterDefinitionsHasPendingChanges;
 }
 
 void FNiagaraParameterDefinitionsToolkit::OnEditedParameterDefinitionsPropertyFinishedChanging(const FPropertyChangedEvent& InEvent)
@@ -357,4 +369,9 @@ void FNiagaraParameterDefinitionsToolkit::OnEditedParameterDefinitionsPropertyFi
 	bEditedParameterDefinitionsHasPendingChanges = true;
 }
 
+void FNiagaraParameterDefinitionsToolkit::OnEditedParameterDefinitionsChanged()
+{
+	bEditedParameterDefinitionsHasPendingChanges = true;
+}
+
 #undef LOCTEXT_NAMESPACE /*NiagaraParameterDefinitionsToolkit*/
diff --git a/Engine/Plugins/FX/Niagara/Source/NiagaraEditor/Private/Toolkits/NiagaraParameterDefinitionsToolkit.h b/Engine/Plugins/FX/Niagara/Source/NiagaraEditor/Private/Toolkits/NiagaraParameterDefinitionsToolkit.h
--- a/Engine/Plugins/FX/Niagara/Source/NiagaraEditor/Private/Toolkits/NiagaraParameterDefinitionsToolkit.h
+++ b/Engine/Plugins/FX/Niagara/Source/NiagaraEditor/Private/Toolkits/NiagaraParameterDefinitionsToolkit.h
@@ -45,15 +45,16 @@ protected:
 private:
 	TSharedRef<SDockTab> SpawnTab_ParameterDefinitionsDetails(const FSpawnTabArgs& Args);
 	TSharedRef<SDockTab> SpawnTab_ParameterPanel(const FSpawnTabArgs& Args);
-	TSharedRef<SDockTab> SpawnTab_SelectedDetails(const FSpawnTabArgs& Args);
+	TSharedRef<SDockTab> SpawnTab_SelectedScriptVarDetails(const FSpawnTabArgs& Args);
 
 	void ExtendToolbar();
 	void SetupCommands();
 
-	void OnApply();
+	void OnApply(bool bInvokedFromSaveAsset = false);
 	bool OnApplyEnabled() const;
 
 	void OnEditedParameterDefinitionsPropertyFinishedChanging(const FPropertyChangedEvent& InEvent);
+	void OnEditedParameterDefinitionsChanged();
 
 private:
 	/** The duplicate instance of ParameterDefinitions to be live edited. Overwrites ParameterDefinitionsSource when changes are applied. */
@@ -71,7 +72,7 @@ private:
 	/** The Parameter Panel displaying parameter definitions. */
 	TSharedPtr<FNiagaraParameterDefinitionsToolkitParameterPanelViewModel> ParameterPanelViewModel;
 
-	TSharedPtr<SNiagaraSelectedObjectsDetails> SelectedDetailsWidget;
+	TSharedPtr<SNiagaraSelectedObjectsDetails> SelectedScriptVarDetailsWidget;
 
 	/** The selection displayed by the details tab. */
 	TSharedPtr<FNiagaraObjectSelection> DetailsScriptSelection;
diff --git a/Engine/Plugins/FX/Niagara/Source/NiagaraEditor/Private/ViewModels/NiagaraParameterDefinitionsSubscriberViewModel.cpp b/Engine/Plugins/FX/Niagara/Source/NiagaraEditor/Private/ViewModels/NiagaraParameterDefinitionsSubscriberViewModel.cpp
--- a/Engine/Plugins/FX/Niagara/Source/NiagaraEditor/Private/ViewModels/NiagaraParameterDefinitionsSubscriberViewModel.cpp
+++ b/Engine/Plugins/FX/Niagara/Source/NiagaraEditor/Private/ViewModels/NiagaraParameterDefinitionsSubscriberViewModel.cpp
@@ -16,7 +16,6 @@
 void INiagaraParameterDefinitionsSubscriberViewModel::SubscribeToParameterDefinitions(UNiagaraParameterDefinitions* NewParameterDefinitions)
 {
 	GetParameterDefinitionsSubscriber()->SubscribeToParameterDefinitions(NewParameterDefinitions);
-	NewParameterDefinitions->GetOnParameterDefinitionsChanged().AddRaw(this, &INiagaraParameterDefinitionsSubscriberViewModel::SynchronizeWithParameterDefinitions, FSynchronizeWithParameterDefinitionsArgs());
 }
 
 void INiagaraParameterDefinitionsSubscriberViewModel::UnsubscribeFromParameterDefinitions(const FGuid& ParameterDefinitionsToRemoveUniqueId)
diff --git a/Engine/Plugins/FX/Niagara/Source/NiagaraEditor/Private/ViewModels/NiagaraParameterPanelViewModel.cpp b/Engine/Plugins/FX/Niagara/Source/NiagaraEditor/Private/ViewModels/NiagaraParameterPanelViewModel.cpp
--- a/Engine/Plugins/FX/Niagara/Source/NiagaraEditor/Private/ViewModels/NiagaraParameterPanelViewModel.cpp
+++ b/Engine/Plugins/FX/Niagara/Source/NiagaraEditor/Private/ViewModels/NiagaraParameterPanelViewModel.cpp
@@ -1403,14 +1403,22 @@ TArray<FNiagaraParameterPanelItem> FNiagaraSystemToolkitParameterPanelViewModel:
 			}
 	
 			// Get all UNiagaraScriptVariables of visited graphs in the ParameterToScriptVariableMap so that generated items are in sync.
-			TSet<UNiagaraGraph*> VisitedGraphs;
+			TSet<UNiagaraGraph*> VisitedExternalGraphs;
 			for (const UEdGraphPin* MapPin : Builder.Histories[0].MapPinHistory)
 			{
-				VisitedGraphs.Add(CastChecked<UNiagaraGraph>(MapPin->GetOwningNode()->GetOuter()));
+				const UNiagaraNodeFunctionCall* MapPinOuterFunctionCallNode = Cast<UNiagaraNodeFunctionCall>(MapPin->GetOuter());
+				if (MapPinOuterFunctionCallNode != nullptr)
+				{
+					UNiagaraGraph* VisitedExternalGraph = MapPinOuterFunctionCallNode->GetCalledGraph();
+					if (VisitedExternalGraph != nullptr)
+					{
+						VisitedExternalGraphs.Add(VisitedExternalGraph);
+					}
+				}
 			}
-			for (const UNiagaraGraph* VisitedGraph : VisitedGraphs)
+			for (const UNiagaraGraph* VisitedExternalGraph : VisitedExternalGraphs)
 			{
-				ParameterToScriptVariableMap.Append(VisitedGraph->GetAllMetaData());
+				ParameterToScriptVariableMap.Append(VisitedExternalGraph->GetAllMetaData());
 			}
 
 			const TArray<FName>& CustomIterationSourceNamespaces = Builder.Histories[0].IterationNamespaceOverridesEncountered;
diff --git a/Engine/Plugins/FX/Niagara/Source/NiagaraEditor/Private/ViewModels/NiagaraScratchPadScriptViewModel.cpp b/Engine/Plugins/FX/Niagara/Source/NiagaraEditor/Private/ViewModels/NiagaraScratchPadScriptViewModel.cpp
--- a/Engine/Plugins/FX/Niagara/Source/NiagaraEditor/Private/ViewModels/NiagaraScratchPadScriptViewModel.cpp
+++ b/Engine/Plugins/FX/Niagara/Source/NiagaraEditor/Private/ViewModels/NiagaraScratchPadScriptViewModel.cpp
@@ -48,7 +48,6 @@ void FNiagaraScratchPadScriptViewModel::Initialize(UNiagaraScript* Script)
 	OriginalScript = Script;
 	EditScript.Script = CastChecked<UNiagaraScript>(StaticDuplicateObject(Script, GetTransientPackage()));
 	SetScript(EditScript);
-	EditScript.InitParameterDefinitionsSubscriptions();
 	UNiagaraScriptSource* EditScriptSource = CastChecked<UNiagaraScriptSource>(EditScript.Script->GetLatestSource());
 	OnGraphNeedsRecompileHandle = EditScriptSource->NodeGraph->AddOnGraphNeedsRecompileHandler(FOnGraphChanged::FDelegate::CreateSP(this, &FNiagaraScratchPadScriptViewModel::OnScriptGraphChanged));
 	EditScript.Script->OnPropertyChanged().AddSP(this, &FNiagaraScratchPadScriptViewModel::OnScriptPropertyChanged);
diff --git a/Engine/Plugins/FX/Niagara/Source/NiagaraEditor/Private/ViewModels/NiagaraScriptViewModel.cpp b/Engine/Plugins/FX/Niagara/Source/NiagaraEditor/Private/ViewModels/NiagaraScriptViewModel.cpp
--- a/Engine/Plugins/FX/Niagara/Source/NiagaraEditor/Private/ViewModels/NiagaraScriptViewModel.cpp
+++ b/Engine/Plugins/FX/Niagara/Source/NiagaraEditor/Private/ViewModels/NiagaraScriptViewModel.cpp
@@ -130,11 +130,6 @@ const TArray<FVersionedNiagaraScriptWeakPtr>& FNiagaraScriptViewModel::GetScript
 }
 
 void FNiagaraScriptViewModel::SetScripts(UNiagaraScriptSource* InScriptSource, TArray<FVersionedNiagaraScript>& InScripts)
-{
-	SetScriptsImpl(InScriptSource, InScripts, false);
-}
-
-void FNiagaraScriptViewModel::SetScriptsImpl(UNiagaraScriptSource* InScriptSource, TArray<FVersionedNiagaraScript>& InScripts, bool bStandalone)
 {
 	// Remove the graph changed handler on the child.
 	if (Source.IsValid())
@@ -169,10 +164,6 @@ void FNiagaraScriptViewModel::SetScriptsImpl(UNiagaraScriptSource* InScriptSourc
 	for (FVersionedNiagaraScript& VersionedScript : InScripts)
 	{
 		int32 i = Scripts.Add(VersionedScript.ToWeakPtr());
-		if (bStandalone)
-		{
-			Scripts[i].InitParameterDefinitionsSubscriptions(); //@todo(ng) cleanup
-		}
 		check(VersionedScript.Script->GetSource(VersionedScript.Version) == InScriptSource);
 		Scripts[i].Script->OnVMScriptCompiled().AddSP(this, &FNiagaraScriptViewModel::OnVMScriptCompiled);
 		Scripts[i].Script->OnGPUScriptCompiled().AddSP(this, &FNiagaraScriptViewModel::OnGPUScriptCompiled);
@@ -264,7 +255,6 @@ void FNiagaraScriptViewModel::SetScript(FVersionedNiagaraScript InScript)
 	UNiagaraScriptSource* InSource = nullptr;
 	if (InScript.Script)
 	{
-		InScript.InitParameterDefinitionsSubscriptions(); //@todo(ng) cleanup
 		InScripts.Add(InScript);
 		InSource = Cast<UNiagaraScriptSource>(InScript.Script->GetSource(InScript.Version));
 	}
@@ -407,6 +397,7 @@ void FNiagaraScriptViewModel::CompileStandaloneScript(bool bForceCompile)
 		FVersionedNiagaraScriptWeakPtr VersionedScript = Scripts[0];
 		if (VersionedScript.Script->IsStandaloneScript() && VersionedScript.Script->IsCompilable())
 		{
+			VersionedScript.SynchronizeWithParameterDefinitions();
 			VersionedScript.Script->RequestCompile(VersionedScript.Version, bForceCompile);
 		}
 		else if (!VersionedScript.Script->IsCompilable())
diff --git a/Engine/Plugins/FX/Niagara/Source/NiagaraEditor/Private/ViewModels/NiagaraStandaloneScriptViewModel.cpp b/Engine/Plugins/FX/Niagara/Source/NiagaraEditor/Private/ViewModels/NiagaraStandaloneScriptViewModel.cpp
--- a/Engine/Plugins/FX/Niagara/Source/NiagaraEditor/Private/ViewModels/NiagaraStandaloneScriptViewModel.cpp
+++ b/Engine/Plugins/FX/Niagara/Source/NiagaraEditor/Private/ViewModels/NiagaraStandaloneScriptViewModel.cpp
@@ -24,7 +24,6 @@ FNiagaraStandaloneScriptViewModel::FNiagaraStandaloneScriptViewModel(
 void FNiagaraStandaloneScriptViewModel::Initialize(FVersionedNiagaraScript& InScript, const FVersionedNiagaraScript& InSourceScript)
 {
 	SetScript(InScript);
-	GetStandaloneScript().InitParameterDefinitionsSubscriptions();
 	SourceScript = InSourceScript;
 
 	SendLastCompileMessages(SourceScript);
@@ -47,12 +46,6 @@ INiagaraParameterDefinitionsSubscriber* FNiagaraStandaloneScriptViewModel::GetPa
 	return &Scripts[0];
 }
 
-void FNiagaraStandaloneScriptViewModel::SetScripts(UNiagaraScriptSource* InScriptSource, TArray<FVersionedNiagaraScript>& InScripts)
-{
-	ensureMsgf(InScripts.Num() == 1, TEXT("Tried to set more than one script for a standalone script viewmodel!"));
-	SetScriptsImpl(InScriptSource, InScripts, true);
-}
-
 void FNiagaraStandaloneScriptViewModel::OnVMScriptCompiled(UNiagaraScript* InScript, const FGuid& ScriptVersion)
 {
 	FNiagaraScriptViewModel::OnVMScriptCompiled(InScript, ScriptVersion);
diff --git a/Engine/Plugins/FX/Niagara/Source/NiagaraEditor/Private/ViewModels/NiagaraStandaloneScriptViewModel.h b/Engine/Plugins/FX/Niagara/Source/NiagaraEditor/Private/ViewModels/NiagaraStandaloneScriptViewModel.h
--- a/Engine/Plugins/FX/Niagara/Source/NiagaraEditor/Private/ViewModels/NiagaraStandaloneScriptViewModel.h
+++ b/Engine/Plugins/FX/Niagara/Source/NiagaraEditor/Private/ViewModels/NiagaraStandaloneScriptViewModel.h
@@ -33,8 +33,6 @@ public:
 	const FVersionedNiagaraScript GetStandaloneScript() const;
 
 private:
-	virtual void SetScripts(UNiagaraScriptSource* InScriptSource, TArray<FVersionedNiagaraScript>& InScripts) override;
-
 	virtual void OnVMScriptCompiled(UNiagaraScript* InScript, const FGuid& ScriptVersion) override;
 
 	/** Sends messages to FNiagaraMessageManager for all compile events from the last compile. */
diff --git a/Engine/Plugins/FX/Niagara/Source/NiagaraEditor/Public/NiagaraParameterDefinitions.h b/Engine/Plugins/FX/Niagara/Source/NiagaraEditor/Public/NiagaraParameterDefinitions.h
--- a/Engine/Plugins/FX/Niagara/Source/NiagaraEditor/Public/NiagaraParameterDefinitions.h
+++ b/Engine/Plugins/FX/Niagara/Source/NiagaraEditor/Public/NiagaraParameterDefinitions.h
@@ -59,8 +59,6 @@ public:
 	virtual void PostEditChangeProperty(struct FPropertyChangedEvent& PropertyChangedEvent) override;
 	//~ End UObject Interface
 
-	void InitBindings();
-
 	void AddParameter(const FNiagaraVariable& NewVariable);
 	void RemoveParameter(const FNiagaraVariable& VariableToRemove);
 	void RenameParameter(const FNiagaraVariable& VariableToRename, const FName NewName);
diff --git a/Engine/Plugins/FX/Niagara/Source/NiagaraEditor/Public/ViewModels/NiagaraScriptViewModel.h b/Engine/Plugins/FX/Niagara/Source/NiagaraEditor/Public/ViewModels/NiagaraScriptViewModel.h
--- a/Engine/Plugins/FX/Niagara/Source/NiagaraEditor/Public/ViewModels/NiagaraScriptViewModel.h
+++ b/Engine/Plugins/FX/Niagara/Source/NiagaraEditor/Public/ViewModels/NiagaraScriptViewModel.h
@@ -88,9 +88,6 @@ public:
 	virtual FVersionedNiagaraScript GetStandaloneScript();
 	//const UNiagaraScript* GetStandaloneScript() const;
 
-protected:
-	void SetScriptsImpl(UNiagaraScriptSource* InScriptSource, TArray<FVersionedNiagaraScript>& InScripts, bool bStandalone);
-
 private:
 	/** Handles the selection changing in the graph view model. */
 	void GraphViewModelSelectedNodesChanged();
@@ -101,7 +98,7 @@ private:
 	/** Marks this script view model as dirty and marks the scripts as needing synchrnozation. */
 	void MarkAllDirty(FString Reason);
 
-	virtual void SetScripts(UNiagaraScriptSource* InScriptSource, TArray<FVersionedNiagaraScript>& InScripts);
+	void SetScripts(UNiagaraScriptSource* InScriptSource, TArray<FVersionedNiagaraScript>& InScripts);
 
 	/** Handles when a value in the input parameter collection changes. */
 	void InputParameterValueChanged(FName ParameterName);
