Merging using //UE5/Main_to_//UE5/Release-Engine-Staging
Add OnGetKeyFor Item/Category/Section binds to SItemSelector. Now allow selective support to preserve selection and expansion state by keying Items via templated ItemKeyType.
Fix FItemSelectorViewModel::GetChildrenRecursive() only grabbing top level ItemViewModels.

diff --git a/Engine/Plugins/FX/Niagara/Source/NiagaraEditor/Private/Widgets/SNiagaraParameterPanel.cpp b/Engine/Plugins/FX/Niagara/Source/NiagaraEditor/Private/Widgets/SNiagaraParameterPanel.cpp
--- a/Engine/Plugins/FX/Niagara/Source/NiagaraEditor/Private/Widgets/SNiagaraParameterPanel.cpp
+++ b/Engine/Plugins/FX/Niagara/Source/NiagaraEditor/Private/Widgets/SNiagaraParameterPanel.cpp
@@ -78,6 +78,8 @@ void SNiagaraParameterPanel::Construct(const FArguments& InArgs, const TSharedPt
 	ParameterPanelViewModel->GetParametersWithNamespaceModifierRenamePendingDelegate().BindSP(this, &SNiagaraParameterPanel::GetParametersWithNamespaceModifierRenamePending);
 
 	SAssignNew(ItemSelector, SNiagaraParameterPanelSelector)
+	.PreserveSelectionOnRefresh(true)
+	.PreserveExpansionOnRefresh(true)
 	.Items(ParameterPanelViewModel->GetViewedParameterItems())
 	.OnGetCategoriesForItem(this, &SNiagaraParameterPanel::OnGetCategoriesForItem)
 	.OnCompareCategoriesForEquality(this, &SNiagaraParameterPanel::OnCompareCategoriesForEquality)
@@ -98,7 +100,9 @@ void SNiagaraParameterPanel::Construct(const FArguments& InArgs, const TSharedPt
 	.OnGetCategoryBackgroundImage(this, &SNiagaraParameterPanel::GetCategoryBackgroundImage)
 	.CategoryBorderBackgroundColor(FLinearColor(.6, .6, .6, 1.0f))
 	.CategoryChildSlotPadding(FMargin(0.0f, 2.0f, 0.0f, 0.0f))
-	.CategoryBorderBackgroundPadding(FMargin(0.0f, 3.0f));
+	.CategoryBorderBackgroundPadding(FMargin(0.0f, 3.0f))
+	.OnGetKeyForItem(this, &SNiagaraParameterPanel::OnGetKeyForItem)
+	.OnGetKeyForCategory(this, &SNiagaraParameterPanel::OnGetKeyForCategory);
 
 	// View options
 	TSharedRef<SWidget> ViewOptionsWidget = SNew(SBorder)
@@ -183,6 +187,11 @@ bool SNiagaraParameterPanel::OnCompareCategoriesForSorting(const FNiagaraParamet
 	return CategoryA.NamespaceMetaData.SortId < CategoryB.NamespaceMetaData.SortId;
 }
 
+const FGuid& SNiagaraParameterPanel::OnGetKeyForCategory(const FNiagaraParameterPanelCategory& Category) const
+{
+	return Category.NamespaceMetaData.GetGuid();
+}
+
 bool SNiagaraParameterPanel::OnCompareItemsForEquality(const FNiagaraParameterPanelItem& ItemA, const FNiagaraParameterPanelItem& ItemB) const
 {
 	return ItemA.GetVariable()  == ItemB.GetVariable();
@@ -193,6 +202,11 @@ bool SNiagaraParameterPanel::OnCompareItemsForSorting(const FNiagaraParameterPan
 	return ItemA.GetVariable().GetName().LexicalLess(ItemB.GetVariable().GetName());
 }
 
+const FNiagaraVariableBase& SNiagaraParameterPanel::OnGetKeyForItem(const FNiagaraParameterPanelItem& Item) const
+{
+	return Item.GetVariable();
+}
+
 bool SNiagaraParameterPanel::OnDoesItemMatchFilterText(const FText& FilterText, const FNiagaraParameterPanelItem& Item)
 {
 	return Item.GetVariable().GetName().ToString().Contains(FilterText.ToString());
diff --git a/Engine/Plugins/FX/Niagara/Source/NiagaraEditor/Public/Widgets/SItemSelector.h b/Engine/Plugins/FX/Niagara/Source/NiagaraEditor/Public/Widgets/SItemSelector.h
--- a/Engine/Plugins/FX/Niagara/Source/NiagaraEditor/Public/Widgets/SItemSelector.h
+++ b/Engine/Plugins/FX/Niagara/Source/NiagaraEditor/Public/Widgets/SItemSelector.h
@@ -18,7 +18,7 @@ enum class EItemSelectorClickActivateMode
 };
 
 /** A generic widget for selecting an item from an array of items including optional filtering and categorization. */
-template<typename CategoryType, typename ItemType, typename SectionType = CategoryType>
+template<typename CategoryType, typename ItemType, typename SectionType = CategoryType, typename CategoryKeyType = uint32, typename ItemKeyType = uint32, typename SectionKeyType = uint32>
 class SItemSelector : public SCompoundWidget
 {
 public:
@@ -69,11 +69,16 @@ public:
 	DECLARE_DELEGATE_RetVal_TwoParams(const FSlateBrush*, FOnGetCategoryBackgroundImage, bool /* bIsHovered */, bool /* bIsExpanded */)
 	DECLARE_DELEGATE_TwoParams(FOnItemSelected, const ItemType& /* SelectedItem */, ESelectInfo::Type /* SelectInfo */);
 	DECLARE_DELEGATE_RetVal_TwoParams(FReply, FOnItemsDragged, const TArray<ItemType>& /* DraggedItems */, const FPointerEvent& /* MouseEvent */);
+	DECLARE_DELEGATE_RetVal_OneParam(const ItemKeyType&, FOnGetKeyForItem, const ItemType& /* Item */);
+	DECLARE_DELEGATE_RetVal_OneParam(const CategoryKeyType&, FOnGetKeyForCategory, const CategoryType& /* Category */);
+	DECLARE_DELEGATE_RetVal_OneParam(const SectionKeyType&, FOnGetKeyForSection, const SectionType& /* Section */);
 
 public:
 	SLATE_BEGIN_ARGS(SItemSelector)
 		: _AllowMultiselect(false)
 		, _ClickActivateMode(EItemSelectorClickActivateMode::DoubleClick)
+		, _PreserveExpansionOnRefresh(false)
+		, _PreserveSelectionOnRefresh(false)
 		, _CategoryRowStyle(&FNiagaraEditorStyle::Get().GetWidgetStyle<FTableRowStyle>("ActionMenu.Row"))
 		, _SectionRowStyle(&FNiagaraEditorStyle::Get().GetWidgetStyle<FTableRowStyle>("ActionMenu.Row"))
 		, _ItemRowStyle(&FNiagaraEditorStyle::Get().GetWidgetStyle<FTableRowStyle>("ActionMenu.Row"))
@@ -97,6 +102,12 @@ public:
 		/** Whether or not a single click activates an item. */
 		SLATE_ARGUMENT(EItemSelectorClickActivateMode, ClickActivateMode)
 
+		/** Whether or not to preserve the expansion state when refresh is called. If true, OnGetKeyForCategory and OnGetKeyForSection must be set if Categories or Sections are used, respectively. */
+		SLATE_ARGUMENT(bool, PreserveExpansionOnRefresh)
+
+		/** Whether or not to preserve the selection state when refresh is called. If true, OnGetKeyForItem must be set, and OnGetKeyForCategory and OnGetKeyForSection must be set if Categories or Sections are used, respectively. */
+		SLATE_ARGUMENT(bool, PreserveSelectionOnRefresh)
+
 		/** Optional style override to use for category rows. */
 		SLATE_STYLE_ARGUMENT(FTableRowStyle, CategoryRowStyle)
 
@@ -185,12 +196,21 @@ public:
 		/** An optional delegate called when a context menu would open for an item. The user returns the menu content to display or null if a context menu should not be opened. */
 		SLATE_EVENT(FOnContextMenuOpening, OnContextMenuOpening)
 
-			/** An optional delegate called when one or more items are selected. */
+		/** An optional delegate called when one or more items are selected. */
 		SLATE_EVENT(FOnItemSelected, OnItemSelected)
 
 		/** An optional delegate called when one or more items are dragged. */
 		SLATE_EVENT(FOnItemsDragged, OnItemsDragged)
 
+		/** An optional delegate called to get a stable key (supports GetTypeHash()) for an item. If set, this delegate is used to uniquely key items to preserve selection and expansion state when refresh is called. */
+		SLATE_EVENT(FOnGetKeyForItem, OnGetKeyForItem)
+
+		/** An optional delegate called to get a stable key (supports GetTypeHash()) for a category. If set, this delegate is used to uniquely key categories to preserve selection and expansion state when refresh is called. */
+		SLATE_EVENT(FOnGetKeyForCategory, OnGetKeyForCategory)
+
+		/** An optional delegate called to get a stable key (supports GetTypeHash()) for a section. If set, this delegate is used to uniquely key sections to preserve selection and expansion state when refresh is called. */
+		SLATE_EVENT(FOnGetKeyForSection, OnGetKeyForSection)
+
 		/** An optional array of delegates to refresh the item selector view when executed. */
 		SLATE_ARGUMENT(TArray<FRefreshItemSelectorDelegate*>, RefreshItemSelectorDelegates)
 	
@@ -800,8 +820,9 @@ private:
 			FOnCompareItemsForEquality InOnCompareItemsForEquality, FOnCompareItemsForSorting InOnCompareItemsForSorting,
 			FOnDoesItemMatchFilterText InOnDoesItemMatchFilterText, FOnGetItemWeight InOnGetItemWeight, 
 			FOnDoesItemPassCustomFilter InOnDoesItemPassCustomFilter, FOnDoesSectionPassCustomFilter InOnDoesSectionPassCustomFilter,
-			FOnGetSectionData InOnGetSectionData,
-			TAttribute<bool> InHideSingleSection)
+			FOnGetSectionData InOnGetSectionData, TAttribute<bool> InHideSingleSection, bool bInPreseveExpansionOnRefresh, 
+			bool bInPreserveSelectionOnRefresh, FOnGetKeyForItem& InOnGetKeyForItem, FOnGetKeyForCategory& InOnGetKeyForCategory,
+			FOnGetKeyForSection& InOnGetKeyForSection)
 			: Items(InItems)
 			, DefaultCategoryPaths(InDefaultCategoryPaths)
 			, OnGetCategoriesForItem(InOnGetCategoriesForItem)
@@ -818,6 +839,11 @@ private:
 			, OnDoesSectionPassCustomFilter(InOnDoesSectionPassCustomFilter)
 			, OnGetSectionData(InOnGetSectionData)
 			, HideSingleSection(InHideSingleSection)
+			, bPreserveExpansionOnRefresh(bInPreseveExpansionOnRefresh)
+			, bPreserveSelectionOnRefresh(bInPreserveSelectionOnRefresh)
+			, OnGetKeyForItem(InOnGetKeyForItem)
+			, OnGetKeyForCategory(InOnGetKeyForCategory)
+			, OnGetKeyForSection(InOnGetKeyForSection)
 		{
 		}
 
@@ -887,7 +913,7 @@ private:
 		void GetChildrenRecursive(TArray<TSharedRef<FItemSelectorItemViewModel>>& OutChildren)
         {
         	TArray<TSharedRef<FItemSelectorItemViewModel>> ItemsToProcess;
-        	OutChildren.Append(*GetRootItems());
+        	ItemsToProcess.Append(*GetRootItems());
         	while (ItemsToProcess.Num() > 0)
         	{
         		TSharedRef<FItemSelectorItemViewModel> ItemToProcess = ItemsToProcess[0];
@@ -895,7 +921,6 @@ private:
         		ItemsToProcess.RemoveAtSwap(0);
         		ItemToProcess->GetChildren(ItemsToProcess);
         	}
-        	
         }
 
 		void GetItemViewModelsForItems(const TArray<ItemType>& InItems, TArray<TSharedRef<FItemSelectorItemViewModel>>& OutItemViewModelsForItems)
@@ -1048,10 +1073,10 @@ private:
 			return FilterText;
 		}
 
-		void SetFilterText(FText InFilterText)
+		void SetFilterText(FText InFilterText, const TSharedPtr<STreeView<TSharedRef<FItemSelectorItemViewModel>>>& ItemTreeRef)
 		{
 			FilterText = InFilterText;
-			Refresh(Items, DefaultCategoryPaths);
+			Refresh(Items, DefaultCategoryPaths, ItemTreeRef);
 		}
 
 		virtual bool IsFiltering() const override
@@ -1201,14 +1226,51 @@ private:
 			return RootViewModel->GetSections();
 		}
 		
-		void Refresh(const TArray<ItemType>& InItems, const TArray<TArray<CategoryType>>& InDefaultCategoryPaths, bool bPreserveExpansion = true)
+		void Refresh(const TArray<ItemType>& InItems, const TArray<TArray<CategoryType>>& InDefaultCategoryPaths, const TSharedPtr<STreeView<TSharedRef<FItemSelectorItemViewModel>>>& ItemTreeRef)
 		{
 			if (RootViewModel.IsValid())
 			{
-				// todo expansion preservation
-				if(bPreserveExpansion)
+				// If preserving expansion or selection, cache the key for items/categories/sections before clearing out item viewmodels so that they may be re-marked as expanded or selected after regeneration.
+				if(bPreserveExpansionOnRefresh)
 				{
-					//TreeView.Pin()->GetExpandedItems(ExpansionCache);
+					ExpandedCategoryKeyCache.Reset();
+					ExpandedSectionKeyCache.Reset();
+					TSet<TSharedRef<FItemSelectorItemViewModel>> ExpandedItemViewModels;
+					ItemTreeRef->GetExpandedItems(ExpandedItemViewModels);
+					for (const TSharedRef<FItemSelectorItemViewModel>& ExpandedItemViewModel : ExpandedItemViewModels)
+					{
+						if (ExpandedItemViewModel->GetType() == EItemSelectorItemViewModelType::Category)
+						{
+							ExpandedCategoryKeyCache.Add(OnGetKeyForCategory.Execute(StaticCastSharedRef<FItemSelectorItemCategoryViewModel>(ExpandedItemViewModel)->GetCategory()));
+						}
+						else if (ExpandedItemViewModel->GetType() == EItemSelectorItemViewModelType::Section)
+						{
+							ExpandedSectionKeyCache.Add(OnGetKeyForSection.Execute(StaticCastSharedRef<FSectionViewModel>(ExpandedItemViewModel)->GetSection()));
+						}
+					}
+				}
+				if (bPreserveSelectionOnRefresh)
+				{
+					SelectedItemKeyCache.Reset();
+					SelectedCategoryKeyCache.Reset();
+					SelectedSectionKeyCache.Reset();
+					TArray<TSharedRef<FItemSelectorItemViewModel>> SelectedItemViewModels;
+					ItemTreeRef->GetSelectedItems(SelectedItemViewModels);
+					for (const TSharedRef<FItemSelectorItemViewModel>& SelectedItemViewModel : SelectedItemViewModels)
+					{
+						if (SelectedItemViewModel->GetType() == EItemSelectorItemViewModelType::Item)
+						{
+							SelectedItemKeyCache.Add(OnGetKeyForItem.Execute(StaticCastSharedRef<FItemSelectorItemContainerViewModel>(SelectedItemViewModel)->GetItem()));
+						}
+						else if (SelectedItemViewModel->GetType() == EItemSelectorItemViewModelType::Category)
+						{
+							SelectedCategoryKeyCache.Add(OnGetKeyForCategory.Execute(StaticCastSharedRef<FItemSelectorItemCategoryViewModel>(SelectedItemViewModel)->GetCategory()));
+						}
+						else if (SelectedItemViewModel->GetType() == EItemSelectorItemViewModelType::Section)
+						{
+							SelectedSectionKeyCache.Add(OnGetKeyForSection.Execute(StaticCastSharedRef<FSectionViewModel>(SelectedItemViewModel)->GetSection()));
+						}
+					}
 				}
 				
 				RootViewModel.Reset();
@@ -1221,30 +1283,65 @@ private:
 			DefaultCategoryPaths = InDefaultCategoryPaths;
 			GetRootItems();
 
-			// todo expansion preservation
-			if(bPreserveExpansion && ExpansionCache.Num() > 0)
+			TArray<TSharedRef<FItemSelectorItemViewModel>> Children;
+			// Set expansion state if necessary.
+			if(bPreserveExpansionOnRefresh && ( ExpandedCategoryKeyCache.Num() > 0 || ExpandedSectionKeyCache.Num() > 0 ))
 			{
-				TArray<TSharedRef<FItemSelectorItemViewModel>> Children;
-				GetChildrenRecursive(Children);
+				if (Children.Num() == 0)
+				{
+					GetChildrenRecursive(Children);
+				}
 				
-				for(auto& Item : Children)
+				for(const TSharedRef<FItemSelectorItemViewModel>& ChildItemViewModel : Children)
+				{
+					if (ChildItemViewModel->GetType() == EItemSelectorItemViewModelType::Category)
+					{
+						if (ExpandedCategoryKeyCache.Contains(OnGetKeyForCategory.Execute(StaticCastSharedRef<FItemSelectorItemCategoryViewModel>(ChildItemViewModel)->GetCategory())))
+						{
+							ItemTreeRef->SetItemExpansion(ChildItemViewModel, true);
+						}
+					}
+					else if (ChildItemViewModel->GetType() == EItemSelectorItemViewModelType::Section)
+					{
+						if (ExpandedSectionKeyCache.Contains(OnGetKeyForSection.Execute(StaticCastSharedRef<FSectionViewModel>(ChildItemViewModel)->GetSection())))
+						{
+							ItemTreeRef->SetItemExpansion(ChildItemViewModel, true);
+						}
+					}	
+				}
+			}
+
+			// Set selection state if necessary.
+			if (bPreserveSelectionOnRefresh && ( SelectedItemKeyCache.Num() > 0 || SelectedCategoryKeyCache.Num() > 0 || SelectedSectionKeyCache.Num() > 0 ))
+			{
+				if (Children.Num() == 0)
+				{
+					GetChildrenRecursive(Children);
+				}
+
+				for (const TSharedRef<FItemSelectorItemViewModel>& ChildItemViewModel : Children)
 				{
-					if(OnCompareCategoriesForEquality.IsBound() && Item->GetType() == EItemSelectorItemViewModelType::Category)
+					if (ChildItemViewModel->GetType() == EItemSelectorItemViewModelType::Item)
+					{
+						if (SelectedItemKeyCache.Contains(OnGetKeyForItem.Execute(StaticCastSharedRef<FItemSelectorItemContainerViewModel>(ChildItemViewModel)->GetItem())))
+						{
+							ItemTreeRef->SetItemSelection(ChildItemViewModel, true);
+						}
+					}
+					else if (ChildItemViewModel->GetType() == EItemSelectorItemViewModelType::Category)
 					{
-						TSharedRef<FItemSelectorItemCategoryViewModel> CategoryViewModel = StaticCastSharedRef<FItemSelectorItemCategoryViewModel>(Item);
-						for(auto& CacheItem : ExpansionCache)
+						if (SelectedCategoryKeyCache.Contains(OnGetKeyForCategory.Execute(StaticCastSharedRef<FItemSelectorItemCategoryViewModel>(ChildItemViewModel)->GetCategory())))
 						{
-							if(CacheItem->GetType() == EItemSelectorItemViewModelType::Category)
-							{
-								TSharedRef<FItemSelectorItemCategoryViewModel> CachedCategoryViewModel = StaticCastSharedRef<FItemSelectorItemCategoryViewModel>(Item);
-
-								if(OnCompareCategoriesForEquality.Execute(CategoryViewModel->GetCategory(), CachedCategoryViewModel->GetCategory()))
-								{
-									//TreeView.Pin()->SetItemExpansion(CategoryViewModel, true);																
-								}
-							}
+							ItemTreeRef->SetItemSelection(ChildItemViewModel, true);
 						}
-					}				
+					}
+					else if (ChildItemViewModel->GetType() == EItemSelectorItemViewModelType::Section)
+					{
+						if (SelectedSectionKeyCache.Contains(OnGetKeyForSection.Execute(StaticCastSharedRef<FSectionViewModel>(ChildItemViewModel)->GetSection())))
+						{
+							ItemTreeRef->SetItemSelection(ChildItemViewModel, true);
+						}
+					}
 				}
 			}
 		}
@@ -1274,6 +1371,11 @@ private:
 		FOnDoesSectionPassCustomFilter OnDoesSectionPassCustomFilter;
 		FOnGetSectionData OnGetSectionData;
 		TAttribute<bool> HideSingleSection;
+		bool bPreserveExpansionOnRefresh;
+		bool bPreserveSelectionOnRefresh;
+		FOnGetKeyForItem OnGetKeyForItem;
+		FOnGetKeyForCategory OnGetKeyForCategory;
+		FOnGetKeyForSection OnGetKeyForSection;
 		
 		int32 CurrentMaxWeight = INDEX_NONE;
 		int32 CurrentSuggestionIndex = INDEX_NONE;
@@ -1282,7 +1384,11 @@ private:
 		TArray<TSharedRef<FItemSelectorItemContainerViewModel>> FilteredFlattenedItems;
 		FText FilterText;
 
-		TSet<TSharedRef<FItemSelectorItemViewModel>> ExpansionCache;
+		TSet<CategoryKeyType> ExpandedCategoryKeyCache;
+		TSet<SectionKeyType> ExpandedSectionKeyCache;
+		TSet<ItemKeyType> SelectedItemKeyCache;
+		TSet<CategoryKeyType> SelectedCategoryKeyCache;
+		TSet<SectionKeyType> SelectedSectionKeyCache;
 	};
 
 	typedef STableRow<TSharedRef<FItemSelectorItemViewModel>> SItemSelectorTableRow;
@@ -1398,6 +1504,8 @@ public:
 		}
 		DefaultCategoryPaths.Append(InArgs._DefaultCategoryPaths);
 		ClickActivateMode = InArgs._ClickActivateMode;
+		bPreserveExpansionOnRefresh = InArgs._PreserveExpansionOnRefresh;
+		bPreserveSelectionOnRefresh = InArgs._PreserveSelectionOnRefresh;
 		SectionRowStyle = InArgs._SectionRowStyle;
 		CategoryRowStyle = InArgs._CategoryRowStyle;
 		ItemRowStyle = InArgs._ItemRowStyle;
@@ -1427,6 +1535,9 @@ public:
 		HideSingleSection = InArgs._HideSingleSection;
 		OnItemSelected = InArgs._OnItemSelected;
 		OnItemsDragged = InArgs._OnItemsDragged;
+		OnGetKeyForItem = InArgs._OnGetKeyForItem;
+		OnGetKeyForCategory = InArgs._OnGetKeyForCategory;
+		OnGetKeyForSection = InArgs._OnGetKeyForSection;
 		SearchBoxAdjacentContentWidget = InArgs._SearchBoxAdjacentContent.Widget;
 		bIsSettingSelection = false;
 
@@ -1439,6 +1550,7 @@ public:
 			}
 		}
 
+		// Validate bindings for options set.
 		checkf(DefaultCategoryPaths.Num() == 0 || OnCompareCategoriesForEquality.IsBound(), TEXT("OnCompareCategoriesForEquality must be bound if default categories are supplied."));
 		checkf(DefaultCategoryPaths.Num() == 0 || OnGenerateWidgetForCategory.IsBound(), TEXT("OnGenerateWidgetForCategory must be bound if default categories are supplied."));
 		checkf(OnGetCategoriesForItem.IsBound() == false || OnCompareCategoriesForEquality.IsBound(), TEXT("OnCompareCategoriesForEquality must be bound if OnGenerateCategoriesForItem is bound."));
@@ -1447,6 +1559,22 @@ public:
 		checkf(OnGetSectionData.IsBound() == false || OnGetSectionsForItem.IsBound(), TEXT("OnGetSectionsForItem must be bound if OnGetListSections is bound."))
 		checkf(OnGenerateWidgetForItem.IsBound(), TEXT("OnGenerateWidgetForItem must be bound"));
 
+		if (bPreserveExpansionOnRefresh || bPreserveSelectionOnRefresh)
+		{
+			if (bPreserveSelectionOnRefresh)
+			{
+				checkf(OnGetKeyForItem.IsBound(), TEXT("OnGetKeyForItem must be bound if PreserveSelectionOnRefresh is set."));
+			}
+			if (DefaultCategoryPaths.Num() != 0 || OnGetCategoriesForItem.IsBound())
+			{
+				checkf(OnGetKeyForCategory.IsBound(), TEXT("OnGetKeyForCategory must be bound if: default categories are supplied or OnGetCategoriesForItem is bound, and PreserveExpansionOnRefresh or PreserveSelectionOnRefresh are set."));
+			}
+			if (OnGetSectionsForItem.IsBound())
+			{
+				checkf(OnGetKeyForSection.IsBound(), TEXT("OnGetKeyForSection must be bound if OnGetSectionsForItem is bound and PreserveExpansionOnRefresh or PreserveSelectionOnRefresh are set."));
+			}
+		}
+
 		ViewModelUtilities = MakeShared<FItemSelectorViewModel>(
             Items,
             DefaultCategoryPaths, 
@@ -1456,7 +1584,9 @@ public:
             OnCompareItemsForEquality, OnCompareItemsForSorting,
             OnDoesItemMatchFilterText, OnGetItemWeight, 
             OnDoesItemPassCustomFilter, OnDoesSectionPassCustomFilter,
-            OnGetSectionData, HideSingleSection);
+            OnGetSectionData, HideSingleSection, bPreserveExpansionOnRefresh,
+			bPreserveSelectionOnRefresh, OnGetKeyForItem, OnGetKeyForCategory,
+			OnGetKeyForSection);
 
 		// Search Box
 		SAssignNew(SearchBox, SSearchBox)
@@ -1668,8 +1798,11 @@ public:
 
 	void RefreshItemsAndDefaultCategories(const TArray<ItemType>& InItems, const TArray<TArray<CategoryType>>& InDefaultCategoryPaths)
 	{
-		ViewModelUtilities->Refresh(InItems, InDefaultCategoryPaths);
-		ExpandTree();
+		ViewModelUtilities->Refresh(InItems, InDefaultCategoryPaths, ItemTree);
+		if (!bPreserveExpansionOnRefresh)
+		{
+			ExpandTree();
+		}
 		ItemTree->RequestTreeRefresh();
 	}
 
@@ -1692,7 +1825,7 @@ public:
 
 	void RefreshAllCurrentItems(bool bForceExpansion = false)
 	{
-		ViewModelUtilities->Refresh(Items, DefaultCategoryPaths);
+		ViewModelUtilities->Refresh(Items, DefaultCategoryPaths, ItemTree);
 		// let's manually expand the tree here if we are still searching, as refreshing wipes the tree expansion state
 		if(IsSearching() || bForceExpansion)
 		{
@@ -1778,7 +1911,7 @@ private:
 		{
 			if(!SearchText.IsEmpty())
 			{
-				ViewModelUtilities->SetFilterText(SearchText);
+				ViewModelUtilities->SetFilterText(SearchText, ItemTree);
 				ExpandTree();
 				MarkActiveSuggestion();
 
@@ -1789,7 +1922,7 @@ private:
 				// causes a refresh of items and therefore expansion as well
 				// @todo cache expansion state. Not that easy since item data doesn't live across multiple refreshes.
 				// Even caching off the items doesn't help since their category information gets invalidated as internal cache gets cleared
-				ViewModelUtilities->SetFilterText(SearchText);
+				ViewModelUtilities->SetFilterText(SearchText, ItemTree);
 				ExpandSections();
 
 				ItemTree->RequestTreeRefresh();
@@ -1920,6 +2053,9 @@ private:
 	
 	EItemSelectorClickActivateMode ClickActivateMode;
 
+	bool bPreserveExpansionOnRefresh;
+	bool bPreserveSelectionOnRefresh;
+
 	const FTableRowStyle* SectionRowStyle;
 	const FTableRowStyle* CategoryRowStyle;
 	const FTableRowStyle* ItemRowStyle;
@@ -1950,6 +2086,9 @@ private:
 	TAttribute<bool> HideSingleSection;
 	FOnItemSelected OnItemSelected;
 	FOnItemsDragged OnItemsDragged;
+	FOnGetKeyForItem OnGetKeyForItem;
+	FOnGetKeyForCategory OnGetKeyForCategory;
+	FOnGetKeyForSection OnGetKeyForSection;
 
 	TSharedPtr<FItemSelectorViewModel> ViewModelUtilities;
 	TSharedPtr<STreeView<TSharedRef<FItemSelectorItemViewModel>>> ItemTree;
diff --git a/Engine/Plugins/FX/Niagara/Source/NiagaraEditor/Public/Widgets/SNiagaraParameterPanel.h b/Engine/Plugins/FX/Niagara/Source/NiagaraEditor/Public/Widgets/SNiagaraParameterPanel.h
--- a/Engine/Plugins/FX/Niagara/Source/NiagaraEditor/Public/Widgets/SNiagaraParameterPanel.h
+++ b/Engine/Plugins/FX/Niagara/Source/NiagaraEditor/Public/Widgets/SNiagaraParameterPanel.h
@@ -10,14 +10,14 @@
 #include "Widgets/SCompoundWidget.h"
 #include "Widgets/SItemSelector.h"
 
-typedef SItemSelector<FNiagaraParameterPanelCategory, FNiagaraParameterPanelItem> SNiagaraParameterPanelSelector;
-
 class FUICommandList;
 class INiagaraParameterPanelViewModel;
 class SExpanderArrow;
 class SGraphActionMenu;
 class SSearchBox;
+struct FGuid;
 
+typedef SItemSelector<FNiagaraParameterPanelCategory /*CategoryType*/, FNiagaraParameterPanelItem /*ItemType*/, FNiagaraParameterPanelCategory /*SectionType*/, FGuid /*CategoryKeyType*/, FNiagaraVariableBase /*ItemKeyType*/> SNiagaraParameterPanelSelector;
 
 class FNiagaraParameterPanelCommands : public TCommands<FNiagaraParameterPanelCommands>
 {
@@ -55,8 +55,10 @@ public:
 	TArray<FNiagaraParameterPanelCategory> OnGetCategoriesForItem(const FNiagaraParameterPanelItem& Item);
 	bool OnCompareCategoriesForEquality(const FNiagaraParameterPanelCategory& CategoryA, const FNiagaraParameterPanelCategory& CategoryB) const;
 	bool OnCompareCategoriesForSorting(const FNiagaraParameterPanelCategory& CategoryA, const FNiagaraParameterPanelCategory& CategoryB) const;
+	const FGuid& OnGetKeyForCategory(const FNiagaraParameterPanelCategory& Category) const;
 	bool OnCompareItemsForEquality(const FNiagaraParameterPanelItem& ItemA, const FNiagaraParameterPanelItem& ItemB) const;
 	bool OnCompareItemsForSorting(const FNiagaraParameterPanelItem& ItemA, const FNiagaraParameterPanelItem& ItemB) const;
+	const FNiagaraVariableBase& OnGetKeyForItem(const FNiagaraParameterPanelItem& Item) const;
 	bool OnDoesItemMatchFilterText(const FText& FilterText, const FNiagaraParameterPanelItem& Item);
 	TSharedRef<SWidget> OnGenerateWidgetForCategory(const FNiagaraParameterPanelCategory& Category);
 	TSharedRef<SWidget> OnGenerateWidgetForItem(const FNiagaraParameterPanelItem& Item);
