Merging using //UE5/Main_to_//UE5/Release-Engine-Staging
Don't crash on hitting null graph on scriptsource during apply to scripts, and log the scriptsource and outer.

diff --git a/Engine/Plugins/FX/Niagara/Source/NiagaraEditor/Private/NiagaraEditorUtilities.cpp b/Engine/Plugins/FX/Niagara/Source/NiagaraEditor/Private/NiagaraEditorUtilities.cpp
--- a/Engine/Plugins/FX/Niagara/Source/NiagaraEditor/Private/NiagaraEditorUtilities.cpp
+++ b/Engine/Plugins/FX/Niagara/Source/NiagaraEditor/Private/NiagaraEditorUtilities.cpp
@@ -2769,6 +2769,11 @@ void FNiagaraEditorUtilities::RefreshAllScriptsFromExternalChanges(FRefreshAllSc
 		{
 			continue;
 		}
+		if (Source->NodeGraph == nullptr)
+		{
+			ensureMsgf(false, TEXT("Encountered null nodegraph on source script: %s (outer: %s)"), *Source->GetName(), *Source->GetOuter()->GetName());
+			continue;
+		}
 		TArray<UNiagaraNode*> NiagaraNodes;
 		Source->NodeGraph->GetNodesOfClass<UNiagaraNode>(NiagaraNodes);
 		bool bRefreshed = false;
