Merging using //UE5/Main_to_//UE5/Release-Engine-Staging
Fix crash loading last options index from niagara asset dialog config if the last index is greater than the number of options now available.

diff --git a/Engine/Plugins/FX/Niagara/Source/NiagaraEditor/Private/SNiagaraNewAssetDialog.cpp b/Engine/Plugins/FX/Niagara/Source/NiagaraEditor/Private/SNiagaraNewAssetDialog.cpp
--- a/Engine/Plugins/FX/Niagara/Source/NiagaraEditor/Private/SNiagaraNewAssetDialog.cpp
+++ b/Engine/Plugins/FX/Niagara/Source/NiagaraEditor/Private/SNiagaraNewAssetDialog.cpp
@@ -29,6 +29,11 @@ void SNiagaraNewAssetDialog::Construct(const FArguments& InArgs, FName InSaveCon
 	SelectedOptionIndex = DialogConfig.SelectedOptionIndex;
 
 	Options = InOptions;
+	// It is possible that the number of options has changed since the options config was last saved; make sure the SelectedOptionsIndex from the config is valid.
+	if (SelectedOptionIndex > Options.Num() - 1)
+	{
+		SelectedOptionIndex = Options.Num() - 1;
+	}
 
 	SetOnWindowClosed(FOnWindowClosed::CreateSP(this, &SNiagaraNewAssetDialog::OnWindowClosed));
 	
