Merging using //UE5/Main_to_//UE5/Release-Engine-Staging
Skip applying changes to scripts that are in the transient package as they should only be used as intermediate objects during compile and are not guaranteed to be valid.

diff --git a/Engine/Plugins/FX/Niagara/Source/NiagaraEditor/Private/NiagaraEditorUtilities.cpp b/Engine/Plugins/FX/Niagara/Source/NiagaraEditor/Private/NiagaraEditorUtilities.cpp
--- a/Engine/Plugins/FX/Niagara/Source/NiagaraEditor/Private/NiagaraEditorUtilities.cpp
+++ b/Engine/Plugins/FX/Niagara/Source/NiagaraEditor/Private/NiagaraEditorUtilities.cpp
@@ -2758,7 +2758,7 @@ void FNiagaraEditorUtilities::RefreshAllScriptsFromExternalChanges(FRefreshAllSc
 
 	for (TObjectIterator<UNiagaraScript> It; It; ++It)
 	{
-		if (*It == OriginatingScript || It->IsPendingKillOrUnreachable())
+		if (*It == OriginatingScript || It->IsPendingKillOrUnreachable() || It->GetOutermost() == GetTransientPackage())
 		{
 			continue;
 		}
