Merging using //UE5/Main_to_//UE5/Release-Engine-Staging
SItemSelector; Reorganize item keying to cache item key on the itemviewmodel when the itemviewmodel is generated. This is to prevent a crash on trying to get the item key on refresh; as we own the item by ref it is not guaranteed to be valid during refresh, so now cache the key when it is valid.

diff --git a/Engine/Plugins/FX/Niagara/Source/NiagaraEditor/Public/Widgets/SItemSelector.h b/Engine/Plugins/FX/Niagara/Source/NiagaraEditor/Public/Widgets/SItemSelector.h
--- a/Engine/Plugins/FX/Niagara/Source/NiagaraEditor/Public/Widgets/SItemSelector.h
+++ b/Engine/Plugins/FX/Niagara/Source/NiagaraEditor/Public/Widgets/SItemSelector.h
@@ -250,6 +250,8 @@ private:
 		virtual const FOnCompareCategoriesForSorting& GetOnCompareCategoriesForSorting() const = 0;
 		virtual bool CompareItemsForEquality(const ItemType& ItemA, const ItemType& ItemB) const = 0;
 		virtual const FOnCompareItemsForSorting& GetOnCompareItemsForSorting() const = 0;
+
+		virtual const ItemKeyType GetKeyForItem(const ItemType& InItem) const = 0;
 	};
 
 	class FItemSelectorItemViewModel
@@ -320,9 +322,10 @@ private:
 	class FItemSelectorItemContainerViewModel : public FItemSelectorItemViewModel
 	{
 	public:
-		FItemSelectorItemContainerViewModel(TSharedRef<IItemSelectorItemViewModelUtilities> InItemUtilities, const ItemType& InItem)
+		FItemSelectorItemContainerViewModel(TSharedRef<IItemSelectorItemViewModelUtilities> InItemUtilities, const ItemType& InItem, const ItemKeyType& InItemKey)
 			: FItemSelectorItemViewModel(InItemUtilities, EItemSelectorItemViewModelType::Item)
 			, Item(InItem)
+			, ItemKey(InItemKey)
 		{
 		}
 
@@ -331,6 +334,11 @@ private:
 			return Item;
 		}
 
+		const ItemKeyType& GetItemKey() const
+		{
+			return ItemKey;
+		}
+
 		virtual int32 GetItemWeight(const TArray<FString>& FilterTerms) const
 		{
 			return this->GetItemUtilities()->GetItemWeight(Item, FilterTerms);
@@ -343,6 +351,7 @@ private:
 
 	private:
 		const ItemType& Item;
+		const ItemKeyType ItemKey;
 	};
 
 	class FItemSelectorItemCategoryViewModel : public FItemSelectorItemViewModel
@@ -397,7 +406,8 @@ private:
 		
 		void AddItemDirect(const ItemType& InItem)
 		{
-			ChildItemViewModels.Add(MakeShared<FItemSelectorItemContainerViewModel>(this->GetItemUtilities(), InItem));
+			const ItemKeyType ItemKey = this->GetItemUtilities()->GetKeyForItem(InItem);
+			ChildItemViewModels.Add(MakeShared<FItemSelectorItemContainerViewModel>(this->GetItemUtilities(), InItem, ItemKey));
 		}
 
 		void FlattenChildren(TArray<TSharedRef<FItemSelectorItemContainerViewModel>>& FilteredFlattenedItems)
@@ -492,7 +502,8 @@ private:
 
 		void AddItemDirect(const ItemType& InItem)
 		{
-			ChildItemViewModels.Add(MakeShared<FItemSelectorItemContainerViewModel>(this->GetItemUtilities(), InItem));
+			const ItemKeyType ItemKey = this->GetItemUtilities()->GetKeyForItem(InItem);
+			ChildItemViewModels.Add(MakeShared<FItemSelectorItemContainerViewModel>(this->GetItemUtilities(), InItem, ItemKey));
 		}
 
 		TSharedRef<FItemSelectorItemCategoryViewModel> AddCategory(const CategoryType& InCategory)
@@ -776,7 +787,8 @@ private:
 		
 		void AddItemToRootDirect(const ItemType& Item)
 		{
-			ChildItemViewModels.Add(MakeShared<FItemSelectorItemContainerViewModel>(this->GetItemUtilities(), Item));
+			const ItemKeyType ItemKey = this->GetItemUtilities()->GetKeyForItem(Item);
+			ChildItemViewModels.Add(MakeShared<FItemSelectorItemContainerViewModel>(this->GetItemUtilities(), Item, ItemKey));
 		}
 		
 		void AddCategoryToRoot(const CategoryType& Category)
@@ -1186,6 +1198,11 @@ private:
 			return OnCompareItemsForSorting;
 		}
 
+		virtual const ItemKeyType GetKeyForItem(const ItemType& InItem) const override
+		{
+			return OnGetKeyForItem.IsBound() ? OnGetKeyForItem.Execute(InItem) : ItemKeyType();
+		}
+
 		bool CanCompareItems() const
 		{
 			return OnCompareItemsForSorting.IsBound();
@@ -1260,7 +1277,7 @@ private:
 					{
 						if (SelectedItemViewModel->GetType() == EItemSelectorItemViewModelType::Item)
 						{
-							SelectedItemKeyCache.Add(OnGetKeyForItem.Execute(StaticCastSharedRef<FItemSelectorItemContainerViewModel>(SelectedItemViewModel)->GetItem()));
+							SelectedItemKeyCache.Add(StaticCastSharedRef<FItemSelectorItemContainerViewModel>(SelectedItemViewModel)->GetItemKey());
 						}
 						else if (SelectedItemViewModel->GetType() == EItemSelectorItemViewModelType::Category)
 						{
