Merging using //UE5/Main_to_//UE5/Release-Engine-Staging
Fix excessive refreshing and loss of focus when editing parameters in NiagaraScriptVariableCustomization.

diff --git a/Engine/Plugins/FX/Niagara/Source/NiagaraEditor/Private/Customizations/NiagaraScriptVariableCustomization.cpp b/Engine/Plugins/FX/Niagara/Source/NiagaraEditor/Private/Customizations/NiagaraScriptVariableCustomization.cpp
--- a/Engine/Plugins/FX/Niagara/Source/NiagaraEditor/Private/Customizations/NiagaraScriptVariableCustomization.cpp
+++ b/Engine/Plugins/FX/Niagara/Source/NiagaraEditor/Private/Customizations/NiagaraScriptVariableCustomization.cpp
@@ -33,6 +33,7 @@ TSharedRef<IDetailCustomization> FNiagaraScriptVariableDetails::MakeInstance()
 FNiagaraScriptVariableDetails::FNiagaraScriptVariableDetails()
 {
 	GEditor->RegisterForUndo(this);
+	bParameterValueChangedDuringOnValueChanged = false;
 	LibraryDefaultModeValue = 0;
 
 	LibrarySynchronizedDefaultModeEnum = StaticEnum<ENiagaraLibrarySynchronizedDefaultMode>();
@@ -407,7 +408,6 @@ void FNiagaraScriptVariableDetails::OnValueChanged()
 {
 	if (TypeUtilityValue && ParameterEditorValue && GetDefaultPins().Num() > 0)
 	{
-		FString NewDefaultValue;
 		if (!ParameterEditorValue->CanChangeContinuously())
 		{
 			const FScopedTransaction Transaction(NSLOCTEXT("ScriptVariableCustomization", "ChangeValue", "Change Default Value"));
@@ -416,7 +416,7 @@ void FNiagaraScriptVariableDetails::OnValueChanged()
 			TSharedPtr<FStructOnScope> ParameterValue = MakeShareable(new FStructOnScope(Variable->Variable.GetType().GetStruct()));
 			ParameterEditorValue->UpdateStructFromInternalValue(ParameterValue.ToSharedRef());
 			Variable->Variable.SetData(ParameterValue->GetStructMemory());
-			NewDefaultValue = TypeUtilityValue->GetPinDefaultStringFromValue(Variable->Variable);
+			const FString NewDefaultValue = TypeUtilityValue->GetPinDefaultStringFromValue(Variable->Variable);
 
 			for (UEdGraphPin* Pin : GetDefaultPins())
 			{
@@ -426,31 +426,18 @@ void FNiagaraScriptVariableDetails::OnValueChanged()
 		}
 		else
 		{
-			// Values that can change continuously are not guaranteed to call OnBeginChange() before OnChange(), so if a transaction has not been begun by OnBeginChange(), create one now.
-			bool bOnValueChangedTransactionActive = false;
-			if (GEditor->IsTransactionActive() == false)
-			{
-				bOnValueChangedTransactionActive = true;
-				GEditor->BeginTransaction(NSLOCTEXT("ScriptVariableCustomization", "ChangeValue", "Change Default Value"));
-				Variable->Modify();
-			}
 			TSharedPtr<FStructOnScope> ParameterValue = MakeShareable(new FStructOnScope(Variable->Variable.GetType().GetStruct()));
 			ParameterEditorValue->UpdateStructFromInternalValue(ParameterValue.ToSharedRef());
 			Variable->Variable.SetData(ParameterValue->GetStructMemory());
-			NewDefaultValue = TypeUtilityValue->GetPinDefaultStringFromValue(Variable->Variable);
+			const FString NewDefaultValue = TypeUtilityValue->GetPinDefaultStringFromValue(Variable->Variable);
 
 			for (UEdGraphPin* Pin : GetDefaultPins())
 			{
-				if (bOnValueChangedTransactionActive)
+				if (Pin->GetDefaultAsString() != NewDefaultValue)
 				{
-					Pin->Modify();
+					bParameterValueChangedDuringOnValueChanged = true;
+					GetDefault<UEdGraphSchema_Niagara>()->TrySetDefaultValue(*Pin, NewDefaultValue, false);
 				}
-				GetDefault<UEdGraphSchema_Niagara>()->TrySetDefaultValue(*Pin, NewDefaultValue, true);
-			}
-
-			if (bOnValueChangedTransactionActive)
-			{
-				GEditor->EndTransaction();
 			}
 		}
 	}
@@ -458,6 +445,7 @@ void FNiagaraScriptVariableDetails::OnValueChanged()
  
 void FNiagaraScriptVariableDetails::OnBeginValueChanged()
 {
+	bParameterValueChangedDuringOnValueChanged = false;
 	if (!ParameterEditorValue->CanChangeContinuously())
 	{
 		return;
@@ -470,12 +458,16 @@ void FNiagaraScriptVariableDetails::OnBeginValueChanged()
 		TSharedPtr<FStructOnScope> ParameterValue = MakeShareable(new FStructOnScope(Variable->Variable.GetType().GetStruct()));
 		ParameterEditorValue->UpdateStructFromInternalValue(ParameterValue.ToSharedRef());
 		Variable->Variable.SetData(ParameterValue->GetStructMemory());
+		const FString NewDefaultValue = TypeUtilityValue->GetPinDefaultStringFromValue(Variable->Variable);
 
 		for (UEdGraphPin* Pin : GetDefaultPins())
 		{
-			Pin->Modify();
-			FString NewDefaultValue = TypeUtilityValue->GetPinDefaultStringFromValue(Variable->Variable);
-			GetDefault<UEdGraphSchema_Niagara>()->TrySetDefaultValue(*Pin, NewDefaultValue, true);
+			if (Pin->GetDefaultAsString() != NewDefaultValue)
+			{
+				bParameterValueChangedDuringOnValueChanged = true;
+				Pin->Modify();
+				GetDefault<UEdGraphSchema_Niagara>()->TrySetDefaultValue(*Pin, NewDefaultValue, false);
+			}
 		}
 	}
 }
@@ -484,6 +476,15 @@ void FNiagaraScriptVariableDetails::OnEndValueChanged()
 {
 	if (GEditor->IsTransactionActive())
 	{
+		// If the parameter value was actually changed, invoke PinDefaultValueChanged() on the target pin(s) owning Node(s) to propagate the change.
+		if (bParameterValueChangedDuringOnValueChanged)
+		{
+			bParameterValueChangedDuringOnValueChanged = false;
+			for (UEdGraphPin* Pin : GetDefaultPins())
+			{
+				Pin->GetOwningNode()->PinDefaultValueChanged(Pin);
+			}
+		}
 		GEditor->EndTransaction();
 	}
 }
@@ -569,35 +570,10 @@ void FNiagaraScriptVariableDetails::OnLibraryValueChanged()
 		}
 		else
 		{
-			// Values that can change continuously are not guaranteed to call OnBeginChange() before OnChange(), so if a transaction has not been begun by OnBeginChange(), create one now.
-			bool bOnValueChangedTransactionActive = false;
-			if (GEditor->IsTransactionActive() == false)
-			{
-				bOnValueChangedTransactionActive = true;
-				GEditor->BeginTransaction(NSLOCTEXT("ScriptVariableCustomization", "ChangeValue", "Change Default Value"));
-				Variable->Modify();
-			}
 			DefaultValueHandle->NotifyPreChange();
 			TSharedPtr<FStructOnScope> ParameterValue = MakeShareable(new FStructOnScope(Variable->Variable.GetType().GetStruct()));
 			ParameterEditorLibraryValue->UpdateStructFromInternalValue(ParameterValue.ToSharedRef());
 			Variable->SetDefaultValueData(ParameterValue->GetStructMemory());
-
-			if (bOnValueChangedTransactionActive)
-			{
-				Variable->UpdateChangeId();
-
-				DefaultValueHandle->NotifyPostChange(EPropertyChangeType::ValueSet);
-				DefaultValueHandle->NotifyFinishedChangingProperties();
-
-				if (UNiagaraParameterDefinitions* OuterParameterDefinitions = Cast<UNiagaraParameterDefinitions>(Variable->GetOuter()))
-				{
-					OuterParameterDefinitions->NotifyParameterDefinitionsChanged();
-				}
-
-				GEditor->EndTransaction();
-			}
-			DefaultValueHandle->NotifyPostChange(EPropertyChangeType::ValueSet);
-			DefaultValueHandle->NotifyFinishedChangingProperties();
 		}
 	}
 }
diff --git a/Engine/Plugins/FX/Niagara/Source/NiagaraEditor/Private/Customizations/NiagaraScriptVariableCustomization.h b/Engine/Plugins/FX/Niagara/Source/NiagaraEditor/Private/Customizations/NiagaraScriptVariableCustomization.h
--- a/Engine/Plugins/FX/Niagara/Source/NiagaraEditor/Private/Customizations/NiagaraScriptVariableCustomization.h
+++ b/Engine/Plugins/FX/Niagara/Source/NiagaraEditor/Private/Customizations/NiagaraScriptVariableCustomization.h
@@ -96,6 +96,9 @@ private:
 	TArray<UEdGraphPin*> GetDefaultPins();
 
 private:
+	// Bool to track whether a parameters value has actually changed at any point during OnBeginValueChanged(), OnValueChanged() and OnEndValueChanged() to prevent excessive refreshing.
+	bool bParameterValueChangedDuringOnValueChanged;
+
 	// Cached value of enum selector for choosing the library default mode.
 	int32 LibraryDefaultModeValue;
 
diff --git a/Engine/Plugins/FX/Niagara/Source/NiagaraEditor/Private/EdGraphSchema_Niagara.cpp b/Engine/Plugins/FX/Niagara/Source/NiagaraEditor/Private/EdGraphSchema_Niagara.cpp
--- a/Engine/Plugins/FX/Niagara/Source/NiagaraEditor/Private/EdGraphSchema_Niagara.cpp
+++ b/Engine/Plugins/FX/Niagara/Source/NiagaraEditor/Private/EdGraphSchema_Niagara.cpp
@@ -1154,6 +1154,17 @@ void UEdGraphSchema_Niagara::GetAssetsGraphHoverMessage(const TArray<FAssetData>
 	}
 }
 
+void UEdGraphSchema_Niagara::TrySetDefaultValue(UEdGraphPin& Pin, const FString& NewDefaultValue, bool bMarkAsModified /*= true*/) const
+{
+	Pin.DefaultValue = NewDefaultValue;
+
+	if (bMarkAsModified)
+	{
+		UEdGraphNode* Node = Pin.GetOwningNode();
+		checkf(Node != nullptr, TEXT("Encountered null node owning pin!"));
+		Node->PinDefaultValueChanged(&Pin);
+	}
+}
 
 bool UEdGraphSchema_Niagara::TryCreateConnection(UEdGraphPin* PinA, UEdGraphPin* PinB) const
 {
diff --git a/Engine/Plugins/FX/Niagara/Source/NiagaraEditor/Private/SNiagaraParameterEditor.cpp b/Engine/Plugins/FX/Niagara/Source/NiagaraEditor/Private/SNiagaraParameterEditor.cpp
--- a/Engine/Plugins/FX/Niagara/Source/NiagaraEditor/Private/SNiagaraParameterEditor.cpp
+++ b/Engine/Plugins/FX/Niagara/Source/NiagaraEditor/Private/SNiagaraParameterEditor.cpp
@@ -10,6 +10,7 @@ void SNiagaraParameterEditor::Construct(const FArguments& InArgs)
 	VerticalAlignment = InArgs._VAlign;
 	MinimumDesiredWidth = InArgs._MinimumDesiredWidth;
 	MaximumDesiredWidth = InArgs._MaximumDesiredWidth;
+	bContinousChangeActive = false;
 }
 
 void SNiagaraParameterEditor::SetOnBeginValueChange(FOnValueChange InOnBeginValueChange)
@@ -76,15 +77,25 @@ void SNiagaraParameterEditor::SetIsEditingExclusively(bool bInIsEditingExclusive
 
 void SNiagaraParameterEditor::ExecuteOnBeginValueChange()
 {
+	bContinousChangeActive = true;
 	OnBeginValueChange.ExecuteIfBound();
 }
 
 void SNiagaraParameterEditor::ExecuteOnEndValueChange()
 {
+	bContinousChangeActive = false;
 	OnEndValueChange.ExecuteIfBound();
 }
 
 void SNiagaraParameterEditor::ExecuteOnValueChanged()
 {
+	if (bContinousChangeActive == false)
+	{
+		OnBeginValueChange.ExecuteIfBound();
+	}
 	OnValueChanged.ExecuteIfBound();
+	if (bContinousChangeActive == false)
+	{
+		OnEndValueChange.ExecuteIfBound();
+	}
 }
\ No newline at end of file
diff --git a/Engine/Plugins/FX/Niagara/Source/NiagaraEditor/Private/Widgets/SNiagaraSelectedObjectsDetails.cpp b/Engine/Plugins/FX/Niagara/Source/NiagaraEditor/Private/Widgets/SNiagaraSelectedObjectsDetails.cpp
--- a/Engine/Plugins/FX/Niagara/Source/NiagaraEditor/Private/Widgets/SNiagaraSelectedObjectsDetails.cpp
+++ b/Engine/Plugins/FX/Niagara/Source/NiagaraEditor/Private/Widgets/SNiagaraSelectedObjectsDetails.cpp
@@ -19,8 +19,9 @@ void SNiagaraSelectedObjectsDetails::Construct(const FArguments& InArgs, TShared
 {
 	bAllowEditingLibraryOwnedScriptVars = InArgs._AllowEditingLibraryScriptVariables;
 	bViewingLibrarySubscribedScriptVar = false;
+	LastSetSelectedObjectsArrayIdx = 0;
 	SelectedObjectsArray.Push(InSelectedObjects);
-	SelectedObjectsArray[0]->OnSelectedObjectsChanged().AddSP(this, &SNiagaraSelectedObjectsDetails::SelectedObjectsChanged);
+	SelectedObjectsArray[0]->OnSelectedObjectsChanged().AddSP(this, &SNiagaraSelectedObjectsDetails::SelectedObjectsChangedFirst);
 
 	FPropertyEditorModule& PropertyEditorModule = FModuleManager::GetModuleChecked<FPropertyEditorModule>("PropertyEditor");
 	FDetailsViewArgs DetailsViewArgs;
@@ -44,9 +45,10 @@ void SNiagaraSelectedObjectsDetails::Construct(const FArguments& InArgs, TShared
 {
 	bAllowEditingLibraryOwnedScriptVars = InArgs._AllowEditingLibraryScriptVariables;
 	bViewingLibrarySubscribedScriptVar = false;
+	LastSetSelectedObjectsArrayIdx = 0;
 	SelectedObjectsArray.Push(InSelectedObjects);
 	SelectedObjectsArray.Push(InSelectedObjects2);
-	SelectedObjectsArray[0]->OnSelectedObjectsChanged().AddSP(this, &SNiagaraSelectedObjectsDetails::SelectedObjectsChanged);
+	SelectedObjectsArray[0]->OnSelectedObjectsChanged().AddSP(this, &SNiagaraSelectedObjectsDetails::SelectedObjectsChangedFirst);
 	SelectedObjectsArray[1]->OnSelectedObjectsChanged().AddSP(this, &SNiagaraSelectedObjectsDetails::SelectedObjectsChangedSecond);
 
 	FPropertyEditorModule& PropertyEditorModule = FModuleManager::GetModuleChecked<FPropertyEditorModule>("PropertyEditor");
@@ -68,6 +70,16 @@ void SNiagaraSelectedObjectsDetails::Construct(const FArguments& InArgs, TShared
 } 
 
 void SNiagaraSelectedObjectsDetails::SelectedObjectsChanged()
+{
+	// Do not update selected object info flags if editing library owned script vars is enabled as these flags are not used.
+	if (bAllowEditingLibraryOwnedScriptVars == false)
+	{
+		UpdateSelectedObjectInfoFlags(SelectedObjectsArray[LastSetSelectedObjectsArrayIdx]);
+	}
+	DetailsView->SetObjects(SelectedObjectsArray[LastSetSelectedObjectsArrayIdx]->GetSelectedObjects().Array());
+}
+
+void SNiagaraSelectedObjectsDetails::SelectedObjectsChangedFirst()
 {
 	// Do not update selected object info flags if editing library owned script vars is enabled as these flags are not used.
 	if (bAllowEditingLibraryOwnedScriptVars == false)
@@ -75,6 +87,7 @@ void SNiagaraSelectedObjectsDetails::SelectedObjectsChanged()
 		UpdateSelectedObjectInfoFlags(SelectedObjectsArray[0]);
 	}
 	DetailsView->SetObjects(SelectedObjectsArray[0]->GetSelectedObjects().Array());
+	LastSetSelectedObjectsArrayIdx = 0;
 }
 
 void SNiagaraSelectedObjectsDetails::RefreshDetails()
@@ -91,6 +104,7 @@ void SNiagaraSelectedObjectsDetails::SelectedObjectsChangedSecond()
 		UpdateSelectedObjectInfoFlags(SelectedObjectsArray[1]);
 	}
 	DetailsView->SetObjects(SelectedObjectsArray[1]->GetSelectedObjects().Array());
+	LastSetSelectedObjectsArrayIdx = 1;
 }
 
 void SNiagaraSelectedObjectsDetails::OnDetailsPanelFinishedChangingProperties(const FPropertyChangedEvent& InEvent)
diff --git a/Engine/Plugins/FX/Niagara/Source/NiagaraEditor/Public/EdGraphSchema_Niagara.h b/Engine/Plugins/FX/Niagara/Source/NiagaraEditor/Public/EdGraphSchema_Niagara.h
--- a/Engine/Plugins/FX/Niagara/Source/NiagaraEditor/Public/EdGraphSchema_Niagara.h
+++ b/Engine/Plugins/FX/Niagara/Source/NiagaraEditor/Public/EdGraphSchema_Niagara.h
@@ -105,6 +105,16 @@ class NIAGARAEDITOR_API UEdGraphSchema_Niagara : public UEdGraphSchema
 	virtual bool ShouldAlwaysPurgeOnModification() const override { return false; }
 	virtual void DroppedAssetsOnGraph(const TArray<FAssetData>& Assets, const FVector2D& GraphPosition, UEdGraph* Graph) const override;
 	virtual void GetAssetsGraphHoverMessage(const TArray<FAssetData>& Assets, const UEdGraph* HoverGraph, FString& OutTooltipText, bool& OutOkIcon) const override;
+
+	/**
+	* Sets the string to the specified pin; even if it is invalid it is still set.
+	*
+	* @param	Pin			   	The pin on which to set the default value.
+	* @param	NewDefaultValue	The new default value.
+	* @param   bMarkAsModified Marks the container of the value as modified. For Niagara Schema, whether or not to notify the Node owning the Pin that the Pin has changed.
+	*/
+	virtual void TrySetDefaultValue(UEdGraphPin& Pin, const FString& NewDefaultValue, bool bMarkAsModified = true) const override;
+
 	//~ End EdGraphSchema Interface
 
 	static FLinearColor GetTypeColor(const FNiagaraTypeDefinition& Type);
diff --git a/Engine/Plugins/FX/Niagara/Source/NiagaraEditor/Public/SNiagaraParameterEditor.h b/Engine/Plugins/FX/Niagara/Source/NiagaraEditor/Public/SNiagaraParameterEditor.h
--- a/Engine/Plugins/FX/Niagara/Source/NiagaraEditor/Public/SNiagaraParameterEditor.h
+++ b/Engine/Plugins/FX/Niagara/Source/NiagaraEditor/Public/SNiagaraParameterEditor.h
@@ -81,6 +81,9 @@ protected:
 	static const float DefaultInputSize;
 
 private:
+	//** Flag set when a continuous change is in progress to prevent invoking ExecuteOnBeginValueChange() or ExecuteOnEndValueChange() more than once for a change. */
+	bool bContinousChangeActive;
+
 	/** Whether this is currently the exclusive editor of this parameter, meaning that the corresponding details view
 		should not be updated.  This hack is necessary because the details view closes all color pickers when
 		it's changed! */
diff --git a/Engine/Plugins/FX/Niagara/Source/NiagaraEditor/Public/Widgets/SNiagaraSelectedObjectsDetails.h b/Engine/Plugins/FX/Niagara/Source/NiagaraEditor/Public/Widgets/SNiagaraSelectedObjectsDetails.h
--- a/Engine/Plugins/FX/Niagara/Source/NiagaraEditor/Public/Widgets/SNiagaraSelectedObjectsDetails.h
+++ b/Engine/Plugins/FX/Niagara/Source/NiagaraEditor/Public/Widgets/SNiagaraSelectedObjectsDetails.h
@@ -33,6 +33,7 @@ public:
 	void RefreshDetails();
 private:
 	/** Called whenever the object selection changes. */
+	void SelectedObjectsChangedFirst();
 	void SelectedObjectsChangedSecond();
 
 	/** Internal delegate to route to third parties. */
@@ -70,4 +71,7 @@ private:
 
 	/** Flag to notify if a UNiagaraScriptVariable owned by a UNiagaraParameterDefinitions is being viewed. Used to selectively disable editing. */
 	bool bViewingLibraryOwnedScriptVar;
+
+	/** Index of the SelectedObjectsArray that was last set. Used to track which SelectedObjects to reset on refresh. */
+	int32 LastSetSelectedObjectsArrayIdx;
 };
