- PostInitViews also needs to set final data to ensure MeshProcessor / Sorting uses the correct data

diff --git a/Engine/Plugins/FX/Niagara/Source/Niagara/Private/NiagaraEmitterInstanceBatcher.cpp b/Engine/Plugins/FX/Niagara/Source/Niagara/Private/NiagaraEmitterInstanceBatcher.cpp
--- a/Engine/Plugins/FX/Niagara/Source/Niagara/Private/NiagaraEmitterInstanceBatcher.cpp
+++ b/Engine/Plugins/FX/Niagara/Source/Niagara/Private/NiagaraEmitterInstanceBatcher.cpp
@@ -1111,9 +1111,10 @@ void NiagaraEmitterInstanceBatcher::PrepareTicksForProxy(FRHICommandListImmediat
 			ComputeContext->DataBuffers_RT[BufferIndex]->AllocateGPU(RHICmdList, ComputeContext->CurrentMaxAllocateInstances_RT + 1, FeatureLevel, ComputeContext->GetDebugSimName());
 		}
 
-		// RDG will defer the execution of the actual dispatch calls until much later in the frame, so anything that is required for
-		// Mesh Processor commands to execute correct we must setup immediately.  I.e. the final buffer / final count.
-		if (ComputeProxy->GetComputeTickStage() == ENiagaraGpuComputeTickStage::PreInitViews)
+		// RDG will defer the Niagara dispatches until the graph is executed.
+		// Therefore we need to setup the DataToRender for MeshProcessors & sorting to use the correct data,
+		// that is anything that happens before PostRenderOpaque
+		if ((ComputeProxy->GetComputeTickStage() == ENiagaraGpuComputeTickStage::PreInitViews) || (ComputeProxy->GetComputeTickStage() == ENiagaraGpuComputeTickStage::PostInitViews))
 		{
 			FNiagaraDataBuffer* FinalBuffer = ComputeContext->GetCurrDataBuffer();
 			FinalBuffer->GPUInstanceCountBufferOffset = ComputeContext->CountOffset_RT;
