- Fix handling instances that complete & activate while we are waiting on them

#rnx

diff --git a/Engine/Plugins/FX/Niagara/Source/Niagara/Private/NiagaraSystemSimulation.cpp b/Engine/Plugins/FX/Niagara/Source/Niagara/Private/NiagaraSystemSimulation.cpp
--- a/Engine/Plugins/FX/Niagara/Source/Niagara/Private/NiagaraSystemSimulation.cpp
+++ b/Engine/Plugins/FX/Niagara/Source/Niagara/Private/NiagaraSystemSimulation.cpp
@@ -1058,15 +1058,10 @@ void FNiagaraSystemSimulation::Tick_GameThread(float DeltaSeconds, const FGraphE
 		// Perform instance tick
 		Instance->Tick_GameThread(DeltaSeconds);
 
-		// TickDataInterfaces could remove the system so we only increment if the system has changed
-		// Also possible for this system to have been transfered to another system simulation.
-		if ( Instance->GetSystemSimulation().Get() == this )
+		// Ticking the instance can result in it being removed, completing + reactivating or transferring
+		if (SystemInstances.IsValidIndex(SystemIndex) && (SystemInstances[SystemIndex] == Instance))
 		{
-			if (Instance->SystemInstanceIndex != INDEX_NONE)
-			{
-				check(Instance->SystemInstanceIndex == SystemIndex);
-				++SystemIndex;
-			}
+			++SystemIndex;
 		}
 	}
 
@@ -1079,10 +1074,9 @@ void FNiagaraSystemSimulation::Tick_GameThread(float DeltaSeconds, const FGraphE
 	{
 		SystemInstances.Reserve(SystemInstances.Num() + PendingSystemInstances.Num());
 
-		SystemIndex = 0;
-		while (SystemIndex < PendingSystemInstances.Num())
+		while (PendingSystemInstances.Num() > 0)
 		{
-			FNiagaraSystemInstance* Instance = PendingSystemInstances[SystemIndex];
+			FNiagaraSystemInstance* Instance = PendingSystemInstances[0];
 
 			// Not solo, look to see if we should move tick group
 			if (!bIsSolo)
@@ -1100,11 +1094,9 @@ void FNiagaraSystemSimulation::Tick_GameThread(float DeltaSeconds, const FGraphE
 			PARTICLE_PERF_STAT_CYCLES_WITH_COUNT_GT(FParticlePerfStatsContext(GetInstancePerfStats(Instance)), TickGameThread, 1);
 			Instance->Tick_GameThread(DeltaSeconds);
 
-			if (Instance->SystemInstanceIndex != INDEX_NONE)
+			// Ensure we handle the instance being killed or potentially transfering system
+			if ( PendingSystemInstances.IsValidIndex(0) && (PendingSystemInstances[0] == Instance) )
 			{
-				// We should not move tick group during Tick_GameThread but let's be safe
-				check(Instance->SystemSimulation.Get() == this);
-
 				// When the first instance is added we need to initialize the parameter store to data set bindings.
 				if (!bBindingsInitialized)
 				{
@@ -1368,9 +1360,9 @@ void FNiagaraSystemSimulation::WaitForInstancesTickComplete(bool bEnsureComplete
 			FNiagaraSystemInstance* Instance = SystemInstances[SystemInstanceIndex];
 			Instance->WaitForConcurrentTickAndFinalize();
 
-			if ( !Instance->IsComplete() )
+			// Finalize may result in the instance being removed or it being completed & reactivated
+			if (SystemInstances.IsValidIndex(Instance->SystemInstanceIndex) && (SystemInstances[Instance->SystemInstanceIndex] == Instance) )
 			{
-				check(SystemInstances[SystemInstanceIndex] == Instance);
 				++SystemInstanceIndex;
 			}
 		}
@@ -1428,16 +1420,16 @@ void FNiagaraSystemSimulation::Tick_Concurrent(FNiagaraSystemSimulationTickConte
 		if ( !Context.IsRunningAsync() )
 		{
 			check(IsInGameThread());
-			int32 SystemInstIndex = 0;
-			while (SystemInstIndex < Context.Instances.Num())
+			int32 InstanceIndex = 0;
+			while (InstanceIndex < Context.Instances.Num())
 			{
-				FNiagaraSystemInstance* Inst = Context.Instances[SystemInstIndex];
-				Inst->FinalizeTick_GameThread();
+				FNiagaraSystemInstance* Instance = Context.Instances[InstanceIndex];
+				Instance->FinalizeTick_GameThread();
 
-				// If the system completes during finalize it will be removed from the instances, therefore we do not need to increment our system index;
-				if (!Inst->IsComplete())
+				// Finalize can complete the instance and potentially reactivate
+				if (Context.Instances.IsValidIndex(InstanceIndex) && (Context.Instances[InstanceIndex] == Instance))
 				{
-					++SystemInstIndex;
+					++InstanceIndex;
 				}
 
 				check(Context.DataSet.GetCurrentDataChecked().GetNumInstances() == Context.Instances.Num());
