- Editgrate CL 16474571
- Fix pulling from incorrect data source on first tick
- Fix overwriting the buffer we could potentially be rendering from, shows up with translucent systems that tick in PostRenderOpaque with multiple buffer passes

#rnx

diff --git a/Engine/Plugins/FX/Niagara/Source/Niagara/Classes/NiagaraComputeExecutionContext.h b/Engine/Plugins/FX/Niagara/Source/Niagara/Classes/NiagaraComputeExecutionContext.h
--- a/Engine/Plugins/FX/Niagara/Source/Niagara/Classes/NiagaraComputeExecutionContext.h
+++ b/Engine/Plugins/FX/Niagara/Source/Niagara/Classes/NiagaraComputeExecutionContext.h
@@ -25,10 +25,12 @@ struct FNiagaraSimStageData
 	{
 		bFirstStage = false;
 		bLastStage = false;
+		bSetDataToRender = false;
 	}
 
 	uint32 bFirstStage : 1;
 	uint32 bLastStage : 1;
+	uint32 bSetDataToRender : 1;
 
 	uint32 StageIndex = INDEX_NONE;
 
diff --git a/Engine/Plugins/FX/Niagara/Source/Niagara/Classes/NiagaraDataSet.h b/Engine/Plugins/FX/Niagara/Source/Niagara/Classes/NiagaraDataSet.h
--- a/Engine/Plugins/FX/Niagara/Source/Niagara/Classes/NiagaraDataSet.h
+++ b/Engine/Plugins/FX/Niagara/Source/Niagara/Classes/NiagaraDataSet.h
@@ -115,7 +115,7 @@ public:
 	void Allocate(uint32 NumInstances, bool bMaintainExisting = false);
 
 	void AllocateGPU(FRHICommandList& RHICmdList, uint32 InNumInstances, ERHIFeatureLevel::Type FeatureLevel, const TCHAR* DebugSimName);
-	void AliasGPU(FNiagaraDataBuffer* BufferToAlias);
+	void SwapGPU(FNiagaraDataBuffer* BufferToAlias);
 	void ReleaseGPU();
 
 	void SwapInstances(uint32 OldIndex, uint32 NewIndex);
@@ -166,7 +166,6 @@ public:
 	FORCEINLINE FRWBuffer& GetGPUBufferInt() { return GPUBufferInt; }
 	FORCEINLINE FRWBuffer& GetGPUBufferHalf() { return GPUBufferHalf; }
 	FORCEINLINE uint32 GetGPUInstanceCountBufferOffset() const { return GPUInstanceCountBufferOffset; }
-	FORCEINLINE void ClearGPUInstanceCountBufferOffset() { GPUInstanceCountBufferOffset = INDEX_NONE; }
 	FORCEINLINE FRWBuffer& GetGPUIDToIndexTable() { return GPUIDToIndexTable; }
 
 	FORCEINLINE int32 GetSafeComponentBufferSize() const { return GetSafeComponentBufferSize(GetNumInstancesAllocated()); }
diff --git a/Engine/Plugins/FX/Niagara/Source/Niagara/Private/NiagaraDataSet.cpp b/Engine/Plugins/FX/Niagara/Source/Niagara/Private/NiagaraDataSet.cpp
--- a/Engine/Plugins/FX/Niagara/Source/Niagara/Private/NiagaraDataSet.cpp
+++ b/Engine/Plugins/FX/Niagara/Source/Niagara/Private/NiagaraDataSet.cpp
@@ -850,26 +850,24 @@ void FNiagaraDataBuffer::AllocateGPU(FRHICommandList& RHICmdList, uint32 InNumIn
 #endif
 }
 
-void FNiagaraDataBuffer::AliasGPU(FNiagaraDataBuffer* BufferToAlias)
+void FNiagaraDataBuffer::SwapGPU(FNiagaraDataBuffer* BufferToAlias)
 {
 #if NIAGARA_MEMORY_TRACKING
-	DEC_MEMORY_STAT_BY(STAT_NiagaraGPUParticleMemory, AllocationSizeBytes);
-	Owner->BufferSizeBytes -= AllocationSizeBytes;
-	AllocationSizeBytes = 0;
+	Swap(AllocationSizeBytes, BufferToAlias->AllocationSizeBytes);
 #endif
 
-	GPUInstanceCountBufferOffset = BufferToAlias->GPUInstanceCountBufferOffset;
-	GPUBufferFloat = BufferToAlias->GPUBufferFloat;
-	GPUBufferInt = BufferToAlias->GPUBufferInt;
-	GPUIDToIndexTable = BufferToAlias->GPUIDToIndexTable;
-	GPUBufferHalf = BufferToAlias->GPUBufferHalf;
-	NumInstances = BufferToAlias->NumInstances;
-	NumInstancesAllocated = BufferToAlias->NumInstancesAllocated;
-	FloatStride = BufferToAlias->FloatStride;
-	Int32Stride = BufferToAlias->Int32Stride;
-	HalfStride = BufferToAlias->HalfStride;
-	NumSpawnedInstances = BufferToAlias->NumSpawnedInstances;
-	IDAcquireTag = BufferToAlias->IDAcquireTag;
+	Swap(GPUInstanceCountBufferOffset, BufferToAlias->GPUInstanceCountBufferOffset);
+	Swap(GPUBufferFloat, BufferToAlias->GPUBufferFloat);
+	Swap(GPUBufferInt, BufferToAlias->GPUBufferInt);
+	Swap(GPUIDToIndexTable, BufferToAlias->GPUIDToIndexTable);
+	Swap(GPUBufferHalf, BufferToAlias->GPUBufferHalf);
+	Swap(NumInstances, BufferToAlias->NumInstances);
+	Swap(NumInstancesAllocated, BufferToAlias->NumInstancesAllocated);
+	Swap(FloatStride, BufferToAlias->FloatStride);
+	Swap(Int32Stride, BufferToAlias->Int32Stride);
+	Swap(HalfStride, BufferToAlias->HalfStride);
+	Swap(NumSpawnedInstances, BufferToAlias->NumSpawnedInstances);
+	Swap(IDAcquireTag, BufferToAlias->IDAcquireTag);
 }
 
 void FNiagaraDataBuffer::ReleaseGPU()
diff --git a/Engine/Plugins/FX/Niagara/Source/Niagara/Private/NiagaraEmitterInstanceBatcher.cpp b/Engine/Plugins/FX/Niagara/Source/Niagara/Private/NiagaraEmitterInstanceBatcher.cpp
--- a/Engine/Plugins/FX/Niagara/Source/Niagara/Private/NiagaraEmitterInstanceBatcher.cpp
+++ b/Engine/Plugins/FX/Niagara/Source/Niagara/Private/NiagaraEmitterInstanceBatcher.cpp
@@ -735,11 +735,17 @@ void NiagaraEmitterInstanceBatcher::DumpDebugFrame()
 					Builder.Append(TEXT("LastStage "));
 				}
 
+				if (DispatchInstance.SimStageData.bSetDataToRender)
+				{
+					Builder.Append(TEXT("SetDataToRender "));
+				}
+					
+
 				if (InstanceData.Context->EmitterInstanceReadback.GPUCountOffset != INDEX_NONE)
 				{
 					if (InstanceData.Context->EmitterInstanceReadback.GPUCountOffset == SimStageData.SourceCountOffset)
 					{
-						Builder.Append(TEXT("ReadbackSource "));
+						Builder.Appendf(TEXT("ReadbackSource(%d) "), InstanceData.Context->EmitterInstanceReadback.CPUCount);
 					}
 				}
 				Builder.Appendf(TEXT("Source(%p 0x%08x %d) "), SimStageData.Source, SimStageData.SourceCountOffset, SimStageData.SourceNumInstances);
@@ -789,8 +795,8 @@ void NiagaraEmitterInstanceBatcher::UpdateInstanceCountManager(FRHICommandListIm
 {
 	// Resize dispatch buffer count
 	//-OPT: No need to iterate over all the ticks, we can store this as ticks are queued
-	int32 TotalDispatchCount = 0;
 	{
+		int32 TotalDispatchCount = 0;
 		for (int iTickStage = 0; iTickStage < ENiagaraGpuComputeTickStage::Max; ++iTickStage)
 		{
 			for (FNiagaraSystemGpuComputeProxy* ComputeProxy : ProxiesPerStage[iTickStage])
@@ -844,7 +850,7 @@ void NiagaraEmitterInstanceBatcher::UpdateInstanceCountManager(FRHICommandListIm
 						}
 
 						// Release the counter for the readback
-						ComputeContext->ReleaseReadbackCounter(GPUInstanceCounterManager);
+						GPUInstanceCounterManager.FreeEntry(ComputeContext->EmitterInstanceReadback.GPUCountOffset);
 					}
 				}
 			}
@@ -955,7 +961,7 @@ void NiagaraEmitterInstanceBatcher::PrepareTicksForProxy(FRHICommandListImmediat
 
 				SimStageData.bFirstStage = true;
 				SimStageData.StageIndex = 0;
-				SimStageData.Source = ComputeContext->GetCurrDataBuffer();
+				SimStageData.Source = ComputeContext->bHasTickedThisFrame_RT ? ComputeContext->GetCurrDataBuffer() : ComputeContext->MainDataSet->GetCurrentData();
 				SimStageData.SourceCountOffset = ComputeContext->CountOffset_RT;
 				SimStageData.SourceNumInstances = PrevNumInstances;
 
@@ -991,11 +997,10 @@ void NiagaraEmitterInstanceBatcher::PrepareTicksForProxy(FRHICommandListImmediat
 			}
 			ComputeContext->bHasTickedThisFrame_RT = true;
 
-			//-OPT: Seems like a weird test
+			//-OPT: Do we need this test?  Can remove in favor of MaxUpdateIterations
 			if (Tick.NumInstancesWithSimStages > 0)
 			{
-				//-OPT: Is this ever true?
-				for (uint32 SimStageIndex = 1; SimStageIndex < ComputeContext->MaxUpdateIterations; ++SimStageIndex)
+				for (uint32 SimStageIndex=1; SimStageIndex < ComputeContext->MaxUpdateIterations; ++SimStageIndex)
 				{
 					// Determine if the iteration is outputting to a custom data size
 					FNiagaraDataInterfaceProxyRW* IterationInterface = InstanceData.FindIterationInterface(SimStageIndex);
@@ -1058,6 +1063,7 @@ void NiagaraEmitterInstanceBatcher::PrepareTicksForProxy(FRHICommandListImmediat
 			// The final instance in the last group we modified is the final stage for this tick
 			FNiagaraGpuDispatchGroup& FinalDispatchGroup = GpuDispatchList.DispatchGroups[iInstanceCurrDispatchGroup - 1];
 			FinalDispatchGroup.DispatchInstances.Last().SimStageData.bLastStage = true;
+			FinalDispatchGroup.DispatchInstances.Last().SimStageData.bSetDataToRender = Tick.bIsFinalTick;
 
 			// Keep track of where the next set of dispatch should occur
 			iTickStartDispatchGroup = FMath::Max(iTickStartDispatchGroup, iInstanceCurrDispatchGroup);
@@ -1105,6 +1111,7 @@ void NiagaraEmitterInstanceBatcher::PrepareTicksForProxy(FRHICommandListImmediat
 		// Allocate space for all buffers we use
 		//-OPT: We can batch the allocation of persistent IDs together so the compute shaders overlap
 		const uint32 NumBuffersToResize = FMath::Min<uint32>(ComputeContext->BufferSwapsThisFrame_RT, UE_ARRAY_COUNT(ComputeContext->DataBuffers_RT));
+
 		for (uint32 i = 0; i < NumBuffersToResize; ++i)
 		{
 			const uint32 BufferIndex = (ComputeContext->DataBufferIndex_RT + i) % UE_ARRAY_COUNT(ComputeContext->DataBuffers_RT);
@@ -1289,7 +1296,10 @@ void NiagaraEmitterInstanceBatcher::ExecuteTicks(FRHICommandList& RHICmdList, FR
 			if ( DispatchInstance.SimStageData.bLastStage )
 			{
 				PostSimulateInterface(RHICmdList, DispatchInstance.Tick, DispatchInstance.InstanceData);
+			}
 
+			if (DispatchInstance.SimStageData.bSetDataToRender)
+			{
 				// Update CurrentData with the latest information as things like ParticleReads can use this data
 				FNiagaraComputeExecutionContext* ComputeContext = DispatchInstance.InstanceData.Context;
 				const FNiagaraSimStageData& FinalSimStageData = DispatchInstance.SimStageData;
@@ -1297,9 +1307,8 @@ void NiagaraEmitterInstanceBatcher::ExecuteTicks(FRHICommandList& RHICmdList, FR
 				FNiagaraDataBuffer* FinalData = FinalSimStageData.Destination != nullptr ? FinalSimStageData.Destination : FinalSimStageData.Source;
 				check(FinalData);
 
-				CurrentData->AliasGPU(FinalData);
+				CurrentData->SwapGPU(FinalData);
 
-				ComputeContext->SetTranslucentDataToRender(CurrentData);
 				ComputeContext->SetDataToRender(CurrentData);
 			}
 		}
diff --git a/Engine/Plugins/FX/Niagara/Source/Niagara/Private/NiagaraSystemInstance.cpp b/Engine/Plugins/FX/Niagara/Source/Niagara/Private/NiagaraSystemInstance.cpp
--- a/Engine/Plugins/FX/Niagara/Source/Niagara/Private/NiagaraSystemInstance.cpp
+++ b/Engine/Plugins/FX/Niagara/Source/Niagara/Private/NiagaraSystemInstance.cpp
@@ -852,9 +852,6 @@ void FNiagaraSystemInstance::Reset(FNiagaraSystemInstance::EResetMode Mode)
 				TickCount = 0;
 			}
 		}
-		else
-		{
-		}
 	}
 	else
 	{
