CIS fix post merge

diff --git a/Engine/Plugins/FX/Niagara/Source/Niagara/Private/NiagaraSystem.cpp b/Engine/Plugins/FX/Niagara/Source/Niagara/Private/NiagaraSystem.cpp
--- a/Engine/Plugins/FX/Niagara/Source/Niagara/Private/NiagaraSystem.cpp
+++ b/Engine/Plugins/FX/Niagara/Source/Niagara/Private/NiagaraSystem.cpp
@@ -782,11 +782,14 @@ void UNiagaraSystem::PostLoad()
 	// there will be no merging or compiling which makes it safe to do so.
 	UpdateSystemAfterLoad();
 #endif
+
+#if WITH_EDITORONLY_DATA
 	// see the equivalent in NiagaraEmitter for details
 	if(bIsTemplateAsset_DEPRECATED)
 	{
 		TemplateSpecification = bIsTemplateAsset_DEPRECATED ? ENiagaraScriptTemplateSpecification::Template : ENiagaraScriptTemplateSpecification::None;
 	}
+#endif // WITH_EDITORONLY_DATA
 }
 
 #if WITH_EDITORONLY_DATA
