fixing bad merge

diff --git a/Engine/Plugins/FX/Niagara/Source/Niagara/Classes/NiagaraScript.h b/Engine/Plugins/FX/Niagara/Source/Niagara/Classes/NiagaraScript.h
--- a/Engine/Plugins/FX/Niagara/Source/Niagara/Classes/NiagaraScript.h
+++ b/Engine/Plugins/FX/Niagara/Source/Niagara/Classes/NiagaraScript.h
@@ -516,6 +516,7 @@ public:
 	mutable FNiagaraVMExecutableDataId LastGeneratedVMId;
 
 	/** Reference to a python script that is executed when the user updates from a previous version to this version. */
+	UPROPERTY()
 	ENiagaraPythonUpdateScriptReference UpdateScriptExecution = ENiagaraPythonUpdateScriptReference::None;
 
 	/** Python script to run when updating to this script version. */
