- When world rendering is disabled ensure we flush Niagara GPU jobs to avoid a larger flush or loosing work depending on tick mode

diff --git a/Engine/Plugins/FX/Niagara/Source/Niagara/Classes/NiagaraEmitterInstanceBatcher.h b/Engine/Plugins/FX/Niagara/Source/Niagara/Classes/NiagaraEmitterInstanceBatcher.h
--- a/Engine/Plugins/FX/Niagara/Source/Niagara/Classes/NiagaraEmitterInstanceBatcher.h
+++ b/Engine/Plugins/FX/Niagara/Source/Niagara/Classes/NiagaraEmitterInstanceBatcher.h
@@ -87,7 +87,7 @@ public:
 	virtual void PreRender(FRDGBuilder& GraphBuilder, TConstArrayView<FViewInfo> Views, bool bAllowGPUParticleUpdate) override;
 	virtual void OnDestroy() override; // Called on the gamethread to delete the batcher on the renderthread.
 
-	virtual void Tick(float DeltaTime) override;
+	virtual void Tick(UWorld* World, float DeltaTime) override;
 
 	virtual void PostRenderOpaque(FRDGBuilder& GraphBuilder, TConstArrayView<FViewInfo> Views, bool bAllowGPUParticleUpdate) override;
 
diff --git a/Engine/Plugins/FX/Niagara/Source/Niagara/Private/NiagaraEmitterInstanceBatcher.cpp b/Engine/Plugins/FX/Niagara/Source/Niagara/Private/NiagaraEmitterInstanceBatcher.cpp
--- a/Engine/Plugins/FX/Niagara/Source/Niagara/Private/NiagaraEmitterInstanceBatcher.cpp
+++ b/Engine/Plugins/FX/Niagara/Source/Niagara/Private/NiagaraEmitterInstanceBatcher.cpp
@@ -6,8 +6,11 @@
 #include "CanvasTypes.h"
 #include "ClearQuad.h"
 #include "Engine/Canvas.h"
+#include "Engine/GameViewportClient.h"
+#include "Engine/World.h"
 #include "GPUSort.h"
 #include "Misc/ScopeExit.h"
+
 #include "NiagaraDataInterfaceRW.h"
 #if WITH_EDITOR
 #include "NiagaraGpuComputeDebug.h"
@@ -405,13 +408,23 @@ void NiagaraEmitterInstanceBatcher::BuildConstantBuffers(FNiagaraGPUSystemTick&
 	}
 }
 
-void NiagaraEmitterInstanceBatcher::Tick(float DeltaTime)
+void NiagaraEmitterInstanceBatcher::Tick(UWorld* World, float DeltaTime)
 {
+	bool bFlushTicks = false;
+
+	if ((World != nullptr) && World->IsGameWorld())
+	{
+		if (UGameViewportClient* GameViewport = World->GetGameViewport())
+		{
+			bFlushTicks |= GameViewport->bDisableWorldRendering;
+		}
+	}
+
 	check(IsInGameThread());
 	ENQUEUE_RENDER_COMMAND(NiagaraPumpBatcher)(
-		[RT_NiagaraBatcher=this](FRHICommandListImmediate& RHICmdList)
+		[RT_NiagaraBatcher=this, RT_bFlushTicks=bFlushTicks](FRHICommandListImmediate& RHICmdList)
 		{
-			RT_NiagaraBatcher->ProcessPendingTicksFlush(RHICmdList, false);
+			RT_NiagaraBatcher->ProcessPendingTicksFlush(RHICmdList, RT_bFlushTicks);
 			RT_NiagaraBatcher->GetGPUInstanceCounterManager().FlushIndirectArgsPool();
 		}
 	);
diff --git a/Engine/Source/Editor/Cascade/Private/Cascade.cpp b/Engine/Source/Editor/Cascade/Private/Cascade.cpp
--- a/Engine/Source/Editor/Cascade/Private/Cascade.cpp
+++ b/Engine/Source/Editor/Cascade/Private/Cascade.cpp
@@ -1608,12 +1608,12 @@ void FCascade::Tick(float DeltaTime)
 			ParticleSystemComponent->CascadeTickComponent(CurrDeltaTime, LEVELTICK_All);
 		}
 		ParticleSystemComponent->DoDeferredRenderUpdates_Concurrent();
-		GetFXSystem()->Tick(CurrDeltaTime);
+		UWorld* World = PreviewViewport->GetViewportClient()->GetPreviewScene().GetWorld();
+		GetFXSystem()->Tick(World, CurrDeltaTime);
 		TotalTime += CurrDeltaTime;
 		ParticleSystem->UpdateTime_Delta = fSaveUpdateDelta;
 
 		// Tick the physics scene
-		UWorld* World = PreviewViewport->GetViewportClient()->GetPreviewScene().GetWorld();
 		FPhysScene* PhysScene = World->GetPhysicsScene();
 		AWorldSettings * WorldSettings = World->GetWorldSettings();
 		check(WorldSettings);
diff --git a/Engine/Source/Editor/UnrealEd/Private/ThumbnailHelpers.cpp b/Engine/Source/Editor/UnrealEd/Private/ThumbnailHelpers.cpp
--- a/Engine/Source/Editor/UnrealEd/Private/ThumbnailHelpers.cpp
+++ b/Engine/Source/Editor/UnrealEd/Private/ThumbnailHelpers.cpp
@@ -221,7 +221,7 @@ void FParticleSystemThumbnailScene::SetParticleSystem(UParticleSystem* ParticleS
 				{
 					ParticleSystem->PreviewComponent->TickComponent(WarmupTimestep, LEVELTICK_All, NULL);
 					WarmupElapsed += WarmupTimestep;
-					ThumbnailFXSystem->Tick(WarmupTimestep);
+					ThumbnailFXSystem->Tick(ParticleSystem->PreviewComponent->GetWorld(), WarmupTimestep);
 				}
 			}
 		}
diff --git a/Engine/Source/Runtime/Engine/Private/LevelTick.cpp b/Engine/Source/Runtime/Engine/Private/LevelTick.cpp
--- a/Engine/Source/Runtime/Engine/Private/LevelTick.cpp
+++ b/Engine/Source/Runtime/Engine/Private/LevelTick.cpp
@@ -1733,7 +1733,7 @@ void UWorld::Tick( ELevelTick TickType, float DeltaSeconds )
 	{
 		TRACE_CPUPROFILER_EVENT_SCOPE(FXSystem);
 		SCOPE_TIME_GUARD_MS(TEXT("UWorld::Tick - FX"), 5);
-		FXSystem->Tick(DeltaSeconds);
+		FXSystem->Tick(this, DeltaSeconds);
 	}
 
 #if WITH_EDITOR
diff --git a/Engine/Source/Runtime/Engine/Private/Particles/FXSystem.cpp b/Engine/Source/Runtime/Engine/Private/Particles/FXSystem.cpp
--- a/Engine/Source/Runtime/Engine/Private/Particles/FXSystem.cpp
+++ b/Engine/Source/Runtime/Engine/Private/Particles/FXSystem.cpp
@@ -255,7 +255,7 @@ FFXSystemInterface* FFXSystem::GetInterface(const FName& InName)
 	return InName == Name ? this : nullptr;
 }
 
-void FFXSystem::Tick(float DeltaSeconds)
+void FFXSystem::Tick(UWorld* World, float DeltaSeconds)
 {
 	if (RHISupportsGPUParticles())
 	{
diff --git a/Engine/Source/Runtime/Engine/Private/Particles/FXSystemPrivate.h b/Engine/Source/Runtime/Engine/Private/Particles/FXSystemPrivate.h
--- a/Engine/Source/Runtime/Engine/Private/Particles/FXSystemPrivate.h
+++ b/Engine/Source/Runtime/Engine/Private/Particles/FXSystemPrivate.h
@@ -140,7 +140,7 @@ public:
 	virtual FFXSystemInterface* GetInterface(const FName& InName) override;
 
 	// Begin FFXSystemInterface.
-	virtual void Tick(float DeltaSeconds) override;
+	virtual void Tick(UWorld* World, float DeltaSeconds) override;
 #if WITH_EDITOR
 	virtual void Suspend() override;
 	virtual void Resume() override;
diff --git a/Engine/Source/Runtime/Engine/Private/Particles/FXSystemSet.cpp b/Engine/Source/Runtime/Engine/Private/Particles/FXSystemSet.cpp
--- a/Engine/Source/Runtime/Engine/Private/Particles/FXSystemSet.cpp
+++ b/Engine/Source/Runtime/Engine/Private/Particles/FXSystemSet.cpp
@@ -26,12 +26,12 @@ FFXSystemInterface* FFXSystemSet::GetInterface(const FName& InName)
 	return nullptr;
 }
 
-void FFXSystemSet::Tick(float DeltaSeconds)
+void FFXSystemSet::Tick(UWorld* World, float DeltaSeconds)
 {
 	for (FFXSystemInterface* FXSystem : FXSystems)
 	{
 		check(FXSystem);
-		FXSystem->Tick(DeltaSeconds);
+		FXSystem->Tick(World, DeltaSeconds);
 	}
 }
 
diff --git a/Engine/Source/Runtime/Engine/Private/Particles/FXSystemSet.h b/Engine/Source/Runtime/Engine/Private/Particles/FXSystemSet.h
--- a/Engine/Source/Runtime/Engine/Private/Particles/FXSystemSet.h
+++ b/Engine/Source/Runtime/Engine/Private/Particles/FXSystemSet.h
@@ -25,7 +25,7 @@ public:
 	TArray<FFXSystemInterface*> FXSystems;
 
 	virtual FFXSystemInterface* GetInterface(const FName& InName) override;
-	virtual void Tick(float DeltaSeconds) override;
+	virtual void Tick(UWorld* World, float DeltaSeconds) override;
 
 #if WITH_EDITOR
 	virtual void Suspend() override;
diff --git a/Engine/Source/Runtime/Engine/Public/FXSystem.h b/Engine/Source/Runtime/Engine/Public/FXSystem.h
--- a/Engine/Source/Runtime/Engine/Public/FXSystem.h
+++ b/Engine/Source/Runtime/Engine/Public/FXSystem.h
@@ -170,7 +170,7 @@ public:
 	 * Tick the effects system.
 	 * @param DeltaSeconds The number of seconds by which to step simulations forward.
 	 */
-	virtual void Tick(float DeltaSeconds) = 0;
+	virtual void Tick(UWorld* World, float DeltaSeconds) = 0;
 
 #if WITH_EDITOR
 	/**
