- Replace RDG based mip generation with custom until Niagara becomes compliant

#rnx

diff --git a/Engine/Plugins/FX/Niagara/Shaders/Private/NiagaraGenerateMips.usf b/Engine/Plugins/FX/Niagara/Shaders/Private/NiagaraGenerateMips.usf
new file mode 100644
--- /dev/null
+++ b/Engine/Plugins/FX/Niagara/Shaders/Private/NiagaraGenerateMips.usf
@@ -0,0 +1,30 @@
+// Copyright Epic Games, Inc. All Rights Reserved.
+
+#include "/Engine/Private/Common.ush"
+#include "/Engine/Private/GammaCorrectionCommon.ush"
+
+float2 TexelSize;
+
+Texture2D MipInSRV;
+SamplerState MipInSampler;
+
+int2 MipOutSize;
+RWTexture2D<float4> MipOutUAV;
+
+[numthreads(THREADGROUPSIZE, THREADGROUPSIZE, 1)]
+void MainCS(uint3 ThreadID : SV_DispatchThreadID)
+{
+	if ( any(int2(ThreadID.xy) >= MipOutSize) )
+	{
+		return;
+	}
+
+	float2 SrcUV = (float2(ThreadID.xy) + 0.5f) * TexelSize;
+    float4 OutColor = MipInSRV.SampleLevel(MipInSampler, SrcUV, 0);
+
+#if GENMIPS_SRGB
+    OutColor.rgb = LinearToSrgb(OutColor.rgb);
+#endif
+
+    MipOutUAV[ThreadID.xy] = OutColor;
+}
diff --git a/Engine/Plugins/FX/Niagara/Source/Niagara/Private/NiagaraDataInterfaceRenderTarget2D.cpp b/Engine/Plugins/FX/Niagara/Source/Niagara/Private/NiagaraDataInterfaceRenderTarget2D.cpp
--- a/Engine/Plugins/FX/Niagara/Source/Niagara/Private/NiagaraDataInterfaceRenderTarget2D.cpp
+++ b/Engine/Plugins/FX/Niagara/Source/Niagara/Private/NiagaraDataInterfaceRenderTarget2D.cpp
@@ -16,6 +16,7 @@
 #if WITH_EDITOR
 #include "NiagaraGpuComputeDebug.h"
 #endif
+#include "NiagaraGenerateMips.h"
 #include "NiagaraStats.h"
 
 #define LOCTEXT_NAMESPACE "NiagaraDataInterfaceRenderTarget2D"
@@ -803,15 +804,7 @@ void FNiagaraDataInterfaceProxyRenderTarget2DProxy::PostStage(FRHICommandList& R
 		if (ProxyData->bWasWrittenTo && (ProxyData->MipMapGeneration == ENiagaraMipMapGeneration::PostStage))
 		{
 			ProxyData->bWasWrittenTo = false;
-
-			check(&RHICmdList == &FRHICommandListExecutor::GetImmediateCommandList());
-
-			FMemMark MemMark(FMemStack::Get());
-			FRDGBuilder GraphBuilder(FRHICommandListExecutor::GetImmediateCommandList());
-			TRefCountPtr<IPooledRenderTarget> PoolRenderTarget = CreateRenderTarget(ProxyData->TextureRHI, TEXT("NiagaraRenderTarget2D"));
-			FRDGTextureRef MipOutputTexture = GraphBuilder.RegisterExternalTexture(PoolRenderTarget);
-			FGenerateMips::Execute(GraphBuilder, MipOutputTexture);
-			GraphBuilder.Execute();
+			NiagaraGenerateMips::GenerateMips(RHICmdList, ProxyData->TextureRHI);
 		}
 	}
 }
@@ -823,15 +816,7 @@ void FNiagaraDataInterfaceProxyRenderTarget2DProxy::PostSimulate(FRHICommandList
 		if (ProxyData->bWasWrittenTo && (ProxyData->MipMapGeneration == ENiagaraMipMapGeneration::PostSimulate))
 		{
 			ProxyData->bWasWrittenTo = false;
-
-			check(&RHICmdList == &FRHICommandListExecutor::GetImmediateCommandList());
-
-			FMemMark MemMark(FMemStack::Get());
-			FRDGBuilder GraphBuilder(FRHICommandListExecutor::GetImmediateCommandList());
-			TRefCountPtr<IPooledRenderTarget> PoolRenderTarget = CreateRenderTarget(ProxyData->TextureRHI, TEXT("NiagaraRenderTarget2D"));
-			FRDGTextureRef MipOutputTexture = GraphBuilder.RegisterExternalTexture(PoolRenderTarget);
-			FGenerateMips::Execute(GraphBuilder, MipOutputTexture);
-			GraphBuilder.Execute();
+			NiagaraGenerateMips::GenerateMips(RHICmdList, ProxyData->TextureRHI);
 		}
 
 #if NIAGARA_COMPUTEDEBUG_ENABLED
diff --git a/Engine/Plugins/FX/Niagara/Source/NiagaraShader/Private/NiagaraGenerateMips.cpp b/Engine/Plugins/FX/Niagara/Source/NiagaraShader/Private/NiagaraGenerateMips.cpp
new file mode 100644
--- /dev/null
+++ b/Engine/Plugins/FX/Niagara/Source/NiagaraShader/Private/NiagaraGenerateMips.cpp
@@ -0,0 +1,84 @@
+// Copyright Epic Games, Inc. All Rights Reserved.
+
+#include "NiagaraGenerateMips.h"
+#include "RenderGraphUtils.h"
+
+class FNiagaraGenerateMipsCS : public FGlobalShader
+{
+public:
+	DECLARE_GLOBAL_SHADER(FNiagaraGenerateMipsCS)
+	SHADER_USE_PARAMETER_STRUCT(FNiagaraGenerateMipsCS, FGlobalShader)
+
+	class FGenMipsSRGB : SHADER_PERMUTATION_BOOL("GENMIPS_SRGB");
+	using FPermutationDomain = TShaderPermutationDomain<FGenMipsSRGB>;
+
+	BEGIN_SHADER_PARAMETER_STRUCT(FParameters, )
+		SHADER_PARAMETER(FVector2D, TexelSize)
+		SHADER_PARAMETER(FIntPoint, MipOutSize)
+		SHADER_PARAMETER_UAV(RWTexture2D<float4>, MipOutUAV)
+		SHADER_PARAMETER_SRV(Texture2D, MipInSRV)
+		SHADER_PARAMETER_SAMPLER(SamplerState, MipInSampler)
+	END_SHADER_PARAMETER_STRUCT()
+
+	static bool ShouldCompilePermutation(const FGlobalShaderPermutationParameters& Parameters)
+	{
+		return RHISupportsComputeShaders(Parameters.Platform);
+	}
+
+	static void ModifyCompilationEnvironment(const FGlobalShaderPermutationParameters& Parameters, FShaderCompilerEnvironment& OutEnvironment)
+	{
+		FGlobalShader::ModifyCompilationEnvironment(Parameters, OutEnvironment);
+		OutEnvironment.SetDefine(TEXT("THREADGROUPSIZE"), FComputeShaderUtils::kGolden2DGroupSize);
+	}
+};
+
+IMPLEMENT_GLOBAL_SHADER(FNiagaraGenerateMipsCS, "/Plugin/FX/Niagara/Private/NiagaraGenerateMips.usf", "MainCS", SF_Compute);
+
+void NiagaraGenerateMips::GenerateMips(FRHICommandList& RHICmdList, FRHITexture2D* TextureRHI)
+{
+	SCOPED_DRAW_EVENT(RHICmdList, NiagaraGenerateMips);
+
+	// Select compute shader variant (normal vs. sRGB etc.)
+	bool bMipsSRGB = EnumHasAnyFlags(TextureRHI->GetFlags(), TexCreate_SRGB);
+#if PLATFORM_ANDROID
+	if (IsVulkanPlatform(GMaxRHIShaderPlatform))
+	{
+		// Vulkan Android seems to skip sRGB->Lin conversion when sampling texture in compute
+		bMipsSRGB = false;
+	}
+#endif
+
+	FNiagaraGenerateMipsCS::FPermutationDomain PermutationVector;
+	PermutationVector.Set<FNiagaraGenerateMipsCS::FGenMipsSRGB>(bMipsSRGB);
+	TShaderMapRef<FNiagaraGenerateMipsCS> ComputeShader(GetGlobalShaderMap(GMaxRHIFeatureLevel), PermutationVector);
+	FRHIComputeShader* ShaderRHI = ComputeShader.GetComputeShader();
+
+	const int32 NumMips = TextureRHI->GetNumMips();
+	for ( int32 iDstMip=1; iDstMip < NumMips; ++iDstMip)
+	{
+		const int32 iSrcMip = iDstMip - 1;
+		const FIntPoint SrcMipSize(FMath::Max<int32>(TextureRHI->GetSizeX() >> iSrcMip, 1), FMath::Max<int32>(TextureRHI->GetSizeY() >> iSrcMip, 1));
+		const FIntPoint DstMipSize(FMath::Max<int32>(TextureRHI->GetSizeX() >> iDstMip, 1), FMath::Max<int32>(TextureRHI->GetSizeY() >> iDstMip, 1));
+
+		FShaderResourceViewRHIRef MipInSRV = RHICreateShaderResourceView(TextureRHI, iSrcMip);
+		FUnorderedAccessViewRHIRef MipOutUAV = RHICreateUnorderedAccessView(TextureRHI, iDstMip);
+
+		FNiagaraGenerateMipsCS::FParameters PassParameters;
+		PassParameters.TexelSize	= FVector2D(1.0f / float(DstMipSize.X), 1.0f / float(DstMipSize.Y));
+		PassParameters.MipOutSize	= DstMipSize;
+		PassParameters.MipInSRV		= MipInSRV;
+		PassParameters.MipInSampler	= TStaticSamplerState<SF_Bilinear>::GetRHI();
+		PassParameters.MipOutUAV	= MipOutUAV;
+
+		const FIntVector NumThreadGroups = FComputeShaderUtils::GetGroupCount(DstMipSize, FComputeShaderUtils::kGolden2DGroupSize);
+
+		RHICmdList.Transition(FRHITransitionInfo(MipOutUAV, ERHIAccess::SRVMask, ERHIAccess::UAVCompute));
+
+		RHICmdList.SetComputeShader(ShaderRHI);
+		SetShaderParameters(RHICmdList, ComputeShader, ShaderRHI, PassParameters);
+		RHICmdList.DispatchComputeShader(NumThreadGroups.X, NumThreadGroups.Y, NumThreadGroups.Z);
+		UnsetShaderUAVs(RHICmdList, ComputeShader, ShaderRHI);
+
+		RHICmdList.Transition(FRHITransitionInfo(MipOutUAV, ERHIAccess::UAVCompute, ERHIAccess::SRVMask));
+	}
+}
diff --git a/Engine/Plugins/FX/Niagara/Source/NiagaraShader/Public/NiagaraGenerateMips.h b/Engine/Plugins/FX/Niagara/Source/NiagaraShader/Public/NiagaraGenerateMips.h
new file mode 100644
--- /dev/null
+++ b/Engine/Plugins/FX/Niagara/Source/NiagaraShader/Public/NiagaraGenerateMips.h
@@ -0,0 +1,10 @@
+// Copyright Epic Games, Inc. All Rights Reserved.
+
+#pragma once
+
+#include "RHI.h"
+
+namespace NiagaraGenerateMips
+{
+	NIAGARASHADER_API void GenerateMips(FRHICommandList& RHICmdList, FRHITexture2D* TextureRHI);
+}
