- Allow NiagaraMeshRenderer to pass in PreSkinnedBounds information so it can be used in the material graph

diff --git a/Engine/Plugins/FX/Niagara/Source/Niagara/Private/NiagaraComponent.cpp b/Engine/Plugins/FX/Niagara/Source/Niagara/Private/NiagaraComponent.cpp
--- a/Engine/Plugins/FX/Niagara/Source/Niagara/Private/NiagaraComponent.cpp
+++ b/Engine/Plugins/FX/Niagara/Source/Niagara/Private/NiagaraComponent.cpp
@@ -191,7 +191,13 @@ void FNiagaraSceneProxy::DestroyRenderState_Concurrent()
 FNiagaraSceneProxy::~FNiagaraSceneProxy()
 {
 	check(IsInRenderingThread());
-	UniformBufferNoVelocity.ReleaseResource();
+
+	for ( auto& CustomUB : CustomUniformBuffers)
+	{
+		CustomUB.Value.ReleaseResource();
+	}
+	CustomUniformBuffers.Empty();
+
 	delete RenderData;
 }
 
@@ -201,7 +207,11 @@ void FNiagaraSceneProxy::ReleaseRenderThreadResources()
 	{
 		RenderData->ReleaseRenderThreadResources();
 	}
-	UniformBufferNoVelocity.ReleaseResource();
+
+	for (auto& CustomUB : CustomUniformBuffers)
+	{
+		CustomUB.Value.ReleaseResource();
+	}
 }
 
 // FPrimitiveSceneProxy interface.
@@ -217,7 +227,11 @@ void FNiagaraSceneProxy::CreateRenderThreadResources()
 void FNiagaraSceneProxy::OnTransformChanged()
 {
 	LocalToWorldInverse = GetLocalToWorld().Inverse();
-	UniformBufferNoVelocity.ReleaseResource();
+
+	for (auto& CustomUB : CustomUniformBuffers)
+	{
+		CustomUB.Value.ReleaseResource();
+	}
 }
 
 FPrimitiveViewRelevance FNiagaraSceneProxy::GetViewRelevance(const FSceneView* View) const
@@ -244,51 +258,68 @@ FPrimitiveViewRelevance FNiagaraSceneProxy::GetViewRelevance(const FSceneView* V
 	return Relevance;
 }
 
-FRHIUniformBuffer* FNiagaraSceneProxy::GetUniformBufferNoVelocity() const
+FRHIUniformBuffer* FNiagaraSceneProxy::GetCustomUniformBuffer(bool bHasVelocity, const FBox& PreSkinnedBounds) const
 {
-	if ( !UniformBufferNoVelocity.IsInitialized() )
+	// Default UB we create for the primitive
+	if (bHasVelocity && !PreSkinnedBounds.IsValid)
+	{
+		return GetUniformBuffer();
+	}
+
+	// Custom UB we need to create
+	uint64 KeyHash = HashCombine(bHasVelocity, PreSkinnedBounds.IsValid);
+	if ( PreSkinnedBounds.IsValid )
+	{
+		KeyHash = HashCombine(KeyHash, PreSkinnedBounds.Min.X);
+		KeyHash = HashCombine(KeyHash, PreSkinnedBounds.Min.Y);
+		KeyHash = HashCombine(KeyHash, PreSkinnedBounds.Min.Z);
+		KeyHash = HashCombine(KeyHash, PreSkinnedBounds.Max.X);
+		KeyHash = HashCombine(KeyHash, PreSkinnedBounds.Max.Y);
+		KeyHash = HashCombine(KeyHash, PreSkinnedBounds.Max.Z);
+	}
+
+	TUniformBuffer<FPrimitiveUniformShaderParameters>& CustomUB = CustomUniformBuffers.FindOrAdd(KeyHash);
+	if (!CustomUB.IsInitialized())
 	{
 		bool bHasPrecomputedVolumetricLightmap;
 		FMatrix PreviousLocalToWorld;
 		int32 SingleCaptureIndex;
 		bool bOutputVelocity;
 		FPrimitiveSceneInfo* LocalPrimitiveSceneInfo = GetPrimitiveSceneInfo();
-		GetScene().GetPrimitiveUniformShaderParameters_RenderThread(
-			LocalPrimitiveSceneInfo,
-			bHasPrecomputedVolumetricLightmap,
-			PreviousLocalToWorld,
-			SingleCaptureIndex,
-			bOutputVelocity
-		);
-
-		UniformBufferNoVelocity.SetContents(
-			FPrimitiveUniformShaderParametersBuilder{}
-			.Defaults()
-				.LocalToWorld(GetLocalToWorld())
-				.PreviousLocalToWorld(PreviousLocalToWorld)
-				.ActorWorldPosition(GetActorPosition())
-				.WorldBounds(GetBounds())
-				.LocalBounds(GetLocalBounds())
-				.CustomPrimitiveData(GetCustomPrimitiveData())
-				.LightingChannelMask(GetLightingChannelMask())
-				.LightmapDataIndex(LocalPrimitiveSceneInfo ? LocalPrimitiveSceneInfo->GetLightmapDataOffset() : 0)
-				.LightmapUVIndex(GetLightMapCoordinateIndex())
-				.SingleCaptureIndex(SingleCaptureIndex)
-				.InstanceDataOffset(LocalPrimitiveSceneInfo ? LocalPrimitiveSceneInfo->GetInstanceDataOffset() : INDEX_NONE)
-				.NumInstanceDataEntries(LocalPrimitiveSceneInfo ? LocalPrimitiveSceneInfo->GetNumInstanceDataEntries() : 0)
-				.ReceivesDecals(ReceivesDecals())
-				.DrawsVelocity(DrawsVelocity())
-				.OutputVelocity(bOutputVelocity)
-				.CastContactShadow(CastsContactShadow())
-				.CastShadow(CastsDynamicShadow())
-				.HasCapsuleRepresentation(HasDynamicIndirectShadowCasterRepresentation())
-				.UseVolumetricLightmap(bHasPrecomputedVolumetricLightmap)
-				.UseSingleSampleShadowFromStationaryLights(UseSingleSampleShadowFromStationaryLights())
-			.Build()
-		);
-		UniformBufferNoVelocity.InitResource();
-	}
-	return UniformBufferNoVelocity.GetUniformBufferRHI();
+		GetScene().GetPrimitiveUniformShaderParameters_RenderThread(LocalPrimitiveSceneInfo, bHasPrecomputedVolumetricLightmap, PreviousLocalToWorld, SingleCaptureIndex, bOutputVelocity);
+
+		FPrimitiveUniformShaderParametersBuilder UBBuilder;
+		UBBuilder.Defaults();
+		UBBuilder.LocalToWorld(GetLocalToWorld());
+		UBBuilder.PreviousLocalToWorld(PreviousLocalToWorld);
+		UBBuilder.ActorWorldPosition(GetActorPosition());
+		UBBuilder.WorldBounds(GetBounds());
+		UBBuilder.LocalBounds(GetLocalBounds());
+		UBBuilder.CustomPrimitiveData(GetCustomPrimitiveData());
+		UBBuilder.LightingChannelMask(GetLightingChannelMask());
+		UBBuilder.LightmapDataIndex(LocalPrimitiveSceneInfo ? LocalPrimitiveSceneInfo->GetLightmapDataOffset() : 0);
+		UBBuilder.LightmapUVIndex(GetLightMapCoordinateIndex());
+		UBBuilder.SingleCaptureIndex(SingleCaptureIndex);
+		UBBuilder.InstanceDataOffset(LocalPrimitiveSceneInfo ? LocalPrimitiveSceneInfo->GetInstanceDataOffset() : INDEX_NONE);
+		UBBuilder.NumInstanceDataEntries(LocalPrimitiveSceneInfo ? LocalPrimitiveSceneInfo->GetNumInstanceDataEntries() : 0);
+		UBBuilder.ReceivesDecals(ReceivesDecals());
+		UBBuilder.DrawsVelocity(DrawsVelocity());
+		UBBuilder.OutputVelocity(bOutputVelocity);
+		UBBuilder.CastContactShadow(CastsContactShadow());
+		UBBuilder.CastShadow(CastsDynamicShadow());
+		UBBuilder.HasCapsuleRepresentation(HasDynamicIndirectShadowCasterRepresentation());
+		UBBuilder.UseVolumetricLightmap(bHasPrecomputedVolumetricLightmap);
+		UBBuilder.UseSingleSampleShadowFromStationaryLights(UseSingleSampleShadowFromStationaryLights());
+		if ( PreSkinnedBounds.IsValid )
+		{
+			UBBuilder.PreSkinnedLocalBounds(PreSkinnedBounds);
+		}
+
+		CustomUB.SetContents(UBBuilder.Build());
+		CustomUB.InitResource();
+	}
+
+	return CustomUB.GetUniformBufferRHI();
 }
 
 uint32 FNiagaraSceneProxy::GetMemoryFootprint() const
diff --git a/Engine/Plugins/FX/Niagara/Source/Niagara/Private/NiagaraRendererMeshes.cpp b/Engine/Plugins/FX/Niagara/Source/Niagara/Private/NiagaraRendererMeshes.cpp
--- a/Engine/Plugins/FX/Niagara/Source/Niagara/Private/NiagaraRendererMeshes.cpp
+++ b/Engine/Plugins/FX/Niagara/Source/Niagara/Private/NiagaraRendererMeshes.cpp
@@ -204,6 +204,7 @@ void FNiagaraRendererMeshes::Initialize(const UNiagaraRendererProperties* InProp
 			FBox LocalBounds = Mesh->GetExtendedBounds().GetBox();
 			LocalBounds.Min *= MeshProperties.Scale;
 			LocalBounds.Max *= MeshProperties.Scale;
+			MeshData.LocalBounds = LocalBounds;
 			MeshData.LocalCullingSphere.Center = LocalBounds.GetCenter();
 			MeshData.LocalCullingSphere.W = LocalBounds.GetExtent().Size();
 
@@ -802,6 +803,7 @@ void FNiagaraRendererMeshes::CreateMeshBatchForSection(
 	FVertexFactory& VertexFactory,
 	FMaterialRenderProxy& MaterialProxy,
 	const FNiagaraSceneProxy& SceneProxy,
+	const FMeshData& MeshData,
 	const FStaticMeshLODResources& LODModel,
 	const FStaticMeshSection& Section,
 	const FSceneView& View,
@@ -828,7 +830,7 @@ void FNiagaraRendererMeshes::CreateMeshBatchForSection(
 	Mesh.DepthPriorityGroup = (ESceneDepthPriorityGroup)SceneProxy.GetDepthPriorityGroup(&View);
 
 	FMeshBatchElement& BatchElement = Mesh.Elements[0];
-	BatchElement.PrimitiveUniformBuffer = IsMotionBlurEnabled() ? SceneProxy.GetUniformBuffer() : SceneProxy.GetUniformBufferNoVelocity();
+	BatchElement.PrimitiveUniformBuffer = SceneProxy.GetCustomUniformBuffer(IsMotionBlurEnabled(), MeshData.LocalBounds);
 	BatchElement.FirstIndex = 0;
 	BatchElement.MinVertexIndex = 0;
 	BatchElement.MaxVertexIndex = 0;
@@ -1060,7 +1062,7 @@ void FNiagaraRendererMeshes::GetDynamicMeshElements(const TArray<const FSceneVie
 					}
 
 					const uint32 GPUCountBufferOffset = bDoGPUCulling ? SortInfo.CulledGPUParticleCountOffset : SourceParticleData->GetGPUInstanceCountBufferOffset();
-					CreateMeshBatchForSection(Collector, VertexFactory, *MaterialProxy, *SceneProxy, LODModel, Section, *View, ViewIndex, NumInstances,
+					CreateMeshBatchForSection(Collector, VertexFactory, *MaterialProxy, *SceneProxy, MeshData, LODModel, Section, *View, ViewIndex, NumInstances,
 						GPUCountBufferOffset, bIsWireframe, bIsInstancedStereo, bDoGPUCulling);
 				}
 			}
@@ -1229,7 +1231,7 @@ void FNiagaraRendererMeshes::GetDynamicRayTracingInstances(FRayTracingMaterialGa
 			MeshBatch.LCI = NULL;
 			MeshBatch.ReverseCulling = SceneProxy->IsLocalToWorldDeterminantNegative();
 
-			MeshBatchElement.PrimitiveUniformBuffer = IsMotionBlurEnabled() ? SceneProxy->GetUniformBuffer() : SceneProxy->GetUniformBufferNoVelocity();
+			MeshBatchElement.PrimitiveUniformBuffer = SceneProxy->GetCustomUniformBuffer(IsMotionBlurEnabled(), MeshData.LocalBounds);
 			MeshBatchElement.FirstIndex = 0;
 			MeshBatchElement.MinVertexIndex = 0;
 			MeshBatchElement.MaxVertexIndex = 0;
diff --git a/Engine/Plugins/FX/Niagara/Source/Niagara/Private/NiagaraRendererRibbons.cpp b/Engine/Plugins/FX/Niagara/Source/Niagara/Private/NiagaraRendererRibbons.cpp
--- a/Engine/Plugins/FX/Niagara/Source/Niagara/Private/NiagaraRendererRibbons.cpp
+++ b/Engine/Plugins/FX/Niagara/Source/Niagara/Private/NiagaraRendererRibbons.cpp
@@ -1167,7 +1167,7 @@ void FNiagaraRendererRibbons::SetupMeshBatchAndCollectorResourceForView(
 	MeshElement.NumInstances = 1;
 	MeshElement.MinVertexIndex = 0;
 	MeshElement.MaxVertexIndex = 0;
-	MeshElement.PrimitiveUniformBuffer = SceneProxy->GetUniformBufferNoVelocity();	// Note: Ribbons don't generate accurate velocities so disabling
+	MeshElement.PrimitiveUniformBuffer = SceneProxy->GetCustomUniformBuffer(false);	// Note: Ribbons don't generate accurate velocities so disabling
 }
 
 FNiagaraRendererRibbons::FCPUSimParticleDataAllocation FNiagaraRendererRibbons::AllocateParticleDataIfCPUSim(FNiagaraDynamicDataRibbon* DynamicDataRibbon, FGlobalDynamicReadBuffer& DynamicReadBuffer) const
diff --git a/Engine/Plugins/FX/Niagara/Source/Niagara/Private/NiagaraRendererSprites.cpp b/Engine/Plugins/FX/Niagara/Source/Niagara/Private/NiagaraRendererSprites.cpp
--- a/Engine/Plugins/FX/Niagara/Source/Niagara/Private/NiagaraRendererSprites.cpp
+++ b/Engine/Plugins/FX/Niagara/Source/Niagara/Private/NiagaraRendererSprites.cpp
@@ -820,7 +820,7 @@ void FNiagaraRendererSprites::CreateMeshBatchForView(
 	MeshElement.NumInstances = FMath::Max(0, NumInstances);	//->VertexData.Num();
 	MeshElement.MinVertexIndex = 0;
 	MeshElement.MaxVertexIndex = 0;// MeshElement.NumInstances * 4 - 1;
-	MeshElement.PrimitiveUniformBuffer = IsMotionBlurEnabled() ? SceneProxy->GetUniformBuffer() : SceneProxy->GetUniformBufferNoVelocity();
+	MeshElement.PrimitiveUniformBuffer = SceneProxy->GetCustomUniformBuffer(IsMotionBlurEnabled());
 	if (IndirectDraw.IsValid())
 	{
 		MeshElement.IndirectArgsBuffer = IndirectDraw.Buffer;
diff --git a/Engine/Plugins/FX/Niagara/Source/Niagara/Public/NiagaraComponent.h b/Engine/Plugins/FX/Niagara/Source/Niagara/Public/NiagaraComponent.h
--- a/Engine/Plugins/FX/Niagara/Source/Niagara/Public/NiagaraComponent.h
+++ b/Engine/Plugins/FX/Niagara/Source/Niagara/Public/NiagaraComponent.h
@@ -728,7 +728,7 @@ public:
 
 	FORCEINLINE const FMatrix& GetLocalToWorldInverse() const { return LocalToWorldInverse; }
 
-	FRHIUniformBuffer* GetUniformBufferNoVelocity() const;
+	FRHIUniformBuffer* GetCustomUniformBuffer(bool bHasVelocity, const FBox& PreSkinnedBounds = FBox(ForceInitToZero)) const;
 
 	virtual FPrimitiveViewRelevance GetViewRelevance(const FSceneView* View) const override;
 
@@ -759,8 +759,8 @@ private:
 	uint32 GetAllocatedSize() const;
 
 private:
-	/** Uniform Buffer with Velocity writes disabled.  Mutable as it is updated during GetDynamicMeshElements is required. */
-	mutable TUniformBuffer<FPrimitiveUniformShaderParameters> UniformBufferNoVelocity;
+	/** Custom Uniform Buffers, allows us to have renderer specific data packed inside such as pre-skinned bounds. */
+	mutable TMap<uint64, TUniformBuffer<FPrimitiveUniformShaderParameters>> CustomUniformBuffers;
 
 	/** The data required to render a single instance of a NiagaraSystem */
 	FNiagaraSystemRenderData* RenderData = nullptr;
diff --git a/Engine/Plugins/FX/Niagara/Source/Niagara/Public/NiagaraRendererMeshes.h b/Engine/Plugins/FX/Niagara/Source/Niagara/Public/NiagaraRendererMeshes.h
--- a/Engine/Plugins/FX/Niagara/Source/Niagara/Public/NiagaraRendererMeshes.h
+++ b/Engine/Plugins/FX/Niagara/Source/Niagara/Public/NiagaraRendererMeshes.h
@@ -55,6 +55,7 @@ protected:
 		FVector PivotOffset = FVector(ForceInitToZero);
 		ENiagaraMeshPivotOffsetSpace PivotOffsetSpace = ENiagaraMeshPivotOffsetSpace::Mesh;
 		FVector Scale = FVector(1.0f, 1.0f, 1.0f);
+		FBox LocalBounds = FBox(ForceInitToZero);
 		FSphere LocalCullingSphere = FSphere(ForceInitToZero);
 		TArray<uint32, TInlineAllocator<4>> MaterialRemapTable;
 	};
@@ -130,6 +131,7 @@ protected:
 		FVertexFactory& VertexFactory,
 		FMaterialRenderProxy& MaterialProxy,
 		const FNiagaraSceneProxy& SceneProxy,
+		const FMeshData& MeshData,
 		const FStaticMeshLODResources& LODModel,
 		const FStaticMeshSection& Section,
 		const FSceneView& View,
