Merging using //UE5/Main_to_//UE5/Release-Engine-Staging
Niagara - Fix code accidentally removed in last check-in.

diff --git a/Engine/Plugins/FX/Niagara/Source/Niagara/Private/NiagaraSystem.cpp b/Engine/Plugins/FX/Niagara/Source/Niagara/Private/NiagaraSystem.cpp
--- a/Engine/Plugins/FX/Niagara/Source/Niagara/Private/NiagaraSystem.cpp
+++ b/Engine/Plugins/FX/Niagara/Source/Niagara/Private/NiagaraSystem.cpp
@@ -777,6 +777,42 @@ void UNiagaraSystem::PostLoad()
 	ExposedParameters.RecreateRedirections();
 #endif
 
+	for (FNiagaraEmitterHandle& EmitterHandle : EmitterHandles)
+	{
+#if WITH_EDITORONLY_DATA
+		EmitterHandle.ConditionalPostLoad(NiagaraVer);
+#else
+		if (UNiagaraEmitter* NiagaraEmitter = EmitterHandle.GetInstance())
+		{
+			NiagaraEmitter->ConditionalPostLoad();
+		}
+#endif
+	}
+
+#if WITH_EDITORONLY_DATA
+	if (EditorData == nullptr)
+	{
+		INiagaraModule& NiagaraModule = FModuleManager::GetModuleChecked<INiagaraModule>("Niagara");
+		EditorData = NiagaraModule.GetEditorOnlyDataUtilities().CreateDefaultEditorData(this);
+	}
+	else
+	{
+		EditorData->PostLoadFromOwner(this);
+	}
+
+	if (EditorParameters == nullptr)
+	{
+		INiagaraModule& NiagaraModule = FModuleManager::GetModuleChecked<INiagaraModule>("Niagara");
+		EditorParameters = NiagaraModule.GetEditorOnlyDataUtilities().CreateDefaultEditorParameters(this);
+	}
+
+	// see the equivalent in NiagaraEmitter for details
+	if (bIsTemplateAsset_DEPRECATED)
+	{
+		TemplateSpecification = bIsTemplateAsset_DEPRECATED ? ENiagaraScriptTemplateSpecification::Template : ENiagaraScriptTemplateSpecification::None;
+	}
+#endif // WITH_EDITORONLY_DATA
+
 #if !WITH_EDITOR
 	// When running without the editor in a cooked build we run the update immediately in post load since
 	// there will be no merging or compiling which makes it safe to do so.
