Fixes for no unity build errors
#rnx

diff --git a/Engine/Plugins/FX/Niagara/Source/NiagaraShader/Private/NiagaraGenerateMips.cpp b/Engine/Plugins/FX/Niagara/Source/NiagaraShader/Private/NiagaraGenerateMips.cpp
--- a/Engine/Plugins/FX/Niagara/Source/NiagaraShader/Private/NiagaraGenerateMips.cpp
+++ b/Engine/Plugins/FX/Niagara/Source/NiagaraShader/Private/NiagaraGenerateMips.cpp
@@ -3,6 +3,7 @@
 #include "NiagaraGenerateMips.h"
 #include "GlobalShader.h"
 #include "RenderGraphUtils.h"
+#include "RHIStaticStates.h"
 
 class FNiagaraGenerateMipsCS : public FGlobalShader
 {
