Use slate common FOnDrop delegate in SDropTarget and deprecate the one off delegate version.

diff --git a/Engine/Plugins/FX/Niagara/Source/NiagaraEditor/Private/Widgets/SNiagaraGraphParameterMapGetNode.cpp b/Engine/Plugins/FX/Niagara/Source/NiagaraEditor/Private/Widgets/SNiagaraGraphParameterMapGetNode.cpp
--- a/Engine/Plugins/FX/Niagara/Source/NiagaraEditor/Private/Widgets/SNiagaraGraphParameterMapGetNode.cpp
+++ b/Engine/Plugins/FX/Niagara/Source/NiagaraEditor/Private/Widgets/SNiagaraGraphParameterMapGetNode.cpp
@@ -59,7 +59,7 @@ TSharedRef<SWidget> SNiagaraGraphParameterMapGetNode::CreateNodeContentArea()
 {
 	// NODE CONTENT AREA
 	return 	SNew(SDropTarget)
-		.OnDrop(this, &SNiagaraGraphParameterMapGetNode::OnDroppedOnTarget)
+		.OnDropped(this, &SNiagaraGraphParameterMapGetNode::OnDroppedOnTarget)
 		.OnAllowDrop(this, &SNiagaraGraphParameterMapGetNode::OnAllowDrop)
 		.Content()
 		[
@@ -195,12 +195,16 @@ FReply SNiagaraGraphParameterMapGetNode::OnBorderMouseButtonDown(const FGeometry
 	return FReply::Unhandled();
 }
 
-FReply SNiagaraGraphParameterMapGetNode::OnDroppedOnTarget(TSharedPtr<FDragDropOperation> DropOperation)
+FReply SNiagaraGraphParameterMapGetNode::OnDroppedOnTarget(const FGeometry& InGeometry, const FDragDropEvent& InDragDropEvent)
 {
-	UNiagaraNodeParameterMapBase* MapNode = Cast<UNiagaraNodeParameterMapBase>(GraphNode);
-	if (MapNode != nullptr && MapNode->HandleDropOperation(DropOperation))
+	TSharedPtr<FDragDropOperation> DragDropOperation = InDragDropEvent.GetOperation();
+	if (DragDropOperation)
 	{
-		return FReply::Handled();
+		UNiagaraNodeParameterMapBase* MapNode = Cast<UNiagaraNodeParameterMapBase>(GraphNode);
+		if (MapNode != nullptr && MapNode->HandleDropOperation(DragDropOperation))
+		{
+			return FReply::Handled();
+		}
 	}
 	return FReply::Unhandled();
 }
diff --git a/Engine/Plugins/FX/Niagara/Source/NiagaraEditor/Private/Widgets/SNiagaraGraphParameterMapGetNode.h b/Engine/Plugins/FX/Niagara/Source/NiagaraEditor/Private/Widgets/SNiagaraGraphParameterMapGetNode.h
--- a/Engine/Plugins/FX/Niagara/Source/NiagaraEditor/Private/Widgets/SNiagaraGraphParameterMapGetNode.h
+++ b/Engine/Plugins/FX/Niagara/Source/NiagaraEditor/Private/Widgets/SNiagaraGraphParameterMapGetNode.h
@@ -18,7 +18,7 @@ public:
 protected:
 	FReply OnBorderMouseButtonDown(const FGeometry& MyGeometry, const FPointerEvent& MouseEvent, int32 InWhichPin);
 	const FSlateBrush* GetBackgroundBrush(TSharedPtr<SWidget> Border) const;
-	FReply OnDroppedOnTarget(TSharedPtr<FDragDropOperation> DropOperation);
+	FReply OnDroppedOnTarget(const FGeometry& InGeometry, const FDragDropEvent& InDragDropEvent);
 	bool OnAllowDrop(TSharedPtr<FDragDropOperation> DragDropOperation);
 	TSharedPtr<SVerticalBox> PinContainerRoot;
 
diff --git a/Engine/Plugins/FX/Niagara/Source/NiagaraEditor/Private/Widgets/SNiagaraGraphParameterMapSetNode.cpp b/Engine/Plugins/FX/Niagara/Source/NiagaraEditor/Private/Widgets/SNiagaraGraphParameterMapSetNode.cpp
--- a/Engine/Plugins/FX/Niagara/Source/NiagaraEditor/Private/Widgets/SNiagaraGraphParameterMapSetNode.cpp
+++ b/Engine/Plugins/FX/Niagara/Source/NiagaraEditor/Private/Widgets/SNiagaraGraphParameterMapSetNode.cpp
@@ -30,7 +30,7 @@ TSharedRef<SWidget> SNiagaraGraphParameterMapSetNode::CreateNodeContentArea()
 	TSharedRef<SWidget> ExistingContent = SNiagaraGraphNode::CreateNodeContentArea();
 	// NODE CONTENT AREA
 	return 	SNew(SDropTarget)
-		.OnDrop(this, &SNiagaraGraphParameterMapSetNode::OnDroppedOnTarget)
+		.OnDropped(this, &SNiagaraGraphParameterMapSetNode::OnDroppedOnTarget)
 		.OnAllowDrop(this, &SNiagaraGraphParameterMapSetNode::OnAllowDrop)
 		.Content()
 		[
@@ -45,12 +45,16 @@ TSharedRef<SWidget> SNiagaraGraphParameterMapSetNode::CreateNodeContentArea()
 		];
 }
 
-FReply SNiagaraGraphParameterMapSetNode::OnDroppedOnTarget(TSharedPtr<FDragDropOperation> DropOperation)
+FReply SNiagaraGraphParameterMapSetNode::OnDroppedOnTarget(const FGeometry& InGeometry, const FDragDropEvent& InDragDropEvent)
 {
-	UNiagaraNodeParameterMapBase* MapNode = Cast<UNiagaraNodeParameterMapBase>(GraphNode);
-	if (MapNode != nullptr && MapNode->HandleDropOperation(DropOperation))
+	TSharedPtr<FDragDropOperation> DragDropOperation = InDragDropEvent.GetOperation();
+	if (DragDropOperation)
 	{
-		return FReply::Handled();
+		UNiagaraNodeParameterMapBase* MapNode = Cast<UNiagaraNodeParameterMapBase>(GraphNode);
+		if (MapNode != nullptr && MapNode->HandleDropOperation(DragDropOperation))
+		{
+			return FReply::Handled();
+		}
 	}
 	return FReply::Unhandled();
 }
diff --git a/Engine/Plugins/FX/Niagara/Source/NiagaraEditor/Private/Widgets/SNiagaraGraphParameterMapSetNode.h b/Engine/Plugins/FX/Niagara/Source/NiagaraEditor/Private/Widgets/SNiagaraGraphParameterMapSetNode.h
--- a/Engine/Plugins/FX/Niagara/Source/NiagaraEditor/Private/Widgets/SNiagaraGraphParameterMapSetNode.h
+++ b/Engine/Plugins/FX/Niagara/Source/NiagaraEditor/Private/Widgets/SNiagaraGraphParameterMapSetNode.h
@@ -13,7 +13,7 @@ public:
 	virtual TSharedRef<SWidget> CreateNodeContentArea();
 
 protected:
-	FReply OnDroppedOnTarget(TSharedPtr<FDragDropOperation> DropOperation);
+	FReply OnDroppedOnTarget(const FGeometry& InGeometry, const FDragDropEvent& InDragDropEvent);
 	bool OnAllowDrop(TSharedPtr<FDragDropOperation> DragDropOperation);
 
 };
\ No newline at end of file
diff --git a/Engine/Plugins/FX/Niagara/Source/NiagaraEditor/Private/Widgets/SNiagaraParameterPanel.cpp b/Engine/Plugins/FX/Niagara/Source/NiagaraEditor/Private/Widgets/SNiagaraParameterPanel.cpp
--- a/Engine/Plugins/FX/Niagara/Source/NiagaraEditor/Private/Widgets/SNiagaraParameterPanel.cpp
+++ b/Engine/Plugins/FX/Niagara/Source/NiagaraEditor/Private/Widgets/SNiagaraParameterPanel.cpp
@@ -145,7 +145,7 @@ void SNiagaraParameterPanel::Construct(const FArguments& InArgs, const TSharedPt
 		[
 			// Drop target
 			SNew(SDropTarget)
-			.OnDrop(this, &SNiagaraParameterPanel::HandleDragDropOperation)
+			.OnDropped(this, &SNiagaraParameterPanel::HandleDragDropOperation)
 			.OnAllowDrop(this, &SNiagaraParameterPanel::GetCanHandleDragDropOperation)
 			.Content()
 			[
@@ -703,9 +703,15 @@ TSharedRef<SWidget> SNiagaraParameterPanel::GetViewOptionsMenu()
 	return MenuBuilder.MakeWidget();
 }
 
-FReply SNiagaraParameterPanel::HandleDragDropOperation(TSharedPtr<FDragDropOperation> DropOperation)
+FReply SNiagaraParameterPanel::HandleDragDropOperation(const FGeometry& InGeometry, const FDragDropEvent& InDragDropEvent)
 {
-	return ParameterPanelViewModel->HandleDragDropOperation(DropOperation);
+	TSharedPtr<FDragDropOperation> DragDropOperation = InDragDropEvent.GetOperation();
+	if (DragDropOperation)
+	{
+		return ParameterPanelViewModel->HandleDragDropOperation(DragDropOperation);
+	}
+
+	return FReply::Unhandled();
 }
 
 bool SNiagaraParameterPanel::GetCanHandleDragDropOperation(TSharedPtr<FDragDropOperation> DragDropOperation)
diff --git a/Engine/Plugins/FX/Niagara/Source/NiagaraEditor/Public/Widgets/SNiagaraParameterPanel.h b/Engine/Plugins/FX/Niagara/Source/NiagaraEditor/Public/Widgets/SNiagaraParameterPanel.h
--- a/Engine/Plugins/FX/Niagara/Source/NiagaraEditor/Public/Widgets/SNiagaraParameterPanel.h
+++ b/Engine/Plugins/FX/Niagara/Source/NiagaraEditor/Public/Widgets/SNiagaraParameterPanel.h
@@ -120,7 +120,7 @@ private:
 
 	TSharedRef<SWidget> GetViewOptionsMenu();
 
-	FReply HandleDragDropOperation(TSharedPtr<FDragDropOperation> DropOperation);
+	FReply HandleDragDropOperation(const FGeometry& InGeometry, const FDragDropEvent& InDragDropEvent);
 	
 	bool GetCanHandleDragDropOperation(TSharedPtr<FDragDropOperation> DragDropOperation);
 
diff --git a/Engine/Plugins/FX/Niagara/Source/NiagaraEditorWidgets/Private/Stack/SNiagaraStackFunctionInputValue.cpp b/Engine/Plugins/FX/Niagara/Source/NiagaraEditorWidgets/Private/Stack/SNiagaraStackFunctionInputValue.cpp
--- a/Engine/Plugins/FX/Niagara/Source/NiagaraEditorWidgets/Private/Stack/SNiagaraStackFunctionInputValue.cpp
+++ b/Engine/Plugins/FX/Niagara/Source/NiagaraEditorWidgets/Private/Stack/SNiagaraStackFunctionInputValue.cpp
@@ -59,7 +59,7 @@ void SNiagaraStackFunctionInputValue::Construct(const FArguments& InArgs, UNiaga
 	[
 		SNew(SDropTarget)
 		.OnAllowDrop(this, &SNiagaraStackFunctionInputValue::OnFunctionInputAllowDrop)
-		.OnDrop(this, &SNiagaraStackFunctionInputValue::OnFunctionInputDrop)
+		.OnDropped(this, &SNiagaraStackFunctionInputValue::OnFunctionInputDrop)
 		.HorizontalImage(FNiagaraEditorWidgetsStyle::Get().GetBrush("NiagaraEditor.Stack.DropTarget.BorderHorizontal"))
 		.VerticalImage(FNiagaraEditorWidgetsStyle::Get().GetBrush("NiagaraEditor.Stack.DropTarget.BorderVertical"))
 		.IsEnabled_UObject(FunctionInput, &UNiagaraStackEntry::GetOwnerIsEnabled)
@@ -740,11 +740,11 @@ FSlateColor SNiagaraStackFunctionInputValue::GetInputIconColor() const
 	return FNiagaraEditorWidgetsStyle::Get().GetColor(FNiagaraStackEditorWidgetsUtilities::GetIconColorNameForInputMode(FunctionInput->GetValueMode()));
 }
 
-FReply SNiagaraStackFunctionInputValue::OnFunctionInputDrop(TSharedPtr<FDragDropOperation> DragDropOperation)
+FReply SNiagaraStackFunctionInputValue::OnFunctionInputDrop(const FGeometry& InGeometry, const FDragDropEvent& InDragDropEvent)
 {
-	if (DragDropOperation->IsOfType<FNiagaraParameterDragOperation>())
+	TSharedPtr<FNiagaraParameterDragOperation> InputDragDropOperation = InDragDropEvent.GetOperationAs<FNiagaraParameterDragOperation>();
+	if (InputDragDropOperation)
 	{
-		TSharedPtr<FNiagaraParameterDragOperation> InputDragDropOperation = StaticCastSharedPtr<FNiagaraParameterDragOperation>(DragDropOperation);
 		TSharedPtr<FNiagaraParameterAction> Action = StaticCastSharedPtr<FNiagaraParameterAction>(InputDragDropOperation->GetSourceAction());
 		if (Action.IsValid())
 		{
diff --git a/Engine/Plugins/FX/Niagara/Source/NiagaraEditorWidgets/Private/Stack/SNiagaraStackFunctionInputValue.h b/Engine/Plugins/FX/Niagara/Source/NiagaraEditorWidgets/Private/Stack/SNiagaraStackFunctionInputValue.h
--- a/Engine/Plugins/FX/Niagara/Source/NiagaraEditorWidgets/Private/Stack/SNiagaraStackFunctionInputValue.h
+++ b/Engine/Plugins/FX/Niagara/Source/NiagaraEditorWidgets/Private/Stack/SNiagaraStackFunctionInputValue.h
@@ -133,7 +133,7 @@ private:
 
 	FSlateColor GetInputIconColor() const;
 
-	FReply OnFunctionInputDrop(TSharedPtr<FDragDropOperation> DragDropOperation);
+	FReply OnFunctionInputDrop(const FGeometry& InGeometry, const FDragDropEvent& InDragDropEvent);
 
 	bool OnFunctionInputAllowDrop(TSharedPtr<FDragDropOperation> DragDropOperation);
 
diff --git a/Engine/Plugins/FX/Niagara/Source/NiagaraEditorWidgets/Private/Stack/SNiagaraStackModuleItem.cpp b/Engine/Plugins/FX/Niagara/Source/NiagaraEditorWidgets/Private/Stack/SNiagaraStackModuleItem.cpp
--- a/Engine/Plugins/FX/Niagara/Source/NiagaraEditorWidgets/Private/Stack/SNiagaraStackModuleItem.cpp
+++ b/Engine/Plugins/FX/Niagara/Source/NiagaraEditorWidgets/Private/Stack/SNiagaraStackModuleItem.cpp
@@ -344,7 +344,7 @@ TSharedRef<SWidget> SNiagaraStackModuleItem::AddContainerForRowWidgets(TSharedRe
 
 	return SNew(SDropTarget)
 	.OnAllowDrop(this, &SNiagaraStackModuleItem::OnModuleItemAllowDrop)
-	.OnDrop(this, &SNiagaraStackModuleItem::OnModuleItemDrop)
+	.OnDropped(this, &SNiagaraStackModuleItem::OnModuleItemDrop)
 	.HorizontalImage(FNiagaraEditorWidgetsStyle::Get().GetBrush("NiagaraEditor.Stack.DropTarget.BorderHorizontal"))
 	.VerticalImage(FNiagaraEditorWidgetsStyle::Get().GetBrush("NiagaraEditor.Stack.DropTarget.BorderVertical"))
 	.Content()
@@ -517,11 +517,17 @@ FReply SNiagaraStackModuleItem::RefreshClicked()
 	return FReply::Handled();
 }
 
-FReply SNiagaraStackModuleItem::OnModuleItemDrop(TSharedPtr<FDragDropOperation> DragDropOperation)
+FReply SNiagaraStackModuleItem::OnModuleItemDrop(const FGeometry& InGeometry, const FDragDropEvent& InDragDropEvent)
 {
-	UNiagaraStackEntry::FDropRequest DropRequest(DragDropOperation.ToSharedRef(), EItemDropZone::OntoItem, UNiagaraStackEntry::EDragOptions::None, UNiagaraStackEntry::EDropOptions::None);
-	TOptional<UNiagaraStackEntry::FDropRequestResponse> DropResponse = ModuleItem->Drop(DropRequest);
-	return DropResponse.IsSet() && DropResponse->DropZone == EItemDropZone::OntoItem ? FReply::Handled() : FReply::Unhandled();
+	TSharedPtr<FDragDropOperation> DragDropOperation = InDragDropEvent.GetOperation();
+	if (DragDropOperation)
+	{
+		UNiagaraStackEntry::FDropRequest DropRequest(DragDropOperation.ToSharedRef(), EItemDropZone::OntoItem, UNiagaraStackEntry::EDragOptions::None, UNiagaraStackEntry::EDropOptions::None);
+		TOptional<UNiagaraStackEntry::FDropRequestResponse> DropResponse = ModuleItem->Drop(DropRequest);
+		return (DropResponse.IsSet() && (DropResponse->DropZone == EItemDropZone::OntoItem)) ? FReply::Handled() : FReply::Unhandled();
+	}
+
+	return FReply::Unhandled();
 }
 
 bool SNiagaraStackModuleItem::OnModuleItemAllowDrop(TSharedPtr<FDragDropOperation> DragDropOperation)
diff --git a/Engine/Plugins/FX/Niagara/Source/NiagaraEditorWidgets/Private/Stack/SNiagaraStackModuleItem.h b/Engine/Plugins/FX/Niagara/Source/NiagaraEditorWidgets/Private/Stack/SNiagaraStackModuleItem.h
--- a/Engine/Plugins/FX/Niagara/Source/NiagaraEditorWidgets/Private/Stack/SNiagaraStackModuleItem.h
+++ b/Engine/Plugins/FX/Niagara/Source/NiagaraEditorWidgets/Private/Stack/SNiagaraStackModuleItem.h
@@ -54,7 +54,7 @@ private:
 
 	FReply RefreshClicked();
 
-	FReply OnModuleItemDrop(TSharedPtr<class FDragDropOperation> DragDropOperation);
+	FReply OnModuleItemDrop(const FGeometry& InGeometry, const FDragDropEvent& InDragDropEvent);
 
 	FSlateColor GetVersionSelectorColor() const;
 
diff --git a/Engine/Plugins/Media/MediaCompositing/Source/MediaCompositingEditor/Private/Sequencer/MediaSequenceRecorderExtender.cpp b/Engine/Plugins/Media/MediaCompositing/Source/MediaCompositingEditor/Private/Sequencer/MediaSequenceRecorderExtender.cpp
--- a/Engine/Plugins/Media/MediaCompositing/Source/MediaCompositingEditor/Private/Sequencer/MediaSequenceRecorderExtender.cpp
+++ b/Engine/Plugins/Media/MediaCompositing/Source/MediaCompositingEditor/Private/Sequencer/MediaSequenceRecorderExtender.cpp
@@ -197,7 +197,7 @@ TSharedPtr<SWidget> FMediaSequenceRecorderExtender::MakeListWidget(TSharedPtr<SL
 		[
 			SNew(SDropTarget)
 			.OnAllowDrop(this, &FMediaSequenceRecorderExtender::OnRecordingMediaPlayerListAllowDrop)
-			.OnDrop(this, &FMediaSequenceRecorderExtender::OnRecordingMediaPlayerListDrop)
+			.OnDropped(this, &FMediaSequenceRecorderExtender::OnRecordingMediaPlayerListDrop)
 			.Content()
 			[
 				SAssignNew(MediaPlayerListView, SListView<USequenceRecordingBase*>)
@@ -347,12 +347,11 @@ bool FMediaSequenceRecorderExtender::OnRecordingMediaPlayerListAllowDrop(TShared
 }
 
 
- FReply FMediaSequenceRecorderExtender::OnRecordingMediaPlayerListDrop(TSharedPtr<FDragDropOperation> DragDropOperation)
+ FReply FMediaSequenceRecorderExtender::OnRecordingMediaPlayerListDrop(const FGeometry& InGeometry, const FDragDropEvent& InDragDropEvent)
 {
-	if (DragDropOperation->IsOfType<FAssetDragDropOp>())
+	TSharedPtr<FAssetDragDropOp> AssetDragDropOperation = InDragDropEvent.GetOperationAs<FAssetDragDropOp>();
+	if (AssetDragDropOperation)
 	{
-		TSharedPtr<FAssetDragDropOp> AssetDragDropOperation = StaticCastSharedPtr<FAssetDragDropOp>(DragDropOperation);
-
 		for (const FAssetData& AssetData : AssetDragDropOperation->GetAssets())
 		{
 			if (AssetData.IsValid() && AssetData.GetClass()->IsChildOf<UMediaPlayer>())
diff --git a/Engine/Plugins/Media/MediaCompositing/Source/MediaCompositingEditor/Private/Sequencer/MediaSequenceRecorderExtender.h b/Engine/Plugins/Media/MediaCompositing/Source/MediaCompositingEditor/Private/Sequencer/MediaSequenceRecorderExtender.h
--- a/Engine/Plugins/Media/MediaCompositing/Source/MediaCompositingEditor/Private/Sequencer/MediaSequenceRecorderExtender.h
+++ b/Engine/Plugins/Media/MediaCompositing/Source/MediaCompositingEditor/Private/Sequencer/MediaSequenceRecorderExtender.h
@@ -58,7 +58,7 @@ private:
 	void OnMediaPlayerListSelectionChanged(USequenceRecordingBase* Recording, ESelectInfo::Type SelectionType);
 	EVisibility GetRecordMediaPlayerVisible() const;
 	bool OnRecordingMediaPlayerListAllowDrop(TSharedPtr<FDragDropOperation> DragDropOperation);
-	FReply OnRecordingMediaPlayerListDrop(TSharedPtr<FDragDropOperation> DragDropOperation);
+	FReply OnRecordingMediaPlayerListDrop(const FGeometry& InGeometry, const FDragDropEvent& InDragDropEvent);
 
 private:
 	bool bInsideSelectionChanged;
diff --git a/Engine/Plugins/MovieScene/MovieRenderPipeline/Source/MovieRenderPipelineEditor/Private/Widgets/SMoviePipelineQueueEditor.cpp b/Engine/Plugins/MovieScene/MovieRenderPipeline/Source/MovieRenderPipelineEditor/Private/Widgets/SMoviePipelineQueueEditor.cpp
--- a/Engine/Plugins/MovieScene/MovieRenderPipeline/Source/MovieRenderPipelineEditor/Private/Widgets/SMoviePipelineQueueEditor.cpp
+++ b/Engine/Plugins/MovieScene/MovieRenderPipeline/Source/MovieRenderPipelineEditor/Private/Widgets/SMoviePipelineQueueEditor.cpp
@@ -889,7 +889,7 @@ void SMoviePipelineQueueEditor::Construct(const FArguments& InArgs)
 	ChildSlot
 	[
 		SNew(SDropTarget)
-		.OnDrop(this, &SMoviePipelineQueueEditor::OnDragDropTarget)
+		.OnDropped(this, &SMoviePipelineQueueEditor::OnDragDropTarget)
 		.OnAllowDrop(this, &SMoviePipelineQueueEditor::CanDragDropTarget)
 		.OnIsRecognized(this, &SMoviePipelineQueueEditor::CanDragDropTarget)
 		[
@@ -1159,22 +1159,18 @@ void SMoviePipelineQueueEditor::OnGetChildren(TSharedPtr<IMoviePipelineQueueTree
 	}
 }
 
-FReply SMoviePipelineQueueEditor::OnDragDropTarget(TSharedPtr<FDragDropOperation> InOperation)
+FReply SMoviePipelineQueueEditor::OnDragDropTarget(const FGeometry& InGeometry, const FDragDropEvent& InDragDropEvent)
 {
-	if (InOperation)
+	if (TSharedPtr<FAssetDragDropOp> AssetDragDrop = InDragDropEvent.GetOperationAs<FAssetDragDropOp>())
 	{
-		if (InOperation->IsOfType<FAssetDragDropOp>())
-		{
-			TSharedPtr<FAssetDragDropOp> AssetDragDrop = StaticCastSharedPtr<FAssetDragDropOp>(InOperation);
-			FScopedTransaction Transaction(FText::Format(LOCTEXT("CreateJob_Transaction", "Add {0}|plural(one=Job, other=Jobs)"), AssetDragDrop->GetAssets().Num()));
+		FScopedTransaction Transaction(FText::Format(LOCTEXT("CreateJob_Transaction", "Add {0}|plural(one=Job, other=Jobs)"), AssetDragDrop->GetAssets().Num()));
 
-			for (const FAssetData& Asset : AssetDragDrop->GetAssets())
-			{
-				OnCreateJobFromAsset(Asset);
-			}
-
-			return FReply::Handled();
+		for (const FAssetData& Asset : AssetDragDrop->GetAssets())
+		{
+			OnCreateJobFromAsset(Asset);
 		}
+
+		return FReply::Handled();
 	}
 
 	return FReply::Unhandled();
diff --git a/Engine/Plugins/MovieScene/MovieRenderPipeline/Source/MovieRenderPipelineEditor/Private/Widgets/SMoviePipelineQueueEditor.h b/Engine/Plugins/MovieScene/MovieRenderPipeline/Source/MovieRenderPipelineEditor/Private/Widgets/SMoviePipelineQueueEditor.h
--- a/Engine/Plugins/MovieScene/MovieRenderPipeline/Source/MovieRenderPipelineEditor/Private/Widgets/SMoviePipelineQueueEditor.h
+++ b/Engine/Plugins/MovieScene/MovieRenderPipeline/Source/MovieRenderPipelineEditor/Private/Widgets/SMoviePipelineQueueEditor.h
@@ -61,7 +61,7 @@ private:
 
 	void OnGetChildren(TSharedPtr<IMoviePipelineQueueTreeItem> Item, TArray<TSharedPtr<IMoviePipelineQueueTreeItem>>& OutChildItems);
 
-	FReply OnDragDropTarget(TSharedPtr<FDragDropOperation> InOperation);
+	FReply OnDragDropTarget(const FGeometry& InGeometry, const FDragDropEvent& InDragDropEvent);
 
 	bool CanDragDropTarget(TSharedPtr<FDragDropOperation> InOperation);
 
diff --git a/Engine/Plugins/Runtime/ActorLayerUtilities/Source/ActorLayerUtilitiesEditor/Private/ActorLayerPropertyTypeCustomization.cpp b/Engine/Plugins/Runtime/ActorLayerUtilities/Source/ActorLayerUtilitiesEditor/Private/ActorLayerPropertyTypeCustomization.cpp
--- a/Engine/Plugins/Runtime/ActorLayerUtilities/Source/ActorLayerUtilitiesEditor/Private/ActorLayerPropertyTypeCustomization.cpp
+++ b/Engine/Plugins/Runtime/ActorLayerUtilities/Source/ActorLayerUtilitiesEditor/Private/ActorLayerPropertyTypeCustomization.cpp
@@ -36,7 +36,7 @@ void FActorLayerPropertyTypeCustomization::CustomizeHeader(TSharedRef<IPropertyH
 	.MaxDesiredWidth(TOptional<float>())
 	[
 		SNew(SDropTarget)
-		.OnDrop(this, &FActorLayerPropertyTypeCustomization::OnDrop)
+		.OnDropped(this, &FActorLayerPropertyTypeCustomization::OnDrop)
 		.OnAllowDrop(this, &FActorLayerPropertyTypeCustomization::OnVerifyDrag)
 		.OnIsRecognized(this, &FActorLayerPropertyTypeCustomization::OnVerifyDrag)
 		[
@@ -205,11 +205,12 @@ void FActorLayerPropertyTypeCustomization::OpenLayerBrowser()
 	LevelEditorModule.GetLevelEditorTabManager()->TryInvokeTab(FTabId("LevelEditorLayerBrowser"));
 }
 
-FReply FActorLayerPropertyTypeCustomization::OnDrop(TSharedPtr<FDragDropOperation> InDragDrop)
+FReply FActorLayerPropertyTypeCustomization::OnDrop(const FGeometry& InGeometry, const FDragDropEvent& InDragDropEvent)
 {
-	if (InDragDrop.IsValid() && InDragDrop->IsOfType<FLayersDragDropOp>())
+	TSharedPtr<FLayersDragDropOp> LayersDragDropOp = InDragDropEvent.GetOperationAs<FLayersDragDropOp>();
+	if (LayersDragDropOp)
 	{
-		const TArray<FName>& LayerNames = StaticCastSharedPtr<FLayersDragDropOp>(InDragDrop)->Layers;
+		const TArray<FName>& LayerNames = LayersDragDropOp->Layers;
 		if (ensure(LayerNames.Num() == 1))
 		{
 			AssignLayer(LayerNames[0]);
diff --git a/Engine/Plugins/Runtime/ActorLayerUtilities/Source/ActorLayerUtilitiesEditor/Private/ActorLayerPropertyTypeCustomization.h b/Engine/Plugins/Runtime/ActorLayerUtilities/Source/ActorLayerUtilitiesEditor/Private/ActorLayerPropertyTypeCustomization.h
--- a/Engine/Plugins/Runtime/ActorLayerUtilities/Source/ActorLayerUtilitiesEditor/Private/ActorLayerPropertyTypeCustomization.h
+++ b/Engine/Plugins/Runtime/ActorLayerUtilities/Source/ActorLayerUtilitiesEditor/Private/ActorLayerPropertyTypeCustomization.h
@@ -30,7 +30,7 @@ private:
 
 	void OpenLayerBrowser();
 
-	FReply OnDrop(TSharedPtr<FDragDropOperation> InDragDrop);
+	FReply OnDrop(const FGeometry& InGeometry, const FDragDropEvent& InDragDropEvent);
 	bool OnVerifyDrag(TSharedPtr<FDragDropOperation> InDragDrop);
 
 	TSharedPtr<IPropertyHandle> PropertyHandle;
diff --git a/Engine/Plugins/VirtualProduction/RemoteControl/Source/RemoteControlUI/Private/UI/SRCPanelFieldGroup.cpp b/Engine/Plugins/VirtualProduction/RemoteControl/Source/RemoteControlUI/Private/UI/SRCPanelFieldGroup.cpp
--- a/Engine/Plugins/VirtualProduction/RemoteControl/Source/RemoteControlUI/Private/UI/SRCPanelFieldGroup.cpp
+++ b/Engine/Plugins/VirtualProduction/RemoteControl/Source/RemoteControlUI/Private/UI/SRCPanelFieldGroup.cpp
@@ -131,7 +131,7 @@ void SRCPanelGroup::Construct(const FArguments& InArgs, URemoteControlPreset* In
 			SNew(SDropTarget)
 			.VerticalImage(FRemoteControlPanelStyle::Get()->GetBrush("RemoteControlPanel.VerticalDash"))
 			.HorizontalImage(FRemoteControlPanelStyle::Get()->GetBrush("RemoteControlPanel.HorizontalDash"))
-			.OnDrop_Lambda([this](TSharedPtr<FDragDropOperation> DragDropOperation) { return OnFieldDropGroup(DragDropOperation, nullptr); })
+			.OnDropped_Lambda([this](const FGeometry& InGeometry, const FDragDropEvent& InDragDropEvent) { return OnFieldDropGroup(InDragDropEvent.GetOperation(), nullptr); })
 			.OnAllowDrop(this, &SRCPanelGroup::OnAllowDropFromOtherGroup)
 			.OnIsRecognized(this, &SRCPanelGroup::OnAllowDropFromOtherGroup)
 			[
diff --git a/Engine/Plugins/VirtualProduction/Takes/Source/TakeRecorder/Private/Widgets/STakeRecorderSources.cpp b/Engine/Plugins/VirtualProduction/Takes/Source/TakeRecorder/Private/Widgets/STakeRecorderSources.cpp
--- a/Engine/Plugins/VirtualProduction/Takes/Source/TakeRecorder/Private/Widgets/STakeRecorderSources.cpp
+++ b/Engine/Plugins/VirtualProduction/Takes/Source/TakeRecorder/Private/Widgets/STakeRecorderSources.cpp
@@ -400,7 +400,7 @@ void STakeRecorderSources::Construct(const FArguments& InArgs)
 	ChildSlot
 	[
 		SNew(SDropTarget)
-		.OnDrop(this, &STakeRecorderSources::OnDragDropTarget)
+		.OnDropped(this, &STakeRecorderSources::OnDragDropTarget)
 		.OnAllowDrop(this, &STakeRecorderSources::CanDragDropTarget)
 		.OnIsRecognized(this, &STakeRecorderSources::CanDragDropTarget)
 		[
@@ -597,16 +597,22 @@ void STakeRecorderSources::OnGetChildren(TSharedPtr<ITakeRecorderSourceTreeItem>
 	}
 }
 
-FReply STakeRecorderSources::OnDragDropTarget(TSharedPtr<FDragDropOperation> InOperation)
+FReply STakeRecorderSources::OnDragDropTarget(const FGeometry& InGeometry, const FDragDropEvent& InDragDropEvent)
 {
+	TSharedPtr<FDragDropOperation> DragDropOperation = InDragDropEvent.GetOperation();
+	if (!DragDropOperation)
+	{
+		return FReply::Unhandled();
+	}
+
 	UTakeRecorderSources* Sources = WeakSources.Get();
 	if (Sources)
 	{
 		for (ITakeRecorderDropHandler* Handler : ITakeRecorderDropHandler::GetDropHandlers())
 		{
-			if (Handler->CanHandleOperation(InOperation, Sources))
+			if (Handler->CanHandleOperation(DragDropOperation, Sources))
 			{
-				Handler->HandleOperation(InOperation, Sources);
+				Handler->HandleOperation(DragDropOperation, Sources);
 				return FReply::Handled();
 			}
 		}
diff --git a/Engine/Plugins/VirtualProduction/Takes/Source/TakeRecorder/Private/Widgets/STakeRecorderSources.h b/Engine/Plugins/VirtualProduction/Takes/Source/TakeRecorder/Private/Widgets/STakeRecorderSources.h
--- a/Engine/Plugins/VirtualProduction/Takes/Source/TakeRecorder/Private/Widgets/STakeRecorderSources.h
+++ b/Engine/Plugins/VirtualProduction/Takes/Source/TakeRecorder/Private/Widgets/STakeRecorderSources.h
@@ -52,7 +52,7 @@ private:
 
 	void OnGetChildren(TSharedPtr<ITakeRecorderSourceTreeItem> Item, TArray<TSharedPtr<ITakeRecorderSourceTreeItem>>& OutChildItems);
 
-	FReply OnDragDropTarget(TSharedPtr<FDragDropOperation> InOperation);
+	FReply OnDragDropTarget(const FGeometry& InGeometry, const FDragDropEvent& InDragDropEvent);
 
 	bool CanDragDropTarget(TSharedPtr<FDragDropOperation> InOperation);
 
diff --git a/Engine/Source/Editor/DataLayerEditor/Private/DataLayer/DataLayerPropertyTypeCustomization.cpp b/Engine/Source/Editor/DataLayerEditor/Private/DataLayer/DataLayerPropertyTypeCustomization.cpp
--- a/Engine/Source/Editor/DataLayerEditor/Private/DataLayer/DataLayerPropertyTypeCustomization.cpp
+++ b/Engine/Source/Editor/DataLayerEditor/Private/DataLayer/DataLayerPropertyTypeCustomization.cpp
@@ -36,7 +36,7 @@ void FDataLayerPropertyTypeCustomization::CustomizeHeader(TSharedRef<IPropertyHa
 	.MaxDesiredWidth(TOptional<float>())
 	[
 		SNew(SDropTarget)
-		.OnDrop(this, &FDataLayerPropertyTypeCustomization::OnDrop)
+		.OnDropped(this, &FDataLayerPropertyTypeCustomization::OnDrop)
 		.OnAllowDrop(this, &FDataLayerPropertyTypeCustomization::OnVerifyDrag)
 		.OnIsRecognized(this, &FDataLayerPropertyTypeCustomization::OnVerifyDrag)
 		[
@@ -148,9 +148,9 @@ void FDataLayerPropertyTypeCustomization::AssignDataLayer(const UDataLayer* InDa
 	}
 }
 
-FReply FDataLayerPropertyTypeCustomization::OnDrop(TSharedPtr<FDragDropOperation> InDragDrop)
+FReply FDataLayerPropertyTypeCustomization::OnDrop(const FGeometry& InGeometry, const FDragDropEvent& InDragDropEvent)
 {
-	TSharedPtr<const FDataLayerDragDropOp> DataLayerDragDropOp = GetDataLayerDragDropOp(InDragDrop);
+	TSharedPtr<const FDataLayerDragDropOp> DataLayerDragDropOp = GetDataLayerDragDropOp(InDragDropEvent.GetOperation());
 	if (DataLayerDragDropOp.IsValid())
 	{
 		const TArray<FName>& DataLayerLabels = DataLayerDragDropOp->DataLayerLabels;
diff --git a/Engine/Source/Editor/DataLayerEditor/Private/DataLayer/DataLayerPropertyTypeCustomization.h b/Engine/Source/Editor/DataLayerEditor/Private/DataLayer/DataLayerPropertyTypeCustomization.h
--- a/Engine/Source/Editor/DataLayerEditor/Private/DataLayer/DataLayerPropertyTypeCustomization.h
+++ b/Engine/Source/Editor/DataLayerEditor/Private/DataLayer/DataLayerPropertyTypeCustomization.h
@@ -28,7 +28,7 @@ private:
 
 	TSharedRef<SWidget> OnGetDataLayerMenu();
 	FReply OnSelectDataLayer();
-	FReply OnDrop(TSharedPtr<FDragDropOperation> InDragDrop);
+	FReply OnDrop(const FGeometry& InGeometry, const FDragDropEvent& InDragDropEvent);
 	bool OnVerifyDrag(TSharedPtr<FDragDropOperation> InDragDrop);
 	TSharedPtr<const FDataLayerDragDropOp> GetDataLayerDragDropOp(TSharedPtr<FDragDropOperation> InDragDrop);
 
diff --git a/Engine/Source/Editor/EditorWidgets/Private/SAssetDropTarget.cpp b/Engine/Source/Editor/EditorWidgets/Private/SAssetDropTarget.cpp
--- a/Engine/Source/Editor/EditorWidgets/Private/SAssetDropTarget.cpp
+++ b/Engine/Source/Editor/EditorWidgets/Private/SAssetDropTarget.cpp
@@ -17,16 +17,16 @@ void SAssetDropTarget::Construct(const FArguments& InArgs )
 
 	SDropTarget::Construct(
 		SDropTarget::FArguments()
-		.OnDrop(this, &SAssetDropTarget::OnDropped)
+		.OnDropped(this, &SAssetDropTarget::OnDropped)
 		[
 			InArgs._Content.Widget
 		]);
 }
 
-FReply SAssetDropTarget::OnDropped(TSharedPtr<FDragDropOperation> DragDropOperation)
+FReply SAssetDropTarget::OnDropped(const FGeometry& InGeometry, const FDragDropEvent& InDragDropEvent)
 {
 	bool bUnused;
-	UObject* Object = GetDroppedObject(DragDropOperation, bUnused);
+	UObject* Object = GetDroppedObject(InDragDropEvent.GetOperation(), bUnused);
 
 	if ( Object )
 	{
@@ -105,7 +105,12 @@ void SAssetDropTarget::OnDragLeave(const FDragDropEvent& DragDropEvent)
 UObject* SAssetDropTarget::GetDroppedObject(TSharedPtr<FDragDropOperation> DragDropOperation, bool& bOutRecognizedEvent) const
 {
 	bOutRecognizedEvent = false;
-	UObject* DroppedObject = NULL;
+	UObject* DroppedObject = nullptr;
+
+	if (!DragDropOperation)
+	{
+		return DroppedObject;
+	}
 
 	// Asset being dragged from content browser
 	if ( DragDropOperation->IsOfType<FAssetDragDropOp>() )
diff --git a/Engine/Source/Editor/EditorWidgets/Private/SDropTarget.cpp b/Engine/Source/Editor/EditorWidgets/Private/SDropTarget.cpp
--- a/Engine/Source/Editor/EditorWidgets/Private/SDropTarget.cpp
+++ b/Engine/Source/Editor/EditorWidgets/Private/SDropTarget.cpp
@@ -12,7 +12,7 @@
 
 void SDropTarget::Construct(const FArguments& InArgs)
 {
-	DroppedEvent = InArgs._OnDrop;
+	DroppedEvent = InArgs._OnDropped;
 	AllowDropEvent = InArgs._OnAllowDrop;
 	IsRecognizedEvent = InArgs._OnIsRecognized;
 
@@ -111,7 +111,7 @@ FReply SDropTarget::OnDrop(const FGeometry& MyGeometry, const FDragDropEvent& Dr
 	{
 		if ( DroppedEvent.IsBound() )
 		{
-			return DroppedEvent.Execute(DragDropEvent.GetOperation());
+			return DroppedEvent.Execute(MyGeometry, DragDropEvent);
 		}
 
 		return FReply::Handled();
diff --git a/Engine/Source/Editor/EditorWidgets/Public/SAssetDropTarget.h b/Engine/Source/Editor/EditorWidgets/Public/SAssetDropTarget.h
--- a/Engine/Source/Editor/EditorWidgets/Public/SAssetDropTarget.h
+++ b/Engine/Source/Editor/EditorWidgets/Public/SAssetDropTarget.h
@@ -37,7 +37,7 @@ public:
 	void Construct(const FArguments& InArgs );
 
 protected:
-	FReply OnDropped(TSharedPtr<FDragDropOperation> DragDropOperation);
+	FReply OnDropped(const FGeometry& InGeometry, const FDragDropEvent& InDragDropEvent);
 	virtual bool OnAllowDrop(TSharedPtr<FDragDropOperation> DragDropOperation) const override;
 	virtual bool OnIsRecognized(TSharedPtr<FDragDropOperation> DragDropOperation) const override;
 	virtual void OnDragLeave(const FDragDropEvent& DragDropEvent) override;
diff --git a/Engine/Source/Editor/EditorWidgets/Public/SDropTarget.h b/Engine/Source/Editor/EditorWidgets/Public/SDropTarget.h
--- a/Engine/Source/Editor/EditorWidgets/Public/SDropTarget.h
+++ b/Engine/Source/Editor/EditorWidgets/Public/SDropTarget.h
@@ -6,6 +6,7 @@
 #include "Styling/SlateColor.h"
 #include "Input/Reply.h"
 #include "Layout/Visibility.h"
+#include "Framework/SlateDelegates.h"
 #include "Widgets/DeclarativeSyntaxSupport.h"
 #include "Widgets/SCompoundWidget.h"
 #include "EditorStyleSet.h"
@@ -20,10 +21,8 @@ class FSlateWindowElementList;
 class EDITORWIDGETS_API SDropTarget : public SCompoundWidget
 {
 public:
-	/** Called when a valid asset is dropped */
-	DECLARE_DELEGATE_RetVal_OneParam(FReply, FOnDrop, TSharedPtr<FDragDropOperation>);
-
 	DECLARE_DELEGATE_RetVal_OneParam(bool, FVerifyDrag, TSharedPtr<FDragDropOperation>);
+	DECLARE_DELEGATE_RetVal_OneParam(FReply, FOnDropDeprecated, TSharedPtr<FDragDropOperation>);
 
 	SLATE_BEGIN_ARGS(SDropTarget)
 		: _ValidColor(FStyleColors::AccentBlue)
@@ -57,11 +56,25 @@ public:
 		/** The background image that is applied after the surface. */
 		SLATE_ATTRIBUTE(const FSlateBrush*, BackgroundImage)
 		/** Called when a valid asset is dropped */
-		SLATE_EVENT(FOnDrop, OnDrop)
+		SLATE_EVENT(FOnDrop, OnDropped)
 		/** Called to check if an asset is acceptable for dropping */
 		SLATE_EVENT(FVerifyDrag, OnAllowDrop)
 		/** Called to check if an asset is acceptable for dropping */
 		SLATE_EVENT(FVerifyDrag, OnIsRecognized)
+
+		FOnDrop ConvertOnDropFn(const FOnDropDeprecated& LegacyDelegate)
+		{
+			return FOnDrop::CreateLambda([LegacyDelegate](const FGeometry&, const FDragDropEvent& DragDropEvent)
+			{
+				if (LegacyDelegate.IsBound())
+				{
+					return LegacyDelegate.Execute(DragDropEvent.GetOperation());
+				}
+
+				return FReply::Unhandled();
+			});
+		}
+		SLATE_EVENT_DEPRECATED(5.0, "Use OnDropped instead.", FOnDropDeprecated, OnDrop, OnDropped, ConvertOnDropFn)
 	SLATE_END_ARGS()
 
 	void Construct(const FArguments& InArgs );
diff --git a/Engine/Source/Editor/MovieSceneTools/Private/MovieSceneObjectBindingIDCustomization.cpp b/Engine/Source/Editor/MovieSceneTools/Private/MovieSceneObjectBindingIDCustomization.cpp
--- a/Engine/Source/Editor/MovieSceneTools/Private/MovieSceneObjectBindingIDCustomization.cpp
+++ b/Engine/Source/Editor/MovieSceneTools/Private/MovieSceneObjectBindingIDCustomization.cpp
@@ -61,7 +61,7 @@ void FMovieSceneObjectBindingIDCustomization::CustomizeHeader(TSharedRef<IProper
 		+ SHorizontalBox::Slot()
 		[
 			SNew(SDropTarget)
-			.OnDrop(this, &FMovieSceneObjectBindingIDCustomization::OnDrop)
+			.OnDropped(this, &FMovieSceneObjectBindingIDCustomization::OnDrop)
 			.OnAllowDrop_Static(IsAcceptable)
 			.OnIsRecognized_Static(IsAcceptable)
 			[
@@ -88,9 +88,9 @@ void FMovieSceneObjectBindingIDCustomization::CustomizeHeader(TSharedRef<IProper
 	];
 }
 
-FReply FMovieSceneObjectBindingIDCustomization::OnDrop(TSharedPtr<FDragDropOperation> InOperation)
+FReply FMovieSceneObjectBindingIDCustomization::OnDrop(const FGeometry& InGeometry, const FDragDropEvent& InDragDropEvent)
 {
-	FSequencerObjectBindingDragDropOp* SequencerOp = InOperation->IsOfType<FSequencerObjectBindingDragDropOp>() ? static_cast<FSequencerObjectBindingDragDropOp*>(InOperation.Get()) : nullptr;
+	TSharedPtr<FSequencerObjectBindingDragDropOp> SequencerOp = InDragDropEvent.GetOperationAs<FSequencerObjectBindingDragDropOp>();
 	if (SequencerOp)
 	{
 		TArray<UE::MovieScene::FFixedObjectBindingID> Bindings = SequencerOp->GetDraggedBindings();
diff --git a/Engine/Source/Editor/MovieSceneTools/Private/Sections/LevelVisibilitySection.cpp b/Engine/Source/Editor/MovieSceneTools/Private/Sections/LevelVisibilitySection.cpp
--- a/Engine/Source/Editor/MovieSceneTools/Private/Sections/LevelVisibilitySection.cpp
+++ b/Engine/Source/Editor/MovieSceneTools/Private/Sections/LevelVisibilitySection.cpp
@@ -34,7 +34,7 @@ TSharedRef<SWidget> FLevelVisibilitySection::GenerateSectionWidget()
 	return
 		SNew( SDropTarget )
 		.OnAllowDrop( this, &FLevelVisibilitySection::OnAllowDrop )
-		.OnDrop( this, &FLevelVisibilitySection::OnDrop )
+		.OnDropped( this, &FLevelVisibilitySection::OnDrop )
 		.Content()
 		[
 			SNew( SBorder )
@@ -111,11 +111,11 @@ bool FLevelVisibilitySection::OnAllowDrop( TSharedPtr<FDragDropOperation> DragDr
 }
 
 
-FReply FLevelVisibilitySection::OnDrop( TSharedPtr<FDragDropOperation> DragDropOperation )
+FReply FLevelVisibilitySection::OnDrop( const FGeometry& InGeometry, const FDragDropEvent& InDragDropEvent )
 {
-	if ( DragDropOperation->IsOfType<FLevelDragDropOp>() )
+	TSharedPtr<FLevelDragDropOp> LevelDragDropOperation = InDragDropEvent.GetOperationAs<FLevelDragDropOp>();
+	if ( LevelDragDropOperation )
 	{
-		TSharedPtr<FLevelDragDropOp> LevelDragDropOperation = StaticCastSharedPtr<FLevelDragDropOp>( DragDropOperation );
 		if ( LevelDragDropOperation->StreamingLevelsToDrop.Num() > 0 )
 		{
 			FScopedTransaction Transaction(NSLOCTEXT("LevelVisibilitySection", "TransactionText", "Add Level(s) to Level Visibility Section"));
diff --git a/Engine/Source/Editor/MovieSceneTools/Private/Sections/LevelVisibilitySection.h b/Engine/Source/Editor/MovieSceneTools/Private/Sections/LevelVisibilitySection.h
--- a/Engine/Source/Editor/MovieSceneTools/Private/Sections/LevelVisibilitySection.h
+++ b/Engine/Source/Editor/MovieSceneTools/Private/Sections/LevelVisibilitySection.h
@@ -37,7 +37,7 @@ private:
 	FText GetVisibilityToolTip() const;
 
 	bool OnAllowDrop( TSharedPtr<FDragDropOperation> DragDropOperation );
-	FReply OnDrop( TSharedPtr<FDragDropOperation> DragDropOperation );
+	FReply OnDrop( const FGeometry& InGeometry, const FDragDropEvent& InDragDropEvent );
 
 private:
 
diff --git a/Engine/Source/Editor/MovieSceneTools/Public/MovieSceneObjectBindingIDCustomization.h b/Engine/Source/Editor/MovieSceneTools/Public/MovieSceneObjectBindingIDCustomization.h
--- a/Engine/Source/Editor/MovieSceneTools/Public/MovieSceneObjectBindingIDCustomization.h
+++ b/Engine/Source/Editor/MovieSceneTools/Public/MovieSceneObjectBindingIDCustomization.h
@@ -50,7 +50,7 @@ private:
 
 	virtual bool HasMultipleValues() const override;
 
-	FReply OnDrop(TSharedPtr<FDragDropOperation> InOperation);
+	FReply OnDrop(const FGeometry& InGeometry, const FDragDropEvent& InDragDropEvent);
 
 	TSharedPtr<IPropertyHandle> StructProperty;
 };
\ No newline at end of file
diff --git a/Engine/Source/Editor/PropertyEditor/Private/UserInterface/PropertyEditor/SPropertyEditorArray.h b/Engine/Source/Editor/PropertyEditor/Private/UserInterface/PropertyEditor/SPropertyEditorArray.h
--- a/Engine/Source/Editor/PropertyEditor/Private/UserInterface/PropertyEditor/SPropertyEditorArray.h
+++ b/Engine/Source/Editor/PropertyEditor/Private/UserInterface/PropertyEditor/SPropertyEditorArray.h
@@ -46,7 +46,7 @@ public:
 		.Padding( 0.0f, 0.0f, 2.0f, 0.0f )
 		[
 			SNew(SDropTarget)
-			.OnDrop(this, &SPropertyEditorArray::OnDragDropTarget)
+			.OnDropped(this, &SPropertyEditorArray::OnDragDropTarget)
 			.OnAllowDrop(this, &SPropertyEditorArray::WillAddValidElements)
 			.OnIsRecognized(this, &SPropertyEditorArray::IsValidAssetDropOp)
 			.Content()
@@ -92,7 +92,7 @@ private:
 		return PropertyEditor.IsValid() ? !PropertyEditor->IsEditConst() : true;
 	}
 
-	FReply OnDragDropTarget(TSharedPtr<FDragDropOperation> InOperation)
+	FReply OnDragDropTarget(const FGeometry& InGeometry, const FDragDropEvent& InDragDropEvent)
 	{
 		FObjectProperty* ObjectProperty = nullptr;
 		if (const FArrayProperty* ArrayProperty = CastField<FArrayProperty>(PropertyEditor->GetProperty()))
@@ -101,9 +101,10 @@ private:
 		}
 
 		// Only try to add entries if we are dropping on an asset array
-		if (ObjectProperty && InOperation->IsOfType<FAssetDragDropOp>())
+		TSharedPtr<FDragDropOperation> DragOperation = InDragDropEvent.GetOperation();
+		if (ObjectProperty && DragOperation && DragOperation->IsOfType<FAssetDragDropOp>())
 		{
-			TSharedPtr<FAssetDragDropOp> DragDropOp = StaticCastSharedPtr<FAssetDragDropOp>(InOperation);
+			TSharedPtr<FAssetDragDropOp> DragDropOp = StaticCastSharedPtr<FAssetDragDropOp>(DragOperation);
 			if (DragDropOp.IsValid())
 			{
 				for (FAssetData AssetData : DragDropOp->GetAssets())
diff --git a/Engine/Source/Editor/SequenceRecorder/Private/SSequenceRecorder.cpp b/Engine/Source/Editor/SequenceRecorder/Private/SSequenceRecorder.cpp
--- a/Engine/Source/Editor/SequenceRecorder/Private/SSequenceRecorder.cpp
+++ b/Engine/Source/Editor/SequenceRecorder/Private/SSequenceRecorder.cpp
@@ -449,7 +449,7 @@ void SSequenceRecorder::Construct(const FArguments& Args)
 						[
 							SNew( SDropTarget )
 							.OnAllowDrop( this, &SSequenceRecorder::OnRecordingActorListAllowDrop )
-							.OnDrop( this, &SSequenceRecorder::OnRecordingActorListDrop )
+							.OnDropped( this, &SSequenceRecorder::OnRecordingActorListDrop )
 							.Content()
 							[
 								SAssignNew(ActorListView, SListView<UActorRecording*>)
@@ -986,12 +986,11 @@ bool SSequenceRecorder::OnRecordingActorListAllowDrop( TSharedPtr<FDragDropOpera
 	return DragDropOperation->IsOfType<FActorDragDropOp>();
 }
 
-FReply SSequenceRecorder::OnRecordingActorListDrop( TSharedPtr<FDragDropOperation> DragDropOperation )
+FReply SSequenceRecorder::OnRecordingActorListDrop( const FGeometry& InGeometry, const FDragDropEvent& InDragDropEvent )
 {
-	if ( DragDropOperation->IsOfType<FActorDragDropOp>() )
+	TSharedPtr<FActorDragDropOp> ActorDragDropOperation = InDragDropEvent.GetOperationAs<FActorDragDropOp>();
+	if ( ActorDragDropOperation )
 	{
-		TSharedPtr<FActorDragDropOp> ActorDragDropOperation = StaticCastSharedPtr<FActorDragDropOp>( DragDropOperation );
-
 		for (auto Actor : ActorDragDropOperation->Actors)
 		{
 			if (Actor.IsValid())
diff --git a/Engine/Source/Editor/SequenceRecorder/Private/SSequenceRecorder.h b/Engine/Source/Editor/SequenceRecorder/Private/SSequenceRecorder.h
--- a/Engine/Source/Editor/SequenceRecorder/Private/SSequenceRecorder.h
+++ b/Engine/Source/Editor/SequenceRecorder/Private/SSequenceRecorder.h
@@ -105,7 +105,7 @@ private:
 
 	FText GetTargetSequenceName() const;
 
-	FReply OnRecordingActorListDrop( TSharedPtr<FDragDropOperation> DragDropOperation );
+	FReply OnRecordingActorListDrop( const FGeometry& InGeometry, const FDragDropEvent& InDragDropEvent );
 
 	bool OnRecordingActorListAllowDrop( TSharedPtr<FDragDropOperation> DragDropOperation );
 
