- Clear translucent data to render before we set actual data to render to fix ensure

diff --git a/Engine/Plugins/FX/Niagara/Source/Niagara/Private/NiagaraEmitterInstanceBatcher.cpp b/Engine/Plugins/FX/Niagara/Source/Niagara/Private/NiagaraEmitterInstanceBatcher.cpp
--- a/Engine/Plugins/FX/Niagara/Source/Niagara/Private/NiagaraEmitterInstanceBatcher.cpp
+++ b/Engine/Plugins/FX/Niagara/Source/Niagara/Private/NiagaraEmitterInstanceBatcher.cpp
@@ -1367,7 +1367,7 @@ void NiagaraEmitterInstanceBatcher::ExecuteTicks(FRHICommandList& RHICmdList, FR
 				check(FinalData);
 
 				CurrentData->SwapGPU(FinalData);
-
+				ComputeContext->SetTranslucentDataToRender(nullptr);
 				ComputeContext->SetDataToRender(CurrentData);
 #if WITH_MGPU
 				AddTemporalEffectBuffers(CurrentData);
