Merged CL 16589009: Fixed a bug where Niagara assets were marked as dirty unnecessarily

diff --git a/Engine/Plugins/FX/Niagara/Source/Niagara/Private/NiagaraSystem.cpp b/Engine/Plugins/FX/Niagara/Source/Niagara/Private/NiagaraSystem.cpp
--- a/Engine/Plugins/FX/Niagara/Source/Niagara/Private/NiagaraSystem.cpp
+++ b/Engine/Plugins/FX/Niagara/Source/Niagara/Private/NiagaraSystem.cpp
@@ -466,9 +466,12 @@ void UNiagaraSystem::UpdateSystemAfterLoad()
 
 		if (bSystemScriptsAreSynchronized == false || bEmitterScriptsAreSynchronized == false)
 		{
-			// Call modify here so that the system will resave the compile ids and script vm when running the resave
-			// commandlet.  In normal post load, it will be ignored.
-			Modify();
+			if (IsRunningCommandlet())
+			{
+				// Call modify here so that the system will resave the compile ids and script vm when running the resave
+				// commandlet. We don't need it for normal post-loading.
+				Modify();
+			}
 			RequestCompile(false);
 		}
 
