Update emitter notifications so that they are applied if loaded or created (as is the case with an emitter being edited as a transient copy will be created for the system editor).

diff --git a/Engine/Plugins/FX/Niagara/Source/Niagara/Classes/NiagaraEmitter.h b/Engine/Plugins/FX/Niagara/Source/Niagara/Classes/NiagaraEmitter.h
--- a/Engine/Plugins/FX/Niagara/Source/Niagara/Classes/NiagaraEmitter.h
+++ b/Engine/Plugins/FX/Niagara/Source/Niagara/Classes/NiagaraEmitter.h
@@ -247,6 +247,8 @@ public:
 	/** Helper method for when a rename has been detected within the graph. Covers resetting the internal renderer bindings.*/
 	NIAGARA_API void HandleVariableRemoved(const FNiagaraVariable& InOldVariable, bool bUpdateContexts);
 
+	/** Helper method for binding the notifications needed for proper editor integration. */
+	NIAGARA_API void BindNotifications();
 #endif
 	virtual bool NeedsLoadForTargetPlatform(const ITargetPlatform* TargetPlatform) const override;
 	void Serialize(FArchive& Ar)override;
diff --git a/Engine/Plugins/FX/Niagara/Source/Niagara/Private/NiagaraEmitter.cpp b/Engine/Plugins/FX/Niagara/Source/Niagara/Private/NiagaraEmitter.cpp
--- a/Engine/Plugins/FX/Niagara/Source/Niagara/Private/NiagaraEmitter.cpp
+++ b/Engine/Plugins/FX/Niagara/Source/Niagara/Private/NiagaraEmitter.cpp
@@ -754,6 +754,8 @@ UNiagaraEmitter* UNiagaraEmitter::CreateWithParentAndOwner(UNiagaraEmitter& InPa
 	NewEmitter->ScratchPadScripts.Empty();
 	NewEmitter->SetUniqueEmitterName(InName.GetPlainNameString());
 	NewEmitter->GraphSource->MarkNotSynchronized(InitialNotSynchronizedReason);
+	NewEmitter->BindNotifications();
+
 	return NewEmitter;
 }
 
@@ -770,6 +772,7 @@ UNiagaraEmitter* UNiagaraEmitter::CreateAsDuplicate(const UNiagaraEmitter& InEmi
 	}
 	NewEmitter->SetUniqueEmitterName(InDuplicateName.GetPlainNameString());
 	NewEmitter->GraphSource->MarkNotSynchronized(InitialNotSynchronizedReason);
+	NewEmitter->BindNotifications();
 
 	return NewEmitter;
 }
@@ -1314,47 +1317,7 @@ void UNiagaraEmitter::UpdateEmitterAfterLoad()
 			UpdateChangeId(GenerateNewChangeIdReason);
 		}
 
-		GraphSource->OnChanged().AddUObject(this, &UNiagaraEmitter::GraphSourceChanged);
-
-		EmitterSpawnScriptProps.Script->RapidIterationParameters.AddOnChangedHandler(
-			FNiagaraParameterStore::FOnChanged::FDelegate::CreateUObject(this, &UNiagaraEmitter::ScriptRapidIterationParameterChanged));
-		EmitterUpdateScriptProps.Script->RapidIterationParameters.AddOnChangedHandler(
-			FNiagaraParameterStore::FOnChanged::FDelegate::CreateUObject(this, &UNiagaraEmitter::ScriptRapidIterationParameterChanged));
-
-		if (SpawnScriptProps.Script)
-		{
-			SpawnScriptProps.Script->RapidIterationParameters.AddOnChangedHandler(
-				FNiagaraParameterStore::FOnChanged::FDelegate::CreateUObject(this, &UNiagaraEmitter::ScriptRapidIterationParameterChanged));
-		}
-	
-		if (UpdateScriptProps.Script)
-		{
-			UpdateScriptProps.Script->RapidIterationParameters.AddOnChangedHandler(
-				FNiagaraParameterStore::FOnChanged::FDelegate::CreateUObject(this, &UNiagaraEmitter::ScriptRapidIterationParameterChanged));
-		}
-
-		for (FNiagaraEventScriptProperties& EventScriptProperties : EventHandlerScriptProps)
-		{
-			EventScriptProperties.Script->RapidIterationParameters.AddOnChangedHandler(
-				FNiagaraParameterStore::FOnChanged::FDelegate::CreateUObject(this, &UNiagaraEmitter::ScriptRapidIterationParameterChanged));
-		}
-
-		for (UNiagaraSimulationStageBase* SimulationStage : SimulationStages)
-		{
-			SimulationStage->OnChanged().AddUObject(this, &UNiagaraEmitter::SimulationStageChanged);
-			SimulationStage->Script->RapidIterationParameters.AddOnChangedHandler(
-				FNiagaraParameterStore::FOnChanged::FDelegate::CreateUObject(this, &UNiagaraEmitter::ScriptRapidIterationParameterChanged));
-		}
-
-		for (UNiagaraRendererProperties* Renderer : RendererProperties)
-		{
-			Renderer->OnChanged().AddUObject(this, &UNiagaraEmitter::RendererChanged);
-		}
-
-		if (EditorData != nullptr)
-		{
-			EditorData->OnPersistentDataChanged().AddUObject(this, &UNiagaraEmitter::PersistentEditorDataChanged);
-		}
+		BindNotifications();
 	}
 #endif
 
@@ -1537,6 +1500,75 @@ void UNiagaraEmitter::OnPostCompile()
 
 	OnEmitterVMCompiled().Broadcast(this);
 }
+
+void UNiagaraEmitter::BindNotifications()
+{
+	if (GraphSource)
+	{
+		GraphSource->OnChanged().AddUObject(this, &UNiagaraEmitter::GraphSourceChanged);
+	}
+
+	if (EmitterSpawnScriptProps.Script)
+	{
+		EmitterSpawnScriptProps.Script->RapidIterationParameters.AddOnChangedHandler(
+			FNiagaraParameterStore::FOnChanged::FDelegate::CreateUObject(this, &UNiagaraEmitter::ScriptRapidIterationParameterChanged));
+	}
+
+	if (EmitterUpdateScriptProps.Script)
+	{
+		EmitterUpdateScriptProps.Script->RapidIterationParameters.AddOnChangedHandler(
+			FNiagaraParameterStore::FOnChanged::FDelegate::CreateUObject(this, &UNiagaraEmitter::ScriptRapidIterationParameterChanged));
+	}
+
+	if (SpawnScriptProps.Script)
+	{
+		SpawnScriptProps.Script->RapidIterationParameters.AddOnChangedHandler(
+			FNiagaraParameterStore::FOnChanged::FDelegate::CreateUObject(this, &UNiagaraEmitter::ScriptRapidIterationParameterChanged));
+	}
+
+	if (UpdateScriptProps.Script)
+	{
+		UpdateScriptProps.Script->RapidIterationParameters.AddOnChangedHandler(
+			FNiagaraParameterStore::FOnChanged::FDelegate::CreateUObject(this, &UNiagaraEmitter::ScriptRapidIterationParameterChanged));
+	}
+
+	for (FNiagaraEventScriptProperties& EventScriptProperties : EventHandlerScriptProps)
+	{
+		if (EventScriptProperties.Script)
+		{
+			EventScriptProperties.Script->RapidIterationParameters.AddOnChangedHandler(
+				FNiagaraParameterStore::FOnChanged::FDelegate::CreateUObject(this, &UNiagaraEmitter::ScriptRapidIterationParameterChanged));
+		}
+	}
+
+	for (UNiagaraSimulationStageBase* SimulationStage : SimulationStages)
+	{
+		if (SimulationStage)
+		{
+			SimulationStage->OnChanged().AddUObject(this, &UNiagaraEmitter::SimulationStageChanged);
+
+			if (SimulationStage->Script)
+			{
+				SimulationStage->Script->RapidIterationParameters.AddOnChangedHandler(
+					FNiagaraParameterStore::FOnChanged::FDelegate::CreateUObject(this, &UNiagaraEmitter::ScriptRapidIterationParameterChanged));
+			}
+		}
+	}
+
+	for (UNiagaraRendererProperties* Renderer : RendererProperties)
+	{
+		if (Renderer)
+		{
+			Renderer->OnChanged().AddUObject(this, &UNiagaraEmitter::RendererChanged);
+		}
+	}
+
+	if (EditorData != nullptr)
+	{
+		EditorData->OnPersistentDataChanged().AddUObject(this, &UNiagaraEmitter::PersistentEditorDataChanged);
+	}
+}
+
 #endif
 
 bool UNiagaraEmitter::UsesScript(const UNiagaraScript* Script)const
