Merged CL 16620400: Fixed a bug where Niagara assets were marked as dirty as result of the merge during load

diff --git a/Engine/Plugins/FX/Niagara/Source/Niagara/Private/NiagaraEmitter.cpp b/Engine/Plugins/FX/Niagara/Source/Niagara/Private/NiagaraEmitter.cpp
--- a/Engine/Plugins/FX/Niagara/Source/Niagara/Private/NiagaraEmitter.cpp
+++ b/Engine/Plugins/FX/Niagara/Source/Niagara/Private/NiagaraEmitter.cpp
@@ -517,7 +517,7 @@ void UNiagaraEmitter::PostLoad()
 			}
 		}
 
-		if (IsSynchronizedWithParent() == false)
+		if (IsSynchronizedWithParent() == false && IsRunningCommandlet())
 		{
 			// Modify here so that the asset will be marked dirty when using the resave commandlet.  This will be ignored during regular post load.
 			Modify();
@@ -1278,8 +1278,13 @@ void UNiagaraEmitter::UpdateEmitterAfterLoad()
 	{
 		if (IsSynchronizedWithParent() == false)
 		{
-			// Modify here so that the asset will be marked dirty when using the resave commandlet.  This will be ignored during regular post load.
+			bool bIsPackageDirty = GetOutermost()->IsDirty();
 			MergeChangesFromParent();
+			if (bIsPackageDirty == false)
+			{
+				// we do not want to dirty the system from the merge on load
+				GetOutermost()->SetDirtyFlag(false);
+			}
 		}
 
 		// Reset scripts if recompile is forced.
