Fixing callsites for new raw pointer to TObjectPtr upgrade pass on Engine + ShooterGame.

diff --git a/Engine/Plugins/Enterprise/VariantManagerContent/Source/VariantManagerContent/Private/LevelVariantSets.cpp b/Engine/Plugins/Enterprise/VariantManagerContent/Source/VariantManagerContent/Private/LevelVariantSets.cpp
--- a/Engine/Plugins/Enterprise/VariantManagerContent/Source/VariantManagerContent/Private/LevelVariantSets.cpp
+++ b/Engine/Plugins/Enterprise/VariantManagerContent/Source/VariantManagerContent/Private/LevelVariantSets.cpp
@@ -305,7 +305,7 @@ UObject* ULevelVariantSets::GetDirectorInstance(UObject* WorldContext)
 	}
 
 	ULevelVariantSetsFunctionDirector* TargetDirector = nullptr;
-	if ( ULevelVariantSetsFunctionDirector** ExistingDirector = DirectorOuter->DirectorInstances.Find( DirectorClass ) )
+	if ( UE_TRANSITIONAL_OBJECT_PTR(ULevelVariantSetsFunctionDirector*)* ExistingDirector = DirectorOuter->DirectorInstances.Find( DirectorClass ) )
 	{
 		TargetDirector = *ExistingDirector;
 	}
diff --git a/Engine/Plugins/Experimental/MeshModelingToolset/Source/MeshModelingTools/Private/BakeMeshAttributeMapsTool.cpp b/Engine/Plugins/Experimental/MeshModelingToolset/Source/MeshModelingTools/Private/BakeMeshAttributeMapsTool.cpp
--- a/Engine/Plugins/Experimental/MeshModelingToolset/Source/MeshModelingTools/Private/BakeMeshAttributeMapsTool.cpp
+++ b/Engine/Plugins/Experimental/MeshModelingToolset/Source/MeshModelingTools/Private/BakeMeshAttributeMapsTool.cpp
@@ -401,8 +401,8 @@ void UBakeMeshAttributeMapsTool::Setup()
 	AddToolPropertySource(MultiTextureProps);
 	SetToolPropertySourceEnabled(MultiTextureProps, false);
 
-	auto SetDirtyCallback = [this](TMap<int32, UTexture2D*>) { bInputsDirty = true; };
-	auto NotEqualsCallback = [](const TMap<int32, UTexture2D*>& A, const TMap<int32, UTexture2D*>& B) -> bool { return !(A.OrderIndependentCompareEqual(B)); };
+	auto SetDirtyCallback = [this](decltype(MultiTextureProps->MaterialIDSourceTextureMap)) { bInputsDirty = true; };
+	auto NotEqualsCallback = [](const decltype(MultiTextureProps->MaterialIDSourceTextureMap)& A, const decltype(MultiTextureProps->MaterialIDSourceTextureMap)& B) -> bool { return !(A.OrderIndependentCompareEqual(B)); };
 	MultiTextureProps->WatchProperty(MultiTextureProps->MaterialIDSourceTextureMap, SetDirtyCallback, NotEqualsCallback);
 	MultiTextureProps->WatchProperty(MultiTextureProps->UVLayer, [this](float) { bInputsDirty = true; });
 
@@ -1143,7 +1143,7 @@ EBakeOpState UBakeMeshAttributeMapsTool::UpdateResult_MultiTexture()
 		return EBakeOpState::Invalid;
 	}
 
-	for (TPair<int32, UTexture2D*>& InputTexture : MultiTextureProps->MaterialIDSourceTextureMap)
+	for (auto& InputTexture : MultiTextureProps->MaterialIDSourceTextureMap)
 	{
 		if (InputTexture.Value == nullptr)
 		{
@@ -1155,7 +1155,7 @@ EBakeOpState UBakeMeshAttributeMapsTool::UpdateResult_MultiTexture()
 
 	CachedMultiTextures.Reset();
 
-	for ( TPair<int32, UTexture2D*>& InputTexture : MultiTextureProps->MaterialIDSourceTextureMap)
+	for ( auto& InputTexture : MultiTextureProps->MaterialIDSourceTextureMap)
 	{
 		UTexture2D* Texture = InputTexture.Value;
 		if (!ensure(Texture != nullptr))
diff --git a/Engine/Plugins/Experimental/MeshModelingToolset/Source/MeshModelingTools/Private/Sculpting/MeshSculptToolBase.cpp b/Engine/Plugins/Experimental/MeshModelingToolset/Source/MeshModelingTools/Private/Sculpting/MeshSculptToolBase.cpp
--- a/Engine/Plugins/Experimental/MeshModelingToolset/Source/MeshModelingTools/Private/Sculpting/MeshSculptToolBase.cpp
+++ b/Engine/Plugins/Experimental/MeshModelingToolset/Source/MeshModelingTools/Private/Sculpting/MeshSculptToolBase.cpp
@@ -329,7 +329,7 @@ void UMeshSculptToolBase::SetActivePrimaryBrushType(int32 Identifier)
 	PrimaryBrushOp = (*Factory)->Build();
 	PrimaryBrushOp->Falloff = PrimaryFalloff;
 
-	UMeshSculptBrushOpProps** FoundProps = BrushOpPropSets.Find(Identifier);
+	UE_TRANSITIONAL_OBJECT_PTR(UMeshSculptBrushOpProps)* FoundProps = BrushOpPropSets.Find(Identifier);
 	if (FoundProps != nullptr)
 	{
 		SetToolPropertySourceEnabled(*FoundProps, bBrushOpPropsVisible);
@@ -361,7 +361,7 @@ void UMeshSculptToolBase::SetActiveSecondaryBrushType(int32 Identifier)
 	SecondaryFalloff->FalloffFunc = UE::SculptFalloffs::MakeStandardSmoothFalloff();
 	SecondaryBrushOp->Falloff = SecondaryFalloff;
 
-	UMeshSculptBrushOpProps** FoundProps = SecondaryBrushOpPropSets.Find(Identifier);
+	UE_TRANSITIONAL_OBJECT_PTR(UMeshSculptBrushOpProps)* FoundProps = SecondaryBrushOpPropSets.Find(Identifier);
 	if (FoundProps != nullptr)
 	{
 		SetToolPropertySourceEnabled(*FoundProps, bBrushOpPropsVisible);
diff --git a/Engine/Plugins/Experimental/MeshModelingToolset/Source/ModelingComponents/Private/Drawing/PreviewGeometryActor.cpp b/Engine/Plugins/Experimental/MeshModelingToolset/Source/ModelingComponents/Private/Drawing/PreviewGeometryActor.cpp
--- a/Engine/Plugins/Experimental/MeshModelingToolset/Source/ModelingComponents/Private/Drawing/PreviewGeometryActor.cpp
+++ b/Engine/Plugins/Experimental/MeshModelingToolset/Source/ModelingComponents/Private/Drawing/PreviewGeometryActor.cpp
@@ -102,7 +102,7 @@ ULineSetComponent* UPreviewGeometry::AddLineSet(const FString& SetIdentifier)
 
 ULineSetComponent* UPreviewGeometry::FindLineSet(const FString& LineSetIdentifier)
 {
-	ULineSetComponent** Found = LineSets.Find(LineSetIdentifier);
+	UE_TRANSITIONAL_OBJECT_PTR(ULineSetComponent)* Found = LineSets.Find(LineSetIdentifier);
 	if (Found != nullptr)
 	{
 		return *Found;
@@ -112,7 +112,7 @@ ULineSetComponent* UPreviewGeometry::FindLineSet(const FString& LineSetIdentifie
 
 bool UPreviewGeometry::RemoveLineSet(const FString& LineSetIdentifier, bool bDestroy)
 {
-	ULineSetComponent** Found = LineSets.Find(LineSetIdentifier);
+	UE_TRANSITIONAL_OBJECT_PTR(ULineSetComponent)* Found = LineSets.Find(LineSetIdentifier);
 	if (Found != nullptr)
 	{
 		ULineSetComponent* LineSet = *Found;
@@ -146,7 +146,7 @@ void UPreviewGeometry::RemoveAllLineSets(bool bDestroy)
 
 bool UPreviewGeometry::SetLineSetVisibility(const FString& LineSetIdentifier, bool bVisible)
 {
-	ULineSetComponent** Found = LineSets.Find(LineSetIdentifier);
+	UE_TRANSITIONAL_OBJECT_PTR(ULineSetComponent)* Found = LineSets.Find(LineSetIdentifier);
 	if (Found != nullptr)
 	{
 		(*Found)->SetVisibility(bVisible);
@@ -158,7 +158,7 @@ bool UPreviewGeometry::SetLineSetVisibility(const FString& LineSetIdentifier, bo
 
 bool UPreviewGeometry::SetLineSetMaterial(const FString& LineSetIdentifier, UMaterialInterface* NewMaterial)
 {
-	ULineSetComponent** Found = LineSets.Find(LineSetIdentifier);
+	UE_TRANSITIONAL_OBJECT_PTR(ULineSetComponent)* Found = LineSets.Find(LineSetIdentifier);
 	if (Found != nullptr)
 	{
 		(*Found)->SetLineMaterial(NewMaterial);
diff --git a/Engine/Plugins/FX/Niagara/Source/Niagara/Private/NiagaraDataInterfaceIntRenderTarget2D.cpp b/Engine/Plugins/FX/Niagara/Source/Niagara/Private/NiagaraDataInterfaceIntRenderTarget2D.cpp
--- a/Engine/Plugins/FX/Niagara/Source/Niagara/Private/NiagaraDataInterfaceIntRenderTarget2D.cpp
+++ b/Engine/Plugins/FX/Niagara/Source/Niagara/Private/NiagaraDataInterfaceIntRenderTarget2D.cpp
@@ -586,7 +586,7 @@ void UNiagaraDataInterfaceIntRenderTarget2D::DestroyPerInstanceData(void* PerIns
 
 	// Make sure to clear out the reference to the render target if we created one.
 	extern int32 GNiagaraReleaseResourceOnRemove;
-	UTextureRenderTarget2D* ExistingRenderTarget = nullptr;
+	decltype(ManagedRenderTargets)::ValueType ExistingRenderTarget = nullptr;
 	if ( ManagedRenderTargets.RemoveAndCopyValue(SystemInstance->GetId(), ExistingRenderTarget) && GNiagaraReleaseResourceOnRemove)
 	{
 		ExistingRenderTarget->ReleaseResource();
@@ -674,7 +674,7 @@ bool UNiagaraDataInterfaceIntRenderTarget2D::UpdateInstanceTexture(FNiagaraSyste
 				InstanceData->TargetTexture = UserTargetTexture;
 
 				extern int32 GNiagaraReleaseResourceOnRemove;
-				UTextureRenderTarget2D* ExistingRenderTarget = nullptr;
+				decltype(ManagedRenderTargets)::ValueType ExistingRenderTarget = nullptr;
 				if (ManagedRenderTargets.RemoveAndCopyValue(SystemInstance->GetId(), ExistingRenderTarget) && GNiagaraReleaseResourceOnRemove)
 				{
 					ExistingRenderTarget->ReleaseResource();
diff --git a/Engine/Plugins/FX/Niagara/Source/Niagara/Private/NiagaraDataInterfaceRenderTargetCube.cpp b/Engine/Plugins/FX/Niagara/Source/Niagara/Private/NiagaraDataInterfaceRenderTargetCube.cpp
--- a/Engine/Plugins/FX/Niagara/Source/Niagara/Private/NiagaraDataInterfaceRenderTargetCube.cpp
+++ b/Engine/Plugins/FX/Niagara/Source/Niagara/Private/NiagaraDataInterfaceRenderTargetCube.cpp
@@ -536,7 +536,7 @@ void UNiagaraDataInterfaceRenderTargetCube::DestroyPerInstanceData(void* PerInst
 
 	// Make sure to clear out the reference to the render target if we created one.
 	extern int32 GNiagaraReleaseResourceOnRemove;
-	UTextureRenderTargetCube* ExistingRenderTarget = nullptr;
+	decltype(ManagedRenderTargets)::ValueType ExistingRenderTarget = nullptr;
 	if (ManagedRenderTargets.RemoveAndCopyValue(SystemInstance->GetId(), ExistingRenderTarget) && GNiagaraReleaseResourceOnRemove)
 	{
 		ExistingRenderTarget->ReleaseResource();
@@ -604,7 +604,7 @@ bool UNiagaraDataInterfaceRenderTargetCube::PerInstanceTick(void* PerInstanceDat
 		InstanceData->TargetTexture = UserTargetTexture;
 
 		extern int32 GNiagaraReleaseResourceOnRemove;
-		UTextureRenderTargetCube* ExistingRenderTarget = nullptr;
+		decltype(ManagedRenderTargets)::ValueType ExistingRenderTarget = nullptr;
 		if (ManagedRenderTargets.RemoveAndCopyValue(SystemInstance->GetId(), ExistingRenderTarget) && GNiagaraReleaseResourceOnRemove)
 		{
 			ExistingRenderTarget->ReleaseResource();
diff --git a/Engine/Plugins/FX/Niagara/Source/NiagaraEditor/Private/NiagaraEditorData.cpp b/Engine/Plugins/FX/Niagara/Source/NiagaraEditor/Private/NiagaraEditorData.cpp
--- a/Engine/Plugins/FX/Niagara/Source/NiagaraEditor/Private/NiagaraEditorData.cpp
+++ b/Engine/Plugins/FX/Niagara/Source/NiagaraEditor/Private/NiagaraEditorData.cpp
@@ -101,7 +101,7 @@ TOptional<TTuple<FName /*SyncedOldName*/, FName /*SyncedNewName*/>> UNiagaraEdit
 	if (DestScriptVar == nullptr)
 	{
 		const FGuid& SourceScriptVarId = SourceScriptVar->Metadata.GetVariableGuid();
-		UNiagaraScriptVariable** ScriptVarPtr = EditorOnlyScriptVars.FindByPredicate([&SourceScriptVarId](const UNiagaraScriptVariable* ScriptVar) { return ScriptVar->Metadata.GetVariableGuid() == SourceScriptVarId; });
+		UE_TRANSITIONAL_OBJECT_PTR(UNiagaraScriptVariable)* ScriptVarPtr = EditorOnlyScriptVars.FindByPredicate([&SourceScriptVarId](const UNiagaraScriptVariable* ScriptVar) { return ScriptVar->Metadata.GetVariableGuid() == SourceScriptVarId; });
 		if (ScriptVarPtr == nullptr)
 		{
 			// Failed to find a DestScriptVar with an Id matching that of SourceScriptVar.
diff --git a/Engine/Plugins/FX/Niagara/Source/NiagaraEditor/Private/NiagaraParameterDefinitions.cpp b/Engine/Plugins/FX/Niagara/Source/NiagaraEditor/Private/NiagaraParameterDefinitions.cpp
--- a/Engine/Plugins/FX/Niagara/Source/NiagaraEditor/Private/NiagaraParameterDefinitions.cpp
+++ b/Engine/Plugins/FX/Niagara/Source/NiagaraEditor/Private/NiagaraParameterDefinitions.cpp
@@ -44,7 +44,7 @@ void UNiagaraParameterDefinitions::PostEditChangeProperty(struct FPropertyChange
 void UNiagaraParameterDefinitions::AddParameter(const FNiagaraVariable& NewVariable)
 {
 	Modify();
-	UNiagaraScriptVariable*& NewScriptVariable = ScriptVariables.Add_GetRef(NewObject<UNiagaraScriptVariable>(this, FName(), RF_Transactional));
+	UE_TRANSITIONAL_OBJECT_PTR(UNiagaraScriptVariable)& NewScriptVariable = ScriptVariables.Add_GetRef(NewObject<UNiagaraScriptVariable>(this, FName(), RF_Transactional));
 	NewScriptVariable->Init(NewVariable, FNiagaraVariableMetaData());
 	NewScriptVariable->SetIsStaticSwitch(false);
 	NewScriptVariable->SetIsSubscribedToParameterDefinitions(true);
@@ -69,7 +69,7 @@ void UNiagaraParameterDefinitions::RemoveParameter(const FNiagaraVariable& Varia
 
 void UNiagaraParameterDefinitions::RenameParameter(const FNiagaraVariable& VariableToRename, const FName NewName)
 {
-	if (UNiagaraScriptVariable** ScriptVariablePtr = ScriptVariables.FindByPredicate([&VariableToRename](const UNiagaraScriptVariable* ScriptVariable) { return ScriptVariable->Variable == VariableToRename; }))
+	if (UE_TRANSITIONAL_OBJECT_PTR(UNiagaraScriptVariable)* ScriptVariablePtr = ScriptVariables.FindByPredicate([&VariableToRename](const UNiagaraScriptVariable* ScriptVariable) { return ScriptVariable->Variable == VariableToRename; }))
 	{
 		const FName OldName = VariableToRename.GetName();
 		Modify();
@@ -112,7 +112,7 @@ TArray<FGuid> UNiagaraParameterDefinitions::GetParameterIds() const
 
 UNiagaraScriptVariable* UNiagaraParameterDefinitions::GetScriptVariable(const FNiagaraVariable& Var)
 {
-	if (UNiagaraScriptVariable* const* ScriptVarPtr = ScriptVariables.FindByPredicate([Var](const UNiagaraScriptVariable* ScriptVar) { return ScriptVar->Variable == Var; }))
+	if (UE_TRANSITIONAL_OBJECT_PTR(UNiagaraScriptVariable) const* ScriptVarPtr = ScriptVariables.FindByPredicate([Var](const UNiagaraScriptVariable* ScriptVar) { return ScriptVar->Variable == Var; }))
 	{
 		return *ScriptVarPtr;
 	}
@@ -121,7 +121,7 @@ UNiagaraScriptVariable* UNiagaraParameterDefinitions::GetScriptVariable(const FN
 
 UNiagaraScriptVariable* UNiagaraParameterDefinitions::GetScriptVariable(const FGuid& ScriptVarId)
 {
-	if (UNiagaraScriptVariable* const* ScriptVarPtr = ScriptVariables.FindByPredicate([ScriptVarId](const UNiagaraScriptVariable* ScriptVar) { return ScriptVar->Metadata.GetVariableGuid() == ScriptVarId; }))
+	if (UE_TRANSITIONAL_OBJECT_PTR(UNiagaraScriptVariable) const* ScriptVarPtr = ScriptVariables.FindByPredicate([ScriptVarId](const UNiagaraScriptVariable* ScriptVar) { return ScriptVar->Metadata.GetVariableGuid() == ScriptVarId; }))
 	{
 		return *ScriptVarPtr;
 	}
@@ -184,7 +184,7 @@ void UNiagaraParameterDefinitions::SynchronizeWithSubscribedParameterDefinitions
 				const FName& ExternalScriptVarName = ExternalScriptVar->Variable.GetName();
 				for (const FGuid& InternalScriptVarId : BindingNameSubscription.InternalScriptVarIds)
 				{
-					if (UNiagaraScriptVariable* const* ScriptVarPtr = ScriptVariables.FindByPredicate([&InternalScriptVarId](const UNiagaraScriptVariable* ScriptVar) { return ScriptVar->Metadata.GetVariableGuid() == InternalScriptVarId; }))
+					if (UE_TRANSITIONAL_OBJECT_PTR(UNiagaraScriptVariable) const* ScriptVarPtr = ScriptVariables.FindByPredicate([&InternalScriptVarId](const UNiagaraScriptVariable* ScriptVar) { return ScriptVar->Metadata.GetVariableGuid() == InternalScriptVarId; }))
 					{
 						UNiagaraScriptVariable* ScriptVar = *ScriptVarPtr;
 						if (ScriptVar->GetChangeId() != ExternalScriptVarChangeId)
diff --git a/Engine/Plugins/FX/Niagara/Source/NiagaraEditor/Public/NiagaraNodeFunctionCall.h b/Engine/Plugins/FX/Niagara/Source/NiagaraEditor/Public/NiagaraNodeFunctionCall.h
--- a/Engine/Plugins/FX/Niagara/Source/NiagaraEditor/Public/NiagaraNodeFunctionCall.h
+++ b/Engine/Plugins/FX/Niagara/Source/NiagaraEditor/Public/NiagaraNodeFunctionCall.h
@@ -168,7 +168,7 @@ public:
 	virtual TSharedPtr<SGraphNode> CreateVisualWidget() override;
 
 	// Messages API
-	NIAGARAEDITOR_API const TMap<FGuid, UNiagaraMessageData*>& GetMessages() const { return MessageKeyToMessageMap; };
+	NIAGARAEDITOR_API const auto& GetMessages() const { return MessageKeyToMessageMap; };
 	NIAGARAEDITOR_API void AddMessage(const FGuid& MessageKey, UNiagaraMessageData* NewMessage) { MessageKeyToMessageMap.Add(MessageKey, NewMessage); };
 	NIAGARAEDITOR_API void RemoveMessage(const FGuid& MessageKey) { MessageKeyToMessageMap.Remove(MessageKey); };
 	void RemoveMessageDelegateable(const FGuid MessageKey) { MessageKeyToMessageMap.Remove(MessageKey); };
@@ -216,7 +216,7 @@ protected:
 	FString FunctionDisplayName;
 
 	UPROPERTY(meta = (SkipForCompileHash="true"))
-	TMap<FGuid, UNiagaraMessageData*> MessageKeyToMessageMap;
+	TMap<FGuid, TObjectPtr<UNiagaraMessageData>> MessageKeyToMessageMap;
 	
 	UPROPERTY(meta = (SkipForCompileHash="true"))
 	TArray<FNiagaraStackMessage> StackMessages;
diff --git a/Engine/Plugins/Runtime/HairStrands/Source/HairStrandsCore/Private/GroomAsset.cpp b/Engine/Plugins/Runtime/HairStrands/Source/HairStrandsCore/Private/GroomAsset.cpp
--- a/Engine/Plugins/Runtime/HairStrands/Source/HairStrandsCore/Private/GroomAsset.cpp
+++ b/Engine/Plugins/Runtime/HairStrands/Source/HairStrandsCore/Private/GroomAsset.cpp
@@ -1593,7 +1593,7 @@ void UGroomAsset::RemoveUserDataOfClass(TSubclassOf<UAssetUserData> InUserDataCl
 
 const TArray<UAssetUserData*>* UGroomAsset::GetAssetUserDataArray() const
 {
-	return &AssetUserData;
+	return &ToRawPtrTArrayUnsafe(AssetUserData);
 }
 
 bool UGroomAsset::CanRebuildFromDescription() const
diff --git a/Engine/Plugins/Runtime/HairStrands/Source/HairStrandsCore/Private/GroomComponent.cpp b/Engine/Plugins/Runtime/HairStrands/Source/HairStrandsCore/Private/GroomComponent.cpp
--- a/Engine/Plugins/Runtime/HairStrands/Source/HairStrandsCore/Private/GroomComponent.cpp
+++ b/Engine/Plugins/Runtime/HairStrands/Source/HairStrandsCore/Private/GroomComponent.cpp
@@ -1129,7 +1129,7 @@ void UGroomComponent::ReleaseHairSimulation()
 {
 	for (int32 CompIndex = 0; CompIndex < NiagaraComponents.Num(); ++CompIndex)
 	{
-		UNiagaraComponent*& NiagaraComponent = NiagaraComponents[CompIndex];
+		UE_TRANSITIONAL_OBJECT_PTR(UNiagaraComponent)& NiagaraComponent = NiagaraComponents[CompIndex];
 		if (NiagaraComponent && !NiagaraComponent->IsBeingDestroyed())
 		{
 			if (GetWorld() && NiagaraComponent->IsRegistered())
@@ -1167,7 +1167,7 @@ void UGroomComponent::UpdateHairSimulation()
 	NiagaraComponents.SetNumZeroed(NumComponents);
 	for (int32 CompIndex = 0; CompIndex < NumComponents; ++CompIndex)
 	{
-		UNiagaraComponent*& NiagaraComponent = NiagaraComponents[CompIndex];
+		UE_TRANSITIONAL_OBJECT_PTR(UNiagaraComponent)& NiagaraComponent = NiagaraComponents[CompIndex];
 		if (ValidComponents[CompIndex])
 		{
 			if (!NiagaraComponent)
@@ -1187,8 +1187,8 @@ void UGroomComponent::UpdateHairSimulation()
 				if(!CosseratRodsSystem) CosseratRodsSystem = LoadObject<UNiagaraSystem>(nullptr, TEXT("/HairStrands/Emitters/StableRodsSystem.StableRodsSystem"));
 
 				UNiagaraSystem* NiagaraAsset =
-					(GroomAsset->HairGroupsPhysics[CompIndex].SolverSettings.NiagaraSolver == EGroomNiagaraSolvers::AngularSprings) ? AngularSpringsSystem :
-					(GroomAsset->HairGroupsPhysics[CompIndex].SolverSettings.NiagaraSolver == EGroomNiagaraSolvers::CosseratRods) ? CosseratRodsSystem :
+					(GroomAsset->HairGroupsPhysics[CompIndex].SolverSettings.NiagaraSolver == EGroomNiagaraSolvers::AngularSprings) ? ToRawPtr(AngularSpringsSystem) :
+					(GroomAsset->HairGroupsPhysics[CompIndex].SolverSettings.NiagaraSolver == EGroomNiagaraSolvers::CosseratRods) ? ToRawPtr(CosseratRodsSystem) :
 					GroomAsset->HairGroupsPhysics[CompIndex].SolverSettings.CustomSystem.LoadSynchronous();
 
 				NiagaraComponent->SetAsset(NiagaraAsset);
diff --git a/Engine/Plugins/Runtime/HairStrands/Source/HairStrandsCore/Private/HairCardsBuilder.cpp b/Engine/Plugins/Runtime/HairStrands/Source/HairStrandsCore/Private/HairCardsBuilder.cpp
--- a/Engine/Plugins/Runtime/HairStrands/Source/HairStrandsCore/Private/HairCardsBuilder.cpp
+++ b/Engine/Plugins/Runtime/HairStrands/Source/HairStrandsCore/Private/HairCardsBuilder.cpp
@@ -3441,10 +3441,10 @@ void BuildGeometry(
 		HairCards::CreateCardsInterpolation(Out.Cards, OutGuides, OutInterpolation);
 		check(Out.Cards.GetNum() > 0);
 
-		AllocateAtlasTexture(OutTextures.DepthTexture,		Out.Atlas.Resolution, LODName, TEXT("_CardsAtlas_Depth"), AllocateAtlasTexture_Depth);
-		AllocateAtlasTexture(OutTextures.CoverageTexture,	Out.Atlas.Resolution, LODName, TEXT("_CardsAtlas_Coverage"), AllocateAtlasTexture_Coverage);
-		AllocateAtlasTexture(OutTextures.TangentTexture,	Out.Atlas.Resolution, LODName, TEXT("_CardsAtlas_Tangent"), AllocateAtlasTexture_Tangent);
-		AllocateAtlasTexture(OutTextures.AttributeTexture,	Out.Atlas.Resolution, LODName, TEXT("_CardsAtlas_Attribute"), AllocateAtlasTexture_Attribute);
+		AllocateAtlasTexture(static_cast<UTexture2D*&>(OutTextures.DepthTexture),		Out.Atlas.Resolution, LODName, TEXT("_CardsAtlas_Depth"), AllocateAtlasTexture_Depth);
+		AllocateAtlasTexture(static_cast<UTexture2D*&>(OutTextures.CoverageTexture),	Out.Atlas.Resolution, LODName, TEXT("_CardsAtlas_Coverage"), AllocateAtlasTexture_Coverage);
+		AllocateAtlasTexture(static_cast<UTexture2D*&>(OutTextures.TangentTexture),	Out.Atlas.Resolution, LODName, TEXT("_CardsAtlas_Tangent"), AllocateAtlasTexture_Tangent);
+		AllocateAtlasTexture(static_cast<UTexture2D*&>(OutTextures.AttributeTexture),	Out.Atlas.Resolution, LODName, TEXT("_CardsAtlas_Attribute"), AllocateAtlasTexture_Attribute);
 	}
 
 	Out.Cards.BoundingBox.Init();
diff --git a/Engine/Source/Runtime/ClothingSystemRuntimeInterface/Private/ClothingSimulationInteractor.cpp b/Engine/Source/Runtime/ClothingSystemRuntimeInterface/Private/ClothingSimulationInteractor.cpp
--- a/Engine/Source/Runtime/ClothingSystemRuntimeInterface/Private/ClothingSimulationInteractor.cpp
+++ b/Engine/Source/Runtime/ClothingSystemRuntimeInterface/Private/ClothingSimulationInteractor.cpp
@@ -24,7 +24,7 @@ void UClothingSimulationInteractor::DestroyClothingInteractors()
 
 UClothingInteractor* UClothingSimulationInteractor::GetClothingInteractor(const FString& ClothingAssetName) const
 {
-	if (UClothingInteractor* const* const ClothingInteractor = ClothingInteractors.Find(FName(ClothingAssetName)))
+	if (UE_TRANSITIONAL_OBJECT_PTR(UClothingInteractor* const)* const ClothingInteractor = ClothingInteractors.Find(FName(ClothingAssetName)))
 	{
 		return *ClothingInteractor;
 	}
diff --git a/Engine/Source/Runtime/Experimental/InteractiveToolsFramework/Private/BaseGizmos/GizmoArrowComponent.cpp b/Engine/Source/Runtime/Experimental/InteractiveToolsFramework/Private/BaseGizmos/GizmoArrowComponent.cpp
--- a/Engine/Source/Runtime/Experimental/InteractiveToolsFramework/Private/BaseGizmos/GizmoArrowComponent.cpp
+++ b/Engine/Source/Runtime/Experimental/InteractiveToolsFramework/Private/BaseGizmos/GizmoArrowComponent.cpp
@@ -191,7 +191,7 @@ bool UGizmoArrowComponent::LineTraceComponent(FHitResult& OutHit, const FVector
 	// Copy what is done in the proxy object, but get data from the gizmo view context.
 	FVector StartPoint, EndPoint;
 	float PixelToWorldScale = 0;
-	bool bRenderVisibility = GetWorldEndpoints(bIsViewDependent, GizmoViewContext, UseOrigin, Direction, Gap, Length, bWorld,
+	bool bRenderVisibility = GetWorldEndpoints(bIsViewDependent, ToRawPtr(GizmoViewContext), UseOrigin, Direction, Gap, Length, bWorld,
 		[&Transform](const FVector& VectorIn) { return Transform.TransformVector(VectorIn); },
 		StartPoint, EndPoint, PixelToWorldScale);
 
diff --git a/Engine/Source/Runtime/Experimental/InteractiveToolsFramework/Private/BaseGizmos/GizmoBoxComponent.cpp b/Engine/Source/Runtime/Experimental/InteractiveToolsFramework/Private/BaseGizmos/GizmoBoxComponent.cpp
--- a/Engine/Source/Runtime/Experimental/InteractiveToolsFramework/Private/BaseGizmos/GizmoBoxComponent.cpp
+++ b/Engine/Source/Runtime/Experimental/InteractiveToolsFramework/Private/BaseGizmos/GizmoBoxComponent.cpp
@@ -274,7 +274,7 @@ bool UGizmoBoxComponent::LineTraceComponent(FHitResult& OutHit, const FVector St
 	if (bIsViewDependent)
 	{
 		FVector UseDirectionX, UseDirectionY, UseDirectionZ; // not used
-		GetBoxDirections(bIsViewDependent, GizmoViewContext,
+		GetBoxDirections(bIsViewDependent, ToRawPtr(GizmoViewContext),
 			Transform.TransformPosition(FVector::ZeroVector), // Origin
 			Rotation * FVector(1, 0, 0), // DirectionX
 			Rotation * FVector(0, 1, 0), // DirectionY
diff --git a/Engine/Source/Runtime/Experimental/InteractiveToolsFramework/Private/BaseGizmos/GizmoCircleComponent.cpp b/Engine/Source/Runtime/Experimental/InteractiveToolsFramework/Private/BaseGizmos/GizmoCircleComponent.cpp
--- a/Engine/Source/Runtime/Experimental/InteractiveToolsFramework/Private/BaseGizmos/GizmoCircleComponent.cpp
+++ b/Engine/Source/Runtime/Experimental/InteractiveToolsFramework/Private/BaseGizmos/GizmoCircleComponent.cpp
@@ -258,7 +258,7 @@ bool UGizmoCircleComponent::LineTraceComponent(FHitResult& OutHit, const FVector
 			: GizmoViewContext->GetViewDirection();
 		ViewDirection.Normalize();
 
-		GetVisibility(GizmoViewContext, ViewDirection, WorldNormal, WorldOrigin, bRenderVisibility, bCircleIsViewPlaneParallel);
+		GetVisibility(ToRawPtr(GizmoViewContext), ViewDirection, WorldNormal, WorldOrigin, bRenderVisibility, bCircleIsViewPlaneParallel);
 
 		if (!bRenderVisibility)
 		{
diff --git a/Engine/Source/Runtime/Experimental/InteractiveToolsFramework/Private/BaseGizmos/GizmoRectangleComponent.cpp b/Engine/Source/Runtime/Experimental/InteractiveToolsFramework/Private/BaseGizmos/GizmoRectangleComponent.cpp
--- a/Engine/Source/Runtime/Experimental/InteractiveToolsFramework/Private/BaseGizmos/GizmoRectangleComponent.cpp
+++ b/Engine/Source/Runtime/Experimental/InteractiveToolsFramework/Private/BaseGizmos/GizmoRectangleComponent.cpp
@@ -250,7 +250,7 @@ bool UGizmoRectangleComponent::LineTraceComponent(FHitResult& OutHit, const FVec
 	FVector UseOrigin = Transform.TransformPosition(FVector::ZeroVector);
 
 	TArray<FVector> Corners;
-	bool bRenderVisibility = GetWorldCorners(bIsViewDependent, GizmoViewContext, UseOrigin, DirectionX, DirectionY, 
+	bool bRenderVisibility = GetWorldCorners(bIsViewDependent, ToRawPtr(GizmoViewContext), UseOrigin, DirectionX, DirectionY, 
 		OffsetX, OffsetY, LengthX, LengthY,	bWorld, bOrientYAccordingToCamera,
 		[&Transform](const FVector& VectorIn) { return Transform.TransformVector(VectorIn); },
 		Corners);
