- Fix for some changes that were lost

#rnx

diff --git a/Engine/Plugins/FX/Niagara/Source/Niagara/Private/NiagaraRendererMeshes.cpp b/Engine/Plugins/FX/Niagara/Source/Niagara/Private/NiagaraRendererMeshes.cpp
--- a/Engine/Plugins/FX/Niagara/Source/Niagara/Private/NiagaraRendererMeshes.cpp
+++ b/Engine/Plugins/FX/Niagara/Source/Niagara/Private/NiagaraRendererMeshes.cpp
@@ -558,6 +558,8 @@ uint32 FNiagaraRendererMeshes::PerformSortAndCull(FParticleMeshRenderData& Parti
 {
 	if ( SourceMode == ENiagaraRendererSourceDataMode::Emitter )
 	{
+		ParticleMeshRenderData.ParticleSortedIndicesSRV = GFNiagaraNullSortedIndicesVertexBuffer.VertexBufferSRV.GetReference();
+		ParticleMeshRenderData.ParticleSortedIndicesOffset = 0xffffffff;
 		return 1;
 	}
 
@@ -876,7 +878,7 @@ void FNiagaraRendererMeshes::CreateMeshBatchForSection(
 		BatchElement.NumPrimitives = Section.NumTriangles;
 	}
 
-	if (GPUCountBufferOffset != INDEX_NONE)
+	if ((SourceMode != ENiagaraRendererSourceDataMode::Emitter) && (GPUCountBufferOffset != INDEX_NONE))
 	{
 		// We need to use indirect draw args, because the number of actual instances is coming from the GPU
 		auto Batcher = SceneProxy.GetBatcher();
diff --git a/Engine/Plugins/FX/Niagara/Source/Niagara/Private/NiagaraRendererSprites.cpp b/Engine/Plugins/FX/Niagara/Source/Niagara/Private/NiagaraRendererSprites.cpp
--- a/Engine/Plugins/FX/Niagara/Source/Niagara/Private/NiagaraRendererSprites.cpp
+++ b/Engine/Plugins/FX/Niagara/Source/Niagara/Private/NiagaraRendererSprites.cpp
@@ -745,7 +745,7 @@ void FNiagaraRendererSprites::CreateMeshBatchForView(
 	VFLooseParams.SortedIndicesOffset = VertexFactory.GetSortedIndicesOffset();
 
 	FNiagaraGPUInstanceCountManager::FIndirectArgSlot IndirectDraw;
-	if (GPUCountBufferOffset != INDEX_NONE)
+	if ((SourceMode != ENiagaraRendererSourceDataMode::Emitter) && (GPUCountBufferOffset != INDEX_NONE))
 	{
 		NiagaraEmitterInstanceBatcher* Batcher = SceneProxy.GetBatcher();
 		check(Batcher);
