Fix unbound sampler state parameters in Niagara data interfaces.

diff --git a/Engine/Plugins/FX/Niagara/Source/Niagara/Private/NiagaraDataInterfaceRenderTarget2D.cpp b/Engine/Plugins/FX/Niagara/Source/Niagara/Private/NiagaraDataInterfaceRenderTarget2D.cpp
--- a/Engine/Plugins/FX/Niagara/Source/Niagara/Private/NiagaraDataInterfaceRenderTarget2D.cpp
+++ b/Engine/Plugins/FX/Niagara/Source/Niagara/Private/NiagaraDataInterfaceRenderTarget2D.cpp
@@ -85,7 +85,7 @@ public:
 		OutputParam.Bind(ParameterMap, *(UNiagaraDataInterfaceRenderTarget2D::OutputName + ParameterInfo.DataInterfaceHLSLSymbol));
 
 		InputParam.Bind(ParameterMap, *(UNiagaraDataInterfaceRenderTarget2D::InputName + ParameterInfo.DataInterfaceHLSLSymbol));
-		InputSamplerStateParam.Bind(ParameterMap, *(UNiagaraDataInterfaceRenderTarget2D::InputName + TEXT("SamplerState") + ParameterInfo.DataInterfaceHLSLSymbol));
+		InputSamplerStateParam.Bind(ParameterMap, *(UNiagaraDataInterfaceRenderTarget2D::InputName + ParameterInfo.DataInterfaceHLSLSymbol + TEXT("SamplerState")));
 	}
 
 	void Set(FRHICommandList& RHICmdList, const FNiagaraDataInterfaceSetArgs& Context) const
diff --git a/Engine/Plugins/FX/Niagara/Source/Niagara/Private/NiagaraDataInterfaceRenderTarget2DArray.cpp b/Engine/Plugins/FX/Niagara/Source/Niagara/Private/NiagaraDataInterfaceRenderTarget2DArray.cpp
--- a/Engine/Plugins/FX/Niagara/Source/Niagara/Private/NiagaraDataInterfaceRenderTarget2DArray.cpp
+++ b/Engine/Plugins/FX/Niagara/Source/Niagara/Private/NiagaraDataInterfaceRenderTarget2DArray.cpp
@@ -59,7 +59,7 @@ public:
 		OutputParam.Bind(ParameterMap, *(UNiagaraDataInterfaceRenderTarget2DArray::OutputName + ParameterInfo.DataInterfaceHLSLSymbol));
 
 		InputParam.Bind(ParameterMap, *(UNiagaraDataInterfaceRenderTarget2DArray::InputName + ParameterInfo.DataInterfaceHLSLSymbol));
-		InputSamplerStateParam.Bind(ParameterMap, *(UNiagaraDataInterfaceRenderTarget2DArray::InputName + TEXT("SamplerState") + ParameterInfo.DataInterfaceHLSLSymbol));
+		InputSamplerStateParam.Bind(ParameterMap, *(UNiagaraDataInterfaceRenderTarget2DArray::InputName + ParameterInfo.DataInterfaceHLSLSymbol + TEXT("SamplerState")));
 	}
 
 	void Set(FRHICommandList& RHICmdList, const FNiagaraDataInterfaceSetArgs& Context) const
diff --git a/Engine/Plugins/FX/Niagara/Source/Niagara/Private/NiagaraDataInterfaceRenderTargetCube.cpp b/Engine/Plugins/FX/Niagara/Source/Niagara/Private/NiagaraDataInterfaceRenderTargetCube.cpp
--- a/Engine/Plugins/FX/Niagara/Source/Niagara/Private/NiagaraDataInterfaceRenderTargetCube.cpp
+++ b/Engine/Plugins/FX/Niagara/Source/Niagara/Private/NiagaraDataInterfaceRenderTargetCube.cpp
@@ -57,7 +57,7 @@ public:
 		OutputParam.Bind(ParameterMap, *(UNiagaraDataInterfaceRenderTargetCube::OutputName + ParameterInfo.DataInterfaceHLSLSymbol));
 
 		InputParam.Bind(ParameterMap, *(UNiagaraDataInterfaceRenderTargetCube::InputName + ParameterInfo.DataInterfaceHLSLSymbol));
-		InputSamplerStateParam.Bind(ParameterMap, *(UNiagaraDataInterfaceRenderTargetCube::InputName + TEXT("SamplerState") + ParameterInfo.DataInterfaceHLSLSymbol));
+		InputSamplerStateParam.Bind(ParameterMap, *(UNiagaraDataInterfaceRenderTargetCube::InputName + ParameterInfo.DataInterfaceHLSLSymbol + TEXT("SamplerState")));
 	}
 
 	void Set(FRHICommandList& RHICmdList, const FNiagaraDataInterfaceSetArgs& Context) const
diff --git a/Engine/Plugins/FX/Niagara/Source/Niagara/Private/NiagaraDataInterfaceRenderTargetVolume.cpp b/Engine/Plugins/FX/Niagara/Source/Niagara/Private/NiagaraDataInterfaceRenderTargetVolume.cpp
--- a/Engine/Plugins/FX/Niagara/Source/Niagara/Private/NiagaraDataInterfaceRenderTargetVolume.cpp
+++ b/Engine/Plugins/FX/Niagara/Source/Niagara/Private/NiagaraDataInterfaceRenderTargetVolume.cpp
@@ -58,7 +58,7 @@ public:
 		OutputParam.Bind(ParameterMap, *(UNiagaraDataInterfaceRenderTargetVolume::OutputName + ParameterInfo.DataInterfaceHLSLSymbol));
 
 		InputParam.Bind(ParameterMap, *(UNiagaraDataInterfaceRenderTargetVolume::InputName + ParameterInfo.DataInterfaceHLSLSymbol));
-		InputSamplerStateParam.Bind(ParameterMap, *(UNiagaraDataInterfaceRenderTargetVolume::InputName + TEXT("SamplerState") + ParameterInfo.DataInterfaceHLSLSymbol));
+		InputSamplerStateParam.Bind(ParameterMap, *(UNiagaraDataInterfaceRenderTargetVolume::InputName + ParameterInfo.DataInterfaceHLSLSymbol + TEXT("SamplerState")));
 	}
 
 	void Set(FRHICommandList& RHICmdList, const FNiagaraDataInterfaceSetArgs& Context) const
diff --git a/Engine/Plugins/FX/Niagara/Source/NiagaraCore/Private/NiagaraCustomVersion.cpp b/Engine/Plugins/FX/Niagara/Source/NiagaraCore/Private/NiagaraCustomVersion.cpp
--- a/Engine/Plugins/FX/Niagara/Source/NiagaraCore/Private/NiagaraCustomVersion.cpp
+++ b/Engine/Plugins/FX/Niagara/Source/NiagaraCore/Private/NiagaraCustomVersion.cpp
@@ -8,4 +8,4 @@ const FGuid FNiagaraCustomVersion::GUID(0xFCF57AFA, 0x50764283, 0xB9A9E658, 0xFF
 // Register the custom version with core
 FCustomVersionRegistration GRegisterNiagaraCustomVersion(FNiagaraCustomVersion::GUID, FNiagaraCustomVersion::LatestVersion, TEXT("NiagaraVer"));
 
-const FGuid FNiagaraCustomVersion::LatestScriptCompileVersion(0x03BD6918, 0xE41E8A4D, 0xB8A1BA26, 0xD7E638E4);
+const FGuid FNiagaraCustomVersion::LatestScriptCompileVersion(0xA83D40AB, 0xE1F64D9F, 0x9F5451FB, 0x8C954501);
