Integrating fix for particle cutouts.

diff --git a/Engine/Plugins/FX/Niagara/Source/Niagara/Private/NiagaraRendererSprites.cpp b/Engine/Plugins/FX/Niagara/Source/Niagara/Private/NiagaraRendererSprites.cpp
--- a/Engine/Plugins/FX/Niagara/Source/Niagara/Private/NiagaraRendererSprites.cpp
+++ b/Engine/Plugins/FX/Niagara/Source/Niagara/Private/NiagaraRendererSprites.cpp
@@ -405,7 +405,6 @@ void FNiagaraRendererSprites::InitializeSortInfo(FParticleSpriteRenderData& Part
 void FNiagaraRendererSprites::SetupVertexFactory(FParticleSpriteRenderData& ParticleSpriteRenderData, FNiagaraSpriteVertexFactory& VertexFactory) const
 {
 	VertexFactory.SetParticleFactoryType(NVFT_Sprite);
-	VertexFactory.InitResource();
 
 	// Set facing / alignment
 	{
@@ -449,6 +448,9 @@ void FNiagaraRendererSprites::SetupVertexFactory(FParticleSpriteRenderData& Part
 			VertexFactory.SetVertexBufferOverride(&CutoutVertexBuffer);
 		}
 	}
+	
+	// The InitResource needs to happen at the end here as SetVertexBufferOverride will set the UV buffers.
+	VertexFactory.InitResource();
 }
 
 FNiagaraSpriteUniformBufferRef FNiagaraRendererSprites::CreateViewUniformBuffer(FParticleSpriteRenderData& ParticleSpriteRenderData, const FSceneView& View, const FSceneViewFamily& ViewFamily, const FNiagaraSceneProxy& SceneProxy, FNiagaraSpriteVertexFactory& VertexFactory) const
