- Remove check for disabled world rendering, this doesn't work with the MRQ which will set this variable and manually rendering in OnEndFrame
- Adjust flush to 1 frame, this allows a single frame to be in queue before we will flush on the next frame

#rnx

diff --git a/Engine/Plugins/FX/Niagara/Source/Niagara/Private/NiagaraEmitterInstanceBatcher.cpp b/Engine/Plugins/FX/Niagara/Source/Niagara/Private/NiagaraEmitterInstanceBatcher.cpp
--- a/Engine/Plugins/FX/Niagara/Source/Niagara/Private/NiagaraEmitterInstanceBatcher.cpp
+++ b/Engine/Plugins/FX/Niagara/Source/Niagara/Private/NiagaraEmitterInstanceBatcher.cpp
@@ -58,7 +58,7 @@ uint32 FNiagaraComputeExecutionContext::TickCounter = 0;
 const FName NiagaraEmitterInstanceBatcher::TemporalEffectName("NiagaraEmitterInstanceBatcher");
 #endif // WITH_MGPU
 
-int32 GNiagaraGpuMaxQueuedRenderFrames = 10;
+int32 GNiagaraGpuMaxQueuedRenderFrames = 1;
 static FAutoConsoleVariableRef CVarNiagaraGpuMaxQueuedRenderFrames(
 	TEXT("fx.Niagara.Batcher.MaxQueuedFramesWithoutRender"),
 	GNiagaraGpuMaxQueuedRenderFrames,
@@ -410,21 +410,11 @@ void NiagaraEmitterInstanceBatcher::BuildConstantBuffers(FNiagaraGPUSystemTick&
 
 void NiagaraEmitterInstanceBatcher::Tick(UWorld* World, float DeltaTime)
 {
-	bool bFlushTicks = false;
-
-	if ((World != nullptr) && World->IsGameWorld())
-	{
-		if (UGameViewportClient* GameViewport = World->GetGameViewport())
-		{
-			bFlushTicks |= GameViewport->bDisableWorldRendering;
-		}
-	}
-
 	check(IsInGameThread());
 	ENQUEUE_RENDER_COMMAND(NiagaraPumpBatcher)(
-		[RT_NiagaraBatcher=this, RT_bFlushTicks=bFlushTicks](FRHICommandListImmediate& RHICmdList)
+		[RT_NiagaraBatcher=this](FRHICommandListImmediate& RHICmdList)
 		{
-			RT_NiagaraBatcher->ProcessPendingTicksFlush(RHICmdList, RT_bFlushTicks);
+			RT_NiagaraBatcher->ProcessPendingTicksFlush(RHICmdList, false);
 			RT_NiagaraBatcher->GetGPUInstanceCounterManager().FlushIndirectArgsPool();
 		}
 	);
