- Fix crash when changing the material on a mesh while used by a mesh renderer

#rnx

diff --git a/Engine/Plugins/FX/Niagara/Source/Niagara/Private/NiagaraRendererMeshes.cpp b/Engine/Plugins/FX/Niagara/Source/Niagara/Private/NiagaraRendererMeshes.cpp
--- a/Engine/Plugins/FX/Niagara/Source/Niagara/Private/NiagaraRendererMeshes.cpp
+++ b/Engine/Plugins/FX/Niagara/Source/Niagara/Private/NiagaraRendererMeshes.cpp
@@ -975,6 +975,11 @@ void FNiagaraRendererMeshes::GetDynamicMeshElements(const TArray<const FSceneVie
 
 				// @TODO : support multiple LOD
 				const int32 LODIndex = GetLODIndex(MeshIndex);
+				if (LODIndex == INDEX_NONE)
+				{
+					continue;
+				}
+
 				const FStaticMeshLODResources& LODModel = MeshData.RenderData->LODResources[LODIndex];
 				const int32 SectionCount = LODModel.Sections.Num();
 
@@ -1082,6 +1087,10 @@ void FNiagaraRendererMeshes::GetDynamicRayTracingInstances(FRayTracingMaterialGa
 	{
 		const FMeshData& MeshData = Meshes[MeshIndex];
 		const int32 LODIndex = GetLODIndex(MeshIndex);
+		if (LODIndex == INDEX_NONE)
+		{
+			continue;
+		}
 
 		const FStaticMeshLODResources& LODModel = MeshData.RenderData->LODResources[LODIndex];
 		FRayTracingGeometry& Geometry = MeshData.RenderData->LODResources[LODIndex].RayTracingGeometry;
