- Fix assert when using dump command

#rnx

diff --git a/Engine/Plugins/FX/Niagara/Source/Niagara/Private/NiagaraSystemSimulation.cpp b/Engine/Plugins/FX/Niagara/Source/Niagara/Private/NiagaraSystemSimulation.cpp
--- a/Engine/Plugins/FX/Niagara/Source/Niagara/Private/NiagaraSystemSimulation.cpp
+++ b/Engine/Plugins/FX/Niagara/Source/Niagara/Private/NiagaraSystemSimulation.cpp
@@ -874,6 +874,12 @@ void FNiagaraSystemSimulation::DumpTickInfo(FOutputDevice& Ar)
 
 	for ( int i=0; i < int(ENiagaraSystemInstanceState::Num); ++i )
 	{
+		// None is invalid so don't use
+		if ( i == int(ENiagaraSystemInstanceState::None) )
+		{
+			continue;
+		}
+
 		TArray<FNiagaraSystemInstance*> Instances = GetSystemInstances(ENiagaraSystemInstanceState(i));
 		if (Instances.Num() > 0)
 		{
