- Add prereq calculation for things attached to skeletal mesh components who have blend physics enabled

diff --git a/Engine/Plugins/FX/Niagara/Source/Niagara/Private/NiagaraSystemInstance.cpp b/Engine/Plugins/FX/Niagara/Source/Niagara/Private/NiagaraSystemInstance.cpp
--- a/Engine/Plugins/FX/Niagara/Source/Niagara/Private/NiagaraSystemInstance.cpp
+++ b/Engine/Plugins/FX/Niagara/Source/Niagara/Private/NiagaraSystemInstance.cpp
@@ -6,17 +6,19 @@
 #include "NiagaraCommon.h"
 #include "NiagaraDataInterface.h"
 #include "NiagaraStats.h"
-#include "Async/ParallelFor.h"
 #include "NiagaraParameterCollection.h"
 #include "NiagaraWorldManager.h"
 #include "NiagaraComponent.h"
 #include "NiagaraRenderer.h"
 #include "NiagaraGpuComputeDebug.h"
-#include "Templates/AlignmentTemplates.h"
 #include "NiagaraEmitterInstanceBatcher.h"
-#include "GameFramework/PlayerController.h"
 #include "NiagaraCrashReporterHandler.h"
+
 #include "Async/Async.h"
+#include "Async/ParallelFor.h"
+#include "Components/SkeletalMeshComponent.h"
+#include "GameFramework/PlayerController.h"
+#include "Templates/AlignmentTemplates.h"
 
 #if WITH_EDITORONLY_DATA
 #include "Editor.h"
@@ -1754,6 +1756,14 @@ ETickingGroup FNiagaraSystemInstance::CalculateTickGroup() const
 			{
 				//-TODO: This doesn't deal with 'DontCompleteUntil' on the prereq's tick, if we have to handle that it could mean continual TG demotion
 				ETickingGroup PrereqTG = ETickingGroup(FMath::Max(PrereqComponent->PrimaryComponentTick.TickGroup, PrereqComponent->PrimaryComponentTick.EndTickGroup) + 1);
+
+				// If we are attached to a skeletal mesh component and blend physics is enabled we can not tick until after that has completed
+				// otherwise we may not get the correct final transform location leading to latency.
+				if ( USkeletalMeshComponent* PrereqSMC = Cast<USkeletalMeshComponent>(PrereqComponent) )
+				{
+					PrereqTG = PrereqSMC->bBlendPhysics ? FMath::Max(PrereqTG, ETickingGroup(TG_EndPhysics + 1)) : PrereqTG;
+				}
+
 				NewTickGroup = FMath::Max(NewTickGroup, PrereqTG);
 			}
 
