- Integrate 16675172 - Contains fixes for instance culling / ray tracing

#rnx

diff --git a/Engine/Plugins/FX/Niagara/Source/Niagara/Private/NiagaraRendererMeshes.cpp b/Engine/Plugins/FX/Niagara/Source/Niagara/Private/NiagaraRendererMeshes.cpp
--- a/Engine/Plugins/FX/Niagara/Source/Niagara/Private/NiagaraRendererMeshes.cpp
+++ b/Engine/Plugins/FX/Niagara/Source/Niagara/Private/NiagaraRendererMeshes.cpp
@@ -554,9 +554,31 @@ void FNiagaraRendererMeshes::InitializeSortInfo(FParticleMeshRenderData& Particl
 	}
 }
 
+void FNiagaraRendererMeshes::PreparePerMeshData(FParticleMeshRenderData& ParticleMeshRenderData, const FNiagaraSceneProxy& SceneProxy, const FMeshData& MeshData) const
+{
+	// Calculate pivot offset / culling sphere
+	ParticleMeshRenderData.CullingSphere = MeshData.LocalCullingSphere;
+	if (MeshData.PivotOffsetSpace == ENiagaraMeshPivotOffsetSpace::Mesh)
+	{
+		ParticleMeshRenderData.WorldSpacePivotOffset = FVector::ZeroVector;
+		ParticleMeshRenderData.CullingSphere.Center += MeshData.PivotOffset;
+	}
+	else
+	{
+		ParticleMeshRenderData.WorldSpacePivotOffset = MeshData.PivotOffset;
+		if (MeshData.PivotOffsetSpace == ENiagaraMeshPivotOffsetSpace::Local ||
+			(bLocalSpace && MeshData.PivotOffsetSpace == ENiagaraMeshPivotOffsetSpace::Simulation))
+		{
+			// The offset is in local space, transform it to world
+			ParticleMeshRenderData.WorldSpacePivotOffset = SceneProxy.GetLocalToWorld().TransformVector(ParticleMeshRenderData.WorldSpacePivotOffset);
+		}
+	}
+}
+
 uint32 FNiagaraRendererMeshes::PerformSortAndCull(FParticleMeshRenderData& ParticleMeshRenderData, FGlobalDynamicReadBuffer& ReadBuffer, FNiagaraGPUSortInfo& SortInfo, NiagaraEmitterInstanceBatcher* Batcher, int32 MeshIndex) const
 {
-	if ( SourceMode == ENiagaraRendererSourceDataMode::Emitter )
+	// Emitter mode culls earlier on
+	if (SourceMode == ENiagaraRendererSourceDataMode::Emitter)
 	{
 		ParticleMeshRenderData.ParticleSortedIndicesSRV = GFNiagaraNullSortedIndicesVertexBuffer.VertexBufferSRV.GetReference();
 		ParticleMeshRenderData.ParticleSortedIndicesOffset = 0xffffffff;
@@ -597,9 +619,6 @@ uint32 FNiagaraRendererMeshes::PerformSortAndCull(FParticleMeshRenderData& Parti
 
 FNiagaraMeshUniformBufferRef FNiagaraRendererMeshes::CreatePerViewUniformBuffer(FParticleMeshRenderData& ParticleMeshRenderData, const FSceneView& View, const FMeshData& MeshData, const FNiagaraSceneProxy& SceneProxy) const
 {
-	ParticleMeshRenderData.WorldSpacePivotOffset = FVector::ZeroVector;
-	ParticleMeshRenderData.CullingSphere = MeshData.LocalCullingSphere;
-
 	// Compute the per-view uniform buffers.
 	FNiagaraMeshUniformParameters PerViewUniformParameters;
 	FMemory::Memzero(&PerViewUniformParameters, sizeof(PerViewUniformParameters)); // Clear unset bytes
@@ -609,24 +628,13 @@ FNiagaraMeshUniformBufferRef FNiagaraRendererMeshes::CreatePerViewUniformBuffer(
 	PerViewUniformParameters.DeltaSeconds = View.Family->DeltaWorldTime;
 	PerViewUniformParameters.MeshScale = MeshData.Scale;
 
-	// Calculate pivot offset
 	if (MeshData.PivotOffsetSpace == ENiagaraMeshPivotOffsetSpace::Mesh)
 	{
-		ParticleMeshRenderData.CullingSphere.Center += MeshData.PivotOffset;
-
 		PerViewUniformParameters.PivotOffset = MeshData.PivotOffset;
 		PerViewUniformParameters.bPivotOffsetIsWorldSpace = false;
 	}
 	else
 	{
-		ParticleMeshRenderData.WorldSpacePivotOffset = MeshData.PivotOffset;
-		if (MeshData.PivotOffsetSpace == ENiagaraMeshPivotOffsetSpace::Local ||
-			(bUseLocalSpace && MeshData.PivotOffsetSpace == ENiagaraMeshPivotOffsetSpace::Simulation))
-		{
-			// The offset is in local space, transform it to world
-			ParticleMeshRenderData.WorldSpacePivotOffset = SceneProxy.GetLocalToWorld().TransformVector(ParticleMeshRenderData.WorldSpacePivotOffset);
-		}
-
 		PerViewUniformParameters.PivotOffset = ParticleMeshRenderData.WorldSpacePivotOffset;
 		PerViewUniformParameters.bPivotOffsetIsWorldSpace = true;
 	}
@@ -1002,6 +1010,8 @@ void FNiagaraRendererMeshes::GetDynamicMeshElements(const TArray<const FSceneVie
 				VertexFactory.InitResource();
 				SetupVertexFactory(VertexFactory, LODModel);
 
+				PreparePerMeshData(ParticleMeshRenderData, *SceneProxy, MeshData);
+
 				// Sort/Cull particles if needed.
 				const uint32 NumInstances = PerformSortAndCull(ParticleMeshRenderData, Collector.GetDynamicReadBuffer(), SortInfo, Batcher, MeshData.SourceMeshIndex);
 				if ( NumInstances > 0 )
@@ -1039,6 +1049,7 @@ void FNiagaraRendererMeshes::GetDynamicMeshElements(const TArray<const FSceneVie
 							GPUCountBufferOffset, bIsWireframe, bIsInstancedStereo,
 							ParticleMeshRenderData.bNeedsCull && ParticleMeshRenderData.bSortCullOnGpu
 						);
+
 						Collector.AddMesh(ViewIndex, MeshBatch);
 					}
 				}
@@ -1068,6 +1079,11 @@ void FNiagaraRendererMeshes::GetDynamicRayTracingInstances(FRayTracingMaterialGa
 		return;
 	}
 
+	// Disable sorting and we don't support culling on the GPU
+	ensureMsgf(ParticleMeshRenderData.bNeedsCull == false || SimTarget == ENiagaraSimTarget::CPUSim, TEXT("Mesh culling & raytracing is not suppoerted currently."));
+	ParticleMeshRenderData.bNeedsSort = false;
+	ParticleMeshRenderData.bSortCullOnGpu = false;
+
 	SCOPE_CYCLE_COUNTER(STAT_NiagaraRenderMeshes);
 
 	PrepareParticleRenderBuffers(ParticleMeshRenderData, Context.RayTracingMeshResourceCollector.GetDynamicReadBuffer());
@@ -1093,6 +1109,7 @@ void FNiagaraRendererMeshes::GetDynamicRayTracingInstances(FRayTracingMaterialGa
 		}
 
 		const FStaticMeshLODResources& LODModel = MeshData.RenderData->LODResources[LODIndex];
+		const FStaticMeshVertexFactories& VFs = MeshData.RenderData->LODVertexFactories[LODIndex];
 		FRayTracingGeometry& Geometry = MeshData.RenderData->LODResources[LODIndex].RayTracingGeometry;
 		FRayTracingInstance RayTracingInstance;
 		RayTracingInstance.Geometry = &Geometry;
@@ -1121,6 +1138,8 @@ void FNiagaraRendererMeshes::GetDynamicRayTracingInstances(FRayTracingMaterialGa
 			continue;
 		}
 
+		PreparePerMeshData(ParticleMeshRenderData, *SceneProxy, MeshData);
+
 		// Sort/Cull particles if needed.
 		NiagaraEmitterInstanceBatcher* Batcher = SceneProxy->GetBatcher();
 		const uint32 NumInstances = PerformSortAndCull(ParticleMeshRenderData, Context.RayTracingMeshResourceCollector.GetDynamicReadBuffer(), SortInfo, Batcher, MeshData.SourceMeshIndex);
@@ -1155,8 +1174,21 @@ void FNiagaraRendererMeshes::GetDynamicRayTracingInstances(FRayTracingMaterialGa
 			}
 
 			FMeshBatch MeshBatch;
-			const uint32 GPUCountBufferOffset = SortInfo.CulledGPUParticleCountOffset != INDEX_NONE ? SortInfo.CulledGPUParticleCountOffset : ParticleMeshRenderData.SourceParticleData->GetGPUInstanceCountBufferOffset();
-			CreateMeshBatchForSection(MeshBatch, VertexFactory, *MaterialProxy, *SceneProxy, MeshData, LODModel, Section, *View, ViewIndex, NumInstances, GPUCountBufferOffset, false, bIsInstancedStereo, ParticleMeshRenderData.bNeedsCull && ParticleMeshRenderData.bSortCullOnGpu);
+			CreateMeshBatchForSection(MeshBatch, VertexFactory, *MaterialProxy, *SceneProxy, MeshData, LODModel, Section, *View, ViewIndex, NumInstances, INDEX_NONE, false, bIsInstancedStereo, ParticleMeshRenderData.bNeedsCull && ParticleMeshRenderData.bSortCullOnGpu);
+			MeshBatch.SegmentIndex = SectionIndex;
+			MeshBatch.LODIndex = LODIndex;
+
+#if !(UE_BUILD_SHIPPING || UE_BUILD_TEST)
+			MeshBatch.VisualizeLODIndex = LODIndex;
+#endif
+
+			MeshBatch.bCanApplyViewModeOverrides = false;
+
+			MeshBatch.Elements[0].VertexFactoryUserData = VFs.VertexFactory.GetUniformBuffer();
+#if !(UE_BUILD_SHIPPING || UE_BUILD_TEST)
+			MeshBatch.Elements[0].VisualizeElementIndex = SectionIndex;
+#endif
+
 			RayTracingInstance.Materials.Add(MeshBatch);
 		}
 
@@ -1284,51 +1316,30 @@ void FNiagaraRendererMeshes::GetDynamicRayTracingInstances(FRayTracingMaterialGa
 				{
 					FMatrix InstanceTransform(FMatrix::Identity);
 
-					if (SimTarget == ENiagaraSimTarget::CPUSim)
+					if (bHasPosition)
 					{
-						InstanceTransform = FMatrix::Identity;
-
-						if (bHasPosition)
-						{
-							FVector InstancePos(GetInstancePosition(InstanceIndex));
-							InstanceTransform.SetOrigin(InstanceTransform.GetOrigin() + InstancePos);
-						}
-
-						if (bHasRotation)
-						{
-							FQuat InstanceQuat = GetInstanceQuat(InstanceIndex);
-							FTransform RotationTransform(InstanceQuat.GetNormalized());
-							FMatrix RotationMatrix = RotationTransform.ToMatrixWithScale();
-							InstanceTransform = RotationMatrix * InstanceTransform;
-						}
-
-						if (bHasScale)
-						{
-							FMatrix ScaleTransform = FMatrix::Identity;
-
-							FVector InstanceSca(GetInstanceScale(InstanceIndex));
-							ScaleTransform.M[0][0] *= InstanceSca.X;
-							ScaleTransform.M[1][1] *= InstanceSca.Y;
-							ScaleTransform.M[2][2] *= InstanceSca.Z;
+						FVector InstancePos(GetInstancePosition(InstanceIndex));
+						InstanceTransform.SetOrigin(InstanceTransform.GetOrigin() + InstancePos);
+					}
 
-							InstanceTransform = ScaleTransform * InstanceTransform;
-						}
+					if (bHasRotation)
+					{
+						FQuat InstanceQuat = GetInstanceQuat(InstanceIndex);
+						FTransform RotationTransform(InstanceQuat.GetNormalized());
+						FMatrix RotationMatrix = RotationTransform.ToMatrixWithScale();
+						InstanceTransform = RotationMatrix * InstanceTransform;
+					}
 
-						if (bUseLocalSpace)
-						{
-							InstanceTransform = InstanceTransform * LocalTransform;
-						}
+					FVector InstanceScale = MeshData.Scale;
+					if (bHasScale)
+					{
+						InstanceScale *= GetInstanceScale(InstanceIndex);
 					}
-					else
+					InstanceTransform = FScaleMatrix::Make(InstanceScale) * InstanceTransform;
+
+					if (bLocalSpace)
 					{
-						// Indirect instancing dispatching: transforms are not available at this point but computed in GPU instead
-						// Set invalid transforms so ray tracing ignores them. Valid transforms will be set later directly in the GPU
-						FMatrix ScaleTransform = FMatrix::Identity;
-						ScaleTransform.M[0][0] = 0.0;
-						ScaleTransform.M[1][1] = 0.0;
-						ScaleTransform.M[2][2] = 0.0;
-
-						InstanceTransform = ScaleTransform * InstanceTransform;
+						InstanceTransform = InstanceTransform * LocalTransform;
 					}
 
 					RayTracingInstance.InstanceTransforms.Add(InstanceTransform);
@@ -1339,13 +1350,11 @@ void FNiagaraRendererMeshes::GetDynamicRayTracingInstances(FRayTracingMaterialGa
 			{
 				FRHICommandListImmediate& RHICmdList = Context.RHICmdList;
 
-				const uint32 CPUInstancesCount = ParticleMeshRenderData.SourceParticleData->GetNumInstances();
-
-				RayTracingInstance.NumTransforms = CPUInstancesCount;
+				RayTracingInstance.NumTransforms = NumInstances;
 
 				FRWBufferStructured InstanceGPUTransformsBuffer;
-				//InstanceGPUTransformsBuffer.Initialize(sizeof(FMatrix), CPUInstancesCount, BUF_Static);
-					InstanceGPUTransformsBuffer.Initialize(TEXT("InstanceGPUTransformsBuffer"), 3 * 4 * sizeof(float), CPUInstancesCount, BUF_Static);
+				//InstanceGPUTransformsBuffer.Initialize(sizeof(FMatrix), NumInstances, BUF_Static);
+				InstanceGPUTransformsBuffer.Initialize(TEXT("InstanceGPUTransformsBuffer"), 3 * 4 * sizeof(float), NumInstances, BUF_Static);
 				RayTracingInstance.InstanceGPUTransformsSRV = InstanceGPUTransformsBuffer.SRV;
 
 				FNiagaraGPURayTracingTransformsCS::FPermutationDomain PermutationVector;
@@ -1359,19 +1368,20 @@ void FNiagaraRendererMeshes::GetDynamicRayTracingInstances(FRayTracingMaterialGa
 					VFVariables[ENiagaraMeshVFLayout::Scale].GetGPUOffset(),
 					bUseLocalSpace ? 1 : 0);
 
+				const uint32 GPUCountBufferOffset = SortInfo.CulledGPUParticleCountOffset != INDEX_NONE ? SortInfo.CulledGPUParticleCountOffset : ParticleMeshRenderData.SourceParticleData->GetGPUInstanceCountBufferOffset();
 				GPURayTracingTransformsCS->SetParameters(
 					RHICmdList,
-					CPUInstancesCount,
+					NumInstances,
 					ParticleMeshRenderData.ParticleFloatSRV,
 					0,	//FloatDataOffset
-					ParticleMeshRenderData.ParticleFloatDataStride / sizeof(float),
-					ParticleMeshRenderData.SourceParticleData->GetGPUInstanceCountBufferOffset(),
+					ParticleMeshRenderData.ParticleFloatDataStride,
+					GPUCountBufferOffset,
 					Batcher->GetGPUInstanceCounterManager().GetInstanceCountBuffer().SRV,
 					NiagaraOffsets,
 					LocalTransform,
 					InstanceGPUTransformsBuffer.UAV);
 
-				uint32 NGroups = FMath::DivideAndRoundUp(CPUInstancesCount, FNiagaraGPURayTracingTransformsCS::ThreadGroupSize);
+				uint32 NGroups = FMath::DivideAndRoundUp(NumInstances, FNiagaraGPURayTracingTransformsCS::ThreadGroupSize);
 				DispatchComputeShader(RHICmdList, GPURayTracingTransformsCS, NGroups, 1, 1);
 				GPURayTracingTransformsCS->UnbindBuffers(RHICmdList);
 
@@ -1415,12 +1425,10 @@ FNiagaraDynamicDataBase* FNiagaraRendererMeshes::GenerateDynamicData(const FNiag
 		return nullptr;
 	}
 
-
 	// Bail if we have cached mesh render data for any meshes that are no longer valid
 	for (const auto& MeshData : Meshes)
 	{
-		if (!Properties->Meshes.IsValidIndex(MeshData.SourceMeshIndex) || 
-			!Properties->Meshes[MeshData.SourceMeshIndex].HasValidMeshProperties())
+		if (!Properties->Meshes.IsValidIndex(MeshData.SourceMeshIndex) || !Properties->Meshes[MeshData.SourceMeshIndex].Mesh)
 		{
 			return nullptr;
 		}
diff --git a/Engine/Plugins/FX/Niagara/Source/Niagara/Public/NiagaraRendererMeshes.h b/Engine/Plugins/FX/Niagara/Source/Niagara/Public/NiagaraRendererMeshes.h
--- a/Engine/Plugins/FX/Niagara/Source/Niagara/Public/NiagaraRendererMeshes.h
+++ b/Engine/Plugins/FX/Niagara/Source/Niagara/Public/NiagaraRendererMeshes.h
@@ -106,6 +106,7 @@ protected:
 	void PrepareParticleMeshRenderData(FParticleMeshRenderData& ParticleMeshRenderData, FNiagaraDynamicDataBase* InDynamicData, const FNiagaraSceneProxy* SceneProxy) const;
 	void PrepareParticleRenderBuffers(FParticleMeshRenderData& ParticleMeshRenderData, FGlobalDynamicReadBuffer& DynamicReadBuffer) const;
 	void InitializeSortInfo(FParticleMeshRenderData& ParticleMeshRenderData, const FNiagaraSceneProxy& SceneProxy, const FSceneView& View, int32 ViewIndex, bool bIsInstancedStereo, FNiagaraGPUSortInfo& OutSortInfo) const;
+	void PreparePerMeshData(FParticleMeshRenderData& ParticleMeshRenderData, const FNiagaraSceneProxy& SceneProxy, const FMeshData& MeshData) const;
 	uint32 PerformSortAndCull(FParticleMeshRenderData& ParticleMeshRenderData, FGlobalDynamicReadBuffer& ReadBuffer, FNiagaraGPUSortInfo& SortInfo, class NiagaraEmitterInstanceBatcher* Batcher, int32 MeshIndex) const;
 	FNiagaraMeshUniformBufferRef CreatePerViewUniformBuffer(FParticleMeshRenderData& ParticleMeshRenderData, const FSceneView& View, const FMeshData& MeshData, const FNiagaraSceneProxy& SceneProxy) const;
 
