LWC: Niagara uses large world coordinate float types only.
Submitted on behalf of shaun.kime.

diff --git a/Engine/Plugins/FX/Niagara/Source/Niagara/Classes/NiagaraDataInterface.h b/Engine/Plugins/FX/Niagara/Source/Niagara/Classes/NiagaraDataInterface.h
--- a/Engine/Plugins/FX/Niagara/Source/Niagara/Classes/NiagaraDataInterface.h
+++ b/Engine/Plugins/FX/Niagara/Source/Niagara/Classes/NiagaraDataInterface.h
@@ -717,6 +717,8 @@ struct FNDIInputParam<FVector3f>
 	FORCEINLINE FVector3f GetAndAdvance() { return FVector3f(X.GetAndAdvance(), Y.GetAndAdvance(), Z.GetAndAdvance()); }
 };
 
+
+
 template<>
 struct FNDIInputParam<FVector4>
 {
@@ -802,14 +804,14 @@ struct FNDIOutputParam<FVector2D>
 };
 
 template<>
-struct FNDIOutputParam<FVector>
+struct FNDIOutputParam<FVector3f>
 {
 	VectorVM::FExternalFuncRegisterHandler<float> X;
 	VectorVM::FExternalFuncRegisterHandler<float> Y;
 	VectorVM::FExternalFuncRegisterHandler<float> Z;
 	FNDIOutputParam(FVectorVMContext& Context) : X(Context), Y(Context), Z(Context) {}
 	FORCEINLINE bool IsValid() const { return X.IsValid() || Y.IsValid() || Z.IsValid(); }
-	FORCEINLINE void SetAndAdvance(FVector Val)
+	FORCEINLINE void SetAndAdvance(FVector3f Val)
 	{
 		*X.GetDestAndAdvance() = Val.X;
 		*Y.GetDestAndAdvance() = Val.Y;
@@ -817,6 +819,7 @@ struct FNDIOutputParam<FVector>
 	}
 };
 
+
 template<>
 struct FNDIOutputParam<FVector4>
 {
@@ -854,7 +857,7 @@ struct FNDIOutputParam<FQuat>
 };
 
 template<>
-struct FNDIOutputParam<FMatrix>
+struct FNDIOutputParam<FMatrix44f>
 {
 	VectorVM::FExternalFuncRegisterHandler<float> Out00;
 	VectorVM::FExternalFuncRegisterHandler<float> Out01;
@@ -876,7 +879,7 @@ struct FNDIOutputParam<FMatrix>
 	FORCEINLINE FNDIOutputParam(FVectorVMContext& Context) : Out00(Context), Out01(Context), Out02(Context), Out03(Context), Out04(Context), Out05(Context),
 		Out06(Context), Out07(Context), Out08(Context), Out09(Context), Out10(Context), Out11(Context), Out12(Context), Out13(Context), Out14(Context), Out15(Context)	{}
 	FORCEINLINE bool IsValid() const { return Out00.IsValid(); }
-	FORCEINLINE void SetAndAdvance(const FMatrix& Val)
+	FORCEINLINE void SetAndAdvance(const FMatrix44f& Val)
 	{
 		*Out00.GetDestAndAdvance() = Val.M[0][0];
 		*Out01.GetDestAndAdvance() = Val.M[0][1];
diff --git a/Engine/Plugins/FX/Niagara/Source/Niagara/Classes/NiagaraDataInterfaceArrayFloat.h b/Engine/Plugins/FX/Niagara/Source/Niagara/Classes/NiagaraDataInterfaceArrayFloat.h
--- a/Engine/Plugins/FX/Niagara/Source/Niagara/Classes/NiagaraDataInterfaceArrayFloat.h
+++ b/Engine/Plugins/FX/Niagara/Source/Niagara/Classes/NiagaraDataInterfaceArrayFloat.h
@@ -33,7 +33,7 @@ class NIAGARA_API UNiagaraDataInterfaceArrayFloat3 : public UNiagaraDataInterfac
 	GENERATED_UCLASS_BODY()
 public:
 	UPROPERTY(EditAnywhere, BlueprintReadWrite, Category = "Array")
-	TArray<FVector> FloatData;
+	TArray<FVector> FloatData;		// LWC_TODO: Should be FVector3f, but only FVector is blueprint accessible, so we copy it to FloatDataAsFloat before passing to the interface
 
 	TArray<FVector>& GetArrayReference() { return FloatData; }
 };
diff --git a/Engine/Plugins/FX/Niagara/Source/Niagara/Classes/NiagaraDataInterfaceArrayImpl.h b/Engine/Plugins/FX/Niagara/Source/Niagara/Classes/NiagaraDataInterfaceArrayImpl.h
--- a/Engine/Plugins/FX/Niagara/Source/Niagara/Classes/NiagaraDataInterfaceArrayImpl.h
+++ b/Engine/Plugins/FX/Niagara/Source/Niagara/Classes/NiagaraDataInterfaceArrayImpl.h
@@ -8,7 +8,6 @@
 template<typename TArrayType>
 struct FNDIArrayImplHelperBase
 {
-	typedef TArrayType TVMArrayType;
 
 	static constexpr bool bSupportsCPU = true;
 	static constexpr bool bSupportsGPU = true;
@@ -21,6 +20,11 @@ struct FNDIArrayImplHelperBase
 
 	static void GPUGetFetchHLSL(FString& OutHLSL, const TCHAR* BufferName) { OutHLSL.Appendf(TEXT("OutValue = %s[ClampedIndex];"), BufferName); }
 	static int32 GPUGetTypeStride() { return sizeof(TArrayType); }
+	static void CopyData(void* Dest, const TArrayType* Src, int32 BufferSize)
+	{
+		FMemory::Memcpy(Dest, Src, BufferSize);
+	}
+
 };
 
 template<typename TArrayType>
@@ -424,7 +428,7 @@ struct FNiagaraDataInterfaceArrayImpl : public INiagaraDataInterfaceArrayImpl
 
 					RT_Proxy->Buffer.Initialize(TEXT("NiagaraArrayFloat"), BufferStride, BufferNumElements, FNDIArrayImplHelper<TArrayType>::PixelFormat, BUF_Static);
 					void* GPUMemory = RHICmdList.LockBuffer(RT_Proxy->Buffer.Buffer, 0, BufferSize, RLM_WriteOnly);
-					FMemory::Memcpy(GPUMemory, RT_Array.GetData(), BufferSize);
+					T::CopyData(GPUMemory, RT_Array.GetData(), BufferSize);
 					RHICmdList.UnlockBuffer(RT_Proxy->Buffer.Buffer);
 				}
 				else
@@ -438,7 +442,7 @@ struct FNiagaraDataInterfaceArrayImpl : public INiagaraDataInterfaceArrayImpl
 
 					RT_Proxy->Buffer.Initialize(TEXT("NiagaraArrayFloat"), BufferStride, BufferNumElements, FNDIArrayImplHelper<TArrayType>::PixelFormat, BUF_Static);
 					void* GPUMemory = RHICmdList.LockBuffer(RT_Proxy->Buffer.Buffer, 0, BufferSize, RLM_WriteOnly);
-					FMemory::Memcpy(GPUMemory, &DefaultValue, BufferSize);
+					T::CopyData(GPUMemory, &DefaultValue, BufferSize);
 					RHICmdList.UnlockBuffer(RT_Proxy->Buffer.Buffer);
 				}
 
@@ -658,4 +662,4 @@ struct FNiagaraDataInterfaceArrayImpl : public INiagaraDataInterfaceArrayImpl
 			}
 		}
 	}
-};
+};
\ No newline at end of file
diff --git a/Engine/Plugins/FX/Niagara/Source/Niagara/Classes/NiagaraDataInterfaceSkeletalMesh.h b/Engine/Plugins/FX/Niagara/Source/Niagara/Classes/NiagaraDataInterfaceSkeletalMesh.h
--- a/Engine/Plugins/FX/Niagara/Source/Niagara/Classes/NiagaraDataInterfaceSkeletalMesh.h
+++ b/Engine/Plugins/FX/Niagara/Source/Niagara/Classes/NiagaraDataInterfaceSkeletalMesh.h
@@ -270,8 +270,8 @@ struct FSkeletalMeshUvMappingHandle
 	FSkeletalMeshUvMappingUsage Usage;
 
 	void FindOverlappingTriangles(const FVector2D& InUv, float Tolerance, TArray<int32>& TriangleIndices) const;
-	int32 FindFirstTriangle(const FVector2D& InUv, float Tolerance, FVector& BarycentricCoord) const;
-	int32 FindFirstTriangle(const FBox2D& InUvBox, FVector& BarycentricCoord) const;
+	int32 FindFirstTriangle(const FVector2D& InUv, float Tolerance, FVector3f& BarycentricCoord) const;
+	int32 FindFirstTriangle(const FBox2D& InUvBox, FVector3f& BarycentricCoord) const;
 	const FSkeletalMeshUvMappingBufferProxy* GetQuadTreeProxy() const;
 	int32 GetUvSetIndex() const;
 	int32 GetLodIndex() const;
diff --git a/Engine/Plugins/FX/Niagara/Source/Niagara/Classes/NiagaraDataSetAccessor.h b/Engine/Plugins/FX/Niagara/Source/Niagara/Classes/NiagaraDataSetAccessor.h
--- a/Engine/Plugins/FX/Niagara/Source/Niagara/Classes/NiagaraDataSetAccessor.h
+++ b/Engine/Plugins/FX/Niagara/Source/Niagara/Classes/NiagaraDataSetAccessor.h
@@ -242,9 +242,9 @@ struct FNiagaraDataSetAccessorTypeInfo<FVector2D>
 };
 
 template<>
-struct FNiagaraDataSetAccessorTypeInfo<FVector>
+struct FNiagaraDataSetAccessorTypeInfo<FVector3f>
 {
-	using TAccessorBaseClass = FNiagaraDataSetAccessorFloat<FVector>;
+	using TAccessorBaseClass = FNiagaraDataSetAccessorFloat<FVector3f>;
 
 	static constexpr bool bSupportsHalf = true;
 	static constexpr int32 NumElements = 3;
@@ -252,6 +252,8 @@ struct FNiagaraDataSetAccessorTypeInfo<FVector>
 	static const FNiagaraTypeDefinition& GetHalfType() { return FNiagaraTypeDefinition::GetHalfVec3Def(); }
 };
 
+
+
 template<>
 struct FNiagaraDataSetAccessorTypeInfo<FVector4>
 {
diff --git a/Engine/Plugins/FX/Niagara/Source/Niagara/Classes/NiagaraScript.h b/Engine/Plugins/FX/Niagara/Source/Niagara/Classes/NiagaraScript.h
--- a/Engine/Plugins/FX/Niagara/Source/Niagara/Classes/NiagaraScript.h
+++ b/Engine/Plugins/FX/Niagara/Source/Niagara/Classes/NiagaraScript.h
@@ -895,6 +895,11 @@ public:
 		return TOptional<T>();
 	}
 
+	TOptional<FVector3d> GetCompilerTag(const FNiagaraVariableBase& InVar, const FNiagaraParameterStore* FallbackParameterStore) const
+	{
+		return TOptional<FVector3d>(*GetCompilerTag<FVector3f>(InVar, FallbackParameterStore));
+	}
+
 #if WITH_EDITORONLY_DATA
 	NIAGARA_API void InvalidateCompileResults(const FString& Reason);
 	NIAGARA_API FText GetDescription(const FGuid& VersionGuid);
diff --git a/Engine/Plugins/FX/Niagara/Source/Niagara/Private/DataInterface/NiagaraDataInterfaceDebugDraw.cpp b/Engine/Plugins/FX/Niagara/Source/Niagara/Private/DataInterface/NiagaraDataInterfaceDebugDraw.cpp
--- a/Engine/Plugins/FX/Niagara/Source/Niagara/Private/DataInterface/NiagaraDataInterfaceDebugDraw.cpp
+++ b/Engine/Plugins/FX/Niagara/Source/Niagara/Private/DataInterface/NiagaraDataInterfaceDebugDraw.cpp
@@ -39,7 +39,7 @@ struct FNDIDebugDrawInstanceData_GameThread
 
 
 #if NIAGARA_COMPUTEDEBUG_ENABLED
-	void AddLine(const FVector& Start, const FVector& End, const FLinearColor& Color)
+	void AddLine(const FVector3f& Start, const FVector3f& End, const FLinearColor& Color)
 	{
 		//-OPT: Need to improve this
 		FScopeLock RWLock(&LineBufferLock);
@@ -53,21 +53,21 @@ struct FNDIDebugDrawInstanceData_GameThread
 		Line.Color |= uint32(FMath::Clamp(Color.A, 0.0f, 1.0f) * 255.0f) << 0;
 	}
 
-	void AddSphere(const FVector& Location, float Radius, int32 Segments, const FLinearColor& Color)
+	void AddSphere(const FVector3f& Location, float Radius, int32 Segments, const FLinearColor& Color)
 	{
-		AddSphere(Location, FVector(Radius), FQuat::Identity, FQuat::Identity, FVector::OneVector, Segments, true, true, true, Color);
+		AddSphere(Location, FVector3f(Radius), FQuat::Identity, FQuat::Identity, FVector3f::OneVector, Segments, true, true, true, Color);
 	}
 
-	void AddSphere(const FVector& Location, const FVector& Radius, const FQuat& AxisRotator, const FQuat& WorldRotate, const FVector& WorldScale, int32 Segments, bool bHemiX, bool bHemiY, bool bHemiZ, const FLinearColor& Color)
+	void AddSphere(const FVector3f& Location, const FVector3f& Radius, const FQuat& AxisRotator, const FQuat& WorldRotate, const FVector3f& WorldScale, int32 Segments, bool bHemiX, bool bHemiY, bool bHemiZ, const FLinearColor& Color)
 	{
 		const float uinc = 2.0f * PI / float(Segments);
 
 		float ux = 0.0f;
 		float SinX0 = FMath::Sin(ux);
 		float CosX0 = FMath::Cos(ux);
-		FVector LerpVector = FVector(bHemiX ? 1.0f : 0.0f, bHemiY ? 1.0f : 0.0f, bHemiZ ? 1.0f : 0.0f);
+		FVector3f LerpVector = FVector3f(bHemiX ? 1.0f : 0.0f, bHemiY ? 1.0f : 0.0f, bHemiZ ? 1.0f : 0.0f);
 		FTransform LocalToWorld;
-		LocalToWorld.SetComponents(WorldRotate, FVector::ZeroVector, WorldScale);
+		LocalToWorld.SetComponents(WorldRotate, FVector3f::ZeroVector, WorldScale);
 
 		for (int x = 0; x < Segments; ++x)
 		{
@@ -85,9 +85,9 @@ struct FNDIDebugDrawInstanceData_GameThread
 				const float CosY1 = FMath::Cos(uy);
 
 				// Get uniform sphere location...
-				FVector A = FVector(CosX0 * CosY0, SinY0, SinX0 * CosY0);
-				FVector B = FVector(CosX1 * CosY0, SinY0, SinX1 * CosY0);
-				FVector C = FVector(CosX0 * CosY1, SinY1, SinX0 * CosY1);
+				FVector3f A = FVector3f(CosX0 * CosY0, SinY0, SinX0 * CosY0);
+				FVector3f B = FVector3f(CosX1 * CosY0, SinY0, SinX1 * CosY0);
+				FVector3f C = FVector3f(CosX0 * CosY1, SinY1, SinX0 * CosY1);
 
 				// Handle hemisphere mapping and scaling to radii/nonuniform... SphereVector = Lerp(SphereVector, Abs(SphereVector), (HemiX, HemiY, HemiZ)) * Radius
 				A = FMath::Lerp(A, A.GetAbs(), LerpVector) * Radius;
@@ -99,9 +99,9 @@ struct FNDIDebugDrawInstanceData_GameThread
 				B = AxisRotator.RotateVector(B);
 				C = AxisRotator.RotateVector(C);
 
-				const FVector Point0 = Location + LocalToWorld.TransformVector(A);
-				const FVector Point1 = Location + LocalToWorld.TransformVector(B);
-				const FVector Point2 = Location + LocalToWorld.TransformVector(C);
+				const FVector3f Point0 = Location + LocalToWorld.TransformVector(A);
+				const FVector3f Point1 = Location + LocalToWorld.TransformVector(B);
+				const FVector3f Point2 = Location + LocalToWorld.TransformVector(C);
 
 				AddLine(Point0, Point1, Color);
 				AddLine(Point0, Point2, Color);
@@ -115,25 +115,25 @@ struct FNDIDebugDrawInstanceData_GameThread
 		}
 	}
 
-	void AddBox(const FVector& Location, const FQuat& Rotation, const FVector& Extents, const FLinearColor& Color)
+	void AddBox(const FVector3f& Location, const FQuat& Rotation, const FVector3f& Extents, const FLinearColor& Color)
 	{
-		AddBox(Location, Rotation, Extents, FQuat::Identity, FVector::OneVector, Color);
+		AddBox(Location, Rotation, Extents, FQuat::Identity, FVector3f::OneVector, Color);
 	}
 
-	void AddBox(const FVector& Location, const FQuat& Rotation, const FVector& Extents, const FQuat& WorldRotate, const FVector& WorldScale, const FLinearColor& Color)
+	void AddBox(const FVector3f& Location, const FQuat& Rotation, const FVector3f& Extents, const FQuat& WorldRotate, const FVector3f& WorldScale, const FLinearColor& Color)
 	{
 		FTransform LocalToWorld;
-		LocalToWorld.SetComponents(WorldRotate, FVector::ZeroVector, WorldScale);
-		const FVector Points[] =
-		{
-			Location + LocalToWorld.TransformVector(Rotation.RotateVector(FVector(Extents.X,  Extents.Y,  Extents.Z))),
-			Location + LocalToWorld.TransformVector(Rotation.RotateVector(FVector(-Extents.X,  Extents.Y,  Extents.Z))),
-			Location + LocalToWorld.TransformVector(Rotation.RotateVector(FVector(-Extents.X, -Extents.Y,  Extents.Z))),
-			Location + LocalToWorld.TransformVector(Rotation.RotateVector(FVector(Extents.X, -Extents.Y,  Extents.Z))),
-			Location + LocalToWorld.TransformVector(Rotation.RotateVector(FVector(Extents.X,  Extents.Y, -Extents.Z))),
-			Location + LocalToWorld.TransformVector(Rotation.RotateVector(FVector(-Extents.X,  Extents.Y, -Extents.Z))),
-			Location + LocalToWorld.TransformVector(Rotation.RotateVector(FVector(-Extents.X, -Extents.Y, -Extents.Z))),
-			Location + LocalToWorld.TransformVector(Rotation.RotateVector(FVector(Extents.X, -Extents.Y, -Extents.Z))),
+		LocalToWorld.SetComponents(WorldRotate, FVector3f::ZeroVector, WorldScale);
+		const FVector3f Points[] =
+		{
+			Location + LocalToWorld.TransformVector(Rotation.RotateVector(FVector3f(Extents.X,  Extents.Y,  Extents.Z))),
+			Location + LocalToWorld.TransformVector(Rotation.RotateVector(FVector3f(-Extents.X,  Extents.Y,  Extents.Z))),
+			Location + LocalToWorld.TransformVector(Rotation.RotateVector(FVector3f(-Extents.X, -Extents.Y,  Extents.Z))),
+			Location + LocalToWorld.TransformVector(Rotation.RotateVector(FVector3f(Extents.X, -Extents.Y,  Extents.Z))),
+			Location + LocalToWorld.TransformVector(Rotation.RotateVector(FVector3f(Extents.X,  Extents.Y, -Extents.Z))),
+			Location + LocalToWorld.TransformVector(Rotation.RotateVector(FVector3f(-Extents.X,  Extents.Y, -Extents.Z))),
+			Location + LocalToWorld.TransformVector(Rotation.RotateVector(FVector3f(-Extents.X, -Extents.Y, -Extents.Z))),
+			Location + LocalToWorld.TransformVector(Rotation.RotateVector(FVector3f(Extents.X, -Extents.Y, -Extents.Z))),
 		};
 		AddLine(Points[0], Points[1], Color);
 		AddLine(Points[1], Points[2], Color);
@@ -149,36 +149,36 @@ struct FNDIDebugDrawInstanceData_GameThread
 		AddLine(Points[3], Points[7], Color);
 	}
 
-	void AddCircle(const FVector& Location, const FVector& XAxis, const FVector& YAxis, float Scale, int32 Segments, const FLinearColor& Color)
+	void AddCircle(const FVector3f& Location, const FVector3f& XAxis, const FVector3f& YAxis, float Scale, int32 Segments, const FLinearColor& Color)
 	{
-		const FVector X = XAxis * Scale;
-		const FVector Y = YAxis * Scale;
+		const FVector3f X = XAxis * Scale;
+		const FVector3f Y = YAxis * Scale;
 
 		const float d = 2.0f * PI / float(Segments);
 		float u = 0.0f;
 
-		FVector LastPoint = Location + (X * FMath::Cos(u)) + (Y * FMath::Sin(u));
+		FVector3f LastPoint = Location + (X * FMath::Cos(u)) + (Y * FMath::Sin(u));
 
 		for (int32 x = 0; x < Segments; ++x)
 		{
 			u += d;
-			const FVector CurrPoint = Location + (X * FMath::Cos(u)) + (Y * FMath::Sin(u));
+			const FVector3f CurrPoint = Location + (X * FMath::Cos(u)) + (Y * FMath::Sin(u));
 			AddLine(LastPoint, CurrPoint, Color);
 			LastPoint = CurrPoint;
 		}
 	}
 
-	void AddRectangle(const FVector& Location, const FVector& XAxis, const FVector& YAxis, const FVector2D& Extents, const FIntPoint& Segments, const FLinearColor& Color, bool bUnbounded)
+	void AddRectangle(const FVector3f& Location, const FVector3f& XAxis, const FVector3f& YAxis, const FVector2D& Extents, const FIntPoint& Segments, const FLinearColor& Color, bool bUnbounded)
 	{
-		const FVector XAxisNorm = XAxis.GetSafeNormal();
-		const FVector YAxisNorm = YAxis.GetSafeNormal();
+		const FVector3f XAxisNorm = XAxis.GetSafeNormal();
+		const FVector3f YAxisNorm = YAxis.GetSafeNormal();
 
 		const FIntPoint NumSegments = FIntPoint(FMath::Clamp(Segments.X, bUnbounded ? 2 : 1, 16), FMath::Clamp(Segments.Y, bUnbounded ? 2 : 1, 16));
 
-		const FVector StartPosition = Location - ((XAxisNorm * Extents.X) + (YAxisNorm * Extents.Y));
+		const FVector3f StartPosition = Location - ((XAxisNorm * Extents.X) + (YAxisNorm * Extents.Y));
 		const FVector2D Step = (Extents * 2) / NumSegments;
-		const FVector StepX = XAxisNorm * Step.X;
-		const FVector StepY = YAxisNorm * Step.Y;
+		const FVector3f StepX = XAxisNorm * Step.X;
+		const FVector3f StepY = YAxisNorm * Step.Y;
 
 		const int32 LowerBound = bUnbounded ? 1 : 0;
 		const int32 UpperBoundX = bUnbounded ? FMath::Max(NumSegments.X - 1, 1) : NumSegments.X;
@@ -187,8 +187,8 @@ struct FNDIDebugDrawInstanceData_GameThread
 		// Add all Y axis lines
 		for (int32 X = LowerBound; X <= UpperBoundX; X++)
 		{
-			FVector LineStart = StartPosition + (StepX * X);
-			FVector LineEnd = LineStart + (StepY * NumSegments.Y);
+			FVector3f LineStart = StartPosition + (StepX * X);
+			FVector3f LineEnd = LineStart + (StepY * NumSegments.Y);
 
 			AddLine(LineStart, LineEnd, Color);
 		}
@@ -196,43 +196,43 @@ struct FNDIDebugDrawInstanceData_GameThread
 		// Add all X axis lines
 		for (int32 Y = LowerBound; Y <= UpperBoundY; Y++)
 		{
-			FVector LineStart = StartPosition + (StepY * Y);
-			FVector LineEnd = LineStart + (StepX * NumSegments.X);
+			FVector3f LineStart = StartPosition + (StepY * Y);
+			FVector3f LineEnd = LineStart + (StepX * NumSegments.X);
 
 			AddLine(LineStart, LineEnd, Color);
 		}
 	}
 
-	void AddCylinder(const FVector& Location, const FVector& Axis, float Height, float Radius, int32 NumHeightSegments, int32 NumRadiusSegments, const FLinearColor& Color)
+	void AddCylinder(const FVector3f& Location, const FVector3f& Axis, float Height, float Radius, int32 NumHeightSegments, int32 NumRadiusSegments, const FLinearColor& Color)
 	{
-		FQuat Rotation = FQuat::FindBetweenNormals(FVector::UpVector, Axis);
+		FQuat Rotation = FQuat::FindBetweenNormals(FVector3f::UpVector, Axis);
 
-		AddCylinder(Location, Rotation, Height, Radius, NumHeightSegments, NumRadiusSegments, Color, false, false, FVector::OneVector, true);
+		AddCylinder(Location, Rotation, Height, Radius, NumHeightSegments, NumRadiusSegments, Color, false, false, FVector3f::OneVector, true);
 	}
 
-	void AddCylinder(const FVector& Location, const FQuat& Rotation, float Height, float Radius, int32 NumHeightSegments, int32 NumRadiusSegments, const FLinearColor& Color, bool bHemiXDraw, bool bHemiYDraw, const FVector& NonUniformScale, bool bCenterVertically)
+	void AddCylinder(const FVector3f& Location, const FQuat& Rotation, float Height, float Radius, int32 NumHeightSegments, int32 NumRadiusSegments, const FLinearColor& Color, bool bHemiXDraw, bool bHemiYDraw, const FVector3f& NonUniformScale, bool bCenterVertically)
 	{
-		const FVector HeightVector = FVector::ZAxisVector * Height;
-		const FVector HalfHeightOffset = bCenterVertically? HeightVector * 0.5f : FVector::ZeroVector;
-		const FVector AxisStep = FVector::ZAxisVector * Height / NumHeightSegments;
+		const FVector3f HeightVector = FVector3f::ZAxisVector * Height;
+		const FVector3f HalfHeightOffset = bCenterVertically? HeightVector * 0.5f : FVector3f::ZeroVector;
+		const FVector3f AxisStep = FVector3f::ZAxisVector * Height / NumHeightSegments;
 
-		const FVector HemisphereMask = FVector(bHemiXDraw ? 1.0f : 0.0f, bHemiYDraw ? 1.0f : 0.0f, 0.0f);
+		const FVector3f HemisphereMask = FVector3f(bHemiXDraw ? 1.0f : 0.0f, bHemiYDraw ? 1.0f : 0.0f, 0.0f);
 
 		float CurrentRotation = 0.0f;
-		FVector LastPointNormal = (FVector::YAxisVector * FMath::Cos(CurrentRotation)) + (FVector::XAxisVector * FMath::Sin(CurrentRotation));
+		FVector3f LastPointNormal = (FVector3f::YAxisVector * FMath::Cos(CurrentRotation)) + (FVector3f::XAxisVector * FMath::Sin(CurrentRotation));
 
 		for (int32 Idx = 1; Idx <= NumRadiusSegments; Idx++)
 		{
 			CurrentRotation = 2.0f * PI / float(NumRadiusSegments) * Idx;
 
-			FVector CurrentPointNormal = (FVector::YAxisVector * FMath::Cos(CurrentRotation)) + (FVector::XAxisVector * FMath::Sin(CurrentRotation));
+			FVector3f CurrentPointNormal = (FVector3f::YAxisVector * FMath::Cos(CurrentRotation)) + (FVector3f::XAxisVector * FMath::Sin(CurrentRotation));
 			CurrentPointNormal = FMath::Lerp(CurrentPointNormal, CurrentPointNormal.GetAbs(), HemisphereMask);
 
-			const FVector LastBottom = LastPointNormal * Radius - HalfHeightOffset;
-			const FVector LastTop = LastPointNormal * Radius + HeightVector - HalfHeightOffset;
+			const FVector3f LastBottom = LastPointNormal * Radius - HalfHeightOffset;
+			const FVector3f LastTop = LastPointNormal * Radius + HeightVector - HalfHeightOffset;
 
-			const FVector CurrentBottom = CurrentPointNormal * Radius - HalfHeightOffset;
-			const FVector CurrentTop = CurrentPointNormal * Radius + HeightVector - HalfHeightOffset;
+			const FVector3f CurrentBottom = CurrentPointNormal * Radius - HalfHeightOffset;
+			const FVector3f CurrentTop = CurrentPointNormal * Radius + HeightVector - HalfHeightOffset;
 
 
 			// Add Height line
@@ -251,35 +251,35 @@ struct FNDIDebugDrawInstanceData_GameThread
 		}
 	}
 
-	void AddCone(const FVector& Location, const FVector& Axis, float Height, float RadiusTop, float RadiusBottom, int32 NumHeightSegments, int32 NumRadiusSegments, const FLinearColor& Color)
+	void AddCone(const FVector3f& Location, const FVector3f& Axis, float Height, float RadiusTop, float RadiusBottom, int32 NumHeightSegments, int32 NumRadiusSegments, const FLinearColor& Color)
 	{
-		FQuat Rotation = FQuat::FindBetweenNormals(FVector::UpVector, Axis);
+		FQuat Rotation = FQuat::FindBetweenNormals(FVector3f::UpVector, Axis);
 
-		AddCone(Location, Rotation, Height, RadiusTop, RadiusBottom, NumHeightSegments, NumRadiusSegments, Color, FVector::OneVector, true);
+		AddCone(Location, Rotation, Height, RadiusTop, RadiusBottom, NumHeightSegments, NumRadiusSegments, Color, FVector3f::OneVector, true);
 	}
 
-	void AddCone(const FVector& Location, const FQuat& Rotation, float Height, float RadiusTop, float RadiusBottom, int32 NumHeightSegments, int32 NumRadiusSegments, const FLinearColor& Color, const FVector& NonUniformScale, bool bCenterVertically)
+	void AddCone(const FVector3f& Location, const FQuat& Rotation, float Height, float RadiusTop, float RadiusBottom, int32 NumHeightSegments, int32 NumRadiusSegments, const FLinearColor& Color, const FVector3f& NonUniformScale, bool bCenterVertically)
 	{
-		const FVector HeightVector = FVector::ZAxisVector * Height;
-		const FVector HalfHeightOffset = bCenterVertically ? HeightVector * 0.5f : FVector::ZeroVector;
-		const FVector AxisStep = FVector::ZAxisVector * Height / NumHeightSegments;
+		const FVector3f HeightVector = FVector3f::ZAxisVector * Height;
+		const FVector3f HalfHeightOffset = bCenterVertically ? HeightVector * 0.5f : FVector3f::ZeroVector;
+		const FVector3f AxisStep = FVector3f::ZAxisVector * Height / NumHeightSegments;
 
 		const float HeightSegmentStepAlpha = 1.0 / NumHeightSegments;
 
 		float CurrentRotation = 0.0f;
-		FVector LastPointNormal = (FVector::YAxisVector * FMath::Cos(CurrentRotation)) + (FVector::XAxisVector * FMath::Sin(CurrentRotation));
+		FVector3f LastPointNormal = (FVector3f::YAxisVector * FMath::Cos(CurrentRotation)) + (FVector3f::XAxisVector * FMath::Sin(CurrentRotation));
 
 		for (int32 Idx = 1; Idx <= NumRadiusSegments; Idx++)
 		{
 			CurrentRotation = 2.0f * PI / float(NumRadiusSegments) * Idx;
 
-			FVector CurrentPointNormal = (FVector::YAxisVector * FMath::Cos(CurrentRotation)) + (FVector::XAxisVector * FMath::Sin(CurrentRotation));
+			FVector3f CurrentPointNormal = (FVector3f::YAxisVector * FMath::Cos(CurrentRotation)) + (FVector3f::XAxisVector * FMath::Sin(CurrentRotation));
 
-			const FVector LastBottom = LastPointNormal * RadiusBottom - HalfHeightOffset;
-			const FVector LastTop = LastPointNormal * RadiusTop + HeightVector - HalfHeightOffset;
+			const FVector3f LastBottom = LastPointNormal * RadiusBottom - HalfHeightOffset;
+			const FVector3f LastTop = LastPointNormal * RadiusTop + HeightVector - HalfHeightOffset;
 
-			const FVector CurrentBottom = CurrentPointNormal * RadiusBottom - HalfHeightOffset;
-			const FVector CurrentTop = CurrentPointNormal * RadiusTop + HeightVector - HalfHeightOffset;
+			const FVector3f CurrentBottom = CurrentPointNormal * RadiusBottom - HalfHeightOffset;
+			const FVector3f CurrentTop = CurrentPointNormal * RadiusTop + HeightVector - HalfHeightOffset;
 
 			// Add Height line
 			AddLine(Location + Rotation.RotateVector(CurrentBottom * NonUniformScale), Location + Rotation.RotateVector(CurrentTop * NonUniformScale), Color);
@@ -292,10 +292,10 @@ struct FNDIDebugDrawInstanceData_GameThread
 			{
 				float RingRadius = FMath::Lerp(RadiusBottom, RadiusTop, HeightSegmentStepAlpha * HeightIdx);
 
-				const FVector RingOffset = AxisStep * HeightIdx;
+				const FVector3f RingOffset = AxisStep * HeightIdx;
 
-				const FVector LastRingPosition = LastPointNormal * RingRadius + RingOffset - HalfHeightOffset;
-				const FVector CurrentRingPosition = CurrentPointNormal * RingRadius + RingOffset - HalfHeightOffset;
+				const FVector3f LastRingPosition = LastPointNormal * RingRadius + RingOffset - HalfHeightOffset;
+				const FVector3f CurrentRingPosition = CurrentPointNormal * RingRadius + RingOffset - HalfHeightOffset;
 
 				AddLine(Location + Rotation.RotateVector(LastRingPosition * NonUniformScale), Location + Rotation.RotateVector(CurrentRingPosition * NonUniformScale), Color);
 			}
@@ -305,37 +305,37 @@ struct FNDIDebugDrawInstanceData_GameThread
 	}
 
 
-	void AddTorus(const FVector& Location, const FVector& Axis, float MajorRadius, float MinorRadius, int32 MajorRadiusSegments, int32 MinorRadiusSegments, const FLinearColor& Color)
+	void AddTorus(const FVector3f& Location, const FVector3f& Axis, float MajorRadius, float MinorRadius, int32 MajorRadiusSegments, int32 MinorRadiusSegments, const FLinearColor& Color)
 	{
-		const FVector AxisNorm = Axis.GetSafeNormal();
+		const FVector3f AxisNorm = Axis.GetSafeNormal();
 
-		const FQuat AxisRotation = FQuat::FindBetweenNormals(FVector::UpVector, AxisNorm);
+		const FQuat AxisRotation = FQuat::FindBetweenNormals(FVector3f::UpVector, AxisNorm);
 
-		const FVector TangentX = AxisRotation.RotateVector(FVector::XAxisVector);
-		const FVector TangentY = AxisRotation.RotateVector(FVector::YAxisVector);
+		const FVector3f TangentX = AxisRotation.RotateVector(FVector3f::XAxisVector);
+		const FVector3f TangentY = AxisRotation.RotateVector(FVector3f::YAxisVector);
 
 		const float MajorRotationDelta = 2.0f * PI / float(MajorRadiusSegments);
 		const float MinorRotationDelta = 2.0f * PI / float(MinorRadiusSegments);
 
 		float CurrentMajorRotation = 0.0f;
-		FVector PreviousMajorNormal = (TangentX * FMath::Cos(CurrentMajorRotation)) + (TangentY * FMath::Sin(CurrentMajorRotation));
+		FVector3f PreviousMajorNormal = (TangentX * FMath::Cos(CurrentMajorRotation)) + (TangentY * FMath::Sin(CurrentMajorRotation));
 
 		for (int32 MajorIdx = 0; MajorIdx < MajorRadiusSegments; MajorIdx++)
 		{
 			CurrentMajorRotation += MajorRotationDelta;
-			const FVector CurrentMajorNormal = (TangentX * FMath::Cos(CurrentMajorRotation)) + (TangentY * FMath::Sin(CurrentMajorRotation));
+			const FVector3f CurrentMajorNormal = (TangentX * FMath::Cos(CurrentMajorRotation)) + (TangentY * FMath::Sin(CurrentMajorRotation));
 
 			float CurrentMinorRotation = 0.0f;
-			FVector PreviousMinorNormal = (CurrentMajorNormal * FMath::Cos(CurrentMinorRotation)) + (AxisNorm * FMath::Sin(CurrentMinorRotation));
+			FVector3f PreviousMinorNormal = (CurrentMajorNormal * FMath::Cos(CurrentMinorRotation)) + (AxisNorm * FMath::Sin(CurrentMinorRotation));
 
 			for (int32 MinorIdx = 0; MinorIdx < MinorRadiusSegments; MinorIdx++)
 			{
 				CurrentMinorRotation += MinorRotationDelta;
-				FVector CurrentMinorNormal = (CurrentMajorNormal * FMath::Cos(CurrentMinorRotation)) + (AxisNorm * FMath::Sin(CurrentMinorRotation));
+				FVector3f CurrentMinorNormal = (CurrentMajorNormal * FMath::Cos(CurrentMinorRotation)) + (AxisNorm * FMath::Sin(CurrentMinorRotation));
 
 				AddLine(Location + CurrentMajorNormal * MajorRadius + PreviousMinorNormal * MinorRadius, Location + CurrentMajorNormal * MajorRadius + CurrentMinorNormal * MinorRadius, Color);
 
-				FVector PreviousMajorMinorNormal = (PreviousMajorNormal * FMath::Cos(CurrentMinorRotation)) + (AxisNorm * FMath::Sin(CurrentMinorRotation));
+				FVector3f PreviousMajorMinorNormal = (PreviousMajorNormal * FMath::Cos(CurrentMinorRotation)) + (AxisNorm * FMath::Sin(CurrentMinorRotation));
 				AddLine(Location + PreviousMajorNormal * MajorRadius + PreviousMajorMinorNormal * MinorRadius, Location + CurrentMajorNormal * MajorRadius + CurrentMinorNormal * MinorRadius, Color);
 
 
@@ -348,56 +348,56 @@ struct FNDIDebugDrawInstanceData_GameThread
 
 	}
 
-	void AddCoordinateSystem(const FVector& Location, const FQuat& Rotation, float Scale)
+	void AddCoordinateSystem(const FVector3f& Location, const FQuat& Rotation, float Scale)
 	{
-		const FVector XAxis = Rotation.RotateVector(FVector(Scale, 0.0f, 0.0f));
-		const FVector YAxis = Rotation.RotateVector(FVector(0.0f, Scale, 0.0f));
-		const FVector ZAxis = Rotation.RotateVector(FVector(0.0f, 0.0f, Scale));
+		const FVector3f XAxis = Rotation.RotateVector(FVector3f(Scale, 0.0f, 0.0f));
+		const FVector3f YAxis = Rotation.RotateVector(FVector3f(0.0f, Scale, 0.0f));
+		const FVector3f ZAxis = Rotation.RotateVector(FVector3f(0.0f, 0.0f, Scale));
 
 		AddLine(Location, Location + XAxis, FLinearColor::Red);
 		AddLine(Location, Location + YAxis, FLinearColor::Green);
 		AddLine(Location, Location + ZAxis, FLinearColor::Blue);
 	}
 
-	void AddGrid2D(const FVector& Center, const FQuat& Rotation, const FVector2D& Extents, const FIntPoint& NumCells, const FLinearColor& Color)
+	void AddGrid2D(const FVector3f& Center, const FQuat& Rotation, const FVector2D& Extents, const FIntPoint& NumCells, const FLinearColor& Color)
 	{
-		const FVector Corner = Center - Rotation.RotateVector(FVector(Extents.X, Extents.Y, 0.0f));
-		const FVector XLength = Rotation.RotateVector(FVector(Extents.X * 2.0f, 0.0f, 0.0f));
-		const FVector YLength = Rotation.RotateVector(FVector(0.0f, Extents.Y * 2.0f, 0.0f));
-		const FVector XDelta = XLength / float(NumCells.X);
-		const FVector YDelta = YLength / float(NumCells.Y);
+		const FVector3f Corner = Center - Rotation.RotateVector(FVector3f(Extents.X, Extents.Y, 0.0f));
+		const FVector3f XLength = Rotation.RotateVector(FVector3f(Extents.X * 2.0f, 0.0f, 0.0f));
+		const FVector3f YLength = Rotation.RotateVector(FVector3f(0.0f, Extents.Y * 2.0f, 0.0f));
+		const FVector3f XDelta = XLength / float(NumCells.X);
+		const FVector3f YDelta = YLength / float(NumCells.Y);
 
 		for (int X = 0; X <= NumCells.X; ++X)
 		{
-			const FVector XOffset = XDelta * float(X);
+			const FVector3f XOffset = XDelta * float(X);
 			for (int Y = 0; Y <= NumCells.Y; ++Y)
 			{
-				const FVector YOffset = YDelta * float(Y);
+				const FVector3f YOffset = YDelta * float(Y);
 				AddLine(Corner + XOffset, Corner + XOffset + YLength, Color);
 				AddLine(Corner + YOffset, Corner + YOffset + XLength, Color);
 			}
 		}
 	}
 
-	void AddGrid3D(const FVector& Center, const FQuat& Rotation, const FVector& Extents, const FIntVector& NumCells, const FLinearColor& Color)
+	void AddGrid3D(const FVector3f& Center, const FQuat& Rotation, const FVector3f& Extents, const FIntVector& NumCells, const FLinearColor& Color)
 	{
-		const FVector Corner = Center - Rotation.RotateVector(Extents);
-		const FVector XLength = Rotation.RotateVector(FVector(Extents.X * 2.0f, 0.0f, 0.0f));
-		const FVector YLength = Rotation.RotateVector(FVector(0.0f, Extents.Y * 2.0f, 0.0f));
-		const FVector ZLength = Rotation.RotateVector(FVector(0.0f, 0.0f, Extents.Z * 2.0f));
-		const FVector XDelta = XLength / float(NumCells.X);
-		const FVector YDelta = YLength / float(NumCells.Y);
-		const FVector ZDelta = ZLength / float(NumCells.Z);
+		const FVector3f Corner = Center - Rotation.RotateVector(Extents);
+		const FVector3f XLength = Rotation.RotateVector(FVector3f(Extents.X * 2.0f, 0.0f, 0.0f));
+		const FVector3f YLength = Rotation.RotateVector(FVector3f(0.0f, Extents.Y * 2.0f, 0.0f));
+		const FVector3f ZLength = Rotation.RotateVector(FVector3f(0.0f, 0.0f, Extents.Z * 2.0f));
+		const FVector3f XDelta = XLength / float(NumCells.X);
+		const FVector3f YDelta = YLength / float(NumCells.Y);
+		const FVector3f ZDelta = ZLength / float(NumCells.Z);
 
 		for (int X = 0; X <= NumCells.X; ++X)
 		{
-			const FVector XOffset = XDelta * float(X);
+			const FVector3f XOffset = XDelta * float(X);
 			for (int Y = 0; Y <= NumCells.Y; ++Y)
 			{
-				const FVector YOffset = YDelta * float(Y);
+				const FVector3f YOffset = YDelta * float(Y);
 				for (int Z = 0; Z <= NumCells.Z; ++Z)
 				{
-					const FVector ZOffset = ZDelta * float(Z);
+					const FVector3f ZOffset = ZDelta * float(Z);
 
 					AddLine(Corner + ZOffset + XOffset, Corner + ZOffset + XOffset + YLength, Color);		// Z Slice: X -> Y
 					AddLine(Corner + ZOffset + YOffset, Corner + ZOffset + YOffset + XLength, Color);		// Z Slice: Y -> X
@@ -446,7 +446,7 @@ struct FNDIDebugDrawInstanceData_GameThread
 				return ENiagaraCoordinateSpace::Local;
 		}
 
-		void TransformVector(bool bVectorWasSet, FVector& Vector, const TOptional<ENiagaraCoordinateSpace>& SourceSpace, ENiagaraCoordinateSpace DestSpace, const FNiagaraSystemInstance* SystemInstance)
+		void TransformVector(bool bVectorWasSet, FVector3f& Vector, const TOptional<ENiagaraCoordinateSpace>& SourceSpace, ENiagaraCoordinateSpace DestSpace, const FNiagaraSystemInstance* SystemInstance)
 		{
 			ENiagaraCoordinateSpace SourceSpaceConcrete = GetConcreteSource(bVectorWasSet, SourceSpace);
 			ENiagaraCoordinateSpace DestSpaceConcrete = GetConcreteSource(true, DestSpace);
@@ -469,7 +469,7 @@ struct FNDIDebugDrawInstanceData_GameThread
 			}
 		}
 
-		void ScaleVector(bool bVectorWasSet, FVector& Vector, const TOptional<ENiagaraCoordinateSpace>& SourceSpace, ENiagaraCoordinateSpace DestSpace, const FNiagaraSystemInstance* SystemInstance)
+		void ScaleVector(bool bVectorWasSet, FVector3f& Vector, const TOptional<ENiagaraCoordinateSpace>& SourceSpace, ENiagaraCoordinateSpace DestSpace, const FNiagaraSystemInstance* SystemInstance)
 		{
 			ENiagaraCoordinateSpace SourceSpaceConcrete = GetConcreteSource(bVectorWasSet, SourceSpace);
 			ENiagaraCoordinateSpace DestSpaceConcrete = GetConcreteSource(true, DestSpace);
@@ -492,7 +492,7 @@ struct FNDIDebugDrawInstanceData_GameThread
 			}
 		}
 
-		void TransformPosition(bool bPointWasSet, FVector& Point, const TOptional<ENiagaraCoordinateSpace>& SourceSpace, ENiagaraCoordinateSpace DestSpace, const FNiagaraSystemInstance* SystemInstance)
+		void TransformPosition(bool bPointWasSet, FVector3f& Point, const TOptional<ENiagaraCoordinateSpace>& SourceSpace, ENiagaraCoordinateSpace DestSpace, const FNiagaraSystemInstance* SystemInstance)
 		{
 			ENiagaraCoordinateSpace SourceSpaceConcrete = GetConcreteSource(bPointWasSet, SourceSpace);
 			ENiagaraCoordinateSpace DestSpaceConcrete = GetConcreteSource(true, DestSpace);
@@ -598,8 +598,8 @@ struct FNDIDebugDrawInstanceData_GameThread
 			TOptional<ENiagaraCoordinateSpace> EndLocationCoordinateSpace = GetTag<ENiagaraCoordinateSpace>(SystemInstance, CoordTypeDef, 3);
 			TOptional<FLinearColor> Color = GetTag<FLinearColor>(SystemInstance, FNiagaraTypeDefinition::GetColorDef(), 4);
 
-			FVector DrawStartLocation = StartLocation.Get(FVector::ZeroVector);
-			FVector DrawEndLocation = EndLocation.Get(FVector::ZeroVector);
+			FVector3f DrawStartLocation = StartLocation.Get(FVector3f::ZeroVector);
+			FVector3f DrawEndLocation = EndLocation.Get(FVector3f::ZeroVector);
 			FLinearColor DrawColor = Color.Get(FLinearColor::Red);
 
 			TransformPosition(StartLocation.IsSet(), DrawStartLocation, StartLocationCoordinateSpace, ENiagaraCoordinateSpace::Simulation, SystemInstance);
@@ -631,11 +631,11 @@ struct FNDIDebugDrawInstanceData_GameThread
 			TOptional<FNiagaraBool> UnboundedPlane = GetTag<FNiagaraBool>(SystemInstance, FNiagaraTypeDefinition::GetBoolDef(), 11);
 			TOptional<FLinearColor> Color = GetTag<FLinearColor>(SystemInstance, FNiagaraTypeDefinition::GetColorDef(), 12);
 
-			FVector DrawCenter = Center.Get(FVector::ZeroVector);
-			FVector DrawOffset = Offset.Get(FVector::ZeroVector);
+			FVector3f DrawCenter = Center.Get(FVector3f::ZeroVector);
+			FVector3f DrawOffset = Offset.Get(FVector3f::ZeroVector);
 			FVector2D DrawExtents = Extents.Get(FVector2D(10));
-			FVector DrawXAxis = XAxis.Get(FVector::XAxisVector);
-			FVector DrawYAxis = YAxis.Get(FVector::YAxisVector);
+			FVector3f DrawXAxis = XAxis.Get(FVector3f::XAxisVector);
+			FVector3f DrawYAxis = YAxis.Get(FVector3f::YAxisVector);
 			int32 DrawNumXSegments = NumXSegments.Get(1);
 			int32 DrawNumYSegments = NumYSegments.Get(1);
 			bool bDrawUnbounded = UnboundedPlane.Get(FNiagaraBool(false)).GetValue();
@@ -675,11 +675,11 @@ struct FNDIDebugDrawInstanceData_GameThread
 			TOptional<int32> NumSegments = GetTag<int32>(SystemInstance, FNiagaraTypeDefinition::GetIntDef(), 8);
 			TOptional<FLinearColor> Color = GetTag<FLinearColor>(SystemInstance, FNiagaraTypeDefinition::GetColorDef(), 9);
 
-			FVector DrawCenter = Center.Get(FVector::ZeroVector);
-			FVector DrawOffset = Offset.Get(FVector::ZeroVector);
+			FVector3f DrawCenter = Center.Get(FVector3f::ZeroVector);
+			FVector3f DrawOffset = Offset.Get(FVector3f::ZeroVector);
 			float DrawRadius = Radius.Get(10);
-			FVector DrawXAxis = XAxis.Get(FVector::XAxisVector);
-			FVector DrawYAxis = YAxis.Get(FVector::YAxisVector);
+			FVector3f DrawXAxis = XAxis.Get(FVector3f::XAxisVector);
+			FVector3f DrawYAxis = YAxis.Get(FVector3f::YAxisVector);
 			int32 DrawNumSegments = NumSegments.Get(1);
 			FLinearColor DrawColor = Color.Get(FLinearColor::Green);
 
@@ -714,10 +714,10 @@ struct FNDIDebugDrawInstanceData_GameThread
 			TOptional<FLinearColor> Color = GetTag<FLinearColor>(SystemInstance, FNiagaraTypeDefinition::GetColorDef(), 9);
 
 
-			FVector DrawCenter = Center.Get(FVector::ZeroVector);
-			FVector DrawOffset = Offset.Get(FVector::ZeroVector);
-			FVector DrawExtents = Extents.Get(FVector(10.0f));
-			FVector DrawRotationAxis = RotationAxis.Get(FVector::ZAxisVector);
+			FVector3f DrawCenter = Center.Get(FVector3f::ZeroVector);
+			FVector3f DrawOffset = Offset.Get(FVector3f::ZeroVector);
+			FVector3f DrawExtents = Extents.Get(FVector3f(10.0f));
+			FVector3f DrawRotationAxis = RotationAxis.Get(FVector3f::ZAxisVector);
 			float DrawRotationNormalizedAngle = RotationNormalizedAngle.Get(0.0f);
 			FLinearColor DrawColor = Color.Get(FLinearColor::Green);
 			FQuat DrawRotation = FQuat::Identity;
@@ -733,7 +733,7 @@ struct FNDIDebugDrawInstanceData_GameThread
 
 			FQuat WorldRot = FQuat::Identity;
 			TransformQuat(true, WorldRot, RotationWorldSpace, ENiagaraCoordinateSpace::World, SystemInstance);
-			FVector WorldScale = FVector::OneVector;
+			FVector3f WorldScale = FVector3f::OneVector;
 			ScaleVector(true, WorldScale, RotationWorldSpace, ENiagaraCoordinateSpace::World, SystemInstance);
 
 			if ((HalfExtents.IsSet() && HalfExtents.GetValue().GetValue()) || !HalfExtents.IsSet())
@@ -743,7 +743,7 @@ struct FNDIDebugDrawInstanceData_GameThread
 
 			if (Extents.IsSet())
 			{
-				FVector DebugPoint = DrawCenter + DrawOffset;
+				FVector3f DebugPoint = DrawCenter + DrawOffset;
 				TransformPosition(true, DebugPoint, TOptional<ENiagaraCoordinateSpace>(ENiagaraCoordinateSpace::Simulation), ENiagaraCoordinateSpace::World, SystemInstance);
 
 				InstanceData->AddBox(DebugPoint, DrawRotation, DrawExtents, WorldRot, WorldScale, DrawColor);
@@ -772,36 +772,36 @@ struct FNDIDebugDrawInstanceData_GameThread
 			TOptional<int32> NumSegments = GetTag<int32>(SystemInstance, FNiagaraTypeDefinition::GetIntDef(), 13);
 			TOptional<FLinearColor> Color = GetTag<FLinearColor>(SystemInstance, FNiagaraTypeDefinition::GetColorDef(), 14);
 
-			FVector DrawCenter = Center.Get(FVector::ZeroVector);
-			FVector DrawOffset = Offset.Get(FVector::ZeroVector);
-			FVector DrawRadius = FVector(Radius.Get(1.0f));
+			FVector3f DrawCenter = Center.Get(FVector3f::ZeroVector);
+			FVector3f DrawOffset = Offset.Get(FVector3f::ZeroVector);
+			FVector3f DrawRadius = FVector3f(Radius.Get(1.0f));
 			FLinearColor DrawColor = Color.Get(FLinearColor::Green);
 			int32 DrawNumSegments = NumSegments.Get(36);
-			FVector DrawNonUniformScale = NonUniformScale.Get(FVector::OneVector);
+			FVector3f DrawNonUniformScale = NonUniformScale.Get(FVector3f::OneVector);
 
-			FVector DrawOrientationAxis = OrientationAxisForSphere.Get(FVector::XAxisVector);
+			FVector3f DrawOrientationAxis = OrientationAxisForSphere.Get(FVector3f::XAxisVector);
 			bool bHemiXDraw = HemiX.IsSet() ? HemiX.GetValue().GetValue() : true;
 			bool bHemiYDraw = HemiY.IsSet() ? HemiY.GetValue().GetValue() : true;
 			bool bHemiZDraw = HemiZ.IsSet() ? HemiZ.GetValue().GetValue() : true;
 
 
-			FVector DrawRadii = DrawNonUniformScale * DrawRadius;
+			FVector3f DrawRadii = DrawNonUniformScale * DrawRadius;
 
 
 			TransformPosition(Center.IsSet(), DrawCenter, CenterWorldSpace, ENiagaraCoordinateSpace::Simulation, SystemInstance);
 			TransformVector(Offset.IsSet(), DrawOffset, OffsetWorldSpace, ENiagaraCoordinateSpace::Simulation, SystemInstance);
 
-			FQuat OrientationAxis = FQuat::FindBetweenVectors(FVector(1.0f, 0.0f, 0.0f), DrawOrientationAxis) * AdditionalRotation.Get(FQuat::Identity);
+			FQuat OrientationAxis = FQuat::FindBetweenVectors(FVector3f(1.0f, 0.0f, 0.0f), DrawOrientationAxis) * AdditionalRotation.Get(FQuat::Identity);
 
 			FQuat WorldRot = FQuat::Identity;
 			TransformQuat(true, WorldRot, RadiusWorldSpace, ENiagaraCoordinateSpace::World, SystemInstance);
-			FVector WorldScale = FVector::OneVector;
+			FVector3f WorldScale = FVector3f::OneVector;
 			ScaleVector(true, WorldScale, RadiusWorldSpace, ENiagaraCoordinateSpace::World, SystemInstance);
 
 
 			if (Radius.IsSet())
 			{
-				FVector DebugPoint = DrawCenter + DrawOffset;
+				FVector3f DebugPoint = DrawCenter + DrawOffset;
 				TransformPosition(true, DebugPoint, TOptional<ENiagaraCoordinateSpace>(ENiagaraCoordinateSpace::Simulation), ENiagaraCoordinateSpace::World, SystemInstance);
 
 				InstanceData->AddSphere(DebugPoint, DrawRadii, OrientationAxis, WorldRot, WorldScale, DrawNumSegments, bHemiXDraw, bHemiYDraw, bHemiZDraw, DrawColor);
@@ -831,15 +831,15 @@ struct FNDIDebugDrawInstanceData_GameThread
 			TOptional<int32> NumHeightSegments = GetTag<int32>(SystemInstance, FNiagaraTypeDefinition::GetIntDef(), 13);
 			TOptional<FLinearColor> Color = GetTag<FLinearColor>(SystemInstance, FNiagaraTypeDefinition::GetColorDef(), 14);
 
-			FVector DrawCenter = Center.Get(FVector::ZeroVector);
-			FVector DrawOffset = Offset.Get(FVector::ZeroVector);
+			FVector3f DrawCenter = Center.Get(FVector3f::ZeroVector);
+			FVector3f DrawOffset = Offset.Get(FVector3f::ZeroVector);
 			float DrawRadius = Radius.Get(10.0f);
 			float DrawHeight = Height.Get(5.0f);
 			bool DrawHalfHeight = IsHalfHeight.Get(false);
 			bool bHemiXDraw = HemiX.Get(FNiagaraBool(false)).GetValue();
 			bool bHemiYDraw = HemiY.Get(FNiagaraBool(false)).GetValue();
 			ENiagaraOrientationAxis DrawOrientationAxisEnum = OrientationAxis.Get(ENiagaraOrientationAxis::ZAxis);
-			FVector DrawNonUniformScale = NonUniformScale.Get(FVector::OneVector);
+			FVector3f DrawNonUniformScale = NonUniformScale.Get(FVector3f::OneVector);
 			int32 DrawNumRadiusSegmentss = NumRadiusSegments.Get(16);
 			int32 DrawNumHeightSegments = NumHeightSegments.Get(1);
 			FLinearColor DrawColor = Color.Get(FLinearColor::Green);
@@ -884,24 +884,24 @@ struct FNDIDebugDrawInstanceData_GameThread
 			TOptional<int32> NumHeightSegments = GetTag<int32>(SystemInstance, FNiagaraTypeDefinition::GetIntDef(), 10);
 			TOptional<FLinearColor> Color = GetTag<FLinearColor>(SystemInstance, FNiagaraTypeDefinition::GetColorDef(), 11);
 
-			FVector DrawCenter = Center.Get(FVector::ZeroVector);
-			FVector DrawOffset = Offset.Get(FVector::ZeroVector);
+			FVector3f DrawCenter = Center.Get(FVector3f::ZeroVector);
+			FVector3f DrawOffset = Offset.Get(FVector3f::ZeroVector);
 			float DrawAngle = Angle.Get(25.0f);
 			float DrawLength = Length.Get(10.0f);
-			FVector DrawOrientationAxis = OrientationAxis.Get(FVector::ZAxisVector);;
+			FVector3f DrawOrientationAxis = OrientationAxis.Get(FVector3f::ZAxisVector);;
 			int32 DrawNumRadiusSegmentss = NumRadiusSegments.Get(16);
 			int32 DrawNumHeightSegments = NumHeightSegments.Get(1);
 			FLinearColor DrawColor = Color.Get(FLinearColor::Green);
 
-			FVector DrawNonUniformScale = NonUniformScale.Get(FVector::OneVector);
+			FVector3f DrawNonUniformScale = NonUniformScale.Get(FVector3f::OneVector);
 			TransformPosition(Center.IsSet(), DrawCenter, CenterWorldSpace, ENiagaraCoordinateSpace::World, SystemInstance);
 			TransformVector(Offset.IsSet(), DrawOffset, OffsetWorldSpace, ENiagaraCoordinateSpace::World, SystemInstance);
 
 
-			FQuat WorldRot = FQuat::FindBetweenNormals(FVector::ZAxisVector, DrawOrientationAxis);
+			FQuat WorldRot = FQuat::FindBetweenNormals(FVector3f::ZAxisVector, DrawOrientationAxis);
 			TransformQuat(true, WorldRot, OrientationAxisCoordinateSpace, ENiagaraCoordinateSpace::World, SystemInstance);
 
-			FVector HalfOffset = DrawOrientationAxis.GetSafeNormal() * DrawLength * 0.5f;
+			FVector3f HalfOffset = DrawOrientationAxis.GetSafeNormal() * DrawLength * 0.5f;
 			TransformVector(Offset.IsSet(), HalfOffset, OrientationAxisCoordinateSpace, ENiagaraCoordinateSpace::World, SystemInstance);
 
 
@@ -933,16 +933,16 @@ struct FNDIDebugDrawInstanceData_GameThread
 			TOptional<int32> MinorRadiusSegments = GetTag<int32>(SystemInstance, FNiagaraTypeDefinition::GetFloatDef(), 10);
 			TOptional<FLinearColor> Color = GetTag<FLinearColor>(SystemInstance, FNiagaraTypeDefinition::GetColorDef(), 11);
 
-			FVector DrawCenter = Center.Get(FVector::ZeroVector);
-			FVector DrawOffset = Offset.Get(FVector::ZeroVector);
+			FVector3f DrawCenter = Center.Get(FVector3f::ZeroVector);
+			FVector3f DrawOffset = Offset.Get(FVector3f::ZeroVector);
 			float DrawMajorRadius = MajorRadius.Get(25.0f);
 			float DrawMinorRadius = MinorRadius.Get(5.0f);
-			FVector DrawOrientationAxis = OrientationAxis.Get(FVector::ZAxisVector);;
+			FVector3f DrawOrientationAxis = OrientationAxis.Get(FVector3f::ZAxisVector);;
 			int32 DrawMajorRadiusSegments = MajorRadiusSegments.Get(16);
 			int32 DrawMinorRadiusSegments = MinorRadiusSegments.Get(8);
 			FLinearColor DrawColor = Color.Get(FLinearColor::Green);
 
-			FVector DrawNonUniformScale = NonUniformScale.Get(FVector::OneVector);
+			FVector3f DrawNonUniformScale = NonUniformScale.Get(FVector3f::OneVector);
 			TransformPosition(Center.IsSet(), DrawCenter, CenterWorldSpace, ENiagaraCoordinateSpace::World, SystemInstance);
 			TransformVector(Offset.IsSet(), DrawOffset, OffsetWorldSpace, ENiagaraCoordinateSpace::World, SystemInstance);
 
@@ -972,8 +972,8 @@ struct FNDIDebugDrawInstanceData_GameThread
 			TOptional<ENiagaraCoordinateSpace> RotationWorldSpace = GetTag<ENiagaraCoordinateSpace>(SystemInstance, CoordTypeDef, 5);
 			TOptional<float> Scale = GetTag<float>(SystemInstance, FNiagaraTypeDefinition::GetFloatDef(), 6);
 
-			FVector DrawCenter = Center.Get(FVector::ZeroVector);
-			FVector DrawOffset = Offset.Get(FVector::ZeroVector);
+			FVector3f DrawCenter = Center.Get(FVector3f::ZeroVector);
+			FVector3f DrawOffset = Offset.Get(FVector3f::ZeroVector);
 			FQuat DrawRotation = Rotation.Get(FQuat::Identity);
 			float DrawScale = Scale.Get(1.0f);
 
@@ -1006,8 +1006,8 @@ struct FNDIDebugDrawInstanceData_GameThread
 
 
 
-			FVector DrawCenter = Center.Get(FVector::ZeroVector);
-			FVector DrawOffset = Offset.Get(FVector::ZeroVector);
+			FVector3f DrawCenter = Center.Get(FVector3f::ZeroVector);
+			FVector3f DrawOffset = Offset.Get(FVector3f::ZeroVector);
 			FQuat DrawRotation = Rotation.Get(FQuat::Identity);
 			FVector2D DrawExtents = Extents.Get(FVector2D(10.f));
 			int32 DrawNumCellsX = NumCellsX.Get(10);
@@ -1043,10 +1043,10 @@ struct FNDIDebugDrawInstanceData_GameThread
 			TOptional<FLinearColor> Color = GetTag<FLinearColor>(SystemInstance, FNiagaraTypeDefinition::GetColorDef(), 10);
 
 
-			FVector DrawCenter = Center.Get(FVector::ZeroVector);
-			FVector DrawOffset = Offset.Get(FVector::ZeroVector);
+			FVector3f DrawCenter = Center.Get(FVector3f::ZeroVector);
+			FVector3f DrawOffset = Offset.Get(FVector3f::ZeroVector);
 			FQuat DrawRotation = Rotation.Get(FQuat::Identity);
-			FVector DrawExtents = Extents.Get(FVector(10.f));
+			FVector3f DrawExtents = Extents.Get(FVector3f(10.f));
 			int32 DrawNumCellsX = NumCellsX.Get(10);
 			int32 DrawNumCellsY = NumCellsY.Get(10);
 			int32 DrawNumCellsZ = NumCellsZ.Get(10);
@@ -1471,8 +1471,8 @@ namespace NDIDebugDrawLocal
 			{
 			}
 
-			FNDIInputParam<FVector> LineStartParam;
-			FNDIInputParam<FVector> LineEndParam;
+			FNDIInputParam<FVector3f> LineStartParam;
+			FNDIInputParam<FVector3f> LineEndParam;
 			FNDIInputParam<FLinearColor> ColorParam;
 		};
 
@@ -1488,9 +1488,9 @@ namespace NDIDebugDrawLocal
 
 			}
 
-			FNDIInputParam<FVector> StartLocationParam;
+			FNDIInputParam<FVector3f> StartLocationParam;
 			FNDIInputParam<ENiagaraCoordinateSpace> StartLocationWSParam;
-			FNDIInputParam<FVector> EndLocationParam;
+			FNDIInputParam<FVector3f> EndLocationParam;
 			FNDIInputParam<ENiagaraCoordinateSpace> EndLocationWSParam;
 			FNDIInputParam<FLinearColor> ColorParam;
 		};
@@ -1498,8 +1498,8 @@ namespace NDIDebugDrawLocal
 #if NIAGARA_COMPUTEDEBUG_ENABLED
 		static void Draw(FNDIDebugDrawInstanceData_GameThread* InstanceData, VMBindings& Bindings, bool bExecute)
 		{
-			const FVector LineStart = Bindings.LineStartParam.GetAndAdvance();
-			const FVector LineEnd = Bindings.LineEndParam.GetAndAdvance();
+			const FVector3f LineStart = Bindings.LineStartParam.GetAndAdvance();
+			const FVector3f LineEnd = Bindings.LineEndParam.GetAndAdvance();
 			const FLinearColor Color = Bindings.ColorParam.GetAndAdvance();
 
 			if (bExecute)
@@ -1526,9 +1526,9 @@ namespace NDIDebugDrawLocal
 			{
 			}
 
-			FNDIInputParam<FVector> LocationParam;
-			FNDIInputParam<FVector> XAxisParam;
-			FNDIInputParam<FVector> YAxisParam;
+			FNDIInputParam<FVector3f> LocationParam;
+			FNDIInputParam<FVector3f> XAxisParam;
+			FNDIInputParam<FVector3f> YAxisParam;
 			FNDIInputParam<FVector2D> ExtentsParam;
 			FNDIInputParam<int32> NumXSegments;
 			FNDIInputParam<int32> NumYSegments;
@@ -1556,14 +1556,14 @@ namespace NDIDebugDrawLocal
 
 			}
 
-			FNDIInputParam<FVector> LocationParam;
+			FNDIInputParam<FVector3f> LocationParam;
 			FNDIInputParam<ENiagaraCoordinateSpace> LocationWSParam;
-			FNDIInputParam<FVector> OffsetParam;
+			FNDIInputParam<FVector3f> OffsetParam;
 			FNDIInputParam<ENiagaraCoordinateSpace> OffsetWSParam;
 			FNDIInputParam<FVector2D> ExtentsParam;
 			FNDIInputParam<FNiagaraBool> HalfExtentsParam;
-			FNDIInputParam<FVector> XAxisParam;
-			FNDIInputParam<FVector> YAxisParam;
+			FNDIInputParam<FVector3f> XAxisParam;
+			FNDIInputParam<FVector3f> YAxisParam;
 			FNDIInputParam<ENiagaraCoordinateSpace> AxisCoordinateSpace;
 			FNDIInputParam<int32> NumXSegmentsParam;
 			FNDIInputParam<int32> NumYSegmentsParam;
@@ -1574,9 +1574,9 @@ namespace NDIDebugDrawLocal
 #if NIAGARA_COMPUTEDEBUG_ENABLED
 		static void Draw(FNDIDebugDrawInstanceData_GameThread* InstanceData, VMBindings& Bindings, bool bExecute)
 		{
-			const FVector Location = Bindings.LocationParam.GetAndAdvance();
-			const FVector XAxis = Bindings.XAxisParam.GetAndAdvance();
-			const FVector YAxis = Bindings.YAxisParam.GetAndAdvance();
+			const FVector3f Location = Bindings.LocationParam.GetAndAdvance();
+			const FVector3f XAxis = Bindings.XAxisParam.GetAndAdvance();
+			const FVector3f YAxis = Bindings.YAxisParam.GetAndAdvance();
 			const FVector2D Extents = Bindings.ExtentsParam.GetAndAdvance();
 			const FIntPoint NumSegments = FIntPoint(FMath::Clamp(Bindings.NumXSegments.GetAndAdvance(), 1, 16), FMath::Clamp(Bindings.NumYSegments.GetAndAdvance(), 1, 16));
 			const FLinearColor Color = Bindings.ColorParam.GetAndAdvance();
@@ -1605,9 +1605,9 @@ namespace NDIDebugDrawLocal
 			{
 			}
 
-			FNDIInputParam<FVector> LocationParam;
-			FNDIInputParam<FVector> XAxisParam;
-			FNDIInputParam<FVector> YAxisParam;
+			FNDIInputParam<FVector3f> LocationParam;
+			FNDIInputParam<FVector3f> XAxisParam;
+			FNDIInputParam<FVector3f> YAxisParam;
 			FNDIInputParam<float> ScaleParam;
 			FNDIInputParam<int32> SegmentsParam;
 			FNDIInputParam<FLinearColor> ColorParam;
@@ -1630,13 +1630,13 @@ namespace NDIDebugDrawLocal
 
 			}
 
-			FNDIInputParam<FVector> LocationParam;
+			FNDIInputParam<FVector3f> LocationParam;
 			FNDIInputParam<ENiagaraCoordinateSpace> LocationWSParam;
-			FNDIInputParam<FVector> OffsetParam;
+			FNDIInputParam<FVector3f> OffsetParam;
 			FNDIInputParam<ENiagaraCoordinateSpace> OffsetWSParam;
 			FNDIInputParam<float> RadiusParam;
-			FNDIInputParam<FVector> XAxisParam;
-			FNDIInputParam<FVector> YAxisParam;
+			FNDIInputParam<FVector3f> XAxisParam;
+			FNDIInputParam<FVector3f> YAxisParam;
 			FNDIInputParam<ENiagaraCoordinateSpace> AxisCoordinateSpace;
 			FNDIInputParam<int32> NumSegmentsParam;
 			FNDIInputParam<FLinearColor> ColorParam;
@@ -1645,9 +1645,9 @@ namespace NDIDebugDrawLocal
 #if NIAGARA_COMPUTEDEBUG_ENABLED
 		static void Draw(FNDIDebugDrawInstanceData_GameThread* InstanceData, VMBindings& Bindings, bool bExecute)
 		{
-			const FVector Location = Bindings.LocationParam.GetAndAdvance();
-			const FVector XAxis = Bindings.XAxisParam.GetAndAdvance();
-			const FVector YAxis = Bindings.YAxisParam.GetAndAdvance();
+			const FVector3f Location = Bindings.LocationParam.GetAndAdvance();
+			const FVector3f XAxis = Bindings.XAxisParam.GetAndAdvance();
+			const FVector3f YAxis = Bindings.YAxisParam.GetAndAdvance();
 			const float Scale = Bindings.ScaleParam.GetAndAdvance();
 			const int32 Segments = FMath::Clamp(Bindings.SegmentsParam.GetAndAdvance(), 4, 16);
 			const FLinearColor Color = Bindings.ColorParam.GetAndAdvance();
@@ -1672,9 +1672,9 @@ namespace NDIDebugDrawLocal
 			{
 			}
 
-			FNDIInputParam<FVector> LocationParam;
+			FNDIInputParam<FVector3f> LocationParam;
 			FNDIInputParam<FQuat> RotationParam;
-			FNDIInputParam<FVector> ExtentsParam;
+			FNDIInputParam<FVector3f> ExtentsParam;
 			FNDIInputParam<FLinearColor> ColorParam;
 		};
 
@@ -1695,13 +1695,13 @@ namespace NDIDebugDrawLocal
 
 			}
 
-			FNDIInputParam<FVector> LocationParam;
+			FNDIInputParam<FVector3f> LocationParam;
 			FNDIInputParam<ENiagaraCoordinateSpace> LocationWSParam;
-			FNDIInputParam<FVector> OffsetParam;
+			FNDIInputParam<FVector3f> OffsetParam;
 			FNDIInputParam<ENiagaraCoordinateSpace> OffsetWSParam;
-			FNDIInputParam<FVector> ExtentsParam;
+			FNDIInputParam<FVector3f> ExtentsParam;
 			FNDIInputParam<FNiagaraBool> HalfExtentsParam;
-			FNDIInputParam<FVector> RotationAxisParam;
+			FNDIInputParam<FVector3f> RotationAxisParam;
 			FNDIInputParam<float> RotationAngleParam;
 			FNDIInputParam<ENiagaraCoordinateSpace> RotationWSParam;
 			FNDIInputParam<FLinearColor> ColorParam;
@@ -1710,9 +1710,9 @@ namespace NDIDebugDrawLocal
 #if NIAGARA_COMPUTEDEBUG_ENABLED
 		static void Draw(FNDIDebugDrawInstanceData_GameThread* InstanceData, VMBindings& Bindings, bool bExecute)
 		{
-			const FVector Location = Bindings.LocationParam.GetAndAdvance();
+			const FVector3f Location = Bindings.LocationParam.GetAndAdvance();
 			const FQuat Rotation = Bindings.RotationParam.GetAndAdvance();
-			const FVector Extents = Bindings.ExtentsParam.GetAndAdvance();
+			const FVector3f Extents = Bindings.ExtentsParam.GetAndAdvance();
 			const FLinearColor Color = Bindings.ColorParam.GetAndAdvance();
 
 			if (bExecute)
@@ -1735,7 +1735,7 @@ namespace NDIDebugDrawLocal
 			{
 			}
 
-			FNDIInputParam<FVector> LocationParam;
+			FNDIInputParam<FVector3f> LocationParam;
 			FNDIInputParam<float> RadiusParam;
 			FNDIInputParam<int32> SegmentsParam;
 			FNDIInputParam<FLinearColor> ColorParam;
@@ -1762,19 +1762,19 @@ namespace NDIDebugDrawLocal
 			{
 			}
 
-			FNDIInputParam<FVector> CenterParam;
+			FNDIInputParam<FVector3f> CenterParam;
 			FNDIInputParam<ENiagaraCoordinateSpace> CenterWSParam;
-			FNDIInputParam<FVector> OffsetParam;
+			FNDIInputParam<FVector3f> OffsetParam;
 			FNDIInputParam<ENiagaraCoordinateSpace> OffsetWSParam;
 			FNDIInputParam<float> RadiusParam;
 			FNDIInputParam<ENiagaraCoordinateSpace> RadiusWSParam;
 			FNDIInputParam<FNiagaraBool> HemiXParam;
 			FNDIInputParam<FNiagaraBool> HemiYParam;
 			FNDIInputParam<FNiagaraBool> HemiZParam;
-			FNDIInputParam<FVector> OrientationAxisParam;
+			FNDIInputParam<FVector3f> OrientationAxisParam;
 			FNDIInputParam<ENiagaraCoordinateSpace> RotationWSParam;
 			FNDIInputParam<FQuat> AdditionalRotationParam;
-			FNDIInputParam<FVector> NonUniformScaleParam;
+			FNDIInputParam<FVector3f> NonUniformScaleParam;
 			FNDIInputParam<int32> SegmentsParam;
 			FNDIInputParam<FLinearColor> ColorParam;
 		};
@@ -1782,7 +1782,7 @@ namespace NDIDebugDrawLocal
 #if NIAGARA_COMPUTEDEBUG_ENABLED
 		static void Draw(FNDIDebugDrawInstanceData_GameThread* InstanceData, VMBindings& Bindings, bool bExecute)
 		{
-			const FVector Location = Bindings.LocationParam.GetAndAdvance();
+			const FVector3f Location = Bindings.LocationParam.GetAndAdvance();
 			const float Radius = Bindings.RadiusParam.GetAndAdvance();
 			const int32 Segments = FMath::Clamp(Bindings.SegmentsParam.GetAndAdvance(), 4, 16);
 			const FLinearColor Color = Bindings.ColorParam.GetAndAdvance();
@@ -1810,8 +1810,8 @@ namespace NDIDebugDrawLocal
 			{
 			}
 
-			FNDIInputParam<FVector> LocationParam;
-			FNDIInputParam<FVector> AxisParam;
+			FNDIInputParam<FVector3f> LocationParam;
+			FNDIInputParam<FVector3f> AxisParam;
 			FNDIInputParam<float> HeightParam;
 			FNDIInputParam<float> RadiusParam;
 			FNDIInputParam<int32> NumHeightSegmentsParam;
@@ -1841,9 +1841,9 @@ namespace NDIDebugDrawLocal
 
 			}
 
-			FNDIInputParam<FVector> LocationParam;
+			FNDIInputParam<FVector3f> LocationParam;
 			FNDIInputParam<ENiagaraCoordinateSpace> LocationWSParam;
-			FNDIInputParam<FVector> OffsetParam;
+			FNDIInputParam<FVector3f> OffsetParam;
 			FNDIInputParam<ENiagaraCoordinateSpace> OffsetWSParam;
 			FNDIInputParam<float> RadiusParam;
 			FNDIInputParam<float> HeightParam;
@@ -1852,7 +1852,7 @@ namespace NDIDebugDrawLocal
 			FNDIInputParam<FNiagaraBool> HemisphereYParam;
 			FNDIInputParam<ENiagaraOrientationAxis> OrientationAxisParam;
 			FNDIInputParam<ENiagaraCoordinateSpace> OrientationAxisCoordinateSpaceParam;
-			FNDIInputParam<FVector> NonUniformScaleParam;
+			FNDIInputParam<FVector3f> NonUniformScaleParam;
 			FNDIInputParam<int32> NumRadiusSegmentsParam;
 			FNDIInputParam<int32> NumHeightSegmentsParam;
 			FNDIInputParam<FLinearColor> ColorParam;
@@ -1862,8 +1862,8 @@ namespace NDIDebugDrawLocal
 #if NIAGARA_COMPUTEDEBUG_ENABLED
 		static void Draw(FNDIDebugDrawInstanceData_GameThread* InstanceData, VMBindings& Bindings, bool bExecute)
 		{
-			const FVector Location = Bindings.LocationParam.GetAndAdvance();
-			const FVector Axis = Bindings.AxisParam.GetAndAdvance();
+			const FVector3f Location = Bindings.LocationParam.GetAndAdvance();
+			const FVector3f Axis = Bindings.AxisParam.GetAndAdvance();
 			const float Height = Bindings.HeightParam.GetAndAdvance();
 			const float Radius = Bindings.RadiusParam.GetAndAdvance();
 			const int32 NumHeightSegments = Bindings.NumHeightSegmentsParam.GetAndAdvance();
@@ -1894,8 +1894,8 @@ namespace NDIDebugDrawLocal
 			{
 			}
 
-			FNDIInputParam<FVector> LocationParam;
-			FNDIInputParam<FVector> AxisParam;
+			FNDIInputParam<FVector3f> LocationParam;
+			FNDIInputParam<FVector3f> AxisParam;
 			FNDIInputParam<float> HeightParam;
 			FNDIInputParam<float> RadiusTopParam;
 			FNDIInputParam<float> RadiusBottomParam;
@@ -1923,15 +1923,15 @@ namespace NDIDebugDrawLocal
 
 			}
 
-			FNDIInputParam<FVector> LocationParam;
+			FNDIInputParam<FVector3f> LocationParam;
 			FNDIInputParam<ENiagaraCoordinateSpace> LocationWSParam;
-			FNDIInputParam<FVector> OffsetParam;
+			FNDIInputParam<FVector3f> OffsetParam;
 			FNDIInputParam<ENiagaraCoordinateSpace> OffsetWSParam;
 			FNDIInputParam<float> AngleParam;
 			FNDIInputParam<float> LengthParam;
-			FNDIInputParam<FVector> OrientationAxisParam;
+			FNDIInputParam<FVector3f> OrientationAxisParam;
 			FNDIInputParam<ENiagaraCoordinateSpace> OrientationAxisCoordinateSpaceParam;
-			FNDIInputParam<FVector> NonUniformScaleParam;
+			FNDIInputParam<FVector3f> NonUniformScaleParam;
 			FNDIInputParam<int32> NumRadiusSegmentsParam;
 			FNDIInputParam<int32> NumHeightSegmentsParam;
 			FNDIInputParam<FLinearColor> ColorParam;
@@ -1940,8 +1940,8 @@ namespace NDIDebugDrawLocal
 #if NIAGARA_COMPUTEDEBUG_ENABLED
 		static void Draw(FNDIDebugDrawInstanceData_GameThread* InstanceData, VMBindings& Bindings, bool bExecute)
 		{
-			const FVector Location = Bindings.LocationParam.GetAndAdvance();
-			const FVector Axis = Bindings.AxisParam.GetAndAdvance();
+			const FVector3f Location = Bindings.LocationParam.GetAndAdvance();
+			const FVector3f Axis = Bindings.AxisParam.GetAndAdvance();
 			const float Height = Bindings.HeightParam.GetAndAdvance();
 			const float RadiusTop = Bindings.RadiusTopParam.GetAndAdvance();
 			const float RadiusBottom = Bindings.RadiusBottomParam.GetAndAdvance();
@@ -1972,8 +1972,8 @@ namespace NDIDebugDrawLocal
 			{
 			}
 
-			FNDIInputParam<FVector> LocationParam;
-			FNDIInputParam<FVector> AxisParam;
+			FNDIInputParam<FVector3f> LocationParam;
+			FNDIInputParam<FVector3f> AxisParam;
 			FNDIInputParam<float> MajorRadiusParam;
 			FNDIInputParam<float> MinorRadiusParam;
 			FNDIInputParam<int32> MajorRadiusSegmentsParam;
@@ -2000,15 +2000,15 @@ namespace NDIDebugDrawLocal
 
 			}
 
-			FNDIInputParam<FVector> LocationParam;
+			FNDIInputParam<FVector3f> LocationParam;
 			FNDIInputParam<ENiagaraCoordinateSpace> LocationWSParam;
-			FNDIInputParam<FVector> OffsetParam;
+			FNDIInputParam<FVector3f> OffsetParam;
 			FNDIInputParam<ENiagaraCoordinateSpace> OffsetWSParam;
 			FNDIInputParam<float> MajorRadiusParam;
 			FNDIInputParam<float> MinorRadiusParam;
-			FNDIInputParam<FVector> OrientationAxisParam;
+			FNDIInputParam<FVector3f> OrientationAxisParam;
 			FNDIInputParam<ENiagaraCoordinateSpace> OrientationAxisCoordinateSpaceParam;
-			FNDIInputParam<FVector> NonUniformScaleParam;
+			FNDIInputParam<FVector3f> NonUniformScaleParam;
 			FNDIInputParam<int32> MajorRadiusSegments;
 			FNDIInputParam<int32> MinorRadiusSegments;
 			FNDIInputParam<FLinearColor> ColorParam;
@@ -2017,8 +2017,8 @@ namespace NDIDebugDrawLocal
 #if NIAGARA_COMPUTEDEBUG_ENABLED
 		static void Draw(FNDIDebugDrawInstanceData_GameThread* InstanceData, VMBindings& Bindings, bool bExecute)
 		{
-			const FVector Location = Bindings.LocationParam.GetAndAdvance();
-			const FVector Axis = Bindings.AxisParam.GetAndAdvance();
+			const FVector3f Location = Bindings.LocationParam.GetAndAdvance();
+			const FVector3f Axis = Bindings.AxisParam.GetAndAdvance();
 			const float MajorRadius = Bindings.MajorRadiusParam.GetAndAdvance();
 			const float MinorRadius = Bindings.MinorRadiusParam.GetAndAdvance();
 			const int32 MajorRadiusSegments = Bindings.MajorRadiusSegmentsParam.GetAndAdvance();
@@ -2044,7 +2044,7 @@ namespace NDIDebugDrawLocal
 			{
 			}
 
-			FNDIInputParam<FVector> LocationParam;
+			FNDIInputParam<FVector3f> LocationParam;
 			FNDIInputParam<FQuat> RotationParam;
 			FNDIInputParam<float> ScaleParam;
 		};		
@@ -2063,9 +2063,9 @@ namespace NDIDebugDrawLocal
 
 			}
 
-			FNDIInputParam<FVector> LocationParam;
+			FNDIInputParam<FVector3f> LocationParam;
 			FNDIInputParam<ENiagaraCoordinateSpace> LocationWSParam;
-			FNDIInputParam<FVector> OffsetParam;
+			FNDIInputParam<FVector3f> OffsetParam;
 			FNDIInputParam<ENiagaraCoordinateSpace> OffsetWSParam;
 			FNDIInputParam<FQuat> RotationParam;
 			FNDIInputParam<ENiagaraCoordinateSpace> RotationCoordinateSpace;
@@ -2076,7 +2076,7 @@ namespace NDIDebugDrawLocal
 #if NIAGARA_COMPUTEDEBUG_ENABLED
 		static void Draw(FNDIDebugDrawInstanceData_GameThread* InstanceData, VMBindings& Bindings, bool bExecute)
 		{
-			const FVector Location = Bindings.LocationParam.GetAndAdvance();
+			const FVector3f Location = Bindings.LocationParam.GetAndAdvance();
 			const FQuat Rotation = Bindings.RotationParam.GetAndAdvance();
 			const float Scale = Bindings.ScaleParam.GetAndAdvance();
 
@@ -2102,7 +2102,7 @@ namespace NDIDebugDrawLocal
 			{
 			}
 
-			FNDIInputParam<FVector> CenterParam;
+			FNDIInputParam<FVector3f> CenterParam;
 			FNDIInputParam<FQuat> RotationParam;
 			FNDIInputParam<FVector2D> ExtentsParam;
 			FNDIInputParam<int32> NumCellsXParam;
@@ -2127,9 +2127,9 @@ namespace NDIDebugDrawLocal
 
 			}
 
-			FNDIInputParam<FVector> LocationParam;
+			FNDIInputParam<FVector3f> LocationParam;
 			FNDIInputParam<ENiagaraCoordinateSpace> LocationWSParam;
-			FNDIInputParam<FVector> OffsetParam;
+			FNDIInputParam<FVector3f> OffsetParam;
 			FNDIInputParam<ENiagaraCoordinateSpace> OffsetWSParam;
 			FNDIInputParam<FQuat> RotationParam;
 			FNDIInputParam<ENiagaraCoordinateSpace> RotationCoordinateSpace;
@@ -2142,7 +2142,7 @@ namespace NDIDebugDrawLocal
 #if NIAGARA_COMPUTEDEBUG_ENABLED
 		static void Draw(FNDIDebugDrawInstanceData_GameThread* InstanceData, VMBindings& Bindings, bool bExecute)
 		{
-			const FVector Center = Bindings.CenterParam.GetAndAdvance();
+			const FVector3f Center = Bindings.CenterParam.GetAndAdvance();
 			const FQuat Rotation = Bindings.RotationParam.GetAndAdvance();
 			const FVector2D Extents = Bindings.ExtentsParam.GetAndAdvance();
 			const FIntPoint NumCells(Bindings.NumCellsXParam.GetAndAdvance(), Bindings.NumCellsYParam.GetAndAdvance());
@@ -2171,9 +2171,9 @@ namespace NDIDebugDrawLocal
 			{
 			}
 
-			FNDIInputParam<FVector> CenterParam;
+			FNDIInputParam<FVector3f> CenterParam;
 			FNDIInputParam<FQuat> RotationParam;
-			FNDIInputParam<FVector> ExtentsParam;
+			FNDIInputParam<FVector3f> ExtentsParam;
 			FNDIInputParam<int32> NumCellsXParam;
 			FNDIInputParam<int32> NumCellsYParam;
 			FNDIInputParam<int32> NumCellsZParam;
@@ -2198,13 +2198,13 @@ namespace NDIDebugDrawLocal
 
 			}
 
-			FNDIInputParam<FVector> LocationParam;
+			FNDIInputParam<FVector3f> LocationParam;
 			FNDIInputParam<ENiagaraCoordinateSpace> LocationWSParam;
-			FNDIInputParam<FVector> OffsetParam;
+			FNDIInputParam<FVector3f> OffsetParam;
 			FNDIInputParam<ENiagaraCoordinateSpace> OffsetWSParam;
 			FNDIInputParam<FQuat> RotationParam;
 			FNDIInputParam<ENiagaraCoordinateSpace> RotationCoordinateSpace;
-			FNDIInputParam<FVector> ExtentsParam;
+			FNDIInputParam<FVector3f> ExtentsParam;
 			FNDIInputParam<int32> NumCellsXParam;
 			FNDIInputParam<int32> NumCellsYParam;
 			FNDIInputParam<int32> NumCellsZParam;
@@ -2215,9 +2215,9 @@ namespace NDIDebugDrawLocal
 #if NIAGARA_COMPUTEDEBUG_ENABLED
 		static void Draw(FNDIDebugDrawInstanceData_GameThread* InstanceData, VMBindings& Bindings, bool bExecute)
 		{
-			const FVector Center = Bindings.CenterParam.GetAndAdvance();
+			const FVector3f Center = Bindings.CenterParam.GetAndAdvance();
 			const FQuat Rotation = Bindings.RotationParam.GetAndAdvance();
-			const FVector Extents = Bindings.ExtentsParam.GetAndAdvance();
+			const FVector3f Extents = Bindings.ExtentsParam.GetAndAdvance();
 			const FIntVector NumCells(Bindings.NumCellsXParam.GetAndAdvance(), Bindings.NumCellsYParam.GetAndAdvance(), Bindings.NumCellsZParam.GetAndAdvance());
 			const FLinearColor Color = Bindings.ColorParam.GetAndAdvance();
 
@@ -2414,8 +2414,8 @@ void UNiagaraDataInterfaceDebugDraw::GetFunctions(TArray<FNiagaraFunctionSignatu
 		Signature.Inputs.Add(FNiagaraVariable(FNiagaraTypeDefinition(GetClass()), TEXT("DebugDrawInterface")));
 		Signature.Inputs.Add_GetRef(FNiagaraVariable(FNiagaraTypeDefinition::GetBoolDef(), TEXT("Execute"))).SetValue(true);
 		Signature.Inputs.Add(FNiagaraVariable(FNiagaraTypeDefinition::GetVec3Def(), TEXT("Center")));
-		Signature.Inputs.Add_GetRef(FNiagaraVariable(FNiagaraTypeDefinition::GetVec3Def(), TEXT("XAxis"))).SetValue(FVector::XAxisVector);
-		Signature.Inputs.Add_GetRef(FNiagaraVariable(FNiagaraTypeDefinition::GetVec3Def(), TEXT("YAxis"))).SetValue(FVector::YAxisVector);
+		Signature.Inputs.Add_GetRef(FNiagaraVariable(FNiagaraTypeDefinition::GetVec3Def(), TEXT("XAxis"))).SetValue(FVector3f::XAxisVector);
+		Signature.Inputs.Add_GetRef(FNiagaraVariable(FNiagaraTypeDefinition::GetVec3Def(), TEXT("YAxis"))).SetValue(FVector3f::YAxisVector);
 		Signature.Inputs.Add_GetRef(FNiagaraVariable(FNiagaraTypeDefinition::GetVec2Def(), TEXT("Extents"))).SetValue(FVector2D(10.0f));
 		Signature.Inputs.Add_GetRef(FNiagaraVariable(FNiagaraTypeDefinition::GetIntDef(), TEXT("NumXSegments"))).SetValue(1);
 		Signature.Inputs.Add_GetRef(FNiagaraVariable(FNiagaraTypeDefinition::GetIntDef(), TEXT("NumYSegments"))).SetValue(1);
@@ -2430,8 +2430,8 @@ void UNiagaraDataInterfaceDebugDraw::GetFunctions(TArray<FNiagaraFunctionSignatu
 		Signature.Inputs.Add(FNiagaraVariable(FNiagaraTypeDefinition(GetClass()), TEXT("DebugDrawInterface")));
 		Signature.Inputs.Add_GetRef(FNiagaraVariable(FNiagaraTypeDefinition::GetBoolDef(), TEXT("Execute"))).SetValue(true);
 		Signature.Inputs.Add(FNiagaraVariable(FNiagaraTypeDefinition::GetVec3Def(), TEXT("Center")));
-		Signature.Inputs.Add_GetRef(FNiagaraVariable(FNiagaraTypeDefinition::GetVec3Def(), TEXT("XAxis"))).SetValue(FVector::XAxisVector);
-		Signature.Inputs.Add_GetRef(FNiagaraVariable(FNiagaraTypeDefinition::GetVec3Def(), TEXT("YAxis"))).SetValue(FVector::YAxisVector);
+		Signature.Inputs.Add_GetRef(FNiagaraVariable(FNiagaraTypeDefinition::GetVec3Def(), TEXT("XAxis"))).SetValue(FVector3f::XAxisVector);
+		Signature.Inputs.Add_GetRef(FNiagaraVariable(FNiagaraTypeDefinition::GetVec3Def(), TEXT("YAxis"))).SetValue(FVector3f::YAxisVector);
 		Signature.Inputs.Add_GetRef(FNiagaraVariable(FNiagaraTypeDefinition::GetFloatDef(), TEXT("Radius"))).SetValue(10.0f);
 		Signature.Inputs.Add_GetRef(FNiagaraVariable(FNiagaraTypeDefinition::GetIntDef(), TEXT("Num Segments"))).SetValue(6);
 		Signature.Inputs.Add(FNiagaraVariable(FNiagaraTypeDefinition::GetColorDef(), TEXT("Color")));
@@ -2445,7 +2445,7 @@ void UNiagaraDataInterfaceDebugDraw::GetFunctions(TArray<FNiagaraFunctionSignatu
 		Signature.Inputs.Add_GetRef(FNiagaraVariable(FNiagaraTypeDefinition::GetBoolDef(), TEXT("Execute"))).SetValue(true);
 		Signature.Inputs.Add(FNiagaraVariable(FNiagaraTypeDefinition::GetVec3Def(), TEXT("Center")));
 		Signature.Inputs.Add(FNiagaraVariable(FNiagaraTypeDefinition::GetQuatDef(), TEXT("Rotation")));
-		Signature.Inputs.Add_GetRef(FNiagaraVariable(FNiagaraTypeDefinition::GetVec3Def(), TEXT("Extents"))).SetValue(FVector(10.0f));
+		Signature.Inputs.Add_GetRef(FNiagaraVariable(FNiagaraTypeDefinition::GetVec3Def(), TEXT("Extents"))).SetValue(FVector3f(10.0f));
 		Signature.Inputs.Add(FNiagaraVariable(FNiagaraTypeDefinition::GetColorDef(), TEXT("Color")));
 	}
 
@@ -2468,7 +2468,7 @@ void UNiagaraDataInterfaceDebugDraw::GetFunctions(TArray<FNiagaraFunctionSignatu
 		Signature.Inputs.Add(FNiagaraVariable(FNiagaraTypeDefinition(GetClass()), TEXT("DebugDrawInterface")));
 		Signature.Inputs.Add_GetRef(FNiagaraVariable(FNiagaraTypeDefinition::GetBoolDef(), TEXT("Execute"))).SetValue(true);
 		Signature.Inputs.Add(FNiagaraVariable(FNiagaraTypeDefinition::GetVec3Def(), TEXT("Center")));
-		Signature.Inputs.Add_GetRef(FNiagaraVariable(FNiagaraTypeDefinition::GetVec3Def(), TEXT("Axis"))).SetValue(FVector::ZAxisVector);
+		Signature.Inputs.Add_GetRef(FNiagaraVariable(FNiagaraTypeDefinition::GetVec3Def(), TEXT("Axis"))).SetValue(FVector3f::ZAxisVector);
 		Signature.Inputs.Add_GetRef(FNiagaraVariable(FNiagaraTypeDefinition::GetFloatDef(), TEXT("Height"))).SetValue(10.0f);
 		Signature.Inputs.Add_GetRef(FNiagaraVariable(FNiagaraTypeDefinition::GetFloatDef(), TEXT("Radius"))).SetValue(10.0f);
 		Signature.Inputs.Add_GetRef(FNiagaraVariable(FNiagaraTypeDefinition::GetIntDef(), TEXT("NumHeightSegments"))).SetValue(1);
@@ -2483,7 +2483,7 @@ void UNiagaraDataInterfaceDebugDraw::GetFunctions(TArray<FNiagaraFunctionSignatu
 		Signature.Inputs.Add(FNiagaraVariable(FNiagaraTypeDefinition(GetClass()), TEXT("DebugDrawInterface")));
 		Signature.Inputs.Add_GetRef(FNiagaraVariable(FNiagaraTypeDefinition::GetBoolDef(), TEXT("Execute"))).SetValue(true);
 		Signature.Inputs.Add(FNiagaraVariable(FNiagaraTypeDefinition::GetVec3Def(), TEXT("Center")));
-		Signature.Inputs.Add_GetRef(FNiagaraVariable(FNiagaraTypeDefinition::GetVec3Def(), TEXT("Axis"))).SetValue(FVector::ZAxisVector);
+		Signature.Inputs.Add_GetRef(FNiagaraVariable(FNiagaraTypeDefinition::GetVec3Def(), TEXT("Axis"))).SetValue(FVector3f::ZAxisVector);
 		Signature.Inputs.Add_GetRef(FNiagaraVariable(FNiagaraTypeDefinition::GetFloatDef(), TEXT("Height"))).SetValue(10.0f);
 		Signature.Inputs.Add_GetRef(FNiagaraVariable(FNiagaraTypeDefinition::GetFloatDef(), TEXT("RadiusTop"))).SetValue(0.0f);
 		Signature.Inputs.Add_GetRef(FNiagaraVariable(FNiagaraTypeDefinition::GetFloatDef(), TEXT("RadiusBottom"))).SetValue(10.0f);
@@ -2499,7 +2499,7 @@ void UNiagaraDataInterfaceDebugDraw::GetFunctions(TArray<FNiagaraFunctionSignatu
 		Signature.Inputs.Add(FNiagaraVariable(FNiagaraTypeDefinition(GetClass()), TEXT("DebugDrawInterface")));
 		Signature.Inputs.Add_GetRef(FNiagaraVariable(FNiagaraTypeDefinition::GetBoolDef(), TEXT("Execute"))).SetValue(true);
 		Signature.Inputs.Add(FNiagaraVariable(FNiagaraTypeDefinition::GetVec3Def(), TEXT("Center")));
-		Signature.Inputs.Add_GetRef(FNiagaraVariable(FNiagaraTypeDefinition::GetVec3Def(), TEXT("Axis"))).SetValue(FVector::ZAxisVector);
+		Signature.Inputs.Add_GetRef(FNiagaraVariable(FNiagaraTypeDefinition::GetVec3Def(), TEXT("Axis"))).SetValue(FVector3f::ZAxisVector);
 		Signature.Inputs.Add_GetRef(FNiagaraVariable(FNiagaraTypeDefinition::GetFloatDef(), TEXT("MajorRadius"))).SetValue(100.0f);
 		Signature.Inputs.Add_GetRef(FNiagaraVariable(FNiagaraTypeDefinition::GetFloatDef(), TEXT("MinorRadius"))).SetValue(25.0f);
 		Signature.Inputs.Add_GetRef(FNiagaraVariable(FNiagaraTypeDefinition::GetIntDef(), TEXT("MajorRadiusSegments"))).SetValue(12);
@@ -2540,7 +2540,7 @@ void UNiagaraDataInterfaceDebugDraw::GetFunctions(TArray<FNiagaraFunctionSignatu
 		Signature.Inputs.Add_GetRef(FNiagaraVariable(FNiagaraTypeDefinition::GetBoolDef(), TEXT("Execute"))).SetValue(true);
 		Signature.Inputs.Add(FNiagaraVariable(FNiagaraTypeDefinition::GetVec3Def(), TEXT("Center")));
 		Signature.Inputs.Add(FNiagaraVariable(FNiagaraTypeDefinition::GetQuatDef(), TEXT("Rotation")));
-		Signature.Inputs.Add_GetRef(FNiagaraVariable(FNiagaraTypeDefinition::GetVec3Def(), TEXT("Extents"))).SetValue(FVector(10.0f));
+		Signature.Inputs.Add_GetRef(FNiagaraVariable(FNiagaraTypeDefinition::GetVec3Def(), TEXT("Extents"))).SetValue(FVector3f(10.0f));
 		Signature.Inputs.Add_GetRef(FNiagaraVariable(FNiagaraTypeDefinition::GetIntDef(), TEXT("NumCellsX"))).SetValue(1);
 		Signature.Inputs.Add_GetRef(FNiagaraVariable(FNiagaraTypeDefinition::GetIntDef(), TEXT("NumCellsY"))).SetValue(1);
 		Signature.Inputs.Add_GetRef(FNiagaraVariable(FNiagaraTypeDefinition::GetIntDef(), TEXT("NumCellsZ"))).SetValue(1);
@@ -2573,12 +2573,12 @@ void UNiagaraDataInterfaceDebugDraw::GetFunctions(TArray<FNiagaraFunctionSignatu
 		Signature.Inputs.Add(FNiagaraVariable(FNiagaraTypeDefinition(GetClass()), TEXT("DebugDrawInterface")));
 		Signature.Inputs.Add(FNiagaraVariable(FNiagaraTypeDefinition::GetVec3Def(), TEXT("Center")));
 		Signature.Inputs.Add_GetRef(FNiagaraVariable(CoordTypeDef, TEXT("CenterCoordinateSpace"))).SetValue(ENiagaraCoordinateSpace::Local);
-		Signature.Inputs.Add_GetRef(FNiagaraVariable(FNiagaraTypeDefinition::GetVec3Def(), TEXT("Offset"))).SetValue(FVector::ZeroVector);
+		Signature.Inputs.Add_GetRef(FNiagaraVariable(FNiagaraTypeDefinition::GetVec3Def(), TEXT("Offset"))).SetValue(FVector3f::ZeroVector);
 		Signature.Inputs.Add_GetRef(FNiagaraVariable(CoordTypeDef, TEXT("OffsetCoordinateSpace"))).SetValue(ENiagaraCoordinateSpace::Local);
 		Signature.Inputs.Add_GetRef(FNiagaraVariable(FNiagaraTypeDefinition::GetVec2Def(), TEXT("Extents"))).SetValue(FVector2D(5, 5));
 		Signature.Inputs.Add_GetRef(FNiagaraVariable(FNiagaraTypeDefinition::GetBoolDef(), TEXT("HalfExtents"))).SetValue(true);
-		Signature.Inputs.Add_GetRef(FNiagaraVariable(FNiagaraTypeDefinition::GetVec3Def(), TEXT("XAxis"))).SetValue(FVector::XAxisVector);
-		Signature.Inputs.Add_GetRef(FNiagaraVariable(FNiagaraTypeDefinition::GetVec3Def(), TEXT("YAxis"))).SetValue(FVector::YAxisVector);
+		Signature.Inputs.Add_GetRef(FNiagaraVariable(FNiagaraTypeDefinition::GetVec3Def(), TEXT("XAxis"))).SetValue(FVector3f::XAxisVector);
+		Signature.Inputs.Add_GetRef(FNiagaraVariable(FNiagaraTypeDefinition::GetVec3Def(), TEXT("YAxis"))).SetValue(FVector3f::YAxisVector);
 		Signature.Inputs.Add_GetRef(FNiagaraVariable(CoordTypeDef, TEXT("AxisVectorsCoordinateSpace"))).SetValue(ENiagaraCoordinateSpace::Local);
 		Signature.Inputs.Add_GetRef(FNiagaraVariable(FNiagaraTypeDefinition::GetIntDef(), TEXT("NumXSegments"))).SetValue(1);
 		Signature.Inputs.Add_GetRef(FNiagaraVariable(FNiagaraTypeDefinition::GetIntDef(), TEXT("NumYSegments"))).SetValue(1);
@@ -2595,11 +2595,11 @@ void UNiagaraDataInterfaceDebugDraw::GetFunctions(TArray<FNiagaraFunctionSignatu
 		Signature.Inputs.Add(FNiagaraVariable(FNiagaraTypeDefinition(GetClass()), TEXT("DebugDrawInterface")));
 		Signature.Inputs.Add(FNiagaraVariable(FNiagaraTypeDefinition::GetVec3Def(), TEXT("Center")));
 		Signature.Inputs.Add_GetRef(FNiagaraVariable(CoordTypeDef, TEXT("CenterCoordinateSpace"))).SetValue(ENiagaraCoordinateSpace::Local);
-		Signature.Inputs.Add_GetRef(FNiagaraVariable(FNiagaraTypeDefinition::GetVec3Def(), TEXT("Offset"))).SetValue(FVector::ZeroVector);
+		Signature.Inputs.Add_GetRef(FNiagaraVariable(FNiagaraTypeDefinition::GetVec3Def(), TEXT("Offset"))).SetValue(FVector3f::ZeroVector);
 		Signature.Inputs.Add_GetRef(FNiagaraVariable(CoordTypeDef, TEXT("OffsetCoordinateSpace"))).SetValue(ENiagaraCoordinateSpace::Local);
 		Signature.Inputs.Add_GetRef(FNiagaraVariable(FNiagaraTypeDefinition::GetFloatDef(), TEXT("Radius"))).SetValue(10.0f);
-		Signature.Inputs.Add_GetRef(FNiagaraVariable(FNiagaraTypeDefinition::GetVec3Def(), TEXT("XAxis"))).SetValue(FVector::XAxisVector);
-		Signature.Inputs.Add_GetRef(FNiagaraVariable(FNiagaraTypeDefinition::GetVec3Def(), TEXT("YAxis"))).SetValue(FVector::YAxisVector);
+		Signature.Inputs.Add_GetRef(FNiagaraVariable(FNiagaraTypeDefinition::GetVec3Def(), TEXT("XAxis"))).SetValue(FVector3f::XAxisVector);
+		Signature.Inputs.Add_GetRef(FNiagaraVariable(FNiagaraTypeDefinition::GetVec3Def(), TEXT("YAxis"))).SetValue(FVector3f::YAxisVector);
 		Signature.Inputs.Add_GetRef(FNiagaraVariable(CoordTypeDef, TEXT("AxisVectorsCoordinateSpace"))).SetValue(ENiagaraCoordinateSpace::Local);
 		Signature.Inputs.Add_GetRef(FNiagaraVariable(FNiagaraTypeDefinition::GetIntDef(), TEXT("NumSegments"))).SetValue(6);
 		Signature.Inputs.Add(FNiagaraVariable(FNiagaraTypeDefinition::GetColorDef(), TEXT("Color")));
@@ -2614,11 +2614,11 @@ void UNiagaraDataInterfaceDebugDraw::GetFunctions(TArray<FNiagaraFunctionSignatu
 		Signature.Inputs.Add(FNiagaraVariable(FNiagaraTypeDefinition(GetClass()), TEXT("DebugDrawInterface")));
 		Signature.Inputs.Add(FNiagaraVariable(FNiagaraTypeDefinition::GetVec3Def(), TEXT("Center")));
 		Signature.Inputs.Add_GetRef(FNiagaraVariable(CoordTypeDef, TEXT("CenterCoordinateSpace"))).SetValue(ENiagaraCoordinateSpace::Local);
-		Signature.Inputs.Add_GetRef(FNiagaraVariable(FNiagaraTypeDefinition::GetVec3Def(), TEXT("Offset"))).SetValue(FVector::ZeroVector);
+		Signature.Inputs.Add_GetRef(FNiagaraVariable(FNiagaraTypeDefinition::GetVec3Def(), TEXT("Offset"))).SetValue(FVector3f::ZeroVector);
 		Signature.Inputs.Add_GetRef(FNiagaraVariable(CoordTypeDef, TEXT("OffsetCoordinateSpace"))).SetValue(ENiagaraCoordinateSpace::Local);
-		Signature.Inputs.Add_GetRef(FNiagaraVariable(FNiagaraTypeDefinition::GetVec3Def(), TEXT("Extents"))).SetValue(FVector(10.0f));
+		Signature.Inputs.Add_GetRef(FNiagaraVariable(FNiagaraTypeDefinition::GetVec3Def(), TEXT("Extents"))).SetValue(FVector3f(10.0f));
 		Signature.Inputs.Add_GetRef(FNiagaraVariable(FNiagaraTypeDefinition::GetBoolDef(), TEXT("HalfExtents"))).SetValue(true);
-		Signature.Inputs.Add_GetRef(FNiagaraVariable(FNiagaraTypeDefinition::GetVec3Def(), TEXT("RotationAxis"))).SetValue(FVector::ZAxisVector);
+		Signature.Inputs.Add_GetRef(FNiagaraVariable(FNiagaraTypeDefinition::GetVec3Def(), TEXT("RotationAxis"))).SetValue(FVector3f::ZAxisVector);
 		Signature.Inputs.Add_GetRef(FNiagaraVariable(FNiagaraTypeDefinition::GetFloatDef(), TEXT("RotationNormalizedAngle"))).SetValue(0.0f);
 		Signature.Inputs.Add_GetRef(FNiagaraVariable(CoordTypeDef, TEXT("RotationCoordinateSpace"))).SetValue(ENiagaraCoordinateSpace::Local);
 		Signature.Inputs.Add(FNiagaraVariable(FNiagaraTypeDefinition::GetColorDef(), TEXT("Color")));
@@ -2633,17 +2633,17 @@ void UNiagaraDataInterfaceDebugDraw::GetFunctions(TArray<FNiagaraFunctionSignatu
 		Signature.Inputs.Add(FNiagaraVariable(FNiagaraTypeDefinition(GetClass()), TEXT("DebugDrawInterface")));
 		Signature.Inputs.Add(FNiagaraVariable(FNiagaraTypeDefinition::GetVec3Def(), TEXT("Center")));
 		Signature.Inputs.Add_GetRef(FNiagaraVariable(CoordTypeDef, TEXT("Center Coordinate Space"))).SetValue(ENiagaraCoordinateSpace::Simulation);
-		Signature.Inputs.Add_GetRef(FNiagaraVariable(FNiagaraTypeDefinition::GetVec3Def(), TEXT("Offset From Center"))).SetValue(FVector::ZeroVector);
+		Signature.Inputs.Add_GetRef(FNiagaraVariable(FNiagaraTypeDefinition::GetVec3Def(), TEXT("Offset From Center"))).SetValue(FVector3f::ZeroVector);
 		Signature.Inputs.Add_GetRef(FNiagaraVariable(CoordTypeDef, TEXT("Offset Coordinate Space"))).SetValue(ENiagaraCoordinateSpace::Local);
 		Signature.Inputs.Add_GetRef(FNiagaraVariable(FNiagaraTypeDefinition::GetFloatDef(), TEXT("Radius"))).SetValue(10.0f);
 		Signature.Inputs.Add_GetRef(FNiagaraVariable(CoordTypeDef, TEXT("Radius Coordinate Space"))).SetValue(ENiagaraCoordinateSpace::Local);
 		Signature.Inputs.Add_GetRef(FNiagaraVariable(FNiagaraTypeDefinition::GetBoolDef(), TEXT("Hemisphere X"))).SetValue(true);
 		Signature.Inputs.Add_GetRef(FNiagaraVariable(FNiagaraTypeDefinition::GetBoolDef(), TEXT("Hemisphere Y"))).SetValue(true);
 		Signature.Inputs.Add_GetRef(FNiagaraVariable(FNiagaraTypeDefinition::GetBoolDef(), TEXT("Hemisphere Z"))).SetValue(true);
-		Signature.Inputs.Add_GetRef(FNiagaraVariable(FNiagaraTypeDefinition::GetVec3Def(), TEXT("Sphere Orientation Axis"))).SetValue(FVector::ZAxisVector);
+		Signature.Inputs.Add_GetRef(FNiagaraVariable(FNiagaraTypeDefinition::GetVec3Def(), TEXT("Sphere Orientation Axis"))).SetValue(FVector3f::ZAxisVector);
 		Signature.Inputs.Add_GetRef(FNiagaraVariable(CoordTypeDef, TEXT("RotationCoordinateSpace"))).SetValue(ENiagaraCoordinateSpace::Local);
 		Signature.Inputs.Add_GetRef(FNiagaraVariable(FNiagaraTypeDefinition::GetQuatDef(), TEXT("Additional Rotation"))).SetValue(FQuat::Identity);
-		Signature.Inputs.Add_GetRef(FNiagaraVariable(FNiagaraTypeDefinition::GetVec3Def(), TEXT("NonUniform Scale"))).SetValue(FVector::OneVector);
+		Signature.Inputs.Add_GetRef(FNiagaraVariable(FNiagaraTypeDefinition::GetVec3Def(), TEXT("NonUniform Scale"))).SetValue(FVector3f::OneVector);
 		Signature.Inputs.Add_GetRef(FNiagaraVariable(FNiagaraTypeDefinition::GetIntDef(), TEXT("Num Segments"))).SetValue(12);
 		Signature.Inputs.Add(FNiagaraVariable(FNiagaraTypeDefinition::GetColorDef(), TEXT("Color")));
 	}
@@ -2657,7 +2657,7 @@ void UNiagaraDataInterfaceDebugDraw::GetFunctions(TArray<FNiagaraFunctionSignatu
 		Signature.Inputs.Add(FNiagaraVariable(FNiagaraTypeDefinition(GetClass()), TEXT("DebugDrawInterface")));
 		Signature.Inputs.Add(FNiagaraVariable(FNiagaraTypeDefinition::GetVec3Def(), TEXT("Center")));
 		Signature.Inputs.Add_GetRef(FNiagaraVariable(CoordTypeDef, TEXT("CenterCoordinateSpace"))).SetValue(ENiagaraCoordinateSpace::Local);
-		Signature.Inputs.Add_GetRef(FNiagaraVariable(FNiagaraTypeDefinition::GetVec3Def(), TEXT("Offset"))).SetValue(FVector::ZeroVector);
+		Signature.Inputs.Add_GetRef(FNiagaraVariable(FNiagaraTypeDefinition::GetVec3Def(), TEXT("Offset"))).SetValue(FVector3f::ZeroVector);
 		Signature.Inputs.Add_GetRef(FNiagaraVariable(CoordTypeDef, TEXT("OffsetCoordinateSpace"))).SetValue(ENiagaraCoordinateSpace::Local);
 		Signature.Inputs.Add_GetRef(FNiagaraVariable(FNiagaraTypeDefinition::GetFloatDef(), TEXT("Radius"))).SetValue(5.0f);
 		Signature.Inputs.Add_GetRef(FNiagaraVariable(FNiagaraTypeDefinition::GetFloatDef(), TEXT("Height"))).SetValue(10.0f);
@@ -2666,7 +2666,7 @@ void UNiagaraDataInterfaceDebugDraw::GetFunctions(TArray<FNiagaraFunctionSignatu
 		Signature.Inputs.Add_GetRef(FNiagaraVariable(FNiagaraTypeDefinition::GetBoolDef(), TEXT("Hemisphere Y"))).SetValue(true);
 		Signature.Inputs.Add_GetRef(FNiagaraVariable(AxisTypeDef, TEXT("Orientation Axis"))).SetValue(ENiagaraOrientationAxis::ZAxis);
 		Signature.Inputs.Add_GetRef(FNiagaraVariable(CoordTypeDef, TEXT("OrientationAxisCoordinateSpace"))).SetValue(ENiagaraCoordinateSpace::Local);
-		Signature.Inputs.Add_GetRef(FNiagaraVariable(FNiagaraTypeDefinition::GetVec3Def(), TEXT("NonUniform Scale"))).SetValue(FVector::OneVector);
+		Signature.Inputs.Add_GetRef(FNiagaraVariable(FNiagaraTypeDefinition::GetVec3Def(), TEXT("NonUniform Scale"))).SetValue(FVector3f::OneVector);
 		Signature.Inputs.Add_GetRef(FNiagaraVariable(FNiagaraTypeDefinition::GetIntDef(), TEXT("Num Radius Segments"))).SetValue(12);
 		Signature.Inputs.Add_GetRef(FNiagaraVariable(FNiagaraTypeDefinition::GetIntDef(), TEXT("Num Height Segments"))).SetValue(1);
 		Signature.Inputs.Add(FNiagaraVariable(FNiagaraTypeDefinition::GetColorDef(), TEXT("Color")));
@@ -2681,13 +2681,13 @@ void UNiagaraDataInterfaceDebugDraw::GetFunctions(TArray<FNiagaraFunctionSignatu
 		Signature.Inputs.Add(FNiagaraVariable(FNiagaraTypeDefinition(GetClass()), TEXT("DebugDrawInterface")));
 		Signature.Inputs.Add(FNiagaraVariable(FNiagaraTypeDefinition::GetVec3Def(), TEXT("Center")));
 		Signature.Inputs.Add_GetRef(FNiagaraVariable(CoordTypeDef, TEXT("CenterCoordinateSpace"))).SetValue(ENiagaraCoordinateSpace::Local);
-		Signature.Inputs.Add_GetRef(FNiagaraVariable(FNiagaraTypeDefinition::GetVec3Def(), TEXT("Offset"))).SetValue(FVector::ZeroVector);
+		Signature.Inputs.Add_GetRef(FNiagaraVariable(FNiagaraTypeDefinition::GetVec3Def(), TEXT("Offset"))).SetValue(FVector3f::ZeroVector);
 		Signature.Inputs.Add_GetRef(FNiagaraVariable(CoordTypeDef, TEXT("OffsetCoordinateSpace"))).SetValue(ENiagaraCoordinateSpace::Local);
 		Signature.Inputs.Add_GetRef(FNiagaraVariable(FNiagaraTypeDefinition::GetFloatDef(), TEXT("Angle"))).SetValue(10.0f);
 		Signature.Inputs.Add_GetRef(FNiagaraVariable(FNiagaraTypeDefinition::GetFloatDef(), TEXT("Length"))).SetValue(10.0f);
-		Signature.Inputs.Add_GetRef(FNiagaraVariable(FNiagaraTypeDefinition::GetVec3Def(), TEXT("Orientation Axis"))).SetValue(FVector::ZAxisVector);
+		Signature.Inputs.Add_GetRef(FNiagaraVariable(FNiagaraTypeDefinition::GetVec3Def(), TEXT("Orientation Axis"))).SetValue(FVector3f::ZAxisVector);
 		Signature.Inputs.Add_GetRef(FNiagaraVariable(CoordTypeDef, TEXT("OrientationAxisCoordinateSpace"))).SetValue(ENiagaraCoordinateSpace::Local);
-		Signature.Inputs.Add_GetRef(FNiagaraVariable(FNiagaraTypeDefinition::GetVec3Def(), TEXT("NonUniform Scale"))).SetValue(FVector::OneVector);
+		Signature.Inputs.Add_GetRef(FNiagaraVariable(FNiagaraTypeDefinition::GetVec3Def(), TEXT("NonUniform Scale"))).SetValue(FVector3f::OneVector);
 		Signature.Inputs.Add_GetRef(FNiagaraVariable(FNiagaraTypeDefinition::GetIntDef(), TEXT("Num Radius Segments"))).SetValue(12);
 		Signature.Inputs.Add_GetRef(FNiagaraVariable(FNiagaraTypeDefinition::GetIntDef(), TEXT("Num Height Segments"))).SetValue(1);
 		Signature.Inputs.Add(FNiagaraVariable(FNiagaraTypeDefinition::GetColorDef(), TEXT("Color")));
@@ -2702,13 +2702,13 @@ void UNiagaraDataInterfaceDebugDraw::GetFunctions(TArray<FNiagaraFunctionSignatu
 		Signature.Inputs.Add(FNiagaraVariable(FNiagaraTypeDefinition(GetClass()), TEXT("DebugDrawInterface")));
 		Signature.Inputs.Add(FNiagaraVariable(FNiagaraTypeDefinition::GetVec3Def(), TEXT("Center")));
 		Signature.Inputs.Add_GetRef(FNiagaraVariable(CoordTypeDef, TEXT("CenterCoordinateSpace"))).SetValue(ENiagaraCoordinateSpace::Local);
-		Signature.Inputs.Add_GetRef(FNiagaraVariable(FNiagaraTypeDefinition::GetVec3Def(), TEXT("Offset"))).SetValue(FVector::ZeroVector);
+		Signature.Inputs.Add_GetRef(FNiagaraVariable(FNiagaraTypeDefinition::GetVec3Def(), TEXT("Offset"))).SetValue(FVector3f::ZeroVector);
 		Signature.Inputs.Add_GetRef(FNiagaraVariable(CoordTypeDef, TEXT("OffsetCoordinateSpace"))).SetValue(ENiagaraCoordinateSpace::Local);
 		Signature.Inputs.Add_GetRef(FNiagaraVariable(FNiagaraTypeDefinition::GetFloatDef(), TEXT("MajorRadius"))).SetValue(100.0f);
 		Signature.Inputs.Add_GetRef(FNiagaraVariable(FNiagaraTypeDefinition::GetFloatDef(), TEXT("MinorRadius"))).SetValue(25.0f);
-		Signature.Inputs.Add_GetRef(FNiagaraVariable(FNiagaraTypeDefinition::GetVec3Def(), TEXT("Orientation Axis"))).SetValue(FVector::ZAxisVector);
+		Signature.Inputs.Add_GetRef(FNiagaraVariable(FNiagaraTypeDefinition::GetVec3Def(), TEXT("Orientation Axis"))).SetValue(FVector3f::ZAxisVector);
 		Signature.Inputs.Add_GetRef(FNiagaraVariable(CoordTypeDef, TEXT("OrientationAxisCoordinateSpace"))).SetValue(ENiagaraCoordinateSpace::Local);
-		Signature.Inputs.Add_GetRef(FNiagaraVariable(FNiagaraTypeDefinition::GetVec3Def(), TEXT("NonUniform Scale"))).SetValue(FVector::OneVector);
+		Signature.Inputs.Add_GetRef(FNiagaraVariable(FNiagaraTypeDefinition::GetVec3Def(), TEXT("NonUniform Scale"))).SetValue(FVector3f::OneVector);
 		Signature.Inputs.Add_GetRef(FNiagaraVariable(FNiagaraTypeDefinition::GetIntDef(), TEXT("MajorRadiusSegments"))).SetValue(12);
 		Signature.Inputs.Add_GetRef(FNiagaraVariable(FNiagaraTypeDefinition::GetIntDef(), TEXT("MinorRadiusSegments"))).SetValue(6);
 		Signature.Inputs.Add(FNiagaraVariable(FNiagaraTypeDefinition::GetColorDef(), TEXT("Color")));
@@ -2723,7 +2723,7 @@ void UNiagaraDataInterfaceDebugDraw::GetFunctions(TArray<FNiagaraFunctionSignatu
 		Signature.Inputs.Add(FNiagaraVariable(FNiagaraTypeDefinition(GetClass()), TEXT("DebugDrawInterface")));
 		Signature.Inputs.Add(FNiagaraVariable(FNiagaraTypeDefinition::GetVec3Def(), TEXT("Center")));
 		Signature.Inputs.Add_GetRef(FNiagaraVariable(CoordTypeDef, TEXT("CenterCoordinateSpace"))).SetValue(ENiagaraCoordinateSpace::Local);
-		Signature.Inputs.Add_GetRef(FNiagaraVariable(FNiagaraTypeDefinition::GetVec3Def(), TEXT("Offset"))).SetValue(FVector::ZeroVector);
+		Signature.Inputs.Add_GetRef(FNiagaraVariable(FNiagaraTypeDefinition::GetVec3Def(), TEXT("Offset"))).SetValue(FVector3f::ZeroVector);
 		Signature.Inputs.Add_GetRef(FNiagaraVariable(CoordTypeDef, TEXT("OffsetCoordinateSpace"))).SetValue(ENiagaraCoordinateSpace::Local);
 		Signature.Inputs.Add(FNiagaraVariable(FNiagaraTypeDefinition::GetQuatDef(), TEXT("Rotation")));
 		Signature.Inputs.Add_GetRef(FNiagaraVariable(CoordTypeDef, TEXT("RotationCoordinateSpace"))).SetValue(ENiagaraCoordinateSpace::Local);
@@ -2739,7 +2739,7 @@ void UNiagaraDataInterfaceDebugDraw::GetFunctions(TArray<FNiagaraFunctionSignatu
 		Signature.Inputs.Add(FNiagaraVariable(FNiagaraTypeDefinition(GetClass()), TEXT("DebugDrawInterface")));
 		Signature.Inputs.Add(FNiagaraVariable(FNiagaraTypeDefinition::GetVec3Def(), TEXT("Center")));
 		Signature.Inputs.Add_GetRef(FNiagaraVariable(CoordTypeDef, TEXT("CenterCoordinateSpace"))).SetValue(ENiagaraCoordinateSpace::Local);
-		Signature.Inputs.Add_GetRef(FNiagaraVariable(FNiagaraTypeDefinition::GetVec3Def(), TEXT("Offset"))).SetValue(FVector::ZeroVector);
+		Signature.Inputs.Add_GetRef(FNiagaraVariable(FNiagaraTypeDefinition::GetVec3Def(), TEXT("Offset"))).SetValue(FVector3f::ZeroVector);
 		Signature.Inputs.Add_GetRef(FNiagaraVariable(CoordTypeDef, TEXT("OffsetCoordinateSpace"))).SetValue(ENiagaraCoordinateSpace::Local);
 		Signature.Inputs.Add(FNiagaraVariable(FNiagaraTypeDefinition::GetQuatDef(), TEXT("Rotation")));
 		Signature.Inputs.Add_GetRef(FNiagaraVariable(CoordTypeDef, TEXT("RotationCoordinateSpace"))).SetValue(ENiagaraCoordinateSpace::Local);
@@ -2758,11 +2758,11 @@ void UNiagaraDataInterfaceDebugDraw::GetFunctions(TArray<FNiagaraFunctionSignatu
 		Signature.Inputs.Add(FNiagaraVariable(FNiagaraTypeDefinition(GetClass()), TEXT("DebugDrawInterface")));
 		Signature.Inputs.Add(FNiagaraVariable(FNiagaraTypeDefinition::GetVec3Def(), TEXT("Center")));
 		Signature.Inputs.Add_GetRef(FNiagaraVariable(CoordTypeDef, TEXT("CenterCoordinateSpace"))).SetValue(ENiagaraCoordinateSpace::Local);
-		Signature.Inputs.Add_GetRef(FNiagaraVariable(FNiagaraTypeDefinition::GetVec3Def(), TEXT("Offset"))).SetValue(FVector::ZeroVector);
+		Signature.Inputs.Add_GetRef(FNiagaraVariable(FNiagaraTypeDefinition::GetVec3Def(), TEXT("Offset"))).SetValue(FVector3f::ZeroVector);
 		Signature.Inputs.Add_GetRef(FNiagaraVariable(CoordTypeDef, TEXT("OffsetCoordinateSpace"))).SetValue(ENiagaraCoordinateSpace::Local);
 		Signature.Inputs.Add(FNiagaraVariable(FNiagaraTypeDefinition::GetQuatDef(), TEXT("Rotation")));
 		Signature.Inputs.Add_GetRef(FNiagaraVariable(CoordTypeDef, TEXT("RotationCoordinateSpace"))).SetValue(ENiagaraCoordinateSpace::Local);
-		Signature.Inputs.Add_GetRef(FNiagaraVariable(FNiagaraTypeDefinition::GetVec3Def(), TEXT("Extents"))).SetValue(FVector(10.0f));
+		Signature.Inputs.Add_GetRef(FNiagaraVariable(FNiagaraTypeDefinition::GetVec3Def(), TEXT("Extents"))).SetValue(FVector3f(10.0f));
 		Signature.Inputs.Add_GetRef(FNiagaraVariable(FNiagaraTypeDefinition::GetIntDef(), TEXT("NumCellsX"))).SetValue(1);
 		Signature.Inputs.Add_GetRef(FNiagaraVariable(FNiagaraTypeDefinition::GetIntDef(), TEXT("NumCellsY"))).SetValue(1);
 		Signature.Inputs.Add_GetRef(FNiagaraVariable(FNiagaraTypeDefinition::GetIntDef(), TEXT("NumCellsZ"))).SetValue(1);
diff --git a/Engine/Plugins/FX/Niagara/Source/Niagara/Private/MovieScene/Parameters/MovieSceneNiagaraVectorParameterTrack.cpp b/Engine/Plugins/FX/Niagara/Source/Niagara/Private/MovieScene/Parameters/MovieSceneNiagaraVectorParameterTrack.cpp
--- a/Engine/Plugins/FX/Niagara/Source/Niagara/Private/MovieScene/Parameters/MovieSceneNiagaraVectorParameterTrack.cpp
+++ b/Engine/Plugins/FX/Niagara/Source/Niagara/Private/MovieScene/Parameters/MovieSceneNiagaraVectorParameterTrack.cpp
@@ -31,9 +31,9 @@ void UMovieSceneNiagaraVectorParameterTrack::SetSectionChannelDefaults(UMovieSce
 			SetChannelDefault(VectorChannelProxy, VectorSection->GetChannel(0), DefaultValue.X);
 			SetChannelDefault(VectorChannelProxy, VectorSection->GetChannel(1), DefaultValue.Y);
 		}
-		else if (VectorSection->GetChannelsUsed() == 3 && ensureMsgf(DefaultValueData.Num() == sizeof(FVector), TEXT("DefaultValueData must be a FVector when channels used is 3")))
+		else if (VectorSection->GetChannelsUsed() == 3 && ensureMsgf(DefaultValueData.Num() == sizeof(FVector3f), TEXT("DefaultValueData must be a FVector3f when channels used is 3")))
 		{
-			FVector DefaultValue = *((FVector*)DefaultValueData.GetData());
+			FVector3f DefaultValue = *((FVector3f*)DefaultValueData.GetData());
 			SetChannelDefault(VectorChannelProxy, VectorSection->GetChannel(0), DefaultValue.X);
 			SetChannelDefault(VectorChannelProxy, VectorSection->GetChannel(1), DefaultValue.Y);
 			SetChannelDefault(VectorChannelProxy, VectorSection->GetChannel(2), DefaultValue.Z);
diff --git a/Engine/Plugins/FX/Niagara/Source/Niagara/Private/NDISkeletalMesh_BoneSampling.cpp b/Engine/Plugins/FX/Niagara/Source/Niagara/Private/NDISkeletalMesh_BoneSampling.cpp
--- a/Engine/Plugins/FX/Niagara/Source/Niagara/Private/NDISkeletalMesh_BoneSampling.cpp
+++ b/Engine/Plugins/FX/Niagara/Source/Niagara/Private/NDISkeletalMesh_BoneSampling.cpp
@@ -711,9 +711,9 @@ struct FBoneSocketSkinnedDataOutputHandler
 	{
 	}
 
-	FNDIOutputParam<FVector> Position;
+	FNDIOutputParam<FVector3f> Position;
 	FNDIOutputParam<FQuat> Rotation;
-	FNDIOutputParam<FVector> Velocity;
+	FNDIOutputParam<FVector3f> Velocity;
 
 	//TODO: Rotation + Scale too? Use quats so we can get proper interpolation between bone and parent.
 
@@ -768,18 +768,18 @@ void UNiagaraDataInterfaceSkeletalMesh::GetSkinnedBoneData(FVectorVMContext& Con
 			const bool bIsSocket = Bone >= BoneCount;
 			const int32 Socket = Bone - BoneCount;
 
-			FVector Pos;
-			FVector Prev;
+			FVector3f Pos;
+			FVector3f Prev;
 
 			// Handle invalid bone indices first
 			if (Bone < 0 || Bone >= BoneAndSocketCount)
 			{
-				Pos = FVector::ZeroVector;
+				Pos = FVector3f::ZeroVector;
 				TransformHandler.TransformPosition(Pos, InstanceTransform);
 
 				if (Output.bNeedsVelocity || bInterpolated::Value)
 				{
-					Prev = FVector::ZeroVector;
+					Prev = FVector3f::ZeroVector;
 					TransformHandler.TransformPosition(Prev, PrevInstanceTransform);
 				}
 				if (Output.bNeedsRotation)
@@ -855,7 +855,7 @@ void UNiagaraDataInterfaceSkeletalMesh::GetSkinnedBoneData(FVectorVMContext& Con
 			if(Output.bNeedsVelocity)
 			{
 				//Don't have enough information to get a better interpolated velocity.
-				FVector Velocity = (Pos - Prev) * InvDt;
+				FVector3f Velocity = (Pos - Prev) * InvDt;
 				Output.Velocity.SetAndAdvance(Velocity);
 			}
 		}
@@ -868,8 +868,8 @@ void UNiagaraDataInterfaceSkeletalMesh::GetSkinnedBoneData(FVectorVMContext& Con
 		{
 			const float Interp = bInterpolated::Value ? InterpParam.GetAndAdvance() : 1.0f;
 
-			FVector Prev = FVector::ZeroVector;
-			FVector Pos = FVector::ZeroVector;
+			FVector3f Prev = FVector3f::ZeroVector;
+			FVector3f Pos = FVector3f::ZeroVector;
 			TransformHandler.TransformPosition(Pos, InstanceTransform);
 
 			if (Output.bNeedsVelocity || bInterpolated::Value)
@@ -894,7 +894,7 @@ void UNiagaraDataInterfaceSkeletalMesh::GetSkinnedBoneData(FVectorVMContext& Con
 
 			if (Output.bNeedsVelocity)
 			{
-				FVector Velocity = (Pos - Prev) * InvDt;
+				FVector3f Velocity = (Pos - Prev) * InvDt;
 				Output.Velocity.SetAndAdvance(Velocity);
 			}
 		}
@@ -949,9 +949,9 @@ void UNiagaraDataInterfaceSkeletalMesh::GetFilteredSocketTransform(FVectorVMCont
 	FNDIInputParam<int32> SocketParam(Context);
 	FNDIInputParam<FNiagaraBool> ApplyWorldTransform(Context);
 
-	FNDIOutputParam<FVector> OutSocketTranslate(Context);
+	FNDIOutputParam<FVector3f> OutSocketTranslate(Context);
 	FNDIOutputParam<FQuat> OutSocketRotation(Context);
-	FNDIOutputParam<FVector> OutSocketScale(Context);
+	FNDIOutputParam<FVector3f> OutSocketScale(Context);
 
 	const TArray<FTransform>& CurrentFilteredSockets = InstData->GetFilteredSocketsCurrBuffer();
 	const int32 SocketMax = CurrentFilteredSockets.Num() - 1;
@@ -964,9 +964,9 @@ void UNiagaraDataInterfaceSkeletalMesh::GetFilteredSocketTransform(FVectorVMCont
 		for (int32 i = 0; i < Context.NumInstances; ++i)
 		{
 			const int32 SocketIndex = FMath::Clamp(SocketParam.GetAndAdvance(), 0, SocketMax);
-			FVector SocketTranslation = CurrentFilteredSockets[SocketIndex].GetTranslation();
+			FVector3f SocketTranslation = CurrentFilteredSockets[SocketIndex].GetTranslation();
 			FQuat SocketRotation = CurrentFilteredSockets[SocketIndex].GetRotation();
-			FVector SocketScale = CurrentFilteredSockets[SocketIndex].GetScale3D();
+			FVector3f SocketScale = CurrentFilteredSockets[SocketIndex].GetScale3D();
 
 			const bool bApplyTransform = ApplyWorldTransform.GetAndAdvance();
 			if (bApplyTransform)
@@ -985,9 +985,9 @@ void UNiagaraDataInterfaceSkeletalMesh::GetFilteredSocketTransform(FVectorVMCont
 	{
 		for (int32 i=0; i < Context.NumInstances; ++i)
 		{
-			OutSocketTranslate.SetAndAdvance(FVector::ZeroVector);
+			OutSocketTranslate.SetAndAdvance(FVector3f::ZeroVector);
 			OutSocketRotation.SetAndAdvance(FQuat::Identity);
-			OutSocketScale.SetAndAdvance(FVector::ZeroVector);
+			OutSocketScale.SetAndAdvance(FVector3f::ZeroVector);
 		}
 	}
 }
diff --git a/Engine/Plugins/FX/Niagara/Source/Niagara/Private/NDISkeletalMesh_TriangleSampling.cpp b/Engine/Plugins/FX/Niagara/Source/Niagara/Private/NDISkeletalMesh_TriangleSampling.cpp
--- a/Engine/Plugins/FX/Niagara/Source/Niagara/Private/NDISkeletalMesh_TriangleSampling.cpp
+++ b/Engine/Plugins/FX/Niagara/Source/Niagara/Private/NDISkeletalMesh_TriangleSampling.cpp
@@ -221,7 +221,7 @@ void UNiagaraDataInterfaceSkeletalMesh::GetTriangleSamplingFunctions(TArray<FNia
 		Sig.Name = FSkeletalMeshInterfaceHelper::GetFilteredTriangleAtName;
 		Sig.Inputs.Add(FNiagaraVariable(FNiagaraTypeDefinition(GetClass()), TEXT("SkeletalMesh")));
 		Sig.Inputs.Add(FNiagaraVariable(FNiagaraTypeDefinition::GetIntDef(), TEXT("Index")));
-		Sig.Inputs.Add_GetRef(FNiagaraVariable(FNiagaraTypeDefinition::GetVec3Def(), TEXT("BaryCoord"))).SetValue(FVector(1.0f / 3.0f));
+		Sig.Inputs.Add_GetRef(FNiagaraVariable(FNiagaraTypeDefinition::GetVec3Def(), TEXT("BaryCoord"))).SetValue(FVector3f(1.0f / 3.0f));
 		Sig.Outputs.Add(FNiagaraVariable(FNiagaraTypeDefinition(FMeshTriCoordinate::StaticStruct()), TEXT("Coord")));
 		Sig.bMemberFunction = true;
 		Sig.bRequiresContext = false;
@@ -538,7 +538,7 @@ void UNiagaraDataInterfaceSkeletalMesh::RandomTriCoord(FVectorVMContext& Context
 	checkfSlow(InstData.Get(), TEXT("Skeletal Mesh Interface has invalid instance data. %s"), *GetPathName());
 
 	FNDIOutputParam<int32> OutTri(Context);
-	FNDIOutputParam<FVector> OutBary(Context);
+	FNDIOutputParam<FVector3f> OutBary(Context);
 
 	FSkeletalMeshAccessorHelper MeshAccessor;
 	MeshAccessor.Init<FilterMode, AreaWeightingMode>(InstData);
@@ -557,7 +557,7 @@ void UNiagaraDataInterfaceSkeletalMesh::RandomTriCoord(FVectorVMContext& Context
 		for (int32 i = 0; i < Context.NumInstances; ++i)
 		{
 			OutTri.SetAndAdvance(-1);
-			OutBary.SetAndAdvance(FVector::ZeroVector);
+			OutBary.SetAndAdvance(FVector3f::ZeroVector);
 		}
 	}
 }
@@ -607,7 +607,7 @@ void UNiagaraDataInterfaceSkeletalMesh::RandomTriangle(FVectorVMContext& Context
 	VectorVM::FUserPtrHandler<FNDISkeletalMesh_InstanceData> InstData(Context);
 	FNDIRandomHelper RandHelper(Context);
 	FNDIOutputParam<int32> OutTri(Context);
-	FNDIOutputParam<FVector> OutBary(Context);
+	FNDIOutputParam<FVector3f> OutBary(Context);
 
 	FSkeletalMeshAccessorHelper MeshAccessor;
 	MeshAccessor.Init<TIntegralConstant<int32, 0>, TIntegralConstant<int32, 0>>(InstData);
@@ -617,7 +617,7 @@ void UNiagaraDataInterfaceSkeletalMesh::RandomTriangle(FVectorVMContext& Context
 		for (int32 i = 0; i < Context.NumInstances; ++i)
 		{
 			OutTri.SetAndAdvance(-1);
-			OutBary.SetAndAdvance(FVector::ZeroVector);
+			OutBary.SetAndAdvance(FVector3f::ZeroVector);
 		}
 		return;
 	}
@@ -659,7 +659,7 @@ void UNiagaraDataInterfaceSkeletalMesh::RandomTriangle(FVectorVMContext& Context
 		for (int32 i = 0; i < Context.NumInstances; ++i)
 		{
 			OutTri.SetAndAdvance(-1);
-			OutBary.SetAndAdvance(FVector::ZeroVector);
+			OutBary.SetAndAdvance(FVector3f::ZeroVector);
 		}
 	}
 }
@@ -884,9 +884,9 @@ void UNiagaraDataInterfaceSkeletalMesh::GetFilteredTriangleAt(FVectorVMContext&
 	VectorVM::FExternalFuncInputHandler<int32> TriParam(Context);
 	checkfSlow(InstData.Get(), TEXT("Skeletal Mesh Interface has invalid instance data. %s"), *GetPathName());
 
-	FNDIInputParam<FVector> InBary(Context);
+	FNDIInputParam<FVector3f> InBary(Context);
 	FNDIOutputParam<int32> OutTri(Context);
-	FNDIOutputParam<FVector> OutBary(Context);
+	FNDIOutputParam<FVector3f> OutBary(Context);
 
 	FSkeletalMeshAccessorHelper Accessor;
 	Accessor.Init<FilterMode, AreaWeightingMode>(InstData);
@@ -921,7 +921,7 @@ void UNiagaraDataInterfaceSkeletalMesh::GetTriCoordColor(FVectorVMContext& Conte
 	SCOPE_CYCLE_COUNTER(STAT_NiagaraSkel_Sample);
 	VectorVM::FUserPtrHandler<FNDISkeletalMesh_InstanceData> InstData(Context);
 	FNDIInputParam<int32> TriParam(Context);
-	FNDIInputParam<FVector> BaryParam(Context);
+	FNDIInputParam<FVector3f> BaryParam(Context);
 
 	FNDIOutputParam<FLinearColor> OutColor(Context);
 
@@ -963,7 +963,7 @@ void UNiagaraDataInterfaceSkeletalMesh::GetTriCoordColorFallback(FVectorVMContex
 {
 	VectorVM::FUserPtrHandler<FNDISkeletalMesh_InstanceData> InstData(Context);
 	FNDIInputParam<int32> TriParam(Context);
-	FNDIInputParam<FVector> BaryParam(Context);
+	FNDIInputParam<FVector3f> BaryParam(Context);
 
 	FNDIOutputParam<FLinearColor> OutColor(Context);
 
@@ -980,7 +980,7 @@ void UNiagaraDataInterfaceSkeletalMesh::GetTriCoordUV(FVectorVMContext& Context)
 	VectorVM::FUserPtrHandler<FNDISkeletalMesh_InstanceData> InstData(Context);
 	VertexAccessorType VertAccessor;
 	FNDIInputParam<int32> TriParam(Context);
-	FNDIInputParam<FVector> BaryParam(Context);
+	FNDIInputParam<FVector3f> BaryParam(Context);
 	FNDIInputParam<int32> UVSetParam(Context);
 
 	checkf(InstData.Get(), TEXT("Skeletal Mesh Interface has invalid instance data. %s"), *GetPathName());
@@ -1031,7 +1031,7 @@ void UNiagaraDataInterfaceSkeletalMesh::GetTriCoordUV<FSkelMeshVertexAccessorNoo
 {
 	VectorVM::FUserPtrHandler<FNDISkeletalMesh_InstanceData> InstData(Context);
 	FNDIInputParam<int32> TriParam(Context);
-	FNDIInputParam<FVector> BaryParam(Context);
+	FNDIInputParam<FVector3f> BaryParam(Context);
 	FNDIInputParam<int32> UVSetParam(Context);
 
 	FNDIOutputParam<FVector2D> OutUV(Context);
@@ -1052,7 +1052,7 @@ void UNiagaraDataInterfaceSkeletalMesh::GetTriangleCoordAtUV(FVectorVMContext& C
 	FNDIInputParam<float> InTolerance(Context);
 
 	FNDIOutputParam<int32> OutTriangleIndex(Context);
-	FNDIOutputParam<FVector> OutBaryCoord(Context);
+	FNDIOutputParam<FVector3f> OutBaryCoord(Context);
 	FNDIOutputParam<FNiagaraBool> OutIsValid(Context);
 
 	checkf(InstData.Get(), TEXT("Skeletal Mesh Interface has invalid instance data. %s"), *GetPathName());
@@ -1066,7 +1066,7 @@ void UNiagaraDataInterfaceSkeletalMesh::GetTriangleCoordAtUV(FVectorVMContext& C
 			const FVector2D SourceUv = InUV.GetAndAdvance();
 			const float Tolerance = InTolerance.GetAndAdvance();
 
-			FVector BaryCoord(ForceInitToZero);
+			FVector3f BaryCoord(ForceInitToZero);
 			int32 TriangleIndex = INDEX_NONE;
 
 			if (Enabled)
@@ -1084,7 +1084,7 @@ void UNiagaraDataInterfaceSkeletalMesh::GetTriangleCoordAtUV(FVectorVMContext& C
 		for (int32 i = 0; i < Context.NumInstances; ++i)
 		{
 			OutTriangleIndex.SetAndAdvance(INDEX_NONE);
-			OutBaryCoord.SetAndAdvance(FVector::ZeroVector);
+			OutBaryCoord.SetAndAdvance(FVector3f::ZeroVector);
 			OutIsValid.SetAndAdvance(false);
 		}
 	}
@@ -1100,13 +1100,13 @@ void UNiagaraDataInterfaceSkeletalMesh::GetTriangleCoordAtUV<FSkelMeshVertexAcce
 	FNDIInputParam<float> InTolerance(Context);
 
 	FNDIOutputParam<int32> OutTriangleIndex(Context);
-	FNDIOutputParam<FVector> OutBaryCoord(Context);
+	FNDIOutputParam<FVector3f> OutBaryCoord(Context);
 	FNDIOutputParam<bool> OutIsValid(Context);
 
 	for (int32 i = 0; i < Context.NumInstances; ++i)
 	{
 		OutTriangleIndex.SetAndAdvance(INDEX_NONE);
-		OutBaryCoord.SetAndAdvance(FVector::ZeroVector);
+		OutBaryCoord.SetAndAdvance(FVector3f::ZeroVector);
 		OutIsValid.SetAndAdvance(false);
 	}
 }
@@ -1122,7 +1122,7 @@ void UNiagaraDataInterfaceSkeletalMesh::GetTriangleCoordInAabb(FVectorVMContext&
 	FNDIInputParam<FVector2D> InMaxExtent(Context);
 
 	FNDIOutputParam<int32> OutTriangleIndex(Context);
-	FNDIOutputParam<FVector> OutBaryCoord(Context);
+	FNDIOutputParam<FVector3f> OutBaryCoord(Context);
 	FNDIOutputParam<FNiagaraBool> OutIsValid(Context);
 
 	checkf(InstData.Get(), TEXT("Skeletal Mesh Interface has invalid instance data. %s"), *GetPathName());
@@ -1136,7 +1136,7 @@ void UNiagaraDataInterfaceSkeletalMesh::GetTriangleCoordInAabb(FVectorVMContext&
 			const FVector2D MinExtent = InMinExtent.GetAndAdvance();
 			const FVector2D MaxExtent = InMaxExtent.GetAndAdvance();
 
-			FVector BaryCoord(ForceInitToZero);
+			FVector3f BaryCoord(ForceInitToZero);
 			int32 TriangleIndex = INDEX_NONE;
 			if (Enabled)
 			{
@@ -1153,7 +1153,7 @@ void UNiagaraDataInterfaceSkeletalMesh::GetTriangleCoordInAabb(FVectorVMContext&
 		for (int32 i = 0; i < Context.NumInstances; ++i)
 		{
 			OutTriangleIndex.SetAndAdvance(INDEX_NONE);
-			OutBaryCoord.SetAndAdvance(FVector::ZeroVector);
+			OutBaryCoord.SetAndAdvance(FVector3f::ZeroVector);
 			OutIsValid.SetAndAdvance(false);
 		}
 	}
@@ -1168,13 +1168,13 @@ void UNiagaraDataInterfaceSkeletalMesh::GetTriangleCoordInAabb<FSkelMeshVertexAc
 	FNDIInputParam<FVector2D> InMaxExtent(Context);
 
 	FNDIOutputParam<int32> OutTriangleIndex(Context);
-	FNDIOutputParam<FVector> OutBaryCoord(Context);
+	FNDIOutputParam<FVector3f> OutBaryCoord(Context);
 	FNDIOutputParam<bool> OutIsValid(Context);
 
 	for (int32 i = 0; i < Context.NumInstances; ++i)
 	{
 		OutTriangleIndex.SetAndAdvance(INDEX_NONE);
-		OutBaryCoord.SetAndAdvance(FVector::ZeroVector);
+		OutBaryCoord.SetAndAdvance(FVector3f::ZeroVector);
 		OutIsValid.SetAndAdvance(false);
 	}
 }
@@ -1195,11 +1195,11 @@ struct FGetTriCoordSkinnedDataOutputHandler
 	{
 	}
 
-	FNDIOutputParam<FVector> Position;
-	FNDIOutputParam<FVector> Velocity;
-	FNDIOutputParam<FVector> Normal;
-	FNDIOutputParam<FVector> Binormal;
-	FNDIOutputParam<FVector> Tangent;
+	FNDIOutputParam<FVector3f> Position;
+	FNDIOutputParam<FVector3f> Velocity;
+	FNDIOutputParam<FVector3f> Normal;
+	FNDIOutputParam<FVector3f> Binormal;
+	FNDIOutputParam<FVector3f> Tangent;
 
 	const bool bNeedsPosition;
 	const bool bNeedsVelocity;
@@ -1217,7 +1217,7 @@ void UNiagaraDataInterfaceSkeletalMesh::GetTriCoordSkinnedData(FVectorVMContext&
 	SkinningHandlerType SkinningHandler;
 	TransformHandlerType TransformHandler;
 	FNDIInputParam<int32> TriParam(Context);
-	FNDIInputParam<FVector> BaryParam(Context);
+	FNDIInputParam<FVector3f> BaryParam(Context);
 	VectorVM::FExternalFuncInputHandler<float> InterpParam;
 
  	if(bInterpolated::Value)
@@ -1250,23 +1250,23 @@ void UNiagaraDataInterfaceSkeletalMesh::GetTriCoordSkinnedData(FVectorVMContext&
 		for (int32 i = 0; i < Context.NumInstances; ++i)
 		{
 			const float Interp = bInterpolated::Value ? InterpParam.GetAndAdvance() : 1.0f;
-			FVector PrevPosition = FVector::ZeroVector;
+			FVector3f PrevPosition = FVector3f::ZeroVector;
 			if (bNeedsPrev)
 			{
 				TransformHandler.TransformPosition(PrevPosition, PrevTransform);
 			}
-			FVector CurrPosition =  FVector::ZeroVector;
+			FVector3f CurrPosition =  FVector3f::ZeroVector;
 			if (Output.bNeedsPosition || Output.bNeedsVelocity)
 			{
 				TransformHandler.TransformPosition(CurrPosition, Transform);
 			}
-			const FVector LerpPosition = bInterpolated::Value ? FMath::Lerp(PrevPosition, CurrPosition, Interp) : CurrPosition;
+			const FVector3f LerpPosition = bInterpolated::Value ? FMath::Lerp(PrevPosition, CurrPosition, Interp) : CurrPosition;
 
 			Output.Position.SetAndAdvance(LerpPosition);
 			Output.Velocity.SetAndAdvance((LerpPosition - PrevPosition) * InvDt);
-			Output.Normal.SetAndAdvance(FVector::ZAxisVector);
-			Output.Binormal.SetAndAdvance(FVector::YAxisVector);
-			Output.Tangent.SetAndAdvance(FVector::XAxisVector);
+			Output.Normal.SetAndAdvance(FVector3f::ZAxisVector);
+			Output.Binormal.SetAndAdvance(FVector3f::YAxisVector);
+			Output.Tangent.SetAndAdvance(FVector3f::XAxisVector);
 		}
 		// Handled missing / failed data
 		return;
@@ -1394,7 +1394,7 @@ void UNiagaraDataInterfaceSkeletalMesh::GetTriCoordSkinnedDataFallback(FVectorVM
 
 	VectorVM::FUserPtrHandler<FNDISkeletalMesh_InstanceData> InstData(Context);	
 	FNDIInputParam<int32> TriParam(Context);
-	FNDIInputParam<FVector> BaryParam(Context);
+	FNDIInputParam<FVector3f> BaryParam(Context);
 	VectorVM::FExternalFuncInputHandler<float> InterpParam;
 
 	if (bInterpolated::Value)
@@ -1422,8 +1422,8 @@ void UNiagaraDataInterfaceSkeletalMesh::GetTriCoordSkinnedDataFallback(FVectorVM
 			Interp = InterpParam.GetAndAdvance();
 		}
 
-		FVector Prev(0.0f);
-		FVector Pos(0.0f);
+		FVector3f Prev(0.0f);
+		FVector3f Pos(0.0f);
 		if (bNeedsPrev)
 		{
 			TransformHandler.TransformPosition(Prev, PrevTransform);
@@ -1443,23 +1443,23 @@ void UNiagaraDataInterfaceSkeletalMesh::GetTriCoordSkinnedDataFallback(FVectorVM
 
 		if (Output.bNeedsVelocity)
 		{			
-			FVector Velocity = (Pos - Prev) * InvDt;
+			FVector3f Velocity = (Pos - Prev) * InvDt;
 			Output.Velocity.SetAndAdvance(Velocity);
 		}
 
 		if (Output.bNeedsNorm)
 		{
-			Output.Normal.SetAndAdvance(FVector(0.0f, 0.0f, 1.0f));
+			Output.Normal.SetAndAdvance(FVector3f(0.0f, 0.0f, 1.0f));
 		}
 
 		if (Output.bNeedsBinorm)
 		{
-			Output.Binormal.SetAndAdvance(FVector(0.0f, 1.0f, 0.0f));
+			Output.Binormal.SetAndAdvance(FVector3f(0.0f, 1.0f, 0.0f));
 		}
 		
 		if (Output.bNeedsTangent)
 		{
-			Output.Tangent.SetAndAdvance(FVector(1.0f, 0.0f, 0.0f));
+			Output.Tangent.SetAndAdvance(FVector3f(1.0f, 0.0f, 0.0f));
 		}		
 	}
 }
diff --git a/Engine/Plugins/FX/Niagara/Source/Niagara/Private/NDISkeletalMesh_VertexSampling.cpp b/Engine/Plugins/FX/Niagara/Source/Niagara/Private/NDISkeletalMesh_VertexSampling.cpp
--- a/Engine/Plugins/FX/Niagara/Source/Niagara/Private/NDISkeletalMesh_VertexSampling.cpp
+++ b/Engine/Plugins/FX/Niagara/Source/Niagara/Private/NDISkeletalMesh_VertexSampling.cpp
@@ -664,11 +664,11 @@ struct FGetVertexSkinnedDataOutputHandler
 		, bNeedsTangentZ(TangentZ.IsValid())
 	{
 	}
-	FNDIOutputParam<FVector> Position;
-	FNDIOutputParam<FVector> Velocity;
-	FNDIOutputParam<FVector> TangentZ;
-	FNDIOutputParam<FVector> TangentY;
-	FNDIOutputParam<FVector> TangentX;
+	FNDIOutputParam<FVector3f> Position;
+	FNDIOutputParam<FVector3f> Velocity;
+	FNDIOutputParam<FVector3f> TangentZ;
+	FNDIOutputParam<FVector3f> TangentY;
+	FNDIOutputParam<FVector3f> TangentX;
 
 	const bool bNeedsPosition;
 	const bool bNeedsVelocity;
@@ -708,7 +708,7 @@ void UNiagaraDataInterfaceSkeletalMesh::GetVertexSkinnedData(FVectorVMContext& C
 			{
 				const int32 Vertex = FMath::Clamp(VertParam.GetAndAdvance(), 0, VertMax);
 
-				FVector Pos = FVector::ZeroVector;
+				FVector3f Pos = FVector3f::ZeroVector;
 				if (Output.bNeedsPosition || Output.bNeedsVelocity)
 				{
 					Pos = SkinningHandler.GetSkinnedVertexPosition(Accessor, Vertex);
@@ -718,9 +718,9 @@ void UNiagaraDataInterfaceSkeletalMesh::GetVertexSkinnedData(FVectorVMContext& C
 
 				if (Output.bNeedsVelocity)
 				{
-					FVector Prev = SkinningHandler.GetSkinnedVertexPreviousPosition(Accessor, Vertex);
+					FVector3f Prev = SkinningHandler.GetSkinnedVertexPreviousPosition(Accessor, Vertex);
 					TransformHandler.TransformPosition(Prev, PrevTransform);
-					const FVector Velocity = (Pos - Prev) * InvDt;
+					const FVector3f Velocity = (Pos - Prev) * InvDt;
 					Output.Velocity.SetAndAdvance(Velocity);
 				}
 
@@ -758,16 +758,16 @@ void UNiagaraDataInterfaceSkeletalMesh::GetVertexSkinnedData(FVectorVMContext& C
 	// Fall though for bad data
 	for (int32 i = 0; i < Context.NumInstances; ++i)
 	{
-		FVector Position = FVector::ZeroVector;
+		FVector3f Position = FVector3f::ZeroVector;
 		if (Output.bNeedsPosition || Output.bNeedsVelocity)
 		{
 			TransformHandler.TransformPosition(Position, Transform);
 		}
 		Output.Position.SetAndAdvance(Position);
-		Output.Velocity.SetAndAdvance(FVector::ZeroVector);
-		Output.TangentX.SetAndAdvance(FVector::XAxisVector);
-		Output.TangentY.SetAndAdvance(FVector::YAxisVector);
-		Output.TangentZ.SetAndAdvance(FVector::ZAxisVector);
+		Output.Velocity.SetAndAdvance(FVector3f::ZeroVector);
+		Output.TangentX.SetAndAdvance(FVector3f::XAxisVector);
+		Output.TangentY.SetAndAdvance(FVector3f::YAxisVector);
+		Output.TangentZ.SetAndAdvance(FVector3f::ZAxisVector);
 	}
 }
 
diff --git a/Engine/Plugins/FX/Niagara/Source/Niagara/Private/NiagaraBoundsCalculatorHelper.h b/Engine/Plugins/FX/Niagara/Source/Niagara/Private/NiagaraBoundsCalculatorHelper.h
--- a/Engine/Plugins/FX/Niagara/Source/Niagara/Private/NiagaraBoundsCalculatorHelper.h
+++ b/Engine/Plugins/FX/Niagara/Source/Niagara/Private/NiagaraBoundsCalculatorHelper.h
@@ -51,8 +51,10 @@ public:
 
 		constexpr float kDefaultSize = 50.0f;
 
-		FBox Bounds(ForceInitToZero);
-		PositionAccessor.GetReader(DataSet).GetMinMax(Bounds.Min, Bounds.Max);
+		FVector3f BoundsMin(ForceInitToZero);
+		FVector3f BoundsMax(ForceInitToZero);
+		PositionAccessor.GetReader(DataSet).GetMinMax(BoundsMin, BoundsMax);
+		FBox Bounds(BoundsMin, BoundsMax);
 
 		float MaxSize = KINDA_SMALL_NUMBER;
 		if (bUsedWithMeshes)
@@ -109,9 +111,9 @@ public:
 		return Bounds.ExpandBy(MaxSize);
 	}
 
-	FNiagaraDataSetAccessor<FVector> PositionAccessor;
+	FNiagaraDataSetAccessor<FVector3f> PositionAccessor;
 	FNiagaraDataSetAccessor<FVector2D> SpriteSizeAccessor;
-	FNiagaraDataSetAccessor<FVector> ScaleAccessor;
+	FNiagaraDataSetAccessor<FVector3f> ScaleAccessor;
 	FNiagaraDataSetAccessor<float> RibbonWidthAccessor;
 
 	const FVector MeshExtents = FVector::OneVector;
diff --git a/Engine/Plugins/FX/Niagara/Source/Niagara/Private/NiagaraCommon.cpp b/Engine/Plugins/FX/Niagara/Source/Niagara/Private/NiagaraCommon.cpp
--- a/Engine/Plugins/FX/Niagara/Source/Niagara/Private/NiagaraCommon.cpp
+++ b/Engine/Plugins/FX/Niagara/Source/Niagara/Private/NiagaraCommon.cpp
@@ -97,6 +97,10 @@ FString FNiagaraTypeHelper::ToString(const uint8* ValueData, const UObject* Stru
 				{
 					Ret += FString::Printf(TEXT("%s: %g "), *Property->GetNameCPP(), *(float*)PropPtr);
 				}
+				else if (Property->IsA(FDoubleProperty::StaticClass()))
+				{
+					Ret += FString::Printf(TEXT("%s: %g "), *Property->GetNameCPP(), *(double*)PropPtr);
+				}
 				else if (Property->IsA(FUInt16Property::StaticClass()))
 				{
 					FFloat16 Val = *(FFloat16*)PropPtr;
@@ -126,6 +130,32 @@ FString FNiagaraTypeHelper::ToString(const uint8* ValueData, const UObject* Stru
 	}
 	return Ret;
 }
+
+UScriptStruct* FNiagaraTypeHelper::FindNiagaraFriendlyTopLevelStruct(UScriptStruct* InStruct)
+{
+	if (!InStruct)
+		return nullptr;
+
+	// Note: UE core types are converted to the float variant as Niagara only works with float types.
+	if (InStruct->GetFName() == NAME_Vector || InStruct->GetFName() == NAME_Vector3d) // LWC support
+	{
+		return FNiagaraTypeDefinition::GetVec3Struct();
+	}
+
+	return InStruct;
+}
+
+bool FNiagaraTypeHelper::IsNiagaraFriendlyTopLevelStruct(UScriptStruct* InStruct)
+{
+	if (FindNiagaraFriendlyTopLevelStruct(InStruct) == InStruct)
+	{
+		return true;
+	}
+	else
+	{
+		return false;
+	}
+}
 //////////////////////////////////////////////////////////////////////////
 
 FNiagaraSystemUpdateContext::~FNiagaraSystemUpdateContext()
diff --git a/Engine/Plugins/FX/Niagara/Source/Niagara/Private/NiagaraComponent.cpp b/Engine/Plugins/FX/Niagara/Source/Niagara/Private/NiagaraComponent.cpp
--- a/Engine/Plugins/FX/Niagara/Source/Niagara/Private/NiagaraComponent.cpp
+++ b/Engine/Plugins/FX/Niagara/Source/Niagara/Private/NiagaraComponent.cpp
@@ -1850,17 +1850,18 @@ void UNiagaraComponent::SetNiagaraVariableVec4(const FString& InVariableName, co
 
 void UNiagaraComponent::SetVariableVec3(FName InVariableName, FVector InValue)
 {
+	FVector3f AsFloat = InValue;
 	const FNiagaraVariable VariableDesc(FNiagaraTypeDefinition::GetVec3Def(), InVariableName);
 	if (SystemInstanceController.IsValid())
 	{
-		SystemInstanceController->SetVariable_Deferred(InVariableName, InValue);
+		SystemInstanceController->SetVariable_Deferred(InVariableName, AsFloat);
 	}
 	else
 	{
-		OverrideParameters.SetParameterValue(InValue, VariableDesc, true);
+		OverrideParameters.SetParameterValue(AsFloat, VariableDesc, true);
 	}
 #if WITH_EDITOR
-	SetParameterOverride(VariableDesc, FNiagaraVariant(&InValue, sizeof(FVector)));
+	SetParameterOverride(VariableDesc, FNiagaraVariant(&AsFloat, sizeof(FVector3f)));
 #endif
 }
 
diff --git a/Engine/Plugins/FX/Niagara/Source/Niagara/Private/NiagaraComponentRendererProperties.cpp b/Engine/Plugins/FX/Niagara/Source/Niagara/Private/NiagaraComponentRendererProperties.cpp
--- a/Engine/Plugins/FX/Niagara/Source/Niagara/Private/NiagaraComponentRendererProperties.cpp
+++ b/Engine/Plugins/FX/Niagara/Source/Niagara/Private/NiagaraComponentRendererProperties.cpp
@@ -38,6 +38,7 @@ bool UNiagaraComponentRendererProperties::IsConvertible(const FNiagaraTypeDefini
 	if ((SourceType == FNiagaraTypeDefinition::GetColorDef() && TargetType.GetStruct() == GetFColorDef().GetStruct()) ||
 		(SourceType == FNiagaraTypeDefinition::GetVec3Def() && TargetType.GetStruct() == GetFColorDef().GetStruct()) ||
 		(SourceType == FNiagaraTypeDefinition::GetVec3Def() && TargetType.GetStruct() == GetFRotatorDef().GetStruct()) ||
+		(SourceType == FNiagaraTypeDefinition::GetVec3Def() && TargetType.GetStruct() == GetFVectorDef().GetStruct()) ||
 		(SourceType == FNiagaraTypeDefinition::GetVec4Def() && TargetType.GetStruct() == GetFColorDef().GetStruct()) ||
 		(SourceType == FNiagaraTypeDefinition::GetQuatDef() && TargetType.GetStruct() == GetFRotatorDef().GetStruct()))
 	{
@@ -66,7 +67,7 @@ NIAGARA_API FNiagaraTypeDefinition UNiagaraComponentRendererProperties::ToNiagar
 		FStructProperty* StructProperty = (FStructProperty*)Property;
 		if (StructProperty->Struct)
 		{
-			return FNiagaraTypeDefinition(StructProperty->Struct);
+			return FNiagaraTypeDefinition(FNiagaraTypeHelper::FindNiagaraFriendlyTopLevelStruct(StructProperty->Struct));
 		}
 	}
 
@@ -104,6 +105,13 @@ FNiagaraTypeDefinition UNiagaraComponentRendererProperties::GetFRotatorDef()
 	return FNiagaraTypeDefinition(RotatorStruct);
 }
 
+FNiagaraTypeDefinition UNiagaraComponentRendererProperties::GetFVectorDef()
+{
+	static UPackage* CoreUObjectPkg = FindObjectChecked<UPackage>(nullptr, TEXT("/Script/CoreUObject"));
+	static UScriptStruct* VectorStruct = FindObjectChecked<UScriptStruct>(CoreUObjectPkg, TEXT("Vector"));
+	return FNiagaraTypeDefinition(VectorStruct, FNiagaraTypeDefinition::EAllowUnfriendlyStruct::Allow);
+}
+
 TArray<TWeakObjectPtr<UNiagaraComponentRendererProperties>> UNiagaraComponentRendererProperties::ComponentRendererPropertiesToDeferredInit;
 
 UNiagaraComponentRendererProperties::UNiagaraComponentRendererProperties()
@@ -289,7 +297,7 @@ void UNiagaraComponentRendererProperties::UpdateSetterFunctions()
 						FNiagaraTypeDefinition FieldType = ToNiagaraType(Property);
 						if (FieldType != PropertyBinding.PropertyType && FieldType == PropertyBinding.AttributeBinding.GetType())
 						{
-							// we can use the original Niagara value with the setter instead of converting it
+							// we can use the original Niagara value with the setter instead of converting it.
 							Setter.bIgnoreConversion = true;
 						}
 						else if (FieldType != PropertyBinding.PropertyType)
diff --git a/Engine/Plugins/FX/Niagara/Source/Niagara/Private/NiagaraConstants.cpp b/Engine/Plugins/FX/Niagara/Source/Niagara/Private/NiagaraConstants.cpp
--- a/Engine/Plugins/FX/Niagara/Source/Niagara/Private/NiagaraConstants.cpp
+++ b/Engine/Plugins/FX/Niagara/Source/Niagara/Private/NiagaraConstants.cpp
@@ -315,7 +315,7 @@ void FNiagaraConstants::Init()
 		
 		AttrDefaultsStrMap.Add(SYS_PARAM_PARTICLES_VELOCITY, TEXT("0.0,0.0,0.0"));
 		Var = SYS_PARAM_PARTICLES_VELOCITY;
-		Var.SetValue<FVector>(FVector(0.0f, 0.0f, 0.0f));
+		Var.SetValue<FVector3f>(FVector3f(0.0f, 0.0f, 0.0f));
 		AttrDefaultsValueMap.Add(SYS_PARAM_PARTICLES_VELOCITY, Var);
 
 		AttrDefaultsStrMap.Add(SYS_PARAM_PARTICLES_COLOR, FLinearColor(1.0f, 1.0f, 1.0f, 1.0f).ToString());
@@ -340,12 +340,12 @@ void FNiagaraConstants::Init()
 
 		AttrDefaultsStrMap.Add(SYS_PARAM_PARTICLES_SPRITE_FACING, TEXT("1.0,0.0,0.0"));
 		Var = SYS_PARAM_PARTICLES_SPRITE_FACING;
-		Var.SetValue<FVector>(FVector(1.0f, 0.0f, 0.0f));
+		Var.SetValue<FVector3f>(FVector3f(1.0f, 0.0f, 0.0f));
 		AttrDefaultsValueMap.Add(SYS_PARAM_PARTICLES_SPRITE_FACING, Var);
 
 		AttrDefaultsStrMap.Add(SYS_PARAM_PARTICLES_SPRITE_ALIGNMENT, TEXT("1.0,0.0,0.0"));
 		Var = SYS_PARAM_PARTICLES_SPRITE_ALIGNMENT;
-		Var.SetValue<FVector>(FVector(1.0f, 0.0f, 0.0f));
+		Var.SetValue<FVector3f>(FVector3f(1.0f, 0.0f, 0.0f));
 		AttrDefaultsValueMap.Add(SYS_PARAM_PARTICLES_SPRITE_ALIGNMENT, Var);
 		
 		AttrDefaultsStrMap.Add(SYS_PARAM_PARTICLES_SUB_IMAGE_INDEX, TEXT("0.0"));
@@ -375,7 +375,7 @@ void FNiagaraConstants::Init()
 
 		AttrDefaultsStrMap.Add(SYS_PARAM_PARTICLES_SCALE, TEXT("1.0,1.0,1.0"));
 		Var = SYS_PARAM_PARTICLES_SCALE;
-		Var.SetValue<FVector>(FVector(1.0f, 1.0f, 1.0f));
+		Var.SetValue<FVector3f>(FVector3f(1.0f, 1.0f, 1.0f));
 		AttrDefaultsValueMap.Add(SYS_PARAM_PARTICLES_SCALE, Var);
 
 		AttrDefaultsStrMap.Add(SYS_PARAM_PARTICLES_LIFETIME, TEXT("5.0"));
@@ -460,7 +460,7 @@ void FNiagaraConstants::Init()
 
 		AttrDefaultsStrMap.Add(SYS_PARAM_PARTICLES_RIBBONFACING, TEXT("0.0, 0.0, 1.0"));
 		Var = SYS_PARAM_PARTICLES_RIBBONFACING;
-		Var.SetValue<FVector>(FVector(0.0f, 0.0f, 1.0f));
+		Var.SetValue<FVector3f>(FVector3f(0.0f, 0.0f, 1.0f));
 		AttrDefaultsValueMap.Add(SYS_PARAM_PARTICLES_RIBBONFACING, Var);
 
 		AttrDefaultsStrMap.Add(SYS_PARAM_PARTICLES_RIBBONLINKORDER, TEXT("0"));
diff --git a/Engine/Plugins/FX/Niagara/Source/Niagara/Private/NiagaraDataInterface2DArrayTexture.cpp b/Engine/Plugins/FX/Niagara/Source/Niagara/Private/NiagaraDataInterface2DArrayTexture.cpp
--- a/Engine/Plugins/FX/Niagara/Source/Niagara/Private/NiagaraDataInterface2DArrayTexture.cpp
+++ b/Engine/Plugins/FX/Niagara/Source/Niagara/Private/NiagaraDataInterface2DArrayTexture.cpp
@@ -129,7 +129,7 @@ bool UNiagaraDataInterface2DArrayTexture::PerInstanceTick(void* PerInstanceData,
 
 void UNiagaraDataInterface2DArrayTexture::GetTextureDimensions(FVectorVMContext& Context)
 {
-	FNDIOutputParam<FVector> DimensionsOut(Context);
+	FNDIOutputParam<FVector3f> DimensionsOut(Context);
 
 	const FVector FloatTextureSize(TextureSize.X, TextureSize.Y, TextureSize.Z);
 	for (int32 i = 0; i < Context.NumInstances; ++i)
diff --git a/Engine/Plugins/FX/Niagara/Source/Niagara/Private/NiagaraDataInterfaceArrayFloat.cpp b/Engine/Plugins/FX/Niagara/Source/Niagara/Private/NiagaraDataInterfaceArrayFloat.cpp
--- a/Engine/Plugins/FX/Niagara/Source/Niagara/Private/NiagaraDataInterfaceArrayFloat.cpp
+++ b/Engine/Plugins/FX/Niagara/Source/Niagara/Private/NiagaraDataInterfaceArrayFloat.cpp
@@ -7,6 +7,7 @@
 template<>
 struct FNDIArrayImplHelper<float> : public FNDIArrayImplHelperBase<float>
 {
+	typedef float TVMArrayType;
 	static constexpr TCHAR const* HLSLValueTypeName = TEXT("float");
 	static constexpr TCHAR const* HLSLBufferTypeName = TEXT("float");
 	static constexpr EPixelFormat PixelFormat = PF_R32_FLOAT;
@@ -17,6 +18,7 @@ struct FNDIArrayImplHelper<float> : public FNDIArrayImplHelperBase<float>
 template<>
 struct FNDIArrayImplHelper<FVector2D> : public FNDIArrayImplHelperBase<FVector2D>
 {
+	typedef FVector2D TVMArrayType;
 	static constexpr TCHAR const* HLSLValueTypeName = TEXT("float2");
 	static constexpr TCHAR const* HLSLBufferTypeName = TEXT("float2");
 	static constexpr EPixelFormat PixelFormat = PF_G32R32F;
@@ -25,13 +27,14 @@ struct FNDIArrayImplHelper<FVector2D> : public FNDIArrayImplHelperBase<FVector2D
 };
 
 template<>
-struct FNDIArrayImplHelper<FVector> : public FNDIArrayImplHelperBase<FVector>
+struct FNDIArrayImplHelper<FVector3f> : public FNDIArrayImplHelperBase<FVector3f>
 {
+	typedef FVector3f TVMArrayType;
 	static constexpr TCHAR const* HLSLValueTypeName = TEXT("float3");
 	static constexpr TCHAR const* HLSLBufferTypeName = TEXT("float");	//-OPT: Current we have no float3 pixel format, when we add one update this to use it
 	static constexpr EPixelFormat PixelFormat = PF_R32_FLOAT;
 	static const FNiagaraTypeDefinition& GetTypeDefinition() { return FNiagaraTypeDefinition::GetVec3Def(); }
-	static const FVector GetDefaultValue() { return FVector::ZeroVector; }
+	static const FVector3f GetDefaultValue() { return FVector3f::ZeroVector; }
 
 	static void GPUGetFetchHLSL(FString& OutHLSL, const TCHAR* BufferName)
 	{
@@ -45,9 +48,41 @@ struct FNDIArrayImplHelper<FVector> : public FNDIArrayImplHelperBase<FVector>
 	}
 };
 
+// LWC_TODO: This is represented as an FVector3f internally (array is converted to floats during PushToRenderThread)
+template<>
+struct FNDIArrayImplHelper<FVector3d> : public FNDIArrayImplHelperBase<FVector3d>
+{
+	typedef FVector3f TVMArrayType;
+	static constexpr TCHAR const* HLSLValueTypeName = TEXT("float3");
+	static constexpr TCHAR const* HLSLBufferTypeName = TEXT("float");	//-OPT: Current we have no float3 pixel format, when we add one update this to use it
+	static constexpr EPixelFormat PixelFormat = PF_R32_FLOAT;
+	static const FNiagaraTypeDefinition& GetTypeDefinition() { return FNiagaraTypeDefinition::GetVec3Def(); }
+	static const FVector3f GetDefaultValue() { return FVector3f::ZeroVector; }
+
+	static void GPUGetFetchHLSL(FString& OutHLSL, const TCHAR* BufferName)
+	{
+		OutHLSL.Appendf(TEXT("OutValue.x = %s[ClampedIndex * 3 + 0];"), BufferName);
+		OutHLSL.Appendf(TEXT("OutValue.y = %s[ClampedIndex * 3 + 1];"), BufferName);
+		OutHLSL.Appendf(TEXT("OutValue.z = %s[ClampedIndex * 3 + 2];"), BufferName);
+	}
+	static int32 GPUGetTypeStride()
+	{
+		return sizeof(float);
+	}
+
+	static void CopyData(void* Dest, const FVector3d* Src, int32 BufferSize)
+	{
+		FVector3f* DestFloats = (FVector3f*)Dest;
+		int32 Num = BufferSize / sizeof(FVector3d);
+		for (int32 i = 0; i < Num; i++)
+			DestFloats[i] = Src[i];
+	}
+};
+
 template<>
 struct FNDIArrayImplHelper<FVector4> : public FNDIArrayImplHelperBase<FVector4>
 {
+	typedef FVector4 TVMArrayType;
 	static constexpr TCHAR const* HLSLValueTypeName = TEXT("float4");
 	static constexpr TCHAR const* HLSLBufferTypeName = TEXT("float4");
 	static constexpr EPixelFormat PixelFormat = PF_A32B32G32R32F;
@@ -58,6 +93,7 @@ struct FNDIArrayImplHelper<FVector4> : public FNDIArrayImplHelperBase<FVector4>
 template<>
 struct FNDIArrayImplHelper<FLinearColor> : public FNDIArrayImplHelperBase<FLinearColor>
 {
+	typedef FLinearColor TVMArrayType;
 	static constexpr TCHAR const* HLSLValueTypeName = TEXT("float4");
 	static constexpr TCHAR const* HLSLBufferTypeName = TEXT("float4");
 	static constexpr EPixelFormat PixelFormat = PF_A32B32G32R32F;
@@ -68,6 +104,7 @@ struct FNDIArrayImplHelper<FLinearColor> : public FNDIArrayImplHelperBase<FLinea
 template<>
 struct FNDIArrayImplHelper<FQuat> : public FNDIArrayImplHelperBase<FQuat>
 {
+	typedef FQuat TVMArrayType;
 	static constexpr TCHAR const* HLSLValueTypeName = TEXT("float4");
 	static constexpr TCHAR const* HLSLBufferTypeName = TEXT("float4");
 	static constexpr EPixelFormat PixelFormat = PF_A32B32G32R32F;
@@ -115,4 +152,4 @@ UNiagaraDataInterfaceArrayQuat::UNiagaraDataInterfaceArrayQuat(FObjectInitialize
 {
 	Proxy.Reset(new FNiagaraDataInterfaceProxyArrayImpl());
 	Impl.Reset(new FNiagaraDataInterfaceArrayImpl<FQuat, UNiagaraDataInterfaceArrayQuat>(this, QuatData));
-}
+}
\ No newline at end of file
diff --git a/Engine/Plugins/FX/Niagara/Source/Niagara/Private/NiagaraDataInterfaceArrayInt.cpp b/Engine/Plugins/FX/Niagara/Source/Niagara/Private/NiagaraDataInterfaceArrayInt.cpp
--- a/Engine/Plugins/FX/Niagara/Source/Niagara/Private/NiagaraDataInterfaceArrayInt.cpp
+++ b/Engine/Plugins/FX/Niagara/Source/Niagara/Private/NiagaraDataInterfaceArrayInt.cpp
@@ -7,6 +7,7 @@
 template<>
 struct FNDIArrayImplHelper<int32> : public FNDIArrayImplHelperBase<int32>
 {
+	typedef int32 TVMArrayType;
 	static constexpr TCHAR const* HLSLValueTypeName = TEXT("int");
 	static constexpr TCHAR const* HLSLBufferTypeName = TEXT("int");
 	static constexpr EPixelFormat PixelFormat = PF_R32_SINT;
diff --git a/Engine/Plugins/FX/Niagara/Source/Niagara/Private/NiagaraDataInterfaceAudioPlayer.cpp b/Engine/Plugins/FX/Niagara/Source/Niagara/Private/NiagaraDataInterfaceAudioPlayer.cpp
--- a/Engine/Plugins/FX/Niagara/Source/Niagara/Private/NiagaraDataInterfaceAudioPlayer.cpp
+++ b/Engine/Plugins/FX/Niagara/Source/Niagara/Private/NiagaraDataInterfaceAudioPlayer.cpp
@@ -598,7 +598,7 @@ void UNiagaraDataInterfaceAudioPlayer::UpdateLocation(FVectorVMContext& Context)
 	VectorVM::FUserPtrHandler<FAudioPlayerInterface_InstanceData> InstData(Context);
 
 	FNDIInputParam<int32> AudioHandleInParam(Context);
-	FNDIInputParam<FVector> LocationParam(Context);
+	FNDIInputParam<FVector3f> LocationParam(Context);
 	checkfSlow(InstData.Get(), TEXT("Audio player interface has invalid instance data. %s"), *GetPathName());
 
 	InstData->bHadPersistentAudioUpdateThisTick = true;
@@ -628,7 +628,7 @@ void UNiagaraDataInterfaceAudioPlayer::UpdateRotation(FVectorVMContext& Context)
 	VectorVM::FUserPtrHandler<FAudioPlayerInterface_InstanceData> InstData(Context);
 
 	FNDIInputParam<int32> AudioHandleInParam(Context);
-	FNDIInputParam<FVector> RotationParam(Context);
+	FNDIInputParam<FVector3f> RotationParam(Context);
 	checkfSlow(InstData.Get(), TEXT("Audio player interface has invalid instance data. %s"), *GetPathName());
 
 	InstData->bHadPersistentAudioUpdateThisTick = true;
@@ -732,8 +732,8 @@ void UNiagaraDataInterfaceAudioPlayer::PlayPersistentAudio(FVectorVMContext& Con
 
 	FNDIInputParam<FNiagaraBool> PlayAudioParam(Context);
 	FNDIInputParam<int32> AudioHandleInParam(Context);
-	FNDIInputParam<FVector> PositionParam(Context);
-	FNDIInputParam<FVector> RotationParam(Context);
+	FNDIInputParam<FVector3f> PositionParam(Context);
+	FNDIInputParam<FVector3f> RotationParam(Context);
 	FNDIInputParam<float> VolumeParam(Context);
 	FNDIInputParam<float> PitchParam(Context);
 	FNDIInputParam<float> StartTimeParam(Context);
diff --git a/Engine/Plugins/FX/Niagara/Source/Niagara/Private/NiagaraDataInterfaceCamera.cpp b/Engine/Plugins/FX/Niagara/Source/Niagara/Private/NiagaraDataInterfaceCamera.cpp
--- a/Engine/Plugins/FX/Niagara/Source/Niagara/Private/NiagaraDataInterfaceCamera.cpp
+++ b/Engine/Plugins/FX/Niagara/Source/Niagara/Private/NiagaraDataInterfaceCamera.cpp
@@ -520,7 +520,7 @@ void UNiagaraDataInterfaceCamera::CalculateParticleDistances(FVectorVMContext& C
 	VectorVM::FUserPtrHandler<FCameraDataInterface_InstanceData> InstData(Context);
 
 	FNDIInputParam<FNiagaraID> IDParam(Context);
-	FNDIInputParam<FVector> ParticlePosParam(Context);
+	FNDIInputParam<FVector3f> ParticlePosParam(Context);
 
 	FVector CameraPos = InstData->CameraLocation;
 	for (int32 i = 0; i < Context.NumInstances; ++i)
diff --git a/Engine/Plugins/FX/Niagara/Source/Niagara/Private/NiagaraDataInterfaceCubeTexture.cpp b/Engine/Plugins/FX/Niagara/Source/Niagara/Private/NiagaraDataInterfaceCubeTexture.cpp
--- a/Engine/Plugins/FX/Niagara/Source/Niagara/Private/NiagaraDataInterfaceCubeTexture.cpp
+++ b/Engine/Plugins/FX/Niagara/Source/Niagara/Private/NiagaraDataInterfaceCubeTexture.cpp
@@ -156,7 +156,7 @@ void UNiagaraDataInterfaceCubeTexture::GetTextureDimensions(FVectorVMContext& Co
 
 void UNiagaraDataInterfaceCubeTexture::SampleCubeTexture(FVectorVMContext& Context)
 {
-	FNDIInputParam<FVector> InCoord(Context);
+	FNDIInputParam<FVector3f> InCoord(Context);
 	FNDIInputParam<float> InMipLevel(Context);
 	FNDIOutputParam<FVector4> OutColor(Context);
 
diff --git a/Engine/Plugins/FX/Niagara/Source/Niagara/Private/NiagaraDataInterfaceExport.cpp b/Engine/Plugins/FX/Niagara/Source/Niagara/Private/NiagaraDataInterfaceExport.cpp
--- a/Engine/Plugins/FX/Niagara/Source/Niagara/Private/NiagaraDataInterfaceExport.cpp
+++ b/Engine/Plugins/FX/Niagara/Source/Niagara/Private/NiagaraDataInterfaceExport.cpp
@@ -511,9 +511,9 @@ void UNiagaraDataInterfaceExport::StoreData(FVectorVMContext& Context)
 	VectorVM::FUserPtrHandler<FNDIExportInstanceData_GameThread> InstData(Context);
 
 	FNDIInputParam<FNiagaraBool> StoreDataParam(Context);
-	FNDIInputParam<FVector> PositionParam(Context);
+	FNDIInputParam<FVector3f> PositionParam(Context);
 	FNDIInputParam<float> SizeParam(Context);
-	FNDIInputParam<FVector> VelocityParam(Context);
+	FNDIInputParam<FVector3f> VelocityParam(Context);
 
 	FNDIOutputParam<FNiagaraBool> OutSample(Context);
 
@@ -543,9 +543,9 @@ void UNiagaraDataInterfaceExport::ExportData(FVectorVMContext& Context)
 	VectorVM::FUserPtrHandler<FNDIExportInstanceData_GameThread> InstData(Context);
 
 	FNDIInputParam<FNiagaraBool> StoreDataParam(Context);
-	FNDIInputParam<FVector> PositionParam(Context);
+	FNDIInputParam<FVector3f> PositionParam(Context);
 	FNDIInputParam<float> SizeParam(Context);
-	FNDIInputParam<FVector> VelocityParam(Context);
+	FNDIInputParam<FVector3f> VelocityParam(Context);
 
 	FNDIOutputParam<FNiagaraBool> OutWasExported(Context);
 
diff --git a/Engine/Plugins/FX/Niagara/Source/Niagara/Private/NiagaraDataInterfaceMeshRendererInfo.cpp b/Engine/Plugins/FX/Niagara/Source/Niagara/Private/NiagaraDataInterfaceMeshRendererInfo.cpp
--- a/Engine/Plugins/FX/Niagara/Source/Niagara/Private/NiagaraDataInterfaceMeshRendererInfo.cpp
+++ b/Engine/Plugins/FX/Niagara/Source/Niagara/Private/NiagaraDataInterfaceMeshRendererInfo.cpp
@@ -38,8 +38,8 @@ class FNDIMeshRendererInfo
 public:
 	struct FMeshData
 	{
-		FVector MinLocalBounds = FVector(ForceInitToZero);
-		FVector MaxLocalBounds = FVector(ForceInitToZero);
+		FVector3f MinLocalBounds = FVector3f(ForceInitToZero);
+		FVector3f MaxLocalBounds = FVector3f(ForceInitToZero);
 	};
 
 	using FMeshDataArray = TArray<FMeshData>;
@@ -79,13 +79,13 @@ public:
 
 	virtual void InitRHI() override final
 	{
-		uint32 SizeByte = MeshData.Num() * sizeof(FVector) * 2;
+		uint32 SizeByte = MeshData.Num() * sizeof(FVector3f) * 2;
 		
 		if (SizeByte > 0)
 		{
 			FRHIResourceCreateInfo CreateInfo(TEXT("FNDIMeshRendererInfoGPUData"));
 			BufferMeshDataRHI = RHICreateBuffer(SizeByte, BUF_Static | BUF_VertexBuffer | BUF_ShaderResource, 0, ERHIAccess::VertexOrIndexBuffer | ERHIAccess::SRVMask, CreateInfo);
-			FVector* BufferDataVector = reinterpret_cast<FVector*>(RHILockBuffer(BufferMeshDataRHI, 0, SizeByte, RLM_WriteOnly));
+			FVector3f* BufferDataVector = reinterpret_cast<FVector3f*>(RHILockBuffer(BufferMeshDataRHI, 0, SizeByte, RLM_WriteOnly));
 			for (const auto& Mesh : MeshData)
 			{
 				*BufferDataVector++ = Mesh.MinLocalBounds;
@@ -599,14 +599,14 @@ void UNiagaraDataInterfaceMeshRendererInfo::GetNumMeshes(FVectorVMContext& Conte
 void UNiagaraDataInterfaceMeshRendererInfo::GetMeshLocalBounds(FVectorVMContext& Context)
 {
 	FNDIInputParam<int32> InMeshIdx(Context);
-	FNDIOutputParam<FVector> OutMinBounds(Context);
-	FNDIOutputParam<FVector> OutMaxBounds(Context);
-	FNDIOutputParam<FVector> OutSize(Context);
+	FNDIOutputParam<FVector3f> OutMinBounds(Context);
+	FNDIOutputParam<FVector3f> OutMaxBounds(Context);
+	FNDIOutputParam<FVector3f> OutSize(Context);
 
 	for (int32 Instance = 0; Instance < Context.GetNumInstances(); ++Instance)
 	{
-		FVector MinLocalBounds(ForceInitToZero);
-		FVector MaxLocalBounds(ForceInitToZero);
+		FVector3f MinLocalBounds(ForceInitToZero);
+		FVector3f MaxLocalBounds(ForceInitToZero);
 		if (Info.IsValid())
 		{
 			const FNDIMeshRendererInfo::FMeshDataArray& MeshData = Info->GetMeshData();
diff --git a/Engine/Plugins/FX/Niagara/Source/Niagara/Private/NiagaraDataInterfaceNeighborGrid3D.cpp b/Engine/Plugins/FX/Niagara/Source/Niagara/Private/NiagaraDataInterfaceNeighborGrid3D.cpp
--- a/Engine/Plugins/FX/Niagara/Source/Niagara/Private/NiagaraDataInterfaceNeighborGrid3D.cpp
+++ b/Engine/Plugins/FX/Niagara/Source/Niagara/Private/NiagaraDataInterfaceNeighborGrid3D.cpp
@@ -352,7 +352,7 @@ void UNiagaraDataInterfaceNeighborGrid3D::GetWorldBBoxSize(FVectorVMContext& Con
 {
 	VectorVM::FUserPtrHandler<NeighborGrid3DRWInstanceData> InstData(Context);
 
-	FNDIOutputParam<FVector> OutWorldBounds(Context);
+	FNDIOutputParam<FVector3f> OutWorldBounds(Context);
 
 	for (int32 InstanceIdx = 0; InstanceIdx < Context.NumInstances; ++InstanceIdx)
 	{
diff --git a/Engine/Plugins/FX/Niagara/Source/Niagara/Private/NiagaraDataInterfaceParticleRead.cpp b/Engine/Plugins/FX/Niagara/Source/Niagara/Private/NiagaraDataInterfaceParticleRead.cpp
--- a/Engine/Plugins/FX/Niagara/Source/Niagara/Private/NiagaraDataInterfaceParticleRead.cpp
+++ b/Engine/Plugins/FX/Niagara/Source/Niagara/Private/NiagaraDataInterfaceParticleRead.cpp
@@ -1577,7 +1577,7 @@ void UNiagaraDataInterfaceParticleRead::ReadVector2(FVectorVMContext& Context, F
 
 void UNiagaraDataInterfaceParticleRead::ReadVector3(FVectorVMContext& Context, FName AttributeToRead)
 {
-	ReadWithCheck<FVector>(Context, AttributeToRead, FVector::ZeroVector);
+	ReadWithCheck<FVector3f>(Context, AttributeToRead, FVector3f::ZeroVector);
 }
 
 void UNiagaraDataInterfaceParticleRead::ReadVector4(FVectorVMContext& Context, FName AttributeToRead)
@@ -1697,7 +1697,7 @@ void UNiagaraDataInterfaceParticleRead::ReadVector2ByIndex(FVectorVMContext& Con
 
 void UNiagaraDataInterfaceParticleRead::ReadVector3ByIndex(FVectorVMContext& Context, FName AttributeToRead)
 {
-	ReadByIndexWithCheck<FVector>(Context, AttributeToRead, FVector::ZeroVector);
+	ReadByIndexWithCheck<FVector3f>(Context, AttributeToRead, FVector3f::ZeroVector);
 }
 
 void UNiagaraDataInterfaceParticleRead::ReadVector4ByIndex(FVectorVMContext& Context, FName AttributeToRead)
diff --git a/Engine/Plugins/FX/Niagara/Source/Niagara/Private/NiagaraDataInterfaceSkeletalMesh.cpp b/Engine/Plugins/FX/Niagara/Source/Niagara/Private/NiagaraDataInterfaceSkeletalMesh.cpp
--- a/Engine/Plugins/FX/Niagara/Source/Niagara/Private/NiagaraDataInterfaceSkeletalMesh.cpp
+++ b/Engine/Plugins/FX/Niagara/Source/Niagara/Private/NiagaraDataInterfaceSkeletalMesh.cpp
@@ -3399,7 +3399,7 @@ bool UNiagaraDataInterfaceSkeletalMesh::UpgradeFunctionCall(FNiagaraFunctionSign
 	{
 		if ( FunctionSignature.Name == FSkeletalMeshInterfaceHelper::GetFilteredTriangleAtName )
 		{
-			FunctionSignature.Inputs.Add_GetRef(FNiagaraVariable(FNiagaraTypeDefinition::GetVec3Def(), TEXT("BaryCoord"))).SetValue(FVector(1.0f / 3.0f));
+			FunctionSignature.Inputs.Add_GetRef(FNiagaraVariable(FNiagaraTypeDefinition::GetVec3Def(), TEXT("BaryCoord"))).SetValue(FVector3f(1.0f / 3.0f));
 			bWasChanged = true;
 		}
 	}
diff --git a/Engine/Plugins/FX/Niagara/Source/Niagara/Private/NiagaraDataInterfaceSkeletalMeshUvMapping.cpp b/Engine/Plugins/FX/Niagara/Source/Niagara/Private/NiagaraDataInterfaceSkeletalMeshUvMapping.cpp
--- a/Engine/Plugins/FX/Niagara/Source/Niagara/Private/NiagaraDataInterfaceSkeletalMeshUvMapping.cpp
+++ b/Engine/Plugins/FX/Niagara/Source/Niagara/Private/NiagaraDataInterfaceSkeletalMeshUvMapping.cpp
@@ -221,21 +221,27 @@ void FSkeletalMeshUvMappingHandle::FindOverlappingTriangles(const FVector2D& InU
 	}
 }
 
-int32 FSkeletalMeshUvMappingHandle::FindFirstTriangle(const FVector2D& InUv, float Tolerance, FVector& BarycentricCoord) const
+int32 FSkeletalMeshUvMappingHandle::FindFirstTriangle(const FVector2D& InUv, float Tolerance, FVector3f& BarycentricCoord) const
 {
 	if (UvMappingData)
 	{
-		return UvMappingData->FindFirstTriangle(InUv, Tolerance, BarycentricCoord);
+		FVector BarycentricCoord3d = FVector::ZeroVector;
+		int32 Tri = UvMappingData->FindFirstTriangle(InUv, Tolerance, BarycentricCoord3d);
+		BarycentricCoord = BarycentricCoord3d;
+		return Tri;
 	}
 
 	return INDEX_NONE;
 }
 
-int32 FSkeletalMeshUvMappingHandle::FindFirstTriangle(const FBox2D& InUvBox, FVector& BarycentricCoord) const
+int32 FSkeletalMeshUvMappingHandle::FindFirstTriangle(const FBox2D& InUvBox, FVector3f& BarycentricCoord) const
 {
 	if (UvMappingData)
 	{
-		return UvMappingData->FindFirstTriangle(InUvBox, BarycentricCoord);
+		FVector BarycentricCoord3d = FVector::ZeroVector;
+		int32 Tri = UvMappingData->FindFirstTriangle(InUvBox, BarycentricCoord3d);
+		BarycentricCoord = BarycentricCoord3d;
+		return Tri;
 	}
 	return INDEX_NONE;
 }
diff --git a/Engine/Plugins/FX/Niagara/Source/Niagara/Private/NiagaraDataInterfaceStaticMesh.cpp b/Engine/Plugins/FX/Niagara/Source/Niagara/Private/NiagaraDataInterfaceStaticMesh.cpp
--- a/Engine/Plugins/FX/Niagara/Source/Niagara/Private/NiagaraDataInterfaceStaticMesh.cpp
+++ b/Engine/Plugins/FX/Niagara/Source/Niagara/Private/NiagaraDataInterfaceStaticMesh.cpp
@@ -1920,7 +1920,7 @@ void UNiagaraDataInterfaceStaticMesh::RandomTriCoord(FVectorVMContext& Context)
 	VectorVM::FUserPtrHandler<FNDIStaticMesh_InstanceData> InstData(Context);
 
 	FNDIOutputParam<int32> OutTri(Context);
-	FNDIOutputParam<FVector> OutBary(Context);
+	FNDIOutputParam<FVector3f> OutBary(Context);
 
 	const FStaticMeshLODResources* Res = InstData->CachedLOD;
 	check(Res && InstData->StaticMesh.IsValid()); // sanity check - should have bound the invalid specialization below
@@ -1939,7 +1939,7 @@ void UNiagaraDataInterfaceStaticMesh::RandomTriCoord<TSampleModeInvalid>(FVector
 	VectorVM::FUserPtrHandler<FNDIStaticMesh_InstanceData> InstData(Context);
 
 	FNDIOutputParam<int32> OutTri(Context);
-	FNDIOutputParam<FVector> OutBary(Context);
+	FNDIOutputParam<FVector3f> OutBary(Context);
 
 	for (int32 i = 0; i < Context.NumInstances; ++i)
 	{
@@ -1955,7 +1955,7 @@ void UNiagaraDataInterfaceStaticMesh::RandomTriCoordVertexColorFiltered(FVectorV
 	VectorVM::FExternalFuncRegisterHandler<int32> RangeValue(Context);
 
 	FNDIOutputParam<int32> OutTri(Context);
-	FNDIOutputParam<FVector> OutBary(Context);
+	FNDIOutputParam<FVector3f> OutBary(Context);
 	
 	// Handle no mesh case
 	//TODO: Maybe figure out a good way to stub this in bindings to prevent the branch
@@ -2034,7 +2034,7 @@ void UNiagaraDataInterfaceStaticMesh::RandomTriCoordOnSection(FVectorVMContext&
 	FNDIInputParam<int32> SectionIdxParam(Context);
 
 	FNDIOutputParam<int32> OutTri(Context);
-	FNDIOutputParam<FVector> OutBary(Context);
+	FNDIOutputParam<FVector3f> OutBary(Context);
 
 	const FStaticMeshLODResources* Res = InstData->CachedLOD;
 	check(Res && InstData->bMeshValid); // sanity check - should have bound the invalid specialization below
@@ -2069,7 +2069,7 @@ void UNiagaraDataInterfaceStaticMesh::RandomTriCoordOnSection<TSampleModeInvalid
 	FNDIInputParam<int32> SectionIdxParam(Context);
 
 	FNDIOutputParam<int32> OutTri(Context);
-	FNDIOutputParam<FVector> OutBary(Context);
+	FNDIOutputParam<FVector3f> OutBary(Context);
 
 	for (int32 i = 0; i < Context.NumInstances; ++i)
 	{
@@ -2084,9 +2084,9 @@ void UNiagaraDataInterfaceStaticMesh::GetTriCoordPosition(FVectorVMContext& Cont
 	VectorVM::FUserPtrHandler<FNDIStaticMesh_InstanceData> InstData(Context);
 	TransformHandlerType TransformHandler;
 	FNDIInputParam<int32> TriParam(Context);
-	FNDIInputParam<FVector> BaryParam(Context);
+	FNDIInputParam<FVector3f> BaryParam(Context);
 
-	FNDIOutputParam<FVector> OutPos(Context);
+	FNDIOutputParam<FVector3f> OutPos(Context);
 
 	if (InstData->bMeshValid)
 	{
@@ -2133,9 +2133,9 @@ void UNiagaraDataInterfaceStaticMesh::GetTriCoordNormal(FVectorVMContext& Contex
 	TransformHandlerType TransformHandler;
 
 	FNDIInputParam<int32> TriParam(Context);
-	FNDIInputParam<FVector> BaryParam(Context);
+	FNDIInputParam<FVector3f> BaryParam(Context);
 
-	FNDIOutputParam<FVector> OutNorm(Context);
+	FNDIOutputParam<FVector3f> OutNorm(Context);
 
 	if (InstData->bMeshValid)
 	{
@@ -2179,11 +2179,11 @@ void UNiagaraDataInterfaceStaticMesh::GetTriCoordTangents(FVectorVMContext& Cont
 	TransformHandlerType TransformHandler;
 
 	FNDIInputParam<int32> TriParam(Context);
-	FNDIInputParam<FVector> BaryParam(Context);
+	FNDIInputParam<FVector3f> BaryParam(Context);
 
-	FNDIOutputParam<FVector> OutTangent(Context);
-	FNDIOutputParam<FVector> OutBinorm(Context);
-	FNDIOutputParam<FVector> OutNorm(Context);
+	FNDIOutputParam<FVector3f> OutTangent(Context);
+	FNDIOutputParam<FVector3f> OutBinorm(Context);
+	FNDIOutputParam<FVector3f> OutNorm(Context);
 
 	if (InstData->bMeshValid)
 	{
@@ -2228,7 +2228,7 @@ void UNiagaraDataInterfaceStaticMesh::GetTriCoordColor(FVectorVMContext& Context
 	VectorVM::FUserPtrHandler<FNDIStaticMesh_InstanceData> InstData(Context);
 
 	FNDIInputParam<int32> TriParam(Context);
-	FNDIInputParam<FVector> BaryParam(Context);
+	FNDIInputParam<FVector3f> BaryParam(Context);
 
 	FNDIOutputParam<FLinearColor> OutColor(Context);
 
@@ -2274,7 +2274,7 @@ void UNiagaraDataInterfaceStaticMesh::GetTriCoordUV(FVectorVMContext& Context)
 	VectorVM::FUserPtrHandler<FNDIStaticMesh_InstanceData> InstData(Context);
 
 	FNDIInputParam<int32> TriParam(Context);
-	FNDIInputParam<FVector> BaryParam(Context);
+	FNDIInputParam<FVector3f> BaryParam(Context);
 	FNDIInputParam<int32> UVSetParam(Context);
 
 	FNDIOutputParam<FVector2D> OutUV(Context);
@@ -2318,10 +2318,10 @@ void UNiagaraDataInterfaceStaticMesh::GetTriCoordPositionAndVelocity(FVectorVMCo
 	VectorVM::FUserPtrHandler<FNDIStaticMesh_InstanceData> InstData(Context);
 
 	FNDIInputParam<int32> TriParam(Context);
-	FNDIInputParam<FVector> BaryParam(Context);
+	FNDIInputParam<FVector3f> BaryParam(Context);
 
-	FNDIOutputParam<FVector> OutPos(Context);
-	FNDIOutputParam<FVector> OutVel(Context);
+	FNDIOutputParam<FVector3f> OutPos(Context);
+	FNDIOutputParam<FVector3f> OutVel(Context);
 
 	if (InstData->bMeshValid)
 	{
@@ -2428,7 +2428,7 @@ void UNiagaraDataInterfaceStaticMesh::GetWorldVelocity(FVectorVMContext& Context
 {
 	VectorVM::FUserPtrHandler<FNDIStaticMesh_InstanceData> InstData(Context);
 
-	FNDIOutputParam<FVector> OutVel(Context);
+	FNDIOutputParam<FVector3f> OutVel(Context);
 
 	FVector Velocity(0.0f, 0.0f, 0.0f);
 	if (InstData->bUsePhysicsVelocity)
@@ -2459,7 +2459,7 @@ void UNiagaraDataInterfaceStaticMesh::GetVertexPosition(FVectorVMContext& Contex
 	TransformHandlerType TransformHandler;
 	FNDIInputParam<int32> VertexIndexParam(Context);
 
-	FNDIOutputParam<FVector> OutPos(Context);
+	FNDIOutputParam<FVector3f> OutPos(Context);
 
 	if (InstData->bMeshValid)
 	{
@@ -2535,9 +2535,9 @@ void UNiagaraDataInterfaceStaticMesh::GetSocketTransform(FVectorVMContext& Conte
 	VectorVM::FUserPtrHandler<FNDIStaticMesh_InstanceData> InstData(Context);
 	FNDIInputParam<int32> InSocketIndex(Context);
 
-	FNDIOutputParam<FVector> OutTranslate(Context);
+	FNDIOutputParam<FVector3f> OutTranslate(Context);
 	FNDIOutputParam<FQuat> OutRotate(Context);
-	FNDIOutputParam<FVector> OutScale(Context);
+	FNDIOutputParam<FVector3f> OutScale(Context);
 
 	const FQuat InstRotation = bWorldSpace ? InstData->Rotation : FQuat::Identity;
 
@@ -2573,9 +2573,9 @@ void UNiagaraDataInterfaceStaticMesh::GetFilteredSocketTransform(FVectorVMContex
 	VectorVM::FUserPtrHandler<FNDIStaticMesh_InstanceData> InstData(Context);
 	FNDIInputParam<int32> InSocketIndex(Context);
 
-	FNDIOutputParam<FVector> OutTranslate(Context);
+	FNDIOutputParam<FVector3f> OutTranslate(Context);
 	FNDIOutputParam<FQuat> OutRotate(Context);
-	FNDIOutputParam<FVector> OutScale(Context);
+	FNDIOutputParam<FVector3f> OutScale(Context);
 
 	const FQuat InstRotation = bWorldSpace ? InstData->Rotation : FQuat::Identity;
 
@@ -2611,9 +2611,9 @@ void UNiagaraDataInterfaceStaticMesh::GetUnfilteredSocketTransform(FVectorVMCont
 	VectorVM::FUserPtrHandler<FNDIStaticMesh_InstanceData> InstData(Context);
 	FNDIInputParam<int32> InSocketIndex(Context);
 
-	FNDIOutputParam<FVector> OutTranslate(Context);
+	FNDIOutputParam<FVector3f> OutTranslate(Context);
 	FNDIOutputParam<FQuat> OutRotate(Context);
-	FNDIOutputParam<FVector> OutScale(Context);
+	FNDIOutputParam<FVector3f> OutScale(Context);
 
 	const FQuat InstRotation = bWorldSpace ? InstData->Rotation : FQuat::Identity;
 
diff --git a/Engine/Plugins/FX/Niagara/Source/Niagara/Private/NiagaraDebugHud.cpp b/Engine/Plugins/FX/Niagara/Source/Niagara/Private/NiagaraDebugHud.cpp
--- a/Engine/Plugins/FX/Niagara/Source/Niagara/Private/NiagaraDebugHud.cpp
+++ b/Engine/Plugins/FX/Niagara/Source/Niagara/Private/NiagaraDebugHud.cpp
@@ -58,7 +58,7 @@ namespace NiagaraDebugLocal
 		TArray<FNiagaraVariableWithOffset> UserVariables;						// Exposed user parameters which will pull from the component
 
 		TArray<TArray<FNiagaraDataSetDebugAccessor>> ParticleVariables;			// Per Emitter Particle variables
-		TArray<FNiagaraDataSetAccessor<FVector>> ParticlePositionAccessors;		// Only valid if we have particle attributes
+		TArray<FNiagaraDataSetAccessor<FVector3f>> ParticlePositionAccessors;		// Only valid if we have particle attributes
 	};
 
 	static TMap<TWeakObjectPtr<UNiagaraSystem>, FCachedVariables> GCachedSystemVariables;
diff --git a/Engine/Plugins/FX/Niagara/Source/Niagara/Private/NiagaraFunctionLibrary.cpp b/Engine/Plugins/FX/Niagara/Source/Niagara/Private/NiagaraFunctionLibrary.cpp
--- a/Engine/Plugins/FX/Niagara/Source/Niagara/Private/NiagaraFunctionLibrary.cpp
+++ b/Engine/Plugins/FX/Niagara/Source/Niagara/Private/NiagaraFunctionLibrary.cpp
@@ -683,29 +683,29 @@ struct FVectorKernelFastDot4
 
 		VectorVM::FExternalFuncRegisterHandler<float> OutValue(Context);
 
-		VectorRegister* RESTRICT AX = (VectorRegister *)InVecA[0].GetDest();
-		VectorRegister* RESTRICT AY = (VectorRegister *)InVecA[1].GetDest();
-		VectorRegister* RESTRICT AZ = (VectorRegister *)InVecA[2].GetDest();
-		VectorRegister* RESTRICT AW = (VectorRegister *)InVecA[3].GetDest();
+		VectorRegister4Float* RESTRICT AX = (VectorRegister4Float *)InVecA[0].GetDest();
+		VectorRegister4Float* RESTRICT AY = (VectorRegister4Float *)InVecA[1].GetDest();
+		VectorRegister4Float* RESTRICT AZ = (VectorRegister4Float *)InVecA[2].GetDest();
+		VectorRegister4Float* RESTRICT AW = (VectorRegister4Float *)InVecA[3].GetDest();
 
-		VectorRegister* RESTRICT BX = (VectorRegister *)InVecB[0].GetDest();
-		VectorRegister* RESTRICT BY = (VectorRegister *)InVecB[1].GetDest();
-		VectorRegister* RESTRICT BZ = (VectorRegister *)InVecB[2].GetDest();
-		VectorRegister* RESTRICT BW = (VectorRegister *)InVecB[3].GetDest();
-		VectorRegister* RESTRICT Out = (VectorRegister *)OutValue.GetDest();
+		VectorRegister4Float* RESTRICT BX = (VectorRegister4Float *)InVecB[0].GetDest();
+		VectorRegister4Float* RESTRICT BY = (VectorRegister4Float *)InVecB[1].GetDest();
+		VectorRegister4Float* RESTRICT BZ = (VectorRegister4Float *)InVecB[2].GetDest();
+		VectorRegister4Float* RESTRICT BW = (VectorRegister4Float *)InVecB[3].GetDest();
+		VectorRegister4Float* RESTRICT Out = (VectorRegister4Float *)OutValue.GetDest();
 
 		const int32 Loops = Context.GetNumLoops<4>();
 		for (int32 i = 0; i < Loops; ++i)
 		{
 
-			VectorRegister AVX0= VectorLoadAligned(&AX[i]);
-			VectorRegister AVY0= VectorLoadAligned(&AY[i]);
-			VectorRegister AVZ0= VectorLoadAligned(&AZ[i]);
-			VectorRegister AVW0= VectorLoadAligned(&AW[i]);
-			VectorRegister BVX0= VectorLoadAligned(&BX[i]);
-			VectorRegister BVY0= VectorLoadAligned(&BY[i]);
-			VectorRegister BVZ0= VectorLoadAligned(&BZ[i]);
-			VectorRegister BVW0= VectorLoadAligned(&BW[i]);
+			VectorRegister4Float AVX0= VectorLoadAligned(&AX[i]);
+			VectorRegister4Float AVY0= VectorLoadAligned(&AY[i]);
+			VectorRegister4Float AVZ0= VectorLoadAligned(&AZ[i]);
+			VectorRegister4Float AVW0= VectorLoadAligned(&AW[i]);
+			VectorRegister4Float BVX0= VectorLoadAligned(&BX[i]);
+			VectorRegister4Float BVY0= VectorLoadAligned(&BY[i]);
+			VectorRegister4Float BVZ0= VectorLoadAligned(&BZ[i]);
+			VectorRegister4Float BVW0= VectorLoadAligned(&BW[i]);
 
 			/*
 				 R[19] = :mul(R[21], R[25]);
@@ -713,10 +713,10 @@ struct FVectorKernelFastDot4
 				 R[19] = :mad(R[22], R[26], R[21]);
 				 R[20] = :mad(R[23], R[27], R[19]);
 			*/
-			VectorRegister AMBX0	= VectorMultiply(AVX0, BVX0);
-			VectorRegister AMBXY0= VectorMultiplyAdd(AVY0, BVY0, AMBX0);
-			VectorRegister AMBXYZ0= VectorMultiplyAdd(AVZ0, BVZ0, AMBXY0);
-			VectorRegister AMBXYZW0= VectorMultiplyAdd(AVW0, BVW0, AMBXYZ0);
+			VectorRegister4Float AMBX0	= VectorMultiply(AVX0, BVX0);
+			VectorRegister4Float AMBXY0= VectorMultiplyAdd(AVY0, BVY0, AMBX0);
+			VectorRegister4Float AMBXYZ0= VectorMultiplyAdd(AVZ0, BVZ0, AMBXY0);
+			VectorRegister4Float AMBXYZW0= VectorMultiplyAdd(AVW0, BVW0, AMBXYZ0);
 			VectorStoreAligned(AMBXYZW0, &Out[i]) ;
 
 			/*
@@ -1461,37 +1461,37 @@ struct FVectorKernel_SolveVelocitiesAndForces
 			*OutParticlesVelocityZ.GetDestAndAdvance() = ParticleVelocity.Z;
 		}
 #else
-		const VectorRegister EngineDeltaTime = VectorSetFloat1(InEngineDeltaTime.Get());
-		const VectorRegister MassMin = VectorSetFloat1(0.0001f);
+		const VectorRegister4Float EngineDeltaTime = VectorSetFloat1(InEngineDeltaTime.Get());
+		const VectorRegister4Float MassMin = VectorSetFloat1(0.0001f);
 
 		for (int i=0; i < Context.GetNumLoops<4>(); ++i)
 		{
 			// Gather values
-			VectorRegister PhysicsForceX		= VectorLoad(InPhysicsForceX.GetDestAndAdvance());
-			VectorRegister PhysicsForceY		= VectorLoad(InPhysicsForceY.GetDestAndAdvance());
-			VectorRegister PhysicsForceZ		= VectorLoad(InPhysicsForceZ.GetDestAndAdvance());
-			VectorRegister PhysicsDrag			= VectorLoad(InPhysicsDrag.GetDestAndAdvance());
-
-			VectorRegister ParticlesMass		= VectorLoad(InParticlesMass.GetDestAndAdvance());
-			VectorRegister ParticlesPositionX	= VectorLoad(InParticlesPositionX.GetDestAndAdvance());
-			VectorRegister ParticlesPositionY	= VectorLoad(InParticlesPositionY.GetDestAndAdvance());
-			VectorRegister ParticlesPositionZ	= VectorLoad(InParticlesPositionZ.GetDestAndAdvance());
-			VectorRegister ParticlesVelocityX	= VectorLoad(InParticlesVelocityX.GetDestAndAdvance());
-			VectorRegister ParticlesVelocityY	= VectorLoad(InParticlesVelocityY.GetDestAndAdvance());
-			VectorRegister ParticlesVelocityZ	= VectorLoad(InParticlesVelocityZ.GetDestAndAdvance());
+			VectorRegister4Float PhysicsForceX		= VectorLoad(InPhysicsForceX.GetDestAndAdvance());
+			VectorRegister4Float PhysicsForceY		= VectorLoad(InPhysicsForceY.GetDestAndAdvance());
+			VectorRegister4Float PhysicsForceZ		= VectorLoad(InPhysicsForceZ.GetDestAndAdvance());
+			VectorRegister4Float PhysicsDrag			= VectorLoad(InPhysicsDrag.GetDestAndAdvance());
+
+			VectorRegister4Float ParticlesMass		= VectorLoad(InParticlesMass.GetDestAndAdvance());
+			VectorRegister4Float ParticlesPositionX	= VectorLoad(InParticlesPositionX.GetDestAndAdvance());
+			VectorRegister4Float ParticlesPositionY	= VectorLoad(InParticlesPositionY.GetDestAndAdvance());
+			VectorRegister4Float ParticlesPositionZ	= VectorLoad(InParticlesPositionZ.GetDestAndAdvance());
+			VectorRegister4Float ParticlesVelocityX	= VectorLoad(InParticlesVelocityX.GetDestAndAdvance());
+			VectorRegister4Float ParticlesVelocityY	= VectorLoad(InParticlesVelocityY.GetDestAndAdvance());
+			VectorRegister4Float ParticlesVelocityZ	= VectorLoad(InParticlesVelocityZ.GetDestAndAdvance());
 
 			VectorStore(ParticlesVelocityX, OutParticlesPreviousVelocityX.GetDestAndAdvance());
 			VectorStore(ParticlesVelocityY, OutParticlesPreviousVelocityY.GetDestAndAdvance());
 			VectorStore(ParticlesVelocityZ, OutParticlesPreviousVelocityZ.GetDestAndAdvance());
 
 			// Apply velocity
-			const VectorRegister OOParticleMassDT = VectorMultiply(VectorReciprocal(VectorMax(ParticlesMass, MassMin)), EngineDeltaTime);
+			const VectorRegister4Float OOParticleMassDT = VectorMultiply(VectorReciprocal(VectorMax(ParticlesMass, MassMin)), EngineDeltaTime);
 			ParticlesVelocityX = VectorMultiplyAdd(PhysicsForceX, OOParticleMassDT, ParticlesVelocityX);
 			ParticlesVelocityY = VectorMultiplyAdd(PhysicsForceY, OOParticleMassDT, ParticlesVelocityY);
 			ParticlesVelocityZ = VectorMultiplyAdd(PhysicsForceZ, OOParticleMassDT, ParticlesVelocityZ);
 
 			// Apply Drag
-			VectorRegister ClampedDrag = VectorMultiply(PhysicsDrag, EngineDeltaTime);
+			VectorRegister4Float ClampedDrag = VectorMultiply(PhysicsDrag, EngineDeltaTime);
 			ClampedDrag = VectorMax(VectorMin(ClampedDrag, VectorOne()), VectorZero());
 			ClampedDrag = VectorNegate(ClampedDrag);
 
@@ -1526,45 +1526,45 @@ struct FVectorKernel_SolveVelocitiesAndForces
 	FORCEINLINE static void ExecOptimized(FVectorVMContext& Context)
 	{
 #if 0
-		const VectorRegister MassMin = VectorSetFloat1(0.0001f);
-		const VectorRegister EngineDeltaTime = VectorSetFloat1(VectorVM::FExternalFuncInputHandler<float>(Context).Get());
-		VectorRegister PhysicsForceX = VectorLoadFloat1(VectorVM::FExternalFuncInputHandler<VectorRegister>(Context).GetDest());
-		VectorRegister PhysicsForceY = VectorLoadFloat1(VectorVM::FExternalFuncInputHandler<VectorRegister>(Context).GetDest());
-		VectorRegister PhysicsForceZ = VectorLoadFloat1(VectorVM::FExternalFuncInputHandler<VectorRegister>(Context).GetDest());
-		VectorRegister PhysicsDrag = VectorLoadFloat1(VectorVM::FExternalFuncInputHandler<VectorRegister>(Context).GetDest());
-		VectorRegister ParticlesMass = VectorLoadFloat1(VectorVM::FExternalFuncInputHandler<VectorRegister>(Context).GetDest());
-		const VectorRegister* RESTRICT InParticlesPositionX = VectorVM::FExternalFuncInputHandler<VectorRegister>(Context).GetDest();
-		const VectorRegister* RESTRICT InParticlesPositionY = VectorVM::FExternalFuncInputHandler<VectorRegister>(Context).GetDest();
-		const VectorRegister* RESTRICT InParticlesPositionZ = VectorVM::FExternalFuncInputHandler<VectorRegister>(Context).GetDest();
-		const VectorRegister* RESTRICT InParticlesVelocityX = VectorVM::FExternalFuncInputHandler<VectorRegister>(Context).GetDest();
-		const VectorRegister* RESTRICT InParticlesVelocityY = VectorVM::FExternalFuncInputHandler<VectorRegister>(Context).GetDest();
-		const VectorRegister* RESTRICT InParticlesVelocityZ = VectorVM::FExternalFuncInputHandler<VectorRegister>(Context).GetDest();
-
-		VectorRegister* RESTRICT OutParticlesPositionX = VectorVM::FExternalFuncRegisterHandler<VectorRegister>(Context).GetDest();
-		VectorRegister* RESTRICT OutParticlesPositionY = VectorVM::FExternalFuncRegisterHandler<VectorRegister>(Context).GetDest();
-		VectorRegister* RESTRICT OutParticlesPositionZ = VectorVM::FExternalFuncRegisterHandler<VectorRegister>(Context).GetDest();
-		VectorRegister* RESTRICT OutParticlesVelocityX = VectorVM::FExternalFuncRegisterHandler<VectorRegister>(Context).GetDest();
-		VectorRegister* RESTRICT OutParticlesVelocityY = VectorVM::FExternalFuncRegisterHandler<VectorRegister>(Context).GetDest();
-		VectorRegister* RESTRICT OutParticlesVelocityZ = VectorVM::FExternalFuncRegisterHandler<VectorRegister>(Context).GetDest();
-		VectorRegister* RESTRICT OutParticlesPreviousVelocityX = VectorVM::FExternalFuncRegisterHandler<VectorRegister>(Context).GetDest();
-		VectorRegister* RESTRICT OutParticlesPreviousVelocityY = VectorVM::FExternalFuncRegisterHandler<VectorRegister>(Context).GetDest();
-		VectorRegister* RESTRICT OutParticlesPreviousVelocityZ = VectorVM::FExternalFuncRegisterHandler<VectorRegister>(Context).GetDest();
-
-		const VectorRegister OOParticleMassDT = VectorMultiply(VectorReciprocal(VectorMax(ParticlesMass, MassMin)), EngineDeltaTime);
-
-		VectorRegister ClampedDrag = VectorMultiply(PhysicsDrag, EngineDeltaTime);
+		const VectorRegister4Float MassMin = VectorSetFloat1(0.0001f);
+		const VectorRegister4Float EngineDeltaTime = VectorSetFloat1(VectorVM::FExternalFuncInputHandler<float>(Context).Get());
+		VectorRegister4Float PhysicsForceX = VectorLoadFloat1(VectorVM::FExternalFuncInputHandler<VectorRegister4Float>(Context).GetDest());
+		VectorRegister4Float PhysicsForceY = VectorLoadFloat1(VectorVM::FExternalFuncInputHandler<VectorRegister4Float>(Context).GetDest());
+		VectorRegister4Float PhysicsForceZ = VectorLoadFloat1(VectorVM::FExternalFuncInputHandler<VectorRegister4Float>(Context).GetDest());
+		VectorRegister4Float PhysicsDrag = VectorLoadFloat1(VectorVM::FExternalFuncInputHandler<VectorRegister4Float>(Context).GetDest());
+		VectorRegister4Float ParticlesMass = VectorLoadFloat1(VectorVM::FExternalFuncInputHandler<VectorRegister4Float>(Context).GetDest());
+		const VectorRegister4Float* RESTRICT InParticlesPositionX = VectorVM::FExternalFuncInputHandler<VectorRegister4Float>(Context).GetDest();
+		const VectorRegister4Float* RESTRICT InParticlesPositionY = VectorVM::FExternalFuncInputHandler<VectorRegister4Float>(Context).GetDest();
+		const VectorRegister4Float* RESTRICT InParticlesPositionZ = VectorVM::FExternalFuncInputHandler<VectorRegister4Float>(Context).GetDest();
+		const VectorRegister4Float* RESTRICT InParticlesVelocityX = VectorVM::FExternalFuncInputHandler<VectorRegister4Float>(Context).GetDest();
+		const VectorRegister4Float* RESTRICT InParticlesVelocityY = VectorVM::FExternalFuncInputHandler<VectorRegister4Float>(Context).GetDest();
+		const VectorRegister4Float* RESTRICT InParticlesVelocityZ = VectorVM::FExternalFuncInputHandler<VectorRegister4Float>(Context).GetDest();
+
+		VectorRegister4Float* RESTRICT OutParticlesPositionX = VectorVM::FExternalFuncRegisterHandler<VectorRegister4Float>(Context).GetDest();
+		VectorRegister4Float* RESTRICT OutParticlesPositionY = VectorVM::FExternalFuncRegisterHandler<VectorRegister4Float>(Context).GetDest();
+		VectorRegister4Float* RESTRICT OutParticlesPositionZ = VectorVM::FExternalFuncRegisterHandler<VectorRegister4Float>(Context).GetDest();
+		VectorRegister4Float* RESTRICT OutParticlesVelocityX = VectorVM::FExternalFuncRegisterHandler<VectorRegister4Float>(Context).GetDest();
+		VectorRegister4Float* RESTRICT OutParticlesVelocityY = VectorVM::FExternalFuncRegisterHandler<VectorRegister4Float>(Context).GetDest();
+		VectorRegister4Float* RESTRICT OutParticlesVelocityZ = VectorVM::FExternalFuncRegisterHandler<VectorRegister4Float>(Context).GetDest();
+		VectorRegister4Float* RESTRICT OutParticlesPreviousVelocityX = VectorVM::FExternalFuncRegisterHandler<VectorRegister4Float>(Context).GetDest();
+		VectorRegister4Float* RESTRICT OutParticlesPreviousVelocityY = VectorVM::FExternalFuncRegisterHandler<VectorRegister4Float>(Context).GetDest();
+		VectorRegister4Float* RESTRICT OutParticlesPreviousVelocityZ = VectorVM::FExternalFuncRegisterHandler<VectorRegister4Float>(Context).GetDest();
+
+		const VectorRegister4Float OOParticleMassDT = VectorMultiply(VectorReciprocal(VectorMax(ParticlesMass, MassMin)), EngineDeltaTime);
+
+		VectorRegister4Float ClampedDrag = VectorMultiply(PhysicsDrag, EngineDeltaTime);
 		ClampedDrag = VectorMax(VectorMin(ClampedDrag, VectorOne()), VectorZero());
 		ClampedDrag = VectorNegate(ClampedDrag);
 
 		for (int i = 0; i < Context.GetNumLoops<4>(); ++i)
 		{
 			// Gather values
-			VectorRegister ParticlesPositionX = VectorLoad(InParticlesPositionX + i);
-			VectorRegister ParticlesPositionY = VectorLoad(InParticlesPositionY + i);
-			VectorRegister ParticlesPositionZ = VectorLoad(InParticlesPositionZ + i);
-			VectorRegister ParticlesVelocityX = VectorLoad(InParticlesVelocityX + i);
-			VectorRegister ParticlesVelocityY = VectorLoad(InParticlesVelocityY + i);
-			VectorRegister ParticlesVelocityZ = VectorLoad(InParticlesVelocityZ + i);
+			VectorRegister4Float ParticlesPositionX = VectorLoad(InParticlesPositionX + i);
+			VectorRegister4Float ParticlesPositionY = VectorLoad(InParticlesPositionY + i);
+			VectorRegister4Float ParticlesPositionZ = VectorLoad(InParticlesPositionZ + i);
+			VectorRegister4Float ParticlesVelocityX = VectorLoad(InParticlesVelocityX + i);
+			VectorRegister4Float ParticlesVelocityY = VectorLoad(InParticlesVelocityY + i);
+			VectorRegister4Float ParticlesVelocityZ = VectorLoad(InParticlesVelocityZ + i);
 
 			VectorStore(ParticlesVelocityX, OutParticlesPreviousVelocityX + i);
 			VectorStore(ParticlesVelocityY, OutParticlesPreviousVelocityY + i);
@@ -1602,87 +1602,87 @@ struct FVectorKernel_SolveVelocitiesAndForces
 		}
 #else
 		VectorVM::FExternalFuncInputHandler<float> InEngineDeltaTime(Context);
-		VectorVM::FExternalFuncInputHandler<VectorRegister> InPhysicsForceXHandler(Context);
-		VectorVM::FExternalFuncInputHandler<VectorRegister> InPhysicsForceYHandler(Context);
-		VectorVM::FExternalFuncInputHandler<VectorRegister> InPhysicsForceZHandler(Context);
-		VectorVM::FExternalFuncInputHandler<VectorRegister> InPhysicsDragHandler(Context);
-		VectorVM::FExternalFuncInputHandler<VectorRegister> InParticlesMassHandler(Context);
-		VectorVM::FExternalFuncInputHandler<VectorRegister> InParticlesPositionXHandler(Context);
-		VectorVM::FExternalFuncInputHandler<VectorRegister> InParticlesPositionYHandler(Context);
-		VectorVM::FExternalFuncInputHandler<VectorRegister> InParticlesPositionZHandler(Context);
-		VectorVM::FExternalFuncInputHandler<VectorRegister> InParticlesVelocityXHandler(Context);
-		VectorVM::FExternalFuncInputHandler<VectorRegister> InParticlesVelocityYHandler(Context);
-		VectorVM::FExternalFuncInputHandler<VectorRegister> InParticlesVelocityZHandler(Context);
-
-		VectorVM::FExternalFuncRegisterHandler<VectorRegister> OutParticlesPositionXHandler(Context);
-		VectorVM::FExternalFuncRegisterHandler<VectorRegister> OutParticlesPositionYHandler(Context);
-		VectorVM::FExternalFuncRegisterHandler<VectorRegister> OutParticlesPositionZHandler(Context);
-		VectorVM::FExternalFuncRegisterHandler<VectorRegister> OutParticlesVelocityXHandler(Context);
-		VectorVM::FExternalFuncRegisterHandler<VectorRegister> OutParticlesVelocityYHandler(Context);
-		VectorVM::FExternalFuncRegisterHandler<VectorRegister> OutParticlesVelocityZHandler(Context);
-		VectorVM::FExternalFuncRegisterHandler<VectorRegister> OutParticlesPreviousVelocityXHandler(Context);
-		VectorVM::FExternalFuncRegisterHandler<VectorRegister> OutParticlesPreviousVelocityYHandler(Context);
-		VectorVM::FExternalFuncRegisterHandler<VectorRegister> OutParticlesPreviousVelocityZHandler(Context);
-
-		const VectorRegister EngineDeltaTime = VectorSetFloat1(InEngineDeltaTime.Get());
-		const VectorRegister* RESTRICT InPhysicsForceX = InPhysicsForceXHandler.GetDest();
-		const VectorRegister* RESTRICT InPhysicsForceY = InPhysicsForceYHandler.GetDest();
-		const VectorRegister* RESTRICT InPhysicsForceZ = InPhysicsForceZHandler.GetDest();
-		const VectorRegister* RESTRICT InPhysicsDrag = InPhysicsDragHandler.GetDest();
-		const VectorRegister* RESTRICT InParticlesMass = InParticlesMassHandler.GetDest();
-		const VectorRegister* RESTRICT InParticlesPositionX = InParticlesPositionXHandler.GetDest();
-		const VectorRegister* RESTRICT InParticlesPositionY = InParticlesPositionYHandler.GetDest();
-		const VectorRegister* RESTRICT InParticlesPositionZ = InParticlesPositionZHandler.GetDest();
-		const VectorRegister* RESTRICT InParticlesVelocityX = InParticlesVelocityXHandler.GetDest();
-		const VectorRegister* RESTRICT InParticlesVelocityY = InParticlesVelocityYHandler.GetDest();
-		const VectorRegister* RESTRICT InParticlesVelocityZ = InParticlesVelocityZHandler.GetDest();
-
-		VectorRegister* RESTRICT OutParticlesPositionX = OutParticlesPositionXHandler.GetDest();
-		VectorRegister* RESTRICT OutParticlesPositionY = OutParticlesPositionYHandler.GetDest();
-		VectorRegister* RESTRICT OutParticlesPositionZ = OutParticlesPositionZHandler.GetDest();
-		VectorRegister* RESTRICT OutParticlesVelocityX = OutParticlesVelocityXHandler.GetDest();
-		VectorRegister* RESTRICT OutParticlesVelocityY = OutParticlesVelocityYHandler.GetDest();
-		VectorRegister* RESTRICT OutParticlesVelocityZ = OutParticlesVelocityZHandler.GetDest();
-		VectorRegister* RESTRICT OutParticlesPreviousVelocityX = OutParticlesPreviousVelocityXHandler.GetDest();
-		VectorRegister* RESTRICT OutParticlesPreviousVelocityY = OutParticlesPreviousVelocityYHandler.GetDest();
-		VectorRegister* RESTRICT OutParticlesPreviousVelocityZ = OutParticlesPreviousVelocityZHandler.GetDest();
-
-		const VectorRegister MassMin = VectorSetFloat1(0.0001f);
-
-		const VectorRegister ConstantPhysicsForceX = bForceConstant ? VectorLoadFloat1(InPhysicsForceX) : VectorZero();
-		const VectorRegister ConstantPhysicsForceY = bForceConstant ? VectorLoadFloat1(InPhysicsForceY) : VectorZero();
-		const VectorRegister ConstantPhysicsForceZ = bForceConstant ? VectorLoadFloat1(InPhysicsForceZ) : VectorZero();
-		const VectorRegister ConstantPhysicsDrag = bDragConstant ? VectorLoadFloat1(InPhysicsDrag) : VectorZero();
-		const VectorRegister ConstantParticlesMass = bMassConstant ? VectorLoadFloat1(InParticlesMass) : VectorZero();
+		VectorVM::FExternalFuncInputHandler<VectorRegister4Float> InPhysicsForceXHandler(Context);
+		VectorVM::FExternalFuncInputHandler<VectorRegister4Float> InPhysicsForceYHandler(Context);
+		VectorVM::FExternalFuncInputHandler<VectorRegister4Float> InPhysicsForceZHandler(Context);
+		VectorVM::FExternalFuncInputHandler<VectorRegister4Float> InPhysicsDragHandler(Context);
+		VectorVM::FExternalFuncInputHandler<VectorRegister4Float> InParticlesMassHandler(Context);
+		VectorVM::FExternalFuncInputHandler<VectorRegister4Float> InParticlesPositionXHandler(Context);
+		VectorVM::FExternalFuncInputHandler<VectorRegister4Float> InParticlesPositionYHandler(Context);
+		VectorVM::FExternalFuncInputHandler<VectorRegister4Float> InParticlesPositionZHandler(Context);
+		VectorVM::FExternalFuncInputHandler<VectorRegister4Float> InParticlesVelocityXHandler(Context);
+		VectorVM::FExternalFuncInputHandler<VectorRegister4Float> InParticlesVelocityYHandler(Context);
+		VectorVM::FExternalFuncInputHandler<VectorRegister4Float> InParticlesVelocityZHandler(Context);
+
+		VectorVM::FExternalFuncRegisterHandler<VectorRegister4Float> OutParticlesPositionXHandler(Context);
+		VectorVM::FExternalFuncRegisterHandler<VectorRegister4Float> OutParticlesPositionYHandler(Context);
+		VectorVM::FExternalFuncRegisterHandler<VectorRegister4Float> OutParticlesPositionZHandler(Context);
+		VectorVM::FExternalFuncRegisterHandler<VectorRegister4Float> OutParticlesVelocityXHandler(Context);
+		VectorVM::FExternalFuncRegisterHandler<VectorRegister4Float> OutParticlesVelocityYHandler(Context);
+		VectorVM::FExternalFuncRegisterHandler<VectorRegister4Float> OutParticlesVelocityZHandler(Context);
+		VectorVM::FExternalFuncRegisterHandler<VectorRegister4Float> OutParticlesPreviousVelocityXHandler(Context);
+		VectorVM::FExternalFuncRegisterHandler<VectorRegister4Float> OutParticlesPreviousVelocityYHandler(Context);
+		VectorVM::FExternalFuncRegisterHandler<VectorRegister4Float> OutParticlesPreviousVelocityZHandler(Context);
+
+		const VectorRegister4Float EngineDeltaTime = VectorSetFloat1(InEngineDeltaTime.Get());
+		const VectorRegister4Float* RESTRICT InPhysicsForceX = InPhysicsForceXHandler.GetDest();
+		const VectorRegister4Float* RESTRICT InPhysicsForceY = InPhysicsForceYHandler.GetDest();
+		const VectorRegister4Float* RESTRICT InPhysicsForceZ = InPhysicsForceZHandler.GetDest();
+		const VectorRegister4Float* RESTRICT InPhysicsDrag = InPhysicsDragHandler.GetDest();
+		const VectorRegister4Float* RESTRICT InParticlesMass = InParticlesMassHandler.GetDest();
+		const VectorRegister4Float* RESTRICT InParticlesPositionX = InParticlesPositionXHandler.GetDest();
+		const VectorRegister4Float* RESTRICT InParticlesPositionY = InParticlesPositionYHandler.GetDest();
+		const VectorRegister4Float* RESTRICT InParticlesPositionZ = InParticlesPositionZHandler.GetDest();
+		const VectorRegister4Float* RESTRICT InParticlesVelocityX = InParticlesVelocityXHandler.GetDest();
+		const VectorRegister4Float* RESTRICT InParticlesVelocityY = InParticlesVelocityYHandler.GetDest();
+		const VectorRegister4Float* RESTRICT InParticlesVelocityZ = InParticlesVelocityZHandler.GetDest();
+
+		VectorRegister4Float* RESTRICT OutParticlesPositionX = OutParticlesPositionXHandler.GetDest();
+		VectorRegister4Float* RESTRICT OutParticlesPositionY = OutParticlesPositionYHandler.GetDest();
+		VectorRegister4Float* RESTRICT OutParticlesPositionZ = OutParticlesPositionZHandler.GetDest();
+		VectorRegister4Float* RESTRICT OutParticlesVelocityX = OutParticlesVelocityXHandler.GetDest();
+		VectorRegister4Float* RESTRICT OutParticlesVelocityY = OutParticlesVelocityYHandler.GetDest();
+		VectorRegister4Float* RESTRICT OutParticlesVelocityZ = OutParticlesVelocityZHandler.GetDest();
+		VectorRegister4Float* RESTRICT OutParticlesPreviousVelocityX = OutParticlesPreviousVelocityXHandler.GetDest();
+		VectorRegister4Float* RESTRICT OutParticlesPreviousVelocityY = OutParticlesPreviousVelocityYHandler.GetDest();
+		VectorRegister4Float* RESTRICT OutParticlesPreviousVelocityZ = OutParticlesPreviousVelocityZHandler.GetDest();
+
+		const VectorRegister4Float MassMin = VectorSetFloat1(0.0001f);
+
+		const VectorRegister4Float ConstantPhysicsForceX = bForceConstant ? VectorLoadFloat1(InPhysicsForceX) : VectorZero();
+		const VectorRegister4Float ConstantPhysicsForceY = bForceConstant ? VectorLoadFloat1(InPhysicsForceY) : VectorZero();
+		const VectorRegister4Float ConstantPhysicsForceZ = bForceConstant ? VectorLoadFloat1(InPhysicsForceZ) : VectorZero();
+		const VectorRegister4Float ConstantPhysicsDrag = bDragConstant ? VectorLoadFloat1(InPhysicsDrag) : VectorZero();
+		const VectorRegister4Float ConstantParticlesMass = bMassConstant ? VectorLoadFloat1(InParticlesMass) : VectorZero();
 
 		for (int i = 0; i < Context.GetNumLoops<4>(); ++i)
 		{
 			// Gather values
-			VectorRegister PhysicsForceX = bForceConstant ? ConstantPhysicsForceX : VectorLoad(InPhysicsForceX + i);
-			VectorRegister PhysicsForceY = bForceConstant ? ConstantPhysicsForceY : VectorLoad(InPhysicsForceY + i);
-			VectorRegister PhysicsForceZ = bForceConstant ? ConstantPhysicsForceZ : VectorLoad(InPhysicsForceZ + i);
-			VectorRegister PhysicsDrag = bDragConstant ? ConstantPhysicsDrag : VectorLoad(InPhysicsDrag + i);
-
-			VectorRegister ParticlesMass = bMassConstant ? ConstantParticlesMass : VectorLoad(InParticlesMass + i);
-			VectorRegister ParticlesPositionX = VectorLoad(InParticlesPositionX + i);
-			VectorRegister ParticlesPositionY = VectorLoad(InParticlesPositionY + i);
-			VectorRegister ParticlesPositionZ = VectorLoad(InParticlesPositionZ + i);
-			VectorRegister ParticlesVelocityX = VectorLoad(InParticlesVelocityX + i);
-			VectorRegister ParticlesVelocityY = VectorLoad(InParticlesVelocityY + i);
-			VectorRegister ParticlesVelocityZ = VectorLoad(InParticlesVelocityZ + i);
+			VectorRegister4Float PhysicsForceX = bForceConstant ? ConstantPhysicsForceX : VectorLoad(InPhysicsForceX + i);
+			VectorRegister4Float PhysicsForceY = bForceConstant ? ConstantPhysicsForceY : VectorLoad(InPhysicsForceY + i);
+			VectorRegister4Float PhysicsForceZ = bForceConstant ? ConstantPhysicsForceZ : VectorLoad(InPhysicsForceZ + i);
+			VectorRegister4Float PhysicsDrag = bDragConstant ? ConstantPhysicsDrag : VectorLoad(InPhysicsDrag + i);
+
+			VectorRegister4Float ParticlesMass = bMassConstant ? ConstantParticlesMass : VectorLoad(InParticlesMass + i);
+			VectorRegister4Float ParticlesPositionX = VectorLoad(InParticlesPositionX + i);
+			VectorRegister4Float ParticlesPositionY = VectorLoad(InParticlesPositionY + i);
+			VectorRegister4Float ParticlesPositionZ = VectorLoad(InParticlesPositionZ + i);
+			VectorRegister4Float ParticlesVelocityX = VectorLoad(InParticlesVelocityX + i);
+			VectorRegister4Float ParticlesVelocityY = VectorLoad(InParticlesVelocityY + i);
+			VectorRegister4Float ParticlesVelocityZ = VectorLoad(InParticlesVelocityZ + i);
 
 			VectorStore(ParticlesVelocityX, OutParticlesPreviousVelocityX + i);
 			VectorStore(ParticlesVelocityY, OutParticlesPreviousVelocityY + i);
 			VectorStore(ParticlesVelocityZ, OutParticlesPreviousVelocityZ + i);
 
 			// Apply velocity
-			const VectorRegister OOParticleMassDT = VectorMultiply(VectorReciprocal(VectorMax(ParticlesMass, MassMin)), EngineDeltaTime);
+			const VectorRegister4Float OOParticleMassDT = VectorMultiply(VectorReciprocal(VectorMax(ParticlesMass, MassMin)), EngineDeltaTime);
 			ParticlesVelocityX = VectorMultiplyAdd(PhysicsForceX, OOParticleMassDT, ParticlesVelocityX);
 			ParticlesVelocityY = VectorMultiplyAdd(PhysicsForceY, OOParticleMassDT, ParticlesVelocityY);
 			ParticlesVelocityZ = VectorMultiplyAdd(PhysicsForceZ, OOParticleMassDT, ParticlesVelocityZ);
 
 			// Apply Drag
-			VectorRegister ClampedDrag = VectorMultiply(PhysicsDrag, EngineDeltaTime);
+			VectorRegister4Float ClampedDrag = VectorMultiply(PhysicsDrag, EngineDeltaTime);
 			ClampedDrag = VectorMax(VectorMin(ClampedDrag, VectorOne()), VectorZero());
 			ClampedDrag = VectorNegate(ClampedDrag);
 
diff --git a/Engine/Plugins/FX/Niagara/Source/Niagara/Private/NiagaraModule.cpp b/Engine/Plugins/FX/Niagara/Source/Niagara/Private/NiagaraModule.cpp
--- a/Engine/Plugins/FX/Niagara/Source/Niagara/Private/NiagaraModule.cpp
+++ b/Engine/Plugins/FX/Niagara/Source/Niagara/Private/NiagaraModule.cpp
@@ -636,7 +636,7 @@ void FNiagaraTypeDefinition::Init()
 	FNiagaraTypeDefinition::HalfVec4Struct = FindObjectChecked<UScriptStruct>(NiagaraPkg, TEXT("NiagaraHalfVector4"));
 
 	FNiagaraTypeDefinition::Vec2Struct = FindObjectChecked<UScriptStruct>(CoreUObjectPkg, TEXT("Vector2D"));
-	FNiagaraTypeDefinition::Vec3Struct = FindObjectChecked<UScriptStruct>(CoreUObjectPkg, TEXT("Vector"));
+	FNiagaraTypeDefinition::Vec3Struct = FindObjectChecked<UScriptStruct>(CoreUObjectPkg, TEXT("Vector3f"));
 	FNiagaraTypeDefinition::Vec4Struct = FindObjectChecked<UScriptStruct>(CoreUObjectPkg, TEXT("Vector4"));
 	FNiagaraTypeDefinition::ColorStruct = FindObjectChecked<UScriptStruct>(CoreUObjectPkg, TEXT("LinearColor"));
 	FNiagaraTypeDefinition::QuatStruct = FindObjectChecked<UScriptStruct>(CoreUObjectPkg, TEXT("Quat"));
@@ -1033,15 +1033,15 @@ bool FNiagaraTypeDefinition::TypesAreAssignable(const FNiagaraTypeDefinition& Ty
 		bIsSupportedConversion = (TypeA == ColorDef && TypeB == Vec4Def) || (TypeB == ColorDef && TypeA == Vec4Def);
 	}
 
-	if (bIsSupportedConversion)
-	{
-		return true;
-	}
+if (bIsSupportedConversion)
+{
+	return true;
+}
 
-	return	(TypeA == NumericDef && NumericStructs.Contains(TypeB.GetScriptStruct())) ||
-			(TypeB == NumericDef && NumericStructs.Contains(TypeA.GetScriptStruct())) ||
-			(TypeA == NumericDef && (TypeB.GetStruct() == GetIntStruct()) && TypeB.GetEnum() != nullptr) ||
-			(TypeB == NumericDef && (TypeA.GetStruct() == GetIntStruct()) && TypeA.GetEnum() != nullptr);
+return	(TypeA == NumericDef && NumericStructs.Contains(TypeB.GetScriptStruct())) ||
+(TypeB == NumericDef && NumericStructs.Contains(TypeA.GetScriptStruct())) ||
+(TypeA == NumericDef && (TypeB.GetStruct() == GetIntStruct()) && TypeB.GetEnum() != nullptr) ||
+(TypeB == NumericDef && (TypeA.GetStruct() == GetIntStruct()) && TypeA.GetEnum() != nullptr);
 }
 
 bool FNiagaraTypeDefinition::IsLossyConversion(const FNiagaraTypeDefinition& TypeA, const FNiagaraTypeDefinition& TypeB)
@@ -1081,11 +1081,11 @@ FNiagaraTypeDefinition FNiagaraTypeDefinition::GetNumericOutputType(const TArray
 
 	TArray<FNiagaraTypeDefinition> SortedTypeDefinitions = TypeDefinintions;
 	SortedTypeDefinitions.Sort([&](const FNiagaraTypeDefinition& TypeA, const FNiagaraTypeDefinition& TypeB)
-	{
-		int32 AIndex = OrderedNumericTypes.IndexOfByKey(TypeA);
-		int32 BIndex = OrderedNumericTypes.IndexOfByKey(TypeB);
-		return AIndex < BIndex;
-	});
+		{
+			int32 AIndex = OrderedNumericTypes.IndexOfByKey(TypeA);
+			int32 BIndex = OrderedNumericTypes.IndexOfByKey(TypeB);
+			return AIndex < BIndex;
+		});
 
 	if (SelectionMode == ENiagaraNumericOutputTypeSelectionMode::Largest)
 	{
@@ -1128,6 +1128,11 @@ void FNiagaraTypeDefinition::PostSerialize(const FArchive& Ar)
 			ClassStructOrEnum = nullptr;
 		}
 	}
+	if (Ar.IsLoading() && ClassStructOrEnum != nullptr)
+	{
+		if (ClassStructOrEnum.GetClass()->IsChildOf(UScriptStruct::StaticClass()))
+			ClassStructOrEnum = FNiagaraTypeHelper::FindNiagaraFriendlyTopLevelStruct((UScriptStruct*)ClassStructOrEnum);
+	}
 #endif
 }
 
diff --git a/Engine/Plugins/FX/Niagara/Source/Niagara/Private/NiagaraParameterStore.cpp b/Engine/Plugins/FX/Niagara/Source/Niagara/Private/NiagaraParameterStore.cpp
--- a/Engine/Plugins/FX/Niagara/Source/Niagara/Private/NiagaraParameterStore.cpp
+++ b/Engine/Plugins/FX/Niagara/Source/Niagara/Private/NiagaraParameterStore.cpp
@@ -427,7 +427,7 @@ void FNiagaraParameterStore::CheckForNaNs()const
 		else if (Var.GetType() == FNiagaraTypeDefinition::GetMatrix4Def())
 		{
 			FMatrix Val;
-			FMemory::Memcpy(&Val, GetParameterData(Offset), sizeof(FMatrix));
+			FMemory::Memcpy(&Val, GetParameterData(Offset), sizeof(FMatrix44f));
 			bContainsNans = Val.ContainsNaN();
 		}
 
diff --git a/Engine/Plugins/FX/Niagara/Source/Niagara/Private/NiagaraRendererComponents.cpp b/Engine/Plugins/FX/Niagara/Source/Niagara/Private/NiagaraRendererComponents.cpp
--- a/Engine/Plugins/FX/Niagara/Source/Niagara/Private/NiagaraRendererComponents.cpp
+++ b/Engine/Plugins/FX/Niagara/Source/Niagara/Private/NiagaraRendererComponents.cpp
@@ -46,7 +46,7 @@ void SetVariableByType(FNiagaraVariable& DataVariable, FNiagaraDataSet& Data, in
 	else if (VarType == FNiagaraTypeDefinition::GetIntDef()) { SetValueWithAccessor<int32>(DataVariable, Data, ParticleIndex); }
 	else if (VarType == FNiagaraTypeDefinition::GetBoolDef()) { SetValueWithAccessor<FNiagaraBool>(DataVariable, Data, ParticleIndex); }
 	else if (VarType == FNiagaraTypeDefinition::GetVec2Def()) { SetValueWithAccessor<FVector2D>(DataVariable, Data, ParticleIndex); }
-	else if (VarType == FNiagaraTypeDefinition::GetVec3Def()) { SetValueWithAccessor<FVector>(DataVariable, Data, ParticleIndex); }
+	else if (VarType == FNiagaraTypeDefinition::GetVec3Def()) { SetValueWithAccessor<FVector3f>(DataVariable, Data, ParticleIndex); }
 	else if (VarType == FNiagaraTypeDefinition::GetVec4Def()) { SetValueWithAccessor<FVector4>(DataVariable, Data, ParticleIndex); }
 	else if (VarType == FNiagaraTypeDefinition::GetColorDef()) { SetValueWithAccessor<FLinearColor>(DataVariable, Data, ParticleIndex); }
 	else if (VarType == FNiagaraTypeDefinition::GetQuatDef()) { SetValueWithAccessor<FQuat>(DataVariable, Data, ParticleIndex); }
@@ -59,7 +59,7 @@ void ConvertVariableToType(const FNiagaraVariable& SourceVariable, FNiagaraVaria
 
 	if (SourceType == FNiagaraTypeDefinition::GetVec3Def() && TargetType == UNiagaraComponentRendererProperties::GetFColorDef())
 	{
-		FVector Data = SourceVariable.GetValue<FVector>();
+		FVector3f Data = SourceVariable.GetValue<FVector3f>();
 		FColor ColorData((uint8)FMath::Clamp<int32>(FMath::TruncToInt(Data.X * 255.f), 0, 255),
 			(uint8)FMath::Clamp<int32>(FMath::TruncToInt(Data.Y * 255.f), 0, 255),
 			(uint8)FMath::Clamp<int32>(FMath::TruncToInt(Data.Z * 255.f), 0, 255));
@@ -81,10 +81,15 @@ void ConvertVariableToType(const FNiagaraVariable& SourceVariable, FNiagaraVaria
 	}
 	else if (SourceType == FNiagaraTypeDefinition::GetVec3Def() && TargetType == UNiagaraComponentRendererProperties::GetFRotatorDef())
 	{
-		FVector Data = SourceVariable.GetValue<FVector>();
+		FVector3f Data = SourceVariable.GetValue<FVector3f>();
 		FRotator Rotator(Data.X, Data.Y, Data.Z);
 		TargetVariable.SetValue<FRotator>(Rotator);
 	}
+	else if (SourceType == FNiagaraTypeDefinition::GetVec3Def() && TargetType == UNiagaraComponentRendererProperties::GetFVectorDef())
+	{
+		FVector3f Data = SourceVariable.GetValue<FVector3f>();
+		TargetVariable.SetValue<FVector>(FVector(Data));
+	}
 	else if (SourceType == FNiagaraTypeDefinition::GetQuatDef() && TargetType == UNiagaraComponentRendererProperties::GetFRotatorDef())
 	{
 		FQuat Data = SourceVariable.GetValue<FQuat>();
@@ -576,9 +581,28 @@ void FNiagaraRendererComponents::TickPropertyBindings(
 		}
 
 		SetVariableByType(DataVariable, Data, ParticleIndex);
-		if (PropertyBinding.PropertyType.IsValid() && DataVariable.GetType() != PropertyBinding.PropertyType && !PropertySetter->bIgnoreConversion)
+
+		FNiagaraTypeDefinition PropertyType = PropertyBinding.PropertyType;
+		// Get property type from property setter function
+		if(!PropertyType.IsValid() && PropertySetter->Function)
+		{
+			for (FProperty* Property = PropertySetter->Function->PropertyLink; Property; Property = Property->PropertyLinkNext)
+			{
+				if (Property->HasAnyPropertyFlags(CPF_Parm) && !Property->HasAnyPropertyFlags(CPF_ReturnParm))
+				{
+					if (FStructProperty* StructProp = CastField<FStructProperty>(Property))
+					{
+						PropertyType = FNiagaraTypeDefinition(StructProp->Struct, FNiagaraTypeDefinition::EAllowUnfriendlyStruct::Allow);
+						break;
+					}
+				}
+			}
+		}
+
+		bool bForceStructConversion = !UE_LARGE_WORLD_COORDINATES_DISABLED && PropertyType.GetScriptStruct() && !FNiagaraTypeHelper::IsNiagaraFriendlyTopLevelStruct(PropertyType.GetScriptStruct());
+		if (PropertyType.IsValid() && DataVariable.GetType() != PropertyType && (!PropertySetter->bIgnoreConversion || bForceStructConversion))
 		{
-			FNiagaraVariable TargetVariable(PropertyBinding.PropertyType, DataVariable.GetName());
+			FNiagaraVariable TargetVariable(PropertyType, DataVariable.GetName());
 			ConvertVariableToType(DataVariable, TargetVariable);
 			DataVariable = TargetVariable;
 		}
diff --git a/Engine/Plugins/FX/Niagara/Source/Niagara/Private/NiagaraRendererMeshes.cpp b/Engine/Plugins/FX/Niagara/Source/Niagara/Private/NiagaraRendererMeshes.cpp
--- a/Engine/Plugins/FX/Niagara/Source/Niagara/Private/NiagaraRendererMeshes.cpp
+++ b/Engine/Plugins/FX/Niagara/Source/Niagara/Private/NiagaraRendererMeshes.cpp
@@ -742,16 +742,16 @@ FNiagaraMeshUniformBufferRef FNiagaraRendererMeshes::CreatePerViewUniformBuffer(
 				switch (i)
 				{
 				case ENiagaraMeshVFLayout::Type::Position:
-					FMemory::Memcpy(&PerViewUniformParameters.DefaultPos, ParameterBoundData + VFBoundOffsetsInParamStore[i], sizeof(FVector));
+					FMemory::Memcpy(&PerViewUniformParameters.DefaultPos, ParameterBoundData + VFBoundOffsetsInParamStore[i], sizeof(FVector3f));
 					break;
 				case ENiagaraMeshVFLayout::Type::Velocity:
-					FMemory::Memcpy(&PerViewUniformParameters.DefaultVelocity, ParameterBoundData + VFBoundOffsetsInParamStore[i], sizeof(FVector));
+					FMemory::Memcpy(&PerViewUniformParameters.DefaultVelocity, ParameterBoundData + VFBoundOffsetsInParamStore[i], sizeof(FVector3f));
 					break;
 				case ENiagaraMeshVFLayout::Type::Color:
 					FMemory::Memcpy(&PerViewUniformParameters.DefaultColor, ParameterBoundData + VFBoundOffsetsInParamStore[i], sizeof(FLinearColor));
 					break;
 				case ENiagaraMeshVFLayout::Type::Scale:
-					FMemory::Memcpy(&PerViewUniformParameters.DefaultScale, ParameterBoundData + VFBoundOffsetsInParamStore[i], sizeof(FVector));
+					FMemory::Memcpy(&PerViewUniformParameters.DefaultScale, ParameterBoundData + VFBoundOffsetsInParamStore[i], sizeof(FVector3f));
 					break;
 				case ENiagaraMeshVFLayout::Type::Rotation:
 					FMemory::Memcpy(&PerViewUniformParameters.DefaultRotation, ParameterBoundData + VFBoundOffsetsInParamStore[i], sizeof(FVector4));
@@ -784,10 +784,10 @@ FNiagaraMeshUniformBufferRef FNiagaraRendererMeshes::CreatePerViewUniformBuffer(
 					FMemory::Memcpy(&PerViewUniformParameters.DefaultCamOffset, ParameterBoundData + VFBoundOffsetsInParamStore[i], sizeof(float));
 					break;
 				case ENiagaraMeshVFLayout::Type::PrevPosition:
-					FMemory::Memcpy(&PerViewUniformParameters.DefaultPrevPos, ParameterBoundData + VFBoundOffsetsInParamStore[i], sizeof(FVector));
+					FMemory::Memcpy(&PerViewUniformParameters.DefaultPrevPos, ParameterBoundData + VFBoundOffsetsInParamStore[i], sizeof(FVector3f));
 					break;
 				case ENiagaraMeshVFLayout::Type::PrevScale:
-					FMemory::Memcpy(&PerViewUniformParameters.DefaultPrevScale, ParameterBoundData + VFBoundOffsetsInParamStore[i], sizeof(FVector));
+					FMemory::Memcpy(&PerViewUniformParameters.DefaultPrevScale, ParameterBoundData + VFBoundOffsetsInParamStore[i], sizeof(FVector3f));
 					break;
 				case ENiagaraMeshVFLayout::Type::PrevRotation:
 					FMemory::Memcpy(&PerViewUniformParameters.DefaultPrevRotation, ParameterBoundData + VFBoundOffsetsInParamStore[i], sizeof(FVector4));
@@ -1202,8 +1202,8 @@ void FNiagaraRendererMeshes::GetDynamicRayTracingInstances(FRayTracingMaterialGa
 		const bool bUseLocalSpace = UseLocalSpace(SceneProxy);
 		if (SourceMode == ENiagaraRendererSourceDataMode::Emitter)
 		{
-			FVector Pos = bUseLocalSpace ? FVector() : LocalTransform.GetOrigin();
-			FVector Scale{ 1.0f, 1.0f, 1.0f };
+			FVector3f Pos = bUseLocalSpace ? FVector() : LocalTransform.GetOrigin();
+			FVector3f Scale{ 1.0f, 1.0f, 1.0f };
 			FQuat Rot = FQuat::Identity;
 
 			if (bSetAnyBoundVars)
@@ -1212,12 +1212,12 @@ void FNiagaraRendererMeshes::GetDynamicRayTracingInstances(FRayTracingMaterialGa
 				if (VFBoundOffsetsInParamStore[ENiagaraMeshVFLayout::Type::Position] != INDEX_NONE
 					&& ParticleMeshRenderData.DynamicDataMesh->ParameterDataBound.IsValidIndex(VFBoundOffsetsInParamStore[ENiagaraMeshVFLayout::Type::Position]))
 				{
-					FMemory::Memcpy(&Pos, ParameterBoundData + VFBoundOffsetsInParamStore[ENiagaraMeshVFLayout::Type::Position], sizeof(FVector));
+					FMemory::Memcpy(&Pos, ParameterBoundData + VFBoundOffsetsInParamStore[ENiagaraMeshVFLayout::Type::Position], sizeof(FVector3f));
 				}
 				if (VFBoundOffsetsInParamStore[ENiagaraMeshVFLayout::Type::Scale] != INDEX_NONE
 					&& ParticleMeshRenderData.DynamicDataMesh->ParameterDataBound.IsValidIndex(VFBoundOffsetsInParamStore[ENiagaraMeshVFLayout::Type::Scale]))
 				{
-					FMemory::Memcpy(&Scale, ParameterBoundData + VFBoundOffsetsInParamStore[ENiagaraMeshVFLayout::Type::Scale], sizeof(FVector));
+					FMemory::Memcpy(&Scale, ParameterBoundData + VFBoundOffsetsInParamStore[ENiagaraMeshVFLayout::Type::Scale], sizeof(FVector3f));
 				}
 				if (VFBoundOffsetsInParamStore[ENiagaraMeshVFLayout::Type::Rotation] != INDEX_NONE
 					&& ParticleMeshRenderData.DynamicDataMesh->ParameterDataBound.IsValidIndex(VFBoundOffsetsInParamStore[ENiagaraMeshVFLayout::Type::Rotation]))
diff --git a/Engine/Plugins/FX/Niagara/Source/Niagara/Private/NiagaraRendererRibbons.cpp b/Engine/Plugins/FX/Niagara/Source/Niagara/Private/NiagaraRendererRibbons.cpp
--- a/Engine/Plugins/FX/Niagara/Source/Niagara/Private/NiagaraRendererRibbons.cpp
+++ b/Engine/Plugins/FX/Niagara/Source/Niagara/Private/NiagaraRendererRibbons.cpp
@@ -874,13 +874,13 @@ FNiagaraDynamicDataBase* FNiagaraRendererRibbons::GenerateDynamicData(const FNia
 			// Update the tangents for the first and last vertex, apply a reflect vector logic so that the initial and final curvature is continuous.
 			if (NumSegments > 1)
 			{
-				FVector& FirstTangent = reinterpret_cast<FVector&>(DynamicData->TangentAndDistances[StartIndex]);
-				FVector& NextToFirstTangent = reinterpret_cast<FVector&>(DynamicData->TangentAndDistances[StartIndex + 1]);
-				FirstTangent = (2.f * FVector::DotProduct(FirstTangent, NextToFirstTangent)) * FirstTangent - NextToFirstTangent;
+				FVector3f& FirstTangent =  reinterpret_cast<FVector3f&>(DynamicData->TangentAndDistances[StartIndex]);
+				FVector3f& NextToFirstTangent = reinterpret_cast<FVector3f&>(DynamicData->TangentAndDistances[StartIndex + 1]);
+				FirstTangent = (2.f * FVector3f::DotProduct(FirstTangent, NextToFirstTangent)) * FirstTangent - NextToFirstTangent;
 
-				FVector& LastTangent = reinterpret_cast<FVector&>(DynamicData->TangentAndDistances[EndIndex]);
-				FVector& PrevToLastTangent = reinterpret_cast<FVector&>(DynamicData->TangentAndDistances[EndIndex - 1]);
-				LastTangent = (2.f * FVector::DotProduct(LastTangent, PrevToLastTangent)) * LastTangent - PrevToLastTangent;
+				FVector3f& LastTangent = reinterpret_cast<FVector3f&>(DynamicData->TangentAndDistances[EndIndex]);
+				FVector3f& PrevToLastTangent = reinterpret_cast<FVector3f&>(DynamicData->TangentAndDistances[EndIndex - 1]);
+				LastTangent = (2.f * FVector3f::DotProduct(LastTangent, PrevToLastTangent)) * LastTangent - PrevToLastTangent;
 			}
 
 			// Add segment data
@@ -1464,7 +1464,7 @@ void FNiagaraRendererRibbons::GetDynamicRayTracingInstances(FRayTracingMaterialG
 			RayTracingInstance.Materials,
 			false,
 			VertexCount,
-			VertexCount * (uint32)sizeof(FVector),
+			VertexCount * (uint32)sizeof(FVector3f),
 			MeshBatch.Elements[0].NumPrimitives,
 			&RayTracingGeometry,
 			VertexBuffer,
diff --git a/Engine/Plugins/FX/Niagara/Source/Niagara/Private/NiagaraRendererSprites.cpp b/Engine/Plugins/FX/Niagara/Source/Niagara/Private/NiagaraRendererSprites.cpp
--- a/Engine/Plugins/FX/Niagara/Source/Niagara/Private/NiagaraRendererSprites.cpp
+++ b/Engine/Plugins/FX/Niagara/Source/Niagara/Private/NiagaraRendererSprites.cpp
@@ -578,13 +578,13 @@ FNiagaraSpriteUniformBufferRef FNiagaraRendererSprites::CreateViewUniformBuffer(
 				switch (i)
 				{
 				case ENiagaraSpriteVFLayout::Type::Position:
-					FMemory::Memcpy(&PerViewUniformParameters.DefaultPos, DynamicDataSprites->ParameterDataBound.GetData() + VFBoundOffsetsInParamStore[i], sizeof(FVector));
+					FMemory::Memcpy(&PerViewUniformParameters.DefaultPos, DynamicDataSprites->ParameterDataBound.GetData() + VFBoundOffsetsInParamStore[i], sizeof(FVector3f));
 					break;
 				case ENiagaraSpriteVFLayout::Type::Color:
 					FMemory::Memcpy(&PerViewUniformParameters.DefaultColor, DynamicDataSprites->ParameterDataBound.GetData() + VFBoundOffsetsInParamStore[i], sizeof(FLinearColor));
 					break;
 				case ENiagaraSpriteVFLayout::Type::Velocity:
-					FMemory::Memcpy(&PerViewUniformParameters.DefaultVelocity, DynamicDataSprites->ParameterDataBound.GetData() + VFBoundOffsetsInParamStore[i], sizeof(FVector));
+					FMemory::Memcpy(&PerViewUniformParameters.DefaultVelocity, DynamicDataSprites->ParameterDataBound.GetData() + VFBoundOffsetsInParamStore[i], sizeof(FVector3f));
 					break;
 				case ENiagaraSpriteVFLayout::Type::Rotation:
 					FMemory::Memcpy(&PerViewUniformParameters.DefaultRotation, DynamicDataSprites->ParameterDataBound.GetData() + VFBoundOffsetsInParamStore[i], sizeof(float));
@@ -593,10 +593,10 @@ FNiagaraSpriteUniformBufferRef FNiagaraRendererSprites::CreateViewUniformBuffer(
 					FMemory::Memcpy(&PerViewUniformParameters.DefaultSize, DynamicDataSprites->ParameterDataBound.GetData() + VFBoundOffsetsInParamStore[i], sizeof(FVector2D));
 					break;
 				case ENiagaraSpriteVFLayout::Type::Facing:
-					FMemory::Memcpy(&PerViewUniformParameters.DefaultFacing, DynamicDataSprites->ParameterDataBound.GetData() + VFBoundOffsetsInParamStore[i], sizeof(FVector));
+					FMemory::Memcpy(&PerViewUniformParameters.DefaultFacing, DynamicDataSprites->ParameterDataBound.GetData() + VFBoundOffsetsInParamStore[i], sizeof(FVector3f));
 					break;
 				case ENiagaraSpriteVFLayout::Type::Alignment:
-					FMemory::Memcpy(&PerViewUniformParameters.DefaultAlignment, DynamicDataSprites->ParameterDataBound.GetData() + VFBoundOffsetsInParamStore[i], sizeof(FVector));
+					FMemory::Memcpy(&PerViewUniformParameters.DefaultAlignment, DynamicDataSprites->ParameterDataBound.GetData() + VFBoundOffsetsInParamStore[i], sizeof(FVector3f));
 					break;
 				case ENiagaraSpriteVFLayout::Type::SubImage:
 					FMemory::Memcpy(&PerViewUniformParameters.DefaultSubImage, DynamicDataSprites->ParameterDataBound.GetData() + VFBoundOffsetsInParamStore[i], sizeof(float));
@@ -636,10 +636,10 @@ FNiagaraSpriteUniformBufferRef FNiagaraRendererSprites::CreateViewUniformBuffer(
 					FMemory::Memcpy(&PerViewUniformParameters.DefaultNormAge, DynamicDataSprites->ParameterDataBound.GetData() + VFBoundOffsetsInParamStore[i], sizeof(float));
 					break;
 				case ENiagaraSpriteVFLayout::Type::PrevPosition:
-					FMemory::Memcpy(&PerViewUniformParameters.DefaultPrevPos, DynamicDataSprites->ParameterDataBound.GetData() + VFBoundOffsetsInParamStore[i], sizeof(FVector));
+					FMemory::Memcpy(&PerViewUniformParameters.DefaultPrevPos, DynamicDataSprites->ParameterDataBound.GetData() + VFBoundOffsetsInParamStore[i], sizeof(FVector3f));
 					break;
 				case ENiagaraSpriteVFLayout::Type::PrevVelocity:
-					FMemory::Memcpy(&PerViewUniformParameters.DefaultPrevVelocity, DynamicDataSprites->ParameterDataBound.GetData() + VFBoundOffsetsInParamStore[i], sizeof(FVector));
+					FMemory::Memcpy(&PerViewUniformParameters.DefaultPrevVelocity, DynamicDataSprites->ParameterDataBound.GetData() + VFBoundOffsetsInParamStore[i], sizeof(FVector3f));
 					break;
 				case ENiagaraSpriteVFLayout::Type::PrevRotation:
 					FMemory::Memcpy(&PerViewUniformParameters.DefaultPrevRotation, DynamicDataSprites->ParameterDataBound.GetData() + VFBoundOffsetsInParamStore[i], sizeof(float));
@@ -648,10 +648,10 @@ FNiagaraSpriteUniformBufferRef FNiagaraRendererSprites::CreateViewUniformBuffer(
 					FMemory::Memcpy(&PerViewUniformParameters.DefaultPrevSize, DynamicDataSprites->ParameterDataBound.GetData() + VFBoundOffsetsInParamStore[i], sizeof(FVector2D));
 					break;
 				case ENiagaraSpriteVFLayout::Type::PrevFacing:
-					FMemory::Memcpy(&PerViewUniformParameters.DefaultPrevFacing, DynamicDataSprites->ParameterDataBound.GetData() + VFBoundOffsetsInParamStore[i], sizeof(FVector));
+					FMemory::Memcpy(&PerViewUniformParameters.DefaultPrevFacing, DynamicDataSprites->ParameterDataBound.GetData() + VFBoundOffsetsInParamStore[i], sizeof(FVector3f));
 					break;
 				case ENiagaraSpriteVFLayout::Type::PrevAlignment:
-					FMemory::Memcpy(&PerViewUniformParameters.DefaultPrevAlignment, DynamicDataSprites->ParameterDataBound.GetData() + VFBoundOffsetsInParamStore[i], sizeof(FVector));
+					FMemory::Memcpy(&PerViewUniformParameters.DefaultPrevAlignment, DynamicDataSprites->ParameterDataBound.GetData() + VFBoundOffsetsInParamStore[i], sizeof(FVector3f));
 					break;
 				case ENiagaraSpriteVFLayout::Type::PrevCameraOffset:
 					FMemory::Memcpy(&PerViewUniformParameters.DefaultPrevCamOffset, DynamicDataSprites->ParameterDataBound.GetData() + VFBoundOffsetsInParamStore[i], sizeof(float));
@@ -850,12 +850,12 @@ void FNiagaraRendererSprites::GetDynamicMeshElements(const TArray<const FSceneVi
 			if (SourceMode == ENiagaraRendererSourceDataMode::Emitter && bEnableDistanceCulling)
 			{
 				FVector ViewOrigin = View->ViewMatrices.GetViewOrigin();
-				FVector RefPosition = SceneProxy->GetLocalToWorld().GetOrigin();
+				FVector3f RefPosition = SceneProxy->GetLocalToWorld().GetOrigin();
 				const int32 BoundPosOffset = VFBoundOffsetsInParamStore[ENiagaraSpriteVFLayout::Type::Position];
 				if (BoundPosOffset != INDEX_NONE && ParticleSpriteRenderData.DynamicDataSprites->ParameterDataBound.IsValidIndex(BoundPosOffset))
 				{
 					// retrieve the reference position from the parameter store
-					FMemory::Memcpy(&RefPosition, ParticleSpriteRenderData.DynamicDataSprites->ParameterDataBound.GetData() + BoundPosOffset, sizeof(FVector));
+					FMemory::Memcpy(&RefPosition, ParticleSpriteRenderData.DynamicDataSprites->ParameterDataBound.GetData() + BoundPosOffset, sizeof(FVector3f));
 					if (UseLocalSpace(SceneProxy))
 					{
 						RefPosition = SceneProxy->GetLocalToWorld().TransformPosition(RefPosition);
@@ -1037,7 +1037,7 @@ void FNiagaraRendererSprites::GetDynamicRayTracingInstances(FRayTracingMaterialG
 				RayTracingInstance.Materials,
 				MeshBatch.Elements[0].NumPrimitives == 0,
 				NumVerticesPerInstance* NumInstances,
-				NumVerticesPerInstance* NumInstances* (uint32)sizeof(FVector),
+				NumVerticesPerInstance* NumInstances* (uint32)sizeof(FVector3f),
 				NumTrianglesPerInstance * NumInstances,
 				&RayTracingGeometry,
 				VertexBuffer,
@@ -1117,4 +1117,4 @@ int FNiagaraRendererSprites::GetDynamicDataSize()const
 bool FNiagaraRendererSprites::IsMaterialValid(const UMaterialInterface* Mat)const
 {
 	return Mat && Mat->CheckMaterialUsage_Concurrent(MATUSAGE_NiagaraSprites);
-}
+}
\ No newline at end of file
diff --git a/Engine/Plugins/FX/Niagara/Source/Niagara/Private/NiagaraScriptExecutionParameterStore.cpp b/Engine/Plugins/FX/Niagara/Source/Niagara/Private/NiagaraScriptExecutionParameterStore.cpp
--- a/Engine/Plugins/FX/Niagara/Source/Niagara/Private/NiagaraScriptExecutionParameterStore.cpp
+++ b/Engine/Plugins/FX/Niagara/Source/Niagara/Private/NiagaraScriptExecutionParameterStore.cpp
@@ -71,6 +71,13 @@ uint32 FNiagaraScriptExecutionParameterStore::GenerateLayoutInfoInternal(TArray<
 		InSrcOffset += sizeof(FVector3f);
 		NextMemberOffset = Members[Members.Num() - 1].DestOffset + Members[Members.Num() - 1].DestSize;
 	}
+	else if (InSrcStruct && (InSrcStruct->GetFName() == NAME_Vector || InSrcStruct->GetFName() == NAME_Vector3d))
+	{
+		uint32 StructFinalSize = (TShaderParameterTypeInfo<FVector3f>::NumRows * TShaderParameterTypeInfo<FVector3f>::NumColumns) * sizeof(float);
+		Members.Emplace(InSrcOffset, OffsetAlign(NextMemberOffset, VectorPaddedSize), StructFinalSize, VectorPaddedSize);
+		InSrcOffset += sizeof(FVector3f);
+		NextMemberOffset = Members[Members.Num() - 1].DestOffset + Members[Members.Num() - 1].DestSize;
+	}
 	else if (InSrcStruct == FNiagaraTypeDefinition::GetVec4Struct() || InSrcStruct == FNiagaraTypeDefinition::GetColorStruct() || InSrcStruct == FNiagaraTypeDefinition::GetQuatStruct())
 	{
 		uint32 StructFinalSize = (TShaderParameterTypeInfo<FVector4>::NumRows * TShaderParameterTypeInfo<FVector4>::NumColumns) * sizeof(float);
@@ -80,9 +87,9 @@ uint32 FNiagaraScriptExecutionParameterStore::GenerateLayoutInfoInternal(TArray<
 	}
 	else if (InSrcStruct == FNiagaraTypeDefinition::GetMatrix4Struct())
 	{
-		uint32 StructFinalSize = (TShaderParameterTypeInfo<FMatrix>::NumRows * TShaderParameterTypeInfo<FMatrix>::NumColumns) * sizeof(float);
-		Members.Emplace(InSrcOffset, Align(NextMemberOffset, TShaderParameterTypeInfo<FMatrix>::Alignment), StructFinalSize, StructFinalSize);
-		InSrcOffset += sizeof(FMatrix);
+		uint32 StructFinalSize = (TShaderParameterTypeInfo<FMatrix44f>::NumRows * TShaderParameterTypeInfo<FMatrix44f>::NumColumns) * sizeof(float);
+		Members.Emplace(InSrcOffset, Align(NextMemberOffset, TShaderParameterTypeInfo<FMatrix44f>::Alignment), StructFinalSize, StructFinalSize);
+		InSrcOffset += sizeof(FMatrix44f);
 		NextMemberOffset = Members[Members.Num() - 1].DestOffset + Members[Members.Num() - 1].DestSize;
 	}
 	else if (InSrcStruct == FNiagaraTypeDefinition::GetHalfStruct())
@@ -135,10 +142,9 @@ uint32 FNiagaraScriptExecutionParameterStore::GenerateLayoutInfoInternal(TArray<
 			{
 				Struct = FNiagaraTypeDefinition::GetBoolStruct();
 			}
-			//Should be able to support double easily enough
 			else if (FStructProperty* StructProp = CastFieldChecked<FStructProperty>(Property))
 			{
-				Struct = StructProp->Struct;
+				Struct = FNiagaraTypeHelper::FindNiagaraFriendlyTopLevelStruct(StructProp->Struct);
 			}
 			else
 			{
diff --git a/Engine/Plugins/FX/Niagara/Source/Niagara/Private/NiagaraSystemInstanceController.cpp b/Engine/Plugins/FX/Niagara/Source/Niagara/Private/NiagaraSystemInstanceController.cpp
--- a/Engine/Plugins/FX/Niagara/Source/Niagara/Private/NiagaraSystemInstanceController.cpp
+++ b/Engine/Plugins/FX/Niagara/Source/Niagara/Private/NiagaraSystemInstanceController.cpp
@@ -183,7 +183,7 @@ bool FNiagaraSystemInstanceController::GetParticleValueVec3_DebugOnly(TArray<FVe
 				int32 NumParticles = ParticleData.GetNumInstances();
 				OutValues.SetNum(NumParticles);
 
-				const auto Reader = FNiagaraDataSetAccessor<FVector>::CreateReader(Sim->GetData(), ValueName);
+				const auto Reader = FNiagaraDataSetAccessor<FVector3f>::CreateReader(Sim->GetData(), ValueName);
 				if (!Reader.IsValid())
 				{
 					return false;
@@ -367,7 +367,7 @@ void FNiagaraSystemInstanceController::SetVariable(FName InVariableName, FVector
 	}
 }
 
-void FNiagaraSystemInstanceController::SetVariable(FName InVariableName, FVector InValue)
+void FNiagaraSystemInstanceController::SetVariable(FName InVariableName, FVector3f InValue)
 {
 	if (OverrideParameters)
 	{
diff --git a/Engine/Plugins/FX/Niagara/Source/Niagara/Public/NiagaraComponentRendererProperties.h b/Engine/Plugins/FX/Niagara/Source/Niagara/Public/NiagaraComponentRendererProperties.h
--- a/Engine/Plugins/FX/Niagara/Source/Niagara/Public/NiagaraComponentRendererProperties.h
+++ b/Engine/Plugins/FX/Niagara/Source/Niagara/Public/NiagaraComponentRendererProperties.h
@@ -135,6 +135,7 @@ public:
 	NIAGARA_API static FNiagaraTypeDefinition ToNiagaraType(FProperty* Property);
 	static FNiagaraTypeDefinition GetFColorDef();
 	static FNiagaraTypeDefinition GetFRotatorDef();
+	static FNiagaraTypeDefinition GetFVectorDef();
 
 	virtual void CacheFromCompiledData(const FNiagaraDataSetCompiledData* CompiledData) override;
 
diff --git a/Engine/Plugins/FX/Niagara/Source/Niagara/Public/NiagaraLightRendererProperties.h b/Engine/Plugins/FX/Niagara/Source/Niagara/Public/NiagaraLightRendererProperties.h
--- a/Engine/Plugins/FX/Niagara/Source/Niagara/Public/NiagaraLightRendererProperties.h
+++ b/Engine/Plugins/FX/Niagara/Source/Niagara/Public/NiagaraLightRendererProperties.h
@@ -102,7 +102,7 @@ public:
 	UPROPERTY(EditAnywhere, AdvancedDisplay, Category = "Bindings")
 	FNiagaraVariableAttributeBinding RendererVisibilityTagBinding;
 
-	FNiagaraDataSetAccessor<FVector> PositionDataSetAccessor;
+	FNiagaraDataSetAccessor<FVector3f> PositionDataSetAccessor;
 	FNiagaraDataSetAccessor<FLinearColor> ColorDataSetAccessor;
 	FNiagaraDataSetAccessor<float> RadiusDataSetAccessor;
 	FNiagaraDataSetAccessor<float> ExponentDataSetAccessor;
diff --git a/Engine/Plugins/FX/Niagara/Source/Niagara/Public/NiagaraParameterStore.h b/Engine/Plugins/FX/Niagara/Source/Niagara/Public/NiagaraParameterStore.h
--- a/Engine/Plugins/FX/Niagara/Source/Niagara/Public/NiagaraParameterStore.h
+++ b/Engine/Plugins/FX/Niagara/Source/Niagara/Public/NiagaraParameterStore.h
@@ -657,6 +657,12 @@ protected:
 	friend struct FNiagaraParameterStoreToDataSetBinding;    // this should be the only class calling SetParameterByOffset
 };
 
+template<>
+FORCEINLINE bool FNiagaraParameterStore::SetParameterValue(const FVector3d& InValue, const FNiagaraVariable& Param, bool bAdd)
+{
+	return SetParameterValue((FVector3f)InValue, Param, bAdd);
+}
+
 FORCEINLINE_DEBUGGABLE void FNiagaraParameterStore::Tick()
 {
 #if NIAGARA_NAN_CHECKING
@@ -961,9 +967,9 @@ struct FNiagaraParameterDirectBinding
 };
 
 template<>
-struct FNiagaraParameterDirectBinding<FMatrix>
+struct FNiagaraParameterDirectBinding<FMatrix44f>
 {
-	mutable FMatrix* ValuePtr;
+	mutable FMatrix44f* ValuePtr;
 #if NIAGARA_VALIDATE_DIRECT_BINDINGS
 	FNiagaraParameterStore* BoundStore;
 	FNiagaraVariable BoundVariable;
@@ -977,40 +983,40 @@ struct FNiagaraParameterDirectBinding<FMatrix>
 #endif
 	{}
 
-	FMatrix* Init(FNiagaraParameterStore& InStore, const FNiagaraVariable& DestVariable)
+	FMatrix44f* Init(FNiagaraParameterStore& InStore, const FNiagaraVariable& DestVariable)
 	{
 #if NIAGARA_VALIDATE_DIRECT_BINDINGS
 		BoundStore = &InStore;
 		BoundVariable = DestVariable;
 		LayoutVersion = BoundStore->GetLayoutVersion();
 #endif
-		check(DestVariable.GetSizeInBytes() == sizeof(FMatrix));
-		ValuePtr = (FMatrix*)InStore.GetParameterData(DestVariable);
+		check(DestVariable.GetSizeInBytes() == sizeof(FMatrix44f));
+		ValuePtr = (FMatrix44f*)InStore.GetParameterData(DestVariable);
 		return ValuePtr;
 	}
 
-	FORCEINLINE void SetValue(const FMatrix& InValue)
+	FORCEINLINE void SetValue(const FMatrix44f& InValue)
 	{
 #if NIAGARA_VALIDATE_DIRECT_BINDINGS
-		checkSlow(BoundVariable.GetSizeInBytes() == sizeof(FMatrix));
+		checkSlow(BoundVariable.GetSizeInBytes() == sizeof(FMatrix44f));
 		checkfSlow(LayoutVersion == BoundStore->GetLayoutVersion(), TEXT("This binding is invalid, its bound parameter store's layout was changed since it was created"));
 #endif
 		if (ValuePtr)
 		{
-			FMemory::Memcpy(ValuePtr, &InValue, sizeof(FMatrix));//Temp annoyance until we fix the alignment issues with parameter stores.
+			FMemory::Memcpy(ValuePtr, &InValue, sizeof(FMatrix44f));//Temp annoyance until we fix the alignment issues with parameter stores.
 		}
 	}
 
-	FORCEINLINE FMatrix GetValue()const
+	FORCEINLINE FMatrix44f GetValue()const
 	{
 #if NIAGARA_VALIDATE_DIRECT_BINDINGS
-		checkSlow(BoundVariable.GetSizeInBytes() == sizeof(FMatrix));
+		checkSlow(BoundVariable.GetSizeInBytes() == sizeof(FMatrix44f));
 		checkfSlow(LayoutVersion == BoundStore->GetLayoutVersion(), TEXT("This binding is invalid, its bound parameter store's layout was changed since it was created"));
 #endif
-		FMatrix Ret;
+		FMatrix44f Ret;
 		if (ValuePtr)
 		{
-			FMemory::Memcpy(&Ret, ValuePtr, sizeof(FMatrix));//Temp annoyance until we fix the alignment issues with parameter stores.
+			FMemory::Memcpy(&Ret, ValuePtr, sizeof(FMatrix44f));//Temp annoyance until we fix the alignment issues with parameter stores.
 		}
 		return Ret;
 	}
diff --git a/Engine/Plugins/FX/Niagara/Source/Niagara/Public/NiagaraRibbonRendererProperties.h b/Engine/Plugins/FX/Niagara/Source/Niagara/Public/NiagaraRibbonRendererProperties.h
--- a/Engine/Plugins/FX/Niagara/Source/Niagara/Public/NiagaraRibbonRendererProperties.h
+++ b/Engine/Plugins/FX/Niagara/Source/Niagara/Public/NiagaraRibbonRendererProperties.h
@@ -415,11 +415,11 @@ public:
 
 	bool								bSortKeyDataSetAccessorIsAge = false;
 	FNiagaraDataSetAccessor<float>		SortKeyDataSetAccessor;
-	FNiagaraDataSetAccessor<FVector>	PositionDataSetAccessor;
+	FNiagaraDataSetAccessor<FVector3f>	PositionDataSetAccessor;
 	FNiagaraDataSetAccessor<float>		NormalizedAgeAccessor;
 	FNiagaraDataSetAccessor<float>		SizeDataSetAccessor;
 	FNiagaraDataSetAccessor<float>		TwistDataSetAccessor;
-	FNiagaraDataSetAccessor<FVector>	FacingDataSetAccessor;
+	FNiagaraDataSetAccessor<FVector3f>	FacingDataSetAccessor;
 	FNiagaraDataSetAccessor<FVector4>	MaterialParam0DataSetAccessor;
 	FNiagaraDataSetAccessor<FVector4>	MaterialParam1DataSetAccessor;
 	FNiagaraDataSetAccessor<FVector4>	MaterialParam2DataSetAccessor;
diff --git a/Engine/Plugins/FX/Niagara/Source/Niagara/Public/NiagaraSystemInstanceController.h b/Engine/Plugins/FX/Niagara/Source/Niagara/Public/NiagaraSystemInstanceController.h
--- a/Engine/Plugins/FX/Niagara/Source/Niagara/Public/NiagaraSystemInstanceController.h
+++ b/Engine/Plugins/FX/Niagara/Source/Niagara/Public/NiagaraSystemInstanceController.h
@@ -153,7 +153,7 @@ private:
 	void SetVariable(FName InVariableName, int32 InValue);
 	void SetVariable(FName InVariableName, float InValue);
 	void SetVariable(FName InVariableName, FVector2D InValue);
-	void SetVariable(FName InVariableName, FVector InValue);
+	void SetVariable(FName InVariableName, FVector3f InValue);
 	void SetVariable(FName InVariableName, FVector4 InValue);
 	void SetVariable(FName InVariableName, FLinearColor InValue);
 	void SetVariable(FName InVariableName, FQuat InValue);
diff --git a/Engine/Plugins/FX/Niagara/Source/Niagara/Public/NiagaraTypes.h b/Engine/Plugins/FX/Niagara/Source/Niagara/Public/NiagaraTypes.h
--- a/Engine/Plugins/FX/Niagara/Source/Niagara/Public/NiagaraTypes.h
+++ b/Engine/Plugins/FX/Niagara/Source/Niagara/Public/NiagaraTypes.h
@@ -248,6 +248,18 @@ FORCEINLINE uint32 GetTypeHash(const FNiagaraID& ID)
 	return HashCombine(GetTypeHash(ID.Index), GetTypeHash(ID.AcquireTag));
 }
 
+/*
+*  Can convert a UStruct with fields of base types only (float, int... - will likely add native vector types here as well)
+*	to an FNiagaraTypeDefinition (internal representation)
+*/
+class NIAGARA_API FNiagaraTypeHelper
+{
+public:
+	static FString ToString(const uint8* ValueData, const UObject* StructOrEnum);
+	static UScriptStruct* FindNiagaraFriendlyTopLevelStruct(UScriptStruct* InStruct);
+	static bool IsNiagaraFriendlyTopLevelStruct(UScriptStruct* InStruct);
+};
+
 /** Information about how this type should be laid out in an FNiagaraDataSet */
 USTRUCT()
 struct FNiagaraTypeLayoutInfo
@@ -317,7 +329,7 @@ private:
 			//Should be able to support double easily enough
 			else if (FStructProperty* StructProp = CastFieldChecked<FStructProperty>(Property))
 			{
-				GenerateLayoutInfoInternal(Layout, StructProp->Struct, PropOffset);
+				GenerateLayoutInfoInternal(Layout, FNiagaraTypeHelper::FindNiagaraFriendlyTopLevelStruct(StructProp->Struct), PropOffset);
 			}
 			else
 			{
@@ -327,15 +339,6 @@ private:
 	}
 };
 
-/*
-*  Can convert a UStruct with fields of base types only (float, int... - will likely add native vector types here as well)
-*	to an FNiagaraTypeDefinition (internal representation)
-*/
-class NIAGARA_API FNiagaraTypeHelper
-{
-public:
-	static FString ToString(const uint8* ValueData, const UObject* StructOrEnum);
-};
 
 /** Defines different modes for selecting the output numeric type of a function or operation based on the types of the inputs. */
 UENUM()
@@ -672,14 +675,17 @@ public:
 		checkSlow(ClassStructOrEnum != nullptr);
 	}
 
-	FORCEINLINE FNiagaraTypeDefinition(UScriptStruct *StructDef)
+	enum class EAllowUnfriendlyStruct : uint8 { Deny, Allow };
+	FORCEINLINE FNiagaraTypeDefinition(UScriptStruct* StructDef, EAllowUnfriendlyStruct AllowUnfriendlyStruct)
 		: ClassStructOrEnum(StructDef), UnderlyingType(UT_Struct), Size(INDEX_NONE), Alignment(INDEX_NONE)
 #if WITH_EDITORONLY_DATA
 		, Struct_DEPRECATED(nullptr), Enum_DEPRECATED(nullptr)
 #endif
 	{
 		checkSlow(ClassStructOrEnum != nullptr);
+		ensure(AllowUnfriendlyStruct == EAllowUnfriendlyStruct::Allow || FNiagaraTypeHelper::IsNiagaraFriendlyTopLevelStruct(StructDef));
 	}
+	FORCEINLINE FNiagaraTypeDefinition(UScriptStruct* StructDef) : FNiagaraTypeDefinition(StructDef, EAllowUnfriendlyStruct::Deny) {}
 
 	FORCEINLINE FNiagaraTypeDefinition(const FNiagaraTypeDefinition &Other)
 		: ClassStructOrEnum(Other.ClassStructOrEnum), UnderlyingType(Other.UnderlyingType), Size(INDEX_NONE), Alignment(INDEX_NONE)
@@ -937,6 +943,9 @@ public:
 
 	static const TArray<FNiagaraTypeDefinition>& GetNumericTypes() { return OrderedNumericTypes; }
 	static bool IsValidNumericInput(const FNiagaraTypeDefinition& TypeDef);
+
+	
+
 private:
 
 	static FNiagaraTypeDefinition FloatDef;
@@ -1030,12 +1039,12 @@ const FNiagaraTypeDefinition& FNiagaraTypeDefinition::Get()
 {
 	if (TIsSame<T, float>::Value) { return FNiagaraTypeDefinition::GetFloatDef(); }
 	if (TIsSame<T, FVector2D>::Value) { return FNiagaraTypeDefinition::GetVec2Def(); }
-	if (TIsSame<T, FVector>::Value) { return FNiagaraTypeDefinition::GetVec3Def(); }
+	if (TIsSame<T, FVector3f>::Value) { return FNiagaraTypeDefinition::GetVec3Def(); }	
 	if (TIsSame<T, FVector4>::Value) { return FNiagaraTypeDefinition::GetVec4Def(); }
 	if (TIsSame<T, int32>::Value) { return FNiagaraTypeDefinition::GetIntDef(); }
 	if (TIsSame<T, FNiagaraBool>::Value) { return FNiagaraTypeDefinition::GetBoolDef(); }
 	if (TIsSame<T, FQuat>::Value) { return FNiagaraTypeDefinition::GetQuatDef(); }
-	if (TIsSame<T, FMatrix>::Value) { return FNiagaraTypeDefinition::GetMatrix4Def(); }
+	if (TIsSame<T, FMatrix44f>::Value) { return FNiagaraTypeDefinition::GetMatrix4Def(); }
 	if (TIsSame<T, FLinearColor>::Value) { return FNiagaraTypeDefinition::GetColorDef(); }
 	if (TIsSame<T, FNiagaraID>::Value) { return FNiagaraTypeDefinition::GetIDDef(); }
 }
@@ -1637,6 +1646,8 @@ inline void FNiagaraVariable::SetValue<bool>(const bool& Data)
 	BoolStruct->SetValue(Data);
 }
 
+
+
 // Any change to this structure, or it's GetVariables implementation will require a bump in the CustomNiagaraVersion so that we
 // properly rebuild the scripts
 // You must pad this struct and the results of GetVariables() to a 16 byte boundry.
@@ -1690,12 +1701,12 @@ struct alignas(16) FNiagaraOwnerParameters
 	NIAGARA_API static const TArray<FNiagaraVariable>& GetVariables();
 #endif
 
-	FMatrix EngineLocalToWorld = FMatrix::Identity;
-	FMatrix EngineWorldToLocal = FMatrix::Identity;
-	FMatrix EngineLocalToWorldTransposed = FMatrix::Identity;
-	FMatrix EngineWorldToLocalTransposed = FMatrix::Identity;
-	FMatrix EngineLocalToWorldNoScale = FMatrix::Identity;
-	FMatrix EngineWorldToLocalNoScale = FMatrix::Identity;
+	FMatrix44f EngineLocalToWorld = FMatrix44f::Identity;
+	FMatrix44f EngineWorldToLocal = FMatrix44f::Identity;
+	FMatrix44f EngineLocalToWorldTransposed = FMatrix44f::Identity;
+	FMatrix44f EngineWorldToLocalTransposed = FMatrix44f::Identity;
+	FMatrix44f EngineLocalToWorldNoScale = FMatrix44f::Identity;
+	FMatrix44f EngineWorldToLocalNoScale = FMatrix44f::Identity;
 	FQuat EngineRotation = FQuat::Identity;
 	FVector4 EnginePosition = FVector4(EForceInit::ForceInitToZero);
 	FVector4 EngineVelocity = FVector4(EForceInit::ForceInitToZero);
diff --git a/Engine/Plugins/FX/Niagara/Source/NiagaraCore/Private/NiagaraCustomVersion.cpp b/Engine/Plugins/FX/Niagara/Source/NiagaraCore/Private/NiagaraCustomVersion.cpp
--- a/Engine/Plugins/FX/Niagara/Source/NiagaraCore/Private/NiagaraCustomVersion.cpp
+++ b/Engine/Plugins/FX/Niagara/Source/NiagaraCore/Private/NiagaraCustomVersion.cpp
@@ -8,4 +8,4 @@ const FGuid FNiagaraCustomVersion::GUID(0xFCF57AFA, 0x50764283, 0xB9A9E658, 0xFF
 // Register the custom version with core
 FCustomVersionRegistration GRegisterNiagaraCustomVersion(FNiagaraCustomVersion::GUID, FNiagaraCustomVersion::LatestVersion, TEXT("NiagaraVer"));
 
-const FGuid FNiagaraCustomVersion::LatestScriptCompileVersion(0xA83D40AB, 0xE1F64D9F, 0x9F5451FB, 0x8C954501);
+const FGuid FNiagaraCustomVersion::LatestScriptCompileVersion(0xA83D40AB, 0xE1F64D9F, 0x9F5451FB, 0x8C95FFFF);
diff --git a/Engine/Plugins/FX/Niagara/Source/NiagaraEditor/Private/EdGraphSchema_Niagara.cpp b/Engine/Plugins/FX/Niagara/Source/NiagaraEditor/Private/EdGraphSchema_Niagara.cpp
--- a/Engine/Plugins/FX/Niagara/Source/NiagaraEditor/Private/EdGraphSchema_Niagara.cpp
+++ b/Engine/Plugins/FX/Niagara/Source/NiagaraEditor/Private/EdGraphSchema_Niagara.cpp
@@ -1386,6 +1386,22 @@ bool UEdGraphSchema_Niagara::TryGetPinDefaultValueFromNiagaraVariable(const FNia
 	return false;
 }
 
+void UEdGraphSchema_Niagara::ConvertIllegalPinsInPlace(UEdGraphPin* Pin)
+{
+	if (Pin == nullptr)
+		return;
+
+	if (Pin->PinType.PinCategory == PinCategoryType && Pin->PinType.PinSubCategoryObject.IsValid())
+	{
+		UScriptStruct* Struct = Cast<UScriptStruct>(Pin->PinType.PinSubCategoryObject.Get());
+		UScriptStruct* ConvertedStruct = FNiagaraTypeHelper::FindNiagaraFriendlyTopLevelStruct(Struct);
+		if (Struct != ConvertedStruct)
+		{
+			Pin->PinType.PinSubCategoryObject = ConvertedStruct;
+		}
+	}
+}
+
 FNiagaraTypeDefinition UEdGraphSchema_Niagara::PinToTypeDefinition(const UEdGraphPin* Pin)
 {
 	if (Pin == nullptr)
@@ -1402,6 +1418,20 @@ FNiagaraTypeDefinition UEdGraphSchema_Niagara::PinToTypeDefinition(const UEdGrap
 				*Pin->PinName.ToString(), OwningNode ? *OwningNode->GetName() : TEXT("Invalid"));
 			return FNiagaraTypeDefinition();
 		}
+		else if (Struct && !FNiagaraTypeHelper::IsNiagaraFriendlyTopLevelStruct(Struct)) // LWC swapover support
+		{
+			if (OwningNode && OwningNode->HasAnyFlags(EObjectFlags::RF_NeedPostLoad))
+			{
+				UE_LOG(LogNiagaraEditor, Verbose, TEXT("Pin states that it is not a niagara friendly struct, but didn't get converted by PostLoad yet, it still has RF_NeedPostLoad . Pin Name '%s' Owning Node '%s'."),
+					*Pin->PinName.ToString(), OwningNode ? *OwningNode->GetPathName() : TEXT("Invalid"));
+			}
+			else 
+			{
+				UE_LOG(LogNiagaraEditor, Verbose, TEXT("Pin states that it is not a niagara friendly struct, but didn't get converted by PostLoad yet, it does not have RF_NeedPostLoad though. Pin Name '%s' Owning Node '%s'."),
+					*Pin->PinName.ToString(), OwningNode ? *OwningNode->GetPathName() : TEXT("Invalid"));
+			}
+			return FNiagaraTypeDefinition(FNiagaraTypeHelper::FindNiagaraFriendlyTopLevelStruct(Struct));
+		}
 		return FNiagaraTypeDefinition(Struct);
 	}
 	else if (Pin->PinType.PinCategory == PinCategoryClass)
@@ -1679,7 +1709,7 @@ FNiagaraTypeDefinition UEdGraphSchema_Niagara::GetTypeDefForProperty(const FProp
 	}
 	else if (const FStructProperty* StructProp = CastFieldChecked<const FStructProperty>(Property))
 	{
-		return FNiagaraTypeDefinition(StructProp->Struct);
+		return FNiagaraTypeDefinition(FNiagaraTypeHelper::FindNiagaraFriendlyTopLevelStruct(StructProp->Struct));
 	}
 
 	check(0);
diff --git a/Engine/Plugins/FX/Niagara/Source/NiagaraEditor/Private/NiagaraHlslTranslator.cpp b/Engine/Plugins/FX/Niagara/Source/NiagaraEditor/Private/NiagaraHlslTranslator.cpp
--- a/Engine/Plugins/FX/Niagara/Source/NiagaraEditor/Private/NiagaraHlslTranslator.cpp
+++ b/Engine/Plugins/FX/Niagara/Source/NiagaraEditor/Private/NiagaraHlslTranslator.cpp
@@ -2268,17 +2268,18 @@ void FHlslNiagaraTranslator::GatherComponentsForDataSetAccess(UScriptStruct* Str
 
 		if (const FStructProperty* StructProp = CastField<const FStructProperty>(Property))
 		{
-			if (bMatrixRoot && FNiagaraTypeDefinition(StructProp->Struct) == FNiagaraTypeDefinition::GetFloatDef())
+			UScriptStruct* NiagaraStruct = FNiagaraTypeHelper::FindNiagaraFriendlyTopLevelStruct(StructProp->Struct);
+			if (bMatrixRoot && FNiagaraTypeDefinition(NiagaraStruct) == FNiagaraTypeDefinition::GetFloatDef())
 			{
-				GatherComponentsForDataSetAccess(StructProp->Struct, VariableSymbol + ComputeMatrixColumnAccess(Property->GetName()), bMatrixRoot, Components, Types);
+				GatherComponentsForDataSetAccess(NiagaraStruct, VariableSymbol + ComputeMatrixColumnAccess(Property->GetName()), bMatrixRoot, Components, Types);
 			}
-			else if (bMatrixRoot &&  FNiagaraTypeDefinition(StructProp->Struct) == FNiagaraTypeDefinition::GetVec4Def())
+			else if (bMatrixRoot &&  FNiagaraTypeDefinition(NiagaraStruct) == FNiagaraTypeDefinition::GetVec4Def())
 			{
-				GatherComponentsForDataSetAccess(StructProp->Struct, VariableSymbol + ComputeMatrixRowAccess(Property->GetName()), bMatrixRoot, Components, Types);
+				GatherComponentsForDataSetAccess(NiagaraStruct, VariableSymbol + ComputeMatrixRowAccess(Property->GetName()), bMatrixRoot, Components, Types);
 			}
 			else
 			{
-				GatherComponentsForDataSetAccess(StructProp->Struct, VariableSymbol + TEXT(".") + Property->GetName(), bMatrixRoot, Components, Types);
+				GatherComponentsForDataSetAccess(NiagaraStruct, VariableSymbol + TEXT(".") + Property->GetName(), bMatrixRoot, Components, Types);
 			}
 		}
 		else
@@ -3664,10 +3665,11 @@ void FHlslNiagaraTranslator::DecomposeVariableAccess(UStruct* Struct, bool bRead
 
 		if (FStructProperty* StructProp = CastFieldChecked<FStructProperty>(Property))
 		{
-			FNiagaraTypeDefinition PropDef(StructProp->Struct);
+			UScriptStruct* NiagaraStruct = FNiagaraTypeHelper::FindNiagaraFriendlyTopLevelStruct(StructProp->Struct);
+			FNiagaraTypeDefinition PropDef(NiagaraStruct);
 			if (!IsHlslBuiltinVector(PropDef))
 			{
-				DecomposeVariableAccess(StructProp->Struct, bRead, IndexSymbol, AccessStr);
+				DecomposeVariableAccess(NiagaraStruct, bRead, IndexSymbol, AccessStr);
 				return;
 			}
 		}
@@ -5053,7 +5055,7 @@ bool FHlslNiagaraTranslator::HandleBoundConstantVariableToDataSetRead(FNiagaraVa
 		}
 		else
 		{
-			InVariable.SetValue(FVector(EForceInit::ForceInitToZero));
+			InVariable.SetValue(FVector3f(EForceInit::ForceInitToZero));
 			float* ValuePtr = (float*)InVariable.GetData();
 			const FString ConstantStr = FString::Printf(TEXT("float3(%g,%g,%g)"), *ValuePtr, *(ValuePtr + 1), *(ValuePtr + 2));
 			Output = AddBodyChunk(GetUniqueSymbolName(TEXT("Constant")), ConstantStr, InVariable.GetType()); 
@@ -7603,7 +7605,7 @@ FNiagaraTypeDefinition FHlslNiagaraTranslator::GetChildType(const FNiagaraTypeDe
 				}
 				else if (const FStructProperty* StructProp = CastFieldChecked<const FStructProperty>(Property))
 				{
-					return FNiagaraTypeDefinition(StructProp->Struct);
+					return FNiagaraTypeDefinition(FNiagaraTypeHelper::FindNiagaraFriendlyTopLevelStruct(StructProp->Struct));
 				}
 			}
 		}
@@ -8488,7 +8490,7 @@ FString FHlslNiagaraTranslator::GetPropertyHlslTypeName(const FProperty* Propert
 	else if (Property->IsA(FStructProperty::StaticClass()))
 	{
 		const FStructProperty* StructProp = CastField<const FStructProperty>(Property);
-		return GetStructHlslTypeName(StructProp->Struct);
+		return GetStructHlslTypeName(FNiagaraTypeHelper::FindNiagaraFriendlyTopLevelStruct(StructProp->Struct));
 	}
 	else if (Property->IsA(FEnumProperty::StaticClass()))
 	{
@@ -8580,7 +8582,7 @@ bool FHlslNiagaraTranslator::AddStructToDefinitionSet(const FNiagaraTypeDefiniti
 			const FStructProperty* StructProp = CastField<const FStructProperty>(Property);
 			if (StructProp)
 			{
-				if (!AddStructToDefinitionSet(StructProp->Struct))
+				if (!AddStructToDefinitionSet(FNiagaraTypeHelper::FindNiagaraFriendlyTopLevelStruct(StructProp->Struct)))
 				{
 					return false;
 				}
@@ -8640,7 +8642,7 @@ TArray<FName> FHlslNiagaraTranslator::ConditionPropertyPath(const FNiagaraTypeDe
 						ReturnPath.Add(InPath[0]);
 						TArray<FName> Subset = InPath;
 						Subset.RemoveAt(0);
-						TArray<FName> Children = ConditionPropertyPath(StructProp->Struct, Subset);
+						TArray<FName> Children = ConditionPropertyPath(FNiagaraTypeHelper::FindNiagaraFriendlyTopLevelStruct(StructProp->Struct), Subset);
 						for (const FName& Child : Children)
 						{
 							ReturnPath.Add(Child);
diff --git a/Engine/Plugins/FX/Niagara/Source/NiagaraEditor/Private/NiagaraNode.cpp b/Engine/Plugins/FX/Niagara/Source/NiagaraEditor/Private/NiagaraNode.cpp
--- a/Engine/Plugins/FX/Niagara/Source/NiagaraEditor/Private/NiagaraNode.cpp
+++ b/Engine/Plugins/FX/Niagara/Source/NiagaraEditor/Private/NiagaraNode.cpp
@@ -43,6 +43,8 @@ void UNiagaraNode::PostLoad()
 		const UEdGraphSchema_Niagara* Schema = CastChecked<UEdGraphSchema_Niagara>(GetSchema());
 		check(Schema);
 
+		UEdGraphSchema_Niagara::ConvertIllegalPinsInPlace(Pin);
+
 		const FNiagaraTypeDefinition PinType = Schema->PinToTypeDefinition(Pin);
 		if (PinType.GetEnum() == FNiagaraTypeDefinition::GetCoordinateSpaceEnum())
 		{
diff --git a/Engine/Plugins/FX/Niagara/Source/NiagaraEditor/Private/NiagaraNodeConvert.cpp b/Engine/Plugins/FX/Niagara/Source/NiagaraEditor/Private/NiagaraNodeConvert.cpp
--- a/Engine/Plugins/FX/Niagara/Source/NiagaraEditor/Private/NiagaraNodeConvert.cpp
+++ b/Engine/Plugins/FX/Niagara/Source/NiagaraEditor/Private/NiagaraNodeConvert.cpp
@@ -107,9 +107,9 @@ struct FNiagaraConvertEntry
 
 			const FStructProperty* StructProperty = CastField<FStructProperty>(Property);
 
-			if (StructProperty != nullptr)
+			 if (StructProperty != nullptr)
 			{
-				CreateEntries(Schema, InPinId, InPin, StructProperty->Struct, OutEntries[Index].Children);				
+				CreateEntries(Schema, InPinId, InPin, FNiagaraTypeHelper::FindNiagaraFriendlyTopLevelStruct(StructProperty->Struct), OutEntries[Index].Children);
 			}
 		}
 	}
diff --git a/Engine/Plugins/FX/Niagara/Source/NiagaraEditor/Private/NiagaraNodeDataSetBase.cpp b/Engine/Plugins/FX/Niagara/Source/NiagaraEditor/Private/NiagaraNodeDataSetBase.cpp
--- a/Engine/Plugins/FX/Niagara/Source/NiagaraEditor/Private/NiagaraNodeDataSetBase.cpp
+++ b/Engine/Plugins/FX/Niagara/Source/NiagaraEditor/Private/NiagaraNodeDataSetBase.cpp
@@ -266,7 +266,7 @@ bool UNiagaraNodeDataSetBase::GetSupportedNiagaraTypeDef(const FProperty* Proper
 	}
 	else if (StructProp && StructProp->Struct)
 	{
-		TypeDef = FNiagaraTypeDefinition(StructProp->Struct);
+		TypeDef = FNiagaraTypeDefinition(FNiagaraTypeHelper::FindNiagaraFriendlyTopLevelStruct(StructProp->Struct));
 		return FNiagaraTypeRegistry::GetRegisteredPayloadTypes().Contains(TypeDef);
 	}
 	return false;
diff --git a/Engine/Plugins/FX/Niagara/Source/NiagaraEditor/Private/NiagaraNodeFunctionCall.cpp b/Engine/Plugins/FX/Niagara/Source/NiagaraEditor/Private/NiagaraNodeFunctionCall.cpp
--- a/Engine/Plugins/FX/Niagara/Source/NiagaraEditor/Private/NiagaraNodeFunctionCall.cpp
+++ b/Engine/Plugins/FX/Niagara/Source/NiagaraEditor/Private/NiagaraNodeFunctionCall.cpp
@@ -113,6 +113,9 @@ void UNiagaraNodeFunctionCall::PostLoad()
 	}
 	
 	// check if maybe the parameter names in the referenced module were changed and try to move over existing values
+	UNiagaraGraph* Graph = GetCalledGraph();
+	if (Graph)
+		Graph->ConditionalPostLoad();
 	FixupPinNames();
 }
 
diff --git a/Engine/Plugins/FX/Niagara/Source/NiagaraEditor/Private/TypeEditorUtilities/NiagaraMatrixTypeEditorUtilities.cpp b/Engine/Plugins/FX/Niagara/Source/NiagaraEditor/Private/TypeEditorUtilities/NiagaraMatrixTypeEditorUtilities.cpp
--- a/Engine/Plugins/FX/Niagara/Source/NiagaraEditor/Private/TypeEditorUtilities/NiagaraMatrixTypeEditorUtilities.cpp
+++ b/Engine/Plugins/FX/Niagara/Source/NiagaraEditor/Private/TypeEditorUtilities/NiagaraMatrixTypeEditorUtilities.cpp
@@ -8,6 +8,6 @@ void FNiagaraEditorMatrixTypeUtilities::UpdateVariableWithDefaultValue(FNiagaraV
 	// This is implemented directly here to avoid alignment issues with the default initialization path for structs which
 	// uses pointer assignment rather than a mem-copy.
 	checkf(Variable.GetType().GetStruct() == FNiagaraTypeDefinition::GetMatrix4Struct(), TEXT("Struct type not supported."));
-	FMatrix DefaultMatrix(FMatrix::Identity);
+	FMatrix44f DefaultMatrix(FMatrix44f::Identity);
 	Variable.SetData((uint8*)&DefaultMatrix);
 }
\ No newline at end of file
diff --git a/Engine/Plugins/FX/Niagara/Source/NiagaraEditor/Private/TypeEditorUtilities/NiagaraVectorTypeEditorUtilities.cpp b/Engine/Plugins/FX/Niagara/Source/NiagaraEditor/Private/TypeEditorUtilities/NiagaraVectorTypeEditorUtilities.cpp
--- a/Engine/Plugins/FX/Niagara/Source/NiagaraEditor/Private/TypeEditorUtilities/NiagaraVectorTypeEditorUtilities.cpp
+++ b/Engine/Plugins/FX/Niagara/Source/NiagaraEditor/Private/TypeEditorUtilities/NiagaraVectorTypeEditorUtilities.cpp
@@ -201,14 +201,14 @@ public:
 
 	virtual void UpdateInternalValueFromStruct(TSharedRef<FStructOnScope> Struct) override
 	{
-		checkf(Struct->GetStruct() == FNiagaraTypeDefinition::GetVec3Struct(), TEXT("Struct type not supported."));
-		VectorValue = *((FVector*)Struct->GetStructMemory());
+		checkf(Struct->GetStruct() == FNiagaraTypeDefinition::GetVec3Struct(), TEXT("Struct type not supported."));	// LWC_TODO: Support for FVector3d likely required here.
+		VectorValue = *((FVector3f*)Struct->GetStructMemory());
 	}
 
 	virtual void UpdateStructFromInternalValue(TSharedRef<FStructOnScope> Struct) override
 	{
 		checkf(Struct->GetStruct() == FNiagaraTypeDefinition::GetVec3Struct(), TEXT("Struct type not supported."));
-		*((FVector*)Struct->GetStructMemory()) = VectorValue;
+		*((FVector3f*)Struct->GetStructMemory()) = VectorValue;
 	}
 
 protected:
@@ -223,7 +223,7 @@ protected:
 	}
 
 private:
-	FVector VectorValue;
+	FVector3f VectorValue;
 };
 
 TSharedPtr<SNiagaraParameterEditor> FNiagaraEditorVector3TypeUtilities::CreateParameterEditor(const FNiagaraTypeDefinition& ParameterType) const
@@ -241,7 +241,7 @@ FString FNiagaraEditorVector3TypeUtilities::GetPinDefaultStringFromValue(const F
 	checkf(AllocatedVariable.IsDataAllocated(), TEXT("Can not generate a default value string for an unallocated variable."));
 
 	// NOTE: We can not use ToString() here since the vector pin control doesn't use the standard 'X=0,Y=0,Z=0' syntax.
-	FVector Value = AllocatedVariable.GetValue<FVector>();
+	FVector Value = AllocatedVariable.GetValue<FVector3f>();
 	return FString::Printf(TEXT("%3.3f,%3.3f,%3.3f"), Value.X, Value.Y, Value.Z);
 }
 
@@ -251,7 +251,7 @@ bool FNiagaraEditorVector3TypeUtilities::SetValueFromPinDefaultString(const FStr
 	FVector Value = FVector::ZeroVector;
 	if (FDefaultValueHelper::ParseVector(StringValue, Value) || !Variable.IsDataAllocated())
 	{
-		Variable.SetValue<FVector>(Value);
+		Variable.SetValue<FVector3f>((FVector3f)Value);
 		return true;
 	}
 	return false;
@@ -264,7 +264,7 @@ FText FNiagaraEditorVector3TypeUtilities::GetSearchTextFromValue(const FNiagaraV
 
 FText FNiagaraEditorVector3TypeUtilities::GetStackDisplayText(FNiagaraVariable& Variable) const
 {
-	FVector Value = Variable.GetValue<FVector>();
+	FVector Value = Variable.GetValue<FVector3f>();
 	return FText::Format(FText::FromString("({0}, {1}, {2})"), Value.X, Value.Y, Value.Z);
 }
 
diff --git a/Engine/Plugins/FX/Niagara/Source/NiagaraEditor/Private/UpgradeNiagaraScriptResults.cpp b/Engine/Plugins/FX/Niagara/Source/NiagaraEditor/Private/UpgradeNiagaraScriptResults.cpp
--- a/Engine/Plugins/FX/Niagara/Source/NiagaraEditor/Private/UpgradeNiagaraScriptResults.cpp
+++ b/Engine/Plugins/FX/Niagara/Source/NiagaraEditor/Private/UpgradeNiagaraScriptResults.cpp
@@ -85,7 +85,7 @@ FVector UNiagaraPythonScriptModuleInput::AsVec3() const
 {
 	if (IsSet() && Input->InputType == FNiagaraTypeDefinition::GetVec3Def())
 	{
-		return GetValue<FVector>(Input);
+		return GetValue<FVector3f>(Input);
 	}
 	return FVector();
 }
diff --git a/Engine/Plugins/FX/Niagara/Source/NiagaraEditor/Private/Widgets/SNiagaraSpreadsheetView.cpp b/Engine/Plugins/FX/Niagara/Source/NiagaraEditor/Private/Widgets/SNiagaraSpreadsheetView.cpp
--- a/Engine/Plugins/FX/Niagara/Source/NiagaraEditor/Private/Widgets/SNiagaraSpreadsheetView.cpp
+++ b/Engine/Plugins/FX/Niagara/Source/NiagaraEditor/Private/Widgets/SNiagaraSpreadsheetView.cpp
@@ -1243,7 +1243,7 @@ void SNiagaraSpreadsheetView::GenerateLayoutInfo(FNiagaraTypeLayoutInfo& Layout,
 		else if (Property->IsA(FStructProperty::StaticClass()))
 		{
 			FStructProperty* StructProp = CastFieldChecked<FStructProperty>(Property);
-			GenerateLayoutInfo(Layout, StructProp->Struct, nullptr, PropertyName, PropertyNames, FieldInfo);
+			GenerateLayoutInfo(Layout, FNiagaraTypeHelper::FindNiagaraFriendlyTopLevelStruct(StructProp->Struct), nullptr, PropertyName, PropertyNames, FieldInfo);
 		}
 		else
 		{
diff --git a/Engine/Plugins/FX/Niagara/Source/NiagaraEditor/Public/EdGraphSchema_Niagara.h b/Engine/Plugins/FX/Niagara/Source/NiagaraEditor/Public/EdGraphSchema_Niagara.h
--- a/Engine/Plugins/FX/Niagara/Source/NiagaraEditor/Public/EdGraphSchema_Niagara.h
+++ b/Engine/Plugins/FX/Niagara/Source/NiagaraEditor/Public/EdGraphSchema_Niagara.h
@@ -141,6 +141,7 @@ class NIAGARAEDITOR_API UEdGraphSchema_Niagara : public UEdGraphSchema
 	  */
 	bool TryGetPinDefaultValueFromNiagaraVariable(const FNiagaraVariable& Variable, FString& OutPinDefaultValue) const;
 
+	static void ConvertIllegalPinsInPlace(UEdGraphPin* Pin);
 	static FNiagaraTypeDefinition PinToTypeDefinition(const UEdGraphPin* Pin);
 	static FNiagaraTypeDefinition PinTypeToTypeDefinition(const FEdGraphPinType& PinType) ;
 	static FEdGraphPinType TypeDefinitionToPinType(FNiagaraTypeDefinition TypeDef);
diff --git a/Engine/Source/Runtime/Engine/Private/Particles/ParticleModules_Location.cpp b/Engine/Source/Runtime/Engine/Private/Particles/ParticleModules_Location.cpp
--- a/Engine/Source/Runtime/Engine/Private/Particles/ParticleModules_Location.cpp
+++ b/Engine/Source/Runtime/Engine/Private/Particles/ParticleModules_Location.cpp
@@ -355,7 +355,7 @@ void UParticleModuleLocationDirect::Update(FParticleEmitterInstance* Owner, int3
 
 uint32 UParticleModuleLocationDirect::RequiredBytes(UParticleModuleTypeDataBase* TypeData)
 {
-	return sizeof(FVector);
+	return sizeof(FVector3f);
 }
 
 /*-----------------------------------------------------------------------------
@@ -1757,7 +1757,7 @@ uint32 UParticleModuleLocationBoneSocket::RequiredBytesPerInstance()
 
 	//Have to take the max of all variants as lots of code assumes all LODs use the same memory and prep it the same way :(
 	int32 ArraySize = FMath::Max(SourceLocations.Num(), NumPreSelectedIndices);
-	int32 ElemSize = (sizeof(FVector)* 2) + sizeof(int32);
+	int32 ElemSize = (sizeof(FVector3f)* 2) + sizeof(int32);
 	
 	int32 BoneArraySize = ArraySize * ElemSize;
 	
diff --git a/Engine/Source/Runtime/Engine/Private/Particles/ParticleSystemRender.cpp b/Engine/Source/Runtime/Engine/Private/Particles/ParticleSystemRender.cpp
--- a/Engine/Source/Runtime/Engine/Private/Particles/ParticleSystemRender.cpp
+++ b/Engine/Source/Runtime/Engine/Private/Particles/ParticleSystemRender.cpp
@@ -2299,8 +2299,8 @@ void FDynamicBeam2EmitterData::Init( bool bInSelected )
 
 	check(Source.ActiveParticleCount < (MaxBeams));	// TTP #33330 - Max of 2048 beams from a single emitter
 	check(Source.ParticleStride < 
-		((MaxInterpolationPoints + 2) * (sizeof(FVector) + sizeof(float))) + 
-		(MaxNoiseFrequency * (sizeof(FVector) + sizeof(FVector) + sizeof(float) + sizeof(float)))
+		((MaxInterpolationPoints + 2) * (sizeof(FVector3f) + sizeof(float))) + 
+		(MaxNoiseFrequency * (sizeof(FVector3f) + sizeof(FVector3f) + sizeof(float) + sizeof(float)))
 		);	// TTP #33330 - Max of 10k per beam (includes interpolation points, noise, etc.)
 
 	MaterialResource = Source.MaterialInterface->GetRenderProxy();
diff --git a/Engine/Source/Runtime/VectorVM/Private/VectorVM.cpp b/Engine/Source/Runtime/VectorVM/Private/VectorVM.cpp
--- a/Engine/Source/Runtime/VectorVM/Private/VectorVM.cpp
+++ b/Engine/Source/Runtime/VectorVM/Private/VectorVM.cpp
@@ -703,7 +703,7 @@ struct TUnaryKernel
 template<typename Kernel>
 struct TUnaryScalarKernel : public TUnaryKernel<Kernel, FRegisterHandler<float>, FConstantHandler<float>, FRegisterHandler<float>, 1> {};
 template<typename Kernel>
-struct TUnaryVectorKernel : public TUnaryKernel<Kernel, FRegisterHandler<VectorRegister>, FConstantHandler<VectorRegister>, FRegisterHandler<VectorRegister>, VECTOR_WIDTH_FLOATS> {};
+struct TUnaryVectorKernel : public TUnaryKernel<Kernel, FRegisterHandler<VectorRegister4Float>, FConstantHandler<VectorRegister4Float>, FRegisterHandler<VectorRegister4Float>, VECTOR_WIDTH_FLOATS> {};
 template<typename Kernel>
 struct TUnaryScalarIntKernel : public TUnaryKernel<Kernel, FRegisterHandler<int32>, FConstantHandler<int32>, FRegisterHandler<int32>, 1> {};
 template<typename Kernel>
@@ -742,7 +742,7 @@ struct TBinaryKernel
 template<typename Kernel>
 struct TBinaryScalarKernel : public TBinaryKernel<Kernel, FRegisterHandler<float>, FConstantHandler<float>, FRegisterHandler<float>, 1> {};
 template<typename Kernel>
-struct TBinaryVectorKernel : public TBinaryKernel<Kernel, FRegisterHandler<VectorRegister>, FConstantHandler<VectorRegister>, FRegisterHandler<VectorRegister>, VECTOR_WIDTH_FLOATS> {};
+struct TBinaryVectorKernel : public TBinaryKernel<Kernel, FRegisterHandler<VectorRegister4Float>, FConstantHandler<VectorRegister4Float>, FRegisterHandler<VectorRegister4Float>, VECTOR_WIDTH_FLOATS> {};
 template<typename Kernel>
 struct TBinaryVectorIntKernel : public TBinaryKernel<Kernel, FRegisterHandler<VectorRegister4Int>, FConstantHandler<VectorRegister4Int>, FRegisterHandler<VectorRegister4Int>, VECTOR_WIDTH_FLOATS> {};
 
@@ -788,7 +788,7 @@ struct TTrinaryKernel
 template<typename Kernel>
 struct TTrinaryScalarKernel : public TTrinaryKernel<Kernel, FRegisterHandler<float>, FConstantHandler<float>, FRegisterHandler<float>, 1> {};
 template<typename Kernel>
-struct TTrinaryVectorKernel : public TTrinaryKernel<Kernel, FRegisterHandler<VectorRegister>, FConstantHandler<VectorRegister>, FRegisterHandler<VectorRegister>, VECTOR_WIDTH_FLOATS> {};
+struct TTrinaryVectorKernel : public TTrinaryKernel<Kernel, FRegisterHandler<VectorRegister4Float>, FConstantHandler<VectorRegister4Float>, FRegisterHandler<VectorRegister4Float>, VECTOR_WIDTH_FLOATS> {};
 template<typename Kernel>
 struct TTrinaryVectorIntKernel : public TTrinaryKernel<Kernel, FRegisterHandler<VectorRegister4Int>, FConstantHandler<VectorRegister4Int>, FRegisterHandler<VectorRegister4Int>, VECTOR_WIDTH_FLOATS> {};
 
@@ -799,7 +799,7 @@ struct TTrinaryVectorIntKernel : public TTrinaryKernel<Kernel, FRegisterHandler<
 
 struct FVectorKernelAdd : public TBinaryVectorKernel<FVectorKernelAdd>
 {
-	static void VM_FORCEINLINE DoKernel(FVectorVMContext& Context, VectorRegister* RESTRICT Dst,VectorRegister Src0,VectorRegister Src1)
+	static void VM_FORCEINLINE DoKernel(FVectorVMContext& Context, VectorRegister4Float* RESTRICT Dst,VectorRegister4Float Src0,VectorRegister4Float Src1)
 	{
 		*Dst = VectorAdd(Src0, Src1);
 	}
@@ -807,7 +807,7 @@ struct FVectorKernelAdd : public TBinaryVectorKernel<FVectorKernelAdd>
 
 struct FVectorKernelSub : public TBinaryVectorKernel<FVectorKernelSub>
 {
-	static void VM_FORCEINLINE DoKernel(FVectorVMContext& Context, VectorRegister* RESTRICT Dst,VectorRegister Src0,VectorRegister Src1)
+	static void VM_FORCEINLINE DoKernel(FVectorVMContext& Context, VectorRegister4Float* RESTRICT Dst,VectorRegister4Float Src0,VectorRegister4Float Src1)
 	{
 		*Dst = VectorSubtract(Src0, Src1);
 	}
@@ -815,7 +815,7 @@ struct FVectorKernelSub : public TBinaryVectorKernel<FVectorKernelSub>
 
 struct FVectorKernelMul : public TBinaryVectorKernel<FVectorKernelMul>
 {
-	static void VM_FORCEINLINE DoKernel(FVectorVMContext& Context, VectorRegister* RESTRICT Dst,VectorRegister Src0,VectorRegister Src1)
+	static void VM_FORCEINLINE DoKernel(FVectorVMContext& Context, VectorRegister4Float* RESTRICT Dst,VectorRegister4Float Src0,VectorRegister4Float Src1)
 	{
 		*Dst = VectorMultiply(Src0, Src1);
 	}
@@ -823,7 +823,7 @@ struct FVectorKernelMul : public TBinaryVectorKernel<FVectorKernelMul>
 
 struct FVectorKernelDiv : public TBinaryVectorKernel<FVectorKernelDiv>
 {
-	static void VM_FORCEINLINE DoKernel(FVectorVMContext& Context, VectorRegister* RESTRICT Dst, VectorRegister Src0, VectorRegister Src1)
+	static void VM_FORCEINLINE DoKernel(FVectorVMContext& Context, VectorRegister4Float* RESTRICT Dst, VectorRegister4Float Src0, VectorRegister4Float Src1)
 	{
 		*Dst = VectorDivide(Src0, Src1);
 	}
@@ -831,16 +831,16 @@ struct FVectorKernelDiv : public TBinaryVectorKernel<FVectorKernelDiv>
 
 struct FVectorKernelDivSafe : public TBinaryVectorKernel<FVectorKernelDivSafe>
 {
-	static void VM_FORCEINLINE DoKernel(FVectorVMContext& Context, VectorRegister* RESTRICT Dst, VectorRegister Src0, VectorRegister Src1)
+	static void VM_FORCEINLINE DoKernel(FVectorVMContext& Context, VectorRegister4Float* RESTRICT Dst, VectorRegister4Float Src0, VectorRegister4Float Src1)
 	{
-		VectorRegister ValidMask = VectorCompareGT(VectorAbs(Src1), GlobalVectorConstants::SmallNumber);
+		VectorRegister4Float ValidMask = VectorCompareGT(VectorAbs(Src1), GlobalVectorConstants::SmallNumber);
 		*Dst = VectorSelect(ValidMask, VectorDivide(Src0, Src1), GlobalVectorConstants::FloatZero);
 	}
 };
 
 struct FVectorKernelMad : public TTrinaryVectorKernel<FVectorKernelMad>
 {
-	static void VM_FORCEINLINE DoKernel(FVectorVMContext& Context, VectorRegister* RESTRICT Dst,VectorRegister Src0,VectorRegister Src1,VectorRegister Src2)
+	static void VM_FORCEINLINE DoKernel(FVectorVMContext& Context, VectorRegister4Float* RESTRICT Dst,VectorRegister4Float Src0,VectorRegister4Float Src1,VectorRegister4Float Src2)
 	{
 		*Dst = VectorMultiplyAdd(Src0, Src1, Src2);
 	}
@@ -848,17 +848,17 @@ struct FVectorKernelMad : public TTrinaryVectorKernel<FVectorKernelMad>
 
 struct FVectorKernelLerp : public TTrinaryVectorKernel<FVectorKernelLerp>
 {
-	static void VM_FORCEINLINE DoKernel(FVectorVMContext& Context, VectorRegister* RESTRICT Dst,VectorRegister Src0,VectorRegister Src1,VectorRegister Src2)
+	static void VM_FORCEINLINE DoKernel(FVectorVMContext& Context, VectorRegister4Float* RESTRICT Dst,VectorRegister4Float Src0,VectorRegister4Float Src1,VectorRegister4Float Src2)
 	{
-		const VectorRegister OneMinusAlpha = VectorSubtract(GlobalVectorConstants::FloatOne, Src2);
-		const VectorRegister Tmp = VectorMultiply(Src0, OneMinusAlpha);
+		const VectorRegister4Float OneMinusAlpha = VectorSubtract(GlobalVectorConstants::FloatOne, Src2);
+		const VectorRegister4Float Tmp = VectorMultiply(Src0, OneMinusAlpha);
 		*Dst = VectorMultiplyAdd(Src1, Src2, Tmp);
 	}
 };
 
 struct FVectorKernelRcp : public TUnaryVectorKernel<FVectorKernelRcp>
 {
-	static void VM_FORCEINLINE DoKernel(FVectorVMContext& Context, VectorRegister* RESTRICT Dst,VectorRegister Src0)
+	static void VM_FORCEINLINE DoKernel(FVectorVMContext& Context, VectorRegister4Float* RESTRICT Dst,VectorRegister4Float Src0)
 	{
 		*Dst = VectorReciprocal(Src0);
 	}
@@ -867,16 +867,16 @@ struct FVectorKernelRcp : public TUnaryVectorKernel<FVectorKernelRcp>
 // if the magnitude of the value is too small, then the result will be 0 (not NaN/Inf)
 struct FVectorKernelRcpSafe : public TUnaryVectorKernel<FVectorKernelRcpSafe>
 {
-	static void VM_FORCEINLINE DoKernel(FVectorVMContext& Context, VectorRegister* RESTRICT Dst, VectorRegister Src0)
+	static void VM_FORCEINLINE DoKernel(FVectorVMContext& Context, VectorRegister4Float* RESTRICT Dst, VectorRegister4Float Src0)
 	{
-		VectorRegister ValidMask = VectorCompareGT(VectorAbs(Src0), GlobalVectorConstants::SmallNumber);
+		VectorRegister4Float ValidMask = VectorCompareGT(VectorAbs(Src0), GlobalVectorConstants::SmallNumber);
 		*Dst = VectorSelect(ValidMask, VectorReciprocal(Src0), GlobalVectorConstants::FloatZero);
 	}
 };
 
 struct FVectorKernelRsq : public TUnaryVectorKernel<FVectorKernelRsq>
 {
-	static void VM_FORCEINLINE DoKernel(FVectorVMContext& Context, VectorRegister* RESTRICT Dst,VectorRegister Src0)
+	static void VM_FORCEINLINE DoKernel(FVectorVMContext& Context, VectorRegister4Float* RESTRICT Dst,VectorRegister4Float Src0)
 	{
 		*Dst = VectorReciprocalSqrt(Src0);
 	}
@@ -885,16 +885,16 @@ struct FVectorKernelRsq : public TUnaryVectorKernel<FVectorKernelRsq>
 // if the value is very small or negative, then the result will be 0 (not NaN/Inf/imaginary)
 struct FVectorKernelRsqSafe : public TUnaryVectorKernel<FVectorKernelRsqSafe>
 {
-	static void VM_FORCEINLINE DoKernel(FVectorVMContext& Context, VectorRegister* RESTRICT Dst, VectorRegister Src0)
+	static void VM_FORCEINLINE DoKernel(FVectorVMContext& Context, VectorRegister4Float* RESTRICT Dst, VectorRegister4Float Src0)
 	{
-		VectorRegister ValidMask = VectorCompareGT(Src0, GlobalVectorConstants::SmallNumber);
+		VectorRegister4Float ValidMask = VectorCompareGT(Src0, GlobalVectorConstants::SmallNumber);
 		*Dst = VectorSelect(ValidMask, VectorReciprocalSqrt(Src0), GlobalVectorConstants::FloatZero);
 	}
 };
 
 struct FVectorKernelSqrt : public TUnaryVectorKernel<FVectorKernelSqrt>
 {
-	static void VM_FORCEINLINE DoKernel(FVectorVMContext& Context, VectorRegister* RESTRICT Dst,VectorRegister Src0)
+	static void VM_FORCEINLINE DoKernel(FVectorVMContext& Context, VectorRegister4Float* RESTRICT Dst,VectorRegister4Float Src0)
 	{
 		// TODO: Need a SIMD sqrt!
 		*Dst = VectorReciprocal(VectorReciprocalSqrt(Src0));
@@ -903,16 +903,16 @@ struct FVectorKernelSqrt : public TUnaryVectorKernel<FVectorKernelSqrt>
 
 struct FVectorKernelSqrtSafe : public TUnaryVectorKernel<FVectorKernelSqrtSafe>
 {
-	static void VM_FORCEINLINE DoKernel(FVectorVMContext& Context, VectorRegister* RESTRICT Dst, VectorRegister Src0)
+	static void VM_FORCEINLINE DoKernel(FVectorVMContext& Context, VectorRegister4Float* RESTRICT Dst, VectorRegister4Float Src0)
 	{
-		VectorRegister ValidMask = VectorCompareGT(Src0, GlobalVectorConstants::SmallNumber);
+		VectorRegister4Float ValidMask = VectorCompareGT(Src0, GlobalVectorConstants::SmallNumber);
 		*Dst = VectorSelect(ValidMask, VectorReciprocal(VectorReciprocalSqrt(Src0)), GlobalVectorConstants::FloatZero);
 	}
 };
 
 struct FVectorKernelNeg : public TUnaryVectorKernel<FVectorKernelNeg>
 {
-	static void VM_FORCEINLINE DoKernel(FVectorVMContext& Context, VectorRegister* RESTRICT Dst,VectorRegister Src0)
+	static void VM_FORCEINLINE DoKernel(FVectorVMContext& Context, VectorRegister4Float* RESTRICT Dst,VectorRegister4Float Src0)
 	{
 		*Dst = VectorNegate(Src0);
 	}
@@ -920,7 +920,7 @@ struct FVectorKernelNeg : public TUnaryVectorKernel<FVectorKernelNeg>
 
 struct FVectorKernelAbs : public TUnaryVectorKernel<FVectorKernelAbs>
 {
-	static void VM_FORCEINLINE DoKernel(FVectorVMContext& Context, VectorRegister* RESTRICT Dst,VectorRegister Src0)
+	static void VM_FORCEINLINE DoKernel(FVectorVMContext& Context, VectorRegister4Float* RESTRICT Dst,VectorRegister4Float Src0)
 	{
 		*Dst = VectorAbs(Src0);
 	}
@@ -928,7 +928,7 @@ struct FVectorKernelAbs : public TUnaryVectorKernel<FVectorKernelAbs>
 
 struct FVectorKernelExp : public TUnaryVectorKernel<FVectorKernelExp>
 {
-	static void VM_FORCEINLINE DoKernel(FVectorVMContext& Context, VectorRegister* RESTRICT Dst, VectorRegister Src0)
+	static void VM_FORCEINLINE DoKernel(FVectorVMContext& Context, VectorRegister4Float* RESTRICT Dst, VectorRegister4Float Src0)
 	{
 		*Dst = VectorExp(Src0);
 	}
@@ -936,7 +936,7 @@ struct FVectorKernelExp : public TUnaryVectorKernel<FVectorKernelExp>
 
 struct FVectorKernelExp2 : public TUnaryVectorKernel<FVectorKernelExp2>
 {
-	static void VM_FORCEINLINE DoKernel(FVectorVMContext& Context, VectorRegister* RESTRICT Dst, VectorRegister Src0)
+	static void VM_FORCEINLINE DoKernel(FVectorVMContext& Context, VectorRegister4Float* RESTRICT Dst, VectorRegister4Float Src0)
 	{
 		*Dst = VectorExp2(Src0);
 	}
@@ -944,7 +944,7 @@ struct FVectorKernelExp2 : public TUnaryVectorKernel<FVectorKernelExp2>
 
 struct FVectorKernelLog : public TUnaryVectorKernel<FVectorKernelLog>
 {
-	static void VM_FORCEINLINE DoKernel(FVectorVMContext& Context, VectorRegister* RESTRICT Dst, VectorRegister Src0)
+	static void VM_FORCEINLINE DoKernel(FVectorVMContext& Context, VectorRegister4Float* RESTRICT Dst, VectorRegister4Float Src0)
 	{
 		*Dst = VectorLog(Src0);
 	}
@@ -952,9 +952,9 @@ struct FVectorKernelLog : public TUnaryVectorKernel<FVectorKernelLog>
 
 struct FVectorKernelLogSafe : public TUnaryVectorKernel<FVectorKernelLogSafe>
 {
-	static void VM_FORCEINLINE DoKernel(FVectorVMContext& Context, VectorRegister* RESTRICT Dst, VectorRegister Src0)
+	static void VM_FORCEINLINE DoKernel(FVectorVMContext& Context, VectorRegister4Float* RESTRICT Dst, VectorRegister4Float Src0)
 	{
-		VectorRegister ValidMask = VectorCompareGT(Src0, GlobalVectorConstants::FloatZero);
+		VectorRegister4Float ValidMask = VectorCompareGT(Src0, GlobalVectorConstants::FloatZero);
 
 		*Dst = VectorSelect(ValidMask, VectorLog(Src0), GlobalVectorConstants::FloatZero);
 	}
@@ -962,7 +962,7 @@ struct FVectorKernelLogSafe : public TUnaryVectorKernel<FVectorKernelLogSafe>
 
 struct FVectorKernelLog2 : public TUnaryVectorKernel<FVectorKernelLog2>
 {
-	static void VM_FORCEINLINE DoKernel(FVectorVMContext& Context, VectorRegister* RESTRICT Dst, VectorRegister Src0)
+	static void VM_FORCEINLINE DoKernel(FVectorVMContext& Context, VectorRegister4Float* RESTRICT Dst, VectorRegister4Float Src0)
 	{
 		*Dst = VectorLog2(Src0);
 	}
@@ -970,7 +970,7 @@ struct FVectorKernelLog2 : public TUnaryVectorKernel<FVectorKernelLog2>
 
 struct FVectorKernelStep : public TBinaryVectorKernel<FVectorKernelStep>
 {
-	static void VM_FORCEINLINE DoKernel(FVectorVMContext& Context, VectorRegister* RESTRICT Dst, VectorRegister Src0, VectorRegister Src1)
+	static void VM_FORCEINLINE DoKernel(FVectorVMContext& Context, VectorRegister4Float* RESTRICT Dst, VectorRegister4Float Src0, VectorRegister4Float Src1)
 	{
 		*Dst = VectorStep(VectorSubtract(Src1, Src0));
 	}
@@ -978,16 +978,16 @@ struct FVectorKernelStep : public TBinaryVectorKernel<FVectorKernelStep>
 
 struct FVectorKernelClamp : public TTrinaryVectorKernel<FVectorKernelClamp>
 {
-	static void VM_FORCEINLINE DoKernel(FVectorVMContext& Context, VectorRegister* RESTRICT Dst,VectorRegister Src0,VectorRegister Src1,VectorRegister Src2)
+	static void VM_FORCEINLINE DoKernel(FVectorVMContext& Context, VectorRegister4Float* RESTRICT Dst,VectorRegister4Float Src0,VectorRegister4Float Src1,VectorRegister4Float Src2)
 	{
-		const VectorRegister Tmp = VectorMax(Src0, Src1);
+		const VectorRegister4Float Tmp = VectorMax(Src0, Src1);
 		*Dst = VectorMin(Tmp, Src2);
 	}
 };
 
 struct FVectorKernelSin : public TUnaryVectorKernel<FVectorKernelSin>
 {
-	static void VM_FORCEINLINE DoKernel(FVectorVMContext& Context, VectorRegister* RESTRICT Dst, VectorRegister Src0)
+	static void VM_FORCEINLINE DoKernel(FVectorVMContext& Context, VectorRegister4Float* RESTRICT Dst, VectorRegister4Float Src0)
 	{
 		*Dst = VectorSin(Src0);
 	}
@@ -995,7 +995,7 @@ struct FVectorKernelSin : public TUnaryVectorKernel<FVectorKernelSin>
 
 struct FVectorKernelCos : public TUnaryVectorKernel<FVectorKernelCos>
 {
-	static void VM_FORCEINLINE DoKernel(FVectorVMContext& Context, VectorRegister* RESTRICT Dst, VectorRegister Src0)
+	static void VM_FORCEINLINE DoKernel(FVectorVMContext& Context, VectorRegister4Float* RESTRICT Dst, VectorRegister4Float Src0)
  	{
 		*Dst = VectorCos(Src0);
 	}
@@ -1003,7 +1003,7 @@ struct FVectorKernelCos : public TUnaryVectorKernel<FVectorKernelCos>
 
 struct FVectorKernelTan : public TUnaryVectorKernel<FVectorKernelTan>
 {
-	static void VM_FORCEINLINE DoKernel(FVectorVMContext& Context, VectorRegister* RESTRICT Dst, VectorRegister Src0)
+	static void VM_FORCEINLINE DoKernel(FVectorVMContext& Context, VectorRegister4Float* RESTRICT Dst, VectorRegister4Float Src0)
 	{
 		*Dst = VectorTan(Src0);
 	}
@@ -1011,7 +1011,7 @@ struct FVectorKernelTan : public TUnaryVectorKernel<FVectorKernelTan>
 
 struct FVectorKernelASin : public TUnaryVectorKernel<FVectorKernelASin>
 {
-	static void VM_FORCEINLINE DoKernel(FVectorVMContext& Context, VectorRegister* RESTRICT Dst, VectorRegister Src0)
+	static void VM_FORCEINLINE DoKernel(FVectorVMContext& Context, VectorRegister4Float* RESTRICT Dst, VectorRegister4Float Src0)
 	{
 		*Dst = VectorASin(Src0);
 	}
@@ -1019,7 +1019,7 @@ struct FVectorKernelASin : public TUnaryVectorKernel<FVectorKernelASin>
 
 struct FVectorKernelACos : public TUnaryVectorKernel<FVectorKernelACos>
 {
-	static void VM_FORCEINLINE DoKernel(FVectorVMContext& Context, VectorRegister* RESTRICT Dst, VectorRegister Src0)
+	static void VM_FORCEINLINE DoKernel(FVectorVMContext& Context, VectorRegister4Float* RESTRICT Dst, VectorRegister4Float Src0)
 	{
 		*Dst = VectorACos(Src0);
 	}
@@ -1027,7 +1027,7 @@ struct FVectorKernelACos : public TUnaryVectorKernel<FVectorKernelACos>
 
 struct FVectorKernelATan : public TUnaryVectorKernel<FVectorKernelATan>
 {
-	static void VM_FORCEINLINE DoKernel(FVectorVMContext& Context, VectorRegister* RESTRICT Dst, VectorRegister Src0)
+	static void VM_FORCEINLINE DoKernel(FVectorVMContext& Context, VectorRegister4Float* RESTRICT Dst, VectorRegister4Float Src0)
 	{
 		*Dst = VectorATan(Src0);
 	}
@@ -1035,7 +1035,7 @@ struct FVectorKernelATan : public TUnaryVectorKernel<FVectorKernelATan>
 
 struct FVectorKernelATan2 : public TBinaryVectorKernel<FVectorKernelATan2>
 {
-	static void VM_FORCEINLINE DoKernel(FVectorVMContext& Context, VectorRegister* RESTRICT Dst, VectorRegister Src0, VectorRegister Src1)
+	static void VM_FORCEINLINE DoKernel(FVectorVMContext& Context, VectorRegister4Float* RESTRICT Dst, VectorRegister4Float Src0, VectorRegister4Float Src1)
 	{
 		*Dst = VectorATan2(Src0, Src1);
 	}
@@ -1043,7 +1043,7 @@ struct FVectorKernelATan2 : public TBinaryVectorKernel<FVectorKernelATan2>
 
 struct FVectorKernelCeil : public TUnaryVectorKernel<FVectorKernelCeil>
 {
-	static void VM_FORCEINLINE DoKernel(FVectorVMContext& Context, VectorRegister* RESTRICT Dst, VectorRegister Src0)
+	static void VM_FORCEINLINE DoKernel(FVectorVMContext& Context, VectorRegister4Float* RESTRICT Dst, VectorRegister4Float Src0)
 	{
 		*Dst = VectorCeil(Src0);
 	}
@@ -1051,7 +1051,7 @@ struct FVectorKernelCeil : public TUnaryVectorKernel<FVectorKernelCeil>
 
 struct FVectorKernelFloor : public TUnaryVectorKernel<FVectorKernelFloor>
 {
-	static void VM_FORCEINLINE DoKernel(FVectorVMContext& Context, VectorRegister* RESTRICT Dst, VectorRegister Src0)
+	static void VM_FORCEINLINE DoKernel(FVectorVMContext& Context, VectorRegister4Float* RESTRICT Dst, VectorRegister4Float Src0)
 	{
 		*Dst = VectorFloor(Src0);
 	}
@@ -1059,17 +1059,17 @@ struct FVectorKernelFloor : public TUnaryVectorKernel<FVectorKernelFloor>
 
 struct FVectorKernelRound : public TUnaryVectorKernel<FVectorKernelRound>
 {
-	static void VM_FORCEINLINE DoKernel(FVectorVMContext& Context, VectorRegister* RESTRICT Dst, VectorRegister Src0)
+	static void VM_FORCEINLINE DoKernel(FVectorVMContext& Context, VectorRegister4Float* RESTRICT Dst, VectorRegister4Float Src0)
 	{
 		//TODO: >SSE4 has direct ops for this.		
-		VectorRegister Trunc = VectorTruncate(Src0);
+		VectorRegister4Float Trunc = VectorTruncate(Src0);
 		*Dst = VectorAdd(Trunc, VectorTruncate(VectorMultiply(VectorSubtract(Src0, Trunc), GlobalVectorConstants::FloatAlmostTwo)));
 	}
 };
 
 struct FVectorKernelMod : public TBinaryVectorKernel<FVectorKernelMod>
 {
-	static void VM_FORCEINLINE DoKernel(FVectorVMContext& Context, VectorRegister* RESTRICT Dst, VectorRegister Src0, VectorRegister Src1)
+	static void VM_FORCEINLINE DoKernel(FVectorVMContext& Context, VectorRegister4Float* RESTRICT Dst, VectorRegister4Float Src0, VectorRegister4Float Src1)
 	{
 		*Dst = VectorMod(Src0, Src1);
 	}
@@ -1077,7 +1077,7 @@ struct FVectorKernelMod : public TBinaryVectorKernel<FVectorKernelMod>
 
 struct FVectorKernelFrac : public TUnaryVectorKernel<FVectorKernelFrac>
 {
-	static void VM_FORCEINLINE DoKernel(FVectorVMContext& Context, VectorRegister* RESTRICT Dst, VectorRegister Src0)
+	static void VM_FORCEINLINE DoKernel(FVectorVMContext& Context, VectorRegister4Float* RESTRICT Dst, VectorRegister4Float Src0)
 	{
 		*Dst = VectorFractional(Src0);
 	}
@@ -1085,7 +1085,7 @@ struct FVectorKernelFrac : public TUnaryVectorKernel<FVectorKernelFrac>
 
 struct FVectorKernelTrunc : public TUnaryVectorKernel<FVectorKernelTrunc>
 {
-	static void VM_FORCEINLINE DoKernel(FVectorVMContext& Context, VectorRegister* RESTRICT Dst, VectorRegister Src0)
+	static void VM_FORCEINLINE DoKernel(FVectorVMContext& Context, VectorRegister4Float* RESTRICT Dst, VectorRegister4Float Src0)
 	{
 		*Dst = VectorTruncate(Src0);
 	}
@@ -1093,7 +1093,7 @@ struct FVectorKernelTrunc : public TUnaryVectorKernel<FVectorKernelTrunc>
 
 struct FVectorKernelCompareLT : public TBinaryVectorKernel<FVectorKernelCompareLT>
 {
-	static void VM_FORCEINLINE DoKernel(FVectorVMContext& Context, VectorRegister* RESTRICT Dst, VectorRegister Src0, VectorRegister Src1)
+	static void VM_FORCEINLINE DoKernel(FVectorVMContext& Context, VectorRegister4Float* RESTRICT Dst, VectorRegister4Float Src0, VectorRegister4Float Src1)
 	{
 		*Dst = VectorCompareLT(Src0, Src1);
 	}
@@ -1101,7 +1101,7 @@ struct FVectorKernelCompareLT : public TBinaryVectorKernel<FVectorKernelCompareL
 
 struct FVectorKernelCompareLE : public TBinaryVectorKernel<FVectorKernelCompareLE>
 {
-	static void VM_FORCEINLINE DoKernel(FVectorVMContext& Context, VectorRegister* RESTRICT Dst, VectorRegister Src0, VectorRegister Src1)
+	static void VM_FORCEINLINE DoKernel(FVectorVMContext& Context, VectorRegister4Float* RESTRICT Dst, VectorRegister4Float Src0, VectorRegister4Float Src1)
 	{
 		*Dst = VectorCompareLE(Src0, Src1);
 	}
@@ -1109,7 +1109,7 @@ struct FVectorKernelCompareLE : public TBinaryVectorKernel<FVectorKernelCompareL
 
 struct FVectorKernelCompareGT : public TBinaryVectorKernel<FVectorKernelCompareGT>
 {
-	static void VM_FORCEINLINE DoKernel(FVectorVMContext& Context, VectorRegister* RESTRICT Dst, VectorRegister Src0, VectorRegister Src1)
+	static void VM_FORCEINLINE DoKernel(FVectorVMContext& Context, VectorRegister4Float* RESTRICT Dst, VectorRegister4Float Src0, VectorRegister4Float Src1)
 	{
 		*Dst = VectorCompareGT(Src0, Src1);
 	}
@@ -1117,7 +1117,7 @@ struct FVectorKernelCompareGT : public TBinaryVectorKernel<FVectorKernelCompareG
 
 struct FVectorKernelCompareGE : public TBinaryVectorKernel<FVectorKernelCompareGE>
 {
-	static void VM_FORCEINLINE DoKernel(FVectorVMContext& Context, VectorRegister* RESTRICT Dst, VectorRegister Src0, VectorRegister Src1)
+	static void VM_FORCEINLINE DoKernel(FVectorVMContext& Context, VectorRegister4Float* RESTRICT Dst, VectorRegister4Float Src0, VectorRegister4Float Src1)
 	{
 		*Dst = VectorCompareGE(Src0, Src1);
 	}
@@ -1125,7 +1125,7 @@ struct FVectorKernelCompareGE : public TBinaryVectorKernel<FVectorKernelCompareG
 
 struct FVectorKernelCompareEQ : public TBinaryVectorKernel<FVectorKernelCompareEQ>
 {
-	static void VM_FORCEINLINE DoKernel(FVectorVMContext& Context, VectorRegister* RESTRICT Dst, VectorRegister Src0, VectorRegister Src1)
+	static void VM_FORCEINLINE DoKernel(FVectorVMContext& Context, VectorRegister4Float* RESTRICT Dst, VectorRegister4Float Src0, VectorRegister4Float Src1)
 	{
 		*Dst = VectorCompareEQ(Src0, Src1);
 	}
@@ -1133,7 +1133,7 @@ struct FVectorKernelCompareEQ : public TBinaryVectorKernel<FVectorKernelCompareE
 
 struct FVectorKernelCompareNEQ : public TBinaryVectorKernel<FVectorKernelCompareNEQ>
 {
-	static void VM_FORCEINLINE DoKernel(FVectorVMContext& Context, VectorRegister* RESTRICT Dst, VectorRegister Src0, VectorRegister Src1)
+	static void VM_FORCEINLINE DoKernel(FVectorVMContext& Context, VectorRegister4Float* RESTRICT Dst, VectorRegister4Float Src0, VectorRegister4Float Src1)
 	{		
 		*Dst = VectorCompareNE(Src0, Src1);
 	}
@@ -1141,7 +1141,7 @@ struct FVectorKernelCompareNEQ : public TBinaryVectorKernel<FVectorKernelCompare
 
 struct FVectorKernelSelect : public TTrinaryVectorKernel<FVectorKernelSelect>
 {
-	static void VM_FORCEINLINE DoKernel(FVectorVMContext& Context, VectorRegister* RESTRICT Dst, VectorRegister Mask, VectorRegister A, VectorRegister B)
+	static void VM_FORCEINLINE DoKernel(FVectorVMContext& Context, VectorRegister4Float* RESTRICT Dst, VectorRegister4Float Mask, VectorRegister4Float A, VectorRegister4Float B)
 	{
 		*Dst = VectorSelect(Mask, A, B);
 	}
@@ -1249,10 +1249,10 @@ struct FVectorKernelExitStatScope
 
 struct FVectorKernelRandom : public TUnaryVectorKernel<FVectorKernelRandom>
 {
-	static void VM_FORCEINLINE DoKernel(FVectorVMContext& Context, VectorRegister* RESTRICT Dst, VectorRegister Src0)
+	static void VM_FORCEINLINE DoKernel(FVectorVMContext& Context, VectorRegister4Float* RESTRICT Dst, VectorRegister4Float Src0)
 	{
 		//EEK!. Improve this. Implement GPU style seeded rand instead of this.
-		VectorRegister Result = MakeVectorRegister(Context.RandStream.GetFraction(),
+		VectorRegister4Float Result = MakeVectorRegister(Context.RandStream.GetFraction(),
 			Context.RandStream.GetFraction(),
 			Context.RandStream.GetFraction(),
 			Context.RandStream.GetFraction());
@@ -1263,9 +1263,9 @@ struct FVectorKernelRandom : public TUnaryVectorKernel<FVectorKernelRandom>
 /* gaussian distribution random number (not working yet) */
 struct FVectorKernelRandomGauss : public TBinaryVectorKernel<FVectorKernelRandomGauss>
 {
-	static void VM_FORCEINLINE DoKernel(FVectorVMContext& Context, VectorRegister* RESTRICT Dst, VectorRegister Src0, VectorRegister Src1)
+	static void VM_FORCEINLINE DoKernel(FVectorVMContext& Context, VectorRegister4Float* RESTRICT Dst, VectorRegister4Float Src0, VectorRegister4Float Src1)
 	{
-		VectorRegister Result = MakeVectorRegister(Context.RandStream.GetFraction(),
+		VectorRegister4Float Result = MakeVectorRegister(Context.RandStream.GetFraction(),
 			Context.RandStream.GetFraction(),
 			Context.RandStream.GetFraction(),
 			Context.RandStream.GetFraction());
@@ -1274,9 +1274,9 @@ struct FVectorKernelRandomGauss : public TBinaryVectorKernel<FVectorKernelRandom
 		Result = VectorMultiply(MakeVectorRegister(3.0f, 3.0f, 3.0f, 3.0f), Result);
 
 		// taylor series gaussian approximation
-		const VectorRegister SPi2 = VectorReciprocal(VectorReciprocalSqrt(MakeVectorRegister(2 * PI, 2 * PI, 2 * PI, 2 * PI)));
-		VectorRegister Gauss = VectorReciprocal(SPi2);
-		VectorRegister Div = VectorMultiply(GlobalVectorConstants::FloatTwo, SPi2);
+		const VectorRegister4Float SPi2 = VectorReciprocal(VectorReciprocalSqrt(MakeVectorRegister(2 * PI, 2 * PI, 2 * PI, 2 * PI)));
+		VectorRegister4Float Gauss = VectorReciprocal(SPi2);
+		VectorRegister4Float Div = VectorMultiply(GlobalVectorConstants::FloatTwo, SPi2);
 		Gauss = VectorSubtract(Gauss, VectorDivide(VectorMultiply(Result, Result), Div));
 		Div = VectorMultiply(MakeVectorRegister(8.0f, 8.0f, 8.0f, 8.0f), SPi2);
 		Gauss = VectorAdd(Gauss, VectorDivide(VectorPow(MakeVectorRegister(4.0f, 4.0f, 4.0f, 4.0f), Result), Div));
@@ -1291,7 +1291,7 @@ struct FVectorKernelRandomGauss : public TBinaryVectorKernel<FVectorKernelRandom
 
 struct FVectorKernelMin : public TBinaryVectorKernel<FVectorKernelMin>
 {
-	static void VM_FORCEINLINE DoKernel(FVectorVMContext& Context, VectorRegister* RESTRICT Dst,VectorRegister Src0,VectorRegister Src1)
+	static void VM_FORCEINLINE DoKernel(FVectorVMContext& Context, VectorRegister4Float* RESTRICT Dst,VectorRegister4Float Src0,VectorRegister4Float Src1)
 	{
 		*Dst = VectorMin(Src0, Src1);
 	}
@@ -1299,7 +1299,7 @@ struct FVectorKernelMin : public TBinaryVectorKernel<FVectorKernelMin>
 
 struct FVectorKernelMax : public TBinaryVectorKernel<FVectorKernelMax>
 {
-	static void VM_FORCEINLINE DoKernel(FVectorVMContext& Context, VectorRegister* RESTRICT Dst,VectorRegister Src0,VectorRegister Src1)
+	static void VM_FORCEINLINE DoKernel(FVectorVMContext& Context, VectorRegister4Float* RESTRICT Dst,VectorRegister4Float Src0,VectorRegister4Float Src1)
 	{
 		*Dst = VectorMax(Src0, Src1);
 	}
@@ -1307,7 +1307,7 @@ struct FVectorKernelMax : public TBinaryVectorKernel<FVectorKernelMax>
 
 struct FVectorKernelPow : public TBinaryVectorKernel<FVectorKernelPow>
 {
-	static void VM_FORCEINLINE DoKernel(FVectorVMContext& Context, VectorRegister* RESTRICT Dst,VectorRegister Src0,VectorRegister Src1)
+	static void VM_FORCEINLINE DoKernel(FVectorVMContext& Context, VectorRegister4Float* RESTRICT Dst,VectorRegister4Float Src0,VectorRegister4Float Src1)
 	{
 		*Dst = VectorPow(Src0, Src1);
 	}
@@ -1316,16 +1316,16 @@ struct FVectorKernelPow : public TBinaryVectorKernel<FVectorKernelPow>
 // if the base is small, then the result will be 0
 struct FVectorKernelPowSafe : public TBinaryVectorKernel<FVectorKernelPowSafe>
 {
-	static void VM_FORCEINLINE DoKernel(FVectorVMContext& Context, VectorRegister* RESTRICT Dst, VectorRegister Src0, VectorRegister Src1)
+	static void VM_FORCEINLINE DoKernel(FVectorVMContext& Context, VectorRegister4Float* RESTRICT Dst, VectorRegister4Float Src0, VectorRegister4Float Src1)
 	{
-		VectorRegister ValidMask = VectorCompareGT(Src0, GlobalVectorConstants::SmallNumber);
+		VectorRegister4Float ValidMask = VectorCompareGT(Src0, GlobalVectorConstants::SmallNumber);
 		*Dst = VectorSelect(ValidMask, VectorPow(Src0, Src1), GlobalVectorConstants::FloatZero);
 	}
 };
 
 struct FVectorKernelSign : public TUnaryVectorKernel<FVectorKernelSign>
 {
-	static void VM_FORCEINLINE DoKernel(FVectorVMContext& Context, VectorRegister* RESTRICT Dst, VectorRegister Src0)
+	static void VM_FORCEINLINE DoKernel(FVectorVMContext& Context, VectorRegister4Float* RESTRICT Dst, VectorRegister4Float Src0)
 	{
 		*Dst = VectorSign(Src0);
 	}
@@ -1461,48 +1461,48 @@ namespace VectorVMNoise
 //Possibly needs to be a data interface as the VM can't output Vectors?
 struct FVectorKernelNoise : public TUnaryVectorKernel<FVectorKernelNoise>
 {
-	static VectorRegister RandomTable[17][17][17];
+	static VectorRegister4Float RandomTable[17][17][17];
 
-	static void VM_FORCEINLINE DoKernel(FVectorVMContext& Context, VectorRegister* RESTRICT Dst, VectorRegister Src0)
+	static void VM_FORCEINLINE DoKernel(FVectorVMContext& Context, VectorRegister4Float* RESTRICT Dst, VectorRegister4Float Src0)
 	{
-		const VectorRegister VecSize = MakeVectorRegister(16.0f, 16.0f, 16.0f, 16.0f);
+		const VectorRegister4Float VecSize = MakeVectorRegister(16.0f, 16.0f, 16.0f, 16.0f);
 
 		*Dst = GlobalVectorConstants::FloatZero;
 		
 		for (uint32 i = 1; i < 2; i++)
 		{
 			float Di = 0.2f * (1.0f/(1<<i));
-			VectorRegister Div = MakeVectorRegister(Di, Di, Di, Di);
-			VectorRegister Coords = VectorMod( VectorAbs( VectorMultiply(Src0, Div) ), VecSize );
+			VectorRegister4Float Div = MakeVectorRegister(Di, Di, Di, Di);
+			VectorRegister4Float Coords = VectorMod( VectorAbs( VectorMultiply(Src0, Div) ), VecSize );
 			const float *CoordPtr = reinterpret_cast<float const*>(&Coords);
 			const int32 Cx = CoordPtr[0];
 			const int32 Cy = CoordPtr[1];
 			const int32 Cz = CoordPtr[2];
 
-			VectorRegister Frac = VectorFractional(Coords);
-			VectorRegister Alpha = VectorReplicate(Frac, 0);
-			VectorRegister OneMinusAlpha = VectorSubtract(GlobalVectorConstants::FloatOne, Alpha);
+			VectorRegister4Float Frac = VectorFractional(Coords);
+			VectorRegister4Float Alpha = VectorReplicate(Frac, 0);
+			VectorRegister4Float OneMinusAlpha = VectorSubtract(GlobalVectorConstants::FloatOne, Alpha);
 			
-			VectorRegister XV1 = VectorMultiplyAdd(RandomTable[Cx][Cy][Cz], Alpha, VectorMultiply(RandomTable[Cx+1][Cy][Cz], OneMinusAlpha));
-			VectorRegister XV2 = VectorMultiplyAdd(RandomTable[Cx][Cy+1][Cz], Alpha, VectorMultiply(RandomTable[Cx+1][Cy+1][Cz], OneMinusAlpha));
-			VectorRegister XV3 = VectorMultiplyAdd(RandomTable[Cx][Cy][Cz+1], Alpha, VectorMultiply(RandomTable[Cx+1][Cy][Cz+1], OneMinusAlpha));
-			VectorRegister XV4 = VectorMultiplyAdd(RandomTable[Cx][Cy+1][Cz+1], Alpha, VectorMultiply(RandomTable[Cx+1][Cy+1][Cz+1], OneMinusAlpha));
+			VectorRegister4Float XV1 = VectorMultiplyAdd(RandomTable[Cx][Cy][Cz], Alpha, VectorMultiply(RandomTable[Cx+1][Cy][Cz], OneMinusAlpha));
+			VectorRegister4Float XV2 = VectorMultiplyAdd(RandomTable[Cx][Cy+1][Cz], Alpha, VectorMultiply(RandomTable[Cx+1][Cy+1][Cz], OneMinusAlpha));
+			VectorRegister4Float XV3 = VectorMultiplyAdd(RandomTable[Cx][Cy][Cz+1], Alpha, VectorMultiply(RandomTable[Cx+1][Cy][Cz+1], OneMinusAlpha));
+			VectorRegister4Float XV4 = VectorMultiplyAdd(RandomTable[Cx][Cy+1][Cz+1], Alpha, VectorMultiply(RandomTable[Cx+1][Cy+1][Cz+1], OneMinusAlpha));
 
 			Alpha = VectorReplicate(Frac, 1);
 			OneMinusAlpha = VectorSubtract(GlobalVectorConstants::FloatOne, Alpha);
-			VectorRegister YV1 = VectorMultiplyAdd(XV1, Alpha, VectorMultiply(XV2, OneMinusAlpha));
-			VectorRegister YV2 = VectorMultiplyAdd(XV3, Alpha, VectorMultiply(XV4, OneMinusAlpha));
+			VectorRegister4Float YV1 = VectorMultiplyAdd(XV1, Alpha, VectorMultiply(XV2, OneMinusAlpha));
+			VectorRegister4Float YV2 = VectorMultiplyAdd(XV3, Alpha, VectorMultiply(XV4, OneMinusAlpha));
 
 			Alpha = VectorReplicate(Frac, 2);
 			OneMinusAlpha = VectorSubtract(GlobalVectorConstants::FloatOne, Alpha);
-			VectorRegister ZV = VectorMultiplyAdd(YV1, Alpha, VectorMultiply(YV2, OneMinusAlpha));
+			VectorRegister4Float ZV = VectorMultiplyAdd(YV1, Alpha, VectorMultiply(YV2, OneMinusAlpha));
 
 			*Dst = VectorAdd(*Dst, ZV);
 		}
 	}
 };
 
-VectorRegister FVectorKernelNoise::RandomTable[17][17][17];
+VectorRegister4Float FVectorKernelNoise::RandomTable[17][17][17];
 
 //////////////////////////////////////////////////////////////////////////
 //Special Kernels.
@@ -1636,15 +1636,15 @@ struct FVectorKernelReadInput
 
 	static VM_FORCEINLINE void Exec(FVectorVMContext& Context)
 	{
-		static const int32 InstancesPerVector = sizeof(VectorRegister) / sizeof(SourceType);
+		static const int32 InstancesPerVector = sizeof(VectorRegister4Float ) / sizeof(SourceType);
 
 		const int32 DataSetIndex = Context.DecodeU16();
 		const int32 InputRegisterIdx = Context.DecodeU16();
 		const int32 DestRegisterIdx = Context.DecodeU16();
 		int32 Loops = Context.GetNumLoops<InstancesPerVector>();
 
-		VectorRegister* DestReg = (VectorRegister*)(Context.GetTempRegister(DestRegisterIdx));
-		VectorRegister* InputReg = (VectorRegister*)(Context.GetInputRegister<SourceType, TypeOffset>(DataSetIndex, InputRegisterIdx) + Context.GetStartInstance());
+		VectorRegister4Float* DestReg = (VectorRegister4Float*)(Context.GetTempRegister(DestRegisterIdx));
+		VectorRegister4Float* InputReg = (VectorRegister4Float*)(Context.GetInputRegister<SourceType, TypeOffset>(DataSetIndex, InputRegisterIdx) + Context.GetStartInstance());
 
 		//TODO: We can actually do some scalar loads into the first and final vectors to get around alignment issues and then use the aligned load for all others.
 		while(Loops > 0)
@@ -1715,15 +1715,15 @@ struct FVectorKernelReadInputNoAdvance
 
 	static VM_FORCEINLINE void Exec(FVectorVMContext& Context)
 	{
-		static const int32 InstancesPerVector = sizeof(VectorRegister) / sizeof(T);
+		static const int32 InstancesPerVector = sizeof(VectorRegister4Float) / sizeof(T);
 
 		const int32 DataSetIndex = Context.DecodeU16();
 		const int32 InputRegisterIdx = Context.DecodeU16();
 		const int32 DestRegisterIdx = Context.DecodeU16();
 		const int32 Loops = Context.GetNumLoops<InstancesPerVector>();
 
-		VectorRegister* DestReg = (VectorRegister*)(Context.GetTempRegister(DestRegisterIdx));
-		VectorRegister* InputReg = (VectorRegister*)(Context.GetInputRegister<T, TypeOffset>(DataSetIndex, InputRegisterIdx));
+		VectorRegister4Float* DestReg = (VectorRegister4Float*)(Context.GetTempRegister(DestRegisterIdx));
+		VectorRegister4Float* InputReg = (VectorRegister4Float*)(Context.GetInputRegister<T, TypeOffset>(DataSetIndex, InputRegisterIdx));
 
 		//TODO: We can actually do some scalar loads into the first and final vectors to get around alignment issues and then use the aligned load for all others.
 		for (int32 i = 0; i < Loops; ++i)
@@ -1752,7 +1752,7 @@ struct FVectorKernelReadInputNoAdvance<FFloat16, 2>
 		const int32 DestRegisterIdx = Context.DecodeU16();
 		const int32 Loops = Context.GetNumLoops<4>();
 
-		VectorRegister* DestReg = (VectorRegister*)(Context.GetTempRegister(DestRegisterIdx));
+		VectorRegister4Float* DestReg = (VectorRegister4Float*)(Context.GetTempRegister(DestRegisterIdx));
 		uint16* InputReg = (uint16*)(Context.GetInputRegister<FFloat16, 2>(DataSetIndex, InputRegisterIdx));
 
 		//TODO: We can actually do some scalar loads into the first and final vectors to get around alignment issues and then use the aligned load for all others.
@@ -2348,36 +2348,36 @@ struct FVectorIntKernelLogicNot : TUnaryVectorIntKernel<FVectorIntKernelLogicNot
 
 //conversions
 //f2i,
-struct FVectorKernelFloatToInt : TUnaryKernel<FVectorKernelFloatToInt, FRegisterHandler<VectorRegister4Int>, FConstantHandler<VectorRegister>, FRegisterHandler<VectorRegister>, VECTOR_WIDTH_FLOATS>
+struct FVectorKernelFloatToInt : TUnaryKernel<FVectorKernelFloatToInt, FRegisterHandler<VectorRegister4Int>, FConstantHandler<VectorRegister4Float>, FRegisterHandler<VectorRegister4Float>, VECTOR_WIDTH_FLOATS>
 {
-	static void VM_FORCEINLINE DoKernel(FVectorVMContext& Context, VectorRegister4Int* Dst, VectorRegister Src0)
+	static void VM_FORCEINLINE DoKernel(FVectorVMContext& Context, VectorRegister4Int* Dst, VectorRegister4Float Src0)
 	{
 		*Dst = VectorFloatToInt(Src0);
 	}
 };
 
 //i2f,
-struct FVectorKernelIntToFloat : TUnaryKernel<FVectorKernelIntToFloat, FRegisterHandler<VectorRegister>, FConstantHandler<VectorRegister4Int>, FRegisterHandler<VectorRegister4Int>, VECTOR_WIDTH_FLOATS>
+struct FVectorKernelIntToFloat : TUnaryKernel<FVectorKernelIntToFloat, FRegisterHandler<VectorRegister4Float>, FConstantHandler<VectorRegister4Int>, FRegisterHandler<VectorRegister4Int>, VECTOR_WIDTH_FLOATS>
 {
-	static void VM_FORCEINLINE DoKernel(FVectorVMContext& Context, VectorRegister* Dst, VectorRegister4Int Src0)
+	static void VM_FORCEINLINE DoKernel(FVectorVMContext& Context, VectorRegister4Float* Dst, VectorRegister4Int Src0)
 	{
 		*Dst = VectorIntToFloat(Src0);
 	}
 };
 
 //f2b,
-struct FVectorKernelFloatToBool : TUnaryKernel<FVectorKernelFloatToBool, FRegisterHandler<VectorRegister>, FConstantHandler<VectorRegister>, FRegisterHandler<VectorRegister>, VECTOR_WIDTH_FLOATS>
+struct FVectorKernelFloatToBool : TUnaryKernel<FVectorKernelFloatToBool, FRegisterHandler<VectorRegister4Float>, FConstantHandler<VectorRegister4Float>, FRegisterHandler<VectorRegister4Float>, VECTOR_WIDTH_FLOATS>
 {
-	static void VM_FORCEINLINE DoKernel(FVectorVMContext& Context, VectorRegister* Dst, VectorRegister Src0)
+	static void VM_FORCEINLINE DoKernel(FVectorVMContext& Context, VectorRegister4Float* Dst, VectorRegister4Float Src0)
 	{		
 		*Dst = VectorCompareGT(Src0, GlobalVectorConstants::FloatZero);
 	}
 };
 
 //b2f,
-struct FVectorKernelBoolToFloat : TUnaryKernel<FVectorKernelBoolToFloat, FRegisterHandler<VectorRegister>, FConstantHandler<VectorRegister>, FRegisterHandler<VectorRegister>, VECTOR_WIDTH_FLOATS>
+struct FVectorKernelBoolToFloat : TUnaryKernel<FVectorKernelBoolToFloat, FRegisterHandler<VectorRegister4Float>, FConstantHandler<VectorRegister4Float>, FRegisterHandler<VectorRegister4Float>, VECTOR_WIDTH_FLOATS>
 {
-	static void VM_FORCEINLINE DoKernel(FVectorVMContext& Context, VectorRegister* Dst, VectorRegister Src0)
+	static void VM_FORCEINLINE DoKernel(FVectorVMContext& Context, VectorRegister4Float* Dst, VectorRegister4Float Src0)
 	{
 		*Dst = VectorSelect(Src0, GlobalVectorConstants::FloatOne, GlobalVectorConstants::FloatZero);
 	}
@@ -2851,7 +2851,7 @@ struct FBatchedWriteIndexedOutput
 	static void HandleRegisterValidIndices(FVectorVMContext& Context)
 	{
 		FDataSetCounterHandler CounterHandler(Context);
-		FRegisterHandler<VectorRegister> ValidReader(Context);
+		FRegisterHandler<VectorRegister4Float> ValidReader(Context);
 		FRegisterHandler<int8> Dst(Context);
 
 		int8* DestAddr = Dst.GetDest();
@@ -3141,7 +3141,7 @@ struct FBatchedWriteIndexedOutput
 			//						our output buffers padded out to 16 bytes!
 			//	VectorIntShuffle(	- swizzle our source register to pack the valid entries at the beginning, with 0s at the end
 			//    Source,			- source data
-			//    ShuffleMask),		- result of the VectorMaskBits done in the acquireindex, int8/VectorRegister of input
+			//    ShuffleMask),		- result of the VectorMaskBits done in the acquireindex, int8/VectorRegister4Float of input
 			//  Destination);
 			VectorIntStore(VectorIntShuffle(SourceVectors[SourceIt], VectorVMConstants::RegisterShuffleMask[ShuffleMask]), OutputElements);
 
@@ -3172,7 +3172,7 @@ struct FBatchedWriteIndexedOutput
 			//						our output buffers padded out to 16 bytes!
 			//	VectorIntShuffle(	- swizzle our source register to pack the valid entries at the beginning, with 0s at the end
 			//    Source,			- source data
-			//    ShuffleMask),		- result of the VectorMaskBits done in the acquireindex, int8/VectorRegister of input
+			//    ShuffleMask),		- result of the VectorMaskBits done in the acquireindex, int8/VectorRegister4Float of input
 			//  Destination);
 
 			const VectorRegister4Int ShuffledFloats = VectorIntShuffle(SourceVectors[SourceIt], VectorVMConstants::RegisterShuffleMask[ShuffleMask]);
diff --git a/Engine/Source/Runtime/VectorVM/Private/VectorVMPrivate.h b/Engine/Source/Runtime/VectorVM/Private/VectorVMPrivate.h
--- a/Engine/Source/Runtime/VectorVM/Private/VectorVMPrivate.h
+++ b/Engine/Source/Runtime/VectorVM/Private/VectorVMPrivate.h
@@ -21,8 +21,8 @@ namespace VectorVM
 struct FDummyHandler
 {
 	FORCEINLINE void Advance(){ }
-	FORCEINLINE VectorRegister Get(){ return GlobalVectorConstants::FloatZero; }
-	FORCEINLINE VectorRegister GetValue(){ return GlobalVectorConstants::FloatZero; }
+	FORCEINLINE VectorRegister4Float Get(){ return GlobalVectorConstants::FloatZero; }
+	FORCEINLINE VectorRegister4Float GetValue(){ return GlobalVectorConstants::FloatZero; }
 };
 extern FDummyHandler DummyHandler;
 
