- Add optional bool to move UVs into viewport space

diff --git a/Engine/Plugins/FX/Niagara/Shaders/Private/NiagaraDataInterfaceGBuffer.ush b/Engine/Plugins/FX/Niagara/Shaders/Private/NiagaraDataInterfaceGBuffer.ush
--- a/Engine/Plugins/FX/Niagara/Shaders/Private/NiagaraDataInterfaceGBuffer.ush
+++ b/Engine/Plugins/FX/Niagara/Shaders/Private/NiagaraDataInterfaceGBuffer.ush
@@ -13,9 +13,19 @@
 Texture2D NDIGBuffer_VelocityTexture;
 SamplerState NDIGBuffer_VelocityTextureSampler;
 
-void DIGBuffer_DecodeDiffuseColor(in float2 ScreenUV, out bool IsValid, out float3 OutValue)
+float2 DIGBuffer_ApplyViewportOffset(float2 UV, bool bApplyViewportOffset)
+{
+	if (bApplyViewportOffset)
+	{
+		UV = (UV * float2(2.0f, -2.0f) + float2(-1.0f, 1.0f)) * View.ScreenPositionScaleBias.xy + View.ScreenPositionScaleBias.wz;
+	}
+	return UV;
+}
+
+void DIGBuffer_DecodeDiffuseColor(in float2 ScreenUV, in bool bApplyViewportOffset, out bool IsValid, out float3 OutValue)
 {
 #if NDIGBUFFER_HAS_GBUFFER
+	ScreenUV = DIGBuffer_ApplyViewportOffset(ScreenUV, bApplyViewportOffset);
 	FGBufferData GBufferData = GetGBufferData(ScreenUV);
 	IsValid = true;
 	OutValue = GBufferData.DiffuseColor;
@@ -25,9 +35,10 @@ void DIGBuffer_DecodeDiffuseColor(in float2 ScreenUV, out bool IsValid, out floa
 #endif
 }
 
-void DIGBuffer_DecodeWorldNormal(in float2 ScreenUV, out bool IsValid, out float3 OutValue)
+void DIGBuffer_DecodeWorldNormal(in float2 ScreenUV, in bool bApplyViewportOffset, out bool IsValid, out float3 OutValue)
 {
 #if NDIGBUFFER_HAS_GBUFFER
+	ScreenUV = DIGBuffer_ApplyViewportOffset(ScreenUV, bApplyViewportOffset);
 	FGBufferData GBufferData = GetGBufferData(ScreenUV);
 	IsValid = true;
 	OutValue = GBufferData.WorldNormal;
@@ -37,11 +48,12 @@ void DIGBuffer_DecodeWorldNormal(in float2 ScreenUV, out bool IsValid, out float
 #endif
 }
 
-void DIGBuffer_DecodeScreenVelocity(in float2 ScreenUV, out bool IsValid, out float3 OutValue)
+void DIGBuffer_DecodeScreenVelocity(in float2 ScreenUV, in bool bApplyViewportOffset, out bool IsValid, out float3 OutValue)
 {
 	IsValid = false;
 	OutValue = 0;
 #if NDIGBUFFER_HAS_GBUFFER
+	ScreenUV = DIGBuffer_ApplyViewportOffset(ScreenUV, bApplyViewportOffset);
 	float4 EncodedVelocity = 0;
 	#if NDIGBUFFER_USE_GBUFFER_VELOCITY
 		EncodedVelocity = GetGBufferData(ScreenUV).Velocity;
@@ -60,11 +72,12 @@ void DIGBuffer_DecodeScreenVelocity(in float2 ScreenUV, out bool IsValid, out fl
 #endif
 }
 
-void DIGBuffer_DecodeWorldVelocity(in float2 ScreenUV, out bool IsValid, out float3 OutValue)
+void DIGBuffer_DecodeWorldVelocity(in float2 ScreenUV, in bool bApplyViewportOffset, out bool IsValid, out float3 OutValue)
 {
 	IsValid = false;
 	OutValue = 0;
 #if NDIGBUFFER_HAS_GBUFFER
+	ScreenUV = DIGBuffer_ApplyViewportOffset(ScreenUV, bApplyViewportOffset);
 	float4 EncodedVelocity = 0;
 	#if NDIGBUFFER_USE_GBUFFER_VELOCITY
 		EncodedVelocity = GetGBufferData(ScreenUV).Velocity;
@@ -100,9 +113,10 @@ void DIGBuffer_DecodeWorldVelocity(in float2 ScreenUV, out bool IsValid, out flo
 #endif
 }
 
-void DIGBuffer_DecodeBaseColor(in float2 ScreenUV, out bool IsValid, out float3 OutValue)
+void DIGBuffer_DecodeBaseColor(in float2 ScreenUV, in bool bApplyViewportOffset, out bool IsValid, out float3 OutValue)
 {
 #if NDIGBUFFER_HAS_GBUFFER
+	ScreenUV = DIGBuffer_ApplyViewportOffset(ScreenUV, bApplyViewportOffset);
 	FGBufferData GBufferData = GetGBufferData(ScreenUV);
 	IsValid = true;
 	OutValue = GBufferData.BaseColor;
@@ -112,9 +126,10 @@ void DIGBuffer_DecodeBaseColor(in float2 ScreenUV, out bool IsValid, out float3
 #endif
 }
 
-void DIGBuffer_DecodeMetallic(in float2 ScreenUV, out bool IsValid, out float OutValue)
+void DIGBuffer_DecodeMetallic(in float2 ScreenUV, in bool bApplyViewportOffset, out bool IsValid, out float OutValue)
 {
 #if NDIGBUFFER_HAS_GBUFFER
+	ScreenUV = DIGBuffer_ApplyViewportOffset(ScreenUV, bApplyViewportOffset);
 	FGBufferData GBufferData = GetGBufferData(ScreenUV);
 	IsValid = true;
 	OutValue = GBufferData.Metallic;
@@ -124,9 +139,10 @@ void DIGBuffer_DecodeMetallic(in float2 ScreenUV, out bool IsValid, out float Ou
 #endif
 }
 
-void DIGBuffer_DecodeSpecular(in float2 ScreenUV, out bool IsValid, out float OutValue)
+void DIGBuffer_DecodeSpecular(in float2 ScreenUV, in bool bApplyViewportOffset, out bool IsValid, out float OutValue)
 {
 #if NDIGBUFFER_HAS_GBUFFER
+	ScreenUV = DIGBuffer_ApplyViewportOffset(ScreenUV, bApplyViewportOffset);
 	FGBufferData GBufferData = GetGBufferData(ScreenUV);
 	IsValid = true;
 	OutValue = GBufferData.Specular;
@@ -136,9 +152,10 @@ void DIGBuffer_DecodeSpecular(in float2 ScreenUV, out bool IsValid, out float Ou
 #endif
 }
 
-void DIGBuffer_DecodeRoughness(in float2 ScreenUV, out bool IsValid, out float OutValue)
+void DIGBuffer_DecodeRoughness(in float2 ScreenUV, in bool bApplyViewportOffset, out bool IsValid, out float OutValue)
 {
 #if NDIGBUFFER_HAS_GBUFFER
+	ScreenUV = DIGBuffer_ApplyViewportOffset(ScreenUV, bApplyViewportOffset);
 	FGBufferData GBufferData = GetGBufferData(ScreenUV);
 	IsValid = true;
 	OutValue = GBufferData.Roughness;
@@ -148,9 +165,10 @@ void DIGBuffer_DecodeRoughness(in float2 ScreenUV, out bool IsValid, out float O
 #endif
 }
 
-void DIGBuffer_DecodeDepth(in float2 ScreenUV, out bool IsValid, out float OutValue)
+void DIGBuffer_DecodeDepth(in float2 ScreenUV, in bool bApplyViewportOffset, out bool IsValid, out float OutValue)
 {
 #if NDIGBUFFER_HAS_GBUFFER
+	ScreenUV = DIGBuffer_ApplyViewportOffset(ScreenUV, bApplyViewportOffset);
 	FGBufferData GBufferData = GetGBufferData(ScreenUV);
 	IsValid = true;
 	OutValue = GBufferData.Depth;
@@ -169,11 +187,12 @@ void DIGBuffer_DecodeDepth(in float2 ScreenUV, out bool IsValid, out float OutVa
 
 #if 0
 // This has not been tested yet
-void DIGBuffer_DecodeStencil(in float2 ScreenUV, out bool IsValid, out int OutValue)
+void DIGBuffer_DecodeStencil(in float2 ScreenUV, in bool bApplyViewportOffset, out bool IsValid, out int OutValue)
 {
 #if NDIGBUFFER_HAS_GBUFFER
 	IsValid = true;
 
+	ScreenUV = DIGBuffer_ApplyViewportOffset(ScreenUV, bApplyViewportOffset);
 	uint2 IntUV = (uint2)(ScreenUV.xy * (float2)View.BufferSizeAndInvSize.xy);
 	uint Stencil = SceneStencilTexture.Load(uint3(IntUV, 0)) STENCIL_COMPONENT_SWIZZLE;
 	OutValue = int(Stencil);
@@ -184,9 +203,10 @@ void DIGBuffer_DecodeStencil(in float2 ScreenUV, out bool IsValid, out int OutVa
 }
 #endif
 
-void DIGBuffer_DecodeCustomDepth(in float2 ScreenUV, out bool IsValid, out float OutValue)
+void DIGBuffer_DecodeCustomDepth(in float2 ScreenUV, in bool bApplyViewportOffset, out bool IsValid, out float OutValue)
 {
 #if NDIGBUFFER_HAS_GBUFFER
+	ScreenUV = DIGBuffer_ApplyViewportOffset(ScreenUV, bApplyViewportOffset);
 	FGBufferData GBufferData = GetGBufferData(ScreenUV);
 	IsValid = true;
 	OutValue = GBufferData.CustomDepth;
@@ -199,9 +219,10 @@ void DIGBuffer_DecodeCustomDepth(in float2 ScreenUV, out bool IsValid, out float
 #endif
 }
 
-void DIGBuffer_DecodeCustomStencil(in float2 ScreenUV, out bool IsValid, out int OutValue)
+void DIGBuffer_DecodeCustomStencil(in float2 ScreenUV, in bool bApplyViewportOffset, out bool IsValid, out int OutValue)
 {
 #if NDIGBUFFER_HAS_GBUFFER
+	ScreenUV = DIGBuffer_ApplyViewportOffset(ScreenUV, bApplyViewportOffset);
 	FGBufferData GBufferData = GetGBufferData(ScreenUV);
 	IsValid = true;
 	OutValue = GBufferData.CustomStencil;
@@ -214,10 +235,11 @@ void DIGBuffer_DecodeCustomStencil(in float2 ScreenUV, out bool IsValid, out int
 #endif
 }
 
-void DIGBuffer_DecodeSceneColor(in float2 ScreenUV, out bool IsValid, out float4 OutValue)
+void DIGBuffer_DecodeSceneColor(in float2 ScreenUV, in bool bApplyViewportOffset, out bool IsValid, out float4 OutValue)
 {
 #if NDIGBUFFER_HAS_GBUFFER
 	IsValid = true;
+	ScreenUV = DIGBuffer_ApplyViewportOffset(ScreenUV, bApplyViewportOffset);
 	OutValue = Texture2DSampleLevel(SceneTexturesStruct_SceneColorTexture, SceneTexturesStruct_SceneColorTextureSampler, ScreenUV, 0);
 	OutValue.rgb *= View.OneOverPreExposure.xxx;
 #elif NDIGBUFFER_IS_MOBILE
diff --git a/Engine/Plugins/FX/Niagara/Source/Niagara/Private/DataInterface/NiagaraDataInterfaceGBuffer.cpp b/Engine/Plugins/FX/Niagara/Source/Niagara/Private/DataInterface/NiagaraDataInterfaceGBuffer.cpp
--- a/Engine/Plugins/FX/Niagara/Source/Niagara/Private/DataInterface/NiagaraDataInterfaceGBuffer.cpp
+++ b/Engine/Plugins/FX/Niagara/Source/Niagara/Private/DataInterface/NiagaraDataInterfaceGBuffer.cpp
@@ -13,6 +13,18 @@
 
 namespace NiagaraDataInterfaceGBufferLocal
 {
+	struct EDIFunctionVersion
+	{
+		enum Type
+		{
+			InitialVersion = 0,
+			AddedApplyViewportOffset = 1,
+
+			VersionPlusOne,
+			LatestVersion = VersionPlusOne - 1
+		};
+	};
+
 	struct FGBufferAttribute
 	{
 		FGBufferAttribute(const TCHAR* InAttributeName, const TCHAR* InAttributeType, FNiagaraTypeDefinition InTypeDef, FText InDescription)
@@ -170,8 +182,12 @@ void UNiagaraDataInterfaceGBuffer::GetFunctions(TArray<FNiagaraFunctionSignature
 		Signature.bExperimental = true;
 		Signature.Inputs.Add(FNiagaraVariable(FNiagaraTypeDefinition(GetClass()), TEXT("GBufferInterface")));
 		Signature.Inputs.Add(FNiagaraVariable(FNiagaraTypeDefinition::GetVec2Def(), TEXT("ScreenUV")));
+		Signature.Inputs.Add_GetRef(FNiagaraVariable(FNiagaraTypeDefinition::GetBoolDef(), TEXT("ApplyViewportOffset"))).SetValue(true);
 		Signature.Outputs.Add(FNiagaraVariable(FNiagaraTypeDefinition::GetBoolDef(), TEXT("IsValid")));
 		Signature.Outputs.Add(FNiagaraVariable(Attribute.TypeDef, Attribute.AttributeName));
+#if WITH_EDITORONLY_DATA
+		Signature.FunctionVersion = EDIFunctionVersion::LatestVersion;
+#endif
 	}
 }
 
@@ -207,7 +223,7 @@ bool UNiagaraDataInterfaceGBuffer::GetFunctionHLSL(const FNiagaraDataInterfaceGP
 			ArgsSample.Emplace(TEXT("AttributeName"), Attribute.AttributeName);
 			ArgsSample.Emplace(TEXT("AttributeType"), Attribute.AttributeType);
 
-			static const TCHAR* FormatSample = TEXT("void {InstanceFunctionName}(float2 ScreenUV, out bool IsValid, out {AttributeType} {AttributeName}) { DIGBuffer_Decode{AttributeName}(ScreenUV, IsValid, {AttributeName}); }\n");
+			static const TCHAR* FormatSample = TEXT("void {InstanceFunctionName}(float2 ScreenUV, bool bApplyViewportOffset, out bool IsValid, out {AttributeType} {AttributeName}) { DIGBuffer_Decode{AttributeName}(ScreenUV, bApplyViewportOffset, IsValid, {AttributeName}); }\n");
 			OutHLSL += FString::Format(FormatSample, ArgsSample);
 			return true;
 		}
@@ -215,4 +231,27 @@ bool UNiagaraDataInterfaceGBuffer::GetFunctionHLSL(const FNiagaraDataInterfaceGP
 
 	return false;
 }
+
+bool UNiagaraDataInterfaceGBuffer::UpgradeFunctionCall(FNiagaraFunctionSignature& FunctionSignature)
+{
+	using namespace NiagaraDataInterfaceGBufferLocal;
+
+	bool bWasChanged = false;
+
+	// Early out for version matching
+	if (FunctionSignature.FunctionVersion == EDIFunctionVersion::LatestVersion)
+	{
+		return bWasChanged;
+	}
+
+	// AddedApplyViewportOffset
+	if ( FunctionSignature.FunctionVersion < EDIFunctionVersion::AddedApplyViewportOffset )
+	{
+		FunctionSignature.Inputs.Add_GetRef(FNiagaraVariable(FNiagaraTypeDefinition::GetBoolDef(), TEXT("ApplyViewportOffset"))).SetValue(false);
+		bWasChanged = true;
+	}
+
+	FunctionSignature.FunctionVersion = EDIFunctionVersion::LatestVersion;
+	return bWasChanged;
+}
 #endif
diff --git a/Engine/Plugins/FX/Niagara/Source/Niagara/Private/DataInterface/NiagaraDataInterfaceGBuffer.h b/Engine/Plugins/FX/Niagara/Source/Niagara/Private/DataInterface/NiagaraDataInterfaceGBuffer.h
--- a/Engine/Plugins/FX/Niagara/Source/Niagara/Private/DataInterface/NiagaraDataInterfaceGBuffer.h
+++ b/Engine/Plugins/FX/Niagara/Source/Niagara/Private/DataInterface/NiagaraDataInterfaceGBuffer.h
@@ -26,6 +26,7 @@ public:
 	virtual bool AppendCompileHash(FNiagaraCompileHashVisitor* InVisitor) const override;
 	virtual void GetCommonHLSL(FString& OutHLSL) override;
 	virtual bool GetFunctionHLSL(const FNiagaraDataInterfaceGPUParamInfo& ParamInfo, const FNiagaraDataInterfaceGeneratedFunction& FunctionInfo, int FunctionInstanceIndex, FString& OutHLSL) override;
+	virtual bool UpgradeFunctionCall(FNiagaraFunctionSignature& FunctionSignature);
 #endif
 
 	virtual bool CanExecuteOnTarget(ENiagaraSimTarget Target) const override { return Target == ENiagaraSimTarget::GPUComputeSim; }
