rasterization grid DI that uses atomic operations
built sdf generation and rendering into base flip emitter
updated all examples to use fluid rendering properly

diff --git a/Engine/Plugins/FX/Niagara/Source/Niagara/Classes/NiagaraDataInterfaceRasterizationGrid3D.h b/Engine/Plugins/FX/Niagara/Source/Niagara/Classes/NiagaraDataInterfaceRasterizationGrid3D.h
new file mode 100644
--- /dev/null
+++ b/Engine/Plugins/FX/Niagara/Source/Niagara/Classes/NiagaraDataInterfaceRasterizationGrid3D.h
@@ -0,0 +1,122 @@
+// Copyright Epic Games, Inc. All Rights Reserved.
+#pragma once
+
+#include "NiagaraDataInterface.h"
+#include "NiagaraCommon.h"
+#include "NiagaraDataInterfaceRW.h"
+#include "ClearQuad.h"
+#include "NiagaraStats.h"
+
+#include "NiagaraDataInterfaceRasterizationGrid3D.generated.h"
+
+class FNiagaraSystemInstance;
+
+// store all data in in a class
+// move all data management to use per instance data
+// remove references to push data to render thread
+
+class RasterizationGrid3DRWInstanceData
+{
+public:
+
+	~RasterizationGrid3DRWInstanceData()
+	{
+#if STATS
+		DEC_MEMORY_STAT_BY(STAT_NiagaraGPUDataInterfaceMemory, GPUMemory);
+#endif
+	}
+
+	void ResizeBuffers();
+
+	FIntVector NumCells = FIntVector::ZeroValue;
+	float Precision = 0;
+	int ResetValue = 0;
+
+	bool NeedsRealloc = false;
+	
+	FTextureRWBuffer3D RasterizationBuffer;
+
+#if STATS
+	int32 GPUMemory = 0;
+#endif
+};
+
+struct FNiagaraDataInterfaceProxyRasterizationGrid3D : public FNiagaraDataInterfaceProxyRW
+{	
+	virtual void PreStage(FRHICommandList& RHICmdList, const FNiagaraDataInterfaceStageArgs& Context) override;
+	virtual void ConsumePerInstanceDataFromGameThread(void* PerInstanceData, const FNiagaraSystemInstanceID& Instance) override {}
+	virtual int32 PerInstanceDataPassedToRenderThreadSize() const override { return sizeof(RasterizationGrid3DRWInstanceData); }	
+
+	virtual FIntVector GetElementCount(FNiagaraSystemInstanceID SystemInstanceID) const override;
+
+	/* List of proxy data for each system instances*/
+	// #todo(dmp): this should all be refactored to avoid duplicate code
+	TMap<FNiagaraSystemInstanceID, RasterizationGrid3DRWInstanceData> SystemInstancesToProxyData;
+};
+
+UCLASS(EditInlineNew, Category = "Grid", meta = (DisplayName = "Rasterization Grid3D"))
+class NIAGARA_API UNiagaraDataInterfaceRasterizationGrid3D : public UNiagaraDataInterfaceGrid3D
+{
+	GENERATED_UCLASS_BODY()
+
+public:	
+
+	UPROPERTY(EditAnywhere, Category = "Grid")
+	float Precision;
+
+	UPROPERTY(EditAnywhere, Category = "Grid")
+	int ResetValue;
+
+public:
+
+	DECLARE_NIAGARA_DI_PARAMETER();
+
+	virtual void PostInitProperties() override
+	{
+		Super::PostInitProperties();
+
+		//Can we register data interfaces as regular types and fold them into the FNiagaraVariable framework for UI and function calls etc?
+		if (HasAnyFlags(RF_ClassDefaultObject))
+		{
+			ENiagaraTypeRegistryFlags Flags = ENiagaraTypeRegistryFlags::AllowAnyVariable | ENiagaraTypeRegistryFlags::AllowParameter;
+			FNiagaraTypeRegistry::Register(FNiagaraTypeDefinition(GetClass()), Flags);
+		}
+	}
+
+	//~ UNiagaraDataInterface interface
+	// VM functionality
+	virtual void GetFunctions(TArray<FNiagaraFunctionSignature>& OutFunctions) override;
+	virtual void GetVMExternalFunction(const FVMExternalFunctionBindingInfo& BindingInfo, void* InstanceData, FVMExternalFunction &OutFunc) override;
+
+	virtual bool Equals(const UNiagaraDataInterface* Other) const override;
+
+	// GPU sim functionality
+#if WITH_EDITORONLY_DATA
+	virtual void GetParameterDefinitionHLSL(const FNiagaraDataInterfaceGPUParamInfo& ParamInfo, FString& OutHLSL) override;
+	virtual bool GetFunctionHLSL(const FNiagaraDataInterfaceGPUParamInfo& ParamInfo, const FNiagaraDataInterfaceGeneratedFunction& FunctionInfo, int FunctionInstanceIndex, FString& OutHLSL) override;
+#endif
+
+	virtual void ProvidePerInstanceDataForRenderThread(void* DataForRenderThread, void* PerInstanceData, const FNiagaraSystemInstanceID& SystemInstance) override {}
+	virtual bool InitPerInstanceData(void* PerInstanceData, FNiagaraSystemInstance* SystemInstance) override;
+	virtual void DestroyPerInstanceData(void* PerInstanceData, FNiagaraSystemInstance* SystemInstance) override;
+	virtual bool PerInstanceTick(void* PerInstanceData, FNiagaraSystemInstance* SystemInstance, float DeltaSeconds) override { return false;  }
+	virtual int32 PerInstanceDataSize()const override { return sizeof(RasterizationGrid3DRWInstanceData); }
+	virtual bool PerInstanceTickPostSimulate(void* PerInstanceData, FNiagaraSystemInstance* SystemInstance, float DeltaSeconds) override;
+	virtual bool HasPostSimulateTick() const override { return true; }
+	virtual bool HasPreSimulateTick() const override { return true; }
+	//~ UNiagaraDataInterface interface END
+
+	void GetNumCells(FVectorVMContext& Context);
+	void SetNumCells(FVectorVMContext& Context);
+	void SetFloatResetValue(FVectorVMContext& Context);
+
+	static const FName SetNumCellsFunctionName;
+	static const FName SetFloatResetValueFunctionName;	
+
+protected:
+	//~ UNiagaraDataInterface interface
+	virtual bool CopyToInternal(UNiagaraDataInterface* Destination) const override;
+	//~ UNiagaraDataInterface interface END	
+};
+
+
diff --git a/Engine/Plugins/FX/Niagara/Source/Niagara/Private/NiagaraDataInterfaceRasterizationGrid3D.cpp b/Engine/Plugins/FX/Niagara/Source/Niagara/Private/NiagaraDataInterfaceRasterizationGrid3D.cpp
new file mode 100644
--- /dev/null
+++ b/Engine/Plugins/FX/Niagara/Source/Niagara/Private/NiagaraDataInterfaceRasterizationGrid3D.cpp
@@ -0,0 +1,785 @@
+// Copyright Epic Games, Inc. All Rights Reserved.
+#include "NiagaraDataInterfaceRasterizationGrid3D.h"
+#include "NiagaraShader.h"
+#include "ShaderParameterUtils.h"
+#include "NiagaraEmitterInstanceBatcher.h"
+#include "NiagaraSystemInstance.h"
+#include "NiagaraRenderer.h"
+#include "NiagaraShaderParticleID.h"
+
+#define LOCTEXT_NAMESPACE "NiagaraDataInterfaceRasterizationGrid3D"
+
+static const FString IntGridName(TEXT("IntGrid_"));
+static const FString OutputIntGridName(TEXT("OutputIntGrid_"));
+static const FString PrecisionName(TEXT("Precision_"));
+
+const FName UNiagaraDataInterfaceRasterizationGrid3D::SetNumCellsFunctionName("SetNumCells");
+const FName UNiagaraDataInterfaceRasterizationGrid3D::SetFloatResetValueFunctionName("SetFloatResetValue");
+
+
+
+// Global VM function names, also used by the shaders code generation methods.
+
+static const FName SetFloatValueFunctionName("SetFloatGridValue");
+static const FName GetFloatValueFunctionName("GetFloatGridValue");
+static const FName InterlockedAddFloatGridValueFunctionName("InterlockedAddFloatGridValue");
+static const FName InterlockedMinFloatGridValueFunctionName("InterlockedMinFloatGridValue");
+static const FName InterlockedMaxFloatGridValueFunctionName("InterlockedMaxFloatGridValue");
+
+static const FName IntToFloatFunctionName("IntToFloat");
+static const FName FloatToIntFunctionName("FloatToInt");
+
+
+static const FName SetIntValueFunctionName("SetIntGridValue");
+static const FName GetIntValueFunctionName("GetIntGridValue");
+
+
+static int32 GMaxNiagaraRasterizationGridCells = (1000*1000*1000);
+static FAutoConsoleVariableRef CVarMaxNiagaraRasterizationGridCells(
+	TEXT("fx.MaxNiagaraRasterizationGridCells"),
+	GMaxNiagaraRasterizationGridCells,
+	TEXT("The max number of supported grid cells in Niagara. Overflowing this threshold will cause the sim to warn and fail. \n"),
+	ECVF_Default
+);
+
+
+/*--------------------------------------------------------------------------------------------------------------------------*/
+struct FNiagaraDataInterfaceParametersCS_RasterizationGrid3D : public FNiagaraDataInterfaceParametersCS
+{
+	DECLARE_TYPE_LAYOUT(FNiagaraDataInterfaceParametersCS_RasterizationGrid3D, NonVirtual);
+public:
+	void Bind(const FNiagaraDataInterfaceGPUParamInfo& ParameterInfo, const class FShaderParameterMap& ParameterMap)
+	{
+		NumCellsParam.Bind(ParameterMap, *(UNiagaraDataInterfaceRWBase::NumCellsName + ParameterInfo.DataInterfaceHLSLSymbol));
+		UnitToUVParam.Bind(ParameterMap, *(UNiagaraDataInterfaceRWBase::UnitToUVName + ParameterInfo.DataInterfaceHLSLSymbol));		
+		PrecisionParam.Bind(ParameterMap, *(PrecisionName + ParameterInfo.DataInterfaceHLSLSymbol));
+				
+		IntGridParam.Bind(ParameterMap,  *(IntGridName + ParameterInfo.DataInterfaceHLSLSymbol));
+		
+		OutputIntGridParam.Bind(ParameterMap, *(OutputIntGridName + ParameterInfo.DataInterfaceHLSLSymbol));
+	}
+
+	// #todo(dmp): make resource transitions batched
+	void Set(FRHICommandList& RHICmdList, const FNiagaraDataInterfaceSetArgs& Context) const
+	{
+		check(IsInRenderingThread());
+
+		// Get shader and DI
+		FRHIComputeShader* ComputeShaderRHI = Context.Shader.GetComputeShader();
+		FNiagaraDataInterfaceProxyRasterizationGrid3D* VFDI = static_cast<FNiagaraDataInterfaceProxyRasterizationGrid3D*>(Context.DataInterface);
+
+		RasterizationGrid3DRWInstanceData* ProxyData = VFDI->SystemInstancesToProxyData.Find(Context.SystemInstanceID);	
+
+		if (!(ProxyData && ProxyData->RasterizationBuffer.Buffer.IsValid()))
+		{			
+			SetShaderValue(RHICmdList, ComputeShaderRHI, NumCellsParam, 0);
+			SetShaderValue(RHICmdList, ComputeShaderRHI, UnitToUVParam, FVector3f::ZeroVector);						
+			SetShaderValue(RHICmdList, ComputeShaderRHI, PrecisionParam, 0);
+			SetSRVParameter(RHICmdList, ComputeShaderRHI, IntGridParam, FNiagaraRenderer::GetDummyIntBuffer());
+
+			if (OutputIntGridParam.IsUAVBound())
+			{
+				RHICmdList.SetUAVParameter(ComputeShaderRHI, OutputIntGridParam.GetUAVIndex(), Context.Batcher->GetEmptyUAVFromPool(RHICmdList, PF_R32_SINT, ENiagaraEmptyUAVType::Buffer));
+			}
+
+			return;
+		}
+
+		SetShaderValue(RHICmdList, ComputeShaderRHI, NumCellsParam, ProxyData->NumCells);
+		SetShaderValue(RHICmdList, ComputeShaderRHI, UnitToUVParam, FVector3f(1.0f) / FVector3f(ProxyData->NumCells));
+		SetShaderValue(RHICmdList, ComputeShaderRHI, PrecisionParam, ProxyData->Precision);
+		
+		if (!Context.IsOutputStage)
+		{			
+			if (IntGridParam.IsBound())
+			{
+				RHICmdList.Transition(FRHITransitionInfo(ProxyData->RasterizationBuffer.UAV, ERHIAccess::Unknown, ERHIAccess::SRVCompute));
+				SetSRVParameter(RHICmdList, ComputeShaderRHI, IntGridParam, ProxyData->RasterizationBuffer.SRV);
+			}
+
+			if (OutputIntGridParam.IsUAVBound())
+			{
+				RHICmdList.SetUAVParameter(ComputeShaderRHI, OutputIntGridParam.GetUAVIndex(), Context.Batcher->GetEmptyUAVFromPool(RHICmdList, PF_R32_SINT, ENiagaraEmptyUAVType::Buffer));
+			}
+		}
+		else
+		{			
+			SetSRVParameter(RHICmdList, ComputeShaderRHI, IntGridParam, FNiagaraRenderer::GetDummyIntBuffer());
+
+			if (OutputIntGridParam.IsBound())
+			{
+				RHICmdList.Transition(FRHITransitionInfo(ProxyData->RasterizationBuffer.UAV, ERHIAccess::Unknown, ERHIAccess::UAVCompute));
+
+				//inline void SetTexture(TRHICmdList & RHICmdList, const TShaderRHIRef & Shader, FRHITexture * Texture, FRHIUnorderedAccessView * UAV) const;
+				//OutputIntGridParam.SetBuffer(RHICmdList, ComputeShaderRHI, ProxyData->RasterizationBuffer.Buffer, ProxyData->RasterizationBuffer.UAV);
+				RHICmdList.SetUAVParameter(ComputeShaderRHI, OutputIntGridParam.GetUAVIndex(), ProxyData->RasterizationBuffer.UAV);
+			}
+		}
+		// Note: There is a flush in PreEditChange to make sure everything is synced up at this point 
+	}
+
+	void Unset(FRHICommandList& RHICmdList, const FNiagaraDataInterfaceSetArgs& Context) const 
+	{
+		if (OutputIntGridParam.IsBound())
+		{
+			OutputIntGridParam.UnsetUAV(RHICmdList, Context.Shader.GetComputeShader());
+		}
+	}
+
+private:
+	LAYOUT_FIELD(FShaderParameter, NumCellsParam);
+	LAYOUT_FIELD(FShaderParameter, UnitToUVParam);		
+	LAYOUT_FIELD(FShaderParameter, PrecisionParam);
+	LAYOUT_FIELD(FShaderResourceParameter, IntGridParam);
+	LAYOUT_FIELD(FRWShaderParameter, OutputIntGridParam);	
+};
+
+void RasterizationGrid3DRWInstanceData::ResizeBuffers()
+{
+	uint32 NumTotalCells = NumCells.X * NumCells.Y * NumCells.Z;
+
+	uint32 NumIntsInGridBuffer = NumTotalCells * 1;
+
+	if (NumTotalCells > (uint32)GMaxNiagaraRasterizationGridCells)
+		return;
+
+	RasterizationBuffer.Initialize(TEXT("NiagaraRasterizationGrid3D::IntGrid"), sizeof(int32), NumCells.X, NumCells.Y, NumCells.Z, EPixelFormat::PF_R32_SINT);
+
+	#if STATS
+		DEC_MEMORY_STAT_BY(STAT_NiagaraGPUDataInterfaceMemory, GPUMemory);
+		GPUMemory = NumTotalCells * sizeof(int32) + NumIntsInGridBuffer * sizeof(int32);
+		INC_MEMORY_STAT_BY(STAT_NiagaraGPUDataInterfaceMemory, GPUMemory);
+	#endif
+}
+
+IMPLEMENT_TYPE_LAYOUT(FNiagaraDataInterfaceParametersCS_RasterizationGrid3D);
+
+IMPLEMENT_NIAGARA_DI_PARAMETER(UNiagaraDataInterfaceRasterizationGrid3D, FNiagaraDataInterfaceParametersCS_RasterizationGrid3D);
+
+UNiagaraDataInterfaceRasterizationGrid3D::UNiagaraDataInterfaceRasterizationGrid3D(FObjectInitializer const& ObjectInitializer)
+	: Super(ObjectInitializer)
+	, Precision(1.f)
+	, ResetValue(0)
+{	
+	Proxy.Reset(new FNiagaraDataInterfaceProxyRasterizationGrid3D());	
+}
+
+
+void UNiagaraDataInterfaceRasterizationGrid3D::GetFunctions(TArray<FNiagaraFunctionSignature>& OutFunctions)
+{
+	Super::GetFunctions(OutFunctions);
+
+	{
+		FNiagaraFunctionSignature Sig;
+		Sig.Name = SetNumCellsFunctionName;
+		Sig.Inputs.Add(FNiagaraVariable(FNiagaraTypeDefinition(GetClass()), TEXT("Grid")));
+		Sig.Inputs.Add(FNiagaraVariable(FNiagaraTypeDefinition::GetIntDef(), TEXT("NumCellsX")));
+		Sig.Inputs.Add(FNiagaraVariable(FNiagaraTypeDefinition::GetIntDef(), TEXT("NumCellsY")));
+		Sig.Inputs.Add(FNiagaraVariable(FNiagaraTypeDefinition::GetIntDef(), TEXT("NumCellsZ")));		
+		Sig.Outputs.Add(FNiagaraVariable(FNiagaraTypeDefinition::GetBoolDef(), TEXT("Success")));
+
+		Sig.ModuleUsageBitmask = ENiagaraScriptUsageMask::Emitter | ENiagaraScriptUsageMask::System;
+		Sig.bExperimental = true;
+		Sig.bMemberFunction = true;
+		Sig.bRequiresExecPin = true;
+		Sig.bRequiresContext = false;
+		Sig.bSupportsCPU = true;
+		Sig.bSupportsGPU = false;
+		OutFunctions.Add(Sig);
+	}
+
+	{
+		FNiagaraFunctionSignature Sig;
+		Sig.Name = SetFloatResetValueFunctionName;
+		Sig.Inputs.Add(FNiagaraVariable(FNiagaraTypeDefinition(GetClass()), TEXT("Grid")));
+		Sig.Inputs.Add(FNiagaraVariable(FNiagaraTypeDefinition::GetFloatDef(), TEXT("ResetValue")));		
+		Sig.Outputs.Add(FNiagaraVariable(FNiagaraTypeDefinition::GetBoolDef(), TEXT("Success")));
+
+		Sig.ModuleUsageBitmask = ENiagaraScriptUsageMask::Emitter | ENiagaraScriptUsageMask::System;
+		Sig.bExperimental = true;
+		Sig.bMemberFunction = true;
+		Sig.bRequiresExecPin = true;
+		Sig.bRequiresContext = false;
+		Sig.bSupportsCPU = true;
+		Sig.bSupportsGPU = false;
+		OutFunctions.Add(Sig);
+	}
+
+	{
+		FNiagaraFunctionSignature Sig;
+		Sig.Name = SetFloatValueFunctionName;
+		Sig.Inputs.Add(FNiagaraVariable(FNiagaraTypeDefinition(GetClass()), TEXT("Grid")));
+		Sig.Inputs.Add(FNiagaraVariable(FNiagaraTypeDefinition::GetIntDef(), TEXT("IndexX")));
+		Sig.Inputs.Add(FNiagaraVariable(FNiagaraTypeDefinition::GetIntDef(), TEXT("IndexY")));
+		Sig.Inputs.Add(FNiagaraVariable(FNiagaraTypeDefinition::GetIntDef(), TEXT("IndexZ")));		
+		Sig.Inputs.Add(FNiagaraVariable(FNiagaraTypeDefinition::GetFloatDef(), TEXT("Value")));
+		Sig.Outputs.Add(FNiagaraVariable(FNiagaraTypeDefinition::GetIntDef(), TEXT("IGNORE")));
+
+		Sig.bExperimental = true;
+		Sig.bMemberFunction = true;
+		Sig.bRequiresContext = false;
+		Sig.bWriteFunction = true;
+		Sig.ModuleUsageBitmask = ENiagaraScriptUsageMask::Particle;
+		
+		Sig.bSupportsCPU = false;
+		Sig.bSupportsGPU = true;
+#if WITH_EDITORONLY_DATA
+		Sig.Description = NSLOCTEXT("Niagara", "NiagaraDataInterfaceGridColl2D_SetValueFunction", "Set the value at a specific index. Note that this is an older way of working with Grids. Consider using the SetFloat or other typed, named functions or parameter map variables with StackContext namespace instead.");
+#endif
+		OutFunctions.Add(Sig);
+	}
+
+	{
+		FNiagaraFunctionSignature Sig;
+		Sig.Name = SetIntValueFunctionName;
+		Sig.Inputs.Add(FNiagaraVariable(FNiagaraTypeDefinition(GetClass()), TEXT("Grid")));
+		Sig.Inputs.Add(FNiagaraVariable(FNiagaraTypeDefinition::GetIntDef(), TEXT("IndexX")));
+		Sig.Inputs.Add(FNiagaraVariable(FNiagaraTypeDefinition::GetIntDef(), TEXT("IndexY")));
+		Sig.Inputs.Add(FNiagaraVariable(FNiagaraTypeDefinition::GetIntDef(), TEXT("IndexZ")));
+		Sig.Inputs.Add(FNiagaraVariable(FNiagaraTypeDefinition::GetIntDef(), TEXT("Value")));		
+		Sig.Outputs.Add(FNiagaraVariable(FNiagaraTypeDefinition::GetIntDef(), TEXT("IGNORE")));
+
+		Sig.bExperimental = true;
+		Sig.bMemberFunction = true;
+		Sig.bRequiresContext = false;
+		Sig.bWriteFunction = true;
+		Sig.ModuleUsageBitmask = ENiagaraScriptUsageMask::Particle;
+		
+		Sig.bSupportsCPU = false;
+		Sig.bSupportsGPU = true;
+#if WITH_EDITORONLY_DATA
+		Sig.Description = NSLOCTEXT("Niagara", "NiagaraDataInterfaceGridColl2D_SetValueFunction", "Set the value at a specific index. Note that this is an older way of working with Grids. Consider using the SetFloat or other typed, named functions or parameter map variables with StackContext namespace instead.");
+#endif
+		OutFunctions.Add(Sig);
+	}
+
+	{
+		FNiagaraFunctionSignature Sig;
+		Sig.Name = InterlockedAddFloatGridValueFunctionName;
+		Sig.Inputs.Add(FNiagaraVariable(FNiagaraTypeDefinition(GetClass()), TEXT("Grid")));
+		Sig.Inputs.Add(FNiagaraVariable(FNiagaraTypeDefinition::GetIntDef(), TEXT("IndexX")));
+		Sig.Inputs.Add(FNiagaraVariable(FNiagaraTypeDefinition::GetIntDef(), TEXT("IndexY")));
+		Sig.Inputs.Add(FNiagaraVariable(FNiagaraTypeDefinition::GetIntDef(), TEXT("IndexZ")));
+		Sig.Inputs.Add(FNiagaraVariable(FNiagaraTypeDefinition::GetFloatDef(), TEXT("Value")));
+		Sig.Outputs.Add(FNiagaraVariable(FNiagaraTypeDefinition::GetIntDef(), TEXT("IGNORE")));
+
+		Sig.bExperimental = true;
+		Sig.bMemberFunction = true;
+		Sig.bRequiresContext = false;
+		Sig.bWriteFunction = true;
+		Sig.ModuleUsageBitmask = ENiagaraScriptUsageMask::Particle;
+		
+		Sig.bSupportsCPU = false;
+		Sig.bSupportsGPU = true;
+#if WITH_EDITORONLY_DATA
+		Sig.Description = NSLOCTEXT("Niagara", "NiagaraDataInterfaceGridColl2D_SetValueFunction", "Set the value at a specific index. Note that this is an older way of working with Grids. Consider using the SetFloat or other typed, named functions or parameter map variables with StackContext namespace instead.");
+#endif
+		OutFunctions.Add(Sig);
+	}
+
+	{
+		FNiagaraFunctionSignature Sig;
+		Sig.Name = InterlockedMinFloatGridValueFunctionName;
+		Sig.Inputs.Add(FNiagaraVariable(FNiagaraTypeDefinition(GetClass()), TEXT("Grid")));
+		Sig.Inputs.Add(FNiagaraVariable(FNiagaraTypeDefinition::GetIntDef(), TEXT("IndexX")));
+		Sig.Inputs.Add(FNiagaraVariable(FNiagaraTypeDefinition::GetIntDef(), TEXT("IndexY")));
+		Sig.Inputs.Add(FNiagaraVariable(FNiagaraTypeDefinition::GetIntDef(), TEXT("IndexZ")));
+		Sig.Inputs.Add(FNiagaraVariable(FNiagaraTypeDefinition::GetFloatDef(), TEXT("Value")));
+		Sig.Outputs.Add(FNiagaraVariable(FNiagaraTypeDefinition::GetIntDef(), TEXT("IGNORE")));
+
+		Sig.bExperimental = true;
+		Sig.bMemberFunction = true;
+		Sig.bRequiresContext = false;
+		Sig.bWriteFunction = true;
+		Sig.ModuleUsageBitmask = ENiagaraScriptUsageMask::Particle;
+
+		Sig.bSupportsCPU = false;
+		Sig.bSupportsGPU = true;
+#if WITH_EDITORONLY_DATA
+		Sig.Description = NSLOCTEXT("Niagara", "NiagaraDataInterfaceGridColl2D_SetValueFunction", "Set the value at a specific index. Note that this is an older way of working with Grids. Consider using the SetFloat or other typed, named functions or parameter map variables with StackContext namespace instead.");
+#endif
+		OutFunctions.Add(Sig);
+	}
+
+	{
+		FNiagaraFunctionSignature Sig;
+		Sig.Name = InterlockedMaxFloatGridValueFunctionName;
+		Sig.Inputs.Add(FNiagaraVariable(FNiagaraTypeDefinition(GetClass()), TEXT("Grid")));
+		Sig.Inputs.Add(FNiagaraVariable(FNiagaraTypeDefinition::GetIntDef(), TEXT("IndexX")));
+		Sig.Inputs.Add(FNiagaraVariable(FNiagaraTypeDefinition::GetIntDef(), TEXT("IndexY")));
+		Sig.Inputs.Add(FNiagaraVariable(FNiagaraTypeDefinition::GetIntDef(), TEXT("IndexZ")));
+		Sig.Inputs.Add(FNiagaraVariable(FNiagaraTypeDefinition::GetFloatDef(), TEXT("Value")));
+		Sig.Outputs.Add(FNiagaraVariable(FNiagaraTypeDefinition::GetIntDef(), TEXT("IGNORE")));
+
+		Sig.bExperimental = true;
+		Sig.bMemberFunction = true;
+		Sig.bRequiresContext = false;
+		Sig.bWriteFunction = true;
+		Sig.ModuleUsageBitmask = ENiagaraScriptUsageMask::Particle;
+
+		Sig.bSupportsCPU = false;
+		Sig.bSupportsGPU = true;
+#if WITH_EDITORONLY_DATA
+		Sig.Description = NSLOCTEXT("Niagara", "NiagaraDataInterfaceGridColl2D_SetValueFunction", "Set the value at a specific index. Note that this is an older way of working with Grids. Consider using the SetFloat or other typed, named functions or parameter map variables with StackContext namespace instead.");
+#endif
+		OutFunctions.Add(Sig);
+	}
+
+	{
+		// Older, deprecated form
+		FNiagaraFunctionSignature Sig;
+		Sig.Name = GetFloatValueFunctionName;
+		Sig.Inputs.Add(FNiagaraVariable(FNiagaraTypeDefinition(GetClass()), TEXT("Grid")));
+		Sig.Inputs.Add(FNiagaraVariable(FNiagaraTypeDefinition::GetIntDef(), TEXT("IndexX")));
+		Sig.Inputs.Add(FNiagaraVariable(FNiagaraTypeDefinition::GetIntDef(), TEXT("IndexY")));
+		Sig.Inputs.Add(FNiagaraVariable(FNiagaraTypeDefinition::GetIntDef(), TEXT("IndexZ")));		
+		Sig.Outputs.Add(FNiagaraVariable(FNiagaraTypeDefinition::GetFloatDef(), TEXT("Value")));
+
+		Sig.bExperimental = true;
+		Sig.bMemberFunction = true;
+		Sig.bRequiresContext = false;
+		Sig.ModuleUsageBitmask = ENiagaraScriptUsageMask::Particle;
+		
+		Sig.bSupportsCPU = false;
+		Sig.bSupportsGPU = true;
+#if WITH_EDITORONLY_DATA
+		Sig.Description = NSLOCTEXT("Niagara", "NiagaraDataInterfaceGridColl3D_GetValueFunction", "Get the value at a specific index. Note that this is an older way of working with Grids. Consider using the SetFloat or other typed, named functions or parameter map variables with StackContext namespace instead.");
+#endif
+
+		OutFunctions.Add(Sig);
+	}
+
+	{
+		// Older, deprecated form
+		FNiagaraFunctionSignature Sig;
+		Sig.Name = GetIntValueFunctionName;
+		Sig.Inputs.Add(FNiagaraVariable(FNiagaraTypeDefinition(GetClass()), TEXT("Grid")));
+		Sig.Inputs.Add(FNiagaraVariable(FNiagaraTypeDefinition::GetIntDef(), TEXT("IndexX")));
+		Sig.Inputs.Add(FNiagaraVariable(FNiagaraTypeDefinition::GetIntDef(), TEXT("IndexY")));
+		Sig.Inputs.Add(FNiagaraVariable(FNiagaraTypeDefinition::GetIntDef(), TEXT("IndexZ")));		
+		Sig.Outputs.Add(FNiagaraVariable(FNiagaraTypeDefinition::GetIntDef(), TEXT("Value")));
+
+		Sig.bExperimental = true;
+		Sig.bMemberFunction = true;
+		Sig.bRequiresContext = false;
+		Sig.ModuleUsageBitmask = ENiagaraScriptUsageMask::Particle;
+		
+		Sig.bSupportsCPU = false;
+		Sig.bSupportsGPU = true;
+#if WITH_EDITORONLY_DATA
+		Sig.Description = NSLOCTEXT("Niagara", "NiagaraDataInterfaceGridColl3D_GetValueFunction", "Get the value at a specific index. Note that this is an older way of working with Grids. Consider using the SetFloat or other typed, named functions or parameter map variables with StackContext namespace instead.");
+#endif
+
+		OutFunctions.Add(Sig);
+	}
+}
+
+DEFINE_NDI_DIRECT_FUNC_BINDER(UNiagaraDataInterfaceRasterizationGrid3D, SetNumCells);
+DEFINE_NDI_DIRECT_FUNC_BINDER(UNiagaraDataInterfaceRasterizationGrid3D, SetFloatResetValue);
+void UNiagaraDataInterfaceRasterizationGrid3D::GetVMExternalFunction(const FVMExternalFunctionBindingInfo& BindingInfo, void* InstanceData, FVMExternalFunction &OutFunc)
+{
+	Super::GetVMExternalFunction(BindingInfo, InstanceData, OutFunc);
+
+	// #todo(dmp): this overrides the empty function set by the super class
+	if (BindingInfo.Name == UNiagaraDataInterfaceRWBase::NumCellsFunctionName)
+	{
+		check(BindingInfo.GetNumInputs() == 1 && BindingInfo.GetNumOutputs() == 3);
+		OutFunc = FVMExternalFunction::CreateLambda([&](FVectorVMContext& Context) { GetNumCells(Context); });
+	}
+	else if (BindingInfo.Name == SetNumCellsFunctionName)
+	{
+		check(BindingInfo.GetNumInputs() == 4 && BindingInfo.GetNumOutputs() == 1);
+		NDI_FUNC_BINDER(UNiagaraDataInterfaceRasterizationGrid3D, SetNumCells)::Bind(this, OutFunc);
+	}
+	else if (BindingInfo.Name == SetFloatResetValueFunctionName)
+	{
+		check(BindingInfo.GetNumInputs() == 2 && BindingInfo.GetNumOutputs() == 1);
+		NDI_FUNC_BINDER(UNiagaraDataInterfaceRasterizationGrid3D, SetFloatResetValue)::Bind(this, OutFunc);
+	}
+}
+
+
+void UNiagaraDataInterfaceRasterizationGrid3D::GetNumCells(FVectorVMContext& Context)
+{
+	VectorVM::FUserPtrHandler<RasterizationGrid3DRWInstanceData> InstData(Context);
+
+	FNDIOutputParam<int32> NumCellsX(Context);
+	FNDIOutputParam<int32> NumCellsY(Context);
+	FNDIOutputParam<int32> NumCellsZ(Context);
+
+	for (int32 InstanceIdx = 0; InstanceIdx < Context.NumInstances; ++InstanceIdx)
+	{
+		NumCellsX.SetAndAdvance(NumCells.X);
+		NumCellsY.SetAndAdvance(NumCells.Y);
+		NumCellsZ.SetAndAdvance(NumCells.Z);
+	}
+}
+
+
+bool UNiagaraDataInterfaceRasterizationGrid3D::Equals(const UNiagaraDataInterface* Other) const
+{
+	if (!Super::Equals(Other))
+	{
+		return false;
+	}
+	const UNiagaraDataInterfaceRasterizationGrid3D* OtherTyped = CastChecked<const UNiagaraDataInterfaceRasterizationGrid3D>(Other);
+
+	return OtherTyped->Precision == Precision && OtherTyped->ResetValue == ResetValue;
+}
+
+#if WITH_EDITORONLY_DATA
+void UNiagaraDataInterfaceRasterizationGrid3D::GetParameterDefinitionHLSL(const FNiagaraDataInterfaceGPUParamInfo& ParamInfo, FString& OutHLSL)
+{
+	Super::GetParameterDefinitionHLSL(ParamInfo, OutHLSL);
+
+	static const TCHAR *FormatDeclarations = TEXT(R"(			
+		Texture3D<int> {IntGridName};		
+		RWTexture3D<int> RW{OutputIntGridName};
+		float {Precision};
+	)");
+	TMap<FString, FStringFormatArg> ArgsDeclarations = {				
+		{ TEXT("IntGridName"),    IntGridName + ParamInfo.DataInterfaceHLSLSymbol },		
+		{ TEXT("OutputIntGridName"),    OutputIntGridName + ParamInfo.DataInterfaceHLSLSymbol },
+		{ TEXT("Precision"), PrecisionName + ParamInfo.DataInterfaceHLSLSymbol},
+	};
+	OutHLSL += FString::Format(FormatDeclarations, ArgsDeclarations);
+
+	// always generate the code for these functions that are used internally by other DI functions
+	// to quantize values
+	{
+		static const TCHAR* FormatHLSL = TEXT(R"(
+				float {IntToFloatFunction}(int IntValue)
+				{
+					return 1. * IntValue / {Precision};
+				}
+			)");
+		TMap<FString, FStringFormatArg> FormatArgs =
+		{
+			{ TEXT("IntToFloatFunction"), IntToFloatFunctionName.ToString() + ParamInfo.DataInterfaceHLSLSymbol},
+			{ TEXT("FloatToIntFunction"), FloatToIntFunctionName.ToString() + ParamInfo.DataInterfaceHLSLSymbol},
+			{ TEXT("Precision"), PrecisionName + ParamInfo.DataInterfaceHLSLSymbol},
+		};
+		OutHLSL += FString::Format(FormatHLSL, FormatArgs);
+	}
+	
+	{
+		static const TCHAR* FormatHLSL = TEXT(R"(
+				int {FloatToIntFunction}(float FloatValue)
+				{
+					return FloatValue * {Precision};
+				}
+			)");
+
+		TMap<FString, FStringFormatArg> FormatArgs =
+		{
+			{ TEXT("IntToFloatFunction"), IntToFloatFunctionName.ToString() + ParamInfo.DataInterfaceHLSLSymbol},
+			{ TEXT("FloatToIntFunction"), FloatToIntFunctionName.ToString() + ParamInfo.DataInterfaceHLSLSymbol},
+			{ TEXT("Precision"), PrecisionName + ParamInfo.DataInterfaceHLSLSymbol},
+		};
+		OutHLSL += FString::Format(FormatHLSL, FormatArgs);
+	}
+
+}
+
+bool UNiagaraDataInterfaceRasterizationGrid3D::GetFunctionHLSL(const FNiagaraDataInterfaceGPUParamInfo& ParamInfo, const FNiagaraDataInterfaceGeneratedFunction& FunctionInfo, int FunctionInstanceIndex, FString& OutHLSL)
+{
+	bool ParentRet = Super::GetFunctionHLSL(ParamInfo, FunctionInfo, FunctionInstanceIndex, OutHLSL);
+
+	TMap<FString, FStringFormatArg> ArgsBounds =
+	{
+		{ TEXT("FunctionName"), FunctionInfo.InstanceName},
+		{ TEXT("IntGrid"), IntGridName + ParamInfo.DataInterfaceHLSLSymbol},
+		{ TEXT("OutputIntGrid"), OutputIntGridName + ParamInfo.DataInterfaceHLSLSymbol},		
+		{ TEXT("NumCellsName"), UNiagaraDataInterfaceRWBase::NumCellsName + ParamInfo.DataInterfaceHLSLSymbol},
+		{ TEXT("UnitToUVName"), UNiagaraDataInterfaceRWBase::UnitToUVName + ParamInfo.DataInterfaceHLSLSymbol},		
+		{ TEXT("IntToFloatFunctionName"), IntToFloatFunctionName.ToString() + ParamInfo.DataInterfaceHLSLSymbol},
+		{ TEXT("FloatToIntFunctionName"), FloatToIntFunctionName.ToString() + ParamInfo.DataInterfaceHLSLSymbol},
+		{ TEXT("Precision"), PrecisionName + ParamInfo.DataInterfaceHLSLSymbol},
+	};
+
+	if (ParentRet)
+	{
+		return true;
+	}	
+	else if (FunctionInfo.DefinitionName == UNiagaraDataInterfaceRWBase::NumCellsFunctionName)
+	{
+		static const TCHAR* FormatHLSL = TEXT(R"(
+			void {FunctionName}(out int OutNumCellsX, out int OutNumCellsY, out int OutNumCellsZ)
+			{
+				OutNumCellsX = {NumCellsName}.x;
+				OutNumCellsY = {NumCellsName}.y;
+				OutNumCellsZ = {NumCellsName}.z;
+			}
+		)");
+
+		OutHLSL += FString::Format(FormatHLSL, ArgsBounds);
+	}
+	else if (FunctionInfo.DefinitionName == SetFloatValueFunctionName)
+	{
+		static const TCHAR* FormatBounds = TEXT(R"(
+			void {FunctionName}(int In_IndexX, int In_IndexY, int In_IndexZ, float In_Value, out int val)
+			{			
+				val = 0;				
+				RW{OutputIntGrid}[int3(In_IndexX, In_IndexY, In_IndexZ)] = {FloatToIntFunctionName}(In_Value);
+			}
+		)");
+
+		OutHLSL += FString::Format(FormatBounds, ArgsBounds);
+		return true;
+	}
+	else if (FunctionInfo.DefinitionName == SetIntValueFunctionName)
+	{
+		static const TCHAR* FormatBounds = TEXT(R"(
+			void {FunctionName}(int In_IndexX, int In_IndexY, int In_IndexZ, int In_Value, out int val)
+			{			
+				val = 0;			
+				RW{OutputIntGrid}[int3(In_IndexX, In_IndexY, In_IndexZ)] = In_Value;				
+			}
+		)");
+		OutHLSL += FString::Format(FormatBounds, ArgsBounds);
+		return true;
+	}
+	else if (FunctionInfo.DefinitionName == InterlockedAddFloatGridValueFunctionName)
+	{
+		static const TCHAR* FormatBounds = TEXT(R"(
+			void {FunctionName}(int In_IndexX, int In_IndexY, int In_IndexZ, float In_Value, out int val)
+			{							
+				val = 0;			
+				InterlockedAdd(RW{OutputIntGrid}[int3(In_IndexX, In_IndexY, In_IndexZ)], {FloatToIntFunctionName}(In_Value));
+			}
+		)");
+		OutHLSL += FString::Format(FormatBounds, ArgsBounds);
+		return true;
+	}
+	else if (FunctionInfo.DefinitionName == InterlockedMinFloatGridValueFunctionName)
+	{
+		static const TCHAR* FormatBounds = TEXT(R"(
+			void {FunctionName}(int In_IndexX, int In_IndexY, int In_IndexZ, float In_Value, out int val)
+			{							
+				val = 0;			
+				InterlockedMin(RW{OutputIntGrid}[int3(In_IndexX, In_IndexY, In_IndexZ)], {FloatToIntFunctionName}(In_Value));
+			}
+		)");
+		OutHLSL += FString::Format(FormatBounds, ArgsBounds);
+		return true;
+	}
+	else if (FunctionInfo.DefinitionName == InterlockedMaxFloatGridValueFunctionName)
+	{
+		static const TCHAR* FormatBounds = TEXT(R"(
+			void {FunctionName}(int In_IndexX, int In_IndexY, int In_IndexZ, float In_Value, out int val)
+			{							
+				val = 0;			
+				InterlockedMax(RW{OutputIntGrid}[int3(In_IndexX, In_IndexY, In_IndexZ)], {FloatToIntFunctionName}(In_Value));
+			}
+		)");
+		OutHLSL += FString::Format(FormatBounds, ArgsBounds);
+		return true;
+	}
+	else if (FunctionInfo.DefinitionName == GetFloatValueFunctionName)
+	{
+		static const TCHAR* FormatBounds = TEXT(R"(
+			void {FunctionName}(int In_IndexX, int In_IndexY, int In_IndexZ, out float Out_Val)
+			{				
+				Out_Val =  {IntToFloatFunctionName}({IntGrid}.Load(int4(In_IndexX, In_IndexY, In_IndexZ, 0)));
+			}
+		)");
+		OutHLSL += FString::Format(FormatBounds, ArgsBounds);
+		return true;
+	}
+	else if (FunctionInfo.DefinitionName == GetIntValueFunctionName)
+	{
+		static const TCHAR* FormatBounds = TEXT(R"(
+			void {FunctionName}(int In_IndexX, int In_IndexY, int In_IndexZ, out int Out_Val)
+			{
+				Out_Val = Out_Val = {IntGrid}.Load(int4(In_IndexX, In_IndexY, In_IndexZ, 0));				
+			}
+		)");
+		OutHLSL += FString::Format(FormatBounds, ArgsBounds);
+		return true;
+	}
+	return false;
+}
+#endif
+
+bool UNiagaraDataInterfaceRasterizationGrid3D::InitPerInstanceData(void* PerInstanceData, FNiagaraSystemInstance* SystemInstance)
+{
+
+	RasterizationGrid3DRWInstanceData* InstanceData = new (PerInstanceData) RasterizationGrid3DRWInstanceData();
+
+	FNiagaraDataInterfaceProxyRasterizationGrid3D* RT_Proxy = GetProxyAs<FNiagaraDataInterfaceProxyRasterizationGrid3D>();
+
+	FIntVector RT_NumCells = NumCells;		
+
+	float RT_Precision = Precision;
+	int RT_ResetValue = ResetValue;
+
+	RT_NumCells.X = FMath::Max(RT_NumCells.X, 1);
+	RT_NumCells.Y = FMath::Max(RT_NumCells.Y, 1);
+	RT_NumCells.Z = FMath::Max(RT_NumCells.Z, 1);
+		
+	InstanceData->NumCells = RT_NumCells;
+	InstanceData->Precision = RT_Precision;
+	InstanceData->ResetValue = RT_ResetValue;
+
+	if ((RT_NumCells.X * RT_NumCells.Y * RT_NumCells.Z) > GMaxNiagaraRasterizationGridCells)
+	{
+		UE_LOG(LogNiagara, Warning, TEXT("Dimensions are too big! Please adjust! %d x %d x %d > %d for ==> %s"), RT_NumCells.X, RT_NumCells.Y, RT_NumCells.Z, GMaxNiagaraRasterizationGridCells , *GetFullNameSafe(this))
+			return false;
+	}
+
+	// @todo-threadsafety. This would be a race but I'm taking a ref here. Not ideal in the long term.
+	// Push Updates to Proxy.
+	ENQUEUE_RENDER_COMMAND(FUpdateData)(
+		[RT_Proxy, RT_NumCells, RT_Precision, RT_ResetValue, InstanceID = SystemInstance->GetId()](FRHICommandListImmediate& RHICmdList)
+	{
+		check(!RT_Proxy->SystemInstancesToProxyData.Contains(InstanceID));
+		RasterizationGrid3DRWInstanceData* TargetData = &RT_Proxy->SystemInstancesToProxyData.Add(InstanceID);
+
+		TargetData->NumCells = RT_NumCells;				
+		TargetData->Precision = RT_Precision;
+		TargetData->ResetValue = RT_ResetValue;
+
+		TargetData->ResizeBuffers();
+	});
+
+	return true;
+}
+
+void UNiagaraDataInterfaceRasterizationGrid3D::SetNumCells(FVectorVMContext& Context)
+{
+	// This should only be called from a system or emitter script due to a need for only setting up initially.
+	VectorVM::FUserPtrHandler<RasterizationGrid3DRWInstanceData> InstData(Context);
+	VectorVM::FExternalFuncInputHandler<int> InNumCellsX(Context);
+	VectorVM::FExternalFuncInputHandler<int> InNumCellsY(Context);
+	VectorVM::FExternalFuncInputHandler<int> InNumCellsZ(Context);	
+	VectorVM::FExternalFuncRegisterHandler<FNiagaraBool> OutSuccess(Context);
+
+	for (int32 InstanceIdx = 0; InstanceIdx < Context.NumInstances; ++InstanceIdx)
+	{
+		int NewNumCellsX = InNumCellsX.GetAndAdvance();
+		int NewNumCellsY = InNumCellsY.GetAndAdvance();
+		int NewNumCellsZ = InNumCellsZ.GetAndAdvance();		
+		bool bSuccess = (InstData.Get() != nullptr && Context.NumInstances == 1 && NumCells.X >= 0 && NumCells.Y >= 0 && NumCells.Z >= 0);
+		*OutSuccess.GetDestAndAdvance() = bSuccess;
+		if (bSuccess)
+		{
+			FIntVector OldNumCells = InstData->NumCells;			
+
+			InstData->NumCells.X = NewNumCellsX;
+			InstData->NumCells.Y = NewNumCellsY;
+			InstData->NumCells.Z = NewNumCellsZ;			
+		
+			InstData->NeedsRealloc = OldNumCells != InstData->NumCells;
+		}
+	}
+}
+
+void UNiagaraDataInterfaceRasterizationGrid3D::SetFloatResetValue(FVectorVMContext& Context)
+{
+	VectorVM::FUserPtrHandler<RasterizationGrid3DRWInstanceData> InstData(Context);
+	VectorVM::FExternalFuncInputHandler<float> InResetValue(Context);
+	VectorVM::FExternalFuncRegisterHandler<FNiagaraBool> OutSuccess(Context);
+
+	for (int32 InstanceIdx = 0; InstanceIdx < Context.NumInstances; ++InstanceIdx)
+	{
+		float NewResetValue = InResetValue.GetAndAdvance();
+		bool bSuccess = InstData.Get() != nullptr && Context.NumInstances == 1;
+		*OutSuccess.GetDestAndAdvance() = bSuccess;
+		if (bSuccess)
+		{
+			// Quantize value
+			InstData->ResetValue = NewResetValue * InstData->Precision;
+		}
+	}
+}
+
+bool UNiagaraDataInterfaceRasterizationGrid3D::PerInstanceTickPostSimulate(void* PerInstanceData, FNiagaraSystemInstance* SystemInstance, float DeltaSeconds)
+{
+	RasterizationGrid3DRWInstanceData* InstanceData = static_cast<RasterizationGrid3DRWInstanceData*>(PerInstanceData);
+	bool bNeedsReset = false;
+
+	if (InstanceData->NeedsRealloc && InstanceData->NumCells.X > 0 && InstanceData->NumCells.Y > 0 && InstanceData->NumCells.Z > 0)
+	{
+		InstanceData->NeedsRealloc = false;
+		
+		FNiagaraDataInterfaceProxyRasterizationGrid3D* RT_Proxy = GetProxyAs<FNiagaraDataInterfaceProxyRasterizationGrid3D>();
+		ENQUEUE_RENDER_COMMAND(FUpdateData)(
+			[RT_Proxy, RT_NumCells = InstanceData->NumCells, RT_Precision = InstanceData->Precision, RT_ResetValue = InstanceData->ResetValue, InstanceID = SystemInstance->GetId()](FRHICommandListImmediate& RHICmdList)
+		{
+			check(RT_Proxy->SystemInstancesToProxyData.Contains(InstanceID));
+			RasterizationGrid3DRWInstanceData* TargetData = &RT_Proxy->SystemInstancesToProxyData.Add(InstanceID);
+
+			TargetData->NumCells = RT_NumCells;					
+			TargetData->Precision = RT_Precision;
+			TargetData->ResetValue = RT_ResetValue;
+			TargetData->ResizeBuffers();
+		});
+	}
+
+	return false;
+}
+
+
+void UNiagaraDataInterfaceRasterizationGrid3D::DestroyPerInstanceData(void* PerInstanceData, FNiagaraSystemInstance* SystemInstance)
+{		
+	RasterizationGrid3DRWInstanceData* InstanceData = static_cast<RasterizationGrid3DRWInstanceData*>(PerInstanceData);
+	InstanceData->~RasterizationGrid3DRWInstanceData();
+
+	FNiagaraDataInterfaceProxyRasterizationGrid3D* ThisProxy = GetProxyAs<FNiagaraDataInterfaceProxyRasterizationGrid3D>();
+	if (!ThisProxy)
+		return;
+
+	ENQUEUE_RENDER_COMMAND(FNiagaraDIDestroyInstanceData) (
+		[ThisProxy, InstanceID = SystemInstance->GetId(), Batcher = SystemInstance->GetBatcher()](FRHICommandListImmediate& CmdList)
+	{
+		//check(ThisProxy->SystemInstancesToProxyData.Contains(InstanceID));
+		ThisProxy->SystemInstancesToProxyData.Remove(InstanceID);
+	}
+	);
+}
+
+void FNiagaraDataInterfaceProxyRasterizationGrid3D::PreStage(FRHICommandList& RHICmdList, const FNiagaraDataInterfaceStageArgs& Context)
+{
+	if (Context.IsOutputStage)
+	{
+		RasterizationGrid3DRWInstanceData* ProxyData = SystemInstancesToProxyData.Find(Context.SystemInstanceID);
+
+		SCOPED_DRAW_EVENT(RHICmdList, NiagaraRasterizationGrid3DClearNeighborInfo);
+		ERHIFeatureLevel::Type FeatureLevel = Context.Batcher->GetFeatureLevel();
+
+		FRHITransitionInfo TransitionInfos[] =
+		{			
+			FRHITransitionInfo(ProxyData->RasterizationBuffer.UAV, ERHIAccess::Unknown, ERHIAccess::UAVCompute),
+		};
+		RHICmdList.Transition(MakeArrayView(TransitionInfos, UE_ARRAY_COUNT(TransitionInfos)));		
+					
+		RHICmdList.ClearUAVUint(ProxyData->RasterizationBuffer.UAV, FUintVector4(ProxyData->ResetValue, ProxyData->ResetValue, ProxyData->ResetValue, ProxyData->ResetValue));		
+	}
+}
+
+FIntVector FNiagaraDataInterfaceProxyRasterizationGrid3D::GetElementCount(FNiagaraSystemInstanceID SystemInstanceID) const
+{
+	if ( const RasterizationGrid3DRWInstanceData* TargetData = SystemInstancesToProxyData.Find(SystemInstanceID) )
+	{
+		return TargetData->NumCells;
+	}
+	return FIntVector::ZeroValue;
+}
+
+bool UNiagaraDataInterfaceRasterizationGrid3D::CopyToInternal(UNiagaraDataInterface* Destination) const
+{
+	if (!Super::CopyToInternal(Destination))
+	{
+		return false;
+	}
+
+	UNiagaraDataInterfaceRasterizationGrid3D* OtherTyped = CastChecked<UNiagaraDataInterfaceRasterizationGrid3D>(Destination);
+
+	
+	OtherTyped->Precision = Precision;
+	OtherTyped->ResetValue = ResetValue;
+
+	return true;
+}
+
+#undef LOCTEXT_NAMESPACE
