- Add buffer size & vp offset to camera DI

diff --git a/Engine/Plugins/FX/Niagara/Source/Niagara/Private/NiagaraDataInterfaceCamera.cpp b/Engine/Plugins/FX/Niagara/Source/Niagara/Private/NiagaraDataInterfaceCamera.cpp
--- a/Engine/Plugins/FX/Niagara/Source/Niagara/Private/NiagaraDataInterfaceCamera.cpp
+++ b/Engine/Plugins/FX/Niagara/Source/Niagara/Private/NiagaraDataInterfaceCamera.cpp
@@ -20,6 +20,7 @@ struct FNiagaraCameraDIFunctionVersion
 	{
 		InitialVersion = 0,
         AddedPreViewTranslation = 1,
+		AddedBufferSizeInvSizeForGPU = 2,
 
         VersionPlusOne,
         LatestVersion = VersionPlusOne - 1
@@ -132,6 +133,8 @@ void UNiagaraDataInterfaceCamera::GetFunctions(TArray<FNiagaraFunctionSignature>
 	Sig.AddOutput(FNiagaraVariable(FNiagaraTypeDefinition::GetVec2Def(), TEXT("Temporal AA Jitter (Current Frame)")), LOCTEXT("TemporalAACurrentDescription", "Returns the movement of the current frame view due to AA jittering"));
 	Sig.AddOutput(FNiagaraVariable(FNiagaraTypeDefinition::GetVec2Def(), TEXT("Temporal AA Jitter (Previous Frame)")), LOCTEXT("TemporalAAPreviousDescription", "Returns the movement of the previous frame view due to AA jittering"));
 	Sig.AddOutput(FNiagaraVariable(FNiagaraTypeDefinition::GetVec3Def(), TEXT("PreViewTranslation")), LOCTEXT("PreViewTranslationDescription", "Returns the translation to apply for the various 'Translated XX to XX' transforms."));
+	Sig.AddOutput(FNiagaraVariable(FNiagaraTypeDefinition::GetVec4Def(), TEXT("Buffer Size And Inverse Size")), LOCTEXT("BufferSizeAndInverseSizeDescription", "Returns a vec4 - the x and y values are width and height of the buffer; the z and w values are the inverse width and height."));
+	Sig.AddOutput(FNiagaraVariable(FNiagaraTypeDefinition::GetVec2Def(), TEXT("Viewport Offset")), LOCTEXT("ViewportOffsetDescription", "Returns the viewport offset in the buffer."));
 	OutFunctions.Add(Sig);
 
 
@@ -255,7 +258,7 @@ bool UNiagaraDataInterfaceCamera::GetFunctionHLSL(const FNiagaraDataInterfaceGPU
 	if (FunctionInfo.DefinitionName == GetViewPropertiesName)
 	{
 		static const TCHAR *FormatSample = TEXT(R"(
-			void {FunctionName}(out float3 Out_ViewPositionWorld, out float3 Out_ViewForwardVector, out float3 Out_ViewUpVector, out float3 Out_ViewRightVector, out float4 Out_ViewSizeAndInverseSize, out float4 Out_ScreenToViewSpace, out float2 Out_Current_TAAJitter, out float2 Out_Previous_TAAJitter, out float3 Out_PreViewTranslation)
+			void {FunctionName}(out float3 Out_ViewPositionWorld, out float3 Out_ViewForwardVector, out float3 Out_ViewUpVector, out float3 Out_ViewRightVector, out float4 Out_ViewSizeAndInverseSize, out float4 Out_ScreenToViewSpace, out float2 Out_Current_TAAJitter, out float2 Out_Previous_TAAJitter, out float3 Out_PreViewTranslation, out float4 Out_BufferSizeAndInverseSize, out float2 Out_ViewportOffset)
 			{
 				Out_ViewPositionWorld.xyz = View.WorldViewOrigin.xyz;
 				Out_ViewForwardVector.xyz = View.ViewForward.xyz;
@@ -266,6 +269,8 @@ bool UNiagaraDataInterfaceCamera::GetFunctionHLSL(const FNiagaraDataInterfaceGPU
 				Out_Current_TAAJitter = View.TemporalAAJitter.xy;
 				Out_Previous_TAAJitter = View.TemporalAAJitter.zw;
 				Out_PreViewTranslation = View.PreViewTranslation;
+				Out_BufferSizeAndInverseSize = View.BufferSizeAndInvSize;
+				Out_ViewportOffset = View.ViewRectMin.xy;
 			} 
 		)");
 		OutHLSL += FString::Format(FormatSample, ArgsSample);
@@ -364,6 +369,16 @@ bool UNiagaraDataInterfaceCamera::UpgradeFunctionCall(FNiagaraFunctionSignature&
 		bWasChanged = true;
 	}
 	
+	// Added a few new outputs for GPU properties
+	if (FunctionSignature.FunctionVersion < FNiagaraCameraDIFunctionVersion::AddedBufferSizeInvSizeForGPU && FunctionSignature.Name == GetViewPropertiesName)
+	{
+		ensure(FunctionSignature.Outputs.Num() == 9);
+		FunctionSignature.AddOutput(FNiagaraVariable(FNiagaraTypeDefinition::GetVec4Def(), TEXT("Buffer Size And Inverse Size")), LOCTEXT("BufferSizeAndInverseSizeDescription", "Returns a vec4 - the x and y values are width and height of the buffer; the z and w values are the inverse width and height."));
+		FunctionSignature.AddOutput(FNiagaraVariable(FNiagaraTypeDefinition::GetVec2Def(), TEXT("Viewport Offset")), LOCTEXT("ViewportOffsetDescription", "Returns the viewport offset in the buffer."));
+
+		bWasChanged = true;
+	}
+
 	// Set latest version
     FunctionSignature.FunctionVersion = FNiagaraCameraDIFunctionVersion::LatestVersion;
 
