- Fix various missing transitions and cleanup some where we can specify initial state

#rnx

diff --git a/Engine/Plugins/FX/Niagara/Source/Niagara/Private/NiagaraDataInterfaceCollisionQuery.cpp b/Engine/Plugins/FX/Niagara/Source/Niagara/Private/NiagaraDataInterfaceCollisionQuery.cpp
--- a/Engine/Plugins/FX/Niagara/Source/Niagara/Private/NiagaraDataInterfaceCollisionQuery.cpp
+++ b/Engine/Plugins/FX/Niagara/Source/Niagara/Private/NiagaraDataInterfaceCollisionQuery.cpp
@@ -89,11 +89,28 @@ struct FNiagaraDataIntefaceProxyCollisionQuery : public FNiagaraDataInterfacePro
 		FNiagaraDataInterfaceProxy::PreStage(RHICmdList, Context);
 
 #if RHI_RAYTRACING
-		if (MaxRayTraceCount > 0)
+		// First stage needs to clear the indirect args
+		if (MaxRayTraceCount > 0 && Context.SimStageData->bFirstStage)
 		{
-			// clear out indirect args
 			RHICmdList.ClearUAVUint(RayTraceCounts.UAV, FUintVector4(0, 0, 0, 0));
-			RHICmdList.Transition(FRHITransitionInfo(RayTraceCounts.UAV, ERHIAccess::Unknown, ERHIAccess::UAVCompute));
+			RHICmdList.Transition(FRHITransitionInfo(RayTraceCounts.UAV, ERHIAccess::UAVCompute, ERHIAccess::UAVCompute));
+		}
+#endif
+	}
+
+	virtual void PostStage(FRHICommandList& RHICmdList, const FNiagaraDataInterfaceStageArgs& Context) override
+	{
+#if RHI_RAYTRACING
+		// If we are not the last stage we need to transition because another stage may write
+		// If we do not have a ray tracing scene active we also need to transition as PostSimulate will not transition the buffers
+		if (MaxRayTraceCount > 0 && (!Context.SimStageData->bLastStage || !Context.Batcher->HasRayTracingScene()))
+		{
+			FRHITransitionInfo Transitions[] =
+			{
+				FRHITransitionInfo(RayTraceRequests.UAV, ERHIAccess::UAVCompute, ERHIAccess::UAVCompute),
+				FRHITransitionInfo(RayTraceCounts.UAV, ERHIAccess::UAVCompute, ERHIAccess::UAVCompute),
+			};
+			RHICmdList.Transition(Transitions);
 		}
 #endif
 	}
@@ -109,7 +126,8 @@ struct FNiagaraDataIntefaceProxyCollisionQuery : public FNiagaraDataInterfacePro
 				FRHITransitionInfo PreTransitions[] =
 				{
 					FRHITransitionInfo(RayTraceRequests.UAV, ERHIAccess::UAVCompute, ERHIAccess::SRVCompute),
-					FRHITransitionInfo(RayTraceIntersections.UAV, ERHIAccess::SRVCompute, ERHIAccess::UAVCompute)
+					FRHITransitionInfo(RayTraceIntersections.UAV, ERHIAccess::SRVCompute, ERHIAccess::UAVCompute),
+					FRHITransitionInfo(RayTraceCounts.UAV, ERHIAccess::UAVCompute, ERHIAccess::IndirectArgs),
 				};
 				RHICmdList.Transition(PreTransitions);
 			}
@@ -120,7 +138,8 @@ struct FNiagaraDataIntefaceProxyCollisionQuery : public FNiagaraDataInterfacePro
 				FRHITransitionInfo PostTransitions[] =
 				{
 					FRHITransitionInfo(RayTraceRequests.UAV, ERHIAccess::SRVCompute, ERHIAccess::UAVCompute),
-					FRHITransitionInfo(RayTraceIntersections.UAV, ERHIAccess::UAVCompute, ERHIAccess::SRVCompute)
+					FRHITransitionInfo(RayTraceIntersections.UAV, ERHIAccess::UAVCompute, ERHIAccess::SRVCompute),
+					FRHITransitionInfo(RayTraceCounts.UAV, ERHIAccess::IndirectArgs, ERHIAccess::UAVCompute),
 				};
 				RHICmdList.Transition(PostTransitions);
 			}
diff --git a/Engine/Plugins/FX/Niagara/Source/Niagara/Private/NiagaraDataSet.cpp b/Engine/Plugins/FX/Niagara/Source/Niagara/Private/NiagaraDataSet.cpp
--- a/Engine/Plugins/FX/Niagara/Source/Niagara/Private/NiagaraDataSet.cpp
+++ b/Engine/Plugins/FX/Niagara/Source/Niagara/Private/NiagaraDataSet.cpp
@@ -781,8 +781,6 @@ void FNiagaraDataBuffer::AllocateGPU(FRHICommandList& RHICmdList, uint32 InNumIn
 	}
 	else // Otherwise check for growing and possibly shrinking (if GNiagaraGPUDataBufferBufferSlack > 1) .
 	{
-		TArray<FRHITransitionInfo, TInlineAllocator<5>> Transitions;
-
 		// Float buffer requires growing or shrinking?
 		const int32 RequiredFloatByteSize = Align(FloatStride * Owner->GetNumFloatComponents(), GNiagaraGPUDataBufferChunkSize);
 		const int32 CurrentFloatByteSize = (int32)GPUBufferFloat.NumBytes;
@@ -792,8 +790,7 @@ void FNiagaraDataBuffer::AllocateGPU(FRHICommandList& RHICmdList, uint32 InNumIn
 			{
 				GPUBufferFloat.Release();
 			}
-			GPUBufferFloat.Initialize(TEXT("GPUBufferFloat"), sizeof(float), RequiredFloatByteSize / sizeof(float), EPixelFormat::PF_R32_FLOAT, GPUBufferFlags);
-			Transitions.Add(FRHITransitionInfo(GPUBufferFloat.UAV, ERHIAccess::Unknown, ERHIAccess::SRVMask));
+			GPUBufferFloat.Initialize(TEXT("GPUBufferFloat"), sizeof(float), RequiredFloatByteSize / sizeof(float), EPixelFormat::PF_R32_FLOAT, ERHIAccess::SRVMask, GPUBufferFlags);
 		}
 
 		// Half buffer requires growing or shrinking?
@@ -805,8 +802,7 @@ void FNiagaraDataBuffer::AllocateGPU(FRHICommandList& RHICmdList, uint32 InNumIn
 			{
 				GPUBufferHalf.Release();
 			}
-			GPUBufferHalf.Initialize(TEXT("GPUBufferHalf"), sizeof(FFloat16), RequiredHalfByteSize / sizeof(FFloat16), EPixelFormat::PF_R16F, GPUBufferFlags);
-			Transitions.Add(FRHITransitionInfo(GPUBufferHalf.UAV, ERHIAccess::Unknown, ERHIAccess::SRVMask));
+			GPUBufferHalf.Initialize(TEXT("GPUBufferHalf"), sizeof(FFloat16), RequiredHalfByteSize / sizeof(FFloat16), EPixelFormat::PF_R16F, ERHIAccess::SRVMask, GPUBufferFlags);
 		}
 
 		// Int buffer requires growing or shrinking?
@@ -818,8 +814,7 @@ void FNiagaraDataBuffer::AllocateGPU(FRHICommandList& RHICmdList, uint32 InNumIn
 			{
 				GPUBufferInt.Release();
 			}
-			GPUBufferInt.Initialize(TEXT("GPUBufferInt"), sizeof(int32), RequiredInt32ByteSize / sizeof(int32), EPixelFormat::PF_R32_SINT, GPUBufferFlags);
-			Transitions.Add(FRHITransitionInfo(GPUBufferInt.UAV, ERHIAccess::Unknown, ERHIAccess::SRVMask));
+			GPUBufferInt.Initialize(TEXT("GPUBufferInt"), sizeof(int32), RequiredInt32ByteSize / sizeof(int32), EPixelFormat::PF_R32_SINT, ERHIAccess::SRVMask, GPUBufferFlags);
 		}
 
 		// Allocate persistent IDs?
@@ -835,13 +830,9 @@ void FNiagaraDataBuffer::AllocateGPU(FRHICommandList& RHICmdList, uint32 InNumIn
 				}
 				TStringBuilder<128> DebugBufferName;
 				DebugBufferName.Appendf(TEXT("NiagaraIDToIndexTable_%s_%p"), DebugSimName ? DebugSimName : TEXT(""), this);
-				GPUIDToIndexTable.Initialize(DebugBufferName.ToString(), sizeof(int32), NumNeededElems, EPixelFormat::PF_R32_SINT, BUF_Static);
-				Transitions.Add(FRHITransitionInfo(GPUIDToIndexTable.UAV, ERHIAccess::Unknown, ERHIAccess::SRVCompute));
+				GPUIDToIndexTable.Initialize(DebugBufferName.ToString(), sizeof(int32), NumNeededElems, EPixelFormat::PF_R32_SINT, ERHIAccess::SRVCompute, BUF_Static);
 			}
 		}
-
-		// NiagaraEmitterInstanceBatcher expects the buffers to be readable before running the sim.
-		RHICmdList.Transition(MakeArrayView(Transitions.GetData(), Transitions.Num()));
 	}
 #if NIAGARA_MEMORY_TRACKING
 	AllocationSizeBytes = GPUBufferFloat.NumBytes + GPUBufferHalf.NumBytes + GPUBufferInt.NumBytes + GPUIDToIndexTable.NumBytes;
diff --git a/Engine/Plugins/FX/Niagara/Source/Niagara/Private/NiagaraEmitterInstanceBatcher.cpp b/Engine/Plugins/FX/Niagara/Source/Niagara/Private/NiagaraEmitterInstanceBatcher.cpp
--- a/Engine/Plugins/FX/Niagara/Source/Niagara/Private/NiagaraEmitterInstanceBatcher.cpp
+++ b/Engine/Plugins/FX/Niagara/Source/Niagara/Private/NiagaraEmitterInstanceBatcher.cpp
@@ -619,7 +619,6 @@ void NiagaraEmitterInstanceBatcher::PostSimulateInterface(FRHICommandList& RHICm
 
 void NiagaraEmitterInstanceBatcher::UpdateFreeIDsListSizesBuffer(FRHICommandList& RHICmdList, uint32 NumInstances)
 {
-	ERHIAccess BeforeState = ERHIAccess::UAVCompute;
 	if (NumInstances > NumAllocatedFreeIDListSizes)
 	{
 		constexpr uint32 ALLOC_CHUNK_SIZE = 128;
@@ -628,14 +627,12 @@ void NiagaraEmitterInstanceBatcher::UpdateFreeIDsListSizesBuffer(FRHICommandList
 		{
 			FreeIDListSizesBuffer.Release();
 		}
-		FreeIDListSizesBuffer.Initialize(TEXT("NiagaraFreeIDListSizes"), sizeof(uint32), NumAllocatedFreeIDListSizes, EPixelFormat::PF_R32_SINT, BUF_Static);
-
-		BeforeState = ERHIAccess::Unknown;
+		FreeIDListSizesBuffer.Initialize(TEXT("NiagaraFreeIDListSizes"), sizeof(uint32), NumAllocatedFreeIDListSizes, EPixelFormat::PF_R32_SINT, ERHIAccess::UAVCompute, BUF_Static);
 	}
 
 	{
 		SCOPED_DRAW_EVENT(RHICmdList, NiagaraGPUComputeClearFreeIDListSizes);
-		RHICmdList.Transition(FRHITransitionInfo(FreeIDListSizesBuffer.UAV, BeforeState, ERHIAccess::UAVCompute));
+		RHICmdList.Transition(FRHITransitionInfo(FreeIDListSizesBuffer.UAV, ERHIAccess::UAVCompute, ERHIAccess::UAVCompute));
 		NiagaraFillGPUIntBuffer(RHICmdList, FeatureLevel, FreeIDListSizesBuffer, 0);
 	}
 }
@@ -2141,6 +2138,7 @@ FRHIUnorderedAccessView* NiagaraEmitterInstanceBatcher::GetEmptyUAVFromPool(FRHI
 	{
 		FDummyUAV& NewUAV = Pool.UAVs.AddDefaulted_GetRef();
 		NewUAV.Init(RHICmdList, Format, Type, TEXT("NiagaraEmitterInstanceBatcher::DummyUAV"));
+
 		// Dispatches which use dummy UAVs are allowed to overlap, since we don't care about the contents of these buffers.
 		// We never need to calll EndUAVOverlap() on these.
 		RHICmdList.BeginUAVOverlap(NewUAV.UAV);
diff --git a/Engine/Plugins/FX/Niagara/Source/Niagara/Private/NiagaraGPUInstanceCountManager.cpp b/Engine/Plugins/FX/Niagara/Source/Niagara/Private/NiagaraGPUInstanceCountManager.cpp
--- a/Engine/Plugins/FX/Niagara/Source/Niagara/Private/NiagaraGPUInstanceCountManager.cpp
+++ b/Engine/Plugins/FX/Niagara/Source/Niagara/Private/NiagaraGPUInstanceCountManager.cpp
@@ -200,21 +200,19 @@ FRWBuffer* FNiagaraGPUInstanceCountManager::AcquireCulledCountsBuffer(FRHIComman
 		if (!bAcquiredCulledCounts)
 		{
 			const int32 RecommendedCulledCounts = FMath::Max(GNiagaraMinCulledGPUInstanceCount, (int32)(RequiredCulledCounts * GNiagaraGPUCountBufferSlack));
-			ERHIAccess BeforeState = ERHIAccess::SRVCompute;
 			if (RecommendedCulledCounts > AllocatedCulledCounts)
 			{
 				// We need a bigger buffer
 				CulledCountBuffer.Release();
 
 				AllocatedCulledCounts = RecommendedCulledCounts;
-				CulledCountBuffer.Initialize(TEXT("NiagaraCulledGPUInstanceCounts"), sizeof(uint32), AllocatedCulledCounts, EPixelFormat::PF_R32_UINT, BUF_Transient);
-				BeforeState = ERHIAccess::Unknown;
+				CulledCountBuffer.Initialize(TEXT("NiagaraCulledGPUInstanceCounts"), sizeof(uint32), AllocatedCulledCounts, EPixelFormat::PF_R32_UINT, ERHIAccess::SRVCompute, BUF_Transient);
 			}
 
 			CulledCountBuffer.AcquireTransientResource();
 
 			// Initialize the buffer by clearing it to zero then transition it to be ready to write to
-			RHICmdList.Transition(FRHITransitionInfo(CulledCountBuffer.UAV, BeforeState, ERHIAccess::UAVCompute));
+			RHICmdList.Transition(FRHITransitionInfo(CulledCountBuffer.UAV, ERHIAccess::SRVCompute, ERHIAccess::UAVCompute));
 			RHICmdList.ClearUAVUint(CulledCountBuffer.UAV, FUintVector4(EForceInit::ForceInitToZero));
 			RHICmdList.Transition(FRHITransitionInfo(CulledCountBuffer.UAV, ERHIAccess::UAVCompute, ERHIAccess::UAVCompute));
 
@@ -239,9 +237,7 @@ void FNiagaraGPUInstanceCountManager::ResizeBuffers(FRHICommandListImmediate& RH
 			AllocatedInstanceCounts = RecommendedInstanceCounts;
 			TResourceArray<uint32> InitData;
 			InitData.AddZeroed(AllocatedInstanceCounts);
-			CountBuffer.Initialize(TEXT("NiagaraGPUInstanceCounts"), sizeof(uint32), AllocatedInstanceCounts, EPixelFormat::PF_R32_UINT, BUF_Static | BUF_SourceCopy, &InitData);
-			// NiagaraEmitterInstanceBatcher expects the count buffer to be readable and copyable before running the sim.
-			RHICmdList.Transition(FRHITransitionInfo(CountBuffer.UAV, ERHIAccess::UAVCompute, kCountBufferDefaultState));
+			CountBuffer.Initialize(TEXT("NiagaraGPUInstanceCounts"), sizeof(uint32), AllocatedInstanceCounts, EPixelFormat::PF_R32_UINT, kCountBufferDefaultState, BUF_Static | BUF_SourceCopy, &InitData);
 			//UE_LOG(LogNiagara, Log, TEXT("FNiagaraGPUInstanceCountManager::ResizeBuffers Alloc AllocatedInstanceCounts: %d ReservedInstanceCounts: %d"), AllocatedInstanceCounts, ReservedInstanceCounts);
 		}
 		// If we need to increase the buffer size to RecommendedInstanceCounts because the buffer is too small.
@@ -253,7 +249,7 @@ void FNiagaraGPUInstanceCountManager::ResizeBuffers(FRHICommandListImmediate& RH
 			TResourceArray<uint32> InitData;
 			InitData.AddZeroed(RecommendedInstanceCounts);
 			FRWBuffer NextCountBuffer;
-			NextCountBuffer.Initialize(TEXT("NiagaraGPUInstanceCounts"), sizeof(uint32), RecommendedInstanceCounts, EPixelFormat::PF_R32_UINT, BUF_Static | BUF_SourceCopy, &InitData);
+			NextCountBuffer.Initialize(TEXT("NiagaraGPUInstanceCounts"), sizeof(uint32), RecommendedInstanceCounts, EPixelFormat::PF_R32_UINT, ERHIAccess::UAVCompute, BUF_Static | BUF_SourceCopy, &InitData);
 
 			// Copy the current buffer in the next buffer. We don't need to transition any of the buffers, because the current buffer is transitioned to readable after
 			// the simulation, and the new buffer is created in the UAVCompute state.
@@ -307,7 +303,7 @@ void FNiagaraGPUInstanceCountManager::FlushIndirectArgsPool()
 
 		TResourceArray<uint32> InitData;
 		InitData.AddZeroed(PoolEntry->NumAllocated * NIAGARA_DRAW_INDIRECT_ARGS_SIZE);
-		PoolEntry->Buffer.Initialize(TEXT("NiagaraGPUDrawIndirectArgs"), sizeof(uint32), PoolEntry->NumAllocated * NIAGARA_DRAW_INDIRECT_ARGS_SIZE, EPixelFormat::PF_R32_UINT, BUF_Static | BUF_DrawIndirect, &InitData);
+		PoolEntry->Buffer.Initialize(TEXT("NiagaraGPUDrawIndirectArgs"), sizeof(uint32), PoolEntry->NumAllocated * NIAGARA_DRAW_INDIRECT_ARGS_SIZE, EPixelFormat::PF_R32_UINT, ERHIAccess::IndirectArgs, BUF_Static | BUF_DrawIndirect, &InitData);
 
 		// Reset the timer
 		DrawIndirectLowWaterFrames = 0;
@@ -336,7 +332,7 @@ FNiagaraGPUInstanceCountManager::FIndirectArgSlot FNiagaraGPUInstanceCountManage
 
 			TResourceArray<uint32> InitData;
 			InitData.AddZeroed(NewEntry->NumAllocated * NIAGARA_DRAW_INDIRECT_ARGS_SIZE);
-			NewEntry->Buffer.Initialize(TEXT("NiagaraGPUDrawIndirectArgs"), sizeof(uint32), NewEntry->NumAllocated * NIAGARA_DRAW_INDIRECT_ARGS_SIZE, EPixelFormat::PF_R32_UINT, BUF_Static | BUF_DrawIndirect, &InitData);
+			NewEntry->Buffer.Initialize(TEXT("NiagaraGPUDrawIndirectArgs"), sizeof(uint32), NewEntry->NumAllocated * NIAGARA_DRAW_INDIRECT_ARGS_SIZE, EPixelFormat::PF_R32_UINT, ERHIAccess::IndirectArgs, BUF_Static | BUF_DrawIndirect, &InitData);
 
 			PoolEntry = NewEntry.Get();
 			DrawIndirectPool.Emplace(MoveTemp(NewEntry));
@@ -384,10 +380,11 @@ void FNiagaraGPUInstanceCountManager::UpdateDrawIndirectBuffers(NiagaraEmitterIn
 			TArray<FRHITransitionInfo, TInlineAllocator<10>> Transitions;
 			for (auto& PoolEntry : DrawIndirectPool)
 			{
-				Transitions.Emplace(PoolEntry->Buffer.UAV, ERHIAccess::Unknown, ERHIAccess::UAVCompute);
+				Transitions.Emplace(PoolEntry->Buffer.UAV, ERHIAccess::IndirectArgs, ERHIAccess::UAVCompute);
 			}
 
-			if (CountBuffer.UAV.IsValid())
+			const bool bCountBufferIsValid = CountBuffer.UAV.IsValid();
+			if (bCountBufferIsValid)
 			{
 				Transitions.Emplace(CountBuffer.UAV, kCountBufferDefaultState, ERHIAccess::UAVCompute);
 				CountsUAV = CountBuffer.UAV;
@@ -417,6 +414,11 @@ void FNiagaraGPUInstanceCountManager::UpdateDrawIndirectBuffers(NiagaraEmitterIn
 			PermutationVector.Set<FNiagaraDrawIndirectArgsGenCS::FSupportsTextureRW>(GRHISupportsRWTextureBuffers ? 1 : 0);
 			TShaderMapRef<FNiagaraDrawIndirectArgsGenCS> DrawIndirectArgsGenCS(GetGlobalShaderMap(FeatureLevel), PermutationVector);
 
+			if (bCountBufferIsValid)
+			{
+				RHICmdList.BeginUAVOverlap(CountsUAV);
+			}
+
 			const int32 NumDispatches = FMath::Max(DrawIndirectPool.Num(), 1);
 			uint32 ArgGenTaskOffset = 0;
 			for (int32 DispatchIdx = 0; DispatchIdx < NumDispatches; ++DispatchIdx)
@@ -470,6 +472,11 @@ void FNiagaraGPUInstanceCountManager::UpdateDrawIndirectBuffers(NiagaraEmitterIn
 				ArgGenTaskOffset += NumArgGenTasks;
 			}
 
+			if (bCountBufferIsValid)
+			{
+				RHICmdList.EndUAVOverlap(CountsUAV);
+			}
+
 			Transitions.Reset();
 			for (auto& PoolEntry : DrawIndirectPool)
 			{
