- Release counts back into the count manager immediately we don't need to keep hold when readback is in flight since that's a different buffer

diff --git a/Engine/Plugins/FX/Niagara/Source/Niagara/Classes/NiagaraComputeExecutionContext.h b/Engine/Plugins/FX/Niagara/Source/Niagara/Classes/NiagaraComputeExecutionContext.h
--- a/Engine/Plugins/FX/Niagara/Source/Niagara/Classes/NiagaraComputeExecutionContext.h
+++ b/Engine/Plugins/FX/Niagara/Source/Niagara/Classes/NiagaraComputeExecutionContext.h
@@ -142,7 +142,6 @@ struct FNiagaraComputeExecutionContext
 		// The CPU instance count at the time the GPU count readback was issued. Always bigger or equal to the GPU count.
 		uint32 CPUCount = 0;
 	}  EmitterInstanceReadback;
-	void ReleaseReadbackCounter(FNiagaraGPUInstanceCountManager& GPUInstanceCountManager);
 	
 #if !UE_BUILD_SHIPPING
 	const TCHAR* GetDebugSimName() const { return *DebugSimName; }
diff --git a/Engine/Plugins/FX/Niagara/Source/Niagara/Private/NiagaraComputeExecutionContext.cpp b/Engine/Plugins/FX/Niagara/Source/Niagara/Private/NiagaraComputeExecutionContext.cpp
--- a/Engine/Plugins/FX/Niagara/Source/Niagara/Private/NiagaraComputeExecutionContext.cpp
+++ b/Engine/Plugins/FX/Niagara/Source/Niagara/Private/NiagaraComputeExecutionContext.cpp
@@ -293,24 +293,12 @@ void FNiagaraComputeExecutionContext::PostTick()
 	}
 }
 
-void FNiagaraComputeExecutionContext::ReleaseReadbackCounter(FNiagaraGPUInstanceCountManager& GPUInstanceCountManager)
-{
-	if ( EmitterInstanceReadback.GPUCountOffset != INDEX_NONE )
-	{
-		check(EmitterInstanceReadback.GPUCountOffset != CountOffset_RT);
-
-		GPUInstanceCountManager.FreeEntry(EmitterInstanceReadback.GPUCountOffset);
-		EmitterInstanceReadback.GPUCountOffset = INDEX_NONE;
-	}
-}
-
 void FNiagaraComputeExecutionContext::ResetInternal(NiagaraEmitterInstanceBatcher* Batcher)
 {
 	checkf(IsInRenderingThread(), TEXT("Can only reset the gpu context from the render thread"));
 
 	if (Batcher)
 	{
-		ReleaseReadbackCounter(Batcher->GetGPUInstanceCounterManager());
 		Batcher->GetGPUInstanceCounterManager().FreeEntry(CountOffset_RT);
 	}
 
diff --git a/Engine/Plugins/FX/Niagara/Source/Niagara/Private/NiagaraEmitterInstanceBatcher.cpp b/Engine/Plugins/FX/Niagara/Source/Niagara/Private/NiagaraEmitterInstanceBatcher.cpp
--- a/Engine/Plugins/FX/Niagara/Source/Niagara/Private/NiagaraEmitterInstanceBatcher.cpp
+++ b/Engine/Plugins/FX/Niagara/Source/Niagara/Private/NiagaraEmitterInstanceBatcher.cpp
@@ -806,6 +806,8 @@ void NiagaraEmitterInstanceBatcher::DumpDebugFrame()
 
 void NiagaraEmitterInstanceBatcher::UpdateInstanceCountManager(FRHICommandListImmediate& RHICmdList)
 {
+	TRACE_CPUPROFILER_EVENT_SCOPE(NiagaraEmitterInstanceBatcher_UpdateInstanceCountManager);
+
 	// Resize dispatch buffer count
 	//-OPT: No need to iterate over all the ticks, we can store this as ticks are queued
 	{
@@ -822,7 +824,7 @@ void NiagaraEmitterInstanceBatcher::UpdateInstanceCountManager(FRHICommandListIm
 					{
 						if (InstanceData.bResetData)
 						{
-							GPUInstanceCounterManager.FreeEntry(InstanceData.Context->EmitterInstanceReadback.GPUCountOffset);
+							InstanceData.Context->EmitterInstanceReadback.GPUCountOffset = INDEX_NONE;
 						}
 					}
 				}
@@ -862,8 +864,8 @@ void NiagaraEmitterInstanceBatcher::UpdateInstanceCountManager(FRHICommandListIm
 							UE_LOG(LogNiagara, Warning, TEXT("ComputeContext(%p) Emitter(%s) DeadInstances(%d) CountReleased(0x%08x)"), ComputeContext, ComputeContext->GetDebugSimName(), DeadInstanceCount, ComputeContext->EmitterInstanceReadback.GPUCountOffset);
 						}
 
-						// Release the counter for the readback
-						GPUInstanceCounterManager.FreeEntry(ComputeContext->EmitterInstanceReadback.GPUCountOffset);
+						// Readback complete
+						ComputeContext->EmitterInstanceReadback.GPUCountOffset = INDEX_NONE;
 					}
 				}
 			}
@@ -998,10 +1000,7 @@ void NiagaraEmitterInstanceBatcher::PrepareTicksForProxy(FRHICommandListImmediat
 						ComputeContext->EmitterInstanceReadback.CPUCount = SimStageData.SourceNumInstances;
 						ComputeContext->EmitterInstanceReadback.GPUCountOffset = SimStageData.SourceCountOffset;
 					}
-					else if (SimStageData.SourceCountOffset != ComputeContext->EmitterInstanceReadback.GPUCountOffset)
-					{
-						GpuDispatchList.CountsToRelease.Add(SimStageData.SourceCountOffset);
-					}
+					GpuDispatchList.CountsToRelease.Add(SimStageData.SourceCountOffset);
 				}
 				if (Tick.NumInstancesWithSimStages > 0)
 				{
@@ -1207,7 +1206,7 @@ void NiagaraEmitterInstanceBatcher::ExecuteTicks(FRHICommandList& RHICmdList, FR
 		return;
 	}
 
-
+	TRACE_CPUPROFILER_EVENT_SCOPE(NiagaraEmitterInstanceBatcher_ExecuteTicks);
 	SCOPED_DRAW_EVENTF(RHICmdList, NiagaraEmitterInstanceBatcher_ExecuteTicks, TEXT("NiagaraEmitterInstanceBatcher_ExecuteTicks - TickStage(%d)"), TickStage);
 	SCOPE_CYCLE_COUNTER(STAT_NiagaraGPUSimTick_RT);
 	SCOPED_GPU_STAT(RHICmdList, NiagaraGPUSimulation);
@@ -1390,12 +1389,9 @@ void NiagaraEmitterInstanceBatcher::ExecuteTicks(FRHICommandList& RHICmdList, FR
 		}
 	}
 
+	// Clear dispatch groups
+	// We do not release the counts as we won't do that until we finish the dispatches
 	DispatchList.DispatchGroups.Empty();
-	if (DispatchList.CountsToRelease.Num() > 0)
-	{
-		GPUInstanceCounterManager.FreeEntryArray(DispatchList.CountsToRelease);
-		DispatchList.CountsToRelease.Empty();
-	}
 }
 
 void NiagaraEmitterInstanceBatcher::DispatchStage(FRHICommandList& RHICmdList, FRHIUniformBuffer* ViewUniformBuffer, const FNiagaraGPUSystemTick& Tick, const FNiagaraComputeInstanceData& InstanceData, const FNiagaraSimStageData& SimStageData)
@@ -1729,9 +1725,11 @@ void NiagaraEmitterInstanceBatcher::PostRenderOpaque(FRDGBuilder& GraphBuilder,
 			ProcessDebugReadbacks(RHICmdList, false);
 		}
 
-		if (!GPUInstanceCounterManager.HasPendingGPUReadback() && bRequiresReadback)
+		if ( bRequiresReadback )
 		{
+			check(!GPUInstanceCounterManager.HasPendingGPUReadback());
 			GPUInstanceCounterManager.EnqueueGPUReadback(RHICmdList);
+			bRequiresReadback = false;
 		}
 	});
 
diff --git a/Engine/Plugins/FX/Niagara/Source/Niagara/Private/NiagaraSystemGpuComputeProxy.cpp b/Engine/Plugins/FX/Niagara/Source/Niagara/Private/NiagaraSystemGpuComputeProxy.cpp
--- a/Engine/Plugins/FX/Niagara/Source/Niagara/Private/NiagaraSystemGpuComputeProxy.cpp
+++ b/Engine/Plugins/FX/Niagara/Source/Niagara/Private/NiagaraSystemGpuComputeProxy.cpp
@@ -69,6 +69,9 @@ void FNiagaraSystemGpuComputeProxy::AddToBatcher(NiagaraEmitterInstanceBatcher*
 				ComputeContext->CurrentNumInstances_RT = 0;
 				ComputeContext->CurrentMaxInstances_RT = 0;
 
+				ComputeContext->EmitterInstanceReadback.CPUCount = 0;
+				ComputeContext->EmitterInstanceReadback.GPUCountOffset = INDEX_NONE;
+
 				for (int i=0; i < UE_ARRAY_COUNT(ComputeContext->DataBuffers_RT); ++i)
 				{
 					check(ComputeContext->DataBuffers_RT[i] == nullptr);
