- Add more information to outliner, tick group, what features are required, etc

diff --git a/Engine/Plugins/FX/Niagara/Source/Niagara/Private/NiagaraDebuggerClient.cpp b/Engine/Plugins/FX/Niagara/Source/Niagara/Private/NiagaraDebuggerClient.cpp
--- a/Engine/Plugins/FX/Niagara/Source/Niagara/Private/NiagaraDebuggerClient.cpp
+++ b/Engine/Plugins/FX/Niagara/Source/Niagara/Private/NiagaraDebuggerClient.cpp
@@ -360,20 +360,31 @@ bool FNiagaraDebuggerClient::UpdateOutliner(float DeltaSeconds)
 					// TODO: need to be able to access at least a shadow copy of per-emitter execution state and num particles. For now, we'll just allow unsafe access
 					if (FNiagaraSystemInstance* Inst = InstController->GetSystemInstance_Unsafe())
 					{
+						InstData.TickGroup = Inst->CalculateTickGroup();
+
+						InstData.bIsSolo = Inst->IsSolo();
+						InstData.bRequiresDistanceFieldData = Inst->RequiresDistanceFieldData();
+						InstData.bRequiresDepthBuffer = Inst->RequiresDepthBuffer();
+						InstData.bRequiresEarlyViewData = Inst->RequiresEarlyViewData();
+						InstData.bRequiresViewUniformBuffer = Inst->RequiresViewUniformBuffer();
+						InstData.bRequiresRayTracingScene = Inst->RequiresRayTracingScene();
+
 						InstData.Emitters.Reserve(Inst->GetEmitters().Num());
 						for (TSharedRef<FNiagaraEmitterInstance, ESPMode::ThreadSafe>& EmitterInst : Inst->GetEmitters())
 						{
 							FNiagaraOutlinerEmitterInstanceData& EmitterData = InstData.Emitters.AddDefaulted_GetRef();
-							if (EmitterInst->GetCachedEmitter())
+							if (UNiagaraEmitter* NiagaraEmitter = EmitterInst->GetCachedEmitter())
 							{
 								//TODO: This is a bit wasteful to copy the name into each instance data. Though we can't rely on the debugger side data matchin the actul running data on the device.
 								//We need to build a shared representation of the asset data from the client that we then reference from this per instance data.
-								EmitterData.EmitterName = EmitterInst->GetCachedEmitter()->GetUniqueEmitterName();
-								EmitterData.SimTarget = EmitterInst->GetCachedEmitter()->SimTarget;
+								EmitterData.EmitterName = NiagaraEmitter->GetUniqueEmitterName();
+								EmitterData.SimTarget = NiagaraEmitter->SimTarget;
 								//Move all above to a shared asset representation.
 
 								EmitterData.ExecState = EmitterInst->GetExecutionState();
 								EmitterData.NumParticles = EmitterInst->GetNumParticles();
+
+								EmitterData.bRequiresPersistentIDs = NiagaraEmitter->RequiresPersistentIDs();
 							}
 						}
 					}
diff --git a/Engine/Plugins/FX/Niagara/Source/Niagara/Public/NiagaraDebuggerCommon.h b/Engine/Plugins/FX/Niagara/Source/Niagara/Public/NiagaraDebuggerCommon.h
--- a/Engine/Plugins/FX/Niagara/Source/Niagara/Public/NiagaraDebuggerCommon.h
+++ b/Engine/Plugins/FX/Niagara/Source/Niagara/Public/NiagaraDebuggerCommon.h
@@ -49,6 +49,13 @@ struct FNiagaraOutlinerEmitterInstanceData
 	UPROPERTY(VisibleAnywhere, Category = "Emitter")
 	int32 NumParticles = 0;
 
+	UPROPERTY(VisibleAnywhere, Category = "Emitter")
+	uint32 bRequiresPersistentIDs : 1;
+
+	FNiagaraOutlinerEmitterInstanceData()
+		: bRequiresPersistentIDs(false)
+	{}
+
 	//Mem Usage?
 	//Scalability info?
 };
@@ -90,9 +97,37 @@ struct FNiagaraOutlinerSystemInstanceData
 	UPROPERTY(VisibleAnywhere, Category = "Performance")
 	FNiagaraOutlinerTimingData MaxTime;
 
-	FNiagaraOutlinerSystemInstanceData() 
-	: bPendingKill(false)
-	, bUsingCullProxy(false)
+	UPROPERTY(VisibleAnywhere, Category = "Ticking")
+	TEnumAsByte<ETickingGroup> TickGroup;
+
+	UPROPERTY(VisibleAnywhere, Category = "Gpu")
+	uint32 bIsSolo : 1;
+
+	UPROPERTY(VisibleAnywhere, Category = "Gpu")
+	uint32 bRequiresDistanceFieldData : 1;
+
+	UPROPERTY(VisibleAnywhere, Category = "Gpu")
+	uint32 bRequiresDepthBuffer : 1;
+
+	UPROPERTY(VisibleAnywhere, Category = "Gpu")
+	uint32 bRequiresEarlyViewData : 1;
+
+	UPROPERTY(VisibleAnywhere, Category = "Gpu")
+	uint32 bRequiresViewUniformBuffer : 1;
+
+	UPROPERTY(VisibleAnywhere, Category = "Gpu")
+	uint32 bRequiresRayTracingScene : 1;
+
+	FNiagaraOutlinerSystemInstanceData()
+		: bPendingKill(false)
+		, bUsingCullProxy(false)
+		, TickGroup(0)
+		, bIsSolo(false)
+		, bRequiresDistanceFieldData(false)
+		, bRequiresDepthBuffer(false)
+		, bRequiresEarlyViewData(false)
+		, bRequiresViewUniformBuffer(false)
+		, bRequiresRayTracingScene(false)
 	{}
 
 	//TODO:
diff --git a/Engine/Plugins/FX/Niagara/Source/NiagaraEditor/Private/Customizations/NiagaraOutlinerCustomization.cpp b/Engine/Plugins/FX/Niagara/Source/NiagaraEditor/Private/Customizations/NiagaraOutlinerCustomization.cpp
--- a/Engine/Plugins/FX/Niagara/Source/NiagaraEditor/Private/Customizations/NiagaraOutlinerCustomization.cpp
+++ b/Engine/Plugins/FX/Niagara/Source/NiagaraEditor/Private/Customizations/NiagaraOutlinerCustomization.cpp
@@ -1425,6 +1425,17 @@ void FNiagaraOutlinerEmitterInstanceDetailsCustomization::CustomizeHeader(TShare
 
 void FNiagaraOutlinerEmitterInstanceDetailsCustomization::CustomizeChildren(TSharedRef<IPropertyHandle> StructPropertyHandle, class IDetailChildrenBuilder& ChildBuilder, IPropertyTypeCustomizationUtils& StructCustomizationUtils)
 {
+	FNiagaraEditorModule& NiagaraEditorModule = FModuleManager::GetModuleChecked<FNiagaraEditorModule>("NiagaraEditor");
+
+	uint32 NumChildren = 0;
+	StructPropertyHandle->GetNumChildren(NumChildren);
+
+	for (uint32 ChildNum = 0; ChildNum < NumChildren; ++ChildNum)
+	{
+		TSharedPtr<IPropertyHandle> ChildProperty = StructPropertyHandle->GetChildHandle(ChildNum);
+
+		ChildBuilder.AddProperty(ChildProperty.ToSharedRef());
+	}
 }
 
 #undef LOCTEXT_NAMESPACE
