- Destroy render state when the system completes
- Most paths will auto destroy but if the user did want to keep the component around to re-activate later we had a small cost with GDME lingering

diff --git a/Engine/Plugins/FX/Niagara/Source/Niagara/Private/NiagaraComponent.cpp b/Engine/Plugins/FX/Niagara/Source/Niagara/Private/NiagaraComponent.cpp
--- a/Engine/Plugins/FX/Niagara/Source/Niagara/Private/NiagaraComponent.cpp
+++ b/Engine/Plugins/FX/Niagara/Source/Niagara/Private/NiagaraComponent.cpp
@@ -1039,6 +1039,12 @@ void UNiagaraComponent::ActivateInternal(bool bReset /* = false */, bool bIsScal
 
 	RegisterWithScalabilityManager();
 
+	// If we have no render state create one
+	if (!IsRenderStateCreated())
+	{
+		CreateRenderState_Concurrent(nullptr);
+	}
+
 	// NOTE: This call can cause SystemInstance itself to get destroyed with auto destroy systems
 	SystemInstanceController->Activate(ResetMode);
 	
@@ -1269,8 +1275,20 @@ void UNiagaraComponent::OnSystemComplete(bool bExternalCompletion)
 	SetComponentTickEnabled(false);
 	SetActiveFlag(false);
 
+	// Give renderers a chance to handle completion
+	if (SystemInstanceController)
+	{
+		if (auto NiagaraProxy = static_cast<FNiagaraSceneProxy*>(SceneProxy))
+		{
+			if (auto RenderData = NiagaraProxy->GetSystemRenderData())
+			{
+				SystemInstanceController->NotifyRenderersComplete(*RenderData);
+			}
+		}
+	}
+
+	// Mark our render data dirty as we want to create an empty set of render data in the case of scalability deactivating
 	MarkRenderDynamicDataDirty();
-	//TODO: Mark the render state dirty?
 
 	//Don't really complete if we're being culled by scalability.
 	//We want to stop ticking but not be reclaimed by the pools etc.
@@ -1298,9 +1316,19 @@ void UNiagaraComponent::OnSystemComplete(bool bExternalCompletion)
 			DestroyComponent();
 			//UE_LOG(LogNiagara, Log, TEXT("OnSystemComplete DestroyComponent();: } %p - %s"), SystemInstance.Get(), *Asset->GetName());
 		}
-		else if (bAutoManageAttachment)
+		else
 		{
-			CancelAutoAttachment(/*bDetachFromParent=*/ true);
+			if (bAutoManageAttachment)
+			{
+				CancelAutoAttachment(/*bDetachFromParent=*/ true);
+			}
+
+			// For components who do not auto destroy or auto release back into the pool ensure we remove the render state
+			// Without doing this we will 'leak' proxies until the user manually destroys or releases.
+			if (IsRenderStateCreated())
+			{
+				DestroyRenderState_Concurrent();
+			}
 		}
 
 		if (IsRegisteredWithScalabilityManager())
@@ -1327,18 +1355,6 @@ void UNiagaraComponent::OnSystemComplete(bool bExternalCompletion)
 		}
 	}
 
-	// Give renderers a chance to handle completion
-	if (SystemInstanceController)
-	{
-		if (auto NiagaraProxy = static_cast<FNiagaraSceneProxy*>(SceneProxy))
-		{
-			if (auto RenderData = NiagaraProxy->GetSystemRenderData())
-			{
-				SystemInstanceController->NotifyRenderersComplete(*RenderData);
-			}
-		}
-	}
-
 	if (GIsAutomationTesting && GetAsset() != nullptr)
 	{
 		UE_LOG(LogNiagara, Log, TEXT("UNiagaraComponent> System %s completed."), *GetAsset()->GetPathName());
