Merging CL 16774974: Move ddc request for compilation to be async. Also fixed a bug where FNiagaraEditorModule::OnPreExit would crash when trying to complete the compilation.

diff --git a/Engine/Plugins/FX/Niagara/Source/Niagara/Classes/NiagaraScript.h b/Engine/Plugins/FX/Niagara/Source/Niagara/Classes/NiagaraScript.h
--- a/Engine/Plugins/FX/Niagara/Source/Niagara/Classes/NiagaraScript.h
+++ b/Engine/Plugins/FX/Niagara/Source/Niagara/Classes/NiagaraScript.h
@@ -941,7 +941,7 @@ public:
 	NIAGARA_API FOnPropertyChanged& OnPropertyChanged();
 
 	/** External call used to identify the values for a successful VM script compilation. OnVMScriptCompiled will be issued in this case.*/
-	void SetVMCompilationResults(const FNiagaraVMExecutableDataId& InCompileId, FNiagaraVMExecutableData& InScriptVM, FNiagaraCompileRequestDataBase* InRequestData);
+	void SetVMCompilationResults(const FNiagaraVMExecutableDataId& InCompileId, FNiagaraVMExecutableData& InScriptVM, FString EmitterUniqueName, const TMap<FName, UNiagaraDataInterface*>& ObjectNameMap);
 
 	/** In the event where we "merge" we duplicate the changes of the master copy onto the newly cloned copy. This function will synchronize the compiled script 
 		results assuming that the scripts themselves are bound to the same key. This saves looking things up in the DDC. It returns true if successfully synchronized and 
diff --git a/Engine/Plugins/FX/Niagara/Source/Niagara/Classes/NiagaraScriptSourceBase.h b/Engine/Plugins/FX/Niagara/Source/Niagara/Classes/NiagaraScriptSourceBase.h
--- a/Engine/Plugins/FX/Niagara/Source/Niagara/Classes/NiagaraScriptSourceBase.h
+++ b/Engine/Plugins/FX/Niagara/Source/Niagara/Classes/NiagaraScriptSourceBase.h
@@ -103,6 +103,8 @@ class UNiagaraScriptSourceBase : public UObject
 
 	virtual void ComputeVMCompilationId(struct FNiagaraVMExecutableDataId& Id, ENiagaraScriptUsage InUsage, const FGuid& InUsageId, bool bForceRebuild = false) const {};
 
+	virtual TMap<FName, UNiagaraDataInterface*> ComputeObjectNameMap(FString EmitterUniqueName) const { return TMap<FName, UNiagaraDataInterface*>(); }
+
 	virtual FGuid GetCompileBaseId(ENiagaraScriptUsage InUsage, const FGuid& InUsageId) const { return FGuid(); }
 
 	virtual FNiagaraCompileHash GetCompileHash(ENiagaraScriptUsage InUsage, const FGuid& InUsageId) const { return FNiagaraCompileHash(); }
diff --git a/Engine/Plugins/FX/Niagara/Source/Niagara/Classes/NiagaraSystem.h b/Engine/Plugins/FX/Niagara/Source/Niagara/Classes/NiagaraSystem.h
--- a/Engine/Plugins/FX/Niagara/Source/Niagara/Classes/NiagaraSystem.h
+++ b/Engine/Plugins/FX/Niagara/Source/Niagara/Classes/NiagaraSystem.h
@@ -3,20 +3,22 @@
 #pragma once
 
 #include "CoreMinimal.h"
-#include "UObject/ObjectMacros.h"
-#include "UObject/Object.h"
-#include "Particles/ParticleSystem.h"
-#include "Particles/ParticlePerfStats.h"
+
+#include "DerivedDataPluginInterface.h"
+#include "NiagaraBakerSettings.h"
 #include "NiagaraCommon.h"
 #include "NiagaraDataSet.h"
 #include "NiagaraDataSetAccessor.h"
-#include "NiagaraEmitterInstance.h"
+#include "NiagaraEffectType.h"
 #include "NiagaraEmitterHandle.h"
-#include "NiagaraBakerSettings.h"
+#include "NiagaraEmitterInstance.h"
 #include "NiagaraParameterCollection.h"
-#include "NiagaraUserRedirectionParameterStore.h"
-#include "NiagaraEffectType.h"
 #include "NiagaraParameterDefinitionsSubscriber.h"
+#include "NiagaraUserRedirectionParameterStore.h"
+#include "Particles/ParticlePerfStats.h"
+#include "Particles/ParticleSystem.h"
+#include "UObject/Object.h"
+#include "UObject/ObjectMacros.h"
 
 #include "NiagaraSystem.generated.h"
 
@@ -150,15 +152,101 @@ struct FEmitterCompiledScriptPair
 {
 	GENERATED_USTRUCT_BODY()
 	
-	bool bResultsReady;
-	UNiagaraEmitter* Emitter;
-	UNiagaraScript* CompiledScript;
+	bool bResultsReady = false;
+	UNiagaraEmitter* Emitter = nullptr;
+	UNiagaraScript* CompiledScript = nullptr;
 	uint32 PendingJobID = INDEX_NONE; // this is the ID for any active shader compiler worker job
 	FNiagaraVMExecutableDataId CompileId;
 	TSharedPtr<FNiagaraVMExecutableData> CompileResults;
 	int32 ParentIndex = INDEX_NONE;
 };
 
+UENUM()
+enum class ENiagaraCompilationState : uint8
+{
+	Precompile,
+	StartCompileJob,
+	AwaitResult,
+	ProcessResult,
+	Finished,
+	Aborted
+};
+
+struct FNiagaraLazyPrecompileReference
+{
+	TSharedPtr<FNiagaraCompileRequestDataBase, ESPMode::ThreadSafe> GetPrecompileData(UNiagaraScript* ForScript);
+	
+	UNiagaraSystem* System = nullptr;
+	TArray<UNiagaraScript*> Scripts;
+	TMap<UNiagaraScript*, int32> EmitterScriptIndex;
+	TArray<TObjectPtr<UObject>>* CompilationRootObjects;
+	TArray<FNiagaraVariable> EncounteredExposedVars;
+
+private:
+	TSharedPtr<FNiagaraCompileRequestDataBase, ESPMode::ThreadSafe> SystemPrecompiledData;
+	TMap<UNiagaraScript*, TSharedPtr<FNiagaraCompileRequestDataBase, ESPMode::ThreadSafe>> EmitterMapping;
+};
+
+struct FNiagaraAsyncTaskSharedData
+{
+	FNiagaraAsyncTaskSharedData();
+	~FNiagaraAsyncTaskSharedData();
+	
+	double StartTaskTime = 0;
+	double StartCompileTime = 0;
+	bool bWaitForCompileJob = false;
+	FEmitterCompiledScriptPair ScriptPair;
+
+	std::atomic<ENiagaraCompilationState> CurrentState;
+	FCriticalSection StateProcessingCS;
+	FEvent* WaitOnGTEvent;
+	std::atomic<bool> bWaitingOnGT;
+
+	TSharedPtr<FNiagaraVMExecutableData> ExeData;
+	TArray<uint8> DDCOutData;
+};
+
+class FNiagaraDerivedDataAsyncRequest : public FDerivedDataPluginInterface
+{
+#if WITH_EDITORONLY_DATA
+public:
+	FString DDCKey;
+	FString AssetPath;
+	FString UniqueEmitterName;
+	uint32 TaskHandle = 0;
+	bool bResultProcessed = false;
+
+	// this data is shared between the ddc thread and the game thread that starts the compilation
+	TSharedPtr<FNiagaraAsyncTaskSharedData, ESPMode::ThreadSafe> TaskSharedData;
+	TSharedPtr<FNiagaraLazyPrecompileReference, ESPMode::ThreadSafe> PrecompileReference;
+	TSharedPtr<FNiagaraCompileRequestDataBase, ESPMode::ThreadSafe> ComputedPrecompileData;
+
+	FNiagaraDerivedDataAsyncRequest(FString InAssetPath, const FEmitterCompiledScriptPair& InScriptPair);
+
+	virtual FString GetPluginSpecificCacheKeySuffix() const override { return DDCKey; }
+	virtual bool IsBuildThreadsafe() const override	{ return true; }
+	virtual bool IsDeterministic() const override {	return true; }
+	virtual FString GetDebugContextString() const override { return AssetPath; }
+	virtual const TCHAR* GetPluginName() const override;
+	virtual const TCHAR* GetVersionString() const override;
+	virtual bool Build(TArray<uint8>& OutData) override;
+
+	const FEmitterCompiledScriptPair& GetScriptPair() const { return TaskSharedData->ScriptPair; };
+	void WaitAndResolveResult();
+	void AbortTask();
+
+	void ProcessCurrentState();
+	void MoveToState(ENiagaraCompilationState NewState);
+	bool IsDone() const;
+
+	void PrecompileData();
+	void StartCompileJob();
+	bool AwaitResult();
+	void ProcessResult();
+
+#endif
+};
+
 USTRUCT()
 struct FNiagaraSystemCompileRequest
 {
@@ -169,13 +257,13 @@ struct FNiagaraSystemCompileRequest
 	UPROPERTY()
 	TArray<TObjectPtr<UObject>> RootObjects;
 
-	TArray<FEmitterCompiledScriptPair> EmitterCompiledScriptPairs;
+	TArray<TSharedPtr<FNiagaraDerivedDataAsyncRequest, ESPMode::ThreadSafe>> DDCTasks;
 	
-	TMap<UNiagaraScript*, TSharedPtr<FNiagaraCompileRequestDataBase, ESPMode::ThreadSafe> > MappedData;
-
 	bool bIsValid = true;
-
 	bool bForced = false;
+
+	bool bHasCompiledJobs = false;
+	float CombinedCompileTime = 0.0f;
 };
 
 struct FNiagaraEmitterExecutionIndex
@@ -362,6 +450,9 @@ public:
 	/** Blocks until all active compile jobs have finished */
 	void WaitForCompilationComplete(bool bIncludingGPUShaders = false, bool bShowProgress = true);
 
+	/** Tries to abort all running shader compilations */
+	void KillAllActiveCompilations();
+
 	/** Invalidates any active compilation requests which will ignore their results. */
 	void InvalidateActiveCompiles();
 
@@ -571,15 +662,13 @@ public:
 
 private:
 #if WITH_EDITORONLY_DATA
-	/** Checks the ddc for vm execution data for the given script. Return true if the data was loaded from the ddc, false otherwise. */
-	bool GetFromDDC(FEmitterCompiledScriptPair& ScriptPair);
 
 	/** Since the shader compilation is done in another process, this is used to check if the result for any ongoing compilations is done.
-	*   If bWait is true then this *blocks* the game thread (and ui) until all running compilations are finished.
+	*   If bWait is true then this *blocks* the game thread (and ui) until all shader compilations are finished.
 	*/
 	bool QueryCompileComplete(bool bWait, bool bDoPost, bool bDoNotApply = false);
 
-	bool ProcessCompilationResult(FEmitterCompiledScriptPair& ScriptPair, bool bWait, bool bDoNotApply);
+	void ProcessWaitingDDCTasks(bool bProcessForWait);
 
 	bool CompilationResultsValid(FNiagaraSystemCompileRequest& CompileRequest) const;
 
@@ -595,6 +684,7 @@ private:
 
 	/** Helper for filling in attribute datasets per emitter. */
 	void InitEmitterDataSetCompiledData(FNiagaraDataSetCompiledData& DataSetToInit, const UNiagaraEmitter* InAssociatedEmitter, const FNiagaraEmitterHandle& InAssociatedEmitterHandle);
+	void PrepareRapidIterationParametersForCompilation();
 #endif
 
 	void ResolveScalabilitySettings();
diff --git a/Engine/Plugins/FX/Niagara/Source/Niagara/Niagara.Build.cs b/Engine/Plugins/FX/Niagara/Source/Niagara/Niagara.Build.cs
--- a/Engine/Plugins/FX/Niagara/Source/Niagara/Niagara.Build.cs
+++ b/Engine/Plugins/FX/Niagara/Source/Niagara/Niagara.Build.cs
@@ -44,6 +44,11 @@ public class Niagara : ModuleRules
                 "VectorVM",
             }
         );
+        
+        PublicIncludePathModuleNames.AddRange(
+	        new string[] {
+		        "DerivedDataCache",
+	        });
 
 		PrivateIncludePathModuleNames.AddRange(
 			new string[] {
diff --git a/Engine/Plugins/FX/Niagara/Source/Niagara/Private/NiagaraScript.cpp b/Engine/Plugins/FX/Niagara/Source/Niagara/Private/NiagaraScript.cpp
--- a/Engine/Plugins/FX/Niagara/Source/Niagara/Private/NiagaraScript.cpp
+++ b/Engine/Plugins/FX/Niagara/Source/Niagara/Private/NiagaraScript.cpp
@@ -537,11 +537,6 @@ void FNiagaraVMExecutableDataId::AppendKeyString(FString& KeyString, const FStri
 	}
 	KeyString += Delimiter;
 
-	if (bAppendObjectForDebugging)
-	{
-		KeyString += TEXT("[AdditionalDefines]") + Delimiter;
-	}
-
 	if (bUsesRapidIterationParams)
 	{
 		KeyString += TEXT("USESRI") + Delimiter;
@@ -551,12 +546,22 @@ void FNiagaraVMExecutableDataId::AppendKeyString(FString& KeyString, const FStri
 		KeyString += TEXT("NORI") + Delimiter;
 	}
 
+	if (bAppendObjectForDebugging)
+	{
+		KeyString += TEXT("[AdditionalDefines]") + Delimiter;
+	}
+
 	for (int32 Idx = 0; Idx < AdditionalDefines.Num(); Idx++)
 	{
 		KeyString += AdditionalDefines[Idx];
 		KeyString += Delimiter;
 	}
 
+	if (bAppendObjectForDebugging)
+	{
+		KeyString += TEXT("[AdditionalVariables]") + Delimiter;
+	}
+
 	for (int32 Idx = 0; Idx < AdditionalVariables.Num(); Idx++)
 	{
 		KeyString += AdditionalVariables[Idx].GetName().ToString();
@@ -807,9 +812,8 @@ void UNiagaraScript::ComputeVMCompilationId(FNiagaraVMExecutableDataId& Id, FGui
 			HashState.GetHash(DataHash.GetData());
 
 			FNiagaraCompileHash Hash(DataHash);
-			Id.ReferencedCompileHashes.Add(Hash);
+			Id.ReferencedCompileHashes.AddUnique(Hash);
 			Id.DebugReferencedObjects.Add(TEXT("SimulationStageHeaders"));
-
 		}
 	}
 
@@ -893,7 +897,7 @@ void UNiagaraScript::ComputeVMCompilationId(FNiagaraVMExecutableDataId& Id, FGui
 		HashState.GetHash(DataHash.GetData());
 
 		FNiagaraCompileHash Hash(DataHash);
-		Id.ReferencedCompileHashes.Add(Hash);
+		Id.ReferencedCompileHashes.AddUnique(Hash);
 		Id.DebugReferencedObjects.Add(TEXT("RIParams"));
 	}
 
@@ -917,6 +921,11 @@ void UNiagaraScript::ComputeVMCompilationId(FNiagaraVMExecutableDataId& Id, FGui
 		}
 	}
 
+	// normalize the order
+	Id.AdditionalDefines.Sort();
+	Id.ReferencedCompileHashes.Sort([](const FNiagaraCompileHash& Lhs, const FNiagaraCompileHash& Rhs) { return Lhs.ToString() < Rhs.ToString(); });
+	Id.AdditionalVariables.Sort([](const FNiagaraVariableBase& Lhs, const FNiagaraVariableBase& Rhs) { return Lhs.GetName().LexicalLess(Rhs.GetName()); });
+
 	LastGeneratedVMId = Id;
 }
 
@@ -2056,9 +2065,9 @@ FVersionedNiagaraScriptData* FVersionedNiagaraScript::GetScriptData() const
 }
 #endif
 
-UNiagaraDataInterface* ResolveDataInterface(FNiagaraCompileRequestDataBase* InBase, FName VariableName)
+UNiagaraDataInterface* ResolveDataInterface(const TMap<FName, UNiagaraDataInterface*>& ObjectNameMap, FName VariableName)
 {
-	UNiagaraDataInterface* const* FoundDI = InBase->GetObjectNameMap().Find(VariableName);
+	UNiagaraDataInterface* const* FoundDI = ObjectNameMap.Find(VariableName);
 	if (FoundDI && *(FoundDI))
 	{
 		return *FoundDI;
@@ -2066,20 +2075,31 @@ UNiagaraDataInterface* ResolveDataInterface(FNiagaraCompileRequestDataBase* InBa
 	return nullptr;
 }
 
-void DumpNameMap(FNiagaraCompileRequestDataBase* InBase)
+void DumpNameMap(const TMap<FName, UNiagaraDataInterface*>& ObjectNameMap)
 {
-	for (const TPair<FName, UNiagaraDataInterface*>& Pair : InBase->GetObjectNameMap())
+	for (const TPair<FName, UNiagaraDataInterface*>& Pair : ObjectNameMap)
 	{
 		UE_LOG(LogNiagara, Log, TEXT("%s -> %s"), *Pair.Key.ToString(), *GetPathNameSafe(Pair.Value));
 	}
 }
 
+FName ResolveEmitterAlias(const FName& InName, const FString& InAlias)
+{
+	// If the alias is empty than the name can't be resolved.
+	if (InAlias.IsEmpty())
+	{
+		return InName;
+	}
 
+	FNiagaraVariable Var(FNiagaraTypeDefinition::GetFloatDef(), InName);
+	FNiagaraAliasContext ResolveAliasesContext(FNiagaraAliasContext::ERapidIterationParameterMode::EmitterOrParticleScript);
+	ResolveAliasesContext.ChangeEmitterToEmitterName(InAlias);
+	Var = FNiagaraUtilities::ResolveAliases(Var, ResolveAliasesContext);
+	return Var.GetName();
+}
 
-void UNiagaraScript::SetVMCompilationResults(const FNiagaraVMExecutableDataId& InCompileId, FNiagaraVMExecutableData& InScriptVM, FNiagaraCompileRequestDataBase* InRequestData)
+void UNiagaraScript::SetVMCompilationResults(const FNiagaraVMExecutableDataId& InCompileId, FNiagaraVMExecutableData& InScriptVM, FString EmitterUniqueName, const TMap<FName, UNiagaraDataInterface*>& ObjectNameMap)
 {
-	check(InRequestData != nullptr);
-
 	CachedScriptVMId = InCompileId;
 	CachedScriptVM = InScriptVM;
 	CachedParameterCollectionReferences.Empty();
@@ -2104,13 +2124,13 @@ void UNiagaraScript::SetVMCompilationResults(const FNiagaraVMExecutableDataId& I
 	{
 		int32 Idx = CachedDefaultDataInterfaces.AddDefaulted();
 		CachedDefaultDataInterfaces[Idx].UserPtrIdx = Info.UserPtrIdx;
-		CachedDefaultDataInterfaces[Idx].Name = InRequestData->ResolveEmitterAlias(Info.Name);
+		CachedDefaultDataInterfaces[Idx].Name = ResolveEmitterAlias(Info.Name, EmitterUniqueName);
 		CachedDefaultDataInterfaces[Idx].Type = Info.Type;
-		CachedDefaultDataInterfaces[Idx].RegisteredParameterMapRead = InRequestData->ResolveEmitterAlias(Info.RegisteredParameterMapRead);
-		CachedDefaultDataInterfaces[Idx].RegisteredParameterMapWrite = InRequestData->ResolveEmitterAlias(Info.RegisteredParameterMapWrite);
+		CachedDefaultDataInterfaces[Idx].RegisteredParameterMapRead = ResolveEmitterAlias(Info.RegisteredParameterMapRead, EmitterUniqueName);
+		CachedDefaultDataInterfaces[Idx].RegisteredParameterMapWrite = ResolveEmitterAlias(Info.RegisteredParameterMapWrite, EmitterUniqueName);
 
 		// We compiled it just a bit ago, so we should be able to resolve it from the table that we passed in.
-		UNiagaraDataInterface* FindDIById = ResolveDataInterface(InRequestData, CachedDefaultDataInterfaces[Idx].Name);
+		UNiagaraDataInterface* FindDIById = ResolveDataInterface(ObjectNameMap, CachedDefaultDataInterfaces[Idx].Name);
 		if (FindDIById != nullptr )
 		{
 			CachedDefaultDataInterfaces[Idx].DataInterface = CopyDataInterface(FindDIById, this);
@@ -2120,14 +2140,14 @@ void UNiagaraScript::SetVMCompilationResults(const FNiagaraVMExecutableDataId& I
 		if (CachedDefaultDataInterfaces[Idx].DataInterface == nullptr)
 		{
 			// Use the CDO since we didn't have a default..
-			UObject* Obj = const_cast<UClass*>(Info.Type.GetClass())->GetDefaultObject(true);
+			UObject* Obj = Info.Type.GetClass()->GetDefaultObject(true);
 			CachedDefaultDataInterfaces[Idx].DataInterface = Cast<UNiagaraDataInterface>(CopyDataInterface(CastChecked<UNiagaraDataInterface>(Obj), this));
 
 			if (Info.bIsPlaceholder == false)
 			{
 				UE_LOG(LogNiagara, Warning, TEXT("We somehow ended up with a data interface that we couldn't match post compile. This shouldn't happen. Creating a dummy to prevent crashes. DataInterfaceInfoName:%s Object:%s"), *Info.Name.ToString(), *GetPathNameSafe(this));
 				UE_LOG(LogNiagara, Log, TEXT("Object to Name map contents:"));
-				DumpNameMap(InRequestData);
+				DumpNameMap(ObjectNameMap);
 			}
 		}
 		check(CachedDefaultDataInterfaces[Idx].DataInterface != nullptr);
@@ -2202,7 +2222,7 @@ void UNiagaraScript::RequestCompile(const FGuid& ScriptVersion, bool bForceCompi
 			if (BinaryToExecData(this, OutData, ExeData))
 			{
 				COOK_STAT(Timer.AddHit(OutData.Num()));
-				SetVMCompilationResults(LastGeneratedVMId, ExeData, RequestData.Get());
+				SetVMCompilationResults(LastGeneratedVMId, ExeData, FString(), RequestData->GetObjectNameMap());
 				return;
 			}
 		}
@@ -2215,7 +2235,7 @@ void UNiagaraScript::RequestCompile(const FGuid& ScriptVersion, bool bForceCompi
 		TSharedPtr<FNiagaraVMExecutableData> ExeData = NiagaraModule.GetCompileJobResult(JobHandle, true);
 		if (ExeData)
 		{
-			SetVMCompilationResults(LastGeneratedVMId, *ExeData, RequestData.Get());
+			SetVMCompilationResults(LastGeneratedVMId, *ExeData, FString(), RequestData->GetObjectNameMap());
 			// save result to the ddc
 			if (ExecToBinaryData(this, OutData, *ExeData))
 			{
@@ -2247,6 +2267,7 @@ bool UNiagaraScript::RequestExternallyManagedAsyncCompile(const TSharedPtr<FNiag
 			OutAsyncHandle = (uint32)INDEX_NONE;
 			CachedScriptVM.LastCompileStatus = ENiagaraScriptCompileStatus::NCS_Unknown;
 			CachedScriptVMId = LastGeneratedVMId;
+			UE_LOG(LogNiagara, Warning, TEXT("Requested compile of script '%s', which cannot be compiled!"), *GetFullName());
 			return false;
 		}
 
diff --git a/Engine/Plugins/FX/Niagara/Source/Niagara/Private/NiagaraSystem.cpp b/Engine/Plugins/FX/Niagara/Source/Niagara/Private/NiagaraSystem.cpp
--- a/Engine/Plugins/FX/Niagara/Source/Niagara/Private/NiagaraSystem.cpp
+++ b/Engine/Plugins/FX/Niagara/Source/Niagara/Private/NiagaraSystem.cpp
@@ -40,6 +40,7 @@ static uint32 CompileGuardSlot = 0;
 
 DECLARE_CYCLE_STAT(TEXT("Niagara - System - Precompile"), STAT_Niagara_System_Precompile, STATGROUP_Niagara);
 DECLARE_CYCLE_STAT(TEXT("Niagara - System - CompileScript"), STAT_Niagara_System_CompileScript, STATGROUP_Niagara);
+DECLARE_CYCLE_STAT(TEXT("Niagara - System - CompileScriptAsyncTaskGT"), STAT_Niagara_System_CompileScriptAsyncTaskGT, STATGROUP_Niagara);
 DECLARE_CYCLE_STAT(TEXT("Niagara - System - CompileScript_ResetAfter"), STAT_Niagara_System_CompileScriptResetAfter, STATGROUP_Niagara);
 
 #if ENABLE_COOK_STATS
@@ -93,6 +94,14 @@ static FAutoConsoleVariableRef CVarNiagaraCompileWaitLoggingTerminationCap(
 	ECVF_Default
 );
 
+static int GNiagaraCompileAsyncDDCWaitTimeout = 600;
+static FAutoConsoleVariableRef CVarNiagaraCompileAsyncDDCWaitTimeout(
+	TEXT("fx.Niagara.CompileAsyncDDCWaitTimeout"),
+	GNiagaraCompileAsyncDDCWaitTimeout,
+	TEXT("The max amount of seconds that the ddc thread should wait for the shader compiler to finish"),
+	ECVF_Default
+);
+
 //////////////////////////////////////////////////////////////////////////
 
 UNiagaraSystem::UNiagaraSystem(const FObjectInitializer& ObjectInitializer)
@@ -132,6 +141,361 @@ UNiagaraSystem::UNiagaraSystem(const FObjectInitializer& ObjectInitializer)
 #endif
 }
 
+#if WITH_EDITORONLY_DATA
+class FAsyncCompileResultPollTask : public FAsyncGraphTaskBase
+{
+public:
+
+	FAsyncCompileResultPollTask(TUniqueFunction<void()>&& InFunction)
+		: Function(MoveTemp(InFunction))
+	{ }
+
+	void DoTask(ENamedThreads::Type CurrentThread, const FGraphEventRef& MyCompletionGraphEvent)
+	{
+		Function();
+	}
+
+	FORCEINLINE TStatId GetStatId() const { RETURN_QUICK_DECLARE_CYCLE_STAT(FAsyncCompileResultPollTask, STATGROUP_TaskGraphTasks); }
+	static FORCEINLINE ENamedThreads::Type GetDesiredThread() { return ENamedThreads::GameThread; }
+	static FORCEINLINE ESubsequentsMode::Type GetSubsequentsMode() { return ESubsequentsMode::TrackSubsequents; }
+
+private:
+	TUniqueFunction<void()> Function;
+};
+
+TSharedPtr<FNiagaraCompileRequestDataBase, ESPMode::ThreadSafe> FNiagaraLazyPrecompileReference::GetPrecompileData(UNiagaraScript* ForScript)
+{
+	if (SystemPrecompiledData == nullptr)
+	{
+		UNiagaraPrecompileContainer* Container = NewObject<UNiagaraPrecompileContainer>(GetTransientPackage());
+		Container->System = System;
+		Container->Scripts = Scripts;
+		INiagaraModule& NiagaraModule = FModuleManager::Get().LoadModuleChecked<INiagaraModule>(TEXT("Niagara"));
+		SystemPrecompiledData = NiagaraModule.Precompile(Container, FGuid());
+
+		// Grab the list of user variables that were actually encountered so that we can add to them later.
+		SystemPrecompiledData->GatherPreCompiledVariables(TEXT("User"), EncounteredExposedVars);
+		SystemPrecompiledData->GetReferencedObjects(*CompilationRootObjects);
+		TArray<FNiagaraEmitterHandle> EmitterHandles = System->GetEmitterHandles();
+		for (int32 i = 0; i < EmitterHandles.Num(); i++)
+		{
+			FNiagaraEmitterHandle Handle = EmitterHandles[i];
+			if (Handle.GetInstance() && Handle.GetIsEnabled())
+			{
+				TSharedPtr<FNiagaraCompileRequestDataBase, ESPMode::ThreadSafe> EmitterPrecompiledData = SystemPrecompiledData->GetDependentRequest(i);
+				EmitterPrecompiledData->GetReferencedObjects(*CompilationRootObjects);
+				
+				TArray<UNiagaraScript*> EmitterScripts;
+				Handle.GetInstance()->GetScripts(EmitterScripts, false, true);
+				check(EmitterScripts.Num() > 0);
+				for (UNiagaraScript* EmitterScript : EmitterScripts)
+				{
+					EmitterMapping.Add(EmitterScript, EmitterPrecompiledData);
+				}
+
+				// Add the emitter's User variables to the encountered list to expose for later.
+				EmitterPrecompiledData->GatherPreCompiledVariables(TEXT("User"), EncounteredExposedVars);
+			}
+		}
+	}
+	if (SystemPrecompiledData.IsValid() == false)
+	{
+		UE_LOG(LogNiagara, Error, TEXT("Failed to precompile %s.  This is due to unexpected invalid or broken data.  Additional details should be in the log."), *System->GetPathName());
+		return nullptr;
+	}
+
+	if (ForScript->IsSystemSpawnScript() || ForScript->IsSystemUpdateScript())
+	{
+		return SystemPrecompiledData;
+	}
+	return EmitterMapping.FindChecked(ForScript);
+}
+
+FNiagaraAsyncTaskSharedData::FNiagaraAsyncTaskSharedData()
+{
+	CurrentState.store(ENiagaraCompilationState::Precompile);
+	bWaitingOnGT.store(false);
+	WaitOnGTEvent = FPlatformProcess::GetSynchEventFromPool(false);
+}
+
+FNiagaraAsyncTaskSharedData::~FNiagaraAsyncTaskSharedData()
+{
+	FPlatformProcess::ReturnSynchEventToPool(WaitOnGTEvent);
+}
+
+FNiagaraDerivedDataAsyncRequest::FNiagaraDerivedDataAsyncRequest(FString InAssetPath, const FEmitterCompiledScriptPair& InScriptPair)
+{
+	AssetPath = InAssetPath;
+	TaskSharedData = MakeShared<FNiagaraAsyncTaskSharedData, ESPMode::ThreadSafe>();
+	TaskSharedData->ScriptPair = InScriptPair;
+}
+
+const TCHAR* FNiagaraDerivedDataAsyncRequest::GetPluginName() const
+{
+	return TEXT("Niagara");
+}
+
+const TCHAR* FNiagaraDerivedDataAsyncRequest::GetVersionString() const
+{
+	return TEXT("832C5B24BBCC4F4EA235EB8112450C56");
+}
+
+void FNiagaraDerivedDataAsyncRequest::ProcessCurrentState()
+{
+	// This is the logic for the state machine that is the async ddc fetch task. The task can run on any thread (including the game thread),
+	// but most things the task does actually need to run on the game thread. So once per tick the GT checks if the task is stalling
+	// and waiting for the GT to process anything. Even though some states might *look* thread safe (e.g. StartCompileJob), if they access
+	// any UObject, they are not thread safe because the garbage collector can run in parallel on the GT and will invalidate pointers.
+	{
+		FScopeLock ScopeLock(&TaskSharedData->StateProcessingCS);
+		if (IsDone())
+		{
+			return;
+		}
+	
+		if (TaskSharedData->CurrentState == ENiagaraCompilationState::Precompile)
+		{
+			// gather precompile data
+			TaskSharedData->StartTaskTime = FPlatformTime::Seconds();
+			if (IsInGameThread())
+			{
+				SCOPE_CYCLE_COUNTER(STAT_Niagara_System_CompileScriptAsyncTaskGT);
+				TaskSharedData->StartCompileTime = FPlatformTime::Seconds();
+				PrecompileData();
+				MoveToState(ENiagaraCompilationState::StartCompileJob);
+			}
+			else
+			{
+				TaskSharedData->bWaitingOnGT = true;
+			}
+		}
+		else if (TaskSharedData->CurrentState == ENiagaraCompilationState::StartCompileJob)
+		{
+			// start the async compile job
+			if (IsInGameThread())
+			{
+				SCOPE_CYCLE_COUNTER(STAT_Niagara_System_CompileScriptAsyncTaskGT);
+				TaskSharedData->StartCompileTime = FPlatformTime::Seconds();
+				StartCompileJob();
+				MoveToState(ENiagaraCompilationState::AwaitResult);
+			}
+			else
+			{
+				TaskSharedData->bWaitingOnGT = true;
+			}
+		}
+		else if (TaskSharedData->CurrentState == ENiagaraCompilationState::AwaitResult)
+		{
+			// wait for the async compile job to finish
+			if (IsInGameThread())
+			{
+				SCOPE_CYCLE_COUNTER(STAT_Niagara_System_CompileScriptAsyncTaskGT);
+				if (AwaitResult())
+				{
+					MoveToState(ENiagaraCompilationState::ProcessResult);
+				}
+			}
+			else if (FPlatformTime::Seconds() - TaskSharedData->StartCompileTime > GNiagaraCompileAsyncDDCWaitTimeout)
+			{
+				AbortTask();
+			}
+			else
+			{
+				TaskSharedData->bWaitingOnGT = true;
+			}
+		}
+		else if (TaskSharedData->CurrentState == ENiagaraCompilationState::ProcessResult)
+		{
+			// save the result from the compile job
+			if (IsInGameThread())
+			{
+				SCOPE_CYCLE_COUNTER(STAT_Niagara_System_CompileScriptAsyncTaskGT);
+				ProcessResult();
+				MoveToState(ENiagaraCompilationState::Finished);
+			}
+			else
+			{
+				TaskSharedData->bWaitingOnGT = true;
+			}
+		}
+		else
+		{
+			check(false);
+		}
+	}
+	if (TaskSharedData->bWaitingOnGT && !IsInGameThread())
+	{
+		TaskSharedData->WaitOnGTEvent->Wait();
+		TaskSharedData->bWaitingOnGT = false;
+	}
+}
+
+void FNiagaraDerivedDataAsyncRequest::MoveToState(ENiagaraCompilationState NewState)
+{
+	if (TaskSharedData->CurrentState == ENiagaraCompilationState::Finished || TaskSharedData->CurrentState == ENiagaraCompilationState::Aborted)
+	{
+		// we are already in an end state, so nothing to do other than waking up waiting threads
+		if (TaskSharedData->bWaitingOnGT)
+		{
+			TaskSharedData->WaitOnGTEvent->Trigger();
+		}
+		return;
+	}
+
+	// check state machine integrity
+	check(NewState != ENiagaraCompilationState::Precompile);
+	if (NewState == ENiagaraCompilationState::StartCompileJob)
+	{
+		check(TaskSharedData->CurrentState == ENiagaraCompilationState::Precompile);
+	}
+	if (NewState == ENiagaraCompilationState::AwaitResult)
+	{
+		check(TaskSharedData->CurrentState == ENiagaraCompilationState::StartCompileJob);
+	}
+	if (NewState == ENiagaraCompilationState::ProcessResult)
+	{
+		check(TaskSharedData->CurrentState == ENiagaraCompilationState::AwaitResult);
+	}
+	if (NewState == ENiagaraCompilationState::Finished)
+	{
+		check(TaskSharedData->CurrentState == ENiagaraCompilationState::ProcessResult);
+	}
+
+	UE_LOG(LogNiagara, Verbose, TEXT("Changing state %i -> %i for for %s!"), TaskSharedData->CurrentState.load(), NewState, *AssetPath);
+	TaskSharedData->CurrentState = NewState;
+
+	if (TaskSharedData->bWaitingOnGT)
+	{
+		// wake up any waiting ddc thread
+		TaskSharedData->WaitOnGTEvent->Trigger();
+	}
+}
+
+bool FNiagaraDerivedDataAsyncRequest::IsDone() const
+{
+	return TaskSharedData->CurrentState == ENiagaraCompilationState::Finished || TaskSharedData->CurrentState == ENiagaraCompilationState::Aborted;
+}
+
+void LogVMID(const FNiagaraVMExecutableDataId& ID, FString Name)
+{
+	UE_LOG(LogNiagara, Display, TEXT("FNiagaraVMExecutableDataId %s:"), *Name);
+	UE_LOG(LogNiagara, Display, TEXT("  ScriptUsageType: %i"), ID.ScriptUsageType);
+	UE_LOG(LogNiagara, Display, TEXT("  ScriptUsageTypeID: %s"), *ID.ScriptUsageTypeID.ToString());
+	UE_LOG(LogNiagara, Display, TEXT("  CompilerVersionID: %s"), *ID.CompilerVersionID.ToString());
+	UE_LOG(LogNiagara, Display, TEXT("  BaseScriptCompileHash: %s"), *ID.BaseScriptCompileHash.ToString());
+	UE_LOG(LogNiagara, Display, TEXT("  bUsesRapidIterationParams: %i"), ID.bUsesRapidIterationParams);
+	UE_LOG(LogNiagara, Display, TEXT("  AdditionalDefines:"));
+	for (int32 Idx = 0; Idx < ID.AdditionalDefines.Num(); Idx++)
+	{
+		UE_LOG(LogNiagara, Display, TEXT("     %s"), *ID.AdditionalDefines[Idx]);
+	}
+	UE_LOG(LogNiagara, Display, TEXT("  AdditionalVariables:"));
+	for (int32 Idx = 0; Idx < ID.AdditionalVariables.Num(); Idx++)
+	{
+		UE_LOG(LogNiagara, Display, TEXT("     %s / %s"), *ID.AdditionalVariables[Idx].GetName().ToString(), *ID.AdditionalVariables[Idx].GetType().GetName());
+	}
+	UE_LOG(LogNiagara, Display, TEXT("  ReferencedCompileHashes:"));
+	for (int32 HashIndex = 0; HashIndex < ID.ReferencedCompileHashes.Num(); HashIndex++)
+	{
+		UE_LOG(LogNiagara, Display, TEXT("     %s"), *ID.ReferencedCompileHashes[HashIndex].ToString());
+	}
+}
+
+void FNiagaraDerivedDataAsyncRequest::PrecompileData()
+{
+	SCOPE_CYCLE_COUNTER(STAT_Niagara_System_Precompile);
+	UE_LOG(LogNiagara, Verbose, TEXT("Getting precompile data for %s!"), *AssetPath);
+	FNiagaraVMExecutableDataId OldID = GetScriptPair().CompileId;
+	FNiagaraVMExecutableDataId NewID;
+	// we compute the vmID here again to check if the script changed between the start of the async task and now. If it was changed then we need to abort,
+	// because we would save the compiled data under the wrong ddc key 
+	GetScriptPair().CompiledScript->ComputeVMCompilationId(NewID, FGuid());
+	FString CurrentDDCKey = UNiagaraScript::BuildNiagaraDDCKeyString(NewID);
+	if (DDCKey != CurrentDDCKey)
+	{
+		UE_LOG(LogNiagara, Warning, TEXT("Compile ID for %s changed, cannot compute precompile data without corrupting the ddc."), *AssetPath);
+		LogVMID(OldID, "OldID");
+		LogVMID(NewID, "NewID");
+		AbortTask();
+		return;
+	}
+	
+	ComputedPrecompileData = PrecompileReference->GetPrecompileData(GetScriptPair().CompiledScript);
+}
+
+void FNiagaraDerivedDataAsyncRequest::StartCompileJob()
+{
+	UE_LOG(LogNiagara, Verbose, TEXT("Starting compilation for %s!"), *AssetPath);
+	
+	// Fire off the compile job
+	TaskSharedData->StartCompileTime = FPlatformTime::Seconds();
+	if (TaskSharedData->ScriptPair.CompiledScript->RequestExternallyManagedAsyncCompile(ComputedPrecompileData, TaskSharedData->ScriptPair.CompileId, TaskSharedData->ScriptPair.PendingJobID))
+	{
+		UE_LOG(LogNiagara, Verbose, TEXT("Designated compilation ID %i for compilation of %s!"), TaskSharedData->ScriptPair.PendingJobID, *AssetPath);
+	}
+	else
+	{
+		UE_LOG(LogNiagara, Warning, TEXT("For some reason we are reporting that %s is in sync even though AreScriptAndSourceSynchronized returned false!"), *TaskSharedData->ScriptPair.CompiledScript->GetPathName())
+		AbortTask();
+	}
+}
+
+bool FNiagaraDerivedDataAsyncRequest::AwaitResult()
+{
+	check(IsInGameThread());
+	UE_LOG(LogNiagara, Verbose, TEXT("Waiting on compilation to finish for %s!"), *AssetPath);
+	INiagaraModule& NiagaraModule = FModuleManager::Get().LoadModuleChecked<INiagaraModule>(TEXT("Niagara"));
+	TaskSharedData->ExeData = NiagaraModule.GetCompileJobResult(TaskSharedData->ScriptPair.PendingJobID, TaskSharedData->bWaitForCompileJob);
+	return TaskSharedData->ExeData.IsValid();
+}
+
+void FNiagaraDerivedDataAsyncRequest::ProcessResult()
+{
+	check(TaskSharedData->ExeData);
+	UE_LOG(LogNiagara, Verbose, TEXT("Processing compilation results for %s!"), *AssetPath);
+	
+	// convert results to be saved in the ddc
+	TaskSharedData->ExeData->CompileTime = FPlatformTime::Seconds() - TaskSharedData->StartCompileTime;
+	TaskSharedData->ScriptPair.CompiledScript->ExecToBinaryData(TaskSharedData->ScriptPair.CompiledScript, TaskSharedData->DDCOutData, *TaskSharedData->ExeData);
+	UE_LOG(LogNiagara, Verbose, TEXT("Got %i bytes in ddc data for %s"), TaskSharedData->DDCOutData.Num(), *AssetPath);
+}
+
+bool FNiagaraDerivedDataAsyncRequest::Build(TArray<uint8>& OutData)
+{
+	while (!IsDone())
+	{
+		ProcessCurrentState();
+	}
+	
+	if (TaskSharedData->CurrentState == ENiagaraCompilationState::Aborted)
+	{
+		UE_LOG(LogNiagara, Log, TEXT("DDC build task for %s was aborted, data will be discarded."), *AssetPath);
+		return false;
+	}
+	UE_LOG(LogNiagara, Verbose, TEXT("Finished ddc build task for %s in %f seconds! Resulted in %i bytes of ddc data."), *AssetPath, FPlatformTime::Seconds() - TaskSharedData->StartTaskTime, TaskSharedData->DDCOutData.Num());
+	OutData = TaskSharedData->DDCOutData;
+	return OutData.Num() > 0;
+}
+
+void FNiagaraDerivedDataAsyncRequest::WaitAndResolveResult()
+{
+	check(IsInGameThread());
+	UE_LOG(LogNiagara, Verbose, TEXT("Resolving result from game thread for task %s!"), *AssetPath);
+
+	TaskSharedData->bWaitForCompileJob = true;
+	while (!IsDone())
+	{
+		ProcessCurrentState();
+	}
+}
+
+void FNiagaraDerivedDataAsyncRequest::AbortTask()
+{
+	MoveToState(ENiagaraCompilationState::Aborted);
+}
+
+#endif
+
 UNiagaraSystem::UNiagaraSystem(FVTableHelper& Helper)
 	: Super(Helper)
 {
@@ -144,7 +508,7 @@ void UNiagaraSystem::BeginDestroy()
 #if WITH_EDITORONLY_DATA
 	while (ActiveCompilations.Num() > 0)
 	{
-		QueryCompileComplete(true, false, true);
+		KillAllActiveCompilations();
 	}
 #endif
 
@@ -204,7 +568,6 @@ void UNiagaraSystem::HandleVariableRenamed(const FNiagaraVariable& InOldVariable
 	}
 }
 
-
 void UNiagaraSystem::HandleVariableRemoved(const FNiagaraVariable& InOldVariable, bool bUpdateContexts)
 {
 	if (InOldVariable.IsInNameSpace(FNiagaraConstants::UserNamespace))
@@ -330,9 +693,9 @@ void UNiagaraSystem::UpdateSystemAfterLoad()
 	// We remove emitters and scripts on dedicated servers, so skip further work.
 	const bool bIsDedicatedServer = !GIsClient && GIsServer;
 
+	TArray<UNiagaraScript*> AllSystemScripts;
 	if (!GetOutermost()->bIsCookedForEditor && !bIsDedicatedServer)
 	{
-		TArray<UNiagaraScript*> AllSystemScripts;
 		UNiagaraScriptSourceBase* SystemScriptSource;
 		if (SystemSpawnScript == nullptr)
 		{
@@ -361,16 +724,9 @@ void UNiagaraSystem::UpdateSystemAfterLoad()
 		}
 		AllSystemScripts.Add(SystemUpdateScript);
 
-		bool bSystemScriptsAreSynchronized = true;
-		for (UNiagaraScript* SystemScript : AllSystemScripts)
-		{
-			bSystemScriptsAreSynchronized &= SystemScript->AreScriptAndSourceSynchronized();
-		}
-
 		// Synchronize with parameter definitions
 		PostLoadDefinitionsSubscriptions();
 
-		bool bEmitterScriptsAreSynchronized = true;
 #if 0
 		UE_LOG(LogNiagara, Log, TEXT("PreMerger"));
 		for (FNiagaraEmitterHandle& EmitterHandle : EmitterHandles)
@@ -389,30 +745,12 @@ void UNiagaraSystem::UpdateSystemAfterLoad()
 		}
 #endif
 
-		for (FNiagaraEmitterHandle& EmitterHandle : EmitterHandles)
-		{
-			if (EmitterHandle.GetIsEnabled() && EmitterHandle.GetInstance() && !EmitterHandle.GetInstance()->AreAllScriptAndSourcesSynchronized())
-			{
-				bEmitterScriptsAreSynchronized = false;
-			}
-		}
-
 		if (UNiagaraEmitter::GetForceCompileOnLoad())
 		{
 			ForceGraphToRecompileOnNextCheck();
 			UE_LOG(LogNiagara, Log, TEXT("System %s being rebuilt because UNiagaraEmitter::GetForceCompileOnLoad() == true."), *GetPathName());
 		}
 
-		if (bSystemScriptsAreSynchronized == false && GEnableVerboseNiagaraChangeIdLogging)
-		{
-			UE_LOG(LogNiagara, Log, TEXT("System %s being compiled because there were changes to a system script Change ID."), *GetPathName());
-		}
-
-		if (bEmitterScriptsAreSynchronized == false && GEnableVerboseNiagaraChangeIdLogging)
-		{
-			UE_LOG(LogNiagara, Log, TEXT("System %s being compiled because there were changes to an emitter script Change ID."), *GetPathName());
-		}
-
 		if (EmitterCompiledData.Num() == 0 || EmitterCompiledData[0]->DataSetCompiledData.Variables.Num() == 0)
 		{
 			InitEmitterCompiledData();
@@ -439,20 +777,7 @@ void UNiagaraSystem::UpdateSystemAfterLoad()
 			UE_LOG(LogNiagara, Log, TEXT("Update RI Parameters"));
 			UpdateScript->RapidIterationParameters.DumpParameters();
 		}
-#endif
-
-		if (bSystemScriptsAreSynchronized == false || bEmitterScriptsAreSynchronized == false)
-		{
-			if (IsRunningCommandlet())
-			{
-				// Call modify here so that the system will resave the compile ids and script vm when running the resave
-				// commandlet. We don't need it for normal post-loading.
-				Modify();
-			}
-			RequestCompile(false);
-		}
-
-#if 0
+		
 		UE_LOG(LogNiagara, Log, TEXT("After"));
 		for (FNiagaraEmitterHandle& EmitterHandle : EmitterHandles)
 		{
@@ -470,12 +795,6 @@ void UNiagaraSystem::UpdateSystemAfterLoad()
 		}
 #endif
 	}
-	if (GNiagaraForceSystemsToCookOutRapidIterationOnLoad == 1 && !bBakeOutRapidIteration)
-	{
-		WaitForCompilationComplete();
-		bBakeOutRapidIteration = true;
-		RequestCompile(false);
-	}
 #endif
 
 	if ( FPlatformProperties::RequiresCookedData() )
@@ -501,6 +820,53 @@ void UNiagaraSystem::UpdateSystemAfterLoad()
 	{
 		FNiagaraWorldManager::PrimePoolForAllWorlds(this);
 	}
+
+#if WITH_EDITORONLY_DATA
+	// check if the system needs to be compiled and start a compile task if necessary
+	{
+		bool bSystemScriptsAreSynchronized = true;
+		for (UNiagaraScript* SystemScript : AllSystemScripts)
+		{
+			bSystemScriptsAreSynchronized &= SystemScript->AreScriptAndSourceSynchronized();
+		}
+
+		bool bEmitterScriptsAreSynchronized = true;
+		for (FNiagaraEmitterHandle& EmitterHandle : EmitterHandles)
+		{
+			if (EmitterHandle.GetIsEnabled() && EmitterHandle.GetInstance() && !EmitterHandle.GetInstance()->AreAllScriptAndSourcesSynchronized())
+			{
+				bEmitterScriptsAreSynchronized = false;
+			}
+		}
+
+		if (bSystemScriptsAreSynchronized == false && GEnableVerboseNiagaraChangeIdLogging)
+		{
+			UE_LOG(LogNiagara, Log, TEXT("System %s being compiled because there were changes to a system script Change ID."), *GetPathName());
+		}
+
+		if (bEmitterScriptsAreSynchronized == false && GEnableVerboseNiagaraChangeIdLogging)
+		{
+			UE_LOG(LogNiagara, Log, TEXT("System %s being compiled because there were changes to an emitter script Change ID."), *GetPathName());
+		}
+
+		if (GNiagaraForceSystemsToCookOutRapidIterationOnLoad == 1 && !bBakeOutRapidIteration)
+		{
+			bSystemScriptsAreSynchronized = false;
+			bBakeOutRapidIteration = true;
+		}
+		
+		if (bSystemScriptsAreSynchronized == false || bEmitterScriptsAreSynchronized == false)
+		{
+			if (IsRunningCommandlet())
+			{
+				// Call modify here so that the system will resave the compile ids and script vm when running the resave
+				// commandlet. We don't need it for normal post-loading.
+				Modify();
+			}
+			RequestCompile(false);
+		}
+	}
+#endif
 }
 
 #if WITH_EDITORONLY_DATA
@@ -1817,6 +2183,19 @@ const UNiagaraScript* UNiagaraSystem::GetSystemUpdateScript() const
 
 #if WITH_EDITORONLY_DATA
 
+void UNiagaraSystem::KillAllActiveCompilations()
+{
+	InvalidateActiveCompiles();
+	for (FNiagaraSystemCompileRequest& Request : ActiveCompilations)
+	{
+		for (auto& AsyncTask : Request.DDCTasks)
+		{
+			AsyncTask->AbortTask();
+		}
+	}
+	ActiveCompilations.Empty();
+}
+
 bool UNiagaraSystem::GetIsolateEnabled() const
 {
 	return bIsolateEnabled;
@@ -1896,8 +2275,6 @@ void UNiagaraSystem::WaitForCompilationComplete(bool bIncludingGPUShaders, bool
 			Progress.EnterProgressFrame();
 		}
 	
-
-#if WITH_EDITORONLY_DATA
 		if (GIsAutomationTesting)
 		{
 			double CurrentTime = FPlatformTime::Seconds();
@@ -1915,7 +2292,6 @@ void UNiagaraSystem::WaitForCompilationComplete(bool bIncludingGPUShaders, bool
 				break;
 			}
 		}
-#endif
 	}
 	
 	for (FNiagaraShaderScript* ShaderScript : GPUScripts)
@@ -1927,9 +2303,9 @@ void UNiagaraSystem::WaitForCompilationComplete(bool bIncludingGPUShaders, bool
 
 void UNiagaraSystem::InvalidateActiveCompiles()
 {
-	for (FNiagaraSystemCompileRequest& ActiveCompilation : ActiveCompilations)
+	for (FNiagaraSystemCompileRequest& CompileRequest : ActiveCompilations)
 	{
-		ActiveCompilation.bIsValid = false;
+		CompileRequest.bIsValid = false;
 	}
 }
 
@@ -1961,22 +2337,22 @@ bool UNiagaraSystem::CompilationResultsValid(FNiagaraSystemCompileRequest& Compi
 {
 	// for now the only thing we're concerned about is if we've got results for SystemSpawn and SystemUpdate scripts
 	// then we need to make sure that they agree in terms of the dataset attributes
-	const FEmitterCompiledScriptPair* SpawnScriptRequest =
-		Algo::FindBy(CompileRequest.EmitterCompiledScriptPairs, SystemSpawnScript, &FEmitterCompiledScriptPair::CompiledScript);
-	const FEmitterCompiledScriptPair* UpdateScriptRequest =
-		Algo::FindBy(CompileRequest.EmitterCompiledScriptPairs, SystemUpdateScript, &FEmitterCompiledScriptPair::CompiledScript);
+	TSharedPtr<FNiagaraDerivedDataAsyncRequest, ESPMode::ThreadSafe>* SpawnScriptRequest = CompileRequest.DDCTasks.FindByPredicate([this](const auto& Task) { return Task->GetScriptPair().CompiledScript == SystemSpawnScript; });
+	TSharedPtr<FNiagaraDerivedDataAsyncRequest, ESPMode::ThreadSafe>* UpdateScriptRequest = CompileRequest.DDCTasks.FindByPredicate([this](const auto& Task) { return Task->GetScriptPair().CompiledScript == SystemUpdateScript; });
 
 	const bool SpawnScriptValid = SpawnScriptRequest
-		&& SpawnScriptRequest->CompileResults.IsValid()
-		&& SpawnScriptRequest->CompileResults->LastCompileStatus != ENiagaraScriptCompileStatus::NCS_Error;
+		&& SpawnScriptRequest->Get()->GetScriptPair().CompileResults.IsValid()
+		&& SpawnScriptRequest->Get()->GetScriptPair().CompileResults->LastCompileStatus != ENiagaraScriptCompileStatus::NCS_Error;
 
 	const bool UpdateScriptValid = UpdateScriptRequest
-		&& UpdateScriptRequest->CompileResults.IsValid()
-		&& UpdateScriptRequest->CompileResults->LastCompileStatus != ENiagaraScriptCompileStatus::NCS_Error;
+		&& UpdateScriptRequest->Get()->GetScriptPair().CompileResults.IsValid()
+		&& UpdateScriptRequest->Get()->GetScriptPair().CompileResults->LastCompileStatus != ENiagaraScriptCompileStatus::NCS_Error;
 
 	if (SpawnScriptValid && UpdateScriptValid)
 	{
-		if (SpawnScriptRequest->CompileResults->Attributes != UpdateScriptRequest->CompileResults->Attributes)
+		const FEmitterCompiledScriptPair& SpawnScriptPair = SpawnScriptRequest->Get()->GetScriptPair();
+		const FEmitterCompiledScriptPair& UpdateScriptPair = UpdateScriptRequest->Get()->GetScriptPair();
+		if (SpawnScriptPair.CompileResults->Attributes != UpdateScriptPair.CompileResults->Attributes)
 		{
 			// if we had requested a full rebuild, then we've got a case where the generated scripts are not compatible.  This indicates
 			// a significant issue where we're allowing graphs to generate invalid collections of scripts.  One known example is using
@@ -1987,17 +2363,17 @@ bool UNiagaraSystem::CompilationResultsValid(FNiagaraSystemCompileRequest& Compi
 				FString MissingAttributes;
 				FString AdditionalAttributes;
 
-				for (const auto& SpawnAttrib : SpawnScriptRequest->CompileResults->Attributes)
+				for (const auto& SpawnAttrib : SpawnScriptPair.CompileResults->Attributes)
 				{
-					if (!UpdateScriptRequest->CompileResults->Attributes.Contains(SpawnAttrib))
+					if (!UpdateScriptPair.CompileResults->Attributes.Contains(SpawnAttrib))
 					{
 						MissingAttributes.Appendf(TEXT("%s%s"), MissingAttributes.Len() ? TEXT(", ") : TEXT(""), *SpawnAttrib.GetName().ToString());
 					}
 				}
 
-				for (const auto& UpdateAttrib : UpdateScriptRequest->CompileResults->Attributes)
+				for (const auto& UpdateAttrib : UpdateScriptPair.CompileResults->Attributes)
 				{
-					if (!SpawnScriptRequest->CompileResults->Attributes.Contains(UpdateAttrib))
+					if (!SpawnScriptPair.CompileResults->Attributes.Contains(UpdateAttrib))
 					{
 						AdditionalAttributes.Appendf(TEXT("%s%s"), AdditionalAttributes.Len() ? TEXT(", ") : TEXT(""), *UpdateAttrib.GetName().ToString());
 					}
@@ -2010,8 +2386,8 @@ bool UNiagaraSystem::CompilationResultsValid(FNiagaraSystemCompileRequest& Compi
 						FText::FromString(AdditionalAttributes))
 					.ToString());
 
-				SpawnScriptRequest->CompileResults->LastCompileStatus = ENiagaraScriptCompileStatus::NCS_Error;
-				SpawnScriptRequest->CompileResults->LastCompileEvents.Add(AttributeMismatchEvent);
+				SpawnScriptPair.CompileResults->LastCompileStatus = ENiagaraScriptCompileStatus::NCS_Error;
+				SpawnScriptPair.CompileResults->LastCompileEvents.Add(AttributeMismatchEvent);
 			}
 			else
 			{
@@ -2023,8 +2399,9 @@ bool UNiagaraSystem::CompilationResultsValid(FNiagaraSystemCompileRequest& Compi
 	}
 
 	// Now iterate over all dependencies and verify that they are met. If not, emit an error.
-	for (FEmitterCompiledScriptPair& CompilePair : CompileRequest.EmitterCompiledScriptPairs)
+	for (auto& AsyncTask : CompileRequest.DDCTasks)
 	{
+		const FEmitterCompiledScriptPair& CompilePair = AsyncTask->GetScriptPair();
 		bool bValid = CompilePair.CompileResults.IsValid()
 			&& CompilePair.CompileResults->LastCompileStatus != ENiagaraScriptCompileStatus::NCS_Error;
 
@@ -2044,9 +2421,9 @@ bool UNiagaraSystem::CompilationResultsValid(FNiagaraSystemCompileRequest& Compi
 				int32 TestIdx = CompilePair.ParentIndex;
 				while (TestIdx != INDEX_NONE && bDependencyMet == false)
 				{
-					if (CompileRequest.EmitterCompiledScriptPairs.IsValidIndex(TestIdx))
+					if (CompileRequest.DDCTasks.IsValidIndex(TestIdx))
 					{
-						const FEmitterCompiledScriptPair& TestPair = CompileRequest.EmitterCompiledScriptPairs[TestIdx];
+						const FEmitterCompiledScriptPair& TestPair = CompileRequest.DDCTasks[TestIdx]->GetScriptPair();
 						if (TestPair.CompileResults.IsValid() && TestPair.CompileResults->AttributesWritten.Num() > 0)
 						{
 							if (TestPair.CompileResults->AttributesWritten.Contains(TestVar))
@@ -2069,32 +2446,93 @@ bool UNiagaraSystem::CompilationResultsValid(FNiagaraSystemCompileRequest& Compi
 		}
 	}
 	
-	
-
 	return true;
 }
 
-bool UNiagaraSystem::QueryCompileComplete(bool bWait, bool bDoPost, bool bDoNotApply)
+void UNiagaraSystem::ProcessWaitingDDCTasks(bool bProcessForWait)
 {
+	for (FNiagaraSystemCompileRequest& CompileRequest : ActiveCompilations)
+	{
+		for (auto& AsyncTask : CompileRequest.DDCTasks)
+		{
+			if (bProcessForWait)
+			{
+				// before we start to wait for the compile results, we start the compilation of all remaining tasks
+				while (!AsyncTask->bResultProcessed && !AsyncTask->IsDone() && AsyncTask->TaskSharedData->CurrentState < ENiagaraCompilationState::AwaitResult)
+				{
+					AsyncTask->ProcessCurrentState();
+				}
+			}
+			else if (AsyncTask->TaskSharedData->bWaitingOnGT)
+			{
+				ensure(!AsyncTask->bResultProcessed && !AsyncTask->IsDone());
+				AsyncTask->ProcessCurrentState();
+			}
+		}
+	}
+}
 
+bool UNiagaraSystem::QueryCompileComplete(bool bWait, bool bDoPost, bool bDoNotApply)
+{
 	bool bCompileGuardInProgress = InternalCompileGuardCheck(this);
 
 	if (ActiveCompilations.Num() > 0 && !bCompileGuardInProgress)
 	{
-		int32 ActiveCompileIdx = 0;
-
+		ProcessWaitingDDCTasks(bWait);
+		
+		FDerivedDataCacheInterface& DerivedDataCache = GetDerivedDataCacheRef();
+		FNiagaraSystemCompileRequest& CompileRequest = ActiveCompilations[0];
 		bool bAreWeWaitingForAnyResults = false;
-
-		// Check to see if ALL of the sub-requests have resolved. 
-		for (FEmitterCompiledScriptPair& EmitterCompiledScriptPair : ActiveCompilations[ActiveCompileIdx].EmitterCompiledScriptPairs)
+		
+		// Check to see if ALL of the sub-requests have resolved.
+		for (auto& AsyncTask : CompileRequest.DDCTasks)
 		{
-			if ((uint32)INDEX_NONE == EmitterCompiledScriptPair.PendingJobID || EmitterCompiledScriptPair.bResultsReady)
+			if (AsyncTask->bResultProcessed)
 			{
 				continue;
 			}
-			EmitterCompiledScriptPair.bResultsReady = ProcessCompilationResult(EmitterCompiledScriptPair, bWait, bDoNotApply);
-			if (!EmitterCompiledScriptPair.bResultsReady)
+
+			// if we went through the ddc, the result was either fetched or compiled via the ddc task
+			COOK_STAT(auto Timer = NiagaraScriptCookStats::UsageStats.TimeAsyncWait());
+			uint32 TaskHandle = AsyncTask->TaskHandle;
+			if (bWait)
+			{
+				AsyncTask->WaitAndResolveResult();
+				DerivedDataCache.WaitAsynchronousCompletion(TaskHandle);
+			}
+			
+			if (DerivedDataCache.PollAsynchronousCompletion(TaskHandle))
+			{
+				AsyncTask->bResultProcessed = true;
+				bool bDataWasBuilt;
+				TArray<uint8> Data;
+				bool bGetSuccessful = GetDerivedDataCacheRef().GetAsynchronousResults(TaskHandle, Data, &bDataWasBuilt);
+				if (!bGetSuccessful)
+				{
+					UE_LOG(LogNiagara, Error, TEXT("Error getting compiled data for asset %s from ddc"), *AsyncTask->AssetPath);
+					return false;
+				}
+				if (bDataWasBuilt)
+				{
+					CompileRequest.bHasCompiledJobs = true;
+				}
+				else
+				{
+					UE_LOG(LogNiagara, Verbose, TEXT("Fetched data for %s from ddc"), *AsyncTask->AssetPath);
+				}
+	
+				TSharedPtr<FNiagaraVMExecutableData> ExeData = MakeShared<FNiagaraVMExecutableData>();
+				if (AsyncTask->GetScriptPair().CompiledScript->BinaryToExecData(AsyncTask->GetScriptPair().CompiledScript, Data, *ExeData))
+				{
+					AsyncTask->TaskSharedData->ScriptPair.CompileResults = ExeData;
+					AsyncTask->TaskSharedData->ScriptPair.bResultsReady = true;
+					CompileRequest.CombinedCompileTime += ExeData->CompileTime;
+				}
+				COOK_STAT(Timer.AddHitOrMiss(bDataWasBuilt ? FCookStats::CallStats::EHitOrMiss::Miss : FCookStats::CallStats::EHitOrMiss::Hit, Data.Num()));
+			}
+			else
 			{
+				COOK_STAT(Timer.TrackCyclesOnly());
 				bAreWeWaitingForAnyResults = true;
 			}
 		}
@@ -2113,50 +2551,61 @@ bool UNiagaraSystem::QueryCompileComplete(bool bWait, bool bDoPost, bool bDoNotA
 		{
 			// if we've gotten all the results, run a quick check to see if the data is valid, if it's not then that indicates that
 			// we've run into a compatibility issue and so we should see if we should issue a full rebuild
-			const bool ResultsValid = CompilationResultsValid(ActiveCompilations[ActiveCompileIdx]);
-			if (!ResultsValid && !ActiveCompilations[ActiveCompileIdx].bForced)
+			const bool ResultsValid = CompilationResultsValid(CompileRequest);
+			if (!ResultsValid && !CompileRequest.bForced)
 			{
-				ActiveCompilations[ActiveCompileIdx].RootObjects.Empty();
-				ActiveCompilations.RemoveAt(ActiveCompileIdx);
+				CompileRequest.RootObjects.Empty();
+				ActiveCompilations.RemoveAt(0);
 				RequestCompile(true, nullptr);
 				return false;
 			}
 		}
 
 		// In the world of do not apply, we're exiting the system completely so let's just kill any active compilations altogether.
-		if (bDoNotApply || ActiveCompilations[ActiveCompileIdx].bIsValid == false)
+		if (bDoNotApply || CompileRequest.bIsValid == false)
 		{
-			ActiveCompilations[ActiveCompileIdx].RootObjects.Empty();
-			ActiveCompilations.RemoveAt(ActiveCompileIdx);
+			CompileRequest.RootObjects.Empty();
+			ActiveCompilations.RemoveAt(0);
 			return true;
 		}
 
-
-		SCOPE_CYCLE_COUNTER(STAT_Niagara_System_CompileScript);
-
 		// Now that the above code says they are all complete, go ahead and resolve them all at once.
-		float CombinedCompileTime = 0.0f;
-		bool HasCompiledJobs = false;
-		for (FEmitterCompiledScriptPair& EmitterCompiledScriptPair : ActiveCompilations[ActiveCompileIdx].EmitterCompiledScriptPairs)
+		for (auto& AsyncTask : CompileRequest.DDCTasks)
 		{
-			if ((uint32)INDEX_NONE == EmitterCompiledScriptPair.PendingJobID)
+			FEmitterCompiledScriptPair& EmitterCompiledScriptPair = AsyncTask->TaskSharedData->ScriptPair;
+			if (EmitterCompiledScriptPair.bResultsReady)
 			{
-				if (!EmitterCompiledScriptPair.bResultsReady)
+				TSharedPtr<FNiagaraVMExecutableData> ExeData = EmitterCompiledScriptPair.CompileResults;
+				TMap<FName, UNiagaraDataInterface*> ObjectNameMap;
+				if (AsyncTask->ComputedPrecompileData.IsValid())
 				{
-					continue;
+					ObjectNameMap = AsyncTask->ComputedPrecompileData->GetObjectNameMap(); 
 				}
+				else
+				{
+					ObjectNameMap = EmitterCompiledScriptPair.CompiledScript->GetLatestSource()->ComputeObjectNameMap(AsyncTask->UniqueEmitterName);
+				}
+				EmitterCompiledScriptPair.CompiledScript->SetVMCompilationResults(EmitterCompiledScriptPair.CompileId, *ExeData, AsyncTask->UniqueEmitterName, ObjectNameMap);
 			}
-			else
+			// clean up the precompile data
+			if (AsyncTask->ComputedPrecompileData.IsValid())
 			{
-				HasCompiledJobs = true;
+				// Synchronize the variables that we actually encountered during precompile so that we can expose them to the end user.
+				TArray<FNiagaraVariable> OriginalExposedParams;
+				GetExposedParameters().GetParameters(OriginalExposedParams);
+				TArray<FNiagaraVariable>& EncounteredExposedVars = AsyncTask->PrecompileReference->EncounteredExposedVars;
+				for (int32 i = 0; i < EncounteredExposedVars.Num(); i++)
+				{
+					if (OriginalExposedParams.Contains(EncounteredExposedVars[i]) == false)
+					{
+						// Just in case it wasn't added previously..
+						ExposedParameters.AddParameter(EncounteredExposedVars[i]);
+					}
+				}
+				
+				AsyncTask->ComputedPrecompileData->ReleaseCompilationCopies();
+				AsyncTask->ComputedPrecompileData.Reset();
 			}
-
-			CombinedCompileTime += EmitterCompiledScriptPair.CompileResults->CompileTime;
-			check(EmitterCompiledScriptPair.bResultsReady);
-
-			TSharedPtr<FNiagaraVMExecutableData> ExeData = EmitterCompiledScriptPair.CompileResults;
-			TSharedPtr<FNiagaraCompileRequestDataBase, ESPMode::ThreadSafe> PrecompData = ActiveCompilations[ActiveCompileIdx].MappedData.FindChecked(EmitterCompiledScriptPair.CompiledScript);
-			EmitterCompiledScriptPair.CompiledScript->SetVMCompilationResults(EmitterCompiledScriptPair.CompileId, *ExeData, PrecompData.Get());
 		}
 
 		if (bDoPost)
@@ -2184,8 +2633,9 @@ bool UNiagaraSystem::QueryCompileComplete(bool bWait, bool bDoPost, bool bDoNotA
 		// particle update script aren't being shared by the interpolated spawn script when accessed directly.  This works
 		// properly if the data interface is assigned to a named particle parameter and then linked to an input.
 		// TODO: Bind these data interfaces the same way parameter data interfaces are bound.
-		for (FEmitterCompiledScriptPair& EmitterCompiledScriptPair : ActiveCompilations[ActiveCompileIdx].EmitterCompiledScriptPairs)
+		for (auto& AsyncTask : CompileRequest.DDCTasks)
 		{
+			FEmitterCompiledScriptPair& EmitterCompiledScriptPair = AsyncTask->TaskSharedData->ScriptPair;
 			UNiagaraEmitter* Emitter = EmitterCompiledScriptPair.Emitter;
 			UNiagaraScript* CompiledScript = EmitterCompiledScriptPair.CompiledScript;
 
@@ -2211,7 +2661,7 @@ bool UNiagaraSystem::QueryCompileComplete(bool bWait, bool bDoPost, bool bDoNotA
 			}
 		}
 
-		ActiveCompilations[ActiveCompileIdx].RootObjects.Empty();
+		CompileRequest.RootObjects.Empty();
 
 		UpdatePostCompileDIInfo();
 
@@ -2226,19 +2676,19 @@ bool UNiagaraSystem::QueryCompileComplete(bool bWait, bool bDoPost, bool bDoNotA
 
 		ResolveScalabilitySettings();
 
-		const float ElapsedWallTime = (float)(FPlatformTime::Seconds() - ActiveCompilations[ActiveCompileIdx].StartTime);
+		const float ElapsedWallTime = (float)(FPlatformTime::Seconds() - CompileRequest.StartTime);
 
-		if (HasCompiledJobs)
+		if (CompileRequest.bHasCompiledJobs)
 		{
 			UE_LOG(LogNiagara, Log, TEXT("Compiling System %s took %f sec (time since issued), %f sec (combined shader worker time)."),
-				*GetFullName(), ElapsedWallTime, CombinedCompileTime);
+				*GetFullName(), ElapsedWallTime, CompileRequest.CombinedCompileTime);
 		}
 		else
 		{
-			UE_LOG(LogNiagara, Verbose, TEXT("Compiling System %s took %f sec."), *GetFullName(), ElapsedWallTime);
+			UE_LOG(LogNiagara, Verbose, TEXT("Compiling System %s took %f sec, no shader worker used."), *GetFullName(), ElapsedWallTime);
 		}
 
-		ActiveCompilations.RemoveAt(ActiveCompileIdx);
+		ActiveCompilations.RemoveAt(0);
 
 		if (bDoPost)
 		{
@@ -2253,84 +2703,6 @@ bool UNiagaraSystem::QueryCompileComplete(bool bWait, bool bDoPost, bool bDoNotA
 	return false;
 }
 
-bool UNiagaraSystem::ProcessCompilationResult(FEmitterCompiledScriptPair& ScriptPair, bool bWait, bool bDoNotApply)
-{
-	COOK_STAT(auto Timer = NiagaraScriptCookStats::UsageStats.TimeAsyncWait());
-
-	INiagaraModule& NiagaraModule = FModuleManager::Get().LoadModuleChecked<INiagaraModule>(TEXT("Niagara"));
-	TSharedPtr<FNiagaraVMExecutableData> ExeData = NiagaraModule.GetCompileJobResult(ScriptPair.PendingJobID, bWait);
-
-	if (!bWait && !ExeData.IsValid())
-	{
-		COOK_STAT(Timer.TrackCyclesOnly());
-		return false;
-	}
-	check(ExeData.IsValid());
-	if (!bDoNotApply)
-	{
-		ScriptPair.CompileResults = ExeData;
-	}
-
-	// save result to the ddc
-	TArray<uint8> OutData;
-	if (UNiagaraScript::ExecToBinaryData(ScriptPair.CompiledScript, OutData, *ExeData))
-	{
-		COOK_STAT(Timer.AddMiss(OutData.Num()));
-
-		// be sure to use the CompileId that is associated with the compilation
-		const FString DDCKey = UNiagaraScript::BuildNiagaraDDCKeyString(ScriptPair.CompileId);
-
-		GetDerivedDataCacheRef().Put(*DDCKey, OutData, GetPathName());
-		return true;
-	}
-
-	COOK_STAT(Timer.TrackCyclesOnly());
-	return false;
-}
-
-bool UNiagaraSystem::GetFromDDC(FEmitterCompiledScriptPair& ScriptPair)
-{
-	if (!ScriptPair.CompiledScript->IsCompilable())
-	{
-		return false;
-	}
-
-	COOK_STAT(auto Timer = NiagaraScriptCookStats::UsageStats.TimeSyncWork());
-
-	FNiagaraVMExecutableDataId NewID;
-	ScriptPair.CompiledScript->ComputeVMCompilationId(NewID, FGuid());
-	ScriptPair.CompileId = NewID;
-
-	TArray<uint8> Data;
-	if (GetDerivedDataCacheRef().GetSynchronous(*ScriptPair.CompiledScript->GetNiagaraDDCKeyString(FGuid()), Data, GetPathName()))
-	{
-		TSharedPtr<FNiagaraVMExecutableData> ExeData = MakeShared<FNiagaraVMExecutableData>();
-		if (ScriptPair.CompiledScript->BinaryToExecData(ScriptPair.CompiledScript, Data, *ExeData))
-		{
-			COOK_STAT(Timer.AddHit(Data.Num()));
-			ExeData->CompileTime = 0; // since we didn't actually compile anything
-			ScriptPair.CompileResults = ExeData;
-			ScriptPair.bResultsReady = true;
-			if (GNiagaraLogDDCStatusForSystems != 0)
-			{
-				UE_LOG(LogNiagara, Verbose, TEXT("Niagara Script pulled from DDC ... %s"), *ScriptPair.CompiledScript->GetPathName());
-			}
-			return true;
-		}
-	}
-	
-	if (GNiagaraLogDDCStatusForSystems != 0)
-	{
-	    UE_LOG(LogNiagara, Verbose, TEXT("Need Compile! Niagara Script GotFromDDC could not find ... %s"), *ScriptPair.CompiledScript->GetPathName());
-	}
-
-	COOK_STAT(Timer.TrackCyclesOnly());
-	return false;
-}
-
-
-#if WITH_EDITORONLY_DATA
-
 void UNiagaraSystem::InitEmitterVariableAliasNames(FNiagaraEmitterCompiledData& EmitterCompiledDataToInit, const UNiagaraEmitter* InAssociatedEmitter)
 {
 	EmitterCompiledDataToInit.EmitterSpawnIntervalVar.SetName(GetEmitterVariableAliasName(SYS_PARAM_EMITTER_SPAWN_INTERVAL, InAssociatedEmitter));
@@ -2371,7 +2743,83 @@ void UNiagaraSystem::InitEmitterDataSetCompiledData(FNiagaraDataSetCompiledData&
 
 	DataSetToInit.BuildLayout();
 }
-#endif
+
+void UNiagaraSystem::PrepareRapidIterationParametersForCompilation()
+{
+	// prepare rapid iteration parameters for compilation
+	TArray<UNiagaraScript*> Scripts;
+	TMap<UNiagaraScript*, const UNiagaraEmitter*> ScriptToEmitterMap;
+	for (int32 i = 0; i < GetEmitterHandles().Num(); i++)
+	{
+		const FNiagaraEmitterHandle& Handle = GetEmitterHandle(i);
+		if (Handle.GetIsEnabled()) // Don't pull in the emitter if it isn't going to be used.
+		{
+			TArray<UNiagaraScript*> EmitterScripts;
+			Handle.GetInstance()->GetScripts(EmitterScripts, false, true);
+
+			for (int32 ScriptIdx = 0; ScriptIdx < EmitterScripts.Num(); ScriptIdx++)
+			{
+				if (EmitterScripts[ScriptIdx])
+				{
+					Scripts.AddUnique(EmitterScripts[ScriptIdx]);
+					ScriptToEmitterMap.Add(EmitterScripts[ScriptIdx], Handle.GetInstance());
+				}
+			}
+		}
+	}
+    		
+	TMap<UNiagaraScript*, UNiagaraScript*> ScriptDependencyMap;
+	for (auto ScriptEmitterPair = ScriptToEmitterMap.CreateIterator(); ScriptEmitterPair; ++ScriptEmitterPair)
+	{
+		UNiagaraScript* CompiledScript = ScriptEmitterPair->Key;
+		const UNiagaraEmitter* Emitter = ScriptEmitterPair->Value;
+
+		if (UNiagaraScript::IsEquivalentUsage(CompiledScript->GetUsage(), ENiagaraScriptUsage::EmitterSpawnScript))
+		{
+			Scripts.AddUnique(SystemSpawnScript);
+			ScriptDependencyMap.Add(CompiledScript, SystemSpawnScript);
+			ScriptToEmitterMap.Add(SystemSpawnScript, nullptr);
+		}
+
+		if (UNiagaraScript::IsEquivalentUsage(CompiledScript->GetUsage(), ENiagaraScriptUsage::EmitterUpdateScript))
+		{
+			Scripts.AddUnique(SystemUpdateScript);
+			ScriptDependencyMap.Add(CompiledScript, SystemUpdateScript);
+			ScriptToEmitterMap.Add(SystemUpdateScript, nullptr);
+		}
+
+		if (UNiagaraScript::IsEquivalentUsage(CompiledScript->GetUsage(), ENiagaraScriptUsage::ParticleSpawnScript))
+		{
+			if (Emitter && Emitter->SimTarget == ENiagaraSimTarget::GPUComputeSim)
+			{
+				UNiagaraScript* ComputeScript = const_cast<UNiagaraScript*>(Emitter->GetGPUComputeScript());
+
+				Scripts.AddUnique(ComputeScript);
+				ScriptDependencyMap.Add(CompiledScript, ComputeScript);
+				ScriptToEmitterMap.Add(ComputeScript, Emitter);
+			}
+		}
+
+		if (UNiagaraScript::IsEquivalentUsage(CompiledScript->GetUsage(), ENiagaraScriptUsage::ParticleUpdateScript))
+		{
+			if (Emitter && Emitter->SimTarget == ENiagaraSimTarget::GPUComputeSim)
+			{
+				UNiagaraScript* ComputeScript = const_cast<UNiagaraScript*>(Emitter->GetGPUComputeScript());
+
+				Scripts.AddUnique(ComputeScript);
+				ScriptDependencyMap.Add(CompiledScript, ComputeScript);
+				ScriptToEmitterMap.Add(ComputeScript, Emitter);
+			}
+			else if (Emitter && Emitter->bInterpolatedSpawning)
+			{
+				Scripts.AddUnique(Emitter->SpawnScriptProps.Script);
+				ScriptDependencyMap.Add(CompiledScript, Emitter->SpawnScriptProps.Script);
+				ScriptToEmitterMap.Add(Emitter->SpawnScriptProps.Script, Emitter);
+			}
+		}
+	}
+	FNiagaraUtilities::PrepareRapidIterationParameters(Scripts, ScriptDependencyMap, ScriptToEmitterMap);
+}
 
 bool UNiagaraSystem::RequestCompile(bool bForce, FNiagaraSystemUpdateContext* OptionalUpdateContext)
 {
@@ -2408,17 +2856,15 @@ bool UNiagaraSystem::RequestCompile(bool bForce, FNiagaraSystemUpdateContext* Op
 	
 
 	// Record that we entered this function already.
+	SCOPE_CYCLE_COUNTER(STAT_Niagara_System_CompileScript);
 	FPlatformTLS::SetTlsValue(CompileGuardSlot, this);
 
 	FNiagaraSystemCompileRequest& ActiveCompilation = ActiveCompilations.AddDefaulted_GetRef();
 	ActiveCompilation.bForced = bForce;
 	ActiveCompilation.StartTime = FPlatformTime::Seconds();
 
-	SCOPE_CYCLE_COUNTER(STAT_Niagara_System_Precompile);
-	
 	check(SystemSpawnScript->GetLatestSource() == SystemUpdateScript->GetLatestSource());
-	TArray<FNiagaraVariable> OriginalExposedParams;
-	GetExposedParameters().GetParameters(OriginalExposedParams);
+	PrepareRapidIterationParametersForCompilation();
 
 	TArray<UNiagaraScript*> ScriptsNeedingCompile;
 	bool bAnyCompiled = false;
@@ -2427,13 +2873,11 @@ bool UNiagaraSystem::RequestCompile(bool bForce, FNiagaraSystemUpdateContext* Op
 		COOK_STAT(Timer.TrackCyclesOnly());
 		INiagaraModule& NiagaraModule = FModuleManager::Get().LoadModuleChecked<INiagaraModule>(TEXT("Niagara"));
 
-
 		//Compile all emitters
 		bool bAnyUnsynchronized = false;	
 
 		// Pass one... determine if any need to be compiled.
 		{
-
 			for (int32 i = 0; i < EmitterHandles.Num(); i++)
 			{
 				FNiagaraEmitterHandle Handle = EmitterHandles[i];
@@ -2445,20 +2889,29 @@ bool UNiagaraSystem::RequestCompile(bool bForce, FNiagaraSystemUpdateContext* Op
 					int32 Parent = INDEX_NONE;
 					for (UNiagaraScript* EmitterScript : EmitterScripts)
 					{
-
 						FEmitterCompiledScriptPair Pair;
-						Pair.bResultsReady = false;
 						Pair.Emitter = Handle.GetInstance();
 						Pair.CompiledScript = EmitterScript;
 						Pair.ParentIndex = Parent;
-						if (!GetFromDDC(Pair) && EmitterScript->IsCompilable() && !EmitterScript->AreScriptAndSourceSynchronized())
+
+						FNiagaraVMExecutableDataId NewID;
+						// we need to compute the vmID here to check later in the ddc task before doing the precompile if anything has changed in the meantime
+						Pair.CompiledScript->ComputeVMCompilationId(NewID, FGuid());
+						Pair.CompileId = NewID;
+
+						TSharedPtr<FNiagaraDerivedDataAsyncRequest, ESPMode::ThreadSafe> AsyncTask = MakeShared<FNiagaraDerivedDataAsyncRequest, ESPMode::ThreadSafe>(EmitterScript->GetPathName(), Pair);
+						if (EmitterScript->IsCompilable() && !EmitterScript->AreScriptAndSourceSynchronized())
 						{
 							ScriptsNeedingCompile.Add(EmitterScript);
 							bAnyUnsynchronized = true;
 						}
-						Parent = ActiveCompilation.EmitterCompiledScriptPairs.Add(Pair);
+						else
+						{
+							AsyncTask->bResultProcessed = true;
+						}
+						AsyncTask->UniqueEmitterName = Handle.GetInstance()->GetUniqueEmitterName();
+						Parent = ActiveCompilation.DDCTasks.Add(AsyncTask);
 					}
-
 				}
 			}
 
@@ -2468,34 +2921,53 @@ bool UNiagaraSystem::RequestCompile(bool bForce, FNiagaraSystemUpdateContext* Op
 			int32 Parent = INDEX_NONE;
 			{
 				FEmitterCompiledScriptPair Pair;
-				Pair.bResultsReady = false;
-				Pair.Emitter = nullptr;
 				Pair.CompiledScript = SystemSpawnScript;
-				if (!GetFromDDC(Pair) && !SystemSpawnScript->AreScriptAndSourceSynchronized())
+				
+				FNiagaraVMExecutableDataId NewID;
+				// we need to compute the vmID here to check later in the ddc task before doing the precompile if anything has changed in the meantime
+				Pair.CompiledScript->ComputeVMCompilationId(NewID, FGuid());
+				Pair.CompileId = NewID;
+
+				TSharedPtr<FNiagaraDerivedDataAsyncRequest, ESPMode::ThreadSafe> AsyncTask = MakeShared<FNiagaraDerivedDataAsyncRequest, ESPMode::ThreadSafe>(SystemSpawnScript->GetPathName(), Pair);
+				if (!SystemSpawnScript->AreScriptAndSourceSynchronized())
 				{
 					ScriptsNeedingCompile.Add(SystemSpawnScript);
 					bAnyCompiled = true;
 				}
-				Parent = ActiveCompilation.EmitterCompiledScriptPairs.Add(Pair);
+				else
+				{
+					AsyncTask->bResultProcessed = true;
+				}
+				Parent = ActiveCompilation.DDCTasks.Add(AsyncTask);
 			}
 
 			{
 				FEmitterCompiledScriptPair Pair;
-				Pair.bResultsReady = false;
-				Pair.Emitter = nullptr;
 				Pair.CompiledScript = SystemUpdateScript;
 				Pair.ParentIndex = Parent;
-				if (!GetFromDDC(Pair) && !SystemUpdateScript->AreScriptAndSourceSynchronized())
+				
+				FNiagaraVMExecutableDataId NewID;
+				// we need to compute the vmID here to check later in the ddc task before doing the precompile if anything has changed in the meantime
+				Pair.CompiledScript->ComputeVMCompilationId(NewID, FGuid());
+				Pair.CompileId = NewID;
+
+				TSharedPtr<FNiagaraDerivedDataAsyncRequest, ESPMode::ThreadSafe> AsyncTask = MakeShared<FNiagaraDerivedDataAsyncRequest, ESPMode::ThreadSafe>(SystemUpdateScript->GetPathName(), Pair);
+				if (!SystemUpdateScript->AreScriptAndSourceSynchronized())
 				{
 					ScriptsNeedingCompile.Add(SystemUpdateScript);
 					bAnyCompiled = true;
 				}
-				Parent = ActiveCompilation.EmitterCompiledScriptPairs.Add(Pair);
+				else
+				{
+					AsyncTask->bResultProcessed = true;
+				}
+				Parent = ActiveCompilation.DDCTasks.Add(AsyncTask);
 			}
 
 			// Need to set the EmitterParent on the emitter spawn scripts
-			for (FEmitterCompiledScriptPair & Pair: ActiveCompilation.EmitterCompiledScriptPairs)
+			for (auto & AsyncTask: ActiveCompilation.DDCTasks)
 			{
+				FEmitterCompiledScriptPair& Pair = AsyncTask->TaskSharedData->ScriptPair;
 				if (Pair.Emitter != nullptr && Pair.ParentIndex == INDEX_NONE)
 				{
 					Pair.ParentIndex = Parent;
@@ -2504,94 +2976,48 @@ bool UNiagaraSystem::RequestCompile(bool bForce, FNiagaraSystemUpdateContext* Op
 		}
 
 
-		// We found things needing compilation, now we have to go through an static duplicate everything that will be translated...
+		// prepare data for any precompile the ddc tasks need to do
+		TSharedPtr<FNiagaraLazyPrecompileReference, ESPMode::ThreadSafe> PrecompileReference = MakeShared<FNiagaraLazyPrecompileReference, ESPMode::ThreadSafe>();
+		PrecompileReference->System = this;
+		PrecompileReference->CompilationRootObjects = &ActiveCompilation.RootObjects;
+		PrecompileReference->Scripts = ScriptsNeedingCompile;
+		
+		for (int32 i = 0; i < EmitterHandles.Num(); i++)
 		{
-			UNiagaraPrecompileContainer* Container = NewObject<UNiagaraPrecompileContainer>(GetTransientPackage());
-			Container->System = this;
-			Container->Scripts = ScriptsNeedingCompile;
-			TSharedPtr<FNiagaraCompileRequestDataBase, ESPMode::ThreadSafe> SystemPrecompiledData = NiagaraModule.Precompile(Container, FGuid());
-
-			if (SystemPrecompiledData.IsValid() == false)
-			{
-				UE_LOG(LogNiagara, Error, TEXT("Failed to precompile %s.  This is due to unexpected invalid or broken data.  Additional details should be in the log."), *GetPathName());
-				return false;
-			}
-
-			SystemPrecompiledData->GetReferencedObjects(ActiveCompilation.RootObjects);
-			ActiveCompilation.MappedData.Add(SystemSpawnScript, SystemPrecompiledData);
-			ActiveCompilation.MappedData.Add(SystemUpdateScript, SystemPrecompiledData);
-
-			check(EmitterHandles.Num() == SystemPrecompiledData->GetDependentRequestCount());
-
-
-			// Grab the list of user variables that were actually encountered so that we can add to them later.
-			TArray<FNiagaraVariable> EncounteredExposedVars;
-			SystemPrecompiledData->GatherPreCompiledVariables(TEXT("User"), EncounteredExposedVars);
-
-			for (int32 i = 0; i < EmitterHandles.Num(); i++)
+			FNiagaraEmitterHandle Handle = EmitterHandles[i];
+			if (Handle.GetInstance() && Handle.GetIsEnabled())
 			{
-				FNiagaraEmitterHandle Handle = EmitterHandles[i];
-				if (Handle.GetInstance() && Handle.GetIsEnabled())
+				TArray<UNiagaraScript*> EmitterScripts;
+				Handle.GetInstance()->GetScripts(EmitterScripts, false, true);
+				check(EmitterScripts.Num() > 0);
+				for (UNiagaraScript* EmitterScript : EmitterScripts)
 				{
-					TSharedPtr<FNiagaraCompileRequestDataBase, ESPMode::ThreadSafe> EmitterPrecompiledData = SystemPrecompiledData->GetDependentRequest(i);
-					EmitterPrecompiledData->GetReferencedObjects(ActiveCompilation.RootObjects);
-
-					TArray<UNiagaraScript*> EmitterScripts;
-					Handle.GetInstance()->GetScripts(EmitterScripts, false, true);
-					check(EmitterScripts.Num() > 0);
-					for (UNiagaraScript* EmitterScript : EmitterScripts)
-					{
-						ActiveCompilation.MappedData.Add(EmitterScript, EmitterPrecompiledData);
-					}
-
-					// Add the emitter's User variables to the encountered list to expose for later.
-					EmitterPrecompiledData->GatherPreCompiledVariables(TEXT("User"), EncounteredExposedVars);
-				}
-			}
-
-
-			// Now let's synchronize the variables that we actually encountered during precompile so that we can expose them to the end user.
-			for (int32 i = 0; i < EncounteredExposedVars.Num(); i++)
-			{
-				if (OriginalExposedParams.Contains(EncounteredExposedVars[i]) == false)
-				{
-					// Just in case it wasn't added previously..
-					ExposedParameters.AddParameter(EncounteredExposedVars[i]);
+					PrecompileReference->EmitterScriptIndex.Add(EmitterScript, i);
 				}
 			}
 		}
-		
 
 		// We have previously duplicated all that is needed for compilation, so let's now issue the compile requests!
 		for (UNiagaraScript* CompiledScript : ScriptsNeedingCompile)
 		{
 
-			const auto InPairs = [&CompiledScript](const FEmitterCompiledScriptPair& Other) -> bool
+			const auto InPairs = [&CompiledScript](const TSharedPtr<FNiagaraDerivedDataAsyncRequest, ESPMode::ThreadSafe>& Other) -> bool
 			{
-				return CompiledScript == Other.CompiledScript;
+				return CompiledScript == Other->GetScriptPair().CompiledScript;
 			};
 
-			TSharedPtr<FNiagaraCompileRequestDataBase, ESPMode::ThreadSafe> EmitterPrecompiledData = ActiveCompilation.MappedData.FindChecked(CompiledScript);
-			FEmitterCompiledScriptPair* Pair = ActiveCompilation.EmitterCompiledScriptPairs.FindByPredicate(InPairs);
-			check(Pair);
+			TSharedPtr<FNiagaraDerivedDataAsyncRequest, ESPMode::ThreadSafe> AsyncTask = *ActiveCompilation.DDCTasks.FindByPredicate(InPairs);
 
-			// now that we've done the precompile check with the DDC again as our key may have changed.  Currently the Precompile can update the rapid
+			// now that we've done the precompile, check with the DDC. Currently the precompile can update the rapid
 			// iteration parameters, which if they are baked out, will impact the DDC key.
 			// TODO - Handling of the rapid iteration parameters should move to follow merging of emitter sripts rather than be a part of the precompile.
-			if (GetFromDDC(*Pair))
-			{
-				continue;
-			}
+			AsyncTask->DDCKey = UNiagaraScript::BuildNiagaraDDCKeyString(AsyncTask->GetScriptPair().CompileId);
+			AsyncTask->PrecompileReference = PrecompileReference;
 
-			if (!CompiledScript->RequestExternallyManagedAsyncCompile(EmitterPrecompiledData, Pair->CompileId, Pair->PendingJobID))
-			{
-				UE_LOG(LogNiagara, Warning, TEXT("For some reason we are reporting that %s is in sync even though AreScriptAndSourceSynchronized returned false!"), *CompiledScript->GetPathName())
-			}
+			// fire off all the ddc tasks, which will trigger the compilation if the data is not in the ddc
+			UE_LOG(LogNiagara, Verbose, TEXT("Scheduling async get task for %s"), *AsyncTask->AssetPath);
+			AsyncTask->TaskHandle = GetDerivedDataCacheRef().GetAsynchronous(new FNiagaraDerivedDataAsyncRequest(*AsyncTask.Get()));
 		}
-
-		// clean up copies from precompile data
-		TSharedPtr<FNiagaraCompileRequestDataBase, ESPMode::ThreadSafe> SystemPrecompiledData = ActiveCompilation.MappedData.FindChecked(SystemSpawnScript);
-		SystemPrecompiledData->ReleaseCompilationCopies();
 	}
 
 
@@ -2618,10 +3044,6 @@ bool UNiagaraSystem::RequestCompile(bool bForce, FNiagaraSystemUpdateContext* Op
 	return bAnyCompiled;
 }
 
-
-#endif
-
-#if WITH_EDITORONLY_DATA
 void UNiagaraSystem::InitEmitterCompiledData()
 {
 	EmitterCompiledData.Empty();
@@ -2756,6 +3178,7 @@ void UNiagaraSystem::InitSystemCompiledData()
 	}
 
 }
+
 #endif
 
 TStatId UNiagaraSystem::GetStatID(bool bGameThread, bool bConcurrent)const
diff --git a/Engine/Plugins/FX/Niagara/Source/NiagaraCore/Private/NiagaraCustomVersion.cpp b/Engine/Plugins/FX/Niagara/Source/NiagaraCore/Private/NiagaraCustomVersion.cpp
--- a/Engine/Plugins/FX/Niagara/Source/NiagaraCore/Private/NiagaraCustomVersion.cpp
+++ b/Engine/Plugins/FX/Niagara/Source/NiagaraCore/Private/NiagaraCustomVersion.cpp
@@ -8,5 +8,4 @@ const FGuid FNiagaraCustomVersion::GUID(0xFCF57AFA, 0x50764283, 0xB9A9E658, 0xFF
 // Register the custom version with core
 FCustomVersionRegistration GRegisterNiagaraCustomVersion(FNiagaraCustomVersion::GUID, FNiagaraCustomVersion::LatestVersion, TEXT("NiagaraVer"));
 
-const FGuid FNiagaraCustomVersion::LatestScriptCompileVersion(0x4BC4AAF3, 0x89AB4E70, 0xBF75C979, 0x53BE2748);
-
+const FGuid FNiagaraCustomVersion::LatestScriptCompileVersion(0xA83D40AB, 0xE1F64D9F, 0x9F5451FB, 0x8C954559);
diff --git a/Engine/Plugins/FX/Niagara/Source/NiagaraEditor/Private/NiagaraCompiler.cpp b/Engine/Plugins/FX/Niagara/Source/NiagaraEditor/Private/NiagaraCompiler.cpp
--- a/Engine/Plugins/FX/Niagara/Source/NiagaraEditor/Private/NiagaraCompiler.cpp
+++ b/Engine/Plugins/FX/Niagara/Source/NiagaraEditor/Private/NiagaraCompiler.cpp
@@ -307,7 +307,7 @@ void FNiagaraCompileRequestData::VisitReferencedGraphsRecursive(UNiagaraGraph* I
 					{
 						EmitterNode->SyncEnabledState(); // Just to be safe, sync here while we likely still have the handle source.
 						EmitterNode->SetOwnerSystem(nullptr);
-						EmitterNode->SetCachedVariablesForCompilation(*Ptr->EmitterUniqueName, Ptr->NodeGraphDeepCopy.Get(), Ptr->Source);
+						EmitterNode->SetCachedVariablesForCompilation(*Ptr->EmitterUniqueName, Ptr->NodeGraphDeepCopy.Get(), Ptr->Source.Get());
 					}
 				}
 			}
@@ -499,12 +499,28 @@ void FNiagaraCompileRequestData::ReleaseCompilationCopies()
 	// clean up graph copies
 	for (auto& ClonedGraph : ClonedGraphs)
 	{
-		ClonedGraph->ReleaseCompilationCopy();
+		if (ClonedGraph.IsValid())
+		{
+			ClonedGraph->ReleaseCompilationCopy();
+		}
 	}
+	ClonedGraphs.Empty();
+	
 	for (auto& EmitterCompileData : EmitterData)
 	{
-		EmitterCompileData->ReleaseCompilationCopies();
+		if (EmitterCompileData.IsValid())
+		{
+			EmitterCompileData->ReleaseCompilationCopies();
+		}
 	}
+	EmitterData.Empty();
+
+	// clean up script source
+	if (Source.IsValid())
+	{
+		Source->ReleaseCompilationCopy();
+	}
+	Source = nullptr;
 }
 
 void FNiagaraCompileRequestData::FinishPrecompile(const TArray<FNiagaraVariable>& EncounterableVariables, ENiagaraScriptUsage InUsage, FCompileConstantResolver ConstantResolver, const TArray<UNiagaraSimulationStageBase*>* SimStages)
@@ -876,8 +892,8 @@ TSharedPtr<FNiagaraCompileRequestDataBase, ESPMode::ThreadSafe> FNiagaraEditorMo
 
 			FNiagaraUtilities::PrepareRapidIterationParameters(Scripts, ScriptDependencyMap, ScriptToEmitterMap);
 
-			BasePtr->AddRapidIterationParameters(System->GetSystemSpawnScript()->RapidIterationParameters, EmptyResolver);
-			BasePtr->AddRapidIterationParameters(System->GetSystemUpdateScript()->RapidIterationParameters, EmptyResolver);
+			BasePtr->AddRapidIterationParameters(System->GetSystemSpawnScript()->RapidIterationParameters, FCompileConstantResolver(System, ENiagaraScriptUsage::SystemSpawnScript));
+			BasePtr->AddRapidIterationParameters(System->GetSystemUpdateScript()->RapidIterationParameters, FCompileConstantResolver(System, ENiagaraScriptUsage::SystemUpdateScript));
 
 			// Now we can finish off the emitters.
 			for (int32 i = 0; i < System->GetEmitterHandles().Num(); i++)
@@ -960,6 +976,7 @@ int32 FNiagaraEditorModule::CompileScript(const FNiagaraCompileRequestDataBase*
 	int32 JobID = Compiler->CompileScript(CompileRequest, InCompileOptions, TranslateResults, &Translator.GetTranslateOutput(), Translator.GetTranslatedHLSL());
 	ActiveCompilations.Add(JobID, Compiler);
 	return JobID;
+	
 }
 
 TSharedPtr<FNiagaraVMExecutableData> FNiagaraEditorModule::GetCompilationResult(int32 JobID, bool bWait)
diff --git a/Engine/Plugins/FX/Niagara/Source/NiagaraEditor/Private/NiagaraEditorModule.cpp b/Engine/Plugins/FX/Niagara/Source/NiagaraEditor/Private/NiagaraEditorModule.cpp
--- a/Engine/Plugins/FX/Niagara/Source/NiagaraEditor/Private/NiagaraEditorModule.cpp
+++ b/Engine/Plugins/FX/Niagara/Source/NiagaraEditor/Private/NiagaraEditorModule.cpp
@@ -381,7 +381,7 @@ void CompileEmitterStandAlone(UNiagaraEmitter* Emitter, TSet<UNiagaraEmitter*>&
 		}
 
 		Emitter->MarkPackageDirty();
-		UNiagaraSystem* TransientSystem = NewObject<UNiagaraSystem>(GetTransientPackage(), NAME_None, RF_Transient);
+		UNiagaraSystem* TransientSystem = NewObject<UNiagaraSystem>(GetTransientPackage(), FName("StandaloneEmitter_TempSystem"), RF_Transient);
 		UNiagaraSystemFactoryNew::InitializeSystem(TransientSystem, true);
 		TransientSystem->AddEmitterHandle(*Emitter, TEXT("Emitter"));
 		FNiagaraStackGraphUtilities::RebuildEmitterNodes(*TransientSystem);
@@ -755,7 +755,7 @@ void FNiagaraEditorModule::OnPreExit()
 		UNiagaraSystem* Sys = *It;
 		if (Sys)
 		{
-			Sys->WaitForCompilationComplete();
+			Sys->KillAllActiveCompilations();
 		}
 	}
 	ClearObjectPool();
diff --git a/Engine/Plugins/FX/Niagara/Source/NiagaraEditor/Private/NiagaraGraph.cpp b/Engine/Plugins/FX/Niagara/Source/NiagaraEditor/Private/NiagaraGraph.cpp
--- a/Engine/Plugins/FX/Niagara/Source/NiagaraEditor/Private/NiagaraGraph.cpp
+++ b/Engine/Plugins/FX/Niagara/Source/NiagaraEditor/Private/NiagaraGraph.cpp
@@ -2600,11 +2600,11 @@ void UNiagaraGraph::GatherExternalDependencyData(ENiagaraScriptUsage InUsage, co
 		{
 			if (GNiagaraUseGraphHash == 1)
 			{
-				InReferencedCompileHashes.Add(CachedUsageInfo[i].CompileHashFromGraph);
+				InReferencedCompileHashes.AddUnique(CachedUsageInfo[i].CompileHashFromGraph);
 			}
 			else
 			{
-				InReferencedCompileHashes.Add(CachedUsageInfo[i].CompileHash);
+				InReferencedCompileHashes.AddUnique(CachedUsageInfo[i].CompileHash);
 			}
 			InReferencedObjs.Add(CachedUsageInfo[i].Traversal.Last()->GetPathName());
 
diff --git a/Engine/Plugins/FX/Niagara/Source/NiagaraEditor/Private/NiagaraHlslTranslator.h b/Engine/Plugins/FX/Niagara/Source/NiagaraEditor/Private/NiagaraHlslTranslator.h
--- a/Engine/Plugins/FX/Niagara/Source/NiagaraEditor/Private/NiagaraHlslTranslator.h
+++ b/Engine/Plugins/FX/Niagara/Source/NiagaraEditor/Private/NiagaraHlslTranslator.h
@@ -105,7 +105,7 @@ public:
 	TMap<UClass*, UObject*> CDOs;
 	FString EmitterUniqueName;
 	TArray<TSharedPtr<FNiagaraCompileRequestData, ESPMode::ThreadSafe>> EmitterData;
-	UNiagaraScriptSource* Source;
+	TWeakObjectPtr<UNiagaraScriptSource> Source;
 	FString SourceName;
 	bool bUseRapidIterationParams = true;
 	bool bSimulationStagesEnabled = false;
diff --git a/Engine/Plugins/FX/Niagara/Source/NiagaraEditor/Private/NiagaraNodeEmitter.cpp b/Engine/Plugins/FX/Niagara/Source/NiagaraEditor/Private/NiagaraNodeEmitter.cpp
--- a/Engine/Plugins/FX/Niagara/Source/NiagaraEditor/Private/NiagaraNodeEmitter.cpp
+++ b/Engine/Plugins/FX/Niagara/Source/NiagaraEditor/Private/NiagaraNodeEmitter.cpp
@@ -492,8 +492,10 @@ void UNiagaraNodeEmitter::GatherExternalDependencyData(ENiagaraScriptUsage InMas
 
 	if (CalledGraph && IsNodeEnabled()) // Skip if disabled
 	{
+		CalledGraph->RebuildCachedCompileIds();
 		ENiagaraScriptUsage TargetUsage = InMasterUsage == ENiagaraScriptUsage::SystemSpawnScript ? ENiagaraScriptUsage::EmitterSpawnScript : ENiagaraScriptUsage::EmitterUpdateScript;
-		InReferencedCompileHashes.Add(CalledGraph->GetCompileDataHash(TargetUsage, FGuid(0,0,0,0)));
+		FNiagaraCompileHash Hash = CalledGraph->GetCompileDataHash(TargetUsage, FGuid(0,0,0,0));
+		InReferencedCompileHashes.AddUnique(Hash);
 		InReferencedObjs.Add(CalledGraph->GetPathName());
 		CalledGraph->GatherExternalDependencyData(TargetUsage, FGuid(0, 0, 0, 0), InReferencedCompileHashes, InReferencedObjs);
 	}
diff --git a/Engine/Plugins/FX/Niagara/Source/NiagaraEditor/Private/NiagaraNodeFunctionCall.cpp b/Engine/Plugins/FX/Niagara/Source/NiagaraEditor/Private/NiagaraNodeFunctionCall.cpp
--- a/Engine/Plugins/FX/Niagara/Source/NiagaraEditor/Private/NiagaraNodeFunctionCall.cpp
+++ b/Engine/Plugins/FX/Niagara/Source/NiagaraEditor/Private/NiagaraNodeFunctionCall.cpp
@@ -1206,7 +1206,7 @@ void UNiagaraNodeFunctionCall::GatherExternalDependencyData(ENiagaraScriptUsage
 				FNiagaraCompileHash FoundCompileHash = FunctionGraph->GetCompileDataHash((ENiagaraScriptUsage)i, FGuid(0, 0, 0, 0));
 				if (FoundGuid.IsValid() && FoundCompileHash.IsValid())
 				{
-					InReferencedCompileHashes.Add(FoundCompileHash);
+					InReferencedCompileHashes.AddUnique(FoundCompileHash);
 					InReferencedObjs.Add(FunctionGraph->GetPathName());
 					FunctionGraph->GatherExternalDependencyData((ENiagaraScriptUsage)i, FGuid(0, 0, 0, 0), InReferencedCompileHashes, InReferencedObjs);
 				}
diff --git a/Engine/Plugins/FX/Niagara/Source/NiagaraEditor/Private/NiagaraNodeParameterMapGet.cpp b/Engine/Plugins/FX/Niagara/Source/NiagaraEditor/Private/NiagaraNodeParameterMapGet.cpp
--- a/Engine/Plugins/FX/Niagara/Source/NiagaraEditor/Private/NiagaraNodeParameterMapGet.cpp
+++ b/Engine/Plugins/FX/Niagara/Source/NiagaraEditor/Private/NiagaraNodeParameterMapGet.cpp
@@ -520,7 +520,8 @@ void UNiagaraNodeParameterMapGet::GatherExternalDependencyData(ENiagaraScriptUsa
 		UNiagaraParameterCollection* Collection = Schema->VariableIsFromParameterCollection(Var);
 		if (Collection)
 		{
-			InReferencedCompileHashes.Add(Collection->GetCompileHash());
+			FNiagaraCompileHash Hash = Collection->GetCompileHash();
+			InReferencedCompileHashes.AddUnique(Hash);
 			InReferencedObjs.Add(Collection->GetPathName());
 		}
 	}
diff --git a/Engine/Plugins/FX/Niagara/Source/NiagaraEditor/Private/NiagaraScriptSource.cpp b/Engine/Plugins/FX/Niagara/Source/NiagaraEditor/Private/NiagaraScriptSource.cpp
--- a/Engine/Plugins/FX/Niagara/Source/NiagaraEditor/Private/NiagaraScriptSource.cpp
+++ b/Engine/Plugins/FX/Niagara/Source/NiagaraEditor/Private/NiagaraScriptSource.cpp
@@ -47,13 +47,13 @@ void UNiagaraScriptSource::ComputeVMCompilationId(FNiagaraVMExecutableDataId& Id
 
 	// Add in any referenced HLSL files.
 	FSHAHash Hash = GetShaderFileHash((TEXT("/Plugin/FX/Niagara/Private/NiagaraEmitterInstanceShader.usf")), EShaderPlatform::SP_PCD3D_SM5);
-	Id.ReferencedCompileHashes.Emplace(Hash.Hash, sizeof(Hash.Hash)/sizeof(uint8));
+	Id.ReferencedCompileHashes.AddUnique(FNiagaraCompileHash(Hash.Hash, sizeof(Hash.Hash)/sizeof(uint8)));
 	Id.DebugReferencedObjects.Emplace(TEXT("/Plugin/FX/Niagara/Private/NiagaraEmitterInstanceShader.usf"));
 	Hash = GetShaderFileHash((TEXT("/Plugin/FX/Niagara/Private/NiagaraShaderVersion.ush")), EShaderPlatform::SP_PCD3D_SM5);
-	Id.ReferencedCompileHashes.Emplace(Hash.Hash, sizeof(Hash.Hash) / sizeof(uint8));
+	Id.ReferencedCompileHashes.AddUnique(FNiagaraCompileHash(Hash.Hash, sizeof(Hash.Hash)/sizeof(uint8)));
 	Id.DebugReferencedObjects.Emplace(TEXT("/Plugin/FX/Niagara/Private/NiagaraShaderVersion.ush"));
 	Hash = GetShaderFileHash((TEXT("/Engine/Public/ShaderVersion.ush")), EShaderPlatform::SP_PCD3D_SM5);
-	Id.ReferencedCompileHashes.Emplace(Hash.Hash, sizeof(Hash.Hash) / sizeof(uint8));
+	Id.ReferencedCompileHashes.AddUnique(FNiagaraCompileHash(Hash.Hash, sizeof(Hash.Hash)/sizeof(uint8)));
 	Id.DebugReferencedObjects.Emplace(TEXT("/Engine/Public/ShaderVersion.ush"));
 }
 
@@ -108,9 +108,22 @@ UNiagaraScriptSource* UNiagaraScriptSource::CreateCompilationCopy()
 {
 	UNiagaraScriptSource* Result = NewObject<UNiagaraScriptSource>();
 	Result->NodeGraph = NodeGraph->CreateCompilationCopy();
+	Result->bIsCompilationCopy = true;
+	Result->AddToRoot();
 	return Result;
 }
 
+void UNiagaraScriptSource::ReleaseCompilationCopy()
+{
+	if (bIsCompilationCopy && !bIsReleased)
+	{
+		bIsReleased = true;
+		NodeGraph = nullptr;
+		RemoveFromRoot();
+		MarkPendingKill();
+	}
+}
+
 bool UNiagaraScriptSource::IsSynchronized(const FGuid& InChangeId)
 {
 	if (NodeGraph)
@@ -158,6 +171,56 @@ void UNiagaraScriptSource::PostLoadFromEmitter(UNiagaraEmitter& OwningEmitter)
 	}
 }
 
+void FindObjectNamesRecursive(UNiagaraGraph* InGraph, FString EmitterUniqueName, TMap<FName, UNiagaraDataInterface*>& Result)
+{
+	if (!InGraph)
+	{
+		return;
+	}
+	TArray<UNiagaraNode*> Nodes;
+	InGraph->GetNodesOfClass(Nodes);
+	for (UEdGraphNode* Node : Nodes)
+	{
+		if (UNiagaraNode* InNode = Cast<UNiagaraNode>(Node))
+		{
+			UNiagaraNodeInput* InputNode = Cast<UNiagaraNodeInput>(InNode);
+			if (InputNode)
+			{
+				if (InputNode->Input.IsDataInterface())
+				{
+					UNiagaraDataInterface* DataInterface = InputNode->GetDataInterface();
+					bool bIsParameterMapDataInterface = false;
+					FName DIName = FHlslNiagaraTranslator::GetDataInterfaceName(InputNode->Input.GetName(), EmitterUniqueName, bIsParameterMapDataInterface);
+					if (Result.Contains(DIName) && Result[DIName] != DataInterface)
+					{
+						UE_LOG(LogNiagaraEditor, Log, TEXT("Duplicate data interface name %s in graph %s. One of the data interfaces will override the other when resolving the name."), *DIName.ToString(), *InGraph->GetPathName());
+					}
+					Result.Add(DIName, DataInterface);
+				}
+				continue;
+			}
+
+			UNiagaraNodeFunctionCall* FunctionCallNode = Cast<UNiagaraNodeFunctionCall>(InNode);
+			if (FunctionCallNode)
+			{
+				UNiagaraGraph* FunctionGraph = FunctionCallNode->GetCalledGraph();
+				FindObjectNamesRecursive(FunctionGraph, EmitterUniqueName, Result);
+			}
+		}
+	}
+}
+
+TMap<FName, UNiagaraDataInterface*> UNiagaraScriptSource::ComputeObjectNameMap(FString EmitterUniqueName) const
+{
+	TMap<FName, UNiagaraDataInterface*> Result;
+	if (!NodeGraph)
+	{
+		return Result;
+	}
+	FindObjectNamesRecursive(NodeGraph, EmitterUniqueName, Result);
+	return Result;
+}
+
 bool UNiagaraScriptSource::AddModuleIfMissing(FString ModulePath, ENiagaraScriptUsage Usage, bool& bOutFoundModule)
 {
 	FSoftObjectPath SystemUpdateScriptRef(ModulePath);
diff --git a/Engine/Plugins/FX/Niagara/Source/NiagaraEditor/Public/NiagaraScriptSource.h b/Engine/Plugins/FX/Niagara/Source/NiagaraEditor/Public/NiagaraScriptSource.h
--- a/Engine/Plugins/FX/Niagara/Source/NiagaraEditor/Public/NiagaraScriptSource.h
+++ b/Engine/Plugins/FX/Niagara/Source/NiagaraEditor/Public/NiagaraScriptSource.h
@@ -19,12 +19,17 @@ class UNiagaraScriptSource : public UNiagaraScriptSourceBase
 	/** Graph for particle update expression */
 	UPROPERTY()
 	TObjectPtr<class UNiagaraGraph>	NodeGraph;
+
+	bool bIsCompilationCopy = false;
+	bool bIsReleased = false;
 	
 	// UObject interface
 	virtual void PostLoad() override;
 
 	UNiagaraScriptSource* CreateCompilationCopy();
 
+	void ReleaseCompilationCopy();
+
 	// UNiagaraScriptSourceBase interface.
 	//virtual ENiagaraScriptCompileStatus Compile(UNiagaraScript* ScriptOwner, FString& OutGraphLevelErrorMessages) override;
 	virtual bool IsSynchronized(const FGuid& InChangeId) override;
@@ -40,6 +45,8 @@ class UNiagaraScriptSource : public UNiagaraScriptSourceBase
 
 	virtual void PostLoadFromEmitter(UNiagaraEmitter& OwningEmitter) override;
 
+	virtual TMap<FName, UNiagaraDataInterface*> ComputeObjectNameMap(FString EmitterUniqueName) const override;
+
 	NIAGARAEDITOR_API virtual bool AddModuleIfMissing(FString ModulePath, ENiagaraScriptUsage Usage, bool& bOutFoundModule)override;
 
 	virtual void FixupRenamedParameters(UNiagaraNode* Node, FNiagaraParameterStore& RapidIterationParameters, const TArray<FNiagaraVariable>& OldRapidIterationVariables, const TSet<FName>& ValidRapidIterationParameterNames, const UNiagaraEmitter* Emitter, ENiagaraScriptUsage ScriptUsage) const;
diff --git a/Engine/Plugins/FX/Niagara/Source/NiagaraShader/Private/NiagaraShared.cpp b/Engine/Plugins/FX/Niagara/Source/NiagaraShader/Private/NiagaraShared.cpp
--- a/Engine/Plugins/FX/Niagara/Source/NiagaraShader/Private/NiagaraShared.cpp
+++ b/Engine/Plugins/FX/Niagara/Source/NiagaraShader/Private/NiagaraShared.cpp
@@ -604,11 +604,11 @@ bool FNiagaraShaderScript::CacheShaders(const FNiagaraShaderMapId& ShaderMapId,
 			FNiagaraShaderMap::LoadFromDerivedDataCache(this, ShaderMapId, ShaderPlatform, GameThreadShaderMap);
 			if (GameThreadShaderMap && GameThreadShaderMap->IsValid())
 			{
-				UE_LOG(LogTemp, Verbose, TEXT("Loaded shader %s for Niagara script %s from DDC"), *GameThreadShaderMap->GetFriendlyName(), *GetFriendlyName());
+				UE_LOG(LogShaders, Verbose, TEXT("Loaded shader %s for Niagara script %s from DDC"), *GameThreadShaderMap->GetFriendlyName(), *GetFriendlyName());
 			}
 			else
 			{
-				UE_LOG(LogTemp, Display, TEXT("Loading shader for Niagara script %s from DDC failed. Shader needs recompile."), *GetFriendlyName());
+				UE_LOG(LogShaders, Display, TEXT("Loading shader for Niagara script %s from DDC failed. Shader needs recompile."), *GetFriendlyName());
 			}
 		}
 	}
