Prevent active Niagara compilations from starving all other background tasks. Prevent transient systems from starting compilations.

diff --git a/Engine/Plugins/FX/Niagara/Source/Niagara/Private/NiagaraModule.cpp b/Engine/Plugins/FX/Niagara/Source/Niagara/Private/NiagaraModule.cpp
--- a/Engine/Plugins/FX/Niagara/Source/Niagara/Private/NiagaraModule.cpp
+++ b/Engine/Plugins/FX/Niagara/Source/Niagara/Private/NiagaraModule.cpp
@@ -69,6 +69,14 @@ static FAutoConsoleVariableRef CVarUseGlobalFXBudget(
 	ECVF_Default
 );
 
+static float GNiagaraMaxCompilePollTimePerFrame = 0.016f;
+static FAutoConsoleVariableRef CVarNiagaraMaxCompilePollTimePerFrame(
+	TEXT("fx.Niagara.MaxCompilePollTimePerFrame"),
+	GNiagaraMaxCompilePollTimePerFrame,
+	TEXT("When a lot of system compile tasks queue up, this is the max time per frame that is used to advance them."),
+	ECVF_Default
+);
+
 void INiagaraModule::OnUseGlobalFXBudgetChanged(IConsoleVariable* Variable)
 {
 	//Warn if we're enabling Niagara's use of budgeting but the budgeting itself is disabled.
@@ -368,8 +376,28 @@ void INiagaraModule::OnPreExit()
 #endif
 }
 
+#if WITH_EDITOR
+void PollSystemCompilations()
+{
+	double Start = FPlatformTime::Seconds();
+	for (TObjectIterator<UNiagaraSystem> SysIt; SysIt; ++SysIt)
+	{
+		UNiagaraSystem* System = *SysIt;
+		System->PollForCompilationComplete();
+		if (FPlatformTime::Seconds() - Start > GNiagaraMaxCompilePollTimePerFrame)
+		{
+			break;
+		}
+	}
+}
+#endif
+
 void INiagaraModule::OnWorldTickStart(UWorld* World, ELevelTick TickType, float DeltaSeconds)
 {
+#if WITH_EDITOR
+	PollSystemCompilations();
+#endif
+	
 #if NIAGARA_PERF_BASELINES
 	if (BaselineHandler.IsValid())
 	{
diff --git a/Engine/Plugins/FX/Niagara/Source/NiagaraEditor/Private/ViewModels/NiagaraSystemViewModel.cpp b/Engine/Plugins/FX/Niagara/Source/NiagaraEditor/Private/ViewModels/NiagaraSystemViewModel.cpp
--- a/Engine/Plugins/FX/Niagara/Source/NiagaraEditor/Private/ViewModels/NiagaraSystemViewModel.cpp
+++ b/Engine/Plugins/FX/Niagara/Source/NiagaraEditor/Private/ViewModels/NiagaraSystemViewModel.cpp
@@ -301,6 +301,10 @@ TSharedPtr<FNiagaraSystemScriptViewModel> FNiagaraSystemViewModel::GetSystemScri
 
 void FNiagaraSystemViewModel::CompileSystem(bool bForce)
 {
+	if (EditMode == ENiagaraSystemViewModelEditMode::EmitterDuringMerge && bForce == false)
+	{
+		return;
+	}
 	SCOPE_CYCLE_COUNTER(STAT_NiagaraEditor_SystemViewModel_CompileSystem);
 	FNiagaraEditorUtilities::KillSystemInstances(GetSystem());
 	check(SystemScriptViewModel.IsValid());
