- Fix race condition around releasing skel mesh UV / connectivity

#rnx

diff --git a/Engine/Plugins/FX/Niagara/Source/Niagara/Private/NiagaraDataInterfaceSkeletalMeshConnectivity.cpp b/Engine/Plugins/FX/Niagara/Source/Niagara/Private/NiagaraDataInterfaceSkeletalMeshConnectivity.cpp
--- a/Engine/Plugins/FX/Niagara/Source/Niagara/Private/NiagaraDataInterfaceSkeletalMeshConnectivity.cpp
+++ b/Engine/Plugins/FX/Niagara/Source/Niagara/Private/NiagaraDataInterfaceSkeletalMeshConnectivity.cpp
@@ -80,8 +80,6 @@ FSkeletalMeshConnectivity::FSkeletalMeshConnectivity(TWeakObjectPtr<USkeletalMes
 	: LodIndex(InLodIndex)
 	, MeshObject(InMeshObject)
 	, GpuUserCount(0)
-	, ReleasedByRT(false)
-	, QueuedForRelease(false)
 {
 }
 
@@ -92,7 +90,7 @@ bool FSkeletalMeshConnectivity::IsUsed() const
 
 bool FSkeletalMeshConnectivity::CanBeDestroyed() const
 {
-	return !IsUsed() && (!QueuedForRelease || ReleasedByRT);
+	return !IsUsed();
 }
 
 void FSkeletalMeshConnectivity::RegisterUser(FSkeletalMeshConnectivityUsage Usage, bool bNeedsDataImmediately)
@@ -101,8 +99,10 @@ void FSkeletalMeshConnectivity::RegisterUser(FSkeletalMeshConnectivityUsage Usag
 	{
 		if (GpuUserCount++ == 0)
 		{
-			Proxy.Initialize(*this);
-			BeginInitResource(&Proxy);
+			check(Proxy == nullptr);
+			Proxy.Reset(new FSkeletalMeshConnectivityProxy());
+			Proxy->Initialize(*this);
+			BeginInitResource(Proxy.Get());
 		}
 	}
 }
@@ -113,24 +113,20 @@ void FSkeletalMeshConnectivity::UnregisterUser(FSkeletalMeshConnectivityUsage Us
 	{
 		if (--GpuUserCount == 0)
 		{
-			QueuedForRelease = true;
-			ReleasedByRT = false;
-			FThreadSafeBool* Released = &ReleasedByRT;
-
-			BeginReleaseResource(&Proxy);
-
 			ENQUEUE_RENDER_COMMAND(BeginDestroyCommand)(
-				[Released](FRHICommandListImmediate& RHICmdList)
+				[RT_Proxy=Proxy.Release()](FRHICommandListImmediate& RHICmdList)
 				{
-					*Released = true;
-				});
+					RT_Proxy->ReleaseResource();
+					delete RT_Proxy;
+				}
+			);
 		}
 	}
 }
 
 bool FSkeletalMeshConnectivity::CanBeUsed(const TWeakObjectPtr<USkeletalMesh>& InMeshObject, int32 InLodIndex) const
 {
-	return !QueuedForRelease && LodIndex == InLodIndex && MeshObject == InMeshObject;
+	return IsUsed() && LodIndex == InLodIndex && MeshObject == InMeshObject;
 }
 
 bool FSkeletalMeshConnectivity::IsValidMeshObject(TWeakObjectPtr<USkeletalMesh>& MeshObject, int32 InLodIndex)
@@ -210,7 +206,7 @@ FString FSkeletalMeshConnectivity::GetMeshName() const
 
 const FSkeletalMeshConnectivityProxy* FSkeletalMeshConnectivity::GetProxy() const
 {
-	return &Proxy;
+	return Proxy.Get();
 }
 
 struct FAdjacencyVertexOverlapKey
diff --git a/Engine/Plugins/FX/Niagara/Source/Niagara/Private/NiagaraDataInterfaceSkeletalMeshConnectivity.h b/Engine/Plugins/FX/Niagara/Source/Niagara/Private/NiagaraDataInterfaceSkeletalMeshConnectivity.h
--- a/Engine/Plugins/FX/Niagara/Source/Niagara/Private/NiagaraDataInterfaceSkeletalMeshConnectivity.h
+++ b/Engine/Plugins/FX/Niagara/Source/Niagara/Private/NiagaraDataInterfaceSkeletalMeshConnectivity.h
@@ -57,10 +57,7 @@ private:
 	static const FSkeletalMeshLODRenderData* GetLodRenderData(const USkeletalMesh& Mesh, int32 LodIndex);
 
 	TWeakObjectPtr<USkeletalMesh> MeshObject;
-	FSkeletalMeshConnectivityProxy Proxy;
+	TUniquePtr<FSkeletalMeshConnectivityProxy> Proxy;
 
 	std::atomic<int32> GpuUserCount;
-
-	FThreadSafeBool ReleasedByRT;
-	bool QueuedForRelease;
 };
\ No newline at end of file
diff --git a/Engine/Plugins/FX/Niagara/Source/Niagara/Private/NiagaraDataInterfaceSkeletalMeshUvMapping.cpp b/Engine/Plugins/FX/Niagara/Source/Niagara/Private/NiagaraDataInterfaceSkeletalMeshUvMapping.cpp
--- a/Engine/Plugins/FX/Niagara/Source/Niagara/Private/NiagaraDataInterfaceSkeletalMeshUvMapping.cpp
+++ b/Engine/Plugins/FX/Niagara/Source/Niagara/Private/NiagaraDataInterfaceSkeletalMeshUvMapping.cpp
@@ -285,8 +285,6 @@ FSkeletalMeshUvMapping::FSkeletalMeshUvMapping(TWeakObjectPtr<USkeletalMesh> InM
 	, TriangleIndexQuadTree(8 /* Internal node capacity */, 8 /* Maximum tree depth */)
 	, CpuQuadTreeUserCount(0)
 	, GpuQuadTreeUserCount(0)
-	, ReleasedByRT(false)
-	, QueuedForRelease(false)
 {
 }
 
@@ -344,22 +342,19 @@ void FSkeletalMeshUvMapping::ReleaseQuadTree()
 
 void FSkeletalMeshUvMapping::BuildGpuQuadTree()
 {
-	FrozenQuadTreeProxy.Initialize(*this);
-	BeginInitResource(&FrozenQuadTreeProxy);
+	check(FrozenQuadTreeProxy == nullptr);
+	FrozenQuadTreeProxy.Reset(new FSkeletalMeshUvMappingBufferProxy());
+	FrozenQuadTreeProxy->Initialize(*this);
+	BeginInitResource(FrozenQuadTreeProxy.Get());
 }
 
 void FSkeletalMeshUvMapping::ReleaseGpuQuadTree()
 {
-	QueuedForRelease = true;
-	ReleasedByRT = false;
-	FThreadSafeBool* Released = &ReleasedByRT;
-
-	BeginReleaseResource(&FrozenQuadTreeProxy);
-
 	ENQUEUE_RENDER_COMMAND(BeginDestroyCommand)(
-		[Released](FRHICommandListImmediate& RHICmdList)
+		[RT_Proxy=FrozenQuadTreeProxy.Release()](FRHICommandListImmediate& RHICmdList)
 		{
-			*Released = true;
+			RT_Proxy->ReleaseResource();
+			delete RT_Proxy;
 		});
 }
 
@@ -371,7 +366,7 @@ bool FSkeletalMeshUvMapping::IsUsed() const
 
 bool FSkeletalMeshUvMapping::CanBeDestroyed() const
 {
-	return !IsUsed() && (!QueuedForRelease || ReleasedByRT);
+	return !IsUsed();
 }
 
 void FSkeletalMeshUvMapping::RegisterUser(FSkeletalMeshUvMappingUsage Usage, bool bNeedsDataImmediately)
@@ -504,7 +499,7 @@ const FSkeletalMeshLODRenderData* FSkeletalMeshUvMapping::GetLodRenderData() con
 
 const FSkeletalMeshUvMappingBufferProxy* FSkeletalMeshUvMapping::GetQuadTreeProxy() const
 {
-	return &FrozenQuadTreeProxy;
+	return FrozenQuadTreeProxy.Get();
 }
 
 void
diff --git a/Engine/Plugins/FX/Niagara/Source/Niagara/Private/NiagaraDataInterfaceSkeletalMeshUvMapping.h b/Engine/Plugins/FX/Niagara/Source/Niagara/Private/NiagaraDataInterfaceSkeletalMeshUvMapping.h
--- a/Engine/Plugins/FX/Niagara/Source/Niagara/Private/NiagaraDataInterfaceSkeletalMeshUvMapping.h
+++ b/Engine/Plugins/FX/Niagara/Source/Niagara/Private/NiagaraDataInterfaceSkeletalMeshUvMapping.h
@@ -66,11 +66,8 @@ private:
 	TWeakObjectPtr<USkeletalMesh> MeshObject;
 
 	FNiagaraUvQuadTree TriangleIndexQuadTree;
-	FSkeletalMeshUvMappingBufferProxy FrozenQuadTreeProxy;
+	TUniquePtr<FSkeletalMeshUvMappingBufferProxy> FrozenQuadTreeProxy;
 
 	std::atomic<int32> CpuQuadTreeUserCount;
 	std::atomic<int32> GpuQuadTreeUserCount;
-
-	FThreadSafeBool ReleasedByRT;
-	bool QueuedForRelease;
 };
