Dumping data to try and figure out differences between frames at screenshot intervals.

diff --git a/Engine/Plugins/FX/Niagara/Source/Niagara/Classes/NiagaraDataSet.h b/Engine/Plugins/FX/Niagara/Source/Niagara/Classes/NiagaraDataSet.h
--- a/Engine/Plugins/FX/Niagara/Source/Niagara/Classes/NiagaraDataSet.h
+++ b/Engine/Plugins/FX/Niagara/Source/Niagara/Classes/NiagaraDataSet.h
@@ -122,7 +122,7 @@ public:
 	void KillInstance(uint32 InstanceIdx);
 	void CopyTo(FNiagaraDataBuffer& DestBuffer, int32 SrcStartIdx, int32 DestStartIdx, int32 NumInstances)const;
 	void GPUCopyFrom(const float* GPUReadBackFloat, const int* GPUReadBackInt, const FFloat16* GPUReadBackHalf, int32 StartIdx, int32 NumInstances, uint32 InSrcFloatStride, uint32 InSrcIntStride, uint32 InSrcHalfStride);
-	void Dump(int32 StartIndex, int32 NumInstances, const FString& Label)const;
+	void Dump(int32 StartIndex, int32 NumInstances, const FString& Label, const FName& SortParameterKey = FName())const;
 
 	FORCEINLINE TArrayView<uint8 const* RESTRICT const> GetRegisterTable() const { return TArrayView<uint8 const* RESTRICT const>(RegisterTable); }
 	
@@ -373,7 +373,7 @@ public:
 
 	void CheckForNaNs() const;
 
-	void Dump(int32 StartIndex, int32 NumInstances, const FString& Label) const;
+	void Dump(int32 StartIndex, int32 NumInstances, const FString& Label, const FName& SortParameterKey = FName()) const;
 
 	FORCEINLINE bool IsCurrentDataValid()const { return CurrentData != nullptr; }
 	FORCEINLINE FNiagaraDataBuffer* GetCurrentData()const {	return CurrentData; }
diff --git a/Engine/Plugins/FX/Niagara/Source/Niagara/Private/NiagaraDataSet.cpp b/Engine/Plugins/FX/Niagara/Source/Niagara/Private/NiagaraDataSet.cpp
--- a/Engine/Plugins/FX/Niagara/Source/Niagara/Private/NiagaraDataSet.cpp
+++ b/Engine/Plugins/FX/Niagara/Source/Niagara/Private/NiagaraDataSet.cpp
@@ -381,17 +381,17 @@ void FNiagaraDataSet::CheckForNaNs() const
 	}
 }
 
-void FNiagaraDataSet::Dump(int32 StartIndex, int32 NumInstances, const FString& Label)const
+void FNiagaraDataSet::Dump(int32 StartIndex, int32 NumInstances, const FString& Label, const FName& SortParameterKey)const
 {
 	if (CurrentData)
 	{
-		CurrentData->Dump(StartIndex, NumInstances, Label);
+		CurrentData->Dump(StartIndex, NumInstances, Label, SortParameterKey);
 	}
 
 	if (GetDestinationData())
 	{
 		FString DestLabel = Label + TEXT("[Destination]");
-		DestinationData->Dump(StartIndex, NumInstances, DestLabel);
+		DestinationData->Dump(StartIndex, NumInstances, DestLabel, SortParameterKey);
 	}
 }
 
@@ -1120,7 +1120,7 @@ void FNiagaraDataBuffer::GPUCopyFrom(const float* GPUReadBackFloat, const int* G
 	}
 }
 
-void FNiagaraDataBuffer::Dump(int32 StartIndex, int32 InNumInstances, const FString& Label)const
+void FNiagaraDataBuffer::Dump(int32 StartIndex, int32 InNumInstances, const FString& Label, const FName& SortParameterKey)const
 {
 	FNiagaraDataVariableIterator Itr(this, StartIndex);
 
@@ -1131,18 +1131,35 @@ void FNiagaraDataBuffer::Dump(int32 StartIndex, int32 InNumInstances, const FStr
 	}
 
 	int32 NumInstancesDumped = 0;
-	TArray<FString> Lines;
+	TArray<TPair<float, FString>> Lines;
 	Lines.Reserve(GetNumInstances());
+	bool bKeyFound = false;
+	int32 MaxKeySize = 0;
 	while (Itr.IsValid() && NumInstancesDumped < InNumInstances)
 	{
 		Itr.Get();
 
 		FString Line = TEXT("| ");
+		float Key = 0.0f;
 		for (const FNiagaraVariable& Var : Itr.GetVariables())
 		{
 			Line += Var.ToString() + TEXT(" | ");
+			if (SortParameterKey.IsNone() == false && Var.GetName() == SortParameterKey)
+			{
+				// Dump value to string so that we can sort on it later.
+				if (Var.GetType() == FNiagaraTypeDefinition::GetFloatDef())
+				{
+					Key = Var.GetValue<float>();
+					bKeyFound = true;
+				}
+				else if (Var.GetType() == FNiagaraTypeDefinition::GetIntDef())
+				{
+					Key = (float)Var.GetValue<int32>();
+					bKeyFound = true;
+				}
+			}
 		}
-		Lines.Add(Line);
+		Lines.Add(TPair<float, FString>(Key, Line));
 		Itr.Advance();
 		NumInstancesDumped++;
 	}
@@ -1156,14 +1173,32 @@ void FNiagaraDataBuffer::Dump(int32 StartIndex, int32 InNumInstances, const FStr
 		}
 	}
 
+	if (bKeyFound)
+	{
+		Lines.StableSort([](const TPair<float, FString>& Lhs, const TPair<float, FString>& Rhs)
+			{
+				return Lhs.Key < Rhs.Key;
+			});
+	}
+
 	UE_LOG(LogNiagara, Log, TEXT("%s"), *Sep);
 	UE_LOG(LogNiagara, Log, TEXT(" %s "), *Label);
 	UE_LOG(LogNiagara, Log, TEXT("%s"), *Sep);
 	// 	UE_LOG(LogNiagara, Log, TEXT("%s"), *HeaderStr);
 	// 	UE_LOG(LogNiagara, Log, TEXT("%s"), *Sep);
-	for (FString& Str : Lines)
+	if (bKeyFound)
+	{
+		for (TPair<float, FString>& Str : Lines)
+		{
+			UE_LOG(LogNiagara, Log, TEXT("[%f] %s"), Str.Key, *Str.Value);
+		}
+	}
+	else
 	{
-		UE_LOG(LogNiagara, Log, TEXT("%s"), *Str);
+		for (TPair<float, FString>& Str : Lines)
+		{
+			UE_LOG(LogNiagara, Log, TEXT("%s"), *Str.Value);
+		}
 	}
 	if (IDToIndexTable.Num() > 0)
 	{
diff --git a/Engine/Plugins/FX/Niagara/Source/Niagara/Private/NiagaraEmitterInstance.cpp b/Engine/Plugins/FX/Niagara/Source/Niagara/Private/NiagaraEmitterInstance.cpp
--- a/Engine/Plugins/FX/Niagara/Source/Niagara/Private/NiagaraEmitterInstance.cpp
+++ b/Engine/Plugins/FX/Niagara/Source/Niagara/Private/NiagaraEmitterInstance.cpp
@@ -217,7 +217,7 @@ void FNiagaraEmitterInstance::Dump()const
 	}
 	else
 	{
-		ParticleDataSet->Dump(0, INDEX_NONE, TEXT("Particle Data"));
+		ParticleDataSet->Dump(0, INDEX_NONE, TEXT("Particle Data"), TEXT("UniqueID"));
 	}
 }
 
