Fixed normalization missed out of recent refactor.

diff --git a/Engine/Plugins/FX/Niagara/Source/Niagara/Private/NiagaraRendererMeshes.cpp b/Engine/Plugins/FX/Niagara/Source/Niagara/Private/NiagaraRendererMeshes.cpp
--- a/Engine/Plugins/FX/Niagara/Source/Niagara/Private/NiagaraRendererMeshes.cpp
+++ b/Engine/Plugins/FX/Niagara/Source/Niagara/Private/NiagaraRendererMeshes.cpp
@@ -1332,7 +1332,7 @@ void FNiagaraRendererMeshes::GetDynamicRayTracingInstances(FRayTracingMaterialGa
 					}
 
 					const FVector InstancePosition = bHasPosition ? GetInstancePosition(InstanceIndex) : FVector::ZeroVector;
-					const FQuat InstanceRotation = bHasRotation ? GetInstanceQuat(InstanceIndex) : FQuat::Identity;
+					const FQuat InstanceRotation = bHasRotation ? GetInstanceQuat(InstanceIndex).GetNormalized() : FQuat::Identity;
 					FMatrix InstanceTransform = FQuatRotationTranslationMatrix::Make(InstanceRotation, InstancePosition);
 
 					const FVector InstanceScale = bHasScale ? GetInstanceScale(InstanceIndex) * MeshData.Scale : MeshData.Scale;
