- Add options for Niagara auto mip generation

diff --git a/Engine/Plugins/FX/Niagara/Shaders/Private/NiagaraGenerateMips.usf b/Engine/Plugins/FX/Niagara/Shaders/Private/NiagaraGenerateMips.usf
--- a/Engine/Plugins/FX/Niagara/Shaders/Private/NiagaraGenerateMips.usf
+++ b/Engine/Plugins/FX/Niagara/Shaders/Private/NiagaraGenerateMips.usf
@@ -3,7 +3,8 @@
 #include "/Engine/Private/Common.ush"
 #include "/Engine/Private/GammaCorrectionCommon.ush"
 
-float2 TexelSize;
+float2 SrcTexelSize;
+float2 DstTexelSize;
 
 Texture2D MipInSRV;
 SamplerState MipInSampler;
@@ -11,6 +12,8 @@ SamplerState MipInSampler;
 int2 MipOutSize;
 RWTexture2D<float4> MipOutUAV;
 
+int KernelHWidth;
+
 [numthreads(THREADGROUPSIZE, THREADGROUPSIZE, 1)]
 void MainCS(uint3 ThreadID : SV_DispatchThreadID)
 {
@@ -19,8 +22,33 @@ void MainCS(uint3 ThreadID : SV_DispatchThreadID)
 		return;
 	}
 
-	float2 SrcUV = (float2(ThreadID.xy) + 0.5f) * TexelSize;
-    float4 OutColor = MipInSRV.SampleLevel(MipInSampler, SrcUV, 0);
+	float4 OutColor = 0;
+	float2 UV = (float2(ThreadID.xy) + 0.5f) * DstTexelSize;
+#if GENMIPS_GAUSSIAN
+	{
+		float TotalWeight = 0.0f;
+		float s = KernelHWidth;
+
+		//-OPT: If we still want single pass use groupshared to reduce taps, or go two pass instead (which would add more barriers)
+		//-OPT: Consider moving the Offsets / Weights out
+		for ( int y=-KernelHWidth; y <= KernelHWidth; ++y )
+		{
+			for ( int x=-KernelHWidth; x <= KernelHWidth; ++x )
+			{
+				float2 Offset = float2(x, y);
+				float r = sqrt(Offset.x*Offset.x + Offset.y*Offset.y);
+				float Weight = exp(-(r*r) / s) / (3.14f * s);
+				OutColor += MipInSRV.SampleLevel(MipInSampler, UV + (Offset * SrcTexelSize * 2), 0) * Weight;
+				TotalWeight += Weight;
+			}
+		}
+		OutColor /= TotalWeight;
+	}
+#else
+	{
+		OutColor = MipInSRV.SampleLevel(MipInSampler, UV, 0);
+	}
+#endif
 
 #if GENMIPS_SRGB
     OutColor.rgb = LinearToSrgb(OutColor.rgb);
diff --git a/Engine/Plugins/FX/Niagara/Source/Niagara/Classes/NiagaraDataInterfaceRenderTarget2D.h b/Engine/Plugins/FX/Niagara/Source/Niagara/Classes/NiagaraDataInterfaceRenderTarget2D.h
--- a/Engine/Plugins/FX/Niagara/Source/Niagara/Classes/NiagaraDataInterfaceRenderTarget2D.h
+++ b/Engine/Plugins/FX/Niagara/Source/Niagara/Classes/NiagaraDataInterfaceRenderTarget2D.h
@@ -5,6 +5,7 @@
 #include "NiagaraDataInterfaceRW.h"
 #include "ClearQuad.h"
 #include "NiagaraComponent.h"
+#include "NiagaraGenerateMips.h"
 #include "NiagaraDataInterfaceRenderTarget2D.generated.h"
 
 class FNiagaraSystemInstance;
@@ -23,6 +24,7 @@ struct FRenderTarget2DRWInstanceData_GameThread
 	FIntPoint Size = FIntPoint(EForceInit::ForceInitToZero);
 	ETextureRenderTargetFormat Format = RTF_RGBA16f;
 	ENiagaraMipMapGeneration MipMapGeneration = ENiagaraMipMapGeneration::Disabled;
+	ENiagaraMipMapGenerationType MipMapGenerationType = ENiagaraMipMapGenerationType::Linear;
 	
 	UTextureRenderTarget2D* TargetTexture = nullptr;
 #if WITH_EDITORONLY_DATA
@@ -42,6 +44,7 @@ struct FRenderTarget2DRWInstanceData_RenderThread
 
 	FIntPoint Size = FIntPoint(EForceInit::ForceInitToZero);
 	ENiagaraMipMapGeneration MipMapGeneration = ENiagaraMipMapGeneration::Disabled;
+	ENiagaraMipMapGenerationType MipMapGenerationType = ENiagaraMipMapGenerationType::Linear;
 	bool bWasWrittenTo = false;
 
 	FSamplerStateRHIRef SamplerStateRHI;
@@ -140,6 +143,9 @@ public:
 	UPROPERTY(EditAnywhere, Category = "Render Target", meta = (EditCondition = "!bInheritUserParameterSettings"))
 	ENiagaraMipMapGeneration MipMapGeneration = ENiagaraMipMapGeneration::Disabled;
 
+	UPROPERTY(EditAnywhere, Category = "Render Target")
+	ENiagaraMipMapGenerationType MipMapGenerationType = ENiagaraMipMapGenerationType::Linear;
+
 	/** When enabled overrides the format of the render target, otherwise uses the project default setting. */
 	UPROPERTY(EditAnywhere, Category = "Render Target", meta = (EditCondition = "!bInheritUserParameterSettings && bOverrideFormat"))
 	TEnumAsByte<ETextureRenderTargetFormat> OverrideRenderTargetFormat;
diff --git a/Engine/Plugins/FX/Niagara/Source/Niagara/Private/NiagaraDataInterfaceRenderTarget2D.cpp b/Engine/Plugins/FX/Niagara/Source/Niagara/Private/NiagaraDataInterfaceRenderTarget2D.cpp
--- a/Engine/Plugins/FX/Niagara/Source/Niagara/Private/NiagaraDataInterfaceRenderTarget2D.cpp
+++ b/Engine/Plugins/FX/Niagara/Source/Niagara/Private/NiagaraDataInterfaceRenderTarget2D.cpp
@@ -16,7 +16,6 @@
 #if WITH_EDITOR
 #include "NiagaraGpuComputeDebug.h"
 #endif
-#include "NiagaraGenerateMips.h"
 #include "NiagaraStats.h"
 
 #define LOCTEXT_NAMESPACE "NiagaraDataInterfaceRenderTarget2D"
@@ -400,6 +399,7 @@ bool UNiagaraDataInterfaceRenderTarget2D::Equals(const UNiagaraDataInterface* Ot
 		OtherTyped->RenderTargetUserParameter == RenderTargetUserParameter &&
 		OtherTyped->Size == Size &&
 		OtherTyped->MipMapGeneration == MipMapGeneration &&
+		OtherTyped->MipMapGenerationType == MipMapGenerationType &&
 		OtherTyped->OverrideRenderTargetFormat == OverrideRenderTargetFormat &&
 		OtherTyped->bInheritUserParameterSettings == bInheritUserParameterSettings &&
 		OtherTyped->bOverrideFormat == bOverrideFormat;
@@ -420,6 +420,7 @@ bool UNiagaraDataInterfaceRenderTarget2D::CopyToInternal(UNiagaraDataInterface*
 
 	DestinationTyped->Size = Size;
 	DestinationTyped->MipMapGeneration = MipMapGeneration;
+	DestinationTyped->MipMapGenerationType = MipMapGenerationType;
 	DestinationTyped->OverrideRenderTargetFormat = OverrideRenderTargetFormat;
 	DestinationTyped->bInheritUserParameterSettings = bInheritUserParameterSettings;
 	DestinationTyped->bOverrideFormat = bOverrideFormat;
@@ -542,6 +543,7 @@ bool UNiagaraDataInterfaceRenderTarget2D::InitPerInstanceData(void* PerInstanceD
 	InstanceData->Size.X = FMath::Clamp<int>(int(float(Size.X) * GNiagaraRenderTargetResolutionMultiplier), 1, GMaxTextureDimensions);
 	InstanceData->Size.Y = FMath::Clamp<int>(int(float(Size.Y) * GNiagaraRenderTargetResolutionMultiplier), 1, GMaxTextureDimensions);
 	InstanceData->MipMapGeneration = MipMapGeneration;
+	InstanceData->MipMapGenerationType = MipMapGenerationType;
 	InstanceData->Format = bOverrideFormat ? OverrideRenderTargetFormat : GetDefault<UNiagaraSettings>()->DefaultRenderTargetFormat;
 	InstanceData->RTUserParamBinding.Init(SystemInstance->GetInstanceParameters(), RenderTargetUserParameter.Parameter);
 #if WITH_EDITORONLY_DATA
@@ -702,10 +704,12 @@ bool UNiagaraDataInterfaceRenderTarget2D::PerInstanceTick(void* PerInstanceData,
 			{
 				// We have to take a guess at user intention
 				InstanceData->MipMapGeneration = MipMapGeneration == ENiagaraMipMapGeneration::Disabled ? ENiagaraMipMapGeneration::PostStage : MipMapGeneration;
+				InstanceData->MipMapGenerationType = MipMapGenerationType;
 			}
 			else
 			{
 				InstanceData->MipMapGeneration = ENiagaraMipMapGeneration::Disabled;
+				InstanceData->MipMapGenerationType = ENiagaraMipMapGenerationType::Unfiltered;
 			}
 			InstanceData->Format = InstanceData->TargetTexture->RenderTargetFormat;
 		}
@@ -772,6 +776,7 @@ bool UNiagaraDataInterfaceRenderTarget2D::PerInstanceTickPostSimulate(void* PerI
 				FRenderTarget2DRWInstanceData_RenderThread* TargetData = &RT_Proxy->SystemInstancesToProxyData_RT.FindOrAdd(RT_InstanceID);
 				TargetData->Size = RT_InstanceData.Size;
 				TargetData->MipMapGeneration = RT_InstanceData.MipMapGeneration;
+				TargetData->MipMapGenerationType = RT_InstanceData.MipMapGenerationType;
 			#if WITH_EDITORONLY_DATA
 				TargetData->bPreviewTexture = RT_InstanceData.bPreviewTexture;
 			#endif
@@ -804,7 +809,7 @@ void FNiagaraDataInterfaceProxyRenderTarget2DProxy::PostStage(FRHICommandList& R
 		if (ProxyData->bWasWrittenTo && (ProxyData->MipMapGeneration == ENiagaraMipMapGeneration::PostStage))
 		{
 			ProxyData->bWasWrittenTo = false;
-			NiagaraGenerateMips::GenerateMips(RHICmdList, ProxyData->TextureRHI);
+			NiagaraGenerateMips::GenerateMips(RHICmdList, ProxyData->TextureRHI, ProxyData->MipMapGenerationType);
 		}
 	}
 }
@@ -816,7 +821,7 @@ void FNiagaraDataInterfaceProxyRenderTarget2DProxy::PostSimulate(FRHICommandList
 		if (ProxyData->bWasWrittenTo && (ProxyData->MipMapGeneration == ENiagaraMipMapGeneration::PostSimulate))
 		{
 			ProxyData->bWasWrittenTo = false;
-			NiagaraGenerateMips::GenerateMips(RHICmdList, ProxyData->TextureRHI);
+			NiagaraGenerateMips::GenerateMips(RHICmdList, ProxyData->TextureRHI, ProxyData->MipMapGenerationType);
 		}
 
 #if NIAGARA_COMPUTEDEBUG_ENABLED
diff --git a/Engine/Plugins/FX/Niagara/Source/NiagaraShader/Private/NiagaraGenerateMips.cpp b/Engine/Plugins/FX/Niagara/Source/NiagaraShader/Private/NiagaraGenerateMips.cpp
--- a/Engine/Plugins/FX/Niagara/Source/NiagaraShader/Private/NiagaraGenerateMips.cpp
+++ b/Engine/Plugins/FX/Niagara/Source/NiagaraShader/Private/NiagaraGenerateMips.cpp
@@ -12,10 +12,13 @@ public:
 	SHADER_USE_PARAMETER_STRUCT(FNiagaraGenerateMipsCS, FGlobalShader)
 
 	class FGenMipsSRGB : SHADER_PERMUTATION_BOOL("GENMIPS_SRGB");
-	using FPermutationDomain = TShaderPermutationDomain<FGenMipsSRGB>;
+	class FGaussianBlur : SHADER_PERMUTATION_BOOL("GENMIPS_GAUSSIAN");
+	using FPermutationDomain = TShaderPermutationDomain<FGenMipsSRGB, FGaussianBlur>;
 
 	BEGIN_SHADER_PARAMETER_STRUCT(FParameters, )
-		SHADER_PARAMETER(FVector2D, TexelSize)
+		SHADER_PARAMETER(FVector2D, SrcTexelSize)
+		SHADER_PARAMETER(FVector2D, DstTexelSize)
+		SHADER_PARAMETER(int32, KernelHWidth)
 		SHADER_PARAMETER(FIntPoint, MipOutSize)
 		SHADER_PARAMETER_UAV(RWTexture2D<float4>, MipOutUAV)
 		SHADER_PARAMETER_SRV(Texture2D, MipInSRV)
@@ -36,7 +39,7 @@ public:
 
 IMPLEMENT_GLOBAL_SHADER(FNiagaraGenerateMipsCS, "/Plugin/FX/Niagara/Private/NiagaraGenerateMips.usf", "MainCS", SF_Compute);
 
-void NiagaraGenerateMips::GenerateMips(FRHICommandList& RHICmdList, FRHITexture2D* TextureRHI)
+void NiagaraGenerateMips::GenerateMips(FRHICommandList& RHICmdList, FRHITexture2D* TextureRHI, ENiagaraMipMapGenerationType GenType)
 {
 	SCOPED_DRAW_EVENT(RHICmdList, NiagaraGenerateMips);
 
@@ -50,11 +53,16 @@ void NiagaraGenerateMips::GenerateMips(FRHICommandList& RHICmdList, FRHITexture2
 	}
 #endif
 
+	const bool bBlurMips = GenType >= ENiagaraMipMapGenerationType::Blur1;
+
 	FNiagaraGenerateMipsCS::FPermutationDomain PermutationVector;
 	PermutationVector.Set<FNiagaraGenerateMipsCS::FGenMipsSRGB>(bMipsSRGB);
+	PermutationVector.Set<FNiagaraGenerateMipsCS::FGaussianBlur>(bBlurMips);
 	TShaderMapRef<FNiagaraGenerateMipsCS> ComputeShader(GetGlobalShaderMap(GMaxRHIFeatureLevel), PermutationVector);
 	FRHIComputeShader* ShaderRHI = ComputeShader.GetComputeShader();
 
+	const int32 KernelHWidth = bBlurMips ? (int32(GenType) + 1 - int32(ENiagaraMipMapGenerationType::Blur1)) : 1;
+
 	const int32 NumMips = TextureRHI->GetNumMips();
 	for ( int32 iDstMip=1; iDstMip < NumMips; ++iDstMip)
 	{
@@ -66,10 +74,12 @@ void NiagaraGenerateMips::GenerateMips(FRHICommandList& RHICmdList, FRHITexture2
 		FUnorderedAccessViewRHIRef MipOutUAV = RHICreateUnorderedAccessView(TextureRHI, iDstMip);
 
 		FNiagaraGenerateMipsCS::FParameters PassParameters;
-		PassParameters.TexelSize	= FVector2D(1.0f / float(DstMipSize.X), 1.0f / float(DstMipSize.Y));
+		PassParameters.SrcTexelSize	= FVector2D(1.0f / float(SrcMipSize.X), 1.0f / float(SrcMipSize.Y));
+		PassParameters.DstTexelSize = FVector2D(1.0f / float(DstMipSize.X), 1.0f / float(DstMipSize.Y));
+		PassParameters.KernelHWidth = KernelHWidth;
 		PassParameters.MipOutSize	= DstMipSize;
 		PassParameters.MipInSRV		= MipInSRV;
-		PassParameters.MipInSampler	= TStaticSamplerState<SF_Bilinear>::GetRHI();
+		PassParameters.MipInSampler	= GenType == ENiagaraMipMapGenerationType::Unfiltered ? TStaticSamplerState<SF_Point>::GetRHI() : TStaticSamplerState<SF_Bilinear>::GetRHI();
 		PassParameters.MipOutUAV	= MipOutUAV;
 
 		const FIntVector NumThreadGroups = FComputeShaderUtils::GetGroupCount(DstMipSize, FComputeShaderUtils::kGolden2DGroupSize);
diff --git a/Engine/Plugins/FX/Niagara/Source/NiagaraShader/Public/NiagaraGenerateMips.h b/Engine/Plugins/FX/Niagara/Source/NiagaraShader/Public/NiagaraGenerateMips.h
--- a/Engine/Plugins/FX/Niagara/Source/NiagaraShader/Public/NiagaraGenerateMips.h
+++ b/Engine/Plugins/FX/Niagara/Source/NiagaraShader/Public/NiagaraGenerateMips.h
@@ -3,8 +3,26 @@
 #pragma once
 
 #include "RHI.h"
+#include "NiagaraGenerateMips.generated.h"
+
+UENUM()
+enum class ENiagaraMipMapGenerationType : uint8
+{
+	/** Result is taken from whatever texel the sample lies on, aka Point Sampling. */
+	Unfiltered,
+	/** Linear blending is performed between a 2x2 (or 2x2x2) region of texels, aka Bilinear / Trilinear. */
+	Linear,
+	/** A blur filter across a 3x3 (or 3x3x3) region of texels. */
+	Blur1 UMETA(DisplayName="Gaussian 3 texel filter"),
+	/** A blur filter across a 5x5 (or 5x5x5) region of texels. */
+	Blur2 UMETA(DisplayName="Gaussian 5 texel filter"),
+	/** A blur filter across a 7x7 (or 7x7x7) region of texels. */
+	Blur3 UMETA(DisplayName="Gaussian 7 texel filter"),
+	/** A blur filter across a 9x9 (or 9x9x9) region of texels. */
+	Blur4 UMETA(DisplayName="Gaussian 9 texel filter"),
+};
 
 namespace NiagaraGenerateMips
 {
-	NIAGARASHADER_API void GenerateMips(FRHICommandList& RHICmdList, FRHITexture2D* TextureRHI);
+	NIAGARASHADER_API void GenerateMips(FRHICommandList& RHICmdList, FRHITexture2D* TextureRHI, ENiagaraMipMapGenerationType GenType);
 }
