Added functionality to pin generated data for skeletal mesh connnectivity and uv mapping
-includes optimizations for generation of adjacency buffer (65% faster)

diff --git a/Engine/Plugins/FX/Niagara/Source/Niagara/Classes/NiagaraDataInterfaceSkeletalMesh.h b/Engine/Plugins/FX/Niagara/Source/Niagara/Classes/NiagaraDataInterfaceSkeletalMesh.h
--- a/Engine/Plugins/FX/Niagara/Source/Niagara/Classes/NiagaraDataInterfaceSkeletalMesh.h
+++ b/Engine/Plugins/FX/Niagara/Source/Niagara/Classes/NiagaraDataInterfaceSkeletalMesh.h
@@ -255,7 +255,7 @@ struct FSkeletalMeshUvMappingUsage
 	bool RequiresGpuAccess = false;
 };
 
-struct FSkeletalMeshUvMappingHandle
+struct NIAGARA_API FSkeletalMeshUvMappingHandle
 {
 	FSkeletalMeshUvMappingHandle();
 	FSkeletalMeshUvMappingHandle(FSkeletalMeshUvMappingUsage InUsage, const TSharedPtr<struct FSkeletalMeshUvMapping>& InUvMappingData, bool bNeedsDataImmediately);
@@ -276,6 +276,8 @@ struct FSkeletalMeshUvMappingHandle
 	int32 GetUvSetIndex() const;
 	int32 GetLodIndex() const;
 
+	void PinAndInvalidateHandle();
+
 private:
 	TSharedPtr<FSkeletalMeshUvMapping> UvMappingData;
 };
@@ -294,7 +296,7 @@ struct FSkeletalMeshConnectivityUsage
 	bool RequiresGpuAccess = false;
 };
 
-struct FSkeletalMeshConnectivityHandle
+struct NIAGARA_API FSkeletalMeshConnectivityHandle
 {
 	FSkeletalMeshConnectivityHandle();
 	FSkeletalMeshConnectivityHandle(FSkeletalMeshConnectivityUsage InUsage, const TSharedPtr<struct FSkeletalMeshConnectivity>& InConnectivityData, bool bNeedsDataImmediately);
@@ -312,11 +314,13 @@ struct FSkeletalMeshConnectivityHandle
 
 	const FSkeletalMeshConnectivityProxy* GetProxy() const;
 
+	void PinAndInvalidateHandle();
+
 private:
 	TSharedPtr<FSkeletalMeshConnectivity> ConnectivityData;
 };
 
-class FNDI_SkeletalMesh_GeneratedData : public FNDI_GeneratedData
+class NIAGARA_API FNDI_SkeletalMesh_GeneratedData : public FNDI_GeneratedData
 {
 	FRWLock CachedSkinningDataGuard;
 	TMap<TWeakObjectPtr<USkeletalMeshComponent>, TSharedPtr<FSkeletalMeshSkinningData> > CachedSkinningData;
diff --git a/Engine/Plugins/FX/Niagara/Source/Niagara/Private/NiagaraDataInterfaceSkeletalMeshConnectivity.cpp b/Engine/Plugins/FX/Niagara/Source/Niagara/Private/NiagaraDataInterfaceSkeletalMeshConnectivity.cpp
--- a/Engine/Plugins/FX/Niagara/Source/Niagara/Private/NiagaraDataInterfaceSkeletalMeshConnectivity.cpp
+++ b/Engine/Plugins/FX/Niagara/Source/Niagara/Private/NiagaraDataInterfaceSkeletalMeshConnectivity.cpp
@@ -11,6 +11,8 @@
 
 #include <limits>
 
+DECLARE_CYCLE_STAT(TEXT("Niagara - SkelMesh - BuildAdjacencyBuffer"), STAT_NiagaraSkel_Connectivity_Adjacency, STATGROUP_Niagara);
+
 FSkeletalMeshConnectivityHandle::FSkeletalMeshConnectivityHandle()
 	: ConnectivityData(nullptr)
 {
@@ -76,6 +78,11 @@ const FSkeletalMeshConnectivityProxy* FSkeletalMeshConnectivityHandle::GetProxy(
 	return nullptr;
 }
 
+void FSkeletalMeshConnectivityHandle::PinAndInvalidateHandle()
+{
+	ConnectivityData.Reset();
+}
+
 FSkeletalMeshConnectivity::FSkeletalMeshConnectivity(TWeakObjectPtr<USkeletalMesh> InMeshObject, int32 InLodIndex)
 	: LodIndex(InLodIndex)
 	, MeshObject(InMeshObject)
@@ -83,6 +90,11 @@ FSkeletalMeshConnectivity::FSkeletalMeshConnectivity(TWeakObjectPtr<USkeletalMes
 {
 }
 
+FSkeletalMeshConnectivity::~FSkeletalMeshConnectivity()
+{
+	Release();
+}
+
 bool FSkeletalMeshConnectivity::IsUsed() const
 {
 	return (GpuUserCount > 0);
@@ -101,8 +113,14 @@ void FSkeletalMeshConnectivity::RegisterUser(FSkeletalMeshConnectivityUsage Usag
 		{
 			check(Proxy == nullptr);
 			Proxy.Reset(new FSkeletalMeshConnectivityProxy());
-			Proxy->Initialize(*this);
-			BeginInitResource(Proxy.Get());
+			if (Proxy->Initialize(*this))
+			{
+				BeginInitResource(Proxy.Get());
+			}
+			else
+			{
+				Proxy.Release();
+			}
 		}
 	}
 }
@@ -113,17 +131,23 @@ void FSkeletalMeshConnectivity::UnregisterUser(FSkeletalMeshConnectivityUsage Us
 	{
 		if (--GpuUserCount == 0)
 		{
-			ENQUEUE_RENDER_COMMAND(BeginDestroyCommand)(
-				[RT_Proxy=Proxy.Release()](FRHICommandListImmediate& RHICmdList)
-				{
-					RT_Proxy->ReleaseResource();
-					delete RT_Proxy;
-				}
-			);
+			Release();
 		}
 	}
 }
 
+void FSkeletalMeshConnectivity::Release()
+{
+	if (FSkeletalMeshConnectivityProxy* ProxyPtr = Proxy.Release())
+	{
+		ENQUEUE_RENDER_COMMAND(BeginDestroyCommand)([RT_Proxy = ProxyPtr](FRHICommandListImmediate& RHICmdList)
+		{
+			RT_Proxy->ReleaseResource();
+			delete RT_Proxy;
+		});
+	}
+}
+
 bool FSkeletalMeshConnectivity::CanBeUsed(const TWeakObjectPtr<USkeletalMesh>& InMeshObject, int32 InLodIndex) const
 {
 	return IsUsed() && LodIndex == InLodIndex && MeshObject == InMeshObject;
@@ -227,8 +251,10 @@ bool operator==(const FAdjacencyVertexOverlapKey& Lhs, const FAdjacencyVertexOve
 template<typename TriangleIndexType, bool SortBySize>
 static bool BuildAdjacencyBuffer(const FSkeletalMeshLODRenderData& LodRenderData, int32 MaxAdjacencyCount, TResourceArray<uint8>& Buffer, int32& MaxFoundAdjacentTriangleCount)
 {
+	SCOPE_CYCLE_COUNTER(STAT_NiagaraSkel_Connectivity_Adjacency);
+
 	const FRawStaticIndexBuffer16or32Interface* IndexBuffer = LodRenderData.MultiSizeIndexContainer.GetIndexBuffer();
-	const FPositionVertexBuffer* VertexBuffer = SortBySize ? &LodRenderData.StaticVertexBuffers.PositionVertexBuffer : nullptr;
+	const FPositionVertexBuffer& VertexBuffer = LodRenderData.StaticVertexBuffers.PositionVertexBuffer;
 	const uint32 IndexCount = IndexBuffer->Num();
 	const uint32 TriangleCount = IndexCount / 3;
 
@@ -239,127 +265,103 @@ static bool BuildAdjacencyBuffer(const FSkeletalMeshLODRenderData& LodRenderData
 		return false;
 	}
 
-	int32 MinVertexValue = std::numeric_limits<int32>::max();
-	int32 MaxVertexValue = std::numeric_limits<int32>::min();
-
-	int32 UniqueVertexCount = 0;
-	TMultiMap<int32 /*UniqueVertexIndex*/, TriangleIndexType> RawAdjacency;
-	TMap<FAdjacencyVertexOverlapKey, int32 /*UniqueVertexIndex*/> UniqueIndexMap;
-	TMap<int32 /*VertexIndex*/, int32 /*UniqueVertexIndex*/> VertexToUniqueIndexMap;
-
-	for (uint32 IndexIt = 0; IndexIt < IndexCount; ++IndexIt)
-	{
-		const TriangleIndexType TriangleId = IndexIt / 3;
-		const int32 VertexId = IndexBuffer->Get(IndexIt);
-		MinVertexValue = FMath::Min(MinVertexValue, VertexId);
-		MaxVertexValue = FMath::Max(MaxVertexValue, VertexId);
-
-		FAdjacencyVertexOverlapKey OverlapKey;
-		OverlapKey.Position = VertexBuffer->VertexPosition(VertexId);
+	const uint32 VertexCount = LodRenderData.StaticVertexBuffers.PositionVertexBuffer.GetNumVertices();
+	const int32 SizePerVertex = MaxAdjacencyCount * sizeof(TriangleIndexType);
+	const int32 BufferSize = VertexCount * SizePerVertex;
+	const int32 PaddedBufferSize = 4 * FMath::DivideAndRoundUp(BufferSize, 4);
+	Buffer.SetNum(PaddedBufferSize, true);
+	FMemory::Memset(Buffer.GetData(), 0xFF, Buffer.Num());
+	TriangleIndexType* AdjacencyBuffer = reinterpret_cast<TriangleIndexType*>(Buffer.GetData());
 
-		int32 UniqueIndex = UniqueVertexCount;
-		if (int32* ExistingGroup = UniqueIndexMap.Find(OverlapKey))
-		{
-			UniqueIndex = *ExistingGroup;
-		}
-		else
-		{
-			UniqueIndexMap.Add(OverlapKey, UniqueVertexCount);
-			++UniqueVertexCount;
-		}
+	TMap<FVector, int32 /*UniqueVertexIndex*/> UniqueIndexMap;
 
-		int32* ExistingUniqueIndex = VertexToUniqueIndexMap.Find(VertexId);
-		if (!ExistingUniqueIndex)
-		{
-			VertexToUniqueIndexMap.Add(VertexId, UniqueIndex);
-		}
-		else
-		{
-			check(*ExistingUniqueIndex == UniqueIndex);
-		}
+	TArray<uint32> RedirectionArray;
+	TArray<float> TriangleSizes;
 
-		RawAdjacency.Add(UniqueIndex, TriangleId);
+	RedirectionArray.SetNum(VertexCount);
+	if (SortBySize)
+	{
+		TriangleSizes.SetNum(TriangleCount);
 	}
 
-	const int32 SizePerVertex = MaxAdjacencyCount * sizeof(TriangleIndexType);
-	const int32 BufferSize = (MaxVertexValue - MinVertexValue) * SizePerVertex;
-	const int32 PaddedBufferSize = 4 * FMath::DivideAndRoundUp(BufferSize, 4);
-	Buffer.Init(0xFF, PaddedBufferSize);
-
-	MaxFoundAdjacentTriangleCount = 0;
-
+	for (TriangleIndexType TriangleIt = 0; TriangleIt < TriangleCount; ++TriangleIt)
 	{
-		FNiagaraResourceArrayWriter Ar(Buffer);
+		const uint32 V[3] = 
+		{
+			IndexBuffer->Get(TriangleIt * 3 + 0),
+			IndexBuffer->Get(TriangleIt * 3 + 1),
+			IndexBuffer->Get(TriangleIt * 3 + 2)
+		};
 
-		TArray<TriangleIndexType> TriangleValues;
-		TArray<TriangleIndexType> SortedValues;
-		TArray<float> TriangleSizes;
-		TArray<int32> SortIndices;
+		const FVector P[3] =
+		{
+			VertexBuffer.VertexPosition(V[0]),
+			VertexBuffer.VertexPosition(V[1]),
+			VertexBuffer.VertexPosition(V[2])
+		};
 
-		TriangleValues.Reserve(MaxAdjacencyCount);
 		if (SortBySize)
 		{
-			TriangleSizes.Reserve(MaxAdjacencyCount);
-			SortIndices.Reserve(MaxAdjacencyCount);
-			SortedValues.Reserve(MaxAdjacencyCount);
+			TriangleSizes[TriangleIt] = 0.5f * ((P[2] - P[0]) ^ (P[1] - P[0])).Size();
 		}
 
-		for (int32 VertexIt = MinVertexValue; VertexIt <= MaxVertexValue; ++VertexIt)
+		for (int32 i = 0; i < 3; ++i)
 		{
-			Ar.Seek(VertexIt * SizePerVertex);
-
-			TriangleValues.Reset();
+			const uint32 VertexIndex = RedirectionArray[V[i]] = UniqueIndexMap.FindOrAdd(P[i], V[i]);
 
-			{
-				RawAdjacency.MultiFind(VertexToUniqueIndexMap[VertexIt], TriangleValues);
-			
-				TriangleValues.Sort();
-				TriangleValues.SetNum(Algo::Unique(TriangleValues));
-			}
+			TArrayView<TriangleIndexType> AdjacentTriangles = MakeArrayView(AdjacencyBuffer + VertexIndex * MaxAdjacencyCount, MaxAdjacencyCount);
 
-			MaxFoundAdjacentTriangleCount = FMath::Max(MaxFoundAdjacentTriangleCount, TriangleValues.Num());
-
-			if (SortBySize)
+			int32 InsertionPoint = 0;
+			while (InsertionPoint < MaxAdjacencyCount)
 			{
-				const int32 AdjacentCount = TriangleValues.Num();
+				const TriangleIndexType TestTriangle = AdjacentTriangles[InsertionPoint];
 
-				TriangleSizes.Reset(AdjacentCount);
-				SortIndices.Reset(AdjacentCount);
-				SortedValues.Reset(AdjacentCount);
-				
-				for (int32 TriangleIt = 0; TriangleIt < AdjacentCount; ++TriangleIt)
+				if (TestTriangle == INDEX_NONE)
 				{
-					const int32 TriangleIndex = TriangleValues[TriangleIt];
-
-					const FVector& v0 = VertexBuffer->VertexPosition(IndexBuffer->Get(TriangleIndex * 3 + 0));
-					const FVector& v1 = VertexBuffer->VertexPosition(IndexBuffer->Get(TriangleIndex * 3 + 1));
-					const FVector& v2 = VertexBuffer->VertexPosition(IndexBuffer->Get(TriangleIndex * 3 + 2));
-
-					const float TriangleSize = 0.5f * ((v2 - v0) ^ (v1 - v0)).Size();
-					TriangleSizes.Add(TriangleSize);
-					SortIndices.Add(TriangleIt);
+					AdjacentTriangles[InsertionPoint] = TriangleIt;
+					break;
 				}
 
-				Algo::StableSort(SortIndices, [&](int32 Lhs, int32 Rhs){ return TriangleSizes[Lhs] > TriangleSizes[Rhs]; });
-
-				for (int32 TriangleIt = 0; TriangleIt < FMath::Min(AdjacentCount, MaxAdjacencyCount); ++TriangleIt)
+				if ((SortBySize && (TriangleSizes[TriangleIt] > TriangleSizes[TestTriangle]))
+					|| (!SortBySize && TriangleIt < TestTriangle))
 				{
-					SortedValues.Add(TriangleValues[SortIndices[TriangleIt]]);
+					// skip empty entries
+					int32 ShiftIt = MaxAdjacencyCount - 1;
+					while (AdjacentTriangles[ShiftIt - 1] == INDEX_NONE)
+					{
+						--ShiftIt;
+					}
+
+					// shift the results down and then insert
+					do 
+					{
+						AdjacentTriangles[ShiftIt] = AdjacentTriangles[ShiftIt - 1];
+						--ShiftIt;
+					} while (ShiftIt > InsertionPoint);
+
+					AdjacentTriangles[InsertionPoint] = TriangleIt;
+					break;
 				}
 
-				Swap(SortedValues, TriangleValues);
+				++InsertionPoint;
 			}
+		}
+	}
 
-			TriangleIndexType AdjacentTriangleCount = FMath::Clamp<TriangleIndexType>(TriangleValues.Num(), 0, MaxAdjacencyCount);
-
-			Ar.Serialize(TriangleValues.GetData(), AdjacentTriangleCount * sizeof(TriangleIndexType));
+	for (uint32 VertexIt = 1; VertexIt < VertexCount; ++VertexIt)
+	{
+		// if this vertex has a sibling we just copy the data over
+		const int32 SiblingIndex = RedirectionArray[VertexIt];
+		if (SiblingIndex != VertexIt)
+		{
+			FMemory::Memcpy(AdjacencyBuffer + VertexIt * MaxAdjacencyCount, AdjacencyBuffer + SiblingIndex * MaxAdjacencyCount, SizePerVertex);
 		}
 	}
 
 	return true;
 }
 
-void
+bool
 FSkeletalMeshConnectivityProxy::Initialize(const FSkeletalMeshConnectivity& Connectivity)
 {
 	if (const FSkeletalMeshLODRenderData* LodRenderData = Connectivity.GetLodRenderData())
@@ -387,7 +389,11 @@ FSkeletalMeshConnectivityProxy::Initialize(const FSkeletalMeshConnectivity& Conn
 			UE_LOG(LogNiagara, Warning, TEXT("Max adjacency limit of %d exceeded (up to %d found) when processing %s.  Some connections will be ignored."),
 				MaxAdjacentTriangleCount, MaxFoundAdjacentTriangleCount, *Connectivity.GetMeshName());
 		}
+
+		return AdjacencySuccess;
 	}
+
+	return false;
 }
 
 void
diff --git a/Engine/Plugins/FX/Niagara/Source/Niagara/Private/NiagaraDataInterfaceSkeletalMeshConnectivity.h b/Engine/Plugins/FX/Niagara/Source/Niagara/Private/NiagaraDataInterfaceSkeletalMeshConnectivity.h
--- a/Engine/Plugins/FX/Niagara/Source/Niagara/Private/NiagaraDataInterfaceSkeletalMeshConnectivity.h
+++ b/Engine/Plugins/FX/Niagara/Source/Niagara/Private/NiagaraDataInterfaceSkeletalMeshConnectivity.h
@@ -10,7 +10,7 @@ struct FSkeletalMeshConnectivity;
 class FSkeletalMeshConnectivityProxy : public FRenderResource
 {
 public:
-	void Initialize(const FSkeletalMeshConnectivity& ConnectivityData);
+	bool Initialize(const FSkeletalMeshConnectivity& ConnectivityData);
 
 	virtual void InitRHI() override;
 	virtual void ReleaseRHI() override;
@@ -36,6 +36,7 @@ struct FSkeletalMeshConnectivity
 	FSkeletalMeshConnectivity() = delete;
 	FSkeletalMeshConnectivity(const FSkeletalMeshConnectivity&) = delete;
 	FSkeletalMeshConnectivity(TWeakObjectPtr<USkeletalMesh> InMeshObject, int32 InLodIndex);
+	~FSkeletalMeshConnectivity();
 
 	bool IsUsed() const;
 	bool CanBeDestroyed() const;
@@ -56,6 +57,8 @@ struct FSkeletalMeshConnectivity
 private:
 	static const FSkeletalMeshLODRenderData* GetLodRenderData(const USkeletalMesh& Mesh, int32 LodIndex);
 
+	void Release();
+
 	TWeakObjectPtr<USkeletalMesh> MeshObject;
 	TUniquePtr<FSkeletalMeshConnectivityProxy> Proxy;
 
diff --git a/Engine/Plugins/FX/Niagara/Source/Niagara/Private/NiagaraDataInterfaceSkeletalMeshUvMapping.cpp b/Engine/Plugins/FX/Niagara/Source/Niagara/Private/NiagaraDataInterfaceSkeletalMeshUvMapping.cpp
--- a/Engine/Plugins/FX/Niagara/Source/Niagara/Private/NiagaraDataInterfaceSkeletalMeshUvMapping.cpp
+++ b/Engine/Plugins/FX/Niagara/Source/Niagara/Private/NiagaraDataInterfaceSkeletalMeshUvMapping.cpp
@@ -6,6 +6,8 @@
 #include "NiagaraResourceArrayWriter.h"
 #include "NiagaraStats.h"
 
+DECLARE_CYCLE_STAT(TEXT("Niagara - SkelMesh - UvQuadTree Cpu"), STAT_NiagaraSkel_UvMapping_Cpu, STATGROUP_Niagara);
+DECLARE_CYCLE_STAT(TEXT("Niagara - SkelMesh - UvQuadTree Gpu"), STAT_NiagaraSkel_UvMapping_Gpu, STATGROUP_Niagara);
 
 template<bool UseFullPrecisionUv>
 struct FQuadTreeQueryHelper
@@ -277,6 +279,10 @@ int32 FSkeletalMeshUvMappingHandle::GetLodIndex() const
 	return 0;
 }
 
+void FSkeletalMeshUvMappingHandle::PinAndInvalidateHandle()
+{
+	UvMappingData.Reset();
+}
 
 FSkeletalMeshUvMapping::FSkeletalMeshUvMapping(TWeakObjectPtr<USkeletalMesh> InMeshObject, int32 InLodIndex, int32 InUvSetIndex)
 	: LodIndex(InLodIndex)
@@ -288,6 +294,12 @@ FSkeletalMeshUvMapping::FSkeletalMeshUvMapping(TWeakObjectPtr<USkeletalMesh> InM
 {
 }
 
+FSkeletalMeshUvMapping::~FSkeletalMeshUvMapping()
+{
+	ReleaseQuadTree();
+	ReleaseGpuQuadTree();
+}
+
 template<bool UseFullPrecisionUv>
 static void BuildQuadTreeHelper(FNiagaraUvQuadTree& QuadTree, const FSkeletalMeshLODRenderData* LodRenderData, int32 UvSetIndex)
 {
@@ -318,6 +330,7 @@ void FSkeletalMeshUvMapping::BuildQuadTree()
 {
 	if (const FSkeletalMeshLODRenderData* LodRenderData = GetLodRenderData())
 	{
+		SCOPE_CYCLE_COUNTER(STAT_NiagaraSkel_UvMapping_Cpu);
 		if (LodRenderData->StaticVertexBuffers.StaticMeshVertexBuffer.GetUseFullPrecisionUVs())
 		{
 			BuildQuadTreeHelper<true>(TriangleIndexQuadTree, LodRenderData, UvSetIndex);
@@ -342,6 +355,8 @@ void FSkeletalMeshUvMapping::ReleaseQuadTree()
 
 void FSkeletalMeshUvMapping::BuildGpuQuadTree()
 {
+	SCOPE_CYCLE_COUNTER(STAT_NiagaraSkel_UvMapping_Gpu);
+
 	check(FrozenQuadTreeProxy == nullptr);
 	FrozenQuadTreeProxy.Reset(new FSkeletalMeshUvMappingBufferProxy());
 	FrozenQuadTreeProxy->Initialize(*this);
@@ -350,12 +365,14 @@ void FSkeletalMeshUvMapping::BuildGpuQuadTree()
 
 void FSkeletalMeshUvMapping::ReleaseGpuQuadTree()
 {
-	ENQUEUE_RENDER_COMMAND(BeginDestroyCommand)(
-		[RT_Proxy=FrozenQuadTreeProxy.Release()](FRHICommandListImmediate& RHICmdList)
+	if (FSkeletalMeshUvMappingBufferProxy* ProxyPtr = FrozenQuadTreeProxy.Release())
+	{
+		ENQUEUE_RENDER_COMMAND(BeginDestroyCommand)([RT_Proxy=ProxyPtr](FRHICommandListImmediate& RHICmdList)
 		{
 			RT_Proxy->ReleaseResource();
 			delete RT_Proxy;
 		});
+	}
 }
 
 bool FSkeletalMeshUvMapping::IsUsed() const
diff --git a/Engine/Plugins/FX/Niagara/Source/Niagara/Private/NiagaraDataInterfaceSkeletalMeshUvMapping.h b/Engine/Plugins/FX/Niagara/Source/Niagara/Private/NiagaraDataInterfaceSkeletalMeshUvMapping.h
--- a/Engine/Plugins/FX/Niagara/Source/Niagara/Private/NiagaraDataInterfaceSkeletalMeshUvMapping.h
+++ b/Engine/Plugins/FX/Niagara/Source/Niagara/Private/NiagaraDataInterfaceSkeletalMeshUvMapping.h
@@ -34,6 +34,7 @@ struct FSkeletalMeshUvMapping
 	FSkeletalMeshUvMapping() = delete;
 	FSkeletalMeshUvMapping(const FSkeletalMeshUvMapping&) = delete;
 	FSkeletalMeshUvMapping(TWeakObjectPtr<USkeletalMesh> InMeshObject, int32 InLodIndex, int32 InUvSetIndex);
+	~FSkeletalMeshUvMapping();
 
 	bool IsUsed() const;
 	bool CanBeDestroyed() const;
