Fix for initial bounds of a system instance when emitter fixed bounds are used.  Specifically resolves issues with the system's bounds on the first frame.

diff --git a/Engine/Plugins/FX/Niagara/Source/Niagara/Private/NiagaraSystemInstance.cpp b/Engine/Plugins/FX/Niagara/Source/Niagara/Private/NiagaraSystemInstance.cpp
--- a/Engine/Plugins/FX/Niagara/Source/Niagara/Private/NiagaraSystemInstance.cpp
+++ b/Engine/Plugins/FX/Niagara/Source/Niagara/Private/NiagaraSystemInstance.cpp
@@ -2021,7 +2021,7 @@ void FNiagaraSystemInstance::InitEmitters()
 
 	bHasGPUEmitters = false;
 
-	LocalBounds = FBox(FVector::ZeroVector, FVector::ZeroVector);
+	LocalBounds = FBox();
 
 	Emitters.Empty(false);
 	UNiagaraSystem* System = GetSystem();
@@ -2035,6 +2035,9 @@ void FNiagaraSystemInstance::InitEmitters()
 		Emitters.Reserve(NumEmitters);
 		for (int32 EmitterIdx=0; EmitterIdx < NumEmitters; ++EmitterIdx)
 		{
+			//-TODO: We should not create emitter instances for disable emitters
+			const bool EmitterEnabled = EmitterHandles[EmitterIdx].GetIsEnabled();
+
 			TSharedRef<FNiagaraEmitterInstance, ESPMode::ThreadSafe> Sim = MakeShared<FNiagaraEmitterInstance, ESPMode::ThreadSafe>(this);
 
 			if (System->bFixedBounds)
@@ -2042,12 +2045,21 @@ void FNiagaraSystemInstance::InitEmitters()
 				// be sure to set the system bounds first so that we can bypass work in the initialization of the emitter
 				Sim->SetSystemFixedBoundsOverride(System->GetFixedBounds());
 			}
+			else if (EmitterEnabled)
+			{
+				if (const UNiagaraEmitter* EmitterAsset = EmitterHandles[EmitterIdx].GetInstance())
+				{
+					if (EmitterAsset->bFixedBounds)
+					{
+						LocalBounds += EmitterAsset->FixedBounds;
+					}
+				}
+			}
 
 			Sim->Init(EmitterIdx, ID);
 			Emitters.Add(Sim);
 
-			//-TODO: We should not create emitter instances for disable emitters
-			if (EmitterHandles[EmitterIdx].GetIsEnabled())
+			if (EmitterEnabled)
 			{
 				// Only set bHasGPUEmitters if we allow compute shaders on the platform
 				if (bAllowComputeShaders)
@@ -2064,6 +2076,10 @@ void FNiagaraSystemInstance::InitEmitters()
 		{
 			LocalBounds = System->GetFixedBounds();
 		}
+		else if (!LocalBounds.IsValid)
+		{
+			LocalBounds = FBox(EForceInit::ForceInitToZero);
+		}
 	}
 
 	ResetParameters();
