Niagara - Merge 16888080 - Update the ComputeObjectNameMap function so that finds data interfaces by script usage to avoid duplicate warnings, and also includes emitter data interfaces when handling the system scripts to fix the missing data interfaces.

diff --git a/Engine/Plugins/FX/Niagara/Source/Niagara/Classes/NiagaraScriptSourceBase.h b/Engine/Plugins/FX/Niagara/Source/Niagara/Classes/NiagaraScriptSourceBase.h
--- a/Engine/Plugins/FX/Niagara/Source/Niagara/Classes/NiagaraScriptSourceBase.h
+++ b/Engine/Plugins/FX/Niagara/Source/Niagara/Classes/NiagaraScriptSourceBase.h
@@ -103,7 +103,7 @@ class UNiagaraScriptSourceBase : public UObject
 
 	virtual void ComputeVMCompilationId(struct FNiagaraVMExecutableDataId& Id, ENiagaraScriptUsage InUsage, const FGuid& InUsageId, bool bForceRebuild = false) const {};
 
-	virtual TMap<FName, UNiagaraDataInterface*> ComputeObjectNameMap(FString EmitterUniqueName) const { return TMap<FName, UNiagaraDataInterface*>(); }
+	virtual TMap<FName, UNiagaraDataInterface*> ComputeObjectNameMap(UNiagaraSystem& System, ENiagaraScriptUsage Usage, FGuid UsageId, FString EmitterUniqueName) const { return TMap<FName, UNiagaraDataInterface*>(); }
 
 	virtual FGuid GetCompileBaseId(ENiagaraScriptUsage InUsage, const FGuid& InUsageId) const { return FGuid(); }
 
diff --git a/Engine/Plugins/FX/Niagara/Source/Niagara/Private/NiagaraSystem.cpp b/Engine/Plugins/FX/Niagara/Source/Niagara/Private/NiagaraSystem.cpp
--- a/Engine/Plugins/FX/Niagara/Source/Niagara/Private/NiagaraSystem.cpp
+++ b/Engine/Plugins/FX/Niagara/Source/Niagara/Private/NiagaraSystem.cpp
@@ -2481,7 +2481,8 @@ bool UNiagaraSystem::QueryCompileComplete(bool bWait, bool bDoPost, bool bDoNotA
 				}
 				else
 				{
-					ObjectNameMap = EmitterCompiledScriptPair.CompiledScript->GetLatestSource()->ComputeObjectNameMap(AsyncTask->UniqueEmitterName);
+					UNiagaraScript* CompiledScript = EmitterCompiledScriptPair.CompiledScript;
+					ObjectNameMap = CompiledScript->GetLatestSource()->ComputeObjectNameMap(*this, CompiledScript->GetUsage(), CompiledScript->GetUsageId(), AsyncTask->UniqueEmitterName);
 				}
 				EmitterCompiledScriptPair.CompiledScript->SetVMCompilationResults(EmitterCompiledScriptPair.CompileId, *ExeData, AsyncTask->UniqueEmitterName, ObjectNameMap);
 			}
diff --git a/Engine/Plugins/FX/Niagara/Source/NiagaraEditor/Private/NiagaraScriptSource.cpp b/Engine/Plugins/FX/Niagara/Source/NiagaraEditor/Private/NiagaraScriptSource.cpp
--- a/Engine/Plugins/FX/Niagara/Source/NiagaraEditor/Private/NiagaraScriptSource.cpp
+++ b/Engine/Plugins/FX/Niagara/Source/NiagaraEditor/Private/NiagaraScriptSource.cpp
@@ -171,14 +171,16 @@ void UNiagaraScriptSource::PostLoadFromEmitter(UNiagaraEmitter& OwningEmitter)
 	}
 }
 
-void FindObjectNamesRecursive(UNiagaraGraph* InGraph, FString EmitterUniqueName, TMap<FName, UNiagaraDataInterface*>& Result)
+void FindObjectNamesRecursive(UNiagaraGraph* InGraph, ENiagaraScriptUsage Usage, FGuid UsageId, FString EmitterUniqueName, TSet<UNiagaraGraph*>& VisitedGraphs, TMap<FName, UNiagaraDataInterface*>& Result)
 {
 	if (!InGraph)
 	{
 		return;
 	}
+
 	TArray<UNiagaraNode*> Nodes;
-	InGraph->GetNodesOfClass(Nodes);
+	UNiagaraNodeOutput* OutputNode = InGraph->FindEquivalentOutputNode(Usage, UsageId);
+	InGraph->BuildTraversal(Nodes, OutputNode);
 	for (UEdGraphNode* Node : Nodes)
 	{
 		if (UNiagaraNode* InNode = Cast<UNiagaraNode>(Node))
@@ -204,20 +206,71 @@ void FindObjectNamesRecursive(UNiagaraGraph* InGraph, FString EmitterUniqueName,
 			if (FunctionCallNode)
 			{
 				UNiagaraGraph* FunctionGraph = FunctionCallNode->GetCalledGraph();
-				FindObjectNamesRecursive(FunctionGraph, EmitterUniqueName, Result);
+				if(VisitedGraphs.Contains(FunctionGraph) == false)
+				{
+					VisitedGraphs.Add(FunctionGraph);
+					FindObjectNamesRecursive(FunctionGraph, FunctionCallNode->GetCalledUsage(), FGuid(), EmitterUniqueName, VisitedGraphs, Result);
+				}
 			}
 		}
 	}
 }
 
-TMap<FName, UNiagaraDataInterface*> UNiagaraScriptSource::ComputeObjectNameMap(FString EmitterUniqueName) const
+TMap<FName, UNiagaraDataInterface*> UNiagaraScriptSource::ComputeObjectNameMap(UNiagaraSystem& System, ENiagaraScriptUsage Usage, FGuid UsageId, FString EmitterUniqueName) const
 {
 	TMap<FName, UNiagaraDataInterface*> Result;
 	if (!NodeGraph)
 	{
 		return Result;
 	}
-	FindObjectNamesRecursive(NodeGraph, EmitterUniqueName, Result);
+
+	TSet<UNiagaraGraph*> VisitedGraphs;
+	if (Usage == ENiagaraScriptUsage::ParticleGPUComputeScript)
+	{
+		// GPU scripts need to include spawn, update, and sim stage script's data interfaces.
+		TArray<UNiagaraNodeOutput*> OutputNodes;
+		NodeGraph->GetNodesOfClass(OutputNodes);
+		for (UNiagaraNodeOutput* OutputNode : OutputNodes)
+		{
+			if (OutputNode->GetUsage() == ENiagaraScriptUsage::ParticleSpawnScript ||
+				OutputNode->GetUsage() == ENiagaraScriptUsage::ParticleSpawnScriptInterpolated ||
+				OutputNode->GetUsage() == ENiagaraScriptUsage::ParticleUpdateScript ||
+				OutputNode->GetUsage() == ENiagaraScriptUsage::ParticleSimulationStageScript)
+			{
+				FindObjectNamesRecursive(NodeGraph, OutputNode->GetUsage(), OutputNode->GetUsageId(), EmitterUniqueName, VisitedGraphs, Result);
+			}
+		}
+	}
+	else
+	{
+		// Collect the data interfaces for the current script.
+		FindObjectNamesRecursive(NodeGraph, Usage, UsageId, EmitterUniqueName, VisitedGraphs, Result);
+
+		if (Usage == ENiagaraScriptUsage::ParticleSpawnScriptInterpolated)
+		{
+			// The interpolated spawn script needs to include the particle update script's data interfaces as well.
+			FindObjectNamesRecursive(NodeGraph, ENiagaraScriptUsage::ParticleUpdateScript, FGuid(), EmitterUniqueName, VisitedGraphs, Result);
+		}
+		else if (Usage == ENiagaraScriptUsage::SystemSpawnScript || Usage == ENiagaraScriptUsage::SystemUpdateScript)
+		{
+			// System scripts need to include data interfaces from the corresponding emitter scripts.
+			ENiagaraScriptUsage EmitterUsage = Usage == ENiagaraScriptUsage::SystemSpawnScript
+				? ENiagaraScriptUsage::EmitterSpawnScript 
+				: ENiagaraScriptUsage::EmitterUpdateScript;
+
+			for (const FNiagaraEmitterHandle& EmitterHandle : System.GetEmitterHandles())
+			{
+				if (EmitterHandle.GetIsEnabled() && EmitterHandle.GetInstance() != nullptr)
+				{
+					UNiagaraScriptSource* EmitterSource = Cast<UNiagaraScriptSource>(EmitterHandle.GetInstance()->GraphSource);
+					if (EmitterSource != nullptr)
+					{
+						FindObjectNamesRecursive(EmitterSource->NodeGraph, EmitterUsage, FGuid(), EmitterHandle.GetUniqueInstanceName(), VisitedGraphs, Result);
+					}
+				}
+			}
+		}
+	}
 	return Result;
 }
 
diff --git a/Engine/Plugins/FX/Niagara/Source/NiagaraEditor/Public/NiagaraScriptSource.h b/Engine/Plugins/FX/Niagara/Source/NiagaraEditor/Public/NiagaraScriptSource.h
--- a/Engine/Plugins/FX/Niagara/Source/NiagaraEditor/Public/NiagaraScriptSource.h
+++ b/Engine/Plugins/FX/Niagara/Source/NiagaraEditor/Public/NiagaraScriptSource.h
@@ -45,7 +45,7 @@ class UNiagaraScriptSource : public UNiagaraScriptSourceBase
 
 	virtual void PostLoadFromEmitter(UNiagaraEmitter& OwningEmitter) override;
 
-	virtual TMap<FName, UNiagaraDataInterface*> ComputeObjectNameMap(FString EmitterUniqueName) const override;
+	virtual TMap<FName, UNiagaraDataInterface*> ComputeObjectNameMap(UNiagaraSystem& System, ENiagaraScriptUsage Usage, FGuid UsageId, FString EmitterUniqueName) const override;
 
 	NIAGARAEDITOR_API virtual bool AddModuleIfMissing(FString ModulePath, ENiagaraScriptUsage Usage, bool& bOutFoundModule)override;
 
