Fixed the Niagara user parameter binding customization not being applied properly in some cases when a dummy DI with engine/transient as its outer is used.

diff --git a/Engine/Plugins/FX/Niagara/Source/NiagaraEditor/Private/Customizations/NiagaraTypeCustomizations.cpp b/Engine/Plugins/FX/Niagara/Source/NiagaraEditor/Private/Customizations/NiagaraTypeCustomizations.cpp
--- a/Engine/Plugins/FX/Niagara/Source/NiagaraEditor/Private/Customizations/NiagaraTypeCustomizations.cpp
+++ b/Engine/Plugins/FX/Niagara/Source/NiagaraEditor/Private/Customizations/NiagaraTypeCustomizations.cpp
@@ -9,6 +9,7 @@
 #include "IPropertyUtilities.h"
 #include "NiagaraConstants.h"
 #include "NiagaraDataInterfaceRW.h"
+#include "NiagaraEditorModule.h"
 #include "NiagaraEditorUtilities.h"
 #include "NiagaraEmitter.h"
 #include "NiagaraNodeOutput.h"
@@ -623,6 +624,17 @@ void FNiagaraUserParameterBindingCustomization::CustomizeHeader(TSharedRef<IProp
 	if (Objects.Num() == 1)
 	{
 		BaseSystem = Objects[0]->GetTypedOuter<UNiagaraSystem>();
+		
+		if (BaseSystem == nullptr)
+		{
+			if (UNiagaraDataInterface* ObjectAsDataInterface = Cast<UNiagaraDataInterface>(Objects[0]))
+			{
+				FNiagaraEditorModule& NiagaraEditorModule = FModuleManager::LoadModuleChecked<FNiagaraEditorModule>("NiagaraEditor");
+				UNiagaraEmitter* Emitter = nullptr;
+				NiagaraEditorModule.GetTargetSystemAndEmitterForDataInterface(ObjectAsDataInterface, BaseSystem, Emitter);
+			}
+		}
+		
 		if (BaseSystem)
 		{
 			TargetUserParameterBinding = (FNiagaraUserParameterBinding*)PropertyHandle->GetValueBaseAddress((uint8*)Objects[0]);
