Fixed non-unity build error

diff --git a/Engine/Plugins/FX/Niagara/Source/NiagaraEditor/Private/Customizations/NiagaraTypeCustomizations.cpp b/Engine/Plugins/FX/Niagara/Source/NiagaraEditor/Private/Customizations/NiagaraTypeCustomizations.cpp
--- a/Engine/Plugins/FX/Niagara/Source/NiagaraEditor/Private/Customizations/NiagaraTypeCustomizations.cpp
+++ b/Engine/Plugins/FX/Niagara/Source/NiagaraEditor/Private/Customizations/NiagaraTypeCustomizations.cpp
@@ -10,6 +10,7 @@
 #include "NiagaraConstants.h"
 #include "NiagaraDataInterfaceRW.h"
 #include "NiagaraEditorModule.h"
+#include "Modules/ModuleManager.h"
 #include "NiagaraEditorUtilities.h"
 #include "NiagaraEmitter.h"
 #include "NiagaraNodeOutput.h"
