- Split indirect argument generation into pre opaque and post opaque
- This allows mesh renderers, sorting & culling to all fully function on the low latency translucent path

diff --git a/Engine/Plugins/FX/Niagara/Shaders/Private/NiagaraDrawIndirectArgsGen.usf b/Engine/Plugins/FX/Niagara/Shaders/Private/NiagaraDrawIndirectArgsGen.usf
--- a/Engine/Plugins/FX/Niagara/Shaders/Private/NiagaraDrawIndirectArgsGen.usf
+++ b/Engine/Plugins/FX/Niagara/Shaders/Private/NiagaraDrawIndirectArgsGen.usf
@@ -30,9 +30,9 @@ void MainCS(uint TaskIndex : SV_DispatchThreadID)
 
 	if (TaskIndex < NumArgGenTasks)
 	{
-		const uint ArgOffset = NIAGARA_DRAW_INDIRECT_ARGS_SIZE * TaskIndex;
 		const uint InfoOffset = (ArgGenTaskOffset + TaskIndex) * NIAGARA_DRAW_INDIRECT_TASK_INFO_SIZE;
-		const uint Flags = TaskInfos[InfoOffset + 3];
+		const uint ArgOffset = TaskInfos[InfoOffset + 0];
+		const uint Flags = TaskInfos[InfoOffset + 4];
 		const bool bUseCulledCounts = (Flags & FLAG_USE_CULLED_COUNTS) != 0;
 		const bool bInstancedStereo = (Flags & FLAG_INSTANCED_STEREO) != 0;
 
@@ -40,11 +40,11 @@ void MainCS(uint TaskIndex : SV_DispatchThreadID)
 		BRANCH
 		if (bUseCulledCounts)
 		{
-			InstanceCount = CulledInstanceCounts[TaskInfos[InfoOffset + 0]];
+			InstanceCount = CulledInstanceCounts[TaskInfos[InfoOffset + 1]];
 		}
 		else
 		{
-			InstanceCount = RWInstanceCounts[TaskInfos[InfoOffset + 0]];
+			InstanceCount = RWInstanceCounts[TaskInfos[InfoOffset + 1]];
 		}
 
 		if (bInstancedStereo)
@@ -52,9 +52,9 @@ void MainCS(uint TaskIndex : SV_DispatchThreadID)
 			InstanceCount *= 2;
 		}
 
-		RWDrawIndirectArgs[ArgOffset + 0] = TaskInfos[InfoOffset + 1]; // NumIndicesPerInstance
+		RWDrawIndirectArgs[ArgOffset + 0] = TaskInfos[InfoOffset + 2]; // NumIndicesPerInstance
 		RWDrawIndirectArgs[ArgOffset + 1] = InstanceCount;
-		RWDrawIndirectArgs[ArgOffset + 2] = TaskInfos[InfoOffset + 2]; // StartIndexLocation
+		RWDrawIndirectArgs[ArgOffset + 2] = TaskInfos[InfoOffset + 3]; // StartIndexLocation
 		RWDrawIndirectArgs[ArgOffset + 3] = 0; // BaseVertexLocation
 		RWDrawIndirectArgs[ArgOffset + 4] = 0; // StartInstanceLocation
 	}
@@ -86,7 +86,7 @@ void MainCS(uint TaskIndex : SV_DispatchThreadID)
 		if (TaskIndex < NumTotalTasks && TaskIndex >= NumArgGenTasks)
 		{
 			const uint InfoOffset = (ArgGenTaskOffset + NumArgGenTasks) * NIAGARA_DRAW_INDIRECT_TASK_INFO_SIZE + (TaskIndex - NumArgGenTasks);
-			RWInstanceCounts[TaskInfos[InfoOffset + 0]] = 0;
+			RWInstanceCounts[TaskInfos[InfoOffset]] = 0;
 		}
 #endif
 }
@@ -100,6 +100,6 @@ void ResetCountsCS(uint TaskIndex : SV_DispatchThreadID)
 	if (TaskIndex < NumInstanceCountClearTasks)
 	{
 		const uint InfoOffset = (ArgGenTaskOffset + NumArgGenTasks) * NIAGARA_DRAW_INDIRECT_TASK_INFO_SIZE + (TaskIndex);
-		RWInstanceCounts[TaskInfos[InfoOffset + 0]] = 0;
+		RWInstanceCounts[TaskInfos[InfoOffset]] = 0;
 	}
 }
\ No newline at end of file
diff --git a/Engine/Plugins/FX/Niagara/Source/Niagara/Classes/NiagaraDataSet.h b/Engine/Plugins/FX/Niagara/Source/Niagara/Classes/NiagaraDataSet.h
--- a/Engine/Plugins/FX/Niagara/Source/Niagara/Classes/NiagaraDataSet.h
+++ b/Engine/Plugins/FX/Niagara/Source/Niagara/Classes/NiagaraDataSet.h
@@ -115,7 +115,7 @@ public:
 	void Allocate(uint32 NumInstances, bool bMaintainExisting = false);
 
 	void AllocateGPU(FRHICommandList& RHICmdList, uint32 InNumInstances, ERHIFeatureLevel::Type FeatureLevel, const TCHAR* DebugSimName);
-	void SwapGPU(FNiagaraDataBuffer* BufferToAlias);
+	void SwapGPU(FNiagaraDataBuffer* BufferToSwap);
 	void ReleaseGPU();
 
 	void SwapInstances(uint32 OldIndex, uint32 NewIndex);
@@ -162,6 +162,9 @@ public:
 
 	FORCEINLINE void SetNumInstances(uint32 InNumInstances) { check(InNumInstances <= NumInstancesAllocated); NumInstances = InNumInstances; }
 	FORCEINLINE void SetNumSpawnedInstances(uint32 InNumSpawnedInstances) { NumSpawnedInstances = InNumSpawnedInstances; }
+
+	FORCEINLINE void SetGPUDataReadyStage(ENiagaraGpuComputeTickStage::Type InReadyStage) { GPUDataReadyStage = InReadyStage; }
+	FORCEINLINE ENiagaraGpuComputeTickStage::Type GetGPUDataReadyStage() const { return GPUDataReadyStage; }
 	FORCEINLINE FRWBuffer& GetGPUBufferFloat() { return GPUBufferFloat; }
 	FORCEINLINE FRWBuffer& GetGPUBufferInt() { return GPUBufferInt; }
 	FORCEINLINE FRWBuffer& GetGPUBufferHalf() { return GPUBufferHalf; }
@@ -224,6 +227,8 @@ private:
 
 	//////////////////////////////////////////////////////////////////////////
 	// GPU Data
+	/** Location in the frame where GPU data will be ready, for CPU this is always the first group, for GPU is depends on the features used as to which phase. */
+	ENiagaraGpuComputeTickStage::Type GPUDataReadyStage;
 	/** The buffer offset where the instance count is accumulated. */
 	uint32 GPUInstanceCountBufferOffset;
 	/** GPU Buffer containing floating point values for GPU simulations. */
diff --git a/Engine/Plugins/FX/Niagara/Source/Niagara/Classes/NiagaraGPUInstanceCountManager.h b/Engine/Plugins/FX/Niagara/Source/Niagara/Classes/NiagaraGPUInstanceCountManager.h
--- a/Engine/Plugins/FX/Niagara/Source/Niagara/Classes/NiagaraGPUInstanceCountManager.h
+++ b/Engine/Plugins/FX/Niagara/Source/Niagara/Classes/NiagaraGPUInstanceCountManager.h
@@ -27,6 +27,16 @@ FORCEINLINE uint32 GetTypeHash(const FNiagaraDrawIndirectArgGenTaskInfo& Info)
 	return HashCombine(Info.InstanceCountBufferOffset, HashCombine(Info.NumIndicesPerInstance, HashCombine(Info.StartIndexLocation, Info.Flags)));
 }
 
+namespace ENiagaraGPUCountUpdatePhase
+{
+	enum Type
+	{
+		PreOpaque,
+		PostOpaque,
+		Max
+	};
+}
+
 /**
  * A manager that handles the buffer containing the GPU particle count.
  * Also provides related functionalities like the generation of the draw indirect buffer.
@@ -81,8 +91,7 @@ public:
 	bool HasPendingGPUReadback() const;
 
 	/** Add a draw indirect task to generate the draw indirect args. Returns the draw indirect arg buffer offset. */
-	FIndirectArgSlot AddDrawIndirect(uint32 InstanceCountBufferOffset, uint32 NumIndicesPerInstance, uint32 StartIndexLocation,
-		bool bIsInstancedStereoEnabled, bool bCulled);
+	FIndirectArgSlot AddDrawIndirect(uint32 InstanceCountBufferOffset, uint32 NumIndicesPerInstance, uint32 StartIndexLocation, bool bIsInstancedStereoEnabled, bool bCulled, ENiagaraGpuComputeTickStage::Type ReadyTickStage);
 
 	// Resize instance count and draw indirect buffers to ensure it is big enough to hold all draw indirect args.
 	void ResizeBuffers(FRHICommandListImmediate& RHICmdList, ERHIFeatureLevel::Type FeatureLevel, int32 ReservedInstanceCounts);
@@ -90,7 +99,7 @@ public:
 	void FlushIndirectArgsPool();
 
 	// Generate the draw indirect buffers, and reset all release counts.
-	void UpdateDrawIndirectBuffers(NiagaraEmitterInstanceBatcher& Batcher, FRHICommandList& RHICmdList, ERHIFeatureLevel::Type FeatureLevel);
+	void UpdateDrawIndirectBuffers(NiagaraEmitterInstanceBatcher& Batcher, FRHICommandList& RHICmdList, ERHIFeatureLevel::Type FeatureLevel, ENiagaraGPUCountUpdatePhase::Type CountPhase);
 
 	static const ERHIAccess kCountBufferDefaultState;
 
@@ -99,9 +108,11 @@ public:
 protected:
 	struct FIndirectArgsPoolEntry
 	{
-		FRWBuffer Buffer;
-		uint32 NumAllocated = 0;
-		uint32 NumUsed = 0;
+		FRWBuffer	Buffer;
+
+		uint32		AllocatedEntries = 0;
+		uint32		UsedEntriesTotal = 0;
+		uint32		UsedEntries[ENiagaraGPUCountUpdatePhase::Max] = {};
 	};
 
 	using FIndirectArgsPoolEntryPtr = TUniquePtr<FIndirectArgsPoolEntry>;
@@ -128,14 +139,19 @@ protected:
 	int32 CountReadbackSize = 0;
 
 	/** The list of all draw indirected tasks that are to be run in UpdateDrawIndirectBuffers() */
-	using FArgGenTaskInfo = FNiagaraDrawIndirectArgGenTaskInfo;
-	using FArgGenSlotInfo = TTuple<uint32, uint32>;
+	TArray<FNiagaraDrawIndirectArgGenTaskInfo> DrawIndirectArgGenTasks[ENiagaraGPUCountUpdatePhase::Max];
+
+	/** Set to allow de-duplication of FNiagaraDrawIndirectArgGenTaskInfo. */
+	struct FNiagaraDrawIndirectArgGenSlotInfo
+	{
+		uint32 PoolIndex = INDEX_NONE;
+		uint32 BufferOffset = INDEX_NONE;
+	};
+	TMap<FNiagaraDrawIndirectArgGenTaskInfo, FNiagaraDrawIndirectArgGenSlotInfo> DrawIndirectArgMap;
 
-	TArray<FArgGenTaskInfo> DrawIndirectArgGenTasks;
-	/** The map between each task FArgGenTaskInfo and entry offset from DrawIndirectArgGenTasks. Used to reuse entries. */
-	TMap<FArgGenTaskInfo, FArgGenSlotInfo> DrawIndirectArgMap;
 	/** The list of all instance count clear tasks that are to be run in UpdateDrawIndirectBuffers() */
 	TArray<uint32> InstanceCountClearTasks;
+
 	/** Buffers holding drawindirect data to render GPU emitter renderers. */
 	TArray<FIndirectArgsPoolEntryPtr> DrawIndirectPool;
 	uint32 DrawIndirectLowWaterFrames = 0;
diff --git a/Engine/Plugins/FX/Niagara/Source/Niagara/Classes/NiagaraGPUSortInfo.h b/Engine/Plugins/FX/Niagara/Source/Niagara/Classes/NiagaraGPUSortInfo.h
--- a/Engine/Plugins/FX/Niagara/Source/Niagara/Classes/NiagaraGPUSortInfo.h
+++ b/Engine/Plugins/FX/Niagara/Source/Niagara/Classes/NiagaraGPUSortInfo.h
@@ -9,6 +9,7 @@ NiagaraGPUSortInfo.h: GPU particle sorting helper
 #include "RHI.h"
 #include "RHIResources.h"
 #include "GPUSortManager.h"
+#include "NiagaraCommon.h"
 #include "NiagaraGPUSortInfo.generated.h"
 
 UENUM()
@@ -70,10 +71,11 @@ struct FNiagaraGPUSortInfo
 	EGPUSortFlags SortFlags = EGPUSortFlags::None;
 
 	// Set the SortFlags based on the emitter and material constraints.
-	FORCEINLINE void SetSortFlags(bool bHighPrecisionKeys, bool bTranslucentMaterial)
+	FORCEINLINE void SetSortFlags(bool bHighPrecisionKeys, ENiagaraGpuComputeTickStage::Type ReadyTickStage)
 	{
-		SortFlags = EGPUSortFlags::KeyGenAfterPreRender | EGPUSortFlags::ValuesAsInt32 | 
+		SortFlags =
+			EGPUSortFlags::ValuesAsInt32 | 
 			(bHighPrecisionKeys ? EGPUSortFlags::HighPrecisionKeys : EGPUSortFlags::LowPrecisionKeys) |
-			(bTranslucentMaterial ? EGPUSortFlags::AnySortLocation : EGPUSortFlags::SortAfterPreRender);
+			(ReadyTickStage == ENiagaraGpuComputeTickStage::PostOpaqueRender ? EGPUSortFlags::SortAfterPostRenderOpaque | EGPUSortFlags::KeyGenAfterPostRenderOpaque : EGPUSortFlags::SortAfterPreRender | EGPUSortFlags::KeyGenAfterPreRender);
 	}
 };
\ No newline at end of file
diff --git a/Engine/Plugins/FX/Niagara/Source/Niagara/Private/NiagaraDataSet.cpp b/Engine/Plugins/FX/Niagara/Source/Niagara/Private/NiagaraDataSet.cpp
--- a/Engine/Plugins/FX/Niagara/Source/Niagara/Private/NiagaraDataSet.cpp
+++ b/Engine/Plugins/FX/Niagara/Source/Niagara/Private/NiagaraDataSet.cpp
@@ -841,24 +841,25 @@ void FNiagaraDataBuffer::AllocateGPU(FRHICommandList& RHICmdList, uint32 InNumIn
 #endif
 }
 
-void FNiagaraDataBuffer::SwapGPU(FNiagaraDataBuffer* BufferToAlias)
+void FNiagaraDataBuffer::SwapGPU(FNiagaraDataBuffer* BufferToSwap)
 {
 #if NIAGARA_MEMORY_TRACKING
-	Swap(AllocationSizeBytes, BufferToAlias->AllocationSizeBytes);
+	Swap(AllocationSizeBytes, BufferToSwap->AllocationSizeBytes);
 #endif
 
-	Swap(GPUInstanceCountBufferOffset, BufferToAlias->GPUInstanceCountBufferOffset);
-	Swap(GPUBufferFloat, BufferToAlias->GPUBufferFloat);
-	Swap(GPUBufferInt, BufferToAlias->GPUBufferInt);
-	Swap(GPUIDToIndexTable, BufferToAlias->GPUIDToIndexTable);
-	Swap(GPUBufferHalf, BufferToAlias->GPUBufferHalf);
-	Swap(NumInstances, BufferToAlias->NumInstances);
-	Swap(NumInstancesAllocated, BufferToAlias->NumInstancesAllocated);
-	Swap(FloatStride, BufferToAlias->FloatStride);
-	Swap(Int32Stride, BufferToAlias->Int32Stride);
-	Swap(HalfStride, BufferToAlias->HalfStride);
-	Swap(NumSpawnedInstances, BufferToAlias->NumSpawnedInstances);
-	Swap(IDAcquireTag, BufferToAlias->IDAcquireTag);
+	Swap(GPUDataReadyStage, BufferToSwap->GPUDataReadyStage);
+	Swap(GPUInstanceCountBufferOffset, BufferToSwap->GPUInstanceCountBufferOffset);
+	Swap(GPUBufferFloat, BufferToSwap->GPUBufferFloat);
+	Swap(GPUBufferInt, BufferToSwap->GPUBufferInt);
+	Swap(GPUIDToIndexTable, BufferToSwap->GPUIDToIndexTable);
+	Swap(GPUBufferHalf, BufferToSwap->GPUBufferHalf);
+	Swap(NumInstances, BufferToSwap->NumInstances);
+	Swap(NumInstancesAllocated, BufferToSwap->NumInstancesAllocated);
+	Swap(FloatStride, BufferToSwap->FloatStride);
+	Swap(Int32Stride, BufferToSwap->Int32Stride);
+	Swap(HalfStride, BufferToSwap->HalfStride);
+	Swap(NumSpawnedInstances, BufferToSwap->NumSpawnedInstances);
+	Swap(IDAcquireTag, BufferToSwap->IDAcquireTag);
 }
 
 void FNiagaraDataBuffer::ReleaseGPU()
diff --git a/Engine/Plugins/FX/Niagara/Source/Niagara/Private/NiagaraEmitterInstanceBatcher.cpp b/Engine/Plugins/FX/Niagara/Source/Niagara/Private/NiagaraEmitterInstanceBatcher.cpp
--- a/Engine/Plugins/FX/Niagara/Source/Niagara/Private/NiagaraEmitterInstanceBatcher.cpp
+++ b/Engine/Plugins/FX/Niagara/Source/Niagara/Private/NiagaraEmitterInstanceBatcher.cpp
@@ -199,12 +199,16 @@ NiagaraEmitterInstanceBatcher::NiagaraEmitterInstanceBatcher(ERHIFeatureLevel::T
 	// The registration also involves defining the list of flags possibly used in GPUSortManager::AddTask()
 	if (GPUSortManager)
 	{
-		GPUSortManager->Register(FGPUSortKeyGenDelegate::CreateLambda([this](FRHICommandListImmediate& RHICmdList, int32 BatchId, int32 NumElementsInBatch, EGPUSortFlags Flags, FRHIUnorderedAccessView* KeysUAV, FRHIUnorderedAccessView* ValuesUAV)
-		{
-			GenerateSortKeys(RHICmdList, BatchId, NumElementsInBatch, Flags, KeysUAV, ValuesUAV);
-		}),
-		EGPUSortFlags::AnyKeyPrecision | EGPUSortFlags::KeyGenAfterPreRender | EGPUSortFlags::AnySortLocation | EGPUSortFlags::ValuesAsInt32,
-		Name);
+		GPUSortManager->Register(
+			FGPUSortKeyGenDelegate::CreateLambda(
+				[this](FRHICommandListImmediate& RHICmdList, int32 BatchId, int32 NumElementsInBatch, EGPUSortFlags Flags, FRHIUnorderedAccessView* KeysUAV, FRHIUnorderedAccessView* ValuesUAV)
+				{
+					GenerateSortKeys(RHICmdList, BatchId, NumElementsInBatch, Flags, KeysUAV, ValuesUAV);
+				}
+			),
+			EGPUSortFlags::AnyKeyPrecision | EGPUSortFlags::AnyKeyGenLocation | EGPUSortFlags::AnySortLocation | EGPUSortFlags::ValuesAsInt32,
+			Name
+		);
 
 		if (FNiagaraUtilities::AllowComputeShaders(GetShaderPlatform()))
 		{
@@ -212,7 +216,7 @@ NiagaraEmitterInstanceBatcher::NiagaraEmitterInstanceBatcher(ERHIFeatureLevel::T
 			GPUSortManager->PostPreRenderEvent.AddLambda(
 				[this](FRHICommandListImmediate& RHICmdList)
 				{
-					GPUInstanceCounterManager.UpdateDrawIndirectBuffers(*this, RHICmdList, FeatureLevel);
+					GPUInstanceCounterManager.UpdateDrawIndirectBuffers(*this, RHICmdList, FeatureLevel, ENiagaraGPUCountUpdatePhase::PreOpaque);
 #if WITH_MGPU
 					// For PreInitViews we actually need to broadcast here and not in ExecuteAll since
 					// this is the last place the instance count buffer is written to.
@@ -223,6 +227,14 @@ NiagaraEmitterInstanceBatcher::NiagaraEmitterInstanceBatcher(ERHIFeatureLevel::T
 #endif // WITH_MGPU
 				}
 			);
+
+			GPUSortManager->PostPostRenderEvent.AddLambda
+			(
+				[this](FRHICommandListImmediate& RHICmdList)
+				{
+					GPUInstanceCounterManager.UpdateDrawIndirectBuffers(*this, RHICmdList, FeatureLevel, ENiagaraGPUCountUpdatePhase::PostOpaque);
+				}
+			);
 		}
 	}
 
@@ -1145,6 +1157,7 @@ void NiagaraEmitterInstanceBatcher::PrepareTicksForProxy(FRHICommandListImmediat
 			FNiagaraDataBuffer* FinalBuffer = ComputeContext->GetPrevDataBuffer();
 			FinalBuffer->GPUInstanceCountBufferOffset = ComputeContext->CountOffset_RT;
 			FinalBuffer->SetNumInstances(ComputeContext->CurrentNumInstances_RT);
+			FinalBuffer->SetGPUDataReadyStage(ComputeProxy->GetComputeTickStage());
 			ComputeContext->SetDataToRender(FinalBuffer);
 #if WITH_MGPU
 			AddTemporalEffectBuffers(FinalBuffer);
@@ -1158,6 +1171,7 @@ void NiagaraEmitterInstanceBatcher::PrepareTicksForProxy(FRHICommandListImmediat
 			FNiagaraDataBuffer* FinalBuffer = ComputeContext->GetPrevDataBuffer();
 			FinalBuffer->GPUInstanceCountBufferOffset = ComputeContext->CountOffset_RT;
 			FinalBuffer->SetNumInstances(ComputeContext->CurrentNumInstances_RT);
+			FinalBuffer->SetGPUDataReadyStage(ComputeProxy->GetComputeTickStage());
 			ComputeContext->SetTranslucentDataToRender(FinalBuffer);
 		}
 	}
@@ -1872,9 +1886,6 @@ bool NiagaraEmitterInstanceBatcher::AddSortedGPUSimulation(FNiagaraGPUSortInfo&
 
 void NiagaraEmitterInstanceBatcher::GenerateSortKeys(FRHICommandListImmediate& RHICmdList, int32 BatchId, int32 NumElementsInBatch, EGPUSortFlags Flags, FRHIUnorderedAccessView* KeysUAV, FRHIUnorderedAccessView* ValuesUAV)
 {
-	// Currently all Niagara KeyGen must execute after PreRender() - in between PreInitViews() and PostRenderOpaque(), when the GPU simulation are possibly ticked.
-	check(EnumHasAnyFlags(Flags, EGPUSortFlags::KeyGenAfterPreRender));
-
 	const FGPUSortManager::FKeyGenInfo KeyGenInfo((uint32)NumElementsInBatch, EnumHasAnyFlags(Flags, EGPUSortFlags::HighPrecisionKeys));
 
 	FNiagaraSortKeyGenCS::FPermutationDomain SortPermutationVector;
diff --git a/Engine/Plugins/FX/Niagara/Source/Niagara/Private/NiagaraGPUInstanceCountManager.cpp b/Engine/Plugins/FX/Niagara/Source/Niagara/Private/NiagaraGPUInstanceCountManager.cpp
--- a/Engine/Plugins/FX/Niagara/Source/Niagara/Private/NiagaraGPUInstanceCountManager.cpp
+++ b/Engine/Plugins/FX/Niagara/Source/Niagara/Private/NiagaraGPUInstanceCountManager.cpp
@@ -163,13 +163,6 @@ void FNiagaraGPUInstanceCountManager::FreeEntry(uint32& BufferOffset)
 		checkf(!FreeEntries.Contains(BufferOffset), TEXT("BufferOffset %u exists in FreeEntries"), BufferOffset);
 		checkf(!InstanceCountClearTasks.Contains(BufferOffset), TEXT("BufferOffset %u exists in InstanceCountClearTasks"), BufferOffset);
 
-		//UE_LOG(LogNiagara, Warning, TEXT("FNiagaraGPUInstanceCountManager::FreeEntry %d"), BufferOffset);
-		// Add a reset to 0 task.
-		// The entry will only become available/reusable after being reset to 0 in UpdateDrawIndirectBuffer()
-		//if (InstanceCountClearTasks.Find(BufferOffset) != -1)
-		//{
-		//	UE_LOG(LogNiagara, Warning, TEXT("FNiagaraGPUInstanceCountManager::FreeEntry DUPLICATED!!!!!!!!!!!!!! %d"), BufferOffset);
-		//}
 		InstanceCountClearTasks.Add(BufferOffset);
 		BufferOffset = INDEX_NONE;
 	}
@@ -294,249 +287,278 @@ void FNiagaraGPUInstanceCountManager::FlushIndirectArgsPool()
 	if (GNiagaraIndirectArgsPoolAllowShrinking && DrawIndirectPool.Num() > 0 && DrawIndirectLowWaterFrames >= uint32(GNiagaraIndirectArgsPoolLowWaterFrames))
 	{
 		FIndirectArgsPoolEntryPtr& PoolEntry = DrawIndirectPool[0];
-		const uint32 NewSize = FMath::Max<uint32>(GNiagaraIndirectArgsPoolMinSize, PoolEntry->NumAllocated / GNiagaraIndirectArgsPoolBlockSizeFactor);
+		const uint32 NewSize = FMath::Max<uint32>(GNiagaraIndirectArgsPoolMinSize, PoolEntry->AllocatedEntries / GNiagaraIndirectArgsPoolBlockSizeFactor);
 
-		INDIRECT_ARG_POOL_LOG("Shrinking pool from size %d to %d", PoolEntry->NumAllocated, NewSize);
+		INDIRECT_ARG_POOL_LOG("Shrinking pool from size %d to %d", PoolEntry->AllocatedEntries, NewSize);
 
 		PoolEntry->Buffer.Release();
-		PoolEntry->NumAllocated = NewSize;
+		PoolEntry->AllocatedEntries = NewSize;
 
 		TResourceArray<uint32> InitData;
-		InitData.AddZeroed(PoolEntry->NumAllocated * NIAGARA_DRAW_INDIRECT_ARGS_SIZE);
-		PoolEntry->Buffer.Initialize(TEXT("NiagaraGPUDrawIndirectArgs"), sizeof(uint32), PoolEntry->NumAllocated * NIAGARA_DRAW_INDIRECT_ARGS_SIZE, EPixelFormat::PF_R32_UINT, ERHIAccess::IndirectArgs, BUF_Static | BUF_DrawIndirect, &InitData);
+		InitData.AddZeroed(PoolEntry->AllocatedEntries * NIAGARA_DRAW_INDIRECT_ARGS_SIZE);
+		PoolEntry->Buffer.Initialize(TEXT("NiagaraGPUDrawIndirectArgs"), sizeof(uint32), PoolEntry->AllocatedEntries * NIAGARA_DRAW_INDIRECT_ARGS_SIZE, EPixelFormat::PF_R32_UINT, ERHIAccess::IndirectArgs, BUF_Static | BUF_DrawIndirect, &InitData);
 
 		// Reset the timer
 		DrawIndirectLowWaterFrames = 0;
 	}
 }
 
-FNiagaraGPUInstanceCountManager::FIndirectArgSlot FNiagaraGPUInstanceCountManager::AddDrawIndirect(uint32 InstanceCountBufferOffset, uint32 NumIndicesPerInstance,
-	uint32 StartIndexLocation, bool bIsInstancedStereoEnabled, bool bCulled)
+FNiagaraGPUInstanceCountManager::FIndirectArgSlot FNiagaraGPUInstanceCountManager::AddDrawIndirect(uint32 InstanceCountBufferOffset, uint32 NumIndicesPerInstance, uint32 StartIndexLocation, bool bIsInstancedStereoEnabled, bool bCulled, ENiagaraGpuComputeTickStage::Type ReadyTickStage)
 {
 	checkSlow(IsInRenderingThread());
 
-	const FArgGenTaskInfo Info(InstanceCountBufferOffset, NumIndicesPerInstance, StartIndexLocation, bIsInstancedStereoEnabled, bCulled);
+	const ENiagaraDrawIndirectArgGenTaskFlags TaskFlags =
+		(bIsInstancedStereoEnabled ? ENiagaraDrawIndirectArgGenTaskFlags::InstancedStereo : ENiagaraDrawIndirectArgGenTaskFlags::None)
+		| (bCulled ? ENiagaraDrawIndirectArgGenTaskFlags::UseCulledCounts : ENiagaraDrawIndirectArgGenTaskFlags::None)
+		| (ReadyTickStage == ENiagaraGpuComputeTickStage::PostOpaqueRender ? ENiagaraDrawIndirectArgGenTaskFlags::PostOpaque : ENiagaraDrawIndirectArgGenTaskFlags::None);
+	FNiagaraDrawIndirectArgGenTaskInfo Info(InstanceCountBufferOffset, NumIndicesPerInstance, StartIndexLocation, TaskFlags);
 
-	static const FArgGenSlotInfo InvalidSlot(INDEX_NONE, INDEX_NONE);
-	FArgGenSlotInfo& CachedSlot = DrawIndirectArgMap.FindOrAdd(Info, InvalidSlot);
-	if (CachedSlot == InvalidSlot)
+	FNiagaraDrawIndirectArgGenSlotInfo* SlotInfo = DrawIndirectArgMap.Find(Info);
+	if ( SlotInfo == nullptr )
 	{
 		// Attempt to allocate a new slot from the pool, or add to the pool if it's full
 		FIndirectArgsPoolEntry* PoolEntry = DrawIndirectPool.Num() > 0 ? DrawIndirectPool.Last().Get() : nullptr;
-		if (PoolEntry == nullptr || PoolEntry->NumUsed >= PoolEntry->NumAllocated)
+		if (PoolEntry == nullptr || PoolEntry->UsedEntriesTotal >= PoolEntry->AllocatedEntries)
 		{
 			FIndirectArgsPoolEntryPtr NewEntry = MakeUnique<FIndirectArgsPoolEntry>();
-			NewEntry->NumAllocated = PoolEntry ? uint32(PoolEntry->NumAllocated * GNiagaraIndirectArgsPoolBlockSizeFactor) : uint32(GNiagaraIndirectArgsPoolMinSize);
+			NewEntry->AllocatedEntries = PoolEntry ? uint32(PoolEntry->AllocatedEntries * GNiagaraIndirectArgsPoolBlockSizeFactor) : uint32(GNiagaraIndirectArgsPoolMinSize);
 
-			INDIRECT_ARG_POOL_LOG("Increasing pool from size %d to %d", PoolEntry ? PoolEntry->NumAllocated : 0, NewEntry->NumAllocated);
+			INDIRECT_ARG_POOL_LOG("Increasing pool from size %d to %d", PoolEntry ? PoolEntry->AllocatedEntries : 0, NewEntry->AllocatedEntries);
 
 			TResourceArray<uint32> InitData;
-			InitData.AddZeroed(NewEntry->NumAllocated * NIAGARA_DRAW_INDIRECT_ARGS_SIZE);
-			NewEntry->Buffer.Initialize(TEXT("NiagaraGPUDrawIndirectArgs"), sizeof(uint32), NewEntry->NumAllocated * NIAGARA_DRAW_INDIRECT_ARGS_SIZE, EPixelFormat::PF_R32_UINT, ERHIAccess::IndirectArgs, BUF_Static | BUF_DrawIndirect, &InitData);
+			InitData.AddZeroed(NewEntry->AllocatedEntries * NIAGARA_DRAW_INDIRECT_ARGS_SIZE);
+			NewEntry->Buffer.Initialize(TEXT("NiagaraGPUDrawIndirectArgs"), sizeof(uint32), NewEntry->AllocatedEntries * NIAGARA_DRAW_INDIRECT_ARGS_SIZE, EPixelFormat::PF_R32_UINT, ERHIAccess::IndirectArgs, BUF_Static | BUF_DrawIndirect, &InitData);
 
 			PoolEntry = NewEntry.Get();
 			DrawIndirectPool.Emplace(MoveTemp(NewEntry));
 		}
 
-		DrawIndirectArgGenTasks.Add(Info);
-		CachedSlot.Key = DrawIndirectPool.Num() - 1;
-		CachedSlot.Value = PoolEntry->NumUsed * NIAGARA_DRAW_INDIRECT_ARGS_SIZE * sizeof(uint32);
-		++PoolEntry->NumUsed;
+		Info.IndirectArgsBufferOffset = PoolEntry->UsedEntriesTotal * NIAGARA_DRAW_INDIRECT_ARGS_SIZE;
+		++PoolEntry->UsedEntriesTotal;
+
+		SlotInfo = &DrawIndirectArgMap.Add(Info);
+		SlotInfo->PoolIndex = DrawIndirectPool.Num() - 1;
+		SlotInfo->BufferOffset = Info.IndirectArgsBufferOffset * sizeof(uint32);
+
+		const ENiagaraGPUCountUpdatePhase::Type CountPhase = ReadyTickStage == ENiagaraGpuComputeTickStage::PostOpaqueRender ? ENiagaraGPUCountUpdatePhase::PostOpaque : ENiagaraGPUCountUpdatePhase::PreOpaque;
+		DrawIndirectArgGenTasks[CountPhase].Add(Info);
+		++PoolEntry->UsedEntries[CountPhase];
 	}
 
-	return FIndirectArgSlot(DrawIndirectPool[CachedSlot.Key]->Buffer.Buffer, DrawIndirectPool[CachedSlot.Key]->Buffer.SRV, CachedSlot.Value);
+	return FIndirectArgSlot(DrawIndirectPool[SlotInfo->PoolIndex]->Buffer.Buffer, DrawIndirectPool[SlotInfo->PoolIndex]->Buffer.SRV, SlotInfo->BufferOffset);
 }
 
-void FNiagaraGPUInstanceCountManager::UpdateDrawIndirectBuffers(NiagaraEmitterInstanceBatcher& Batcher, FRHICommandList& RHICmdList, ERHIFeatureLevel::Type FeatureLevel)
-{
-	INC_DWORD_STAT_BY(STAT_NiagaraIndirectDraws, DrawIndirectArgGenTasks.Num());
+////////////////////////////////////////////////////////////////////////////////////////////////////////////////////////////////////////////////////////////////////////////////////////////////////////
 
-	if (DrawIndirectArgGenTasks.Num() || InstanceCountClearTasks.Num())
+void FNiagaraGPUInstanceCountManager::UpdateDrawIndirectBuffers(NiagaraEmitterInstanceBatcher& Batcher, FRHICommandList& RHICmdList, ERHIFeatureLevel::Type FeatureLevel, ENiagaraGPUCountUpdatePhase::Type CountPhase)
+{
+	// Anything to process?
+	TArray<FNiagaraDrawIndirectArgGenTaskInfo>& ArgTasks = DrawIndirectArgGenTasks[CountPhase];
+	const bool bClearCounts = (CountPhase == ENiagaraGPUCountUpdatePhase::PreOpaque) && (InstanceCountClearTasks.Num() > 0);
+	if ( (ArgTasks.Num() > 0 || bClearCounts) && FNiagaraUtilities::AllowComputeShaders(GShaderPlatformForFeatureLevel[FeatureLevel]) )
 	{
-		if (FNiagaraUtilities::AllowComputeShaders(GShaderPlatformForFeatureLevel[FeatureLevel]))
+		INC_DWORD_STAT_BY(STAT_NiagaraIndirectDraws, ArgTasks.Num());
+
+		SCOPED_DRAW_EVENT(RHICmdList, NiagaraUpdateDrawIndirectBuffers);
+
+		// Allocate task buffer
+		FReadBuffer TaskInfosBuffer;
 		{
-			SCOPED_DRAW_EVENT(RHICmdList, NiagaraUpdateDrawIndirectBuffers);
-			FReadBuffer TaskInfosBuffer;
+			const uint32 ArgGenSize = ArgTasks.Num() * sizeof(FNiagaraDrawIndirectArgGenTaskInfo);
+			const uint32 InstanceCountClearSize = InstanceCountClearTasks.Num() * sizeof(uint32);
+			const uint32 TaskBufferSize = ArgGenSize + InstanceCountClearSize;
+			TaskInfosBuffer.Initialize(TEXT("NiagaraTaskInfosBuffer"), sizeof(uint32), TaskBufferSize / sizeof(uint32), EPixelFormat::PF_R32_UINT, BUF_Volatile);
+
+			uint8* TaskBufferData = (uint8*)RHILockBuffer(TaskInfosBuffer.Buffer, 0, TaskBufferSize, RLM_WriteOnly);
+			FMemory::Memcpy(TaskBufferData, ArgTasks.GetData(), ArgGenSize);
+			FMemory::Memcpy(TaskBufferData + ArgGenSize, InstanceCountClearTasks.GetData(), InstanceCountClearSize);
+			RHIUnlockBuffer(TaskInfosBuffer.Buffer);
+		}
 
-			//for (int32 i = 0; i < InstanceCountClearTasks.Num(); i++)
-			//{
-			//		UE_LOG(LogNiagara, Log, TEXT("InstanceCountClearTasks[%d] = %d"), i, InstanceCountClearTasks[i]);
-			//}
-			{
-				// All draw indirect args task are run first because of the binding between the task ID and arg write offset.
-				const uint32 ArgGenSize = DrawIndirectArgGenTasks.Num() * sizeof(FArgGenTaskInfo);
-				const uint32 InstanceCountClearSize = InstanceCountClearTasks.Num() * sizeof(uint32);
-				const uint32 TaskBufferSize = ArgGenSize + InstanceCountClearSize;
-				TaskInfosBuffer.Initialize(TEXT("NiagaraTaskInfosBuffer"), sizeof(uint32), TaskBufferSize / sizeof(uint32), EPixelFormat::PF_R32_UINT, BUF_Volatile);
-				uint8* TaskBufferData = (uint8*)RHILockBuffer(TaskInfosBuffer.Buffer, 0, TaskBufferSize, RLM_WriteOnly);
-				FMemory::Memcpy(TaskBufferData, DrawIndirectArgGenTasks.GetData(), ArgGenSize);
-				FMemory::Memcpy(TaskBufferData + ArgGenSize, InstanceCountClearTasks.GetData(), InstanceCountClearSize);
-				RHIUnlockBuffer(TaskInfosBuffer.Buffer);
-			}
-			FNiagaraUAVPoolAccessScope UAVPoolAccessScope(Batcher);
+		FNiagaraUAVPoolAccessScope UAVPoolAccessScope(Batcher);
+		TArray<FRHITransitionInfo, TInlineAllocator<10>> Transitions;
+		Transitions.Reserve(DrawIndirectPool.Num() + 2);
 
-			FUnorderedAccessViewRHIRef CountsUAV = nullptr;
-			FShaderResourceViewRHIRef CulledCountsSRV = nullptr;
-			TArray<FRHITransitionInfo, TInlineAllocator<10>> Transitions;
-			for (auto& PoolEntry : DrawIndirectPool)
+		// Get counts buffer
+		FUnorderedAccessViewRHIRef CountsUAV = nullptr;
+		const bool bCountBufferIsValid = CountBuffer.UAV.IsValid();
+		if (bCountBufferIsValid)
+		{
+			Transitions.Emplace(CountBuffer.UAV, kCountBufferDefaultState, ERHIAccess::UAVCompute);
+			CountsUAV = CountBuffer.UAV;
+		}
+		else
+		{
+			// This can happen if there are no InstanceCountClearTasks and all DrawIndirectArgGenTasks_PreOpaque are using culled counts
+			CountsUAV = Batcher.GetEmptyUAVFromPool(RHICmdList, PF_R32_UINT, ENiagaraEmptyUAVType::Buffer);
+		}
+
+		// Get culled counts buffer
+		FShaderResourceViewRHIRef CulledCountsSRV = nullptr;
+		if (CulledCountBuffer.SRV.IsValid())
+		{
+			if (bAcquiredCulledCounts)
 			{
-				Transitions.Emplace(PoolEntry->Buffer.UAV, ERHIAccess::IndirectArgs, ERHIAccess::UAVCompute);
+				Transitions.Emplace(CulledCountBuffer.UAV, ERHIAccess::UAVCompute, ERHIAccess::SRVCompute);
 			}
+			CulledCountsSRV = CulledCountBuffer.SRV.GetReference();
+		}
+		else
+		{
+			CulledCountsSRV = FNiagaraRenderer::GetDummyUIntBuffer();
+		}
 
-			const bool bCountBufferIsValid = CountBuffer.UAV.IsValid();
-			if (bCountBufferIsValid)
+		// Execute transitions
+		for (auto& PoolEntry : DrawIndirectPool)
+		{
+			Transitions.Emplace(PoolEntry->Buffer.UAV, ERHIAccess::IndirectArgs, ERHIAccess::UAVCompute);
+		}
+		RHICmdList.Transition(Transitions);
+
+		// Execute tasks
+		FNiagaraDrawIndirectArgsGenCS::FPermutationDomain PermutationVector;
+		PermutationVector.Set<FNiagaraDrawIndirectArgsGenCS::FSupportsTextureRW>(GRHISupportsRWTextureBuffers ? 1 : 0);
+		TShaderMapRef<FNiagaraDrawIndirectArgsGenCS> DrawIndirectArgsGenCS(GetGlobalShaderMap(FeatureLevel), PermutationVector);
+
+		if (bCountBufferIsValid)
+		{
+			RHICmdList.BeginUAVOverlap(CountsUAV);
+		}
+
+		const int32 NumDispatches = FMath::Max(DrawIndirectPool.Num(), 1);
+		uint32 ArgGenTaskOffset = 0;
+		for (int32 DispatchIdx = 0; DispatchIdx < NumDispatches; ++DispatchIdx)
+		{
+			// Get draw indirect pool UAV
+			// Note: If we have counts to clear but no indirect args we won't have a DrawIndirectPool entry
+			FRHIUnorderedAccessView* ArgsUAV = nullptr;
+			int32 NumArgGenTasks = 0;
+			if (DrawIndirectPool.IsValidIndex(DispatchIdx))
 			{
-				Transitions.Emplace(CountBuffer.UAV, kCountBufferDefaultState, ERHIAccess::UAVCompute);
-				CountsUAV = CountBuffer.UAV;
+				FIndirectArgsPoolEntryPtr& PoolEntry = DrawIndirectPool[DispatchIdx];
+				ArgsUAV = PoolEntry->Buffer.UAV;
+				NumArgGenTasks = PoolEntry->UsedEntries[CountPhase];
 			}
 			else
 			{
-				// This can happen if there are no InstanceCountClearTasks and all DrawIndirectArgGenTasks are using culled counts
-				CountsUAV = Batcher.GetEmptyUAVFromPool(RHICmdList, PF_R32_UINT, ENiagaraEmptyUAVType::Buffer);
+				ArgsUAV = Batcher.GetEmptyUAVFromPool(RHICmdList, PF_R32_UINT, ENiagaraEmptyUAVType::Buffer);
 			}
 
-			if (CulledCountBuffer.SRV.IsValid())
-			{
-				if (bAcquiredCulledCounts)
-				{
-					Transitions.Emplace(CulledCountBuffer.UAV, ERHIAccess::UAVCompute, ERHIAccess::SRVCompute);
-				}
-				CulledCountsSRV = CulledCountBuffer.SRV.GetReference();
-			}
-			else
+			const bool bIsLastDispatch = DispatchIdx == (NumDispatches - 1);
+			const int32 NumInstanceCountClearTasks = bIsLastDispatch ? InstanceCountClearTasks.Num() : 0;
+
+			// Do we have anything to do for this pool?
+			if (NumArgGenTasks + NumInstanceCountClearTasks == 0)
 			{
-				CulledCountsSRV = FNiagaraRenderer::GetDummyUIntBuffer();
+				continue;
 			}
 
-			RHICmdList.Transition(Transitions);
+			SetComputePipelineState(RHICmdList, DrawIndirectArgsGenCS.GetComputeShader());
+			DrawIndirectArgsGenCS->SetOutput(RHICmdList, ArgsUAV, CountsUAV);
+			DrawIndirectArgsGenCS->SetParameters(RHICmdList, TaskInfosBuffer.SRV, CulledCountsSRV, ArgGenTaskOffset, NumArgGenTasks, NumInstanceCountClearTasks);
 
-			FNiagaraDrawIndirectArgsGenCS::FPermutationDomain PermutationVector;
-			PermutationVector.Set<FNiagaraDrawIndirectArgsGenCS::FSupportsTextureRW>(GRHISupportsRWTextureBuffers ? 1 : 0);
-			TShaderMapRef<FNiagaraDrawIndirectArgsGenCS> DrawIndirectArgsGenCS(GetGlobalShaderMap(FeatureLevel), PermutationVector);
-
-			if (bCountBufferIsValid)
+			// If the device supports RW Texture buffers then we can use a single compute pass, otherwise we need to split into two passes
+			if (GRHISupportsRWTextureBuffers)
 			{
-				RHICmdList.BeginUAVOverlap(CountsUAV);
+				DispatchComputeShader(RHICmdList, DrawIndirectArgsGenCS.GetShader(), FMath::DivideAndRoundUp(NumArgGenTasks + NumInstanceCountClearTasks, NIAGARA_DRAW_INDIRECT_ARGS_GEN_THREAD_COUNT), 1, 1);
+				DrawIndirectArgsGenCS->UnbindBuffers(RHICmdList);
 			}
-
-			const int32 NumDispatches = FMath::Max(DrawIndirectPool.Num(), 1);
-			uint32 ArgGenTaskOffset = 0;
-			for (int32 DispatchIdx = 0; DispatchIdx < NumDispatches; ++DispatchIdx)
+			else
 			{
-				int32 NumArgGenTasks = 0;
-				FUnorderedAccessViewRHIRef ArgsUAV;
-				if (DrawIndirectPool.IsValidIndex(DispatchIdx))
+				if (NumArgGenTasks > 0)
 				{
-					FIndirectArgsPoolEntryPtr& PoolEntry = DrawIndirectPool[DispatchIdx];
-					ArgsUAV = PoolEntry->Buffer.UAV;
-					NumArgGenTasks = PoolEntry->NumUsed;
-				}
-				else
-				{
-					ArgsUAV = Batcher.GetEmptyUAVFromPool(RHICmdList, PF_R32_UINT, ENiagaraEmptyUAVType::Buffer);
-				}
-
-				const bool bIsLastDispatch = DispatchIdx == (NumDispatches - 1);
-				const int32 NumInstanceCountClearTasks = bIsLastDispatch ? InstanceCountClearTasks.Num() : 0;
-
-				SetComputePipelineState(RHICmdList, DrawIndirectArgsGenCS.GetComputeShader());
-				DrawIndirectArgsGenCS->SetOutput(RHICmdList, ArgsUAV, CountsUAV);
-				DrawIndirectArgsGenCS->SetParameters(RHICmdList, TaskInfosBuffer.SRV, CulledCountsSRV, ArgGenTaskOffset, NumArgGenTasks, NumInstanceCountClearTasks);
-
-				// If the device supports RW Texture buffers then we can use a single compute pass, otherwise we need to split into two passes
-				if (GRHISupportsRWTextureBuffers)
-				{
-					DispatchComputeShader(RHICmdList, DrawIndirectArgsGenCS.GetShader(), FMath::DivideAndRoundUp(NumArgGenTasks + NumInstanceCountClearTasks, NIAGARA_DRAW_INDIRECT_ARGS_GEN_THREAD_COUNT), 1, 1);
+					DispatchComputeShader(RHICmdList, DrawIndirectArgsGenCS.GetShader(), FMath::DivideAndRoundUp(NumArgGenTasks, NIAGARA_DRAW_INDIRECT_ARGS_GEN_THREAD_COUNT), 1, 1);
 					DrawIndirectArgsGenCS->UnbindBuffers(RHICmdList);
 				}
-				else
+
+				if (NumInstanceCountClearTasks > 0)
 				{
-					if (NumArgGenTasks > 0)
-					{
-						DispatchComputeShader(RHICmdList, DrawIndirectArgsGenCS.GetShader(), FMath::DivideAndRoundUp(NumArgGenTasks, NIAGARA_DRAW_INDIRECT_ARGS_GEN_THREAD_COUNT), 1, 1);
-						DrawIndirectArgsGenCS->UnbindBuffers(RHICmdList);
-					}
-
-					if (NumInstanceCountClearTasks > 0)
-					{
-						FNiagaraDrawIndirectResetCountsCS::FPermutationDomain PermutationVectorResetCounts;
-						TShaderMapRef<FNiagaraDrawIndirectResetCountsCS> DrawIndirectResetCountsArgsGenCS(GetGlobalShaderMap(FeatureLevel), PermutationVectorResetCounts);
-						SetComputePipelineState(RHICmdList, DrawIndirectResetCountsArgsGenCS.GetComputeShader());
-						DrawIndirectResetCountsArgsGenCS->SetOutput(RHICmdList, CountBuffer.UAV);
-						DrawIndirectResetCountsArgsGenCS->SetParameters(RHICmdList, TaskInfosBuffer.SRV, DrawIndirectArgGenTasks.Num(), NumInstanceCountClearTasks);
-						DispatchComputeShader(RHICmdList, DrawIndirectResetCountsArgsGenCS.GetShader(), FMath::DivideAndRoundUp(NumInstanceCountClearTasks, NIAGARA_DRAW_INDIRECT_ARGS_GEN_THREAD_COUNT), 1, 1);
-						DrawIndirectResetCountsArgsGenCS->UnbindBuffers(RHICmdList);
-					}
+					FNiagaraDrawIndirectResetCountsCS::FPermutationDomain PermutationVectorResetCounts;
+					TShaderMapRef<FNiagaraDrawIndirectResetCountsCS> DrawIndirectResetCountsArgsGenCS(GetGlobalShaderMap(FeatureLevel), PermutationVectorResetCounts);
+					SetComputePipelineState(RHICmdList, DrawIndirectResetCountsArgsGenCS.GetComputeShader());
+					DrawIndirectResetCountsArgsGenCS->SetOutput(RHICmdList, CountBuffer.UAV);
+					DrawIndirectResetCountsArgsGenCS->SetParameters(RHICmdList, TaskInfosBuffer.SRV, ArgTasks.Num(), NumInstanceCountClearTasks);
+					DispatchComputeShader(RHICmdList, DrawIndirectResetCountsArgsGenCS.GetShader(), FMath::DivideAndRoundUp(NumInstanceCountClearTasks, NIAGARA_DRAW_INDIRECT_ARGS_GEN_THREAD_COUNT), 1, 1);
+					DrawIndirectResetCountsArgsGenCS->UnbindBuffers(RHICmdList);
 				}
-
-				ArgGenTaskOffset += NumArgGenTasks;
-			}
-
-			if (bCountBufferIsValid)
-			{
-				RHICmdList.EndUAVOverlap(CountsUAV);
 			}
 
-			Transitions.Reset();
-			for (auto& PoolEntry : DrawIndirectPool)
-			{
-				Transitions.Emplace(PoolEntry->Buffer.UAV, ERHIAccess::UAVCompute, ERHIAccess::IndirectArgs);
-			}
-			Transitions.Emplace(CountBuffer.UAV, ERHIAccess::UAVCompute, kCountBufferDefaultState);
-			RHICmdList.Transition(Transitions);
+			ArgGenTaskOffset += NumArgGenTasks;
 		}
-		// Once cleared to 0, the count are reusable.
-		FreeEntries.Append(InstanceCountClearTasks);
 
-		//for (int32 i = 0; i < FreeEntries.Num(); i++)
-		//{
-		//	UE_LOG(LogNiagara, Log, TEXT("FreeEntries[%d] = %d"), i, FreeEntries[i]);
-		//}
-
-		DrawIndirectArgGenTasks.Empty();
-		DrawIndirectArgMap.Empty();
-		InstanceCountClearTasks.Empty();
+		if (bCountBufferIsValid)
+		{
+			RHICmdList.EndUAVOverlap(CountsUAV);
+		}
 
+		// Generate and execute transitions
+		Transitions.Reset();
+		for (auto& PoolEntry : DrawIndirectPool)
+		{
+			Transitions.Emplace(PoolEntry->Buffer.UAV, ERHIAccess::UAVCompute, ERHIAccess::IndirectArgs);
+		}
+		Transitions.Emplace(CountBuffer.UAV, ERHIAccess::UAVCompute, kCountBufferDefaultState);
+		RHICmdList.Transition(Transitions);
 	}
 
-	// Release culled count buffers
-	// This is done outside of the above if as a mesh renderer could request a culled count but never add any indirect draws
-	if (bAcquiredCulledCounts && RequiredCulledCounts > 0)
+	// Add free counts back to list as we have cleared them
+	if ( bClearCounts )
 	{
-		CulledCountBuffer.DiscardTransientResource();
+		FreeEntries.Append(InstanceCountClearTasks);
+		InstanceCountClearTasks.Empty();
 	}
-	bAcquiredCulledCounts = false;
-	RequiredCulledCounts = 0; // reset this counter now that we're done with them
 
-	if (GNiagaraIndirectArgsPoolAllowShrinking)
+	DrawIndirectArgGenTasks[CountPhase].Empty();
+
+	// Final phase we can do some booking keeping / clearing of data
+	if ( CountPhase == ENiagaraGPUCountUpdatePhase::PostOpaque )
 	{
-		if (DrawIndirectPool.Num() == 1 && DrawIndirectPool[0]->NumAllocated > uint32(GNiagaraIndirectArgsPoolMinSize))
+		// Clear indirect arg map
+		DrawIndirectArgMap.Empty();
+
+		// Release culled count buffers
+		if (bAcquiredCulledCounts && RequiredCulledCounts > 0)
 		{
-			// See if this was a low water mark frame
-			FIndirectArgsPoolEntryPtr& PoolEntry = DrawIndirectPool[0];
-			const uint32 LowWaterCount = FMath::Max<uint32>(GNiagaraIndirectArgsPoolMinSize, PoolEntry->NumAllocated * GNiagaraIndirectArgsPoolLowWaterAmount);
-			if (PoolEntry->NumUsed < LowWaterCount)
+			CulledCountBuffer.DiscardTransientResource();
+		}
+		bAcquiredCulledCounts = false;
+		RequiredCulledCounts = 0;
+
+		// Optionally shrink the indirect arg pool
+		if (GNiagaraIndirectArgsPoolAllowShrinking)
+		{
+			if (DrawIndirectPool.Num() == 1 && DrawIndirectPool[0]->AllocatedEntries > uint32(GNiagaraIndirectArgsPoolMinSize))
 			{
-				++DrawIndirectLowWaterFrames;
+				// See if this was a low water mark frame
+				FIndirectArgsPoolEntryPtr& PoolEntry = DrawIndirectPool[0];
+				const uint32 LowWaterCount = FMath::Max<uint32>(GNiagaraIndirectArgsPoolMinSize, PoolEntry->AllocatedEntries * GNiagaraIndirectArgsPoolLowWaterAmount);
+				if (PoolEntry->UsedEntriesTotal < LowWaterCount)
+				{
+					++DrawIndirectLowWaterFrames;
+				}
+				else
+				{
+					// We've allocated above the low water amount, reset the timer
+					DrawIndirectLowWaterFrames = 0;
+				}
 			}
 			else
 			{
-				// We've allocated above the low water amount, reset the timer
+				// Either the pool is empty, at the min size, or we had to increase the pool size this frame. Either way, reset the shrink timer
 				DrawIndirectLowWaterFrames = 0;
 			}
 		}
-		else
+
+		// Clear indirect args pool counts
+		for (auto& Pool : DrawIndirectPool)
 		{
-			// Either the pool is empty, at the min size, or we had to increase the pool size this frame. Either way, reset the shrink timer
-			DrawIndirectLowWaterFrames = 0;
+			Pool->UsedEntriesTotal = 0;
+			for ( uint32& UsedEntry : Pool->UsedEntries )
+			{
+				UsedEntry = 0;
+			}
 		}
 	}
-
-	// Clear indirect args pool counts
-	for (auto& Pool : DrawIndirectPool)
-	{
-		Pool->NumUsed = 0;
-	}
 }
 
 const uint32* FNiagaraGPUInstanceCountManager::GetGPUReadback()
diff --git a/Engine/Plugins/FX/Niagara/Source/Niagara/Private/NiagaraRendererMeshes.cpp b/Engine/Plugins/FX/Niagara/Source/Niagara/Private/NiagaraRendererMeshes.cpp
--- a/Engine/Plugins/FX/Niagara/Source/Niagara/Private/NiagaraRendererMeshes.cpp
+++ b/Engine/Plugins/FX/Niagara/Source/Niagara/Private/NiagaraRendererMeshes.cpp
@@ -27,7 +27,6 @@ DECLARE_CYCLE_STAT(TEXT("Generate Mesh Vertex Data [GT]"), STAT_NiagaraGenMeshVe
 DECLARE_CYCLE_STAT(TEXT("Render Meshes [RT]"), STAT_NiagaraRenderMeshes, STATGROUP_Niagara);
 DECLARE_CYCLE_STAT(TEXT("Render Meshes - Allocate GPU Data [RT]"), STAT_NiagaraRenderMeshes_AllocateGPUData, STATGROUP_Niagara);
 
-
 DECLARE_DWORD_COUNTER_STAT(TEXT("NumMeshesRenderer"), STAT_NiagaraNumMeshes, STATGROUP_Niagara);
 DECLARE_DWORD_COUNTER_STAT(TEXT("NumMesheVerts"), STAT_NiagaraNumMeshVerts, STATGROUP_Niagara);
 
@@ -76,7 +75,7 @@ FNiagaraRendererMeshes::FNiagaraRendererMeshes(ERHIFeatureLevel::Type FeatureLev
 	LockedAxisSpace = Properties->LockedAxisSpace;
 	SortMode = Properties->SortMode;
 	bSortOnlyWhenTranslucent = Properties->bSortOnlyWhenTranslucent;
-	bGpuLowLatencyTranslucency = false;//-TEMP: Need to add IndirectArgs update point.  Properties->bGpuLowLatencyTranslucency && (SortMode == ENiagaraSortMode::None);
+	bGpuLowLatencyTranslucency = Properties->bGpuLowLatencyTranslucency;
 	bOverrideMaterials = Properties->bOverrideMaterials;
 	SubImageSize = Properties->SubImageSize;
 	bSubImageBlend = Properties->bSubImageBlend;
@@ -405,7 +404,7 @@ void FNiagaraRendererMeshes::InitializeSortInfo(FParticleMeshRenderData& Particl
 
 	OutSortInfo.ParticleCount = ParticleMeshRenderData.SourceParticleData->GetNumInstances();
 	OutSortInfo.SortMode = SortMode;
-	OutSortInfo.SetSortFlags(GNiagaraGPUSortingUseMaxPrecision != 0, ParticleMeshRenderData.bHasTranslucentMaterials);
+	OutSortInfo.SetSortFlags(GNiagaraGPUSortingUseMaxPrecision != 0, ParticleMeshRenderData.SourceParticleData->GetGPUDataReadyStage());
 	OutSortInfo.bEnableCulling = ParticleMeshRenderData.bNeedsCull;
 	OutSortInfo.RendererVisTagAttributeOffset = ParticleMeshRenderData.RendererVisTagOffset;
 	OutSortInfo.MeshIndexAttributeOffset = ParticleMeshRenderData.MeshIndexOffset;
@@ -825,6 +824,7 @@ FNiagaraMeshUniformBufferRef FNiagaraRendererMeshes::CreatePerViewUniformBuffer(
 }
 
 void FNiagaraRendererMeshes::CreateMeshBatchForSection(
+	const FParticleMeshRenderData& ParticleMeshRenderData,
 	FMeshBatch& MeshBatch,
 	FVertexFactory& VertexFactory,
 	FMaterialRenderProxy& MaterialProxy,
@@ -837,8 +837,7 @@ void FNiagaraRendererMeshes::CreateMeshBatchForSection(
 	uint32 NumInstances,
 	uint32 GPUCountBufferOffset,
 	bool bIsWireframe,
-	bool bIsInstancedStereo,
-	bool bDoGPUCulling
+	bool bIsInstancedStereo
 ) const
 {
 	MeshBatch.VertexFactory = &VertexFactory;
@@ -898,7 +897,9 @@ void FNiagaraRendererMeshes::CreateMeshBatchForSection(
 			Section.NumTriangles * 3,
 			Section.FirstIndex,
 			bIsInstancedStereo,
-			bDoGPUCulling);
+			ParticleMeshRenderData.bNeedsCull && ParticleMeshRenderData.bSortCullOnGpu,
+			ParticleMeshRenderData.SourceParticleData->GetGPUDataReadyStage()
+		);
 
 		BatchElement.NumPrimitives = 0;
 		BatchElement.IndirectArgsBuffer = IndirectDraw.Buffer;
@@ -1046,9 +1047,9 @@ void FNiagaraRendererMeshes::GetDynamicMeshElements(const TArray<const FSceneVie
 						const uint32 GPUCountBufferOffset = SortInfo.CulledGPUParticleCountOffset != INDEX_NONE ? SortInfo.CulledGPUParticleCountOffset : ParticleMeshRenderData.SourceParticleData->GetGPUInstanceCountBufferOffset();
 						FMeshBatch& MeshBatch = Collector.AllocateMesh();
 						CreateMeshBatchForSection(
+							ParticleMeshRenderData,
 							MeshBatch, VertexFactory, *MaterialProxy, *SceneProxy, MeshData, LODModel, Section, *View, ViewIndex, NumInstances,
-							GPUCountBufferOffset, bIsWireframe, bIsInstancedStereo,
-							ParticleMeshRenderData.bNeedsCull && ParticleMeshRenderData.bSortCullOnGpu
+							GPUCountBufferOffset, bIsWireframe, bIsInstancedStereo
 						);
 
 						Collector.AddMesh(ViewIndex, MeshBatch);
@@ -1176,7 +1177,7 @@ void FNiagaraRendererMeshes::GetDynamicRayTracingInstances(FRayTracingMaterialGa
 			}
 
 			FMeshBatch MeshBatch;
-			CreateMeshBatchForSection(MeshBatch, VertexFactory, *MaterialProxy, *SceneProxy, MeshData, LODModel, Section, *View, ViewIndex, NumInstances, INDEX_NONE, false, bIsInstancedStereo, ParticleMeshRenderData.bNeedsCull && ParticleMeshRenderData.bSortCullOnGpu);
+			CreateMeshBatchForSection(ParticleMeshRenderData, MeshBatch, VertexFactory, *MaterialProxy, *SceneProxy, MeshData, LODModel, Section, *View, ViewIndex, NumInstances, INDEX_NONE, false, bIsInstancedStereo);
 			MeshBatch.SegmentIndex = SectionIndex;
 			MeshBatch.LODIndex = LODIndex;
 
diff --git a/Engine/Plugins/FX/Niagara/Source/Niagara/Private/NiagaraRendererSprites.cpp b/Engine/Plugins/FX/Niagara/Source/Niagara/Private/NiagaraRendererSprites.cpp
--- a/Engine/Plugins/FX/Niagara/Source/Niagara/Private/NiagaraRendererSprites.cpp
+++ b/Engine/Plugins/FX/Niagara/Source/Niagara/Private/NiagaraRendererSprites.cpp
@@ -88,7 +88,7 @@ FNiagaraRendererSprites::FNiagaraRendererSprites(ERHIFeatureLevel::Type FeatureL
 	bSubImageBlend = Properties->bSubImageBlend;
 	bRemoveHMDRollInVR = Properties->bRemoveHMDRollInVR;
 	bSortOnlyWhenTranslucent = Properties->bSortOnlyWhenTranslucent;
-	bGpuLowLatencyTranslucency = Properties->bGpuLowLatencyTranslucency && (SortMode == ENiagaraSortMode::None);
+	bGpuLowLatencyTranslucency = Properties->bGpuLowLatencyTranslucency;
 	MinFacingCameraBlendDistance = Properties->MinFacingCameraBlendDistance;
 	MaxFacingCameraBlendDistance = Properties->MaxFacingCameraBlendDistance;
 	RendererVisibility = Properties->RendererVisibility;
@@ -346,7 +346,7 @@ void FNiagaraRendererSprites::InitializeSortInfo(FParticleSpriteRenderData& Part
 
 	OutSortInfo.ParticleCount = ParticleSpriteRenderData.SourceParticleData->GetNumInstances();
 	OutSortInfo.SortMode = SortMode;
-	OutSortInfo.SetSortFlags(GNiagaraGPUSortingUseMaxPrecision != 0, ParticleSpriteRenderData.bHasTranslucentMaterials);
+	OutSortInfo.SetSortFlags(GNiagaraGPUSortingUseMaxPrecision != 0, ParticleSpriteRenderData.SourceParticleData->GetGPUDataReadyStage());
 	OutSortInfo.bEnableCulling = ParticleSpriteRenderData.bNeedsCull;
 	OutSortInfo.RendererVisTagAttributeOffset = ParticleSpriteRenderData.RendererVisTagOffset;
 	OutSortInfo.RendererVisibility = RendererVisibility;
@@ -752,7 +752,14 @@ void FNiagaraRendererSprites::CreateMeshBatchForView(
 		NiagaraEmitterInstanceBatcher* Batcher = SceneProxy.GetBatcher();
 		check(Batcher);
 
-		IndirectDraw = Batcher->GetGPUInstanceCounterManager().AddDrawIndirect(GPUCountBufferOffset, NumIndicesPerInstance, 0, View.IsInstancedStereoPass(), bDoGPUCulling);
+		IndirectDraw = Batcher->GetGPUInstanceCounterManager().AddDrawIndirect(
+			GPUCountBufferOffset,
+			NumIndicesPerInstance,
+			0,
+			View.IsInstancedStereoPass(),
+			bDoGPUCulling,
+			ParticleSpriteRenderData.SourceParticleData->GetGPUDataReadyStage()
+		);
 	}
 
 	if (IndirectDraw.IsValid())
diff --git a/Engine/Plugins/FX/Niagara/Source/Niagara/Public/NiagaraCommon.h b/Engine/Plugins/FX/Niagara/Source/Niagara/Public/NiagaraCommon.h
--- a/Engine/Plugins/FX/Niagara/Source/Niagara/Public/NiagaraCommon.h
+++ b/Engine/Plugins/FX/Niagara/Source/Niagara/Public/NiagaraCommon.h
@@ -1550,3 +1550,16 @@ struct NIAGARA_API FSynchronizeWithParameterDefinitionsArgs
 	TArray<TTuple<FName, FName>> AdditionalOldToNewNames;
 };
 
+/** Enum used to track stage that GPU compute proxies will execute in. */
+namespace ENiagaraGpuComputeTickStage
+{
+	enum Type
+	{
+		PreInitViews,
+		PostInitViews,
+		PostOpaqueRender,
+		Max,
+		First = PreInitViews,
+		Last = PostOpaqueRender,
+	};
+};
diff --git a/Engine/Plugins/FX/Niagara/Source/Niagara/Public/NiagaraRendererMeshes.h b/Engine/Plugins/FX/Niagara/Source/Niagara/Public/NiagaraRendererMeshes.h
--- a/Engine/Plugins/FX/Niagara/Source/Niagara/Public/NiagaraRendererMeshes.h
+++ b/Engine/Plugins/FX/Niagara/Source/Niagara/Public/NiagaraRendererMeshes.h
@@ -43,6 +43,7 @@ protected:
 		class FNiagaraDataBuffer*		SourceParticleData = nullptr;
 
 		bool							bHasTranslucentMaterials = false;
+		bool							bAllTranslucentMaterials = false;
 		bool							bSortCullOnGpu = false;
 		bool							bNeedsSort = false;
 		bool							bNeedsCull = false;
@@ -111,6 +112,7 @@ protected:
 	FNiagaraMeshUniformBufferRef CreatePerViewUniformBuffer(FParticleMeshRenderData& ParticleMeshRenderData, const FSceneView& View, const FMeshData& MeshData, const FNiagaraSceneProxy& SceneProxy) const;
 
 	void CreateMeshBatchForSection(
+		const FParticleMeshRenderData& ParticleMeshRenderData,
 		FMeshBatch& MeshBatch,
 		FVertexFactory& VertexFactory,
 		FMaterialRenderProxy& MaterialProxy,
@@ -123,8 +125,7 @@ protected:
 		uint32 NumInstances,
 		uint32 GPUCountBufferOffset,
 		bool bIsWireframe,
-		bool bIsInstancedStereo,
-		bool bDoGPUCulling
+		bool bIsInstancedStereo
 	) const;
 
 private:
diff --git a/Engine/Plugins/FX/Niagara/Source/Niagara/Public/NiagaraSystemGpuComputeProxy.h b/Engine/Plugins/FX/Niagara/Source/Niagara/Public/NiagaraSystemGpuComputeProxy.h
--- a/Engine/Plugins/FX/Niagara/Source/Niagara/Public/NiagaraSystemGpuComputeProxy.h
+++ b/Engine/Plugins/FX/Niagara/Source/Niagara/Public/NiagaraSystemGpuComputeProxy.h
@@ -10,19 +10,6 @@ struct FNiagaraComputeExecutionContext;
 class FNiagaraGPUInstanceCountManager;
 class FNiagaraGPUSystemTick;
 
-namespace ENiagaraGpuComputeTickStage
-{
-	enum Type
-	{
-		PreInitViews,
-		PostInitViews,
-		PostOpaqueRender,
-		Max,
-		First = PreInitViews,
-		Last = PostOpaqueRender,
-	};
-};
-
 class FNiagaraSystemGpuComputeProxy
 {
 	friend class NiagaraEmitterInstanceBatcher;
diff --git a/Engine/Plugins/FX/Niagara/Source/NiagaraVertexFactories/Public/NiagaraDrawIndirect.h b/Engine/Plugins/FX/Niagara/Source/NiagaraVertexFactories/Public/NiagaraDrawIndirect.h
--- a/Engine/Plugins/FX/Niagara/Source/NiagaraVertexFactories/Public/NiagaraDrawIndirect.h
+++ b/Engine/Plugins/FX/Niagara/Source/NiagaraVertexFactories/Public/NiagaraDrawIndirect.h
@@ -13,11 +13,20 @@ NiagaraDrawIndirect.h: Niagara shader to generate the draw indirect args for Nia
 
 #define NIAGARA_DRAW_INDIRECT_ARGS_GEN_THREAD_COUNT 64
 #define NIAGARA_DRAW_INDIRECT_ARGS_SIZE 5
-#define NIAGARA_DRAW_INDIRECT_TASK_INFO_SIZE 4
+#define NIAGARA_DRAW_INDIRECT_TASK_INFO_SIZE 5
 
 // #define NIAGARA_COPY_BUFFER_THREAD_COUNT 64
 // #define NIAGARA_COPY_BUFFER_BUFFER_COUNT 3
 
+enum class ENiagaraDrawIndirectArgGenTaskFlags : uint32
+{
+	None = 0,
+	UseCulledCounts = 1 << 0,
+	InstancedStereo = 1 << 1,
+	PostOpaque = 1 << 2,
+};
+ENUM_CLASS_FLAGS(ENiagaraDrawIndirectArgGenTaskFlags);
+
 /**
 * Task info when generating draw indirect frame buffer.
 * Task is either about generate Niagara renderers drawindirect buffer,
@@ -25,29 +34,16 @@ NiagaraDrawIndirect.h: Niagara shader to generate the draw indirect args for Nia
 */
 struct FNiagaraDrawIndirectArgGenTaskInfo
 {
-	enum Flags : uint32
-	{
-		Flag_UseCulledCounts = 1 << 0,
-		Flag_InstancedStereo = 1 << 1,
-	};
-
-	FNiagaraDrawIndirectArgGenTaskInfo(uint32 InInstanceCountBufferOffset, uint32 InNumIndicesPerInstance, uint32 InStartIndexLocation,
-		bool bInIsInstancedStereo, bool bInUseCulledCounts)
-		: InstanceCountBufferOffset(InInstanceCountBufferOffset)
+	explicit FNiagaraDrawIndirectArgGenTaskInfo(uint32 InInstanceCountBufferOffset, uint32 InNumIndicesPerInstance, uint32 InStartIndexLocation, ENiagaraDrawIndirectArgGenTaskFlags InFlags)
+		: IndirectArgsBufferOffset(INDEX_NONE)
+		, InstanceCountBufferOffset(InInstanceCountBufferOffset)
 		, NumIndicesPerInstance(InNumIndicesPerInstance)
 		, StartIndexLocation(InStartIndexLocation)
+		, Flags((uint32)InFlags)
 	{
-		Flags = 0;
-		if (bInUseCulledCounts)
-		{
-			Flags |= Flag_UseCulledCounts;
-		}
-		if (bInIsInstancedStereo)
-		{
-			Flags |= Flag_InstancedStereo;
-		}
 	}
 
+	uint32 IndirectArgsBufferOffset;	// Ignored in hash / comparison since we can de-duplicate the rest
 	uint32 InstanceCountBufferOffset;
 	uint32 NumIndicesPerInstance; // When -1 the counter needs to be reset to 0.
 	uint32 StartIndexLocation;
diff --git a/Engine/Source/Runtime/Engine/Private/GPUSortManager.cpp b/Engine/Source/Runtime/Engine/Private/GPUSortManager.cpp
--- a/Engine/Source/Runtime/Engine/Private/GPUSortManager.cpp
+++ b/Engine/Source/Runtime/Engine/Private/GPUSortManager.cpp
@@ -399,12 +399,12 @@ void FGPUSortManager::FSortBatch::UpdateProcessingOrder()
 {
 	if (EnumHasAnyFlags(Flags, EGPUSortFlags::SortAfterPreRender))
 	{
-		checkSlow(!EnumHasAnyFlags(Flags, EGPUSortFlags::KeyGenAfterPostRenderOpaque));
+		check(!EnumHasAnyFlags(Flags, EGPUSortFlags::KeyGenAfterPostRenderOpaque));
 		ProcessingOrder = ESortBatchProcessingOrder::KeyGenAndSortAfterPreRender;
 	}
 	else // SortAfterPostRenderOpaque
 	{
-		checkSlow(EnumHasAnyFlags(Flags, EGPUSortFlags::SortAfterPostRenderOpaque));
+		check(EnumHasAnyFlags(Flags, EGPUSortFlags::SortAfterPostRenderOpaque));
 		if (EnumHasAnyFlags(Flags, EGPUSortFlags::KeyGenAfterPreRender))
 		{
 			ProcessingOrder = ESortBatchProcessingOrder::KeyGenAfterPreRenderAndSortAfterPostRenderOpaque;
@@ -817,6 +817,7 @@ void FGPUSortManager::OnPostRenderOpaque(FRDGBuilder& GraphBuilder)
 					SortBatch.SortBuffers = nullptr;
 				}
 			}
+			PostPostRenderEvent.Broadcast(RHICmdList);
 
 			ResetDynamicValuesBuffers();
 		}
diff --git a/Engine/Source/Runtime/Engine/Public/GPUSortManager.h b/Engine/Source/Runtime/Engine/Public/GPUSortManager.h
--- a/Engine/Source/Runtime/Engine/Public/GPUSortManager.h
+++ b/Engine/Source/Runtime/Engine/Public/GPUSortManager.h
@@ -452,6 +452,12 @@ public:
 	DECLARE_EVENT_OneParam(FGPUSortManager, FPostPreRenderEvent, FRHICommandListImmediate&);
 	FPostPreRenderEvent PostPreRenderEvent;
 
+	/**
+	 * Event to register and receive post-postrender notification.
+	 */
+	DECLARE_EVENT_OneParam(FGPUSortManager, FPostPostRenderEvent, FRHICommandListImmediate&);
+	FPostPreRenderEvent PostPostRenderEvent;
+
 private:
 
 	/** A helper to test that a flag has one and exactly one of A or B. */
