Wrapped FVectorVMContext in NiagaraDataInterfaceFunctionContext for external functions to facilitate running an experimental VM in parallel.

diff --git a/Engine/Plugins/FX/Niagara/Source/Niagara/Classes/NiagaraDataInterfaceTexture.h b/Engine/Plugins/FX/Niagara/Source/Niagara/Classes/NiagaraDataInterfaceTexture.h
--- a/Engine/Plugins/FX/Niagara/Source/Niagara/Classes/NiagaraDataInterfaceTexture.h
+++ b/Engine/Plugins/FX/Niagara/Source/Niagara/Classes/NiagaraDataInterfaceTexture.h
@@ -37,9 +37,9 @@ public:
 	virtual bool PerInstanceTick(void* PerInstanceData, FNiagaraSystemInstance* SystemInstance, float DeltaSeconds) override;
 	//UNiagaraDataInterface Interface End
 
-	void SampleTexture(FVectorVMContext& Context);
-	void GetTextureDimensions(FVectorVMContext& Context);
-	void SamplePseudoVolumeTexture(FVectorVMContext& Context);
+	void SampleTexture(FVectorVMExternalFunctionContext& Context);
+	void GetTextureDimensions(FVectorVMExternalFunctionContext& Context);
+	void SamplePseudoVolumeTexture(FVectorVMExternalFunctionContext& Context);
 
 	virtual bool Equals(const UNiagaraDataInterface* Other) const override;
 
diff --git a/Engine/Plugins/FX/Niagara/Source/Niagara/Classes/NiagaraDataInterfaceVectorField.h b/Engine/Plugins/FX/Niagara/Source/Niagara/Classes/NiagaraDataInterfaceVectorField.h
--- a/Engine/Plugins/FX/Niagara/Source/Niagara/Classes/NiagaraDataInterfaceVectorField.h
+++ b/Engine/Plugins/FX/Niagara/Source/Niagara/Classes/NiagaraDataInterfaceVectorField.h
@@ -58,10 +58,10 @@ public:
 	//~ UNiagaraDataInterface interface END
 
 	// VM functions
-	void GetFieldDimensions(FVectorVMContext& Context);
-	void GetFieldBounds(FVectorVMContext& Context); 
-	void GetFieldTilingAxes(FVectorVMContext& Context);
-	void SampleVectorField(FVectorVMContext& Context);
+	void GetFieldDimensions(FVectorVMExternalFunctionContext& Context);
+	void GetFieldBounds(FVectorVMExternalFunctionContext& Context); 
+	void GetFieldTilingAxes(FVectorVMExternalFunctionContext& Context);
+	void SampleVectorField(FVectorVMExternalFunctionContext& Context);
 	
 	//	
 	FVector GetTilingAxes() const;
diff --git a/Engine/Plugins/FX/Niagara/Source/Niagara/Classes/NiagaraDataInterfaceVolumeTexture.h b/Engine/Plugins/FX/Niagara/Source/Niagara/Classes/NiagaraDataInterfaceVolumeTexture.h
--- a/Engine/Plugins/FX/Niagara/Source/Niagara/Classes/NiagaraDataInterfaceVolumeTexture.h
+++ b/Engine/Plugins/FX/Niagara/Source/Niagara/Classes/NiagaraDataInterfaceVolumeTexture.h
@@ -38,8 +38,8 @@ public:
 	virtual bool PerInstanceTick(void* PerInstanceData, FNiagaraSystemInstance* SystemInstance, float DeltaSeconds) override;
 	//UNiagaraDataInterface Interface End
 
-	void SampleVolumeTexture(FVectorVMContext& Context);
-	void GetTextureDimensions(FVectorVMContext& Context);
+	void SampleVolumeTexture(FVectorVMExternalFunctionContext& Context);
+	void GetTextureDimensions(FVectorVMExternalFunctionContext& Context);
 
 	virtual bool Equals(const UNiagaraDataInterface* Other) const override;
 
diff --git a/Engine/Plugins/FX/Niagara/Source/Niagara/Classes/NiagaraScriptExecutionContext.h b/Engine/Plugins/FX/Niagara/Source/Niagara/Classes/NiagaraScriptExecutionContext.h
--- a/Engine/Plugins/FX/Niagara/Source/Niagara/Classes/NiagaraScriptExecutionContext.h
+++ b/Engine/Plugins/FX/Niagara/Source/Niagara/Classes/NiagaraScriptExecutionContext.h
@@ -247,7 +247,7 @@ protected:
 	ENiagaraSystemSimulationScript ScriptType;
 
 	/** Helper function that handles calling into per instance DI calls and massages the VM context appropriately. */
-	void PerInstanceFunctionHook(FVectorVMContext& Context, int32 PerInstFunctionIndex, int32 UserPtrIndex);
+	void PerInstanceFunctionHook(FVectorVMExternalFunctionContext& Context, int32 PerInstFunctionIndex, int32 UserPtrIndex);
 
 public:
 	FNiagaraSystemScriptExecutionContext(ENiagaraSystemSimulationScript InScriptType) : SystemInstances(nullptr), ScriptType(InScriptType){}
diff --git a/Engine/Plugins/FX/Niagara/Source/Niagara/Private/DataInterface/NiagaraDataInterfaceDebugDraw.cpp b/Engine/Plugins/FX/Niagara/Source/Niagara/Private/DataInterface/NiagaraDataInterfaceDebugDraw.cpp
--- a/Engine/Plugins/FX/Niagara/Source/Niagara/Private/DataInterface/NiagaraDataInterfaceDebugDraw.cpp
+++ b/Engine/Plugins/FX/Niagara/Source/Niagara/Private/DataInterface/NiagaraDataInterfaceDebugDraw.cpp
@@ -1464,7 +1464,7 @@ namespace NDIDebugDrawLocal
 	{
 		struct VMBindings
 		{
-			VMBindings(FVectorVMContext& Context)
+			VMBindings(FVectorVMExternalFunctionContext& Context)
 				: LineStartParam(Context)
 				, LineEndParam(Context)
 				, ColorParam(Context)
@@ -1478,7 +1478,7 @@ namespace NDIDebugDrawLocal
 
 		struct PersistentVMBindings
 		{
-			PersistentVMBindings(FVectorVMContext& Context)
+			PersistentVMBindings(FVectorVMExternalFunctionContext& Context)
 				: StartLocationParam(Context)
 				, StartLocationWSParam(Context)
 				, EndLocationParam(Context)
@@ -1514,7 +1514,7 @@ namespace NDIDebugDrawLocal
 	{
 		struct VMBindings
 		{
-			VMBindings(FVectorVMContext& Context)
+			VMBindings(FVectorVMExternalFunctionContext& Context)
 				: LocationParam(Context)
 				, XAxisParam(Context)
 				, YAxisParam(Context)
@@ -1538,7 +1538,7 @@ namespace NDIDebugDrawLocal
 
 		struct PersistentVMBindings
 		{
-			PersistentVMBindings(FVectorVMContext& Context)
+			PersistentVMBindings(FVectorVMExternalFunctionContext& Context)
 				: LocationParam(Context)
 				, LocationWSParam(Context)
 				, OffsetParam(Context)
@@ -1595,7 +1595,7 @@ namespace NDIDebugDrawLocal
 	{
 		struct VMBindings
 		{
-			VMBindings(FVectorVMContext& Context)
+			VMBindings(FVectorVMExternalFunctionContext& Context)
 				: LocationParam(Context)
 				, XAxisParam(Context)
 				, YAxisParam(Context)
@@ -1615,7 +1615,7 @@ namespace NDIDebugDrawLocal
 
 		struct PersistentVMBindings
 		{
-			PersistentVMBindings(FVectorVMContext& Context)
+			PersistentVMBindings(FVectorVMExternalFunctionContext& Context)
 				: LocationParam(Context)
 				, LocationWSParam(Context)
 				, OffsetParam(Context)
@@ -1664,7 +1664,7 @@ namespace NDIDebugDrawLocal
 	{
 		struct VMBindings
 		{
-			VMBindings(FVectorVMContext& Context)
+			VMBindings(FVectorVMExternalFunctionContext& Context)
 				: LocationParam(Context)
 				, RotationParam(Context)
 				, ExtentsParam(Context)
@@ -1680,7 +1680,7 @@ namespace NDIDebugDrawLocal
 
 		struct PersistentVMBindings
 		{
-			PersistentVMBindings(FVectorVMContext& Context)
+			PersistentVMBindings(FVectorVMExternalFunctionContext& Context)
 				: LocationParam(Context)
 				, LocationWSParam(Context)
 				, OffsetParam(Context)
@@ -1727,7 +1727,7 @@ namespace NDIDebugDrawLocal
 	{
 		struct VMBindings
 		{
-			VMBindings(FVectorVMContext& Context)
+			VMBindings(FVectorVMExternalFunctionContext& Context)
 				: LocationParam(Context)
 				, RadiusParam(Context)
 				, SegmentsParam(Context)
@@ -1743,7 +1743,7 @@ namespace NDIDebugDrawLocal
 
 		struct PersistentVMBindings
 		{
-			PersistentVMBindings(FVectorVMContext& Context)
+			PersistentVMBindings(FVectorVMExternalFunctionContext& Context)
 				: CenterParam(Context)
 				, CenterWSParam(Context)
 				, OffsetParam(Context)
@@ -1799,7 +1799,7 @@ namespace NDIDebugDrawLocal
 	{
 		struct VMBindings
 		{
-			VMBindings(FVectorVMContext& Context)
+			VMBindings(FVectorVMExternalFunctionContext& Context)
 				: LocationParam(Context)
 				, AxisParam(Context)
 				, HeightParam(Context)
@@ -1821,7 +1821,7 @@ namespace NDIDebugDrawLocal
 
 		struct PersistentVMBindings
 		{
-			PersistentVMBindings(FVectorVMContext& Context)
+			PersistentVMBindings(FVectorVMExternalFunctionContext& Context)
 				: LocationParam(Context)
 				, LocationWSParam(Context)
 				, OffsetParam(Context)
@@ -1882,7 +1882,7 @@ namespace NDIDebugDrawLocal
 	{
 		struct VMBindings
 		{
-			VMBindings(FVectorVMContext& Context)
+			VMBindings(FVectorVMExternalFunctionContext& Context)
 				: LocationParam(Context)
 				, AxisParam(Context)
 				, HeightParam(Context)
@@ -1906,7 +1906,7 @@ namespace NDIDebugDrawLocal
 
 		struct PersistentVMBindings
 		{
-			PersistentVMBindings(FVectorVMContext& Context)
+			PersistentVMBindings(FVectorVMExternalFunctionContext& Context)
 				: LocationParam(Context)
 				, LocationWSParam(Context)
 				, OffsetParam(Context)
@@ -1961,7 +1961,7 @@ namespace NDIDebugDrawLocal
 	{
 		struct VMBindings
 		{
-			VMBindings(FVectorVMContext& Context)
+			VMBindings(FVectorVMExternalFunctionContext& Context)
 				: LocationParam(Context)
 				, AxisParam(Context)
 				, MajorRadiusParam(Context)
@@ -1983,7 +1983,7 @@ namespace NDIDebugDrawLocal
 
 		struct PersistentVMBindings
 		{
-			PersistentVMBindings(FVectorVMContext& Context)
+			PersistentVMBindings(FVectorVMExternalFunctionContext& Context)
 				: LocationParam(Context)
 				, LocationWSParam(Context)
 				, OffsetParam(Context)
@@ -2037,7 +2037,7 @@ namespace NDIDebugDrawLocal
 	{
 		struct VMBindings
 		{
-			VMBindings(FVectorVMContext& Context)
+			VMBindings(FVectorVMExternalFunctionContext& Context)
 				: LocationParam(Context)
 				, RotationParam(Context)
 				, ScaleParam(Context)
@@ -2051,7 +2051,7 @@ namespace NDIDebugDrawLocal
 		
 		struct PersistentVMBindings
 		{
-			PersistentVMBindings(FVectorVMContext& Context)
+			PersistentVMBindings(FVectorVMExternalFunctionContext& Context)
 				: LocationParam(Context)
 				, LocationWSParam(Context)
 				, OffsetParam(Context)
@@ -2092,7 +2092,7 @@ namespace NDIDebugDrawLocal
 	{
 		struct VMBindings
 		{
-			VMBindings(FVectorVMContext& Context)
+			VMBindings(FVectorVMExternalFunctionContext& Context)
 				: CenterParam(Context)
 				, RotationParam(Context)
 				, ExtentsParam(Context)
@@ -2112,7 +2112,7 @@ namespace NDIDebugDrawLocal
 
 		struct PersistentVMBindings
 		{
-			PersistentVMBindings(FVectorVMContext& Context)
+			PersistentVMBindings(FVectorVMExternalFunctionContext& Context)
 				: LocationParam(Context)
 				, LocationWSParam(Context)
 				, OffsetParam(Context)
@@ -2160,7 +2160,7 @@ namespace NDIDebugDrawLocal
 	{
 		struct VMBindings
 		{
-			VMBindings(FVectorVMContext& Context)
+			VMBindings(FVectorVMExternalFunctionContext& Context)
 				: CenterParam(Context)
 				, RotationParam(Context)
 				, ExtentsParam(Context)
@@ -2182,7 +2182,7 @@ namespace NDIDebugDrawLocal
 
 		struct PersistentVMBindings
 		{
-			PersistentVMBindings(FVectorVMContext& Context)
+			PersistentVMBindings(FVectorVMExternalFunctionContext& Context)
 				: LocationParam(Context)
 				, LocationWSParam(Context)
 				, OffsetParam(Context)
@@ -2231,7 +2231,7 @@ namespace NDIDebugDrawLocal
 
 
 	template<typename TPrimType>
-	void DrawDebug(FVectorVMContext& Context)
+	void DrawDebug(FVectorVMExternalFunctionContext& Context)
 	{
 		VectorVM::FUserPtrHandler<FNDIDebugDrawInstanceData_GameThread> InstanceData(Context);
 		FNDIInputParam<FNiagaraBool> ExecuteParam(Context);
@@ -2243,7 +2243,7 @@ namespace NDIDebugDrawLocal
 			return;
 		}
 
-		for ( int i=0; i < Context.NumInstances; ++i )
+		for ( int i=0; i < Context.GetNumInstances(); ++i )
 		{
 			const bool bExecute = ExecuteParam.GetAndAdvance();
 			TPrimType::Draw(InstanceData, Bindings, bExecute);
@@ -2252,7 +2252,7 @@ namespace NDIDebugDrawLocal
 	}
 
 	template<typename TPrimType>
-	void DrawDebugPersistent(FVectorVMContext& Context)
+	void DrawDebugPersistent(FVectorVMExternalFunctionContext& Context)
 	{
 		VectorVM::FUserPtrHandler<FNDIDebugDrawInstanceData_GameThread> InstanceData(Context);
 		typename TPrimType::PersistentVMBindings Bindings(Context);
diff --git a/Engine/Plugins/FX/Niagara/Source/Niagara/Private/NDISkeletalMesh_BoneSampling.cpp b/Engine/Plugins/FX/Niagara/Source/Niagara/Private/NDISkeletalMesh_BoneSampling.cpp
--- a/Engine/Plugins/FX/Niagara/Source/Niagara/Private/NDISkeletalMesh_BoneSampling.cpp
+++ b/Engine/Plugins/FX/Niagara/Source/Niagara/Private/NDISkeletalMesh_BoneSampling.cpp
@@ -366,53 +366,53 @@ void UNiagaraDataInterfaceSkeletalMesh::BindSkeletonSamplingFunction(const FVMEx
 	else if (BindingInfo.Name == FSkeletalMeshInterfaceHelper::RandomBoneName)
 	{
 		check(BindingInfo.GetNumInputs() == 4 && BindingInfo.GetNumOutputs() == 1);
-		OutFunc = FVMExternalFunction::CreateLambda([this](FVectorVMContext& Context) { this->RandomBone(Context); });
+		OutFunc = FVMExternalFunction::CreateLambda([this](FVectorVMExternalFunctionContext& Context) { this->RandomBone(Context); });
 	}
 	else if (BindingInfo.Name == FSkeletalMeshInterfaceHelper::GetBoneCountName)
 	{
 		check(BindingInfo.GetNumInputs() == 1 && BindingInfo.GetNumOutputs() == 1);
-		OutFunc = FVMExternalFunction::CreateLambda([this](FVectorVMContext& Context) { this->GetBoneCount(Context); });
+		OutFunc = FVMExternalFunction::CreateLambda([this](FVectorVMExternalFunctionContext& Context) { this->GetBoneCount(Context); });
 	}
 	else if (BindingInfo.Name == FSkeletalMeshInterfaceHelper::RandomFilteredBoneName)
 	{
 		check(BindingInfo.GetNumInputs() == 4 && BindingInfo.GetNumOutputs() == 1);
-		OutFunc = FVMExternalFunction::CreateLambda([this](FVectorVMContext& Context) { this->RandomFilteredBone(Context); });
+		OutFunc = FVMExternalFunction::CreateLambda([this](FVectorVMExternalFunctionContext& Context) { this->RandomFilteredBone(Context); });
 	}
 	else if (BindingInfo.Name == FSkeletalMeshInterfaceHelper::GetFilteredBoneCountName)
 	{
 		check(BindingInfo.GetNumInputs() == 1 && BindingInfo.GetNumOutputs() == 1);
-		OutFunc = FVMExternalFunction::CreateLambda([this](FVectorVMContext& Context) { this->GetFilteredBoneCount(Context); });
+		OutFunc = FVMExternalFunction::CreateLambda([this](FVectorVMExternalFunctionContext& Context) { this->GetFilteredBoneCount(Context); });
 	}
 	else if (BindingInfo.Name == FSkeletalMeshInterfaceHelper::GetFilteredBoneAtName)
 	{
 		check(BindingInfo.GetNumInputs() == 2 && BindingInfo.GetNumOutputs() == 1);
-		OutFunc = FVMExternalFunction::CreateLambda([this](FVectorVMContext& Context) { this->GetFilteredBoneAt(Context); });
+		OutFunc = FVMExternalFunction::CreateLambda([this](FVectorVMExternalFunctionContext& Context) { this->GetFilteredBoneAt(Context); });
 	}
 	else if (BindingInfo.Name == FSkeletalMeshInterfaceHelper::RandomUnfilteredBoneName)
 	{
 		check(BindingInfo.GetNumInputs() == 4 && BindingInfo.GetNumOutputs() == 1);
-		OutFunc = FVMExternalFunction::CreateLambda([this](FVectorVMContext& Context) { this->RandomUnfilteredBone(Context); });
+		OutFunc = FVMExternalFunction::CreateLambda([this](FVectorVMExternalFunctionContext& Context) { this->RandomUnfilteredBone(Context); });
 	}
 	else if (BindingInfo.Name == FSkeletalMeshInterfaceHelper::GetUnfilteredBoneCountName)
 	{
 		check(BindingInfo.GetNumInputs() == 1 && BindingInfo.GetNumOutputs() == 1);
-		OutFunc = FVMExternalFunction::CreateLambda([this](FVectorVMContext& Context) { this->GetUnfilteredBoneCount(Context); });
+		OutFunc = FVMExternalFunction::CreateLambda([this](FVectorVMExternalFunctionContext& Context) { this->GetUnfilteredBoneCount(Context); });
 	}
 	else if (BindingInfo.Name == FSkeletalMeshInterfaceHelper::GetUnfilteredBoneAtName)
 	{
 		check(BindingInfo.GetNumInputs() == 2 && BindingInfo.GetNumOutputs() == 1);
-		OutFunc = FVMExternalFunction::CreateLambda([this](FVectorVMContext& Context) { this->GetUnfilteredBoneAt(Context); });
+		OutFunc = FVMExternalFunction::CreateLambda([this](FVectorVMExternalFunctionContext& Context) { this->GetUnfilteredBoneAt(Context); });
 	}
 	//Socket Functions
 	else if (BindingInfo.Name == FSkeletalMeshInterfaceHelper::RandomFilteredSocketName)
 	{
 		check(BindingInfo.GetNumInputs() == 4 && BindingInfo.GetNumOutputs() == 1);
-		OutFunc = FVMExternalFunction::CreateLambda([this](FVectorVMContext& Context) { this->RandomFilteredSocket(Context); });
+		OutFunc = FVMExternalFunction::CreateLambda([this](FVectorVMExternalFunctionContext& Context) { this->RandomFilteredSocket(Context); });
 	}
 	else if (BindingInfo.Name == FSkeletalMeshInterfaceHelper::GetFilteredSocketCountName)
 	{
 		check(BindingInfo.GetNumInputs() == 1 && BindingInfo.GetNumOutputs() == 1);
-		OutFunc = FVMExternalFunction::CreateLambda([this](FVectorVMContext& Context) { this->GetFilteredSocketCount(Context); });
+		OutFunc = FVMExternalFunction::CreateLambda([this](FVectorVMExternalFunctionContext& Context) { this->GetFilteredSocketCount(Context); });
 	}
 	else if (BindingInfo.Name == FSkeletalMeshInterfaceHelper::GetFilteredSocketBoneAtName)
 	{
@@ -428,24 +428,24 @@ void UNiagaraDataInterfaceSkeletalMesh::BindSkeletonSamplingFunction(const FVMEx
 	else if (BindingInfo.Name == FSkeletalMeshInterfaceHelper::RandomFilteredSocketOrBoneName)
 	{
 		check(BindingInfo.GetNumInputs() == 4 && BindingInfo.GetNumOutputs() == 1);
-		OutFunc = FVMExternalFunction::CreateLambda([this](FVectorVMContext& Context) { this->RandomFilteredSocketOrBone(Context); });
+		OutFunc = FVMExternalFunction::CreateLambda([this](FVectorVMExternalFunctionContext& Context) { this->RandomFilteredSocketOrBone(Context); });
 	}
 	else if (BindingInfo.Name == FSkeletalMeshInterfaceHelper::GetFilteredSocketOrBoneCountName)
 	{
 		check(BindingInfo.GetNumInputs() == 1 && BindingInfo.GetNumOutputs() == 1);
-		OutFunc = FVMExternalFunction::CreateLambda([this](FVectorVMContext& Context) { this->GetFilteredSocketOrBoneCount(Context); });
+		OutFunc = FVMExternalFunction::CreateLambda([this](FVectorVMExternalFunctionContext& Context) { this->GetFilteredSocketOrBoneCount(Context); });
 	}
 	else if (BindingInfo.Name == FSkeletalMeshInterfaceHelper::GetFilteredSocketOrBoneAtName)
 	{
 		check(BindingInfo.GetNumInputs() == 2 && BindingInfo.GetNumOutputs() == 1);
-		OutFunc = FVMExternalFunction::CreateLambda([this](FVectorVMContext& Context) { this->GetFilteredSocketOrBoneBoneAt(Context); });
+		OutFunc = FVMExternalFunction::CreateLambda([this](FVectorVMExternalFunctionContext& Context) { this->GetFilteredSocketOrBoneBoneAt(Context); });
 	}
 }
 
 //////////////////////////////////////////////////////////////////////////
 // Direct sampling from listed sockets and bones.
 
-void UNiagaraDataInterfaceSkeletalMesh::GetFilteredBoneCount(FVectorVMContext& Context)
+void UNiagaraDataInterfaceSkeletalMesh::GetFilteredBoneCount(FVectorVMExternalFunctionContext& Context)
 {
 	SCOPE_CYCLE_COUNTER(STAT_NiagaraSkel_Bone_Sample);
 
@@ -454,13 +454,13 @@ void UNiagaraDataInterfaceSkeletalMesh::GetFilteredBoneCount(FVectorVMContext& C
 	FNDIOutputParam<int32> OutCount(Context);
 
 	const int32 Num = InstData->NumFilteredBones;
-	for (int32 i = 0; i < Context.NumInstances; ++i)
+	for (int32 i = 0; i < Context.GetNumInstances(); ++i)
 	{
 		OutCount.SetAndAdvance(Num);
 	}
 }
 
-void UNiagaraDataInterfaceSkeletalMesh::GetFilteredBoneAt(FVectorVMContext& Context)
+void UNiagaraDataInterfaceSkeletalMesh::GetFilteredBoneAt(FVectorVMExternalFunctionContext& Context)
 {
 	VectorVM::FUserPtrHandler<FNDISkeletalMesh_InstanceData> InstData(Context);
 	FNDIInputParam<int32> BoneParam(Context);
@@ -470,7 +470,7 @@ void UNiagaraDataInterfaceSkeletalMesh::GetFilteredBoneAt(FVectorVMContext& Cont
 	const int32 Max = InstData->NumFilteredBones - 1;
 	if (Max >= 0)
 	{
-		for (int32 i = 0; i < Context.NumInstances; ++i)
+		for (int32 i = 0; i < Context.GetNumInstances(); ++i)
 		{
 			const int32 BoneIndex = FMath::Clamp(BoneParam.GetAndAdvance(), 0, Max);
 			OutBone.SetAndAdvance(InstData->FilteredAndUnfilteredBones[BoneIndex]);
@@ -478,14 +478,14 @@ void UNiagaraDataInterfaceSkeletalMesh::GetFilteredBoneAt(FVectorVMContext& Cont
 	}
 	else
 	{
-		for (int32 i = 0; i < Context.NumInstances; ++i)
+		for (int32 i = 0; i < Context.GetNumInstances(); ++i)
 		{
 			OutBone.SetAndAdvance(-1);
 		}
 	}
 }
 
-void UNiagaraDataInterfaceSkeletalMesh::RandomFilteredBone(FVectorVMContext& Context)
+void UNiagaraDataInterfaceSkeletalMesh::RandomFilteredBone(FVectorVMExternalFunctionContext& Context)
 {
 	VectorVM::FUserPtrHandler<FNDISkeletalMesh_InstanceData> InstData(Context);
 	FNDIRandomHelper RandHelper(Context);
@@ -495,7 +495,7 @@ void UNiagaraDataInterfaceSkeletalMesh::RandomFilteredBone(FVectorVMContext& Con
 	const int32 Max = InstData->NumFilteredBones - 1;
 	if (Max >= 0)
 	{
-		for (int32 i = 0; i < Context.NumInstances; ++i)
+		for (int32 i = 0; i < Context.GetNumInstances(); ++i)
 		{
 			RandHelper.GetAndAdvance();
 			const int32 BoneIndex = RandHelper.RandRange(i, 0, Max);
@@ -504,14 +504,14 @@ void UNiagaraDataInterfaceSkeletalMesh::RandomFilteredBone(FVectorVMContext& Con
 	}
 	else
 	{
-		for (int32 i = 0; i < Context.NumInstances; ++i)
+		for (int32 i = 0; i < Context.GetNumInstances(); ++i)
 		{
 			OutBone.SetAndAdvance(-1);
 		}
 	}
 }
 
-void UNiagaraDataInterfaceSkeletalMesh::GetUnfilteredBoneCount(FVectorVMContext& Context)
+void UNiagaraDataInterfaceSkeletalMesh::GetUnfilteredBoneCount(FVectorVMExternalFunctionContext& Context)
 {
 	SCOPE_CYCLE_COUNTER(STAT_NiagaraSkel_Bone_Sample);
 	VectorVM::FUserPtrHandler<FNDISkeletalMesh_InstanceData> InstData(Context);
@@ -519,13 +519,13 @@ void UNiagaraDataInterfaceSkeletalMesh::GetUnfilteredBoneCount(FVectorVMContext&
 	FNDIOutputParam<int32> OutCount(Context);
 
 	const int32 Num = InstData->NumUnfilteredBones;
-	for (int32 i = 0; i < Context.NumInstances; ++i)
+	for (int32 i = 0; i < Context.GetNumInstances(); ++i)
 	{
 		OutCount.SetAndAdvance(Num);
 	}
 }
 
-void UNiagaraDataInterfaceSkeletalMesh::GetUnfilteredBoneAt(FVectorVMContext& Context)
+void UNiagaraDataInterfaceSkeletalMesh::GetUnfilteredBoneAt(FVectorVMExternalFunctionContext& Context)
 {
 	VectorVM::FUserPtrHandler<FNDISkeletalMesh_InstanceData> InstData(Context);
 	FNDIInputParam<int32> BoneParam(Context);
@@ -537,7 +537,7 @@ void UNiagaraDataInterfaceSkeletalMesh::GetUnfilteredBoneAt(FVectorVMContext& Co
 	{
 		if (InstData->NumFilteredBones == 0)
 		{
-			for (int32 i = 0; i < Context.NumInstances; ++i)
+			for (int32 i = 0; i < Context.GetNumInstances(); ++i)
 			{
 				const int32 BoneIndex = FMath::Clamp(BoneParam.GetAndAdvance(), 0, Max);
 				OutBone.SetAndAdvance(BoneIndex);
@@ -545,7 +545,7 @@ void UNiagaraDataInterfaceSkeletalMesh::GetUnfilteredBoneAt(FVectorVMContext& Co
 		}
 		else
 		{
-			for (int32 i = 0; i < Context.NumInstances; ++i)
+			for (int32 i = 0; i < Context.GetNumInstances(); ++i)
 			{
 				const int32 BoneIndex = FMath::Clamp(BoneParam.GetAndAdvance(), 0, Max);
 				OutBone.SetAndAdvance(InstData->FilteredAndUnfilteredBones[BoneIndex + InstData->NumFilteredBones]);
@@ -554,13 +554,13 @@ void UNiagaraDataInterfaceSkeletalMesh::GetUnfilteredBoneAt(FVectorVMContext& Co
 	}
 	else
 	{
-		for (int32 i = 0; i < Context.NumInstances; ++i)
+		for (int32 i = 0; i < Context.GetNumInstances(); ++i)
 		{
 			OutBone.SetAndAdvance(-1);
 		}
 	}
 }
-void UNiagaraDataInterfaceSkeletalMesh::RandomUnfilteredBone(FVectorVMContext& Context)
+void UNiagaraDataInterfaceSkeletalMesh::RandomUnfilteredBone(FVectorVMExternalFunctionContext& Context)
 {
 	VectorVM::FUserPtrHandler<FNDISkeletalMesh_InstanceData> InstData(Context);
 	FNDIRandomHelper RandHelper(Context);
@@ -576,7 +576,7 @@ void UNiagaraDataInterfaceSkeletalMesh::RandomUnfilteredBone(FVectorVMContext& C
 			const int32 NumBones = InstData->NumUnfilteredBones - (ExcludedBoneIndex >= 0 ? 2 : 1);
 			if (NumBones >= 0)
 			{
-				for (int32 i = 0; i < Context.NumInstances; ++i)
+				for (int32 i = 0; i < Context.GetNumInstances(); ++i)
 				{
 					RandHelper.GetAndAdvance();
 					const int32 BoneIndex = RandHelper.RandRange(i, 0, NumBones);
@@ -585,7 +585,7 @@ void UNiagaraDataInterfaceSkeletalMesh::RandomUnfilteredBone(FVectorVMContext& C
 			}
 			else
 			{
-				for (int32 i = 0; i < Context.NumInstances; ++i)
+				for (int32 i = 0; i < Context.GetNumInstances(); ++i)
 				{
 					OutBone.SetAndAdvance(-1);
 				}
@@ -593,7 +593,7 @@ void UNiagaraDataInterfaceSkeletalMesh::RandomUnfilteredBone(FVectorVMContext& C
 		}
 		else
 		{
-			for (int32 i = 0; i < Context.NumInstances; ++i)
+			for (int32 i = 0; i < Context.GetNumInstances(); ++i)
 			{
 				RandHelper.GetAndAdvance();
 				const int32 BoneIndex = RandHelper.RandRange(i, 0, UnfilteredMax);
@@ -603,14 +603,14 @@ void UNiagaraDataInterfaceSkeletalMesh::RandomUnfilteredBone(FVectorVMContext& C
 	}
 	else
 	{
-		for (int32 i = 0; i < Context.NumInstances; ++i)
+		for (int32 i = 0; i < Context.GetNumInstances(); ++i)
 		{
 			OutBone.SetAndAdvance(-1);
 		}
 	}
 }
 
-void UNiagaraDataInterfaceSkeletalMesh::IsValidBone(FVectorVMContext& Context)
+void UNiagaraDataInterfaceSkeletalMesh::IsValidBone(FVectorVMExternalFunctionContext& Context)
 {
 	VectorVM::FUserPtrHandler<FNDISkeletalMesh_InstanceData> InstData(Context);
 	FNDIInputParam<int32> BoneParam(Context);
@@ -626,7 +626,7 @@ void UNiagaraDataInterfaceSkeletalMesh::IsValidBone(FVectorVMContext& Context)
 	{
 		const FReferenceSkeleton& RefSkeleton = MeshAccessor.Mesh->GetRefSkeleton();
 		int32 NumBones = RefSkeleton.GetNum();
-		for (int32 i = 0; i < Context.NumInstances; ++i)
+		for (int32 i = 0; i < Context.GetNumInstances(); ++i)
 		{
 			int32 RequestedIndex = BoneParam.GetAndAdvance();
 			OutValid.SetAndAdvance(RequestedIndex >= 0 && RequestedIndex < NumBones);
@@ -634,14 +634,14 @@ void UNiagaraDataInterfaceSkeletalMesh::IsValidBone(FVectorVMContext& Context)
 	}
 	else
 	{
-		for (int32 i = 0; i < Context.NumInstances; ++i)
+		for (int32 i = 0; i < Context.GetNumInstances(); ++i)
 		{
 			OutValid.SetAndAdvance(false);
 		}
 	}
 }
 
-void UNiagaraDataInterfaceSkeletalMesh::RandomBone(FVectorVMContext& Context)
+void UNiagaraDataInterfaceSkeletalMesh::RandomBone(FVectorVMExternalFunctionContext& Context)
 {
 	VectorVM::FUserPtrHandler<FNDISkeletalMesh_InstanceData> InstData(Context);
 	FNDIRandomHelper RandHelper(Context);
@@ -660,7 +660,7 @@ void UNiagaraDataInterfaceSkeletalMesh::RandomBone(FVectorVMContext& Context)
 
 	if (NumBones >= 0)
 	{
-		for (int32 i = 0; i < Context.NumInstances; ++i)
+		for (int32 i = 0; i < Context.GetNumInstances(); ++i)
 		{
 			RandHelper.GetAndAdvance();
 			const int32 BoneIndex = RandHelper.RandRange(i, 0, NumBones);
@@ -669,14 +669,14 @@ void UNiagaraDataInterfaceSkeletalMesh::RandomBone(FVectorVMContext& Context)
 	}
 	else
 	{
-		for (int32 i = 0; i < Context.NumInstances; ++i)
+		for (int32 i = 0; i < Context.GetNumInstances(); ++i)
 		{
 			OutBone.SetAndAdvance(-1);
 		}
 	}
 }
 
-void UNiagaraDataInterfaceSkeletalMesh::GetBoneCount(FVectorVMContext& Context)
+void UNiagaraDataInterfaceSkeletalMesh::GetBoneCount(FVectorVMExternalFunctionContext& Context)
 {
 	SCOPE_CYCLE_COUNTER(STAT_NiagaraSkel_Bone_Sample);
 
@@ -693,7 +693,7 @@ void UNiagaraDataInterfaceSkeletalMesh::GetBoneCount(FVectorVMContext& Context)
 		NumBones = RefSkeleton.GetNum();
 	}
 
-	for (int32 i = 0; i < Context.NumInstances; ++i)
+	for (int32 i = 0; i < Context.GetNumInstances(); ++i)
 	{
 		OutCount.SetAndAdvance(NumBones);
 	}
@@ -701,7 +701,7 @@ void UNiagaraDataInterfaceSkeletalMesh::GetBoneCount(FVectorVMContext& Context)
 
 struct FBoneSocketSkinnedDataOutputHandler
 {
-	FBoneSocketSkinnedDataOutputHandler(FVectorVMContext& Context)
+	FBoneSocketSkinnedDataOutputHandler(FVectorVMExternalFunctionContext& Context)
 		: Position(Context)
 		, Rotation(Context)
 		, Velocity(Context)
@@ -723,7 +723,7 @@ struct FBoneSocketSkinnedDataOutputHandler
 };
 
 template<typename SkinningHandlerType, typename TransformHandlerType, typename bInterpolated>
-void UNiagaraDataInterfaceSkeletalMesh::GetSkinnedBoneData(FVectorVMContext& Context)
+void UNiagaraDataInterfaceSkeletalMesh::GetSkinnedBoneData(FVectorVMExternalFunctionContext& Context)
 {
 	VectorVM::FUserPtrHandler<FNDISkeletalMesh_InstanceData> InstData(Context);
 
@@ -759,7 +759,7 @@ void UNiagaraDataInterfaceSkeletalMesh::GetSkinnedBoneData(FVectorVMContext& Con
 		const TArray<FTransform>& FilteredSocketCurrTransforms = InstData->GetFilteredSocketsCurrBuffer();
 		const TArray<FTransform>& FilteredSocketPrevTransforms = InstData->GetFilteredSocketsPrevBuffer();
 
-		for (int32 i = 0; i < Context.NumInstances; ++i)
+		for (int32 i = 0; i < Context.GetNumInstances(); ++i)
 		{
 			const float Interp = bInterpolated::Value ? InterpParam.GetAndAdvance() : 1.0f;
 
@@ -864,7 +864,7 @@ void UNiagaraDataInterfaceSkeletalMesh::GetSkinnedBoneData(FVectorVMContext& Con
 	{
 		const float InvDt = 1.0f / InstData->DeltaSeconds;
 
-		for (int32 i = 0; i < Context.NumInstances; ++i)
+		for (int32 i = 0; i < Context.GetNumInstances(); ++i)
 		{
 			const float Interp = bInterpolated::Value ? InterpParam.GetAndAdvance() : 1.0f;
 
@@ -904,20 +904,20 @@ void UNiagaraDataInterfaceSkeletalMesh::GetSkinnedBoneData(FVectorVMContext& Con
 //////////////////////////////////////////////////////////////////////////
 // Sockets
 
-void UNiagaraDataInterfaceSkeletalMesh::GetFilteredSocketCount(FVectorVMContext& Context)
+void UNiagaraDataInterfaceSkeletalMesh::GetFilteredSocketCount(FVectorVMExternalFunctionContext& Context)
 {
 	VectorVM::FUserPtrHandler<FNDISkeletalMesh_InstanceData> InstData(Context);
 
 	FNDIOutputParam<int32> OutCount(Context);
 
 	const int32 Num = InstData->FilteredSocketInfo.Num();
-	for (int32 i = 0; i < Context.NumInstances; ++i)
+	for (int32 i = 0; i < Context.GetNumInstances(); ++i)
 	{
 		OutCount.SetAndAdvance(Num);
 	}
 }
 
-void UNiagaraDataInterfaceSkeletalMesh::GetFilteredSocketBoneAt(FVectorVMContext& Context)
+void UNiagaraDataInterfaceSkeletalMesh::GetFilteredSocketBoneAt(FVectorVMExternalFunctionContext& Context)
 {
 	VectorVM::FUserPtrHandler<FNDISkeletalMesh_InstanceData> InstData(Context);
 	FNDIInputParam<int32> SocketParam(Context);
@@ -928,7 +928,7 @@ void UNiagaraDataInterfaceSkeletalMesh::GetFilteredSocketBoneAt(FVectorVMContext
 
 	if (Max != INDEX_NONE)
 	{
-		for (int32 i = 0; i < Context.NumInstances; ++i)
+		for (int32 i = 0; i < Context.GetNumInstances(); ++i)
 		{
 			const int32 SocketIndex = FMath::Clamp(SocketParam.GetAndAdvance(), 0, Max);
 			OutSocketBone.SetAndAdvance(FilteredSocketBoneOffset + SocketIndex);
@@ -936,14 +936,14 @@ void UNiagaraDataInterfaceSkeletalMesh::GetFilteredSocketBoneAt(FVectorVMContext
 	}
 	else
 	{
-		for (int32 i = 0; i < Context.NumInstances; ++i)
+		for (int32 i = 0; i < Context.GetNumInstances(); ++i)
 		{
 			OutSocketBone.SetAndAdvance(-1);
 		}
 	}
 }
 
-void UNiagaraDataInterfaceSkeletalMesh::GetFilteredSocketTransform(FVectorVMContext& Context)
+void UNiagaraDataInterfaceSkeletalMesh::GetFilteredSocketTransform(FVectorVMExternalFunctionContext& Context)
 {
 	VectorVM::FUserPtrHandler<FNDISkeletalMesh_InstanceData> InstData(Context);
 	FNDIInputParam<int32> SocketParam(Context);
@@ -961,7 +961,7 @@ void UNiagaraDataInterfaceSkeletalMesh::GetFilteredSocketTransform(FVectorVMCont
 		const FMatrix InstanceTransform = InstData->Transform;
 		const FQuat InstanceRotation = bNeedsRotation ? InstanceTransform.GetMatrixWithoutScale().ToQuat() : FQuat::Identity;
 
-		for (int32 i = 0; i < Context.NumInstances; ++i)
+		for (int32 i = 0; i < Context.GetNumInstances(); ++i)
 		{
 			const int32 SocketIndex = FMath::Clamp(SocketParam.GetAndAdvance(), 0, SocketMax);
 			FVector3f SocketTranslation = CurrentFilteredSockets[SocketIndex].GetTranslation();
@@ -983,7 +983,7 @@ void UNiagaraDataInterfaceSkeletalMesh::GetFilteredSocketTransform(FVectorVMCont
 	}
 	else
 	{
-		for (int32 i=0; i < Context.NumInstances; ++i)
+		for (int32 i=0; i < Context.GetNumInstances(); ++i)
 		{
 			OutSocketTranslate.SetAndAdvance(FVector3f::ZeroVector);
 			OutSocketRotation.SetAndAdvance(FQuat::Identity);
@@ -992,7 +992,7 @@ void UNiagaraDataInterfaceSkeletalMesh::GetFilteredSocketTransform(FVectorVMCont
 	}
 }
 
-void UNiagaraDataInterfaceSkeletalMesh::RandomFilteredSocket(FVectorVMContext& Context)
+void UNiagaraDataInterfaceSkeletalMesh::RandomFilteredSocket(FVectorVMExternalFunctionContext& Context)
 {
 	VectorVM::FUserPtrHandler<FNDISkeletalMesh_InstanceData> InstData(Context);
 	FNDIRandomHelper RandHelper(Context);
@@ -1003,7 +1003,7 @@ void UNiagaraDataInterfaceSkeletalMesh::RandomFilteredSocket(FVectorVMContext& C
 	int32 Max = FilteredSockets.Num() - 1;
 	if (Max != INDEX_NONE)
 	{
-		for (int32 i = 0; i < Context.NumInstances; ++i)
+		for (int32 i = 0; i < Context.GetNumInstances(); ++i)
 		{
 			RandHelper.GetAndAdvance();
 			const int32 SocketIndex = RandHelper.RandRange(i, 0, Max);
@@ -1012,14 +1012,14 @@ void UNiagaraDataInterfaceSkeletalMesh::RandomFilteredSocket(FVectorVMContext& C
 	}
 	else
 	{
-		for (int32 i = 0; i < Context.NumInstances; ++i)
+		for (int32 i = 0; i < Context.GetNumInstances(); ++i)
 		{
 			OutSocketBone.SetAndAdvance(-1);
 		}
 	}
 }
 
-void UNiagaraDataInterfaceSkeletalMesh::RandomFilteredSocketOrBone(FVectorVMContext& Context)
+void UNiagaraDataInterfaceSkeletalMesh::RandomFilteredSocketOrBone(FVectorVMExternalFunctionContext& Context)
 {
 	VectorVM::FUserPtrHandler<FNDISkeletalMesh_InstanceData> InstData(Context);
 	FNDIRandomHelper RandHelper(Context);
@@ -1031,7 +1031,7 @@ void UNiagaraDataInterfaceSkeletalMesh::RandomFilteredSocketOrBone(FVectorVMCont
 	{
 		const int32 NumFilteredBones = InstData->NumFilteredBones;
 		const int32 FilteredSocketBoneOffset = InstData->FilteredSocketBoneOffset;
-		for (int32 i = 0; i < Context.NumInstances; ++i)
+		for (int32 i = 0; i < Context.GetNumInstances(); ++i)
 		{
 			RandHelper.GetAndAdvance();
 			const int32 FilteredIndex = RandHelper.RandRange(i, 0, Max);
@@ -1047,26 +1047,26 @@ void UNiagaraDataInterfaceSkeletalMesh::RandomFilteredSocketOrBone(FVectorVMCont
 	}
 	else
 	{
-		for (int32 i = 0; i < Context.NumInstances; ++i)
+		for (int32 i = 0; i < Context.GetNumInstances(); ++i)
 		{
 			OutBoneIndex.SetAndAdvance(-1);
 		}
 	}
 }
 
-void UNiagaraDataInterfaceSkeletalMesh::GetFilteredSocketOrBoneCount(FVectorVMContext& Context)
+void UNiagaraDataInterfaceSkeletalMesh::GetFilteredSocketOrBoneCount(FVectorVMExternalFunctionContext& Context)
 {
 	VectorVM::FUserPtrHandler<FNDISkeletalMesh_InstanceData> InstData(Context);
 
 	FNDIOutputParam<int32> OutCount(Context);
 	const int32 Count = FilteredSockets.Num() + InstData->NumFilteredBones;
-	for (int32 i=0; i < Context.NumInstances; ++i)
+	for (int32 i=0; i < Context.GetNumInstances(); ++i)
 	{
 		OutCount.SetAndAdvance(Count);
 	}
 }
 
-void UNiagaraDataInterfaceSkeletalMesh::GetFilteredSocketOrBoneBoneAt(FVectorVMContext& Context)
+void UNiagaraDataInterfaceSkeletalMesh::GetFilteredSocketOrBoneBoneAt(FVectorVMExternalFunctionContext& Context)
 {
 	VectorVM::FUserPtrHandler<FNDISkeletalMesh_InstanceData> InstData(Context);
 	FNDIInputParam<int32> IndexParam(Context);
@@ -1077,7 +1077,7 @@ void UNiagaraDataInterfaceSkeletalMesh::GetFilteredSocketOrBoneBoneAt(FVectorVMC
 	{
 		const int32 NumFilteredBones = InstData->NumFilteredBones;
 		const int32 FilteredSocketBoneOffset = InstData->FilteredSocketBoneOffset;
-		for (int32 i = 0; i < Context.NumInstances; ++i)
+		for (int32 i = 0; i < Context.GetNumInstances(); ++i)
 		{
 			const int32 FilteredIndex = IndexParam.GetAndAdvance();
 			if (FilteredIndex < NumFilteredBones)
@@ -1092,7 +1092,7 @@ void UNiagaraDataInterfaceSkeletalMesh::GetFilteredSocketOrBoneBoneAt(FVectorVMC
 	}
 	else
 	{
-		for (int32 i = 0; i < Context.NumInstances; ++i)
+		for (int32 i = 0; i < Context.GetNumInstances(); ++i)
 		{
 			OutBoneIndex.SetAndAdvance(-1);
 		}
diff --git a/Engine/Plugins/FX/Niagara/Source/Niagara/Private/NDISkeletalMesh_TriangleSampling.cpp b/Engine/Plugins/FX/Niagara/Source/Niagara/Private/NDISkeletalMesh_TriangleSampling.cpp
--- a/Engine/Plugins/FX/Niagara/Source/Niagara/Private/NDISkeletalMesh_TriangleSampling.cpp
+++ b/Engine/Plugins/FX/Niagara/Source/Niagara/Private/NDISkeletalMesh_TriangleSampling.cpp
@@ -385,12 +385,12 @@ void UNiagaraDataInterfaceSkeletalMesh::BindTriangleSamplingFunction(const FVMEx
 	else if (BindingInfo.Name == FSkeletalMeshInterfaceHelper::RandomTriangleName)
 	{
 		check(BindingInfo.GetNumInputs() == 4 && BindingInfo.GetNumOutputs() == 4);
-		OutFunc = FVMExternalFunction::CreateLambda([this](FVectorVMContext& Context) { this->RandomTriangle(Context); });
+		OutFunc = FVMExternalFunction::CreateLambda([this](FVectorVMExternalFunctionContext& Context) { this->RandomTriangle(Context); });
 	}
 	else if (BindingInfo.Name == FSkeletalMeshInterfaceHelper::GetTriangleCountName)
 	{
 		check(BindingInfo.GetNumInputs() == 1 && BindingInfo.GetNumOutputs() == 1);
-		OutFunc = FVMExternalFunction::CreateLambda([this](FVectorVMContext& Context) { this->GetTriangleCount(Context); });
+		OutFunc = FVMExternalFunction::CreateLambda([this](FVectorVMExternalFunctionContext& Context) { this->GetTriangleCount(Context); });
 	}
 	else if (BindingInfo.Name == FSkeletalMeshInterfaceHelper::RandomFilteredTriangleName)
 	{
@@ -527,7 +527,7 @@ FORCEINLINE int32 UNiagaraDataInterfaceSkeletalMesh::RandomTriIndex<TNDISkelMesh
 }
 
 template<typename FilterMode, typename AreaWeightingMode>
-void UNiagaraDataInterfaceSkeletalMesh::RandomTriCoord(FVectorVMContext& Context)
+void UNiagaraDataInterfaceSkeletalMesh::RandomTriCoord(FVectorVMExternalFunctionContext& Context)
 {
 	SCOPE_CYCLE_COUNTER(STAT_NiagaraSkel_Sample);
 
@@ -545,16 +545,16 @@ void UNiagaraDataInterfaceSkeletalMesh::RandomTriCoord(FVectorVMContext& Context
 
 	if (MeshAccessor.IsLODAccessible())
 	{
-		for (int32 i = 0; i < Context.NumInstances; ++i)
+		for (int32 i = 0; i < Context.GetNumInstances(); ++i)
 		{
 			RandHelper.GetAndAdvance();//We grab the rand info to a local value first so it can be used for multiple rand calls from the helper.
 			OutTri.SetAndAdvance(RandomTriIndex<FilterMode, AreaWeightingMode>(RandHelper, MeshAccessor, InstData, i));
-			OutBary.SetAndAdvance(RandomBarycentricCoord(Context.RandStream));
+			OutBary.SetAndAdvance(RandomBarycentricCoord(Context.GetRandStream()));
 		}
 	}
 	else
 	{
-		for (int32 i = 0; i < Context.NumInstances; ++i)
+		for (int32 i = 0; i < Context.GetNumInstances(); ++i)
 		{
 			OutTri.SetAndAdvance(-1);
 			OutBary.SetAndAdvance(FVector3f::ZeroVector);
@@ -563,7 +563,7 @@ void UNiagaraDataInterfaceSkeletalMesh::RandomTriCoord(FVectorVMContext& Context
 }
 
 template<typename FilterMode, typename AreaWeightingMode>
-void UNiagaraDataInterfaceSkeletalMesh::IsValidTriCoord(FVectorVMContext& Context)
+void UNiagaraDataInterfaceSkeletalMesh::IsValidTriCoord(FVectorVMExternalFunctionContext& Context)
 {
 	SCOPE_CYCLE_COUNTER(STAT_NiagaraSkel_Sample);
 
@@ -583,7 +583,7 @@ void UNiagaraDataInterfaceSkeletalMesh::IsValidTriCoord(FVectorVMContext& Contex
 
 	if (MeshAccessor.IsLODAccessible())
 	{
-		for (int32 i = 0; i < Context.NumInstances; ++i)
+		for (int32 i = 0; i < Context.GetNumInstances(); ++i)
 		{
 			const int32 RequestedIndex = (TriParam.GetAndAdvance() * 3) + 2; // Get the last triangle index of the set
 			OutValid.SetAndAdvance(MeshAccessor.IndexBuffer != nullptr && MeshAccessor.IndexBuffer->Num() > RequestedIndex);
@@ -591,7 +591,7 @@ void UNiagaraDataInterfaceSkeletalMesh::IsValidTriCoord(FVectorVMContext& Contex
 	}
 	else
 	{
-		for (int32 i = 0; i < Context.NumInstances; ++i)
+		for (int32 i = 0; i < Context.GetNumInstances(); ++i)
 		{
 			OutValid.SetAndAdvance(false);
 		}
@@ -600,7 +600,7 @@ void UNiagaraDataInterfaceSkeletalMesh::IsValidTriCoord(FVectorVMContext& Contex
 
 //////////////////////////////////////////////////////////////////////////
 
-void UNiagaraDataInterfaceSkeletalMesh::RandomTriangle(FVectorVMContext& Context)
+void UNiagaraDataInterfaceSkeletalMesh::RandomTriangle(FVectorVMExternalFunctionContext& Context)
 {
 	SCOPE_CYCLE_COUNTER(STAT_NiagaraSkel_Sample);
 
@@ -614,7 +614,7 @@ void UNiagaraDataInterfaceSkeletalMesh::RandomTriangle(FVectorVMContext& Context
 
 	if (!MeshAccessor.IsLODAccessible())
 	{
-		for (int32 i = 0; i < Context.NumInstances; ++i)
+		for (int32 i = 0; i < Context.GetNumInstances(); ++i)
 		{
 			OutTri.SetAndAdvance(-1);
 			OutBary.SetAndAdvance(FVector3f::ZeroVector);
@@ -634,7 +634,7 @@ void UNiagaraDataInterfaceSkeletalMesh::RandomTriangle(FVectorVMContext& Context
 		const FSkeletalMeshSamplingLODBuiltData& WholeMeshBuiltData = SamplingInfo.GetWholeMeshLODBuiltData(InstData->GetLODIndex());
 		if (WholeMeshBuiltData.AreaWeightedTriangleSampler.GetNumEntries() > 0)
 		{
-			for (int32 i = 0; i < Context.NumInstances; ++i)
+			for (int32 i = 0; i < Context.GetNumInstances(); ++i)
 			{
 				RandHelper.GetAndAdvance();
 				OutTri.SetAndAdvance(WholeMeshBuiltData.AreaWeightedTriangleSampler.GetEntryIndex(RandHelper.Rand(i), RandHelper.Rand(i)));
@@ -647,7 +647,7 @@ void UNiagaraDataInterfaceSkeletalMesh::RandomTriangle(FVectorVMContext& Context
 	const int32 MaxTriangle = (MeshAccessor.IndexBuffer->Num() / 3) - 1;
 	if (MaxTriangle >= 0)
 	{
-		for (int32 i = 0; i < Context.NumInstances; ++i)
+		for (int32 i = 0; i < Context.GetNumInstances(); ++i)
 		{
 			RandHelper.GetAndAdvance();
 			OutTri.SetAndAdvance(RandHelper.RandRange(i, 0, MaxTriangle));
@@ -656,7 +656,7 @@ void UNiagaraDataInterfaceSkeletalMesh::RandomTriangle(FVectorVMContext& Context
 	}
 	else
 	{
-		for (int32 i = 0; i < Context.NumInstances; ++i)
+		for (int32 i = 0; i < Context.GetNumInstances(); ++i)
 		{
 			OutTri.SetAndAdvance(-1);
 			OutBary.SetAndAdvance(FVector3f::ZeroVector);
@@ -664,7 +664,7 @@ void UNiagaraDataInterfaceSkeletalMesh::RandomTriangle(FVectorVMContext& Context
 	}
 }
 
-void UNiagaraDataInterfaceSkeletalMesh::GetTriangleCount(FVectorVMContext& Context)
+void UNiagaraDataInterfaceSkeletalMesh::GetTriangleCount(FVectorVMExternalFunctionContext& Context)
 {
 	SCOPE_CYCLE_COUNTER(STAT_NiagaraSkel_Sample);
 
@@ -675,7 +675,7 @@ void UNiagaraDataInterfaceSkeletalMesh::GetTriangleCount(FVectorVMContext& Conte
 	MeshAccessor.Init<TIntegralConstant<int32, 0>, TIntegralConstant<int32, 0>>(InstData);
 	
 	const int32 NumTriangles = MeshAccessor.IsLODAccessible() ? MeshAccessor.IndexBuffer->Num() / 3 : 0;
-	for (int32 i = 0; i < Context.NumInstances; ++i)
+	for (int32 i = 0; i < Context.GetNumInstances(); ++i)
 	{
 		OutCount.SetAndAdvance(NumTriangles);
 	}
@@ -761,7 +761,7 @@ FORCEINLINE int32 UNiagaraDataInterfaceSkeletalMesh::GetFilteredTriangleCount<TN
 }
 
 template<typename FilterMode, typename AreaWeightingMode>
-void UNiagaraDataInterfaceSkeletalMesh::GetFilteredTriangleCount(FVectorVMContext& Context)
+void UNiagaraDataInterfaceSkeletalMesh::GetFilteredTriangleCount(FVectorVMExternalFunctionContext& Context)
 {
 	SCOPE_CYCLE_COUNTER(STAT_NiagaraSkel_Sample);
 	VectorVM::FUserPtrHandler<FNDISkeletalMesh_InstanceData> InstData(Context);
@@ -773,7 +773,7 @@ void UNiagaraDataInterfaceSkeletalMesh::GetFilteredTriangleCount(FVectorVMContex
 	MeshAccessor.Init<FilterMode, AreaWeightingMode>(InstData);
 
 	int32 Count = MeshAccessor.IsLODAccessible() ? GetFilteredTriangleCount<FilterMode, AreaWeightingMode>(MeshAccessor, InstData) : 0;
-	for (int32 i = 0; i < Context.NumInstances; ++i)
+	for (int32 i = 0; i < Context.GetNumInstances(); ++i)
 	{
 		OutTri.SetAndAdvance(Count);
 	}
@@ -876,7 +876,7 @@ FORCEINLINE int32 UNiagaraDataInterfaceSkeletalMesh::GetFilteredTriangleAt<TNDIS
 }
 
 template<typename FilterMode, typename AreaWeightingMode>
-void UNiagaraDataInterfaceSkeletalMesh::GetFilteredTriangleAt(FVectorVMContext& Context)
+void UNiagaraDataInterfaceSkeletalMesh::GetFilteredTriangleAt(FVectorVMExternalFunctionContext& Context)
 {
 	SCOPE_CYCLE_COUNTER(STAT_NiagaraSkel_Sample);
 
@@ -893,7 +893,7 @@ void UNiagaraDataInterfaceSkeletalMesh::GetFilteredTriangleAt(FVectorVMContext&
 
 	if (Accessor.IsLODAccessible())
 	{
-		for (int32 i = 0; i < Context.NumInstances; ++i)
+		for (int32 i = 0; i < Context.GetNumInstances(); ++i)
 		{
 			int32 Tri = TriParam.GetAndAdvance();
 			int32 RealIdx = 0;
@@ -908,7 +908,7 @@ void UNiagaraDataInterfaceSkeletalMesh::GetFilteredTriangleAt(FVectorVMContext&
 	}
 	else
 	{
-		for (int32 i = 0; i < Context.NumInstances; ++i)
+		for (int32 i = 0; i < Context.GetNumInstances(); ++i)
 		{
 			OutTri.SetAndAdvance(-1);
 			OutBary.SetAndAdvance(InBary.GetAndAdvance());
@@ -916,7 +916,7 @@ void UNiagaraDataInterfaceSkeletalMesh::GetFilteredTriangleAt(FVectorVMContext&
 	}
 }
 
-void UNiagaraDataInterfaceSkeletalMesh::GetTriCoordColor(FVectorVMContext& Context)
+void UNiagaraDataInterfaceSkeletalMesh::GetTriCoordColor(FVectorVMExternalFunctionContext& Context)
 {
 	SCOPE_CYCLE_COUNTER(STAT_NiagaraSkel_Sample);
 	VectorVM::FUserPtrHandler<FNDISkeletalMesh_InstanceData> InstData(Context);
@@ -936,7 +936,7 @@ void UNiagaraDataInterfaceSkeletalMesh::GetTriCoordColor(FVectorVMContext& Conte
 		const int32 TriMax = (IndexBuffer->Num() / 3) - 1;
 		if (TriMax >= 0)
 		{
-			for (int32 i = 0; i < Context.NumInstances; ++i)
+			for (int32 i = 0; i < Context.GetNumInstances(); ++i)
 			{
 				const int32 Tri = FMath::Clamp(TriParam.GetAndAdvance(), 0, TriMax) * 3;
 				const int32 Idx0 = IndexBuffer->Get(Tri);
@@ -952,14 +952,14 @@ void UNiagaraDataInterfaceSkeletalMesh::GetTriCoordColor(FVectorVMContext& Conte
 	}
 
 	// Error fall-through
-	for (int32 i = 0; i < Context.NumInstances; ++i)
+	for (int32 i = 0; i < Context.GetNumInstances(); ++i)
 	{
 		OutColor.SetAndAdvance(FLinearColor::White);
 	}
 }
 
 // Where we determine we are sampling a skeletal mesh without tri color we bind to this fallback method 
-void UNiagaraDataInterfaceSkeletalMesh::GetTriCoordColorFallback(FVectorVMContext& Context)
+void UNiagaraDataInterfaceSkeletalMesh::GetTriCoordColorFallback(FVectorVMExternalFunctionContext& Context)
 {
 	VectorVM::FUserPtrHandler<FNDISkeletalMesh_InstanceData> InstData(Context);
 	FNDIInputParam<int32> TriParam(Context);
@@ -967,14 +967,14 @@ void UNiagaraDataInterfaceSkeletalMesh::GetTriCoordColorFallback(FVectorVMContex
 
 	FNDIOutputParam<FLinearColor> OutColor(Context);
 
-	for (int32 i = 0; i < Context.NumInstances; ++i)
+	for (int32 i = 0; i < Context.GetNumInstances(); ++i)
 	{
 		OutColor.SetAndAdvance(FLinearColor::White);
 	}
 }
 
 template<typename VertexAccessorType>
-void UNiagaraDataInterfaceSkeletalMesh::GetTriCoordUV(FVectorVMContext& Context)
+void UNiagaraDataInterfaceSkeletalMesh::GetTriCoordUV(FVectorVMExternalFunctionContext& Context)
 {
 	SCOPE_CYCLE_COUNTER(STAT_NiagaraSkel_Sample);
 	VectorVM::FUserPtrHandler<FNDISkeletalMesh_InstanceData> InstData(Context);
@@ -998,7 +998,7 @@ void UNiagaraDataInterfaceSkeletalMesh::GetTriCoordUV(FVectorVMContext& Context)
 		{
 			const int32 UVSetMax = LODData->StaticVertexBuffers.StaticMeshVertexBuffer.GetNumTexCoords() - 1;
 			const float InvDt = 1.0f / InstData->DeltaSeconds;
-			for (int32 i = 0; i < Context.NumInstances; ++i)
+			for (int32 i = 0; i < Context.GetNumInstances(); ++i)
 			{
 				const int32 Tri = FMath::Clamp(TriParam.GetAndAdvance(), 0, TriMax) * 3;
 				const int32 Idx0 = IndexBuffer->Get(Tri);
@@ -1019,7 +1019,7 @@ void UNiagaraDataInterfaceSkeletalMesh::GetTriCoordUV(FVectorVMContext& Context)
 	}
 
 	// Fall-though for failure conditions
-	for (int32 i=0; i < Context.NumInstances; ++i)
+	for (int32 i=0; i < Context.GetNumInstances(); ++i)
 	{
 		OutUV.SetAndAdvance(FVector2D::ZeroVector);
 	}
@@ -1027,7 +1027,7 @@ void UNiagaraDataInterfaceSkeletalMesh::GetTriCoordUV(FVectorVMContext& Context)
 
 // Stub specialization for no valid mesh data on the data interface
 template<>
-void UNiagaraDataInterfaceSkeletalMesh::GetTriCoordUV<FSkelMeshVertexAccessorNoop>(FVectorVMContext& Context)
+void UNiagaraDataInterfaceSkeletalMesh::GetTriCoordUV<FSkelMeshVertexAccessorNoop>(FVectorVMExternalFunctionContext& Context)
 {
 	VectorVM::FUserPtrHandler<FNDISkeletalMesh_InstanceData> InstData(Context);
 	FNDIInputParam<int32> TriParam(Context);
@@ -1036,14 +1036,14 @@ void UNiagaraDataInterfaceSkeletalMesh::GetTriCoordUV<FSkelMeshVertexAccessorNoo
 
 	FNDIOutputParam<FVector2D> OutUV(Context);
 
-	for (int32 i = 0; i < Context.NumInstances; ++i)
+	for (int32 i = 0; i < Context.GetNumInstances(); ++i)
 	{
 		OutUV.SetAndAdvance(FVector2D::ZeroVector);
 	}
 }
 
 template<typename VertexAccessorType>
-void UNiagaraDataInterfaceSkeletalMesh::GetTriangleCoordAtUV(FVectorVMContext& Context)
+void UNiagaraDataInterfaceSkeletalMesh::GetTriangleCoordAtUV(FVectorVMExternalFunctionContext& Context)
 {
 	SCOPE_CYCLE_COUNTER(STAT_NiagaraSkel_Sample);
 	VectorVM::FUserPtrHandler<FNDISkeletalMesh_InstanceData> InstData(Context);
@@ -1060,7 +1060,7 @@ void UNiagaraDataInterfaceSkeletalMesh::GetTriangleCoordAtUV(FVectorVMContext& C
 
 	if (InstData->UvMapping)
 	{
-		for (int32 i = 0; i < Context.NumInstances; ++i)
+		for (int32 i = 0; i < Context.GetNumInstances(); ++i)
 		{
 			const bool Enabled = InEnabled.GetAndAdvance();
 			const FVector2D SourceUv = InUV.GetAndAdvance();
@@ -1081,7 +1081,7 @@ void UNiagaraDataInterfaceSkeletalMesh::GetTriangleCoordAtUV(FVectorVMContext& C
 	}
 	else
 	{
-		for (int32 i = 0; i < Context.NumInstances; ++i)
+		for (int32 i = 0; i < Context.GetNumInstances(); ++i)
 		{
 			OutTriangleIndex.SetAndAdvance(INDEX_NONE);
 			OutBaryCoord.SetAndAdvance(FVector3f::ZeroVector);
@@ -1092,7 +1092,7 @@ void UNiagaraDataInterfaceSkeletalMesh::GetTriangleCoordAtUV(FVectorVMContext& C
 
 // Stub specialization for no valid mesh data on the data interface
 template<>
-void UNiagaraDataInterfaceSkeletalMesh::GetTriangleCoordAtUV<FSkelMeshVertexAccessorNoop>(FVectorVMContext& Context)
+void UNiagaraDataInterfaceSkeletalMesh::GetTriangleCoordAtUV<FSkelMeshVertexAccessorNoop>(FVectorVMExternalFunctionContext& Context)
 {
 	VectorVM::FUserPtrHandler<FNDISkeletalMesh_InstanceData> InstData(Context);
 	FNDIInputParam<bool> InEnabled(Context);
@@ -1103,7 +1103,7 @@ void UNiagaraDataInterfaceSkeletalMesh::GetTriangleCoordAtUV<FSkelMeshVertexAcce
 	FNDIOutputParam<FVector3f> OutBaryCoord(Context);
 	FNDIOutputParam<bool> OutIsValid(Context);
 
-	for (int32 i = 0; i < Context.NumInstances; ++i)
+	for (int32 i = 0; i < Context.GetNumInstances(); ++i)
 	{
 		OutTriangleIndex.SetAndAdvance(INDEX_NONE);
 		OutBaryCoord.SetAndAdvance(FVector3f::ZeroVector);
@@ -1113,7 +1113,7 @@ void UNiagaraDataInterfaceSkeletalMesh::GetTriangleCoordAtUV<FSkelMeshVertexAcce
 
 
 template<typename VertexAccessorType>
-void UNiagaraDataInterfaceSkeletalMesh::GetTriangleCoordInAabb(FVectorVMContext& Context)
+void UNiagaraDataInterfaceSkeletalMesh::GetTriangleCoordInAabb(FVectorVMExternalFunctionContext& Context)
 {
 	SCOPE_CYCLE_COUNTER(STAT_NiagaraSkel_Sample);
 	VectorVM::FUserPtrHandler<FNDISkeletalMesh_InstanceData> InstData(Context);
@@ -1130,7 +1130,7 @@ void UNiagaraDataInterfaceSkeletalMesh::GetTriangleCoordInAabb(FVectorVMContext&
 
 	if (InstData->UvMapping)
 	{
-		for (int32 i = 0; i < Context.NumInstances; ++i)
+		for (int32 i = 0; i < Context.GetNumInstances(); ++i)
 		{
 			const bool Enabled = InEnabled.GetAndAdvance();
 			const FVector2D MinExtent = InMinExtent.GetAndAdvance();
@@ -1150,7 +1150,7 @@ void UNiagaraDataInterfaceSkeletalMesh::GetTriangleCoordInAabb(FVectorVMContext&
 	}
 	else
 	{
-		for (int32 i = 0; i < Context.NumInstances; ++i)
+		for (int32 i = 0; i < Context.GetNumInstances(); ++i)
 		{
 			OutTriangleIndex.SetAndAdvance(INDEX_NONE);
 			OutBaryCoord.SetAndAdvance(FVector3f::ZeroVector);
@@ -1160,7 +1160,7 @@ void UNiagaraDataInterfaceSkeletalMesh::GetTriangleCoordInAabb(FVectorVMContext&
 }
 
 template<>
-void UNiagaraDataInterfaceSkeletalMesh::GetTriangleCoordInAabb<FSkelMeshVertexAccessorNoop>(FVectorVMContext& Context)
+void UNiagaraDataInterfaceSkeletalMesh::GetTriangleCoordInAabb<FSkelMeshVertexAccessorNoop>(FVectorVMExternalFunctionContext& Context)
 {
 	VectorVM::FUserPtrHandler<FNDISkeletalMesh_InstanceData> InstData(Context);
 	FNDIInputParam<bool> InEnabled(Context);
@@ -1171,7 +1171,7 @@ void UNiagaraDataInterfaceSkeletalMesh::GetTriangleCoordInAabb<FSkelMeshVertexAc
 	FNDIOutputParam<FVector3f> OutBaryCoord(Context);
 	FNDIOutputParam<bool> OutIsValid(Context);
 
-	for (int32 i = 0; i < Context.NumInstances; ++i)
+	for (int32 i = 0; i < Context.GetNumInstances(); ++i)
 	{
 		OutTriangleIndex.SetAndAdvance(INDEX_NONE);
 		OutBaryCoord.SetAndAdvance(FVector3f::ZeroVector);
@@ -1181,7 +1181,7 @@ void UNiagaraDataInterfaceSkeletalMesh::GetTriangleCoordInAabb<FSkelMeshVertexAc
 
 struct FGetTriCoordSkinnedDataOutputHandler
 {
-	FGetTriCoordSkinnedDataOutputHandler(FVectorVMContext& Context)
+	FGetTriCoordSkinnedDataOutputHandler(FVectorVMExternalFunctionContext& Context)
 		: Position(Context)
 		, Velocity(Context)
 		, Normal(Context)
@@ -1209,7 +1209,7 @@ struct FGetTriCoordSkinnedDataOutputHandler
 };
 
 template<typename SkinningHandlerType, typename TransformHandlerType, typename VertexAccessorType, typename bInterpolated>
-void UNiagaraDataInterfaceSkeletalMesh::GetTriCoordSkinnedData(FVectorVMContext& Context)
+void UNiagaraDataInterfaceSkeletalMesh::GetTriCoordSkinnedData(FVectorVMExternalFunctionContext& Context)
 {
 	SCOPE_CYCLE_COUNTER(STAT_NiagaraSkel_Sample);
 	VectorVM::FUserPtrHandler<FNDISkeletalMesh_InstanceData> InstData(Context);
@@ -1247,7 +1247,7 @@ void UNiagaraDataInterfaceSkeletalMesh::GetTriCoordSkinnedData(FVectorVMContext&
 	// If LOD isn't accessible push out failure data
 	if ( !Accessor.IsLODAccessible() )
 	{
-		for (int32 i = 0; i < Context.NumInstances; ++i)
+		for (int32 i = 0; i < Context.GetNumInstances(); ++i)
 		{
 			const float Interp = bInterpolated::Value ? InterpParam.GetAndAdvance() : 1.0f;
 			FVector3f PrevPosition = FVector3f::ZeroVector;
@@ -1282,7 +1282,7 @@ void UNiagaraDataInterfaceSkeletalMesh::GetTriCoordSkinnedData(FVectorVMContext&
 	FVector3f Prev;
 	FVector3f Velocity;
 
-	for (int32 i = 0; i < Context.NumInstances; ++i)
+	for (int32 i = 0; i < Context.GetNumInstances(); ++i)
 	{
 		FMeshTriCoordinate MeshTriCoord(TriParam.GetAndAdvance(), BaryParam.GetAndAdvance());
 
@@ -1387,7 +1387,7 @@ void UNiagaraDataInterfaceSkeletalMesh::GetTriCoordSkinnedData(FVectorVMContext&
 
 // Fallback sampling function for no valid mesh on the interface
 template<typename TransformHandlerType, typename bInterpolated>
-void UNiagaraDataInterfaceSkeletalMesh::GetTriCoordSkinnedDataFallback(FVectorVMContext& Context)
+void UNiagaraDataInterfaceSkeletalMesh::GetTriCoordSkinnedDataFallback(FVectorVMExternalFunctionContext& Context)
 {
 	SCOPE_CYCLE_COUNTER(STAT_NiagaraSkel_Sample);
 	TransformHandlerType TransformHandler;
@@ -1414,7 +1414,7 @@ void UNiagaraDataInterfaceSkeletalMesh::GetTriCoordSkinnedDataFallback(FVectorVM
 
 	const float InvDt = 1.0f / InstData->DeltaSeconds;
 
-	for (int32 i = 0; i < Context.NumInstances; ++i)
+	for (int32 i = 0; i < Context.GetNumInstances(); ++i)
 	{
 		float Interp = 1.0f;
 		if (bInterpolated::Value)
@@ -1465,7 +1465,7 @@ void UNiagaraDataInterfaceSkeletalMesh::GetTriCoordSkinnedDataFallback(FVectorVM
 }
 
 template<typename SkinningHandlerType>
-void UNiagaraDataInterfaceSkeletalMesh::GetTriCoordVertices(FVectorVMContext& Context)
+void UNiagaraDataInterfaceSkeletalMesh::GetTriCoordVertices(FVectorVMExternalFunctionContext& Context)
 {
 	SCOPE_CYCLE_COUNTER(STAT_NiagaraSkel_Sample);
 	VectorVM::FUserPtrHandler<FNDISkeletalMesh_InstanceData> InstData(Context);
@@ -1489,7 +1489,7 @@ void UNiagaraDataInterfaceSkeletalMesh::GetTriCoordVertices(FVectorVMContext& Co
 		const int32 TriMax = (Accessor.IndexBuffer->Num() / 3) - 1;
 		if (TriMax >= 0)
 		{
-			for (int32 i = 0; i < Context.NumInstances; ++i)
+			for (int32 i = 0; i < Context.GetNumInstances(); ++i)
 			{
 				const int32 Tri = FMath::Clamp(TriParam.GetAndAdvance(), 0, TriMax);
 				SkinningHandler.GetTriangleIndices(Accessor, Tri, Idx0, Idx1, Idx2);
@@ -1503,7 +1503,7 @@ void UNiagaraDataInterfaceSkeletalMesh::GetTriCoordVertices(FVectorVMContext& Co
 	}
 
 	// Bad or missing data
-	for (int32 i = 0; i < Context.NumInstances; ++i)
+	for (int32 i = 0; i < Context.GetNumInstances(); ++i)
 	{
 		OutV0.SetAndAdvance(0);
 		OutV1.SetAndAdvance(0);
diff --git a/Engine/Plugins/FX/Niagara/Source/Niagara/Private/NDISkeletalMesh_VertexSampling.cpp b/Engine/Plugins/FX/Niagara/Source/Niagara/Private/NDISkeletalMesh_VertexSampling.cpp
--- a/Engine/Plugins/FX/Niagara/Source/Niagara/Private/NDISkeletalMesh_VertexSampling.cpp
+++ b/Engine/Plugins/FX/Niagara/Source/Niagara/Private/NDISkeletalMesh_VertexSampling.cpp
@@ -205,17 +205,17 @@ void UNiagaraDataInterfaceSkeletalMesh::BindVertexSamplingFunction(const FVMExte
 	else if (BindingInfo.Name == FSkeletalMeshInterfaceHelper::IsValidVertexName)
 	{
 		check(BindingInfo.GetNumInputs() == 2 && BindingInfo.GetNumOutputs() == 1);
-		OutFunc = FVMExternalFunction::CreateLambda([this](FVectorVMContext& Context) { this->IsValidVertex(Context); });
+		OutFunc = FVMExternalFunction::CreateLambda([this](FVectorVMExternalFunctionContext& Context) { this->IsValidVertex(Context); });
 	}
 	else if (BindingInfo.Name == FSkeletalMeshInterfaceHelper::RandomVertexName)
 	{
 		check(BindingInfo.GetNumInputs() == 4 && BindingInfo.GetNumOutputs() == 1);
-		OutFunc = FVMExternalFunction::CreateLambda([this](FVectorVMContext& Context) { this->RandomVertex(Context); });
+		OutFunc = FVMExternalFunction::CreateLambda([this](FVectorVMExternalFunctionContext& Context) { this->RandomVertex(Context); });
 	}
 	else if (BindingInfo.Name == FSkeletalMeshInterfaceHelper::GetVertexCountName)
 	{
 		check(BindingInfo.GetNumInputs() == 1 && BindingInfo.GetNumOutputs() == 1);
-		OutFunc = FVMExternalFunction::CreateLambda([this](FVectorVMContext& Context) { this->GetVertexCount(Context); });
+		OutFunc = FVMExternalFunction::CreateLambda([this](FVectorVMExternalFunctionContext& Context) { this->GetVertexCount(Context); });
 	}
 	//////////////////////////////////////////////////////////////////////////
 	else if (BindingInfo.Name == FSkeletalMeshInterfaceHelper::IsValidFilteredVertexName)
@@ -242,7 +242,7 @@ void UNiagaraDataInterfaceSkeletalMesh::BindVertexSamplingFunction(const FVMExte
 
 //////////////////////////////////////////////////////////////////////////
 
-void UNiagaraDataInterfaceSkeletalMesh::IsValidVertex(FVectorVMContext& Context)
+void UNiagaraDataInterfaceSkeletalMesh::IsValidVertex(FVectorVMExternalFunctionContext& Context)
 {
 	SCOPE_CYCLE_COUNTER(STAT_NiagaraSkel_Vertex_Sample);
 
@@ -254,14 +254,14 @@ void UNiagaraDataInterfaceSkeletalMesh::IsValidVertex(FVectorVMContext& Context)
 	MeshAccessor.Init<TNDISkelMesh_FilterModeNone, TNDISkelMesh_AreaWeightingOff>(InstData);
 
 	const int32 MaxVertex = MeshAccessor.IsLODAccessible() ? MeshAccessor.LODData->GetNumVertices() : 0;
-	for (int32 i = 0; i < Context.NumInstances; ++i)
+	for (int32 i = 0; i < Context.GetNumInstances(); ++i)
 	{
 		const int32 VertexIndex = VertexParam.GetAndAdvance();
 		OutValid.SetAndAdvance(VertexIndex < MaxVertex);
 	}
 }
 
-void UNiagaraDataInterfaceSkeletalMesh::RandomVertex(FVectorVMContext& Context)
+void UNiagaraDataInterfaceSkeletalMesh::RandomVertex(FVectorVMExternalFunctionContext& Context)
 {
 	SCOPE_CYCLE_COUNTER(STAT_NiagaraSkel_Vertex_Sample);
 
@@ -275,7 +275,7 @@ void UNiagaraDataInterfaceSkeletalMesh::RandomVertex(FVectorVMContext& Context)
 	const int32 MaxVertex = MeshAccessor.IsLODAccessible() ? MeshAccessor.LODData->GetNumVertices() - 1 : -1;
 	if (MaxVertex >= 0)
 	{
-		for (int32 i = 0; i < Context.NumInstances; ++i)
+		for (int32 i = 0; i < Context.GetNumInstances(); ++i)
 		{
 			RandHelper.GetAndAdvance();
 			OutVertex.SetAndAdvance(RandHelper.RandRange(i, 0, MaxVertex));
@@ -283,14 +283,14 @@ void UNiagaraDataInterfaceSkeletalMesh::RandomVertex(FVectorVMContext& Context)
 	}
 	else
 	{
-		for (int32 i = 0; i < Context.NumInstances; ++i)
+		for (int32 i = 0; i < Context.GetNumInstances(); ++i)
 		{
 			OutVertex.SetAndAdvance(-1);
 		}
 	}
 }
 
-void UNiagaraDataInterfaceSkeletalMesh::GetVertexCount(FVectorVMContext& Context)
+void UNiagaraDataInterfaceSkeletalMesh::GetVertexCount(FVectorVMExternalFunctionContext& Context)
 {
 	SCOPE_CYCLE_COUNTER(STAT_NiagaraSkel_Vertex_Sample);
 
@@ -301,7 +301,7 @@ void UNiagaraDataInterfaceSkeletalMesh::GetVertexCount(FVectorVMContext& Context
 	MeshAccessor.Init<TNDISkelMesh_FilterModeNone, TNDISkelMesh_AreaWeightingOff>(InstData);
 
 	const int32 MaxVertex = MeshAccessor.IsLODAccessible() ? MeshAccessor.LODData->GetNumVertices() : 0;
-	for (int32 i = 0; i < Context.NumInstances; ++i)
+	for (int32 i = 0; i < Context.GetNumInstances(); ++i)
 	{
 		OutVertexCount.SetAndAdvance(MaxVertex);
 	}
@@ -346,7 +346,7 @@ FORCEINLINE int32 UNiagaraDataInterfaceSkeletalMesh::RandomFilteredVertIndex<TIn
 }
 
 template<typename FilterMode>
-void UNiagaraDataInterfaceSkeletalMesh::RandomFilteredVertex(FVectorVMContext& Context)
+void UNiagaraDataInterfaceSkeletalMesh::RandomFilteredVertex(FVectorVMExternalFunctionContext& Context)
 {
 	SCOPE_CYCLE_COUNTER(STAT_NiagaraSkel_Vertex_Sample);
 
@@ -361,7 +361,7 @@ void UNiagaraDataInterfaceSkeletalMesh::RandomFilteredVertex(FVectorVMContext& C
 
 	if (MeshAccessor.IsLODAccessible())
 	{
-		for (int32 i = 0; i < Context.NumInstances; ++i)
+		for (int32 i = 0; i < Context.GetNumInstances(); ++i)
 		{
 			RandHelper.GetAndAdvance();
 			OutVert.SetAndAdvance(RandomFilteredVertIndex<FilterMode>(RandHelper, i, MeshAccessor, InstData));
@@ -369,7 +369,7 @@ void UNiagaraDataInterfaceSkeletalMesh::RandomFilteredVertex(FVectorVMContext& C
 	}
 	else
 	{
-		for (int32 i = 0; i < Context.NumInstances; ++i)
+		for (int32 i = 0; i < Context.GetNumInstances(); ++i)
 		{
 			OutVert.SetAndAdvance(-1);
 		}
@@ -377,7 +377,7 @@ void UNiagaraDataInterfaceSkeletalMesh::RandomFilteredVertex(FVectorVMContext& C
 }
 
 template<typename FilterMode>
-void UNiagaraDataInterfaceSkeletalMesh::IsValidFilteredVertex(FVectorVMContext& Context)
+void UNiagaraDataInterfaceSkeletalMesh::IsValidFilteredVertex(FVectorVMExternalFunctionContext& Context)
 {
 	SCOPE_CYCLE_COUNTER(STAT_NiagaraSkel_Vertex_Sample);
 
@@ -393,7 +393,7 @@ void UNiagaraDataInterfaceSkeletalMesh::IsValidFilteredVertex(FVectorVMContext&
 
 	const int32 MaxVertex = MeshAccessor.IsLODAccessible() ? MeshAccessor.LODData->GetNumVertices() : 0;
 
-	for (int32 i = 0; i < Context.NumInstances; ++i)
+	for (int32 i = 0; i < Context.GetNumInstances(); ++i)
 	{
 		const int32 RequestedIndex = VertexParam.GetAndAdvance();
 		OutValid.SetAndAdvance(RequestedIndex < MaxVertex);
@@ -442,7 +442,7 @@ FORCEINLINE_DEBUGGABLE int32 UNiagaraDataInterfaceSkeletalMesh::GetFilteredVerte
 }
 
 template<typename FilterMode>
-void UNiagaraDataInterfaceSkeletalMesh::GetFilteredVertexCount(FVectorVMContext& Context)
+void UNiagaraDataInterfaceSkeletalMesh::GetFilteredVertexCount(FVectorVMExternalFunctionContext& Context)
 {
 	SCOPE_CYCLE_COUNTER(STAT_NiagaraSkel_Vertex_Sample);
 	VectorVM::FUserPtrHandler<FNDISkeletalMesh_InstanceData> InstData(Context);
@@ -454,7 +454,7 @@ void UNiagaraDataInterfaceSkeletalMesh::GetFilteredVertexCount(FVectorVMContext&
 	MeshAccessor.Init<FilterMode, TNDISkelMesh_AreaWeightingOff>(InstData);
 
 	const int32 Count = MeshAccessor.IsLODAccessible() ? GetFilteredVertexCount<FilterMode>(MeshAccessor, InstData) : 0;
-	for (int32 i = 0; i < Context.NumInstances; ++i)
+	for (int32 i = 0; i < Context.GetNumInstances(); ++i)
 	{
 		OutVert.SetAndAdvance(Count);
 	}
@@ -508,7 +508,7 @@ FORCEINLINE_DEBUGGABLE int32 UNiagaraDataInterfaceSkeletalMesh::GetFilteredVerte
 }
 
 template<typename FilterMode>
-void UNiagaraDataInterfaceSkeletalMesh::GetFilteredVertexAt(FVectorVMContext& Context)
+void UNiagaraDataInterfaceSkeletalMesh::GetFilteredVertexAt(FVectorVMExternalFunctionContext& Context)
 {
 	SCOPE_CYCLE_COUNTER(STAT_NiagaraSkel_Vertex_Sample);
 
@@ -523,7 +523,7 @@ void UNiagaraDataInterfaceSkeletalMesh::GetFilteredVertexAt(FVectorVMContext& Co
 
 	if (Accessor.IsLODAccessible())
 	{
-		for (int32 i = 0; i < Context.NumInstances; ++i)
+		for (int32 i = 0; i < Context.GetNumInstances(); ++i)
 		{
 			const int32 FilteredVert = FilteredVertexParam.GetAndAdvance();
 			const int32 RealIdx = GetFilteredVertexAt<FilterMode>(Accessor, InstData, FilteredVert);
@@ -532,7 +532,7 @@ void UNiagaraDataInterfaceSkeletalMesh::GetFilteredVertexAt(FVectorVMContext& Co
 	}
 	else
 	{
-		for (int32 i = 0; i < Context.NumInstances; ++i)
+		for (int32 i = 0; i < Context.GetNumInstances(); ++i)
 		{
 			OutVert.SetAndAdvance(-1);
 		}
@@ -541,7 +541,7 @@ void UNiagaraDataInterfaceSkeletalMesh::GetFilteredVertexAt(FVectorVMContext& Co
 
 //////////////////////////////////////////////////////////////////////////
 
-void UNiagaraDataInterfaceSkeletalMesh::GetVertexColor(FVectorVMContext& Context)
+void UNiagaraDataInterfaceSkeletalMesh::GetVertexColor(FVectorVMExternalFunctionContext& Context)
 {
 	SCOPE_CYCLE_COUNTER(STAT_NiagaraSkel_Vertex_Sample);
 	VectorVM::FUserPtrHandler<FNDISkeletalMesh_InstanceData> InstData(Context);
@@ -560,7 +560,7 @@ void UNiagaraDataInterfaceSkeletalMesh::GetVertexColor(FVectorVMContext& Context
 		const int32 VertMax = LODData->GetNumVertices() - 1;
 		if ( VertMax >= 0 )
 		{
-			for (int32 i = 0; i < Context.NumInstances; ++i)
+			for (int32 i = 0; i < Context.GetNumInstances(); ++i)
 			{
 				int32 Vertex = VertParam.GetAndAdvance();
 				Vertex = FMath::Clamp(VertMax, 0, Vertex);
@@ -573,27 +573,27 @@ void UNiagaraDataInterfaceSkeletalMesh::GetVertexColor(FVectorVMContext& Context
 	}
 
 	// Fall though for bad data
-	for (int32 i = 0; i < Context.NumInstances; ++i)
+	for (int32 i = 0; i < Context.GetNumInstances(); ++i)
 	{
 		OutColor.SetAndAdvance(FLinearColor::White);
 	}
 }
 
-void UNiagaraDataInterfaceSkeletalMesh::GetVertexColorFallback(FVectorVMContext& Context)
+void UNiagaraDataInterfaceSkeletalMesh::GetVertexColorFallback(FVectorVMExternalFunctionContext& Context)
 {
 	VectorVM::FUserPtrHandler<FNDISkeletalMesh_InstanceData> InstData(Context);
 	FNDIInputParam<int32> VertParam(Context);
 
 	FNDIOutputParam<FLinearColor> OutColor(Context);
 
-	for (int32 i = 0; i < Context.NumInstances; ++i)
+	for (int32 i = 0; i < Context.GetNumInstances(); ++i)
 	{
 		OutColor.SetAndAdvance(FLinearColor::White);
 	}
 }
 
 template<typename VertexAccessorType>
-void UNiagaraDataInterfaceSkeletalMesh::GetVertexUV(FVectorVMContext& Context)
+void UNiagaraDataInterfaceSkeletalMesh::GetVertexUV(FVectorVMExternalFunctionContext& Context)
 {
 	SCOPE_CYCLE_COUNTER(STAT_NiagaraSkel_Vertex_Sample);
 	VectorVM::FUserPtrHandler<FNDISkeletalMesh_InstanceData> InstData(Context);
@@ -614,7 +614,7 @@ void UNiagaraDataInterfaceSkeletalMesh::GetVertexUV(FVectorVMContext& Context)
 		const int32 VertMax = LODData->GetNumVertices() - 1;
 		if (VertMax >= 0)
 		{
-			for (int32 i = 0; i < Context.NumInstances; ++i)
+			for (int32 i = 0; i < Context.GetNumInstances(); ++i)
 			{
 				const int32 Vert = FMath::Clamp(VertParam.GetAndAdvance(), 0, VertMax);
 				const int32 UVSet = UVSetParam.GetAndAdvance();
@@ -627,7 +627,7 @@ void UNiagaraDataInterfaceSkeletalMesh::GetVertexUV(FVectorVMContext& Context)
 	}
 
 	// Fall though for bad data
-	for (int32 i = 0; i < Context.NumInstances; ++i)
+	for (int32 i = 0; i < Context.GetNumInstances(); ++i)
 	{
 		OutUV.SetAndAdvance(FVector2D::ZeroVector);
 	}
@@ -635,7 +635,7 @@ void UNiagaraDataInterfaceSkeletalMesh::GetVertexUV(FVectorVMContext& Context)
 
 // Stub specialization for no valid mesh data on the data interface
 template<>
-void UNiagaraDataInterfaceSkeletalMesh::GetVertexUV<FSkelMeshVertexAccessorNoop>(FVectorVMContext& Context)
+void UNiagaraDataInterfaceSkeletalMesh::GetVertexUV<FSkelMeshVertexAccessorNoop>(FVectorVMExternalFunctionContext& Context)
 {
 	FNDIInputParam<int32> VertParam(Context);
 	FNDIInputParam<int32> UVSetParam(Context);
@@ -643,7 +643,7 @@ void UNiagaraDataInterfaceSkeletalMesh::GetVertexUV<FSkelMeshVertexAccessorNoop>
 
 	FNDIOutputParam<FVector2D> OutUV(Context);
 
-	for (int32 i = 0; i < Context.NumInstances; ++i)
+	for (int32 i = 0; i < Context.GetNumInstances(); ++i)
 	{
 		OutUV.SetAndAdvance(FVector2D::ZeroVector);
 	}
@@ -651,7 +651,7 @@ void UNiagaraDataInterfaceSkeletalMesh::GetVertexUV<FSkelMeshVertexAccessorNoop>
 
 struct FGetVertexSkinnedDataOutputHandler
 {
-	FGetVertexSkinnedDataOutputHandler(FVectorVMContext& Context)
+	FGetVertexSkinnedDataOutputHandler(FVectorVMExternalFunctionContext& Context)
 		: Position(Context)
 		, Velocity(Context)
 		, TangentZ(Context)
@@ -678,7 +678,7 @@ struct FGetVertexSkinnedDataOutputHandler
 };
 
 template<typename SkinningHandlerType, typename TransformHandlerType, typename VertexAccessorType>
-void UNiagaraDataInterfaceSkeletalMesh::GetVertexSkinnedData(FVectorVMContext& Context)
+void UNiagaraDataInterfaceSkeletalMesh::GetVertexSkinnedData(FVectorVMExternalFunctionContext& Context)
 {
 	SCOPE_CYCLE_COUNTER(STAT_NiagaraSkel_Vertex_Sample);
 	VectorVM::FUserPtrHandler<FNDISkeletalMesh_InstanceData> InstData(Context);
@@ -704,7 +704,7 @@ void UNiagaraDataInterfaceSkeletalMesh::GetVertexSkinnedData(FVectorVMContext& C
 			const float InvDt = 1.0f / InstData->DeltaSeconds;
 			const bool bNeedsTangentBasis = Output.bNeedsTangentX || Output.bNeedsTangentY || Output.bNeedsTangentZ;
 
-			for (int32 i = 0; i < Context.NumInstances; ++i)
+			for (int32 i = 0; i < Context.GetNumInstances(); ++i)
 			{
 				const int32 Vertex = FMath::Clamp(VertParam.GetAndAdvance(), 0, VertMax);
 
@@ -756,7 +756,7 @@ void UNiagaraDataInterfaceSkeletalMesh::GetVertexSkinnedData(FVectorVMContext& C
 	}
 
 	// Fall though for bad data
-	for (int32 i = 0; i < Context.NumInstances; ++i)
+	for (int32 i = 0; i < Context.GetNumInstances(); ++i)
 	{
 		FVector3f Position = FVector3f::ZeroVector;
 		if (Output.bNeedsPosition || Output.bNeedsVelocity)
diff --git a/Engine/Plugins/FX/Niagara/Source/Niagara/Private/NiagaraDataInterface2DArrayTexture.cpp b/Engine/Plugins/FX/Niagara/Source/Niagara/Private/NiagaraDataInterface2DArrayTexture.cpp
--- a/Engine/Plugins/FX/Niagara/Source/Niagara/Private/NiagaraDataInterface2DArrayTexture.cpp
+++ b/Engine/Plugins/FX/Niagara/Source/Niagara/Private/NiagaraDataInterface2DArrayTexture.cpp
@@ -127,12 +127,12 @@ bool UNiagaraDataInterface2DArrayTexture::PerInstanceTick(void* PerInstanceData,
 	return false;
 }
 
-void UNiagaraDataInterface2DArrayTexture::GetTextureDimensions(FVectorVMContext& Context)
+void UNiagaraDataInterface2DArrayTexture::GetTextureDimensions(FVectorVMExternalFunctionContext& Context)
 {
 	FNDIOutputParam<FVector3f> DimensionsOut(Context);
 
 	const FVector FloatTextureSize(TextureSize.X, TextureSize.Y, TextureSize.Z);
-	for (int32 i = 0; i < Context.NumInstances; ++i)
+	for (int32 i = 0; i < Context.GetNumInstances(); ++i)
 	{
 		DimensionsOut.SetAndAdvance(FloatTextureSize);
 	}
diff --git a/Engine/Plugins/FX/Niagara/Source/Niagara/Private/NiagaraDataInterfaceAudioOscilloscope.cpp b/Engine/Plugins/FX/Niagara/Source/Niagara/Private/NiagaraDataInterfaceAudioOscilloscope.cpp
--- a/Engine/Plugins/FX/Niagara/Source/Niagara/Private/NiagaraDataInterfaceAudioOscilloscope.cpp
+++ b/Engine/Plugins/FX/Niagara/Source/Niagara/Private/NiagaraDataInterfaceAudioOscilloscope.cpp
@@ -150,7 +150,7 @@ float FNiagaraDataInterfaceProxyOscilloscope::SampleAudio(float NormalizedPositi
 	return FMath::Lerp<float>(LowerFrameAmplitude, HigherFrameAmplitude, Fraction);
 }
 
-void UNiagaraDataInterfaceAudioOscilloscope::SampleAudio(FVectorVMContext& Context)
+void UNiagaraDataInterfaceAudioOscilloscope::SampleAudio(FVectorVMExternalFunctionContext& Context)
 {
 	const int32 NumFramesInDownsampledBuffer = GetProxyAs<FNiagaraDataInterfaceProxyOscilloscope>()->DownsampleAudioToBuffer();
 
@@ -160,7 +160,7 @@ void UNiagaraDataInterfaceAudioOscilloscope::SampleAudio(FVectorVMContext& Conte
 
 	const int32 NumChannelsInDownsampledBuffer = GetProxyAs<FNiagaraDataInterfaceProxyOscilloscope>()->GetNumChannels();
 
-	for (int32 InstanceIdx = 0; InstanceIdx < Context.NumInstances; ++InstanceIdx)
+	for (int32 InstanceIdx = 0; InstanceIdx < Context.GetNumInstances(); ++InstanceIdx)
 	{
 		float Position = InNormalizedPos.Get();
 		int32 Channel = InChannel.Get();
@@ -177,11 +177,11 @@ int32 FNiagaraDataInterfaceProxyOscilloscope::GetNumChannels()
 	return NumChannelsInDownsampledBuffer.GetValue();
 }
 
-void UNiagaraDataInterfaceAudioOscilloscope::GetNumChannels(FVectorVMContext& Context)
+void UNiagaraDataInterfaceAudioOscilloscope::GetNumChannels(FVectorVMExternalFunctionContext& Context)
 {
 	VectorVM::FExternalFuncRegisterHandler<int32> OutChannel(Context);
 
-	for (int32 InstanceIdx = 0; InstanceIdx < Context.NumInstances; ++InstanceIdx)
+	for (int32 InstanceIdx = 0; InstanceIdx < Context.GetNumInstances(); ++InstanceIdx)
 	{
 		*OutChannel.GetDestAndAdvance() = GetProxyAs<FNiagaraDataInterfaceProxyOscilloscope>()->GetNumChannels();
 	}
diff --git a/Engine/Plugins/FX/Niagara/Source/Niagara/Private/NiagaraDataInterfaceAudioPlayer.cpp b/Engine/Plugins/FX/Niagara/Source/Niagara/Private/NiagaraDataInterfaceAudioPlayer.cpp
--- a/Engine/Plugins/FX/Niagara/Source/Niagara/Private/NiagaraDataInterfaceAudioPlayer.cpp
+++ b/Engine/Plugins/FX/Niagara/Source/Niagara/Private/NiagaraDataInterfaceAudioPlayer.cpp
@@ -434,7 +434,7 @@ void UNiagaraDataInterfaceAudioPlayer::GetVMExternalFunction(const FVMExternalFu
 	}
 }
 
-void UNiagaraDataInterfaceAudioPlayer::SetParameterBool(FVectorVMContext& Context)
+void UNiagaraDataInterfaceAudioPlayer::SetParameterBool(FVectorVMExternalFunctionContext& Context)
 {
 	VectorVM::FUserPtrHandler<FAudioPlayerInterface_InstanceData> InstData(Context);
 
@@ -444,7 +444,7 @@ void UNiagaraDataInterfaceAudioPlayer::SetParameterBool(FVectorVMContext& Contex
 	checkfSlow(InstData.Get(), TEXT("Audio player interface has invalid instance data. %s"), *GetPathName());
 
 	InstData->bHadPersistentAudioUpdateThisTick = true;
-	for (int32 i = 0; i < Context.NumInstances; ++i)
+	for (int32 i = 0; i < Context.GetNumInstances(); ++i)
 	{
 		int32 Handle = AudioHandleInParam.GetAndAdvance();
 		int32 NameIndex = NameIndexParam.GetAndAdvance();
@@ -467,7 +467,7 @@ void UNiagaraDataInterfaceAudioPlayer::SetParameterBool(FVectorVMContext& Contex
 	}
 }
 
-void UNiagaraDataInterfaceAudioPlayer::SetParameterInteger(FVectorVMContext& Context)
+void UNiagaraDataInterfaceAudioPlayer::SetParameterInteger(FVectorVMExternalFunctionContext& Context)
 {
 	VectorVM::FUserPtrHandler<FAudioPlayerInterface_InstanceData> InstData(Context);
 
@@ -477,7 +477,7 @@ void UNiagaraDataInterfaceAudioPlayer::SetParameterInteger(FVectorVMContext& Con
 	checkfSlow(InstData.Get(), TEXT("Audio player interface has invalid instance data. %s"), *GetPathName());
 
 	InstData->bHadPersistentAudioUpdateThisTick = true;
-	for (int32 i = 0; i < Context.NumInstances; ++i)
+	for (int32 i = 0; i < Context.GetNumInstances(); ++i)
 	{
 		int32 Handle = AudioHandleInParam.GetAndAdvance();
 		int32 NameIndex = NameIndexParam.GetAndAdvance();
@@ -500,7 +500,7 @@ void UNiagaraDataInterfaceAudioPlayer::SetParameterInteger(FVectorVMContext& Con
 	}
 }
 
-void UNiagaraDataInterfaceAudioPlayer::SetParameterFloat(FVectorVMContext& Context)
+void UNiagaraDataInterfaceAudioPlayer::SetParameterFloat(FVectorVMExternalFunctionContext& Context)
 {
 	VectorVM::FUserPtrHandler<FAudioPlayerInterface_InstanceData> InstData(Context);
 
@@ -510,7 +510,7 @@ void UNiagaraDataInterfaceAudioPlayer::SetParameterFloat(FVectorVMContext& Conte
 	checkfSlow(InstData.Get(), TEXT("Audio player interface has invalid instance data. %s"), *GetPathName());
 
 	InstData->bHadPersistentAudioUpdateThisTick = true;
-	for (int32 i = 0; i < Context.NumInstances; ++i)
+	for (int32 i = 0; i < Context.GetNumInstances(); ++i)
 	{
 		int32 Handle = AudioHandleInParam.GetAndAdvance();
 		int32 NameIndex = NameIndexParam.GetAndAdvance();
@@ -533,7 +533,7 @@ void UNiagaraDataInterfaceAudioPlayer::SetParameterFloat(FVectorVMContext& Conte
 	}
 }
 
-void UNiagaraDataInterfaceAudioPlayer::UpdateVolume(FVectorVMContext& Context)
+void UNiagaraDataInterfaceAudioPlayer::UpdateVolume(FVectorVMExternalFunctionContext& Context)
 {
 	VectorVM::FUserPtrHandler<FAudioPlayerInterface_InstanceData> InstData(Context);
 
@@ -542,7 +542,7 @@ void UNiagaraDataInterfaceAudioPlayer::UpdateVolume(FVectorVMContext& Context)
 	checkfSlow(InstData.Get(), TEXT("Audio player interface has invalid instance data. %s"), *GetPathName());
 
 	InstData->bHadPersistentAudioUpdateThisTick = true;
-	for (int32 i = 0; i < Context.NumInstances; ++i)
+	for (int32 i = 0; i < Context.GetNumInstances(); ++i)
 	{
 		int32 Handle = AudioHandleInParam.GetAndAdvance();
 		float Volume = VolumeParam.GetAndAdvance();
@@ -563,7 +563,7 @@ void UNiagaraDataInterfaceAudioPlayer::UpdateVolume(FVectorVMContext& Context)
 	}
 }
 
-void UNiagaraDataInterfaceAudioPlayer::UpdatePitch(FVectorVMContext& Context)
+void UNiagaraDataInterfaceAudioPlayer::UpdatePitch(FVectorVMExternalFunctionContext& Context)
 {
 	VectorVM::FUserPtrHandler<FAudioPlayerInterface_InstanceData> InstData(Context);
 
@@ -572,7 +572,7 @@ void UNiagaraDataInterfaceAudioPlayer::UpdatePitch(FVectorVMContext& Context)
 	checkfSlow(InstData.Get(), TEXT("Audio player interface has invalid instance data. %s"), *GetPathName());
 
 	InstData->bHadPersistentAudioUpdateThisTick = true;
-	for (int32 i = 0; i < Context.NumInstances; ++i)
+	for (int32 i = 0; i < Context.GetNumInstances(); ++i)
 	{
 		int32 Handle = AudioHandleInParam.GetAndAdvance();
 		float Pitch = PitchParam.GetAndAdvance();
@@ -593,7 +593,7 @@ void UNiagaraDataInterfaceAudioPlayer::UpdatePitch(FVectorVMContext& Context)
 	}
 }
 
-void UNiagaraDataInterfaceAudioPlayer::UpdateLocation(FVectorVMContext& Context)
+void UNiagaraDataInterfaceAudioPlayer::UpdateLocation(FVectorVMExternalFunctionContext& Context)
 {
 	VectorVM::FUserPtrHandler<FAudioPlayerInterface_InstanceData> InstData(Context);
 
@@ -602,7 +602,7 @@ void UNiagaraDataInterfaceAudioPlayer::UpdateLocation(FVectorVMContext& Context)
 	checkfSlow(InstData.Get(), TEXT("Audio player interface has invalid instance data. %s"), *GetPathName());
 
 	InstData->bHadPersistentAudioUpdateThisTick = true;
-	for (int32 i = 0; i < Context.NumInstances; ++i)
+	for (int32 i = 0; i < Context.GetNumInstances(); ++i)
 	{
 		int32 Handle = AudioHandleInParam.GetAndAdvance();
 		FVector Location = LocationParam.GetAndAdvance();
@@ -623,7 +623,7 @@ void UNiagaraDataInterfaceAudioPlayer::UpdateLocation(FVectorVMContext& Context)
 	}
 }
 
-void UNiagaraDataInterfaceAudioPlayer::UpdateRotation(FVectorVMContext& Context)
+void UNiagaraDataInterfaceAudioPlayer::UpdateRotation(FVectorVMExternalFunctionContext& Context)
 {
 	VectorVM::FUserPtrHandler<FAudioPlayerInterface_InstanceData> InstData(Context);
 
@@ -632,7 +632,7 @@ void UNiagaraDataInterfaceAudioPlayer::UpdateRotation(FVectorVMContext& Context)
 	checkfSlow(InstData.Get(), TEXT("Audio player interface has invalid instance data. %s"), *GetPathName());
 
 	InstData->bHadPersistentAudioUpdateThisTick = true;
-	for (int32 i = 0; i < Context.NumInstances; ++i)
+	for (int32 i = 0; i < Context.GetNumInstances(); ++i)
 	{
 		int32 Handle = AudioHandleInParam.GetAndAdvance();
 		FVector Rotation = RotationParam.GetAndAdvance();
@@ -654,7 +654,7 @@ void UNiagaraDataInterfaceAudioPlayer::UpdateRotation(FVectorVMContext& Context)
 	}
 }
 
-void UNiagaraDataInterfaceAudioPlayer::SetPausedState(FVectorVMContext& Context)
+void UNiagaraDataInterfaceAudioPlayer::SetPausedState(FVectorVMExternalFunctionContext& Context)
 {
 	VectorVM::FUserPtrHandler<FAudioPlayerInterface_InstanceData> InstData(Context);
 
@@ -663,7 +663,7 @@ void UNiagaraDataInterfaceAudioPlayer::SetPausedState(FVectorVMContext& Context)
 	checkfSlow(InstData.Get(), TEXT("Audio player interface has invalid instance data. %s"), *GetPathName());
 
 	InstData->bHadPersistentAudioUpdateThisTick = true;
-	for (int32 i = 0; i < Context.NumInstances; ++i)
+	for (int32 i = 0; i < Context.GetNumInstances(); ++i)
 	{
 		int32 Handle = AudioHandleInParam.GetAndAdvance();
 		bool IsPaused = PausedParam.GetAndAdvance();
@@ -684,7 +684,7 @@ void UNiagaraDataInterfaceAudioPlayer::SetPausedState(FVectorVMContext& Context)
 	}
 }
 
-void UNiagaraDataInterfaceAudioPlayer::PlayOneShotAudio(FVectorVMContext& Context)
+void UNiagaraDataInterfaceAudioPlayer::PlayOneShotAudio(FVectorVMExternalFunctionContext& Context)
 {
 	VectorVM::FUserPtrHandler<FAudioPlayerInterface_InstanceData> InstData(Context);
 
@@ -707,7 +707,7 @@ void UNiagaraDataInterfaceAudioPlayer::PlayOneShotAudio(FVectorVMContext& Contex
 	checkfSlow(InstData.Get(), TEXT("Audio player interface has invalid instance data. %s"), *GetPathName());
 	bool ValidSoundData = InstData->SoundToPlay.IsValid();
 
-	for (int32 i = 0; i < Context.NumInstances; ++i)
+	for (int32 i = 0; i < Context.GetNumInstances(); ++i)
 	{
 		FNiagaraBool ShouldPlay = PlayDataParam.GetAndAdvance();
 		FAudioParticleData Data;
@@ -726,7 +726,7 @@ void UNiagaraDataInterfaceAudioPlayer::PlayOneShotAudio(FVectorVMContext& Contex
 	}
 }
 
-void UNiagaraDataInterfaceAudioPlayer::PlayPersistentAudio(FVectorVMContext& Context)
+void UNiagaraDataInterfaceAudioPlayer::PlayPersistentAudio(FVectorVMExternalFunctionContext& Context)
 {
 	VectorVM::FUserPtrHandler<FAudioPlayerInterface_InstanceData> InstData(Context);
 
@@ -745,7 +745,7 @@ void UNiagaraDataInterfaceAudioPlayer::PlayPersistentAudio(FVectorVMContext& Con
 	checkfSlow(InstData.Get(), TEXT("Audio player interface has invalid instance data. %s"), *GetPathName());
 
 	InstData->bHadPersistentAudioUpdateThisTick = true;
-	for (int32 i = 0; i < Context.NumInstances; ++i)
+	for (int32 i = 0; i < Context.GetNumInstances(); ++i)
 	{
 		bool ShouldPlay = PlayAudioParam.GetAndAdvance();
 		int32 Handle = AudioHandleInParam.GetAndAdvance();
diff --git a/Engine/Plugins/FX/Niagara/Source/Niagara/Private/NiagaraDataInterfaceAudioSpectrum.cpp b/Engine/Plugins/FX/Niagara/Source/Niagara/Private/NiagaraDataInterfaceAudioSpectrum.cpp
--- a/Engine/Plugins/FX/Niagara/Source/Niagara/Private/NiagaraDataInterfaceAudioSpectrum.cpp
+++ b/Engine/Plugins/FX/Niagara/Source/Niagara/Private/NiagaraDataInterfaceAudioSpectrum.cpp
@@ -564,7 +564,7 @@ UNiagaraDataInterfaceAudioSpectrum::UNiagaraDataInterfaceAudioSpectrum(FObjectIn
 }
 
 
-void UNiagaraDataInterfaceAudioSpectrum::GetSpectrumValue(FVectorVMContext& Context)
+void UNiagaraDataInterfaceAudioSpectrum::GetSpectrumValue(FVectorVMExternalFunctionContext& Context)
 {
 	GetProxyAs<FNiagaraDataInterfaceProxySpectrum>()->UpdateSpectrum();
 
@@ -572,7 +572,7 @@ void UNiagaraDataInterfaceAudioSpectrum::GetSpectrumValue(FVectorVMContext& Cont
 	VectorVM::FExternalFuncInputHandler<int32> InChannel(Context);
 	VectorVM::FExternalFuncRegisterHandler<float> OutValue(Context);
 
-	for (int32 InstanceIdx = 0; InstanceIdx < Context.NumInstances; ++InstanceIdx)
+	for (int32 InstanceIdx = 0; InstanceIdx < Context.GetNumInstances(); ++InstanceIdx)
 	{
 		float Position = InNormalizedPos.Get();
 		int32 Channel = InChannel.Get();
@@ -584,11 +584,11 @@ void UNiagaraDataInterfaceAudioSpectrum::GetSpectrumValue(FVectorVMContext& Cont
 	}
 }
 
-void UNiagaraDataInterfaceAudioSpectrum::GetNumChannels(FVectorVMContext& Context)
+void UNiagaraDataInterfaceAudioSpectrum::GetNumChannels(FVectorVMExternalFunctionContext& Context)
 {
 	VectorVM::FExternalFuncRegisterHandler<int32> OutChannel(Context);
 
-	for (int32 InstanceIdx = 0; InstanceIdx < Context.NumInstances; ++InstanceIdx)
+	for (int32 InstanceIdx = 0; InstanceIdx < Context.GetNumInstances(); ++InstanceIdx)
 	{
 		*OutChannel.GetDestAndAdvance() = GetProxyAs<FNiagaraDataInterfaceProxySpectrum>()->GetNumChannels();
 	}
diff --git a/Engine/Plugins/FX/Niagara/Source/Niagara/Private/NiagaraDataInterfaceCamera.cpp b/Engine/Plugins/FX/Niagara/Source/Niagara/Private/NiagaraDataInterfaceCamera.cpp
--- a/Engine/Plugins/FX/Niagara/Source/Niagara/Private/NiagaraDataInterfaceCamera.cpp
+++ b/Engine/Plugins/FX/Niagara/Source/Niagara/Private/NiagaraDataInterfaceCamera.cpp
@@ -434,20 +434,20 @@ void UNiagaraDataInterfaceCamera::GetVMExternalFunction(const FVMExternalFunctio
 	}
 }
 
-void UNiagaraDataInterfaceCamera::GetCameraFOV(FVectorVMContext& Context)
+void UNiagaraDataInterfaceCamera::GetCameraFOV(FVectorVMExternalFunctionContext& Context)
 {
 	VectorVM::FUserPtrHandler<FCameraDataInterface_InstanceData> InstData(Context);
 	VectorVM::FExternalFuncRegisterHandler<float> OutFov(Context);
 
 	float Fov = InstData.Get()->CameraFOV;
 
-	for (int32 i = 0; i < Context.NumInstances; ++i)
+	for (int32 i = 0; i < Context.GetNumInstances(); ++i)
 	{
 		*OutFov.GetDestAndAdvance() = Fov;
 	}
 }
 
-void UNiagaraDataInterfaceCamera::GetCameraProperties(FVectorVMContext& Context)
+void UNiagaraDataInterfaceCamera::GetCameraProperties(FVectorVMExternalFunctionContext& Context)
 {
 	VectorVM::FUserPtrHandler<FCameraDataInterface_InstanceData> InstData(Context);
 
@@ -477,7 +477,7 @@ void UNiagaraDataInterfaceCamera::GetCameraProperties(FVectorVMContext& Context)
 	const FVector Up = RotationMatrix.GetScaledAxis(EAxis::Z);
 	const FVector Right = RotationMatrix.GetScaledAxis(EAxis::Y);
 
-	for (int32 i = 0; i < Context.NumInstances; ++i)
+	for (int32 i = 0; i < Context.GetNumInstances(); ++i)
 	{
 		*CamPosX.GetDestAndAdvance() = XPos;
 		*CamPosY.GetDestAndAdvance() = YPos;
@@ -497,7 +497,7 @@ void UNiagaraDataInterfaceCamera::GetCameraProperties(FVectorVMContext& Context)
 	}
 }
 
-void UNiagaraDataInterfaceCamera::GetClosestParticles(FVectorVMContext& Context)
+void UNiagaraDataInterfaceCamera::GetClosestParticles(FVectorVMExternalFunctionContext& Context)
 {
 	VectorVM::FUserPtrHandler<FCameraDataInterface_InstanceData> InstData(Context);
 
@@ -505,10 +505,10 @@ void UNiagaraDataInterfaceCamera::GetClosestParticles(FVectorVMContext& Context)
 	FNDIInputParam<int32> CountParam(Context);
 	FNDIOutputParam<FNiagaraBool> ResultOutParam(Context);
 
-	int32 Count = Context.NumInstances > 0 ? CountParam.GetAndAdvance() : 0;
+	int32 Count = Context.GetNumInstances() > 0 ? CountParam.GetAndAdvance() : 0;
 	if (Count == 0 || InstData->ParticlesSortedByDistance.Num() == 0)
 	{
-		for (int32 i = 0; i < Context.NumInstances; ++i)
+		for (int32 i = 0; i < Context.GetNumInstances(); ++i)
 		{
 			ResultOutParam.SetAndAdvance(false);
 		}
@@ -523,14 +523,14 @@ void UNiagaraDataInterfaceCamera::GetClosestParticles(FVectorVMContext& Context)
 	}
 
 	// Assign each particles their result
-	for (int32 i = 0; i < Context.NumInstances; ++i)
+	for (int32 i = 0; i < Context.GetNumInstances(); ++i)
 	{
 		FNiagaraID ParticleID = ParticleIDParam.GetAndAdvance();
 		ResultOutParam.SetAndAdvance(ClosestParticleIDs.Contains(ParticleID));
 	}
 }
 
-void UNiagaraDataInterfaceCamera::CalculateParticleDistances(FVectorVMContext& Context)
+void UNiagaraDataInterfaceCamera::CalculateParticleDistances(FVectorVMExternalFunctionContext& Context)
 {
 	VectorVM::FUserPtrHandler<FCameraDataInterface_InstanceData> InstData(Context);
 
@@ -538,7 +538,7 @@ void UNiagaraDataInterfaceCamera::CalculateParticleDistances(FVectorVMContext& C
 	FNDIInputParam<FVector3f> ParticlePosParam(Context);
 
 	FVector CameraPos = InstData->CameraLocation;
-	for (int32 i = 0; i < Context.NumInstances; ++i)
+	for (int32 i = 0; i < Context.GetNumInstances(); ++i)
 	{
 		FDistanceData DistanceData;
 		FVector ParticlePos = ParticlePosParam.GetAndAdvance();
@@ -622,7 +622,7 @@ void UNiagaraDataInterfaceCamera::GetFeedback(UNiagaraSystem* Asset, UNiagaraCom
 
 // ------- Dummy implementations for CPU execution ------------
 
-void UNiagaraDataInterfaceCamera::GetViewPropertiesGPU(FVectorVMContext& Context)
+void UNiagaraDataInterfaceCamera::GetViewPropertiesGPU(FVectorVMExternalFunctionContext& Context)
 {
 	VectorVM::FUserPtrHandler<FCameraDataInterface_InstanceData> InstData(Context);
 	TArray<VectorVM::FExternalFuncRegisterHandler<float>> OutParams;
@@ -632,7 +632,7 @@ void UNiagaraDataInterfaceCamera::GetViewPropertiesGPU(FVectorVMContext& Context
 		OutParams.Emplace(Context);
 	}
 
-	for (int32 k = 0; k < Context.NumInstances; ++k)
+	for (int32 k = 0; k < Context.GetNumInstances(); ++k)
 	{
 		for (int i = 0; i < 24; i++)
 		{
@@ -641,7 +641,7 @@ void UNiagaraDataInterfaceCamera::GetViewPropertiesGPU(FVectorVMContext& Context
 	}
 }
 
-void UNiagaraDataInterfaceCamera::GetClipSpaceTransformsGPU(FVectorVMContext& Context)
+void UNiagaraDataInterfaceCamera::GetClipSpaceTransformsGPU(FVectorVMExternalFunctionContext& Context)
 {
 	VectorVM::FUserPtrHandler<FCameraDataInterface_InstanceData> InstData(Context);
 	TArray<VectorVM::FExternalFuncRegisterHandler<float>> OutParams;
@@ -651,7 +651,7 @@ void UNiagaraDataInterfaceCamera::GetClipSpaceTransformsGPU(FVectorVMContext& Co
 		OutParams.Emplace(Context);
 	}
 
-	for (int32 k = 0; k < Context.NumInstances; ++k)
+	for (int32 k = 0; k < Context.GetNumInstances(); ++k)
 	{
 		for (int i = 0; i < 128; i++)
 		{
@@ -660,7 +660,7 @@ void UNiagaraDataInterfaceCamera::GetClipSpaceTransformsGPU(FVectorVMContext& Co
 	}
 }
 
-void UNiagaraDataInterfaceCamera::GetViewSpaceTransformsGPU(FVectorVMContext& Context)
+void UNiagaraDataInterfaceCamera::GetViewSpaceTransformsGPU(FVectorVMExternalFunctionContext& Context)
 {
 	VectorVM::FUserPtrHandler<FCameraDataInterface_InstanceData> InstData(Context);
 	TArray<VectorVM::FExternalFuncRegisterHandler<float>> OutParams;
@@ -670,7 +670,7 @@ void UNiagaraDataInterfaceCamera::GetViewSpaceTransformsGPU(FVectorVMContext& Co
 		OutParams.Emplace(Context);
 	}
 
-	for (int32 k = 0; k < Context.NumInstances; ++k)
+	for (int32 k = 0; k < Context.GetNumInstances(); ++k)
 	{
 		for (int i = 0; i < 96; i++)
 		{
@@ -680,7 +680,7 @@ void UNiagaraDataInterfaceCamera::GetViewSpaceTransformsGPU(FVectorVMContext& Co
 }
 
 
-void UNiagaraDataInterfaceCamera::GetTAAJitter(FVectorVMContext& Context)
+void UNiagaraDataInterfaceCamera::GetTAAJitter(FVectorVMExternalFunctionContext& Context)
 {
 	VectorVM::FUserPtrHandler<FCameraDataInterface_InstanceData> InstData(Context);
 	TArray<VectorVM::FExternalFuncRegisterHandler<float>> OutParams;
@@ -692,7 +692,7 @@ void UNiagaraDataInterfaceCamera::GetTAAJitter(FVectorVMContext& Context)
 		OutParams.Emplace(Context);
 	}
 
-	for (int32 k = 0; k < Context.NumInstances; ++k)
+	for (int32 k = 0; k < Context.GetNumInstances(); ++k)
 	{
 		for (int i = 0; i < ElementCount; i++)
 		{
diff --git a/Engine/Plugins/FX/Niagara/Source/Niagara/Private/NiagaraDataInterfaceCollisionQuery.cpp b/Engine/Plugins/FX/Niagara/Source/Niagara/Private/NiagaraDataInterfaceCollisionQuery.cpp
--- a/Engine/Plugins/FX/Niagara/Source/Niagara/Private/NiagaraDataInterfaceCollisionQuery.cpp
+++ b/Engine/Plugins/FX/Niagara/Source/Niagara/Private/NiagaraDataInterfaceCollisionQuery.cpp
@@ -720,7 +720,7 @@ void UNiagaraDataInterfaceCollisionQuery::ModifyCompilationEnvironment(EShaderPl
 
 #endif
 
-void UNiagaraDataInterfaceCollisionQuery::PerformQuerySyncCPU(FVectorVMContext & Context)
+void UNiagaraDataInterfaceCollisionQuery::PerformQuerySyncCPU(FVectorVMExternalFunctionContext & Context)
 {
 	VectorVM::FUserPtrHandler<CQDIPerInstanceData> InstanceData(Context);
 
@@ -749,7 +749,7 @@ void UNiagaraDataInterfaceCollisionQuery::PerformQuerySyncCPU(FVectorVMContext &
 	VectorVM::FExternalFuncRegisterHandler<int32> OutPhysicalMaterialIdx(Context);
 
 	FScopeLock ScopeLock(&CriticalSection);
-	for (int32 i = 0; i < Context.NumInstances; ++i)
+	for (int32 i = 0; i < Context.GetNumInstances(); ++i)
 	{
 		FVector Pos(StartPosParamX.GetAndAdvance(), StartPosParamY.GetAndAdvance(), StartPosParamZ.GetAndAdvance());
 		FVector Dir(EndPosParamX.GetAndAdvance(), EndPosParamY.GetAndAdvance(), EndPosParamZ.GetAndAdvance());
@@ -789,7 +789,7 @@ void UNiagaraDataInterfaceCollisionQuery::PerformQuerySyncCPU(FVectorVMContext &
 	}
 }
 
-void UNiagaraDataInterfaceCollisionQuery::PerformQueryAsyncCPU(FVectorVMContext & Context)
+void UNiagaraDataInterfaceCollisionQuery::PerformQueryAsyncCPU(FVectorVMExternalFunctionContext & Context)
 {
 	VectorVM::FUserPtrHandler<CQDIPerInstanceData> InstanceData(Context);
 
@@ -821,7 +821,7 @@ void UNiagaraDataInterfaceCollisionQuery::PerformQueryAsyncCPU(FVectorVMContext
 	VectorVM::FExternalFuncRegisterHandler<int32> OutPhysicalMaterialIdx(Context);
 
 	FScopeLock ScopeLock(&CriticalSection);
-	for (int32 i = 0; i < Context.NumInstances; ++i)
+	for (int32 i = 0; i < Context.GetNumInstances(); ++i)
 	{
 		FVector Pos(StartPosParamX.GetAndAdvance(), StartPosParamY.GetAndAdvance(), StartPosParamZ.GetAndAdvance());
 		FVector End(EndPosParamX.GetAndAdvance(), EndPosParamY.GetAndAdvance(), EndPosParamZ.GetAndAdvance());
@@ -865,7 +865,7 @@ void UNiagaraDataInterfaceCollisionQuery::PerformQueryAsyncCPU(FVectorVMContext
 	}
 }
 
-void UNiagaraDataInterfaceCollisionQuery::QuerySceneDepth(FVectorVMContext & Context)
+void UNiagaraDataInterfaceCollisionQuery::QuerySceneDepth(FVectorVMExternalFunctionContext & Context)
 {
 	UE_LOG(LogNiagara, Error, TEXT("GPU only function 'QuerySceneDepthGPU' called on CPU VM, check your module code to fix."));
 
@@ -888,7 +888,7 @@ void UNiagaraDataInterfaceCollisionQuery::QuerySceneDepth(FVectorVMContext & Con
 	VectorVM::FExternalFuncRegisterHandler<float> OutWorldNormZ(Context);
 
 	FScopeLock ScopeLock(&CriticalSection);
-	for (int32 i = 0; i < Context.NumInstances; ++i)
+	for (int32 i = 0; i < Context.GetNumInstances(); ++i)
 	{
 		*OutSceneDepth.GetDestAndAdvance() = -1;
 		*OutIsInsideView.GetDestAndAdvance() = 0;
@@ -904,7 +904,7 @@ void UNiagaraDataInterfaceCollisionQuery::QuerySceneDepth(FVectorVMContext & Con
 	}
 }
 
-void UNiagaraDataInterfaceCollisionQuery::QueryMeshDistanceField(FVectorVMContext& Context)
+void UNiagaraDataInterfaceCollisionQuery::QueryMeshDistanceField(FVectorVMExternalFunctionContext& Context)
 {
 	UE_LOG(LogNiagara, Error, TEXT("GPU only function 'QueryMeshDistanceFieldGPU' called on CPU VM, check your module code to fix."));
 
@@ -921,7 +921,7 @@ void UNiagaraDataInterfaceCollisionQuery::QueryMeshDistanceField(FVectorVMContex
 	FNDIOutputParam<FNiagaraBool> OutIsFieldValid(Context);
 
 	FScopeLock ScopeLock(&CriticalSection);
-	for (int32 i = 0; i < Context.NumInstances; ++i)
+	for (int32 i = 0; i < Context.GetNumInstances(); ++i)
 	{
 		*OutSurfaceDistance.GetDestAndAdvance() = -1;
 		*OutFieldGradientX.GetDestAndAdvance() = 0.0f;
diff --git a/Engine/Plugins/FX/Niagara/Source/Niagara/Private/NiagaraDataInterfaceColorCurve.cpp b/Engine/Plugins/FX/Niagara/Source/Niagara/Private/NiagaraDataInterfaceColorCurve.cpp
--- a/Engine/Plugins/FX/Niagara/Source/Niagara/Private/NiagaraDataInterfaceColorCurve.cpp
+++ b/Engine/Plugins/FX/Niagara/Source/Niagara/Private/NiagaraDataInterfaceColorCurve.cpp
@@ -238,7 +238,7 @@ FORCEINLINE_DEBUGGABLE FLinearColor UNiagaraDataInterfaceColorCurve::SampleCurve
 }
 
 template<typename UseLUT>
-void UNiagaraDataInterfaceColorCurve::SampleCurve(FVectorVMContext& Context)
+void UNiagaraDataInterfaceColorCurve::SampleCurve(FVectorVMExternalFunctionContext& Context)
 {
 	//TODO: Create some SIMDable optimized representation of the curve to do this faster.
 	VectorVM::FExternalFuncInputHandler<float> XParam(Context);
@@ -247,7 +247,7 @@ void UNiagaraDataInterfaceColorCurve::SampleCurve(FVectorVMContext& Context)
 	VectorVM::FExternalFuncRegisterHandler<float> SamplePtrB(Context);
 	VectorVM::FExternalFuncRegisterHandler<float> SamplePtrA(Context);
 
-	for (int32 i = 0; i < Context.NumInstances; ++i)
+	for (int32 i = 0; i < Context.GetNumInstances(); ++i)
 	{
 		float X = XParam.GetAndAdvance();
 		FLinearColor C = SampleCurveInternal<UseLUT>(X);
diff --git a/Engine/Plugins/FX/Niagara/Source/Niagara/Private/NiagaraDataInterfaceCubeTexture.cpp b/Engine/Plugins/FX/Niagara/Source/Niagara/Private/NiagaraDataInterfaceCubeTexture.cpp
--- a/Engine/Plugins/FX/Niagara/Source/Niagara/Private/NiagaraDataInterfaceCubeTexture.cpp
+++ b/Engine/Plugins/FX/Niagara/Source/Niagara/Private/NiagaraDataInterfaceCubeTexture.cpp
@@ -142,25 +142,25 @@ bool UNiagaraDataInterfaceCubeTexture::PerInstanceTick(void* PerInstanceData, FN
 	return false;
 }
 
-void UNiagaraDataInterfaceCubeTexture::GetTextureDimensions(FVectorVMContext& Context)
+void UNiagaraDataInterfaceCubeTexture::GetTextureDimensions(FVectorVMExternalFunctionContext& Context)
 {
 	FNDIOutputParam<int32> OutWidth(Context);
 	FNDIOutputParam<int32> OutHeight(Context);
 
-	for (int32 i = 0; i < Context.NumInstances; ++i)
+	for (int32 i = 0; i < Context.GetNumInstances(); ++i)
 	{
 		OutWidth.SetAndAdvance(TextureSize.X);
 		OutHeight.SetAndAdvance(TextureSize.Y);
 	}
 }
 
-void UNiagaraDataInterfaceCubeTexture::SampleCubeTexture(FVectorVMContext& Context)
+void UNiagaraDataInterfaceCubeTexture::SampleCubeTexture(FVectorVMExternalFunctionContext& Context)
 {
 	FNDIInputParam<FVector3f> InCoord(Context);
 	FNDIInputParam<float> InMipLevel(Context);
 	FNDIOutputParam<FVector4> OutColor(Context);
 
-	for (int32 i = 0; i < Context.NumInstances; ++i)
+	for (int32 i = 0; i < Context.GetNumInstances(); ++i)
 	{
 		OutColor.SetAndAdvance(FVector4(1.0f, 0.0f, 1.0f, 1.0f));
 	}
diff --git a/Engine/Plugins/FX/Niagara/Source/Niagara/Private/NiagaraDataInterfaceCurlNoise.cpp b/Engine/Plugins/FX/Niagara/Source/Niagara/Private/NiagaraDataInterfaceCurlNoise.cpp
--- a/Engine/Plugins/FX/Niagara/Source/Niagara/Private/NiagaraDataInterfaceCurlNoise.cpp
+++ b/Engine/Plugins/FX/Niagara/Source/Niagara/Private/NiagaraDataInterfaceCurlNoise.cpp
@@ -403,7 +403,7 @@ void UNiagaraDataInterfaceCurlNoise::GetVMExternalFunction(const FVMExternalFunc
 	NDI_FUNC_BINDER(UNiagaraDataInterfaceCurlNoise, SampleNoiseField)::Bind(this, OutFunc);
 }
 
-void UNiagaraDataInterfaceCurlNoise::SampleNoiseField(FVectorVMContext& Context)
+void UNiagaraDataInterfaceCurlNoise::SampleNoiseField(FVectorVMExternalFunctionContext& Context)
 {
 	VectorVM::FExternalFuncInputHandler<float> XParam(Context);
 	VectorVM::FExternalFuncInputHandler<float> YParam(Context);
@@ -412,7 +412,7 @@ void UNiagaraDataInterfaceCurlNoise::SampleNoiseField(FVectorVMContext& Context)
 	VectorVM::FExternalFuncRegisterHandler<float> OutSampleY(Context);
 	VectorVM::FExternalFuncRegisterHandler<float> OutSampleZ(Context);
 
-	for (int32 InstanceIdx = 0; InstanceIdx < Context.NumInstances; ++InstanceIdx)
+	for (int32 InstanceIdx = 0; InstanceIdx < Context.GetNumInstances(); ++InstanceIdx)
 	{
 		FVector InCoords = FVector(XParam.GetAndAdvance(), YParam.GetAndAdvance(), ZParam.GetAndAdvance());
 
diff --git a/Engine/Plugins/FX/Niagara/Source/Niagara/Private/NiagaraDataInterfaceCurve.cpp b/Engine/Plugins/FX/Niagara/Source/Niagara/Private/NiagaraDataInterfaceCurve.cpp
--- a/Engine/Plugins/FX/Niagara/Source/Niagara/Private/NiagaraDataInterfaceCurve.cpp
+++ b/Engine/Plugins/FX/Niagara/Source/Niagara/Private/NiagaraDataInterfaceCurve.cpp
@@ -196,12 +196,12 @@ FORCEINLINE_DEBUGGABLE float UNiagaraDataInterfaceCurve::SampleCurveInternal<TIn
 }
 
 template<typename UseLUT>
-void UNiagaraDataInterfaceCurve::SampleCurve(FVectorVMContext& Context)
+void UNiagaraDataInterfaceCurve::SampleCurve(FVectorVMExternalFunctionContext& Context)
 {
 	VectorVM::FExternalFuncInputHandler<float> XParam(Context);
 	VectorVM::FExternalFuncRegisterHandler<float> OutSample(Context);
 
-	for (int32 i = 0; i < Context.NumInstances; ++i)
+	for (int32 i = 0; i < Context.GetNumInstances(); ++i)
 	{
 		*OutSample.GetDest() = SampleCurveInternal<UseLUT>(XParam.Get());
 		XParam.Advance();
