Wrapped FVectorVMContext in NiagaraDataInterfaceFunctionContext for external functions to facilitate running an experimental VM in parallel.

diff --git a/Engine/Plugins/FX/Niagara/Source/Niagara/Private/NiagaraDataInterfaceExport.cpp b/Engine/Plugins/FX/Niagara/Source/Niagara/Private/NiagaraDataInterfaceExport.cpp
--- a/Engine/Plugins/FX/Niagara/Source/Niagara/Private/NiagaraDataInterfaceExport.cpp
+++ b/Engine/Plugins/FX/Niagara/Source/Niagara/Private/NiagaraDataInterfaceExport.cpp
@@ -506,7 +506,7 @@ void UNiagaraDataInterfaceExport::GetVMExternalFunction(const FVMExternalFunctio
 	}
 }
 
-void UNiagaraDataInterfaceExport::StoreData(FVectorVMContext& Context)
+void UNiagaraDataInterfaceExport::StoreData(FVectorVMExternalFunctionContext& Context)
 {
 	VectorVM::FUserPtrHandler<FNDIExportInstanceData_GameThread> InstData(Context);
 
@@ -520,7 +520,7 @@ void UNiagaraDataInterfaceExport::StoreData(FVectorVMContext& Context)
 	checkfSlow(InstData.Get(), TEXT("Export data interface has invalid instance data. %s"), *GetPathName());
 	bool ValidHandlerData = InstData->UserParamBinding.BoundVariable.IsValid() && InstData->CallbackHandler.IsValid();
 
-	for (int32 i = 0; i < Context.NumInstances; ++i)
+	for (int32 i = 0; i < Context.GetNumInstances(); ++i)
 	{
 		const bool ShouldStore = StoreDataParam.GetAndAdvance();
 
@@ -538,7 +538,7 @@ void UNiagaraDataInterfaceExport::StoreData(FVectorVMContext& Context)
 	}
 }
 
-void UNiagaraDataInterfaceExport::ExportData(FVectorVMContext& Context)
+void UNiagaraDataInterfaceExport::ExportData(FVectorVMExternalFunctionContext& Context)
 {
 	VectorVM::FUserPtrHandler<FNDIExportInstanceData_GameThread> InstData(Context);
 
@@ -552,7 +552,7 @@ void UNiagaraDataInterfaceExport::ExportData(FVectorVMContext& Context)
 	checkfSlow(InstData.Get(), TEXT("Export data interface has invalid instance data. %s"), *GetPathName());
 	bool ValidHandlerData = InstData->UserParamBinding.BoundVariable.IsValid() && InstData->CallbackHandler.IsValid();
 
-	for (int32 i = 0; i < Context.NumInstances; ++i)
+	for (int32 i = 0; i < Context.GetNumInstances(); ++i)
 	{
 		const bool ShouldStore = StoreDataParam.GetAndAdvance();
 
diff --git a/Engine/Plugins/FX/Niagara/Source/Niagara/Private/NiagaraDataInterfaceGrid2DCollection.cpp b/Engine/Plugins/FX/Niagara/Source/Niagara/Private/NiagaraDataInterfaceGrid2DCollection.cpp
--- a/Engine/Plugins/FX/Niagara/Source/Niagara/Private/NiagaraDataInterfaceGrid2DCollection.cpp
+++ b/Engine/Plugins/FX/Niagara/Source/Niagara/Private/NiagaraDataInterfaceGrid2DCollection.cpp
@@ -2291,43 +2291,43 @@ void UNiagaraDataInterfaceGrid2DCollection::GetTextureSize(const UNiagaraCompone
 	SizeY = Grid2DInstanceData->NumCells.Y;
 }
 
-void UNiagaraDataInterfaceGrid2DCollection::GetWorldBBoxSize(FVectorVMContext& Context)
+void UNiagaraDataInterfaceGrid2DCollection::GetWorldBBoxSize(FVectorVMExternalFunctionContext& Context)
 {
 	VectorVM::FUserPtrHandler<FGrid2DCollectionRWInstanceData_GameThread> InstData(Context);
 	FNDIOutputParam<FVector2D> OutWorldBounds(Context);
 
-	for (int32 InstanceIdx = 0; InstanceIdx < Context.NumInstances; ++InstanceIdx)
+	for (int32 InstanceIdx = 0; InstanceIdx < Context.GetNumInstances(); ++InstanceIdx)
 	{
 		OutWorldBounds.SetAndAdvance(InstData->WorldBBoxSize);
 	}
 }
 
 
-void UNiagaraDataInterfaceGrid2DCollection::GetCellSize(FVectorVMContext& Context)
+void UNiagaraDataInterfaceGrid2DCollection::GetCellSize(FVectorVMExternalFunctionContext& Context)
 {
 	VectorVM::FUserPtrHandler<FGrid2DCollectionRWInstanceData_GameThread> InstData(Context);
 	FNDIOutputParam<FVector2D> OutCellSize(Context);
 
-	for (int32 InstanceIdx = 0; InstanceIdx < Context.NumInstances; ++InstanceIdx)
+	for (int32 InstanceIdx = 0; InstanceIdx < Context.GetNumInstances(); ++InstanceIdx)
 	{
 		OutCellSize.SetAndAdvance(InstData->CellSize);
 	}
 }
 
-void UNiagaraDataInterfaceGrid2DCollection::GetNumCells(FVectorVMContext& Context)
+void UNiagaraDataInterfaceGrid2DCollection::GetNumCells(FVectorVMExternalFunctionContext& Context)
 {
 	VectorVM::FUserPtrHandler<FGrid2DCollectionRWInstanceData_GameThread> InstData(Context);
 	FNDIOutputParam<int> OutNumCellsX(Context);
 	FNDIOutputParam<int> OutNumCellsY(Context);
 
-	for (int32 InstanceIdx = 0; InstanceIdx < Context.NumInstances; ++InstanceIdx)
+	for (int32 InstanceIdx = 0; InstanceIdx < Context.GetNumInstances(); ++InstanceIdx)
 	{
 		OutNumCellsX.SetAndAdvance(InstData->NumCells.X);
 		OutNumCellsY.SetAndAdvance(InstData->NumCells.Y);
 	}
 }
 
-void UNiagaraDataInterfaceGrid2DCollection::SetNumCells(FVectorVMContext& Context)
+void UNiagaraDataInterfaceGrid2DCollection::SetNumCells(FVectorVMExternalFunctionContext& Context)
 {
 	// This should only be called from a system or emitter script due to a need for only setting up initially.
 	VectorVM::FUserPtrHandler<FGrid2DCollectionRWInstanceData_GameThread> InstData(Context);
@@ -2335,11 +2335,11 @@ void UNiagaraDataInterfaceGrid2DCollection::SetNumCells(FVectorVMContext& Contex
 	VectorVM::FExternalFuncInputHandler<int> InNumCellsY(Context);
 	VectorVM::FExternalFuncRegisterHandler<FNiagaraBool> OutSuccess(Context);
 
-	for (int32 InstanceIdx = 0; InstanceIdx < Context.NumInstances; ++InstanceIdx)
+	for (int32 InstanceIdx = 0; InstanceIdx < Context.GetNumInstances(); ++InstanceIdx)
 	{
 		int NewNumCellsX = InNumCellsX.GetAndAdvance();
 		int NewNumCellsY = InNumCellsY.GetAndAdvance();
-		bool bSuccess = (InstData.Get() != nullptr && Context.NumInstances == 1 && NumCellsX >= 0 && NumCellsY >= 0);
+		bool bSuccess = (InstData.Get() != nullptr && Context.GetNumInstances() == 1 && NumCellsX >= 0 && NumCellsY >= 0);
 		*OutSuccess.GetDestAndAdvance() = bSuccess;
 		if (bSuccess)
 		{
@@ -2413,7 +2413,7 @@ bool UNiagaraDataInterfaceGrid2DCollection::PerInstanceTickPostSimulate(void* Pe
 }
 
 
-void UNiagaraDataInterfaceGrid2DCollection::GetAttributeIndex(FVectorVMContext& Context, const FName& InName, int32 NumChannels)
+void UNiagaraDataInterfaceGrid2DCollection::GetAttributeIndex(FVectorVMExternalFunctionContext& Context, const FName& InName, int32 NumChannels)
 {
 	VectorVM::FUserPtrHandler<FGrid2DCollectionRWInstanceData_GameThread> InstData(Context);
 	VectorVM::FExternalFuncRegisterHandler<int> OutIndex(Context);
@@ -2421,7 +2421,7 @@ void UNiagaraDataInterfaceGrid2DCollection::GetAttributeIndex(FVectorVMContext&
 	if (InstData.Get())
 		Index = InstData.Get()->FindAttributeIndexByName(InName, NumChannels);
 
-	for (int32 InstanceIdx = 0; InstanceIdx < Context.NumInstances; ++InstanceIdx)
+	for (int32 InstanceIdx = 0; InstanceIdx < Context.GetNumInstances(); ++InstanceIdx)
 	{
 		*OutIndex.GetDestAndAdvance() = Index;
 	}
diff --git a/Engine/Plugins/FX/Niagara/Source/Niagara/Private/NiagaraDataInterfaceGrid3DCollection.cpp b/Engine/Plugins/FX/Niagara/Source/Niagara/Private/NiagaraDataInterfaceGrid3DCollection.cpp
--- a/Engine/Plugins/FX/Niagara/Source/Niagara/Private/NiagaraDataInterfaceGrid3DCollection.cpp
+++ b/Engine/Plugins/FX/Niagara/Source/Niagara/Private/NiagaraDataInterfaceGrid3DCollection.cpp
@@ -1026,7 +1026,7 @@ void UNiagaraDataInterfaceGrid3DCollection::GetVMExternalFunction(const FVMExter
 	}
 }
 
-void UNiagaraDataInterfaceGrid3DCollection::GetNumCells(FVectorVMContext& Context)
+void UNiagaraDataInterfaceGrid3DCollection::GetNumCells(FVectorVMExternalFunctionContext& Context)
 {
 	VectorVM::FUserPtrHandler<FGrid3DCollectionRWInstanceData_GameThread> InstData(Context);
 
@@ -1034,7 +1034,7 @@ void UNiagaraDataInterfaceGrid3DCollection::GetNumCells(FVectorVMContext& Contex
 	FNDIOutputParam<int32> NumCellsY(Context);
 	FNDIOutputParam<int32> NumCellsZ(Context);
 
-	for (int32 InstanceIdx = 0; InstanceIdx < Context.NumInstances; ++InstanceIdx)
+	for (int32 InstanceIdx = 0; InstanceIdx < Context.GetNumInstances(); ++InstanceIdx)
 	{
 		NumCellsX.SetAndAdvance(InstData->NumCells.X);
 		NumCellsY.SetAndAdvance(InstData->NumCells.Y);
@@ -2316,14 +2316,14 @@ void UNiagaraDataInterfaceGrid3DCollection::GetTextureSize(const UNiagaraCompone
 	SizeZ = Grid3DInstanceData->NumCells.Z;
 }
 
-void UNiagaraDataInterfaceGrid3DCollection::GetWorldBBoxSize(FVectorVMContext& Context)
+void UNiagaraDataInterfaceGrid3DCollection::GetWorldBBoxSize(FVectorVMExternalFunctionContext& Context)
 {
 	VectorVM::FUserPtrHandler<FGrid3DCollectionRWInstanceData_GameThread> InstData(Context);
 	VectorVM::FExternalFuncRegisterHandler<float> OutWorldBoundsX(Context);
 	VectorVM::FExternalFuncRegisterHandler<float> OutWorldBoundsY(Context);
 	VectorVM::FExternalFuncRegisterHandler<float> OutWorldBoundsZ(Context);
 
-	for (int32 InstanceIdx = 0; InstanceIdx < Context.NumInstances; ++InstanceIdx)
+	for (int32 InstanceIdx = 0; InstanceIdx < Context.GetNumInstances(); ++InstanceIdx)
 	{
 		*OutWorldBoundsX.GetDestAndAdvance() = InstData->WorldBBoxSize.X;
 		*OutWorldBoundsY.GetDestAndAdvance() = InstData->WorldBBoxSize.Y;
@@ -2331,7 +2331,7 @@ void UNiagaraDataInterfaceGrid3DCollection::GetWorldBBoxSize(FVectorVMContext& C
 	}
 }
 
-void UNiagaraDataInterfaceGrid3DCollection::SetNumCells(FVectorVMContext& Context)
+void UNiagaraDataInterfaceGrid3DCollection::SetNumCells(FVectorVMExternalFunctionContext& Context)
 {
 	// This should only be called from a system or emitter script due to a need for only setting up initially.
 	VectorVM::FUserPtrHandler<FGrid3DCollectionRWInstanceData_GameThread> InstData(Context);
@@ -2340,13 +2340,13 @@ void UNiagaraDataInterfaceGrid3DCollection::SetNumCells(FVectorVMContext& Contex
 	VectorVM::FExternalFuncInputHandler<int> InNumCellsZ(Context);
 	VectorVM::FExternalFuncRegisterHandler<FNiagaraBool> OutSuccess(Context);
 
-	for (int32 InstanceIdx = 0; InstanceIdx < Context.NumInstances; ++InstanceIdx)
+	for (int32 InstanceIdx = 0; InstanceIdx < Context.GetNumInstances(); ++InstanceIdx)
 	{
 		int NewNumCellsX = InNumCellsX.GetAndAdvance();
 		int NewNumCellsY = InNumCellsY.GetAndAdvance();
 		int NewNumCellsZ = InNumCellsZ.GetAndAdvance();
 
-		bool bSuccess = (InstData.Get() != nullptr && Context.NumInstances == 1 && NumCells.X >= 0 && NumCells.Y >= 0 && NumCells.Z >= 0);
+		bool bSuccess = (InstData.Get() != nullptr && Context.GetNumInstances() == 1 && NumCells.X >= 0 && NumCells.Y >= 0 && NumCells.Z >= 0);
 		*OutSuccess.GetDestAndAdvance() = bSuccess;
 		if (bSuccess)
 		{
@@ -2435,14 +2435,14 @@ bool UNiagaraDataInterfaceGrid3DCollection::PerInstanceTickPostSimulate(void* Pe
 	return false;
 }
 
-void UNiagaraDataInterfaceGrid3DCollection::GetCellSize(FVectorVMContext& Context)
+void UNiagaraDataInterfaceGrid3DCollection::GetCellSize(FVectorVMExternalFunctionContext& Context)
 {
 	VectorVM::FUserPtrHandler<FGrid3DCollectionRWInstanceData_GameThread> InstData(Context);
 	VectorVM::FExternalFuncRegisterHandler<float> OutCellSizeX(Context);
 	VectorVM::FExternalFuncRegisterHandler<float> OutCellSizeY(Context);
 	VectorVM::FExternalFuncRegisterHandler<float> OutCellSizeZ(Context);
 
-	for (int32 InstanceIdx = 0; InstanceIdx < Context.NumInstances; ++InstanceIdx)
+	for (int32 InstanceIdx = 0; InstanceIdx < Context.GetNumInstances(); ++InstanceIdx)
 	{
 		*OutCellSizeX.GetDestAndAdvance() = InstData->CellSize.X;
 		*OutCellSizeY.GetDestAndAdvance() = InstData->CellSize.Y;
@@ -2450,7 +2450,7 @@ void UNiagaraDataInterfaceGrid3DCollection::GetCellSize(FVectorVMContext& Contex
 	}
 }
 
-void UNiagaraDataInterfaceGrid3DCollection::GetAttributeIndex(FVectorVMContext& Context, const FName& InName, int32 NumChannels)
+void UNiagaraDataInterfaceGrid3DCollection::GetAttributeIndex(FVectorVMExternalFunctionContext& Context, const FName& InName, int32 NumChannels)
 {
 	VectorVM::FUserPtrHandler<FGrid3DCollectionRWInstanceData_GameThread> InstData(Context);
 	VectorVM::FExternalFuncRegisterHandler<int> OutIndex(Context);
@@ -2458,7 +2458,7 @@ void UNiagaraDataInterfaceGrid3DCollection::GetAttributeIndex(FVectorVMContext&
 	if (InstData.Get())
 		Index = InstData.Get()->FindAttributeIndexByName(InName, NumChannels);
 
-	for (int32 InstanceIdx = 0; InstanceIdx < Context.NumInstances; ++InstanceIdx)
+	for (int32 InstanceIdx = 0; InstanceIdx < Context.GetNumInstances(); ++InstanceIdx)
 	{
 		*OutIndex.GetDestAndAdvance() = Index;
 	}
diff --git a/Engine/Plugins/FX/Niagara/Source/Niagara/Private/NiagaraDataInterfaceIntRenderTarget2D.cpp b/Engine/Plugins/FX/Niagara/Source/Niagara/Private/NiagaraDataInterfaceIntRenderTarget2D.cpp
--- a/Engine/Plugins/FX/Niagara/Source/Niagara/Private/NiagaraDataInterfaceIntRenderTarget2D.cpp
+++ b/Engine/Plugins/FX/Niagara/Source/Niagara/Private/NiagaraDataInterfaceIntRenderTarget2D.cpp
@@ -433,11 +433,11 @@ void UNiagaraDataInterfaceIntRenderTarget2D::GetVMExternalFunction(const FVMExte
 	Super::GetVMExternalFunction(BindingInfo, InstanceData, OutFunc);
 	if (BindingInfo.Name == NDIIntRenderTarget2DLocal::GetSizeFunctionName)
 	{
-		OutFunc = FVMExternalFunction::CreateLambda([this](FVectorVMContext& Context) { VMGetSize(Context); });
+		OutFunc = FVMExternalFunction::CreateLambda([this](FVectorVMExternalFunctionContext& Context) { VMGetSize(Context); });
 	}
 	else if (BindingInfo.Name == NDIIntRenderTarget2DLocal::SetSizeFunctionName)
 	{
-		OutFunc = FVMExternalFunction::CreateLambda([this](FVectorVMContext& Context) { VMSetSize(Context); });
+		OutFunc = FVMExternalFunction::CreateLambda([this](FVectorVMExternalFunctionContext& Context) { VMSetSize(Context); });
 	}
 }
 
@@ -726,20 +726,20 @@ bool UNiagaraDataInterfaceIntRenderTarget2D::UpdateInstanceTexture(FNiagaraSyste
 	return bHasChanged;
 }
 
-void UNiagaraDataInterfaceIntRenderTarget2D::VMGetSize(FVectorVMContext& Context)
+void UNiagaraDataInterfaceIntRenderTarget2D::VMGetSize(FVectorVMExternalFunctionContext& Context)
 {
 	VectorVM::FUserPtrHandler<FNDIIntRenderTarget2DInstanceData_GameThread> InstData(Context);
 	FNDIOutputParam<int32> OutSizeX(Context);
 	FNDIOutputParam<int32> OutSizeY(Context);
 	
-	for (int32 i=0; i < Context.NumInstances; ++i)
+	for (int32 i=0; i < Context.GetNumInstances(); ++i)
 	{
 		OutSizeX.SetAndAdvance(InstData->Size.X);
 		OutSizeY.SetAndAdvance(InstData->Size.Y);
 	}
 }
 
-void UNiagaraDataInterfaceIntRenderTarget2D::VMSetSize(FVectorVMContext& Context)
+void UNiagaraDataInterfaceIntRenderTarget2D::VMSetSize(FVectorVMExternalFunctionContext& Context)
 {
 	VectorVM::FUserPtrHandler<FNDIIntRenderTarget2DInstanceData_GameThread> InstData(Context);
 	FNDIInputParam<int32> InSizeX(Context);
@@ -747,11 +747,11 @@ void UNiagaraDataInterfaceIntRenderTarget2D::VMSetSize(FVectorVMContext& Context
 	FNDIOutputParam<FNiagaraBool> OutSuccess(Context);
 
 	extern float GNiagaraRenderTargetResolutionMultiplier;
-	for (int32 i=0; i < Context.NumInstances; ++i)
+	for (int32 i=0; i < Context.GetNumInstances(); ++i)
 	{
 		const int SizeX = InSizeX.GetAndAdvance();
 		const int SizeY = InSizeY.GetAndAdvance();
-		const bool bSuccess = (InstData.Get() != nullptr && Context.NumInstances == 1 && SizeX >= 0 && SizeY >= 0);
+		const bool bSuccess = (InstData.Get() != nullptr && Context.GetNumInstances() == 1 && SizeX >= 0 && SizeY >= 0);
 		OutSuccess.SetAndAdvance(bSuccess);
 		if (bSuccess)
 		{
diff --git a/Engine/Plugins/FX/Niagara/Source/Niagara/Private/NiagaraDataInterfaceMeshRendererInfo.cpp b/Engine/Plugins/FX/Niagara/Source/Niagara/Private/NiagaraDataInterfaceMeshRendererInfo.cpp
--- a/Engine/Plugins/FX/Niagara/Source/Niagara/Private/NiagaraDataInterfaceMeshRendererInfo.cpp
+++ b/Engine/Plugins/FX/Niagara/Source/Niagara/Private/NiagaraDataInterfaceMeshRendererInfo.cpp
@@ -589,7 +589,7 @@ void UNiagaraDataInterfaceMeshRendererInfo::PushToRenderThreadImpl()
 	}
 }
 
-void UNiagaraDataInterfaceMeshRendererInfo::GetNumMeshes(FVectorVMContext& Context)
+void UNiagaraDataInterfaceMeshRendererInfo::GetNumMeshes(FVectorVMExternalFunctionContext& Context)
 {
 	FNDIOutputParam<int32> OutNum(Context);
 
@@ -600,7 +600,7 @@ void UNiagaraDataInterfaceMeshRendererInfo::GetNumMeshes(FVectorVMContext& Conte
 	}
 }
 
-void UNiagaraDataInterfaceMeshRendererInfo::GetMeshLocalBounds(FVectorVMContext& Context)
+void UNiagaraDataInterfaceMeshRendererInfo::GetMeshLocalBounds(FVectorVMExternalFunctionContext& Context)
 {
 	FNDIInputParam<int32> InMeshIdx(Context);
 	FNDIOutputParam<FVector3f> OutMinBounds(Context);
diff --git a/Engine/Plugins/FX/Niagara/Source/Niagara/Private/NiagaraDataInterfaceNeighborGrid3D.cpp b/Engine/Plugins/FX/Niagara/Source/Niagara/Private/NiagaraDataInterfaceNeighborGrid3D.cpp
--- a/Engine/Plugins/FX/Niagara/Source/Niagara/Private/NiagaraDataInterfaceNeighborGrid3D.cpp
+++ b/Engine/Plugins/FX/Niagara/Source/Niagara/Private/NiagaraDataInterfaceNeighborGrid3D.cpp
@@ -324,17 +324,17 @@ void UNiagaraDataInterfaceNeighborGrid3D::GetVMExternalFunction(const FVMExterna
 	if (BindingInfo.Name == UNiagaraDataInterfaceRWBase::WorldBBoxSizeFunctionName)
 	{
 		check(BindingInfo.GetNumInputs() == 1 && BindingInfo.GetNumOutputs() == 3);
-		OutFunc = FVMExternalFunction::CreateLambda([&](FVectorVMContext& Context) { GetWorldBBoxSize(Context); });
+		OutFunc = FVMExternalFunction::CreateLambda([&](FVectorVMExternalFunctionContext& Context) { GetWorldBBoxSize(Context); });
 	}
 	else if (BindingInfo.Name == UNiagaraDataInterfaceRWBase::NumCellsFunctionName)
 	{
 		check(BindingInfo.GetNumInputs() == 1 && BindingInfo.GetNumOutputs() == 3);
-		OutFunc = FVMExternalFunction::CreateLambda([&](FVectorVMContext& Context) { GetNumCells(Context); });
+		OutFunc = FVMExternalFunction::CreateLambda([&](FVectorVMExternalFunctionContext& Context) { GetNumCells(Context); });
 	}
 	else if (BindingInfo.Name == MaxNeighborsPerCellFunctionName)
 	{
 		check(BindingInfo.GetNumInputs() == 1 && BindingInfo.GetNumOutputs() == 1);
-		OutFunc = FVMExternalFunction::CreateLambda([&](FVectorVMContext& Context) { GetMaxNeighborsPerCell(Context); });
+		OutFunc = FVMExternalFunction::CreateLambda([&](FVectorVMExternalFunctionContext& Context) { GetMaxNeighborsPerCell(Context); });
 	}
 	else if (BindingInfo.Name == SetNumCellsFunctionName)
 	{
@@ -348,19 +348,19 @@ void UNiagaraDataInterfaceNeighborGrid3D::GetVMExternalFunction(const FVMExterna
 	//else if (BindingInfo.Name == SetParticleNeighborCountFunctionName) { OutFunc = FVMExternalFunction::CreateUObject(this, &UNiagaraDataInterfaceRWBase::EmptyVMFunction); }
 }
 
-void UNiagaraDataInterfaceNeighborGrid3D::GetWorldBBoxSize(FVectorVMContext& Context)
+void UNiagaraDataInterfaceNeighborGrid3D::GetWorldBBoxSize(FVectorVMExternalFunctionContext& Context)
 {
 	VectorVM::FUserPtrHandler<NeighborGrid3DRWInstanceData> InstData(Context);
 
 	FNDIOutputParam<FVector3f> OutWorldBounds(Context);
 
-	for (int32 InstanceIdx = 0; InstanceIdx < Context.NumInstances; ++InstanceIdx)
+	for (int32 InstanceIdx = 0; InstanceIdx < Context.GetNumInstances(); ++InstanceIdx)
 	{
 		OutWorldBounds.SetAndAdvance(WorldBBoxSize);
 	}
 }
 
-void UNiagaraDataInterfaceNeighborGrid3D::GetNumCells(FVectorVMContext& Context)
+void UNiagaraDataInterfaceNeighborGrid3D::GetNumCells(FVectorVMExternalFunctionContext& Context)
 {
 	VectorVM::FUserPtrHandler<NeighborGrid3DRWInstanceData> InstData(Context);
 
@@ -368,7 +368,7 @@ void UNiagaraDataInterfaceNeighborGrid3D::GetNumCells(FVectorVMContext& Context)
 	FNDIOutputParam<int32> NumCellsY(Context);
 	FNDIOutputParam<int32> NumCellsZ(Context);
 
-	for (int32 InstanceIdx = 0; InstanceIdx < Context.NumInstances; ++InstanceIdx)
+	for (int32 InstanceIdx = 0; InstanceIdx < Context.GetNumInstances(); ++InstanceIdx)
 	{
 		NumCellsX.SetAndAdvance(NumCells.X);
 		NumCellsY.SetAndAdvance(NumCells.Y);
@@ -376,13 +376,13 @@ void UNiagaraDataInterfaceNeighborGrid3D::GetNumCells(FVectorVMContext& Context)
 	}
 }
 
-void UNiagaraDataInterfaceNeighborGrid3D::GetMaxNeighborsPerCell(FVectorVMContext& Context)
+void UNiagaraDataInterfaceNeighborGrid3D::GetMaxNeighborsPerCell(FVectorVMExternalFunctionContext& Context)
 {
 	VectorVM::FUserPtrHandler<NeighborGrid3DRWInstanceData> InstData(Context);
 
 	FNDIOutputParam<int32> OutMaxNeighborsPerCell(Context);
 
-	for (int32 InstanceIdx = 0; InstanceIdx < Context.NumInstances; ++InstanceIdx)
+	for (int32 InstanceIdx = 0; InstanceIdx < Context.GetNumInstances(); ++InstanceIdx)
 	{
 		OutMaxNeighborsPerCell.SetAndAdvance(InstData->MaxNeighborsPerCell);
 	}
@@ -647,7 +647,7 @@ bool UNiagaraDataInterfaceNeighborGrid3D::InitPerInstanceData(void* PerInstanceD
 	return true;
 }
 
-void UNiagaraDataInterfaceNeighborGrid3D::SetNumCells(FVectorVMContext& Context)
+void UNiagaraDataInterfaceNeighborGrid3D::SetNumCells(FVectorVMExternalFunctionContext& Context)
 {
 	// This should only be called from a system or emitter script due to a need for only setting up initially.
 	VectorVM::FUserPtrHandler<NeighborGrid3DRWInstanceData> InstData(Context);
@@ -657,13 +657,13 @@ void UNiagaraDataInterfaceNeighborGrid3D::SetNumCells(FVectorVMContext& Context)
 	VectorVM::FExternalFuncInputHandler<int> InMaxNeighborsPerCell(Context);
 	VectorVM::FExternalFuncRegisterHandler<FNiagaraBool> OutSuccess(Context);
 
-	for (int32 InstanceIdx = 0; InstanceIdx < Context.NumInstances; ++InstanceIdx)
+	for (int32 InstanceIdx = 0; InstanceIdx < Context.GetNumInstances(); ++InstanceIdx)
 	{
 		int NewNumCellsX = InNumCellsX.GetAndAdvance();
 		int NewNumCellsY = InNumCellsY.GetAndAdvance();
 		int NewNumCellsZ = InNumCellsZ.GetAndAdvance();
 		int NewMaxNeighborsPerCell = InMaxNeighborsPerCell.GetAndAdvance();
-		bool bSuccess = (InstData.Get() != nullptr && Context.NumInstances == 1 && NumCells.X >= 0 && NumCells.Y >= 0 && NumCells.Z >= 0 && MaxNeighborsPerCell >= 0);
+		bool bSuccess = (InstData.Get() != nullptr && Context.GetNumInstances() == 1 && NumCells.X >= 0 && NumCells.Y >= 0 && NumCells.Z >= 0 && MaxNeighborsPerCell >= 0);
 		*OutSuccess.GetDestAndAdvance() = bSuccess;
 		if (bSuccess)
 		{
diff --git a/Engine/Plugins/FX/Niagara/Source/Niagara/Private/NiagaraDataInterfaceOcclusion.cpp b/Engine/Plugins/FX/Niagara/Source/Niagara/Private/NiagaraDataInterfaceOcclusion.cpp
--- a/Engine/Plugins/FX/Niagara/Source/Niagara/Private/NiagaraDataInterfaceOcclusion.cpp
+++ b/Engine/Plugins/FX/Niagara/Source/Niagara/Private/NiagaraDataInterfaceOcclusion.cpp
@@ -135,7 +135,7 @@ void UNiagaraDataInterfaceOcclusion::GetVMExternalFunction(const FVMExternalFunc
 
 // ------- Dummy implementations for CPU execution ------------
 
-void UNiagaraDataInterfaceOcclusion::QueryOcclusionFactorGPU(FVectorVMContext& Context)
+void UNiagaraDataInterfaceOcclusion::QueryOcclusionFactorGPU(FVectorVMExternalFunctionContext& Context)
 {
 	VectorVM::FExternalFuncInputHandler<float> PosParamX(Context);
 	VectorVM::FExternalFuncInputHandler<float> PosParamY(Context);
@@ -147,7 +147,7 @@ void UNiagaraDataInterfaceOcclusion::QueryOcclusionFactorGPU(FVectorVMContext& C
 	VectorVM::FExternalFuncRegisterHandler<float> OutOcclusion(Context);
 	VectorVM::FExternalFuncRegisterHandler<float> OutFactor(Context);
 
-	for (int32 i = 0; i < Context.NumInstances; ++i)
+	for (int32 i = 0; i < Context.GetNumInstances(); ++i)
 	{
 		PosParamX.GetAndAdvance();
 		PosParamY.GetAndAdvance();
@@ -161,7 +161,7 @@ void UNiagaraDataInterfaceOcclusion::QueryOcclusionFactorGPU(FVectorVMContext& C
 	}
 }
 
-void UNiagaraDataInterfaceOcclusion::QueryOcclusionFactorCircleGPU(FVectorVMContext& Context)
+void UNiagaraDataInterfaceOcclusion::QueryOcclusionFactorCircleGPU(FVectorVMExternalFunctionContext& Context)
 {
 	VectorVM::FExternalFuncInputHandler<float> PosParamX(Context);
 	VectorVM::FExternalFuncInputHandler<float> PosParamY(Context);
@@ -173,7 +173,7 @@ void UNiagaraDataInterfaceOcclusion::QueryOcclusionFactorCircleGPU(FVectorVMCont
 	VectorVM::FExternalFuncRegisterHandler<float> OutOcclusion(Context);
 	VectorVM::FExternalFuncRegisterHandler<float> OutFactor(Context);
 
-	for (int32 i = 0; i < Context.NumInstances; ++i)
+	for (int32 i = 0; i < Context.GetNumInstances(); ++i)
 	{
 		PosParamX.GetAndAdvance();
 		PosParamY.GetAndAdvance();
diff --git a/Engine/Plugins/FX/Niagara/Source/Niagara/Private/NiagaraDataInterfaceParticleRead.cpp b/Engine/Plugins/FX/Niagara/Source/Niagara/Private/NiagaraDataInterfaceParticleRead.cpp
--- a/Engine/Plugins/FX/Niagara/Source/Niagara/Private/NiagaraDataInterfaceParticleRead.cpp
+++ b/Engine/Plugins/FX/Niagara/Source/Niagara/Private/NiagaraDataInterfaceParticleRead.cpp
@@ -1322,7 +1322,7 @@ void UNiagaraDataInterfaceParticleRead::GetVMExternalFunction(const FVMExternalF
 	}
 }
 
-void UNiagaraDataInterfaceParticleRead::GetNumSpawnedParticles(FVectorVMContext& Context)
+void UNiagaraDataInterfaceParticleRead::GetNumSpawnedParticles(FVectorVMExternalFunctionContext& Context)
 {
 	VectorVM::FUserPtrHandler<FNDIParticleRead_InstanceData> InstData(Context);
 	VectorVM::FExternalFuncRegisterHandler<int32> OutNumSpawned(Context);
@@ -1331,13 +1331,13 @@ void UNiagaraDataInterfaceParticleRead::GetNumSpawnedParticles(FVectorVMContext&
 	const FNiagaraDataBuffer* CurrentData = EmitterInstance ? EmitterInstance->GetData().GetCurrentData() : nullptr;
 	const int32 NumSpawned = CurrentData ? CurrentData->GetNumSpawnedInstances() : 0;
 
-	for (int32 InstanceIdx = 0; InstanceIdx < Context.NumInstances; ++InstanceIdx)
+	for (int32 InstanceIdx = 0; InstanceIdx < Context.GetNumInstances(); ++InstanceIdx)
 	{
 		*OutNumSpawned.GetDestAndAdvance() = NumSpawned;
 	}
 }
 
-void UNiagaraDataInterfaceParticleRead::GetSpawnedIDAtIndex(FVectorVMContext& Context)
+void UNiagaraDataInterfaceParticleRead::GetSpawnedIDAtIndex(FVectorVMExternalFunctionContext& Context)
 {
 	VectorVM::FUserPtrHandler<FNDIParticleRead_InstanceData> InstData(Context);
 	VectorVM::FExternalFuncInputHandler<int32> InIndex(Context);
@@ -1351,7 +1351,7 @@ void UNiagaraDataInterfaceParticleRead::GetSpawnedIDAtIndex(FVectorVMContext& Co
 	int32 NumSpawned = SpawnedIDsTable ? SpawnedIDsTable->Num() : 0;
 	int32 IDAcquireTag = EmitterInstance ? EmitterInstance->GetData().GetIDAcquireTag() : 0;
 
-	for (int32 InstanceIdx = 0; InstanceIdx < Context.NumInstances; ++InstanceIdx)
+	for (int32 InstanceIdx = 0; InstanceIdx < Context.GetNumInstances(); ++InstanceIdx)
 	{
 		FNiagaraBool ValidValue;
 		FNiagaraID IDValue;
@@ -1376,7 +1376,7 @@ void UNiagaraDataInterfaceParticleRead::GetSpawnedIDAtIndex(FVectorVMContext& Co
 	}
 }
 
-void UNiagaraDataInterfaceParticleRead::GetNumParticles(FVectorVMContext& Context)
+void UNiagaraDataInterfaceParticleRead::GetNumParticles(FVectorVMExternalFunctionContext& Context)
 {
 	VectorVM::FUserPtrHandler<FNDIParticleRead_InstanceData> InstData(Context);
 	VectorVM::FExternalFuncRegisterHandler<int32> OutNumParticles (Context);
@@ -1385,13 +1385,13 @@ void UNiagaraDataInterfaceParticleRead::GetNumParticles(FVectorVMContext& Contex
 	const FNiagaraDataBuffer* CurrentData = EmitterInstance ? EmitterInstance->GetData().GetCurrentData() : nullptr;
 	const int32 NumParticles = CurrentData ? CurrentData->GetNumInstances() : 0;
 
-	for (int32 InstanceIdx = 0; InstanceIdx < Context.NumInstances; ++InstanceIdx)
+	for (int32 InstanceIdx = 0; InstanceIdx < Context.GetNumInstances(); ++InstanceIdx)
 	{
 		*OutNumParticles.GetDestAndAdvance() = NumParticles;
 	}
 }
 
-void UNiagaraDataInterfaceParticleRead::GetParticleIndex(FVectorVMContext& Context)
+void UNiagaraDataInterfaceParticleRead::GetParticleIndex(FVectorVMExternalFunctionContext& Context)
 {
 	VectorVM::FUserPtrHandler<FNDIParticleRead_InstanceData> InstData(Context);
 	VectorVM::FExternalFuncInputHandler<int32> ParticleIDIndexParam(Context);
@@ -1403,7 +1403,7 @@ void UNiagaraDataInterfaceParticleRead::GetParticleIndex(FVectorVMContext& Conte
 
 	if (!CurrentData)
 	{
-		for (int32 InstanceIdx = 0; InstanceIdx < Context.NumInstances; ++InstanceIdx)
+		for (int32 InstanceIdx = 0; InstanceIdx < Context.GetNumInstances(); ++InstanceIdx)
 		{
 			*OutIndex.GetDestAndAdvance() = -1;
 		}
@@ -1413,7 +1413,7 @@ void UNiagaraDataInterfaceParticleRead::GetParticleIndex(FVectorVMContext& Conte
 	const auto IDData = FNiagaraDataSetAccessor<FNiagaraID>::CreateReader(EmitterInstance->GetData(), ParticleReadIDName);
 	const TArray<int32>& IDTable = CurrentData->GetIDTable();
 
-	for (int32 InstanceIdx = 0; InstanceIdx < Context.NumInstances; ++InstanceIdx)
+	for (int32 InstanceIdx = 0; InstanceIdx < Context.GetNumInstances(); ++InstanceIdx)
 	{
 		FNiagaraID ParticleID = { ParticleIDIndexParam.GetAndAdvance(), ParticleIDAcquireTagParam.GetAndAdvance() };
 
@@ -1435,7 +1435,7 @@ void UNiagaraDataInterfaceParticleRead::GetParticleIndex(FVectorVMContext& Conte
 	}
 }
 
-void UNiagaraDataInterfaceParticleRead::GetParticleIndexFromIDTable(FVectorVMContext& Context)
+void UNiagaraDataInterfaceParticleRead::GetParticleIndexFromIDTable(FVectorVMExternalFunctionContext& Context)
 {
 	VectorVM::FUserPtrHandler<FNDIParticleRead_InstanceData> InstData(Context);
 	FNDIInputParam<int32> InIDTableIndex(Context);
@@ -1447,7 +1447,7 @@ void UNiagaraDataInterfaceParticleRead::GetParticleIndexFromIDTable(FVectorVMCon
 
 	if ( !CurrentData )
 	{
-		for ( int32 i=0; i < Context.NumInstances; ++i )
+		for ( int32 i=0; i < Context.GetNumInstances(); ++i )
 		{
 			OutValid.SetAndAdvance(false);
 			OutParticleIndex.SetAndAdvance(INDEX_NONE);
@@ -1456,7 +1456,7 @@ void UNiagaraDataInterfaceParticleRead::GetParticleIndexFromIDTable(FVectorVMCon
 	}
 
 	const TArray<int32>& IDTable = CurrentData->GetIDTable();
-	for (int32 i = 0; i < Context.NumInstances; ++i)
+	for (int32 i = 0; i < Context.GetNumInstances(); ++i)
 	{
 		const int32 IDTableIndex = InIDTableIndex.GetAndAdvance();
 		if ( IDTable.IsValidIndex(IDTableIndex) )
@@ -1478,7 +1478,7 @@ void UNiagaraDataInterfaceParticleRead::GetParticleIndexFromIDTable(FVectorVMCon
 template<typename T>
 struct FDirectReadParamHandler
 {
-	FDirectReadParamHandler(FVectorVMContext& Context)
+	FDirectReadParamHandler(FVectorVMExternalFunctionContext& Context)
 		: InstanceData(Context)
 		, ParticleIDIndexParam(Context)
 		, ParticleIDAcquireTagParam(Context)
@@ -1502,7 +1502,7 @@ struct FDirectReadParamHandler
 };
 
 template <typename T>
-FORCEINLINE void ReadWithCheck(FVectorVMContext& Context, FName AttributeToRead, T Default)
+FORCEINLINE void ReadWithCheck(FVectorVMExternalFunctionContext& Context, FName AttributeToRead, T Default)
 {
 	FDirectReadParamHandler<T> Params(Context);
 	bool bWriteDummyData = true;
@@ -1523,7 +1523,7 @@ FORCEINLINE void ReadWithCheck(FVectorVMContext& Context, FName AttributeToRead,
 				{
 					bWriteDummyData = false;
 
-					for (int32 InstanceIdx = 0; InstanceIdx < Context.NumInstances; ++InstanceIdx)
+					for (int32 InstanceIdx = 0; InstanceIdx < Context.GetNumInstances(); ++InstanceIdx)
 					{
 						FNiagaraID ParticleID = Params.GetID();
 						bool bValid = false;
@@ -1554,7 +1554,7 @@ FORCEINLINE void ReadWithCheck(FVectorVMContext& Context, FName AttributeToRead,
 	// Do we need to write dummy data due to not being in a valid state?
 	if (bWriteDummyData)
 	{
-		for (int32 InstanceIdx = 0; InstanceIdx < Context.NumInstances; ++InstanceIdx)
+		for (int32 InstanceIdx = 0; InstanceIdx < Context.GetNumInstances(); ++InstanceIdx)
 		{
 			Params.SetValid(false);
 			Params.SetValue(Default);
@@ -1563,47 +1563,47 @@ FORCEINLINE void ReadWithCheck(FVectorVMContext& Context, FName AttributeToRead,
 }
 
 //TODO: We can remove all this boiler plate and bind directly to ReadWithCheck.
-void UNiagaraDataInterfaceParticleRead::ReadInt(FVectorVMContext& Context, FName AttributeToRead)
+void UNiagaraDataInterfaceParticleRead::ReadInt(FVectorVMExternalFunctionContext& Context, FName AttributeToRead)
 {
 	ReadWithCheck<int32>(Context, AttributeToRead, 0);
 }
 
-void UNiagaraDataInterfaceParticleRead::ReadBool(FVectorVMContext& Context, FName AttributeToRead)
+void UNiagaraDataInterfaceParticleRead::ReadBool(FVectorVMExternalFunctionContext& Context, FName AttributeToRead)
 {
 	ReadWithCheck<FNiagaraBool>(Context, AttributeToRead, FNiagaraBool(false));
 }
 
-void UNiagaraDataInterfaceParticleRead::ReadFloat(FVectorVMContext& Context, FName AttributeToRead)
+void UNiagaraDataInterfaceParticleRead::ReadFloat(FVectorVMExternalFunctionContext& Context, FName AttributeToRead)
 {
 	ReadWithCheck<float>(Context, AttributeToRead, 0.0f);
 }
 
-void UNiagaraDataInterfaceParticleRead::ReadVector2(FVectorVMContext& Context, FName AttributeToRead)
+void UNiagaraDataInterfaceParticleRead::ReadVector2(FVectorVMExternalFunctionContext& Context, FName AttributeToRead)
 {
 	ReadWithCheck<FVector2D>(Context, AttributeToRead, FVector2D::ZeroVector);
 }
 
-void UNiagaraDataInterfaceParticleRead::ReadVector3(FVectorVMContext& Context, FName AttributeToRead)
+void UNiagaraDataInterfaceParticleRead::ReadVector3(FVectorVMExternalFunctionContext& Context, FName AttributeToRead)
 {
 	ReadWithCheck<FVector3f>(Context, AttributeToRead, FVector3f::ZeroVector);
 }
 
-void UNiagaraDataInterfaceParticleRead::ReadVector4(FVectorVMContext& Context, FName AttributeToRead)
+void UNiagaraDataInterfaceParticleRead::ReadVector4(FVectorVMExternalFunctionContext& Context, FName AttributeToRead)
 {
 	ReadWithCheck<FVector4>(Context, AttributeToRead, FVector4(ForceInit));
 }
 
-void UNiagaraDataInterfaceParticleRead::ReadColor(FVectorVMContext& Context, FName AttributeToRead)
+void UNiagaraDataInterfaceParticleRead::ReadColor(FVectorVMExternalFunctionContext& Context, FName AttributeToRead)
 {
 	ReadWithCheck<FLinearColor>(Context, AttributeToRead, FLinearColor::White);
 }
 
-void UNiagaraDataInterfaceParticleRead::ReadQuat(FVectorVMContext& Context, FName AttributeToRead)
+void UNiagaraDataInterfaceParticleRead::ReadQuat(FVectorVMExternalFunctionContext& Context, FName AttributeToRead)
 {
 	ReadWithCheck<FQuat>(Context, AttributeToRead, FQuat::Identity);
 }
 
-void UNiagaraDataInterfaceParticleRead::ReadID(FVectorVMContext& Context, FName AttributeToRead)
+void UNiagaraDataInterfaceParticleRead::ReadID(FVectorVMExternalFunctionContext& Context, FName AttributeToRead)
 {
 	ReadWithCheck<FNiagaraID>(Context, AttributeToRead, FNiagaraID());
 }
@@ -1613,7 +1613,7 @@ void UNiagaraDataInterfaceParticleRead::ReadID(FVectorVMContext& Context, FName
 template<typename T>
 struct FDirectReadByIndexParamHandler
 {
-	FDirectReadByIndexParamHandler(FVectorVMContext& Context)
+	FDirectReadByIndexParamHandler(FVectorVMExternalFunctionContext& Context)
 		: InstanceData(Context)
 		, ParticleIndexParam(Context)
 		, OutValid(Context)
@@ -1635,7 +1635,7 @@ struct FDirectReadByIndexParamHandler
 };
 
 template<typename T>
-FORCEINLINE void ReadByIndexWithCheck(FVectorVMContext& Context, FName AttributeToRead, T Default)
+FORCEINLINE void ReadByIndexWithCheck(FVectorVMExternalFunctionContext& Context, FName AttributeToRead, T Default)
 {
 	FDirectReadByIndexParamHandler<T> Params(Context);
 
@@ -1652,7 +1652,7 @@ FORCEINLINE void ReadByIndexWithCheck(FVectorVMContext& Context, FName Attribute
 			{
 				bWriteDummyData = false;
 
-				for (int32 InstanceIdx = 0; InstanceIdx < Context.NumInstances; ++InstanceIdx)
+				for (int32 InstanceIdx = 0; InstanceIdx < Context.GetNumInstances(); ++InstanceIdx)
 				{
 					int32 ParticleIndex = Params.GetIndex();
 
@@ -1674,7 +1674,7 @@ FORCEINLINE void ReadByIndexWithCheck(FVectorVMContext& Context, FName Attribute
 	// Do we need to write dummy data due to not being in a valid state?
 	if (bWriteDummyData)
 	{
-		for (int32 InstanceIdx = 0; InstanceIdx < Context.NumInstances; ++InstanceIdx)
+		for (int32 InstanceIdx = 0; InstanceIdx < Context.GetNumInstances(); ++InstanceIdx)
 		{
 			Params.SetValid(false);
 			Params.SetValue(Default);
@@ -1683,47 +1683,47 @@ FORCEINLINE void ReadByIndexWithCheck(FVectorVMContext& Context, FName Attribute
 }
 
 //TODO: We can remove all this boiler plate and bind directly to ReadByIndexWithCheck.
-void UNiagaraDataInterfaceParticleRead::ReadIntByIndex(FVectorVMContext& Context, FName AttributeToRead)
+void UNiagaraDataInterfaceParticleRead::ReadIntByIndex(FVectorVMExternalFunctionContext& Context, FName AttributeToRead)
 {
 	ReadByIndexWithCheck<int32>(Context, AttributeToRead, 0);
 }
 
-void UNiagaraDataInterfaceParticleRead::ReadBoolByIndex(FVectorVMContext& Context, FName AttributeToRead)
+void UNiagaraDataInterfaceParticleRead::ReadBoolByIndex(FVectorVMExternalFunctionContext& Context, FName AttributeToRead)
 {
 	ReadByIndexWithCheck<FNiagaraBool>(Context, AttributeToRead, FNiagaraBool(false));
 }
 
-void UNiagaraDataInterfaceParticleRead::ReadFloatByIndex(FVectorVMContext& Context, FName AttributeToRead)
+void UNiagaraDataInterfaceParticleRead::ReadFloatByIndex(FVectorVMExternalFunctionContext& Context, FName AttributeToRead)
 {
 	ReadByIndexWithCheck<float>(Context, AttributeToRead, 0.0f);
 }
 
-void UNiagaraDataInterfaceParticleRead::ReadVector2ByIndex(FVectorVMContext& Context, FName AttributeToRead)
+void UNiagaraDataInterfaceParticleRead::ReadVector2ByIndex(FVectorVMExternalFunctionContext& Context, FName AttributeToRead)
 {
 	ReadByIndexWithCheck<FVector2D>(Context, AttributeToRead, FVector2D::ZeroVector);
 }
 
-void UNiagaraDataInterfaceParticleRead::ReadVector3ByIndex(FVectorVMContext& Context, FName AttributeToRead)
+void UNiagaraDataInterfaceParticleRead::ReadVector3ByIndex(FVectorVMExternalFunctionContext& Context, FName AttributeToRead)
 {
 	ReadByIndexWithCheck<FVector3f>(Context, AttributeToRead, FVector3f::ZeroVector);
 }
 
-void UNiagaraDataInterfaceParticleRead::ReadVector4ByIndex(FVectorVMContext& Context, FName AttributeToRead)
+void UNiagaraDataInterfaceParticleRead::ReadVector4ByIndex(FVectorVMExternalFunctionContext& Context, FName AttributeToRead)
 {
 	ReadByIndexWithCheck<FVector4>(Context, AttributeToRead, FVector4(ForceInit));
 }
 
-void UNiagaraDataInterfaceParticleRead::ReadColorByIndex(FVectorVMContext& Context, FName AttributeToRead)
+void UNiagaraDataInterfaceParticleRead::ReadColorByIndex(FVectorVMExternalFunctionContext& Context, FName AttributeToRead)
 {
 	ReadByIndexWithCheck<FLinearColor>(Context, AttributeToRead, FLinearColor::White);
 }
 
-void UNiagaraDataInterfaceParticleRead::ReadQuatByIndex(FVectorVMContext& Context, FName AttributeToRead)
+void UNiagaraDataInterfaceParticleRead::ReadQuatByIndex(FVectorVMExternalFunctionContext& Context, FName AttributeToRead)
 {
 	ReadByIndexWithCheck<FQuat>(Context, AttributeToRead, FQuat::Identity);
 }
 
-void UNiagaraDataInterfaceParticleRead::ReadIDByIndex(FVectorVMContext& Context, FName AttributeToRead)
+void UNiagaraDataInterfaceParticleRead::ReadIDByIndex(FVectorVMExternalFunctionContext& Context, FName AttributeToRead)
 {
 	ReadByIndexWithCheck<FNiagaraID>(Context, AttributeToRead, FNiagaraID());
 }
diff --git a/Engine/Plugins/FX/Niagara/Source/Niagara/Private/NiagaraDataInterfacePlatformSet.cpp b/Engine/Plugins/FX/Niagara/Source/Niagara/Private/NiagaraDataInterfacePlatformSet.cpp
--- a/Engine/Plugins/FX/Niagara/Source/Niagara/Private/NiagaraDataInterfacePlatformSet.cpp
+++ b/Engine/Plugins/FX/Niagara/Source/Niagara/Private/NiagaraDataInterfacePlatformSet.cpp
@@ -68,12 +68,12 @@ bool UNiagaraDataInterfacePlatformSet::CopyToInternal(UNiagaraDataInterface* Des
 	return true;
 }
 
-void UNiagaraDataInterfacePlatformSet::IsActive(FVectorVMContext& Context)
+void UNiagaraDataInterfacePlatformSet::IsActive(FVectorVMExternalFunctionContext& Context)
 {
 	VectorVM::FExternalFuncRegisterHandler<FNiagaraBool> OutValue(Context);
 
 	bool bIsActive = Platforms.IsActive();
-	for (int32 i = 0; i < Context.NumInstances; ++i)
+	for (int32 i = 0; i < Context.GetNumInstances(); ++i)
 	{
 		*OutValue.GetDestAndAdvance() = FNiagaraBool(bIsActive);
 	}
diff --git a/Engine/Plugins/FX/Niagara/Source/Niagara/Private/NiagaraDataInterfaceRasterizationGrid3D.cpp b/Engine/Plugins/FX/Niagara/Source/Niagara/Private/NiagaraDataInterfaceRasterizationGrid3D.cpp
--- a/Engine/Plugins/FX/Niagara/Source/Niagara/Private/NiagaraDataInterfaceRasterizationGrid3D.cpp
+++ b/Engine/Plugins/FX/Niagara/Source/Niagara/Private/NiagaraDataInterfaceRasterizationGrid3D.cpp
@@ -384,7 +384,7 @@ void UNiagaraDataInterfaceRasterizationGrid3D::GetVMExternalFunction(const FVMEx
 	if (BindingInfo.Name == UNiagaraDataInterfaceRWBase::NumCellsFunctionName)
 	{
 		check(BindingInfo.GetNumInputs() == 1 && BindingInfo.GetNumOutputs() == 3);
-		OutFunc = FVMExternalFunction::CreateLambda([&](FVectorVMContext& Context) { GetNumCells(Context); });
+		OutFunc = FVMExternalFunction::CreateLambda([&](FVectorVMExternalFunctionContext& Context) { GetNumCells(Context); });
 	}
 	else if (BindingInfo.Name == SetNumCellsFunctionName)
 	{
@@ -399,7 +399,7 @@ void UNiagaraDataInterfaceRasterizationGrid3D::GetVMExternalFunction(const FVMEx
 }
 
 
-void UNiagaraDataInterfaceRasterizationGrid3D::GetNumCells(FVectorVMContext& Context)
+void UNiagaraDataInterfaceRasterizationGrid3D::GetNumCells(FVectorVMExternalFunctionContext& Context)
 {
 	VectorVM::FUserPtrHandler<RasterizationGrid3DRWInstanceData> InstData(Context);
 
@@ -407,7 +407,7 @@ void UNiagaraDataInterfaceRasterizationGrid3D::GetNumCells(FVectorVMContext& Con
 	FNDIOutputParam<int32> NumCellsY(Context);
 	FNDIOutputParam<int32> NumCellsZ(Context);
 
-	for (int32 InstanceIdx = 0; InstanceIdx < Context.NumInstances; ++InstanceIdx)
+	for (int32 InstanceIdx = 0; InstanceIdx < Context.GetNumInstances(); ++InstanceIdx)
 	{
 		NumCellsX.SetAndAdvance(NumCells.X);
 		NumCellsY.SetAndAdvance(NumCells.Y);
@@ -645,7 +645,7 @@ bool UNiagaraDataInterfaceRasterizationGrid3D::InitPerInstanceData(void* PerInst
 	return true;
 }
 
-void UNiagaraDataInterfaceRasterizationGrid3D::SetNumCells(FVectorVMContext& Context)
+void UNiagaraDataInterfaceRasterizationGrid3D::SetNumCells(FVectorVMExternalFunctionContext& Context)
 {
 	// This should only be called from a system or emitter script due to a need for only setting up initially.
 	VectorVM::FUserPtrHandler<RasterizationGrid3DRWInstanceData> InstData(Context);
@@ -654,12 +654,12 @@ void UNiagaraDataInterfaceRasterizationGrid3D::SetNumCells(FVectorVMContext& Con
 	VectorVM::FExternalFuncInputHandler<int> InNumCellsZ(Context);	
 	VectorVM::FExternalFuncRegisterHandler<FNiagaraBool> OutSuccess(Context);
 
-	for (int32 InstanceIdx = 0; InstanceIdx < Context.NumInstances; ++InstanceIdx)
+	for (int32 InstanceIdx = 0; InstanceIdx < Context.GetNumInstances(); ++InstanceIdx)
 	{
 		int NewNumCellsX = InNumCellsX.GetAndAdvance();
 		int NewNumCellsY = InNumCellsY.GetAndAdvance();
 		int NewNumCellsZ = InNumCellsZ.GetAndAdvance();		
-		bool bSuccess = (InstData.Get() != nullptr && Context.NumInstances == 1 && NumCells.X >= 0 && NumCells.Y >= 0 && NumCells.Z >= 0);
+		bool bSuccess = (InstData.Get() != nullptr && Context.GetNumInstances() == 1 && NumCells.X >= 0 && NumCells.Y >= 0 && NumCells.Z >= 0);
 		*OutSuccess.GetDestAndAdvance() = bSuccess;
 		if (bSuccess)
 		{
@@ -674,16 +674,16 @@ void UNiagaraDataInterfaceRasterizationGrid3D::SetNumCells(FVectorVMContext& Con
 	}
 }
 
-void UNiagaraDataInterfaceRasterizationGrid3D::SetFloatResetValue(FVectorVMContext& Context)
+void UNiagaraDataInterfaceRasterizationGrid3D::SetFloatResetValue(FVectorVMExternalFunctionContext& Context)
 {
 	VectorVM::FUserPtrHandler<RasterizationGrid3DRWInstanceData> InstData(Context);
 	VectorVM::FExternalFuncInputHandler<float> InResetValue(Context);
 	VectorVM::FExternalFuncRegisterHandler<FNiagaraBool> OutSuccess(Context);
 
-	for (int32 InstanceIdx = 0; InstanceIdx < Context.NumInstances; ++InstanceIdx)
+	for (int32 InstanceIdx = 0; InstanceIdx < Context.GetNumInstances(); ++InstanceIdx)
 	{
 		float NewResetValue = InResetValue.GetAndAdvance();
-		bool bSuccess = InstData.Get() != nullptr && Context.NumInstances == 1;
+		bool bSuccess = InstData.Get() != nullptr && Context.GetNumInstances() == 1;
 		*OutSuccess.GetDestAndAdvance() = bSuccess;
 		if (bSuccess)
 		{
diff --git a/Engine/Plugins/FX/Niagara/Source/Niagara/Private/NiagaraDataInterfaceRenderTarget2D.cpp b/Engine/Plugins/FX/Niagara/Source/Niagara/Private/NiagaraDataInterfaceRenderTarget2D.cpp
--- a/Engine/Plugins/FX/Niagara/Source/Niagara/Private/NiagaraDataInterfaceRenderTarget2D.cpp
+++ b/Engine/Plugins/FX/Niagara/Source/Niagara/Private/NiagaraDataInterfaceRenderTarget2D.cpp
@@ -639,7 +639,7 @@ bool UNiagaraDataInterfaceRenderTarget2D::RenderVariableToCanvas(FNiagaraSystemI
 	return true;
 }
 
-void UNiagaraDataInterfaceRenderTarget2D::SetSize(FVectorVMContext& Context)
+void UNiagaraDataInterfaceRenderTarget2D::SetSize(FVectorVMExternalFunctionContext& Context)
 {
 	// This should only be called from a system or emitter script due to a need for only setting up initially.
 	VectorVM::FUserPtrHandler<FRenderTarget2DRWInstanceData_GameThread> InstData(Context);
@@ -648,11 +648,11 @@ void UNiagaraDataInterfaceRenderTarget2D::SetSize(FVectorVMContext& Context)
 	FNDIOutputParam<FNiagaraBool> OutSuccess(Context);
 
 	extern float GNiagaraRenderTargetResolutionMultiplier;
-	for (int32 InstanceIdx = 0; InstanceIdx < Context.NumInstances; ++InstanceIdx)
+	for (int32 InstanceIdx = 0; InstanceIdx < Context.GetNumInstances(); ++InstanceIdx)
 	{
 		const int SizeX = InSizeX.GetAndAdvance();
 		const int SizeY = InSizeY.GetAndAdvance();
-		const bool bSuccess = (InstData.Get() != nullptr && Context.NumInstances == 1 && SizeX >= 0 && SizeY >= 0);
+		const bool bSuccess = (InstData.Get() != nullptr && Context.GetNumInstances() == 1 && SizeX >= 0 && SizeY >= 0);
 		OutSuccess.SetAndAdvance(bSuccess);
 		if (bSuccess)
 		{
@@ -662,13 +662,13 @@ void UNiagaraDataInterfaceRenderTarget2D::SetSize(FVectorVMContext& Context)
 	}
 }
 
-void UNiagaraDataInterfaceRenderTarget2D::GetSize(FVectorVMContext& Context)
+void UNiagaraDataInterfaceRenderTarget2D::GetSize(FVectorVMExternalFunctionContext& Context)
 {
 	VectorVM::FUserPtrHandler<FRenderTarget2DRWInstanceData_GameThread> InstData(Context);
 	VectorVM::FExternalFuncRegisterHandler<int> OutSizeX(Context);
 	VectorVM::FExternalFuncRegisterHandler<int> OutSizeY(Context);
 
-	for (int32 InstanceIdx = 0; InstanceIdx < Context.NumInstances; ++InstanceIdx)
+	for (int32 InstanceIdx = 0; InstanceIdx < Context.GetNumInstances(); ++InstanceIdx)
 	{
 		*OutSizeX.GetDestAndAdvance() = InstData->Size.X;
 		*OutSizeY.GetDestAndAdvance() = InstData->Size.Y;
diff --git a/Engine/Plugins/FX/Niagara/Source/Niagara/Private/NiagaraDataInterfaceRenderTarget2DArray.cpp b/Engine/Plugins/FX/Niagara/Source/Niagara/Private/NiagaraDataInterfaceRenderTarget2DArray.cpp
--- a/Engine/Plugins/FX/Niagara/Source/Niagara/Private/NiagaraDataInterfaceRenderTarget2DArray.cpp
+++ b/Engine/Plugins/FX/Niagara/Source/Niagara/Private/NiagaraDataInterfaceRenderTarget2DArray.cpp
@@ -575,7 +575,7 @@ bool UNiagaraDataInterfaceRenderTarget2DArray::GetExposedVariableValue(const FNi
 	return false;
 }
 
-void UNiagaraDataInterfaceRenderTarget2DArray::SetSize(FVectorVMContext& Context)
+void UNiagaraDataInterfaceRenderTarget2DArray::SetSize(FVectorVMExternalFunctionContext& Context)
 {
 	// This should only be called from a system or emitter script due to a need for only setting up initially.
 	VectorVM::FUserPtrHandler<FRenderTarget2DArrayRWInstanceData_GameThread> InstData(Context);
@@ -585,12 +585,12 @@ void UNiagaraDataInterfaceRenderTarget2DArray::SetSize(FVectorVMContext& Context
 	FNDIOutputParam<FNiagaraBool> OutSuccess(Context);
 
 	extern float GNiagaraRenderTargetResolutionMultiplier;
-	for (int32 InstanceIdx = 0; InstanceIdx < Context.NumInstances; ++InstanceIdx)
+	for (int32 InstanceIdx = 0; InstanceIdx < Context.GetNumInstances(); ++InstanceIdx)
 	{
 		const int SizeX = InSizeX.GetAndAdvance();
 		const int SizeY = InSizeY.GetAndAdvance();
 		const int Slices = InSlices.GetAndAdvance();
-		const bool bSuccess = (InstData.Get() != nullptr && Context.NumInstances == 1 && SizeX > 0 && SizeY > 0 && Slices > 0);
+		const bool bSuccess = (InstData.Get() != nullptr && Context.GetNumInstances() == 1 && SizeX > 0 && SizeY > 0 && Slices > 0);
 		OutSuccess.SetAndAdvance(bSuccess);
 		if (bSuccess)
 		{
@@ -601,14 +601,14 @@ void UNiagaraDataInterfaceRenderTarget2DArray::SetSize(FVectorVMContext& Context
 	}
 }
 
-void UNiagaraDataInterfaceRenderTarget2DArray::GetSize(FVectorVMContext& Context)
+void UNiagaraDataInterfaceRenderTarget2DArray::GetSize(FVectorVMExternalFunctionContext& Context)
 {
 	VectorVM::FUserPtrHandler<FRenderTarget2DArrayRWInstanceData_GameThread> InstData(Context);
 	FNDIOutputParam<int> OutSizeX(Context);
 	FNDIOutputParam<int> OutSizeY(Context);
 	FNDIOutputParam<int> OutSlices(Context);
 
-	for (int32 InstanceIdx = 0; InstanceIdx < Context.NumInstances; ++InstanceIdx)
+	for (int32 InstanceIdx = 0; InstanceIdx < Context.GetNumInstances(); ++InstanceIdx)
 	{
 		OutSizeX.SetAndAdvance(InstData->Size.X);
 		OutSizeY.SetAndAdvance(InstData->Size.Y);
diff --git a/Engine/Plugins/FX/Niagara/Source/Niagara/Private/NiagaraDataInterfaceRenderTargetCube.cpp b/Engine/Plugins/FX/Niagara/Source/Niagara/Private/NiagaraDataInterfaceRenderTargetCube.cpp
--- a/Engine/Plugins/FX/Niagara/Source/Niagara/Private/NiagaraDataInterfaceRenderTargetCube.cpp
+++ b/Engine/Plugins/FX/Niagara/Source/Niagara/Private/NiagaraDataInterfaceRenderTargetCube.cpp
@@ -561,7 +561,7 @@ bool UNiagaraDataInterfaceRenderTargetCube::GetExposedVariableValue(const FNiaga
 	return false;
 }
 
-void UNiagaraDataInterfaceRenderTargetCube::SetSize(FVectorVMContext& Context)
+void UNiagaraDataInterfaceRenderTargetCube::SetSize(FVectorVMExternalFunctionContext& Context)
 {
 	// This should only be called from a system or emitter script due to a need for only setting up initially.
 	VectorVM::FUserPtrHandler<FRenderTargetCubeRWInstanceData_GameThread> InstData(Context);
@@ -569,10 +569,10 @@ void UNiagaraDataInterfaceRenderTargetCube::SetSize(FVectorVMContext& Context)
 	FNDIOutputParam<FNiagaraBool> OutSuccess(Context);
 
 	extern float GNiagaraRenderTargetResolutionMultiplier;
-	for (int32 InstanceIdx = 0; InstanceIdx < Context.NumInstances; ++InstanceIdx)
+	for (int32 InstanceIdx = 0; InstanceIdx < Context.GetNumInstances(); ++InstanceIdx)
 	{
 		const int NewSize = InSize.GetAndAdvance();
-		const bool bSuccess = (InstData.Get() != nullptr && Context.NumInstances == 1 && NewSize > 0);
+		const bool bSuccess = (InstData.Get() != nullptr && Context.GetNumInstances() == 1 && NewSize > 0);
 		OutSuccess.SetAndAdvance(bSuccess);
 		if (bSuccess)
 		{
@@ -581,12 +581,12 @@ void UNiagaraDataInterfaceRenderTargetCube::SetSize(FVectorVMContext& Context)
 	}
 }
 
-void UNiagaraDataInterfaceRenderTargetCube::GetSize(FVectorVMContext& Context)
+void UNiagaraDataInterfaceRenderTargetCube::GetSize(FVectorVMExternalFunctionContext& Context)
 {
 	VectorVM::FUserPtrHandler<FRenderTargetCubeRWInstanceData_GameThread> InstData(Context);
 	FNDIOutputParam<int> OutSize(Context);
 
-	for (int32 InstanceIdx = 0; InstanceIdx < Context.NumInstances; ++InstanceIdx)
+	for (int32 InstanceIdx = 0; InstanceIdx < Context.GetNumInstances(); ++InstanceIdx)
 	{
 		OutSize.SetAndAdvance(InstData->Size);
 	}
diff --git a/Engine/Plugins/FX/Niagara/Source/Niagara/Private/NiagaraDataInterfaceRenderTargetVolume.cpp b/Engine/Plugins/FX/Niagara/Source/Niagara/Private/NiagaraDataInterfaceRenderTargetVolume.cpp
--- a/Engine/Plugins/FX/Niagara/Source/Niagara/Private/NiagaraDataInterfaceRenderTargetVolume.cpp
+++ b/Engine/Plugins/FX/Niagara/Source/Niagara/Private/NiagaraDataInterfaceRenderTargetVolume.cpp
@@ -571,7 +571,7 @@ bool UNiagaraDataInterfaceRenderTargetVolume::GetExposedVariableValue(const FNia
 	return false;
 }
 
-void UNiagaraDataInterfaceRenderTargetVolume::SetSize(FVectorVMContext& Context)
+void UNiagaraDataInterfaceRenderTargetVolume::SetSize(FVectorVMExternalFunctionContext& Context)
 {
 	// This should only be called from a system or emitter script due to a need for only setting up initially.
 	VectorVM::FUserPtrHandler<FRenderTargetVolumeRWInstanceData_GameThread> InstData(Context);
@@ -581,12 +581,12 @@ void UNiagaraDataInterfaceRenderTargetVolume::SetSize(FVectorVMContext& Context)
 	FNDIOutputParam<FNiagaraBool> OutSuccess(Context);
 
 	extern float GNiagaraRenderTargetResolutionMultiplier;
-	for (int32 InstanceIdx = 0; InstanceIdx < Context.NumInstances; ++InstanceIdx)
+	for (int32 InstanceIdx = 0; InstanceIdx < Context.GetNumInstances(); ++InstanceIdx)
 	{
 		const int SizeX = InSizeX.GetAndAdvance();
 		const int SizeY = InSizeY.GetAndAdvance();
 		const int SizeZ = InSizeZ.GetAndAdvance();
-		const bool bSuccess = (InstData.Get() != nullptr && Context.NumInstances == 1 && SizeX > 0 && SizeY > 0 && SizeZ > 0);
+		const bool bSuccess = (InstData.Get() != nullptr && Context.GetNumInstances() == 1 && SizeX > 0 && SizeY > 0 && SizeZ > 0);
 		OutSuccess.SetAndAdvance(bSuccess);
 		if (bSuccess)
 		{
@@ -597,14 +597,14 @@ void UNiagaraDataInterfaceRenderTargetVolume::SetSize(FVectorVMContext& Context)
 	}
 }
 
-void UNiagaraDataInterfaceRenderTargetVolume::GetSize(FVectorVMContext& Context)
+void UNiagaraDataInterfaceRenderTargetVolume::GetSize(FVectorVMExternalFunctionContext& Context)
 {
 	VectorVM::FUserPtrHandler<FRenderTargetVolumeRWInstanceData_GameThread> InstData(Context);
 	FNDIOutputParam<int> OutSizeX(Context);
 	FNDIOutputParam<int> OutSizeY(Context);
 	FNDIOutputParam<int> OutSizeZ(Context);
 
-	for (int32 InstanceIdx = 0; InstanceIdx < Context.NumInstances; ++InstanceIdx)
+	for (int32 InstanceIdx = 0; InstanceIdx < Context.GetNumInstances(); ++InstanceIdx)
 	{
 		OutSizeX.SetAndAdvance(InstData->Size.X);
 		OutSizeY.SetAndAdvance(InstData->Size.Y);
diff --git a/Engine/Plugins/FX/Niagara/Source/Niagara/Private/NiagaraDataInterfaceSimpleCounter.cpp b/Engine/Plugins/FX/Niagara/Source/Niagara/Private/NiagaraDataInterfaceSimpleCounter.cpp
--- a/Engine/Plugins/FX/Niagara/Source/Niagara/Private/NiagaraDataInterfaceSimpleCounter.cpp
+++ b/Engine/Plugins/FX/Niagara/Source/Niagara/Private/NiagaraDataInterfaceSimpleCounter.cpp
@@ -99,12 +99,12 @@ void UNiagaraDataInterfaceSimpleCounter::GetVMExternalFunction(const FVMExternal
 	}
 }
 
-void UNiagaraDataInterfaceSimpleCounter::GetNextValue(FVectorVMContext& Context)
+void UNiagaraDataInterfaceSimpleCounter::GetNextValue(FVectorVMExternalFunctionContext& Context)
 {	
 	VectorVM::FUserPtrHandler<CounterInstanceData> InstanceData(Context);
 	VectorVM::FExternalFuncRegisterHandler<int32> OutValue(Context);
 
-	for (int32 i = 0; i < Context.NumInstances; ++i)
+	for (int32 i = 0; i < Context.GetNumInstances(); ++i)
 	{
 		*OutValue.GetDest() = FPlatformAtomics::InterlockedIncrement(&InstanceData->Counter);
 		OutValue.Advance();
diff --git a/Engine/Plugins/FX/Niagara/Source/Niagara/Private/NiagaraDataInterfaceSpline.cpp b/Engine/Plugins/FX/Niagara/Source/Niagara/Private/NiagaraDataInterfaceSpline.cpp
--- a/Engine/Plugins/FX/Niagara/Source/Niagara/Private/NiagaraDataInterfaceSpline.cpp
+++ b/Engine/Plugins/FX/Niagara/Source/Niagara/Private/NiagaraDataInterfaceSpline.cpp
@@ -625,7 +625,7 @@ float FNDISpline_InstanceData::FindInputKeyClosestToWorldLocation(const FVector&
 
 
 template<typename TransformHandlerType, typename SplineSampleType>
-void UNiagaraDataInterfaceSpline::SampleSplinePositionByUnitDistance(FVectorVMContext& Context)
+void UNiagaraDataInterfaceSpline::SampleSplinePositionByUnitDistance(FVectorVMExternalFunctionContext& Context)
 {
 	VectorVM::FUserPtrHandler<FNDISpline_InstanceData> InstData(Context);
 	TransformHandlerType TransformHandler;
@@ -637,7 +637,7 @@ void UNiagaraDataInterfaceSpline::SampleSplinePositionByUnitDistance(FVectorVMCo
 	if (InstData->IsValid())
 	{
 		const float SplineLength = InstData->GetSplineLength();
-		for (int32 i = 0; i < Context.NumInstances; ++i)
+		for (int32 i = 0; i < Context.GetNumInstances(); ++i)
 		{
 			float DistanceUnitDistance = SplineSampleParam.Get();
 
@@ -655,7 +655,7 @@ void UNiagaraDataInterfaceSpline::SampleSplinePositionByUnitDistance(FVectorVMCo
 	}
 	else
 	{
-		for (int32 i = 0; i < Context.NumInstances; ++i)
+		for (int32 i = 0; i < Context.GetNumInstances(); ++i)
 		{
 			float DistanceUnitDistance = SplineSampleParam.Get();
 
@@ -674,7 +674,7 @@ void UNiagaraDataInterfaceSpline::SampleSplinePositionByUnitDistance(FVectorVMCo
 }
 
 template<typename TransformHandlerType, typename SplineSampleType>
-void UNiagaraDataInterfaceSpline::SampleSplineRotationByUnitDistance(FVectorVMContext& Context)
+void UNiagaraDataInterfaceSpline::SampleSplineRotationByUnitDistance(FVectorVMExternalFunctionContext& Context)
 {
 	TransformHandlerType TransformHandler;
 	SplineSampleType SplineSampleParam(Context);
@@ -688,7 +688,7 @@ void UNiagaraDataInterfaceSpline::SampleSplineRotationByUnitDistance(FVectorVMCo
 	{
 		const FQuat TransformQuat = InstData->Transform.ToQuat();
 		const float SplineLength = InstData->GetSplineLength();
-		for (int32 i = 0; i < Context.NumInstances; ++i)
+		for (int32 i = 0; i < Context.GetNumInstances(); ++i)
 		{
 			const float DistanceUnitDistance = SplineSampleParam.GetAndAdvance();
 
@@ -703,7 +703,7 @@ void UNiagaraDataInterfaceSpline::SampleSplineRotationByUnitDistance(FVectorVMCo
 	}
 	else
 	{
-		for (int32 i = 0; i < Context.NumInstances; ++i)
+		for (int32 i = 0; i < Context.GetNumInstances(); ++i)
 		{
 			*OutQuatX.GetDestAndAdvance() = FQuat::Identity.X;
 			*OutQuatY.GetDestAndAdvance() = FQuat::Identity.Y;
@@ -714,7 +714,7 @@ void UNiagaraDataInterfaceSpline::SampleSplineRotationByUnitDistance(FVectorVMCo
 }
 
 template<typename TransformHandlerType, typename SplineSampleType>
-void UNiagaraDataInterfaceSpline::SampleSplineUpVectorByUnitDistance(FVectorVMContext& Context)
+void UNiagaraDataInterfaceSpline::SampleSplineUpVectorByUnitDistance(FVectorVMExternalFunctionContext& Context)
 {
 	VectorVM::FUserPtrHandler<FNDISpline_InstanceData> InstData(Context);
 	TransformHandlerType TransformHandler;
@@ -726,7 +726,7 @@ void UNiagaraDataInterfaceSpline::SampleSplineUpVectorByUnitDistance(FVectorVMCo
 	if (InstData->IsValid())
 	{
 		const float SplineLength = InstData->GetSplineLength();
-		for (int32 i = 0; i < Context.NumInstances; ++i)
+		for (int32 i = 0; i < Context.GetNumInstances(); ++i)
 		{
 			float DistanceUnitDistance = SplineSampleParam.Get();
 
@@ -744,7 +744,7 @@ void UNiagaraDataInterfaceSpline::SampleSplineUpVectorByUnitDistance(FVectorVMCo
 	}
 	else
 	{
-		for (int32 i = 0; i < Context.NumInstances; ++i)
+		for (int32 i = 0; i < Context.GetNumInstances(); ++i)
 		{
 			float DistanceUnitDistance = SplineSampleParam.Get();
 
@@ -763,7 +763,7 @@ void UNiagaraDataInterfaceSpline::SampleSplineUpVectorByUnitDistance(FVectorVMCo
 }
 
 template<typename TransformHandlerType, typename SplineSampleType>
-void UNiagaraDataInterfaceSpline::SampleSplineRightVectorByUnitDistance(FVectorVMContext& Context)
+void UNiagaraDataInterfaceSpline::SampleSplineRightVectorByUnitDistance(FVectorVMExternalFunctionContext& Context)
 {
 	VectorVM::FUserPtrHandler<FNDISpline_InstanceData> InstData(Context);
 	TransformHandlerType TransformHandler;
@@ -776,7 +776,7 @@ void UNiagaraDataInterfaceSpline::SampleSplineRightVectorByUnitDistance(FVectorV
 	if (InstData->IsValid())
 	{
 		const float SplineLength = InstData->GetSplineLength();
-		for (int32 i = 0; i < Context.NumInstances; ++i)
+		for (int32 i = 0; i < Context.GetNumInstances(); ++i)
 		{
 			float DistanceUnitDistance = SplineSampleParam.Get();
 
@@ -794,7 +794,7 @@ void UNiagaraDataInterfaceSpline::SampleSplineRightVectorByUnitDistance(FVectorV
 	}
 	else
 	{
-		for (int32 i = 0; i < Context.NumInstances; ++i)
+		for (int32 i = 0; i < Context.GetNumInstances(); ++i)
 		{
 			FVector Pos = FVector(-1.0f, 0.0f, 0.0f); 
 			TransformHandler.TransformVector(Pos, InstData->Transform);
@@ -811,7 +811,7 @@ void UNiagaraDataInterfaceSpline::SampleSplineRightVectorByUnitDistance(FVectorV
 }
 
 template<typename TransformHandlerType, typename SplineSampleType>
-void UNiagaraDataInterfaceSpline::SampleSplineTangentByUnitDistance(FVectorVMContext& Context)
+void UNiagaraDataInterfaceSpline::SampleSplineTangentByUnitDistance(FVectorVMExternalFunctionContext& Context)
 {
 	VectorVM::FUserPtrHandler<FNDISpline_InstanceData> InstData(Context);
 	TransformHandlerType TransformHandler;
@@ -823,7 +823,7 @@ void UNiagaraDataInterfaceSpline::SampleSplineTangentByUnitDistance(FVectorVMCon
 	if (InstData->IsValid())
 	{
 		const float SplineLength = InstData->GetSplineLength();
-		for (int32 i = 0; i < Context.NumInstances; ++i)
+		for (int32 i = 0; i < Context.GetNumInstances(); ++i)
 		{
 			float DistanceUnitDistance = SplineSampleParam.Get();
 
@@ -841,7 +841,7 @@ void UNiagaraDataInterfaceSpline::SampleSplineTangentByUnitDistance(FVectorVMCon
 	}
 	else
 	{
-		for (int32 i = 0; i < Context.NumInstances; ++i)
+		for (int32 i = 0; i < Context.GetNumInstances(); ++i)
 		{
 			float DistanceUnitDistance = SplineSampleParam.Get();
 
@@ -860,7 +860,7 @@ void UNiagaraDataInterfaceSpline::SampleSplineTangentByUnitDistance(FVectorVMCon
 }
 
 template<typename TransformHandlerType, typename SplineSampleType>
-void UNiagaraDataInterfaceSpline::SampleSplineDirectionByUnitDistance(FVectorVMContext& Context)
+void UNiagaraDataInterfaceSpline::SampleSplineDirectionByUnitDistance(FVectorVMExternalFunctionContext& Context)
 {
 	VectorVM::FUserPtrHandler<FNDISpline_InstanceData> InstData(Context);
 	TransformHandlerType TransformHandler;
@@ -872,7 +872,7 @@ void UNiagaraDataInterfaceSpline::SampleSplineDirectionByUnitDistance(FVectorVMC
 	if (InstData->IsValid())
 	{
 		const float SplineLength = InstData->GetSplineLength();
-		for (int32 i = 0; i < Context.NumInstances; ++i)
+		for (int32 i = 0; i < Context.GetNumInstances(); ++i)
 		{
 			float DistanceUnitDistance = SplineSampleParam.Get();
 
@@ -890,7 +890,7 @@ void UNiagaraDataInterfaceSpline::SampleSplineDirectionByUnitDistance(FVectorVMC
 	}
 	else
 	{
-		for (int32 i = 0; i < Context.NumInstances; ++i)
+		for (int32 i = 0; i < Context.GetNumInstances(); ++i)
 		{
 			float DistanceUnitDistance = SplineSampleParam.Get();
 
@@ -908,7 +908,7 @@ void UNiagaraDataInterfaceSpline::SampleSplineDirectionByUnitDistance(FVectorVMC
 	}
 }
 
-void UNiagaraDataInterfaceSpline::WriteTransform(const FMatrix& ToWrite, FVectorVMContext& Context)
+void UNiagaraDataInterfaceSpline::WriteTransform(const FMatrix& ToWrite, FVectorVMExternalFunctionContext& Context)
 {
 	VectorVM::FExternalFuncRegisterHandler<float> Out00(Context);
 	VectorVM::FExternalFuncRegisterHandler<float> Out01(Context);
@@ -927,7 +927,7 @@ void UNiagaraDataInterfaceSpline::WriteTransform(const FMatrix& ToWrite, FVector
 	VectorVM::FExternalFuncRegisterHandler<float> Out14(Context);
 	VectorVM::FExternalFuncRegisterHandler<float> Out15(Context);
 
-	for (int32 i = 0; i < Context.NumInstances; ++i)
+	for (int32 i = 0; i < Context.GetNumInstances(); ++i)
 	{
 		*Out00.GetDest() = ToWrite.M[0][0]; Out00.Advance();
 		*Out01.GetDest() = ToWrite.M[0][1]; Out01.Advance();
@@ -949,7 +949,7 @@ void UNiagaraDataInterfaceSpline::WriteTransform(const FMatrix& ToWrite, FVector
 }
 
 template<typename PosXType, typename PosYType, typename PosZType>
-void UNiagaraDataInterfaceSpline::FindClosestUnitDistanceFromPositionWS(FVectorVMContext& Context)
+void UNiagaraDataInterfaceSpline::FindClosestUnitDistanceFromPositionWS(FVectorVMExternalFunctionContext& Context)
 {
 	VectorVM::FUserPtrHandler<FNDISpline_InstanceData> InstData(Context);
 	PosXType PosXParam(Context);
@@ -963,7 +963,7 @@ void UNiagaraDataInterfaceSpline::FindClosestUnitDistanceFromPositionWS(FVectorV
 		const int32 NumPoints = InstData->GetSplinePointsPosition().Points.Num();
 		const float FinalKeyTime = InstData->GetSplinePointsPosition().Points[NumPoints - 1].InVal;
 
-		for (int32 i = 0; i < Context.NumInstances; ++i)
+		for (int32 i = 0; i < Context.GetNumInstances(); ++i)
 		{
 			float PosX = PosXParam.Get();
 			float PosY = PosYParam.Get();
@@ -986,7 +986,7 @@ void UNiagaraDataInterfaceSpline::FindClosestUnitDistanceFromPositionWS(FVectorV
 	}
 	else
 	{
-		for (int32 i = 0; i < Context.NumInstances; ++i)
+		for (int32 i = 0; i < Context.GetNumInstances(); ++i)
 		{
 			float PosX = PosXParam.Get();
 			float PosY = PosYParam.Get();
@@ -1002,13 +1002,13 @@ void UNiagaraDataInterfaceSpline::FindClosestUnitDistanceFromPositionWS(FVectorV
 	}
 }
 
-void UNiagaraDataInterfaceSpline::GetLocalToWorld(FVectorVMContext& Context)
+void UNiagaraDataInterfaceSpline::GetLocalToWorld(FVectorVMExternalFunctionContext& Context)
 {
 	VectorVM::FUserPtrHandler<FNDISpline_InstanceData> InstData(Context);
 	WriteTransform(InstData->Transform, Context);
 }
 
-void UNiagaraDataInterfaceSpline::GetLocalToWorldInverseTransposed(FVectorVMContext& Context)
+void UNiagaraDataInterfaceSpline::GetLocalToWorldInverseTransposed(FVectorVMExternalFunctionContext& Context)
 {
 	VectorVM::FUserPtrHandler<FNDISpline_InstanceData> InstData(Context);
 	WriteTransform(InstData->TransformInverseTransposed, Context);
diff --git a/Engine/Plugins/FX/Niagara/Source/Niagara/Private/NiagaraDataInterfaceStaticMesh.cpp b/Engine/Plugins/FX/Niagara/Source/Niagara/Private/NiagaraDataInterfaceStaticMesh.cpp
--- a/Engine/Plugins/FX/Niagara/Source/Niagara/Private/NiagaraDataInterfaceStaticMesh.cpp
+++ b/Engine/Plugins/FX/Niagara/Source/Niagara/Private/NiagaraDataInterfaceStaticMesh.cpp
@@ -1444,47 +1444,47 @@ void UNiagaraDataInterfaceStaticMesh::GetVMExternalFunction(const FVMExternalFun
 	else if (BindingInfo.Name == StaticMeshHelpers::GetSocketCountName)
 	{
 		check(BindingInfo.GetNumInputs() == 1 && BindingInfo.GetNumOutputs() == 1);
-		OutFunc = FVMExternalFunction::CreateLambda([this](FVectorVMContext& Context) { this->GetSocketCount(Context); });
+		OutFunc = FVMExternalFunction::CreateLambda([this](FVectorVMExternalFunctionContext& Context) { this->GetSocketCount(Context); });
 	}
 	else if (BindingInfo.Name == StaticMeshHelpers::GetFilteredSocketCountName)
 	{
 		check(BindingInfo.GetNumInputs() == 1 && BindingInfo.GetNumOutputs() == 1);
-		OutFunc = FVMExternalFunction::CreateLambda([this](FVectorVMContext& Context) { this->GetFilteredSocketCount(Context); });
+		OutFunc = FVMExternalFunction::CreateLambda([this](FVectorVMExternalFunctionContext& Context) { this->GetFilteredSocketCount(Context); });
 	}
 	else if (BindingInfo.Name == StaticMeshHelpers::GetUnfilteredSocketCountName)
 	{
 		check(BindingInfo.GetNumInputs() == 1 && BindingInfo.GetNumOutputs() == 1);
-		OutFunc = FVMExternalFunction::CreateLambda([this](FVectorVMContext& Context) { this->GetUnfilteredSocketCount(Context); });
+		OutFunc = FVMExternalFunction::CreateLambda([this](FVectorVMExternalFunctionContext& Context) { this->GetUnfilteredSocketCount(Context); });
 	}
 	else if (BindingInfo.Name == StaticMeshHelpers::GetSocketTransformName)
 	{
 		check(BindingInfo.GetNumInputs() == 2 && BindingInfo.GetNumOutputs() == 10);
-		OutFunc = FVMExternalFunction::CreateLambda([this](FVectorVMContext& Context) { this->GetSocketTransform<false>(Context); });
+		OutFunc = FVMExternalFunction::CreateLambda([this](FVectorVMExternalFunctionContext& Context) { this->GetSocketTransform<false>(Context); });
 	}
 	else if (BindingInfo.Name == StaticMeshHelpers::GetSocketTransformWSName)
 	{
 		check(BindingInfo.GetNumInputs() == 2 && BindingInfo.GetNumOutputs() == 10);
-		OutFunc = FVMExternalFunction::CreateLambda([this](FVectorVMContext& Context) { this->GetSocketTransform<true>(Context); });
+		OutFunc = FVMExternalFunction::CreateLambda([this](FVectorVMExternalFunctionContext& Context) { this->GetSocketTransform<true>(Context); });
 	}
 	else if (BindingInfo.Name == StaticMeshHelpers::GetFilteredSocketTransformName)
 	{
 		check(BindingInfo.GetNumInputs() == 2 && BindingInfo.GetNumOutputs() == 10);
-		OutFunc = FVMExternalFunction::CreateLambda([this](FVectorVMContext& Context) { this->GetFilteredSocketTransform<false>(Context); });
+		OutFunc = FVMExternalFunction::CreateLambda([this](FVectorVMExternalFunctionContext& Context) { this->GetFilteredSocketTransform<false>(Context); });
 	}
 	else if (BindingInfo.Name == StaticMeshHelpers::GetFilteredSocketTransformWSName)
 	{
 		check(BindingInfo.GetNumInputs() == 2 && BindingInfo.GetNumOutputs() == 10);
-		OutFunc = FVMExternalFunction::CreateLambda([this](FVectorVMContext& Context) { this->GetFilteredSocketTransform<true>(Context); });
+		OutFunc = FVMExternalFunction::CreateLambda([this](FVectorVMExternalFunctionContext& Context) { this->GetFilteredSocketTransform<true>(Context); });
 	}
 	else if (BindingInfo.Name == StaticMeshHelpers::GetUnfilteredSocketTransformName)
 	{
 		check(BindingInfo.GetNumInputs() == 2 && BindingInfo.GetNumOutputs() == 10);
-		OutFunc = FVMExternalFunction::CreateLambda([this](FVectorVMContext& Context) { this->GetUnfilteredSocketTransform<false>(Context); });
+		OutFunc = FVMExternalFunction::CreateLambda([this](FVectorVMExternalFunctionContext& Context) { this->GetUnfilteredSocketTransform<false>(Context); });
 	}
 	else if (BindingInfo.Name == StaticMeshHelpers::GetUnfilteredSocketTransformWSName)
 	{
 		check(BindingInfo.GetNumInputs() == 2 && BindingInfo.GetNumOutputs() == 10);
-		OutFunc = FVMExternalFunction::CreateLambda([this](FVectorVMContext& Context) { this->GetUnfilteredSocketTransform<true>(Context); });
+		OutFunc = FVMExternalFunction::CreateLambda([this](FVectorVMExternalFunctionContext& Context) { this->GetUnfilteredSocketTransform<true>(Context); });
 	}
 }
 
@@ -1768,14 +1768,14 @@ UStaticMesh* UNiagaraDataInterfaceStaticMesh::GetStaticMesh(TWeakObjectPtr<UScen
 }
 
 //Whether or not there is valid mesh data on this interface
-void UNiagaraDataInterfaceStaticMesh::IsValid(FVectorVMContext& Context)
+void UNiagaraDataInterfaceStaticMesh::IsValid(FVectorVMExternalFunctionContext& Context)
 {
 	VectorVM::FUserPtrHandler<FNDIStaticMesh_InstanceData> InstData(Context);
 	VectorVM::FExternalFuncRegisterHandler<FNiagaraBool> OutValid(Context);
 	
 	FNiagaraBool Valid;
 	Valid.SetValue(InstData->bMeshValid);
-	for (int32 i = 0; i < Context.NumInstances; ++i)
+	for (int32 i = 0; i < Context.GetNumInstances(); ++i)
 	{
 		*OutValid.GetDest() = Valid;
 		OutValid.Advance();
@@ -1825,7 +1825,7 @@ FORCEINLINE int32 UNiagaraDataInterfaceStaticMesh::RandomSection<TSampleModeDefa
 }
 
 template<typename TSampleMode>
-void UNiagaraDataInterfaceStaticMesh::RandomSection(FVectorVMContext& Context)
+void UNiagaraDataInterfaceStaticMesh::RandomSection(FVectorVMExternalFunctionContext& Context)
 {
 	VectorVM::FUserPtrHandler<FNDIStaticMesh_InstanceData> InstData(Context);
 	VectorVM::FExternalFuncRegisterHandler<int32> OutSection(Context);
@@ -1833,20 +1833,20 @@ void UNiagaraDataInterfaceStaticMesh::RandomSection(FVectorVMContext& Context)
 	const FStaticMeshLODResources* Res = InstData->CachedLOD;
 	check(Res); // sanity check - should have bound the invalid specialization below
 
-	for (int32 i = 0; i < Context.NumInstances; ++i)
+	for (int32 i = 0; i < Context.GetNumInstances(); ++i)
 	{
-		*OutSection.GetDestAndAdvance() = RandomSection<TSampleMode, true>(Context.RandStream, *Res, InstData);
+		*OutSection.GetDestAndAdvance() = RandomSection<TSampleMode, true>(Context.GetRandStream(), *Res, InstData);
 	}
 }
 
 // Invalid mesh specialization
 template<>
-void UNiagaraDataInterfaceStaticMesh::RandomSection<TSampleModeInvalid>(FVectorVMContext& Context)
+void UNiagaraDataInterfaceStaticMesh::RandomSection<TSampleModeInvalid>(FVectorVMExternalFunctionContext& Context)
 {
 	VectorVM::FUserPtrHandler<FNDIStaticMesh_InstanceData> InstData(Context);
 	VectorVM::FExternalFuncRegisterHandler<int32> OutSection(Context);
 
-	for (int32 i = 0; i < Context.NumInstances; ++i)
+	for (int32 i = 0; i < Context.GetNumInstances(); ++i)
 	{
 		*OutSection.GetDestAndAdvance() = -1;
 	}
@@ -1915,7 +1915,7 @@ FORCEINLINE int32 UNiagaraDataInterfaceStaticMesh::RandomTriIndex<TSampleModeDef
 }
 
 template<typename TSampleMode>
-void UNiagaraDataInterfaceStaticMesh::RandomTriCoord(FVectorVMContext& Context)
+void UNiagaraDataInterfaceStaticMesh::RandomTriCoord(FVectorVMExternalFunctionContext& Context)
 {
 	VectorVM::FUserPtrHandler<FNDIStaticMesh_InstanceData> InstData(Context);
 
@@ -1925,30 +1925,30 @@ void UNiagaraDataInterfaceStaticMesh::RandomTriCoord(FVectorVMContext& Context)
 	const FStaticMeshLODResources* Res = InstData->CachedLOD;
 	check(Res && InstData->StaticMesh.IsValid()); // sanity check - should have bound the invalid specialization below
 
-	for (int32 i = 0; i < Context.NumInstances; ++i)
+	for (int32 i = 0; i < Context.GetNumInstances(); ++i)
 	{
-		OutTri.SetAndAdvance(RandomTriIndex<TSampleMode, true>(Context.RandStream, *Res, InstData));
-		OutBary.SetAndAdvance(RandomBarycentricCoord(Context.RandStream));
+		OutTri.SetAndAdvance(RandomTriIndex<TSampleMode, true>(Context.GetRandStream(), *Res, InstData));
+		OutBary.SetAndAdvance(RandomBarycentricCoord(Context.GetRandStream()));
 	}
 }
 
 // Invalid mesh specialization
 template<>
-void UNiagaraDataInterfaceStaticMesh::RandomTriCoord<TSampleModeInvalid>(FVectorVMContext& Context)
+void UNiagaraDataInterfaceStaticMesh::RandomTriCoord<TSampleModeInvalid>(FVectorVMExternalFunctionContext& Context)
 {
 	VectorVM::FUserPtrHandler<FNDIStaticMesh_InstanceData> InstData(Context);
 
 	FNDIOutputParam<int32> OutTri(Context);
 	FNDIOutputParam<FVector3f> OutBary(Context);
 
-	for (int32 i = 0; i < Context.NumInstances; ++i)
+	for (int32 i = 0; i < Context.GetNumInstances(); ++i)
 	{
 		OutTri.SetAndAdvance(-1);
 		OutBary.SetAndAdvance(FVector::ZeroVector);
 	}
 }
 
-void UNiagaraDataInterfaceStaticMesh::RandomTriCoordVertexColorFiltered(FVectorVMContext& Context)
+void UNiagaraDataInterfaceStaticMesh::RandomTriCoordVertexColorFiltered(FVectorVMExternalFunctionContext& Context)
 {
 	VectorVM::FUserPtrHandler<FNDIStaticMesh_InstanceData> InstData(Context);
 	VectorVM::FExternalFuncRegisterHandler<int32> MinValue(Context);
@@ -1962,7 +1962,7 @@ void UNiagaraDataInterfaceStaticMesh::RandomTriCoordVertexColorFiltered(FVectorV
 	const FStaticMeshLODResources* Res = InstData->CachedLOD;
 	if (!Res)
 	{
-		for (int32 i = 0; i < Context.NumInstances; ++i)
+		for (int32 i = 0; i < Context.GetNumInstances(); ++i)
 		{
 			OutTri.SetAndAdvance(-1);
 			OutBary.SetAndAdvance(FVector::ZeroVector);
@@ -1973,7 +1973,7 @@ void UNiagaraDataInterfaceStaticMesh::RandomTriCoordVertexColorFiltered(FVectorV
 	FDynamicVertexColorFilterData* VCFData = InstData->DynamicVertexColorSampler.Get();	
 	FIndexArrayView Indices = Res->IndexBuffer.GetArrayView();
 
-	for (int32 i = 0; i < Context.NumInstances; ++i)
+	for (int32 i = 0; i < Context.GetNumInstances(); ++i)
 	{
 		uint32 StartIdx = (uint32)(MinValue.GetAndAdvance()*255.0f);
 		uint32 Range = (uint32)(RangeValue.GetAndAdvance()*255.0f + 0.5f);
@@ -2002,12 +2002,12 @@ void UNiagaraDataInterfaceStaticMesh::RandomTriCoordVertexColorFiltered(FVectorV
 		}
 
 		// Select a random triangle from the list.
-		uint32 RandomTri = Context.RandStream.GetFraction()*NumTris;
+		uint32 RandomTri = Context.GetRandStream().GetFraction()*NumTris;
 
 		// Now emit that triangle...
 		OutTri.SetAndAdvance(VCFData->TrianglesSortedByVertexColor[VCFData->VertexColorToTriangleStart[StartIdx] + RandomTri]);
 
-		OutBary.SetAndAdvance(RandomBarycentricCoord(Context.RandStream));
+		OutBary.SetAndAdvance(RandomBarycentricCoord(Context.GetRandStream()));
 	}
 }
 
@@ -2028,7 +2028,7 @@ FORCEINLINE int32 UNiagaraDataInterfaceStaticMesh::RandomTriIndexOnSection<TSamp
 }
 
 template<typename TSampleMode>
-void UNiagaraDataInterfaceStaticMesh::RandomTriCoordOnSection(FVectorVMContext& Context)
+void UNiagaraDataInterfaceStaticMesh::RandomTriCoordOnSection(FVectorVMExternalFunctionContext& Context)
 {
 	VectorVM::FUserPtrHandler<FNDIStaticMesh_InstanceData> InstData(Context);
 	FNDIInputParam<int32> SectionIdxParam(Context);
@@ -2043,18 +2043,18 @@ void UNiagaraDataInterfaceStaticMesh::RandomTriCoordOnSection(FVectorVMContext&
 	const int32 MaxSection = Res->Sections.Num() - 1;
 	if (MaxSection >= 0)
 	{
-		for (int32 i = 0; i < Context.NumInstances; ++i)
+		for (int32 i = 0; i < Context.GetNumInstances(); ++i)
 		{
 			int32 SecIdx = FMath::Clamp(SectionIdxParam.GetAndAdvance(), 0, MaxSection);
-			OutTri.SetAndAdvance(RandomTriIndexOnSection<TSampleMode>(Context.RandStream, *Res, SecIdx, InstData));
-			OutBary.SetAndAdvance(RandomBarycentricCoord(Context.RandStream));
+			OutTri.SetAndAdvance(RandomTriIndexOnSection<TSampleMode>(Context.GetRandStream(), *Res, SecIdx, InstData));
+			OutBary.SetAndAdvance(RandomBarycentricCoord(Context.GetRandStream()));
 		}
 		// Early out as we are done
 		return;
 	}
 
 	// Fall through that handles missing or invalid data
-	for (int32 i = 0; i < Context.NumInstances; ++i)
+	for (int32 i = 0; i < Context.GetNumInstances(); ++i)
 	{
 		OutTri.SetAndAdvance(-1);
 		OutBary.SetAndAdvance(FVector::ZeroVector);
@@ -2063,7 +2063,7 @@ void UNiagaraDataInterfaceStaticMesh::RandomTriCoordOnSection(FVectorVMContext&
 
 // Invalid mesh specialization
 template<>
-void UNiagaraDataInterfaceStaticMesh::RandomTriCoordOnSection<TSampleModeInvalid>(FVectorVMContext& Context)
+void UNiagaraDataInterfaceStaticMesh::RandomTriCoordOnSection<TSampleModeInvalid>(FVectorVMExternalFunctionContext& Context)
 {
 	VectorVM::FUserPtrHandler<FNDIStaticMesh_InstanceData> InstData(Context);
 	FNDIInputParam<int32> SectionIdxParam(Context);
@@ -2071,7 +2071,7 @@ void UNiagaraDataInterfaceStaticMesh::RandomTriCoordOnSection<TSampleModeInvalid
 	FNDIOutputParam<int32> OutTri(Context);
 	FNDIOutputParam<FVector3f> OutBary(Context);
 
-	for (int32 i = 0; i < Context.NumInstances; ++i)
+	for (int32 i = 0; i < Context.GetNumInstances(); ++i)
 	{
 		OutTri.SetAndAdvance(-1);
 		OutBary.SetAndAdvance(FVector::ZeroVector);
@@ -2079,7 +2079,7 @@ void UNiagaraDataInterfaceStaticMesh::RandomTriCoordOnSection<TSampleModeInvalid
 }
 
 template<typename TransformHandlerType>
-void UNiagaraDataInterfaceStaticMesh::GetTriCoordPosition(FVectorVMContext& Context)
+void UNiagaraDataInterfaceStaticMesh::GetTriCoordPosition(FVectorVMExternalFunctionContext& Context)
 {
 	VectorVM::FUserPtrHandler<FNDIStaticMesh_InstanceData> InstData(Context);
 	TransformHandlerType TransformHandler;
@@ -2098,7 +2098,7 @@ void UNiagaraDataInterfaceStaticMesh::GetTriCoordPosition(FVectorVMContext& Cont
 		if ( (Indices.Num() > 0) && (Positions.GetNumVertices() > 0) && Positions.GetVertexData() )
 		{
 			const int32 NumTriangles = Indices.Num() / 3;
-			for (int32 i = 0; i < Context.NumInstances; ++i)
+			for (int32 i = 0; i < Context.GetNumInstances(); ++i)
 			{
 				const int32 Tri = (TriParam.GetAndAdvance() % NumTriangles) * 3;
 				const int32 Idx0 = Indices[Tri];
@@ -2119,14 +2119,14 @@ void UNiagaraDataInterfaceStaticMesh::GetTriCoordPosition(FVectorVMContext& Cont
 	FVector Pos(0.0f);
 	TransformHandler.TransformPosition(Pos, InstData->Transform);
 
-	for (int32 i = 0; i < Context.NumInstances; ++i)
+	for (int32 i = 0; i < Context.GetNumInstances(); ++i)
 	{
 		OutPos.SetAndAdvance(Pos);
 	}
 }
 
 template<typename TransformHandlerType>
-void UNiagaraDataInterfaceStaticMesh::GetTriCoordNormal(FVectorVMContext& Context)
+void UNiagaraDataInterfaceStaticMesh::GetTriCoordNormal(FVectorVMExternalFunctionContext& Context)
 {
 	VectorVM::FUserPtrHandler<FNDIStaticMesh_InstanceData> InstData(Context);
 
@@ -2147,7 +2147,7 @@ void UNiagaraDataInterfaceStaticMesh::GetTriCoordNormal(FVectorVMContext& Contex
 		if ((Indices.Num() > 0) && (Verts.GetNumVertices() > 0) && Verts.GetTangentData())
 		{
 			const int32 NumTriangles = Indices.Num() / 3;
-			for (int32 i = 0; i < Context.NumInstances; ++i)
+			for (int32 i = 0; i < Context.GetNumInstances(); ++i)
 			{
 				const int32 Tri = (TriParam.GetAndAdvance() % NumTriangles) * 3;
 				const int32 Idx0 = Indices[Tri];
@@ -2165,14 +2165,14 @@ void UNiagaraDataInterfaceStaticMesh::GetTriCoordNormal(FVectorVMContext& Contex
 	}
 
 	// Fall through that handles missing or invalid data
-	for (int32 i = 0; i < Context.NumInstances; ++i)
+	for (int32 i = 0; i < Context.GetNumInstances(); ++i)
 	{
 		OutNorm.SetAndAdvance(FVector::ZeroVector);
 	}
 }
 
 template<typename VertexAccessorType, typename TransformHandlerType>
-void UNiagaraDataInterfaceStaticMesh::GetTriCoordTangents(FVectorVMContext& Context)
+void UNiagaraDataInterfaceStaticMesh::GetTriCoordTangents(FVectorVMExternalFunctionContext& Context)
 {
 	VectorVM::FUserPtrHandler<FNDIStaticMesh_InstanceData> InstData(Context);
 
@@ -2195,7 +2195,7 @@ void UNiagaraDataInterfaceStaticMesh::GetTriCoordTangents(FVectorVMContext& Cont
 		if ((Indices.Num() > 0) && (Res->VertexBuffers.StaticMeshVertexBuffer.GetNumVertices() > 0) && Res->VertexBuffers.StaticMeshVertexBuffer.GetTangentData())
 		{
 			const int32 NumTriangles = Indices.Num() / 3;
-			for (int32 i = 0; i < Context.NumInstances; ++i)
+			for (int32 i = 0; i < Context.GetNumInstances(); ++i)
 			{
 				const int32 Tri = (TriParam.GetAndAdvance() % NumTriangles) * 3;
 				const int32 Idx0 = Indices[Tri];
@@ -2215,7 +2215,7 @@ void UNiagaraDataInterfaceStaticMesh::GetTriCoordTangents(FVectorVMContext& Cont
 	}
 
 	// Fall through that handles missing or invalid data
-	for (int32 i = 0; i < Context.NumInstances; ++i)
+	for (int32 i = 0; i < Context.GetNumInstances(); ++i)
 	{
 		OutTangent.SetAndAdvance(FVector::ForwardVector);
 		OutBinorm.SetAndAdvance(FVector::RightVector);
@@ -2223,7 +2223,7 @@ void UNiagaraDataInterfaceStaticMesh::GetTriCoordTangents(FVectorVMContext& Cont
 	}
 }
 
-void UNiagaraDataInterfaceStaticMesh::GetTriCoordColor(FVectorVMContext& Context)
+void UNiagaraDataInterfaceStaticMesh::GetTriCoordColor(FVectorVMExternalFunctionContext& Context)
 {
 	VectorVM::FUserPtrHandler<FNDIStaticMesh_InstanceData> InstData(Context);
 
@@ -2242,7 +2242,7 @@ void UNiagaraDataInterfaceStaticMesh::GetTriCoordColor(FVectorVMContext& Context
 		if ((Indices.Num() > 0) && (Colors.GetNumVertices() > 0) && Colors.GetVertexData())
 		{
 			const int32 NumTriangles = Indices.Num() / 3;
-			for (int32 i = 0; i < Context.NumInstances; ++i)
+			for (int32 i = 0; i < Context.GetNumInstances(); ++i)
 			{
 				const int32 Tri = (TriParam.GetAndAdvance() % NumTriangles) * 3;
 				const int32 Idx0 = Indices[Tri];
@@ -2262,14 +2262,14 @@ void UNiagaraDataInterfaceStaticMesh::GetTriCoordColor(FVectorVMContext& Context
 	// Fall through that handles missing or invalid data
 	// This mesh is invalid or doesn't have color information so set the color to white.
 	FLinearColor Color = FLinearColor::White;
-	for (int32 i = 0; i < Context.NumInstances; ++i)
+	for (int32 i = 0; i < Context.GetNumInstances(); ++i)
 	{
 		OutColor.SetAndAdvance(Color);
 	}
 }
 
 template<typename VertexAccessorType>
-void UNiagaraDataInterfaceStaticMesh::GetTriCoordUV(FVectorVMContext& Context)
+void UNiagaraDataInterfaceStaticMesh::GetTriCoordUV(FVectorVMExternalFunctionContext& Context)
 {
 	VectorVM::FUserPtrHandler<FNDIStaticMesh_InstanceData> InstData(Context);
 
@@ -2289,7 +2289,7 @@ void UNiagaraDataInterfaceStaticMesh::GetTriCoordUV(FVectorVMContext& Context)
 		if ( (Indices.Num() > 0) && (Res->VertexBuffers.StaticMeshVertexBuffer.GetNumTexCoords() > 0) && Res->VertexBuffers.StaticMeshVertexBuffer.GetTexCoordData() )
 		{
 			const int32 NumTriangles = Indices.Num() / 3;
-			for (int32 i = 0; i < Context.NumInstances; ++i)
+			for (int32 i = 0; i < Context.GetNumInstances(); ++i)
 			{
 				const int32 Tri = (TriParam.GetAndAdvance() % NumTriangles) * 3;
 				const int32 Idx0 = Indices[Tri];
@@ -2307,13 +2307,13 @@ void UNiagaraDataInterfaceStaticMesh::GetTriCoordUV(FVectorVMContext& Context)
 	}
 
 	// Fall through that handles missing or invalid data
-	for (int32 i = 0; i < Context.NumInstances; ++i)
+	for (int32 i = 0; i < Context.GetNumInstances(); ++i)
 	{
 		OutUV.SetAndAdvance(FVector2D::ZeroVector);
 	}
 }
 
-void UNiagaraDataInterfaceStaticMesh::GetTriCoordPositionAndVelocity(FVectorVMContext& Context)
+void UNiagaraDataInterfaceStaticMesh::GetTriCoordPositionAndVelocity(FVectorVMExternalFunctionContext& Context)
 {
 	VectorVM::FUserPtrHandler<FNDIStaticMesh_InstanceData> InstData(Context);
 
@@ -2334,7 +2334,7 @@ void UNiagaraDataInterfaceStaticMesh::GetTriCoordPositionAndVelocity(FVectorVMCo
 		{
 			const int32 NumTriangles = Indices.Num() / 3;
 			float InvDt = InstData->DeltaSeconds > 0.0f ? (1.0f / InstData->DeltaSeconds) : 0.0f;
-			for (int32 i = 0; i < Context.NumInstances; ++i)
+			for (int32 i = 0; i < Context.GetNumInstances(); ++i)
 			{
 				const int32 Tri = (TriParam.GetAndAdvance() % NumTriangles) * 3;
 				const int32 Idx0 = Indices[Tri];
@@ -2365,14 +2365,14 @@ void UNiagaraDataInterfaceStaticMesh::GetTriCoordPositionAndVelocity(FVectorVMCo
 
 	// Fall through that handles missing or invalid data
 	FVector WSPos = InstData->Transform.TransformPosition(FVector(0.0f));
-	for (int32 i = 0; i < Context.NumInstances; ++i)
+	for (int32 i = 0; i < Context.GetNumInstances(); ++i)
 	{
 		OutPos.SetAndAdvance(WSPos);
 		OutVel.SetAndAdvance(FVector::ZeroVector);
 	}
 }
 
-void UNiagaraDataInterfaceStaticMesh::WriteTransform(const FMatrix& ToWrite, FVectorVMContext& Context)
+void UNiagaraDataInterfaceStaticMesh::WriteTransform(const FMatrix& ToWrite, FVectorVMExternalFunctionContext& Context)
 {
 	VectorVM::FExternalFuncRegisterHandler<float> Out00(Context);
 	VectorVM::FExternalFuncRegisterHandler<float> Out01(Context);
@@ -2391,7 +2391,7 @@ void UNiagaraDataInterfaceStaticMesh::WriteTransform(const FMatrix& ToWrite, FVe
 	VectorVM::FExternalFuncRegisterHandler<float> Out14(Context);
 	VectorVM::FExternalFuncRegisterHandler<float> Out15(Context);
 
-	for (int32 i = 0; i < Context.NumInstances; ++i)
+	for (int32 i = 0; i < Context.GetNumInstances(); ++i)
 	{
 		*Out00.GetDest() = ToWrite.M[0][0]; Out00.Advance();
 		*Out01.GetDest() = ToWrite.M[0][1]; Out01.Advance();
@@ -2412,19 +2412,19 @@ void UNiagaraDataInterfaceStaticMesh::WriteTransform(const FMatrix& ToWrite, FVe
 	}
 }
 
-void UNiagaraDataInterfaceStaticMesh::GetLocalToWorld(FVectorVMContext& Context)
+void UNiagaraDataInterfaceStaticMesh::GetLocalToWorld(FVectorVMExternalFunctionContext& Context)
 {
 	VectorVM::FUserPtrHandler<FNDIStaticMesh_InstanceData> InstData(Context);
 	WriteTransform(InstData->Transform, Context);
 }
 
-void UNiagaraDataInterfaceStaticMesh::GetLocalToWorldInverseTransposed(FVectorVMContext& Context)
+void UNiagaraDataInterfaceStaticMesh::GetLocalToWorldInverseTransposed(FVectorVMExternalFunctionContext& Context)
 {
 	VectorVM::FUserPtrHandler<FNDIStaticMesh_InstanceData> InstData(Context);
 	WriteTransform(InstData->TransformInverseTransposed, Context);
 }
 
-void UNiagaraDataInterfaceStaticMesh::GetWorldVelocity(FVectorVMContext& Context)
+void UNiagaraDataInterfaceStaticMesh::GetWorldVelocity(FVectorVMExternalFunctionContext& Context)
 {
 	VectorVM::FUserPtrHandler<FNDIStaticMesh_InstanceData> InstData(Context);
 
@@ -2445,14 +2445,14 @@ void UNiagaraDataInterfaceStaticMesh::GetWorldVelocity(FVectorVMContext& Context
 		}
 	}
 
-	for (int32 i = 0; i < Context.NumInstances; i++)
+	for (int32 i = 0; i < Context.GetNumInstances(); i++)
 	{
 		OutVel.SetAndAdvance(Velocity);
 	}
 }
 
 template<typename TransformHandlerType>
-void UNiagaraDataInterfaceStaticMesh::GetVertexPosition(FVectorVMContext& Context)
+void UNiagaraDataInterfaceStaticMesh::GetVertexPosition(FVectorVMExternalFunctionContext& Context)
 {
 	VectorVM::FUserPtrHandler<FNDIStaticMesh_InstanceData> InstData(Context);
 
@@ -2471,7 +2471,7 @@ void UNiagaraDataInterfaceStaticMesh::GetVertexPosition(FVectorVMContext& Contex
 		if ((NumVerts > 0) && Positions.GetVertexData())
 		{
 			FVector Pos;
-			for (int32 i = 0; i < Context.NumInstances; i++)
+			for (int32 i = 0; i < Context.GetNumInstances(); i++)
 			{
 				int32 VertexIndex = VertexIndexParam.GetAndAdvance() % NumVerts;
 				Pos = Positions.VertexPosition(VertexIndex);		
@@ -2486,51 +2486,51 @@ void UNiagaraDataInterfaceStaticMesh::GetVertexPosition(FVectorVMContext& Contex
 
 	// Fall through that handles missing or invalid data
 	FVector WSPos = InstData->Transform.TransformPosition(FVector(0.0f));
-	for (int32 i = 0; i < Context.NumInstances; ++i)
+	for (int32 i = 0; i < Context.GetNumInstances(); ++i)
 	{
 		OutPos.SetAndAdvance(WSPos);
 	}
 }
 
-void UNiagaraDataInterfaceStaticMesh::GetSocketCount(FVectorVMContext& Context)
+void UNiagaraDataInterfaceStaticMesh::GetSocketCount(FVectorVMExternalFunctionContext& Context)
 {
 	VectorVM::FUserPtrHandler<FNDIStaticMesh_InstanceData> InstData(Context);
 
 	FNDIOutputParam<int32> OutNum(Context);
 
-	for (int32 i = 0; i < Context.NumInstances; ++i)
+	for (int32 i = 0; i < Context.GetNumInstances(); ++i)
 	{
 		OutNum.SetAndAdvance(InstData->CachedSockets.Num());
 	}
 }
 
-void UNiagaraDataInterfaceStaticMesh::GetFilteredSocketCount(FVectorVMContext& Context)
+void UNiagaraDataInterfaceStaticMesh::GetFilteredSocketCount(FVectorVMExternalFunctionContext& Context)
 {
 	VectorVM::FUserPtrHandler<FNDIStaticMesh_InstanceData> InstData(Context);
 
 	FNDIOutputParam<int32> OutNum(Context);
 
-	for (int32 i = 0; i < Context.NumInstances; ++i)
+	for (int32 i = 0; i < Context.GetNumInstances(); ++i)
 	{
 		OutNum.SetAndAdvance(InstData->NumFilteredSockets);
 	}
 }
 
-void UNiagaraDataInterfaceStaticMesh::GetUnfilteredSocketCount(FVectorVMContext& Context)
+void UNiagaraDataInterfaceStaticMesh::GetUnfilteredSocketCount(FVectorVMExternalFunctionContext& Context)
 {
 	VectorVM::FUserPtrHandler<FNDIStaticMesh_InstanceData> InstData(Context);
 
 	FNDIOutputParam<int32> OutNum(Context);
 
 	const int32 NumUnfilteredSockets = InstData->CachedSockets.Num() - InstData->NumFilteredSockets;
-	for (int32 i = 0; i < Context.NumInstances; ++i)
+	for (int32 i = 0; i < Context.GetNumInstances(); ++i)
 	{
 		OutNum.SetAndAdvance(NumUnfilteredSockets);
 	}
 }
 
 template<bool bWorldSpace>
-void UNiagaraDataInterfaceStaticMesh::GetSocketTransform(FVectorVMContext& Context)
+void UNiagaraDataInterfaceStaticMesh::GetSocketTransform(FVectorVMExternalFunctionContext& Context)
 {
 	VectorVM::FUserPtrHandler<FNDIStaticMesh_InstanceData> InstData(Context);
 	FNDIInputParam<int32> InSocketIndex(Context);
@@ -2544,7 +2544,7 @@ void UNiagaraDataInterfaceStaticMesh::GetSocketTransform(FVectorVMContext& Conte
 	const int32 SocketMax = InstData->CachedSockets.Num() - 1;
 	if (SocketMax >= 0)
 	{
-		for (int32 i=0; i < Context.NumInstances; ++i)
+		for (int32 i=0; i < Context.GetNumInstances(); ++i)
 		{
 			const int32 SocketIndex = FMath::Clamp(InSocketIndex.GetAndAdvance(), 0, SocketMax);
 
@@ -2558,7 +2558,7 @@ void UNiagaraDataInterfaceStaticMesh::GetSocketTransform(FVectorVMContext& Conte
 	{
 		const FVector DefaultTranslate = bWorldSpace ? InstData->Transform.GetOrigin() : FVector::ZeroVector;
 		const FVector DefaultScale = bWorldSpace ? InstData->Transform.ExtractScaling() : FVector::OneVector;
-		for (int32 i = 0; i < Context.NumInstances; ++i)
+		for (int32 i = 0; i < Context.GetNumInstances(); ++i)
 		{
 			OutTranslate.SetAndAdvance(DefaultTranslate);
 			OutRotate.SetAndAdvance(InstRotation);
@@ -2568,7 +2568,7 @@ void UNiagaraDataInterfaceStaticMesh::GetSocketTransform(FVectorVMContext& Conte
 }
 
 template<bool bWorldSpace>
-void UNiagaraDataInterfaceStaticMesh::GetFilteredSocketTransform(FVectorVMContext& Context)
+void UNiagaraDataInterfaceStaticMesh::GetFilteredSocketTransform(FVectorVMExternalFunctionContext& Context)
 {
 	VectorVM::FUserPtrHandler<FNDIStaticMesh_InstanceData> InstData(Context);
 	FNDIInputParam<int32> InSocketIndex(Context);
@@ -2582,7 +2582,7 @@ void UNiagaraDataInterfaceStaticMesh::GetFilteredSocketTransform(FVectorVMContex
 	const int32 SocketMax = InstData->NumFilteredSockets - 1;
 	if (SocketMax >= 0)
 	{
-		for (int32 i = 0; i < Context.NumInstances; ++i)
+		for (int32 i = 0; i < Context.GetNumInstances(); ++i)
 		{
 			const int32 SocketIndex = InstData->FilteredAndUnfilteredSockets[FMath::Clamp(InSocketIndex.GetAndAdvance(), 0, SocketMax)];
 
@@ -2596,7 +2596,7 @@ void UNiagaraDataInterfaceStaticMesh::GetFilteredSocketTransform(FVectorVMContex
 	{
 		const FVector DefaultTranslate = bWorldSpace ? InstData->Transform.GetOrigin() : FVector::ZeroVector;
 		const FVector DefaultScale = bWorldSpace ? InstData->Transform.ExtractScaling() : FVector::OneVector;
-		for (int32 i = 0; i < Context.NumInstances; ++i)
+		for (int32 i = 0; i < Context.GetNumInstances(); ++i)
 		{
 			OutTranslate.SetAndAdvance(DefaultTranslate);
 			OutRotate.SetAndAdvance(InstRotation);
@@ -2606,7 +2606,7 @@ void UNiagaraDataInterfaceStaticMesh::GetFilteredSocketTransform(FVectorVMContex
 }
 
 template<bool bWorldSpace>
-void UNiagaraDataInterfaceStaticMesh::GetUnfilteredSocketTransform(FVectorVMContext& Context)
+void UNiagaraDataInterfaceStaticMesh::GetUnfilteredSocketTransform(FVectorVMExternalFunctionContext& Context)
 {
 	VectorVM::FUserPtrHandler<FNDIStaticMesh_InstanceData> InstData(Context);
 	FNDIInputParam<int32> InSocketIndex(Context);
@@ -2623,7 +2623,7 @@ void UNiagaraDataInterfaceStaticMesh::GetUnfilteredSocketTransform(FVectorVMCont
 		check(SocketMax >= 0);
 
 		const int32 SocketOffset = InstData->NumFilteredSockets;
-		for (int32 i = 0; i < Context.NumInstances; ++i)
+		for (int32 i = 0; i < Context.GetNumInstances(); ++i)
 		{
 			const int32 SocketIndex = InstData->FilteredAndUnfilteredSockets[FMath::Clamp(InSocketIndex.GetAndAdvance(), 0, SocketMax) + SocketOffset];
 
@@ -2636,7 +2636,7 @@ void UNiagaraDataInterfaceStaticMesh::GetUnfilteredSocketTransform(FVectorVMCont
 	else if (InstData->CachedSockets.Num() > 0)
 	{
 		const int32 SocketMax = InstData->CachedSockets.Num() - 1;
-		for (int32 i = 0; i < Context.NumInstances; ++i)
+		for (int32 i = 0; i < Context.GetNumInstances(); ++i)
 		{
 			const int32 SocketIndex = FMath::Clamp(InSocketIndex.GetAndAdvance(), 0, SocketMax);
 
@@ -2650,7 +2650,7 @@ void UNiagaraDataInterfaceStaticMesh::GetUnfilteredSocketTransform(FVectorVMCont
 	{
 		const FVector DefaultTranslate = bWorldSpace ? InstData->Transform.GetOrigin() : FVector::ZeroVector;
 		const FVector DefaultScale = bWorldSpace ? InstData->Transform.ExtractScaling() : FVector::OneVector;
-		for (int32 i = 0; i < Context.NumInstances; ++i)
+		for (int32 i = 0; i < Context.GetNumInstances(); ++i)
 		{
 			OutTranslate.SetAndAdvance(DefaultTranslate);
 			OutRotate.SetAndAdvance(InstRotation);
diff --git a/Engine/Plugins/FX/Niagara/Source/Niagara/Private/NiagaraDataInterfaceTexture.cpp b/Engine/Plugins/FX/Niagara/Source/Niagara/Private/NiagaraDataInterfaceTexture.cpp
--- a/Engine/Plugins/FX/Niagara/Source/Niagara/Private/NiagaraDataInterfaceTexture.cpp
+++ b/Engine/Plugins/FX/Niagara/Source/Niagara/Private/NiagaraDataInterfaceTexture.cpp
@@ -228,19 +228,19 @@ bool UNiagaraDataInterfaceTexture::PerInstanceTick(void* PerInstanceData, FNiaga
 	return false;
 }
 
-void UNiagaraDataInterfaceTexture::GetTextureDimensions(FVectorVMContext& Context)
+void UNiagaraDataInterfaceTexture::GetTextureDimensions(FVectorVMExternalFunctionContext& Context)
 {
 	FNDIOutputParam<float> OutWidth(Context);
 	FNDIOutputParam<float> OutHeight(Context);
 
-	for (int32 i = 0; i < Context.NumInstances; ++i)
+	for (int32 i = 0; i < Context.GetNumInstances(); ++i)
 	{
 		OutWidth.SetAndAdvance(TextureSize.X);
 		OutHeight.SetAndAdvance(TextureSize.Y);
 	}
 }
 
-void UNiagaraDataInterfaceTexture::SampleTexture(FVectorVMContext& Context)
+void UNiagaraDataInterfaceTexture::SampleTexture(FVectorVMExternalFunctionContext& Context)
 {
 	VectorVM::FExternalFuncInputHandler<float> XParam(Context);
 	VectorVM::FExternalFuncInputHandler<float> YParam(Context);
@@ -249,7 +249,7 @@ void UNiagaraDataInterfaceTexture::SampleTexture(FVectorVMContext& Context)
 	VectorVM::FExternalFuncRegisterHandler<float> OutSampleB(Context);
 	VectorVM::FExternalFuncRegisterHandler<float> OutSampleA(Context);
 
-	for (int32 i = 0; i < Context.NumInstances; ++i)
+	for (int32 i = 0; i < Context.GetNumInstances(); ++i)
 	{
 		float X = XParam.GetAndAdvance();
 		float Y = YParam.GetAndAdvance();
@@ -261,7 +261,7 @@ void UNiagaraDataInterfaceTexture::SampleTexture(FVectorVMContext& Context)
 
 }
 
-void UNiagaraDataInterfaceTexture::SamplePseudoVolumeTexture(FVectorVMContext& Context)
+void UNiagaraDataInterfaceTexture::SamplePseudoVolumeTexture(FVectorVMExternalFunctionContext& Context)
 {
 	// Noop handler which just returns magenta since this doesn't run on CPU.
 	VectorVM::FExternalFuncInputHandler<float> UVW_UParam(Context);
@@ -288,7 +288,7 @@ void UNiagaraDataInterfaceTexture::SamplePseudoVolumeTexture(FVectorVMContext& C
 	VectorVM::FExternalFuncRegisterHandler<float> OutSampleB(Context);
 	VectorVM::FExternalFuncRegisterHandler<float> OutSampleA(Context);
 
-	for (int32 i = 0; i < Context.NumInstances; ++i)
+	for (int32 i = 0; i < Context.GetNumInstances(); ++i)
 	{
 		UVW_UParam.Advance();
 		UVW_VParam.Advance();
diff --git a/Engine/Plugins/FX/Niagara/Source/Niagara/Private/NiagaraDataInterfaceVector2DCurve.cpp b/Engine/Plugins/FX/Niagara/Source/Niagara/Private/NiagaraDataInterfaceVector2DCurve.cpp
--- a/Engine/Plugins/FX/Niagara/Source/Niagara/Private/NiagaraDataInterfaceVector2DCurve.cpp
+++ b/Engine/Plugins/FX/Niagara/Source/Niagara/Private/NiagaraDataInterfaceVector2DCurve.cpp
@@ -213,14 +213,14 @@ FORCEINLINE_DEBUGGABLE FVector2D UNiagaraDataInterfaceVector2DCurve::SampleCurve
 }
 
 template<typename UseLUT>
-void UNiagaraDataInterfaceVector2DCurve::SampleCurve(FVectorVMContext& Context)
+void UNiagaraDataInterfaceVector2DCurve::SampleCurve(FVectorVMExternalFunctionContext& Context)
 {
 	//TODO: Create some SIMDable optimized representation of the curve to do this faster.
 	VectorVM::FExternalFuncInputHandler<float> XParam(Context);
 	VectorVM::FExternalFuncRegisterHandler<float> OutSampleX(Context);
 	VectorVM::FExternalFuncRegisterHandler<float> OutSampleY(Context);
 
-	for (int32 i = 0; i < Context.NumInstances; ++i)
+	for (int32 i = 0; i < Context.GetNumInstances(); ++i)
 	{
 		float X = XParam.GetAndAdvance();
 		FVector2D V = SampleCurveInternal<UseLUT>(X);
diff --git a/Engine/Plugins/FX/Niagara/Source/Niagara/Private/NiagaraDataInterfaceVector4Curve.cpp b/Engine/Plugins/FX/Niagara/Source/Niagara/Private/NiagaraDataInterfaceVector4Curve.cpp
--- a/Engine/Plugins/FX/Niagara/Source/Niagara/Private/NiagaraDataInterfaceVector4Curve.cpp
+++ b/Engine/Plugins/FX/Niagara/Source/Niagara/Private/NiagaraDataInterfaceVector4Curve.cpp
@@ -236,7 +236,7 @@ FORCEINLINE_DEBUGGABLE FVector4 UNiagaraDataInterfaceVector4Curve::SampleCurveIn
 
 
 template<typename UseLUT>
-void UNiagaraDataInterfaceVector4Curve::SampleCurve(FVectorVMContext& Context)
+void UNiagaraDataInterfaceVector4Curve::SampleCurve(FVectorVMExternalFunctionContext& Context)
 {
 	//TODO: Create some SIMDable optimized representation of the curve to do this faster.
 	VectorVM::FExternalFuncInputHandler<float> XParam(Context);
@@ -245,7 +245,7 @@ void UNiagaraDataInterfaceVector4Curve::SampleCurve(FVectorVMContext& Context)
 	VectorVM::FExternalFuncRegisterHandler<float> SamplePtrB(Context);
 	VectorVM::FExternalFuncRegisterHandler<float> SamplePtrA(Context);
 
-	for (int32 i = 0; i < Context.NumInstances; ++i)
+	for (int32 i = 0; i < Context.GetNumInstances(); ++i)
 	{
 		float X = XParam.GetAndAdvance();
 		FVector4 Sample = SampleCurveInternal<UseLUT>(X);
