- Set current data to be in ready state

#rnx

diff --git a/Engine/Plugins/FX/Niagara/Source/Niagara/Classes/NiagaraDataSet.h b/Engine/Plugins/FX/Niagara/Source/Niagara/Classes/NiagaraDataSet.h
--- a/Engine/Plugins/FX/Niagara/Source/Niagara/Classes/NiagaraDataSet.h
+++ b/Engine/Plugins/FX/Niagara/Source/Niagara/Classes/NiagaraDataSet.h
@@ -228,7 +228,7 @@ private:
 	//////////////////////////////////////////////////////////////////////////
 	// GPU Data
 	/** Location in the frame where GPU data will be ready, for CPU this is always the first group, for GPU is depends on the features used as to which phase. */
-	ENiagaraGpuComputeTickStage::Type GPUDataReadyStage;
+	ENiagaraGpuComputeTickStage::Type GPUDataReadyStage = ENiagaraGpuComputeTickStage::First;
 	/** The buffer offset where the instance count is accumulated. */
 	uint32 GPUInstanceCountBufferOffset;
 	/** GPU Buffer containing floating point values for GPU simulations. */
diff --git a/Engine/Plugins/FX/Niagara/Source/Niagara/Private/NiagaraEmitterInstanceBatcher.cpp b/Engine/Plugins/FX/Niagara/Source/Niagara/Private/NiagaraEmitterInstanceBatcher.cpp
--- a/Engine/Plugins/FX/Niagara/Source/Niagara/Private/NiagaraEmitterInstanceBatcher.cpp
+++ b/Engine/Plugins/FX/Niagara/Source/Niagara/Private/NiagaraEmitterInstanceBatcher.cpp
@@ -1379,8 +1379,12 @@ void NiagaraEmitterInstanceBatcher::ExecuteTicks(FRHICommandList& RHICmdList, FR
 					ComputeContext->MainDataSet->CurrentData = CurrentData;
 					CurrentData->SwapGPU(FinalSimStageDataBuffer);
 
+					// Mark data as ready for anyone who picks up the buffer on the next frame
+					CurrentData->SetGPUDataReadyStage(ENiagaraGpuComputeTickStage::First);
+
 					ComputeContext->SetTranslucentDataToRender(nullptr);
 					ComputeContext->SetDataToRender(CurrentData);
+
 #if WITH_MGPU
 					AddTemporalEffectBuffers(CurrentData);
 #endif // WITH_MGPU
