Fixed a bug where the name of a static switch parameter could contain a dot char

diff --git a/Engine/Plugins/FX/Niagara/Source/NiagaraEditor/Private/Customizations/NiagaraStaticSwitchNodeDetails.cpp b/Engine/Plugins/FX/Niagara/Source/NiagaraEditor/Private/Customizations/NiagaraStaticSwitchNodeDetails.cpp
--- a/Engine/Plugins/FX/Niagara/Source/NiagaraEditor/Private/Customizations/NiagaraStaticSwitchNodeDetails.cpp
+++ b/Engine/Plugins/FX/Niagara/Source/NiagaraEditor/Private/Customizations/NiagaraStaticSwitchNodeDetails.cpp
@@ -95,6 +95,7 @@ void FNiagaraStaticSwitchNodeDetails::CustomizeDetails(IDetailLayoutBuilder& Det
 				.IsEnabled(this, &FNiagaraStaticSwitchNodeDetails::GetDefaultOptionEnabled)
 				.ToolTipText(LOCTEXT("NiagaraSwitchNodeNameTooltip", "This is the name of the parameter that is exposed to the user calling this function graph."))
 				.OnTextCommitted(this, &FNiagaraStaticSwitchNodeDetails::OnParameterNameCommited)
+				.OnVerifyTextChanged(this, &FNiagaraStaticSwitchNodeDetails::OnVerifyParameterNameChanged)
 				.SelectAllTextWhenFocused(true)
 				.RevertTextOnEscape(true)
 			]
@@ -546,6 +547,39 @@ void FNiagaraStaticSwitchNodeDetails::OnParameterNameCommited(const FText& InTex
 	}
 }
 
+bool FNiagaraStaticSwitchNodeDetails::OnVerifyParameterNameChanged(const FText& InName, FText& OutErrorMessage)
+{
+	FText TrimmedName = FText::TrimPrecedingAndTrailing(InName);
+
+	if (TrimmedName.IsEmpty())
+	{
+		OutErrorMessage = LOCTEXT("RenameFailed_LeftBlank", "Static switch name cannot be left blank");
+		return false;
+	}
+
+	if (TrimmedName.ToString().Len() >= NAME_SIZE)
+	{
+		FFormatNamedArguments Arguments;
+		Arguments.Add(TEXT("CharCount"), NAME_SIZE);
+		OutErrorMessage = FText::Format(LOCTEXT("RenameFailed_TooLong", "Name must be less than {CharCount} characters long."), Arguments);
+		return false;
+	}
+
+	if (FName(*TrimmedName.ToString()) == NAME_None)
+	{
+		OutErrorMessage = LOCTEXT("RenameFailed_ReservedNameNone", "\"None\" is a reserved term and cannot be used for a static switch name");
+		return false;
+	}
+
+	if (TrimmedName.ToString().Contains("."))
+	{
+		OutErrorMessage = LOCTEXT("RenameFailed_IncludesDot", "'.' is used to separate namespaces and cannot be used in the name of a static switch");
+		return false;
+	}
+
+	return true;
+}
+
 bool FNiagaraStaticSwitchNodeDetails::GetDefaultOptionEnabled() const
 {
 	return !IsConstantSelection();
diff --git a/Engine/Plugins/FX/Niagara/Source/NiagaraEditor/Private/Customizations/NiagaraStaticSwitchNodeDetails.h b/Engine/Plugins/FX/Niagara/Source/NiagaraEditor/Private/Customizations/NiagaraStaticSwitchNodeDetails.h
--- a/Engine/Plugins/FX/Niagara/Source/NiagaraEditor/Private/Customizations/NiagaraStaticSwitchNodeDetails.h
+++ b/Engine/Plugins/FX/Niagara/Source/NiagaraEditor/Private/Customizations/NiagaraStaticSwitchNodeDetails.h
@@ -82,6 +82,7 @@ private:
 	// parameter name option function
 	FText GetParameterNameText() const;
 	void OnParameterNameCommited(const FText& InText, ETextCommit::Type InCommitType);
+	bool OnVerifyParameterNameChanged(const FText& InLabel, FText& OutErrorMessage);
 
 	// default value option functions
 	bool GetDefaultOptionEnabled() const;
