Fixed a bug where UNiagaraFunctionLibrary::SpawnSystemAttached() did not handle EAttachLocation::SnapToTarget correctly

diff --git a/Engine/Plugins/FX/Niagara/Source/Niagara/Private/NiagaraFunctionLibrary.cpp b/Engine/Plugins/FX/Niagara/Source/Niagara/Private/NiagaraFunctionLibrary.cpp
--- a/Engine/Plugins/FX/Niagara/Source/Niagara/Private/NiagaraFunctionLibrary.cpp
+++ b/Engine/Plugins/FX/Niagara/Source/Niagara/Private/NiagaraFunctionLibrary.cpp
@@ -140,59 +140,7 @@ UNiagaraComponent* UNiagaraFunctionLibrary::SpawnSystemAtLocation(const UObject*
 */
 UNiagaraComponent* UNiagaraFunctionLibrary::SpawnSystemAttached(UNiagaraSystem* SystemTemplate, USceneComponent* AttachToComponent, FName AttachPointName, FVector Location, FRotator Rotation, EAttachLocation::Type LocationType, bool bAutoDestroy, bool bAutoActivate, ENCPoolMethod PoolingMethod, bool bPreCullCheck)
 {
-	UNiagaraComponent* PSC = nullptr;
-	if (SystemTemplate)
-	{
-		if (AttachToComponent == NULL)
-		{
-			UE_LOG(LogScript, Warning, TEXT("UNiagaraFunctionLibrary::SpawnSystemAttached: NULL AttachComponent specified!"));
-		}
-		else
-		{
-			bool bShouldCull = false;
-			if (bPreCullCheck)
-			{
-				FNiagaraWorldManager* WorldManager = FNiagaraWorldManager::Get(AttachToComponent->GetWorld());
-				//TODO: For now using the attach parent location and ignoring the emitters relative location which is clearly going to be a bit wrong in some cases.
-				bShouldCull = WorldManager->ShouldPreCull(SystemTemplate, AttachToComponent->GetComponentLocation());
-			}
-
-			if (!bShouldCull)
-			{
-				PSC = CreateNiagaraSystem(SystemTemplate, AttachToComponent->GetWorld(), AttachToComponent->GetOwner(), bAutoDestroy, PoolingMethod);
-				if(PSC)
-				{
-#if WITH_EDITOR
-					if (GForceNiagaraSpawnAttachedSolo > 0)
-					{
-						PSC->SetForceSolo(true);
-					}
-#endif
-					if(!PSC->IsRegistered())
-					{
-						PSC->RegisterComponentWithWorld(AttachToComponent->GetWorld());
-					}
-
-					PSC->AttachToComponent(AttachToComponent, FAttachmentTransformRules::KeepRelativeTransform, AttachPointName);
-					if (LocationType == EAttachLocation::KeepWorldPosition)
-					{
-						PSC->SetWorldLocationAndRotation(Location, Rotation);
-					}
-					else
-					{
-						PSC->SetRelativeLocationAndRotation(Location, Rotation);
-					}
-					PSC->SetRelativeScale3D(FVector(1.f));
-
-					if (bAutoActivate)
-					{
-						PSC->Activate();
-					}
-				}
-			}
-		}
-	}
-	return PSC;
+	return SpawnSystemAttached(SystemTemplate, AttachToComponent, AttachPointName, Location, Rotation, FVector(1, 1, 1), LocationType, bAutoDestroy, PoolingMethod, bAutoActivate, bPreCullCheck);
 }
 
 /**
