Add ShadingModelID to GBuffer DI

diff --git a/Engine/Plugins/FX/Niagara/Shaders/Private/NiagaraDataInterfaceGBuffer.ush b/Engine/Plugins/FX/Niagara/Shaders/Private/NiagaraDataInterfaceGBuffer.ush
--- a/Engine/Plugins/FX/Niagara/Shaders/Private/NiagaraDataInterfaceGBuffer.ush
+++ b/Engine/Plugins/FX/Niagara/Shaders/Private/NiagaraDataInterfaceGBuffer.ush
@@ -252,3 +252,15 @@ void DIGBuffer_DecodeSceneColor(in float2 ScreenUV, in bool bApplyViewportOffset
 #endif
 }
 
+void DIGBuffer_DecodeShadingModelID(in float2 ScreenUV, in bool bApplyViewportOffset, out bool IsValid, out int OutValue)
+{
+#if NDIGBUFFER_HAS_GBUFFER
+	ScreenUV = DIGBuffer_ApplyViewportOffset(ScreenUV, bApplyViewportOffset);
+	FGBufferData GBufferData = GetGBufferData(ScreenUV);
+	IsValid = true;
+	OutValue = GBufferData.ShadingModelID;
+#else
+	IsValid = false;
+	OutValue = 0;
+#endif
+}
diff --git a/Engine/Plugins/FX/Niagara/Source/Niagara/Private/DataInterface/NiagaraDataInterfaceGBuffer.cpp b/Engine/Plugins/FX/Niagara/Source/Niagara/Private/DataInterface/NiagaraDataInterfaceGBuffer.cpp
--- a/Engine/Plugins/FX/Niagara/Source/Niagara/Private/DataInterface/NiagaraDataInterfaceGBuffer.cpp
+++ b/Engine/Plugins/FX/Niagara/Source/Niagara/Private/DataInterface/NiagaraDataInterfaceGBuffer.cpp
@@ -93,6 +93,7 @@ namespace NiagaraDataInterfaceGBufferLocal
 			FGBufferAttribute(TEXT("CustomStencil"),	TEXT("int"),	FNiagaraTypeDefinition::GetIntDef(), FText::GetEmpty()),
 
 			FGBufferAttribute(TEXT("SceneColor"),		TEXT("float4"),	FNiagaraTypeDefinition::GetVec4Def(), GetDescription_SceneColor()),
+			FGBufferAttribute(TEXT("ShadingModelID"),	TEXT("int"),	FNiagaraTypeDefinition::GetIntDef(), FText::GetEmpty()),
 		};
 
 		return MakeArrayView(GBufferAttributes);
