Merged CL 16853681: Fixed a crash when copying emitter between systems

diff --git a/Engine/Plugins/FX/Niagara/Source/Niagara/Private/NiagaraSystem.cpp b/Engine/Plugins/FX/Niagara/Source/Niagara/Private/NiagaraSystem.cpp
--- a/Engine/Plugins/FX/Niagara/Source/Niagara/Private/NiagaraSystem.cpp
+++ b/Engine/Plugins/FX/Niagara/Source/Niagara/Private/NiagaraSystem.cpp
@@ -2491,14 +2491,14 @@ bool UNiagaraSystem::QueryCompileComplete(bool bWait, bool bDoPost, bool bDoNotA
 			{
 				// Synchronize the variables that we actually encountered during precompile so that we can expose them to the end user.
 				TArray<FNiagaraVariable> OriginalExposedParams;
-				GetExposedParameters().GetParameters(OriginalExposedParams);
+				ExposedParameters.GetParameters(OriginalExposedParams);
 				TArray<FNiagaraVariable>& EncounteredExposedVars = AsyncTask->PrecompileReference->EncounteredExposedVars;
 				for (int32 i = 0; i < EncounteredExposedVars.Num(); i++)
 				{
 					if (OriginalExposedParams.Contains(EncounteredExposedVars[i]) == false)
 					{
 						// Just in case it wasn't added previously..
-						ExposedParameters.AddParameter(EncounteredExposedVars[i]);
+						ExposedParameters.AddParameter(EncounteredExposedVars[i], true, false);
 					}
 				}
 				
