Merging CL 16854568: Fixed a crash that could occurr when using the cascade-niagara converter

diff --git a/Engine/Plugins/FX/Niagara/Source/Niagara/Classes/NiagaraSystem.h b/Engine/Plugins/FX/Niagara/Source/Niagara/Classes/NiagaraSystem.h
--- a/Engine/Plugins/FX/Niagara/Source/Niagara/Classes/NiagaraSystem.h
+++ b/Engine/Plugins/FX/Niagara/Source/Niagara/Classes/NiagaraSystem.h
@@ -180,7 +180,7 @@ struct FNiagaraLazyPrecompileReference
 	UNiagaraSystem* System = nullptr;
 	TArray<UNiagaraScript*> Scripts;
 	TMap<UNiagaraScript*, int32> EmitterScriptIndex;
-	TArray<TObjectPtr<UObject>>* CompilationRootObjects;
+	TArray<TObjectPtr<UObject>> CompilationRootObjects;
 	TArray<FNiagaraVariable> EncounteredExposedVars;
 
 private:
@@ -210,6 +210,7 @@ public:
 	double StartCompileTime = 0;
 	bool bWaitForCompileJob = false;
 	bool bUsedShaderCompilerWorker = false;
+	bool bFetchedGCObjects = false;
 	FEmitterCompiledScriptPair ScriptPair;
 
 	ENiagaraCompilationState CurrentState;
diff --git a/Engine/Plugins/FX/Niagara/Source/Niagara/Private/NiagaraSystem.cpp b/Engine/Plugins/FX/Niagara/Source/Niagara/Private/NiagaraSystem.cpp
--- a/Engine/Plugins/FX/Niagara/Source/Niagara/Private/NiagaraSystem.cpp
+++ b/Engine/Plugins/FX/Niagara/Source/Niagara/Private/NiagaraSystem.cpp
@@ -146,7 +146,7 @@ TSharedPtr<FNiagaraCompileRequestDataBase, ESPMode::ThreadSafe> FNiagaraLazyPrec
 
 		// Grab the list of user variables that were actually encountered so that we can add to them later.
 		SystemPrecompiledData->GatherPreCompiledVariables(TEXT("User"), EncounteredExposedVars);
-		SystemPrecompiledData->GetReferencedObjects(*CompilationRootObjects);
+		SystemPrecompiledData->GetReferencedObjects(CompilationRootObjects);
 		TArray<FNiagaraEmitterHandle> EmitterHandles = System->GetEmitterHandles();
 		for (int32 i = 0; i < EmitterHandles.Num(); i++)
 		{
@@ -154,7 +154,7 @@ TSharedPtr<FNiagaraCompileRequestDataBase, ESPMode::ThreadSafe> FNiagaraLazyPrec
 			if (Handle.GetInstance() && Handle.GetIsEnabled())
 			{
 				TSharedPtr<FNiagaraCompileRequestDataBase, ESPMode::ThreadSafe> EmitterPrecompiledData = SystemPrecompiledData->GetDependentRequest(i);
-				EmitterPrecompiledData->GetReferencedObjects(*CompilationRootObjects);
+				EmitterPrecompiledData->GetReferencedObjects(CompilationRootObjects);
 				
 				TArray<UNiagaraScript*> EmitterScripts;
 				Handle.GetInstance()->GetScripts(EmitterScripts, false, true);
@@ -2428,6 +2428,12 @@ bool UNiagaraSystem::QueryCompileComplete(bool bWait, bool bDoPost, bool bDoNotA
 			{
 				continue;
 			}
+
+			if (!AsyncTask->bFetchedGCObjects && AsyncTask->CurrentState > ENiagaraCompilationState::Precompile && AsyncTask->CurrentState <= ENiagaraCompilationState::ProcessResult)
+			{
+				CompileRequest.RootObjects.Append(AsyncTask->PrecompileReference->CompilationRootObjects);
+				AsyncTask->bFetchedGCObjects = true;
+			}
 			
 			if (bWait)
 			{
@@ -2873,7 +2879,6 @@ bool UNiagaraSystem::RequestCompile(bool bForce, FNiagaraSystemUpdateContext* Op
 		// prepare data for any precompile the ddc tasks need to do
 		TSharedPtr<FNiagaraLazyPrecompileReference, ESPMode::ThreadSafe> PrecompileReference = MakeShared<FNiagaraLazyPrecompileReference, ESPMode::ThreadSafe>();
 		PrecompileReference->System = this;
-		PrecompileReference->CompilationRootObjects = &ActiveCompilation.RootObjects;
 		PrecompileReference->Scripts = ScriptsNeedingCompile;
 		
 		for (int32 i = 0; i < EmitterHandles.Num(); i++)
