HWRT Collision Filtering.

- Build a hash table on GPU of GPUSceneIndex to CollisionGroup. Hash table is updated just before we issue collision rays if it's dirty.
- Passing GPUSceneIndex into RG payload. + misc boiler plate
- Pass CollisionGroup into RG Ray data. + misc boiler plate
- Lookup and test primitive hash table group vs the ray group

API:
User calls AcquireNiagaraGPURayTracedCollisionGroup() and ReleaseNiagaraGPURayTracedCollisionGroup() to get/release a collision group index they can use.
Call SetActorNiagaraGPURayTracedCollisionGroup() to set all primitives on this actor to the given collision group.
Pass same collision group into Niagara Systems as a user parameter and pass through to the Collision DI.
Any rays from this DI will now ignore collisions with primitives of the same CollisionGroup.

Misc:

- Batcher now has access to it's FScene.
- Delegate broadcast when GPUSceneInstances are reallocated.
- A couple of fixes to the hash table ush functions. No one is using these anyway.
- Some Test Assets

diff --git a/Engine/Plugins/FX/Niagara/Shaders/Private/NiagaraDataInterfaceCollisionQuery.ush b/Engine/Plugins/FX/Niagara/Shaders/Private/NiagaraDataInterfaceCollisionQuery.ush
--- a/Engine/Plugins/FX/Niagara/Shaders/Private/NiagaraDataInterfaceCollisionQuery.ush
+++ b/Engine/Plugins/FX/Niagara/Shaders/Private/NiagaraDataInterfaceCollisionQuery.ush
@@ -151,7 +151,7 @@ uint BuildRayTraceMask(int TraceChannel)
 	return 0xFF;
 }
 
-void NDICollisionQuery_IssueAsyncRayTrace(uint MaxRayTraceCount, RWStructuredBuffer<FBasicRayData> RWRayRequests, RWBuffer<uint> RWRayTraceCounts, int In_QueryID, float3 In_TraceStart, float3 In_TraceEnd, int In_TraceChannel, out bool Out_IsQueryValid)
+void NDICollisionQuery_IssueAsyncRayTrace(uint MaxRayTraceCount, RWStructuredBuffer<FNiagaraRayData> RWRayRequests, RWBuffer<uint> RWRayTraceCounts, int In_QueryID, float3 In_TraceStart, float3 In_TraceEnd, int In_TraceChannel, out bool Out_IsQueryValid)
 {
 	Out_IsQueryValid = false;
 
@@ -196,7 +196,7 @@ void NDICollisionQuery_ReserveRayTraceIndex(uint MaxRayTraceCount, RWBuffer<uint
 	}
 }
 
-void NDICollisionQuery_CreateAsyncRayTrace(uint MaxRayTraceCount, RWStructuredBuffer<FBasicRayData> RWRayRequests, RWBuffer<uint> RWRayTraceCounts, float3 In_TraceStart, float3 In_TraceEnd, int In_TraceChannel, out int Out_QueryID, out bool Out_IsQueryValid)
+void NDICollisionQuery_CreateAsyncRayTrace(uint MaxRayTraceCount, RWStructuredBuffer<FNiagaraRayData> RWRayRequests, RWBuffer<uint> RWRayTraceCounts, float3 In_TraceStart, float3 In_TraceEnd, int In_TraceChannel, out int Out_QueryID, out bool Out_IsQueryValid)
 {
 	Out_IsQueryValid = false;
 
@@ -211,6 +211,9 @@ void NDICollisionQuery_CreateAsyncRayTrace(uint MaxRayTraceCount, RWStructuredBu
 		RWRayRequests[Out_QueryID].Mask = BuildRayTraceMask(In_TraceChannel);
 		RWRayRequests[Out_QueryID].Direction = normalize(Trace);
 		RWRayRequests[Out_QueryID].TFar = TraceLength;
+
+		//TODO: We can pack collision group and trace channels into the same int and test both in the RG Shader.
+		RWRayRequests[Out_QueryID].CollisionGroup = In_TraceChannel;
 	}
 }
 
diff --git a/Engine/Plugins/FX/Niagara/Shaders/Private/NiagaraDataInterfaceCollisionQueryTemplate.ush b/Engine/Plugins/FX/Niagara/Shaders/Private/NiagaraDataInterfaceCollisionQueryTemplate.ush
--- a/Engine/Plugins/FX/Niagara/Shaders/Private/NiagaraDataInterfaceCollisionQueryTemplate.ush
+++ b/Engine/Plugins/FX/Niagara/Shaders/Private/NiagaraDataInterfaceCollisionQueryTemplate.ush
@@ -3,7 +3,7 @@
 #if NIAGARA_USE_RAYTRACING
 bool											RayTracingEnabled_{ParameterName};
 uint                                            MaxRayTraceCount_{ParameterName};
-RWStructuredBuffer<FBasicRayData>               RWRayRequests_{ParameterName};
+RWStructuredBuffer<FNiagaraRayData>               RWRayRequests_{ParameterName};
 StructuredBuffer<FNiagaraRayTracingPayload>     IntersectionResults_{ParameterName};
 RWBuffer<uint>                                  RWRayTraceCounts_{ParameterName};
 #endif
diff --git a/Engine/Plugins/FX/Niagara/Shaders/Private/NiagaraRayTraceCollisionGroupShaders.usf b/Engine/Plugins/FX/Niagara/Shaders/Private/NiagaraRayTraceCollisionGroupShaders.usf
new file mode 100644
--- /dev/null
+++ b/Engine/Plugins/FX/Niagara/Shaders/Private/NiagaraRayTraceCollisionGroupShaders.usf
@@ -0,0 +1,28 @@
+// Copyright Epic Games, Inc. All Rights Reserved.
+
+#include "/Engine/Private/Common.ush"
+#include "/Engine/Private/HashTable.ush"
+
+
+RWBuffer<uint> RWHashToCollisionGroups;
+
+Buffer<uint2> NewPrimIdCollisionGroupPairs;
+uint NumNewPrims;
+
+[numthreads(THREAD_COUNT, 1, 1)]
+void UpdatePrimIdToCollisionGroupMap(uint Index : SV_DispatchThreadID)
+{
+	if (Index < NumNewPrims)
+	{
+		uint PrimId = NewPrimIdCollisionGroupPairs[Index].x;
+		uint CollisionGroupId = NewPrimIdCollisionGroupPairs[Index].y;
+		uint CollisionGroupIdIndex;
+
+		//TODO: Handle failures.
+		if (PrimId != -1)
+		{
+			HashTableAdd(PrimId, CollisionGroupIdIndex);
+			RWHashToCollisionGroups[CollisionGroupIdIndex] = CollisionGroupId;
+		}
+	}
+}
diff --git a/Engine/Plugins/FX/Niagara/Shaders/Private/NiagaraRayTracingCommon.ush b/Engine/Plugins/FX/Niagara/Shaders/Private/NiagaraRayTracingCommon.ush
--- a/Engine/Plugins/FX/Niagara/Shaders/Private/NiagaraRayTracingCommon.ush
+++ b/Engine/Plugins/FX/Niagara/Shaders/Private/NiagaraRayTracingCommon.ush
@@ -8,8 +8,14 @@
 
 #include "/Engine/Private/RayTracing/RayTracingCommon.ush"
 
+// ray input structure for Niagara ray traces
+struct FNiagaraRayData : FBasicRayData
+{
+	uint CollisionGroup;
+};
+
 // payload structure for Niagara ray traces
-struct FNiagaraRayTracingPayload : FMinimalPayload
+struct FNiagaraRayTracingPayload : FDefaultPayload
 {
 	float3 WorldPosition;
 	float3 WorldNormal;
diff --git a/Engine/Plugins/FX/Niagara/Shaders/Private/NiagaraRayTracingShaders.usf b/Engine/Plugins/FX/Niagara/Shaders/Private/NiagaraRayTracingShaders.usf
--- a/Engine/Plugins/FX/Niagara/Shaders/Private/NiagaraRayTracingShaders.usf
+++ b/Engine/Plugins/FX/Niagara/Shaders/Private/NiagaraRayTracingShaders.usf
@@ -1,19 +1,56 @@
 // Copyright Epic Games, Inc. All Rights Reserved.
 
 #include "/Engine/Private/Common.ush"
+#include "/Engine/Private/SceneData.ush"
 #include "/Engine/Private/ShadingCommon.ush"
 #include "/Engine/Private/RayTracing/RayTracingCommon.ush"
 #include "/Engine/Private/RayTracing/RayTracingHitGroupCommon.ush"
 #include "NiagaraRayTracingCommon.ush"
+#include "/Engine/Private/HashTable.ush"
 
 RaytracingAccelerationStructure               TLAS;
-StructuredBuffer<FBasicRayData>               Rays;
+StructuredBuffer<FNiagaraRayData>               Rays;
 RWStructuredBuffer<FNiagaraRayTracingPayload> CollisionOutput;
 
 #if NIAGARA_RAYTRACE_FAKE_INDIRECT
 Buffer<uint>                                  RayTraceCounts;
 #endif
 
+Buffer<uint> HashToCollisionGroups;
+uint MaxRetraces;
+uint TestParam;
+uint TestControl;
+
+uint GetGPUSceneInstanceIndex(uint PrimitiveIndex, uint InstanceIndex)
+{
+#if VF_USE_PRIMITIVE_SCENE_DATA
+	FPrimitiveSceneData PrimData = GetPrimitiveData(PrimitiveIndex);
+	// Handle ray tracing auto instancing which can cause PrimitiveInstanceIndex to be above 0 for meshes with a single instance
+	if (PrimData.NumInstanceSceneDataEntries == 1)
+	{
+		InstanceIndex = 0;
+	}
+
+	return PrimData.InstanceSceneDataOffset + InstanceIndex;
+#else
+
+	return 0;
+#endif
+}
+
+bool FindCollisionGroup(int GPUSceneIndexId, out int CollisionGroup)
+{
+	int Index;
+	if (HashTableFind(GPUSceneIndexId, Index))
+	{
+		CollisionGroup = HashToCollisionGroups[Index];
+		return true;
+	}
+
+	return false;
+}
+
+
 RAY_TRACING_ENTRY_RAYGEN(NiagaraCollisionRayTraceRG)
 {
 	const uint RayIndex = DispatchRaysIndex().x;
@@ -25,7 +62,7 @@ RAY_TRACING_ENTRY_RAYGEN(NiagaraCollisionRayTraceRG)
 	}
 #endif
 
-	FBasicRayData InputRay = Rays[RayIndex];
+	FNiagaraRayData InputRay = Rays[RayIndex];
 
 	RayDesc Ray;
 	Ray.Origin = InputRay.Origin;
@@ -33,26 +70,69 @@ RAY_TRACING_ENTRY_RAYGEN(NiagaraCollisionRayTraceRG)
 	Ray.TMin = 0.0f;
 	Ray.TMax = InputRay.TFar;
 
-	uint RayFlags = RAY_FLAG_FORCE_OPAQUE; // don't run anyhit shader
+	uint RayFlags = RAY_FLAG_FORCE_OPAQUE | RAY_FLAG_CULL_BACK_FACING_TRIANGLES; // don't run anyhit shader
 
 	FNiagaraRayTracingPayload Payload = (FNiagaraRayTracingPayload)0;
-
-	TraceRay(
-		TLAS,   // AccelerationStructure
-		RayFlags,
-		RAY_TRACING_MASK_OPAQUE,          // InstanceInclusionMask
-		RAY_TRACING_SHADER_SLOT_MATERIAL, // RayContributionToHitGroupIndex
-		RAY_TRACING_NUM_SHADER_SLOTS,     // MultiplierForGeometryContributionToShaderIndex
-		0,      // MissShaderIndex
-		Ray,    // RayDesc
-		Payload // Payload
-	);
+	uint RayCollisionGroup = InputRay.CollisionGroup;
+	if (RayCollisionGroup != -1)
+	{
+		uint NumCasts = 0;
+		LOOP
+		while (NumCasts++ <= MaxRetraces)
+		{
+			Payload = (FNiagaraRayTracingPayload)0;
+			TraceRay(
+				TLAS,   // AccelerationStructure
+				RayFlags,
+				RAY_TRACING_MASK_OPAQUE,          // InstanceInclusionMask
+				RAY_TRACING_SHADER_SLOT_MATERIAL, // RayContributionToHitGroupIndex
+				RAY_TRACING_NUM_SHADER_SLOTS,     // MultiplierForGeometryContributionToShaderIndex
+				0,      // MissShaderIndex
+				Ray,    // RayDesc
+				Payload // Payload
+			);
+
+			uint GPUSceneInstanceIndex = GetGPUSceneInstanceIndex(Payload.PrimitiveIndex, Payload.InstanceIndex);
+			uint PrimCollisionGroup;
+			if (FindCollisionGroup(GPUSceneInstanceIndex, PrimCollisionGroup))
+			{
+				//If we collide with something of a different group to the Ray then break and allow the collision. Otherwise loop and attempt to recast just infront of the hit we're discarding.
+				if (PrimCollisionGroup == RayCollisionGroup)
+				{
+					Ray.TMin = asfloat(asuint(Payload.HitT) + 1);
+					//Clear the hit state on the payload in case we don't attempt a retrace
+					Payload.HitT = -1;
+				}
+				else
+				{
+					break;
+				}
+			}
+		}
+	}
+	else
+	{
+		TraceRay(
+			TLAS,   // AccelerationStructure
+			RayFlags,
+			RAY_TRACING_MASK_OPAQUE,          // InstanceInclusionMask
+			RAY_TRACING_SHADER_SLOT_MATERIAL, // RayContributionToHitGroupIndex
+			RAY_TRACING_NUM_SHADER_SLOTS,     // MultiplierForGeometryContributionToShaderIndex
+			0,      // MissShaderIndex
+			Ray,    // RayDesc
+			Payload // Payload
+		);
+	}
 
 	CollisionOutput[RayIndex].HitT = Payload.HitT;
 	if (CollisionOutput[RayIndex].IsHit())
 	{
 		CollisionOutput[RayIndex].WorldPosition = Payload.WorldPosition;
 		CollisionOutput[RayIndex].WorldNormal = Payload.WorldNormal;
+
+		CollisionOutput[RayIndex].PrimitiveIndex = Payload.PrimitiveIndex;
+		CollisionOutput[RayIndex].InstanceIndex = Payload.InstanceIndex;
+		CollisionOutput[RayIndex].InstanceID = Payload.InstanceID;
 	}
 }
 
@@ -69,14 +149,20 @@ RAY_TRACING_ENTRY_CLOSEST_HIT(NiagaraCollisionRayTraceCH,
 {
 	Payload.HitT = RayTCurrent();
 
+	const int Primitive = GetInstanceUserData();
+	Payload.PrimitiveIndex = Primitive;
+
+	const int PrimitiveInstanceIndex = InstanceIndex() - GetBaseInstanceIndex();
+	Payload.InstanceIndex = PrimitiveInstanceIndex;
+
 	FTriangleBaseAttributes Tri = LoadTriangleBaseAttributes(PrimitiveIndex());
 	float3 LocalPosition = GetBarycentricPosition(Attributes, Tri);
 	float3 LocalEdges[2] = { Tri.LocalPositions[1] - Tri.LocalPositions[0], Tri.LocalPositions[2] - Tri.LocalPositions[0] };
 	float3 LocalNormal = cross(LocalEdges[1], LocalEdges[0]);
 	float3x3 InverseTranspose3x3 = transpose((float3x3)WorldToObject4x3());
 
-	Payload.WorldPosition = mul(float4(LocalPosition, 1.0f), ObjectToWorld4x3()).xyz;
-	Payload.WorldNormal = normalize(mul(LocalNormal, InverseTranspose3x3));
+	Payload.WorldPosition = WorldRayOrigin() + WorldRayDirection() * RayTCurrent();
+	Payload.WorldNormal = WorldRayDirection();
 }
 
 RAY_TRACING_ENTRY_MISS(NiagaraCollisionRayTraceMiss, FNiagaraRayTracingPayload, Payload)
diff --git a/Engine/Plugins/FX/Niagara/Source/Niagara/Classes/NiagaraDataInterfaceCollisionQuery.h b/Engine/Plugins/FX/Niagara/Source/Niagara/Classes/NiagaraDataInterfaceCollisionQuery.h
--- a/Engine/Plugins/FX/Niagara/Source/Niagara/Classes/NiagaraDataInterfaceCollisionQuery.h
+++ b/Engine/Plugins/FX/Niagara/Source/Niagara/Classes/NiagaraDataInterfaceCollisionQuery.h
@@ -32,6 +32,10 @@ public:
 	UPROPERTY(EditAnywhere, Category = "Ray Trace")
 	int32 MaxRayTraceCount = 0;
 
+	/** If a collision is rejected, how many times do we attempt to retrace from that collision point forward to find a new, valid collision.*/
+	UPROPERTY(EditAnywhere, Category = "Ray Trace")
+	int32 MaxRetraces = 0;
+
 	FNiagaraSystemInstance *SystemInstance;
 
 	//UObject Interface
diff --git a/Engine/Plugins/FX/Niagara/Source/Niagara/Classes/NiagaraEmitterInstanceBatcher.h b/Engine/Plugins/FX/Niagara/Source/Niagara/Classes/NiagaraEmitterInstanceBatcher.h
--- a/Engine/Plugins/FX/Niagara/Source/Niagara/Classes/NiagaraEmitterInstanceBatcher.h
+++ b/Engine/Plugins/FX/Niagara/Source/Niagara/Classes/NiagaraEmitterInstanceBatcher.h
@@ -158,10 +158,18 @@ public:
 	FNiagaraGpuReadbackManager* GetGpuReadbackManager() const { return GpuReadbackManagerPtr.Get(); }
 
 #if RHI_RAYTRACING
-	void IssueRayTraces(FRHICommandList& RHICmdList, const FIntPoint& RayTraceCounts, FRHIShaderResourceView* RayTraceRequests, FRWBuffer* IndirectArgsBuffer, uint32 IndirectArgsOffset, FRHIUnorderedAccessView* RayTraceResults) const;
+	void IssueRayTraces(FRHICommandList& RHICmdList, const FIntPoint& RayTraceCounts, uint32 MaxRetraces, FRHIShaderResourceView* RayTraceRequests, FRWBuffer* IndirectArgsBuffer, uint32 IndirectArgsOffset, FRHIUnorderedAccessView* RayTraceResults) const;
 	bool HasRayTracingScene() const;
+
+	void SetPrimitiveRayTracingCollisionGroup(UPrimitiveComponent* PrimitiveComponent, uint32 Group);
+
+	int32 AcquireGPURayTracedCollisionGroup(); 
+	void ReleaseGPURayTracedCollisionGroup(int32 CollisionGroup);
 #endif
 
+	void OnPrimitiveGPUSceneInstancesAllocated();
+	void OnPrimitiveGPUSceneInstancesFreed();
+
 private:
 	void DumpDebugFrame();
 	void UpdateInstanceCountManager(FRHICommandListImmediate& RHICmdList);
@@ -299,4 +307,7 @@ private:
 	void BroadcastTemporalEffect(FRHICommandList& RHICmdList);
 #endif // WITH_MGPU
 
+	/** Pool of free GPU ray traced collision groups. */
+	TArray<int32> FreeGPURayTracedCollisionGroups;
+	int32 NumGPURayTracedCollisionGroups = 0;
 };
diff --git a/Engine/Plugins/FX/Niagara/Source/Niagara/Private/NiagaraDataInterfaceCollisionQuery.cpp b/Engine/Plugins/FX/Niagara/Source/Niagara/Private/NiagaraDataInterfaceCollisionQuery.cpp
--- a/Engine/Plugins/FX/Niagara/Source/Niagara/Private/NiagaraDataInterfaceCollisionQuery.cpp
+++ b/Engine/Plugins/FX/Niagara/Source/Niagara/Private/NiagaraDataInterfaceCollisionQuery.cpp
@@ -63,6 +63,8 @@ struct FNiagaraDataIntefaceProxyCollisionQuery : public FNiagaraDataInterfacePro
 	FRWBuffer RayTraceCounts;
 
 	int32 MaxRayTraceCount = 0;
+
+	uint32 MaxRetraces = 0;
 #endif
 
 	virtual ~FNiagaraDataIntefaceProxyCollisionQuery()
@@ -132,7 +134,7 @@ struct FNiagaraDataIntefaceProxyCollisionQuery : public FNiagaraDataInterfacePro
 				RHICmdList.Transition(PreTransitions);
 			}
 
-			Context.Batcher->IssueRayTraces(RHICmdList, FIntPoint(MaxRayTraceCount, 1), RayTraceRequests.SRV, &RayTraceCounts, 0, RayTraceIntersections.UAV);
+			Context.Batcher->IssueRayTraces(RHICmdList, FIntPoint(MaxRayTraceCount, 1), MaxRetraces, RayTraceRequests.SRV, &RayTraceCounts, 0, RayTraceIntersections.UAV);
 
 			{
 				FRHITransitionInfo PostTransitions[] =
@@ -147,10 +149,11 @@ struct FNiagaraDataIntefaceProxyCollisionQuery : public FNiagaraDataInterfacePro
 #endif
 	}
 
-	void RenderThreadInitialize(int32 InMaxRayTraceRequests)
+	void RenderThreadInitialize(int32 InMaxRayTraceRequests, uint32 InMaxRetraces)
 	{
 #if RHI_RAYTRACING
 		MaxRayTraceCount = 0;
+		MaxRetraces = InMaxRetraces;
 
 		DEC_MEMORY_STAT_BY(STAT_NiagaraGPUDataInterfaceMemory, RayTraceRequests.NumBytes);
 		RayTraceRequests.Release();
@@ -167,7 +170,7 @@ struct FNiagaraDataIntefaceProxyCollisionQuery : public FNiagaraDataInterfacePro
 
 			RayTraceRequests.Initialize(
 				TEXT("NiagaraRayTraceRequests"),
-				sizeof(FBasicRayData),
+				sizeof(FNiagaraRayData),
 				MaxRayTraceCount,
 				BUF_Static,
 				false /*bUseUavCounter*/,
@@ -956,6 +959,7 @@ bool UNiagaraDataInterfaceCollisionQuery::Equals(const UNiagaraDataInterface* Ot
 
 	const UNiagaraDataInterfaceCollisionQuery* OtherTyped = CastChecked<const UNiagaraDataInterfaceCollisionQuery>(Other);
 	return OtherTyped->MaxRayTraceCount == MaxRayTraceCount;
+	return OtherTyped->MaxRetraces == MaxRetraces;
 }
 
 bool UNiagaraDataInterfaceCollisionQuery::CopyToInternal(UNiagaraDataInterface* Destination) const
@@ -967,6 +971,7 @@ bool UNiagaraDataInterfaceCollisionQuery::CopyToInternal(UNiagaraDataInterface*
 
 	UNiagaraDataInterfaceCollisionQuery* OtherTyped = CastChecked<UNiagaraDataInterfaceCollisionQuery>(Destination);
 	OtherTyped->MaxRayTraceCount = MaxRayTraceCount;
+	OtherTyped->MaxRetraces = MaxRetraces;
 	OtherTyped->MarkRenderDataDirty();
 	return true;
 }
@@ -977,9 +982,9 @@ void UNiagaraDataInterfaceCollisionQuery::PushToRenderThreadImpl()
 
 	// Push Updates to Proxy, first release any resources
 	ENQUEUE_RENDER_COMMAND(FUpdateDI)(
-		[RT_Proxy, RT_MaxRayTraceRequests = MaxRayTraceCount](FRHICommandListImmediate& RHICmdList)
+		[RT_Proxy, RT_MaxRayTraceRequests = MaxRayTraceCount, RT_MaxRetraces = MaxRetraces](FRHICommandListImmediate& RHICmdList)
 		{
-			RT_Proxy->RenderThreadInitialize(RT_MaxRayTraceRequests);
+			RT_Proxy->RenderThreadInitialize(RT_MaxRayTraceRequests, RT_MaxRetraces);
 		});
 }
 
diff --git a/Engine/Plugins/FX/Niagara/Source/Niagara/Private/NiagaraEmitterInstanceBatcher.cpp b/Engine/Plugins/FX/Niagara/Source/Niagara/Private/NiagaraEmitterInstanceBatcher.cpp
--- a/Engine/Plugins/FX/Niagara/Source/Niagara/Private/NiagaraEmitterInstanceBatcher.cpp
+++ b/Engine/Plugins/FX/Niagara/Source/Niagara/Private/NiagaraEmitterInstanceBatcher.cpp
@@ -262,6 +262,9 @@ NiagaraEmitterInstanceBatcher::NiagaraEmitterInstanceBatcher(ERHIFeatureLevel::T
 	GpuComputeDebugPtr.Reset(new FNiagaraGpuComputeDebug(FeatureLevel));
 #endif
 	GpuReadbackManagerPtr.Reset(new FNiagaraGpuReadbackManager());
+
+	FPrimitiveSceneInfo::OnGPUSceneInstancesAllocated.AddRaw(this, &NiagaraEmitterInstanceBatcher::OnPrimitiveGPUSceneInstancesAllocated);
+	FPrimitiveSceneInfo::OnGPUSceneInstancesAllocated.AddRaw(this, &NiagaraEmitterInstanceBatcher::OnPrimitiveGPUSceneInstancesFreed);
 }
 
 NiagaraEmitterInstanceBatcher::~NiagaraEmitterInstanceBatcher()
@@ -272,6 +275,9 @@ NiagaraEmitterInstanceBatcher::~NiagaraEmitterInstanceBatcher()
 	SystemCBufferLayout = nullptr;
 	OwnerCBufferLayout = nullptr;
 	EmitterCBufferLayout = nullptr;
+
+	FPrimitiveSceneInfo::OnGPUSceneInstancesAllocated.RemoveAll(this);
+	FPrimitiveSceneInfo::OnGPUSceneInstancesFreed.RemoveAll(this);
 }
 
 void NiagaraEmitterInstanceBatcher::AddGpuComputeProxy(FNiagaraSystemGpuComputeProxy* ComputeProxy)
@@ -1993,6 +1999,8 @@ void NiagaraEmitterInstanceBatcher::BuildRayTracingSceneInfo(FRHICommandList& RH
 	if (NumProxiesThatRequireRayTracingScene > 0)
 	{
 		RayTracingHelper->BuildRayTracingSceneInfo(RHICmdList, Views);
+
+		RayTracingHelper->UpdateCollisionGroupMap(RHICmdList, GetScene(), FeatureLevel);
 	}
 	else
 	{
@@ -2005,19 +2013,68 @@ void NiagaraEmitterInstanceBatcher::ResetRayTracingSceneInfo()
 	RayTracingHelper->Reset();
 }
 
-void NiagaraEmitterInstanceBatcher::IssueRayTraces(FRHICommandList& RHICmdList, const FIntPoint& RayTraceCounts, FRHIShaderResourceView* RayTraceRequests, FRWBuffer* IndirectArgsBuffer, uint32 IndirectArgsOffset, FRHIUnorderedAccessView* RayTraceResults) const
+void NiagaraEmitterInstanceBatcher::IssueRayTraces(FRHICommandList& RHICmdList, const FIntPoint& RayTraceCounts, uint32 MaxRetraces, FRHIShaderResourceView* RayTraceRequests, FRWBuffer* IndirectArgsBuffer, uint32 IndirectArgsOffset, FRHIUnorderedAccessView* RayTraceResults) const
 {
 	check(NumProxiesThatRequireRayTracingScene > 0);
 
-	RayTracingHelper->IssueRayTraces(RHICmdList, RayTraceCounts, RayTraceRequests, IndirectArgsBuffer, IndirectArgsOffset, RayTraceResults);
+	RayTracingHelper->IssueRayTraces(RHICmdList, GetScene(), RayTraceCounts, MaxRetraces, RayTraceRequests, IndirectArgsBuffer, IndirectArgsOffset, RayTraceResults);
 }
 
 bool NiagaraEmitterInstanceBatcher::HasRayTracingScene() const
 {
 	return NumProxiesThatRequireRayTracingScene && RayTracingHelper->IsValid();
 }
+
+void NiagaraEmitterInstanceBatcher::SetPrimitiveRayTracingCollisionGroup(UPrimitiveComponent* Primitive, uint32 Group)
+{
+	check(IsInGameThread());
+	if (Primitive == nullptr || Primitive->SceneProxy == nullptr)
+	{
+		return;
+	}
+
+	ENQUEUE_RENDER_COMMAND(NiagaraSetPrimHWRTCollisionGroup)(
+		[RT_NiagaraBatcher = this, Proxy = Primitive->SceneProxy, CollisionGroup=Group](FRHICommandListImmediate& RHICmdList)
+		{
+			check(Proxy);
+			RT_NiagaraBatcher->RayTracingHelper->SetPrimitiveCollisionGroup(*Proxy->GetPrimitiveSceneInfo(), CollisionGroup);
+		}
+	);
+}
+
+int32 NiagaraEmitterInstanceBatcher::AcquireGPURayTracedCollisionGroup()
+{
+	check(IsInGameThread());//one of the few batcher functions that should be called from the GT
+	if (FreeGPURayTracedCollisionGroups.Num() > 0)
+	{
+		return FreeGPURayTracedCollisionGroups.Pop();
+	}
+
+	return NumGPURayTracedCollisionGroups++;
+}
+
+void NiagaraEmitterInstanceBatcher::ReleaseGPURayTracedCollisionGroup(int32 CollisionGroup)
+{
+	check(IsInGameThread());//one of the few batcher functions that should be called from the GT
+	FreeGPURayTracedCollisionGroups.Add(CollisionGroup);
+}
+
 #endif
 
+void NiagaraEmitterInstanceBatcher::OnPrimitiveGPUSceneInstancesAllocated()
+{
+#if RHI_RAYTRACING
+	RayTracingHelper->RefreshPrimitiveInstanceData();
+#endif
+}
+
+void NiagaraEmitterInstanceBatcher::OnPrimitiveGPUSceneInstancesFreed()
+{
+#if RHI_RAYTRACING
+	RayTracingHelper->RefreshPrimitiveInstanceData();
+#endif
+}
+
 /* Set shader parameters for data interfaces
  */
 void NiagaraEmitterInstanceBatcher::SetDataInterfaceParameters(FRHICommandList& RHICmdList, const FNiagaraGPUSystemTick& Tick, const FNiagaraComputeInstanceData& InstanceData, const FNiagaraShaderRef& ComputeShader, const FNiagaraSimStageData& SimStageData) const
@@ -2256,4 +2313,4 @@ void NiagaraEmitterInstanceBatcher::BroadcastTemporalEffect(FRHICommandList& RHI
 		}
 	}
 }
-#endif // WITH_MGPU
\ No newline at end of file
+#endif // WITH_MGPU
diff --git a/Engine/Plugins/FX/Niagara/Source/Niagara/Private/NiagaraFunctionLibrary.cpp b/Engine/Plugins/FX/Niagara/Source/Niagara/Private/NiagaraFunctionLibrary.cpp
--- a/Engine/Plugins/FX/Niagara/Source/Niagara/Private/NiagaraFunctionLibrary.cpp
+++ b/Engine/Plugins/FX/Niagara/Source/Niagara/Private/NiagaraFunctionLibrary.cpp
@@ -18,6 +18,7 @@
 #include "NiagaraDataInterfaceVolumeTexture.h"
 #include "Engine/VolumeTexture.h"
 #include "Engine/Texture2DArray.h"
+#include "NiagaraEmitterInstanceBatcher.h"
 
 #define LOCTEXT_NAMESPACE "NiagaraFunctionLibrary"
 
@@ -1769,4 +1770,92 @@ void UNiagaraFunctionLibrary::InitVectorVMFastPathOps()
 	check(VectorVMOps.Num() == VectorVMOpsHLSL.Num());
 }
 
+
+//////////////////////////////////////////////////////////////////////////
+//GPU Ray Traced Collision Functions
+
+void UNiagaraFunctionLibrary::SetComponentNiagaraGPURayTracedCollisionGroup(UObject* WorldContextObject, UPrimitiveComponent* Primitive, int32 CollisionGroup)
+{
+#if WITH_RAYTRACING
+	if(UWorld* World = GEngine->GetWorldFromContextObject(WorldContextObject, EGetWorldErrorMode::LogAndReturnNull))
+	{
+		if(FNiagaraWorldManager* WorldMan = FNiagaraWorldManager::Get(World))
+		{
+			if(FFXSystemInterface* FXSystemInterface = World->Scene->GetFXSystem())			
+			{
+				if (NiagaraEmitterInstanceBatcher* Batcher = static_cast<NiagaraEmitterInstanceBatcher*>(FXSystemInterface->GetInterface(NiagaraEmitterInstanceBatcher::Name)))
+				{
+					Batcher->SetPrimitiveRayTracingCollisionGroup(Primitive, CollisionGroup);
+				}
+			}
+		}
+	}
+#endif
+}
+
+void UNiagaraFunctionLibrary::SetActorNiagaraGPURayTracedCollisionGroup(UObject* WorldContextObject, AActor* Actor, int32 CollisionGroup)
+{
+#if WITH_RAYTRACING
+	if (Actor == nullptr)
+	{
+		return;
+	}
+
+	if (UWorld* World = GEngine->GetWorldFromContextObject(WorldContextObject, EGetWorldErrorMode::LogAndReturnNull))
+	{
+		if (FNiagaraWorldManager* WorldMan = FNiagaraWorldManager::Get(World))
+		{
+			if (FFXSystemInterface* FXSystemInterface = World->Scene->GetFXSystem())
+			{
+				if (NiagaraEmitterInstanceBatcher* Batcher = static_cast<NiagaraEmitterInstanceBatcher*>(FXSystemInterface->GetInterface(NiagaraEmitterInstanceBatcher::Name)))
+				{
+					Actor->ForEachComponent<UPrimitiveComponent>(/*bIncludeFromChildActors*/true, [&Batcher, &CollisionGroup](UPrimitiveComponent* PrimitiveComponent)
+					{
+							Batcher->SetPrimitiveRayTracingCollisionGroup(PrimitiveComponent, CollisionGroup);
+					});
+				}
+			}
+		}
+	}
+#endif
+}
+
+int32 UNiagaraFunctionLibrary::AcquireNiagaraGPURayTracedCollisionGroup(UObject* WorldContextObject)
+{
+#if WITH_RAYTRACING
+	if (UWorld* World = GEngine->GetWorldFromContextObject(WorldContextObject, EGetWorldErrorMode::LogAndReturnNull))
+	{
+		if (FNiagaraWorldManager* WorldMan = FNiagaraWorldManager::Get(World))
+		{
+			if (FFXSystemInterface* FXSystemInterface = World->Scene->GetFXSystem())
+			{
+				if (NiagaraEmitterInstanceBatcher* Batcher = static_cast<NiagaraEmitterInstanceBatcher*>(FXSystemInterface->GetInterface(NiagaraEmitterInstanceBatcher::Name)))
+				{
+					return Batcher->AcquireGPURayTracedCollisionGroup();
+				}
+			}
+		}
+	}
+#endif
+	return INDEX_NONE;
+}
+
+void UNiagaraFunctionLibrary::ReleaseNiagaraGPURayTracedCollisionGroup(UObject* WorldContextObject, int32 CollisionGroup)
+{
+#if WITH_RAYTRACING
+	if (UWorld* World = GEngine->GetWorldFromContextObject(WorldContextObject, EGetWorldErrorMode::LogAndReturnNull))
+	{
+		if (FNiagaraWorldManager* WorldMan = FNiagaraWorldManager::Get(World))
+		{
+			if (FFXSystemInterface* FXSystemInterface = World->Scene->GetFXSystem())
+			{
+				if (NiagaraEmitterInstanceBatcher* Batcher = static_cast<NiagaraEmitterInstanceBatcher*>(FXSystemInterface->GetInterface(NiagaraEmitterInstanceBatcher::Name)))
+				{
+					Batcher->ReleaseGPURayTracedCollisionGroup(CollisionGroup);
+				}
+			}
+		}
+	}
+#endif
+}
 #undef LOCTEXT_NAMESPACE
\ No newline at end of file
diff --git a/Engine/Plugins/FX/Niagara/Source/Niagara/Public/NiagaraFunctionLibrary.h b/Engine/Plugins/FX/Niagara/Source/Niagara/Public/NiagaraFunctionLibrary.h
--- a/Engine/Plugins/FX/Niagara/Source/Niagara/Public/NiagaraFunctionLibrary.h
+++ b/Engine/Plugins/FX/Niagara/Source/Niagara/Public/NiagaraFunctionLibrary.h
@@ -90,6 +90,24 @@ public:
 
 	static bool GetVectorVMFastPathExternalFunction(const FVMExternalFunctionBindingInfo& BindingInfo, FVMExternalFunction &OutFunc);
 
+
+	//Functions providing access to HWRT collision specific features
+
+	/** Sets the Niagara GPU ray traced collision group for the give primitive component. */
+	UFUNCTION(BlueprintCallable, Category = Niagara, meta = (Keywords = "niagara collision ray tracing", WorldContext = "WorldContextObject"))
+	static void SetComponentNiagaraGPURayTracedCollisionGroup(UObject* WorldContextObject, UPrimitiveComponent* Primitive, int32 CollisionGroup);
+	
+	/** Sets the Niagara GPU ray traced collision group for all primitive components on the given actor. */
+	UFUNCTION(BlueprintCallable, Category = Niagara, meta = (Keywords = "niagara collision ray tracing", WorldContext = "WorldContextObject"))
+	static void SetActorNiagaraGPURayTracedCollisionGroup(UObject* WorldContextObject, AActor* Actor, int32 CollisionGroup);
+
+	/** Returns a free collision group for use in HWRT collision group filtering. Returns -1 on failure. */
+	UFUNCTION(BlueprintCallable, Category = Niagara, meta = (Keywords = "niagara collision ray tracing", WorldContext = "WorldContextObject"))
+	static int32 AcquireNiagaraGPURayTracedCollisionGroup(UObject* WorldContextObject);
+
+	/** Releases a collision group back to the system for use by ohers. */
+	UFUNCTION(BlueprintCallable, Category = Niagara, meta = (Keywords = "niagara collision ray tracing", WorldContext = "WorldContextObject"))
+	static void ReleaseNiagaraGPURayTracedCollisionGroup(UObject* WorldContextObject, int32 CollisionGroup);
 private:
 	static void InitVectorVMFastPathOps();
 	static TArray<FNiagaraFunctionSignature> VectorVMOps;
diff --git a/Engine/Plugins/FX/Niagara/Source/NiagaraShader/Private/NiagaraRayTracingHelper.cpp b/Engine/Plugins/FX/Niagara/Source/NiagaraShader/Private/NiagaraRayTracingHelper.cpp
--- a/Engine/Plugins/FX/Niagara/Source/NiagaraShader/Private/NiagaraRayTracingHelper.cpp
+++ b/Engine/Plugins/FX/Niagara/Source/NiagaraShader/Private/NiagaraRayTracingHelper.cpp
@@ -10,20 +10,66 @@
 #include "Renderer/Private/ScenePrivate.h"
 #include "Renderer/Private/SceneRendering.h"
 #include "ShaderParameterStruct.h"
+#include "NiagaraShaderParticleID.h"
+
+//////////////////////////////////////////////////////////////////////////
+
+class FNiagaraUpdateCollisionGroupMapCS : public FGlobalShader
+{
+	DECLARE_GLOBAL_SHADER(FNiagaraUpdateCollisionGroupMapCS);
+	SHADER_USE_PARAMETER_STRUCT(FNiagaraUpdateCollisionGroupMapCS, FGlobalShader);
+
+	BEGIN_SHADER_PARAMETER_STRUCT(FParameters, )
+		SHADER_PARAMETER_UAV(RWStructuredBuffer<UINT>, RWHashTable)		
+		SHADER_PARAMETER(uint32, HashTableSize)
+		SHADER_PARAMETER_UAV(Buffer<UINT>, RWHashToCollisionGroups)
+		SHADER_PARAMETER_SRV(Buffer<UINT>, NewPrimIdCollisionGroupPairs)
+		SHADER_PARAMETER(uint32, NumNewPrims)
+	END_SHADER_PARAMETER_STRUCT()
+public:
+
+	static bool ShouldCompilePermutation(const FGlobalShaderPermutationParameters& Parameters)
+	{
+		return RHISupportsComputeShaders(Parameters.Platform);
+	}
+
+	static void ModifyCompilationEnvironment(const FGlobalShaderPermutationParameters& Parameters, FShaderCompilerEnvironment& OutEnvironment)
+	{
+		FGlobalShader::ModifyCompilationEnvironment(Parameters, OutEnvironment);
+		OutEnvironment.SetDefine(TEXT("THREAD_COUNT"), THREAD_COUNT);
+	}
+
+	static constexpr uint32 THREAD_COUNT = 64;
+};
+IMPLEMENT_GLOBAL_SHADER(FNiagaraUpdateCollisionGroupMapCS, "/Plugin/FX/Niagara/Private/NiagaraRayTraceCollisionGroupShaders.usf", "UpdatePrimIdToCollisionGroupMap", SF_Compute);
+
+//////////////////////////////////////////////////////////////////////////
 
 /// TODO
 ///  -get geometry masking working when an environmental mask is implemented
 
+BEGIN_SHADER_PARAMETER_STRUCT(FGPUSceneParameters, )
+SHADER_PARAMETER_SRV(StructuredBuffer<float4>, GPUSceneInstanceSceneData)
+SHADER_PARAMETER_SRV(StructuredBuffer<float4>, GPUScenePrimitiveSceneData)
+SHADER_PARAMETER(uint32, GPUSceneFrameNumber)
+END_SHADER_PARAMETER_STRUCT()
+
 class FNiagaraCollisionRayTraceRG : public FGlobalShader
 {
 	DECLARE_GLOBAL_SHADER(FNiagaraCollisionRayTraceRG)
 	SHADER_USE_ROOT_PARAMETER_STRUCT(FNiagaraCollisionRayTraceRG, FGlobalShader)
 
 	BEGIN_SHADER_PARAMETER_STRUCT(FParameters, )
-		SHADER_PARAMETER_RDG_BUFFER_SRV(RaytracingAccelerationStructure, TLAS)
-		SHADER_PARAMETER_RDG_BUFFER_SRV(StructuredBuffer<FBasicRayData>, Rays)
-		SHADER_PARAMETER_RDG_BUFFER_UAV(RWStructuredBuffer<FNiagaraRayTracingPayload>, CollisionOutput)
-		SHADER_PARAMETER_RDG_BUFFER_SRV(Buffer<UINT>, RayTraceCounts)
+		SHADER_PARAMETER_STRUCT_INCLUDE(FGPUSceneParameters, GPUSceneParameters)
+		SHADER_PARAMETER_SRV(StructuredBuffer<UINT>, HashTable)
+		SHADER_PARAMETER_UAV(RWStructuredBuffer<UINT>,RWHashTable)
+		SHADER_PARAMETER(uint32, HashTableSize)
+		SHADER_PARAMETER_SRV(RaytracingAccelerationStructure, TLAS)
+		SHADER_PARAMETER_SRV(StructuredBuffer<FNiagaraRayData>, Rays)
+		SHADER_PARAMETER_UAV(RWStructuredBuffer<FNiagaraRayTracingPayload>, CollisionOutput)
+		SHADER_PARAMETER_SRV(Buffer<UINT>, RayTraceCounts)
+		SHADER_PARAMETER_SRV(Buffer<UINT>, HashToCollisionGroups)
+		SHADER_PARAMETER(uint32, MaxRetraces)
 	END_SHADER_PARAMETER_STRUCT()
 
 	class FFakeIndirectDispatch : SHADER_PERMUTATION_BOOL("NIAGARA_RAYTRACE_FAKE_INDIRECT");
@@ -38,14 +84,21 @@ class FNiagaraCollisionRayTraceRG : public FGlobalShader
 	{
 		FGlobalShader::ModifyCompilationEnvironment(Parameters, OutEnvironment);
 		OutEnvironment.SetDefine(TEXT("NIAGARA_SUPPORTS_RAY_TRACING"), 1);
+		OutEnvironment.SetDefine(TEXT("VF_SUPPORTS_PRIMITIVE_SCENE_DATA"), 1);
+		OutEnvironment.SetDefine(TEXT("USE_GLOBAL_GPU_SCENE_DATA"), 1);
 	}
 
-	static FRHIRayTracingShader* GetShader(FGlobalShaderMap* ShaderMap)
+	static TShaderRef< FNiagaraCollisionRayTraceRG> GetShader(FGlobalShaderMap* ShaderMap)
 	{
 		FPermutationDomain PermutationVector;
 		PermutationVector.Set<FFakeIndirectDispatch>(!SupportsIndirectDispatch());
 
-		return ShaderMap->GetShader<FNiagaraCollisionRayTraceRG>(PermutationVector).GetRayTracingShader();
+		return ShaderMap->GetShader<FNiagaraCollisionRayTraceRG>(PermutationVector);
+	}
+
+	static FRHIRayTracingShader* GetRayTracingShader(FGlobalShaderMap* ShaderMap)
+	{
+		return GetShader(ShaderMap).GetRayTracingShader();
 	}
 
 	static bool SupportsIndirectDispatch()
@@ -190,6 +243,104 @@ static void BindNiagaraRayTracingMeshCommands(
 		bCopyDataToInlineStorage);
 }
 
+void FNiagaraRayTracingHelper::SetPrimitiveCollisionGroup(FPrimitiveSceneInfo& Primitive, uint32 CollisionGroup)
+{
+	CollisionGroupMap.FindOrAdd(Primitive.PrimitiveComponentId) = CollisionGroup;
+	bCollisionGroupMapDirty = true;
+}
+
+void FNiagaraRayTracingHelper::RefreshPrimitiveInstanceData()
+{
+	bCollisionGroupMapDirty = true;
+}
+
+void FNiagaraRayTracingHelper::UpdateCollisionGroupMap(FRHICommandList& RHICmdList, FScene* Scene, ERHIFeatureLevel::Type FeatureLevel)
+{
+	if (bCollisionGroupMapDirty)
+	{
+		SCOPED_DRAW_EVENT(RHICmdList, NiagaraUpdateCollisionGroupsMap);
+		bCollisionGroupMapDirty = false;
+
+		uint32 MinBufferInstances = FNiagaraUpdateCollisionGroupMapCS::THREAD_COUNT;
+		uint32 NeededInstances = FMath::Max((uint32)CollisionGroupMap.Num(), MinBufferInstances);
+		uint32 AllocInstances = Align(NeededInstances, FNiagaraUpdateCollisionGroupMapCS::THREAD_COUNT);
+
+		//We can probably be smarter here but for now just push the whole map over to the GPU each time it's dirty.
+		//Ideally this should be a small amount of data and updated infrequently.
+		FReadBuffer NewPrimIdCollisionGroupPairs;		
+		NewPrimIdCollisionGroupPairs.Initialize(TEXT("NewPrimIdCollisionGroupPairs"), sizeof(uint32) * 2, AllocInstances, EPixelFormat::PF_R32G32_UINT, BUF_Volatile);
+
+		uint32* PrimIdCollisionGroupPairPtr = (uint32*)RHILockBuffer(NewPrimIdCollisionGroupPairs.Buffer, 0, AllocInstances * sizeof(uint32) * 2, RLM_WriteOnly);
+		FMemory::Memset(PrimIdCollisionGroupPairPtr, 0, AllocInstances * sizeof(uint32) * 2);
+
+		for (auto Entry : CollisionGroupMap)
+		{
+			FPrimitiveComponentId PrimId = Entry.Key;
+			uint32 CollisionGroup = Entry.Value;
+			
+			uint32 GPUSceneInstanceIndex = INDEX_NONE;
+			//Ugh this is a bit pants. Maybe try to rework things so I can use the direct prim index.
+			int32 PrimIndex = Scene->PrimitiveComponentIds.Find(PrimId);
+			if (PrimIndex != INDEX_NONE)
+			{
+				GPUSceneInstanceIndex = Scene->Primitives[PrimIndex]->GetInstanceSceneDataOffset();
+			}
+
+			PrimIdCollisionGroupPairPtr[0] = GPUSceneInstanceIndex;
+			PrimIdCollisionGroupPairPtr[1] = CollisionGroup;
+			PrimIdCollisionGroupPairPtr += 2;
+		}
+		RHIUnlockBuffer(NewPrimIdCollisionGroupPairs.Buffer);
+
+		//Init the hash table if needed
+		if (HashTableSize < AllocInstances)
+		{
+			HashTableSize = AllocInstances;
+
+			PrimIdHashTable.Release();
+			PrimIdHashTable.Initialize(
+				TEXT("NiagaraPrimIdHashTable"),
+				sizeof(uint32),
+				AllocInstances,
+				BUF_Static,
+				false /*bUseUavCounter*/,
+				false /*bAppendBuffer*/,
+				ERHIAccess::UAVCompute);
+
+			HashToCollisionGroups.Release();
+			HashToCollisionGroups.Initialize(TEXT("NiagaraPrimIdHashToCollisionGroups"), sizeof(uint32), AllocInstances, EPixelFormat::PF_R32_UINT, BUF_Static);
+		}
+
+		RHICmdList.Transition(FRHITransitionInfo(PrimIdHashTable.UAV, ERHIAccess::Unknown, ERHIAccess::UAVCompute));
+		RHICmdList.Transition(FRHITransitionInfo(HashToCollisionGroups.UAV, ERHIAccess::Unknown, ERHIAccess::UAVCompute));
+
+		//First we have to clear the buffers. Can probably do this better.
+		NiagaraFillGPUIntBuffer(RHICmdList, FeatureLevel, PrimIdHashTable, 0);
+		NiagaraFillGPUIntBuffer(RHICmdList, FeatureLevel, HashToCollisionGroups, INDEX_NONE);
+
+		TShaderMapRef<FNiagaraUpdateCollisionGroupMapCS> ComputeShader(GetGlobalShaderMap(GMaxRHIFeatureLevel));
+		FRHIComputeShader* ShaderRHI = ComputeShader.GetComputeShader();
+
+		FNiagaraUpdateCollisionGroupMapCS::FParameters Params;
+		Params.RWHashTable = PrimIdHashTable.UAV;
+		Params.HashTableSize = HashTableSize;
+		Params.RWHashToCollisionGroups = HashToCollisionGroups.UAV;
+		Params.NewPrimIdCollisionGroupPairs = NewPrimIdCollisionGroupPairs.SRV;
+		Params.NumNewPrims = CollisionGroupMap.Num();
+
+		// To simplify the shader code, the size of the table must be a multiple of the thread count.
+		check(HashTableSize % FNiagaraUpdateCollisionGroupMapCS::THREAD_COUNT == 0);
+
+		RHICmdList.SetComputeShader(ShaderRHI);
+		SetShaderParameters(RHICmdList, ComputeShader, ShaderRHI, Params);
+		RHICmdList.DispatchComputeShader(FMath::DivideAndRoundUp(CollisionGroupMap.Num(), (int32)FNiagaraUpdateCollisionGroupMapCS::THREAD_COUNT), 1, 1);
+		UnsetShaderUAVs(RHICmdList, ComputeShader, ShaderRHI);
+
+		RHICmdList.Transition(FRHITransitionInfo(PrimIdHashTable.UAV, ERHIAccess::UAVCompute, ERHIAccess::SRVCompute));
+		RHICmdList.Transition(FRHITransitionInfo(HashToCollisionGroups.UAV, ERHIAccess::UAVCompute, ERHIAccess::SRVCompute));
+	}
+}
+
 void FNiagaraRayTracingHelper::BuildRayTracingSceneInfo(FRHICommandList& RHICmdList, TConstArrayView<FViewInfo> Views)
 {
 	check(Views.Num() > 0);
@@ -204,7 +355,7 @@ void FNiagaraRayTracingHelper::BuildRayTracingSceneInfo(FRHICommandList& RHICmdL
 		ViewUniformBuffer = ReferenceView.ViewUniformBuffer;
 
 		FGlobalShaderMap* ShaderMap = GetGlobalShaderMap(ShaderPlatform);
-		auto RayGenShader = FNiagaraCollisionRayTraceRG::GetShader(ShaderMap);
+		auto RayGenShader = FNiagaraCollisionRayTraceRG::GetRayTracingShader(ShaderMap);
 		auto ClosestHitShader = ShaderMap->GetShader<FNiagaraCollisionRayTraceCH>().GetRayTracingShader();
 		auto MissShader = ShaderMap->GetShader<FNiagaraCollisionRayTraceMiss>().GetRayTracingShader();
 
@@ -234,18 +385,33 @@ void FNiagaraRayTracingHelper::BuildRayTracingSceneInfo(FRHICommandList& RHICmdL
 	}
 }
 
-void FNiagaraRayTracingHelper::IssueRayTraces(FRHICommandList& RHICmdList, const FIntPoint& RayTraceCounts, FRHIShaderResourceView* RayTraceRequests, FRWBuffer* IndirectArgsBuffer, uint32 IndirectArgsOffset, FRHIUnorderedAccessView* RayTraceResults) const
+void FNiagaraRayTracingHelper::IssueRayTraces(FRHICommandList& RHICmdList, FScene* Scene, const FIntPoint& RayTraceCounts, uint32 MaxRetraces, FRHIShaderResourceView* RayTraceRequests, FRWBuffer* IndirectArgsBuffer, uint32 IndirectArgsOffset, FRHIUnorderedAccessView* RayTraceResults) const
 {
+	SCOPED_DRAW_EVENT(RHICmdList, NiagaraIssueCollisionRayTraces);
 	check(RayTracingPipelineState);
 	check(RayTracingScene);
 	check(RayTracingSceneView);
+	check(Scene);
 
 	FGlobalShaderMap* ShaderMap = GetGlobalShaderMap(ShaderPlatform);
 
-	FRayTracingShaderBindings GlobalBindings;
-	GlobalBindings.SRVs[0] = RayTracingSceneView;
-	GlobalBindings.SRVs[1] = RayTraceRequests;
-	GlobalBindings.UAVs[0] = RayTraceResults;
+	TShaderRef<FNiagaraCollisionRayTraceRG> RGShader = FNiagaraCollisionRayTraceRG::GetShader(ShaderMap);
+
+	FNiagaraCollisionRayTraceRG::FParameters Params;
+
+	Params.GPUSceneParameters.GPUSceneInstanceSceneData = Scene->GPUScene.InstanceSceneDataBuffer.SRV;
+	Params.GPUSceneParameters.GPUScenePrimitiveSceneData = Scene->GPUScene.PrimitiveBuffer.SRV;
+	Params.GPUSceneParameters.GPUSceneFrameNumber = Scene->GPUScene.GetSceneFrameNumber();
+
+	Params.HashTable = PrimIdHashTable.SRV;
+	Params.HashTableSize = HashTableSize;
+	Params.TLAS = RayTracingSceneView;
+	Params.Rays = RayTraceRequests;
+	Params.CollisionOutput = RayTraceResults;
+	Params.HashToCollisionGroups = HashToCollisionGroups.SRV;
+	Params.MaxRetraces = MaxRetraces;
+	FRayTracingShaderBindingsWriter GlobalResources;
+	SetShaderParameters(GlobalResources, RGShader, Params);
 
 	if (FNiagaraCollisionRayTraceRG::SupportsIndirectDispatch())
 	{
@@ -253,9 +419,9 @@ void FNiagaraRayTracingHelper::IssueRayTraces(FRHICommandList& RHICmdList, const
 
 		RHICmdList.RayTraceDispatchIndirect(
 			RayTracingPipelineState,
-			FNiagaraCollisionRayTraceRG::GetShader(ShaderMap),
+			RGShader.GetRayTracingShader(),
 			RayTracingScene,
-			GlobalBindings,
+			GlobalResources,
 			IndirectArgsBuffer->Buffer,
 			IndirectArgsOffset);
 
@@ -264,13 +430,13 @@ void FNiagaraRayTracingHelper::IssueRayTraces(FRHICommandList& RHICmdList, const
 	else
 	{
 		RHICmdList.Transition(FRHITransitionInfo(IndirectArgsBuffer->UAV, ERHIAccess::UAVCompute, ERHIAccess::SRVCompute));
-		GlobalBindings.SRVs[2] = IndirectArgsBuffer->SRV;
+		GlobalResources.SetSRV(2, IndirectArgsBuffer->SRV);
 
 		RHICmdList.RayTraceDispatch(
 			RayTracingPipelineState,
-			FNiagaraCollisionRayTraceRG::GetShader(ShaderMap),
+			RGShader.GetRayTracingShader(),
 			RayTracingScene,
-			GlobalBindings,
+			GlobalResources,
 			RayTraceCounts.X,
 			RayTraceCounts.Y
 		);
@@ -287,4 +453,4 @@ bool FNiagaraRayTracingHelper::IsValid() const
 		&& ViewUniformBuffer;
 }
 
-#endif
\ No newline at end of file
+#endif
diff --git a/Engine/Plugins/FX/Niagara/Source/NiagaraShader/Private/NiagaraShaderParticleID.cpp b/Engine/Plugins/FX/Niagara/Source/NiagaraShader/Private/NiagaraShaderParticleID.cpp
--- a/Engine/Plugins/FX/Niagara/Source/NiagaraShader/Private/NiagaraShaderParticleID.cpp
+++ b/Engine/Plugins/FX/Niagara/Source/NiagaraShader/Private/NiagaraShaderParticleID.cpp
@@ -199,6 +199,22 @@ public:
 		SetUAVParameter(RHICmdList, ComputeShader, TargetBufferParam, nullptr);
 	}
 
+	void Execute(FRHICommandList& RHICmdList, FRHIComputeShader* ComputeShader, FRWBufferStructured& Buffer, int32 Value)
+	{
+		const uint32 THREAD_COUNT = 64;
+		const uint32 NumInts = Buffer.NumBytes / sizeof(int32);
+		const uint32 ThreadGroups = FMath::DivideAndRoundUp(NumInts, THREAD_COUNT);
+
+		SetComputePipelineState(RHICmdList, ComputeShader);
+
+		SetUAVParameter(RHICmdList, ComputeShader, TargetBufferParam, Buffer.UAV);
+		SetShaderValue(RHICmdList, ComputeShader, FillValueParam, Value);
+		SetShaderValue(RHICmdList, ComputeShader, BufferSizeParam, NumInts);
+
+		DispatchComputeShader(RHICmdList, this, ThreadGroups, 1, 1);
+
+		SetUAVParameter(RHICmdList, ComputeShader, TargetBufferParam, nullptr);
+	}
 private:
 	LAYOUT_FIELD(FShaderResourceParameter, TargetBufferParam);
 	LAYOUT_FIELD(FShaderParameter, FillValueParam);
@@ -213,3 +229,10 @@ void NiagaraFillGPUIntBuffer(FRHICommandList& RHICmdList, ERHIFeatureLevel::Type
 	FRHIComputeShader* ComputeShader = FillCS.GetComputeShader();
 	FillCS->Execute(RHICmdList, ComputeShader, Buffer, Value);
 }
+
+void NiagaraFillGPUIntBuffer(FRHICommandList& RHICmdList, ERHIFeatureLevel::Type FeatureLevel, FRWBufferStructured& Buffer, int32 Value)
+{
+	TShaderMapRef<NiagaraFillIntBufferCS> FillCS(GetGlobalShaderMap(FeatureLevel));
+	FRHIComputeShader* ComputeShader = FillCS.GetComputeShader();
+	FillCS->Execute(RHICmdList, ComputeShader, Buffer, Value);
+}
diff --git a/Engine/Plugins/FX/Niagara/Source/NiagaraShader/Public/NiagaraRayTracingHelper.h b/Engine/Plugins/FX/Niagara/Source/NiagaraShader/Public/NiagaraRayTracingHelper.h
--- a/Engine/Plugins/FX/Niagara/Source/NiagaraShader/Public/NiagaraRayTracingHelper.h
+++ b/Engine/Plugins/FX/Niagara/Source/NiagaraShader/Public/NiagaraRayTracingHelper.h
@@ -16,12 +16,25 @@ class FRHIUniformBuffer;
 class FRHIUnorderedAccessView;
 struct FRWBuffer;
 
+//c++ mirror of the ray struct defined in NiagaraRayTracingCommon.ush
+struct FNiagaraRayData : public FBasicRayData
+{
+	int32 CollisionGroup;
+};
+
 // c++ mirror of the struct defined in NiagaraRayTracingCommon.ush
 struct FNiagaraRayTracingPayload
 {
-	float	HitT;            // Distance from ray origin to the intersection point in the ray direction. Negative on miss.
-	float	WorldPosition[3];
-	float	WorldNormal[3];
+	float HitT;
+
+	uint32 PrimitiveIndex;
+	uint32 InstanceIndex;
+	float Barycentrics[2];
+
+	uint32 InstanceID;
+
+	float WorldPosition[3];
+	float WorldNormal[3];
 };
 
 class NIAGARASHADER_API FNiagaraRayTracingHelper
@@ -35,9 +48,16 @@ class NIAGARASHADER_API FNiagaraRayTracingHelper
 
 		void Reset();
 		void BuildRayTracingSceneInfo(FRHICommandList& RHICmdList, TConstArrayView<FViewInfo> Views);
-		void IssueRayTraces(FRHICommandList& RHICmdList, const FIntPoint& RayTraceCounts, FRHIShaderResourceView* RayTraceRequests, FRWBuffer* IndirectArgsBuffer, uint32 IndirectArgsOffset, FRHIUnorderedAccessView* RayTraceResults) const;
+		void IssueRayTraces(FRHICommandList& RHICmdList, FScene* Scene, const FIntPoint& RayTraceCounts, uint32 MaxRetraces, FRHIShaderResourceView* RayTraceRequests, FRWBuffer* IndirectArgsBuffer, uint32 IndirectArgsOffset, FRHIUnorderedAccessView* RayTraceResults) const;
 		bool IsValid() const;
 
+		/** Adds a primitive to a collision group. This data is sent to the GPU on the start of the next frame. */
+		void SetPrimitiveCollisionGroup(FPrimitiveSceneInfo& Primitive, uint32 CollisionGroup);
+
+		/** Pushes cpu side copy of the collision group map to the GPU. */
+		void UpdateCollisionGroupMap(FRHICommandList& RHICmdList, FScene* Scene, ERHIFeatureLevel::Type FeatureLevel);
+
+		void RefreshPrimitiveInstanceData();
 	private:
 		const EShaderPlatform ShaderPlatform;
 
@@ -45,6 +65,21 @@ class NIAGARASHADER_API FNiagaraRayTracingHelper
 		mutable FRHIRayTracingScene* RayTracingScene = nullptr;
 		mutable FRHIShaderResourceView* RayTracingSceneView = nullptr;
 		mutable FRHIUniformBuffer* ViewUniformBuffer = nullptr;
+
+		/**
+		CPU side copy of the PrimID to Collision Group Map.
+		This is uploaded to the GPU and used in Niagara RG shader to filter self collisions between objects of the same group.
+		*/
+		TMap<FPrimitiveComponentId, uint32> CollisionGroupMap;
+		
+		/** Hash table. 
+		 PrimIdHashTable is the main hash table that maps GPUSceneInstanceIndex to and Index we can use to store Collision Groups inside HashToCollisionGroups.
+*/
+		FRWBufferStructured PrimIdHashTable;
+		uint32 HashTableSize = 0;
+		FRWBuffer HashToCollisionGroups;
+
+		bool bCollisionGroupMapDirty = true;
 };
 
 #endif
diff --git a/Engine/Plugins/FX/Niagara/Source/NiagaraShader/Public/NiagaraShaderParticleID.h b/Engine/Plugins/FX/Niagara/Source/NiagaraShader/Public/NiagaraShaderParticleID.h
--- a/Engine/Plugins/FX/Niagara/Source/NiagaraShader/Public/NiagaraShaderParticleID.h
+++ b/Engine/Plugins/FX/Niagara/Source/NiagaraShader/Public/NiagaraShaderParticleID.h
@@ -7,3 +7,4 @@
 extern NIAGARASHADER_API void NiagaraInitGPUFreeIDList(FRHICommandList& RHICmdList, ERHIFeatureLevel::Type FeatureLevel, uint32 NumElementsToAlloc, FRWBuffer& NewBuffer, uint32 NumExistingElements, FRHIShaderResourceView* ExistingBuffer);
 extern NIAGARASHADER_API void NiagaraComputeGPUFreeIDs(FRHICommandList& RHICmdList, ERHIFeatureLevel::Type FeatureLevel, uint32 NumIDs, FRHIShaderResourceView* IDToIndexTable, FRWBuffer& FreeIDList, FRWBuffer& FreeIDListSizes, uint32 FreeIDListIndex);
 extern NIAGARASHADER_API void NiagaraFillGPUIntBuffer(FRHICommandList& RHICmdList, ERHIFeatureLevel::Type FeatureLevel, FRWBuffer& Buffer, int32 Value);
+extern NIAGARASHADER_API void NiagaraFillGPUIntBuffer(FRHICommandList& RHICmdList, ERHIFeatureLevel::Type FeatureLevel, FRWBufferStructured& Buffer, int32 Value);
diff --git a/Engine/Shaders/Private/HashTable.ush b/Engine/Shaders/Private/HashTable.ush
--- a/Engine/Shaders/Private/HashTable.ush
+++ b/Engine/Shaders/Private/HashTable.ush
@@ -35,6 +35,7 @@ bool HashTableAdd( uint Key, out uint Index )
 	Key++;
 
 	LOOP
+	[allow_uav_condition]
 	for( Index = MurmurMix( Key );; Index++ )
 	{
 		//Index &= HashTableSize - 1;
@@ -69,6 +70,7 @@ bool HashTableFind( uint Key, out uint Index )
 	Key++;
 
 	LOOP
+	[allow_uav_condition]
 	for( Index = MurmurMix( Key );; Index++ )
 	{
 		//Index &= HashTableSize - 1;
@@ -79,7 +81,9 @@ bool HashTableFind( uint Key, out uint Index )
 		{
 			if( StoredKey != 0 )
 				continue;
-
+		}
+		else
+		{
 			return true;
 		}
 
@@ -87,4 +91,4 @@ bool HashTableFind( uint Key, out uint Index )
 	}
 
 	return false;
-}
\ No newline at end of file
+}
diff --git a/Engine/Source/Editor/UnrealEd/Private/ThumbnailHelpers.cpp b/Engine/Source/Editor/UnrealEd/Private/ThumbnailHelpers.cpp
--- a/Engine/Source/Editor/UnrealEd/Private/ThumbnailHelpers.cpp
+++ b/Engine/Source/Editor/UnrealEd/Private/ThumbnailHelpers.cpp
@@ -166,7 +166,7 @@ FParticleSystemThumbnailScene::FParticleSystemThumbnailScene()
 	bForceAllUsedMipsResident = false;
 	PartComponent = NULL;
 
-	ThumbnailFXSystem = FFXSystemInterface::Create(GetScene()->GetFeatureLevel(), GetScene()->GetShaderPlatform());
+	ThumbnailFXSystem = FFXSystemInterface::Create(GetScene()->GetFeatureLevel(), GetScene());
 	GetScene()->SetFXSystem( ThumbnailFXSystem );
 }
 
diff --git a/Engine/Source/Runtime/Engine/Private/Particles/FXSystem.cpp b/Engine/Source/Runtime/Engine/Private/Particles/FXSystem.cpp
--- a/Engine/Source/Runtime/Engine/Private/Particles/FXSystem.cpp
+++ b/Engine/Source/Runtime/Engine/Private/Particles/FXSystem.cpp
@@ -24,8 +24,10 @@ TMap<FName, FCreateCustomFXSystemDelegate> FFXSystemInterface::CreateCustomFXDel
 	External FX system interface.
 -----------------------------------------------------------------------------*/
 
-FFXSystemInterface* FFXSystemInterface::Create(ERHIFeatureLevel::Type InFeatureLevel, EShaderPlatform InShaderPlatform)
+FFXSystemInterface* FFXSystemInterface::Create(ERHIFeatureLevel::Type InFeatureLevel, FSceneInterface* Scene)
 {
+	check(Scene);
+	EShaderPlatform InShaderPlatform = Scene->GetShaderPlatform();
 	// The FGPUSortManager is currently only being used by FFXSystemInterface implementations.
 	// Because of this, the lifetime management of the GPUSortManager only consists of a ref counter incremented initially 
 	// in the gamethread (in this function) and decremented on the renderthread (when the system interfaces are deleted).
@@ -34,6 +36,9 @@ FFXSystemInterface* FFXSystemInterface::Create(ERHIFeatureLevel::Type InFeatureL
 	if (CreateCustomFXDelegates.Num())
 	{
 		FFXSystemSet* Set = new FFXSystemSet(GPUSortManager);
+		Scene->SetFXSystem(Set);
+		Set->SetScene((FScene*)Scene);
+
 		Set->FXSystems.Add(new FFXSystem(InFeatureLevel, InShaderPlatform, GPUSortManager));
 
 		for (TMap<FName, FCreateCustomFXSystemDelegate>::TConstIterator Ite(CreateCustomFXDelegates); Ite; ++Ite)
@@ -41,6 +46,7 @@ FFXSystemInterface* FFXSystemInterface::Create(ERHIFeatureLevel::Type InFeatureL
 			FFXSystemInterface* CustomFX = Ite.Value().Execute(InFeatureLevel, InShaderPlatform, GPUSortManager);
 			if (CustomFX)
 			{
+				CustomFX->SetScene((FScene*)Scene);
 				Set->FXSystems.Add(CustomFX);
 			}
 		}
@@ -48,8 +54,12 @@ FFXSystemInterface* FFXSystemInterface::Create(ERHIFeatureLevel::Type InFeatureL
 	}
 	else
 	{
-		return new FFXSystem(InFeatureLevel, InShaderPlatform, GPUSortManager);
+		FFXSystemInterface* Ret = new FFXSystem(InFeatureLevel, InShaderPlatform, GPUSortManager);
+		Scene->SetFXSystem(Ret);
+		Ret->SetScene((FScene*)Scene);
+		return Ret;
 	}
+
 }
 void FFXSystemInterface::QueueDestroyGPUSimulation(FFXSystemInterface* FXSystem)
 {
diff --git a/Engine/Source/Runtime/Engine/Private/World.cpp b/Engine/Source/Runtime/Engine/Private/World.cpp
--- a/Engine/Source/Runtime/Engine/Private/World.cpp
+++ b/Engine/Source/Runtime/Engine/Private/World.cpp
@@ -7759,8 +7759,7 @@ void UWorld::CreateFXSystem()
 {
 	if ( !IsRunningDedicatedServer() && !IsRunningCommandlet() )
 	{
-		FXSystem = FFXSystemInterface::Create(FeatureLevel, Scene->GetShaderPlatform());
-		Scene->SetFXSystem(FXSystem);
+		FXSystem = FFXSystemInterface::Create(FeatureLevel, Scene);
 	}
 	else
 	{
diff --git a/Engine/Source/Runtime/Engine/Public/FXSystem.h b/Engine/Source/Runtime/Engine/Public/FXSystem.h
--- a/Engine/Source/Runtime/Engine/Public/FXSystem.h
+++ b/Engine/Source/Runtime/Engine/Public/FXSystem.h
@@ -127,7 +127,7 @@ public:
 	/**
 	 * Create an effects system instance.
 	 */
-	ENGINE_API static FFXSystemInterface* Create(ERHIFeatureLevel::Type InFeatureLevel, EShaderPlatform InShaderPlatform);
+	ENGINE_API static FFXSystemInterface* Create(ERHIFeatureLevel::Type InFeatureLevel, FSceneInterface* Scene);
 
 	/**
 	 * Destroy an effects system instance.
@@ -263,6 +263,9 @@ public:
 
 	virtual void SetSceneTexturesUniformBuffer(FRHIUniformBuffer* InSceneTexturesUniformParams) {}
 
+	FORCEINLINE FScene* GetScene()const { return Scene; }
+	FORCEINLINE void SetScene(FScene* InScene) { Scene = InScene; }
+
 protected:
 	
 	friend class FFXSystemSet;
@@ -278,6 +281,8 @@ private:
 
 	bool bIsPendingKill = false;
 
+	class FScene* Scene = nullptr;
+
 	static TMap<FName, FCreateCustomFXSystemDelegate> CreateCustomFXDelegates;
 };
 
diff --git a/Engine/Source/Runtime/Renderer/Private/PrimitiveSceneInfo.cpp b/Engine/Source/Runtime/Renderer/Private/PrimitiveSceneInfo.cpp
--- a/Engine/Source/Runtime/Renderer/Private/PrimitiveSceneInfo.cpp
+++ b/Engine/Source/Runtime/Renderer/Private/PrimitiveSceneInfo.cpp
@@ -109,6 +109,9 @@ private:
 	TRefCountPtr<HHitProxy> CurrentHitProxy;
 };
 
+FPrimitiveSceneInfo::FPrimitiveSceneInfoEvent FPrimitiveSceneInfo::OnGPUSceneInstancesAllocated;
+FPrimitiveSceneInfo::FPrimitiveSceneInfoEvent FPrimitiveSceneInfo::OnGPUSceneInstancesFreed;
+
 FPrimitiveFlagsCompact::FPrimitiveFlagsCompact(const FPrimitiveSceneProxy* Proxy)
 	: bCastDynamicShadow(Proxy->CastsDynamicShadow())
 	, bStaticLighting(Proxy->HasStaticLighting())
@@ -1002,6 +1005,8 @@ void FPrimitiveSceneInfo::AllocateGPUSceneInstances(FScene* Scene, const TArrayV
 				Scene->LumenSceneData->UpdatePrimitiveInstanceOffset(SceneInfo->PackedIndex);
 			}
 		}
+
+		OnGPUSceneInstancesAllocated.Broadcast();
 	}
 }
 
@@ -1040,6 +1045,8 @@ void FPrimitiveSceneInfo::FreeGPUSceneInstances()
 		Scene->GPUScene.FreeInstanceSceneDataSlots(InstanceSceneDataOffset, NumInstanceSceneDataEntries);
 		InstanceSceneDataOffset = INDEX_NONE;
 		NumInstanceSceneDataEntries = 0;
+
+		OnGPUSceneInstancesFreed.Broadcast();
 	}
 }
 
diff --git a/Engine/Source/Runtime/Renderer/Public/PrimitiveSceneInfo.h b/Engine/Source/Runtime/Renderer/Public/PrimitiveSceneInfo.h
--- a/Engine/Source/Runtime/Renderer/Public/PrimitiveSceneInfo.h
+++ b/Engine/Source/Runtime/Renderer/Public/PrimitiveSceneInfo.h
@@ -635,6 +635,11 @@ private:
 	/** Updates cached transforms in CachedRayTracingInstance */
 	void UpdateCachedRayTracingInstanceTransforms(FMatrix NewPrimitiveLocalToWorld);
 #endif
+
+public:
+	DECLARE_MULTICAST_DELEGATE(FPrimitiveSceneInfoEvent);
+	RENDERER_API static FPrimitiveSceneInfoEvent OnGPUSceneInstancesAllocated;
+	RENDERER_API static FPrimitiveSceneInfoEvent OnGPUSceneInstancesFreed;
 };
 
 /** Defines how the primitive is stored in the scene's primitive octree. */
