- LWC fixes

#rnx

diff --git a/Engine/Plugins/FX/Niagara/Source/Niagara/Private/DataInterface/NiagaraDataInterfaceActorComponent.cpp b/Engine/Plugins/FX/Niagara/Source/Niagara/Private/DataInterface/NiagaraDataInterfaceActorComponent.cpp
--- a/Engine/Plugins/FX/Niagara/Source/Niagara/Private/DataInterface/NiagaraDataInterfaceActorComponent.cpp
+++ b/Engine/Plugins/FX/Niagara/Source/Niagara/Private/DataInterface/NiagaraDataInterfaceActorComponent.cpp
@@ -94,10 +94,13 @@ public:
 		FInstanceData_RenderThread* InstanceData = DataInterfaceProxy->SystemInstancesToInstanceData_RT.Find(Context.SystemInstanceID);
 		check(InstanceData != nullptr);
 
+		const FMatrix44f InstanceMatrix = InstanceData->CachedTransform.ToMatrixWithScale();
+		const FQuat InstanceRotation = InstanceData->CachedTransform.GetRotation();
+		const FVector3f InstanceScale = InstanceData->CachedTransform.GetScale3D();
 		SetShaderValue(RHICmdList, ComputeShaderRHI, ValidParam, InstanceData->bCachedValid ? 1 : 0);
-		SetShaderValue(RHICmdList, ComputeShaderRHI, MatrixParam, InstanceData->CachedTransform.ToMatrixWithScale());
-		SetShaderValue(RHICmdList, ComputeShaderRHI, RotationParam, InstanceData->CachedTransform.GetRotation());
-		SetShaderValue(RHICmdList, ComputeShaderRHI, ScaleParam, InstanceData->CachedTransform.GetScale3D());
+		SetShaderValue(RHICmdList, ComputeShaderRHI, MatrixParam, InstanceMatrix);
+		SetShaderValue(RHICmdList, ComputeShaderRHI, RotationParam, InstanceRotation);
+		SetShaderValue(RHICmdList, ComputeShaderRHI, ScaleParam, InstanceScale);
 	}
 
 private:
@@ -344,9 +347,9 @@ void UNiagaraDataInterfaceActorComponent::VMGetMatrix(FVectorVMExternalFunctionC
 
 	VectorVM::FUserPtrHandler<FInstanceData_GameThread> InstanceData(Context);
 	FNDIOutputParam<bool>		OutValid(Context);
-	FNDIOutputParam<FMatrix>	OutMatrix(Context);
+	FNDIOutputParam<FMatrix44f>	OutMatrix(Context);
 
-	const FMatrix& InstanceMatrix = InstanceData->CachedTransform.ToMatrixWithScale();
+	const FMatrix44f InstanceMatrix = InstanceData->CachedTransform.ToMatrixWithScale();
 	for (int32 i = 0; i < Context.GetNumInstances(); ++i)
 	{
 		OutValid.SetAndAdvance(InstanceData->bCachedValid);
@@ -360,9 +363,9 @@ void UNiagaraDataInterfaceActorComponent::VMGetTransform(FVectorVMExternalFuncti
 
 	VectorVM::FUserPtrHandler<FInstanceData_GameThread> InstanceData(Context);
 	FNDIOutputParam<bool>		OutValid(Context);
-	FNDIOutputParam<FVector>	OutPosition(Context);
+	FNDIOutputParam<FVector3f>	OutPosition(Context);
 	FNDIOutputParam<FQuat>		OutRotation(Context);
-	FNDIOutputParam<FVector>	OutScale(Context);
+	FNDIOutputParam<FVector3f>	OutScale(Context);
 
 	for (int32 i=0; i < Context.GetNumInstances(); ++i)
 	{
