Fixed collision normals on PC

diff --git a/Engine/Plugins/FX/Niagara/Shaders/Private/NiagaraRayTracingShaders.usf b/Engine/Plugins/FX/Niagara/Shaders/Private/NiagaraRayTracingShaders.usf
--- a/Engine/Plugins/FX/Niagara/Shaders/Private/NiagaraRayTracingShaders.usf
+++ b/Engine/Plugins/FX/Niagara/Shaders/Private/NiagaraRayTracingShaders.usf
@@ -18,8 +18,6 @@ Buffer<uint>                                  RayTraceCounts;
 
 Buffer<uint> HashToCollisionGroups;
 uint MaxRetraces;
-uint TestParam;
-uint TestControl;
 
 uint GetGPUSceneInstanceIndex(uint PrimitiveIndex, uint InstanceIndex)
 {
@@ -162,7 +160,8 @@ RAY_TRACING_ENTRY_CLOSEST_HIT(NiagaraCollisionRayTraceCH,
 	float3x3 InverseTranspose3x3 = transpose((float3x3)WorldToObject4x3());
 
 	Payload.WorldPosition = WorldRayOrigin() + WorldRayDirection() * RayTCurrent();
-	Payload.WorldNormal = WorldRayDirection();
+	//Payload.WorldPosition = mul(float4(LocalPosition, 1.0f), ObjectToWorld4x3()).xyz;
+	Payload.WorldNormal = normalize(mul(LocalNormal, InverseTranspose3x3));
 }
 
 RAY_TRACING_ENTRY_MISS(NiagaraCollisionRayTraceMiss, FNiagaraRayTracingPayload, Payload)
