Faster attribute pass through for grid3d
Skips attributes that have been written to for the current stage
includes a test

diff --git a/Engine/Plugins/FX/Niagara/Source/Niagara/Classes/NiagaraDataInterface.h b/Engine/Plugins/FX/Niagara/Source/Niagara/Classes/NiagaraDataInterface.h
--- a/Engine/Plugins/FX/Niagara/Source/Niagara/Classes/NiagaraDataInterface.h
+++ b/Engine/Plugins/FX/Niagara/Source/Niagara/Classes/NiagaraDataInterface.h
@@ -278,9 +278,9 @@ public:
 	/** Can this data interface be used as a StackContext parameter map replacement when being used as a sim stage iteration source? */
 	virtual bool SupportsIterationSourceNamespaceAttributesHLSL() const { return false; }
 	/** Generate the necessary plumbing HLSL at the beginning of the stage where this is used as a sim stage iteration source. Note that this should inject other internal calls using the CustomHLSL node syntax. See GridCollection2D for an example.*/
-	virtual bool GenerateIterationSourceNamespaceReadAttributesHLSL(FNiagaraDataInterfaceGPUParamInfo& DIInstanceInfo, const FNiagaraVariable& InIterationSourceVariable, TConstArrayView<FNiagaraVariable> InArguments, TConstArrayView<FNiagaraVariable> InAttributes, TConstArrayView<FString> InAttributeHLSLNames, bool bInSetToDefaults, bool bPartialWrites, TArray<FText>& OutErrors, FString& OutHLSL) const { return false; };
+	virtual bool GenerateIterationSourceNamespaceReadAttributesHLSL(FNiagaraDataInterfaceGPUParamInfo& DIInstanceInfo, const FNiagaraVariable& InIterationSourceVariable, TConstArrayView<FNiagaraVariable> InArguments, TConstArrayView<FNiagaraVariable> InAttributes, TConstArrayView<FString> InAttributeHLSLNames, bool bInSetToDefaults, bool bSpawnOnly, bool bPartialWrites, TArray<FText>& OutErrors, FString& OutHLSL) const { return false; };
 	/** Generate the necessary plumbing HLSL at the end of the stage where this is used as a sim stage iteration source. Note that this should inject other internal calls using the CustomHLSL node syntax. See GridCollection2D for an example.*/
-	virtual bool GenerateIterationSourceNamespaceWriteAttributesHLSL(FNiagaraDataInterfaceGPUParamInfo& DIInstanceInfo, const FNiagaraVariable& InIterationSourceVariable, TConstArrayView<FNiagaraVariable> InArguments, TConstArrayView<FNiagaraVariable> InAttributes, TConstArrayView<FString> InAttributeHLSLNames, bool bPartialWrites, TArray<FText>& OutErrors, FString& OutHLSL) const { return false; };
+	virtual bool GenerateIterationSourceNamespaceWriteAttributesHLSL(FNiagaraDataInterfaceGPUParamInfo& DIInstanceInfo, const FNiagaraVariable& InIterationSourceVariable, TConstArrayView<FNiagaraVariable> InArguments, TConstArrayView<FNiagaraVariable> InAttributes, TConstArrayView<FString> InAttributeHLSLNames, bool bSpawnOnly, bool bPartialWrites, TArray<FText>& OutErrors, FString& OutHLSL) const { return false; };
 	/** Used by the translator when dealing with signatures that turn into compiler tags to figure out the precise compiler tag. */
 	virtual bool GenerateCompilerTagPrefix(const FNiagaraFunctionSignature& InSignature, FString& OutPrefix) const  { return false; }
 
diff --git a/Engine/Plugins/FX/Niagara/Source/Niagara/Classes/NiagaraDataInterfaceGrid2DCollection.h b/Engine/Plugins/FX/Niagara/Source/Niagara/Classes/NiagaraDataInterfaceGrid2DCollection.h
--- a/Engine/Plugins/FX/Niagara/Source/Niagara/Classes/NiagaraDataInterfaceGrid2DCollection.h
+++ b/Engine/Plugins/FX/Niagara/Source/Niagara/Classes/NiagaraDataInterfaceGrid2DCollection.h
@@ -267,8 +267,8 @@ public:
 	virtual bool GenerateSetupHLSL(FNiagaraDataInterfaceGPUParamInfo& DIInstanceInfo, TConstArrayView<FNiagaraVariable> InArguments, bool bSpawnOnly, bool bPartialWrites, TArray<FText>& OutErrors, FString& OutHLSL) const;
 	virtual bool GenerateTeardownHLSL(FNiagaraDataInterfaceGPUParamInfo& DIInstanceInfo, TConstArrayView<FNiagaraVariable> InArguments, bool bSpawnOnly, bool bPartialWrites, TArray<FText>& OutErrors, FString& OutHLSL) const;
 	virtual bool SupportsIterationSourceNamespaceAttributesHLSL() const override { return true; }
-	virtual bool GenerateIterationSourceNamespaceReadAttributesHLSL(FNiagaraDataInterfaceGPUParamInfo& DIInstanceInfo, const FNiagaraVariable& IterationSourceVar, TConstArrayView<FNiagaraVariable> InArguments, TConstArrayView<FNiagaraVariable> InAttributes, TConstArrayView<FString> InAttributeHLSLNames, bool bInSetToDefaults, bool bPartialWrites, TArray<FText>& OutErrors, FString& OutHLSL) const override;
-	virtual bool GenerateIterationSourceNamespaceWriteAttributesHLSL(FNiagaraDataInterfaceGPUParamInfo& DIInstanceInfo, const FNiagaraVariable& IterationSourceVar, TConstArrayView<FNiagaraVariable> InArguments, TConstArrayView<FNiagaraVariable> InAttributes, TConstArrayView<FString> InAttributeHLSLNames, bool bPartialWrites, TArray<FText>& OutErrors, FString& OutHLSL) const override;
+	virtual bool GenerateIterationSourceNamespaceReadAttributesHLSL(FNiagaraDataInterfaceGPUParamInfo& DIInstanceInfo, const FNiagaraVariable& IterationSourceVar, TConstArrayView<FNiagaraVariable> InArguments, TConstArrayView<FNiagaraVariable> InAttributes, TConstArrayView<FString> InAttributeHLSLNames, bool bInSetToDefaults, bool bSpawnOnly, bool bPartialWrites, TArray<FText>& OutErrors, FString& OutHLSL) const override;
+	virtual bool GenerateIterationSourceNamespaceWriteAttributesHLSL(FNiagaraDataInterfaceGPUParamInfo& DIInstanceInfo, const FNiagaraVariable& IterationSourceVar, TConstArrayView<FNiagaraVariable> InArguments, TConstArrayView<FNiagaraVariable> InAttributes, TConstArrayView<FString> InAttributeHLSLNames, bool bSpawnOnly, bool bPartialWrites, TArray<FText>& OutErrors, FString& OutHLSL) const override;
 #endif
 
 	static int32 GetComponentCountFromFuncName(const FName& FuncName);
diff --git a/Engine/Plugins/FX/Niagara/Source/Niagara/Classes/NiagaraDataInterfaceGrid3DCollection.h b/Engine/Plugins/FX/Niagara/Source/Niagara/Classes/NiagaraDataInterfaceGrid3DCollection.h
--- a/Engine/Plugins/FX/Niagara/Source/Niagara/Classes/NiagaraDataInterfaceGrid3DCollection.h
+++ b/Engine/Plugins/FX/Niagara/Source/Niagara/Classes/NiagaraDataInterfaceGrid3DCollection.h
@@ -217,6 +217,7 @@ public:
 
 	static const FName ClearCellFunctionName;
 	static const FName CopyPreviousToCurrentForCellFunctionName;
+	static const FName CopyMaskedPreviousToCurrentForCellFunctionName;
 
 	static const FName SetValueFunctionName;
 	static const FName GetValueFunctionName;
@@ -262,8 +263,8 @@ public:
 	virtual bool GenerateSetupHLSL(FNiagaraDataInterfaceGPUParamInfo& DIInstanceInfo, TConstArrayView<FNiagaraVariable> InArguments, bool bSpawnOnly, bool bPartialWrites, TArray<FText>& OutErrors, FString& OutHLSL) const;
 	virtual bool GenerateTeardownHLSL(FNiagaraDataInterfaceGPUParamInfo& DIInstanceInfo, TConstArrayView<FNiagaraVariable> InArguments, bool bSpawnOnly, bool bPartialWrites, TArray<FText>& OutErrors, FString& OutHLSL) const;
 	virtual bool SupportsIterationSourceNamespaceAttributesHLSL() const override { return true; }
-	virtual bool GenerateIterationSourceNamespaceReadAttributesHLSL(FNiagaraDataInterfaceGPUParamInfo& DIInstanceInfo, const FNiagaraVariable& IterationSourceVar, TConstArrayView<FNiagaraVariable> InArguments, TConstArrayView<FNiagaraVariable> InAttributes, TConstArrayView<FString> InAttributeHLSLNames, bool bInSetToDefaults, bool bPartialWrites, TArray<FText>& OutErrors, FString& OutHLSL) const override;
-	virtual bool GenerateIterationSourceNamespaceWriteAttributesHLSL(FNiagaraDataInterfaceGPUParamInfo& DIInstanceInfo, const FNiagaraVariable& IterationSourceVar, TConstArrayView<FNiagaraVariable> InArguments, TConstArrayView<FNiagaraVariable> InAttributes, TConstArrayView<FString> InAttributeHLSLNames, bool bPartialWrites, TArray<FText>& OutErrors, FString& OutHLSL) const override;
+	virtual bool GenerateIterationSourceNamespaceReadAttributesHLSL(FNiagaraDataInterfaceGPUParamInfo& DIInstanceInfo, const FNiagaraVariable& IterationSourceVar, TConstArrayView<FNiagaraVariable> InArguments, TConstArrayView<FNiagaraVariable> InAttributes, TConstArrayView<FString> InAttributeHLSLNames, bool bInSetToDefaults, bool bSpawnOnly, bool bPartialWrites, TArray<FText>& OutErrors, FString& OutHLSL) const override;
+	virtual bool GenerateIterationSourceNamespaceWriteAttributesHLSL(FNiagaraDataInterfaceGPUParamInfo& DIInstanceInfo, const FNiagaraVariable& IterationSourceVar, TConstArrayView<FNiagaraVariable> InArguments, TConstArrayView<FNiagaraVariable> InAttributes, TConstArrayView<FString> InAttributeHLSLNames, bool bSpawnOnly, bool bPartialWrites, TArray<FText>& OutErrors, FString& OutHLSL) const override;
 #endif
 
 	static int32 GetComponentCountFromFuncName(const FName& FuncName);
@@ -279,6 +280,7 @@ protected:
 	const TCHAR* TypeDefinitionToHLSLTypeString(const FNiagaraTypeDefinition& InDef) const;
 	FName TypeDefinitionToGetFunctionName(const FNiagaraTypeDefinition& InDef) const;
 	FName TypeDefinitionToSetFunctionName(const FNiagaraTypeDefinition& InDef) const;
+	FName TypeDefinitionToAttributeIndexFunctionName(const FNiagaraTypeDefinition& InDef) const;
 #endif
 
 	//~ UNiagaraDataInterface interface
diff --git a/Engine/Plugins/FX/Niagara/Source/Niagara/Private/NiagaraDataInterfaceGrid2DCollection.cpp b/Engine/Plugins/FX/Niagara/Source/Niagara/Private/NiagaraDataInterfaceGrid2DCollection.cpp
--- a/Engine/Plugins/FX/Niagara/Source/Niagara/Private/NiagaraDataInterfaceGrid2DCollection.cpp
+++ b/Engine/Plugins/FX/Niagara/Source/Niagara/Private/NiagaraDataInterfaceGrid2DCollection.cpp
@@ -1547,7 +1547,7 @@ bool UNiagaraDataInterfaceGrid2DCollection::GetFunctionHLSL(const FNiagaraDataIn
 #endif
 
 #if WITH_EDITOR
-bool UNiagaraDataInterfaceGrid2DCollection::GenerateIterationSourceNamespaceReadAttributesHLSL(FNiagaraDataInterfaceGPUParamInfo& DIInstanceInfo, const FNiagaraVariable& IterationSourceVar, TConstArrayView<FNiagaraVariable> InArguments, TConstArrayView<FNiagaraVariable> InAttributes, TConstArrayView<FString> InAttributeHLSLNames, bool bInSetToDefaults, bool bPartialWrites, TArray<FText>& OutErrors, FString& OutHLSL) const
+bool UNiagaraDataInterfaceGrid2DCollection::GenerateIterationSourceNamespaceReadAttributesHLSL(FNiagaraDataInterfaceGPUParamInfo& DIInstanceInfo, const FNiagaraVariable& IterationSourceVar, TConstArrayView<FNiagaraVariable> InArguments, TConstArrayView<FNiagaraVariable> InAttributes, TConstArrayView<FString> InAttributeHLSLNames, bool bInSetToDefaults, bool bSpawnOnly, bool bPartialWrites, TArray<FText>& OutErrors, FString& OutHLSL) const
 {
 
 	FString DIVarName;
@@ -1608,7 +1608,7 @@ bool UNiagaraDataInterfaceGrid2DCollection::GenerateIterationSourceNamespaceRead
 	}
 	return true;
 }
-bool UNiagaraDataInterfaceGrid2DCollection::GenerateIterationSourceNamespaceWriteAttributesHLSL(FNiagaraDataInterfaceGPUParamInfo& DIInstanceInfo, const FNiagaraVariable& IterationSourceVar, TConstArrayView<FNiagaraVariable> InArguments, TConstArrayView<FNiagaraVariable> InAttributes, TConstArrayView<FString> InAttributeHLSLNames, bool bPartialWrites, TArray<FText>& OutErrors, FString& OutHLSL) const
+bool UNiagaraDataInterfaceGrid2DCollection::GenerateIterationSourceNamespaceWriteAttributesHLSL(FNiagaraDataInterfaceGPUParamInfo& DIInstanceInfo, const FNiagaraVariable& IterationSourceVar, TConstArrayView<FNiagaraVariable> InArguments, TConstArrayView<FNiagaraVariable> InAttributes, TConstArrayView<FString> InAttributeHLSLNames, bool bSpawnOnly, bool bPartialWrites, TArray<FText>& OutErrors, FString& OutHLSL) const
 {
 	FString DIVarName;
 	OutHLSL += TEXT("\t//Generated by UNiagaraDataInterfaceGrid2DCollection::GenerateIterationSourceNamespaceWriteAttributesHLSL\n");
diff --git a/Engine/Plugins/FX/Niagara/Source/Niagara/Private/NiagaraDataInterfaceGrid3DCollection.cpp b/Engine/Plugins/FX/Niagara/Source/Niagara/Private/NiagaraDataInterfaceGrid3DCollection.cpp
--- a/Engine/Plugins/FX/Niagara/Source/Niagara/Private/NiagaraDataInterfaceGrid3DCollection.cpp
+++ b/Engine/Plugins/FX/Niagara/Source/Niagara/Private/NiagaraDataInterfaceGrid3DCollection.cpp
@@ -30,6 +30,7 @@ const FName UNiagaraDataInterfaceGrid3DCollection::SetNumCellsFunctionName("SetN
 
 const FName UNiagaraDataInterfaceGrid3DCollection::ClearCellFunctionName("ClearCell");
 const FName UNiagaraDataInterfaceGrid3DCollection::CopyPreviousToCurrentForCellFunctionName("CopyPreviousToCurrentForCell");
+const FName UNiagaraDataInterfaceGrid3DCollection::CopyMaskedPreviousToCurrentForCellFunctionName("CopyMaskedPreviousToCurrentForCell");
 
 // Global VM function names, also used by the shaders code generation methods.
 const FName UNiagaraDataInterfaceGrid3DCollection::SetValueFunctionName("SetGridValue");
@@ -578,6 +579,30 @@ void UNiagaraDataInterfaceGrid3DCollection::GetFunctions(TArray<FNiagaraFunction
 		OutFunctions.Add(Sig);
 	}
 
+	{
+		FNiagaraFunctionSignature Sig;
+		Sig.Name = CopyMaskedPreviousToCurrentForCellFunctionName;
+		Sig.Inputs.Add(FNiagaraVariable(FNiagaraTypeDefinition(GetClass()), TEXT("Grid")));
+		Sig.Inputs.Add(FNiagaraVariable(FNiagaraTypeDefinition::GetIntDef(), TEXT("IndexX")));
+		Sig.Inputs.Add(FNiagaraVariable(FNiagaraTypeDefinition::GetIntDef(), TEXT("IndexY")));
+		Sig.Inputs.Add(FNiagaraVariable(FNiagaraTypeDefinition::GetIntDef(), TEXT("IndexZ")));		
+		Sig.Inputs.Add(FNiagaraVariable(FNiagaraTypeDefinition::GetIntDef(), TEXT("NumAttributesSet")));
+		Sig.Inputs.Add(FNiagaraVariable(FNiagaraTypeDefinition::GetIntDef(), TEXT("AttributeMask")));
+
+		Sig.bExperimental = true;
+		Sig.bMemberFunction = true;
+		Sig.bRequiresContext = false;
+		Sig.bWriteFunction = true;
+		Sig.bRequiresExecPin = true;
+		Sig.ModuleUsageBitmask = ENiagaraScriptUsageMask::Particle;
+		Sig.bSupportsCPU = false;
+		Sig.bSupportsGPU = true;
+#if WITH_EDITORONLY_DATA
+		Sig.Description = NSLOCTEXT("Niagara", "NiagaraDataInterfaceGridColl2D_CopyMaskedPreviousToCurrentForCell", "Take the previous contents of the cell and copy to the output location for the cell.");
+#endif
+		OutFunctions.Add(Sig);
+	}
+
 	{
 		FNiagaraFunctionSignature Sig;
 		Sig.Name = SetVector4ValueFunctionName;
@@ -1326,6 +1351,19 @@ FName UNiagaraDataInterfaceGrid3DCollection::TypeDefinitionToSetFunctionName(con
 	return NAME_None;;
 }
 
+FName UNiagaraDataInterfaceGrid3DCollection::TypeDefinitionToAttributeIndexFunctionName(const FNiagaraTypeDefinition& InDef) const
+{
+	if (InDef == FNiagaraTypeDefinition::GetFloatDef())
+		return GetFloatAttributeIndexFunctionName;
+	if (InDef == FNiagaraTypeDefinition::GetVec2Def())
+		return GetVector2DAttributeIndexFunctionName;
+	if (InDef == FNiagaraTypeDefinition::GetVec3Def())
+		return GetVectorAttributeIndexFunctionName;
+	if (InDef == FNiagaraTypeDefinition::GetVec4Def() || InDef == FNiagaraTypeDefinition::GetColorDef())
+		return GetVector4AttributeIndexFunctionName;
+	return NAME_None;;
+}
+
 bool UNiagaraDataInterfaceGrid3DCollection::GetFunctionHLSL(const FNiagaraDataInterfaceGPUParamInfo& ParamInfo, const FNiagaraDataInterfaceGeneratedFunction& FunctionInfo, int FunctionInstanceIndex, FString& OutHLSL)
 {
 	bool ParentRet = Super::GetFunctionHLSL(ParamInfo, FunctionInfo, FunctionInstanceIndex, OutHLSL);
@@ -1425,6 +1463,54 @@ bool UNiagaraDataInterfaceGrid3DCollection::GetFunctionHLSL(const FNiagaraDataIn
 		OutHLSL += FString::Format(FormatBounds, ArgsBounds);
 		return true;
 	}
+	else if (FunctionInfo.DefinitionName == CopyMaskedPreviousToCurrentForCellFunctionName)
+	{
+		static const TCHAR* FormatBounds = TEXT(R"(
+			void {FunctionName}(int In_IndexX, int In_IndexY, int In_IndexZ, int NumAttributesSet, int AttributeMask)
+			{	
+				// early out if we set all the attributes that exist on the DI
+				[branch]
+				if (NumAttributesSet == {NumAttributes}.x)
+				{
+					return;					
+				}
+
+				for (int AttributeIndex = 0; AttributeIndex < {NumAttributes}.x; AttributeIndex++)
+				{					
+					// check the attribute index in the attribute mask to see if it has been set
+					// we automatically pass through attributes higher than 31 since we can only
+					// store 32 attributes in the mask
+					[branch]
+					if ((AttributeMask & (1 << AttributeIndex)) == 0 || AttributeIndex >= 32)
+					{						
+						int3 TileOffset = {PerAttributeDataName}[AttributeIndex].xyz;
+
+						float Val = {Grid}.Load(int4(In_IndexX + TileOffset.x, In_IndexY + TileOffset.y, In_IndexZ + TileOffset.z, 0));	
+						RW{OutputGrid}[int3(In_IndexX + TileOffset.x, In_IndexY + TileOffset.y, In_IndexZ + TileOffset.z)] = Val;
+					}
+				}
+
+	// alternative implementation - look
+	/*
+				int AttributeIndex = -1;
+				[loop]			
+				for (int x = 0; x < {NumAttributes}.x - NumAttributesSet; ++x)
+				{											
+					do
+					{
+						AttributeIndex++;																		
+					} while(AttributeIndex < 32 && (AttributeMask & (1 << AttributeIndex)) == 1);
+					
+					int3 TileOffset = {PerAttributeDataName}[AttributeIndex].xyz;
+					float Val = {Grid}.Load(int4(In_IndexX + TileOffset.x, In_IndexY + TileOffset.y, In_IndexZ + TileOffset.z, 0));	
+					RW{OutputGrid}[int3(In_IndexX + TileOffset.x, In_IndexY + TileOffset.y, In_IndexZ + TileOffset.z)] = Val;
+				}
+*/
+			}
+		)");
+		OutHLSL += FString::Format(FormatBounds, ArgsBounds);
+		return true;
+	}
 	else if (FunctionInfo.DefinitionName == ClearCellFunctionName)
 	{
 		static const TCHAR* FormatBounds = TEXT(R"(
@@ -1558,7 +1644,7 @@ bool UNiagaraDataInterfaceGrid3DCollection::GetFunctionHLSL(const FNiagaraDataIn
 #endif
 
 #if WITH_EDITOR
-bool UNiagaraDataInterfaceGrid3DCollection::GenerateIterationSourceNamespaceReadAttributesHLSL(FNiagaraDataInterfaceGPUParamInfo& DIInstanceInfo, const FNiagaraVariable& IterationSourceVar, TConstArrayView<FNiagaraVariable> InArguments, TConstArrayView<FNiagaraVariable> InAttributes, TConstArrayView<FString> InAttributeHLSLNames, bool bInSetToDefaults, bool bPartialWrites, TArray<FText>& OutErrors, FString& OutHLSL) const
+bool UNiagaraDataInterfaceGrid3DCollection::GenerateIterationSourceNamespaceReadAttributesHLSL(FNiagaraDataInterfaceGPUParamInfo& DIInstanceInfo, const FNiagaraVariable& IterationSourceVar, TConstArrayView<FNiagaraVariable> InArguments, TConstArrayView<FNiagaraVariable> InAttributes, TConstArrayView<FString> InAttributeHLSLNames, bool bInSetToDefaults, bool bSpawnOnly, bool bPartialWrites, TArray<FText>& OutErrors, FString& OutHLSL) const
 {
 
 	FString DIVarName;
@@ -1619,7 +1705,7 @@ bool UNiagaraDataInterfaceGrid3DCollection::GenerateIterationSourceNamespaceRead
 	}
 	return true;
 }
-bool UNiagaraDataInterfaceGrid3DCollection::GenerateIterationSourceNamespaceWriteAttributesHLSL(FNiagaraDataInterfaceGPUParamInfo& DIInstanceInfo, const FNiagaraVariable& IterationSourceVar, TConstArrayView<FNiagaraVariable> InArguments, TConstArrayView<FNiagaraVariable> InAttributes, TConstArrayView<FString> InAttributeHLSLNames, bool bPartialWrites, TArray<FText>& OutErrors, FString& OutHLSL) const
+bool UNiagaraDataInterfaceGrid3DCollection::GenerateIterationSourceNamespaceWriteAttributesHLSL(FNiagaraDataInterfaceGPUParamInfo& DIInstanceInfo, const FNiagaraVariable& IterationSourceVar, TConstArrayView<FNiagaraVariable> InArguments, TConstArrayView<FNiagaraVariable> InAttributes, TConstArrayView<FString> InAttributeHLSLNames, bool bSpawnOnly, bool bPartialWrites, TArray<FText>& OutErrors, FString& OutHLSL) const
 {
 	FString DIVarName;
 	OutHLSL += TEXT("\t//Generated by UNiagaraDataInterfaceGrid3DCollection::GenerateIterationSourceNamespaceWriteAttributesHLSL\n");
@@ -1634,16 +1720,16 @@ bool UNiagaraDataInterfaceGrid3DCollection::GenerateIterationSourceNamespaceWrit
 	if (InAttributes.Num() != InAttributeHLSLNames.Num())
 		return false;
 
-	// First we need to copy all the data over from the input buffer, because we can't assume that this function will know all the attributes held within the grid. Instead, we copy all of them
-	// over AND THEN overlay the local changes. Hopefully the optimizer will know enough to fix this up.
-	if (InAttributes.Num() > 0)
-	{
-		OutHLSL += FString::Printf(TEXT("\tint X, Y, Z;\n\t%s.ExecutionIndexToGridIndex(X, Y, Z);\n"), *DIVarName);
-	}
-
 	TArray<FString> RootArray;
 	IterationSourceVar.GetName().ToString().ParseIntoArray(RootArray, TEXT("."));
 
+	OutHLSL += FString("\tint AttributeIsSetMask = 0;\n");		
+	OutHLSL += FString("\tint CurrAttributeIndex;\n");
+	OutHLSL += FString("\tint X, Y, Z;\n");
+	OutHLSL += FString::Printf(TEXT("\t%s.ExecutionIndexToGridIndex(X, Y, Z);\n"), *DIVarName);
+	
+	
+	int NumAttributesSet = 0;
 	for (int32 i = 0; i < InAttributes.Num(); i++)
 	{
 		OutHLSL += FString::Printf(TEXT("\t// Name \"%s\" Type \"%s\" Var \"%s\"\n"), *InAttributes[i].GetName().ToString(), *InAttributes[i].GetType().GetName(), *InAttributeHLSLNames[i]);
@@ -1673,11 +1759,36 @@ bool UNiagaraDataInterfaceGrid3DCollection::GenerateIterationSourceNamespaceWrit
 				if (AttributeName.Len() != 0)
 					AttributeName += TEXT(".");
 				AttributeName += OutArray[NamespaceIdx];
-			}
+			}						
+
+			OutHLSL += FString::Printf(TEXT("\t%s.%s<Attribute=\"%s\">(CurrAttributeIndex);\n"), *DIVarName, *TypeDefinitionToAttributeIndexFunctionName(InAttributes[i].GetType()).ToString(), *AttributeName);
+			
+			int ComponentCount = GetComponentCountFromFuncName(TypeDefinitionToSetFunctionName(InAttributes[i].GetType()));
 
+			for (int j = 0; j < ComponentCount; ++j)
+			{
+				OutHLSL += FString::Printf(TEXT("\tAttributeIsSetMask |= 1 << (CurrAttributeIndex+%d);\n"), j);				
+				NumAttributesSet++;				
+			}			
 			OutHLSL += FString::Printf(TEXT("\t%s.%s<Attribute=\"%s\">(X, Y, Z,  %s);\n"), *DIVarName, *TypeDefinitionToSetFunctionName(InAttributes[i].GetType()).ToString(), *AttributeName, *InAttributeHLSLNames[i]);
 		}
 	}
+
+	// First we need to copy all the data over from the input buffer, because we can't assume that this function will know all the attributes held within the grid. Instead, we copy all of them
+// over AND THEN overlay the local changes. Hopefully the optimizer will know enough to fix this up.
+	if (InAttributes.Num() > 0 && !bSpawnOnly && !bPartialWrites)
+	{
+
+		static const TCHAR* FormatBounds = TEXT(R"(			
+			{Grid}.CopyMaskedPreviousToCurrentForCell(X,Y,Z,{NumAttributesSet},AttributeIsSetMask);
+		)");
+		TMap<FString, FStringFormatArg> ArgsBounds = {
+			{TEXT("Grid"), DIVarName},
+			{TEXT("NumAttributesSet"), NumAttributesSet},
+		};
+		OutHLSL += FString::Format(FormatBounds, ArgsBounds);
+
+	}
 	return true;
 }
 
@@ -1696,20 +1807,6 @@ bool UNiagaraDataInterfaceGrid3DCollection::GenerateSetupHLSL(FNiagaraDataInterf
 		}
 	}
 
-	if (!bSpawnOnly && !bPartialWrites)
-	{
-		static const TCHAR* FormatBounds = TEXT(R"(
-			// We need to copy from previous to current first thing, because other functions afterwards may just set values on the local grid.
-			int X, Y, Z;
-			{Grid}.ExecutionIndexToGridIndex(X, Y, Z);
-			{Grid}.CopyPreviousToCurrentForCell(X,Y, Z);
-		)");
-		TMap<FString, FStringFormatArg> ArgsBounds = {
-			{TEXT("Grid"), DIVarName},
-		};
-		OutHLSL += FString::Format(FormatBounds, ArgsBounds);
-	}
-
 	return true;
 }
 
@@ -1832,12 +1929,18 @@ bool UNiagaraDataInterfaceGrid3DCollection::InitPerInstanceData(void* PerInstanc
 	int32 MaxTilesY = floor(MaxDim / InstanceData->NumCells.Y);
 	int32 MaxTilesZ = floor(MaxDim / InstanceData->NumCells.Z);
 	int32 MaxAttributes = MaxTilesX * MaxTilesY * MaxTilesZ;
-	if ((NumAttribChannelsFound > MaxAttributes && MaxAttributes > 0) || NumAttribChannelsFound == 0)
+	if (NumAttribChannelsFound > MaxAttributes && MaxAttributes > 0)
 	{
 		UE_LOG(LogNiagara, Error, TEXT("Invalid number of attributes defined on %s... max is %i, num defined is %i"), *FNiagaraUtilities::SystemInstanceIDToString(SystemInstance->GetId()), MaxAttributes, NumAttribChannelsFound);
 		return false;
 	}
 
+	if (NumAttribChannelsFound == 0)
+	{
+		UE_LOG(LogNiagara, Warning, TEXT("Invalid number of attributes defined on %s... max is %i, num defined is %i"), *FNiagaraUtilities::SystemInstanceIDToString(SystemInstance->GetId()), MaxAttributes, NumAttribChannelsFound);
+		return false;
+	}
+
 	// need to determine number of tiles in x and y based on number of attributes and max dimension size
 	int32 NumTilesX = FMath::Min<int32>(MaxTilesX, NumAttribChannelsFound);
 	int32 NumTilesY = FMath::Min<int32>(MaxTilesY, ceil(1.0 * NumAttribChannelsFound / NumTilesX));
diff --git a/Engine/Plugins/FX/Niagara/Source/NiagaraCore/Private/NiagaraCustomVersion.cpp b/Engine/Plugins/FX/Niagara/Source/NiagaraCore/Private/NiagaraCustomVersion.cpp
--- a/Engine/Plugins/FX/Niagara/Source/NiagaraCore/Private/NiagaraCustomVersion.cpp
+++ b/Engine/Plugins/FX/Niagara/Source/NiagaraCore/Private/NiagaraCustomVersion.cpp
@@ -8,5 +8,4 @@ const FGuid FNiagaraCustomVersion::GUID(0xFCF57AFA, 0x50764283, 0xB9A9E658, 0xFF
 // Register the custom version with core
 FCustomVersionRegistration GRegisterNiagaraCustomVersion(FNiagaraCustomVersion::GUID, FNiagaraCustomVersion::LatestVersion, TEXT("NiagaraVer"));
 
-const FGuid FNiagaraCustomVersion::LatestScriptCompileVersion(0xE1236282, 0x02E2477E, 0x808C6EC6, 0xAC96EE7F);
-
+const FGuid FNiagaraCustomVersion::LatestScriptCompileVersion(0x323AF294, 0x45834FA5, 0x835EB6A5, 0x80C83B36);
\ No newline at end of file
diff --git a/Engine/Plugins/FX/Niagara/Source/NiagaraEditor/Private/NiagaraHlslTranslator.cpp b/Engine/Plugins/FX/Niagara/Source/NiagaraEditor/Private/NiagaraHlslTranslator.cpp
--- a/Engine/Plugins/FX/Niagara/Source/NiagaraEditor/Private/NiagaraHlslTranslator.cpp
+++ b/Engine/Plugins/FX/Niagara/Source/NiagaraEditor/Private/NiagaraHlslTranslator.cpp
@@ -2127,7 +2127,7 @@ void FHlslNiagaraTranslator::HandleSimStageSetupAndTeardown(int32 InWhichStage,
 				AttributeHLSLNames.Emplace(TEXT("Map.") + GetSanitizedSymbolName(Var.GetName().ToString()));
 			}
 
-			if (CDO->GenerateIterationSourceNamespaceReadAttributesHLSL(DIInstanceInfo,  IterationSourceVar, MakeArrayView(Sig.Inputs), MakeArrayView(ReadVars), MakeArrayView(AttributeHLSLNames), TranslationStage.bSpawnOnly, bPartialWrites, GeneratedErrors, AttributeReadGeneratedHLSL) && AttributeReadGeneratedHLSL.Len() > 0)
+			if (CDO->GenerateIterationSourceNamespaceReadAttributesHLSL(DIInstanceInfo,  IterationSourceVar, MakeArrayView(Sig.Inputs), MakeArrayView(ReadVars), MakeArrayView(AttributeHLSLNames), TranslationStage.bSpawnOnly, TranslationStage.bSpawnOnly, bPartialWrites, GeneratedErrors, AttributeReadGeneratedHLSL) && AttributeReadGeneratedHLSL.Len() > 0)
 			{
 				Sig.Name = *FString::Printf(TEXT("SetupFromIterationSource_%s_GeneratedReadAttributes"), *GetSanitizedFunctionNameSuffix(TranslationStage.PassNamespace));
 				FNiagaraFunctionSignature SignatureOut = Sig;
@@ -2148,7 +2148,7 @@ void FHlslNiagaraTranslator::HandleSimStageSetupAndTeardown(int32 InWhichStage,
 				AttributeHLSLNames.Emplace(TEXT("Map.") + GetSanitizedSymbolName(Var.GetName().ToString()));
 			}
 
-			if (CDO->GenerateIterationSourceNamespaceWriteAttributesHLSL(DIInstanceInfo, IterationSourceVar, MakeArrayView(Sig.Inputs), MakeArrayView(WriteVars), MakeArrayView(AttributeHLSLNames), bPartialWrites, GeneratedErrors, AttributeWriteGeneratedHLSL) && AttributeWriteGeneratedHLSL.Len() > 0)
+			if (CDO->GenerateIterationSourceNamespaceWriteAttributesHLSL(DIInstanceInfo, IterationSourceVar, MakeArrayView(Sig.Inputs), MakeArrayView(WriteVars), MakeArrayView(AttributeHLSLNames), TranslationStage.bSpawnOnly, bPartialWrites, GeneratedErrors, AttributeWriteGeneratedHLSL) && AttributeWriteGeneratedHLSL.Len() > 0)
 			{
 				Sig.Name = *FString::Printf(TEXT("TeardownFromIterationSource_%s_GeneratedWriteAttributes"), *GetSanitizedFunctionNameSuffix(TranslationStage.PassNamespace));
 				FNiagaraFunctionSignature SignatureOut = Sig;
