Fix incorrectly merged whitespace.
#rnx

diff --git a/Engine/Plugins/FX/Niagara/Source/Niagara/Private/NiagaraScript.cpp b/Engine/Plugins/FX/Niagara/Source/Niagara/Private/NiagaraScript.cpp
--- a/Engine/Plugins/FX/Niagara/Source/Niagara/Private/NiagaraScript.cpp
+++ b/Engine/Plugins/FX/Niagara/Source/Niagara/Private/NiagaraScript.cpp
@@ -852,16 +852,16 @@ void UNiagaraScript::ComputeVMCompilationId(FNiagaraVMExecutableDataId& Id, FGui
 					TArray<const UNiagaraDataInterfaceBase*> DataInterfaces;
 					if (OtherEmitter && OtherEmitter->GraphSource)
 					{
-					OtherEmitter->GraphSource->CollectDataInterfaces(DataInterfaces);
+						OtherEmitter->GraphSource->CollectDataInterfaces(DataInterfaces);
 
-					for (const UNiagaraDataInterfaceBase* DataInterface : DataInterfaces)
-					{
-						if (DataInterface->HasInternalAttributeReads(OtherEmitter, Emitter))
+						for (const UNiagaraDataInterfaceBase* DataInterface : DataInterfaces)
 						{
-							return false;
+							if (DataInterface->HasInternalAttributeReads(OtherEmitter, Emitter))
+							{
+								return false;
+							}
 						}
-					}
-					return true;
+						return true;
 					}
 					else
 					{
@@ -1062,8 +1062,8 @@ void UNiagaraScript::ComputeVMCompilationId(FNiagaraVMExecutableDataId& Id, FGui
 	{
 		if (ScriptData->Source)
 		{
-		ScriptData->Source->ComputeVMCompilationId(Id, Usage, UsageId);
-	}
+			ScriptData->Source->ComputeVMCompilationId(Id, Usage, UsageId);
+		}
 	}
 
 	FNiagaraVMExecutableDataId& LastGeneratedVMId = GetLastGeneratedVMId(VersionGuid);
@@ -1291,23 +1291,23 @@ void UNiagaraScript::AsyncOptimizeByteCode(bool bIsInPostLoad)
 		if (OptimizedByteCode.HasByteCode())
 		{
 			Result->OptimizedByteCode = MoveTemp(OptimizedByteCode);
-	}
+		}
 
 		return Result;
 	};
 
 	if (!bShouldOptimizeImmediately)
-					{
+	{
 		CachedScriptVM.OptimizationTask = MakeShared<FNiagaraVMExecutableByteCodeOptimizationTaskState, ESPMode::ThreadSafe>();
 		CachedScriptVM.OptimizationTask->SharedTask = Async(EAsyncExecution::TaskGraph, MoveTemp(Task)).Share();
 	}
 	else
-						{
+	{
 		FNiagaraVMExecutableByteCodeOptimizationTaskResultPtr Result = Task();
 		CachedScriptVM.ApplyFinishedOptimization(Result);
 		CachedScriptVM.OptimizationTask = MakeShared<FNiagaraVMExecutableByteCodeOptimizationTaskState, ESPMode::ThreadSafe>(true);
-					}
-			}
+	}
+}
 
 void UNiagaraScript::GenerateDefaultFunctionBindings()
 {
@@ -1487,7 +1487,7 @@ void UNiagaraScript::Serialize(FArchive& Ar)
 		if (GNiagaraCompressScriptByteCode)
 		{
 			ExecutableData.ByteCode.Compress();
-	}
+		}
 	}
 
 	if (Ar.IsLoading())
