- All GPU stages are now 'simulation stages'
- The intrinsic particle stage is now just another simulation stage that can be enabled or not
- Iteration count is now dynamic and can be bound to a user variable.  In the future we should be able to skip compilation on the constant version.
- Stages can be enabled / disabled dynamically by a user variable
- Brought some clarification to SpawnOnly on simulation stages

diff --git a/Engine/Plugins/FX/Niagara/Source/Niagara/Classes/NiagaraComputeExecutionContext.h b/Engine/Plugins/FX/Niagara/Source/Niagara/Classes/NiagaraComputeExecutionContext.h
--- a/Engine/Plugins/FX/Niagara/Source/Niagara/Classes/NiagaraComputeExecutionContext.h
+++ b/Engine/Plugins/FX/Niagara/Source/Niagara/Classes/NiagaraComputeExecutionContext.h
@@ -33,6 +33,7 @@ struct FNiagaraSimStageData
 	uint32 bSetDataToRender : 1;
 
 	uint32 StageIndex = INDEX_NONE;
+	uint32 IterationIndex = 0;
 
 	FNiagaraDataBuffer* Source = nullptr;
 	uint32 SourceCountOffset = INDEX_NONE;
@@ -191,9 +192,6 @@ public:
 	// Game thread spawn info will be sent to the render thread inside FNiagaraComputeInstanceData
 	FNiagaraGpuSpawnInfo GpuSpawnInfo_GT;
 
-	uint32 MaxUpdateIterations = 0;
-	TSet<uint32> SpawnStages;
-
 	bool HasInterpolationParameters = false;
 
 	// Do we have a reset pending, controlled by the game thread and passed to instance data
@@ -229,8 +227,7 @@ public:
 
 	TArray<FSimulationStageMetaData> SimStageInfo;
 
-	bool IsOutputStage(FNiagaraDataInterfaceProxy* DIProxy, uint32 CurrentStage) const;
-	bool IsIterationStage(FNiagaraDataInterfaceProxy* DIProxy, uint32 CurrentStage) const;
+	bool IsOutputStage(FNiagaraDataInterfaceProxy* DIProxy, uint32 SimulationStageIndex) const;
+	bool IsIterationStage(FNiagaraDataInterfaceProxy* DIProxy, uint32 SimulationStageIndex) const;
 	FNiagaraDataInterfaceProxyRW* FindIterationInterface(const TArray<FNiagaraDataInterfaceProxyRW*>& InProxies, uint32 SimulationStageIndex) const;
-	const FSimulationStageMetaData* GetSimStageMetaData(uint32 SimulationStageIndex) const;
 };
diff --git a/Engine/Plugins/FX/Niagara/Source/Niagara/Classes/NiagaraGPUSystemTick.h b/Engine/Plugins/FX/Niagara/Source/Niagara/Classes/NiagaraGPUSystemTick.h
--- a/Engine/Plugins/FX/Niagara/Source/Niagara/Classes/NiagaraGPUSystemTick.h
+++ b/Engine/Plugins/FX/Niagara/Source/Niagara/Classes/NiagaraGPUSystemTick.h
@@ -28,7 +28,9 @@ struct FNiagaraComputeInstanceData
 	FNiagaraComputeExecutionContext* Context = nullptr;
 	TArray<FNiagaraDataInterfaceProxy*> DataInterfaceProxies;
 	TArray<FNiagaraDataInterfaceProxyRW*> IterationDataInterfaceProxies;
+	TArray<int32, TInlineAllocator<1>> NumIterationsPerStage;
 	uint32 ParticleCountFence = INDEX_NONE;
+	uint32 TotalDispatches = 0;
 	uint32 bResetData : 1;
 	uint32 bStartNewOverlapGroup : 1;
 	uint32 bUsesSimStages : 1;
@@ -98,7 +100,6 @@ public:
 	uint8* OwnerParamData = nullptr;
 	uint32 InstanceCount = 0;
 	uint32 TotalDispatches = 0;
-	uint32 NumInstancesWithSimStages = 0;									//-TODO: Remove me
 	bool bIsFinalTick = false;
 
 #if !UE_BUILD_SHIPPING && !UE_BUILD_TEST
diff --git a/Engine/Plugins/FX/Niagara/Source/Niagara/Private/NiagaraCommon.cpp b/Engine/Plugins/FX/Niagara/Source/Niagara/Private/NiagaraCommon.cpp
--- a/Engine/Plugins/FX/Niagara/Source/Niagara/Private/NiagaraCommon.cpp
+++ b/Engine/Plugins/FX/Niagara/Source/Niagara/Private/NiagaraCommon.cpp
@@ -779,7 +779,11 @@ void FNiagaraVariableAttributeBinding::CacheValues(const UNiagaraEmitter* InEmit
 	}
 
 	// If this is one of the possible namespaces that is implicitly defined, go ahead and expand the full namespace. RootVariable should be non-namespaced at this point.
-	if ((InSourceMode == ENiagaraRendererSourceDataMode::Emitter && BindingSourceMode == ENiagaraBindingSource::ImplicitFromSource) ||
+	if (DataSetVariable.GetName().IsNone())
+	{
+		ParamMapVariable.SetName(NAME_None);
+	}
+	else if ((InSourceMode == ENiagaraRendererSourceDataMode::Emitter && BindingSourceMode == ENiagaraBindingSource::ImplicitFromSource) ||
 		BindingSourceMode == ENiagaraBindingSource::ExplicitEmitter)
 	{
 		ensure(!DataSetVariable.IsInNameSpace(FNiagaraConstants::EmitterNamespace));
diff --git a/Engine/Plugins/FX/Niagara/Source/Niagara/Private/NiagaraComputeExecutionContext.cpp b/Engine/Plugins/FX/Niagara/Source/Niagara/Private/NiagaraComputeExecutionContext.cpp
--- a/Engine/Plugins/FX/Niagara/Source/Niagara/Private/NiagaraComputeExecutionContext.cpp
+++ b/Engine/Plugins/FX/Niagara/Source/Niagara/Private/NiagaraComputeExecutionContext.cpp
@@ -46,97 +46,18 @@ void FNiagaraComputeExecutionContext::InitParams(UNiagaraScript* InGPUComputeScr
 {
 	GPUScript = InGPUComputeScript;
 	CombinedParamStore.InitFromOwningContext(InGPUComputeScript, InSimTarget, true);
-	SpawnStages.Empty();
 	
 	HasInterpolationParameters = GPUScript && GPUScript->GetComputedVMCompilationId().HasInterpolatedParameters();
-	MaxUpdateIterations = 1;
 
 	if (InGPUComputeScript)
 	{
 		FNiagaraVMExecutableData& VMData = InGPUComputeScript->GetVMExecutableData();
-		if (VMData.IsValid() && VMData.SimulationStageMetaData.Num() > 0)
+		if ( VMData.IsValid() )
 		{
 			SimStageInfo = VMData.SimulationStageMetaData;
-
-			int32 FoundMaxUpdateIterations = SimStageInfo[SimStageInfo.Num() - 1].MaxStage;
-
-			// Some useful debugging code should we need to look up differences between old and new
-			const bool bDebugSimStages = false;
-			if (bDebugSimStages)
-			{
-				UE_LOG(LogNiagara, Log, TEXT("Stored vs:"));
-				bool bPass = FoundMaxUpdateIterations == MaxUpdateIterations;
-				UE_LOG(LogNiagara, Log, TEXT("MaxUpdateIterations: %d vs %d %s"), FoundMaxUpdateIterations, MaxUpdateIterations, bPass ? TEXT("Pass") : TEXT("FAIL!!!!!!!!"));
-
-				int32 NumSpawnFound = 0;
-				bool bMatchesFound = true;
-				for (int32 i = 0; i < SimStageInfo.Num(); i++)
-				{
-					if (SimStageInfo[i].bSpawnOnly)
-					{
-						NumSpawnFound++;
-
-						if (!SpawnStages.Contains(SimStageInfo[i].MinStage))
-						{
-							bMatchesFound = false;
-							UE_LOG(LogNiagara, Log, TEXT("Missing spawn stage: %d FAIL!!!!!!!!!"), SimStageInfo[i].MinStage);
-						}
-					}
-				}
-
-				bPass = SpawnStages.Num() == NumSpawnFound;
-				UE_LOG(LogNiagara, Log, TEXT("SpawnStages.Num(): %d vs %d %s"), NumSpawnFound, SpawnStages.Num(), bPass ? TEXT("Pass") : TEXT("FAIL!!!!!!!!"));
-
-				TArray<FNiagaraVariable> Params;
-				CombinedParamStore.GetParameters(Params);
-				for (FNiagaraVariable& Var : Params)
-				{
-					if (!Var.IsDataInterface())
-						continue;
-
-					UNiagaraDataInterface* DI = CombinedParamStore.GetDataInterface(Var);
-					UNiagaraDataInterfaceRWBase* DIRW = Cast<UNiagaraDataInterfaceRWBase>(DI);
-					if (DIRW)
-					{
-						for (int32 i = 0; i < SimStageInfo.Num(); i++)
-						{
-							if (SimStageInfo[i].IterationSource == Var.GetName())
-							{
-								if (!DIRW->IterationShaderStages.Contains(SimStageInfo[i].MinStage))
-								{
-									UE_LOG(LogNiagara, Log, TEXT("Missing iteration stage for %s: %d FAIL!!!!!!!!!"), *Var.GetName().ToString(), SimStageInfo[i].MinStage);
-								}
-							}
-
-							if (SimStageInfo[i].OutputDestinations.Contains(Var.GetName()))
-							{
-								if (!DIRW->OutputShaderStages.Contains(SimStageInfo[i].MinStage))
-								{
-									UE_LOG(LogNiagara, Log, TEXT("Missing output stage for %s: %d FAIL!!!!!!!!!"), *Var.GetName().ToString(), SimStageInfo[i].MinStage);
-								}
-							}
-						}
-					}
-
-				}
-			}
-
-
-			// Set the values that we are using from compiled data instead...
-			MaxUpdateIterations = SimStageInfo[SimStageInfo.Num() - 1].MaxStage;
-			SpawnStages.Empty();
-
-			for (int32 i = 0; i < SimStageInfo.Num(); i++)
-			{
-				if (SimStageInfo[i].bSpawnOnly)
-				{
-					SpawnStages.Add(SimStageInfo[i].MinStage);
-				}
-			}
 		}
 	}
 
-	
 #if DO_CHECK
 	// DI Parameters are the same between all shader permutations so we can just get the first one
 	FNiagaraShaderRef Shader = InGPUComputeScript->GetRenderThreadScript()->GetShaderGameThread(0);
@@ -159,66 +80,42 @@ void FNiagaraComputeExecutionContext::InitParams(UNiagaraScript* InGPUComputeScr
 #endif
 }
 
-
-const FSimulationStageMetaData* FNiagaraComputeExecutionContext::GetSimStageMetaData(uint32 SimulationStageIndex) const
+bool FNiagaraComputeExecutionContext::IsOutputStage(FNiagaraDataInterfaceProxy* DIProxy, uint32 SimulationStageIndex) const
 {
-	if (SimStageInfo.Num() > 0)
-	{
-		for (int32 i = 0; i < SimStageInfo.Num(); i++)
-		{
-			if (SimulationStageIndex >= (uint32)SimStageInfo[i].MinStage  && SimulationStageIndex < (uint32)SimStageInfo[i].MaxStage)
-			{
-				return &SimStageInfo[i];
-			}
-		}
+	if (DIProxy && !DIProxy->SourceDIName.IsNone())
+	{		
+		return SimStageInfo[SimulationStageIndex].OutputDestinations.Contains(DIProxy->SourceDIName);
 	}
-	return nullptr;
+	return false;
 }
 
-bool FNiagaraComputeExecutionContext::IsOutputStage(FNiagaraDataInterfaceProxy* DIProxy, uint32 CurrentStage) const
+bool FNiagaraComputeExecutionContext::IsIterationStage(FNiagaraDataInterfaceProxy* DIProxy, uint32 SimulationStageIndex) const
 {
-	const FSimulationStageMetaData* MetaData = GetSimStageMetaData(CurrentStage);
-	if (MetaData && DIProxy && !DIProxy->SourceDIName.IsNone())
+	if (DIProxy && !DIProxy->SourceDIName.IsNone())
 	{
-		if (MetaData->OutputDestinations.Contains(DIProxy->SourceDIName))
-			return true;
+		return !SimStageInfo[SimulationStageIndex].IterationSource.IsNone() && (SimStageInfo[SimulationStageIndex].IterationSource == DIProxy->SourceDIName);
 	}
 	return false;
 }
 
-bool FNiagaraComputeExecutionContext::IsIterationStage(FNiagaraDataInterfaceProxy* DIProxy, uint32 CurrentStage) const
+FNiagaraDataInterfaceProxyRW* FNiagaraComputeExecutionContext::FindIterationInterface(const TArray<FNiagaraDataInterfaceProxyRW*>& InProxies, uint32 SimulationStageIndex) const
 {
-	const FSimulationStageMetaData* MetaData = GetSimStageMetaData(CurrentStage);
-	if (MetaData && DIProxy && !DIProxy->SourceDIName.IsNone())
+	// Particle stage
+	if ( SimStageInfo[SimulationStageIndex].IterationSource.IsNone() )
 	{
-		if (MetaData->IterationSource.IsNone()) // Per particle iteration...
-			return false;
-
-		if (MetaData->IterationSource == DIProxy->SourceDIName)
-			return true;
+		return nullptr;
 	}
-	return false;
-}
 
-FNiagaraDataInterfaceProxyRW* FNiagaraComputeExecutionContext::FindIterationInterface(const TArray<FNiagaraDataInterfaceProxyRW*>& InProxies, uint32 CurrentStage) const
-{
-	const FSimulationStageMetaData* MetaData = GetSimStageMetaData(CurrentStage);
-	if (MetaData)
+	for (FNiagaraDataInterfaceProxyRW* Proxy : InProxies)
 	{
-		if (MetaData->IterationSource.IsNone()) // Per particle iteration...
-			return nullptr;
-
-		for (FNiagaraDataInterfaceProxyRW* Proxy : InProxies)
+		if (Proxy->SourceDIName == SimStageInfo[SimulationStageIndex].IterationSource)
 		{
-			if (Proxy->SourceDIName == MetaData->IterationSource)
-				return Proxy;
+			return Proxy;
 		}
-
-		UE_LOG(LogNiagara, Verbose, TEXT("FNiagaraComputeExecutionContext::FindIterationInterface could not find IterationInterface %s"), *MetaData->IterationSource.ToString());
-
-		return nullptr;
 	}
 
+	UE_LOG(LogNiagara, Verbose, TEXT("FNiagaraComputeExecutionContext::FindIterationInterface could not find IterationInterface %s"), *SimStageInfo[SimulationStageIndex].IterationSource.ToString());
+
 	return nullptr;
 }
 
@@ -347,20 +244,20 @@ void FNiagaraComputeExecutionContext::SetTranslucentDataToRender(FNiagaraDataBuf
 	}
 }
 
-bool FNiagaraComputeInstanceData::IsOutputStage(FNiagaraDataInterfaceProxy* DIProxy, uint32 CurrentStage) const
+bool FNiagaraComputeInstanceData::IsOutputStage(FNiagaraDataInterfaceProxy* DIProxy, uint32 SimulationStageIndex) const
 {
 	if (bUsesSimStages)
 	{
-		return Context->IsOutputStage(DIProxy, CurrentStage);
+		return Context->IsOutputStage(DIProxy, SimulationStageIndex);
 	}
 	return false;
 }
 
-bool FNiagaraComputeInstanceData::IsIterationStage(FNiagaraDataInterfaceProxy* DIProxy, uint32 CurrentStage) const
+bool FNiagaraComputeInstanceData::IsIterationStage(FNiagaraDataInterfaceProxy* DIProxy, uint32 SimulationStageIndex) const
 {
 	if (bUsesSimStages)
 	{
-		return Context->IsIterationStage(DIProxy, CurrentStage);
+		return Context->IsIterationStage(DIProxy, SimulationStageIndex);
 	}
 	return false;
 }
diff --git a/Engine/Plugins/FX/Niagara/Source/Niagara/Private/NiagaraEmitter.cpp b/Engine/Plugins/FX/Niagara/Source/Niagara/Private/NiagaraEmitter.cpp
--- a/Engine/Plugins/FX/Niagara/Source/Niagara/Private/NiagaraEmitter.cpp
+++ b/Engine/Plugins/FX/Niagara/Source/Niagara/Private/NiagaraEmitter.cpp
@@ -555,6 +555,22 @@ void UNiagaraEmitter::PostLoad()
 	// there will be no merging or compiling which makes it safe to do so.
 	UpdateEmitterAfterLoad();
 #endif
+
+	//-TODO: should this be a version
+#if WITH_EDITORONLY_DATA
+	// If we have a simulation stage using the old bSpawnOnly we need to modify all generic stages to the new method
+	if ( SimulationStages.ContainsByPredicate([](UNiagaraSimulationStageBase* Stage) { UNiagaraSimulationStageGeneric* StageAsGeneric = Cast<UNiagaraSimulationStageGeneric>(Stage); return StageAsGeneric && StageAsGeneric->bSpawnOnly_DEPRECATED; }) )
+	{
+		for (UNiagaraSimulationStageBase* Stage : SimulationStages)
+		{
+			if ( UNiagaraSimulationStageGeneric* StageAsGeneric = Cast<UNiagaraSimulationStageGeneric>(Stage) )
+			{
+				StageAsGeneric->ExecuteBehavior = StageAsGeneric->bSpawnOnly_DEPRECATED ? ENiagaraSimStageExecuteBehavior::OnSimulationReset : ENiagaraSimStageExecuteBehavior::NotOnSimulationReset;
+				StageAsGeneric->bSpawnOnly_DEPRECATED = false;
+			}
+		}
+	}
+#endif
 }
 
 bool UNiagaraEmitter::IsEditorOnly() const
diff --git a/Engine/Plugins/FX/Niagara/Source/Niagara/Private/NiagaraEmitterInstance.cpp b/Engine/Plugins/FX/Niagara/Source/Niagara/Private/NiagaraEmitterInstance.cpp
--- a/Engine/Plugins/FX/Niagara/Source/Niagara/Private/NiagaraEmitterInstance.cpp
+++ b/Engine/Plugins/FX/Niagara/Source/Niagara/Private/NiagaraEmitterInstance.cpp
@@ -488,8 +488,9 @@ void FNiagaraEmitterInstance::Init(int32 InEmitterIdx, FNiagaraSystemInstanceID
 			}
 		}
 
-
 		// We may need to populate bindings that will be used in rendering
+		RendererBindings.UnbindAll();
+
 		bool bAnyRendererBindingsAdded = false;
 		for (UNiagaraRendererProperties* Props : CachedEmitter->GetRenderers())
 		{
@@ -499,6 +500,21 @@ void FNiagaraEmitterInstance::Init(int32 InEmitterIdx, FNiagaraSystemInstanceID
 			}
 		}
 
+		if ( GPUExecContext != nullptr )
+		{
+			for (const FSimulationStageMetaData& SimStageMetaData : GPUExecContext->SimStageInfo )
+			{
+				if (!SimStageMetaData.EnabledBinding.IsNone())
+				{
+					bAnyRendererBindingsAdded |= RendererBindings.AddParameter(FNiagaraVariable(FNiagaraTypeDefinition::GetBoolDef(), SimStageMetaData.EnabledBinding), false);
+				}
+				if ( !SimStageMetaData.NumIterationsBinding.IsNone() )
+				{
+					bAnyRendererBindingsAdded |= RendererBindings.AddParameter(FNiagaraVariable(FNiagaraTypeDefinition::GetIntDef(), SimStageMetaData.NumIterationsBinding), false);
+				}				
+			}
+		}
+
 		if (bAnyRendererBindingsAdded)
 		{
 			if (ParentSystemInstance)
diff --git a/Engine/Plugins/FX/Niagara/Source/Niagara/Private/NiagaraEmitterInstanceBatcher.cpp b/Engine/Plugins/FX/Niagara/Source/Niagara/Private/NiagaraEmitterInstanceBatcher.cpp
--- a/Engine/Plugins/FX/Niagara/Source/Niagara/Private/NiagaraEmitterInstanceBatcher.cpp
+++ b/Engine/Plugins/FX/Niagara/Source/Niagara/Private/NiagaraEmitterInstanceBatcher.cpp
@@ -155,17 +155,6 @@ namespace NiagaraEmitterInstanceBatcherLocal
 			AfterTransitionArray.Emplace(IntUAV, AfterState, BeforeState);
 		}
 	}
-
-	static bool ShouldRunStage(FNiagaraComputeExecutionContext* ComputeContext, FNiagaraDataInterfaceProxyRW* IterationInterface, uint32 StageIndex, bool bResetData)
-	{
-		if (!IterationInterface || ComputeContext->SpawnStages.Num() == 0)
-		{
-			return true;
-		}
-
-		const bool bIsSpawnStage = ComputeContext->SpawnStages.Contains(StageIndex);
-		return bResetData == bIsSpawnStage;
-	}
 }
 
 FNiagaraUAVPoolAccessScope::FNiagaraUAVPoolAccessScope(NiagaraEmitterInstanceBatcher& InBatcher)
@@ -739,14 +728,7 @@ void NiagaraEmitterInstanceBatcher::DumpDebugFrame()
 				Builder.Appendf(TEXT("Proxy(%p) "), DispatchInstance.Tick.SystemGpuComputeProxy);
 				Builder.Appendf(TEXT("ComputeContext(%p) "), InstanceData.Context);
 				Builder.Appendf(TEXT("Emitter(%s) "), InstanceData.Context->GetDebugSimName());
-				if (SimStageData.StageMetaData)
-				{
-					Builder.Appendf(TEXT("Stage(%d | %s) "), SimStageData.StageIndex, *SimStageData.StageMetaData->SimulationStageName.ToString());
-				}
-				else
-				{
-					Builder.Appendf(TEXT("Stage(%d) "), SimStageData.StageIndex);
-				}
+				Builder.Appendf(TEXT("Stage(%d | %s) "), SimStageData.StageIndex, *SimStageData.StageMetaData->SimulationStageName.ToString());
 
 				if (InstanceData.bResetData)
 				{
@@ -783,7 +765,7 @@ void NiagaraEmitterInstanceBatcher::DumpDebugFrame()
 				}
 				Builder.Appendf(TEXT("Source(%p 0x%08x %d) "), SimStageData.Source, SimStageData.SourceCountOffset, SimStageData.SourceNumInstances);
 				Builder.Appendf(TEXT("Destination(%p 0x%08x %d) "), SimStageData.Destination, SimStageData.DestinationCountOffset, SimStageData.DestinationNumInstances);
-				Builder.Appendf(TEXT("Iteration(%s) "), SimStageData.AlternateIterationSource ? *SimStageData.AlternateIterationSource->SourceDIName.ToString() : TEXT("Particles"));
+				Builder.Appendf(TEXT("Iteration(%d | %s) "), SimStageData.IterationIndex, SimStageData.AlternateIterationSource ? *SimStageData.AlternateIterationSource->SourceDIName.ToString() : TEXT("Particles"));
 				UE_LOG(LogNiagara, Warning, TEXT("%s"), Builder.ToString());
 			}
 
@@ -961,11 +943,18 @@ void NiagaraEmitterInstanceBatcher::PrepareTicksForProxy(FRHICommandListImmediat
 			}
 #endif
 
+			// Nothing to dispatch?
+			if ( InstanceData.TotalDispatches == 0 )
+			{
+				continue;
+			}
+
 			// Determine this instances start dispatch group, in the case of emitter dependencies (i.e. particle reads) we need to continue rather than starting again
 			iInstanceStartDispatchGroup = InstanceData.bStartNewOverlapGroup ? iInstanceCurrDispatchGroup : iInstanceStartDispatchGroup;
 			iInstanceCurrDispatchGroup = iInstanceStartDispatchGroup;
 
-			GpuDispatchList.PreAllocateGroups(iInstanceCurrDispatchGroup + ComputeContext->MaxUpdateIterations);
+			// Pre-allocator groups
+			GpuDispatchList.PreAllocateGroups(iInstanceCurrDispatchGroup + InstanceData.TotalDispatches);
 
 			// Calculate instance counts
 			const uint32 MaxBufferInstances = ComputeContext->MainDataSet->GetMaxInstanceCount();
@@ -988,73 +977,37 @@ void NiagaraEmitterInstanceBatcher::PrepareTicksForProxy(FRHICommandListImmediat
 
 			bHasFreeIDUpdates |= ComputeContext->MainDataSet->RequiresPersistentIDs();
 
-			// Add particle update / spawn stage
+			//-OPT: Do we need this test?  Can remove in favor of MaxUpdateIterations
+			bool bFirstStage = true;
+			for (int32 SimStageIndex=0; SimStageIndex < ComputeContext->SimStageInfo.Num(); ++SimStageIndex)
 			{
-				FNiagaraGpuDispatchGroup& DispatchGroup = GpuDispatchList.DispatchGroups[iInstanceCurrDispatchGroup++];
-				FNiagaraGpuDispatchInstance& DispatchInstance = DispatchGroup.DispatchInstances.Emplace_GetRef(Tick, InstanceData);
-				FNiagaraSimStageData& SimStageData = DispatchInstance.SimStageData;
-
-				SimStageData.bFirstStage = true;
-				SimStageData.StageIndex = 0;
-				SimStageData.Source = ComputeContext->bHasTickedThisFrame_RT ? ComputeContext->GetPrevDataBuffer() : ComputeContext->MainDataSet->GetCurrentData();
-				SimStageData.SourceCountOffset = ComputeContext->CountOffset_RT;
-				SimStageData.SourceNumInstances = PrevNumInstances;
-
-				// Validation as a non zero instance count must have a valid count offset
-				check(SimStageData.SourceCountOffset != INDEX_NONE || SimStageData.SourceNumInstances == 0);
-
-				SimStageData.Destination = ComputeContext->GetNextDataBuffer();
-				SimStageData.DestinationCountOffset = GPUInstanceCounterManager.AcquireEntry();
-				SimStageData.DestinationNumInstances = ComputeContext->CurrentNumInstances_RT;
-
-				ComputeContext->AdvanceDataBuffer();
-				ComputeContext->CountOffset_RT = SimStageData.DestinationCountOffset;
-
-				// If we are the last tick then we may want to enqueue for a readback
-				// Note: Do not pull count from SimStageData as a reset tick will be INDEX_NONE
-				if (SimStageData.SourceCountOffset != INDEX_NONE)
-				{
-					if (bEnqueueCountReadback && Tick.bIsFinalTick && (ComputeContext->EmitterInstanceReadback.GPUCountOffset == INDEX_NONE))
-					{
-						bRequiresReadback = true;
-						ComputeContext->EmitterInstanceReadback.CPUCount = SimStageData.SourceNumInstances;
-						ComputeContext->EmitterInstanceReadback.GPUCountOffset = SimStageData.SourceCountOffset;
-					}
-					GpuDispatchList.CountsToRelease.Add(SimStageData.SourceCountOffset);
-				}
-				if (Tick.NumInstancesWithSimStages > 0)
+				const FSimulationStageMetaData& SimStageMetaData = ComputeContext->SimStageInfo[SimStageIndex];
+				if (InstanceData.NumIterationsPerStage[SimStageIndex] == 0)
 				{
-					SimStageData.AlternateIterationSource = InstanceData.FindIterationInterface(0);
+					continue;
 				}
-			}
-			ComputeContext->bHasTickedThisFrame_RT = true;
 
-			//-OPT: Do we need this test?  Can remove in favor of MaxUpdateIterations
-			if (Tick.NumInstancesWithSimStages > 0)
-			{
-				for (uint32 SimStageIndex=1; SimStageIndex < ComputeContext->MaxUpdateIterations; ++SimStageIndex)
+				FNiagaraDataInterfaceProxyRW* IterationInterface = InstanceData.FindIterationInterface(SimStageIndex);
+				for ( int32 IterationIndex=0; IterationIndex < InstanceData.NumIterationsPerStage[SimStageIndex]; ++IterationIndex )
 				{
-					// Determine if the iteration is outputting to a custom data size
-					FNiagaraDataInterfaceProxyRW* IterationInterface = InstanceData.FindIterationInterface(SimStageIndex);
-					if (!NiagaraEmitterInstanceBatcherLocal::ShouldRunStage(ComputeContext, IterationInterface, SimStageIndex, InstanceData.bResetData))
-					{
-						continue;
-					}
-
 					// Build SimStage data
 					FNiagaraGpuDispatchGroup& DispatchGroup = GpuDispatchList.DispatchGroups[iInstanceCurrDispatchGroup++];
 					FNiagaraGpuDispatchInstance& DispatchInstance = DispatchGroup.DispatchInstances.Emplace_GetRef(Tick, InstanceData);
 					FNiagaraSimStageData& SimStageData = DispatchInstance.SimStageData;
+					SimStageData.bFirstStage = bFirstStage;
 					SimStageData.StageIndex = SimStageIndex;
-					SimStageData.StageMetaData = ComputeContext->GetSimStageMetaData(SimStageIndex);
+					SimStageData.IterationIndex = IterationIndex;
+					SimStageData.StageMetaData = &SimStageMetaData;
 					SimStageData.AlternateIterationSource = IterationInterface;
 
-					checkf(SimStageData.StageMetaData, TEXT("No simulation stage meta data found for stage '%s' iteration interface '%s'"), SimStageIndex, *(IterationInterface ? IterationInterface->SourceDIName : FName()).ToString());
+					bFirstStage = false;
+
+					FNiagaraDataBuffer* SourceData = ComputeContext->bHasTickedThisFrame_RT ? ComputeContext->GetPrevDataBuffer() : ComputeContext->MainDataSet->GetCurrentData();
 
 					// This stage does not modify particle data, i.e. read only or not related to particles at all
 					if (!SimStageData.StageMetaData->bWritesParticles)
 					{
-						SimStageData.Source = ComputeContext->GetPrevDataBuffer();
+						SimStageData.Source = SourceData;
 						SimStageData.SourceCountOffset = ComputeContext->CountOffset_RT;
 						SimStageData.SourceNumInstances = ComputeContext->CurrentNumInstances_RT;
 						SimStageData.Destination = nullptr;
@@ -1067,27 +1020,38 @@ void NiagaraEmitterInstanceBatcher::PrepareTicksForProxy(FRHICommandListImmediat
 						SimStageData.Source = nullptr;
 						SimStageData.SourceCountOffset = ComputeContext->CountOffset_RT;
 						SimStageData.SourceNumInstances = ComputeContext->CurrentNumInstances_RT;
-						SimStageData.Destination = ComputeContext->GetPrevDataBuffer();
+						SimStageData.Destination = SourceData;
 						SimStageData.DestinationCountOffset = ComputeContext->CountOffset_RT;
 						SimStageData.DestinationNumInstances = ComputeContext->CurrentNumInstances_RT;
 					}
 					// This stage may kill particles, we need to allocate a new destination buffer
 					else
 					{
-						SimStageData.Source = ComputeContext->GetPrevDataBuffer();
+						SimStageData.Source = SourceData;
 						SimStageData.SourceCountOffset = ComputeContext->CountOffset_RT;
-						SimStageData.SourceNumInstances = ComputeContext->CurrentNumInstances_RT;
+						//-TODO: This is a little odd, perhaps we need to change the preallocate
+						SimStageData.SourceNumInstances = SimStageIndex == 0 && IterationIndex == 0 ? PrevNumInstances : ComputeContext->CurrentNumInstances_RT;
 						SimStageData.Destination = ComputeContext->GetNextDataBuffer();
 						SimStageData.DestinationCountOffset = GPUInstanceCounterManager.AcquireEntry();
 						SimStageData.DestinationNumInstances = ComputeContext->CurrentNumInstances_RT;
 
-						if (ensure(SimStageData.SourceCountOffset != INDEX_NONE))
+						ComputeContext->AdvanceDataBuffer();
+						ComputeContext->CountOffset_RT = SimStageData.DestinationCountOffset;
+						ComputeContext->bHasTickedThisFrame_RT = true;
+
+						// If we are the last tick then we may want to enqueue for a readback
+						// Note: Do not pull count from SimStageData as a reset tick will be INDEX_NONE
+						check(SimStageData.SourceCountOffset != INDEX_NONE || SimStageData.SourceNumInstances == 0);
+						if (SimStageData.SourceCountOffset != INDEX_NONE)
 						{
+							if ( bEnqueueCountReadback && Tick.bIsFinalTick && (ComputeContext->EmitterInstanceReadback.GPUCountOffset == INDEX_NONE))
+							{
+								bRequiresReadback = true;
+								ComputeContext->EmitterInstanceReadback.CPUCount = SimStageData.SourceNumInstances;
+								ComputeContext->EmitterInstanceReadback.GPUCountOffset = SimStageData.SourceCountOffset;
+							}
 							GpuDispatchList.CountsToRelease.Add(SimStageData.SourceCountOffset);
 						}
-
-						ComputeContext->AdvanceDataBuffer();
-						ComputeContext->CountOffset_RT = SimStageData.DestinationCountOffset;
 					}
 				}
 			}
@@ -1340,17 +1304,8 @@ void NiagaraEmitterInstanceBatcher::ExecuteTicks(FRHICommandList& RHICmdList, FR
 			}
 #if STATS
 			if (GPUProfiler.IsProfilingEnabled() && DispatchInstance.Tick.bIsFinalTick)
-			{
-				TStatId StatId;
-				if (const FSimulationStageMetaData* StageMetaData = DispatchInstance.SimStageData.StageMetaData)
-				{
-					StatId = FDynamicStats::CreateStatId<FStatGroup_STATGROUP_NiagaraDetailed>("GPU_Stage_" + StageMetaData->SimulationStageName.ToString());
-				}
-				else
-				{
-					static FString DefaultStageName = "GPU_Stage_SpawnUpdate";
-					StatId = FDynamicStats::CreateStatId<FStatGroup_STATGROUP_NiagaraDetailed>(DefaultStageName);
-				}
+			{			
+				const TStatId StatId = FDynamicStats::CreateStatId<FStatGroup_STATGROUP_NiagaraDetailed>("GPU_Stage_" + DispatchInstance.SimStageData.StageMetaData->SimulationStageName.ToString());
 				const int32 TimerHandle = GPUProfiler.StartTimer((uint64)DispatchInstance.InstanceData.Context, StatId, RHICmdList);
 				DispatchStage(RHICmdList, ViewUniformBuffer, DispatchInstance.Tick, DispatchInstance.InstanceData, DispatchInstance.SimStageData);
 				GPUProfiler.EndTimer(TimerHandle, RHICmdList);
@@ -1484,17 +1439,17 @@ void NiagaraEmitterInstanceBatcher::DispatchStage(FRHICommandList& RHICmdList, F
 	}
 
 	// Get Shader
-	const int32 PermutationId = Tick.NumInstancesWithSimStages > 0 ? InstanceData.Context->GPUScript_RT->ShaderStageIndexToPermutationId_RenderThread(SimStageData.StageIndex) : 0;
-	const FNiagaraShaderRef ComputeShader = InstanceData.Context->GPUScript_RT->GetShader(PermutationId);
+	const FNiagaraShaderRef ComputeShader = InstanceData.Context->GPUScript_RT->GetShader(SimStageData.StageIndex);
 	FRHIComputeShader* RHIComputeShader = ComputeShader.GetComputeShader();
 	SetComputePipelineState(RHICmdList, RHIComputeShader);
 
 	SCOPED_DRAW_EVENTF(RHICmdList, NiagaraGPUSimulationCS,
-		TEXT("NiagaraGpuSim(%s) DispatchCount(%u) Stage(%s %u) NumInstructions(%u)"),
+		TEXT("NiagaraGpuSim(%s) DispatchCount(%u) Stage(%s %u) Iteration(%d) NumInstructions(%u)"),
 		InstanceData.Context->GetDebugSimName(),
 		DispatchCount,
-		SimStageData.StageMetaData ? *SimStageData.StageMetaData->SimulationStageName.ToString() : TEXT("Particles"),
+		*SimStageData.StageMetaData->SimulationStageName.ToString(),
 		SimStageData.StageIndex,
+		SimStageData.IterationIndex,
 		ComputeShader->GetNumInstructions()
 	);
 	FNiagaraUAVPoolAccessScope UAVPoolAccessScope(*this);
@@ -1557,15 +1512,12 @@ void NiagaraEmitterInstanceBatcher::DispatchStage(FRHICommandList& RHICmdList, F
 			SimulationStageIterationInfo.Y = IterationInstanceCountOffset == INDEX_NONE ? DispatchCount : 0;
 		}
 
-		if (SimStageData.StageMetaData != nullptr)
-		{
-			const int32 NumStages = SimStageData.StageMetaData->MaxStage - SimStageData.StageMetaData->MinStage;
-			ensure((int32(SimStageData.StageIndex) >= SimStageData.StageMetaData->MinStage) && (int32(SimStageData.StageIndex) < SimStageData.StageMetaData->MaxStage));
-			const int32 IterationIndex = SimStageData.StageIndex - SimStageData.StageMetaData->MinStage;
-			SimulationStageIterationInfo.Z = SimStageData.StageIndex - SimStageData.StageMetaData->MinStage;
-			SimulationStageIterationInfo.W = NumStages;
-			SimulationStageNormalizedIterationIndex = NumStages > 1 ? float(IterationIndex) / float(NumStages - 1) : 1.0f;
-		}
+		const int32 NumIterations = InstanceData.NumIterationsPerStage[SimStageData.StageIndex];
+		const int32 IterationIndex = SimStageData.IterationIndex;
+		SimulationStageIterationInfo.Z = IterationIndex;
+		SimulationStageIterationInfo.W = NumIterations;
+		SimulationStageNormalizedIterationIndex = NumIterations > 1 ? float(IterationIndex) / float(NumIterations - 1) : 1.0f;
+
 		SetShaderValue(RHICmdList, RHIComputeShader, ComputeShader->SimulationStageIterationInfoParam, SimulationStageIterationInfo);
 		SetShaderValue(RHICmdList, RHIComputeShader, ComputeShader->SimulationStageNormalizedIterationIndexParam, SimulationStageNormalizedIterationIndex);
 	}
@@ -2313,4 +2265,4 @@ void NiagaraEmitterInstanceBatcher::BroadcastTemporalEffect(FRHICommandList& RHI
 		}
 	}
 }
-#endif // WITH_MGPU
+#endif // WITH_MGPU
\ No newline at end of file
diff --git a/Engine/Plugins/FX/Niagara/Source/Niagara/Private/NiagaraGPUSystemTick.cpp b/Engine/Plugins/FX/Niagara/Source/Niagara/Private/NiagaraGPUSystemTick.cpp
--- a/Engine/Plugins/FX/Niagara/Source/Niagara/Private/NiagaraGPUSystemTick.cpp
+++ b/Engine/Plugins/FX/Niagara/Source/Niagara/Private/NiagaraGPUSystemTick.cpp
@@ -3,6 +3,7 @@
 #include "NiagaraGPUSystemTick.h"
 #include "NiagaraSystemGpuComputeProxy.h"
 #include "NiagaraSystemInstance.h"
+#include "NiagaraSystemSimulation.h"
 
 void FNiagaraGPUSystemTick::Init(FNiagaraSystemInstance* InSystemInstance)
 {
@@ -70,10 +71,6 @@ void FNiagaraGPUSystemTick::Init(FNiagaraSystemInstance* InSystemInstance)
 	FNiagaraComputeInstanceData* Instances = (FNiagaraComputeInstanceData*)(InstanceData_ParamData_Packed);
 	uint8* ParamDataBufferPtr = InstanceData_ParamData_Packed + PackedDispatchesSizeAligned;
 
-	NumInstancesWithSimStages = 0;
-
-	TotalDispatches = 0;
-
 	// we want to include interpolation parameters (current and previous frame) if any of the emitters in the system
 	// require it
 	const bool IncludeInterpolationParameters = InSystemInstance->GPUParamIncludeInterpolation;
@@ -177,14 +174,6 @@ void FNiagaraGPUSystemTick::Init(FNiagaraSystemInstance* InSystemInstance)
 
 			InstanceData->bUsesSimStages = Emitter->bSimulationStagesEnabled; /* TODO limit to just with stages in the future! Leaving like this so what can convert! && EmitterRaw->GetSimulationStages().Num() > 0*/
 
-			if (InstanceData->bUsesSimStages)
-			{
-				NumInstancesWithSimStages++;
-			}
-
-			check(GPUContext->MaxUpdateIterations > 0);
-			TotalDispatches += FMath::Max<int32>(GPUContext->MaxUpdateIterations, 1);
-
 			// @todo-threadsafety Think of a better way to do this!
 			const TArray<UNiagaraDataInterface*>& DataInterfaces = GPUContext->CombinedParamStore.GetDataInterfaces();
 			InstanceData->DataInterfaceProxies.Reserve(DataInterfaces.Num());
@@ -201,6 +190,42 @@ void FNiagaraGPUSystemTick::Init(FNiagaraSystemInstance* InSystemInstance)
 					InstanceData->IterationDataInterfaceProxies.Add(RWProxy);
 				}
 			}
+
+			// Gather number of iterations for each stage, and if the stage should run or not
+			InstanceData->NumIterationsPerStage.Reserve(GPUContext->SimStageInfo.Num());
+			for ( FSimulationStageMetaData& SimStageMetaData : GPUContext->SimStageInfo )
+			{
+				int32 NumIterations = SimStageMetaData.NumIterations;
+				if (SimStageMetaData.ShouldRunStage(InstanceData->bResetData))
+				{
+					if (!SimStageMetaData.NumIterationsBinding.IsNone())
+					{
+						FNiagaraParameterStore& BoundParamStore = EmitterInstance->GetRendererBoundVariables();
+						if ( const uint8* ParameterData = BoundParamStore.GetParameterData(FNiagaraVariable(FNiagaraTypeDefinition::GetIntDef(), SimStageMetaData.NumIterationsBinding)) )
+						{
+							NumIterations = *reinterpret_cast<const int32*>(ParameterData);
+							NumIterations = FMath::Max(NumIterations, 0);
+						}
+					}
+					if (!SimStageMetaData.EnabledBinding.IsNone())
+					{
+						FNiagaraParameterStore& BoundParamStore = EmitterInstance->GetRendererBoundVariables();
+						if (const uint8* ParameterData = BoundParamStore.GetParameterData(FNiagaraVariable(FNiagaraTypeDefinition::GetBoolDef(), SimStageMetaData.EnabledBinding)))
+						{
+							const FNiagaraBool StageEnabled = *reinterpret_cast<const FNiagaraBool*>(ParameterData);
+							NumIterations = StageEnabled.GetValue() ? NumIterations : 0;
+						}
+					}
+				}
+				else
+				{
+					TotalDispatches = 0;
+				}
+
+				InstanceData->TotalDispatches += NumIterations;
+				InstanceData->NumIterationsPerStage.Add(NumIterations);
+			}
+			TotalDispatches += InstanceData->TotalDispatches;
 		}
 	}
 
diff --git a/Engine/Plugins/FX/Niagara/Source/Niagara/Private/NiagaraScript.cpp b/Engine/Plugins/FX/Niagara/Source/Niagara/Private/NiagaraScript.cpp
--- a/Engine/Plugins/FX/Niagara/Source/Niagara/Private/NiagaraScript.cpp
+++ b/Engine/Plugins/FX/Niagara/Source/Niagara/Private/NiagaraScript.cpp
@@ -891,6 +891,14 @@ void UNiagaraScript::ComputeVMCompilationId(FNiagaraVMExecutableDataId& Id, FGui
 				Id.AdditionalDefines.Add(GNiagaraForceSafeScriptAttributeTrim ? TEXT("TrimAttributesSafe") : TEXT("TrimAttributes"));
 
 				TArray<FString> PreserveAttributes;
+				auto AddAttributeToPreserve = 
+					[&](FName AttributeName)
+					{
+						if ( !AttributeName.IsNone() )
+						{
+							PreserveAttributes.AddUnique(FString::Printf(TEXT("PreserveAttribute=%s"), *AttributeName.ToString()));
+						}
+					};
 
 				// preserve the attributes that have been defined on the emitter directly
 				for (const FString& Attribute : Emitter->AttributesToPreserve)
@@ -904,8 +912,16 @@ void UNiagaraScript::ComputeVMCompilationId(FNiagaraVMExecutableDataId& Id, FGui
 				{
 					for (const FNiagaraVariable& BoundAttribute : RendererProperty->GetBoundAttributes())
 					{
-						const FString PreserveDefine = TEXT("PreserveAttribute=") + BoundAttribute.GetName().ToString();
-						PreserveAttributes.AddUnique(PreserveDefine);
+						AddAttributeToPreserve(BoundAttribute.GetName());
+					}
+				}
+
+				for (const UNiagaraSimulationStageBase* SimStageBase : Emitter->GetSimulationStages())
+				{
+					if (const UNiagaraSimulationStageGeneric* SimStageGeneric = Cast<const UNiagaraSimulationStageGeneric>(SimStageBase))
+					{
+						AddAttributeToPreserve(SimStageGeneric->EnabledBinding.GetParamMapBindableVariable().GetName());
+						AddAttributeToPreserve(SimStageGeneric->NumIterationsBinding.GetParamMapBindableVariable().GetName());
 					}
 				}
 
@@ -2939,28 +2955,32 @@ void UNiagaraScript::SyncAliases(const FNiagaraAliasContext& ResolveAliasesConte
 	}
 
 	// Sync up any simulation stage name references.
-	for (int32 i = 0; i < GetVMExecutableData().SimulationStageMetaData.Num(); i++)
+	for (FSimulationStageMetaData& SimStageMetaData : GetVMExecutableData().SimulationStageMetaData)
 	{
-		if (!GetVMExecutableData().SimulationStageMetaData[i].IterationSource.IsNone())
+		if ( !SimStageMetaData.IterationSource.IsNone() )
 		{
-			FNiagaraVariable Var(FNiagaraTypeDefinition(UNiagaraDataInterface::StaticClass()), GetVMExecutableData().SimulationStageMetaData[i].IterationSource);
-			FNiagaraVariable NewVar = FNiagaraUtilities::ResolveAliases(Var, ResolveAliasesContext);
-			if (NewVar.GetName() != Var.GetName())
-			{
-				GetVMExecutableData().SimulationStageMetaData[i].IterationSource = NewVar.GetName();
-			}
+			FNiagaraVariable Var(FNiagaraTypeDefinition(UNiagaraDataInterface::StaticClass()), SimStageMetaData.IterationSource);
+			SimStageMetaData.IterationSource = FNiagaraUtilities::ResolveAliases(Var, ResolveAliasesContext).GetName();
+		}
+
+		if (!SimStageMetaData.EnabledBinding.IsNone())
+		{
+			FNiagaraVariable Var(FNiagaraTypeDefinition::GetIntDef(), SimStageMetaData.EnabledBinding);
+			SimStageMetaData.EnabledBinding = FNiagaraUtilities::ResolveAliases(Var, ResolveAliasesContext).GetName();
 		}
 
-		for (int32 DestIdx = 0; DestIdx < GetVMExecutableData().SimulationStageMetaData[i].OutputDestinations.Num(); DestIdx++)
+		if (!SimStageMetaData.NumIterationsBinding.IsNone())
 		{
-			if (!GetVMExecutableData().SimulationStageMetaData[i].OutputDestinations[DestIdx].IsNone())
+			FNiagaraVariable Var(FNiagaraTypeDefinition::GetIntDef(), SimStageMetaData.NumIterationsBinding);
+			SimStageMetaData.NumIterationsBinding = FNiagaraUtilities::ResolveAliases(Var, ResolveAliasesContext).GetName();
+		}
+
+		for ( FName& OutputDestination : SimStageMetaData.OutputDestinations )
+		{
+			if ( !OutputDestination.IsNone() )
 			{
-				FNiagaraVariable Var(FNiagaraTypeDefinition(UNiagaraDataInterface::StaticClass()), GetVMExecutableData().SimulationStageMetaData[i].OutputDestinations[DestIdx]);
-				FNiagaraVariable NewVar = FNiagaraUtilities::ResolveAliases(Var, ResolveAliasesContext);
-				if (NewVar.GetName() != Var.GetName())
-				{
-					GetVMExecutableData().SimulationStageMetaData[i].OutputDestinations[DestIdx] = NewVar.GetName();
-				}
+				FNiagaraVariable Var(FNiagaraTypeDefinition(UNiagaraDataInterface::StaticClass()), OutputDestination);
+				OutputDestination = FNiagaraUtilities::ResolveAliases(Var, ResolveAliasesContext).GetName();
 			}
 		}
 	}
diff --git a/Engine/Plugins/FX/Niagara/Source/Niagara/Private/NiagaraSimulationStageBase.cpp b/Engine/Plugins/FX/Niagara/Source/Niagara/Private/NiagaraSimulationStageBase.cpp
--- a/Engine/Plugins/FX/Niagara/Source/Niagara/Private/NiagaraSimulationStageBase.cpp
+++ b/Engine/Plugins/FX/Niagara/Source/Niagara/Private/NiagaraSimulationStageBase.cpp
@@ -59,15 +59,37 @@ bool UNiagaraSimulationStageGeneric::AppendCompileHash(FNiagaraCompileHashVisito
 {
 	Super::AppendCompileHash(InVisitor);
 
+	InVisitor->UpdateString(TEXT("EnabledBinding"), EnabledBinding.GetDataSetBindableVariable().GetName().ToString());
 	InVisitor->UpdatePOD(TEXT("Iterations"), Iterations);
+	InVisitor->UpdateString(TEXT("NumIterationsBinding"), NumIterationsBinding.GetDataSetBindableVariable().GetName().ToString());
 	InVisitor->UpdatePOD(TEXT("IterationSource"), (int32)IterationSource);
-	InVisitor->UpdatePOD(TEXT("bSpawnOnly"), bSpawnOnly ? 1 : 0);
+	InVisitor->UpdatePOD(TEXT("ExecuteBehavior"), (int32)ExecuteBehavior);
 	InVisitor->UpdatePOD(TEXT("bDisablePartialParticleUpdate"), bDisablePartialParticleUpdate ? 1 : 0);
 	InVisitor->UpdateString(TEXT("DataInterface"), DataInterface.BoundVariable.GetName().ToString());
 	InVisitor->UpdateString(TEXT("SimulationStageName"), SimulationStageName.ToString());
 	return true;
 }
 
+void UNiagaraSimulationStageGeneric::PostInitProperties()
+{
+	Super::PostInitProperties();
+
+	if (HasAnyFlags(RF_ClassDefaultObject) == false)
+	{
+		EnabledBinding.Setup(
+			FNiagaraVariableBase(FNiagaraTypeDefinition::GetBoolDef(), NAME_None),
+			FNiagaraVariableBase(FNiagaraTypeDefinition::GetBoolDef(), NAME_None),
+			ENiagaraRendererSourceDataMode::Emitter
+		);
+
+		NumIterationsBinding.Setup(
+			FNiagaraVariableBase(FNiagaraTypeDefinition::GetIntDef(), NAME_None),
+			FNiagaraVariableBase(FNiagaraTypeDefinition::GetIntDef(), NAME_None),
+			ENiagaraRendererSourceDataMode::Emitter
+		);
+	}
+}
+
 #if WITH_EDITOR
 void UNiagaraSimulationStageGeneric::PostEditChangeProperty(struct FPropertyChangedEvent& PropertyChangedEvent) 
 {
@@ -80,7 +102,15 @@ void UNiagaraSimulationStageGeneric::PostEditChangeProperty(struct FPropertyChan
 	}
 
 	bool bNeedsRecompile = false;
-	if (PropertyName == GET_MEMBER_NAME_CHECKED(UNiagaraSimulationStageGeneric, Iterations))
+	if (PropertyName == GET_MEMBER_NAME_CHECKED(UNiagaraSimulationStageGeneric, EnabledBinding))
+	{
+		bNeedsRecompile = true;
+	}
+	else if (PropertyName == GET_MEMBER_NAME_CHECKED(UNiagaraSimulationStageGeneric, Iterations))
+	{
+		bNeedsRecompile = true;
+	}
+	else if (PropertyName == GET_MEMBER_NAME_CHECKED(UNiagaraSimulationStageGeneric, NumIterationsBinding))
 	{
 		bNeedsRecompile = true;
 	}
@@ -88,7 +118,7 @@ void UNiagaraSimulationStageGeneric::PostEditChangeProperty(struct FPropertyChan
 	{
 		bNeedsRecompile = true;
 	}
-	else if (PropertyName == GET_MEMBER_NAME_CHECKED(UNiagaraSimulationStageGeneric, bSpawnOnly))
+	else if (PropertyName == GET_MEMBER_NAME_CHECKED(UNiagaraSimulationStageGeneric, ExecuteBehavior))
 	{
 		bNeedsRecompile = true;
 	}
diff --git a/Engine/Plugins/FX/Niagara/Source/Niagara/Public/NiagaraCommon.h b/Engine/Plugins/FX/Niagara/Source/Niagara/Public/NiagaraCommon.h
--- a/Engine/Plugins/FX/Niagara/Source/Niagara/Public/NiagaraCommon.h
+++ b/Engine/Plugins/FX/Niagara/Source/Niagara/Public/NiagaraCommon.h
@@ -383,10 +383,7 @@ struct NIAGARA_API FNiagaraFunctionSignature
 
 	/** When using simulation stages and bRequiresContext is true this will be the index of the stage that is associated with the function. */
 	UPROPERTY()
-	int32 ContextStageMinIndex;
-
-	UPROPERTY()
-	int32 ContextStageMaxIndex;
+	int32 ContextStageIndex;
 
 	/** Function specifiers verified at bind time. */
 	UPROPERTY()
@@ -416,8 +413,7 @@ struct NIAGARA_API FNiagaraFunctionSignature
 		, bIsCompileTagGenerator(false)
 		, bHidden(false)
 		, ModuleUsageBitmask(0)
-		, ContextStageMinIndex(INDEX_NONE)
-		, ContextStageMaxIndex(INDEX_NONE)
+		, ContextStageIndex(INDEX_NONE)
 	{
 	}
 
@@ -434,8 +430,7 @@ struct NIAGARA_API FNiagaraFunctionSignature
 		, bIsCompileTagGenerator(false)
 		, bHidden(false)
 		, ModuleUsageBitmask(0)
-		, ContextStageMinIndex(INDEX_NONE)
-		, ContextStageMaxIndex(INDEX_NONE)
+		, ContextStageIndex(INDEX_NONE)
 	{
 		Inputs.Reserve(InInputs.Num());
 		for (FNiagaraVariable& Var : InInputs)
@@ -479,8 +474,7 @@ struct NIAGARA_API FNiagaraFunctionSignature
 		bMatches &= bRequiresExecPin == Other.bRequiresExecPin;
 		bMatches &= bMemberFunction == Other.bMemberFunction;
 		bMatches &= OwnerName == Other.OwnerName;
-		bMatches &= ContextStageMinIndex == Other.ContextStageMinIndex;
-		bMatches &= ContextStageMaxIndex == Other.ContextStageMaxIndex;
+		bMatches &= ContextStageIndex == Other.ContextStageIndex;
 		bMatches &= bIsCompileTagGenerator == Other.bIsCompileTagGenerator;
 		return bMatches;
 	}
@@ -966,6 +960,7 @@ struct FNiagaraVariableAttributeBinding
 	NIAGARA_API const FNiagaraVariableBase& GetParamMapBindableVariable() const { return ParamMapVariable; }
 	NIAGARA_API const FNiagaraVariableBase& GetDataSetBindableVariable() const { return DataSetVariable; }
 	NIAGARA_API const FNiagaraTypeDefinition& GetType() const { return DataSetVariable.GetType(); }
+	NIAGARA_API ENiagaraBindingSource GetBindingSourceMode() const { return BindingSourceMode; }
 
 	NIAGARA_API bool IsValid() const {	return DataSetVariable.IsValid();}
 	
diff --git a/Engine/Plugins/FX/Niagara/Source/Niagara/Public/NiagaraSimulationStageBase.h b/Engine/Plugins/FX/Niagara/Source/Niagara/Public/NiagaraSimulationStageBase.h
--- a/Engine/Plugins/FX/Niagara/Source/Niagara/Public/NiagaraSimulationStageBase.h
+++ b/Engine/Plugins/FX/Niagara/Source/Niagara/Public/NiagaraSimulationStageBase.h
@@ -3,6 +3,7 @@
 #pragma once
 
 #include "NiagaraMergeable.h"
+#include "NiagaraScriptBase.h"
 #include "NiagaraCommon.h"
 #include "NiagaraTypes.h"
 #include "NiagaraConstants.h"
@@ -54,6 +55,10 @@ public:
 	{
 	}
 
+	/** Binding to a bool parameter which dynamically controls if the simulation stage is enabled or not. */
+	UPROPERTY(EditAnywhere, Category = "Simulation Stage")
+	FNiagaraVariableAttributeBinding EnabledBinding;
+
 	/** Determine which elements this script is iterating over. You are not allowed to */
 	UPROPERTY(EditAnywhere, Category = "Simulation Stage")
 	ENiagaraIterationSource IterationSource;
@@ -61,8 +66,15 @@ public:
 	UPROPERTY(EditAnywhere, Category = "Simulation Stage", meta = (NoSpinbox = "true", ClampMin = 1, Tooltip = "The number of times we run this simulation stage before moving to the next stage."))
 	int32 Iterations;
 
-	UPROPERTY(EditAnywhere, Category = "Simulation Stage", meta = (DisplayName = "Emitter Reset Only", Tooltip = "When enabled the stage will only run on the first tick after the emitter is reset, only valid for data interface iteration stages", EditCondition = "IterationSource == ENiagaraIterationSource::DataInterface"))
-	uint32 bSpawnOnly : 1;
+	/** Binding to an int parameter which dynamically controls the number of times the simulation stage runs. */
+	UPROPERTY(EditAnywhere, Category = "Simulation Stage")
+	FNiagaraVariableAttributeBinding NumIterationsBinding;
+
+	UPROPERTY()
+	uint32 bSpawnOnly_DEPRECATED : 1;
+
+	UPROPERTY(EditAnywhere, Category = "Simulation Stage", meta = (Tooltip = "Controls when the simulation stage should execute, only valid for data interface iteration stages", EditCondition = "IterationSource == ENiagaraIterationSource::DataInterface"))
+	ENiagaraSimStageExecuteBehavior ExecuteBehavior = ENiagaraSimStageExecuteBehavior::Always;
 
 	UPROPERTY(EditAnywhere, AdvancedDisplay, Category = "Simulation Stage", meta = (Tooltip = "Disables the ability to read / write from the same particle buffer, i.e. only update position and no other attributes.  By default this should not be changed and is a debugging tool.", EditCondition = "IterationSource == ENiagaraIterationSource::Particles"))
 	uint32 bDisablePartialParticleUpdate : 1;
@@ -71,6 +83,7 @@ public:
 	UPROPERTY(EditAnywhere, Category = "Simulation Stage", meta = (editcondition = "IterationSource == ENiagaraIterationSource::DataInterface"))
 	FNiagaraVariableDataInterfaceBinding DataInterface;
 
+	virtual void PostInitProperties() override;
 	virtual bool AppendCompileHash(FNiagaraCompileHashVisitor* InVisitor) const override;
 #if WITH_EDITOR
 	virtual FName GetStackContextReplacementName() const override; 
diff --git a/Engine/Plugins/FX/Niagara/Source/NiagaraCore/Private/NiagaraCustomVersion.cpp b/Engine/Plugins/FX/Niagara/Source/NiagaraCore/Private/NiagaraCustomVersion.cpp
--- a/Engine/Plugins/FX/Niagara/Source/NiagaraCore/Private/NiagaraCustomVersion.cpp
+++ b/Engine/Plugins/FX/Niagara/Source/NiagaraCore/Private/NiagaraCustomVersion.cpp
@@ -8,4 +8,4 @@ const FGuid FNiagaraCustomVersion::GUID(0xFCF57AFA, 0x50764283, 0xB9A9E658, 0xFF
 // Register the custom version with core
 FCustomVersionRegistration GRegisterNiagaraCustomVersion(FNiagaraCustomVersion::GUID, FNiagaraCustomVersion::LatestVersion, TEXT("NiagaraVer"));
 
-const FGuid FNiagaraCustomVersion::LatestScriptCompileVersion(0x323AF294, 0x45834FA5, 0x835EB6A5, 0x80C83B36);
\ No newline at end of file
+const FGuid FNiagaraCustomVersion::LatestScriptCompileVersion(0x02497746, 0x33D75C49, 0x9A60CB65, 0x7D6AC2A2);
diff --git a/Engine/Plugins/FX/Niagara/Source/NiagaraEditor/Private/Customizations/NiagaraTypeCustomizations.cpp b/Engine/Plugins/FX/Niagara/Source/NiagaraEditor/Private/Customizations/NiagaraTypeCustomizations.cpp
--- a/Engine/Plugins/FX/Niagara/Source/NiagaraEditor/Private/Customizations/NiagaraTypeCustomizations.cpp
+++ b/Engine/Plugins/FX/Niagara/Source/NiagaraEditor/Private/Customizations/NiagaraTypeCustomizations.cpp
@@ -256,6 +256,10 @@ TSharedRef<SWidget> FNiagaraVariableAttributeBindingCustomization::OnGetMenuCont
 TArray<FName> FNiagaraVariableAttributeBindingCustomization::GetNames(UNiagaraEmitter* InEmitter) const
 {
 	TArray<FName> Names;
+	if (!PropertyHandle.IsValid() || !PropertyHandle->GetProperty() || !TargetVariableBinding)
+	{
+		return Names;
+	}
 
 	bool bSystem = true;
 	bool bEmitter = true;
@@ -264,10 +268,17 @@ TArray<FName> FNiagaraVariableAttributeBindingCustomization::GetNames(UNiagaraEm
 	TArray<FNiagaraVariableBase> Vars = FNiagaraStackAssetAction_VarBind::FindVariables(InEmitter, bSystem, bEmitter, bParticles, bUser);
 	for (const FNiagaraVariableBase& Var : Vars)
 	{
-		if (RenderProps && PropertyHandle.IsValid() && PropertyHandle->GetProperty() && RenderProps->IsSupportedVariableForBinding(Var, *PropertyHandle->GetProperty()->GetName()))
+		if ( Var.GetType() != TargetVariableBinding->GetType() )
 		{
-			if (Var.GetType() == TargetVariableBinding->GetType())
-				Names.AddUnique(Var.GetName());
+			continue;
+		}
+		if ( RenderProps && RenderProps->IsSupportedVariableForBinding(Var, *PropertyHandle->GetProperty()->GetName()) )
+		{
+			Names.AddUnique(Var.GetName());
+		}
+		else if ( SimulationStage && !Var.IsInNameSpace(FNiagaraConstants::ParticleAttributeNamespace) )
+		{
+			Names.AddUnique(Var.GetName());
 		}
 	}
 
@@ -333,10 +344,11 @@ void FNiagaraVariableAttributeBindingCustomization::ChangeSource(FName InVarName
 		Obj->Modify();
 	}
 	check(BaseEmitter);
-	check(RenderProps);
+	check(RenderProps || SimulationStage);
 
 	PropertyHandle->NotifyPreChange();
-	TargetVariableBinding->SetValue(InVarName, BaseEmitter, RenderProps->GetCurrentSourceMode());
+	const ENiagaraRendererSourceDataMode BindingSourceMode = RenderProps ? RenderProps->GetCurrentSourceMode() : ENiagaraRendererSourceDataMode::Emitter;
+	TargetVariableBinding->SetValue(InVarName, BaseEmitter, BindingSourceMode);
 	PropertyHandle->NotifyPostChange(EPropertyChangeType::ValueSet);
 	PropertyHandle->NotifyFinishedChangingProperties();
 }
@@ -349,12 +361,12 @@ void FNiagaraVariableAttributeBindingCustomization::ResetToDefault()
 EVisibility FNiagaraVariableAttributeBindingCustomization::IsResetToDefaultsVisible() const
 {
 	check(BaseEmitter);
-	check(RenderProps);
+	check(RenderProps || SimulationStage);
 	check(TargetVariableBinding);
 	check(DefaultVariableBinding);
-	return (!TargetVariableBinding->MatchesDefault(*DefaultVariableBinding, RenderProps->GetCurrentSourceMode()))
-		? EVisibility::Visible
-		: EVisibility::Hidden;
+
+	const ENiagaraRendererSourceDataMode BindingSourceMode = RenderProps ? RenderProps->GetCurrentSourceMode() : ENiagaraRendererSourceDataMode::Emitter;
+	return TargetVariableBinding->MatchesDefault(*DefaultVariableBinding, BindingSourceMode) ? EVisibility::Hidden : EVisibility::Visible;
 }
 
 FReply FNiagaraVariableAttributeBindingCustomization::OnResetToDefaultsClicked()
@@ -367,12 +379,13 @@ FReply FNiagaraVariableAttributeBindingCustomization::OnResetToDefaultsClicked()
 		Obj->Modify();
 	}
 	check(BaseEmitter);
-	check(RenderProps);
+	check(RenderProps || SimulationStage);
 	check(TargetVariableBinding);
 	check(DefaultVariableBinding);
 
 	PropertyHandle->NotifyPreChange();
-	TargetVariableBinding->ResetToDefault(*DefaultVariableBinding, BaseEmitter, RenderProps->GetCurrentSourceMode());
+	const ENiagaraRendererSourceDataMode BindingSourceMode = RenderProps ? RenderProps->GetCurrentSourceMode() : ENiagaraRendererSourceDataMode::Emitter;
+	TargetVariableBinding->ResetToDefault(*DefaultVariableBinding, BaseEmitter, BindingSourceMode);
 	PropertyHandle->NotifyPostChange(EPropertyChangeType::ValueSet);
 	PropertyHandle->NotifyFinishedChangingProperties();
 	return FReply::Handled();
@@ -381,6 +394,7 @@ FReply FNiagaraVariableAttributeBindingCustomization::OnResetToDefaultsClicked()
 void FNiagaraVariableAttributeBindingCustomization::CustomizeHeader(TSharedRef<IPropertyHandle> InPropertyHandle, FDetailWidgetRow& HeaderRow, IPropertyTypeCustomizationUtils& CustomizationUtils)
 {
 	RenderProps = nullptr;
+	SimulationStage = nullptr;
 	BaseEmitter = nullptr;
 	PropertyHandle = InPropertyHandle;
 	TArray<UObject*> Objects;
@@ -404,66 +418,71 @@ void FNiagaraVariableAttributeBindingCustomization::CustomizeHeader(TSharedRef<I
 		if (RenderProps)
 		{
 			BaseEmitter = Cast<UNiagaraEmitter>(RenderProps->GetOuter());
+		}
 
-			if (BaseEmitter)
-			{
-				TargetVariableBinding = (FNiagaraVariableAttributeBinding*)PropertyHandle->GetValueBaseAddress((uint8*)Objects[0]);
-				DefaultVariableBinding = (FNiagaraVariableAttributeBinding*)PropertyHandle->GetValueBaseAddress((uint8*)Objects[0]->GetClass()->GetDefaultObject());
+		SimulationStage = Cast<UNiagaraSimulationStageBase>(Objects[0]);
+		if ( SimulationStage )
+		{
+			BaseEmitter = SimulationStage->GetTypedOuter<UNiagaraEmitter>();
+		}
+
+		if (BaseEmitter)
+		{
+			TargetVariableBinding = (FNiagaraVariableAttributeBinding*)PropertyHandle->GetValueBaseAddress((uint8*)Objects[0]);
+			DefaultVariableBinding = (FNiagaraVariableAttributeBinding*)PropertyHandle->GetValueBaseAddress((uint8*)Objects[0]->GetClass()->GetDefaultObject());
 				
-				HeaderRow
-					.NameContent()
-					[
-						PropertyHandle->CreatePropertyNameWidget()
-					]
-					.ValueContent()
+			HeaderRow
+				.NameContent()
+				[
+					PropertyHandle->CreatePropertyNameWidget()
+				]
+				.ValueContent()
+				[
+					SNew(SHorizontalBox)
+					+ SHorizontalBox::Slot()
+					.AutoWidth()
+					.Padding(0.0f, 0.0f, 4.0f, 0.0f)
 					[
-						SNew(SHorizontalBox)
-						+ SHorizontalBox::Slot()
-						.AutoWidth()
-						.Padding(0.0f, 0.0f, 4.0f, 0.0f)
+						SAssignNew(ComboButton, SComboButton)					
+						.OnGetMenuContent(this, &FNiagaraVariableAttributeBindingCustomization::OnGetMenuContent)
+						.ContentPadding(1)
+						.ToolTipText(this, &FNiagaraVariableAttributeBindingCustomization::GetTooltipText)
+						.ButtonContent()
 						[
-							SAssignNew(ComboButton, SComboButton)					
-							.OnGetMenuContent(this, &FNiagaraVariableAttributeBindingCustomization::OnGetMenuContent)
-							.ContentPadding(1)
-							.ToolTipText(this, &FNiagaraVariableAttributeBindingCustomization::GetTooltipText)
-							.ButtonContent()
-							[
-								/*SNew(STextBlock)
-								.Text(this, &FNiagaraVariableAttributeBindingCustomization::GetCurrentText)
-								.Font(IDetailLayoutBuilder::GetDetailFont())*/
-								SNew(SNiagaraParameterName)
-								.ParameterName(this, &FNiagaraVariableAttributeBindingCustomization::GetVariableName)
-								.IsReadOnly(true)
-								//SNew(STextBlock)
-								//.Text(InCreateData->Action->GetMenuDescription())
-							]
+							/*SNew(STextBlock)
+							.Text(this, &FNiagaraVariableAttributeBindingCustomization::GetCurrentText)
+							.Font(IDetailLayoutBuilder::GetDetailFont())*/
+							SNew(SNiagaraParameterName)
+							.ParameterName(this, &FNiagaraVariableAttributeBindingCustomization::GetVariableName)
+							.IsReadOnly(true)
+							//SNew(STextBlock)
+							//.Text(InCreateData->Action->GetMenuDescription())
 						]
-						+ SHorizontalBox::Slot()
-						.AutoWidth()
-						.VAlign(VAlign_Center)
-						.Padding(2, 1)
+					]
+					+ SHorizontalBox::Slot()
+					.AutoWidth()
+					.VAlign(VAlign_Center)
+					.Padding(2, 1)
+					[
+						SNew(SButton)
+						.IsFocusable(false)
+						.ToolTipText(LOCTEXT("ResetToDefaultToolTip", "Reset to Default"))
+						.ButtonStyle(FEditorStyle::Get(), "NoBorder")
+						.ContentPadding(0)
+						.Visibility(this, &FNiagaraVariableAttributeBindingCustomization::IsResetToDefaultsVisible)
+						.OnClicked(this, &FNiagaraVariableAttributeBindingCustomization::OnResetToDefaultsClicked)
+						.Content()
 						[
-							SNew(SButton)
-							.IsFocusable(false)
-							.ToolTipText(LOCTEXT("ResetToDefaultToolTip", "Reset to Default"))
-							.ButtonStyle(FEditorStyle::Get(), "NoBorder")
-							.ContentPadding(0)
-							.Visibility(this, &FNiagaraVariableAttributeBindingCustomization::IsResetToDefaultsVisible)
-							.OnClicked(this, &FNiagaraVariableAttributeBindingCustomization::OnResetToDefaultsClicked)
-							.Content()
-							[
-								SNew(SImage)
-								.Image(FEditorStyle::GetBrush("PropertyWindow.DiffersFromDefault"))
-								.ColorAndOpacity(FSlateColor::UseForeground())
-							]
+							SNew(SImage)
+							.Image(FEditorStyle::GetBrush("PropertyWindow.DiffersFromDefault"))
+							.ColorAndOpacity(FSlateColor::UseForeground())
 						]
-					];
-				bAddDefault = false;
-			}
+					]
+				];
+			bAddDefault = false;
 		}
 	}
 	
-
 	if (bAddDefault)
 	{
 		HeaderRow
@@ -471,7 +490,7 @@ void FNiagaraVariableAttributeBindingCustomization::CustomizeHeader(TSharedRef<I
 			[
 				PropertyHandle->CreatePropertyNameWidget()
 			]
-		.ValueContent()
+			.ValueContent()
 			.MaxDesiredWidth(200.f)
 			[
 				SNew(STextBlock)
diff --git a/Engine/Plugins/FX/Niagara/Source/NiagaraEditor/Private/Customizations/NiagaraTypeCustomizations.h b/Engine/Plugins/FX/Niagara/Source/NiagaraEditor/Private/Customizations/NiagaraTypeCustomizations.h
--- a/Engine/Plugins/FX/Niagara/Source/NiagaraEditor/Private/Customizations/NiagaraTypeCustomizations.h
+++ b/Engine/Plugins/FX/Niagara/Source/NiagaraEditor/Private/Customizations/NiagaraTypeCustomizations.h
@@ -184,10 +184,11 @@ private:
 
 	TSharedPtr<IPropertyHandle> PropertyHandle;
 	TSharedPtr<SComboButton> ComboButton;
-	class UNiagaraEmitter* BaseEmitter;
-	class UNiagaraRendererProperties* RenderProps;
-	struct FNiagaraVariableAttributeBinding* TargetVariableBinding;
-	const struct FNiagaraVariableAttributeBinding* DefaultVariableBinding;
+	class UNiagaraEmitter* BaseEmitter = nullptr;
+	class UNiagaraRendererProperties* RenderProps = nullptr;
+	class UNiagaraSimulationStageBase* SimulationStage = nullptr;
+	struct FNiagaraVariableAttributeBinding* TargetVariableBinding = nullptr;
+	const struct FNiagaraVariableAttributeBinding* DefaultVariableBinding = nullptr;
 };
 
 class FNiagaraUserParameterBindingCustomization : public IPropertyTypeCustomization
diff --git a/Engine/Plugins/FX/Niagara/Source/NiagaraEditor/Private/NiagaraCompiler.cpp b/Engine/Plugins/FX/Niagara/Source/NiagaraEditor/Private/NiagaraCompiler.cpp
--- a/Engine/Plugins/FX/Niagara/Source/NiagaraEditor/Private/NiagaraCompiler.cpp
+++ b/Engine/Plugins/FX/Niagara/Source/NiagaraEditor/Private/NiagaraCompiler.cpp
@@ -617,14 +617,9 @@ void FNiagaraCompileRequestData::FinishPrecompile(const TArray<FNiagaraVariable>
 
 		if (SimStages && NumSimStageNodes)
 		{
-			SpawnOnlyPerStage.Reserve(NumSimStageNodes);
-			PartialParticleUpdatePerStage.Reserve(NumSimStageNodes);
-			NumIterationsPerStage.Reserve(NumSimStageNodes);
-			IterationSourcePerStage.Reserve(NumSimStageNodes);
-			StageGuids.Reserve(NumSimStageNodes);
-			StageNames.Reserve(NumSimStageNodes);
-			const int32 NumProvidedStages = SimStages->Num();
+			CompileSimStageData.Reserve(NumSimStageNodes);
 
+			const int32 NumProvidedStages = SimStages->Num();
 			for (int32 i=0, ActiveStageCount = 0; ActiveStageCount < NumSimStageNodes && i < NumProvidedStages; ++i)
 			{
 				UNiagaraSimulationStageBase* SimStage = (*SimStages)[i];
@@ -635,12 +630,15 @@ void FNiagaraCompileRequestData::FinishPrecompile(const TArray<FNiagaraVariable>
 
 				if ( UNiagaraSimulationStageGeneric* GenericStage = Cast<UNiagaraSimulationStageGeneric>(SimStage) )
 				{
-					NumIterationsPerStage.Add(GenericStage->Iterations);
-					IterationSourcePerStage.Add(GenericStage->IterationSource == ENiagaraIterationSource::DataInterface ? GenericStage->DataInterface.BoundVariable.GetName() : FName());
-					SpawnOnlyPerStage.Add(GenericStage->bSpawnOnly);
-					PartialParticleUpdatePerStage.Add(GenericStage->bDisablePartialParticleUpdate == false);
-					StageGuids.Add(GenericStage->Script->GetUsageId());
-					StageNames.Add(GenericStage->SimulationStageName);
+					FCompileSimStageData& SimStageData	= CompileSimStageData.AddDefaulted_GetRef();
+					SimStageData.StageGuid				= GenericStage->Script->GetUsageId();
+					SimStageData.StageName				= GenericStage->SimulationStageName;
+					SimStageData.EnabledBinding			= GenericStage->EnabledBinding.GetName();
+					SimStageData.NumIterations			= GenericStage->Iterations;
+					SimStageData.NumIterationsBinding	= GenericStage->NumIterationsBinding.GetName();
+					SimStageData.IterationSource		= GenericStage->IterationSource == ENiagaraIterationSource::DataInterface ? GenericStage->DataInterface.BoundVariable.GetName() : FName();
+					SimStageData.ExecuteBehavior		= GenericStage->ExecuteBehavior;
+					SimStageData.PartialParticleUpdate	= GenericStage->bDisablePartialParticleUpdate == false;
 
 					++ActiveStageCount;
 				}
diff --git a/Engine/Plugins/FX/Niagara/Source/NiagaraEditor/Private/NiagaraHlslTranslator.cpp b/Engine/Plugins/FX/Niagara/Source/NiagaraEditor/Private/NiagaraHlslTranslator.cpp
--- a/Engine/Plugins/FX/Niagara/Source/NiagaraEditor/Private/NiagaraHlslTranslator.cpp
+++ b/Engine/Plugins/FX/Niagara/Source/NiagaraEditor/Private/NiagaraHlslTranslator.cpp
@@ -255,17 +255,14 @@ private:
 		FString ConditionString;
 
 		const int32 StageIndexCount = StageIndices.Num();
-		for (int32 i = 0; i < StageIndexCount; ++i)
+		for (int32 i=0; i < StageIndexCount; ++i)
 		{
 			const int32 StageIndex = StageIndices[i];
-
 			if (i)
 			{
 				ConditionString.Append(TEXT(" || "));
 			}
-			ConditionString.Appendf(TEXT("((SimulationStageIndex >= %d) && (SimulationStageIndex < %d))"),
-				TranslationStages[StageIndex].SimulationStageIndexMin,
-				TranslationStages[StageIndex].SimulationStageIndexMax);
+			ConditionString.Appendf(TEXT("(SimulationStageIndex == %d)"), TranslationStages[StageIndex].SimulationStageIndex);
 		}
 
 		return ConditionString;
@@ -567,8 +564,7 @@ void FHlslNiagaraTranslator::GenerateFunctionSignature(ENiagaraScriptUsage Scrip
 	// with so that if we encounter a function implementation for another stage that it will also be added
 	if (OutSig.bRequiresContext && FDeclarationPermutationContext::SupportsBranching(*this))
 	{
-		OutSig.ContextStageMinIndex = TranslationStages[ActiveStageIdx].SimulationStageIndexMin;
-		OutSig.ContextStageMaxIndex = TranslationStages[ActiveStageIdx].SimulationStageIndexMax;
+		OutSig.ContextStageIndex = TranslationStages[ActiveStageIdx].SimulationStageIndex;
 	}
 }
 
@@ -1138,10 +1134,10 @@ const FNiagaraTranslateResults &FHlslNiagaraTranslator::Translate(const FNiagara
 		TranslationStages[0].OutputNode = SourceGraph->FindEquivalentOutputNode(ENiagaraScriptUsage::ParticleSpawnScript, TranslationStages[0].UsageId);
 		TranslationStages[1].OutputNode = SourceGraph->FindEquivalentOutputNode(TranslationStages[1].ScriptUsage, TranslationStages[1].UsageId);
 		TranslationStages[1].bInterpolatePreviousParams = true;
-		TranslationStages[0].SimulationStageIndexMin = 0;
-		TranslationStages[0].SimulationStageIndexMax = 1;
-		TranslationStages[1].SimulationStageIndexMin = 0;
-		TranslationStages[1].SimulationStageIndexMax = 1;
+		TranslationStages[0].SimulationStageIndex = 0;
+		TranslationStages[0].NumIterations = 1;
+		TranslationStages[1].SimulationStageIndex = 0;
+		TranslationStages[1].NumIterations = 1;
 		TranslationStages[0].bWritesParticles = true;
 		TranslationStages[1].bWritesParticles = true;
 		ParamMapHistories.AddDefaulted(2);
@@ -1158,26 +1154,35 @@ const FNiagaraTranslateResults &FHlslNiagaraTranslator::Translate(const FNiagara
 		TranslationStages[0].OutputNode = SourceGraph->FindEquivalentOutputNode(ENiagaraScriptUsage::ParticleSpawnScript, TranslationStages[0].UsageId);
 		TranslationStages[1].OutputNode = SourceGraph->FindEquivalentOutputNode(TranslationStages[1].ScriptUsage, TranslationStages[1].UsageId);
 		TranslationStages[1].bInterpolatePreviousParams = bHasInterpolatedSpawn;
-		TranslationStages[0].SimulationStageIndexMin = 0;
-		TranslationStages[0].SimulationStageIndexMax = 1;
-		TranslationStages[1].SimulationStageIndexMin = 0;
-		TranslationStages[1].SimulationStageIndexMax = 1;
+		TranslationStages[0].SimulationStageIndex = 0;
+		TranslationStages[0].NumIterations = 1;
+		TranslationStages[1].SimulationStageIndex = 0;
+		TranslationStages[1].NumIterations = 1;
 		TranslationStages[0].bWritesParticles = true;
 		TranslationStages[1].bWritesParticles = true;
 		ParamMapHistories.AddDefaulted(2);
 		ParamMapHistoriesSourceInOtherHistories.AddDefaulted(2);
 		ParamMapSetVariablesToChunks.AddDefaulted(2);
 
+		// Add the spawn / update stage
 		{
-			int32 SimStageStartIndex = 1;
-			int32 SimStageIndex = 0;
+			FSimulationStageMetaData& SimulationStageMetaData = CompilationOutput.ScriptData.SimulationStageMetaData.AddDefaulted_GetRef();
+			SimulationStageMetaData.SimulationStageName = FName("ParticleSpawnUpdate");
+			SimulationStageMetaData.NumIterations = 1;
+			SimulationStageMetaData.bWritesParticles = true;
+			SimulationStageMetaData.bPartialParticleUpdate = false;
+		}
+
+		{
+			int32 SourceSimStageIndex = 0;
 			
 			// OutputNode order in traversal doesn't necessarily match the stack ordering. Use the GUID order to define the actual stages.
 			TArray<const UNiagaraNodeOutput*> FoundOutputNodes;
 			TArray<int32> FoundStageHistories;
 
-			for (const FGuid& StageGuid : InCompileData->StageGuids)
+			for (const auto& CompileSimStageData : InCompileData->CompileSimStageData )
 			{
+				const FGuid& StageGuid = CompileSimStageData.StageGuid;
 				for (int32 FoundHistoryIdx = 0; FoundHistoryIdx < OtherOutputParamMapHistories.Num(); FoundHistoryIdx++)
 				{
 					FNiagaraParameterMapHistory& FoundHistory = OtherOutputParamMapHistories[FoundHistoryIdx];
@@ -1199,14 +1204,14 @@ const FNiagaraTranslateResults &FHlslNiagaraTranslator::Translate(const FNiagara
 
 				if (HistoryOutputNode && HistoryOutputNode->ScriptType == ENiagaraScriptUsage::ParticleSimulationStageScript)
 				{
-					bool bSpawnOnly = InCompileData->NumIterationsPerStage.Num() > SimStageIndex ? InCompileData->SpawnOnlyPerStage[SimStageIndex] : false;
-					int32 NumIterationsThisStage = InCompileData->NumIterationsPerStage.Num() > SimStageIndex ? InCompileData->NumIterationsPerStage[SimStageIndex] : 1;
-					FName IterationSrc = InCompileData->IterationSourcePerStage.Num() > SimStageIndex ? InCompileData->IterationSourcePerStage[SimStageIndex] : FName();
-					FString StageName = InCompileData->StageNames.Num() > SimStageIndex ? InCompileData->StageNames[SimStageIndex].ToString() : TEXT("Unnamed");
+					const auto& CompileSimStageData = InCompileData->CompileSimStageData[SourceSimStageIndex];
+
+					FString StageName = CompileSimStageData.StageName.ToString();
 					StageName = TEXT("_") + GetSanitizedFunctionNameSuffix(StageName);
-					int32 Index = TranslationStages.Add(FHlslNiagaraTranslationStage(HistoryOutputNode->ScriptType, HistoryOutputNode->ScriptTypeId));
+					const int32 Index = TranslationStages.Add(FHlslNiagaraTranslationStage(HistoryOutputNode->ScriptType, HistoryOutputNode->ScriptTypeId));
 
-					TranslationStages[Index].PassNamespace = FString::Printf(TEXT("MapSimStage%d%s"), SimStageIndex + 1, *StageName);
+					const int32 DestSimStageIndex = CompilationOutput.ScriptData.SimulationStageMetaData.Num();
+					TranslationStages[Index].PassNamespace = FString::Printf(TEXT("MapSimStage%d%s"), DestSimStageIndex, *StageName);
 					TranslationStages[Index].ChunkModeIndex = (ENiagaraCodeChunkMode)(((int32)ENiagaraCodeChunkMode::SimulationStageBody) + (Index - 2));
 					if (TranslationStages[Index].ChunkModeIndex >= ENiagaraCodeChunkMode::SimulationStageBodyMax)
 					{
@@ -1218,14 +1223,14 @@ const FNiagaraTranslateResults &FHlslNiagaraTranslator::Translate(const FNiagara
 					ensure(TranslationStages[Index].OutputNode == HistoryOutputNode);
 					TranslationStages[Index].bInterpolatePreviousParams = false;
 					TranslationStages[Index].bCopyPreviousParams = false;
-					TranslationStages[Index].SimulationStageIndexMin = SimStageStartIndex;
-					TranslationStages[Index].SimulationStageIndexMax = SimStageStartIndex + NumIterationsThisStage;
-					TranslationStages[Index].NumIterationsThisStage = NumIterationsThisStage;
-					TranslationStages[Index].bSpawnOnly = bSpawnOnly;
-					TranslationStages[Index].bPartialParticleUpdate = InCompileData->PartialParticleUpdatePerStage.IsValidIndex(SimStageIndex) ? InCompileData->PartialParticleUpdatePerStage[SimStageIndex] : false;
-					TranslationStages[Index].IterationSource = IterationSrc;
-					TranslationStages[Index].SourceSimStage = SimStageIndex;
-					SimStageStartIndex += NumIterationsThisStage;
+					TranslationStages[Index].SimulationStageIndex = DestSimStageIndex;					
+					TranslationStages[Index].EnabledBinding = CompileSimStageData.EnabledBinding;
+					TranslationStages[Index].NumIterations = CompileSimStageData.NumIterations;
+					TranslationStages[Index].ExecuteBehavior = CompileSimStageData.ExecuteBehavior;
+					TranslationStages[Index].bPartialParticleUpdate = CompileSimStageData.PartialParticleUpdate;
+					TranslationStages[Index].IterationSource = CompileSimStageData.IterationSource;
+					TranslationStages[Index].NumIterationsBinding = CompileSimStageData.NumIterationsBinding;
+
 					ParamMapHistories.AddDefaulted(1);
 					ParamMapHistoriesSourceInOtherHistories.AddDefaulted(1);
 
@@ -1282,18 +1287,20 @@ const FNiagaraTranslateResults &FHlslNiagaraTranslator::Translate(const FNiagara
 					
 					// Set up the compile output for the shader stages so that we can properly execute at runtime.
 					FSimulationStageMetaData& SimulationStageMetaData = CompilationOutput.ScriptData.SimulationStageMetaData.AddDefaulted_GetRef();
-					SimulationStageMetaData.SimulationStageName = InCompileData->StageNames.IsValidIndex(SimStageIndex) ? InCompileData->StageNames[SimStageIndex] : FName();
-					SimulationStageMetaData.bSpawnOnly = bSpawnOnly;
-					SimulationStageMetaData.IterationSource = IterationSrc;
-					SimulationStageMetaData.MinStage = TranslationStages[Index].SimulationStageIndexMin;
-					SimulationStageMetaData.MaxStage = TranslationStages[Index].SimulationStageIndexMax;
+					SimulationStageMetaData.SimulationStageName = InCompileData->CompileSimStageData[SourceSimStageIndex].StageName;
+					SimulationStageMetaData.EnabledBinding = InCompileData->CompileSimStageData[SourceSimStageIndex].EnabledBinding;
+					SimulationStageMetaData.ExecuteBehavior = TranslationStages[Index].ExecuteBehavior;
+					SimulationStageMetaData.IterationSource = InCompileData->CompileSimStageData[SourceSimStageIndex].IterationSource;
+					SimulationStageMetaData.NumIterationsBinding = InCompileData->CompileSimStageData[SourceSimStageIndex].NumIterationsBinding;
+					SimulationStageMetaData.NumIterations = TranslationStages[Index].NumIterations;
 					SimulationStageMetaData.bWritesParticles = TranslationStages[Index].bWritesParticles;
 					SimulationStageMetaData.bPartialParticleUpdate = TranslationStages[Index].bPartialParticleUpdate;
 
-					// Other outputs are written to as appropriate data interfaces are found. See HandleDataInterfaceCall for details.
+					// Increment source stage index
+					++SourceSimStageIndex;
 
+					// Other outputs are written to as appropriate data interfaces are found. See HandleDataInterfaceCall for details
 					ParamMapSetVariablesToChunks.AddDefaulted(1);
-					SimStageIndex++;
 				}
 			}
 		}
@@ -1303,18 +1310,18 @@ const FNiagaraTranslateResults &FHlslNiagaraTranslator::Translate(const FNiagara
 		TranslationStages[0].PassNamespace = TEXT("Map");
 		TranslationStages[0].OutputNode = SourceGraph->FindEquivalentOutputNode(TranslationStages[0].ScriptUsage, TranslationStages[0].UsageId);
 		TranslationStages[0].ChunkModeIndex = ENiagaraCodeChunkMode::Body;
-		TranslationStages[0].SimulationStageIndexMin = 0;
-		TranslationStages[0].SimulationStageIndexMax = 0;
+		TranslationStages[0].SimulationStageIndex = 0;
+		TranslationStages[0].NumIterations = 1;
 		TranslationStages[0].bWritesParticles = true;
 
 		if (CompileOptions.TargetUsage == ENiagaraScriptUsage::ParticleSimulationStageScript)
 		{
-			for (int32 StageIdx = 0; StageIdx < InCompileData->StageGuids.Num(); StageIdx++)
+			for (int32 StageIdx = 0; StageIdx < InCompileData->CompileSimStageData.Num(); StageIdx++)
 			{
-				const FGuid& StageGuid = InCompileData->StageGuids[StageIdx];
-				if (StageGuid == CompileOptions.TargetUsageId && InCompileData->IterationSourcePerStage.IsValidIndex(StageIdx))
+				const FGuid& StageGuid = InCompileData->CompileSimStageData[StageIdx].StageGuid;
+				if (StageGuid == CompileOptions.TargetUsageId && InCompileData->CompileSimStageData.IsValidIndex(StageIdx))
 				{
-					TranslationStages[0].IterationSource = InCompileData->IterationSourcePerStage[StageIdx];
+					TranslationStages[0].IterationSource = InCompileData->CompileSimStageData[StageIdx].IterationSource;
 				}
 			}
 		}
@@ -1876,30 +1883,30 @@ const FNiagaraTranslateResults &FHlslNiagaraTranslator::Translate(const FNiagara
 		// Log out all the information we've built thus far for assistance debugging
 		FString Preamble = TEXT("// Shader generated by Niagara HLSL Translator\n\n");
 		
+		const UEnum* ExecuteBehaviorEnum = StaticEnum<ENiagaraSimStageExecuteBehavior>();
 		for (int32 i = 0; i < CompilationOutput.ScriptData.SimulationStageMetaData.Num(); i++)
 		{
-			Preamble += FString::Printf(TEXT("// SimStage[%d]\n//\t Iteration Src: \"%s\"\n//\tMinIndex: %d  MaxIndex: %d\n//\tbSpawnOnly: %s\n//\tWritesParticles: %s\n//\tPartialParticleUpdate: %s\n"), i,
-				*CompilationOutput.ScriptData.SimulationStageMetaData[i].IterationSource.ToString(),
-				CompilationOutput.ScriptData.SimulationStageMetaData[i].MinStage,
-				CompilationOutput.ScriptData.SimulationStageMetaData[i].MaxStage,
-				CompilationOutput.ScriptData.SimulationStageMetaData[i].bSpawnOnly ? TEXT("True") : TEXT("False"),
-				CompilationOutput.ScriptData.SimulationStageMetaData[i].bWritesParticles ? TEXT("True") : TEXT("False"),
-				CompilationOutput.ScriptData.SimulationStageMetaData[i].bPartialParticleUpdate ? TEXT("True") : TEXT("False")
-			);
-			for (const FName& Dest : CompilationOutput.ScriptData.SimulationStageMetaData[i].OutputDestinations)
+			const FSimulationStageMetaData& SimStageMetaData = CompilationOutput.ScriptData.SimulationStageMetaData[i];
+			Preamble.Appendf(TEXT("// SimStage[%d] = %s\n"), i, *SimStageMetaData.SimulationStageName.ToString());
+			Preamble.Appendf(TEXT("//\tNumIterations = %d\n"), SimStageMetaData.NumIterations);
+			Preamble.Appendf(TEXT("//\tExecuteBehavior = %s\n"), *ExecuteBehaviorEnum->GetNameStringByValue((int64)SimStageMetaData.ExecuteBehavior));
+			Preamble.Appendf(TEXT("//\tWritesParticles = %s\n"), SimStageMetaData.bWritesParticles ? TEXT("True") : TEXT("False"));
+			Preamble.Appendf(TEXT("//\tPartialParticleUpdate = %s\n"), SimStageMetaData.bPartialParticleUpdate ? TEXT("True") : TEXT("False"));
+
+			for (const FName& Dest : SimStageMetaData.OutputDestinations)
 			{
-				Preamble += FString::Printf(TEXT("//\tOutputs to: \"%s\"\n"), *Dest.ToString());
+				Preamble.Appendf(TEXT("//\tOutputs to: \"%s\"\n"), *Dest.ToString());
 			}
 		}
 
 		// Display the computed compile tags in the source hlsl to make checking easier.
 		if (TranslateResults.CompileTags.Num() != 0)
 		{
-			Preamble += FString::Printf(TEXT("//\tCompile Tags: \n"));
+			Preamble.Appendf(TEXT("// Compile Tags: \n"));
 
 			for (const FNiagaraCompilerTag& Tag : TranslateResults.CompileTags)
 			{
-				Preamble += FString::Printf(TEXT("//\t\tVariable: \"%s\" StringValue: \"%s\" \n"), *Tag.Variable.ToString(),  *Tag.StringValue);
+				Preamble.Appendf(TEXT("//\tVariable: \"%s\" StringValue: \"%s\" \n"), *Tag.Variable.ToString(),  *Tag.StringValue);
 			}
 
 		}
@@ -1996,8 +2003,8 @@ void FHlslNiagaraTranslator::HandleSimStageSetupAndTeardown(int32 InWhichStage,
 	// RIght now we need to know if anyone wrote to the IterationSource this stage. That can be one of two ways:
 	// 1) Someone wrote to StackContext.XXXX
 	// 2) Someone called a function that was marked to write 
-	int32 SourceSimStage = TranslationStage.SourceSimStage;
-	ensure(CompilationOutput.ScriptData.SimulationStageMetaData.Num() > SourceSimStage && SourceSimStage >= 0);
+	const int32 SourceSimStage = TranslationStage.SimulationStageIndex;
+	ensure(CompilationOutput.ScriptData.SimulationStageMetaData.IsValidIndex(SourceSimStage));
 	bool bWroteToIterationSource = CompilationOutput.ScriptData.SimulationStageMetaData[SourceSimStage].OutputDestinations.Contains(TranslationStage.IterationSource);
 	if (WriteVars.Num() > 0)
 		bWroteToIterationSource = true;
@@ -2091,10 +2098,11 @@ void FHlslNiagaraTranslator::HandleSimStageSetupAndTeardown(int32 InWhichStage,
 		TArray<FText> GeneratedErrors;
 		bool bPartialWrites = false;
 		
+		const bool bSpawnOnly = TranslationStage.ExecuteBehavior == ENiagaraSimStageExecuteBehavior::OnSimulationReset;
 		if (bNeedsSetupAndTeardown)
 		{
-			FString SetupGeneratedHLSL;			
-			if (CDO->GenerateSetupHLSL(DIInstanceInfo, MakeArrayView(Sig.Inputs), TranslationStage.bSpawnOnly, bPartialWrites, GeneratedErrors, SetupGeneratedHLSL) && SetupGeneratedHLSL.Len() > 0)
+			FString SetupGeneratedHLSL;
+			if (CDO->GenerateSetupHLSL(DIInstanceInfo, MakeArrayView(Sig.Inputs), bSpawnOnly, bPartialWrites, GeneratedErrors, SetupGeneratedHLSL) && SetupGeneratedHLSL.Len() > 0)
 			{
 				Sig.Name = *FString::Printf(TEXT("SetupFromIterationSource_%s_GeneratedSetup"), *GetSanitizedFunctionNameSuffix(TranslationStage.PassNamespace));
 				FNiagaraFunctionSignature SignatureOut = Sig;
@@ -2105,7 +2113,7 @@ void FHlslNiagaraTranslator::HandleSimStageSetupAndTeardown(int32 InWhichStage,
 			}
 
 			FString TeardownGeneratedHLSL;
-			if (CDO->GenerateTeardownHLSL(DIInstanceInfo, MakeArrayView(Sig.Inputs), TranslationStage.bSpawnOnly, bPartialWrites, GeneratedErrors, TeardownGeneratedHLSL) && TeardownGeneratedHLSL.Len() > 0)
+			if (CDO->GenerateTeardownHLSL(DIInstanceInfo, MakeArrayView(Sig.Inputs), bSpawnOnly, bPartialWrites, GeneratedErrors, TeardownGeneratedHLSL) && TeardownGeneratedHLSL.Len() > 0)
 			{
 				Sig.Name = *FString::Printf(TEXT("TeardownFromIterationSource_%s_GeneratedTeardown"), *GetSanitizedFunctionNameSuffix(TranslationStage.PassNamespace));
 				FNiagaraFunctionSignature SignatureOut = Sig;
@@ -2127,7 +2135,7 @@ void FHlslNiagaraTranslator::HandleSimStageSetupAndTeardown(int32 InWhichStage,
 				AttributeHLSLNames.Emplace(TEXT("Map.") + GetSanitizedSymbolName(Var.GetName().ToString()));
 			}
 
-			if (CDO->GenerateIterationSourceNamespaceReadAttributesHLSL(DIInstanceInfo,  IterationSourceVar, MakeArrayView(Sig.Inputs), MakeArrayView(ReadVars), MakeArrayView(AttributeHLSLNames), TranslationStage.bSpawnOnly, TranslationStage.bSpawnOnly, bPartialWrites, GeneratedErrors, AttributeReadGeneratedHLSL) && AttributeReadGeneratedHLSL.Len() > 0)
+			if (CDO->GenerateIterationSourceNamespaceReadAttributesHLSL(DIInstanceInfo,  IterationSourceVar, MakeArrayView(Sig.Inputs), MakeArrayView(ReadVars), MakeArrayView(AttributeHLSLNames), bSpawnOnly, bSpawnOnly, bPartialWrites, GeneratedErrors, AttributeReadGeneratedHLSL) && AttributeReadGeneratedHLSL.Len() > 0)
 			{
 				Sig.Name = *FString::Printf(TEXT("SetupFromIterationSource_%s_GeneratedReadAttributes"), *GetSanitizedFunctionNameSuffix(TranslationStage.PassNamespace));
 				FNiagaraFunctionSignature SignatureOut = Sig;
@@ -2148,7 +2156,7 @@ void FHlslNiagaraTranslator::HandleSimStageSetupAndTeardown(int32 InWhichStage,
 				AttributeHLSLNames.Emplace(TEXT("Map.") + GetSanitizedSymbolName(Var.GetName().ToString()));
 			}
 
-			if (CDO->GenerateIterationSourceNamespaceWriteAttributesHLSL(DIInstanceInfo, IterationSourceVar, MakeArrayView(Sig.Inputs), MakeArrayView(WriteVars), MakeArrayView(AttributeHLSLNames), TranslationStage.bSpawnOnly, bPartialWrites, GeneratedErrors, AttributeWriteGeneratedHLSL) && AttributeWriteGeneratedHLSL.Len() > 0)
+			if (CDO->GenerateIterationSourceNamespaceWriteAttributesHLSL(DIInstanceInfo, IterationSourceVar, MakeArrayView(Sig.Inputs), MakeArrayView(WriteVars), MakeArrayView(AttributeHLSLNames), bSpawnOnly, bPartialWrites, GeneratedErrors, AttributeWriteGeneratedHLSL) && AttributeWriteGeneratedHLSL.Len() > 0)
 			{
 				Sig.Name = *FString::Printf(TEXT("TeardownFromIterationSource_%s_GeneratedWriteAttributes"), *GetSanitizedFunctionNameSuffix(TranslationStage.PassNamespace));
 				FNiagaraFunctionSignature SignatureOut = Sig;
@@ -2329,7 +2337,7 @@ void FHlslNiagaraTranslator::GatherComponentsForDataSetAccess(UScriptStruct* Str
 void FHlslNiagaraTranslator::DefinePreviousParametersFunction(FString& HlslOutputString, TArray<TArray<FNiagaraVariable>>& DataSetVariables, TMap<FNiagaraDataSetID, int32>& DataSetReads, TMap<FNiagaraDataSetID, int32>& DataSetWrites)
 {
 	HlslOutputString +=
-		TEXT("#if (((SimulationStageIndex >= 0) && (SimulationStageIndex < 1))) // MapSpawn\n")
+		TEXT("#if (SimulationStageIndex == 0) // MapSpawn\n")
 		TEXT("void HandlePreviousValuesForSpawn(inout FSimulationContext Context)\n{\n");
 	if (UNiagaraScript::IsParticleSpawnScript(CompileOptions.TargetUsage) || UNiagaraScript::IsGPUScript(CompileOptions.TargetUsage))
 	{
@@ -4508,10 +4516,11 @@ FString FHlslNiagaraTranslator::GenerateConstantString(const FNiagaraVariable& C
 bool FHlslNiagaraTranslationStage::ShouldDoSpawnOnlyLogic() const
 {
 	if (UNiagaraScript::IsSpawnScript(ScriptUsage))
+	{
 		return true;
-	if (ScriptUsage == ENiagaraScriptUsage::ParticleSimulationStageScript && bSpawnOnly)
-		return true;
-	return false;
+	}
+	
+	return (ScriptUsage == ENiagaraScriptUsage::ParticleSimulationStageScript) && (ExecuteBehavior == ENiagaraSimStageExecuteBehavior::OnSimulationReset);
 }
 
 bool FHlslNiagaraTranslationStage::IsExternalConstantNamespace(const FNiagaraVariable& InVar, ENiagaraScriptUsage InTargetUsage, uint32 InTargetBitmask)
@@ -4532,7 +4541,8 @@ bool FHlslNiagaraTranslationStage::IsRelevantToSpawnForStage(const FNiagaraParam
 	{
 		return true;
 	}
-	if (ScriptUsage == ENiagaraScriptUsage::ParticleSimulationStageScript && bSpawnOnly)
+
+	if ((ScriptUsage == ENiagaraScriptUsage::ParticleSimulationStageScript) && (ExecuteBehavior == ENiagaraSimStageExecuteBehavior::OnSimulationReset))
 	{
 		if (IterationSource == NAME_None)
 		{
@@ -4929,12 +4939,13 @@ void FHlslNiagaraTranslator::ParameterMapSet(UNiagaraNodeParameterMapSet* SetNod
 			{
 				if (Var == FNiagaraVariable(FNiagaraTypeDefinition::GetBoolDef(), TEXT("DataInstance.Alive")))
 				{
-					const int32 OutputStageIndex = ActiveStageIdx - 2;
-					if (CompilationOutput.ScriptData.SimulationStageMetaData.IsValidIndex(OutputStageIndex))
+					const int32 OutputStageIndex = TranslationStages[ActiveStageIdx].SimulationStageIndex;
+					if ( CompilationOutput.ScriptData.SimulationStageMetaData.IsValidIndex(OutputStageIndex) )
 					{
 						CompilationOutput.ScriptData.SimulationStageMetaData[OutputStageIndex].bWritesParticles = true;
 						CompilationOutput.ScriptData.SimulationStageMetaData[OutputStageIndex].bPartialParticleUpdate = false;
 					}
+
 					TranslationStages[ActiveStageIdx].bWritesParticles = true;
 					TranslationStages[ActiveStageIdx].bPartialParticleUpdate = false;
 					TranslationStages[ActiveStageIdx].bWritesAlive = true;
@@ -6586,9 +6597,9 @@ void FHlslNiagaraTranslator::ProcessCustomHlsl(const FString& InCustomHlsl, ENia
 			if (bUsesInputData)
 			{
 				// Clear out the ability to use partial particle writes as we can't be sure how InputData is being used
-				for ( int i=0; i < CompileData->PartialParticleUpdatePerStage.Num(); ++i)
+				for (const auto& CompileStageData : CompileData->CompileSimStageData)
 				{
-					CompileData->PartialParticleUpdatePerStage[i] = false;
+					CompileStageData.PartialParticleUpdate = false;
 				}
 				break;
 			}
@@ -6893,14 +6904,14 @@ void FHlslNiagaraTranslator::HandleDataInterfaceCall(FNiagaraScriptDataInterface
 	
 	//UE_LOG(LogNiagaraEditor, Log, TEXT("HandleDataInterfaceCall %d %s %s %s"), ActiveStageIdx, *InMatchingSignature.Name.ToString(), InMatchingSignature.bWriteFunction ? TEXT("true") : TEXT("False"), *Info.Name.ToString());
 
-	if (InMatchingSignature.bWriteFunction && CompilationOutput.ScriptData.SimulationStageMetaData.Num() != 0 && TranslationStages[ActiveStageIdx].SourceSimStage != -1)
+	if (InMatchingSignature.bWriteFunction && CompilationOutput.ScriptData.SimulationStageMetaData.Num() > 1 && TranslationStages[ActiveStageIdx].SimulationStageIndex != -1)
 	{
-		int32 SourceSimStage = TranslationStages[ActiveStageIdx].SourceSimStage;
-		ensure(CompilationOutput.ScriptData.SimulationStageMetaData.Num() > SourceSimStage);
+		const int32 SourceSimStage = TranslationStages[ActiveStageIdx].SimulationStageIndex;
 		CompilationOutput.ScriptData.SimulationStageMetaData[SourceSimStage].OutputDestinations.AddUnique(Info.Name);
 		if (ActiveStageWriteTargets.Num() > 0)
+		{
 			ActiveStageWriteTargets.Top().AddUnique(Info.Name);
-		//UE_LOG(LogNiagaraEditor, Log, TEXT(" Add to %d %d"), ActiveStageIdx, SourceSimStage);
+		}
 	}
 }
 
@@ -7157,12 +7168,9 @@ void FHlslNiagaraTranslator::RegisterFunctionCall(ENiagaraScriptUsage ScriptUsag
 			{
 				for (const FName& Entry : *Entries)
 				{
-					int32 SourceSimStage = TranslationStages[ActiveStageIdx].SourceSimStage;
-					ensure(CompilationOutput.ScriptData.SimulationStageMetaData.Num() > SourceSimStage);
+					const int32 SourceSimStage = TranslationStages[ActiveStageIdx].SimulationStageIndex;
 					CompilationOutput.ScriptData.SimulationStageMetaData[SourceSimStage].OutputDestinations.AddUnique(Entry);
 					ActiveStageWriteTargets.Top().AddUnique(Entry);
-					//UE_LOG(LogNiagaraEditor, Log, TEXT(" Add to %d %d"), ActiveStageIdx, SourceSimStage);
-
 				}
 			}
 		}
diff --git a/Engine/Plugins/FX/Niagara/Source/NiagaraEditor/Private/NiagaraHlslTranslator.h b/Engine/Plugins/FX/Niagara/Source/NiagaraEditor/Private/NiagaraHlslTranslator.h
--- a/Engine/Plugins/FX/Niagara/Source/NiagaraEditor/Private/NiagaraHlslTranslator.h
+++ b/Engine/Plugins/FX/Niagara/Source/NiagaraEditor/Private/NiagaraHlslTranslator.h
@@ -91,12 +91,20 @@ public:
 	virtual void ReleaseCompilationCopies() override;
 
 	// Simulation Stage Variables. Sim stage of 0 is always Spawn/Update
-	TArray<uint32> NumIterationsPerStage;
-	TArray<FName> IterationSourcePerStage;
-	TArray<bool> SpawnOnlyPerStage;
-	mutable TArray<bool> PartialParticleUpdatePerStage;		//-TODO: Remove mutable and communicate that we can do a partial write in another way
-	TArray<FGuid> StageGuids;
-	TArray<FName> StageNames;
+	struct FCompileSimStageData
+	{
+		FGuid							StageGuid;
+		FName							StageName;
+		FName							EnabledBinding;
+		uint32							NumIterations = 1;
+		FName							NumIterationsBinding;
+		FName							IterationSource;
+		ENiagaraSimStageExecuteBehavior	ExecuteBehavior = ENiagaraSimStageExecuteBehavior::Always;
+		mutable bool					PartialParticleUpdate = false;
+	};
+	TArray<FCompileSimStageData> CompileSimStageData;
+
+	//TArray<FName> IterationSourcePerStage;
 
 	TWeakObjectPtr<UNiagaraGraph> NodeGraphDeepCopy;
 	TArray<FNiagaraParameterMapHistory> PrecompiledHistories;
@@ -195,8 +203,7 @@ FORCEINLINE uint32 GetTypeHash(const FNiagaraFunctionSignature& Sig)
 		Hash = HashCombine(Hash, GetTypeHash(Var));
 	}
 	Hash = HashCombine(Hash, GetTypeHash(Sig.OwnerName));
-	Hash = HashCombine(Hash, GetTypeHash(Sig.ContextStageMinIndex));
-	Hash = HashCombine(Hash, GetTypeHash(Sig.ContextStageMaxIndex));
+	Hash = HashCombine(Hash, GetTypeHash(Sig.ContextStageIndex));
 	return Hash;
 }
 
@@ -302,11 +309,11 @@ public:
 	bool bCopyPreviousParams = true;
 	ENiagaraCodeChunkMode ChunkModeIndex = (ENiagaraCodeChunkMode)-1;
 	FName IterationSource;
-	int32 SimulationStageIndexMin = -1;
-	int32 SimulationStageIndexMax = -1;
-	int32 NumIterationsThisStage = 1;
-	int32 SourceSimStage = -1;
-	bool bSpawnOnly = false;
+	int32 SimulationStageIndex = -1;
+	FName EnabledBinding;
+	int32 NumIterations = 1;
+	FName NumIterationsBinding;
+	ENiagaraSimStageExecuteBehavior ExecuteBehavior = ENiagaraSimStageExecuteBehavior::Always;
 	bool bUsesAlive = false;
 	bool bWritesAlive = false;
 	bool bWritesParticles = false;
diff --git a/Engine/Plugins/FX/Niagara/Source/NiagaraEditor/Private/ViewModels/NiagaraScriptStatsViewModel.cpp b/Engine/Plugins/FX/Niagara/Source/NiagaraEditor/Private/ViewModels/NiagaraScriptStatsViewModel.cpp
--- a/Engine/Plugins/FX/Niagara/Source/NiagaraEditor/Private/ViewModels/NiagaraScriptStatsViewModel.cpp
+++ b/Engine/Plugins/FX/Niagara/Source/NiagaraEditor/Private/ViewModels/NiagaraScriptStatsViewModel.cpp
@@ -186,22 +186,12 @@ namespace NiagaraScriptStatsLocal
 						for (const auto& ShaderScript : ShaderScripts)
 						{
 							TConstArrayView<FSimulationStageMetaData> SimulationStageMetaData = ShaderScript->GetBaseVMScript()->GetSimulationStageMetaData();
-							for (int32 iPermutation = 0; iPermutation < ShaderScript->GetNumPermutations(); ++iPermutation)
+							for (int32 iSimStageIndex=0; iSimStageIndex < SimulationStageMetaData.Num(); ++iSimStageIndex)
 							{
-								FNiagaraShaderRef Shader = ShaderScript->GetShaderGameThread(iPermutation);
+								FNiagaraShaderRef Shader = ShaderScript->GetShaderGameThread(iSimStageIndex);
 								if ( Shader.IsValid() )
 								{
-									FName StageName;
-									const int32 ShaderStage = iPermutation - 1;
-									if (SimulationStageMetaData.IsValidIndex(ShaderStage))
-									{
-										if (!SimulationStageMetaData[ShaderStage].SimulationStageName.IsNone())
-										{
-											StageName = SimulationStageMetaData[ShaderStage].SimulationStageName;
-										}
-									}
-
-									ResultsString.Appendf(TEXT("GPU: %s = "), StageName.IsNone() ? TEXT("Particles") : *StageName.ToString());
+									ResultsString.Appendf(TEXT("GPU: %s = "), *SimulationStageMetaData[iSimStageIndex].SimulationStageName.ToString());
 									if (Shader->GetNumInstructions() == 0)
 									{
 										ResultsString.Append(TEXT("n/a"));
diff --git a/Engine/Plugins/FX/Niagara/Source/NiagaraEditor/Private/Widgets/SNiagaraSystemViewport.cpp b/Engine/Plugins/FX/Niagara/Source/NiagaraEditor/Private/Widgets/SNiagaraSystemViewport.cpp
--- a/Engine/Plugins/FX/Niagara/Source/NiagaraEditor/Private/Widgets/SNiagaraSystemViewport.cpp
+++ b/Engine/Plugins/FX/Niagara/Source/NiagaraEditor/Private/Widgets/SNiagaraSystemViewport.cpp
@@ -206,29 +206,14 @@ void FNiagaraSystemViewportClient::DrawInstructionCounts(UNiagaraSystem* Particl
 				if (ShaderScript != nullptr && ShaderScript->GetBaseVMScript() != nullptr)
 				{
 					TConstArrayView<FSimulationStageMetaData> SimulationStageMetaData = ShaderScript->GetBaseVMScript()->GetSimulationStageMetaData();
-
-					for (int32 iPermutation=0; iPermutation < ShaderScript->GetNumPermutations(); ++iPermutation)
+					for (int32 iSimStageIndex = 0; iSimStageIndex < SimulationStageMetaData.Num(); ++iSimStageIndex)
 					{
-						FNiagaraShaderRef Shader = ShaderScript->GetShaderGameThread(iPermutation);
+						FNiagaraShaderRef Shader = ShaderScript->GetShaderGameThread(iSimStageIndex);
 						if (Shader.IsValid())
 						{
 							FColor DisplayColor = FColor(196, 196, 196);
-							FString StageName = TEXT("Particles");
-							int32 MinStage = 0;
-							int32 MaxStage = 1;
-
-							const int32 ShaderStage = iPermutation - 1;
-							if (SimulationStageMetaData.IsValidIndex(ShaderStage))
-							{
-								if (!SimulationStageMetaData[ShaderStage].SimulationStageName.IsNone())
-								{
-									StageName = SimulationStageMetaData[ShaderStage].SimulationStageName.ToString();
-								}
-								MinStage = SimulationStageMetaData[ShaderStage].MinStage;
-								MaxStage = SimulationStageMetaData[ShaderStage].MaxStage;
-							}
-
-							Canvas->DrawShadowedString(CurrentX + 20.0f, CurrentY, *FString::Printf(TEXT("GPU StageName(%s) Stages(%d - %d) = %u"), *StageName, MinStage, MaxStage, Shader->GetNumInstructions()), Font, DisplayColor);
+							FString StageName = SimulationStageMetaData[iSimStageIndex].SimulationStageName.ToString();
+							Canvas->DrawShadowedString(CurrentX + 20.0f, CurrentY, *FString::Printf(TEXT("GPU StageName(%s) Stage(%d) = %u"), *StageName, iSimStageIndex, Shader->GetNumInstructions()), Font, DisplayColor);
 							CurrentY += FontHeight;
 						}
 					}
diff --git a/Engine/Plugins/FX/Niagara/Source/NiagaraShader/Private/NiagaraScriptBase.cpp b/Engine/Plugins/FX/Niagara/Source/NiagaraShader/Private/NiagaraScriptBase.cpp
--- a/Engine/Plugins/FX/Niagara/Source/NiagaraShader/Private/NiagaraScriptBase.cpp
+++ b/Engine/Plugins/FX/Niagara/Source/NiagaraShader/Private/NiagaraScriptBase.cpp
@@ -8,8 +8,7 @@ UNiagaraScriptBase::UNiagaraScriptBase(const FObjectInitializer& ObjectInitializ
 }
 
 FSimulationStageMetaData::FSimulationStageMetaData()
-	: bSpawnOnly(0)
-	, bWritesParticles(0)
+	: bWritesParticles(0)
 	, bPartialParticleUpdate(0)
 {
 }
diff --git a/Engine/Plugins/FX/Niagara/Source/NiagaraShader/Private/NiagaraShader.cpp b/Engine/Plugins/FX/Niagara/Source/NiagaraShader/Private/NiagaraShader.cpp
--- a/Engine/Plugins/FX/Niagara/Source/NiagaraShader/Private/NiagaraShader.cpp
+++ b/Engine/Plugins/FX/Niagara/Source/NiagaraShader/Private/NiagaraShader.cpp
@@ -405,20 +405,16 @@ void FNiagaraShaderType::BeginCompileShader(
 	const bool bUsesSimulationStages = Script->GetUsesSimulationStages();
 	if (bUsesSimulationStages)
 	{
-		const int32 StageIndex = Script->PermutationIdToShaderStageIndex(PermutationId);
 		NewJob->Input.Environment.SetDefine(TEXT("NIAGARA_SHADER_PERMUTATIONS"), 1);
 		NewJob->Input.Environment.SetDefine(TEXT("DefaultSimulationStageIndex"), 0);
-		NewJob->Input.Environment.SetDefine(TEXT("SimulationStageIndex"), StageIndex);
+		NewJob->Input.Environment.SetDefine(TEXT("SimulationStageIndex"), PermutationId);
 		NewJob->Input.Environment.SetDefine(TEXT("USE_SIMULATION_STAGES"), 1);
 
-		if (PermutationId > 0)
+		TConstArrayView<FSimulationStageMetaData> SimStageMetaDataArray = const_cast<FNiagaraShaderScript*>(Script)->GetBaseVMScript()->GetSimulationStageMetaData();
+		const FSimulationStageMetaData& SimStageMetaData = SimStageMetaDataArray[PermutationId];
+		if (SimStageMetaData.bWritesParticles && SimStageMetaData.bPartialParticleUpdate)
 		{
-			TConstArrayView<FSimulationStageMetaData> SimStageMetaDataArray = const_cast<FNiagaraShaderScript*>(Script)->GetBaseVMScript()->GetSimulationStageMetaData();
-			const FSimulationStageMetaData& SimStageMetaData = SimStageMetaDataArray[PermutationId - 1];
-			if (SimStageMetaData.bWritesParticles && SimStageMetaData.bPartialParticleUpdate)
-			{
-				NewJob->Input.Environment.SetDefine(TEXT("NIAGARA_PARTICLE_PARTIAL_ENABLED"), 1);
-			}
+			NewJob->Input.Environment.SetDefine(TEXT("NIAGARA_PARTICLE_PARTIAL_ENABLED"), 1);
 		}
 	}
 	else
diff --git a/Engine/Plugins/FX/Niagara/Source/NiagaraShader/Private/NiagaraShared.cpp b/Engine/Plugins/FX/Niagara/Source/NiagaraShader/Private/NiagaraShared.cpp
--- a/Engine/Plugins/FX/Niagara/Source/NiagaraShader/Private/NiagaraShared.cpp
+++ b/Engine/Plugins/FX/Niagara/Source/NiagaraShader/Private/NiagaraShared.cpp
@@ -160,11 +160,6 @@ bool FNiagaraShaderScript::IsSame(const FNiagaraShaderMapId& InId) const
 		InId.CompilerVersionID == CompilerVersionId;
 }
 
-int32 FNiagaraShaderScript::PermutationIdToShaderStageIndex(int32 PermutationId) const
-{
-	return ShaderStageToPermutation[PermutationId].Key;
-}
-
 bool FNiagaraShaderScript::IsShaderMapComplete() const
 {
 	if (GameThreadShaderMap == nullptr)
@@ -192,25 +187,6 @@ bool FNiagaraShaderScript::IsShaderMapComplete() const
 	return true;
 }
 
-int32 FNiagaraShaderScript::ShaderStageIndexToPermutationId_RenderThread(int32 ShaderStageIndex) const
-{
-	check(IsInRenderingThread());
-	if (CachedData_RenderThread.NumPermutations > 1)
-	{
-		for (int32 i = 0; i < CachedData_RenderThread.ShaderStageToPermutation.Num(); ++i)
-		{
-			const TPair<int32, int32> MinMaxStage = CachedData_RenderThread.ShaderStageToPermutation[i];
-			if ((ShaderStageIndex >= MinMaxStage.Key) && (ShaderStageIndex < MinMaxStage.Value))
-			{
-				return i;
-			}
-		}
-		UE_LOG(LogShaders, Fatal, TEXT("FNiagaraShaderScript::ShaderStageIndexToPermutationId_RenderThread: Failed to map from simulation stage(%d) to permutation id."), ShaderStageIndex);
-	}
-
-	return 0;
-}
-
 void FNiagaraShaderScript::GetDependentShaderTypes(EShaderPlatform Platform, TArray<FShaderType*>& OutShaderTypes) const
 {
 	for (TLinkedList<FShaderType*>::TIterator ShaderTypeIt(FShaderType::GetTypeList()); ShaderTypeIt; ShaderTypeIt.Next())
@@ -321,7 +297,6 @@ void FNiagaraShaderScript::SerializeShaderMap(FArchive& Ar)
 	bool bCooked = Ar.IsCooking();
 	Ar << bCooked;
 	Ar << NumPermutations;
-	Ar << ShaderStageToPermutation;
 
 	if (Ar.IsLoading())
 	{
@@ -488,24 +463,12 @@ void FNiagaraShaderScript::UpdateCachedData_PreCompile()
 {
 	if (BaseVMScript)
 	{
-		NumPermutations = 1;
-		ShaderStageToPermutation.Empty();
-
 		TConstArrayView<FSimulationStageMetaData> SimulationStages = BaseVMScript->GetSimulationStageMetaData();
-
-		// We add the number of simulation stages as Stage 0 is always the particle stage currently
-		NumPermutations += SimulationStages.Num();
-
-		ShaderStageToPermutation.Emplace(0, 1);
-		for (const FSimulationStageMetaData& StageMeta : SimulationStages)
-		{
-			ShaderStageToPermutation.Emplace(StageMeta.MinStage, StageMeta.MaxStage);
-		}
+		NumPermutations = SimulationStages.Num();
 	}
 	else
 	{
 		NumPermutations = 0;
-		ShaderStageToPermutation.Empty();
 	}
 }
 
@@ -552,8 +515,6 @@ void FNiagaraShaderScript::UpdateCachedData_PostCompile(bool bCalledFromSerializ
 		CachedData.bIsComplete = 0;
 	}
 
-	CachedData.ShaderStageToPermutation = ShaderStageToPermutation;
-
 	if (bCalledFromSerialize)
 	{
 		CachedData_RenderThread = MoveTemp(CachedData);
diff --git a/Engine/Plugins/FX/Niagara/Source/NiagaraShader/Public/NiagaraScriptBase.h b/Engine/Plugins/FX/Niagara/Source/NiagaraShader/Public/NiagaraScriptBase.h
--- a/Engine/Plugins/FX/Niagara/Source/NiagaraShader/Public/NiagaraScriptBase.h
+++ b/Engine/Plugins/FX/Niagara/Source/NiagaraShader/Public/NiagaraScriptBase.h
@@ -7,6 +7,17 @@
 
 #include "NiagaraScriptBase.generated.h"
 
+UENUM()
+enum class ENiagaraSimStageExecuteBehavior : uint8
+{
+	/** The stage will run every frame. */
+	Always,
+	/** The stage will only run on the frame when the simulation is reset. */
+	OnSimulationReset,
+	/** The stage will not run on the frame where the simulation is reset. */
+	NotOnSimulationReset,
+};
+
 USTRUCT()
 struct NIAGARASHADER_API FSimulationStageMetaData
 {
@@ -18,13 +29,16 @@ public:
 	UPROPERTY()
 	FName SimulationStageName;
 
+	UPROPERTY()
+	FName EnabledBinding;
+
 	/** The Data Interface that we iterate over for this stage. If None, then use particles.*/
 	UPROPERTY()
 	FName IterationSource;
 
-	/** Is this stage a spawn-only stage? */
+	/** Controls when the simulation stage will execute. */
 	UPROPERTY()
-	uint32 bSpawnOnly : 1;
+	ENiagaraSimStageExecuteBehavior ExecuteBehavior = ENiagaraSimStageExecuteBehavior::Always;
 
 	/** Do we write to particles this stage? */
 	UPROPERTY()
@@ -38,13 +52,20 @@ public:
 	UPROPERTY()
 	TArray<FName> OutputDestinations;
 
-	/** Index of the simulation stage where we begin iterating. This is meant to encompass iteration count without having an entry for each iteration.*/
+	/** The number of iterations for the stage. */
 	UPROPERTY()
-	int32 MinStage = 0;
+	int32 NumIterations = 1;
 
-	/** Index of the simulation stage where we end iterating. This is meant to encompass iteration count without having an entry for each iteration.*/
+	/** Optional binding to gather num iterations from. */
 	UPROPERTY()
-	int32 MaxStage = 0;
+	FName NumIterationsBinding;
+
+	inline bool ShouldRunStage(bool bResetData) const
+	{
+		const bool bAlways = ExecuteBehavior == ENiagaraSimStageExecuteBehavior::Always;
+		const bool bResetOnly = ExecuteBehavior == ENiagaraSimStageExecuteBehavior::OnSimulationReset;
+		return bAlways || (bResetOnly == bResetData);
+	}
 };
 
 UCLASS(MinimalAPI, abstract)
diff --git a/Engine/Plugins/FX/Niagara/Source/NiagaraShader/Public/NiagaraShared.h b/Engine/Plugins/FX/Niagara/Source/NiagaraShader/Public/NiagaraShared.h
--- a/Engine/Plugins/FX/Niagara/Source/NiagaraShader/Public/NiagaraShared.h
+++ b/Engine/Plugins/FX/Niagara/Source/NiagaraShader/Public/NiagaraShared.h
@@ -608,7 +608,6 @@ class FNiagaraShaderScript
 			bViewUniformBufferUsed = 0;
 		}
 
-		TArray<TPair<int32, int32>> ShaderStageToPermutation;
 		int32 NumPermutations;
 		uint32 bIsComplete : 1;
 		uint32 bGlobalConstantBufferUsed : 2;
@@ -812,7 +811,6 @@ public:
 	bool IsSame(const FNiagaraShaderMapId& InId) const;
 
 	NIAGARASHADER_API int32 GetNumPermutations() const { return NumPermutations; }
-	NIAGARASHADER_API int32 PermutationIdToShaderStageIndex(int32 PermutationId) const;
 
 	NIAGARASHADER_API bool IsShaderMapComplete() const;
 
@@ -824,7 +822,6 @@ public:
 	FORCEINLINE bool IsEmitterConstantBufferUsed_RenderThread(int32 Index) const { return (CachedData_RenderThread.bEmitterConstantBufferUsed & (1 << Index)) != 0; }
 	FORCEINLINE bool IsExternalConstantBufferUsed_RenderThread(int32 Index) const { return (CachedData_RenderThread.bExternalConstantBufferUsed & (1 << Index)) != 0; }
 	FORCEINLINE bool IsViewUniformBufferUsed_RenderThread() const { return CachedData_RenderThread.bViewUniformBufferUsed != 0; }
-	NIAGARASHADER_API int32 ShaderStageIndexToPermutationId_RenderThread(int32 ShaderStageIndex) const;
 
 protected:
 
@@ -904,7 +901,6 @@ private:
 	uint32 bQueuedForRelease : 1;
 
 	int32 NumPermutations = 0;
-	TArray<TPair<int32, int32>> ShaderStageToPermutation;
 
 	FNiagaraShaderMapCachedData CachedData_RenderThread;
 
