- Fix for Execute Behaviour not working as expected

#rnx

diff --git a/Engine/Plugins/FX/Niagara/Source/Niagara/Private/NiagaraGPUSystemTick.cpp b/Engine/Plugins/FX/Niagara/Source/Niagara/Private/NiagaraGPUSystemTick.cpp
--- a/Engine/Plugins/FX/Niagara/Source/Niagara/Private/NiagaraGPUSystemTick.cpp
+++ b/Engine/Plugins/FX/Niagara/Source/Niagara/Private/NiagaraGPUSystemTick.cpp
@@ -219,7 +219,7 @@ void FNiagaraGPUSystemTick::Init(FNiagaraSystemInstance* InSystemInstance)
 				}
 				else
 				{
-					TotalDispatches = 0;
+					NumIterations = 0;
 				}
 
 				InstanceData->TotalDispatches += NumIterations;
