- Fix for GPU export count not being correct when going over the limit

diff --git a/Engine/Plugins/FX/Niagara/Shaders/Private/NiagaraDataInterfaceExportTemplate.ush b/Engine/Plugins/FX/Niagara/Shaders/Private/NiagaraDataInterfaceExportTemplate.ush
--- a/Engine/Plugins/FX/Niagara/Shaders/Private/NiagaraDataInterfaceExportTemplate.ush
+++ b/Engine/Plugins/FX/Niagara/Shaders/Private/NiagaraDataInterfaceExportTemplate.ush
@@ -22,6 +22,10 @@ void ExportParticleData_{ParameterName}_UEImpureCall(in bool bStoreData, in floa
 			RWWriteBuffer_{ParameterName}[StoreOffset + 6] = asuint(Velocity.z);
 			bSuccess = true;
 		}
+		else
+		{
+			InterlockedAdd(RWWriteBuffer_{ParameterName}[0], -1u, StoreIndex);
+		}
 	}
 }
 
diff --git a/Engine/Plugins/FX/Niagara/Source/Niagara/Private/NiagaraDataInterfaceExport.cpp b/Engine/Plugins/FX/Niagara/Source/Niagara/Private/NiagaraDataInterfaceExport.cpp
--- a/Engine/Plugins/FX/Niagara/Source/Niagara/Private/NiagaraDataInterfaceExport.cpp
+++ b/Engine/Plugins/FX/Niagara/Source/Niagara/Private/NiagaraDataInterfaceExport.cpp
@@ -171,6 +171,7 @@ struct FNDIExportProxy : public FNiagaraDataInterfaceProxy
 				[MaxInstances=InstanceData->WriteBufferInstanceCount, WeakCallbackHandler=InstanceData->WeakCallbackHandler, WeakSystem=InstanceData->WeakSystem](TConstArrayView<TPair<void*, uint32>> Buffers)
 				{
 					const uint32 ReadbackInstanceCount = *reinterpret_cast<uint32*>(Buffers[0].Key);
+					check(ReadbackInstanceCount <= MaxInstances);
 					if (ReadbackInstanceCount > 0)
 					{
 						// Translate float data into Export Data
