- Functionality to push which LOD level the data interface samples from

#rnx

diff --git a/Engine/Plugins/FX/Niagara/Source/Niagara/Classes/NiagaraDataInterfaceSkeletalMesh.h b/Engine/Plugins/FX/Niagara/Source/Niagara/Classes/NiagaraDataInterfaceSkeletalMesh.h
--- a/Engine/Plugins/FX/Niagara/Source/Niagara/Classes/NiagaraDataInterfaceSkeletalMesh.h
+++ b/Engine/Plugins/FX/Niagara/Source/Niagara/Classes/NiagaraDataInterfaceSkeletalMesh.h
@@ -1009,6 +1009,7 @@ public:
 
 	void SetSourceComponentFromBlueprints(USkeletalMeshComponent* ComponentToUse);
 	void SetSamplingRegionsFromBlueprints(const TArray<FName>& InSamplingRegions);
+	void SetWholeMeshLODFromBlueprints(int32 MeshLODLevel);
 };
 
 
diff --git a/Engine/Plugins/FX/Niagara/Source/Niagara/Private/NiagaraDataInterfaceSkeletalMesh.cpp b/Engine/Plugins/FX/Niagara/Source/Niagara/Private/NiagaraDataInterfaceSkeletalMesh.cpp
--- a/Engine/Plugins/FX/Niagara/Source/Niagara/Private/NiagaraDataInterfaceSkeletalMesh.cpp
+++ b/Engine/Plugins/FX/Niagara/Source/Niagara/Private/NiagaraDataInterfaceSkeletalMesh.cpp
@@ -3431,6 +3431,13 @@ void UNiagaraDataInterfaceSkeletalMesh::SetSamplingRegionsFromBlueprints(const T
 	SamplingRegions = InSamplingRegions;
 }
 
+void UNiagaraDataInterfaceSkeletalMesh::SetWholeMeshLODFromBlueprints(int32 InWholeMeshLOD)
+{
+	// NOTE: When ChangeId changes the next tick will be skipped and a reset of the per-instance data will be initiated.
+	++ChangeId;
+	WholeMeshLOD = InWholeMeshLOD;
+}
+
 ETickingGroup UNiagaraDataInterfaceSkeletalMesh::CalculateTickGroup(const void* PerInstanceData) const
 {
 	const FNDISkeletalMesh_InstanceData* InstData = static_cast<const FNDISkeletalMesh_InstanceData*>(PerInstanceData);
