Default all new parameters to "Fail if Unset" default mode

diff --git a/Engine/Plugins/FX/Niagara/Source/NiagaraEditor/Private/NiagaraGraph.cpp b/Engine/Plugins/FX/Niagara/Source/NiagaraEditor/Private/NiagaraGraph.cpp
--- a/Engine/Plugins/FX/Niagara/Source/NiagaraEditor/Private/NiagaraGraph.cpp
+++ b/Engine/Plugins/FX/Niagara/Source/NiagaraEditor/Private/NiagaraGraph.cpp
@@ -1375,6 +1375,7 @@ UNiagaraScriptVariable* UNiagaraGraph::AddParameter(const FNiagaraVariable& Para
 		UNiagaraScriptVariable* NewScriptVariable = NewObject<UNiagaraScriptVariable>(this, FName(), RF_Transactional);
 		NewScriptVariable->Init(Parameter, ParameterMetaData);
 		NewScriptVariable->SetIsStaticSwitch(bIsStaticSwitch);
+		NewScriptVariable->DefaultMode = ENiagaraDefaultMode::FailIfPreviouslyNotSet;
 		VariableToScriptVariable.Add(Parameter, NewScriptVariable);
 		if (bNotifyChanged)
 		{
