Fixed a bug where GetDynamicDataSize used GT data on the RT

diff --git a/Engine/Plugins/FX/Niagara/Source/Niagara/Private/NiagaraSystemRenderData.cpp b/Engine/Plugins/FX/Niagara/Source/Niagara/Private/NiagaraSystemRenderData.cpp
--- a/Engine/Plugins/FX/Niagara/Source/Niagara/Private/NiagaraSystemRenderData.cpp
+++ b/Engine/Plugins/FX/Niagara/Source/Niagara/Private/NiagaraSystemRenderData.cpp
@@ -104,8 +104,10 @@ FPrimitiveViewRelevance FNiagaraSystemRenderData::GetViewRelevance(const FSceneV
 
 uint32 FNiagaraSystemRenderData::GetDynamicDataSize() const
 {
+	// only the render thread should use the RT data, the game thread and other parallel running tasks use the GT data
+	const TArray<FNiagaraRenderer*>& Emitter_Renderers = IsInParallelRenderingThread() ? EmitterRenderers_RT :EmitterRenderers_GT;
 	uint32 Size = 0;
-	for (const auto Renderer : EmitterRenderers_GT)
+	for (const auto Renderer : Emitter_Renderers)
 	{
 		if (Renderer)
 		{
