- Dynamically enabled bindings for renderers

diff --git a/Engine/Plugins/FX/Niagara/Source/Niagara/Private/NiagaraMeshRendererProperties.cpp b/Engine/Plugins/FX/Niagara/Source/Niagara/Private/NiagaraMeshRendererProperties.cpp
--- a/Engine/Plugins/FX/Niagara/Source/Niagara/Private/NiagaraMeshRendererProperties.cpp
+++ b/Engine/Plugins/FX/Niagara/Source/Niagara/Private/NiagaraMeshRendererProperties.cpp
@@ -477,7 +477,7 @@ void UNiagaraMeshRendererProperties::GetUsedMaterials(const FNiagaraEmitterInsta
 
 bool UNiagaraMeshRendererProperties::PopulateRequiredBindings(FNiagaraParameterStore& InParameterStore)
 {
-	bool bAnyAdded = false;
+	bool bAnyAdded = Super::PopulateRequiredBindings(InParameterStore);
 
 	for (const FNiagaraVariableAttributeBinding* Binding : AttributeBindings)
 	{
diff --git a/Engine/Plugins/FX/Niagara/Source/Niagara/Private/NiagaraRenderer.cpp b/Engine/Plugins/FX/Niagara/Source/Niagara/Private/NiagaraRenderer.cpp
--- a/Engine/Plugins/FX/Niagara/Source/Niagara/Private/NiagaraRenderer.cpp
+++ b/Engine/Plugins/FX/Niagara/Source/Niagara/Private/NiagaraRenderer.cpp
@@ -445,6 +445,23 @@ struct FParticleOrderAsUint
 	FORCEINLINE operator uint32() const { return OrderAsUint; }
 };
 
+bool FNiagaraRenderer::IsRendererEnabled(const UNiagaraRendererProperties* InProperties, const FNiagaraEmitterInstance* Emitter) const
+{
+	if (InProperties->RendererEnabledBinding.GetParamMapBindableVariable().IsValid())
+	{
+		const FNiagaraParameterStore& BoundParamStore = Emitter->GetRendererBoundVariables();
+		if (const uint8* ParameterData = BoundParamStore.GetParameterData(InProperties->RendererEnabledBinding.GetParamMapBindableVariable()))
+		{
+			const FNiagaraBool RendererEnabled = *reinterpret_cast<const FNiagaraBool*>(ParameterData);
+			if (RendererEnabled.GetValue() == false)
+			{
+				return false;
+			}
+		}
+	}
+	return true;
+}
+
 bool FNiagaraRenderer::UseLocalSpace(const FNiagaraSceneProxy* Proxy)const
 {
 	//Force local space if we're using a cull proxy. Cull proxies are simulated at the origin so their world space particles can be transformed to this system as though they were local space.
diff --git a/Engine/Plugins/FX/Niagara/Source/Niagara/Private/NiagaraRendererMeshes.cpp b/Engine/Plugins/FX/Niagara/Source/Niagara/Private/NiagaraRendererMeshes.cpp
--- a/Engine/Plugins/FX/Niagara/Source/Niagara/Private/NiagaraRendererMeshes.cpp
+++ b/Engine/Plugins/FX/Niagara/Source/Niagara/Private/NiagaraRendererMeshes.cpp
@@ -1406,6 +1406,11 @@ FNiagaraDynamicDataBase* FNiagaraRendererMeshes::GenerateDynamicData(const FNiag
 
 	const UNiagaraMeshRendererProperties* Properties = CastChecked<const UNiagaraMeshRendererProperties>(InProperties);
 	
+	if (!IsRendererEnabled(Properties, Emitter))
+	{
+		return nullptr;
+	}
+
 	if (Properties->bAllowInCullProxies == false)
 	{
 		check(Emitter);
diff --git a/Engine/Plugins/FX/Niagara/Source/Niagara/Private/NiagaraRendererProperties.cpp b/Engine/Plugins/FX/Niagara/Source/Niagara/Private/NiagaraRendererProperties.cpp
--- a/Engine/Plugins/FX/Niagara/Source/Niagara/Private/NiagaraRendererProperties.cpp
+++ b/Engine/Plugins/FX/Niagara/Source/Niagara/Private/NiagaraRendererProperties.cpp
@@ -306,6 +306,15 @@ void UNiagaraRendererProperties::GetAssetTagsForContext(const UObject* InAsset,
 	}
 }
 
+bool UNiagaraRendererProperties::PopulateRequiredBindings(FNiagaraParameterStore& InParameterStore)
+{
+	bool bAnyAdded = false;
+	if (RendererEnabledBinding.GetParamMapBindableVariable().IsValid())
+	{
+		bAnyAdded |= InParameterStore.AddParameter(RendererEnabledBinding.GetParamMapBindableVariable(), false);
+	}
+	return bAnyAdded;
+}
 
 bool UNiagaraRendererProperties::NeedsLoadForTargetPlatform(const ITargetPlatform* TargetPlatform) const
 {
@@ -337,6 +346,9 @@ void UNiagaraRendererProperties::PostInitProperties()
 	if (HasAnyFlags(RF_ClassDefaultObject) == false)
 	{
 		SetFlags(RF_Transactional);
+
+		FNiagaraVariableBase EnabledDefaultVariable(FNiagaraTypeDefinition::GetBoolDef(), NAME_None);
+		RendererEnabledBinding.Setup(EnabledDefaultVariable, EnabledDefaultVariable, ENiagaraRendererSourceDataMode::Emitter);
 	}
 #endif
 }
diff --git a/Engine/Plugins/FX/Niagara/Source/Niagara/Private/NiagaraRendererRibbons.cpp b/Engine/Plugins/FX/Niagara/Source/Niagara/Private/NiagaraRendererRibbons.cpp
--- a/Engine/Plugins/FX/Niagara/Source/Niagara/Private/NiagaraRendererRibbons.cpp
+++ b/Engine/Plugins/FX/Niagara/Source/Niagara/Private/NiagaraRendererRibbons.cpp
@@ -586,6 +586,11 @@ FNiagaraDynamicDataBase* FNiagaraRendererRibbons::GenerateDynamicData(const FNia
 	FNiagaraDataSet& Data = Emitter->GetData();
 	const UNiagaraRibbonRendererProperties* Properties = CastChecked<const UNiagaraRibbonRendererProperties>(InProperties);
 
+	if (!IsRendererEnabled(Properties, Emitter))
+	{
+		return nullptr;
+	}
+
 	if (Properties->bAllowInCullProxies == false)
 	{
 		check(Emitter);
diff --git a/Engine/Plugins/FX/Niagara/Source/Niagara/Private/NiagaraRendererSprites.cpp b/Engine/Plugins/FX/Niagara/Source/Niagara/Private/NiagaraRendererSprites.cpp
--- a/Engine/Plugins/FX/Niagara/Source/Niagara/Private/NiagaraRendererSprites.cpp
+++ b/Engine/Plugins/FX/Niagara/Source/Niagara/Private/NiagaraRendererSprites.cpp
@@ -1067,6 +1067,11 @@ FNiagaraDynamicDataBase *FNiagaraRendererSprites::GenerateDynamicData(const FNia
 
 	if (Properties)
 	{
+		if ( !IsRendererEnabled(Properties, Emitter) )
+		{
+			return nullptr;
+		}
+
 		if (Properties->bAllowInCullProxies == false)
 		{
 			check(Emitter);
diff --git a/Engine/Plugins/FX/Niagara/Source/Niagara/Private/NiagaraRibbonRendererProperties.cpp b/Engine/Plugins/FX/Niagara/Source/Niagara/Private/NiagaraRibbonRendererProperties.cpp
--- a/Engine/Plugins/FX/Niagara/Source/Niagara/Private/NiagaraRibbonRendererProperties.cpp
+++ b/Engine/Plugins/FX/Niagara/Source/Niagara/Private/NiagaraRibbonRendererProperties.cpp
@@ -160,7 +160,7 @@ void UNiagaraRibbonRendererProperties::GetUsedMaterials(const FNiagaraEmitterIns
 
 bool UNiagaraRibbonRendererProperties::PopulateRequiredBindings(FNiagaraParameterStore& InParameterStore)
 {
-	bool bAnyAdded = false;
+	bool bAnyAdded = Super::PopulateRequiredBindings(InParameterStore);
 
 	for (const FNiagaraVariableAttributeBinding* Binding : AttributeBindings)
 	{
diff --git a/Engine/Plugins/FX/Niagara/Source/Niagara/Private/NiagaraSpriteRendererProperties.cpp b/Engine/Plugins/FX/Niagara/Source/Niagara/Private/NiagaraSpriteRendererProperties.cpp
--- a/Engine/Plugins/FX/Niagara/Source/Niagara/Private/NiagaraSpriteRendererProperties.cpp
+++ b/Engine/Plugins/FX/Niagara/Source/Niagara/Private/NiagaraSpriteRendererProperties.cpp
@@ -321,7 +321,7 @@ void UNiagaraSpriteRendererProperties::CacheFromCompiledData(const FNiagaraDataS
 
 bool UNiagaraSpriteRendererProperties::PopulateRequiredBindings(FNiagaraParameterStore& InParameterStore)
 {
-	bool bAnyAdded = false;
+	bool bAnyAdded = Super::PopulateRequiredBindings(InParameterStore);
 
 	for (const FNiagaraVariableAttributeBinding* Binding : AttributeBindings)
 	{
diff --git a/Engine/Plugins/FX/Niagara/Source/Niagara/Public/NiagaraRenderer.h b/Engine/Plugins/FX/Niagara/Source/Niagara/Public/NiagaraRenderer.h
--- a/Engine/Plugins/FX/Niagara/Source/Niagara/Public/NiagaraRenderer.h
+++ b/Engine/Plugins/FX/Niagara/Source/Niagara/Public/NiagaraRenderer.h
@@ -139,7 +139,8 @@ protected:
 
 	virtual void ProcessMaterialParameterBindings(TConstArrayView< FNiagaraMaterialAttributeBinding > InMaterialParameterBindings, const FNiagaraEmitterInstance* InEmitter, TConstArrayView<UMaterialInterface*> InMaterials) const;
 
-	bool UseLocalSpace(const FNiagaraSceneProxy* Proxy)const;
+	bool IsRendererEnabled(const UNiagaraRendererProperties* InProperties, const FNiagaraEmitterInstance* Emitter) const;
+	bool UseLocalSpace(const FNiagaraSceneProxy* Proxy) const;
 
 	struct FNiagaraDynamicDataBase *DynamicDataRender;
 
diff --git a/Engine/Plugins/FX/Niagara/Source/Niagara/Public/NiagaraRendererProperties.h b/Engine/Plugins/FX/Niagara/Source/Niagara/Public/NiagaraRendererProperties.h
--- a/Engine/Plugins/FX/Niagara/Source/Niagara/Public/NiagaraRendererProperties.h
+++ b/Engine/Plugins/FX/Niagara/Source/Niagara/Public/NiagaraRendererProperties.h
@@ -194,7 +194,7 @@ public:
 	virtual void GetAssetTagsForContext(const UObject* InAsset, const TArray<const UNiagaraRendererProperties*>& InProperties, TMap<FName, uint32>& NumericKeys, TMap<FName, FString>& StringKeys) const;
 
 	/** In the case that we need parameters bound in that aren't Particle variables, these should be set up here so that the data is appropriately populated after the simulation.*/
-	virtual bool PopulateRequiredBindings(FNiagaraParameterStore& InParameterStore) { return false; }
+	virtual bool PopulateRequiredBindings(FNiagaraParameterStore& InParameterStore);
 
 #if WITH_EDITORONLY_DATA
 
@@ -263,6 +263,10 @@ public:
 	UPROPERTY(EditAnywhere, Category = "Motion Blur")
 	ENiagaraRendererMotionVectorSetting MotionVectorSetting;
 
+	/** Optional bool binding to dynamically enable / disable the renderer. */
+	UPROPERTY(EditAnywhere, Category = "Bindings")
+	FNiagaraVariableAttributeBinding RendererEnabledBinding;
+
 	UPROPERTY()
 	bool bIsEnabled;
 
