Added a check to catch wrongly implemented renderers early

diff --git a/Engine/Plugins/FX/Niagara/Source/NiagaraEditorWidgets/Private/SNiagaraOverviewStackNode.cpp b/Engine/Plugins/FX/Niagara/Source/NiagaraEditorWidgets/Private/SNiagaraOverviewStackNode.cpp
--- a/Engine/Plugins/FX/Niagara/Source/NiagaraEditorWidgets/Private/SNiagaraOverviewStackNode.cpp
+++ b/Engine/Plugins/FX/Niagara/Source/NiagaraEditorWidgets/Private/SNiagaraOverviewStackNode.cpp
@@ -432,6 +432,7 @@ void SNiagaraOverviewStackNode::FillTopContentBar()
 				RendererItem->GetRendererProperties()->GetRendererWidgets(InInstance, Widgets, UThumbnailManager::Get().GetSharedThumbnailPool());
 				TArray<TSharedPtr<SWidget>> TooltipWidgets;
 				RendererItem->GetRendererProperties()->GetRendererTooltipWidgets(InInstance, TooltipWidgets, UThumbnailManager::Get().GetSharedThumbnailPool());
+				check(Widgets.Num() == TooltipWidgets.Num());
 				for (int32 WidgetIndex = 0; WidgetIndex < Widgets.Num(); WidgetIndex++)
 				{
 					ToolBarBuilder.AddWidget(
