Do not force synchronize parameters with definitions when postloading graphs.
Forcing synchronization was resulting in graph generating a new ChangeID which would cause an assert on applying changes to Emitters.

diff --git a/Engine/Plugins/FX/Niagara/Source/Niagara/Private/NiagaraCommon.cpp b/Engine/Plugins/FX/Niagara/Source/Niagara/Private/NiagaraCommon.cpp
--- a/Engine/Plugins/FX/Niagara/Source/Niagara/Private/NiagaraCommon.cpp
+++ b/Engine/Plugins/FX/Niagara/Source/Niagara/Private/NiagaraCommon.cpp
@@ -1397,10 +1397,10 @@ void FNiagaraUtilities::PrepareRapidIterationParameters(const TArray<UNiagaraScr
 		{
 			if (Script->GetLatestSource())
 			{
-			Script->GetLatestSource()->CleanUpOldAndInitializeNewRapidIterationParameters(*Emitter, Script->GetUsage(), Script->GetUsageId(), ParameterStoreToPrepare);
+				Script->GetLatestSource()->CleanUpOldAndInitializeNewRapidIterationParameters(*Emitter, Script->GetUsage(), Script->GetUsageId(), ParameterStoreToPrepare);
+			}
 		}
 	}
-	}
 
 	// Copy parameters for dependencies.
 	for (auto It = ScriptToPreparedParameterStoreMap.CreateIterator(); It; ++It)
@@ -1511,7 +1511,8 @@ const FString FNiagaraCompileOptions::GpuScriptDefine = TEXT("GPUComputeSim");
 FSynchronizeWithParameterDefinitionsArgs::FSynchronizeWithParameterDefinitionsArgs()
 	: SpecificDefinitionsUniqueIds(TArray<FGuid>())
 	, SpecificDestScriptVarIds(TArray<FGuid>())
-	, bForceSynchronizeDefinitions(false)
+	, bForceGatherDefinitions(false)
+	, bForceSynchronizeParameters(false)
 	, bSubscribeAllNameMatchParameters(false)
 	, AdditionalOldToNewNames()
 {
diff --git a/Engine/Plugins/FX/Niagara/Source/Niagara/Private/NiagaraParameterDefinitionsSubscriber.cpp b/Engine/Plugins/FX/Niagara/Source/Niagara/Private/NiagaraParameterDefinitionsSubscriber.cpp
--- a/Engine/Plugins/FX/Niagara/Source/Niagara/Private/NiagaraParameterDefinitionsSubscriber.cpp
+++ b/Engine/Plugins/FX/Niagara/Source/Niagara/Private/NiagaraParameterDefinitionsSubscriber.cpp
@@ -23,8 +23,11 @@ void INiagaraParameterDefinitionsSubscriber::PostLoadDefinitionsSubscriptions()
 		}
 	}
 
+	// When postloading definition subscriptions, we want to synchronize all parameters with all parameter definitions that are matching by name.
+	// As such; Set bForceGatherDefinitions so that all NiagaraParameterDefinitions assets are gathered to consider for linking, and;
+	// Set bSubscribeAllNameMatchParameters so that name matches are considered for linking parameters to parameter definitions.
 	FSynchronizeWithParameterDefinitionsArgs Args;
-	Args.bForceSynchronizeDefinitions = true;
+	Args.bForceGatherDefinitions = true;
 	Args.bSubscribeAllNameMatchParameters = true;
 	SynchronizeWithParameterDefinitions(Args);
 }
@@ -149,7 +152,7 @@ void INiagaraParameterDefinitionsSubscriber::SynchronizeWithParameterDefinitions
 		if (const FParameterDefinitionsSubscription* Subscription = Subscriptions.FindByPredicate(
 			[&DefinitionAndChangeIdHash](const FParameterDefinitionsSubscription& inSubscription){ return inSubscription.DefinitionsId == DefinitionAndChangeIdHash.Definition->GetDefinitionsUniqueId(); }))
 		{
-			if (Args.bForceSynchronizeDefinitions)
+			if (Args.bForceGatherDefinitions)
 			{
 				TargetDefinitionAndChangeIdHashes.Add(DefinitionAndChangeIdHash);
 				TargetDefinitions.Add(DefinitionAndChangeIdHash.Definition);
diff --git a/Engine/Plugins/FX/Niagara/Source/Niagara/Public/NiagaraCommon.h b/Engine/Plugins/FX/Niagara/Source/Niagara/Public/NiagaraCommon.h
--- a/Engine/Plugins/FX/Niagara/Source/Niagara/Public/NiagaraCommon.h
+++ b/Engine/Plugins/FX/Niagara/Source/Niagara/Public/NiagaraCommon.h
@@ -1531,7 +1531,12 @@ struct NIAGARA_API FSynchronizeWithParameterDefinitionsArgs
 	 *  ignore the pending synchronization flag and gather via INiagaraParameterDefinitionsSubscriber::GetSubscribedParameterDefinitions().
 	 *  Note: If true, SpecificDefinitionsUniqueIds and SpecificDestScriptVarIds will still apply.
 	 */
-	bool bForceSynchronizeDefinitions;
+	bool bForceGatherDefinitions;
+
+	/** Default false; If true, instead of comparing parameter change IDs to check if a parameter needs to be synchronized, ignore the change IDs.
+	 *  Note: If true, SpecificDefinitionsUniqueIds and SpecificDestScriptVarIds will still apply.
+	 */
+	bool bForceSynchronizeParameters;
 
 	/** Default false; If true, search for parameters with the same name and type as definitions, and if there is a match, subscribe to the definition. 
 	 *  If the existing parameter default value does not match the definition default value, set the parameter to override the definition default value.
diff --git a/Engine/Plugins/FX/Niagara/Source/NiagaraEditor/Private/NiagaraGraph.cpp b/Engine/Plugins/FX/Niagara/Source/NiagaraEditor/Private/NiagaraGraph.cpp
--- a/Engine/Plugins/FX/Niagara/Source/NiagaraEditor/Private/NiagaraGraph.cpp
+++ b/Engine/Plugins/FX/Niagara/Source/NiagaraEditor/Private/NiagaraGraph.cpp
@@ -1575,13 +1575,13 @@ bool UNiagaraGraph::RenameParameterFromPin(const FNiagaraVariable& Parameter, FN
 	{
 		// Only create a new variable if needed.
 		UNiagaraScriptVariable* FoundOldScriptVariable = *FoundOldScriptVariablePtr;
-			UNiagaraScriptVariable* NewScriptVariable = CastChecked<UNiagaraScriptVariable>(StaticDuplicateObject(FoundOldScriptVariable, this, FName()));
-			NewScriptVariable->SetFlags(RF_Transactional);
-			CopyScriptVariableDataForRename(*FoundOldScriptVariable, *NewScriptVariable);
-			NewScriptVariable->Variable.SetName(NewName);
+		UNiagaraScriptVariable* NewScriptVariable = CastChecked<UNiagaraScriptVariable>(StaticDuplicateObject(FoundOldScriptVariable, this, FName()));
+		NewScriptVariable->SetFlags(RF_Transactional);
+		CopyScriptVariableDataForRename(*FoundOldScriptVariable, *NewScriptVariable);
+		NewScriptVariable->Variable.SetName(NewName);
 		NewScriptVariable->Metadata.CreateNewGuid();
 		NewScriptVariable->SetIsSubscribedToParameterDefinitions(false);
-			VariableToScriptVariable.Add(NewParameter, NewScriptVariable);
+		VariableToScriptVariable.Add(NewParameter, NewScriptVariable);
 		ParameterPanelViewModel->SubscribeParameterToLibraryIfMatchingDefinition(NewScriptVariable, NewName);
 
 		const FNiagaraGraphParameterReferenceCollection* ReferenceCollection = GetParameterReferenceMap().Find(Parameter);
@@ -1937,7 +1937,7 @@ void UNiagaraGraph::SynchronizeParametersWithParameterDefinitions(
 		{
 			if (const UNiagaraScriptVariable* TargetLibraryScriptVar = GetTargetDefinitionScriptVarWithSameId(TargetScriptVar))
 			{
-				bMarkRequiresSync |= SynchronizeScriptVariable(TargetLibraryScriptVar, TargetScriptVar, Args.bForceSynchronizeDefinitions);
+				bMarkRequiresSync |= SynchronizeScriptVariable(TargetLibraryScriptVar, TargetScriptVar, Args.bForceSynchronizeParameters);
 			}
 			else if(AllDefinitionsParameterIds.Contains(TargetScriptVar->Metadata.GetVariableGuid()) == false)
 			{ 
diff --git a/Engine/Plugins/FX/Niagara/Source/NiagaraEditor/Private/ViewModels/NiagaraParameterDefinitionsSubscriberViewModel.cpp b/Engine/Plugins/FX/Niagara/Source/NiagaraEditor/Private/ViewModels/NiagaraParameterDefinitionsSubscriberViewModel.cpp
--- a/Engine/Plugins/FX/Niagara/Source/NiagaraEditor/Private/ViewModels/NiagaraParameterDefinitionsSubscriberViewModel.cpp
+++ b/Engine/Plugins/FX/Niagara/Source/NiagaraEditor/Private/ViewModels/NiagaraParameterDefinitionsSubscriberViewModel.cpp
@@ -39,7 +39,8 @@ void INiagaraParameterDefinitionsSubscriberViewModel::SynchronizeScriptVarWithPa
 {
 	FSynchronizeWithParameterDefinitionsArgs Args = FSynchronizeWithParameterDefinitionsArgs();
 	Args.SpecificDestScriptVarIds.Add(ScriptVarToSynchronize->Metadata.GetVariableGuid());
-	Args.bForceSynchronizeDefinitions = bForce;
+	Args.bForceGatherDefinitions = bForce;
+	Args.bForceSynchronizeParameters = bForce;
 	SynchronizeWithParameterDefinitions(Args); //@todo(ng) error if we failed to find the var
 }
 
