Set Niagara plugin to be enabled by default in UE5

diff --git a/Engine/Plugins/FX/Niagara/Niagara.uplugin b/Engine/Plugins/FX/Niagara/Niagara.uplugin
--- a/Engine/Plugins/FX/Niagara/Niagara.uplugin
+++ b/Engine/Plugins/FX/Niagara/Niagara.uplugin
@@ -10,7 +10,7 @@
 	"DocsURL" : "",
 	"MarketplaceURL" : "",
 	"SupportURL" : "",
-	"EnabledByDefault" : false,
+	"EnabledByDefault" : true,
 	"CanContainContent" : true,
 	"IsBetaVersion" : false,
 	"Installed" : false,
