Removed unnecessary Niagara compilation log messages

diff --git a/Engine/Plugins/FX/Niagara/Source/Niagara/Classes/NiagaraSystem.h b/Engine/Plugins/FX/Niagara/Source/Niagara/Classes/NiagaraSystem.h
--- a/Engine/Plugins/FX/Niagara/Source/Niagara/Classes/NiagaraSystem.h
+++ b/Engine/Plugins/FX/Niagara/Source/Niagara/Classes/NiagaraSystem.h
@@ -256,6 +256,7 @@ struct FNiagaraSystemCompileRequest
 	
 	bool bIsValid = true;
 	bool bForced = false;
+	bool bAllScriptsSynchronized = false;
 
 	float CombinedCompileTime = 0.0f;
 };
diff --git a/Engine/Plugins/FX/Niagara/Source/Niagara/Private/NiagaraSystem.cpp b/Engine/Plugins/FX/Niagara/Source/Niagara/Private/NiagaraSystem.cpp
--- a/Engine/Plugins/FX/Niagara/Source/Niagara/Private/NiagaraSystem.cpp
+++ b/Engine/Plugins/FX/Niagara/Source/Niagara/Private/NiagaraSystem.cpp
@@ -2588,7 +2588,7 @@ bool UNiagaraSystem::QueryCompileComplete(bool bWait, bool bDoPost, bool bDoNotA
 			UE_LOG(LogNiagara, Log, TEXT("Compiling System %s took %f sec (time since issued), %f sec (combined shader worker time)."),
 				*GetFullName(), ElapsedWallTime, CompileRequest.CombinedCompileTime);
 		}
-		else
+		else if (CompileRequest.bAllScriptsSynchronized == false)
 		{
 			UE_LOG(LogNiagara, Log, TEXT("Compiling System %s took %f sec, no shader worker used."), *GetFullName(), ElapsedWallTime);
 		}
@@ -2929,6 +2929,7 @@ bool UNiagaraSystem::RequestCompile(bool bForce, FNiagaraSystemUpdateContext* Op
 	// We might be able to just complete compilation right now if nothing needed compilation.
 	if (ScriptsNeedingCompile.Num() == 0)
 	{
+		ActiveCompilation.bAllScriptsSynchronized = true;
 		PollForCompilationComplete();
 	}
 
diff --git a/Engine/Plugins/FX/Niagara/Source/NiagaraEditor/Private/NiagaraScriptSource.cpp b/Engine/Plugins/FX/Niagara/Source/NiagaraEditor/Private/NiagaraScriptSource.cpp
--- a/Engine/Plugins/FX/Niagara/Source/NiagaraEditor/Private/NiagaraScriptSource.cpp
+++ b/Engine/Plugins/FX/Niagara/Source/NiagaraEditor/Private/NiagaraScriptSource.cpp
@@ -195,7 +195,7 @@ void FindObjectNamesRecursive(UNiagaraGraph* InGraph, ENiagaraScriptUsage Usage,
 					FName DIName = FHlslNiagaraTranslator::GetDataInterfaceName(InputNode->Input.GetName(), EmitterUniqueName, bIsParameterMapDataInterface);
 					if (Result.Contains(DIName) && Result[DIName] != DataInterface)
 					{
-						UE_LOG(LogNiagaraEditor, Log, TEXT("Duplicate data interface name %s in graph %s. One of the data interfaces will override the other when resolving the name."), *DIName.ToString(), *InGraph->GetPathName());
+						UE_LOG(LogNiagaraEditor, Verbose, TEXT("Duplicate data interface name %s in graph %s. One of the data interfaces will override the other when resolving the name."), *DIName.ToString(), *InGraph->GetPathName());
 					}
 					Result.Add(DIName, DataInterface);
 				}
